SkyLight - fix real time capture to account for lower hemispehre color alpha when blending.

diff --git a/Engine/Shaders/Private/ReflectionEnvironmentShaders.usf b/Engine/Shaders/Private/ReflectionEnvironmentShaders.usf
--- a/Engine/Shaders/Private/ReflectionEnvironmentShaders.usf
+++ b/Engine/Shaders/Private/ReflectionEnvironmentShaders.usf
@@ -393,29 +393,35 @@ void DownsampleIntegrateCubemapIlluminancePS(
 
 
 
-#ifdef USE_COMPUTE
+#if APPLY_LOWER_HEMISPHERE_COLOR_PIXELSHADER
+
 float4 LowerHemisphereSolidColor;
-[numthreads(THREADGROUP_SIZE, THREADGROUP_SIZE, 1)]
-void ApplyLowerHemisphereColorCS(uint3 ThreadId : SV_DispatchThreadID)
+
+// This pixel shader is designed to work with RenderRealTimeReflectionHeightFogVS
+void ApplyLowerHemisphereColorPS(
+	in float4 SVPos : SV_POSITION,
+	out float4 OutLuminanceAlpha : SV_Target0)
 {
-	const uint2 FaceCoord = uint2(ThreadId.x % uint(FaceThreadGroupSize), ThreadId.y);
-	if (any(FaceCoord >= uint2(ValidDispatchCoord)))
-	{
-		return;
-	}
-	const int SelectedCubeFace = int(ThreadId.x) / FaceThreadGroupSize;
-	float2 ScaledUVs = ((float2(FaceCoord) + 0.5f) / float2(ValidDispatchCoord)) * 2.0f - 1.0f;
-	float4 OutColor;
+	float2 UV = SVPos.xy * SvPositionToUVScale;
+	float2 ScaledUVs = UV * 2 - 1;
 
-	float3 CubeCoordinates = GetCubemapVector(ScaledUVs, SelectedCubeFace);
+	float3 CubeCoordinates = GetCubemapVector(ScaledUVs, CubeFace);
 	float3 N = normalize(CubeCoordinates);
+
+	OutLuminanceAlpha = float4(0.0f, 0.0f, 0.0f, 1.0f);
 	if (N.z < 0.0f)
 	{
-		OutTextureMipColor[uint3(FaceCoord, SelectedCubeFace)] = float4(LowerHemisphereSolidColor.rgb * LowerHemisphereSolidColor.a, 1.0f);
+		float Coverage = 0.0f;
+		Coverage = pow(saturate(LowerHemisphereSolidColor.a), 2.2);
+		Coverage = saturate(LowerHemisphereSolidColor.a);
+		OutLuminanceAlpha = float4(LowerHemisphereSolidColor.rgb * Coverage, 1.0 - Coverage);
 	}
 }
+
 #endif
 
+
+
 int NumCaptureArrayMips;
 
 /** Cube map array of reflection captures. */
diff --git a/Engine/Source/Runtime/Renderer/Private/ReflectionEnvironmentRealTimeCapture.cpp b/Engine/Source/Runtime/Renderer/Private/ReflectionEnvironmentRealTimeCapture.cpp
--- a/Engine/Source/Runtime/Renderer/Private/ReflectionEnvironmentRealTimeCapture.cpp
+++ b/Engine/Source/Runtime/Renderer/Private/ReflectionEnvironmentRealTimeCapture.cpp
@@ -151,33 +151,38 @@ public:
 IMPLEMENT_GLOBAL_SHADER(FComputeSkyEnvMapDiffuseIrradianceCS, "/Engine/Private/ReflectionEnvironmentShaders.usf", "ComputeSkyEnvMapDiffuseIrradianceCS", SF_Compute);
 
 
-
-class FApplyLowerHemisphereColor : public FGlobalShader
+class FApplyLowerHemisphereColorPS : public FGlobalShader
 {
-	DECLARE_GLOBAL_SHADER(FApplyLowerHemisphereColor);
-	SHADER_USE_PARAMETER_STRUCT(FApplyLowerHemisphereColor, FGlobalShader);
-
-	static const uint32 ThreadGroupSize = 8;
+	DECLARE_GLOBAL_SHADER(FApplyLowerHemisphereColorPS);
+	SHADER_USE_PARAMETER_STRUCT(FApplyLowerHemisphereColorPS, FGlobalShader);
 
 	using FPermutationDomain = TShaderPermutationDomain<>;
 
 	BEGIN_SHADER_PARAMETER_STRUCT(FParameters, )
-		SHADER_PARAMETER(FLinearColor, LowerHemisphereSolidColor)
-		SHADER_PARAMETER(FIntPoint, ValidDispatchCoord)
-		SHADER_PARAMETER(int32, FaceThreadGroupSize)
-		SHADER_PARAMETER_RDG_TEXTURE_UAV(RWTexture2D, OutTextureMipColor)
+		SHADER_PARAMETER_STRUCT_REF(FViewUniformShaderParameters, ViewUniformBuffer)
+		SHADER_PARAMETER(FVector4f, LowerHemisphereSolidColor)
+		SHADER_PARAMETER(FVector2f, SvPositionToUVScale)
+		SHADER_PARAMETER(int32, CubeFace)
+		RENDER_TARGET_BINDING_SLOTS()
 	END_SHADER_PARAMETER_STRUCT()
 
-public:
-	static bool ShouldCompilePermutation(const FGlobalShaderPermutationParameters& Parameters) { return GetMaxSupportedFeatureLevel(Parameters.Platform) >= ERHIFeatureLevel::SM5; }
+	static FPermutationDomain RemapPermutation(FPermutationDomain PermutationVector)
+	{
+		return PermutationVector;
+	}
+
+	static bool ShouldCompilePermutation(const FGlobalShaderPermutationParameters& Parameters)
+	{
+		return true;
+	}
+
 	static void ModifyCompilationEnvironment(const FGlobalShaderPermutationParameters& Parameters, FShaderCompilerEnvironment& OutEnvironment)
 	{
 		FGlobalShader::ModifyCompilationEnvironment(Parameters, OutEnvironment);
-		OutEnvironment.SetDefine(TEXT("THREADGROUP_SIZE"), ThreadGroupSize);
-		OutEnvironment.SetDefine(TEXT("USE_COMPUTE"), 1);
+		OutEnvironment.SetDefine(TEXT("APPLY_LOWER_HEMISPHERE_COLOR_PIXELSHADER"), 1);
 	}
 };
-IMPLEMENT_GLOBAL_SHADER(FApplyLowerHemisphereColor, "/Engine/Private/ReflectionEnvironmentShaders.usf", "ApplyLowerHemisphereColorCS", SF_Compute);
+IMPLEMENT_GLOBAL_SHADER(FApplyLowerHemisphereColorPS, "/Engine/Private/ReflectionEnvironmentShaders.usf", "ApplyLowerHemisphereColorPS", SF_Pixel);
 
 
 class FRenderRealTimeReflectionHeightFogVS : public FGlobalShader
@@ -746,33 +751,53 @@ void FScene::AllocateAndCaptureFrameSkyEnvMap(
 
 					SceneRenderer.RenderVolumetricCloudsInternal(GraphBuilder, CloudRC, InstanceCullingManager);
 				}
-			}
-
-			// Render lower hemisphere color
-			if (SkyLight->bLowerHemisphereIsSolidColor)
-			{
-				FApplyLowerHemisphereColor::FPermutationDomain PermutationVector;
-				TShaderMapRef<FApplyLowerHemisphereColor> ComputeShader(GetGlobalShaderMap(FeatureLevel), PermutationVector);
 
-				const uint32 MipIndex = 0;
-				const uint32 Mip0Resolution = SkyCubeTexture->Desc.GetSize().X;
-				FApplyLowerHemisphereColor::FParameters* PassParameters = GraphBuilder.AllocParameters<FApplyLowerHemisphereColor::FParameters>();
-				PassParameters->ValidDispatchCoord = FIntPoint(Mip0Resolution, Mip0Resolution);
-				PassParameters->LowerHemisphereSolidColor = SkyLight->LowerHemisphereColor;
-
-				FRDGTextureUAVDesc OutTextureMipColorDesc(SkyCubeTexture, MipIndex);
-				OutTextureMipColorDesc.DimensionOverride = ETextureDimension::Texture2DArray;
-				PassParameters->OutTextureMipColor = GraphBuilder.CreateUAV(OutTextureMipColorDesc);
-
-				FIntVector NumGroups = FIntVector::DivideAndRoundUp(FIntVector(Mip0Resolution, Mip0Resolution, 1), FIntVector(FApplyLowerHemisphereColor::ThreadGroupSize, FApplyLowerHemisphereColor::ThreadGroupSize, 1));
-
-				// The groupd size per face with padding
-				PassParameters->FaceThreadGroupSize = NumGroups.X * FConvolveSpecularFaceCS::ThreadGroupSize;
-
-				// We are going to dispatch once for all faces 
-				NumGroups.X *= 6;
-
-				FComputeShaderUtils::AddPass(GraphBuilder, RDG_EVENT_NAME("ApplyLowerHemisphereColor"), ComputeShader, PassParameters, NumGroups);
+				// Now blend over the lower hemisphere color, during the second pass for clouds
+				if (SkyLight->bLowerHemisphereIsSolidColor && bExecuteCloud)
+				{
+					FRenderRealTimeReflectionHeightFogVS::FPermutationDomain VsPermutationVector;
+					TShaderMapRef<FRenderRealTimeReflectionHeightFogVS> VertexShader(GetGlobalShaderMap(SkyRC.FeatureLevel), VsPermutationVector);
+
+					FApplyLowerHemisphereColorPS::FPermutationDomain PsPermutationVector;
+					TShaderMapRef<FApplyLowerHemisphereColorPS> PixelShader(GetGlobalShaderMap(SkyRC.FeatureLevel), PsPermutationVector);
+
+					FApplyLowerHemisphereColorPS::FParameters* PsPassParameters = GraphBuilder.AllocParameters<FApplyLowerHemisphereColorPS::FParameters>();
+					PsPassParameters->RenderTargets = SkyRC.RenderTargets;
+					PsPassParameters->ViewUniformBuffer = CubeView.ViewUniformBuffer;
+					PsPassParameters->LowerHemisphereSolidColor = SkyLight->LowerHemisphereColor;
+					PsPassParameters->CubeFace = CubeFace;
+					PsPassParameters->SvPositionToUVScale = FVector2f(1.0f / float(CubeWidth), 1.0f / float(CubeWidth));
+
+					// Render height fog at an infinite distance since real time reflections does not have a depth buffer for now.
+					// Volumetric fog is not supported in such reflections.
+					GraphBuilder.AddPass(
+						RDG_EVENT_NAME("ApplyLowerHemisphereColor"),
+						PsPassParameters,
+						ERDGPassFlags::Raster,
+						[PsPassParameters, VertexShader, PixelShader, CubeWidth](FRHICommandList& RHICmdListLambda)
+						{
+							RHICmdListLambda.SetViewport(0.0f, 0.0f, 0.0f, CubeWidth, CubeWidth, 1.0f);
+
+							FGraphicsPipelineStateInitializer GraphicsPSOInit;
+							RHICmdListLambda.ApplyCachedRenderTargets(GraphicsPSOInit);
+
+							GraphicsPSOInit.BlendState = TStaticBlendState<CW_RGB, BO_Add, BF_One, BF_SourceAlpha, BO_Add, BF_Zero, BF_Zero>::GetRHI();
+							GraphicsPSOInit.DepthStencilState = TStaticDepthStencilState<false, CF_Always>::GetRHI();
+							GraphicsPSOInit.RasterizerState = TStaticRasterizerState<FM_Solid, CM_None>::GetRHI();
+							GraphicsPSOInit.BoundShaderState.VertexDeclarationRHI = GEmptyVertexDeclaration.VertexDeclarationRHI;
+							GraphicsPSOInit.BoundShaderState.VertexShaderRHI = VertexShader.GetVertexShader();
+							GraphicsPSOInit.BoundShaderState.PixelShaderRHI = PixelShader.GetPixelShader();
+							GraphicsPSOInit.PrimitiveType = PT_TriangleList;
+							SetGraphicsPipelineState(RHICmdListLambda, GraphicsPSOInit, 0);
+
+							FRenderRealTimeReflectionHeightFogVS::FParameters VsPassParameters;
+							VsPassParameters.ViewUniformBuffer = PsPassParameters->ViewUniformBuffer;
+							SetShaderParameters(RHICmdListLambda, VertexShader, VertexShader.GetVertexShader(), VsPassParameters);
+							SetShaderParameters(RHICmdListLambda, PixelShader, PixelShader.GetPixelShader(), *PsPassParameters);
+
+							RHICmdListLambda.DrawPrimitive(0, 1, 1);
+						});
+				}
 			}
 		}
 		else
