Add ETaskTag::EParallelRenderingThread tag to RayTracingSceneInitTask
- prevents ensure when calling CVarRayTracingSceneBuildMode.GetValueOnRenderThread()

diff --git a/Engine/Source/Runtime/Renderer/Private/RayTracing/RayTracing.cpp b/Engine/Source/Runtime/Renderer/Private/RayTracing/RayTracing.cpp
--- a/Engine/Source/Runtime/Renderer/Private/RayTracing/RayTracing.cpp
+++ b/Engine/Source/Runtime/Renderer/Private/RayTracing/RayTracing.cpp
@@ -1444,6 +1444,8 @@ namespace RayTracing
 		View.RayTracingSceneInitTask = FFunctionGraphTask::CreateAndDispatchWhenReady(
 			[&View, &RayTracingScene]()
 			{
+				FTaskTagScope TaskTagScope(ETaskTag::EParallelRenderingThread);
+
 				TRACE_CPUPROFILER_EVENT_SCOPE(RayTracingSceneInitTask);
 				View.RayTracingSceneInitData = RayTracingScene.BuildInitializationData();
 			},
