Prescale DeferredCommands array and optimize MDC debug data.

diff --git a/Engine/Source/Runtime/Renderer/Private/MeshPassProcessor.cpp b/Engine/Source/Runtime/Renderer/Private/MeshPassProcessor.cpp
--- a/Engine/Source/Runtime/Renderer/Private/MeshPassProcessor.cpp
+++ b/Engine/Source/Runtime/Renderer/Private/MeshPassProcessor.cpp
@@ -1667,13 +1667,14 @@ void FMeshDrawCommand::SetDebugData(const FPrimitiveSceneProxy* PrimitiveScenePr
 	DebugData.MaterialRenderProxy = MaterialRenderProxy;
 	DebugData.VertexShader = UntypedShaders.VertexShader;
 	DebugData.PixelShader = UntypedShaders.PixelShader;
-	DebugData.VertexFactory = VertexFactory;
-	DebugData.VertexFactoryType = VertexFactory->GetType();
 	DebugData.LODIndex = MeshBatch.LODIndex;
 	DebugData.SegmentIndex = MeshBatch.SegmentIndex;
-	DebugData.PSOCollectorIndex = PSOCollectorIndex;
 	DebugData.ResourceName =  PrimitiveSceneProxy ? PrimitiveSceneProxy->GetResourceName() : FName();
-	DebugData.MaterialName = MaterialRenderProxy->GetMaterialName();
+#if PSO_PRECACHING_VALIDATE
+	DebugData.VertexFactory = VertexFactory;
+	DebugData.VertexFactoryType = VertexFactory->GetType();
+	DebugData.PSOCollectorIndex = PSOCollectorIndex;
+#endif
 }
 #endif
 
@@ -1691,7 +1692,7 @@ void FMeshDrawCommand::GetStatsData(FVisibleMeshDrawCommandStatsData& OutVisible
 	OutVisibleStatsData.LODIndex = DebugData.LODIndex;
 	OutVisibleStatsData.SegmentIndex = DebugData.SegmentIndex;
 	OutVisibleStatsData.ResourceName = DebugData.ResourceName;
-	OutVisibleStatsData.MaterialName = DebugData.MaterialName;
+	OutVisibleStatsData.MaterialName = DebugData.MaterialRenderProxy->GetMaterialName();
 #endif
 }
 #endif // MESH_DRAW_COMMAND_STATS
@@ -2262,7 +2263,7 @@ FMeshDrawCommand::FMeshDrawEvent::FMeshDrawEvent(const FMeshDrawCommand& MeshDra
 		  RHICmdList
 		, GShowMaterialDrawEvents != 0
 		, TEXT("%s %s (%u instances)")
-		, MeshDrawCommand.DebugData.MaterialName
+		, MeshDrawCommand.DebugData.MaterialRenderProxy->GetMaterialName()
 		, MeshDrawCommand.DebugData.ResourceName
 		, MeshDrawCommand.NumInstances * InstanceFactor
 	)
diff --git a/Engine/Source/Runtime/Renderer/Public/MeshPassProcessor.h b/Engine/Source/Runtime/Renderer/Public/MeshPassProcessor.h
--- a/Engine/Source/Runtime/Renderer/Public/MeshPassProcessor.h
+++ b/Engine/Source/Runtime/Renderer/Public/MeshPassProcessor.h
@@ -737,18 +737,23 @@ inline const int32 NumInlineShaderBindings = 10;
 struct FMeshDrawCommandDebugData
 {
 #if MESH_DRAW_COMMAND_DEBUG_DATA
+	// sorted from larger to smaller to minimize padding
+	TShaderRef<FShader> VertexShader;
+	TShaderRef<FShader> PixelShader;
 	const FPrimitiveSceneProxy* PrimitiveSceneProxyIfNotUsingStateBuckets;
 	const FMaterial* Material;
 	const FMaterialRenderProxy* MaterialRenderProxy;
-	TShaderRef<FShader> VertexShader;
-	TShaderRef<FShader> PixelShader;
+#if PSO_PRECACHING_VALIDATE
+	// so far these are only used for PSO precaching validation
 	const FVertexFactory* VertexFactory;
 	const FVertexFactoryType* VertexFactoryType;
+#endif
+	FName ResourceName;
+#if PSO_PRECACHING_VALIDATE
+	uint32 PSOCollectorIndex;
+#endif
 	int8 LODIndex;
 	uint8 SegmentIndex;
-	uint32 PSOCollectorIndex;
-	FName ResourceName;
-	FString MaterialName;
 #endif
 };
 
@@ -1911,7 +1916,13 @@ public:
 class FCachedPassMeshDrawListContextDeferred : public FCachedPassMeshDrawListContext
 {
 public:
-	FCachedPassMeshDrawListContextDeferred(FScene& InScene) : FCachedPassMeshDrawListContext(InScene) {}
+	FCachedPassMeshDrawListContextDeferred(FScene& InScene) : FCachedPassMeshDrawListContext(InScene)
+	{
+		if (LIKELY(bUseGPUScene))
+		{
+			DeferredCommandHashes.Reserve(DeferredCommands.Max());
+		}
+	}
 
 	virtual void FinalizeCommand(
 		const FMeshBatch& MeshBatch, 
@@ -1928,7 +1939,7 @@ public:
 	void DeferredFinalizeMeshDrawCommands(const TArrayView<FPrimitiveSceneInfo*>& SceneInfos, int32 Start, int32 End);
 
 private:
-	TArray<FMeshDrawCommand> DeferredCommands;
+	TArray<FMeshDrawCommand, TInlineAllocator<192>> DeferredCommands;
 	TArray<Experimental::FHashType> DeferredCommandHashes;
 };
 
