VolumetricCloud - removed unused parameter.

diff --git a/Engine/Source/Runtime/Renderer/Private/VolumetricCloudRendering.cpp b/Engine/Source/Runtime/Renderer/Private/VolumetricCloudRendering.cpp
--- a/Engine/Source/Runtime/Renderer/Private/VolumetricCloudRendering.cpp
+++ b/Engine/Source/Runtime/Renderer/Private/VolumetricCloudRendering.cpp
@@ -846,7 +846,6 @@ class FRenderVolumetricCloudRenderViewCS : public FMeshMaterialShader
 		SHADER_PARAMETER_STRUCT_REF(FViewUniformShaderParameters, View)
 		SHADER_PARAMETER_RDG_UNIFORM_BUFFER(FSceneUniformParameters, Scene)
 		SHADER_PARAMETER_STRUCT_INCLUDE(FVirtualShadowMapSamplingParameters, VirtualShadowMap)
-		SHADER_PARAMETER(int32, VirtualShadowMapId0)
 		SHADER_PARAMETER_RDG_TEXTURE_UAV(RWTexture2D, OutCloudColor0)
 		SHADER_PARAMETER_RDG_TEXTURE_UAV(RWTexture2D, OutCloudColor1)
 		SHADER_PARAMETER_RDG_TEXTURE_UAV(RWTexture2D, OutCloudDepth)
