Fixed overflow when batching transforms when the presized number isn't large enough (this can happen if a component happens to submit two updates, as GroomComponent currently does).

diff --git a/Engine/Source/Runtime/Renderer/Private/RendererScene.cpp b/Engine/Source/Runtime/Renderer/Private/RendererScene.cpp
--- a/Engine/Source/Runtime/Renderer/Private/RendererScene.cpp
+++ b/Engine/Source/Runtime/Renderer/Private/RendererScene.cpp
@@ -2170,13 +2170,20 @@ void FScene::UpdatePrimitiveTransformInternal(T* Primitive)
 
 			if (bPerformUpdate)
 			{
+				bool bNeedsTransformCommand = true;
+
 				// Accumulate all transform updates and enqueue them as once
 				if (bPrimitivesUpdateBatching)
 				{
-					int32 Index = PrimitiveUpdateIndex.fetch_add(1, std::memory_order_relaxed);
-					PrimitivesUpdates[Index] = MoveTemp(UpdateParams);
+					const int32 Index = PrimitiveUpdateIndex.fetch_add(1, std::memory_order_relaxed);
+					if (ensure(Index < PrimitivesUpdates.Num()))
+					{
+						PrimitivesUpdates[Index] = MoveTemp(UpdateParams);
+						bNeedsTransformCommand = false;
+					}
 				}
-				else
+
+				if (bNeedsTransformCommand)
 				{
 					ENQUEUE_RENDER_COMMAND(UpdateTransformCommand)(
 						[UpdateParams](FRHICommandListBase&)
@@ -2218,7 +2225,7 @@ void FScene::FinishUpdatePrimitiveTransform()
 	{
 		// Pass the collection and actual number of accumulated updates
 		ENQUEUE_RENDER_COMMAND(UpdateTransformCommand)(
-			[PrimitivesUpdates = MoveTemp(PrimitivesUpdates), NumUpdates = PrimitiveUpdateIndex.load(std::memory_order_relaxed)](FRHICommandListBase&)
+			[PrimitivesUpdates = MoveTemp(PrimitivesUpdates), NumUpdates = FMath::Min(PrimitivesUpdates.Num(), PrimitiveUpdateIndex.load(std::memory_order_relaxed))](FRHICommandListBase&)
 			{
 				for (int32 Index = 0; Index < NumUpdates; ++Index)
 				{
