Path Tracer: Add a cvar to quickly override the use of reference Depth of Field

This matches the new EnableReferenceAtmosphere cvar, which can be faster for quickly testing this feature than tracking down the right PostProcessVolume.

diff --git a/Engine/Source/Runtime/Renderer/Private/PathTracing.cpp b/Engine/Source/Runtime/Renderer/Private/PathTracing.cpp
--- a/Engine/Source/Runtime/Renderer/Private/PathTracing.cpp
+++ b/Engine/Source/Runtime/Renderer/Private/PathTracing.cpp
@@ -194,6 +194,16 @@ TAutoConsoleVariable<int32> CVarPathTracingEnableCameraBackfaceCulling(
 	ECVF_RenderThreadSafe
 );
 
+TAutoConsoleVariable<int32> CVarPathTracingEnableReferenceDOF(
+	TEXT("r.PathTracing.EnableReferenceDOF"),
+	-1,
+	TEXT("Should the path tracer ray trace the depth-of-field effect instead of the post-processed effect?\n")
+	TEXT("-1: Inherit from PostProcess settings (default)\n")
+	TEXT(" 0: Disabled\n")
+	TEXT(" 1: Enable\n"),
+	ECVF_RenderThreadSafe
+);
+
 TAutoConsoleVariable<int32> CVarPathTracingEnableReferenceAtmosphere(
 	TEXT("r.PathTracing.EnableReferenceAtmosphere"),
 	-1,
@@ -684,6 +694,12 @@ namespace PathTracing
 		const int32 EnableReferenceAtmosphereCVar = CVarPathTracingEnableReferenceAtmosphere.GetValueOnRenderThread();
 		return EnableReferenceAtmosphereCVar < 0 ? View.FinalPostProcessSettings.PathTracingEnableReferenceAtmosphere != 0 : EnableReferenceAtmosphereCVar != 0;
 	}
+
+	bool UsesReferenceDOF(const FViewInfo& View)
+	{
+		const int32 EnableReferenceDOFCvar = CVarPathTracingEnableReferenceDOF.GetValueOnRenderThread();
+		return EnableReferenceDOFCvar < 0 ? View.FinalPostProcessSettings.PathTracingEnableReferenceDOF != 0 : EnableReferenceDOFCvar != 0;
+	}
 }
 
 // This function prepares the portion of shader arguments that may involve invalidating the path traced state
@@ -719,7 +735,7 @@ static void PreparePathTracingData(const FScene* Scene, const FViewInfo& View, F
 	PathTracingData.CameraFocusDistance = 0;
 	PathTracingData.CameraLensRadius = FVector2f::ZeroVector;
 	if (ShowFlags.DepthOfField &&
-		PPV.PathTracingEnableReferenceDOF &&
+		PathTracing::UsesReferenceDOF(View) &&
 		PPV.DepthOfFieldFocalDistance > 0 &&
 		PPV.DepthOfFieldFstop > 0)
 	{
@@ -3987,6 +4003,11 @@ namespace PathTracing
 	{
 		return false;
 	}
+
+	bool UsesReferenceDOF(const FViewInfo& View)
+	{
+		return false;
+	}
 }
 
 #endif
diff --git a/Engine/Source/Runtime/Renderer/Private/PathTracing.h b/Engine/Source/Runtime/Renderer/Private/PathTracing.h
--- a/Engine/Source/Runtime/Renderer/Private/PathTracing.h
+++ b/Engine/Source/Runtime/Renderer/Private/PathTracing.h
@@ -37,4 +37,5 @@ namespace PathTracing
 {
 	bool UsesDecals(const FSceneViewFamily& ViewFamily);
 	bool UsesReferenceAtmosphere(const FViewInfo& View);
+	bool UsesReferenceDOF(const FViewInfo& View);
 }
diff --git a/Engine/Source/Runtime/Renderer/Private/PostProcess/DiaphragmDOF.cpp b/Engine/Source/Runtime/Renderer/Private/PostProcess/DiaphragmDOF.cpp
--- a/Engine/Source/Runtime/Renderer/Private/PostProcess/DiaphragmDOF.cpp
+++ b/Engine/Source/Runtime/Renderer/Private/PostProcess/DiaphragmDOF.cpp
@@ -1361,7 +1361,7 @@ bool DiaphragmDOF::IsEnabled(const FViewInfo& View)
 		DiaphragmDOF::IsSupported(View.GetShaderPlatform()) &&
 		View.Family->EngineShowFlags.DepthOfField &&
 		bDepthOfFieldRequestedByCVar &&
-		!(View.Family->EngineShowFlags.PathTracing && View.FinalPostProcessSettings.PathTracingEnableReferenceDOF) &&
+		!(View.Family->EngineShowFlags.PathTracing && PathTracing::UsesReferenceDOF(View)) &&
 		((View.FinalPostProcessSettings.DepthOfFieldFstop > 0.f && View.FinalPostProcessSettings.DepthOfFieldFocalDistance > 0.f) || View.FinalPostProcessSettings.DepthOfFieldDepthBlurRadius > 0.f);
 }
 
