Add a callback to modify the wireframe pre-pass viewfamily through FMeshEdgesViewFamilySettings, and move this API to the public folder.

diff --git a/Engine/Source/Runtime/Renderer/Private/MeshEdgesRendering.cpp b/Engine/Source/Runtime/Renderer/Private/MeshEdgesRendering.cpp
--- a/Engine/Source/Runtime/Renderer/Private/MeshEdgesRendering.cpp
+++ b/Engine/Source/Runtime/Renderer/Private/MeshEdgesRendering.cpp
@@ -158,6 +158,7 @@ public:
 	TUniquePtr<FRenderTargetTexture> WireframeDepth;
 	TArray<FIntRect> ViewRects;
 	TArray<FMeshEdgesViewSettings> ViewSettings;
+	FMeshEdgesViewFamilySettings ViewFamilySettings = {};
 };
 
 const FMeshEdgesViewSettings& GetMeshEdgesViewSettings(const FSceneViewFamily& ViewFamily, int ViewIndex)
@@ -190,6 +191,17 @@ FMeshEdgesViewSettings& GetMeshEdgesViewSettings(FSceneView& View)
 	return const_cast<FMeshEdgesViewSettings&>(GetMeshEdgesViewSettings(AsConst(View)));
 }
 
+const FMeshEdgesViewFamilySettings& GetMeshEdgesViewFamilySettings(const FSceneViewFamily& ViewFamily)
+{
+	FMeshEdgesViewFamilyData* FamilyData = const_cast<FSceneViewFamily&>(ViewFamily).GetOrCreateExtentionData<FMeshEdgesViewFamilyData>();
+	return FamilyData->ViewFamilySettings;
+}
+
+FMeshEdgesViewFamilySettings& GetMeshEdgesViewFamilySettings(FSceneViewFamily& ViewFamily)
+{
+	return const_cast<FMeshEdgesViewFamilySettings&>(GetMeshEdgesViewFamilySettings(AsConst(ViewFamily)));
+}
+
 void RenderMeshEdges(FSceneViewFamily& InViewFamily);
 
 class FMeshEdgesExtension : public FSceneViewExtensionBase
@@ -290,46 +302,47 @@ void RenderMeshEdges(FSceneViewFamily& ViewFamily)
 	}
 
 	FMeshEdgesViewFamilyData* ViewFamilyData = ViewFamily.GetOrCreateExtentionData<FMeshEdgesViewFamilyData>();
+	const FMeshEdgesViewFamilySettings& Settings = ViewFamilyData->ViewFamilySettings;
 
 	ERHIFeatureLevel::Type FeatureLevel = ViewFamily.GetFeatureLevel();
 	FIntPoint DesiredBufferSize = FSceneRenderer::GetDesiredInternalBufferSize(ViewFamily);
 	ViewFamilyData->CreateRenderTargets(FeatureLevel, DesiredBufferSize);
 
-	FEngineShowFlags ShowFlags = ViewFamily.EngineShowFlags;
+	FEngineShowFlags WireframeShowFlags = ViewFamily.EngineShowFlags;
 	{
 		// Render a wireframe view
-		ShowFlags.SetWireframe(true);
+		WireframeShowFlags.SetWireframe(true);
 
 		// Copy the MSAA wireframe view only, don't copy other scene elements
-		ShowFlags.SetSceneCaptureCopySceneDepth(false);
+		WireframeShowFlags.SetSceneCaptureCopySceneDepth(false);
 
 		// Disable rendering of elements that are not needed
-		ShowFlags.SetMeshEdges(false);
-		ShowFlags.SetLighting(false);
-		ShowFlags.SetLightFunctions(false);
-		ShowFlags.SetGlobalIllumination(false);
-		ShowFlags.SetLumenGlobalIllumination(false);
-		ShowFlags.SetLumenReflections(false);
-		ShowFlags.SetDynamicShadows(false);
-		ShowFlags.SetCapsuleShadows(false);
-		ShowFlags.SetDistanceFieldAO(false);
-		ShowFlags.SetFog(false);
-		ShowFlags.SetVolumetricFog(false);
-		ShowFlags.SetCloud(false);
-		ShowFlags.SetDecals(false);
-		ShowFlags.SetAtmosphere(false);
-		ShowFlags.SetPostProcessing(false);
-		ShowFlags.SetCompositeEditorPrimitives(false);
-		ShowFlags.SetGrid(false);
-		//ShowFlags.SetScreenPercentage(false);
-		//ShowFlags.SetTranslucency(false);
+		WireframeShowFlags.SetMeshEdges(false);
+		WireframeShowFlags.SetLighting(false);
+		WireframeShowFlags.SetLightFunctions(false);
+		WireframeShowFlags.SetGlobalIllumination(false);
+		WireframeShowFlags.SetLumenGlobalIllumination(false);
+		WireframeShowFlags.SetLumenReflections(false);
+		WireframeShowFlags.SetDynamicShadows(false);
+		WireframeShowFlags.SetCapsuleShadows(false);
+		WireframeShowFlags.SetDistanceFieldAO(false);
+		WireframeShowFlags.SetFog(false);
+		WireframeShowFlags.SetVolumetricFog(false);
+		WireframeShowFlags.SetCloud(false);
+		WireframeShowFlags.SetDecals(false);
+		WireframeShowFlags.SetAtmosphere(false);
+		WireframeShowFlags.SetPostProcessing(false);
+		WireframeShowFlags.SetCompositeEditorPrimitives(false);
+		WireframeShowFlags.SetGrid(false);
+		//WireframeShowFlags.SetScreenPercentage(false);
+		//WireframeShowFlags.SetTranslucency(false);
 	}
 	
 	FSceneViewFamilyContext CaptureViewFamily(
 		FSceneViewFamily::ConstructionValues(
 			ViewFamilyData->WireframeColor.Get(),
 			ViewFamily.Scene,
-			ShowFlags)
+			WireframeShowFlags)
 		.SetRenderTargetDepth(ViewFamilyData->WireframeDepth.Get())
 		.SetResolveScene(true)
 		.SetRealtimeUpdate(true)
@@ -346,6 +359,8 @@ void RenderMeshEdges(FSceneViewFamily& ViewFamily)
 		CaptureViewFamily.SetScreenPercentageInterface(ViewFamily.GetScreenPercentageInterface()->Fork_GameThread(CaptureViewFamily));
 	}
 
+	Settings.OnBeforeWireframeRender(CaptureViewFamily);
+
 	FSceneRenderer* SceneRenderer = FSceneRenderer::CreateSceneRenderer(&CaptureViewFamily, nullptr);
 
 	for (const FSceneViewExtensionRef& Extension : CaptureViewFamily.ViewExtensions)
diff --git a/Engine/Source/Runtime/Renderer/Private/MeshEdgesRendering.h b/Engine/Source/Runtime/Renderer/Private/MeshEdgesRendering.h
--- a/Engine/Source/Runtime/Renderer/Private/MeshEdgesRendering.h
+++ b/Engine/Source/Runtime/Renderer/Private/MeshEdgesRendering.h
@@ -2,20 +2,13 @@
 
 #pragma once
 
+#include "MeshEdges.h"
+
 class FRDGBuilder;
-class FSceneView;
 class FViewInfo;
 struct FCompositePrimitiveInputs;
 struct FScreenPassRenderTarget;
 
-struct FMeshEdgesViewSettings
-{
-	float Opacity = 1.0;
-};
-
-const FMeshEdgesViewSettings& GetMeshEdgesViewSettings(const FSceneView& View);
-FMeshEdgesViewSettings& GetMeshEdgesViewSettings(FSceneView& View);
-
 void ComposeMeshEdges(FRDGBuilder& GraphBuilder,
 						const FViewInfo& View,
 						FScreenPassRenderTarget& EditorPrimitivesColor,
diff --git a/Engine/Source/Runtime/Renderer/Public/MeshEdges.h b/Engine/Source/Runtime/Renderer/Public/MeshEdges.h
new file mode 100644
--- /dev/null
+++ b/Engine/Source/Runtime/Renderer/Public/MeshEdges.h
@@ -0,0 +1,26 @@
+// Copyright Epic Games, Inc. All Rights Reserved.
+
+#pragma once
+
+#include "HAL/Platform.h"
+#include "Templates/Function.h"
+
+class FSceneView;
+class FSceneViewFamily;
+
+struct FMeshEdgesViewSettings
+{
+	// Opacity of the wireframe blended with the shaded view.
+	float Opacity = 1.0;
+};
+
+RENDERER_API const FMeshEdgesViewSettings& GetMeshEdgesViewSettings(const FSceneView& View);
+RENDERER_API FMeshEdgesViewSettings& GetMeshEdgesViewSettings(FSceneView& View);
+
+struct FMeshEdgesViewFamilySettings
+{
+	TFunction<void(FSceneViewFamily& WireframeViewFamily)> OnBeforeWireframeRender = [](auto&){};
+};
+
+RENDERER_API const FMeshEdgesViewFamilySettings& GetMeshEdgesViewFamilySettings(const FSceneViewFamily& ViewFamily);
+RENDERER_API FMeshEdgesViewFamilySettings& GetMeshEdgesViewFamilySettings(FSceneViewFamily& ViewFamily);
