Changes to various volumetric areas to allow them to work with Orthographic rendering. This includes:

Sky, Cloud and Fog
- Created 2 pathways for resolving these volumetrics, a typical Ortho view which is tightened to the camera width and height (ignoring depth) and a pseudo FOV option which resolves a 90 degree viewing angle of the sky for a wider scene view.
- Enabled AO resolving for cloud etc definition.
- Planar reflections fixed to allow this on various surfaces.
- Motion Blur improvements due to large depth diffs in ortho cameras between near (character) and far (cloud) objects
- Disabled VRT option as it doesn't resolve well without Ortho having access to depth velocities for compensating velocity diffs, and ultimately isn't entirely necessary at this stage.
- Clamped sky materials to the far plane as they will otherwise will be clipped. Materials not marked sky but being used for the sky dome can't be recognised by the engine.

diff --git a/Engine/Shaders/Private/BasePassVertexShader.usf b/Engine/Shaders/Private/BasePassVertexShader.usf
--- a/Engine/Shaders/Private/BasePassVertexShader.usf
+++ b/Engine/Shaders/Private/BasePassVertexShader.usf
@@ -66,6 +66,16 @@ void Main(
 	{
 		float4 RasterizedWorldPosition = VertexFactoryGetRasterizedWorldPosition(Input, VFIntermediates, WorldPosition);
 		ClipSpacePosition = mul(RasterizedWorldPosition, ResolvedView.TranslatedWorldToClip);
+#if MATERIAL_IS_SKY
+		if (IsOrthoProjection())
+		{
+			//Ortho rendering doesn't stretch to LWC typically, so clamp the sky to as close to the reverse-z far plane as possible unless the dome is already closer to the camera.
+			if(ConvertFromDeviceZ(ClipSpacePosition.z) > ResolvedView.OrthoFarPlane)
+			{
+				ClipSpacePosition.z = FarDepthValue;
+			}
+		}
+#endif
 		Output.Position = INVARIANT(ClipSpacePosition);
 	}
 
diff --git a/Engine/Shaders/Private/Common.ush b/Engine/Shaders/Private/Common.ush
--- a/Engine/Shaders/Private/Common.ush
+++ b/Engine/Shaders/Private/Common.ush
@@ -1275,6 +1275,28 @@ float ConvertFromDeviceZ(float DeviceZ)
 	return DeviceZ * View.InvDeviceZToWorldZTransform[0] + View.InvDeviceZToWorldZTransform[1] + 1.0f / (DeviceZ * View.InvDeviceZToWorldZTransform[2] - View.InvDeviceZToWorldZTransform[3]);
 }
 
+float2 AdjustClipToPrevClipForProjectionType(float2 Velocity, float DeviceZ)
+{
+	if(IsOrthoProjection())
+	{
+		/*
+		* Motion vectors on far distance planes are greatly influenced by slight rotations in camera angle
+		* And with Orthographic cameras, this magnifies the velocities over large distances.
+		* As a result, we need to counter this by reducing the speed with the DeviceZ
+		*/
+#if HAS_INVERTED_Z_BUFFER
+		/**
+		* Because this is ClipToPrevClip rather than velocity, we use the non - linear device z which works better than linear depth,
+		* however if adjusting velocity directly, linear depth should be used instead. 
+		*/
+		Velocity *= DeviceZ;
+#else
+		Velocity *= 1.0f - DeviceZ;
+#endif
+	}
+	return Velocity;
+}
+
 float GetScreenPositionDepth(float4 ScreenPosition)
 {
 #if RAYTRACINGSHADER
@@ -1441,12 +1463,16 @@ void SvPositionToResolvedScreenPositionDeriv(float4 SvPosition, float2 PPZ_DDX_D
 	ScreenPositionDDY = NDCPos * W_DDX_DDY.y + NDCPosDDY * SvPosition.w;
 }
 
-float3 GetTranslatedWorldCameraPosFromView(ViewState InView, float2 SvPosition)
+float3 GetTranslatedWorldCameraPosFromView(ViewState InView, float2 SvPosition, bool bForceAddOrthoHeight = false)
 {
 	// Get the camera position per pixel in ortho, or just the camera location in perspective.
 	if (IsOrthoProjection(InView))
 	{
-		return SvPositionToTranslatedWorld(float4(SvPosition, 1.0f, 1.0f));
+		//When using sky height, for example, the horizon calculation requires bForceAddOrthoHeight to be true to 
+		//avoid negative calculations when comparing the relative sky heights + view camera height
+		float3 CameraPosition = SvPositionToTranslatedWorld(float4(SvPosition, 1.0f, 1.0f));
+		CameraPosition.z += select(bForceAddOrthoHeight, View.ClipToView[1][1], 0.0f);
+		return CameraPosition;
 	}
 	else
 	{
@@ -1454,9 +1480,25 @@ float3 GetTranslatedWorldCameraPosFromView(ViewState InView, float2 SvPosition)
 	}
 }
 
-float3 GetTranslatedWorldCameraPosFromView(float2 SvPosition)
+FDFVector3 GetWorldCameraPosFromView(ViewState InView, float2 SvPosition, bool bForceAddOrthoHeight = false)
+{
+	if(IsOrthoProjection(InView))
+	{
+		//When using sky height, for example, the horizon calculation requires bForceAddOrthoHeight to be true to 
+		//avoid negative calculations when comparing the relative sky heights + view camera height
+		FDFVector3 CameraPosition = SvPositionToWorld(float4(SvPosition, 1.0f, 1.0f));
+		CameraPosition = DFFastAdd(select(bForceAddOrthoHeight, float3(0.0f, 0.0f, View.ClipToView[1][1]), float3(0.0f,0.0f,0.0f)), CameraPosition);
+		return CameraPosition;
+	}
+	else
+	{
+		return InView.WorldCameraOrigin;
+	}
+}
+
+float3 GetTranslatedWorldCameraPosFromView(float2 SvPosition, bool bForceAddOrthoHeight = false)
 {
-	return GetTranslatedWorldCameraPosFromView(PrimaryView, SvPosition);
+	return GetTranslatedWorldCameraPosFromView(PrimaryView, SvPosition, bForceAddOrthoHeight);
 }
 
 float2 SvPositionToViewportUV(float4 SvPosition)
diff --git a/Engine/Shaders/Private/DebugProbes.usf b/Engine/Shaders/Private/DebugProbes.usf
--- a/Engine/Shaders/Private/DebugProbes.usf
+++ b/Engine/Shaders/Private/DebugProbes.usf
@@ -252,8 +252,8 @@ void MainPS(
 	const float Depth = 1000000.0f;
 	float4 TranslatedWorldPos = mul(float4(GetScreenPositionForProjectionType(ScreenPosition, Depth), Depth, 1), View.ScreenToTranslatedWorld);
 
-	const float3 RayDir = normalize(TranslatedWorldPos.xyz - View.TranslatedWorldCameraOrigin);
-	const float3 RayOrig = DFHackToFloat(ResolvedView.WorldCameraOrigin);
+	const float3 RayDir = GetCameraVectorFromTranslatedWorldPosition(TranslatedWorldPos.xyz);
+	const float3 RayOrig = DFHackToFloat(GetWorldCameraPosFromView(ResolvedView, SvPosition.xy));
 
 #if PERMUTATION_ILLUMINANCEMETER /////////////////////////////////////////////////////////////////////////////////////////////////////////////////////////////////////
 
diff --git a/Engine/Shaders/Private/DeferredLightPixelShaders.usf b/Engine/Shaders/Private/DeferredLightPixelShaders.usf
--- a/Engine/Shaders/Private/DeferredLightPixelShaders.usf
+++ b/Engine/Shaders/Private/DeferredLightPixelShaders.usf
@@ -97,7 +97,7 @@ FDerivedParams GetDerivedParams(in FInputParams Input, in float SceneDepth)
 	Out.TranslatedWorldPosition = mul(float4(ClipPosition, SceneDepth, 1), View.ScreenToTranslatedWorld).xyz;
 	Out.CameraVector = GetCameraVectorFromTranslatedWorldPosition(Out.TranslatedWorldPosition);
 #else
-	Out.TranslatedWorldPosition = Input.ScreenVector * SceneDepth + View.TranslatedWorldCameraOrigin;
+	Out.TranslatedWorldPosition = Input.ScreenVector * SceneDepth + GetTranslatedWorldCameraPosFromView(Input.PixelPos + 0.5f);
 	Out.CameraVector = normalize(Input.ScreenVector);
 #endif
 	return Out;
@@ -471,8 +471,8 @@ void DeferredLightPixelMain(
 	HairScreenSpaceData.AmbientOcclusion = 1;
 	HairScreenSpaceData.GBuffer = HairSampleToGBufferData(Sample, HairStrands.HairDualScatteringRoughnessOverride);
 
-	const float3 TranslatedWorldPosition = mul(float4(ScreenPosition * HairScreenSpaceData.GBuffer.Depth, HairScreenSpaceData.GBuffer.Depth, 1), View.ScreenToTranslatedWorld).xyz;
-	const float3 CameraVector = normalize(TranslatedWorldPosition - View.TranslatedWorldCameraOrigin);
+	const float3 TranslatedWorldPosition = mul(float4(GetScreenPositionForProjectionType(ScreenPosition, HairScreenSpaceData.GBuffer.Depth), HairScreenSpaceData.GBuffer.Depth, 1), View.ScreenToTranslatedWorld).xyz;
+	const float3 CameraVector = GetCameraVectorFromTranslatedWorldPosition(TranslatedWorldPosition);
 
 	FDeferredLightData LightData = InitDeferredLightFromUniforms(CURRENT_LIGHT_TYPE);
 
diff --git a/Engine/Shaders/Private/MaterialTemplate.ush b/Engine/Shaders/Private/MaterialTemplate.ush
--- a/Engine/Shaders/Private/MaterialTemplate.ush
+++ b/Engine/Shaders/Private/MaterialTemplate.ush
@@ -2178,8 +2178,13 @@ float3 MaterialExpressionSkyAtmosphereLightDiskLuminance(FMaterialPixelParameter
 #if MATERIAL_SKY_ATMOSPHERE && PROJECT_SUPPORT_SKY_ATMOSPHERE
 	if (ResolvedView.RenderingReflectionCaptureMask == 0.0f) // Do not render light disk when in reflection capture in order to avoid double specular. The sun contribution is already computed analyticaly.
 	{
-		const float3 PlanetCenterToCameraTranslatedWorld = (ResolvedView.SkyCameraTranslatedWorldOrigin - ResolvedView.SkyPlanetTranslatedWorldCenterAndViewHeight.xyz) * CM_TO_SKY_UNIT;
 		const float3 ViewDir = -Parameters.CameraVector;
+		float3 PlanetCenterToCameraTranslatedWorld = (ResolvedView.SkyCameraTranslatedWorldOrigin - ResolvedView.SkyPlanetTranslatedWorldCenterAndViewHeight.xyz);
+		if (IsOrthoProjection())
+		{
+			PlanetCenterToCameraTranslatedWorld = GetTranslatedWorldCameraPosFromView(Parameters.SvPosition.xy, true) + PlanetCenterToCameraTranslatedWorld;
+		}
+		PlanetCenterToCameraTranslatedWorld *= CM_TO_SKY_UNIT;
 
 		// GetLightDiskLuminance does a test against the virtual planet but SkyCameraTranslatedWorldOrigin is always put safely setup above it (to never have the camera into the virtual planet with a black screen)
 		LightDiskLuminance =  GetLightDiskLuminance(PlanetCenterToCameraTranslatedWorld, ViewDir, ResolvedView.SkyAtmosphereBottomRadiusKm, ResolvedView.SkyAtmosphereTopRadiusKm,
@@ -2199,12 +2204,20 @@ float3 MaterialExpressionSkyAtmosphereViewLuminance(FMaterialPixelParameters Par
 	{
 		return float3(0.0f, 0.0f, 0.0f);
 	}
-	
-	const float ViewHeight = ResolvedView.SkyPlanetTranslatedWorldCenterAndViewHeight.w * CM_TO_SKY_UNIT;
-	const float3 ViewDir = -Parameters.CameraVector;
 
 	// The referencial used to build the Sky View lut
 	float3x3 LocalReferencial = GetSkyViewLutReferential(ResolvedView.SkyViewLutReferential);
+	const float3 ViewDir = -Parameters.CameraVector;
+
+	float ViewHeight = ResolvedView.SkyPlanetTranslatedWorldCenterAndViewHeight.w;
+	if(IsOrthoProjection())
+	{
+		float3 CameraPosition = GetTranslatedWorldCameraPosFromView(Parameters.SvPosition.xy, true);
+		CameraPosition.z += ViewHeight;
+		ViewHeight = mul(LocalReferencial, CameraPosition).z;
+	}
+	ViewHeight *= CM_TO_SKY_UNIT;
+
 	// Compute inputs in this referential
 	float3 WorldPosLocal = float3(0.0, 0.0, ViewHeight);
 	float3 UpVectorLocal = float3(0.0, 0.0, 1.0);
diff --git a/Engine/Shaders/Private/MotionBlur/MotionBlurVelocityFlatten.usf b/Engine/Shaders/Private/MotionBlur/MotionBlurVelocityFlatten.usf
--- a/Engine/Shaders/Private/MotionBlur/MotionBlurVelocityFlatten.usf
+++ b/Engine/Shaders/Private/MotionBlur/MotionBlurVelocityFlatten.usf
@@ -44,7 +44,7 @@ float2 GetCameraMotionVelocity(uint2 DispatchThreadId, float Depth, float4x4 Cli
 	float4 ThisClip = float4(ScreenPos, Depth, 1);
 	float4 PrevClip = mul(ThisClip, ClipToPrevClip);
 	float2 PrevScreen = PrevClip.xy / PrevClip.w;
-	return ScreenPos - PrevScreen;
+	return AdjustClipToPrevClipForProjectionType(ScreenPos - PrevScreen, Depth);
 }
 
 float2 PreprocessVelocityForMotionBlur(uint2 DispatchThreadId, float DeviceZ, float2 Velocity, float2 CameraMotionVelocity)
diff --git a/Engine/Shaders/Private/MotionBlur/VisualizeMotionVectors.usf b/Engine/Shaders/Private/MotionBlur/VisualizeMotionVectors.usf
--- a/Engine/Shaders/Private/MotionBlur/VisualizeMotionVectors.usf
+++ b/Engine/Shaders/Private/MotionBlur/VisualizeMotionVectors.usf
@@ -36,7 +36,7 @@ float2 ComputeStaticVelocity(float2 ScreenPos, float DeviceZ)
 	float4 ThisClip = float4(PosN, 1);
 	float4 PrevClip = mul( ThisClip, View.ClipToPrevClip );
 	float2 PrevScreen = PrevClip.xy / PrevClip.w;
-	return PosN.xy - PrevScreen;
+	return AdjustClipToPrevClipForProjectionType(PosN.xy - PrevScreen, DeviceZ);
 }
 
 
diff --git a/Engine/Shaders/Private/PlanarReflectionShaders.usf b/Engine/Shaders/Private/PlanarReflectionShaders.usf
--- a/Engine/Shaders/Private/PlanarReflectionShaders.usf
+++ b/Engine/Shaders/Private/PlanarReflectionShaders.usf
@@ -168,7 +168,7 @@ void PlanarReflectionPS(
 
 	FGBufferData GBuffer = GetGBufferDataFromSceneTextures(BufferUV);
 
-	float3 TranslatedWorldPosition = ResolvedView.TranslatedWorldCameraOrigin + ScreenVector * GBuffer.Depth;
+	float3 TranslatedWorldPosition = GetTranslatedWorldCameraPosFromView(SVPos.xy) + ScreenVector * GBuffer.Depth;
 
 	// Skip unlit pixels (skybox)
 	BRANCH
diff --git a/Engine/Shaders/Private/PlanarReflectionShared.ush b/Engine/Shaders/Private/PlanarReflectionShared.ush
--- a/Engine/Shaders/Private/PlanarReflectionShared.ush
+++ b/Engine/Shaders/Private/PlanarReflectionShared.ush
@@ -26,7 +26,7 @@ half4 ComputePlanarReflections(float3 TranslatedWorldPosition, half3 WorldNormal
 	if (FinalFade > 0)
 	{
 		// CameraToPixel in the main view is what we used as ReflectionVector when rendering the reflection pass to PlanarReflectionTexture
-		float3 CameraToPixel = normalize(TranslatedWorldPosition - ResolvedView.TranslatedWorldCameraOrigin);
+		float3 CameraToPixel = GetCameraVectorFromTranslatedWorldPosition(ResolvedView, TranslatedWorldPosition);
 		// Reflect the effective ReflectionVector in mirrored space to get the original camera vector
 		float3 MirroredCameraVector = reflect(CameraToPixel, -PlanarReflectionStruct.ReflectionPlane.xyz);
 		// Transform the GBuffer normal into mirrored space
diff --git a/Engine/Shaders/Private/SkyAtmosphere.usf b/Engine/Shaders/Private/SkyAtmosphere.usf
--- a/Engine/Shaders/Private/SkyAtmosphere.usf
+++ b/Engine/Shaders/Private/SkyAtmosphere.usf
@@ -164,8 +164,8 @@ float3 GetScreenWorldDir(in float4 SVPos)
 {
 	float2 ScreenPosition = SvPositionToScreenPosition(SVPos).xy;
 	const float Depth = 1000000.0f;
-	float4 WorldPos = mul(float4(GetScreenPositionForProjectionType(ScreenPosition, Depth), Depth, 1), View.ScreenToTranslatedWorld);
-	return GetCameraVectorFromTranslatedWorldPosition(WorldPos.xyz);
+	float4 TranslatedWorldPos = mul(float4(GetScreenPositionForProjectionType(ScreenPosition, Depth), Depth, 1), View.ScreenToTranslatedWorld);
+	return GetCameraVectorFromTranslatedWorldPosition(TranslatedWorldPos.xyz);
 }
 
 // This is the world position of the camera. It is also force to be at the top of the virutal planet surface. 
@@ -869,8 +869,12 @@ void RenderSkyAtmosphereRayMarchingPS(
 	//}
 
 	// World position are relative to the planet center (itself expressed within translated world space)
-	float3 WorldPos = GetTranslatedCameraPlanetPos();
 	float3 WorldDir = GetScreenWorldDir(SVPos);
+	float3 WorldPos = GetTranslatedCameraPlanetPos();
+	if(IsOrthoProjection())
+	{
+		WorldPos += GetTranslatedWorldCameraPosFromView(SVPos.xy, true);
+	}
 
 	float3 PreExposedL = 0;
 	float3 LuminanceScale = 1.0f;
@@ -1470,6 +1474,11 @@ void RenderCameraAerialPerspectiveVolumeCS(uint3 ThreadId : SV_DispatchThreadID)
 	float4 SVPos = float4(View.ViewRectMin.xy + UV * View.ViewSizeAndInvSize.xy, 0.0f, 1.0f);// SV_POS as if resolution was the one from the scene view.
 	float3 WorldDir = GetScreenWorldDir(SVPos);
 	float3 CamPos = GetTranslatedCameraPlanetPos();
+	if (IsOrthoProjection())
+	{
+		CamPos += GetTranslatedWorldCameraPosFromView(SVPos.xy, true);
+	}
+
 
 	if (RealTimeReflection360Mode)
 	{
diff --git a/Engine/Shaders/Private/SkyAtmosphereCommon.ush b/Engine/Shaders/Private/SkyAtmosphereCommon.ush
--- a/Engine/Shaders/Private/SkyAtmosphereCommon.ush
+++ b/Engine/Shaders/Private/SkyAtmosphereCommon.ush
@@ -55,6 +55,12 @@ float4 GetAerialPerspectiveLuminanceTransmittance(
 	{
 		return float4(0.0f, 0.0f, 0.0f, 1.0f);
 	}
+
+	if (IsOrthoProjection())
+	{
+		float2 SvPosXY = (((NDC.xy * float2(0.5f, -0.5f)) + 0.5f) * View.ViewSizeAndInvSize.xy) + ResolvedView.ViewRectMin.xy;		
+		WorldPositionRelativeToCamera += (GetTranslatedWorldCameraPosFromView(SvPosXY, true) * CM_TO_SKY_UNIT) + WorldPositionRelativeToCamera;
+	}
 	
 	float2 ScreenUv = (NDC.xy / NDC.ww) * float2(0.5f, -0.5f) + 0.5f;
 
diff --git a/Engine/Shaders/Private/VolumetricCloud.usf b/Engine/Shaders/Private/VolumetricCloud.usf
--- a/Engine/Shaders/Private/VolumetricCloud.usf
+++ b/Engine/Shaders/Private/VolumetricCloud.usf
@@ -470,7 +470,6 @@ void MainCommon(in FMaterialPixelParameters MaterialParameters, in float4 SvPosi
 	//
 	// Initialise all the parameters
 	//
-	
 	OutColor0					= float4(0.0f, 0.0f, 0.0f, 1.0f);
 	OutDepth					= MaxHalfFloat;
 
@@ -479,17 +478,18 @@ void MainCommon(in FMaterialPixelParameters MaterialParameters, in float4 SvPosi
 	FPixelMaterialInputs PixelMaterialInputs	= (FPixelMaterialInputs)0;
 	CalcMaterialParameters(MaterialParameters, PixelMaterialInputs, SvPosition, true);
 
-	float3 RayWorldOriginKm				= DFHackToFloat(ResolvedView.WorldCameraOrigin) * CENTIMETER_TO_KILOMETER;
-	float3 RayTranslatedWorldOrigin		= View.TranslatedWorldCameraOrigin;
+	float3 RayTranslatedWorldOrigin		= GetTranslatedWorldCameraPosFromView(SvPosition.xy);
 	float3 RayTranslatedWorldOriginKm	= RayTranslatedWorldOrigin * CENTIMETER_TO_KILOMETER;
+	float3 RayWorldOriginKm				= select(IsOrthoProjection(), DFHackToFloat(DFSubtract(RayTranslatedWorldOrigin, PrimaryView.PreViewTranslation)) * CENTIMETER_TO_KILOMETER, DFHackToFloat(ResolvedView.WorldCameraOrigin) * CENTIMETER_TO_KILOMETER);
 	float3 Raydir						= -MaterialParameters.CameraVector;
+
 	float TMin							= -999999999.0f;
 	float TMax							= -999999999.0f;
 
-	float3 Luminance			= 0.0f;
-	MATVEC TransmittanceToView	= 1.0f;
-	float tAPWeightedSum		= 0.0f;
-	float tAPWeightsSum			= 0.0f;
+	float3 Luminance					= 0.0f;
+	MATVEC TransmittanceToView			= 1.0f;
+	float tAPWeightedSum				= 0.0f;
+	float tAPWeightsSum					= 0.0f;
 
 #if CLOUD_MIN_AND_MAX_DEPTH
 	OutColor1 = float4(0.0f, 0.0f, 0.0f, 1.0f);
@@ -505,10 +505,8 @@ void MainCommon(in FMaterialPixelParameters MaterialParameters, in float4 SvPosi
 	const float3 RayDirOffsetX = GetCameraVector(ResolvedView, TranslatedWorldPositionOffsetX, ResolvedView.TranslatedWorldCameraOrigin);
 	const float3 RayDirOffsetY = GetCameraVector(ResolvedView, TranslatedWorldPositionOffsetY, ResolvedView.TranslatedWorldCameraOrigin);
 
-	// TODO: The current cloud tracing code assumes the ray origin to be at the camera origin, which breaks on ortho projections.
-	// We still define these variables to reduce the risk of bugs in the future when the code is fixed to properly support ortho projections as well.
-	const float3 RayTranslatedWorldOriginOffsetX = RayTranslatedWorldOrigin;
-	const float3 RayTranslatedWorldOriginOffsetY = RayTranslatedWorldOrigin;
+	const float3 RayTranslatedWorldOriginOffsetX = GetTranslatedWorldCameraPosFromView(ResolvedView, SvPosition.xy + float2(1.0f, 0.0f));
+	const float3 RayTranslatedWorldOriginOffsetY = GetTranslatedWorldCameraPosFromView(ResolvedView, SvPosition.xy + float2(0.0f, 1.0f));
 #endif // USE_ANALYTIC_DERIVATIVES
 
 	uint4 TracingCoordToZbufferCoordScaleBias = RenderVolumetricCloudParameters.TracingCoordToZbufferCoordScaleBias;
@@ -563,6 +561,15 @@ void MainCommon(in FMaterialPixelParameters MaterialParameters, in float4 SvPosi
 		TMin = max(0.0f, TMin) * KILOMETER_TO_CENTIMETER;
 		TMax = max(0.0f, TMax) * KILOMETER_TO_CENTIMETER;
 
+		const bool bDoDepthIntersectionAndTest = RenderVolumetricCloudParameters.OpaqueIntersectionMode >= 2;
+		if (IsOrthoProjection())
+		{
+			//Ortho is always drawing to the far plane for now, so clip anything that is in front of it when we are inside the sky sphere. Earlies out compared to the further checks below.
+			if (bDoDepthIntersectionAndTest && length(RayWorldOriginKm - RenderVolumetricCloudParameters.CloudLayerCenterKm) < RenderVolumetricCloudParameters.TopRadiusKm && OutDepth.z > 0.0f)
+			{
+				return;
+			}
+		}
 
 		//
 		// Skip tracing if the range is 0, e.g. cloud layer is behind an object and we do not want to pay for SampleCountMin, or the distance at which tracing should start it too far away
@@ -613,21 +620,21 @@ void MainCommon(in FMaterialPixelParameters MaterialParameters, in float4 SvPosi
 		// Skip tracing if the depth buffer is in front of the cloud layer front interface. If inside the layer, clamp the tracing to the depth buffer
 		//
 
-		const bool bDoDepthInteresectionAndTest = RenderVolumetricCloudParameters.OpaqueIntersectionMode >= 2;
-
-		if (bDoDepthInteresectionAndTest && (TDepthBufferKm * KILOMETER_TO_CENTIMETER) < TMin)
+		if(!IsOrthoProjection() || TMax < View.OrthoFarPlane)
 		{
-			//OutColor0 = float4(0.0f, 0.0f, 1.0f, 1.0f);
-			return;
-		}
+			if (bDoDepthIntersectionAndTest && (TDepthBufferKm * KILOMETER_TO_CENTIMETER) < TMin)
+			{
+				//OutColor0 = float4(0.0f, 0.0f, 1.0f, 1.0f);
+				return;
+			}
 
-		if (bDoDepthInteresectionAndTest)
-		{
-			// Only trace up to the closest between the cloud layer far distance and the depth buffer 
-			TMax = min(TMax, TDepthBufferKm * KILOMETER_TO_CENTIMETER);
+			if (bDoDepthIntersectionAndTest)
+			{
+				// Only trace up to the closest between the cloud layer far distance and the depth buffer 
+				TMax = min(TMax, TDepthBufferKm * KILOMETER_TO_CENTIMETER);
+			}
 		}
 
-
 		//
 		// Prepare a bunch of variable for the tracing
 		//
@@ -1867,15 +1874,15 @@ float GetSvPositionEmptySpaceSkippingValue(in ViewState ResolvedView, in uint3 D
 
 	const float SampleDistanceFromCamera = (float(GroupThreadId.z) + SampleFroxelOffset.z) * GetSliceDepth();
 	const float3 Raydir = -MaterialParameters.CameraVector;
-	const float3 SampleTranslatedWorldPosition = ResolvedView.TranslatedWorldCameraOrigin + Raydir * SampleDistanceFromCamera;
+	const float3 SampleTranslatedWorldPosition = GetTranslatedWorldCameraPosFromView(ResolvedView, SvPosition.xy) + Raydir * SampleDistanceFromCamera;
 
 #if USE_ANALYTIC_DERIVATIVES
 	const float3 TranslatedWorldPositionOffsetX = SvPositionToResolvedTranslatedWorld(SvPosition + float4(1.0f, 0.0f, 0.0f, 0.0f));
 	const float3 TranslatedWorldPositionOffsetY = SvPositionToResolvedTranslatedWorld(SvPosition + float4(0.0f, 1.0f, 0.0f, 0.0f));
 	const float3 RayDirOffsetX = GetCameraVector(ResolvedView, TranslatedWorldPositionOffsetX, ResolvedView.TranslatedWorldCameraOrigin);
 	const float3 RayDirOffsetY = GetCameraVector(ResolvedView, TranslatedWorldPositionOffsetY, ResolvedView.TranslatedWorldCameraOrigin);
-	float3 SampleWorldPositionDDX = (ResolvedView.TranslatedWorldCameraOrigin + t * RayDirOffsetX) - SampleTranslatedWorldPosition;
-	float3 SampleWorldPositionDDY = (ResolvedView.TranslatedWorldCameraOrigin + t * RayDirOffsetY) - SampleTranslatedWorldPosition;
+	float3 SampleWorldPositionDDX = (GetTranslatedWorldCameraPosFromView(ResolvedView, SvPosition.xy + float2(1.0f, 0.0f)) + RayDirOffsetX * SampleDistanceFromCamera) - SampleTranslatedWorldPosition;
+	float3 SampleWorldPositionDDY = (GetTranslatedWorldCameraPosFromView(ResolvedView, SvPosition.xy + float2(0.0f, 1.0f)) + RayDirOffsetY * SampleDistanceFromCamera) - SampleTranslatedWorldPosition;
 #else
 	float3 SampleWorldPositionDDX = 0.0f;
 	float3 SampleWorldPositionDDY = 0.0f;
diff --git a/Engine/Shaders/Private/VolumetricCloudMaterialPixelCommon.ush b/Engine/Shaders/Private/VolumetricCloudMaterialPixelCommon.ush
--- a/Engine/Shaders/Private/VolumetricCloudMaterialPixelCommon.ush
+++ b/Engine/Shaders/Private/VolumetricCloudMaterialPixelCommon.ush
@@ -53,7 +53,7 @@ void UpdateMaterialCloudParam(inout FMaterialPixelParameters MaterialParameters,
 	//MaterialParameters.WorldPosition_NoOffsets;				// TODO
 	//MaterialParameters.LWCData.WorldPosition_NoOffsets;	// TODO
 	//MaterialParameters.WorldPosition_NoOffsets_CamRelative;	// TODO
-	MaterialParameters.CameraVector = normalize(InputView.TranslatedWorldCameraOrigin - TranslatedWorldPosition);
+	MaterialParameters.CameraVector = select(IsOrthoProjection(InputView), InputView.ViewForward, normalize(InputView.TranslatedWorldCameraOrigin - TranslatedWorldPosition));
 	MaterialParameters.WorldPosition_DDX = WorldPositionDDX;
 	MaterialParameters.WorldPosition_DDY = WorldPositionDDY;
 
diff --git a/Engine/Shaders/Private/VolumetricFog.usf b/Engine/Shaders/Private/VolumetricFog.usf
--- a/Engine/Shaders/Private/VolumetricFog.usf
+++ b/Engine/Shaders/Private/VolumetricFog.usf
@@ -164,12 +164,8 @@ void MaterialSetupCS(
 	float CellOffset = .5f;
 	float SceneDepth;
 	float3 TranslatedWorldPosition = ComputeCellTranslatedWorldPosition(GridCoordinate, CellOffset, SceneDepth);
-	if(length(TranslatedWorldPosition - PrimaryView.TranslatedWorldCameraOrigin) > LFVGlobalStartDistance)
+	if(GetDistanceToCameraFromViewVector(TranslatedWorldPosition - PrimaryView.TranslatedWorldCameraOrigin) > LFVGlobalStartDistance)
 	{
-		float3 CamRayTranslatedWorldOrigin = PrimaryView.TranslatedWorldCameraOrigin;
-		float3 PreViewTranslation = DFHackToFloat(PrimaryView.PreViewTranslation);
-		//float3 TranslatedWorldPosition = WorldPosition + PreViewTranslation;
-
 
 		float4 VertexClipSpacePosition = mul(float4(TranslatedWorldPosition.xyz, 1), PrimaryView.TranslatedWorldToClip);
 		float2 SvPosition = (VertexClipSpacePosition.xy / VertexClipSpacePosition.w * float2(.5f, -.5f) + .5f) * PrimaryView.ViewSizeAndInvSize.xy;
@@ -181,10 +177,6 @@ void MaterialSetupCS(
 		{
 			uint InstanceIndex = LFVTileDataTex[uint3(TilePos, 1 + LFVIndex)];
 
-			float3 CamRayWorldOrigin = CamRayTranslatedWorldOrigin - PreViewTranslation;
-			float3 CamRayTranslatedWorldDir = normalize(TranslatedWorldPosition - CamRayTranslatedWorldOrigin);
-			float3 CamRayWorldDir = CamRayTranslatedWorldDir;
-
 			// LFV_TODO skip if out of sphere of influence
 			FLocalFogVolumeGPUInstanceData FogInstance = GetLocalFogVolumeGPUInstanceData(InstanceIndex);
 
@@ -407,7 +399,7 @@ float4 InjectShadowedLocalLightCommon(uint3 GridCoordinate)
 			float SceneDepth;
 			float3 TranslatedWorldPosition = ComputeCellTranslatedWorldPosition(GridCoordinate, CellOffset, SceneDepth);
 			float3 WorldPosition = TranslatedWorldPosition - DFHackToFloat(PrimaryView.PreViewTranslation);
-			float3 CameraVector = normalize(TranslatedWorldPosition - PrimaryView.TranslatedWorldCameraOrigin);
+			float3 CameraVector = GetCameraVectorFromTranslatedWorldPosition(TranslatedWorldPosition);
 
 			float CellRadius = length(TranslatedWorldPosition - ComputeCellTranslatedWorldPosition(GridCoordinate + uint3(1, 1, 1), CellOffset));
 			// Bias the inverse squared light falloff based on voxel size to prevent aliasing near the light source
@@ -748,8 +740,8 @@ void LightScatteringCS(
 		float SceneDepth;
 		float3 TranslatedWorldPosition = ComputeCellTranslatedWorldPosition(GridCoordinate, CellOffset, SceneDepth);
 		float3 WorldPosition = TranslatedWorldPosition - DFHackToFloat(PrimaryView.PreViewTranslation); // LWC_TODO
-		float CameraVectorLength = length(TranslatedWorldPosition - PrimaryView.TranslatedWorldCameraOrigin);
-		float3 CameraVector = (TranslatedWorldPosition - PrimaryView.TranslatedWorldCameraOrigin) / CameraVectorLength;
+		float CameraVectorLength = GetDistanceToCameraFromViewVector(TranslatedWorldPosition - PrimaryView.TranslatedWorldCameraOrigin);
+		float3 CameraVector = GetCameraVectorFromTranslatedWorldPosition(TranslatedWorldPosition);
 		
 	#if SHADING_PATH_MOBILE
 		const bool	 bHasDirectionalLight		= MobileHasDirectionalLight > 0;
diff --git a/Engine/Shaders/Private/VolumetricFogLightFunction.usf b/Engine/Shaders/Private/VolumetricFogLightFunction.usf
--- a/Engine/Shaders/Private/VolumetricFogLightFunction.usf
+++ b/Engine/Shaders/Private/VolumetricFogLightFunction.usf
@@ -89,7 +89,7 @@ void Main(
 	float GreyScale = dot(LightFunction, .3333f).x;
 
 	// Calculate radial view distance for stable fading
-	float ViewDistance = length(PrimaryView.TranslatedWorldCameraOrigin - TranslatedWorldPosition.xyz);
+	float ViewDistance = GetDistanceToCameraFromViewVector(PrimaryView.TranslatedWorldCameraOrigin - TranslatedWorldPosition.xyz);
 
 	float DistanceFadeAlpha = saturate((LightFunctionParameters2.x - ViewDistance) / (LightFunctionParameters2.x * .2f));
 	// Fade to disabled based on LightFunctionFadeDistance
diff --git a/Engine/Source/Runtime/Engine/Private/Camera/CameraStackTypes.cpp b/Engine/Source/Runtime/Engine/Private/Camera/CameraStackTypes.cpp
--- a/Engine/Source/Runtime/Engine/Private/Camera/CameraStackTypes.cpp
+++ b/Engine/Source/Runtime/Engine/Private/Camera/CameraStackTypes.cpp
@@ -375,7 +375,7 @@ bool FMinimalViewInfo::AutoCalculateOrthoPlanes(FSceneViewProjectionData& InOutP
 		 * Forcing the absolute value for this means that Up/Down is 1.0f and Forward (90 degrees) is 0.
 		 * We use this to scale the Near Plane, and the far plane if 16 bit scaling is disabled.
 		 */
-		float CosAngle = FMath::Abs((ViewForward.Dot(FVector(0, 0, -1.0f))));
+		float CosAngle = FMath::Abs(ViewForward.Z);
 		
 		/** 
 		 * We still max out at UE_OLD_WORLD_MAX or Max32FP, but we scale the FarPlane depending on ratio of the pixel size to the world unit size.
diff --git a/Engine/Source/Runtime/Engine/Private/SceneView.cpp b/Engine/Source/Runtime/Engine/Private/SceneView.cpp
--- a/Engine/Source/Runtime/Engine/Private/SceneView.cpp
+++ b/Engine/Source/Runtime/Engine/Private/SceneView.cpp
@@ -600,7 +600,7 @@ bool FSceneViewProjectionData::UpdateOrthoPlanes(FSceneViewProjectionData* InOut
 		float CameraHeightAdjustment = 0.0f;
 		if (CVarOrthoCameraHeightAsViewTarget.GetValueOnAnyThread() && bUseCameraHeightAsViewTarget)
 		{
-			CameraHeightAdjustment = FMath::Abs(FMath::Min(InOutProjectionData->ViewOrigin.Z, HalfOrthoWidth)) * FMath::Abs((ViewForward.Dot(FVector(0, 0, -1.0f))));
+			CameraHeightAdjustment = FMath::Abs(FMath::Min(InOutProjectionData->ViewOrigin.Z, HalfOrthoWidth)) * FMath::Abs(ViewForward.Z);
 		}
 		InOutProjectionData->ViewOrigin += ViewForward * (CameraHeightAdjustment + NearPlane);
 	}
@@ -2641,6 +2641,7 @@ void FSceneView::SetupCommonViewUniformBufferParameters(
 	ViewUniformShaderParameters.WorldCameraMovementSinceLastFrame = FVector3f(InViewMatrices.GetViewOrigin() - InPrevViewMatrices.GetViewOrigin());
 	ViewUniformShaderParameters.CullingSign = bReverseCulling ? -1.0f : 1.0f;
 	ViewUniformShaderParameters.NearPlane = InViewMatrices.ComputeNearPlane();
+	ViewUniformShaderParameters.OrthoFarPlane = InViewMatrices.ComputeOrthoFarPlane();
 	ViewUniformShaderParameters.MaterialTextureMipBias = 0.0f;
 	ViewUniformShaderParameters.MaterialTextureDerivativeMultiply = 1.0f;
 	ViewUniformShaderParameters.ResolutionFractionAndInv = FVector2f(1.0f, 1.0f);
diff --git a/Engine/Source/Runtime/Engine/Public/SceneView.h b/Engine/Source/Runtime/Engine/Public/SceneView.h
--- a/Engine/Source/Runtime/Engine/Public/SceneView.h
+++ b/Engine/Source/Runtime/Engine/Public/SceneView.h
@@ -646,9 +646,13 @@ public:
 		return ( ProjectionMatrix.M[3][3] - ProjectionMatrix.M[3][2] ) / ( ProjectionMatrix.M[2][2] - ProjectionMatrix.M[2][3] );
 	}
 
-	FMatrix::FReal ComputeFarPlane() const
+	FMatrix::FReal ComputeOrthoFarPlane() const
 	{
-		return ComputeNearPlane() - InvProjectionMatrix.M[2][2];
+		if(!IsPerspectiveProjection())
+		{
+			return ComputeNearPlane() - InvProjectionMatrix.M[2][2];
+		}
+		return 0.0;
 	}
 
 	void ApplyWorldOffset(const FVector& InOffset)
@@ -848,6 +852,7 @@ enum ETranslucencyVolumeCascade
 	VIEW_UNIFORM_BUFFER_MEMBER_PER_VIEW(FVector3f, WorldCameraMovementSinceLastFrame) \
 	VIEW_UNIFORM_BUFFER_MEMBER(float, CullingSign) \
 	VIEW_UNIFORM_BUFFER_MEMBER_PER_VIEW_EX(float, NearPlane, EShaderPrecisionModifier::Half) \
+	VIEW_UNIFORM_BUFFER_MEMBER_PER_VIEW(float, OrthoFarPlane) \
 	VIEW_UNIFORM_BUFFER_MEMBER(float, GameTime) \
 	VIEW_UNIFORM_BUFFER_MEMBER(float, RealTime) \
 	VIEW_UNIFORM_BUFFER_MEMBER(float, DeltaTime) \
diff --git a/Engine/Source/Runtime/Renderer/Private/DistanceFieldAmbientOcclusion.cpp b/Engine/Source/Runtime/Renderer/Private/DistanceFieldAmbientOcclusion.cpp
--- a/Engine/Source/Runtime/Renderer/Private/DistanceFieldAmbientOcclusion.cpp
+++ b/Engine/Source/Runtime/Renderer/Private/DistanceFieldAmbientOcclusion.cpp
@@ -813,8 +813,7 @@ bool FDeferredShadingSceneRenderer::ShouldRenderDistanceFieldLighting() const
 	{
 		const FViewInfo& View = Views[ViewIndex];
 
-		if (!SupportsDistanceFieldAO(View.GetFeatureLevel(), View.GetShaderPlatform())
-			|| !View.IsPerspectiveProjection())
+		if (!SupportsDistanceFieldAO(View.GetFeatureLevel(), View.GetShaderPlatform()))
 		{
 			bSupportsDistanceFieldAO = false;
 			break;
diff --git a/Engine/Source/Runtime/Renderer/Private/FogRendering.cpp b/Engine/Source/Runtime/Renderer/Private/FogRendering.cpp
--- a/Engine/Source/Runtime/Renderer/Private/FogRendering.cpp
+++ b/Engine/Source/Runtime/Renderer/Private/FogRendering.cpp
@@ -43,6 +43,12 @@ static TAutoConsoleVariable<bool> CVarUnderwaterFogWhenCameraIsAboveWater(
 	TEXT("Renders height fog behind the water surface even when the camera is above water. This avoids artifacts when entering and exiting the water with strong height fog in the scene but causes artifacts when looking at the water surface from a distance."),
 	ECVF_RenderThreadSafe);
 
+static TAutoConsoleVariable<bool> CVarOrthoFogHeightAdjustment(
+	TEXT("r.Ortho.FogHeightAdjustment"),
+	true,
+	TEXT("When enabled, uses the Ortho camera height to determine the fog cutoff height"),
+	ECVF_RenderThreadSafe);
+
 IMPLEMENT_GLOBAL_SHADER_PARAMETER_STRUCT(FFogUniformParameters, "FogStruct");
 
 void SetupFogUniformParameters(FRDGBuilder& GraphBuilder, const FViewInfo& View, FFogUniformParameters& OutParameters)
@@ -192,8 +198,15 @@ void FSceneRenderer::InitFogConstants()
 			{
 				const FExponentialHeightFogSceneInfo& FogInfo = Scene->ExponentialFogs[0];
 				float CollapsedFogParameter[FExponentialHeightFogSceneInfo::NumFogs];
-				static constexpr float MaxObserverHeightDifference = 65536.0f;
+				float MaxObserverHeightDifference = 65536.0f;
 				float MaxObserverHeight = FLT_MAX;
+				
+				if (!View.IsPerspectiveProjection() && CVarOrthoFogHeightAdjustment.GetValueOnRenderThread())
+				{
+					FVector ViewForward = View.ViewMatrices.GetViewMatrix().GetColumn(2);
+					MaxObserverHeightDifference = FMath::Max(View.ViewMatrices.GetViewOrigin().Z, View.ViewMatrices.GetCameraToViewTarget().Length() * FMath::Abs(ViewForward.Z));
+				}
+
 				for (int i = 0; i < FExponentialHeightFogSceneInfo::NumFogs; i++)
 				{
 					// Only limit the observer height to fog if it has any density
@@ -416,46 +429,42 @@ void RenderFogOnClouds(
 	{
 		RDG_EVENT_SCOPE(GraphBuilder, "ExponentialHeightFog on Clouds");
 		RDG_GPU_STAT_SCOPE(GraphBuilder, Fog);
+		RDG_GPU_MASK_SCOPE(GraphBuilder, View.GPUMask);
 
-		if (View.IsPerspectiveProjection())
-		{
-			RDG_GPU_MASK_SCOPE(GraphBuilder, View.GPUMask);
-
-			TRDGUniformBufferRef<FFogUniformParameters> FogUniformBuffer = CreateFogUniformBuffer(GraphBuilder, View);
+		TRDGUniformBufferRef<FFogUniformParameters> FogUniformBuffer = CreateFogUniformBuffer(GraphBuilder, View);
 
-			// Light shaft is not accounted for in this case
-			const FScreenPassTextureViewportParameters LightShaftParameters;
-			// Local fog volume are not accounted for in this case
-			const bool bFogComposeLocalFogVolumes = false;
+		// Light shaft is not accounted for in this case
+		const FScreenPassTextureViewportParameters LightShaftParameters;
+		// Local fog volume are not accounted for in this case
+		const bool bFogComposeLocalFogVolumes = false;
 
-			FFogPassParameters* PassParameters = CreateDefaultFogPassParameters(
-				GraphBuilder, View, 
-				CreateSceneTextureUniformBuffer(GraphBuilder,View, ESceneTextureSetupMode::None),
-				FogUniformBuffer, nullptr /*LightShaftOcclusionTexture*/, LightShaftParameters);
+		FFogPassParameters* PassParameters = CreateDefaultFogPassParameters(
+			GraphBuilder, View, 
+			CreateSceneTextureUniformBuffer(GraphBuilder,View, ESceneTextureSetupMode::None),
+			FogUniformBuffer, nullptr /*LightShaftOcclusionTexture*/, LightShaftParameters);
 
-			// Patch the pass parameter for it to work on clouds
-			PassParameters->VS.ViewUniformBuffer = GetShaderBinding(bUseVolumetricRenderTarget ? View.VolumetricRenderTargetViewUniformBuffer : View.ViewUniformBuffer);
-			PassParameters->PS.ViewUniformBuffer = GetShaderBinding(bUseVolumetricRenderTarget ? View.VolumetricRenderTargetViewUniformBuffer : View.ViewUniformBuffer);
-			PassParameters->PS.bOnlyOnRenderedOpaque = false;
+		// Patch the pass parameter for it to work on clouds
+		PassParameters->VS.ViewUniformBuffer = GetShaderBinding(bUseVolumetricRenderTarget ? View.VolumetricRenderTargetViewUniformBuffer : View.ViewUniformBuffer);
+		PassParameters->PS.ViewUniformBuffer = GetShaderBinding(bUseVolumetricRenderTarget ? View.VolumetricRenderTargetViewUniformBuffer : View.ViewUniformBuffer);
+		PassParameters->PS.bOnlyOnRenderedOpaque = false;
 
-			PassParameters->PS.SrcCloudDepthTexture = SrcCloudDepth;
-			PassParameters->PS.SrcCloudViewTexture = SrcCloudView;
+		PassParameters->PS.SrcCloudDepthTexture = SrcCloudDepth;
+		PassParameters->PS.SrcCloudViewTexture = SrcCloudView;
 
-			PassParameters->RenderTargets[0] = FRenderTargetBinding(DstCloudView, ERenderTargetLoadAction::ENoAction);
-			// No depth target
+		PassParameters->RenderTargets[0] = FRenderTargetBinding(DstCloudView, ERenderTargetLoadAction::ENoAction);
+		// No depth target
 
-			// We enable the blending when volumetric render target is not enabled. Because in this case, the fog pass is compositing directly over the scene.
-			const bool bEnableBlending = !bUseVolumetricRenderTarget;
+		// We enable the blending when volumetric render target is not enabled. Because in this case, the fog pass is compositing directly over the scene.
+		const bool bEnableBlending = !bUseVolumetricRenderTarget;
 
-			FIntRect ViewRect(0, 0, SrcCloudView->Desc.Extent.X, SrcCloudView->Desc.Extent.Y);
-			GraphBuilder.AddPass(RDG_EVENT_NAME("Fog"), PassParameters, ERDGPassFlags::Raster,
-				[&View, ViewRect, PassParameters, bShouldRenderVolumetricFog, bFogComposeLocalFogVolumes, bEnableBlending](FRHICommandList& RHICmdList)
-				{
-					const bool bSampleFogOnClouds = true;
-					RenderViewFog(RHICmdList, View, ViewRect, PassParameters, bShouldRenderVolumetricFog, bFogComposeLocalFogVolumes, bSampleFogOnClouds, bEnableBlending);
-				});
-		}
-	}
+		FIntRect ViewRect(0, 0, SrcCloudView->Desc.Extent.X, SrcCloudView->Desc.Extent.Y);
+		GraphBuilder.AddPass(RDG_EVENT_NAME("Fog"), PassParameters, ERDGPassFlags::Raster,
+			[&View, ViewRect, PassParameters, bShouldRenderVolumetricFog, bFogComposeLocalFogVolumes, bEnableBlending](FRHICommandList& RHICmdList)
+			{
+				const bool bSampleFogOnClouds = true;
+				RenderViewFog(RHICmdList, View, ViewRect, PassParameters, bShouldRenderVolumetricFog, bFogComposeLocalFogVolumes, bSampleFogOnClouds, bEnableBlending);
+			});
+}
 }
 
 void FDeferredShadingSceneRenderer::RenderFog(
@@ -476,27 +485,24 @@ void FDeferredShadingSceneRenderer::RenderFog(
 		for(int32 ViewIndex = 0;ViewIndex < Views.Num();ViewIndex++)
 		{
 			const FViewInfo& View = Views[ViewIndex];
-			if (View.IsPerspectiveProjection())
-			{
-				RDG_EVENT_SCOPE_CONDITIONAL(GraphBuilder, Views.Num() > 1, "View%d", ViewIndex);
-				RDG_GPU_MASK_SCOPE(GraphBuilder, View.GPUMask);
+			RDG_EVENT_SCOPE_CONDITIONAL(GraphBuilder, Views.Num() > 1, "View%d", ViewIndex);
+			RDG_GPU_MASK_SCOPE(GraphBuilder, View.GPUMask);
 
-				TRDGUniformBufferRef<FFogUniformParameters> FogUniformBuffer = CreateFogUniformBuffer(GraphBuilder, View);
+			TRDGUniformBufferRef<FFogUniformParameters> FogUniformBuffer = CreateFogUniformBuffer(GraphBuilder, View);
 
-				const FScreenPassTextureViewport SceneViewport(SceneTextures.Config.Extent, View.ViewRect);
-				const FScreenPassTextureViewport OutputViewport(GetDownscaledViewport(SceneViewport, GetLightShaftDownsampleFactor()));
-				const FScreenPassTextureViewportParameters LightShaftParameters = GetScreenPassTextureViewportParameters(OutputViewport);
+			const FScreenPassTextureViewport SceneViewport(SceneTextures.Config.Extent, View.ViewRect);
+			const FScreenPassTextureViewport OutputViewport(GetDownscaledViewport(SceneViewport, GetLightShaftDownsampleFactor()));
+			const FScreenPassTextureViewportParameters LightShaftParameters = GetScreenPassTextureViewportParameters(OutputViewport);
 
-				FFogPassParameters* PassParameters = CreateDefaultFogPassParameters(GraphBuilder, View, SceneTextures.UniformBuffer, FogUniformBuffer, LightShaftOcclusionTexture, LightShaftParameters);
-				PassParameters->RenderTargets[0] = FRenderTargetBinding(SceneTextures.Color.Target, ERenderTargetLoadAction::ELoad);
-				PassParameters->RenderTargets.DepthStencil = FDepthStencilBinding(SceneTextures.Depth.Target, ERenderTargetLoadAction::ELoad, ERenderTargetLoadAction::ELoad, FExclusiveDepthStencil::DepthRead_StencilWrite);
+			FFogPassParameters* PassParameters = CreateDefaultFogPassParameters(GraphBuilder, View, SceneTextures.UniformBuffer, FogUniformBuffer, LightShaftOcclusionTexture, LightShaftParameters);
+			PassParameters->RenderTargets[0] = FRenderTargetBinding(SceneTextures.Color.Target, ERenderTargetLoadAction::ELoad);
+			PassParameters->RenderTargets.DepthStencil = FDepthStencilBinding(SceneTextures.Depth.Target, ERenderTargetLoadAction::ELoad, ERenderTargetLoadAction::ELoad, FExclusiveDepthStencil::DepthRead_StencilWrite);
 
-				GraphBuilder.AddPass(RDG_EVENT_NAME("Fog"), PassParameters, ERDGPassFlags::Raster, 
-					[this, &View, PassParameters, bShouldRenderVolumetricFog, bFogComposeLocalFogVolumes](FRHICommandList& RHICmdList)
+			GraphBuilder.AddPass(RDG_EVENT_NAME("Fog"), PassParameters, ERDGPassFlags::Raster,
+				[this, &View, PassParameters, bShouldRenderVolumetricFog, bFogComposeLocalFogVolumes](FRHICommandList& RHICmdList)
 				{
 					RenderViewFog(RHICmdList, View, View.ViewRect, PassParameters, bShouldRenderVolumetricFog, bFogComposeLocalFogVolumes);
 				});
-			}
 		}
 	}
 }
@@ -521,7 +527,7 @@ void FDeferredShadingSceneRenderer::RenderUnderWaterFog(
 		for (int32 ViewIndex = 0; ViewIndex < Views.Num(); ViewIndex++)
 		{
 			const FViewInfo& View = Views[ViewIndex];
-			if (View.IsPerspectiveProjection() && (View.IsUnderwater() || CVarUnderwaterFogWhenCameraIsAboveWater.GetValueOnRenderThread()))
+			if (View.IsUnderwater() || CVarUnderwaterFogWhenCameraIsAboveWater.GetValueOnRenderThread())
 			{
 				RDG_EVENT_SCOPE_CONDITIONAL(GraphBuilder, Views.Num() > 1, "View%d", ViewIndex);
 				RDG_GPU_MASK_SCOPE(GraphBuilder, View.GPUMask);
diff --git a/Engine/Source/Runtime/Renderer/Private/PlanarReflectionRendering.cpp b/Engine/Source/Runtime/Renderer/Private/PlanarReflectionRendering.cpp
--- a/Engine/Source/Runtime/Renderer/Private/PlanarReflectionRendering.cpp
+++ b/Engine/Source/Runtime/Renderer/Private/PlanarReflectionRendering.cpp
@@ -600,16 +600,16 @@ void FScene::UpdatePlanarReflectionContents(UPlanarReflectionComponent* CaptureC
 			const FMatrix ViewMatrix(MirrorMatrix * View.ViewMatrices.GetViewMatrix());
 			const FVector ViewOrigin = ViewMatrix.InverseTransformPosition(FVector::ZeroVector);
 			const FMatrix ViewRotationMatrix = ViewMatrix.RemoveTranslation();
-			const float HalfFOV = FMath::Atan(1.0f / View.ViewMatrices.GetProjectionMatrix().M[0][0]);
 
 			FMatrix ProjectionMatrix;
-			if (CaptureComponent->ExtraFOV == 0.f && MainSceneRenderer.Views.Num() > 1)
+			if (!View.IsPerspectiveProjection() || (CaptureComponent->ExtraFOV == 0.f && MainSceneRenderer.Views.Num() > 1))
 			{
 				// Prefer exact (potentially uneven) stereo projection matrices when no extra FOV is requested
 				ProjectionMatrix = View.ViewMatrices.GetProjectionMatrix();
 			}
 			else
 			{
+				const float HalfFOV = FMath::Atan(1.0f / View.ViewMatrices.GetProjectionMatrix().M[0][0]);
 				BuildProjectionMatrix(View.UnscaledViewRect.Size(), HalfFOV + FMath::DegreesToRadians(CaptureComponent->ExtraFOV), GNearClippingPlane, ProjectionMatrix);
 			}
 
diff --git a/Engine/Source/Runtime/Renderer/Private/PostProcess/PostProcessMotionBlur.cpp b/Engine/Source/Runtime/Renderer/Private/PostProcess/PostProcessMotionBlur.cpp
--- a/Engine/Source/Runtime/Renderer/Private/PostProcess/PostProcessMotionBlur.cpp
+++ b/Engine/Source/Runtime/Renderer/Private/PostProcess/PostProcessMotionBlur.cpp
@@ -76,6 +76,13 @@ namespace
 		ECVF_RenderThreadSafe
 	);
 
+	TAutoConsoleVariable<int32> CVarOrthoUsePreviousMotionVelocityFlattenPass(
+		TEXT("r.Ortho.UsePreviousMotionVelocityFlattenPass"),
+		0,
+		TEXT("Enable/Disable the option for Ortho motion blur pass to use the previous flatten textures (e.g. from TSR). ")
+		TEXT("Currently causes shimmering on distance planes if disabled"),
+		ECVF_RenderThreadSafe);
+
 	FMatrix GetPreviousWorldToClipMatrix(const FViewInfo& View)
 	{
 		if (View.Family->EngineShowFlags.CameraInterpolation)
@@ -619,7 +626,9 @@ void AddMotionBlurVelocityPass(
 	// Velocity flatten pass: combines depth / velocity into a single target for sampling efficiency.
 	FRDGTextureRef VelocityFlatTexture = nullptr;
 	FRDGTextureRef VelocityTileTextureSetup = nullptr;
-	if (Inputs.VelocityFlattenTextures.IsValid())
+
+	bool bViewUsesPreviousFlattenTexture = View.IsPerspectiveProjection() || CVarOrthoUsePreviousMotionVelocityFlattenPass.GetValueOnRenderThread() != 0;
+	if (bViewUsesPreviousFlattenTexture && Inputs.VelocityFlattenTextures.IsValid())
 	{
 		ensure(Inputs.VelocityFlattenTextures.VelocityFlatten.ViewRect == View.ViewRect);
 		ensure(Inputs.VelocityFlattenTextures.VelocityTileArray.ViewRect == FIntRect(FIntPoint::ZeroValue, VelocityTileCount));
diff --git a/Engine/Source/Runtime/Renderer/Private/SceneVisibility.cpp b/Engine/Source/Runtime/Renderer/Private/SceneVisibility.cpp
--- a/Engine/Source/Runtime/Renderer/Private/SceneVisibility.cpp
+++ b/Engine/Source/Runtime/Renderer/Private/SceneVisibility.cpp
@@ -4940,7 +4940,7 @@ void FSceneRenderer::PrepareViewStateForVisibility(const FSceneTexturesConfig& S
 			}
 
 			// Choose sub pixel sample coordinate in the temporal sequence.
-			float SampleX, SampleY;
+			float SampleX = 0.0f, SampleY = 0.0f;
 			if (View.PrimaryScreenPercentageMethod == EPrimaryScreenPercentageMethod::TemporalUpscale)
 			{
 				// Uniformly distribute temporal jittering in [-.5; .5], because there is no longer any alignement of input and output pixels.
@@ -5004,7 +5004,7 @@ void FSceneRenderer::PrepareViewStateForVisibility(const FSceneTexturesConfig& S
 				SampleX = SamplesX[ TemporalSampleIndex ];
 				SampleY = SamplesY[ TemporalSampleIndex ];
 			}
-			else
+			else if(View.IsPerspectiveProjection())
 			{
 				float u1 = Halton( TemporalSampleIndex + 1, 2 );
 				float u2 = Halton( TemporalSampleIndex + 1, 3 );
diff --git a/Engine/Source/Runtime/Renderer/Private/SkyAtmosphereRendering.cpp b/Engine/Source/Runtime/Renderer/Private/SkyAtmosphereRendering.cpp
--- a/Engine/Source/Runtime/Renderer/Private/SkyAtmosphereRendering.cpp
+++ b/Engine/Source/Runtime/Renderer/Private/SkyAtmosphereRendering.cpp
@@ -1888,9 +1888,17 @@ void FSceneRenderer::RenderSkyAtmosphereInternal(
 		if (SkyRC.bFastAerialPerspectiveDepthTest)
 		{
 			const FMatrix ProjectionMatrix = ViewMatrices.GetProjectionMatrix();
-			float HalfHorizontalFOV = FMath::Atan(1.0f / ProjectionMatrix.M[0][0]);
-			float HalfVerticalFOV = FMath::Atan(1.0f / ProjectionMatrix.M[1][1]);
-			float StartDepthViewCm = FMath::Cos(FMath::Max(HalfHorizontalFOV, HalfVerticalFOV)) * AerialPerspectiveStartDepthInCm;
+			float StartDepthViewCm;
+			if (ViewMatrices.IsPerspectiveProjection())
+			{
+				float HalfHorizontalFOV = FMath::Atan(1.0f / ProjectionMatrix.M[0][0]);
+				float HalfVerticalFOV = FMath::Atan(1.0f / ProjectionMatrix.M[1][1]);
+				StartDepthViewCm = FMath::Cos(FMath::Max(HalfHorizontalFOV, HalfVerticalFOV)) * AerialPerspectiveStartDepthInCm;
+			}
+			else
+			{
+				StartDepthViewCm = AerialPerspectiveStartDepthInCm;
+			}
 			StartDepthViewCm = FMath::Max(StartDepthViewCm, SkyRC.NearClippingDistance); // In any case, we need to limit the distance to frustum near plane to not be clipped away.
 			const FVector4 Projected = ProjectionMatrix.TransformFVector4(FVector4(0.0f, 0.0f, StartDepthViewCm, 1.0f));
 			StartDepthZ = float(Projected.Z / Projected.W); // LWC_TODO: precision loss
diff --git a/Engine/Source/Runtime/Renderer/Private/VolumetricRenderTarget.cpp b/Engine/Source/Runtime/Renderer/Private/VolumetricRenderTarget.cpp
--- a/Engine/Source/Runtime/Renderer/Private/VolumetricRenderTarget.cpp
+++ b/Engine/Source/Runtime/Renderer/Private/VolumetricRenderTarget.cpp
@@ -69,7 +69,7 @@ static bool ShouldPipelineCompileVolumetricRenderTargetShaders(EShaderPlatform S
 bool ShouldViewRenderVolumetricCloudRenderTarget(const FViewInfo& ViewInfo)
 {
 	return CVarVolumetricRenderTarget.GetValueOnRenderThread() && ShouldPipelineCompileVolumetricRenderTargetShaders(ViewInfo.GetShaderPlatform())
-		&& (ViewInfo.ViewState != nullptr) && !ViewInfo.bIsReflectionCapture;
+		&& (ViewInfo.ViewState != nullptr) && !ViewInfo.bIsReflectionCapture && ViewInfo.IsPerspectiveProjection(); //Do not use for ortho as the resolution resolves do not blend well when depth is uniform anyway.
 }
 
 bool IsVolumetricRenderTargetEnabled()
