Path Tracer: Workaround assert from cvar lookup in background thread

Use a global variable instead of querying the cvar.

diff --git a/Engine/Source/Runtime/Renderer/Private/PathTracing.cpp b/Engine/Source/Runtime/Renderer/Private/PathTracing.cpp
--- a/Engine/Source/Runtime/Renderer/Private/PathTracing.cpp
+++ b/Engine/Source/Runtime/Renderer/Private/PathTracing.cpp
@@ -194,9 +194,10 @@ TAutoConsoleVariable<int32> CVarPathTracingEnableCameraBackfaceCulling(
 	ECVF_RenderThreadSafe
 );
 
-TAutoConsoleVariable<int32> CVarPathTracingEnableReferenceDOF(
+static int32 GEnableReferenceDOF = -1;
+FAutoConsoleVariableRef CVarPathTracingEnableReferenceDOF(
 	TEXT("r.PathTracing.EnableReferenceDOF"),
-	-1,
+	GEnableReferenceDOF,
 	TEXT("Should the path tracer ray trace the depth-of-field effect instead of the post-processed effect?\n")
 	TEXT("-1: Inherit from PostProcess settings (default)\n")
 	TEXT(" 0: Disabled\n")
@@ -697,8 +698,7 @@ namespace PathTracing
 
 	bool UsesReferenceDOF(const FViewInfo& View)
 	{
-		const int32 EnableReferenceDOFCvar = CVarPathTracingEnableReferenceDOF.GetValueOnRenderThread();
-		return EnableReferenceDOFCvar < 0 ? View.FinalPostProcessSettings.PathTracingEnableReferenceDOF != 0 : EnableReferenceDOFCvar != 0;
+		return GEnableReferenceDOF < 0 ? View.FinalPostProcessSettings.PathTracingEnableReferenceDOF != 0 : GEnableReferenceDOF != 0;
 	}
 }
 
