Fixed debug probe not showing up when enabled.

diff --git a/Engine/Source/Runtime/Renderer/Private/DebugProbeRendering.cpp b/Engine/Source/Runtime/Renderer/Private/DebugProbeRendering.cpp
--- a/Engine/Source/Runtime/Renderer/Private/DebugProbeRendering.cpp
+++ b/Engine/Source/Runtime/Renderer/Private/DebugProbeRendering.cpp
@@ -127,7 +127,7 @@ void StampDeferredDebugProbeDepthPS(
 	const FRDGTextureRef SceneDepthTexture)
 {
 #if DEBUG_PROBE_ENABLED
-	if (FPlatformProperties::HasEditorOnlyData())
+	if (!FPlatformProperties::HasEditorOnlyData())
 	{
 		return;
 	}
@@ -161,7 +161,7 @@ void StampDeferredDebugProbeMaterialPS(
 	const FMinimalSceneTextures& SceneTextures)
 {
 #if DEBUG_PROBE_ENABLED
-	if (FPlatformProperties::HasEditorOnlyData())
+	if (!FPlatformProperties::HasEditorOnlyData())
 	{
 		return;
 	}
@@ -226,7 +226,7 @@ void StampDeferredDebugProbeVelocityPS(
 	const FRenderTargetBindingSlots& BasePassRenderTargets)
 {
 #if DEBUG_PROBE_ENABLED
-	if (FPlatformProperties::HasEditorOnlyData())
+	if (!FPlatformProperties::HasEditorOnlyData())
 	{
 		return;
 	}
