Add back in a missed part of previous Volumetric change for auto-adjusting fog height in Orthographic view modes.

diff --git a/Engine/Source/Runtime/Renderer/Private/FogRendering.cpp b/Engine/Source/Runtime/Renderer/Private/FogRendering.cpp
--- a/Engine/Source/Runtime/Renderer/Private/FogRendering.cpp
+++ b/Engine/Source/Runtime/Renderer/Private/FogRendering.cpp
@@ -200,8 +200,9 @@ void FSceneRenderer::InitFogConstants()
 				float CollapsedFogParameter[FExponentialHeightFogSceneInfo::NumFogs];
 				float MaxObserverHeightDifference = 65536.0f;
 				float MaxObserverHeight = FLT_MAX;
-				
-				if (!View.IsPerspectiveProjection() && CVarOrthoFogHeightAdjustment.GetValueOnRenderThread())
+
+				const bool bUsingOrthoHeightAdjustment = !View.IsPerspectiveProjection() && CVarOrthoFogHeightAdjustment.GetValueOnRenderThread();				
+				if (bUsingOrthoHeightAdjustment)
 				{
 					FVector ViewForward = View.ViewMatrices.GetViewMatrix().GetColumn(2);
 					MaxObserverHeightDifference = FMath::Max(View.ViewMatrices.GetViewOrigin().Z, View.ViewMatrices.GetCameraToViewTarget().Length() * FMath::Abs(ViewForward.Z));
@@ -217,7 +218,7 @@ void FSceneRenderer::InitFogConstants()
 				}
 				
 				// Clamping the observer height to avoid numerical precision issues in the height fog equation. The max observer height is relative to the fog height.
-				const float ObserverHeight = FMath::Min<float>(View.ViewMatrices.GetViewOrigin().Z, MaxObserverHeight);
+				const float ObserverHeight = bUsingOrthoHeightAdjustment ? MaxObserverHeight : FMath::Min<float>(View.ViewMatrices.GetViewOrigin().Z, MaxObserverHeight);
 
 				for (int i = 0; i < FExponentialHeightFogSceneInfo::NumFogs; i++)
 				{
