Fixed MGPU scope for RayTracing build and improved clarity of comments.

diff --git a/Engine/Source/Runtime/Renderer/Private/DeferredShadingRenderer.cpp b/Engine/Source/Runtime/Renderer/Private/DeferredShadingRenderer.cpp
--- a/Engine/Source/Runtime/Renderer/Private/DeferredShadingRenderer.cpp
+++ b/Engine/Source/Runtime/Renderer/Private/DeferredShadingRenderer.cpp
@@ -862,12 +862,12 @@ bool FDeferredShadingSceneRenderer::DispatchRayTracingWorldUpdates(FRDGBuilder&
 				Scene->GetRayTracingDynamicGeometryCollection()->DispatchUpdates(RHICmdList, DynamicGeometryScratchBuffer);
 				Scene->GetRayTracingDynamicGeometryCollection()->EndUpdate();
 
-				// Force this pass to run on the RHI thread.
+				// Disables parallel translate and forces the work to happen on the RHI thread.
 				RHICmdList.RHIThreadFence(true);
 			});
-
-			RayTracingScene.Build(GraphBuilder, ComputePassFlags | ERDGPassFlags::NeverCull, OutDynamicGeometryScratchBuffer);
 		}
+
+		RayTracingScene.Build(GraphBuilder, ComputePassFlags | ERDGPassFlags::NeverCull, OutDynamicGeometryScratchBuffer);
 	}
 
 	GraphBuilder.AddDispatchHint();
diff --git a/Engine/Source/Runtime/Renderer/Private/RayTracing/RayTracingScene.cpp b/Engine/Source/Runtime/Renderer/Private/RayTracing/RayTracingScene.cpp
--- a/Engine/Source/Runtime/Renderer/Private/RayTracing/RayTracingScene.cpp
+++ b/Engine/Source/Runtime/Renderer/Private/RayTracing/RayTracingScene.cpp
@@ -289,7 +289,7 @@ void FRayTracingScene::CreateWithInitializationData(FRDGBuilder& GraphBuilder, c
 					CullingParameters.bUseInstanceCulling ? &CullingParameters : nullptr,
 					PassParams->DebugInstanceGPUSceneIndexBuffer ? PassParams->DebugInstanceGPUSceneIndexBuffer->GetRHI() : nullptr);
 
-				// Force this pass to run on the RHI thread.
+				// Disables parallel translate and forces the work to happen on the RHI thread.
 				RHICmdList.RHIThreadFence(true);
 			});
 	}
@@ -326,7 +326,7 @@ void FRayTracingScene::Build(FRDGBuilder& GraphBuilder, ERDGPassFlags ComputePas
 			RHICmdList.BindAccelerationStructureMemory(RayTracingSceneRHI, PassParams->TLASBuffer->GetRHI(), 0);
 			RHICmdList.BuildAccelerationStructure(BuildParams);
 
-			// Force this pass to run on the RHI thread.
+			// Disables parallel translate and forces the work to happen on the RHI thread.
 			RHICmdList.RHIThreadFence(true);
 		});
 }
