HWRT: Don't run ray tracing primitive visiblity tests if ray tracing is disabled

#tests Run CitySample with HWRT on/off

diff --git a/Engine/Source/Runtime/Renderer/Private/SceneVisibility.cpp b/Engine/Source/Runtime/Renderer/Private/SceneVisibility.cpp
--- a/Engine/Source/Runtime/Renderer/Private/SceneVisibility.cpp
+++ b/Engine/Source/Runtime/Renderer/Private/SceneVisibility.cpp
@@ -643,6 +643,8 @@ static void UpdateAlwaysVisible(const FScene& Scene, FViewInfo& View, FFrustumCu
 	uint32* RESTRICT VisWords = View.PrimitiveVisibilityMap.GetData();
 #if RHI_RAYTRACING
 	uint32* RESTRICT  RTWords = View.PrimitiveRayTracingVisibilityMap.GetData();
+
+	const bool bRayTracingEnabled = IsRayTracingEnabled(View.GetShaderPlatform());
 #endif
 
 	for (int32 WordIndex = TaskWordOffset; WordIndex < TaskWordOffset + int32(TaskConfig.AlwaysVisible.NumWordsPerTask) && WordIndex * NumBitsPerDWORD < BitArrayNumInner; ++WordIndex)
@@ -660,7 +662,7 @@ static void UpdateAlwaysVisible(const FScene& Scene, FViewInfo& View, FFrustumCu
 			VisBits |= Mask;
 
 		#if RHI_RAYTRACING
-			if (!IsPrimitiveHidden(Scene, View, Index, Flags) && !ShouldCullForRayTracing(Scene, View, Index))
+			if (bRayTracingEnabled && !IsPrimitiveHidden(Scene, View, Index, Flags) && !ShouldCullForRayTracing(Scene, View, Index))
 			{
 				RayTracingBits |= Mask;
 			}
@@ -707,6 +709,8 @@ static int32 FrustumCull(const FScene& Scene, FViewInfo& View, FFrustumCullingFl
 	uint32* RESTRICT FadeWords = View.PotentiallyFadingPrimitiveMap.GetData();
 #if RHI_RAYTRACING
 	uint32* RESTRICT  RTWords = View.PrimitiveRayTracingVisibilityMap.GetData();
+
+	const bool bRayTracingEnabled = IsRayTracingEnabled(View.GetShaderPlatform());
 #endif
 
 	for (int32 WordIndex = TaskWordOffset; WordIndex < TaskWordOffset + int32(TaskConfig.FrustumCull.NumWordsPerTask) && WordIndex * NumBitsPerDWORD < BitArrayNumInner; WordIndex++)
@@ -735,7 +739,7 @@ static int32 FrustumCull(const FScene& Scene, FViewInfo& View, FFrustumCullingFl
 		#if RHI_RAYTRACING
 			bool bIsVisibleInRayTracing = true;
 
-			if (bPrimitiveIsHidden || ShouldCullForRayTracing(Scene, View, Index))
+			if (bPrimitiveIsHidden || !bRayTracingEnabled || ShouldCullForRayTracing(Scene, View, Index))
 			{
 				bIsVisibleInRayTracing = false;
 			}
@@ -986,6 +990,10 @@ static void UpdatePrimitiveFading(const FScene& Scene, FViewInfo& View, FSceneVi
 
 	SCOPE_CYCLE_COUNTER(STAT_UpdatePrimitiveFading);
 
+#if RHI_RAYTRACING
+	const bool bRayTracingEnabled = IsRayTracingEnabled(View.GetShaderPlatform());
+#endif 
+
 	// Should we allow fading transitions at all this frame?  For frames where the camera moved
 	// a large distance or where we haven't rendered a view in awhile, it's best to disable
 	// fading so users don't see unexpected object transitions.
@@ -1008,7 +1016,7 @@ static void UpdatePrimitiveFading(const FScene& Scene, FViewInfo& View, FSceneVi
 				// This should be a very rare occurrence, so the hit is not worrisome.
 				// TODO:  Could this be moved into the actual culling phase?
 
-				if (!ShouldCullForRayTracing(Scene, View, BitIt.GetIndex()))
+				if (bRayTracingEnabled && !ShouldCullForRayTracing(Scene, View, BitIt.GetIndex()))
 				{
 					View.PrimitiveRayTracingVisibilityMap.AccessCorrespondingBit(BitIt) = true;
 				}
