Fixed shadowing of TransformUpdates in Lambda with compiler dependent non-determinsm around move ordering.

diff --git a/Engine/Source/Runtime/Renderer/Private/RendererScene.cpp b/Engine/Source/Runtime/Renderer/Private/RendererScene.cpp
--- a/Engine/Source/Runtime/Renderer/Private/RendererScene.cpp
+++ b/Engine/Source/Runtime/Renderer/Private/RendererScene.cpp
@@ -2223,11 +2223,13 @@ void FScene::FinishUpdatePrimitiveTransform()
 {
 	if (bPrimitivesUpdateBatching)
 	{
+		const int32 NumPrimitiveUpdates = FMath::Min(PrimitivesUpdates.Num(), PrimitiveUpdateIndex.load(std::memory_order_relaxed));
+
 		// Pass the collection and actual number of accumulated updates
 		ENQUEUE_RENDER_COMMAND(UpdateTransformCommand)(
-			[PrimitivesUpdates = MoveTemp(PrimitivesUpdates), NumUpdates = FMath::Min(PrimitivesUpdates.Num(), PrimitiveUpdateIndex.load(std::memory_order_relaxed))](FRHICommandListBase&)
+			[PrimitivesUpdates = MoveTemp(PrimitivesUpdates), NumPrimitiveUpdates](FRHICommandListBase&)
 			{
-				for (int32 Index = 0; Index < NumUpdates; ++Index)
+				for (int32 Index = 0; Index < NumPrimitiveUpdates; ++Index)
 				{
 					const auto& UpdateParams = PrimitivesUpdates[Index];
 					FScopeCycleCounter Context(UpdateParams.PrimitiveSceneProxy->GetStatId());
@@ -2243,7 +2245,7 @@ void FScene::FinishUpdatePrimitiveTransform()
 			}
 		);
 
-		PrimitivesUpdates.Reset();
+		PrimitivesUpdates.Empty();
 		PrimitiveUpdateIndex = 0;
 		bPrimitivesUpdateBatching = false;
 	}
