Take into account that draw can be skipped when doing batched MDC submission

diff --git a/Engine/Source/Runtime/Renderer/Private/InstanceCulling/InstanceCullingContext.cpp b/Engine/Source/Runtime/Renderer/Private/InstanceCulling/InstanceCullingContext.cpp
--- a/Engine/Source/Runtime/Renderer/Private/InstanceCulling/InstanceCullingContext.cpp
+++ b/Engine/Source/Runtime/Renderer/Private/InstanceCulling/InstanceCullingContext.cpp
@@ -1712,14 +1712,15 @@ void FInstanceCullingContext::SubmitDrawCommands(
 			}
 			
 			SceneArgs.PrimitiveIdOffset = OverrideArgs.InstanceDataByteOffset + DrawCommandInfo.InstanceDataByteOffset;
-			FMeshDrawCommand::SubmitDraw(*VisibleMeshDrawCommand.MeshDrawCommand, GraphicsMinimalPipelineStateSet, SceneArgs, InstanceFactor, RHICmdList, StateCache);
-			
-			// If MDC was split to a more than one batch, submit them without changing state
-			for (uint32 BatchIdx = 1; BatchIdx < DrawCommandInfo.NumBatches; ++BatchIdx)
+			if (FMeshDrawCommand::SubmitDraw(*VisibleMeshDrawCommand.MeshDrawCommand, GraphicsMinimalPipelineStateSet, SceneArgs, InstanceFactor, RHICmdList, StateCache))
 			{
-				SceneArgs.PrimitiveIdOffset += DrawCommandInfo.BatchDataStride;
-				SceneArgs.IndirectArgsByteOffset += sizeof(FRHIDrawIndexedIndirectParameters);
-				FMeshDrawCommand::SubmitDrawEnd(*VisibleMeshDrawCommand.MeshDrawCommand, SceneArgs, InstanceFactor, RHICmdList);
+				// If MDC was split to a more than one batch, submit them without changing state
+				for (uint32 BatchIdx = 1; BatchIdx < DrawCommandInfo.NumBatches; ++BatchIdx)
+				{
+					SceneArgs.PrimitiveIdOffset += DrawCommandInfo.BatchDataStride;
+					SceneArgs.IndirectArgsByteOffset += sizeof(FRHIDrawIndexedIndirectParameters);
+					FMeshDrawCommand::SubmitDrawEnd(*VisibleMeshDrawCommand.MeshDrawCommand, SceneArgs, InstanceFactor, RHICmdList);
+				}
 			}
 		}
 	}
diff --git a/Engine/Source/Runtime/Renderer/Private/MeshPassProcessor.cpp b/Engine/Source/Runtime/Renderer/Private/MeshPassProcessor.cpp
--- a/Engine/Source/Runtime/Renderer/Private/MeshPassProcessor.cpp
+++ b/Engine/Source/Runtime/Renderer/Private/MeshPassProcessor.cpp
@@ -1570,7 +1570,7 @@ void FMeshDrawCommand::SubmitDrawIndirectEnd(
 	}
 }
 
-void FMeshDrawCommand::SubmitDraw(
+bool FMeshDrawCommand::SubmitDraw(
 	const FMeshDrawCommand& RESTRICT MeshDrawCommand,
 	const FGraphicsMinimalPipelineStateSet& GraphicsMinimalPipelineStateSet,
 	const FMeshDrawCommandSceneArgs& SceneArgs,
@@ -1585,7 +1585,10 @@ void FMeshDrawCommand::SubmitDraw(
 	if (SubmitDrawBegin(MeshDrawCommand, GraphicsMinimalPipelineStateSet, SceneArgs, InstanceFactor, RHICmdList, StateCache, bAllowSkipDrawCommand))
 	{
 		SubmitDrawEnd(MeshDrawCommand, SceneArgs, InstanceFactor, RHICmdList);
+		return true;
 	}
+
+	return false;
 }
 
 void ApplyTargetsInfo(FGraphicsPipelineStateInitializer& GraphicsPSOInit, const FGraphicsPipelineRenderTargetsInfo& RenderTargetsInfo)
diff --git a/Engine/Source/Runtime/Renderer/Public/MeshPassProcessor.h b/Engine/Source/Runtime/Renderer/Public/MeshPassProcessor.h
--- a/Engine/Source/Runtime/Renderer/Public/MeshPassProcessor.h
+++ b/Engine/Source/Runtime/Renderer/Public/MeshPassProcessor.h
@@ -1283,7 +1283,7 @@ public:
 	static void SubmitDrawIndirectEnd(const FMeshDrawCommand& MeshDrawCommand, const FMeshDrawCommandSceneArgs& SceneArgs, uint32 InstanceFactor, FRHICommandList& RHICmdList);
 
 	/** Submits commands to the RHI Commandlist to draw the MeshDrawCommand. */
-	static void SubmitDraw(
+	static bool SubmitDraw(
 		const FMeshDrawCommand& RESTRICT MeshDrawCommand,
 		const FGraphicsMinimalPipelineStateSet& GraphicsMinimalPipelineStateSet,
 		const FMeshDrawCommandSceneArgs& SceneArgs,
