Also disable motion blur visualize debug text to make comparing screenshots much easier.

- Rename r.MotionBlur.VisualizeCheckerboard to r.MotionBlur.VisualizeDebugInformation since it disables more than just the checkerboard now.
- Change the cvar to a bool since it is just a toggle.

diff --git a/Engine/Source/Runtime/Renderer/Private/PostProcess/PostProcessMotionBlur.cpp b/Engine/Source/Runtime/Renderer/Private/PostProcess/PostProcessMotionBlur.cpp
--- a/Engine/Source/Runtime/Renderer/Private/PostProcess/PostProcessMotionBlur.cpp
+++ b/Engine/Source/Runtime/Renderer/Private/PostProcess/PostProcessMotionBlur.cpp
@@ -67,12 +67,10 @@ namespace
 		0,
 		TEXT("For projects with motion blur enabled, this allows motion blur to be enabled even while in VR."));
 
-	TAutoConsoleVariable<int32> CVarVisualizeMotionBlurEnableCheckerboard(
-		TEXT("r.MotionBlur.VisualizeCheckerboard"),
-		1,
-		TEXT("Enables checkerboard visualization when Visualize Motion Blur show flag is enabled.\n")
-		TEXT("0: off\n")
-		TEXT("1: on (default)\n"),
+	TAutoConsoleVariable<bool> CVarVisualizeMotionBlurEnableDebugInformation(
+		TEXT("r.MotionBlur.VisualizeDebugInformation"),
+		true,
+		TEXT("Enables checkerboard and debug text visualization when Visualize Motion Blur show flag is enabled.\n"),
 		ECVF_RenderThreadSafe
 	);
 
@@ -1140,6 +1138,8 @@ FScreenPassTextureSlice AddVisualizeMotionBlurPass(FRDGBuilder& GraphBuilder, co
 	// NOTE: Scene depth is used as the velocity viewport because velocity can actually be a 1x1 black texture.
 	const FMotionBlurViewports Viewports(FScreenPassTextureViewport(Inputs.SceneColor), FScreenPassTextureViewport(Inputs.SceneDepth));
 
+	const bool bVisualizeDebugInfo = CVarVisualizeMotionBlurEnableDebugInformation.GetValueOnRenderThread();
+
 	FMotionBlurVisualizePS::FParameters* PassParameters = GraphBuilder.AllocParameters<FMotionBlurVisualizePS::FParameters>();
 	PassParameters->WorldToClipPrev = FMatrix44f(GetPreviousWorldToClipMatrix(View));		// LWC_TODO: Precision loss
 	PassParameters->View = View.ViewUniformBuffer;
@@ -1152,7 +1152,7 @@ FScreenPassTextureSlice AddVisualizeMotionBlurPass(FRDGBuilder& GraphBuilder, co
 	PassParameters->VelocitySampler = GetMotionBlurVelocitySampler();
 	PassParameters->DepthSampler = GetMotionBlurVelocitySampler();
 	PassParameters->RenderTargets[0] = Output.GetRenderTargetBinding();
-	PassParameters->CheckerboardEnabled = CVarVisualizeMotionBlurEnableCheckerboard.GetValueOnRenderThread();
+	PassParameters->CheckerboardEnabled = bVisualizeDebugInfo;
 
 	TShaderMapRef<FMotionBlurVisualizePS> PixelShader(View.ShaderMap);
 
@@ -1160,41 +1160,44 @@ FScreenPassTextureSlice AddVisualizeMotionBlurPass(FRDGBuilder& GraphBuilder, co
 
 	Output.LoadAction = ERenderTargetLoadAction::ELoad;
 
-	AddDrawCanvasPass(GraphBuilder, RDG_EVENT_NAME("VisualizeMotionBlurOverlay"), View, Output,
-		[&View](FCanvas& Canvas)
+	if (bVisualizeDebugInfo)
 	{
-		float X = 20;
-		float Y = 38;
-		const float YStep = 14;
-		const float ColumnWidth = 200;
+		AddDrawCanvasPass(GraphBuilder, RDG_EVENT_NAME("VisualizeMotionBlurOverlay"), View, Output,
+			[&View](FCanvas& Canvas)
+			{
+				float X = 20;
+				float Y = 38;
+				const float YStep = 14;
+				const float ColumnWidth = 200;
 
-		FString Line;
+				FString Line;
 
-		Line = FString::Printf(TEXT("Visualize MotionBlur"));
-		Canvas.DrawShadowedString(X, Y += YStep, *Line, GetStatsFont(), FLinearColor(1, 1, 0));
+				Line = FString::Printf(TEXT("Visualize MotionBlur"));
+				Canvas.DrawShadowedString(X, Y += YStep, *Line, GetStatsFont(), FLinearColor(1, 1, 0));
 
-		static const auto MotionBlurDebugVar = IConsoleManager::Get().FindTConsoleVariableDataInt(TEXT("r.MotionBlurDebug"));
-		const int32 MotionBlurDebug = MotionBlurDebugVar ? MotionBlurDebugVar->GetValueOnRenderThread() : 0;
+				static const auto MotionBlurDebugVar = IConsoleManager::Get().FindTConsoleVariableDataInt(TEXT("r.MotionBlurDebug"));
+				const int32 MotionBlurDebug = MotionBlurDebugVar ? MotionBlurDebugVar->GetValueOnRenderThread() : 0;
 
-		Line = FString::Printf(TEXT("%d, %d"), View.Family->FrameNumber, MotionBlurDebug);
-		Canvas.DrawShadowedString(X, Y += YStep, TEXT("FrameNo, r.MotionBlurDebug:"), GetStatsFont(), FLinearColor(1, 1, 0));
-		Canvas.DrawShadowedString(X + ColumnWidth, Y, *Line, GetStatsFont(), FLinearColor(1, 1, 0));
+				Line = FString::Printf(TEXT("%d, %d"), View.Family->FrameNumber, MotionBlurDebug);
+				Canvas.DrawShadowedString(X, Y += YStep, TEXT("FrameNo, r.MotionBlurDebug:"), GetStatsFont(), FLinearColor(1, 1, 0));
+				Canvas.DrawShadowedString(X + ColumnWidth, Y, *Line, GetStatsFont(), FLinearColor(1, 1, 0));
 
-		static const auto VelocityTestVar = IConsoleManager::Get().FindTConsoleVariableDataInt(TEXT("r.VelocityTest"));
-		const int32 VelocityTest = VelocityTestVar ? VelocityTestVar->GetValueOnRenderThread() : 0;
+				static const auto VelocityTestVar = IConsoleManager::Get().FindTConsoleVariableDataInt(TEXT("r.VelocityTest"));
+				const int32 VelocityTest = VelocityTestVar ? VelocityTestVar->GetValueOnRenderThread() : 0;
 
-		Line = FString::Printf(TEXT("%d, %d, %d"), View.Family->bWorldIsPaused, VelocityTest, FVelocityRendering::IsParallelVelocity(View.GetShaderPlatform()));
-		Canvas.DrawShadowedString(X, Y += YStep, TEXT("Paused, r.VelocityTest, Parallel:"), GetStatsFont(), FLinearColor(1, 1, 0));
-		Canvas.DrawShadowedString(X + ColumnWidth, Y, *Line, GetStatsFont(), FLinearColor(1, 1, 0));
+				Line = FString::Printf(TEXT("%d, %d, %d"), View.Family->bWorldIsPaused, VelocityTest, FVelocityRendering::IsParallelVelocity(View.GetShaderPlatform()));
+				Canvas.DrawShadowedString(X, Y += YStep, TEXT("Paused, r.VelocityTest, Parallel:"), GetStatsFont(), FLinearColor(1, 1, 0));
+				Canvas.DrawShadowedString(X + ColumnWidth, Y, *Line, GetStatsFont(), FLinearColor(1, 1, 0));
 
-		const FSceneViewState *SceneViewState = (const FSceneViewState*)View.State;
+				const FSceneViewState* SceneViewState = (const FSceneViewState*)View.State;
 
-		Line = FString::Printf(TEXT("View=%.4x PrevView=%.4x"),
-			View.ViewMatrices.GetViewMatrix().ComputeHash() & 0xffff,
-			View.PrevViewInfo.ViewMatrices.GetViewMatrix().ComputeHash() & 0xffff);
-		Canvas.DrawShadowedString(X, Y += YStep, TEXT("ViewMatrix:"), GetStatsFont(), FLinearColor(1, 1, 0));
-		Canvas.DrawShadowedString(X + ColumnWidth, Y, *Line, GetStatsFont(), FLinearColor(1, 1, 0));
-	});
+				Line = FString::Printf(TEXT("View=%.4x PrevView=%.4x"),
+					View.ViewMatrices.GetViewMatrix().ComputeHash() & 0xffff,
+					View.PrevViewInfo.ViewMatrices.GetViewMatrix().ComputeHash() & 0xffff);
+				Canvas.DrawShadowedString(X, Y += YStep, TEXT("ViewMatrix:"), GetStatsFont(), FLinearColor(1, 1, 0));
+				Canvas.DrawShadowedString(X + ColumnWidth, Y, *Line, GetStatsFont(), FLinearColor(1, 1, 0));
+			});
+	}
 
 	return FScreenPassTextureSlice::CreateFromScreenPassTexture(GraphBuilder, Output);
 }
