Generate cached mesh draw commands for one-pass point lights.
 - Considerable performance improvement to mesh pass setup tasks in scenes with many non-nanite shadow-casting lights.
 - Controlled by a read-only CVar r.Shadow.OnePassPointLights.UseCachedMDCs
 - Enabled by default.

diff --git a/Engine/Source/Runtime/Renderer/Private/MeshPassProcessor.cpp b/Engine/Source/Runtime/Renderer/Private/MeshPassProcessor.cpp
--- a/Engine/Source/Runtime/Renderer/Private/MeshPassProcessor.cpp
+++ b/Engine/Source/Runtime/Renderer/Private/MeshPassProcessor.cpp
@@ -2071,6 +2071,8 @@ void FCachedPassMeshDrawListContext::FinalizeCommandCommon(
 		MeshDrawCommand.ClearDebugPrimitiveSceneProxy(); //When using State Buckets multiple PrimitiveSceneProxies use the same MeshDrawCommand, so The PrimitiveSceneProxy pointer can't be stored.
 	}
 #endif
+	ensureMsgf(MeshDrawCommand.VertexStreams.GetAllocatedSize() == 0, TEXT("Cached Mesh Draw command overflows VertexStreams. VertexStream inline size should be tweaked."));
+
 #if DO_GUARD_SLOW
 	if (bUseGPUScene)
 	{
@@ -2082,7 +2084,12 @@ void FCachedPassMeshDrawListContext::FinalizeCommandCommon(
 	{
 		ensureMsgf(MeshDrawCommand.VertexStreams.GetAllocatedSize() == 0, TEXT("Cached Mesh Draw command overflows VertexStreams. VertexStream inline size should be tweaked."));
 
-		if (CurrMeshPass == EMeshPass::BasePass || CurrMeshPass == EMeshPass::DepthPass || CurrMeshPass == EMeshPass::SecondStageDepthPass || CurrMeshPass == EMeshPass::CSMShadowDepth || CurrMeshPass == EMeshPass::VSMShadowDepth)
+		if (   CurrMeshPass == EMeshPass::BasePass
+			|| CurrMeshPass == EMeshPass::DepthPass
+			|| CurrMeshPass == EMeshPass::SecondStageDepthPass
+			|| CurrMeshPass == EMeshPass::OnePassPointLightShadowDepth
+			|| CurrMeshPass == EMeshPass::CSMShadowDepth
+			|| CurrMeshPass == EMeshPass::VSMShadowDepth)
 		{
 			TArray<EShaderFrequency, TInlineAllocator<SF_NumFrequencies>> ShaderFrequencies;
 			MeshDrawCommand.ShaderBindings.GetShaderFrequencies(ShaderFrequencies);
diff --git a/Engine/Source/Runtime/Renderer/Private/ShadowDepthRendering.cpp b/Engine/Source/Runtime/Renderer/Private/ShadowDepthRendering.cpp
--- a/Engine/Source/Runtime/Renderer/Private/ShadowDepthRendering.cpp
+++ b/Engine/Source/Runtime/Renderer/Private/ShadowDepthRendering.cpp
@@ -72,6 +72,13 @@ static FAutoConsoleVariableRef CVarFarShadowScreenMultiple(
 	TEXT("Multiplier for LOD selection distance when rendering far shadows"),
 	ECVF_RenderThreadSafe);
 
+static bool GShadowOnePassPointLightUseCachedMDCs = true;
+static FAutoConsoleVariableRef CVarShadowOnePassPointLightUseCachedMDCs(
+	TEXT("r.Shadow.OnePassPointLight.UseCachedMDCs"),
+	GShadowOnePassPointLightUseCachedMDCs,
+	TEXT("Controls whether to use a dedicated cached mesh draw command pass for one pass point lights. This increases the cost of caching / storing mesh draw commands but improves the CPU rendering time of point lights."),
+	ECVF_ReadOnly);
+
 DEFINE_GPU_DRAWCALL_STAT(ShadowDepths);
 
 IMPLEMENT_STATIC_UNIFORM_BUFFER_STRUCT(FShadowDepthPassUniformParameters, "ShadowDepthPass", SceneTextures);
@@ -527,10 +534,9 @@ bool GetShadowDepthPassShaders(
 	const FMaterial& Material,
 	const FVertexFactoryType* VertexFactoryType,
 	ERHIFeatureLevel::Type FeatureLevel,
-	FShadowDepthType ShadowDepthType,
+	EShadowDepthType ShadowDepthType,
 	bool bSupportsPositionAndNormalOnlyStream,
 	bool bMaterialModifiesMeshPosition,
-	EMeshPass::Type MeshPassType,
 	TShaderRef<FShadowDepthVS>& VertexShader,
 	TShaderRef<FShadowDepthBasePS>& PixelShader)
 {
@@ -544,10 +550,10 @@ bool GetShadowDepthPassShaders(
 	// Depth will be interpolated to the pixel shader and written out, which disables HiZ and double speed Z.
 	// Directional light shadows use an ortho projection and can use the non-perspective correct path without artifacts.
 	// One pass point lights don't output a linear depth, so they are already perspective correct.
-	bool bUsePerspectiveCorrectShadowDepths = !ShadowDepthType.bDirectionalLight && !ShadowDepthType.bOnePassPointLightShadow;
-	bool bOnePassPointLightShadow = ShadowDepthType.bOnePassPointLightShadow;
+	bool bUsePerspectiveCorrectShadowDepths = !EnumHasAnyFlags(ShadowDepthType, EShadowDepthType::Directional | EShadowDepthType::OnePassPoint);
+	bool bOnePassPointLightShadow = EnumHasAnyFlags(ShadowDepthType, EShadowDepthType::OnePassPoint);
 
-	bool bVirtualShadowMap = MeshPassType == EMeshPass::VSMShadowDepth;
+	bool bVirtualShadowMap = EnumHasAnyFlags(ShadowDepthType, EShadowDepthType::VSM);
 	if (bVirtualShadowMap)
 	{
 		bUsePerspectiveCorrectShadowDepths = false;
@@ -702,16 +708,16 @@ void FProjectedShadowInfo::SetStateForView(FRHICommandList& RHICmdList) const
 	}
 }
 
-void SetStateForShadowDepth(bool bOnePassPointLightShadow, bool bDirectionalLight, FMeshPassProcessorRenderState& DrawRenderState, EMeshPass::Type InMeshPassTargetType)
+void SetStateForShadowDepth(EShadowDepthType ShadowDepthType, FMeshPassProcessorRenderState& DrawRenderState)
 {
 	// Disable color writes
 	DrawRenderState.SetBlendState(TStaticBlendState<CW_NONE>::GetRHI());
 
-	if( InMeshPassTargetType == EMeshPass::VSMShadowDepth )
+	if (EnumHasAnyFlags(ShadowDepthType, EShadowDepthType::VSM))
 	{
 		DrawRenderState.SetDepthStencilState(TStaticDepthStencilState<false, CF_Always>::GetRHI());
 	}
-	else if (bOnePassPointLightShadow || InMeshPassTargetType == EMeshPass::VSMShadowDepth)
+	else if (EnumHasAnyFlags(ShadowDepthType, EShadowDepthType::OnePassPoint))
 	{
 		// Point lights use reverse Z depth maps
 		DrawRenderState.SetDepthStencilState(TStaticDepthStencilState<true, CF_DepthNearOrEqual>::GetRHI());
@@ -1119,7 +1125,7 @@ void FProjectedShadowInfo::RenderDepth(
 	{
 		// Copy in depths of static primitives before we render movable primitives.
 		FMeshPassProcessorRenderState DrawRenderState;
-		SetStateForShadowDepth(bOnePassPointLightShadow, bDirectionalLight, DrawRenderState, MeshPassTargetType);
+		SetStateForShadowDepth(GetShadowDepthType(), DrawRenderState);
 		CopyCachedShadowMap(GraphBuilder, *ShadowDepthView, SceneRenderer, PassParameters->RenderTargets, DrawRenderState);
 	}
 
@@ -1969,17 +1975,16 @@ void FSceneRenderer::RenderShadowDepthMaps(FRDGBuilder& GraphBuilder, FDynamicSh
 }
 
 ERasterizerCullMode SetupShadowCullMode(
-	ERHIFeatureLevel::Type FeatureLevel, 
-	EMeshPass::Type MeshPassTargetType, 
-	FShadowDepthType ShadowDepthType, 
-	const FMaterial& Material, 
+	ERHIFeatureLevel::Type FeatureLevel,
+	EShadowDepthType ShadowDepthType,
+	const FMaterial& Material,
 	const ERasterizerCullMode MeshCullMode,
 	bool bCastShadowAsTwoSided)
 {
 	const bool bTwoSided = Material.IsTwoSided() || bCastShadowAsTwoSided;
 
 	const bool bRenderSceneTwoSided = bTwoSided;
-	const bool bShadowReversesCulling = MeshPassTargetType == EMeshPass::VSMShadowDepth ? false : ShadowDepthType.bOnePassPointLightShadow;
+	const bool bShadowReversesCulling = EnumHasAnyFlags(ShadowDepthType, EShadowDepthType::VSM) ? false : EnumHasAnyFlags(ShadowDepthType, EShadowDepthType::OnePassPoint);
 	const bool bReverseCullMode = bShadowReversesCulling;
 
 	return bRenderSceneTwoSided ? CM_None : bReverseCullMode ? FMeshPassProcessor::InverseCullMode(MeshCullMode) : MeshCullMode;
@@ -1988,7 +1993,7 @@ ERasterizerCullMode SetupShadowCullMode(
 void FShadowDepthPassMeshProcessor::CollectPSOInitializersInternal(
 	const FPSOPrecacheVertexFactoryData& VertexFactoryData,
 	const FMaterial& RESTRICT MaterialResource,
-	const FShadowDepthType& InShadowDepthType,
+	EShadowDepthType InShadowDepthType,
 	ERasterizerFillMode MeshFillMode,
 	ERasterizerCullMode MeshCullMode,
 	bool bSupportsPositionAndNormalOnlyStream,
@@ -2005,7 +2010,6 @@ void FShadowDepthPassMeshProcessor::CollectPSOInitializersInternal(
 		InShadowDepthType,
 		bSupportsPositionAndNormalOnlyStream,
 		MaterialResource.MaterialModifiesMeshPosition_GameThread(),
-		MeshPassTargetType,
 		ShadowDepthPassShaders.VertexShader,
 		ShadowDepthPassShaders.PixelShader))
 	{
@@ -2023,7 +2027,7 @@ void FShadowDepthPassMeshProcessor::CollectPSOInitializersInternal(
 	{
 		ETextureCreateFlags ShadowMapCreateFlags = TexCreate_DepthStencilTargetable | TexCreate_ShaderResource;
 		ShadowMapCreateFlags |= GFastVRamConfig.ShadowPointLight;
-		if (InShadowDepthType.bOnePassPointLightShadow)
+		if (EnumHasAnyFlags(InShadowDepthType, EShadowDepthType::OnePassPoint))
 		{
 			ShadowMapCreateFlags |= TexCreate_NoFastClear;
 		}
@@ -2069,7 +2073,6 @@ bool FShadowDepthPassMeshProcessor::Process(
 		ShadowDepthType,
 		VertexFactory->SupportsPositionAndNormalOnlyStream(),
 		MaterialResource.MaterialModifiesMeshPosition_RenderThread(),
-		MeshPassTargetType,
 		ShadowDepthPassShaders.VertexShader,
 		ShadowDepthPassShaders.PixelShader))
 	{
@@ -2089,7 +2092,7 @@ bool FShadowDepthPassMeshProcessor::Process(
 		&& !MaterialResource.MaterialModifiesMeshPosition_RenderThread();
 
 	// Need to replicate for cube faces on host if GPU-scene is not available (for this draw).
-	const bool bPerformHostCubeFaceReplication = ShadowDepthType.bOnePassPointLightShadow && !bUseGpuSceneInstancing;
+	const bool bPerformHostCubeFaceReplication = EnumHasAnyFlags(ShadowDepthType, EShadowDepthType::OnePassPoint) && !bUseGpuSceneInstancing;
 	const uint32 InstanceFactor = bPerformHostCubeFaceReplication ? 6 : 1;
 
 	for (uint32 i = 0; i < InstanceFactor; i++)
@@ -2128,7 +2131,7 @@ bool FShadowDepthPassMeshProcessor::TryAddMeshBatch(const FMeshBatch& RESTRICT M
 	const ERasterizerFillMode MeshFillMode = ComputeMeshFillMode(Material, OverrideSettings);
 
 	const ERasterizerCullMode MeshCullMode = FMeshPassProcessor::ComputeMeshCullMode(Material, OverrideSettings);
-	ERasterizerCullMode FinalCullMode = SetupShadowCullMode(FeatureLevel, MeshPassTargetType, ShadowDepthType, Material, MeshCullMode, PrimitiveSceneProxy->CastsShadowAsTwoSided());
+	ERasterizerCullMode FinalCullMode = SetupShadowCullMode(FeatureLevel, ShadowDepthType, Material, MeshCullMode, PrimitiveSceneProxy->CastsShadowAsTwoSided());
 
 	bool bResult = true;
 	if (bShouldCastShadow
@@ -2238,30 +2241,27 @@ void FShadowDepthPassMeshProcessor::CollectPSOInitializersForEachShadowDepthType
 	bool bCastShadowAsTwoSided,
 	TArray<FPSOPrecacheData>& PSOInitializers)
 {
-	FShadowDepthType LocalShadowDepthType(true, false);
+	const EShadowDepthType ShadowDepthVSMBit = ShadowDepthType & EShadowDepthType::VSM;
 	bool bRequired = true;
 
 	// Collect for directional shadows
 	{
-		LocalShadowDepthType.bDirectionalLight = true;
-		LocalShadowDepthType.bOnePassPointLightShadow = false;
-		ERasterizerCullMode FinalCullMode = SetupShadowCullMode(FeatureLevel, MeshPassTargetType, LocalShadowDepthType, Material, MeshCullMode, bCastShadowAsTwoSided);
+		const EShadowDepthType LocalShadowDepthType = ShadowDepthVSMBit | EShadowDepthType::Directional;
+		ERasterizerCullMode FinalCullMode = SetupShadowCullMode(FeatureLevel, LocalShadowDepthType, Material, MeshCullMode, bCastShadowAsTwoSided);
 		CollectPSOInitializersForEachStreamSetup(VertexFactoryData, Material, LocalShadowDepthType, MeshFillMode, FinalCullMode, bRequired, PSOInitializers);
 	}
 
 	// Collect for non-directional one pass point light shadows
 	{
-		LocalShadowDepthType.bDirectionalLight = false;
-		LocalShadowDepthType.bOnePassPointLightShadow = true;
-		ERasterizerCullMode FinalCullMode = SetupShadowCullMode(FeatureLevel, MeshPassTargetType, LocalShadowDepthType, Material, MeshCullMode, bCastShadowAsTwoSided);
+		const EShadowDepthType LocalShadowDepthType = ShadowDepthVSMBit | EShadowDepthType::OnePassPoint;
+		ERasterizerCullMode FinalCullMode = SetupShadowCullMode(FeatureLevel, LocalShadowDepthType, Material, MeshCullMode, bCastShadowAsTwoSided);
 		CollectPSOInitializersForEachStreamSetup(VertexFactoryData, Material, LocalShadowDepthType, MeshFillMode, FinalCullMode, bRequired, PSOInitializers);
 	}
 
 	// Collect for non-directional non-one pass point light shadows
 	{
-		LocalShadowDepthType.bDirectionalLight = false;
-		LocalShadowDepthType.bOnePassPointLightShadow = false;
-		ERasterizerCullMode FinalCullMode = SetupShadowCullMode(FeatureLevel, MeshPassTargetType, LocalShadowDepthType, Material, MeshCullMode, bCastShadowAsTwoSided);
+		const EShadowDepthType LocalShadowDepthType = ShadowDepthVSMBit;
+		ERasterizerCullMode FinalCullMode = SetupShadowCullMode(FeatureLevel, LocalShadowDepthType, Material, MeshCullMode, bCastShadowAsTwoSided);
 		CollectPSOInitializersForEachStreamSetup(VertexFactoryData, Material, LocalShadowDepthType, MeshFillMode, FinalCullMode, bRequired, PSOInitializers);
 	}
 }
@@ -2269,13 +2269,13 @@ void FShadowDepthPassMeshProcessor::CollectPSOInitializersForEachShadowDepthType
 void FShadowDepthPassMeshProcessor::CollectPSOInitializersForEachStreamSetup(
 	const FPSOPrecacheVertexFactoryData& VertexFactoryData,
 	const FMaterial& RESTRICT MaterialResource,
-	const FShadowDepthType& InShadowDepthType,
+	EShadowDepthType InShadowDepthType,
 	ERasterizerFillMode MeshFillMode,
 	ERasterizerCullMode MeshCullMode,
 	bool bRequired,
 	TArray<FPSOPrecacheData>& PSOInitializers)
 {
-	SetStateForShadowDepth(InShadowDepthType.bOnePassPointLightShadow, InShadowDepthType.bDirectionalLight, PassDrawRenderState, MeshPassTargetType);
+	SetStateForShadowDepth(InShadowDepthType, PassDrawRenderState);
 
 	// Collect for when both use cases to be complete
 	// Ideally position and normal stream is always available or can be checked on VF to be always available to reduce PSO precache count
@@ -2294,28 +2294,54 @@ FShadowDepthPassMeshProcessor::FShadowDepthPassMeshProcessor(
 	const FScene* Scene,
 	const ERHIFeatureLevel::Type InFeatureLevel,
 	const FSceneView* InViewIfDynamicMeshCommand,
-	FShadowDepthType InShadowDepthType,
-	FMeshPassDrawListContext* InDrawListContext,
-	EMeshPass::Type InMeshPassTargetType)
-	: FMeshPassProcessor(InMeshPassTargetType, Scene, InFeatureLevel, InViewIfDynamicMeshCommand, InDrawListContext)
+	EShadowDepthType InShadowDepthType,
+	FMeshPassDrawListContext* InDrawListContext)
+	: FMeshPassProcessor(GetShadowMeshPassType(InShadowDepthType), Scene, InFeatureLevel, InViewIfDynamicMeshCommand, InDrawListContext)
 	, ShadowDepthType(InShadowDepthType)
-	, MeshPassTargetType(InMeshPassTargetType)
 {
 	EShaderPlatform ShaderPlatform = GetFeatureLevelShaderPlatform(InFeatureLevel);
 	if (UseNonNaniteVirtualShadowMaps(ShaderPlatform, FeatureLevel))
 	{
 		// set up mesh filtering.
-		MeshSelectionMask = MeshPassTargetType == EMeshPass::VSMShadowDepth ? EShadowMeshSelection::VSM : EShadowMeshSelection::SM;
+		MeshSelectionMask = EnumHasAnyFlags(ShadowDepthType, EShadowDepthType::VSM) ? EShadowMeshSelection::VSM : EShadowMeshSelection::SM;
 	}
 	else
 	{
 		// If VSMs are disabled, pipe all kinds of draws into the regular SMs
 		MeshSelectionMask = EShadowMeshSelection::All;
 	}
-	SetStateForShadowDepth(ShadowDepthType.bOnePassPointLightShadow, ShadowDepthType.bDirectionalLight, PassDrawRenderState, MeshPassTargetType);
+	SetStateForShadowDepth(ShadowDepthType, PassDrawRenderState);
 }
 
-FShadowDepthType CSMShadowDepthType(true, false);
+EMeshPass::Type GetShadowMeshPassType(EShadowDepthType ShadowDepthType)
+{
+	if (EnumHasAnyFlags(ShadowDepthType, EShadowDepthType::VSM))
+	{
+		return EMeshPass::VSMShadowDepth;
+	}
+
+	if (EnumHasAnyFlags(ShadowDepthType, EShadowDepthType::OnePassPoint))
+	{
+		return EMeshPass::OnePassPointLightShadowDepth;
+	}
+
+	return EMeshPass::CSMShadowDepth;
+}
+
+bool UseCachedMeshDrawCommands(EShadowDepthType ShadowDepthType)
+{
+	if (EnumHasAnyFlags(ShadowDepthType, EShadowDepthType::VSM | EShadowDepthType::Directional))
+	{
+		return true;
+	}
+
+	if (EnumHasAnyFlags(ShadowDepthType, EShadowDepthType::OnePassPoint) && GShadowOnePassPointLightUseCachedMDCs)
+	{
+		return true;
+	}
+
+	return false;
+}
 
 FMeshPassProcessor* CreateCSMShadowDepthPassProcessor(ERHIFeatureLevel::Type FeatureLevel, const FScene* Scene, const FSceneView* InViewIfDynamicMeshCommand, FMeshPassDrawListContext* InDrawListContext)
 {
@@ -2323,12 +2349,30 @@ FMeshPassProcessor* CreateCSMShadowDepthPassProcessor(ERHIFeatureLevel::Type Fea
 		Scene,
 		FeatureLevel,
 		InViewIfDynamicMeshCommand,
-		CSMShadowDepthType,
-		InDrawListContext,
-		EMeshPass::CSMShadowDepth);
+		EShadowDepthType::Directional,
+		InDrawListContext);
 }
 
 REGISTER_MESHPASSPROCESSOR_AND_PSOCOLLECTOR(CSMShadowDepthPass, CreateCSMShadowDepthPassProcessor, EShadingPath::Deferred, EMeshPass::CSMShadowDepth, EMeshPassFlags::CachedMeshCommands);
+REGISTER_MESHPASSPROCESSOR_AND_PSOCOLLECTOR(CSMMobileShadowDepthPass, CreateCSMShadowDepthPassProcessor, EShadingPath::Mobile, EMeshPass::CSMShadowDepth, EMeshPassFlags::CachedMeshCommands);
+
+FMeshPassProcessor* CreateOnePassPointLightShadowDepthPassProcessor(ERHIFeatureLevel::Type FeatureLevel, const FScene* Scene, const FSceneView* InViewIfDynamicMeshCommand, FMeshPassDrawListContext* InDrawListContext)
+{
+	const EShaderPlatform ShaderPlatform = GetFeatureLevelShaderPlatform(FeatureLevel);
+
+	if (GShadowOnePassPointLightUseCachedMDCs && DoesRuntimeSupportOnePassPointLightShadows(ShaderPlatform))
+	{
+		return new FShadowDepthPassMeshProcessor(
+			Scene,
+			FeatureLevel,
+			InViewIfDynamicMeshCommand,
+			EShadowDepthType::OnePassPoint,
+			InDrawListContext);
+	}
+	return nullptr;
+}
+
+REGISTER_MESHPASSPROCESSOR_AND_PSOCOLLECTOR(OnePassPointLightShadowDepthPass, CreateOnePassPointLightShadowDepthPassProcessor, EShadingPath::Deferred, EMeshPass::OnePassPointLightShadowDepth, EMeshPassFlags::CachedMeshCommands);
 
 FMeshPassProcessor* CreateVSMShadowDepthPassProcessor(ERHIFeatureLevel::Type FeatureLevel, const FScene* Scene, const FSceneView* InViewIfDynamicMeshCommand, FMeshPassDrawListContext* InDrawListContext)
 {
@@ -2341,13 +2385,11 @@ FMeshPassProcessor* CreateVSMShadowDepthPassProcessor(ERHIFeatureLevel::Type Fea
 			Scene,
 			FeatureLevel,
 			InViewIfDynamicMeshCommand,
-			CSMShadowDepthType,
-			InDrawListContext,
-			EMeshPass::VSMShadowDepth);
+			EShadowDepthType::VSM | EShadowDepthType::Directional,
+			InDrawListContext);
 	}
 
 	return nullptr;
 }
 
-REGISTER_MESHPASSPROCESSOR_AND_PSOCOLLECTOR(VSMShadowDepthPass, CreateVSMShadowDepthPassProcessor, EShadingPath::Deferred, EMeshPass::VSMShadowDepth, EMeshPassFlags::CachedMeshCommands);
-REGISTER_MESHPASSPROCESSOR_AND_PSOCOLLECTOR(CSMMobileShadowDepthPass, CreateCSMShadowDepthPassProcessor, EShadingPath::Mobile, EMeshPass::CSMShadowDepth, EMeshPassFlags::CachedMeshCommands);
+REGISTER_MESHPASSPROCESSOR_AND_PSOCOLLECTOR(VSMShadowDepthPass, CreateVSMShadowDepthPassProcessor, EShadingPath::Deferred, EMeshPass::VSMShadowDepth, EMeshPassFlags::CachedMeshCommands);
\ No newline at end of file
diff --git a/Engine/Source/Runtime/Renderer/Private/ShadowRendering.h b/Engine/Source/Runtime/Renderer/Private/ShadowRendering.h
--- a/Engine/Source/Runtime/Renderer/Private/ShadowRendering.h
+++ b/Engine/Source/Runtime/Renderer/Private/ShadowRendering.h
@@ -79,32 +79,20 @@ enum EShadowDepthRenderMode
 	ShadowDepthRenderMode_GIBlockingVolumes,
 };
 
-class FShadowDepthType
+enum class EShadowDepthType
 {
-public:
-	bool bDirectionalLight;
-	bool bOnePassPointLightShadow;
-
-	FShadowDepthType(bool bInDirectionalLight, bool bInOnePassPointLightShadow) 
-	: bDirectionalLight(bInDirectionalLight)
-	, bOnePassPointLightShadow(bInOnePassPointLightShadow)
-	{}
-
-	inline bool operator==(const FShadowDepthType& rhs) const
-	{
-		if (bDirectionalLight != rhs.bDirectionalLight || 
-			bOnePassPointLightShadow != rhs.bOnePassPointLightShadow)
-		{
-			return false;
-		}
-
-		return true;
-	}
+	None = 0,
+	VSM = 1 << 0,
+	Directional = 1 << 1,
+	Point = 1 << 2,
+	OnePassPoint = 1 << 3
 };
 
-extern FShadowDepthType CSMShadowDepthType;
+ENUM_CLASS_FLAGS(EShadowDepthType);
 
+extern EMeshPass::Type GetShadowMeshPassType(EShadowDepthType ShadowDepthType);
 
+extern bool UseCachedMeshDrawCommands(EShadowDepthType ShadowDepthType);
 
 /**
  * Used to select what meshes to draw into what shadow infos. Each mesh selects the support it has, and
@@ -128,10 +116,9 @@ public:
 	FShadowDepthPassMeshProcessor(
 		const FScene* Scene, 
 		const ERHIFeatureLevel::Type InFeatureLevel,
-		const FSceneView* InViewIfDynamicMeshCommand, 
-		FShadowDepthType InShadowDepthType,
-		FMeshPassDrawListContext* InDrawListContext,
-		EMeshPass::Type InMeshPassTargetType);
+		const FSceneView* InViewIfDynamicMeshCommand,
+		EShadowDepthType ShadowDepthType,
+		FMeshPassDrawListContext* InDrawListContext);
 
 	virtual void AddMeshBatch(const FMeshBatch& RESTRICT MeshBatch, uint64 BatchElementMask, const FPrimitiveSceneProxy* RESTRICT PrimitiveSceneProxy, int32 StaticMeshId = -1) override final;
 	virtual void CollectPSOInitializers(const FSceneTexturesConfig& SceneTexturesConfig, const FMaterial& Material, const FPSOPrecacheVertexFactoryData& VertexFactoryData, const FPSOPrecacheParams& PreCacheParams, TArray<FPSOPrecacheData>& PSOInitializers) override final;
@@ -167,7 +154,7 @@ private:
 	void CollectPSOInitializersForEachStreamSetup(
 		const FPSOPrecacheVertexFactoryData& VertexFactoryData,
 		const FMaterial& RESTRICT MaterialResource,
-		const FShadowDepthType& InShadowDepthType,
+		EShadowDepthType InShadowDepthType,
 		ERasterizerFillMode MeshFillMode,
 		ERasterizerCullMode MeshCullMode,
 		bool bRequired,
@@ -176,15 +163,14 @@ private:
 	void CollectPSOInitializersInternal(
 		const FPSOPrecacheVertexFactoryData& VertexFactoryData,
 		const FMaterial& RESTRICT MaterialResource,
-		const FShadowDepthType& InShadowDepthType,
+		EShadowDepthType InShadowDepthType,
 		ERasterizerFillMode MeshFillMode,
 		ERasterizerCullMode MeshCullMode,
 		bool bSupportsPositionAndNormalOnlyStream,
 		bool bRequired,
 		TArray<FPSOPrecacheData>& PSOInitializers);
 
-	FShadowDepthType ShadowDepthType;
-	EMeshPass::Type MeshPassTargetType = EMeshPass::CSMShadowDepth;
+	EShadowDepthType ShadowDepthType;
 	EShadowMeshSelection MeshSelectionMask = EShadowMeshSelection::All;
 };
 
@@ -403,6 +389,9 @@ public:
 	/** Whether the shadow is a point light shadow that renders all faces of a cubemap in one pass. */
 	uint32 bOnePassPointLightShadow : 1;
 
+	/** Whether the shadow is a virtual shadow map. */
+	uint32 bVSM : 1;
+
 	/** Whether this shadow affects the whole scene or only a group of objects. */
 	uint32 bWholeSceneShadow : 1;
 
@@ -439,9 +428,6 @@ public:
 
 	/** Whether this shadow should support casting shadows from volumetric surfaces. */
 	uint32 bVolumetricShadow : 1;
-	
-	/** Used to fetch the correct cached static mesh draw commands */
-	EMeshPass::Type MeshPassTargetType = EMeshPass::CSMShadowDepth;
 
 	EShadowMeshSelection MeshSelectionMask = EShadowMeshSelection::All;
 
@@ -784,9 +770,29 @@ public:
 	/** Creates a new view from the pool and caches it in ShadowDepthView for depth rendering. */
 	void SetupShadowDepthView(FSceneRenderer* SceneRenderer);
 
-	FShadowDepthType GetShadowDepthType() const 
+	EShadowDepthType GetShadowDepthType() const 
+	{
+		EShadowDepthType ShadowDepthType = bVSM ? EShadowDepthType::VSM : EShadowDepthType::None;
+
+		if (bDirectionalLight)
+		{
+			ShadowDepthType |= EShadowDepthType::Directional;
+		}
+		else if (bOnePassPointLightShadow)
+		{
+			ShadowDepthType |= EShadowDepthType::OnePassPoint;
+		}
+		else
+		{
+			ShadowDepthType |= EShadowDepthType::Point;
+		}
+
+		return ShadowDepthType;
+	}
+
+	EMeshPass::Type GetTargetMeshPassType() const
 	{
-		return FShadowDepthType(bDirectionalLight, bOnePassPointLightShadow);
+		return ::GetShadowMeshPassType(GetShadowDepthType());
 	}
 
 	bool HasVirtualShadowMap() const { return VirtualShadowMapId != INDEX_NONE || VirtualShadowMapClipmap.IsValid(); }
diff --git a/Engine/Source/Runtime/Renderer/Private/ShadowSetup.cpp b/Engine/Source/Runtime/Renderer/Private/ShadowSetup.cpp
--- a/Engine/Source/Runtime/Renderer/Private/ShadowSetup.cpp
+++ b/Engine/Source/Runtime/Renderer/Private/ShadowSetup.cpp
@@ -738,6 +738,7 @@ FProjectedShadowInfo::FProjectedShadowInfo()
 	, bDepthsCached(false)
 	, bDirectionalLight(false)
 	, bOnePassPointLightShadow(false)
+	, bVSM(false)
 	, bWholeSceneShadow(false)
 	, bTranslucentShadow(false)
 	, bRayTracedDistanceField(false)
@@ -1107,7 +1108,7 @@ void FProjectedShadowInfo::SetupClipmapProjection(FLightSceneInfo* InLightSceneI
 	bWholeSceneShadow = true;
 	BorderSize = 0;
 	MaxNonFarCascadeDistance = InMaxNonFarCascadeDistance;
-	MeshPassTargetType = EMeshPass::VSMShadowDepth;
+	bVSM = true;
 	MeshSelectionMask = EShadowMeshSelection::VSM;
 
 	const int32 ClipmapIndex = VirtualShadowMapClipmap->GetLevelCount() - 1;
@@ -1663,7 +1664,7 @@ void FProjectedShadowInfo::AddCachedMeshDrawCommands_AnyThread(
 	FAddSubjectPrimitiveStats& OutStats,
 	FAddSubjectPrimitiveOverflowedIndices& OverflowBuffer) const
 {
-	const EMeshPass::Type PassType = MeshPassTargetType;
+	const EMeshPass::Type PassType = GetTargetMeshPassType();
 	const EShadingPath ShadingPath = GetFeatureLevelShadingPath(Scene->GetFeatureLevel());
 	const bool bUseCachedMeshCommand = UseCachedMeshDrawCommands_AnyThread()
 		&& !!(FPassProcessorManager::GetPassFlags(ShadingPath, PassType) & EMeshPassFlags::CachedMeshCommands)
@@ -1774,66 +1775,57 @@ FORCEINLINE bool FProjectedShadowInfo::ShouldDrawStaticMesh(const FStaticMeshBat
 
 bool FProjectedShadowInfo::ShouldDrawStaticMeshes(FViewInfo& InCurrentView, FPrimitiveSceneInfo* InPrimitiveSceneInfo)
 {
-	bool WholeSceneDirectionalShadow = IsWholeSceneDirectionalShadow();
 	bool bDrawingStaticMeshes = false;
-	int32 PrimitiveId = InPrimitiveSceneInfo->GetIndex();
+	const int32 PrimitiveId = InPrimitiveSceneInfo->GetIndex();
 	const FMeshDrawCommandPrimitiveIdInfo PrimitiveIdInfo = InPrimitiveSceneInfo->GetMDCIdInfo();
+	const int32 ForcedLOD = (InCurrentView.Family->EngineShowFlags.LOD) ? (GetCVarForceLODShadow() != -1 ? GetCVarForceLODShadow() : GetCVarForceLOD()) : -1;
+	FLODMask ShadowLODToRender = CalcAndUpdateLODToRender(InCurrentView, InPrimitiveSceneInfo->Proxy->GetBounds(), InPrimitiveSceneInfo, ForcedLOD);
+	const EShadowDepthType ShadowDepthType = GetShadowDepthType();
+
+	// Don't cache if it requires per view per mesh state for distance cull fade.
+	const bool bIsPrimitiveDistanceCullFading = InCurrentView.PotentiallyFadingPrimitiveMap[PrimitiveId];
+	const bool bCanCache = !bIsPrimitiveDistanceCullFading;
+
+	if (bCanCache && UseCachedMeshDrawCommands(ShadowDepthType))
 	{
-		const int32 ForcedLOD = (InCurrentView.Family->EngineShowFlags.LOD) ? (GetCVarForceLODShadow() != -1 ? GetCVarForceLODShadow() : GetCVarForceLOD()) : -1;
-		FLODMask ShadowLODToRender = CalcAndUpdateLODToRender(InCurrentView, InPrimitiveSceneInfo->Proxy->GetBounds(), InPrimitiveSceneInfo, ForcedLOD);
+		const EMeshPass::Type MeshPassType = GetShadowMeshPassType(ShadowDepthType);
 
-		if (WholeSceneDirectionalShadow)
+		for (int32 MeshIndex = 0; MeshIndex < InPrimitiveSceneInfo->StaticMeshRelevances.Num(); MeshIndex++)
 		{
-			// Don't cache if it requires per view per mesh state for distance cull fade.
-			const bool bIsPrimitiveDistanceCullFading = InCurrentView.PotentiallyFadingPrimitiveMap[InPrimitiveSceneInfo->GetIndex()];
-			const bool bCanCache = !bIsPrimitiveDistanceCullFading;
+			const FStaticMeshBatchRelevance& StaticMeshRelevance = InPrimitiveSceneInfo->StaticMeshRelevances[MeshIndex];
+			const FStaticMeshBatch& StaticMesh = InPrimitiveSceneInfo->StaticMeshes[MeshIndex];
 
-			for (int32 MeshIndex = 0; MeshIndex < InPrimitiveSceneInfo->StaticMeshRelevances.Num(); MeshIndex++)
+			if (ShouldDrawStaticMesh(StaticMeshRelevance, ShadowLODToRender, bDrawingStaticMeshes))
 			{
-				const FStaticMeshBatchRelevance& StaticMeshRelevance = InPrimitiveSceneInfo->StaticMeshRelevances[MeshIndex];
-				const FStaticMeshBatch& StaticMesh = InPrimitiveSceneInfo->StaticMeshes[MeshIndex];
-
-				if (ShouldDrawStaticMesh(StaticMeshRelevance, ShadowLODToRender, bDrawingStaticMeshes))
-				{
-					const EMeshDrawCommandCullingPayloadFlags CullingPayloadFlags = GetCullingPayloadFlags(ShadowLODToRender, StaticMeshRelevance.GetLODIndex());
-
-					if (GetShadowDepthType() == CSMShadowDepthType && bCanCache)
-					{
-						AddCachedMeshDrawCommandsForPass(
-							PrimitiveIdInfo,
-							InPrimitiveSceneInfo,
-							StaticMeshRelevance,
-							StaticMesh,
-							CullingPayloadFlags,
-							InPrimitiveSceneInfo->Scene,
-							MeshPassTargetType,
-							ShadowDepthPassVisibleCommands,
-							SubjectMeshCommandBuildRequests,
-							SubjectMeshCommandBuildFlags,
-							NumSubjectMeshCommandBuildRequestElements);
-					}
-					else
-					{
-						NumSubjectMeshCommandBuildRequestElements += StaticMeshRelevance.NumElements;
-						SubjectMeshCommandBuildRequests.Add(&StaticMesh);
-						SubjectMeshCommandBuildFlags.Add(CullingPayloadFlags);
-					}
-				}
+				const EMeshDrawCommandCullingPayloadFlags CullingPayloadFlags = GetCullingPayloadFlags(ShadowLODToRender, StaticMeshRelevance.GetLODIndex());
+
+				AddCachedMeshDrawCommandsForPass(
+					PrimitiveIdInfo,
+					InPrimitiveSceneInfo,
+					StaticMeshRelevance,
+					StaticMesh,
+					CullingPayloadFlags,
+					InPrimitiveSceneInfo->Scene,
+					MeshPassType,
+					ShadowDepthPassVisibleCommands,
+					SubjectMeshCommandBuildRequests,
+					SubjectMeshCommandBuildFlags,
+					NumSubjectMeshCommandBuildRequestElements);
 			}
 		}
-		else
+	}
+	else
+	{
+		for (int32 MeshIndex = 0; MeshIndex < InPrimitiveSceneInfo->StaticMeshRelevances.Num(); MeshIndex++)
 		{
-			for (int32 MeshIndex = 0; MeshIndex < InPrimitiveSceneInfo->StaticMeshRelevances.Num(); MeshIndex++)
-			{
-				const FStaticMeshBatchRelevance& StaticMeshRelevance = InPrimitiveSceneInfo->StaticMeshRelevances[MeshIndex];
-				const FStaticMeshBatch& StaticMesh = InPrimitiveSceneInfo->StaticMeshes[MeshIndex];
+			const FStaticMeshBatchRelevance& StaticMeshRelevance = InPrimitiveSceneInfo->StaticMeshRelevances[MeshIndex];
+			const FStaticMeshBatch& StaticMesh = InPrimitiveSceneInfo->StaticMeshes[MeshIndex];
 
-				if (ShouldDrawStaticMesh(StaticMeshRelevance, ShadowLODToRender, bDrawingStaticMeshes))
-				{
-					NumSubjectMeshCommandBuildRequestElements += StaticMeshRelevance.NumElements;
-					SubjectMeshCommandBuildRequests.Add(&StaticMesh);
-					SubjectMeshCommandBuildFlags.Add(GetCullingPayloadFlags(ShadowLODToRender, StaticMeshRelevance.GetLODIndex()));
-				}
+			if (ShouldDrawStaticMesh(StaticMeshRelevance, ShadowLODToRender, bDrawingStaticMeshes))
+			{
+				NumSubjectMeshCommandBuildRequestElements += StaticMeshRelevance.NumElements;
+				SubjectMeshCommandBuildRequests.Add(&StaticMesh);
+				SubjectMeshCommandBuildFlags.Add(GetCullingPayloadFlags(ShadowLODToRender, StaticMeshRelevance.GetLODIndex()));
 			}
 		}
 	}
@@ -1853,6 +1845,7 @@ bool FProjectedShadowInfo::ShouldDrawStaticMeshes_AnyThread(
 	const bool WholeSceneDirectionalShadow = IsWholeSceneDirectionalShadow();
 	const FPrimitiveSceneInfo* PrimitiveSceneInfo = PrimitiveSceneInfoCompact.PrimitiveSceneInfo;
 	const FPrimitiveSceneProxy* Proxy = PrimitiveSceneInfoCompact.Proxy;
+	const EShadowDepthType ShadowDepthType = GetShadowDepthType();
 
 	{
 		const int32 ForcedLOD = CurrentView.Family->EngineShowFlags.LOD ? (GetCVarForceLODShadow_AnyThread() != -1 ? GetCVarForceLODShadow_AnyThread() : GetCVarForceLOD_AnyThread()) : -1;
@@ -1860,10 +1853,11 @@ bool FProjectedShadowInfo::ShouldDrawStaticMeshes_AnyThread(
 		FLODMask ShadowLODToRender = CalcAndUpdateLODToRender(CurrentView, FBoxSphereBounds(PrimitiveSceneInfoCompact.Bounds), PrimitiveSceneInfo, ForcedLOD);
 		OutResult.SetLodRange(ShadowLODToRender);
 
-		if (WholeSceneDirectionalShadow)
+		// Don't cache if it requires per view per mesh state for distance cull fade.
+		const bool bCanCache = !bMayBeFading;
+
+		if (bCanCache && UseCachedMeshDrawCommands(ShadowDepthType))
 		{
-			// Don't cache if it requires per view per mesh state for distance cull fade.
-			const bool bCanCache = !bMayBeFading;
 			int32 NumAcceptedStaticMeshes = 0;
 
 			for (int32 MeshIndex = 0; MeshIndex < PrimitiveSceneInfo->StaticMeshRelevances.Num(); MeshIndex++)
@@ -1873,15 +1867,7 @@ bool FProjectedShadowInfo::ShouldDrawStaticMeshes_AnyThread(
 
 				if (ShouldDrawStaticMesh(StaticMeshRelevance, ShadowLODToRender, bDrawingStaticMeshes))
 				{
-					if (bCanCache && GetShadowDepthType() == CSMShadowDepthType)
-					{
-						AddCachedMeshDrawCommands_AnyThread(PrimitiveSceneInfo->Scene, StaticMeshRelevance, MeshIndex, NumAcceptedStaticMeshes, OutResult, OutStats, OverflowBuffer);
-					}
-					else
-					{
-						++OutStats.NumMDCBuildRequests;
-						OutResult.AcceptMesh(NumAcceptedStaticMeshes++, MeshIndex, OverflowBuffer);
-					}
+					AddCachedMeshDrawCommands_AnyThread(PrimitiveSceneInfo->Scene, StaticMeshRelevance, MeshIndex, NumAcceptedStaticMeshes, OutResult, OutStats, OverflowBuffer);
 				}
 			}
 		}
@@ -2011,7 +1997,7 @@ bool FProjectedShadowInfo::AddSubjectPrimitive(FDynamicShadowsTaskData& TaskData
 				for (FViewInfo& CurrentView : Views)
 				{
 					// Note: skip small-mesh culling for VSM since it needs it drawn for GPU-side caching.
-					if (bWholeSceneShadow && CacheMode != SDCM_StaticPrimitivesOnly && MeshPassTargetType != EMeshPass::VSMShadowDepth)
+					if (bWholeSceneShadow && CacheMode != SDCM_StaticPrimitivesOnly && !bVSM)
 					{
 						const float DistanceSquared = ( Bounds.Origin - CurrentView.ShadowViewMatrices.GetViewOrigin() ).SizeSquared();
 						const float LODScaleSquared = FMath::Square(CurrentView.LODDistanceFactor);
@@ -2076,7 +2062,7 @@ bool FProjectedShadowInfo::TestPrimitiveFarCascadeConditions(bool bPrimitiveCast
 		return bPrimitiveCastsFarShadow;
 	}
 
-	if (GEnableNonNaniteVSM != 0 && MeshPassTargetType == EMeshPass::VSMShadowDepth)
+	if (GEnableNonNaniteVSM != 0 && bVSM)
 	{
 		const bool bWholeSceneDirectionalShadow = IsWholeSceneDirectionalShadow();
 
@@ -2355,14 +2341,16 @@ void FProjectedShadowInfo::FinalizeAddSubjectPrimitive(
 		int32 IdxBias;
 		int32 NumMDCs = Result.GetMDCIndices(Context, MDCIndices, IdxBias);
 
+		const EMeshPass::Type MeshPassType = GetShadowMeshPassType(GetShadowDepthType());
+
 		for (int32 Idx = 0; Idx < NumMDCs; ++Idx)
 		{
 			const int32 CmdIdx = (int32)MDCIndices[Idx] + IdxBias;
 			const FCachedMeshDrawCommandInfo& CmdInfo = PrimitiveSceneInfo->StaticMeshCommandInfos[CmdIdx];
 			const FScene* Scene = PrimitiveSceneInfo->Scene;
 			const FMeshDrawCommand* CachedCmd = CmdInfo.StateBucketId >= 0 ?
-				&Scene->CachedMeshDrawCommandStateBuckets[MeshPassTargetType].GetByElementId(CmdInfo.StateBucketId).Key :
-				&Scene->CachedDrawLists[MeshPassTargetType].MeshDrawCommands[CmdInfo.CommandIndex];
+				&Scene->CachedMeshDrawCommandStateBuckets[MeshPassType].GetByElementId(CmdInfo.StateBucketId).Key :
+				&Scene->CachedDrawLists[MeshPassType].MeshDrawCommands[CmdInfo.CommandIndex];
 			const EMeshDrawCommandCullingPayloadFlags CullingPayloadFlags = GetCullingPayloadFlags(Result.bIsLodRange, CmdInfo.CullingPayload.LodIndex == Result.LodRangeMin, CmdInfo.CullingPayload.LodIndex == Result.LodRangeMax);
 
 			const FMeshDrawCommandPrimitiveIdInfo PrimitiveIdInfo = PrimitiveSceneInfo->GetMDCIdInfo();
@@ -2441,8 +2429,7 @@ void FProjectedShadowInfo::SetupMeshDrawCommandsForShadowDepth(FSceneRenderer& R
 		Renderer.Scene->GetFeatureLevel(),
 		ShadowDepthView,
 		GetShadowDepthType(),
-		nullptr,
-		MeshPassTargetType);
+		nullptr);
 
 	if (Renderer.ShouldDumpMeshDrawCommandInstancingStats())
 	{
@@ -4262,7 +4249,7 @@ void FSceneRenderer::CreateWholeSceneProjectedShadow(
 
 					TSharedPtr<FVirtualShadowMapPerLightCacheEntry> VirtualSmPerLightCacheEntry = ShadowSceneRenderer->AddLocalLightShadow(ProjectedShadowInitializer, ProjectedShadowInfo, LightSceneInfo, MaxScreenRadius);
 
-					ProjectedShadowInfo->MeshPassTargetType = EMeshPass::VSMShadowDepth;
+					ProjectedShadowInfo->bVSM = true;
 					ProjectedShadowInfo->MeshSelectionMask = EShadowMeshSelection::VSM;
 
 					bool bContainsNaniteSubjects = false;
diff --git a/Engine/Source/Runtime/Renderer/Private/TranslucentLighting.cpp b/Engine/Source/Runtime/Renderer/Private/TranslucentLighting.cpp
--- a/Engine/Source/Runtime/Renderer/Private/TranslucentLighting.cpp
+++ b/Engine/Source/Runtime/Renderer/Private/TranslucentLighting.cpp
@@ -383,7 +383,6 @@ private:
 
 	FMeshPassProcessorRenderState PassDrawRenderState;
 	const FProjectedShadowInfo* ShadowInfo;
-	FShadowDepthType ShadowDepthType;
 	const bool bDirectionalLight;
 };
 
@@ -395,7 +394,6 @@ FTranslucencyDepthPassMeshProcessor::FTranslucencyDepthPassMeshProcessor(const F
 	: FMeshPassProcessor(EMeshPass::Num, Scene, Scene->GetFeatureLevel(), InViewIfDynamicMeshCommand, InDrawListContext)
 	, PassDrawRenderState(InPassDrawRenderState)
 	, ShadowInfo(InShadowInfo)
-	, ShadowDepthType(InShadowInfo->GetShadowDepthType())
 	, bDirectionalLight(InShadowInfo->bDirectionalLight)
 {
 }
diff --git a/Engine/Source/Runtime/Renderer/Public/MeshPassProcessor.h b/Engine/Source/Runtime/Renderer/Public/MeshPassProcessor.h
--- a/Engine/Source/Runtime/Renderer/Public/MeshPassProcessor.h
+++ b/Engine/Source/Runtime/Renderer/Public/MeshPassProcessor.h
@@ -42,6 +42,7 @@ namespace EMeshPass
 		SingleLayerWaterDepthPrepass,
 		CSMShadowDepth,
 		VSMShadowDepth,
+		OnePassPointLightShadowDepth,
 		Distortion,
 		Velocity,
 		TranslucentVelocity,
@@ -97,6 +98,7 @@ inline const TCHAR* GetMeshPassName(EMeshPass::Type MeshPass)
 	case EMeshPass::SingleLayerWaterDepthPrepass: return TEXT("SingleLayerWaterDepthPrepass");
 	case EMeshPass::CSMShadowDepth: return TEXT("CSMShadowDepth");
 	case EMeshPass::VSMShadowDepth: return TEXT("VSMShadowDepth");
+	case EMeshPass::OnePassPointLightShadowDepth: return TEXT("OnePassPointLightShadowDepth");
 	case EMeshPass::Distortion: return TEXT("Distortion");
 	case EMeshPass::Velocity: return TEXT("Velocity");
 	case EMeshPass::TranslucentVelocity: return TEXT("TranslucentVelocity");
@@ -133,9 +135,9 @@ inline const TCHAR* GetMeshPassName(EMeshPass::Type MeshPass)
 	}
 
 #if WITH_EDITOR
-	static_assert(EMeshPass::Num == 36 + 4, "Need to update switch(MeshPass) after changing EMeshPass"); // GUID to prevent incorrect auto-resolves, please change when changing the expression: {674D7D62-CFD8-4971-9A8D-CD91E5612CD8}
+	static_assert(EMeshPass::Num == 37 + 4, "Need to update switch(MeshPass) after changing EMeshPass"); // GUID to prevent incorrect auto-resolves, please change when changing the expression: {674D7D62-CFD8-4971-9A8D-CD91E5612CD8}
 #else
-	static_assert(EMeshPass::Num == 36, "Need to update switch(MeshPass) after changing EMeshPass"); // GUID to prevent incorrect auto-resolves, please change when changing the expression: {674D7D62-CFD8-4971-9A8D-CD91E5612CD8}
+	static_assert(EMeshPass::Num == 37, "Need to update switch(MeshPass) after changing EMeshPass"); // GUID to prevent incorrect auto-resolves, please change when changing the expression: {674D7D62-CFD8-4971-9A8D-CD91E5612CD8}
 #endif
 
 	checkf(0, TEXT("Missing case for EMeshPass %u"), (uint32)MeshPass);
