Fix bug with VT not streaming in texture editor.
When no scene is visible the VT system will only try and map pages requested by the texture editor CPU logic. This happens in the LoadPendingTiles() path.
But if the texture's root page is not yet mapped then other pages are blocked from streaming.
The fix is to make sure that we are also requesting the pending TilesToLock in LoadPendingTiles().

diff --git a/Engine/Source/Runtime/Renderer/Private/VT/VirtualTextureSystem.cpp b/Engine/Source/Runtime/Renderer/Private/VT/VirtualTextureSystem.cpp
--- a/Engine/Source/Runtime/Renderer/Private/VT/VirtualTextureSystem.cpp
+++ b/Engine/Source/Runtime/Renderer/Private/VT/VirtualTextureSystem.cpp
@@ -1234,6 +1234,7 @@ void FVirtualTextureSystem::LoadPendingTiles(FRDGBuilder& GraphBuilder, ERHIFeat
 
 		FUniqueRequestList* RequestList = Allocator.Create<FUniqueRequestList>(Allocator);
 		RequestList->Initialize();
+		GatherLockedTileRequests(RequestList);
 		GatherRequests(RequestList, UniquePageList, Frame, Allocator, Settings);
 		// No need to sort requests, since we're submitting all of them here (no throttling)
 		AllocateResources(GraphBuilder);
