Avoid recreation of drawmesh commands when switching SkyLight cube within one frame.

diff --git a/Engine/Source/Runtime/Renderer/Private/RendererScene.cpp b/Engine/Source/Runtime/Renderer/Private/RendererScene.cpp
--- a/Engine/Source/Runtime/Renderer/Private/RendererScene.cpp
+++ b/Engine/Source/Runtime/Renderer/Private/RendererScene.cpp
@@ -1630,6 +1630,7 @@ FScene::FScene(UWorld* InWorld, bool bInRequiresHitProxies, bool bInIsEditorScen
 :	FSceneInterface(InFeatureLevel)
 ,	World(InWorld)
 ,	FXSystem(nullptr)
+,	bCachedShouldRenderSkylightInBasePass(false)
 ,	bScenesPrimitivesNeedStaticMeshElementUpdate(false)
 ,   PathTracingInvalidationCounter(0)
 #if RHI_RAYTRACING
@@ -2871,20 +2872,10 @@ void FScene::SetSkyLight(FSkyLightSceneProxy* LightProxy)
 		{
 			check(!Scene->SkyLightStack.Contains(LightProxy));
 			Scene->SkyLightStack.Push(LightProxy);
-			const bool bOriginalHadSkylight = Scene->ShouldRenderSkylightInBasePass(false);
 
 			// Use the most recently enabled skylight
 			Scene->SkyLight = LightProxy;
 
-			const bool bNewHasSkylight = Scene->ShouldRenderSkylightInBasePass(false);
-
-			if (bOriginalHadSkylight != bNewHasSkylight)
-			{
-				// Mark the scene as needing static draw lists to be recreated if needed
-				// The base pass chooses shaders based on whether there's a skylight in the scene, and that is cached in static draw lists
-				Scene->bScenesPrimitivesNeedStaticMeshElementUpdate = true;
-				UE_CLOG(!GIsEditor, LogRenderer, Log, TEXT("Forcing update for all mesh draw commands: Enable SkyLight"));
-			}
 			Scene->InvalidatePathTracedOutput();
 		});
 }
@@ -2899,8 +2890,6 @@ void FScene::DisableSkyLight(FSkyLightSceneProxy* LightProxy)
 	ENQUEUE_RENDER_COMMAND(FDisableSkyLightCommand)(
 		[Scene, LightProxy] (FRHICommandListBase&)
 	{
-		const bool bOriginalHadSkylight = Scene->ShouldRenderSkylightInBasePass(false);
-
 		Scene->SkyLightStack.RemoveSingle(LightProxy);
 
 		if (Scene->SkyLightStack.Num() > 0)
@@ -2913,14 +2902,6 @@ void FScene::DisableSkyLight(FSkyLightSceneProxy* LightProxy)
 			Scene->SkyLight = NULL;
 		}
 
-		const bool bNewHasSkylight = Scene->ShouldRenderSkylightInBasePass(false);
-
-		// Update the scene if we switched skylight enabled states
-		if (bOriginalHadSkylight != bNewHasSkylight)
-		{
-			Scene->bScenesPrimitivesNeedStaticMeshElementUpdate = true;
-			UE_CLOG(!GIsEditor, LogRenderer, Log, TEXT("Forcing update for all mesh draw commands: Disable SkyLight"));
-		}
 		Scene->InvalidatePathTracedOutput();
 	});
 }
@@ -6559,6 +6540,17 @@ void FScene::Update(FRDGBuilder& GraphBuilder, const FUpdateParameters& Paramete
 	{
 		SCOPED_NAMED_EVENT(UpdateStaticMeshes, FColor::Emerald);
 		
+		const bool bLastFrameShouldRenderSkylightInBasePass = bCachedShouldRenderSkylightInBasePass;
+		bCachedShouldRenderSkylightInBasePass = ShouldRenderSkylightInBasePass(false);
+
+		if (bCachedShouldRenderSkylightInBasePass != bLastFrameShouldRenderSkylightInBasePass)
+		{
+			// Mark the scene as needing static draw lists to be recreated if needed
+			// The base pass chooses shaders based on whether there's a skylight in the scene, and that is cached in static draw lists
+			UE_CLOG(!GIsEditor, LogRenderer, Log, TEXT("Forcing update for all mesh draw commands: SkyLight change"));
+			bScenesPrimitivesNeedStaticMeshElementUpdate = true;
+		}
+
 		if (bScenesPrimitivesNeedStaticMeshElementUpdate || CachedDefaultBasePassDepthStencilAccess != DefaultBasePassDepthStencilAccess)
 		{
 			// Mark all primitives as needing an update
diff --git a/Engine/Source/Runtime/Renderer/Private/ScenePrivate.h b/Engine/Source/Runtime/Renderer/Private/ScenePrivate.h
--- a/Engine/Source/Runtime/Renderer/Private/ScenePrivate.h
+++ b/Engine/Source/Runtime/Renderer/Private/ScenePrivate.h
@@ -3002,6 +3002,9 @@ public:
 	/** Default base pass depth stencil access used to cache mesh draw commands. */
 	FExclusiveDepthStencil::Type CachedDefaultBasePassDepthStencilAccess;
 
+	/** Previous frame SkyLight state. */
+	bool bCachedShouldRenderSkylightInBasePass;
+
 	/** True if a change to SkyLight / Lighting has occurred that requires static draw lists to be updated. */
 	bool bScenesPrimitivesNeedStaticMeshElementUpdate;
 
