Fixed VSM cache not invalidating properly on pagepool size change when AllocateViaLRU=1

diff --git a/Engine/Source/Runtime/Renderer/Private/VirtualShadowMaps/VirtualShadowMapCacheManager.cpp b/Engine/Source/Runtime/Renderer/Private/VirtualShadowMaps/VirtualShadowMapCacheManager.cpp
--- a/Engine/Source/Runtime/Renderer/Private/VirtualShadowMaps/VirtualShadowMapCacheManager.cpp
+++ b/Engine/Source/Runtime/Renderer/Private/VirtualShadowMaps/VirtualShadowMapCacheManager.cpp
@@ -765,6 +765,8 @@ void FVirtualShadowMapArrayCacheManager::Invalidate(FRDGBuilder& GraphBuilder)
 	// Clear the cache
 	CacheEntries.Reset();
 
+	PrevBuffers = FVirtualShadowMapArrayFrameData();
+
 	//UE_LOG(LogRenderer, Display, TEXT("Virtual shadow map cache invalidated."));
 
 	// Clear the physical page metadata (on all GPUs)
