Fixed incorrect culling volume space for scene culling of local lights. The convex volume was set up in pre-translated space, but tested in world space.

diff --git a/Engine/Source/Runtime/Renderer/Private/SceneCulling/SceneCulling.cpp b/Engine/Source/Runtime/Renderer/Private/SceneCulling/SceneCulling.cpp
--- a/Engine/Source/Runtime/Renderer/Private/SceneCulling/SceneCulling.cpp
+++ b/Engine/Source/Runtime/Renderer/Private/SceneCulling/SceneCulling.cpp
@@ -496,7 +496,7 @@ inline FSceneCulling::FFootprint64 ToLevelRelative(const FSceneCulling::FFootpri
 	return Result;
 };
 
-void FSceneCulling::TestConvexVolume(const FConvexVolume& ViewCullVolume, TArray<FCellDraw, SceneRenderingAllocator>& OutCellDraws, uint32 ViewGroupId, uint32 MaxNumViews, uint32& OutNumInstanceGroups)
+void FSceneCulling::TestConvexVolume(const FConvexVolume& ViewCullVolume, const FVector3d &WorldToVolumeTranslation, TArray<FCellDraw, SceneRenderingAllocator>& OutCellDraws, uint32 ViewGroupId, uint32 MaxNumViews, uint32& OutNumInstanceGroups)
 {
 	LLM_SCOPE_BYTAG(SceneCulling);
 
@@ -523,7 +523,7 @@ void FSceneCulling::TestConvexVolume(const FConvexVolume& ViewCullVolume, TArray
 		const double LevelCellSize = SpatialHash.GetCellSize(BlockLoc.GetLevel() - FSpatialHash::CellBlockDimLog2);
 		// Extend extent by half a cell size in all directions
 		FVector3d BlockBoundsExtent = FVector((BlockLevelSize + LevelCellSize) * 0.5);
-		FOutcode BlockCullResult = ViewCullVolume.GetBoxIntersectionOutcode(BlockBoundsCenter, BlockBoundsExtent);
+		FOutcode BlockCullResult = ViewCullVolume.GetBoxIntersectionOutcode(BlockBoundsCenter + WorldToVolumeTranslation, BlockBoundsExtent);
 
 		if (BlockCullResult.GetInside())
 		{
@@ -563,7 +563,7 @@ void FSceneCulling::TestConvexVolume(const FConvexVolume& ViewCullVolume, TArray
 					FVector3d CellCenter = FVector3d(CellCoord) * LevelCellSize + MinCellCenter;
 
 					// emit for intersecting cells
-					bool bCellIntersects = ViewCullVolume.IntersectBox(CellCenter, CellBoundsExtent);
+					bool bCellIntersects = ViewCullVolume.IntersectBox(CellCenter, WorldToVolumeTranslation, CellBoundsExtent);
 					if (bCellIntersects)
 					{
 						FCellHeader CellHeader = UnpackCellHeader(CellHeaders[CellId]);
@@ -711,7 +711,7 @@ void FSceneCulling::Test(const FCullingVolume& CullingVolume, TArray<FCellDraw,
 			return;
 		}
 	}
-	TestConvexVolume(CullingVolume.ConvexVolume, OutCellDraws, ViewGroupId, MaxNumViews, OutNumInstanceGroups);
+	TestConvexVolume(CullingVolume.ConvexVolume, CullingVolume.WorldToVolumeTranslation, OutCellDraws, ViewGroupId, MaxNumViews, OutNumInstanceGroups);
 }
 
 void FSceneCulling::Empty()
diff --git a/Engine/Source/Runtime/Renderer/Private/SceneCulling/SceneCulling.h b/Engine/Source/Runtime/Renderer/Private/SceneCulling/SceneCulling.h
--- a/Engine/Source/Runtime/Renderer/Private/SceneCulling/SceneCulling.h
+++ b/Engine/Source/Runtime/Renderer/Private/SceneCulling/SceneCulling.h
@@ -25,7 +25,10 @@ struct FPrimitiveBounds;
  */
 struct FCullingVolume
 {
+	// Negative translation to add to the tested location prior to testing the ConvexVolume.
+	FVector3d WorldToVolumeTranslation; 
 	FConvexVolume ConvexVolume;
+	// Bounding sphere in world space, if radius is zero OR the footprint is <= r.SceneCulling.SmallFootprintSideThreshold, the ConvexVolume is used
 	FSphere3d Sphere = FSphere(ForceInit);
 };
 
@@ -83,7 +86,7 @@ public:
 
 	void PublishStats();
 
-	void TestConvexVolume(const FConvexVolume& ViewCullVolume, TArray<FCellDraw, SceneRenderingAllocator>& OutCellDraws, uint32 ViewGroupId, uint32 MaxNumViews, uint32& OutNumInstanceGroups);
+	void TestConvexVolume(const FConvexVolume& ViewCullVolume, const FVector3d &WorldToVolumeTranslation, TArray<FCellDraw, SceneRenderingAllocator>& OutCellDraws, uint32 ViewGroupId, uint32 MaxNumViews, uint32& OutNumInstanceGroups);
 
 	void TestSphere(const FSphere& Sphere, TArray<FCellDraw, SceneRenderingAllocator>& OutCellDraws, uint32 ViewGroupId, uint32 MaxNumViews, uint32& OutNumInstanceGroups);
 
diff --git a/Engine/Source/Runtime/Renderer/Private/SceneCulling/SceneCullingRenderer.cpp b/Engine/Source/Runtime/Renderer/Private/SceneCulling/SceneCullingRenderer.cpp
--- a/Engine/Source/Runtime/Renderer/Private/SceneCulling/SceneCullingRenderer.cpp
+++ b/Engine/Source/Runtime/Renderer/Private/SceneCulling/SceneCullingRenderer.cpp
@@ -66,6 +66,8 @@ FSceneInstanceCullingQuery* FSceneCullingRenderer::CullInstances(FRDGBuilder& Gr
 		for (int32 Index = 0; Index < ViewCullVolumes.Num(); ++Index)
 		{
 			FCullingVolume CullingVolume;
+			// Assume world-space
+			CullingVolume.WorldToVolumeTranslation = FVector3d::ZeroVector;
 			CullingVolume.ConvexVolume = ViewCullVolumes[Index];
 			Query->Add(Index, 1, 1, CullingVolume);
 		}
diff --git a/Engine/Source/Runtime/Renderer/Private/VirtualShadowMaps/VirtualShadowMapArray.cpp b/Engine/Source/Runtime/Renderer/Private/VirtualShadowMaps/VirtualShadowMapArray.cpp
--- a/Engine/Source/Runtime/Renderer/Private/VirtualShadowMaps/VirtualShadowMapArray.cpp
+++ b/Engine/Source/Runtime/Renderer/Private/VirtualShadowMaps/VirtualShadowMapArray.cpp
@@ -2510,6 +2510,8 @@ static FCullingVolume GetCullingVolume(const FProjectedShadowInfo* ProjectedShad
 {
 	FCullingVolume CullingVolume;
 
+	CullingVolume.WorldToVolumeTranslation = ProjectedShadowInfo->PreShadowTranslation;
+
 	uint32 NumPrimaryViews = 0;
 	if (ProjectedShadowInfo->VirtualShadowMapClipmap)
 	{
@@ -2519,6 +2521,8 @@ static FCullingVolume GetCullingVolume(const FProjectedShadowInfo* ProjectedShad
 		if (!bIsCached && !ProjectedShadowInfo->CascadeSettings.ShadowBoundsAccurate.Planes.IsEmpty())
 		{
 			CullingVolume.ConvexVolume = ProjectedShadowInfo->CascadeSettings.ShadowBoundsAccurate;
+			// ShadowBoundsAccurate is in world-space
+			CullingVolume.WorldToVolumeTranslation = FVector3d::ZeroVector;
 		}
 		else
 		{
