Fix some issues related with CSM caching.
1. Fix an issue that the LastFrameExtraStaticShadowSubjects could be overridden when there are multiple Packets in FGatherShadowPrimitivesPacket.
2. Uses atomic add to count LastFrameExtraStaticShadowSubjects for multi-thread rendering.
3. Invalid the cached shadow map if the LastFrameExtraStaticShadowSubjects is greater than a throttle for both SDCM_MovablePrimitivesOnly and SDCM_CSMScrolling

diff --git a/Engine/Source/Runtime/Renderer/Private/ShadowSetup.cpp b/Engine/Source/Runtime/Renderer/Private/ShadowSetup.cpp
--- a/Engine/Source/Runtime/Renderer/Private/ShadowSetup.cpp
+++ b/Engine/Source/Runtime/Renderer/Private/ShadowSetup.cpp
@@ -2300,7 +2300,7 @@ void FProjectedShadowInfo::FinalizeAddSubjectPrimitive(
 		else if (!PrimitiveSceneInfo->Proxy->IsMeshShapeOftenMoving())
 		{
 			// Count the number of extra draw calls of static meshes for filling the scrolling area.
-			CachedShadowMapData.LastFrameExtraStaticShadowSubjects += 1;
+			FPlatformAtomics::InterlockedAdd(&CachedShadowMapData.LastFrameExtraStaticShadowSubjects, 1);
 		}
 	}
 
@@ -3802,7 +3802,8 @@ void ComputeViewDependentWholeSceneShadowCacheModes(
 					&& ProjectedShadowInitializer.MinLightW == CachedShadowInitializer.MinLightW
 					&& ProjectedShadowInitializer.MaxDistanceToCastInLightW == CachedShadowInitializer.MaxDistanceToCastInLightW;
 
-				if (CachedShadowMapData->ShadowMap.IsValid() && CachedShadowMapData->ShadowBufferResolution == ShadowMapSize)
+				// Reuse the cached shadow map if it is valid and the extra draw calls of the static meshes are less than a throttle
+				if (CachedShadowMapData->ShadowMap.IsValid() && CachedShadowMapData->ShadowBufferResolution == ShadowMapSize && CachedShadowMapData->LastFrameExtraStaticShadowSubjects <= CVarCSMScrollingMaxExtraStaticShadowSubjects.GetValueOnAnyThread())
 				{
 					if (bExactlyEqual)
 					{
@@ -3833,8 +3834,8 @@ void ComputeViewDependentWholeSceneShadowCacheModes(
 
 						float OverlappedAreaRatio = (OverlappedArea.X * OverlappedArea.Y) / (CachedShadowMapData->Initializer.SubjectBounds.SphereRadius * CachedShadowMapData->Initializer.SubjectBounds.SphereRadius * 4);
 
-						// if the overlapped area of the cached shadow map and the current shadow map smaller than a throttle or the extra draw calls of the static meshes when scrolling the shadow are greater than a throttle, update the cached shadow map, otherwise scroll the shadow map instead.
-						if (OverlappedAreaRatio > CVarCSMScrollingOverlapAreaThrottle.GetValueOnAnyThread() && CachedShadowMapData->LastFrameExtraStaticShadowSubjects <= CVarCSMScrollingMaxExtraStaticShadowSubjects.GetValueOnAnyThread())
+						// if the overlapped area of the cached shadow map and the current shadow map smaller than a throttle, update the cached shadow map, otherwise scroll the shadow map instead.
+						if (OverlappedAreaRatio > CVarCSMScrollingOverlapAreaThrottle.GetValueOnAnyThread())
 						{
 							OutNumShadowMaps = 1;
 							OutCacheModes[0] = SDCM_CSMScrolling;
@@ -5006,18 +5007,6 @@ struct FGatherShadowPrimitivesPacket
 			Context.OverflowedMDCIndices = OverflowBuffer.MDCIndices.GetData();
 			Context.OverflowedMeshIndices = OverflowBuffer.MeshIndices.GetData();
 
-			FCachedShadowMapData* CachedShadowMapData = nullptr;
-
-			if (ProjectedShadowInfo->CacheMode != SDCM_Uncached && ProjectedShadowInfo->IsWholeSceneDirectionalShadow())
-			{
-				CachedShadowMapData = const_cast<FCachedShadowMapData*>(TaskData.Scene->GetCachedShadowMapData(ProjectedShadowInfo->GetLightSceneInfo().Id, ProjectedShadowInfo->CascadeSettings.ShadowSplitIndex));
-
-				if (ProjectedShadowInfo->CacheMode == SDCM_MovablePrimitivesOnly || ProjectedShadowInfo->CacheMode == SDCM_CSMScrolling)
-				{
-					CachedShadowMapData->LastFrameExtraStaticShadowSubjects = 0;
-				}
-			}
-
 			for (const FAddSubjectPrimitiveOp& PrimitiveOp : ViewDependentWholeSceneShadowSubjectPrimitives[ShadowIndex])
 			{
 				ProjectedShadowInfo->FinalizeAddSubjectPrimitive(TaskData, PrimitiveOp, TArrayView<FViewInfo>(), Context);
@@ -5025,6 +5014,7 @@ struct FGatherShadowPrimitivesPacket
 
 			if (ProjectedShadowInfo->CacheMode == SDCM_StaticPrimitivesOnly && ProjectedShadowInfo->IsWholeSceneDirectionalShadow())
 			{
+				FCachedShadowMapData* CachedShadowMapData = const_cast<FCachedShadowMapData*>(TaskData.Scene->GetCachedShadowMapData(ProjectedShadowInfo->GetLightSceneInfo().Id, ProjectedShadowInfo->CascadeSettings.ShadowSplitIndex));
 				checkSlow(CachedShadowMapData != nullptr);
 				CachedShadowMapData->bCachedShadowMapHasPrimitives = ProjectedShadowInfo->HasSubjectPrims();
 			}
@@ -5494,17 +5484,19 @@ void FSceneRenderer::AddViewDependentWholeSceneShadowsForView(
 							{
 								ShadowInfosThatNeedCulling.Add(ProjectedShadowInfo);
 
-								if (CacheMode[CacheModeIndex] == SDCM_StaticPrimitivesOnly || CacheMode[CacheModeIndex] == SDCM_CSMScrolling)
+								if (CacheMode[CacheModeIndex] != SDCM_Uncached)
 								{
 									FCachedShadowMapData& CachedShadowMapData = Scene->GetCachedShadowMapDataRef(LightSceneInfo.Id, ProjectedShadowInfo->CascadeSettings.ShadowSplitIndex);
 
+									CachedShadowMapData.LastFrameExtraStaticShadowSubjects = 0;
+
 									if (CacheMode[CacheModeIndex] == SDCM_StaticPrimitivesOnly)
 									{
 										CachedShadowMapData.MaxSubjectZ = ProjectedShadowInfo->MaxSubjectZ;
 										CachedShadowMapData.MinSubjectZ = ProjectedShadowInfo->MinSubjectZ;
 										CachedShadowMapData.PreShadowTranslation = ProjectedShadowInfo->PreShadowTranslation;
 									}
-									else // CacheMode[CacheModeIndex] == SDCM_CSMScrolling
+									else if (CacheMode[CacheModeIndex] == SDCM_CSMScrolling)
 									{
 										const FVector FaceDirection(1, 0, 0);
 										FVector	XAxis, YAxis;
