Removed additional ensure to check mesh draw command inline vertex stream size.

diff --git a/Engine/Source/Runtime/Renderer/Private/MeshPassProcessor.cpp b/Engine/Source/Runtime/Renderer/Private/MeshPassProcessor.cpp
--- a/Engine/Source/Runtime/Renderer/Private/MeshPassProcessor.cpp
+++ b/Engine/Source/Runtime/Renderer/Private/MeshPassProcessor.cpp
@@ -2071,7 +2071,6 @@ void FCachedPassMeshDrawListContext::FinalizeCommandCommon(
 		MeshDrawCommand.ClearDebugPrimitiveSceneProxy(); //When using State Buckets multiple PrimitiveSceneProxies use the same MeshDrawCommand, so The PrimitiveSceneProxy pointer can't be stored.
 	}
 #endif
-	ensureMsgf(MeshDrawCommand.VertexStreams.GetAllocatedSize() == 0, TEXT("Cached Mesh Draw command overflows VertexStreams. VertexStream inline size should be tweaked."));
 
 #if DO_GUARD_SLOW
 	if (bUseGPUScene)
