Fixed incorrect per-instance custom data/random in raytracing when using FInstancedStaticMeshSceneProxy.

diff --git a/Engine/Source/Runtime/Engine/Private/InstancedStaticMesh.cpp b/Engine/Source/Runtime/Engine/Private/InstancedStaticMesh.cpp
--- a/Engine/Source/Runtime/Engine/Private/InstancedStaticMesh.cpp
+++ b/Engine/Source/Runtime/Engine/Private/InstancedStaticMesh.cpp
@@ -1864,7 +1864,7 @@ void FInstancedStaticMeshSceneProxy::GetDynamicRayTracingInstances(struct FRayTr
 
 	//preallocate the worst-case to prevent an explosion of reallocs
 	//#dxr_todo: possibly track used instances and reserve based on previous behavior
-	RayTracingInstanceTemplate.InstanceTransforms.Reserve(InstanceCount);
+	RayTracingInstanceTemplate.PrimitiveInstanceIndices.Reserve(InstanceCount);
 
 	// Add all visible instances
 	for (FVisibleInstance VisibleInstance : VisibleInstances)
@@ -1888,7 +1888,7 @@ void FInstancedStaticMeshSceneProxy::GetDynamicRayTracingInstances(struct FRayTr
 				ActiveInstances[DynamicInstanceIdx] = OutRayTracingInstances.Num();
 				FRayTracingInstance& RayTracingInstance = OutRayTracingInstances.Add_GetRef(RayTracingWPOInstanceTemplate);
 				RayTracingInstance.Geometry = &DynamicData.DynamicGeometry;
-				RayTracingInstance.InstanceTransforms.Reserve(InstanceCount);
+				RayTracingInstance.PrimitiveInstanceIndices.Reserve(InstanceCount);
 
 				DynamicInstance = &RayTracingInstance;
 
@@ -1913,16 +1913,16 @@ void FInstancedStaticMeshSceneProxy::GetDynamicRayTracingInstances(struct FRayTr
 				DynamicInstance = &OutRayTracingInstances[ActiveInstances[DynamicInstanceIdx]];
 			}
 
-			DynamicInstance->InstanceTransforms.Add(InstanceToWorld);
+			DynamicInstance->PrimitiveInstanceIndices.Add(InstanceIndex);
 
 		}
 		else
 		{
-			RayTracingInstanceTemplate.InstanceTransforms.Emplace(InstanceToWorld);
+			RayTracingInstanceTemplate.PrimitiveInstanceIndices.Add(InstanceIndex);
 		}
 	}
 
-	if (RayTracingInstanceTemplate.InstanceTransforms.Num() > 0)
+	if (RayTracingInstanceTemplate.PrimitiveInstanceIndices.Num() > 0)
 	{
 		int32 SectionCount = InstancedRenderData.LODModels[LODIndex].Sections.Num();
 
diff --git a/Engine/Source/Runtime/Renderer/Private/RayTracing/RayTracing.cpp b/Engine/Source/Runtime/Renderer/Private/RayTracing/RayTracing.cpp
--- a/Engine/Source/Runtime/Renderer/Private/RayTracing/RayTracing.cpp
+++ b/Engine/Source/Runtime/Renderer/Private/RayTracing/RayTracing.cpp
@@ -961,8 +961,24 @@ namespace RayTracing
 						RayTracingInstance.Flags = SceneInfo->CachedRayTracingInstance.Flags;
 						AddDebugRayTracingInstanceFlags(RayTracingInstance.Flags);
 
-						if(!Instance.GetTransforms().IsEmpty())
+						if (!Instance.GetPrimitiveInstanceIndices().IsEmpty())
 						{
+							TConstArrayView<uint32> PrimitiveInstanceIndices = Instance.GetPrimitiveInstanceIndices();
+
+							// Convert from instance indices to InstanceSceneDataOffsets
+							TArrayView<uint32> InstanceSceneDataOffsets = RayTracingScene.Allocate<uint32>(PrimitiveInstanceIndices.Num());
+							for (int32 InstanceIndex = 0; InstanceIndex < PrimitiveInstanceIndices.Num(); ++InstanceIndex)
+							{
+								InstanceSceneDataOffsets[InstanceIndex] = SceneInfo->GetInstanceSceneDataOffset() + PrimitiveInstanceIndices[InstanceIndex];
+							}
+
+							RayTracingInstance.InstanceSceneDataOffsets = InstanceSceneDataOffsets;
+							RayTracingInstance.UserData = InstanceSceneDataOffsets;
+							RayTracingInstance.NumTransforms = PrimitiveInstanceIndices.Num();
+						}
+						else if(!Instance.GetTransforms().IsEmpty())
+						{
+							TConstArrayView<FMatrix> TransformsView;
 							if (Instance.OwnsTransforms())
 							{
 								// Slow path: copy transforms to the owned storage
@@ -971,16 +987,17 @@ namespace RayTracing
 								FMemory::Memcpy(SceneOwnedTransforms.GetData(), Instance.InstanceTransforms.GetData(), Instance.InstanceTransforms.Num() * sizeof(RayTracingInstance.Transforms[0]));
 								static_assert(std::is_same_v<decltype(SceneOwnedTransforms[0]), decltype(Instance.InstanceTransforms[0])>, "Unexpected transform type");
 
-								RayTracingInstance.NumTransforms = SceneOwnedTransforms.Num();
-								RayTracingInstance.Transforms = SceneOwnedTransforms;
+								TransformsView = SceneOwnedTransforms;
 							}
 							else
 							{
 								// Fast path: just reference persistently-allocated transforms and avoid a copy
 								checkf(Instance.InstanceTransforms.Num() == 0, TEXT("InstanceTransforms is expected to be empty if using InstanceTransformsView"));
-								RayTracingInstance.NumTransforms = Instance.InstanceTransformsView.Num();
-								RayTracingInstance.Transforms = Instance.InstanceTransformsView;
+								TransformsView = Instance.InstanceTransformsView;
 							}
+
+							RayTracingInstance.NumTransforms = TransformsView.Num();
+							RayTracingInstance.Transforms = TransformsView;
 						}
 						else
 						{
diff --git a/Engine/Source/Runtime/Renderer/Public/RayTracingInstance.h b/Engine/Source/Runtime/Renderer/Public/RayTracingInstance.h
--- a/Engine/Source/Runtime/Renderer/Public/RayTracingInstance.h
+++ b/Engine/Source/Runtime/Renderer/Public/RayTracingInstance.h
@@ -59,7 +59,7 @@ struct FRayTracingInstance
 	TArray<FMeshBatch> Materials;
 
 	/** Similar to Materials, but memory is owned by someone else (i.g. FPrimitiveSceneProxy). */
-	TArrayView<const FMeshBatch> MaterialsView;
+	TConstArrayView<FMeshBatch> MaterialsView;
 
 	bool OwnsMaterials() const
 	{
@@ -104,23 +104,48 @@ struct FRayTracingInstance
 	*/
 	uint32 NumTransforms = 0;
 
+	// Indices of primitive instances to be included in ray tracing scene
+	TArray<uint32> PrimitiveInstanceIndices;
+
+	/** Similar to PrimitiveInstanceIndices, but memory is owned by someone else (i.g. FPrimitiveSceneProxy). */
+	TConstArrayView<uint32> PrimitiveInstanceIndicesView;
+
+	bool OwnsPrimitiveInstanceIndices() const
+	{
+		return PrimitiveInstanceIndices.Num() != 0;
+	}
+
+	TConstArrayView<uint32> GetPrimitiveInstanceIndices() const
+	{
+		if (OwnsPrimitiveInstanceIndices())
+		{
+			check(PrimitiveInstanceIndicesView.Num() == 0);
+			return TConstArrayView<uint32>(PrimitiveInstanceIndices);
+		}
+		else
+		{
+			check(PrimitiveInstanceIndices.Num() == 0);
+			return PrimitiveInstanceIndicesView;
+		}
+	}
+
 	/** Instance transforms. */
 	TArray<FMatrix> InstanceTransforms;
 
 	/** Similar to InstanceTransforms, but memory is owned by someone else (i.g. FPrimitiveSceneProxy). */
-	TArrayView<const FMatrix> InstanceTransformsView;
+	TConstArrayView<FMatrix> InstanceTransformsView;
 
 	bool OwnsTransforms() const
 	{
 		return InstanceTransforms.Num() != 0;
 	}
 
-	TArrayView<const FMatrix> GetTransforms() const
+	TConstArrayView<FMatrix> GetTransforms() const
 	{
 		if (OwnsTransforms())
 		{
 			check(InstanceTransformsView.Num() == 0);
-			return TArrayView<const FMatrix>(InstanceTransforms);
+			return TConstArrayView<FMatrix>(InstanceTransforms);
 		}
 		else
 		{
