Fix crash when using the level instance editor with Instanced Stereo enabled due to a missing InstancedView uniform buffer for the visualization shader.

diff --git a/Engine/Source/Runtime/Renderer/Private/PostProcess/PostProcessVisualizeLevelInstance.cpp b/Engine/Source/Runtime/Renderer/Private/PostProcess/PostProcessVisualizeLevelInstance.cpp
--- a/Engine/Source/Runtime/Renderer/Private/PostProcess/PostProcessVisualizeLevelInstance.cpp
+++ b/Engine/Source/Runtime/Renderer/Private/PostProcess/PostProcessVisualizeLevelInstance.cpp
@@ -38,7 +38,7 @@ IMPLEMENT_GLOBAL_SHADER(FVisualizeLevelInstancePS, "/Engine/Private/PostProcessV
 
 
 BEGIN_SHADER_PARAMETER_STRUCT(FVisualizeLevelInstancePassParameters, )
-	SHADER_PARAMETER_STRUCT_REF(FViewUniformShaderParameters, View)
+	SHADER_PARAMETER_STRUCT_INCLUDE(FViewShaderParameters, View)
 	SHADER_PARAMETER_STRUCT_INCLUDE(FSceneTextureShaderParameters, SceneTextures)
 	SHADER_PARAMETER_STRUCT_INCLUDE(FInstanceCullingDrawParams, InstanceCullingDrawParams)
 	RENDER_TARGET_BINDING_SLOTS()
@@ -88,7 +88,7 @@ FScreenPassTexture AddVisualizeLevelInstancePass(
 		{
 			auto* PassParameters = GraphBuilder.AllocParameters<FVisualizeLevelInstancePassParameters>();
 
-			PassParameters->View = EditorView->ViewUniformBuffer;
+			PassParameters->View = EditorView->GetShaderParameters();
 			PassParameters->SceneTextures = Inputs.SceneTextures;
 			PassParameters->RenderTargets.DepthStencil = FDepthStencilBinding(
 				DepthStencilTexture,
