HWRT: Turn r.RayTracing.CompileMaterialCHS/AHS into per-platform cvars.

diff --git a/Engine/Source/Runtime/RenderCore/Private/Shader.cpp b/Engine/Source/Runtime/RenderCore/Private/Shader.cpp
--- a/Engine/Source/Runtime/RenderCore/Private/Shader.cpp
+++ b/Engine/Source/Runtime/RenderCore/Private/Shader.cpp
@@ -2619,16 +2619,17 @@ void ShaderMapAppendKey(EShaderPlatform Platform, FShaderKeyGenerator& KeyGen)
 
 	if (ShouldCompileRayTracingShadersForProject(Platform))
 	{
-		static const auto CVarCompileCHS = IConsoleManager::Get().FindConsoleVariable(TEXT("r.RayTracing.CompileMaterialCHS"));
-		static const auto CVarCompileAHS = IConsoleManager::Get().FindConsoleVariable(TEXT("r.RayTracing.CompileMaterialAHS"));
+		static FShaderPlatformCachedIniValue<int32> CVarCompileRayTracingMaterialCHS(TEXT("r.RayTracing.CompileMaterialCHS"));
+		static FShaderPlatformCachedIniValue<int32> CVarCompileRayTracingMaterialAHS(TEXT("r.RayTracing.CompileMaterialAHS"));
+
 		static const auto CVarTextureLod = IConsoleManager::Get().FindConsoleVariable(TEXT("r.RayTracing.UseTextureLod"));
 
 		KeyGen.AppendSeparator();
 		KeyGen.Append(TEXT("RAY"));
 		KeyGen.AppendDebugText(TEXT("-CHS"));
-		KeyGen.AppendBoolInt(CVarCompileCHS && CVarCompileCHS->GetBool());
+		KeyGen.AppendBoolInt(CVarCompileRayTracingMaterialCHS.Get(Platform) != 0);
 		KeyGen.AppendDebugText(TEXT("AHS"));
-		KeyGen.AppendBoolInt(CVarCompileAHS && CVarCompileAHS->GetBool());
+		KeyGen.AppendBoolInt(CVarCompileRayTracingMaterialAHS.Get(Platform) != 0);
 		KeyGen.AppendDebugText(TEXT("LOD"));
 		KeyGen.AppendBoolInt(CVarTextureLod && CVarTextureLod->GetBool());
 	}
diff --git a/Engine/Source/Runtime/Renderer/Private/RayTracing/RayTracingMaterialHitShaders.cpp b/Engine/Source/Runtime/Renderer/Private/RayTracing/RayTracingMaterialHitShaders.cpp
--- a/Engine/Source/Runtime/Renderer/Private/RayTracing/RayTracingMaterialHitShaders.cpp
+++ b/Engine/Source/Runtime/Renderer/Private/RayTracing/RayTracingMaterialHitShaders.cpp
@@ -20,6 +20,7 @@
 #include "RayTracingDecals.h"
 #include "PathTracing.h"
 #include "RendererModule.h"
+#include "ShaderPlatformCachedIniValue.h"
 
 int32 GEnableRayTracingMaterials = 1;
 static FAutoConsoleVariableRef CVarEnableRayTracingMaterials(
@@ -66,6 +67,14 @@ static bool IsSupportedVertexFactoryType(const FVertexFactoryType* VertexFactory
 	return VertexFactoryType->SupportsRayTracing();
 }
 
+static bool AreRayTracingMaterialsCompiled(EShaderPlatform Platform)
+{
+	static FShaderPlatformCachedIniValue<int32> CVarCompileMaterialCHS(TEXT("r.RayTracing.CompileMaterialCHS"));
+	static FShaderPlatformCachedIniValue<int32> CVarCompileMaterialAHS(TEXT("r.RayTracing.CompileMaterialAHS"));
+
+	return CVarCompileMaterialCHS.Get(Platform) || CVarCompileMaterialAHS.Get(Platform);
+}
+
 class FMaterialCHS : public FMeshMaterialShader, public FUniformLightMapPolicyShaderParametersType
 {
 	DECLARE_INLINE_TYPE_LAYOUT_EXPLICIT_BASES(FMaterialCHS, NonVirtual, FMeshMaterialShader, FUniformLightMapPolicyShaderParametersType);
@@ -129,12 +138,13 @@ public:
 
 	static bool ShouldCompilePermutation(const FMeshMaterialShaderPermutationParameters& Parameters)
 	{
-		if (!GCompileRayTracingMaterialAHS && !GCompileRayTracingMaterialCHS)
+		if (!AreRayTracingMaterialsCompiled(Parameters.Platform))
 		{
 			return false;
 		}
 
-		const bool bWantAnyHitShader = (GCompileRayTracingMaterialAHS && (Parameters.MaterialParameters.bIsMasked || IsTranslucentOnlyBlendMode(Parameters.MaterialParameters)));
+		static FShaderPlatformCachedIniValue<int32> CVarCompileMaterialAHS(TEXT("r.RayTracing.CompileMaterialAHS"));
+		const bool bWantAnyHitShader = ((CVarCompileMaterialAHS.Get(Parameters.Platform) != 0) && (Parameters.MaterialParameters.bIsMasked || IsTranslucentOnlyBlendMode(Parameters.MaterialParameters)));
 		const bool bSupportProceduralPrimitive = Parameters.VertexFactoryType->SupportsRayTracingProceduralPrimitive() && FDataDrivenShaderPlatformInfo::GetSupportsRayTracingProceduralPrimitive(Parameters.Platform);
 
 		return IsSupportedVertexFactoryType(Parameters.VertexFactoryType)
@@ -148,9 +158,11 @@ public:
 	static void ModifyCompilationEnvironment(const FMaterialShaderPermutationParameters& Parameters, FShaderCompilerEnvironment& OutEnvironment)
 	{
 		// NOTE: Any CVars that are used in this function must be handled in ShaderMapAppendKeyString() to ensure shaders are recompiled when necessary.
+		static FShaderPlatformCachedIniValue<int32> CVarCompileMaterialCHS(TEXT("r.RayTracing.CompileMaterialCHS"));
+		static FShaderPlatformCachedIniValue<int32> CVarCompileMaterialAHS(TEXT("r.RayTracing.CompileMaterialAHS"));
 
-		OutEnvironment.SetDefine(TEXT("USE_MATERIAL_CLOSEST_HIT_SHADER"), GCompileRayTracingMaterialCHS ? 1 : 0);
-		OutEnvironment.SetDefine(TEXT("USE_MATERIAL_ANY_HIT_SHADER"), GCompileRayTracingMaterialAHS ? 1 : 0);
+		OutEnvironment.SetDefine(TEXT("USE_MATERIAL_CLOSEST_HIT_SHADER"), CVarCompileMaterialCHS.Get(Parameters.Platform) ? 1 : 0);
+		OutEnvironment.SetDefine(TEXT("USE_MATERIAL_ANY_HIT_SHADER"), CVarCompileMaterialAHS.Get(Parameters.Platform) ? 1 : 0);
 		OutEnvironment.SetDefine(TEXT("USE_MATERIAL_INTERSECTION_SHADER"), UseIntersectionShader ? 1 : 0);
 		OutEnvironment.SetDefine(TEXT("USE_RAYTRACED_TEXTURE_RAYCONE_LOD"), UseRayConeTextureLod ? 1 : 0);
 		OutEnvironment.SetDefine(TEXT("SCENE_TEXTURES_DISABLED"), 1);
@@ -202,7 +214,7 @@ public:
 
 	static bool ShouldCompilePermutation(const FMeshMaterialShaderPermutationParameters& Parameters)
 	{
-		if (GCompileRayTracingMaterialAHS || GCompileRayTracingMaterialCHS)
+		if (AreRayTracingMaterialsCompiled(Parameters.Platform))
 		{
 			return false;
 		}
@@ -299,10 +311,10 @@ inline void GetMaterialHitShader_AnyHit_Intersection_TextureLOD(FMaterialShaderT
 }
 
 template<typename LightMapPolicyType>
-static bool GetMaterialHitShader(const FMaterial& RESTRICT MaterialResource, const FVertexFactory* VertexFactory, bool UseTextureLod, TShaderRef<FMaterialCHS>& OutShader)
+static bool GetMaterialHitShader(const FMaterial& RESTRICT MaterialResource, const FVertexFactory* VertexFactory, bool UseTextureLod, EShaderPlatform Platform, TShaderRef<FMaterialCHS>& OutShader)
 {
-	const bool bMaterialsCompiled = GCompileRayTracingMaterialAHS || GCompileRayTracingMaterialCHS;
-	checkf(bMaterialsCompiled, TEXT(""));
+	const bool bMaterialsCompiled = AreRayTracingMaterialsCompiled(Platform);
+	checkf(bMaterialsCompiled, TEXT("Material hit shaders are requested but they were not compiled for current platform [%s]"), *LexToString(Platform));
 
 	FMaterialShaderTypes ShaderTypes;
 	const FVertexFactoryType* VFType = VertexFactory->GetType();
@@ -325,11 +337,12 @@ static bool GetRayTracingMeshProcessorShaders(
 	const FUniformLightMapPolicy& RESTRICT LightMapPolicy,
 	const FVertexFactory* VertexFactory,
 	const FMaterial& RESTRICT MaterialResource,
+	EShaderPlatform Platform,
 	TShaderRef<FMaterialCHS>& OutRayHitGroupShader)
 {
 	check(GRHISupportsRayTracingShaders);
 
-	const bool bMaterialsCompiled = GCompileRayTracingMaterialAHS || GCompileRayTracingMaterialCHS;
+	const bool bMaterialsCompiled = AreRayTracingMaterialsCompiled(Platform);
 
 	if (bMaterialsCompiled)
 	{
@@ -338,31 +351,31 @@ static bool GetRayTracingMeshProcessorShaders(
 		switch (LightMapPolicy.GetIndirectPolicy())
 		{
 		case LMP_PRECOMPUTED_IRRADIANCE_VOLUME_INDIRECT_LIGHTING:
-			if (!GetMaterialHitShader<TUniformLightMapPolicy<LMP_PRECOMPUTED_IRRADIANCE_VOLUME_INDIRECT_LIGHTING>>(MaterialResource, VertexFactory, bUseTextureLOD, OutRayHitGroupShader))
+			if (!GetMaterialHitShader<TUniformLightMapPolicy<LMP_PRECOMPUTED_IRRADIANCE_VOLUME_INDIRECT_LIGHTING>>(MaterialResource, VertexFactory, bUseTextureLOD, Platform, OutRayHitGroupShader))
 			{
 				return false;
 			}
 			break;
 		case LMP_LQ_LIGHTMAP:
-			if (!GetMaterialHitShader<TUniformLightMapPolicy<LMP_LQ_LIGHTMAP>>(MaterialResource, VertexFactory, bUseTextureLOD, OutRayHitGroupShader))
+			if (!GetMaterialHitShader<TUniformLightMapPolicy<LMP_LQ_LIGHTMAP>>(MaterialResource, VertexFactory, bUseTextureLOD, Platform, OutRayHitGroupShader))
 			{
 				return false;
 			}
 			break;
 		case LMP_HQ_LIGHTMAP:
-			if (!GetMaterialHitShader<TUniformLightMapPolicy<LMP_HQ_LIGHTMAP>>(MaterialResource, VertexFactory, bUseTextureLOD, OutRayHitGroupShader))
+			if (!GetMaterialHitShader<TUniformLightMapPolicy<LMP_HQ_LIGHTMAP>>(MaterialResource, VertexFactory, bUseTextureLOD, Platform, OutRayHitGroupShader))
 			{
 				return false;
 			}
 			break;
 		case LMP_DISTANCE_FIELD_SHADOWS_AND_HQ_LIGHTMAP:
-			if (!GetMaterialHitShader<TUniformLightMapPolicy<LMP_DISTANCE_FIELD_SHADOWS_AND_HQ_LIGHTMAP>>(MaterialResource, VertexFactory, bUseTextureLOD, OutRayHitGroupShader))
+			if (!GetMaterialHitShader<TUniformLightMapPolicy<LMP_DISTANCE_FIELD_SHADOWS_AND_HQ_LIGHTMAP>>(MaterialResource, VertexFactory, bUseTextureLOD, Platform, OutRayHitGroupShader))
 			{
 				return false;
 			}
 			break;
 		case LMP_NO_LIGHTMAP:
-			if (!GetMaterialHitShader<TUniformLightMapPolicy<LMP_NO_LIGHTMAP>>(MaterialResource, VertexFactory, bUseTextureLOD, OutRayHitGroupShader))
+			if (!GetMaterialHitShader<TUniformLightMapPolicy<LMP_NO_LIGHTMAP>>(MaterialResource, VertexFactory, bUseTextureLOD, Platform, OutRayHitGroupShader))
 			{
 				return false;
 			}
@@ -411,7 +424,7 @@ bool FRayTracingMeshProcessor::Process(
 	TShaderRef<FMaterialCHS> RayTracingShader;
 	if (GRHISupportsRayTracingShaders)
 	{
-		if (!GetRayTracingMeshProcessorShaders(LightMapPolicy, MeshBatch.VertexFactory, MaterialResource, RayTracingShader))
+		if (!GetRayTracingMeshProcessorShaders(LightMapPolicy, MeshBatch.VertexFactory, MaterialResource, Scene->GetShaderPlatform(), RayTracingShader))
 		{
 			return false;
 		}
@@ -656,7 +669,7 @@ void FDeferredShadingSceneRenderer::CreateRayTracingMaterialPipeline(
 
 	Initializer.SetRayGenShaderTable(RayGenShaderTable);
 
-	const bool bMaterialsCompiled = GCompileRayTracingMaterialAHS || GCompileRayTracingMaterialCHS;
+	const bool bMaterialsCompiled = AreRayTracingMaterialsCompiled(View.GetShaderPlatform());
 	const bool bEnableMaterials = bMaterialsCompiled && GEnableRayTracingMaterials != 0;
 	static auto CVarEnableShadowMaterials = IConsoleManager::Get().FindConsoleVariable(TEXT("r.RayTracing.Shadows.EnableMaterials"));
 	const bool bEnableShadowMaterials = bMaterialsCompiled && (CVarEnableShadowMaterials ? CVarEnableShadowMaterials->GetInt() != 0 : true);
