Handle multiple views when computing visibility masks for per-instance occlusion queries

* Add a shader permutation to load previous visibility mask values before setting/clearing the visiblity bit for the current view

diff --git a/Engine/Shaders/Private/InstanceCulling/InstanceCullingOcclusionQuery.usf b/Engine/Shaders/Private/InstanceCulling/InstanceCullingOcclusionQuery.usf
--- a/Engine/Shaders/Private/InstanceCulling/InstanceCullingOcclusionQuery.usf
+++ b/Engine/Shaders/Private/InstanceCulling/InstanceCullingOcclusionQuery.usf
@@ -11,11 +11,6 @@
 
 static const uint GBumpMeToForceRebuild = 2;
 
-// Whether to run per-instance culling step using a CS or perform all the work in the VS
-#ifndef USE_COMPUTE_INSTANCE_CULLING_SETUP
-#	define USE_COMPUTE_INSTANCE_CULLING_SETUP 1
-#endif
-
 #ifndef HZB_DESIRED_FOOTPRINT_PIXELS
 #	define HZB_DESIRED_FOOTPRINT_PIXELS 4
 #endif
@@ -24,12 +19,17 @@ static const uint GBumpMeToForceRebuild = 2;
 #error This shader requires accessing GPUScene buffers as SRVs from vertex stage
 #endif
 
-RWBuffer<uint> OutVisibilityMask;
+#ifndef DIM_MULTI_VIEW
+#	define DIM_MULTI_VIEW 0
+#endif
+
+RWBuffer<uint> RWVisibilityMask;
 
 uint NumInstances;
 uint InstanceSceneDataSOAStride;
-int4 ViewRect;
+uint ViewMask;
 float OcclusionSlop;
+int2 HZBViewRectSize;
 
 #define EInstanceVisibilityStatus uint // Not all platforms support native enums
 #define IVS_Hidden          0 // Instance is hidden or invalid, no occlusion query needed
@@ -37,10 +37,6 @@ float OcclusionSlop;
 #define IVS_PossiblyVisible 2 // Occlusion query is needed to determine visibility
 #define IVS_Incompatible    3 // Instance can't use occlusion queries
 
-// TODO: Compute appropriate mask value based on current view index and previous mask value
-#define MaskValueVisible 1
-#define MaskValueHidden  0
-
 EInstanceVisibilityStatus GetInstanceDataAndVisibility(uint InstanceId, out FInstanceSceneData OutInstanceData)
 {
 	const FInstanceSceneData InstanceData = GetInstanceSceneData(InstanceId, InstanceSceneDataSOAStride);
@@ -68,11 +64,11 @@ EInstanceVisibilityStatus GetInstanceDataAndVisibility(uint InstanceId, out FIns
 	float4x4 LocalToTranslatedWorld = DFMultiplyTranslationDemote(InstanceData.LocalToWorld, ResolvedView.PreViewTranslation);
 
 	FFrustumCullData Cull = BoxCullFrustum(
-		BoxOrigin.xyz, 
+		BoxOrigin.xyz,
 		BoxExtent.xyz,
 		LocalToTranslatedWorld,
 		ResolvedView.TranslatedWorldToClip,
-		ResolvedView.ViewToClip, 
+		ResolvedView.ViewToClip,
 		bIsOrtho, true /* near clip */, false /* skip culling */
 	);
 
@@ -82,7 +78,7 @@ EInstanceVisibilityStatus GetInstanceDataAndVisibility(uint InstanceId, out FIns
 		return IVS_Visible;
 	}
 
-	if(!Cull.bIsVisible)
+	if (!Cull.bIsVisible)
 	{
 		// Conservative culling heuristic:
 		// If instance is frustum-culled this frame, it may become visible next frame when instance visibility buffer is consumed.
@@ -92,9 +88,10 @@ EInstanceVisibilityStatus GetInstanceDataAndVisibility(uint InstanceId, out FIns
 		return IVS_Visible;
 	}
 
-	FScreenRect Rect = GetScreenRect(ViewRect, Cull, HZB_DESIRED_FOOTPRINT_PIXELS);
+	int4 HZBViewRect = int4(0, 0, HZBViewRectSize.x, HZBViewRectSize.y);
+	FScreenRect Rect = GetScreenRect(HZBViewRect, Cull, HZB_DESIRED_FOOTPRINT_PIXELS);
 
-	if (any(Rect.Pixels == ViewRect))
+	if (any(Rect.Pixels == HZBViewRect))
 	{
 		// Conservative culling heuristic:
 		// Objects may be visible with HZB test but occluded with per-pixel depth test during this frame.
@@ -106,7 +103,7 @@ EInstanceVisibilityStatus GetInstanceDataAndVisibility(uint InstanceId, out FIns
 		return IVS_Visible;
 	}
 
-	if (IsVisibleHZB(Rect, true /*bSample4x4*/ ))
+	if (IsVisibleHZB(Rect, true /*bSample4x4*/))
 	{
 		// Instance is visible using conservative HZB test and may benefit from a more accurate per-pixel occlusion test.
 		return IVS_PossiblyVisible;
@@ -141,7 +138,7 @@ void MainCS(uint3 GroupId : SV_GroupID, uint GroupThreadIndex : SV_GroupIndex, u
 
 	const bool bValidInvocation = DispatchThreadId < NumInstances;
 
-	uint ClampedThreadId = min(DispatchThreadId, NumInstances-1);
+	uint ClampedThreadId = min(DispatchThreadId, NumInstances - 1);
 	uint InstanceId = InstanceIdBuffer[ClampedThreadId];
 
 	GroupMemoryBarrierWithGroupSync();
@@ -164,19 +161,26 @@ void MainCS(uint3 GroupId : SV_GroupID, uint GroupThreadIndex : SV_GroupIndex, u
 	{
 		Status = GetInstanceDataAndVisibility(InstanceId, /*out*/ InstanceData);
 
+#if DIM_MULTI_VIEW
+		const uint OldVisibilityMask = RWVisibilityMask[InstanceId];
+		const uint MaskValueHidden = OldVisibilityMask & (~ViewMask);
+		const uint MaskValueVisible = OldVisibilityMask | ViewMask;
+#else // DIM_MULTI_VIEW
+		const uint MaskValueHidden = 0;
+		const uint MaskValueVisible = ViewMask;
+#endif // DIM_MULTI_VIEW
+
 		if (Status == IVS_Visible || Status == IVS_Incompatible)
 		{
-			// YURIY_TODO: multi-view support
-			OutVisibilityMask[InstanceId] = MaskValueVisible; // this is the final value
+			RWVisibilityMask[InstanceId] = MaskValueVisible; // this is the final value
 		}
 		else if (Status == IVS_Hidden)
 		{
-			// YURIY_TODO: multi-view support
-			OutVisibilityMask[InstanceId] = MaskValueHidden; // this is the final value
+			RWVisibilityMask[InstanceId] = MaskValueHidden; // this is the final value
 		}
 		else // Status == IVS_PossiblyVisible
 		{
-			OutVisibilityMask[InstanceId] = MaskValueHidden; // clear the value which will be updated by pixel shader later
+			RWVisibilityMask[InstanceId] = MaskValueHidden; // clear the value which will be updated by pixel shader later
 
 			uint LocalOutputIndex = 0;
 			InterlockedAdd(GroupNumVisibleInstances, 1, LocalOutputIndex);
@@ -212,32 +216,21 @@ void MainVS(
 
 	FInstanceSceneData InstanceData = (FInstanceSceneData)0;
 
-#if USE_COMPUTE_INSTANCE_CULLING_SETUP
 	InstanceId = InstanceIdBuffer[InstanceId];
 	GetInstanceDataAndVisibility(InstanceId, /*out*/ InstanceData);
-#else
-	EInstanceVisibilityStatus Status = GetInstanceDataAndVisibility(InstanceId, /*out*/ InstanceData);
-	if (Status == IVS_Visible || Status == IVS_Incompatible)
-	{
-		// Render a box that covers just one pixel at the near plane which will mark the instance visible
-		float2 ViewRectSize = float2(ViewRect.z - ViewRect.x, ViewRect.w - ViewRect.y);
-		OutPosition = float4(InPosition.xy / ViewRectSize, 1, 1);
-	}
-	else if (Status == IVS_Hidden)
-	{
-		// Output NaN position to cull the primitive, preventing any pixel shader execution
-		OutPosition = asfloat(0xFFFFFFFF).xxxx;
-	}
-	else // Status == IVS_PossiblyVisible
-#endif // USE_COMPUTE_INSTANCE_CULLING_SETUP
-	{
-		float3 LocalPosition = (InPosition * (InstanceData.LocalBoundsExtent + OcclusionSlop) * InstanceData.DeterminantSign) + InstanceData.LocalBoundsCenter;
-		float4 WorldPosition = TransformLocalToTranslatedWorld(LocalPosition, InstanceData.LocalToWorld);
-		OutPosition = mul(WorldPosition, ResolvedView.TranslatedWorldToClip);
-	}
+
+	float3 LocalPosition = (InPosition * (InstanceData.LocalBoundsExtent + OcclusionSlop) * InstanceData.DeterminantSign) + InstanceData.LocalBoundsCenter;
+	float4 WorldPosition = TransformLocalToTranslatedWorld(LocalPosition, InstanceData.LocalToWorld);
+	OutPosition = mul(WorldPosition, ResolvedView.TranslatedWorldToClip);
 
 	OutInstanceId = InstanceId;
-	OutVisibilityMask = MaskValueVisible;
+
+#if DIM_MULTI_VIEW
+	const uint OldVisibilityMask = RWVisibilityMask[InstanceId];
+	OutVisibilityMask = OldVisibilityMask | ViewMask;
+#else // DIM_MULTI_VIEW
+	OutVisibilityMask = ViewMask;
+#endif // DIM_MULTI_VIEW
 }
 
 
@@ -248,7 +241,7 @@ void MainPS(
 	out float4 OutColor : SV_Target0
 )
 {
-	OutVisibilityMask[InstanceId] = VisibilityMask;
+	RWVisibilityMask[InstanceId] = VisibilityMask;
 	OutColor = float4(1, 0, 0, 1);
 }
 
@@ -275,7 +268,7 @@ void DebugMainVS(
 
 		OutPosition = mul(WorldPosition, ResolvedView.TranslatedWorldToClip);
 
-		const bool bWasVisibleLastFrame = InstanceOcclusionQueryBuffer[InstanceId] != 0; // YURIY_TODO: multi-view support
+		const bool bWasVisibleLastFrame = (InstanceOcclusionQueryBuffer[InstanceId] & ViewMask) == ViewMask;
 
 		if (VisibilityStatus == IVS_PossiblyVisible)
 		{
@@ -303,20 +296,20 @@ void DebugMainPS(
 	case 0:
 		// Definitely invisible due to HZB test.
 		// No occlusion query performed.
-		OutColor = float4(0.1, 0.0, 0.0, 0.0);
+		OutColor = float4(0.1, 0.0, 0.0, 0.0); // red
 		break;
 	case 1:
 		// Definitely visible due to near plane intersection or other heuristics.
 		// No occlusion query performed.
-		OutColor = float4(0.1, 0.1, 0.0, 0.0);
+		OutColor = float4(0.1, 0.1, 0.0, 0.0); // yellow
 		break;
 	case 2:
 		// Visible due to occlusion query.
-		OutColor = float4(0.0, 0.025, 0.0, 0.0);
+		OutColor = float4(0.0, 0.025, 0.0, 0.0); // green
 		break;
 	case 3:
 		// Invisible due to occlusion query.
-		OutColor = float4(0.25, 0.0, 0.25, 0.0);
+		OutColor = float4(0.25, 0.0, 0.25, 0.0); // magenta
 		break;
 	default:
 		OutColor = (float4)0;
diff --git a/Engine/Source/Runtime/Renderer/Private/InstanceCulling/InstanceCullingOcclusionQuery.cpp b/Engine/Source/Runtime/Renderer/Private/InstanceCulling/InstanceCullingOcclusionQuery.cpp
--- a/Engine/Source/Runtime/Renderer/Private/InstanceCulling/InstanceCullingOcclusionQuery.cpp
+++ b/Engine/Source/Runtime/Renderer/Private/InstanceCulling/InstanceCullingOcclusionQuery.cpp
@@ -95,6 +95,10 @@ class FInstanceCullingOcclusionQueryCS : public FGlobalShader
 	SHADER_USE_PARAMETER_STRUCT(FInstanceCullingOcclusionQueryCS, FGlobalShader);
 
 public:
+
+	class FMultiView : SHADER_PERMUTATION_BOOL("DIM_MULTI_VIEW");
+	using FPermutationDomain = TShaderPermutationDomain<FMultiView>;
+
 	static bool ShouldCompilePermutation(const FGlobalShaderPermutationParameters& Parameters)
 	{
 		return FDataDrivenShaderPlatformInfo::GetSupportsVertexShaderSRVs(Parameters.Platform);
@@ -113,7 +117,7 @@ public:
 	BEGIN_SHADER_PARAMETER_STRUCT(FParameters, )
 		SHADER_PARAMETER_RDG_BUFFER_UAV(RWBuffer<uint32>, OutIndirectArgsBuffer)
 		SHADER_PARAMETER_RDG_BUFFER_UAV(RWBuffer<uint32>, OutInstanceIdBuffer)
-		SHADER_PARAMETER_RDG_BUFFER_UAV(RWBuffer<uint>, OutVisibilityMask) // One uint32 per instance (0 if instance is culled, non-0 otherwise)
+		SHADER_PARAMETER_RDG_BUFFER_UAV(RWBuffer<uint>, RWVisibilityMask) // One uint8/32 per instance (0 if instance is culled, non-0 otherwise)
 		SHADER_PARAMETER(uint32, InstanceSceneDataSOAStride)
 		SHADER_PARAMETER_RDG_BUFFER_SRV(StructuredBuffer<float4>, GPUSceneInstanceSceneData)
 		SHADER_PARAMETER_RDG_BUFFER_SRV(StructuredBuffer<float4>, GPUSceneInstancePayloadData)
@@ -123,9 +127,10 @@ public:
 		SHADER_PARAMETER_RDG_TEXTURE(Texture2D, HZBTexture)
 		SHADER_PARAMETER_SAMPLER(SamplerState, HZBSampler)
 		SHADER_PARAMETER(FVector2f, HZBSize)
-		SHADER_PARAMETER(FIntVector4, ViewRect)
+		SHADER_PARAMETER(FIntPoint, HZBViewRectSize)
 		SHADER_PARAMETER(float, OcclusionSlop)
 		SHADER_PARAMETER(int32, NumInstances)
+		SHADER_PARAMETER(uint32, ViewMask)
 	END_SHADER_PARAMETER_STRUCT()
 };
 
@@ -137,6 +142,10 @@ class FInstanceCullingOcclusionQueryVS : public FGlobalShader
 	SHADER_USE_PARAMETER_STRUCT(FInstanceCullingOcclusionQueryVS, FGlobalShader);
 
 public:
+
+	class FMultiView : SHADER_PERMUTATION_BOOL("DIM_MULTI_VIEW");
+	using FPermutationDomain = TShaderPermutationDomain<FMultiView>;
+
 	static bool ShouldCompilePermutation(const FGlobalShaderPermutationParameters& Parameters)
 	{
 		return FDataDrivenShaderPlatformInfo::GetSupportsVertexShaderSRVs(Parameters.Platform);
@@ -156,12 +165,14 @@ public:
 		SHADER_PARAMETER_RDG_BUFFER_SRV(StructuredBuffer<float4>, GPUSceneInstancePayloadData)
 		SHADER_PARAMETER_RDG_BUFFER_SRV(StructuredBuffer<float4>, GPUScenePrimitiveSceneData)
 		SHADER_PARAMETER_RDG_BUFFER_SRV(Buffer<uint32>, InstanceIdBuffer)
+		SHADER_PARAMETER_RDG_BUFFER_UAV(RWBuffer<uint>, RWVisibilityMask) // One uint8/32 per instance (0 if instance is culled, non-0 otherwise)
 		SHADER_PARAMETER_STRUCT_REF(FViewUniformShaderParameters, View)
 		SHADER_PARAMETER_RDG_TEXTURE(Texture2D, HZBTexture)
 		SHADER_PARAMETER_SAMPLER(SamplerState, HZBSampler)
 		SHADER_PARAMETER(FVector2f, HZBSize)
-		SHADER_PARAMETER(FIntVector4, ViewRect)
+		SHADER_PARAMETER(FIntPoint, HZBViewRectSize)
 		SHADER_PARAMETER(float, OcclusionSlop)
+		SHADER_PARAMETER(uint32, ViewMask)
 	END_SHADER_PARAMETER_STRUCT()
 };
 
@@ -181,7 +192,7 @@ public:
 	}
 
 	BEGIN_SHADER_PARAMETER_STRUCT(FParameters, )
-		SHADER_PARAMETER_RDG_BUFFER_UAV(RWBuffer<uint>, OutVisibilityMask) // One uint32 per instance (0 if instance is culled, non-0 otherwise)
+		SHADER_PARAMETER_RDG_BUFFER_UAV(RWBuffer<uint>, RWVisibilityMask) // One uint8/32 per instance (0 if instance is culled, non-0 otherwise)
 	END_SHADER_PARAMETER_STRUCT()
 };
 
@@ -260,15 +271,19 @@ TGlobalResource<FInstanceCullingOcclusionQueryBox> GInstanceCullingOcclusionQuer
 static void RenderInstanceOcclusionCulling(
 	FRHICommandList& RHICmdList,
 	FViewInfo& View,
-	FOcclusionInstanceCullingParameters* PassParameters)
+	FOcclusionInstanceCullingParameters* PassParameters,
+	bool bMultiView)
 {
-	TShaderMapRef<FInstanceCullingOcclusionQueryVS> VertexShader(View.ShaderMap);
+	FInstanceCullingOcclusionQueryVS::FPermutationDomain VSPermutationVector;
+	VSPermutationVector.Set<FInstanceCullingOcclusionQueryVS::FMultiView>(bMultiView);
+	TShaderMapRef<FInstanceCullingOcclusionQueryVS> VertexShader(View.ShaderMap, VSPermutationVector);
+
 	TShaderMapRef<FInstanceCullingOcclusionQueryPS> PixelShader(View.ShaderMap);
 
 	FGraphicsPipelineStateInitializer GraphicsPSOInit;
 	RHICmdList.ApplyCachedRenderTargets(GraphicsPSOInit);
 
-	FIntVector4 ViewRect = PassParameters->VS.ViewRect;
+	FIntVector4 ViewRect = FIntVector4(View.ViewRect.Min.X, View.ViewRect.Min.Y, View.ViewRect.Max.X, View.ViewRect.Max.Y);
 	RHICmdList.SetViewport(ViewRect.X, ViewRect.Y, 0.0f, ViewRect.Z, ViewRect.W, 1.0f);
 
 	GraphicsPSOInit.BoundShaderState.VertexDeclarationRHI = GInstanceCullingOcclusionQueryBox.VertexDeclaration;
@@ -510,7 +525,7 @@ uint32 FInstanceCullingOcclusionQueryRenderer::Render(
 		return 0;
 	}
 
-	const uint32 ViewMask = RegisterView(View);
+	const uint32 ViewMask = FindOrAddViewSlot(View);
 
 	if (ViewMask == 0)
 	{
@@ -520,6 +535,9 @@ uint32 FInstanceCullingOcclusionQueryRenderer::Render(
 
 	TRACE_CPUPROFILER_EVENT_SCOPE(FInstanceCullingOcclusionQueryRenderer::Render);
 
+	// Whether to use shader permutation that preserves visibility bits corresponding to other views (slight extra cost)
+	const bool bMultiView = CurrentRenderedViewIDs.Num() > 1;
+
 	const int32 NumGPUSceneInstances = GPUScene.GetNumInstances();
 
 	FInstanceCullingOcclusionQueryDeferredContext* DeferredContext = GraphBuilder.AllocObject<FInstanceCullingOcclusionQueryDeferredContext>(&View, NumGPUSceneInstances, EMeshPass::BasePass);
@@ -598,12 +616,12 @@ uint32 FInstanceCullingOcclusionQueryRenderer::Render(
 
 		PassParameters->OutIndirectArgsBuffer = IndirectArgsUAV;
 		PassParameters->OutInstanceIdBuffer = InstanceIdUAV;
-		PassParameters->OutVisibilityMask = VisibilityMaskUAV;
+		PassParameters->RWVisibilityMask = VisibilityMaskUAV;
 		PassParameters->View = View.ViewUniformBuffer;
 		PassParameters->HZBTexture = HZBTexture;
 		PassParameters->HZBSampler = TStaticSamplerState<SF_Point, AM_Clamp, AM_Clamp, AM_Clamp>::GetRHI();
 		PassParameters->HZBSize = FVector2f(HZBSize.X, HZBSize.Y);
-		PassParameters->ViewRect = FIntVector4(View.ViewRect.Min.X, View.ViewRect.Min.Y, View.ViewRect.Max.X, View.ViewRect.Max.Y);
+		PassParameters->HZBViewRectSize = ViewRectSize;
 		PassParameters->OcclusionSlop = OCCLUSION_SLOP;
 		PassParameters->InstanceSceneDataSOAStride = GPUSceneParameters.InstanceDataSOAStride;
 		PassParameters->GPUSceneInstanceSceneData = GPUSceneParameters.GPUSceneInstanceSceneData;
@@ -611,8 +629,11 @@ uint32 FInstanceCullingOcclusionQueryRenderer::Render(
 		PassParameters->GPUScenePrimitiveSceneData = GPUSceneParameters.GPUScenePrimitiveSceneData;
 		PassParameters->NumInstances = 0; // filled from DeferredContext later
 		PassParameters->InstanceIdBuffer = SetupInstanceIdBufferSRV;
+		PassParameters->ViewMask = ViewMask;
 
-		TShaderMapRef<FInstanceCullingOcclusionQueryCS> ComputeShader(View.ShaderMap);
+		FInstanceCullingOcclusionQueryCS::FPermutationDomain CSPermutationVector;
+		CSPermutationVector.Set<FInstanceCullingOcclusionQueryCS::FMultiView>(bMultiView);
+		TShaderMapRef<FInstanceCullingOcclusionQueryCS> ComputeShader(View.ShaderMap, CSPermutationVector);
 
 		ClearUnusedGraphResources(ComputeShader, PassParameters);
 
@@ -646,29 +667,32 @@ uint32 FInstanceCullingOcclusionQueryRenderer::Render(
 		PassParameters->VS.HZBTexture = HZBTexture;
 		PassParameters->VS.HZBSampler = TStaticSamplerState<SF_Point, AM_Clamp, AM_Clamp, AM_Clamp>::GetRHI();
 		PassParameters->VS.HZBSize = FVector2f(HZBSize.X, HZBSize.Y);
-		PassParameters->VS.ViewRect = FIntVector4(View.ViewRect.Min.X, View.ViewRect.Min.Y, View.ViewRect.Max.X, View.ViewRect.Max.Y);
+		PassParameters->VS.HZBViewRectSize = ViewRectSize;
 		PassParameters->VS.OcclusionSlop = OCCLUSION_SLOP;
 		PassParameters->VS.InstanceSceneDataSOAStride = GPUSceneParameters.InstanceDataSOAStride;
 		PassParameters->VS.GPUSceneInstanceSceneData = GPUSceneParameters.GPUSceneInstanceSceneData;
 		PassParameters->VS.GPUSceneInstancePayloadData = GPUSceneParameters.GPUSceneInstancePayloadData;
 		PassParameters->VS.GPUScenePrimitiveSceneData = GPUSceneParameters.GPUScenePrimitiveSceneData;
 		PassParameters->VS.InstanceIdBuffer = InstanceIdSRV;
-		PassParameters->PS.OutVisibilityMask = VisibilityMaskUAV;
+		PassParameters->VS.ViewMask = ViewMask;
+		PassParameters->VS.RWVisibilityMask = VisibilityMaskUAV;
+		PassParameters->PS.RWVisibilityMask = VisibilityMaskUAV;
 		PassParameters->RenderTargets.DepthStencil = FDepthStencilBinding(DepthTexture,
 			ERenderTargetLoadAction::ELoad, ERenderTargetLoadAction::ENoAction,
 			FExclusiveDepthStencil::DepthRead_StencilNop);
 
+
 		GraphBuilder.AddPass(
 			RDG_EVENT_NAME("InstanceCullingOcclusionQueryRenderer_Draw"),
 			PassParameters, ERDGPassFlags::Raster | ERDGPassFlags::NeverCull,
-			[PassParameters, DeferredContext, &View](FRHICommandList& RHICmdList)
+			[PassParameters, DeferredContext, bMultiView, &View](FRHICommandList& RHICmdList)
 			{
 				if (!DeferredContext->bValid)
 				{
 					return;
 				}
 
-				RenderInstanceOcclusionCulling(RHICmdList, View, PassParameters);
+				RenderInstanceOcclusionCulling(RHICmdList, View, PassParameters, bMultiView);
 			});
 	}
 
@@ -713,11 +737,13 @@ void FInstanceCullingOcclusionQueryRenderer::EndFrame(FRDGBuilder& GraphBuilder)
 	CurrentRenderedViewIDs.Empty();
 }
 
-uint32 FInstanceCullingOcclusionQueryRenderer::RegisterView(const FViewInfo& View)
+uint32 FInstanceCullingOcclusionQueryRenderer::FindOrAddViewSlot(const FViewInfo& View)
 {
-	if (CurrentRenderedViewIDs.Num() < MaxViews)
+	const uint32 ViewKey = View.GetViewKey();
+
+	if (CurrentRenderedViewIDs.Num() < MaxViews && ViewKey != 0)
 	{
-		int32 Index = CurrentRenderedViewIDs.AddUnique(View.GetViewKey());
+		int32 Index = CurrentRenderedViewIDs.AddUnique(ViewKey);
 		check(Index >= 0 && Index < MaxViews);
 		return 1u << Index;
 	}
@@ -731,6 +757,7 @@ bool FInstanceCullingOcclusionQueryRenderer::IsCompatibleWithView(const FViewInf
 {
 	EPixelFormat VisibilityMaskFormat = GetPreferredVisibilityMaskFormat();
 	return FDataDrivenShaderPlatformInfo::GetSupportsVertexShaderSRVs(View.GetShaderPlatform())
+		&& View.GetViewKey() != 0
 		&& View.GetSceneTextures().Depth.Target
 		&& View.HZB
 		&& VisibilityMaskFormat != PF_Unknown
@@ -764,11 +791,13 @@ public:
 		SHADER_PARAMETER_RDG_BUFFER_SRV(StructuredBuffer<float4>, GPUScenePrimitiveSceneData)
 		SHADER_PARAMETER_STRUCT_REF(FViewUniformShaderParameters, View)
 		SHADER_PARAMETER_RDG_BUFFER_SRV(Buffer<uint>, InstanceOcclusionQueryBuffer)
+		SHADER_PARAMETER_RDG_BUFFER_UAV(RWBuffer<uint>, RWVisibilityMask) // One uint8/32 per instance (0 if instance is culled, non-0 otherwise)
 		SHADER_PARAMETER_RDG_TEXTURE(Texture2D, HZBTexture)
 		SHADER_PARAMETER_SAMPLER(SamplerState, HZBSampler)
 		SHADER_PARAMETER(FVector2f, HZBSize)
-		SHADER_PARAMETER(FIntVector4, ViewRect)
+		SHADER_PARAMETER(FIntPoint, HZBViewRectSize)
 		SHADER_PARAMETER(float, OcclusionSlop)
+		SHADER_PARAMETER(uint32, ViewMask)
 	END_SHADER_PARAMETER_STRUCT()
 };
 
@@ -812,7 +841,7 @@ static void RenderInstanceOcclusionCullingDebug(
 	FGraphicsPipelineStateInitializer GraphicsPSOInit;
 	RHICmdList.ApplyCachedRenderTargets(GraphicsPSOInit);
 
-	FIntVector4 ViewRect = PassParameters->VS.ViewRect;
+	FIntVector4 ViewRect = FIntVector4(View.ViewRect.Min.X, View.ViewRect.Min.Y, View.ViewRect.Max.X, View.ViewRect.Max.Y);
 	RHICmdList.SetViewport(ViewRect.X, ViewRect.Y, 0.0f, ViewRect.Z, ViewRect.W, 1.0f);
 
 	GraphicsPSOInit.BoundShaderState.VertexDeclarationRHI = GInstanceCullingOcclusionQueryBox.VertexDeclaration;
@@ -843,6 +872,8 @@ void FInstanceCullingOcclusionQueryRenderer::RenderDebug(FRDGBuilder& GraphBuild
 		return;
 	}
 
+	const uint32 ViewMask = FindOrAddViewSlot(View);
+
 	FRDGTextureRef SceneColor = View.GetSceneTextures().Color.Target;
 	FRDGTextureRef SceneDepth = View.GetSceneTextures().Depth.Target;
 	FRDGBufferRef InstanceOcclusionQueryBufferRDG = GraphBuilder.RegisterExternalBuffer(InstanceOcclusionQueryBuffer);
@@ -858,10 +889,12 @@ void FInstanceCullingOcclusionQueryRenderer::RenderDebug(FRDGBuilder& GraphBuild
 	const int32 NumInstances = GPUScene.GetNumInstances();
 	const FGPUSceneResourceParameters GPUSceneParameters = GPUScene.GetShaderParameters(GraphBuilder);
 
+	const FIntPoint ViewRectSize = View.ViewRect.Size();
+
 	FOcclusionInstanceCullingDebugParameters* PassParameters = GraphBuilder.AllocParameters<FOcclusionInstanceCullingDebugParameters>();
 
 	PassParameters->VS.OcclusionSlop = OCCLUSION_SLOP;
-	PassParameters->VS.ViewRect = FIntVector4(View.ViewRect.Min.X, View.ViewRect.Min.Y, View.ViewRect.Max.X, View.ViewRect.Max.Y);
+	PassParameters->VS.HZBViewRectSize = ViewRectSize;
 	PassParameters->VS.View = View.ViewUniformBuffer;
 	PassParameters->VS.InstanceSceneDataSOAStride = GPUSceneParameters.InstanceDataSOAStride;
 	PassParameters->VS.GPUSceneInstanceSceneData = GPUSceneParameters.GPUSceneInstanceSceneData;
@@ -871,6 +904,7 @@ void FInstanceCullingOcclusionQueryRenderer::RenderDebug(FRDGBuilder& GraphBuild
 	PassParameters->VS.HZBTexture = HZBTexture;
 	PassParameters->VS.HZBSampler = TStaticSamplerState<SF_Point, AM_Clamp, AM_Clamp, AM_Clamp>::GetRHI();
 	PassParameters->VS.HZBSize = FVector2f(HZBSize.X, HZBSize.Y);
+	PassParameters->VS.ViewMask = ViewMask;
 	PassParameters->RenderTargets[0] = FRenderTargetBinding(SceneColor, ERenderTargetLoadAction::ELoad);
 	PassParameters->RenderTargets.DepthStencil = FDepthStencilBinding(SceneDepth,
 		ERenderTargetLoadAction::ELoad, ERenderTargetLoadAction::ENoAction,
diff --git a/Engine/Source/Runtime/Renderer/Private/InstanceCulling/InstanceCullingOcclusionQuery.h b/Engine/Source/Runtime/Renderer/Private/InstanceCulling/InstanceCullingOcclusionQuery.h
--- a/Engine/Source/Runtime/Renderer/Private/InstanceCulling/InstanceCullingOcclusionQuery.h
+++ b/Engine/Source/Runtime/Renderer/Private/InstanceCulling/InstanceCullingOcclusionQuery.h
@@ -70,11 +70,11 @@ private:
 	* Returns bit mask with a single set bit that should be used to access the occlusion results.
 	* Returns 0 if maximum number of supported views is reached, falling back to no-occlusion-query code path.
 	*/
-	uint32 RegisterView(const FViewInfo& View);
+	uint32 FindOrAddViewSlot(const FViewInfo& View);
 
 	FRDGBufferRef CurrentInstanceOcclusionQueryBuffer = {};
 
-	static constexpr uint32 MaxViews = 1; // -- YURIY_TODO: we could theoretically support up to 8 views for uint8 or 32 for uint32 mask
+	static constexpr uint32 MaxViews = 8; // Visibility mask is stored as uint8
 	TArray<uint32> CurrentRenderedViewIDs;
 
 	uint32 AllocatedNumInstances = 0;
