-Change UpdateStaticMeshesForMaterials to use FObjectCacheContext for primitive iteration
-Changed FGlobalComponentRecreateRenderStateContex to use a combination of FObjectCacheContext & UActorComponent iteration to identitfy the components to recreate )

diff --git a/Engine/Source/Runtime/Engine/Private/Components/ActorComponent.cpp b/Engine/Source/Runtime/Engine/Private/Components/ActorComponent.cpp
--- a/Engine/Source/Runtime/Engine/Private/Components/ActorComponent.cpp
+++ b/Engine/Source/Runtime/Engine/Private/Components/ActorComponent.cpp
@@ -39,6 +39,7 @@
 #include "UObject/FrameworkObjectVersion.h"
 #include "PSOPrecacheMaterial.h"
 #include "Materials/MaterialInterface.h"
+#include "ObjectCacheContext.h"
 
 #if WITH_EDITOR
 #include "Kismet2/ComponentEditorUtils.h"
@@ -309,12 +310,25 @@ FGlobalComponentRecreateRenderStateContext::FGlobalComponentRecreateRenderStateC
 		// wait until resources are released
 		FlushRenderingCommands();
 
+		FObjectCacheContextScope ObjectCacheScope;
+		for (IPrimitiveComponent* PrimitiveComponent: ObjectCacheScope.GetContext().GetPrimitiveComponents())
+		{
+			if (PrimitiveComponent->IsRegistered() && PrimitiveComponent->IsRenderStateCreated())
+			{
+				ComponentContexts.Emplace(PrimitiveComponent, &ScenesToUpdateAllPrimitiveSceneInfos);
+			}
+		}
+
 		// recreate render state for all components.
 		for (UActorComponent* Component : TObjectRange<UActorComponent>())
 		{
-			if (Component->IsRegistered() && Component->IsRenderStateCreated())
+			// Those are obtained through FObjectCacheContext
+			if (!Component->IsA<UPrimitiveComponent>())
 			{
-				ComponentContexts.Emplace(Component, &ScenesToUpdateAllPrimitiveSceneInfos);
+				if (Component->IsRegistered() && Component->IsRenderStateCreated())
+				{
+					ComponentContexts.Emplace(Component, &ScenesToUpdateAllPrimitiveSceneInfos);
+				}
 			}
 		}
 
diff --git a/Engine/Source/Runtime/Renderer/Private/RendererScene.cpp b/Engine/Source/Runtime/Renderer/Private/RendererScene.cpp
--- a/Engine/Source/Runtime/Renderer/Private/RendererScene.cpp
+++ b/Engine/Source/Runtime/Renderer/Private/RendererScene.cpp
@@ -83,6 +83,7 @@
 #include "InstanceCulling/InstanceCullingOcclusionQuery.h"
 #include "ComputeWorkerInterface.h"
 #include "Nanite/NaniteMaterialsSceneExtension.h"
+#include "ObjectCacheContext.h"
 
 #if RHI_RAYTRACING
 #include "Nanite/NaniteRayTracing.h"
@@ -7159,18 +7160,20 @@ void UpdateStaticMeshesForMaterials(const TArray<const FMaterial*>& MaterialReso
 
 	TArray<UMaterialInterface*> UsedMaterials;
 	TSet<UMaterialInterface*> UsedMaterialsDependencies;
-	TMap<FScene*, TArray<FPrimitiveSceneInfo*>> UsedPrimitives;
-	for (TObjectIterator<UPrimitiveComponent> PrimitiveIt; PrimitiveIt; ++PrimitiveIt)
-	{
-		UPrimitiveComponent* PrimitiveComponent = *PrimitiveIt;
 
-		if (PrimitiveComponent->IsRenderStateCreated() && PrimitiveComponent->SceneProxy && PrimitiveComponent->SceneProxy->GetPrimitiveSceneInfo()->IsIndexValid())
+	FObjectCacheContextScope ObjectCacheScope;
+
+	TMap<FScene*, TArray<FPrimitiveSceneInfo*>> UsedPrimitives;
+	
+	for (IPrimitiveComponent* PrimitiveComponent: ObjectCacheScope.GetContext().GetPrimitiveComponents())
+	{		
+		if (PrimitiveComponent->IsRenderStateCreated() && PrimitiveComponent->GetSceneProxy() && PrimitiveComponent->GetSceneProxy()->GetPrimitiveSceneInfo()->IsIndexValid())
 		{
 			UsedMaterialsDependencies.Reset();
 			UsedMaterials.Reset();
 
 			// Note: relying on GetUsedMaterials to be accurate, or else we won't propagate to the right primitives and the renderer will crash later
-			// FPrimitiveSceneProxy::VerifyUsedMaterial is used to make sure that all materials used for rendering are reported in GetUsedMaterials
+			// FPrimitiveSceneProxy::VerifyUsedMaterial is used to make sure that all materials used for rendering are reported in GetUsedMaterials			
 			PrimitiveComponent->GetUsedMaterials(UsedMaterials);
 
 			for (UMaterialInterface* UsedMaterial : UsedMaterials)
@@ -7191,7 +7194,7 @@ void UpdateStaticMeshesForMaterials(const TArray<const FMaterial*>& MaterialReso
 					{
 						if (UsedMaterialsDependencies.Contains(UpdatedMaterialInterface))
 						{
-							FPrimitiveSceneProxy* SceneProxy = PrimitiveComponent->SceneProxy;
+							FPrimitiveSceneProxy* SceneProxy = PrimitiveComponent->GetSceneProxy();
 							FPrimitiveSceneInfo* SceneInfo = SceneProxy->GetPrimitiveSceneInfo();
 							FScene* Scene = SceneInfo->Scene;
 							TArray<FPrimitiveSceneInfo*>& SceneInfos = UsedPrimitives.FindOrAdd(Scene);
