Added additional CVars to control the temporal antialiasing jitter sequence.

r.TemporalAAScaleSamples - use to enable or disable automatic scaling of the sample count with TAA upsampling factor
r.InvertTemporalJitterX - use to invert the X-values of TAA Halton sequences
r.InvertTemporalJitterY - use to invert the Y-values of TAA Halton sequences

These are set to 1, 0, and 0 respectively by default, matching pre-existing behavior.

diff --git a/Engine/Source/Runtime/Renderer/Private/SceneVisibility.cpp b/Engine/Source/Runtime/Renderer/Private/SceneVisibility.cpp
--- a/Engine/Source/Runtime/Renderer/Private/SceneVisibility.cpp
+++ b/Engine/Source/Runtime/Renderer/Private/SceneVisibility.cpp
@@ -94,6 +94,24 @@ static TAutoConsoleVariable<int32> CVarTemporalAASamples(
 	TEXT("Number of jittered positions for temporal AA (4, 8=default, 16, 32, 64)."),
 	ECVF_RenderThreadSafe);
 
+static TAutoConsoleVariable<int32> CVarTemporalAAScaleSamples(
+	TEXT("r.TemporalAAScaleSamples"),
+	1,
+	TEXT("Whether or not to scale the number of jittered positions for temporal AA when upsampling to maintain a consistent density."),
+	ECVF_RenderThreadSafe);
+
+static TAutoConsoleVariable<int32> CVarInvertTemporalJitterX(
+	TEXT("r.InvertTemporalJitterX"),
+	0,
+	TEXT("Whether or not to invert the X value of jittered positions for temporal AA."),
+	ECVF_RenderThreadSafe);
+
+static TAutoConsoleVariable<int32> CVarInvertTemporalJitterY(
+	TEXT("r.InvertTemporalJitterY"),
+	0,
+	TEXT("Whether or not to invert the Y value of jittered positions for temporal AA."),
+	ECVF_RenderThreadSafe);
+
 static int32 GHZBOcclusion = 0;
 static FAutoConsoleVariableRef CVarHZBOcclusion(
 	TEXT("r.HZBOcclusion"),
@@ -4865,6 +4883,7 @@ void FSceneRenderer::PrepareViewStateForVisibility(const FSceneTexturesConfig& S
 		
 		// Subpixel jitter for temporal AA
 		int32 CVarTemporalAASamplesValue = CVarTemporalAASamples.GetValueOnRenderThread();
+		const bool bShouldScaleTemporalAASampleCount = (CVarTemporalAAScaleSamples.GetValueOnRenderThread() != 0);
 
 		EMainTAAPassConfig TAAConfig = GetMainTAAPassConfig(View);
 
@@ -4889,7 +4908,7 @@ void FSceneRenderer::PrepareViewStateForVisibility(const FSceneTexturesConfig& S
 					TemporalAASamples = FMath::Clamp(CVarTemporalAASamplesValue, 1, 255);
 				}
 
-				if (bTemporalUpsampling)
+				if (bTemporalUpsampling && bShouldScaleTemporalAASampleCount)
 				{
 					// When doing TAA upsample with screen percentage < 100%, we need extra temporal samples to have a
 					// constant temporal sample density for final output pixels to avoid output pixel aligned converging issues.
@@ -5034,6 +5053,16 @@ void FSceneRenderer::PrepareViewStateForVisibility(const FSceneTexturesConfig& S
 				SampleY = r * FMath::Sin( Theta );
 			}
 
+			if (CVarInvertTemporalJitterX.GetValueOnRenderThread() != 0)
+			{
+				SampleX = -SampleX;
+			}
+
+			if (CVarInvertTemporalJitterY.GetValueOnRenderThread() != 0)
+			{
+				SampleY = -SampleY;
+			}
+
 			View.TemporalJitterSequenceLength = TemporalAASamples;
 			View.TemporalJitterIndex = TemporalSampleIndex;
 			View.TemporalJitterPixels.X = SampleX;
