Added ECVF_Scalability flag to r.Shadow.NaniteLODBias

diff --git a/Engine/Source/Runtime/Renderer/Private/Shadows/ShadowSceneRenderer.cpp b/Engine/Source/Runtime/Renderer/Private/Shadows/ShadowSceneRenderer.cpp
--- a/Engine/Source/Runtime/Renderer/Private/Shadows/ShadowSceneRenderer.cpp
+++ b/Engine/Source/Runtime/Renderer/Private/Shadows/ShadowSceneRenderer.cpp
@@ -53,7 +53,7 @@ static TAutoConsoleVariable<float> CVarNaniteShadowsLODBias(
 	TEXT("r.Shadow.NaniteLODBias"),
 	1.0f,
 	TEXT("LOD bias for nanite geometry in shadows. 0 = full detail. >0 = reduced detail."),
-	ECVF_RenderThreadSafe);
+	ECVF_Scalability | ECVF_RenderThreadSafe);
 
 TAutoConsoleVariable<int32> CVarVirtualShadowOnePassProjection(
 	TEXT("r.Shadow.Virtual.OnePassProjection"),
