Change Displacement RVT blend mode.
Displacement now does opacity blend to RVT instead of Max blend.
#tests Editor test of RVT with displacement decal

diff --git a/Engine/Shaders/Private/VirtualTextureMaterial.usf b/Engine/Shaders/Private/VirtualTextureMaterial.usf
--- a/Engine/Shaders/Private/VirtualTextureMaterial.usf
+++ b/Engine/Shaders/Private/VirtualTextureMaterial.usf
@@ -131,7 +131,7 @@ void FPixelShaderInOut_MainPS(
 	float PackedHeight = PackWorldHeight(WorldHeight, ResolvedView.RuntimeVirtualTexturePackHeight);
 	Out.MRT[0] = float4(PackedHeight, 0, 0, 1);
 #elif defined(OUT_DISPLACEMENT)
-	Out.MRT[0] = float4(Displacement, 0, 0, 1);
+	Out.MRT[0] = float4(Displacement, 0, 0, 1) * Opacity;
 #endif
 }
 
diff --git a/Engine/Source/Runtime/Renderer/Private/VT/RuntimeVirtualTextureRender.cpp b/Engine/Source/Runtime/Renderer/Private/VT/RuntimeVirtualTextureRender.cpp
--- a/Engine/Source/Runtime/Renderer/Private/VT/RuntimeVirtualTextureRender.cpp
+++ b/Engine/Source/Runtime/Renderer/Private/VT/RuntimeVirtualTextureRender.cpp
@@ -382,7 +382,7 @@ namespace RuntimeVirtualTexture
 
 		static FRHIBlendState* GetBlendState(uint8 OutputAttributeMask)
 		{
-			return TStaticBlendState< CW_RED, BO_Max, BF_One, BF_One, BO_Add, BF_One, BF_One >::GetRHI();
+			return TStaticBlendState< CW_RED, BO_Add, BF_One, BF_InverseSourceAlpha, BO_Add, BF_Zero, BF_One >::GetRHI();
 		}
 	};
 
