Add RVT volume option for only using the streaming low mips.
This only applies when the streaming low mips are already enabled and showing. It forces the use of only the streaming virtual texture with no runtime pages. This is intended for workflows that use RVT only as a means to generate some baked world texture for sampling.
Also changed "Show in Editor" to be a persistent setting, since that is what would be required for such a baked only workflow.
To help with debug only workflow to show SVT in editor, added r.VT.RVT.StreamingMips.UseInEditor which when set to 0 will force SVT off in editor, 2 will force on in editor, and 1 will obey the RVT volume setting.

diff --git a/Engine/Source/Editor/VirtualTexturingEditor/Private/RuntimeVirtualTextureDetailsCustomization.cpp b/Engine/Source/Editor/VirtualTexturingEditor/Private/RuntimeVirtualTextureDetailsCustomization.cpp
--- a/Engine/Source/Editor/VirtualTexturingEditor/Private/RuntimeVirtualTextureDetailsCustomization.cpp
+++ b/Engine/Source/Editor/VirtualTexturingEditor/Private/RuntimeVirtualTextureDetailsCustomization.cpp
@@ -313,7 +313,7 @@ bool FRuntimeVirtualTextureComponentDetailsCustomization::IsBuildStreamedMipsEna
 {
 	if (URuntimeVirtualTextureComponent* Component = RuntimeVirtualTextureComponent.Get())
 	{
-		return Component->GetVirtualTexture() != nullptr;
+		return Component->GetVirtualTexture() != nullptr && Component->NumStreamingMips() > 0;
 	}
 	return false;
 }
diff --git a/Engine/Source/Runtime/Engine/Classes/Components/RuntimeVirtualTextureComponent.h b/Engine/Source/Runtime/Engine/Classes/Components/RuntimeVirtualTextureComponent.h
--- a/Engine/Source/Runtime/Engine/Classes/Components/RuntimeVirtualTextureComponent.h
+++ b/Engine/Source/Runtime/Engine/Classes/Components/RuntimeVirtualTextureComponent.h
@@ -89,9 +89,13 @@ protected:
 	UPROPERTY(EditAnywhere, BlueprintReadOnly, AdvancedDisplay, Category = StreamingVirtualTexture, meta = (DisplayName = "Fixed Color", HideAlphaChannel, EditCondition = bUseStreamingMipsFixedColor))
 	FLinearColor StreamingMipsFixedColor;
 
+	/** Whenever streaming low mips are in use, only show the streaming mips and never show runtime generated pages. */
+	UPROPERTY(EditAnywhere, AdvancedDisplay, Category = StreamingVirtualTexture, meta = (DisplayName = "Use Only Streaming Texture"))
+	bool bUseStreamingMipsOnly = false;
+
 	/** Use streaming low mips when rendering in editor. Set true to view and debug the baked streaming low mips. */
-	UPROPERTY(EditAnywhere, Transient, AdvancedDisplay, Category = StreamingVirtualTexture, meta = (DisplayName = "View in Editor"))
-	bool bUseStreamingLowMipsInEditor = false;
+	UPROPERTY(EditAnywhere, AdvancedDisplay, Category = StreamingVirtualTexture, meta = (DisplayName = "View in Editor"))
+	bool bUseStreamingMipsInEditor = false;
 
 #if WITH_EDITOR
 	/** Delegate handle for our function called on PIE end. */
@@ -146,6 +150,9 @@ public:
 	/** Public getter for streaming mips fixed color. */
 	ENGINE_API FLinearColor GetStreamingMipsFixedColor() const;
 
+	/** Public getter for using streaming low mips only. */
+	ENGINE_API bool IsStreamingLowMipsOnly() { return bUseStreamingMipsOnly; }
+
 	/** Public getter for lossy compression setting. */
 	TEnumAsByte<ETextureLossyCompressionAmount> GetLossyCompressionAmount() const { return LossyCompressionAmount; }
 
diff --git a/Engine/Source/Runtime/Engine/Private/Components/RuntimeVirtualTextureComponent.cpp b/Engine/Source/Runtime/Engine/Private/Components/RuntimeVirtualTextureComponent.cpp
--- a/Engine/Source/Runtime/Engine/Private/Components/RuntimeVirtualTextureComponent.cpp
+++ b/Engine/Source/Runtime/Engine/Private/Components/RuntimeVirtualTextureComponent.cpp
@@ -3,27 +3,42 @@
 #include "Components/RuntimeVirtualTextureComponent.h"
 
 #include "Async/TaskGraphInterfaces.h"
-#include "GameDelegates.h"
+#include "ComponentRecreateRenderStateContext.h"
 #include "Engine/Texture.h"
 #include "Engine/World.h"
-#include "Logging/MessageLog.h"
+#include "GameDelegates.h"
 #include "GameFramework/Actor.h"
+#include "Logging/MessageLog.h"
 #include "Misc/UObjectToken.h"
 #include "Misc/MapErrors.h"
+#include "RenderUtils.h"
+#include "RHIGlobals.h"
 #include "SceneInterface.h"
+#include "SceneUtils.h"
 #include "UObject/UObjectIterator.h"
 #include "UObject/UnrealType.h"
 #include "VT/RuntimeVirtualTexture.h"
 #include "VT/VirtualTexture.h"
 #include "VT/VirtualTextureBuilder.h"
-#include "RenderUtils.h"
-#include "RHIGlobals.h"
-#include "SceneUtils.h"
 
 #include UE_INLINE_GENERATED_CPP_BY_NAME(RuntimeVirtualTextureComponent)
 
 #define LOCTEXT_NAMESPACE "URuntimeVirtualTextureComponent"
 
+static TAutoConsoleVariable<int32> CVarVTStreamingMipsShowInEditor(
+	TEXT("r.VT.RVT.StreamingMips.UseInEditor"),
+	1,
+	TEXT("Use streaming mips for RVT when in Editor.\n")
+	TEXT("  0: Never use.\n")
+	TEXT("  1: Use the setting from RVT component (default).\n")
+	TEXT("  2: Always use when available.\n"),
+	FConsoleVariableDelegate::CreateLambda([](IConsoleVariable* InVariable)
+		{
+			FGlobalComponentRecreateRenderStateContext Context;
+		}),
+	ECVF_Default);
+
+
 URuntimeVirtualTextureComponent::URuntimeVirtualTextureComponent(const FObjectInitializer& ObjectInitializer)
 	: Super(ObjectInitializer)
 	, EnableInGamePerPlatform(true)
@@ -234,9 +249,13 @@ bool URuntimeVirtualTextureComponent::IsStreamingLowMips(EShadingPath ShadingPat
 	checkf(IsActiveInWorld(), TEXT("This function should never be called for a world where we're inactive"));
 	
 #if WITH_EDITOR
-	if (!bUseStreamingLowMipsInEditor && GIsEditor)
+	if (GIsEditor)
 	{
-		return false;
+		const int32 ShowStreamingMipsInEditor = CVarVTStreamingMipsShowInEditor.GetValueOnAnyThread();
+		if (ShowStreamingMipsInEditor == 0 || (ShowStreamingMipsInEditor == 1 && !bUseStreamingMipsInEditor))
+		{
+			return false;
+		}
 	}
 #endif
 	return VirtualTexture != nullptr && StreamingTexture != nullptr && StreamingTexture->GetVirtualTexture(ShadingPath) != nullptr;
@@ -382,7 +401,8 @@ void URuntimeVirtualTextureComponent::InitializeStreamingTexture(EShadingPath Sh
 bool URuntimeVirtualTextureComponent::CanEditChange(const FProperty* InProperty) const
 {
 	bool bCanEdit = Super::CanEditChange(InProperty);
-	if (InProperty->GetFName() == GET_MEMBER_NAME_CHECKED(URuntimeVirtualTextureComponent, bUseStreamingLowMipsInEditor))
+	if (InProperty->GetFName() == GET_MEMBER_NAME_CHECKED(URuntimeVirtualTextureComponent, bUseStreamingMipsInEditor) || 
+		InProperty->GetFName() == GET_MEMBER_NAME_CHECKED(URuntimeVirtualTextureComponent, bUseStreamingMipsOnly))
 	{
 		bCanEdit &= GetVirtualTexture() != nullptr && GetStreamingTexture() != nullptr;
 	}
diff --git a/Engine/Source/Runtime/Renderer/Private/VT/RuntimeVirtualTextureSceneProxy.cpp b/Engine/Source/Runtime/Renderer/Private/VT/RuntimeVirtualTextureSceneProxy.cpp
--- a/Engine/Source/Runtime/Renderer/Private/VT/RuntimeVirtualTextureSceneProxy.cpp
+++ b/Engine/Source/Runtime/Renderer/Private/VT/RuntimeVirtualTextureSceneProxy.cpp
@@ -62,20 +62,15 @@ FRuntimeVirtualTextureSceneProxy::FRuntimeVirtualTextureSceneProxy(URuntimeVirtu
 			// The producer description is calculated using the transform to determine the aspect ratio
 			FVTProducerDescription ProducerDesc;
 			VirtualTexture->GetProducerDescription(ProducerDesc, InitSettings, Transform);
-			VirtualTextureSize = FIntPoint(ProducerDesc.BlockWidthInTiles * ProducerDesc.TileSize, ProducerDesc.BlockHeightInTiles * ProducerDesc.TileSize);
-			// We only need to dirty flush up to the producer description MaxLevel which accounts for the RemoveLowMips
-			MaxDirtyLevel = ProducerDesc.MaxLevel;
 
 			const ERuntimeVirtualTextureMaterialType MaterialType = VirtualTexture->GetMaterialType();
 			const bool bClearTextures = VirtualTexture->GetClearTextures();
+
+			// Get streaming texture if it is valid.
+			UVirtualTexture2D* StreamingTexture = nullptr;
+
 			FSceneInterface* SceneInterface = InComponent->GetScene();
 			const EShadingPath ShadingPath = SceneInterface ? SceneInterface->GetShadingPath() : EShadingPath::Deferred;
-
-			// The producer object created here will be passed into the virtual texture system which will take ownership.
-			IVirtualTexture* Producer = new FRuntimeVirtualTextureProducer(ProducerDesc, ProducerId, MaterialType, bClearTextures, SceneInterface, Transform, Bounds);
-		
-			// Create a producer for the streaming low mips. 
-			// This is bound with the main producer so that one allocated VT can use both runtime or streaming producers dependent on mip level.
 			if (InComponent->IsStreamingLowMips(ShadingPath))
 			{
 				if (CVarVTStreamingMips.GetValueOnAnyThread() == 0)
@@ -105,23 +100,56 @@ FRuntimeVirtualTextureSceneProxy::FRuntimeVirtualTextureSceneProxy(URuntimeVirtu
 				}
 				else
 				{
-					UVirtualTexture2D* StreamingTexture = InComponent->GetStreamingTexture()->GetVirtualTexture(ShadingPath);
+					StreamingTexture = InComponent->GetStreamingTexture()->GetVirtualTexture(ShadingPath);
+				}
+			}
 
-					FVTProducerDescription StreamingProducerDesc;
-					IVirtualTexture* StreamingProducer = RuntimeVirtualTexture::CreateStreamingTextureProducer(StreamingTexture, ProducerDesc, StreamingProducerDesc);
+			// The producer object created here will be passed into the virtual texture system which will take ownership.
+			IVirtualTexture* Producer = nullptr;
 
-					const int32 NumLevels = (int32)FMath::CeilLogTwo(FMath::Max(ProducerDesc.BlockWidthInTiles, ProducerDesc.BlockHeightInTiles));
-					const int32 NumStreamingLevels = (int32)FMath::CeilLogTwo(FMath::Max(StreamingProducerDesc.BlockWidthInTiles, StreamingProducerDesc.BlockHeightInTiles));
-					ensure(NumLevels >= NumStreamingLevels);
-					const int32 TransitionLevel = NumLevels - NumStreamingLevels;
+			// Create a producer for the streaming low mips. 
+			// This is bound with the main producer so that one allocated VT can use both runtime or streaming producers dependent on mip level.
+			if (StreamingTexture == nullptr)
+			{
+				// Create the runtime virtual texture producer.
+				Producer = new FRuntimeVirtualTextureProducer(ProducerDesc, ProducerId, MaterialType, bClearTextures, SceneInterface, Transform, Bounds);
 
-					Producer = RuntimeVirtualTexture::BindStreamingTextureProducer(Producer, StreamingProducer, TransitionLevel);
+				// We only need to dirty flush up to the producer description MaxLevel which accounts for the RemoveLowMips
+				MaxDirtyLevel = ProducerDesc.MaxLevel;
+			}
+			else
+			{
+				// Create the streaming virtual texture producer.
+				FVTProducerDescription StreamingProducerDesc;
+				IVirtualTexture* StreamingProducer = RuntimeVirtualTexture::CreateStreamingTextureProducer(StreamingTexture, ProducerDesc, StreamingProducerDesc);
 
-					// Any dirty flushes don't need to flush the streaming mips (they only change with a build step).
-					MaxDirtyLevel = TransitionLevel - 1;
+				if (InComponent->IsStreamingLowMipsOnly())
+				{
+					// Clamp the the runtime virtual texture producer dimensions to the streaming virtual texture dimensions.
+					// This will force to only using streaming pages.
+					ProducerDesc.BlockWidthInTiles = StreamingProducerDesc.BlockWidthInTiles;
+					ProducerDesc.BlockHeightInTiles = StreamingProducerDesc.BlockHeightInTiles;
+					ProducerDesc.MaxLevel = StreamingProducerDesc.MaxLevel;
 				}
+
+				// Create the runtime virtual texture producer.
+				Producer = new FRuntimeVirtualTextureProducer(ProducerDesc, ProducerId, MaterialType, bClearTextures, SceneInterface, Transform, Bounds);
+
+				// Bind the runtime virtual texture producer to the streaming producer.
+				const int32 NumLevels = (int32)FMath::CeilLogTwo(FMath::Max(ProducerDesc.BlockWidthInTiles, ProducerDesc.BlockHeightInTiles));
+				const int32 NumStreamingLevels = (int32)FMath::CeilLogTwo(FMath::Max(StreamingProducerDesc.BlockWidthInTiles, StreamingProducerDesc.BlockHeightInTiles));
+				ensure(NumLevels >= NumStreamingLevels);
+				const int32 TransitionLevel = NumLevels - NumStreamingLevels;
+
+				Producer = RuntimeVirtualTexture::BindStreamingTextureProducer(Producer, StreamingProducer, TransitionLevel);
+				
+				// Any dirty flushes don't need to flush the streaming mips (they only change with a build step).
+				MaxDirtyLevel = TransitionLevel - 1;
 			}
 
+			// Store effective virtual texture size used when calculating dirty regions.
+			VirtualTextureSize = FIntPoint(ProducerDesc.BlockWidthInTiles * ProducerDesc.TileSize, ProducerDesc.BlockHeightInTiles * ProducerDesc.TileSize);
+
 			// The Initialize() call will allocate the virtual texture by spawning work on the render thread.
 			VirtualTexture->Initialize(Producer, ProducerDesc, Transform, Bounds);
 
