Fix BasicRayTracing test - it was missing CommitRayTracingBindings call.

diff --git a/Engine/Source/Runtime/Renderer/Private/RayTracing/RayTracingBasicShaders.cpp b/Engine/Source/Runtime/Renderer/Private/RayTracing/RayTracingBasicShaders.cpp
--- a/Engine/Source/Runtime/Renderer/Private/RayTracing/RayTracingBasicShaders.cpp
+++ b/Engine/Source/Runtime/Renderer/Private/RayTracing/RayTracingBasicShaders.cpp
@@ -58,6 +58,7 @@ void DispatchBasicOcclusionRays(FRHICommandList& RHICmdList, FRHIRayTracingScene
 
 	RHICmdList.SetRayTracingHitGroup(Scene, 0, 0, 0, RayTracingPipeline.PipelineState, 0, 0, nullptr, 0, nullptr, 0);
 	RHICmdList.SetRayTracingMissShader(Scene, 0, RayTracingPipeline.PipelineState, 0, 0, nullptr, 0);
+	RHICmdList.CommitRayTracingBindings(Scene);
 
 	FBasicOcclusionMainRGS::FParameters OcclusionParameters;
 	OcclusionParameters.TLAS = SceneView;
@@ -75,6 +76,7 @@ void DispatchBasicIntersectionRays(FRHICommandList& RHICmdList, FRHIRayTracingSc
 
 	RHICmdList.SetRayTracingHitGroup(Scene, 0, 0, 0, RayTracingPipeline.PipelineState, 0, 0, nullptr, 0, nullptr, 0);
 	RHICmdList.SetRayTracingMissShader(Scene, 0, RayTracingPipeline.PipelineState, 0, 0, nullptr, 0);
+	RHICmdList.CommitRayTracingBindings(Scene);
 
 	FBasicIntersectionMainRGS::FParameters OcclusionParameters;
 	OcclusionParameters.TLAS = SceneView;
