Fixed r.water.waterinfo.rendermethod 1 not properly transitioning the water info texture to SRVMask after rendering to it.

#rnx

diff --git a/Engine/Source/Runtime/Renderer/Private/WaterInfoTextureRendering.cpp b/Engine/Source/Runtime/Renderer/Private/WaterInfoTextureRendering.cpp
--- a/Engine/Source/Runtime/Renderer/Private/WaterInfoTextureRendering.cpp
+++ b/Engine/Source/Runtime/Renderer/Private/WaterInfoTextureRendering.cpp
@@ -1021,5 +1021,8 @@ void RenderWaterInfoTexture(
 				PassParameters,
 				Viewport);
 		}
+
+		// Make sure the texture is in the SRV state by the time it is used in water draws (referenced outside of RDG).
+		GraphBuilder.UseExternalAccessMode(OutputTexture, ERHIAccess::SRVMask);
 	}
 }
