Add console variable r.VT.RVT.StreamingMips.UseAlways
This forces the use of RVT streaming mips and disables runtime generation of pages.
Can be useful for debugging. Especially when combined with the recently added r.VT.RVT.StreamingMips.UseInEditor
Also moved related cvar locations in code to live next to one another.
#tests In editor map

diff --git a/Engine/Source/Runtime/Engine/Classes/Components/RuntimeVirtualTextureComponent.h b/Engine/Source/Runtime/Engine/Classes/Components/RuntimeVirtualTextureComponent.h
--- a/Engine/Source/Runtime/Engine/Classes/Components/RuntimeVirtualTextureComponent.h
+++ b/Engine/Source/Runtime/Engine/Classes/Components/RuntimeVirtualTextureComponent.h
@@ -147,12 +147,12 @@ public:
 	/** Get if we want to use any streaming low mips on this component. */
 	ENGINE_API bool IsStreamingLowMips(EShadingPath ShadingPath) const;
 
+	/** Public getter for using streaming low mips only. */
+	ENGINE_API bool IsStreamingLowMipsOnly();
+
 	/** Public getter for streaming mips fixed color. */
 	ENGINE_API FLinearColor GetStreamingMipsFixedColor() const;
 
-	/** Public getter for using streaming low mips only. */
-	ENGINE_API bool IsStreamingLowMipsOnly() { return bUseStreamingMipsOnly; }
-
 	/** Public getter for lossy compression setting. */
 	TEnumAsByte<ETextureLossyCompressionAmount> GetLossyCompressionAmount() const { return LossyCompressionAmount; }
 
diff --git a/Engine/Source/Runtime/Engine/Private/Components/RuntimeVirtualTextureComponent.cpp b/Engine/Source/Runtime/Engine/Private/Components/RuntimeVirtualTextureComponent.cpp
--- a/Engine/Source/Runtime/Engine/Private/Components/RuntimeVirtualTextureComponent.cpp
+++ b/Engine/Source/Runtime/Engine/Private/Components/RuntimeVirtualTextureComponent.cpp
@@ -25,6 +25,15 @@
 
 #define LOCTEXT_NAMESPACE "URuntimeVirtualTextureComponent"
 
+static TAutoConsoleVariable<bool> CVarVTStreamingMips(
+	TEXT("r.VT.RVT.StreamingMips"),
+	true,
+	TEXT("Enable streaming mips for RVT"),
+	FConsoleVariableDelegate::CreateLambda([](IConsoleVariable* InVariable)	{ FGlobalComponentRecreateRenderStateContext Context; }),
+	ECVF_Default);
+
+#if WITH_EDITOR
+
 static TAutoConsoleVariable<int32> CVarVTStreamingMipsShowInEditor(
 	TEXT("r.VT.RVT.StreamingMips.UseInEditor"),
 	1,
@@ -32,10 +41,16 @@ static TAutoConsoleVariable<int32> CVarVTStreamingMipsShowInEditor(
 	TEXT("  0: Never use.\n")
 	TEXT("  1: Use the setting from RVT component (default).\n")
 	TEXT("  2: Always use when available.\n"),
-	FConsoleVariableDelegate::CreateLambda([](IConsoleVariable* InVariable)
-		{
-			FGlobalComponentRecreateRenderStateContext Context;
-		}),
+	FConsoleVariableDelegate::CreateLambda([](IConsoleVariable* InVariable) { FGlobalComponentRecreateRenderStateContext Context; }),
+	ECVF_Default);
+
+#endif
+
+static TAutoConsoleVariable<bool> CVarVTStreamingMipsUseAlways(
+	TEXT("r.VT.RVT.StreamingMips.UseAlways"),
+	false,
+	TEXT("Whenever streaming low mips are in use, only show the streaming mips and never show runtime generated pages.\n"),
+	FConsoleVariableDelegate::CreateLambda([](IConsoleVariable* InVariable) { FGlobalComponentRecreateRenderStateContext Context; }),
 	ECVF_Default);
 
 
@@ -258,7 +273,12 @@ bool URuntimeVirtualTextureComponent::IsStreamingLowMips(EShadingPath ShadingPat
 		}
 	}
 #endif
-	return VirtualTexture != nullptr && StreamingTexture != nullptr && StreamingTexture->GetVirtualTexture(ShadingPath) != nullptr;
+	return VirtualTexture != nullptr && StreamingTexture != nullptr && StreamingTexture->GetVirtualTexture(ShadingPath) != nullptr && CVarVTStreamingMips.GetValueOnAnyThread();
+}
+
+bool URuntimeVirtualTextureComponent::IsStreamingLowMipsOnly()
+{
+	return bUseStreamingMipsOnly || CVarVTStreamingMipsUseAlways.GetValueOnAnyThread();
 }
 
 bool IsCompatibleFormat(URuntimeVirtualTexture const& RuntimeVirtualTexture, UVirtualTexture2D const& StreamingVirtualTexture)
diff --git a/Engine/Source/Runtime/Engine/Private/VT/VirtualTextureScalability.cpp b/Engine/Source/Runtime/Engine/Private/VT/VirtualTextureScalability.cpp
--- a/Engine/Source/Runtime/Engine/Private/VT/VirtualTextureScalability.cpp
+++ b/Engine/Source/Runtime/Engine/Private/VT/VirtualTextureScalability.cpp
@@ -71,6 +71,12 @@ namespace VirtualTextureScalability
 		TEXT("Increasing this threshold reduces the chances that an in-use frame is considered free."),
 		ECVF_RenderThreadSafe);
 
+	static TAutoConsoleVariable<int32> CVarVTKeepDirtyPageMappedFrameThreshold(
+		TEXT("r.VT.RVT.DirtyPagesKeptMappedFrames"), 
+		8,
+		TEXT("Number of frames since the last time a VT page was used, before it will be unmapped instead of simply updated during RVT page invalidation."),
+		ECVF_RenderThreadSafe);
+
 	static const int NumScalabilityGroups = 3;
 
 	static float GTileCountBiases[NumScalabilityGroups] = { 0 };
@@ -210,6 +216,11 @@ namespace VirtualTextureScalability
 		return FMath::Max(CVarVTPageFreeThreshold.GetValueOnRenderThread(), 0);
 	}
 
+	uint32 GetKeepDirtyPageMappedFrameThreshold()
+	{
+		return FMath::Max(CVarVTKeepDirtyPageMappedFrameThreshold.GetValueOnRenderThread(), 0);
+	}
+
 	int32 GetRuntimeVirtualTextureSizeBias(uint32 GroupIndex)
 	{
 		return GroupIndex < NumScalabilityGroups ? GTileCountBiases[GroupIndex] : 0;
diff --git a/Engine/Source/Runtime/Engine/Public/VT/VirtualTextureScalability.h b/Engine/Source/Runtime/Engine/Public/VT/VirtualTextureScalability.h
--- a/Engine/Source/Runtime/Engine/Public/VT/VirtualTextureScalability.h
+++ b/Engine/Source/Runtime/Engine/Public/VT/VirtualTextureScalability.h
@@ -18,6 +18,8 @@ namespace VirtualTextureScalability
 	ENGINE_API int32 GetMaxAllocatedVTReleasedPerFrame();
 	/** Get the number of frames a page must be unused, before it's considered free */
 	ENGINE_API uint32 GetPageFreeThreshold();
+	/** Get the number of frames a page must be unused, before it is unmapped instead of updated during RVT page invalidation. */
+	ENGINE_API uint32 GetKeepDirtyPageMappedFrameThreshold();
 	/** Get resolution bias for runtime virtual textures. */
 	ENGINE_API int32 GetRuntimeVirtualTextureSizeBias();
 	/** Get resolution bias for runtime virtual textures. */
diff --git a/Engine/Source/Runtime/Renderer/Private/VT/RuntimeVirtualTextureSceneProxy.cpp b/Engine/Source/Runtime/Renderer/Private/VT/RuntimeVirtualTextureSceneProxy.cpp
--- a/Engine/Source/Runtime/Renderer/Private/VT/RuntimeVirtualTextureSceneProxy.cpp
+++ b/Engine/Source/Runtime/Renderer/Private/VT/RuntimeVirtualTextureSceneProxy.cpp
@@ -15,22 +15,6 @@
 
 #define LOCTEXT_NAMESPACE "VirtualTexture"
 
-static TAutoConsoleVariable<int32> CVarVTStreamingMips(
-	TEXT("r.VT.RVT.StreamingMips"),
-	1,
-	TEXT("Enable streaming mips for RVT"),
-	FConsoleVariableDelegate::CreateLambda([](IConsoleVariable* InVariable)
-	{
-		FGlobalComponentRecreateRenderStateContext Context;
-	}),
-	ECVF_Default);
-
-static TAutoConsoleVariable<int32> CVarDirtyPagesKeptMappedFrames(
-	TEXT("r.VT.RVT.DirtyPagesKeptMappedFrames"),
-	8,
-	TEXT("When invalidating RVT pages, we keep them mapped if they gave feedback in the last N frames."),
-	ECVF_RenderThreadSafe);
-
 int32 FRuntimeVirtualTextureSceneProxy::ProducerIdGenerator = 1;
 
 FRuntimeVirtualTextureSceneProxy::FRuntimeVirtualTextureSceneProxy(URuntimeVirtualTextureComponent* InComponent)
@@ -73,19 +57,7 @@ FRuntimeVirtualTextureSceneProxy::FRuntimeVirtualTextureSceneProxy(URuntimeVirtu
 			const EShadingPath ShadingPath = SceneInterface ? SceneInterface->GetShadingPath() : EShadingPath::Deferred;
 			if (InComponent->IsStreamingLowMips(ShadingPath))
 			{
-				if (CVarVTStreamingMips.GetValueOnAnyThread() == 0)
-				{
-#if !UE_BUILD_SHIPPING
-					// Notify that streaming texture is turned off.
-					OnScreenWarningDelegateHandle = FRendererOnScreenNotification::Get().AddLambda([](FCoreDelegates::FSeverityMessageMap& OutMessages)
-					{
-						OutMessages.Add(
-							FCoreDelegates::EOnScreenMessageSeverity::Warning,
-							LOCTEXT("SVTDisabled", "Runtime Virtual Texture streaming mips disabled."));
-					});
-#endif
-				}
-				else if (InComponent->IsStreamingTextureInvalid(ShadingPath))
+				if (InComponent->IsStreamingTextureInvalid(ShadingPath))
 				{
 #if !UE_BUILD_SHIPPING
 					// Notify that streaming texture is invalid since this can cause performance regression.
@@ -233,7 +205,7 @@ void FRuntimeVirtualTextureSceneProxy::FlushDirtyPages()
 		{
 			// Keeping visible pages mapped reduces update flicker due to the latency in the unmap/feedback/map sequence.
 			// But it potentially creates more page update work since more pages may get updated.
-			const uint32 MaxAgeToKeepMapped = CVarDirtyPagesKeptMappedFrames.GetValueOnRenderThread();
+			const uint32 MaxAgeToKeepMapped = VirtualTextureScalability::GetKeepDirtyPageMappedFrameThreshold();
 
 			//todo[vt]: 
 			// Profile to work out best heuristic for when we should use the CombinedDirtyRect
