Adaptive Virtual Texture: Defer release of producers when managing adaptive page table.
Release of producer can stall if waiting on streaming events, so instead always put in a pending release array and try releasing every update.

diff --git a/Engine/Source/Runtime/Renderer/Private/VT/AdaptiveVirtualTexture.cpp b/Engine/Source/Runtime/Renderer/Private/VT/AdaptiveVirtualTexture.cpp
--- a/Engine/Source/Runtime/Renderer/Private/VT/AdaptiveVirtualTexture.cpp
+++ b/Engine/Source/Runtime/Renderer/Private/VT/AdaptiveVirtualTexture.cpp
@@ -199,19 +199,14 @@ IAllocatedVirtualTexture* FAdaptiveVirtualTexture::AllocateVirtualTexture(
 }
 
 /** Destroy an allocated virtual texture and release its producers. */
-void FAdaptiveVirtualTexture::DestroyVirtualTexture(FVirtualTextureSystem* InSystem, IAllocatedVirtualTexture* InAllocatedVT)
+void FAdaptiveVirtualTexture::DestroyVirtualTexture(FVirtualTextureSystem* InSystem, IAllocatedVirtualTexture* InAllocatedVT, TArray<FVirtualTextureProducerHandle>& OutProducersToRelease)
 {
 	FAllocatedVTDescription const& Desc = InAllocatedVT->GetDescription();
-	TArray<FVirtualTextureProducerHandle, TInlineAllocator<8>> ProducersToRelease;
 	for (int32 LayerIndex = 0; LayerIndex < Desc.NumTextureLayers; ++LayerIndex)
 	{
-		ProducersToRelease.AddUnique(Desc.ProducerHandle[LayerIndex]);
+		OutProducersToRelease.AddUnique(Desc.ProducerHandle[LayerIndex]);
 	}
 	InSystem->DestroyVirtualTexture(InAllocatedVT);
-	for (int32 ProducerIndex = 0; ProducerIndex < ProducersToRelease.Num(); ++ProducerIndex)
-	{
-		InSystem->ReleaseProducer(ProducersToRelease[ProducerIndex]);
-	}
 }
 
 /** Remaps the page mappings from one allocated virtual texture to another. */
@@ -275,16 +270,21 @@ void FAdaptiveVirtualTexture::Init(FRHICommandListBase& RHICmdList, FVirtualText
 
 void FAdaptiveVirtualTexture::Destroy(FVirtualTextureSystem* InSystem)
 {
-	DestroyVirtualTexture(InSystem, AllocatedVirtualTextureLowMips);
+	DestroyVirtualTexture(InSystem, AllocatedVirtualTextureLowMips, ProducersToRelease);
 
 	for (FAllocation& Allocation : AllocationSlots)
 	{
 		if (Allocation.AllocatedVT != nullptr)
 		{
-			DestroyVirtualTexture(InSystem, Allocation.AllocatedVT);
+			DestroyVirtualTexture(InSystem, Allocation.AllocatedVT, ProducersToRelease);
 		}
 	}
 
+	for (FVirtualTextureProducerHandle Handle : ProducersToRelease)
+	{
+		InSystem->ReleaseProducer(Handle);
+	}
+
 	delete this;
 }
 
@@ -557,7 +557,7 @@ void FAdaptiveVirtualTexture::Allocate(FRHICommandListBase& RHICmdList, FVirtual
 	{
 		// Remap the old allocated virtual texture before destroying it.
 		RemapVirtualTexturePages(InSystem, OldAllocatedVT, NewAllocatedVT, InFrame);
-		DestroyVirtualTexture(InSystem, OldAllocatedVT);
+		DestroyVirtualTexture(InSystem, OldAllocatedVT, ProducersToRelease);
 
 		// Adjust allocation structures.
 		AllocatedVTMap.Remove(GetAllocatedVTHash(OldAllocatedVT), InAllocationIndex);
@@ -612,7 +612,7 @@ void FAdaptiveVirtualTexture::Free(FVirtualTextureSystem* InSystem, uint32 InAll
 	// Destroy allocated virtual texture.
 	const uint32 GridIndex = AllocationSlots[InAllocationIndex].GridIndex;
 	FAllocatedVirtualTexture* OldAllocatedVT = AllocationSlots[InAllocationIndex].AllocatedVT;
-	DestroyVirtualTexture(InSystem, OldAllocatedVT);
+	DestroyVirtualTexture(InSystem, OldAllocatedVT, ProducersToRelease);
 
 	// Remove from all allocation structures.
 	GridIndexMap.Remove(GetGridIndexHash(GridIndex), InAllocationIndex);
@@ -760,4 +760,13 @@ void FAdaptiveVirtualTexture::UpdateAllocations(FVirtualTextureSystem* InSystem,
 	// Clear requests
 	RequestsToMap.Reset();
 	TextureUpdates.Reset();
+
+	// Release any producers
+	for (int32 ProducerIndex = ProducersToRelease.Num() - 1; ProducerIndex >= 0; ProducerIndex--)
+	{
+		if (InSystem->TryReleaseProducer(ProducersToRelease[ProducerIndex]))
+		{
+			ProducersToRelease.RemoveAt(ProducerIndex, 1, EAllowShrinking::No);
+		}
+	}
 }
diff --git a/Engine/Source/Runtime/Renderer/Private/VT/AdaptiveVirtualTexture.h b/Engine/Source/Runtime/Renderer/Private/VT/AdaptiveVirtualTexture.h
--- a/Engine/Source/Runtime/Renderer/Private/VT/AdaptiveVirtualTexture.h
+++ b/Engine/Source/Runtime/Renderer/Private/VT/AdaptiveVirtualTexture.h
@@ -82,7 +82,7 @@ private:
 		FIntPoint InAddressOffset,
 		int32 InLevelOffset);
 
-	static void DestroyVirtualTexture(FVirtualTextureSystem* InSystem, IAllocatedVirtualTexture* InAllocatedVT);
+	static void DestroyVirtualTexture(FVirtualTextureSystem* InSystem, IAllocatedVirtualTexture* InAllocatedVT, TArray<FVirtualTextureProducerHandle>& OutProducersToRelease);
 	static void RemapVirtualTexturePages(FVirtualTextureSystem* InSystem, FAllocatedVirtualTexture* OldAllocatedVT, FAllocatedVirtualTexture* NewAllocatedVT, uint32 InFrame);
 
 private:
@@ -140,4 +140,7 @@ private:
 
 	/** Array of indirection texture updates to process. */
 	TArray<FIndirectionTextureUpdate> TextureUpdates;
-};
+
+	/** */
+	TArray<FVirtualTextureProducerHandle> ProducersToRelease;
+};
\ No newline at end of file
diff --git a/Engine/Source/Runtime/Renderer/Private/VT/VirtualTextureProducer.cpp b/Engine/Source/Runtime/Renderer/Private/VT/VirtualTextureProducer.cpp
--- a/Engine/Source/Runtime/Renderer/Private/VT/VirtualTextureProducer.cpp
+++ b/Engine/Source/Runtime/Renderer/Private/VT/VirtualTextureProducer.cpp
@@ -129,6 +129,24 @@ void FVirtualTextureProducerCollection::ReleaseProducer(FVirtualTextureSystem* S
 	}
 }
 
+bool FVirtualTextureProducerCollection::TryReleaseProducer(FVirtualTextureSystem* System, const FVirtualTextureProducerHandle& Handle)
+{
+	if (FProducerEntry* Entry = GetEntry(Handle))
+	{
+		FGraphEventArray ProducePageTasks;
+		Entry->Producer.GetVirtualTexture()->GatherProducePageDataTasks(Handle, ProducePageTasks);
+
+		if (ProducePageTasks.Num() != 0)
+		{
+			// Don't release with tasks remaining.
+			return false;
+		}
+	}
+
+	ReleaseProducer(System, Handle);
+	return true;
+}
+
 void FVirtualTextureProducerCollection::CallPendingCallbacks()
 {
 	uint32 CallbackIndex = Callbacks[CallbackList_Pending].NextIndex;
diff --git a/Engine/Source/Runtime/Renderer/Private/VT/VirtualTextureProducer.h b/Engine/Source/Runtime/Renderer/Private/VT/VirtualTextureProducer.h
--- a/Engine/Source/Runtime/Renderer/Private/VT/VirtualTextureProducer.h
+++ b/Engine/Source/Runtime/Renderer/Private/VT/VirtualTextureProducer.h
@@ -58,6 +58,7 @@ public:
 	FVirtualTextureProducerCollection();
 	FVirtualTextureProducerHandle RegisterProducer(FRHICommandListBase& RHICmdList, FVirtualTextureSystem* System, const FVTProducerDescription& InDesc, IVirtualTexture* InProducer);
 	void ReleaseProducer(FVirtualTextureSystem* System, const FVirtualTextureProducerHandle& Handle);
+	bool TryReleaseProducer(FVirtualTextureSystem* System, const FVirtualTextureProducerHandle& Handle);
 	
 	void AddDestroyedCallback(const FVirtualTextureProducerHandle& Handle, FVTProducerDestroyedFunction* Function, void* Baton);
 	uint32 RemoveAllCallbacks(const void* Baton);
diff --git a/Engine/Source/Runtime/Renderer/Private/VT/VirtualTextureSystem.cpp b/Engine/Source/Runtime/Renderer/Private/VT/VirtualTextureSystem.cpp
--- a/Engine/Source/Runtime/Renderer/Private/VT/VirtualTextureSystem.cpp
+++ b/Engine/Source/Runtime/Renderer/Private/VT/VirtualTextureSystem.cpp
@@ -803,6 +803,13 @@ void FVirtualTextureSystem::ReleaseProducer(const FVirtualTextureProducerHandle&
 	Producers.ReleaseProducer(this, Handle);
 }
 
+bool FVirtualTextureSystem::TryReleaseProducer(const FVirtualTextureProducerHandle& Handle)
+{
+	check(!bUpdating);
+	UE::TScopeLock Lock(Mutex);
+	return Producers.TryReleaseProducer(this, Handle);
+}
+
 void FVirtualTextureSystem::AddProducerDestroyedCallback(const FVirtualTextureProducerHandle& Handle, FVTProducerDestroyedFunction* Function, void* Baton)
 {
 	check(!bUpdating);
diff --git a/Engine/Source/Runtime/Renderer/Private/VT/VirtualTextureSystem.h b/Engine/Source/Runtime/Renderer/Private/VT/VirtualTextureSystem.h
--- a/Engine/Source/Runtime/Renderer/Private/VT/VirtualTextureSystem.h
+++ b/Engine/Source/Runtime/Renderer/Private/VT/VirtualTextureSystem.h
@@ -126,6 +126,7 @@ public:
 
 	FVirtualTextureProducerHandle RegisterProducer(FRHICommandListBase& RHICmdList, const FVTProducerDescription& InDesc, IVirtualTexture* InProducer);
 	void ReleaseProducer(const FVirtualTextureProducerHandle& Handle);
+	bool TryReleaseProducer(const FVirtualTextureProducerHandle& Handle);
 	void AddProducerDestroyedCallback(const FVirtualTextureProducerHandle& Handle, FVTProducerDestroyedFunction* Function, void* Baton);
 	uint32 RemoveAllProducerDestroyedCallbacks(const void* Baton);
 
