Fix incorrect FOV value on SceneCapture entity.

This was causing the hair strands to look incorrect, due to the FOV being use for computing pixel size based on FOV.

diff --git a/Engine/Source/Runtime/Engine/Classes/Components/SceneCaptureComponent.h b/Engine/Source/Runtime/Engine/Classes/Components/SceneCaptureComponent.h
--- a/Engine/Source/Runtime/Engine/Classes/Components/SceneCaptureComponent.h
+++ b/Engine/Source/Runtime/Engine/Classes/Components/SceneCaptureComponent.h
@@ -26,6 +26,7 @@ struct FSceneCaptureViewInfo : public FSceneViewProjectionData
 	FRotator ViewRotation;
 	EStereoscopicPass StereoPass;
 	int32 StereoViewIndex;
+	float FOV;
 };
 
 #if WITH_EDITORONLY_DATA
diff --git a/Engine/Source/Runtime/Renderer/Private/PlanarReflectionRendering.cpp b/Engine/Source/Runtime/Renderer/Private/PlanarReflectionRendering.cpp
--- a/Engine/Source/Runtime/Renderer/Private/PlanarReflectionRendering.cpp
+++ b/Engine/Source/Runtime/Renderer/Private/PlanarReflectionRendering.cpp
@@ -620,6 +620,7 @@ void FScene::UpdatePlanarReflectionContents(UPlanarReflectionComponent* CaptureC
 			NewView.ProjectionMatrix = ProjectionMatrix;
 			NewView.StereoPass = View.StereoPass;
 			NewView.StereoViewIndex = View.StereoViewIndex;
+			NewView.FOV = View.FOV;
 
 			SceneCaptureViewInfo.Add(NewView);
 		}
diff --git a/Engine/Source/Runtime/Renderer/Private/SceneCaptureRendering.cpp b/Engine/Source/Runtime/Renderer/Private/SceneCaptureRendering.cpp
--- a/Engine/Source/Runtime/Renderer/Private/SceneCaptureRendering.cpp
+++ b/Engine/Source/Runtime/Renderer/Private/SceneCaptureRendering.cpp
@@ -818,6 +818,8 @@ void SetupViewFamilyForSceneCapture(
 		ViewInitOptions.ProjectionMatrix = SceneCaptureViewInfo.ProjectionMatrix;
 		ViewInitOptions.bIsSceneCapture = true;
 		ViewInitOptions.bIsPlanarReflection = bIsPlanarReflection;
+		ViewInitOptions.FOV = SceneCaptureViewInfo.FOV;
+		ViewInitOptions.DesiredFOV = SceneCaptureViewInfo.FOV;
 
 		if (ViewFamily.Scene->GetWorld() != nullptr && ViewFamily.Scene->GetWorld()->GetWorldSettings() != nullptr)
 		{
@@ -885,6 +887,7 @@ static FSceneRenderer* CreateSceneRendererForSceneCapture(
 	const FVector& ViewLocation,
 	const FMatrix& ProjectionMatrix,
 	float MaxViewDistance,
+	float InFOV,
 	bool bCaptureSceneColor,
 	FPostProcessSettings* PostProcessSettings,
 	float PostProcessBlendWeight,
@@ -898,6 +901,7 @@ static FSceneRenderer* CreateSceneRendererForSceneCapture(
 	SceneCaptureViewInfo.StereoPass = EStereoscopicPass::eSSP_FULL;
 	SceneCaptureViewInfo.StereoViewIndex = INDEX_NONE;
 	SceneCaptureViewInfo.ViewRect = FIntRect(0, 0, RenderTargetSize.X, RenderTargetSize.Y);
+	SceneCaptureViewInfo.FOV = InFOV;
 
 	// Use camera position correction for ortho scene captures
 	if(USceneCaptureComponent2D * SceneCaptureComponent2D = Cast<USceneCaptureComponent2D>(SceneCaptureComponent))
@@ -1118,6 +1122,7 @@ void FScene::UpdateSceneCaptureContents(USceneCaptureComponent2D* CaptureCompone
 			ViewLocation, 
 			ProjectionMatrix, 
 			CaptureComponent->MaxViewDistanceOverride, 
+			CaptureComponent->FOVAngle,
 			bUseSceneColorTexture,
 			&CaptureComponent->PostProcessSettings, 
 			CaptureComponent->PostProcessBlendWeight,
@@ -1435,9 +1440,10 @@ void FScene::UpdateSceneCaptureContents(USceneCaptureComponentCube* CaptureCompo
 		FGenerateMipsParams GenerateMipsParams{ TextureTarget->MipsSamplerFilter == TF_Nearest ? SF_Point : (TextureTarget->MipsSamplerFilter == TF_Trilinear ? SF_Trilinear : SF_Bilinear), AM_Clamp, AM_Clamp };
 
 		FIntPoint CaptureSize(TextureTarget->GetSurfaceWidth(), TextureTarget->GetSurfaceHeight());
-		const float FOV = 90 * (float)PI / 360.0f;
+		const float FOVInDegrees = 90.f;
+		const float FOVInRadians = FOVInDegrees * (float)PI / 360.0f;
 
-		auto ComputeProjectionMatrix = [CaptureComponent, Transform, CaptureSize, FOV](ECubeFace TargetFace, FMatrix& OutViewRotationMatrix, FMatrix& OutProjectionMatrix)
+		auto ComputeProjectionMatrix = [CaptureComponent, Transform, CaptureSize, FOVInRadians](ECubeFace TargetFace, FMatrix& OutViewRotationMatrix, FMatrix& OutProjectionMatrix)
 		{
 			if (CaptureComponent->bCaptureRotation)
 			{
@@ -1447,7 +1453,7 @@ void FScene::UpdateSceneCaptureContents(USceneCaptureComponentCube* CaptureCompo
 			{
 				OutViewRotationMatrix = FLocal::CalcCubeFaceTransform(TargetFace);
 			}
-			BuildProjectionMatrix(CaptureSize, FOV, GNearClippingPlane, OutProjectionMatrix);
+			BuildProjectionMatrix(CaptureSize, FOVInRadians, GNearClippingPlane, OutProjectionMatrix);
 		};
 
 		const FVector Location = CaptureComponent->GetComponentToWorld().GetTranslation();
@@ -1480,7 +1486,7 @@ void FScene::UpdateSceneCaptureContents(USceneCaptureComponentCube* CaptureCompo
 
 				FSceneRenderer* SceneRenderer = CreateSceneRendererForSceneCapture(this, CaptureComponent,
 					TextureTarget->GameThread_GetRenderTargetResource(), CaptureSize, ViewRotationMatrix,
-					Location, ProjectionMatrix, CaptureComponent->MaxViewDistanceOverride,
+					Location, ProjectionMatrix, CaptureComponent->MaxViewDistanceOverride, FOVInDegrees,
 					bCaptureSceneColor, &CaptureComponent->PostProcessSettings, CaptureComponent->PostProcessBlendWeight, CaptureComponent->GetViewOwner(), faceidx);
 
 				// When bIsMultipleSceneCapture is true, set bIsFirstSceneRenderer to false, which tells the scene renderer it can skip RHI resource flush, saving performance.
@@ -1576,6 +1582,7 @@ void FScene::UpdateSceneCaptureContents(USceneCaptureComponentCube* CaptureCompo
 				SceneCaptureViewInfos[faceidx].StereoPass = EStereoscopicPass::eSSP_FULL;
 				SceneCaptureViewInfos[faceidx].StereoViewIndex = INDEX_NONE;
 				SceneCaptureViewInfos[faceidx].ViewRect = FIntRect(ViewportOffset.X, ViewportOffset.Y, ViewportOffset.X + CaptureSize.X, ViewportOffset.Y + CaptureSize.Y);
+				SceneCaptureViewInfos[faceidx].FOV = 90.f;
 			}
 
 			// Render target that includes all six tiled faces of the cube map
