Use an array of render queries for LastOcclusionQuery in order to respect the NumBufferedFrames config and do not always stall in the N+1 frame

diff --git a/Engine/Source/Runtime/Renderer/Private/RendererScene.cpp b/Engine/Source/Runtime/Renderer/Private/RendererScene.cpp
--- a/Engine/Source/Runtime/Renderer/Private/RendererScene.cpp
+++ b/Engine/Source/Runtime/Renderer/Private/RendererScene.cpp
@@ -419,6 +419,8 @@ FSceneViewState::FSceneViewState(ERHIFeatureLevel::Type FeatureLevel, FSceneView
 			BeginInitResource(&OcclusionFeedback);
 		}
 	}
+
+	Occlusion.LastOcclusionQueryArray.SetNumZeroed(FOcclusionQueryHelpers::MaxBufferedOcclusionFrames);
 }
 
 FSceneViewState::~FSceneViewState()
diff --git a/Engine/Source/Runtime/Renderer/Private/ScenePrivate.h b/Engine/Source/Runtime/Renderer/Private/ScenePrivate.h
--- a/Engine/Source/Runtime/Renderer/Private/ScenePrivate.h
+++ b/Engine/Source/Runtime/Renderer/Private/ScenePrivate.h
@@ -1646,7 +1646,7 @@ PRAGMA_ENABLE_DEPRECATION_WARNINGS
 		TSet<FPrimitiveOcclusionHistory, FPrimitiveOcclusionHistoryKeyFuncs> PrimitiveOcclusionHistorySet;
 
 		/** The last occlusion query of last frame to test in the following frame to block the GPU. */
-		FRHIRenderQuery* LastOcclusionQuery = nullptr;
+		TArray<FRHIRenderQuery*, TInlineAllocator<FOcclusionQueryHelpers::MaxBufferedOcclusionFrames> > LastOcclusionQueryArray;
 
 		/** The number of queries requested last frame. */
 		uint32 NumRequestedQueries = 0;
diff --git a/Engine/Source/Runtime/Renderer/Private/SceneVisibility.cpp b/Engine/Source/Runtime/Renderer/Private/SceneVisibility.cpp
--- a/Engine/Source/Runtime/Renderer/Private/SceneVisibility.cpp
+++ b/Engine/Source/Runtime/Renderer/Private/SceneVisibility.cpp
@@ -2841,7 +2841,8 @@ void FGPUOcclusionPacket::FProcessVisitor::AddOcclusionQuery(const FOcclusionQue
 		RenderQuery = Packet.View.IndividualOcclusionQueries.BatchPrimitive(Query.Bounds.Origin, Query.Bounds.Extent, DynamicVertexBuffer);
 	}
 
-	Packet.ViewState.Occlusion.LastOcclusionQuery = RenderQuery;
+	const uint32 QueryIndex = FOcclusionQueryHelpers::GetQueryIssueIndex(Packet.ViewState.OcclusionFrameCounter, Packet.OcclusionState.NumBufferedFrames);
+	Packet.ViewState.Occlusion.LastOcclusionQueryArray[QueryIndex] = RenderQuery;
 	Packet.ViewState.Occlusion.NumRequestedQueries++;
 
 	Query.PrimitiveOcclusionHistory->SetCurrentQuery(
@@ -2901,7 +2902,8 @@ void FGPUOcclusionPacket::FProcessVisitor::SubmitThrottledOcclusionQueries()
 
 		FRHIRenderQuery* RenderQuery = Packet.View.IndividualOcclusionQueries.BatchPrimitive(Query->Bounds.Origin, Query->Bounds.Extent, DynamicVertexBuffer);
 
-		Packet.ViewState.Occlusion.LastOcclusionQuery = RenderQuery;
+		const uint32 QueryIndex = FOcclusionQueryHelpers::GetQueryIssueIndex(Packet.ViewState.OcclusionFrameCounter, Packet.OcclusionState.NumBufferedFrames);
+		Packet.ViewState.Occlusion.LastOcclusionQueryArray[QueryIndex] = RenderQuery;
 		Packet.ViewState.Occlusion.NumRequestedQueries++;
 
 		PrimitiveOcclusionHistory->SetCurrentQuery(
@@ -3014,12 +3016,13 @@ void FGPUOcclusion::Unmap(FRHICommandListImmediate& RHICmdListImmediate)
 
 void FGPUOcclusion::WaitForLastOcclusionQuery()
 {
-	if (ViewState.Occlusion.LastOcclusionQuery)
+	const uint32 QueryIndex = FOcclusionQueryHelpers::GetQueryLookupIndex(ViewState.OcclusionFrameCounter, State.NumBufferedFrames);
+	if (ViewState.Occlusion.LastOcclusionQueryArray[QueryIndex])
 	{
 		uint64 Result;
 		const bool bWait = true;
-		RHIGetRenderQueryResult(ViewState.Occlusion.LastOcclusionQuery, Result, bWait);
-		ViewState.Occlusion.LastOcclusionQuery = nullptr;
+		RHIGetRenderQueryResult(ViewState.Occlusion.LastOcclusionQueryArray[QueryIndex], Result, bWait);
+		ViewState.Occlusion.LastOcclusionQueryArray[QueryIndex] = nullptr;
 	}
 }
 
