Add Cvar for visibility culling to cull to outer bounds of the sphere instead of to sphere origin.  This will be the true fix for PLAY-36400 which currently has a workaround.
(CVar is currently disabled so I can A/B test to double check perf)

diff --git a/Engine/Source/Runtime/Renderer/Private/SceneVisibility.cpp b/Engine/Source/Runtime/Renderer/Private/SceneVisibility.cpp
--- a/Engine/Source/Runtime/Renderer/Private/SceneVisibility.cpp
+++ b/Engine/Source/Runtime/Renderer/Private/SceneVisibility.cpp
@@ -56,6 +56,14 @@
 #include "DecalRenderingCommon.h"
 #include "CompositionLighting/PostProcessDeferredDecals.h"
 
+bool GDistanceCullToSphereEdge = false;
+static FAutoConsoleVariableRef CVarDistanceCullToSphereEdge(
+	TEXT("r.DistanceCullToSphereEdge"),
+	GDistanceCullToSphereEdge,
+	TEXT("If true frustum cull will check distance to bounding sphere edge rather than origin."),
+	ECVF_RenderThreadSafe
+);
+
 static float GWireframeCullThreshold = 5.0f;
 static FAutoConsoleVariableRef CVarWireframeCullThreshold(
 	TEXT("r.WireframeCullThreshold"),
@@ -605,6 +613,16 @@ inline bool ShouldCullForRayTracing(const FScene& Scene, FViewInfo& View, int32
 };
 #endif //RHI_RAYTRACING
 
+inline void ComputeDistances(const FPrimitiveBounds& Bounds, const FVector& ViewOriginForDistanceCulling, float& OutClosestDistSq, float& OutFurthestDistSq)
+{
+	float Dist = (Bounds.BoxSphereBounds.Origin - ViewOriginForDistanceCulling).Length();
+	float ClosestDist = Dist - Bounds.BoxSphereBounds.SphereRadius;
+	float FurthestDist = Dist + Bounds.BoxSphereBounds.SphereRadius;
+
+	OutClosestDistSq = Dist >= Bounds.BoxSphereBounds.SphereRadius ? FMath::Square(ClosestDist) : 0;
+	OutFurthestDistSq = FMath::Square(FurthestDist);
+}
+
 static void CullOctree(const FScene& Scene, FViewInfo& View, const FFrustumCullingFlags& Flags, FSceneBitArray& OutVisibleNodes)
 {
 	TRACE_CPUPROFILER_EVENT_SCOPE(SceneVisibility_CullOctree);
@@ -831,14 +849,23 @@ static int32 FrustumCull(const FScene& Scene, FViewInfo& View, FFrustumCullingFl
 					{
 						float MaxDrawDistance = Bounds.MaxCullDistance * MaxDrawDistanceScale;
 						float MinDrawDistanceSq = FMath::Square(Bounds.MinDrawDistance * MaxDrawDistanceScale);
-						float DistanceSquared = FVector::DistSquared(Bounds.BoxSphereBounds.Origin, ViewOriginForDistanceCulling);
+						float ClosestDistSquared, FurthestDistSquared;
+
+						if (GDistanceCullToSphereEdge)
+						{
+							ComputeDistances(Bounds, ViewOriginForDistanceCulling, ClosestDistSquared, FurthestDistSquared);
+						}
+						else
+						{
+							ClosestDistSquared = FurthestDistSquared = FVector::DistSquared(Bounds.BoxSphereBounds.Origin, ViewOriginForDistanceCulling);
+						}
 
 						// Always test the fade in distance.  If a primitive was set to always draw, it may need to be faded in.
 						if (bHasMaxDrawDistance)
 						{
 							float MaxFadeDistanceSquared = FMath::Square(MaxDrawDistance + FadeRadius);
 							float MinFadeDistanceSquared = FMath::Square(MaxDrawDistance - FadeRadius);
-							if ((DistanceSquared < MaxFadeDistanceSquared && DistanceSquared > MinFadeDistanceSquared)
+							if ((ClosestDistSquared < MaxFadeDistanceSquared && ClosestDistSquared > MinFadeDistanceSquared)
 								&& Scene.PrimitiveSceneProxies[Index]->IsUsingDistanceCullFade())  // Proxy call is intentionally behind the fade check to prevent an expensive memory read
 							{
 								FadingBits |= Mask;
@@ -846,8 +873,9 @@ static int32 FrustumCull(const FScene& Scene, FViewInfo& View, FFrustumCullingFl
 						}
 
 						// Check for distance culling first
-						const bool bFarDistanceCulled = bHasMaxDrawDistance && (DistanceSquared > FMath::Square(MaxDrawDistance));
-						const bool bNearDistanceCulled = bHasMinDrawDistance && (DistanceSquared < MinDrawDistanceSq);
+						const bool bFarDistanceCulled = bHasMaxDrawDistance && (ClosestDistSquared > FMath::Square(MaxDrawDistance));
+						const bool bNearDistanceCulled = bHasMinDrawDistance && (FurthestDistSquared < MinDrawDistanceSq);
+
 						bool bIsDistanceCulled = bNearDistanceCulled || bFarDistanceCulled;
 
 						if (bIsDistanceCulled)
@@ -5912,9 +5940,20 @@ void FLODSceneTree::UpdateVisibilityStates(FViewInfo& View, UE::Tasks::FTaskEven
 			const bool bForcedIntoView = FMath::IsNearlyZero(Bounds.MinDrawDistance);
 
 			// Update visibility states of this node and owned children
-			const float DistanceSquared = Bounds.BoxSphereBounds.ComputeSquaredDistanceFromBoxToPoint(View.ViewMatrices.GetViewOrigin());
-			const bool bNearCulled = DistanceSquared < FMath::Square(Bounds.MinDrawDistance) * HLODState.FOVDistanceScaleSq;
-			const bool bFarCulled = DistanceSquared > Bounds.MaxDrawDistance * Bounds.MaxDrawDistance * HLODState.FOVDistanceScaleSq;
+			float ClosestDistSquared, FurthestDistSquared;
+
+			if (GDistanceCullToSphereEdge)
+			{
+				ComputeDistances(Bounds, View.ViewMatrices.GetViewOrigin(), ClosestDistSquared, FurthestDistSquared);
+			}
+			else
+			{
+				ClosestDistSquared = FurthestDistSquared = Bounds.BoxSphereBounds.ComputeSquaredDistanceFromBoxToPoint(View.ViewMatrices.GetViewOrigin());
+			}
+
+			const bool bNearCulled = FurthestDistSquared  < FMath::Square(Bounds.MinDrawDistance) * HLODState.FOVDistanceScaleSq;
+			const bool bFarCulled = ClosestDistSquared > Bounds.MaxDrawDistance * Bounds.MaxDrawDistance * HLODState.FOVDistanceScaleSq;
+
 			const bool bIsInDrawRange = !bNearCulled && !bFarCulled;
 
 			const bool bWasFadingPreUpdate = !!NodeVisibility.bIsFading;
