Path Tracer: Give control over background alpha

This is intended for Twin Motion which would like the alpha channel to always be 1.0 for the entire image, with the ability to render in layers via holdouts to decompose the image. By using this cvar, they can control how the background contributes to alpha.

This cvar applies both to rays which miss the scene as well as infinite lights (distant and skylight) when they are made visible to camera rays.

diff --git a/Engine/Shaders/Private/PathTracing/PathTracingCore.ush b/Engine/Shaders/Private/PathTracing/PathTracingCore.ush
--- a/Engine/Shaders/Private/PathTracing/PathTracingCore.ush
+++ b/Engine/Shaders/Private/PathTracing/PathTracingCore.ush
@@ -38,6 +38,8 @@ uint MaxBounces;
 uint MaxSSSBounces;
 float SSSGuidingRatio;
 
+float BackgroundAlpha;
+
 // 0: only Material sampling
 // 1: only Light sampling
 // 2: both Material and Light
@@ -515,7 +517,7 @@ FPackedPathTracingPayload TraceTransparentRay(inout FPathState PathState, int Bo
 				LightHitPayload.ShadingModelID = SHADINGMODELID_UNLIT;
 				// light hits are fully transparent to match how the light loop works (the light surface does not cast a shadow)
 				LightHitPayload.BSDFOpacity = 1.0;
-				LightHitPayload.TransparencyColor = 1.0;
+				LightHitPayload.TransparencyColor = LightHit.HitT == POSITIVE_INFINITY ? 1.0 - BackgroundAlpha : 1.0;
 				LightHitPayload.SetFrontFace();
 				PackedPayload = PackPathTracingPayload(LightHitPayload);
 			}
@@ -668,6 +670,10 @@ FPackedPathTracingPayload TraceTransparentRay(inout FPathState PathState, int Bo
 		if (PackedPayload.IsMiss())
 		{
 			// Ray didn't hit any real geometry
+			if (bIsCameraRay)
+			{
+				PathState.Alpha += Luminance(PathState.PathThroughput) * BackgroundAlpha;
+			}
 			break;
 		}
 
diff --git a/Engine/Source/Runtime/Renderer/Private/PathTracing.cpp b/Engine/Source/Runtime/Renderer/Private/PathTracing.cpp
--- a/Engine/Source/Runtime/Renderer/Private/PathTracing.cpp
+++ b/Engine/Source/Runtime/Renderer/Private/PathTracing.cpp
@@ -499,6 +499,15 @@ TAutoConsoleVariable<int32> CVarPathTracingAdaptiveSamplingVisualize(
 	ECVF_RenderThreadSafe
 );
 
+TAutoConsoleVariable<float> CVarPathTracingBackgroundAlpha(
+	TEXT("r.PathTracing.BackgroundAlpha"),
+	0.0f,
+	TEXT("Value of the alpha channel for pixels that do hit anything (default 0.0)\n")
+	TEXT("Note that this refers to the normal interpretation of alpha which the path tracer uses internally, so 0 corresponds to a transparent pixel while 1 refers to a solid pixel."),
+	ECVF_RenderThreadSafe
+);
+
+
 BEGIN_SHADER_PARAMETER_STRUCT(FPathTracingData, )
 	SHADER_PARAMETER(float, BlendFactor)
 	SHADER_PARAMETER(uint32, Iteration)
@@ -506,6 +515,7 @@ BEGIN_SHADER_PARAMETER_STRUCT(FPathTracingData, )
 	SHADER_PARAMETER(uint32, MaxSamples)
 	SHADER_PARAMETER(uint32, MaxBounces)
 	SHADER_PARAMETER(uint32, MaxSSSBounces)
+	SHADER_PARAMETER(float, BackgroundAlpha)
 	SHADER_PARAMETER(float , SSSGuidingRatio)
 	SHADER_PARAMETER(uint32, MISMode)
 	SHADER_PARAMETER(uint32, VolumeMISMode)
@@ -559,6 +569,7 @@ struct FPathTracingConfig
 		return
 			PathTracingData.MaxSamples != Other.PathTracingData.MaxSamples ||
 			PathTracingData.MaxBounces != Other.PathTracingData.MaxBounces ||
+			PathTracingData.BackgroundAlpha != Other.PathTracingData.BackgroundAlpha ||
 			PathTracingData.MaxSSSBounces != Other.PathTracingData.MaxSSSBounces ||
 			PathTracingData.SSSGuidingRatio != Other.PathTracingData.SSSGuidingRatio ||
 			PathTracingData.MISMode != Other.PathTracingData.MISMode ||
@@ -715,6 +726,7 @@ static void PreparePathTracingData(const FScene* Scene, const FViewInfo& View, F
 	}
 
 	PathTracingData.MaxBounces = MaxBounces;
+	PathTracingData.BackgroundAlpha = FMath::Clamp(CVarPathTracingBackgroundAlpha.GetValueOnRenderThread(), 0.0f, 1.0f);
 	PathTracingData.MaxSSSBounces = ShowFlags.SubsurfaceScattering ? CVarPathTracingMaxSSSBounces.GetValueOnRenderThread() : 0;
 	PathTracingData.SSSGuidingRatio = FMath::Clamp(CVarPathTracingSSSGuidingRatio.GetValueOnRenderThread(), 0.0f, 1.0f);
 	PathTracingData.MaxNormalBias = GetRaytracingMaxNormalBias();
