Fix issue with displacement decal alpha blending.
Set render target output format for displacement RVT to include alpha.

diff --git a/Engine/Source/Runtime/Renderer/Private/VT/RuntimeVirtualTextureRender.cpp b/Engine/Source/Runtime/Renderer/Private/VT/RuntimeVirtualTextureRender.cpp
--- a/Engine/Source/Runtime/Renderer/Private/VT/RuntimeVirtualTextureRender.cpp
+++ b/Engine/Source/Runtime/Renderer/Private/VT/RuntimeVirtualTextureRender.cpp
@@ -377,7 +377,7 @@ namespace RuntimeVirtualTexture
 		static void ModifyCompilationEnvironment(FShaderCompilerEnvironment& OutEnvironment)
 		{
 			OutEnvironment.SetDefine(TEXT("OUT_DISPLACEMENT"), 1);
-			OutEnvironment.SetRenderTargetOutputFormat(0, PF_R32_FLOAT);
+			OutEnvironment.SetRenderTargetOutputFormat(0, PF_A32B32G32R32F);
 		}
 
 		static FRHIBlendState* GetBlendState(uint8 OutputAttributeMask)
