Fix decals instance culling data not being set to the right pass parameter when bRequiresMultiPass is enabled (and causes a GPU crash)

#rnx

diff --git a/Engine/Source/Runtime/Renderer/Private/MobileShadingRenderer.cpp b/Engine/Source/Runtime/Renderer/Private/MobileShadingRenderer.cpp
--- a/Engine/Source/Runtime/Renderer/Private/MobileShadingRenderer.cpp
+++ b/Engine/Source/Runtime/Renderer/Private/MobileShadingRenderer.cpp
@@ -1407,10 +1407,14 @@ void FMobileSceneRenderer::BuildInstanceCullingDrawParams(FRDGBuilder& GraphBuil
 		}
 		View.ParallelMeshDrawCommandPasses[EMeshPass::BasePass].BuildRenderingCommands(GraphBuilder, Scene->GPUScene, PassParameters->InstanceCullingDrawParams);
 		View.ParallelMeshDrawCommandPasses[EMeshPass::SkyPass].BuildRenderingCommands(GraphBuilder, Scene->GPUScene, SkyPassInstanceCullingDrawParams);
-		View.ParallelMeshDrawCommandPasses[EMeshPass::MeshDecal_SceneColor].BuildRenderingCommands(GraphBuilder, Scene->GPUScene, PassParameters->InstanceCullingDrawParams);
-		View.ParallelMeshDrawCommandPasses[EMeshPass::MeshDecal_SceneColorAndGBuffer].BuildRenderingCommands(GraphBuilder, Scene->GPUScene, PassParameters->InstanceCullingDrawParams);
 		View.ParallelMeshDrawCommandPasses[StandardTranslucencyMeshPass].BuildRenderingCommands(GraphBuilder, Scene->GPUScene, TranslucencyInstanceCullingDrawParams);
 		View.ParallelMeshDrawCommandPasses[EMeshPass::DebugViewMode].BuildRenderingCommands(GraphBuilder, Scene->GPUScene, DebugViewModeInstanceCullingDrawParams);
+	
+		if (!bRequiresMultiPass)
+		{
+			View.ParallelMeshDrawCommandPasses[EMeshPass::MeshDecal_SceneColor].BuildRenderingCommands(GraphBuilder, Scene->GPUScene, PassParameters->InstanceCullingDrawParams);
+			View.ParallelMeshDrawCommandPasses[EMeshPass::MeshDecal_SceneColorAndGBuffer].BuildRenderingCommands(GraphBuilder, Scene->GPUScene, PassParameters->InstanceCullingDrawParams);
+		}
 	}
 }
 
@@ -1682,6 +1686,9 @@ void FMobileSceneRenderer::RenderForwardMultiPass(FRDGBuilder& GraphBuilder, FMo
 	const bool bDoOcclusionQueries = (!bIsFullDepthPrepassEnabled && ViewContext.bIsLastView && DoOcclusionQueries());
 	SecondPassParameters->RenderTargets.NumOcclusionQueries = bDoOcclusionQueries ? ComputeNumOcclusionQueriesToBatch() : 0u;
 
+	View.ParallelMeshDrawCommandPasses[EMeshPass::MeshDecal_SceneColor].BuildRenderingCommands(GraphBuilder, Scene->GPUScene, SecondPassParameters->InstanceCullingDrawParams);
+	View.ParallelMeshDrawCommandPasses[EMeshPass::MeshDecal_SceneColorAndGBuffer].BuildRenderingCommands(GraphBuilder, Scene->GPUScene, SecondPassParameters->InstanceCullingDrawParams);
+
 	GraphBuilder.AddPass(
 		RDG_EVENT_NAME("DecalsAndTranslucency"),
 		SecondPassParameters,
