Add a VT CSV category (off by default in shipping), and a RenderedPages stat so we can track the number of page updates per frame

diff --git a/Engine/Source/Runtime/Renderer/Private/VT/RuntimeVirtualTextureRender.cpp b/Engine/Source/Runtime/Renderer/Private/VT/RuntimeVirtualTextureRender.cpp
--- a/Engine/Source/Runtime/Renderer/Private/VT/RuntimeVirtualTextureRender.cpp
+++ b/Engine/Source/Runtime/Renderer/Private/VT/RuntimeVirtualTextureRender.cpp
@@ -8,6 +8,7 @@
 #include "GPUScene.h"
 #include "MaterialShader.h"
 #include "MeshPassProcessor.h"
+#include "ProfilingDebugging/CsvProfiler.h"
 #include "RenderGraphBuilder.h"
 #include "RenderGraphUtils.h"
 #include "RenderUtils.h"
@@ -24,6 +25,8 @@
 #include "SceneRendering.h"
 #include "EngineModule.h"
 
+CSV_DEFINE_CATEGORY(VT, (!UE_BUILD_SHIPPING));
+
 namespace RuntimeVirtualTexture
 {
 	static TAutoConsoleVariable<bool> CVarVTMipColors(
@@ -1393,6 +1396,7 @@ namespace RuntimeVirtualTexture
 		FLinearColor const& FixedColor)
 	{
 		RDG_EVENT_SCOPE(GraphBuilder, "VirtualTextureDynamicCache");
+		CSV_CUSTOM_STAT(VT, RenderedPages, 1, ECsvCustomStatOp::Accumulate);
 
 		// Initialize a temporary view required for the material render pass
 		//todo[vt]: Some of this, such as ViewRotationMatrix, can be computed once in the Finalizer and passed down.
