Fixed incorrect mesh decal instancing setup in a mobile renderer

diff --git a/Engine/Source/Runtime/Renderer/Private/MobileDecalRendering.cpp b/Engine/Source/Runtime/Renderer/Private/MobileDecalRendering.cpp
--- a/Engine/Source/Runtime/Renderer/Private/MobileDecalRendering.cpp
+++ b/Engine/Source/Runtime/Renderer/Private/MobileDecalRendering.cpp
@@ -43,7 +43,7 @@ static bool DoesPlatformSupportDecals(EShaderPlatform ShaderPlatform)
 	return true;
 }
 
-void FMobileSceneRenderer::RenderDecals(FRHICommandList& RHICmdList, FViewInfo& View, const FInstanceCullingDrawParams* InstanceCullingDrawParams)
+void FMobileSceneRenderer::RenderDecals(FRHICommandList& RHICmdList, FViewInfo& View)
 {
 	if (!DoesPlatformSupportDecals(View.GetShaderPlatform()) || !ViewFamily.EngineShowFlags.Decals || View.bIsPlanarReflection)
 	{
@@ -67,6 +67,19 @@ void FMobileSceneRenderer::RenderDecals(FRHICommandList& RHICmdList, FViewInfo&
 	EMeshPass::Type DecalMeshPassType = DecalRendering::GetMeshPassType(RenderTargetMode);
 	if (View.ParallelMeshDrawCommandPasses[DecalMeshPassType].HasAnyDraw())
 	{
+		const FInstanceCullingDrawParams* InstanceCullingDrawParams = nullptr;
+		switch (DecalMeshPassType)
+		{
+		case EMeshPass::MeshDecal_SceneColor:
+			InstanceCullingDrawParams = &MeshDecalSceneColorInstanceCullingDrawParams;
+			break;
+		case EMeshPass::MeshDecal_SceneColorAndGBuffer:
+			InstanceCullingDrawParams = &MeshDecalSceneColorAndGBufferInstanceCullingDrawParams;
+			break;
+		default:
+			checkf(false, TEXT("Unexpected MeshDecal pass, please add corresponding InstanceCullingDrawParams!"))
+		};
+				
 		RHICmdList.SetViewport(View.ViewRect.Min.X, View.ViewRect.Min.Y, 0, View.ViewRect.Max.X, View.ViewRect.Max.Y, 1);
 		View.ParallelMeshDrawCommandPasses[DecalMeshPassType].DispatchDraw(nullptr, RHICmdList, InstanceCullingDrawParams);
 	}
diff --git a/Engine/Source/Runtime/Renderer/Private/MobileShadingRenderer.cpp b/Engine/Source/Runtime/Renderer/Private/MobileShadingRenderer.cpp
--- a/Engine/Source/Runtime/Renderer/Private/MobileShadingRenderer.cpp
+++ b/Engine/Source/Runtime/Renderer/Private/MobileShadingRenderer.cpp
@@ -1409,12 +1409,8 @@ void FMobileSceneRenderer::BuildInstanceCullingDrawParams(FRDGBuilder& GraphBuil
 		View.ParallelMeshDrawCommandPasses[EMeshPass::SkyPass].BuildRenderingCommands(GraphBuilder, Scene->GPUScene, SkyPassInstanceCullingDrawParams);
 		View.ParallelMeshDrawCommandPasses[StandardTranslucencyMeshPass].BuildRenderingCommands(GraphBuilder, Scene->GPUScene, TranslucencyInstanceCullingDrawParams);
 		View.ParallelMeshDrawCommandPasses[EMeshPass::DebugViewMode].BuildRenderingCommands(GraphBuilder, Scene->GPUScene, DebugViewModeInstanceCullingDrawParams);
-	
-		if (!bRequiresMultiPass)
-		{
-			View.ParallelMeshDrawCommandPasses[EMeshPass::MeshDecal_SceneColor].BuildRenderingCommands(GraphBuilder, Scene->GPUScene, PassParameters->InstanceCullingDrawParams);
-			View.ParallelMeshDrawCommandPasses[EMeshPass::MeshDecal_SceneColorAndGBuffer].BuildRenderingCommands(GraphBuilder, Scene->GPUScene, PassParameters->InstanceCullingDrawParams);
-		}
+		View.ParallelMeshDrawCommandPasses[EMeshPass::MeshDecal_SceneColor].BuildRenderingCommands(GraphBuilder, Scene->GPUScene, MeshDecalSceneColorInstanceCullingDrawParams);
+		View.ParallelMeshDrawCommandPasses[EMeshPass::MeshDecal_SceneColorAndGBuffer].BuildRenderingCommands(GraphBuilder, Scene->GPUScene, MeshDecalSceneColorAndGBufferInstanceCullingDrawParams);
 	}
 }
 
@@ -1584,7 +1580,7 @@ void FMobileSceneRenderer::RenderForwardSinglePass(FRDGBuilder& GraphBuilder, FM
 		PostRenderBasePass(RHICmdList, View);
 		// scene depth is read only and can be fetched
 		RHICmdList.NextSubpass();
-		RenderDecals(RHICmdList, View, &PassParameters->InstanceCullingDrawParams);
+		RenderDecals(RHICmdList, View);
 		RenderModulatedShadowProjections(RHICmdList, ViewContext.ViewIndex, View);
 		if (GMaxRHIShaderPlatform != SP_METAL_SIM)
 		{
@@ -1686,9 +1682,6 @@ void FMobileSceneRenderer::RenderForwardMultiPass(FRDGBuilder& GraphBuilder, FMo
 	const bool bDoOcclusionQueries = (!bIsFullDepthPrepassEnabled && ViewContext.bIsLastView && DoOcclusionQueries());
 	SecondPassParameters->RenderTargets.NumOcclusionQueries = bDoOcclusionQueries ? ComputeNumOcclusionQueriesToBatch() : 0u;
 
-	View.ParallelMeshDrawCommandPasses[EMeshPass::MeshDecal_SceneColor].BuildRenderingCommands(GraphBuilder, Scene->GPUScene, SecondPassParameters->InstanceCullingDrawParams);
-	View.ParallelMeshDrawCommandPasses[EMeshPass::MeshDecal_SceneColorAndGBuffer].BuildRenderingCommands(GraphBuilder, Scene->GPUScene, SecondPassParameters->InstanceCullingDrawParams);
-
 	GraphBuilder.AddPass(
 		RDG_EVENT_NAME("DecalsAndTranslucency"),
 		SecondPassParameters,
@@ -1698,7 +1691,7 @@ void FMobileSceneRenderer::RenderForwardMultiPass(FRDGBuilder& GraphBuilder, FMo
 		FViewInfo& View = *ViewContext.ViewInfo;
 			
 		// scene depth is read only and can be fetched
-		RenderDecals(RHICmdList, View, &SecondPassParameters->InstanceCullingDrawParams);
+		RenderDecals(RHICmdList, View);
 		RenderModulatedShadowProjections(RHICmdList, ViewContext.ViewIndex, View);
 		RenderFog(RHICmdList, View);
 		// Draw translucency.
@@ -1944,7 +1937,7 @@ void FMobileSceneRenderer::RenderDeferredSinglePass(FRDGBuilder& GraphBuilder, c
 		PostRenderBasePass(RHICmdList, View);
 		// SceneColor + GBuffer write, SceneDepth is read only
 		RHICmdList.NextSubpass();
-		RenderDecals(RHICmdList, View, &PassParameters->InstanceCullingDrawParams);
+		RenderDecals(RHICmdList, View);
 		// SceneColor write, SceneDepth is read only
 		RHICmdList.NextSubpass();
 		MobileDeferredShadingPass(RHICmdList, ViewContext.ViewIndex, Views.Num(), View, *Scene, SortedLightSet, VisibleLightInfos);
@@ -2006,7 +1999,7 @@ void FMobileSceneRenderer::RenderDeferredMultiPass(FRDGBuilder& GraphBuilder, cl
 		[this, SecondPassParameters, ViewContext](FRHICommandList& RHICmdList)
 	{
 		FViewInfo& View = *ViewContext.ViewInfo;
-		RenderDecals(RHICmdList, View, &SecondPassParameters->InstanceCullingDrawParams);
+		RenderDecals(RHICmdList, View);
 	});
 
 	auto* ThirdPassParameters = GraphBuilder.AllocParameters<FMobileRenderPassParameters>();
diff --git a/Engine/Source/Runtime/Renderer/Private/SceneRendering.h b/Engine/Source/Runtime/Renderer/Private/SceneRendering.h
--- a/Engine/Source/Runtime/Renderer/Private/SceneRendering.h
+++ b/Engine/Source/Runtime/Renderer/Private/SceneRendering.h
@@ -2729,7 +2729,7 @@ protected:
 
 
 	/** Renders decals. */
-	void RenderDecals(FRHICommandList& RHICmdList, FViewInfo& View, const FInstanceCullingDrawParams* InstanceCullingDrawParams);
+	void RenderDecals(FRHICommandList& RHICmdList, FViewInfo& View);
 	void RenderDBuffer(FRDGBuilder& GraphBuilder, FSceneTextures& SceneTextures, FDBufferTextures& DBufferTextures, FInstanceCullingManager& InstanceCullingManager);
 
 	/** Renders the atmospheric and height fog */
@@ -2799,6 +2799,8 @@ private:
 	FInstanceCullingDrawParams SkyPassInstanceCullingDrawParams;
 	FInstanceCullingDrawParams DebugViewModeInstanceCullingDrawParams;
 	FInstanceCullingDrawParams TranslucencyInstanceCullingDrawParams;
+	FInstanceCullingDrawParams MeshDecalSceneColorInstanceCullingDrawParams;
+	FInstanceCullingDrawParams MeshDecalSceneColorAndGBufferInstanceCullingDrawParams;
 
 	const FViewInfo* CachedView = nullptr;
 };
