HWRT: Allow skinned mesh geometry update and instance buffer build to run fully async.
* It's controlled by r.RayTracing.AsyncBuild same as other BLAS/TLAS building passes.
#tests CitySample on all platforms

diff --git a/Engine/Source/Runtime/Engine/Private/RayTracingSkinnedGeometry.cpp b/Engine/Source/Runtime/Engine/Private/RayTracingSkinnedGeometry.cpp
--- a/Engine/Source/Runtime/Engine/Private/RayTracingSkinnedGeometry.cpp
+++ b/Engine/Source/Runtime/Engine/Private/RayTracingSkinnedGeometry.cpp
@@ -19,13 +19,6 @@ static TAutoConsoleVariable<int32> CVarSkinCacheRayTracingMaxUpdatePrimitivesPer
 	ECVF_RenderThreadSafe
 );
 
-static TAutoConsoleVariable<int32> CVarSkinCacheRayTracingAsyncBuild(
-	TEXT("r.SkinCache.RayTracing.AsyncBuild"),
-	0,
-	TEXT("Whether to build skinned ray tracing acceleration structures on async compute queue (default = 0).\n"),
-	ECVF_RenderThreadSafe
-);
-
 static int32 GMaxRayTracingPrimitivesPerCmdList = -1;
 FAutoConsoleVariableRef CVarSkinnedGeometryMaxRayTracingPrimitivesPerCmdList(
 	TEXT("r.SkinCache.MaxRayTracingPrimitivesPerCmdList"),
@@ -200,7 +193,7 @@ BEGIN_SHADER_PARAMETER_STRUCT(FSkinnedGeometryBLASUpdateParams, )
 	RDG_BUFFER_ACCESS(SharedScratchBuffer, ERHIAccess::UAVCompute)
 END_SHADER_PARAMETER_STRUCT()
 
-void FRayTracingSkinnedGeometryUpdateQueue::Commit(FRDGBuilder& GraphBuilder)
+void FRayTracingSkinnedGeometryUpdateQueue::Commit(FRDGBuilder& GraphBuilder, ERDGPassFlags ComputePassFlags)
 {
 	TRACE_CPUPROFILER_EVENT_SCOPE(FRayTracingSkinnedGeometryUpdateQueue::Commit);
 
@@ -320,9 +313,6 @@ void FRayTracingSkinnedGeometryUpdateQueue::Commit(FRDGBuilder& GraphBuilder)
 
 	if (GeometryBuildRequests.Num())
 	{
-		const bool bRayTracingAsyncBuild = CVarSkinCacheRayTracingAsyncBuild.GetValueOnRenderThread() != 0 && GRHISupportsRayTracingAsyncBuildAccelerationStructure;
-		const ERDGPassFlags ComputePassFlags = bRayTracingAsyncBuild ? ERDGPassFlags::AsyncCompute : ERDGPassFlags::Compute;
-
 		GraphBuilder.AddPass(RDG_EVENT_NAME("CommitRayTracingSkinnedGeometryUpdates"), BLASUpdateParams, ComputePassFlags | ERDGPassFlags::NeverCull,
 			[
 				BuildRequests = MoveTemp(GeometryBuildRequests),
diff --git a/Engine/Source/Runtime/Engine/Public/RayTracingSkinnedGeometry.h b/Engine/Source/Runtime/Engine/Public/RayTracingSkinnedGeometry.h
--- a/Engine/Source/Runtime/Engine/Public/RayTracingSkinnedGeometry.h
+++ b/Engine/Source/Runtime/Engine/Public/RayTracingSkinnedGeometry.h
@@ -4,6 +4,7 @@
 
 #include "CoreMinimal.h"
 #include "RHIDefinitions.h"
+#include "RenderGraphDefinitions.h"
 
 #if RHI_RAYTRACING
 
@@ -35,7 +36,7 @@ public:
 	ENGINE_API void Commit(FRHICommandList& RHICmdList, FRHIBuffer* ScratchBuffer);
 
 	/** Commit all pending work using render graph. This allocates a transient scratch buffer internally. */
-	ENGINE_API void Commit(FRDGBuilder& GraphBuilder);
+	ENGINE_API void Commit(FRDGBuilder& GraphBuilder, ERDGPassFlags ComputePassFlags = ERDGPassFlags::Compute);
 
 private:
 	/** Info about pending updates. */
diff --git a/Engine/Source/Runtime/Renderer/Private/DeferredShadingRenderer.cpp b/Engine/Source/Runtime/Renderer/Private/DeferredShadingRenderer.cpp
--- a/Engine/Source/Runtime/Renderer/Private/DeferredShadingRenderer.cpp
+++ b/Engine/Source/Runtime/Renderer/Private/DeferredShadingRenderer.cpp
@@ -773,13 +773,16 @@ bool FDeferredShadingSceneRenderer::DispatchRayTracingWorldUpdates(FRDGBuilder&
 
 	TRACE_CPUPROFILER_EVENT_SCOPE(FDeferredShadingSceneRenderer::DispatchRayTracingWorldUpdates);
 
+	const bool bRayTracingAsyncBuild = CVarRayTracingAsyncBuild.GetValueOnRenderThread() != 0 && GRHISupportsRayTracingAsyncBuildAccelerationStructure;
+	const ERDGPassFlags ComputePassFlags = bRayTracingAsyncBuild ? ERDGPassFlags::AsyncCompute : ERDGPassFlags::Compute;
+
 	// Make sure there are no pending skin cache builds and updates anymore:
 	// FSkeletalMeshObjectGPUSkin::UpdateDynamicData_RenderThread could have enqueued build operations which might not have
 	// been processed by CommitRayTracingGeometryUpdates. 
 	// All pending builds should be done before adding them to the top level BVH.
 	if (FRayTracingSkinnedGeometryUpdateQueue* RayTracingSkinnedGeometryUpdateQueue = Scene->GetRayTracingSkinnedGeometryUpdateQueue())
 	{
-		RayTracingSkinnedGeometryUpdateQueue->Commit(GraphBuilder);
+		RayTracingSkinnedGeometryUpdateQueue->Commit(GraphBuilder, ComputePassFlags);
 	}
 
 	const int32 ReferenceViewIndex = 0;
@@ -815,7 +818,7 @@ bool FDeferredShadingSceneRenderer::DispatchRayTracingWorldUpdates(FRDGBuilder&
 	// Keep mask the same as what's already set (which will be the view mask) if TLAS updates should be masked to the view
 	RDG_GPU_MASK_SCOPE(GraphBuilder, GRayTracingMultiGpuTLASMask ? GraphBuilder.RHICmdList.GetGPUMask() : FRHIGPUMask::All());
 
-	RayTracingScene.CreateWithInitializationData(GraphBuilder, ReferenceView, &Scene->GPUScene, MoveTemp(ReferenceView.RayTracingSceneInitData));
+	RayTracingScene.CreateWithInitializationData(GraphBuilder, ReferenceView, &Scene->GPUScene, MoveTemp(ReferenceView.RayTracingSceneInitData), ComputePassFlags);
 
 	const uint32 BLASScratchSize = Scene->GetRayTracingDynamicGeometryCollection()->ComputeScratchBufferSize();
 	if (BLASScratchSize > 0)
@@ -829,9 +832,6 @@ bool FDeferredShadingSceneRenderer::DispatchRayTracingWorldUpdates(FRDGBuilder&
 		OutDynamicGeometryScratchBuffer = GraphBuilder.CreateBuffer(ScratchBufferDesc, TEXT("DynamicGeometry.BLASSharedScratchBuffer"));
 	}
 
-	const bool bRayTracingAsyncBuild = CVarRayTracingAsyncBuild.GetValueOnRenderThread() != 0 && GRHISupportsRayTracingAsyncBuildAccelerationStructure;
-	const ERDGPassFlags ComputePassFlags = bRayTracingAsyncBuild ? ERDGPassFlags::AsyncCompute : ERDGPassFlags::Compute;
-
 	{
 		RDG_GPU_STAT_SCOPE(GraphBuilder, RayTracingUpdate);
 
diff --git a/Engine/Source/Runtime/Renderer/Private/RayTracing/RayTracingScene.cpp b/Engine/Source/Runtime/Renderer/Private/RayTracing/RayTracingScene.cpp
--- a/Engine/Source/Runtime/Renderer/Private/RayTracing/RayTracingScene.cpp
+++ b/Engine/Source/Runtime/Renderer/Private/RayTracing/RayTracingScene.cpp
@@ -58,7 +58,7 @@ FRayTracingSceneWithGeometryInstances FRayTracingScene::BuildInitializationData(
 
 void FRayTracingScene::Create(FRDGBuilder& GraphBuilder, const FViewInfo& View, const FGPUScene* GPUScene)
 {
-	CreateWithInitializationData(GraphBuilder, View, GPUScene, BuildInitializationData());
+	CreateWithInitializationData(GraphBuilder, View, GPUScene, BuildInitializationData(), ERDGPassFlags::Compute);
 }
 
 void FRayTracingScene::InitPreViewTranslation(const FViewMatrices& ViewMatrices)
@@ -66,7 +66,7 @@ void FRayTracingScene::InitPreViewTranslation(const FViewMatrices& ViewMatrices)
 	PreViewTranslation = FDFVector3(ViewMatrices.GetPreViewTranslation());
 }
 
-void FRayTracingScene::CreateWithInitializationData(FRDGBuilder& GraphBuilder, const FViewInfo& View, const FGPUScene* GPUScene, FRayTracingSceneWithGeometryInstances SceneWithGeometryInstances)
+void FRayTracingScene::CreateWithInitializationData(FRDGBuilder& GraphBuilder, const FViewInfo& View, const FGPUScene* GPUScene, FRayTracingSceneWithGeometryInstances SceneWithGeometryInstances, ERDGPassFlags ComputePassFlags)
 {
 	QUICK_SCOPE_CYCLE_COUNTER(FRayTracingScene_BeginCreate);
 
@@ -237,7 +237,7 @@ void FRayTracingScene::CreateWithInitializationData(FRDGBuilder& GraphBuilder, c
 		GraphBuilder.AddPass(
 			RDG_EVENT_NAME("RayTracingBuildInstanceBuffer"),
 			PassParams,
-			ERDGPassFlags::Compute,
+			ComputePassFlags,
 			[PassParams,
 			this,
 			GPUScene,
diff --git a/Engine/Source/Runtime/Renderer/Private/RayTracing/RayTracingScene.h b/Engine/Source/Runtime/Renderer/Private/RayTracing/RayTracingScene.h
--- a/Engine/Source/Runtime/Renderer/Private/RayTracing/RayTracingScene.h
+++ b/Engine/Source/Runtime/Renderer/Private/RayTracing/RayTracingScene.h
@@ -61,7 +61,7 @@ public:
 	// Allocates GPU memory to fit at least the current number of instances.
 	// Kicks off instance buffer build to parallel thread along with RDG pass.
 	// NOTE: SceneWithGeometryInstances is passed in by value because ownership of the internal data is taken over. Use MoveTemp at call site, if possible.
-	void CreateWithInitializationData(FRDGBuilder& GraphBuilder, const FViewInfo& View, const FGPUScene* GPUScene, FRayTracingSceneWithGeometryInstances SceneWithGeometryInstances);
+	void CreateWithInitializationData(FRDGBuilder& GraphBuilder, const FViewInfo& View, const FGPUScene* GPUScene, FRayTracingSceneWithGeometryInstances SceneWithGeometryInstances, ERDGPassFlags ComputePassFlags);
 
 	// Backwards-compatible version of Create() which internally calls CreateRayTracingSceneWithGeometryInstances().
 	void Create(FRDGBuilder& GraphBuilder, const FViewInfo& View, const FGPUScene* GPUScene);
