PR #11823: Additional profiling information to LightRendering

diff --git a/Engine/Source/Runtime/Renderer/Private/LightRendering.cpp b/Engine/Source/Runtime/Renderer/Private/LightRendering.cpp
--- a/Engine/Source/Runtime/Renderer/Private/LightRendering.cpp
+++ b/Engine/Source/Runtime/Renderer/Private/LightRendering.cpp
@@ -2453,7 +2453,7 @@ static void InternalRenderLight(
 	const ELightComponentType LightType = (ELightComponentType)LightProxy->GetLightType();
 
 	GraphBuilder.AddPass(
-		RDG_EVENT_NAME("%s", ShaderName),
+		RDG_EVENT_NAME("%s: %s", ShaderName, *LightProxy->GetOwnerNameOrLabel()),
 		PassParameters,
 		ERDGPassFlags::Raster,
 		[Scene, &View, PixelShader, LightSceneInfo, PassParameters, LightBounds, LightType, SubstrateTileMaterialType](FRHICommandList& RHICmdList)
