Track number of active raytracing instances (passing GPU instance culling).
- since instance culling is done on GPU this requires using readback buffers.

diff --git a/Engine/Plugins/Experimental/GPULightmass/Source/GPULightmass/Private/LightmapRenderer.cpp b/Engine/Plugins/Experimental/GPULightmass/Source/GPULightmass/Private/LightmapRenderer.cpp
--- a/Engine/Plugins/Experimental/GPULightmass/Source/GPULightmass/Private/LightmapRenderer.cpp
+++ b/Engine/Plugins/Experimental/GPULightmass/Source/GPULightmass/Private/LightmapRenderer.cpp
@@ -1207,6 +1207,7 @@ bool FSceneRenderState::SetupRayTracingScene(FRDGBuilder& GraphBuilder, FSceneUn
 				SceneWithGeometryInstances.NumNativeGPUSceneInstances,
 				SceneWithGeometryInstances.NumNativeCPUInstances,
 				nullptr,
+				nullptr,
 				nullptr);
 
 			RHICmdList.BindAccelerationStructureMemory(RayTracingScene, RayTracingSceneBuffer, 0);
diff --git a/Engine/Shaders/Private/RayTracing/RayTracingInstanceBufferUtil.usf b/Engine/Shaders/Private/RayTracing/RayTracingInstanceBufferUtil.usf
--- a/Engine/Shaders/Private/RayTracing/RayTracingInstanceBufferUtil.usf
+++ b/Engine/Shaders/Private/RayTracing/RayTracingInstanceBufferUtil.usf
@@ -44,6 +44,10 @@ StructuredBuffer<float4> InstanceTransforms;
 // Extra per instance debug data
 RWStructuredBuffer<uint> RWDebugInstanceGPUSceneIndices;
 
+#if OUTPUT_STATS
+RWStructuredBuffer<uint> RWOutputStats;
+#endif
+
 #if GPU_CULLING
 
 float CullingRadius;
@@ -172,5 +176,16 @@ void RayTracingBuildInstanceBufferCS(uint3 DispatchThreadId : SV_DispatchThreadI
 			InputDesc.InstanceContributionToHitGroupIndex,
 			LocalToTranslatedWorld,
 			BlasAddress);
+
+#if OUTPUT_STATS
+		if (!bCulled)
+		{
+#if USE_WAVE_OPS
+			WaveInterlockedAddScalar(RWOutputStats[0], 1);
+#else
+			InterlockedAdd(RWOutputStats[0], 1);
+#endif
+		}
+#endif
 	}
 }
diff --git a/Engine/Source/Runtime/Renderer/Private/RayTracing/RayTracingInstanceBufferUtil.cpp b/Engine/Source/Runtime/Renderer/Private/RayTracing/RayTracingInstanceBufferUtil.cpp
--- a/Engine/Source/Runtime/Renderer/Private/RayTracing/RayTracingInstanceBufferUtil.cpp
+++ b/Engine/Source/Runtime/Renderer/Private/RayTracing/RayTracingInstanceBufferUtil.cpp
@@ -141,8 +141,6 @@ void FillRayTracingInstanceUploadBuffer(
 		LayerBaseIndices[LayerIndex] = LayerBaseIndices[LayerIndex - 1] + SceneInitializer.NumNativeInstancesPerLayer[LayerIndex - 1];
 	}
 
-	int32 NumInactiveNativeInstances = 0;
-
 	const int32 NumSceneInstances = Instances.Num();
 	const int32 MinBatchSize = 128;
 	ParallelFor(
@@ -160,8 +158,7 @@ void FillRayTracingInstanceUploadBuffer(
 			BaseInstancePrefixSum,
 			LayerBaseIndices,
 			PreViewTranslation,
-			&SceneInitializer,
-			&NumInactiveNativeInstances
+			&SceneInitializer
 		](int32 SceneInstanceIndex)
 		{
 			const FRayTracingGeometryInstance& SceneInstance = Instances[SceneInstanceIndex];
@@ -191,7 +188,6 @@ void FillRayTracingInstanceUploadBuffer(
 				BaseDescriptorIndex += NumNativeGPUSceneInstances;
 			}
 
-			int32 NumInactiveNativeInstancesThisSceneInstance = 0;
 			for (uint32 TransformIndex = 0; TransformIndex < NumTransforms; ++TransformIndex)
 			{
 				FRayTracingInstanceDescriptorInput InstanceDesc;
@@ -250,22 +246,7 @@ void FillRayTracingInstanceUploadBuffer(
 
 				OutInstanceUploadData[BaseDescriptorIndex + TransformIndex] = InstanceDesc;
 			}
-
-#if STATS
-			FPlatformAtomics::InterlockedAdd(&NumInactiveNativeInstances, NumInactiveNativeInstancesThisSceneInstance);
-#endif
 		});
-
-#if STATS
-	uint32 TotalNumNativeInstances = 0;
-	for (uint32 NumNativeInstances : SceneInitializer.NumNativeInstancesPerLayer)
-	{
-		TotalNumNativeInstances += NumNativeInstances;
-	}
-
-	SET_DWORD_STAT(STAT_RayTracingTotalInstances, TotalNumNativeInstances);
-	SET_DWORD_STAT(STAT_RayTracingActiveInstances, TotalNumNativeInstances - NumInactiveNativeInstances);
-#endif
 }
 
 struct FRayTracingBuildInstanceBufferCS : public FGlobalShader
@@ -300,6 +281,8 @@ struct FRayTracingBuildInstanceBufferCS : public FGlobalShader
 		SHADER_PARAMETER(FVector3f, ViewOrigin)
 		SHADER_PARAMETER(uint32, CullingMode)
 
+		SHADER_PARAMETER_UAV(RWStructuredBuffer<uint>, RWOutputStats)
+
 		// Debug parameters
 		SHADER_PARAMETER_UAV(RWStructuredBuffer<uint>, RWDebugInstanceGPUSceneIndices)
 	END_SHADER_PARAMETER_STRUCT()
@@ -307,7 +290,9 @@ struct FRayTracingBuildInstanceBufferCS : public FGlobalShader
 	class FUseGPUSceneDim : SHADER_PERMUTATION_BOOL("USE_GPUSCENE");
 	class FOutputInstanceGPUSceneIndexDim : SHADER_PERMUTATION_BOOL("OUTPUT_INSTANCE_GPUSCENE_INDEX");
 	class FGpuCullingDim : SHADER_PERMUTATION_BOOL("GPU_CULLING");
-	using FPermutationDomain = TShaderPermutationDomain<FUseGPUSceneDim, FOutputInstanceGPUSceneIndexDim, FGpuCullingDim>;
+	class FOutputStatsDim : SHADER_PERMUTATION_BOOL("OUTPUT_STATS");
+	class FUseWaveOpsDim : SHADER_PERMUTATION_BOOL("USE_WAVE_OPS");
+	using FPermutationDomain = TShaderPermutationDomain<FUseGPUSceneDim, FOutputInstanceGPUSceneIndexDim, FGpuCullingDim, FOutputStatsDim, FUseWaveOpsDim>;
 		
 	static constexpr uint32 ThreadGroupSize = 64;
 
@@ -325,6 +310,13 @@ struct FRayTracingBuildInstanceBufferCS : public FGlobalShader
 
 	static bool ShouldCompilePermutation(const FGlobalShaderPermutationParameters& Parameters)
 	{
+		FPermutationDomain PermutationVector(Parameters.PermutationId);
+
+		if (PermutationVector.Get<FUseWaveOpsDim>() && !RHISupportsWaveOperations(Parameters.Platform))
+		{
+			return false;
+		}
+
 		return IsRayTracingEnabledForProject(Parameters.Platform);
 	}
 };
@@ -342,6 +334,7 @@ void BuildRayTracingInstanceBuffer(
 	FShaderResourceViewRHIRef AccelerationStructureAddressesSRV,
 	FShaderResourceViewRHIRef InstanceTransformSRV,
 	const FRayTracingCullingParameters* CullingParameters,
+	FUnorderedAccessViewRHIRef OutputStatsUAV,
 	FUnorderedAccessViewRHIRef DebugInstanceGPUSceneIndexUAV)
 {
 	FRayTracingBuildInstanceBufferCS::FParameters PassParams;
@@ -372,11 +365,15 @@ void BuildRayTracingInstanceBuffer(
 		PassParams.CullingMode = uint32(CullingParameters->CullingMode);
 	}
 
+	PassParams.RWOutputStats = OutputStatsUAV;
+
 	PassParams.RWDebugInstanceGPUSceneIndices = DebugInstanceGPUSceneIndexUAV;
 
 	FRayTracingBuildInstanceBufferCS::FPermutationDomain PermutationVector;
 	PermutationVector.Set<FRayTracingBuildInstanceBufferCS::FUseGPUSceneDim>(InstanceTransformSRV == nullptr);
 	PermutationVector.Set<FRayTracingBuildInstanceBufferCS::FOutputInstanceGPUSceneIndexDim>(DebugInstanceGPUSceneIndexUAV != nullptr);
+	PermutationVector.Set<FRayTracingBuildInstanceBufferCS::FOutputStatsDim>(OutputStatsUAV != nullptr);
+	PermutationVector.Set<FRayTracingBuildInstanceBufferCS::FUseWaveOpsDim>(GRHISupportsWaveOperations);
 	PermutationVector.Set<FRayTracingBuildInstanceBufferCS::FGpuCullingDim>(CullingParameters != nullptr);
 
 	auto ComputeShader = GetGlobalShaderMap(GMaxRHIFeatureLevel)->GetShader<FRayTracingBuildInstanceBufferCS>(PermutationVector);
@@ -404,6 +401,7 @@ void BuildRayTracingInstanceBuffer(
 	uint32 NumNativeGPUSceneInstances,
 	uint32 NumNativeCPUInstances,
 	const FRayTracingCullingParameters* CullingParameters,
+	FUnorderedAccessViewRHIRef OutputStatsUAV,
 	FUnorderedAccessViewRHIRef DebugInstanceGPUSceneIndexUAV)
 {
 	if (NumNativeGPUSceneInstances > 0)
@@ -419,6 +417,7 @@ void BuildRayTracingInstanceBuffer(
 			AccelerationStructureAddressesSRV,
 			nullptr,
 			CullingParameters,
+			OutputStatsUAV,
 			DebugInstanceGPUSceneIndexUAV);
 	}
 
@@ -435,6 +434,7 @@ void BuildRayTracingInstanceBuffer(
 			AccelerationStructureAddressesSRV,
 			CPUInstanceTransformSRV,
 			nullptr,
+			OutputStatsUAV,
 			nullptr);
 	}
 }
diff --git a/Engine/Source/Runtime/Renderer/Private/RayTracing/RayTracingScene.cpp b/Engine/Source/Runtime/Renderer/Private/RayTracing/RayTracingScene.cpp
--- a/Engine/Source/Runtime/Renderer/Private/RayTracing/RayTracingScene.cpp
+++ b/Engine/Source/Runtime/Renderer/Private/RayTracing/RayTracingScene.cpp
@@ -28,6 +28,7 @@ static TAutoConsoleVariable<int32> CVarRayTracingSceneBuildMode(
 
 BEGIN_SHADER_PARAMETER_STRUCT(FBuildInstanceBufferPassParams, )
 	SHADER_PARAMETER_RDG_BUFFER_UAV(RWStructuredBuffer, InstanceBuffer)
+	SHADER_PARAMETER_RDG_BUFFER_UAV(RWStructuredBuffer, OutputStats)
 	SHADER_PARAMETER_RDG_BUFFER_UAV(RWStructuredBuffer, DebugInstanceGPUSceneIndexBuffer)
 	SHADER_PARAMETER_RDG_UNIFORM_BUFFER(FSceneUniformParameters, Scene)
 END_SHADER_PARAMETER_STRUCT()
@@ -174,6 +175,12 @@ void FRayTracingScene::CreateWithInitializationData(FRDGBuilder& GraphBuilder, c
 		}
 	}
 
+#if STATS
+	FRDGBufferRef OutputStatsBuffer = GraphBuilder.CreateBuffer(FRDGBufferDesc::CreateStructuredDesc(sizeof(uint32), 1), TEXT("FRayTracingScene::OutputStatsBuffer"));
+	FRDGBufferUAVRef OutputStatsBufferUAV = GraphBuilder.CreateUAV(OutputStatsBuffer);
+	AddClearUAVPass(GraphBuilder, OutputStatsBufferUAV, 0);
+#endif
+
 	FRDGBufferUAVRef DebugInstanceGPUSceneIndexBufferUAV = nullptr;
 	if (bNeedsDebugInstanceGPUSceneIndexBuffer)
 	{
@@ -234,6 +241,10 @@ void FRayTracingScene::CreateWithInitializationData(FRDGBuilder& GraphBuilder, c
 		PassParams->DebugInstanceGPUSceneIndexBuffer = DebugInstanceGPUSceneIndexBufferUAV;
 		PassParams->Scene = View.GetSceneUniforms().GetBuffer(GraphBuilder);
 
+#if STATS
+		PassParams->OutputStats = OutputStatsBufferUAV;
+#endif
+
 		GraphBuilder.AddPass(
 			RDG_EVENT_NAME("RayTracingBuildInstanceBuffer"),
 			PassParams,
@@ -287,9 +298,62 @@ void FRayTracingScene::CreateWithInitializationData(FRDGBuilder& GraphBuilder, c
 					NumNativeGPUSceneInstances,
 					NumNativeCPUInstances,
 					CullingParameters.bUseInstanceCulling ? &CullingParameters : nullptr,
+					PassParams->OutputStats ? PassParams->OutputStats->GetRHI() : nullptr,
 					PassParams->DebugInstanceGPUSceneIndexBuffer ? PassParams->DebugInstanceGPUSceneIndexBuffer->GetRHI() : nullptr);
 			});
 	}
+
+#if STATS
+	// readback
+	{
+		//  if necessary create readback buffers
+		if (StatsReadbackBuffers.IsEmpty())
+		{
+			StatsReadbackBuffers.SetNum(MaxReadbackBuffers);
+
+			for (uint32 Index = 0; Index < MaxReadbackBuffers; ++Index)
+			{
+				StatsReadbackBuffers[Index] = new FRHIGPUBufferReadback(TEXT("FRayTracingScene::StatsReadbackBuffer"));
+			}
+		}
+		
+		// copy stats to readback buffer
+		{
+			AddReadbackBufferPass(GraphBuilder, RDG_EVENT_NAME("FRayTracingScene::StatsReadback"), OutputStatsBuffer,
+				[ReadbackBuffer = StatsReadbackBuffers[StatsReadbackBuffersWriteIndex], OutputStatsBuffer](FRHICommandList& RHICmdList)
+				{
+					ReadbackBuffer->EnqueueCopy(RHICmdList, OutputStatsBuffer->GetRHI(), 0u);
+				});
+
+			StatsReadbackBuffersWriteIndex = (StatsReadbackBuffersWriteIndex + 1u) % MaxReadbackBuffers;
+			StatsReadbackBuffersNumPending = FMath::Min(StatsReadbackBuffersNumPending + 1u, MaxReadbackBuffers);
+		}
+
+		// process ready results
+		while (StatsReadbackBuffersNumPending > 0)
+		{
+			uint32 Index = (StatsReadbackBuffersWriteIndex + MaxReadbackBuffers - StatsReadbackBuffersNumPending) % MaxReadbackBuffers;
+			FRHIGPUBufferReadback* ReadbackBuffer = StatsReadbackBuffers[Index];
+			if (ReadbackBuffer->IsReady())
+			{
+				StatsReadbackBuffersNumPending--;
+
+				auto ReadbackBufferPtr = (const uint32*)ReadbackBuffer->Lock(sizeof(uint32));
+
+				NumActiveInstances = ReadbackBufferPtr[0];
+
+				ReadbackBuffer->Unlock();
+			}
+			else
+			{
+				break;
+			}
+		}
+
+		SET_DWORD_STAT(STAT_RayTracingTotalInstances, NumNativeInstances);
+		SET_DWORD_STAT(STAT_RayTracingActiveInstances, FMath::Min(NumActiveInstances, NumNativeInstances));
+	}
+#endif
 }
 
 BEGIN_SHADER_PARAMETER_STRUCT(FRayTracingSceneBuildPassParams, )
@@ -458,6 +522,20 @@ void FRayTracingScene::EndFrame()
 
 		RayTracingSceneBuffer = nullptr;
 		RayTracingScenePooledBuffer = nullptr;
+
+#if STATS
+		for (auto& ReadbackBuffer : StatsReadbackBuffers)
+		{
+			delete ReadbackBuffer;
+		}
+
+		StatsReadbackBuffers.Empty();
+
+		StatsReadbackBuffersWriteIndex = 0;
+		StatsReadbackBuffersNumPending = 0;
+
+		NumActiveInstances = 0;
+#endif
 	}
 
 	bUsedThisFrame = false;
diff --git a/Engine/Source/Runtime/Renderer/Private/RayTracing/RayTracingScene.h b/Engine/Source/Runtime/Renderer/Private/RayTracing/RayTracingScene.h
--- a/Engine/Source/Runtime/Renderer/Private/RayTracing/RayTracingScene.h
+++ b/Engine/Source/Runtime/Renderer/Private/RayTracing/RayTracingScene.h
@@ -11,6 +11,7 @@
 #include "Math/DoubleFloat.h"
 #include "RHI.h"
 #include "RHIUtilities.h"
+#include "RHIGPUReadback.h"
 #include "RenderGraphResources.h"
 #include "Misc/MemStack.h"
 #include "Containers/ArrayView.h"
@@ -171,6 +172,16 @@ private:
 	bool bInstanceDebugDataEnabled = false;
 
 	bool bUsedThisFrame = false;
+
+#if STATS
+	const uint32 MaxReadbackBuffers = 4;
+
+	TArray<FRHIGPUBufferReadback*> StatsReadbackBuffers;
+	uint32 StatsReadbackBuffersWriteIndex = 0;
+	uint32 StatsReadbackBuffersNumPending = 0;
+
+	uint32 NumActiveInstances = 0;
+#endif
 };
 
 #endif // RHI_RAYTRACING
diff --git a/Engine/Source/Runtime/Renderer/Private/Tests/RayTracingTestbed.cpp b/Engine/Source/Runtime/Renderer/Private/Tests/RayTracingTestbed.cpp
--- a/Engine/Source/Runtime/Renderer/Private/Tests/RayTracingTestbed.cpp
+++ b/Engine/Source/Runtime/Renderer/Private/Tests/RayTracingTestbed.cpp
@@ -275,6 +275,7 @@ bool RunRayTracingTestbed_RenderThread(const FString& Parameters)
 		RayTracingScene.NumNativeGPUSceneInstances,
 		RayTracingScene.NumNativeCPUInstances,
 		nullptr,
+		nullptr,
 		nullptr);
 
 	RHICmdList.BindAccelerationStructureMemory(RayTracingScene.Scene, SceneBuffer, 0);
diff --git a/Engine/Source/Runtime/Renderer/Public/RayTracingInstanceBufferUtil.h b/Engine/Source/Runtime/Renderer/Public/RayTracingInstanceBufferUtil.h
--- a/Engine/Source/Runtime/Renderer/Public/RayTracingInstanceBufferUtil.h
+++ b/Engine/Source/Runtime/Renderer/Public/RayTracingInstanceBufferUtil.h
@@ -127,6 +127,7 @@ RENDERER_API void BuildRayTracingInstanceBuffer(
 	uint32 NumNativeGPUSceneInstances,
 	uint32 NumNativeCPUInstances,
 	const FRayTracingCullingParameters* CullingParameters,
+	FUnorderedAccessViewRHIRef OutputStatsUAV,
 	FUnorderedAccessViewRHIRef DebugInstanceGPUSceneIndexUAV);
 
 UE_DEPRECATED(5.5, "GPUInstances no longer supported. Use GPUSceneInstances instead.")
@@ -148,6 +149,7 @@ inline void BuildRayTracingInstanceBuffer(
 		InstancesUAV, InstanceUploadSRV, AccelerationStructureAddressesSRV, CPUInstanceTransformSRV,
 		NumNativeGPUSceneInstances, NumNativeCPUInstances,
 		CullingParameters,
+		nullptr,
 		DebugInstanceGPUSceneIndexUAV);
 }
 
