Changes to GPU stats related to HWRT:
- Moved dynamic ray tracing geometry updates out of "RayTracingUpdate" since those are already tracked by "RayTracingDynamicGeometry"
- Renamed "RayTracingUpdate" to "RayTracingScene" which tracks FRayTracingBuildInstanceBufferCS and TLAS builds

diff --git a/Engine/Source/Runtime/Renderer/Private/DeferredShadingRenderer.cpp b/Engine/Source/Runtime/Renderer/Private/DeferredShadingRenderer.cpp
--- a/Engine/Source/Runtime/Renderer/Private/DeferredShadingRenderer.cpp
+++ b/Engine/Source/Runtime/Renderer/Private/DeferredShadingRenderer.cpp
@@ -282,7 +282,7 @@ DECLARE_CYCLE_STAT(TEXT("DeferredShadingSceneRenderer RenderLocalFogVolume"), ST
 DECLARE_CYCLE_STAT(TEXT("DeferredShadingSceneRenderer RenderLightShaftBloom"), STAT_FDeferredShadingSceneRenderer_RenderLightShaftBloom, STATGROUP_SceneRendering);
 DECLARE_CYCLE_STAT(TEXT("DeferredShadingSceneRenderer RenderFinish"), STAT_FDeferredShadingSceneRenderer_RenderFinish, STATGROUP_SceneRendering);
 
-DECLARE_GPU_STAT(RayTracingUpdate);
+DECLARE_GPU_STAT(RayTracingScene);
 DECLARE_GPU_STAT(RayTracingGeometry);
 DECLARE_GPU_STAT(RayTracingDynamicGeometry);
 
@@ -818,8 +818,6 @@ bool FDeferredShadingSceneRenderer::DispatchRayTracingWorldUpdates(FRDGBuilder&
 	// Keep mask the same as what's already set (which will be the view mask) if TLAS updates should be masked to the view
 	RDG_GPU_MASK_SCOPE(GraphBuilder, GRayTracingMultiGpuTLASMask ? GraphBuilder.RHICmdList.GetGPUMask() : FRHIGPUMask::All());
 
-	RayTracingScene.CreateWithInitializationData(GraphBuilder, ReferenceView, &Scene->GPUScene, MoveTemp(ReferenceView.RayTracingSceneInitData), ComputePassFlags);
-
 	const uint32 BLASScratchSize = Scene->GetRayTracingDynamicGeometryCollection()->ComputeScratchBufferSize();
 	if (BLASScratchSize > 0)
 	{
@@ -833,32 +831,33 @@ bool FDeferredShadingSceneRenderer::DispatchRayTracingWorldUpdates(FRDGBuilder&
 	}
 
 	{
-		RDG_GPU_STAT_SCOPE(GraphBuilder, RayTracingUpdate);
-
-		{
-			// Dynamic geometry (BLAS) updates must always run on all GPUs.  Other passes may either run on all GPUs or be scoped to the view's GPUs.
-			// See GRayTracingMultiGpuTLASMask.
-			RDG_GPU_MASK_SCOPE(GraphBuilder, FRHIGPUMask::All());
+		// Dynamic geometry (BLAS) updates must always run on all GPUs.  Other passes may either run on all GPUs or be scoped to the view's GPUs.
+		// See GRayTracingMultiGpuTLASMask.
+		RDG_GPU_MASK_SCOPE(GraphBuilder, FRHIGPUMask::All());
 
-			FBuildAccelerationStructurePassParams* PassParams = GraphBuilder.AllocParameters<FBuildAccelerationStructurePassParams>();
-			PassParams->View = ReferenceView.ViewUniformBuffer;
-			PassParams->Scene = GetSceneUniforms().GetBuffer(GraphBuilder);
-			PassParams->DynamicGeometryScratchBuffer = OutDynamicGeometryScratchBuffer;
-			PassParams->LightGridPacked = nullptr;
-			PassParams->ClusterPageData = nullptr;
-			PassParams->HierarchyBuffer = nullptr;
-			PassParams->RayTracingDataBuffer = nullptr;
+		FBuildAccelerationStructurePassParams* PassParams = GraphBuilder.AllocParameters<FBuildAccelerationStructurePassParams>();
+		PassParams->View = ReferenceView.ViewUniformBuffer;
+		PassParams->Scene = GetSceneUniforms().GetBuffer(GraphBuilder);
+		PassParams->DynamicGeometryScratchBuffer = OutDynamicGeometryScratchBuffer;
+		PassParams->LightGridPacked = nullptr;
+		PassParams->ClusterPageData = nullptr;
+		PassParams->HierarchyBuffer = nullptr;
+		PassParams->RayTracingDataBuffer = nullptr;
 
-			GraphBuilder.AddPass(RDG_EVENT_NAME("RayTracingDynamicUpdate"), PassParams, ComputePassFlags | ERDGPassFlags::NeverCull,
-				[this, PassParams, bRayTracingAsyncBuild](FRHICommandList& RHICmdList)
+		GraphBuilder.AddPass(RDG_EVENT_NAME("RayTracingDynamicUpdate"), PassParams, ComputePassFlags | ERDGPassFlags::NeverCull,
+			[this, PassParams, bRayTracingAsyncBuild](FRHICommandList& RHICmdList)
 			{
 				SCOPED_GPU_STAT(RHICmdList, RayTracingDynamicGeometry);
 				FRHIBuffer* DynamicGeometryScratchBuffer = PassParams->DynamicGeometryScratchBuffer ? PassParams->DynamicGeometryScratchBuffer->GetRHI() : nullptr;
 				Scene->GetRayTracingDynamicGeometryCollection()->DispatchUpdates(RHICmdList, DynamicGeometryScratchBuffer);
 				Scene->GetRayTracingDynamicGeometryCollection()->EndUpdate();
 			});
-		}
+	}
 
+
+	{
+		RDG_GPU_STAT_SCOPE(GraphBuilder, RayTracingScene);
+		RayTracingScene.CreateWithInitializationData(GraphBuilder, ReferenceView, &Scene->GPUScene, MoveTemp(ReferenceView.RayTracingSceneInitData), ComputePassFlags);
 		RayTracingScene.Build(GraphBuilder, ComputePassFlags | ERDGPassFlags::NeverCull, OutDynamicGeometryScratchBuffer);
 	}
 
diff --git a/Engine/Source/Runtime/Renderer/Private/RayTracing/RayTracingScene.cpp b/Engine/Source/Runtime/Renderer/Private/RayTracing/RayTracingScene.cpp
--- a/Engine/Source/Runtime/Renderer/Private/RayTracing/RayTracingScene.cpp
+++ b/Engine/Source/Runtime/Renderer/Private/RayTracing/RayTracingScene.cpp
@@ -15,8 +15,6 @@
 #include "SceneRendering.h"
 #include "RayTracingInstanceCulling.h"
 
-DECLARE_GPU_STAT(RayTracingScene);
-
 static TAutoConsoleVariable<int32> CVarRayTracingSceneBuildMode(
 	TEXT("r.RayTracing.Scene.BuildMode"),
 	1,
@@ -375,8 +373,6 @@ void FRayTracingScene::Build(FRDGBuilder& GraphBuilder, ERDGPassFlags ComputePas
 	GraphBuilder.AddPass(RDG_EVENT_NAME("RayTracingBuildScene"), PassParams, ComputePassFlags,
 		[this, PassParams](FRHICommandList& RHICmdList)
 		{
-			SCOPED_GPU_STAT(RHICmdList, RayTracingScene);
-
 			FRayTracingSceneBuildParams BuildParams;
 			BuildParams.Scene = RayTracingSceneRHI;
 			BuildParams.ScratchBuffer = PassParams->ScratchBuffer->GetRHI();
