RHI: Atomic flag should be set for all RHI when in SM6.

#rnx

diff --git a/Engine/Source/Runtime/Renderer/Private/VirtualShadowMaps/VirtualShadowMapCacheManager.cpp b/Engine/Source/Runtime/Renderer/Private/VirtualShadowMaps/VirtualShadowMapCacheManager.cpp
--- a/Engine/Source/Runtime/Renderer/Private/VirtualShadowMaps/VirtualShadowMapCacheManager.cpp
+++ b/Engine/Source/Runtime/Renderer/Private/VirtualShadowMaps/VirtualShadowMapCacheManager.cpp
@@ -679,14 +679,7 @@ void FVirtualShadowMapArrayCacheManager::SetPhysicalPoolSize(FRDGBuilder& GraphB
 		// Track changes to these ourselves instead of from the GetDesc() since that may get manipulated internally
 		PhysicalPagePoolCreateFlags = RequestedCreateFlags;
         
-        ETextureCreateFlags PoolTexCreateFlags = TexCreate_ShaderResource | TexCreate_UAV;
-        
-#if PLATFORM_MAC
-        if(GMaxRHIFeatureLevel >= ERHIFeatureLevel::SM6)
-        {
-            PoolTexCreateFlags |= TexCreate_AtomicCompatible;
-        }
-#endif
+        const ETextureCreateFlags PoolTexCreateFlags = TexCreate_ShaderResource | TexCreate_UAV | TexCreate_AtomicCompatible;
         
 		FPooledRenderTargetDesc Desc2D = FPooledRenderTargetDesc::Create2DArrayDesc(
 			RequestedSize,
