Fix over budget RVT page updates per frame.
The dirty pages in MappedTilesToProduce were being updated twice in a frame, once at BeginUpdate() and once at EndUpdate() and could use the full budget twice.
Changed to tracking remaining budget in the Updater so that BeginUpdate() affects budget for EndUpdate().
MappedTilesToProduce are only handled in BeginUpdate() and Continuous updates are only handled in EndUpdate().
Also the load request sort had a bug where requests were not being correctly sorted by stream/non-stream status. This can lead, for example, to 2 non-streamed pages being prioritized for update instead of 1 streamed and 1 non-streamed.

diff --git a/Engine/Source/Runtime/Renderer/Private/VT/RuntimeVirtualTextureRender.cpp b/Engine/Source/Runtime/Renderer/Private/VT/RuntimeVirtualTextureRender.cpp
--- a/Engine/Source/Runtime/Renderer/Private/VT/RuntimeVirtualTextureRender.cpp
+++ b/Engine/Source/Runtime/Renderer/Private/VT/RuntimeVirtualTextureRender.cpp
@@ -25,7 +25,9 @@
 #include "SceneRendering.h"
 #include "EngineModule.h"
 
-CSV_DEFINE_CATEGORY(VT, (!UE_BUILD_SHIPPING));
+CSV_DECLARE_CATEGORY_EXTERN(VirtualTexturing);
+
+DECLARE_DWORD_COUNTER_STAT(TEXT("Num pages rendered"), STAT_RenderedPages, STATGROUP_VirtualTexturing);
 
 namespace RuntimeVirtualTexture
 {
@@ -1396,7 +1398,8 @@ namespace RuntimeVirtualTexture
 		FLinearColor const& FixedColor)
 	{
 		RDG_EVENT_SCOPE(GraphBuilder, "VirtualTextureDynamicCache");
-		CSV_CUSTOM_STAT(VT, RenderedPages, 1, ECsvCustomStatOp::Accumulate);
+		CSV_CUSTOM_STAT(VirtualTexturing, RenderedPages, 1, ECsvCustomStatOp::Accumulate);
+		INC_DWORD_STAT_BY(STAT_RenderedPages, 1);
 
 		// Initialize a temporary view required for the material render pass
 		//todo[vt]: Some of this, such as ViewRotationMatrix, can be computed once in the Finalizer and passed down.
diff --git a/Engine/Source/Runtime/Renderer/Private/VT/UniqueRequestList.h b/Engine/Source/Runtime/Renderer/Private/VT/UniqueRequestList.h
--- a/Engine/Source/Runtime/Renderer/Private/VT/UniqueRequestList.h
+++ b/Engine/Source/Runtime/Renderer/Private/VT/UniqueRequestList.h
@@ -377,11 +377,11 @@ inline void FUniqueRequestList::SortRequests(FVirtualTextureProducerCollection&
 {
 	struct FPriorityAndIndex
 	{
-		uint32 Priroity;
+		uint32 Priority;
 		uint16 Index;
 
 		// sort from largest to smallest
-		inline bool operator<(const FPriorityAndIndex& Rhs) const { return Priroity > Rhs.Priroity; }
+		inline bool operator<(const FPriorityAndIndex& Rhs) const { return Priority > Rhs.Priority; }
 	};
 
 	if (bUseCombinedLimit)
@@ -409,17 +409,17 @@ inline void FUniqueRequestList::SortRequests(FVirtualTextureProducerCollection&
 		if (Flags.bLocked)
 		{
 			// Lock request, use max priority
-			SortedKeys[i].Priroity = MAX_uint32;
+			SortedKeys[i].Priority = MAX_uint32;
 			NumNonStreamingLockRequests += (bUseCombinedLimit || !Flags.bStreaming ? 1 : 0);
 			++CheckNumLockRequests;
 		}
 		else if (bUseCombinedLimit || !Flags.bStreaming)
 		{
-			SortedKeys[i].Priroity = 0xd | Priority;
+			SortedKeys[i].Priority = 0xd0000000 | Priority;
 		}
 		else
 		{
-			SortedKeys[i].Priroity = 0xe | Priority;
+			SortedKeys[i].Priority = 0xe0000000 | Priority;
 			++NumStreamingNonLockRequests;
 		}
 	}
diff --git a/Engine/Source/Runtime/Renderer/Private/VT/VirtualTextureSystem.cpp b/Engine/Source/Runtime/Renderer/Private/VT/VirtualTextureSystem.cpp
--- a/Engine/Source/Runtime/Renderer/Private/VT/VirtualTextureSystem.cpp
+++ b/Engine/Source/Runtime/Renderer/Private/VT/VirtualTextureSystem.cpp
@@ -29,7 +29,7 @@
 
 #define LOCTEXT_NAMESPACE "VirtualTexture"
 
-CSV_DEFINE_CATEGORY(VirtualTexturing, true);
+CSV_DEFINE_CATEGORY(VirtualTexturing, (!UE_BUILD_SHIPPING));
 
 DECLARE_CYCLE_STAT(TEXT("VirtualTextureSystem Update"), STAT_VirtualTextureSystem_Update, STATGROUP_VirtualTexturing);
 
@@ -1964,7 +1964,7 @@ void FVirtualTextureSystem::GrowPhysicalPools() const
 	}
 }
 
-void FVirtualTextureSystem::SubmitRequestsFromLocalTileList(FRHICommandList& RHICmdList, TArray<FVirtualTextureLocalTile>& OutDeferredTiles, const TSet<FVirtualTextureLocalTile>& LocalTileList, EVTProducePageFlags Flags, ERHIFeatureLevel::Type FeatureLevel, uint32 MaxRequestsToProduce)
+int32 FVirtualTextureSystem::SubmitRequestsFromLocalTileList(FRHICommandList& RHICmdList, TArray<FVirtualTextureLocalTile>& OutDeferredTiles, const TSet<FVirtualTextureLocalTile>& LocalTileList, EVTProducePageFlags Flags, ERHIFeatureLevel::Type FeatureLevel, uint32 MaxRequestsToProduce)
 {
 	LLM_SCOPE(ELLMTag::VirtualTextureSystem);
 
@@ -2045,29 +2045,11 @@ void FVirtualTextureSystem::SubmitRequestsFromLocalTileList(FRHICommandList& RHI
 
 		NumPagesProduced++;
 	}
-}
-
-void FVirtualTextureSystem::SubmitPreMappedRequests(FRHICommandList& RHICmdList, ERHIFeatureLevel::Type FeatureLevel, uint32 MaxMappedTilesToProduce)
-{
-	check(TransientCollectedPages.Num() == 0);
-
-	{
-		INC_DWORD_STAT_BY(STAT_NumMappedPageUpdate, MappedTilesToProduce.Num());
-		SubmitRequestsFromLocalTileList(RHICmdList, TransientCollectedPages, MappedTilesToProduce, EVTProducePageFlags::None, FeatureLevel, MaxMappedTilesToProduce);
-		MappedTilesToProduce.Reset();
-		MappedTilesToProduce.Append(TransientCollectedPages);
-		TransientCollectedPages.Reset();
-	}
 
-	{
-		INC_DWORD_STAT_BY(STAT_NumContinuousPageUpdate, ContinuousUpdateTilesToProduce.Num());
-		SubmitRequestsFromLocalTileList(RHICmdList, TransientCollectedPages, ContinuousUpdateTilesToProduce, EVTProducePageFlags::ContinuousUpdate, FeatureLevel, 0);
-		ContinuousUpdateTilesToProduce.Reset();
-		TransientCollectedPages.Reset();
-	}
+	return NumPagesProduced;
 }
 
-void FVirtualTextureSystem::SubmitThrottledRequests(FRHICommandList& RHICmdList, FVirtualTextureUpdater* Updater, bool bContinousUpdates)
+void FVirtualTextureSystem::SubmitThrottledRequests(FRHICommandList& RHICmdList, FVirtualTextureUpdater* Updater, EUpdatePhase UpdatePhase)
 {
 	const FVirtualTextureUpdateSettings& Settings   = Updater->Settings;
 	const ERHIFeatureLevel::Type  FeatureLevel      = Updater->FeatureLevel;
@@ -2080,47 +2062,54 @@ void FVirtualTextureSystem::SubmitThrottledRequests(FRHICommandList& RHICmdList,
 		FAdaptiveVirtualTexture::QueuePackedAllocationRequests(this, &MergedRequestList->GetAdaptiveAllocationRequest(0), MergedRequestList->GetNumAdaptiveAllocationRequests(), Frame);
 	}
 
+	// Deal with dirty MappedTilesToProduce pages first.
+	if (UpdatePhase == EUpdatePhase::Begin)
 	{
-		SCOPE_CYCLE_COUNTER(STAT_ProcessRequests_Sort);
-
-		// Limit the number of uploads (account for MappedTilesToProduce this frame but only for RVTs since SVT pages will take a while to load and likely won't be produced this frame).
-		// Use a separate budget for SVTs since we may want more of them to go through so that I/O can be initiated as early as possible
-		const int32 MaxNonStreamingLoadRequests = FMath::Max(Settings.MaxRVTPageUploads - MappedTilesToProduce.Num(), 1);
-		const int32 MaxStreamingLoadRequests = Settings.MaxSVTPageUploads;
-		check(MaxStreamingLoadRequests >= 0);
-		// 0 is a special value that enables the old behavior where all pages are limited by a single budget
-		const bool bUseCombinedLimit = !MaxStreamingLoadRequests;
-		const uint32 OldNumLoadRequests = MergedRequestList->GetNumLoadRequests();
+		// Only take runtime generated page budget into account since we expect that only runtime generated pages should have been dirtied.
+		const int32 NumPagesProduced = SubmitRequestsFromLocalTileList(RHICmdList, TransientCollectedPages, MappedTilesToProduce, EVTProducePageFlags::None, FeatureLevel, Updater->PageUploadBudgetRVT);
+		Updater->PageUploadBudgetRVT = FMath::Max(Updater->PageUploadBudgetRVT - NumPagesProduced, 0);
 
-		MergedRequestList->SortRequests(Producers, Allocator, MaxNonStreamingLoadRequests, MaxStreamingLoadRequests, bUseCombinedLimit);
-
-		if (MergedRequestList->GetNumLoadRequests() < OldNumLoadRequests)
-		{
-			// Dropping requests is normal but track to log here if we want to tune settings.
-			if (CVarVTVerbose.GetValueOnRenderThread())
-			{
-				UE_LOG(LogConsoleResponse, Display, TEXT("VT dropped %d load requests."), MergedRequestList->GetNumLoadRequests() - OldNumLoadRequests);
-			}
-		}
+		INC_DWORD_STAT_BY(STAT_NumMappedPageUpdate, MappedTilesToProduce.Num());
+		MappedTilesToProduce.Reset();
+		MappedTilesToProduce.Append(TransientCollectedPages);
+		TransientCollectedPages.Reset();
 	}
 
-	if (bContinousUpdates)
 	{
-		// After sorting and clamping the load requests, if we still have unused upload bandwidth then use it to add some continuous updates.
-		// Not taking SVT requests into account since they take a while to load and likely won't be produced this frame
-		const int32 MaxTilesToProduce = FMath::Max(Settings.MaxRVTPageUploads - MappedTilesToProduce.Num() - (int32)MergedRequestList->GetNumNonStreamingLoadRequests(), 0);
-		const int32 MaxContinuousUpdates = Settings.MaxContinuousUpdates;
-
-		GetContinuousUpdatesToProduce(MergedRequestList, MaxTilesToProduce, MaxContinuousUpdates);
+		SCOPE_CYCLE_COUNTER(STAT_ProcessRequests_Sort);
+		
+		// Throttle the total number of requests during sort according to the remaining budget.
+		// Note that we use distinct budgets for streaming and runtime generated pages, since they have difference performance characteristics.
+		// We may want a higher budget for streaming pages so that I/O can be initiated as early as possible.
+		
+		// SVT budget of 0 is a special value that enables the old behavior where all pages are limited by a single budget.
+		const bool bUseCombinedLimit = Settings.MaxSVTPageUploads == 0;
+		
+		MergedRequestList->SortRequests(Producers, Allocator, Updater->PageUploadBudgetRVT, Updater->PageUploadBudgetSVT, bUseCombinedLimit);
+		
+		// Subtract sorted request count from the remaining budgets.
+		const int32 NumRequestsRVT = (int32)MergedRequestList->GetNumNonStreamingLoadRequests();
+		const int32 NumRequestsSVT = (int32)MergedRequestList->GetNumLoadRequests() - (int32)MergedRequestList->GetNumNonStreamingLoadRequests();
+		check(NumRequestsSVT >= 0);
+		Updater->PageUploadBudgetRVT = FMath::Max(Updater->PageUploadBudgetRVT - NumRequestsRVT, 0);
+		Updater->PageUploadBudgetSVT = FMath::Max(Updater->PageUploadBudgetSVT - NumRequestsSVT, 0);
 	}
 
-	// Track total number of requests made this frame.
-	Updater->NumProcessedLoadRequests += MergedRequestList->GetNumLoadRequests();
+	// If we have any remaining page budget then use it to add continuous updates.
+	if (UpdatePhase == EUpdatePhase::End)
+	{
+		// Don't take streaming page budget into account since they async load and likely won't be produced this frame.
+		// Also we expect only runtime generated pages to need continous updates (since continuous updates are there to handle stale pages coming from rendering before ready etc).
+		GetContinuousUpdatesToProduce(MergedRequestList, Updater->PageUploadBudgetRVT, Settings.MaxContinuousUpdates);
+		const int32 NumPagesProduced = SubmitRequestsFromLocalTileList(RHICmdList, TransientCollectedPages, ContinuousUpdateTilesToProduce, EVTProducePageFlags::ContinuousUpdate, FeatureLevel, 0);
+		Updater->PageUploadBudgetRVT = FMath::Max(Updater->PageUploadBudgetRVT - NumPagesProduced, 0);
 
-	// Submit the requests to produce pages that are already mapped
-	SubmitPreMappedRequests(RHICmdList, FeatureLevel, Settings.MaxRVTPageUploads);
+		INC_DWORD_STAT_BY(STAT_NumContinuousPageUpdate, ContinuousUpdateTilesToProduce.Num());
+		ContinuousUpdateTilesToProduce.Reset();
+		TransientCollectedPages.Reset();
+	}
 
-	// Submit the merged requests
+	// Submit the merged, sorted and throttled page load requests.
 	SubmitRequests(RHICmdList, FeatureLevel, Allocator, Settings, MergedRequestList, true);
 }
 
@@ -2684,18 +2673,17 @@ void FVirtualTextureSystem::BeginUpdate(FRDGBuilder& GraphBuilder, FVirtualTextu
 	{
 		TRACE_CPUPROFILER_EVENT_SCOPE(FVirtualTextureSystem::GatherAndSubmitRequests);
 
-		const bool bContinousUpdates = false;
-
 		Updater->MergedRequestList = Allocator.Create<FUniqueRequestList>(Allocator);
 		Updater->MergedRequestList->Initialize();
 
 		GatherFeedbackRequests(Allocator, Settings, FeedbackResult, Updater->MergedRequestList);
 		GatherLockedTileRequests(Updater->MergedRequestList);
 		GatherPackedTileRequests(Allocator, Settings, Updater->MergedRequestList);
-		SubmitThrottledRequests(RHICmdList, Updater, bContinousUpdates);
+		SubmitThrottledRequests(RHICmdList, Updater, EUpdatePhase::Begin);
 
 		// Reset the request list for the gather in EndUpdate.
-		Updater->MergedRequestList->Reset(bContinousUpdates);
+		const bool bResetContinousUpdates = false;
+		Updater->MergedRequestList->Reset(bResetContinousUpdates);
 
 	}, UE::Tasks::ETaskPriority::High, Updater->bAsyncTaskAllowed);
 }
@@ -2795,6 +2783,8 @@ TUniquePtr<FVirtualTextureUpdater> FVirtualTextureSystem::BeginUpdate(FRDGBuilde
 	Updater->Settings = Settings;
 	Updater->FeatureLevel = FeatureLevel;
 	Updater->bAsyncTaskAllowed = Settings.bEnableAsyncTasks && CVarVTAsyncPageRequestTask.GetValueOnRenderThread();
+	Updater->PageUploadBudgetRVT = Settings.MaxRVTPageUploads;
+	Updater->PageUploadBudgetSVT = Settings.MaxSVTPageUploads;
 
 	if (Updater->bAsyncTaskAllowed)
 	{
@@ -2856,7 +2846,7 @@ void FVirtualTextureSystem::EndUpdate(FRDGBuilder& GraphBuilder, TUniquePtr<FVir
 			GatherPackedTileRequests(Updater->Allocator, Updater->Settings, Updater->MergedRequestList);
 		}
 
-		SubmitThrottledRequests(GraphBuilder.RHICmdList, Updater.Get(), bContinousUpdates);
+		SubmitThrottledRequests(GraphBuilder.RHICmdList, Updater.Get(), EUpdatePhase::End);
 	}
 
 	FinalizeRequests(GraphBuilder);
diff --git a/Engine/Source/Runtime/Renderer/Private/VT/VirtualTextureSystem.h b/Engine/Source/Runtime/Renderer/Private/VT/VirtualTextureSystem.h
--- a/Engine/Source/Runtime/Renderer/Private/VT/VirtualTextureSystem.h
+++ b/Engine/Source/Runtime/Renderer/Private/VT/VirtualTextureSystem.h
@@ -95,8 +95,9 @@ private:
 	FVirtualTextureFeedback::FMapResult FeedbackMapResult;
 	UE::Tasks::FTask AsyncTask;
 	ERHIFeatureLevel::Type FeatureLevel = ERHIFeatureLevel::Num;
-	uint32 NumProcessedLoadRequests = 0;
 	bool bAsyncTaskAllowed = false;
+	int32 PageUploadBudgetRVT = 0;
+	int32 PageUploadBudgetSVT = 0;
 
 	friend FVirtualTextureSystem;
 };
@@ -205,15 +206,19 @@ private:
 	void RequestTilesInternal(const IAllocatedVirtualTexture* AllocatedVT, int32 InMipLevel);
 	void RequestTilesInternal(const IAllocatedVirtualTexture* AllocatedVT, const FVector2D& InScreenSpaceSize, int32 InMipLevel);
 	
-	void SubmitRequestsFromLocalTileList(FRHICommandList& RHICmdList, TArray<FVirtualTextureLocalTile>& OutDeferredTiles, const TSet<FVirtualTextureLocalTile>& LocalTileList, EVTProducePageFlags Flags, ERHIFeatureLevel::Type FeatureLevel, uint32 MaxRequestsToProduce);
+	int32 SubmitRequestsFromLocalTileList(FRHICommandList& RHICmdList, TArray<FVirtualTextureLocalTile>& OutDeferredTiles, const TSet<FVirtualTextureLocalTile>& LocalTileList, EVTProducePageFlags Flags, ERHIFeatureLevel::Type FeatureLevel, uint32 MaxRequestsToProduce);
 
 	void GatherFeedbackRequests(FConcurrentLinearBulkObjectAllocator& Allocator, const FVirtualTextureUpdateSettings& Settings, const FVirtualTextureFeedback::FMapResult& FeedbackResult, FUniqueRequestList* MergedRequestList);
 	void GatherLockedTileRequests(FUniqueRequestList* MergedRequestList);
 	void GatherPackedTileRequests(FConcurrentLinearBulkObjectAllocator& Allocator, const FVirtualTextureUpdateSettings& Settings, FUniqueRequestList* MergedRequestList);
 
-	void SubmitPreMappedRequests(FRHICommandList& RHICmdList, ERHIFeatureLevel::Type FeatureLevel, uint32 MaxMappedTilesToProduce);
+	enum class EUpdatePhase
+	{
+		Begin,
+		End
+	};
 
-	void SubmitThrottledRequests(FRHICommandList& RHICmdList, FVirtualTextureUpdater* Updater, bool bContinuousUpdates);
+	void SubmitThrottledRequests(FRHICommandList& RHICmdList, FVirtualTextureUpdater* Updater, EUpdatePhase UpdatePhase);
 	void SubmitRequests(FRHICommandList& RHICmdList, ERHIFeatureLevel::Type FeatureLevel, FConcurrentLinearBulkObjectAllocator& Allocator, FVirtualTextureUpdateSettings const& Settings, FUniqueRequestList* RequestList, bool bAsync);
 	void FinalizeRequests(FRDGBuilder& GraphBuilder);
 
