LFV - fixed LFV not rendered in the sky after a recent (valid, 30677967) fix for stencil to be used when rendering the mobile fog pass in forward.

diff --git a/Engine/Shaders/Private/MobileBasePassPixelShader.usf b/Engine/Shaders/Private/MobileBasePassPixelShader.usf
--- a/Engine/Shaders/Private/MobileBasePassPixelShader.usf
+++ b/Engine/Shaders/Private/MobileBasePassPixelShader.usf
@@ -1101,6 +1101,18 @@ void Main(
 	OutColor.rgb *= ResolvedView.PreExposure;
 #endif
 
+#if FORWARD_SHADING && MATERIAL_IS_SKY
+	if (MobileBasePass.bApplyHalfResLocalFogToSkyMeshes)
+	{
+		// Apply the half resolution texture to the sky dome mesh when needed.
+		const float2 ResolutionRatioCorrection = (ResolvedView.ViewSizeAndInvSize.xy * MobileBasePass.LFV.HalfResTextureSizeAndInvSize.zw * 0.5);
+		const float2 HalfResSourceTexCoord = float2(SvPosition.xy - ResolvedView.ViewRectMin.xy) * ResolvedView.ViewSizeAndInvSize.zw * ResolutionRatioCorrection;
+		float4 LFVContribution = Texture2DSampleLevel(MobileBasePass.HalfResLocalFogVolumeViewTexture, MobileBasePass.HalfResLocalFogVolumeViewSampler, HalfResSourceTexCoord, 0.0f);
+
+		OutColor.rgb = OutColor.rgb * LFVContribution.aaa + LFVContribution.rgb;
+	}
+#endif
+
 #if MATERIAL_IS_SKY || (!MATERIAL_SHADINGMODEL_UNLIT && !MOBILE_USE_GBUBUFFER)
 	// clamp lit color to avoid overflows
 	OutColor.rgb = SafeGetOutColor(OutColor.rgb);
diff --git a/Engine/Shaders/Private/MobileFog.usf b/Engine/Shaders/Private/MobileFog.usf
--- a/Engine/Shaders/Private/MobileFog.usf
+++ b/Engine/Shaders/Private/MobileFog.usf
@@ -120,7 +120,7 @@ void MobileFogPS(
 	}
 
 	// HalfResLocalFogVolumeViewSRV contains pre exposed color, so remove it for composition.
-	LFVContribution.rgb /= ResolvedView.PreExposure;
+	LFVContribution.rgb *= ResolvedView.OneOverPreExposure;
 
 	// Combine with the curren fog
 	FinalFog = float4(LFVContribution.rgb + FinalFog.rgb * LFVContribution.a, LFVContribution.a * FinalFog.a);
diff --git a/Engine/Source/Runtime/Renderer/Private/LocalFogVolumeRendering.cpp b/Engine/Source/Runtime/Renderer/Private/LocalFogVolumeRendering.cpp
--- a/Engine/Source/Runtime/Renderer/Private/LocalFogVolumeRendering.cpp
+++ b/Engine/Source/Runtime/Renderer/Private/LocalFogVolumeRendering.cpp
@@ -194,7 +194,7 @@ void SetDummyLocalFogVolumeForView(FRDGBuilder& GraphBuilder, FViewInfo& View)
 	View.LocalFogVolumeViewData.GPUTileDrawIndirectBufferUAV= nullptr;
 
 	View.LocalFogVolumeViewData.bUseHalfResLocalFogVolume = false;
-	View.LocalFogVolumeViewData.HalfResLocalFogVolumeView = GSystemTextures.GetBlackDummy(GraphBuilder);
+	View.LocalFogVolumeViewData.HalfResLocalFogVolumeView = GSystemTextures.GetBlackAlphaOneDummy(GraphBuilder);
 	View.LocalFogVolumeViewData.HalfResLocalFogVolumeViewSRV = GraphBuilder.CreateSRV(View.LocalFogVolumeViewData.HalfResLocalFogVolumeView);
 	View.LocalFogVolumeViewData.HalfResLocalFogVolumeDepth = GSystemTextures.GetBlackDummy(GraphBuilder);
 	View.LocalFogVolumeViewData.HalfResLocalFogVolumeDepthSRV = GraphBuilder.CreateSRV(View.LocalFogVolumeViewData.HalfResLocalFogVolumeDepth);
@@ -602,7 +602,7 @@ void CreateViewLocalFogVolumeBufferSRV(const FScene* Scene, FViewInfo& View, FRD
 	else
 	{
 		View.LocalFogVolumeViewData.UniformParametersStruct.LocalFogVolumeCommon.HalfResTextureSizeAndInvSize = FVector4f(1.0f, 1.0f, 1.0f, 1.0f);
-		View.LocalFogVolumeViewData.HalfResLocalFogVolumeView = GSystemTextures.GetBlackDummy(GraphBuilder);
+		View.LocalFogVolumeViewData.HalfResLocalFogVolumeView = GSystemTextures.GetBlackAlphaOneDummy(GraphBuilder);
 		View.LocalFogVolumeViewData.HalfResLocalFogVolumeViewSRV = GraphBuilder.CreateSRV(View.LocalFogVolumeViewData.HalfResLocalFogVolumeView);
 		View.LocalFogVolumeViewData.HalfResLocalFogVolumeDepth = GSystemTextures.GetBlackDummy(GraphBuilder);
 		View.LocalFogVolumeViewData.HalfResLocalFogVolumeDepthSRV = GraphBuilder.CreateSRV(View.LocalFogVolumeViewData.HalfResLocalFogVolumeDepth);
@@ -1033,10 +1033,14 @@ void RenderLocalFogVolumeHalfResMobile(
 	}
 	FMobileLocalFogVolumeTiledRenderHalfResPS::FParameters* PassParameters = GraphBuilder.AllocParameters<FMobileLocalFogVolumeTiledRenderHalfResPS::FParameters>();
 
-	PassParameters->View = View.GetShaderParameters();;
+	// When rendering the half resolution LFV texture, we are right after the depth prepass so it should not be bound (the texture is going to be generated now).
+	// That is why we notify that we are still in the prepass stage.
+	const EMobileBasePass BasePass = EMobileBasePass::DepthPrePass;
+
+	PassParameters->View = View.GetShaderParameters();
 	PassParameters->LFV = View.LocalFogVolumeViewData.UniformParametersStruct;
 	PassParameters->LocalFogVolumeTileDebug = FMath::Clamp(CVarLocalFogVolumeTileDebug.GetValueOnRenderThread(), 0, 2);
-	PassParameters->MobileBasePass = CreateMobileBasePassUniformBuffer(GraphBuilder, View, EMobileBasePass::Opaque, EMobileSceneTextureSetupMode::SceneDepth);
+	PassParameters->MobileBasePass = CreateMobileBasePassUniformBuffer(GraphBuilder, View, BasePass, EMobileSceneTextureSetupMode::SceneDepth);
 	PassParameters->RenderTargets[0] = FRenderTargetBinding(View.LocalFogVolumeViewData.HalfResLocalFogVolumeView, ERenderTargetLoadAction::ENoAction);
 	PassParameters->RenderTargets[1] = FRenderTargetBinding(View.LocalFogVolumeViewData.HalfResLocalFogVolumeDepth, ERenderTargetLoadAction::ENoAction);
 
diff --git a/Engine/Source/Runtime/Renderer/Private/MobileBasePassRendering.cpp b/Engine/Source/Runtime/Renderer/Private/MobileBasePassRendering.cpp
--- a/Engine/Source/Runtime/Renderer/Private/MobileBasePassRendering.cpp
+++ b/Engine/Source/Runtime/Renderer/Private/MobileBasePassRendering.cpp
@@ -337,6 +337,20 @@ void SetupMobileBasePassUniformParameters(
 
 	BasePassParameters.LFV = View.LocalFogVolumeViewData.UniformParametersStruct;
 
+	// We need to compose the half resolution LFV texture when rendering mesh with Sky materials. So that the Fog passes remains cheap and we can keep the stencil test on the fog pass.
+	if (BasePass > EMobileBasePass::DepthPrePass)
+	{
+		// HalfResLocalFogVolumeView is rendered after the depth pre pass so we only bind it after the depth pre pass
+		BasePassParameters.bApplyHalfResLocalFogToSkyMeshes = View.LocalFogVolumeViewData.bUseHalfResLocalFogVolume ? 1 : 0;
+		BasePassParameters.HalfResLocalFogVolumeViewTexture = View.LocalFogVolumeViewData.HalfResLocalFogVolumeView;
+	}
+	else
+	{
+		BasePassParameters.bApplyHalfResLocalFogToSkyMeshes = 0;
+		BasePassParameters.HalfResLocalFogVolumeViewTexture = SystemTextures.BlackAlphaOne;
+	}
+	BasePassParameters.HalfResLocalFogVolumeViewSampler = TStaticSamplerState<SF_Bilinear, AM_Clamp, AM_Clamp, AM_Clamp>::GetRHI();
+
 	SetupReflectionUniformParameters(GraphBuilder, View, BasePassParameters.ReflectionsParameters);
 }
 
diff --git a/Engine/Source/Runtime/Renderer/Private/MobileBasePassRendering.h b/Engine/Source/Runtime/Renderer/Private/MobileBasePassRendering.h
--- a/Engine/Source/Runtime/Renderer/Private/MobileBasePassRendering.h
+++ b/Engine/Source/Runtime/Renderer/Private/MobileBasePassRendering.h
@@ -59,11 +59,15 @@ BEGIN_GLOBAL_SHADER_PARAMETER_STRUCT(FMobileBasePassUniformParameters, )
 	SHADER_PARAMETER_SAMPLER(SamplerState, AmbientOcclusionSampler)
 	SHADER_PARAMETER_RDG_TEXTURE(Texture2D, ScreenSpaceShadowMaskTexture)
 	SHADER_PARAMETER_SAMPLER(SamplerState, ScreenSpaceShadowMaskSampler)
+	SHADER_PARAMETER(uint32, bApplyHalfResLocalFogToSkyMeshes)
+	SHADER_PARAMETER_RDG_TEXTURE(Texture2D, HalfResLocalFogVolumeViewTexture)
+	SHADER_PARAMETER_SAMPLER(SamplerState, HalfResLocalFogVolumeViewSampler)
 	SHADER_PARAMETER_STRUCT_INCLUDE(FDBufferParameters, DBuffer)
 END_GLOBAL_SHADER_PARAMETER_STRUCT()
 
 enum class EMobileBasePass
 {
+	DepthPrePass,
 	Opaque,
 	Translucent
 };
diff --git a/Engine/Source/Runtime/Renderer/Private/MobileShadingRenderer.cpp b/Engine/Source/Runtime/Renderer/Private/MobileShadingRenderer.cpp
--- a/Engine/Source/Runtime/Renderer/Private/MobileShadingRenderer.cpp
+++ b/Engine/Source/Runtime/Renderer/Private/MobileShadingRenderer.cpp
@@ -804,7 +804,7 @@ void FMobileSceneRenderer::RenderFullDepthPrepass(FRDGBuilder& GraphBuilder, TAr
 		auto* PassParameters = GraphBuilder.AllocParameters<FMobileRenderPassParameters>();
 		PassParameters->RenderTargets = BasePassRenderTargets;
 		PassParameters->View = View.GetShaderParameters();
-		PassParameters->MobileBasePass = CreateMobileBasePassUniformBuffer(GraphBuilder, View, EMobileBasePass::Opaque, EMobileSceneTextureSetupMode::None);
+		PassParameters->MobileBasePass = CreateMobileBasePassUniformBuffer(GraphBuilder, View, EMobileBasePass::DepthPrePass, EMobileSceneTextureSetupMode::None);
 		
 		View.ParallelMeshDrawCommandPasses[EMeshPass::DepthPass].BuildRenderingCommands(GraphBuilder, Scene->GPUScene, PassParameters->InstanceCullingDrawParams);
 
