Fixes Heterogeneous Volume shadow flickering with Directional Light

diff --git a/Engine/Source/Runtime/Renderer/Private/HeterogeneousVolumes/HeterogeneousVolumesLiveShadingPipeline.cpp b/Engine/Source/Runtime/Renderer/Private/HeterogeneousVolumes/HeterogeneousVolumesLiveShadingPipeline.cpp
--- a/Engine/Source/Runtime/Renderer/Private/HeterogeneousVolumes/HeterogeneousVolumesLiveShadingPipeline.cpp
+++ b/Engine/Source/Runtime/Renderer/Private/HeterogeneousVolumes/HeterogeneousVolumesLiveShadingPipeline.cpp
@@ -1507,7 +1507,7 @@ bool RenderVolumetricShadowMapForLightWithLiveShading(
 
 	// Collect shadow-casting volumes
 	TSet<FVolumetricMeshBatch> HeterogeneousVolumesMeshBatches;
-	FBoxSphereBounds WorldVolumeBounds;
+	FBoxSphereBounds WorldVolumeBounds(ForceInit);
 	CollectHeterogeneousVolumeMeshBatchesForLight(LightSceneInfo, VisibleLightInfo, View, HeterogeneousVolumesMeshBatches, WorldVolumeBounds);
 	if (HeterogeneousVolumesMeshBatches.IsEmpty())
 	{
