Fix directional light shadow projections failing to appear in instanced stereo views.

diff --git a/Engine/Source/Runtime/Engine/Private/SceneView.cpp b/Engine/Source/Runtime/Engine/Private/SceneView.cpp
--- a/Engine/Source/Runtime/Engine/Private/SceneView.cpp
+++ b/Engine/Source/Runtime/Engine/Private/SceneView.cpp
@@ -2407,12 +2407,6 @@ bool FSceneView::IsInstancedStereoPass() const
 	return bIsInstancedStereoEnabled && IStereoRendering::IsStereoEyeView(*this) && IStereoRendering::IsAPrimaryView(*this);
 }
 
-int32 FSceneView::GetStereoPassInstanceFactor() const
-{
-	return bIsInstancedStereoEnabled && IStereoRendering::IsStereoEyeView(*this) && GEngine->StereoRenderingDevice.IsValid() ?
-		GEngine->StereoRenderingDevice->GetDesiredNumberOfViews(true) : 1;
-}
-
 FVector4f FSceneView::GetScreenPositionScaleBias(const FIntPoint& BufferSize, const FIntRect& ViewRect) const
 {
 	const float InvBufferSizeX = 1.0f / BufferSize.X;
diff --git a/Engine/Source/Runtime/Engine/Public/SceneView.h b/Engine/Source/Runtime/Engine/Public/SceneView.h
--- a/Engine/Source/Runtime/Engine/Public/SceneView.h
+++ b/Engine/Source/Runtime/Engine/Public/SceneView.h
@@ -1614,6 +1614,9 @@ public:
 	/** Feature level for this scene */
 	const ERHIFeatureLevel::Type FeatureLevel;
 
+	/** When using mobile multi view fallback path we need to instance draw calls ourselves to cover both eyes instead of letting the drivers do it for us. */
+	uint32 InstanceFactor = 1;
+
 	/** Use to allow ray tracing on this view. */
 	bool bAllowRayTracing = true;
 
@@ -1790,7 +1793,7 @@ public:
 	ENGINE_API bool IsInstancedStereoPass() const;
 
 	/** Instance factor for a stereo pass (normally 2 for ISR views, but see IStereoRendering::GetDesiredNumberOfViews()). Returns 1 for non-instanced stereo views or regular (split screen etc) views. */
-	ENGINE_API int32 GetStereoPassInstanceFactor() const;
+	ENGINE_API int32 GetStereoPassInstanceFactor() const { return InstanceFactor; }
 
 	/** Sets up the view rect parameters in the view's uniform shader parameters */
 	ENGINE_API void SetupViewRectUniformBufferParameters(FViewUniformShaderParameters& ViewUniformShaderParameters, 
diff --git a/Engine/Source/Runtime/Renderer/Private/MobileDecalRendering.cpp b/Engine/Source/Runtime/Renderer/Private/MobileDecalRendering.cpp
--- a/Engine/Source/Runtime/Renderer/Private/MobileDecalRendering.cpp
+++ b/Engine/Source/Runtime/Renderer/Private/MobileDecalRendering.cpp
@@ -148,7 +148,7 @@ void RenderDeferredDecalsMobile(FRHICommandList& RHICmdList, const FScene& Scene
 			// Set shader params
 			DecalRendering::SetShader(RHICmdList, GraphicsPSOInit, 0, View, DecalData, DecalRenderStage, FrustumComponentToClip);
 
-			RHICmdList.DrawIndexedPrimitive(GetUnitCubeIndexBuffer(), 0, 0, 8, 0, UE_ARRAY_COUNT(GCubeIndices) / 3, View.InstanceFactor);
+			RHICmdList.DrawIndexedPrimitive(GetUnitCubeIndexBuffer(), 0, 0, 8, 0, UE_ARRAY_COUNT(GCubeIndices) / 3, View.GetStereoPassInstanceFactor());
 		}
 	}
 }
diff --git a/Engine/Source/Runtime/Renderer/Private/MobileShadingRenderer.cpp b/Engine/Source/Runtime/Renderer/Private/MobileShadingRenderer.cpp
--- a/Engine/Source/Runtime/Renderer/Private/MobileShadingRenderer.cpp
+++ b/Engine/Source/Runtime/Renderer/Private/MobileShadingRenderer.cpp
@@ -446,12 +446,14 @@ void FMobileSceneRenderer::InitViews(
 
 	// Create GPU-side representation of the view for instance culling.
 	InstanceCullingManager.AllocateViews(Views.Num());
-	for (int32 ViewIndex = 0; ViewIndex < Views.Num(); ++ViewIndex)
+	for (FViewInfo& ViewInfo : Views)
 	{
-		Views[ViewIndex].GPUSceneViewId = InstanceCullingManager.RegisterView(Views[ViewIndex]);
-		// Set the stereo instance factor as GetStereoPassInstanceFactor can't be called from inside an RDG pass.
-		uint32 InstanceFactor = Views[ViewIndex].GetStereoPassInstanceFactor();
-		Views[ViewIndex].InstanceFactor = InstanceFactor > 0 ? InstanceFactor : 1;
+		ViewInfo.GPUSceneViewId = InstanceCullingManager.RegisterView(ViewInfo);
+
+		uint32 InstanceFactor = ViewInfo.bIsInstancedStereoEnabled && IStereoRendering::IsStereoEyeView(ViewInfo) && GEngine->StereoRenderingDevice.IsValid() ?
+			GEngine->StereoRenderingDevice->GetDesiredNumberOfViews(true) : 1;
+
+		ViewInfo.InstanceFactor = InstanceFactor > 0 ? InstanceFactor : 1;
 	}
 
 	FILCUpdatePrimTaskData* ILCTaskData = nullptr;
diff --git a/Engine/Source/Runtime/Renderer/Private/PostProcess/PostProcessTonemap.cpp b/Engine/Source/Runtime/Renderer/Private/PostProcess/PostProcessTonemap.cpp
--- a/Engine/Source/Runtime/Renderer/Private/PostProcess/PostProcessTonemap.cpp
+++ b/Engine/Source/Runtime/Renderer/Private/PostProcess/PostProcessTonemap.cpp
@@ -1139,7 +1139,7 @@ void RenderMobileCustomResolve(FRHICommandList& RHICmdList, const FViewInfo& Vie
 		TargetSize,
 		VertexShader,
 		EDRF_UseTriangleOptimization,
-		View.InstanceFactor);
+		View.GetStereoPassInstanceFactor());
 }
 
 BEGIN_SHADER_PARAMETER_STRUCT(FMobileCustomResolveParameters, )
diff --git a/Engine/Source/Runtime/Renderer/Private/SceneRendering.h b/Engine/Source/Runtime/Renderer/Private/SceneRendering.h
--- a/Engine/Source/Runtime/Renderer/Private/SceneRendering.h
+++ b/Engine/Source/Runtime/Renderer/Private/SceneRendering.h
@@ -1539,9 +1539,6 @@ public:
 
 	FShaderPrintData ShaderPrintData;
 
-	/** When using mobile multi view fallback path we need to instance draw calls ourselves to cover both eyes instead of letting the drivers do it for us. */
-	uint32 InstanceFactor = 1;
-
 private:
 	FLumenTranslucencyGIVolume LumenTranslucencyGIVolume;
 public:
diff --git a/Engine/Source/Runtime/Renderer/Private/SceneVisibility.cpp b/Engine/Source/Runtime/Renderer/Private/SceneVisibility.cpp
--- a/Engine/Source/Runtime/Renderer/Private/SceneVisibility.cpp
+++ b/Engine/Source/Runtime/Renderer/Private/SceneVisibility.cpp
@@ -5588,9 +5588,14 @@ void FDeferredShadingSceneRenderer::BeginInitViews(
 
 	// Create GPU-side representation of the view for instance culling.
 	InstanceCullingManager.AllocateViews(Views.Num());
-	for (int32 ViewIndex = 0; ViewIndex < Views.Num(); ++ViewIndex)
+	for (FViewInfo& ViewInfo : Views)
 	{
-		Views[ViewIndex].GPUSceneViewId = InstanceCullingManager.RegisterView(Views[ViewIndex]);
+		ViewInfo.GPUSceneViewId = InstanceCullingManager.RegisterView(ViewInfo);
+
+		uint32 InstanceFactor = ViewInfo.bIsInstancedStereoEnabled && IStereoRendering::IsStereoEyeView(ViewInfo) && GEngine->StereoRenderingDevice.IsValid() ?
+			GEngine->StereoRenderingDevice->GetDesiredNumberOfViews(true) : 1;
+
+		ViewInfo.InstanceFactor = InstanceFactor > 0 ? InstanceFactor : 1;
 	}
 
 	{
diff --git a/Engine/Source/Runtime/Renderer/Private/ShadowRendering.cpp b/Engine/Source/Runtime/Renderer/Private/ShadowRendering.cpp
--- a/Engine/Source/Runtime/Renderer/Private/ShadowRendering.cpp
+++ b/Engine/Source/Runtime/Renderer/Private/ShadowRendering.cpp
@@ -1070,7 +1070,7 @@ void FProjectedShadowInfo::SetupProjectionStencilMask(
 		RHICmdList.SetStreamSource(0, GFrustumVertexBuffer.VertexBufferRHI, 0);
 
 		// Draw the frustum using the stencil buffer to mask just the pixels which are inside the shadow frustum.
-		RHICmdList.DrawIndexedPrimitive(GCubeIndexBuffer.IndexBufferRHI, 0, 0, 8, 0, 12, View->InstanceFactor);
+		RHICmdList.DrawIndexedPrimitive(GCubeIndexBuffer.IndexBufferRHI, 0, 0, 8, 0, 12, View->GetStereoPassInstanceFactor());
 
 		// if rendering modulated shadows mask out subject mesh elements to prevent self shadowing.
 		if (bMobileModulatedProjections && !CVarEnableModulatedSelfShadow.GetValueOnRenderThread())
@@ -1413,7 +1413,7 @@ void FProjectedShadowInfo::RenderProjectionInternal(
 		}
 	}
 
-	uint32 NumberOfInstances = View->InstanceFactor;
+	uint32 NumberOfInstances = View->GetStereoPassInstanceFactor();
 	if (IsWholeSceneDirectionalShadow())
 	{
 		RHICmdList.SetStreamSource(0, GClearVertexBuffer.VertexBufferRHI, 0);
diff --git a/Engine/Source/Runtime/Renderer/Private/SimpleMeshDrawCommandPass.cpp b/Engine/Source/Runtime/Renderer/Private/SimpleMeshDrawCommandPass.cpp
--- a/Engine/Source/Runtime/Renderer/Private/SimpleMeshDrawCommandPass.cpp
+++ b/Engine/Source/Runtime/Renderer/Private/SimpleMeshDrawCommandPass.cpp
@@ -25,7 +25,7 @@ FSimpleMeshDrawCommandPass::FSimpleMeshDrawCommandPass(const FSceneView& View, F
 	static FName NAME_SimpleMDCPass("SimpleMDCPass");
 	InstanceCullingContext = FInstanceCullingContext(NAME_SimpleMDCPass, ViewInfo->GetShaderPlatform(), InstanceCullingManager, ViewIds, nullptr, bUsingStereo ? EInstanceCullingMode::Stereo : EInstanceCullingMode::Normal);
 
-	InstanceFactor = ViewInfo->InstanceFactor;
+	InstanceFactor = ViewInfo->GetStereoPassInstanceFactor();
 }
 
 void FSimpleMeshDrawCommandPass::BuildRenderingCommands(FRDGBuilder& GraphBuilder, const FSceneView& View, const FGPUScene& GPUScene, FInstanceCullingDrawParams& OutInstanceCullingDrawParams)
