Fix stale CachedRayTracingMeshCommands when RefreshRayTracingMeshCommandCache() is used.
- CachedRayTracingMeshCommands was being emptied but CachedRayTracingMeshCommandIndicesPerLOD / CachedRayTracingMeshCommandsHashPerLOD could end up with stale data if primitive ray tracing flags included DisableCacheMeshCommandsFlags.

Repro steps:
- r.RayTracing.UseReferenceBasedResidency 1
- Move away from primitives in the level so they are not referenced in TLAS
- r.RayTracing.UpdateCachedState
- Move closer to primitives in the level so they start being referenced in TLAS

diff --git a/Engine/Source/Runtime/Renderer/Private/PrimitiveSceneInfo.cpp b/Engine/Source/Runtime/Renderer/Private/PrimitiveSceneInfo.cpp
--- a/Engine/Source/Runtime/Renderer/Private/PrimitiveSceneInfo.cpp
+++ b/Engine/Source/Runtime/Renderer/Private/PrimitiveSceneInfo.cpp
@@ -924,7 +924,12 @@ void FPrimitiveSceneInfo::RemoveCachedNaniteMaterialBins()
 void FScene::RefreshRayTracingMeshCommandCache()
 {
 	// Get rid of all existing cached commands
-	CachedRayTracingMeshCommands.Empty(CachedRayTracingMeshCommands.Num());
+	for (FPrimitiveSceneInfo* SceneInfo : Primitives)
+	{
+		SceneInfo->RemoveCachedRayTracingPrimitives();
+	}
+
+	check(CachedRayTracingMeshCommands.IsEmpty());
 
 	// Re-cache all current primitives
 	FPrimitiveSceneInfo::CacheRayTracingPrimitives(this, Primitives);
