Add missing BUF_SourceCopy flag to FRayTracingScene::OutputStatsBuffer.

diff --git a/Engine/Source/Runtime/Renderer/Private/RayTracing/RayTracingScene.cpp b/Engine/Source/Runtime/Renderer/Private/RayTracing/RayTracingScene.cpp
--- a/Engine/Source/Runtime/Renderer/Private/RayTracing/RayTracingScene.cpp
+++ b/Engine/Source/Runtime/Renderer/Private/RayTracing/RayTracingScene.cpp
@@ -169,7 +169,10 @@ void FRayTracingScene::CreateWithInitializationData(FRDGBuilder& GraphBuilder, c
 	}
 
 #if STATS
-	FRDGBufferRef OutputStatsBuffer = GraphBuilder.CreateBuffer(FRDGBufferDesc::CreateStructuredDesc(sizeof(uint32), 1), TEXT("FRayTracingScene::OutputStatsBuffer"));
+	FRDGBufferDesc OutputStatsBufferDesc(FRDGBufferDesc::CreateStructuredDesc(sizeof(uint32), 1));
+	OutputStatsBufferDesc.Usage |= BUF_SourceCopy;
+
+	FRDGBufferRef OutputStatsBuffer = GraphBuilder.CreateBuffer(OutputStatsBufferDesc, TEXT("FRayTracingScene::OutputStatsBuffer"));
 	FRDGBufferUAVRef OutputStatsBufferUAV = GraphBuilder.CreateUAV(OutputStatsBuffer);
 	AddClearUAVPass(GraphBuilder, OutputStatsBufferUAV, 0);
 #endif
