Fixed GPU crash on iOS on a landscape render capture

diff --git a/Engine/Source/Runtime/Renderer/Private/SimpleMeshDrawCommandPass.cpp b/Engine/Source/Runtime/Renderer/Private/SimpleMeshDrawCommandPass.cpp
--- a/Engine/Source/Runtime/Renderer/Private/SimpleMeshDrawCommandPass.cpp
+++ b/Engine/Source/Runtime/Renderer/Private/SimpleMeshDrawCommandPass.cpp
@@ -81,7 +81,7 @@ void FSimpleMeshDrawCommandPass::SubmitDraw(FRHICommandList& RHICmdList, const F
 		else
 		{
 			FMeshDrawCommandSceneArgs SceneArgs;
-			SceneArgs.PrimitiveIdsBuffer = PrimitiveIdVertexBuffer;
+			SceneArgs.PrimitiveIdsBuffer = GetMeshDrawCommandOverrideArgs(InstanceCullingDrawParams).InstanceBuffer;
 			SceneArgs.PrimitiveIdOffset = 0u;
 			SceneArgs.BatchedPrimitiveSlot = InstanceCullingContext.BatchedPrimitiveSlot;
 			const uint32 PrimitiveIdBufferStride = FInstanceCullingContext::GetInstanceIdBufferStride(InstanceCullingContext.ShaderPlatform);
