Refactor VRS RHI implementation and manager to avoid using GRHI variables bound to CVars, and to allow disabling support for VRS with a read-only variable (requires shader recompile, may avoid performance cost on some platforms).

diff --git a/Engine/Plugins/Runtime/OpenXR/Source/OpenXRHMD/Public/OpenXRHMDSettings.h b/Engine/Plugins/Runtime/OpenXR/Source/OpenXRHMD/Public/OpenXRHMDSettings.h
--- a/Engine/Plugins/Runtime/OpenXR/Source/OpenXRHMD/Public/OpenXRHMDSettings.h
+++ b/Engine/Plugins/Runtime/OpenXR/Source/OpenXRHMD/Public/OpenXRHMDSettings.h
@@ -4,6 +4,7 @@
 
 #include "CoreMinimal.h"
 #include "UObject/Object.h"
+#include "Engine/RendererSettings.h"
 #include "OpenXRHMDSettings.generated.h"
 
 /**
@@ -16,6 +17,8 @@ class OPENXRHMD_API UOpenXRHMDSettings : public UObject
 
 public:
 	/** Enables foveation provided by the XR_FB_foveation OpenXR extension. */
-	UPROPERTY(config, EditAnywhere, Category = "Foveation", meta = (DisplayName = "Enable XR_FB_foveation extension"))
+	UPROPERTY(config, EditAnywhere, Category = "Foveation", meta = (
+		ToolTip = "Enables foveation provided by the XR_FB_foveation OpenXR extension. Requires support for hardware variable rate shading.", 
+		DisplayName = "Enable XR_FB_foveation extension"))
 	bool bIsFBFoveationEnabled = false;
 };
diff --git a/Engine/Source/Runtime/D3D12RHI/Private/D3D12Commands.cpp b/Engine/Source/Runtime/D3D12RHI/Private/D3D12Commands.cpp
--- a/Engine/Source/Runtime/D3D12RHI/Private/D3D12Commands.cpp
+++ b/Engine/Source/Runtime/D3D12RHI/Private/D3D12Commands.cpp
@@ -792,7 +792,7 @@ void FD3D12CommandContext::RHISetGraphicsPipelineState(FRHIGraphicsPipelineState
 		StateCache.SetDepthBounds(0.0f, 1.0f);
 	}
 
-	if (GRHISupportsPipelineVariableRateShading && GRHIVariableRateShadingEnabled)
+	if (GRHISupportsPipelineVariableRateShading)
 	{
 		if (GraphicsPipelineState->PipelineStateInitializer.bAllowVariableRateShading)
 		{
@@ -800,6 +800,7 @@ void FD3D12CommandContext::RHISetGraphicsPipelineState(FRHIGraphicsPipelineState
 		}
 		else
 		{
+			// Also forces shading rate image attachment to be ignored
 			StateCache.SetShadingRate(EVRSShadingRate::VRSSR_1x1, VRSRB_Passthrough, VRSRB_Passthrough);
 		}
 	}
@@ -1371,26 +1372,19 @@ void FD3D12CommandContext::SetRenderTargetsAndClear(const FRHISetRenderTargetsIn
 	}
 
 #if PLATFORM_SUPPORTS_VARIABLE_RATE_SHADING
-	if (GRHISupportsPipelineVariableRateShading && GRHIVariableRateShadingEnabled)
+	if (GRHISupportsAttachmentVariableRateShading)
 	{
-		if (GRHISupportsAttachmentVariableRateShading && GRHIAttachmentVariableRateShadingEnabled)
+		if (RenderTargetsInfo.ShadingRateTexture != nullptr)
 		{
-			if (RenderTargetsInfo.ShadingRateTexture != nullptr)
-			{
-				FD3D12Resource* Resource = RetrieveTexture(RenderTargetsInfo.ShadingRateTexture)->GetResource();
+			FD3D12Resource* Resource = RetrieveTexture(RenderTargetsInfo.ShadingRateTexture)->GetResource();
 
-				TransitionResource(
-					Resource,
-					D3D12_RESOURCE_STATE_TBD,
-					D3D12_RESOURCE_STATE_SHADING_RATE_SOURCE,
-					D3D12_RESOURCE_BARRIER_ALL_SUBRESOURCES);
+			TransitionResource(
+				Resource,
+				D3D12_RESOURCE_STATE_TBD,
+				D3D12_RESOURCE_STATE_SHADING_RATE_SOURCE,
+				D3D12_RESOURCE_BARRIER_ALL_SUBRESOURCES);
 
-				StateCache.SetShadingRateImage(Resource);
-			}
-			else
-			{
-				StateCache.SetShadingRateImage(nullptr);
-			}
+			StateCache.SetShadingRateImage(Resource);
 		}
 		else
 		{
@@ -1930,14 +1924,17 @@ void FD3D12CommandContext::RHISetShadingRate(EVRSShadingRate ShadingRate, EVRSRa
 {
 #if PLATFORM_SUPPORTS_VARIABLE_RATE_SHADING
 	// Note - this will override per-material VRS opt-out, but FRHICommandSetShadingRate isn't called from anywhere
-	StateCache.SetShadingRate(ShadingRate, Combiner, VRSRB_Max);
+	if (GRHISupportsPipelineVariableRateShading)
+	{
+		StateCache.SetShadingRate(ShadingRate, Combiner, VRSRB_Max);
+	}
 #endif
 }
 
 void FD3D12CommandContext::SetShadingRate(EVRSShadingRate ShadingRate, const TStaticArray<EVRSRateCombiner, ED3D12VRSCombinerStages::Num>& Combiners)
 {
  #if PLATFORM_SUPPORTS_VARIABLE_RATE_SHADING
- 	if (GRHISupportsPipelineVariableRateShading && GRHIVariableRateShadingEnabled && GraphicsCommandList5())
+ 	if (GRHISupportsPipelineVariableRateShading && GraphicsCommandList5())
  	{
 		for (int32 CombinerIndex = 0; CombinerIndex < Combiners.Num(); ++CombinerIndex)
 		{
@@ -1960,7 +1957,7 @@ void FD3D12CommandContext::SetShadingRate(EVRSShadingRate ShadingRate, const TSt
 void FD3D12CommandContext::SetShadingRateImage(FD3D12Resource* RateImageTexture)
 {
 #if PLATFORM_SUPPORTS_VARIABLE_RATE_SHADING
-	if (GRHISupportsAttachmentVariableRateShading && GRHIAttachmentVariableRateShadingEnabled && GraphicsCommandList5())
+	if (GRHISupportsAttachmentVariableRateShading && GraphicsCommandList5())
 	{
 		if (RateImageTexture)
 		{
diff --git a/Engine/Source/Runtime/D3D12RHI/Private/D3D12StateCache.cpp b/Engine/Source/Runtime/D3D12RHI/Private/D3D12StateCache.cpp
--- a/Engine/Source/Runtime/D3D12RHI/Private/D3D12StateCache.cpp
+++ b/Engine/Source/Runtime/D3D12RHI/Private/D3D12StateCache.cpp
@@ -228,9 +228,8 @@ void FD3D12StateCache::DirtyStateForNewCommandList()
 		bNeedSetDepthBounds = GSupportsDepthBoundsTest;
 	}
 	
-	bNeedSetShadingRate = GRHISupportsPipelineVariableRateShading && GRHIVariableRateShadingEnabled;
-	
-	bNeedSetShadingRateImage = GRHISupportsAttachmentVariableRateShading && GRHIAttachmentVariableRateShadingEnabled;
+	bNeedSetShadingRate = GRHISupportsPipelineVariableRateShading;
+	bNeedSetShadingRateImage = GRHISupportsAttachmentVariableRateShading;
 
 	// Always dirty View and Sampler bindings. We detect the slots that are actually used at Draw/Dispatch time.
 	PipelineState.Common.SRVCache.DirtyAll();
@@ -254,8 +253,8 @@ void FD3D12StateCache::DirtyState()
 	bNeedSetBlendFactor = true;
 	bNeedSetStencilRef = true;
 	bNeedSetDepthBounds = GSupportsDepthBoundsTest;
-	bNeedSetShadingRate = GRHISupportsPipelineVariableRateShading && GRHIVariableRateShadingEnabled;
-	bNeedSetShadingRateImage = GRHISupportsAttachmentVariableRateShading && GRHIAttachmentVariableRateShadingEnabled;
+	bNeedSetShadingRate = GRHISupportsPipelineVariableRateShading;
+	bNeedSetShadingRateImage = GRHISupportsAttachmentVariableRateShading;
 	PipelineState.Common.SRVCache.DirtyAll();
 	PipelineState.Common.UAVCache.DirtyAll();
 	PipelineState.Common.CBVCache.DirtyAll();
diff --git a/Engine/Source/Runtime/D3D12RHI/Private/D3D12StateCachePrivate.h b/Engine/Source/Runtime/D3D12RHI/Private/D3D12StateCachePrivate.h
--- a/Engine/Source/Runtime/D3D12RHI/Private/D3D12StateCachePrivate.h
+++ b/Engine/Source/Runtime/D3D12RHI/Private/D3D12StateCachePrivate.h
@@ -727,7 +727,7 @@ public:
 			PipelineState.Graphics.DrawShadingRate = ShadingRate;
 			PipelineState.Graphics.Combiners[ED3D12VRSCombinerStages::PerPrimitive] = PerPrimitiveCombiner;
 			PipelineState.Graphics.Combiners[ED3D12VRSCombinerStages::ScreenSpace] = ScreenSpaceCombiner;
-			bNeedSetShadingRate = GRHISupportsPipelineVariableRateShading && GRHIVariableRateShadingEnabled;
+			bNeedSetShadingRate = GRHISupportsPipelineVariableRateShading;
 		}
 	}	
 
@@ -736,7 +736,7 @@ public:
 		if (PipelineState.Graphics.ShadingRateImage != ShadingRateImage)
 		{
 			PipelineState.Graphics.ShadingRateImage = ShadingRateImage;
-			bNeedSetShadingRateImage = GRHISupportsAttachmentVariableRateShading && GRHIAttachmentVariableRateShadingEnabled;
+			bNeedSetShadingRateImage = GRHISupportsAttachmentVariableRateShading;
 		}
 	}
 
diff --git a/Engine/Source/Runtime/D3D12RHI/Private/Windows/WindowsD3D12Device.cpp b/Engine/Source/Runtime/D3D12RHI/Private/Windows/WindowsD3D12Device.cpp
--- a/Engine/Source/Runtime/D3D12RHI/Private/Windows/WindowsD3D12Device.cpp
+++ b/Engine/Source/Runtime/D3D12RHI/Private/Windows/WindowsD3D12Device.cpp
@@ -1703,7 +1703,7 @@ void FD3D12DynamicRHI::Init()
 	// - Standalones are added to the deferred deletion queue of its parent FD3D12Adapter
 	GRHIForceNoDeletionLatencyForStreamingTextures = !!PLATFORM_WINDOWS;
 
-	if(Options6HR == S_OK && options.VariableShadingRateTier != D3D12_VARIABLE_SHADING_RATE_TIER_NOT_SUPPORTED)
+	if(Options6HR == S_OK && options.VariableShadingRateTier != D3D12_VARIABLE_SHADING_RATE_TIER_NOT_SUPPORTED && HardwareVariableRateShadingSupportedByPlatform(GMaxRHIShaderPlatform))
 	{
 		GRHISupportsPipelineVariableRateShading = true;		// We have at least tier 1.
 		GRHISupportsLargerVariableRateShadingSizes = (options.AdditionalShadingRatesSupported != 0);
diff --git a/Engine/Source/Runtime/Engine/Classes/Engine/RendererSettings.h b/Engine/Source/Runtime/Engine/Classes/Engine/RendererSettings.h
--- a/Engine/Source/Runtime/Engine/Classes/Engine/RendererSettings.h
+++ b/Engine/Source/Runtime/Engine/Classes/Engine/RendererSettings.h
@@ -633,11 +633,13 @@ class URendererSettings : public UDeveloperSettings
 	uint32 bLocalFogVolumeApplyOnTranslucent:1;
 
 	UPROPERTY(config, EditAnywhere, Category = VR, meta = (
+		EditCondition = "bSupportHardwareVariableRateShading",
 		ConsoleVariable = "xr.VRS.FoveationLevel", DisplayName = "Stereo Foveation Level (Experimental)",
 		ToolTip = "Set the level of foveation to apply when generating the Variable Rate Shading attachment. This feature is currently experimental.\nThis can yield some fairly significant performance benefits on GPUs that support Tier 2 VRS.\nLower settings will result in almost no discernible artifacting on most HMDs; higher settings will show some artifacts towards the edges of the view."))
 	TEnumAsByte<EFixedFoveationLevels::Type> FoveationLevel;
 
 	UPROPERTY(config, EditAnywhere, Category = VR, meta = (
+		EditCondition = "bSupportHardwareVariableRateShading && FoveationLevel != EFixedFoveationLevels::Disabled",
 		ConsoleVariable = "xr.VRS.DynamicFoveation", DisplayName = "Dynamic Foveation (Experimental)",
 		ToolTip = "Allows foveation level to adjust dynamically based on GPU utilization.\nLevel will range between none at the minimum, and the currently selected foveation level at the maximum."))
 	uint32 bDynamicFoveation:1;
@@ -1154,6 +1156,15 @@ class URendererSettings : public UDeveloperSettings
 		ConfigRestartRequired = true))
 	uint32 bSupportSkinCacheShaders : 1;
 
+	/**
+	"Support hardware variable rate shading.
+	*/
+	UPROPERTY(config, EditAnywhere, Category = Optimizations, meta = (
+		ConsoleVariable = "r.VRS.Support", DisplayName = "Support Hardware Variable Rate Shading",
+		ToolTip = "Allows selectively shading certain portions of the image at lower rates, using one pixel shader invocation to shade multiple pixels. Rates are selected per-material, or in screenspace by enabling a shading rate image generator (such as Contrast Adaptive Shading or Stereo Foveation). Changing this setting requires restarting the editor.",
+		ConfigRestartRequired = true))
+	uint32 bSupportHardwareVariableRateShading : 1;
+
 	/**
 	"When enabled this will skip compiling GPU skin vertex factory shader variants with the assumption that all skinning work will be done via the skin cache."
 	*/
diff --git a/Engine/Source/Runtime/Engine/Private/Materials/MaterialShared.cpp b/Engine/Source/Runtime/Engine/Private/Materials/MaterialShared.cpp
--- a/Engine/Source/Runtime/Engine/Private/Materials/MaterialShared.cpp
+++ b/Engine/Source/Runtime/Engine/Private/Materials/MaterialShared.cpp
@@ -2773,11 +2773,10 @@ void FMaterial::SetupMaterialEnvironment(
 		SET_SHADER_DEFINE(OutEnvironment, USE_STENCIL_LOD_DITHER_DEFAULT, CVar->GetValueOnAnyThread() != 0 ? 1 : 0);
 	}
 
-	if (FDataDrivenShaderPlatformInfo::GetSupportsVariableRateShading(Platform) && 
-		(GRHIAttachmentVariableRateShadingEnabled || GetShadingRate() != MSR_1x1) &&
-		IsVariableRateShadingAllowed())
+	if (FDataDrivenShaderPlatformInfo::GetSupportsVariableRateShading(Platform) && IsVariableRateShadingAllowed())
 	{
-		OutEnvironment.SetCompileArgument(TEXT("USING_VARIABLE_RATE_SHADING"), true);
+		static const auto CVar = IConsoleManager::Get().FindTConsoleVariableDataInt(TEXT("r.VRS.Support"));
+		OutEnvironment.SetCompileArgument(TEXT("USING_VARIABLE_RATE_SHADING"), CVar->GetValueOnAnyThread() != 0 ? 1 : 0);
 	}
 
 	{
diff --git a/Engine/Source/Runtime/Engine/Private/RendererSettings.cpp b/Engine/Source/Runtime/Engine/Private/RendererSettings.cpp
--- a/Engine/Source/Runtime/Engine/Private/RendererSettings.cpp
+++ b/Engine/Source/Runtime/Engine/Private/RendererSettings.cpp
@@ -388,6 +388,15 @@ void URendererSettings::PostEditChangeProperty(FPropertyChangedEvent& PropertyCh
 			CheckForMissingShaderModels();
 		}
 
+		if (PropertyChangedEvent.GetPropertyName() == GET_MEMBER_NAME_CHECKED(URendererSettings, bSupportHardwareVariableRateShading))
+		{
+			if (!bSupportHardwareVariableRateShading)
+			{
+				FoveationLevel = EFixedFoveationLevels::Disabled;
+				UpdateDependentPropertyInConfigFile(this, GET_MEMBER_NAME_CHECKED(URendererSettings, FoveationLevel));
+			}
+		}
+
 		FixAntiAliasingOnShadingPathChange(PropertyChangedEvent);
 	}
 }
diff --git a/Engine/Source/Runtime/RHI/Private/RHI.cpp b/Engine/Source/Runtime/RHI/Private/RHI.cpp
--- a/Engine/Source/Runtime/RHI/Private/RHI.cpp
+++ b/Engine/Source/Runtime/RHI/Private/RHI.cpp
@@ -1134,19 +1134,6 @@ static TAutoConsoleVariable<int32> CVarGPUCrashOnOutOfMemory(
 	ECVF_ReadOnly
 );
 
-static FAutoConsoleVariableRef CVarEnableVariableRateShading(
-	TEXT("r.VRS.Enable"),
-	GRHIVariableRateShadingEnabled,
-	TEXT("Toggle to enable Variable Rate Shading."),
-	ECVF_RenderThreadSafe);
-
-static FAutoConsoleVariableRef CVarEnableAttachmentVariableRateShading(
-	TEXT("r.VRS.EnableImage"),
-	GRHIAttachmentVariableRateShadingEnabled,
-	TEXT("Toggle to enable image-based Variable Rate Shading."),
-	ECVF_RenderThreadSafe);
-
-
 FString GRHIBindlessResourceConfiguration = TEXT("Disabled");
 static FAutoConsoleVariableRef CVarEnableBindlessResources(
 	TEXT("rhi.Bindless.Resources"),
diff --git a/Engine/Source/Runtime/RHI/Public/RHIGlobals.h b/Engine/Source/Runtime/RHI/Public/RHIGlobals.h
--- a/Engine/Source/Runtime/RHI/Public/RHIGlobals.h
+++ b/Engine/Source/Runtime/RHI/Public/RHIGlobals.h
@@ -503,9 +503,11 @@ struct FRHIGlobals
 	struct FVariableRateShading
 	{
 		/** Whether VRS (in all flavors) is currently enabled (separate from whether it's supported/available). */
+		UE_DEPRECATED(5.5, "GRHIGlobals.VariableRateShading.Enabled has been deprecated - please use GRHIGlobals.VariableRateShading.SupportsPipeline or GVRSImageManager.IsPipelineVRSEnabled() instead.")
 		bool Enabled = true;
 
 		/** Whether attachment (image-based) VRS is currently enabled (separate from whether it's supported/available). */
+		UE_DEPRECATED(5.5, "GRHIGlobals.VariableRateShading.AttachmentEnabled has been deprecated - please use GRHIGlobals.VariableRateShading.SupportsAttachment or GVRSImageManager.IsAttachmentVRSEnabled() instead.")
 		bool AttachmentEnabled = true;
 
 		/** Whether or not the RHI can support per-draw Variable Rate Shading. */
@@ -791,8 +793,8 @@ extern RHI_API FRHIGlobals GRHIGlobals;
 #define GRHIIsHDREnabled                                       GRHIGlobals.IsHDREnabled
 #define GRHISupportsHDROutput                                  GRHIGlobals.SupportsHDROutput
 #define GRHIMaxDispatchThreadGroupsPerDimension                GRHIGlobals.MaxDispatchThreadGroupsPerDimension
-#define GRHIVariableRateShadingEnabled                         GRHIGlobals.VariableRateShading.Enabled
-#define GRHIAttachmentVariableRateShadingEnabled               GRHIGlobals.VariableRateShading.AttachmentEnabled
+#define GRHIVariableRateShadingEnabled                         GRHIGlobals.VariableRateShading.Enabled UE_DEPRECATED_MACRO(5.5, "GRHIVariableRateShadingEnabled has been deprecated - please use GRHISupportsPipelineVariableRateShading or GVRSImageManager.IsPipelineVRSEnabled() instead.")
+#define GRHIAttachmentVariableRateShadingEnabled               GRHIGlobals.VariableRateShading.AttachmentEnabled UE_DEPRECATED_MACRO(5.5, "GRHIAttachmentVariableRateShadingEnabled has been deprecated - please use GRHISupportsAttachmentVariableRateShading or GVRSImageManager.IsAttachmentVRSEnabled() instead.")
 #define GRHISupportsPipelineVariableRateShading                GRHIGlobals.VariableRateShading.SupportsPipeline
 #define GRHISupportsLargerVariableRateShadingSizes             GRHIGlobals.VariableRateShading.SupportsLargerSizes
 #define GRHISupportsAttachmentVariableRateShading              GRHIGlobals.VariableRateShading.SupportsAttachment
diff --git a/Engine/Source/Runtime/RenderCore/Private/RenderUtils.cpp b/Engine/Source/Runtime/RenderCore/Private/RenderUtils.cpp
--- a/Engine/Source/Runtime/RenderCore/Private/RenderUtils.cpp
+++ b/Engine/Source/Runtime/RenderCore/Private/RenderUtils.cpp
@@ -2048,3 +2048,14 @@ bool ForwardShadingForcesSkyLightCubemapBlending(const FStaticShaderPlatform Pla
 	static FShaderPlatformCachedIniValue<int32> CVar(TEXT("r.ForwardShading.ForceSkyLightCubemapBlending"));
 	return CVar.Get(Platform) > 0;
 }
+
+bool HardwareVariableRateShadingSupportedByPlatform(EShaderPlatform ShaderPlatform)
+{
+	const static auto SupportVRSCVar = IConsoleManager::Get().FindTConsoleVariableDataInt(TEXT("r.VRS.Support"));
+	if (SupportVRSCVar && SupportVRSCVar->GetValueOnAnyThread() != 0)
+	{
+		return FDataDrivenShaderPlatformInfo::GetSupportsVariableRateShading(ShaderPlatform);
+	}
+
+	return false;
+}
\ No newline at end of file
diff --git a/Engine/Source/Runtime/RenderCore/Private/Shader.cpp b/Engine/Source/Runtime/RenderCore/Private/Shader.cpp
--- a/Engine/Source/Runtime/RenderCore/Private/Shader.cpp
+++ b/Engine/Source/Runtime/RenderCore/Private/Shader.cpp
@@ -2462,12 +2462,6 @@ void ShaderMapAppendKey(EShaderPlatform Platform, FShaderKeyGenerator& KeyGen)
 		KeyGen.Append(TEXT("sdct"));
 	}
 
-	if (FDataDrivenShaderPlatformInfo::GetSupportsVariableRateShading(Platform) && GRHIAttachmentVariableRateShadingEnabled)
-	{
-		KeyGen.AppendSeparator();
-		KeyGen.Append(TEXT("VRS"));
-	}
-
 	{
 		static const auto CVar = IConsoleManager::Get().FindTConsoleVariableDataInt(TEXT("r.GBufferDiffuseSampleOcclusion"));
 		if (CVar && CVar->GetValueOnAnyThread() != 0)
diff --git a/Engine/Source/Runtime/RenderCore/Public/RenderUtils.h b/Engine/Source/Runtime/RenderCore/Public/RenderUtils.h
--- a/Engine/Source/Runtime/RenderCore/Public/RenderUtils.h
+++ b/Engine/Source/Runtime/RenderCore/Public/RenderUtils.h
@@ -618,4 +618,7 @@ FORCEINLINE float LuminanceToEV100(float Luminance)
 FORCEINLINE float Log2ToEV100(float LuminanceMax, float Log2)
 {
 	return Log2 - FMath::Log2(LuminanceMax);
-}
\ No newline at end of file
+}
+
+// Whether or not VRS is supported via r.VRS.Support and the current platform's DDPI
+RENDERCORE_API bool HardwareVariableRateShadingSupportedByPlatform(EShaderPlatform ShaderPlatform);
\ No newline at end of file
diff --git a/Engine/Source/Runtime/Renderer/Private/MeshPassProcessor.cpp b/Engine/Source/Runtime/Renderer/Private/MeshPassProcessor.cpp
--- a/Engine/Source/Runtime/Renderer/Private/MeshPassProcessor.cpp
+++ b/Engine/Source/Runtime/Renderer/Private/MeshPassProcessor.cpp
@@ -21,6 +21,7 @@
 #include "UnrealEngine.h"
 #include "SceneUniformBuffer.h"
 #include "MeshDrawCommandStats.h"
+#include "VariableRateShadingImageManager.h"
 
 FRWLock FGraphicsMinimalPipelineStateId::PersistentIdTableLock;
 FGraphicsMinimalPipelineStateId::PersistentTableType FGraphicsMinimalPipelineStateId::PersistentIdTable;
@@ -2058,6 +2059,28 @@ bool FMeshPassProcessor::ShouldSkipMeshDrawCommand(const FMeshBatch& RESTRICT Me
 	return bSkipMeshDrawCommand;
 }
 
+bool FMeshPassProcessor::PipelineVariableRateShadingEnabled() const
+{
+	return GVRSImageManager.IsPipelineVRSEnabled();
+}
+
+bool FMeshPassProcessor::HardwareVariableRateShadingSupportedByScene() const
+{
+	if (Scene)
+	{
+		return HardwareVariableRateShadingSupportedByPlatform(Scene->GetShaderPlatform());
+	}
+	else if (ViewIfDynamicMeshCommand)
+	{
+		// When applying changes to a material in the editor, we take this pathway
+		return HardwareVariableRateShadingSupportedByPlatform(ViewIfDynamicMeshCommand->GetShaderPlatform());
+	}
+	else
+	{
+		return false;
+	}
+}
+
 FCachedPassMeshDrawListContext::FCachedPassMeshDrawListContext(FScene& InScene)
 	: Scene(InScene)
 	, bUseGPUScene(UseGPUScene(GMaxRHIShaderPlatform, GMaxRHIFeatureLevel))
diff --git a/Engine/Source/Runtime/Renderer/Private/PostProcess/PostProcessTonemap.cpp b/Engine/Source/Runtime/Renderer/Private/PostProcess/PostProcessTonemap.cpp
--- a/Engine/Source/Runtime/Renderer/Private/PostProcess/PostProcessTonemap.cpp
+++ b/Engine/Source/Runtime/Renderer/Private/PostProcess/PostProcessTonemap.cpp
@@ -186,6 +186,8 @@ FDesktopDomain RemapPermutation(FDesktopDomain PermutationVector, ERHIFeatureLev
 		PermutationVector.Set<FTonemapperColorFringeDim>(true);
 	}
 
+	// Luminance output is only used for Contrast Adaptive Shading.
+	// Even if HW VRS support is disabled, we may still need CAS for Nanite SW VRS.
 	if (!FVariableRateShadingImageManager::IsVRSCompatibleWithOutputType(PermutationVector.Get<FTonemapperOutputDeviceDim>()))
 	{
 		PermutationVector.Set<FTonemapperOutputLuminance>(false);
diff --git a/Engine/Source/Runtime/Renderer/Private/RendererScene.cpp b/Engine/Source/Runtime/Renderer/Private/RendererScene.cpp
--- a/Engine/Source/Runtime/Renderer/Private/RendererScene.cpp
+++ b/Engine/Source/Runtime/Renderer/Private/RendererScene.cpp
@@ -96,6 +96,7 @@
 
 #include "VirtualShadowMaps/VirtualShadowMapCacheManager.h"
 #include "Shadows/ShadowScene.h"
+#include "VariableRateShadingImageManager.h"
 
 #if WITH_EDITOR
 #include "Rendering/StaticLightingSystemInterface.h"
@@ -6542,7 +6543,6 @@ void FScene::Update(FRDGBuilder& GraphBuilder, const FUpdateParameters& Paramete
 
 	{
 		SCOPED_NAMED_EVENT(UpdateStaticMeshes, FColor::Emerald);
-		
 		const bool bLastFrameShouldRenderSkylightInBasePass = bCachedShouldRenderSkylightInBasePass;
 		bCachedShouldRenderSkylightInBasePass = ShouldRenderSkylightInBasePass(false);
 
@@ -6554,6 +6554,9 @@ void FScene::Update(FRDGBuilder& GraphBuilder, const FUpdateParameters& Paramete
 			bScenesPrimitivesNeedStaticMeshElementUpdate = true;
 		}
 
+
+		bool bVariableRateShadingToggled = GVRSImageManager.GetNeedStaticMeshUpdate();
+
 		if (bScenesPrimitivesNeedStaticMeshElementUpdate || CachedDefaultBasePassDepthStencilAccess != DefaultBasePassDepthStencilAccess)
 		{
 			// Mark all primitives as needing an update
@@ -6562,6 +6565,26 @@ void FScene::Update(FRDGBuilder& GraphBuilder, const FUpdateParameters& Paramete
 			bScenesPrimitivesNeedStaticMeshElementUpdate = false;
 			CachedDefaultBasePassDepthStencilAccess = DefaultBasePassDepthStencilAccess;
 		}
+		else if (bVariableRateShadingToggled)
+		{
+			// If VRS was just toggled, mark all primitives using non-1x1 per-material shading rates as needing an update
+			for (FPrimitiveSceneInfo* PrimitiveSceneInfo : Primitives)
+			{
+				for (FMeshBatch& MeshBatch : PrimitiveSceneInfo->StaticMeshes)
+				{
+					if (MeshBatch.MaterialRenderProxy->GetMaterialNoFallback(FeatureLevel)->GetShadingRate() != EMaterialShadingRate::MSR_1x1)
+					{
+						PrimitivesNeedingStaticMeshUpdate[PrimitiveSceneInfo->PackedIndex] = true;
+						break;
+					}
+				}
+			}
+		}
+
+		if (bVariableRateShadingToggled)
+		{
+			GVRSImageManager.SetNeedStaticMeshUpdate(false);
+		}
 
 		for (FPrimitiveSceneInfo* PrimitiveSceneInfo : SceneInfosWithStaticDrawListUpdate)
 		{
diff --git a/Engine/Source/Runtime/Renderer/Private/VariableRateShading/VariableRateShadingImageManager.cpp b/Engine/Source/Runtime/Renderer/Private/VariableRateShading/VariableRateShadingImageManager.cpp
--- a/Engine/Source/Runtime/Renderer/Private/VariableRateShading/VariableRateShadingImageManager.cpp
+++ b/Engine/Source/Runtime/Renderer/Private/VariableRateShading/VariableRateShadingImageManager.cpp
@@ -30,6 +30,45 @@ DEFINE_LOG_CATEGORY(LogVRS);
  * Basic CVars
  */
 
+TAutoConsoleVariable<int32> CVarSupportVRS(
+	TEXT("r.VRS.Support"),
+	1,
+	TEXT("Toggles support for hardware Variable Rate Shading. Requires shader recompilation.")
+	TEXT("0: Off, 1: On"),
+	ECVF_ReadOnly);
+
+void CVarEnableVRSCallback(IConsoleVariable* Var)
+{
+	const int32 Value = Var->GetInt();
+
+	// Maintain deprecated globals
+	PRAGMA_DISABLE_DEPRECATION_WARNINGS
+	GRHIGlobals.VariableRateShading.Enabled = (Value != 0) ? GRHISupportsPipelineVariableRateShading : false;
+	GRHIGlobals.VariableRateShading.AttachmentEnabled = (Value != 0) ? GRHISupportsAttachmentVariableRateShading : false;
+	PRAGMA_ENABLE_DEPRECATION_WARNINGS
+
+	// If pipeline VRS is enabled, we need to update static meshes using per-material rates to reflect a change
+	if (GRHISupportsPipelineVariableRateShading)
+	{
+		GVRSImageManager.SetNeedStaticMeshUpdate(true);
+	}
+}
+
+TAutoConsoleVariable<int32> CVarEnableVRS(
+	TEXT("r.VRS.Enable"),
+	0,
+	TEXT("Enables hardware Variable Rate Shading and Shading Rate Image generation (8x8 or 16x16 tile size).")
+	TEXT("0: Off, 1: On"),
+	FConsoleVariableDelegate::CreateStatic(&CVarEnableVRSCallback),
+	ECVF_RenderThreadSafe);
+
+TAutoConsoleVariable<int32> CVarEnableVRSSoftwareImage(
+	TEXT("r.VRS.EnableSoftware"),
+	0,
+	TEXT("Enables software (2x2 tile size) Shading Rate Image generation for use with Nanite Software VRS. Allows generating iamges even when r.VRS.Enable/r.VRS.Support=0 or Tier 2 VRS is unsupported by the hardware.")
+	TEXT("0: Off, 1: On"),
+	ECVF_RenderThreadSafe);
+
 void CVarVRSPreviewCallback(IConsoleVariable* Var)
 {
 	const int32 RequestedPreview = Var->GetInt();
@@ -65,13 +104,6 @@ FAutoConsoleVariableRef CVarVRSDebugForceRate(
 	FConsoleVariableDelegate::CreateStatic(&CVarVRSDebugForceRateCallback),
 	ECVF_RenderThreadSafe);
 
-TAutoConsoleVariable<int32> CVarVRSSoftwareImage(
-	TEXT("r.VRS.EnableSoftware"),
-	0,
-	TEXT("Generate 2x2 tile size software shading rate images when possible for use with nanite CS. Works even when r.VRS.Enable = 0 or Tier 2 VRS is unsupported by the hardware.")
-	TEXT("0: Off, 1: On"),
-	ECVF_RenderThreadSafe);
-
 
 /**
  * Pass Settings
@@ -300,9 +332,17 @@ FVariableRateShadingImageManager::~FVariableRateShadingImageManager() {}
 
 void FVariableRateShadingImageManager::InitRHI(FRHICommandListBase& RHICmdList)
 {
-	if (IsHardwareVRSSupported())
+	if (GRHISupportsPipelineVariableRateShading && GRHISupportsAttachmentVariableRateShading)
+	{
+		UE_LOG(LogVRS, Log, TEXT("Current RHI supports per-draw and screenspace Variable Rate Shading"));
+	}
+	else if (GRHISupportsPipelineVariableRateShading)
 	{
-		UE_LOG(LogVRS, Log, TEXT("Current RHI supports Variable Rate Shading"));
+		UE_LOG(LogVRS, Log, TEXT("Current RHI supports per-draw Variable Rate Shading"));
+	}
+	else if (GRHISupportsAttachmentVariableRateShading)
+	{
+		UE_LOG(LogVRS, Log, TEXT("Current RHI supports screenspace Variable Rate Shading"));
 	}
 	else
 	{
@@ -310,14 +350,9 @@ void FVariableRateShadingImageManager::InitRHI(FRHICommandListBase& RHICmdList)
 	}
 }
 
-void FVariableRateShadingImageManager::ReleaseRHI()
-{
-	GRenderTargetPool.FreeUnusedResources();
-}
-
 bool FVariableRateShadingImageManager::IsHardwareVRSSupported()
 {
-	return GRHISupportsAttachmentVariableRateShading && FDataDrivenShaderPlatformInfo::GetSupportsVariableRateShading(GMaxRHIShaderPlatform);
+	return (GRHISupportsPipelineVariableRateShading || GRHISupportsAttachmentVariableRateShading);
 }
 
 bool FVariableRateShadingImageManager::IsSoftwareVRSSupported()
@@ -325,32 +360,48 @@ bool FVariableRateShadingImageManager::IsSoftwareVRSSupported()
 	return IsFeatureLevelSupported(GMaxRHIShaderPlatform, ERHIFeatureLevel::SM6);
 }
 
-bool FVariableRateShadingImageManager::IsHardwareVRSEnabled()
+void FVariableRateShadingImageManager::ReleaseRHI()
+{
+	GRenderTargetPool.FreeUnusedResources();
+}
+
+bool FVariableRateShadingImageManager::IsPipelineVRSEnabled() const
 {
-	// Currently corresponds to r.VRS.Enable and r.VRS.EnableImage
-	return GRHIVariableRateShadingEnabled && GRHIAttachmentVariableRateShadingEnabled;
+	// This GRHI should be forced to false in DynamicRHI::Init() if r.VRS.Support=0 or bSupportsVariableRateShading=false in the platform's DDPI
+	return GRHISupportsPipelineVariableRateShading && CVarEnableVRS.GetValueOnRenderThread() != 0;
 }
 
-bool FVariableRateShadingImageManager::IsSoftwareVRSEnabled()
+bool FVariableRateShadingImageManager::IsAttachmentVRSEnabled() const
 {
-	return CVarVRSSoftwareImage.GetValueOnRenderThread() > 0;
+	// Ditto
+	return GRHISupportsAttachmentVariableRateShading && CVarEnableVRS.GetValueOnRenderThread() != 0;
 }
 
-bool FVariableRateShadingImageManager::IsVRSEnabledForFrame()
+bool FVariableRateShadingImageManager::IsVRSEnabledForFrame() const
 {
 	return bHardwareVRSEnabledForFrame || bSoftwareVRSEnabledForFrame;
 }
 
-bool FVariableRateShadingImageManager::IsHardwareVRSEnabledForFrame()
+bool FVariableRateShadingImageManager::IsHardwareVRSEnabledForFrame() const
 {
 	return bHardwareVRSEnabledForFrame;
 }
 
-bool FVariableRateShadingImageManager::IsSoftwareVRSEnabledForFrame()
+bool FVariableRateShadingImageManager::IsSoftwareVRSEnabledForFrame() const
 {
 	return bSoftwareVRSEnabledForFrame;
 }
 
+bool FVariableRateShadingImageManager::GetNeedStaticMeshUpdate() const
+{
+	return bNeedStaticMeshUpdate;
+}
+
+void FVariableRateShadingImageManager::SetNeedStaticMeshUpdate(bool bInNeedStaticMeshUpdate)
+{
+	bNeedStaticMeshUpdate = bInNeedStaticMeshUpdate;
+}
+
 static EDisplayOutputFormat GetDisplayOutputFormat(const FViewInfo& View)
 {
 	FTonemapperOutputDeviceParameters Parameters = GetTonemapperOutputDeviceParameters(*View.Family);
@@ -423,7 +474,7 @@ FRDGTextureRef FVariableRateShadingImageManager::GetVariableRateShadingImage(FRD
 	// Use debug rate if provided, otherwise bail if no generators available
 	if (VRSForceRateForFrame >= 0)
 	{
-		return GetForceRateImage(GraphBuilder, *ViewInfo.Family, VRSForceRateForFrame, GetImageTypeFromPassType(PassType), bRequestSoftwareImage);
+		return GetForceRateImage(GraphBuilder, *ViewInfo.Family, VRSForceRateForFrame, bRequestSoftwareImage);
 	}
 
 	if (ActiveGenerators.IsEmpty())
@@ -468,8 +519,11 @@ FRDGTextureRef FVariableRateShadingImageManager::GetVariableRateShadingImage(FRD
 
 void FVariableRateShadingImageManager::PrepareImageBasedVRS(FRDGBuilder& GraphBuilder, const FSceneViewFamily& ViewFamily, const FMinimalSceneTextures& SceneTextures)
 {
-	bHardwareVRSEnabledForFrame = IsHardwareVRSSupported() && IsHardwareVRSEnabled();
-	bSoftwareVRSEnabledForFrame = IsSoftwareVRSSupported() && CVarVRSSoftwareImage.GetValueOnRenderThread() != 0;
+	EShaderPlatform ShaderPlatform = ViewFamily.Scene->GetShaderPlatform();
+
+	bHardwareVRSEnabledForFrame = IsAttachmentVRSEnabled() && HardwareVariableRateShadingSupportedByPlatform(ShaderPlatform); // Additional check is required here for preview levels
+	bSoftwareVRSEnabledForFrame = CVarEnableVRSSoftwareImage.GetValueOnRenderThread() > 0 && IsFeatureLevelSupported(ShaderPlatform, ERHIFeatureLevel::SM6);
+
 	if (!IsVRSEnabledForFrame())
 	{
 		return;
@@ -595,7 +649,7 @@ void FVariableRateShadingImageManager::DrawDebugPreview(FRDGBuilder& GraphBuilde
 			// Use debug rate if provided
 			if (VRSForceRateForFrame >= 0)
 			{
-				PreviewTexture = GetForceRateImage(GraphBuilder, ViewFamily, VRSForceRateForFrame, PreviewImageType, bUseSoftwareImage);
+				PreviewTexture = GetForceRateImage(GraphBuilder, ViewFamily, VRSForceRateForFrame, bUseSoftwareImage);
 			}
 
 			// Otherwise collate debug images
@@ -622,7 +676,7 @@ void FVariableRateShadingImageManager::DrawDebugPreview(FRDGBuilder& GraphBuilde
 				// Generate a dummy 1x1 image if we have no VRS sources
 				if (!PreviewTexture)
 				{
-					PreviewTexture = GetForceRateImage(GraphBuilder, ViewFamily, VRSSR_1x1, EVRSImageType::Full, bUseSoftwareImage);
+					PreviewTexture = GetForceRateImage(GraphBuilder, ViewFamily, VRSSR_1x1, bUseSoftwareImage);
 				}
 			}
 
@@ -738,24 +792,20 @@ FRDGTextureRef FVariableRateShadingImageManager::CombineShadingRateImages(FRDGBu
 
 }
 
-FRDGTextureRef FVariableRateShadingImageManager::GetForceRateImage(FRDGBuilder& GraphBuilder, const FSceneViewFamily& ViewFamily, int RateIndex /* = 0*/, EVRSImageType ImageType /* = EVRSImageType::Full*/, bool bGetSoftwareImage)
+FRDGTextureRef FVariableRateShadingImageManager::GetForceRateImage(FRDGBuilder& GraphBuilder, const FSceneViewFamily& ViewFamily, int RateIndex /* = 0*/, bool bGetSoftwareImage)
 {
 	static const TArray<uint32> ValidShadingRates = { VRSSR_1x1, VRSSR_1x2, VRSSR_2x1, VRSSR_2x2, VRSSR_2x4, VRSSR_4x2, VRSSR_4x4 };
 
-	const int32 NumberOfAvailableRates = GetNumberOfSupportedRates();
-
-	if (RateIndex >= NumberOfAvailableRates)
-	{
-		RateIndex = NumberOfAvailableRates - 1; // Default to maximum shading rate if value exceeds valid rates
-	}
-
-	if (ImageType == EVRSImageType::Disabled)
+	const bool bImageTypeAvailable = bGetSoftwareImage ? bSoftwareVRSEnabledForFrame : bHardwareVRSEnabledForFrame;
+	if (!bImageTypeAvailable)
 	{
-		RateIndex = 0; // Force to minimum shading rate if VRS is disabled for this pass
+		return nullptr;
 	}
 
 	FRDGTextureRef ForceShadingRateTexture = GraphBuilder.CreateTexture(GetSRIDesc(ViewFamily, bGetSoftwareImage), TEXT("ForceShadingRateTexture"));
 	FRDGTextureUAVRef ForceShadingRateUAV = GraphBuilder.CreateUAV(ForceShadingRateTexture);
+
+	RateIndex = FMath::Clamp(RateIndex, 0, GetNumberOfSupportedRates() - 1);
 	AddClearUAVPass(GraphBuilder, ForceShadingRateUAV, ValidShadingRates[RateIndex]);
 
 	return ForceShadingRateTexture;
diff --git a/Engine/Source/Runtime/Renderer/Public/MeshPassProcessor.h b/Engine/Source/Runtime/Renderer/Public/MeshPassProcessor.h
--- a/Engine/Source/Runtime/Renderer/Public/MeshPassProcessor.h
+++ b/Engine/Source/Runtime/Renderer/Public/MeshPassProcessor.h
@@ -2185,6 +2185,9 @@ protected:
 		const FMeshBatch& RESTRICT MeshBatch,
 		const FPrimitiveSceneProxy* RESTRICT PrimitiveSceneProxy
 	) const;
+
+	RENDERER_API bool PipelineVariableRateShadingEnabled() const;
+	RENDERER_API bool HardwareVariableRateShadingSupportedByScene() const;
 };
 
 #if PSO_PRECACHING_VALIDATE
diff --git a/Engine/Source/Runtime/Renderer/Public/MeshPassProcessor.inl b/Engine/Source/Runtime/Renderer/Public/MeshPassProcessor.inl
--- a/Engine/Source/Runtime/Renderer/Public/MeshPassProcessor.inl
+++ b/Engine/Source/Runtime/Renderer/Public/MeshPassProcessor.inl
@@ -12,34 +12,36 @@ MeshPassProcessor.inl:
 #include "RHIStaticStates.h"
 #include "RenderGraphBuilder.h"
 #include "PSOPrecacheValidation.h"
+#include "VariableRateShadingImageManager.h"
+
 
 static EVRSShadingRate GetShadingRateFromMaterial(EMaterialShadingRate MaterialShadingRate)
 {
-	if (GRHISupportsPipelineVariableRateShading && GRHIVariableRateShadingEnabled)
+	switch (MaterialShadingRate)
 	{
-		switch (MaterialShadingRate)
-		{
-		case MSR_1x2:
-			return EVRSShadingRate::VRSSR_1x2;
-		case MSR_2x1:
-			return EVRSShadingRate::VRSSR_2x1;
-		case MSR_2x2:
-			return EVRSShadingRate::VRSSR_2x2;
-		}
+	case MSR_1x1:
+		return EVRSShadingRate::VRSSR_1x1;
+	case MSR_1x2:
+		return EVRSShadingRate::VRSSR_1x2;
+	case MSR_2x1:
+		return EVRSShadingRate::VRSSR_2x1;
+	case MSR_2x2:
+		return EVRSShadingRate::VRSSR_2x2;
+	}
 
-		if (GRHISupportsLargerVariableRateShadingSizes)
+	if (GRHISupportsLargerVariableRateShadingSizes)
+	{
+		switch (MaterialShadingRate)
 		{
-			switch (MaterialShadingRate)
-			{
-			case MSR_4x2:
-				return EVRSShadingRate::VRSSR_4x2;
-			case MSR_2x4:
-				return EVRSShadingRate::VRSSR_2x4;
-			case MSR_4x4:
-				return EVRSShadingRate::VRSSR_4x4;
-			}
+		case MSR_4x2:
+			return EVRSShadingRate::VRSSR_4x2;
+		case MSR_2x4:
+			return EVRSShadingRate::VRSSR_2x4;
+		case MSR_4x4:
+			return EVRSShadingRate::VRSSR_4x4;
 		}
 	}
+
 	return EVRSShadingRate::VRSSR_1x1;
 }
 
@@ -102,8 +104,8 @@ void FMeshPassProcessor::BuildMeshDrawCommands(
 
 	PipelineState.BlendState = DrawRenderState.GetBlendState();
 	PipelineState.DepthStencilState = DrawRenderState.GetDepthStencilState();
-	PipelineState.DrawShadingRate = GetShadingRateFromMaterial(MaterialResource.GetShadingRate());
-	PipelineState.bAllowVariableRateShading = MaterialResource.IsVariableRateShadingAllowed();
+	PipelineState.DrawShadingRate = PipelineVariableRateShadingEnabled() ? GetShadingRateFromMaterial(MaterialResource.GetShadingRate()) : VRSSR_1x1;
+	PipelineState.bAllowVariableRateShading = MaterialResource.IsVariableRateShadingAllowed() && HardwareVariableRateShadingSupportedByScene();
 
 	// PSO Precache hash only needed when PSO precaching is enabled
 	if (PipelineStateCache::IsPSOPrecachingEnabled())
diff --git a/Engine/Source/Runtime/Renderer/Public/VariableRateShadingImageManager.h b/Engine/Source/Runtime/Renderer/Public/VariableRateShadingImageManager.h
--- a/Engine/Source/Runtime/Renderer/Public/VariableRateShadingImageManager.h
+++ b/Engine/Source/Runtime/Renderer/Public/VariableRateShadingImageManager.h
@@ -103,12 +103,21 @@ public:
 	static bool IsHardwareVRSSupported();
 	static bool IsSoftwareVRSSupported();
 
-	bool IsHardwareVRSEnabled();
-	bool IsSoftwareVRSEnabled();
+	RENDERER_API bool IsPipelineVRSEnabled() const;
+	RENDERER_API bool IsAttachmentVRSEnabled() const;
 
-	bool IsVRSEnabledForFrame();
-	bool IsHardwareVRSEnabledForFrame();
-	bool IsSoftwareVRSEnabledForFrame();
+	RENDERER_API bool IsVRSEnabledForFrame()  const;
+	RENDERER_API bool IsHardwareVRSEnabledForFrame() const;
+	RENDERER_API bool IsSoftwareVRSEnabledForFrame() const;
+
+	UE_DEPRECATED(5.5, "IsHardwareVRSEnabled() is deprecated, please use IsHardwareVRSEnabledForFrame() instead.")
+	bool IsHardwareVRSEnabled() { return IsHardwareVRSEnabledForFrame(); };
+
+	UE_DEPRECATED(5.5, "IsSoftwareVRSEnabled() is deprecated, please use IsSoftwareVRSEnabledForFrame() instead.")
+	bool IsSoftwareVRSEnabled() { return IsSoftwareVRSEnabledForFrame(); };
+
+	bool GetNeedStaticMeshUpdate() const;
+	void SetNeedStaticMeshUpdate(bool bInNeedStaticMeshUpdate);
 
 	static bool IsVRSCompatibleWithView(const FViewInfo& View);
 	static bool IsVRSCompatibleWithOutputType(const EDisplayOutputFormat& DisplayOutputFormat);
@@ -131,13 +140,19 @@ private:
 	FRWLock	GeneratorsMutex;
 
 	FRDGTextureRef CombineShadingRateImages(FRDGBuilder& GraphBuilder, const FSceneViewFamily& ViewFamily, TArray<FRDGTextureRef> Sources);
-	FRDGTextureRef GetForceRateImage(FRDGBuilder& GraphBuilder, const FSceneViewFamily& ViewFamily, int RateIndex = 0, EVRSImageType ImageType = EVRSImageType::Full, bool bGetSoftwareImage = false);
+	FRDGTextureRef GetForceRateImage(FRDGBuilder& GraphBuilder, const FSceneViewFamily& ViewFamily, int RateIndex = 0, bool bGetSoftwareImage = false);
+
+	UE_DEPRECATED(5.5, "This version of GetForceRateImage is deprecated, please use GetForceRateImage(FRDGBuilder& GraphBuilder, const FSceneViewFamily& ViewFamily, int RateIndex, bool bGetSoftwareImage) instead.")
+	FRDGTextureRef GetForceRateImage(FRDGBuilder& GraphBuilder, const FSceneViewFamily& ViewFamily, int RateIndex = 0,
+		EVRSImageType ImageType = EVRSImageType::Full, bool bGetSoftwareImage = false) { return GetForceRateImage(GraphBuilder, ViewFamily, RateIndex, bGetSoftwareImage); };
 
 	EVRSImageType GetImageTypeFromPassType(EVRSPassType PassType);
 
 	bool bHardwareVRSEnabledForFrame = false;
 	bool bSoftwareVRSEnabledForFrame = false;
 	int32 VRSForceRateForFrame = -1;
+
+	bool bNeedStaticMeshUpdate = false;
 };
 
 ENUM_CLASS_FLAGS(FVariableRateShadingImageManager::EVRSSourceType);
diff --git a/Engine/Source/Runtime/VulkanRHI/Private/VulkanExtensions.cpp b/Engine/Source/Runtime/VulkanRHI/Private/VulkanExtensions.cpp
--- a/Engine/Source/Runtime/VulkanRHI/Private/VulkanExtensions.cpp
+++ b/Engine/Source/Runtime/VulkanRHI/Private/VulkanExtensions.cpp
@@ -502,9 +502,10 @@ public:
 	FVulkanKHRFragmentShadingRateExtension(FVulkanDevice* InDevice)
 		: FVulkanDeviceExtension(InDevice, VK_KHR_FRAGMENT_SHADING_RATE_EXTENSION_NAME, VULKAN_EXTENSION_ENABLED)
 	{
-		int32 VRSFormatPreference = GVulkanVariableRateShadingFormatCVar->GetInt();
-		bEnabledInCode = bEnabledInCode && GRHIVariableRateShadingEnabled;
+		bEnabledInCode &= HardwareVariableRateShadingSupportedByPlatform(GMaxRHIShaderPlatform);
+
 		// FSR should be enabled even if FDM is preferred because it could be not available.
+		int32 VRSFormatPreference = GVulkanVariableRateShadingFormatCVar->GetInt();
 		bEnabledInCode &= (VRSFormatPreference <= (uint8) EVulkanVariableRateShadingPreference::RequireFSR || VRSFormatPreference == (uint8)EVulkanVariableRateShadingPreference::PreferFDM);
 	}
 
@@ -586,9 +587,10 @@ public:
 	FVulkanEXTFragmentDensityMapExtension(FVulkanDevice* InDevice)
 		: FVulkanDeviceExtension(InDevice, VK_EXT_FRAGMENT_DENSITY_MAP_EXTENSION_NAME, VULKAN_EXTENSION_ENABLED)
 	{
-		int32 VRSFormatPreference = GVulkanVariableRateShadingFormatCVar->GetInt();
-		bEnabledInCode = bEnabledInCode && GRHIVariableRateShadingEnabled;
+		bEnabledInCode &= HardwareVariableRateShadingSupportedByPlatform(GMaxRHIShaderPlatform);
+
 		// FDM should be enabled even if the preferred choice is FSR because that might not be available.
+		int32 VRSFormatPreference = GVulkanVariableRateShadingFormatCVar->GetInt();
 		bEnabledInCode &= (VRSFormatPreference >= (uint8) EVulkanVariableRateShadingPreference::PreferFDM || VRSFormatPreference == (uint8) EVulkanVariableRateShadingPreference::PreferFSR);
 	}
 
@@ -660,7 +662,7 @@ public:
 	FVulkanEXTFragmentDensityMap2Extension(FVulkanDevice* InDevice)
 		: FVulkanDeviceExtension(InDevice, VK_EXT_FRAGMENT_DENSITY_MAP_2_EXTENSION_NAME, VULKAN_EXTENSION_ENABLED)
 	{
-		bEnabledInCode = bEnabledInCode && GRHIVariableRateShadingEnabled;
+		bEnabledInCode &= HardwareVariableRateShadingSupportedByPlatform(GMaxRHIShaderPlatform);
 	}
 
 	virtual void PrePhysicalDeviceFeatures(VkPhysicalDeviceFeatures2KHR& PhysicalDeviceFeatures2) override final
diff --git a/Engine/Source/Runtime/VulkanRHI/Private/VulkanPipeline.cpp b/Engine/Source/Runtime/VulkanRHI/Private/VulkanPipeline.cpp
--- a/Engine/Source/Runtime/VulkanRHI/Private/VulkanPipeline.cpp
+++ b/Engine/Source/Runtime/VulkanRHI/Private/VulkanPipeline.cpp
@@ -1350,8 +1350,10 @@ bool FVulkanPipelineStateCacheManager::CreateGfxPipelineFromEntry(FVulkanRHIGrap
 
 	PipelineInfo.pDynamicState = &DynamicState;
 
+	const bool bUsingVariableRateShading = PSO->Desc.ShadingRate != EVRSShadingRate::VRSSR_1x1 || PSO->RenderPass->Layout.bHasFragmentDensityAttachment;
+
 	VkPipelineFragmentShadingRateStateCreateInfoKHR PipelineFragmentShadingRate;
-	if (GRHISupportsPipelineVariableRateShading && GRHIVariableRateShadingEnabled && GRHIVariableRateShadingImageDataType == VRSImage_Palette)
+	if (GRHISupportsPipelineVariableRateShading && GRHIVariableRateShadingImageDataType == VRSImage_Palette && bUsingVariableRateShading)
 	{
 		const VkExtent2D FragmentSize = Device->GetBestMatchedFragmentSize(PSO->Desc.ShadingRate);
 		VkFragmentShadingRateCombinerOpKHR PipelineToPrimitiveCombinerOperation = FragmentCombinerOpMap[(uint8)PSO->Desc.Combiner];
@@ -1921,7 +1923,7 @@ void FVulkanPipelineStateCacheManager::CreateGfxEntry(const FGraphicsPipelineSta
 	OutGfxEntry->RenderTargets.ReadFrom(RTLayout);
 
 	// Shading rate:
-	OutGfxEntry->ShadingRate = PSOInitializer.ShadingRate;
+	OutGfxEntry->ShadingRate = PSOInitializer.bAllowVariableRateShading ? PSOInitializer.ShadingRate : EVRSShadingRate::VRSSR_1x1;
 	OutGfxEntry->Combiner = EVRSRateCombiner::VRSRB_Max;		// @todo: This needs to be specified twice; from pipeline-to-primitive, and from primitive-to-attachment. 
 																// We don't have per-primitive VRS so that should just be hard-coded to "passthrough" until this is supported; but we should expose 
 																// this setting in the material properies, especially since there's some materials that don't play nicely with 
diff --git a/Engine/Source/Runtime/VulkanRHI/Private/VulkanRHI.cpp b/Engine/Source/Runtime/VulkanRHI/Private/VulkanRHI.cpp
--- a/Engine/Source/Runtime/VulkanRHI/Private/VulkanRHI.cpp
+++ b/Engine/Source/Runtime/VulkanRHI/Private/VulkanRHI.cpp
@@ -89,15 +89,6 @@ static FAutoConsoleVariableRef GCVarEnableTransientResourceAllocator(
 	ECVF_ReadOnly
 );
 
-int32 GVulkanAllowVariableRateShading = 1;
-static FAutoConsoleVariableRef CVarVulkanVariableRateShading(
-	TEXT("r.Vulkan.AllowVariableRateShading"),
-	GVulkanAllowVariableRateShading,
-	TEXT("0 to disable variable rate shading")
-	TEXT("1 to allow use of variable rate shading if available (default)"),
-	ECVF_ReadOnly
-);
-
 static TAutoConsoleVariable<bool> CVarAllowVulkanPSOPrecache(
 	TEXT("r.Vulkan.AllowPSOPrecaching"),
 	false,
@@ -519,7 +510,6 @@ FVulkanDynamicRHI::FVulkanDynamicRHI()
 	GRHISupportsMultithreading = true;
 	GRHISupportsMultithreadedResources = true;
 	GRHISupportsPipelineFileCache = true;
-	GRHIVariableRateShadingEnabled &= (GVulkanAllowVariableRateShading != 0); // before extensions setup
 	GRHITransitionPrivateData_SizeInBytes = sizeof(FVulkanPipelineBarrier);
 	GRHITransitionPrivateData_AlignInBytes = alignof(FVulkanPipelineBarrier);
 	GConfig->GetInt(TEXT("TextureStreaming"), TEXT("PoolSizeVRAMPercentage"), GPoolSizeVRAMPercentage, GEngineIni);
@@ -2075,7 +2065,8 @@ uint64 FVulkanDynamicRHI::RHIComputeStatePrecachePSOHash(const FGraphicsPipeline
 		uint32 PrimitiveType : 8;
 		uint32 bDepthBounds : 1;
 		uint32 bHasFragmentDensityAttachment : 1;
-		uint32 Unused : 6;
+		uint32 bAllowVariableRateShading : 1;
+		uint32 Unused : 5;
 	} HashKey;
 
 	FMemory::Memzero(&HashKey, sizeof(FHashKey));
@@ -2114,6 +2105,7 @@ uint64 FVulkanDynamicRHI::RHIComputeStatePrecachePSOHash(const FGraphicsPipeline
 	HashKey.PrimitiveType = Initializer.PrimitiveType;
 	HashKey.bDepthBounds = Initializer.bDepthBounds;
 	HashKey.bHasFragmentDensityAttachment = Initializer.bHasFragmentDensityAttachment;
+	HashKey.bAllowVariableRateShading = Initializer.bAllowVariableRateShading;
 
 	uint64 PrecachePSOHash = CityHash64((const char*)&HashKey, sizeof(FHashKey));
 
@@ -2152,6 +2144,7 @@ uint64 FVulkanDynamicRHI::RHIComputePrecachePSOHash(const FGraphicsPipelineState
 		EVRSShadingRate					ShadingRate;
 		bool							bDepthBounds;
 		bool							bHasFragmentDensityAttachment;
+		bool							bAllowVariableRateShading;
 	} HashKey;
 
 	FMemory::Memzero(&HashKey, sizeof(FNonStateHashKey));
@@ -2172,6 +2165,7 @@ uint64 FVulkanDynamicRHI::RHIComputePrecachePSOHash(const FGraphicsPipelineState
 	HashKey.ShadingRate = Initializer.ShadingRate;
 	HashKey.bDepthBounds = Initializer.bDepthBounds;
 	HashKey.bHasFragmentDensityAttachment = Initializer.bHasFragmentDensityAttachment;
+	HashKey.bAllowVariableRateShading = Initializer.bAllowVariableRateShading;
 
 	// TODO: check if any RT flags actually affect PSO in VK
 	for (ETextureCreateFlags& Flags : HashKey.RenderTargetFlags)
@@ -2192,6 +2186,7 @@ bool FVulkanDynamicRHI::RHIMatchPrecachePSOInitializers(const FGraphicsPipelineS
 		LHS.MultiViewCount != RHS.MultiViewCount ||
 		LHS.ShadingRate != RHS.ShadingRate ||
 		LHS.bHasFragmentDensityAttachment != RHS.bHasFragmentDensityAttachment ||
+		LHS.bAllowVariableRateShading != RHS.bAllowVariableRateShading ||
 		LHS.RenderTargetsEnabled != RHS.RenderTargetsEnabled ||
 		LHS.RenderTargetFormats != RHS.RenderTargetFormats ||
 		!FGraphicsPipelineStateInitializer::RelevantRenderTargetFlagsEqual(LHS.RenderTargetFlags, RHS.RenderTargetFlags) ||
diff --git a/Engine/Source/Runtime/VulkanRHI/Private/VulkanRenderTarget.cpp b/Engine/Source/Runtime/VulkanRHI/Private/VulkanRenderTarget.cpp
--- a/Engine/Source/Runtime/VulkanRHI/Private/VulkanRenderTarget.cpp
+++ b/Engine/Source/Runtime/VulkanRHI/Private/VulkanRenderTarget.cpp
@@ -939,7 +939,7 @@ FVulkanRenderTargetLayout::FVulkanRenderTargetLayout(FVulkanDevice& InDevice, co
 		}
 	}
 
-	if (GRHISupportsAttachmentVariableRateShading && GRHIVariableRateShadingEnabled && GRHIAttachmentVariableRateShadingEnabled && RTInfo.ShadingRateTexture)
+	if (GRHISupportsAttachmentVariableRateShading && RTInfo.ShadingRateTexture)
 	{
 		FVulkanTexture* Texture = ResourceCast(RTInfo.ShadingRateTexture);
 		check(Texture->GetFormat() == GRHIVariableRateShadingImageFormat);
@@ -1194,7 +1194,7 @@ FVulkanRenderTargetLayout::FVulkanRenderTargetLayout(FVulkanDevice& InDevice, co
 		Extent.Extent3D.depth = 1;
 	}
 
-	if (GRHISupportsAttachmentVariableRateShading && GRHIVariableRateShadingEnabled && GRHIAttachmentVariableRateShadingEnabled && RPInfo.ShadingRateTexture)
+	if (GRHISupportsAttachmentVariableRateShading && RPInfo.ShadingRateTexture)
 	{
 		FVulkanTexture* Texture = ResourceCast(RPInfo.ShadingRateTexture);
 		check(Texture->GetFormat() == GRHIVariableRateShadingImageFormat);
@@ -1364,7 +1364,7 @@ FVulkanRenderTargetLayout::FVulkanRenderTargetLayout(const FGraphicsPipelineStat
 		bHasDepthStencil = true;
 	}
 
-	if (Initializer.bHasFragmentDensityAttachment)
+	if (Initializer.bHasFragmentDensityAttachment && Initializer.bAllowVariableRateShading)
 	{
 		VkAttachmentDescription& CurrDesc = Desc[NumAttachmentDescriptions];
 		FMemory::Memzero(CurrDesc);
diff --git a/Engine/Source/Runtime/VulkanRHI/Private/VulkanRenderpass.h b/Engine/Source/Runtime/VulkanRHI/Private/VulkanRenderpass.h
--- a/Engine/Source/Runtime/VulkanRHI/Private/VulkanRenderpass.h
+++ b/Engine/Source/Runtime/VulkanRHI/Private/VulkanRenderpass.h
@@ -458,9 +458,7 @@ public:
 		uint32 MultiviewMask = (0b1 << RTLayout.GetMultiViewCount()) - 1;
 
 		const bool bDeferredShadingSubpass = RTLayout.GetSubpassHint() == ESubpassHint::DeferredShadingSubpass;
-		const bool bApplyFragmentShadingRate =  GRHISupportsAttachmentVariableRateShading 
-												&& GRHIVariableRateShadingEnabled 
-												&& GRHIAttachmentVariableRateShadingEnabled 
+		const bool bApplyFragmentShadingRate =  GRHISupportsAttachmentVariableRateShading
 												&& RTLayout.GetFragmentDensityAttachmentReference() != nullptr
 												&& Device.GetOptionalExtensions().HasKHRFragmentShadingRate 
 												&& Device.GetOptionalExtensionProperties().FragmentShadingRateFeatures.attachmentFragmentShadingRate == VK_TRUE;
diff --git a/Engine/Source/Runtime/VulkanRHI/Private/VulkanViewport.cpp b/Engine/Source/Runtime/VulkanRHI/Private/VulkanViewport.cpp
--- a/Engine/Source/Runtime/VulkanRHI/Private/VulkanViewport.cpp
+++ b/Engine/Source/Runtime/VulkanRHI/Private/VulkanViewport.cpp
@@ -534,7 +534,7 @@ FVulkanFramebuffer::FVulkanFramebuffer(FVulkanDevice& Device, const FRHISetRende
 		}
 	}
 
-	if (GRHISupportsAttachmentVariableRateShading && GRHIVariableRateShadingEnabled && GRHIAttachmentVariableRateShadingEnabled && RTLayout.GetHasFragmentDensityAttachment())
+	if (GRHISupportsAttachmentVariableRateShading && RTLayout.GetHasFragmentDensityAttachment())
 	{
 		FVulkanTexture* Texture = ResourceCast(InRTInfo.ShadingRateTexture);
 		FragmentDensityImage = Texture->Image;
