Added CVars to enable inset shadows on the mobile renderer and to lower inset shadows resolution.

diff --git a/Engine/Source/Runtime/Renderer/Private/ShadowSetup.cpp b/Engine/Source/Runtime/Renderer/Private/ShadowSetup.cpp
--- a/Engine/Source/Runtime/Renderer/Private/ShadowSetup.cpp
+++ b/Engine/Source/Runtime/Renderer/Private/ShadowSetup.cpp
@@ -50,6 +50,8 @@ static const int32 NumCubeShadowDepthSurfaces = 5;
 /** Number of surfaces used for translucent shadows. */
 static const int32 NumTranslucencyShadowSurfaces = 2;
 
+static bool bGMobileInsetShadows = false;
+
 static float GMinScreenRadiusForShadowCaster = 0.01f;
 static FAutoConsoleVariableRef CVarMinScreenRadiusForShadowCaster(
 	TEXT("r.Shadow.RadiusThreshold"),
@@ -162,6 +164,13 @@ static TAutoConsoleVariable<int32> CVarCachePreshadows(
 	ECVF_RenderThreadSafe
 	);
 
+static TAutoConsoleVariable<int32> CVarInsetDownscaleFactor(
+	TEXT("r.Shadow.InsetDownscaleFactor"),
+	1,
+	TEXT("Use a lower resolution for inset shadows (OriginalResolution / DownscaleFactor)"),
+	ECVF_RenderThreadSafe
+);
+
 /**
  * NOTE: This flag is intended to be kept only as long as deemed neccessary to be sure that no artifacts were introduced.
  *       This allows a quick hot-fix to disable the change if need be.
@@ -352,6 +361,13 @@ FAutoConsoleCommand CmdDumpShadowDumpSetup(
 	);
 #endif // !UE_BUILD_SHIPPING
 
+static TAutoConsoleVariable<bool> CVarMobileInsetShadows(
+	TEXT("r.Mobile.SupportInsetShadows"),
+	false,
+	TEXT("Enables inset shadows on mobile. Currently not supported on platforms not using ShadowMaskTexture. Default = false."),
+	ECVF_RenderThreadSafe
+);
+
 /** Whether to round the shadow map up to power of two on mobile platform. */
 static TAutoConsoleVariable<int32> CVarMobileShadowmapRoundUpToPowerOfTwo(
 	TEXT("r.Mobile.ShadowmapRoundUpToPowerOfTwo"),
@@ -3304,8 +3320,9 @@ void FSceneRenderer::CreatePerObjectProjectedShadow(
 
 	// Shadowing constants.
 	
+	const uint32 InsetDownscaleFactor = FMath::Max<int32>(1, CVarInsetDownscaleFactor.GetValueOnRenderThread());
 	const uint32 MaxShadowResolutionSetting = GetCachedScalabilityCVars().MaxShadowResolution;
-	const FIntPoint ShadowBufferResolution = GetShadowDepthTextureResolution(FeatureLevel);
+	const FIntPoint ShadowBufferResolution = GetShadowDepthTextureResolution(FeatureLevel) / InsetDownscaleFactor;
 	const uint32 MaxShadowResolution = FMath::Min<int32>(MaxShadowResolutionSetting, ShadowBufferResolution.X) - SHADOW_BORDER * 2;
 	const uint32 MaxShadowResolutionY = FMath::Min<int32>(MaxShadowResolutionSetting, ShadowBufferResolution.Y) - SHADOW_BORDER * 2;
 	const uint32 MinShadowResolution     = FMath::Max<int32>(0, CVarMinShadowResolution.GetValueOnRenderThread());
@@ -6170,7 +6187,19 @@ void FSceneRenderer::CreateDynamicShadows(FDynamicShadowsTaskData& TaskData)
 								AddViewDependentWholeSceneShadowsForView(ViewDependentWholeSceneShadows, ViewDependentWholeSceneShadowsThatNeedCulling, VisibleLightInfo, *LightSceneInfo, GetCachedShadowMapsSize(), NumCSMCachesUpdatedThisFrame);
 							}
 
-							if (!bMobile || (LightSceneInfo->Proxy->CastsModulatedShadows() && !LightSceneInfo->Proxy->UseCSMForDynamicObjects() && LightSceneInfo->Proxy->HasStaticShadowing()))
+#if !(UE_BUILD_SHIPPING || UE_BUILD_TEST)
+							bool bMobileInsetShadowsCVarValueChanged = bGMobileInsetShadows != CVarMobileInsetShadows.GetValueOnRenderThread();
+							if (bMobile && CVarMobileInsetShadows.GetValueOnRenderThread() && bMobileInsetShadowsCVarValueChanged)
+							{
+								if (!MobileUsesShadowMaskTexture(ShaderPlatform))
+								{
+									UE_LOG(LogRenderer, Warning, TEXT("r.Mobile.SupportInsetShadows is enabled on a platform not using ShadowMaskTexture, so the inset shadows projections won't be rendered."));
+								}
+							}
+#endif
+							bGMobileInsetShadows = CVarMobileInsetShadows.GetValueOnRenderThread();
+
+							if (!bMobile || bGMobileInsetShadows || (LightSceneInfo->Proxy->CastsModulatedShadows() && !LightSceneInfo->Proxy->UseCSMForDynamicObjects() && LightSceneInfo->Proxy->HasStaticShadowing()))
 							{
 								const TArray<FLightPrimitiveInteraction*>* InteractionShadowPrimitives = LightSceneInfo->GetInteractionShadowPrimitives();
 
