Fix incorrect buffer type being bound to FRayTracingLightGrid::SceneLights when RT light grid isn't used.

This would produce an error like "Buffer type for SRV DefaultBuffer must match buffer type declared in the shader" at project startup when running with -rhivalidation.

#rnx
#tests Ran EngineTest with -rhivalidation and confirmed that the error no longer occurs.

diff --git a/Engine/Source/Runtime/Renderer/Private/RayTracing/RayTracingLighting.cpp b/Engine/Source/Runtime/Renderer/Private/RayTracing/RayTracingLighting.cpp
--- a/Engine/Source/Runtime/Renderer/Private/RayTracing/RayTracingLighting.cpp
+++ b/Engine/Source/Runtime/Renderer/Private/RayTracing/RayTracingLighting.cpp
@@ -361,7 +361,7 @@ TRDGUniformBufferRef<FRayTracingLightGrid> CreateRayTracingLightData(
 		LightGridParameters->SceneInfiniteLightCount = 0;
 		LightGridParameters->SceneLightsTranslatedBoundMin = FVector3f::ZeroVector;
 		LightGridParameters->SceneLightsTranslatedBoundMax = FVector3f::ZeroVector;
-		LightGridParameters->SceneLights = GraphBuilder.CreateSRV(GSystemTextures.GetDefaultBuffer(GraphBuilder, sizeof(uint32), 0u), PF_R32_UINT);
+		LightGridParameters->SceneLights = GraphBuilder.CreateSRV(GSystemTextures.GetDefaultStructuredBuffer(GraphBuilder, sizeof(FRTLightingData), 0u));
 		LightGridParameters->LightGrid = GSystemTextures.GetDefaultTexture2D(GraphBuilder, PF_R32_UINT, 0u);
 		LightGridParameters->LightGridData = GraphBuilder.CreateSRV(GSystemTextures.GetDefaultStructuredBuffer(GraphBuilder, sizeof(uint32), 0u), PF_R32_UINT);
 		LightGridParameters->LightGridResolution = 0;
