Scene Capture:  Upgrades to performance and robustness for base pass scene captures (captures that output albedo or normal gbuffers).
* Experimental CVar (r.SceneCapture.BasePassOptimization) to enable optimized base pass rendering, which skips unnecessary scene rendering work (ray tracing, lighting, shadows, sky, atmosphere, etc).  Over 2x CPU and 3x GPU improvement.  Disabled by default.
* Support for bRenderInMainRenderer flag on base pass scene captures, which embeds the scene capture into the main view as a CustomRenderPass.  9x CPU and 5x GPU improvement over original baseline in simple test scene.  Previously, the flag only worked for depth only scene captures -- base pass CustomRenderPasses required writing code to render.
* EngineShowFlags work correctly.  Since these are referenced from FSceneViewFamily, a separate copy of the structure is required for the CustomRenderPass.  The FSceneTextures (also stored in the view family) are shared via pointer reference.  The most noticeable artifact this fixes is editor objects showing up in the custom render pass code path, but not the scene capture code path, but it also makes user flag settings in the scene capture component work.
* Single layer water rendering works for optimized scene captures and custom render passes.
* Mip map generation works for custom render pass scene captures.  With these three fixes, full visual parity is achieved with regular scene captures in a test scene.
* CustomRenderPass manually transitions the output render target to SRV in the RDG.  Render targets can be used by other rendering in the same graph, so the transition should happen immediately -- without this, RDG can delay the transition to the end of graph execution, which can cause rendering bugs.
* CustomRenderPasses reset the scene texture uniform buffer state after running.  Without this, it's possible downstream rendering would be allowed to reference GBuffer textures in the GBuffer pass itself (if referenced from a material), rather than dummy versions they should see at that point in the frame, which could cause rendering bugs.  Other logic could also be confused thinking GBuffers are initialized, when they contain rendering from a CustomRenderPass, and not the view family.

There remains a major discrepancy in base pass scene captures versus custom render passes in the mobile renderer.  The mobile renderer doesn't support base pass scene captures in general, because there's no gbuffer pass, given that it's forward rendering.  So it treats the "BasePass" setting the same as "SceneColor", rendering the whole scene, with lighting, shadows, translucent passes, and whatever else.  By contrast, custom render passes skip shadows and translucency (lights still render, with dummy shadow buffers that lead to no shadowing).  I think it would make sense for "BasePass" to skip shadows and translucency on the mobile renderer, which is closer visually to a base pass renderer, and would then match custom render pass behavior.  From there, you could get closer to the deferred shading result by replacing the lighting with a user specified flat ambient light value, or even have a dynamic flag in the shader (perhaps gated on a project setting, since there are minor perf implications) which switches the shader output to what the deferred renderer would produce for the albedo channel, producing identical results.  But that functionality is well beyond the scope of this CL.

diff --git a/Engine/Source/Runtime/Engine/Classes/Components/SceneCaptureComponent2D.h b/Engine/Source/Runtime/Engine/Classes/Components/SceneCaptureComponent2D.h
--- a/Engine/Source/Runtime/Engine/Classes/Components/SceneCaptureComponent2D.h
+++ b/Engine/Source/Runtime/Engine/Classes/Components/SceneCaptureComponent2D.h
@@ -113,8 +113,8 @@ public:
 	UPROPERTY(EditAnywhere, BlueprintReadWrite, AdvancedDisplay, Category=SceneCapture)
 	FVector ClipPlaneNormal;
 	
-	/** Render scene capture as additional render passes of the main renderer rather than as an independent renderer. Can only apply to scene depth and device depth modes. */
-	UPROPERTY(EditAnywhere, BlueprintReadWrite, AdvancedDisplay, Category = SceneCapture, meta = (EditCondition = "CaptureSource == ESceneCaptureSource::SCS_SceneDepth || CaptureSource == ESceneCaptureSource::SCS_DeviceDepth"))
+	/** Render scene capture as additional render passes of the main renderer rather than as an independent renderer. Applies to scene depth, device depth, base color, and normal modes. */
+	UPROPERTY(EditAnywhere, BlueprintReadWrite, AdvancedDisplay, Category = SceneCapture, meta = (EditCondition = "CaptureSource == ESceneCaptureSource::SCS_SceneDepth || CaptureSource == ESceneCaptureSource::SCS_DeviceDepth || CaptureSource == ESceneCaptureSource::SCS_BaseColor || CaptureSource == ESceneCaptureSource::SCS_Normal"))
 	bool bRenderInMainRenderer = false;
 
 	/** 
diff --git a/Engine/Source/Runtime/Engine/Private/Rendering/CustomRenderPass.cpp b/Engine/Source/Runtime/Engine/Private/Rendering/CustomRenderPass.cpp
--- a/Engine/Source/Runtime/Engine/Private/Rendering/CustomRenderPass.cpp
+++ b/Engine/Source/Runtime/Engine/Private/Rendering/CustomRenderPass.cpp
@@ -67,6 +67,10 @@ ESceneCaptureSource FCustomRenderPassBase::GetSceneCaptureSource() const
 		return SCS_SceneColorHDR;
 	else if (RenderOutput == ERenderOutput::SceneColorNoAlpha)
 		return SCS_SceneColorHDRNoAlpha;
+	else if (RenderOutput == ERenderOutput::BaseColor)
+		return SCS_BaseColor;
+	else if (RenderOutput == ERenderOutput::Normal)
+		return SCS_Normal;
 	else
 		return SCS_MAX;
 }
diff --git a/Engine/Source/Runtime/Engine/Public/Rendering/CustomRenderPass.h b/Engine/Source/Runtime/Engine/Public/Rendering/CustomRenderPass.h
--- a/Engine/Source/Runtime/Engine/Public/Rendering/CustomRenderPass.h
+++ b/Engine/Source/Runtime/Engine/Public/Rendering/CustomRenderPass.h
@@ -70,6 +70,8 @@ public:
 		SceneColorAndDepth,
 		SceneColorAndAlpha, // The alpha is inverted opacity (throughput).
 		SceneColorNoAlpha,
+		BaseColor,
+		Normal
 	};
 
 	enum class ERenderCaptureType
diff --git a/Engine/Source/Runtime/Engine/Public/SceneInterface.h b/Engine/Source/Runtime/Engine/Public/SceneInterface.h
--- a/Engine/Source/Runtime/Engine/Public/SceneInterface.h
+++ b/Engine/Source/Runtime/Engine/Public/SceneInterface.h
@@ -9,6 +9,7 @@
 #include "RenderGraphFwd.h"
 #include "Engine/EngineTypes.h"
 #include "PrimitiveComponentId.h"
+#include "ShowFlags.h"
 
 class AWorldSettings;
 class FArchive;
@@ -689,15 +690,22 @@ public:
 	/** Contains settings used to construct scene view for custom render pass during the renderer construction. */
 	struct FCustomRenderPassRendererInput
 	{
+		FCustomRenderPassRendererInput()
+			: EngineShowFlags(ESFIM_Game)
+		{
+		}
+
 		/** Data used to construct scene view for the custom render pass. */
 		FVector ViewLocation;
 		FMatrix ViewRotationMatrix;
 		FMatrix ProjectionMatrix;
 		TSet<FPrimitiveComponentId> HiddenPrimitives;
 		TOptional<TSet<FPrimitiveComponentId>> ShowOnlyPrimitives;
+		FEngineShowFlags EngineShowFlags;
 		const AActor* ViewActor = nullptr;
 		class FSceneViewStateInterface* ViewStateInterface = nullptr;
 		bool bIsSceneCapture = false;
+		bool bUseMainViewFamilyShowFlags = false;			// Custom render pass should use flags from main view family, rather than ones in this structure
 
 		class FCustomRenderPassBase* CustomRenderPass = nullptr;
 	};
diff --git a/Engine/Source/Runtime/Engine/Public/SceneView.h b/Engine/Source/Runtime/Engine/Public/SceneView.h
--- a/Engine/Source/Runtime/Engine/Public/SceneView.h
+++ b/Engine/Source/Runtime/Engine/Public/SceneView.h
@@ -2168,7 +2168,6 @@ public:
 	*/
 	bool bIsFirstViewInMultipleViewFamily = true;
 
-	bool bIsSceneTexturesInitialized = false;
 	bool bIsViewFamilyInfo = false;
 
 	/** Whether this view allows split screen debug -- we only want it for editor and game viewports, not other random scene renders */
diff --git a/Engine/Source/Runtime/Renderer/Internal/SceneTextures.h b/Engine/Source/Runtime/Renderer/Internal/SceneTextures.h
--- a/Engine/Source/Runtime/Renderer/Internal/SceneTextures.h
+++ b/Engine/Source/Runtime/Renderer/Internal/SceneTextures.h
@@ -21,6 +21,12 @@ struct FMinimalSceneTextures
 	// Initializes the minimal scene textures structure in the FViewFamilyInfo
 	static RENDERER_API void InitializeViewFamily(FRDGBuilder& GraphBuilder, FViewFamilyInfo& ViewFamily);
 
+	// Structure may be pointed to by multiple FViewFamilyInfo during scene rendering, through CustomRenderPasses.  The Owner
+	// handles deleting the structure when the scene renderer is destroyed.  TRefCountPtr doesn't work, because the structure
+	// is also copied by value, and the copy constructor is disabled for reference counted structures.
+	const FViewFamilyInfo* Owner = nullptr;
+	bool bIsSceneTexturesInitialized = false;
+
 	// Immutable copy of the config used to create scene textures.
 	FSceneTexturesConfig Config;
 
diff --git a/Engine/Source/Runtime/Renderer/Private/AnisotropyRendering.cpp b/Engine/Source/Runtime/Renderer/Private/AnisotropyRendering.cpp
--- a/Engine/Source/Runtime/Renderer/Private/AnisotropyRendering.cpp
+++ b/Engine/Source/Runtime/Renderer/Private/AnisotropyRendering.cpp
@@ -295,6 +295,12 @@ bool ShouldRenderAnisotropyPass(const FViewInfo& View)
 		return false;
 	}
 
+	// The anisotropy GBuffer is used for lighting, and not needed for custom render passes, which don't run lighting.
+	if (View.CustomRenderPass)
+	{
+		return false;
+	}
+
 	if (View.ShouldRenderView() && View.ParallelMeshDrawCommandPasses[EMeshPass::AnisotropyPass].HasAnyDraw())
 	{
 		return true;
@@ -324,6 +330,7 @@ END_SHADER_PARAMETER_STRUCT()
 
 void FDeferredShadingSceneRenderer::RenderAnisotropyPass(
 	FRDGBuilder& GraphBuilder, 
+	TArrayView<FViewInfo> InViews,
 	FSceneTextures& SceneTextures,
 	bool bDoParallelPass
 )
@@ -333,9 +340,9 @@ void FDeferredShadingSceneRenderer::RenderAnisotropyPass(
 	SCOPE_CYCLE_COUNTER(STAT_AnisotropyPassDrawTime);
 	RDG_GPU_STAT_SCOPE(GraphBuilder, RenderAnisotropyPass);
 
-	for (int32 ViewIndex = 0; ViewIndex < Views.Num(); ViewIndex++)
+	for (int32 ViewIndex = 0; ViewIndex < InViews.Num(); ViewIndex++)
 	{
-		FViewInfo& View = Views[ViewIndex];
+		FViewInfo& View = InViews[ViewIndex];
 
 		if (View.ShouldRenderView())
 		{
diff --git a/Engine/Source/Runtime/Renderer/Private/BasePassRendering.cpp b/Engine/Source/Runtime/Renderer/Private/BasePassRendering.cpp
--- a/Engine/Source/Runtime/Renderer/Private/BasePassRendering.cpp
+++ b/Engine/Source/Runtime/Renderer/Private/BasePassRendering.cpp
@@ -1187,9 +1187,10 @@ void FDeferredShadingSceneRenderer::RenderBasePass(
 		ClearGBufferAtMaxZ(GraphBuilder, InViews, BasePassRenderTargets, SceneColorClearValue);
 	}
 
-	if (ShouldRenderAnisotropyPass(InViews))
+	// The anisotropy GBuffer is used for lighting, and not needed when running a scene capture with ERendererOutput::BasePass (BaseColor or Normal output), which doesn't run lighting.
+	if (ShouldRenderAnisotropyPass(InViews) && GetRendererOutput() == ERendererOutput::FinalSceneColor)
 	{
-		RenderAnisotropyPass(GraphBuilder, SceneTextures, bEnableParallelBasePasses);
+		RenderAnisotropyPass(GraphBuilder, InViews, SceneTextures, bEnableParallelBasePasses);
 	}
 
 #if !(UE_BUILD_SHIPPING)
diff --git a/Engine/Source/Runtime/Renderer/Private/DeferredShadingRenderer.cpp b/Engine/Source/Runtime/Renderer/Private/DeferredShadingRenderer.cpp
--- a/Engine/Source/Runtime/Renderer/Private/DeferredShadingRenderer.cpp
+++ b/Engine/Source/Runtime/Renderer/Private/DeferredShadingRenderer.cpp
@@ -101,6 +101,7 @@
 #include "ManyLights/ManyLights.h"
 #include "Rendering/CustomRenderPass.h"
 #include "EnvironmentComponentsFlags.h"
+#include "GenerateMips.h"
 
 #if !UE_BUILD_SHIPPING
 #include "RenderCaptureInterface.h"
@@ -1397,7 +1398,7 @@ void FDeferredShadingSceneRenderer::Render(FRDGBuilder& GraphBuilder)
 	GPU_MESSAGE_SCOPE(GraphBuilder);
 
 #if RHI_RAYTRACING
-	if (RendererOutput != FSceneRenderer::ERendererOutput::DepthPrepassOnly)
+	if (RendererOutput == FSceneRenderer::ERendererOutput::FinalSceneColor)
 	{
 		GRayTracingGeometryManager->PreRender();
 
@@ -1417,7 +1418,8 @@ void FDeferredShadingSceneRenderer::Render(FRDGBuilder& GraphBuilder)
 
 	const bool bUseVirtualTexturing = UseVirtualTexturing(ShaderPlatform);
 
-	if (bUseVirtualTexturing && RendererOutput == ERendererOutput::FinalSceneColor)
+	// Virtual texturing runs for ERendererOutput::BasePass or ERendererOutput::FinalSceneColor
+	if (bUseVirtualTexturing && RendererOutput != ERendererOutput::DepthPrepassOnly)
 	{
 		FVirtualTextureUpdateSettings Settings;
 		Settings.EnableThrottling(!ViewFamily.bOverrideVirtualTextureThrottle);
@@ -1807,14 +1809,18 @@ void FDeferredShadingSceneRenderer::Render(FRDGBuilder& GraphBuilder)
 	bool bComputeLightGrid = false;
 	bool bAnyLumenEnabled = false;
 
-	if (RendererOutput == ERendererOutput::FinalSceneColor)
+	// Virtual texturing runs for ERendererOutput::BasePass or ERendererOutput::FinalSceneColor
+	if (RendererOutput != ERendererOutput::DepthPrepassOnly)
 	{
 		if (bUseVirtualTexturing)
 		{
 			// Note, should happen after the GPU-Scene update to ensure rendering to runtime virtual textures is using the correctly updated scene
 			FVirtualTextureSystem::Get().EndUpdate(GraphBuilder, MoveTemp(VirtualTextureUpdater), FeatureLevel);
 		}
+	}
 
+	if (RendererOutput == ERendererOutput::FinalSceneColor)
+	{
 #if RHI_RAYTRACING
 		if (bAnyRayTracingPassEnabled)
 		{
@@ -2027,12 +2033,21 @@ void FDeferredShadingSceneRenderer::Render(FRDGBuilder& GraphBuilder)
 
 	FDBufferTextures DBufferTextures = CreateDBufferTextures(GraphBuilder, SceneTextures.Config.Extent, ShaderPlatform);
 
+	// Initialise local fog volume with dummy data before volumetric cloud view initialization (further down) which can bind LFV data.
+	// Also need to do this before custom render passes (included in AllViews), as base pass rendering may bind LFV data.
+	SetDummyLocalFogVolumeForViews(GraphBuilder, AllViews);
+
 	if (CustomRenderPassInfos.Num() > 0)
 	{
 		QUICK_SCOPE_CYCLE_COUNTER(STAT_CustomRenderPasses);
 		RDG_EVENT_SCOPE(GraphBuilder, "CustomRenderPasses");
 		RDG_GPU_STAT_SCOPE(GraphBuilder, CustomRenderPasses);
 
+		// We want to reset the scene texture uniform buffer to its original state after custom render passes,
+		// so they can't affect downstream rendering.
+		ESceneTextureSetupMode OriginalSceneTextureSetupMode = SceneTextures.SetupMode;
+		TRDGUniformBufferRef<FSceneTextureUniformParameters> OriginalSceneTextureUniformBuffer = SceneTextures.UniformBuffer;
+
 		for (int32 i = 0; i < CustomRenderPassInfos.Num(); ++i)
 		{
 			FCustomRenderPassBase* CustomRenderPass = CustomRenderPassInfos[i].CustomRenderPass;
@@ -2053,31 +2068,65 @@ void FDeferredShadingSceneRenderer::Render(FRDGBuilder& GraphBuilder)
 				FNaniteBasePassVisibility DummyNaniteBasePassVisibility;
 				RenderPrepassAndVelocity(CustomRenderPassViews, DummyNaniteBasePassVisibility, NaniteRasterResults, PrimaryNaniteViews);
 
+				const FSingleLayerWaterPrePassResult* SingleLayerWaterPrePassResult = nullptr;
+				if (ShouldRenderSingleLayerWaterDepthPrepass(CustomRenderPassViews))
+				{
+					SingleLayerWaterPrePassResult = RenderSingleLayerWaterDepthPrepass(GraphBuilder, CustomRenderPassViews, SceneTextures);
+				}
+
 				if (CustomRenderPass->GetRenderMode() == FCustomRenderPassBase::ERenderMode::DepthAndBasePass)
 				{
 					SceneTextures.SetupMode |= ESceneTextureSetupMode::SceneColor;
 					SceneTextures.UniformBuffer = CreateSceneTextureUniformBuffer(GraphBuilder, &SceneTextures, FeatureLevel, SceneTextures.SetupMode);
 
-					// Setup dummy uniform buffer parameters for fog volume.
-					SetDummyLocalFogVolumeForViews(GraphBuilder, CustomRenderPassViews);
-
 					if (bNaniteEnabled)
 					{
 						Nanite::BuildShadingCommands(GraphBuilder, *Scene, ENaniteMeshPass::BasePass, NaniteBasePassShadingCommands, Nanite::EBuildShadingCommandsMode::Custom);
 					}
 
 					RenderBasePass(GraphBuilder, CustomRenderPassViews, SceneTextures, DBufferTextures, BasePassDepthStencilAccess, /*ForwardScreenSpaceShadowMaskTexture=*/nullptr, InstanceCullingManager, bNaniteEnabled, NaniteBasePassShadingCommands, NaniteRasterResults);
+
+					if (ShouldRenderSingleLayerWater(CustomRenderPassViews))
+					{
+						// GBuffer code paths in RenderSingleLayerWater don't use the bIsCameraUnderWater flag, so just pass in false.  Normally this is
+						// computed by a render extension, but those aren't run for custom render passes.
+						FSceneWithoutWaterTextures SceneWithoutWaterTextures;
+						RenderSingleLayerWater(GraphBuilder, CustomRenderPassViews, SceneTextures, SingleLayerWaterPrePassResult, /*bShouldRenderVolumetricCloud=*/false, SceneWithoutWaterTextures, LumenFrameTemporaries, /*bIsCameraUnderWater=*/false);
+					}
+
+					FCustomRenderPassBase::ERenderOutput RenderOutput = CustomRenderPass->GetRenderOutput();
+					if (RenderOutput == FCustomRenderPassBase::ERenderOutput::BaseColor || RenderOutput == FCustomRenderPassBase::ERenderOutput::Normal)
+					{
+						// CopySceneCaptureComponentToTarget uses scene texture uniforms
+						SceneTextures.SetupMode |= ESceneTextureSetupMode::GBuffers;
+						SceneTextures.UniformBuffer = CreateSceneTextureUniformBuffer(GraphBuilder, &SceneTextures, FeatureLevel, SceneTextures.SetupMode);
+					}
 				}
 
 				CopySceneCaptureComponentToTarget(GraphBuilder, SceneTextures, CustomRenderPass->GetRenderTargetTexture(), ViewFamily, CustomRenderPassViews);
 				CustomRenderPass->PostRender(GraphBuilder);
 
+				// Mips are normally generated in UpdateSceneCaptureContentDeferred_RenderThread, but that doesn't run when the
+				// scene capture runs as a custom render pass.  The function does nothing if the render target doesn't have mips.
+				if (CustomRenderPassViews[0].bIsSceneCapture)
+				{
+					FGenerateMips::Execute(GraphBuilder, FeatureLevel, CustomRenderPass->GetRenderTargetTexture(), FGenerateMipsParams());
+				}
+
 			#if WITH_MGPU
 				DoCrossGPUTransfers(GraphBuilder, CustomRenderPass->GetRenderTargetTexture(), CustomRenderPassViews, false, FRHIGPUMask::All());
 			#endif
+
+				// Materials in the main view renderer will be using this render target, so we need RDG to transition it back to SRV now,
+				// rather than at the end of graph execution.
+				GraphBuilder.UseExternalAccessMode(CustomRenderPass->GetRenderTargetTexture(), ERHIAccess::SRVMask);
 			}
 
 			CustomRenderPass->EndPass(GraphBuilder);
+
+			// Restore original scene texture uniforms
+			SceneTextures.SetupMode = OriginalSceneTextureSetupMode;
+			SceneTextures.UniformBuffer = OriginalSceneTextureUniformBuffer;
 		}
 	}
 
@@ -2086,10 +2135,10 @@ void FDeferredShadingSceneRenderer::Render(FRDGBuilder& GraphBuilder)
 	RenderPrepassAndVelocity(Views, NaniteBasePassVisibility, NaniteRasterResults, PrimaryNaniteViews);
 
 	// Run Nanite compute commands early in the frame to allow some task overlap on the CPU until the base pass runs.
-	if (bNaniteEnabled && RendererOutput == ERendererOutput::FinalSceneColor && !bHasRayTracedOverlay)
+	if (bNaniteEnabled && RendererOutput != ERendererOutput::DepthPrepassOnly && !bHasRayTracedOverlay)
 	{
 		Nanite::BuildShadingCommands(GraphBuilder, *Scene, ENaniteMeshPass::BasePass, Scene->NaniteShadingCommands[ENaniteMeshPass::BasePass]);
-		if (bAnyLumenEnabled)
+		if (bAnyLumenEnabled && RendererOutput == ERendererOutput::FinalSceneColor)
 		{
 			Nanite::BuildShadingCommands(GraphBuilder, *Scene, ENaniteMeshPass::LumenCardCapture, Scene->NaniteShadingCommands[ENaniteMeshPass::LumenCardCapture]);
 		}
@@ -2127,9 +2176,58 @@ void FDeferredShadingSceneRenderer::Render(FRDGBuilder& GraphBuilder)
 
 	FRDGTextureRef ViewFamilyTexture = TryCreateViewFamilyTexture(GraphBuilder, ViewFamily);
 	FRDGTextureRef ViewFamilyDepthTexture = TryCreateViewFamilyDepthTexture(GraphBuilder, ViewFamily);
-	if (RendererOutput == ERendererOutput::DepthPrepassOnly)
+	if (RendererOutput == ERendererOutput::DepthPrepassOnly || RendererOutput == ERendererOutput::BasePass)
 	{
-		RenderOcclusionLambda();
+		const FSingleLayerWaterPrePassResult* SingleLayerWaterPrePassResult = nullptr;
+		if (ShouldRenderSingleLayerWaterDepthPrepass(Views))
+		{
+			SingleLayerWaterPrePassResult = RenderSingleLayerWaterDepthPrepass(GraphBuilder, Views, SceneTextures);
+		}
+
+		bool bOcclusionBeforeBasePass = false;
+		if (RendererOutput == ERendererOutput::BasePass)
+		{
+			// Early occlusion queries
+			bOcclusionBeforeBasePass = ((DepthPass.EarlyZPassMode == EDepthDrawingMode::DDM_AllOccluders) || bIsEarlyDepthComplete);
+			if (bOcclusionBeforeBasePass)
+			{
+				RenderOcclusionLambda();
+			}
+
+			SceneTextures.SetupMode |= ESceneTextureSetupMode::SceneColor;
+			SceneTextures.UniformBuffer = CreateSceneTextureUniformBuffer(GraphBuilder, &SceneTextures, FeatureLevel, SceneTextures.SetupMode);
+
+			for (FSceneViewExtensionRef& ViewExtension : ViewFamily.ViewExtensions)
+			{
+				ViewExtension->PreRenderBasePass_RenderThread(GraphBuilder, ShouldRenderPrePass() /*bDepthBufferIsPopulated*/);
+			}
+
+			RenderBasePass(GraphBuilder, Views, SceneTextures, DBufferTextures, BasePassDepthStencilAccess, /*ForwardScreenSpaceShadowMaskTexture=*/nullptr, InstanceCullingManager, bNaniteEnabled, Scene->NaniteShadingCommands[ENaniteMeshPass::BasePass], NaniteRasterResults);
+
+			if (ShouldRenderSingleLayerWater(Views))
+			{
+				// Virtual shadow map uniforms need to be initialized with dummy data for water.  Their initialization
+				// was skipped above due to (RendererOutput == ERendererOutput::FinalSceneColor) being false.
+				VirtualShadowMapArray.Initialize(GraphBuilder, Scene->GetVirtualShadowMapCache(), /*bEnableVirtualShadowMaps=*/false, ViewFamily.EngineShowFlags);
+
+				FSceneWithoutWaterTextures SceneWithoutWaterTextures;
+				RenderSingleLayerWater(GraphBuilder, Views, SceneTextures, SingleLayerWaterPrePassResult, /*bShouldRenderVolumetricCloud=*/false, SceneWithoutWaterTextures, LumenFrameTemporaries, /*bIsCameraUnderWater=*/false);
+			}
+
+			SceneTextures.SetupMode |= ESceneTextureSetupMode::GBuffers;
+			SceneTextures.UniformBuffer = CreateSceneTextureUniformBuffer(GraphBuilder, &SceneTextures, FeatureLevel, SceneTextures.SetupMode);
+
+			if (bUseVirtualTexturing)
+			{
+				RDG_GPU_STAT_SCOPE(GraphBuilder, VirtualTextureUpdate);
+				VirtualTextureFeedbackEnd(GraphBuilder);
+			}
+		}
+
+		if (!bOcclusionBeforeBasePass)
+		{
+			RenderOcclusionLambda();
+		}
 
 		if (bUpdateNaniteStreaming)
 		{
@@ -2175,9 +2273,6 @@ void FDeferredShadingSceneRenderer::Render(FRDGBuilder& GraphBuilder)
 		CSV_CUSTOM_STAT(LightCount, Batched, float(SortedLightSet.UnbatchedLightStart), ECsvCustomStatOp::Set);
 		CSV_CUSTOM_STAT(LightCount, Unbatched, float(SortedLightSet.SortedLights.Num()) - float(SortedLightSet.UnbatchedLightStart), ECsvCustomStatOp::Set);
 
-		// Initialise local fog volume with dummy data before volumetric cloud view initialization which can bind LFV data.
-		SetDummyLocalFogVolumeForViews(GraphBuilder, Views);
-
 		// Run before RenderSkyAtmosphereLookUpTables for cloud shadows to be valid.
 		InitVolumetricCloudsForViews(GraphBuilder, bShouldRenderVolumetricCloudBase, InstanceCullingManager);
 
@@ -2503,7 +2598,7 @@ void FDeferredShadingSceneRenderer::Render(FRDGBuilder& GraphBuilder)
 		const bool bShouldRenderSingleLayerWaterDepthPrepass = !bHasRayTracedOverlay && ShouldRenderSingleLayerWaterDepthPrepass(Views);
 		if (bShouldRenderSingleLayerWaterDepthPrepass)
 		{
-			SingleLayerWaterPrePassResult = RenderSingleLayerWaterDepthPrepass(GraphBuilder, SceneTextures);
+			SingleLayerWaterPrePassResult = RenderSingleLayerWaterDepthPrepass(GraphBuilder, Views, SceneTextures);
 		}
 
 		FAsyncLumenIndirectLightingOutputs AsyncLumenIndirectLightingOutputs;
@@ -2914,7 +3009,7 @@ void FDeferredShadingSceneRenderer::Render(FRDGBuilder& GraphBuilder)
 				EnumRemoveFlags(TranslucencyViewsToRender, ETranslucencyView::UnderWater);
 			}
 
-			RenderSingleLayerWater(GraphBuilder, SceneTextures, SingleLayerWaterPrePassResult, bShouldRenderVolumetricCloud, SceneWithoutWaterTextures, LumenFrameTemporaries, bIsCameraUnderWater);
+			RenderSingleLayerWater(GraphBuilder, Views, SceneTextures, SingleLayerWaterPrePassResult, bShouldRenderVolumetricCloud, SceneWithoutWaterTextures, LumenFrameTemporaries, bIsCameraUnderWater);
 
 			// Replace main depth texture with the output of the SLW depth prepass which contains the scene + water.
 			// Note: Stencil now has all water bits marked with 1. As long as no other passes after this point want to read the depth buffer,
diff --git a/Engine/Source/Runtime/Renderer/Private/DeferredShadingRenderer.h b/Engine/Source/Runtime/Renderer/Private/DeferredShadingRenderer.h
--- a/Engine/Source/Runtime/Renderer/Private/DeferredShadingRenderer.h
+++ b/Engine/Source/Runtime/Renderer/Private/DeferredShadingRenderer.h
@@ -377,6 +377,7 @@ public:
 
 	void RenderAnisotropyPass(
 		FRDGBuilder& GraphBuilder,
+		TArrayView<FViewInfo> InViews,
 		FSceneTextures& SceneTextures,
 		bool bDoParallelPass);
 	/**
@@ -384,10 +385,12 @@ public:
 	 */
 	FSingleLayerWaterPrePassResult* RenderSingleLayerWaterDepthPrepass(
 		FRDGBuilder& GraphBuilder,
+		TArrayView<FViewInfo> InViews,
 		const FSceneTextures& SceneTextures);
 
 	void RenderSingleLayerWater(
 		FRDGBuilder& GraphBuilder,
+		TArrayView<FViewInfo> InViews,
 		const FSceneTextures& SceneTextures,
 		const FSingleLayerWaterPrePassResult* SingleLayerWaterPrePassResult,
 		bool bShouldRenderVolumetricCloud,
@@ -397,12 +400,14 @@ public:
 
 	void RenderSingleLayerWaterInner(
 		FRDGBuilder& GraphBuilder,
+		TArrayView<FViewInfo> InViews,
 		const FSceneTextures& SceneTextures,
 		const FSceneWithoutWaterTextures& SceneWithoutWaterTextures,
 		const FSingleLayerWaterPrePassResult* SingleLayerWaterPrePassResult);
 
 	void RenderSingleLayerWaterReflections(
 		FRDGBuilder& GraphBuilder,
+		TArrayView<FViewInfo> InViews,
 		const FSceneTextures& SceneTextures,
 		const FSceneWithoutWaterTextures& SceneWithoutWaterTextures,
 		const FSingleLayerWaterPrePassResult* SingleLayerWaterPrePassResult,
diff --git a/Engine/Source/Runtime/Renderer/Private/LocalFogVolumeRendering.cpp b/Engine/Source/Runtime/Renderer/Private/LocalFogVolumeRendering.cpp
--- a/Engine/Source/Runtime/Renderer/Private/LocalFogVolumeRendering.cpp
+++ b/Engine/Source/Runtime/Renderer/Private/LocalFogVolumeRendering.cpp
@@ -208,6 +208,14 @@ void SetDummyLocalFogVolumeForViews(FRDGBuilder& GraphBuilder, TArray<FViewInfo>
 	}
 }
 
+void SetDummyLocalFogVolumeForViews(FRDGBuilder& GraphBuilder, TArray<FViewInfo*>& Views)
+{
+	for (FViewInfo* View : Views)
+	{
+		SetDummyLocalFogVolumeForView(GraphBuilder, *View);
+	}
+}
+
 
 /*=============================================================================
 	FScene functions
diff --git a/Engine/Source/Runtime/Renderer/Private/LocalFogVolumeRendering.h b/Engine/Source/Runtime/Renderer/Private/LocalFogVolumeRendering.h
--- a/Engine/Source/Runtime/Renderer/Private/LocalFogVolumeRendering.h
+++ b/Engine/Source/Runtime/Renderer/Private/LocalFogVolumeRendering.h
@@ -138,6 +138,7 @@ bool IsLocalFogVolumeHalfResolution();
 void GetLocalFogVolumeSortingData(const FScene* Scene, FRDGBuilder& GraphBuilder, FLocalFogVolumeSortingData& Out);
 
 void SetDummyLocalFogVolumeForViews(FRDGBuilder& GraphBuilder, TArray<FViewInfo>& Views);
+void SetDummyLocalFogVolumeForViews(FRDGBuilder& GraphBuilder, TArray<FViewInfo*>& Views);
 void SetDummyLocalFogVolumeForView(FRDGBuilder& GraphBuilder, FViewInfo& View);
 
 void InitLocalFogVolumesForViews(
diff --git a/Engine/Source/Runtime/Renderer/Private/MobileShadingRenderer.cpp b/Engine/Source/Runtime/Renderer/Private/MobileShadingRenderer.cpp
--- a/Engine/Source/Runtime/Renderer/Private/MobileShadingRenderer.cpp
+++ b/Engine/Source/Runtime/Renderer/Private/MobileShadingRenderer.cpp
@@ -74,6 +74,7 @@
 #include "SceneCaptureRendering.h"
 #include "WaterInfoTextureRendering.h"
 #include "Rendering/CustomRenderPass.h"
+#include "GenerateMips.h"
 
 uint32 GetShadowQuality();
 
@@ -275,7 +276,7 @@ FMobileSceneRenderer::FMobileSceneRenderer(const FSceneViewFamily* InViewFamily,
 	, bGammaSpace(!IsMobileHDR())
 	, bDeferredShading(IsMobileDeferredShadingEnabled(ShaderPlatform))
 	, bRequiresDBufferDecals(bDeferredShading ? false : IsUsingDBuffers(ShaderPlatform))
-	, bUseVirtualTexturing(UseVirtualTexturing(ShaderPlatform) && GetRendererOutput() == FSceneRenderer::ERendererOutput::FinalSceneColor)
+	, bUseVirtualTexturing(UseVirtualTexturing(ShaderPlatform) && GetRendererOutput() != FSceneRenderer::ERendererOutput::DepthPrepassOnly)
 {
 	bRenderToSceneColor = false;
 	bRequiresMultiPass = false;
@@ -305,7 +306,7 @@ FMobileSceneRenderer::FMobileSceneRenderer(const FSceneViewFamily* InViewFamily,
 
 	NumMSAASamples = GetDefaultMSAACount(ERHIFeatureLevel::ES3_1);
 	// As of UE 5.4 only vulkan supports inline (single pass) tonemap
-	bTonemapSubpass = IsMobileTonemapSubpassEnabled(ShaderPlatform) && ViewFamily.bResolveScene && GetRendererOutput() == FSceneRenderer::ERendererOutput::FinalSceneColor;
+	bTonemapSubpass = IsMobileTonemapSubpassEnabled(ShaderPlatform) && ViewFamily.bResolveScene && GetRendererOutput() != FSceneRenderer::ERendererOutput::DepthPrepassOnly;
 	bTonemapSubpassInline = bTonemapSubpass && IsVulkanPlatform(ShaderPlatform) && (GRHISupportsMSAAShaderResolve || NumMSAASamples == 1);
 	bRequiresSceneDepthAux = MobileRequiresSceneDepthAux(ShaderPlatform) && !bTonemapSubpass;
 }
@@ -432,7 +433,7 @@ void FMobileSceneRenderer::InitViews(
 
 	check(Scene);
 
-	const bool bRendererOutputFinalSceneColor = (GetRendererOutput() == ERendererOutput::FinalSceneColor);
+	const bool bRendererOutputFinalSceneColor = (GetRendererOutput() != ERendererOutput::DepthPrepassOnly);
 
 	PreVisibilityFrameSetup(GraphBuilder);
 
@@ -900,7 +901,7 @@ void FMobileSceneRenderer::Render(FRDGBuilder& GraphBuilder)
 	}
 
 	const ERendererOutput RendererOutput = GetRendererOutput();
-	const bool bRendererOutputFinalSceneColor = (RendererOutput == ERendererOutput::FinalSceneColor);
+	const bool bRendererOutputFinalSceneColor = (RendererOutput != ERendererOutput::DepthPrepassOnly);
 
 	RDG_RHI_EVENT_SCOPE(GraphBuilder, MobileSceneRender);
 	RDG_RHI_GPU_STAT_SCOPE(GraphBuilder, MobileSceneRender);
@@ -1098,7 +1099,7 @@ void FMobileSceneRenderer::Render(FRDGBuilder& GraphBuilder)
 		}
 
 		// Hair update
-		if (IsHairStrandsEnabled(EHairStrandsShaderType::All, Scene->GetShaderPlatform()) && RendererOutput == ERendererOutput::FinalSceneColor)
+		if (IsHairStrandsEnabled(EHairStrandsShaderType::All, Scene->GetShaderPlatform()) && RendererOutput != ERendererOutput::DepthPrepassOnly)
 		{
 			FHairStrandsBookmarkParameters& HairStrandsBookmarkParameters = *GraphBuilder.AllocObject<FHairStrandsBookmarkParameters>();
 			CreateHairStrandsBookmarkParameters(Scene, Views, AllFamilyViews, HairStrandsBookmarkParameters);
@@ -1174,6 +1175,11 @@ void FMobileSceneRenderer::Render(FRDGBuilder& GraphBuilder)
 		RDG_EVENT_SCOPE(GraphBuilder, "CustomRenderPasses");
 		RDG_GPU_STAT_SCOPE(GraphBuilder, CustomRenderPasses);
 
+		// We want to reset the scene texture uniform buffer to its original state after custom render passes,
+		// so they can't affect downstream rendering.
+		EMobileSceneTextureSetupMode OriginalSceneTextureSetupMode = SceneTextures.MobileSetupMode;
+		TRDGUniformBufferRef<FMobileSceneTextureUniformParameters> OriginalSceneTextureUniformBuffer = SceneTextures.MobileUniformBuffer;
+
 		for (int32 i = 0; i < CustomRenderPassInfos.Num(); ++i)
 		{
 			FCustomRenderPassBase* CustomRenderPass = CustomRenderPassInfos[i].CustomRenderPass;
@@ -1212,9 +1218,23 @@ void FMobileSceneRenderer::Render(FRDGBuilder& GraphBuilder)
 				CopySceneCaptureComponentToTarget(GraphBuilder, SceneTextures, CustomRenderPass->GetRenderTargetTexture(), ViewFamily, CustomRenderPassViews);
 
 				CustomRenderPass->PostRender(GraphBuilder);
+
+				// Mips are normally generated in UpdateSceneCaptureContentMobile_RenderThread, but that doesn't run when the
+				// scene capture runs as a custom render pass.  The function does nothing if the render target doesn't have mips.
+				if (CustomRenderPassViews[0].bIsSceneCapture)
+				{
+					FGenerateMips::Execute(GraphBuilder, FeatureLevel, CustomRenderPass->GetRenderTargetTexture(), FGenerateMipsParams());
+				}
+
+				// Materials in the main view renderer will be using this render target, so we need RDG to transition it back to SRV now,
+				// rather than at the end of graph execution.
+				GraphBuilder.UseExternalAccessMode(CustomRenderPass->GetRenderTargetTexture(), ERHIAccess::SRVMask);
 			}
 
 			CustomRenderPass->EndPass(GraphBuilder);
+
+			SceneTextures.MobileSetupMode = OriginalSceneTextureSetupMode;
+			SceneTextures.MobileUniformBuffer = OriginalSceneTextureUniformBuffer;
 		}
 	}
 
diff --git a/Engine/Source/Runtime/Renderer/Private/SceneCaptureRendering.cpp b/Engine/Source/Runtime/Renderer/Private/SceneCaptureRendering.cpp
--- a/Engine/Source/Runtime/Renderer/Private/SceneCaptureRendering.cpp
+++ b/Engine/Source/Runtime/Renderer/Private/SceneCaptureRendering.cpp
@@ -1090,7 +1090,8 @@ void FScene::UpdateSceneCaptureContents(USceneCaptureComponent2D* CaptureCompone
 		// As optimization for depth capture modes, render scene capture as additional render passes inside the main renderer.
 		if (GSceneCaptureAllowRenderInMainRenderer && 
 			CaptureComponent->bRenderInMainRenderer && 
-			(CaptureComponent->CaptureSource == ESceneCaptureSource::SCS_SceneDepth || CaptureComponent->CaptureSource == ESceneCaptureSource::SCS_DeviceDepth)
+			(CaptureComponent->CaptureSource == ESceneCaptureSource::SCS_SceneDepth || CaptureComponent->CaptureSource == ESceneCaptureSource::SCS_DeviceDepth ||
+			 CaptureComponent->CaptureSource == ESceneCaptureSource::SCS_BaseColor || CaptureComponent->CaptureSource == ESceneCaptureSource::SCS_Normal)
 			)
 		{
 			FCustomRenderPassRendererInput PassInput;
@@ -1100,13 +1101,42 @@ void FScene::UpdateSceneCaptureContents(USceneCaptureComponent2D* CaptureCompone
 			PassInput.ViewActor = CaptureComponent->GetViewOwner();
 			PassInput.bIsSceneCapture = true;
 
-			FString DebugName = CaptureComponent->CaptureSource == ESceneCaptureSource::SCS_SceneDepth ? TEXT("SceneCapturePass_SceneDepth") : TEXT("SceneCapturePass_DeviceDepth");
-			FCustomRenderPassBase::ERenderOutput RenderOutput = CaptureComponent->CaptureSource == ESceneCaptureSource::SCS_SceneDepth ? FCustomRenderPassBase::ERenderOutput::SceneDepth : FCustomRenderPassBase::ERenderOutput::DeviceDepth;
-			FSceneCapturePass* CustomPass = new FSceneCapturePass(DebugName, FCustomRenderPassBase::ERenderMode::DepthPass, RenderOutput, TextureRenderTarget);
+			FCustomRenderPassBase::ERenderMode RenderMode;
+			FCustomRenderPassBase::ERenderOutput RenderOutput;
+			const TCHAR* DebugName;
+
+			switch (CaptureComponent->CaptureSource)
+			{
+			case ESceneCaptureSource::SCS_SceneDepth:
+				RenderMode = FCustomRenderPassBase::ERenderMode::DepthPass;
+				RenderOutput = FCustomRenderPassBase::ERenderOutput::SceneDepth;
+				DebugName = TEXT("SceneCapturePass_SceneDepth");
+				break;
+			case ESceneCaptureSource::SCS_DeviceDepth:
+				RenderMode = FCustomRenderPassBase::ERenderMode::DepthPass;
+				RenderOutput = FCustomRenderPassBase::ERenderOutput::DeviceDepth;
+				DebugName = TEXT("SceneCapturePass_DeviceDepth");
+				break;
+			case ESceneCaptureSource::SCS_Normal:
+				RenderMode = FCustomRenderPassBase::ERenderMode::DepthAndBasePass;
+				RenderOutput = FCustomRenderPassBase::ERenderOutput::Normal;
+				DebugName = TEXT("SceneCapturePass_Normal");
+				break;
+			case ESceneCaptureSource::SCS_BaseColor:
+			default:
+				RenderMode = FCustomRenderPassBase::ERenderMode::DepthAndBasePass;
+				RenderOutput = FCustomRenderPassBase::ERenderOutput::BaseColor;
+				DebugName = TEXT("SceneCapturePass_BaseColor");
+				break;
+			}
+
+			FSceneCapturePass* CustomPass = new FSceneCapturePass(DebugName, RenderMode, RenderOutput, TextureRenderTarget);
 			PassInput.CustomRenderPass = CustomPass;
 
 			GetShowOnlyAndHiddenComponents(CaptureComponent, PassInput.HiddenPrimitives, PassInput.ShowOnlyPrimitives);
 
+			PassInput.EngineShowFlags = CaptureComponent->ShowFlags;
+
 			// Caching scene capture info to be passed to the scene renderer.
 			// #todo: We cannot (yet) guarantee for which ViewFamily this CRP will eventually be rendered since it will just execute the next time the scene is rendered by any FSceneRenderer. This seems quite problematic and could easily lead to unexpected behavior...
 			AddCustomRenderPass(nullptr, PassInput);
diff --git a/Engine/Source/Runtime/Renderer/Private/SceneRendering.cpp b/Engine/Source/Runtime/Renderer/Private/SceneRendering.cpp
--- a/Engine/Source/Runtime/Renderer/Private/SceneRendering.cpp
+++ b/Engine/Source/Runtime/Renderer/Private/SceneRendering.cpp
@@ -114,8 +114,14 @@ int32 GSceneCaptureDepthPrepassOptimization = 0;
 static FAutoConsoleVariableRef CVarSceneCaptureDepthPrepassOptimization(
 	TEXT("r.SceneCapture.DepthPrepassOptimization"),
 	GSceneCaptureDepthPrepassOptimization,
-	TEXT("Whether to apply optimized render path when capturing depth prepass for scene capture 2D. Experimental!\n")
-	TEXT("Warning: turning it on means rendering after depth pre-pass (e.g. SingleLayerWater) is ignored, hence result is different from when CVar is off.\n"),
+	TEXT("Whether to apply optimized render path when capturing depth prepass for scene capture 2D. Experimental!\n"),
+	ECVF_RenderThreadSafe | ECVF_Scalability);
+
+int32 GSceneCaptureBasePassOptimization = 0;
+static FAutoConsoleVariableRef CVarSceneCaptureBasePassOptimization(
+	TEXT("r.SceneCapture.BasePassOptimization"),
+	GSceneCaptureBasePassOptimization,
+	TEXT("Whether to apply optimized render path when capturing base pass or normals for scene capture 2D. Experimental!\n"),
 	ECVF_RenderThreadSafe | ECVF_Scalability);
 
 static int32 GAsyncCreateLightPrimitiveInteractions = 1;
@@ -2606,9 +2612,31 @@ FViewFamilyInfo::FViewFamilyInfo(const FSceneViewFamily& InViewFamily)
 :	FSceneViewFamily(InViewFamily)
 {
 	bIsViewFamilyInfo = true;
+
+	SceneTextures = new FSceneTextures;
+	SceneTextures->Owner = this;
+}
+
+// Constructor that shares scene textures with a MainViewFamily.  Used to create a separate FViewFamilyInfo for custom render passes, so
+// they can have distinct EngineShowFlags from the view family they are rendering with.
+FViewFamilyInfo::FViewFamilyInfo(const FSceneViewFamily::ConstructionValues& CVS, const FViewFamilyInfo& MainViewFamily)
+:	FSceneViewFamily(CVS)
+{
+	bIsViewFamilyInfo = true;
+
+	SceneTextures = MainViewFamily.SceneTextures;
 }
 
 FViewFamilyInfo::~FViewFamilyInfo()
+{
+	if (SceneTextures && SceneTextures->Owner == this)
+	{
+		delete SceneTextures;
+	}
+}
+
+FSceneRenderer::FCustomRenderPassInfo::FCustomRenderPassInfo(const FSceneViewFamily::ConstructionValues& CVS, const FViewFamilyInfo& MainViewFamily)
+	: ViewFamily(CVS, MainViewFamily)
 {
 }
 
@@ -2876,7 +2904,6 @@ FSceneRenderer::FSceneRenderer(const FSceneViewFamily* InViewFamily, FHitProxyCo
 
 	// Prepare custom render passes and their views:
 	CustomRenderPassInfos.Empty(Scene->CustomRenderPassRendererInputs.Num());
-	CustomRenderPassInfos.AddDefaulted(Scene->CustomRenderPassRendererInputs.Num());
 
 	int32 NumAdditionalViews = 0;
 	for (int32 i = 0; i < Scene->CustomRenderPassRendererInputs.Num(); i++)
@@ -2884,7 +2911,19 @@ FSceneRenderer::FSceneRenderer(const FSceneViewFamily* InViewFamily, FHitProxyCo
 		const FScene::FCustomRenderPassRendererInput& PassInput = Scene->CustomRenderPassRendererInputs[i];
 		FCustomRenderPassBase* CustomRenderPass = PassInput.CustomRenderPass;
 		check(CustomRenderPass);
-		CustomRenderPassInfos[i].CustomRenderPass = CustomRenderPass;
+
+		// We construct from scratch, rather than copying, as we don't want to copy interfaces attached to the view family
+		// (ScreenPercentageInterface, TemporalUpscalerInterface, etc), which can assert or double free if copied.  Those aren't
+		// relevant for custom render passes anyway.
+		FSceneViewFamily::ConstructionValues FamilyCVS(ViewFamily.RenderTarget, Scene, PassInput.bUseMainViewFamilyShowFlags ? ViewFamily.EngineShowFlags : PassInput.EngineShowFlags);
+
+		// Disable sky rendering, which is gated by the Atmosphere flag (unnecessary perf cost during base pass rendering)
+		FamilyCVS.EngineShowFlags.Atmosphere = false;
+
+		FCustomRenderPassInfo& CustomRenderPassInfo = CustomRenderPassInfos.Emplace_GetRef(FamilyCVS, ViewFamily);
+		CustomRenderPassInfo.CustomRenderPass = CustomRenderPass;
+		CustomRenderPassInfo.ViewFamily.Time = ViewFamily.Time;
+		CustomRenderPassInfo.ViewFamily.SetSceneRenderer(this);
 
 		FSceneViewInitOptions ViewInitOptions;
 		ViewInitOptions.SceneViewStateInterface = PassInput.ViewStateInterface;
@@ -2893,13 +2932,15 @@ FSceneRenderer::FSceneRenderer(const FSceneViewFamily* InViewFamily, FHitProxyCo
 		ViewInitOptions.ViewRotationMatrix = PassInput.ViewRotationMatrix;
 		ViewInitOptions.ProjectionMatrix = PassInput.ProjectionMatrix;
 		ViewInitOptions.bIsSceneCapture = PassInput.bIsSceneCapture;
-		ViewInitOptions.ViewFamily = &ViewFamily;
+		ViewInitOptions.ViewFamily = &CustomRenderPassInfo.ViewFamily;
 		ViewInitOptions.ViewActor = PassInput.ViewActor;
 		ViewInitOptions.ShowOnlyPrimitives = PassInput.ShowOnlyPrimitives;
 		ViewInitOptions.HiddenPrimitives = PassInput.HiddenPrimitives;
 
 		FSceneView NewView(ViewInitOptions);
-		FViewInfo* ViewInfo = &CustomRenderPassInfos[i].Views.Emplace_GetRef(&NewView);
+		FViewInfo* ViewInfo = &CustomRenderPassInfo.Views.Emplace_GetRef(&NewView);
+		CustomRenderPassInfo.ViewFamily.Views.Add(ViewInfo);
+		CustomRenderPassInfo.ViewFamily.AllViews.Add(ViewInfo);
 		// Must initialize to have a GPUScene connected to be able to collect dynamic primitives.
 		ViewInfo->DynamicPrimitiveCollector = FGPUScenePrimitiveCollector(&GPUSceneDynamicContext);
 		ViewInfo->bDisableQuerySubmissions = true;
@@ -3078,8 +3119,30 @@ FIntPoint FSceneRenderer::GetDesiredInternalBufferSize(const FSceneViewFamily& V
 
 FSceneRenderer::ERendererOutput FSceneRenderer::GetRendererOutput() const
 {
-	const bool bSceneCaptureDepthPrepass = Views[0].bIsSceneCapture && (ViewFamily.SceneCaptureSource == ESceneCaptureSource::SCS_SceneDepth || ViewFamily.SceneCaptureSource == ESceneCaptureSource::SCS_DeviceDepth);
-	return bSceneCaptureDepthPrepass && GSceneCaptureDepthPrepassOptimization ? ERendererOutput::DepthPrepassOnly : ERendererOutput::FinalSceneColor;
+	if (!Views[0].bIsSceneCapture)
+	{
+		return ERendererOutput::FinalSceneColor;
+	}
+	if (ViewFamily.SceneCaptureSource == ESceneCaptureSource::SCS_SceneDepth || ViewFamily.SceneCaptureSource == ESceneCaptureSource::SCS_DeviceDepth)
+	{
+		if (GSceneCaptureDepthPrepassOptimization)
+		{
+			return ERendererOutput::DepthPrepassOnly;
+		}
+	}
+	if (ViewFamily.SceneCaptureSource == ESceneCaptureSource::SCS_BaseColor || ViewFamily.SceneCaptureSource == ESceneCaptureSource::SCS_Normal)
+	{
+		// Forward shading in the deferred renderer requires shadows to run (skipped by BasePass optimization)
+		// HasRayTracedOverlay requires ray tracing to run -- appears to be used by path tracer, and ray tracing debug modes (ditto)
+		//
+		// TODO:  Could there be additional features that should disable the BasePass optimization?  We can add more cases if we run into bugs.
+		//        The above were just obvious from looking at the code.  Also, should some features disable DepthPrepassOnly above?
+		if (GSceneCaptureBasePassOptimization && !IsForwardShadingEnabled(Scene->GetShaderPlatform()) && !HasRayTracedOverlay(ViewFamily))
+		{
+			return ERendererOutput::BasePass;
+		}
+	}
+	return ERendererOutput::FinalSceneColor;
 }
 
 void FSceneRenderer::PrepareViewRectsForRendering(FRHICommandListImmediate& RHICmdList)
@@ -3679,6 +3742,13 @@ IVisibilityTaskData* FSceneRenderer::OnRenderBegin(FRDGBuilder& GraphBuilder)
 		InitializeSceneTexturesConfig(ViewFamily.SceneTexturesConfig, ViewFamily);
 		FSceneTexturesConfig& SceneTexturesConfig = GetActiveSceneTexturesConfig();
 		FSceneTexturesConfig::Set(SceneTexturesConfig);
+
+		// Custom render passes have their own view family structure, so they can have separate EngineShowFlags, so the SceneTexturesConfig
+		// needs to be copied.  The FSceneTextures structure itself is pointer shared, and doesn't need to be copied.
+		for (FCustomRenderPassInfo& CustomRenderPass : CustomRenderPassInfos)
+		{
+			CustomRenderPass.ViewFamily.SceneTexturesConfig = ViewFamily.SceneTexturesConfig;
+		}
 	
 		PrepareViewStateForVisibility(SceneTexturesConfig);
 	
@@ -3699,7 +3769,11 @@ IVisibilityTaskData* FSceneRenderer::OnRenderBegin(FRDGBuilder& GraphBuilder)
 			}
 		}
 	
-		LightFunctionAtlas::OnRenderBegin(LightFunctionAtlas, *Scene, Views, ViewFamily);
+		// Lighting is skipped when running ERendererOutput::DepthPrepassOnly or ERendererOutput::BasePass
+		if (GetRendererOutput() == ERendererOutput::FinalSceneColor)
+		{
+			LightFunctionAtlas::OnRenderBegin(LightFunctionAtlas, *Scene, Views, ViewFamily);
+		}
 	
 		GraphBuilder.RHICmdList.BeginScene();
 
diff --git a/Engine/Source/Runtime/Renderer/Private/SceneRendering.h b/Engine/Source/Runtime/Renderer/Private/SceneRendering.h
--- a/Engine/Source/Runtime/Renderer/Private/SceneRendering.h
+++ b/Engine/Source/Runtime/Renderer/Private/SceneRendering.h
@@ -1963,6 +1963,7 @@ class FViewFamilyInfo : public FSceneViewFamily
 {
 public:
 	explicit FViewFamilyInfo(const FSceneViewFamily& InViewFamily);
+	explicit FViewFamilyInfo(const FSceneViewFamily::ConstructionValues& CVS, const FViewFamilyInfo& MainViewFamily);
 	virtual ~FViewFamilyInfo();
 
 	FSceneTexturesConfig SceneTexturesConfig;
@@ -1970,31 +1971,34 @@ public:
 	/** Get scene textures associated with this view family -- asserts or checks that they have been initialized */
 	inline FSceneTextures& GetSceneTextures()
 	{
-		checkf(bIsSceneTexturesInitialized, TEXT("FSceneTextures was not initialized. Call FSceneTextures::InitializeViewFamily() first."));
-		return SceneTextures;
+		checkf(SceneTextures->bIsSceneTexturesInitialized, TEXT("FSceneTextures was not initialized. Call FSceneTextures::InitializeViewFamily() first."));
+		return *SceneTextures;
 	}
 
 	inline const FSceneTextures& GetSceneTextures() const
 	{
-		checkf(bIsSceneTexturesInitialized, TEXT("FSceneTextures was not initialized. Call FSceneTextures::InitializeViewFamily() first."));
-		return SceneTextures;
+		checkf(SceneTextures->bIsSceneTexturesInitialized, TEXT("FSceneTextures was not initialized. Call FSceneTextures::InitializeViewFamily() first."));
+		return *SceneTextures;
 	}
 
 	inline FSceneTextures* GetSceneTexturesChecked()
 	{
-		return bIsSceneTexturesInitialized ? &SceneTextures : nullptr;
+		return SceneTextures->bIsSceneTexturesInitialized ? SceneTextures : nullptr;
 	}
 
 	inline const FSceneTextures* GetSceneTexturesChecked() const
 	{
-		return bIsSceneTexturesInitialized ? &SceneTextures : nullptr;
+		return SceneTextures->bIsSceneTexturesInitialized ? SceneTextures : nullptr;
 	}
 
 private:
 	friend struct FMinimalSceneTextures;
 	friend struct FSceneTextures;
 
-	FSceneTextures SceneTextures;
+	// Structure may be pointed to by multiple FViewFamilyInfo during scene rendering, through CustomRenderPasses.  The Owner
+	// (pointed to by FSceneTextures) handles deleting the structure when the scene renderer is destroyed.  TRefCountPtr doesn't
+	// work, because the structure is also copied by value, and the copy constructor is disabled for reference counted structures.
+	FSceneTextures* SceneTextures;
 };
 
 struct FComputeLightGridOutput
@@ -2052,8 +2056,12 @@ public:
 	/** Information of a custom render pass that renders as part of the main renderer. */
 	struct FCustomRenderPassInfo
 	{
+		FCustomRenderPassInfo(const FSceneViewFamily::ConstructionValues& CVS, const FViewFamilyInfo& MainViewFamily);
+
 		/** Custom render pass that render as part of the main renderer. */
 		FCustomRenderPassBase* CustomRenderPass;
+		/** View family used by custom render pass.  NOT added to AllFamilies.  Required to allow different EngineShowFlags from main renderer ViewFamily. */
+		FViewFamilyInfo ViewFamily;
 		/** Views used to render the custom render pass. */
 		TArray<FViewInfo> Views;
 		FNaniteShadingCommands NaniteBasePassShadingCommands;
@@ -2333,6 +2341,7 @@ protected:
 	enum class ERendererOutput
 	{
 		DepthPrepassOnly,	// Only render depth prepass and its related code paths
+		BasePass,			// Render base pass (deferred renderer only, for mobile renderer equivalent to FinalSceneColor)
 		FinalSceneColor		// Render the whole pipeline
 	};
 
diff --git a/Engine/Source/Runtime/Renderer/Private/SceneTextures.cpp b/Engine/Source/Runtime/Renderer/Private/SceneTextures.cpp
--- a/Engine/Source/Runtime/Renderer/Private/SceneTextures.cpp
+++ b/Engine/Source/Runtime/Renderer/Private/SceneTextures.cpp
@@ -434,8 +434,10 @@ void InitializeSceneTexturesConfig(FSceneTexturesConfig& Config, const FSceneVie
 
 void FMinimalSceneTextures::InitializeViewFamily(FRDGBuilder& GraphBuilder, FViewFamilyInfo& ViewFamily)
 {
+	checkf(ViewFamily.SceneTextures->Owner == &ViewFamily, TEXT("Scene Textures should only be initialized by their owning view family -- possible duplicate initialization"));
+
 	const FSceneTexturesConfig& Config = ViewFamily.SceneTexturesConfig;
-	FSceneTextures& SceneTextures = ViewFamily.SceneTextures;
+	FSceneTextures& SceneTextures = *ViewFamily.SceneTextures;
 
 	checkf(Config.IsValid(), TEXT("Attempted to create scene textures with an empty config."));
 
@@ -495,7 +497,7 @@ void FMinimalSceneTextures::InitializeViewFamily(FRDGBuilder& GraphBuilder, FVie
 	// Custom Depth
 	SceneTextures.CustomDepth = FCustomDepthTextures::Create(GraphBuilder, Config.Extent, Config.ShaderPlatform);
 
-	ViewFamily.bIsSceneTexturesInitialized = true;
+	SceneTextures.bIsSceneTexturesInitialized = true;
 }
 
 FSceneTextureShaderParameters FMinimalSceneTextures::GetSceneTextureShaderParameters(ERHIFeatureLevel::Type FeatureLevel) const
@@ -515,7 +517,7 @@ FSceneTextureShaderParameters FMinimalSceneTextures::GetSceneTextureShaderParame
 void FSceneTextures::InitializeViewFamily(FRDGBuilder& GraphBuilder, FViewFamilyInfo& ViewFamily, FIntPoint FamilySize)
 {
 	const FSceneTexturesConfig& Config = ViewFamily.SceneTexturesConfig;
-	FSceneTextures& SceneTextures = ViewFamily.SceneTextures;
+	FSceneTextures& SceneTextures = *ViewFamily.SceneTextures;
 
 	FMinimalSceneTextures::InitializeViewFamily(GraphBuilder, ViewFamily);
 
diff --git a/Engine/Source/Runtime/Renderer/Private/SingleLayerWaterRendering.cpp b/Engine/Source/Runtime/Renderer/Private/SingleLayerWaterRendering.cpp
--- a/Engine/Source/Runtime/Renderer/Private/SingleLayerWaterRendering.cpp
+++ b/Engine/Source/Runtime/Renderer/Private/SingleLayerWaterRendering.cpp
@@ -531,7 +531,7 @@ static FSingleLayerWaterTileClassification ClassifyTiles(FRDGBuilder& GraphBuild
 	return Result;
 }
 
-FSingleLayerWaterPrePassResult* FDeferredShadingSceneRenderer::RenderSingleLayerWaterDepthPrepass(FRDGBuilder& GraphBuilder, const FSceneTextures& SceneTextures)
+FSingleLayerWaterPrePassResult* FDeferredShadingSceneRenderer::RenderSingleLayerWaterDepthPrepass(FRDGBuilder& GraphBuilder, TArrayView<FViewInfo> InViews, const FSceneTextures& SceneTextures)
 {
 	RDG_CSV_STAT_EXCLUSIVE_SCOPE(GraphBuilder, Water);
 	SCOPED_NAMED_EVENT(FDeferredShadingSceneRenderer_RenderSingleLayerWaterDepthPrepass, FColor::Emerald);
@@ -540,7 +540,7 @@ FSingleLayerWaterPrePassResult* FDeferredShadingSceneRenderer::RenderSingleLayer
 	RDG_GPU_STAT_SCOPE(GraphBuilder, SingleLayerWaterDepthPrepass);
 
 	FSingleLayerWaterPrePassResult* Result = GraphBuilder.AllocObject<FSingleLayerWaterPrePassResult>();
-	Result->ViewTileClassification.SetNum(Views.Num());
+	Result->ViewTileClassification.SetNum(InViews.Num());
 
 	FRDGTextureMSAA &OutDepthPrepassTexture = Result->DepthPrepassTexture;
 	// Create an identical copy of the main depth buffer
@@ -585,9 +585,9 @@ FSingleLayerWaterPrePassResult* FDeferredShadingSceneRenderer::RenderSingleLayer
 			TShaderMapRef<FCopyDepthPS> PixelShader(ShaderMap, PermutationVector);
 
 			FIntRect Viewport(0, 0, DepthPrepassTextureDesc.Extent.X, DepthPrepassTextureDesc.Extent.Y);
-			if (bOptimizedClear && Views.Num() == 1)
+			if (bOptimizedClear && InViews.Num() == 1)
 			{
-				Viewport = Views[0].ViewRect;
+				Viewport = InViews[0].ViewRect;
 			}
 
 			// Set depth test to always pass and stencil test to replace all pixels with zero, essentially also clearing stencil while doing the depth copy.
@@ -617,9 +617,9 @@ FSingleLayerWaterPrePassResult* FDeferredShadingSceneRenderer::RenderSingleLayer
 
 	const bool bRenderInParallel = GRHICommandList.UseParallelAlgorithms() && CVarParallelSingleLayerWaterPass.GetValueOnRenderThread() == 1;
 
-	for (int32 ViewIndex = 0; ViewIndex < Views.Num(); ++ViewIndex)
+	for (int32 ViewIndex = 0; ViewIndex < InViews.Num(); ++ViewIndex)
 	{
-		FViewInfo& View = Views[ViewIndex];
+		FViewInfo& View = InViews[ViewIndex];
 
 		if (!View.ShouldRenderView())
 		{
@@ -627,7 +627,7 @@ FSingleLayerWaterPrePassResult* FDeferredShadingSceneRenderer::RenderSingleLayer
 		}
 
 		RDG_GPU_MASK_SCOPE(GraphBuilder, View.GPUMask);
-		RDG_EVENT_SCOPE_CONDITIONAL(GraphBuilder, Views.Num() > 1, "View%d", ViewIndex);
+		RDG_EVENT_SCOPE_CONDITIONAL(GraphBuilder, InViews.Num() > 1, "View%d", ViewIndex);
 		View.BeginRenderView();
 
 		FSingleLayerWaterDepthPassParameters* PassParameters = GetSingleLayerWaterDepthPassParameters(GraphBuilder, View, OutDepthPrepassTexture.Target);
@@ -660,12 +660,12 @@ FSingleLayerWaterPrePassResult* FDeferredShadingSceneRenderer::RenderSingleLayer
 		}
 	}
 
-	AddResolveSceneDepthPass(GraphBuilder, Views, OutDepthPrepassTexture);
+	AddResolveSceneDepthPass(GraphBuilder, InViews, OutDepthPrepassTexture);
 
 	// Run classification pass.
-	for (int32 ViewIndex = 0; ViewIndex < Views.Num(); ++ViewIndex)
+	for (int32 ViewIndex = 0; ViewIndex < InViews.Num(); ++ViewIndex)
 	{
-		FViewInfo& View = Views[ViewIndex];
+		FViewInfo& View = InViews[ViewIndex];
 		if (UseSingleLayerWaterIndirectDraw(View.GetShaderPlatform()) && CVarWaterSingleLayerTiledComposite.GetValueOnRenderThread())
 		{
 			Result->ViewTileClassification[ViewIndex] = ClassifyTiles(GraphBuilder, View, SceneTextures, OutDepthPrepassTexture.Resolve);
@@ -803,6 +803,7 @@ END_SHADER_PARAMETER_STRUCT()
 
 void FDeferredShadingSceneRenderer::RenderSingleLayerWaterReflections(
 	FRDGBuilder& GraphBuilder,
+	TArrayView<FViewInfo> InViews,
 	const FSceneTextures& SceneTextures,
 	const FSceneWithoutWaterTextures& SceneWithoutWaterTextures,
 	const FSingleLayerWaterPrePassResult* SingleLayerWaterPrePassResult,
@@ -816,9 +817,9 @@ void FDeferredShadingSceneRenderer::RenderSingleLayerWaterReflections(
 	const FRDGSystemTextures& SystemTextures = FRDGSystemTextures::Get(GraphBuilder);
 	FRDGTextureRef SceneColorTexture = SceneTextures.Color.Resolve;
 
-	for (int32 ViewIndex = 0; ViewIndex < Views.Num(); ViewIndex++)
+	for (int32 ViewIndex = 0; ViewIndex < InViews.Num(); ViewIndex++)
 	{
-		FViewInfo& View = Views[ViewIndex];
+		FViewInfo& View = InViews[ViewIndex];
 
 		// Unfortunately, reflections cannot handle two views at once (yet?) - because of that, allow the secondary pass here.
 		// Note: not completely removing ShouldRenderView in case some other reason to not render it is valid.
@@ -828,7 +829,7 @@ void FDeferredShadingSceneRenderer::RenderSingleLayerWaterReflections(
 		}
 
 		RDG_GPU_MASK_SCOPE(GraphBuilder, View.GPUMask);
-		RDG_EVENT_SCOPE_CONDITIONAL(GraphBuilder, Views.Num() > 1, "View%d", ViewIndex);
+		RDG_EVENT_SCOPE_CONDITIONAL(GraphBuilder, InViews.Num() > 1, "View%d", ViewIndex);
 
 		FRDGTextureRef ReflectionsColor = nullptr;
 		FRDGTextureRef BlackDummyTexture = SystemTextures.Black;
@@ -1145,6 +1146,7 @@ void FDeferredShadingSceneRenderer::RenderSingleLayerWaterReflections(
 
 void FDeferredShadingSceneRenderer::RenderSingleLayerWater(
 	FRDGBuilder& GraphBuilder,
+	TArrayView<FViewInfo> InViews,
 	const FSceneTextures& SceneTextures,
 	const FSingleLayerWaterPrePassResult* SingleLayerWaterPrePassResult,
 	bool bShouldRenderVolumetricCloud,
@@ -1157,23 +1159,29 @@ void FDeferredShadingSceneRenderer::RenderSingleLayerWater(
 	SCOPED_NAMED_EVENT(SingleLayerWater, FColor::Emerald);
 
 	// Copy the texture to be available for the water surface to refract
-	SceneWithoutWaterTextures = AddCopySceneWithoutWaterPass(GraphBuilder, ViewFamily, Views, SceneTextures.Color.Resolve, SceneTextures.Depth.Resolve);
+	SceneWithoutWaterTextures = AddCopySceneWithoutWaterPass(GraphBuilder, ViewFamily, InViews, SceneTextures.Color.Resolve, SceneTextures.Depth.Resolve);
 
-	// Render height fog over the color buffer if it is allocated, e.g. SingleLayerWaterUsesSimpleShading is true.
-	if (!bIsCameraUnderWater && SceneWithoutWaterTextures.ColorTexture && ShouldRenderFog(ViewFamily))
-	{
-		RenderUnderWaterFog(GraphBuilder, SceneWithoutWaterTextures, SceneTextures.UniformBuffer);
-	}
-	if (!bIsCameraUnderWater && SceneWithoutWaterTextures.ColorTexture && bShouldRenderVolumetricCloud)
+	// Check if this is depth or base pass only renderer, where the final scene color isn't relevant, and we don't need fog, clouds, or reflections
+	bool bFinalSceneColor = !InViews[0].CustomRenderPass && GetRendererOutput() == ERendererOutput::FinalSceneColor;
+
+	if (bFinalSceneColor)
 	{
-		// This path is only taken when rendering the clouds in a render target that can be composited
-		ComposeVolumetricRenderTargetOverSceneUnderWater(GraphBuilder, Views, SceneWithoutWaterTextures, SceneTextures);
+		// Render height fog over the color buffer if it is allocated, e.g. SingleLayerWaterUsesSimpleShading is true.
+		if (!bIsCameraUnderWater && SceneWithoutWaterTextures.ColorTexture && ShouldRenderFog(ViewFamily))
+		{
+			RenderUnderWaterFog(GraphBuilder, SceneWithoutWaterTextures, SceneTextures.UniformBuffer);
+		}
+		if (!bIsCameraUnderWater && SceneWithoutWaterTextures.ColorTexture && bShouldRenderVolumetricCloud)
+		{
+			// This path is only taken when rendering the clouds in a render target that can be composited
+			ComposeVolumetricRenderTargetOverSceneUnderWater(GraphBuilder, InViews, SceneWithoutWaterTextures, SceneTextures);
+		}
 	}
 
-	RenderSingleLayerWaterInner(GraphBuilder, SceneTextures, SceneWithoutWaterTextures, SingleLayerWaterPrePassResult);
+	RenderSingleLayerWaterInner(GraphBuilder, InViews, SceneTextures, SceneWithoutWaterTextures, SingleLayerWaterPrePassResult);
 
 	// No SSR or composite needed in Forward. Reflections are applied in the WaterGBuffer pass.
-	if (!IsForwardShadingEnabled(ShaderPlatform))
+	if (!IsForwardShadingEnabled(ShaderPlatform) && bFinalSceneColor)
 	{
 		// Reflection composite expects the depth buffer in FSceneTextures to contain water but the swap of the main depth buffer with the water prepass depth buffer
 		// is only done at the call site after this function returns (for visibility and to keep SceneTextures const), so we need to swap the depth buffers on an internal copy.
@@ -1186,7 +1194,7 @@ void FDeferredShadingSceneRenderer::RenderSingleLayerWater(
 		}
 
 		// If supported render SSR, the composite pass in non deferred and/or under water effect.
-		RenderSingleLayerWaterReflections(GraphBuilder, SceneTexturesInternal, SceneWithoutWaterTextures, SingleLayerWaterPrePassResult, LumenFrameTemporaries);
+		RenderSingleLayerWaterReflections(GraphBuilder, InViews, SceneTexturesInternal, SceneWithoutWaterTextures, SingleLayerWaterPrePassResult, LumenFrameTemporaries);
 	}
 }
 
@@ -1226,6 +1234,7 @@ bool IsVSMTranslucentHighQualityEnabled();
 
 void FDeferredShadingSceneRenderer::RenderSingleLayerWaterInner(
 	FRDGBuilder& GraphBuilder,
+	TArrayView<FViewInfo> InViews,
 	const FSceneTextures& SceneTextures,
 	const FSceneWithoutWaterTextures& SceneWithoutWaterTextures,
 	const FSingleLayerWaterPrePassResult* SingleLayerWaterPrePassResult)
@@ -1262,9 +1271,9 @@ void FDeferredShadingSceneRenderer::RenderSingleLayerWaterInner(
 		DepthStencilBinding = FDepthStencilBinding(SceneTextures.Depth.Target, ERenderTargetLoadAction::ELoad, ERenderTargetLoadAction::ELoad, FExclusiveDepthStencil::DepthWrite_StencilNop);
 	}
 
-	for (int32 ViewIndex = 0; ViewIndex < Views.Num(); ++ViewIndex)
+	for (int32 ViewIndex = 0; ViewIndex < InViews.Num(); ++ViewIndex)
 	{
-		FViewInfo& View = Views[ViewIndex];
+		FViewInfo& View = InViews[ViewIndex];
 
 		if (!View.ShouldRenderView())
 		{
@@ -1272,7 +1281,7 @@ void FDeferredShadingSceneRenderer::RenderSingleLayerWaterInner(
 		}
 
 		RDG_GPU_MASK_SCOPE(GraphBuilder, View.GPUMask);
-		RDG_EVENT_SCOPE_CONDITIONAL(GraphBuilder, Views.Num() > 1, "View%d", ViewIndex);
+		RDG_EVENT_SCOPE_CONDITIONAL(GraphBuilder, InViews.Num() > 1, "View%d", ViewIndex);
 		View.BeginRenderView();
 
 		FSingleLayerWaterPassUniformParameters& SLWUniformParameters = *GraphBuilder.AllocParameters<FSingleLayerWaterPassUniformParameters>();
@@ -1347,7 +1356,7 @@ void FDeferredShadingSceneRenderer::RenderSingleLayerWaterInner(
 
 	if (!bHasDepthPrepass)
 	{
-		AddResolveSceneDepthPass(GraphBuilder, Views, SceneTextures.Depth);
+		AddResolveSceneDepthPass(GraphBuilder, InViews, SceneTextures.Depth);
 	}
 }
 
