Fix "Hit Count Per Instance" ray tracing visualization.
- use InstanceIndex() in TLAS + PrimitiveID (aka PersistentPrimitiveIndex).
- also changes behaviour for ISMs since hits are now counted per-instance instead of per-primitive.

diff --git a/Engine/Shaders/Private/RayTracing/RayTracingDebug.usf b/Engine/Shaders/Private/RayTracing/RayTracingDebug.usf
--- a/Engine/Shaders/Private/RayTracing/RayTracingDebug.usf
+++ b/Engine/Shaders/Private/RayTracing/RayTracingDebug.usf
@@ -211,7 +211,7 @@ float4 DebugRayTracingInstance(FRayDesc Ray, out float OutDepth)
 			Result.rgb = ColorMapTurbo(Payload.TriangleHitCountPerRay / TriangleHitCountMaxThreshold);
 			break;
 		case RAY_TRACING_DEBUG_VIZ_HITCOUNT_PER_INSTANCE:
-			Result.rgb = ColorMapTurbo(RayTracingDebugHitStatsUniformBuffer.HitStatsOutput[Payload.ScenePrimitiveIndex].Count / (float)TriangleHitCountPerInstanceMaxThreshold);
+			Result.rgb = ColorMapTurbo(RayTracingDebugHitStatsUniformBuffer.HitStatsOutput[Payload.InstanceIndex].Count / (float)TriangleHitCountPerInstanceMaxThreshold);
 			break;
 		case RAY_TRACING_DEBUG_VIZ_INSTANCE_OVERLAP:
 			Result.rgb = 1.0; // We only use fake lighting as instance heatmap is blended in
diff --git a/Engine/Shaders/Private/RayTracing/RayTracingDebugHitStatsCHS.usf b/Engine/Shaders/Private/RayTracing/RayTracingDebugHitStatsCHS.usf
--- a/Engine/Shaders/Private/RayTracing/RayTracingDebugHitStatsCHS.usf
+++ b/Engine/Shaders/Private/RayTracing/RayTracingDebugHitStatsCHS.usf
@@ -12,11 +12,12 @@ RAY_TRACING_ENTRY_ANY_HIT(RayTracingDebugHitStatsAHS,
 	FRayTracingDebugPayload, Payload,
 	FRayTracingIntersectionAttributes, Attributes)
 {
+	const uint Index = InstanceIndex();
+
 	const uint GPUSceneInstanceId = GetInstanceUserData();
 	const FInstanceSceneData InstanceSceneData = GetInstanceSceneData(GPUSceneInstanceId, Scene.GPUScene.InstanceDataSOAStride);
-	const uint Index = InstanceSceneData.PrimitiveId;
 
-	RayTracingDebugHitStatsUniformBuffer.HitStatsOutput[Index].PrimitiveID = Index;
+	RayTracingDebugHitStatsUniformBuffer.HitStatsOutput[Index].PrimitiveID = InstanceSceneData.PrimitiveId;
 	InterlockedAdd(RayTracingDebugHitStatsUniformBuffer.HitStatsOutput[Index].Count, 1);
 }
 
diff --git a/Engine/Source/Runtime/Renderer/Private/RayTracing/RayTracingDebug.cpp b/Engine/Source/Runtime/Renderer/Private/RayTracing/RayTracingDebug.cpp
--- a/Engine/Source/Runtime/Renderer/Private/RayTracing/RayTracingDebug.cpp
+++ b/Engine/Source/Runtime/Renderer/Private/RayTracing/RayTracingDebug.cpp
@@ -655,7 +655,6 @@ void BindRayTracingDebugHitStatsCHSMaterialBindings(FRHICommandList& RHICmdList,
 		Binding.SegmentIndex = MeshCommand.GeometrySegmentIndex;
 		Binding.UniformBuffers = HelperBinding.UniformBufferArray;
 		Binding.NumUniformBuffers = HelperBinding.NumUniformBuffers;
-		Binding.UserData = VisibleMeshCommand.InstanceIndex;
 
 		Bindings[BindingIndex] = Binding;
 		BindingIndex++;
@@ -990,11 +989,18 @@ static void PrintTopKMostHitMessage(FRDGBuilder& GraphBuilder, const FScene* Sce
 		for (int32 HitStatsID = 0; HitStatsID < NumPrimitives; ++HitStatsID)
 		{
 			const uint32 PrimitiveID = HitStatsArray[HitStatsID].PrimitiveID;
-			if (PrimitiveID >= (uint32)Scene->Primitives.Num())
+
+			FPersistentPrimitiveIndex PersistentPrimitiveIndex;
+			PersistentPrimitiveIndex.Index = PrimitiveID;
+
+			FPrimitiveSceneInfo* SceneInfo = Scene->GetPrimitiveSceneInfo(PersistentPrimitiveIndex);
+
+			if (SceneInfo == nullptr)
 			{
 				continue;
 			}
-			const FString OwnerName = Scene->Primitives[PrimitiveID]->GetFullnameForDebuggingOnly();
+
+			const FString OwnerName = SceneInfo->GetFullnameForDebuggingOnly();
 			const uint32 NameOffset = SelectedNames.Num();
 			const uint32 NameLength = OwnerName.Len();
 			for (TCHAR C : OwnerName)
@@ -1089,9 +1095,8 @@ static FRDGBufferRef RayTracingPerformHitStatsPerPrimitive(FRDGBuilder& GraphBui
 
 	FShaderBindingTableRHIRef HitStatsSBT = RHICreateShaderBindingTable(SBTInitializer);
 
-	const uint32 NumPrimitives = FMath::Max<uint32>(Scene->Primitives.Num(),
-		(uint32)CVarRayTracingDebugHitCountTopKHits.GetValueOnRenderThread());
-	FRDGBufferDesc HitStatsPerPrimitiveBufferDesc = FRDGBufferDesc::CreateStructuredDesc(sizeof(FRayTracingHitStatsEntry), NumPrimitives);
+	const uint32 NumInstancesInTLAS = FMath::Max(RayTracingScene.GetRHIRayTracingSceneChecked()->GetInitializer().NumNativeInstancesPerLayer[uint32(ERayTracingSceneLayer::Base)], (uint32)CVarRayTracingDebugHitCountTopKHits.GetValueOnRenderThread());
+	FRDGBufferDesc HitStatsPerPrimitiveBufferDesc = FRDGBufferDesc::CreateStructuredDesc(sizeof(FRayTracingHitStatsEntry), NumInstancesInTLAS);
 	HitStatsPerPrimitiveBufferDesc.Usage = EBufferUsageFlags(HitStatsPerPrimitiveBufferDesc.Usage | BUF_SourceCopy);
 	FRDGBufferRef HitStatsBuffer = GraphBuilder.CreateBuffer(HitStatsPerPrimitiveBufferDesc, TEXT("RayTracingDebug.HitStatsBuffer"));
 	
@@ -1107,7 +1112,7 @@ static FRDGBufferRef RayTracingPerformHitStatsPerPrimitive(FRDGBuilder& GraphBui
 	RayGenParameters->SceneUniformBuffer = GetSceneUniformBufferRef(GraphBuilder, View); // TODO: use a separate params structure
 	RayGenParameters->RayTracingDebugHitStatsUniformBuffer = DebugHitStatsUniformBuffer;
 
-	AddClearUAVPass(GraphBuilder, DebugHitStatsUniformBufferParameters->HitStatsOutput, 0, ERDGPassFlags::Compute);
+	AddClearUAVPass(GraphBuilder, DebugHitStatsUniformBufferParameters->HitStatsOutput, 0);
 
 	FIntRect ViewRect = View.ViewRect;
 	GraphBuilder.AddPass(
