Removed unecessary NANITE_RAY_TRACING permutation of FRayTracingDebugHitStatsCHS.

diff --git a/Engine/Source/Runtime/Renderer/Private/RayTracing/RayTracingDebug.cpp b/Engine/Source/Runtime/Renderer/Private/RayTracing/RayTracingDebug.cpp
--- a/Engine/Source/Runtime/Renderer/Private/RayTracing/RayTracingDebug.cpp
+++ b/Engine/Source/Runtime/Renderer/Private/RayTracing/RayTracingDebug.cpp
@@ -283,25 +283,11 @@ class FRayTracingDebugHitStatsCHS : public FGlobalShader
 	END_SHADER_PARAMETER_STRUCT()
 public:
 
-	class FNaniteRayTracing : SHADER_PERMUTATION_BOOL("NANITE_RAY_TRACING");
-	using FPermutationDomain = TShaderPermutationDomain<FNaniteRayTracing>;
-
 	static bool ShouldCompilePermutation(const FGlobalShaderPermutationParameters& Parameters)
 	{
 		return ShouldCompileRayTracingShadersForProject(Parameters.Platform);
 	}
 
-	static void ModifyCompilationEnvironment(const FGlobalShaderPermutationParameters& Parameters, FShaderCompilerEnvironment& OutEnvironment)
-	{
-		FGlobalShader::ModifyCompilationEnvironment(Parameters, OutEnvironment);
-
-		FPermutationDomain PermutationVector(Parameters.PermutationId);
-		if (PermutationVector.Get<FNaniteRayTracing>())
-		{
-			OutEnvironment.SetDefine(TEXT("VF_SUPPORTS_PRIMITIVE_SCENE_DATA"), 1);
-		}
-	}
-
 	static ERayTracingPayloadType GetRayTracingPayloadType(const int32 PermutationId)
 	{
 		return ERayTracingPayloadType::RayTracingDebug;
@@ -636,18 +622,14 @@ void BindRayTracingDebugHitStatsCHSMaterialBindings(FRHICommandList& RHICmdList,
 
 	FRayTracingDebugHitStatsCHS::FPermutationDomain PermutationVector;
 
-	PermutationVector.Set<FRayTracingDebugHitStatsCHS::FNaniteRayTracing>(false);
 	FBinding ShaderBinding = SetupBinding(PermutationVector);
 
-	PermutationVector.Set<FRayTracingDebugHitStatsCHS::FNaniteRayTracing>(true);
-	FBinding ShaderBindingNaniteRT = SetupBinding(PermutationVector);
-
 	uint32 BindingIndex = 0;
 	for (const FVisibleRayTracingMeshCommand VisibleMeshCommand : View.VisibleRayTracingMeshCommands)
 	{
 		const FRayTracingMeshCommand& MeshCommand = *VisibleMeshCommand.RayTracingMeshCommand;
 
-		const FBinding& HelperBinding = MeshCommand.IsUsingNaniteRayTracing() ? ShaderBindingNaniteRT : ShaderBinding;
+		const FBinding& HelperBinding = ShaderBinding;
 
 		FRayTracingLocalShaderBindings Binding = {};
 		Binding.ShaderIndexInPipeline = HelperBinding.ShaderIndexInPipeline;
@@ -1071,13 +1053,9 @@ static FRDGBufferRef RayTracingPerformHitStatsPerPrimitive(FRDGBuilder& GraphBui
 
 	FRayTracingDebugHitStatsCHS::FPermutationDomain PermutationVector;
 
-	PermutationVector.Set<FRayTracingDebugHitStatsCHS::FNaniteRayTracing>(false);
 	auto HitGroupShader = View.ShaderMap->GetShader<FRayTracingDebugHitStatsCHS>(PermutationVector);
 
-	PermutationVector.Set<FRayTracingDebugHitStatsCHS::FNaniteRayTracing>(true);
-	auto HitGroupShaderNaniteRT = View.ShaderMap->GetShader<FRayTracingDebugHitStatsCHS>(PermutationVector);
-
-	FRHIRayTracingShader* HitGroupTable[] = { HitGroupShader.GetRayTracingShader(), HitGroupShaderNaniteRT.GetRayTracingShader() };
+	FRHIRayTracingShader* HitGroupTable[] = { HitGroupShader.GetRayTracingShader() };
 	Initializer.SetHitGroupTable(HitGroupTable);
 	Initializer.bAllowHitGroupIndexing = true; // Required for stable output using GetBaseInstanceIndex().
 
