Path Tracer: Add support for cloud shadowing flag, split callable shader clouds into their own permutation

The clouds do not cast shadows by default in the deferred renderer, so match this behavior in the path tracer for improved compatibility.

Callable shader based clouds are currently very expensive, so split the use of those into an experimental permutation so we don't grow the live state in the common case.

Cleanup some console variables that should be bools and make the helptext more consistent

#rnx

diff --git a/Engine/Shaders/Private/PathTracing/Light/PathTracingLightCommon.ush b/Engine/Shaders/Private/PathTracing/Light/PathTracingLightCommon.ush
--- a/Engine/Shaders/Private/PathTracing/Light/PathTracingLightCommon.ush
+++ b/Engine/Shaders/Private/PathTracing/Light/PathTracingLightCommon.ush
@@ -342,6 +342,11 @@ bool CastsVolumeShadow(int LightId)
 	return (SceneLights[LightId].Flags & PATHTRACER_FLAG_CAST_VOL_SHADOW_MASK) != 0;
 }
 
+bool CastsCloudShadow(int LightId)
+{
+	return (SceneLights[LightId].Flags & PATHTRACER_FLAG_CAST_CLOUD_SHADOW_MASK) != 0;
+}
+
 bool HasRectTexture(int LightId)
 {
 	return (SceneLights[LightId].Flags & PATHTRACER_FLAG_HAS_RECT_TEXTURE_MASK) != 0;
diff --git a/Engine/Shaders/Private/PathTracing/PathTracingCore.ush b/Engine/Shaders/Private/PathTracing/PathTracingCore.ush
--- a/Engine/Shaders/Private/PathTracing/PathTracingCore.ush
+++ b/Engine/Shaders/Private/PathTracing/PathTracingCore.ush
@@ -487,7 +487,7 @@ FPackedPathTracingPayload TraceTransparentRay(inout FPathState PathState, int Bo
 			PathState.Ray.GetNativeDesc(),
 			PackedPayload);
 
-		FVolumeIntersectionList VolumeIsectList = VolumeIntersect(PathState.Ray.Origin, PathState.Ray.Direction, PathState.Ray.TMin, PathState.Ray.TMax, PathState.GetPathRoughness());
+		FVolumeIntersectionList VolumeIsectList = VolumeIntersect(PathState.Ray.Origin, PathState.Ray.Direction, PathState.Ray.TMin, PathState.Ray.TMax, PathState.GetPathRoughness(), VolumeFlags);
 		if (VolumeIsectList.HitBlocker() && (PackedPayload.IsMiss() || VolumeIsectList.BlockerHitT < PackedPayload.HitT))
 		{
 			// we didn't hit any real geometry, but we did hit the volume's blocker geometry
@@ -800,7 +800,7 @@ FPackedPathTracingPayload TraceTransparentRay(inout FPathState PathState, int Bo
 	return Payload;
 }
 
-float3 TraceTransparentVisibilityRay(FRayDesc Ray, float PathRoughness, float PreviousPathRoughness, uint MissShaderIndex, bool bCastShadows, inout RandomSequence RandSequence)
+float3 TraceTransparentVisibilityRay(FRayDesc Ray, float PathRoughness, float PreviousPathRoughness, uint MissShaderIndex, bool bCastShadows, bool bCloudCastShadows, inout RandomSequence RandSequence)
 {
 	if (!bCastShadows && MissShaderIndex == 0)
 	{
@@ -841,7 +841,14 @@ float3 TraceTransparentVisibilityRay(FRayDesc Ray, float PathRoughness, float Pr
 	if (bCastShadows && any(Throughput > 0))
 	{
 		// NOTE: use the path roughness before scattering off the current vertex (so we can have clouds casting shadows onto diffuse surfaces)
-		FVolumeIntersectionList VolumeIsectList = VolumeIntersect(Ray.Origin, Ray.Direction, Ray.TMin, Ray.TMax, PreviousPathRoughness);
+		uint VolumeMask = VolumeFlags;
+		if (!bCloudCastShadows)
+		{
+			// if skipping clouds on shadows, clear the contribution here
+			VolumeMask &= ~PATH_TRACER_VOLUME_ENABLE_CLOUDS;
+		}
+
+		FVolumeIntersectionList VolumeIsectList = VolumeIntersect(Ray.Origin, Ray.Direction, Ray.TMin, Ray.TMax, PreviousPathRoughness, VolumeMask);
 		if (VolumeIsectList.HitBlocker())
 		{
 			// blockers are opaque -- done!
@@ -1429,7 +1436,7 @@ bool PathTracingKernel(inout FPathState PathState, int Bounce)
 
 #if 0
 	// visualize the shadow ray handling directly
-	PathState.Radiance = TraceTransparentVisibilityRay(PathState.Ray, 0.0, 0, true, PathState.RandSequence);
+	PathState.Radiance = TraceTransparentVisibilityRay(PathState.Ray, 0.0, 0, true, true, PathState.RandSequence);
 	return false;
 #endif
 
@@ -1501,6 +1508,7 @@ bool PathTracingKernel(inout FPathState PathState, int Bounce)
 				float3 LightRandValue = RandomSequence_GenerateSample3D(PathState.RandSequence);
 
 				const bool bCastShadows = CastsVolumeShadow(LightSampleLightId);
+				const bool bCloudCastShadows = true; // volumes should get shadows from clouds
 				const uint MissShaderIndex = GetLightMissShaderIndex(LightSampleLightId);
 				// compute direct light sampling?
 				if (MISMode != 0)
@@ -1524,7 +1532,7 @@ bool PathTracingKernel(inout FPathState PathState, int Bounce)
 							{
 								LightContrib *= MISWeightPower(LightSample.Pdf, MaterialEval.Pdf);
 							}
-							LightContrib *= TraceTransparentVisibilityRay(LightRay, 1.0, PathState.GetPathRoughness(), MissShaderIndex, bCastShadows, PathState.RandSequence);
+							LightContrib *= TraceTransparentVisibilityRay(LightRay, 1.0, PathState.GetPathRoughness(), MissShaderIndex, bCastShadows, bCloudCastShadows, PathState.RandSequence);
 
 							AccumulateRadiance(PathState.Radiance, LightContrib, bIsCameraRay);
 						}
@@ -1550,7 +1558,7 @@ bool PathTracingKernel(inout FPathState PathState, int Bounce)
 								LightContrib *= MISWeightPower(MaterialSample.Pdf, LightResult.Pdf);
 							}
 							MaterialRay.TMax = LightResult.HitT;
-							LightContrib *= TraceTransparentVisibilityRay(MaterialRay, 1.0, PathState.GetPathRoughness(), MissShaderIndex, bCastShadows, PathState.RandSequence);
+							LightContrib *= TraceTransparentVisibilityRay(MaterialRay, 1.0, PathState.GetPathRoughness(), MissShaderIndex, bCastShadows, bCloudCastShadows, PathState.RandSequence);
 							AccumulateRadiance(PathState.Radiance, LightContrib, bIsCameraRay);
 						}
 					}
@@ -1693,8 +1701,9 @@ bool PathTracingKernel(inout FPathState PathState, int Bounce)
 						LightContrib *= MISWeightPower(LightSample.Pdf, MaterialEval.Pdf);
 					}
 					const bool bCastShadows = bIsVolumeSample ? CastsVolumeShadow(LightId) : CastsShadow(LightId);
+					const bool bCloudCastShadows = bIsVolumeSample ? true : CastsCloudShadow(LightId);
 					const uint MissShaderIndex = GetLightMissShaderIndex(LightId);
-					LightContrib *= TraceTransparentVisibilityRay(LightRay, AvgRoughness, PathState.GetPathRoughness(), MissShaderIndex, bCastShadows, PathState.RandSequence);
+					LightContrib *= TraceTransparentVisibilityRay(LightRay, AvgRoughness, PathState.GetPathRoughness(), MissShaderIndex, bCastShadows, bCloudCastShadows, PathState.RandSequence);
 
 					AccumulateRadiance(PathState.Radiance, LightContrib, bIsCameraRay);
 				}
@@ -1837,9 +1846,10 @@ bool PathTracingKernel(inout FPathState PathState, int Bounce)
 				FRayDesc LightRay = PathState.Ray;
 				LightRay.TMax = LightResult.HitT;
 				const bool bCastShadows = bIsVolumeSample ? CastsVolumeShadow(LightId) : CastsShadow(LightId);
+				const bool bCloudCastShadows = bIsVolumeSample ? true : CastsCloudShadow(LightId);
 				const uint MissShaderIndex = GetLightMissShaderIndex(LightId);
 				const float ShadowRayRoughness = ApproximateCaustics ? PathState.GetPathRoughness() : 0.0;
-				LightContrib *= TraceTransparentVisibilityRay(LightRay, ShadowRayRoughness, PreviousPathRoughness, MissShaderIndex, bCastShadows, PathState.RandSequence);
+				LightContrib *= TraceTransparentVisibilityRay(LightRay, ShadowRayRoughness, PreviousPathRoughness, MissShaderIndex, bCastShadows, bCloudCastShadows, PathState.RandSequence);
 				AccumulateRadiance(PathState.Radiance, LightContrib, bIsCameraRay);
 			}
 		}
diff --git a/Engine/Shaders/Private/PathTracing/Volume/PathTracingClouds.ush b/Engine/Shaders/Private/PathTracing/Volume/PathTracingClouds.ush
--- a/Engine/Shaders/Private/PathTracing/Volume/PathTracingClouds.ush
+++ b/Engine/Shaders/Private/PathTracing/Volume/PathTracingClouds.ush
@@ -203,21 +203,9 @@ FVolumeShadedResult CloudsGetDensity(float3 TranslatedWorldPos)
 	}
 
 	uint4 Data = 0;
-#if PLATFORM_SUPPORTS_CALLABLE_SHADERS
-	const bool bUseCloudMap = (VolumeFlags & PATH_TRACER_VOLUME_USE_CLOUD_MAP) != 0;
-	if (bUseCloudMap)
-#endif
-	{
-		// Main path -- via baked cloud map
-		const float3 P = mul(GetCloudClipBasis(), TranslatedWorldPos * CM_TO_SKY_UNIT); // WorldToTangent
-		const float3 UVW = GetUVWFromCloudClipPos(P + float3(0, 0, CloudClipRadiusKm));
-		const int3 GridCoord = int3(floor(float3(CloudMapResolution) * UVW));
-		Data = CloudMap.Load(int4(GridCoord, 0));
-	}
-#if PLATFORM_SUPPORTS_CALLABLE_SHADERS
-	else
+#if PATH_TRACER_USE_CLOUD_SHADER
 	{
-		// Main path -- invoke the actual callable shader
+		// Reference path -- invoke the actual callable shader
 		FPackedPathTracingPayload Payload = (FPackedPathTracingPayload)0;
 		Payload.SetVolumetricCallableShaderInput(
 			DFFastSubtract(TranslatedWorldPos, PrimaryView.PreViewTranslation),
@@ -226,6 +214,14 @@ FVolumeShadedResult CloudsGetDensity(float3 TranslatedWorldPos)
 		CallShader(CloudCallableShaderId, Payload);
 		Data = Payload.GetVolumetricCallableShaderOutput();
 	}
+#else
+	{
+		// Main path -- via baked cloud map
+		const float3 P = mul(GetCloudClipBasis(), TranslatedWorldPos * CM_TO_SKY_UNIT); // WorldToTangent
+		const float3 UVW = GetUVWFromCloudClipPos(P + float3(0, 0, CloudClipRadiusKm));
+		const int3 GridCoord = int3(floor(float3(CloudMapResolution) * UVW));
+		Data = CloudMap.Load(int4(GridCoord, 0));
+	}
 #endif
 	const FSampleCloudMaterialResult CloudResult = DecodeCloudResult(Data);
 	Result.SigmaT        = CloudResult.Extinction;
diff --git a/Engine/Shaders/Private/PathTracing/Volume/PathTracingVolume.ush b/Engine/Shaders/Private/PathTracing/Volume/PathTracingVolume.ush
--- a/Engine/Shaders/Private/PathTracing/Volume/PathTracingVolume.ush
+++ b/Engine/Shaders/Private/PathTracing/Volume/PathTracingVolume.ush
@@ -14,23 +14,23 @@ int MaxRaymarchSteps;
 
 
 // Given an input ray, figure out where the provided volume intersects it
-FVolumeIntersectionList VolumeIntersect(float3 Origin, float3 Direction, float TMin, float TMax, float PathRoughness)
+FVolumeIntersectionList VolumeIntersect(float3 Origin, float3 Direction, float TMin, float TMax, float PathRoughness, uint VolumeMask)
 {
 	// NOTE: only this call needs to be guarded by enable flags at the moment since the others can't be reached if this returns nothing
 	FVolumeIntersectionList Result = CreateEmptyVolumeIntersectionList();
-	if (VolumeFlags & PATH_TRACER_VOLUME_ENABLE_ATMOSPHERE)
+	if (VolumeMask & PATH_TRACER_VOLUME_ENABLE_ATMOSPHERE)
 	{
 		Result.Add(VOLUMEID_ATMOSPHERE, AtmosphereIntersect(Origin, Direction, TMin, TMax));
 	}
-	if (VolumeFlags & PATH_TRACER_VOLUME_ENABLE_CLOUDS)
+	if (VolumeMask & PATH_TRACER_VOLUME_ENABLE_CLOUDS)
 	{
 		Result.Add(VOLUMEID_CLOUDS, CloudsIntersect(Origin, Direction, TMin, TMax, PathRoughness));
 	}
-	if (VolumeFlags & PATH_TRACER_VOLUME_ENABLE_FOG)
+	if (VolumeMask & PATH_TRACER_VOLUME_ENABLE_FOG)
 	{
 		Result.Add(VOLUMEID_FOG, FogIntersect(Origin, Direction, TMin, TMax));
 	}
-	if (VolumeFlags & PATH_TRACER_VOLUME_ENABLE_HETEROGENEOUS_VOLUMES)
+	if (VolumeMask & PATH_TRACER_VOLUME_ENABLE_HETEROGENEOUS_VOLUMES)
 	{
 		Result.Add(VOLUMEID_HETEROGENEOUS_VOLUMES, HeterogeneousVolumesIntersect(Origin, Direction, TMin, TMax));
 	}
diff --git a/Engine/Shaders/Shared/PathTracingDefinitions.h b/Engine/Shaders/Shared/PathTracingDefinitions.h
--- a/Engine/Shaders/Shared/PathTracingDefinitions.h
+++ b/Engine/Shaders/Shared/PathTracingDefinitions.h
@@ -19,7 +19,8 @@
 #define PATHTRACING_LIGHT_RECT								(4 << 6)
 #define PATHTRACER_FLAG_CAST_SHADOW_MASK 					(1 << 9)
 #define PATHTRACER_FLAG_CAST_VOL_SHADOW_MASK 				(1 << 10)
-#define PATHTRACER_FLAG_HAS_RECT_TEXTURE_MASK				(1 << 11)
+#define PATHTRACER_FLAG_CAST_CLOUD_SHADOW_MASK              (1 << 11)
+#define PATHTRACER_FLAG_HAS_RECT_TEXTURE_MASK				(1 << 12)
 
 #define PATHTRACER_MASK_CAMERA								0x01			// opaque and alpha tested meshes and particles as a whole (primary ray) excluding hairs
 #define PATHTRACER_MASK_HAIR_CAMERA							0x02			// For primary ray tracing against hair
diff --git a/Engine/Source/Runtime/Renderer/Private/PathTracing.cpp b/Engine/Source/Runtime/Renderer/Private/PathTracing.cpp
--- a/Engine/Source/Runtime/Renderer/Private/PathTracing.cpp
+++ b/Engine/Source/Runtime/Renderer/Private/PathTracing.cpp
@@ -39,18 +39,18 @@ TAutoConsoleVariable<int32> CVarPathTracing(
 #include "EnvironmentComponentsFlags.h"
 #include <limits>
 
-TAutoConsoleVariable<int32> CVarPathTracingExperimental(
+TAutoConsoleVariable<bool> CVarPathTracingExperimental(
 	TEXT("r.PathTracing.Experimental"),
-	0,
-	TEXT("Enables some experimental features of the path tracing renderer that require compiling additional permutations of the path tracer."),
+	false,
+	TEXT("Enables some experimental features of the path tracing renderer that require compiling additional permutations of the path tracer. (default: false)"),
 	ECVF_RenderThreadSafe | ECVF_ReadOnly
 );
 
 
-TAutoConsoleVariable<int32> CVarPathTracingCompaction(
+TAutoConsoleVariable<bool> CVarPathTracingCompaction(
 	TEXT("r.PathTracing.Compaction"),
-	1,
-	TEXT("Enables path compaction to improve GPU occupancy for the path tracer (default: 1 (enabled))"),
+	true,
+	TEXT("Enables path compaction to improve GPU occupancy for the path tracer. Requires r.PathTracing.Experimental=true to modify.  (default: true (enabled))"),
 	ECVF_RenderThreadSafe
 );
 
@@ -269,21 +269,24 @@ TAutoConsoleVariable<float> CVarPathTracingCloudRoughnessCutoff(
 TAutoConsoleVariable<bool> CVarPathTracingCloudMapEnabled(
 	TEXT("r.PathTracing.CloudMapEnable"),
 	true,
-	TEXT("If true, clouds will be voxelized into a texture for faster evaluation. If false, the cloud material will be invoked via callable shaders. (default = true)"),
+	TEXT("Controls the use of a texture to cache the cloud material.\n")
+	TEXT(" true : Clouds will be voxelized into a texture for faster evaluation (default)\n")
+	TEXT(" false: Cloud material will be invoked per ray march step via callable shaders (more accurate, but more expensive)\n")
+	TEXT("Requires r.PathTracing.Experimental=true to modify.\n"),
 	ECVF_RenderThreadSafe
 );
 
 TAutoConsoleVariable<int32> CVarPathTracingCloudMapResolution(
 	TEXT("r.PathTracing.CloudMapResolution"),
 	512,
-	TEXT("Size of the texture used to cache cloud appearance in reference atmosphere mode.  (default = 512)"),
+	TEXT("Size of the texture used to cache volumetric cloud appearance in reference atmosphere mode.  (default = 512)"),
 	ECVF_RenderThreadSafe
 );
 
 TAutoConsoleVariable<int32> CVarPathTracingCloudMapDepth(
 	TEXT("r.PathTracing.CloudMapDepth"),
 	64,
-	TEXT("Depth of the texture used to cache cloud appearance in reference atmosphere mode.  (default = 64)"),
+	TEXT("Depth of the texture used to cache volumetric cloud appearance in reference atmosphere mode.  (default = 64)"),
 	ECVF_RenderThreadSafe
 );
 
@@ -323,19 +326,19 @@ TAutoConsoleVariable<bool> CVarPathTracingAdjustMultiGPUPasses(
 );
 #endif  // WITH_MGPU
 
-TAutoConsoleVariable<int32> CVarPathTracingWiperMode(
+TAutoConsoleVariable<bool> CVarPathTracingWiperMode(
 	TEXT("r.PathTracing.WiperMode"),
-	0,
-	TEXT("Enables wiper mode to render using the path tracer only in a region of the screen for debugging purposes (default = 0, wiper mode disabled)"),
+	false,
+	TEXT("Enables wiper mode to render using the path tracer only in a region of the screen for debugging purposes (default = false, wiper mode disabled)"),
 	ECVF_RenderThreadSafe 
 );
 
-TAutoConsoleVariable<int32> CVarPathTracingProgressDisplay(
+TAutoConsoleVariable<bool> CVarPathTracingProgressDisplay(
 	TEXT("r.PathTracing.ProgressDisplay"),
-	1,
-	TEXT("Enables an in-frame display of progress towards the defined sample per pixel limit. The indicator dissapears when the maximum is reached and sample accumulation has stopped (default = 1)\n")
-	TEXT("0: off\n")
-	TEXT("1: on (default)\n"),
+	true,
+	TEXT("Enables an in-frame display of progress towards the defined sample per pixel limit. The indicator dissapears when the maximum is reached and sample accumulation has stopped\n")
+	TEXT(" false: off\n")
+	TEXT(" true : on (default)\n"),
 	ECVF_RenderThreadSafe
 );
 
@@ -374,10 +377,10 @@ TAutoConsoleVariable<int32> CVarPathTracingDecalGridVisualize(
 	ECVF_RenderThreadSafe
 );
 
-TAutoConsoleVariable<int32> CVarPathTracingUseDBuffer(
+TAutoConsoleVariable<bool> CVarPathTracingUseDBuffer(
 	TEXT("r.PathTracing.UseDBuffer"),
-	1,
-	TEXT("Whether to support DBuffer functionality (default=1)"),
+	true,
+	TEXT("Whether to support DBuffer functionality (default=true)"),
 	ECVF_RenderThreadSafe
 );
 
@@ -402,73 +405,65 @@ TAutoConsoleVariable<float> CVarPathTracingMeshDecalBias(
 	ECVF_RenderThreadSafe
 );
 
-TAutoConsoleVariable<int32> CVarPathTracingLightFunctionColor(
+TAutoConsoleVariable<bool> CVarPathTracingLightFunctionColor(
 	TEXT("r.PathTracing.LightFunctionColor"),
-	0,
-	TEXT("Enables light functions to be colored instead of greyscale (default = 0)\n")
+	false,
+	TEXT("Enables light functions to be colored instead of greyscale (default = false)\n")
 	TEXT("0: off (default)\n")
 	TEXT("1: on (light function material output is used directly instead of converting to greyscale)\n"),
 	ECVF_RenderThreadSafe
 );
 
-TAutoConsoleVariable<int32> CVarPathTracingHeterogeneousVolumesRebuildEveryFrame(
+TAutoConsoleVariable<bool> CVarPathTracingHeterogeneousVolumesRebuildEveryFrame(
 	TEXT("r.PathTracing.HeterogeneousVolumes.RebuildEveryFrame"),
-	1,
-	TEXT("Rebuilds volumetric acceleration structures every frame (default = 1)\n"),
+	true,
+	TEXT("Rebuilds volumetric acceleration structures every frame (default = true)\n"),
 	ECVF_RenderThreadSafe
 );
 
-TAutoConsoleVariable<int32> CVarPathTracingCameraMediumTracking(
+TAutoConsoleVariable<bool> CVarPathTracingCameraMediumTracking(
 	TEXT("r.PathTracing.CameraMediumTracking"),
-	1,
-	TEXT("Enables automatic camera medium tracking to detect when a camera starts inside water or solid glass automatically (default = 1)\n")
-	TEXT("0: off\n")
-	TEXT("1: on (default)\n"),
+	true,
+	TEXT("Enables automatic camera medium tracking to detect when a camera starts inside water or solid glass automatically\n")
+	TEXT(" false: off\n")
+	TEXT(" true : on (default)\n"),
 	ECVF_RenderThreadSafe
 );
 
-TAutoConsoleVariable<int32> CVarPathTracingOutputPostProcessResources(
+TAutoConsoleVariable<bool> CVarPathTracingOutputPostProcessResources(
 	TEXT("r.PathTracing.OutputPostProcessResources"),
-	1,
+	true,
 	TEXT("Output the pathtracing resources to the postprocess passes\n")
-	TEXT("0: off\n")
-	TEXT("1: on (Buffers including, raw/denoised radiance, albedo, normal, and variance)\n"),
+	TEXT(" false: off\n")
+	TEXT(" true : on (Buffers including, raw/denoised radiance, albedo, normal, and variance)\n"),
 	ECVF_RenderThreadSafe
 );
 
-TAutoConsoleVariable<int32> CVarPathTracingSubstrateUseSimplifiedMaterial(
+TAutoConsoleVariable<bool> CVarPathTracingSubstrateUseSimplifiedMaterial(
 	TEXT("r.PathTracing.Substrate.UseSimplifiedMaterials"),
-	0,
-	TEXT("Instead of evaluating all layers, use an optimized material in which all slabs have been merged. This is mainly intended for debugging and requires r.PathTracing.Substrate.CompileSimplifiedMaterials to be true.\n")
-	TEXT("0: off (default)\n")
-	TEXT("1: on\n"),
+	false,
+	TEXT("Instead of evaluating all layers, use an optimized material in which all slabs have been merged.\n")
+	TEXT(" false: off (default)\n")
+	TEXT(" true : on\n")
+	TEXT("Requires r.PathTracing.Substrate.CompileSimplifiedMaterials=true to be set.\n"),
 	ECVF_RenderThreadSafe
 );
 
-TAutoConsoleVariable<int32> CVarPathTracingSubstrateCompileSimplifiedMaterial(
+TAutoConsoleVariable<bool> CVarPathTracingSubstrateCompileSimplifiedMaterial(
 	TEXT("r.PathTracing.Substrate.CompileSimplifiedMaterials"),
-	0,
+	false,
 	TEXT("Compile a simplified representation of Substrate materials which merges all slabs into one. This is mainly intended for debugging purposes. Enabling this double the number of path tracing shader permutations.\n")
-	TEXT("0: off (default)\n")
-	TEXT("1: on\n"),
+	TEXT(" false: off (default)\n")
+	TEXT(" true : on\n"),
 	ECVF_RenderThreadSafe | ECVF_ReadOnly
 );
 
-TAutoConsoleVariable<int32> CVarpathTracingOverrideDepth(
-	TEXT("r.PathTracing.Override.Depth"),
-	1,
-	TEXT("Override the scene depth z by the path tracing depth z")
-	TEXT("0: off\n")
-	TEXT("1: On (Default, translucent materials have better DOF with the post-process DOF.)\n"),
-	ECVF_RenderThreadSafe
-);
-
 TAutoConsoleVariable<bool> CVarPathTracingUseAnalyticTransmittance(
 	TEXT("r.PathTracing.UseAnalyticTransmittance"),
 	true,
 	TEXT("Determines use of analytical or null-tracking estimation when evaluating transmittance\n")
-	TEXT("0: off (uses null-tracking estimation)\n")
-	TEXT("1: on (uses analytical estimation when possible) (default)\n"),
+	TEXT(" false: off (uses null-tracking estimation)\n")
+	TEXT(" true : on (uses analytical estimation when possible) (default)\n"),
 	ECVF_RenderThreadSafe
 );
 
@@ -477,7 +472,8 @@ TAutoConsoleVariable<int32> CVarPathTracingAdaptiveSampling(
 	0,
 	TEXT("Determines if adaptive sampling is enabled. When non-zero, the path tracer will try to skip calculation of pixels below the specified error threshold.\n")
 	TEXT("0: off (uniform sampling - default)\n")
-	TEXT("1: on (adaptive sampling)\n"),
+	TEXT("1: on (adaptive sampling)\n")
+	TEXT("Requires r.PathTracing.Experimental=true to modify.\n"),
 	ECVF_RenderThreadSafe
 );
 
@@ -762,7 +758,8 @@ static void PreparePathTracingData(const FScene* Scene, const FViewInfo& View, F
 		PathTracing::UsesReferenceAtmosphere(View);
 	
 	const FVolumetricCloudRenderSceneInfo* CloudInfo = Scene->GetVolumetricCloudSceneInfo();
-	const bool bVolumeCloudMapEnabled = !RHISupportsRayTracingCallableShaders(View.GetShaderPlatform()) || CVarPathTracingCloudMapEnabled.GetValueOnRenderThread();
+	const bool bExperimental = CVarPathTracingExperimental.GetValueOnRenderThread();
+	const bool bVolumeCloudMapEnabled = !ShouldCompileRayTracingCallableShadersForProject(View.GetShaderPlatform()) || (bExperimental == false || CVarPathTracingCloudMapEnabled.GetValueOnRenderThread());
 	const bool bVolumeCloudsVisible = ShouldRenderVolumetricCloud(Scene, ShowFlags)
 		&& bUseReferenceAtmosphere
 		&& CVarPathTracingEnableReferenceClouds.GetValueOnRenderThread() != 0
@@ -1046,12 +1043,13 @@ class FPathTracingRG : public FGlobalShader
 
 	class FCompactionType : SHADER_PERMUTATION_BOOL("PATH_TRACER_USE_COMPACTION");
 	class FAdaptiveSampling : SHADER_PERMUTATION_BOOL("PATH_TRACER_USE_ADAPTIVE_SAMPLING");
+	class FCloudShader : SHADER_PERMUTATION_BOOL("PATH_TRACER_USE_CLOUD_SHADER");
 	class FSubstrateComplexSpecialMaterial : SHADER_PERMUTATION_BOOL("PATH_TRACER_USE_SUBSTRATE_SPECIAL_COMPLEX_MATERIAL");
-	using FPermutationDomain = TShaderPermutationDomain<FCompactionType, FAdaptiveSampling, FSubstrateComplexSpecialMaterial>;
+	using FPermutationDomain = TShaderPermutationDomain<FCompactionType, FAdaptiveSampling, FCloudShader, FSubstrateComplexSpecialMaterial>;
 
 	static bool ShouldCompilePermutation(const FGlobalShaderPermutationParameters& Parameters)
 	{
-		const bool bUseExperimental = CVarPathTracingExperimental.GetValueOnAnyThread() != 0;
+		const bool bUseExperimental = CVarPathTracingExperimental.GetValueOnAnyThread();
 		FPermutationDomain PermutationVector(Parameters.PermutationId);
 		if (bUseExperimental == false)
 		{
@@ -1066,6 +1064,22 @@ class FPathTracingRG : public FGlobalShader
 				// adaptive sampling is experimental
 				return false;
 			}
+			if (PermutationVector.Get<FCloudShader>())
+			{
+				// using the cloud shader via callable shaders is experimental
+				return false;
+			}
+		}
+		else
+		{
+			if (PermutationVector.Get<FCloudShader>())
+			{
+				// the cloud shader version can only be supported if the platform supports callable shaders
+				if (!ShouldCompileRayTracingCallableShadersForProject(Parameters.Platform))
+				{
+					return false;
+				}
+			}
 		}
 		if (!Substrate::IsSubstrateEnabled())
 		{
@@ -1382,8 +1396,11 @@ public:
 
 	static bool ShouldCompilePermutation(const FMaterialShaderPermutationParameters& Parameters)
 	{
+		const bool bUseExperimental = CVarPathTracingExperimental.GetValueOnAnyThread();
+
 		return ShouldCompilePathTracingShadersForProject(Parameters.Platform) &&
 			ShouldCompileRayTracingCallableShadersForProject(Parameters.Platform) &&
+			bUseExperimental &&
 			Parameters.MaterialParameters.bIsUsedWithVolumetricCloud &&
 			Parameters.MaterialParameters.MaterialDomain == MD_Volume;
 	}
@@ -1457,7 +1474,12 @@ void PreparePathTracingCloudMaterial(FScene* Scene, TArrayView<FViewInfo> Views)
 
 	// if we are using reference clouds, or if we are using the cloud map, no need to prepare the callable shader version
 	if (CVarPathTracingEnableReferenceClouds.GetValueOnRenderThread() == 0 ||
-		CVarPathTracingCloudMapEnabled.GetValueOnRenderThread())
+		(CVarPathTracingExperimental.GetValueOnRenderThread() == false || CVarPathTracingCloudMapEnabled.GetValueOnRenderThread()))
+	{
+		return;
+	}
+
+	if (!ShouldCompileRayTracingCallableShadersForProject(Scene->GetShaderPlatform()))
 	{
 		return;
 	}
@@ -1894,7 +1916,7 @@ public:
 		}
 		else
 		{
-			if (SimplifySubstrate && (!Substrate::IsSubstrateEnabled() || CVarPathTracingSubstrateCompileSimplifiedMaterial.GetValueOnAnyThread() == 0))
+			if (SimplifySubstrate && (!Substrate::IsSubstrateEnabled() || CVarPathTracingSubstrateCompileSimplifiedMaterial.GetValueOnAnyThread() == false))
 			{
 				// don't compile the extra Substrate permutation if:
 				//    Substrate is not enabled on this project
@@ -2069,8 +2091,8 @@ bool FRayTracingMeshProcessor::ProcessPathTracing(
 			{
 				// In order to use Substrate simplified materials, Substrate has to be enabled, we have to have _compiled_ the extra permutations _and_ the runtime toggle must be true
 				const bool bUseSimplifiedMaterial = Substrate::IsSubstrateEnabled() &&
-					CVarPathTracingSubstrateCompileSimplifiedMaterial.GetValueOnRenderThread() != 0 &&
-					CVarPathTracingSubstrateUseSimplifiedMaterial.GetValueOnRenderThread() != 0;
+					CVarPathTracingSubstrateCompileSimplifiedMaterial.GetValueOnRenderThread() &&
+					CVarPathTracingSubstrateUseSimplifiedMaterial.GetValueOnRenderThread();
 				if (NeedsAnyHitShader(MaterialResource))
 				{
 					if (bUseSimplifiedMaterial)
@@ -2545,6 +2567,7 @@ void SetLightParameters(FRDGBuilder& GraphBuilder, FPathTracingRG::FParameters*
 			DestLight.Flags |= LightingChannelMask & PATHTRACER_FLAG_LIGHTING_CHANNEL_MASK;
 			DestLight.Flags |= Light.LightSceneInfo->Proxy->CastsDynamicShadow() ? PATHTRACER_FLAG_CAST_SHADOW_MASK : 0;
 			DestLight.Flags |= Light.LightSceneInfo->Proxy->CastsVolumetricShadow() ? PATHTRACER_FLAG_CAST_VOL_SHADOW_MASK : 0;
+			DestLight.Flags |= Light.LightSceneInfo->Proxy->GetCastCloudShadows() ? PATHTRACER_FLAG_CAST_CLOUD_SHADOW_MASK : 0;
 			DestLight.IESAtlasIndex = INDEX_NONE;
 			DestLight.MissShaderIndex = 0;
 
@@ -2617,6 +2640,7 @@ void SetLightParameters(FRDGBuilder& GraphBuilder, FPathTracingRG::FParameters*
 		DestLight.Flags |= LightingChannelMask & PATHTRACER_FLAG_LIGHTING_CHANNEL_MASK;
 		DestLight.Flags |= Light.LightSceneInfo->Proxy->CastsDynamicShadow() ? PATHTRACER_FLAG_CAST_SHADOW_MASK : 0;
 		DestLight.Flags |= Light.LightSceneInfo->Proxy->CastsVolumetricShadow() ? PATHTRACER_FLAG_CAST_VOL_SHADOW_MASK : 0;
+		DestLight.Flags |= Light.LightSceneInfo->Proxy->GetCastCloudShadows() ? PATHTRACER_FLAG_CAST_CLOUD_SHADOW_MASK : 0;
 		DestLight.IESAtlasIndex = LightParameters.IESAtlasIndex;
 		DestLight.MissShaderIndex = 0;
 
@@ -2740,14 +2764,16 @@ IMPLEMENT_SHADER_TYPE(, FPathTracingCompositorPS, TEXT("/Engine/Private/PathTrac
 
 static FPathTracingRG::FPermutationDomain GetPathTracingRGPermutation(const FScene& Scene)
 {
-	const bool bUseExperimental = CVarPathTracingExperimental.GetValueOnRenderThread() != 0;
-	const bool bUseCompaction = (bUseExperimental == false) || CVarPathTracingCompaction.GetValueOnRenderThread() != 0;
+	const bool bUseExperimental = CVarPathTracingExperimental.GetValueOnRenderThread();
+	const bool bUseCompaction = (bUseExperimental == false) || CVarPathTracingCompaction.GetValueOnRenderThread();
 	const bool bUseAdaptiveSampling = bUseExperimental && CVarPathTracingAdaptiveSampling.GetValueOnRenderThread() != 0;
+	const bool bUseCloudShader = bUseExperimental && ShouldCompileRayTracingCallableShadersForProject(Scene.GetShaderPlatform()) && !CVarPathTracingCloudMapEnabled.GetValueOnRenderThread();
 	const bool bHasComplexSpecialRenderPath = Substrate::IsSubstrateEnabled() && Scene.SubstrateSceneData.bUsesComplexSpecialRenderPath;
 
 	FPathTracingRG::FPermutationDomain Out;
 	Out.Set<FPathTracingRG::FCompactionType>(bUseCompaction);
 	Out.Set<FPathTracingRG::FAdaptiveSampling>(bUseAdaptiveSampling);
+	Out.Set<FPathTracingRG::FCloudShader>(bUseCloudShader);
 	Out.Set<FPathTracingRG::FSubstrateComplexSpecialMaterial>(bHasComplexSpecialRenderPath);
 	return Out;
 }
@@ -2916,10 +2942,10 @@ void FDeferredShadingSceneRenderer::RenderPathTracing(
 	uint32 MaxSPP = SamplesPerPixelCVar > -1 ? SamplesPerPixelCVar : View.FinalPostProcessSettings.PathTracingSamplesPerPixel;
 	MaxSPP = FMath::Max(MaxSPP, 1u);
 
-	const bool bUseExperimental = CVarPathTracingExperimental.GetValueOnRenderThread() != 0;
+	const bool bUseExperimental = CVarPathTracingExperimental.GetValueOnRenderThread();
 
 	Config.LockedSamplingPattern = CVarPathTracingFrameIndependentTemporalSeed.GetValueOnRenderThread() == 0;
-	Config.UseCameraMediumTracking = CVarPathTracingCameraMediumTracking.GetValueOnRenderThread() != 0;
+	Config.UseCameraMediumTracking = CVarPathTracingCameraMediumTracking.GetValueOnRenderThread();
 	Config.UseAdaptiveSampling = bUseExperimental && CVarPathTracingAdaptiveSampling.GetValueOnAnyThread() != 0;
 	Config.AdaptiveSamplingThreshold = CVarPathTracingAdaptiveSamplingErrorThreshold.GetValueOnRenderThread();
 	Config.CloudAccelerationMapNumSamples = FMath::Clamp(CVarPathTracingCloudAccelerationMapNumSamples.GetValueOnRenderThread(), 1, 65536);
@@ -2945,9 +2971,9 @@ void FDeferredShadingSceneRenderer::RenderPathTracing(
 	Config.LightShowFlags |= View.Family->EngineShowFlags.ReflectionOverride         ? 1 << 12 : 0;
 	Config.LightShowFlags |= View.Family->EngineShowFlags.SubsurfaceScattering       ? 1 << 13 : 0;
 	// the following affects which material shaders get used and therefore change the image
-	if (Substrate::IsSubstrateEnabled() && CVarPathTracingSubstrateCompileSimplifiedMaterial.GetValueOnRenderThread() != 0)
+	if (Substrate::IsSubstrateEnabled() && CVarPathTracingSubstrateCompileSimplifiedMaterial.GetValueOnRenderThread())
 	{
-		Config.LightShowFlags |= CVarPathTracingSubstrateUseSimplifiedMaterial.GetValueOnRenderThread() != 0 ? 1 << 14 : 0;
+		Config.LightShowFlags |= CVarPathTracingSubstrateUseSimplifiedMaterial.GetValueOnRenderThread() ? 1 << 14 : 0;
 	}
 
 	PreparePathTracingData(Scene, View, Config.PathTracingData);
@@ -3207,7 +3233,7 @@ void FDeferredShadingSceneRenderer::RenderPathTracing(
 		if (bNeedsMoreRays)
 		{
 			RDG_EVENT_SCOPE(GraphBuilder, "Path Tracing Compute (%d x %d)", DispatchResX, DispatchResY);
-			bool bForceRebuild = CVarPathTracingHeterogeneousVolumesRebuildEveryFrame.GetValueOnRenderThread() != 0;
+			const bool bForceRebuild = CVarPathTracingHeterogeneousVolumesRebuildEveryFrame.GetValueOnRenderThread();
 			bCreateVolumeGrids = bForceRebuild ||
 				!PathTracingState->AdaptiveFrustumGridParameterCache.TopLevelGridBuffer ||
 				!PathTracingState->AdaptiveOrthoGridParameterCache.TopLevelGridBuffer;
@@ -3912,7 +3938,7 @@ void FDeferredShadingSceneRenderer::RenderPathTracing(
 	DisplayParameters->MaxSamples = MaxSPP;
 	DisplayParameters->ProgressDisplayEnabled = CVarPathTracingProgressDisplay.GetValueOnRenderThread();
 	DisplayParameters->AdaptiveSamplingErrorThreshold = Config.AdaptiveSamplingThreshold;
-	DisplayParameters->AdaptiveSamplingVisualize = CVarPathTracingAdaptiveSamplingVisualize.GetValueOnRenderThread();
+	DisplayParameters->AdaptiveSamplingVisualize = Config.UseAdaptiveSampling ? CVarPathTracingAdaptiveSamplingVisualize.GetValueOnRenderThread() : 0;
 	DisplayParameters->VarianceTextureDims = FIntVector(DispatchResX, DispatchResY, NumVarianceMips);
 	DisplayParameters->ViewUniformBuffer = View.ViewUniformBuffer;
 	DisplayParameters->RadianceTexture = GraphBuilder.CreateSRV(FRDGTextureSRVDesc::Create(DenoisedRadianceTexture ? DenoisedRadianceTexture : RadianceTexture));
@@ -3927,7 +3953,7 @@ void FDeferredShadingSceneRenderer::RenderPathTracing(
 	const bool IsCursorInsideView = View.CursorPos.X != -1 || View.CursorPos.Y != -1;
 	// wiper mode - reveals the render below the path tracing display
 	// NOTE: we still path trace the full resolution even while wiping the cursor so that rendering does not get out of sync
-	if (CVarPathTracingWiperMode.GetValueOnRenderThread() != 0)
+	if (CVarPathTracingWiperMode.GetValueOnRenderThread())
 	{
 		float DPIScale = FPlatformApplicationMisc::GetDPIScaleFactorAtPoint(View.CursorPos.X, View.CursorPos.Y);
 		
@@ -3944,16 +3970,7 @@ void FDeferredShadingSceneRenderer::RenderPathTracing(
 	TShaderMapRef<FPathTracingCompositorPS> PixelShader(View.ShaderMap);
 	TShaderMapRef<FScreenPassVS> VertexShader(View.ShaderMap);
 	FRHIBlendState* BlendState = FScreenPassPipelineState::FDefaultBlendState::GetRHI();
-	FRHIDepthStencilState* DepthStencilState = nullptr;
-
-	if (CVarpathTracingOverrideDepth.GetValueOnRenderThread() != 0)
-	{
-		DepthStencilState = TStaticDepthStencilState<true, CF_Always>::GetRHI();
-	}
-	else
-	{
-		DepthStencilState = TStaticDepthStencilState<false, CF_Always>::GetRHI();
-	}
+	FRHIDepthStencilState* DepthStencilState = DepthStencilState = TStaticDepthStencilState<true /* bEnableDepthWrite */, CF_Always>::GetRHI();
 
 	AddDrawScreenPass(
 		GraphBuilder,
@@ -3968,7 +3985,7 @@ void FDeferredShadingSceneRenderer::RenderPathTracing(
 		DisplayParameters);
 
 	// Setup the path tracing resources to be used by post process pass.
-	if (CVarPathTracingOutputPostProcessResources.GetValueOnRenderThread() != 0)
+	if (CVarPathTracingOutputPostProcessResources.GetValueOnRenderThread())
 	{
 		PathTracingResources.bPostProcessEnabled = true;
 		PathTracingResources.DenoisedRadiance = DenoisedRadianceTexture ? DenoisedRadianceTexture : RadianceTexture;
@@ -3985,7 +4002,7 @@ void FDeferredShadingSceneRenderer::RenderPathTracing(
 		Inputs.SceneColor =SceneColorOutputTexture;
 
 		FScreenPassTextureViewport MotionVectorViewport(SceneColorOutputTexture, View.ViewRect);
-		if (CVarPathTracingWiperMode.GetValueOnRenderThread() != 0)
+		if (CVarPathTracingWiperMode.GetValueOnRenderThread())
 		{
 			float DPIScale = FPlatformApplicationMisc::GetDPIScaleFactorAtPoint(View.CursorPos.X, View.CursorPos.Y);
 			if (IsCursorInsideView)
