Fix static analysis warning

#rnx

diff --git a/Engine/Source/Runtime/Renderer/Private/PathTracing.cpp b/Engine/Source/Runtime/Renderer/Private/PathTracing.cpp
--- a/Engine/Source/Runtime/Renderer/Private/PathTracing.cpp
+++ b/Engine/Source/Runtime/Renderer/Private/PathTracing.cpp
@@ -3970,7 +3970,7 @@ void FDeferredShadingSceneRenderer::RenderPathTracing(
 	TShaderMapRef<FPathTracingCompositorPS> PixelShader(View.ShaderMap);
 	TShaderMapRef<FScreenPassVS> VertexShader(View.ShaderMap);
 	FRHIBlendState* BlendState = FScreenPassPipelineState::FDefaultBlendState::GetRHI();
-	FRHIDepthStencilState* DepthStencilState = DepthStencilState = TStaticDepthStencilState<true /* bEnableDepthWrite */, CF_Always>::GetRHI();
+	FRHIDepthStencilState* DepthStencilState = TStaticDepthStencilState<true /* bEnableDepthWrite */, CF_Always>::GetRHI();
 
 	AddDrawScreenPass(
 		GraphBuilder,
