If supported pre-cache PSO for velocity passes on mobile
Do not cache MDC and do not cache PSO for DepthPass when not enabled on mobile

diff --git a/Engine/Source/Runtime/Renderer/Private/DepthRendering.cpp b/Engine/Source/Runtime/Renderer/Private/DepthRendering.cpp
--- a/Engine/Source/Runtime/Renderer/Private/DepthRendering.cpp
+++ b/Engine/Source/Runtime/Renderer/Private/DepthRendering.cpp
@@ -926,6 +926,12 @@ bool FDepthPassMeshProcessor::ShouldRender(const FMaterial& Material, bool bMate
 	bUseDefaultMaterial = false;
 	bPositionOnly = false;
 
+	if (FeatureLevel == ERHIFeatureLevel::ES3_1 && EarlyZPassMode == DDM_None)
+	{
+		// Do not cache MDC and do not pre-cache PSOs for a depth pass if it's never going to be used on mobile platforms
+		return false;
+	}
+
 	if (IsOpaqueBlendMode(Material)
 		&& EarlyZPassMode != DDM_MaskedOnly
 		&& bSupportPositionOnlyStream
diff --git a/Engine/Source/Runtime/Renderer/Private/VelocityRendering.cpp b/Engine/Source/Runtime/Renderer/Private/VelocityRendering.cpp
--- a/Engine/Source/Runtime/Renderer/Private/VelocityRendering.cpp
+++ b/Engine/Source/Runtime/Renderer/Private/VelocityRendering.cpp
@@ -879,7 +879,7 @@ FMeshPassProcessor* CreateVelocityPassProcessor(ERHIFeatureLevel::Type FeatureLe
 }
 
 REGISTER_MESHPASSPROCESSOR_AND_PSOCOLLECTOR(VelocityPass, CreateVelocityPassProcessor, EShadingPath::Deferred,  EMeshPass::Velocity, EMeshPassFlags::CachedMeshCommands | EMeshPassFlags::MainView);
-FRegisterPassProcessorCreateFunction MobileRegisterVelocityPass(&CreateVelocityPassProcessor, EShadingPath::Mobile, EMeshPass::Velocity, EMeshPassFlags::CachedMeshCommands | EMeshPassFlags::MainView);
+REGISTER_MESHPASSPROCESSOR_AND_PSOCOLLECTOR(MobileVelocityPass, CreateVelocityPassProcessor, EShadingPath::Mobile,  EMeshPass::Velocity, EMeshPassFlags::CachedMeshCommands | EMeshPassFlags::MainView);
 
 FTranslucentVelocityMeshProcessor::FTranslucentVelocityMeshProcessor(const FScene* Scene, ERHIFeatureLevel::Type FeatureLevel, const FSceneView* InViewIfDynamicMeshCommand, const FMeshPassProcessorRenderState& InPassDrawRenderState, FMeshPassDrawListContext* InDrawListContext)
 	: FVelocityMeshProcessor(EMeshPass::TranslucentVelocity, Scene, FeatureLevel, InViewIfDynamicMeshCommand, InPassDrawRenderState, InDrawListContext)
@@ -895,4 +895,4 @@ FMeshPassProcessor* CreateTranslucentVelocityPassProcessor(ERHIFeatureLevel::Typ
 }
 
 REGISTER_MESHPASSPROCESSOR_AND_PSOCOLLECTOR(TranslucentVelocityPass, CreateTranslucentVelocityPassProcessor, EShadingPath::Deferred, EMeshPass::TranslucentVelocity, EMeshPassFlags::CachedMeshCommands | EMeshPassFlags::MainView);
-FRegisterPassProcessorCreateFunction RegisterMobileTranslucentVelocityPass(&CreateTranslucentVelocityPassProcessor, EShadingPath::Mobile, EMeshPass::TranslucentVelocity, EMeshPassFlags::CachedMeshCommands | EMeshPassFlags::MainView);
+REGISTER_MESHPASSPROCESSOR_AND_PSOCOLLECTOR(MobileTranslucentVelocityPass, CreateTranslucentVelocityPassProcessor, EShadingPath::Mobile, EMeshPass::TranslucentVelocity, EMeshPassFlags::CachedMeshCommands | EMeshPassFlags::MainView);
