Fixed various missing resource transition when rendering RVT  (thumbnails, SVT baking), causing RHI validation issues

diff --git a/Engine/Source/Editor/VirtualTexturingEditor/Private/RuntimeVirtualTextureBuildStreamingMips.cpp b/Engine/Source/Editor/VirtualTexturingEditor/Private/RuntimeVirtualTextureBuildStreamingMips.cpp
--- a/Engine/Source/Editor/VirtualTexturingEditor/Private/RuntimeVirtualTextureBuildStreamingMips.cpp
+++ b/Engine/Source/Editor/VirtualTexturingEditor/Private/RuntimeVirtualTextureBuildStreamingMips.cpp
@@ -54,6 +54,8 @@ namespace
 				RenderTargets[Layer] = RHICreateTexture(Desc);
 
 				Desc.SetFlags(ETextureCreateFlags::CPUReadback);
+				// The staging texture will only ever be in CopyDest state so set it right away : 
+				Desc.SetInitialState(ERHIAccess::CopyDest); 
 				StagingTextures[Layer] = RHICreateTexture(Desc);
 			}
 
@@ -216,6 +218,7 @@ namespace RuntimeVirtualTexture
 		FinalPixels.SetNumUninitialized(RenderTileResourcesBytes);
 
 		UE::RenderCommandPipe::FSyncScope SyncScope;
+		ERHIAccess RenderTargetsCurrentState = ERHIAccess::RTV;
 
 		// Iterate over all tiles and render/store each one to the final image
 		for (int32 TileY = 0; TileY < NumTilesY && !Task.ShouldCancel(); TileY++)
@@ -245,17 +248,11 @@ namespace RuntimeVirtualTexture
 					TileX, TileY,
 					TileSize, ImageSizeX, ImageSizeY, 
 					&FinalPixels,
-					FixedColor](FRHICommandListImmediate& RHICmdList)
+					FixedColor, &RenderTargetsCurrentState](FRHICommandListImmediate& RHICmdList)
 				{
 					const FBox2D TileBox(FVector2D(0, 0), FVector2D(TileSize, TileSize));
 					const FIntRect TileRect(0, 0, TileSize, TileSize);
 
-					// Transition render targets for writing
-					for (int32 Layer = 0; Layer < NumLayers; Layer++)
-					{
-						RHICmdList.Transition(FRHITransitionInfo(RenderTileResources.GetRenderTarget(Layer), ERHIAccess::Unknown, ERHIAccess::RTV));
-					}
-
 					{
 						FRDGBuilder GraphBuilder(RHICmdList);
 
@@ -270,12 +267,14 @@ namespace RuntimeVirtualTexture
 						Desc.bIsThumbnails = false;
 						Desc.FixedColor = FixedColor;
 						Desc.NumPageDescs = 1;
-						Desc.Targets[0].Texture = RenderTileResources.GetRenderTarget(0);
-						Desc.Targets[1].Texture = RenderTileResources.GetRenderTarget(1);
-						Desc.Targets[2].Texture = RenderTileResources.GetRenderTarget(2);
-						Desc.PageDescs[0].DestBox[0] = TileBox;
-						Desc.PageDescs[0].DestBox[1] = TileBox;
-						Desc.PageDescs[0].DestBox[2] = TileBox;
+						for (int32 Layer = 0; Layer < NumLayers; Layer++)
+						{
+							Desc.Targets[Layer].Texture = RenderTileResources.GetRenderTarget(Layer);
+							Desc.Targets[Layer].TextureAccessBefore = RenderTargetsCurrentState;
+							// Leave the render targets in CopySrc state since we'll read from them right after, to avoid a useless transition :
+							Desc.Targets[Layer].TextureAccessAfter = ERHIAccess::CopySrc;
+							Desc.PageDescs[0].DestBox[Layer] = TileBox;
+						}
 						Desc.PageDescs[0].UVRange = UVRange;
 						Desc.PageDescs[0].vLevel = IntCastChecked<uint8>(RenderLevel);
 
@@ -287,9 +286,9 @@ namespace RuntimeVirtualTexture
 					// Copy to staging
 					for (int32 Layer = 0; Layer < NumLayers; Layer++)
 					{
-						RHICmdList.Transition(FRHITransitionInfo(RenderTileResources.GetRenderTarget(Layer), ERHIAccess::RTV, ERHIAccess::CopySrc));
 						RHICmdList.CopyTexture(RenderTileResources.GetRenderTarget(Layer), RenderTileResources.GetStagingTexture(Layer), FRHICopyTextureInfo());
 					}
+					RenderTargetsCurrentState = ERHIAccess::CopySrc;
 
 					RenderTileResources.GetFence()->Clear();
 					RHICmdList.WriteGPUFence(RenderTileResources.GetFence());
diff --git a/Engine/Source/Editor/VirtualTexturingEditor/Private/RuntimeVirtualTextureThumbnailRenderer.cpp b/Engine/Source/Editor/VirtualTexturingEditor/Private/RuntimeVirtualTextureThumbnailRenderer.cpp
--- a/Engine/Source/Editor/VirtualTexturingEditor/Private/RuntimeVirtualTextureThumbnailRenderer.cpp
+++ b/Engine/Source/Editor/VirtualTexturingEditor/Private/RuntimeVirtualTextureThumbnailRenderer.cpp
@@ -116,6 +116,9 @@ void URuntimeVirtualTextureThumbnailRenderer::Draw(UObject* Object, int32 X, int
 		Desc.FixedColor = FLinearColor::Transparent;
 		Desc.NumPageDescs = 1;
 		Desc.Targets[0].Texture = RenderTarget->GetRenderTargetTexture();
+		// The target is in RTV state and should be left in RTV state (RenderPage needs the input/output state) : 
+		Desc.Targets[0].TextureAccessBefore = ERHIAccess::RTV;
+		Desc.Targets[0].TextureAccessAfter = ERHIAccess::RTV;
 		Desc.PageDescs[0].DestBox[0] = DestBox;
 		Desc.PageDescs[0].UVRange = FBox2D(FVector2D(0, 0), FVector2D(1, 1));
 		Desc.PageDescs[0].vLevel = IntCastChecked<uint8>(MaxLevel);
diff --git a/Engine/Source/Runtime/Renderer/Private/VT/RuntimeVirtualTextureRender.cpp b/Engine/Source/Runtime/Renderer/Private/VT/RuntimeVirtualTextureRender.cpp
--- a/Engine/Source/Runtime/Renderer/Private/VT/RuntimeVirtualTextureRender.cpp
+++ b/Engine/Source/Runtime/Renderer/Private/VT/RuntimeVirtualTextureRender.cpp
@@ -1336,7 +1336,7 @@ namespace RuntimeVirtualTexture
 	END_SHADER_PARAMETER_STRUCT()
 
 	/** Set up the copy to final output physical texture. */
-	void AddCopyToOutputPass(FRDGBuilder& GraphBuilder, FRDGTextureRef InputTexture, FRHITexture* OutputTexture, FBox2D const& DestBox)
+	void AddCopyToOutputPass(FRDGBuilder& GraphBuilder, FRDGTextureRef InputTexture, FRHITexture* OutputTexture, ERHIAccess OutputTextureAccessBefore, ERHIAccess OutputTextureAccessAfter, FBox2D const& DestBox)
 	{
 		FRHICopyTextureInfo CopyInfo;
 		CopyInfo.Size = InputTexture->Desc.GetSize();
@@ -1349,11 +1349,11 @@ namespace RuntimeVirtualTexture
 			RDG_EVENT_NAME("VirtualTextureCopyToOutput"),
 			Parameters,
 			ERDGPassFlags::Copy | ERDGPassFlags::NeverCull,
-			[InputTexture, OutputTexture, CopyInfo](FRHICommandList& RHICmdList)
+			[InputTexture, OutputTexture, OutputTextureAccessBefore, OutputTextureAccessAfter, CopyInfo](FRHICommandList& RHICmdList)
 			{
-				RHICmdList.Transition(FRHITransitionInfo(OutputTexture, ERHIAccess::SRVMask, ERHIAccess::CopyDest));
+				RHICmdList.Transition(FRHITransitionInfo(OutputTexture, OutputTextureAccessBefore, ERHIAccess::CopyDest));
 				RHICmdList.CopyTexture(InputTexture->GetRHI(), OutputTexture, CopyInfo);
-				RHICmdList.Transition(FRHITransitionInfo(OutputTexture, ERHIAccess::CopyDest, ERHIAccess::SRVMask));
+				RHICmdList.Transition(FRHITransitionInfo(OutputTexture, ERHIAccess::CopyDest, OutputTextureAccessAfter));
 			});
 	}
 
@@ -1382,12 +1382,18 @@ namespace RuntimeVirtualTexture
 		bool bIsThumbnails,
 		bool bAllowCachedMeshDrawCommands,
 		FRHITexture* OutputTexture0,		// todo[vt]: Only use IPooledRenderTarget or FRDGTextureRef, not raw RHI textures.
+		ERHIAccess OutputTextureAccessBefore0, 
+		ERHIAccess OutputTextureAccessAfter0, 
 		IPooledRenderTarget* OutputTarget0,
 		FBox2D const& DestBox0,
 		FRHITexture* OutputTexture1,
+		ERHIAccess OutputTextureAccessBefore1, 
+		ERHIAccess OutputTextureAccessAfter1,
 		IPooledRenderTarget* OutputTarget1,
 		FBox2D const& DestBox1,
 		FRHITexture* OutputTexture2, 
+		ERHIAccess OutputTextureAccessBefore2, 
+		ERHIAccess OutputTextureAccessAfter2, 
 		IPooledRenderTarget* OutputTarget2,
 		FBox2D const& DestBox2,
 		FTransform const& UVToWorld,
@@ -1538,15 +1544,15 @@ namespace RuntimeVirtualTexture
 		// Copy to Output for each output texture
 		if (GraphSetup.OutputAlias0 != nullptr && OutputTexture0 != nullptr)
 		{
-			AddCopyToOutputPass(GraphBuilder, GraphSetup.OutputAlias0, OutputTexture0, DestBox0);
+			AddCopyToOutputPass(GraphBuilder, GraphSetup.OutputAlias0, OutputTexture0, OutputTextureAccessBefore0, OutputTextureAccessAfter0, DestBox0);
 		}
 		if (GraphSetup.OutputAlias1 != nullptr && OutputTexture1 != nullptr)
 		{
-			AddCopyToOutputPass(GraphBuilder, GraphSetup.OutputAlias1, OutputTexture1, DestBox1);
+			AddCopyToOutputPass(GraphBuilder, GraphSetup.OutputAlias1, OutputTexture1, OutputTextureAccessBefore1, OutputTextureAccessAfter1, DestBox1);
 		}
 		if (GraphSetup.OutputAlias2 != nullptr && OutputTexture2 != nullptr)
 		{
-			AddCopyToOutputPass(GraphBuilder, GraphSetup.OutputAlias2, OutputTexture2, DestBox2);
+			AddCopyToOutputPass(GraphBuilder, GraphSetup.OutputAlias2, OutputTexture2, OutputTextureAccessBefore2, OutputTextureAccessAfter2, DestBox2);
 		}
 	}
 
@@ -1572,9 +1578,9 @@ namespace RuntimeVirtualTexture
 					InDesc.bClearTextures,
 					InDesc.bIsThumbnails,
 					bAllowCachedMeshDrawCommands,
-					InDesc.Targets[0].Texture, InDesc.Targets[0].PooledRenderTarget, PageDesc.DestBox[0],
-					InDesc.Targets[1].Texture, InDesc.Targets[1].PooledRenderTarget, PageDesc.DestBox[1],
-					InDesc.Targets[2].Texture, InDesc.Targets[2].PooledRenderTarget, PageDesc.DestBox[2],
+					InDesc.Targets[0].Texture, InDesc.Targets[0].TextureAccessBefore, InDesc.Targets[0].TextureAccessAfter, InDesc.Targets[0].PooledRenderTarget, PageDesc.DestBox[0],
+					InDesc.Targets[1].Texture, InDesc.Targets[1].TextureAccessBefore, InDesc.Targets[1].TextureAccessAfter, InDesc.Targets[1].PooledRenderTarget, PageDesc.DestBox[1],
+					InDesc.Targets[2].Texture, InDesc.Targets[2].TextureAccessBefore, InDesc.Targets[2].TextureAccessAfter, InDesc.Targets[2].PooledRenderTarget, PageDesc.DestBox[2],
 					InDesc.UVToWorld,
 					InDesc.WorldBounds,
 					PageDesc.UVRange,
diff --git a/Engine/Source/Runtime/Renderer/Public/VT/RuntimeVirtualTextureRender.h b/Engine/Source/Runtime/Renderer/Public/VT/RuntimeVirtualTextureRender.h
--- a/Engine/Source/Runtime/Renderer/Public/VT/RuntimeVirtualTextureRender.h
+++ b/Engine/Source/Runtime/Renderer/Public/VT/RuntimeVirtualTextureRender.h
@@ -4,6 +4,7 @@
 
 #include "CoreMinimal.h"
 #include "VT/RuntimeVirtualTextureEnum.h"
+#include "RHIAccess.h"
 
 class FRHICommandListImmediate;
 class FRHITexture;
@@ -36,8 +37,11 @@ namespace RuntimeVirtualTexture
 		/** Physical texture to render to. */
 		FRHITexture* Texture = nullptr;
 
-		UE_DEPRECATED(5.1, "UAV is deprecated. Register the pooled render target with RDG instead.")
-		FRHIUnorderedAccessView* UAV = nullptr;
+		/** State of the physical texture to render to when it's being rendered */
+		ERHIAccess TextureAccessBefore = ERHIAccess::SRVMask;
+
+		/** State the physical texture to render should be left in after it's being rendered */
+		ERHIAccess TextureAccessAfter = ERHIAccess::SRVMask;
 
 		IPooledRenderTarget* PooledRenderTarget = nullptr;
 	};
