Don't combine draw commands with different GPU culling payload flags.
This avoids case where two distinct ISMs with the same SM but with different LOD ranges get their draw commands combined.
One might be culling max screen size and the other min screen size.
When the two are combined only one set of flags survive, which can lead to unwanted culling for the instances using the wrong flags.
Later it could be good to see if we can get the flags included in the StateBucketId which is sorted so would guarantee optimal batching.

diff --git a/Engine/Source/Runtime/Renderer/Private/InstanceCulling/InstanceCullingContext.cpp b/Engine/Source/Runtime/Renderer/Private/InstanceCulling/InstanceCullingContext.cpp
--- a/Engine/Source/Runtime/Renderer/Private/InstanceCulling/InstanceCullingContext.cpp
+++ b/Engine/Source/Runtime/Renderer/Private/InstanceCulling/InstanceCullingContext.cpp
@@ -1476,6 +1476,7 @@ void FInstanceCullingContext::SetupDrawCommands(
 	}
 
 	int32 CurrentStateBucketId = -1;
+	EMeshDrawCommandCullingPayloadFlags CurrentCullingPayloadFlags = EMeshDrawCommandCullingPayloadFlags::Default;
 	MaxInstances = 1;
 	// Only used to supply stats
 	uint32 CurrentAutoInstanceCount = 1;
@@ -1507,7 +1508,7 @@ void FInstanceCullingContext::SetupDrawCommands(
 		// UniformBufferView path does not support merging ISM draws atm
 		const bool bCompactIdenticalCommands = bInCompactIdenticalCommands && (bUsesUniformBufferView ? (CurrentAutoInstanceCount < MaxPrimitiveBatchSize && !bUseIndirectDraw) : true);
 
-		if (bCompactIdenticalCommands && CurrentStateBucketId != -1 && VisibleMeshDrawCommand.StateBucketId == CurrentStateBucketId)
+		if (bCompactIdenticalCommands && CurrentStateBucketId != -1 && VisibleMeshDrawCommand.StateBucketId == CurrentStateBucketId && VisibleMeshDrawCommand.CullingPayloadFlags == CurrentCullingPayloadFlags)
 		{
 			// Drop since previous covers for this
 
@@ -1568,6 +1569,7 @@ void FInstanceCullingContext::SetupDrawCommands(
 			
 			// Record the last bucket ID (may be -1)
 			CurrentStateBucketId = VisibleMeshDrawCommand.StateBucketId;
+			CurrentCullingPayloadFlags = VisibleMeshDrawCommand.CullingPayloadFlags;
 
 			// If we have dropped any we need to move up to maintain 1:1
 			if (DrawCommandIndex > NumDrawCommandsOut)
