Make culling to sphere edge the new frustum culling default

A/B testing showed no drop in perf w/ sphere edge culling enabled

diff --git a/Engine/Source/Runtime/Renderer/Private/SceneVisibility.cpp b/Engine/Source/Runtime/Renderer/Private/SceneVisibility.cpp
--- a/Engine/Source/Runtime/Renderer/Private/SceneVisibility.cpp
+++ b/Engine/Source/Runtime/Renderer/Private/SceneVisibility.cpp
@@ -56,7 +56,7 @@
 #include "DecalRenderingCommon.h"
 #include "CompositionLighting/PostProcessDeferredDecals.h"
 
-bool GDistanceCullToSphereEdge = false;
+bool GDistanceCullToSphereEdge = true;
 static FAutoConsoleVariableRef CVarDistanceCullToSphereEdge(
 	TEXT("r.DistanceCullToSphereEdge"),
 	GDistanceCullToSphereEdge,
