CustomRenderPass:  Fix for issue with Landscape rendering.  "AllViews" array in main view family needs to be copied to CustomRenderPass view families, so any plugin rendering logic that reads the array through the view family pointer sees the same data.  For Landscape, the inconsistency in "AllViews" caused LOD calculations to be skipped for some of the views, leading to an assert.  This didn't affect most other rendering, because there's an "AllViews" array in the scene renderer itself used nearly everywhere else, but plugin virtual functions don't have access to the scene renderer structure.

#rnx

diff --git a/Engine/Source/Runtime/Renderer/Private/SceneRendering.cpp b/Engine/Source/Runtime/Renderer/Private/SceneRendering.cpp
--- a/Engine/Source/Runtime/Renderer/Private/SceneRendering.cpp
+++ b/Engine/Source/Runtime/Renderer/Private/SceneRendering.cpp
@@ -2940,7 +2940,7 @@ FSceneRenderer::FSceneRenderer(const FSceneViewFamily* InViewFamily, FHitProxyCo
 		FSceneView NewView(ViewInitOptions);
 		FViewInfo* ViewInfo = &CustomRenderPassInfo.Views.Emplace_GetRef(&NewView);
 		CustomRenderPassInfo.ViewFamily.Views.Add(ViewInfo);
-		CustomRenderPassInfo.ViewFamily.AllViews.Add(ViewInfo);
+
 		// Must initialize to have a GPUScene connected to be able to collect dynamic primitives.
 		ViewInfo->DynamicPrimitiveCollector = FGPUScenePrimitiveCollector(&GPUSceneDynamicContext);
 		ViewInfo->bDisableQuerySubmissions = true;
@@ -3014,6 +3014,12 @@ FSceneRenderer::FSceneRenderer(const FSceneViewFamily* InViewFamily, FHitProxyCo
 	check(!ViewFamily.AllViews.Num());
 	ViewFamily.AllViews.Append(AllViews);
 
+	// Mirror AllViews across CustomRenderPass view families
+	for (FCustomRenderPassInfo& PassInfo : CustomRenderPassInfos)
+	{
+		PassInfo.ViewFamily.AllViews = ViewFamily.AllViews;
+	}
+
 	Scene->CustomRenderPassRendererInputs.Reset();
 
 	FeatureLevel = Scene->GetFeatureLevel();
