Fix iteration number range to [0,MaxPP-1] and zero divide for 1spp.
#rnx
#tests dumpgpu shows variance is zero.

diff --git a/Engine/Shaders/Private/PathTracing/PathTracingSpatialTemporalDenoising.usf b/Engine/Shaders/Private/PathTracing/PathTracingSpatialTemporalDenoising.usf
--- a/Engine/Shaders/Private/PathTracing/PathTracingSpatialTemporalDenoising.usf
+++ b/Engine/Shaders/Private/PathTracing/PathTracingSpatialTemporalDenoising.usf
@@ -875,7 +875,7 @@ void PrepassGenerateTextureCS(uint3 DT_ID : SV_DispatchThreadID)
 
 		// VarianceInfo.Var = sigma^2 * n.
 		// Std of mean = sqrt(sigma^2/n) = sqrt(VarianceInfo.Var/n^2)
-		float4 StandDerivation= sqrt(VarianceInfo.Var) / Iteration;
+		float4 StandDerivation= sqrt(VarianceInfo.Var) / (Iteration + 1);
 
 		OutputTexture[Position] = StandDerivation;
 	}
diff --git a/Engine/Source/Runtime/Renderer/Private/PathTracingSpatialTemporalDenoising.cpp b/Engine/Source/Runtime/Renderer/Private/PathTracingSpatialTemporalDenoising.cpp
--- a/Engine/Source/Runtime/Renderer/Private/PathTracingSpatialTemporalDenoising.cpp
+++ b/Engine/Source/Runtime/Renderer/Private/PathTracingSpatialTemporalDenoising.cpp
@@ -1762,7 +1762,7 @@ void PathTracingSpatialTemporalDenoisingPrePass(FRDGBuilder& GraphBuilder, const
 				PassParameters->OutputTexture = GraphBuilder.CreateUAV(SpatialTemporalDenoisingContext.VarianceTexture);
 				PassParameters->VarianceMap	  = GraphBuilder.CreateSRV(VarianceBuffer, EPixelFormat::PF_R32_FLOAT);
 				PassParameters->TargetViewport = TargetViewportParameters;
-				PassParameters->Iteration = IterationNumber;
+				PassParameters->Iteration = FMath::Min(IterationNumber, MaxSPP - 1);
 			}
 
 			SHADER::FPermutationDomain ComputeShaderPermutationVector;
