Add missing ECVF_RenderThreadSafe to r.RayTracing.CompileMaterialAHS/CHS to prevent a warning.
#rnx

diff --git a/Engine/Source/Runtime/Renderer/Private/RayTracing/RayTracingMaterialHitShaders.cpp b/Engine/Source/Runtime/Renderer/Private/RayTracing/RayTracingMaterialHitShaders.cpp
--- a/Engine/Source/Runtime/Renderer/Private/RayTracing/RayTracingMaterialHitShaders.cpp
+++ b/Engine/Source/Runtime/Renderer/Private/RayTracing/RayTracingMaterialHitShaders.cpp
@@ -37,7 +37,7 @@ static FAutoConsoleVariableRef CVarCompileRayTracingMaterialCHS(
 	GCompileRayTracingMaterialCHS,
 	TEXT(" 0: skip compilation of closest-hit shaders for materials (useful if only shadows or ambient occlusion effects are needed)\n")
 	TEXT(" 1: compile closest hit shaders for all ray tracing materials (default)\n"),
-	ECVF_ReadOnly
+	ECVF_ReadOnly | ECVF_RenderThreadSafe
 );
 
 int32 GCompileRayTracingMaterialAHS = 1;
@@ -46,7 +46,7 @@ static FAutoConsoleVariableRef CVarCompileRayTracingMaterialAHS(
 	GCompileRayTracingMaterialAHS,
 	TEXT(" 0: skip compilation of any-hit shaders for materials (useful if alpha masked or translucent materials are not needed)\n")
 	TEXT(" 1: compile any hit shaders for all ray tracing materials (default)\n"),
-	ECVF_ReadOnly
+	ECVF_ReadOnly | ECVF_RenderThreadSafe
 );
 
 static int32 GRayTracingNonBlockingPipelineCreation = 1;
