Add support for lit decals when using mobile forward renderer without dbuffer

diff --git a/Engine/Shaders/Private/DeferredDecal.usf b/Engine/Shaders/Private/DeferredDecal.usf
--- a/Engine/Shaders/Private/DeferredDecal.usf
+++ b/Engine/Shaders/Private/DeferredDecal.usf
@@ -28,6 +28,10 @@
 #include "Substrate/Substrate.ush"
 #include "DecalCommon.ush"
 
+#if DECAL_RENDERSTAGE == DECAL_RENDERSTAGE_MOBILE
+#include "MobileLightingCommon.ush"
+#endif
+
 #endif // PIXELSHADER
 
 // large world coordinate tile position
@@ -136,7 +140,8 @@ void FPixelShaderInOut_MainPS(inout FPixelShaderIn In, inout FPixelShaderOut Out
 	In.SvPosition.z = LookupDeviceZ(ScreenUV);
 	float4x4 SvPositionToDecalForSlice = SvPositionToDecal;
 #endif
-	
+	In.SvPosition.w = ConvertFromDeviceZ(In.SvPosition.z);
+
 	float4 DecalVectorHom = mul(float4(In.SvPosition.xyz,1), SvPositionToDecalForSlice);
 	float3 OSPosition = DecalVectorHom.xyz / DecalVectorHom.w;
 
@@ -169,7 +174,7 @@ void FPixelShaderInOut_MainPS(inout FPixelShaderIn In, inout FPixelShaderOut Out
 	MaterialParameters.VertexColor = 1;
 	MaterialParameters.CameraVector = CameraVector;
 	MaterialParameters.SvPosition = In.SvPosition;
-	MaterialParameters.ScreenPosition = SvPositionToScreenPosition(In.SvPosition);
+	MaterialParameters.ScreenPosition = SvPositionToResolvedScreenPosition(In.SvPosition);
 	MaterialParameters.LightVector = SwizzlePos;
 
 	MaterialParameters.AbsoluteWorldPosition = MaterialParameters.WorldPosition_NoOffsets = DFFastSubtract(TranslatedWorldPosition, PrimaryView.PreViewTranslation);
@@ -233,6 +238,13 @@ void FPixelShaderInOut_MainPS(inout FPixelShaderIn In, inout FPixelShaderOut Out
 	DecalCommonOutput(In, Out, Color, Opacity, GBufferData);
 #endif // SUBSTRATE_ENABLED
 
+#if DECAL_RENDERSTAGE == DECAL_RENDERSTAGE_MOBILE
+	#if DECAL_MOBILE_FORWARD_LIT
+		Out.MRT[0].rgb = ForwardDecalLighting(TranslatedWorldPosition, Color, Opacity, MaterialParameters.ScreenPosition.xy, MaterialParameters.ScreenPosition.w,
+			GBufferData, CameraVector, MaterialParameters.TangentToWorld[2], ReflectionAboutCustomWorldNormal(MaterialParameters, MaterialParameters.WorldNormal, false));
+	#endif
+#endif
+
 #if MATERIAL_VIRTUALTEXTURE_FEEDBACK
 	FinalizeVirtualTextureFeedback(
 		MaterialParameters.VirtualTextureFeedback,
diff --git a/Engine/Shaders/Private/MobileLightingCommon.ush b/Engine/Shaders/Private/MobileLightingCommon.ush
--- a/Engine/Shaders/Private/MobileLightingCommon.ush
+++ b/Engine/Shaders/Private/MobileLightingCommon.ush
@@ -301,7 +301,7 @@ void AccumulateDirectionalLighting(FGBufferData GBuffer, float3 TranslatedWorldP
 	Mobile Reflection.
 ------------------------------------------------------------------------------*/
 
-#if IS_MOBILE_BASE_PASS
+#if IS_MOBILE_BASE_PASS || IS_DECAL
 /** Prenormalized capture of the scene that's closest to the object being rendered. */
 half3 GetMobileSkyLightReflection(half3 ReflectionVector, half Roughness, half CubemapMaxMip)
 {
@@ -677,4 +677,35 @@ void MergeLocalLights(const FCulledLightsGridData InLightGridData
 		}
 	}
 }
-#endif
\ No newline at end of file
+#endif
+
+/*------------------------------------------------------------------------------
+	Mobile Forward Decal Lighting
+------------------------------------------------------------------------------*/
+
+half3 ForwardDecalLighting(float3 TranslatedWorldPosition, half3 EmissiveColor, half Opacity, float2 ScreenPosition, float SceneDepth, FGBufferData GBuffer, half3 CameraVector, half3 VertexNormal, half3 ReflectionVector)
+{
+	float ShadowPositionZ = 0;
+	half ShadowMap = MobileDirectionalLightCSM(ScreenPosition.xy, SceneDepth, ShadowPositionZ);
+	half3 DirectionalLightDirection = MobileDirectionalLight.DirectionalLightDirectionAndShadowTransition.xyz;
+	half NoL = saturate(dot(half3(GBuffer.WorldNormal), DirectionalLightDirection));
+	
+	half3 SpecularColor = ComputeF0(GBuffer.Specular, GBuffer.BaseColor, GBuffer.Metallic);
+	half3 DiffuseColor = GBuffer.BaseColor - GBuffer.BaseColor * GBuffer.Metallic;
+
+	half3 OutColor = EmissiveColor;
+	OutColor += MobileDirectionalLight.DirectionalLightColor.rgb * NoL * SimpleShading(DiffuseColor, SpecularColor, max(GBuffer.Roughness, .04f), DirectionalLightDirection, -CameraVector, GBuffer.WorldNormal) * ShadowMap;
+
+	half3 SkyDiffuse = GetSkySHDiffuseSimple(GBuffer.WorldNormal) * ResolvedView.SkyLightColor.rgb;
+	OutColor += DiffuseColor * SkyDiffuse;
+
+	half3 SpecularIBL = GetImageBasedReflectionSpecular(ReflectionVector, GBuffer.Roughness, TranslatedWorldPosition, GBuffer.WorldNormal, 1.0, Luminance(SkyDiffuse), 0);
+	half NoV = saturate(abs(dot(GBuffer.WorldNormal, -CameraVector)) + 1e-5);
+	OutColor += SpecularIBL * EnvBRDFApprox(SpecularColor, GBuffer.Roughness, NoV);
+
+#if MATERIALBLENDING_MODULATE
+	OutColor *= Opacity;
+#endif
+
+	return OutColor;
+}
\ No newline at end of file
diff --git a/Engine/Source/Runtime/Core/Private/HAL/ConsoleManager.cpp b/Engine/Source/Runtime/Core/Private/HAL/ConsoleManager.cpp
--- a/Engine/Source/Runtime/Core/Private/HAL/ConsoleManager.cpp
+++ b/Engine/Source/Runtime/Core/Private/HAL/ConsoleManager.cpp
@@ -4277,6 +4277,16 @@ static TAutoConsoleVariable<int32> CVarMobileDBuffer(
 		" 1: on"),
 	ECVF_RenderThreadSafe | ECVF_ReadOnly);
 
+static TAutoConsoleVariable<int32> CVarMobileForwardDecalLighting(
+	TEXT("r.Mobile.Forward.DecalLighting"),
+	1,
+	TEXT("Enables lit decals when using the mobile forward renderer without DBuffer.\n"
+		" 0: off\n"
+		" 1: on, lighting channel 1 is used for decal lighting\n"
+		" 2: on, lighting channel 2 is used for decal lighting\n"
+		" 3: on, lighting channel 3 is used for decal lighting"),
+	ECVF_RenderThreadSafe | ECVF_ReadOnly);
+
 static TAutoConsoleVariable<float> CVarSkeletalMeshLODRadiusScale(
 	TEXT("r.SkeletalMeshLODRadiusScale"),
 	1.0f,
diff --git a/Engine/Source/Runtime/RenderCore/Private/ReadOnlyCVARCache.cpp b/Engine/Source/Runtime/RenderCore/Private/ReadOnlyCVARCache.cpp
--- a/Engine/Source/Runtime/RenderCore/Private/ReadOnlyCVARCache.cpp
+++ b/Engine/Source/Runtime/RenderCore/Private/ReadOnlyCVARCache.cpp
@@ -31,6 +31,7 @@ bool FReadOnlyCVARCache::bMobileForwardParticleLights = false;
 bool FReadOnlyCVARCache::bMobileEnableNoPrecomputedLightingCSMShader = false;
 bool FReadOnlyCVARCache::bMobileDeferredShadingValue = false;
 bool FReadOnlyCVARCache::bMobileEnableMovableSpotlightsShadowValue = false;
+int32 FReadOnlyCVARCache::MobileForwardDecalLightingValue = 1;
 
 int32 FReadOnlyCVARCache::MobileEarlyZPassIniValue(EShaderPlatform Platform)
 {
@@ -94,6 +95,7 @@ void FReadOnlyCVARCache::Initialize()
 	const auto CVarMobileSkyLightPermutation = IConsoleManager::Get().FindTConsoleVariableDataInt(TEXT("r.Mobile.SkyLightPermutation"));
 	const auto CVarMobileEnableNoPrecomputedLightingCSMShader = IConsoleManager::Get().FindTConsoleVariableDataInt(TEXT("r.Mobile.EnableNoPrecomputedLightingCSMShader"));
 	const auto CVarMobileSupportGPUScene = IConsoleManager::Get().FindTConsoleVariableDataInt(TEXT("r.Mobile.SupportGPUScene"));
+	const auto CVarMobileForwardDecalLightingValue = IConsoleManager::Get().FindTConsoleVariableDataInt(TEXT("r.Mobile.Forward.DecalLighting"));
 	
 	const bool bForceAllPermutations = CVarSupportAllShaderPermutations && CVarSupportAllShaderPermutations->GetValueOnAnyThread() != 0;
 
@@ -117,6 +119,7 @@ void FReadOnlyCVARCache::Initialize()
 	bMobileEnableMovableSpotlightsShadowValue = MobileEnableMovableSpotlightsShadowIniValue(GMaxRHIShaderPlatform);
 	bMobileSupportsGPUScene = CVarMobileSupportGPUScene->GetValueOnAnyThread() != 0;
 	bMobileForwardParticleLights = MobileForwardParticleLightsIniValue(GMaxRHIShaderPlatform);
+	MobileForwardDecalLightingValue = CVarMobileForwardDecalLightingValue->GetValueOnAnyThread();
 
 #ifdef PROJECT_CVAR_ALLOW_STATIC_LIGHTING
 	check(!!PROJECT_CVAR_ALLOW_STATIC_LIGHTING == bAllowStaticLighting);
diff --git a/Engine/Source/Runtime/RenderCore/Public/ReadOnlyCVARCache.h b/Engine/Source/Runtime/RenderCore/Public/ReadOnlyCVARCache.h
--- a/Engine/Source/Runtime/RenderCore/Public/ReadOnlyCVARCache.h
+++ b/Engine/Source/Runtime/RenderCore/Public/ReadOnlyCVARCache.h
@@ -92,6 +92,11 @@ struct FReadOnlyCVARCache
 		return bMobileEnableNoPrecomputedLightingCSMShader;
 	}
 
+	static inline int32 MobileForwardDecalLighting()
+	{
+		return MobileForwardDecalLightingValue;
+	}
+
 	static inline int32 MobileEarlyZPass(EShaderPlatform Platform)
 	{
 		#if WITH_EDITOR
@@ -163,6 +168,7 @@ private:
 	RENDERCORE_API static int32 MobileEarlyZPassValue;
 	RENDERCORE_API static int32 MobileForwardLocalLightsValue;
 	RENDERCORE_API static bool bMobileForwardParticleLights;
+	RENDERCORE_API static int32 MobileForwardDecalLightingValue;
 	RENDERCORE_API static bool bMobileEnableNoPrecomputedLightingCSMShader;
 	RENDERCORE_API static bool bMobileDeferredShadingValue;
 	RENDERCORE_API static bool bMobileEnableMovableSpotlightsShadowValue;
diff --git a/Engine/Source/Runtime/Renderer/Private/DecalRenderingCommon.cpp b/Engine/Source/Runtime/Renderer/Private/DecalRenderingCommon.cpp
--- a/Engine/Source/Runtime/Renderer/Private/DecalRenderingCommon.cpp
+++ b/Engine/Source/Runtime/Renderer/Private/DecalRenderingCommon.cpp
@@ -27,6 +27,7 @@ namespace DecalRendering
 		const bool bIsSubstrateEnabled = Substrate::IsSubstrateEnabled();
 		const bool bIsMobilePlatform = IsMobilePlatform(Platform);
 		const bool bIsMobileDeferredPlatform = bIsMobilePlatform && IsMobileDeferredShadingEnabled(Platform);
+		const bool bIsMobileLitDecalPlatform = bIsMobilePlatform && FReadOnlyCVARCache::MobileForwardDecalLighting();
 		const bool bIsDBufferPlatform = IsUsingDBuffers(Platform);
 		const bool bIsDBufferMaskPlatform = bIsDBufferPlatform && FDataDrivenShaderPlatformInfo::GetSupportsPerPixelDBufferMask(Platform);
 
@@ -43,7 +44,7 @@ namespace DecalRendering
 		}
 
 		// Enforce platform output limitations.
-		if (bIsMobilePlatform && !bIsMobileDeferredPlatform && !bIsDBufferPlatform)
+		if (bIsMobilePlatform && !bIsMobileDeferredPlatform && !bIsDBufferPlatform && !bIsMobileLitDecalPlatform)
 		{
 			Desc.bWriteNormal = false;
 			Desc.bWriteRoughnessSpecularMetallic = false;
diff --git a/Engine/Source/Runtime/Renderer/Private/DecalRenderingShared.cpp b/Engine/Source/Runtime/Renderer/Private/DecalRenderingShared.cpp
--- a/Engine/Source/Runtime/Renderer/Private/DecalRenderingShared.cpp
+++ b/Engine/Source/Runtime/Renderer/Private/DecalRenderingShared.cpp
@@ -95,9 +95,11 @@ public:
 		DecalParams.Bind(Initializer.ParameterMap, TEXT("DecalParams"));
 		DecalColorParam.Bind(Initializer.ParameterMap, TEXT("DecalColorParam"));
 		MobileBasePassUniformBuffer.Bind(Initializer.ParameterMap, FMobileBasePassUniformParameters::FTypeInfo::GetStructMetadata()->GetShaderVariableName());
+		MobileDirectionLightBufferParam.Bind(Initializer.ParameterMap, FMobileDirectionalLightShaderParameters::FTypeInfo::GetStructMetadata()->GetShaderVariableName());
+		MobileReflectionCaptureParam.Bind(Initializer.ParameterMap, FMobileReflectionCaptureShaderParameters::FTypeInfo::GetStructMetadata()->GetShaderVariableName());
 	}
 
-	void SetParameters(FRHIBatchedShaderParameters& BatchedParameters, const FViewInfo& View, const FDeferredDecalProxy& DecalProxy, const FMaterialRenderProxy* MaterialProxy, const FMaterial* MaterialResource, const float FadeAlphaValue = 1.0f)
+	void SetParameters(FRHIBatchedShaderParameters& BatchedParameters, const FViewInfo& View, const FDeferredDecalProxy& DecalProxy, const FMaterialRenderProxy* MaterialProxy, const FMaterial* MaterialResource, const float FadeAlphaValue = 1.0f, const FScene* Scene = nullptr)
 	{
 		auto& PrimitivePS = GetUniformBufferParameter<FPrimitiveUniformShaderParameters>();
 		SetUniformBufferParameter(BatchedParameters, PrimitivePS, GIdentityPrimitiveUniformBuffer);
@@ -175,6 +177,24 @@ public:
  
 		SetShaderValue(BatchedParameters, DecalParams, FVector2f(FadeAlphaValue, LifetimeAlpha));
 		SetShaderValue(BatchedParameters, DecalColorParam, DecalProxy.DecalColor);
+
+		if (MobileDirectionLightBufferParam.IsBound() && Scene)
+		{
+			const int UniformBufferIndex = FMath::Clamp(FReadOnlyCVARCache::MobileForwardDecalLighting(), 1, 3);
+			SetUniformBufferParameter(BatchedParameters, MobileDirectionLightBufferParam, Scene->UniformBuffers.MobileDirectionalLightUniformBuffers[UniformBufferIndex]);
+		}
+
+		if (MobileReflectionCaptureParam.IsBound())
+		{
+			if (Scene && Scene->SkyLight && Scene->SkyLight->ProcessedTexture && Scene->SkyLight->ProcessedTexture->TextureRHI)
+			{
+				SetUniformBufferParameter(BatchedParameters, MobileReflectionCaptureParam, Scene->UniformBuffers.MobileSkyReflectionUniformBuffer);
+			}
+			else
+			{
+				SetUniformBufferParameter(BatchedParameters, MobileReflectionCaptureParam, GDefaultMobileReflectionCaptureUniformBuffer.GetUniformBufferRHI());
+			}
+		}
 	}
 
 private:
@@ -187,6 +207,8 @@ private:
 	LAYOUT_FIELD(FShaderParameter, DecalParams);
 	LAYOUT_FIELD(FShaderParameter, DecalColorParam);
 	LAYOUT_FIELD(FShaderUniformBufferParameter, MobileBasePassUniformBuffer);
+	LAYOUT_FIELD(FShaderUniformBufferParameter, MobileDirectionLightBufferParam);	
+	LAYOUT_FIELD(FShaderUniformBufferParameter, MobileReflectionCaptureParam);
 };
 
 IMPLEMENT_MATERIAL_SHADER_TYPE(,FDeferredDecalPS,TEXT("/Engine/Private/DeferredDecal.usf"),TEXT("MainPS"),SF_Pixel);
@@ -239,6 +261,32 @@ public:
 
 IMPLEMENT_MATERIAL_SHADER_TYPE(, FDeferredDecalAmbientOcclusionPS, TEXT("/Engine/Private/DeferredDecal.usf"), TEXT("MainPS"), SF_Pixel);
 
+
+class FDeferredDecalMobilePS : public FDeferredDecalPS
+{
+	DECLARE_SHADER_TYPE(FDeferredDecalMobilePS, Material);
+
+public:
+	static bool ShouldCompilePermutation(const FMaterialShaderPermutationParameters& Parameters)
+	{
+		return (Parameters.MaterialParameters.MaterialDomain == MD_DeferredDecal) &&
+			DecalRendering::IsCompatibleWithRenderStage(DecalRendering::ComputeDecalBlendDesc(Parameters.Platform, Parameters.MaterialParameters), EDecalRenderStage::Mobile);
+	}
+	static void ModifyCompilationEnvironment(const FMaterialShaderPermutationParameters& Parameters, FShaderCompilerEnvironment& OutEnvironment)
+	{
+		FMaterialShader::ModifyCompilationEnvironment(Parameters, OutEnvironment);
+		DecalRendering::ModifyCompilationEnvironment(Parameters.Platform, DecalRendering::ComputeDecalBlendDesc(Parameters.Platform, Parameters.MaterialParameters), EDecalRenderStage::Mobile, OutEnvironment);
+		OutEnvironment.SetDefine(TEXT("DECAL_MOBILE_FORWARD_LIT"), FReadOnlyCVARCache::MobileForwardDecalLighting() != 0 ? 1u : 0u);
+	}
+
+	FDeferredDecalMobilePS() {}
+	FDeferredDecalMobilePS(const ShaderMetaType::CompiledShaderInitializerType& Initializer)
+		: FDeferredDecalPS(Initializer)
+	{}
+};
+
+IMPLEMENT_MATERIAL_SHADER_TYPE(, FDeferredDecalMobilePS, TEXT("/Engine/Private/DeferredDecal.usf"), TEXT("MainPS"), SF_Pixel);
+
 namespace DecalRendering
 {
 	float GetDecalFadeScreenSizeMultiplier()
@@ -504,6 +552,10 @@ namespace DecalRendering
 		{
 			ShaderTypes.AddShaderType<FDeferredDecalAmbientOcclusionPS>();
 		}
+		else if (DecalRenderStage == EDecalRenderStage::Mobile)
+		{
+			ShaderTypes.AddShaderType<FDeferredDecalMobilePS>();
+		}
 		else
 		{
 			ShaderTypes.AddShaderType<FDeferredDecalPS>();
@@ -566,7 +618,7 @@ namespace DecalRendering
 	}
 
 	void SetShader(FRHICommandList& RHICmdList, FGraphicsPipelineStateInitializer& GraphicsPSOInit, uint32 StencilRef, const FViewInfo& View,
-		const FTransientDecalRenderData& DecalData, EDecalRenderStage DecalRenderStage, const FMatrix& FrustumComponentToClip)
+		const FTransientDecalRenderData& DecalData, EDecalRenderStage DecalRenderStage, const FMatrix& FrustumComponentToClip, const FScene* Scene)
 	{
 		FMaterial const* MaterialResource = nullptr;
 		TShaderRef<FDeferredDecalPS> PixelShader;
@@ -591,7 +643,7 @@ namespace DecalRendering
 
 		// Set pixel shader parameters.
 		{
-			SetShaderParametersLegacyPS(RHICmdList, PixelShader, View, *DecalData.Proxy, MaterialProxy, MaterialResource, DecalData.FadeAlpha);
+			SetShaderParametersLegacyPS(RHICmdList, PixelShader, View, *DecalData.Proxy, MaterialProxy, MaterialResource, DecalData.FadeAlpha, Scene);
 		}
 
 		// Set stream source after updating cached strides
diff --git a/Engine/Source/Runtime/Renderer/Private/DecalRenderingShared.h b/Engine/Source/Runtime/Renderer/Private/DecalRenderingShared.h
--- a/Engine/Source/Runtime/Renderer/Private/DecalRenderingShared.h
+++ b/Engine/Source/Runtime/Renderer/Private/DecalRenderingShared.h
@@ -43,5 +43,5 @@ namespace DecalRendering
 	FTransientDecalRenderDataList BuildVisibleDecalList(TConstArrayView<FDeferredDecalProxy*> Decals, const FViewInfo& View);
 	bool BuildRelevantDecalList(TConstArrayView<FTransientDecalRenderData> Decals, EDecalRenderStage DecalRenderStage, FTransientDecalRenderDataList* OutVisibleDecals);
 	bool SetupShaderState(ERHIFeatureLevel::Type FeatureLevel, const FMaterial& Material, EDecalRenderStage DecalRenderStage, FBoundShaderStateInput& OutBoundShaderState);
-	void SetShader(FRHICommandList& RHICmdList, FGraphicsPipelineStateInitializer& GraphicsPSOInit, uint32 StencilRef, const FViewInfo& View, const FTransientDecalRenderData& DecalData, EDecalRenderStage DecalRenderStage, const FMatrix& FrustumComponentToClip);
+	void SetShader(FRHICommandList& RHICmdList, FGraphicsPipelineStateInitializer& GraphicsPSOInit, uint32 StencilRef, const FViewInfo& View, const FTransientDecalRenderData& DecalData, EDecalRenderStage DecalRenderStage, const FMatrix& FrustumComponentToClip, const FScene* Scene = nullptr);
 };
diff --git a/Engine/Source/Runtime/Renderer/Private/MobileDecalRendering.cpp b/Engine/Source/Runtime/Renderer/Private/MobileDecalRendering.cpp
--- a/Engine/Source/Runtime/Renderer/Private/MobileDecalRendering.cpp
+++ b/Engine/Source/Runtime/Renderer/Private/MobileDecalRendering.cpp
@@ -146,7 +146,7 @@ void RenderDeferredDecalsMobile(FRHICommandList& RHICmdList, const FScene& Scene
 			GraphicsPSOInit.BlendState = DecalRendering::GetDecalBlendState(DecalData.BlendDesc, DecalRenderStage, RenderTargetMode);
 
 			// Set shader params
-			DecalRendering::SetShader(RHICmdList, GraphicsPSOInit, 0, View, DecalData, DecalRenderStage, FrustumComponentToClip);
+			DecalRendering::SetShader(RHICmdList, GraphicsPSOInit, 0, View, DecalData, DecalRenderStage, FrustumComponentToClip, &Scene);
 
 			RHICmdList.DrawIndexedPrimitive(GetUnitCubeIndexBuffer(), 0, 0, 8, 0, UE_ARRAY_COUNT(GCubeIndices) / 3, View.GetStereoPassInstanceFactor());
 		}
