Removing the WITH_MGPU block from FOcclusionQueryHelpers::GetNumBufferedFrames since it was preventing using r.NumBufferedOcclusionQueries correctly on DESKTOP

diff --git a/Engine/Source/Runtime/Renderer/Private/SceneOcclusion.cpp b/Engine/Source/Runtime/Renderer/Private/SceneOcclusion.cpp
--- a/Engine/Source/Runtime/Renderer/Private/SceneOcclusion.cpp
+++ b/Engine/Source/Runtime/Renderer/Private/SceneOcclusion.cpp
@@ -106,10 +106,6 @@ FOcclusionRandomStream GOcclusionRandomStream;
 
 int32 FOcclusionQueryHelpers::GetNumBufferedFrames(ERHIFeatureLevel::Type FeatureLevel)
 {
-#if WITH_MGPU
-	// TODO:  Should this still be differentiated for MGPU?  Originally this logic was here for AFR, which has been removed.
-	return FMath::Min<int32>(1, (int32)FOcclusionQueryHelpers::MaxBufferedOcclusionFrames);
-#else
 	int32 NumGPUS = 1;
 
 	static const auto NumBufferedQueriesVar = IConsoleManager::Get().FindTConsoleVariableDataInt(TEXT("r.NumBufferedOcclusionQueries"));
@@ -134,7 +130,6 @@ int32 FOcclusionQueryHelpers::GetNumBufferedFrames(ERHIFeatureLevel::Type Featur
 	}
 
 	return FMath::Clamp<int32>(NumExtraMobileFrames + NumBufferedQueriesVar->GetValueOnAnyThread() * NumGPUS, 1, (int32)FOcclusionQueryHelpers::MaxBufferedOcclusionFrames);
-#endif
 }
 
 
