Fixes a d3d validation failure when switching from TSR to TAA

diff --git a/Engine/Source/Runtime/Renderer/Private/PostProcess/TemporalAA.cpp b/Engine/Source/Runtime/Renderer/Private/PostProcess/TemporalAA.cpp
--- a/Engine/Source/Runtime/Renderer/Private/PostProcess/TemporalAA.cpp
+++ b/Engine/Source/Runtime/Renderer/Private/PostProcess/TemporalAA.cpp
@@ -586,7 +586,7 @@ FTAAOutputs AddTemporalAAPass(
 	const bool bIsMainPass = IsMainTAAConfig(Inputs.Pass);
 
 	// Whether to use camera cut shader permutation or not.
-	const bool bCameraCut = !InputHistory.IsValid() || InputHistory.OutputSliceIndex != 0 || View.bCameraCut;
+	const bool bCameraCut = !InputHistory.IsValid() || View.bCameraCut || InputHistory.RT[0]->GetDesc().IsArray();
 
 	const FIntPoint OutputExtent = Inputs.GetOutputExtent();
 
