Fix cluster shading being executed twice causing intensity discrepency on certain platforms.

On platforms supporting rect primitive, the vertex shader was executed twice, due to the vertex shader assuming triangle list topology.

diff --git a/Engine/Source/Runtime/Renderer/Private/ClusteredDeferredShadingPass.cpp b/Engine/Source/Runtime/Renderer/Private/ClusteredDeferredShadingPass.cpp
--- a/Engine/Source/Runtime/Renderer/Private/ClusteredDeferredShadingPass.cpp
+++ b/Engine/Source/Runtime/Renderer/Private/ClusteredDeferredShadingPass.cpp
@@ -206,7 +206,10 @@ static void InternalAddClusteredDeferredShadingPass(
 
 	// VS - Substrate tile parameters
 	EPrimitiveType PrimitiveType = PT_TriangleList;
-	PassParameters->SubstrateTile = Substrate::SetTileParameters(GraphBuilder, View, TileType, PrimitiveType);
+	if (Substrate::IsSubstrateEnabled())
+	{
+		PassParameters->SubstrateTile = Substrate::SetTileParameters(GraphBuilder, View, TileType, PrimitiveType);
+	}
 	
 	const TCHAR* TileTypeName = ToString(TileType);
 	GraphBuilder.AddPass(
