Added some profiler scopes to Hardware Ray Tracing functions.

diff --git a/Engine/Source/Runtime/Engine/Private/Rendering/RayTracingGeometryManager.cpp b/Engine/Source/Runtime/Engine/Private/Rendering/RayTracingGeometryManager.cpp
--- a/Engine/Source/Runtime/Engine/Private/Rendering/RayTracingGeometryManager.cpp
+++ b/Engine/Source/Runtime/Engine/Private/Rendering/RayTracingGeometryManager.cpp
@@ -439,6 +439,8 @@ void FRayTracingGeometryManager::BoostPriority(BuildRequestIndex InRequestIndex,
 
 void FRayTracingGeometryManager::ForceBuildIfPending(FRHIComputeCommandList& InCmdList, const TArrayView<const FRayTracingGeometry*> InGeometries)
 {
+	TRACE_CPUPROFILER_EVENT_SCOPE(FRayTracingGeometryManager::ForceBuildIfPending);
+
 	FScopeLock ScopeLock(&RequestCS);
 
 	BuildParams.Empty(FMath::Max(BuildParams.Max(), InGeometries.Num()));
diff --git a/Engine/Source/Runtime/Renderer/Private/RayTracing/RayTracingInstanceBufferUtil.cpp b/Engine/Source/Runtime/Renderer/Private/RayTracing/RayTracingInstanceBufferUtil.cpp
--- a/Engine/Source/Runtime/Renderer/Private/RayTracing/RayTracingInstanceBufferUtil.cpp
+++ b/Engine/Source/Runtime/Renderer/Private/RayTracing/RayTracingInstanceBufferUtil.cpp
@@ -149,6 +149,8 @@ void FillRayTracingInstanceUploadBuffer(
 	TArrayView<FRayTracingInstanceDescriptorInput> OutInstanceUploadData,
 	TArrayView<FVector4f> OutTransformData)
 {
+	TRACE_CPUPROFILER_EVENT_SCOPE(FillRayTracingInstanceUploadBuffer);
+
 	const FRayTracingSceneInitializer2& SceneInitializer = RayTracingSceneRHI->GetInitializer();
 	const uint32 NumLayers = SceneInitializer.NumNativeInstancesPerLayer.Num();
 
diff --git a/Engine/Source/Runtime/Renderer/Private/RayTracing/RayTracingScene.cpp b/Engine/Source/Runtime/Renderer/Private/RayTracing/RayTracingScene.cpp
--- a/Engine/Source/Runtime/Renderer/Private/RayTracing/RayTracingScene.cpp
+++ b/Engine/Source/Runtime/Renderer/Private/RayTracing/RayTracingScene.cpp
@@ -61,7 +61,8 @@ void FRayTracingScene::InitPreViewTranslation(const FViewMatrices& ViewMatrices)
 
 void FRayTracingScene::CreateWithInitializationData(FRDGBuilder& GraphBuilder, const FViewInfo& View, const FGPUScene* GPUScene, FRayTracingSceneWithGeometryInstances SceneWithGeometryInstances, ERDGPassFlags ComputePassFlags)
 {
-	QUICK_SCOPE_CYCLE_COUNTER(FRayTracingScene_BeginCreate);
+	TRACE_CPUPROFILER_EVENT_SCOPE(FRayTracingScene::CreateWithInitializationData);
+	QUICK_SCOPE_CYCLE_COUNTER(STAT_RayTracingScene_CreateWithInitializationData);
 
 	// Round up buffer sizes to some multiple to avoid pathological growth reallocations.
 	static constexpr uint32 AllocationGranularity = 8 * 1024;
@@ -263,8 +264,6 @@ void FRayTracingScene::CreateWithInitializationData(FRDGBuilder& GraphBuilder, c
 					RHICmdList.UnlockBuffer(TransformUploadBuffer);
 				}
 
-				QUICK_SCOPE_CYCLE_COUNTER(GetAccelerationStructuresAddresses);
-
 				for (uint32 GPUIndex : RHICmdList.GetGPUMask())
 				{
 					FRayTracingAccelerationStructureAddress* AddressesPtr = (FRayTracingAccelerationStructureAddress*)RHICmdList.LockBufferMGPU(
@@ -275,6 +274,8 @@ void FRayTracingScene::CreateWithInitializationData(FRDGBuilder& GraphBuilder, c
 
 					RHICmdList.EnqueueLambda([AddressesPtr, &SceneInitializer, GPUIndex](FRHICommandListBase&)
 					{
+						TRACE_CPUPROFILER_EVENT_SCOPE(GetAccelerationStructuresAddresses);
+
 						for (int32 GeometryIndex = 0; GeometryIndex < SceneInitializer.ReferencedGeometries.Num(); ++GeometryIndex)
 						{
 							AddressesPtr[GeometryIndex] = SceneInitializer.ReferencedGeometries[GeometryIndex]->GetAccelerationStructureAddress(GPUIndex);
