Fixed missing call to update the size of the RenderedPrimitives bit arrays in the VSM cache manager when caching is disabled.

diff --git a/Engine/Source/Runtime/Renderer/Private/VirtualShadowMaps/VirtualShadowMapCacheManager.cpp b/Engine/Source/Runtime/Renderer/Private/VirtualShadowMaps/VirtualShadowMapCacheManager.cpp
--- a/Engine/Source/Runtime/Renderer/Private/VirtualShadowMaps/VirtualShadowMapCacheManager.cpp
+++ b/Engine/Source/Runtime/Renderer/Private/VirtualShadowMaps/VirtualShadowMapCacheManager.cpp
@@ -533,7 +533,6 @@ bool FVirtualShadowMapArrayCacheManager::ShouldCreateExtension(FScene& Scene)
 ISceneExtensionUpdater* FVirtualShadowMapArrayCacheManager::CreateUpdater()
 {
 	if (GVSMNewInvalidations &&
-		IsCacheDataAvailable() &&
 		// NOTE: We need this check because shader platform can change during scene destruction so we need to ensure we
 		// don't try and run shaders on a new platform that doesn't support VSMs...
 		UseVirtualShadowMaps(Scene->GetShaderPlatform(), Scene->GetFeatureLevel()))
@@ -1586,38 +1585,41 @@ void FVirtualShadowMapInvalidationSceneUpdater::PreSceneUpdate(FRDGBuilder& Grap
 	// There may be a way to avoid doing this both in pre and post, but it is pretty light if there is nothing to do anyways.
 	CacheManager.ReallocatePersistentPrimitiveIndices();
 
-	FVirtualShadowMapArrayCacheManager::FInvalidatingPrimitiveCollector InvalidatingPrimitiveCollector(&CacheManager);
-
-	// Primitives that are tracked as always invalidating shadows, pipe through as transform updates
-	for (FPrimitiveSceneInfo* PrimitiveSceneInfo : CacheManager.Scene->ShadowScene->GetAlwaysInvalidatingPrimitives())
+	if (CacheManager.IsCacheDataAvailable())
 	{
-		InvalidatingPrimitiveCollector.UpdatedTransform(PrimitiveSceneInfo);
-	}
+		FVirtualShadowMapArrayCacheManager::FInvalidatingPrimitiveCollector InvalidatingPrimitiveCollector(&CacheManager);
 
-	// All removed primitives must invalidate their footprints in the VSM before leaving
-	for (FPrimitiveSceneInfo* PrimitiveSceneInfo : ChangeSet.RemovedPrimitiveSceneInfos)
-	{
-		InvalidatingPrimitiveCollector.Removed(PrimitiveSceneInfo);
-	}
-	// As must all primitive updates, 
-	for (FPrimitiveSceneInfo* PrimitiveSceneInfo : ChangeSet.UpdatedPrimitiveSceneInfos)
-	{
-		InvalidatingPrimitiveCollector.UpdatedTransform(PrimitiveSceneInfo);
-	}
+		// Primitives that are tracked as always invalidating shadows, pipe through as transform updates
+		for (FPrimitiveSceneInfo* PrimitiveSceneInfo : CacheManager.Scene->ShadowScene->GetAlwaysInvalidatingPrimitives())
+		{
+			InvalidatingPrimitiveCollector.UpdatedTransform(PrimitiveSceneInfo);
+		}
 
-	// TODO! Where do we get this data from...
-	/*
-	for (const auto& CullDistance : UpdatedInstanceCullDistance)
-	{
-		InvalidatingPrimitiveCollector.UpdatedTransform(CullDistance.Key->GetPrimitiveSceneInfo());
-	}
-	*/
+		// All removed primitives must invalidate their footprints in the VSM before leaving
+		for (FPrimitiveSceneInfo* PrimitiveSceneInfo : ChangeSet.RemovedPrimitiveSceneInfos)
+		{
+			InvalidatingPrimitiveCollector.Removed(PrimitiveSceneInfo);
+		}
+		// As must all primitive updates, 
+		for (FPrimitiveSceneInfo* PrimitiveSceneInfo : ChangeSet.UpdatedPrimitiveSceneInfos)
+		{
+			InvalidatingPrimitiveCollector.UpdatedTransform(PrimitiveSceneInfo);
+		}
 
-	// TODO: Perhaps pass this in from the caller in RendererScene?
-	FSceneUniformBuffer SceneUniforms;
-	CacheManager.Scene->GPUScene.FillSceneUniformBuffer(GraphBuilder, SceneUniforms);
+		// TODO! Where do we get this data from...
+		/*
+		for (const auto& CullDistance : UpdatedInstanceCullDistance)
+		{
+			InvalidatingPrimitiveCollector.UpdatedTransform(CullDistance.Key->GetPrimitiveSceneInfo());
+		}
+		*/
 
-	CacheManager.ProcessInvalidations(GraphBuilder, SceneUniforms, InvalidatingPrimitiveCollector);
+		// TODO: Perhaps pass this in from the caller in RendererScene?
+		FSceneUniformBuffer SceneUniforms;
+		CacheManager.Scene->GPUScene.FillSceneUniformBuffer(GraphBuilder, SceneUniforms);
+
+		CacheManager.ProcessInvalidations(GraphBuilder, SceneUniforms, InvalidatingPrimitiveCollector);
+	}
 }
 
 void FVirtualShadowMapInvalidationSceneUpdater::PostSceneUpdate(FRDGBuilder& GraphBuilder, const FScenePostUpdateChangeSet& ChangeSet)
@@ -1632,20 +1634,22 @@ void FVirtualShadowMapInvalidationSceneUpdater::PostGPUSceneUpdate(FRDGBuilder&
 {
 	// TODO: Separate scope for post-update pass?
 	SCOPED_NAMED_EVENT(FScene_VirtualShadowCacheUpdate, FColor::Orange);
+	if (CacheManager.IsCacheDataAvailable())
+	{
+		FVirtualShadowMapArrayCacheManager::FInvalidatingPrimitiveCollector InvalidatingPrimitiveCollector(&CacheManager);
 
-	FVirtualShadowMapArrayCacheManager::FInvalidatingPrimitiveCollector InvalidatingPrimitiveCollector(&CacheManager);
+		// All removed primitives must invalidate their footprints in the VSM before leaving
+		for (FPrimitiveSceneInfo* PrimitiveSceneInfo : PostUpdateChangeSet.AddedPrimitiveSceneInfos)
+		{
+			InvalidatingPrimitiveCollector.UpdatedTransform(PrimitiveSceneInfo);
+		}
+		// As must all primitive updates, 
+		for (FPrimitiveSceneInfo* PrimitiveSceneInfo : PostUpdateChangeSet.UpdatedPrimitiveSceneInfos)
+		{
+			InvalidatingPrimitiveCollector.UpdatedTransform(PrimitiveSceneInfo);
+		}
 
-	// All removed primitives must invalidate their footprints in the VSM before leaving
-	for (FPrimitiveSceneInfo* PrimitiveSceneInfo : PostUpdateChangeSet.AddedPrimitiveSceneInfos)
-	{
-		InvalidatingPrimitiveCollector.UpdatedTransform(PrimitiveSceneInfo);
+		CacheManager.ProcessInvalidations(GraphBuilder, SceneUniforms, InvalidatingPrimitiveCollector);
 	}
-	// As must all primitive updates, 
-	for (FPrimitiveSceneInfo* PrimitiveSceneInfo : PostUpdateChangeSet.UpdatedPrimitiveSceneInfos)
-	{
-		InvalidatingPrimitiveCollector.UpdatedTransform(PrimitiveSceneInfo);
-	}
-
-	CacheManager.ProcessInvalidations(GraphBuilder, SceneUniforms, InvalidatingPrimitiveCollector);
 	PostUpdateChangeSet = FScenePostUpdateChangeSet();
 }
