Better resource names for virtual texture resources.
These appear in the render resource viewer. Before it was difficult to isolate virtual texture resources for inspection, and the names gave no additional context about which VT pools they belonged to.
Now all names begin with "VirtualTexture_" prefix, and pools show their associated formats.
#tests In editor on a map with virtual textures, observe entries in render resource viewer

diff --git a/Engine/Source/Runtime/Engine/Private/VT/VirtualTextureUploadCache.cpp b/Engine/Source/Runtime/Engine/Private/VT/VirtualTextureUploadCache.cpp
--- a/Engine/Source/Runtime/Engine/Private/VT/VirtualTextureUploadCache.cpp
+++ b/Engine/Source/Runtime/Engine/Private/VT/VirtualTextureUploadCache.cpp
@@ -304,7 +304,7 @@ void FVirtualTextureUploadCache::Finalize(FRDGBuilder& GraphBuilder)
 			}
 
 			FRHITextureCreateDesc Desc =
-				FRHITextureCreateDesc::Create2D(TEXT("FVirtualTextureUploadCache_StagingTexture"), TileSize * WidthInTiles, TileSize * HeightInTiles, PoolEntry.Format);
+				FRHITextureCreateDesc::Create2D(TEXT("VirtualTexture_UploadCacheStagingTexture"), TileSize * WidthInTiles, TileSize * HeightInTiles, PoolEntry.Format);
 
 			if (bIsCpuWritable)
 			{
diff --git a/Engine/Source/Runtime/Renderer/Private/VT/VirtualTextureFeedbackBuffer.cpp b/Engine/Source/Runtime/Renderer/Private/VT/VirtualTextureFeedbackBuffer.cpp
--- a/Engine/Source/Runtime/Renderer/Private/VT/VirtualTextureFeedbackBuffer.cpp
+++ b/Engine/Source/Runtime/Renderer/Private/VT/VirtualTextureFeedbackBuffer.cpp
@@ -80,7 +80,7 @@ void FVirtualTextureFeedbackBuffer::Begin(FRDGBuilder& GraphBuilder, const FVirt
 	FRDGBufferDesc BufferDesc(FRDGBufferDesc::CreateStructuredDesc(sizeof(uint32), Desc.BufferSize.X * Desc.BufferSize.Y));
 	BufferDesc.Usage |= BUF_SourceCopy;
 
-	if (AllocatePooledBuffer(BufferDesc, PooledBuffer, TEXT("VirtualTextureFeedbackGPU")))
+	if (AllocatePooledBuffer(BufferDesc, PooledBuffer, TEXT("VirtualTexture_FeedbackBuffer")))
 	{
 		FRDGBufferUAVDesc UAVDesc{};
 		UAV = PooledBuffer->GetOrCreateUAV(GraphBuilder.RHICmdList, UAVDesc);
diff --git a/Engine/Source/Runtime/Renderer/Private/VT/VirtualTexturePhysicalSpace.cpp b/Engine/Source/Runtime/Renderer/Private/VT/VirtualTexturePhysicalSpace.cpp
--- a/Engine/Source/Runtime/Renderer/Private/VT/VirtualTexturePhysicalSpace.cpp
+++ b/Engine/Source/Runtime/Renderer/Private/VT/VirtualTexturePhysicalSpace.cpp
@@ -91,6 +91,21 @@ FVirtualTexturePhysicalSpace::FVirtualTexturePhysicalSpace(uint16 InID, const FV
 			FormatString += TEXT(", ");
 		}
 	}
+
+#if !UE_BUILD_SHIPPING
+	// Store string for resorce debug names.
+	for (uint32 Layer = 0; Layer < Description.NumLayers; ++Layer)
+	{
+		if (Description.NumLayers > 1)
+		{
+			PooledRenderTargetDebugNames[Layer] = FString::Printf(TEXT("VirtualTexture_Physical (%s) %d/%d"), *FormatString, Layer + 1, Description.NumLayers);
+		}
+		else
+		{
+			PooledRenderTargetDebugNames[Layer] = FString::Printf(TEXT("VirtualTexture_Physical (%s)"), *FormatString);
+		}
+	}
+#endif
 }
 
 FVirtualTexturePhysicalSpace::~FVirtualTexturePhysicalSpace()
@@ -152,7 +167,7 @@ void FVirtualTexturePhysicalSpace::InitRHI(FRHICommandListBase& RHICmdList)
 			Desc.UAVFormat = FormatUAV;
 		}
 
-		GRenderTargetPool.FindFreeElement(RHICmdList, Desc, PooledRenderTarget[Layer], TEXT("VirtualPhysicalTexture"));
+		GRenderTargetPool.FindFreeElement(RHICmdList, Desc, PooledRenderTarget[Layer], *PooledRenderTargetDebugNames[Layer]);
 		FRHITexture* TextureRHI = PooledRenderTarget[Layer]->GetRHI();
 
 		// Create sRGB and non-sRGB shader resource views into the physical texture
diff --git a/Engine/Source/Runtime/Renderer/Private/VT/VirtualTexturePhysicalSpace.h b/Engine/Source/Runtime/Renderer/Private/VT/VirtualTexturePhysicalSpace.h
--- a/Engine/Source/Runtime/Renderer/Private/VT/VirtualTexturePhysicalSpace.h
+++ b/Engine/Source/Runtime/Renderer/Private/VT/VirtualTexturePhysicalSpace.h
@@ -157,6 +157,7 @@ private:
 	TRefCountPtr<IPooledRenderTarget> PooledRenderTarget[VIRTUALTEXTURE_SPACE_MAXLAYERS];
 	FShaderResourceViewRHIRef TextureSRV[VIRTUALTEXTURE_SPACE_MAXLAYERS];
 	FShaderResourceViewRHIRef TextureSRV_SRGB[VIRTUALTEXTURE_SPACE_MAXLAYERS];
+	FString PooledRenderTargetDebugNames[VIRTUALTEXTURE_SPACE_MAXLAYERS];
 
 	uint16 ID;
 	uint32 NumRefs;
diff --git a/Engine/Source/Runtime/Renderer/Private/VT/VirtualTextureSpace.cpp b/Engine/Source/Runtime/Renderer/Private/VT/VirtualTextureSpace.cpp
--- a/Engine/Source/Runtime/Renderer/Private/VT/VirtualTextureSpace.cpp
+++ b/Engine/Source/Runtime/Renderer/Private/VT/VirtualTextureSpace.cpp
@@ -79,6 +79,22 @@ FVirtualTextureSpace::FVirtualTextureSpace(FVirtualTextureSystem* InSystem, uint
 		++PageTableIndex;
 	}
 
+#if !UE_BUILD_SHIPPING
+	const uint32 NumPageTableTextures = GetNumPageTableTextures();
+	check (NumPageTableTextures == PageTableIndex);
+	for (uint32 TextureIndex = 0; TextureIndex < NumPageTableTextures; ++TextureIndex)
+	{
+		if (PageTableIndex > 1u)
+		{
+			PageTableDebugNames[TextureIndex] = FString::Printf(TEXT("VirtualTexture_PageTable (%s) %d/%d"), GPixelFormats[TexturePixelFormat[TextureIndex]].Name, TextureIndex + 1, NumPageTableTextures);
+		}
+		else
+		{
+			PageTableDebugNames[TextureIndex] = FString::Printf(TEXT("VirtualTexture_PageTable (%s)"), GPixelFormats[TexturePixelFormat[TextureIndex]].Name);
+		}
+	}
+#endif
+
 	Allocator.Initialize(Description.MaxSpaceSize);
 
 	bNeedToAllocatePageTableIndirection = InDesc.IndirectionTextureSize > 0;
@@ -286,9 +302,6 @@ void FVirtualTextureSpace::AllocateTextures(FRDGBuilder& GraphBuilder)
 	{
 		RDG_GPU_MASK_SCOPE(GraphBuilder, FRHIGPUMask::All());
 
-		const TCHAR* TextureNames[] = { TEXT("VirtualPageTable_0"), TEXT("VirtualPageTable_1") };
-		static_assert(UE_ARRAY_COUNT(TextureNames) == TextureCapacity, "");
-
 		const FUintPoint RequiredPageTableSize = GetRequiredPageTableAllocationSize();
 		CachedPageTableWidth = RequiredPageTableSize.X;
 		CachedPageTableHeight = RequiredPageTableSize.Y;
@@ -305,7 +318,7 @@ void FVirtualTextureSpace::AllocateTextures(FRDGBuilder& GraphBuilder)
 				TexCreate_RenderTargetable | TexCreate_ShaderResource,
 				CachedNumPageTableLevels);
 
-			FRDGTextureRef DstTexture = GraphBuilder.CreateTexture(Desc, TextureNames[TextureIndex]);
+			FRDGTextureRef DstTexture = GraphBuilder.CreateTexture(Desc, *PageTableDebugNames[TextureIndex]);
 
 			if (TextureEntry.RenderTarget)
 			{
@@ -341,7 +354,7 @@ void FVirtualTextureSpace::AllocateTextures(FRDGBuilder& GraphBuilder)
 				FClearValueBinding::None,
 				TexCreate_UAV | TexCreate_ShaderResource);
 
-			FRDGTextureRef PageTableIndirectionTexture = GraphBuilder.CreateTexture(Desc, TEXT("PageTableIndirection"));
+			FRDGTextureRef PageTableIndirectionTexture = GraphBuilder.CreateTexture(Desc, TEXT("VirtualTexture_PageTableAdaptiveIndirection"));
 			AddClearUAVPass(GraphBuilder, GraphBuilder.CreateUAV(PageTableIndirectionTexture), FUintVector4(ForceInitToZero));
 			PageTableIndirection.RenderTarget = GraphBuilder.ConvertToExternalTexture(PageTableIndirectionTexture);
 			RHIUpdateTextureReference(PageTableIndirection.TextureReferenceRHI, PageTableIndirection.RenderTarget->GetRHI());
@@ -428,7 +441,7 @@ void FVirtualTextureSpace::ApplyUpdates(FVirtualTextureSystem* System, FRDGBuild
 			NewBufferSize = FMath::Max(NewBufferSize, UpdateBuffer->GetSize() * 2u);
 		}
 
-		FRHIResourceCreateInfo CreateInfo(TEXT("FVirtualTextureSpace_UpdateBuffer"));
+		FRHIResourceCreateInfo CreateInfo(TEXT("VirtualTexture_PageTableUpdateBuffer"));
 		UpdateBuffer = RHICmdList.CreateVertexBuffer(NewBufferSize, BUF_ShaderResource | BUF_Volatile, CreateInfo);
 		UpdateBufferSRV = RHICmdList.CreateShaderResourceView(UpdateBuffer, sizeof(FPageTableUpdate), PF_R16G16B16A16_UINT);
 	}
diff --git a/Engine/Source/Runtime/Renderer/Private/VT/VirtualTextureSpace.h b/Engine/Source/Runtime/Renderer/Private/VT/VirtualTextureSpace.h
--- a/Engine/Source/Runtime/Renderer/Private/VT/VirtualTextureSpace.h
+++ b/Engine/Source/Runtime/Renderer/Private/VT/VirtualTextureSpace.h
@@ -120,6 +120,7 @@ private:
 
 	FTextureEntry PageTable[TextureCapacity];
 	TEnumAsByte<EPixelFormat> TexturePixelFormat[TextureCapacity];
+	FString PageTableDebugNames[TextureCapacity];
 
 	FTextureEntry PageTableIndirection;
 
