Ray Tracing debug visualization modes without relying on HitGroup UserData
- Replicate data in InstancesDebugData according to NumTransforms so that InstanceDebugBuffer contains one entry per instance in TLAS instead of GeometryInstance.

diff --git a/Engine/Shaders/Private/RayTracing/RayTracingDebug.usf b/Engine/Shaders/Private/RayTracing/RayTracingDebug.usf
--- a/Engine/Shaders/Private/RayTracing/RayTracingDebug.usf
+++ b/Engine/Shaders/Private/RayTracing/RayTracingDebug.usf
@@ -220,10 +220,10 @@ float4 DebugRayTracingInstance(FRayDesc Ray, out float OutDepth)
 			Result.rgb = TriangleColor;
 			break;
 		case RAY_TRACING_DEBUG_VIZ_DYNAMIC_INSTANCES:
-			Result.rgb = InstancesDebugData[Payload.GeometryInstanceIndex].Flags ? float3(0.25f, 0.50f, 0.25f) : float3(0.50f, 0.25f, 0.25f);
+			Result.rgb = InstancesDebugData[Payload.InstanceIndex].Flags ? float3(0.25f, 0.50f, 0.25f) : float3(0.50f, 0.25f, 0.25f);
 			break;
 		case RAY_TRACING_DEBUG_VIZ_PROXY_TYPE:
-			Result.rgb = IntToColor(InstancesDebugData[Payload.GeometryInstanceIndex].ProxyHash);
+			Result.rgb = IntToColor(InstancesDebugData[Payload.InstanceIndex].ProxyHash);
 			break;
 		}
 	
diff --git a/Engine/Shaders/Private/RayTracing/RayTracingDebugCHS.usf b/Engine/Shaders/Private/RayTracing/RayTracingDebugCHS.usf
--- a/Engine/Shaders/Private/RayTracing/RayTracingDebugCHS.usf
+++ b/Engine/Shaders/Private/RayTracing/RayTracingDebugCHS.usf
@@ -22,7 +22,6 @@ RAY_TRACING_ENTRY_CLOSEST_HIT(RayTracingDebugMainCHS,
 	Payload.TriangleIndex = PrimitiveIndex();
 
 	Payload.WorldNormal = GetGeometryNormalFromTriangleBaseAttributes(PrimitiveIndex());
-	Payload.GeometryInstanceIndex = GetHitGroupUserData();
 	Payload.InstanceIndex = InstanceIndex();
 	Payload.GeometryIndex = 0; // GeometryIndex();
 	Payload.ScenePrimitiveIndex = InstanceSceneData.PrimitiveId;
diff --git a/Engine/Shaders/Private/RayTracing/RayTracingDebugPicking.usf b/Engine/Shaders/Private/RayTracing/RayTracingDebugPicking.usf
--- a/Engine/Shaders/Private/RayTracing/RayTracingDebugPicking.usf
+++ b/Engine/Shaders/Private/RayTracing/RayTracingDebugPicking.usf
@@ -41,13 +41,12 @@ RAY_TRACING_ENTRY_RAYGEN(RayTracingDebugPickingRGS)
 
 	if (Payload.IsHit())
 	{
-		FRayTracingInstanceDebugData InstanceDebugData = InstancesDebugData[Payload.GeometryInstanceIndex];
+		FRayTracingInstanceDebugData InstanceDebugData = InstancesDebugData[Payload.InstanceIndex];
 		FRayTracingInstanceDescriptor InstanceDescriptor = InstanceBuffer[Payload.InstanceIndex];
 
 		PickingFeedback.GeometryAddress = InstanceDebugData.GeometryAddress;
 
 		PickingFeedback.InstanceIndex = Payload.InstanceIndex;
-		PickingFeedback.GeometryInstanceIndex = Payload.GeometryInstanceIndex;
 		PickingFeedback.TriangleIndex = Payload.TriangleIndex;
 		PickingFeedback.GeometryIndex = Payload.GeometryIndex;
 
diff --git a/Engine/Shaders/Private/RayTracing/RayTracingDebugUtils.ush b/Engine/Shaders/Private/RayTracing/RayTracingDebugUtils.ush
--- a/Engine/Shaders/Private/RayTracing/RayTracingDebugUtils.ush
+++ b/Engine/Shaders/Private/RayTracing/RayTracingDebugUtils.ush
@@ -11,10 +11,9 @@ struct FRayTracingDebugPayload : FMinimalPayload
 	uint InstanceHash;
 	uint TriangleIndex;
 	uint InstanceIndex;
-	uint GeometryInstanceIndex;
+	uint GeometryIndex;
 
 	float3 WorldNormal;
-	uint GeometryIndex;
 	
 	uint TriangleHitCountPerRay;
 	uint ScenePrimitiveIndex;
diff --git a/Engine/Shaders/Shared/RayTracingDebugTypes.h b/Engine/Shaders/Shared/RayTracingDebugTypes.h
--- a/Engine/Shaders/Shared/RayTracingDebugTypes.h
+++ b/Engine/Shaders/Shared/RayTracingDebugTypes.h
@@ -33,7 +33,6 @@ struct FRayTracingInstanceDebugData
 struct FRayTracingPickingFeedback
 {
 	// Hit data
-	UINT32_TYPE GeometryInstanceIndex;
 	UINT32_TYPE InstanceIndex;
 	UINT32_TYPE GeometryIndex;
 	UINT32_TYPE TriangleIndex;
@@ -49,8 +48,7 @@ struct FRayTracingPickingFeedback
 
 #ifdef __cplusplus
 	FRayTracingPickingFeedback()
-		: GeometryInstanceIndex(0xFFFFFFFF)
-		, InstanceIndex(0xFFFFFFFF)
+		: InstanceIndex(0xFFFFFFFF)
 		, GeometryAddress(0xFFFFFFFFFFFFFFFF)
 	{}
 #endif
diff --git a/Engine/Source/Runtime/Renderer/Private/RayTracing/RayTracingDebug.cpp b/Engine/Source/Runtime/Renderer/Private/RayTracing/RayTracingDebug.cpp
--- a/Engine/Source/Runtime/Renderer/Private/RayTracing/RayTracingDebug.cpp
+++ b/Engine/Source/Runtime/Renderer/Private/RayTracing/RayTracingDebug.cpp
@@ -729,7 +729,6 @@ void BindRayTracingDebugCHSMaterialBindings(FRHICommandList& RHICmdList, FRHISha
 		Binding.SegmentIndex = MeshCommand.GeometrySegmentIndex;
 		Binding.UniformBuffers = HelperBinding.UniformBufferArray;
 		Binding.NumUniformBuffers = HelperBinding.NumUniformBuffers;
-		Binding.UserData = VisibleMeshCommand.InstanceIndex;
 
 		Bindings[BindingIndex] = Binding;
 		BindingIndex++;
@@ -1640,7 +1639,6 @@ void FDeferredShadingSceneRenderer::RayTracingDisplayPicking(const FRayTracingPi
 	Writer.EmptyLine();
 
 	Writer.DrawLine(FText::FromString(FString::Printf(TEXT("Instance Index: %u"), PickingFeedback.InstanceIndex)), 10, FColor::Yellow);
-	Writer.DrawLine(FText::FromString(FString::Printf(TEXT("Geometry Instance Index: %u"), PickingFeedback.GeometryInstanceIndex)), 10, FColor::Yellow);
 	Writer.DrawLine(FText::FromString(FString::Printf(TEXT("Segment Index: %u"), PickingFeedback.GeometryIndex)), 10, FColor::Yellow);
 	Writer.DrawLine(FText::FromString(FString::Printf(TEXT("Triangle Index: %u"), PickingFeedback.TriangleIndex)), 10, FColor::Yellow);
 
diff --git a/Engine/Source/Runtime/Renderer/Private/RayTracing/RayTracingScene.cpp b/Engine/Source/Runtime/Renderer/Private/RayTracing/RayTracingScene.cpp
--- a/Engine/Source/Runtime/Renderer/Private/RayTracing/RayTracingScene.cpp
+++ b/Engine/Source/Runtime/Renderer/Private/RayTracing/RayTracingScene.cpp
@@ -192,6 +192,68 @@ void FRayTracingScene::CreateWithInitializationData(FRDGBuilder& GraphBuilder, c
 		AddClearUAVPass(GraphBuilder, DebugInstanceGPUSceneIndexBufferUAV, 0xFFFFFFFF);
 	}
 
+	if (InstancesDebugData.Num() > 0 && NumNativeInstances > 0)
+	{
+		// Create InstanceDebugBuffer (one entry per instance in TLAS)
+		// This requires replicating the data in InstancesDebugData (one entry per FRayTracingGeometryInstance) according to NumTransforms in each geometry instance
+
+		check(InstancesDebugData.Num() == Instances.Num());
+
+		TArrayView<uint32> LayerBaseIndices = MakeArrayView(GraphBuilder.AllocPODArray<uint32>(NumLayers), NumLayers);
+		LayerBaseIndices[0] = 0;
+
+		for (uint32 LayerIndex = 1; LayerIndex < NumLayers; ++LayerIndex)
+		{
+			LayerBaseIndices[LayerIndex] = LayerBaseIndices[LayerIndex - 1] + SceneInitializer.NumNativeInstancesPerLayer[LayerIndex - 1];
+		}
+
+		// make a copy of SceneWithGeometryInstances.BaseInstancePrefixSum that can be passed to the RDG setup tasks below
+		TArray<uint32, SceneRenderingAllocator> BaseInstancePrefixSum = GraphBuilder.AllocArray<uint32>();
+		BaseInstancePrefixSum = SceneWithGeometryInstances.BaseInstancePrefixSum;
+
+		FRDGUploadData<FRayTracingInstanceDebugData> UploadData(GraphBuilder, NumNativeInstances);
+
+		{
+			const uint32 NumItems = InstancesDebugData.Num();
+
+			// Distribute work evenly to the available task graph workers based on NumItems.
+			const uint32 TargetItemsPerTask = 512;
+			const uint32 NumThreads = FMath::Min(FTaskGraphInterface::Get().GetNumWorkerThreads(), CVarRHICmdWidth.GetValueOnRenderThread());
+			const uint32 NumTasks = FMath::Min(NumThreads, FMath::DivideAndRoundUp(NumItems, TargetItemsPerTask));
+			const uint32 NumItemsPerTask = FMath::DivideAndRoundUp(NumItems, NumTasks);
+
+			for (uint32 TaskIndex = 0; TaskIndex < NumTasks; ++TaskIndex)
+			{
+				const uint32 TaskFirstItemIndex = TaskIndex * NumItemsPerTask;
+				const uint32 TaskNumItems = FMath::Min(NumItemsPerTask, NumItems - TaskFirstItemIndex);
+
+				TConstArrayView<FRayTracingGeometryInstance> TaskInstancesData(Instances.GetData() + TaskFirstItemIndex, TaskNumItems);
+				TConstArrayView<FRayTracingInstanceDebugData> TaskInstancesDebugData(InstancesDebugData.GetData() + TaskFirstItemIndex, TaskNumItems);
+				TConstArrayView<uint32> TaskBaseInstancePrefixSum(BaseInstancePrefixSum.GetData() + TaskFirstItemIndex, TaskNumItems);
+
+				GraphBuilder.AddSetupTask([UploadData, TaskInstancesDebugData, TaskInstancesData, TaskBaseInstancePrefixSum, LayerBaseIndices]()
+					{
+						TRACE_CPUPROFILER_EVENT_SCOPE(FillRayTracingInstanceDebugBuffer);
+
+						for (int32 Index = 0; Index < TaskInstancesDebugData.Num(); ++Index)
+						{
+							const FRayTracingGeometryInstance& SceneInstance = TaskInstancesData[Index];
+							const uint32 BaseInstanceIndex = TaskBaseInstancePrefixSum[Index];
+							const uint32 LayerBaseIndex = LayerBaseIndices[SceneInstance.LayerIndex];
+
+							for (uint32 TransformIndex = 0; TransformIndex < SceneInstance.NumTransforms; ++TransformIndex)
+							{
+								// write data in the same order used in InstanceBuffer used to build TLAS / InstanceIndex() in hit shaders
+								UploadData[LayerBaseIndex + BaseInstanceIndex + TransformIndex] = TaskInstancesDebugData[Index];
+							}
+						}
+					});
+			}
+		}
+
+		InstanceDebugBuffer = CreateStructuredBuffer(GraphBuilder, TEXT("FRayTracingScene::InstanceDebugData"), UploadData);
+	}
+
 	if (NumNativeInstances > 0)
 	{
 		const uint32 InstanceUploadBytes = NumNativeInstances * sizeof(FRayTracingInstanceDescriptorInput);
@@ -227,12 +289,6 @@ void FRayTracingScene::CreateWithInitializationData(FRDGBuilder& GraphBuilder, c
 				InstanceUploadData,
 				TransformUploadData);
 		});
-		
-		if (InstancesDebugData.Num() > 0)
-		{
-			check(InstancesDebugData.Num() == Instances.Num());
-			InstanceDebugBuffer = CreateStructuredBuffer(GraphBuilder, TEXT("FRayTracingScene::InstanceDebugData"), InstancesDebugData);
-		}
 
 		FBuildInstanceBufferPassParams* PassParams = GraphBuilder.AllocParameters<FBuildInstanceBufferPassParams>();
 		PassParams->InstanceBuffer = GraphBuilder.CreateUAV(InstanceBuffer);
