Fixed VSM performance regression caused by static primitives to start out as dynamic, leading to loss of HZB and needless invalidations.

diff --git a/Engine/Source/Runtime/Renderer/Private/VirtualShadowMaps/VirtualShadowMapCacheManager.cpp b/Engine/Source/Runtime/Renderer/Private/VirtualShadowMaps/VirtualShadowMapCacheManager.cpp
--- a/Engine/Source/Runtime/Renderer/Private/VirtualShadowMaps/VirtualShadowMapCacheManager.cpp
+++ b/Engine/Source/Runtime/Renderer/Private/VirtualShadowMaps/VirtualShadowMapCacheManager.cpp
@@ -381,7 +381,7 @@ void FVirtualShadowMapArrayCacheManager::FInvalidatingPrimitiveCollector::AddPri
 	Manager.ShadowInvalidatingInstancesImplementation.PrimitiveInstancesToInvalidate.Reset();
 }
 
-void FVirtualShadowMapArrayCacheManager::FInvalidatingPrimitiveCollector::AddInvalidation(FPrimitiveSceneInfo * PrimitiveSceneInfo, bool bRemovedPrimitive)
+void FVirtualShadowMapArrayCacheManager::FInvalidatingPrimitiveCollector::AddInvalidation(FPrimitiveSceneInfo * PrimitiveSceneInfo, EInvalidationCause InvalidationCause)
 {
 	const int32 PrimitiveID = PrimitiveSceneInfo->GetIndex();
 	const int32 InstanceSceneDataOffset = PrimitiveSceneInfo->GetInstanceSceneDataOffset();
@@ -392,14 +392,15 @@ void FVirtualShadowMapArrayCacheManager::FInvalidatingPrimitiveCollector::AddInv
 		return;
 	}
 
-	if (!Scene->PrimitiveFlagsCompact[PrimitiveID].bCastDynamicShadow)
+	const FPrimitiveFlagsCompact PrimitiveFlagsCompact = Scene->PrimitiveFlagsCompact[PrimitiveID];
+	if (!PrimitiveFlagsCompact.bCastDynamicShadow)
 	{
 		return;
 	}
 
 	const FPersistentPrimitiveIndex PersistentPrimitiveIndex = PrimitiveSceneInfo->GetPersistentIndex();
 
-	if (bRemovedPrimitive)
+	if (InvalidationCause == EInvalidationCause::Removed)
 	{
 		RemovedPrimitives[PersistentPrimitiveIndex.Index] = true;
 	}
@@ -407,7 +408,7 @@ void FVirtualShadowMapArrayCacheManager::FInvalidatingPrimitiveCollector::AddInv
 	// TODO: Early out on stuff that doesn't cast shadows on the CPU?
 
 	// Suppress invalidations from moved primitives that are marked to behave as if they were static.
-	if (!bRemovedPrimitive && PrimitiveSceneInfo->Proxy->GetShadowCacheInvalidationBehavior() == EShadowCacheInvalidationBehavior::Static)
+	if (InvalidationCause == EInvalidationCause::Updated && PrimitiveSceneInfo->Proxy->GetShadowCacheInvalidationBehavior() == EShadowCacheInvalidationBehavior::Static)
 	{
 		return;
 	}
@@ -419,16 +420,28 @@ void FVirtualShadowMapArrayCacheManager::FInvalidatingPrimitiveCollector::AddInv
 	// TODO: Filter out one of the updates for things like WPO animation or other cases where the transform/bounds have not changed
 	// TODO: Should we be using AddedScenePrimitives now instead of this flag?
 	const EPrimitiveDirtyState DirtyState = Scene->GPUScene.GetPrimitiveDirtyState(PersistentPrimitiveIndex);
-	const bool bAddedPrimitive = EnumHasAnyFlags(DirtyState, EPrimitiveDirtyState::Added);
-	if (bAddedPrimitive && (GVSMNewInvalidations == 0))
+	const bool bMarkedAsAddedPrimitive = EnumHasAnyFlags(DirtyState, EPrimitiveDirtyState::Added);
+	if (bMarkedAsAddedPrimitive && (GVSMNewInvalidations == 0))
 	{
 		return;
 	}
 
-	InvalidatedPrimitives[PersistentPrimitiveIndex.Index] = true;
-
+	// Prevent transition to "dynamic" state if the primitive is static or forced to behave as such.
+	if (InvalidationCause == EInvalidationCause::Added)
+	{
+		// Skip marking as dynamic if it is a static mesh (mobility is static & no WPO) or it is forced to behave as static
+		// this avoids needing to re-cache all static meshes.
+		if (PrimitiveSceneInfo->Proxy->IsMeshShapeOftenMoving() && PrimitiveSceneInfo->Proxy->GetShadowCacheInvalidationBehavior() != EShadowCacheInvalidationBehavior::Static)
+		{
+			InvalidatedPrimitives[PersistentPrimitiveIndex.Index] = true;
+		}
+	}
+	else
+	{
+		InvalidatedPrimitives[PersistentPrimitiveIndex.Index] = true;
+	}
 	// Nanite meshes need special handling because they don't get culled on CPU, thus always process invalidations for those
-	const bool bIsNaniteMesh = Scene->PrimitiveFlagsCompact[PrimitiveID].bIsNaniteMesh;
+	const bool bIsNaniteMesh = PrimitiveFlagsCompact.bIsNaniteMesh;
 
 	for (auto& CacheEntry : Manager.CacheEntries)
 	{
@@ -1587,38 +1600,38 @@ void FVirtualShadowMapInvalidationSceneUpdater::PreSceneUpdate(FRDGBuilder& Grap
 
 	if (CacheManager.IsCacheDataAvailable())
 	{
-		FVirtualShadowMapArrayCacheManager::FInvalidatingPrimitiveCollector InvalidatingPrimitiveCollector(&CacheManager);
+	FVirtualShadowMapArrayCacheManager::FInvalidatingPrimitiveCollector InvalidatingPrimitiveCollector(&CacheManager);
 
-		// Primitives that are tracked as always invalidating shadows, pipe through as transform updates
-		for (FPrimitiveSceneInfo* PrimitiveSceneInfo : CacheManager.Scene->ShadowScene->GetAlwaysInvalidatingPrimitives())
-		{
-			InvalidatingPrimitiveCollector.UpdatedTransform(PrimitiveSceneInfo);
-		}
+	// Primitives that are tracked as always invalidating shadows, pipe through as transform updates
+	for (FPrimitiveSceneInfo* PrimitiveSceneInfo : CacheManager.Scene->ShadowScene->GetAlwaysInvalidatingPrimitives())
+	{
+		InvalidatingPrimitiveCollector.UpdatedTransform(PrimitiveSceneInfo);
+	}
 
-		// All removed primitives must invalidate their footprints in the VSM before leaving
-		for (FPrimitiveSceneInfo* PrimitiveSceneInfo : ChangeSet.RemovedPrimitiveSceneInfos)
-		{
-			InvalidatingPrimitiveCollector.Removed(PrimitiveSceneInfo);
-		}
-		// As must all primitive updates, 
-		for (FPrimitiveSceneInfo* PrimitiveSceneInfo : ChangeSet.UpdatedPrimitiveSceneInfos)
-		{
-			InvalidatingPrimitiveCollector.UpdatedTransform(PrimitiveSceneInfo);
-		}
+	// All removed primitives must invalidate their footprints in the VSM before leaving
+	for (FPrimitiveSceneInfo* PrimitiveSceneInfo : ChangeSet.RemovedPrimitiveSceneInfos)
+	{
+		InvalidatingPrimitiveCollector.Removed(PrimitiveSceneInfo);
+	}
+	// As must all primitive updates, 
+	for (FPrimitiveSceneInfo* PrimitiveSceneInfo : ChangeSet.UpdatedPrimitiveSceneInfos)
+	{
+		InvalidatingPrimitiveCollector.UpdatedTransform(PrimitiveSceneInfo);
+	}
 
-		// TODO! Where do we get this data from...
-		/*
-		for (const auto& CullDistance : UpdatedInstanceCullDistance)
-		{
-			InvalidatingPrimitiveCollector.UpdatedTransform(CullDistance.Key->GetPrimitiveSceneInfo());
-		}
-		*/
+	// TODO! Where do we get this data from...
+	/*
+	for (const auto& CullDistance : UpdatedInstanceCullDistance)
+	{
+		InvalidatingPrimitiveCollector.UpdatedTransform(CullDistance.Key->GetPrimitiveSceneInfo());
+	}
+	*/
 
-		// TODO: Perhaps pass this in from the caller in RendererScene?
-		FSceneUniformBuffer SceneUniforms;
-		CacheManager.Scene->GPUScene.FillSceneUniformBuffer(GraphBuilder, SceneUniforms);
+	// TODO: Perhaps pass this in from the caller in RendererScene?
+	FSceneUniformBuffer SceneUniforms;
+	CacheManager.Scene->GPUScene.FillSceneUniformBuffer(GraphBuilder, SceneUniforms);
 
-		CacheManager.ProcessInvalidations(GraphBuilder, SceneUniforms, InvalidatingPrimitiveCollector);
+	CacheManager.ProcessInvalidations(GraphBuilder, SceneUniforms, InvalidatingPrimitiveCollector);
 	}
 }
 
@@ -1636,20 +1649,18 @@ void FVirtualShadowMapInvalidationSceneUpdater::PostGPUSceneUpdate(FRDGBuilder&
 	SCOPED_NAMED_EVENT(FScene_VirtualShadowCacheUpdate, FColor::Orange);
 	if (CacheManager.IsCacheDataAvailable())
 	{
-		FVirtualShadowMapArrayCacheManager::FInvalidatingPrimitiveCollector InvalidatingPrimitiveCollector(&CacheManager);
+	FVirtualShadowMapArrayCacheManager::FInvalidatingPrimitiveCollector InvalidatingPrimitiveCollector(&CacheManager);
 
-		// All removed primitives must invalidate their footprints in the VSM before leaving
-		for (FPrimitiveSceneInfo* PrimitiveSceneInfo : PostUpdateChangeSet.AddedPrimitiveSceneInfos)
-		{
-			InvalidatingPrimitiveCollector.UpdatedTransform(PrimitiveSceneInfo);
-		}
-		// As must all primitive updates, 
-		for (FPrimitiveSceneInfo* PrimitiveSceneInfo : PostUpdateChangeSet.UpdatedPrimitiveSceneInfos)
-		{
-			InvalidatingPrimitiveCollector.UpdatedTransform(PrimitiveSceneInfo);
-		}
+	for (FPrimitiveSceneInfo* PrimitiveSceneInfo : PostUpdateChangeSet.AddedPrimitiveSceneInfos)
+	{
+		InvalidatingPrimitiveCollector.Added(PrimitiveSceneInfo);
+	}
+	for (FPrimitiveSceneInfo* PrimitiveSceneInfo : PostUpdateChangeSet.UpdatedPrimitiveSceneInfos)
+	{
+		InvalidatingPrimitiveCollector.UpdatedTransform(PrimitiveSceneInfo);
+	}
 
-		CacheManager.ProcessInvalidations(GraphBuilder, SceneUniforms, InvalidatingPrimitiveCollector);
+	CacheManager.ProcessInvalidations(GraphBuilder, SceneUniforms, InvalidatingPrimitiveCollector);
 	}
 	PostUpdateChangeSet = FScenePostUpdateChangeSet();
 }
diff --git a/Engine/Source/Runtime/Renderer/Private/VirtualShadowMaps/VirtualShadowMapCacheManager.h b/Engine/Source/Runtime/Renderer/Private/VirtualShadowMaps/VirtualShadowMapCacheManager.h
--- a/Engine/Source/Runtime/Renderer/Private/VirtualShadowMaps/VirtualShadowMapCacheManager.h
+++ b/Engine/Source/Runtime/Renderer/Private/VirtualShadowMaps/VirtualShadowMapCacheManager.h
@@ -287,19 +287,24 @@ public:
 		// Primitive was removed from the scene
 		void Removed(FPrimitiveSceneInfo* PrimitiveSceneInfo)
 		{
-			AddInvalidation(PrimitiveSceneInfo, true);
+			AddInvalidation(PrimitiveSceneInfo, EInvalidationCause::Removed);
 		}
 
 		// Primitive instances updated
 		void UpdatedInstances(FPrimitiveSceneInfo* PrimitiveSceneInfo)
 		{
-			AddInvalidation(PrimitiveSceneInfo, false);
+			AddInvalidation(PrimitiveSceneInfo, EInvalidationCause::Updated);
 		}
 
 		// Primitive moved/transform was updated
 		void UpdatedTransform(FPrimitiveSceneInfo* PrimitiveSceneInfo)
 		{
-			AddInvalidation(PrimitiveSceneInfo, false);
+			AddInvalidation(PrimitiveSceneInfo, EInvalidationCause::Updated);
+		}
+
+		void Added(FPrimitiveSceneInfo* PrimitiveSceneInfo)
+		{
+			AddInvalidation(PrimitiveSceneInfo, EInvalidationCause::Added);
 		}
 
 		FInstanceGPULoadBalancer Instances;
@@ -307,7 +312,14 @@ public:
 		TBitArray<> RemovedPrimitives;
 
 	private:
-		void AddInvalidation(FPrimitiveSceneInfo* PrimitiveSceneInfo, bool bRemovedPrimitive);
+		enum class EInvalidationCause
+		{
+			Added,
+			Removed,
+			Updated,
+		};
+
+		void AddInvalidation(FPrimitiveSceneInfo* PrimitiveSceneInfo, EInvalidationCause InvalidationCause);
 
 		FScene* Scene = nullptr;
 		FVirtualShadowMapArrayCacheManager& Manager;
