Fixing access violation when writing primitive debugger contents out to a CSV file.

diff --git a/Engine/Source/Runtime/Renderer/Private/ViewDebug.cpp b/Engine/Source/Runtime/Renderer/Private/ViewDebug.cpp
--- a/Engine/Source/Runtime/Renderer/Private/ViewDebug.cpp
+++ b/Engine/Source/Runtime/Renderer/Private/ViewDebug.cpp
@@ -175,7 +175,7 @@ void FViewDebugInfo::DumpToCSV() const
 			DrawViewer.AddColumn(*Primitive.Name);
 			DrawViewer.AddColumn(Primitive.Owner.IsValid() ? *Primitive.Owner->GetClass()->GetName() : TEXT(""));
 			DrawViewer.AddColumn(Primitive.Owner.IsValid() ? *Primitive.Owner->GetFullName() : TEXT(""));
-			DrawViewer.AddColumn(Primitive.ComponentInterface ?
+			DrawViewer.AddColumn(Primitive.IsPrimitiveValid() ?
 				*FString::Printf(TEXT("{%s}"), *Primitive.GetPrimitiveLocation().ToString()) : TEXT(""));
 			DrawViewer.AddColumn(*FString::Printf(TEXT("%d"), Primitive.Materials.Num()));
 			FString Materials = "[";
