Fixes a crash in TSR when using -vulkan on a intel Arc GPU

diff --git a/Engine/Source/Runtime/Renderer/Private/PostProcess/TemporalSuperResolution.cpp b/Engine/Source/Runtime/Renderer/Private/PostProcess/TemporalSuperResolution.cpp
--- a/Engine/Source/Runtime/Renderer/Private/PostProcess/TemporalSuperResolution.cpp
+++ b/Engine/Source/Runtime/Renderer/Private/PostProcess/TemporalSuperResolution.cpp
@@ -518,14 +518,6 @@ public:
 
 	static FPermutationDomain RemapPermutation(FPermutationDomain PermutationVector)
 	{
-		int32 WaveSize = PermutationVector.Get<FWaveSizeOps>();
-
-		// WaveSize=16 is for Intel Arc GPU which also supports 16bits ops, so compiling WaveSize=16 32bit ops is useless and should instead fall back to WaveSize=0.
-		if (WaveSize == 16 && !PermutationVector.Get<FTSRShader::F16BitVALUDim>())
-		{
-			PermutationVector.Set<FWaveSizeOps>(0);
-		}
-
 		// Only compile the alpha channel with 32bit ops, as this is mostly targeting enterprise uses on Quadro GPUs
 		if (PermutationVector.Get<FTSRShader::FAlphaChannelDim>())
 		{
@@ -533,7 +525,7 @@ public:
 		}
 
 		// Optimising register pressure with 16bit for waveops that is 1 pixel/lane is pointless.
-		if (WaveSize == 0)
+		if (PermutationVector.Get<FWaveSizeOps>() == 0)
 		{
 			PermutationVector.Set<FTSRShader::F16BitVALUDim>(false);
 		}
