Fix: Translucent After Motion Blur pass has incorrect View Uniform Buffer Params with Dynamic Resolution Translucency

We were always taking the full-scale code path, ignoring ViewportScale, for the After Motion Blur pass. Fix this up to the intended (?) logic, taking the downsample path if applicable.

diff --git a/Engine/Source/Runtime/Renderer/Private/TranslucentRendering.cpp b/Engine/Source/Runtime/Renderer/Private/TranslucentRendering.cpp
--- a/Engine/Source/Runtime/Renderer/Private/TranslucentRendering.cpp
+++ b/Engine/Source/Runtime/Renderer/Private/TranslucentRendering.cpp
@@ -1032,7 +1032,7 @@ static FViewShaderParameters GetSeparateTranslucencyViewParameters(const FViewIn
 		// We can use the existing view uniform buffers if no downsampling is required and is not in the post-motionblur pass
 		ViewParameters = View.GetShaderParameters();
 	}	
-	else if (bIsPostMotionBlur)
+	else if (ViewportScale == 1.0f && bIsPostMotionBlur)
 	{
 		// Full-scale post-motionblur pass
 		FViewUniformShaderParameters ViewUniformParameters;
