Fix shader compilation error in iOS decal shader

diff --git a/Engine/Source/Runtime/Renderer/Private/DecalRenderingCommon.cpp b/Engine/Source/Runtime/Renderer/Private/DecalRenderingCommon.cpp
--- a/Engine/Source/Runtime/Renderer/Private/DecalRenderingCommon.cpp
+++ b/Engine/Source/Runtime/Renderer/Private/DecalRenderingCommon.cpp
@@ -27,7 +27,6 @@ namespace DecalRendering
 		const bool bIsSubstrateEnabled = Substrate::IsSubstrateEnabled();
 		const bool bIsMobilePlatform = IsMobilePlatform(Platform);
 		const bool bIsMobileDeferredPlatform = bIsMobilePlatform && IsMobileDeferredShadingEnabled(Platform);
-		const bool bIsMobileLitDecalPlatform = bIsMobilePlatform && FReadOnlyCVARCache::MobileForwardDecalLighting();
 		const bool bIsDBufferPlatform = IsUsingDBuffers(Platform);
 		const bool bIsDBufferMaskPlatform = bIsDBufferPlatform && FDataDrivenShaderPlatformInfo::GetSupportsPerPixelDBufferMask(Platform);
 
@@ -44,7 +43,7 @@ namespace DecalRendering
 		}
 
 		// Enforce platform output limitations.
-		if (bIsMobilePlatform && !bIsMobileDeferredPlatform && !bIsDBufferPlatform && !bIsMobileLitDecalPlatform)
+		if (bIsMobilePlatform && !bIsMobileDeferredPlatform && !bIsDBufferPlatform)
 		{
 			Desc.bWriteNormal = false;
 			Desc.bWriteRoughnessSpecularMetallic = false;
