OpenImageDenoise: Add alpha denoising and change albedo/normal buffers to 16 bits

This take advantage of new features introduced in OIDN 2.1+.

This required splitting the depth from the normal buffer so it can still be stored with 32 bit precision. This still uses less memory overall (160 bytes/pixel vs 256 bytes/pixel previously).

Optix/NNE/NFOR denoisers will need to be updated independently.

diff --git a/Engine/Plugins/Experimental/OpenImageDenoise/Source/OpenImageDenoise/Private/Plugin.cpp b/Engine/Plugins/Experimental/OpenImageDenoise/Source/OpenImageDenoise/Private/Plugin.cpp
--- a/Engine/Plugins/Experimental/OpenImageDenoise/Source/OpenImageDenoise/Private/Plugin.cpp
+++ b/Engine/Plugins/Experimental/OpenImageDenoise/Source/OpenImageDenoise/Private/Plugin.cpp
@@ -36,31 +36,67 @@ DEFINE_LOG_CATEGORY(LogOpenImageDenoise);
 
 IMPLEMENT_MODULE(FOpenImageDenoiseModule, OpenImageDenoise)
 
-enum class EDenoiseMode
+static TAutoConsoleVariable<bool> CVarOIDNDenoiseAlpha(
+	TEXT("r.OIDN.DenoiseAlpha"),
+	true,
+	TEXT("Should OpenImageDenoise denoise the alpha channel? (default: true)")
+);
+
+static TAutoConsoleVariable<bool> CVarOIDNDenoiseAux(
+	TEXT("r.OIDN.DenoiseAuxilaryInputs"),
+	false,
+	TEXT("Should OpenImageDenoise denoise the auxilary buffers (albedo and normal) prior to using them? (default: false)")
+);
+
+static TAutoConsoleVariable<bool> CVarOIDNUseAux(
+	TEXT("r.OIDN.UseAuxilaryInputs"),
+	true,
+	TEXT("Should OpenImageDenoise make use of auxilary buffers (albedo and normal) to improve image quality? (default: true)")
+);
+
+struct FHalfColor
 {
-	OFF = 0,
-	DEFAULT = 1,   // denoise using albedo/normal as provided
-	CLEAN_AUX = 2, // denoise albedo/normal first and pass cleanAux (haven't found a scene where this is clearly better, and it runs much slower)
-	NO_AOVS = 3,   // denoise beauty only
+	uint16_t r, g, b, a;
 };
 
+struct FDenoiseSettings
+{
+	bool DenoiseAlpha;
+	bool DenoiseAux;
+	bool UseAux;
+
+	bool operator==(const FDenoiseSettings& Other) const = default;
+	bool operator!=(const FDenoiseSettings& Other) const = default;
+};
+
+FDenoiseSettings GetCurrentSettings()
+{
+	FDenoiseSettings Settings = {};
+	Settings.DenoiseAlpha = CVarOIDNDenoiseAlpha.GetValueOnRenderThread();
+	Settings.UseAux = CVarOIDNUseAux.GetValueOnRenderThread();
+	Settings.DenoiseAux = CVarOIDNDenoiseAux.GetValueOnRenderThread();
+
+	return Settings;
+}
+
 struct OIDNState
 {
 	// scratch CPU memory for running the OIDN filter
 	TArray<FLinearColor> RawPixels;
-	TArray<FLinearColor> RawAlbedo;
-	TArray<FLinearColor> RawNormal;
+	TArray<FHalfColor> RawAlbedo;
+	TArray<FHalfColor> RawNormal;
 
 	// re-useable filters
 	oidn::FilterRef AlbedoFilter;
 	oidn::FilterRef NormalFilter;
 	oidn::FilterRef PixelsFilter;
+	oidn::FilterRef AlphaFilter;
 	oidn::DeviceRef OIDNDevice;
 
-	EDenoiseMode CurrentMode = EDenoiseMode::OFF;
+	FDenoiseSettings CurrentSettings = {};
 	FIntPoint CurrentSize = FIntPoint(0, 0);
 
-	void UpdateFilter(FIntPoint Size, EDenoiseMode DenoiserMode)
+	void UpdateFilter(FIntPoint Size, FDenoiseSettings Settings)
 	{
 		int NewSize = Size.X * Size.Y;
 		if (RawPixels.Num() != NewSize)
@@ -72,30 +108,32 @@ struct OIDNState
 			AlbedoFilter = oidn::FilterRef();
 			NormalFilter = oidn::FilterRef();
 			PixelsFilter = oidn::FilterRef();
-		}
-		if (DenoiserMode == EDenoiseMode::OFF)
-		{
-			OIDNDevice = oidn::DeviceRef();
-			return;
+			AlphaFilter = oidn::FilterRef();
 		}
 		if (!OIDNDevice)
 		{
 			OIDNDevice = oidn::newDevice();
 			OIDNDevice.commit();
 		}
-		if (!PixelsFilter || CurrentMode != DenoiserMode || CurrentSize != Size)
+		if (!PixelsFilter || CurrentSettings != Settings || CurrentSize != Size)
 		{
-			CurrentMode = DenoiserMode;
+#if WITH_EDITOR
+			UE_LOG(LogOpenImageDenoise, Log, TEXT("Updating filter configuration for (%d x %d) with denoise_alpha=%s, denoise_aux=%s, use_aux=%s)"), Size.X, Size.Y,
+				   Settings.DenoiseAlpha ? TEXT("on") : TEXT("off"),
+				   Settings.DenoiseAux   ? TEXT("on") : TEXT("off"),
+				   Settings.UseAux       ? TEXT("on") : TEXT("off"));
+#endif
+			CurrentSettings = Settings;
 			CurrentSize = Size;
-			if (CurrentMode == EDenoiseMode::CLEAN_AUX)
+			if (CurrentSettings.DenoiseAux)
 			{
 				AlbedoFilter = OIDNDevice.newFilter("RT");
-				AlbedoFilter.setImage("albedo", RawAlbedo.GetData(), oidn::Format::Float3, Size.X, Size.Y, 0, sizeof(FLinearColor), sizeof(FLinearColor) * Size.X);
-				AlbedoFilter.setImage("output", RawAlbedo.GetData(), oidn::Format::Float3, Size.X, Size.Y, 0, sizeof(FLinearColor), sizeof(FLinearColor) * Size.X);
+				AlbedoFilter.setImage("albedo", RawAlbedo.GetData(), oidn::Format::Half3, Size.X, Size.Y, 0, sizeof(FHalfColor), sizeof(FHalfColor) * Size.X);
+				AlbedoFilter.setImage("output", RawAlbedo.GetData(), oidn::Format::Half3, Size.X, Size.Y, 0, sizeof(FHalfColor), sizeof(FHalfColor) * Size.X);
 				AlbedoFilter.commit();
 				NormalFilter = OIDNDevice.newFilter("RT");
-				NormalFilter.setImage("normal", RawNormal.GetData(), oidn::Format::Float3, Size.X, Size.Y, 0, sizeof(FLinearColor), sizeof(FLinearColor) * Size.X);
-				NormalFilter.setImage("output", RawNormal.GetData(), oidn::Format::Float3, Size.X, Size.Y, 0, sizeof(FLinearColor), sizeof(FLinearColor) * Size.X);
+				NormalFilter.setImage("normal", RawNormal.GetData(), oidn::Format::Half3, Size.X, Size.Y, 0, sizeof(FHalfColor), sizeof(FHalfColor) * Size.X);
+				NormalFilter.setImage("output", RawNormal.GetData(), oidn::Format::Half3, Size.X, Size.Y, 0, sizeof(FHalfColor), sizeof(FHalfColor) * Size.X);
 				NormalFilter.commit();
 			}
 			else
@@ -104,29 +142,41 @@ struct OIDNState
 				NormalFilter = oidn::FilterRef();
 			}
 			PixelsFilter = OIDNDevice.newFilter("RT");
-			// TODO: find a way to denoise the alpha channel? OIDN does not support this yet
 			PixelsFilter.setImage("color" , RawPixels.GetData(), oidn::Format::Float3, Size.X, Size.Y, 0, sizeof(FLinearColor), sizeof(FLinearColor) * Size.X);
 			PixelsFilter.setImage("output", RawPixels.GetData(), oidn::Format::Float3, Size.X, Size.Y, 0, sizeof(FLinearColor), sizeof(FLinearColor) * Size.X);
-			if (CurrentMode == EDenoiseMode::DEFAULT || CurrentMode == EDenoiseMode::CLEAN_AUX)
+
+			if (CurrentSettings.UseAux)
 			{
 				// default behavior, use the albedo/normal buffers to improve quality
-				// TODO: switch these buffers to half precision? (requires OIDN 1.4.2+)
-				PixelsFilter.setImage("albedo", RawAlbedo.GetData(), oidn::Format::Float3, Size.X, Size.Y, 0, sizeof(FLinearColor), sizeof(FLinearColor) * Size.X);
-				PixelsFilter.setImage("normal", RawNormal.GetData(), oidn::Format::Float3, Size.X, Size.Y, 0, sizeof(FLinearColor), sizeof(FLinearColor) * Size.X);
+				PixelsFilter.setImage("albedo", RawAlbedo.GetData(), oidn::Format::Half3, Size.X, Size.Y, 0, sizeof(FHalfColor), sizeof(FHalfColor) * Size.X);
+				PixelsFilter.setImage("normal", RawNormal.GetData(), oidn::Format::Half3, Size.X, Size.Y, 0, sizeof(FHalfColor), sizeof(FHalfColor) * Size.X);
 			}
-			if (CurrentMode == EDenoiseMode::CLEAN_AUX)
+			if (CurrentSettings.DenoiseAux && CurrentSettings.UseAux)
 			{
 				// +cleanAux
 				PixelsFilter.set("cleanAux", true);
 			}
 			PixelsFilter.set("hdr", true);
 			PixelsFilter.commit();
+
+			if (CurrentSettings.DenoiseAlpha)
+			{
+				AlphaFilter = OIDNDevice.newFilter("RT");
+				AlphaFilter.setImage("color" , RawPixels.GetData(), oidn::Format::Float, Size.X, Size.Y, sizeof(float) * 3, sizeof(FLinearColor), sizeof(FLinearColor) * Size.X);
+				AlphaFilter.setImage("output", RawPixels.GetData(), oidn::Format::Float, Size.X, Size.Y, sizeof(float) * 3, sizeof(FLinearColor), sizeof(FLinearColor) * Size.X);
+				AlphaFilter.set("hdr", true);
+				AlphaFilter.commit();
+			}
+			else
+			{
+				AlphaFilter = oidn::FilterRef();
+			}
 		}
 	}
 	
 	void Reset()
 	{
-		UpdateFilter(FIntPoint(0, 0), EDenoiseMode::OFF);
+		UpdateFilter(FIntPoint(0, 0), FDenoiseSettings{});
 	}
 };
 
@@ -167,10 +217,7 @@ static void CopyTextureFromCPUToGPU(FRHICommandListImmediate& RHICmdList, const
 
 static void Denoise(FRHICommandListImmediate& RHICmdList, FRHITexture* ColorTex, FRHITexture* AlbedoTex, FRHITexture* NormalTex, FRHITexture* OutputTex, FRHIGPUMask GPUMask)
 {
-	static IConsoleVariable* DenoiseModeCVar = IConsoleManager::Get().FindConsoleVariable(TEXT("r.PathTracing.Denoiser"));
-
-	const int DenoiseModeCVarValue = DenoiseModeCVar ? DenoiseModeCVar->GetInt() : -1;
-	const EDenoiseMode DenoiseMode = DenoiseModeCVarValue >= 0 ? EDenoiseMode(DenoiseModeCVarValue) : EDenoiseMode::DEFAULT;
+	FDenoiseSettings Settings = GetCurrentSettings();
 
 #if WITH_EDITOR
 	// NOTE: the time will include the transfer from GPU to CPU which will include waiting for the GPU pipeline to complete
@@ -181,16 +228,16 @@ static void Denoise(FRHICommandListImmediate& RHICmdList, FRHITexture* ColorTex,
 	FIntPoint Size = ColorTex->GetSizeXY();
 	FIntRect Rect = FIntRect(0, 0, Size.X, Size.Y);
 
-	DenoiserState.UpdateFilter(Size, DenoiseMode);
+	DenoiserState.UpdateFilter(Size, Settings);
 	CopyTextureFromGPUToCPU(RHICmdList, ColorTex, Size, DenoiserState.RawPixels);
-	if (DenoiseMode == EDenoiseMode::DEFAULT || DenoiseMode == EDenoiseMode::CLEAN_AUX)
+	if (Settings.UseAux || Settings.DenoiseAux)
 	{
 		CopyTextureFromGPUToCPU(RHICmdList, AlbedoTex, Size, DenoiserState.RawAlbedo);
 		CopyTextureFromGPUToCPU(RHICmdList, NormalTex, Size, DenoiserState.RawNormal);
 	}
 	check(DenoiserState.RawPixels.Num() == Size.X * Size.Y);
 
-	if (DenoiseMode == EDenoiseMode::CLEAN_AUX)
+	if (Settings.DenoiseAux)
 	{
 		check(DenoiserState.AlbedoFilter);
 		check(DenoiserState.NormalFilter);
@@ -200,14 +247,19 @@ static void Denoise(FRHICommandListImmediate& RHICmdList, FRHITexture* ColorTex,
 
 	check(DenoiserState.PixelsFilter);
 	DenoiserState.PixelsFilter.execute();
-	// copy pixels back to GPU (including alpha channel which was hopefully untouched by OIDN)
+	if (Settings.DenoiseAlpha)
+	{
+		check(DenoiserState.AlphaFilter);
+		DenoiserState.AlphaFilter.execute();
+	}
+	// copy pixels back to GPU (RGB and A are part of the same buffer)
 	CopyTextureFromCPUToGPU(RHICmdList, DenoiserState.RawPixels, Size, OutputTex);
 
 #if WITH_EDITOR
 	const char* errorMessage;
 	if (DenoiserState.OIDNDevice.getError(errorMessage) != oidn::Error::None)
 	{
-		UE_LOG(LogOpenImageDenoise, Warning, TEXT("Denoiser failed: %s"), *FString(errorMessage));
+		UE_LOG(LogOpenImageDenoise, Warning, TEXT("Denoiser failed: %hs"), errorMessage);
 		return;
 	}
 
@@ -249,7 +301,7 @@ public:
 void FOpenImageDenoiseModule::StartupModule()
 {
 #if WITH_EDITOR
-	UE_LOG(LogOpenImageDenoise, Log, TEXT("OIDN starting up"));
+	UE_LOG(LogOpenImageDenoise, Log, TEXT("OIDN starting up (built with OIDN %hs)"), OIDN_VERSION_STRING);
 #endif
 
 	GPathTracingDenoiserPlugin = MakeUnique<FOIDNDenoiser>();
diff --git a/Engine/Shaders/Private/PathTracing/PathTracingCompositingPixelShader.usf b/Engine/Shaders/Private/PathTracing/PathTracingCompositingPixelShader.usf
--- a/Engine/Shaders/Private/PathTracing/PathTracingCompositingPixelShader.usf
+++ b/Engine/Shaders/Private/PathTracing/PathTracingCompositingPixelShader.usf
@@ -12,7 +12,7 @@ uint Iteration;
 uint MaxSamples;
 int ProgressDisplayEnabled;
 Texture2D<float4> RadianceTexture;
-Texture2D<float4> NormalDepthTexture;
+Texture2D<float> DepthTexture;
 
 int AdaptiveSamplingVisualize;
 
@@ -93,5 +93,5 @@ void CompositeMain(
 		}
 	}
 
-	OutDepth = NormalDepthTexture.Load(TexCoord).w;
+	OutDepth = DepthTexture.Load(TexCoord);
 }
diff --git a/Engine/Shaders/Private/PathTracing/PathTracingCore.ush b/Engine/Shaders/Private/PathTracing/PathTracingCore.ush
--- a/Engine/Shaders/Private/PathTracing/PathTracingCore.ush
+++ b/Engine/Shaders/Private/PathTracing/PathTracingCore.ush
@@ -73,6 +73,7 @@ RWTexture2D<float4> VarianceTexture;
 #endif
 RWTexture2D<float4> AlbedoTexture;
 RWTexture2D<float4> NormalTexture;
+RWTexture2D<float>  DepthTexture;
 RaytracingAccelerationStructure TLAS;
 RaytracingAccelerationStructure DecalTLAS;
 uint SceneVisibleLightCount;
@@ -271,14 +272,14 @@ struct FPathState
 		const float3 Normal = GetNormal();
 
 		RadianceTexture[TextureIndex] = lerp(OldPixel, PixelValue, Blend);
-		AlbedoTexture[TextureIndex]   = lerp(OldAlbedo, float4(Albedo, 0), Blend);
+		AlbedoTexture[TextureIndex].xyz = lerp(OldAlbedo.xyz, Albedo, Blend);
 		NormalTexture[TextureIndex].xyz = lerp(OldNormal.xyz,Normal.xyz, Blend);
 	}
 
 	void WriteDepth(uint2 TextureIndex)
 	{
-		float OldDepth = Iteration > 0 ? NormalTexture[TextureIndex].w : 0;
-		NormalTexture[TextureIndex].w = max(OldDepth, ConvertToDeviceZ(DepthZ));
+		float OldDepth = Iteration > 0 ? DepthTexture[TextureIndex] : 0;
+		DepthTexture[TextureIndex] = max(OldDepth, ConvertToDeviceZ(DepthZ));
 	}
 };
 
@@ -634,7 +635,7 @@ FPackedPathTracingPayload TraceTransparentRay(inout FPathState PathState, int Bo
 							// keep tracing through the volume
 							PathThroughput *= SigmaN;
 							// early termination - fades off from 0.002 to 0.001 and then will be 0
-							PathThroughput *= saturate(1000.0 * max3(PathThroughput.x, PathThroughput.y, PathThroughput.z) - 1.0);
+							PathThroughput *= saturate(10000.0 * max3(PathThroughput.x, PathThroughput.y, PathThroughput.z) - 1.0);
 							if (!any(PathThroughput > 0))
 							{
 								bUseAnalyticTransmittance = false; // avoid re-computing the transmittance after the loop
diff --git a/Engine/Source/Runtime/Renderer/Private/PathTracing.cpp b/Engine/Source/Runtime/Renderer/Private/PathTracing.cpp
--- a/Engine/Source/Runtime/Renderer/Private/PathTracing.cpp
+++ b/Engine/Source/Runtime/Renderer/Private/PathTracing.cpp
@@ -650,13 +650,15 @@ struct FAtmosphereConfig
 	uint32 Resolution;
 };
 
-struct FPathTracingState {
+struct FPathTracingState
+{
 	FPathTracingConfig LastConfig;
 	// Textures holding onto the accumulated frame data
 	TRefCountPtr<IPooledRenderTarget> RadianceRT;
 	TRefCountPtr<IPooledRenderTarget> VarianceRT;
 	TRefCountPtr<IPooledRenderTarget> AlbedoRT;
 	TRefCountPtr<IPooledRenderTarget> NormalRT;
+	TRefCountPtr<IPooledRenderTarget> DepthRT;
 	TRefCountPtr<FRDGPooledBuffer> VarianceBuffer;
 	TRefCountPtr<IPooledRenderTarget> CloudAccelerationMap;
 	TRefCountPtr<IPooledRenderTarget> CloudMap;
@@ -664,8 +666,9 @@ struct FPathTracingState {
 	// Cache to improve the stability when frame denoising (SPP=r.pathtracing.SamplesPerPixel) is used in animation rendering
 	TRefCountPtr<IPooledRenderTarget> LastDenoisedRadianceRT;
 	TRefCountPtr<IPooledRenderTarget> LastRadianceRT;
-	TRefCountPtr<IPooledRenderTarget> LastNormalRT;
 	TRefCountPtr<IPooledRenderTarget> LastAlbedoRT;
+	TRefCountPtr<IPooledRenderTarget> LastNormalRT;
+	TRefCountPtr<IPooledRenderTarget> LastDepthRT;
 	TRefCountPtr<FRDGPooledBuffer> LastVarianceBuffer;
 
 	// Volume acceleration structures
@@ -1109,6 +1112,7 @@ class FPathTracingRG : public FGlobalShader
 		SHADER_PARAMETER_RDG_TEXTURE_UAV(RWTexture2D<float2>, VarianceTexture)
 		SHADER_PARAMETER_RDG_TEXTURE_UAV(RWTexture2D<float4>, AlbedoTexture)
 		SHADER_PARAMETER_RDG_TEXTURE_UAV(RWTexture2D<float4>, NormalTexture)
+		SHADER_PARAMETER_RDG_TEXTURE_UAV(RWTexture2D<float>, DepthTexture)
 		SHADER_PARAMETER_RDG_BUFFER_SRV(RaytracingAccelerationStructure, TLAS)
 		SHADER_PARAMETER_RDG_BUFFER_SRV(RaytracingAccelerationStructure, DecalTLAS)
 
@@ -2748,7 +2752,7 @@ class FPathTracingCompositorPS : public FGlobalShader
 		SHADER_PARAMETER_SAMPLER(SamplerState, VarianceSampler)
 		SHADER_PARAMETER_RDG_TEXTURE_SRV(Texture2D<float4>, RadianceTexture)
 		SHADER_PARAMETER_RDG_TEXTURE_SRV(Texture2D<float2>, VarianceTexture)
-		SHADER_PARAMETER_RDG_TEXTURE_SRV(Texture2D<float4>, NormalDepthTexture)
+		SHADER_PARAMETER_RDG_TEXTURE_SRV(Texture2D<float>, DepthTexture)
 		SHADER_PARAMETER_STRUCT_REF(FViewUniformShaderParameters, ViewUniformBuffer)
 		SHADER_PARAMETER(uint32, Iteration)
 		SHADER_PARAMETER(uint32, MaxSamples)
@@ -2863,6 +2867,7 @@ void FSceneViewState::PathTracingInvalidate(bool InvalidateAnimationStates)
 		State->VarianceRT.SafeRelease();
 		State->AlbedoRT.SafeRelease();
 		State->NormalRT.SafeRelease();
+		State->DepthRT.SafeRelease();
 		State->VarianceBuffer.SafeRelease();
 		State->CloudAccelerationMap.SafeRelease();
 		State->CloudMap.SafeRelease();
@@ -2893,8 +2898,9 @@ static void SplitDouble(double x, float* hi, float* lo)
 #if WITH_MGPU
 BEGIN_SHADER_PARAMETER_STRUCT(FMGPUTransferParameters, )
 	RDG_TEXTURE_ACCESS(InputTexture, ERHIAccess::CopySrc)
-	RDG_TEXTURE_ACCESS(InputAlbedo, ERHIAccess::CopySrc)
-	RDG_TEXTURE_ACCESS(InputNormal, ERHIAccess::CopySrc)
+	RDG_TEXTURE_ACCESS(InputAlbedo , ERHIAccess::CopySrc)
+	RDG_TEXTURE_ACCESS(InputNormal , ERHIAccess::CopySrc)
+	RDG_TEXTURE_ACCESS(InputDepth  , ERHIAccess::CopySrc)
 END_SHADER_PARAMETER_STRUCT()
 #endif
 
@@ -3147,27 +3153,40 @@ void FDeferredShadingSceneRenderer::RenderPathTracing(
 	// Prepare radiance buffer (will be shared with display pass)
 	FRDGTexture* RadianceTexture = nullptr;
 	FRDGTexture* VarianceTexture = nullptr;
-	FRDGTexture* AlbedoTexture = nullptr;
-	FRDGTexture* NormalTexture = nullptr;
+	FRDGTexture* AlbedoTexture   = nullptr;
+	FRDGTexture* NormalTexture   = nullptr;
+	FRDGTexture* DepthTexture    = nullptr;
 	const int NumVarianceMips = FMath::Min(5u, 1 + FMath::FloorLog2(uint32(View.ViewRect.Size().GetMin())));
 	if (PathTracingState->RadianceRT)
 	{
 		// we already have a valid radiance texture, re-use it
 		RadianceTexture = GraphBuilder.RegisterExternalTexture(PathTracingState->RadianceRT, TEXT("PathTracer.Radiance"));
-		AlbedoTexture   = GraphBuilder.RegisterExternalTexture(PathTracingState->AlbedoRT, TEXT("PathTracer.Albedo"));
-		NormalTexture   = GraphBuilder.RegisterExternalTexture(PathTracingState->NormalRT, TEXT("PathTracer.Normal"));
+		AlbedoTexture   = GraphBuilder.RegisterExternalTexture(PathTracingState->AlbedoRT  , TEXT("PathTracer.Albedo"));
+		NormalTexture   = GraphBuilder.RegisterExternalTexture(PathTracingState->NormalRT  , TEXT("PathTracer.Normal"));
+		DepthTexture    = GraphBuilder.RegisterExternalTexture(PathTracingState->DepthRT   , TEXT("PathTracer.Depth"));
 	}
 	else
 	{
 		// First time through, need to make a new texture
-		FRDGTextureDesc Desc = FRDGTextureDesc::Create2D(
+		FRDGTextureDesc RadianceDesc = FRDGTextureDesc::Create2D(
 			View.ViewRect.Size(),
 			PF_A32B32G32R32F,
 			FClearValueBinding::None,
 			TexCreate_ShaderResource | TexCreate_UAV | GetExtraTextureCreateFlagsForDenoiser());
-		RadianceTexture = GraphBuilder.CreateTexture(Desc, TEXT("PathTracer.Radiance"), ERDGTextureFlags::MultiFrame);
-		AlbedoTexture   = GraphBuilder.CreateTexture(Desc, TEXT("PathTracer.Albedo")  , ERDGTextureFlags::MultiFrame);
-		NormalTexture   = GraphBuilder.CreateTexture(Desc, TEXT("PathTracer.Normal")  , ERDGTextureFlags::MultiFrame);
+		FRDGTextureDesc AlbedoNormalDesc = FRDGTextureDesc::Create2D(
+			View.ViewRect.Size(),
+			PF_FloatRGBA,
+			FClearValueBinding::None,
+			TexCreate_ShaderResource | TexCreate_UAV | GetExtraTextureCreateFlagsForDenoiser());
+		FRDGTextureDesc DepthDesc = FRDGTextureDesc::Create2D(
+			View.ViewRect.Size(),
+			PF_R32_FLOAT,
+			FClearValueBinding::None,
+			TexCreate_ShaderResource | TexCreate_UAV | GetExtraTextureCreateFlagsForDenoiser());
+		RadianceTexture = GraphBuilder.CreateTexture(RadianceDesc    , TEXT("PathTracer.Radiance"), ERDGTextureFlags::MultiFrame);
+		AlbedoTexture   = GraphBuilder.CreateTexture(AlbedoNormalDesc, TEXT("PathTracer.Albedo")  , ERDGTextureFlags::MultiFrame);
+		NormalTexture   = GraphBuilder.CreateTexture(AlbedoNormalDesc, TEXT("PathTracer.Normal")  , ERDGTextureFlags::MultiFrame);
+		DepthTexture    = GraphBuilder.CreateTexture(DepthDesc       , TEXT("PathTracer.Depth")   , ERDGTextureFlags::MultiFrame);
 	}
 	if (Config.UseAdaptiveSampling)
 	{
@@ -3509,6 +3528,7 @@ void FDeferredShadingSceneRenderer::RenderPathTracing(
 							PassParameters->RadianceTexture = GraphBuilder.CreateUAV(RadianceTexture);
 							PassParameters->AlbedoTexture = GraphBuilder.CreateUAV(AlbedoTexture);
 							PassParameters->NormalTexture = GraphBuilder.CreateUAV(NormalTexture);
+							PassParameters->DepthTexture = GraphBuilder.CreateUAV(DepthTexture);
 
 							if (Config.UseAdaptiveSampling)
 							{
@@ -3700,6 +3720,7 @@ void FDeferredShadingSceneRenderer::RenderPathTracing(
 			Parameters->InputTexture = RadianceTexture;
 			Parameters->InputAlbedo = AlbedoTexture;
 			Parameters->InputNormal = NormalTexture;
+			Parameters->InputDepth = DepthTexture;
 			GraphBuilder.AddPass(RDG_EVENT_NAME("Path Tracer Cross-GPU Transfer (%d GPUs)", NumGPUs), Parameters, ERDGPassFlags::Readback,
 				[Parameters, DispatchResX, DispatchResY, DispatchSize, GPUMask, MainGPUMask = View.GPUMask, LocalCopyFenceDatas = MoveTemp(CopyFenceDatas)](FRHICommandListImmediate& RHICmdList)
 				{
@@ -3734,6 +3755,7 @@ void FDeferredShadingSceneRenderer::RenderPathTracing(
 									TransferParams.Emplace(Parameters->InputTexture->GetRHI(), TileToCopy, GPUIndex, FirstGPUIndex, false, false);
 									TransferParams.Emplace(Parameters->InputAlbedo->GetRHI(), TileToCopy, GPUIndex, FirstGPUIndex, false, false);
 									TransferParams.Emplace(Parameters->InputNormal->GetRHI(), TileToCopy, GPUIndex, FirstGPUIndex, false, false);
+									TransferParams.Emplace(Parameters->InputDepth->GetRHI(), TileToCopy, GPUIndex, FirstGPUIndex, false, false);
 								}
 								++CurrentGPU;
 							}
@@ -3756,6 +3778,7 @@ void FDeferredShadingSceneRenderer::RenderPathTracing(
 		GraphBuilder.QueueTextureExtraction(RadianceTexture, &PathTracingState->RadianceRT);
 		GraphBuilder.QueueTextureExtraction(AlbedoTexture, &PathTracingState->AlbedoRT);
 		GraphBuilder.QueueTextureExtraction(NormalTexture, &PathTracingState->NormalRT);
+		GraphBuilder.QueueTextureExtraction(DepthTexture, &PathTracingState->DepthRT);
 		if (Config.UseAdaptiveSampling)
 		{
 			check(VarianceTexture != nullptr);
@@ -3787,21 +3810,23 @@ void FDeferredShadingSceneRenderer::RenderPathTracing(
 	{
 		// mGPU renders blocks of pixels that need to be mapped back into alternating scanlines
 		// perform this swizzling now with a simple compute shader
+		// NOTE: we only perform this swizzling for albedo/normals if we are going to use them for denoising
 
 		TShaderMapRef<FPathTracingSwizzleScanlinesCS> ComputeShader(GetGlobalShaderMap(View.FeatureLevel));
-		FRDGTextureDesc Desc = FRDGTextureDesc::Create2D(
-			View.ViewRect.Size(),
-			PF_A32B32G32R32F,
-			FClearValueBinding::None,
-			TexCreate_ShaderResource | TexCreate_UAV);
-		FRDGTexture* NewRadianceTexture = GraphBuilder.CreateTexture(Desc, TEXT("PathTracer.RadianceUnswizzled"));
-		FRDGTexture* NewNormalTexture = GraphBuilder.CreateTexture(Desc, TEXT("PathTracer.NormalUnswizzled"));
-		FRDGTexture* NewAlbedoTexture = NeedsDenoise ? GraphBuilder.CreateTexture(Desc, TEXT("PathTracer.AlbedoUnswizzled")) : nullptr;
-
-		FRDGTexture* InputTextures[3] = { RadianceTexture, NormalTexture, AlbedoTexture};
-		FRDGTexture* OutputTextures[3] = { NewRadianceTexture, NewNormalTexture, NewAlbedoTexture};
-		for (int Index = 0, Num = NeedsDenoise ? 3 : 2; Index < Num; Index++)
+		FRDGTexture* NewRadianceTexture = GraphBuilder.CreateTexture(RadianceTexture->Desc, TEXT("PathTracer.RadianceUnswizzled"));
+		FRDGTexture* NewDepthTexture = GraphBuilder.CreateTexture(DepthTexture->Desc, TEXT("PathTracer.DepthUnswizzled"));
+		FRDGTexture* NewNormalTexture = NeedsDenoise ? GraphBuilder.CreateTexture(NormalTexture->Desc, TEXT("PathTracer.NormalUnswizzled")) : nullptr;
+		FRDGTexture* NewAlbedoTexture = NeedsDenoise ? GraphBuilder.CreateTexture(AlbedoTexture->Desc, TEXT("PathTracer.AlbedoUnswizzled")) : nullptr;
+
+		FRDGTexture* InputTextures[4] = { RadianceTexture, NormalTexture, DepthTexture, AlbedoTexture};
+		FRDGTexture* OutputTextures[4] = { NewRadianceTexture, NewNormalTexture, NewDepthTexture, NewAlbedoTexture};
+		for (int Index = 0; Index < 4; Index++)
 		{
+			if (OutputTextures[Index] == nullptr)
+			{
+				// skip unused textures
+				continue;
+			}
 			FPathTracingSwizzleScanlinesCS::FParameters* PassParameters = GraphBuilder.AllocParameters<FPathTracingSwizzleScanlinesCS::FParameters>();
 			PassParameters->DispatchDim.X = DispatchResX;
 			PassParameters->DispatchDim.Y = DispatchResY;
@@ -3821,6 +3846,7 @@ void FDeferredShadingSceneRenderer::RenderPathTracing(
 		// let the remaining code operate on the unswizzled textures
 		RadianceTexture = NewRadianceTexture;
 		NormalTexture = NewNormalTexture;
+		DepthTexture = NewDepthTexture;
 		AlbedoTexture = NewAlbedoTexture;
 	}
 #endif
@@ -3871,6 +3897,7 @@ void FDeferredShadingSceneRenderer::RenderPathTracing(
 			DenoisingContext.RadianceTexture = RadianceTexture;
 			DenoisingContext.AlbedoTexture = AlbedoTexture;
 			DenoisingContext.NormalTexture = NormalTexture;
+			DenoisingContext.DepthTexture = DepthTexture;
 			DenoisingContext.VarianceBuffer = PathTracingState->VarianceBuffer ? 
 				GraphBuilder.RegisterExternalBuffer(PathTracingState->VarianceBuffer, TEXT("PathTracing.VarianceBuffer")) : nullptr;
 			DenoisingContext.LastVarianceBuffer = PathTracingState->LastVarianceBuffer ?
@@ -3894,10 +3921,12 @@ void FDeferredShadingSceneRenderer::RenderPathTracing(
 					GraphBuilder.RegisterExternalTexture(PathTracingState->LastDenoisedRadianceRT, TEXT("PathTracing.LastPreDenoisedRadiance"));
 				DenoisingContext.LastRadianceTexture =
 					GraphBuilder.RegisterExternalTexture(PathTracingState->LastRadianceRT, TEXT("PathTracing.LastRadianceTexture"));
-				DenoisingContext.LastNormalTexture =
-					GraphBuilder.RegisterExternalTexture(PathTracingState->LastNormalRT, TEXT("PathTracing.LastNormalTexture"));
 				DenoisingContext.LastAlbedoTexture =
 					GraphBuilder.RegisterExternalTexture(PathTracingState->LastAlbedoRT, TEXT("PathTracing.LastAlbedoTexture"));
+				DenoisingContext.LastNormalTexture =
+					GraphBuilder.RegisterExternalTexture(PathTracingState->LastNormalRT, TEXT("PathTracing.LastNormalTexture"));
+				DenoisingContext.LastDepthTexture =
+					GraphBuilder.RegisterExternalTexture(PathTracingState->LastDepthRT, TEXT("PathTracing.LastDepthTexture"));
 			}
 
 			PathTracingSpatialTemporalDenoising(GraphBuilder,
@@ -3907,8 +3936,9 @@ void FDeferredShadingSceneRenderer::RenderPathTracing(
 				DenoisingContext);
 
 			GraphBuilder.QueueTextureExtraction(DenoisedRadianceTexture, &PathTracingState->LastDenoisedRadianceRT);
-			GraphBuilder.QueueTextureExtraction(NormalTexture, &PathTracingState->LastNormalRT);
 			GraphBuilder.QueueTextureExtraction(AlbedoTexture, &PathTracingState->LastAlbedoRT);
+			GraphBuilder.QueueTextureExtraction(NormalTexture, &PathTracingState->LastNormalRT);
+			GraphBuilder.QueueTextureExtraction(DepthTexture, &PathTracingState->LastDepthRT);
 			GraphBuilder.QueueTextureExtraction(RadianceTexture, &PathTracingState->LastRadianceRT);
 
 			PathTracingState->SpatialTemporalDenoiserHistory = DenoisingContext.SpatialTemporalDenoiserHistory;
@@ -3943,7 +3973,7 @@ void FDeferredShadingSceneRenderer::RenderPathTracing(
 	DisplayParameters->ViewUniformBuffer = View.ViewUniformBuffer;
 	DisplayParameters->RadianceTexture = GraphBuilder.CreateSRV(FRDGTextureSRVDesc::Create(DenoisedRadianceTexture ? DenoisedRadianceTexture : RadianceTexture));
 	DisplayParameters->VarianceTexture = GraphBuilder.CreateSRV(FRDGTextureSRVDesc::Create(VarianceTexture ? VarianceTexture : GSystemTextures.GetBlackDummy(GraphBuilder)));
-	DisplayParameters->NormalDepthTexture = GraphBuilder.CreateSRV(FRDGTextureSRVDesc::Create(NormalTexture));
+	DisplayParameters->DepthTexture = GraphBuilder.CreateSRV(FRDGTextureSRVDesc::Create(DepthTexture));
 	DisplayParameters->RenderTargets[0] = FRenderTargetBinding(SceneColorOutputTexture, ERenderTargetLoadAction::ELoad);
 	DisplayParameters->RenderTargets.DepthStencil = FDepthStencilBinding(SceneDepthOutputTexture,  ERenderTargetLoadAction::ELoad, ERenderTargetLoadAction::ENoAction, FExclusiveDepthStencil::DepthWrite_StencilNop);
 	DisplayParameters->VarianceSampler = TStaticSamplerState<SF_Bilinear>::GetRHI();
diff --git a/Engine/Source/Runtime/Renderer/Private/PathTracingSpatialTemporalDenoising.h b/Engine/Source/Runtime/Renderer/Private/PathTracingSpatialTemporalDenoising.h
--- a/Engine/Source/Runtime/Renderer/Private/PathTracingSpatialTemporalDenoising.h
+++ b/Engine/Source/Runtime/Renderer/Private/PathTracingSpatialTemporalDenoising.h
@@ -20,12 +20,14 @@ struct FPathTracingSpatialTemporalDenoisingContext
 	FRDGTexture*	RadianceTexture = nullptr;
 	FRDGTexture*	AlbedoTexture = nullptr;
 	FRDGTexture*	NormalTexture = nullptr;
+	FRDGTexture*	DepthTexture = nullptr;
 	FRDGBuffer*		VarianceBuffer = nullptr;
 
 	FRDGTexture*	LastDenoisedRadianceTexture = nullptr;
 	FRDGTexture*	LastRadianceTexture = nullptr;
-	FRDGTexture*	LastNormalTexture = nullptr;
 	FRDGTexture*	LastAlbedoTexture = nullptr;
+	FRDGTexture*	LastNormalTexture = nullptr;
+	FRDGTexture*	LastDepthTexture = nullptr;
 	FRDGBuffer*		LastVarianceBuffer = nullptr;
 
 	// Custom path tracing spacial temporal denoiser result, used by plugins
