Change r.VT.AVT.AgeToFree so that a value of 0 disables automatic eviction of old adaptive virtual texture entries.
Change the default setting to 0.
This avoids potentially high CPU costs for remapping pages unnecessarily.
By default it is likely to be more efficient to wait until there is pressure on the adaptive page table before starting to evict (which we already also do with r.VT.AVT.MaxPageResidency).
When eviction from residency kicks in, it is much more likely that the associated virtual texture pages have already been evicted due to pressure on the physical pool. That means there will be no time spent needing to remap them.

diff --git a/Engine/Source/Runtime/Renderer/Private/VT/AdaptiveVirtualTexture.cpp b/Engine/Source/Runtime/Renderer/Private/VT/AdaptiveVirtualTexture.cpp
--- a/Engine/Source/Runtime/Renderer/Private/VT/AdaptiveVirtualTexture.cpp
+++ b/Engine/Source/Runtime/Renderer/Private/VT/AdaptiveVirtualTexture.cpp
@@ -32,7 +32,7 @@ static TAutoConsoleVariable<int32> CVarAVTMaxPageResidency(
 
 static TAutoConsoleVariable<int32> CVarAVTAgeToFree(
 	TEXT("r.VT.AVT.AgeToFree"),
-	60,
+	0,
 	TEXT("Number of frames for an allocation to be unused before it is considered for free"),
 	ECVF_RenderThreadSafe
 );
@@ -678,12 +678,15 @@ void FAdaptiveVirtualTexture::UpdateAllocations(FVirtualTextureSystem* InSystem,
 	{
 		// Free old unused pages if there is no other work to do.
 		const uint32 FrameAgeToFree = CVarAVTAgeToFree.GetValueOnRenderThread();
-		const int32 NumToFree = FMath::Min(NumAllocated, CVarAVTMaxFreePerFrame.GetValueOnRenderThread());
-
-		bool bFreeSuccess = true;
-		for (int32 FreeCount = 0; bFreeSuccess && FreeCount < NumToFree; FreeCount++)
+		if (FrameAgeToFree > 0)
 		{
-			bFreeSuccess = FreeLRU(RHICmdList, InSystem, InFrame, FrameAgeToFree);
+			const int32 NumToFree = FMath::Min(NumAllocated, CVarAVTMaxFreePerFrame.GetValueOnRenderThread());
+
+			bool bFreeSuccess = true;
+			for (int32 FreeCount = 0; bFreeSuccess && FreeCount < NumToFree; FreeCount++)
+			{
+				bFreeSuccess = FreeLRU(RHICmdList, InSystem, InFrame, FrameAgeToFree);
+			}
 		}
 	}
 	else
