Fix water rendering in ISR-enabled projects.

diff --git a/Engine/Source/Runtime/Renderer/Private/WaterInfoTextureRendering.cpp b/Engine/Source/Runtime/Renderer/Private/WaterInfoTextureRendering.cpp
--- a/Engine/Source/Runtime/Renderer/Private/WaterInfoTextureRendering.cpp
+++ b/Engine/Source/Runtime/Renderer/Private/WaterInfoTextureRendering.cpp
@@ -33,13 +33,13 @@ static TAutoConsoleVariable<int32> CVarRenderCaptureNextWaterInfoDraws(
 	ECVF_RenderThreadSafe);
 
 BEGIN_SHADER_PARAMETER_STRUCT(FWaterInfoTexturePassParameters, )
-	SHADER_PARAMETER_RDG_UNIFORM_BUFFER(FViewUniformShaderParameters, View)
+	SHADER_PARAMETER_STRUCT_INCLUDE(FViewShaderParameters, View)
 	SHADER_PARAMETER_STRUCT_INCLUDE(FInstanceCullingDrawParams, InstanceCullingDrawParams)
 	RENDER_TARGET_BINDING_SLOTS()
 END_SHADER_PARAMETER_STRUCT()
 
 BEGIN_SHADER_PARAMETER_STRUCT(FWaterInfoTextureDepthPassParameters, )
-	SHADER_PARAMETER_RDG_UNIFORM_BUFFER(FViewUniformShaderParameters, View)
+	SHADER_PARAMETER_STRUCT_INCLUDE(FViewShaderParameters, View)
 	SHADER_PARAMETER_STRUCT_INCLUDE(FInstanceCullingDrawParams, InstanceCullingDrawParams)
 	RENDER_TARGET_BINDING_SLOTS()
 END_SHADER_PARAMETER_STRUCT()
@@ -51,7 +51,7 @@ public:
 	SHADER_USE_PARAMETER_STRUCT(FWaterInfoTextureMergePS, FGlobalShader);
 
 	BEGIN_SHADER_PARAMETER_STRUCT(FParameters, )
-		SHADER_PARAMETER_RDG_UNIFORM_BUFFER(FViewUniformShaderParameters, View)
+		SHADER_PARAMETER_STRUCT_INCLUDE(FViewShaderParameters, View)
 		SHADER_PARAMETER_STRUCT_INCLUDE(FSceneTextureShaderParameters, SceneTextures)
 		SHADER_PARAMETER_RDG_TEXTURE(Texture2D<float>, GroundDepthTexture)
 		SHADER_PARAMETER_SAMPLER(SamplerState, GroundDepthTextureSampler)
@@ -89,7 +89,7 @@ public:
 	SHADER_USE_PARAMETER_STRUCT(FWaterInfoTextureBlurPS, FGlobalShader);
 
 	BEGIN_SHADER_PARAMETER_STRUCT(FParameters, )
-		SHADER_PARAMETER_RDG_UNIFORM_BUFFER(FViewUniformShaderParameters, View)
+		SHADER_PARAMETER_STRUCT_INCLUDE(FViewShaderParameters, View)
 		SHADER_PARAMETER_STRUCT_INCLUDE(FSceneTextureShaderParameters, SceneTextures)
 		SHADER_PARAMETER_RDG_TEXTURE(Texture2D, WaterInfoTexture)
 		SHADER_PARAMETER(float, WaterZMin)
@@ -888,8 +888,6 @@ void RenderWaterInfoTexture(
 			PassDraws[PassIdx].InstanceCullingResult.Parameters.Scene = SceneRenderer.GetSceneUniforms().GetBuffer(GraphBuilder);
 		}
 
-		TRDGUniformBufferRef<FViewUniformShaderParameters> ViewUB = GraphBuilder.CreateUniformBuffer(GraphBuilder.AllocParameters(WaterView.CachedViewUniformShaderParameters.Get()));
-
 		// Execute the three mesh passes involved in generating the water info texture: terrain depth-only, dilated water body depth-only and finally regular water body depth + river velocity
 		for (uint32 PassIndex = 0; PassIndex < static_cast<int32>(EWaterInfoTextureMeshPass::Num); ++PassIndex)
 		{
@@ -900,7 +898,7 @@ void RenderWaterInfoTexture(
 			{
 				const bool bIsTerrainDepthPass = PassType == EWaterInfoTextureMeshPass::TerrainDepth;
 				FWaterInfoTextureDepthPassParameters* PassParameters = GraphBuilder.AllocParameters<FWaterInfoTextureDepthPassParameters>();
-				PassParameters->View = ViewUB;
+				PassParameters->View = WaterView.GetShaderParameters();
 				PassParameters->RenderTargets.DepthStencil = FDepthStencilBinding(bIsTerrainDepthPass ? TerrainDepthBuffer : DilatedDepthBuffer, ERenderTargetLoadAction::EClear, ERenderTargetLoadAction::ENoAction, FExclusiveDepthStencil::DepthWrite_StencilNop);
 				PassDraws[PassIndex].InstanceCullingResult.GetDrawParameters(PassParameters->InstanceCullingDrawParams);
 
@@ -929,7 +927,7 @@ void RenderWaterInfoTexture(
 			else if (PassType == EWaterInfoTextureMeshPass::WaterBody)
 			{
 				FWaterInfoTexturePassParameters* PassParameters = GraphBuilder.AllocParameters<FWaterInfoTexturePassParameters>();
-				PassParameters->View = ViewUB;
+				PassParameters->View = WaterView.GetShaderParameters();
 				PassParameters->RenderTargets = GetRenderTargetBindings(ERenderTargetLoadAction::EClear, MakeArrayView(&WaterInfoColorTexture, 1));
 				PassParameters->RenderTargets.DepthStencil = FDepthStencilBinding(WaterInfoDepthBuffer, ERenderTargetLoadAction::EClear, ERenderTargetLoadAction::ENoAction, FExclusiveDepthStencil::DepthWrite_StencilNop);
 				PassDraws[PassIndex].InstanceCullingResult.GetDrawParameters(PassParameters->InstanceCullingDrawParams);
@@ -968,7 +966,7 @@ void RenderWaterInfoTexture(
 			TShaderMapRef<FWaterInfoTextureMergePS> PixelShader(ShaderMap, PixelPermutationVector);
 
 			FWaterInfoTextureMergePS::FParameters* PassParameters = GraphBuilder.AllocParameters<FWaterInfoTextureMergePS::FParameters>();
-			PassParameters->View = ViewUB;
+			PassParameters->View = WaterView.GetShaderParameters();
 			PassParameters->RenderTargets[0] = FRenderTargetBinding(MergedTexture, ERenderTargetLoadAction::ENoAction);
 			PassParameters->SceneTextures = GetSceneTextureShaderParameters(WaterView);
 			PassParameters->GroundDepthTexture = TerrainDepthBuffer;
@@ -1003,7 +1001,7 @@ void RenderWaterInfoTexture(
 			TShaderMapRef<FWaterInfoTextureBlurPS> PixelShader(ShaderMap, PixelPermutationVector);
 
 			FWaterInfoTextureBlurPS::FParameters* PassParameters = GraphBuilder.AllocParameters<FWaterInfoTextureBlurPS::FParameters>();
-			PassParameters->View = ViewUB;
+			PassParameters->View = WaterView.GetShaderParameters();
 			PassParameters->RenderTargets[0] = FRenderTargetBinding(OutputTexture, ERenderTargetLoadAction::ENoAction);
 			PassParameters->SceneTextures = GetSceneTextureShaderParameters(WaterView);
 			PassParameters->WaterInfoTexture = MergedTexture;
