Quick fix for release after free of resource debug names.
Use fixed strings for now.
Will fix with dynamic naming once the use after free is solved.
#rnx

diff --git a/Engine/Source/Runtime/Renderer/Private/VT/VirtualTexturePhysicalSpace.cpp b/Engine/Source/Runtime/Renderer/Private/VT/VirtualTexturePhysicalSpace.cpp
--- a/Engine/Source/Runtime/Renderer/Private/VT/VirtualTexturePhysicalSpace.cpp
+++ b/Engine/Source/Runtime/Renderer/Private/VT/VirtualTexturePhysicalSpace.cpp
@@ -167,7 +167,7 @@ void FVirtualTexturePhysicalSpace::InitRHI(FRHICommandListBase& RHICmdList)
 			Desc.UAVFormat = FormatUAV;
 		}
 
-		GRenderTargetPool.FindFreeElement(RHICmdList, Desc, PooledRenderTarget[Layer], *PooledRenderTargetDebugNames[Layer]);
+		GRenderTargetPool.FindFreeElement(RHICmdList, Desc, PooledRenderTarget[Layer], TEXT("VirtualTexture_Physical"));
 		FRHITexture* TextureRHI = PooledRenderTarget[Layer]->GetRHI();
 
 		// Create sRGB and non-sRGB shader resource views into the physical texture
diff --git a/Engine/Source/Runtime/Renderer/Private/VT/VirtualTextureSpace.cpp b/Engine/Source/Runtime/Renderer/Private/VT/VirtualTextureSpace.cpp
--- a/Engine/Source/Runtime/Renderer/Private/VT/VirtualTextureSpace.cpp
+++ b/Engine/Source/Runtime/Renderer/Private/VT/VirtualTextureSpace.cpp
@@ -318,7 +318,7 @@ void FVirtualTextureSpace::AllocateTextures(FRDGBuilder& GraphBuilder)
 				TexCreate_RenderTargetable | TexCreate_ShaderResource,
 				CachedNumPageTableLevels);
 
-			FRDGTextureRef DstTexture = GraphBuilder.CreateTexture(Desc, *PageTableDebugNames[TextureIndex]);
+			FRDGTextureRef DstTexture = GraphBuilder.CreateTexture(Desc, TEXT("VirtualTexture_PageTable"));
 
 			if (TextureEntry.RenderTarget)
 			{
