Fix memory leak due to CachedRayTracingMeshCommands not being released when ray tracing is toggled off at runtime.

diff --git a/Engine/Source/Runtime/Renderer/Private/PrimitiveSceneInfo.cpp b/Engine/Source/Runtime/Renderer/Private/PrimitiveSceneInfo.cpp
--- a/Engine/Source/Runtime/Renderer/Private/PrimitiveSceneInfo.cpp
+++ b/Engine/Source/Runtime/Renderer/Private/PrimitiveSceneInfo.cpp
@@ -1122,9 +1122,11 @@ void CacheRayTracingPrimitive(
 			}
 		}
 
+		check(SceneInfo->CachedRayTracingMeshCommandIndicesPerLOD.IsEmpty());
 		SceneInfo->CachedRayTracingMeshCommandIndicesPerLOD.Empty(MaxLOD + 1);
-		SceneInfo->CachedRayTracingMeshCommandIndicesPerLOD.AddDefaulted(MaxLOD + 1); // should be initialzied to -1?
+		SceneInfo->CachedRayTracingMeshCommandIndicesPerLOD.AddDefaulted(MaxLOD + 1); // should be initialized to -1?
 
+		check(SceneInfo->CachedRayTracingMeshCommandsHashPerLOD.IsEmpty());
 		SceneInfo->CachedRayTracingMeshCommandsHashPerLOD.Empty(MaxLOD + 1);
 		SceneInfo->CachedRayTracingMeshCommandsHashPerLOD.AddZeroed(MaxLOD + 1);
 
@@ -1304,7 +1306,7 @@ void FPrimitiveSceneInfo::UpdateCachedRayTracingInstance(FPrimitiveSceneInfo* Sc
 
 void FPrimitiveSceneInfo::RemoveCachedRayTracingPrimitives()
 {
-	if (IsRayTracingEnabled())
+	if (IsRayTracingAllowed())
 	{
 		for (auto& CachedRayTracingMeshCommandIndices : CachedRayTracingMeshCommandIndicesPerLOD)
 		{
@@ -1318,9 +1320,13 @@ void FPrimitiveSceneInfo::RemoveCachedRayTracingPrimitives()
 		}
 
 		CachedRayTracingMeshCommandIndicesPerLOD.Empty();
-
 		CachedRayTracingMeshCommandsHashPerLOD.Empty();
 	}
+	else
+	{
+		check(CachedRayTracingMeshCommandIndicesPerLOD.IsEmpty());
+		check(CachedRayTracingMeshCommandsHashPerLOD.IsEmpty());
+	}
 }
 #endif
 
