Also fix jittering in downsampled post-motionblur translucent pass, exposed in downsampled code path that was previously (incorrectly) never used for post-motionblur pass.

diff --git a/Engine/Source/Runtime/Renderer/Private/TranslucentRendering.cpp b/Engine/Source/Runtime/Renderer/Private/TranslucentRendering.cpp
--- a/Engine/Source/Runtime/Renderer/Private/TranslucentRendering.cpp
+++ b/Engine/Source/Runtime/Renderer/Private/TranslucentRendering.cpp
@@ -774,7 +774,10 @@ void SetupDownsampledTranslucencyViewParameters(
 	{
 		// Remove jitter from this pass
 		ViewMatrices.HackRemoveTemporalAAProjectionJitter();
-		PrevViewMatrices.HackRemoveTemporalAAProjectionJitter();
+		PrevViewMatrices = ViewMatrices;
+
+		FBox VolumeBounds[TVC_MAX];
+		View.SetupUniformBufferParameters(ViewMatrices, PrevViewMatrices, VolumeBounds, TVC_MAX, DownsampledTranslucencyViewParameters);
 	}
 
 	// Update the parts of DownsampledTranslucencyParameters which are dependent on the buffer size and view rect
