Added counters to track added / updated / removed scene primitives in Unreal Insights.

diff --git a/Engine/Source/Runtime/Renderer/Private/RendererScene.cpp b/Engine/Source/Runtime/Renderer/Private/RendererScene.cpp
--- a/Engine/Source/Runtime/Renderer/Private/RendererScene.cpp
+++ b/Engine/Source/Runtime/Renderer/Private/RendererScene.cpp
@@ -160,6 +160,10 @@ IMPLEMENT_GLOBAL_SHADER_PARAMETER_STRUCT(FDistanceCullFadeUniformShaderParameter
 IMPLEMENT_GLOBAL_SHADER_PARAMETER_STRUCT(FDitherUniformShaderParameters, "PrimitiveDither");
 
 TRACE_DECLARE_INT_COUNTER(ScenePrimitives, TEXT("Scene/Primitives"));
+TRACE_DECLARE_INT_COUNTER(ScenePrimitivesAdds, TEXT("Scene/Primitives/Adds"));
+TRACE_DECLARE_INT_COUNTER(ScenePrimitivesRemoves, TEXT("Scene/Primitives/Removes"));
+TRACE_DECLARE_INT_COUNTER(ScenePrimitivesUpdatesTransforms, TEXT("Scene/Primitives/Updates/Transforms"));
+TRACE_DECLARE_INT_COUNTER(ScenePrimitivesUpdatesInstances, TEXT("Scene/Primitives/Updates/Instances"));
 TRACE_DECLARE_INT_COUNTER(ScenePrimitivesArrayMax, TEXT("Scene/PrimitiveArrayMax"));
 TRACE_DECLARE_INT_COUNTER(SceneLights, TEXT("Scene/Lights"));
 
@@ -5609,7 +5613,12 @@ void FScene::Update(FRDGBuilder& GraphBuilder, const FUpdateParameters& Paramete
 			}
 		}
 	}
-	
+
+	TRACE_COUNTER_SET(ScenePrimitivesAdds, AddedPrimitiveSceneInfos.Num());
+	TRACE_COUNTER_SET(ScenePrimitivesRemoves, RemovedPrimitiveSceneInfos.Num());
+	TRACE_COUNTER_SET(ScenePrimitivesUpdatesTransforms, UpdatedInstances.Num());
+	TRACE_COUNTER_SET(ScenePrimitivesUpdatesInstances, UpdatedTransforms.Num());
+
 	// NOTE: We clear this early because IsPrimitiveBeingRemoved gets called from the CreateLightPrimitiveInteraction (to make sure that old primitives are not accessed) 
 	// we cannot safely kick off the AsyncCreateLightPrimitiveInteractionsTask before the RemovedPrimitiveSceneInfos has been cleared.
 	// TODO: this is probably not true anymore!
