Fixed crash on edgecase where a view somehow cannot be found in it's parent viewfamily.
#rnx

diff --git a/Engine/Source/Runtime/Renderer/Private/MeshEdgesRendering.cpp b/Engine/Source/Runtime/Renderer/Private/MeshEdgesRendering.cpp
--- a/Engine/Source/Runtime/Renderer/Private/MeshEdgesRendering.cpp
+++ b/Engine/Source/Runtime/Renderer/Private/MeshEdgesRendering.cpp
@@ -167,7 +167,8 @@ const FMeshEdgesViewSettings& GetMeshEdgesViewSettings(const FSceneViewFamily& V
 
 	if (FamilyData->ViewSettings.IsEmpty())
 	{
-		FamilyData->ViewSettings.Init(FMeshEdgesViewSettings{}, ViewFamily.Views.Num());
+		// Allocate 1 extra as a fallback for weird edgecase where the view can't be found in the viewfamily.
+		FamilyData->ViewSettings.Init(FMeshEdgesViewSettings{}, ViewFamily.Views.Num() + 1);
 	}
 
 	return FamilyData->ViewSettings[ViewIndex];
