Sync RDG builder async deletion tasks at shutdown.

diff --git a/Engine/Source/Runtime/Renderer/Private/Renderer.cpp b/Engine/Source/Runtime/Renderer/Private/Renderer.cpp
--- a/Engine/Source/Runtime/Renderer/Private/Renderer.cpp
+++ b/Engine/Source/Runtime/Renderer/Private/Renderer.cpp
@@ -85,6 +85,7 @@ void FRendererModule::StartupModule()
 		ENQUEUE_RENDER_COMMAND(FSceneRendererCleanUp)(
 			[](FRHICommandListImmediate& RHICmdList)
 		{
+			FRDGBuilder::WaitForAsyncDeleteTask();
 			FSceneRenderer::CleanUp(RHICmdList);
 		});
 	}));
