Use correct PSO collector index for deferred decal stat tracking (was marked as untrackted)

diff --git a/Engine/Source/Runtime/Renderer/Private/CompositionLighting/PostProcessDeferredDecals.cpp b/Engine/Source/Runtime/Renderer/Private/CompositionLighting/PostProcessDeferredDecals.cpp
--- a/Engine/Source/Runtime/Renderer/Private/CompositionLighting/PostProcessDeferredDecals.cpp
+++ b/Engine/Source/Runtime/Renderer/Private/CompositionLighting/PostProcessDeferredDecals.cpp
@@ -657,8 +657,8 @@ void AddDeferredDecalPass(
 		{
 			RHICmdList.SetViewport(View.ViewRect.Min.X, View.ViewRect.Min.Y, 0.0f, View.ViewRect.Max.X, View.ViewRect.Max.Y, 1.0f);
 
-#if PSO_PRECACHING_VALIDATE
-			int32 PSOCollectorIndex = FPSOCollectorCreateManager::GetIndex(EShadingPath::Deferred, TEXT("MeshDecal"));
+#if PSO_PRECACHING_VALIDATE			
+			int32 PSOCollectorIndex = FPassProcessorManager::GetPSOCollectorIndex(EShadingPath::Deferred, DecalRendering::GetMeshPassType(RenderTargetMode));
 #endif // PSO_PRECACHING_VALIDATE
 
 			for (uint32 DecalIndex = DecalIndexBegin; DecalIndex < DecalIndexEnd; ++DecalIndex)
