Remove FRefCountedRHIPooledRenderQuery
 - Unused class.

diff --git a/Engine/Source/Runtime/Renderer/Private/SceneRendering.h b/Engine/Source/Runtime/Renderer/Private/SceneRendering.h
--- a/Engine/Source/Runtime/Renderer/Private/SceneRendering.h
+++ b/Engine/Source/Runtime/Renderer/Private/SceneRendering.h
@@ -324,101 +324,6 @@ private:
 	uint32 NumBufferedFrames = 0;
 };
 
-class FRefCountedRHIPooledRenderQuery
-{
-public:
-	FRefCountedRHIPooledRenderQuery() : FRefCountedRHIPooledRenderQuery(FRHIPooledRenderQuery())
-	{
-	}
-
-	explicit FRefCountedRHIPooledRenderQuery(FRHIPooledRenderQuery&& InQuery)
-	{
-		char* data = new char[sizeof(FRHIPooledRenderQuery) + sizeof(int)];
-		Query = new (data) FRHIPooledRenderQuery();
-		RefCount = new (data + sizeof(FRHIPooledRenderQuery)) int(1);
-		*Query = MoveTemp(InQuery);
-	}
-
-	~FRefCountedRHIPooledRenderQuery()
-	{
-		Deref();
-	}
-
-	bool IsValid() const
-	{
-		return Query && Query->IsValid();
-	}
-
-	FRHIRenderQuery* GetQuery() const
-	{
-		return Query ? Query->GetQuery() : nullptr;
-	}
-
-	void ReleaseQuery()
-	{
-		Deref();
-	}
-
-	FRefCountedRHIPooledRenderQuery(const FRefCountedRHIPooledRenderQuery& Other)
-	{
-		Other.Addref();
-		RefCount = Other.RefCount;
-		Query = Other.Query;
-	}
-
-	FRefCountedRHIPooledRenderQuery& operator= (const FRefCountedRHIPooledRenderQuery& Other)
-	{
-		Other.Addref();
-		Deref();
-		RefCount = Other.RefCount;
-		Query = Other.Query;
-
-		return *this;
-	}
-
-	FRefCountedRHIPooledRenderQuery(FRefCountedRHIPooledRenderQuery&& Other)
-	{
-		RefCount = Other.RefCount;
-		Query = Other.Query;
-		Other.RefCount = nullptr;
-		Other.Query = nullptr;
-	}
-
-	FRefCountedRHIPooledRenderQuery& operator=(FRefCountedRHIPooledRenderQuery&& Other)
-	{
-		Deref();
-		RefCount = Other.RefCount;
-		Query = Other.Query;
-		Other.RefCount = nullptr;
-		Other.Query = nullptr;
-
-		return *this;
-	}
-
-private:
-	void Addref() const
-	{
-		check(RefCount != nullptr);
-		(*RefCount)++;
-	}
-
-	void Deref()
-	{
-		if (RefCount && --(*RefCount) == 0)
-		{
-			Query->~FRHIPooledRenderQuery();
-			char* data = reinterpret_cast<char*>(Query);
-			delete[] data;
-		}
-
-		Query = nullptr;
-		RefCount = nullptr;
-	}
-
-	mutable int* RefCount = nullptr;
-	FRHIPooledRenderQuery* Query = nullptr;
-};
-
 /**
  * Combines consecutive primitives which use the same occlusion query into a single DrawIndexedPrimitive call.
  */
