Add a separate cvar to control how long unreferenced VSM lights live, separate from the per-page ages.
Keeping VSM lights around too long can cause too much bloat in the page table sizes and processing, reducing performance in various page-table-related passes (clearing, etc).

diff --git a/Engine/Shaders/Private/VirtualShadowMaps/VirtualShadowMapCacheInvalidation.ush b/Engine/Shaders/Private/VirtualShadowMaps/VirtualShadowMapCacheInvalidation.ush
--- a/Engine/Shaders/Private/VirtualShadowMaps/VirtualShadowMapCacheInvalidation.ush
+++ b/Engine/Shaders/Private/VirtualShadowMaps/VirtualShadowMapCacheInvalidation.ush
@@ -121,7 +121,7 @@ void InvalidateInstancePages(int VirtualShadowMapId, uint InstanceId, bool bForc
 							uint PageFlagOffset = CalcPageOffset(PageTableLevelOffset, MipLevel, uint2(x, y));
 							uint PageFlags = VirtualShadowMap.PageFlags[PageFlagOffset];
 
-							if (PageFlags != 0)
+							if ((PageFlags & VSM_FLAG_ALLOCATED) != 0)
 							{
 #if USE_HZB_OCCLUSION
 								if (!IsPageVisibleHZB(uint2(x, y), PageFlagOffset, HZBTestRect))
diff --git a/Engine/Source/Runtime/Renderer/Private/VirtualShadowMaps/VirtualShadowMapCacheManager.cpp b/Engine/Source/Runtime/Renderer/Private/VirtualShadowMaps/VirtualShadowMapCacheManager.cpp
--- a/Engine/Source/Runtime/Renderer/Private/VirtualShadowMaps/VirtualShadowMapCacheManager.cpp
+++ b/Engine/Source/Runtime/Renderer/Private/VirtualShadowMaps/VirtualShadowMapCacheManager.cpp
@@ -90,6 +90,13 @@ FAutoConsoleVariableRef CVarVSMMaxPageAgeSinceLastRequest(
 	ECVF_RenderThreadSafe
 );
 
+static TAutoConsoleVariable<int32> CVarMaxLightAgeSinceLastRequest(
+	TEXT("r.Shadow.Virtual.Cache.MaxLightAgeSinceLastRequest"),
+	10,
+	TEXT("The maximum number of frames to allow lights (and their associated pages) that aren't present in the current frame to live in the cache.\n")
+	TEXT("Larger values can allow pages from offscreen local lights to live longer, but can also increase various page table management overheads."),
+	ECVF_RenderThreadSafe);
+
 static TAutoConsoleVariable<int32> CVarFramesStaticThreshold(
 	TEXT("r.Shadow.Virtual.Cache.StaticSeparate.FramesStaticThreshold"),
 	100,
@@ -950,6 +957,7 @@ void FVirtualShadowMapArrayCacheManager::UpdateUnreferencedCacheEntries(
 	FVirtualShadowMapArray& VirtualShadowMapArray)
 {
 	const uint32 SceneFrameNumber = Scene->GetFrameNumberRenderThread();
+	const int32 MaxLightAge = CVarMaxLightAgeSinceLastRequest.GetValueOnRenderThread();
 
 	for (FEntryMap::TIterator It = CacheEntries.CreateIterator(); It; ++It)
 	{
@@ -960,7 +968,7 @@ void FVirtualShadowMapArrayCacheManager::UpdateUnreferencedCacheEntries(
 			// Active this render, leave it alone
 			check(CacheEntry->ShadowMapEntries.Last().CurrentVirtualShadowMapId < VirtualShadowMapArray.GetNumShadowMapSlots());
 		}
-		else if (int32(SceneFrameNumber - CacheEntry->LastReferencedFrameNumber) <= GVSMMaxPageAgeSinceLastRequest)
+		else if (int32(SceneFrameNumber - CacheEntry->LastReferencedFrameNumber) <= MaxLightAge)
 		{
 			// Not active this render, but still recent enough to keep it and its pages alive
 			int PrevBaseVirtualShadowMapId = CacheEntry->ShadowMapEntries[0].CurrentVirtualShadowMapId;
