Mobile: Added support for SSAO in deferred shading
Added support for renderining velocity in a depth pre-pass (default)

diff --git a/Engine/Shaders/Private/MobileBasePassPixelShader.usf b/Engine/Shaders/Private/MobileBasePassPixelShader.usf
--- a/Engine/Shaders/Private/MobileBasePassPixelShader.usf
+++ b/Engine/Shaders/Private/MobileBasePassPixelShader.usf
@@ -95,14 +95,6 @@
 	#define MOBILE_TRANSLUCENT_COLOR_TRANSMITTANCE_SINGLE_SRC_BLENDING 0
 #endif
 
-// Force FORWARD_SHADING on mobile when MOBILE_DEFERRED_SHADING is disabled
-#ifndef FORWARD_SHADING
-#define FORWARD_SHADING (!MOBILE_DEFERRED_SHADING)
-#elif defined(FORWARD_SHADING)
-#undef  FORWARD_SHADING
-#define FORWARD_SHADING (!MOBILE_DEFERRED_SHADING)
-#endif
-
 // SceneDepthAux always enabled for IOS, on other platforms only enabled when in forward and MobileHDR=true
 #define HAS_SCENE_DEPTH_AUX_OUTPUT (USE_SCENE_DEPTH_AUX && (MATERIALBLENDING_SOLID || MATERIALBLENDING_MASKED) && !MATERIAL_SHADINGMODEL_SINGLELAYERWATER)
 
@@ -1101,7 +1093,7 @@ void Main(
 	OutColor.rgb *= ResolvedView.PreExposure;
 #endif
 
-#if FORWARD_SHADING && MATERIAL_IS_SKY
+#if MATERIAL_IS_SKY
 	if (MobileBasePass.bApplyHalfResLocalFogToSkyMeshes)
 	{
 		// Apply the half resolution texture to the sky dome mesh when needed.
diff --git a/Engine/Shaders/Private/MobileDeferredShading.usf b/Engine/Shaders/Private/MobileDeferredShading.usf
--- a/Engine/Shaders/Private/MobileDeferredShading.usf
+++ b/Engine/Shaders/Private/MobileDeferredShading.usf
@@ -176,6 +176,11 @@ void MobileDirectionalLightPS(
 	}
 #endif
 
+#if ENABLE_AMBIENT_OCCLUSION
+	half4 GatheredAmbientOcclusion = Texture2DSample(MobileBasePass.AmbientOcclusionTexture, MobileBasePass.AmbientOcclusionSampler, UVAndScreenPos.xy);
+	GBuffer.GBufferAO *= lerp(1.0, GatheredAmbientOcclusion.r, MobileBasePass.AmbientOcclusionStaticFraction);
+#endif
+
 	// Directional light
 	half4 DynamicShadowFactors = 1.0f;
 	float DirectionalLightShadow = 1.0f;
@@ -328,6 +333,11 @@ void MobileReflectionEnvironmentSkyLightingPS(
 
 	float2 LocalPosition = SvPosition.xy - View.ViewRectMin.xy;
 	uint GridIndex = ComputeLightGridCellIndex(uint2(LocalPosition.x, LocalPosition.y), GBuffer.Depth);
+
+#if ENABLE_AMBIENT_OCCLUSION
+	half4 GatheredAmbientOcclusion = Texture2DSample(MobileBasePass.AmbientOcclusionTexture, MobileBasePass.AmbientOcclusionSampler, UVAndScreenPos.xy);
+	GBuffer.GBufferAO *= lerp(1.0, GatheredAmbientOcclusion.r, MobileBasePass.AmbientOcclusionStaticFraction);
+#endif
 	
 	FLightAccumulator DirectLighting = (FLightAccumulator)0;
 	ReflectionEnvironmentSkyLighting(GBuffer, CameraVector, TranslatedWorldPosition, ReflectionVector, GridIndex, DirectLighting);
diff --git a/Engine/Source/Runtime/Renderer/Private/DepthRendering.cpp b/Engine/Source/Runtime/Renderer/Private/DepthRendering.cpp
--- a/Engine/Source/Runtime/Renderer/Private/DepthRendering.cpp
+++ b/Engine/Source/Runtime/Renderer/Private/DepthRendering.cpp
@@ -644,7 +644,7 @@ void FDeferredShadingSceneRenderer::RenderPrePass(FRDGBuilder& GraphBuilder, TAr
 bool FMobileSceneRenderer::ShouldRenderPrePass() const
 {
 	// Draw a depth pass to avoid overdraw in the other passes.
-	return Scene->EarlyZPassMode == DDM_MaskedOnly || Scene->EarlyZPassMode == DDM_AllOpaque;
+	return Scene->EarlyZPassMode == DDM_MaskedOnly || Scene->EarlyZPassMode == DDM_AllOpaque || Scene->EarlyZPassMode == DDM_AllOpaqueNoVelocity;
 }
 
 void FMobileSceneRenderer::RenderPrePass(FRHICommandList& RHICmdList, const FViewInfo& View, const FInstanceCullingDrawParams* InstanceCullingDrawParams)
@@ -734,39 +734,6 @@ FMeshDrawCommandSortKey CalculateDepthPassMeshStaticSortKey(bool bIsMasked, cons
 	return SortKey;
 }
 
-void SetMobileDepthPassRenderState(const FPrimitiveSceneProxy* RESTRICT PrimitiveSceneProxy, FMeshPassProcessorRenderState& DrawRenderState, const FMeshBatch& RESTRICT MeshBatch, bool bUsesDeferredShading)
-{
-	DrawRenderState.SetDepthStencilState(TStaticDepthStencilState<
-		true, CF_DepthNearOrEqual,
-		true, CF_Always, SO_Keep, SO_Keep, SO_Replace,
-		false, CF_Always, SO_Keep, SO_Keep, SO_Keep,
-		// don't use masking as it has significant performance hit on Mali GPUs (T860MP2)
-		0x00, 0xff >::GetRHI());
-
-	uint8 StencilValue = 0;
-	
-	uint8 ReceiveDecals = (PrimitiveSceneProxy && !PrimitiveSceneProxy->ReceivesDecals() ? 0x01 : 0x00);
-	StencilValue |= GET_STENCIL_BIT_MASK(RECEIVE_DECAL, ReceiveDecals);
-
-	if (bUsesDeferredShading)
-	{
-		extern uint8 GetMobileShadingModelStencilValue(FMaterialShadingModelField ShadingModel);
-
-		// store into [1-3] bits
-		const FMaterial& MaterialResource = MeshBatch.MaterialRenderProxy->GetIncompleteMaterialWithFallback(ERHIFeatureLevel::ES3_1);
-		uint8 ShadingModel = GetMobileShadingModelStencilValue(MaterialResource.GetShadingModels());
-		StencilValue |= GET_STENCIL_MOBILE_SM_MASK(ShadingModel);
-		StencilValue |= STENCIL_LIGHTING_CHANNELS_MASK(PrimitiveSceneProxy ? PrimitiveSceneProxy->GetLightingChannelStencilValue() : 0x00);
-	}
-	else
-	{
-		uint8 CastContactShadows = (PrimitiveSceneProxy && PrimitiveSceneProxy->CastsContactShadow() ? 0x01 : 0x00);
-		StencilValue |= GET_STENCIL_BIT_MASK(MOBILE_CAST_CONTACT_SHADOW, CastContactShadows);
-	}
-
-	DrawRenderState.SetStencilRef(StencilValue);
-}
-
 template<bool bPositionOnly>
 bool FDepthPassMeshProcessor::Process(
 	const FMeshBatch& RESTRICT MeshBatch,
@@ -808,7 +775,11 @@ bool FDepthPassMeshProcessor::Process(
 	// Use StencilMask for DecalOutput on mobile
 	if (FeatureLevel == ERHIFeatureLevel::ES3_1 && !bShadowProjection)
 	{
-		SetMobileDepthPassRenderState(PrimitiveSceneProxy, DrawRenderState, MeshBatch, IsMobileDeferredShadingEnabled(GetFeatureLevelShaderPlatform(FeatureLevel)));
+		extern void SetMobileBasePassDepthState(FMeshPassProcessorRenderState& DrawRenderState, const FPrimitiveSceneProxy* PrimitiveSceneProxy, FMaterialShadingModelField ShadingModels, bool bUsesDeferredShading);
+		
+		FMaterialShadingModelField ShadingModels = MaterialResource.GetShadingModels();
+		bool bUsesDeferredShading = IsMobileDeferredShadingEnabled(GetFeatureLevelShaderPlatform(FeatureLevel));
+		SetMobileBasePassDepthState(DrawRenderState, PrimitiveSceneProxy, ShadingModels, bUsesDeferredShading);
 	}
 
 	FMeshMaterialShaderElementData ShaderElementData;
diff --git a/Engine/Source/Runtime/Renderer/Private/MobileBasePass.cpp b/Engine/Source/Runtime/Renderer/Private/MobileBasePass.cpp
--- a/Engine/Source/Runtime/Renderer/Private/MobileBasePass.cpp
+++ b/Engine/Source/Runtime/Renderer/Private/MobileBasePass.cpp
@@ -77,6 +77,37 @@ uint8 GetMobileShadingModelStencilValue(FMaterialShadingModelField ShadingModel)
 	return MobileUsesGBufferCustomData(GMaxRHIShaderPlatform) ? 2u : 1u;
 }
 
+void SetMobileBasePassDepthState(FMeshPassProcessorRenderState& DrawRenderState, const FPrimitiveSceneProxy* PrimitiveSceneProxy, FMaterialShadingModelField ShadingModels, bool bUsesDeferredShading)
+{
+	DrawRenderState.SetDepthStencilState(TStaticDepthStencilState<
+		true, CF_DepthNearOrEqual,
+		true, CF_Always, SO_Keep, SO_Keep, SO_Replace,
+		false, CF_Always, SO_Keep, SO_Keep, SO_Keep,
+		// don't use masking as it has significant performance hit on Mali GPUs (T860MP2)
+		0x00, 0xff >::GetRHI());
+
+	uint8 StencilValue = 0u;
+
+	uint8 ReceiveDecals = (PrimitiveSceneProxy && !PrimitiveSceneProxy->ReceivesDecals() ? 0x01 : 0x00);
+	StencilValue |= GET_STENCIL_BIT_MASK(RECEIVE_DECAL, ReceiveDecals);
+
+	if (bUsesDeferredShading)
+	{
+		// store into [1-3] bits
+		uint8 ShadingModel = GetMobileShadingModelStencilValue(ShadingModels);
+		StencilValue |= GET_STENCIL_MOBILE_SM_MASK(ShadingModel);
+		StencilValue |= STENCIL_LIGHTING_CHANNELS_MASK(PrimitiveSceneProxy ? PrimitiveSceneProxy->GetLightingChannelStencilValue() : 0x00);
+	}
+	else
+	{
+		// TODO: ContactShadows do not work with deferred shading atm
+		uint8 CastContactShadows = (PrimitiveSceneProxy && PrimitiveSceneProxy->CastsContactShadow() ? 0x01 : 0x00);
+		StencilValue |= GET_STENCIL_BIT_MASK(MOBILE_CAST_CONTACT_SHADOW, CastContactShadows);
+	}
+
+	DrawRenderState.SetStencilRef(StencilValue); 
+}
+
 bool MobileUsesNoLightMapPermutation(const FMeshMaterialShaderPermutationParameters& Parameters)
 {
 	const bool bAllowStaticLighting = IsStaticLightingAllowed();
@@ -528,37 +559,17 @@ static FMobileLightMapPolicyTypeList GetUniformLightMapPolicyTypeForPSOCollectio
 	return Result;
 }
 
-void MobileBasePass::SetOpaqueRenderState(FMeshPassProcessorRenderState& DrawRenderState, const FPrimitiveSceneProxy* PrimitiveSceneProxy, const FMaterial& Material, FMaterialShadingModelField ShadingModels, bool bEnableReceiveDecalOutput, bool bUsesDeferredShading)
+void MobileBasePass::SetOpaqueRenderState(FMeshPassProcessorRenderState& DrawRenderState, const FPrimitiveSceneProxy* PrimitiveSceneProxy, const FMaterial& Material, FMaterialShadingModelField ShadingModels, bool bCanUseDepthStencil, bool bUsesDeferredShading)
 {
-	uint8 StencilValue = 0;
-	if (bEnableReceiveDecalOutput)
+	if (bCanUseDepthStencil)
 	{
-		uint8 ReceiveDecals = (PrimitiveSceneProxy && !PrimitiveSceneProxy->ReceivesDecals() ? 0x01 : 0x00);
-		StencilValue |= GET_STENCIL_BIT_MASK(RECEIVE_DECAL, ReceiveDecals);
-	}
-	
-	if (bUsesDeferredShading)
-	{
-		uint8 ShadingModel = GetMobileShadingModelStencilValue(ShadingModels);
-		StencilValue |= GET_STENCIL_MOBILE_SM_MASK(ShadingModel);
-		StencilValue |= STENCIL_LIGHTING_CHANNELS_MASK(PrimitiveSceneProxy ? PrimitiveSceneProxy->GetLightingChannelStencilValue() : 0x00);
-	}
-		
-	if (bEnableReceiveDecalOutput || bUsesDeferredShading)
-	{
-		DrawRenderState.SetDepthStencilState(TStaticDepthStencilState<
-				true, CF_DepthNearOrEqual,
-				true, CF_Always, SO_Keep, SO_Keep, SO_Replace,
-				false, CF_Always, SO_Keep, SO_Keep, SO_Keep,
-				// don't use masking as it has significant performance hit on Mali GPUs (T860MP2)
-				0x00, 0xff >::GetRHI());
-
-		DrawRenderState.SetStencilRef(StencilValue); 
+		SetMobileBasePassDepthState(DrawRenderState, PrimitiveSceneProxy, ShadingModels, bUsesDeferredShading);
 	}
 	else
 	{
 		// default depth state should be already set
 	}
+	
 	const bool bIsMasked = IsMaskedBlendMode(Material);
 	if (bIsMasked && Material.IsUsingAlphaToCoverage())
 	{
@@ -941,6 +952,7 @@ bool FMobileBasePassMeshProcessor::Process(
 
 	const bool bMaskedInEarlyPass = (MaterialResource.IsMasked() || MeshBatch.bDitheredLODTransition) && Scene && MaskedInEarlyPass(Scene->GetShaderPlatform());
 	const bool bForcePassDrawRenderState = ((Flags & EFlags::ForcePassDrawRenderState) == EFlags::ForcePassDrawRenderState);
+	const bool bIsFullDepthPrepassEnabled = Scene && (Scene->EarlyZPassMode == DDM_AllOpaque || Scene->EarlyZPassMode == DDM_AllOpaqueNoVelocity);
 
 	FMeshPassProcessorRenderState DrawRenderState(PassDrawRenderState);
 	if (!bForcePassDrawRenderState)
@@ -949,14 +961,14 @@ bool FMobileBasePassMeshProcessor::Process(
 		{
 			MobileBasePass::SetTranslucentRenderState(DrawRenderState, MaterialResource, ShadingModels);
 		}
-		else if((MeshBatch.bUseForDepthPass && Scene->EarlyZPassMode == DDM_AllOpaque) || bMaskedInEarlyPass)
+		else if((MeshBatch.bUseForDepthPass && bIsFullDepthPrepassEnabled) || bMaskedInEarlyPass)
 		{
 			DrawRenderState.SetDepthStencilState(TStaticDepthStencilState<false, CF_Equal>::GetRHI());
 		}
 		else
 		{
-			const bool bEnableReceiveDecalOutput = ((Flags & EFlags::CanUseDepthStencil) == EFlags::CanUseDepthStencil);
-			MobileBasePass::SetOpaqueRenderState(DrawRenderState, PrimitiveSceneProxy, MaterialResource, ShadingModels, bEnableReceiveDecalOutput && IsMobileHDR(), bPassUsesDeferredShading);
+			const bool bCanUseDepthStencil = ((Flags & EFlags::CanUseDepthStencil) == EFlags::CanUseDepthStencil);
+			MobileBasePass::SetOpaqueRenderState(DrawRenderState, PrimitiveSceneProxy, MaterialResource, ShadingModels, bCanUseDepthStencil, bPassUsesDeferredShading);
 		}
 	}
 
diff --git a/Engine/Source/Runtime/Renderer/Private/MobileBasePassRendering.cpp b/Engine/Source/Runtime/Renderer/Private/MobileBasePassRendering.cpp
--- a/Engine/Source/Runtime/Renderer/Private/MobileBasePassRendering.cpp
+++ b/Engine/Source/Runtime/Renderer/Private/MobileBasePassRendering.cpp
@@ -291,24 +291,26 @@ void SetupMobileBasePassUniformParameters(
 	}
 
 	const FRDGSystemTextures& SystemTextures = FRDGSystemTextures::Get(GraphBuilder);
+	const FSceneTextures* SceneTextures = View.GetSceneTexturesChecked();
 
-	SetupMobileSceneTextureUniformParameters(GraphBuilder, View.GetSceneTexturesChecked(), SetupMode, BasePassParameters.SceneTextures);
+	SetupMobileSceneTextureUniformParameters(GraphBuilder, SceneTextures, SetupMode, BasePassParameters.SceneTextures);
 
 	BasePassParameters.PreIntegratedGFTexture = GSystemTextures.PreintegratedGF->GetRHI();
 	BasePassParameters.PreIntegratedGFSampler = TStaticSamplerState<SF_Bilinear, AM_Clamp, AM_Clamp, AM_Clamp>::GetRHI();
 	BasePassParameters.EyeAdaptationBuffer = GraphBuilder.CreateSRV(GetEyeAdaptationBuffer(GraphBuilder, View));
 
-	FRDGTextureRef AmbientOcclusionTexture = SystemTextures.White;
-	if (BasePass == EMobileBasePass::Opaque && MobileBasePassTextures.ScreenSpaceAO != nullptr)
+	if (SceneTextures && HasBeenProduced(SceneTextures->ScreenSpaceAO))
 	{
-		AmbientOcclusionTexture = MobileBasePassTextures.ScreenSpaceAO;
+		BasePassParameters.AmbientOcclusionTexture = SceneTextures->ScreenSpaceAO;
+	}
+	else
+	{
+		BasePassParameters.AmbientOcclusionTexture = SystemTextures.White;
 	}
-	
-	BasePassParameters.DBuffer = GetDBufferParameters(GraphBuilder, MobileBasePassTextures.DBufferTextures, View.GetShaderPlatform());
-
-	BasePassParameters.AmbientOcclusionTexture = AmbientOcclusionTexture;
 	BasePassParameters.AmbientOcclusionSampler = TStaticSamplerState<SF_Bilinear, AM_Clamp, AM_Clamp, AM_Clamp>::GetRHI();
 	BasePassParameters.AmbientOcclusionStaticFraction = FMath::Clamp(View.FinalPostProcessSettings.AmbientOcclusionStaticFraction, 0.0f, 1.0f);
+	
+	BasePassParameters.DBuffer = GetDBufferParameters(GraphBuilder, MobileBasePassTextures.DBufferTextures, View.GetShaderPlatform());
 
 	const bool bMobileUsesShadowMaskTexture = MobileUsesShadowMaskTexture(View.GetShaderPlatform());
 	
diff --git a/Engine/Source/Runtime/Renderer/Private/MobileBasePassRendering.h b/Engine/Source/Runtime/Renderer/Private/MobileBasePassRendering.h
--- a/Engine/Source/Runtime/Renderer/Private/MobileBasePassRendering.h
+++ b/Engine/Source/Runtime/Renderer/Private/MobileBasePassRendering.h
@@ -35,7 +35,6 @@ bool MobileMergeLocalLightsInBasepassEnabled(const FStaticShaderPlatform Platfor
 
 struct FMobileBasePassTextures
 {
-	FRDGTextureRef ScreenSpaceAO = nullptr;
 	FDBufferTextures DBufferTextures = {};
 };
 
@@ -420,7 +419,7 @@ public:
 
 		TMobileBasePassPSBaseType<LightMapPolicyType>::ModifyCompilationEnvironment(Parameters, OutEnvironment);
 		OutEnvironment.SetDefine(TEXT("ENABLE_SKY_LIGHT"), bEnableSkylightInBasePass);
-		OutEnvironment.SetDefine(TEXT("ENABLE_AMBIENT_OCCLUSION"), IsMobileAmbientOcclusionEnabled(Parameters.Platform) ? 1u : 0u);
+		OutEnvironment.SetDefine(TEXT("ENABLE_AMBIENT_OCCLUSION"), bForwardShading && IsMobileAmbientOcclusionEnabled(Parameters.Platform) ? 1u : 0u);
 		
 		FForwardLightingParameters::ModifyCompilationEnvironment(Parameters.Platform, OutEnvironment);
 		OutEnvironment.SetDefine(TEXT("ENABLE_CLUSTERED_LIGHTS"), (LocalLightSetting == EMobileLocalLightSetting::LOCAL_LIGHTS_ENABLED) ? 1u : 0u);
diff --git a/Engine/Source/Runtime/Renderer/Private/MobileDeferredShadingPass.cpp b/Engine/Source/Runtime/Renderer/Private/MobileDeferredShadingPass.cpp
--- a/Engine/Source/Runtime/Renderer/Private/MobileDeferredShadingPass.cpp
+++ b/Engine/Source/Runtime/Renderer/Private/MobileDeferredShadingPass.cpp
@@ -12,6 +12,8 @@
 #include "LocalLightSceneProxy.h"
 #include "Materials/MaterialRenderProxy.h"
 
+DECLARE_GPU_STAT(DeferredShading);
+
 int32 GMobileUseClusteredDeferredShading = 0;
 static FAutoConsoleVariableRef CVarMobileUseClusteredDeferredShading(
 	TEXT("r.Mobile.UseClusteredDeferredShading"),
@@ -88,6 +90,7 @@ class FMobileDirectionalLightFunctionPS : public FMaterialShader
 		FForwardLightingParameters::ModifyCompilationEnvironment(Parameters.Platform, OutEnvironment);
 		OutEnvironment.SetDefine(TEXT("USE_LIGHT_FUNCTION"), Parameters.MaterialParameters.bIsDefaultMaterial ? 0 : 1);
 		OutEnvironment.SetDefine(TEXT("USE_SHADOWMASKTEXTURE"), MobileUsesShadowMaskTexture(Parameters.Platform) ? 1u : 0u);
+		OutEnvironment.SetDefine(TEXT("ENABLE_AMBIENT_OCCLUSION"), IsMobileAmbientOcclusionEnabled(Parameters.Platform) ? 1u : 0u);
 		OutEnvironment.SetDefine(TEXT("MATERIAL_SHADER"), 1);
 		OutEnvironment.SetDefine(TEXT("IS_MOBILE_DEFERREDSHADING_SUBPASS"), 1u);
 
@@ -314,6 +317,7 @@ public:
 
 		const bool bMobileForceDepthRead = MobileUsesFullDepthPrepass(Parameters.Platform);
 		OutEnvironment.SetDefine(TEXT("FORCE_DEPTH_TEXTURE_READS"), bMobileForceDepthRead ? 1u : 0u);
+		OutEnvironment.SetDefine(TEXT("ENABLE_AMBIENT_OCCLUSION"), IsMobileAmbientOcclusionEnabled(Parameters.Platform) ? 1u : 0u);
 	}
 };
 
@@ -954,6 +958,7 @@ void MobileDeferredShadingPass(
 	const TArray<FVisibleLightInfo, SceneRenderingAllocator>& VisibleLightInfos)
 {
 	SCOPED_DRAW_EVENT(RHICmdList, DeferredShading);
+	SCOPED_GPU_STAT(RHICmdList, DeferredShading);
 	
 	RHICmdList.SetViewport(View.ViewRect.Min.X, View.ViewRect.Min.Y, 0.0f, View.ViewRect.Max.X, View.ViewRect.Max.Y, 1.0f);
 
diff --git a/Engine/Source/Runtime/Renderer/Private/MobileShadingRenderer.cpp b/Engine/Source/Runtime/Renderer/Private/MobileShadingRenderer.cpp
--- a/Engine/Source/Runtime/Renderer/Private/MobileShadingRenderer.cpp
+++ b/Engine/Source/Runtime/Renderer/Private/MobileShadingRenderer.cpp
@@ -192,7 +192,7 @@ static void RenderOpaqueFX(
 	}
 }
 
-BEGIN_SHADER_PARAMETER_STRUCT(FMobileRenderPassParameters, RENDERER_API)
+BEGIN_SHADER_PARAMETER_STRUCT(FMobileRenderPassParameters,)
 	SHADER_PARAMETER_STRUCT_INCLUDE(FViewShaderParameters, View)
 	SHADER_PARAMETER_STRUCT_INCLUDE(FInstanceCullingDrawParams, InstanceCullingDrawParams)
 	SHADER_PARAMETER_RDG_UNIFORM_BUFFER(FMobileBasePassUniformParameters, MobileBasePass)
@@ -287,7 +287,7 @@ FMobileSceneRenderer::FMobileSceneRenderer(const FSceneViewFamily* InViewFamily,
 	bRequiresPixelProjectedPlanarRelfectionPass = false;
 	bRequiresAmbientOcclusionPass = false;
 	bRequiresShadowProjections = false;
-	bIsFullDepthPrepassEnabled = Scene->EarlyZPassMode == DDM_AllOpaque;
+	bIsFullDepthPrepassEnabled = (Scene->EarlyZPassMode == DDM_AllOpaque || Scene->EarlyZPassMode == DDM_AllOpaqueNoVelocity);
 	bIsMaskedOnlyDepthPrepassEnabled = Scene->EarlyZPassMode == DDM_MaskedOnly;
 	bEnableClusteredLocalLights = MobileForwardEnableLocalLights(ShaderPlatform);
 	bEnableClusteredReflections = MobileForwardEnableClusteredReflections(ShaderPlatform);
@@ -784,51 +784,76 @@ static void EndOcclusionScope(FRDGBuilder& GraphBuilder, TArray<FViewInfo>& View
 */
 void FMobileSceneRenderer::RenderFullDepthPrepass(FRDGBuilder& GraphBuilder, TArrayView<FViewInfo> InViews, FSceneTextures& SceneTextures, bool bIsSceneCaptureRenderPass)
 {
-	FRenderTargetBindingSlots BasePassRenderTargets;
-	BasePassRenderTargets.DepthStencil = FDepthStencilBinding(SceneTextures.Depth.Target, ERenderTargetLoadAction::EClear, ERenderTargetLoadAction::EClear, FExclusiveDepthStencil::DepthWrite_StencilWrite);
-	// If this is scene capture render pass, don't render occlusion
-	BasePassRenderTargets.NumOcclusionQueries = bIsSceneCaptureRenderPass ? 0 : ComputeNumOcclusionQueriesToBatch();
-
 	FRenderViewContextArray RenderViews;
 	GetRenderViews(InViews, RenderViews);
 
+	TRDGUniformBufferBinding<FMobileBasePassUniformParameters> LastViewMobileBasePassUB;
+
 	for (FRenderViewContext& ViewContext : RenderViews)
 	{
 		FViewInfo& View = *ViewContext.ViewInfo;
 
-		if (!ViewContext.bIsFirstView)
-		{
-			BasePassRenderTargets.DepthStencil.SetDepthLoadAction(ERenderTargetLoadAction::ELoad);
-			BasePassRenderTargets.DepthStencil.SetStencilLoadAction(ERenderTargetLoadAction::ELoad);
-			BasePassRenderTargets.DepthStencil.SetDepthStencilAccess(FExclusiveDepthStencil::DepthWrite_StencilWrite);
-		}
-
 		View.BeginRenderView();
 
 		auto* PassParameters = GraphBuilder.AllocParameters<FMobileRenderPassParameters>();
-		PassParameters->RenderTargets = BasePassRenderTargets;
+		PassParameters->RenderTargets.DepthStencil = FDepthStencilBinding(SceneTextures.Depth.Target, ERenderTargetLoadAction::EClear, ERenderTargetLoadAction::EClear, FExclusiveDepthStencil::DepthWrite_StencilWrite);
 		PassParameters->View = View.GetShaderParameters();
 		PassParameters->MobileBasePass = CreateMobileBasePassUniformBuffer(GraphBuilder, View, EMobileBasePass::DepthPrePass, EMobileSceneTextureSetupMode::None);
-		
-		View.ParallelMeshDrawCommandPasses[EMeshPass::DepthPass].BuildRenderingCommands(GraphBuilder, Scene->GPUScene, PassParameters->InstanceCullingDrawParams);
 
-		// Render occlusion at the last view pass only, as they already loop through all views
-		// If this is scene capture render pass, don't render occlusion.
-		bool bDoOcclusionQueries = (ViewContext.bIsLastView && DoOcclusionQueries() && !bIsSceneCaptureRenderPass);
+		if (ViewContext.bIsLastView)
+		{
+			LastViewMobileBasePassUB = PassParameters->MobileBasePass;
+		}
+
+		if (!ViewContext.bIsFirstView)
+		{
+			PassParameters->RenderTargets.DepthStencil.SetDepthLoadAction(ERenderTargetLoadAction::ELoad);
+			PassParameters->RenderTargets.DepthStencil.SetStencilLoadAction(ERenderTargetLoadAction::ELoad);
+			PassParameters->RenderTargets.DepthStencil.SetDepthStencilAccess(FExclusiveDepthStencil::DepthWrite_StencilWrite);
+		}
+
+		View.ParallelMeshDrawCommandPasses[EMeshPass::DepthPass].BuildRenderingCommands(GraphBuilder, Scene->GPUScene, PassParameters->InstanceCullingDrawParams);
 
 		GraphBuilder.AddPass(
 			RDG_EVENT_NAME("FullDepthPrepass"),
 			PassParameters,
 			ERDGPassFlags::Raster,
-			[this, PassParameters, &View, bDoOcclusionQueries](FRHICommandList& RHICmdList)
+			[this, PassParameters, &View](FRHICommandList& RHICmdList)
 			{
 				RenderPrePass(RHICmdList, View, &PassParameters->InstanceCullingDrawParams);
+			});
+	}
+
+	if (bShouldRenderVelocities && Scene->EarlyZPassMode == DDM_AllOpaqueNoVelocity)
+	{
+		// Render the velocities and depth of movable objects
+		RenderVelocities(GraphBuilder, InViews, SceneTextures, EVelocityPass::Opaque, false);
+	}
+
+	for (FRenderViewContext& ViewContext : RenderViews)
+	{
+		FViewInfo& View = *ViewContext.ViewInfo;
 
-				if (bDoOcclusionQueries)
+		// Render occlusion at the last view pass only, as they already loop through all views
+		// If this is scene capture render pass, don't render occlusion.
+		bool bDoOcclusionQueries = (ViewContext.bIsLastView && DoOcclusionQueries() && !bIsSceneCaptureRenderPass);
+		if (bDoOcclusionQueries)
+		{
+			auto* PassParameters = GraphBuilder.AllocParameters<FMobileRenderPassParameters>();
+			PassParameters->RenderTargets.DepthStencil = FDepthStencilBinding(SceneTextures.Depth.Target, ERenderTargetLoadAction::ELoad, ERenderTargetLoadAction::ELoad, FExclusiveDepthStencil::DepthRead_StencilRead);
+			PassParameters->View = View.GetShaderParameters();
+			PassParameters->MobileBasePass = LastViewMobileBasePassUB;
+			PassParameters->RenderTargets.NumOcclusionQueries = ComputeNumOcclusionQueriesToBatch();
+
+			GraphBuilder.AddPass(
+				RDG_EVENT_NAME("RenderOcclusion"),
+				PassParameters,
+				ERDGPassFlags::Raster | ERDGPassFlags::NeverCull,
+				[this](FRHICommandList& RHICmdList)
 				{
 					RenderOcclusion(RHICmdList);
-				}
-			});
+				});
+		}
 	}
 
 	FenceOcclusionTests(GraphBuilder);
@@ -1039,7 +1064,8 @@ void FMobileSceneRenderer::Render(FRDGBuilder& GraphBuilder)
 	SceneTextures.MobileSetupMode = EMobileSceneTextureSetupMode::None;
 	SceneTextures.MobileUniformBuffer = CreateMobileSceneTextureUniformBuffer(GraphBuilder, &SceneTextures, SceneTextures.MobileSetupMode);
 
-	const bool bUseHalfResLocalFogVolume = bIsFullDepthPrepassEnabled && IsLocalFogVolumeHalfResolution(); // We must have a full depth buffer in order to render half res and upsample
+	// We must have a full depth buffer in order to render half res and upsample
+	const bool bUseHalfResLocalFogVolume = bIsFullDepthPrepassEnabled && IsLocalFogVolumeHalfResolution();
 
 	if (bRendererOutputFinalSceneColor)
 	{
@@ -1252,14 +1278,6 @@ void FMobileSceneRenderer::Render(FRDGBuilder& GraphBuilder)
 		SceneTextures.MobileSetupMode = EMobileSceneTextureSetupMode::SceneDepth;
 		SceneTextures.MobileUniformBuffer = CreateMobileSceneTextureUniformBuffer(GraphBuilder, &SceneTextures, SceneTextures.MobileSetupMode);
 
-		// When renderer is in ERendererOutput::DepthPrepassOnly mode, bRequiresShadowProjections is set to false in InitViews()
-		if (bRequiresShadowProjections)
-		{
-			RDG_CSV_STAT_EXCLUSIVE_SCOPE(GraphBuilder, RenderMobileShadowProjections);
-			RDG_GPU_STAT_SCOPE(GraphBuilder, ShadowProjection);
-			RenderMobileShadowProjections(GraphBuilder);
-		}
-
 		// When renderer is in ERendererOutput::DepthPrepassOnly mode, bShouldRenderHZB is set to false in InitViews()
 		if (bShouldRenderHZB)
 		{
@@ -1272,6 +1290,14 @@ void FMobileSceneRenderer::Render(FRDGBuilder& GraphBuilder)
 			RenderAmbientOcclusion(GraphBuilder, SceneTextures.Depth.Resolve, SceneTextures.ScreenSpaceAO);
 		}
 
+		// When renderer is in ERendererOutput::DepthPrepassOnly mode, bRequiresShadowProjections is set to false in InitViews()
+		if (bRequiresShadowProjections)
+		{
+			RDG_CSV_STAT_EXCLUSIVE_SCOPE(GraphBuilder, RenderMobileShadowProjections);
+			RDG_GPU_STAT_SCOPE(GraphBuilder, ShadowProjection);
+			RenderMobileShadowProjections(GraphBuilder);
+		}
+
 		// Local Light prepass
 
 		if (bRendererOutputFinalSceneColor)
@@ -1307,7 +1333,14 @@ void FMobileSceneRenderer::Render(FRDGBuilder& GraphBuilder)
 	{
 		if (bDeferredShading)
 		{
-			RenderDeferred(GraphBuilder, SortedLightSet, ViewFamilyTexture, SceneTextures);
+			if (bRequiresMultiPass)
+			{
+				RenderDeferredMultiPass(GraphBuilder, SceneTextures, SortedLightSet);
+			}
+			else
+			{
+				RenderDeferredSinglePass(GraphBuilder, SceneTextures, SortedLightSet);
+			}
 		}
 		else
 		{
@@ -1328,7 +1361,10 @@ void FMobileSceneRenderer::Render(FRDGBuilder& GraphBuilder)
 		if (bShouldRenderVelocities)
 		{
 			// Render the velocities of movable objects
-			RenderVelocities(GraphBuilder, Views, SceneTextures, EVelocityPass::Opaque, false);
+			if (Scene->EarlyZPassMode != DDM_AllOpaqueNoVelocity)
+			{ 
+				RenderVelocities(GraphBuilder, Views, SceneTextures, EVelocityPass::Opaque, false);
+			}
 			RenderVelocities(GraphBuilder, Views, SceneTextures, EVelocityPass::Translucent, false);
 
 			SceneTextures.MobileSetupMode = EMobileSceneTextureSetupMode::All;
@@ -1538,7 +1574,6 @@ void FMobileSceneRenderer::RenderForward(FRDGBuilder& GraphBuilder, FRDGTextureR
 		UpdateDirectionalLightUniformBuffers(GraphBuilder, View);
 
 		FMobileBasePassTextures MobileBasePassTextures{};
-		MobileBasePassTextures.ScreenSpaceAO = bRequiresAmbientOcclusionPass ? SceneTextures.ScreenSpaceAO : SystemTextures.White;
 		MobileBasePassTextures.DBufferTextures = DBufferTextures;
 
 		EMobileSceneTextureSetupMode SetupMode = (bIsFullDepthPrepassEnabled ? EMobileSceneTextureSetupMode::SceneDepth : EMobileSceneTextureSetupMode::None) | EMobileSceneTextureSetupMode::CustomDepth;
@@ -1873,13 +1908,13 @@ FRenderTargetBindingSlots FMobileSceneRenderer::InitRenderTargetBindings_Deferre
 	return BasePassRenderTargets;
 }
 
-void FMobileSceneRenderer::RenderDeferred(FRDGBuilder& GraphBuilder, const FSortedLightSetSceneInfo& SortedLightSet, FRDGTextureRef ViewFamilyTexture, FSceneTextures& SceneTextures)
+void FMobileSceneRenderer::RenderDeferredSinglePass(FRDGBuilder& GraphBuilder, FSceneTextures& SceneTextures, const FSortedLightSetSceneInfo& SortedLightSet)
 {
+	bool bUsingPixelLocalStorage = UsingPixelLocalStorage(ShaderPlatform);
 	FColorTargets ColorTargets = GetColorTargets_Deferred(SceneTextures);
 	FRenderTargetBindingSlots BasePassRenderTargets = InitRenderTargetBindings_Deferred(SceneTextures, ColorTargets);
-
 	const FRDGSystemTextures& SystemTextures = FRDGSystemTextures::Get(GraphBuilder);
-
+	
 	FRenderViewContextArray RenderViews;
 	GetRenderViews(Views, RenderViews);
 
@@ -1906,157 +1941,222 @@ void FMobileSceneRenderer::RenderDeferred(FRDGBuilder& GraphBuilder, const FSort
 
 		UpdateDirectionalLightUniformBuffers(GraphBuilder, View);
 
-		FMobileBasePassTextures MobileBasePassTextures{};
-		MobileBasePassTextures.ScreenSpaceAO = bRequiresAmbientOcclusionPass ? SceneTextures.ScreenSpaceAO : SystemTextures.White;
-
 		EMobileSceneTextureSetupMode SetupMode = (bIsFullDepthPrepassEnabled ? EMobileSceneTextureSetupMode::SceneDepth : EMobileSceneTextureSetupMode::None) | EMobileSceneTextureSetupMode::CustomDepth;
 		auto* PassParameters = GraphBuilder.AllocParameters<FMobileRenderPassParameters>();
 		PassParameters->View = View.GetShaderParameters();
-		PassParameters->MobileBasePass = CreateMobileBasePassUniformBuffer(GraphBuilder, View, EMobileBasePass::Opaque, SetupMode, MobileBasePassTextures);
+		PassParameters->MobileBasePass = CreateMobileBasePassUniformBuffer(GraphBuilder, View, EMobileBasePass::Opaque, SetupMode);
 		PassParameters->ReflectionCapture = View.MobileReflectionCaptureUniformBuffer;
-		PassParameters->RenderTargets = BasePassRenderTargets;
 		PassParameters->LocalFogVolumeInstances = View.LocalFogVolumeViewData.GPUInstanceDataBufferSRV;
 		PassParameters->LocalFogVolumeTileDrawIndirectBuffer = View.LocalFogVolumeViewData.GPUTileDrawIndirectBuffer;
 		PassParameters->LocalFogVolumeTileDataTexture = View.LocalFogVolumeViewData.TileDataTextureArraySRV;
 		PassParameters->LocalFogVolumeTileDataBuffer = View.LocalFogVolumeViewData.GPUTileDataBufferSRV;
 		PassParameters->HalfResLocalFogVolumeViewSRV = View.LocalFogVolumeViewData.HalfResLocalFogVolumeViewSRV;
 		PassParameters->HalfResLocalFogVolumeDepthSRV = View.LocalFogVolumeViewData.HalfResLocalFogVolumeDepthSRV;
-		
+		PassParameters->RenderTargets = BasePassRenderTargets;
+		PassParameters->RenderTargets.SubpassHint = ESubpassHint::DeferredShadingSubpass;
+		const bool bDoOcclusionQueires = (!bIsFullDepthPrepassEnabled && ViewContext.bIsLastView && DoOcclusionQueries());
+		PassParameters->RenderTargets.NumOcclusionQueries = bDoOcclusionQueires ? ComputeNumOcclusionQueriesToBatch() : 0u;
+
 		BuildInstanceCullingDrawParams(GraphBuilder, View, PassParameters);
 
-		if (bRequiresMultiPass)
-		{
-			RenderDeferredMultiPass(GraphBuilder, PassParameters, ColorTargets.Num(), ViewContext, SceneTextures, SortedLightSet);
-		}
-		else
-		{
-			RenderDeferredSinglePass(GraphBuilder, PassParameters, ViewContext, SceneTextures, SortedLightSet, UsingPixelLocalStorage(ShaderPlatform));
-		}
+		GraphBuilder.AddPass(
+			RDG_EVENT_NAME("SceneColorRendering"),
+			PassParameters,
+			// the second view pass should not be merged with the first view pass on mobile since the subpass would not work properly.
+			ERDGPassFlags::Raster | ERDGPassFlags::NeverMerge,
+			[this, PassParameters, ViewContext, bDoOcclusionQueires, &SceneTextures, &SortedLightSet, bUsingPixelLocalStorage](FRHICommandList& RHICmdList)
+			{
+				FViewInfo& View = *ViewContext.ViewInfo;
+
+				// Depth pre-pass
+				RenderMaskedPrePass(RHICmdList, View);
+				// Opaque and masked
+				RenderMobileBasePass(RHICmdList, View, &PassParameters->InstanceCullingDrawParams);
+				RenderMobileDebugView(RHICmdList, View);
+				RHICmdList.PollOcclusionQueries();
+				PostRenderBasePass(RHICmdList, View);
+				// SceneColor + GBuffer write, SceneDepth is read only
+				RHICmdList.NextSubpass();
+				RenderDecals(RHICmdList, View);
+				// SceneColor write, SceneDepth is read only
+				RHICmdList.NextSubpass();
+				MobileDeferredShadingPass(RHICmdList, ViewContext.ViewIndex, Views.Num(), View, *Scene, SortedLightSet, VisibleLightInfos);
+				if (bUsingPixelLocalStorage)
+				{
+					MobileDeferredCopyBuffer<FMobileDeferredCopyPLSPS>(RHICmdList, View);
+				}
+				RenderFog(RHICmdList, View);
+				// Draw translucency.
+				RenderTranslucency(RHICmdList, View);
+
+				if (bDoOcclusionQueires)
+				{
+					// Issue occlusion queries
+					RenderOcclusion(RHICmdList);
+				}
+			});
 	}
 }
 
-void FMobileSceneRenderer::RenderDeferredSinglePass(FRDGBuilder& GraphBuilder, class FMobileRenderPassParameters* PassParameters, FRenderViewContext& ViewContext, FSceneTextures& SceneTextures, const FSortedLightSetSceneInfo& SortedLightSet, bool bUsingPixelLocalStorage)
+void FMobileSceneRenderer::RenderDeferredMultiPass(FRDGBuilder& GraphBuilder, FSceneTextures& SceneTextures, const FSortedLightSetSceneInfo& SortedLightSet)
 {
-	PassParameters->RenderTargets.SubpassHint = ESubpassHint::DeferredShadingSubpass;
+	FColorTargets ColorTargets = GetColorTargets_Deferred(SceneTextures);
+	FRenderTargetBindingSlots BasePassRenderTargets = InitRenderTargetBindings_Deferred(SceneTextures, ColorTargets);
+	const FRDGSystemTextures& SystemTextures = FRDGSystemTextures::Get(GraphBuilder);
 
-	const bool bDoOcclusionQueires = (!bIsFullDepthPrepassEnabled && ViewContext.bIsLastView && DoOcclusionQueries());
-	PassParameters->RenderTargets.NumOcclusionQueries = bDoOcclusionQueires ? ComputeNumOcclusionQueriesToBatch() : 0u;
-				
-	GraphBuilder.AddPass(
-		RDG_EVENT_NAME("SceneColorRendering"),
-		PassParameters,
-		// the second view pass should not be merged with the first view pass on mobile since the subpass would not work properly.
-		ERDGPassFlags::Raster | ERDGPassFlags::NeverMerge,
-		[this, PassParameters, ViewContext, bDoOcclusionQueires, &SceneTextures, &SortedLightSet, bUsingPixelLocalStorage](FRHICommandList& RHICmdList)
+	FRenderViewContextArray RenderViews;
+	GetRenderViews(Views, RenderViews);
+
+	for (FRenderViewContext& ViewContext : RenderViews)
 	{
 		FViewInfo& View = *ViewContext.ViewInfo;
-			
-		// Depth pre-pass
-		RenderMaskedPrePass(RHICmdList, View);
-		// Opaque and masked
-		RenderMobileBasePass(RHICmdList, View, &PassParameters->InstanceCullingDrawParams);
-		RenderMobileDebugView(RHICmdList, View);
-		RHICmdList.PollOcclusionQueries();
-		PostRenderBasePass(RHICmdList, View);
-		// SceneColor + GBuffer write, SceneDepth is read only
-		RHICmdList.NextSubpass();
-		RenderDecals(RHICmdList, View);
-		// SceneColor write, SceneDepth is read only
-		RHICmdList.NextSubpass();
-		MobileDeferredShadingPass(RHICmdList, ViewContext.ViewIndex, Views.Num(), View, *Scene, SortedLightSet, VisibleLightInfos);
-		if (bUsingPixelLocalStorage)
+
+		SCOPED_GPU_MASK(GraphBuilder.RHICmdList, !View.IsInstancedStereoPass() ? View.GPUMask : (View.GPUMask | View.GetInstancedView()->GPUMask));
+		SCOPED_CONDITIONAL_DRAW_EVENTF(GraphBuilder.RHICmdList, EventView, RenderViews.Num() > 1, TEXT("View%d"), ViewContext.ViewIndex);
+		
+		View.BeginRenderView();
+
+		auto* PassParameters = GraphBuilder.AllocParameters<FMobileRenderPassParameters>();
+		PassParameters->View = View.GetShaderParameters();
+		EMobileSceneTextureSetupMode SetupMode = bIsFullDepthPrepassEnabled ? EMobileSceneTextureSetupMode::SceneDepth : EMobileSceneTextureSetupMode::None;
+		PassParameters->MobileBasePass = CreateMobileBasePassUniformBuffer(GraphBuilder, View, EMobileBasePass::Opaque, SetupMode);
+		PassParameters->RenderTargets = BasePassRenderTargets;
+		if (!ViewContext.bIsFirstView)
 		{
-			MobileDeferredCopyBuffer<FMobileDeferredCopyPLSPS>(RHICmdList, View);
+			// Load targets for a non-first view 
+			for (int32 i = 0; i < ColorTargets.Num(); ++i)
+			{
+				PassParameters->RenderTargets[i].SetLoadAction(ERenderTargetLoadAction::ELoad);
+			}
+			PassParameters->RenderTargets.DepthStencil.SetDepthLoadAction(ERenderTargetLoadAction::ELoad);
+			PassParameters->RenderTargets.DepthStencil.SetStencilLoadAction(ERenderTargetLoadAction::ELoad);
+			PassParameters->RenderTargets.DepthStencil.SetDepthStencilAccess(bIsFullDepthPrepassEnabled ? FExclusiveDepthStencil::DepthRead_StencilWrite : FExclusiveDepthStencil::DepthWrite_StencilWrite);
 		}
-		RenderFog(RHICmdList, View);
-		// Draw translucency.
-		RenderTranslucency(RHICmdList, View);
 
-		if (bDoOcclusionQueires)
+		const bool bDoOcclusionQueires = (!bIsFullDepthPrepassEnabled && ViewContext.bIsLastView && DoOcclusionQueries());
+		PassParameters->RenderTargets.NumOcclusionQueries = bDoOcclusionQueires ? ComputeNumOcclusionQueriesToBatch() : 0u;
+
+		if (!bIsFullDepthPrepassEnabled)
 		{
-			// Issue occlusion queries
-			RenderOcclusion(RHICmdList);
+			View.ParallelMeshDrawCommandPasses[EMeshPass::DepthPass].BuildRenderingCommands(GraphBuilder, Scene->GPUScene, DepthPassInstanceCullingDrawParams);
 		}
-	});
-}
+		View.ParallelMeshDrawCommandPasses[EMeshPass::BasePass].BuildRenderingCommands(GraphBuilder, Scene->GPUScene, PassParameters->InstanceCullingDrawParams);
+		View.ParallelMeshDrawCommandPasses[EMeshPass::SkyPass].BuildRenderingCommands(GraphBuilder, Scene->GPUScene, SkyPassInstanceCullingDrawParams);
+		View.ParallelMeshDrawCommandPasses[EMeshPass::DebugViewMode].BuildRenderingCommands(GraphBuilder, Scene->GPUScene, DebugViewModeInstanceCullingDrawParams);
+		View.ParallelMeshDrawCommandPasses[EMeshPass::MeshDecal_SceneColorAndGBuffer].BuildRenderingCommands(GraphBuilder, Scene->GPUScene, MeshDecalSceneColorAndGBufferInstanceCullingDrawParams);
 
-void FMobileSceneRenderer::RenderDeferredMultiPass(FRDGBuilder& GraphBuilder, class FMobileRenderPassParameters* PassParameters, int32 NumColorTargets, FRenderViewContext& ViewContext, FSceneTextures& SceneTextures, const FSortedLightSetSceneInfo& SortedLightSet)
-{
-	GraphBuilder.AddPass(
-		RDG_EVENT_NAME("SceneColorRendering"),
-		PassParameters,
-		ERDGPassFlags::Raster,
-		[this, PassParameters, ViewContext, &SceneTextures](FRHICommandList& RHICmdList)
-	{
-		FViewInfo& View = *ViewContext.ViewInfo;
-		
-		// Depth pre-pass
-		RenderMaskedPrePass(RHICmdList, View);
-		// Opaque and masked
-		RenderMobileBasePass(RHICmdList, View, &PassParameters->InstanceCullingDrawParams);
-		RenderMobileDebugView(RHICmdList, View);
-		RHICmdList.PollOcclusionQueries();
-		PostRenderBasePass(RHICmdList, View);
-	});
+		GraphBuilder.AddPass(
+			RDG_EVENT_NAME("BasePass"),
+			PassParameters,
+			ERDGPassFlags::Raster,
+			[this, PassParameters, ViewContext, &SceneTextures, bDoOcclusionQueires](FRHICommandList& RHICmdList)
+			{
+				FViewInfo& View = *ViewContext.ViewInfo;
 
-	FViewInfo& View = *ViewContext.ViewInfo;
-	
-	EMobileSceneTextureSetupMode SetupMode = EMobileSceneTextureSetupMode::SceneDepth | EMobileSceneTextureSetupMode::SceneDepthAux | EMobileSceneTextureSetupMode::CustomDepth;
-	auto* SecondPassParameters = GraphBuilder.AllocParameters<FMobileRenderPassParameters>();
-	*SecondPassParameters = *PassParameters;
-	SecondPassParameters->MobileBasePass = CreateMobileBasePassUniformBuffer(GraphBuilder, View, EMobileBasePass::Translucent, SetupMode);
-	SecondPassParameters->ReflectionCapture = View.MobileReflectionCaptureUniformBuffer;
-	// SceneColor + GBuffer write, SceneDepth is read only
-	for (int32 i = 0; i < NumColorTargets; ++i)
-	{
-		SecondPassParameters->RenderTargets[i].SetLoadAction(ERenderTargetLoadAction::ELoad);
+				// Depth pre-pass
+				RenderMaskedPrePass(RHICmdList, View);
+				// Opaque and masked
+				RenderMobileBasePass(RHICmdList, View, &PassParameters->InstanceCullingDrawParams);
+				RenderMobileDebugView(RHICmdList, View);
+				RHICmdList.PollOcclusionQueries();
+				PostRenderBasePass(RHICmdList, View);
+
+				if (bDoOcclusionQueires)
+				{
+					// Issue occlusion queries
+					RenderOcclusion(RHICmdList);
+				}
+
+				if (bIsFullDepthPrepassEnabled)
+				{
+					RenderDecals(RHICmdList, View);
+				}
+			});
 	}
-	SecondPassParameters->RenderTargets.DepthStencil.SetDepthLoadAction(ERenderTargetLoadAction::ELoad);
-	SecondPassParameters->RenderTargets.DepthStencil.SetStencilLoadAction(ERenderTargetLoadAction::ELoad);
-	SecondPassParameters->RenderTargets.DepthStencil.SetDepthStencilAccess(FExclusiveDepthStencil::DepthRead_StencilRead);
 
-	GraphBuilder.AddPass(
-		RDG_EVENT_NAME("Decals"),
-		SecondPassParameters,
-		ERDGPassFlags::Raster,
-		[this, SecondPassParameters, ViewContext](FRHICommandList& RHICmdList)
-	{
-		FViewInfo& View = *ViewContext.ViewInfo;
-		RenderDecals(RHICmdList, View);
-	});
+	BasePassRenderTargets.Enumerate([](FRenderTargetBinding& RenderTarget) {
+		RenderTarget.SetLoadAction(ERenderTargetLoadAction::ELoad);
+		});
+	BasePassRenderTargets.DepthStencil.SetDepthLoadAction(ERenderTargetLoadAction::ELoad);
+	BasePassRenderTargets.DepthStencil.SetStencilLoadAction(ERenderTargetLoadAction::ELoad);
+	BasePassRenderTargets.DepthStencil.SetDepthStencilAccess(FExclusiveDepthStencil::DepthRead_StencilWrite);
 
-	auto* ThirdPassParameters = GraphBuilder.AllocParameters<FMobileRenderPassParameters>();
-	*ThirdPassParameters = *SecondPassParameters;
-	// SceneColor write, SceneDepth is read only
-	for (int32 i = 1; i < NumColorTargets; ++i)
+	// Decals
+	if (!bIsFullDepthPrepassEnabled)
 	{
-		ThirdPassParameters->RenderTargets[i] = FRenderTargetBinding();
+		for (FRenderViewContext& ViewContext : RenderViews)
+		{
+			FViewInfo& View = *ViewContext.ViewInfo;
+
+			SCOPED_GPU_MASK(GraphBuilder.RHICmdList, !View.IsInstancedStereoPass() ? View.GPUMask : (View.GPUMask | View.GetInstancedView()->GPUMask));
+			SCOPED_CONDITIONAL_DRAW_EVENTF(GraphBuilder.RHICmdList, EventView, RenderViews.Num() > 1, TEXT("View%d"), ViewContext.ViewIndex);
+
+			View.BeginRenderView();
+
+			auto* PassParameters = GraphBuilder.AllocParameters<FMobileRenderPassParameters>();
+			PassParameters->View = View.GetShaderParameters();
+			PassParameters->MobileBasePass = CreateMobileBasePassUniformBuffer(GraphBuilder, View, EMobileBasePass::Opaque, EMobileSceneTextureSetupMode::SceneDepth);
+			PassParameters->RenderTargets = BasePassRenderTargets;
+
+			GraphBuilder.AddPass(
+				RDG_EVENT_NAME("Decals"),
+				PassParameters,
+				ERDGPassFlags::Raster,
+				[this, PassParameters, &View](FRHICommandList& RHICmdList)
+				{
+					RenderDecals(RHICmdList, View);
+				});
+		}
 	}
-	ThirdPassParameters->RenderTargets.DepthStencil.SetDepthStencilAccess(FExclusiveDepthStencil::DepthRead_StencilWrite);
-		
-	const bool bDoOcclusionQueires = (!bIsFullDepthPrepassEnabled && ViewContext.bIsLastView && DoOcclusionQueries());
-	ThirdPassParameters->RenderTargets.NumOcclusionQueries = bDoOcclusionQueires ? ComputeNumOcclusionQueriesToBatch() : 0u;
 
-	GraphBuilder.AddPass(
-		RDG_EVENT_NAME("LightingAndTranslucency"),
-		ThirdPassParameters,
-		ERDGPassFlags::Raster,
-		[this, ThirdPassParameters, ViewContext, bDoOcclusionQueires, &SceneTextures, &SortedLightSet](FRHICommandList& RHICmdList)
+	// Lighting and translucency
+	for (FRenderViewContext& ViewContext : RenderViews)
 	{
 		FViewInfo& View = *ViewContext.ViewInfo;
 
-		MobileDeferredShadingPass(RHICmdList, ViewContext.ViewIndex, Views.Num(), View, *Scene, SortedLightSet, VisibleLightInfos);
-		RenderFog(RHICmdList, View);
+		SCOPED_GPU_MASK(GraphBuilder.RHICmdList, !View.IsInstancedStereoPass() ? View.GPUMask : (View.GPUMask | View.GetInstancedView()->GPUMask));
+		SCOPED_CONDITIONAL_DRAW_EVENTF(GraphBuilder.RHICmdList, EventView, RenderViews.Num() > 1, TEXT("View%d"), ViewContext.ViewIndex);
 
-		// Draw translucency.
-		RenderTranslucency(RHICmdList, View);
+		View.BeginRenderView();
+		UpdateDirectionalLightUniformBuffers(GraphBuilder, View);
 
-		if (bDoOcclusionQueires)
-		{
-			// Issue occlusion queries
-			RenderOcclusion(RHICmdList);
-		}
-	});
+		auto* PassParameters = GraphBuilder.AllocParameters<FMobileRenderPassParameters>();
+		PassParameters->View = View.GetShaderParameters();
+
+		EMobileSceneTextureSetupMode SetupMode = 
+			EMobileSceneTextureSetupMode::SceneDepth | 
+			EMobileSceneTextureSetupMode::CustomDepth | 
+			EMobileSceneTextureSetupMode::GBuffers;
+		PassParameters->MobileBasePass = CreateMobileBasePassUniformBuffer(GraphBuilder, View, EMobileBasePass::Translucent, SetupMode);
+		PassParameters->ReflectionCapture = View.MobileReflectionCaptureUniformBuffer;
+		PassParameters->LocalFogVolumeInstances = View.LocalFogVolumeViewData.GPUInstanceDataBufferSRV;
+		PassParameters->LocalFogVolumeTileDrawIndirectBuffer = View.LocalFogVolumeViewData.GPUTileDrawIndirectBuffer;
+		PassParameters->LocalFogVolumeTileDataTexture = View.LocalFogVolumeViewData.TileDataTextureArraySRV;
+		PassParameters->LocalFogVolumeTileDataBuffer = View.LocalFogVolumeViewData.GPUTileDataBufferSRV;
+		PassParameters->HalfResLocalFogVolumeViewSRV = View.LocalFogVolumeViewData.HalfResLocalFogVolumeViewSRV;
+		PassParameters->HalfResLocalFogVolumeDepthSRV = View.LocalFogVolumeViewData.HalfResLocalFogVolumeDepthSRV;
+		// Only SceneColor and Depth
+		PassParameters->RenderTargets[0] = BasePassRenderTargets[0];
+		PassParameters->RenderTargets.DepthStencil = BasePassRenderTargets.DepthStencil;
+
+		View.ParallelMeshDrawCommandPasses[StandardTranslucencyMeshPass].BuildRenderingCommands(GraphBuilder, Scene->GPUScene, TranslucencyInstanceCullingDrawParams);
+		PassParameters->InstanceCullingDrawParams = TranslucencyInstanceCullingDrawParams;
+
+		GraphBuilder.AddPass(
+			RDG_EVENT_NAME("LightingAndTranslucency"),
+			PassParameters,
+			ERDGPassFlags::Raster,
+			[this, PassParameters, ViewContext, &SceneTextures, &SortedLightSet](FRHICommandList& RHICmdList)
+			{
+				FViewInfo& View = *ViewContext.ViewInfo;
+
+				MobileDeferredShadingPass(RHICmdList, ViewContext.ViewIndex, Views.Num(), View, *Scene, SortedLightSet, VisibleLightInfos);
+				RenderFog(RHICmdList, View);
+
+				// Draw translucency.
+				RenderTranslucency(RHICmdList, View);
+			});
+	}
 }
 
 void FMobileSceneRenderer::PostRenderBasePass(FRHICommandList& RHICmdList, FViewInfo& View)
diff --git a/Engine/Source/Runtime/Renderer/Private/RendererScene.cpp b/Engine/Source/Runtime/Renderer/Private/RendererScene.cpp
--- a/Engine/Source/Runtime/Renderer/Private/RendererScene.cpp
+++ b/Engine/Source/Runtime/Renderer/Private/RendererScene.cpp
@@ -4749,7 +4749,8 @@ void FScene::GetEarlyZPassMode(ERHIFeatureLevel::Type InFeatureLevel, EDepthDraw
 
 		if (MobileUsesFullDepthPrepass(ShaderPlatform))
 		{
-			OutZPassMode = DDM_AllOpaque;
+			const bool bDepthPassCanOutputVelocity = FVelocityRendering::DepthPassCanOutputVelocity(InFeatureLevel);
+			OutZPassMode = bDepthPassCanOutputVelocity ? DDM_AllOpaqueNoVelocity : DDM_AllOpaque;
 		}
 	}
 }
diff --git a/Engine/Source/Runtime/Renderer/Private/SceneRendering.h b/Engine/Source/Runtime/Renderer/Private/SceneRendering.h
--- a/Engine/Source/Runtime/Renderer/Private/SceneRendering.h
+++ b/Engine/Source/Runtime/Renderer/Private/SceneRendering.h
@@ -2669,10 +2669,9 @@ protected:
 	void RenderForwardSinglePass(FRDGBuilder& GraphBuilder, class FMobileRenderPassParameters* PassParameters, struct FRenderViewContext& ViewContext, FSceneTextures& SceneTextures);
 	void RenderForwardMultiPass(FRDGBuilder& GraphBuilder, class FMobileRenderPassParameters* PassParameters, struct FRenderViewContext& ViewContext, FSceneTextures& SceneTextures);
 	
-	void RenderDeferred(FRDGBuilder& GraphBuilder, const FSortedLightSetSceneInfo& SortedLightSet, FRDGTextureRef ViewFamilyTexture, FSceneTextures& SceneTextures);
-	void RenderDeferredSinglePass(FRDGBuilder& GraphBuilder, class FMobileRenderPassParameters* PassParameters, struct FRenderViewContext& ViewContext, FSceneTextures& SceneTextures, const FSortedLightSetSceneInfo& SortedLightSet, bool bUsingPixelLocalStorage);
-	void RenderDeferredMultiPass(FRDGBuilder& GraphBuilder, class FMobileRenderPassParameters* PassParameters, int32 NumColorTargets, struct FRenderViewContext& ViewContext, FSceneTextures& SceneTextures, const FSortedLightSetSceneInfo& SortedLightSet);
-	
+	void RenderDeferredSinglePass(FRDGBuilder& GraphBuilder, FSceneTextures& SceneTextures, const FSortedLightSetSceneInfo& SortedLightSet);
+	void RenderDeferredMultiPass(FRDGBuilder& GraphBuilder, FSceneTextures& SceneTextures, const FSortedLightSetSceneInfo& SortedLightSet);
+
 	void RenderAmbientOcclusion(FRDGBuilder& GraphBuilder, FRDGTextureRef SceneDepthTexture, FRDGTextureRef AmbientOcclusionTexture);
 
 	void RenderPixelProjectedReflection(FRDGBuilder& GraphBuilder, FRDGTextureRef SceneColorTexture, FRDGTextureRef SceneDepthTexture, FRDGTextureRef PixelProjectedReflectionTexture, const FPlanarReflectionSceneProxy* PlanarReflectionSceneProxy);
diff --git a/Engine/Source/Runtime/Renderer/Private/VelocityRendering.cpp b/Engine/Source/Runtime/Renderer/Private/VelocityRendering.cpp
--- a/Engine/Source/Runtime/Renderer/Private/VelocityRendering.cpp
+++ b/Engine/Source/Runtime/Renderer/Private/VelocityRendering.cpp
@@ -783,6 +783,16 @@ bool FVelocityMeshProcessor::Process(
 		return false;
 	}
 
+	// When velocity is used as a depth pass we need to set a correct stencil state on mobile
+	if (FeatureLevel == ERHIFeatureLevel::ES3_1 && EarlyZPassMode == DDM_AllOpaqueNoVelocity)
+	{
+		extern void SetMobileBasePassDepthState(FMeshPassProcessorRenderState& DrawRenderState, const FPrimitiveSceneProxy* PrimitiveSceneProxy, FMaterialShadingModelField ShadingModels, bool bUsesDeferredShading);
+
+		FMaterialShadingModelField ShadingModels = MaterialResource.GetShadingModels();
+		bool bUsesDeferredShading = IsMobileDeferredShadingEnabled(GetFeatureLevelShaderPlatform(FeatureLevel));
+		SetMobileBasePassDepthState(PassDrawRenderState, PrimitiveSceneProxy, ShadingModels, bUsesDeferredShading);
+	}
+
 	FMeshMaterialShaderElementData ShaderElementData;
 	ShaderElementData.InitializeMeshMaterialData(ViewIfDynamicMeshCommand, PrimitiveSceneProxy, MeshBatch, StaticMeshId, false);
 
@@ -859,9 +869,11 @@ FVelocityMeshProcessor::FVelocityMeshProcessor(EMeshPass::Type MeshPassType, con
 	, PassDrawRenderState(InPassDrawRenderState)
 {}
 
-FOpaqueVelocityMeshProcessor::FOpaqueVelocityMeshProcessor(const FScene* Scene, ERHIFeatureLevel::Type FeatureLevel, const FSceneView* InViewIfDynamicMeshCommand, const FMeshPassProcessorRenderState& InPassDrawRenderState, FMeshPassDrawListContext* InDrawListContext)
+FOpaqueVelocityMeshProcessor::FOpaqueVelocityMeshProcessor(const FScene* Scene, ERHIFeatureLevel::Type FeatureLevel, const FSceneView* InViewIfDynamicMeshCommand, const FMeshPassProcessorRenderState& InPassDrawRenderState, FMeshPassDrawListContext* InDrawListContext, EDepthDrawingMode InEarlyZPassMode)
 	: FVelocityMeshProcessor(EMeshPass::Velocity, Scene, FeatureLevel, InViewIfDynamicMeshCommand, InPassDrawRenderState, InDrawListContext)
-{}
+{
+	EarlyZPassMode = InEarlyZPassMode;
+}
 
 FMeshPassProcessor* CreateVelocityPassProcessor(ERHIFeatureLevel::Type FeatureLevel, const FScene* Scene, const FSceneView* InViewIfDynamicMeshCommand, FMeshPassDrawListContext* InDrawListContext)
 {
@@ -875,7 +887,7 @@ FMeshPassProcessor* CreateVelocityPassProcessor(ERHIFeatureLevel::Type FeatureLe
 										    ? TStaticDepthStencilState<true, CF_DepthNearOrEqual>::GetRHI()
 											: TStaticDepthStencilState<false, CF_DepthNearOrEqual>::GetRHI());
 
-	return new FOpaqueVelocityMeshProcessor(Scene, FeatureLevel, InViewIfDynamicMeshCommand, VelocityPassState, InDrawListContext);
+	return new FOpaqueVelocityMeshProcessor(Scene, FeatureLevel, InViewIfDynamicMeshCommand, VelocityPassState, InDrawListContext, EarlyZPassMode);
 }
 
 REGISTER_MESHPASSPROCESSOR_AND_PSOCOLLECTOR(VelocityPass, CreateVelocityPassProcessor, EShadingPath::Deferred,  EMeshPass::Velocity, EMeshPassFlags::CachedMeshCommands | EMeshPassFlags::MainView);
diff --git a/Engine/Source/Runtime/Renderer/Private/VelocityRendering.h b/Engine/Source/Runtime/Renderer/Private/VelocityRendering.h
--- a/Engine/Source/Runtime/Renderer/Private/VelocityRendering.h
+++ b/Engine/Source/Runtime/Renderer/Private/VelocityRendering.h
@@ -91,6 +91,9 @@ protected:
 		ERasterizerFillMode MeshFillMode,
 		ERasterizerCullMode MeshCullMode,
 		TArray<FPSOPrecacheData>& PSOInitializers);
+
+protected:
+	EDepthDrawingMode EarlyZPassMode = DDM_None;
 };
 
 /**
@@ -104,7 +107,8 @@ public:
 		ERHIFeatureLevel::Type FeatureLevel,
 		const FSceneView* InViewIfDynamicMeshCommand,
 		const FMeshPassProcessorRenderState& InPassDrawRenderState,
-		FMeshPassDrawListContext* InDrawListContext);
+		FMeshPassDrawListContext* InDrawListContext,
+		EDepthDrawingMode InEarlyZPassMode);
 
 	/** Returns true if the object is capable of having velocity for any frame. */
 	static bool PrimitiveCanHaveVelocity(EShaderPlatform ShaderPlatform, const FPrimitiveSceneProxy* PrimitiveSceneProxy);
