Fixed static validation warning

diff --git a/Engine/Source/Runtime/Renderer/Private/MobileShadingRenderer.cpp b/Engine/Source/Runtime/Renderer/Private/MobileShadingRenderer.cpp
--- a/Engine/Source/Runtime/Renderer/Private/MobileShadingRenderer.cpp
+++ b/Engine/Source/Runtime/Renderer/Private/MobileShadingRenderer.cpp
@@ -1361,7 +1361,8 @@ void FMobileSceneRenderer::Render(FRDGBuilder& GraphBuilder)
 		if (bShouldRenderVelocities)
 		{
 			// Render the velocities of movable objects
-			if (Scene->EarlyZPassMode != DDM_AllOpaqueNoVelocity)
+			EDepthDrawingMode EarlyZPassMode = Scene ? Scene->EarlyZPassMode : DDM_None;
+			if (EarlyZPassMode != DDM_AllOpaqueNoVelocity)
 			{ 
 				RenderVelocities(GraphBuilder, Views, SceneTextures, EVelocityPass::Opaque, false);
 			}
