Moved VSM stats to shared header and implement piping all of the them through to csv (if the stat category is enabled).

diff --git a/Engine/Shaders/Private/VirtualShadowMaps/VirtualShadowMapBuildPerPageDrawCommands.usf b/Engine/Shaders/Private/VirtualShadowMaps/VirtualShadowMapBuildPerPageDrawCommands.usf
--- a/Engine/Shaders/Private/VirtualShadowMaps/VirtualShadowMapBuildPerPageDrawCommands.usf
+++ b/Engine/Shaders/Private/VirtualShadowMaps/VirtualShadowMapBuildPerPageDrawCommands.usf
@@ -414,12 +414,12 @@ void CullPerPageDrawCommandsCs(uint3 GroupId : SV_GroupID, int GroupThreadIndex
 			{
 				uint PrevArea = 0U;
 				// Store these after the stats slots
-				InterlockedMax(OutStatsBuffer[VSM_NUM_STATS + 2U * Index], TotalPageArea, PrevArea);
+				InterlockedMax(OutStatsBuffer[VSM_STAT_NUM + 2U * Index], TotalPageArea, PrevArea);
 				// If the area was greater then we store the persistent Primitive ID after the area.
 				// This obviously has a data race, so numbers & IDs will only be approximately correct, could use 64-bit atomics instead...
 				if (PrevArea < TotalPageArea)
 				{
-					OutStatsBuffer[VSM_NUM_STATS + 2U * Index + 1U] = PrimitiveData.PersistentPrimitiveIndex;
+					OutStatsBuffer[VSM_STAT_NUM + 2U * Index + 1U] = PrimitiveData.PersistentPrimitiveIndex;
 					break;
 				}
 			}
diff --git a/Engine/Shaders/Private/VirtualShadowMaps/VirtualShadowMapCopyStats.usf b/Engine/Shaders/Private/VirtualShadowMaps/VirtualShadowMapCopyStats.usf
--- a/Engine/Shaders/Private/VirtualShadowMaps/VirtualShadowMapCopyStats.usf
+++ b/Engine/Shaders/Private/VirtualShadowMaps/VirtualShadowMapCopyStats.usf
@@ -6,7 +6,6 @@
 
 StructuredBuffer<uint>	InStatsBuffer;
 RWBuffer<uint>			AccumulatedStatsBufferOut;
-uint NumStats;
 
 // Pull out a bunch of nanite stats for our CSV as well
 StructuredBuffer<FNaniteStats> NaniteStatsBuffer;
@@ -18,8 +17,8 @@ void CopyStatsCS()
 	InterlockedAdd(AccumulatedStatsBufferOut[0], 1, Index);
 	if (Index < MAX_STAT_FRAMES)
 	{
-		uint Offset = 1 + Index * NumStats;
-		for (uint StatInd = 0; StatInd < NumStats; ++StatInd)
+		uint Offset = 1 + Index * VSM_STAT_NUM;
+		for (uint StatInd = 0; StatInd < VSM_STAT_NUM; ++StatInd)
 		{
 			AccumulatedStatsBufferOut[Offset + StatInd] = InStatsBuffer[StatInd];
 		}
diff --git a/Engine/Shaders/Private/VirtualShadowMaps/VirtualShadowMapPrintStats.usf b/Engine/Shaders/Private/VirtualShadowMaps/VirtualShadowMapPrintStats.usf
--- a/Engine/Shaders/Private/VirtualShadowMaps/VirtualShadowMapPrintStats.usf
+++ b/Engine/Shaders/Private/VirtualShadowMaps/VirtualShadowMapPrintStats.usf
@@ -63,17 +63,21 @@ StructuredBuffer<FNaniteStats> NaniteStats;
 
 void SendStatsMessage()
 {
-	FGPUMessageWriter Mw = GPUMessageBegin(StatsMessageId, 3U + (MAX_PAGE_AREA_DIAGNOSTIC_SLOTS * 2U));
+	FGPUMessageWriter Mw = GPUMessageBegin(StatsMessageId, 2U  + VSM_STAT_NUM + (MAX_PAGE_AREA_DIAGNOSTIC_SLOTS * 2U));
 
 	// Culling stats
 	GPUMessageWriteItem(Mw, NaniteStats[0].NumTris);
 	GPUMessageWriteItem(Mw, NaniteStats[0].NumMainInstancesPostCull);
-	GPUMessageWriteItem(Mw, InStatsBuffer[VSM_STAT_NON_NANITE_INSTANCES_DRAWN]);
+
+	for (int Stat = 0; Stat < VSM_STAT_NUM; ++Stat)
+	{
+		GPUMessageWriteItem(Mw, InStatsBuffer[Stat]);
+	}
 
 	// Large page area instances
 	for (uint Index = 0; Index < MAX_PAGE_AREA_DIAGNOSTIC_SLOTS * 2U; ++Index)
 	{
-		GPUMessageWriteItem(Mw, InStatsBuffer[VSM_NUM_STATS + Index]);
+		GPUMessageWriteItem(Mw, InStatsBuffer[VSM_STAT_NUM + Index]);
 	}
 }
 
@@ -133,6 +137,12 @@ void LogVirtualSmStatsCS()
 		PrintLeftAlign(Context, InStatsBuffer[VSM_STAT_NUM_PAGES_TO_CLEAR]);
 		Newline(Context);
 
+		Context.Pos.x = ItemX;
+		Print(Context, TEXT("Tmp1"), FontWhite);
+		Context.Pos.x = ValueX;
+		PrintLeftAlign(Context, InStatsBuffer[VSM_STAT_TMP_1]);
+		Newline(Context);
+
 		Context.Pos.x = ItemX;
 		Print(Context, TEXT("HZB Built"), FontWhite);
 		Context.Pos.x = ValueX;
diff --git a/Engine/Shaders/Private/VirtualShadowMaps/VirtualShadowMapStats.ush b/Engine/Shaders/Private/VirtualShadowMaps/VirtualShadowMapStats.ush
--- a/Engine/Shaders/Private/VirtualShadowMaps/VirtualShadowMapStats.ush
+++ b/Engine/Shaders/Private/VirtualShadowMaps/VirtualShadowMapStats.ush
@@ -4,30 +4,12 @@
 	VirtualShadowMapStats.ush: 
 =============================================================================*/
 
+#include "/Engine/Shared/VirtualShadowMapDefinitions.h"
+
 #include "../Common.ush"
 #include "../WaveOpUtil.ush"
 
 
-#define VSM_STAT_REQUESTED_THIS_FRAME_PAGES				 0
-#define VSM_STAT_STATIC_CACHED_PAGES					 1
-#define VSM_STAT_STATIC_INVALIDATED_PAGES				 2
-#define VSM_STAT_DYNAMIC_CACHED_PAGES					 3
-#define VSM_STAT_DYNAMIC_INVALIDATED_PAGES				 4
-#define VSM_STAT_EMPTY_PAGES							 5
-#define VSM_STAT_NON_NANITE_INSTANCES_TOTAL				 6
-#define VSM_STAT_NON_NANITE_INSTANCES_DRAWN				 7
-#define VSM_STAT_NON_NANITE_INSTANCES_HZB_CULLED		 8
-#define VSM_STAT_NON_NANITE_INSTANCES_PAGE_MASK_CULLED	 9
-#define VSM_STAT_NON_NANITE_INSTANCES_EMPTY_RECT_CULLED	10
-#define VSM_STAT_NON_NANITE_INSTANCES_FRUSTUM_CULLED	11
-#define VSM_STAT_NUM_PAGES_TO_MERGE						12
-#define VSM_STAT_NUM_PAGES_TO_CLEAR						13
-#define VSM_STAT_NUM_HZB_PAGES_BUILT					14
-#define VSM_STAT_ALLOCATED_NEW							15
-#define VSM_STAT_NANITE_TRIANGLES						16
-#define VSM_STAT_NANITE_INSTANCES_MAIN					17
-#define VSM_STAT_NANITE_INSTANCES_POST					18
-#define VSM_STAT_WPO_CONSIDERED_PAGES					19
 
 #ifndef VSM_GENERATE_STATS
 #define VSM_GENERATE_STATS 0
diff --git a/Engine/Shaders/Shared/VirtualShadowMapDefinitions.h b/Engine/Shaders/Shared/VirtualShadowMapDefinitions.h
--- a/Engine/Shaders/Shared/VirtualShadowMapDefinitions.h
+++ b/Engine/Shaders/Shared/VirtualShadowMapDefinitions.h
@@ -33,6 +33,32 @@ VirtualShadowMapDefinitions.h: used in virtual shadow map shaders and C++ code t
 // 8 bit flags, 24 bit VSM ID
 #define VSM_INVALIDATION_PAYLOAD_FLAG_BITS                      8
 
+
+#define VSM_STAT_REQUESTED_THIS_FRAME_PAGES				 0
+#define VSM_STAT_STATIC_CACHED_PAGES					 1
+#define VSM_STAT_STATIC_INVALIDATED_PAGES				 2
+#define VSM_STAT_DYNAMIC_CACHED_PAGES					 3
+#define VSM_STAT_DYNAMIC_INVALIDATED_PAGES				 4
+#define VSM_STAT_EMPTY_PAGES							 5
+#define VSM_STAT_NON_NANITE_INSTANCES_TOTAL				 6
+#define VSM_STAT_NON_NANITE_INSTANCES_DRAWN				 7
+#define VSM_STAT_NON_NANITE_INSTANCES_HZB_CULLED		 8
+#define VSM_STAT_NON_NANITE_INSTANCES_PAGE_MASK_CULLED	 9
+#define VSM_STAT_NON_NANITE_INSTANCES_EMPTY_RECT_CULLED	10
+#define VSM_STAT_NON_NANITE_INSTANCES_FRUSTUM_CULLED	11
+#define VSM_STAT_NUM_PAGES_TO_MERGE						12
+#define VSM_STAT_NUM_PAGES_TO_CLEAR						13
+#define VSM_STAT_NUM_HZB_PAGES_BUILT					14
+#define VSM_STAT_ALLOCATED_NEW							15
+#define VSM_STAT_NANITE_TRIANGLES						16
+#define VSM_STAT_NANITE_INSTANCES_MAIN					17
+#define VSM_STAT_NANITE_INSTANCES_POST					18
+#define VSM_STAT_WPO_CONSIDERED_PAGES					19
+#define VSM_STAT_TMP_1									20
+#define VSM_STAT_TMP_2									21
+#define VSM_STAT_TMP_3									22
+#define VSM_STAT_NUM									23
+
 #ifdef __cplusplus
 #include "HLSLTypeAliases.h"
 
diff --git a/Engine/Source/Runtime/Renderer/Private/VirtualShadowMaps/VirtualShadowMapArray.cpp b/Engine/Source/Runtime/Renderer/Private/VirtualShadowMaps/VirtualShadowMapArray.cpp
--- a/Engine/Source/Runtime/Renderer/Private/VirtualShadowMaps/VirtualShadowMapArray.cpp
+++ b/Engine/Source/Runtime/Renderer/Private/VirtualShadowMaps/VirtualShadowMapArray.cpp
@@ -30,6 +30,7 @@ IMPLEMENT_STATIC_UNIFORM_BUFFER_SLOT(VirtualShadowMapUbSlot);
 
 IMPLEMENT_STATIC_UNIFORM_BUFFER_STRUCT(FVirtualShadowMapUniformParameters, "VirtualShadowMap", VirtualShadowMapUbSlot);
 
+// Disabled by default: use either console command "CsvCategory VSM" or command line argument "-CsvCategories=VSM[,...]" to enable.
 CSV_DEFINE_CATEGORY(VSM, false);
 
 DECLARE_DWORD_COUNTER_STAT(TEXT("VSM Nanite Views (Primary)"), STAT_VSMNaniteViewsPrimary, STATGROUP_ShadowRendering);
@@ -643,7 +644,6 @@ void FVirtualShadowMapArray::SetShaderDefines(FShaderCompilerEnvironment& OutEnv
 	OutEnvironment.SetDefine(TEXT("VSM_VIRTUAL_MAX_RESOLUTION_XY"), FVirtualShadowMap::VirtualMaxResolutionXY);
 	OutEnvironment.SetDefine(TEXT("VSM_RASTER_WINDOW_PAGES"), FVirtualShadowMap::RasterWindowPages);
 	OutEnvironment.SetDefine(TEXT("VSM_PAGE_TABLE_SIZE"), FVirtualShadowMap::PageTableSize);
-	OutEnvironment.SetDefine(TEXT("VSM_NUM_STATS"), NumStats);
 	OutEnvironment.SetDefine(TEXT("MAX_PAGE_AREA_DIAGNOSTIC_SLOTS"), MaxPageAreaDiagnosticSlots);
 	OutEnvironment.SetDefine(TEXT("INDEX_NONE"), INDEX_NONE);
 }
@@ -1326,6 +1326,9 @@ void FVirtualShadowMapArray::BuildPageAllocations(
 	// Stats
 	SET_DWORD_STAT(STAT_VSMSinglePageCount, GetNumSinglePageShadowMaps());
 	SET_DWORD_STAT(STAT_VSMFullCount, GetNumFullShadowMaps());
+	// And _other_ stats...
+	CSV_CUSTOM_STAT(VSM, SinglePageCount, GetNumSinglePageShadowMaps(), ECsvCustomStatOp::Set);
+	CSV_CUSTOM_STAT(VSM, FullCount, GetNumFullShadowMaps(), ECsvCustomStatOp::Set);
 
 	UniformParameters.NumFullShadowMaps = GetNumFullShadowMaps();
 	UniformParameters.NumSinglePageShadowMaps = GetNumSinglePageShadowMaps();
@@ -1349,7 +1352,7 @@ void FVirtualShadowMapArray::BuildPageAllocations(
 
 	if (CVarShowStats.GetValueOnRenderThread() || CacheManager->IsAccumulatingStats() || bRunPageAreaDiagnostics || bCsvLogEnabled)
 	{
-		StatsBufferRDG = GraphBuilder.CreateBuffer(FRDGBufferDesc::CreateStructuredDesc(sizeof(uint32), NumStats + MaxPageAreaDiagnosticSlots * 2), TEXT("Shadow.Virtual.StatsBuffer"));
+		StatsBufferRDG = GraphBuilder.CreateBuffer(FRDGBufferDesc::CreateStructuredDesc(sizeof(uint32), VSM_STAT_NUM + MaxPageAreaDiagnosticSlots * 2), TEXT("Shadow.Virtual.StatsBuffer"));
 		AddClearUAVPass(GraphBuilder, GraphBuilder.CreateUAV(StatsBufferRDG), 0);
 
 		// For the rest of the frame we don't want the stats buffer adding additional barriers that are not otherwise present.
diff --git a/Engine/Source/Runtime/Renderer/Private/VirtualShadowMaps/VirtualShadowMapArray.h b/Engine/Source/Runtime/Renderer/Private/VirtualShadowMaps/VirtualShadowMapArray.h
--- a/Engine/Source/Runtime/Renderer/Private/VirtualShadowMaps/VirtualShadowMapArray.h
+++ b/Engine/Source/Runtime/Renderer/Private/VirtualShadowMaps/VirtualShadowMapArray.h
@@ -406,8 +406,6 @@ public:
 	FRDGBufferRef DirtyPageFlagsRDG = nullptr; // Dirty flags that are cleared after render passes
 	bool bHZBBuiltThisFrame = false;
 
-	// See Engine\Shaders\Private\VirtualShadowMaps\VirtualShadowMapStats.ush for definitions of the different stat indexes
-	static constexpr uint32 NumStats = 32;
 	static constexpr uint32 MaxPageAreaDiagnosticSlots = 32;
 
 	FRDGBufferRef StatsBufferRDG = nullptr;
diff --git a/Engine/Source/Runtime/Renderer/Private/VirtualShadowMaps/VirtualShadowMapCacheManager.cpp b/Engine/Source/Runtime/Renderer/Private/VirtualShadowMaps/VirtualShadowMapCacheManager.cpp
--- a/Engine/Source/Runtime/Renderer/Private/VirtualShadowMaps/VirtualShadowMapCacheManager.cpp
+++ b/Engine/Source/Runtime/Renderer/Private/VirtualShadowMaps/VirtualShadowMapCacheManager.cpp
@@ -133,6 +133,35 @@ FAutoConsoleVariableRef CVarVSMNewInvalidations(
 	ECVF_RenderThreadSafe
 );
 
+static const char *VirtualShadowMap_StatNames[] = 
+{
+	"REQUESTED_THIS_FRAME_PAGES",
+	"STATIC_CACHED_PAGES",
+	"STATIC_INVALIDATED_PAGES",
+	"DYNAMIC_CACHED_PAGES",
+	"DYNAMIC_INVALIDATED_PAGES",
+	"EMPTY_PAGES",
+	"NON_NANITE_INSTANCES_TOTAL",
+	"NON_NANITE_INSTANCES_DRAWN",
+	"NON_NANITE_INSTANCES_HZB_CULLED",
+	"NON_NANITE_INSTANCES_PAGE_MASK_CULLED",
+	"NON_NANITE_INSTANCES_EMPTY_RECT_CULLED",
+	"NON_NANITE_INSTANCES_FRUSTUM_CULLED",
+	"NUM_PAGES_TO_MERGE",
+	"NUM_PAGES_TO_CLEAR",
+	"NUM_HZB_PAGES_BUILT",
+	"ALLOCATED_NEW",
+	"NANITE_TRIANGLES",
+	"NANITE_INSTANCES_MAIN",
+	"NANITE_INSTANCES_POST",
+	"WPO_CONSIDERED_PAGES",
+	"TMP_1",
+	"TMP_2",
+	"TMP_3",
+};
+static_assert(UE_ARRAY_COUNT(VirtualShadowMap_StatNames) == VSM_STAT_NUM, "Stat text name array length mismatch!");
+
+
 namespace Nanite
 {
 	extern bool IsStatFilterActive(const FString& FilterName);
@@ -655,11 +684,37 @@ void FVirtualShadowMapArrayCacheManager::InitExtension(FScene& InScene)
 		// Culling stats
 		int32 NaniteNumTris = Message.Read<int32>(0);
 		int32 NanitePostCullNodeCount = Message.Read<int32>(0);
-		int32 NonNanitePostCullInstanceCount = Message.Read<int32>(0);
+
+		TConstArrayView<uint32> Stats = Message.ReadCount(VSM_STAT_NUM);
 
 		CSV_CUSTOM_STAT(VSM, NaniteNumTris, NaniteNumTris, ECsvCustomStatOp::Set);
 		CSV_CUSTOM_STAT(VSM, NanitePostCullNodeCount, NanitePostCullNodeCount, ECsvCustomStatOp::Set);
-		CSV_CUSTOM_STAT(VSM, NonNanitePostCullInstanceCount, NonNanitePostCullInstanceCount, ECsvCustomStatOp::Set);
+#if CSV_PROFILER
+		CSV_CUSTOM_STAT(VSM, NonNanitePostCullInstanceCount, int32(Stats[VSM_STAT_NON_NANITE_INSTANCES_DRAWN]), ECsvCustomStatOp::Set);
+
+		if (FCsvProfiler::Get()->IsCapturing_Renderthread())
+		{
+			static bool bRegisteredInlineStats = false;
+			auto StatCatIndex = CSV_CATEGORY_INDEX(VSM);
+			if (FCsvProfiler::Get()->IsCategoryEnabled(StatCatIndex))
+			{
+				for (int32 StatIndex = 0; StatIndex < UE_ARRAY_COUNT(VirtualShadowMap_StatNames); ++StatIndex)
+				{
+					const char *StatName = VirtualShadowMap_StatNames[StatIndex];
+#if CSVPROFILERTRACE_ENABLED
+					if (!bRegisteredInlineStats)
+					{
+						FCsvProfilerTrace::OutputInlineStat(StatName, StatCatIndex);
+					}
+#endif
+					FCsvProfiler::RecordCustomStat(StatName, StatCatIndex, int32(Stats[StatIndex]), ECsvCustomStatOp::Set);
+				}
+				bRegisteredInlineStats = true;
+			}
+		}
+#endif
+
+
 
 		// Large page area items
 		LastLoggedPageOverlapAppTime.SetNumZeroed(Scene->GetMaxPersistentPrimitiveIndex());
@@ -1018,7 +1073,6 @@ public:
 		SHADER_PARAMETER_RDG_BUFFER_SRV(StructuredBuffer<uint>, InStatsBuffer)
 		SHADER_PARAMETER_RDG_BUFFER_SRV(StructuredBuffer<FNaniteStats>, NaniteStatsBuffer)
 		SHADER_PARAMETER_RDG_BUFFER_UAV(RWBuffer<uint>, AccumulatedStatsBufferOut)
-		SHADER_PARAMETER(uint32, NumStats)
 	END_SHADER_PARAMETER_STRUCT()
 	
 	static bool ShouldCompilePermutation(const FGlobalShaderPermutationParameters& Parameters)
@@ -1137,7 +1191,7 @@ void FVirtualShadowMapArrayCacheManager::ExtractStats(FRDGBuilder& GraphBuilder,
 	{
 		if (!AccumulatedStatsBuffer.IsValid())
 		{
-			FRDGBufferDesc Desc = FRDGBufferDesc::CreateBufferDesc(4, 1 + FVirtualShadowMapArray::NumStats * MaxStatFrames);
+			FRDGBufferDesc Desc = FRDGBufferDesc::CreateBufferDesc(4, 1 + VSM_STAT_NUM * MaxStatFrames);
 			Desc.Usage = EBufferUsageFlags(Desc.Usage | BUF_SourceCopy);
 
 			AccumulatedStatsBufferRDG = GraphBuilder.CreateBuffer(Desc, TEXT("Shadow.Virtual.AccumulatedStatsBuffer"));	// TODO: Can't be a structured buffer as EnqueueCopy is only defined for vertex buffers
@@ -1156,7 +1210,6 @@ void FVirtualShadowMapArrayCacheManager::ExtractStats(FRDGBuilder& GraphBuilder,
 
 		PassParameters->InStatsBuffer = GraphBuilder.CreateSRV(VirtualShadowMapArray.StatsBufferRDG, PF_R32_UINT);
 		PassParameters->AccumulatedStatsBufferOut = GraphBuilder.CreateUAV(AccumulatedStatsBufferRDG, PF_R32_UINT);
-		PassParameters->NumStats = FVirtualShadowMapArray::NumStats;
 
 		// Dummy data
 		PassParameters->NaniteStatsBuffer = GraphBuilder.CreateSRV(GSystemTextures.GetDefaultStructuredBuffer<FNaniteStats>(GraphBuilder));
@@ -1198,10 +1251,10 @@ void FVirtualShadowMapArrayCacheManager::ExtractStats(FRDGBuilder& GraphBuilder,
 	if (GPUBufferReadback && GPUBufferReadback->IsReady())
 	{
 		TArray<uint32> Tmp;
-		Tmp.AddDefaulted(1 + FVirtualShadowMapArray::NumStats * MaxStatFrames);
+		Tmp.AddDefaulted(1 + VSM_STAT_NUM * MaxStatFrames);
 
 		{
-			const uint32* BufferPtr = (const uint32*)GPUBufferReadback->Lock((1 + FVirtualShadowMapArray::NumStats * MaxStatFrames) * sizeof(uint32));
+			const uint32* BufferPtr = (const uint32*)GPUBufferReadback->Lock((1 + VSM_STAT_NUM * MaxStatFrames) * sizeof(uint32));
 			FPlatformMemory::Memcpy(Tmp.GetData(), BufferPtr, Tmp.Num() * Tmp.GetTypeSize());
 			GPUBufferReadback->Unlock();
 
@@ -1219,45 +1272,15 @@ void FVirtualShadowMapArrayCacheManager::ExtractStats(FRDGBuilder& GraphBuilder,
 		ensure(FileToLogTo);
 		if (FileToLogTo)
 		{
-			static const FString StatNames[] =
-			{
-				TEXT("Requested"),
-				TEXT("StaticCached"),
-				TEXT("StaticInvalidated"),
-				TEXT("DynamicCached"),
-				TEXT("DynamicInvalidated"),
-				TEXT("Empty"),
-				TEXT("NonNaniteInstances"),
-				TEXT("NonNaniteInstancesDrawn"),
-				TEXT("NonNaniteInstancesHZBCulled"),
-				TEXT("NonNaniteInstancesPageMaskCulled"),
-				TEXT("NonNaniteInstancesEmptyRectCulled"),
-				TEXT("NonNaniteInstancesFrustumCulled"),
-				TEXT("Merged"),
-				TEXT("Cleared"),
-				TEXT("HZBBuilt"),
-				TEXT("AllocatedNew"),
-				TEXT("NaniteTriangles"),
-				TEXT("NaniteInstancesMain"),
-				TEXT("NaniteInstancesPost"),
-			};
-
 			// Print header
 			FString StringToPrint;
-			for (int32 Index = 0; Index < FVirtualShadowMapArray::NumStats; ++Index)
+			for (int32 Index = 0; Index < VSM_STAT_NUM; ++Index)
 			{
 				if (!StringToPrint.IsEmpty())
 				{
 					StringToPrint += TEXT(",");
 				}
-				if (Index < int32(UE_ARRAY_COUNT(StatNames)))
-				{
-					StringToPrint.Append(StatNames[Index]);
-				}
-				else
-				{
-					StringToPrint.Appendf(TEXT("Stat_%d"), Index);
-				}
+				StringToPrint.Append(VirtualShadowMap_StatNames[Index]);
 			}
 
 			StringToPrint += TEXT("\n");
@@ -1267,14 +1290,14 @@ void FVirtualShadowMapArrayCacheManager::ExtractStats(FRDGBuilder& GraphBuilder,
 			{
 				StringToPrint.Empty();
 
-				for (uint32 StatInd = 0; StatInd < FVirtualShadowMapArray::NumStats; ++StatInd)
+				for (uint32 StatInd = 0; StatInd < VSM_STAT_NUM; ++StatInd)
 				{
 					if (!StringToPrint.IsEmpty())
 					{
 						StringToPrint += TEXT(",");
 					}
 
-					StringToPrint += FString::Printf(TEXT("%d"), Tmp[1 + Ind * FVirtualShadowMapArray::NumStats + StatInd]);
+					StringToPrint += FString::Printf(TEXT("%d"), Tmp[1 + Ind * VSM_STAT_NUM + StatInd]);
 				}
 
 				StringToPrint += TEXT("\n");
