Added r.SkipNaniteLPIs (default 1) to improve light & primitive update performance when there is a lot of Nanite geometry. The default, (r.SkipNaniteLPIs=1), skips Nanite primitives when r.AllowStaticLighting is turned off, while setting it to 2 (forced) always skips the LPI creation and octree insert.

diff --git a/Engine/Source/Runtime/Renderer/Private/LightSceneInfo.cpp b/Engine/Source/Runtime/Renderer/Private/LightSceneInfo.cpp
--- a/Engine/Source/Runtime/Renderer/Private/LightSceneInfo.cpp
+++ b/Engine/Source/Runtime/Renderer/Private/LightSceneInfo.cpp
@@ -114,22 +114,25 @@ void FLightSceneInfo::AddToScene()
 	}
 }
 
-/**
- * If the light affects the primitive, create an interaction, and process children 
- * 
- * @param LightSceneInfoCompact Compact representation of the light
- * @param PrimitiveSceneInfoCompact Compact representation of the primitive
- */
+bool FLightSceneInfo::ShouldCreateLightPrimitiveInteraction(const FLightSceneInfoCompact& LightSceneInfoCompact, const FPrimitiveSceneInfoCompact& PrimitiveSceneInfoCompact)
+{
+	if (!Scene->IsPrimitiveBeingRemoved(PrimitiveSceneInfoCompact.PrimitiveSceneInfo) && LightSceneInfoCompact.AffectsPrimitive(FBoxSphereBounds(PrimitiveSceneInfoCompact.Bounds), PrimitiveSceneInfoCompact.Proxy))
+	{
+		// create light interaction and add to light/primitive lists
+		return FLightPrimitiveInteraction::ShouldCreate(this, PrimitiveSceneInfoCompact.PrimitiveSceneInfo).bShouldCreate;
+	}
+	return false;
+}
+
 void FLightSceneInfo::CreateLightPrimitiveInteraction(const FLightSceneInfoCompact& LightSceneInfoCompact, const FPrimitiveSceneInfoCompact& PrimitiveSceneInfoCompact)
 {
-	if(	!Scene->IsPrimitiveBeingRemoved(PrimitiveSceneInfoCompact.PrimitiveSceneInfo) && LightSceneInfoCompact.AffectsPrimitive(FBoxSphereBounds(PrimitiveSceneInfoCompact.Bounds), PrimitiveSceneInfoCompact.Proxy))
+	if (!Scene->IsPrimitiveBeingRemoved(PrimitiveSceneInfoCompact.PrimitiveSceneInfo) && LightSceneInfoCompact.AffectsPrimitive(FBoxSphereBounds(PrimitiveSceneInfoCompact.Bounds), PrimitiveSceneInfoCompact.Proxy))
 	{
 		// create light interaction and add to light/primitive lists
 		FLightPrimitiveInteraction::Create(this,PrimitiveSceneInfoCompact.PrimitiveSceneInfo);
 	}
 }
 
-
 void FLightSceneInfo::RemoveFromScene()
 {
 	if (OctreeId.IsValidId())
diff --git a/Engine/Source/Runtime/Renderer/Private/LightSceneInfo.h b/Engine/Source/Runtime/Renderer/Private/LightSceneInfo.h
--- a/Engine/Source/Runtime/Renderer/Private/LightSceneInfo.h
+++ b/Engine/Source/Runtime/Renderer/Private/LightSceneInfo.h
@@ -271,6 +271,13 @@ public:
 	/** Adds the light to the scene. */
 	void AddToScene();
 
+	/**
+	 * Returns true if the light affects the primitive
+	 * @param LightSceneInfoCompact Compact representation of the light
+	 * @param PrimitiveSceneInfoCompact Compact representation of the primitive
+	 */
+	bool ShouldCreateLightPrimitiveInteraction(const FLightSceneInfoCompact& LightSceneInfoCompact, const FPrimitiveSceneInfoCompact& PrimitiveSceneInfoCompact);
+
 	/**
 	 * If the light affects the primitive, create an interaction, and process children 
 	 * @param LightSceneInfoCompact Compact representation of the light
diff --git a/Engine/Source/Runtime/Renderer/Private/PrimitiveSceneInfo.cpp b/Engine/Source/Runtime/Renderer/Private/PrimitiveSceneInfo.cpp
--- a/Engine/Source/Runtime/Renderer/Private/PrimitiveSceneInfo.cpp
+++ b/Engine/Source/Runtime/Renderer/Private/PrimitiveSceneInfo.cpp
@@ -38,6 +38,7 @@
 #include "BasePassRendering.h" // TODO: Remove with later refactor (moving Nanite shading into its own files)
 #include "InstanceDataSceneProxy.h"
 #include "DecalRenderingCommon.h"
+#include "RendererPrivateUtils.h"
 
 extern int32 GGPUSceneInstanceClearList;
 
@@ -1640,16 +1641,21 @@ void FPrimitiveSceneInfo::AddToScene(FScene* Scene, TArrayView<FPrimitiveSceneIn
 	}
 
 	{
+		const bool bSkipNaniteInOctree = ShouldSkipNaniteLPIs();
 		SCOPED_NAMED_EVENT(FPrimitiveSceneInfo_AddToScene_AddToPrimitiveOctree, FColor::Red);
 		for (FPrimitiveSceneInfo* SceneInfo : SceneInfos)
 		{
-			// create potential storage for our compact info
-			FPrimitiveSceneInfoCompact CompactPrimitiveSceneInfo(SceneInfo);
+			// doing this check after updating PrimitiveFlagsCompact (next loop) would be more efficient.
+			if (!bSkipNaniteInOctree || !SceneInfo->Proxy->IsNaniteMesh())
+			{
+				// create potential storage for our compact info
+				FPrimitiveSceneInfoCompact CompactPrimitiveSceneInfo(SceneInfo);
 
-			// Add the primitive to the octree.
-			check(!SceneInfo->OctreeId.IsValidId());
-			Scene->PrimitiveOctree.AddElement(CompactPrimitiveSceneInfo);
-			check(SceneInfo->OctreeId.IsValidId());
+				// Add the primitive to the octree.
+				check(!SceneInfo->OctreeId.IsValidId());
+				Scene->PrimitiveOctree.AddElement(CompactPrimitiveSceneInfo);
+				check(SceneInfo->OctreeId.IsValidId());
+			}
 		}
 	}
 
@@ -1810,9 +1816,12 @@ void FPrimitiveSceneInfo::RemoveFromScene(bool bUpdateStaticDrawLists)
 	}
 
 	// Remove the primitive from the octree.
-	check(OctreeId.IsValidId());
-	check(Scene->PrimitiveOctree.GetElementById(OctreeId).PrimitiveSceneInfo == this);
-	Scene->PrimitiveOctree.RemoveElement(OctreeId);
+	if (OctreeId.IsValidId())
+	{
+		check(Scene->PrimitiveOctree.GetElementById(OctreeId).PrimitiveSceneInfo == this);
+		Scene->PrimitiveOctree.RemoveElement(OctreeId);
+	}
+
 	OctreeId = FOctreeElementId2();
 
 	if (LightmapDataOffset != INDEX_NONE && UseGPUScene(GMaxRHIShaderPlatform, Scene->GetFeatureLevel()))
diff --git a/Engine/Source/Runtime/Renderer/Private/RendererPrivateUtils.h b/Engine/Source/Runtime/Renderer/Private/RendererPrivateUtils.h
--- a/Engine/Source/Runtime/Renderer/Private/RendererPrivateUtils.h
+++ b/Engine/Source/Runtime/Renderer/Private/RendererPrivateUtils.h
@@ -7,6 +7,12 @@
 #include "RHICommandList.h"
 #include "RenderGraphUtils.h"
 
+/**
+ * Returns true if Nanite primitives are to be skipped in Light Primitive Interactions (LPI) and the primitive acceleration structure (octree).
+ */
+bool ShouldSkipNaniteLPIs();
+
+
 class FTileTexCoordVertexBuffer : public FVertexBuffer
 {
 public:
diff --git a/Engine/Source/Runtime/Renderer/Private/RendererScene.cpp b/Engine/Source/Runtime/Renderer/Private/RendererScene.cpp
--- a/Engine/Source/Runtime/Renderer/Private/RendererScene.cpp
+++ b/Engine/Source/Runtime/Renderer/Private/RendererScene.cpp
@@ -4524,15 +4524,39 @@ FRHIUniformBuffer* FScene::GetSpeedTreeUniformBuffer(const FVertexFactory* Verte
  * @param	Primitive				Primitive to retrieve interacting lights for
  * @param	RelevantLights	[out]	Array of lights interacting with primitive
  */
-void FScene::GetRelevantLights_RenderThread( UPrimitiveComponent* Primitive, TArray<const ULightComponent*>* RelevantLights ) const
+void FScene::GetRelevantLights_RenderThread( const FPrimitiveSceneProxy* PrimitiveSceneProxy, TArray<const FLightSceneProxy*> &OutRelevantLights ) const
 {
-	check( Primitive );
-	check( RelevantLights );
-	if( Primitive->SceneProxy )
+	if (ShouldSkipNaniteLPIs() && PrimitiveSceneProxy->IsNaniteMesh())
 	{
-		for( const FLightPrimitiveInteraction* Interaction=Primitive->SceneProxy->GetPrimitiveSceneInfo()->LightList; Interaction; Interaction=Interaction->GetNextLight() )
+		if (PrimitiveSceneProxy->GetLightingChannelMask() != 0)
 		{
-			RelevantLights->Add( Interaction->GetLight()->Proxy->GetLightComponent() );
+			const FBoxSphereBounds& Bounds = PrimitiveSceneProxy->GetBounds();
+			const FPrimitiveSceneInfoCompact PrimitiveSceneInfoCompact(PrimitiveSceneProxy->GetPrimitiveSceneInfo());
+			auto TestAddLight = [&](const FLightSceneInfoCompact& LightSceneInfoCompact)
+			{
+				if (LightSceneInfoCompact.LightSceneInfo->ShouldCreateLightPrimitiveInteraction(LightSceneInfoCompact, PrimitiveSceneInfoCompact))
+				{
+					OutRelevantLights.Add( LightSceneInfoCompact.LightSceneInfo->Proxy);
+				}
+			};
+
+			if(DoesPlatformNeedLocalLightPrimitiveInteraction(GetShaderPlatform()))
+			{
+				// Find local lights that affect the primitive in the light octree.
+				LocalShadowCastingLightOctree.FindElementsWithBoundsTest(Bounds.GetBox(), TestAddLight);
+			}
+			// Also loop through non-local (directional) shadow-casting lights
+			for (int32 LightID : DirectionalShadowCastingLightIDs)
+			{
+				TestAddLight(Lights[LightID]);
+			}
+		}
+	}
+	else
+	{
+		for( const FLightPrimitiveInteraction* Interaction = PrimitiveSceneProxy->GetPrimitiveSceneInfo()->LightList; Interaction; Interaction=Interaction->GetNextLight() )
+		{
+			OutRelevantLights.Add( Interaction->GetLight()->Proxy);
 		}
 	}
 }
@@ -4548,11 +4572,15 @@ void FScene::GetRelevantLights( UPrimitiveComponent* Primitive, TArray<const ULi
 	if( Primitive && RelevantLights )
 	{
 		// Add interacting lights to the array.
-		const FScene* Scene = this;
 		ENQUEUE_RENDER_COMMAND(FGetRelevantLightsCommand)(
-			[Scene, Primitive, RelevantLights] (FRHICommandListBase&)
+			[this, PrimitiveSceneProxy = Primitive->GetSceneProxy(), RelevantLights] (FRHICommandListBase&)
 			{
-				Scene->GetRelevantLights_RenderThread( Primitive, RelevantLights );
+				TArray<const FLightSceneProxy*> RelevantLightProxies;
+				GetRelevantLights_RenderThread( PrimitiveSceneProxy, RelevantLightProxies );
+				for (const FLightSceneProxy* LightSceneProxy : RelevantLightProxies)
+				{
+					RelevantLights->Add(LightSceneProxy->GetLightComponent());
+				}
 			});
 
 		// We need to block the main thread as the rendering thread needs to finish modifying the array before we can continue.
@@ -6626,9 +6654,37 @@ void FScene::Update(FRDGBuilder& GraphBuilder, const FUpdateParameters& Paramete
 	{
 		SCOPED_NAMED_EVENT(CreateLightPrimitiveInteractions, FColor::Emerald);
 
-		for (FPrimitiveSceneInfo* SceneInfo : SceneInfosWithAddToScene)
+		bool bSkipNaniteLPIs = ShouldSkipNaniteLPIs();
+
+		for (FPrimitiveSceneInfo* PrimitiveSceneInfo : SceneInfosWithAddToScene)
 		{
-			CreateLightPrimitiveInteractionsForPrimitive(SceneInfo);
+			FPrimitiveSceneProxy* Proxy = PrimitiveSceneInfo->Proxy;
+			if (Proxy->GetLightingChannelMask() != 0)
+			{
+				// Don't create LPIs for Nanite
+				if (bSkipNaniteLPIs && Proxy->IsNaniteMesh())
+				{
+					continue;
+				}
+
+				const FBoxSphereBounds& Bounds = Proxy->GetBounds();
+				const FPrimitiveSceneInfoCompact PrimitiveSceneInfoCompact(PrimitiveSceneInfo);
+
+				if(DoesPlatformNeedLocalLightPrimitiveInteraction(GetShaderPlatform()))
+				{
+					// Find local lights that affect the primitive in the light octree.
+					LocalShadowCastingLightOctree.FindElementsWithBoundsTest(Bounds.GetBox(), [&PrimitiveSceneInfoCompact](const FLightSceneInfoCompact& LightSceneInfoCompact)
+					{
+						LightSceneInfoCompact.LightSceneInfo->CreateLightPrimitiveInteraction(LightSceneInfoCompact, PrimitiveSceneInfoCompact);
+					});
+				}
+				// Also loop through non-local (directional) shadow-casting lights
+				for (int32 LightID : DirectionalShadowCastingLightIDs)
+				{
+					const FLightSceneInfoCompact& LightSceneInfoCompact = Lights[LightID];
+					LightSceneInfoCompact.LightSceneInfo->CreateLightPrimitiveInteraction(LightSceneInfoCompact, PrimitiveSceneInfoCompact);
+				}
+			}			
 		}
 
 	}, EnumHasAnyFlags(Parameters.AsyncOps, EUpdateAllPrimitiveSceneInfosAsyncOps::CreateLightPrimitiveInteractions));
@@ -6893,31 +6949,6 @@ void FScene::Update(FRDGBuilder& GraphBuilder, const FUpdateParameters& Paramete
 #endif
 }
 
-void FScene::CreateLightPrimitiveInteractionsForPrimitive(FPrimitiveSceneInfo* PrimitiveInfo)
-{
-	FPrimitiveSceneProxy* Proxy = PrimitiveInfo->Proxy;
-	if (Proxy->GetLightingChannelMask() != 0)
-	{
-		const FBoxSphereBounds& Bounds = Proxy->GetBounds();
-		const FPrimitiveSceneInfoCompact PrimitiveSceneInfoCompact(PrimitiveInfo);
-
-		if(DoesPlatformNeedLocalLightPrimitiveInteraction(GetShaderPlatform()))
-		{
-			// Find local lights that affect the primitive in the light octree.
-			LocalShadowCastingLightOctree.FindElementsWithBoundsTest(Bounds.GetBox(), [&PrimitiveSceneInfoCompact](const FLightSceneInfoCompact& LightSceneInfoCompact)
-			{
-				LightSceneInfoCompact.LightSceneInfo->CreateLightPrimitiveInteraction(LightSceneInfoCompact, PrimitiveSceneInfoCompact);
-			});
-		}
-		// Also loop through non-local (directional) shadow-casting lights
-		for (int32 LightID : DirectionalShadowCastingLightIDs)
-		{
-			const FLightSceneInfoCompact& LightSceneInfoCompact = Lights[LightID];
-			LightSceneInfoCompact.LightSceneInfo->CreateLightPrimitiveInteraction(LightSceneInfoCompact, PrimitiveSceneInfoCompact);
-		}
-	}
-}
-
 bool FScene::IsPrimitiveBeingRemoved(FPrimitiveSceneInfo* PrimitiveSceneInfo) const
 {
 	check(IsInParallelRenderingThread() || IsInRenderingThread());
diff --git a/Engine/Source/Runtime/Renderer/Private/RendererUtils.cpp b/Engine/Source/Runtime/Renderer/Private/RendererUtils.cpp
--- a/Engine/Source/Runtime/Renderer/Private/RendererUtils.cpp
+++ b/Engine/Source/Runtime/Renderer/Private/RendererUtils.cpp
@@ -9,6 +9,29 @@
 #include "ScenePrivate.h"
 #include "SystemTextures.h"
 #include "UnifiedBuffer.h"
+#include "ComponentRecreateRenderStateContext.h"
+
+static int32 GSkipNaniteLPIs = 1;
+static FAutoConsoleVariableRef CVarSkipNaniteLPIs(
+	TEXT("r.SkipNaniteLPIs"),
+	GSkipNaniteLPIs,
+	TEXT("Skip Nanite primitives in the light-primitive interactions & the primitive octree as they perform GPU-driven culling separately.\n")
+	TEXT(" Values:")
+	TEXT("   1 - (auto, default) Skipping is auto-disabled if r.AllowStaticLighting is enabled for the project as it breaks some associated editor features otherwise.")
+	TEXT("   2 - (forced) Skipping is always enabled regardless of r.AllowStaticLighting. May cause issues with static lighting. Use with care."),
+	FConsoleVariableDelegate::CreateLambda([](IConsoleVariable* InVariable)
+	{
+		// Needed because the primitives need to be re-added to the scene to be removed from the octree and to have existing LPIs cleaned up. And vice versa.
+		// The cvar is not expected to be changed during runtime outside of testing.
+		FGlobalComponentRecreateRenderStateContext Context;
+	}),
+	ECVF_RenderThreadSafe);
+
+bool ShouldSkipNaniteLPIs()
+{
+	return GSkipNaniteLPIs > 1 
+		|| ( GSkipNaniteLPIs == 1 && !IsStaticLightingAllowed());
+}
 
 class FRTWriteMaskDecodeCS : public FGlobalShader
 {
diff --git a/Engine/Source/Runtime/Renderer/Private/SceneCore.cpp b/Engine/Source/Runtime/Renderer/Private/SceneCore.cpp
--- a/Engine/Source/Runtime/Renderer/Private/SceneCore.cpp
+++ b/Engine/Source/Runtime/Renderer/Private/SceneCore.cpp
@@ -112,41 +112,47 @@ uint32 FLightPrimitiveInteraction::GetMemoryPoolSize()
 	return GLightPrimitiveInteractionAllocator.GetAllocatedSize();
 }
 
-void FLightPrimitiveInteraction::Create(FLightSceneInfo* LightSceneInfo,FPrimitiveSceneInfo* PrimitiveSceneInfo)
+FLightPrimitiveInteraction::FShouldCreateResult FLightPrimitiveInteraction::ShouldCreate(FLightSceneInfo* LightSceneInfo, FPrimitiveSceneInfo* PrimitiveSceneInfo)
 {
-	LLM_SCOPE(ELLMTag::SceneRender);
-
-	// Attach the light to the primitive's static meshes.
-	bool bDynamic = true;
-	bool bRelevant = false;
-	bool bIsLightMapped = true;
-	bool bShadowMapped = false;
+	FShouldCreateResult Result;
 
 	// Determine the light's relevance to the primitive.
 	check(PrimitiveSceneInfo->Proxy && LightSceneInfo->Proxy);
-	PrimitiveSceneInfo->Proxy->GetLightRelevance(LightSceneInfo->Proxy, bDynamic, bRelevant, bIsLightMapped, bShadowMapped);
+	PrimitiveSceneInfo->Proxy->GetLightRelevance(LightSceneInfo->Proxy, Result.bDynamic, Result.bRelevant, Result.bIsLightMapped, Result.bShadowMapped);
 
 	// Mobile renders stationary and dynamic local lights as dynamic
-	bDynamic |= (PrimitiveSceneInfo->Scene->GetShadingPath() == EShadingPath::Mobile && bShadowMapped && LightSceneInfo->Proxy->IsLocalLight());
+	Result.bDynamic |= (PrimitiveSceneInfo->Scene->GetShadingPath() == EShadingPath::Mobile && Result.bShadowMapped && LightSceneInfo->Proxy->IsLocalLight());
 
-	if (bRelevant && bDynamic
+	if (Result.bRelevant && Result.bDynamic
 		// Don't let lights with static shadowing or static lighting affect primitives that should use static lighting, but don't have valid settings (lightmap res 0, etc)
 		// This prevents those components with invalid lightmap settings from causing lighting to remain unbuilt after a build
 		&& !(LightSceneInfo->Proxy->HasStaticShadowing() && PrimitiveSceneInfo->Proxy->HasStaticLighting() && !PrimitiveSceneInfo->Proxy->HasValidSettingsForStaticLighting()))
 	{
-		const bool bTranslucentObjectShadow = LightSceneInfo->Proxy->CastsTranslucentShadows() && PrimitiveSceneInfo->Proxy->CastsVolumetricTranslucentShadow();
-		const bool bInsetObjectShadow = 
+		Result.bTranslucentObjectShadow = LightSceneInfo->Proxy->CastsTranslucentShadows() && PrimitiveSceneInfo->Proxy->CastsVolumetricTranslucentShadow();
+		Result.bInsetObjectShadow = 
 			// Currently only supporting inset shadows on directional lights, but could be made to work with any whole scene shadows
 			LightSceneInfo->Proxy->GetLightType() == LightType_Directional
 			&& PrimitiveSceneInfo->Proxy->CastsInsetShadow();
 
 		// Movable directional lights determine shadow relevance dynamically based on the view and CSM settings. Interactions are only required for per-object cases.
-		if (LightSceneInfo->Proxy->GetLightType() != LightType_Directional || LightSceneInfo->Proxy->HasStaticShadowing() || bTranslucentObjectShadow || bInsetObjectShadow)
+		if (LightSceneInfo->Proxy->GetLightType() != LightType_Directional || LightSceneInfo->Proxy->HasStaticShadowing() || Result.bTranslucentObjectShadow || Result.bInsetObjectShadow)
 		{
-			// Create the light interaction.
-			FLightPrimitiveInteraction* Interaction = new FLightPrimitiveInteraction(LightSceneInfo, PrimitiveSceneInfo, bDynamic, bIsLightMapped, bShadowMapped, bTranslucentObjectShadow, bInsetObjectShadow);
-		} //-V773
+			Result.bShouldCreate = true;
+		}
 	}
+
+	return Result;
+}
+
+void FLightPrimitiveInteraction::Create(FLightSceneInfo* LightSceneInfo,FPrimitiveSceneInfo* PrimitiveSceneInfo)
+{
+	LLM_SCOPE(ELLMTag::SceneRender);
+
+	FShouldCreateResult Result = ShouldCreate(LightSceneInfo, PrimitiveSceneInfo);
+	if (Result.bShouldCreate)
+	{
+		FLightPrimitiveInteraction* Interaction = new FLightPrimitiveInteraction(LightSceneInfo, PrimitiveSceneInfo, Result.bDynamic, Result.bIsLightMapped, Result.bShadowMapped, Result.bTranslucentObjectShadow, Result.bInsetObjectShadow);
+	} //-V773
 }
 
 void FLightPrimitiveInteraction::Destroy(FLightPrimitiveInteraction* LightPrimitiveInteraction)
diff --git a/Engine/Source/Runtime/Renderer/Private/SceneCore.h b/Engine/Source/Runtime/Renderer/Private/SceneCore.h
--- a/Engine/Source/Runtime/Renderer/Private/SceneCore.h
+++ b/Engine/Source/Runtime/Renderer/Private/SceneCore.h
@@ -25,6 +25,19 @@ public:
 
 	/** Creates an interaction for a light-primitive pair. */
 	static void InitializeMemoryPool();
+
+	struct FShouldCreateResult
+	{
+		bool bShouldCreate = false;
+		bool bDynamic = true;
+		bool bRelevant = false;
+		bool bIsLightMapped = true;
+		bool bShadowMapped = false;
+		bool bTranslucentObjectShadow = false;
+		bool bInsetObjectShadow = false;
+	};
+
+	static FShouldCreateResult ShouldCreate(FLightSceneInfo* LightSceneInfo,FPrimitiveSceneInfo* PrimitiveSceneInfo);
 	static void Create(FLightSceneInfo* LightSceneInfo,FPrimitiveSceneInfo* PrimitiveSceneInfo);
 	static void Destroy(FLightPrimitiveInteraction* LightPrimitiveInteraction);
 
diff --git a/Engine/Source/Runtime/Renderer/Private/ScenePrivate.h b/Engine/Source/Runtime/Renderer/Private/ScenePrivate.h
--- a/Engine/Source/Runtime/Renderer/Private/ScenePrivate.h
+++ b/Engine/Source/Runtime/Renderer/Private/ScenePrivate.h
@@ -3643,8 +3643,6 @@ public:
 
 	void DumpMeshDrawCommandMemoryStats();
 
-	void CreateLightPrimitiveInteractionsForPrimitive(FPrimitiveSceneInfo* PrimitiveInfo);
-
 	FORCEINLINE TArray<FCachedShadowMapData>* GetCachedShadowMapDatas(int32 LightID)
 	{
 		return CachedShadowMaps.Find(LightID);
@@ -3802,6 +3800,15 @@ public:
 	 * IF using this to drive an async task, the core light scene info may be used, but primitive scene updates will still be ongoing (e.g., light/primitive interactions may change).
 	 */
 	FSceneLightSceneInfoUpdateDelegate OnPostLightSceneInfoUpdate;
+
+	/**
+	 * Retrieves the lights interacting with the passed in primitive and adds them to the out array.
+	 * Render thread version of function.
+	 * @param	PrimitiveSceneProxy		Proxy of Primitive to retrieve interacting lights for
+	 * @param	RelevantLights	[out]	Array of lights interacting with primitive
+	 */
+	void GetRelevantLights_RenderThread( const FPrimitiveSceneProxy* PrimitiveSceneProxy, TArray<const FLightSceneProxy*> &OutRelevantLights ) const;
+
 protected:
 
 private:
@@ -3829,14 +3836,6 @@ private:
 	 */
 	void CheckPrimitiveArrays(int MaxTypeOffsetIndex = -1);
 
-	/**
-	 * Retrieves the lights interacting with the passed in primitive and adds them to the out array.
-	 * Render thread version of function.
-	 * @param	Primitive				Primitive to retrieve interacting lights for
-	 * @param	RelevantLights	[out]	Array of lights interacting with primitive
-	 */
-	void GetRelevantLights_RenderThread( UPrimitiveComponent* Primitive, TArray<const ULightComponent*>* RelevantLights ) const;
-
 	/**
 	 * Adds a primitive to the scene.  Called in the rendering thread by AddPrimitive.
 	 * @param PrimitiveSceneInfo - The primitive being added.
diff --git a/Engine/Source/Runtime/Renderer/Private/SceneVisibility.cpp b/Engine/Source/Runtime/Renderer/Private/SceneVisibility.cpp
--- a/Engine/Source/Runtime/Renderer/Private/SceneVisibility.cpp
+++ b/Engine/Source/Runtime/Renderer/Private/SceneVisibility.cpp
@@ -4768,16 +4768,17 @@ void FSceneRenderer::PreVisibilityFrameSetup(FRDGBuilder& GraphBuilder)
 			for (TConstSetBitIterator<> It(Scene->PrimitivesSelected); It; ++It)
 			{
 				const FPrimitiveSceneInfo* PrimitiveSceneInfo = Scene->Primitives[It.GetIndex()];
-				FLightPrimitiveInteraction* LightList = PrimitiveSceneInfo->LightList;
-				while (LightList)
-				{
-					const FLightSceneInfo* LightSceneInfo = LightList->GetLight();
 
+				TArray<const FLightSceneProxy*> RelevantLights;
+				Scene->GetRelevantLights_RenderThread(PrimitiveSceneInfo->Proxy, RelevantLights);
+
+				for (const FLightSceneProxy* LightSceneProxy : RelevantLights)
+				{
 					bool bDynamic = true;
 					bool bRelevant = false;
 					bool bLightMapped = true;
 					bool bShadowMapped = false;
-					PrimitiveSceneInfo->Proxy->GetLightRelevance(LightSceneInfo->Proxy, bDynamic, bRelevant, bLightMapped, bShadowMapped);
+					PrimitiveSceneInfo->Proxy->GetLightRelevance(LightSceneProxy, bDynamic, bRelevant, bLightMapped, bShadowMapped);
 
 					if (bRelevant)
 					{
@@ -4787,10 +4788,9 @@ void FSceneRenderer::PreVisibilityFrameSetup(FRDGBuilder& GraphBuilder)
 						{
 							FViewInfo& View = Views[ViewIndex];
 							FViewElementPDI LightInfluencesPDI(&View, nullptr, &View.DynamicPrimitiveCollector);
-							LightInfluencesPDI.DrawLine(PrimitiveSceneInfo->Proxy->GetBounds().Origin, LightSceneInfo->Proxy->GetLightToWorld().GetOrigin(), LineColor, SDPG_World);
+							LightInfluencesPDI.DrawLine(PrimitiveSceneInfo->Proxy->GetBounds().Origin, LightSceneProxy->GetLightToWorld().GetOrigin(), LineColor, SDPG_World);
 						}
 					}
-					LightList = LightList->GetNextLight();
 				}
 			}
 		}
diff --git a/Engine/Source/Runtime/Renderer/Private/ShadowSetup.cpp b/Engine/Source/Runtime/Renderer/Private/ShadowSetup.cpp
--- a/Engine/Source/Runtime/Renderer/Private/ShadowSetup.cpp
+++ b/Engine/Source/Runtime/Renderer/Private/ShadowSetup.cpp
@@ -2427,7 +2427,7 @@ bool FProjectedShadowInfo::HasSubjectPrims() const
 		|| ShadowDepthPass.HasAnyDraw()
 		|| SubjectMeshCommandBuildRequests.Num() > 0
 		|| ShadowDepthPassVisibleCommands.Num() > 0
-		|| (bContainsNaniteSubjects && bNaniteGeometry);
+		|| ((bContainsNaniteSubjects || ShouldSkipNaniteLPIs()) && bNaniteGeometry);
 }
 
 void FProjectedShadowInfo::AddReceiverPrimitive(FPrimitiveSceneInfo* PrimitiveSceneInfo)
