Add blue noise for hair env. lighting.

diff --git a/Engine/Shaders/Private/HairStrands/HairStrandsEnvironmentLighting.usf b/Engine/Shaders/Private/HairStrands/HairStrandsEnvironmentLighting.usf
--- a/Engine/Shaders/Private/HairStrands/HairStrandsEnvironmentLighting.usf
+++ b/Engine/Shaders/Private/HairStrands/HairStrandsEnvironmentLighting.usf
@@ -77,6 +77,10 @@
 #define VOXEL_TRAVERSAL_DEBUG 0
 #include "HairStrandsVoxelPageTraversal.ush"
 
+#define USE_BLUE_NOISE 1
+#include "../BlueNoise.ush"
+DEFINE_HAIR_BLUE_NOISE_JITTER_FUNCTION
+
 ////////////////////////////////////////////////////////////////////////////////////////////////////////////////////////////////////////////////////////////////
 // Defines
 
@@ -688,7 +692,11 @@ float3 ReflectionEnvironmentUniformSampling(
 	float3 AccLighting = 0;
 	for (uint SampleIt = 0; SampleIt < SampleCount; ++SampleIt)
 	{
+		#if USE_BLUE_NOISE
+		const float4 Random4 = GetHairBlueNoiseJitter(PixelPosition, SampleIt, SampleCount, View.StateFrameIndexMod8);
+		#else
 		const float4 Random4 = ComputeRandom4(PixelPosition, SampleIt, SampleCount, View.StateFrameIndexMod8);
+		#endif
 
 		// Depth bias
 		const float3 SampleL    = UniformSampleSphere(frac(Random4.xy)).xyz;
diff --git a/Engine/Source/Runtime/Renderer/Private/HairStrands/HairStrandsEnvironment.cpp b/Engine/Source/Runtime/Renderer/Private/HairStrands/HairStrandsEnvironment.cpp
--- a/Engine/Source/Runtime/Renderer/Private/HairStrands/HairStrandsEnvironment.cpp
+++ b/Engine/Source/Runtime/Renderer/Private/HairStrands/HairStrandsEnvironment.cpp
@@ -245,6 +245,7 @@ class FHairEnvironmentLightingPS : public FGlobalShader
 		SHADER_PARAMETER_STRUCT_INCLUDE(LumenRadianceCache::FRadianceCacheInterpolationParameters, RadianceCache)
 		SHADER_PARAMETER_STRUCT_INCLUDE(ShaderPrint::FShaderParameters, ShaderPrintParameters)
 		SHADER_PARAMETER_STRUCT_INCLUDE(FHairStrandsDebugData::FWriteParameters, DebugData)
+		SHADER_PARAMETER_STRUCT_REF(FBlueNoise, BlueNoise)
 
 		SHADER_PARAMETER(uint32, bDynamicSkyLight)
 		SHADER_PARAMETER(uint32, bHasStaticLighting)
@@ -398,6 +399,9 @@ static void AddHairStrandsEnvironmentLightingPassPS(
 	ParametersPS->ReflectionCaptureData = View.ReflectionCaptureUniformBuffer;
 	ParametersPS->ReflectionsParameters = CreateReflectionUniformBuffer(GraphBuilder, View);
 	
+	FBlueNoise BlueNoise = GetBlueNoiseGlobalParameters();
+	ParametersPS->BlueNoise = CreateUniformBufferImmediate(BlueNoise, EUniformBufferUsage::UniformBuffer_SingleDraw);
+
 	const bool bScreenTrace = View.HZB && GHairStrandsSkyLighting_ScreenTraceOcclusion && IntegrationType != EHairLightingIntegrationType::SceneColor;
 	ParametersPS->FurthestHZBTexture = GSystemTextures.GetDepthDummy(GraphBuilder);
 	ParametersPS->FurthestHZBTextureSampler = TStaticSamplerState<SF_Point>::GetRHI();
