Exclude light with contact shadow from being rendered with clustered shading, since the pass does not support this.

diff --git a/Engine/Source/Runtime/Renderer/Private/LightRendering.cpp b/Engine/Source/Runtime/Renderer/Private/LightRendering.cpp
--- a/Engine/Source/Runtime/Renderer/Private/LightRendering.cpp
+++ b/Engine/Source/Runtime/Renderer/Private/LightRendering.cpp
@@ -1244,11 +1244,12 @@ void FSceneRenderer::GatherAndSortLights(FSortedLightSetSceneInfo& OutSortedLigh
 					// And also that are not directional (mostly because it doesn't make so much sense to insert them into every grid cell in the universe)
 					// In the forward case one directional light gets put into its own variables, and in the deferred case it gets a full-screen pass.
 					// Usually it'll have shadows and stuff anyway.
-					// Rect lights are not supported as the performance impact is significant even if not used, for now, left for trad. deferred.
+					// Contact shadow are not supported.
 					const bool bClusteredDeferredSupported =
 						(!SortedLightInfo->SortKey.Fields.bShadowed || bShadowedLightsInClustered) &&
-						(!SortedLightInfo->SortKey.Fields.bLightFunction || bUseLightFunctionAtlas)
+						(!SortedLightInfo->SortKey.Fields.bLightFunction || (bUseLightFunctionAtlas && SortedLightInfo->bIsCompatibleWithLightFunctionAtlas))
 						&& LightSceneInfoCompact.LightType != LightType_Directional
+						&& LightSceneInfo->Proxy->GetContactShadowLength() == 0
 						&& !bHandledByManyLights;
 
 					// Track feature available accross all lights
