Improve VSM debug rendering.
 - Fix label placement with DPI scaling
 - Render overlays below editor primitives
 - Use grayscale scene color to improve visibility of overlays
 - Use alphablending to blend overlay with scene color

#tests Tested locally and ran preflight which failed for unrelated reasons

diff --git a/Engine/Shaders/Private/VirtualShadowMaps/Desaturate.usf b/Engine/Shaders/Private/VirtualShadowMaps/Desaturate.usf
new file mode 100644
--- /dev/null
+++ b/Engine/Shaders/Private/VirtualShadowMaps/Desaturate.usf
@@ -0,0 +1,21 @@
+// Copyright Epic Games, Inc. All Rights Reserved.
+
+#include "../Common.ush"
+
+Texture2D InputTexture;
+SamplerState InputSampler;
+
+float Grayscale(float3 SourceClr)
+{
+	return 0.299 * SourceClr.r + 0.587 * SourceClr.g + 0.114 * SourceClr.b;
+}
+
+float4 DesaturatePS(
+	noperspective float4 UVAndScreenPos : TEXCOORD0
+) : SV_Target0
+{
+	float2 UV = UVAndScreenPos.xy;
+	float4 Color = Texture2DSample(InputTexture, InputSampler, UV);
+	Color.rgb = Grayscale(Color);
+	return Color;
+}
diff --git a/Engine/Shaders/Private/VirtualShadowMaps/VirtualShadowMapDebug.usf b/Engine/Shaders/Private/VirtualShadowMaps/VirtualShadowMapDebug.usf
--- a/Engine/Shaders/Private/VirtualShadowMaps/VirtualShadowMapDebug.usf
+++ b/Engine/Shaders/Private/VirtualShadowMaps/VirtualShadowMapDebug.usf
@@ -21,7 +21,6 @@ int VirtualShadowMapId;
 StructuredBuffer<FPhysicalPageMetaData> PhysicalPageMetaData;
 RWTexture2D<float4> OutVisualize;
 
-static const float3 BackGroundColor = float3(0.1f, 0.1f, 0.1f);
 static const float3 BlackColor = float3(0.0f, 0.0f, 0.0f);
 
 
@@ -109,7 +108,7 @@ float3 DebugVisualizeClipmap(float2 PixelUV, FVirtualShadowMapProjectionShaderDa
 	return ResultColor;
 }
 
-float3 DebugVisualizeClipmap(uint2 OutputPixel, int VirtualShadowMapId)
+float4 DebugVisualizeClipmap(uint2 OutputPixel, int VirtualShadowMapId)
 {
 	// Figure out the square area do draw in (could be any sub-rect)
 	uint SquareDomainDim = min(DebugTargetWidth, DebugTargetHeight) - BorderWidth * 2U;
@@ -117,7 +116,7 @@ float3 DebugVisualizeClipmap(uint2 OutputPixel, int VirtualShadowMapId)
 	// fill background pixels
 	if (any(or(OutputPixel < BorderWidth, OutputPixel >(SquareDomainDim + BorderWidth))))
 	{
-		return BackGroundColor;
+		return float4(0, 0, 0, 0);
 	}
 
 	// The pixel we are drawing in 0-1 range UV within the domain
@@ -126,10 +125,10 @@ float3 DebugVisualizeClipmap(uint2 OutputPixel, int VirtualShadowMapId)
 	FVirtualShadowMapProjectionShaderData ProjectionData = GetVirtualShadowMapProjectionData(VirtualShadowMapId);
 	if (ProjectionData.LightType != LIGHT_TYPE_DIRECTIONAL)
 	{
-		return BackGroundColor;
+		return float4(0, 0, 0, 0);
 	}
 
-	return DebugVisualizeClipmap(PixelUV, ProjectionData);
+	return float4(DebugVisualizeClipmap(PixelUV, ProjectionData), 1.0);
 }
 
 
@@ -140,24 +139,18 @@ void DebugVisualizeVirtualSmCS(uint2 PixelPos : SV_DispatchThreadID)
 	{
 		return;
 	}
-
 	// NOTE: Important to *not* write output if it isn't a recognized mode, as a different
 	// pass may write that output instead.
-	float3 Output = float3(1, 0, 1);
-	bool bWriteOutput = true;
+	float4 Output = float4(1, 0, 1, 0);
 
 	if (VisualizeModeId == VIRTUAL_SHADOW_MAP_VISUALIZE_CLIPMAP_VIRTUAL_SPACE)
 	{
 		Output = DebugVisualizeClipmap(PixelPos, VirtualShadowMapId);
 	}
-	else
-	{
-		bWriteOutput = false;
-	}
 
-	if (bWriteOutput)
+	if (Output.a > 0)
 	{
-		OutVisualize[PixelPos] = float4(Output, 1.0f);
+		OutVisualize[PixelPos] = Output;
 	}
 }
 
diff --git a/Engine/Shaders/Private/VirtualShadowMaps/VirtualShadowMapProjection.usf b/Engine/Shaders/Private/VirtualShadowMaps/VirtualShadowMapProjection.usf
--- a/Engine/Shaders/Private/VirtualShadowMaps/VirtualShadowMapProjection.usf
+++ b/Engine/Shaders/Private/VirtualShadowMaps/VirtualShadowMapProjection.usf
@@ -485,7 +485,9 @@ void OutputVisualize(
 
 		if ( VisualizeModeId == VIRTUAL_SHADOW_MAP_VISUALIZE_SHADOW_FACTOR )
 		{
-			Output = Result.ShadowFactor.xxx;
+			float3 Lit = float3(1,1,0);
+			float3 Shadow = float3(0,0,1);
+			Output = lerp(Shadow, Lit, Result.ShadowFactor);
 		}
 		else if ( VisualizeModeId == VIRTUAL_SHADOW_MAP_VISUALIZE_CLIPMAP_OR_MIP )
 		{
diff --git a/Engine/Source/Runtime/Renderer/Private/PostProcess/PostProcessing.cpp b/Engine/Source/Runtime/Renderer/Private/PostProcess/PostProcessing.cpp
--- a/Engine/Source/Runtime/Renderer/Private/PostProcess/PostProcessing.cpp
+++ b/Engine/Source/Runtime/Renderer/Private/PostProcess/PostProcessing.cpp
@@ -378,8 +378,10 @@ void AddPostProcessingPasses(
 		VisualizeSkyLightIlluminanceMeter,
 		VisualizeLightFunctionAtlas,
 		VisualizeLevelInstance,
+		VisualizeVirtualShadowMaps_PreEditorPrimitives,
 		SelectionOutline,
 		EditorPrimitive,
+		VisualizeVirtualShadowMaps_PostEditorPrimitives,
 		VisualizeShadingModels,
 		VisualizeGBufferHints,
 		VisualizeSubsurface,
@@ -433,8 +435,10 @@ void AddPostProcessingPasses(
 		TEXT("VisualizeSkyLightIlluminanceMeter"),
 		TEXT("VisualizeLightFunctionAtlas"),
 		TEXT("VisualizeLevelInstance"),
+		TEXT("VisualizeVirtualShadowMaps_PreEditorPrimitives"),
 		TEXT("SelectionOutline"),
 		TEXT("EditorPrimitive"),
+		TEXT("VisualizeVirtualShadowMaps_PostEditorPrimitives"),
 		TEXT("VisualizeShadingModels"),
 		TEXT("VisualizeGBufferHints"),
 		TEXT("VisualizeSubsurface"),
@@ -473,6 +477,8 @@ void AddPostProcessingPasses(
 	PassSequence.SetEnabled(EPass::VisualizeSkyAtmosphere, Scene&& View.Family && View.Family->EngineShowFlags.VisualizeSkyAtmosphere && ShouldRenderSkyAtmosphereDebugPasses(Scene, View.Family->EngineShowFlags));
 	PassSequence.SetEnabled(EPass::VisualizeSkyLightIlluminanceMeter, Scene&& Scene->SkyLight && View.Family && View.Family->EngineShowFlags.VisualizeSkyLightIlluminance);
 	PassSequence.SetEnabled(EPass::VisualizeLightFunctionAtlas, Scene && Scene->LightFunctionAtlasSceneData.GetLightFunctionAtlasEnabled() && View.Family && View.Family->EngineShowFlags.VisualizeLightFunctionAtlas);
+	PassSequence.SetEnabled(EPass::VisualizeVirtualShadowMaps_PreEditorPrimitives, EngineShowFlags.VisualizeVirtualShadowMap && VirtualShadowMapArray != nullptr);
+	PassSequence.SetEnabled(EPass::VisualizeVirtualShadowMaps_PostEditorPrimitives, EngineShowFlags.VisualizeVirtualShadowMap && VirtualShadowMapArray != nullptr);
 	PassSequence.SetEnabled(EPass::VisualizeLevelInstance, GIsEditor && EngineShowFlags.EditingLevelInstance && EngineShowFlags.VisualizeLevelInstanceEditing && !bVisualizeHDR);
 	PassSequence.SetEnabled(EPass::SelectionOutline, GIsEditor && EngineShowFlags.Selection && EngineShowFlags.SelectionOutline && !EngineShowFlags.Wireframe && !bVisualizeHDR);
 	PassSequence.SetEnabled(EPass::EditorPrimitive, FSceneRenderer::ShouldCompositeEditorPrimitives(View));
@@ -480,6 +486,8 @@ void AddPostProcessingPasses(
 	PassSequence.SetEnabled(EPass::VisualizeSkyAtmosphere, false);
 	PassSequence.SetEnabled(EPass::VisualizeSkyLightIlluminanceMeter, false);
 	PassSequence.SetEnabled(EPass::VisualizeLightFunctionAtlas, false);
+	PassSequence.SetEnabled(EPass::VisualizeVirtualShadowMaps_PreEditorPrimitives, false);
+	PassSequence.SetEnabled(EPass::VisualizeVirtualShadowMaps_PostEditorPrimitives, false);
 	PassSequence.SetEnabled(EPass::VisualizeLevelInstance, false);
 	PassSequence.SetEnabled(EPass::SelectionOutline, false);
 	PassSequence.SetEnabled(EPass::EditorPrimitive, false);
@@ -1476,6 +1484,7 @@ void AddPostProcessingPasses(
 		}
 	}
 
+
 	if (PassSequence.IsEnabled(EPass::VisualizeLevelInstance))
 	{
 		FVisualizeLevelInstanceInputs PassInputs;
@@ -1486,7 +1495,16 @@ void AddPostProcessingPasses(
 
 		SceneColor = AddVisualizeLevelInstancePass(GraphBuilder, View, SceneUniformBuffer, PassInputs, NaniteRasterResults);
 	}
-
+#endif //WITH_EDITOR
+	
+	if (PassSequence.IsEnabled(EPass::VisualizeVirtualShadowMaps_PreEditorPrimitives))
+	{
+		FScreenPassRenderTarget OverrideOutput;
+		PassSequence.AcceptOverrideIfLastPass(EPass::VisualizeVirtualShadowMaps_PreEditorPrimitives, OverrideOutput);
+		SceneColor = VirtualShadowMapArray->AddVisualizePass(GraphBuilder, View, ViewIndex, EVSMVisualizationPostPass::PreEditorPrimitives, SceneColor, OverrideOutput);
+	}
+	
+#if WITH_EDITOR
 	if (PassSequence.IsEnabled(EPass::SelectionOutline))
 	{
 		FSelectionOutlineInputs PassInputs;
@@ -1511,7 +1529,15 @@ void AddPostProcessingPasses(
 
 		SceneColor = AddEditorPrimitivePass(GraphBuilder, View, PassInputs, InstanceCullingManager);
 	}
-#endif
+#endif //WITH_EDITOR
+	
+	if (PassSequence.IsEnabled(EPass::VisualizeVirtualShadowMaps_PostEditorPrimitives))
+	{
+		FScreenPassRenderTarget OverrideOutput;
+		PassSequence.AcceptOverrideIfLastPass(EPass::VisualizeVirtualShadowMaps_PostEditorPrimitives, OverrideOutput);
+		SceneColor = VirtualShadowMapArray->AddVisualizePass(GraphBuilder, View, ViewIndex, EVSMVisualizationPostPass::PostEditorPrimitives, SceneColor, OverrideOutput);
+	}
+
 
 	if (PassSequence.IsEnabled(EPass::VisualizeShadingModels))
 	{
@@ -1762,11 +1788,6 @@ void AddPostProcessingPasses(
 			AddVisualizeNanitePass(GraphBuilder, View, SceneColor, *NaniteRasterResults);
 		}
 
-		if (EngineShowFlags.VisualizeVirtualShadowMap && VirtualShadowMapArray != nullptr)
-		{
-			VirtualShadowMapArray->AddVisualizePass(GraphBuilder, View, ViewIndex, SceneColor);
-		}
-
 		#if WITH_EDITOR
 		if (CVarGBufferPicking.GetValueOnRenderThread())
 		{
diff --git a/Engine/Source/Runtime/Renderer/Private/VirtualShadowMaps/VirtualShadowMapArray.cpp b/Engine/Source/Runtime/Renderer/Private/VirtualShadowMaps/VirtualShadowMapArray.cpp
--- a/Engine/Source/Runtime/Renderer/Private/VirtualShadowMaps/VirtualShadowMapArray.cpp
+++ b/Engine/Source/Runtime/Renderer/Private/VirtualShadowMaps/VirtualShadowMapArray.cpp
@@ -3453,32 +3453,53 @@ uint32 FVirtualShadowMapArray::AddRenderViews(const FProjectedShadowInfo* Projec
 	return uint32(NumMaps);
 }
 
-void FVirtualShadowMapArray::AddVisualizePass(FRDGBuilder& GraphBuilder, const FViewInfo& View, int32 ViewIndex, FScreenPassTexture Output)
+#if !UE_BUILD_SHIPPING
+class FDesaturatePS : public FGlobalShader
+{
+	DECLARE_GLOBAL_SHADER(FDesaturatePS);
+	SHADER_USE_PARAMETER_STRUCT(FDesaturatePS, FGlobalShader);
+
+	BEGIN_SHADER_PARAMETER_STRUCT(FParameters, )
+		SHADER_PARAMETER_RDG_TEXTURE(Texture2D, InputTexture)
+		SHADER_PARAMETER_SAMPLER(SamplerState, InputSampler)
+		RENDER_TARGET_BINDING_SLOTS()
+	END_SHADER_PARAMETER_STRUCT()
+};
+IMPLEMENT_GLOBAL_SHADER(FDesaturatePS, "/Engine/Private/VirtualShadowMaps/Desaturate.usf", "DesaturatePS", SF_Pixel);
+#endif //!UE_BUILD_SHIPPING
+
+FScreenPassTexture FVirtualShadowMapArray::AddVisualizePass(FRDGBuilder& GraphBuilder, const FViewInfo& View, int32 ViewIndex, EVSMVisualizationPostPass Pass, FScreenPassTexture& SceneColor, FScreenPassRenderTarget& OverrideOutput)
 {
+	FScreenPassTexture Output = SceneColor;
+
 #if !UE_BUILD_SHIPPING
+
 	if (!IsAllocated() || DebugVisualizationOutput.IsEmpty())
 	{
-		return;
+		return MoveTemp(SceneColor);
 	}
 
 	const FVirtualShadowMapVisualizationData& VisualizationData = GetVirtualShadowMapVisualizationData();
-	if (VisualizationData.IsActive() && VisualizeLight[ViewIndex].IsValid())
-	{	
-		FCopyRectPS::FParameters* Parameters = GraphBuilder.AllocParameters<FCopyRectPS::FParameters>();
-		Parameters->InputTexture = DebugVisualizationOutput[ViewIndex];
-		Parameters->InputSampler = TStaticSamplerState<SF_Bilinear, AM_Clamp, AM_Clamp, AM_Clamp>::GetRHI();
-		Parameters->RenderTargets[0] = FRenderTargetBinding(Output.Texture, ERenderTargetLoadAction::ENoAction);
+	if (!VisualizationData.IsActive() || !VisualizeLight[ViewIndex].IsValid())
+	{
+		return MoveTemp(SceneColor);
+	}
 
-		TShaderMapRef<FCopyRectPS> PixelShader(View.ShaderMap);
+	RDG_EVENT_SCOPE(GraphBuilder, "VirtualShadowMapsVisualization");
+
+	FScreenPassTextureViewport InputViewport(DebugVisualizationOutput[ViewIndex]->Desc.Extent);
+	FScreenPassTextureViewport OutputViewport(Output);
+	FScreenPassRenderTarget OutputTarget(Output.Texture, OutputViewport.Rect, ERenderTargetLoadAction::ELoad);
 
-		FScreenPassTextureViewport InputViewport(DebugVisualizationOutput[ViewIndex]->Desc.Extent);
-		FScreenPassTextureViewport OutputViewport(Output);
+	int ActiveModeId = VisualizationData.GetActiveModeID();
 
+	// Resize viewport for layout
+	const int32 VisualizeLayout = CVarVisualizeLayout.GetValueOnRenderThread();
+	{
 		// See CVarVisualizeLayout documentation
-		const int32 VisualizeLayout = CVarVisualizeLayout.GetValueOnRenderThread();
 		if (VisualizeLayout == 1)		// Thumbnail
 		{
-			const int32 TileWidth  = View.UnscaledViewRect.Width() / 3;
+			const int32 TileWidth = View.UnscaledViewRect.Width() / 3;
 			const int32 TileHeight = View.UnscaledViewRect.Height() / 3;
 
 			OutputViewport.Rect.Max = OutputViewport.Rect.Min + FIntPoint(TileWidth, TileHeight);
@@ -3488,28 +3509,92 @@ void FVirtualShadowMapArray::AddVisualizePass(FRDGBuilder& GraphBuilder, const F
 			InputViewport.Rect.Max.X = InputViewport.Rect.Min.X + (InputViewport.Rect.Width() / 2);
 			OutputViewport.Rect.Max.X = OutputViewport.Rect.Min.X + (OutputViewport.Rect.Width() / 2);
 		}
+	}
 
-		// Use separate input and output viewports w/ bilinear sampling to properly support dynamic resolution scaling
-		AddDrawScreenPass(GraphBuilder, RDG_EVENT_NAME("DrawTexture"), View, OutputViewport, InputViewport, PixelShader, Parameters, EScreenPassDrawFlags::None);
-		
-		// Visualization light name
+	auto DrawDebugVisualizationOutput = [&]()
+	{
+		TShaderMapRef<FScreenPassVS> VertexShader(View.ShaderMap);
+		TShaderMapRef<FCopyRectPS> PixelShader(View.ShaderMap);
+
+		FCopyRectPS::FParameters* Parameters = GraphBuilder.AllocParameters<FCopyRectPS::FParameters>();
+		Parameters->InputTexture = DebugVisualizationOutput[ViewIndex];
+		Parameters->InputSampler = TStaticSamplerState<SF_Bilinear, AM_Clamp, AM_Clamp, AM_Clamp>::GetRHI();
+		Parameters->RenderTargets[0] = FRenderTargetBinding(Output.Texture, ERenderTargetLoadAction::ENoAction);
+
+		// Blend with scene color if fullscreen, use black background otherwise
+		FRHIBlendState* BlendState = 
+			VisualizeLayout == 0
+			? TStaticBlendState<CW_RGB, BO_Add, BF_SourceAlpha, BF_InverseSourceAlpha, BO_Add, BF_Zero, BF_One>::GetRHI()
+			: TStaticBlendState<>::GetRHI();
+		FRHIDepthStencilState* DepthStencilState = FScreenPassPipelineState::FDefaultDepthStencilState::GetRHI();
+
+		AddDrawScreenPass(GraphBuilder, RDG_EVENT_NAME("DrawTexture"), View, OutputViewport, InputViewport, VertexShader, PixelShader, BlendState, DepthStencilState, Parameters, EScreenPassDrawFlags::None);
+	};
+	
+	if (Pass == EVSMVisualizationPostPass::PreEditorPrimitives)
+	{
+		// Desaturate scene color
+		{
+			TShaderMapRef<FScreenPassVS> VertexShader(View.ShaderMap);
+			TShaderMapRef<FDesaturatePS> DesaturatePixelShader(View.ShaderMap);
+			TShaderMapRef<FCopyRectPS> PixelShader(View.ShaderMap);
+
+			FDesaturatePS::FParameters* Parameters = GraphBuilder.AllocParameters<FDesaturatePS::FParameters>();
+			Parameters->InputTexture = SceneColor.Texture;
+			Parameters->InputSampler = TStaticSamplerState<SF_Bilinear, AM_Clamp, AM_Clamp, AM_Clamp>::GetRHI();
+			Parameters->RenderTargets[0] = FRenderTargetBinding(Output.Texture, ERenderTargetLoadAction::ENoAction);
+			AddDrawScreenPass(GraphBuilder, RDG_EVENT_NAME("Desaturate"), View, OutputViewport, OutputViewport, VertexShader, DesaturatePixelShader, Parameters, EScreenPassDrawFlags::None);
+		}
+
+		// Render stuff that blends in with scene
+		if (ActiveModeId == VIRTUAL_SHADOW_MAP_VISUALIZE_SHADOW_FACTOR
+			|| ActiveModeId == VIRTUAL_SHADOW_MAP_VISUALIZE_CLIPMAP_OR_MIP
+			|| ActiveModeId == VIRTUAL_SHADOW_MAP_VISUALIZE_VIRTUAL_PAGE
+			|| ActiveModeId == VIRTUAL_SHADOW_MAP_VISUALIZE_CACHED_PAGE
+			|| ActiveModeId == VIRTUAL_SHADOW_MAP_VISUALIZE_SMRT_RAY_COUNT
+			|| ActiveModeId == VIRTUAL_SHADOW_MAP_VISUALIZE_DIRTY_PAGE
+			|| ActiveModeId == VIRTUAL_SHADOW_MAP_VISUALIZE_GPU_INVALIDATED_PAGE
+			|| ActiveModeId == VIRTUAL_SHADOW_MAP_VISUALIZE_MERGED_PAGE
+			|| ActiveModeId == VIRTUAL_SHADOW_MAP_VISUALIZE_NANITE_OVERDRAW)
+		{
+			DrawDebugVisualizationOutput();
+		}
+	}
+	else if (Pass == EVSMVisualizationPostPass::PostEditorPrimitives)
+	{
+		// Render stuff that is not part of scene, e.g. UI
+		if (ActiveModeId == VIRTUAL_SHADOW_MAP_VISUALIZE_CLIPMAP_VIRTUAL_SPACE
+			|| ActiveModeId == VIRTUAL_SHADOW_MAP_VISUALIZE_GENERAL_DEBUG)
 		{
-			FScreenPassRenderTarget OutputTarget(Output.Texture, View.UnscaledViewRect, ERenderTargetLoadAction::ELoad);
+			DrawDebugVisualizationOutput();
+		}
 
-			AddDrawCanvasPass(GraphBuilder, RDG_EVENT_NAME("Labels"), View, OutputTarget,
-				[&VisualizeLight=VisualizeLight[ViewIndex], &OutputViewport=OutputViewport](FCanvas& Canvas)
+		AddDrawCanvasPass(GraphBuilder, RDG_EVENT_NAME("Labels"), View, OutputTarget,
+			[&VisualizeLight = VisualizeLight[ViewIndex], &OutputViewport = OutputViewport](FCanvas& Canvas)
 			{
+				const float DPIScale = Canvas.GetDPIScale();
+				Canvas.SetBaseTransform(FMatrix(FScaleMatrix(DPIScale) * Canvas.CalcBaseTransform2D(Canvas.GetViewRect().Width(), Canvas.GetViewRect().Height())));
+
 				const FLinearColor LabelColor(1, 1, 0);
 				Canvas.DrawShadowedString(
-					OutputViewport.Rect.Min.X + 8,
-					OutputViewport.Rect.Max.Y - 19,
+					(OutputViewport.Rect.Min.X + 8) / DPIScale,
+					(OutputViewport.Rect.Max.Y - 19) / DPIScale,
 					*VisualizeLight.GetLightName(),
 					GetStatsFont(),
 					LabelColor);
 			});
-		}
 	}
-#endif
+
+
+	if (OverrideOutput.IsValid())
+	{
+		AddDrawTexturePass(GraphBuilder, View, Output, OverrideOutput);
+		return OverrideOutput;
+	}
+
+#endif //!UE_BUILD_SHIPPING
+
+	return MoveTemp(Output);
 }
 
 float FVirtualShadowMapArray::InterpolateResolutionBias(float BiasNonMoving, float BiasMoving, float LightMobilityFactor)
diff --git a/Engine/Source/Runtime/Renderer/Private/VirtualShadowMaps/VirtualShadowMapArray.h b/Engine/Source/Runtime/Renderer/Private/VirtualShadowMaps/VirtualShadowMapArray.h
--- a/Engine/Source/Runtime/Renderer/Private/VirtualShadowMaps/VirtualShadowMapArray.h
+++ b/Engine/Source/Runtime/Renderer/Private/VirtualShadowMaps/VirtualShadowMapArray.h
@@ -5,6 +5,7 @@
 #include "SceneManagement.h"
 #include "SceneView.h"
 #include "VirtualShadowMapDefinitions.h"
+#include "ScreenPass.h"
 
 struct FMinimalSceneTextures;
 struct FSortedLightSetSceneInfo;
@@ -221,6 +222,12 @@ private:
 	int FoundVirtualShadowMapId = INDEX_NONE;
 };
 
+enum class EVSMVisualizationPostPass
+{
+	PreEditorPrimitives,
+	PostEditorPrimitives
+};
+
 class FVirtualShadowMapArray
 {
 public:	
@@ -349,7 +356,7 @@ public:
 	uint32 AddRenderViews(const FProjectedShadowInfo* ProjectedShadowInfo, TConstArrayView<FViewInfo> Views, float LODScaleFactor, bool bSetHzbParams, bool bUpdateHZBMetaData, bool bClampToNearPlane, TArray<Nanite::FPackedView, SceneRenderingAllocator>& OutVirtualShadowViews);
 
 	// Add visualization composite pass, if enabled
-	void AddVisualizePass(FRDGBuilder& GraphBuilder, const FViewInfo& View, int32 ViewIndex, FScreenPassTexture Output);
+	FScreenPassTexture AddVisualizePass(FRDGBuilder& GraphBuilder, const FViewInfo& View, int32 ViewIndex, EVSMVisualizationPostPass Pass, FScreenPassTexture& SceneColor, FScreenPassRenderTarget& Output);
 
 	//
 	bool UseHzbOcclusion() const { return bUseHzbOcclusion; }
