Allocate RayTracingScenePooledBuffer directly without explicitly creating a FBufferRHIRef.
- also extended RHIGetDefaultResourceState(...) to handle BUF_AccelerationStructure.

diff --git a/Engine/Source/Runtime/RHI/Private/RHIUtilities.cpp b/Engine/Source/Runtime/RHI/Private/RHIUtilities.cpp
--- a/Engine/Source/Runtime/RHI/Private/RHIUtilities.cpp
+++ b/Engine/Source/Runtime/RHI/Private/RHIUtilities.cpp
@@ -633,6 +633,10 @@ ERHIAccess RHIGetDefaultResourceState(EBufferUsageFlags InUsage, bool bInHasInit
 	{
 		DefaultReadingState = DefaultReadingState | ERHIAccess::SRVMask;
 	}
+	if (EnumHasAnyFlags(InUsage, BUF_AccelerationStructure))
+	{
+		DefaultReadingState = DefaultReadingState | ERHIAccess::BVHRead;
+	}
 
 	// Vertex and index buffers might not have the BUF_ShaderResource flag set and just assume
 	// they are readable by default
@@ -649,6 +653,10 @@ ERHIAccess RHIGetDefaultResourceState(EBufferUsageFlags InUsage, bool bInHasInit
 		{
 			ResourceState = ERHIAccess::UAVMask;
 		}
+		else if (EnumHasAnyFlags(InUsage, BUF_AccelerationStructure))
+		{
+			ResourceState = ERHIAccess::BVHWrite;
+		}
 		else if (EnumHasAnyFlags(InUsage, BUF_ShaderResource))
 		{
 			ResourceState = DefaultReadingState | ERHIAccess::SRVMask;
diff --git a/Engine/Source/Runtime/Renderer/Private/RayTracing/RayTracingScene.cpp b/Engine/Source/Runtime/Renderer/Private/RayTracing/RayTracingScene.cpp
--- a/Engine/Source/Runtime/Renderer/Private/RayTracing/RayTracingScene.cpp
+++ b/Engine/Source/Runtime/Renderer/Private/RayTracing/RayTracingScene.cpp
@@ -90,16 +90,14 @@ void FRayTracingScene::CreateWithInitializationData(FRDGBuilder& GraphBuilder, c
 	SizeInfo.ResultSize = FMath::DivideAndRoundUp(FMath::Max(SizeInfo.ResultSize, 1ull), BufferAllocationGranularity) * BufferAllocationGranularity;
 
 	// Allocate GPU buffer if current one is too small or significantly larger than what we need.
-	if (!RayTracingSceneBuffer.IsValid() 
-		|| SizeInfo.ResultSize > RayTracingSceneBuffer->GetSize() 
-		|| SizeInfo.ResultSize < RayTracingSceneBuffer->GetSize() / 2)
+	if (!RayTracingScenePooledBuffer.IsValid()
+		|| SizeInfo.ResultSize > RayTracingScenePooledBuffer->GetSize()
+		|| SizeInfo.ResultSize < RayTracingScenePooledBuffer->GetSize() / 2)
 	{
-		FRHIResourceCreateInfo CreateInfo(TEXT("FRayTracingScene::SceneBuffer"));
-		RayTracingSceneBuffer = GraphBuilder.RHICmdList.CreateBuffer(uint32(SizeInfo.ResultSize), EBufferUsageFlags::AccelerationStructure, 0, ERHIAccess::BVHWrite, CreateInfo);
-
 		FRDGBufferDesc Desc = FRDGBufferDesc::CreateBufferDesc(1, uint32(SizeInfo.ResultSize));
 		Desc.Usage = EBufferUsageFlags::AccelerationStructure;
-		RayTracingScenePooledBuffer = new FRDGPooledBuffer(RHICmdList, RayTracingSceneBuffer, Desc, Desc.NumElements, TEXT("FRayTracingScene::SceneBuffer::RDG"));
+
+		RayTracingScenePooledBuffer = AllocatePooledBuffer(Desc, TEXT("FRayTracingScene::SceneBuffer"));
 	}
 	RayTracingSceneBufferRDG = GraphBuilder.RegisterExternalBuffer(RayTracingScenePooledBuffer);
 
@@ -466,8 +464,8 @@ FRDGBufferRef FRayTracingScene::GetBufferChecked() const
 FShaderResourceViewRHIRef FRayTracingScene::CreateLayerViewRHI(FRHICommandListBase& RHICmdList, ERayTracingSceneLayer Layer) const
 {
 	const uint8 LayerIndex = uint8(Layer);
-	checkf(RayTracingSceneBuffer, TEXT("Ray tracing scene was not created.Perhaps Create() was not called."));
-	return RHICmdList.CreateShaderResourceView(FShaderResourceViewInitializer(RayTracingSceneBuffer, RayTracingSceneRHI->GetLayerBufferOffset(LayerIndex), 0));
+	checkf(RayTracingScenePooledBuffer, TEXT("Ray tracing scene was not created.Perhaps Create() was not called."));
+	return RHICmdList.CreateShaderResourceView(FShaderResourceViewInitializer(RayTracingScenePooledBuffer->GetRHI(), RayTracingSceneRHI->GetLayerBufferOffset(LayerIndex), 0));
 }
 
 FRDGBufferSRVRef FRayTracingScene::GetLayerView(ERayTracingSceneLayer Layer) const
@@ -551,6 +549,9 @@ void FRayTracingScene::Reset(bool bInInstanceDebugDataEnabled)
 	Allocator.Flush();
 
 	RayTracingSceneRHI = nullptr;
+	RayTracingSceneBufferRDG = nullptr;
+
+	InstanceBuffer = nullptr;
 	BuildScratchBuffer = nullptr;
 	InstanceDebugBuffer = nullptr;
 	DebugInstanceGPUSceneIndexBuffer = nullptr;
@@ -572,7 +573,6 @@ void FRayTracingScene::EndFrame()
 		GeometriesToBuild.Empty();
 		UsedCoarseMeshStreamingHandles.Empty();
 
-		RayTracingSceneBuffer = nullptr;
 		RayTracingScenePooledBuffer = nullptr;
 
 #if STATS
diff --git a/Engine/Source/Runtime/Renderer/Private/RayTracing/RayTracingScene.h b/Engine/Source/Runtime/Renderer/Private/RayTracing/RayTracingScene.h
--- a/Engine/Source/Runtime/Renderer/Private/RayTracing/RayTracingScene.h
+++ b/Engine/Source/Runtime/Renderer/Private/RayTracing/RayTracingScene.h
@@ -147,8 +147,7 @@ private:
 	FRayTracingSceneRHIRef RayTracingSceneRHI;
 
 	// Persistently allocated buffer that holds the built TLAS
-	FBufferRHIRef RayTracingSceneBuffer;
-	TRefCountPtr<FRDGPooledBuffer>	RayTracingScenePooledBuffer;		
+	TRefCountPtr<FRDGPooledBuffer> RayTracingScenePooledBuffer;		
 	FRDGBufferRef RayTracingSceneBufferRDG;
 
 	// Per-layer views for the TLAS buffer that should be used in ray tracing shaders
