Path Tracer: Cleanup compaction logic to avoid clears after each bounce

Avoid redundant clear operations inserted between the raytracing dispatches. This reduces overhead, saving around 0.8ms at the default bounce count.

diff --git a/Engine/Shaders/Private/PathTracing/PathTracing.usf b/Engine/Shaders/Private/PathTracing/PathTracing.usf
--- a/Engine/Shaders/Private/PathTracing/PathTracing.usf
+++ b/Engine/Shaders/Private/PathTracing/PathTracing.usf
@@ -40,8 +40,8 @@ int2 ComputeTextureIndex(uint2 DispatchIdx)
 
 int Bounce;
 RWStructuredBuffer<FPathTracingPackedPathState> PathStateData;
-Buffer<int> ActivePaths;
-RWBuffer<int> NextActivePaths;
+RWBuffer<uint> ActivePaths;
+RWBuffer<uint> NextActivePaths;
 RWBuffer<uint> NumPathStates;
 
 
@@ -83,12 +83,12 @@ RAY_TRACING_ENTRY_RAYGEN(PathTracingMainRG)
 	uint2 DispatchIdx = DispatchRaysIndex().xy;
 	const uint2 DispatchDim = DispatchRaysDimensions().xy;
 
-	int LinearPixelIndex = DispatchIdx.x + DispatchDim.x * DispatchIdx.y;
+	uint LinearPixelIndex = DispatchIdx.x + DispatchDim.x * DispatchIdx.y;
 	if (LinearPixelIndex == 0)
 	{
-		// write other dimensions for indirect dispatch
-		NumPathStates[1] = 1;
-		NumPathStates[2] = 1;
+		// write other dimensions for indirect dispatch of the next bounce
+		NumPathStates[3 * Bounce + 4] = 1;
+		NumPathStates[3 * Bounce + 5] = 1;
 	}
 	FPathState PathState;
 	if (Bounce == 0)
@@ -96,11 +96,11 @@ RAY_TRACING_ENTRY_RAYGEN(PathTracingMainRG)
 #if PATH_TRACER_USE_ADAPTIVE_SAMPLING
 		if (Iteration > 0)
 		{
-			LinearPixelIndex = ActivePaths[LinearPixelIndex];
-			if (LinearPixelIndex < 0)
+			if (LinearPixelIndex >= NumPathStates[3 * Bounce + 0])
 			{
-				return;
+				return; // nothing left to do on this thread
 			}
+			LinearPixelIndex = ActivePaths[LinearPixelIndex];
 			DispatchIdx = uint2(LinearPixelIndex % DispatchDim.x, LinearPixelIndex / DispatchDim.x);
 		}
 #endif
@@ -109,11 +109,12 @@ RAY_TRACING_ENTRY_RAYGEN(PathTracingMainRG)
 	}
 	else
 	{
-		LinearPixelIndex = ActivePaths[LinearPixelIndex];
-		if (LinearPixelIndex < 0)
+		// compare to number of active paths from the previous bounce
+		if (LinearPixelIndex >= NumPathStates[3 * Bounce + 0])
 		{
 			return; // nothing left to do on this thread
 		}
+		LinearPixelIndex = ActivePaths[LinearPixelIndex];
 		PathState = LoadPathStateData(LinearPixelIndex);
 	}
 
@@ -127,9 +128,9 @@ RAY_TRACING_ENTRY_RAYGEN(PathTracingMainRG)
 	}
 	if (KeepGoing)
 	{
-		// NOTE: using wave instructions to reduce contention seems to run slightly slower? (tested on RTX-3080)
+	    // Add the path to the set of valid paths for the next bounce
 		uint PathStateIndex;
-		InterlockedAdd(NumPathStates[0], 1, PathStateIndex);
+		InterlockedAdd(NumPathStates[3 * Bounce + 3], 1, PathStateIndex);
 		NextActivePaths[PathStateIndex] = LinearPixelIndex;
 		StorePathStateData(PathState, LinearPixelIndex);
 	}
@@ -145,6 +146,7 @@ RAY_TRACING_ENTRY_RAYGEN(PathTracingMainRG)
 
 #if PATH_TRACER_USE_ADAPTIVE_SAMPLING
 Buffer<int> ActivePaths;
+Buffer<uint> NumPathStates;
 #endif
 
 RAY_TRACING_ENTRY_RAYGEN(PathTracingMainRG)
@@ -156,11 +158,11 @@ RAY_TRACING_ENTRY_RAYGEN(PathTracingMainRG)
 	if (Iteration > 0)
 	{
 		uint LinearPixelIndex = DispatchIdx.x + DispatchDim.x * DispatchIdx.y;
-		LinearPixelIndex = ActivePaths[LinearPixelIndex];
-		if (LinearPixelIndex < 0)
+		if (LinearPixelIndex >= NumPathStates[0])
 		{
 			return;
 		}
+		LinearPixelIndex = ActivePaths[LinearPixelIndex];
 		DispatchIdx = uint2(LinearPixelIndex % DispatchDim.x, LinearPixelIndex / DispatchDim.x);
 	}
 #endif
diff --git a/Engine/Shaders/Private/PathTracing/PathTracingAdaptiveStart.usf b/Engine/Shaders/Private/PathTracing/PathTracingAdaptiveStart.usf
--- a/Engine/Shaders/Private/PathTracing/PathTracingAdaptiveStart.usf
+++ b/Engine/Shaders/Private/PathTracing/PathTracingAdaptiveStart.usf
@@ -4,7 +4,7 @@
 #include "PathTracingAdaptiveUtils.ush"
 
 
-RWBuffer<int> NextActivePaths;
+RWBuffer<uint> NextActivePaths;
 RWBuffer<uint> NumPathStates;
 int2 TileTextureOffset;
 uint2 DispatchDim;
diff --git a/Engine/Source/Runtime/Renderer/Private/PathTracing.cpp b/Engine/Source/Runtime/Renderer/Private/PathTracing.cpp
--- a/Engine/Source/Runtime/Renderer/Private/PathTracing.cpp
+++ b/Engine/Source/Runtime/Renderer/Private/PathTracing.cpp
@@ -1163,9 +1163,9 @@ class FPathTracingRG : public FGlobalShader
 		// extra parameters required for path compacting kernel
 		SHADER_PARAMETER(int, Bounce)
 		SHADER_PARAMETER_RDG_BUFFER_UAV(RWStructuredBuffer<FPathTracingPackedPathState>, PathStateData)
-		SHADER_PARAMETER_RDG_BUFFER_SRV(Buffer<int>, ActivePaths)
-		SHADER_PARAMETER_RDG_BUFFER_UAV(RWBuffer<int>, NextActivePaths)
-		SHADER_PARAMETER_RDG_BUFFER_UAV(RWBuffer<int>, NumPathStates)
+		SHADER_PARAMETER_RDG_BUFFER_UAV(RWBuffer<uint>, ActivePaths)
+		SHADER_PARAMETER_RDG_BUFFER_UAV(RWBuffer<uint>, NextActivePaths)
+		SHADER_PARAMETER_RDG_BUFFER_UAV(RWBuffer<uint>, NumPathStates)
 
 		RDG_BUFFER_ACCESS(PathTracingIndirectArgs, ERHIAccess::IndirectArgs | ERHIAccess::SRVCompute)
 	END_SHADER_PARAMETER_STRUCT()
@@ -3414,7 +3414,7 @@ void FDeferredShadingSceneRenderer::RenderPathTracing(
 			const int FlushRenderingCommands = CVarPathTracingFlushDispatch.GetValueOnRenderThread();
 
 			FRDGBuffer* ActivePaths[2] = {};
-			FRDGBuffer* NumActivePaths[2] = {};
+			FRDGBuffer* NumActivePaths = nullptr;
 			FRDGBuffer* PathStateData = nullptr;
 			if (bUseCompaction || Config.UseAdaptiveSampling)
 			{
@@ -3424,15 +3424,7 @@ void FDeferredShadingSceneRenderer::RenderPathTracing(
 				);
 				ActivePaths[0] = GraphBuilder.CreateBuffer(FRDGBufferDesc::CreateBufferDesc(sizeof(int32), NumPaths), TEXT("PathTracer.ActivePaths0"));
 				ActivePaths[1] = GraphBuilder.CreateBuffer(FRDGBufferDesc::CreateBufferDesc(sizeof(int32), NumPaths), TEXT("PathTracer.ActivePaths1"));
-				if (bUseIndirectDispatch)
-				{
-					NumActivePaths[0] = GraphBuilder.CreateBuffer(FRDGBufferDesc::CreateIndirectDesc<int32>(3), TEXT("PathTracer.NumActivePaths0"));
-					NumActivePaths[1] = GraphBuilder.CreateBuffer(FRDGBufferDesc::CreateIndirectDesc<int32>(3), TEXT("PathTracer.NumActivePaths1"));
-				}
-				else
-				{
-					NumActivePaths[0] = GraphBuilder.CreateBuffer(FRDGBufferDesc::CreateBufferDesc(sizeof(int32), 3), TEXT("PathTracer.NumActivePaths"));
-				}
+				NumActivePaths = GraphBuilder.CreateBuffer(FRDGBufferDesc::CreateIndirectDesc<uint32>(3 * (Config.PathTracingData.MaxBounces + 1)), TEXT("PathTracer.NumActivePaths"));
 				PathStateData = GraphBuilder.CreateBuffer(FRDGBufferDesc::CreateStructuredDesc(sizeof(FPathTracingPackedPathState), NumPaths), TEXT("PathTracer.PathStateData"));
 			}
 
@@ -3476,9 +3468,8 @@ void FDeferredShadingSceneRenderer::RenderPathTracing(
 							PassParameters->AdaptiveSamplingErrorThreshold = Config.AdaptiveSamplingThreshold;
 							PassParameters->ViewPreExposure = View.PreExposure;
 
-							PassParameters->NextActivePaths = GraphBuilder.CreateUAV(ActivePaths[0], PF_R32_SINT);
-							PassParameters->NumPathStates = GraphBuilder.CreateUAV(NumActivePaths[0], PF_R32_UINT);
-							AddClearUAVPass(GraphBuilder, PassParameters->NextActivePaths, -1); // make sure everything is initialized to -1 since paths that go inactive don't write anything
+							PassParameters->NextActivePaths = GraphBuilder.CreateUAV(ActivePaths[0], PF_R32_UINT);
+							PassParameters->NumPathStates = GraphBuilder.CreateUAV(NumActivePaths, PF_R32_UINT);
 							AddClearUAVPass(GraphBuilder, PassParameters->NumPathStates, 0);
 
 							PassParameters->TileTextureOffset.X = TileX;
@@ -3494,7 +3485,7 @@ void FDeferredShadingSceneRenderer::RenderPathTracing(
 						}
 						else if (bUseCompaction)
 						{
-							AddClearUAVPass(GraphBuilder, GraphBuilder.CreateUAV(ActivePaths[0], PF_R32_UINT), 0);
+							AddClearUAVPass(GraphBuilder, GraphBuilder.CreateUAV(NumActivePaths, PF_R32_UINT), 0);
 						}
 
 						// When using path compaction, we need to run the path tracer once per bounce
@@ -3603,20 +3594,14 @@ void FDeferredShadingSceneRenderer::RenderPathTracing(
 							PassParameters->Bounce = Bounce;
 							if (bUseCompaction)
 							{
-								PassParameters->ActivePaths = GraphBuilder.CreateSRV(ActivePaths[Bounce & 1], PF_R32_SINT);
-								PassParameters->NextActivePaths = GraphBuilder.CreateUAV(ActivePaths[(Bounce & 1) ^ 1], PF_R32_SINT);
+								PassParameters->ActivePaths = GraphBuilder.CreateUAV(ActivePaths[Bounce & 1], PF_R32_UINT);
+								PassParameters->NextActivePaths = GraphBuilder.CreateUAV(ActivePaths[(Bounce & 1) ^ 1], PF_R32_UINT);
 								PassParameters->PathStateData = GraphBuilder.CreateUAV(PathStateData);
+								PassParameters->NumPathStates = GraphBuilder.CreateUAV(NumActivePaths, PF_R32_UINT);
 								if (bUseIndirectDispatch)
 								{
-									PassParameters->NumPathStates = GraphBuilder.CreateUAV(NumActivePaths[Bounce & 1], PF_R32_UINT);
-									PassParameters->PathTracingIndirectArgs = NumActivePaths[(Bounce & 1) ^ 1];
-								}
-								else
-								{
-									PassParameters->NumPathStates = GraphBuilder.CreateUAV(NumActivePaths[0], PF_R32_UINT);
-									AddClearUAVPass(GraphBuilder, PassParameters->NextActivePaths, -1); // make sure everything is initialized to -1 since paths that go inactive don't write anything
+									PassParameters->PathTracingIndirectArgs = NumActivePaths;
 								}
-								AddClearUAVPass(GraphBuilder, PassParameters->NumPathStates, 0);
 							}
 							ClearUnusedGraphResources(RayGenShader, PassParameters);
 							const bool bFlushRenderingCommands = FlushRenderingCommands == 1 || (FlushRenderingCommands == 2 && Bounce == MaxBounces);
@@ -3640,7 +3625,7 @@ void FDeferredShadingSceneRenderer::RenderPathTracing(
 											View.RayTracingMaterialPipeline,
 											RayGenShader.GetRayTracingShader(),
 											RayTracingSceneRHI, View.RayTracingSBT, GlobalResources,
-											PassParameters->PathTracingIndirectArgs->GetIndirectRHICallBuffer(), 0
+											PassParameters->PathTracingIndirectArgs->GetIndirectRHICallBuffer(), 3 * PassParameters->Bounce
 										);
 									}
 									else
