Add new Mask4 RVT format.
Compresses to BC3.

diff --git a/Engine/Shaders/Private/VirtualTextureCompress.usf b/Engine/Shaders/Private/VirtualTextureCompress.usf
--- a/Engine/Shaders/Private/VirtualTextureCompress.usf
+++ b/Engine/Shaders/Private/VirtualTextureCompress.usf
@@ -310,6 +310,45 @@ void CopyBaseColorNormalSpecularMaskYCoCgPS(
 	OutColor2 = float4(RoughnessSpecularNormalZ, NormalXYMask.z);
 }
 
+/** Compress 4 channel mask to BC3. */
+[numthreads(8,8,1)]
+void CompressMask4CS(
+	uint3 ThreadId : SV_DispatchThreadID)
+{
+	if (any(ThreadId.xy * 4 >= SourceRect.zw))
+	{
+		return;
+	}
+
+	uint2 SamplePos = SourceRect.xy + ThreadId.xy * 4;
+	float2 TexelUVSize = 1.f / float2(SourceRect.zw);
+	float2 SampleUV = (float2(SamplePos) + 0.5f) * TexelUVSize;
+
+	float3 BlockXYZ[16];
+	ReadBlockRGB(RenderTexture0, TextureSampler0, SampleUV, TexelUVSize, BlockXYZ);
+
+	float BlockW[16];
+	ReadBlockX(RenderTexture1, TextureSampler0, SampleUV, TexelUVSize, BlockW);
+
+#if ETC_PROFILE
+	OutCompressTexture0_128bit[ThreadId.xy + DestPos0] = CompressBlock_ETC2_RGBA(BlockXYZ, BlockW);
+#else
+	OutCompressTexture0_128bit[ThreadId.xy + DestPos0] = CompressBC3Block(BlockXYZ, BlockW);
+#endif
+}
+
+/** Copy path used when we disable texture compression, because we need to keep the same final channel layout. */
+void CopyMask4PS(
+	in float4 InPosition : SV_POSITION, 
+	in noperspective float2 InTexCoord : TEXCOORD0, 
+	out float4 OutColor0 : SV_Target0)
+{
+	float3 MaskXYZ = RenderTexture0.SampleLevel(TextureSampler0, InTexCoord, 0).xyz;
+	float MaskW = RenderTexture1.SampleLevel(TextureSampler1, InTexCoord, 0).x;
+	
+	OutColor0 = float4(MaskXYZ, MaskW);
+}
+
 /** Compress displacement to a single BC4 target. */
 [numthreads(8, 8, 1)]
 void CompressDisplacementCS(
diff --git a/Engine/Shaders/Private/VirtualTextureMaterial.usf b/Engine/Shaders/Private/VirtualTextureMaterial.usf
--- a/Engine/Shaders/Private/VirtualTextureMaterial.usf
+++ b/Engine/Shaders/Private/VirtualTextureMaterial.usf
@@ -95,9 +95,10 @@ void FPixelShaderInOut_MainPS(
 	half Roughness = GetVirtualTextureOutput2(MaterialParameters);
 	half3 Normal = GetVirtualTextureOutput3(MaterialParameters);
 	FDFScalar WorldHeight = WSToDF(GetVirtualTextureOutput4_LWC(MaterialParameters));
-	float Opacity = GetVirtualTextureOutput5(MaterialParameters);
-	float Mask = GetVirtualTextureOutput6(MaterialParameters);
-	float Displacement = GetVirtualTextureOutput7(MaterialParameters);
+	half Opacity = GetVirtualTextureOutput5(MaterialParameters);
+	half Mask = GetVirtualTextureOutput6(MaterialParameters);
+	half Displacement = GetVirtualTextureOutput7(MaterialParameters);
+	half4 Mask4 = GetVirtualTextureOutput8(MaterialParameters);
 #else
 	// Output is from standard material output attribute node
 	half3 BaseColor = GetMaterialBaseColor(PixelMaterialInputs);
@@ -105,9 +106,10 @@ void FPixelShaderInOut_MainPS(
 	half Roughness = GetMaterialRoughness(PixelMaterialInputs);
 	half3 Normal = MaterialParameters.WorldNormal;
 	FDFScalar WorldHeight = WSToDF(WSGetZ(GetWorldPosition(MaterialParameters)));
-	float Opacity = GetMaterialOpacity(PixelMaterialInputs);
-	float Mask = 0.f;
-	float Displacement = 0.f;
+	half Opacity = GetMaterialOpacity(PixelMaterialInputs);
+	half Mask = 0.f;
+	half Displacement = 0.f;
+	half4 Mask4 = 0.f;
 #endif
 
 	// Apply any fixed color.
@@ -119,19 +121,22 @@ void FPixelShaderInOut_MainPS(
 #elif defined(OUT_BASECOLOR_NORMAL_ROUGHNESS)
 	float3 PackedNormal = PackNormal(Normal);
 	//store basecolor in srgb space to improve image quality
-	Out.MRT[0] = float4(LinearToSrgb(BaseColor), 1.0) * Opacity;
+	Out.MRT[0] = float4(LinearToSrgb(BaseColor), 1.f) * Opacity;
 	//not enough channels for sign of normal  z
-	Out.MRT[1] = float4(PackedNormal.x, Roughness, PackedNormal.y, 1.0) * Opacity;
+	Out.MRT[1] = float4(PackedNormal.x, Roughness, PackedNormal.y, 1.f) * Opacity;
 #elif defined(OUT_BASECOLOR_NORMAL_SPECULAR)
 	float3 PackedNormal = PackNormal(Normal);
 	Out.MRT[0] = float4(BaseColor, 1.f) * Opacity;
 	Out.MRT[1] = float4(PackedNormal.xy, Mask, 1.f) * Opacity;
 	Out.MRT[2] = float4(Specular, Roughness, PackedNormal.z, 1.f) * Opacity;
+#elif defined(OUT_MASK4)
+	Out.MRT[0] = float4(Mask4.xyz, 1.f) * Opacity;
+	Out.MRT[1] = float4(Mask4.w, 0.f, 0.f, 1.f) * Opacity;
 #elif defined(OUT_WORLDHEIGHT)
 	float PackedHeight = PackWorldHeight(WorldHeight, ResolvedView.RuntimeVirtualTexturePackHeight);
-	Out.MRT[0] = float4(PackedHeight, 0, 0, 1);
+	Out.MRT[0] = float4(PackedHeight, 0.f, 0.f, 1.f);
 #elif defined(OUT_DISPLACEMENT)
-	Out.MRT[0] = float4(Displacement, 0, 0, 1) * Opacity;
+	Out.MRT[0] = float4(Displacement, 0.f, 0.f, 1.f) * Opacity;
 #endif
 }
 
@@ -143,6 +148,9 @@ void FPixelShaderInOut_MainPS(
 //#define PIXELSHADEROUTPUT_MRT0 1
 //#define PIXELSHADEROUTPUT_MRT1 1
 //#define PIXELSHADEROUTPUT_MRT2 1
+//#elif defined(OUT_MASK4)
+//#define PIXELSHADEROUTPUT_MRT0 1
+//#define PIXELSHADEROUTPUT_MRT1 1
 //#elif defined(OUT_WORLDHEIGHT)
 //#define PIXELSHADEROUTPUT_MRT0 1
 //#elif defined(OUT_DISPLACEMENT)
diff --git a/Engine/Source/Runtime/Engine/Private/Materials/MaterialExpressions.cpp b/Engine/Source/Runtime/Engine/Private/Materials/MaterialExpressions.cpp
--- a/Engine/Source/Runtime/Engine/Private/Materials/MaterialExpressions.cpp
+++ b/Engine/Source/Runtime/Engine/Private/Materials/MaterialExpressions.cpp
@@ -3159,6 +3159,12 @@ int32 UMaterialExpressionRuntimeVirtualTextureOutput::Compile(class FMaterialCom
 		OutputAttributeMask |= Displacement.IsConnected() ? (1 << (uint8)ERuntimeVirtualTextureAttributeType::Displacement) : 0;
 		ValidateRuntimeVirtualTextureOutput(Compiler, Displacement);
 	}
+	else if (OutputIndex == 8)
+	{
+		CodeInput = Mask4.IsConnected() ? Mask4.Compile(Compiler) : Compiler->Constant4(0.f, 0.f, 0.f, 0.f);
+		OutputAttributeMask |= Mask4.IsConnected() ? (1 << (uint8)ERuntimeVirtualTextureAttributeType::Mask4) : 0;
+		ValidateRuntimeVirtualTextureOutput(Compiler, Mask4);
+	}
 
 	Compiler->VirtualTextureOutput(OutputAttributeMask);
 	return Compiler->CustomOutput(this, OutputIndex, CodeInput);
@@ -3173,7 +3179,7 @@ void UMaterialExpressionRuntimeVirtualTextureOutput::GetCaption(TArray<FString>&
 
 int32 UMaterialExpressionRuntimeVirtualTextureOutput::GetNumOutputs() const
 {
-	return 8; 
+	return 9; 
 }
 
 FString UMaterialExpressionRuntimeVirtualTextureOutput::GetFunctionName() const
@@ -3238,6 +3244,7 @@ void UMaterialExpressionRuntimeVirtualTextureSample::InitOutputs()
 	Outputs.Add(FExpressionOutput(TEXT("WorldHeight")));
 	Outputs.Add(FExpressionOutput(TEXT("Mask")));
 	Outputs.Add(FExpressionOutput(TEXT("Displacement")));
+	Outputs.Add(FExpressionOutput(TEXT("Mask4")));
 #endif // WITH_EDITORONLY_DATA
 }
 
@@ -3262,12 +3269,6 @@ void UMaterialExpressionRuntimeVirtualTextureSample::PostLoad()
 {
 	Super::PostLoad();
 
-	// Convert BaseColor_Normal_DEPRECATED
-	if (MaterialType == ERuntimeVirtualTextureMaterialType::BaseColor_Normal_DEPRECATED)
-	{
-		MaterialType = ERuntimeVirtualTextureMaterialType::BaseColor_Normal_Specular;
-	}
-
 	InitOutputs();
 }
 
@@ -3370,6 +3371,7 @@ int32 UMaterialExpressionRuntimeVirtualTextureSample::Compile(class FMaterialCom
 	bool bIsNormalValid = false;
 	bool bIsWorldHeightValid = false;
 	bool bIsMaskValid = false;
+	bool bIsMask4Valid = false;
 	bool bIsDisplacementValid = false;
 
 	switch (MaterialType)
@@ -3379,6 +3381,7 @@ int32 UMaterialExpressionRuntimeVirtualTextureSample::Compile(class FMaterialCom
 	case ERuntimeVirtualTextureMaterialType::BaseColor_Normal_Specular: bIsRoughnessValid = bIsBaseColorValid = bIsNormalValid = bIsSpecularValid = true; break;
 	case ERuntimeVirtualTextureMaterialType::BaseColor_Normal_Specular_YCoCg: bIsRoughnessValid = bIsBaseColorValid = bIsNormalValid = bIsSpecularValid = true; break;
 	case ERuntimeVirtualTextureMaterialType::BaseColor_Normal_Specular_Mask_YCoCg: bIsRoughnessValid = bIsBaseColorValid = bIsNormalValid = bIsSpecularValid = bIsMaskValid = true; break;
+	case ERuntimeVirtualTextureMaterialType::Mask4: bIsMask4Valid = true; break;
 	case ERuntimeVirtualTextureMaterialType::WorldHeight: bIsWorldHeightValid = true; break;
 	case ERuntimeVirtualTextureMaterialType::Displacement: bIsDisplacementValid = true; break;
 	}
@@ -3461,7 +3464,7 @@ int32 UMaterialExpressionRuntimeVirtualTextureSample::Compile(class FMaterialCom
 	case 5:
 		if (bIsVirtualTextureValid && bIsMaskValid)
 		{
-			UnpackTarget = 2; UnpackMask = 0x8; break;
+			UnpackTarget = 2; UnpackMask = 0x8;
 		}
 		else
 		{
@@ -3478,6 +3481,16 @@ int32 UMaterialExpressionRuntimeVirtualTextureSample::Compile(class FMaterialCom
 			return Compiler->Constant(0.f);
 		}
 		break;
+	case 7:
+		if (bIsVirtualTextureValid && bIsMask4Valid)
+		{
+			UnpackTarget = 0; UnpackMask = 0xf; break;
+		}
+		else
+		{
+			return Compiler->Constant4(0.f, 0.f, 0.f, 0.f);
+		}
+		break;
 	default:
 		return INDEX_NONE;
 	}
diff --git a/Engine/Source/Runtime/Engine/Private/ShaderCompiler/ShaderGenerationUtil.cpp b/Engine/Source/Runtime/Engine/Private/ShaderCompiler/ShaderGenerationUtil.cpp
--- a/Engine/Source/Runtime/Engine/Private/ShaderCompiler/ShaderGenerationUtil.cpp
+++ b/Engine/Source/Runtime/Engine/Private/ShaderCompiler/ShaderGenerationUtil.cpp
@@ -207,6 +207,7 @@ void ApplyFetchEnvironmentInternal(FShaderMaterialPropertyDefines& SrcDefines, c
 	FETCH_COMPILE_BOOL(OUT_BASECOLOR);
 	FETCH_COMPILE_BOOL(OUT_BASECOLOR_NORMAL_ROUGHNESS);
 	FETCH_COMPILE_BOOL(OUT_BASECOLOR_NORMAL_SPECULAR);
+	FETCH_COMPILE_BOOL(OUT_MASK4);
 	FETCH_COMPILE_BOOL(OUT_WORLDHEIGHT);
 	FETCH_COMPILE_BOOL(OUT_DISPLACEMENT);
 
@@ -2000,6 +2001,11 @@ void FShaderCompileUtilities::ApplyDerivedDefines(FShaderCompilerEnvironment& Ou
 			TargetUsage[1] = EGBufferSlotUsage::Written;
 			TargetUsage[2] = EGBufferSlotUsage::Written;
 		}
+		else if (MaterialDefines.OUT_MASK4)
+		{
+			TargetUsage[0] = EGBufferSlotUsage::Written;
+			TargetUsage[1] = EGBufferSlotUsage::Written;
+		}
 		else if (MaterialDefines.OUT_WORLDHEIGHT)
 		{
 			TargetUsage[0] = EGBufferSlotUsage::Written;
diff --git a/Engine/Source/Runtime/Engine/Private/VT/RuntimeVirtualTexture.cpp b/Engine/Source/Runtime/Engine/Private/VT/RuntimeVirtualTexture.cpp
--- a/Engine/Source/Runtime/Engine/Private/VT/RuntimeVirtualTexture.cpp
+++ b/Engine/Source/Runtime/Engine/Private/VT/RuntimeVirtualTexture.cpp
@@ -363,6 +363,7 @@ int32 URuntimeVirtualTexture::GetLayerCount(ERuntimeVirtualTextureMaterialType I
 	switch (InMaterialType)
 	{
 	case ERuntimeVirtualTextureMaterialType::BaseColor:
+	case ERuntimeVirtualTextureMaterialType::Mask4:
 	case ERuntimeVirtualTextureMaterialType::WorldHeight:
 	case ERuntimeVirtualTextureMaterialType::Displacement:
 		return 1;
@@ -438,6 +439,8 @@ EPixelFormat URuntimeVirtualTexture::GetLayerFormat(int32 LayerIndex) const
 		case ERuntimeVirtualTextureMaterialType::BaseColor_Normal_Specular_YCoCg:
 		case ERuntimeVirtualTextureMaterialType::BaseColor_Normal_Specular_Mask_YCoCg:
 			return bCompressTextures ? PlatformCompressedRVTFormat(PF_DXT5) : PF_B8G8R8A8;
+		case ERuntimeVirtualTextureMaterialType::Mask4:
+			return bCompressTextures ? PlatformCompressedRVTFormat(PF_DXT5) : PF_B8G8R8A8;
 		case ERuntimeVirtualTextureMaterialType::WorldHeight:
 			return PF_G16;
 		case ERuntimeVirtualTextureMaterialType::Displacement:
@@ -492,6 +495,7 @@ bool URuntimeVirtualTexture::IsLayerSRGB(int32 LayerIndex) const
 	case ERuntimeVirtualTextureMaterialType::BaseColor_Normal_Specular_YCoCg:
 	case ERuntimeVirtualTextureMaterialType::BaseColor_Normal_Specular_Mask_YCoCg:
 		// These formats have YCoCg packing which can't use sRGB.
+	case ERuntimeVirtualTextureMaterialType::Mask4:
 	case ERuntimeVirtualTextureMaterialType::WorldHeight:
 	case ERuntimeVirtualTextureMaterialType::Displacement:
 		// These formats require linear encoding.
@@ -627,12 +631,6 @@ void URuntimeVirtualTexture::PostLoad()
 		Size_DEPRECATED = -1;
 	}
 
-	// Convert BaseColor_Normal_DEPRECATED
-	if (MaterialType == ERuntimeVirtualTextureMaterialType::BaseColor_Normal_DEPRECATED)
-	{
-		MaterialType = ERuntimeVirtualTextureMaterialType::BaseColor_Normal_Specular;
-	}
-
 	// Remove StreamingTexture_DEPRECATED
 	StreamingTexture_DEPRECATED = nullptr;
 
diff --git a/Engine/Source/Runtime/Engine/Public/Materials/MaterialExpressionRuntimeVirtualTextureOutput.h b/Engine/Source/Runtime/Engine/Public/Materials/MaterialExpressionRuntimeVirtualTextureOutput.h
--- a/Engine/Source/Runtime/Engine/Public/Materials/MaterialExpressionRuntimeVirtualTextureOutput.h
+++ b/Engine/Source/Runtime/Engine/Public/Materials/MaterialExpressionRuntimeVirtualTextureOutput.h
@@ -45,6 +45,10 @@ class UMaterialExpressionRuntimeVirtualTextureOutput : public UMaterialExpressio
 	UPROPERTY()
 	FExpressionInput Displacement;
 
+	/** Input for Mask to output to virtual texture. */
+	UPROPERTY()
+	FExpressionInput Mask4;
+
 public:
 #if WITH_EDITOR
 	//~ Begin UMaterialExpression Interface
diff --git a/Engine/Source/Runtime/Engine/Public/VT/RuntimeVirtualTextureEnum.h b/Engine/Source/Runtime/Engine/Public/VT/RuntimeVirtualTextureEnum.h
--- a/Engine/Source/Runtime/Engine/Public/VT/RuntimeVirtualTextureEnum.h
+++ b/Engine/Source/Runtime/Engine/Public/VT/RuntimeVirtualTextureEnum.h
@@ -19,6 +19,7 @@ enum class ERuntimeVirtualTextureAttributeType : uint8
 	Roughness,
 	Specular,
 	Mask,
+	Mask4,
 	WorldHeight,
 	Displacement,
 
@@ -36,7 +37,7 @@ UENUM()
 enum class ERuntimeVirtualTextureMaterialType : uint8
 {
 	BaseColor UMETA(DisplayName = "Base Color"),
-	BaseColor_Normal_DEPRECATED UMETA(Hidden),
+	Mask4 UMETA(DisplayName = "Mask4", ToolTip = "4 channel mask texture."),
 	BaseColor_Normal_Roughness UMETA(DisplayName = "Base Color, Normal, Roughness", ToolTip = "Local space Normal. Requires less memory than 'Base Color, Normal, Roughness, Specular'. Supports LQ compression."),
 	BaseColor_Normal_Specular UMETA(DisplayName = "Base Color, Normal, Roughness, Specular"),
 	BaseColor_Normal_Specular_YCoCg UMETA(DisplayName = "YCoCg Base Color, Normal, Roughness, Specular", ToolTip = "Base Color is stored in YCoCg space. This requires more memory but may provide better quality."),
diff --git a/Engine/Source/Runtime/Landscape/Private/LandscapeRender.cpp b/Engine/Source/Runtime/Landscape/Private/LandscapeRender.cpp
--- a/Engine/Source/Runtime/Landscape/Private/LandscapeRender.cpp
+++ b/Engine/Source/Runtime/Landscape/Private/LandscapeRender.cpp
@@ -3991,12 +3991,14 @@ public:
 			FName(TEXT("TVirtualTextureVSBaseColorNormal")),
 			FName(TEXT("TVirtualTextureVSBaseColorNormalSpecular")),
 			FName(TEXT("TVirtualTextureVSBaseColorNormalRoughness")),
+			FName(TEXT("TVirtualTextureVSMask4")),
 			FName(TEXT("TVirtualTextureVSWorldHeight")),
 			FName(TEXT("TVirtualTextureVSDisplacement")),
 			FName(TEXT("TVirtualTexturePSBaseColor")),
 			FName(TEXT("TVirtualTexturePSBaseColorNormal")),
 			FName(TEXT("TVirtualTexturePSBaseColorNormalSpecular")),
 			FName(TEXT("TVirtualTexturePSBaseColorNormalRoughness")),
+			FName(TEXT("TVirtualTexturePSMask4")),
 			FName(TEXT("TVirtualTexturePSWorldHeight")),
 			FName(TEXT("TVirtualTexturePSDisplacement")),
 		};
@@ -4158,12 +4160,14 @@ public:
 			FName(TEXT("TVirtualTextureVSBaseColorNormal")),
 			FName(TEXT("TVirtualTextureVSBaseColorNormalSpecular")),
 			FName(TEXT("TVirtualTextureVSBaseColorNormalRoughness")),
+			FName(TEXT("TVirtualTextureVSMask4")),
 			FName(TEXT("TVirtualTextureVSWorldHeight")),
 			FName(TEXT("TVirtualTextureVSDisplacement")),
 			FName(TEXT("TVirtualTexturePSBaseColor")),
 			FName(TEXT("TVirtualTexturePSBaseColorNormal")),
 			FName(TEXT("TVirtualTexturePSBaseColorNormalSpecular")),
 			FName(TEXT("TVirtualTexturePSBaseColorNormalRoughness")),
+			FName(TEXT("TVirtualTexturePSMask4")),
 			FName(TEXT("TVirtualTexturePSWorldHeight")),
 			FName(TEXT("TVirtualTexturePSDisplacement")),
 		};
diff --git a/Engine/Source/Runtime/RenderCore/Private/ShaderMaterialDerivedHelpers.cpp b/Engine/Source/Runtime/RenderCore/Private/ShaderMaterialDerivedHelpers.cpp
--- a/Engine/Source/Runtime/RenderCore/Private/ShaderMaterialDerivedHelpers.cpp
+++ b/Engine/Source/Runtime/RenderCore/Private/ShaderMaterialDerivedHelpers.cpp
@@ -120,6 +120,11 @@ FShaderMaterialDerivedDefines RENDERCORE_API CalculateDerivedMaterialParameters(
 				Dst.PIXELSHADEROUTPUT_MRT1 = 1;
 				Dst.PIXELSHADEROUTPUT_MRT2 = 1;
 			}
+			else if (Mat.OUT_MASK4)
+			{
+				Dst.PIXELSHADEROUTPUT_MRT0 = 1;
+				Dst.PIXELSHADEROUTPUT_MRT1 = 1;
+			}
 			else if (Mat.OUT_WORLDHEIGHT)
 			{
 				Dst.PIXELSHADEROUTPUT_MRT0 = 1;
diff --git a/Engine/Source/Runtime/RenderCore/Public/ShaderMaterial.h b/Engine/Source/Runtime/RenderCore/Public/ShaderMaterial.h
--- a/Engine/Source/Runtime/RenderCore/Public/ShaderMaterial.h
+++ b/Engine/Source/Runtime/RenderCore/Public/ShaderMaterial.h
@@ -178,6 +178,7 @@ struct FShaderMaterialPropertyDefines
 	uint8 OUT_BASECOLOR : 1;
 	uint8 OUT_BASECOLOR_NORMAL_ROUGHNESS : 1;
 	uint8 OUT_BASECOLOR_NORMAL_SPECULAR : 1;
+	uint8 OUT_MASK4 : 1;
 	uint8 OUT_WORLDHEIGHT : 1;
 	uint8 OUT_DISPLACEMENT : 1;
 
diff --git a/Engine/Source/Runtime/Renderer/Private/VT/RuntimeVirtualTextureRender.cpp b/Engine/Source/Runtime/Renderer/Private/VT/RuntimeVirtualTextureRender.cpp
--- a/Engine/Source/Runtime/Renderer/Private/VT/RuntimeVirtualTextureRender.cpp
+++ b/Engine/Source/Runtime/Renderer/Private/VT/RuntimeVirtualTextureRender.cpp
@@ -359,6 +359,23 @@ namespace RuntimeVirtualTexture
 		}
 	};
 
+	/** Specialization for ERuntimeVirtualTextureMaterialType::Mask4 */
+	class FMaterialPolicy_Mask4
+	{
+	public:
+		static void ModifyCompilationEnvironment(FShaderCompilerEnvironment& OutEnvironment)
+		{
+			OutEnvironment.SetDefine(TEXT("OUT_MASK4"), 1);
+		}
+
+		static FRHIBlendState* GetBlendState(uint8 OutputAttributeMask)
+		{
+			return TStaticBlendState<
+				CW_RGBA, BO_Add, BF_One, BF_InverseSourceAlpha, BO_Add, BF_Zero, BF_One,
+				CW_RED, BO_Add, BF_One, BF_InverseSourceAlpha, BO_Add, BF_Zero, BF_One >::GetRHI();
+		}
+	};
+
 	/** Specialization for ERuntimeVirtualTextureMaterialType::WorldHeight */
 	class FMaterialPolicy_WorldHeight
 	{
@@ -443,10 +460,11 @@ namespace RuntimeVirtualTexture
 	IMPLEMENT_MATERIAL_SHADER_TYPE(template<>,TVirtualTexturePS##PolicyName, TEXT("/Engine/Private/VirtualTextureMaterial.usf"), TEXT("MainPS"), SF_Pixel);
 
 	IMPLEMENT_VIRTUALTEXTURE_SHADER_TYPE(FMaterialPolicy_BaseColor, BaseColor);
+	IMPLEMENT_VIRTUALTEXTURE_SHADER_TYPE(FMaterialPolicy_BaseColorNormalRoughness, BaseColorNormalRoughness);
 	IMPLEMENT_VIRTUALTEXTURE_SHADER_TYPE(FMaterialPolicy_BaseColorNormalSpecular, BaseColorNormalSpecular);
+	IMPLEMENT_VIRTUALTEXTURE_SHADER_TYPE(FMaterialPolicy_Mask4, Mask4);
 	IMPLEMENT_VIRTUALTEXTURE_SHADER_TYPE(FMaterialPolicy_WorldHeight, WorldHeight);
 	IMPLEMENT_VIRTUALTEXTURE_SHADER_TYPE(FMaterialPolicy_Displacement, Displacement);
-	IMPLEMENT_VIRTUALTEXTURE_SHADER_TYPE(FMaterialPolicy_BaseColorNormalRoughness, BaseColorNormalRoughness);
 
 	/** Structure to localize the setup of our render graph based on the virtual texture setup. */
 	struct FRenderGraphSetup
@@ -480,6 +498,10 @@ namespace RuntimeVirtualTexture
 				AddRenderTargetInfo(PF_B8G8R8A8, RTCreateFlags, RenderTargetsInfo);
 				AddRenderTargetInfo(PF_B8G8R8A8, RTCreateFlags, RenderTargetsInfo);
 				break;
+			case ERuntimeVirtualTextureMaterialType::Mask4:
+				AddRenderTargetInfo(PF_B8G8R8A8, RTCreateFlags, RenderTargetsInfo);
+				AddRenderTargetInfo(PF_B8G8R8A8, RTCreateFlags, RenderTargetsInfo);
+				break;
 			case ERuntimeVirtualTextureMaterialType::WorldHeight:
 			case ERuntimeVirtualTextureMaterialType::Displacement:
 				AddRenderTargetInfo(PF_G16, RTCreateFlags, RenderTargetsInfo);
@@ -511,7 +533,8 @@ namespace RuntimeVirtualTexture
 			const bool bCompressedFormat = GPixelFormats[Desc.OutputFormat0].BlockSizeX == 4 && GPixelFormats[Desc.OutputFormat0].BlockSizeY == 4;
 			const bool bLQFormat = Desc.OutputFormat0 == PF_R5G6B5_UNORM;
 			bCompressPass = bRenderPass && !bCopyThumbnailPass && bCompressedFormat;
-			bCopyPass = bRenderPass && !bCopyThumbnailPass && !bCompressPass && (Desc.MaterialType == ERuntimeVirtualTextureMaterialType::BaseColor_Normal_Specular || Desc.MaterialType == ERuntimeVirtualTextureMaterialType::BaseColor_Normal_Specular_YCoCg || Desc.MaterialType == ERuntimeVirtualTextureMaterialType::BaseColor_Normal_Specular_Mask_YCoCg);
+			bCopyPass = bRenderPass && !bCopyThumbnailPass && !bCompressPass && 
+				(Desc.MaterialType == ERuntimeVirtualTextureMaterialType::BaseColor_Normal_Specular || Desc.MaterialType == ERuntimeVirtualTextureMaterialType::BaseColor_Normal_Specular_YCoCg || Desc.MaterialType == ERuntimeVirtualTextureMaterialType::BaseColor_Normal_Specular_Mask_YCoCg || Desc.MaterialType == ERuntimeVirtualTextureMaterialType::Mask4);
 			
 			// Use direct aliasing for compression pass on platforms that support it.
 			bDirectAliasing = bCompressedFormat && GRHISupportsUAVFormatAliasing && CVarVTDirectCompress.GetValueOnRenderThread() != 0;
@@ -701,6 +724,31 @@ namespace RuntimeVirtualTexture
 					OutputAlias0 = CopyTexture0 = GraphBuilder.CreateTexture(FRDGTextureDesc::Create2D(Desc.TextureSize, PF_B8G8R8A8, FClearValueBinding::Black, RTCreateFlags | RTSrgbFlags), TEXT("CopyTexture0"));
 				}
 				break;
+			case ERuntimeVirtualTextureMaterialType::Mask4:
+				if (bRenderPass)
+				{
+					OutputAlias0 = RenderTexture0 = GraphBuilder.CreateTexture(FRDGTextureDesc::Create2D(Desc.TextureSize, PF_B8G8R8A8, FClearValueBinding::Black, RTCreateFlags), TEXT("RenderTexture0"));
+					OutputAlias1 = RenderTexture1 = GraphBuilder.CreateTexture(FRDGTextureDesc::Create2D(Desc.TextureSize, PF_B8G8R8A8, FClearValueBinding::Black, RTCreateFlags), TEXT("RenderTexture1"));
+				}
+				if (bCompressPass)
+				{
+					if (bDirectAliasing)
+					{
+						CompressTexture0 = GraphBuilder.RegisterExternalTexture(Desc.OutputTargets[0], ERDGTextureFlags::None);
+						CompressTextureUAV0_128bit = GraphBuilder.CreateUAV(FRDGTextureUAVDesc(CompressTexture0, 0, Compressed128BitFormat));
+						OutputAlias0 = nullptr;
+					}
+					else
+					{
+						OutputAlias0 = CompressTexture0 = GraphBuilder.CreateTexture(FRDGTextureDesc::Create2D(Desc.TextureSize / 4, Compressed128BitFormat, FClearValueBinding::None, TexCreate_UAV), TEXT("CompressTexture0"));
+						CompressTextureUAV0_64bit = GraphBuilder.CreateUAV(FRDGTextureUAVDesc(CompressTexture0));
+					}
+				}
+				if (bCopyPass || bCopyThumbnailPass)
+				{
+					OutputAlias0 = CopyTexture0 = GraphBuilder.CreateTexture(FRDGTextureDesc::Create2D(Desc.TextureSize, PF_B8G8R8A8, FClearValueBinding::Black, RTCreateFlags), TEXT("CopyTexture0"));
+				}
+				break;
 			case ERuntimeVirtualTextureMaterialType::WorldHeight:
 				if (bRenderPass)
 				{
@@ -802,6 +850,8 @@ namespace RuntimeVirtualTexture
 				case ERuntimeVirtualTextureMaterialType::BaseColor_Normal_Specular_YCoCg:
 				case ERuntimeVirtualTextureMaterialType::BaseColor_Normal_Specular_Mask_YCoCg:
 					return Process<FMaterialPolicy_BaseColorNormalSpecular>(MeshBatch, BatchElementMask, StaticMeshId, OutputAttributeMask, PrimitiveSceneProxy, *MaterialRenderProxy, *Material);
+				case ERuntimeVirtualTextureMaterialType::Mask4:
+					return Process<FMaterialPolicy_Mask4>(MeshBatch, BatchElementMask, StaticMeshId, OutputAttributeMask, PrimitiveSceneProxy, *MaterialRenderProxy, *Material);
 				case ERuntimeVirtualTextureMaterialType::WorldHeight:
 					return Process<FMaterialPolicy_WorldHeight>(MeshBatch, BatchElementMask, StaticMeshId, OutputAttributeMask, PrimitiveSceneProxy, *MaterialRenderProxy, *Material);
 				case ERuntimeVirtualTextureMaterialType::Displacement:
@@ -963,6 +1013,7 @@ namespace RuntimeVirtualTexture
 				CollectPSOInitializers<FMaterialPolicy_BaseColor>(VertexFactoryData, Material, MeshFillMode, MeshCullMode, OutputAttributeMask, ERuntimeVirtualTextureMaterialType::BaseColor, PSOInitializers);
 				CollectPSOInitializers<FMaterialPolicy_BaseColorNormalRoughness>(VertexFactoryData, Material, MeshFillMode, MeshCullMode, OutputAttributeMask, ERuntimeVirtualTextureMaterialType::BaseColor_Normal_Roughness, PSOInitializers);
 				CollectPSOInitializers<FMaterialPolicy_BaseColorNormalSpecular>(VertexFactoryData, Material, MeshFillMode, MeshCullMode, OutputAttributeMask, ERuntimeVirtualTextureMaterialType::BaseColor_Normal_Specular, PSOInitializers);
+				CollectPSOInitializers<FMaterialPolicy_Mask4>(VertexFactoryData, Material, MeshFillMode, MeshCullMode, OutputAttributeMask, ERuntimeVirtualTextureMaterialType::BaseColor, PSOInitializers);
 				CollectPSOInitializers<FMaterialPolicy_WorldHeight>(VertexFactoryData, Material, MeshFillMode, MeshCullMode, OutputAttributeMask, ERuntimeVirtualTextureMaterialType::WorldHeight, PSOInitializers);
 				CollectPSOInitializers<FMaterialPolicy_Displacement>(VertexFactoryData, Material, MeshFillMode, MeshCullMode, OutputAttributeMask, ERuntimeVirtualTextureMaterialType::Displacement, PSOInitializers);
 			}
@@ -1155,6 +1206,7 @@ namespace RuntimeVirtualTexture
 	IMPLEMENT_SHADER_TYPE(template<>, FShader_VirtualTextureCompress_CS< ERuntimeVirtualTextureMaterialType::BaseColor_Normal_Roughness >, TEXT("/Engine/Private/VirtualTextureCompress.usf"), TEXT("CompressBaseColorNormalRoughnessCS"), SF_Compute);
 	IMPLEMENT_SHADER_TYPE(template<>, FShader_VirtualTextureCompress_CS< ERuntimeVirtualTextureMaterialType::BaseColor_Normal_Specular_YCoCg >, TEXT("/Engine/Private/VirtualTextureCompress.usf"), TEXT("CompressBaseColorNormalSpecularYCoCgCS"), SF_Compute);
 	IMPLEMENT_SHADER_TYPE(template<>, FShader_VirtualTextureCompress_CS< ERuntimeVirtualTextureMaterialType::BaseColor_Normal_Specular_Mask_YCoCg >, TEXT("/Engine/Private/VirtualTextureCompress.usf"), TEXT("CompressBaseColorNormalSpecularMaskYCoCgCS"), SF_Compute);
+	IMPLEMENT_SHADER_TYPE(template<>, FShader_VirtualTextureCompress_CS< ERuntimeVirtualTextureMaterialType::Mask4 >, TEXT("/Engine/Private/VirtualTextureCompress.usf"), TEXT("CompressMask4CS"), SF_Compute);
 	IMPLEMENT_SHADER_TYPE(template<>, FShader_VirtualTextureCompress_CS< ERuntimeVirtualTextureMaterialType::Displacement >, TEXT("/Engine/Private/VirtualTextureCompress.usf"), TEXT("CompressDisplacementCS"), SF_Compute);
 
 
@@ -1194,6 +1246,9 @@ namespace RuntimeVirtualTexture
 		case ERuntimeVirtualTextureMaterialType::BaseColor_Normal_Specular_Mask_YCoCg:
 			AddCompressPass<ERuntimeVirtualTextureMaterialType::BaseColor_Normal_Specular_Mask_YCoCg>(GraphBuilder, FeatureLevel, Parameters, GroupCount);
 			break;
+		case ERuntimeVirtualTextureMaterialType::Mask4:
+			AddCompressPass<ERuntimeVirtualTextureMaterialType::Mask4>(GraphBuilder, FeatureLevel, Parameters, GroupCount);
+			break;
 		case ERuntimeVirtualTextureMaterialType::Displacement:
 			AddCompressPass<ERuntimeVirtualTextureMaterialType::Displacement>(GraphBuilder, FeatureLevel, Parameters, GroupCount);
 			break;
@@ -1260,6 +1315,7 @@ namespace RuntimeVirtualTexture
 	IMPLEMENT_SHADER_TYPE(template<>, FShader_VirtualTextureCopy_PS< ERuntimeVirtualTextureMaterialType::BaseColor_Normal_Specular >, TEXT("/Engine/Private/VirtualTextureCompress.usf"), TEXT("CopyBaseColorNormalSpecularPS"), SF_Pixel);
 	IMPLEMENT_SHADER_TYPE(template<>, FShader_VirtualTextureCopy_PS< ERuntimeVirtualTextureMaterialType::BaseColor_Normal_Specular_YCoCg >, TEXT("/Engine/Private/VirtualTextureCompress.usf"), TEXT("CopyBaseColorNormalSpecularYCoCgPS"), SF_Pixel);
 	IMPLEMENT_SHADER_TYPE(template<>, FShader_VirtualTextureCopy_PS< ERuntimeVirtualTextureMaterialType::BaseColor_Normal_Specular_Mask_YCoCg >, TEXT("/Engine/Private/VirtualTextureCompress.usf"), TEXT("CopyBaseColorNormalSpecularMaskYCoCgPS"), SF_Pixel);
+	IMPLEMENT_SHADER_TYPE(template<>, FShader_VirtualTextureCopy_PS< ERuntimeVirtualTextureMaterialType::Mask4 >, TEXT("/Engine/Private/VirtualTextureCompress.usf"), TEXT("CopyMask4PS"), SF_Pixel);
 	IMPLEMENT_SHADER_TYPE(template<>, FShader_VirtualTextureCopy_PS< ERuntimeVirtualTextureMaterialType::WorldHeight >, TEXT("/Engine/Private/VirtualTextureCompress.usf"), TEXT("CopyWorldHeightPS"), SF_Pixel);
 
 
@@ -1310,6 +1366,9 @@ namespace RuntimeVirtualTexture
 		case ERuntimeVirtualTextureMaterialType::BaseColor_Normal_Specular_Mask_YCoCg:
 			AddCopyPass<ERuntimeVirtualTextureMaterialType::BaseColor_Normal_Specular_Mask_YCoCg>(GraphBuilder, FeatureLevel, Parameters, TextureSize);
 			break;
+		case ERuntimeVirtualTextureMaterialType::Mask4:
+			AddCopyPass<ERuntimeVirtualTextureMaterialType::Mask4>(GraphBuilder, FeatureLevel, Parameters, TextureSize);
+			break;
 		}
 	}
 
@@ -1319,11 +1378,15 @@ namespace RuntimeVirtualTexture
 		switch (MaterialType)
 		{
 		case ERuntimeVirtualTextureMaterialType::BaseColor:
+		case ERuntimeVirtualTextureMaterialType::BaseColor_Normal_Roughness:
 		case ERuntimeVirtualTextureMaterialType::BaseColor_Normal_Specular:
 		case ERuntimeVirtualTextureMaterialType::BaseColor_Normal_Specular_YCoCg:
 		case ERuntimeVirtualTextureMaterialType::BaseColor_Normal_Specular_Mask_YCoCg:
 			AddCopyPass<ERuntimeVirtualTextureMaterialType::BaseColor>(GraphBuilder, FeatureLevel, Parameters, TextureSize);
 			break;
+		case ERuntimeVirtualTextureMaterialType::Mask4:
+			AddCopyPass<ERuntimeVirtualTextureMaterialType::Mask4>(GraphBuilder, FeatureLevel, Parameters, TextureSize);
+			break;
 		case ERuntimeVirtualTextureMaterialType::WorldHeight:
 		case ERuntimeVirtualTextureMaterialType::Displacement:
 			AddCopyPass<ERuntimeVirtualTextureMaterialType::WorldHeight>(GraphBuilder, FeatureLevel, Parameters, TextureSize);
