Fix Mesh Decals with mobile forward renderer with DBuffer enabled

diff --git a/Engine/Source/Runtime/Renderer/Private/CompositionLighting/PostProcessMeshDecals.cpp b/Engine/Source/Runtime/Renderer/Private/CompositionLighting/PostProcessMeshDecals.cpp
--- a/Engine/Source/Runtime/Renderer/Private/CompositionLighting/PostProcessMeshDecals.cpp
+++ b/Engine/Source/Runtime/Renderer/Private/CompositionLighting/PostProcessMeshDecals.cpp
@@ -493,8 +493,14 @@ FMeshPassProcessor* CreateMeshDecalSceneColorMeshProcessor_Mobile(ERHIFeatureLev
 	return new FMeshDecalMeshProcessor(Scene, FeatureLevel, InViewIfDynamicMeshCommand, EDecalRenderTargetMode::SceneColor, EShadingPath::Mobile, InDrawListContext);
 }
 
+FMeshPassProcessor* CreateMeshDecalDBufferMeshProcessor_Mobile(ERHIFeatureLevel::Type FeatureLevel, const FScene* Scene, const FSceneView* InViewIfDynamicMeshCommand, FMeshPassDrawListContext* InDrawListContext)
+{
+	return new FMeshDecalMeshProcessor(Scene, FeatureLevel, InViewIfDynamicMeshCommand, EDecalRenderTargetMode::DBuffer, EShadingPath::Mobile, InDrawListContext);
+}
+
 REGISTER_MESHPASSPROCESSOR_AND_PSOCOLLECTOR(Mobile_MeshDecal_SceneColorAndGBuffer, CreateMeshDecalSceneColorAndGBufferMeshProcessor_Mobile, EShadingPath::Mobile, EMeshPass::MeshDecal_SceneColorAndGBuffer, EMeshPassFlags::CachedMeshCommands | EMeshPassFlags::MainView);
 REGISTER_MESHPASSPROCESSOR_AND_PSOCOLLECTOR(Mobile_MeshDecal_SceneColor, CreateMeshDecalSceneColorMeshProcessor_Mobile, EShadingPath::Mobile, EMeshPass::MeshDecal_SceneColor, EMeshPassFlags::CachedMeshCommands | EMeshPassFlags::MainView);
+REGISTER_MESHPASSPROCESSOR_AND_PSOCOLLECTOR(Mobile_MeshDecal_DBuffer, CreateMeshDecalDBufferMeshProcessor_Mobile, EShadingPath::Mobile, EMeshPass::MeshDecal_DBuffer, EMeshPassFlags::CachedMeshCommands | EMeshPassFlags::MainView);
 
 void DrawDecalMeshCommands(
 	FRDGBuilder& GraphBuilder,
