Add project renderer settings to disable specific RVT material types.
This allows us to remove shader permutations for types that we know that we don't need.
Specification can be overridden per platform if necessary.
When removed, material types no longer show up in selction option for RVT asset.
If an RVT has a removed material type then the RVT volume is not instantiated.

diff --git a/Engine/Source/Editor/VirtualTexturingEditor/Classes/RuntimeVirtualTextureDetailsCustomization.h b/Engine/Source/Editor/VirtualTexturingEditor/Classes/RuntimeVirtualTextureDetailsCustomization.h
--- a/Engine/Source/Editor/VirtualTexturingEditor/Classes/RuntimeVirtualTextureDetailsCustomization.h
+++ b/Engine/Source/Editor/VirtualTexturingEditor/Classes/RuntimeVirtualTextureDetailsCustomization.h
@@ -21,6 +21,9 @@ public:
 protected:
 	FRuntimeVirtualTextureDetailsCustomization();
 
+	/** Refresh array of supported material types. */
+	void RefreshMaterialTypes();
+
 	/** Callback for updating text values after an edit. */
 	void RefreshTextDetails();
 	/** Callback for full update of details view after an edit. */
@@ -33,6 +36,8 @@ protected:
 private:
 	TWeakObjectPtr<URuntimeVirtualTexture> VirtualTexture;
 
+	TArray<int32> SupportedMaterialTypes;
+
 	IDetailLayoutBuilder* CachedDetailBuilder = nullptr;
 
 	TSharedPtr<STextBlock> TileCountText;
diff --git a/Engine/Source/Editor/VirtualTexturingEditor/Private/RuntimeVirtualTextureDetailsCustomization.cpp b/Engine/Source/Editor/VirtualTexturingEditor/Private/RuntimeVirtualTextureDetailsCustomization.cpp
--- a/Engine/Source/Editor/VirtualTexturingEditor/Private/RuntimeVirtualTextureDetailsCustomization.cpp
+++ b/Engine/Source/Editor/VirtualTexturingEditor/Private/RuntimeVirtualTextureDetailsCustomization.cpp
@@ -13,6 +13,7 @@
 #include "RuntimeVirtualTextureBuildStreamingMips.h"
 #include "RuntimeVirtualTextureSetBounds.h"
 #include "ScopedTransaction.h"
+#include "SEnumCombo.h"
 #include "SResetToDefaultMenu.h"
 #include "VirtualTextureBuilderFactory.h"
 #include "VT/RuntimeVirtualTexture.h"
@@ -100,6 +101,37 @@ void FRuntimeVirtualTextureDetailsCustomization::CustomizeDetails(IDetailLayoutB
 		return;
 	}
 
+	RefreshMaterialTypes();
+
+	TSharedRef<IPropertyHandle> MaterialTypePropertyHandle = DetailBuilder.GetProperty(TEXT("MaterialType"));
+	DetailBuilder.EditDefaultProperty(MaterialTypePropertyHandle)->CustomWidget()
+	.NameContent()
+	[
+		MaterialTypePropertyHandle->CreatePropertyNameWidget()
+	]
+	.ValueContent()
+	[
+		SNew(SEnumComboBox, StaticEnum<ERuntimeVirtualTextureMaterialType>())
+			.Font(FAppStyle::GetFontStyle(TEXT("MenuItem.Font")))
+			.EnumValueSubset(SupportedMaterialTypes)
+			.CurrentValue_Lambda([this]()
+			{
+				if (URuntimeVirtualTexture* Texture = VirtualTexture.Get())
+				{
+					return (int32)Texture->GetMaterialType();
+				}
+				return 0;
+			})
+			.OnEnumSelectionChanged_Lambda([this](uint32 NewValue, ESelectInfo::Type)
+			{
+				if (URuntimeVirtualTexture* Texture = VirtualTexture.Get())
+				{
+					Texture->MaterialType = (ERuntimeVirtualTextureMaterialType)NewValue;
+					RefreshDetailsView();
+				}
+			})
+	];
+
 	// Set UIMax dependent on adaptive page table setting
 	FString MaxTileCountString = FString::Printf(TEXT("%d"), URuntimeVirtualTexture::GetMaxTileCountLog2(VirtualTexture->GetAdaptivePageTable()));
 	DetailBuilder.GetProperty(FName(TEXT("TileCount")))->SetInstanceMetaData("UIMax", MaxTileCountString);
@@ -152,6 +184,28 @@ void FRuntimeVirtualTextureDetailsCustomization::CustomizeDetails(IDetailLayoutB
 	RefreshTextDetails();
 }
 
+void FRuntimeVirtualTextureDetailsCustomization::RefreshMaterialTypes()
+{
+	// Filter for enabled material types.
+	SupportedMaterialTypes.Reset();
+	SupportedMaterialTypes.Reserve((int32)ERuntimeVirtualTextureMaterialType::Count);
+	
+	// Include currently selected type even if it is disabled.
+	ERuntimeVirtualTextureMaterialType CurrentType = ERuntimeVirtualTextureMaterialType::Count;
+	if (URuntimeVirtualTexture* Texture = VirtualTexture.Get())
+	{
+		CurrentType = Texture->GetMaterialType();
+	}
+
+	for (ERuntimeVirtualTextureMaterialType Type : TEnumRange<ERuntimeVirtualTextureMaterialType>())
+	{
+		if (RuntimeVirtualTexture::IsMaterialTypeSupported(Type) || Type == CurrentType)
+		{
+			SupportedMaterialTypes.Add((int32)Type);
+		}
+	}
+}
+
 void FRuntimeVirtualTextureDetailsCustomization::RefreshTextDetails()
 {
 	if (URuntimeVirtualTexture* Texture = VirtualTexture.Get())
diff --git a/Engine/Source/Editor/VirtualTexturingEditor/VirtualTexturingEditor.Build.cs b/Engine/Source/Editor/VirtualTexturingEditor/VirtualTexturingEditor.Build.cs
--- a/Engine/Source/Editor/VirtualTexturingEditor/VirtualTexturingEditor.Build.cs
+++ b/Engine/Source/Editor/VirtualTexturingEditor/VirtualTexturingEditor.Build.cs
@@ -14,6 +14,7 @@ public class VirtualTexturingEditor : ModuleRules
 				"Core",
 				"CoreUObject",
 				"EditorFramework",
+				"EditorWidgets",
 				"Engine",
 				"InputCore",
 				"Landscape",
diff --git a/Engine/Source/Runtime/Engine/Classes/Engine/RendererSettings.h b/Engine/Source/Runtime/Engine/Classes/Engine/RendererSettings.h
--- a/Engine/Source/Runtime/Engine/Classes/Engine/RendererSettings.h
+++ b/Engine/Source/Runtime/Engine/Classes/Engine/RendererSettings.h
@@ -444,6 +444,55 @@ class URendererSettings : public UDeveloperSettings
 		ToolTip = "Lower factor will increase virtual texture feedback resolution which increases CPU/GPU overhead, but may decrease streaming latency, especially if materials use many virtual textures."))
 	uint32 VirtualTextureFeedbackFactor;
 
+	UPROPERTY(config, EditAnywhere, Category = "Runtime Virtual Textures|Enable Material Types", meta = (
+		EditCondition = "bVirtualTextures",
+		ConsoleVariable = "r.vt.rvt.EnableBaseColor", DisplayName = "Enable 'Base Color'",
+		ToolTip = "Disable material type when not required by a project. This minimizes shader permutation count for materials with Runtime Virtual Texture Output nodes.",
+		ConfigRestartRequired = true))
+	bool bEnableRVTBaseColor;
+
+	UPROPERTY(config, EditAnywhere, Category = "Runtime Virtual Textures|Enable Material Types", meta = (
+		EditCondition = "bVirtualTextures",
+		ConsoleVariable = "r.vt.rvt.EnableBaseColorRoughness", DisplayName = "Enable 'Base Color, Normal, Roughness'",
+		ToolTip = "Disable material type when not required by a project. This minimizes shader permutation count for materials with Runtime Virtual Texture Output nodes.",
+		ConfigRestartRequired = true))
+	bool bEnableRVTBaseColorRoughness;
+
+	UPROPERTY(config, EditAnywhere, Category = "Runtime Virtual Textures|Enable Material Types", meta = (
+		EditCondition = "bVirtualTextures",
+		ConsoleVariable = "r.vt.rvt.EnableBaseColorSpecular", DisplayName = "Enable 'Base Color, Normal, Roughness, Specular' family",
+		ToolTip = "Disable family of material types when not required by a project. This minimizes shader permutation count for materials with Runtime Virtual Texture Output nodes.",
+		ConfigRestartRequired = true))
+	bool bEnableRVTBaseColorSpecular;
+
+	UPROPERTY(config, EditAnywhere, Category = "Runtime Virtual Textures|Enable Material Types", meta = (
+		EditCondition = "bVirtualTextures",
+		ConsoleVariable = "r.vt.rvt.EnableMask4", DisplayName = "Enable 'Mask4'",
+		ToolTip = "Disable material type when not required by a project. This minimizes shader permutation count for materials with Runtime Virtual Texture Output nodes.",
+		ConfigRestartRequired = true))
+	bool bEnableRVTMask4;
+
+	UPROPERTY(config, EditAnywhere, Category = "Runtime Virtual Textures|Enable Material Types", meta = (
+		EditCondition = "bVirtualTextures",
+		ConsoleVariable = "r.vt.rvt.EnableWorldHeight", DisplayName = "Enable 'World Height'",
+		ToolTip = "Disable material type when not required by a project. This minimizes shader permutation count for materials with Runtime Virtual Texture Output nodes.",
+		ConfigRestartRequired = true))
+	bool bEnableRVTWorldHeight;
+
+	UPROPERTY(config, EditAnywhere, Category = "Runtime Virtual Textures|Enable Material Types", meta = (
+		EditCondition = "bVirtualTextures",
+		ConsoleVariable = "r.vt.rvt.EnableDisplacement", DisplayName = "Enable 'Displacement'",
+		ToolTip = "Disable material type when not required by a project. This minimizes shader permutation count for materials with Runtime Virtual Texture Output nodes.",
+		ConfigRestartRequired = true))
+	bool bEnableRVTDisplacement;
+
+	UPROPERTY(config, EditAnywhere, Category = "Runtime Virtual Textures", meta = (
+		EditCondition = "bVirtualTextures",
+		ConsoleVariable = "r.vt.rvt.HighQualityPerPixelHeight", DisplayName = "High quality landscape height sampling.",
+		ToolTip = "Use higher quality sampling of per pixel heightmaps when rendering to Runtime Virtual Texture.",
+		ConfigRestartRequired = true))
+	bool bUseHighQualityRVTHeightSampling;
+
 	UPROPERTY(config, EditAnywhere, Category = WorkingColorSpace, meta = (
 		DisplayName = "Working Color Space",
 		ToolTip = "Choose from list of provided working color spaces, or custom to provide user-defined space.",
diff --git a/Engine/Source/Runtime/Engine/Classes/VT/RuntimeVirtualTexture.h b/Engine/Source/Runtime/Engine/Classes/VT/RuntimeVirtualTexture.h
--- a/Engine/Source/Runtime/Engine/Classes/VT/RuntimeVirtualTexture.h
+++ b/Engine/Source/Runtime/Engine/Classes/VT/RuntimeVirtualTexture.h
@@ -8,6 +8,7 @@
 #include "VT/RuntimeVirtualTextureEnum.h"
 #include "RuntimeVirtualTexture.generated.h"
 
+enum EShaderPlatform : uint16;
 namespace UE { namespace Shader	{ enum class EValueType : uint8; } }
 
 /** Runtime virtual texture UObject */
@@ -17,6 +18,8 @@ class URuntimeVirtualTexture : public UObject
 	GENERATED_UCLASS_BODY()
 	ENGINE_API ~URuntimeVirtualTexture();
 
+	friend class FRuntimeVirtualTextureDetailsCustomization;
+
 protected:
 	/** 
 	 * Size of virtual texture in tiles. (Actual values increase in powers of 2).
@@ -86,7 +89,7 @@ protected:
 public:
 	/** Get the material set that this virtual texture stores. */
 	ERuntimeVirtualTextureMaterialType GetMaterialType() const { return MaterialType; }
-
+	
 	/** Public getter for virtual texture tile count */
 	UFUNCTION(BlueprintGetter)
 	int32 GetTileCount() const { return GetClampedTileCount(TileCount, bAdaptive); }
@@ -204,4 +207,8 @@ namespace RuntimeVirtualTexture
 		IVirtualTexture* InProducer,
 		IVirtualTexture* InStreamingProducer,
 		int32 InTransitionLevel);
+
+	/** Get if the material type is supported. Individual material types can be disabled in project settings to reduce shader permutations. */
+	ENGINE_API bool IsMaterialTypeSupported(ERuntimeVirtualTextureMaterialType MaterialType);
+	ENGINE_API bool IsMaterialTypeSupported(ERuntimeVirtualTextureMaterialType MaterialType, EShaderPlatform InPlatform);
 }
diff --git a/Engine/Source/Runtime/Engine/Private/Components/RuntimeVirtualTextureComponent.cpp b/Engine/Source/Runtime/Engine/Private/Components/RuntimeVirtualTextureComponent.cpp
--- a/Engine/Source/Runtime/Engine/Private/Components/RuntimeVirtualTextureComponent.cpp
+++ b/Engine/Source/Runtime/Engine/Private/Components/RuntimeVirtualTextureComponent.cpp
@@ -183,12 +183,18 @@ static ERuntimeVirtualTextureMaterialQuality ConvertMaterialQualityEnum(EMateria
 
 bool URuntimeVirtualTextureComponent::IsEnabledInScene() const
 {
-	const bool bUseNanite = UseNanite(GetScene()->GetShaderPlatform());
+	const EShaderPlatform ShaderPlatform = GetScene()->GetShaderPlatform();
+	const bool bUseNanite = UseNanite(ShaderPlatform);
 	if (bEnableForNaniteOnly && !bUseNanite)
 	{
 		return false;
 	}
 
+	if (!RuntimeVirtualTexture::IsMaterialTypeSupported(VirtualTexture->GetMaterialType(), ShaderPlatform))
+	{
+		return false;
+	}
+
 	if (UWorld* World = GetWorld())
 	{
 		if (World->WorldType == EWorldType::Game || World->WorldType == EWorldType::PIE)
diff --git a/Engine/Source/Runtime/Engine/Private/VT/RuntimeVirtualTexture.cpp b/Engine/Source/Runtime/Engine/Private/VT/RuntimeVirtualTexture.cpp
--- a/Engine/Source/Runtime/Engine/Private/VT/RuntimeVirtualTexture.cpp
+++ b/Engine/Source/Runtime/Engine/Private/VT/RuntimeVirtualTexture.cpp
@@ -8,6 +8,7 @@
 #include "RendererInterface.h"
 #include "RenderingThread.h"
 #include "Shader/ShaderTypes.h"
+#include "ShaderPlatformCachedIniValue.h"
 #include "UObject/AssetRegistryTagsContext.h"
 #include "UObject/UnrealType.h"
 #include "VT/RuntimeVirtualTextureNotify.h"
@@ -18,6 +19,13 @@
 
 #include UE_INLINE_GENERATED_CPP_BY_NAME(RuntimeVirtualTexture)
 
+static TAutoConsoleVariable<bool> CVarRVTEnableBaseColor(TEXT("r.VT.RVT.EnableBaseColor"), true, TEXT("Enable 'Base Color' RVT Material Type"), ECVF_ReadOnly);
+static TAutoConsoleVariable<bool> CVarRVTEnableBaseColorRoughness(TEXT("r.VT.RVT.EnableBaseColorRoughness"), true, TEXT("Enable 'Base Color, Normal, Roughness' RVT Material Type"), ECVF_ReadOnly);
+static TAutoConsoleVariable<bool> CVarRVTEnableBaseColorSpecular(TEXT("r.VT.RVT.EnableBaseColorSpecular"), true, TEXT("Enable 'Base Color, Normal, Roughness, Specular' family of RVT Material Types"), ECVF_ReadOnly);
+static TAutoConsoleVariable<bool> CVarRVTEnableMask4(TEXT("r.VT.RVT.EnableMask4"), true, TEXT("Enable 'Mask4' RVT Material Type"), ECVF_ReadOnly);
+static TAutoConsoleVariable<bool> CVarRVTEnableWorldHeight(TEXT("r.VT.RVT.EnableWorldHeight"), true, TEXT("Enable 'World Height' RVT Material Type"), ECVF_ReadOnly);
+static TAutoConsoleVariable<bool> CVarRVTEnableDisplacement(TEXT("r.VT.RVT.EnableDisplacement"), true, TEXT("Enable 'Displacement' RVT Material Type"), ECVF_ReadOnly);
+
 namespace
 {
 	/** Null producer to use as placeholder when no producer has been set on a URuntimeVirtualTexture */
@@ -706,5 +714,42 @@ namespace RuntimeVirtualTexture
 	{
 		return (InStreamingProducer == nullptr) ? InProducer : new FVirtualTextureLevelRedirector(InProducer, InStreamingProducer, InTransitionLevel);
 	}
+
+	bool IsMaterialTypeSupported(ERuntimeVirtualTextureMaterialType InMaterialType, EShaderPlatform InPlatform)
+	{
+		static FShaderPlatformCachedIniValue<bool> CPlatformVarRVTEnableBaseColor(CVarRVTEnableBaseColor.AsVariable());
+		static FShaderPlatformCachedIniValue<bool> CPlatformVarRVTEnableBaseColorRoughness(CVarRVTEnableBaseColorRoughness.AsVariable());
+		static FShaderPlatformCachedIniValue<bool> CPlatformVarRVTEnableBaseColorSpecular(CVarRVTEnableBaseColorSpecular.AsVariable());
+		static FShaderPlatformCachedIniValue<bool> CPlatformVarRVTEnableMask4(CVarRVTEnableMask4.AsVariable());
+		static FShaderPlatformCachedIniValue<bool> CPlatformVarRVTEnableWorldHeight(CVarRVTEnableWorldHeight.AsVariable());
+		static FShaderPlatformCachedIniValue<bool> CPlatformVarRVTEnableDisplacement(CVarRVTEnableDisplacement.AsVariable());
+
+		switch (InMaterialType)
+		{
+		case ERuntimeVirtualTextureMaterialType::BaseColor:
+			return CPlatformVarRVTEnableBaseColor.Get(InPlatform);
+		case ERuntimeVirtualTextureMaterialType::BaseColor_Normal_Roughness:
+			return CPlatformVarRVTEnableBaseColorRoughness.Get(InPlatform);
+		case ERuntimeVirtualTextureMaterialType::BaseColor_Normal_Specular:
+		case ERuntimeVirtualTextureMaterialType::BaseColor_Normal_Specular_YCoCg:
+		case ERuntimeVirtualTextureMaterialType::BaseColor_Normal_Specular_Mask_YCoCg:
+			return CPlatformVarRVTEnableBaseColorSpecular.Get(InPlatform);
+		case ERuntimeVirtualTextureMaterialType::Mask4:
+			return CPlatformVarRVTEnableMask4.Get(InPlatform);
+		case ERuntimeVirtualTextureMaterialType::WorldHeight:
+			return CPlatformVarRVTEnableWorldHeight.Get(InPlatform);
+		case ERuntimeVirtualTextureMaterialType::Displacement:
+			return CPlatformVarRVTEnableDisplacement.Get(InPlatform);
+		default:
+			checkNoEntry();
+		}
+
+		return false;
+	}
+
+	bool IsMaterialTypeSupported(ERuntimeVirtualTextureMaterialType MaterialType)
+	{
+		return IsMaterialTypeSupported(MaterialType, GetFeatureLevelShaderPlatform(GMaxRHIFeatureLevel));
+	}
 }
 
diff --git a/Engine/Source/Runtime/Engine/Public/VT/RuntimeVirtualTextureEnum.h b/Engine/Source/Runtime/Engine/Public/VT/RuntimeVirtualTextureEnum.h
--- a/Engine/Source/Runtime/Engine/Public/VT/RuntimeVirtualTextureEnum.h
+++ b/Engine/Source/Runtime/Engine/Public/VT/RuntimeVirtualTextureEnum.h
@@ -2,7 +2,7 @@
 
 #pragma once
 
-#include "CoreMinimal.h"
+#include "Misc/EnumRange.h"
 #include "RuntimeVirtualTextureEnum.generated.h"
 
 /** Maximum number of texture layers we will have in a runtime virtual texture. Increase if we add a ERuntimeVirtualTextureMaterialType with more layers. */
@@ -50,6 +50,8 @@ enum class ERuntimeVirtualTextureMaterialType : uint8
 namespace RuntimeVirtualTexture { enum { MaterialType_NumBits = 3 }; }
 static_assert((uint32)ERuntimeVirtualTextureMaterialType::Count <= (1 << (uint32)RuntimeVirtualTexture::MaterialType_NumBits), "NumBits is too small");
 
+ENUM_RANGE_BY_COUNT(ERuntimeVirtualTextureMaterialType, ERuntimeVirtualTextureMaterialType::Count);
+
 /** Enumeration of main pass behaviors when rendering to a runtime virtual texture. */
 UENUM()
 enum class ERuntimeVirtualTextureMainPassType : uint8
diff --git a/Engine/Source/Runtime/Renderer/Private/VT/RuntimeVirtualTextureRender.cpp b/Engine/Source/Runtime/Renderer/Private/VT/RuntimeVirtualTextureRender.cpp
--- a/Engine/Source/Runtime/Renderer/Private/VT/RuntimeVirtualTextureRender.cpp
+++ b/Engine/Source/Runtime/Renderer/Private/VT/RuntimeVirtualTextureRender.cpp
@@ -4,6 +4,7 @@
 
 #include "Components/RuntimeVirtualTextureComponent.h"
 #include "DataDrivenShaderPlatformInfo.h"
+#include "EngineModule.h"
 #include "GlobalShader.h"
 #include "GPUScene.h"
 #include "MaterialShader.h"
@@ -23,7 +24,6 @@
 #include "SimpleMeshDrawCommandPass.h"
 #include "StaticMeshBatch.h"
 #include "SceneRendering.h"
-#include "EngineModule.h"
 
 CSV_DECLARE_CATEGORY_EXTERN(VirtualTexturing);
 
@@ -172,6 +172,11 @@ namespace RuntimeVirtualTexture
 	class FMaterialPolicy_BaseColor
 	{
 	public:
+		static bool ShouldCompilePermutation(const FMeshMaterialShaderPermutationParameters& Parameters)
+		{
+			return RuntimeVirtualTexture::IsMaterialTypeSupported(ERuntimeVirtualTextureMaterialType::BaseColor, Parameters.Platform);
+		}
+
 		static void ModifyCompilationEnvironment(FShaderCompilerEnvironment& OutEnvironment)
 		{
 			OutEnvironment.SetDefine(TEXT("OUT_BASECOLOR"), 1);
@@ -262,6 +267,11 @@ namespace RuntimeVirtualTexture
 		}
 
 	public:
+		static bool ShouldCompilePermutation(const FMeshMaterialShaderPermutationParameters& Parameters)
+		{
+			return RuntimeVirtualTexture::IsMaterialTypeSupported(ERuntimeVirtualTextureMaterialType::BaseColor_Normal_Specular, Parameters.Platform);
+		}
+
 		static void ModifyCompilationEnvironment(FShaderCompilerEnvironment& OutEnvironment)
 		{
 			OutEnvironment.SetDefine(TEXT("OUT_BASECOLOR_NORMAL_SPECULAR"), 1);
@@ -349,6 +359,11 @@ namespace RuntimeVirtualTexture
 			}
 		}
 	public:
+		static bool ShouldCompilePermutation(const FMeshMaterialShaderPermutationParameters& Parameters)
+		{
+			return RuntimeVirtualTexture::IsMaterialTypeSupported(ERuntimeVirtualTextureMaterialType::BaseColor_Normal_Roughness, Parameters.Platform);
+		}
+		
 		static void ModifyCompilationEnvironment(FShaderCompilerEnvironment& OutEnvironment)
 		{
 			OutEnvironment.SetDefine(TEXT("OUT_BASECOLOR_NORMAL_ROUGHNESS"), 1);
@@ -363,6 +378,11 @@ namespace RuntimeVirtualTexture
 	class FMaterialPolicy_Mask4
 	{
 	public:
+		static bool ShouldCompilePermutation(const FMeshMaterialShaderPermutationParameters& Parameters)
+		{
+			return RuntimeVirtualTexture::IsMaterialTypeSupported(ERuntimeVirtualTextureMaterialType::Mask4, Parameters.Platform);
+		}
+
 		static void ModifyCompilationEnvironment(FShaderCompilerEnvironment& OutEnvironment)
 		{
 			OutEnvironment.SetDefine(TEXT("OUT_MASK4"), 1);
@@ -380,6 +400,11 @@ namespace RuntimeVirtualTexture
 	class FMaterialPolicy_WorldHeight
 	{
 	public:
+		static bool ShouldCompilePermutation(const FMeshMaterialShaderPermutationParameters& Parameters)
+		{
+			return RuntimeVirtualTexture::IsMaterialTypeSupported(ERuntimeVirtualTextureMaterialType::WorldHeight, Parameters.Platform);
+		}
+
 		static void ModifyCompilationEnvironment(FShaderCompilerEnvironment& OutEnvironment)
 		{
 			OutEnvironment.SetDefine(TEXT("OUT_WORLDHEIGHT"), 1);
@@ -396,6 +421,11 @@ namespace RuntimeVirtualTexture
 	class FMaterialPolicy_Displacement
 	{
 	public:
+		static bool ShouldCompilePermutation(const FMeshMaterialShaderPermutationParameters& Parameters)
+		{
+			return RuntimeVirtualTexture::IsMaterialTypeSupported(ERuntimeVirtualTextureMaterialType::Displacement, Parameters.Platform);
+		}
+
 		static void ModifyCompilationEnvironment(FShaderCompilerEnvironment& OutEnvironment)
 		{
 			OutEnvironment.SetDefine(TEXT("OUT_DISPLACEMENT"), 1);
@@ -423,6 +453,11 @@ namespace RuntimeVirtualTexture
 			: FShader_VirtualTextureMaterialDraw(Initializer)
 		{}
 
+		static bool ShouldCompilePermutation(const FMeshMaterialShaderPermutationParameters& Parameters)
+		{
+			return FShader_VirtualTextureMaterialDraw::ShouldCompilePermutation(Parameters) && MaterialPolicy::ShouldCompilePermutation(Parameters);
+		}
+
 		static void ModifyCompilationEnvironment(const FMaterialShaderPermutationParameters& Parameters, FShaderCompilerEnvironment& OutEnvironment)
 		{
 			FShader_VirtualTextureMaterialDraw::ModifyCompilationEnvironment(Parameters, OutEnvironment);
@@ -1199,6 +1234,11 @@ namespace RuntimeVirtualTexture
 		FShader_VirtualTextureCompress_CS(const ShaderMetaType::CompiledShaderInitializerType& Initializer)
 			: FShader_VirtualTextureCompress(Initializer)
 		{}
+
+		static bool ShouldCompilePermutation(const FGlobalShaderPermutationParameters& Parameters)
+		{
+			return UseVirtualTexturing(Parameters.Platform) && RuntimeVirtualTexture::IsMaterialTypeSupported(MaterialType, Parameters.Platform);
+		}
 	};
 
 	IMPLEMENT_SHADER_TYPE(template<>, FShader_VirtualTextureCompress_CS< ERuntimeVirtualTextureMaterialType::BaseColor >, TEXT("/Engine/Private/VirtualTextureCompress.usf"), TEXT("CompressBaseColorCS"), SF_Compute);
