Fix Substrate issue in forward rendering.
* Material buffer and top layer texture were bound during scene clearing while not needed.
* Substrate resources where created while not used.

diff --git a/Engine/Source/Runtime/Renderer/Private/LightRendering.cpp b/Engine/Source/Runtime/Renderer/Private/LightRendering.cpp
--- a/Engine/Source/Runtime/Renderer/Private/LightRendering.cpp
+++ b/Engine/Source/Runtime/Renderer/Private/LightRendering.cpp
@@ -2423,7 +2423,7 @@ static FDeferredLightPS::FParameters GetDeferredLightPSParameters(
 
 	// PS - Render Targets
 	Out.RenderTargets[0] = FRenderTargetBinding(SceneColorTexture, ERenderTargetLoadAction::ELoad);
-	if (Substrate::IsOpaqueRoughRefractionEnabled())
+	if (Substrate::IsOpaqueRoughRefractionEnabled() && Substrate::UsesSubstrateMaterialBuffer(View.GetShaderPlatform()))
 	{
 		Out.RenderTargets[1] = FRenderTargetBinding(Scene->SubstrateSceneData.SeparatedOpaqueRoughRefractionSceneColor, ERenderTargetLoadAction::ELoad);
 		Out.RenderTargets[2] = FRenderTargetBinding(Scene->SubstrateSceneData.SeparatedSubSurfaceSceneColor, ERenderTargetLoadAction::ELoad);
diff --git a/Engine/Source/Runtime/Renderer/Private/PostProcess/PostProcessCompositeEditorPrimitives.cpp b/Engine/Source/Runtime/Renderer/Private/PostProcess/PostProcessCompositeEditorPrimitives.cpp
--- a/Engine/Source/Runtime/Renderer/Private/PostProcess/PostProcessCompositeEditorPrimitives.cpp
+++ b/Engine/Source/Runtime/Renderer/Private/PostProcess/PostProcessCompositeEditorPrimitives.cpp
@@ -259,7 +259,7 @@ FScreenPassTexture AddEditorPrimitivePass(
 	
 	// Subtrate data might not be produced in certain case (e.g., path-tracer). In such a case we force generate 
 	// them with a simple clear to please validation.
-	if (Substrate::IsSubstrateEnabled() && !HasBeenProduced(View.SubstrateViewData.SceneData->TopLayerTexture))
+	if (Substrate::IsSubstrateEnabled() && Substrate::UsesSubstrateMaterialBuffer(View.GetShaderPlatform()) && !HasBeenProduced(View.SubstrateViewData.SceneData->TopLayerTexture))
 	{
 		FRDGTextureClearInfo ClearInfo;
 		AddClearRenderTargetPass(GraphBuilder, View.SubstrateViewData.SceneData->TopLayerTexture, ClearInfo);
diff --git a/Engine/Source/Runtime/Renderer/Private/Substrate/Substrate.cpp b/Engine/Source/Runtime/Renderer/Private/Substrate/Substrate.cpp
--- a/Engine/Source/Runtime/Renderer/Private/Substrate/Substrate.cpp
+++ b/Engine/Source/Runtime/Renderer/Private/Substrate/Substrate.cpp
@@ -109,6 +109,11 @@ const TCHAR* ToString(ESubstrateTileType Type)
 namespace Substrate
 {
 
+bool UsesSubstrateMaterialBuffer(EShaderPlatform In)
+{
+	return IsUsingGBuffers(In);
+}
+
 uint32 GetMaterialBufferAllocationMode()
 {
 	return FMath::Clamp(CVarSubstrateAllocationMode.GetValueOnAnyThread(), 0, 2);
@@ -171,7 +176,7 @@ bool GetSubstrateUsesComplexSpecialPath(const FViewInfo& View)
 	return false;
 }
 
-static void BindSubstrateGlobalUniformParameters(FRDGBuilder& GraphBuilder, FSubstrateViewData* SubstrateViewData, FSubstrateGlobalUniformParameters& OutSubstrateUniformParameters);
+static void BindSubstrateGlobalUniformParameters(FRDGBuilder& GraphBuilder, const FSubstrateViewData* SubstrateViewData, bool bNeedsMaterialBuffer, FSubstrateGlobalUniformParameters& OutSubstrateUniformParameters);
 
 bool SupportsCMask(const FStaticShaderPlatform InPlatform)
 {
@@ -188,10 +193,13 @@ static EPixelFormat GetClassificationTileFormat(const FIntPoint& InResolution, u
 	return InTileEncoding == SUBSTRATE_TILE_ENCODING_16BITS ? PF_R32_UINT : PF_R16_UINT;
 }
 
-static void InitialiseSubstrateViewData(FRDGBuilder& GraphBuilder, FViewInfo& View, const FSceneTexturesConfig& SceneTexturesConfig, bool bNeedClosureOffets, FSubstrateSceneData& SceneData)
+static void InitialiseSubstrateViewData(FRDGBuilder& GraphBuilder, FViewInfo& View, const FSceneTexturesConfig& SceneTexturesConfig, bool bNeedsClosureOffets, bool bNeedsMaterialBuffer, FSubstrateSceneData& SceneData)
 {
 	// Sanity check: the scene data should already exist 
-	check(SceneData.MaterialTextureArray != nullptr);
+	if (bNeedsMaterialBuffer)
+	{
+		check(SceneData.MaterialTextureArray != nullptr);
+	}
 
 	FSubstrateViewData& Out = View.SubstrateViewData;
 	Out.Reset();
@@ -204,6 +212,7 @@ static void InitialiseSubstrateViewData(FRDGBuilder& GraphBuilder, FViewInfo& Vi
 		const FIntPoint TileResolution(FMath::DivideAndRoundUp(DynResIndependentViewSize.X, SUBSTRATE_TILE_SIZE), FMath::DivideAndRoundUp(DynResIndependentViewSize.Y, SUBSTRATE_TILE_SIZE));
 
 		// Tile classification buffers
+		if (bNeedsMaterialBuffer)
 		{
 			// Indirect draw
 			Out.ClassificationTileDrawIndirectBuffer = GraphBuilder.CreateBuffer(FRDGBufferDesc::CreateIndirectDesc<FRHIDrawIndirectParameters>(ESubstrateTileType::ECount), TEXT("Substrate.SubstrateTileDrawIndirectBuffer"));
@@ -248,7 +257,7 @@ static void InitialiseSubstrateViewData(FRDGBuilder& GraphBuilder, FViewInfo& Vi
 		}
 
 		// Closure tiles
-		if (bNeedClosureOffets)
+		if (bNeedsClosureOffets)
 		{
 			const FIntPoint TileCount = GetSubstrateTextureTileResolution(View, DynResIndependentViewSize);
 			const uint32 LayerCount = GetSubstrateMaxClosureCount(View);
@@ -263,7 +272,7 @@ static void InitialiseSubstrateViewData(FRDGBuilder& GraphBuilder, FViewInfo& Vi
 		}
 		else
 		{
-			Out.TileCount = GetSubstrateTextureTileResolution(View, DynResIndependentViewSize);
+			Out.TileCount = GetSubstrateTextureTileResolution(View, FIntPoint(2,2));
 			Out.LayerCount = 1;
 			Out.ClosureTilePerThreadDispatchIndirectBuffer = nullptr;
 			Out.ClosureTileDispatchIndirectBuffer = nullptr;
@@ -274,7 +283,7 @@ static void InitialiseSubstrateViewData(FRDGBuilder& GraphBuilder, FViewInfo& Vi
 		// Create the readable uniform buffers
 		{
 			FSubstrateGlobalUniformParameters* SubstrateUniformParameters = GraphBuilder.AllocParameters<FSubstrateGlobalUniformParameters>();
-			BindSubstrateGlobalUniformParameters(GraphBuilder, &Out, *SubstrateUniformParameters);
+			BindSubstrateGlobalUniformParameters(GraphBuilder, &Out, bNeedsMaterialBuffer, *SubstrateUniformParameters);
 			Out.SubstrateGlobalUniformParameters = GraphBuilder.CreateUniformBuffer(SubstrateUniformParameters);
 		}
 	}
@@ -336,9 +345,16 @@ void InitialiseSubstrateFrameSceneData(FRDGBuilder& GraphBuilder, FSceneRenderer
 	};
 
 	// Compute the max byte per pixels required by the views
-	bool bNeedClosureOffsets = false;
-	bool bNeedUAV = false;
+	bool bNeedsMaterialBuffer = UsesSubstrateMaterialBuffer(SceneRenderer.ShaderPlatform);
+	bool bNeedsClosureOffsets = false;
+	bool bNeedsUAV = false;
 	bool bUseDBufferPass = false;
+	
+	FIntPoint SceneTextureExtent = SceneRenderer.GetActiveSceneTexturesConfig().Extent;
+	if (!IsSubstrateEnabled() || !bNeedsMaterialBuffer)
+	{
+		SceneTextureExtent = FIntPoint(2, 2);
+	}
 
 	FIntPoint MaterialBufferSizeXY;
 	UpdateMaterialBufferToTiledResolution(FIntPoint(1, 1), MaterialBufferSizeXY);
@@ -357,8 +373,8 @@ void InitialiseSubstrateFrameSceneData(FRDGBuilder& GraphBuilder, FSceneRenderer
 		Out.ViewsMaxClosurePerPixel = 0;
 		for (const FViewInfo& View : SceneRenderer.Views)
 		{
-			bNeedClosureOffsets = bNeedClosureOffsets || NeedClosureOffsets(SceneRenderer.Scene, View);
-			bNeedUAV = bNeedUAV || IsDBufferPassEnabled(View.GetShaderPlatform()) || DoesPlatformSupportNanite(SceneRenderer.ShaderPlatform, true);
+			bNeedsClosureOffsets = bNeedsClosureOffsets || NeedClosureOffsets(SceneRenderer.Scene, View);
+			bNeedsUAV = bNeedsUAV || IsDBufferPassEnabled(View.GetShaderPlatform()) || DoesPlatformSupportNanite(SceneRenderer.ShaderPlatform, true);
 			Out.ViewsMaxBytesPerPixel = FMath::Max(Out.ViewsMaxBytesPerPixel, View.SubstrateViewData.MaxBytesPerPixel);
 			Out.ViewsMaxClosurePerPixel = FMath::Max(Out.ViewsMaxClosurePerPixel, View.SubstrateViewData.MaxClosurePerPixel);
 			bUseDBufferPass = bUseDBufferPass || IsDBufferPassEnabled(View.GetShaderPlatform());
@@ -411,17 +427,18 @@ void InitialiseSubstrateFrameSceneData(FRDGBuilder& GraphBuilder, FSceneRenderer
 		Out.EffectiveMaxBytesPerPixel = FMath::DivideAndRoundUp(CurrentMaxBytesPerPixel, RoundToValue) * RoundToValue;
 		Out.EffectiveMaxClosurePerPixel = CurrentMaxClosurePerPixel;
 
-		FIntPoint SceneTextureExtent = SceneRenderer.GetActiveSceneTexturesConfig().Extent;
 		
 		// We need to allocate enough for the tiled memory addressing of material data to always work
 		UpdateMaterialBufferToTiledResolution(SceneTextureExtent, MaterialBufferSizeXY);
 
 		// Top layer texture
+		if (bNeedsMaterialBuffer)
 		{
-			Out.TopLayerTexture = GraphBuilder.CreateTexture(FRDGTextureDesc::Create2D(SceneTextureExtent, GetTopLayerTextureFormat(bUseDBufferPass), FClearValueBinding::Black, TexCreate_RenderTargetable | TexCreate_ShaderResource | TexCreate_FastVRAM | (bNeedUAV ? TexCreate_UAV : TexCreate_None)), TEXT("Substrate.TopLayerTexture"));
+			Out.TopLayerTexture = GraphBuilder.CreateTexture(FRDGTextureDesc::Create2D(SceneTextureExtent, GetTopLayerTextureFormat(bUseDBufferPass), FClearValueBinding::Black, TexCreate_RenderTargetable | TexCreate_ShaderResource | TexCreate_FastVRAM | (bNeedsUAV ? TexCreate_UAV : TexCreate_None)), TEXT("Substrate.TopLayerTexture"));
 		}
 
 		// Separated subsurface and rough refraction textures
+		if (bNeedsMaterialBuffer)
 		{
 			const bool bIsSubstrateOpaqueMaterialRoughRefractionEnabled = IsOpaqueRoughRefractionEnabled();
 			const FIntPoint OpaqueRoughRefractionSceneExtent		 = bIsSubstrateOpaqueMaterialRoughRefractionEnabled ? SceneTextureExtent : FIntPoint(4, 4);
@@ -445,7 +462,7 @@ void InitialiseSubstrateFrameSceneData(FRDGBuilder& GraphBuilder, FSceneRenderer
 		}
 
 		// Closure offsets
-		if (bNeedClosureOffsets)
+		if (bNeedsMaterialBuffer && bNeedsClosureOffsets)
 		{
 			Out.ClosureOffsetTexture = GraphBuilder.CreateTexture(FRDGTextureDesc::Create2D(SceneTextureExtent, PF_R32_UINT, FClearValueBinding::None, TexCreate_UAV | TexCreate_ShaderResource), TEXT("Substrate.ClosureOffsets"));
 			AddClearUAVPass(GraphBuilder, GraphBuilder.CreateUAV(Out.ClosureOffsetTexture), 0u);
@@ -457,39 +474,48 @@ void InitialiseSubstrateFrameSceneData(FRDGBuilder& GraphBuilder, FSceneRenderer
 	}
 
 	// Create the material data container
-	FIntPoint SceneTextureExtent = IsSubstrateEnabled() ? SceneRenderer.GetActiveSceneTexturesConfig().Extent : FIntPoint(2, 2);
-
 	const uint32 SliceCountSSS = SUBSTRATE_SSS_DATA_UINT_COUNT;
 	const uint32 SliceCountAdvDebug = IsAdvancedVisualizationEnabled() ? 1 : 0;
-	const uint32 SliceCount = FMath::DivideAndRoundUp(Out.EffectiveMaxBytesPerPixel, 4u) + SliceCountSSS + SliceCountAdvDebug;
-	FRDGTextureDesc MaterialTextureDesc = FRDGTextureDesc::Create2DArray(SceneTextureExtent, PF_R32_UINT, FClearValueBinding::Transparent, TexCreate_TargetArraySlicesIndependently | TexCreate_DisableDCC | TexCreate_NoFastClear | TexCreate_RenderTargetable | TexCreate_ShaderResource | TexCreate_UAV | TexCreate_FastVRAM, SliceCount, 1, 1);
-	MaterialTextureDesc.FastVRAMPercentage = (1.0f / SliceCount) * 0xFF; // Only allocate the first slice into ESRAM
-
-	Out.MaterialTextureArray = GraphBuilder.CreateTexture(MaterialTextureDesc, TEXT("Substrate.Material"));
-	Out.MaterialTextureArraySRV = GraphBuilder.CreateSRV(FRDGTextureSRVDesc::Create(Out.MaterialTextureArray));
-	Out.MaterialTextureArrayUAV = GraphBuilder.CreateUAV(FRDGTextureUAVDesc(Out.MaterialTextureArray, 0));
-
-	// See AppendSubstrateMRTs
-	check(SUBSTRATE_BASE_PASS_MRT_OUTPUT_COUNT <= (SliceCount - SliceCountSSS - SliceCountAdvDebug)); // We want enough slice for MRTs but also do not want the SSSData to be a MRT.
-	Out.MaterialTextureArrayUAVWithoutRTs = GraphBuilder.CreateUAV(FRDGTextureUAVDesc(Out.MaterialTextureArray, 0, PF_Unknown, SUBSTRATE_BASE_PASS_MRT_OUTPUT_COUNT, SliceCount - SUBSTRATE_BASE_PASS_MRT_OUTPUT_COUNT));
+	const uint32 SliceCount = bNeedsMaterialBuffer ? FMath::DivideAndRoundUp(Out.EffectiveMaxBytesPerPixel, 4u) + SliceCountSSS + SliceCountAdvDebug : 1u;
+	if (bNeedsMaterialBuffer)
+	{
+		FRDGTextureDesc MaterialTextureDesc = FRDGTextureDesc::Create2DArray(SceneTextureExtent, PF_R32_UINT, FClearValueBinding::Transparent, TexCreate_TargetArraySlicesIndependently | TexCreate_DisableDCC | TexCreate_NoFastClear | TexCreate_RenderTargetable | TexCreate_ShaderResource | TexCreate_UAV | TexCreate_FastVRAM, SliceCount, 1, 1);
+		MaterialTextureDesc.FastVRAMPercentage = (1.0f / SliceCount) * 0xFF; // Only allocate the first slice into ESRAM
+		Out.MaterialTextureArray = GraphBuilder.CreateTexture(MaterialTextureDesc, TEXT("Substrate.Material"));
+		Out.MaterialTextureArraySRV = GraphBuilder.CreateSRV(FRDGTextureSRVDesc::Create(Out.MaterialTextureArray));
+		Out.MaterialTextureArrayUAV = GraphBuilder.CreateUAV(FRDGTextureUAVDesc(Out.MaterialTextureArray, 0));
+
+		// See AppendSubstrateMRTs
+		check(SUBSTRATE_BASE_PASS_MRT_OUTPUT_COUNT <= (SliceCount - SliceCountSSS - SliceCountAdvDebug)); // We want enough slice for MRTs but also do not want the SSSData to be a MRT.
+		Out.MaterialTextureArrayUAVWithoutRTs = GraphBuilder.CreateUAV(FRDGTextureUAVDesc(Out.MaterialTextureArray, 0, PF_Unknown, SUBSTRATE_BASE_PASS_MRT_OUTPUT_COUNT, SliceCount - SUBSTRATE_BASE_PASS_MRT_OUTPUT_COUNT));
+	}
 
 	// Rough diffuse model
 	Out.bRoughDiffuse = IsRoughDiffuseEnabled();
-
 	Out.PeelLayersAboveDepth = FMath::Max(CVarSubstrateDebugPeelLayersAboveDepth.GetValueOnRenderThread(), 0);
 	Out.bRoughnessTracking = CVarSubstrateDebugRoughnessTracking.GetValueOnRenderThread() > 0 ? 1 : 0;
 
-	// SUBSTRATE_TODO allocate a slice for StoringDebugSubstrate only if SUBSTRATE_ADVANCED_DEBUG_ENABLED is enabled 
-	Out.SliceStoringDebugSubstrateTreeData				= SliceCount - SliceCountAdvDebug;										// When we read, there is no slices excluded
-	Out.SliceStoringDebugSubstrateTreeDataWithoutMRT	= SliceCount - SliceCountAdvDebug - SUBSTRATE_BASE_PASS_MRT_OUTPUT_COUNT;	// The UAV skips the first slices set as render target
+	if (bNeedsMaterialBuffer)
+	{
+		// SUBSTRATE_TODO allocate a slice for StoringDebugSubstrate only if SUBSTRATE_ADVANCED_DEBUG_ENABLED is enabled 
+		Out.SliceStoringDebugSubstrateTreeData				= SliceCount - SliceCountAdvDebug;										// When we read, there is no slices excluded
+		Out.SliceStoringDebugSubstrateTreeDataWithoutMRT	= SliceCount - SliceCountAdvDebug - SUBSTRATE_BASE_PASS_MRT_OUTPUT_COUNT;	// The UAV skips the first slices set as render target
 
-	Out.FirstSliceStoringSubstrateSSSData				= SliceCount - SliceCountSSS - SliceCountAdvDebug;										// When we read, there is no slices excluded
-	Out.FirstSliceStoringSubstrateSSSDataWithoutMRT		= SliceCount - SliceCountSSS - SliceCountAdvDebug - SUBSTRATE_BASE_PASS_MRT_OUTPUT_COUNT;	// The UAV skips the first slices set as render target
+		Out.FirstSliceStoringSubstrateSSSData				= SliceCount - SliceCountSSS - SliceCountAdvDebug;										// When we read, there is no slices excluded
+		Out.FirstSliceStoringSubstrateSSSDataWithoutMRT 	= SliceCount - SliceCountSSS - SliceCountAdvDebug - SUBSTRATE_BASE_PASS_MRT_OUTPUT_COUNT;	// The UAV skips the first slices set as render target
+	}
+	else
+	{
+		Out.SliceStoringDebugSubstrateTreeData 				= -1;
+		Out.SliceStoringDebugSubstrateTreeDataWithoutMRT	= -1;
+		Out.FirstSliceStoringSubstrateSSSData				= -1;
+		Out.FirstSliceStoringSubstrateSSSDataWithoutMRT		= -1;
+	}
 
 	// Initialized view data
 	for (int32 ViewIndex = 0; ViewIndex < SceneRenderer.Views.Num(); ViewIndex++)
 	{
-		Substrate::InitialiseSubstrateViewData(GraphBuilder, SceneRenderer.Views[ViewIndex], SceneRenderer.GetActiveSceneTexturesConfig(), bNeedClosureOffsets, Out);
+		Substrate::InitialiseSubstrateViewData(GraphBuilder, SceneRenderer.Views[ViewIndex], SceneRenderer.GetActiveSceneTexturesConfig(), bNeedsClosureOffsets, bNeedsMaterialBuffer, Out);
 	}
 
 	if (IsSubstrateEnabled())
@@ -508,6 +534,7 @@ static FSubstrateCommonParameters GetSubstrateCommonParameter()
 	Out.bRoughnessTracking 	= 0u;
 	return Out;
 }
+
 static FSubstrateCommonParameters GetSubstrateCommonParameter(const FSubstrateSceneData& In)
 {
 	FSubstrateCommonParameters Out;
@@ -521,16 +548,31 @@ static FSubstrateCommonParameters GetSubstrateCommonParameter(const FSubstrateSc
 
 void BindSubstrateBasePassUniformParameters(FRDGBuilder& GraphBuilder, const FViewInfo& View, FSubstrateBasePassUniformParameters& OutSubstrateUniformParameters)
 {
+	bool bCreateDummyResources = false;
+
 	const FSubstrateSceneData* SubstrateSceneData = View.SubstrateViewData.SceneData;
 	if (IsSubstrateEnabled() && SubstrateSceneData)
 	{
 		OutSubstrateUniformParameters.Common = GetSubstrateCommonParameter(*SubstrateSceneData);
-		OutSubstrateUniformParameters.SliceStoringDebugSubstrateTreeDataWithoutMRT = SubstrateSceneData->SliceStoringDebugSubstrateTreeDataWithoutMRT;
-		OutSubstrateUniformParameters.FirstSliceStoringSubstrateSSSDataWithoutMRT = SubstrateSceneData->FirstSliceStoringSubstrateSSSDataWithoutMRT;
-		OutSubstrateUniformParameters.MaterialTextureArrayUAVWithoutRTs = SubstrateSceneData->MaterialTextureArrayUAVWithoutRTs;
-		OutSubstrateUniformParameters.OpaqueRoughRefractionTextureUAV = SubstrateSceneData->OpaqueRoughRefractionTextureUAV;
+		if (SubstrateSceneData->MaterialTextureArrayUAVWithoutRTs)
+		{
+			OutSubstrateUniformParameters.SliceStoringDebugSubstrateTreeDataWithoutMRT = SubstrateSceneData->SliceStoringDebugSubstrateTreeDataWithoutMRT;
+			OutSubstrateUniformParameters.FirstSliceStoringSubstrateSSSDataWithoutMRT = SubstrateSceneData->FirstSliceStoringSubstrateSSSDataWithoutMRT;
+			OutSubstrateUniformParameters.MaterialTextureArrayUAVWithoutRTs = SubstrateSceneData->MaterialTextureArrayUAVWithoutRTs;
+			OutSubstrateUniformParameters.OpaqueRoughRefractionTextureUAV = SubstrateSceneData->OpaqueRoughRefractionTextureUAV;
+		}
+		else
+		{
+			bCreateDummyResources = true;
+		}
 	}
 	else
+	{
+		OutSubstrateUniformParameters.Common = GetSubstrateCommonParameter();
+		bCreateDummyResources = true;
+	}
+
+	if (bCreateDummyResources)
 	{
 		FRDGTextureRef DummyWritableRefracTexture = GraphBuilder.CreateTexture(FRDGTextureDesc::Create2D(FIntPoint(1, 1), PF_R8, FClearValueBinding::None, TexCreate_ShaderResource | TexCreate_UAV), TEXT("Substrate.DummyWritableTexture"));
 		FRDGTextureUAVRef DummyWritableRefracTextureUAV = GraphBuilder.CreateUAV(FRDGTextureUAVDesc(DummyWritableRefracTexture));
@@ -539,7 +581,6 @@ void BindSubstrateBasePassUniformParameters(FRDGBuilder& GraphBuilder, const FVi
 		FRDGTextureUAVRef DummyWritableTextureArrayUAV = GraphBuilder.CreateUAV(FRDGTextureUAVDesc(DummyWritableTextureArray));
 
 		const FRDGSystemTextures& SystemTextures = FRDGSystemTextures::Get(GraphBuilder);
-		OutSubstrateUniformParameters.Common = GetSubstrateCommonParameter();
 		OutSubstrateUniformParameters.SliceStoringDebugSubstrateTreeDataWithoutMRT = -1;
 		OutSubstrateUniformParameters.FirstSliceStoringSubstrateSSSDataWithoutMRT = -1;
 		OutSubstrateUniformParameters.MaterialTextureArrayUAVWithoutRTs = DummyWritableTextureArrayUAV;
@@ -547,7 +588,7 @@ void BindSubstrateBasePassUniformParameters(FRDGBuilder& GraphBuilder, const FVi
 	}
 }
 
-static void BindSubstrateGlobalUniformParameters(FRDGBuilder& GraphBuilder, FSubstrateViewData* SubstrateViewData, FSubstrateGlobalUniformParameters& OutSubstrateUniformParameters)
+static void BindSubstrateGlobalUniformParameters(FRDGBuilder& GraphBuilder, const FSubstrateViewData* SubstrateViewData, bool bNeedsMaterialBuffer, FSubstrateGlobalUniformParameters& OutSubstrateUniformParameters)
 {
 	FSubstrateSceneData* SubstrateSceneData = SubstrateViewData->SceneData;
 	if (IsSubstrateEnabled() && SubstrateSceneData)
@@ -573,6 +614,18 @@ static void BindSubstrateGlobalUniformParameters(FRDGBuilder& GraphBuilder, FSub
 			OutSubstrateUniformParameters.ClosureTileCountBuffer = DefaultBuffer;
 			OutSubstrateUniformParameters.ClosureTileBuffer = DefaultBuffer;
 		}
+
+		if (!bNeedsMaterialBuffer)
+		{
+			check(OutSubstrateUniformParameters.MaterialTextureArray == nullptr);
+			check(SubstrateSceneData->TopLayerTexture == nullptr);
+			check(SubstrateSceneData->OpaqueRoughRefractionTexture == nullptr);
+			const FRDGSystemTextures& SystemTextures = FRDGSystemTextures::Get(GraphBuilder);
+			FRDGTextureRef DefaultTextureArray = GSystemTextures.GetDefaultTexture(GraphBuilder, ETextureDimension::Texture2DArray, EPixelFormat::PF_R32_UINT, FClearValueBinding::Transparent);
+			OutSubstrateUniformParameters.MaterialTextureArray = DefaultTextureArray;
+			OutSubstrateUniformParameters.TopLayerTexture = SystemTextures.DefaultNormal8Bit;
+			OutSubstrateUniformParameters.OpaqueRoughRefractionTexture = SystemTextures.Black;
+		}
 	}
 	else
 	{
@@ -597,18 +650,31 @@ static void BindSubstrateGlobalUniformParameters(FRDGBuilder& GraphBuilder, FSub
 void BindSubstrateForwardPasslUniformParameters(FRDGBuilder& GraphBuilder, const FViewInfo& View, FSubstrateForwardPassUniformParameters& OutSubstrateUniformParameters)
 {
 	FSubstrateSceneData* SubstrateSceneData = View.SubstrateViewData.SceneData;
+	bool bCreateDummyResources = false;
 	if (IsSubstrateEnabled() && SubstrateSceneData)
 	{
 		OutSubstrateUniformParameters.Common = GetSubstrateCommonParameter(*SubstrateSceneData);
-		OutSubstrateUniformParameters.FirstSliceStoringSubstrateSSSData = SubstrateSceneData->FirstSliceStoringSubstrateSSSData;
-		OutSubstrateUniformParameters.MaterialTextureArray = SubstrateSceneData->MaterialTextureArray;
-		OutSubstrateUniformParameters.TopLayerTexture = SubstrateSceneData->TopLayerTexture;
+		if (SubstrateSceneData->MaterialTextureArray)
+		{
+			OutSubstrateUniformParameters.FirstSliceStoringSubstrateSSSData = SubstrateSceneData->FirstSliceStoringSubstrateSSSData;
+			OutSubstrateUniformParameters.MaterialTextureArray = SubstrateSceneData->MaterialTextureArray;
+			OutSubstrateUniformParameters.TopLayerTexture = SubstrateSceneData->TopLayerTexture;
+		}
+		else
+		{
+			bCreateDummyResources = true;
+		}
 	}
 	else
+	{
+		OutSubstrateUniformParameters.Common = GetSubstrateCommonParameter();
+		bCreateDummyResources = true;
+	}
+
+	if (bCreateDummyResources)
 	{
 		const FRDGSystemTextures& SystemTextures = FRDGSystemTextures::Get(GraphBuilder);
 		FRDGTextureRef DefaultTextureArray = GSystemTextures.GetDefaultTexture(GraphBuilder, ETextureDimension::Texture2DArray, EPixelFormat::PF_R32_UINT, FClearValueBinding::Transparent);
-		OutSubstrateUniformParameters.Common = GetSubstrateCommonParameter();
 		OutSubstrateUniformParameters.FirstSliceStoringSubstrateSSSData = -1;
 		OutSubstrateUniformParameters.MaterialTextureArray = DefaultTextureArray;
 		OutSubstrateUniformParameters.TopLayerTexture = SystemTextures.DefaultNormal8Bit;
@@ -1189,7 +1255,8 @@ void AddSubstrateStencilPass(
 
 void AppendSubstrateMRTs(const FSceneRenderer& SceneRenderer, uint32& RenderTargetCount, TArrayView<FTextureRenderTargetBinding> RenderTargets)
 {
-	if (Substrate::IsSubstrateEnabled() && SceneRenderer.Scene)
+	const bool bUsesMaterialBuffer = UsesSubstrateMaterialBuffer(SceneRenderer.ShaderPlatform);
+	if (Substrate::IsSubstrateEnabled() && SceneRenderer.Scene && bUsesMaterialBuffer)
 	{
 		// If this function changes, update Substrate::SetBasePassRenderTargetOutputFormat()
 		 
@@ -1260,7 +1327,11 @@ void AddSubstrateMaterialClassificationPass(FRDGBuilder& GraphBuilder, const FMi
 		RDG_EVENT_SCOPE_CONDITIONAL(GraphBuilder, Views.Num() > 1, "View%d", i);
 
 		const FViewInfo& View = Views[i];
-		EShaderPlatform Platform = View.GetShaderPlatform();
+		const EShaderPlatform Platform = View.GetShaderPlatform();
+		if (!UsesSubstrateMaterialBuffer(Platform))
+		{
+			continue;
+		}
 
 		const bool bWaveOps = GRHISupportsWaveOperations&& GRHIMaximumWaveSize >= 64 && SubstrateSupportsWaveOps(Platform) != ERHIFeatureSupport::Unsupported;
 		
@@ -1281,7 +1352,7 @@ void AddSubstrateMaterialClassificationPass(FRDGBuilder& GraphBuilder, const FMi
 			}
 
 			// If Dbuffer pass (i.e. apply DBuffer data after the base-pass) is enabled, run special classification for outputing tile with/without tiles
-			const bool bDBufferTiles = IsDBufferPassEnabled(Platform) && CVarSubstrateDBufferPassDedicatedTiles.GetValueOnRenderThread() > 0 && DBufferTextures.IsValid() && IsConsolePlatform(View.GetShaderPlatform());
+			const bool bDBufferTiles = IsDBufferPassEnabled(Platform) && CVarSubstrateDBufferPassDedicatedTiles.GetValueOnRenderThread() > 0 && DBufferTextures.IsValid() && IsConsolePlatform(Platform);
 
 			FSubstrateMaterialTileClassificationPassCS::FPermutationDomain PermutationVector;
 			PermutationVector.Set< FSubstrateMaterialTileClassificationPassCS::FCmask >(bSupportCMask);
diff --git a/Engine/Source/Runtime/Renderer/Private/Substrate/Substrate.h b/Engine/Source/Runtime/Renderer/Private/Substrate/Substrate.h
--- a/Engine/Source/Runtime/Renderer/Private/Substrate/Substrate.h
+++ b/Engine/Source/Runtime/Renderer/Private/Substrate/Substrate.h
@@ -188,6 +188,7 @@ constexpr uint32 StencilBit_ComplexSpecial	= 0x80; // In sync with SceneRenderTa
 FIntPoint GetSubstrateTextureResolution(const FViewInfo& View, const FIntPoint& InResolution);
 uint32 GetSubstrateMaxClosureCount(const FViewInfo& View);
 bool GetSubstrateUsesComplexSpecialPath(const FViewInfo& View);
+bool UsesSubstrateMaterialBuffer(EShaderPlatform In);
 
 void InitialiseSubstrateFrameSceneData(FRDGBuilder& GraphBuilder, FSceneRenderer& SceneRenderer);
 
diff --git a/Engine/Source/Runtime/Renderer/Private/Substrate/SubstrateVisualize.cpp b/Engine/Source/Runtime/Renderer/Private/Substrate/SubstrateVisualize.cpp
--- a/Engine/Source/Runtime/Renderer/Private/Substrate/SubstrateVisualize.cpp
+++ b/Engine/Source/Runtime/Renderer/Private/Substrate/SubstrateVisualize.cpp
@@ -26,6 +26,7 @@ static bool IsSubstrateDebugVisualizationSupported(EShaderPlatform InPlatform, b
 {
 	return 
 		Substrate::IsSubstrateEnabled() && 
+		Substrate::UsesSubstrateMaterialBuffer(InPlatform) &&
 		GetMaxSupportedFeatureLevel(InPlatform) >= ERHIFeatureLevel::SM5 &&
 		!IsVulkanPlatform(InPlatform) && // Does not compile and fails to package games. Fixing this would require a deeper investigation.
 		(bIsEditorOnly ? (IsPCPlatform(InPlatform) || EnumHasAllFlags(Flags, EShaderPermutationFlags::HasEditorOnlyData)) : true);
