Fix Vulkan asset when enabling Substrate.

Vulkan does not support DepthNop operation when writing to Stencil.

diff --git a/Engine/Source/Runtime/Renderer/Private/Substrate/Substrate.cpp b/Engine/Source/Runtime/Renderer/Private/Substrate/Substrate.cpp
--- a/Engine/Source/Runtime/Renderer/Private/Substrate/Substrate.cpp
+++ b/Engine/Source/Runtime/Renderer/Private/Substrate/Substrate.cpp
@@ -1139,7 +1139,7 @@ void AddSubstrateInternalClassificationTilePass(
 			*DepthTexture,
 			ERenderTargetLoadAction::ELoad,
 			ERenderTargetLoadAction::ELoad,
-			FExclusiveDepthStencil::DepthNop_StencilWrite);
+			FExclusiveDepthStencil::DepthRead_StencilWrite);
 		ParametersPS->DebugTileColor = FVector4f(ForceInitToZero);
 	}
 	
