Fix shader compilation error with Substrate and Vulkan SM5.

Vulkan SM5 does not support geometry shader. Disable volumetric fog voxelization on these platforms.

diff --git a/Engine/Source/Runtime/Renderer/Private/VolumetricFog.cpp b/Engine/Source/Runtime/Renderer/Private/VolumetricFog.cpp
--- a/Engine/Source/Runtime/Renderer/Private/VolumetricFog.cpp
+++ b/Engine/Source/Runtime/Renderer/Private/VolumetricFog.cpp
@@ -1120,7 +1120,8 @@ IMPLEMENT_GLOBAL_SHADER(FVolumetricFogFinalIntegrationCS, "/Engine/Private/Volum
 
 bool DoesPlatformSupportVolumetricFogVoxelization(const FStaticShaderPlatform Platform)
 {
-	return !IsMobilePlatform(Platform);
+	const bool bCanRHICompileHlsl2021GeometryShaders = Substrate::IsSubstrateEnabled() && GetMaxSupportedFeatureLevel(Platform) > ERHIFeatureLevel::SM5;
+	return !IsMobilePlatform(Platform) && bCanRHICompileHlsl2021GeometryShaders;
 }
 bool ShouldRenderVolumetricFog(const FScene* Scene, const FSceneViewFamily& ViewFamily)
 {
