VulkanRHI: Fix leaked staging buffers on shutdown.

#rnx

diff --git a/Engine/Source/Runtime/Renderer/Private/RayTracing/RayTracingScene.cpp b/Engine/Source/Runtime/Renderer/Private/RayTracing/RayTracingScene.cpp
--- a/Engine/Source/Runtime/Renderer/Private/RayTracing/RayTracingScene.cpp
+++ b/Engine/Source/Runtime/Renderer/Private/RayTracing/RayTracingScene.cpp
@@ -38,6 +38,7 @@ FRayTracingScene::FRayTracingScene()
 
 FRayTracingScene::~FRayTracingScene()
 {
+	ReleaseReadbackBuffers();
 }
 
 FRayTracingSceneWithGeometryInstances FRayTracingScene::BuildInitializationData() const
@@ -576,12 +577,7 @@ void FRayTracingScene::EndFrame()
 		RayTracingScenePooledBuffer = nullptr;
 
 #if STATS
-		for (auto& ReadbackBuffer : StatsReadbackBuffers)
-		{
-			delete ReadbackBuffer;
-		}
-
-		StatsReadbackBuffers.Empty();
+		ReleaseReadbackBuffers();
 
 		StatsReadbackBuffersWriteIndex = 0;
 		StatsReadbackBuffersNumPending = 0;
@@ -593,4 +589,16 @@ void FRayTracingScene::EndFrame()
 	bUsedThisFrame = false;
 }
 
+void FRayTracingScene::ReleaseReadbackBuffers()
+{
+#if STATS
+	for (auto& ReadbackBuffer : StatsReadbackBuffers)
+	{
+		delete ReadbackBuffer;
+	}
+
+	StatsReadbackBuffers.Empty();
+#endif
+}
+
 #endif // RHI_RAYTRACING
diff --git a/Engine/Source/Runtime/Renderer/Private/RayTracing/RayTracingScene.h b/Engine/Source/Runtime/Renderer/Private/RayTracing/RayTracingScene.h
--- a/Engine/Source/Runtime/Renderer/Private/RayTracing/RayTracingScene.h
+++ b/Engine/Source/Runtime/Renderer/Private/RayTracing/RayTracingScene.h
@@ -139,6 +139,7 @@ public:
 	// Used for transforming to translated world space in which TLAS was built.
 	FDFVector3 PreViewTranslation {};
 private:
+	void ReleaseReadbackBuffers();
 
 	FRDGBufferRef InstanceBuffer;
 	FRDGBufferRef BuildScratchBuffer;
