Fix hair view rect with stereo rendering causing RHI assert due to be out of bound.

diff --git a/Engine/Source/Runtime/Renderer/Private/HairStrands/HairStrandsDebug.cpp b/Engine/Source/Runtime/Renderer/Private/HairStrands/HairStrandsDebug.cpp
--- a/Engine/Source/Runtime/Renderer/Private/HairStrands/HairStrandsDebug.cpp
+++ b/Engine/Source/Runtime/Renderer/Private/HairStrands/HairStrandsDebug.cpp
@@ -1254,7 +1254,7 @@ static void InternalRenderHairStrandsDebugInfo(
 			AddDebugDeepShadowTexturePass(GraphBuilder, &View, MacroGroupData.ScreenRect, nullptr, nullptr, Params.SceneColorTexture);
 		}
 
-		const FIntRect TotalRect = ComputeVisibleHairStrandsMacroGroupsRect(View.ViewRect, HairData.MacroGroupDatas);
+		const FIntRect TotalRect = ComputeVisibleHairStrandsMacroGroupsRect(View, View.ViewRect, HairData.MacroGroupDatas);
 		AddDebugDeepShadowTexturePass(GraphBuilder, &View, TotalRect, nullptr, nullptr, Params.SceneColorTexture);
 	}
 	
diff --git a/Engine/Source/Runtime/Renderer/Private/HairStrands/HairStrandsUtils.cpp b/Engine/Source/Runtime/Renderer/Private/HairStrands/HairStrandsUtils.cpp
--- a/Engine/Source/Runtime/Renderer/Private/HairStrands/HairStrandsUtils.cpp
+++ b/Engine/Source/Runtime/Renderer/Private/HairStrands/HairStrandsUtils.cpp
@@ -277,7 +277,7 @@ FIntRect ComputeProjectedScreenRect(const FBox& B, const FViewInfo& View)
 }
 
 
-FIntRect ComputeVisibleHairStrandsMacroGroupsRect(const FIntRect& ViewRect, const FHairStrandsMacroGroupDatas& Datas)
+FIntRect ComputeVisibleHairStrandsMacroGroupsRect(const FViewInfo& View, const FIntRect& ViewRect, const FHairStrandsMacroGroupDatas& Datas)
 {
 	FIntRect TotalRect(INT_MAX, INT_MAX, -INT_MAX, -INT_MAX);
 	if (IsHairStrandsViewRectOptimEnable())
@@ -287,6 +287,15 @@ FIntRect ComputeVisibleHairStrandsMacroGroupsRect(const FIntRect& ViewRect, cons
 			TotalRect.Union(Data.ScreenRect);
 		}
 
+		// For stereo rendering we don't compute FHairStrandsMacroGroupDatas for each view, but reuse MacroGroup data from View0 to View1. 
+		// See RenderHairPrePass() in Renderer\Private\HairStrands\HairStrandsRendering.cpp for more details.
+		// In order to correct the view rect for the secondary view, we add the view.min value.
+		if (IStereoRendering::IsASecondaryView(View))
+		{
+			TotalRect.Min += View.ViewRect.Min;
+			TotalRect.Max += View.ViewRect.Min;
+		}
+
 		// In case bounds are not initialized correct for some reason, return view rect
 		if (TotalRect.Min.X >= TotalRect.Max.X || TotalRect.Min.Y >= TotalRect.Max.Y) TotalRect = ViewRect;
 	}
diff --git a/Engine/Source/Runtime/Renderer/Private/HairStrands/HairStrandsUtils.h b/Engine/Source/Runtime/Renderer/Private/HairStrands/HairStrandsUtils.h
--- a/Engine/Source/Runtime/Renderer/Private/HairStrands/HairStrandsUtils.h
+++ b/Engine/Source/Runtime/Renderer/Private/HairStrands/HairStrandsUtils.h
@@ -48,7 +48,7 @@ float GetHairDualScatteringRoughnessOverride();
 
 float SampleCountToSubPixelSize(uint32 SamplePerPixelCount);
 
-FIntRect ComputeVisibleHairStrandsMacroGroupsRect(const FIntRect& ViewRect, const FHairStrandsMacroGroupDatas& Datas);
+FIntRect ComputeVisibleHairStrandsMacroGroupsRect(const FViewInfo& View, const FIntRect& ViewRect, const FHairStrandsMacroGroupDatas& Datas);
 
 bool IsHairStrandsViewRectOptimEnable();
 
diff --git a/Engine/Source/Runtime/Renderer/Private/HairStrands/HairStrandsVisibility.cpp b/Engine/Source/Runtime/Renderer/Private/HairStrands/HairStrandsVisibility.cpp
--- a/Engine/Source/Runtime/Renderer/Private/HairStrands/HairStrandsVisibility.cpp
+++ b/Engine/Source/Runtime/Renderer/Private/HairStrands/HairStrandsVisibility.cpp
@@ -2107,7 +2107,7 @@ static void AddHairVisibilityControlPointIdCompactionPass(
 	AddClearGraphicPass(GraphBuilder, RDG_EVENT_NAME("HairStrands::CoverageTexture"), &View, 0, TileData, OutCoverageTexture);
 
 	// Adapt the buffer allocation based on the bounding box of the hair macro groups. This allows to reduce the overall allocation size
-	const FIntRect HairRect = ComputeVisibleHairStrandsMacroGroupsRect(View.ViewRect, MacroGroupDatas);
+	const FIntRect HairRect = ComputeVisibleHairStrandsMacroGroupsRect(View, View.ViewRect, MacroGroupDatas);
 	const FIntPoint EffectiveResolution = bUsePPLL ? FIntPoint(View.ViewRect.Width(), View.ViewRect.Height()) : FIntPoint(HairRect.Width(), HairRect.Height());
 
 	// Select render node count according to current mode
@@ -4600,7 +4600,7 @@ void RenderHairStrandsVisibilityBuffer(
 	check(View.Family);
 	check(MacroGroupDatas.Num() > 0);
 
-	const FIntRect HairRect = ComputeVisibleHairStrandsMacroGroupsRect(View.ViewRect, MacroGroupDatas);
+	const FIntRect HairRect = ComputeVisibleHairStrandsMacroGroupsRect(View, View.ViewRect, MacroGroupDatas);
 	const int32 HairPixelCount = HairRect.Width() * HairRect.Height();
 	if (HairPixelCount <= 0)
 	{
diff --git a/Engine/Source/Runtime/Renderer/Private/LightFunctionRendering.cpp b/Engine/Source/Runtime/Renderer/Private/LightFunctionRendering.cpp
--- a/Engine/Source/Runtime/Renderer/Private/LightFunctionRendering.cpp
+++ b/Engine/Source/Runtime/Renderer/Private/LightFunctionRendering.cpp
@@ -436,7 +436,7 @@ bool FDeferredShadingSceneRenderer::RenderLightFunctionForMaterial(
 					LightSceneProxy->SetScissorRect(RHICmdList, View, View.ViewRect);
 					if (bUseHairStrands)
 					{
-						FIntRect TotalRect = ComputeVisibleHairStrandsMacroGroupsRect(View.ViewRect, View.HairStrandsViewData.MacroGroupDatas);
+						FIntRect TotalRect = ComputeVisibleHairStrandsMacroGroupsRect(View, View.ViewRect, View.HairStrandsViewData.MacroGroupDatas);
 						LightSceneProxy->SetScissorRect(RHICmdList, View, TotalRect);
 					}
 
