Avoid marking revealed non-nanite primitives as dynamic cached (unless they also invalidate for other reasons)
#rnx

diff --git a/Engine/Source/Runtime/Renderer/Private/VirtualShadowMaps/VirtualShadowMapCacheManager.cpp b/Engine/Source/Runtime/Renderer/Private/VirtualShadowMaps/VirtualShadowMapCacheManager.cpp
--- a/Engine/Source/Runtime/Renderer/Private/VirtualShadowMaps/VirtualShadowMapCacheManager.cpp
+++ b/Engine/Source/Runtime/Renderer/Private/VirtualShadowMaps/VirtualShadowMapCacheManager.cpp
@@ -357,7 +357,8 @@ void FVirtualShadowMapArrayCacheManager::FShadowInvalidatingInstancesImplementat
 	PrimitiveInstancesToInvalidate.Add(FVirtualShadowMapInstanceRange{
 		PersistentPrimitiveIndex,
 		int32(InstanceSceneDataOffset),
-		int32(NumInstanceSceneDataEntries)});
+		int32(NumInstanceSceneDataEntries),
+		true});
 }
 
 void FVirtualShadowMapPerLightCacheEntry::Invalidate()
@@ -395,7 +396,7 @@ void FVirtualShadowMapArrayCacheManager::FInvalidatingPrimitiveCollector::AddPri
 			for (const FVirtualShadowMapInstanceRange& Range : Manager.ShadowInvalidatingInstancesImplementation.PrimitiveInstancesToInvalidate)
 			{
 				Instances.Add(Range.InstanceSceneDataOffset, Range.NumInstanceSceneDataEntries, Payload);
-				if (Range.PersistentPrimitiveIndex.IsValid())
+				if (Range.bMarkAsDynamic && Range.PersistentPrimitiveIndex.IsValid())
 				{
 					InvalidatedPrimitives[Range.PersistentPrimitiveIndex.Index] = true;
 				}
@@ -406,7 +407,7 @@ void FVirtualShadowMapArrayCacheManager::FInvalidatingPrimitiveCollector::AddPri
 			{
 				Instances.Add(Range.InstanceSceneDataOffset, Range.NumInstanceSceneDataEntries, Payload);
 				check(Range.PersistentPrimitiveIndex.IsValid());		// Should always be valid currently in this path
-				if (Range.PersistentPrimitiveIndex.IsValid())
+				if (Range.bMarkAsDynamic && Range.PersistentPrimitiveIndex.IsValid())
 				{
 					InvalidatedPrimitives[Range.PersistentPrimitiveIndex.Index] = true;
 				}
@@ -980,6 +981,7 @@ void FVirtualShadowMapPerLightCacheEntry::OnPrimitiveRendered(const FPrimitiveSc
 	CachedPrimitives[PrimitiveSceneInfo->GetPersistentIndex().Index] = true;
 
 	bool bInvalidate = false;
+	bool bMarkAsDynamic = true;
 
 	// Deformable mesh primitives need to trigger invalidation (even if they did not move) or we get artifacts, for example skinned meshes that are animating but not currently moving.
 	// Skip if the invalidation mode is NOT auto (because Always will do it elsewhere & the others should prevent this).
@@ -989,12 +991,12 @@ void FVirtualShadowMapPerLightCacheEntry::OnPrimitiveRendered(const FPrimitiveSc
 	{
 		bInvalidate = true;
 	}
-
 	// With new invalidations on, we need to invalidate any time a (non-nanite) primitive is "revealed", i.e. stopped being culled.
 	// Note that this invalidation will be a frame late - similar to WPO starting - as it will get picked up by the next scene update.
-	if (bPrimitiveRevealed && GVSMNewInvalidations != 0 && GVSMCacheDebugSkipRevealedPrimitivesInvalidate == 0)
-	{
+	else if (bPrimitiveRevealed && GVSMNewInvalidations != 0 && GVSMCacheDebugSkipRevealedPrimitivesInvalidate == 0)
+	{	
 		bInvalidate = true;
+		bMarkAsDynamic = false;		// Don't mark primitives as dynamic just because they were revealed
 		//UE_LOG(LogRenderer, Display, TEXT("VirtualShadowMapCacheManager: Primitive revealed %d!"), PrimitiveSceneInfo->GetPersistentIndex().Index);
 	}
 
@@ -1003,7 +1005,8 @@ void FVirtualShadowMapPerLightCacheEntry::OnPrimitiveRendered(const FPrimitiveSc
 		PrimitiveInstancesToInvalidate.Add(FVirtualShadowMapInstanceRange{
 			PrimitiveSceneInfo->GetPersistentIndex(),
 			PrimitiveSceneInfo->GetInstanceSceneDataOffset(),
-			PrimitiveSceneInfo->GetNumInstanceSceneDataEntries()
+			PrimitiveSceneInfo->GetNumInstanceSceneDataEntries(),
+			bMarkAsDynamic
 		});
 	}
 }
diff --git a/Engine/Source/Runtime/Renderer/Private/VirtualShadowMaps/VirtualShadowMapCacheManager.h b/Engine/Source/Runtime/Renderer/Private/VirtualShadowMaps/VirtualShadowMapCacheManager.h
--- a/Engine/Source/Runtime/Renderer/Private/VirtualShadowMaps/VirtualShadowMapCacheManager.h
+++ b/Engine/Source/Runtime/Renderer/Private/VirtualShadowMaps/VirtualShadowMapCacheManager.h
@@ -24,6 +24,7 @@ struct FVirtualShadowMapInstanceRange
 	FPersistentPrimitiveIndex PersistentPrimitiveIndex;
 	int32 InstanceSceneDataOffset;
 	int32 NumInstanceSceneDataEntries;
+	bool bMarkAsDynamic;					// If true, swaps the primitive/instance to dynamic caching
 };
 
 #define VSM_LOG_INVALIDATIONS 0
