CustomRenderPass: Avoid rendering sky atmosphere editor notifications during custom render passes (which have their atmosphere show flag disabled).

#rnx

diff --git a/Engine/Source/Runtime/Renderer/Private/BasePassRendering.cpp b/Engine/Source/Runtime/Renderer/Private/BasePassRendering.cpp
--- a/Engine/Source/Runtime/Renderer/Private/BasePassRendering.cpp
+++ b/Engine/Source/Runtime/Renderer/Private/BasePassRendering.cpp
@@ -1043,7 +1043,6 @@ void FDeferredShadingSceneRenderer::RenderBasePass(
 	const bool bIsWireframeRenderpass = ViewFamily.EngineShowFlags.Wireframe && FSceneRenderer::ShouldCompositeEditorPrimitives(InViews[0]);
 	const bool bDebugViewMode = ViewFamily.UseDebugViewPS();
 	const bool bRenderLightmapDensity = ViewFamily.EngineShowFlags.LightMapDensity && AllowDebugViewmodes();
-	const bool bRenderSkyAtmosphereEditorNotifications = ShouldRenderSkyAtmosphereEditorNotifications();
 	const bool bDoParallelBasePass = bEnableParallelBasePasses && !bDebugViewMode && !bRenderLightmapDensity; // DebugView and LightmapDensity are non-parallel substitutions inside BasePass
 	const bool bNeedsBeginRender = AllowDebugViewmodes() &&
 		(ViewFamily.EngineShowFlags.RequiredTextureResolution ||
@@ -1126,11 +1125,11 @@ void FDeferredShadingSceneRenderer::RenderBasePass(
 			}
 		});
 
-		if (bRenderSkyAtmosphereEditorNotifications)
+		if (ShouldRenderSkyAtmosphereEditorNotifications(InViews))
 		{
 			// We only render this warning text when bRequiresRHIClear==true to make sure the scene color buffer is allocated at this stage.
 			// When false, the option specifies that all pixels must be written to by a sky dome anyway.
-			RenderSkyAtmosphereEditorNotifications(GraphBuilder, SceneTextures.Color.Target);
+			RenderSkyAtmosphereEditorNotifications(GraphBuilder, InViews, SceneTextures.Color.Target);
 		}
 	}
 
diff --git a/Engine/Source/Runtime/Renderer/Private/SceneRendering.h b/Engine/Source/Runtime/Renderer/Private/SceneRendering.h
--- a/Engine/Source/Runtime/Renderer/Private/SceneRendering.h
+++ b/Engine/Source/Runtime/Renderer/Private/SceneRendering.h
@@ -2427,10 +2427,10 @@ protected:
 		FInstanceCullingManager& InstanceCullingManager);
 
 	/** Render notification to artist when a sky material is used but it might comtains the camera (and then the sky/background would look black).*/
-	void RenderSkyAtmosphereEditorNotifications(FRDGBuilder& GraphBuilder, FRDGTextureRef SceneColorTexture) const;
+	void RenderSkyAtmosphereEditorNotifications(FRDGBuilder& GraphBuilder, TArrayView<FViewInfo> InViews, FRDGTextureRef SceneColorTexture) const;
 
 	/** We should render on screen notification only if any of the scene contains a mesh using a sky material.*/
-	bool ShouldRenderSkyAtmosphereEditorNotifications() const;
+	static bool ShouldRenderSkyAtmosphereEditorNotifications(TArrayView<FViewInfo> Views);
 
 	/**
 	 * Rounds up lights and sorts them according to what type of renderer supports them. The result is stored in OutSortedLights 
diff --git a/Engine/Source/Runtime/Renderer/Private/SkyAtmosphereRendering.cpp b/Engine/Source/Runtime/Renderer/Private/SkyAtmosphereRendering.cpp
--- a/Engine/Source/Runtime/Renderer/Private/SkyAtmosphereRendering.cpp
+++ b/Engine/Source/Runtime/Renderer/Private/SkyAtmosphereRendering.cpp
@@ -2081,15 +2081,16 @@ void FSceneRenderer::RenderSkyAtmosphere(FRDGBuilder& GraphBuilder, const FMinim
 #endif
 }
 
-bool FSceneRenderer::ShouldRenderSkyAtmosphereEditorNotifications() const
+bool FSceneRenderer::ShouldRenderSkyAtmosphereEditorNotifications(TArrayView<FViewInfo> InViews)
 {
 #if WITH_EDITOR
 	if (CVarSkyAtmosphereEditorNotifications.GetValueOnAnyThread() > 0)
 	{
 		bool bAnyViewHasSkyMaterial = false;
-		for (int32 ViewIndex = 0; ViewIndex < Views.Num(); ViewIndex++)
+		for (int32 ViewIndex = 0; ViewIndex < InViews.Num(); ViewIndex++)
 		{
-			bAnyViewHasSkyMaterial |= Views[ViewIndex].bSceneHasSkyMaterial;
+			const FViewInfo& View = InViews[ViewIndex];
+			bAnyViewHasSkyMaterial |= (View.bSceneHasSkyMaterial && View.Family->EngineShowFlags.Atmosphere);
 		}
 		return bAnyViewHasSkyMaterial;
 	}
@@ -2097,16 +2098,16 @@ bool FSceneRenderer::ShouldRenderSkyAtmosphereEditorNotifications() const
 	return false;
 }
 
-void FSceneRenderer::RenderSkyAtmosphereEditorNotifications(FRDGBuilder& GraphBuilder, FRDGTextureRef SceneColorTexture) const
+void FSceneRenderer::RenderSkyAtmosphereEditorNotifications(FRDGBuilder& GraphBuilder, TArrayView<FViewInfo> InViews, FRDGTextureRef SceneColorTexture) const
 {
 #if WITH_EDITOR
 	RDG_EVENT_SCOPE(GraphBuilder, "SkyAtmosphereEditor");
 	RDG_GPU_STAT_SCOPE(GraphBuilder, SkyAtmosphereEditor);
 
-	for (int32 ViewIndex = 0; ViewIndex < Views.Num(); ViewIndex++)
+	for (int32 ViewIndex = 0; ViewIndex < InViews.Num(); ViewIndex++)
 	{
-		const FViewInfo& View = Views[ViewIndex];
-		if (View.bSceneHasSkyMaterial && View.Family->EngineShowFlags.Atmosphere)
+		const FViewInfo& View = InViews[ViewIndex];
+		if (View.bSceneHasSkyMaterial)
 		{
 			RenderSkyAtmosphereEditorHudPS::FPermutationDomain PermutationVector;
 			TShaderMapRef<RenderSkyAtmosphereEditorHudPS> PixelShader(View.ShaderMap, PermutationVector);
