Remove the FastClear on SSAO when we use the ComputeShade. CS version writes all the pixels. 0.2% dynres improvement on some consoles.

diff --git a/Engine/Source/Runtime/Renderer/Private/CompositionLighting/CompositionLighting.cpp b/Engine/Source/Runtime/Renderer/Private/CompositionLighting/CompositionLighting.cpp
--- a/Engine/Source/Runtime/Renderer/Private/CompositionLighting/CompositionLighting.cpp
+++ b/Engine/Source/Runtime/Renderer/Private/CompositionLighting/CompositionLighting.cpp
@@ -101,12 +101,12 @@ bool ShouldRenderScreenSpaceAmbientOcclusion(const FViewInfo& View, bool bLumenW
 
 static ESSAOType GetDownscaleSSAOType(const FViewInfo& View)
 {
-	return FSSAOHelper::IsAmbientOcclusionCompute(View) ? ESSAOType::ECS : ESSAOType::EPS;
+	return FSSAOHelper::IsAmbientOcclusionCompute(View.GetFeatureLevel()) ? ESSAOType::ECS : ESSAOType::EPS;
 }
 
 static ESSAOType GetFullscreenSSAOType(const FViewInfo& View, uint32 Levels)
 {
-	if (FSSAOHelper::IsAmbientOcclusionCompute(View))
+	if (FSSAOHelper::IsAmbientOcclusionCompute(View.GetFeatureLevel()))
 	{
 		if (FSSAOHelper::IsAmbientOcclusionAsyncCompute(View, Levels))
 		{
diff --git a/Engine/Source/Runtime/Renderer/Private/CompositionLighting/PostProcessAmbientOcclusion.cpp b/Engine/Source/Runtime/Renderer/Private/CompositionLighting/PostProcessAmbientOcclusion.cpp
--- a/Engine/Source/Runtime/Renderer/Private/CompositionLighting/PostProcessAmbientOcclusion.cpp
+++ b/Engine/Source/Runtime/Renderer/Private/CompositionLighting/PostProcessAmbientOcclusion.cpp
@@ -179,9 +179,9 @@ int32 FSSAOHelper::GetAmbientOcclusionShaderLevel(const FSceneView& View)
 		(QualityPercent > 5.0f);
 }
 
-bool FSSAOHelper::IsAmbientOcclusionCompute(const FSceneView& View)
+bool FSSAOHelper::IsAmbientOcclusionCompute(const ERHIFeatureLevel::Type FeatureLevel)
 {
-	return View.GetFeatureLevel() >= ERHIFeatureLevel::SM5 && CVarAmbientOcclusionCompute.GetValueOnRenderThread() >= 1;
+	return FeatureLevel >= ERHIFeatureLevel::SM5 && CVarAmbientOcclusionCompute.GetValueOnRenderThread() >= 1;
 }
 
 int32 FSSAOHelper::GetNumAmbientOcclusionLevels()
@@ -211,7 +211,7 @@ bool FSSAOHelper::IsAmbientOcclusionAsyncCompute(const FViewInfo& View, uint32 A
 {
 	// if AsyncCompute is feasible
 	// only single level is allowed.  more levels end up reading from gbuffer normals atm which is not allowed.
-	if(IsAmbientOcclusionCompute(View) && (AOPassCount == 1))
+	if(IsAmbientOcclusionCompute(View.GetFeatureLevel()) && (AOPassCount == 1))
 	{
 		int32 ComputeCVar = CVarAmbientOcclusionCompute.GetValueOnRenderThread();
 
@@ -242,7 +242,7 @@ uint32 FSSAOHelper::ComputeAmbientOcclusionPassCount(const FViewInfo& View)
 	{
 		int32 CVarLevel = GetNumAmbientOcclusionLevels();
 
-		if (IsAmbientOcclusionCompute(View) || IsForwardShadingEnabled(View.GetShaderPlatform()))
+		if (IsAmbientOcclusionCompute(View.GetFeatureLevel()) || IsForwardShadingEnabled(View.GetShaderPlatform()))
 		{	
 			if (CVarLevel<0)
 			{
@@ -308,14 +308,19 @@ EGTAOType FSSAOHelper::GetGTAOPassType(const FViewInfo& View, uint32 Levels)
 	return EGTAOType::EOff;
 }
 
-FRDGTextureDesc GetScreenSpaceAOTextureDesc(FIntPoint Extent)
+FRDGTextureDesc GetScreenSpaceAOTextureDesc(ERHIFeatureLevel::Type FeatureLevel, FIntPoint Extent)
 {
-	return FRDGTextureDesc(FRDGTextureDesc::Create2D(Extent, PF_G8, FClearValueBinding::White, TexCreate_UAV | TexCreate_RenderTargetable | TexCreate_ShaderResource | GFastVRamConfig.ScreenSpaceAO));
+	ETextureCreateFlags TextureCreateFlags = TexCreate_UAV | TexCreate_RenderTargetable | TexCreate_ShaderResource | GFastVRamConfig.ScreenSpaceAO;
+	if (FSSAOHelper::IsAmbientOcclusionCompute(FeatureLevel))
+	{
+		TextureCreateFlags |= TexCreate_NoFastClear;
+	}
+	return FRDGTextureDesc(FRDGTextureDesc::Create2D(Extent, PF_G8, FClearValueBinding::White, TextureCreateFlags));
 }
 
-FRDGTextureRef CreateScreenSpaceAOTexture(FRDGBuilder& GraphBuilder, FIntPoint Extent)
+FRDGTextureRef CreateScreenSpaceAOTexture(FRDGBuilder& GraphBuilder, ERHIFeatureLevel::Type FeatureLevel, FIntPoint Extent)
 {	
-	return GraphBuilder.CreateTexture(GetScreenSpaceAOTextureDesc(Extent), TEXT("ScreenSpaceAO"));
+	return GraphBuilder.CreateTexture(GetScreenSpaceAOTextureDesc(FeatureLevel, Extent), TEXT("ScreenSpaceAO"));
 }
 
 FRDGTextureRef GetScreenSpaceAOFallback(const FRDGSystemTextures& SystemTextures)
diff --git a/Engine/Source/Runtime/Renderer/Private/CompositionLighting/PostProcessAmbientOcclusion.h b/Engine/Source/Runtime/Renderer/Private/CompositionLighting/PostProcessAmbientOcclusion.h
--- a/Engine/Source/Runtime/Renderer/Private/CompositionLighting/PostProcessAmbientOcclusion.h
+++ b/Engine/Source/Runtime/Renderer/Private/CompositionLighting/PostProcessAmbientOcclusion.h
@@ -52,9 +52,9 @@ enum EGTAOPass
 	EGTAOPass_Upsample					= 0x20,
 };
 
-FRDGTextureDesc GetScreenSpaceAOTextureDesc(FIntPoint Extent);
+FRDGTextureDesc GetScreenSpaceAOTextureDesc(ERHIFeatureLevel::Type FeatureLevel, FIntPoint Extent);
 
-FRDGTextureRef CreateScreenSpaceAOTexture(FRDGBuilder& GraphBuilder, FIntPoint Extent);
+FRDGTextureRef CreateScreenSpaceAOTexture(FRDGBuilder& GraphBuilder, ERHIFeatureLevel::Type FeatureLevel, FIntPoint Extent);
 
 FRDGTextureRef GetScreenSpaceAOFallback(const FRDGSystemTextures& SystemTextures);
 
@@ -90,7 +90,7 @@ public:
 	static int32 GetAmbientOcclusionShaderLevel(const FSceneView& View);
 
 	// @return whether AmbientOcclusion should run a compute shader.
-	static bool IsAmbientOcclusionCompute(const FSceneView& View);
+	static bool IsAmbientOcclusionCompute(const ERHIFeatureLevel::Type FeatureLevel);
 
 	static int32 GetNumAmbientOcclusionLevels();
 	static float GetAmbientOcclusionStepMipLevelFactor();
diff --git a/Engine/Source/Runtime/Renderer/Private/CompositionLighting/PostProcessDeferredDecals.cpp b/Engine/Source/Runtime/Renderer/Private/CompositionLighting/PostProcessDeferredDecals.cpp
--- a/Engine/Source/Runtime/Renderer/Private/CompositionLighting/PostProcessDeferredDecals.cpp
+++ b/Engine/Source/Runtime/Renderer/Private/CompositionLighting/PostProcessDeferredDecals.cpp
@@ -152,7 +152,7 @@ void GetDeferredDecalRenderTargetsInfo(
 	}
 	case EDecalRenderTargetMode::AmbientOcclusion:
 	{		
-		const FRDGTextureDesc AOTextureDesc = GetScreenSpaceAOTextureDesc(Config.Extent);
+		const FRDGTextureDesc AOTextureDesc = GetScreenSpaceAOTextureDesc(Config.FeatureLevel, Config.Extent);
 		AddRenderTargetInfo(AOTextureDesc.Format, AOTextureDesc.Flags, RenderTargetsInfo);
 		break;
 	}
diff --git a/Engine/Source/Runtime/Renderer/Private/SceneTextures.cpp b/Engine/Source/Runtime/Renderer/Private/SceneTextures.cpp
--- a/Engine/Source/Runtime/Renderer/Private/SceneTextures.cpp
+++ b/Engine/Source/Runtime/Renderer/Private/SceneTextures.cpp
@@ -503,7 +503,7 @@ void FSceneTextures::InitializeViewFamily(FRDGBuilder& GraphBuilder, FViewFamily
 	if (Config.ShadingPath == EShadingPath::Deferred)
 	{
 		// Screen Space Ambient Occlusion
-		SceneTextures.ScreenSpaceAO = CreateScreenSpaceAOTexture(GraphBuilder, Config.Extent);
+		SceneTextures.ScreenSpaceAO = CreateScreenSpaceAOTexture(GraphBuilder, ViewFamily.GetFeatureLevel(), Config.Extent);
 
 		// Small Depth
 		const FIntPoint SmallDepthExtent = GetDownscaledExtent(Config.Extent, Config.SmallDepthDownsampleFactor);
