Post Process Material:  Support for "UserSceneTextures" -- transient intermediate render targets that can be written and read by name in post process materials.  The post process section of materials has a UserSceneTexture output name, and a new UserSceneTexture node type has been added to read back those outputs as inputs.

Shader features:
* A resolution divisor can be specified on the output to allow things like half resolution targets.  If the same name is used with multiple resolutions, the most recently written resolution is used as an input.
* A UV clamping option was added to solve out of bounds interpolation issues when bilinear sampling lower resolution targets from a higher resolution pass.
* A human readable define based on the user specified name is added to HLSL (with prefix "PPIUser_"), allowing Custom HLSL nodes to refer to UserSceneTexture inputs by name.  Necessary because input slots are dynamically allocated.  Useful for things like blur kernels that use Custom HLSL to sample a range of inputs.
* All 5 Post Process Input slots are potentially available to UserSceneTextures -- built-in SceneTexture slots for a given location (like SceneColor, SeparateTranslucency, or SceneVelocity) are used for UserSceneTextures if not referenced in the shader.
* Supports translucent blending.  If the first write to a UserSceneTexture includes blending, it will be primed with the current input for the given location.
* Missing inputs are replaced with black.
* Simplified syntax for accessing SceneTextures / UserSceneTextures in Custom HLSL (see "SceneTextureFetch" macro).

Editing / debug features:
* Material editor preview window automatically includes other edited materials that feed inputs into the current material, allowing multi-pass effects to be interactively edited.
* Debug display of post process passes with UserSceneTexture inputs or output via r.PostProcessing.UserSceneTextureDebug.  Allows seeing order, location, and priority of passes, plus names of material, inputs, and output, and whether an output is written with blending.
* Debug display pops up automatically on error conditions (missing input, unused output), improving discoverability.  Suppressed by DisableAllScreenMessages (or setting CVar itself to zero).
* FRDGViewableResource can now be given dynamically generated heap allocated debug names (indicated by setting a flag immediately after construction), instead of only supporting static const strings.  Allows UserSceneTexture names to be displayed in debug displays, VisualizeTexture, GPU captures, and GPU dumps, instead of every resource getting the same name.

Mobile doesn't support post process domain materials at all, so the feature is not relevant there.  Tested on scene captures, split screen, and both legacy and new material translate code paths.  An existing bug was fixed where rendering post process materials with blending before bloom would crash, due to the output not being allocated with the render target flag set.

diff --git a/Engine/Shaders/Private/MaterialTemplate.ush b/Engine/Shaders/Private/MaterialTemplate.ush
--- a/Engine/Shaders/Private/MaterialTemplate.ush
+++ b/Engine/Shaders/Private/MaterialTemplate.ush
@@ -79,6 +79,9 @@
 #define PPI_WorldTangent 29
 #define PPI_Anisotropy 30
 
+// Defines to remap PPI_UserSceneTexture0-6 to PPI_PostProcessInput0-6, generated via GenerateUserSceneTextureRemapDefines
+%{user_scene_texture_remap}
+
 //////////////////////////////////////////////////////////////////////////
 
 #define NUM_MATERIAL_TEXCOORDS_VERTEX %{num_material_texcoords_vertex}
@@ -3152,6 +3155,24 @@ float4 SceneTextureLookup(float2 UV, int SceneTextureIndex, bool bFiltered)
 	return float4(0, 0, 0, 0);
 }
 
+/**
+ * Utility function and macro to streamline Custom HLSL fetching from SceneTextures or UserSceneTextures with a pixel offset from default UVs, including
+ * multiplication by texture size and clamping.  Example that fetches at offset 1,1 from the default sample:
+ *
+ *     SceneTextureFetch(PPI_PostProcessInput0, float2(1,1));
+ */
+float4 SceneTextureFetchFunc(FMaterialPixelParameters Parameters, int SceneTextureIndex, float2 PixelOffset)
+{
+#if SCENE_TEXTURES_DISABLED || !POST_PROCESS_MATERIAL
+	return float4(0,0,0,0);
+#else
+	float2 UV = ClampSceneTextureUV(GetDefaultSceneTextureUV(Parameters, SceneTextureIndex) + GetSceneTextureViewSize(SceneTextureIndex).zw * PixelOffset, SceneTextureIndex);
+	return SceneTextureLookup(UV, SceneTextureIndex, true);
+#endif
+}
+
+#define SceneTextureFetch(SceneTextureIndex, PixelOffset) SceneTextureFetchFunc(Parameters, SceneTextureIndex, PixelOffset)
+
 #endif // SHADING_PATH_DEFERRED
 #endif // NEEDS_SCENE_TEXTURES
 
diff --git a/Engine/Source/Editor/MaterialEditor/Private/MaterialEditor.cpp b/Engine/Source/Editor/MaterialEditor/Private/MaterialEditor.cpp
--- a/Engine/Source/Editor/MaterialEditor/Private/MaterialEditor.cpp
+++ b/Engine/Source/Editor/MaterialEditor/Private/MaterialEditor.cpp
@@ -883,6 +883,9 @@ void FMaterialEditor::InitMaterialEditor( const EToolkitMode::Type Mode, const T
 	{
 		SaveAsset_Execute();
 	}
+
+	// Notify other editors if this material editor has a post process named output, which may affect their preview
+	NotifyUserSceneTextureLoadOrUnload();
 }
 
 void FMaterialEditor::UpdateGenerator()
@@ -931,6 +934,11 @@ FMaterialEditor::FMaterialEditor()
 
 FMaterialEditor::~FMaterialEditor()
 {
+	bDestructing = true;
+
+	// Notify other editors if this material editor has a post process named output, which may affect their preview
+	NotifyUserSceneTextureLoadOrUnload();
+
 	// Broadcast that this editor is going down to all listeners
 	OnMaterialEditorClosed().Broadcast();
 
@@ -1009,6 +1017,37 @@ void FMaterialEditor::UpdatePreviewViewportsVisibility()
 	}
 }
 
+void FMaterialEditor::NotifyUserSceneTextureLoadOrUnload()
+{
+	if (Material->IsPostProcessMaterial() && !Material->UserSceneTexture.IsNone())
+	{
+		UAssetEditorSubsystem* AssetEditorSubsystem = GEditor->GetEditorSubsystem<UAssetEditorSubsystem>();
+		TArray<UObject*> EditedAssets = AssetEditorSubsystem->GetAllEditedAssets();
+		for (UObject* EditedAsset : EditedAssets)
+		{
+			UPreviewMaterial* EditedMaterialInterface = Cast<UPreviewMaterial>(EditedAsset);
+
+			if (EditedMaterialInterface && EditedMaterialInterface != Material)
+			{
+				UMaterial* EditedMaterial = EditedMaterialInterface->GetMaterial();
+				if (EditedMaterial->IsPostProcessMaterial())
+				{
+					TArray<IAssetEditorInstance*> Editors = AssetEditorSubsystem->FindEditorsForAsset(EditedAsset);
+					for (IAssetEditorInstance* Editor : Editors)
+					{
+						if (Editor->GetEditorName() == FMaterialEditor::GetToolkitFName())
+						{
+							// Calling "SetPreviewMaterial" will refresh the other editor
+							FMaterialEditor* MaterialEditor = (FMaterialEditor*)Editor;
+							MaterialEditor->SetPreviewMaterial(EditedMaterialInterface);
+						}
+					}
+				}
+			}
+		}
+	}
+}
+
 static void AssignPinSourceIndices(UEdGraphNode* Node)
 {
 	int32 NumInputDataPins = 0;
diff --git a/Engine/Source/Editor/MaterialEditor/Private/MaterialEditor.h b/Engine/Source/Editor/MaterialEditor/Private/MaterialEditor.h
--- a/Engine/Source/Editor/MaterialEditor/Private/MaterialEditor.h
+++ b/Engine/Source/Editor/MaterialEditor/Private/MaterialEditor.h
@@ -460,6 +460,9 @@ public:
 	/** Set to true if stats should be displayed from the preview material. */
 	bool bStatsFromPreviewMaterial;
 
+	/** Set to true when editor is being destructed */
+	bool bDestructing = false;
+
 	/** The material applied to the preview mesh. */
 	TObjectPtr<UMaterial> Material;
 
@@ -635,6 +638,9 @@ private:
 	/** Updates the 3D and UI preview viewport visibility based on material domain */
 	void UpdatePreviewViewportsVisibility();
 
+	/** Called to notify other material editors when a post process material with a UserSceneTexture output is loaded or unloaded, as the material may be used in their preview */
+	void NotifyUserSceneTextureLoadOrUnload();
+
 	//@TODO: these methods are mostly C&P from BlueprintEditor, consider consolidating logic to graph editor. Note: We don't support macros / functions / tunnels / split pins, and also have material expression specific considerations. */
 	// void CollapseNodesIntoGraph(UEdGraphNode* InGatewayNode, UMaterialGraphNode* InEntryNode, UMaterialGraphNode* InResultNode, UEdGraph* InSourceGraph, UEdGraph* InDestinationGraph, TSet<UEdGraphNode*>& InCollapsableNodes);
 	// void CollapseNodes(TSet<class UEdGraphNode*>& InCollapsableNodes);
diff --git a/Engine/Source/Editor/MaterialEditor/Private/SMaterialEditorViewport.cpp b/Engine/Source/Editor/MaterialEditor/Private/SMaterialEditorViewport.cpp
--- a/Engine/Source/Editor/MaterialEditor/Private/SMaterialEditorViewport.cpp
+++ b/Engine/Source/Editor/MaterialEditor/Private/SMaterialEditorViewport.cpp
@@ -14,7 +14,10 @@
 #include "Engine/SkeletalMesh.h"
 #include "Engine/Texture2D.h"
 #include "Editor/UnrealEdEngine.h"
+#include "MaterialEditor.h"
 #include "MaterialEditor/MaterialEditorMeshComponent.h"
+#include "MaterialEditor/PreviewMaterial.h"
+#include "Materials/MaterialExpressionUserSceneTexture.h"
 #include "Engine/StaticMesh.h"
 #include "Engine/Selection.h"
 #include "Editor.h"
@@ -428,6 +431,79 @@ bool SMaterialEditor3DPreviewViewport::SetPreviewAssetByName(const TCHAR* InAsse
 	return bSuccess;
 }
 
+// Add user scene texture inputs from a material.  Doesn't clear TSet first, so can be used to accumulate inputs from multiple materials.
+void GetUserSceneTextureInputs(UMaterial* Material, TSet<FName>& OutUserSceneTextures)
+{
+	TArray<const UMaterialExpressionUserSceneTexture*> UserSceneTextureExpressions;
+	Material->GetAllExpressionsInMaterialAndFunctionsOfType(UserSceneTextureExpressions);
+
+	for (const UMaterialExpressionUserSceneTexture* UserSceneTextureExpression : UserSceneTextureExpressions)
+	{
+		if (!UserSceneTextureExpression->UserSceneTexture.IsNone())
+		{
+			OutUserSceneTextures.Add(UserSceneTextureExpression->UserSceneTexture);
+		}
+	}
+}
+
+// Recursively get all edited materials that have UserSceneTexture outputs that feed into Material.  Also returns inputs that
+// are missing, which may be useful to report as warnings to the log in the future.
+void GetUserSceneTextureDependencies(UMaterial* Material, TSet<UPreviewMaterial*>& OutDependencies, TSet<FName>& OutMissingInputs)
+{
+	// Check if the current material has any UserSceneTexture inputs first
+	TSet<FName> InputsToProcess;
+	GetUserSceneTextureInputs(Material, InputsToProcess);
+	if (InputsToProcess.IsEmpty())
+	{
+		return;
+	}
+
+	// Generate a global list of edited materials that generate a given UserSceneTexture output (minus Material itself)
+	TMap<FName, TSet<UPreviewMaterial*>> MaterialsByUserSceneTextureOutput;
+
+	UAssetEditorSubsystem* AssetEditorSubsystem = GEditor->GetEditorSubsystem<UAssetEditorSubsystem>();
+	TArray<UObject*> EditedAssets = AssetEditorSubsystem->GetAllEditedAssets();
+	for (int32 AssetIdx = 0; AssetIdx < EditedAssets.Num(); AssetIdx++)
+	{
+		UObject* EditedAsset = EditedAssets[AssetIdx];
+		UPreviewMaterial* EditedMaterial = Cast<UPreviewMaterial>(EditedAsset);
+
+		if (EditedMaterial && EditedMaterial != Material && EditedMaterial->IsPostProcessMaterial() && !EditedMaterial->UserSceneTexture.IsNone())
+		{
+			TArray<IAssetEditorInstance*> Editors = AssetEditorSubsystem->FindEditorsForAsset(EditedAsset);
+			if (!Editors.IsEmpty() && Editors[0]->GetEditorName() == FName("MaterialEditor") && !((const FMaterialEditor*)Editors[0])->bDestructing)
+			{
+				MaterialsByUserSceneTextureOutput.FindOrAdd(EditedMaterial->UserSceneTexture).Add(EditedMaterial);
+			}
+		}
+	}
+
+	// Recursively process materials that generate inputs we care about.  InputsToProcess starts with the inputs from the original
+	// material, and accumulates inputs from other encountered materials.  Stops when no new unique elements get added to InputsToProcess.
+	for (int32 ElementIndex = 0; ElementIndex < InputsToProcess.Num(); ++ElementIndex)
+	{
+		// Find materials that generate an input we care about
+		FName Input = InputsToProcess[FSetElementId::FromInteger(ElementIndex)];
+		TSet<UPreviewMaterial*>* MaterialsGeneratingInput = MaterialsByUserSceneTextureOutput.Find(Input);
+
+		if (MaterialsGeneratingInput)
+		{
+			// Add the materials to the dependency list
+			OutDependencies.Append(*MaterialsGeneratingInput);
+
+			// Add any inputs the new dependencies require
+			for (UPreviewMaterial* MaterialGeneratingOutput : *MaterialsGeneratingInput)
+			{
+				GetUserSceneTextureInputs(MaterialGeneratingOutput, InputsToProcess);
+			}
+		}
+		else
+		{
+			OutMissingInputs.Add(Input);
+		}
+	}
+}
+
 void SMaterialEditor3DPreviewViewport::SetPreviewMaterial(UMaterialInterface* InMaterialInterface)
 {
 	PreviewMaterial = InMaterialInterface;
@@ -443,18 +519,37 @@ void SMaterialEditor3DPreviewViewport::SetPreviewMaterial(UMaterialInterface* In
 			GetViewportClient()->EngineShowFlags.SetPostProcessMaterial(true);
 		}
 
+		// Clear blendables, and re-add them (cleans up any post process materials with UserSceneTextures that are no longer used or loaded)
+		PostProcessVolumeActor->Settings.WeightedBlendables.Array.Empty(1);
+
 		check (PreviewMaterial != nullptr);
 		PostProcessVolumeActor->AddOrUpdateBlendable(PreviewMaterial);
 		PostProcessVolumeActor->bEnabled = true;
 		PostProcessVolumeActor->BlendWeight = 1.0f;
 		PostProcessVolumeActor->bUnbound = true;
 
+		// Setting this forces this post process material to write to SceneColor instead of any UserSceneTexture it may have assigned, for preview purposes
+		PostProcessVolumeActor->Settings.PreviewBlendable = PreviewMaterial;
+
 		// Remove preview material from the preview mesh.
 		if (PreviewMeshComponent != nullptr)
 		{
 			PreviewMeshComponent->OverrideMaterials.Empty();
 			PreviewMeshComponent->MarkRenderStateDirty();
 		}
+
+		{
+			// Add any edited post process materials that write UserSceneTextures used by this material, for better visualization
+			TSet<UPreviewMaterial*> UserSceneTextureDependencies;
+			TSet<FName> UserSceneTextureMissingInputs;
+
+			GetUserSceneTextureDependencies(PreviewMaterial->GetMaterial(), UserSceneTextureDependencies, UserSceneTextureMissingInputs);
+
+			for (UPreviewMaterial* Dependency : UserSceneTextureDependencies)
+			{
+				PostProcessVolumeActor->AddOrUpdateBlendable(Dependency);
+			}
+		}
 	}
 	else
 	{
@@ -760,11 +855,14 @@ EVisibility SMaterialEditor3DPreviewViewport::OnGetViewportContentVisibility() c
 
 void SMaterialEditor3DPreviewViewport::OnPropertyChanged(UObject* ObjectBeingModified, FPropertyChangedEvent& PropertyChangedEvent)
 {
-	if (ObjectBeingModified != nullptr && ObjectBeingModified == PreviewMaterial)
+	FProperty* PropertyThatChanged = PropertyChangedEvent.Property;
+	static const FString MaterialDomain = TEXT("MaterialDomain");
+	static const FString UserSceneTexture = TEXT("UserSceneTexture");
+	if (ObjectBeingModified != nullptr && PropertyThatChanged != nullptr && (PropertyThatChanged->GetName() == MaterialDomain || PropertyThatChanged->GetName() == UserSceneTexture))
 	{
-		FProperty* PropertyThatChanged = PropertyChangedEvent.Property;
-		static const FString MaterialDomain = TEXT("MaterialDomain");
-		if (PropertyThatChanged != nullptr && PropertyThatChanged->GetName() == MaterialDomain)
+		// Refresh preview material if this specific material has been modified, or if this is any other post process material that has an editor preview,
+		// which may use this preview material in its preview if it has the UserSceneTexture as one of its inputs.
+		if (ObjectBeingModified == PreviewMaterial || PreviewMaterial->GetMaterial()->IsPostProcessMaterial())
 		{
 			SetPreviewMaterial(PreviewMaterial);
 		}
diff --git a/Engine/Source/Runtime/Engine/Classes/Engine/Scene.h b/Engine/Source/Runtime/Engine/Classes/Engine/Scene.h
--- a/Engine/Source/Runtime/Engine/Classes/Engine/Scene.h
+++ b/Engine/Source/Runtime/Engine/Classes/Engine/Scene.h
@@ -2295,6 +2295,13 @@ struct FPostProcessSettings
 	FWeightedBlendables WeightedBlendables;
 
 #if WITH_EDITORONLY_DATA
+	/**
+	 * For editor material preview windows, we need to support visualizing the output of a blendable that writes to a UserSceneTexture.  Stores
+	 * a pointer to a blendable that's being previewed, forcing its output to write to SceneColor instead of the UserSceneTexture, making it visible.
+	 */
+	UPROPERTY(Transient)
+	TObjectPtr<UObject> PreviewBlendable;
+
 	// for backwards compatibility
 	UPROPERTY()
 	TArray<TObjectPtr<UObject>> Blendables_DEPRECATED;
diff --git a/Engine/Source/Runtime/Engine/Private/Materials/HLSLMaterialTranslator.cpp b/Engine/Source/Runtime/Engine/Private/Materials/HLSLMaterialTranslator.cpp
--- a/Engine/Source/Runtime/Engine/Private/Materials/HLSLMaterialTranslator.cpp
+++ b/Engine/Source/Runtime/Engine/Private/Materials/HLSLMaterialTranslator.cpp
@@ -59,6 +59,7 @@
 #include "DerivedDataRequestOwner.h"
 #include "DerivedDataCache.h"
 #include "MaterialCachedData.h"
+#include "PostProcess/PostProcessMaterialInputs.h"
 
 #if WITH_EDITORONLY_DATA
 #include "Materials/MaterialExpressionSubstrate.h"
@@ -222,6 +223,10 @@ UE::DerivedData::FValueId MaterialCompilationOutputId = UE::DerivedData::FValueI
 UE::DerivedData::FValueId MaterialResultsOutputId = UE::DerivedData::FValueId::FromName("FHLSLMaterialTranslator_Results");
 UE::DerivedData::FValueId EnvironmentDefinesId = UE::DerivedData::FValueId::FromName("FHLSLMaterialTranslator_EnvironmentDefines");
 
+/** HLSL generating utility functions */
+extern FString GenerateUserSceneTextureRemapHLSLDefines(FMaterialCompilationOutput& CompilationOutput);
+extern FString SceneTextureIdToHLSLString(ESceneTextureId TexId);
+
 /** Data structure used to cache a part of material translation results. It contains all the generated
  *  defines that will be declared during the compilation of the generated material shader.
  *  Note: if you make a change to this structure, remember to bump FDevSystemGuids::MaterialTranslationDDCVersion.
@@ -1953,6 +1958,12 @@ void FHLSLMaterialTranslator::DoTranslate()
 		}
 	}
 
+	int32 NumPostProcessInputs = MaterialCompilationOutput.GetNumPostProcessInputsUsed();
+	if (NumPostProcessInputs > kPostProcessMaterialInputCountMax)
+	{
+		Errorf(TEXT("Maximum Scene Texture post process inputs exceeded (%d > %d), between SceneTexture nodes with PostProcessInputs or UserSceneTexture nodes."), NumPostProcessInputs, kPostProcessMaterialInputCountMax);
+	}
+
 	if (IsModulateBlendMode(BlendMode) && Material->IsTranslucencyAfterMotionBlurEnabled())
 	{
 		// We don't currently have a separate translucency modulation pass for After Motion Blur
@@ -4236,6 +4247,11 @@ EMaterialValueType FHLSLMaterialTranslator::GetArithmeticResultType(int32 A,int3
 	return GetArithmeticResultType(TypeA,TypeB);
 }
 
+int32 FHLSLMaterialTranslator::FindOrAddUserSceneTexture(FName UserSceneTextureName)
+{
+	return MaterialCompilationOutput.FindOrAddUserSceneTexture(UserSceneTextureName);
+}
+
 // FMaterialCompiler interface.
 
 /** 
@@ -7730,7 +7746,7 @@ int32 FHLSLMaterialTranslator::SceneDepth(int32 Offset, int32 ViewportUV, bool b
 }
 
 // @param SceneTextureId of type ESceneTextureId e.g. PPI_SubsurfaceColor
-int32 FHLSLMaterialTranslator::SceneTextureLookup(int32 ViewportUV, uint32 InSceneTextureId, bool bFiltered)
+int32 FHLSLMaterialTranslator::SceneTextureLookup(int32 ViewportUV, uint32 InSceneTextureId, bool bFiltered, bool bClamped)
 {
 	ESceneTextureId SceneTextureId = (ESceneTextureId)InSceneTextureId;
 
@@ -7754,16 +7770,25 @@ int32 FHLSLMaterialTranslator::SceneTextureLookup(int32 ViewportUV, uint32 InSce
 
 	UseSceneTextureId(SceneTextureId, true);
 
+	FString SceneTextureIdString = SceneTextureIdToHLSLString(SceneTextureId);
+
 	int32 BufferUV;
 	if (ViewportUV != INDEX_NONE)
 	{
 		BufferUV = AddCodeChunk(MCT_Float2,
-			TEXT("ClampSceneTextureUV(ViewportUVToSceneTextureUV(%s, %d), %d)"),
-			*CoerceParameter(ViewportUV, MCT_Float2), (int)SceneTextureId, (int)SceneTextureId);
+			TEXT("ClampSceneTextureUV(ViewportUVToSceneTextureUV(%s, %s), %s)"),
+			*CoerceParameter(ViewportUV, MCT_Float2), *SceneTextureIdString, *SceneTextureIdString);
 	}
 	else
 	{
-		BufferUV = AddInlinedCodeChunk(MCT_Float2, TEXT("GetDefaultSceneTextureUV(Parameters, %d)"), (int)SceneTextureId);
+		if (bClamped)
+		{
+			BufferUV = AddInlinedCodeChunk(MCT_Float2, TEXT("ClampSceneTextureUV(GetDefaultSceneTextureUV(Parameters, %s), %s)"), *SceneTextureIdString, *SceneTextureIdString);
+		}
+		else
+		{
+			BufferUV = AddInlinedCodeChunk(MCT_Float2, TEXT("GetDefaultSceneTextureUV(Parameters, %s)"), *SceneTextureIdString);
+		}
 	}
 
 	AddEstimatedTextureSample();
@@ -7774,13 +7799,13 @@ int32 FHLSLMaterialTranslator::SceneTextureLookup(int32 ViewportUV, uint32 InSce
 	{
 		LookUp = AddCodeChunk(
 			MCT_Float4,
-			TEXT("SceneTextureLookup(%s, %d, %s)"),
-			*CoerceParameter(BufferUV, MCT_Float2), (int)SceneTextureId, bFiltered ? TEXT("true") : TEXT("false")
+			TEXT("SceneTextureLookup(%s, %s, %s)"),
+			*CoerceParameter(BufferUV, MCT_Float2), *SceneTextureIdString, bFiltered ? TEXT("true") : TEXT("false")
 		);
 	}
 	else // mobile
 	{
-		LookUp = AddCodeChunk(MCT_Float4, TEXT("MobileSceneTextureLookup(Parameters, %d, %s)"), (int32)SceneTextureId, *CoerceParameter(BufferUV, MCT_Float2));
+		LookUp = AddCodeChunk(MCT_Float4, TEXT("MobileSceneTextureLookup(Parameters, %s, %s)"), *SceneTextureIdString, *CoerceParameter(BufferUV, MCT_Float2));
 	}
 
 	// Substrate only
@@ -7792,7 +7817,8 @@ int32 FHLSLMaterialTranslator::SceneTextureLookup(int32 ViewportUV, uint32 InSce
 	SubstrateCtx.SubstrateMaterialComplexity.bIsSimple = false;
 	SubstrateCtx.SubstrateMaterialComplexity.bIsSingle = false;
 	
-	if (SceneTextureId >= PPI_PostProcessInput0 && SceneTextureId <= PPI_PostProcessInput6 && Material->GetMaterialDomain() == MD_PostProcess && Material->GetBlendableLocation() != BL_SceneColorAfterTonemapping)
+	if (((SceneTextureId >= PPI_PostProcessInput0 && SceneTextureId <= PPI_PostProcessInput6) || (SceneTextureId >= PPI_UserSceneTexture0 && SceneTextureId <= PPI_UserSceneTexture6)) &&
+		Material->GetMaterialDomain() == MD_PostProcess && Material->GetBlendableLocation() != BL_SceneColorAfterTonemapping)
 	{
 		return AddInlinedCodeChunk(MCT_Float4, TEXT("(float4(View.OneOverPreExposure.xxx, 1) * %s)"), *CoerceParameter(LookUp, MCT_Float4));
 	}
@@ -7804,11 +7830,13 @@ int32 FHLSLMaterialTranslator::SceneTextureLookup(int32 ViewportUV, uint32 InSce
 
 int32 FHLSLMaterialTranslator::GetSceneTextureViewSize(int32 SceneTextureId, bool InvProperty)
 {
+	FString SceneTextureIdString = SceneTextureIdToHLSLString((ESceneTextureId)SceneTextureId);
+
 	if (InvProperty)
 	{
-		return AddCodeChunkZeroDeriv(MCT_Float2, TEXT("GetSceneTextureViewSize(%d).zw"), SceneTextureId);
+		return AddCodeChunkZeroDeriv(MCT_Float2, TEXT("GetSceneTextureViewSize(%s).zw"), *SceneTextureIdString);
 	}
-	return AddCodeChunkZeroDeriv(MCT_Float2, TEXT("GetSceneTextureViewSize(%d).xy"), SceneTextureId);
+	return AddCodeChunkZeroDeriv(MCT_Float2, TEXT("GetSceneTextureViewSize(%s).xy"), *SceneTextureIdString);
 }
 
 // @param bTextureLookup true: texture, false:no texture lookup, usually to get the size
@@ -7856,6 +7884,7 @@ void FHLSLMaterialTranslator::UseSceneTextureId(ESceneTextureId SceneTextureId,
 	{
 		bNeedsSceneTexturePostProcessInputs = bNeedsSceneTexturePostProcessInputs
 			|| ((SceneTextureId >= PPI_PostProcessInput0 && SceneTextureId <= PPI_PostProcessInput6)
+			|| (SceneTextureId >= PPI_UserSceneTexture0 && SceneTextureId <= PPI_UserSceneTexture6)
 			|| SceneTextureId == PPI_Velocity
 			|| SceneTextureId == PPI_SceneColor);
 
@@ -15470,6 +15499,8 @@ void FHLSLMaterialTranslator::PrepareMaterialSourceStringParameters()
 			MaterialSourceTemplateParams.Add({ TEXT("calc_pixel_material_inputs_analytic_derivatives_other_inputs"), PixelMembersSetupAndAssignments[CompiledPDV_Analytic] });
 		}
 	}
+
+	MaterialSourceTemplateParams.Add({ TEXT("user_scene_texture_remap"), GenerateUserSceneTextureRemapHLSLDefines(MaterialCompilationOutput) });
 }
 
 int32 FHLSLMaterialTranslator::SparseVolumeTexture(USparseVolumeTexture* Texture, int32& TextureReferenceIndex, EMaterialSamplerType SamplerType)
diff --git a/Engine/Source/Runtime/Engine/Private/Materials/HLSLMaterialTranslator.h b/Engine/Source/Runtime/Engine/Private/Materials/HLSLMaterialTranslator.h
--- a/Engine/Source/Runtime/Engine/Private/Materials/HLSLMaterialTranslator.h
+++ b/Engine/Source/Runtime/Engine/Private/Materials/HLSLMaterialTranslator.h
@@ -37,6 +37,7 @@
 
 #if WITH_EDITORONLY_DATA
 #include "Materials/MaterialExpressionSceneTexture.h"
+#include "Materials/MaterialExpressionUserSceneTexture.h"
 #include "Materials/MaterialExpressionNoise.h"
 #include "Materials/MaterialExpressionFunctionInput.h"
 #include "Materials/MaterialExpressionFunctionOutput.h"
@@ -1041,10 +1042,12 @@ protected:
 	virtual int32 SceneDepth(int32 Offset, int32 ViewportUV, bool bUseOffset) override;
 	
 	// @param SceneTextureId of type ESceneTextureId e.g. PPI_SubsurfaceColor
-	virtual int32 SceneTextureLookup(int32 ViewportUV, uint32 InSceneTextureId, bool bFiltered) override;
+	virtual int32 SceneTextureLookup(int32 ViewportUV, uint32 InSceneTextureId, bool bFiltered, bool bClamped) override;
 
 	virtual int32 GetSceneTextureViewSize(int32 SceneTextureId, bool InvProperty) override;
 
+	virtual int32 FindOrAddUserSceneTexture(FName UserSceneTextureName) override;
+
 	virtual int32 DBufferTextureLookup(int32 ViewportUV, uint32 DBufferTextureIndex) override;
 
 	virtual int32 PathTracingBufferTextureLookup(int32 ViewportUV, uint32 PathTracingBufferTextureIndex) override;
diff --git a/Engine/Source/Runtime/Engine/Private/Materials/Material.cpp b/Engine/Source/Runtime/Engine/Private/Materials/Material.cpp
--- a/Engine/Source/Runtime/Engine/Private/Materials/Material.cpp
+++ b/Engine/Source/Runtime/Engine/Private/Materials/Material.cpp
@@ -4540,6 +4540,13 @@ bool UMaterial::CanEditChange(const FProperty* InProperty) const
 			return MaterialDomain == MD_PostProcess;
 		}
 
+		if (PropertyName == GET_MEMBER_NAME_STRING_CHECKED(UMaterial, UserSceneTexture) ||
+			PropertyName == GET_MEMBER_NAME_STRING_CHECKED(UMaterial, UserTextureDivisor))
+		{
+			// "Replacing Tonemapper" blendable location doesn't support a UserSceneTexture output
+			return MaterialDomain == MD_PostProcess && BlendableLocation != BL_ReplacingTonemapper;
+		}
+
 		if (PropertyName == GET_MEMBER_NAME_STRING_CHECKED(UMaterial, BlendMode))
 		{
 			if (bSubstrateEnabled)
diff --git a/Engine/Source/Runtime/Engine/Private/Materials/MaterialExpressionHLSL.cpp b/Engine/Source/Runtime/Engine/Private/Materials/MaterialExpressionHLSL.cpp
--- a/Engine/Source/Runtime/Engine/Private/Materials/MaterialExpressionHLSL.cpp
+++ b/Engine/Source/Runtime/Engine/Private/Materials/MaterialExpressionHLSL.cpp
@@ -227,6 +227,7 @@
 #include "Materials/MaterialExpressionTruncate.h"
 #include "Materials/MaterialExpressionTruncateLWC.h"
 #include "Materials/MaterialExpressionTwoSidedSign.h"
+#include "Materials/MaterialExpressionUserSceneTexture.h"
 #include "Materials/MaterialExpressionVectorNoise.h"
 #include "Materials/MaterialExpressionVectorParameter.h"
 #include "Materials/MaterialExpressionVertexColor.h"
@@ -2244,7 +2245,7 @@ bool UMaterialExpressionSceneTexture::GenerateHLSLExpression(FMaterialHLSLGenera
 		{
 			ExpressionTexCoord = Coordinates.AcquireHLSLExpression(Generator, Scope);
 		}
-		OutExpression = Generator.GetTree().NewExpression<Material::FExpressionSceneTexture>(ExpressionTexCoord, SceneTextureId, bFiltered);
+		OutExpression = Generator.GetTree().NewExpression<Material::FExpressionSceneTexture>(ExpressionTexCoord, SceneTextureId, bFiltered, /*bClamped=*/ false, FName());
 		return true;
 	}
 	else if (OutputIndex == 1 || OutputIndex == 2)
@@ -2261,6 +2262,37 @@ bool UMaterialExpressionSceneTexture::GenerateHLSLExpression(FMaterialHLSLGenera
 	return Generator.Error(TEXT("Invalid input parameter"));
 }
 
+bool UMaterialExpressionUserSceneTexture::GenerateHLSLExpression(FMaterialHLSLGenerator& Generator, UE::HLSLTree::FScope& Scope, int32 OutputIndex, UE::HLSLTree::FExpression const*& OutExpression) const
+{
+	using namespace UE::HLSLTree;
+	if (UserSceneTexture.IsNone())
+	{
+		return Generator.Error(TEXT("UserSceneTexture missing name -- value must be set to something other than None"));
+	}
+
+	if (OutputIndex == 0)
+	{
+		const FExpression* ExpressionTexCoord = nullptr;
+		if (Coordinates.GetTracedInput().Expression)
+		{
+			ExpressionTexCoord = Coordinates.AcquireHLSLExpression(Generator, Scope);
+		}
+
+		// We arbitrarily set the SceneTextureId to PPI_UserSceneTexture0 -- final ID is determined later in a call to FMaterialCompilationOutput::FindOrAddUserSceneTexture
+		// from FExpressionSceneTexture::PrepareValue, with this ID not being used except for classification of the expression as as UserSceneTexture type.
+		OutExpression = Generator.GetTree().NewExpression<Material::FExpressionSceneTexture>(ExpressionTexCoord, PPI_UserSceneTexture0, bFiltered, bClamped, UserSceneTexture);
+		return true;
+	}
+	else if (OutputIndex == 1 || OutputIndex == 2)
+	{
+		const bool bReciprocal = (OutputIndex == 2);
+		OutExpression = Generator.GetTree().NewExpression<Material::FExpressionUserSceneTextureSize>(UserSceneTexture, bReciprocal);
+		return true;
+	}
+
+	return Generator.Error(TEXT("Invalid input parameter"));
+}
+
 bool UMaterialExpressionNoise::GenerateHLSLExpression(FMaterialHLSLGenerator& Generator, UE::HLSLTree::FScope& Scope, int32 OutputIndex, UE::HLSLTree::FExpression const*& OutExpression) const
 {
 	using namespace UE::HLSLTree;
diff --git a/Engine/Source/Runtime/Engine/Private/Materials/MaterialExpressions.cpp b/Engine/Source/Runtime/Engine/Private/Materials/MaterialExpressions.cpp
--- a/Engine/Source/Runtime/Engine/Private/Materials/MaterialExpressions.cpp
+++ b/Engine/Source/Runtime/Engine/Private/Materials/MaterialExpressions.cpp
@@ -254,6 +254,7 @@
 #include "Materials/MaterialExpressionTruncate.h"
 #include "Materials/MaterialExpressionTruncateLWC.h"
 #include "Materials/MaterialExpressionTwoSidedSign.h"
+#include "Materials/MaterialExpressionUserSceneTexture.h"
 #include "Materials/MaterialExpressionVectorNoise.h"
 #include "Materials/MaterialExpressionVertexColor.h"
 #include "Materials/MaterialExpressionVertexNormalWS.h"
@@ -303,6 +304,7 @@
 #include "MaterialGraph/MaterialGraphSchema.h"
 #include "Serialization/ShaderKeyGenerator.h"
 #include "SubstrateMaterial.h"
+#include "PostProcess/PostProcessMaterialInputs.h"
 #else
 #include "Materials/MaterialExpressionVertexInterpolator.h"
 #include "Materials/MaterialParameterCollection.h"
@@ -11804,8 +11806,11 @@ int32 UMaterialExpressionSceneTexture::Compile(class FMaterialCompiler* Compiler
 
 	if(OutputIndex == 0)
 	{
-		// Color
-		return Compiler->SceneTextureLookup(ViewportUV, SceneTextureId, bFiltered);
+		// Color.  Note that clamping support is not necessary for regular SceneTexture, because it's only useful when sampling from lower resolution
+		// maps with filtering, where bilinear blending of a higher resolution UV sample can end up interpolating with pixels outside the valid UV
+		// range on a lower resolution map.  All SceneTextures are full resolution, while UserSceneTextures can be lower resolution (see
+		// UMaterialExpressionUserSceneTexture::Compile below), so those support a user specified clamp flag.
+		return Compiler->SceneTextureLookup(ViewportUV, SceneTextureId, bFiltered, /*bClamped=*/ false);
 	}
 	else if(OutputIndex == 1 || OutputIndex == 2)
 	{
@@ -11827,6 +11832,81 @@ void UMaterialExpressionSceneTexture::GetCaption(TArray<FString>& OutCaptions) c
 }
 #endif // WITH_EDITOR
 
+///////////////////////////////////////////////////////////////////////////////
+// UMaterialExpressionUserSceneTexture
+///////////////////////////////////////////////////////////////////////////////
+UMaterialExpressionUserSceneTexture::UMaterialExpressionUserSceneTexture(const FObjectInitializer& ObjectInitializer)
+	: Super(ObjectInitializer)
+{
+#if WITH_EDITORONLY_DATA
+	// Structure to hold one-time initialization
+	struct FConstructorStatics
+	{
+		FText NAME_Texture;
+		FConstructorStatics()
+			: NAME_Texture(LOCTEXT("Texture", "Texture"))
+		{
+		}
+	};
+	static FConstructorStatics ConstructorStatics;
+
+	MenuCategories.Add(ConstructorStatics.NAME_Texture);
+
+	bShaderInputData = true;
+	bShowOutputNameOnPin = true;
+#endif
+
+	// by default faster, most lookup are read/write the same pixel so this is rarely needed
+	bFiltered = false;
+
+#if WITH_EDITORONLY_DATA
+	Outputs.Reset();
+	Outputs.Add(FExpressionOutput(TEXT("Color"), 1, 1, 1, 1, 1));
+	Outputs.Add(FExpressionOutput(TEXT("Size")));
+	Outputs.Add(FExpressionOutput(TEXT("InvSize")));
+#endif
+}
+
+#if WITH_EDITOR
+int32 UMaterialExpressionUserSceneTexture::Compile(class FMaterialCompiler* Compiler, int32 OutputIndex)
+{
+	if (UserSceneTexture.IsNone())
+	{
+		return Compiler->Errorf(TEXT("UserSceneTexture missing name -- value must be set to something other than None"));
+	}
+
+	int32 SceneTextureId = Compiler->FindOrAddUserSceneTexture(UserSceneTexture);
+	if (SceneTextureId == INDEX_NONE)
+	{
+		return Compiler->Errorf(TEXT("Too many unique UserSceneTexture inputs in the post process material -- max allowed is %d"), kPostProcessMaterialInputCountMax);
+	}
+
+	int32 ViewportUV = INDEX_NONE;
+
+	if (Coordinates.GetTracedInput().Expression)
+	{
+		ViewportUV = Coordinates.Compile(Compiler);
+	}
+
+	if (OutputIndex == 0)
+	{
+		// Color
+		return Compiler->SceneTextureLookup(ViewportUV, SceneTextureId, bFiltered, bClamped);
+	}
+	else if (OutputIndex == 1 || OutputIndex == 2)
+	{
+		return Compiler->GetSceneTextureViewSize(SceneTextureId, /* InvProperty = */ OutputIndex == 2);
+	}
+
+	return Compiler->Errorf(TEXT("Invalid input parameter"));
+}
+
+void UMaterialExpressionUserSceneTexture::GetCaption(TArray<FString>& OutCaptions) const
+{
+	OutCaptions.Add(FString(TEXT("UserSceneTexture:")) + UserSceneTexture.ToString());
+}
+#endif // WITH_EDITOR
+
 ///////////////////////////////////////////////////////////////////////////////
 // UMaterialExpressionSceneColor
 ///////////////////////////////////////////////////////////////////////////////
diff --git a/Engine/Source/Runtime/Engine/Private/Materials/MaterialHLSLEmitter.cpp b/Engine/Source/Runtime/Engine/Private/Materials/MaterialHLSLEmitter.cpp
--- a/Engine/Source/Runtime/Engine/Private/Materials/MaterialHLSLEmitter.cpp
+++ b/Engine/Source/Runtime/Engine/Private/Materials/MaterialHLSLEmitter.cpp
@@ -21,9 +21,13 @@
 #include "RenderUtils.h"
 #include "DataDrivenShaderPlatformInfo.h"
 #include "Engine/SubsurfaceProfile.h"
+#include "PostProcess/PostProcessMaterialInputs.h"
 
 static bool SharedPixelProperties[CompiledMP_MAX];
 
+/** HLSL generating utility function */
+extern FString GenerateUserSceneTextureRemapHLSLDefines(FMaterialCompilationOutput& CompilationOutput);
+
 static const TCHAR* HLSLTypeString(EMaterialValueType Type)
 {
 	switch (Type)
@@ -299,6 +303,8 @@ static FString GenerateMaterialTemplateHLSL(EShaderPlatform ShaderPlatform,
 		}
 	}
 
+	MaterialSourceTemplateParams.Add({ TEXT("user_scene_texture_remap"), GenerateUserSceneTextureRemapHLSLDefines(OutCompilationOutput) });
+
 	return Resolver.Finalize();
 }
 
@@ -1009,6 +1015,12 @@ bool MaterialEmitHLSL(const FMaterialCompileTargetParameters& InCompilerTarget,
 		EmitContext.Error(TEXT("Only transparent or postprocess materials can read from scene depth."));
 	}
 
+	int32 NumPostProcessInputs = OutCompilationOutput.GetNumPostProcessInputsUsed();
+	if (NumPostProcessInputs > kPostProcessMaterialInputCountMax)
+	{
+		EmitContext.Errorf(TEXT("Maximum Scene Texture post process inputs exceeded (%d > %d), between SceneTexture nodes with PostProcessInputs or UserSceneTexture nodes."), NumPostProcessInputs, kPostProcessMaterialInputCountMax);
+	}
+
 	if (EmitContext.NumErrors > 0)
 	{
 		return false;
diff --git a/Engine/Source/Runtime/Engine/Private/Materials/MaterialHLSLTree.cpp b/Engine/Source/Runtime/Engine/Private/Materials/MaterialHLSLTree.cpp
--- a/Engine/Source/Runtime/Engine/Private/Materials/MaterialHLSLTree.cpp
+++ b/Engine/Source/Runtime/Engine/Private/Materials/MaterialHLSLTree.cpp
@@ -16,6 +16,10 @@
 #include "VT/RuntimeVirtualTexture.h"
 #include "DataDrivenShaderPlatformInfo.h"
 #include "RenderUtils.h"
+#include "PostProcess/PostProcessMaterialInputs.h"
+
+/** HLSL generating utility function */
+extern FString SceneTextureIdToHLSLString(ESceneTextureId TexId);
 
 namespace UE::HLSLTree::Material
 {
@@ -2337,7 +2341,20 @@ bool FExpressionSceneTexture::PrepareValue(FEmitContext& Context, FEmitScope& Sc
 	if (Context.bMarkLiveValues && Context.MaterialCompilationOutput)
 	{
 		Context.MaterialCompilationOutput->bNeedsSceneTextures = true;
-		Context.MaterialCompilationOutput->SetIsSceneTextureUsed((ESceneTextureId)SceneTextureId);
+		if (!UserSceneTexture.IsNone())
+		{
+			// We set all HLSL tree expressions with UserSceneTexture to PPI_UserSceneTexture0.  Actual ID is determined by result of FindOrAddUserSceneTexture.
+			check(SceneTextureId == PPI_UserSceneTexture0);
+
+			if (Context.MaterialCompilationOutput->FindOrAddUserSceneTexture(UserSceneTexture) == INDEX_NONE)
+			{
+				return Context.Errorf(TEXT("Too many unique UserSceneTexture inputs in the post process material -- max allowed is %d"), kPostProcessMaterialInputCountMax);
+			}
+		}
+		else
+		{
+			Context.MaterialCompilationOutput->SetIsSceneTextureUsed((ESceneTextureId)SceneTextureId);
+		}
 
 		const bool bNeedsGBuffer = Context.MaterialCompilationOutput->NeedsGBuffer();
 		if (bNeedsGBuffer)
@@ -2368,33 +2385,53 @@ bool FExpressionSceneTexture::PrepareValue(FEmitContext& Context, FEmitScope& Sc
 
 void FExpressionSceneTexture::EmitValueShader(FEmitContext& Context, FEmitScope& Scope, const FRequestedType& RequestedType, FEmitValueShaderResult& OutResult) const
 {
-	const bool bSupportedOnMobile = SceneTextureId == PPI_PostProcessInput0 ||
-		SceneTextureId == PPI_CustomDepth ||
-		SceneTextureId == PPI_SceneDepth ||
-		SceneTextureId == PPI_CustomStencil;
+	FString SceneTextureIdString;
+	if (!UserSceneTexture.IsNone())
+	{
+		// SceneTextureId is arbitrarily set to PPI_UserSceneTexture0 for all UserSceneTextures in the HLSL tree.  Actual ID
+		// is allocated in PrepareValue, based on what PPI_PostProcessInput0-6 slots are available in the compilation output.
+		check(SceneTextureId == PPI_UserSceneTexture0);
+
+		int32 UserSceneTextureId = Context.MaterialCompilationOutput->FindOrAddUserSceneTexture(UserSceneTexture);
+		check(UserSceneTextureId != INDEX_NONE);
+
+		SceneTextureIdString = SceneTextureIdToHLSLString((ESceneTextureId)UserSceneTextureId);
+	}
+	else
+	{
+		SceneTextureIdString = SceneTextureIdToHLSLString((ESceneTextureId)SceneTextureId);
+	}
 
 	FEmitShaderExpression* EmitTexCoord = nullptr;
 	if (TexCoordExpression)
 	{
 		EmitTexCoord = TexCoordExpression->GetValueShader(Context, Scope, Shader::EValueType::Float2);
-		EmitTexCoord = Context.EmitExpression(Scope, Shader::EValueType::Float2, TEXT("ClampSceneTextureUV(ViewportUVToSceneTextureUV(%, %), %)"), EmitTexCoord, (int)SceneTextureId, (int)SceneTextureId);
+		EmitTexCoord = Context.EmitExpression(Scope, Shader::EValueType::Float2, TEXT("ClampSceneTextureUV(ViewportUVToSceneTextureUV(%, %), %)"), EmitTexCoord, *SceneTextureIdString, *SceneTextureIdString);
 	}
 	else
 	{
-		EmitTexCoord = Context.EmitExpression(Scope, Shader::EValueType::Float2, TEXT("GetDefaultSceneTextureUV(Parameters, %)"), (int)SceneTextureId);
+		if (bClamped)
+		{
+			EmitTexCoord = Context.EmitExpression(Scope, Shader::EValueType::Float2, TEXT("ClampSceneTextureUV(GetDefaultSceneTextureUV(Parameters, %), %)"), *SceneTextureIdString, *SceneTextureIdString);
+		}
+		else
+		{
+			EmitTexCoord = Context.EmitExpression(Scope, Shader::EValueType::Float2, TEXT("GetDefaultSceneTextureUV(Parameters, %)"), *SceneTextureIdString);
+		}
 	}
 
 	FEmitShaderExpression* EmitLookup = nullptr;
 	if (Context.Material->GetFeatureLevel() >= ERHIFeatureLevel::SM5)
 	{
-		EmitLookup = Context.EmitExpression(Scope, Shader::EValueType::Float4, TEXT("SceneTextureLookup(%, %, %)"), EmitTexCoord, (int)SceneTextureId, bFiltered);
+		EmitLookup = Context.EmitExpression(Scope, Shader::EValueType::Float4, TEXT("SceneTextureLookup(%, %, %)"), EmitTexCoord, *SceneTextureIdString, bFiltered);
 	}
 	else
 	{
-		EmitLookup = Context.EmitExpression(Scope, Shader::EValueType::Float4, TEXT("MobileSceneTextureLookup(Parameters, %, %)"), (int)SceneTextureId, EmitTexCoord);
+		EmitLookup = Context.EmitExpression(Scope, Shader::EValueType::Float4, TEXT("MobileSceneTextureLookup(Parameters, %, %)"), *SceneTextureIdString, EmitTexCoord);
 	}
 
-	if (SceneTextureId >= PPI_PostProcessInput0 && SceneTextureId <= PPI_PostProcessInput6 && Context.Material->GetMaterialDomain() == MD_PostProcess && Context.Material->GetBlendableLocation() != BL_SceneColorAfterTonemapping)
+	if (((SceneTextureId >= PPI_PostProcessInput0 && SceneTextureId <= PPI_PostProcessInput6) || (SceneTextureId >= PPI_UserSceneTexture0 && SceneTextureId <= PPI_UserSceneTexture6)) &&
+		Context.Material->GetMaterialDomain() == MD_PostProcess && Context.Material->GetBlendableLocation() != BL_SceneColorAfterTonemapping)
 	{
 		EmitLookup = Context.EmitExpression(Scope, Shader::EValueType::Float4, TEXT("(float4(View.OneOverPreExposure.xxx, 1) * %)"), EmitLookup);
 	}
@@ -2402,6 +2439,29 @@ void FExpressionSceneTexture::EmitValueShader(FEmitContext& Context, FEmitScope&
 	OutResult.Code = EmitLookup;
 }
 
+bool FExpressionUserSceneTextureSize::PrepareValue(FEmitContext& Context, FEmitScope& Scope, const FRequestedType& RequestedType, FPrepareValueResult& OutResult) const
+{
+	if (Context.bMarkLiveValues && Context.MaterialCompilationOutput)
+	{
+		Context.MaterialCompilationOutput->bNeedsSceneTextures = true;
+
+		int32 SceneTextureId = Context.MaterialCompilationOutput->FindOrAddUserSceneTexture(UserSceneTexture);
+		if (SceneTextureId == INDEX_NONE)
+		{
+			return Context.Errorf(TEXT("Too many unique UserSceneTexture inputs in the post process material -- max allowed is %d"), kPostProcessMaterialInputCountMax);
+		}
+	}
+	return OutResult.SetType(Context, RequestedType, EExpressionEvaluation::Shader, Shader::EValueType::Float2);
+}
+
+void FExpressionUserSceneTextureSize::EmitValueShader(FEmitContext& Context, FEmitScope& Scope, const FRequestedType& RequestedType, FEmitValueShaderResult& OutResult) const
+{
+	int32 SceneTextureId = Context.MaterialCompilationOutput->FindOrAddUserSceneTexture(UserSceneTexture);
+	FString SceneTextureIdString = SceneTextureIdToHLSLString((ESceneTextureId)SceneTextureId);
+
+	OutResult.Code = Context.EmitExpression(Scope, Shader::EValueType::Float2, TEXT("GetSceneTextureViewSize(%).%"), *SceneTextureIdString, bReciprocal ? TEXT("zw") : TEXT("xy"));
+}
+
 bool FExpressionScreenAlignedUV::PrepareValue(FEmitContext& Context, FEmitScope& Scope, const FRequestedType& RequestedType, FPrepareValueResult& OutResult) const
 {
 	if (OffsetExpression)
diff --git a/Engine/Source/Runtime/Engine/Private/Materials/MaterialShader.cpp b/Engine/Source/Runtime/Engine/Private/Materials/MaterialShader.cpp
--- a/Engine/Source/Runtime/Engine/Private/Materials/MaterialShader.cpp
+++ b/Engine/Source/Runtime/Engine/Private/Materials/MaterialShader.cpp
@@ -2505,6 +2505,15 @@ void FMaterialShaderMap::Compile(
 	ShaderMapId = InShaderMapId;
 	bIsPersistent = Material->IsPersistent();
 
+	// Fill in UserSceneTextureOutput -- only supported for post process domain, and not supported for BL_ReplacingTonemapper
+	const UMaterial* MaterialObject = Material->GetMaterialInterface()->GetMaterial();
+	if (MaterialObject && MaterialObject->MaterialDomain == EMaterialDomain::MD_PostProcess && MaterialObject->BlendableLocation != EBlendableLocation::BL_ReplacingTonemapper)
+	{
+		NewContent->UserSceneTextureOutput = FScriptName(MaterialObject->UserSceneTexture);
+		NewContent->UserTextureDivisorX = MaterialObject->UserTextureDivisor.X;
+		NewContent->UserTextureDivisorY = MaterialObject->UserTextureDivisor.Y;
+	}
+
 #if ALLOW_SHADERMAP_DEBUG_DATA
 	// Store the material name for debugging purposes.
 	// Note: Material instances with static parameters will have the same FriendlyName for their shader maps!
diff --git a/Engine/Source/Runtime/Engine/Private/Materials/MaterialShared.cpp b/Engine/Source/Runtime/Engine/Private/Materials/MaterialShared.cpp
--- a/Engine/Source/Runtime/Engine/Private/Materials/MaterialShared.cpp
+++ b/Engine/Source/Runtime/Engine/Private/Materials/MaterialShared.cpp
@@ -51,6 +51,8 @@
 #include "Serialization/CompactBinary.h"
 #include "Serialization/CompactBinaryWriter.h"
 #include "Shader/PreshaderEvaluate.h"
+#include "Misc/AsciiSet.h"
+#include "PostProcess/PostProcessMaterialInputs.h"
 #endif
 #if WITH_ODSC
 #include "ODSC/ODSCManager.h"
@@ -2222,6 +2224,11 @@ int32 FMaterialResource::GetBlendableLocation() const
 	return Material->BlendableLocation;
 }
 
+int32 FMaterialResource::GetBlendablePriority() const
+{
+	return Material->BlendablePriority;
+}
+
 bool FMaterialResource::GetBlendableOutputAlpha() const
 {
 	return Material->IsPostProcessMaterialOutputingAlpha();
@@ -5622,6 +5629,102 @@ bool LoadFromCompactBinary(FCbFieldView Field, FMaterialShaderParameters& OutVal
 {
 	return OutValue.TryLoad(Field);
 }
+
+int32 FMaterialCompilationOutput::FindOrAddUserSceneTexture(FName UserSceneTexture)
+{
+	FScriptName UserSceneTextureScript(UserSceneTexture);
+	int32 FoundIndex = UserSceneTextureInputs.Find(UserSceneTextureScript);
+	if (FoundIndex != INDEX_NONE)
+	{
+		return FoundIndex + PPI_UserSceneTexture0;
+	}
+
+	if (UserSceneTextureInputs.Num() < kPostProcessMaterialInputCountMax)
+	{
+		FoundIndex = UserSceneTextureInputs.Num();
+		UserSceneTextureInputs.Add(UserSceneTextureScript);
+		return FoundIndex + PPI_UserSceneTexture0;
+	}
+	return INDEX_NONE;
+}
+
+int32 FMaterialCompilationOutput::GetNumPostProcessInputsUsed() const
+{
+	// Check how many post process inputs are used in the mask (SceneTexture nodes explicitly referencing PostProcessInput0-6)
+	int32 NumPostProcessInputs = 0;
+	for (int32 InputIndex = 0; InputIndex < kPostProcessMaterialInputCountMax; InputIndex++)
+	{
+		if (IsSceneTextureUsed((ESceneTextureId)(PPI_PostProcessInput0 + InputIndex)))
+		{
+			NumPostProcessInputs++;
+		}
+	}
+
+	// Add any UserSceneTexture inputs, which take up any remaining unused input slots
+	return NumPostProcessInputs + UserSceneTextureInputs.Num();
+}
+
+// Utility to generate HLSL defines to remap PPI_UserSceneTexture0-6 tokens generated by the HLSL emitter/translator to PPI_PostProcessInput0-6,
+// using unoccupied slots (UsedSceneTextures bits not set) in the order they are encountered.
+FString GenerateUserSceneTextureRemapHLSLDefines(FMaterialCompilationOutput& CompilationOutput)
+{
+	if (CompilationOutput.UserSceneTextureInputs.IsEmpty())
+	{
+		return TEXT("");
+	}
+
+	FString UserSceneTextureEnumRemap;
+	FString UserSceneTextureNameRemap;
+	int32 PostProcessIndex = 0;
+	for (int32 UserIndex = 0; UserIndex < CompilationOutput.UserSceneTextureInputs.Num();)
+	{
+		if (PostProcessIndex >= kPostProcessMaterialInputCountMax)
+		{
+			// If we run out of slots, go ahead and generate something that will still compile -- this should be detected and fail earlier.
+			UserSceneTextureEnumRemap.Appendf(TEXT("#define PPI_UserSceneTexture%d -1\n"), UserIndex);
+			UserIndex++;
+		}
+		else if (CompilationOutput.IsSceneTextureUsed((ESceneTextureId)(PPI_PostProcessInput0 + PostProcessIndex)))
+		{
+			// This slot is used by a SceneTexture node somewhere else in the material, skip over it
+			PostProcessIndex++;
+		}
+		else
+		{
+			UserSceneTextureEnumRemap.Appendf(TEXT("#define PPI_UserSceneTexture%d PPI_PostProcessInput%d\n"), UserIndex, PostProcessIndex);
+
+			// If the name is a token, add another define including the name.  This allows custom HLSL to easily reference a named input,
+			// without worrying about the remapping of the name.
+			constexpr FAsciiSet TokenStartChars = "abcdefghijklmnopqrstuvwxyzABCDEFGHIJKLMNOPQRSTUVWXYZ_";
+			constexpr FAsciiSet TokenContinueChars = "abcdefghijklmnopqrstuvwxyzABCDEFGHIJKLMNOPQRSTUVWXYZ0123456789_";
+
+			FString UserInputName = CompilationOutput.UserSceneTextureInputs[UserIndex].ToString();
+			if (TokenStartChars.Test(UserInputName[0]) && FAsciiSet::HasOnly(UserInputName, TokenContinueChars))
+			{
+				UserSceneTextureNameRemap.Appendf(TEXT("#define PPIUser_%s PPI_PostProcessInput%d\n"), *UserInputName, PostProcessIndex);
+			}
+
+			UserIndex++;
+			PostProcessIndex++;
+		}
+	}
+
+	return UserSceneTextureEnumRemap + TEXT("\n") + UserSceneTextureNameRemap;
+}
+
+FString SceneTextureIdToHLSLString(ESceneTextureId TexId)
+{
+	// PPI_UserSceneTexture0-6 are emitted as symbols rather than numbers, so they can be remapped later using defines generated by the function above.
+	if (TexId >= PPI_UserSceneTexture0 && TexId <= PPI_UserSceneTexture6)
+	{
+		return FString::Printf(TEXT("PPI_UserSceneTexture%d"), TexId - PPI_UserSceneTexture0);
+	}
+	else
+	{
+		return FString::Printf(TEXT("%d"), TexId);
+	}
+
+}
 #endif // WITH_EDITOR
 
 #undef LOCTEXT_NAMESPACE
\ No newline at end of file
diff --git a/Engine/Source/Runtime/Engine/Private/SceneView.cpp b/Engine/Source/Runtime/Engine/Private/SceneView.cpp
--- a/Engine/Source/Runtime/Engine/Private/SceneView.cpp
+++ b/Engine/Source/Runtime/Engine/Private/SceneView.cpp
@@ -1881,6 +1881,13 @@ void FSceneView::OverridePostProcessSettings(const FPostProcessSettings& Src, fl
 				BlendableInterface->OverrideBlendableSettings(*this, LocalWeight);
 			}
 		}
+
+#if WITH_EDITOR
+		if (Src.PreviewBlendable)
+		{
+			FinalPostProcessSettings.PreviewBlendable = Src.PreviewBlendable;
+		}
+#endif
 	}
 }
 
diff --git a/Engine/Source/Runtime/Engine/Public/MaterialCompiler.h b/Engine/Source/Runtime/Engine/Public/MaterialCompiler.h
--- a/Engine/Source/Runtime/Engine/Public/MaterialCompiler.h
+++ b/Engine/Source/Runtime/Engine/Public/MaterialCompiler.h
@@ -436,8 +436,9 @@ public:
 	virtual int32 SceneDepth(int32 Offset, int32 ViewportUV, bool bUseOffset) = 0;
 	virtual int32 SceneColor(int32 Offset, int32 ViewportUV, bool bUseOffset) = 0;
 	// @param SceneTextureId of type ESceneTextureId e.g. PPI_SubsurfaceColor
-	virtual int32 SceneTextureLookup(int32 ViewportUV, uint32 SceneTextureId, bool bFiltered) = 0;
+	virtual int32 SceneTextureLookup(int32 ViewportUV, uint32 SceneTextureId, bool bFiltered, bool bClamped) = 0;
 	virtual int32 GetSceneTextureViewSize(int32 SceneTextureId, bool InvProperty) = 0;
+	virtual int32 FindOrAddUserSceneTexture(FName UserSceneTextureName) = 0;
 	virtual int32 DBufferTextureLookup(int32 ViewportUV, uint32 DBufferTextureIndex) = 0;
 	virtual int32 PathTracingBufferTextureLookup(int32 ViewportUV, uint32 PathTracingBufferTextureIndex) = 0;
 
@@ -983,8 +984,9 @@ public:
 	virtual	int32 PixelDepth() override { return Compiler->PixelDepth(); }
 	virtual int32 SceneDepth(int32 Offset, int32 ViewportUV, bool bUseOffset) override { return Compiler->SceneDepth(Offset, ViewportUV, bUseOffset); }
 	virtual int32 SceneColor(int32 Offset, int32 ViewportUV, bool bUseOffset) override { return Compiler->SceneColor(Offset, ViewportUV, bUseOffset); }
-	virtual int32 SceneTextureLookup(int32 ViewportUV, uint32 InSceneTextureId, bool bFiltered) override { return Compiler->SceneTextureLookup(ViewportUV, InSceneTextureId, bFiltered); }
+	virtual int32 SceneTextureLookup(int32 ViewportUV, uint32 InSceneTextureId, bool bFiltered, bool bClamped) override { return Compiler->SceneTextureLookup(ViewportUV, InSceneTextureId, bFiltered, bClamped); }
 	virtual int32 GetSceneTextureViewSize(int32 SceneTextureId, bool InvProperty) override { return Compiler->GetSceneTextureViewSize(SceneTextureId, InvProperty); }
+	virtual int32 FindOrAddUserSceneTexture(FName UserSceneTextureName) override { return Compiler->FindOrAddUserSceneTexture(UserSceneTextureName); }
 	virtual int32 DBufferTextureLookup(int32 ViewportUV, uint32 DBufferTextureIndex) override { return Compiler->DBufferTextureLookup(ViewportUV, DBufferTextureIndex); }
 	virtual int32 PathTracingBufferTextureLookup(int32 ViewportUV, uint32 PathTracingBufferTextureIndex) override { return Compiler->PathTracingBufferTextureLookup(ViewportUV, PathTracingBufferTextureIndex); }
 
diff --git a/Engine/Source/Runtime/Engine/Public/MaterialHLSLTree.h b/Engine/Source/Runtime/Engine/Public/MaterialHLSLTree.h
--- a/Engine/Source/Runtime/Engine/Public/MaterialHLSLTree.h
+++ b/Engine/Source/Runtime/Engine/Public/MaterialHLSLTree.h
@@ -816,15 +816,34 @@ public:
 class FExpressionSceneTexture : public FExpression
 {
 public:
-	FExpressionSceneTexture(const FExpression* InTexCoordExpression, uint32 InSceneTextureId, bool bInFiltered)
+	FExpressionSceneTexture(const FExpression* InTexCoordExpression, uint32 InSceneTextureId, bool bInFiltered, bool bInClamped = false, FName InUserSceneTexture = NAME_None)
 		: TexCoordExpression(InTexCoordExpression)
 		, SceneTextureId(InSceneTextureId)
 		, bFiltered(bInFiltered)
+		, bClamped(bInClamped)
+		, UserSceneTexture(InUserSceneTexture)
 	{}
 
 	const FExpression* TexCoordExpression;
 	uint32 SceneTextureId;
 	bool bFiltered;
+	bool bClamped;
+	FName UserSceneTexture;
+
+	virtual bool PrepareValue(FEmitContext& Context, FEmitScope& Scope, const FRequestedType& RequestedType, FPrepareValueResult& OutResult) const override;
+	virtual void EmitValueShader(FEmitContext& Context, FEmitScope& Scope, const FRequestedType& RequestedType, FEmitValueShaderResult& OutResult) const override;
+};
+
+class FExpressionUserSceneTextureSize : public FExpression
+{
+public:
+	FExpressionUserSceneTextureSize(FName InUserSceneTexture, bool bInReciprocal)
+		: UserSceneTexture(InUserSceneTexture)
+		, bReciprocal(bInReciprocal)
+	{}
+
+	FName UserSceneTexture;
+	bool bReciprocal;
 
 	virtual bool PrepareValue(FEmitContext& Context, FEmitScope& Scope, const FRequestedType& RequestedType, FPrepareValueResult& OutResult) const override;
 	virtual void EmitValueShader(FEmitContext& Context, FEmitScope& Scope, const FRequestedType& RequestedType, FEmitValueShaderResult& OutResult) const override;
diff --git a/Engine/Source/Runtime/Engine/Public/MaterialSceneTextureId.h b/Engine/Source/Runtime/Engine/Public/MaterialSceneTextureId.h
--- a/Engine/Source/Runtime/Engine/Public/MaterialSceneTextureId.h
+++ b/Engine/Source/Runtime/Engine/Public/MaterialSceneTextureId.h
@@ -73,4 +73,13 @@ enum ESceneTextureId : int
 	PPI_WorldTangent UMETA(DisplayName = "WorldTangent"),
 	/** Material anisotropy, single channel (GBuffer) */
 	PPI_Anisotropy UMETA(DisplayName = "Anisotropy"),
+
+	/** Used internally for user scene texture inputs.  Later remapped to PPI_PostProcessInput0-6 in the shader via procedurally generated defines. */
+	PPI_UserSceneTexture0 = 32 UMETA(Hidden),
+	PPI_UserSceneTexture1 UMETA(Hidden),
+	PPI_UserSceneTexture2 UMETA(Hidden),
+	PPI_UserSceneTexture3 UMETA(Hidden),
+	PPI_UserSceneTexture4 UMETA(Hidden),
+	PPI_UserSceneTexture5 UMETA(Hidden),
+	PPI_UserSceneTexture6 UMETA(Hidden),
 };
diff --git a/Engine/Source/Runtime/Engine/Public/MaterialShared.h b/Engine/Source/Runtime/Engine/Public/MaterialShared.h
--- a/Engine/Source/Runtime/Engine/Public/MaterialShared.h
+++ b/Engine/Source/Runtime/Engine/Public/MaterialShared.h
@@ -753,7 +753,24 @@ public:
 	}
 
 	bool IsSceneTextureUsed(ESceneTextureId TexId) const { return (UsedSceneTextures & (1 << TexId)) != 0; }
-	void SetIsSceneTextureUsed(ESceneTextureId TexId) { UsedSceneTextures |= (1 << TexId); }
+
+	void SetIsSceneTextureUsed(ESceneTextureId TexId)
+	{
+		// User scene textures aren't added to used bits -- those are stored separately in UserSceneTextureInputs
+		if (TexId < PPI_UserSceneTexture0)
+		{
+			UsedSceneTextures |= (1 << TexId);
+		}
+	}
+
+#if WITH_EDITOR
+	// Returns Id for a user scene texture input, or INDEX_NONE if allocation fails.
+	int32 FindOrAddUserSceneTexture(FName UserSceneTexture);
+
+	// Returns number of post process inputs used, between PostProcessInput SceneTexture nodes and UserSceneTexture nodes.  If this exceeds
+	// kPostProcessMaterialInputCountMax, compile should be failed (number is returned rather than bool for error reporting purposes).
+	int32 GetNumPostProcessInputsUsed() const;
+#endif  // WITH_EDITOR
 
 	void SetIsDBufferTextureUsed(int32 TextureIndex) { UsedDBufferTextures |= (1 << TextureIndex); }
 	void SetIsDBufferTextureLookupUsed(bool bValue) { bUsesDBufferTextureLookup = bValue; }
@@ -795,6 +812,9 @@ public:
 
 	LAYOUT_FIELD(FUniformExpressionSet, UniformExpressionSet);
 
+	/** User scene texture inputs to this material, these will use slots not taken up by the UsedSceneTextures bitfield below */
+	LAYOUT_FIELD(TMemoryImageArray<FScriptName>, UserSceneTextureInputs);
+
 	/** Bitfield of the ESceneTextures used */
 	LAYOUT_FIELD(uint32, UsedSceneTextures);
 
@@ -1328,6 +1348,10 @@ private:
 
 	LAYOUT_FIELD(FSHAHash, ShaderContentHash);
 
+	LAYOUT_FIELD(FScriptName, UserSceneTextureOutput);
+	LAYOUT_FIELD(int32, UserTextureDivisorX);
+	LAYOUT_FIELD(int32, UserTextureDivisorY);
+
 	LAYOUT_FIELD_EDITORONLY(TMemoryImageArray<FMaterialProcessedSource>, ShaderProcessedSource);
 	LAYOUT_FIELD_EDITORONLY(FMemoryImageString, FriendlyName);
 	LAYOUT_FIELD_EDITORONLY(FMemoryImageString, DebugDescription);
@@ -1569,6 +1593,10 @@ public:
 	uint32 GetNumVirtualTextureStacks() const { return GetContent()->MaterialCompilationOutput.UniformExpressionSet.VTStacks.Num(); }
 	uint8 GetRuntimeVirtualTextureOutputAttributeMask() const { return GetContent()->MaterialCompilationOutput.RuntimeVirtualTextureOutputAttributeMask; }
 	bool UsesSceneTexture(uint32 TexId) const { return (GetContent()->MaterialCompilationOutput.UsedSceneTextures & (1ull << TexId)) != 0; }
+	TConstArrayView<FScriptName> GetUserSceneTextureInputs() const { return GetContent()->MaterialCompilationOutput.UserSceneTextureInputs; }
+	FScriptName GetUserSceneTextureOutput() const { return GetContent()->UserSceneTextureOutput; }
+	FIntPoint GetUserTextureDivisor() const { return FIntPoint(FMath::Max(GetContent()->UserTextureDivisorX, 1), FMath::Max(GetContent()->UserTextureDivisorY, 1)); }
+
 	bool UsesPathTracingBufferTexture(uint32 TexId) const { return (GetContent()->MaterialCompilationOutput.UsedPathTracingBufferTextures & (1ull << TexId)) != 0;}
 
 	bool IsValidForRendering(bool bFailOnInvalid = false) const
@@ -2147,6 +2175,7 @@ public:
 	virtual bool IsDefaultMaterial() const { return false; };
 	virtual int32 GetNumCustomizedUVs() const { return 0; }
 	virtual int32 GetBlendableLocation() const { return 0; }
+	virtual int32 GetBlendablePriority() const { return 0; }
 	virtual bool GetBlendableOutputAlpha() const { return false; }
 	virtual bool IsStencilTestEnabled() const { return false; }
 	virtual uint32 GetStencilRefValue() const { return 0; }
@@ -2894,6 +2923,7 @@ public:
 	ENGINE_API virtual bool IsDefaultMaterial() const override;
 	ENGINE_API virtual int32 GetNumCustomizedUVs() const override;
 	ENGINE_API virtual int32 GetBlendableLocation() const override;
+	ENGINE_API virtual int32 GetBlendablePriority() const override;
 	ENGINE_API virtual bool GetBlendableOutputAlpha() const override;
 	ENGINE_API virtual bool IsStencilTestEnabled() const override;
 	ENGINE_API virtual uint32 GetStencilRefValue() const override;
diff --git a/Engine/Source/Runtime/Engine/Public/Materials/Material.h b/Engine/Source/Runtime/Engine/Public/Materials/Material.h
--- a/Engine/Source/Runtime/Engine/Public/Materials/Material.h
+++ b/Engine/Source/Runtime/Engine/Public/Materials/Material.h
@@ -988,6 +988,14 @@ public:
 	UPROPERTY(EditAnywhere, BlueprintReadWrite, Category = PostProcessMaterial, meta = (DisplayName = "Output Alpha"))
 	uint8 BlendableOutputAlpha : 1;
 
+	/** Specify a user generated scene texture as output, overriding the default output implied by "Blendable Location", only used if domain is PostProcess */
+	UPROPERTY(EditAnywhere, BlueprintReadWrite, Category = PostProcessMaterial, meta = (DisplayName = "User Scene Texture"))
+	FName UserSceneTexture;
+
+	/** A divisor for the resolution of the User Scene Texture above, allowing an intermediate with reduced resolution */
+	UPROPERTY(EditAnywhere, BlueprintReadWrite, Category = PostProcessMaterial, meta = (DisplayName = "User Texture Divisor"))
+	FIntPoint UserTextureDivisor = FIntPoint(0, 0);
+
 	/**
 	* Indicates that the material and its instances can be used with neural network engine.
 	* This will result in the shaders required to support neural network engine being compiled which will increase shader compile time and memory usage.
diff --git a/Engine/Source/Runtime/Engine/Public/Materials/MaterialExpressionUserSceneTexture.h b/Engine/Source/Runtime/Engine/Public/Materials/MaterialExpressionUserSceneTexture.h
new file mode 100644
--- /dev/null
+++ b/Engine/Source/Runtime/Engine/Public/Materials/MaterialExpressionUserSceneTexture.h
@@ -0,0 +1,41 @@
+// Copyright Epic Games, Inc. All Rights Reserved.
+
+
+#pragma once
+
+#include "CoreMinimal.h"
+#include "UObject/ObjectMacros.h"
+#include "MaterialExpressionIO.h"
+#include "Materials/MaterialExpression.h"
+#include "MaterialExpressionUserSceneTexture.generated.h"
+
+UCLASS(MinimalAPI, collapsecategories, hidecategories=Object)
+class UMaterialExpressionUserSceneTexture : public UMaterialExpression
+{
+	GENERATED_UCLASS_BODY()
+
+	/** UV in 0..1 range */
+	UPROPERTY(meta = (RequiredInput = "false", ToolTip = "Ignored if not specified"))
+	FExpressionInput Coordinates;
+
+	/** User Scene Texture (screen space texture with a user specified name, written by a previous PostProcess shader) to make a lookup into */
+	UPROPERTY(EditAnywhere, Category=UMaterialExpressionUserSceneTexture, meta=(DisplayName = "User Scene Texture", ShowAsInputPin = "Advanced"))
+	FName UserSceneTexture;
+	
+	/** Whether to use point sampled texture lookup (default) or using [bi-linear] filtered (can be slower, avoid faceted lock with distortions) */
+	UPROPERTY(EditAnywhere, Category=UMaterialExpressionUserSceneTexture, meta=(DisplayName = "Filtered", ShowAsInputPin = "Advanced"))
+	bool bFiltered;
+
+	/** Whether to clamp the texture lookup.  Necessary when sampling a UserSceneTexture at reduced resolution with filtering, to avoid blending out of bounds pixels. */
+	UPROPERTY(EditAnywhere, Category=UMaterialExpressionUserSceneTexture, meta=(DisplayName = "Clamped", ShowAsInputPin = "Advanced"))
+	bool bClamped;
+
+	//~ Begin UMaterialExpression Interface
+#if WITH_EDITOR
+	virtual int32 Compile(class FMaterialCompiler* Compiler, int32 OutputIndex) override;
+	virtual void GetCaption(TArray<FString>& OutCaptions) const override;
+	virtual bool GenerateHLSLExpression(FMaterialHLSLGenerator& Generator, UE::HLSLTree::FScope& Scope, int32 OutputIndex, UE::HLSLTree::FExpression const*& OutExpression) const override;
+#endif
+	//~ End UMaterialExpression Interface
+};
+
diff --git a/Engine/Source/Runtime/RenderCore/Private/CommonRenderResources.cpp b/Engine/Source/Runtime/RenderCore/Private/CommonRenderResources.cpp
--- a/Engine/Source/Runtime/RenderCore/Private/CommonRenderResources.cpp
+++ b/Engine/Source/Runtime/RenderCore/Private/CommonRenderResources.cpp
@@ -19,6 +19,7 @@ IMPLEMENT_GLOBAL_SHADER(FScreenVertexShaderVS, "/Engine/Private/Tools/Fullscreen
 IMPLEMENT_GLOBAL_SHADER(FInstancedScreenVertexShaderVS, "/Engine/Private/Tools/FullscreenVertexShader.usf", "MainVS", SF_Vertex);
 IMPLEMENT_GLOBAL_SHADER(FMobileMultiViewVertexShaderVS, "/Engine/Private/Tools/FullscreenVertexShader.usf", "MainVS", SF_Vertex);
 IMPLEMENT_GLOBAL_SHADER(FCopyRectPS, "/Engine/Private/ScreenPass.usf", "CopyRectPS", SF_Pixel);
+IMPLEMENT_GLOBAL_SHADER(FCopyRectSrvPS, "/Engine/Private/ScreenPass.usf", "CopyRectPS", SF_Pixel);
 IMPLEMENT_GLOBAL_SHADER(FImagePreTransformVS, "/Engine/Private/Tools/FullscreenVertexShader.usf", "MainForPreTransform", SF_Vertex);
 
 bool FInstancedScreenVertexShaderVS::ShouldCompilePermutation(const FGlobalShaderPermutationParameters& Parameters)
diff --git a/Engine/Source/Runtime/RenderCore/Private/RenderGraphResources.cpp b/Engine/Source/Runtime/RenderCore/Private/RenderGraphResources.cpp
--- a/Engine/Source/Runtime/RenderCore/Private/RenderGraphResources.cpp
+++ b/Engine/Source/Runtime/RenderCore/Private/RenderGraphResources.cpp
@@ -46,6 +46,9 @@ FRDGViewableResource::FRDGViewableResource(const TCHAR* InName, const ERDGViewab
 	, bQueuedForUpload(0)
 	, bCollectForAllocate(1)
 	, bQueuedForReservedCommit(0)
+#if !(UE_BUILD_SHIPPING || UE_BUILD_TEST)
+	, bHeapAllocatedDebugName(0)
+#endif
 	, TransientExtractionHint(ETransientExtractionHint::None)
 	, ReferenceCount(IsImmediateMode() ? 1 : 0)
 {
diff --git a/Engine/Source/Runtime/RenderCore/Private/VisualizeTexture.cpp b/Engine/Source/Runtime/RenderCore/Private/VisualizeTexture.cpp
--- a/Engine/Source/Runtime/RenderCore/Private/VisualizeTexture.cpp
+++ b/Engine/Source/Runtime/RenderCore/Private/VisualizeTexture.cpp
@@ -767,7 +767,7 @@ FRDGTextureRef FVisualizeTexture::AddVisualizeTexturePass(
 	OutputExtent.X = FMath::Max(OutputExtent.X, 1);
 	OutputExtent.Y = FMath::Max(OutputExtent.Y, 1);
 
-	FRDGTextureRef OutputTexture = GraphBuilder.CreateTexture(FRDGTextureDesc::Create2D(OutputExtent, PF_B8G8R8A8, FClearValueBinding(FLinearColor(1, 1, 0, 1)), TexCreate_RenderTargetable | TexCreate_ShaderResource), TEXT("VisualizeTexture"));
+	FRDGTextureRef OutputTexture = GraphBuilder.CreateTexture(FRDGTextureDesc::Create2D(OutputExtent, PF_B8G8R8A8, FClearValueBinding(FLinearColor(1, 1, 0, 1)), TexCreate_RenderTargetable | TexCreate_ShaderResource), InputTexture->Name);
 
 	{
 		const EVisualisePSType VisualizeType = GetVisualizePSType(InputDesc);
diff --git a/Engine/Source/Runtime/RenderCore/Public/CommonRenderResources.h b/Engine/Source/Runtime/RenderCore/Public/CommonRenderResources.h
--- a/Engine/Source/Runtime/RenderCore/Public/CommonRenderResources.h
+++ b/Engine/Source/Runtime/RenderCore/Public/CommonRenderResources.h
@@ -167,6 +167,19 @@ class FCopyRectPS : public FGlobalShader
 	END_SHADER_PARAMETER_STRUCT()
 };
 
+/** Pixel shader to copy pixels from src to dst performing a format change that works on all platforms -- variation that accepts SRV */
+class FCopyRectSrvPS : public FGlobalShader
+{
+	DECLARE_EXPORTED_GLOBAL_SHADER(FCopyRectSrvPS, RENDERCORE_API);
+	SHADER_USE_PARAMETER_STRUCT(FCopyRectSrvPS, FGlobalShader);
+
+	BEGIN_SHADER_PARAMETER_STRUCT(FParameters, )
+		SHADER_PARAMETER_RDG_TEXTURE_SRV(Texture2D, InputTexture)
+		SHADER_PARAMETER_SAMPLER(SamplerState, InputSampler)
+		RENDER_TARGET_BINDING_SLOTS()
+		END_SHADER_PARAMETER_STRUCT()
+};
+
 /** Vertex shader to perform a screen rotation for Vulkan pre-rotation on mobile. */
 class FImagePreTransformVS : public FGlobalShader
 {
diff --git a/Engine/Source/Runtime/RenderCore/Public/RenderGraphResources.h b/Engine/Source/Runtime/RenderCore/Public/RenderGraphResources.h
--- a/Engine/Source/Runtime/RenderCore/Public/RenderGraphResources.h
+++ b/Engine/Source/Runtime/RenderCore/Public/RenderGraphResources.h
@@ -279,6 +279,16 @@ class FRDGViewableResource
 	: public FRDGResource
 {
 public:
+#if !(UE_BUILD_SHIPPING || UE_BUILD_TEST)
+	virtual ~FRDGViewableResource()
+	{
+		if (bHeapAllocatedDebugName)
+		{
+			delete[] Name;
+		}
+	}
+#endif
+
 	/** The type of this resource; useful for casting between types. */
 	const ERDGViewableResourceType Type;
 
@@ -314,6 +324,22 @@ public:
 #endif
 	}
 
+#if !(UE_BUILD_SHIPPING || UE_BUILD_TEST)
+	inline void SetDebugNameIsHeapAllocated()
+	{
+		bHeapAllocatedDebugName = 1;
+	}
+#endif
+
+	inline bool IsDebugNameHeapAllocated() const
+	{
+#if !(UE_BUILD_SHIPPING || UE_BUILD_TEST)
+		return bHeapAllocatedDebugName;
+#else
+		return false;
+#endif
+	}
+
 protected:
 	RENDERCORE_API FRDGViewableResource(const TCHAR* InName, ERDGViewableResourceType InType, bool bSkipTracking, bool bImmediateFirstBarrier);
 
@@ -389,6 +415,11 @@ protected:
 	/** If true, the reserved resource is having tiles committed. */
 	uint8 bQueuedForReservedCommit : 1;
 
+#if !(UE_BUILD_SHIPPING || UE_BUILD_TEST)
+	/** If true, the debug name passed to the constructor was heap allocated, and needs to be freed in the destructor. */
+	uint8 bHeapAllocatedDebugName : 1;
+#endif
+
 	/** Whether this resource is allowed to be both transient and extracted. */
 	ETransientExtractionHint TransientExtractionHint;
 
diff --git a/Engine/Source/Runtime/Renderer/Internal/SceneTextures.h b/Engine/Source/Runtime/Renderer/Internal/SceneTextures.h
--- a/Engine/Source/Runtime/Renderer/Internal/SceneTextures.h
+++ b/Engine/Source/Runtime/Renderer/Internal/SceneTextures.h
@@ -7,6 +7,7 @@
 #include "CustomDepthRendering.h"
 #include "SceneRenderTargetParameters.h"
 #include "GBufferInfo.h"
+#include "ScreenPass.h"
 
 struct FSceneTextures;
 class FViewInfo;
@@ -15,6 +16,36 @@ class FViewFamilyInfo;
 /** Initializes a scene textures config instance from the view family. */
 extern RENDERER_API void InitializeSceneTexturesConfig(FSceneTexturesConfig& Config, const FSceneViewFamily& ViewFamily);
 
+struct FTransientUserSceneTexture
+{
+	FRDGTextureRef Texture{};
+	FIntPoint ResolutionDivisor;
+	uint16 AllocationOrder;				// Order in which item of a given Name was allocated, mainly for differentiating items in texture visualizer
+	bool bUsed;							// Tracks whether output was used as an input, for debugging
+	uint32 ViewMask;					// Tracks which views this texture has been rendered in, so we can detect the first render in a given view
+};
+
+#if !(UE_BUILD_SHIPPING)
+enum class EUserSceneTextureEvent
+{
+	MissingInput,
+	CollidingInput,		// Input matches the output, and has been unbound as a result
+	FoundInput,
+	Output,
+	Pass				// Marker for the end of events for a given material pass with UserSceneTexture inputs or outputs
+};
+
+struct FUserSceneTextureEventData
+{
+	EUserSceneTextureEvent Event;
+	FName Name;
+	uint16 AllocationOrder;
+	uint16 ViewIndex;					// Necessary to differentiate events from multiple views in split screen
+	const FMaterial* Material;
+	FIntPoint RectSize;					// Only filled in for EUserSceneTextureEvent::Output
+};
+#endif
+
 /** RDG struct containing the minimal set of scene textures common across all rendering configurations. */
 struct FMinimalSceneTextures
 {
@@ -53,7 +84,23 @@ struct FMinimalSceneTextures
 	// Textures containing depth / stencil information from the custom depth pass.
 	FCustomDepthTextures CustomDepth{};
 
+	// Dynamically allocated user scene textures, stored by name.  An array of textures per name is used, as it's possible the
+	// same name is allocated with different resolution divisors.  The most recently written texture resolution with a given name
+	// will be used as an input to other materials, by swapping to the front of the array.
+	mutable TMap<FName, TArray<FTransientUserSceneTexture>> UserSceneTextures;
+
+#if !(UE_BUILD_SHIPPING)
+	mutable TArray<FUserSceneTextureEventData> UserSceneTextureEvents;
+#endif
+
 	RENDERER_API FSceneTextureShaderParameters GetSceneTextureShaderParameters(ERHIFeatureLevel::Type FeatureLevel) const;
+
+	FRDGTextureRef FindOrAddUserSceneTexture(FRDGBuilder& GraphBuilder, int32 ViewIndex, FName Name, FIntPoint ResolutionDivisor, bool& bOutFirstRender, const FMaterial* Material, const FIntRect& OutputRect) const;
+	FScreenPassTextureSlice GetUserSceneTexture(FRDGBuilder& GraphBuilder, const FViewInfo& View, int32 ViewIndex, FName Name, const FMaterial* Material) const;
+
+#if !(UE_BUILD_SHIPPING)
+	const FTransientUserSceneTexture* FindUserSceneTextureByEvent(const FUserSceneTextureEventData& Event) const;
+#endif
 };
 
 /** RDG struct containing the complete set of scene textures for the deferred or mobile renderers. */
diff --git a/Engine/Source/Runtime/Renderer/Private/MobileShadingRenderer.cpp b/Engine/Source/Runtime/Renderer/Private/MobileShadingRenderer.cpp
--- a/Engine/Source/Runtime/Renderer/Private/MobileShadingRenderer.cpp
+++ b/Engine/Source/Runtime/Renderer/Private/MobileShadingRenderer.cpp
@@ -1430,7 +1430,7 @@ void FMobileSceneRenderer::Render(FRDGBuilder& GraphBuilder)
 							}
 							else
 							{ 
-								AddMobilePostProcessingPasses(GraphBuilder, Scene, Views[ViewIndex], GetSceneUniforms(), PostProcessingInputs, InstanceCullingManager);
+								AddMobilePostProcessingPasses(GraphBuilder, Scene, Views[ViewIndex], ViewIndex, GetSceneUniforms(), PostProcessingInputs, InstanceCullingManager);
 							}
 						}
 					}
diff --git a/Engine/Source/Runtime/Renderer/Private/PostProcess/PostProcessMaterial.cpp b/Engine/Source/Runtime/Renderer/Private/PostProcess/PostProcessMaterial.cpp
--- a/Engine/Source/Runtime/Renderer/Private/PostProcess/PostProcessMaterial.cpp
+++ b/Engine/Source/Runtime/Renderer/Private/PostProcess/PostProcessMaterial.cpp
@@ -19,6 +19,7 @@
 #include "SceneRendering.h"
 #include "ClearQuad.h"
 #include "Materials/MaterialExpressionSceneTexture.h"
+#include "Materials/MaterialExpressionUserSceneTexture.h"
 #include "PipelineStateCache.h"
 #include "PostProcess/PostProcessing.h"
 #include "PostProcess/PostProcessMobile.h"
@@ -458,12 +459,25 @@ FPostProcessMaterialParameters* GetPostProcessMaterialParameters(
 
 	FRHISamplerState* PointClampSampler = TStaticSamplerState<SF_Point, AM_Clamp, AM_Clamp, AM_Clamp>::GetRHI();
 
+	int32 NumUserSceneTextures = MaterialShaderMap->GetUserSceneTextureInputs().Num();
 	for (uint32 InputIndex = 0; InputIndex < kPostProcessMaterialInputCountMax; ++InputIndex)
 	{
 		FScreenPassTextureSlice Input = Inputs.GetInput((EPostProcessMaterialInput)InputIndex);
 
+		bool bIsUsed = MaterialShaderMap->UsesSceneTexture(PPI_PostProcessInput0 + InputIndex);
+
+		// User scene textures consume any consecutive slots not used by PPI_PostProcessInput0-6
+		if (!bIsUsed)
+		{
+			if (NumUserSceneTextures > 0)
+			{
+				NumUserSceneTextures--;
+				bIsUsed = true;
+			}
+		}
+
 		// Need to provide valid textures for when shader compilation doesn't cull unused parameters.
-		if (!Input.IsValid() || !MaterialShaderMap->UsesSceneTexture(PPI_PostProcessInput0 + InputIndex))
+		if (!Input.IsValid() || !bIsUsed)
 		{
 			Input = FScreenPassTextureSlice::CreateFromScreenPassTexture(GraphBuilder, BlackDummy);
 		}
@@ -691,7 +705,7 @@ FScreenPassTexture AddPostProcessMaterialPass(
 	GetMaterialInfo(MaterialInterface, FeatureLevel, Inputs, Material, MaterialRenderProxy, MaterialShaderMap, VertexShader, PixelShader);
 
 	EBlendableLocation BlendableLocation = EBlendableLocation(Material->GetBlendableLocation());
-	const FScreenPassTexture SceneColor = FScreenPassTexture::CopyFromSlice(GraphBuilder, Inputs.GetInput(BlendableLocation == BL_TranslucencyAfterDOF ? EPostProcessMaterialInput::SeparateTranslucency : EPostProcessMaterialInput::SceneColor));
+	const FScreenPassTextureSlice SceneColorOutput = Inputs.GetSceneColorOutput(BlendableLocation);
 
 	check(VertexShader.IsValid());
 	check(PixelShader.IsValid());
@@ -729,14 +743,14 @@ FScreenPassTexture AddPostProcessMaterialPass(
 									 BlendState != DefaultBlendState ||
 									 bMayDiscard;
 
-	// We only prime color on the output texture if we are using fixed function Blend / Depth-Stencil,
-	// or we need to retain previously rendered views.
-	const bool bPrimeOutputColor = bCompositeWithInput || !View.IsFirstInFamily();
+	// We only prime color on the output texture if we are using fixed function Blend / Depth-Stencil, or we need to
+	// retain previously rendered views.  UserSceneTexture does its own output priming logic where required.
+	const bool bPrimeOutputColor = (bCompositeWithInput || !View.IsFirstInFamily()) && !Inputs.bUserSceneTextureOutput;
 
 	// Inputs.OverrideOutput is used to force drawing directly to the backbuffer. OpenGL doesn't support using the backbuffer color target with a custom depth/stencil
 	// buffer, so in that case we must draw to an intermediate target and copy to the backbuffer at the end. Ideally, we would test if Inputs.OverrideOutput.Texture
-	// is actually the backbuffer, but it's not worth doing all the plumbing and increasing the RHI surface area just for this hack.
-	const bool bBackbufferWithDepthStencil = (DepthStencilTexture != nullptr && !GRHISupportsBackBufferWithCustomDepthStencil && Inputs.OverrideOutput.IsValid());
+	// is actually the backbuffer, but it's not worth doing all the plumbing and increasing the RHI surface area just for this hack.  UserSceneTexture is never a backbuffer.
+	const bool bBackbufferWithDepthStencil = (DepthStencilTexture != nullptr && !GRHISupportsBackBufferWithCustomDepthStencil && Inputs.OverrideOutput.IsValid() && !Inputs.bUserSceneTextureOutput);
 
 	// We need to decode the target color for blending material, force it rendering to an intermediate render target and decode the color.
 	const bool bCompositeWithInputAndDecode = Inputs.bMetalMSAAHDRDecode && bCompositeWithInput;
@@ -749,8 +763,10 @@ FScreenPassTexture AddPostProcessMaterialPass(
 	// This is only necessary to do if we're going to be priming content from the render target since it avoids
 	// the copy. Otherwise, we just allocate a new render target.
 	const bool bValidShaderPlatform = (GMaxRHIShaderPlatform != SP_PCD3D_ES3_1);
-	if (!Output.IsValid() && !MaterialShaderMap->UsesSceneTexture(PPI_PostProcessInput0) && bPrimeOutputColor && !bForceIntermediateTarget && Inputs.bAllowSceneColorInputAsOutput && bValidShaderPlatform)
+	if (!Output.IsValid() && !MaterialShaderMap->UsesSceneTexture(PPI_PostProcessInput0) && bPrimeOutputColor && !bForceIntermediateTarget && Inputs.bAllowSceneColorInputAsOutput && bValidShaderPlatform && !Inputs.bUserSceneTextureOutput)
 	{
+		FScreenPassTexture SceneColor = FScreenPassTexture::CopyFromSlice(GraphBuilder, SceneColorOutput);
+
 		Output = FScreenPassRenderTarget(SceneColor, ERenderTargetLoadAction::ELoad);
 	}
 	else
@@ -758,7 +774,10 @@ FScreenPassTexture AddPostProcessMaterialPass(
 		// Allocate new transient output texture if none exists.
 		if (!Output.IsValid() || bForceIntermediateTarget)
 		{
-			FRDGTextureDesc OutputDesc = SceneColor.Texture->Desc;
+			FRDGTextureDesc OutputDesc = SceneColorOutput.TextureSRV->Desc.Texture->Desc;
+			OutputDesc.Dimension = ETextureDimension::Texture2D;
+			OutputDesc.ArraySize = 1;
+
 			OutputDesc.Reset();
 			if (Inputs.OutputFormat != PF_Unknown)
 			{
@@ -768,11 +787,13 @@ FScreenPassTexture AddPostProcessMaterialPass(
 			OutputDesc.Flags &= (~ETextureCreateFlags::FastVRAM);
 			OutputDesc.Flags |= GFastVRamConfig.PostProcessMaterial;
 
-			Output = FScreenPassRenderTarget(GraphBuilder.CreateTexture(OutputDesc, TEXT("PostProcessMaterial")), SceneColor.ViewRect, View.GetOverwriteLoadAction());
+			Output = FScreenPassRenderTarget(GraphBuilder.CreateTexture(OutputDesc, TEXT("PostProcessMaterial")), SceneColorOutput.ViewRect, View.GetOverwriteLoadAction());
 		}
 
 		if (bPrimeOutputColor || bForceIntermediateTarget)
 		{
+			FScreenPassTexture SceneColor = FScreenPassTexture::CopyFromSlice(GraphBuilder, SceneColorOutput);
+
 			// Copy existing contents to new output and use load-action to preserve untouched pixels.
 			if (Inputs.bMetalMSAAHDRDecode)
 			{
@@ -784,9 +805,15 @@ FScreenPassTexture AddPostProcessMaterialPass(
 			}
 			Output.LoadAction = ERenderTargetLoadAction::ELoad;
 		}
+
+		// If this is the first render to a UserSceneTexture which requires compositing, we copy the previous output as a starting point.
+		if (bCompositeWithInput && Inputs.bUserSceneTextureOutput && Inputs.bUserSceneTextureFirstRender)
+		{
+			AddDrawTexturePass(GraphBuilder, View, SceneColorOutput, Output);
+		}
 	}
 
-	const FScreenPassTextureViewport SceneColorViewport(SceneColor);
+	const FScreenPassTextureViewport SceneColorViewport(SceneColorOutput);
 	const FScreenPassTextureViewport OutputViewport(Output);
 
 	RDG_EVENT_SCOPE(GraphBuilder, "PostProcessMaterial %dx%d Material=%s", SceneColorViewport.Rect.Width(), SceneColorViewport.Rect.Height(), *Material->GetAssetName());
@@ -878,17 +905,18 @@ FScreenPassTexture AddPostProcessMaterialPass(
 	else
 	{
 		// When skipping the pass, we still need to output a valid FScreenPassRenderTarget
-		Output = FScreenPassRenderTarget(SceneColor, ERenderTargetLoadAction::ENoAction);
-
-		// If there is override output, we need to output to that
 		if (Inputs.OverrideOutput.IsValid())
 		{
-			const FIntPoint SrcPoint = View.ViewRect.Min;
-			const FIntPoint DstPoint = OutputViewport.Rect.Min;
-			const FIntPoint Size = OutputViewport.Rect.Max - DstPoint;
-			AddDrawTexturePass(GraphBuilder, View, Output.Texture, Inputs.OverrideOutput.Texture, SrcPoint, DstPoint, Size);
+			// If there is an override output, we can copy directly to that from the scene color slice.
+			AddDrawTexturePass(GraphBuilder, View, SceneColorOutput, Inputs.OverrideOutput);
 			Output = Inputs.OverrideOutput;
 		}
+		else
+		{
+			// Otherwise, we need to convert output to a single slice before returning
+			FScreenPassTexture SceneColor = FScreenPassTexture::CopyFromSlice(GraphBuilder, SceneColorOutput);
+			Output = FScreenPassRenderTarget(SceneColor, ERenderTargetLoadAction::ENoAction);
+		}
 	}
 
 	return MoveTemp(Output);
@@ -968,7 +996,8 @@ FPostProcessMaterialChain GetPostProcessMaterialChain(const FViewInfo& View, EBl
 		return {};
 	}
 
-	Algo::Sort(Nodes, FPostProcessMaterialNode::FCompare());
+	// Use stable sort, so if nodes with the same priority are explicitly ordered in the post process volume, they stay in that order
+	Algo::StableSort(Nodes, FPostProcessMaterialNode::FCompare());
 
 	FPostProcessMaterialChain OutputChain;
 	OutputChain.Reserve(Nodes.Num());
@@ -981,9 +1010,62 @@ FPostProcessMaterialChain GetPostProcessMaterialChain(const FViewInfo& View, EBl
 	return OutputChain;
 }
 
+static void RemoveCollidingUserSceneTextureInputs(FPostProcessMaterialInputs& PassInputs, const FSceneTextures& SceneTextures)
+{
+	for (int32 InputIndex = 0; InputIndex < kPostProcessMaterialInputCountMax; InputIndex++)
+	{
+		if (PassInputs.bUserSceneTexturesSet[InputIndex])
+		{
+			if (PassInputs.UserSceneTextures[InputIndex].TextureSRV && PassInputs.UserSceneTextures[InputIndex].TextureSRV->GetParent() == PassInputs.OverrideOutput.Texture)
+			{
+				// Zero out the input, and label it as an error in the event stream if necessary
+
+#if !(UE_BUILD_SHIPPING)
+				// Get the name from the resource, and strip off the prefix
+				const TCHAR* InputNameStr = PassInputs.UserSceneTextures[InputIndex].TextureSRV->Desc.Texture->Name;
+				if (FCString::Strstr(InputNameStr, TEXT("UST.")) == InputNameStr)
+				{
+					InputNameStr += 4;
+				}
+
+				// Iterate over the events, looking for that name
+				for (int32 EventIndex = SceneTextures.UserSceneTextureEvents.Num() - 1; EventIndex >= 0; EventIndex--)
+				{
+					// Stop if we reach a pass event marker
+					if (SceneTextures.UserSceneTextureEvents[EventIndex].Event == EUserSceneTextureEvent::Pass)
+					{
+						break;
+					}
+
+					if (SceneTextures.UserSceneTextureEvents[EventIndex].Event == EUserSceneTextureEvent::FoundInput)
+					{
+						FString EventName = SceneTextures.UserSceneTextureEvents[EventIndex].Name.ToString();
+						const TCHAR* EventNameStr = *EventName;
+
+						// Resource may have a numeric allocation order suffix as well, like [1] -- check if the front of the string matches
+						if (FCString::Strstr(InputNameStr, EventNameStr) == InputNameStr)
+						{
+							// Then check if that's the end of the string or an open bracket suffix
+							int32 EventNameLen = FCString::Strlen(EventNameStr);
+							if (InputNameStr[EventNameLen] == 0 || InputNameStr[EventNameLen] == TEXT('['))
+							{
+								SceneTextures.UserSceneTextureEvents[EventIndex].Event = EUserSceneTextureEvent::CollidingInput;
+							}
+						}
+					}
+				}
+#endif
+
+				PassInputs.UserSceneTextures[InputIndex] = FScreenPassTextureSlice();
+			}
+		}
+	}
+}
+
 FScreenPassTexture AddPostProcessMaterialChain(
 	FRDGBuilder& GraphBuilder,
 	const FViewInfo& View,
+	int32 ViewIndex,
 	const FPostProcessMaterialInputs& InputsTemplate,
 	const FPostProcessMaterialChain& Materials,
 	EPostProcessMaterialInput MaterialInput)
@@ -991,34 +1073,122 @@ FScreenPassTexture AddPostProcessMaterialChain(
 	FScreenPassTextureSlice CurrentInput = InputsTemplate.GetInput(MaterialInput);
 	FScreenPassTexture Outputs;
 
+	// Get last material that writes to the output (ignoring materials that write to UserSceneTextures)
+	const UMaterialInterface* LastOutputWrite = nullptr;
+	for (int32 MaterialIndex = Materials.Num() - 1; MaterialIndex >= 0; MaterialIndex--)
+	{
+		const FMaterial* Material = Materials[MaterialIndex]->GetRenderProxy()->GetMaterialNoFallback(View.FeatureLevel);
+		if (!Material || Material->GetRenderingThreadShaderMap()->GetUserSceneTextureOutput().IsNone())
+		{
+			// Doesn't write to a UserSceneTexture, so it writes to the output
+			LastOutputWrite = Materials[MaterialIndex];
+			break;
+		}
+	}
+
 	bool bFirstMaterialInChain = true;
 	for (const UMaterialInterface* MaterialInterface : Materials)
 	{
 		FPostProcessMaterialInputs Inputs = InputsTemplate;
 		Inputs.SetInput(MaterialInput, CurrentInput);
+
+		// Get UserSceneTexture inputs and output from material if present
+		const FMaterial* Material = MaterialInterface->GetRenderProxy()->GetMaterialNoFallback(View.FeatureLevel);
+		FName UserSceneTextureOutput(NAME_None);
+		FIntPoint UserTextureDivisor(0, 0);
+		int32 UserSceneTextureInputNum = 0;
+
+		if (Material)
+		{
+			const FMaterialShaderMap* MaterialShaderMap = Material->GetRenderingThreadShaderMap();
+
+			TConstArrayView<FScriptName> UserSceneTextureInputs = MaterialShaderMap->GetUserSceneTextureInputs();
+			if (!UserSceneTextureInputs.IsEmpty())
+			{
+				UserSceneTextureInputNum = UserSceneTextureInputs.Num();
+
+				const FSceneTextures& SceneTextures = View.GetSceneTextures();
+
+				int32 PostProcessIndex = 0;
+				for (int32 UserIndex = 0; UserIndex < UserSceneTextureInputs.Num();)
+				{
+					check(PostProcessIndex < kPostProcessMaterialInputCountMax);
+				
+					// Skip over this slot if it's used by a SceneTexture node
+					if (!MaterialShaderMap->UsesSceneTexture(PPI_PostProcessInput0 + PostProcessIndex))
+					{
+						// Not used as a SceneTexture, so it's used by the next UserSceneTexture
+						Inputs.SetUserSceneTextureInput((EPostProcessMaterialInput)PostProcessIndex, SceneTextures.GetUserSceneTexture(GraphBuilder, View, ViewIndex, FName(UserSceneTextureInputs[UserIndex]), Material));
+						UserIndex++;
+					}
+					PostProcessIndex++;
+				}
+			}
+
+#if WITH_EDITOR
+			// If this blendable is being previewed, don't write to the UserSceneTexture -- instead it will write to SceneColor
+			if (Material->GetMaterialInterface() != View.FinalPostProcessSettings.PreviewBlendable)
+#endif
+			{
+				UserSceneTextureOutput = FName(MaterialShaderMap->GetUserSceneTextureOutput());
+				UserTextureDivisor = MaterialShaderMap->GetUserTextureDivisor();
+			}
+		}
 		
 		// Only the first material in the chain needs to decode the input color
 		Inputs.bMetalMSAAHDRDecode = Inputs.bMetalMSAAHDRDecode && bFirstMaterialInChain;
 		bFirstMaterialInChain = false;
 
-		// Certain inputs are only respected by the final post process material in the chain.
-		if (MaterialInterface != Materials.Last())
+		if (!UserSceneTextureOutput.IsNone())
 		{
-			Inputs.OverrideOutput = FScreenPassRenderTarget();
+			// Writing to UserSceneTexture, don't set Outputs or CurrentInput, as this is writing to a disjoint texture that's not part of the chain
+			FIntRect OutputRect = GetDownscaledViewRect(View.UnconstrainedViewRect, View.GetFamilyViewRect().Max, UserTextureDivisor);
+			FRDGTextureRef UserOutput = ((FViewFamilyInfo*)View.Family)->GetSceneTextures().FindOrAddUserSceneTexture(GraphBuilder, ViewIndex, UserSceneTextureOutput, UserTextureDivisor, Inputs.bUserSceneTextureFirstRender, Material, OutputRect);
+			Inputs.OverrideOutput = FScreenPassRenderTarget(UserOutput, OutputRect, ERenderTargetLoadAction::ELoad);
+			Inputs.bUserSceneTextureOutput = true;
+
+			RemoveCollidingUserSceneTextureInputs(Inputs, View.GetSceneTextures());
+
+			AddPostProcessMaterialPass(GraphBuilder, View, Inputs, MaterialInterface);
 		}
+		else
+		{
+			// Certain inputs are only respected by the final post process material in the chain, that writes to the Output
+			if (MaterialInterface != LastOutputWrite)
+			{
+				Inputs.OverrideOutput = FScreenPassRenderTarget();
+			}
 
-		Outputs = AddPostProcessMaterialPass(GraphBuilder, View, Inputs, MaterialInterface);
+			Outputs = AddPostProcessMaterialPass(GraphBuilder, View, Inputs, MaterialInterface);
 
-		// Don't create the CurrentInput out of Outputs of the last material as this could possibly be the back buffer for AfterTonemap post process material
-		if (MaterialInterface != Materials.Last())
+			// Don't create the CurrentInput out of Outputs of the last material as this could possibly be the back buffer for AfterTonemap post process material
+			if (MaterialInterface != LastOutputWrite)
+			{
+				CurrentInput = FScreenPassTextureSlice::CreateFromScreenPassTexture(GraphBuilder, Outputs);
+			}
+		}
+
+#if !UE_BUILD_SHIPPING
+		if (UserSceneTextureInputNum || !UserSceneTextureOutput.IsNone())
 		{
-			CurrentInput = FScreenPassTextureSlice::CreateFromScreenPassTexture(GraphBuilder, Outputs);
+			View.GetSceneTextures().UserSceneTextureEvents.Add({ EUserSceneTextureEvent::Pass, NAME_None, 0, (uint16)ViewIndex, Material });
 		}
+#endif
 	}
 
 	if (!Outputs.IsValid())
 	{
-		Outputs = FScreenPassTexture::CopyFromSlice(GraphBuilder, CurrentInput);
+		// If no passes wrote to OverrideOutput, we need to copy to OverrideOutput now.  This can happen if all the passes wrote
+		// to UserSceneTextures instead of the default output.
+		if (InputsTemplate.OverrideOutput.IsValid())
+		{
+			AddDrawTexturePass(GraphBuilder, View, CurrentInput, InputsTemplate.OverrideOutput);
+			Outputs = InputsTemplate.OverrideOutput;
+		}
+		else
+		{
+			Outputs = FScreenPassTexture::CopyFromSlice(GraphBuilder, CurrentInput);
+		}
 	}
 
 	return Outputs;
diff --git a/Engine/Source/Runtime/Renderer/Private/PostProcess/PostProcessMaterial.h b/Engine/Source/Runtime/Renderer/Private/PostProcess/PostProcessMaterial.h
--- a/Engine/Source/Runtime/Renderer/Private/PostProcess/PostProcessMaterial.h
+++ b/Engine/Source/Runtime/Renderer/Private/PostProcess/PostProcessMaterial.h
@@ -48,6 +48,7 @@ FScreenPassTexture AddPostProcessMaterialPass(
 FScreenPassTexture AddPostProcessMaterialChain(
 	FRDGBuilder& GraphBuilder,
 	const FViewInfo& View,
+	int32 ViewIndex,
 	const FPostProcessMaterialInputs& Inputs,
 	const FPostProcessMaterialChain& MaterialChain,
 	EPostProcessMaterialInput MaterialInput = EPostProcessMaterialInput::SceneColor);
diff --git a/Engine/Source/Runtime/Renderer/Private/PostProcess/PostProcessing.cpp b/Engine/Source/Runtime/Renderer/Private/PostProcess/PostProcessing.cpp
--- a/Engine/Source/Runtime/Renderer/Private/PostProcess/PostProcessing.cpp
+++ b/Engine/Source/Runtime/Renderer/Private/PostProcess/PostProcessing.cpp
@@ -72,6 +72,7 @@
 #include "UnrealEngine.h"
 #include "IlluminanceMeter.h"
 #include "SparseVolumeTexture/SparseVolumeTextureStreamingVisualize.h"
+#include "CanvasItem.h"
 
 bool IsMobileEyeAdaptationEnabled(const FViewInfo& View);
 
@@ -151,6 +152,18 @@ TAutoConsoleVariable<int32> CVarGBufferPicking(
 	TEXT("Evaluate GBuffer value for debugging purpose."),
 	ECVF_RenderThreadSafe);
 #endif
+
+#if !(UE_BUILD_SHIPPING)
+TAutoConsoleVariable<int32> CVarUserSceneTextureDebug(
+	TEXT("r.PostProcessing.UserSceneTextureDebug"),
+	2,
+	TEXT("Enable debug display of post process UserSceneTexture inputs and outputs.\n")
+	TEXT(" 0: disabled\n")
+	TEXT(" 1: enabled\n")
+	TEXT(" 2: enable on error -- missing input or unused output (default).  Suppressed by DisableAllScreenMessages.\n")
+	TEXT(" 3: enable only for view with texture visualized through Vis / VisualizeTexture command, to avoid debug clutter in other views.\n"),
+	ECVF_RenderThreadSafe);
+#endif
 }
 
 #if WITH_EDITOR
@@ -203,6 +216,9 @@ bool IsPostProcessingWithAlphaChannelSupported()
 FScreenPassTexture AddFinalPostProcessDebugInfoPasses(FRDGBuilder& GraphBuilder, const FViewInfo& View, FScreenPassTexture& ScreenPassSceneColor);
 #endif
 
+#if !UE_BUILD_SHIPPING
+static void AddUserSceneTextureDebugPass(FRDGBuilder& GraphBuilder, const FViewInfo& View, int32 ViewIndex, FScreenPassTexture Output);
+#endif
 
 FDefaultTemporalUpscaler::FOutputs AddThirdPartyTemporalUpscalerPasses(
 	FRDGBuilder& GraphBuilder,
@@ -704,7 +720,7 @@ void AddPostProcessingPasses(
 		// Post Process Material Chain - BL_SceneColorBeforeDOF
 		if (MaterialChainSceneColorBeforeDOF.Num())
 		{
-			SceneColor = AddPostProcessMaterialChain(GraphBuilder, View, GetPostProcessMaterialInputs(SceneColor), MaterialChainSceneColorBeforeDOF);
+			SceneColor = AddPostProcessMaterialChain(GraphBuilder, View, ViewIndex, GetPostProcessMaterialInputs(SceneColor), MaterialChainSceneColorBeforeDOF);
 		}
 
 		// Diaphragm Depth of Field
@@ -739,7 +755,7 @@ void AddPostProcessingPasses(
 		// Post Process Material Chain - BL_SceneColorAfterDOF
 		if (MaterialChainSceneColorAfterDOF.Num())
 		{
-			SceneColor = AddPostProcessMaterialChain(GraphBuilder, View, GetPostProcessMaterialInputs(SceneColor), MaterialChainSceneColorAfterDOF);
+			SceneColor = AddPostProcessMaterialChain(GraphBuilder, View, ViewIndex, GetPostProcessMaterialInputs(SceneColor), MaterialChainSceneColorAfterDOF);
 		}
 
 		// Post Process Material Chain - BL_TranslucencyAfterDOF
@@ -753,7 +769,7 @@ void AddPostProcessingPasses(
 			if (MaterialChainTranslucencyAfterDOF.Num())
 			{
 				FScreenPassTexture PostDOFTranslucency = AddPostProcessMaterialChain(
-					GraphBuilder, View,
+					GraphBuilder, View, ViewIndex,
 					GetPostProcessMaterialInputs(SceneColor),
 					MaterialChainTranslucencyAfterDOF,
 					EPostProcessMaterialInput::SeparateTranslucency);
@@ -908,7 +924,7 @@ void AddPostProcessingPasses(
 			if (MaterialChain.Num())
 			{
 				InOutPassInputs = GetPostProcessMaterialInputs(FScreenPassTexture::CopyFromSlice(GraphBuilder, SceneColorSlice));
-				PassOutput = AddPostProcessMaterialChain(GraphBuilder, View, InOutPassInputs, MaterialChain);
+				PassOutput = AddPostProcessMaterialChain(GraphBuilder, View, ViewIndex, InOutPassInputs, MaterialChain);
 			}
 
 			for (FAfterPassCallbackDelegate& PassCallback : SSRInputDelegates)
@@ -995,7 +1011,7 @@ void AddPostProcessingPasses(
 			PassSequence.AcceptOverrideIfLastPass(EPass::PostProcessMaterialBeforeBloom, PostProcessMaterialInputs.OverrideOutput);
 			PostProcessMaterialInputs.SetInput(EPostProcessMaterialInput::SceneColor, SceneColorSlice);
 
-			SceneColorSlice = FScreenPassTextureSlice::CreateFromScreenPassTexture(GraphBuilder, AddPostProcessMaterialChain(GraphBuilder, View, PostProcessMaterialInputs, PostProcessMaterialBeforeBloomChain));
+			SceneColorSlice = FScreenPassTextureSlice::CreateFromScreenPassTexture(GraphBuilder, AddPostProcessMaterialChain(GraphBuilder, View, ViewIndex, PostProcessMaterialInputs, PostProcessMaterialBeforeBloomChain));
 		}
 
 		// Generate before bloom lower res scene color if they have not been generated.
@@ -1311,7 +1327,7 @@ void AddPostProcessingPasses(
 			PassInputs.SetInput(GraphBuilder, EPostProcessMaterialInput::PostTonemapHDRColor, SceneColorAfterTonemap);
 			PassInputs.SceneTextures = GetSceneTextureShaderParameters(Inputs.SceneTextures);
 
-			SceneColor = AddPostProcessMaterialChain(GraphBuilder, View, PassInputs, PostProcessMaterialAfterTonemappingChain);
+			SceneColor = AddPostProcessMaterialChain(GraphBuilder, View, ViewIndex, PassInputs, PostProcessMaterialAfterTonemappingChain);
 		}
 
 		if (PassSequence.IsEnabled(EPass::VisualizeLumenScene))
@@ -1821,6 +1837,61 @@ void AddPostProcessingPasses(
 		}
 	}
 	#endif
+
+	#if !UE_BUILD_SHIPPING
+	{
+		int32 UserSceneTextureDebug = CVarUserSceneTextureDebug.GetValueOnRenderThread();
+		bool bEnableUserSceneTextureDebug = false;
+
+		if (UserSceneTextureDebug == 1)
+		{
+			// Enable always
+			bEnableUserSceneTextureDebug = true;
+		}
+		else if (UserSceneTextureDebug == 2 && GAreScreenMessagesEnabled)
+		{
+			// Enable conditionally if there are errors
+			const FSceneTextures& SceneTextures = View.GetSceneTextures();
+			for (const FUserSceneTextureEventData& EventData : SceneTextures.UserSceneTextureEvents)
+			{
+				if (EventData.Event == EUserSceneTextureEvent::MissingInput || EventData.Event == EUserSceneTextureEvent::CollidingInput)
+				{
+					bEnableUserSceneTextureDebug = true;
+					break;
+				}
+			}
+
+			for (auto& UserSceneTextureElement : SceneTextures.UserSceneTextures)
+			{
+				if (bEnableUserSceneTextureDebug)
+				{
+					break;
+				}
+				for (FTransientUserSceneTexture& UserSceneTexture : UserSceneTextureElement.Value)
+				{
+					if (!UserSceneTexture.bUsed)
+					{
+						bEnableUserSceneTextureDebug = true;
+						break;
+					}
+				}
+			}
+		}
+		else if (UserSceneTextureDebug == 3)
+		{
+			// Enable conditionally for view with texture being visualized
+			if (GVisualizeTexture.IsRequestedView())
+			{
+				bEnableUserSceneTextureDebug = true;
+			}
+		}
+
+		if (bEnableUserSceneTextureDebug)
+		{
+			AddUserSceneTextureDebugPass(GraphBuilder, View, ViewIndex, SceneColor);
+		}
+	}
+	#endif
 }
 
 void AddDebugViewPostProcessingPasses(FRDGBuilder& GraphBuilder, const FViewInfo& View, FSceneUniformBuffer &SceneUniformBuffer, const FPostProcessingInputs& Inputs, const Nanite::FRasterResults* NaniteRasterResults)
@@ -2124,7 +2195,7 @@ static bool IsGaussianActive(const FViewInfo& View)
 	return true;
 }
 
-void AddMobilePostProcessingPasses(FRDGBuilder& GraphBuilder, FScene* Scene, const FViewInfo& View, FSceneUniformBuffer &SceneUniformBuffer, const FMobilePostProcessingInputs& Inputs, FInstanceCullingManager& InstanceCullingManager)
+void AddMobilePostProcessingPasses(FRDGBuilder& GraphBuilder, FScene* Scene, const FViewInfo& View, int32 ViewIndex, FSceneUniformBuffer &SceneUniformBuffer, const FMobilePostProcessingInputs& Inputs, FInstanceCullingManager& InstanceCullingManager)
 {
 	RDG_CSV_STAT_EXCLUSIVE_SCOPE(GraphBuilder, RenderPostProcessing);
 	QUICK_SCOPE_CYCLE_COUNTER(STAT_PostProcessing_Process);
@@ -2266,7 +2337,7 @@ void AddMobilePostProcessingPasses(FRDGBuilder& GraphBuilder, FScene* Scene, con
 	// The scene color will be decoded at the first post-process material and output linear color space for the following passes
 	// bMetalMSAAHDRDecode will be set to false if there is any post-process material exist
 
-	auto AddPostProcessMaterialPass = [&GraphBuilder, &View, &Inputs, &SceneColor, &CustomDepth, &bMetalMSAAHDRDecode, &PassSequence](EBlendableLocation BlendableLocation, bool bLastPass)
+	auto AddPostProcessMaterialPass = [&GraphBuilder, &View, ViewIndex, &Inputs, &SceneColor, &CustomDepth, &bMetalMSAAHDRDecode, &PassSequence](EBlendableLocation BlendableLocation, bool bLastPass)
 	{
 		FPostProcessMaterialInputs PostProcessMaterialInputs;
 
@@ -2287,7 +2358,7 @@ void AddMobilePostProcessingPasses(FRDGBuilder& GraphBuilder, FScene* Scene, con
 
 		if (MaterialChain.Num())
 		{
-			SceneColor = AddPostProcessMaterialChain(GraphBuilder, View, PostProcessMaterialInputs, MaterialChain);
+			SceneColor = AddPostProcessMaterialChain(GraphBuilder, View, ViewIndex, PostProcessMaterialInputs, MaterialChain);
 
 			// For solid material, we decode the input color and output the linear color
 			// For blend material, we force it rendering to an intermediate render target and decode there
@@ -3076,6 +3147,201 @@ FScreenPassTexture AddFinalPostProcessDebugInfoPasses(FRDGBuilder& GraphBuilder,
 }
 #endif
 
+#if !UE_BUILD_SHIPPING
+
+// Canvas.DrawShadowedString returns height -- this variation returns width
+static float CanvasDrawShadowedStringReturnWidth(FCanvas& Canvas, float PrintX, float PrintY, const FString& Text, const UFont* Font, FLinearColor TextColor)
+{
+	FCanvasTextStringViewItem TextItem(FVector2D(PrintX, PrintY), Text, Font, TextColor);
+	if (Font && Font->ImportOptions.bUseDistanceFieldAlpha)
+	{
+		TextItem.BlendMode = SE_BLEND_MaskedDistanceFieldShadowed;
+	}
+	else
+	{
+		TextItem.EnableShadow(FLinearColor::Black);
+	}
+	Canvas.DrawItem(TextItem);
+	return TextItem.DrawnSize.X / Canvas.GetDPIScale();
+}
+
+static void AddUserSceneTextureDebugPass(FRDGBuilder& GraphBuilder, const FViewInfo& View, int32 ViewIndex, FScreenPassTexture Output)
+{
+	const FSceneTextures& SceneTextures = View.GetSceneTextures();
+	if (!SceneTextures.UserSceneTextureEvents.IsEmpty())
+	{
+		FScreenPassRenderTarget OutputTarget = FScreenPassRenderTarget(Output, ERenderTargetLoadAction::ELoad);
+		AddDrawCanvasPass(GraphBuilder, RDG_EVENT_NAME("UserSceneTextureDebug"), View, OutputTarget, [&View, ViewIndex, &SceneTextures](FCanvas& Canvas)
+		{
+			FLinearColor TextColor(FLinearColor::White);
+			FLinearColor GrayTextColor(FLinearColor::Gray);
+			FLinearColor GreenTextColor(FLinearColor::Green);
+			FLinearColor RedTextColor(FLinearColor::Red);
+			FLinearColor YellowTextColor(FLinearColor::Yellow);
+			FLinearColor MagentaTextColor(1.f, 0.f, 1.f);
+			FString Text;
+
+			const UFont* Font = GetStatsFont();
+
+			const float ViewPortWidth = float(View.ViewRect.Width());
+			const float ViewPortHeight = float(View.ViewRect.Height());
+
+			const float CRHeight = 20.0f;
+			const float OffsetFromLeft = 0.05f;
+			const float OffsetFromTop = 0.2f;
+			const float OffsetFromHeader = CRHeight * 1.5f;
+
+			float PrintX_CR = ViewPortWidth * OffsetFromLeft;
+
+			float PrintX = PrintX_CR;
+			float PrintY = ViewPortHeight * OffsetFromTop;
+
+			int32 NumPasses = 0;
+			for (const FUserSceneTextureEventData& EventData : SceneTextures.UserSceneTextureEvents)
+			{
+				if (EventData.ViewIndex == ViewIndex && EventData.Event == EUserSceneTextureEvent::Pass)
+				{
+					NumPasses++;
+				}
+			}
+
+			// Draw header
+			Text = FString::Printf(TEXT("User Scene Texture Passes (count = %i)"), NumPasses);
+			PrintX += CanvasDrawShadowedStringReturnWidth(Canvas, PrintX, PrintY, *Text, Font, GreenTextColor);
+			if (CVarUserSceneTextureDebug.GetValueOnRenderThread() == 2)
+			{
+				CanvasDrawShadowedStringReturnWidth(Canvas, PrintX, PrintY, TEXT(" - enabled on error via \"r.PostProcessing.UserSceneTextureDebug 2\""), Font, GreenTextColor);
+			}
+			PrintX = PrintX_CR;
+			PrintY += CRHeight;
+
+			// Draw column description
+			CanvasDrawShadowedStringReturnWidth(Canvas, PrintX, PrintY, TEXT("Location [Priority]  Material:   Inputs   -->  Output"), Font, GrayTextColor);
+			PrintY += OffsetFromHeader;
+
+			// Draw blendable locations and priorities
+			static_assert(BL_MAX == 7);
+			static const TCHAR* GBlendableLocationShortNames[BL_MAX + 1] = {};
+			if (!GBlendableLocationShortNames[0])
+			{
+				// One time init -- enum in header isn't in numerical order, so it's simpler to initialize this way
+				GBlendableLocationShortNames[BL_SceneColorBeforeDOF] = TEXT("BeforeDOF");
+				GBlendableLocationShortNames[BL_SceneColorAfterDOF] = TEXT("AfterDOF");
+				GBlendableLocationShortNames[BL_TranslucencyAfterDOF] = TEXT("Translucent");
+				GBlendableLocationShortNames[BL_SSRInput] = TEXT("SSRInput");
+				GBlendableLocationShortNames[BL_SceneColorBeforeBloom] = TEXT("BeforeBloom");
+				GBlendableLocationShortNames[BL_ReplacingTonemapper] = TEXT("ReplaceTonemap");
+				GBlendableLocationShortNames[BL_SceneColorAfterTonemapping] = TEXT("AfterTonemap");
+				GBlendableLocationShortNames[BL_MAX] = TEXT("MAX");
+			}
+
+			float MaxBlendableInfoWidth = 0.0f;
+			for (const FUserSceneTextureEventData& EventData : SceneTextures.UserSceneTextureEvents)
+			{
+				if (EventData.ViewIndex == ViewIndex && EventData.Event == EUserSceneTextureEvent::Pass)
+				{
+					Text = FString::Printf(TEXT("%s [%d]"), GBlendableLocationShortNames[FMath::Min((uint32)EventData.Material->GetBlendableLocation(), (uint32)BL_MAX)], EventData.Material->GetBlendablePriority());
+
+					float BlendableInfoWidth = CanvasDrawShadowedStringReturnWidth(Canvas, PrintX, PrintY, Text, Font, TextColor);
+					MaxBlendableInfoWidth = FMath::Max(MaxBlendableInfoWidth, BlendableInfoWidth);
+					PrintY += CRHeight;
+				}
+			}
+
+			PrintX_CR = PrintX_CR + MaxBlendableInfoWidth + 10.0f;
+			PrintX = PrintX_CR;
+			PrintY = ViewPortHeight * OffsetFromTop + CRHeight + OffsetFromHeader;
+
+			// Draw material names
+			float MaxNameWidth = 0.0f;
+			for (const FUserSceneTextureEventData& EventData : SceneTextures.UserSceneTextureEvents)
+			{
+				if (EventData.ViewIndex == ViewIndex && EventData.Event == EUserSceneTextureEvent::Pass)
+				{
+					Text = FString::Printf(TEXT("%s:"), *EventData.Material->GetFriendlyName());
+
+					float NameWidth = CanvasDrawShadowedStringReturnWidth(Canvas, PrintX, PrintY, Text, Font, TextColor);
+					MaxNameWidth = FMath::Max(MaxNameWidth, NameWidth);
+					PrintY += CRHeight;
+				}
+			}
+
+			PrintX_CR = PrintX_CR + MaxNameWidth + 10.0f;
+			PrintX = PrintX_CR;
+			PrintY = ViewPortHeight * OffsetFromTop + CRHeight + OffsetFromHeader;
+
+			// Draw everything else (inputs and outputs)
+			bool bAnyMissing = false;
+			bool bAnyUnused = false;
+			bool bAnyColliding = false;
+
+			for (const FUserSceneTextureEventData& EventData : SceneTextures.UserSceneTextureEvents)
+			{
+				if (EventData.ViewIndex == ViewIndex)
+				{
+					switch (EventData.Event)
+					{
+					case EUserSceneTextureEvent::MissingInput:
+						Text = FString::Printf(TEXT("  %s"), *EventData.Name.ToString());
+						PrintX += CanvasDrawShadowedStringReturnWidth(Canvas, PrintX, PrintY, Text, Font, RedTextColor);
+						bAnyMissing = true;
+						break;
+					case EUserSceneTextureEvent::CollidingInput:
+						Text = FString::Printf(TEXT("  %s"), *EventData.Name.ToString());
+						PrintX += CanvasDrawShadowedStringReturnWidth(Canvas, PrintX, PrintY, Text, Font, MagentaTextColor);
+						bAnyColliding = true;
+						break;
+					case EUserSceneTextureEvent::FoundInput:
+						Text = FString::Printf(TEXT("  %s"), *EventData.Name.ToString());
+						PrintX += CanvasDrawShadowedStringReturnWidth(Canvas, PrintX, PrintY, Text, Font, GrayTextColor);
+						break;
+					case EUserSceneTextureEvent::Output:
+						{
+							const FTransientUserSceneTexture* UserTexture = SceneTextures.FindUserSceneTextureByEvent(EventData);
+							check(UserTexture);
+
+							Text = FString::Printf(TEXT("  --> %s [%dx%d]"), *EventData.Name.ToString(), EventData.RectSize.X, EventData.RectSize.Y);
+							PrintX += CanvasDrawShadowedStringReturnWidth(Canvas, PrintX, PrintY, *Text, Font, UserTexture->bUsed ? GrayTextColor : YellowTextColor);
+							bAnyUnused = bAnyUnused || !UserTexture->bUsed;
+
+							if (EventData.Material->GetBlendMode() != BLEND_Opaque)
+							{
+								PrintX += CanvasDrawShadowedStringReturnWidth(Canvas, PrintX, PrintY, TEXT("  Blend"), Font, GrayTextColor);
+							}
+						}
+						break;
+					case EUserSceneTextureEvent::Pass:
+						// End of line
+						PrintY += CRHeight;
+						PrintX = PrintX_CR;
+						break;
+					}
+				}
+			}
+
+			// Print color codings for warnings if present
+			PrintX = ViewPortWidth * OffsetFromLeft;
+			PrintY += CRHeight * 0.5f;
+			if (bAnyUnused)
+			{
+				CanvasDrawShadowedStringReturnWidth(Canvas, PrintX, PrintY, TEXT("Yellow:  Unused Output"), Font, YellowTextColor);
+				PrintY += CRHeight;
+			}
+			if (bAnyMissing)
+			{
+				CanvasDrawShadowedStringReturnWidth(Canvas, PrintX, PrintY, TEXT("Red:  Missing Input"), Font, RedTextColor);
+				PrintY += CRHeight;
+			}
+			if (bAnyColliding)
+			{
+				CanvasDrawShadowedStringReturnWidth(Canvas, PrintX, PrintY, TEXT("Magenta:  Input collides with Output"), Font, MagentaTextColor);
+				PrintY += CRHeight;
+			}
+		});
+	}
+}
+#endif  // !UE_BUILD_SHIPPING
+
 // Shader for visualizing GBuffer values
 class FGBufferPickingCS : public FGlobalShader
 {
diff --git a/Engine/Source/Runtime/Renderer/Private/PostProcess/PostProcessing.h b/Engine/Source/Runtime/Renderer/Private/PostProcess/PostProcessing.h
--- a/Engine/Source/Runtime/Renderer/Private/PostProcess/PostProcessing.h
+++ b/Engine/Source/Runtime/Renderer/Private/PostProcess/PostProcessing.h
@@ -65,7 +65,7 @@ struct FMobilePostProcessingInputs
 	}
 };
 
-void AddMobilePostProcessingPasses(FRDGBuilder& GraphBuilder, FScene* Scene, const FViewInfo& View, FSceneUniformBuffer &SceneUniformBuffer, const FMobilePostProcessingInputs& Inputs, FInstanceCullingManager& InstanceCullingManager);
+void AddMobilePostProcessingPasses(FRDGBuilder& GraphBuilder, FScene* Scene, const FViewInfo& View, int32 ViewIndex, FSceneUniformBuffer &SceneUniformBuffer, const FMobilePostProcessingInputs& Inputs, FInstanceCullingManager& InstanceCullingManager);
 
 void AddBasicPostProcessPasses(FRDGBuilder& GraphBuilder, const FViewInfo& View);
 
diff --git a/Engine/Source/Runtime/Renderer/Private/SceneTextures.cpp b/Engine/Source/Runtime/Renderer/Private/SceneTextures.cpp
--- a/Engine/Source/Runtime/Renderer/Private/SceneTextures.cpp
+++ b/Engine/Source/Runtime/Renderer/Private/SceneTextures.cpp
@@ -493,6 +493,142 @@ FSceneTextureShaderParameters FMinimalSceneTextures::GetSceneTextureShaderParame
 	return OutSceneTextureShaderParameters;
 }
 
+FRDGTextureRef FMinimalSceneTextures::FindOrAddUserSceneTexture(FRDGBuilder& GraphBuilder, int32 ViewIndex, FName Name, FIntPoint ResolutionDivisor, bool& bOutFirstRender, const FMaterial* Material, const FIntRect& OutputRect) const
+{
+	check(ResolutionDivisor.X >= 1 && ResolutionDivisor.Y >= 1);
+
+	bool bFound = false;
+
+	TArray<FTransientUserSceneTexture>* TransientTextures = UserSceneTextures.Find(Name);
+	if (TransientTextures)
+	{
+		if ((*TransientTextures)[0].ResolutionDivisor == ResolutionDivisor)
+		{
+			bFound = true;
+		}
+
+		for (int32 TextureIndex = 1; TextureIndex < TransientTextures->Num(); ++TextureIndex)
+		{
+			if ((*TransientTextures)[TextureIndex].ResolutionDivisor == ResolutionDivisor)
+			{
+				// Swap found item to front of array and return it -- next material render should use the most recently written item
+				TransientTextures->Swap(0, TextureIndex);
+				bFound = true;
+				break;
+			}
+		}
+	}
+	else
+	{
+		TransientTextures = &UserSceneTextures.Add(Name);
+	}
+
+	if (!bFound)
+	{
+		// We didn't find an existing item, need to allocate a new one
+		FTransientUserSceneTexture TransientUserTexture;
+		TransientUserTexture.ResolutionDivisor = ResolutionDivisor;
+		TransientUserTexture.AllocationOrder = TransientTextures->Num();
+		TransientUserTexture.bUsed = false;
+		TransientUserTexture.ViewMask = 0;
+	
+		FIntPoint Extent = FIntPoint::DivideAndRoundUp(Config.Extent, ResolutionDivisor);
+
+		FRDGTextureDesc Desc(Config.bRequireMultiView ?
+			FRDGTextureDesc::Create2DArray(Extent, Config.ColorFormat, Config.ColorClearValue, Config.ColorCreateFlags, 2) :
+			FRDGTextureDesc::Create2D(Extent, Config.ColorFormat, Config.ColorClearValue, Config.ColorCreateFlags));
+
+#if !(UE_BUILD_SHIPPING || UE_BUILD_TEST)
+		// Generate a heap allocated debug name for the resource, required to make UserSceneTextures accessible in VisualizeTexture,
+		// GPU captures, and GPU dumps, as otherwise every texture will end up with the same debug name.
+		FString DebugName = TransientUserTexture.AllocationOrder == 0 ?
+			FString::Printf(TEXT("UST.%s"), *Name.ToString()) :
+			FString::Printf(TEXT("UST.%s[%d]"), *Name.ToString(), TransientUserTexture.AllocationOrder);
+		TCHAR* HeapDebugName = new TCHAR[DebugName.Len() + 1];
+		FCString::Strcpy(HeapDebugName, DebugName.Len() + 1, *DebugName);
+
+		TransientUserTexture.Texture = GraphBuilder.CreateTexture(Desc, HeapDebugName);
+		TransientUserTexture.Texture->SetDebugNameIsHeapAllocated();
+#else
+		TransientUserTexture.Texture = GraphBuilder.CreateTexture(Desc, TEXT("TransientUserTexture"));
+#endif
+
+		TransientTextures->Add(TransientUserTexture);
+
+		// Swap newly allocated item to front of array if it's not the only item
+		if (TransientTextures->Num() > 1)
+		{
+			TransientTextures->Swap(0, TransientTextures->Num() - 1);
+		}
+	}
+
+#if !(UE_BUILD_SHIPPING)
+	UserSceneTextureEvents.Add({ EUserSceneTextureEvent::Output, Name, (*TransientTextures)[0].AllocationOrder, (uint16)ViewIndex, Material, OutputRect.Size() });
+#endif
+
+	// If out of mask range, treat it as the first render.  This may result in certain transparent post process materials writing to UserSceneTextures
+	// rendering incorrectly if there are more than 32 views, but we don't expect to see any real world situation with that number of views anyway
+	// (and transparent post process materials themselves aren't common to begin with).
+	if (ViewIndex > 31)
+	{
+		bOutFirstRender = true;
+	}
+	else
+	{
+		bOutFirstRender = ((*TransientTextures)[0].ViewMask & (1u << ViewIndex)) == 0;
+		(*TransientTextures)[0].ViewMask |= (1u << ViewIndex);
+	}
+
+	return (*TransientTextures)[0].Texture;
+}
+
+FScreenPassTextureSlice FMinimalSceneTextures::GetUserSceneTexture(FRDGBuilder& GraphBuilder, const FViewInfo& View, int32 ViewIndex, FName Name, const FMaterial* Material) const
+{
+	TArray<FTransientUserSceneTexture>* TransientTextures = UserSceneTextures.Find(Name);
+	if (TransientTextures)
+	{
+		FScreenPassTextureSlice TransientTextureSlice(
+			GraphBuilder.CreateSRV(FRDGTextureSRVDesc((*TransientTextures)[0].Texture)),
+			GetDownscaledViewRect(View.UnconstrainedViewRect, View.GetFamilyViewRect().Max, (*TransientTextures)[0].ResolutionDivisor));
+
+#if !(UE_BUILD_SHIPPING)
+		(*TransientTextures)[0].bUsed = true;
+		UserSceneTextureEvents.Add({ EUserSceneTextureEvent::FoundInput, Name, (*TransientTextures)[0].AllocationOrder, (uint16)ViewIndex, Material });
+#endif
+
+		return TransientTextureSlice;
+	}
+	else
+	{
+#if !(UE_BUILD_SHIPPING)
+		UserSceneTextureEvents.Add({ EUserSceneTextureEvent::MissingInput, Name, 0, (uint16)ViewIndex, Material });
+#endif
+
+		return FScreenPassTextureSlice();
+	}
+}
+
+#if !(UE_BUILD_SHIPPING)
+const FTransientUserSceneTexture* FMinimalSceneTextures::FindUserSceneTextureByEvent(const FUserSceneTextureEventData& Event) const
+{
+	const TArray<FTransientUserSceneTexture>* TransientTextures = UserSceneTextures.Find(Event.Name);
+	if (!TransientTextures)
+	{
+		return nullptr;
+	}
+
+	for (const FTransientUserSceneTexture& TransientTexture : *TransientTextures)
+	{
+		if (TransientTexture.AllocationOrder == Event.AllocationOrder)
+		{
+			return &TransientTexture;
+		}
+	}
+
+	return nullptr;
+}
+#endif
+
 void FSceneTextures::InitializeViewFamily(FRDGBuilder& GraphBuilder, FViewFamilyInfo& ViewFamily, FIntPoint FamilySize)
 {
 	const FSceneTexturesConfig& Config = ViewFamily.SceneTexturesConfig;
diff --git a/Engine/Source/Runtime/Renderer/Private/ScreenPass.cpp b/Engine/Source/Runtime/Renderer/Private/ScreenPass.cpp
--- a/Engine/Source/Runtime/Renderer/Private/ScreenPass.cpp
+++ b/Engine/Source/Runtime/Renderer/Private/ScreenPass.cpp
@@ -74,6 +74,11 @@ FScreenPassTexture FScreenPassTexture::CopyFromSlice(FRDGBuilder& GraphBuilder,
 	Desc.Dimension = ETextureDimension::Texture2D;
 	Desc.ArraySize = 1;
 
+	// If a pass uses blending to write to this post process texture, it needs to support being a render target, so make sure this flag is included.
+	// Most post processing uses SceneColor or its FRDGTextureDesc, and SceneColor already has the RenderTargetable flag set, but TSR (the input to
+	// the "Before Bloom" stage) writes to texture slices using compute, and its output doesn't have this flag.
+	Desc.Flags |= ETextureCreateFlags::RenderTargetable;
+
 	FRDGTextureRef NewTexture = GraphBuilder.CreateTexture(Desc, TEXT("CopyToScreenPassTexture2D"));
 
 	FRHICopyTextureInfo CopyInfo;
@@ -328,6 +333,25 @@ void AddDrawTexturePass(
 	AddDrawScreenPass(GraphBuilder, RDG_EVENT_NAME("DrawTexture"), View, OutputViewport, InputViewport, PixelShader, Parameters);
 }
 
+void AddDrawTexturePass(
+	FRDGBuilder& GraphBuilder,
+	const FSceneView& View,
+	FScreenPassTextureSlice Input,
+	FScreenPassRenderTarget Output)
+{
+	const FScreenPassTextureViewport InputViewport(Input);
+	const FScreenPassTextureViewport OutputViewport(Output);
+
+	TShaderMapRef<FCopyRectSrvPS> PixelShader(static_cast<const FViewInfo&>(View).ShaderMap);
+
+	FCopyRectSrvPS::FParameters* Parameters = GraphBuilder.AllocParameters<FCopyRectSrvPS::FParameters>();
+	Parameters->InputTexture = Input.TextureSRV;
+	Parameters->InputSampler = TStaticSamplerState<>::GetRHI();
+	Parameters->RenderTargets[0] = Output.GetRenderTargetBinding();
+
+	AddDrawScreenPass(GraphBuilder, RDG_EVENT_NAME("DrawTexture"), View, OutputViewport, InputViewport, PixelShader, Parameters);
+}
+
 class FDownsampleDepthPS : public FGlobalShader
 {
 public:
diff --git a/Engine/Source/Runtime/Renderer/Public/PostProcess/PostProcessMaterialInputs.h b/Engine/Source/Runtime/Renderer/Public/PostProcess/PostProcessMaterialInputs.h
--- a/Engine/Source/Runtime/Renderer/Public/PostProcess/PostProcessMaterialInputs.h
+++ b/Engine/Source/Runtime/Renderer/Public/PostProcess/PostProcessMaterialInputs.h
@@ -43,6 +43,11 @@ enum class EPathTracingPostProcessMaterialInput : uint32
 
 struct FPostProcessMaterialInputs
 {
+	FPostProcessMaterialInputs()
+	{
+		FMemory::Memset(bUserSceneTexturesSet, 0);
+	}
+
 	inline void SetInput(FRDGBuilder& GraphBuilder, EPostProcessMaterialInput Input, FScreenPassTexture Texture)
 	{
 		SetInput(Input, FScreenPassTextureSlice::CreateFromScreenPassTexture(GraphBuilder, Texture));
@@ -50,12 +55,27 @@ struct FPostProcessMaterialInputs
 
 	inline void SetInput(EPostProcessMaterialInput Input, FScreenPassTextureSlice Texture)
 	{
+		check((uint32)Input < kPostProcessMaterialInputCountMax);
+
 		Textures[(uint32)Input] = Texture;
 	}
 
+	inline void SetUserSceneTextureInput(EPostProcessMaterialInput Input, FScreenPassTextureSlice Texture)
+	{
+		check((uint32)Input < kPostProcessMaterialInputCountMax);
+
+		UserSceneTextures[(uint32)Input] = Texture;
+		bUserSceneTexturesSet[(uint32)Input] = true;
+	}
+
 	inline FScreenPassTextureSlice GetInput(EPostProcessMaterialInput Input) const
 	{
-		return Textures[(uint32)Input];
+		return bUserSceneTexturesSet[(uint32)Input] ? UserSceneTextures[(uint32)Input] : Textures[(uint32)Input];
+	}
+
+	inline FScreenPassTextureSlice GetSceneColorOutput(EBlendableLocation BlendableLocation) const
+	{
+		return Textures[(uint32)(BlendableLocation == BL_TranslucencyAfterDOF ? EPostProcessMaterialInput::SeparateTranslucency : EPostProcessMaterialInput::SceneColor)];
 	}
 
 	inline void SetPathTracingInput(EPathTracingPostProcessMaterialInput Input, FScreenPassTexture Texture)
@@ -71,7 +91,11 @@ struct FPostProcessMaterialInputs
 	inline void Validate() const
 	{
 		ValidateInputExists(EPostProcessMaterialInput::SceneColor);
-		ValidateInputExists(EPostProcessMaterialInput::SeparateTranslucency);
+		
+		// TODO:  Is separate translucency always guaranteeed to be present?  A previous version of the code appeared to be attempting to validate this, but
+		//		  due to a bug (ValidateInputExists ignoring its Input argument and instead validating EPostProcessMaterialInput::SceneColor regardless of
+		//		  what was passed in), didn't actually do so.  I'm afraid to enable this assert now, as I don't know if it will randomly crash some project.
+		// ValidateInputExists(EPostProcessMaterialInput::SeparateTranslucency);
 
 		// Either override output format is valid or the override output texture is; not both.
 		if (OutputFormat != PF_Unknown)
@@ -88,8 +112,7 @@ struct FPostProcessMaterialInputs
 
 	inline void ValidateInputExists(EPostProcessMaterialInput Input) const
 	{
-		const FScreenPassTextureSlice Texture = GetInput(EPostProcessMaterialInput::SceneColor);
-		check(Texture.IsValid());
+		check(Textures[(int32)Input].IsValid());
 	}
 
 	/**
@@ -117,6 +140,15 @@ struct FPostProcessMaterialInputs
 	 */
 	TStaticArray<FScreenPassTextureSlice, kPostProcessMaterialInputCountMax> Textures;
 
+	/**
+	 * UserSceneTexture inputs, which take precedence over Textures above if set.  The reason for separating these from "Textures" above
+	 * is because "Textures" is also where the output SceneColor is fetched from, when OverrideOutput isn't set (see "GetSceneColorOutput"
+	 * function).  The separate bools are needed to track which inputs are UserSceneTextures, as opposed to checking IsValid(), because
+	 * the entry can be invalid when an input being missing -- dummy black is substituted for those downstream.
+	 */
+	TStaticArray<FScreenPassTextureSlice, kPostProcessMaterialInputCountMax> UserSceneTextures;
+	TStaticArray<bool, kPostProcessMaterialInputCountMax> bUserSceneTexturesSet;
+
 	/**
 	*	Array of input textures bound to the material from path tracing. All inputs are optional
 	*/
@@ -145,6 +177,9 @@ struct FPostProcessMaterialInputs
 	bool bAllowSceneColorInputAsOutput = true;
 
 	bool bMetalMSAAHDRDecode = false;
+
+	bool bUserSceneTextureOutput = false;
+	bool bUserSceneTextureFirstRender = false;
 };
 
 class UMaterialInterface;
diff --git a/Engine/Source/Runtime/Renderer/Public/ScreenPass.h b/Engine/Source/Runtime/Renderer/Public/ScreenPass.h
--- a/Engine/Source/Runtime/Renderer/Public/ScreenPass.h
+++ b/Engine/Source/Runtime/Renderer/Public/ScreenPass.h
@@ -675,6 +675,13 @@ void RENDERER_API AddDrawTexturePass(
 	FScreenPassTexture Input,
 	FScreenPassRenderTarget Output);
 
+/** Variation of above that accepts a slice, which contains a texture SRV instead of a texture */
+void RENDERER_API AddDrawTexturePass(
+	FRDGBuilder& GraphBuilder,
+	const FSceneView& View,
+	FScreenPassTextureSlice Input,
+	FScreenPassRenderTarget Output);
+
 template <typename TFunction>
 FORCEINLINE void AddRenderTargetPass(
 	FRDGBuilder& GraphBuilder,
diff --git a/Engine/Source/Runtime/Renderer/Public/ScreenPass.inl b/Engine/Source/Runtime/Renderer/Public/ScreenPass.inl
--- a/Engine/Source/Runtime/Renderer/Public/ScreenPass.inl
+++ b/Engine/Source/Runtime/Renderer/Public/ScreenPass.inl
@@ -31,6 +31,17 @@ inline FIntRect GetDownscaledRect(FIntRect Rect, FIntPoint Divisor)
 	return Rect;
 }
 
+// The normal GetDownscaledRect can produce overlapping rects where multiple viewports are present, because the Min rounds down, while the
+// Max rounds up.  This variation rounds up for the far edges of the viewport, but rounds down for everything else, avoiding this issue.
+inline FIntRect GetDownscaledViewRect(FIntRect Rect, FIntPoint ViewExtent, FIntPoint Divisor)
+{
+	Rect.Min /= Divisor;
+	Rect.Max.X = (Rect.Max.X == ViewExtent.X) ? FMath::DivideAndRoundUp(ViewExtent.X, Divisor.X) : Rect.Max.X / Divisor.X;
+	Rect.Max.Y = (Rect.Max.Y == ViewExtent.Y) ? FMath::DivideAndRoundUp(ViewExtent.Y, Divisor.Y) : Rect.Max.Y / Divisor.Y;
+	Rect.Max = Rect.Max.ComponentMax(Rect.Min + FIntPoint(1, 1));
+	return Rect;
+}
+
 inline FIntRect GetScaledRect(FIntRect Rect, FVector2D Multiplier)
 {
 	Rect.Min.X *= Multiplier.X;
