Fix warnings compiling RayTracingScene.cpp

diff --git a/Engine/Source/Runtime/Renderer/Private/RayTracing/RayTracingScene.cpp b/Engine/Source/Runtime/Renderer/Private/RayTracing/RayTracingScene.cpp
--- a/Engine/Source/Runtime/Renderer/Private/RayTracing/RayTracingScene.cpp
+++ b/Engine/Source/Runtime/Renderer/Private/RayTracing/RayTracingScene.cpp
@@ -57,7 +57,7 @@ void FRayTracingScene::Create(FRDGBuilder& GraphBuilder, const FViewInfo& View,
 
 void FRayTracingScene::InitPreViewTranslation(const FViewMatrices& ViewMatrices)
 {
-	PreViewTranslation = FDFVector3(ViewMatrices.GetPreViewTranslation());
+	PreViewTranslation = ViewMatrices.GetPreViewTranslation();
 }
 
 void FRayTracingScene::CreateWithInitializationData(FRDGBuilder& GraphBuilder, const FViewInfo& View, const FGPUScene* GPUScene, FRayTracingSceneWithGeometryInstances SceneWithGeometryInstances, ERDGPassFlags ComputePassFlags)
@@ -272,7 +272,7 @@ void FRayTracingScene::CreateWithInitializationData(FRDGBuilder& GraphBuilder, c
 			BaseUploadBufferOffsets = MoveTemp(SceneWithGeometryInstances.BaseUploadBufferOffsets),
 			BaseInstancePrefixSum = MoveTemp(SceneWithGeometryInstances.BaseInstancePrefixSum),
 			RayTracingSceneRHI = RayTracingSceneRHI,
-			PreViewTranslation = View.ViewMatrices.GetPreViewTranslation()]()
+			PreViewTranslation = this->PreViewTranslation]()
 		{
 			FTaskTagScope TaskTagScope(ETaskTag::EParallelRenderingThread);
 			FillRayTracingInstanceUploadBuffer(
@@ -304,7 +304,6 @@ void FRayTracingScene::CreateWithInitializationData(FRDGBuilder& GraphBuilder, c
 			[PassParams,
 			this,
 			GPUScene,
-			PreViewTranslation = PreViewTranslation,
 			&SceneInitializer,
 			NumNativeGPUSceneInstances = SceneWithGeometryInstances.NumNativeGPUSceneInstances,
 			NumNativeCPUInstances = SceneWithGeometryInstances.NumNativeCPUInstances,
@@ -342,7 +341,7 @@ void FRayTracingScene::CreateWithInitializationData(FRDGBuilder& GraphBuilder, c
 				BuildRayTracingInstanceBuffer(
 					RHICmdList,
 					GPUScene,
-					PreViewTranslation,
+					FDFVector3(PreViewTranslation),
 					PassParams->InstanceBuffer->GetRHI(),
 					InstanceUploadSRV,
 					AccelerationStructureAddressesBuffer.SRV,
diff --git a/Engine/Source/Runtime/Renderer/Private/RayTracing/RayTracingScene.h b/Engine/Source/Runtime/Renderer/Private/RayTracing/RayTracingScene.h
--- a/Engine/Source/Runtime/Renderer/Private/RayTracing/RayTracingScene.h
+++ b/Engine/Source/Runtime/Renderer/Private/RayTracing/RayTracingScene.h
@@ -137,7 +137,7 @@ public:
 	bool bNeedsDebugInstanceGPUSceneIndexBuffer = false;
 
 	// Used for transforming to translated world space in which TLAS was built.
-	FDFVector3 PreViewTranslation {};
+	FVector PreViewTranslation {};
 private:
 	void ReleaseReadbackBuffers();
 
