Fix Substrate material buffer clear during the hair strands visiblity pass, having incorrect texture format declaration, causing pixel artifact on certain platforms.

diff --git a/Engine/Source/Runtime/Renderer/Private/HairStrands/HairStrandsVisibility.cpp b/Engine/Source/Runtime/Renderer/Private/HairStrands/HairStrandsVisibility.cpp
--- a/Engine/Source/Runtime/Renderer/Private/HairStrands/HairStrandsVisibility.cpp
+++ b/Engine/Source/Runtime/Renderer/Private/HairStrands/HairStrandsVisibility.cpp
@@ -2717,8 +2717,8 @@ public:
 	static void ModifyCompilationEnvironment(const FGlobalShaderPermutationParameters& Parameters, FShaderCompilerEnvironment& OutEnvironment)
 	{
 		FGlobalShader::ModifyCompilationEnvironment(Parameters, OutEnvironment);
-		OutEnvironment.SetRenderTargetOutputFormat(0, PF_B8G8R8A8);
-		OutEnvironment.SetRenderTargetOutputFormat(1, PF_B8G8R8A8);
+		OutEnvironment.SetRenderTargetOutputFormat(0, Substrate::IsSubstrateEnabled() ? PF_R32_UINT : PF_B8G8R8A8);
+		OutEnvironment.SetRenderTargetOutputFormat(1, Substrate::IsSubstrateEnabled() ? PF_R32_UINT : PF_B8G8R8A8);
 		OutEnvironment.SetRenderTargetOutputFormat(2, PF_FloatRGBA);
 	}
 };
