Fix DBuffer decals for stereo rendering

diff --git a/Engine/Shaders/Private/DBufferDecalShared.ush b/Engine/Shaders/Private/DBufferDecalShared.ush
--- a/Engine/Shaders/Private/DBufferDecalShared.ush
+++ b/Engine/Shaders/Private/DBufferDecalShared.ush
@@ -23,6 +23,12 @@
 #define OpaqueBasePassDBufferCTexture DBufferBasePass.DBufferCTexture
 #define OpaqueBasePassDBufferRenderMask DBufferBasePass.DBufferRenderMask
 
+#if MOBILE_MULTI_VIEW && SHADING_PATH_MOBILE
+#define OpaqueBasePassDBufferATextureArray DBufferBasePass.DBufferATextureArray
+#define OpaqueBasePassDBufferBTextureArray DBufferBasePass.DBufferBTextureArray
+#define OpaqueBasePassDBufferCTextureArray DBufferBasePass.DBufferCTextureArray
+#endif
+
 #if SUPPORTS_INDEPENDENT_SAMPLERS
 #define OpaqueBasePassDBufferATextureSampler DBufferBasePass.DBufferATextureSampler
 #define OpaqueBasePassDBufferBTextureSampler DBufferBasePass.DBufferATextureSampler
@@ -409,7 +415,7 @@ FDBufferData DecodeDBufferData(
 }
 
 // @param BufferUV - UV space in the DBuffer textures
-FDBufferData GetDBufferData(float2 BufferUV, uint RTMaskBit)
+FDBufferData GetDBufferData(float2 BufferUV, uint RTMaskBit, uint EyeIndex)
 {
 	// Setup default values, which mean that no decals are present in DBuffer
 	float4 DBufferA = float4(0.0, 0.0, 0.0, 1.0);
@@ -420,19 +426,31 @@ FDBufferData GetDBufferData(float2 BufferUV, uint RTMaskBit)
 	BRANCH
 	if (RTMaskBit & 0x1)
 	{
+	#if ((INSTANCED_STEREO && MOBILE_MULTI_VIEW) || MOBILE_MULTI_VIEW) && SHADING_PATH_MOBILE
+		DBufferA = Texture2DArraySampleLevel(OpaqueBasePassDBufferATextureArray, OpaqueBasePassDBufferATextureSampler, float3(BufferUV.xy, EyeIndex), 0);
+	#else
 		DBufferA = Texture2DSampleLevel(OpaqueBasePassDBufferATexture, OpaqueBasePassDBufferATextureSampler, BufferUV, 0);
+	#endif
 	}
 
 	BRANCH
 	if (RTMaskBit & 0x2)
 	{
+	#if ((INSTANCED_STEREO && MOBILE_MULTI_VIEW) || MOBILE_MULTI_VIEW) && SHADING_PATH_MOBILE
+		DBufferB = Texture2DArraySampleLevel(OpaqueBasePassDBufferBTextureArray, OpaqueBasePassDBufferBTextureSampler, float3(BufferUV.xy, EyeIndex), 0);
+	#else
 		DBufferB = Texture2DSampleLevel(OpaqueBasePassDBufferBTexture, OpaqueBasePassDBufferBTextureSampler, BufferUV, 0);
+	#endif
 	}
 
 	BRANCH
 	if (RTMaskBit & 0x4)
 	{
+	#if ((INSTANCED_STEREO && MOBILE_MULTI_VIEW) || MOBILE_MULTI_VIEW) && SHADING_PATH_MOBILE
+		DBufferC = Texture2DArraySampleLevel(OpaqueBasePassDBufferCTextureArray, OpaqueBasePassDBufferCTextureSampler, float3(BufferUV.xy, EyeIndex), 0);
+	#else
 		DBufferC = Texture2DSampleLevel(OpaqueBasePassDBufferCTexture, OpaqueBasePassDBufferCTextureSampler, BufferUV, 0);
+	#endif
 	}
 #elif USE_DBUFFER && RAYHITGROUPSHADER
 	BRANCH
@@ -458,6 +476,11 @@ FDBufferData GetDBufferData(float2 BufferUV, uint RTMaskBit)
 	return DecodeDBufferData(DBufferA, DBufferB, DBufferC);
 }
 
+FDBufferData GetDBufferData(float2 BufferUV, uint RTMaskBit)
+{
+	return GetDBufferData(BufferUV, RTMaskBit, 0);
+}
+
 /** Populates DBufferA, DBufferB, DBufferC as float4 and puts opacity in alpha for frame buffer blending */
 void ApplyDBufferData(
 	FDBufferData DBufferData, inout float3 WorldNormal, inout float3 SubsurfaceColor, inout float Roughness, 
diff --git a/Engine/Shaders/Private/MeshDecals.usf b/Engine/Shaders/Private/MeshDecals.usf
--- a/Engine/Shaders/Private/MeshDecals.usf
+++ b/Engine/Shaders/Private/MeshDecals.usf
@@ -61,12 +61,26 @@ struct FMeshDecalVSToPS
 void MainVS(
 	FVertexFactoryInput Input,
 	out VS_OUTPUT_TYPE Output
+#if INSTANCED_STEREO && MOBILE_MULTI_VIEW
+	, out nointerpolation uint EyeIndex : VIEW_ID
+	, out uint LayerIndex : SV_RenderTargetArrayIndex
+#elif MOBILE_MULTI_VIEW
+	, in nointerpolation uint ViewId : SV_ViewID
+#endif
 #if USE_GLOBAL_CLIP_PLANE
 	, out float OutGlobalClipPlaneDistance : SV_ClipDistance
 #endif
 	)
 {
+#if INSTANCED_STEREO && MOBILE_MULTI_VIEW
+	EyeIndex = GetEyeIndexFromVF(Input);
+	ResolvedView = ResolveView(EyeIndex);
+	LayerIndex = EyeIndex;
+#elif MOBILE_MULTI_VIEW
+	ResolvedView = ResolveView(ViewId);
+#else
 	ResolvedView = ResolveView();
+#endif
 
 	FVertexFactoryIntermediates VFIntermediates = GetVertexFactoryIntermediates(Input);
 	float4 WorldPosition = VertexFactoryGetWorldPosition(Input, VFIntermediates);
@@ -84,7 +98,7 @@ void MainVS(
 		WorldPosition.xyz += GetMaterialWorldPositionOffset(VertexParameters);
 	}
 	
-	Output.Position = mul(WorldPosition, View.TranslatedWorldToClip);
+	Output.Position = mul(WorldPosition, ResolvedView.TranslatedWorldToClip);
 	Output.FactoryInterpolants = VertexFactoryGetInterpolantsVSToPS(Input, VFIntermediates, VertexParameters);
 
 #if USE_GLOBAL_CLIP_PLANE
@@ -107,10 +121,11 @@ void FPixelShaderInOut_MainPS(
 	FMeshDecalInterpolants MeshDecalInterpolants,
 #endif
 	inout FPixelShaderIn In,
-	inout FPixelShaderOut Out)
+	inout FPixelShaderOut Out,
+	uint EyeIndex)
 {
-#if INSTANCED_STEREO
-	ResolvedView = ResolveView(FactoryInterpolants.EyeIndex);
+#if INSTANCED_STEREO || MOBILE_MULTI_VIEW
+	ResolvedView = ResolveView(EyeIndex);
 #else
 	ResolvedView = ResolveView();
 #endif
@@ -190,6 +205,21 @@ void FPixelShaderInOut_MainPS(
 #endif
 }
 
+void FPixelShaderInOut_MainPS(
+	FVertexFactoryInterpolantsVSToPS FactoryInterpolants,
+#if NEED_MESH_DECAL_INTERPOLATORS
+	FMeshDecalInterpolants MeshDecalInterpolants,
+#endif
+	inout FPixelShaderIn In,
+	inout FPixelShaderOut Out)
+{
+#if NEED_MESH_DECAL_INTERPOLATORS
+	FPixelShaderInOut_MainPS(FactoryInterpolants, MeshDecalInterpolants, In, Out, 0);
+#else 
+	FPixelShaderInOut_MainPS(FactoryInterpolants, In, Out, 0);
+#endif
+}
+
 #define PIXELSHADEROUTPUT_MESHDECALPASS NEED_MESH_DECAL_INTERPOLATORS
 #define PIXELSHADEROUTPUT_INTERPOLANTS 1
 //#define PIXELSHADEROUTPUT_MRT0 (DECAL_RENDERTARGET_COUNT > 0)
diff --git a/Engine/Shaders/Private/MobileBasePassPixelShader.usf b/Engine/Shaders/Private/MobileBasePassPixelShader.usf
--- a/Engine/Shaders/Private/MobileBasePassPixelShader.usf
+++ b/Engine/Shaders/Private/MobileBasePassPixelShader.usf
@@ -480,12 +480,12 @@ void Main(
 			float2 BufferUV = SvPositionToBufferUV(SvPosition);
 		#if SUBSTRATE_ENABLED 
 			#if SUBSTRATE_INLINE_SHADING && !SUBSTRATE_OPTIMIZED_UNLIT
-			const FSubstrateDBuffer SubstrateBufferData = GetDBufferData(BufferUV, ValidDBufferTargetMask);
+			const FSubstrateDBuffer SubstrateBufferData = GetDBufferData(BufferUV, ValidDBufferTargetMask, BasePassInterpolants.MultiViewId);
 			ApplyDBufferData(SubstrateBufferData, SubstratePixelHeader, SubstrateData);
 			DBufferOpacity = SubstrateBufferData.Coverage;
 			#endif
 		#else
-			FDBufferData DBufferData = GetDBufferData(BufferUV, ValidDBufferTargetMask);
+			FDBufferData DBufferData = GetDBufferData(BufferUV, ValidDBufferTargetMask, BasePassInterpolants.MultiViewId);
 			ApplyDBufferData(DBufferData, MaterialParameters.WorldNormal, SubsurfaceColor, Roughness, BaseColor, Metallic, Specular);
 			DBufferOpacity = (DBufferData.ColorOpacity + DBufferData.NormalOpacity + DBufferData.RoughnessOpacity) * (1.0f / 3.0f);
 		#endif // SUBSTRATE_ENABLED
diff --git a/Engine/Shaders/Private/PixelShaderOutputCommon.ush b/Engine/Shaders/Private/PixelShaderOutputCommon.ush
--- a/Engine/Shaders/Private/PixelShaderOutputCommon.ush
+++ b/Engine/Shaders/Private/PixelShaderOutputCommon.ush
@@ -153,7 +153,7 @@ uint ArrayIndex = 0;
 #if PIXELSHADEROUTPUT_BASEPASS
 	FPixelShaderInOut_MainPS(Interpolants, BasePassInterpolants, PixelShaderIn, PixelShaderOut);
 #elif PIXELSHADEROUTPUT_MESHDECALPASS
-	FPixelShaderInOut_MainPS(Interpolants, MeshDecalInterpolants, PixelShaderIn, PixelShaderOut);
+	FPixelShaderInOut_MainPS(Interpolants, MeshDecalInterpolants, PixelShaderIn, PixelShaderOut, ArrayIndex);
 #elif PIXELSHADEROUTPUT_INTERPOLANTS
 	FPixelShaderInOut_MainPS(Interpolants, PixelShaderIn, PixelShaderOut);
 #else
diff --git a/Engine/Shaders/Private/PositionOnlyDepthVertexShader.usf b/Engine/Shaders/Private/PositionOnlyDepthVertexShader.usf
--- a/Engine/Shaders/Private/PositionOnlyDepthVertexShader.usf
+++ b/Engine/Shaders/Private/PositionOnlyDepthVertexShader.usf
@@ -14,16 +14,28 @@ void Main(
 #if USE_GLOBAL_CLIP_PLANE
 	, out float OutGlobalClipPlaneDistance : SV_ClipDistance
 #endif
-#if INSTANCED_STEREO
+#if INSTANCED_STEREO && MOBILE_MULTI_VIEW // Mobile multi view fallback path
+	, out uint LayerIndex : SV_RenderTargetArrayIndex
+#elif MOBILE_MULTI_VIEW
+	, in nointerpolation uint ViewId : SV_ViewID
+#elif INSTANCED_STEREO
 	, out uint ViewportIndex : SV_ViewPortArrayIndex
 #endif
 	)
 {
 #if INSTANCED_STEREO
-	const uint EyeIndex = GetEyeIndexFromVF(Input);
-	ViewportIndex = EyeIndex;
-#endif
+	uint EyeIndex = GetEyeIndexFromVF(Input);
 	ResolvedView = ResolveViewFromVF(Input);
+	#if MOBILE_MULTI_VIEW
+		LayerIndex = EyeIndex;
+	#else
+		ViewportIndex = EyeIndex;
+	#endif
+#elif MOBILE_MULTI_VIEW
+	ResolvedView = ResolveView(ViewId);
+#else
+	ResolvedView = ResolveViewFromVF(Input);
+#endif
 
 	float4 WorldPos = VertexFactoryGetWorldPosition(Input);
 
diff --git a/Engine/Source/Runtime/Renderer/Private/CompositionLighting/PostProcessDeferredDecals.cpp b/Engine/Source/Runtime/Renderer/Private/CompositionLighting/PostProcessDeferredDecals.cpp
--- a/Engine/Source/Runtime/Renderer/Private/CompositionLighting/PostProcessDeferredDecals.cpp
+++ b/Engine/Source/Runtime/Renderer/Private/CompositionLighting/PostProcessDeferredDecals.cpp
@@ -192,13 +192,21 @@ void GetDeferredDecalPassParameters(
 
 	FRenderTargetBindingSlots& RenderTargets = PassParameters.RenderTargets;
 	PassParameters.RenderTargets.ShadingRateTexture = GVRSImageManager.GetVariableRateShadingImage(GraphBuilder, View, FVariableRateShadingImageManager::EVRSPassType::Decals);
-
+	PassParameters.RenderTargets.MultiViewCount = (View.bIsMobileMultiViewEnabled) ? 2 : (UE::StereoRenderUtils::FStereoShaderAspects(View.GetShaderPlatform()).IsMobileMultiViewEnabled() ? 1 : 0);
 	uint32 ColorTargetIndex = 0;
 
-	const auto AddColorTarget = [&](FRDGTextureRef Texture, ERenderTargetLoadAction LoadAction = ERenderTargetLoadAction::ELoad)
+	const auto AddColorTarget = [&](FRDGTextureRef Texture, ERenderTargetLoadAction LoadAction = ERenderTargetLoadAction::ELoad, FRDGTextureRef TextureArray = nullptr, bool bIsMobileMultiView = false)
 	{
-		checkf(Texture, TEXT("Attempting to bind decal render targets, but the texture is null."));
-		RenderTargets[ColorTargetIndex++] = FRenderTargetBinding(Texture, LoadAction);
+		if (bIsMobileMultiView)
+		{
+			checkf(TextureArray, TEXT("Attempting to bind decal render targets, but the texture array is null."));
+			RenderTargets[ColorTargetIndex++] = FRenderTargetBinding(TextureArray, LoadAction);
+		}
+		else
+		{
+			checkf(Texture, TEXT("Attempting to bind decal render targets, but the texture is null."));
+			RenderTargets[ColorTargetIndex++] = FRenderTargetBinding(Texture, LoadAction);
+		}
 	};
 
 	switch (RenderTargetMode)
@@ -224,13 +232,14 @@ void GetDeferredDecalPassParameters(
 
 		const FDBufferTextures& DBufferTextures = *Textures.DBufferTextures;
 
-		const ERenderTargetLoadAction LoadAction = DBufferTextures.DBufferA->HasBeenProduced()
+		const ERenderTargetLoadAction LoadAction = (DBufferTextures.DBufferA->HasBeenProduced() || DBufferTextures.DBufferATexArray->HasBeenProduced())
 			? ERenderTargetLoadAction::ELoad
 			: ERenderTargetLoadAction::EClear;
 
-		AddColorTarget(DBufferTextures.DBufferA, LoadAction);
-		AddColorTarget(DBufferTextures.DBufferB, LoadAction);
-		AddColorTarget(DBufferTextures.DBufferC, LoadAction);
+		bool bIsMobileMultiView = View.bIsMobileMultiViewEnabled || UE::StereoRenderUtils::FStereoShaderAspects(View.GetShaderPlatform()).IsMobileMultiViewEnabled();
+		AddColorTarget(DBufferTextures.DBufferA, LoadAction, DBufferTextures.DBufferATexArray, bIsMobileMultiView);
+		AddColorTarget(DBufferTextures.DBufferB, LoadAction, DBufferTextures.DBufferBTexArray, bIsMobileMultiView);
+		AddColorTarget(DBufferTextures.DBufferC, LoadAction, DBufferTextures.DBufferCTexArray, bIsMobileMultiView);
 
 		if (DBufferTextures.DBufferMask)
 		{
@@ -362,7 +371,7 @@ static bool RenderPreStencil(FRHICommandList& RHICmdList, const FViewInfo& View,
 	RHICmdList.SetStreamSource(0, GetUnitCubeVertexBuffer(), 0);
 
 	// Render decal mask
-	RHICmdList.DrawIndexedPrimitive(GetUnitCubeIndexBuffer(), 0, 0, 8, 0, UE_ARRAY_COUNT(GCubeIndices) / 3, 1);
+	RHICmdList.DrawIndexedPrimitive(GetUnitCubeIndexBuffer(), 0, 0, 8, 0, UE_ARRAY_COUNT(GCubeIndices) / 3, View.GetStereoPassInstanceFactor());
 
 	return true;
 }
@@ -708,7 +717,7 @@ void AddDeferredDecalPass(
 				}
 #endif // PSO_PRECACHING_VALIDATE
 
-				RHICmdList.DrawIndexedPrimitive(GetUnitCubeIndexBuffer(), 0, 0, 8, 0, UE_ARRAY_COUNT(GCubeIndices) / 3, 1);
+				RHICmdList.DrawIndexedPrimitive(GetUnitCubeIndexBuffer(), 0, 0, 8, 0, UE_ARRAY_COUNT(GCubeIndices) / 3, View.GetStereoPassInstanceFactor());
 			}
 		});
 	};
diff --git a/Engine/Source/Runtime/Renderer/Private/DBufferTextures.cpp b/Engine/Source/Runtime/Renderer/Private/DBufferTextures.cpp
--- a/Engine/Source/Runtime/Renderer/Private/DBufferTextures.cpp
+++ b/Engine/Source/Runtime/Renderer/Private/DBufferTextures.cpp
@@ -11,8 +11,8 @@
 
 bool FDBufferTextures::IsValid() const
 {
-	check(!DBufferA || (DBufferB && DBufferC));
-	return HasBeenProduced(DBufferA);
+	check((!DBufferA || (DBufferB && DBufferC)) || (!DBufferATexArray || (DBufferBTexArray && DBufferCTexArray)));
+	return HasBeenProduced(DBufferA) || HasBeenProduced(DBufferATexArray);
 }
 
 EDecalDBufferMaskTechnique GetDBufferMaskTechnique(EShaderPlatform ShaderPlatform)
@@ -43,18 +43,22 @@ FDBufferTexturesDesc GetDBufferTexturesDesc(FIntPoint Extent, EShaderPlatform Sh
 		const ETextureCreateFlags BaseFlags = WriteMaskFlags | TexCreate_ShaderResource | TexCreate_RenderTargetable;
 		
 		FRDGTextureDesc Desc = FRDGTextureDesc::Create2D(Extent, PF_B8G8R8A8, FClearValueBinding::None, BaseFlags);
+		FRDGTextureDesc ArrayDesc = FRDGTextureDesc::Create2DArray(Extent, PF_B8G8R8A8, FClearValueBinding::None, BaseFlags, 2);
 
-		Desc.Flags = BaseFlags | GFastVRamConfig.DBufferA;
-		Desc.ClearValue = FClearValueBinding::Black;
+		Desc.Flags = ArrayDesc.Flags = BaseFlags | GFastVRamConfig.DBufferA;
+		Desc.ClearValue = ArrayDesc.ClearValue = FClearValueBinding::Black;
 		DBufferTexturesDesc.DBufferADesc = Desc;
+		DBufferTexturesDesc.DBufferATexArrayDesc = ArrayDesc;
 
-		Desc.Flags = BaseFlags | GFastVRamConfig.DBufferB;
-		Desc.ClearValue = FClearValueBinding(FLinearColor(128.0f / 255.0f, 128.0f / 255.0f, 128.0f / 255.0f, 1));
+		Desc.Flags = ArrayDesc.Flags = BaseFlags | GFastVRamConfig.DBufferB;
+		Desc.ClearValue = ArrayDesc.ClearValue = FClearValueBinding(FLinearColor(128.0f / 255.0f, 128.0f / 255.0f, 128.0f / 255.0f, 1));
 		DBufferTexturesDesc.DBufferBDesc = Desc;
+		DBufferTexturesDesc.DBufferBTexArrayDesc = ArrayDesc;
 
-		Desc.Flags = BaseFlags | GFastVRamConfig.DBufferC;
-		Desc.ClearValue = FClearValueBinding(FLinearColor(0, 0, 0, 1));
+		Desc.Flags = ArrayDesc.Flags = BaseFlags | GFastVRamConfig.DBufferC;
+		Desc.ClearValue = ArrayDesc.ClearValue = FClearValueBinding(FLinearColor(0, 0, 0, 1));
 		DBufferTexturesDesc.DBufferCDesc = Desc;
+		DBufferTexturesDesc.DBufferCTexArrayDesc = ArrayDesc;
 
 		if (DBufferMaskTechnique == EDecalDBufferMaskTechnique::PerPixel)
 		{
@@ -93,6 +97,10 @@ FDBufferTextures CreateDBufferTextures(FRDGBuilder& GraphBuilder, FIntPoint Exte
 		DBufferTextures.DBufferB = GraphBuilder.CreateTexture(TexturesDesc.DBufferBDesc, TEXT("DBufferB"), TextureFlags);
 		DBufferTextures.DBufferC = GraphBuilder.CreateTexture(TexturesDesc.DBufferCDesc, TEXT("DBufferC"), TextureFlags);
 
+		DBufferTextures.DBufferATexArray = GraphBuilder.CreateTexture(TexturesDesc.DBufferATexArrayDesc, TEXT("DBufferATexArray"), TextureFlags);
+		DBufferTextures.DBufferBTexArray = GraphBuilder.CreateTexture(TexturesDesc.DBufferBTexArrayDesc, TEXT("DBufferBTexArray"), TextureFlags);
+		DBufferTextures.DBufferCTexArray = GraphBuilder.CreateTexture(TexturesDesc.DBufferCTexArrayDesc, TEXT("DBufferCTexArray"), TextureFlags);
+
 		if (DBufferMaskTechnique == EDecalDBufferMaskTechnique::PerPixel)
 		{
 			DBufferTextures.DBufferMask = GraphBuilder.CreateTexture(TexturesDesc.DBufferMaskDesc, TEXT("DBufferMask"));
@@ -113,6 +121,10 @@ FDBufferParameters GetDBufferParameters(FRDGBuilder& GraphBuilder, const FDBuffe
 	Parameters.DBufferATexture = SystemTextures.BlackAlphaOne;
 	Parameters.DBufferBTexture = SystemTextures.DefaultNormal8Bit;
 	Parameters.DBufferCTexture = SystemTextures.BlackAlphaOne;
+	Parameters.DBufferATextureArray = GSystemTextures.GetDefaultTexture(GraphBuilder, ETextureDimension::Texture2DArray, EPixelFormat::PF_B8G8R8A8, FClearValueBinding::Black);
+	Parameters.DBufferBTextureArray = GSystemTextures.GetDefaultTexture(GraphBuilder, ETextureDimension::Texture2DArray, EPixelFormat::PF_B8G8R8A8, FClearValueBinding::Black);
+	Parameters.DBufferCTextureArray = GSystemTextures.GetDefaultTexture(GraphBuilder, ETextureDimension::Texture2DArray, EPixelFormat::PF_B8G8R8A8, FClearValueBinding::Black);
+
 	Parameters.DBufferRenderMask = SystemTextures.White;
 
 	if (DBufferTextures.IsValid())
@@ -121,6 +133,10 @@ FDBufferParameters GetDBufferParameters(FRDGBuilder& GraphBuilder, const FDBuffe
 		Parameters.DBufferBTexture = DBufferTextures.DBufferB;
 		Parameters.DBufferCTexture = DBufferTextures.DBufferC;
 
+		Parameters.DBufferATextureArray = DBufferTextures.DBufferATexArray;
+		Parameters.DBufferBTextureArray = DBufferTextures.DBufferBTexArray;
+		Parameters.DBufferCTextureArray = DBufferTextures.DBufferCTexArray;
+
 		if (DBufferTextures.DBufferMask)
 		{
 			Parameters.DBufferRenderMask = DBufferTextures.DBufferMask;
diff --git a/Engine/Source/Runtime/Renderer/Private/DBufferTextures.h b/Engine/Source/Runtime/Renderer/Private/DBufferTextures.h
--- a/Engine/Source/Runtime/Renderer/Private/DBufferTextures.h
+++ b/Engine/Source/Runtime/Renderer/Private/DBufferTextures.h
@@ -25,6 +25,9 @@ struct FDBufferTexturesDesc
 	FRDGTextureDesc DBufferADesc;
 	FRDGTextureDesc DBufferBDesc;
 	FRDGTextureDesc DBufferCDesc;
+	FRDGTextureDesc DBufferATexArrayDesc;
+	FRDGTextureDesc DBufferBTexArrayDesc;
+	FRDGTextureDesc DBufferCTexArrayDesc;
 	FRDGTextureDesc DBufferMaskDesc;
 };
 
@@ -35,6 +38,9 @@ struct FDBufferTextures
 	FRDGTextureRef DBufferA = nullptr;
 	FRDGTextureRef DBufferB = nullptr;
 	FRDGTextureRef DBufferC = nullptr;
+	FRDGTextureRef DBufferATexArray = nullptr;
+	FRDGTextureRef DBufferBTexArray = nullptr;
+	FRDGTextureRef DBufferCTexArray = nullptr;
 	FRDGTextureRef DBufferMask = nullptr;
 };
 
@@ -45,6 +51,9 @@ BEGIN_SHADER_PARAMETER_STRUCT(FDBufferParameters, )
 	SHADER_PARAMETER_RDG_TEXTURE(Texture2D, DBufferATexture)
 	SHADER_PARAMETER_RDG_TEXTURE(Texture2D, DBufferBTexture)
 	SHADER_PARAMETER_RDG_TEXTURE(Texture2D, DBufferCTexture)
+	SHADER_PARAMETER_RDG_TEXTURE(Texture2DArray, DBufferATextureArray)
+	SHADER_PARAMETER_RDG_TEXTURE(Texture2DArray, DBufferBTextureArray)
+	SHADER_PARAMETER_RDG_TEXTURE(Texture2DArray, DBufferCTextureArray)
 	SHADER_PARAMETER_RDG_TEXTURE(Texture2D<uint>, DBufferRenderMask)
 	SHADER_PARAMETER_SAMPLER(SamplerState, DBufferATextureSampler)
 	SHADER_PARAMETER_SAMPLER(SamplerState, DBufferBTextureSampler)
diff --git a/Engine/Source/Runtime/Renderer/Private/DecalRenderingShared.cpp b/Engine/Source/Runtime/Renderer/Private/DecalRenderingShared.cpp
--- a/Engine/Source/Runtime/Renderer/Private/DecalRenderingShared.cpp
+++ b/Engine/Source/Runtime/Renderer/Private/DecalRenderingShared.cpp
@@ -370,8 +370,12 @@ namespace DecalRendering
 			ComponentToWorldMatrix.GetScaledAxis(EAxis::Y).SizeSquared() +
 			ComponentToWorldMatrix.GetScaledAxis(EAxis::Z).SizeSquared());
 
+		const bool bIsVisibleInFirstView = View.ViewFrustum.IntersectSphere(ComponentToWorldMatrix.GetOrigin(), ConservativeRadius);
+		const FViewInfo* InstancedView = View.GetInstancedView();
+		const bool bIsVisibleInSecondView = InstancedView ? InstancedView->ViewFrustum.IntersectSphere(ComponentToWorldMatrix.GetOrigin(), ConservativeRadius) : false;
+
 		// can be optimized as the test is too conservative (sphere instead of OBB)
-		if (ConservativeRadius < SMALL_NUMBER || !View.ViewFrustum.IntersectSphere(ComponentToWorldMatrix.GetOrigin(), ConservativeRadius))
+		if (ConservativeRadius < SMALL_NUMBER || !(bIsVisibleInFirstView || bIsVisibleInSecondView))
 		{
 			return false;
 		}
diff --git a/Engine/Source/Runtime/Renderer/Private/MobileShadingRenderer.cpp b/Engine/Source/Runtime/Renderer/Private/MobileShadingRenderer.cpp
--- a/Engine/Source/Runtime/Renderer/Private/MobileShadingRenderer.cpp
+++ b/Engine/Source/Runtime/Renderer/Private/MobileShadingRenderer.cpp
@@ -75,6 +75,7 @@
 #include "WaterInfoTextureRendering.h"
 #include "Rendering/CustomRenderPass.h"
 #include "GenerateMips.h"
+#include "StereoRenderUtils.h"
 
 uint32 GetShadowQuality();
 
@@ -799,7 +800,8 @@ void FMobileSceneRenderer::RenderFullDepthPrepass(FRDGBuilder& GraphBuilder, TAr
 		PassParameters->RenderTargets.DepthStencil = FDepthStencilBinding(SceneTextures.Depth.Target, ERenderTargetLoadAction::EClear, ERenderTargetLoadAction::EClear, FExclusiveDepthStencil::DepthWrite_StencilWrite);
 		PassParameters->View = View.GetShaderParameters();
 		PassParameters->MobileBasePass = CreateMobileBasePassUniformBuffer(GraphBuilder, View, EMobileBasePass::DepthPrePass, EMobileSceneTextureSetupMode::None);
-
+		//if the scenecolor isn't multiview but the app is, need to render as a single-view multiview due to shaders
+		PassParameters->RenderTargets.MultiViewCount = (View.bIsMobileMultiViewEnabled) ? 2 : (UE::StereoRenderUtils::FStereoShaderAspects(View.GetShaderPlatform()).IsMobileMultiViewEnabled() ? 1 : 0);
 		if (ViewContext.bIsLastView)
 		{
 			LastViewMobileBasePassUB = PassParameters->MobileBasePass;
@@ -1540,11 +1542,9 @@ void FMobileSceneRenderer::RenderForward(FRDGBuilder& GraphBuilder, FRDGTextureR
 	FRenderTargetBindingSlots BasePassRenderTargets = InitRenderTargetBindings_Forward(ViewFamilyTexture, SceneTextures);
 	BasePassRenderTargets.ShadingRateTexture = (!MainView.bIsSceneCapture && !MainView.bIsReflectionCapture && (NewShadingRateTarget != nullptr)) ? NewShadingRateTarget : nullptr;
 
-	static const auto CVarMobileMultiView = IConsoleManager::Get().FindTConsoleVariableDataInt(TEXT("vr.MobileMultiView"));
-	const bool bIsMultiViewApplication = (CVarMobileMultiView && CVarMobileMultiView->GetValueOnAnyThread() != 0);
-
 	//if the scenecolor isn't multiview but the app is, need to render as a single-view multiview due to shaders
-	BasePassRenderTargets.MultiViewCount = MainView.bIsMobileMultiViewEnabled ? 2 : (bIsMultiViewApplication ? 1 : 0);
+	BasePassRenderTargets.MultiViewCount = (MainView.bIsMobileMultiViewEnabled) ? 2 : (UE::StereoRenderUtils::FStereoShaderAspects(MainView.GetShaderPlatform()).IsMobileMultiViewEnabled() ? 1 : 0);
+
 
 	const FRDGSystemTextures& SystemTextures = FRDGSystemTextures::Get(GraphBuilder);
 
@@ -1900,11 +1900,8 @@ FRenderTargetBindingSlots FMobileSceneRenderer::InitRenderTargetBindings_Deferre
 	BasePassRenderTargets.NumOcclusionQueries = 0u;
 	BasePassRenderTargets.ShadingRateTexture = nullptr;
 	
-	static const auto CVarMobileMultiView = IConsoleManager::Get().FindTConsoleVariableDataInt(TEXT("vr.MobileMultiView"));
-	const bool bIsMultiViewApplication = (CVarMobileMultiView && CVarMobileMultiView->GetValueOnAnyThread() != 0);
-	
 	//if the scenecolor isn't multiview but the app is, need to render as a single-view multiview due to shaders
-	BasePassRenderTargets.MultiViewCount = Views[0].bIsMobileMultiViewEnabled ? 2 : (bIsMultiViewApplication ? 1 : 0);
+	BasePassRenderTargets.MultiViewCount = (Views[0].bIsMobileMultiViewEnabled) ? 2 : (UE::StereoRenderUtils::FStereoShaderAspects(Views[0].GetShaderPlatform()).IsMobileMultiViewEnabled() ? 1 : 0);
 
 	return BasePassRenderTargets;
 }
