Heterogeneous Volumes: Disable sparse-voxel defaults when using pre-shading pipeline.

#tests tested pre-shading pipeline

diff --git a/Engine/Source/Runtime/Renderer/Private/HeterogeneousVolumes/HeterogeneousVolumes.cpp b/Engine/Source/Runtime/Renderer/Private/HeterogeneousVolumes/HeterogeneousVolumes.cpp
--- a/Engine/Source/Runtime/Renderer/Private/HeterogeneousVolumes/HeterogeneousVolumes.cpp
+++ b/Engine/Source/Runtime/Renderer/Private/HeterogeneousVolumes/HeterogeneousVolumes.cpp
@@ -163,22 +163,22 @@ static TAutoConsoleVariable<int32> CVarHeterogeneousVolumesSparseVoxel(
 
 static TAutoConsoleVariable<int32> CVarHeterogeneousVolumesSparseVoxelGenerationMipBias(
 	TEXT("r.HeterogeneousVolumes.SparseVoxel.GenerationMipBias"),
-	3,
-	TEXT("Determines MIP bias for sparse voxel generation (Default = 3)"),
+	0,
+	TEXT("Determines MIP bias for sparse voxel generation (Default = 0)"),
 	ECVF_RenderThreadSafe
 );
 
 static TAutoConsoleVariable<int32> CVarHeterogeneousVolumesSparseVoxelPerTileCulling(
 	TEXT("r.HeterogeneousVolumes.SparseVoxel.PerTileCulling"),
-	1,
+	0,
 	TEXT("Enables sparse-voxel culling when using tiled rendering (Default = 1)"),
 	ECVF_RenderThreadSafe
 );
 
 static TAutoConsoleVariable<int32> CVarHeterogeneousVolumesSparseVoxelRefinement(
 	TEXT("r.HeterogeneousVolumes.SparseVoxel.Refinement"),
-	1,
-	TEXT("Uses hierarchical refinement to coalesce neighboring sparse-voxels (Default = 1)"),
+	0,
+	TEXT("Uses hierarchical refinement to coalesce neighboring sparse-voxels (Default = 0)"),
 	ECVF_RenderThreadSafe
 );
 
@@ -569,7 +569,7 @@ namespace HeterogeneousVolumes
 		{
 			// Determine the pixel-width at the near-plane
 			float TanHalfFOV = FMath::Tan(FMath::DegreesToRadians(View.FOV * 0.5));
-			float HalfWidth = View.UnconstrainedViewRect.Width() * 0.5;
+			float HalfWidth = View.UnconstrainedViewRect.Width() / HeterogeneousVolumes::GetDownsampleFactor() * 0.5;
 			float PixelWidth = TanHalfFOV / HalfWidth;
 
 			// Project to nearest distance of volume bounds
