Heterogeneous Volumes: Add AVSM lookup optimization to Preshading Pipeline

#tests Tested with preshading pipeline

diff --git a/Engine/Source/Runtime/Renderer/Private/HeterogeneousVolumes/HeterogeneousVolumesPreshadingPipeline.cpp b/Engine/Source/Runtime/Renderer/Private/HeterogeneousVolumes/HeterogeneousVolumesPreshadingPipeline.cpp
--- a/Engine/Source/Runtime/Renderer/Private/HeterogeneousVolumes/HeterogeneousVolumesPreshadingPipeline.cpp
+++ b/Engine/Source/Runtime/Renderer/Private/HeterogeneousVolumes/HeterogeneousVolumesPreshadingPipeline.cpp
@@ -209,7 +209,8 @@ class FRenderLightingCacheWithPreshadingCS : public FGlobalShader
 	SHADER_USE_PARAMETER_STRUCT(FRenderLightingCacheWithPreshadingCS, FGlobalShader);
 
 	class FLightingCacheMode : SHADER_PERMUTATION_INT("DIM_LIGHTING_CACHE_MODE", 2);
-	using FPermutationDomain = TShaderPermutationDomain<FLightingCacheMode>;
+	class FUseAdaptiveVolumetricShadowMap : SHADER_PERMUTATION_BOOL("DIM_USE_ADAPTIVE_VOLUMETRIC_SHADOW_MAP");
+	using FPermutationDomain = TShaderPermutationDomain<FLightingCacheMode, FUseAdaptiveVolumetricShadowMap>;
 
 	BEGIN_SHADER_PARAMETER_STRUCT(FParameters, )
 		// Scene data
@@ -229,6 +230,7 @@ class FRenderLightingCacheWithPreshadingCS : public FGlobalShader
 		SHADER_PARAMETER_STRUCT_INCLUDE(FVolumeShadowingShaderParameters, VolumeShadowingShaderParameters)
 		SHADER_PARAMETER_STRUCT_INCLUDE(FVirtualShadowMapSamplingParameters, VirtualShadowMapSamplingParameters)
 		SHADER_PARAMETER(int32, VirtualShadowMapId)
+		SHADER_PARAMETER_RDG_UNIFORM_BUFFER(FAdaptiveVolumetricShadowMapUniformBufferParameters, AVSM)
 
 		// Volume structures
 		SHADER_PARAMETER_RDG_UNIFORM_BUFFER(FSparseVoxelUniformBufferParameters, SparseVoxelUniformBuffer)
@@ -295,8 +297,9 @@ class FRenderSingleScatteringWithPreshadingCS : public FGlobalShader
 	class FUseInscatteringVolume : SHADER_PERMUTATION_BOOL("DIM_USE_INSCATTERING_VOLUME");
 	class FUseLumenGI : SHADER_PERMUTATION_BOOL("DIM_USE_LUMEN_GI");
 	class FWriteVelocity : SHADER_PERMUTATION_BOOL("DIM_WRITE_VELOCITY");
+	class FUseAdaptiveVolumetricShadowMap : SHADER_PERMUTATION_BOOL("DIM_USE_ADAPTIVE_VOLUMETRIC_SHADOW_MAP");
 	class FDebugDim : SHADER_PERMUTATION_BOOL("DIM_DEBUG");
-	using FPermutationDomain = TShaderPermutationDomain<FApplyShadowTransmittanceDim, FVoxelCullingDim, FSparseVoxelTracingDim, FUseTransmittanceVolume, FUseInscatteringVolume, FUseLumenGI, FWriteVelocity, FDebugDim>;
+	using FPermutationDomain = TShaderPermutationDomain<FApplyShadowTransmittanceDim, FVoxelCullingDim, FSparseVoxelTracingDim, FUseTransmittanceVolume, FUseInscatteringVolume, FUseLumenGI, FWriteVelocity, FUseAdaptiveVolumetricShadowMap, FDebugDim>;
 
 	BEGIN_SHADER_PARAMETER_STRUCT(FParameters, )
 		// Scene data
@@ -315,6 +318,7 @@ class FRenderSingleScatteringWithPreshadingCS : public FGlobalShader
 		SHADER_PARAMETER_STRUCT_INCLUDE(FVolumeShadowingShaderParameters, VolumeShadowingShaderParameters)
 		SHADER_PARAMETER_STRUCT_INCLUDE(FVirtualShadowMapSamplingParameters, VirtualShadowMapSamplingParameters)
 		SHADER_PARAMETER(int32, VirtualShadowMapId)
+		SHADER_PARAMETER_RDG_UNIFORM_BUFFER(FAdaptiveVolumetricShadowMapUniformBufferParameters, AVSM)
 
 		// Atmosphere
 		SHADER_PARAMETER_RDG_UNIFORM_BUFFER(FFogUniformParameters, FogStruct)
@@ -502,6 +506,7 @@ void RenderLightingCacheWithPreshadingCompute(
 			PassParameters->VirtualShadowMapId = -1;
 		}
 		PassParameters->VirtualShadowMapSamplingParameters = VirtualShadowMapArray.GetSamplingParameters(GraphBuilder);
+		PassParameters->AVSM = HeterogeneousVolumes::GetAdaptiveVolumetricShadowMapUniformBuffer(GraphBuilder, View.ViewState, LightSceneInfo);
 
 		// Output
 		PassParameters->RWLightingCacheTexture = GraphBuilder.CreateUAV(LightingCacheTexture);
@@ -520,9 +525,11 @@ void RenderLightingCacheWithPreshadingCompute(
 		PassName = FString::Printf(TEXT("RenderLightingCacheWithPreshadingCS [%s] (Light = %s)"), *ModeName, *LightName);
 	}
 #endif // WANTS_DRAW_MESH_EVENTS
+	bool bUseAVSM = HeterogeneousVolumes::UseAdaptiveVolumetricShadowMapForSelfShadowing(HeterogeneousVolumeInterface->GetPrimitiveSceneProxy());
 
 	FRenderLightingCacheWithPreshadingCS::FPermutationDomain PermutationVector;
 	PermutationVector.Set<FRenderLightingCacheWithPreshadingCS::FLightingCacheMode>(HeterogeneousVolumes::GetLightingCacheMode() - 1);
+	PermutationVector.Set<FRenderLightingCacheWithPreshadingCS::FUseAdaptiveVolumetricShadowMap>(bUseAVSM);
 	TShaderRef<FRenderLightingCacheWithPreshadingCS> ComputeShader = View.ShaderMap->GetShader<FRenderLightingCacheWithPreshadingCS>(PermutationVector);
 
 	FIntVector GroupCount = HeterogeneousVolumes::GetLightingCacheResolution(HeterogeneousVolumeInterface, LODFactor);
@@ -619,6 +626,7 @@ void RenderSingleScatteringWithPreshadingCompute(
 			PassParameters->VirtualShadowMapId = -1;
 		}
 		PassParameters->VirtualShadowMapSamplingParameters = VirtualShadowMapArray.GetSamplingParameters(GraphBuilder);
+		PassParameters->AVSM = HeterogeneousVolumes::GetAdaptiveVolumetricShadowMapUniformBuffer(GraphBuilder, View.ViewState, LightSceneInfo);
 
 		TRDGUniformBufferRef<FFogUniformParameters> FogBuffer = CreateFogUniformBuffer(GraphBuilder, View);
 		PassParameters->FogStruct = FogBuffer;
@@ -677,6 +685,7 @@ void RenderSingleScatteringWithPreshadingCompute(
 	{
 		FSceneRenderer::GetLightNameForDrawEvent(LightSceneInfo->Proxy, LightName);
 	}
+	bool bUseAVSM = HeterogeneousVolumes::UseAdaptiveVolumetricShadowMapForSelfShadowing(HeterogeneousVolumeInterface->GetPrimitiveSceneProxy());
 
 	FRenderSingleScatteringWithPreshadingCS::FPermutationDomain PermutationVector;
 	PermutationVector.Set<FRenderSingleScatteringWithPreshadingCS::FApplyShadowTransmittanceDim>(bApplyShadowTransmittance);
@@ -686,6 +695,7 @@ void RenderSingleScatteringWithPreshadingCompute(
 	PermutationVector.Set<FRenderSingleScatteringWithPreshadingCS::FUseInscatteringVolume>(HeterogeneousVolumes::UseLightingCacheForInscattering());
 	PermutationVector.Set<FRenderSingleScatteringWithPreshadingCS::FUseLumenGI>(HeterogeneousVolumes::UseIndirectLighting() && View.GetLumenTranslucencyGIVolume().Texture0 != nullptr);
 	PermutationVector.Set<FRenderSingleScatteringWithPreshadingCS::FWriteVelocity>(bWriteVelocity);
+	PermutationVector.Set<FRenderSingleScatteringWithPreshadingCS::FUseAdaptiveVolumetricShadowMap>(bUseAVSM);
 	PermutationVector.Set<FRenderSingleScatteringWithPreshadingCS::FDebugDim>(HeterogeneousVolumes::GetDebugMode() != 0);
 	TShaderRef<FRenderSingleScatteringWithPreshadingCS> ComputeShader = View.ShaderMap->GetShader<FRenderSingleScatteringWithPreshadingCS>(PermutationVector);
 	FComputeShaderUtils::AddPass(
@@ -1551,7 +1561,9 @@ void RenderWithPreshading(
 )
 {
 	// Determine baking voxel resolution
+	float LODFactor = HeterogeneousVolumes::CalcLODFactor(View, HeterogeneousVolumeInterface);
 	FIntVector VolumeResolution = HeterogeneousVolumes::GetVolumeResolution(HeterogeneousVolumeInterface);
+	// TODO: Modify volume resolution by LODFactor??
 
 	// Create baked material grids
 	uint32 NumMips = FMath::Log2(float(FMath::Min(FMath::Min(VolumeResolution.X, VolumeResolution.Y), VolumeResolution.Z))) + 1;
@@ -1672,9 +1684,9 @@ void RenderWithPreshading(
 		SparseVoxelUniformBufferParameters->MaxTraceDistance = HeterogeneousVolumes::GetMaxTraceDistance();
 		SparseVoxelUniformBufferParameters->MaxShadowTraceDistance = HeterogeneousVolumes::GetMaxShadowTraceDistance();
 		SparseVoxelUniformBufferParameters->StepSize = HeterogeneousVolumes::GetStepSize();
-		SparseVoxelUniformBufferParameters->StepFactor = HeterogeneousVolumeInterface->GetStepFactor();
+		SparseVoxelUniformBufferParameters->StepFactor = HeterogeneousVolumeInterface->GetStepFactor() * LODFactor;
 		SparseVoxelUniformBufferParameters->ShadowStepSize = HeterogeneousVolumes::GetShadowStepSize();
-		SparseVoxelUniformBufferParameters->ShadowStepFactor = HeterogeneousVolumeInterface->GetShadowStepFactor();
+		SparseVoxelUniformBufferParameters->ShadowStepFactor = HeterogeneousVolumeInterface->GetShadowStepFactor() * LODFactor;
 		SparseVoxelUniformBufferParameters->bApplyHeightFog = HeterogeneousVolumes::ShouldApplyHeightFog();
 		SparseVoxelUniformBufferParameters->bApplyVolumetricFog = HeterogeneousVolumes::ShouldApplyVolumetricFog();
 	}
