Re-create reflection data uniform buffer only when reflection state has changed, instead of doing it everyframe

diff --git a/Engine/Source/Runtime/Renderer/Private/ReflectionEnvironment.cpp b/Engine/Source/Runtime/Renderer/Private/ReflectionEnvironment.cpp
--- a/Engine/Source/Runtime/Renderer/Private/ReflectionEnvironment.cpp
+++ b/Engine/Source/Runtime/Renderer/Private/ReflectionEnvironment.cpp
@@ -517,6 +517,9 @@ void FReflectionEnvironmentSceneData::Reset(FScene* Scene)
 			bRegisteredReflectionCapturesHasChanged = true;
 			AllocatedReflectionCaptureStateHasChanged = false;
 
+			ReflectionCaptureUniformBuffer.SafeRelease();
+			MobileReflectionCaptureUniformBuffer.SafeRelease();
+
 			RegisteredReflectionCaptures.Empty();
 			RegisteredReflectionCapturePositionAndRadius.Empty();
 			CubemapArray.Reset();
diff --git a/Engine/Source/Runtime/Renderer/Private/RendererScene.cpp b/Engine/Source/Runtime/Renderer/Private/RendererScene.cpp
--- a/Engine/Source/Runtime/Renderer/Private/RendererScene.cpp
+++ b/Engine/Source/Runtime/Renderer/Private/RendererScene.cpp
@@ -5373,6 +5373,24 @@ FLightSceneChangeSet FScene::UpdateAllLightSceneInfos(FRDGBuilder& GraphBuilder)
 	return FLightSceneChangeSet{ ChangeSet.RemovedLightIds, ChangeSet.AddedLightIds, ChangeSet.TransformUpdatedLightIds, ChangeSet.ColorUpdatedLightIds };
 }
 
+template<class T>
+static void CreateReflectionCaptureUniformBuffer(const TArray<FReflectionCaptureSortData>& SortedCaptures, TUniformBufferRef<T>& OutReflectionCaptureUniformBuffer)
+{
+	T SamplePositionsBuffer;
+	for (int32 CaptureIndex = 0; CaptureIndex < SortedCaptures.Num(); CaptureIndex++)
+	{
+		SamplePositionsBuffer.PositionHighAndRadius[CaptureIndex] = FVector4f(SortedCaptures[CaptureIndex].Position.High, SortedCaptures[CaptureIndex].Radius);
+		SamplePositionsBuffer.PositionLow[CaptureIndex] = FVector4f(SortedCaptures[CaptureIndex].Position.Low, 0);
+
+		SamplePositionsBuffer.CaptureProperties[CaptureIndex] = SortedCaptures[CaptureIndex].CaptureProperties;
+		SamplePositionsBuffer.CaptureOffsetAndAverageBrightness[CaptureIndex] = SortedCaptures[CaptureIndex].CaptureOffsetAndAverageBrightness;
+		SamplePositionsBuffer.BoxTransform[CaptureIndex] = SortedCaptures[CaptureIndex].BoxTransform;
+		SamplePositionsBuffer.BoxScales[CaptureIndex] = SortedCaptures[CaptureIndex].BoxScales;
+	}
+
+	OutReflectionCaptureUniformBuffer = TUniformBufferRef<T>::CreateUniformBufferImmediate(SamplePositionsBuffer, UniformBuffer_MultiFrame);
+}
+
 void UpdateReflectionSceneData(FScene* Scene)
 {
 	QUICK_SCOPE_CYCLE_COUNTER(STAT_UpdateReflectionSceneData)
@@ -5451,18 +5469,37 @@ void UpdateReflectionSceneData(FScene* Scene)
 		ReflectionSceneData.SortedCaptures[CaptureIndex].CaptureProxy->SortedCaptureIndex = CaptureIndex;
 	}
 
+	// Create uniform buffers with a sorted captures
+	if (ReflectionSceneData.bRegisteredReflectionCapturesHasChanged || 
+		ReflectionSceneData.AllocatedReflectionCaptureStateHasChanged)
+	{
+		ReflectionSceneData.ReflectionCaptureUniformBuffer.SafeRelease();
+		ReflectionSceneData.MobileReflectionCaptureUniformBuffer.SafeRelease();
+
+		if (IsMobilePlatform(Scene->GetShaderPlatform()))
+		{
+			CreateReflectionCaptureUniformBuffer(ReflectionSceneData.SortedCaptures, ReflectionSceneData.MobileReflectionCaptureUniformBuffer);
+		}
+		else
+		{
+			CreateReflectionCaptureUniformBuffer(ReflectionSceneData.SortedCaptures, ReflectionSceneData.ReflectionCaptureUniformBuffer);
+		}
+	}
 
 	// If SortedCaptures change, then in case of forward renderer all scene primitives need to be updated, as they 
 	// store index into sorted reflection capture uniform buffer for the forward renderer.
-	if (IsForwardShadingEnabled(Scene->GetShaderPlatform()) && ReflectionSceneData.AllocatedReflectionCaptureStateHasChanged)
+	if (ReflectionSceneData.AllocatedReflectionCaptureStateHasChanged)
 	{
-		const int32 NumPrimitives = Scene->Primitives.Num();
-		for (int32 PrimitiveIndex = 0; PrimitiveIndex < NumPrimitives; ++PrimitiveIndex)
+		if (IsForwardShadingEnabled(Scene->GetShaderPlatform()))
 		{
-			Scene->PrimitivesNeedingUniformBufferUpdate[PrimitiveIndex] = true;
-		}
+			const int32 NumPrimitives = Scene->Primitives.Num();
+			for (int32 PrimitiveIndex = 0; PrimitiveIndex < NumPrimitives; ++PrimitiveIndex)
+			{
+				Scene->PrimitivesNeedingUniformBufferUpdate[PrimitiveIndex] = true;
+			}
 
-		Scene->GPUScene.bUpdateAllPrimitives = true;
+			Scene->GPUScene.bUpdateAllPrimitives = true;
+		}
 		ReflectionSceneData.AllocatedReflectionCaptureStateHasChanged = false;
 	}
 
diff --git a/Engine/Source/Runtime/Renderer/Private/ScenePrivate.h b/Engine/Source/Runtime/Renderer/Private/ScenePrivate.h
--- a/Engine/Source/Runtime/Renderer/Private/ScenePrivate.h
+++ b/Engine/Source/Runtime/Renderer/Private/ScenePrivate.h
@@ -1817,6 +1817,10 @@ public:
 	int32 NumBoxCaptures;
 	int32 NumSphereCaptures;
 
+	// Uniforms buffers with a sorted captures
+	TUniformBufferRef<FReflectionCaptureShaderData> ReflectionCaptureUniformBuffer;
+	TUniformBufferRef<FMobileReflectionCaptureShaderData> MobileReflectionCaptureUniformBuffer;
+
 	/** 
 	 * Game thread list of reflection components that have been allocated in the cubemap array. 
 	 * These are not necessarily all visible or being rendered, but their scene state is stored in the cubemap array.
diff --git a/Engine/Source/Runtime/Renderer/Private/SceneVisibility.cpp b/Engine/Source/Runtime/Renderer/Private/SceneVisibility.cpp
--- a/Engine/Source/Runtime/Renderer/Private/SceneVisibility.cpp
+++ b/Engine/Source/Runtime/Renderer/Private/SceneVisibility.cpp
@@ -5691,51 +5691,21 @@ void FDeferredShadingSceneRenderer::BeginInitViews(
 	PostVisibilityFrameSetup(TaskDatas.ILCUpdatePrim);
 }
 
-template<class T>
-void CreateReflectionCaptureUniformBuffer(const TArray<FReflectionCaptureSortData>& SortedCaptures, TUniformBufferRef<T>& OutReflectionCaptureUniformBuffer)
-{
-	T SamplePositionsBuffer;
-	for (int32 CaptureIndex = 0; CaptureIndex < SortedCaptures.Num(); CaptureIndex++)
-	{
-		SamplePositionsBuffer.PositionHighAndRadius[CaptureIndex] = FVector4f(SortedCaptures[CaptureIndex].Position.High, SortedCaptures[CaptureIndex].Radius);
-		SamplePositionsBuffer.PositionLow[CaptureIndex] = FVector4f(SortedCaptures[CaptureIndex].Position.Low, 0);
-
-		SamplePositionsBuffer.CaptureProperties[CaptureIndex] = SortedCaptures[CaptureIndex].CaptureProperties;
-		SamplePositionsBuffer.CaptureOffsetAndAverageBrightness[CaptureIndex] = SortedCaptures[CaptureIndex].CaptureOffsetAndAverageBrightness;
-		SamplePositionsBuffer.BoxTransform[CaptureIndex] = SortedCaptures[CaptureIndex].BoxTransform;
-		SamplePositionsBuffer.BoxScales[CaptureIndex] = SortedCaptures[CaptureIndex].BoxScales;
-	}
-
-	OutReflectionCaptureUniformBuffer = TUniformBufferRef<T>::CreateUniformBufferImmediate(SamplePositionsBuffer, UniformBuffer_SingleFrame);
-}
-
 void FSceneRenderer::SetupSceneReflectionCaptureBuffer(FRHICommandListImmediate& RHICmdList)
 {
 	const TArray<FReflectionCaptureSortData>& SortedCaptures = Scene->ReflectionSceneData.SortedCaptures;
 
-	TUniformBufferRef<FMobileReflectionCaptureShaderData> MobileReflectionCaptureUniformBuffer;
-	TUniformBufferRef<FReflectionCaptureShaderData> ReflectionCaptureUniformBuffer;
-
-	if (IsMobilePlatform(ShaderPlatform))
-	{
-		CreateReflectionCaptureUniformBuffer(SortedCaptures, MobileReflectionCaptureUniformBuffer);
-	}
-	else
-	{
-		CreateReflectionCaptureUniformBuffer(SortedCaptures, ReflectionCaptureUniformBuffer);
-	}
-
 	for (int32 ViewIndex = 0; ViewIndex < Views.Num(); ++ViewIndex)
 	{
 		FViewInfo& View = Views[ViewIndex];
 
 		if (IsMobilePlatform(ShaderPlatform))
 		{
-			View.MobileReflectionCaptureUniformBuffer = MobileReflectionCaptureUniformBuffer;
+			View.MobileReflectionCaptureUniformBuffer = Scene->ReflectionSceneData.MobileReflectionCaptureUniformBuffer;
 		}
 		else
 		{
-			View.ReflectionCaptureUniformBuffer = ReflectionCaptureUniformBuffer;
+			View.ReflectionCaptureUniformBuffer = Scene->ReflectionSceneData.ReflectionCaptureUniformBuffer;
 		}
 		
 		View.NumBoxReflectionCaptures = 0;
