Fixed reflection captures missing VSM shadowed stationary lights as they are marked as not having static shadowing and thus are excluded when the view requests static scene only.

diff --git a/Engine/Source/Runtime/Renderer/Private/LightSceneInfo.cpp b/Engine/Source/Runtime/Renderer/Private/LightSceneInfo.cpp
--- a/Engine/Source/Runtime/Renderer/Private/LightSceneInfo.cpp
+++ b/Engine/Source/Runtime/Renderer/Private/LightSceneInfo.cpp
@@ -28,6 +28,14 @@ FAutoConsoleVariableRef CVarRecordInteractionShadowPrimitives(
 	TEXT(""),
 	ECVF_RenderThreadSafe);
 
+static int32 GTestMobilityForStaticSceneMembership = 1;
+FAutoConsoleVariableRef CVarTestMobilityForStaticSceneMembership(
+	TEXT("r.Light.TestMobilityForStaticSceneMembership"),
+	GTestMobilityForStaticSceneMembership,
+	TEXT("Deprecated (UE5.5): Temporary flag to switch back to the old behavior (testing the cast static shadow flag)\n.")
+	TEXT("  The old behavior checked the HasStaticShadowing flag on the proxy, but that is cleared for VSMs so the new behavior tests the mobility instead."),
+	ECVF_RenderThreadSafe | ECVF_ReadOnly);
+
 void FLightSceneInfoCompact::Init(FLightSceneInfo* InLightSceneInfo)
 {
 	LightSceneInfo = InLightSceneInfo;
@@ -176,6 +184,23 @@ FBoxCenterAndExtent FLightSceneInfo::GetBoundingBox() const
 	return FBoxCenterAndExtent(BoundingSphere.Center, FVector(BoundingSphere.W, BoundingSphere.W, BoundingSphere.W));
 }
 
+inline static bool IsInDesiredSceneSubset(const FViewInfo& View, FLightSceneProxy *Proxy)
+{
+	if (!View.bStaticSceneOnly)
+	{
+		return true;
+	}
+
+	if (GTestMobilityForStaticSceneMembership)
+	{
+		return !Proxy->IsMovable();
+	}
+	else
+	{
+		return Proxy->HasStaticShadowing();
+	}
+}
+
 bool FLightSceneInfo::ShouldRenderLight(const FViewInfo& View, bool bOffscreen) const
 {
 	// Only render the light if it is in the view frustum
@@ -215,7 +240,7 @@ bool FLightSceneInfo::ShouldRenderLight(const FViewInfo& View, bool bOffscreen)
 
 	return bLocalVisible
 		// Only render lights with static shadowing for reflection captures, since they are only captured at edit time
-		&& (!View.bStaticSceneOnly || Proxy->HasStaticShadowing())
+		&& IsInDesiredSceneSubset(View, Proxy)
 		// Only render lights in the default channel, or if there are any primitives outside the default channel
 		&& (Proxy->GetLightingChannelMask() & GetDefaultLightingChannelMask() || View.bUsesLightingChannels || bOffscreen);
 }
