Improvements to FRayTracingScene
- added new types FInstanceHandle and FInstanceRange instead of using basic indices and offsets.
- improves safety and makes it possible to do extra checks.

diff --git a/Engine/Source/Runtime/Renderer/Private/RayTracing/RayTracing.cpp b/Engine/Source/Runtime/Renderer/Private/RayTracing/RayTracing.cpp
--- a/Engine/Source/Runtime/Renderer/Private/RayTracing/RayTracing.cpp
+++ b/Engine/Source/Runtime/Renderer/Private/RayTracing/RayTracing.cpp
@@ -1245,8 +1245,8 @@ namespace RayTracing
 
 			struct FAutoInstanceBatch
 			{
-				int32 Index = INDEX_NONE;
-				int32 DecalIndex = INDEX_NONE;
+				FRayTracingScene::FInstanceHandle InstanceHandle = FRayTracingScene::INVALID_INSTANCE_HANDLE;
+				FRayTracingScene::FInstanceHandle DecalInstanceHandle = FRayTracingScene::INVALID_INSTANCE_HANDLE;
 
 				// Copies the next InstanceSceneDataOffset and user data into the current batch, returns true if arrays were re-allocated.
 				bool Add(FRayTracingScene& InRayTracingScene, uint32 InInstanceSceneDataOffset)
@@ -1353,9 +1353,9 @@ namespace RayTracing
 
 							bool bReallocated = InstanceBatch.Add(RayTracingScene, SceneInfo->GetInstanceSceneDataOffset());
 
-							if(InstanceBatch.Index != INDEX_NONE)
+							if(InstanceBatch.InstanceHandle.IsValid())
 							{
-								FRayTracingGeometryInstance& RayTracingInstance = RayTracingScene.GetInstance(InstanceBatch.Index);
+								FRayTracingGeometryInstance& RayTracingInstance = RayTracingScene.GetInstance(InstanceBatch.InstanceHandle);
 								++RayTracingInstance.NumTransforms;
 								check(RayTracingInstance.NumTransforms == InstanceBatch.Cursor); // sanity check
 
@@ -1366,9 +1366,9 @@ namespace RayTracing
 								}
 							}
 
-							if (InstanceBatch.DecalIndex != INDEX_NONE)
+							if (InstanceBatch.DecalInstanceHandle.IsValid())
 							{
-								FRayTracingGeometryInstance& RayTracingInstance = RayTracingScene.GetInstance(InstanceBatch.DecalIndex);
+								FRayTracingGeometryInstance& RayTracingInstance = RayTracingScene.GetInstance(InstanceBatch.DecalInstanceHandle);
 								++RayTracingInstance.NumTransforms;
 								check(RayTracingInstance.NumTransforms == InstanceBatch.Cursor); // sanity check
 
@@ -1409,7 +1409,7 @@ namespace RayTracing
 							}
 							AddDebugRayTracingInstanceFlags(RayTracingInstance.Flags);
 
-							InstanceBatch.Index = INDEX_NONE;
+							InstanceBatch.InstanceHandle = FRayTracingScene::INVALID_INSTANCE_HANDLE;
 							uint32 GlobalSegmentIndex = INDEX_NONE;
 
 							if (bNeedMainInstance)
@@ -1419,10 +1419,11 @@ namespace RayTracing
 
 								RayTracingInstance.InstanceContributionToHitGroupIndex = CalculateInstanceContributionToHitGroupIndex(GlobalSegmentIndex);
 								RayTracingInstance.LayerIndex = (uint8)ERayTracingSceneLayer::Base;
-								InstanceBatch.Index = RayTracingScene.AddInstance(RayTracingInstance, SceneProxy, false);
+
+								InstanceBatch.InstanceHandle = RayTracingScene.AddInstance(RayTracingInstance, SceneProxy, false);
 							}
 
-							InstanceBatch.DecalIndex = INDEX_NONE;
+							InstanceBatch.DecalInstanceHandle = FRayTracingScene::INVALID_INSTANCE_HANDLE;
 							uint32 DecalGlobalSegmentIndex = INDEX_NONE;
 
 							if (bNeedDecalInstance)
@@ -1434,7 +1435,7 @@ namespace RayTracing
 								DecalRayTracingInstance.InstanceContributionToHitGroupIndex = CalculateInstanceContributionToHitGroupIndex(DecalGlobalSegmentIndex);
 								DecalRayTracingInstance.LayerIndex = (uint8)ERayTracingSceneLayer::Decals;
 
-								InstanceBatch.DecalIndex = RayTracingScene.AddInstance(MoveTemp(DecalRayTracingInstance), SceneProxy, false);
+								InstanceBatch.DecalInstanceHandle = RayTracingScene.AddInstance(MoveTemp(DecalRayTracingInstance), SceneProxy, false);
 							}
 
 							for (int32 CommandIndex : RelevantPrimitive.CachedRayTracingMeshCommandIndices)
@@ -1470,7 +1471,7 @@ namespace RayTracing
 				{
 					TRACE_CPUPROFILER_EVENT_SCOPE(RayTracingScene_AddCachedStaticInstances);
 					
-					const uint32 BaseCachedStaticInstanceIndex = RayTracingScene.AddInstancesUninitialized(NumCachedStaticSceneInstances);
+					const FRayTracingScene::FInstanceRange CachedStaticInstanceRange = RayTracingScene.AllocateInstanceRangeUninitialized(NumCachedStaticSceneInstances);
 					const uint32 BaseCachedVisibleMeshCommandsIndex = VisibleRayTracingMeshCommands.AddUninitialized(NumCachedStaticVisibleMeshCommands);
 					const uint32 BaseCachedGlobalSegmentIndex = RayTracingScene.NumSegments;
 					RayTracingScene.NumSegments += NumCachedStaticVisibleMeshCommands;
@@ -1480,7 +1481,7 @@ namespace RayTracing
 						TEXT("RayTracingScene_AddCachedStaticInstances_ParallelFor"),
 						RelevantCachedStaticPrimitives.Num(),
 						MinBatchSize,
-						[this, BaseCachedStaticInstanceIndex, BaseCachedVisibleMeshCommandsIndex, BaseCachedGlobalSegmentIndex](int32 Index)
+						[this, CachedStaticInstanceRange, BaseCachedVisibleMeshCommandsIndex, BaseCachedGlobalSegmentIndex](int32 Index)
 					{
 						const FRelevantPrimitive& RelevantPrimitive = RelevantCachedStaticPrimitives[Index];
 						const int32 PrimitiveIndex = RelevantPrimitive.PrimitiveIndex;
@@ -1514,7 +1515,7 @@ namespace RayTracing
 
 						check(RelevantPrimitive.CachedRayTracingInstance);
 
-						int32 InstanceIndex = BaseCachedStaticInstanceIndex + GatherContexts[RelevantPrimitive.ContextIndex].SceneInstanceOffset + RelevantPrimitive.RelativeSceneInstanceOffset;
+						int32 InstanceIndexInRange = GatherContexts[RelevantPrimitive.ContextIndex].SceneInstanceOffset + RelevantPrimitive.RelativeSceneInstanceOffset;
 
 						if (bNeedMainInstance)
 						{
@@ -1523,8 +1524,8 @@ namespace RayTracing
 							RayTracingInstance.InstanceContributionToHitGroupIndex = CalculateInstanceContributionToHitGroupIndex(MainGlobalSegmentIndex);
 							AddDebugRayTracingInstanceFlags(RayTracingInstance.Flags);
 
-							RayTracingScene.SetInstance(InstanceIndex, MoveTemp(RayTracingInstance), SceneProxy, false);
-							InstanceIndex++;
+							RayTracingScene.SetInstance(CachedStaticInstanceRange, InstanceIndexInRange, MoveTemp(RayTracingInstance), SceneProxy, false);
+							InstanceIndexInRange++;
 						}
 
 						if (bNeedDecalInstance)
@@ -1534,8 +1535,8 @@ namespace RayTracing
 							DecalRayTracingInstance.InstanceContributionToHitGroupIndex = CalculateInstanceContributionToHitGroupIndex(DecalGlobalSegmentIndex);
 							AddDebugRayTracingInstanceFlags(DecalRayTracingInstance.Flags);
 
-							RayTracingScene.SetInstance(InstanceIndex, MoveTemp(DecalRayTracingInstance), SceneProxy, false);
-							InstanceIndex++;
+							RayTracingScene.SetInstance(CachedStaticInstanceRange, InstanceIndexInRange, MoveTemp(DecalRayTracingInstance), SceneProxy, false);
+							InstanceIndexInRange++;
 						}
 
 						const int32 VisibleMeshCommandOffset = BaseCachedVisibleMeshCommandsIndex + GatherContexts[RelevantPrimitive.ContextIndex].VisibleMeshCommandOffset + RelevantPrimitive.RelativeVisibleMeshCommandOffset;
diff --git a/Engine/Source/Runtime/Renderer/Private/RayTracing/RayTracingScene.cpp b/Engine/Source/Runtime/Renderer/Private/RayTracing/RayTracingScene.cpp
--- a/Engine/Source/Runtime/Renderer/Private/RayTracing/RayTracingScene.cpp
+++ b/Engine/Source/Runtime/Renderer/Private/RayTracing/RayTracingScene.cpp
@@ -31,6 +31,8 @@ BEGIN_SHADER_PARAMETER_STRUCT(FBuildInstanceBufferPassParams, )
 	SHADER_PARAMETER_RDG_UNIFORM_BUFFER(FSceneUniformParameters, Scene)
 END_SHADER_PARAMETER_STRUCT()
 
+const FRayTracingScene::FInstanceHandle FRayTracingScene::INVALID_INSTANCE_HANDLE = FInstanceHandle();
+
 FRayTracingScene::FRayTracingScene()
 {
 
@@ -474,7 +476,7 @@ FRDGBufferSRVRef FRayTracingScene::GetLayerView(ERayTracingSceneLayer Layer) con
 	return LayerSRVs[uint8(Layer)];
 }
 
-uint32 FRayTracingScene::AddInstance(FRayTracingGeometryInstance Instance, const FPrimitiveSceneProxy* Proxy, bool bDynamic)
+FRayTracingScene::FInstanceHandle FRayTracingScene::AddInstance(FRayTracingGeometryInstance Instance, const FPrimitiveSceneProxy* Proxy, bool bDynamic)
 {
 	FRHIRayTracingGeometry* GeometryRHI = Instance.GeometryRHI;
 
@@ -494,10 +496,10 @@ uint32 FRayTracingScene::AddInstance(FRayTracingGeometryInstance Instance, const
 		check(Instances.Num() == InstancesDebugData.Num());
 	}
 
-	return InstanceIndex;
+	return { InstanceIndex };
 }
 
-uint32 FRayTracingScene::AddInstancesUninitialized(uint32 NumInstances)
+FRayTracingScene::FInstanceRange FRayTracingScene::AllocateInstanceRangeUninitialized(uint32 NumInstances)
 {
 	const uint32 OldNum = Instances.AddUninitialized(NumInstances);
 
@@ -508,11 +510,15 @@ uint32 FRayTracingScene::AddInstancesUninitialized(uint32 NumInstances)
 		check(Instances.Num() == InstancesDebugData.Num());
 	}
 
-	return OldNum;
+	return { OldNum, NumInstances };
 }
 
-void FRayTracingScene::SetInstance(uint32 InstanceIndex, FRayTracingGeometryInstance InInstance, const FPrimitiveSceneProxy* Proxy, bool bDynamic)
+void FRayTracingScene::SetInstance(FInstanceRange InstanceRange, uint32 InstanceIndexInRange, FRayTracingGeometryInstance InInstance, const FPrimitiveSceneProxy* Proxy, bool bDynamic)
 {
+	checkf(InstanceIndexInRange < InstanceRange.Num, TEXT("InstanceIndexInRange (%d) is out of bounds for the range (%d)"), InstanceIndexInRange, InstanceRange.Num);
+
+	const uint32 InstanceIndex = InstanceRange.StartIndex + InstanceIndexInRange;
+
 	FRHIRayTracingGeometry* GeometryRHI = InInstance.GeometryRHI;
 
 	FRayTracingGeometryInstance* Instance = &Instances[InstanceIndex];
diff --git a/Engine/Source/Runtime/Renderer/Private/RayTracing/RayTracingScene.h b/Engine/Source/Runtime/Renderer/Private/RayTracing/RayTracingScene.h
--- a/Engine/Source/Runtime/Renderer/Private/RayTracing/RayTracingScene.h
+++ b/Engine/Source/Runtime/Renderer/Private/RayTracing/RayTracingScene.h
@@ -47,14 +47,51 @@ class FRayTracingScene
 {
 public:
 
+	struct FInstanceHandle
+	{
+		FInstanceHandle()
+			: Index(UINT32_MAX)
+		{}
+
+		bool IsValid() const
+		{
+			return Index != UINT32_MAX;
+		}
+
+	private:
+		FInstanceHandle(uint32 InIndex)
+			: Index(InIndex)
+		{}
+
+		uint32 Index;
+
+		friend class FRayTracingScene;
+	};
+
+	static const FInstanceHandle INVALID_INSTANCE_HANDLE;
+
+	struct FInstanceRange
+	{
+	private:
+		FInstanceRange(uint32 InStartIndex, uint32 InNum)
+			: StartIndex(InStartIndex)
+			, Num(InNum)
+		{}
+
+		uint32 StartIndex;
+		uint32 Num;
+
+		friend class FRayTracingScene;
+	};
+
 	FRayTracingScene();
 	~FRayTracingScene();
 
-	uint32 AddInstance(FRayTracingGeometryInstance Instance, const FPrimitiveSceneProxy* Proxy = nullptr, bool bDynamic = false);
+	FInstanceHandle AddInstance(FRayTracingGeometryInstance Instance, const FPrimitiveSceneProxy* Proxy = nullptr, bool bDynamic = false);
 
-	uint32 AddInstancesUninitialized(uint32 NumInstances);
+	FInstanceRange AllocateInstanceRangeUninitialized(uint32 NumInstances);
 
-	void SetInstance(uint32 InstanceIndex, FRayTracingGeometryInstance Instance, const FPrimitiveSceneProxy* Proxy = nullptr, bool bDynamic = false);
+	void SetInstance(FInstanceRange InstanceRange, uint32 InstanceIndexInRange, FRayTracingGeometryInstance Instance, const FPrimitiveSceneProxy* Proxy = nullptr, bool bDynamic = false);
 
 	// Allocates RayTracingSceneRHI and builds various metadata required to create the final scene.
 	FRayTracingSceneWithGeometryInstances BuildInitializationData() const;
@@ -106,7 +143,7 @@ public:
 
 	TArrayView<const FRayTracingGeometryInstance> GetInstances() const { return MakeArrayView(Instances); }
 
-	FRayTracingGeometryInstance& GetInstance(uint32 InstanceIndex) { return Instances[InstanceIndex]; }
+	FRayTracingGeometryInstance& GetInstance(FInstanceHandle Handle) { return Instances[Handle.Index]; }
 
 	uint32 GetTotalNumSegments() const { return NumSegments; }
 
