More scene extension exports

diff --git a/Engine/Source/Runtime/Renderer/Private/SceneExtensions.h b/Engine/Source/Runtime/Renderer/Private/SceneExtensions.h
--- a/Engine/Source/Runtime/Renderer/Private/SceneExtensions.h
+++ b/Engine/Source/Runtime/Renderer/Private/SceneExtensions.h
@@ -86,11 +86,11 @@ public:
 	TSparseArray<ISceneExtension*> CreateExtensions(FScene& Scene);
 
 private:
-	static void InitRegistry();
+	RENDERER_API static void InitRegistry();
 
 	TArray<ISceneExtensionFactory*> Factories;
 
-	static FSceneExtensionRegistry* GlobalRegistry;
+	RENDERER_API static FSceneExtensionRegistry* GlobalRegistry;
 };
 
 /** A collection of scene extensions */
