Add tile classification to single-scattering pass

#tests tested direct and indirect dispatch for live-shading

diff --git a/Engine/Shaders/Private/HeterogeneousVolumes/HeterogeneousVolumesLiveShadingGlobalPipeline.usf b/Engine/Shaders/Private/HeterogeneousVolumes/HeterogeneousVolumesLiveShadingGlobalPipeline.usf
new file mode 100644
--- /dev/null
+++ b/Engine/Shaders/Private/HeterogeneousVolumes/HeterogeneousVolumesLiveShadingGlobalPipeline.usf
@@ -0,0 +1,111 @@
+// Copyright Epic Games, Inc. All Rights Reserved.
+
+#include "../Common.ush"
+#include "../SceneTexturesCommon.ush"
+#include "HeterogeneousVolumesTracingUtils.ush"
+
+struct FScreenTile
+{
+	// TODO: This could be 16-bit..
+	uint Id;
+};
+
+// Object data
+float4x4 LocalToWorld;
+float4x4 WorldToLocal;
+float3 LocalBoundsOrigin;
+float3 LocalBoundsExtent;
+
+// Ray data
+float MaxTraceDistance;
+
+// Dispatch data
+int3 GroupCount;
+int DownsampleFactor;
+
+// Output
+RWBuffer<uint> RWNumScreenTilesBuffer;
+RWStructuredBuffer<FScreenTile> RWScreenTileBuffer;
+
+groupshared uint GSNumRayHits;
+
+float2 GetDispatchSize()
+{
+	return View.ViewSizeAndInvSize.xy / DownsampleFactor;
+}
+
+[numthreads(THREADGROUP_SIZE_2D, THREADGROUP_SIZE_2D, 1)]
+void ScreenTileClassificationCS(
+	uint2 GroupId : SV_GroupID,
+	uint2 GroupThreadId : SV_GroupThreadID,
+	uint2 DispatchThreadId : SV_DispatchThreadID
+)
+{
+	if (all(GroupThreadId == 0))
+	{
+		GSNumRayHits = 0;
+		RWNumScreenTilesBuffer[1] = 1;
+		RWNumScreenTilesBuffer[2] = 1;
+	}
+
+	// Create screen ray
+	if (any(DispatchThreadId.xy >= GetDispatchSize()))
+	{
+		return;
+	}
+	uint LinearIndex = DispatchThreadId.y * GroupCount.x * THREADGROUP_SIZE_2D + DispatchThreadId.x;
+	uint2 PixelCoord = DispatchThreadId.xy + View.ViewRectMin.xy / DownsampleFactor;
+	uint2 ViewPixelCoord = PixelCoord * DownsampleFactor + View.ViewRectMin.xy;
+#if 0
+	// Extract depth
+	float DeviceZ = SceneDepthTexture.Load(int3(ViewPixelCoord, 0)).r;
+#if HAS_INVERTED_Z_BUFFER
+	DeviceZ = max(0.000000000001, DeviceZ);
+#endif // HAS_INVERTED_Z_BUFFER
+
+	// Clip trace distance
+	float SceneDepth = min(ConvertFromDeviceZ(DeviceZ), MaxTraceDistance);
+#endif
+	float SceneDepth = MaxTraceDistance;
+	float DeviceZ = ConvertToDeviceZ(SceneDepth);
+
+	// Intersect ray with bounding volume
+	// TODO: LWC integration..
+	float3 WorldRayOrigin = DFHackToFloat(DFFastSubtract(GetTranslatedWorldCameraPosFromView(ViewPixelCoord + View.TemporalAAJitter.xy), PrimaryView.PreViewTranslation));
+	float3 WorldRayEnd = DFHackToFloat(SvPositionToWorld(float4(ViewPixelCoord + View.TemporalAAJitter.xy, DeviceZ, 1)));
+	//float3 WorldRayOrigin = DFHackToFloat(DFFastSubtract(GetTranslatedWorldCameraPosFromView(ViewPixelCoord), PrimaryView.PreViewTranslation));
+	//float3 WorldRayEnd = DFHackToFloat(SvPositionToWorld(float4(ViewPixelCoord, DeviceZ, 1)));
+	float3 WorldRayDirection = WorldRayEnd - WorldRayOrigin;
+	float WorldRayLength = length(WorldRayDirection);
+	WorldRayDirection /= WorldRayLength;
+
+	float3 LocalRayOrigin = mul(float4(WorldRayOrigin, 1.0), WorldToLocal).xyz;
+	float3 LocalRayEnd = mul(float4(WorldRayEnd, 1.0), WorldToLocal).xyz;
+	float3 LocalRayDirection = LocalRayEnd - LocalRayOrigin;
+	float LocalRayLength = length(LocalRayDirection);
+	LocalRayDirection /= LocalRayLength;
+
+	float3 LocalBoundsMin = LocalBoundsOrigin - LocalBoundsExtent;
+	float3 LocalBoundsMax = LocalBoundsOrigin + LocalBoundsExtent;
+
+	// Intersect bounding volume
+	float2 HitT = IntersectAABB(LocalRayOrigin, LocalRayDirection, 0.0, LocalRayLength, LocalBoundsMin, LocalBoundsMax);
+	float HitSpan = HitT.y - HitT.x;
+	if (HitSpan > 0.0)
+	{
+		InterlockedAdd(GSNumRayHits, 1);
+	}
+
+	if (all(GroupThreadId == 0))
+	{
+		if (GSNumRayHits > 0)
+		{
+			uint BufferIndex = 0;
+			InterlockedAdd(RWNumScreenTilesBuffer[0], 1, BufferIndex);
+
+			FScreenTile ScreenTile;
+			ScreenTile.Id = GroupId.y * GroupCount.x + GroupId.x;
+			RWScreenTileBuffer[BufferIndex] = ScreenTile;
+		}
+	}
+}
\ No newline at end of file
diff --git a/Engine/Shaders/Private/HeterogeneousVolumes/HeterogeneousVolumesLiveShadingPipeline.usf b/Engine/Shaders/Private/HeterogeneousVolumes/HeterogeneousVolumesLiveShadingPipeline.usf
--- a/Engine/Shaders/Private/HeterogeneousVolumes/HeterogeneousVolumesLiveShadingPipeline.usf
+++ b/Engine/Shaders/Private/HeterogeneousVolumes/HeterogeneousVolumesLiveShadingPipeline.usf
@@ -304,10 +304,8 @@ float2 GetDispatchSize()
 	return View.ViewSizeAndInvSize.xy / DownsampleFactor;
 }
 
-[numthreads(THREADGROUP_SIZE_2D, THREADGROUP_SIZE_2D, 1)]
-void RenderSingleScatteringWithLiveShadingCS(
-	uint2 GroupThreadId : SV_GroupThreadID,
-	uint2 DispatchThreadId : SV_DispatchThreadID
+void RenderSingleScatteringWithLiveShading(
+	uint2 DispatchThreadId
 )
 {
 	FDebugOutput DebugOutput = (FDebugOutput)0;
@@ -418,3 +416,34 @@ void RenderSingleScatteringWithLiveShadingCS(
 		//RWDebugOutputBuffer[LinearIndex] = DebugOutput;
 	}
 }
+
+[numthreads(THREADGROUP_SIZE_2D, THREADGROUP_SIZE_2D, 1)]
+void RenderSingleScatteringWithLiveShadingCS(
+	uint2 GroupThreadId : SV_GroupThreadID,
+	uint2 DispatchThreadId : SV_DispatchThreadID
+)
+{
+	RenderSingleScatteringWithLiveShading(DispatchThreadId);
+}
+
+struct FScreenTile
+{
+	uint Id;
+};
+
+
+
+StructuredBuffer<FScreenTile> ScreenTileBuffer;
+
+[numthreads(THREADGROUP_SIZE_2D, THREADGROUP_SIZE_2D, 1)]
+void RenderSingleScatteringWithLiveShadingIndirectCS(
+	uint GroupId : SV_GroupID,
+	uint2 GroupThreadId : SV_GroupThreadID
+)
+{
+	FScreenTile ScreenTile = ScreenTileBuffer[GroupId];
+	uint2 GroupId_2D = uint2(ScreenTile.Id % GroupCount.x, ScreenTile.Id / GroupCount.x);
+	uint2 DispatchThreadId = GroupId_2D * THREADGROUP_SIZE_2D + GroupThreadId;
+
+	RenderSingleScatteringWithLiveShading(DispatchThreadId);
+}
\ No newline at end of file
diff --git a/Engine/Source/Runtime/Renderer/Private/HeterogeneousVolumes/HeterogeneousVolumesLiveShadingPipeline.cpp b/Engine/Source/Runtime/Renderer/Private/HeterogeneousVolumes/HeterogeneousVolumesLiveShadingPipeline.cpp
--- a/Engine/Source/Runtime/Renderer/Private/HeterogeneousVolumes/HeterogeneousVolumesLiveShadingPipeline.cpp
+++ b/Engine/Source/Runtime/Renderer/Private/HeterogeneousVolumes/HeterogeneousVolumesLiveShadingPipeline.cpp
@@ -24,12 +24,24 @@ static TAutoConsoleVariable<int32> CVarHeterogeneousLightingCacheBoundsCulling(
 	ECVF_RenderThreadSafe
 );
 
+static TAutoConsoleVariable<int32> CVarHeterogeneousLightingLiveShadingScreenTileClassification(
+	TEXT("r.HeterogeneousVolumes.LiveShading.ScreenTileClassification"),
+	0,
+	TEXT("Enables screen tile classification for increased occupancy (Default = 0)"),
+	ECVF_RenderThreadSafe
+);
+
 namespace HeterogeneousVolumes
 {
 	bool ShouldBoundsCull()
 	{
 		return CVarHeterogeneousLightingCacheBoundsCulling.GetValueOnRenderThread() != 0;
 	}
+
+	bool ShouldUseScreenTileClassification()
+	{
+		return CVarHeterogeneousLightingLiveShadingScreenTileClassification.GetValueOnRenderThread() != 0;
+	}
 }
 
 //-OPT: Remove duplicate bindings
@@ -227,6 +239,23 @@ class FRenderLightingCacheWithLiveShadingCS : public FMeshMaterialShader
 
 IMPLEMENT_MATERIAL_SHADER_TYPE(, FRenderLightingCacheWithLiveShadingCS, TEXT("/Engine/Private/HeterogeneousVolumes/HeterogeneousVolumesLiveShadingPipeline.usf"), TEXT("RenderLightingCacheWithLiveShadingCS"), SF_Compute);
 
+namespace HeterogeneousVolumes
+{
+	struct FScreenTile
+	{
+		int32 Id;
+	};
+} // namespace
+
+namespace HeterogeneousVolumes {
+	enum EDispatchMode
+	{
+		DirectDispatch,
+		IndirectDispatch
+	};
+} // namespace HeterogeneousVolumes
+
+template <HeterogeneousVolumes::EDispatchMode DispatchMode>
 class FRenderSingleScatteringWithLiveShadingCS : public FMeshMaterialShader
 {
 	DECLARE_SHADER_TYPE(FRenderSingleScatteringWithLiveShadingCS, MeshMaterial);
@@ -294,6 +323,10 @@ class FRenderSingleScatteringWithLiveShadingCS : public FMeshMaterialShader
 		SHADER_PARAMETER(FIntVector, GroupCount)
 		SHADER_PARAMETER(int32, DownsampleFactor)
 
+		// Optional indirect dispatch data
+		RDG_BUFFER_ACCESS(IndirectArgs, ERHIAccess::IndirectArgs)
+		SHADER_PARAMETER_RDG_BUFFER_SRV(StructuredBuffer<HeterogeneousVolumes::FScreenTile>, ScreenTileBuffer)
+
 		// Output
 		SHADER_PARAMETER_RDG_TEXTURE_UAV(RWTexture2D<float4>, RWLightingTexture)
 		SHADER_PARAMETER_RDG_TEXTURE_UAV(RWTexture2D<float4>, RWVelocityTexture)
@@ -323,7 +356,7 @@ class FRenderSingleScatteringWithLiveShadingCS : public FMeshMaterialShader
 	)
 	{
 		FPermutationDomain PermutationVector(Parameters.PermutationId);
-		if (PermutationVector.Get<FApplyFogInscattering>() == 0)
+		if (PermutationVector.template Get<FApplyFogInscattering>() == 0)
 		{
 			return false;
 		}
@@ -335,9 +368,9 @@ class FRenderSingleScatteringWithLiveShadingCS : public FMeshMaterialShader
 	static FPermutationDomain RemapPermutation(FPermutationDomain PermutationVector)
 	{
 		// Remap Off to Stochastic and turn off individual fog modes
-		if (PermutationVector.Get<FApplyFogInscattering>() == 0)
+		if (PermutationVector.template Get<FApplyFogInscattering>() == 0)
 		{
-			PermutationVector.Set<FApplyFogInscattering>(2);
+			PermutationVector.template Set<FApplyFogInscattering>(2);
 		}
 
 		return PermutationVector;
@@ -373,9 +406,13 @@ class FRenderSingleScatteringWithLiveShadingCS : public FMeshMaterialShader
 	LAYOUT_FIELD(FRenderLightingCacheLooseBindings, ShaderLooseBindings);
 };
 
-IMPLEMENT_MATERIAL_SHADER_TYPE(, FRenderSingleScatteringWithLiveShadingCS, TEXT("/Engine/Private/HeterogeneousVolumes/HeterogeneousVolumesLiveShadingPipeline.usf"), TEXT("RenderSingleScatteringWithLiveShadingCS"), SF_Compute);
+typedef FRenderSingleScatteringWithLiveShadingCS<HeterogeneousVolumes::DirectDispatch> FRenderSingleScatteringWithLiveShadingDirectCS;
+typedef FRenderSingleScatteringWithLiveShadingCS<HeterogeneousVolumes::IndirectDispatch> FRenderSingleScatteringWithLiveShadingIndirectCS;
 
-template<bool bWithLumen, typename ComputeShaderType>
+IMPLEMENT_MATERIAL_SHADER_TYPE(template<>, FRenderSingleScatteringWithLiveShadingDirectCS, TEXT("/Engine/Private/HeterogeneousVolumes/HeterogeneousVolumesLiveShadingPipeline.usf"), TEXT("RenderSingleScatteringWithLiveShadingCS"), SF_Compute);
+IMPLEMENT_MATERIAL_SHADER_TYPE(template<>, FRenderSingleScatteringWithLiveShadingIndirectCS, TEXT("/Engine/Private/HeterogeneousVolumes/HeterogeneousVolumesLiveShadingPipeline.usf"), TEXT("RenderSingleScatteringWithLiveShadingIndirectCS"), SF_Compute);
+
+template<bool bWithLumen, HeterogeneousVolumes::EDispatchMode DispatchMode, typename ComputeShaderType>
 void AddComputePass(
 	FRDGBuilder& GraphBuilder,
 	TShaderRef<ComputeShaderType>& ComputeShader,
@@ -384,7 +421,9 @@ void AddComputePass(
 	const FMaterialRenderProxy* MaterialRenderProxy,
 	const FMaterial& Material,
 	const FString& PassName,
-	FIntVector GroupCount
+	FIntVector GroupCount,
+	FRDGBufferRef IndirectArgsBuffer,
+	uint32 IndirectArgOffset
 )
 {
 	//ClearUnusedGraphResources(ComputeShader, PassParameters);
@@ -393,7 +432,7 @@ void AddComputePass(
 		RDG_EVENT_NAME("%s", *PassName),
 		PassParameters,
 		ERDGPassFlags::Compute,
-		[ComputeShader, PassParameters, Scene, MaterialRenderProxy, &Material, GroupCount](FRHIComputeCommandList& RHICmdList)
+		[ComputeShader, PassParameters, Scene, MaterialRenderProxy, &Material, GroupCount, IndirectArgsBuffer, IndirectArgOffset](FRHIComputeCommandList& RHICmdList)
 		{
 			FMeshMaterialShaderElementData ShaderElementData;
 			ShaderElementData.InitializeMeshMaterialData();
@@ -417,11 +456,74 @@ void AddComputePass(
 				ShaderBindings.Finalize(&PassShaders);
 			}
 
-			UE::MeshPassUtils::Dispatch(RHICmdList, ComputeShader, ShaderBindings, *PassParameters, GroupCount);
+			if constexpr (DispatchMode == HeterogeneousVolumes::EDispatchMode::IndirectDispatch)
+			{
+				UE::MeshPassUtils::DispatchIndirect(RHICmdList, ComputeShader, ShaderBindings, *PassParameters, IndirectArgsBuffer->GetIndirectRHICallBuffer(), IndirectArgOffset);
+			}
+			else
+			{
+				UE::MeshPassUtils::Dispatch(RHICmdList, ComputeShader, ShaderBindings, *PassParameters, GroupCount);
+			}
 		}
 	);
 }
 
+template<bool bWithLumen, typename ComputeShaderType>
+void AddComputePassDirect(
+	FRDGBuilder& GraphBuilder,
+	TShaderRef<ComputeShaderType>& ComputeShader,
+	typename ComputeShaderType::FParameters* PassParameters,
+	const FScene* Scene,
+	const FMaterialRenderProxy* MaterialRenderProxy,
+	const FMaterial& Material,
+	const FString& PassName,
+	FIntVector GroupCount
+)
+{
+	FRDGBufferRef IndirectArgsBuffer = GSystemTextures.GetDefaultBuffer(GraphBuilder, 4);
+	uint32 IndirectArgOffset = 0;
+	AddComputePass<bWithLumen, HeterogeneousVolumes::EDispatchMode::DirectDispatch>(
+		GraphBuilder,
+		ComputeShader,
+		PassParameters,
+		Scene,
+		MaterialRenderProxy,
+		Material,
+		PassName,
+		GroupCount,
+		IndirectArgsBuffer,
+		IndirectArgOffset
+	);
+}
+
+template<bool bWithLumen, typename ComputeShaderType>
+void AddComputePassIndirect(
+	FRDGBuilder& GraphBuilder,
+	TShaderRef<ComputeShaderType>& ComputeShader,
+	typename ComputeShaderType::FParameters* PassParameters,
+	const FScene* Scene,
+	const FMaterialRenderProxy* MaterialRenderProxy,
+	const FMaterial& Material,
+	const FString& PassName,
+	FRDGBufferRef IndirectArgsBuffer,
+	uint32 IndirectArgOffset
+)
+{
+	FIntVector GroupCount = FIntVector::ZeroValue;
+	AddComputePass<bWithLumen, HeterogeneousVolumes::EDispatchMode::IndirectDispatch>(
+		GraphBuilder,
+		ComputeShader,
+		PassParameters,
+		Scene,
+		MaterialRenderProxy,
+		Material,
+		PassName,
+		GroupCount,
+		IndirectArgsBuffer,
+		IndirectArgOffset
+	);
+}
+
 static void RenderLightingCacheWithLiveShading(
 	FRDGBuilder& GraphBuilder,
 	// Scene data
@@ -614,11 +716,136 @@ static void RenderLightingCacheWithLiveShading(
 	TShaderRef<FRenderLightingCacheWithLiveShadingCS> ComputeShader = Material.GetShader<FRenderLightingCacheWithLiveShadingCS>(&FLocalVertexFactory::StaticType, PermutationVector, false);
 	if (!ComputeShader.IsNull())
 	{
-		AddComputePass<false>(GraphBuilder, ComputeShader, PassParameters, Scene, MaterialRenderProxy, Material, PassName, GroupCount);
+		AddComputePassDirect<false>(GraphBuilder, ComputeShader, PassParameters, Scene, MaterialRenderProxy, Material, PassName, GroupCount);
 	}
 }
 
-static void RenderSingleScatteringWithLiveShading(
+class FScreenTileClassificationCS : public FGlobalShader
+{
+	DECLARE_GLOBAL_SHADER(FScreenTileClassificationCS);
+	SHADER_USE_PARAMETER_STRUCT(FScreenTileClassificationCS, FGlobalShader);
+
+
+	BEGIN_SHADER_PARAMETER_STRUCT(FParameters, )
+		// Scene data
+		SHADER_PARAMETER_STRUCT_REF(FViewUniformShaderParameters, View)
+		SHADER_PARAMETER_STRUCT_INCLUDE(FSceneTextureParameters, SceneTextures)
+
+		// Object data
+		SHADER_PARAMETER(FMatrix44f, LocalToWorld)
+		SHADER_PARAMETER(FMatrix44f, WorldToLocal)
+		SHADER_PARAMETER(FVector3f, LocalBoundsOrigin)
+		SHADER_PARAMETER(FVector3f, LocalBoundsExtent)
+
+		// Ray data
+		SHADER_PARAMETER(float, MaxTraceDistance)
+
+		// Dispatch data
+		SHADER_PARAMETER(FIntVector, GroupCount)
+		SHADER_PARAMETER(int32, DownsampleFactor)
+
+		// Output
+		SHADER_PARAMETER_RDG_BUFFER_UAV(RWBuffer<uint>, RWNumScreenTilesBuffer)
+		SHADER_PARAMETER_RDG_BUFFER_UAV(RWStructuredBuffer<HeterogeneousVolumes::FScreenTile>, RWScreenTileBuffer)
+	END_SHADER_PARAMETER_STRUCT()
+
+	static bool ShouldCompilePermutation(
+		const FGlobalShaderPermutationParameters& Parameters
+	)
+	{
+		return DoesPlatformSupportHeterogeneousVolumes(Parameters.Platform);
+	}
+
+	static void ModifyCompilationEnvironment(
+		const FGlobalShaderPermutationParameters& Parameters,
+		FShaderCompilerEnvironment& OutEnvironment
+	)
+	{
+		FMaterialShader::ModifyCompilationEnvironment(Parameters, OutEnvironment);
+		OutEnvironment.SetDefine(TEXT("THREADGROUP_SIZE_1D"), GetThreadGroupSize1D());
+		OutEnvironment.SetDefine(TEXT("THREADGROUP_SIZE_2D"), GetThreadGroupSize2D());
+
+		// This shader takes a very long time to compile with FXC, so we pre-compile it with DXC first and then forward the optimized HLSL to FXC.
+		OutEnvironment.CompilerFlags.Add(CFLAG_AllowTypedUAVLoads);
+	}
+
+	static int32 GetThreadGroupSize1D() { return GetThreadGroupSize2D() * GetThreadGroupSize2D(); }
+	static int32 GetThreadGroupSize2D() { return 8; }
+};
+
+IMPLEMENT_GLOBAL_SHADER(FScreenTileClassificationCS, "/Engine/Private/HeterogeneousVolumes/HeterogeneousVolumesLiveShadingGlobalPipeline.usf", "ScreenTileClassificationCS", SF_Compute);
+
+static void ScreenTileClassification(
+	FRDGBuilder& GraphBuilder,
+	// Scene data
+	const FScene* Scene,
+	const FViewInfo& View,
+	const FSceneTextures& SceneTextures,
+	// Object data
+	const IHeterogeneousVolumeInterface* HeterogeneousVolumeInterface,
+	// Output
+	FRDGBufferRef& NumScreenTilesBuffer,
+	FRDGBufferRef& ScreenTileBuffer
+)
+{
+	FIntVector GroupCount = FComputeShaderUtils::GetGroupCount(HeterogeneousVolumes::GetScaledViewRect(View.ViewRect), FRenderSingleScatteringWithLiveShadingIndirectCS::GetThreadGroupSize2D());
+	int32 NumTiles = GroupCount.X * GroupCount.Y;
+
+	NumScreenTilesBuffer = GraphBuilder.CreateBuffer(
+		FRDGBufferDesc::CreateIndirectDesc<FRHIDispatchIndirectParameters>(1),
+		TEXT("HeterogeneousVolume.NumScreenTilesBuffer")
+	);
+	// TODO: Initialize elsewhere??
+	AddClearUAVPass(GraphBuilder, GraphBuilder.CreateUAV(NumScreenTilesBuffer, PF_R32_UINT), 0);
+
+	ScreenTileBuffer = GraphBuilder.CreateBuffer(
+		FRDGBufferDesc::CreateStructuredDesc(sizeof(HeterogeneousVolumes::FScreenTile), NumTiles),
+		TEXT("HeterogeneousVolume.ScreenTileBuffer")
+	);
+
+	FScreenTileClassificationCS::FParameters* PassParameters = GraphBuilder.AllocParameters<FScreenTileClassificationCS::FParameters>();
+	{
+		PassParameters->View = View.ViewUniformBuffer;
+		PassParameters->SceneTextures = GetSceneTextureParameters(GraphBuilder, SceneTextures);
+
+		// Object data
+		// TODO: Convert to relative-local space
+		//FVector3f ViewOriginHigh = FDFVector3(View.ViewMatrices.GetViewOrigin()).High;
+		//FMatrix44f RelativeLocalToWorld = FDFMatrix::MakeToRelativeWorldMatrix(ViewOriginHigh, HeterogeneousVolumeInterface->GetLocalToWorld()).M;
+		FMatrix InstanceToLocal = HeterogeneousVolumeInterface->GetInstanceToLocal();
+		FMatrix LocalToWorld = HeterogeneousVolumeInterface->GetLocalToWorld();
+		PassParameters->LocalToWorld = FMatrix44f(InstanceToLocal * LocalToWorld);
+		PassParameters->WorldToLocal = PassParameters->LocalToWorld.Inverse();
+
+		FMatrix LocalToInstance = InstanceToLocal.Inverse();
+		FBoxSphereBounds InstanceBoxSphereBounds = HeterogeneousVolumeInterface->GetLocalBounds().TransformBy(LocalToInstance);
+		PassParameters->LocalBoundsOrigin = FVector3f(InstanceBoxSphereBounds.Origin);
+		PassParameters->LocalBoundsExtent = FVector3f(InstanceBoxSphereBounds.BoxExtent);
+
+		// Ray data
+		PassParameters->MaxTraceDistance = HeterogeneousVolumes::GetMaxTraceDistance();
+
+		// Dispatch data
+		PassParameters->GroupCount = GroupCount;
+		PassParameters->DownsampleFactor = HeterogeneousVolumes::GetDownsampleFactor();
+
+		PassParameters->RWNumScreenTilesBuffer = GraphBuilder.CreateUAV(NumScreenTilesBuffer, PF_R32_UINT);
+		PassParameters->RWScreenTileBuffer = GraphBuilder.CreateUAV(ScreenTileBuffer);
+	}
+
+	FScreenTileClassificationCS::FPermutationDomain PermutationVector;
+	TShaderRef<FScreenTileClassificationCS> ComputeShader = View.ShaderMap->GetShader<FScreenTileClassificationCS>(PermutationVector);
+	FComputeShaderUtils::AddPass(
+		GraphBuilder,
+		RDG_EVENT_NAME("FScreenTileClassificationCS"),
+		ComputeShader,
+		PassParameters,
+		GroupCount
+	);
+}
+
+template <HeterogeneousVolumes::EDispatchMode DispatchMode>
+void RenderSingleScatteringWithLiveShading(
 	FRDGBuilder& GraphBuilder,
 	// Scene data
 	const FScene* Scene,
@@ -644,12 +871,29 @@ static void RenderSingleScatteringWithLiveShading(
 	FRDGTextureRef& HeterogeneousVolumeTexture
 )
 {
+	typedef FRenderSingleScatteringWithLiveShadingCS<DispatchMode> FRenderSingleScatteringWithLiveShadingDispatchTypeCS;
+
 	const FMaterialRenderProxy* MaterialRenderProxy = nullptr;
 	const FMaterial& Material = DefaultMaterialRenderProxy->GetMaterialWithFallback(View.GetFeatureLevel(), MaterialRenderProxy);
 	MaterialRenderProxy = MaterialRenderProxy ? MaterialRenderProxy : DefaultMaterialRenderProxy;
 	check(Material.GetMaterialDomain() == MD_Volume);
 
-	FIntVector GroupCount = FComputeShaderUtils::GetGroupCount(HeterogeneousVolumes::GetScaledViewRect(View.ViewRect), FRenderSingleScatteringWithLiveShadingCS::GetThreadGroupSize2D());
+	FRDGBufferRef NumScreenTilesBuffer;
+	FRDGBufferRef ScreenTileBuffer;
+	if (DispatchMode == HeterogeneousVolumes::IndirectDispatch)
+	{
+		ScreenTileClassification(
+			GraphBuilder,
+			Scene,
+			View,
+			SceneTextures,
+			HeterogeneousVolumeInterface,
+			NumScreenTilesBuffer,
+			ScreenTileBuffer
+		);
+	}
+
+	FIntVector GroupCount = FComputeShaderUtils::GetGroupCount(HeterogeneousVolumes::GetScaledViewRect(View.ViewRect), FRenderSingleScatteringWithLiveShadingDispatchTypeCS::GetThreadGroupSize2D());
 
 	// Note must be done in the same scope as we add the pass otherwise the UB lifetime will not be guaranteed
 	FDeferredLightUniformStruct DeferredLightUniform;
@@ -661,7 +905,7 @@ static void RenderSingleScatteringWithLiveShading(
 
 	bool bUseAVSM = HeterogeneousVolumes::UseAdaptiveVolumetricShadowMapForSelfShadowing(HeterogeneousVolumeInterface->GetPrimitiveSceneProxy());
 	bool bWriteVelocity = HeterogeneousVolumes::ShouldWriteVelocity() && HasBeenProduced(SceneTextures.Velocity);
-	FRenderSingleScatteringWithLiveShadingCS::FParameters* PassParameters = GraphBuilder.AllocParameters<FRenderSingleScatteringWithLiveShadingCS::FParameters>();
+	typename FRenderSingleScatteringWithLiveShadingDispatchTypeCS::FParameters* PassParameters = GraphBuilder.AllocParameters<typename FRenderSingleScatteringWithLiveShadingDispatchTypeCS::FParameters>();
 	{
 		// Scene data
 		PassParameters->View = View.ViewUniformBuffer;
@@ -777,6 +1021,11 @@ static void RenderSingleScatteringWithLiveShading(
 		// Dispatch data
 		PassParameters->GroupCount = GroupCount;
 		PassParameters->DownsampleFactor = HeterogeneousVolumes::GetDownsampleFactor();
+		if (DispatchMode == HeterogeneousVolumes::IndirectDispatch)
+		{
+			PassParameters->IndirectArgs = NumScreenTilesBuffer;
+			PassParameters->ScreenTileBuffer = GraphBuilder.CreateSRV(ScreenTileBuffer);
+		}
 
 		// Output
 		PassParameters->RWLightingTexture = GraphBuilder.CreateUAV(HeterogeneousVolumeTexture);
@@ -802,19 +1051,19 @@ static void RenderSingleScatteringWithLiveShading(
 
 	int32 AVSMSampleMode = HeterogeneousVolumes::GetShadowMaxSampleCount() > 16;
 
-	FRenderSingleScatteringWithLiveShadingCS::FPermutationDomain PermutationVector;
-	PermutationVector.Set<FRenderSingleScatteringWithLiveShadingCS::FUseTransmittanceVolume>(HeterogeneousVolumes::UseLightingCacheForTransmittance() && PassParameters->bApplyShadowTransmittance);
-	PermutationVector.Set<FRenderSingleScatteringWithLiveShadingCS::FUseInscatteringVolume>(HeterogeneousVolumes::UseLightingCacheForInscattering());
-	PermutationVector.Set<FRenderSingleScatteringWithLiveShadingCS::FUseLumenGI>(HeterogeneousVolumes::UseIndirectLighting() && View.GetLumenTranslucencyGIVolume().Texture0 != nullptr);
-	PermutationVector.Set<FRenderSingleScatteringWithLiveShadingCS::FWriteVelocity>(bWriteVelocity);
-	PermutationVector.Set<FRenderSingleScatteringWithLiveShadingCS::FUseAdaptiveVolumetricShadowMap>(bUseAVSM);
-	PermutationVector.Set<FRenderSingleScatteringWithLiveShadingCS::FAVSMSampleMode>(AVSMSampleMode);
-	PermutationVector.Set<FRenderSingleScatteringWithLiveShadingCS::FApplyFogInscattering>(static_cast<int32>(HeterogeneousVolumes::GetApplyFogInscattering()));
-	PermutationVector = FRenderSingleScatteringWithLiveShadingCS::RemapPermutation(PermutationVector);
-	TShaderRef<FRenderSingleScatteringWithLiveShadingCS> ComputeShader = Material.GetShader<FRenderSingleScatteringWithLiveShadingCS>(&FLocalVertexFactory::StaticType, PermutationVector, false);
+	typename FRenderSingleScatteringWithLiveShadingDispatchTypeCS::FPermutationDomain PermutationVector;
+	PermutationVector.template Set<typename FRenderSingleScatteringWithLiveShadingDispatchTypeCS::FUseTransmittanceVolume>(HeterogeneousVolumes::UseLightingCacheForTransmittance() && PassParameters->bApplyShadowTransmittance);
+	PermutationVector.template Set<typename FRenderSingleScatteringWithLiveShadingDispatchTypeCS::FUseInscatteringVolume>(HeterogeneousVolumes::UseLightingCacheForInscattering());
+	PermutationVector.template Set<typename FRenderSingleScatteringWithLiveShadingDispatchTypeCS::FUseLumenGI>(HeterogeneousVolumes::UseIndirectLighting() && View.GetLumenTranslucencyGIVolume().Texture0 != nullptr);
+	PermutationVector.template Set<typename FRenderSingleScatteringWithLiveShadingDispatchTypeCS::FWriteVelocity>(bWriteVelocity);
+	PermutationVector.template Set<typename FRenderSingleScatteringWithLiveShadingDispatchTypeCS::FUseAdaptiveVolumetricShadowMap>(bUseAVSM);
+	PermutationVector.template Set<typename FRenderSingleScatteringWithLiveShadingDispatchTypeCS::FAVSMSampleMode>(AVSMSampleMode);
+	PermutationVector.template Set<typename FRenderSingleScatteringWithLiveShadingDispatchTypeCS::FApplyFogInscattering>(static_cast<int32>(HeterogeneousVolumes::GetApplyFogInscattering()));
+	PermutationVector = FRenderSingleScatteringWithLiveShadingDispatchTypeCS::RemapPermutation(PermutationVector);
+	TShaderRef<FRenderSingleScatteringWithLiveShadingDispatchTypeCS> ComputeShader = Material.GetShader<FRenderSingleScatteringWithLiveShadingDispatchTypeCS>(&FLocalVertexFactory::StaticType, PermutationVector, false);
 	if (!ComputeShader.IsNull())
 	{
-		AddComputePass<true>(GraphBuilder, ComputeShader, PassParameters, Scene, MaterialRenderProxy, Material, PassName, GroupCount);
+		AddComputePass<true, DispatchMode>(GraphBuilder, ComputeShader, PassParameters, Scene, MaterialRenderProxy, Material, PassName, GroupCount, PassParameters->IndirectArgs, 0);
 	}
 }
 
@@ -900,31 +1149,62 @@ static void RenderWithTransmittanceVolumePipeline(
 			);
 		}
 
-		RenderSingleScatteringWithLiveShading(
-			GraphBuilder,
-			// Scene data
-			Scene,
-			View,
-			SceneTextures,
-			// Light data
-			bApplyEmissionAndTransmittance,
-			bApplyDirectLighting,
-			bApplyShadowTransmittance,
-			LightType,
-			LightSceneInfo,
-			// Shadow data
-			VisibleLightInfo,
-			VirtualShadowMapArray,
-			// Object data
-			HeterogeneousVolumeInterface,
-			MaterialRenderProxy,
-			PersistentPrimitiveIndex,
-			LocalBoxSphereBounds,
-			// Transmittance acceleration
-			LightingCacheTexture,
-			// Output
-			HeterogeneousVolumeRadiance
-		);
+		if (HeterogeneousVolumes::ShouldUseScreenTileClassification())
+		{
+			RenderSingleScatteringWithLiveShading<HeterogeneousVolumes::IndirectDispatch>(
+				GraphBuilder,
+				// Scene data
+				Scene,
+				View,
+				SceneTextures,
+				// Light data
+				bApplyEmissionAndTransmittance,
+				bApplyDirectLighting,
+				bApplyShadowTransmittance,
+				LightType,
+				LightSceneInfo,
+				// Shadow data
+				VisibleLightInfo,
+				VirtualShadowMapArray,
+				// Object data
+				HeterogeneousVolumeInterface,
+				MaterialRenderProxy,
+				PersistentPrimitiveIndex,
+				LocalBoxSphereBounds,
+				// Transmittance acceleration
+				LightingCacheTexture,
+				// Output
+				HeterogeneousVolumeRadiance
+			);
+		}
+		else
+		{
+			RenderSingleScatteringWithLiveShading<HeterogeneousVolumes::DirectDispatch>(
+				GraphBuilder,
+				// Scene data
+				Scene,
+				View,
+				SceneTextures,
+				// Light data
+				bApplyEmissionAndTransmittance,
+				bApplyDirectLighting,
+				bApplyShadowTransmittance,
+				LightType,
+				LightSceneInfo,
+				// Shadow data
+				VisibleLightInfo,
+				VirtualShadowMapArray,
+				// Object data
+				HeterogeneousVolumeInterface,
+				MaterialRenderProxy,
+				PersistentPrimitiveIndex,
+				LocalBoxSphereBounds,
+				// Transmittance acceleration
+				LightingCacheTexture,
+				// Output
+				HeterogeneousVolumeRadiance
+			);
+		}
 	}
 }
 
@@ -1017,31 +1297,62 @@ static void RenderWithInscatteringVolumePipeline(
 		FLightSceneInfo* LightSceneInfo = nullptr;
 		const FVisibleLightInfo* VisibleLightInfo = nullptr;
 
-		RenderSingleScatteringWithLiveShading(
-			GraphBuilder,
-			// Scene data
-			Scene,
-			View,
-			SceneTextures,
-			// Light data
-			bApplyEmissionAndTransmittance,
-			bApplyDirectLighting,
-			bApplyShadowTransmittance,
-			LightType,
-			LightSceneInfo,
-			// Shadow data
-			VisibleLightInfo,
-			VirtualShadowMapArray,
-			// Object data
-			HeterogeneousVolumeInterface,
-			MaterialRenderProxy,
-			PersistentPrimitiveIndex,
-			LocalBoxSphereBounds,
-			// Transmittance acceleration
-			LightingCacheTexture,
-			// Output
-			HeterogeneousVolumeRadiance
-		);
+		if (HeterogeneousVolumes::ShouldUseScreenTileClassification())
+		{
+			RenderSingleScatteringWithLiveShading<HeterogeneousVolumes::IndirectDispatch>(
+				GraphBuilder,
+				// Scene data
+				Scene,
+				View,
+				SceneTextures,
+				// Light data
+				bApplyEmissionAndTransmittance,
+				bApplyDirectLighting,
+				bApplyShadowTransmittance,
+				LightType,
+				LightSceneInfo,
+				// Shadow data
+				VisibleLightInfo,
+				VirtualShadowMapArray,
+				// Object data
+				HeterogeneousVolumeInterface,
+				MaterialRenderProxy,
+				PersistentPrimitiveIndex,
+				LocalBoxSphereBounds,
+				// Transmittance acceleration
+				LightingCacheTexture,
+				// Output
+				HeterogeneousVolumeRadiance
+			);
+		}
+		else
+		{
+			RenderSingleScatteringWithLiveShading<HeterogeneousVolumes::DirectDispatch>(
+				GraphBuilder,
+				// Scene data
+				Scene,
+				View,
+				SceneTextures,
+				// Light data
+				bApplyEmissionAndTransmittance,
+				bApplyDirectLighting,
+				bApplyShadowTransmittance,
+				LightType,
+				LightSceneInfo,
+				// Shadow data
+				VisibleLightInfo,
+				VirtualShadowMapArray,
+				// Object data
+				HeterogeneousVolumeInterface,
+				MaterialRenderProxy,
+				PersistentPrimitiveIndex,
+				LocalBoxSphereBounds,
+				// Transmittance acceleration
+				LightingCacheTexture,
+				// Output
+				HeterogeneousVolumeRadiance
+			);
+		}
 	}
 }
 
