VolumetricCloud - do not render the cloud alpha hold out contribution pass when none of sky, fog or cloud are holdout (since the contribution would be null anyway).

diff --git a/Engine/Source/Runtime/Renderer/Private/VolumetricCloudRendering.cpp b/Engine/Source/Runtime/Renderer/Private/VolumetricCloudRendering.cpp
--- a/Engine/Source/Runtime/Renderer/Private/VolumetricCloudRendering.cpp
+++ b/Engine/Source/Runtime/Renderer/Private/VolumetricCloudRendering.cpp
@@ -2598,6 +2598,15 @@ bool FSceneRenderer::RenderVolumetricCloud(
 
 				if (bAccumulateAlphaHoldOut && ViewInfo.CachedViewUniformShaderParameters->RenderingReflectionCaptureMask > 0.0f)
 				{
+					// We do not render the cloud holdout pass if we are rendering a reflection captures
+					continue;
+				}
+
+				const FIntVector4& EnvironmentComponentFlags = ViewInfo.CachedViewUniformShaderParameters->EnvironmentComponentsFlags;
+				if (bAccumulateAlphaHoldOut && !(IsSkyAtmosphereHoldout(EnvironmentComponentFlags) || IsVolumetricCloudHoldout(EnvironmentComponentFlags) || IsExponentialFogHoldout(EnvironmentComponentFlags)))
+				{
+					// We do not render the cloud holdout pass alpha contribution if not required (when none of the volumetric effects can contribute holdout alpha value).
+					// The cloud transmittance is still applied on the SkyAtmosphere alpha holdout contribution in any case during the regular pass.
 					continue;
 				}
 
