Fix a bug that got introduced where visualizing level instances would only highlight selected instances.

#rnx

diff --git a/Engine/Source/Runtime/Renderer/Private/SceneVisibility.cpp b/Engine/Source/Runtime/Renderer/Private/SceneVisibility.cpp
--- a/Engine/Source/Runtime/Renderer/Private/SceneVisibility.cpp
+++ b/Engine/Source/Runtime/Renderer/Private/SceneVisibility.cpp
@@ -1856,19 +1856,19 @@ void FRelevancePacket::ComputeRelevance(FDynamicPrimitiveIndexList& DynamicPrimi
 				}
 			}
 
-			const int32 MaxInstances = PrimitiveSceneInfo.GetNumInstanceSceneDataEntries();
-			OutSelectedInstanceDraws.Reserve(OutSelectedInstanceDraws.Num() + MaxInstances);
-			if (OutOverlaidInstanceDraws)
-			{
-				OutOverlaidInstanceDraws->Reserve(OutOverlaidInstanceDraws->Num() + MaxInstances);
-			}
 			const FInstanceSceneDataBuffers* InstanceSceneDataBuffers = PrimitiveSceneInfo.GetInstanceSceneDataBuffers();
 			const bool bCollectInstanceHitProxyIds = bSelectedInstancesOnly &&
 				NaniteProxy->HasSelectedInstances() &&
 				OutSelectedInstanceHitProxyIDs != nullptr &&
 				InstanceSceneDataBuffers != nullptr;
-			const bool bIsSelected = NaniteProxy->IsSelected();
-			const bool bWantsEditorEffects = NaniteProxy->WantsEditorEffects() && OutOverlaidInstanceDraws;
+			const bool bOverlaidDraws = OutOverlaidInstanceDraws &&
+				NaniteProxy->WantsEditorEffects() &&
+				!NaniteProxy->IsSelected();
+			
+			const int32 MaxInstances = PrimitiveSceneInfo.GetNumInstanceSceneDataEntries();
+			TArray<Nanite::FInstanceDraw>& OutDrawArray = bOverlaidDraws ? *OutOverlaidInstanceDraws : OutSelectedInstanceDraws;
+			OutDrawArray.Reserve(OutDrawArray.Num() + MaxInstances);
+
 			for (int32 Idx = 0; Idx < MaxInstances; ++Idx)
 			{
 				if (bCollectInstanceHitProxyIds)
@@ -1891,24 +1891,12 @@ void FRelevancePacket::ComputeRelevance(FDynamicPrimitiveIndexList& DynamicPrimi
 					}
 				}
 
-				if (bIsSelected)
-				{
-					OutSelectedInstanceDraws.Add(
-						Nanite::FInstanceDraw {
-							uint32(PrimitiveSceneInfo.GetInstanceSceneDataOffset() + Idx),
-							0u
-						}
-					);
-				}
-				else if (bWantsEditorEffects)
-				{
-					OutOverlaidInstanceDraws->Add(
-						Nanite::FInstanceDraw {
-							uint32(PrimitiveSceneInfo.GetInstanceSceneDataOffset() + Idx),
-							0u
-						}
-					);
-				}
+				OutDrawArray.Add(
+					Nanite::FInstanceDraw {
+						uint32(PrimitiveSceneInfo.GetInstanceSceneDataOffset() + Idx),
+						0u
+					}
+				);
 			}
 		};
 
