* Fix incorrect output texture format during SLW base pass with Substrate, causing SLW to be incorrect.
* Share logic for enabling/disabling SLW directional texture

diff --git a/Engine/Source/Runtime/Engine/Private/ShaderCompiler/ShaderGenerationUtil.cpp b/Engine/Source/Runtime/Engine/Private/ShaderCompiler/ShaderGenerationUtil.cpp
--- a/Engine/Source/Runtime/Engine/Private/ShaderCompiler/ShaderGenerationUtil.cpp
+++ b/Engine/Source/Runtime/Engine/Private/ShaderCompiler/ShaderGenerationUtil.cpp
@@ -2306,7 +2306,7 @@ FGBufferParams FShaderCompileUtilities::FetchGBufferParamsRuntime(EShaderPlatfor
 	Ret.LegacyFormatIndex = CVarFormat->GetValueOnAnyThread();
 
 	// This should match with SINGLE_LAYER_WATER_SEPARATED_MAIN_LIGHT
-	Ret.bHasSingleLayerWaterSeparatedMainLight = IsWaterDistanceFieldShadowEnabled(Platform) || IsWaterVirtualShadowMapFilteringEnabled(Platform);
+	Ret.bHasSingleLayerWaterSeparatedMainLight = IsWaterSeparateMainDirLightEnabled(Platform);
 
 	return Ret;
 }
diff --git a/Engine/Source/Runtime/RenderCore/Private/RenderUtils.cpp b/Engine/Source/Runtime/RenderCore/Private/RenderUtils.cpp
--- a/Engine/Source/Runtime/RenderCore/Private/RenderUtils.cpp
+++ b/Engine/Source/Runtime/RenderCore/Private/RenderUtils.cpp
@@ -1496,6 +1496,11 @@ bool IsWaterDistanceFieldShadowEnabled(const FStaticShaderPlatform Platform)
 	return !IsForwardShadingEnabled(Platform) && IsUsingDistanceFields(Platform) && bWaterSingleLayerShaderSupportDistanceFieldShadow;
 }
 
+bool IsWaterSeparateMainDirLightEnabled(const FStaticShaderPlatform Platform)
+{
+	return IsWaterDistanceFieldShadowEnabled(Platform) || IsWaterVirtualShadowMapFilteringEnabled(Platform);
+}
+
 bool UseGPUScene(const FStaticShaderPlatform Platform, const FStaticFeatureLevel FeatureLevel)
 {
 	if (FeatureLevel == ERHIFeatureLevel::ES3_1)
diff --git a/Engine/Source/Runtime/RenderCore/Private/Shader.cpp b/Engine/Source/Runtime/RenderCore/Private/Shader.cpp
--- a/Engine/Source/Runtime/RenderCore/Private/Shader.cpp
+++ b/Engine/Source/Runtime/RenderCore/Private/Shader.cpp
@@ -2276,7 +2276,7 @@ void ShaderMapAppendKey(EShaderPlatform Platform, FShaderKeyGenerator& KeyGen)
 		}
 	}
 
-	const bool bNeedsSeparateMainDirLightTexture = IsWaterDistanceFieldShadowEnabled(Platform) || IsWaterVirtualShadowMapFilteringEnabled(Platform);
+	const bool bNeedsSeparateMainDirLightTexture = IsWaterSeparateMainDirLightEnabled(Platform);
 	if (bNeedsSeparateMainDirLightTexture)
 	{
 		KeyGen.AppendSeparator();
diff --git a/Engine/Source/Runtime/RenderCore/Public/RenderUtils.h b/Engine/Source/Runtime/RenderCore/Public/RenderUtils.h
--- a/Engine/Source/Runtime/RenderCore/Public/RenderUtils.h
+++ b/Engine/Source/Runtime/RenderCore/Public/RenderUtils.h
@@ -359,6 +359,9 @@ RENDERCORE_API bool IsUsingDistanceFields(const FStaticShaderPlatform Platform);
 /** Returns if water should render distance field shadow a second time for the water surface. This is for a platofrm so can be used at cook time. */
 RENDERCORE_API bool IsWaterDistanceFieldShadowEnabled(const FStaticShaderPlatform Platform);
 
+/** Returns if water needs a separate main directional light texture. This is for a platofrm so can be used at cook time or at runtime. */
+RENDERCORE_API bool IsWaterSeparateMainDirLightEnabled(const FStaticShaderPlatform Platform);
+
 RENDERCORE_API bool UseGPUScene(const FStaticShaderPlatform Platform, const FStaticFeatureLevel FeatureLevel);
 
 RENDERCORE_API bool UseGPUScene(const FStaticShaderPlatform Platform);
diff --git a/Engine/Source/Runtime/Renderer/Private/SingleLayerWaterRendering.cpp b/Engine/Source/Runtime/Renderer/Private/SingleLayerWaterRendering.cpp
--- a/Engine/Source/Runtime/Renderer/Private/SingleLayerWaterRendering.cpp
+++ b/Engine/Source/Runtime/Renderer/Private/SingleLayerWaterRendering.cpp
@@ -234,11 +234,6 @@ bool IsWaterVirtualShadowMapFilteringEnabled_Runtime(const FStaticShaderPlatform
 	return IsWaterVirtualShadowMapFilteringEnabled(Platform) && UseVirtualShadowMaps(Platform, GetMaxSupportedFeatureLevel(Platform)) && CVarWaterSingleLayerVSMFiltering.GetValueOnRenderThread() > 0;
 }
 
-bool NeedsSeparatedMainDirectionalLightTexture(const FStaticShaderPlatform Platform)
-{
-	return IsWaterDistanceFieldShadowEnabled(Platform) || IsWaterVirtualShadowMapFilteringEnabled(Platform);
-}
-
 bool NeedsSeparatedMainDirectionalLightTexture_Runtime(const FStaticShaderPlatform Platform)
 {
 	return IsWaterDistanceFieldShadowEnabled_Runtime(Platform) || IsWaterVirtualShadowMapFilteringEnabled_Runtime(Platform);
@@ -1250,7 +1245,7 @@ void FDeferredShadingSceneRenderer::RenderSingleLayerWaterInner(
 
 	TStaticArray<FTextureRenderTargetBinding, MaxSimultaneousRenderTargets> BasePassTextures;
 	uint32 BasePassTextureCount = SceneTextures.GetGBufferRenderTargets(BasePassTextures);
-	if(NeedsSeparatedMainDirectionalLightTexture(Scene->GetShaderPlatform()))
+	if(IsWaterSeparateMainDirLightEnabled(Scene->GetShaderPlatform()))
 	{
 		const bool bNeverClear = true;
 		BasePassTextures[BasePassTextureCount++] = FTextureRenderTargetBinding(SceneWithoutWaterTextures.SeparatedMainDirLightTexture, bNeverClear);
@@ -1528,7 +1523,7 @@ void FSingleLayerWaterPassMeshProcessor::CollectPSOInitializers(const FSceneText
 
 		FGraphicsPipelineRenderTargetsInfo RenderTargetsInfo;
 		SetupGBufferRenderTargetInfo(SceneTexturesConfig, RenderTargetsInfo, true /*bSetupDepthStencil*/);
-		if (NeedsSeparatedMainDirectionalLightTexture(GMaxRHIShaderPlatform))
+		if (IsWaterSeparateMainDirLightEnabled(GMaxRHIShaderPlatform))
 		{
 			AddRenderTargetInfo(PF_FloatR11G11B10, TexCreate_ShaderResource | TexCreate_RenderTargetable, RenderTargetsInfo);
 		}
diff --git a/Engine/Source/Runtime/Renderer/Private/Substrate/Substrate.cpp b/Engine/Source/Runtime/Renderer/Private/Substrate/Substrate.cpp
--- a/Engine/Source/Runtime/Renderer/Private/Substrate/Substrate.cpp
+++ b/Engine/Source/Runtime/Renderer/Private/Substrate/Substrate.cpp
@@ -1292,7 +1292,9 @@ void SetBasePassRenderTargetOutputFormat(const EShaderPlatform Platform, const F
 		// If it is not a water material, we force bHasSingleLayerWaterSeparatedMainLight to false, in order to 
 		// ensure non-used MRTs are not inserted in BufferInfo. Otherwise this would offset Substrate MRTs, causing 
 		// MRTs' format to be incorrect
-		if (!MaterialParameters.bIsUsedWithWater)
+		const bool bIsSingleLayerWater = MaterialParameters.ShadingModels.HasShadingModel(MSM_SingleLayerWater);
+		const bool bNeedsSeparateMainDirLightTexture = IsWaterSeparateMainDirLightEnabled(Platform);
+		if (!bIsSingleLayerWater || !bNeedsSeparateMainDirLightTexture)
 		{
 			GBufferParams.bHasSingleLayerWaterSeparatedMainLight = false;
 		}
