Light function atlas - do nt display on screen messages when GAreScreenMessagesEnabled is false.

diff --git a/Engine/Source/Runtime/Renderer/Private/LightFunctionAtlas.cpp b/Engine/Source/Runtime/Renderer/Private/LightFunctionAtlas.cpp
--- a/Engine/Source/Runtime/Renderer/Private/LightFunctionAtlas.cpp
+++ b/Engine/Source/Runtime/Renderer/Private/LightFunctionAtlas.cpp
@@ -766,7 +766,7 @@ void FLightFunctionAtlas::RenderAtlasSlots(FRDGBuilder& GraphBuilder, const TArr
 void FLightFunctionAtlas::RenderDebugInfo(FRDGBuilder& GraphBuilder, TArray<FViewInfo>& Views)
 {
 #if !UE_BUILD_SHIPPING
-	if (!IsLightFunctionAtlasEnabled())
+	if (!IsLightFunctionAtlasEnabled() || !GAreScreenMessagesEnabled)
 	{
 		return;
 	}
