Fix global distance field not being built in some situations even though materials require it.
- wasn't taking into account nanite relevance.

diff --git a/Engine/Source/Runtime/Renderer/Private/SceneVisibility.cpp b/Engine/Source/Runtime/Renderer/Private/SceneVisibility.cpp
--- a/Engine/Source/Runtime/Renderer/Private/SceneVisibility.cpp
+++ b/Engine/Source/Runtime/Renderer/Private/SceneVisibility.cpp
@@ -2191,6 +2191,7 @@ void FComputeAndMarkRelevance::Finalize()
 				const FPrimitiveViewRelevance& CombinedRelevance = ShadingPipelines.CombinedRelevance;
 
 				WriteView.ShadingModelMaskInView |= CombinedRelevance.ShadingModelMask;
+				WriteView.bUsesGlobalDistanceField |= CombinedRelevance.bUsesGlobalDistanceField;
 				WriteView.bUsesLightingChannels |= CombinedRelevance.bUsesLightingChannels;
 				WriteView.bSceneHasSkyMaterial |= CombinedRelevance.bUsesSkyMaterial;
 				WriteView.bHasDistortionPrimitives |= CombinedRelevance.bDistortion;
