Allocate separate instance range for decal instances in RayTracingScene.
- having separate ranges is necessary to be able to move management of ray tracing scene layers out of RHI.

diff --git a/Engine/Source/Runtime/Renderer/Private/RayTracing/RayTracing.cpp b/Engine/Source/Runtime/Renderer/Private/RayTracing/RayTracing.cpp
--- a/Engine/Source/Runtime/Renderer/Private/RayTracing/RayTracing.cpp
+++ b/Engine/Source/Runtime/Renderer/Private/RayTracing/RayTracing.cpp
@@ -168,7 +168,8 @@ namespace RayTracing
 		TArrayView<const int32> CachedRayTracingMeshCommandIndices; // Pointer to FPrimitiveSceneInfo::CachedRayTracingMeshCommandIndicesPerLOD data
 
 		// Offsets relative to FRelevantPrimitiveContext offsets
-		int32 RelativeSceneInstanceOffset = INDEX_NONE;
+		int32 RelativeInstanceOffset = INDEX_NONE;
+		int32 RelativeDecalInstanceOffset = INDEX_NONE;
 		int32 RelativeVisibleMeshCommandOffset = INDEX_NONE;
 		int32 ContextIndex = INDEX_NONE;
 
@@ -199,7 +200,8 @@ namespace RayTracing
 
 	struct FRelevantPrimitiveGatherContext
 	{
-		int32 SceneInstanceOffset = -1;
+		int32 InstanceOffset = -1;
+		int32 DecalInstanceOffset = -1;
 		int32 VisibleMeshCommandOffset = -1;
 	};
 
@@ -222,7 +224,8 @@ namespace RayTracing
 		// Used coarse mesh streaming handles during the last TLAS build
 		TArray<Nanite::CoarseMeshStreamingHandle> UsedCoarseMeshStreamingHandles; // TODO: Should be a set
 
-		int32 NumCachedStaticSceneInstances = 0;
+		int32 NumCachedStaticInstances = 0;
+		int32 NumCachedStaticDecalInstances = 0;
 		int32 NumCachedStaticVisibleMeshCommands = 0;
 
 		// Indicates that this object has been fully produced (for validation)
@@ -449,7 +452,8 @@ namespace RayTracing
 					TChunkedArray<FRelevantPrimitive> CachedStaticPrimitives;
 					TChunkedArray<const FPrimitiveSceneInfo*> VisibleNaniteRayTracingPrimitives;
 
-					int32 NumCachedStaticSceneInstances = 0;
+					int32 NumCachedStaticInstances = 0;
+					int32 NumCachedStaticDecalInstances = 0;
 					int32 NumCachedStaticVisibleMeshCommands = 0;
 
 					int32 ContextIndex = INDEX_NONE;
@@ -540,21 +544,34 @@ namespace RayTracing
 								// TODO: check if this ever happens. should probably skip primitive if so
 							}
 
+							const bool bNeedMainInstance = !RelevantPrimitive->bAllSegmentsDecal;
+
 							// if primitive has mixed decal and non-decal segments we need to have two ray tracing instances
 							// one containing non-decal segments and the other with decal segments
 							// masking of segments is done using "hidden" hitgroups
 							// TODO: Debug Visualization to highlight primitives using this?
 							const bool bNeedDecalInstance = RelevantPrimitive->bAnySegmentsDecal && !ShouldExcludeDecals();
 
-							const uint32 NumTLASInstances = bNeedDecalInstance && !RelevantPrimitive->bAllSegmentsDecal ? 2 : 1;
+							checkf(bNeedMainInstance || bNeedDecalInstance, TEXT("FRelevantPrimitive is expected to have a main instance, decal instance or both."));
+							const uint32 NumTLASInstances = bNeedMainInstance && bNeedDecalInstance ? 2 : 1;
 
 							// For now store offsets relative to current context
 							// Will be patched later to be a global offset
-							RelevantPrimitive->RelativeSceneInstanceOffset = Context.NumCachedStaticSceneInstances;
+							RelevantPrimitive->RelativeInstanceOffset = Context.NumCachedStaticInstances;
+							RelevantPrimitive->RelativeDecalInstanceOffset = Context.NumCachedStaticDecalInstances;
 							RelevantPrimitive->RelativeVisibleMeshCommandOffset = Context.NumCachedStaticVisibleMeshCommands;
 							RelevantPrimitive->ContextIndex = Context.ContextIndex;
 
-							Context.NumCachedStaticSceneInstances += NumTLASInstances;
+							if (bNeedMainInstance)
+							{
+								++Context.NumCachedStaticInstances;
+							}
+
+							if (bNeedDecalInstance)
+							{
+								++Context.NumCachedStaticDecalInstances;
+							}
+
 							Context.NumCachedStaticVisibleMeshCommands += RelevantPrimitive->CachedRayTracingMeshCommandIndices.Num() * NumTLASInstances;
 						}
 						// - DirtyCachedRayTracingPrimitives are only processed after StaticPrimitiveIndices is filled
@@ -658,10 +675,12 @@ namespace RayTracing
 							Context.StaticPrimitives.CopyToLinearArray(Result.StaticPrimitives);
 							Context.CachedStaticPrimitives.CopyToLinearArray(Result.CachedStaticPrimitives);
 
-							GatherContext.SceneInstanceOffset = Result.NumCachedStaticSceneInstances;
+							GatherContext.InstanceOffset = Result.NumCachedStaticInstances;
+							GatherContext.DecalInstanceOffset = Result.NumCachedStaticDecalInstances;
 							GatherContext.VisibleMeshCommandOffset = Result.NumCachedStaticVisibleMeshCommands;
 
-							Result.NumCachedStaticSceneInstances += Context.NumCachedStaticSceneInstances;
+							Result.NumCachedStaticInstances += Context.NumCachedStaticInstances;
+							Result.NumCachedStaticDecalInstances += Context.NumCachedStaticDecalInstances;
 							Result.NumCachedStaticVisibleMeshCommands += Context.NumCachedStaticVisibleMeshCommands;
 
 							for (const FPrimitiveSceneInfo* SceneInfo : Context.VisibleNaniteRayTracingPrimitives)
@@ -1208,7 +1227,8 @@ namespace RayTracing
 			const FRayTracingCullingParameters& CullingParameters;
 			const bool bIsPathTracing;
 
-			const int32& NumCachedStaticSceneInstances;
+			const int32& NumCachedStaticInstances;
+			const int32& NumCachedStaticDecalInstances;
 			const int32& NumCachedStaticVisibleMeshCommands;
 
 			// Outputs
@@ -1223,7 +1243,8 @@ namespace RayTracing
 				TArray<FRelevantPrimitiveGatherContext>& InGatherContexts,
 				const FRayTracingCullingParameters& InCullingParameters,
 				const bool bInIsPathTracing,
-				const int32& InNumCachedStaticSceneInstances,
+				const int32& InNumCachedStaticInstances,
+				const int32& InNumCachedStaticDecalInstances,
 				const int32& InNumCachedStaticVisibleMeshCommands,
 				FRayTracingScene& InRayTracingScene, TArray<FVisibleRayTracingMeshCommand>& InVisibleRayTracingMeshCommands)
 				: Scene(InScene)
@@ -1233,7 +1254,8 @@ namespace RayTracing
 				, GatherContexts(InGatherContexts)
 				, CullingParameters(InCullingParameters)
 				, bIsPathTracing(bInIsPathTracing)
-				, NumCachedStaticSceneInstances(InNumCachedStaticSceneInstances)
+				, NumCachedStaticInstances(InNumCachedStaticInstances)
+				, NumCachedStaticDecalInstances(InNumCachedStaticDecalInstances)
 				, NumCachedStaticVisibleMeshCommands(InNumCachedStaticVisibleMeshCommands)
 				, RayTracingScene(InRayTracingScene)
 				, VisibleRayTracingMeshCommands(InVisibleRayTracingMeshCommands)
@@ -1471,7 +1493,8 @@ namespace RayTracing
 				{
 					TRACE_CPUPROFILER_EVENT_SCOPE(RayTracingScene_AddCachedStaticInstances);
 					
-					const FRayTracingScene::FInstanceRange CachedStaticInstanceRange = RayTracingScene.AllocateInstanceRangeUninitialized(NumCachedStaticSceneInstances);
+					const FRayTracingScene::FInstanceRange CachedStaticInstanceRange = RayTracingScene.AllocateInstanceRangeUninitialized(NumCachedStaticInstances);
+					const FRayTracingScene::FInstanceRange CachedStaticDecalInstanceRange = RayTracingScene.AllocateInstanceRangeUninitialized(NumCachedStaticDecalInstances);
 					const uint32 BaseCachedVisibleMeshCommandsIndex = VisibleRayTracingMeshCommands.AddUninitialized(NumCachedStaticVisibleMeshCommands);
 					const uint32 BaseCachedGlobalSegmentIndex = RayTracingScene.NumSegments;
 					RayTracingScene.NumSegments += NumCachedStaticVisibleMeshCommands;
@@ -1481,7 +1504,7 @@ namespace RayTracing
 						TEXT("RayTracingScene_AddCachedStaticInstances_ParallelFor"),
 						RelevantCachedStaticPrimitives.Num(),
 						MinBatchSize,
-						[this, CachedStaticInstanceRange, BaseCachedVisibleMeshCommandsIndex, BaseCachedGlobalSegmentIndex](int32 Index)
+						[this, CachedStaticInstanceRange, CachedStaticDecalInstanceRange, BaseCachedVisibleMeshCommandsIndex, BaseCachedGlobalSegmentIndex](int32 Index)
 					{
 						const FRelevantPrimitive& RelevantPrimitive = RelevantCachedStaticPrimitives[Index];
 						const int32 PrimitiveIndex = RelevantPrimitive.PrimitiveIndex;
@@ -1515,28 +1538,28 @@ namespace RayTracing
 
 						check(RelevantPrimitive.CachedRayTracingInstance);
 
-						int32 InstanceIndexInRange = GatherContexts[RelevantPrimitive.ContextIndex].SceneInstanceOffset + RelevantPrimitive.RelativeSceneInstanceOffset;
-
 						if (bNeedMainInstance)
 						{
+							const int32 InstanceIndexInRange = GatherContexts[RelevantPrimitive.ContextIndex].InstanceOffset + RelevantPrimitive.RelativeInstanceOffset;
+
 							FRayTracingGeometryInstance RayTracingInstance = *RelevantPrimitive.CachedRayTracingInstance;
 							RayTracingInstance.LayerIndex = (uint8)ERayTracingSceneLayer::Base;
 							RayTracingInstance.InstanceContributionToHitGroupIndex = CalculateInstanceContributionToHitGroupIndex(MainGlobalSegmentIndex);
 							AddDebugRayTracingInstanceFlags(RayTracingInstance.Flags);
 
 							RayTracingScene.SetInstance(CachedStaticInstanceRange, InstanceIndexInRange, MoveTemp(RayTracingInstance), SceneProxy, false);
-							InstanceIndexInRange++;
 						}
 
 						if (bNeedDecalInstance)
 						{
+							const int32 DecalInstanceIndexInRange = GatherContexts[RelevantPrimitive.ContextIndex].DecalInstanceOffset + RelevantPrimitive.RelativeDecalInstanceOffset;
+
 							FRayTracingGeometryInstance DecalRayTracingInstance = *RelevantPrimitive.CachedRayTracingInstance;
 							DecalRayTracingInstance.LayerIndex = (uint8)ERayTracingSceneLayer::Decals;
 							DecalRayTracingInstance.InstanceContributionToHitGroupIndex = CalculateInstanceContributionToHitGroupIndex(DecalGlobalSegmentIndex);
 							AddDebugRayTracingInstanceFlags(DecalRayTracingInstance.Flags);
 
-							RayTracingScene.SetInstance(CachedStaticInstanceRange, InstanceIndexInRange, MoveTemp(DecalRayTracingInstance), SceneProxy, false);
-							InstanceIndexInRange++;
+							RayTracingScene.SetInstance(CachedStaticDecalInstanceRange, DecalInstanceIndexInRange, MoveTemp(DecalRayTracingInstance), SceneProxy, false);
 						}
 
 						const int32 VisibleMeshCommandOffset = BaseCachedVisibleMeshCommandsIndex + GatherContexts[RelevantPrimitive.ContextIndex].VisibleMeshCommandOffset + RelevantPrimitive.RelativeVisibleMeshCommandOffset;
@@ -1579,7 +1602,8 @@ namespace RayTracing
 			RelevantPrimitiveList.GatherContexts,
 			View.RayTracingCullingParameters,
 			bool(View.Family->EngineShowFlags.PathTracing),
-			RelevantPrimitiveList.NumCachedStaticSceneInstances,
+			RelevantPrimitiveList.NumCachedStaticInstances,
+			RelevantPrimitiveList.NumCachedStaticDecalInstances,
 			RelevantPrimitiveList.NumCachedStaticVisibleMeshCommands,
 			// outputs
 			RayTracingScene,
