Replace the Single Layer Water computation of Seperated Main Dir. light render target index, by the GBuffferInfo helper functionality. This ensures the logic is similar in BasePassRendering.cpp and in SetBasePassRenderTargetOutputFormat.

This fives incorrect TopLayerNormal encoding on certain platform for SLW.
This is a follow up of CL 34027621

diff --git a/Engine/Source/Runtime/Renderer/Private/BasePassRendering.cpp b/Engine/Source/Runtime/Renderer/Private/BasePassRendering.cpp
--- a/Engine/Source/Runtime/Renderer/Private/BasePassRendering.cpp
+++ b/Engine/Source/Runtime/Renderer/Private/BasePassRendering.cpp
@@ -28,6 +28,8 @@
 #include "DataDrivenShaderPlatformInfo.h"
 #include "VolumetricFog.h"
 #include "PostProcess/SceneRenderTargets.h"
+#include "ShaderCompiler.h"
+#include "GBufferInfo.h"
 
 #include "BasePassRendering.inl"
 
@@ -951,35 +953,10 @@ void ModifyBasePassCSPSCompilationEnvironment(const FMeshMaterialShaderPermutati
 	const bool bNeedsSeparateMainDirLightTexture = IsWaterDistanceFieldShadowEnabled(Parameters.Platform) || IsWaterVirtualShadowMapFilteringEnabled(Parameters.Platform);
 	if (bIsSingleLayerWater && bNeedsSeparateMainDirLightTexture)
 	{
-		// See FShaderCompileUtilities::FetchGBufferParamsRuntime for the details
-		const bool bHasTangent = false;
-		bool bHasPrecShadowFactor = IsStaticLightingAllowed();
+		const FGBufferParams GBufferParams = FShaderCompileUtilities::FetchGBufferParamsRuntime(Parameters.Platform, GBufferLayout);
+		const FGBufferInfo BufferInfo = FetchFullGBufferInfo(GBufferParams);
+		const uint32 TargetSeparatedMainDirLight = BufferInfo.Slots[GBS_SeparatedMainDirLight].Packing[0].TargetIndex;
 
-		uint32 TargetSeparatedMainDirLight = 5;
-		if (bOutputVelocity == false && bHasTangent == false)
-		{
-			TargetSeparatedMainDirLight = 5;
-			if (bHasPrecShadowFactor)
-			{
-				TargetSeparatedMainDirLight = 6;
-			}
-		}
-		else if (bOutputVelocity)
-		{
-			TargetSeparatedMainDirLight = 6;
-			if (bHasPrecShadowFactor)
-			{
-				TargetSeparatedMainDirLight = 7;
-			}
-		}
-		else if (bHasTangent)
-		{
-			TargetSeparatedMainDirLight = 6;
-			if (bHasPrecShadowFactor)
-			{
-				TargetSeparatedMainDirLight = 7;
-			}
-		}
 		OutEnvironment.SetRenderTargetOutputFormat(TargetSeparatedMainDirLight, PF_FloatR11G11B10);
 	}
 }
