Fix crash in hair transmittance with clustered shading when there is no light casting shadows.

diff --git a/Engine/Source/Runtime/Renderer/Private/HairStrands/HairStrandsTransmittance.cpp b/Engine/Source/Runtime/Renderer/Private/HairStrands/HairStrandsTransmittance.cpp
--- a/Engine/Source/Runtime/Renderer/Private/HairStrands/HairStrandsTransmittance.cpp
+++ b/Engine/Source/Runtime/Renderer/Private/HairStrands/HairStrandsTransmittance.cpp
@@ -262,7 +262,7 @@ static FRDGBufferRef AddHairStrandsVoxelTransmittanceMaskPass(
 		Parameters->VirtualShadowMap = VirtualShadowMapArray->GetSamplingParameters(GraphBuilder);
 		Parameters->ForwardLightData = View.ForwardLightingResources.ForwardLightUniformBuffer;
 		Parameters->RayMarchMaskTexture = nullptr;
-		Parameters->ShadowMaskBitsTexture = ShadowMaskTexture;
+		Parameters->ShadowMaskBitsTexture = ShadowMaskTexture ? ShadowMaskTexture : GSystemTextures.GetZeroUIntDummy(GraphBuilder);
 	}
 	else
 	{
