Fix BasicRayTracing test.
- don't set hit group since bAllowHitGroupIndexing = false;
- specify FRHIRayTracing when creating SRV.

diff --git a/Engine/Source/Runtime/Renderer/Private/RayTracing/RayTracingBasicShaders.cpp b/Engine/Source/Runtime/Renderer/Private/RayTracing/RayTracingBasicShaders.cpp
--- a/Engine/Source/Runtime/Renderer/Private/RayTracing/RayTracingBasicShaders.cpp
+++ b/Engine/Source/Runtime/Renderer/Private/RayTracing/RayTracingBasicShaders.cpp
@@ -53,7 +53,7 @@ FBasicRayTracingPipeline GetBasicRayTracingPipeline(FRHICommandList& RHICmdList,
 	return Result;
 }
 
-void DispatchBasicOcclusionRays(FRHICommandList& RHICmdList, FRHIRayTracingScene* Scene, FRHIShaderResourceView* SceneView, FRHIRayTracingGeometry* Geometry, FRHIShaderResourceView* RayBufferView, FRHIUnorderedAccessView* ResultView, uint32 NumRays)
+void DispatchBasicOcclusionRays(FRHICommandList& RHICmdList, FRHIShaderResourceView* SceneView, FRHIRayTracingGeometry* Geometry, FRHIShaderResourceView* RayBufferView, FRHIUnorderedAccessView* ResultView, uint32 NumRays)
 {
 	FBasicRayTracingPipeline RayTracingPipeline = GetBasicRayTracingPipeline(RHICmdList, GMaxRHIFeatureLevel);
 
@@ -65,7 +65,6 @@ void DispatchBasicOcclusionRays(FRHICommandList& RHICmdList, FRHIRayTracingScene
 
 	FShaderBindingTableRHIRef SBT = RHICreateShaderBindingTable(SBTInitializer);
 
-	RHICmdList.SetRayTracingHitGroup(SBT, 0, Geometry, 0, RayTracingPipeline.PipelineState, 0, 0, nullptr, 0, nullptr, 0);
 	RHICmdList.SetRayTracingMissShader(SBT, 0, RayTracingPipeline.PipelineState, 0, 0, nullptr, 0);
 	RHICmdList.CommitShaderBindingTable(SBT);
 
@@ -79,7 +78,7 @@ void DispatchBasicOcclusionRays(FRHICommandList& RHICmdList, FRHIRayTracingScene
 	RHICmdList.RayTraceDispatch(RayTracingPipeline.PipelineState, RayTracingPipeline.OcclusionRGS.GetRayTracingShader(), SBT, GlobalResources, NumRays, 1);
 }
 
-void DispatchBasicIntersectionRays(FRHICommandList& RHICmdList, FRHIRayTracingScene* Scene, FRHIShaderResourceView* SceneView, FRHIRayTracingGeometry* Geometry, FRHIShaderResourceView* RayBufferView, FRHIUnorderedAccessView* ResultView, uint32 NumRays)
+void DispatchBasicIntersectionRays(FRHICommandList& RHICmdList, FRHIShaderResourceView* SceneView, FRHIRayTracingGeometry* Geometry, FRHIShaderResourceView* RayBufferView, FRHIUnorderedAccessView* ResultView, uint32 NumRays)
 {
 	FBasicRayTracingPipeline RayTracingPipeline = GetBasicRayTracingPipeline(RHICmdList, GMaxRHIFeatureLevel);
 
@@ -91,7 +90,6 @@ void DispatchBasicIntersectionRays(FRHICommandList& RHICmdList, FRHIRayTracingSc
 
 	FShaderBindingTableRHIRef SBT = RHICreateShaderBindingTable(SBTInitializer);
 
-	RHICmdList.SetRayTracingHitGroup(SBT, 0, Geometry, 0, RayTracingPipeline.PipelineState, 0, 0, nullptr, 0, nullptr, 0);
 	RHICmdList.SetRayTracingMissShader(SBT, 0, RayTracingPipeline.PipelineState, 0, 0, nullptr, 0);
 	RHICmdList.CommitShaderBindingTable(SBT);
 
diff --git a/Engine/Source/Runtime/Renderer/Private/RayTracing/RayTracingBasicShaders.h b/Engine/Source/Runtime/Renderer/Private/RayTracing/RayTracingBasicShaders.h
--- a/Engine/Source/Runtime/Renderer/Private/RayTracing/RayTracingBasicShaders.h
+++ b/Engine/Source/Runtime/Renderer/Private/RayTracing/RayTracingBasicShaders.h
@@ -81,12 +81,12 @@ public:
 * Binary intersection results are written to output buffer as R32_UINTs.
 * 0xFFFFFFFF is written if ray intersects any scene triangle, 0 otherwise.
 */
-void DispatchBasicOcclusionRays(FRHICommandList& RHICmdList, FRHIRayTracingScene* Scene, FRHIShaderResourceView* SceneView, FRHIRayTracingGeometry* Geometry, FRHIShaderResourceView* RayBufferView, FRHIUnorderedAccessView* ResultView, uint32 NumRays);
+void DispatchBasicOcclusionRays(FRHICommandList& RHICmdList, FRHIShaderResourceView* SceneView, FRHIRayTracingGeometry* Geometry, FRHIShaderResourceView* RayBufferView, FRHIUnorderedAccessView* ResultView, uint32 NumRays);
 
 /**
 * Trace rays from an input buffer of FBasicRayTracingRay.
 * Primitive intersection results are written to output buffer as FBasicRayTracingIntersectionResult.
 */
-void DispatchBasicIntersectionRays(FRHICommandList& RHICmdList, FRHIRayTracingScene* Scene, FRHIShaderResourceView* SceneView, FRHIRayTracingGeometry* Geometry, FRHIShaderResourceView* RayBufferView, FRHIUnorderedAccessView* ResultView, uint32 NumRays);
+void DispatchBasicIntersectionRays(FRHICommandList& RHICmdList, FRHIShaderResourceView* SceneView, FRHIRayTracingGeometry* Geometry, FRHIShaderResourceView* RayBufferView, FRHIUnorderedAccessView* ResultView, uint32 NumRays);
 
 #endif // RHI_RAYTRACING
diff --git a/Engine/Source/Runtime/Renderer/Private/Tests/RayTracingTestbed.cpp b/Engine/Source/Runtime/Renderer/Private/Tests/RayTracingTestbed.cpp
--- a/Engine/Source/Runtime/Renderer/Private/Tests/RayTracingTestbed.cpp
+++ b/Engine/Source/Runtime/Renderer/Private/Tests/RayTracingTestbed.cpp
@@ -293,11 +293,11 @@ bool RunRayTracingTestbed_RenderThread(const FString& Parameters)
 
 	RHICmdList.Transition(FRHITransitionInfo(RayTracingScene.Scene.GetReference(), ERHIAccess::BVHWrite, ERHIAccess::BVHRead));
 
-	FShaderResourceViewInitializer RayTracingSceneViewInitializer(SceneBuffer, RayTracingScene.Scene->GetLayerBufferOffset(0), 0);
+	FShaderResourceViewInitializer RayTracingSceneViewInitializer(SceneBuffer, RayTracingScene.Scene, 0);
 	FShaderResourceViewRHIRef RayTracingSceneView = RHICmdList.CreateShaderResourceView(RayTracingSceneViewInitializer);
 
-	DispatchBasicOcclusionRays(RHICmdList, RayTracingScene.Scene, RayTracingSceneView, Geometry, RayBufferView, OcclusionResultBufferView, NumRays);
-	DispatchBasicIntersectionRays(RHICmdList, RayTracingScene.Scene, RayTracingSceneView, Geometry, RayBufferView, IntersectionResultBufferView, NumRays);
+	DispatchBasicOcclusionRays(RHICmdList, RayTracingSceneView, Geometry, RayBufferView, OcclusionResultBufferView, NumRays);
+	DispatchBasicIntersectionRays(RHICmdList, RayTracingSceneView, Geometry, RayBufferView, IntersectionResultBufferView, NumRays);
 
 	const bool bValidateResults = true;
 	bool bOcclusionTestOK = false;
