Refactor logic around FRayTracingScene::BuildInitializationData().
- store result of BuildInitializationData() directly inside FRayTracingScene instead of storing in FViewInfo and then move it around.
- will also make it easier to move layer management out of RHI since each layer will have its own FRayTracingSceneWithGeometryInstances data.
- Create(...) method will automatically call BuildInitializationData(...) if it necessary. BuildInitializationData() is still exposed as a public method since it allows higher level code to optionally run part of the FRayTracingScene setup in a async task.

diff --git a/Engine/Source/Runtime/Renderer/Private/DeferredShadingRenderer.cpp b/Engine/Source/Runtime/Renderer/Private/DeferredShadingRenderer.cpp
--- a/Engine/Source/Runtime/Renderer/Private/DeferredShadingRenderer.cpp
+++ b/Engine/Source/Runtime/Renderer/Private/DeferredShadingRenderer.cpp
@@ -893,7 +893,7 @@ bool FDeferredShadingSceneRenderer::DispatchRayTracingWorldUpdates(FRDGBuilder&
 
 	{
 		RDG_GPU_STAT_SCOPE(GraphBuilder, RayTracingScene);
-		RayTracingScene.CreateWithInitializationData(GraphBuilder, ReferenceView, &Scene->GPUScene, MoveTemp(ReferenceView.RayTracingSceneInitData), ComputePassFlags);
+		RayTracingScene.Create(GraphBuilder, ReferenceView, &Scene->GPUScene, ComputePassFlags);
 		RayTracingScene.Build(GraphBuilder, ComputePassFlags | ERDGPassFlags::NeverCull, OutDynamicGeometryScratchBuffer);
 	}
 
diff --git a/Engine/Source/Runtime/Renderer/Private/HeterogeneousVolumes/HeterogeneousVolumesHardwareRayTracing.cpp b/Engine/Source/Runtime/Renderer/Private/HeterogeneousVolumes/HeterogeneousVolumesHardwareRayTracing.cpp
--- a/Engine/Source/Runtime/Renderer/Private/HeterogeneousVolumes/HeterogeneousVolumesHardwareRayTracing.cpp
+++ b/Engine/Source/Runtime/Renderer/Private/HeterogeneousVolumes/HeterogeneousVolumesHardwareRayTracing.cpp
@@ -256,7 +256,7 @@ void GenerateRayTracingScene(
 
 	// Create RayTracingScene
 	const FGPUScene* EmptyGPUScene = nullptr;
-	RayTracingScene.Create(GraphBuilder, View, EmptyGPUScene);
+	RayTracingScene.Create(GraphBuilder, View, EmptyGPUScene, ERDGPassFlags::Compute);
 
 	const bool bRayTracingAsyncBuild = false;//CVarRayTracingAsyncBuild.GetValueOnRenderThread() != 0 && GRHISupportsRayTracingAsyncBuildAccelerationStructure;
 	const ERDGPassFlags ComputePassFlags = bRayTracingAsyncBuild ? ERDGPassFlags::AsyncCompute : ERDGPassFlags::Compute;
diff --git a/Engine/Source/Runtime/Renderer/Private/RayTracing/RayTracing.cpp b/Engine/Source/Runtime/Renderer/Private/RayTracing/RayTracing.cpp
--- a/Engine/Source/Runtime/Renderer/Private/RayTracing/RayTracing.cpp
+++ b/Engine/Source/Runtime/Renderer/Private/RayTracing/RayTracing.cpp
@@ -1617,7 +1617,7 @@ namespace RayTracing
 				FTaskTagScope TaskTagScope(ETaskTag::EParallelRenderingThread);
 
 				TRACE_CPUPROFILER_EVENT_SCOPE(RayTracingSceneInitTask);
-				View.RayTracingSceneInitData = RayTracingScene.BuildInitializationData();
+				RayTracingScene.BuildInitializationData();
 			},
 			TStatId(), AddStaticInstancesTask, ENamedThreads::AnyThread);
 
diff --git a/Engine/Source/Runtime/Renderer/Private/RayTracing/RayTracingDebug.cpp b/Engine/Source/Runtime/Renderer/Private/RayTracing/RayTracingDebug.cpp
--- a/Engine/Source/Runtime/Renderer/Private/RayTracing/RayTracingDebug.cpp
+++ b/Engine/Source/Runtime/Renderer/Private/RayTracing/RayTracingDebug.cpp
@@ -1527,7 +1527,7 @@ void FDeferredShadingSceneRenderer::RenderRayTracingDebug(FRDGBuilder& GraphBuil
 		bRequiresBindings = true;
 	}
 
-	const uint32 NumInstances = View.RayTracingSceneInitData.NumNativeCPUInstances + View.RayTracingSceneInitData.NumNativeGPUSceneInstances;
+	const uint32 NumInstances = RayTracingScene.GetNumNativeInstances();
 
 	FRayTracingDebugRGS::FParameters* RayGenParameters = GraphBuilder.AllocParameters<FRayTracingDebugRGS::FParameters>();
 
diff --git a/Engine/Source/Runtime/Renderer/Private/RayTracing/RayTracingScene.cpp b/Engine/Source/Runtime/Renderer/Private/RayTracing/RayTracingScene.cpp
--- a/Engine/Source/Runtime/Renderer/Private/RayTracing/RayTracingScene.cpp
+++ b/Engine/Source/Runtime/Renderer/Private/RayTracing/RayTracingScene.cpp
@@ -43,18 +43,13 @@ FRayTracingScene::~FRayTracingScene()
 	ReleaseReadbackBuffers();
 }
 
-FRayTracingSceneWithGeometryInstances FRayTracingScene::BuildInitializationData() const
+void FRayTracingScene::BuildInitializationData()
 {
 	ERayTracingAccelerationStructureFlags BuildFlags = CVarRayTracingSceneBuildMode.GetValueOnRenderThread()
 		? ERayTracingAccelerationStructureFlags::FastTrace
 		: ERayTracingAccelerationStructureFlags::FastBuild;
 
-	return CreateRayTracingSceneWithGeometryInstances(Instances, uint8(ERayTracingSceneLayer::NUM), BuildFlags);
-}
-
-void FRayTracingScene::Create(FRDGBuilder& GraphBuilder, const FViewInfo& View, const FGPUScene* GPUScene)
-{
-	CreateWithInitializationData(GraphBuilder, View, GPUScene, BuildInitializationData(), ERDGPassFlags::Compute);
+	InitializationData = CreateRayTracingSceneWithGeometryInstances(Instances, uint8(ERayTracingSceneLayer::NUM), BuildFlags);
 }
 
 void FRayTracingScene::InitPreViewTranslation(const FViewMatrices& ViewMatrices)
@@ -62,10 +57,10 @@ void FRayTracingScene::InitPreViewTranslation(const FViewMatrices& ViewMatrices)
 	PreViewTranslation = ViewMatrices.GetPreViewTranslation();
 }
 
-void FRayTracingScene::CreateWithInitializationData(FRDGBuilder& GraphBuilder, const FViewInfo& View, const FGPUScene* GPUScene, FRayTracingSceneWithGeometryInstances SceneWithGeometryInstances, ERDGPassFlags ComputePassFlags)
+void FRayTracingScene::Create(FRDGBuilder& GraphBuilder, const FViewInfo& View, const FGPUScene* GPUScene, ERDGPassFlags ComputePassFlags)
 {
-	TRACE_CPUPROFILER_EVENT_SCOPE(FRayTracingScene::CreateWithInitializationData);
-	QUICK_SCOPE_CYCLE_COUNTER(STAT_RayTracingScene_CreateWithInitializationData);
+	TRACE_CPUPROFILER_EVENT_SCOPE(FRayTracingScene::Create);
+	QUICK_SCOPE_CYCLE_COUNTER(STAT_RayTracingScene_Create);
 
 	// Round up buffer sizes to some multiple to avoid pathological growth reallocations.
 	static constexpr uint32 AllocationGranularity = 8 * 1024;
@@ -77,17 +72,22 @@ void FRayTracingScene::CreateWithInitializationData(FRDGBuilder& GraphBuilder, c
 
 	static const uint8 NumLayers = uint8(ERayTracingSceneLayer::NUM);
 
-	checkf(SceneWithGeometryInstances.Scene.IsValid(), 
-		TEXT("Ray tracing scene RHI object is expected to have been created by BuildInitializationData() or CreateRayTracingSceneWithGeometryInstances()"));
+	if (!InitializationData.Scene.IsValid())
+	{
+		BuildInitializationData();
+	}
+
+	checkf(InitializationData.Scene.IsValid(), 
+		TEXT("Ray tracing scene RHI object is expected to have been created by BuildInitializationData()"));
 
-	RayTracingSceneRHI = SceneWithGeometryInstances.Scene;
-	check(NumSegments == SceneWithGeometryInstances.TotalNumSegments);
+	RayTracingSceneRHI = InitializationData.Scene;
+	check(NumSegments == InitializationData.TotalNumSegments);
 
 	const FRayTracingSceneInitializer2& SceneInitializer = RayTracingSceneRHI->GetInitializer();
 
-	const uint32 NumNativeInstances = SceneWithGeometryInstances.NumNativeGPUSceneInstances + SceneWithGeometryInstances.NumNativeCPUInstances;
+	const uint32 NumNativeInstances = InitializationData.NumNativeGPUSceneInstances + InitializationData.NumNativeCPUInstances;
 	const uint32 NumNativeInstancesAligned = FMath::DivideAndRoundUp(FMath::Max(NumNativeInstances, 1U), AllocationGranularity) * AllocationGranularity;
-	const uint32 NumTransformsAligned = FMath::DivideAndRoundUp(FMath::Max(SceneWithGeometryInstances.NumNativeCPUInstances, 1U), AllocationGranularity) * AllocationGranularity;
+	const uint32 NumTransformsAligned = FMath::DivideAndRoundUp(FMath::Max(InitializationData.NumNativeCPUInstances, 1U), AllocationGranularity) * AllocationGranularity;
 
 	FRayTracingAccelerationStructureSize SizeInfo = RayTracingSceneRHI->GetSizeInfo();
 	SizeInfo.ResultSize = FMath::DivideAndRoundUp(FMath::Max(SizeInfo.ResultSize, 1ull), BufferAllocationGranularity) * BufferAllocationGranularity;
@@ -208,10 +208,6 @@ void FRayTracingScene::CreateWithInitializationData(FRDGBuilder& GraphBuilder, c
 			LayerBaseIndices[LayerIndex] = LayerBaseIndices[LayerIndex - 1] + SceneInitializer.NumNativeInstancesPerLayer[LayerIndex - 1];
 		}
 
-		// make a copy of SceneWithGeometryInstances.BaseInstancePrefixSum that can be passed to the RDG setup tasks below
-		TArray<uint32, SceneRenderingAllocator> BaseInstancePrefixSum = GraphBuilder.AllocArray<uint32>();
-		BaseInstancePrefixSum = SceneWithGeometryInstances.BaseInstancePrefixSum;
-
 		FRDGUploadData<FRayTracingInstanceDebugData> UploadData(GraphBuilder, NumNativeInstances);
 
 		{
@@ -230,7 +226,7 @@ void FRayTracingScene::CreateWithInitializationData(FRDGBuilder& GraphBuilder, c
 
 				TConstArrayView<FRayTracingGeometryInstance> TaskInstancesData(Instances.GetData() + TaskFirstItemIndex, TaskNumItems);
 				TConstArrayView<FRayTracingInstanceDebugData> TaskInstancesDebugData(InstancesDebugData.GetData() + TaskFirstItemIndex, TaskNumItems);
-				TConstArrayView<uint32> TaskBaseInstancePrefixSum(BaseInstancePrefixSum.GetData() + TaskFirstItemIndex, TaskNumItems);
+				TConstArrayView<uint32> TaskBaseInstancePrefixSum(InitializationData.BaseInstancePrefixSum.GetData() + TaskFirstItemIndex, TaskNumItems);
 
 				GraphBuilder.AddSetupTask([UploadData, TaskInstancesDebugData, TaskInstancesData, TaskBaseInstancePrefixSum, LayerBaseIndices]()
 					{
@@ -258,7 +254,7 @@ void FRayTracingScene::CreateWithInitializationData(FRDGBuilder& GraphBuilder, c
 	if (NumNativeInstances > 0)
 	{
 		const uint32 InstanceUploadBytes = NumNativeInstances * sizeof(FRayTracingInstanceDescriptorInput);
-		const uint32 TransformUploadBytes = SceneWithGeometryInstances.NumNativeCPUInstances * 3 * sizeof(FVector4f);
+		const uint32 TransformUploadBytes = InitializationData.NumNativeCPUInstances * 3 * sizeof(FVector4f);
 
 		FRayTracingInstanceDescriptorInput* InstanceUploadData = (FRayTracingInstanceDescriptorInput*)RHICmdList.LockBuffer(InstanceUploadBuffer, 0, InstanceUploadBytes, RLM_WriteOnly);
 		FVector4f* TransformUploadData = (TransformUploadBytes > 0) ? (FVector4f*)RHICmdList.LockBuffer(TransformUploadBuffer, 0, TransformUploadBytes, RLM_WriteOnly) : nullptr;
@@ -266,13 +262,13 @@ void FRayTracingScene::CreateWithInitializationData(FRDGBuilder& GraphBuilder, c
 		// Fill instance upload buffer on separate thread since results are only needed in RHI thread
 		GraphBuilder.AddSetupTask(
 			[InstanceUploadData = MakeArrayView(InstanceUploadData, NumNativeInstances),
-			TransformUploadData = MakeArrayView(TransformUploadData, SceneWithGeometryInstances.NumNativeCPUInstances * 3),
-			NumNativeGPUSceneInstances = SceneWithGeometryInstances.NumNativeGPUSceneInstances,
-			NumNativeCPUInstances = SceneWithGeometryInstances.NumNativeCPUInstances,
+			TransformUploadData = MakeArrayView(TransformUploadData, InitializationData.NumNativeCPUInstances * 3),
+			NumNativeGPUSceneInstances = InitializationData.NumNativeGPUSceneInstances,
+			NumNativeCPUInstances = InitializationData.NumNativeCPUInstances,
 			Instances = MakeArrayView(Instances),
-			InstanceGeometryIndices = MoveTemp(SceneWithGeometryInstances.InstanceGeometryIndices),
-			BaseUploadBufferOffsets = MoveTemp(SceneWithGeometryInstances.BaseUploadBufferOffsets),
-			BaseInstancePrefixSum = MoveTemp(SceneWithGeometryInstances.BaseInstancePrefixSum),
+			InstanceGeometryIndices = MakeArrayView(InitializationData.InstanceGeometryIndices),
+			BaseUploadBufferOffsets = MakeArrayView(InitializationData.BaseUploadBufferOffsets),
+			BaseInstancePrefixSum = MakeArrayView(InitializationData.BaseInstancePrefixSum),
 			RayTracingSceneRHI = RayTracingSceneRHI,
 			PreViewTranslation = this->PreViewTranslation]()
 		{
@@ -307,8 +303,8 @@ void FRayTracingScene::CreateWithInitializationData(FRDGBuilder& GraphBuilder, c
 			this,
 			GPUScene,
 			&SceneInitializer,
-			NumNativeGPUSceneInstances = SceneWithGeometryInstances.NumNativeGPUSceneInstances,
-			NumNativeCPUInstances = SceneWithGeometryInstances.NumNativeCPUInstances,
+			NumNativeGPUSceneInstances = InitializationData.NumNativeGPUSceneInstances,
+			NumNativeCPUInstances = InitializationData.NumNativeCPUInstances,
 			CullingParameters = View.RayTracingCullingParameters
 			](FRHICommandList& RHICmdList)
 			{
diff --git a/Engine/Source/Runtime/Renderer/Private/RayTracing/RayTracingScene.h b/Engine/Source/Runtime/Renderer/Private/RayTracing/RayTracingScene.h
--- a/Engine/Source/Runtime/Renderer/Private/RayTracing/RayTracingScene.h
+++ b/Engine/Source/Runtime/Renderer/Private/RayTracing/RayTracingScene.h
@@ -94,15 +94,12 @@ public:
 	void SetInstance(FInstanceRange InstanceRange, uint32 InstanceIndexInRange, FRayTracingGeometryInstance Instance, const FPrimitiveSceneProxy* Proxy = nullptr, bool bDynamic = false);
 
 	// Allocates RayTracingSceneRHI and builds various metadata required to create the final scene.
-	FRayTracingSceneWithGeometryInstances BuildInitializationData() const;
+	void BuildInitializationData();
 
 	// Allocates GPU memory to fit at least the current number of instances.
 	// Kicks off instance buffer build to parallel thread along with RDG pass.
 	// NOTE: SceneWithGeometryInstances is passed in by value because ownership of the internal data is taken over. Use MoveTemp at call site, if possible.
-	void CreateWithInitializationData(FRDGBuilder& GraphBuilder, const FViewInfo& View, const FGPUScene* GPUScene, FRayTracingSceneWithGeometryInstances SceneWithGeometryInstances, ERDGPassFlags ComputePassFlags);
-
-	// Backwards-compatible version of Create() which internally calls CreateRayTracingSceneWithGeometryInstances().
-	void Create(FRDGBuilder& GraphBuilder, const FViewInfo& View, const FGPUScene* GPUScene);
+	void Create(FRDGBuilder& GraphBuilder, const FViewInfo& View, const FGPUScene* GPUScene, ERDGPassFlags ComputePassFlags);
 
 	void Build(FRDGBuilder& GraphBuilder, ERDGPassFlags ComputePassFlags, FRDGBufferRef DynamicGeometryScratchBuffer);
 
@@ -147,6 +144,8 @@ public:
 
 	uint32 GetTotalNumSegments() const { return NumSegments; }
 
+	uint32 GetNumNativeInstances() const { return InitializationData.NumNativeCPUInstances + InitializationData.NumNativeGPUSceneInstances; }
+
 	void InitPreViewTranslation(const FViewMatrices& ViewMatrices);
 
 public:
@@ -178,6 +177,8 @@ public:
 private:
 	void ReleaseReadbackBuffers();
 
+	FRayTracingSceneWithGeometryInstances InitializationData;
+
 	FRDGBufferRef InstanceBuffer;
 	FRDGBufferRef BuildScratchBuffer;
 
diff --git a/Engine/Source/Runtime/Renderer/Private/SceneRendering.h b/Engine/Source/Runtime/Renderer/Private/SceneRendering.h
--- a/Engine/Source/Runtime/Renderer/Private/SceneRendering.h
+++ b/Engine/Source/Runtime/Renderer/Private/SceneRendering.h
@@ -1263,9 +1263,7 @@ public:
 
 	FRayTracingCullingParameters RayTracingCullingParameters;
 
-	// Data required for FRayTracingScene, depends on RT instance culling tasks
-	FRayTracingSceneWithGeometryInstances RayTracingSceneInitData;
-	FGraphEventRef RayTracingSceneInitTask; // Task to asynchronously create RayTracingSceneInitData
+	FGraphEventRef RayTracingSceneInitTask; // Task to asynchronously call RayTracingScene.BuildInitializationData()
 
 	FGraphEventArray AddRayTracingMeshBatchTaskList;
 	TArray<FRayTracingMeshCommandOneFrameArray*, SceneRenderingAllocator> VisibleRayTracingMeshCommandsPerTask;
