Fixed r.Visibility.FrustumCull.UseFastIntersect to work and set as default.

diff --git a/Engine/Source/Runtime/Renderer/Private/SceneVisibility.cpp b/Engine/Source/Runtime/Renderer/Private/SceneVisibility.cpp
--- a/Engine/Source/Runtime/Renderer/Private/SceneVisibility.cpp
+++ b/Engine/Source/Runtime/Renderer/Private/SceneVisibility.cpp
@@ -303,10 +303,10 @@ static FAutoConsoleVariableRef CVarFrustumCullUseSphereTestFirst(
 	ECVF_RenderThreadSafe
 );
 
-static bool GFrustumCullUseFastIntersect = false;
-static TAutoConsoleVariable<int32> CVarFrustumCullUseFastIntersect(
+static bool GFrustumCullUseFastIntersect = true;
+static FAutoConsoleVariableRef CVarFrustumCullUseFastIntersect(
 	TEXT("r.Visibility.FrustumCull.UseFastIntersect"),
-	1,
+	GFrustumCullUseFastIntersect,
 	TEXT("Use optimized 8 plane fast intersection code if we have 8 permuted planes."),
 	ECVF_RenderThreadSafe
 );
