Implemented FRayTracingScene::GetInstanceDebugBuffer()/GetDebugInstanceGPUSceneIndexBuffer()
- also made the corresponding buffers private.
- this is necessary to manage layers at FRayTracingScene level in following changes.

diff --git a/Engine/Source/Runtime/Renderer/Private/RayTracing/RayTracingDebug.cpp b/Engine/Source/Runtime/Renderer/Private/RayTracing/RayTracingDebug.cpp
--- a/Engine/Source/Runtime/Renderer/Private/RayTracing/RayTracingDebug.cpp
+++ b/Engine/Source/Runtime/Renderer/Private/RayTracing/RayTracingDebug.cpp
@@ -832,7 +832,7 @@ static FRDGBufferRef RayTracingPerformPicking(FRDGBuilder& GraphBuilder, const F
 	FRDGBufferRef PickingBuffer = GraphBuilder.CreateBuffer(PickingBufferDesc, TEXT("RayTracingDebug.PickingBuffer"));
 
 	FRayTracingPickingRGS::FParameters* RayGenParameters = GraphBuilder.AllocParameters<FRayTracingPickingRGS::FParameters>();
-	RayGenParameters->InstancesDebugData = GraphBuilder.CreateSRV(RayTracingScene.InstanceDebugBuffer);
+	RayGenParameters->InstancesDebugData = GraphBuilder.CreateSRV(RayTracingScene.GetInstanceDebugBuffer());
 	RayGenParameters->TLAS = RayTracingScene.GetLayerView(ERayTracingSceneLayer::Base);
 	RayGenParameters->OpaqueOnly = CVarRayTracingDebugModeOpaqueOnly.GetValueOnRenderThread();
 	RayGenParameters->InstanceBuffer = GraphBuilder.CreateSRV(RayTracingScene.GetInstanceBuffer());
@@ -1271,7 +1271,7 @@ static void DrawInstanceOverlap(FRDGBuilder& GraphBuilder, const FScene* Scene,
 	FRDGTextureDesc InstanceOverlapTextureDesc = FRDGTextureDesc::Create2D(SceneColorTexture->Desc.Extent, PF_R32_FLOAT, FClearValueBinding::Black, TexCreate_ShaderResource | TexCreate_RenderTargetable);
 	FRDGTextureRef InstanceOverlapTexture = GraphBuilder.CreateTexture(InstanceOverlapTextureDesc, TEXT("RayTracingDebug::InstanceOverlap"));
 	
-	RayTracingDrawInstances(GraphBuilder, View, InstanceOverlapTexture, SceneDepthTexture, Scene->RayTracingScene.DebugInstanceGPUSceneIndexBuffer, false);
+	RayTracingDrawInstances(GraphBuilder, View, InstanceOverlapTexture, SceneDepthTexture, Scene->RayTracingScene.GetDebugInstanceGPUSceneIndexBuffer(), false);
 
 	// Calculate heatmap of instance overlap and blend it on top of ray tracing debug output
 	{
@@ -1299,7 +1299,7 @@ static void DrawInstanceOverlap(FRDGBuilder& GraphBuilder, const FScene* Scene,
 	// Draw instance AABB with lines
 	if (CVarRayTracingDebugInstanceOverlapShowWireframe.GetValueOnRenderThread() != 0)
 	{
-		RayTracingDrawInstances(GraphBuilder, View, SceneColorTexture, SceneDepthTexture, Scene->RayTracingScene.DebugInstanceGPUSceneIndexBuffer, true);
+		RayTracingDrawInstances(GraphBuilder, View, SceneColorTexture, SceneDepthTexture, Scene->RayTracingScene.GetDebugInstanceGPUSceneIndexBuffer(), true);
 	}
 }
 
@@ -1446,7 +1446,7 @@ void FDeferredShadingSceneRenderer::RenderRayTracingDebug(FRDGBuilder& GraphBuil
 
 	FRDGBufferRef PickingBuffer = nullptr;
 	FRDGBufferRef StatsBuffer = nullptr;
-	if (IsRayTracingPickingEnabled(DebugVisualizationMode) && RayTracingScene.InstanceDebugBuffer != nullptr)
+	if (IsRayTracingPickingEnabled(DebugVisualizationMode) && RayTracingScene.GetInstanceDebugBuffer() != nullptr)
 	{
 		PickingBuffer = RayTracingPerformPicking(GraphBuilder, Scene, View, PickingFeedback);
 	}
@@ -1470,7 +1470,7 @@ void FDeferredShadingSceneRenderer::RenderRayTracingDebug(FRDGBuilder& GraphBuil
 		AddClearUAVPass(GraphBuilder, DebugHitStatsUniformBufferParameters->HitStatsOutput, 0);
 	}
 
-	FRDGBufferRef InstanceDebugBuffer = RayTracingScene.InstanceDebugBuffer;
+	FRDGBufferRef InstanceDebugBuffer = RayTracingScene.GetInstanceDebugBuffer();
 	if (InstanceDebugBuffer == nullptr)
 	{
 		InstanceDebugBuffer = GSystemTextures.GetDefaultStructuredBuffer(GraphBuilder, sizeof(uint32));
diff --git a/Engine/Source/Runtime/Renderer/Private/RayTracing/RayTracingScene.h b/Engine/Source/Runtime/Renderer/Private/RayTracing/RayTracingScene.h
--- a/Engine/Source/Runtime/Renderer/Private/RayTracing/RayTracingScene.h
+++ b/Engine/Source/Runtime/Renderer/Private/RayTracing/RayTracingScene.h
@@ -146,6 +146,9 @@ public:
 
 	uint32 GetNumNativeInstances() const { return InitializationData.NumNativeCPUInstances + InitializationData.NumNativeGPUSceneInstances; }
 
+	FRDGBufferRef GetInstanceDebugBuffer() const { return InstanceDebugBuffer; }
+	FRDGBufferRef GetDebugInstanceGPUSceneIndexBuffer() const { return DebugInstanceGPUSceneIndexBuffer; }
+
 	void InitPreViewTranslation(const FViewMatrices& ViewMatrices);
 
 public:
@@ -166,10 +169,6 @@ public:
 	// Used coarse mesh streaming handles during the last TLAS build
 	TArray<Nanite::CoarseMeshStreamingHandle> UsedCoarseMeshStreamingHandles;
 
-	// Special data for debugging purposes
-	FRDGBufferRef InstanceDebugBuffer = nullptr;
-
-	FRDGBufferRef DebugInstanceGPUSceneIndexBuffer = nullptr;
 	bool bNeedsDebugInstanceGPUSceneIndexBuffer = false;
 
 	// Used for transforming to translated world space in which TLAS was built.
@@ -203,6 +202,10 @@ private:
 
 	FByteAddressBuffer AccelerationStructureAddressesBuffer;
 
+	// Special data for debugging purposes
+	FRDGBufferRef InstanceDebugBuffer = nullptr;
+	FRDGBufferRef DebugInstanceGPUSceneIndexBuffer = nullptr;
+
 	// Persistent storage for ray tracing instance descriptors.
 	// Cleared every frame without releasing memory to avoid large heap allocations.
 	// This must be filled before calling CreateRayTracingSceneWithGeometryInstances() and Create().
