Fix RDG ensure on DBufferTexArray texture not being produced in non-stereo

diff --git a/Engine/Source/Runtime/Renderer/Private/DBufferTextures.cpp b/Engine/Source/Runtime/Renderer/Private/DBufferTextures.cpp
--- a/Engine/Source/Runtime/Renderer/Private/DBufferTextures.cpp
+++ b/Engine/Source/Runtime/Renderer/Private/DBufferTextures.cpp
@@ -110,7 +110,7 @@ FDBufferTextures CreateDBufferTextures(FRDGBuilder& GraphBuilder, FIntPoint Exte
 	return DBufferTextures;
 }
 
-FDBufferParameters GetDBufferParameters(FRDGBuilder& GraphBuilder, const FDBufferTextures& DBufferTextures, EShaderPlatform ShaderPlatform)
+FDBufferParameters GetDBufferParameters(FRDGBuilder& GraphBuilder, const FDBufferTextures& DBufferTextures, EShaderPlatform ShaderPlatform, const bool bIsMobileMultiView)
 {
 	const FRDGSystemTextures& SystemTextures = FRDGSystemTextures::Get(GraphBuilder);
 
@@ -129,13 +129,18 @@ FDBufferParameters GetDBufferParameters(FRDGBuilder& GraphBuilder, const FDBuffe
 
 	if (DBufferTextures.IsValid())
 	{
-		Parameters.DBufferATexture = DBufferTextures.DBufferA;
-		Parameters.DBufferBTexture = DBufferTextures.DBufferB;
-		Parameters.DBufferCTexture = DBufferTextures.DBufferC;
-
-		Parameters.DBufferATextureArray = DBufferTextures.DBufferATexArray;
-		Parameters.DBufferBTextureArray = DBufferTextures.DBufferBTexArray;
-		Parameters.DBufferCTextureArray = DBufferTextures.DBufferCTexArray;
+		if (bIsMobileMultiView)
+		{
+			Parameters.DBufferATextureArray = DBufferTextures.DBufferATexArray;
+			Parameters.DBufferBTextureArray = DBufferTextures.DBufferBTexArray;
+			Parameters.DBufferCTextureArray = DBufferTextures.DBufferCTexArray;
+		}
+		else
+		{
+			Parameters.DBufferATexture = DBufferTextures.DBufferA;
+			Parameters.DBufferBTexture = DBufferTextures.DBufferB;
+			Parameters.DBufferCTexture = DBufferTextures.DBufferC;
+		}
 
 		if (DBufferTextures.DBufferMask)
 		{
diff --git a/Engine/Source/Runtime/Renderer/Private/DBufferTextures.h b/Engine/Source/Runtime/Renderer/Private/DBufferTextures.h
--- a/Engine/Source/Runtime/Renderer/Private/DBufferTextures.h
+++ b/Engine/Source/Runtime/Renderer/Private/DBufferTextures.h
@@ -60,4 +60,4 @@ BEGIN_SHADER_PARAMETER_STRUCT(FDBufferParameters, )
 	SHADER_PARAMETER_SAMPLER(SamplerState, DBufferCTextureSampler)
 END_SHADER_PARAMETER_STRUCT()
 
-FDBufferParameters GetDBufferParameters(FRDGBuilder& GraphBuilder, const FDBufferTextures& DBufferTextures, EShaderPlatform ShaderPlatform);
+FDBufferParameters GetDBufferParameters(FRDGBuilder& GraphBuilder, const FDBufferTextures& DBufferTextures, EShaderPlatform ShaderPlatform, const bool bIsMobileMultiView = false);
diff --git a/Engine/Source/Runtime/Renderer/Private/MobileBasePassRendering.cpp b/Engine/Source/Runtime/Renderer/Private/MobileBasePassRendering.cpp
--- a/Engine/Source/Runtime/Renderer/Private/MobileBasePassRendering.cpp
+++ b/Engine/Source/Runtime/Renderer/Private/MobileBasePassRendering.cpp
@@ -319,7 +319,8 @@ void SetupMobileBasePassUniformParameters(
 	BasePassParameters.AmbientOcclusionSampler = TStaticSamplerState<SF_Bilinear, AM_Clamp, AM_Clamp, AM_Clamp>::GetRHI();
 	BasePassParameters.AmbientOcclusionStaticFraction = FMath::Clamp(View.FinalPostProcessSettings.AmbientOcclusionStaticFraction, 0.0f, 1.0f);
 	
-	BasePassParameters.DBuffer = GetDBufferParameters(GraphBuilder, MobileBasePassTextures.DBufferTextures, View.GetShaderPlatform());
+	const bool bIsMobileMultiViewEnabled = View.bIsMobileMultiViewEnabled || UE::StereoRenderUtils::FStereoShaderAspects(View.GetShaderPlatform()).IsMobileMultiViewEnabled();
+	BasePassParameters.DBuffer = GetDBufferParameters(GraphBuilder, MobileBasePassTextures.DBufferTextures, View.GetShaderPlatform(), bIsMobileMultiViewEnabled);
 
 	const bool bMobileUsesShadowMaskTexture = MobileUsesShadowMaskTexture(View.GetShaderPlatform());
 	
