Light shaft bloom fix when rendered resolution is non 100%.
Fix wrong assert when rendering light shaft into translucency
Fix flickering bloom when toggliing visibility of element in editor, in general: when TSR is not used.

Now playing with r.LightShaftRenderToSeparateTranslucency and r.SeparateTranslucencyScreenPercentage always work and does not ficker.

// preflight 665a125b2256bb6f21aa249c - some unrealted errors
#rnx

diff --git a/Engine/Source/Runtime/Renderer/Private/DeferredShadingRenderer.cpp b/Engine/Source/Runtime/Renderer/Private/DeferredShadingRenderer.cpp
--- a/Engine/Source/Runtime/Renderer/Private/DeferredShadingRenderer.cpp
+++ b/Engine/Source/Runtime/Renderer/Private/DeferredShadingRenderer.cpp
@@ -3012,7 +3012,8 @@ void FDeferredShadingSceneRenderer::Render(FRDGBuilder& GraphBuilder)
 				SCOPED_NAMED_EVENT(RenderTranslucency, FColor::Emerald);
 				SCOPE_CYCLE_COUNTER(STAT_TranslucencyDrawTime);
 				const bool bStandardTranslucentCanRenderSeparate = false;
-				RenderTranslucency(GraphBuilder, SceneTextures, TranslucencyLightingVolumeTextures, &TranslucencyResourceMap, ETranslucencyView::UnderWater, InstanceCullingManager, bStandardTranslucentCanRenderSeparate);
+				FRDGTextureMSAA SharedDepthTexture;
+				RenderTranslucency(GraphBuilder, SceneTextures, TranslucencyLightingVolumeTextures, &TranslucencyResourceMap, ETranslucencyView::UnderWater, InstanceCullingManager, bStandardTranslucentCanRenderSeparate, SharedDepthTexture);
 				EnumRemoveFlags(TranslucencyViewsToRender, ETranslucencyView::UnderWater);
 			}
 
@@ -3081,6 +3082,7 @@ void FDeferredShadingSceneRenderer::Render(FRDGBuilder& GraphBuilder)
 		}
 
 		// Draw translucency.
+		FRDGTextureMSAA TranslucencySharedDepthTexture;
 		if (!bHasRayTracedOverlay && TranslucencyViewsToRender != ETranslucencyView::None)
 		{
 			RDG_CSV_STAT_EXCLUSIVE_SCOPE(GraphBuilder, RenderTranslucency);
@@ -3122,8 +3124,7 @@ void FDeferredShadingSceneRenderer::Render(FRDGBuilder& GraphBuilder)
 			{
 				// Render all remaining translucency views.
 				const bool bStandardTranslucentCanRenderSeparate = bShouldRenderDistortion; // It is only needed to render standard translucent as separate when there is distortion (non self distortion of transmittance/specular/etc.)
-				RenderTranslucency(GraphBuilder, SceneTextures, TranslucencyLightingVolumeTextures, &TranslucencyResourceMap, TranslucencyViewsToRender, InstanceCullingManager, bStandardTranslucentCanRenderSeparate);
-				TranslucencyViewsToRender = ETranslucencyView::None;
+				RenderTranslucency(GraphBuilder, SceneTextures, TranslucencyLightingVolumeTextures, &TranslucencyResourceMap, TranslucencyViewsToRender, InstanceCullingManager, bStandardTranslucentCanRenderSeparate, TranslucencySharedDepthTexture);
 			}
 
 			// Compose hair before velocity/distortion pass since these pass write depth value, 
@@ -3204,6 +3205,13 @@ void FDeferredShadingSceneRenderer::Render(FRDGBuilder& GraphBuilder)
 			RenderLightShaftBloom(GraphBuilder, SceneTextures, /* inout */ TranslucencyResourceMap);
 		}
 
+		{
+			// Light shaft (rendered just above) can render in separate transluceny at low resolution according to r.SeparateTranslucencyScreenPercentage. 
+			// So we can only upsample that buffer if required after the light shaft bloom pass.
+			UpscaleTranslucencyIfNeeded(GraphBuilder, SceneTextures, TranslucencyViewsToRender, /* inout */ &TranslucencyResourceMap, TranslucencySharedDepthTexture);
+			TranslucencyViewsToRender = ETranslucencyView::None;
+		}
+
 		FPathTracingResources PathTracingResources;
 
 #if RHI_RAYTRACING
diff --git a/Engine/Source/Runtime/Renderer/Private/DeferredShadingRenderer.h b/Engine/Source/Runtime/Renderer/Private/DeferredShadingRenderer.h
--- a/Engine/Source/Runtime/Renderer/Private/DeferredShadingRenderer.h
+++ b/Engine/Source/Runtime/Renderer/Private/DeferredShadingRenderer.h
@@ -859,7 +859,8 @@ private:
 		FTranslucencyPassResourcesMap* OutTranslucencyResourceMap,
 		ETranslucencyView ViewsToRender,
 		FInstanceCullingManager& InstanceCullingManager,
-		bool bStandardTranslucentCanRenderSeparate);
+		bool bStandardTranslucentCanRenderSeparate,
+		FRDGTextureMSAA& OutSharedDepthTexture);
 
 	/** Renders the scene's translucency given a specific pass. */
 	void RenderTranslucencyInner(
@@ -874,6 +875,13 @@ private:
 		FInstanceCullingManager& InstanceCullingManager,
 		bool bStandardTranslucentCanRenderSeparate);
 
+	void UpscaleTranslucencyIfNeeded(
+		FRDGBuilder& GraphBuilder,
+		const FSceneTextures& SceneTextures,
+		ETranslucencyView ViewsToRender,
+		FTranslucencyPassResourcesMap* OutTranslucencyResourceMap,
+		FRDGTextureMSAA& InSharedDepthTexture);
+
 	/** Renders the scene's light shafts */
 	FRDGTextureRef RenderLightShaftOcclusion(
 		FRDGBuilder& GraphBuilder,
diff --git a/Engine/Source/Runtime/Renderer/Private/LightShaftRendering.cpp b/Engine/Source/Runtime/Renderer/Private/LightShaftRendering.cpp
--- a/Engine/Source/Runtime/Renderer/Private/LightShaftRendering.cpp
+++ b/Engine/Source/Runtime/Renderer/Private/LightShaftRendering.cpp
@@ -646,7 +646,14 @@ void FDeferredShadingSceneRenderer::RenderLightShaftBloom(
 							}
 							else if(TranslucencyPassResources.IsValid())
 							{
-								ensure(TranslucencyPassResources.ViewRect == OutputViewport.Rect);
+								if (BloomOutput == ELightShaftBloomOutput::SeparateTranslucency)
+								{
+									ensure(TranslucencyPassResources.ViewRect == OutputViewport.Rect);
+								}
+								else
+								{
+									ensure(View.ViewRect == OutputViewport.Rect);
+								}
 							}
 						}
 					}
diff --git a/Engine/Source/Runtime/Renderer/Private/TranslucentRendering.cpp b/Engine/Source/Runtime/Renderer/Private/TranslucentRendering.cpp
--- a/Engine/Source/Runtime/Renderer/Private/TranslucentRendering.cpp
+++ b/Engine/Source/Runtime/Renderer/Private/TranslucentRendering.cpp
@@ -1490,6 +1490,11 @@ void FDeferredShadingSceneRenderer::RenderTranslucencyInner(
 	}
 }
 
+static bool ShouldRenderTranslucentView(const FViewInfo& View, ETranslucencyView TranslucencyView, ETranslucencyView ViewsToRender)
+{
+	return View.ShouldRenderView() && EnumHasAnyFlags(TranslucencyView, ViewsToRender);
+};
+
 void FDeferredShadingSceneRenderer::RenderTranslucency(
 	FRDGBuilder& GraphBuilder,
 	const FSceneTextures& SceneTextures,
@@ -1497,7 +1502,8 @@ void FDeferredShadingSceneRenderer::RenderTranslucency(
 	FTranslucencyPassResourcesMap* OutTranslucencyResourceMap,
 	ETranslucencyView ViewsToRender,
 	FInstanceCullingManager& InstanceCullingManager,
-	bool bStandardTranslucentCanRenderSeparate)
+	bool bStandardTranslucentCanRenderSeparate,
+	FRDGTextureMSAA& OutSharedDepthTexture)
 {
 	if (!EnumHasAnyFlags(ViewsToRender, ETranslucencyView::UnderWater | ETranslucencyView::AboveWater))
 	{
@@ -1514,13 +1520,7 @@ void FDeferredShadingSceneRenderer::RenderTranslucency(
 		SceneColorCopyTexture = AddCopySceneColorPass(GraphBuilder, Views, SceneTextures.Color);
 	}
 
-	const auto ShouldRenderView = [&](const FViewInfo& View, ETranslucencyView TranslucencyView)
-	{
-		return View.ShouldRenderView() && EnumHasAnyFlags(TranslucencyView, ViewsToRender);
-	};
-
 	// Create a shared depth texture at the correct resolution.
-	FRDGTextureMSAA SharedDepthTexture;
 	const bool bIsScalingTranslucency = SeparateTranslucencyDimensions.Scale != 1.0f;
 	if (bIsScalingTranslucency)
 	{
@@ -1532,7 +1532,7 @@ void FDeferredShadingSceneRenderer::RenderTranslucency(
 			1,
 			SeparateTranslucencyDimensions.NumSamples);
 
-		SharedDepthTexture = CreateTextureMSAA(
+		OutSharedDepthTexture = CreateTextureMSAA(
 			GraphBuilder, Desc,
 			TEXT("Translucency.DepthMS"),
 			TEXT("Translucency.Depth"),
@@ -1544,7 +1544,7 @@ void FDeferredShadingSceneRenderer::RenderTranslucency(
 			FViewInfo& View = Views[ViewIndex];
 			const ETranslucencyView TranslucencyView = GetTranslucencyView(View);
 
-			if (!ShouldRenderView(View, TranslucencyView))
+			if (!ShouldRenderTranslucentView(View, TranslucencyView, ViewsToRender))
 			{
 				continue;
 			}
@@ -1553,35 +1553,49 @@ void FDeferredShadingSceneRenderer::RenderTranslucency(
 			AddDownsampleDepthPass(
 				GraphBuilder, View,
 				FScreenPassTexture(SceneTextures.Depth.Resolve, View.ViewRect),
-				FScreenPassRenderTarget(SharedDepthTexture.Target, SeparateTranslucencyViewport.Rect, ViewIndex == 0 ? ERenderTargetLoadAction::EClear : ERenderTargetLoadAction::ELoad),
+				FScreenPassRenderTarget(OutSharedDepthTexture.Target, SeparateTranslucencyViewport.Rect, ViewIndex == 0 ? ERenderTargetLoadAction::EClear : ERenderTargetLoadAction::ELoad),
 				EDownsampleDepthFilter::Point);
 		}
 	}
 	else
 	{
 		// Uses the existing depth buffer for depth testing the translucency.
-		SharedDepthTexture = SceneTextures.Depth;
+		OutSharedDepthTexture = SceneTextures.Depth;
 	}
 
 	if (ViewFamily.AllowTranslucencyAfterDOF())
 	{
-		RenderTranslucencyInner(GraphBuilder, SceneTextures, TranslucentLightingVolumeTextures, OutTranslucencyResourceMap, SharedDepthTexture, ViewsToRender, SceneColorCopyTexture, ETranslucencyPass::TPT_TranslucencyStandard, InstanceCullingManager, bStandardTranslucentCanRenderSeparate);
+		RenderTranslucencyInner(GraphBuilder, SceneTextures, TranslucentLightingVolumeTextures, OutTranslucencyResourceMap, OutSharedDepthTexture, ViewsToRender, SceneColorCopyTexture, ETranslucencyPass::TPT_TranslucencyStandard, InstanceCullingManager, bStandardTranslucentCanRenderSeparate);
 		if (ViewFamily.AllowStandardTranslucencySeparated() && bStandardTranslucentCanRenderSeparate)
 		{
-			RenderTranslucencyInner(GraphBuilder, SceneTextures, TranslucentLightingVolumeTextures, OutTranslucencyResourceMap, SharedDepthTexture, ViewsToRender, SceneColorCopyTexture, ETranslucencyPass::TPT_TranslucencyStandardModulate, InstanceCullingManager, bStandardTranslucentCanRenderSeparate);
+			RenderTranslucencyInner(GraphBuilder, SceneTextures, TranslucentLightingVolumeTextures, OutTranslucencyResourceMap, OutSharedDepthTexture, ViewsToRender, SceneColorCopyTexture, ETranslucencyPass::TPT_TranslucencyStandardModulate, InstanceCullingManager, bStandardTranslucentCanRenderSeparate);
 		}
 
 		if (GetHairStrandsComposition() == EHairStrandsCompositionType::AfterTranslucentBeforeTranslucentAfterDOF)
 		{
 			RenderHairComposition(GraphBuilder, Views, SceneTextures.Color.Target, SceneTextures.Depth.Target, SceneTextures.Velocity, *OutTranslucencyResourceMap);
 		}
-		RenderTranslucencyInner(GraphBuilder, SceneTextures, TranslucentLightingVolumeTextures, OutTranslucencyResourceMap, SharedDepthTexture, ViewsToRender, SceneColorCopyTexture, ETranslucencyPass::TPT_TranslucencyAfterDOF, InstanceCullingManager, bStandardTranslucentCanRenderSeparate);
-		RenderTranslucencyInner(GraphBuilder, SceneTextures, TranslucentLightingVolumeTextures, OutTranslucencyResourceMap, SharedDepthTexture, ViewsToRender, SceneColorCopyTexture, ETranslucencyPass::TPT_TranslucencyAfterDOFModulate, InstanceCullingManager, bStandardTranslucentCanRenderSeparate);
-		RenderTranslucencyInner(GraphBuilder, SceneTextures, TranslucentLightingVolumeTextures, OutTranslucencyResourceMap, SharedDepthTexture, ViewsToRender, SceneColorCopyTexture, ETranslucencyPass::TPT_TranslucencyAfterMotionBlur, InstanceCullingManager, bStandardTranslucentCanRenderSeparate);
+		RenderTranslucencyInner(GraphBuilder, SceneTextures, TranslucentLightingVolumeTextures, OutTranslucencyResourceMap, OutSharedDepthTexture, ViewsToRender, SceneColorCopyTexture, ETranslucencyPass::TPT_TranslucencyAfterDOF, InstanceCullingManager, bStandardTranslucentCanRenderSeparate);
+		RenderTranslucencyInner(GraphBuilder, SceneTextures, TranslucentLightingVolumeTextures, OutTranslucencyResourceMap, OutSharedDepthTexture, ViewsToRender, SceneColorCopyTexture, ETranslucencyPass::TPT_TranslucencyAfterDOFModulate, InstanceCullingManager, bStandardTranslucentCanRenderSeparate);
+		RenderTranslucencyInner(GraphBuilder, SceneTextures, TranslucentLightingVolumeTextures, OutTranslucencyResourceMap, OutSharedDepthTexture, ViewsToRender, SceneColorCopyTexture, ETranslucencyPass::TPT_TranslucencyAfterMotionBlur, InstanceCullingManager, bStandardTranslucentCanRenderSeparate);
 	}
 	else // Otherwise render translucent primitives in a single bucket.
 	{
-		RenderTranslucencyInner(GraphBuilder, SceneTextures, TranslucentLightingVolumeTextures, OutTranslucencyResourceMap, SharedDepthTexture, ViewsToRender, SceneColorCopyTexture, ETranslucencyPass::TPT_AllTranslucency, InstanceCullingManager, bStandardTranslucentCanRenderSeparate);
+		RenderTranslucencyInner(GraphBuilder, SceneTextures, TranslucentLightingVolumeTextures, OutTranslucencyResourceMap, OutSharedDepthTexture, ViewsToRender, SceneColorCopyTexture, ETranslucencyPass::TPT_AllTranslucency, InstanceCullingManager, bStandardTranslucentCanRenderSeparate);
+	}
+}
+
+
+void FDeferredShadingSceneRenderer::UpscaleTranslucencyIfNeeded(
+	FRDGBuilder& GraphBuilder,
+	const FSceneTextures& SceneTextures,
+	ETranslucencyView ViewsToRender,
+	FTranslucencyPassResourcesMap* OutTranslucencyResourceMap,
+	FRDGTextureMSAA& InSharedDepthTexture)
+{
+	if (!EnumHasAnyFlags(ViewsToRender, ETranslucencyView::UnderWater | ETranslucencyView::AboveWater))
+	{
+		return;
 	}
 
 	bool bUpscalePostDOFTranslucency = true;
@@ -1604,7 +1618,7 @@ void FDeferredShadingSceneRenderer::RenderTranslucency(
 		FViewInfo& View = Views[ViewIndex];
 		const ETranslucencyView TranslucencyView = GetTranslucencyView(View);
 
-		if (!ShouldRenderView(View, TranslucencyView))
+		if (!ShouldRenderTranslucentView(View, TranslucencyView, ViewsToRender))
 		{
 			continue;
 		}
@@ -1612,14 +1626,14 @@ void FDeferredShadingSceneRenderer::RenderTranslucency(
 		// Upscale the responsive AA into original depth buffer.
 		bool bUpscaleResponsiveAA = (
 			IsTemporalAccumulationBasedMethod(View.AntiAliasingMethod) &&
-			SharedDepthTexture.Target != SceneTextures.Depth.Target);
+			InSharedDepthTexture.Target != SceneTextures.Depth.Target);
 		if (bUpscaleResponsiveAA)
 		{
 			const FScreenPassTextureViewport SeparateTranslucencyViewport = SeparateTranslucencyDimensions.GetInstancedStereoViewport(View);
 			AddUpsampleResponsiveAAPass(
 				GraphBuilder,
 				View,
-				FScreenPassTexture(SharedDepthTexture.Target, SeparateTranslucencyViewport.Rect),
+				FScreenPassTexture(InSharedDepthTexture.Target, SeparateTranslucencyViewport.Rect),
 				/* OutputDepthTexture = */ SceneTextures.Depth.Target);
 		}
 
