HWRT: Dont create Lumen ray tracing bindings when running path tracer because those are not released which leads to a memory leak.
#rnx
#tests run path tracer and close the editor

diff --git a/Engine/Source/Runtime/Renderer/Private/DeferredShadingRenderer.cpp b/Engine/Source/Runtime/Renderer/Private/DeferredShadingRenderer.cpp
--- a/Engine/Source/Runtime/Renderer/Private/DeferredShadingRenderer.cpp
+++ b/Engine/Source/Runtime/Renderer/Private/DeferredShadingRenderer.cpp
@@ -694,6 +694,7 @@ bool FDeferredShadingSceneRenderer::SetupRayTracingPipelineStates(FRDGBuilder& G
 	}
 
 	// Add Lumen hardware ray tracing materials
+	if (!bIsPathTracing)
 	{
 		TArray<FRHIRayTracingShader*> LumenHardwareRayTracingRayGenShaders;
 
