Fix assertion in OpenGL RHI when creating HZB

Also removed the code path that does not use gather to sample

diff --git a/Engine/Shaders/Private/HZB.usf b/Engine/Shaders/Private/HZB.usf
--- a/Engine/Shaders/Private/HZB.usf
+++ b/Engine/Shaders/Private/HZB.usf
@@ -19,33 +19,18 @@ Texture2D		ParentTextureMip;
 SamplerState	ParentTextureMipSampler;
 
 
-float4 Gather4(Texture2D Texture, SamplerState TextureSampler, float2 BufferUV)
-#if COMPILER_GLSL || FEATURE_LEVEL < FEATURE_LEVEL_SM5
+float4 Gather4(Texture2D Texture, SamplerState TextureSampler, float2 BufferUV, uint Level)
 {
-	float2 UV[4];
-	
 	// min(..., InputViewportMaxBound) because we don't want to sample outside of the viewport
 	// when the view size has odd dimensions on X/Y axis.
-	UV[0] = min(BufferUV + float2(-0.25f, -0.25f) * InvSize, InputViewportMaxBound);
-	UV[1] = min(BufferUV + float2( 0.25f, -0.25f) * InvSize, InputViewportMaxBound);
-	UV[2] = min(BufferUV + float2(-0.25f,  0.25f) * InvSize, InputViewportMaxBound);
-	UV[3] = min(BufferUV + float2( 0.25f,  0.25f) * InvSize, InputViewportMaxBound);
-
-	float4 Out;
-	Out.x = Texture.SampleLevel(TextureSampler, UV[0], 0 ).r;
-	Out.y = Texture.SampleLevel(TextureSampler, UV[1], 0 ).r;
-	Out.z = Texture.SampleLevel(TextureSampler, UV[2], 0 ).r;
-	Out.w = Texture.SampleLevel(TextureSampler, UV[3], 0 ).r;
-
-	return Out;
-}
-#else
-{
 	float2 UV = min(BufferUV + float2(-0.25f, -0.25f) * InvSize, InputViewportMaxBound - InvSize);
-	return Texture.GatherRed(TextureSampler, UV, 0);
+	return Texture.GatherRed(TextureSampler, UV, Level);
 }
-#endif
 
+float4 Gather4(Texture2D Texture, SamplerState TextureSampler, float2 BufferUV)
+{
+	return Gather4(Texture, TextureSampler, BufferUV, 0);
+}
 
 #if COMPUTESHADER
 
@@ -266,10 +251,12 @@ void HZBBuildCS(
 
 #elif PIXELSHADER
 
+uint SourceMipIndex;
+
 void HZBBuildPS(float4 SvPosition : SV_POSITION, out float4 OutColor : SV_Target0)
 {
 	float2 BufferUV = SvPosition.xy * DispatchThreadIdToBufferUV.xy + DispatchThreadIdToBufferUV.zw;
-	float4 DeviceZ = Gather4(ParentTextureMip, ParentTextureMipSampler, BufferUV);
+	float4 DeviceZ = Gather4(ParentTextureMip, ParentTextureMipSampler, BufferUV, SourceMipIndex);
 
 	float FurthestDeviceZ = min(min(DeviceZ.x, DeviceZ.y), min(DeviceZ.z, DeviceZ.w));
 
diff --git a/Engine/Source/Runtime/Renderer/Private/SceneTextureReductions.cpp b/Engine/Source/Runtime/Renderer/Private/SceneTextureReductions.cpp
--- a/Engine/Source/Runtime/Renderer/Private/SceneTextureReductions.cpp
+++ b/Engine/Source/Runtime/Renderer/Private/SceneTextureReductions.cpp
@@ -40,6 +40,7 @@ class FHZBBuildPS : public FGlobalShader
 
 	BEGIN_SHADER_PARAMETER_STRUCT( FParameters, )
 		SHADER_PARAMETER_STRUCT_INCLUDE(FSharedHZBParameters, Shared)
+		SHADER_PARAMETER(int32, SourceMipIndex)
 		RENDER_TARGET_BINDING_SLOTS()
 	END_SHADER_PARAMETER_STRUCT()
 
@@ -237,6 +238,7 @@ void BuildHZB(
 
 			FHZBBuildPS::FParameters* PassParameters = GraphBuilder.AllocParameters<FHZBBuildPS::FParameters>();
 			PassParameters->Shared = ShaderParameters;
+			PassParameters->SourceMipIndex = GRHISupportsTextureViews ? 0 : StartDestMip - 1;
 			PassParameters->RenderTargets[0] = FRenderTargetBinding(FurthestHZBTexture, ERenderTargetLoadAction::ENoAction, StartDestMip);
 
 			FGlobalShaderMap* GlobalShaderMap = GetGlobalShaderMap(FeatureLevel);
@@ -291,7 +293,7 @@ void BuildHZB(
 		FIntVector4 PixelViewPortMinMax = FIntVector4(0, 0, SrcSize.X - 1, SrcSize.Y - 1);
 		
 		{
-			FRDGTextureSRVRef ParentTextureMip = GraphBuilder.CreateSRV(FRDGTextureSRVDesc::CreateForMipLevel(FurthestHZBTexture, StartDestMip - 1));
+			FRDGTextureSRVRef ParentTextureMip = GraphBuilder.CreateSRV(GRHISupportsTextureViews ? FRDGTextureSRVDesc::CreateForMipLevel(FurthestHZBTexture, StartDestMip - 1) : FRDGTextureSRVDesc::Create(FurthestHZBTexture));
 			ReduceMips(ParentTextureMip, SrcSize,
 				StartDestMip, DispatchThreadIdToBufferUV, InputViewportMaxBound, PixelViewPortMinMax,
 				/* bOutputClosest = */ false, /* bOutputFurthest = */ true);
@@ -300,7 +302,7 @@ void BuildHZB(
 		if (bReduceClosestDepth)
 		{
 			check(ClosestHZBTexture)
-			FRDGTextureSRVRef ParentTextureMip = GraphBuilder.CreateSRV(FRDGTextureSRVDesc::CreateForMipLevel(ClosestHZBTexture, StartDestMip - 1));
+			FRDGTextureSRVRef ParentTextureMip = GraphBuilder.CreateSRV(GRHISupportsTextureViews ? FRDGTextureSRVDesc::CreateForMipLevel(ClosestHZBTexture, StartDestMip - 1) : FRDGTextureSRVDesc::Create(ClosestHZBTexture));
 			ReduceMips(ParentTextureMip, SrcSize,
 				StartDestMip, DispatchThreadIdToBufferUV, InputViewportMaxBound, PixelViewPortMinMax,
 				/* bOutputClosest = */ true, /* bOutputFurthest = */ false);
