Remove NumLayers parameter from CreateRayTracingSceneWithGeometryInstances(...)
- layers are handled by FRayTracingScene since CL 34072502
- 5.4 code path is still available but marked as deprecated.

diff --git a/Engine/Plugins/Experimental/GPULightmass/Source/GPULightmass/Private/LightmapRenderer.cpp b/Engine/Plugins/Experimental/GPULightmass/Source/GPULightmass/Private/LightmapRenderer.cpp
--- a/Engine/Plugins/Experimental/GPULightmass/Source/GPULightmass/Private/LightmapRenderer.cpp
+++ b/Engine/Plugins/Experimental/GPULightmass/Source/GPULightmass/Private/LightmapRenderer.cpp
@@ -1131,7 +1131,6 @@ bool FSceneRenderState::SetupRayTracingScene(FRDGBuilder& GraphBuilder, FSceneUn
 
 			FRayTracingSceneWithGeometryInstances SceneWithGeometryInstances = CreateRayTracingSceneWithGeometryInstances(
 				RayTracingGeometryInstances,
-				1,
 				ERayTracingAccelerationStructureFlags::FastTrace);
 
 			RayTracingScene = SceneWithGeometryInstances.Scene;
diff --git a/Engine/Source/Runtime/Renderer/Private/RayTracing/RayTracingInstanceBufferUtil.cpp b/Engine/Source/Runtime/Renderer/Private/RayTracing/RayTracingInstanceBufferUtil.cpp
--- a/Engine/Source/Runtime/Renderer/Private/RayTracing/RayTracingInstanceBufferUtil.cpp
+++ b/Engine/Source/Runtime/Renderer/Private/RayTracing/RayTracingInstanceBufferUtil.cpp
@@ -134,11 +134,10 @@ FRayTracingSceneWithGeometryInstances CreateRayTracingSceneWithGeometryInstances
 
 FRayTracingSceneWithGeometryInstances CreateRayTracingSceneWithGeometryInstances(
 	TConstArrayView<FRayTracingGeometryInstance> Instances,
-	uint8 NumLayers,
 	ERayTracingAccelerationStructureFlags BuildFlags)
 {
 	PRAGMA_DISABLE_DEPRECATION_WARNINGS
-	return CreateRayTracingSceneWithGeometryInstances(Instances, NumLayers, 1, 1, 0, BuildFlags);
+	return CreateRayTracingSceneWithGeometryInstances(Instances, 1, 1, 1, 0, BuildFlags);
 	PRAGMA_ENABLE_DEPRECATION_WARNINGS
 }
 
diff --git a/Engine/Source/Runtime/Renderer/Private/RayTracing/RayTracingScene.cpp b/Engine/Source/Runtime/Renderer/Private/RayTracing/RayTracingScene.cpp
--- a/Engine/Source/Runtime/Renderer/Private/RayTracing/RayTracingScene.cpp
+++ b/Engine/Source/Runtime/Renderer/Private/RayTracing/RayTracingScene.cpp
@@ -56,7 +56,7 @@ void FRayTracingScene::BuildInitializationData()
 	{
 		FLayer& Layer = Layers[LayerIndex];
 
-		Layer.InitializationData = CreateRayTracingSceneWithGeometryInstances(Layer.Instances, 1, BuildFlags);
+		Layer.InitializationData = CreateRayTracingSceneWithGeometryInstances(Layer.Instances, BuildFlags);
 	}
 }
 
diff --git a/Engine/Source/Runtime/Renderer/Private/Tests/RayTracingTestbed.cpp b/Engine/Source/Runtime/Renderer/Private/Tests/RayTracingTestbed.cpp
--- a/Engine/Source/Runtime/Renderer/Private/Tests/RayTracingTestbed.cpp
+++ b/Engine/Source/Runtime/Renderer/Private/Tests/RayTracingTestbed.cpp
@@ -182,7 +182,7 @@ bool RunRayTracingTestbed_RenderThread(const FString& Parameters)
 	Instances[0].Transforms = MakeArrayView(&FMatrix::Identity, 1);
 	Instances[0].InstanceContributionToHitGroupIndex = 0;
 
-	FRayTracingSceneWithGeometryInstances RayTracingScene = CreateRayTracingSceneWithGeometryInstances(Instances, 1);
+	FRayTracingSceneWithGeometryInstances RayTracingScene = CreateRayTracingSceneWithGeometryInstances(Instances);
 
 	const FRayTracingSceneInitializer2& SceneInitializer = RayTracingScene.Scene->GetInitializer();
 
diff --git a/Engine/Source/Runtime/Renderer/Public/RayTracingInstanceBufferUtil.h b/Engine/Source/Runtime/Renderer/Public/RayTracingInstanceBufferUtil.h
--- a/Engine/Source/Runtime/Renderer/Public/RayTracingInstanceBufferUtil.h
+++ b/Engine/Source/Runtime/Renderer/Public/RayTracingInstanceBufferUtil.h
@@ -81,7 +81,6 @@ RENDERER_API FRayTracingSceneWithGeometryInstances CreateRayTracingSceneWithGeom
 
 RENDERER_API FRayTracingSceneWithGeometryInstances CreateRayTracingSceneWithGeometryInstances(
 	TConstArrayView<FRayTracingGeometryInstance> Instances,
-	uint8 NumLayers,
 	ERayTracingAccelerationStructureFlags BuildFlags = ERayTracingAccelerationStructureFlags::FastTrace);
 
 // Helper function to fill upload buffers required by BuildRayTracingInstanceBuffer with instance descriptors
