Implements experimental in-renderer lens distortion

Lens distortion is done solely in the TSR's HistoryUpdate pass for performance reason,
adjusting directly the alignment between rendered and history pixels. This
aproach has the advantage of not needing to change the resolution of the history
whenever the upscaling factor happening at the center of projection change due to
animated FOV or lens distortion, avoiding VMEM spikes as well as keeping the
TSR's history resurrection operational.

The lens distortion is implemented through a generic FLensDistortionLUT, that
can be generated conveniently with the panini reprojection FPaniniProjectionConfig
(configured with r.LensDistortion.Panini.D cvar) For this lens distortion to happen
in the history update perfectly, TSR requires both the FLensDistortionLUT's
distorting and the undistorting UV displacement LUTs.

Disclaimer:
- Any passes sampling pre-TSR-scene-texture after TSR doesn't correctly handle the lens distortion
yet (Motion Blur, Editor Selection outline, Editor primitives). This is deferred to
followup CL depending on the result.
- This will make post process material after TSR sampling scene textures even more
complicated, but this kind of effects already introduces aliasing and/or flickering anyway.
- The rendering resolution can be automatically adjusted to have 1 rendering pixel per display resolution
at ScreenPercentage=100 at the center of projection with r.LensDistortion.AffectScreenPercentage=1,
but this is intenionally disabled by default as this can reallocate renderer's internal resources whenever
the lens distortion or FOV is animated.

How to use:
0) Make sure you are in PIE, -game or a game build
1) Make sure the lens distortion support is compiled in TSR's shader on your platform with r.TSR.Support.LensDistortion
2) Make sure you are using TSR with r.AntiAliasingMethod=4
3) Make sure you have lens distortion in TSR enabled with r.TSR.LensDistortion=1 (can be toggled at runtime)
4) Generates FLensDistortionLUT with r.LensDistortion.Panini.D > 0 of your choosing
5) Make sure to disable motion blur with r.MotionBlurQuality=0 as it is not yet implemented

diff --git a/Engine/Shaders/Private/PaniniProjection.usf b/Engine/Shaders/Private/PaniniProjection.usf
new file mode 100644
--- /dev/null
+++ b/Engine/Shaders/Private/PaniniProjection.usf
@@ -0,0 +1,103 @@
+// Copyright Epic Games, Inc. All Rights Reserved.
+
+//------------------------------------------------------- INCLUDES
+
+#include "Common.ush"
+#include "ScreenPass.ush"
+#include "PaniniProjection.ush"
+
+
+//------------------------------------------------------- PARAMETERS
+
+float2 ScreenSpaceToPaniniFactor;
+float2 PaniniToScreenSpaceFactor;
+FScreenTransform DispatchThreadIdToDestViewportUV;
+
+// only for MainVS()
+// .x:HalfFOV .y:tan(HalfFov) .z:CylinderDistortion(0=none, 1:full) .w:Correction to scale Y the same as X is scaled in the center
+float PaniniD;
+float PaniniS;
+float ScreenPosScale;
+
+RWTexture2D<float2> DistortingDisplacementOutput;
+RWTexture2D<float2> UndistortingDisplacementOutput;
+
+
+//------------------------------------------------------- FUNCTION
+
+float2 PaniniInverseProjection(float2 ON, float d, float s)
+{
+	// line D equation through DN: A x + B z + C = 0
+	float A = 1.0 + d;
+	float B = - ON.x;
+	float C = ON.x * d;
+	
+	// find intersection K(x,z) between line DN and unit circle with center O
+	// solve: x^2 + z^2 = 1, z < 0, A x + B z + C = 0
+	// ends up with polynom: a z^2 + b z + x = 0
+	float a = 1.0 + (B * B) / (A * A);
+	float b = 2.0 * (B * C) / (A * A);
+	float c = (C * C) / (A * A) - 1.0;
+
+	float z = (-b - sqrt(b * b - 4.0 * a * c)) / (2.0 * a);
+
+	float CosPhi = -z;
+	float SinPhi = sqrt(1.0 - CosPhi * CosPhi) * sign(ON.x);
+
+	float S = (d + 1.0) / (d + CosPhi);
+
+	float OMx = SinPhi / CosPhi;
+	float PaniniDirectionXZInvLength = rsqrt(1.0 + OMx * OMx);
+
+	float TanTheta = ON.y / (S * lerp(1.0, 1.0 / CosPhi, s));
+	float OMy = TanTheta / PaniniDirectionXZInvLength;
+
+	return float2(OMx, OMy);
+}
+
+// Panini projection in the screen pos
+// @param InScreenPos: the position to project in the screen space
+// @return the panini projected postion in screen space
+float2 PaniniProjectionScreenPos(float2 InScreenPos)
+{
+	const float2 PaniniDirection = InScreenPos * ScreenSpaceToPaniniFactor;
+	const float2 PaniniPosition = PaniniProjection(PaniniDirection, PaniniD, PaniniS);
+
+	// Return the new position back in the screen space frame
+	return PaniniPosition * PaniniToScreenSpaceFactor * ScreenPosScale;
+}
+
+float2 PaniniInverseProjectionScreenPos(float2 InScreenPos)
+{
+	float2 PaniniPosition = InScreenPos * ScreenSpaceToPaniniFactor * rcp(ScreenPosScale);
+	float2 PaniniDirection = PaniniInverseProjection(PaniniPosition, PaniniD, PaniniS);
+	
+	// Return the new position back in the screen space frame
+	return PaniniDirection * PaniniToScreenSpaceFactor;
+}
+
+
+//------------------------------------------------------- ENTRY POINT
+
+[numthreads(8, 8, 1)]
+void MainCS(uint3 DispatchThreadId : SV_DispatchThreadID)
+{
+	float2 DestViewportUV = ApplyScreenTransform(float2(DispatchThreadId.xy), DispatchThreadIdToDestViewportUV);
+	float2 DestScreenPos = ViewportUVToScreenPos(DestViewportUV);
+	
+	BRANCH
+	if (DispatchThreadId.z == 0)
+	{
+		float2 SrcScreenPos = PaniniProjectionScreenPos(DestScreenPos);
+		float2 SrcViewportUV = ScreenPosToViewportUV(SrcScreenPos);
+		
+		DistortingDisplacementOutput[DispatchThreadId.xy] = SrcViewportUV - DestViewportUV;
+	}
+	else
+	{
+		float2 SrcInverseScreenPos = PaniniInverseProjectionScreenPos(DestScreenPos);
+		float2 SrcInverseViewportUV = ScreenPosToViewportUV(SrcInverseScreenPos);
+		
+		UndistortingDisplacementOutput[DispatchThreadId.xy] = SrcInverseViewportUV - DestViewportUV;
+	}
+}
diff --git a/Engine/Shaders/Private/PostProcessUpscale.usf b/Engine/Shaders/Private/PostProcessUpscale.usf
--- a/Engine/Shaders/Private/PostProcessUpscale.usf
+++ b/Engine/Shaders/Private/PostProcessUpscale.usf
@@ -17,15 +17,12 @@
 // only for MainPS
 float UpscaleSoftness;
 
-// only for MainVS()
-// .x:HalfFOV .y:tan(HalfFov) .z:CylinderDistortion(0=none, 1:full) .w:Correction to scale Y the same as X is scaled in the center
-float Panini_D;
-float Panini_S;
-float Panini_ScreenPosScale;
-
 SCREEN_PASS_TEXTURE_VIEWPORT(Input)
 SCREEN_PASS_TEXTURE_VIEWPORT(Output)
 
+Texture2D<float2> DistortingDisplacementTexture;
+SamplerState DistortingDisplacementSampler;
+
 Texture2D SceneColorTexture;
 SamplerState SceneColorSampler;
 
@@ -80,20 +77,6 @@ float Gaussian(float Scale, float2 Offset)
 	return exp2(Scale * dot(Offset, Offset));
 }
 
-// Panini projection in the screen pos
-// @param InScreenPos: the position to project in the screen space
-// @return the panini projected postion in screen space
-float2 PaniniProjectionScreenPos(float2 InScreenPos)
-{
-	const float2 ScreenSpaceToPaniniFactor = GetTanHalfFieldOfView();
-	const float2 PaniniToScreenSpaceFactor = GetCotanHalfFieldOfView();
-	const float2 PaniniDirection = InScreenPos * ScreenSpaceToPaniniFactor;
-	const float2 PaniniPosition = PaniniProjection(PaniniDirection, Panini_D, Panini_S);
-
-	// Return the new position back in the screen space frame
-	return PaniniPosition * PaniniToScreenSpaceFactor * Panini_ScreenPosScale;
-}
-
 // vertex shader entry point
 void MainVS(
 	in uint VertexId : SV_VertexID,
@@ -108,10 +91,8 @@ void MainVS(
 	// still in 0..1 range
 	float4 Position = float4(TexCoord.x, TexCoord.y, 0, 1);
 
-	float2 PlanarScreenPos = TexCoord * 2 - 1;
-
 	// distort pos
-	Position.xy = PaniniProjectionScreenPos(PlanarScreenPos) * 0.5f + 0.5f;
+	Position.xy = TexCoord + DistortingDisplacementTexture.SampleLevel(DistortingDisplacementSampler, TexCoord, 0);
 
 	DrawRectangle(Position, TexCoord, OutPosition, OutTexCoord.xy);
 	OutTexCoord.zw =  OutPosition.xy;
diff --git a/Engine/Shaders/Private/TemporalSuperResolution/TSRUpdateHistory.usf b/Engine/Shaders/Private/TemporalSuperResolution/TSRUpdateHistory.usf
--- a/Engine/Shaders/Private/TemporalSuperResolution/TSRUpdateHistory.usf
+++ b/Engine/Shaders/Private/TemporalSuperResolution/TSRUpdateHistory.usf
@@ -35,6 +35,9 @@
 #define CONFIG_REJECT_TRANSLUCENCY_ON_RESPONSIVE_AA_ONLY 0
 #define CONFIG_HISTORY_RESURRECTION 1
 
+/** Whether lens distortion is supported. */
+#define CONFIG_LENS_DISTORTION TSR_SUPPORT_LENS_DISTORTION
+
 
 //------------------------------------------------------- DEFAULTS
 
@@ -91,9 +94,17 @@ Texture2D<tsr_half4> HistoryRejectionTexture;
 Texture2D<float2> DilatedVelocityTexture;
 Texture2D<tsr_ushort2> AntiAliasingTexture;
 
+#if CONFIG_LENS_DISTORTION
+	Texture2D<float2> PrevDistortingDisplacementTexture;
+	Texture2D<float2> ResurrectedDistortingDisplacementTexture;
+	Texture2D<float2> UndistortingDisplacementTexture;
+#endif
+
 int2   TranslucencyPixelPosMin;
 int2   TranslucencyPixelPosMax;
 
+FScreenTransform HistoryPixelPosToViewportUV;
+FScreenTransform ViewportUVToInputPPCo;
 FScreenTransform HistoryPixelPosToScreenPos;
 FScreenTransform HistoryPixelPosToInputPPCo;
 FScreenTransform HistoryPixelPosToTranslucencyPPCo;
@@ -106,6 +117,9 @@ float InvWeightClampingPixelSpeed;
 float InputToHistoryFactor;
 float ResurrectionFrameIndex;
 float PrevFrameIndex;
+#if CONFIG_LENS_DISTORTION
+	uint bLensDistortion;
+#endif
 uint bGenerateOutputMip1;
 uint bGenerateOutputMip2;
 uint bGenerateOutputMip3;
@@ -447,6 +461,17 @@ tsr_halfCx2 DownsampleSceneColor(tsr_halfCx2 SceneColor, uint LocalGroupThreadIn
 	return SceneColor;
 }
 
+#if CONFIG_LENS_DISTORTION
+float2 ApplyDisplacementTextureOnScreenPos(Texture2D<float2> DisplacementTexture, float2 DestScreenPos)
+{
+	float2 DestViewportUV = ScreenPosToViewportUV(DestScreenPos);
+
+	float2 SrcViewportUV = DestViewportUV + DisplacementTexture.SampleLevel(GlobalBilinearClampedSampler, DestViewportUV, 0);
+			
+	return ViewportUVToScreenPos(SrcViewportUV);
+}
+#endif
+
 
 //------------------------------------------------------- ENTRY POINT
 
@@ -476,11 +501,40 @@ void MainCS(
 	tsr_short2x2 HistoryPixelPos = dpv_add(
 		tsr_short2(GroupId) * tsr_short2(TILE_SIZE, TILE_SIZE),
 		tsr_short2x2(Map8x8Tile2x2LaneDPV(GroupThreadIndex)));
+		
+	float2x2 ScreenPos;
+	float2x2 InputPPCo;
+	#if CONFIG_LENS_DISTORTION
+	BRANCH
+	if (bLensDistortion)
+	{
+		float2x2 DestViewportUV = ApplyScreenTransform(float2x2(HistoryPixelPos), HistoryPixelPosToViewportUV);
+		float2x2 SrcViewportUV;
+		UNROLL_N(DPV_PIXEL_PER_LANE)
+		for (uint PixelId = 0; PixelId < DPV_PIXEL_PER_LANE; PixelId++)
+		{
+			float2 LocalSrcViewportUV = dpv_access_pixel(DestViewportUV, PixelId) + UndistortingDisplacementTexture.SampleLevel(GlobalBilinearClampedSampler, dpv_access_pixel(DestViewportUV, PixelId), 0);
+			
+			SrcViewportUV[0][PixelId] = LocalSrcViewportUV.x;
+			SrcViewportUV[1][PixelId] = LocalSrcViewportUV.y;
+		}
 
-	float2x2 ScreenPos = ApplyScreenTransform(float2x2(HistoryPixelPos), HistoryPixelPosToScreenPos);
+		ScreenPos[0][0] = ViewportUVToScreenPos(dpv_access_pixel(SrcViewportUV, 0)).x;
+		ScreenPos[1][0] = ViewportUVToScreenPos(dpv_access_pixel(SrcViewportUV, 0)).y;
+		ScreenPos[0][1] = ViewportUVToScreenPos(dpv_access_pixel(SrcViewportUV, 1)).x;
+		ScreenPos[1][1] = ViewportUVToScreenPos(dpv_access_pixel(SrcViewportUV, 1)).y;
+		
+		// Pixel coordinate of the center of output pixel O in the input viewport.
+		InputPPCo = ApplyScreenTransform(SrcViewportUV, ViewportUVToInputPPCo);
+	}
+	else
+	#endif
+	{
+		ScreenPos = ApplyScreenTransform(float2x2(HistoryPixelPos), HistoryPixelPosToScreenPos);
 	
-	// Pixel coordinate of the center of output pixel O in the input viewport.
-	float2x2 InputPPCo = ApplyScreenTransform(float2x2(HistoryPixelPos), HistoryPixelPosToInputPPCo);
+		// Pixel coordinate of the center of output pixel O in the input viewport.
+		InputPPCo = ApplyScreenTransform(float2x2(HistoryPixelPos), HistoryPixelPosToInputPPCo);
+	}
 
 	// Pixel coordinate of the center of the nearest input pixel K in the input viewport.
 	float2x2 InputPPCk = floor(InputPPCo) + 0.5;
@@ -496,8 +550,7 @@ void MainCS(
 	tsr_half2 DecreaseValidityMultiplier = tsr_half(0.0).xx;
 	tsr_half2 OutputPixelVelocity = tsr_half(0.0).xx;
 	tsr_half2 NoiseFiltering = tsr_half(0.0).xx;
-	bool2 bIsOffScreen = false;
-	bool2 bIsDisoccluded = false;
+	bool2 bIsParallaxRejected = false;
 	#if CONFIG_HISTORY_RESURRECTION
 		bool2 bIsHistoryResurrection = false;
 	#else
@@ -612,19 +665,14 @@ void MainCS(
 				OutputPixelVelocity = tsr_half2(dpv_length(dpv_mul(ScreenVelocity, HistoryInfo_ViewportSize)));
 			}
 			
-			tsr_ushort2 BitMask = tsr_ushort2(round(dpv_interleave_registers_array(RawHistoryRejection)[3] * tsr_half(255.0)));
-
-			// Detect whether the history reprojection is successful
-			IsOffScreenOrDisoccluded(
-				bCameraCut,
-				PrevScreenPos,
-				(BitMask & tsr_ushort(0x1)) == tsr_ushort(0),
-				/* out */ bIsOffScreen,
-				/* out */ bIsDisoccluded);
-				
-			#if CONFIG_HISTORY_RESURRECTION
-				bIsHistoryResurrection = (BitMask & tsr_ushort(0x2)) != tsr_ushort(0);
-			#endif
+			// Unpacks bits masks
+			{
+				tsr_ushort2 BitMask = tsr_ushort2(round(dpv_interleave_registers_array(RawHistoryRejection)[3] * tsr_half(255.0)));
+				bIsParallaxRejected = (BitMask & tsr_ushort(0x1)) == tsr_ushort(0);
+				#if CONFIG_HISTORY_RESURRECTION
+					bIsHistoryResurrection = (BitMask & tsr_ushort(0x2)) != tsr_ushort(0);
+				#endif
+			}
 
 			// Process input texel, forcing to pack register.
 			LowFrequencyRejection      = dpv_force_interleave_registers_array(RawHistoryRejection)[0];
@@ -641,22 +689,57 @@ void MainCS(
 			SpatialAntiAliasingOffset = DecodeSpatialAntiAliasingOffset(EncodedInputTexelOffset);
 		}
 		#endif
-		
-		// Final post processing.
+	}
+
+	// Apply lens distortion on the reprojection.
+	#if CONFIG_LENS_DISTORTION
+	BRANCH
+	if (bLensDistortion)
+	{
+		float2x2 UndistortedScreenPos = ApplyScreenTransform(float2x2(HistoryPixelPos), HistoryPixelPosToScreenPos);
+
+		UNROLL_N(DPV_PIXEL_PER_LANE)
+		for (uint PixelId = 0; PixelId < DPV_PIXEL_PER_LANE; PixelId++)
 		{
-			#if !CONFIG_CLAMP
-			{
-				bIsDisoccluded = false;
-				#if CONFIG_HISTORY_RESURRECTION
-					bIsHistoryResurrection = false;
-				#endif
-				LowFrequencyRejection = tsr_half(1.0).xx;
-				LowFrequencyClamp = tsr_half(1.0).xx;
-			}
+			float2 LocalUndistortedScreenPos = dpv_access_pixel(UndistortedScreenPos, PixelId);
+			float2 LocalPrevScreenPos = dpv_access_pixel(PrevScreenPos, PixelId);
+			float2 LocalSrcScreenPos = ApplyDisplacementTextureOnScreenPos(PrevDistortingDisplacementTexture, LocalPrevScreenPos);
+
+			//LocalSrcScreenPos = lerp(LocalUndistortedScreenPos, LocalSrcScreenPos, saturate(OutputPixelVelocity[PixelId] * tsr_half(8.0)));
+			
+			#if CONFIG_HISTORY_RESURRECTION
+				float2 LocalSrcResurrectedScreenPos = ApplyDisplacementTextureOnScreenPos(ResurrectedDistortingDisplacementTexture, LocalPrevScreenPos);
+				LocalSrcScreenPos = select(bIsHistoryResurrection[PixelId], LocalSrcResurrectedScreenPos, LocalSrcScreenPos);
 			#endif
+
+			PrevScreenPos[0][PixelId] = LocalSrcScreenPos.x;
+			PrevScreenPos[1][PixelId] = LocalSrcScreenPos.y;
 		}
 	}
+	#endif
 	
+	// Detect whether the history reprojection is successful
+	bool2 bIsOffScreen;
+	bool2 bIsDisoccluded;
+	IsOffScreenOrDisoccluded(
+		bCameraCut,
+		PrevScreenPos,
+		bIsParallaxRejected,
+		/* out */ bIsOffScreen,
+		/* out */ bIsDisoccluded);
+				
+	// Final post processing.
+	#if !CONFIG_CLAMP
+	{
+		bIsDisoccluded = false;
+		#if CONFIG_HISTORY_RESURRECTION
+			bIsHistoryResurrection = false;
+		#endif
+		LowFrequencyRejection = tsr_half(1.0).xx;
+		LowFrequencyClamp = tsr_half(1.0).xx;
+	}
+	#endif
+
 	// Reproject history
 	tsr_halfCx2 PrevHighFrequencyColor;
 	tsr_half2 PrevHistoryValidity;
diff --git a/Engine/Shaders/Private/TemporalSuperResolution/TSRVisualize.usf b/Engine/Shaders/Private/TemporalSuperResolution/TSRVisualize.usf
--- a/Engine/Shaders/Private/TemporalSuperResolution/TSRVisualize.usf
+++ b/Engine/Shaders/Private/TemporalSuperResolution/TSRVisualize.usf
@@ -39,6 +39,7 @@ float OutputToHistoryResolutionFractionSquare;
 float FlickeringFramePeriod;
 float PerceptionAdd;
 
+Texture2D<float2>    UndistortingDisplacementTexture;
 Texture2D<tsr_half3> InputTexture;
 Texture2D<tsr_half>  InputMoireLumaTexture;
 Texture2D<tsr_half4> InputSceneTranslucencyTexture;
@@ -55,7 +56,22 @@ Texture2D<tsr_half3> ResurrectedHistoryColorTexture;
 RWTexture2D<tsr_half4> Output;
 
 
-//------------------------------------------------------- ENTRY POINT
+//------------------------------------------------------- FUNCTION
+
+float2 ApplyLensDistortionOnScreenPos(float2 DestScreenPos)
+#if TSR_SUPPORT_LENS_DISTORTION
+{
+	float2 DestViewportUV = ScreenPosToViewportUV(DestScreenPos);
+
+	float2 SrcViewportUV = DestViewportUV + UndistortingDisplacementTexture.SampleLevel(GlobalBilinearClampedSampler, DestViewportUV, 0);
+			
+	return ViewportUVToScreenPos(SrcViewportUV);
+}
+#else
+{
+	return DestScreenPos;
+}
+#endif
 
 tsr_half3 GrayScaleColor(tsr_half3 SceneColor)
 {
@@ -82,6 +98,8 @@ tsr_half3 SampleTSROutputGray(FContext Ctx)
 	return GrayScaleColor(SceneColorTexture[Ctx.OutputPixelPos]);
 }
 
+//------------------------------------------------------- ENTRY POINT
+
 tsr_half3 DisplayHistorySampleCount(FContext Ctx)
 {
 	float2 HistoryUV =  ApplyScreenTransform(Ctx.ScreenPos, ScreenPosToHistoryUV);
@@ -214,7 +232,7 @@ tsr_half3 DisplayAntiFlickering(FContext Ctx)
 	const tsr_half TotalVariationCountThreshold = tsr_half(1.0 / (1.0 - pow(1.0 - BlendFinal, tsr_half(FlickeringFramePeriod))));
 
 	// Texture coordinate
-	float2 MoireHistoryUV = ApplyScreenTransform(Ctx.ScreenPos, ScreenPosToMoireHistoryUV);
+	float2 MoireHistoryUV = ApplyScreenTransform(ApplyLensDistortionOnScreenPos(Ctx.ScreenPos), ScreenPosToMoireHistoryUV);
 	MoireHistoryUV = fastClamp(MoireHistoryUV, MoireHistoryUVBilinearMin, MoireHistoryUVBilinearMax);
 
 	// Texture fetches
@@ -286,9 +304,9 @@ void MainCS(
 		tsr_short2(GroupId) * tsr_short2(TILE_SIZE, TILE_SIZE) +
 		Map8x8Tile2x2Lane(GroupThreadIndex));
 	Ctx.ScreenPos = ApplyScreenTransform(float2(Ctx.OutputPixelPos), OutputPixelPosToScreenPos);
-	Ctx.InputUV = ApplyScreenTransform(Ctx.ScreenPos, ScreenPosToInputUV);
+	Ctx.InputUV = ApplyScreenTransform(ApplyLensDistortionOnScreenPos(Ctx.ScreenPos), ScreenPosToInputUV);
 	Ctx.InputUV = fastClamp(Ctx.InputUV, InputInfo_UVViewportBilinearMin, InputInfo_UVViewportBilinearMax);
-	Ctx.InputPixelPos = tsr_short2(ApplyScreenTransform(Ctx.ScreenPos, ScreenPosToInputPixelPos));
+	Ctx.InputPixelPos = tsr_short2(ApplyScreenTransform(ApplyLensDistortionOnScreenPos(Ctx.ScreenPos), ScreenPosToInputPixelPos));
 	Ctx.InputPixelPos = fastClamp(Ctx.InputPixelPos, tsr_short2(InputInfo_ViewportMin), tsr_short2(InputInfo_ViewportMax - 1));
 
 	tsr_half3 OutputColor = tsr_half(0.0);
diff --git a/Engine/Source/Runtime/Renderer/Private/PostProcess/LensDistortion.cpp b/Engine/Source/Runtime/Renderer/Private/PostProcess/LensDistortion.cpp
new file mode 100644
--- /dev/null
+++ b/Engine/Source/Runtime/Renderer/Private/PostProcess/LensDistortion.cpp
@@ -0,0 +1,218 @@
+// Copyright Epic Games, Inc. All Rights Reserved.
+
+#include "PostProcess/LensDistortion.h"
+#include "PostProcess/SceneFilterRendering.h"
+#include "DataDrivenShaderPlatformInfo.h"
+#include "SceneRendering.h"
+
+namespace
+{
+	
+TAutoConsoleVariable<float> CVarLensDistortionLUTResolutionDivisor(
+	TEXT("r.LensDistortion.LUTScreenPercentage"),
+	100.0f * 256.0f / 3840.0f,
+	TEXT("Screen percentage of the procedurally generated LUTs.\n"),
+	ECVF_RenderThreadSafe);
+
+TAutoConsoleVariable<float> CVarLensDistortionPaniniD(
+	TEXT("r.LensDistortion.Panini.D"),
+	0.0f,
+	TEXT("Allow and configure to apply a panini distortion to the rendered image. Values between 0 and 1 allow to fade the effect (lerp).\n")
+	TEXT("Implementation from research paper \"Pannini: A New Projection for Rendering Wide Angle Perspective Images\"\n")
+	TEXT(" 0: off (default)\n")
+	TEXT(">0: enabled (requires an extra post processing pass if upsampling wasn't used - see r.ScreenPercentage)\n")
+	TEXT(" 1: Panini cylindrical stereographic projection"),
+	ECVF_RenderThreadSafe);
+
+TAutoConsoleVariable<float> CVarLensDistortionPaniniS(
+	TEXT("r.LensDistortion.Panini.S"),
+	0.0f,
+	TEXT("Panini projection's hard vertical compression factor.\n")
+	TEXT(" 0: no vertical compression factor (default)\n")
+	TEXT(" 1: Hard vertical compression"),
+	ECVF_RenderThreadSafe);
+
+#if !UE_BUILD_SHIPPING && !UE_BUILD_TEST
+
+TAutoConsoleVariable<float> CVarLensDistortionPaniniScreenFit(
+	TEXT("r.LensDistortion.Panini.ScreenFit"),
+	1.0f,
+	TEXT("Panini projection screen fit effect factor (lerp) for debugging purposes.\n")
+	TEXT(" 0: fit vertically\n")
+	TEXT(" 1: fit horizontally (default)"),
+	ECVF_RenderThreadSafe);
+
+#endif
+
+FVector2f PaniniProjection(FVector2f OM, float d, float s)
+{
+	float PaniniDirectionXZInvLength = 1.0f / FMath::Sqrt(1.0f + OM.X * OM.X);
+	float SinPhi = OM.X * PaniniDirectionXZInvLength;
+	float TanTheta = OM.Y * PaniniDirectionXZInvLength;
+	float CosPhi = FMath::Sqrt(1.0f - SinPhi * SinPhi);
+	float S = (d + 1.0f) / (d + CosPhi);
+
+	return S * FVector2f(SinPhi, FMath::Lerp(TanTheta, TanTheta / CosPhi, s));
+}
+
+FVector2f PaniniInverseProjection(FVector2f ON, float d, float s)
+{
+	// line D equation through DN: A x + B z + C = 0
+	float A = 1.0f + d;
+	float B = -ON.X;
+	float C = ON.X * d;
+
+	// find intersection K(x,z) between line DN and unit circle with center O
+	// solve: x^2 + z^2 = 1, z < 0, A x + B z + C = 0
+	// ends up with polynom: a z^2 + b z + x = 0
+	float a = 1.0f + (B * B) / (A * A);
+	float b = 2.0f * (B * C) / (A * A);
+	float c = (C * C) / (A * A) - 1.0f;
+
+	float z = (-b - FMath::Sqrt(b * b - 4.0f * a * c)) / (2.0f * a);
+
+	float CosPhi = -z;
+	float SinPhi = FMath::Sqrt(1.0f - CosPhi * CosPhi);
+	if (ON.X < 0.0f)
+	{
+		SinPhi = -SinPhi;
+	}
+
+	float S = (d + 1.0f) / (d + CosPhi);
+
+	float OMx = SinPhi / CosPhi;
+	float PaniniDirectionXZInvLength = 1.0f / FMath::Sqrt(1.0f + OMx * OMx);
+
+	float TanTheta = ON.Y / (S * FMath::Lerp(1.0f, 1.0f / CosPhi, s));
+	float OMy = TanTheta / PaniniDirectionXZInvLength;
+
+	return FVector2f(OMx, OMy);
+}
+
+
+} //! namespace
+
+// static
+bool FPaniniProjectionConfig::IsEnabledByCVars()
+{
+	check(IsInRenderingThread());
+	return CVarLensDistortionPaniniD.GetValueOnRenderThread() > 0.01f;
+}
+
+// static
+FPaniniProjectionConfig FPaniniProjectionConfig::ReadCVars()
+{
+	check(FPaniniProjectionConfig::IsEnabledByCVars());
+
+	FPaniniProjectionConfig Config;
+	Config.D = CVarLensDistortionPaniniD.GetValueOnRenderThread();
+	Config.S = CVarLensDistortionPaniniS.GetValueOnRenderThread();
+	Config.Sanitize();
+	return Config;
+}
+
+class FGeneratePaniniUVDisplacementCS : public FGlobalShader
+{
+	DECLARE_GLOBAL_SHADER(FGeneratePaniniUVDisplacementCS);
+	SHADER_USE_PARAMETER_STRUCT(FGeneratePaniniUVDisplacementCS, FGlobalShader);
+
+	BEGIN_SHADER_PARAMETER_STRUCT(FParameters, )
+		SHADER_PARAMETER(FVector2f, ScreenSpaceToPaniniFactor)
+		SHADER_PARAMETER(FVector2f, PaniniToScreenSpaceFactor)
+		SHADER_PARAMETER(FScreenTransform, DispatchThreadIdToDestViewportUV)
+		SHADER_PARAMETER(float, PaniniD)
+		SHADER_PARAMETER(float, PaniniS)
+		SHADER_PARAMETER(float, ScreenPosScale)
+		SHADER_PARAMETER_RDG_TEXTURE_UAV(RWTexture2D, DistortingDisplacementOutput)
+		SHADER_PARAMETER_RDG_TEXTURE_UAV(RWTexture2D, UndistortingDisplacementOutput)
+	END_SHADER_PARAMETER_STRUCT()
+};
+
+IMPLEMENT_GLOBAL_SHADER(FGeneratePaniniUVDisplacementCS, "/Engine/Private/PaniniProjection.usf", "MainCS", SF_Compute);
+
+FLensDistortionLUT FPaniniProjectionConfig::GenerateLUTPasses(FRDGBuilder& GraphBuilder, const FViewInfo& View) const
+{
+	check(IsEnabled());
+	check(View.ViewMatrices.IsPerspectiveProjection());
+
+	const float LUTResolutionFraction = FMath::Clamp(CVarLensDistortionLUTResolutionDivisor.GetValueOnRenderThread() / 100.0f, 0.25f, 100.0f);
+	const FIntPoint SecondaryViewSize = View.GetSecondaryViewRectSize();
+
+	const FIntPoint LUTResolution = FIntPoint(
+		FMath::RoundToInt(LUTResolutionFraction * float(SecondaryViewSize.X)),
+		FMath::RoundToInt(LUTResolutionFraction * float(SecondaryViewSize.Y)));
+
+	const FVector2f FOVPerAxis = FVector2f(View.ViewMatrices.ComputeHalfFieldOfViewPerAxis());
+	const FVector2f ScreenPosToPaniniFactor = FVector2f(FMath::Tan(FOVPerAxis.X), FMath::Tan(FOVPerAxis.Y));
+
+	// Compute the overscan adjustment.
+	float ScreenPosScale;
+	{
+		FVector2f PaniniDirection = FVector2f(1.0f, 0.0f) * ScreenPosToPaniniFactor;
+		FVector2f PaniniPosition = PaniniProjection(PaniniDirection, D, S);
+
+		float WidthFit = ScreenPosToPaniniFactor.X / PaniniPosition.X;
+
+#if !UE_BUILD_SHIPPING && !UE_BUILD_TEST
+		float ScreenFit = CVarLensDistortionPaniniScreenFit.GetValueOnRenderThread();
+		ScreenPosScale = FMath::Lerp(1.0f, WidthFit, ScreenFit);
+#else
+		ScreenPosScale = WidthFit;
+#endif
+	}
+
+	// Compute the resolution fraction happening at the center of distortion that end upscaling.
+	float ResolutionFraction;
+	{
+		const float PrecisionMultiplier = 10.0f;
+		FVector2f UndistortedScreenPos(PrecisionMultiplier / float(SecondaryViewSize.X), PrecisionMultiplier / float(SecondaryViewSize.Y));
+
+		FVector2f PaniniPosition = UndistortedScreenPos * ScreenPosToPaniniFactor * (1.0f / ScreenPosScale);
+		FVector2f PaniniDirection = PaniniInverseProjection(PaniniPosition, D, S);
+		FVector2f DistortedScreenPos = PaniniDirection / ScreenPosToPaniniFactor;
+
+		FVector2f ResolutionFractionVector = UndistortedScreenPos / DistortedScreenPos;
+		ResolutionFraction = FMath::Max(ResolutionFractionVector.X, ResolutionFractionVector.Y);
+
+		check(FMath::IsFinite(ResolutionFraction));
+		check(ResolutionFraction > 1.0f);
+		check(ResolutionFraction < 2.0f);
+	}
+
+	FLensDistortionLUT LensDistortionLUT;
+	LensDistortionLUT.ResolutionFraction = ResolutionFraction;
+
+	FGeneratePaniniUVDisplacementCS::FParameters* PassParameters = GraphBuilder.AllocParameters<FGeneratePaniniUVDisplacementCS::FParameters>();
+	PassParameters->ScreenSpaceToPaniniFactor = ScreenPosToPaniniFactor;
+	PassParameters->PaniniToScreenSpaceFactor = FVector2f(1.0f, 1.0f) / ScreenPosToPaniniFactor;
+	PassParameters->DispatchThreadIdToDestViewportUV = FScreenTransform::DispatchThreadIdToViewportUV(FIntRect(FIntPoint::ZeroValue, LUTResolution));
+	{
+		PassParameters->PaniniD = D;
+		PassParameters->PaniniS = S;
+		PassParameters->ScreenPosScale = ScreenPosScale;
+	}
+
+	{
+		FRDGTextureDesc Desc = FRDGTextureDesc::Create2D(
+			LUTResolution,
+			PF_G32R32F,
+			FClearValueBinding::None,
+			TexCreate_ShaderResource | TexCreate_UAV);
+
+		LensDistortionLUT.DistortingDisplacementTexture   = GraphBuilder.CreateTexture(Desc, TEXT("Panini.DistortingDisplacement"));
+		LensDistortionLUT.UndistortingDisplacementTexture = GraphBuilder.CreateTexture(Desc, TEXT("Panini.UndistortingDisplacement"));
+
+		PassParameters->DistortingDisplacementOutput   = GraphBuilder.CreateUAV(LensDistortionLUT.DistortingDisplacementTexture);
+		PassParameters->UndistortingDisplacementOutput = GraphBuilder.CreateUAV(LensDistortionLUT.UndistortingDisplacementTexture);
+	}
+
+	TShaderMapRef<FGeneratePaniniUVDisplacementCS> ComputeShader(View.ShaderMap);
+	FComputeShaderUtils::AddPass(
+		GraphBuilder,
+		RDG_EVENT_NAME("GeneratePaniniUVDisplacement %dx%d", LUTResolution.X, LUTResolution.Y),
+		ComputeShader,
+		PassParameters,
+		FIntVector(FMath::DivideAndRoundUp(LUTResolution.X, 8), FMath::DivideAndRoundUp(LUTResolution.Y, 8), 2));
+
+	return LensDistortionLUT;
+}
diff --git a/Engine/Source/Runtime/Renderer/Private/PostProcess/LensDistortion.h b/Engine/Source/Runtime/Renderer/Private/PostProcess/LensDistortion.h
new file mode 100644
--- /dev/null
+++ b/Engine/Source/Runtime/Renderer/Private/PostProcess/LensDistortion.h
@@ -0,0 +1,57 @@
+// Copyright Epic Games, Inc. All Rights Reserved.
+
+#pragma once
+
+#include "ScreenPass.h"
+
+struct FLensDistortionLUT
+{
+	/** Expected pixel format of the LUT. 16bit floating point do not work when distorting in TSR. */
+	static constexpr EPixelFormat kFormat = PF_G32R32F;
+
+	/** DistortedViewportUV = UndistortedViewportUV + DistortingDisplacementTexture.Sample(UndistortedViewportUV) */
+	FRDGTextureRef DistortingDisplacementTexture = nullptr;
+
+	/** UndistortedViewportUV = DistortedViewportUV + UndistortingDisplacementTexture.BilinearSample(DistortedViewportUV) */
+	FRDGTextureRef UndistortingDisplacementTexture = nullptr;
+
+	/** Resolution fraction of the upscaling happening due to distortion. */
+	float ResolutionFraction = 1.0f;
+
+	/** Returns whether the displacement is enabled. */
+	bool IsEnabled() const
+	{
+		return DistortingDisplacementTexture && UndistortingDisplacementTexture;
+	}
+};
+
+struct FPaniniProjectionConfig
+{
+	FPaniniProjectionConfig() = default;
+
+	/** Returns whether the panini is enabled by cvar. */
+	static bool IsEnabledByCVars();
+
+	/** Returns teh cvars' configuration. */
+	static FPaniniProjectionConfig ReadCVars();
+
+	bool IsEnabled() const
+	{
+		return D > 0.01f;
+	}
+
+	void Sanitize()
+	{
+		D = FMath::Max(D, 0.0f);
+	}
+
+	// 0=none..1=full, must be >= 0.
+	float D = 0.0f;
+
+	// Panini hard vertical compression lerp (0=no vertical compression, 1=hard compression).
+	float S = 0.0f;
+
+
+	/** Add a RDG pass to generate the lens distortion LUT from the settings. */
+	FLensDistortionLUT GenerateLUTPasses(FRDGBuilder& GraphBuilder, const FViewInfo& View) const;
+};
diff --git a/Engine/Source/Runtime/Renderer/Private/PostProcess/PostProcessUpscale.cpp b/Engine/Source/Runtime/Renderer/Private/PostProcess/PostProcessUpscale.cpp
--- a/Engine/Source/Runtime/Renderer/Private/PostProcess/PostProcessUpscale.cpp
+++ b/Engine/Source/Runtime/Renderer/Private/PostProcess/PostProcessUpscale.cpp
@@ -15,32 +15,6 @@ TAutoConsoleVariable<float> CVarUpscaleSoftness(
 	TEXT("  0: No sharpening (pure Gaussian)."),
 	ECVF_Scalability | ECVF_RenderThreadSafe);
 
-TAutoConsoleVariable<float> CVarUpscalePaniniD(
-	TEXT("r.Upscale.Panini.D"),
-	0,
-	TEXT("Allow and configure to apply a panini distortion to the rendered image. Values between 0 and 1 allow to fade the effect (lerp).\n")
-	TEXT("Implementation from research paper \"Pannini: A New Projection for Rendering Wide Angle Perspective Images\"\n")
-	TEXT(" 0: off (default)\n")
-	TEXT(">0: enabled (requires an extra post processing pass if upsampling wasn't used - see r.ScreenPercentage)\n")
-	TEXT(" 1: Panini cylindrical stereographic projection"),
-	ECVF_RenderThreadSafe);
-
-TAutoConsoleVariable<float> CVarUpscalePaniniS(
-	TEXT("r.Upscale.Panini.S"),
-	0,
-	TEXT("Panini projection's hard vertical compression factor.\n")
-	TEXT(" 0: no vertical compression factor (default)\n")
-	TEXT(" 1: Hard vertical compression"),
-	ECVF_RenderThreadSafe);
-
-TAutoConsoleVariable<float> CVarUpscalePaniniScreenFit(
-	TEXT("r.Upscale.Panini.ScreenFit"),
-	1.0f,
-	TEXT("Panini projection screen fit effect factor (lerp).\n")
-	TEXT(" 0: fit vertically\n")
-	TEXT(" 1: fit horizontally (default)"),
-	ECVF_RenderThreadSafe);
-
 TAutoConsoleVariable<int32> CVarUpscaleQuality(
 	TEXT("r.Upscale.Quality"),
 	3,
@@ -53,65 +27,19 @@ TAutoConsoleVariable<int32> CVarUpscaleQuality(
 	TEXT(" 5: 36-tap Gaussian-filtered unsharp mask (very expensive, but good for extreme upsampling).\n"),
 	ECVF_Scalability | ECVF_RenderThreadSafe);
 
-FVector2f PaniniProjection(FVector2f OM, float d, float s)
-{
-	float PaniniDirectionXZInvLength = 1.0f / FMath::Sqrt(1.0f + OM.X * OM.X);
-	float SinPhi = OM.X * PaniniDirectionXZInvLength;
-	float TanTheta = OM.Y * PaniniDirectionXZInvLength;
-	float CosPhi = FMath::Sqrt(1.0f - SinPhi * SinPhi);
-	float S = (d + 1.0f) / (d + CosPhi);
-
-	return S * FVector2f(SinPhi, FMath::Lerp(TanTheta, TanTheta / CosPhi, s));
-}
 } //! namespace
 
-const FPaniniProjectionConfig FPaniniProjectionConfig::Default;
-
-FPaniniProjectionConfig::FPaniniProjectionConfig(const FViewInfo& View)
-{
-	*this = FPaniniProjectionConfig();
-
-	if (View.IsPerspectiveProjection() && !GEngine->StereoRenderingDevice.IsValid())
-	{
-		D = FMath::Max(CVarUpscalePaniniD.GetValueOnRenderThread(), 0.0f);
-		S = CVarUpscalePaniniS.GetValueOnRenderThread();
-		ScreenFit = FMath::Max(CVarUpscalePaniniScreenFit.GetValueOnRenderThread(), 0.0f);
-	}
-}
-
-BEGIN_SHADER_PARAMETER_STRUCT(FPaniniProjectionParameters, )
-	SHADER_PARAMETER(float, D)
-	SHADER_PARAMETER(float, S)
-	SHADER_PARAMETER(float, ScreenPosScale)
-END_SHADER_PARAMETER_STRUCT()
-
-FPaniniProjectionParameters GetPaniniProjectionParameters(FPaniniProjectionConfig InPaniniConfig, const FViewInfo& View)
-{
-	const FVector2f FOVPerAxis = FVector2f(View.ViewMatrices.ComputeHalfFieldOfViewPerAxis());
-	const FVector2f ScreenPosToPaniniFactor = FVector2f(FMath::Tan(FOVPerAxis.X), FMath::Tan(FOVPerAxis.Y));
-	const FVector2f PaniniDirection = FVector2f(1.0f, 0.0f) * ScreenPosToPaniniFactor;
-	const FVector2f PaniniPosition = PaniniProjection(PaniniDirection, InPaniniConfig.D, InPaniniConfig.S);
-
-	const float WidthFit = ScreenPosToPaniniFactor.X / PaniniPosition.X;
-	const float ScreenPosScale = FMath::Lerp(1.0f, WidthFit, InPaniniConfig.ScreenFit);
-
-	FPaniniProjectionParameters Parameters;
-	Parameters.D = InPaniniConfig.D;
-	Parameters.S = InPaniniConfig.S;
-	Parameters.ScreenPosScale = ScreenPosScale;
-	return Parameters;
-}
-
 BEGIN_SHADER_PARAMETER_STRUCT(FUpscaleParameters, )
 	SHADER_PARAMETER_STRUCT_REF(FViewUniformShaderParameters, View)
 	SHADER_PARAMETER_STRUCT(FScreenPassTextureViewportParameters, Input)
 	SHADER_PARAMETER_STRUCT(FScreenPassTextureViewportParameters, Output)
+	SHADER_PARAMETER_RDG_TEXTURE(Texture2D, DistortingDisplacementTexture)
+	SHADER_PARAMETER_SAMPLER(SamplerState, DistortingDisplacementSampler)
 	SHADER_PARAMETER_RDG_TEXTURE(Texture2D, SceneColorTexture)
 	SHADER_PARAMETER_SAMPLER(SamplerState, SceneColorSampler)
 	SHADER_PARAMETER_RDG_TEXTURE(Texture2D, PointSceneColorTexture)
 	SHADER_PARAMETER_RDG_TEXTURE(Texture2DArray, PointSceneColorTextureArray)
 	SHADER_PARAMETER_SAMPLER(SamplerState, PointSceneColorSampler)
-	SHADER_PARAMETER_STRUCT(FPaniniProjectionParameters, Panini)
 	SHADER_PARAMETER(float, UpscaleSoftness)
 	RENDER_TARGET_BINDING_SLOTS()
 END_SHADER_PARAMETER_STRUCT()
@@ -174,7 +102,7 @@ FScreenPassTexture ISpatialUpscaler::AddDefaultUpscalePass(
 	const FViewInfo& View,
 	const FInputs& Inputs,
 	EUpscaleMethod Method,
-	FPaniniProjectionConfig PaniniConfig)
+	FLensDistortionLUT LensDistortionLUT)
 {
 	check(Inputs.SceneColor.IsValid());
 	check(Method != EUpscaleMethod::MAX);
@@ -211,20 +139,19 @@ FScreenPassTexture ISpatialUpscaler::AddDefaultUpscalePass(
 	const FScreenPassTextureViewport InputViewport(Inputs.SceneColor);
 	const FScreenPassTextureViewport OutputViewport(Output);
 
-	PaniniConfig.Sanitize();
-
-	const bool bUsePaniniProjection = PaniniConfig.IsEnabled();
+	const bool bApplyLensDistortion = LensDistortionLUT.IsEnabled();
 
 	FUpscaleParameters* PassParameters = GraphBuilder.AllocParameters<FUpscaleParameters>();
 	PassParameters->RenderTargets[0] = Output.GetRenderTargetBinding();
 	PassParameters->Input = GetScreenPassTextureViewportParameters(InputViewport);
 	PassParameters->Output = GetScreenPassTextureViewportParameters(OutputViewport);
+	PassParameters->DistortingDisplacementTexture = LensDistortionLUT.DistortingDisplacementTexture;
+	PassParameters->DistortingDisplacementSampler = TStaticSamplerState<SF_Bilinear, AM_Clamp, AM_Clamp, AM_Clamp>::GetRHI();
 	PassParameters->SceneColorTexture = Inputs.SceneColor.Texture;
 	PassParameters->SceneColorSampler = TStaticSamplerState<SF_Bilinear, AM_Border, AM_Border, AM_Border>::GetRHI();
 	PassParameters->PointSceneColorTexture = Inputs.SceneColor.Texture;
 	PassParameters->PointSceneColorTextureArray = Inputs.SceneColor.Texture;
 	PassParameters->PointSceneColorSampler = TStaticSamplerState<SF_Point, AM_Border, AM_Border, AM_Border>::GetRHI();
-	PassParameters->Panini = GetPaniniProjectionParameters(PaniniConfig, View);
 	PassParameters->UpscaleSoftness = FMath::Clamp(CVarUpscaleSoftness.GetValueOnRenderThread(), 0.0f, 1.0f);
 	PassParameters->View = View.ViewUniformBuffer;
 
@@ -237,15 +164,20 @@ FScreenPassTexture ISpatialUpscaler::AddDefaultUpscalePass(
 	const TCHAR* StageName = StageNames[static_cast<uint32>(Inputs.Stage)];
 
 	GraphBuilder.AddPass(
-		RDG_EVENT_NAME("Upscale(%s Method=%d) %dx%d -> %dx%d", StageName, int32(Method), Inputs.SceneColor.ViewRect.Width(), Inputs.SceneColor.ViewRect.Height(), Output.ViewRect.Width(), Output.ViewRect.Height()),
+		RDG_EVENT_NAME("Upscale(%s Method=%d%s) %dx%d -> %dx%d",
+			StageName,
+			int32(Method),
+			bApplyLensDistortion ? TEXT(" LensDistortion") : TEXT(""),
+			Inputs.SceneColor.ViewRect.Width(), Inputs.SceneColor.ViewRect.Height(),
+			Output.ViewRect.Width(), Output.ViewRect.Height()),
 		PassParameters,
 		ERDGPassFlags::Raster,
-		[&View, bUsePaniniProjection, PixelShader, PassParameters, InputViewport, OutputViewport](FRHICommandList& RHICmdList)
+		[&View, bApplyLensDistortion, PixelShader, PassParameters, InputViewport, OutputViewport](FRHICommandList& RHICmdList)
 	{
 		RHICmdList.SetViewport(OutputViewport.Rect.Min.X, OutputViewport.Rect.Min.Y, 0.0f, OutputViewport.Rect.Max.X, OutputViewport.Rect.Max.Y, 1.0f);
 
 		TShaderRef<FShader> VertexShader;
-		if (bUsePaniniProjection)
+		if (bApplyLensDistortion)
 		{
 			TShaderMapRef<FUpscaleVS> TypedVertexShader(View.ShaderMap);
 			SetScreenPassPipelineState(RHICmdList, FScreenPassPipelineState(TypedVertexShader, PixelShader));
@@ -271,7 +203,7 @@ FScreenPassTexture ISpatialUpscaler::AddDefaultUpscalePass(
 			OutputViewport.Rect.Size(),
 			InputViewport.Extent,
 			VertexShader,
-			bUsePaniniProjection ? EDRF_UseTesselatedIndexBuffer : EDRF_UseTriangleOptimization);
+			bApplyLensDistortion ? EDRF_UseTesselatedIndexBuffer : EDRF_UseTriangleOptimization);
 	});
 
 	return MoveTemp(Output);
diff --git a/Engine/Source/Runtime/Renderer/Private/PostProcess/PostProcessUpscale.h b/Engine/Source/Runtime/Renderer/Private/PostProcess/PostProcessUpscale.h
--- a/Engine/Source/Runtime/Renderer/Private/PostProcess/PostProcessUpscale.h
+++ b/Engine/Source/Runtime/Renderer/Private/PostProcess/PostProcessUpscale.h
@@ -3,34 +3,7 @@
 #pragma once
 
 #include "OverridePassSequence.h"
-
-struct FPaniniProjectionConfig
-{
-	static const FPaniniProjectionConfig Default;
-
-	FPaniniProjectionConfig() = default;
-	FPaniniProjectionConfig(const FViewInfo& View);
-
-	bool IsEnabled() const
-	{
-		return D > 0.01f;
-	}
-
-	void Sanitize()
-	{
-		D = FMath::Max(D, 0.0f);
-		ScreenFit = FMath::Max(ScreenFit, 0.0f);
-	}
-
-	// 0=none..1=full, must be >= 0.
-	float D = 0.0f;
-
-	// Panini hard vertical compression lerp (0=no vertical compression, 1=hard compression).
-	float S = 0.0f;
-
-	// Panini screen fit factor (lerp between vertical and horizontal).
-	float ScreenFit = 1.0f;
-};
+#include "PostProcess/LensDistortion.h"
 
 enum class EUpscaleMethod : uint8
 {
@@ -91,5 +64,5 @@ public:
 		const FViewInfo& View,
 		const FInputs& PassInputs,
 		EUpscaleMethod Method,
-		FPaniniProjectionConfig PaniniConfig);
+		FLensDistortionLUT LensDistortionLUT = FLensDistortionLUT());
 };
diff --git a/Engine/Source/Runtime/Renderer/Private/PostProcess/PostProcessing.cpp b/Engine/Source/Runtime/Renderer/Private/PostProcess/PostProcessing.cpp
--- a/Engine/Source/Runtime/Renderer/Private/PostProcess/PostProcessing.cpp
+++ b/Engine/Source/Runtime/Renderer/Private/PostProcess/PostProcessing.cpp
@@ -374,7 +374,11 @@ void AddPostProcessingPasses(
 	const int32 LumenVisualizeMode = GetLumenVisualizeMode(View);
 	const bool bPostProcessingEnabled = IsPostProcessingEnabled(View);
 
-	const FPaniniProjectionConfig PaniniConfig(View);
+	// Temporal Anti-aliasing. Also may perform a temporal upsample from primary to secondary view rect.
+	const EMainTAAPassConfig TAAConfig = GetMainTAAPassConfig(View);
+
+	bool bApplyLensDistortion = View.LensDistortionLUT.IsEnabled();
+	bool bApplyLensDistortionInTSR = TAAConfig == EMainTAAPassConfig::TSR && bPostProcessingEnabled && IsTSRLensDistortionEnabled(View.GetShaderPlatform());
 
 	enum class EPass : uint32
 	{
@@ -526,7 +530,7 @@ void AddPostProcessingPasses(
 #if UE_ENABLE_DEBUG_DRAWING
 	PassSequence.SetEnabled(EPass::DebugPrimitive, FSceneRenderer::ShouldCompositeDebugPrimitivesInPostProcess(View));
 #endif
-	PassSequence.SetEnabled(EPass::PrimaryUpscale, PaniniConfig.IsEnabled() || (View.PrimaryScreenPercentageMethod == EPrimaryScreenPercentageMethod::SpatialUpscale && PrimaryViewRect.Size() != View.GetSecondaryViewRectSize()));
+	PassSequence.SetEnabled(EPass::PrimaryUpscale, (bApplyLensDistortion && !bApplyLensDistortionInTSR) || (View.PrimaryScreenPercentageMethod == EPrimaryScreenPercentageMethod::SpatialUpscale && PrimaryViewRect.Size() != View.GetSecondaryViewRectSize()));
 	PassSequence.SetEnabled(EPass::SecondaryUpscale, View.RequiresSecondaryUpscale() || View.Family->GetSecondarySpatialUpscalerInterface() != nullptr);
 
 	const auto GetPostProcessMaterialInputs = [&](FScreenPassTexture InSceneColor)
@@ -603,9 +607,6 @@ void AddPostProcessingPasses(
 
 	if (bPostProcessingEnabled)
 	{
-		// Temporal Anti-aliasing. Also may perform a temporal upsample from primary to secondary view rect.
-		EMainTAAPassConfig TAAConfig = GetMainTAAPassConfig(View);
-
 		const bool bPrimaryView = IStereoRendering::IsAPrimaryView(View);
 		const bool bHasViewState = View.ViewState != nullptr;
 		const bool bDepthOfFieldEnabled = DiaphragmDOF::IsEnabled(View);
@@ -850,6 +851,10 @@ void AddPostProcessingPasses(
 			UpscalerPassInputs.DownsampleOverrideFormat = DownsampleOverrideFormat;
 			UpscalerPassInputs.PostDOFTranslucencyResources = PostDOFTranslucencyResources;
 			UpscalerPassInputs.FlickeringInputTexture = TSRFlickeringInput;
+			if (bApplyLensDistortionInTSR)
+			{
+				UpscalerPassInputs.LensDistortionLUT = View.LensDistortionLUT;
+			}
 			check(UpscalerPassInputs.SceneColor.ViewRect == View.ViewRect);
 
 			FDefaultTemporalUpscaler::FOutputs Outputs;
@@ -1756,7 +1761,7 @@ void AddPostProcessingPasses(
 		{
 			EUpscaleMethod Method = GetUpscaleMethod();
 
-			SceneColor = ISpatialUpscaler::AddDefaultUpscalePass(GraphBuilder, View, PassInputs, Method, PaniniConfig);
+			SceneColor = ISpatialUpscaler::AddDefaultUpscalePass(GraphBuilder, View, PassInputs, Method, View.LensDistortionLUT);
 		}
 	}
 
@@ -1786,7 +1791,7 @@ void AddPostProcessingPasses(
 				? EUpscaleMethod::SmoothStep
 				: EUpscaleMethod::Nearest;
 
-			SceneColor = ISpatialUpscaler::AddDefaultUpscalePass(GraphBuilder, View, PassInputs, Method, FPaniniProjectionConfig());
+			SceneColor = ISpatialUpscaler::AddDefaultUpscalePass(GraphBuilder, View, PassInputs, Method);
 		}
 	}
 
@@ -2100,7 +2105,7 @@ void AddDebugViewPostProcessingPasses(FRDGBuilder& GraphBuilder, const FViewInfo
 		{
 			EUpscaleMethod Method = GetUpscaleMethod();
 
-			SceneColor = ISpatialUpscaler::AddDefaultUpscalePass(GraphBuilder, View, PassInputs, Method, FPaniniProjectionConfig());
+			SceneColor = ISpatialUpscaler::AddDefaultUpscalePass(GraphBuilder, View, PassInputs, Method);
 		}
 	}
 
@@ -2129,7 +2134,7 @@ void AddDebugViewPostProcessingPasses(FRDGBuilder& GraphBuilder, const FViewInfo
 				? EUpscaleMethod::SmoothStep
 				: EUpscaleMethod::Nearest;
 
-			SceneColor = ISpatialUpscaler::AddDefaultUpscalePass(GraphBuilder, View, PassInputs, Method, FPaniniProjectionConfig());
+			SceneColor = ISpatialUpscaler::AddDefaultUpscalePass(GraphBuilder, View, PassInputs, Method);
 		}
 	}
 }
@@ -2220,8 +2225,6 @@ void AddMobilePostProcessingPasses(FRDGBuilder& GraphBuilder, FScene* Scene, con
 	const FEyeAdaptationParameters EyeAdaptationParameters = GetEyeAdaptationParameters(View);
 	FRDGBufferRef LastEyeAdaptationBuffer = GetEyeAdaptationBuffer(GraphBuilder, View);
 
-	const FPaniniProjectionConfig PaniniConfig(View);
-
 	enum class EPass : uint32
 	{
 		Distortion,
@@ -2310,7 +2313,7 @@ void AddMobilePostProcessingPasses(FRDGBuilder& GraphBuilder, FScene* Scene, con
 
 	bool bUseHighResolutionScreenshotMask = IsHighResolutionScreenshotMaskEnabled(View);
 
-	bool bShouldPrimaryUpscale = (View.PrimaryScreenPercentageMethod == EPrimaryScreenPercentageMethod::SpatialUpscale && View.UnscaledViewRect != View.ViewRect) || PaniniConfig.IsEnabled();
+	bool bShouldPrimaryUpscale = (View.PrimaryScreenPercentageMethod == EPrimaryScreenPercentageMethod::SpatialUpscale && View.UnscaledViewRect != View.ViewRect) || View.LensDistortionLUT.IsEnabled();
 	bShouldPrimaryUpscale |= View.Family->GetPrimarySpatialUpscalerInterface() != nullptr;
 
 	PassSequence.SetEnabled(EPass::Tonemap, bUseToneMapper);
@@ -2947,7 +2950,7 @@ void AddMobilePostProcessingPasses(FRDGBuilder& GraphBuilder, FScene* Scene, con
 		}
 		else
 		{
-			SceneColor = ISpatialUpscaler::AddDefaultUpscalePass(GraphBuilder, View, PassInputs, EUpscaleMethod::Bilinear, PaniniConfig);
+			SceneColor = ISpatialUpscaler::AddDefaultUpscalePass(GraphBuilder, View, PassInputs, EUpscaleMethod::Bilinear, View.LensDistortionLUT);
 		}
 	}
 
diff --git a/Engine/Source/Runtime/Renderer/Private/PostProcess/TemporalAA.h b/Engine/Source/Runtime/Renderer/Private/PostProcess/TemporalAA.h
--- a/Engine/Source/Runtime/Renderer/Private/PostProcess/TemporalAA.h
+++ b/Engine/Source/Runtime/Renderer/Private/PostProcess/TemporalAA.h
@@ -4,6 +4,7 @@
 
 #include "ScreenPass.h"
 #include "PostProcess/PostProcessMotionBlur.h"
+#include "PostProcess/LensDistortion.h"
 #include "TemporalUpscaler.h"
 
 struct FTemporalAAHistory;
@@ -174,6 +175,12 @@ extern RENDERER_API FTAAOutputs AddTemporalAAPass(
 	const FTemporalAAHistory& InputHistory,
 	FTemporalAAHistory* OutputHistory);
 
+/** Returns whether TSR support lens distortion in its shader. */
+bool IsTSRLensDistortionSupported(EShaderPlatform ShaderPlatform);
+
+/** Returns whether TSR lens distortion is enabled (for runtime toggle). */
+bool IsTSRLensDistortionEnabled(EShaderPlatform ShaderPlatform);
+
 /** Returns whether a given view need to measure luminance of the scene color for moire anti-flickering. */
 bool NeedTSRMoireLuma(const FViewInfo& View);
 
@@ -203,6 +210,7 @@ struct FDefaultTemporalUpscaler
 		FScreenPassTexture SceneVelocity;
 		FTranslucencyPassResources PostDOFTranslucencyResources;
 		FScreenPassTexture FlickeringInputTexture;
+		FLensDistortionLUT LensDistortionLUT;
 	};
 
 	struct FOutputs
diff --git a/Engine/Source/Runtime/Renderer/Private/PostProcess/TemporalSuperResolution.cpp b/Engine/Source/Runtime/Renderer/Private/PostProcess/TemporalSuperResolution.cpp
--- a/Engine/Source/Runtime/Renderer/Private/PostProcess/TemporalSuperResolution.cpp
+++ b/Engine/Source/Runtime/Renderer/Private/PostProcess/TemporalSuperResolution.cpp
@@ -20,7 +20,16 @@
 
 namespace
 {
-	
+
+TAutoConsoleVariable<int32> CVarTSRSupportLensDistortion(
+	TEXT("r.TSR.Support.LensDistortion"), 1,
+	TEXT("Whether to compile lens distortion support in TSR's shaders ")
+	TEXT("(adds the lens distortion LUT in the HistoryUpdate pass in branches that even disabled can add a bit of VALU cost when no lens distortion is used).\n")
+	TEXT(" 0: unsupported;\n")
+	TEXT(" 1: supported only on desktop (default);\n")
+	TEXT(" 2: supported everywhere;\n"),
+	ECVF_ReadOnly);
+
 TAutoConsoleVariable<int32> CVarTSRAlphaChannel(
 	TEXT("r.TSR.AplhaChannel"), -1,
 	TEXT("Controls whether TSR should process the scene color's alpha channel.\n")
@@ -206,6 +215,11 @@ TAutoConsoleVariable<float> CVarTSRShadingExposureOffset(
 	TEXT("with the RGB Linear[0;1] source color space against the Tonemaper's output in sRGB source color space.\n"),
 	ECVF_RenderThreadSafe);
 
+TAutoConsoleVariable<int32> CVarTSRLensDistortion(
+	TEXT("r.TSR.LensDistortion"), 1,
+	TEXT("Whether to apply lens distortion in TSR at runtime (enabled by default, requires r.TSR.Support.LensDistortion enabled at cook time)."),
+	ECVF_RenderThreadSafe);
+
 TAutoConsoleVariable<int32> CVarTSRRejectionAntiAliasingQuality(
 	TEXT("r.TSR.RejectionAntiAliasingQuality"), 3,
 	TEXT("Controls the quality of TSR's built-in spatial anti-aliasing technology when the history needs to be rejected. ")
@@ -445,6 +459,8 @@ public:
 		OutEnvironment.CompilerFlags.Add(CFLAG_WarningsAsErrors);
 		OutEnvironment.CompilerFlags.Add(CFLAG_HLSL2021);
 		OutEnvironment.CompilerFlags.Add(CFLAG_ForceOptimization);
+
+		OutEnvironment.SetDefine(TEXT("TSR_SUPPORT_LENS_DISTORTION"), IsTSRLensDistortionSupported(Parameters.Platform) ? 1 : 0);
 	}
 }; // class FTemporalSuperResolutionShader
 
@@ -922,6 +938,8 @@ class FTSRUpdateHistoryCS : public FTSRShader
 		SHADER_PARAMETER(FIntPoint, TranslucencyPixelPosMin)
 		SHADER_PARAMETER(FIntPoint, TranslucencyPixelPosMax)
 
+		SHADER_PARAMETER(FScreenTransform, HistoryPixelPosToViewportUV)
+		SHADER_PARAMETER(FScreenTransform, ViewportUVToInputPPCo)
 		SHADER_PARAMETER(FScreenTransform, HistoryPixelPosToScreenPos)
 		SHADER_PARAMETER(FScreenTransform, HistoryPixelPosToInputPPCo)
 		SHADER_PARAMETER(FScreenTransform, HistoryPixelPosToTranslucencyPPCo)
@@ -936,6 +954,7 @@ class FTSRUpdateHistoryCS : public FTSRShader
 		SHADER_PARAMETER(float, InputContributionMultiplier)
 		SHADER_PARAMETER(float, ResurrectionFrameIndex)
 		SHADER_PARAMETER(float, PrevFrameIndex)
+		SHADER_PARAMETER(int32, bLensDistortion)
 		SHADER_PARAMETER(int32, bGenerateOutputMip1)
 		SHADER_PARAMETER(int32, bGenerateOutputMip2)
 		SHADER_PARAMETER(int32, bGenerateOutputMip3)
@@ -946,6 +965,10 @@ class FTSRUpdateHistoryCS : public FTSRShader
 		SHADER_PARAMETER_RDG_TEXTURE_SRV(Texture2DArray, PrevHistoryColorTexture)
 		SHADER_PARAMETER_RDG_TEXTURE_SRV(Texture2DArray, PrevHistoryMetadataTexture)
 
+		SHADER_PARAMETER_RDG_TEXTURE(Texture2D, PrevDistortingDisplacementTexture)
+		SHADER_PARAMETER_RDG_TEXTURE(Texture2D, ResurrectedDistortingDisplacementTexture)
+		SHADER_PARAMETER_RDG_TEXTURE(Texture2D, UndistortingDisplacementTexture)
+
 		SHADER_PARAMETER_RDG_TEXTURE_UAV(RWTexture2DArray, HistoryColorOutput)
 		SHADER_PARAMETER_RDG_TEXTURE_UAV(RWTexture2DArray, HistoryMetadataOutput)
 		SHADER_PARAMETER_RDG_TEXTURE_UAV(RWTexture2DArray, SceneColorOutputMip1)
@@ -1097,6 +1120,7 @@ class FTSRVisualizeCS : public FTSRShader
 		SHADER_PARAMETER(float, FlickeringFramePeriod)
 		SHADER_PARAMETER(float, PerceptionAdd)
 
+		SHADER_PARAMETER_RDG_TEXTURE(Texture2D, UndistortingDisplacementTexture)
 		SHADER_PARAMETER_RDG_TEXTURE(Texture2D, InputTexture)
 		SHADER_PARAMETER_RDG_TEXTURE(Texture2D, InputMoireLumaTexture)
 		SHADER_PARAMETER_RDG_TEXTURE(Texture2D, InputSceneTranslucencyTexture)
@@ -1275,6 +1299,32 @@ struct FTSRHistorySliceSequence
 	}
 };
 
+bool IsTSRLensDistortionSupported(EShaderPlatform ShaderPlatform)
+{
+	int32 LensDistortionSupport = CVarTSRSupportLensDistortion.GetValueOnAnyThread();
+	if (LensDistortionSupport <= 0)
+	{
+		return false;
+	}
+	else if (LensDistortionSupport == 1)
+	{
+		return FDataDrivenShaderPlatformInfo::GetIsPC(ShaderPlatform);
+	}
+
+	return true;
+}
+
+bool IsTSRLensDistortionEnabled(EShaderPlatform ShaderPlatform)
+{
+	check(IsInRenderingThread());
+	if (!IsTSRLensDistortionSupported(ShaderPlatform))
+	{
+		return false;
+	}
+
+	return CVarTSRLensDistortion.GetValueOnRenderThread() != 0;
+}
+
 bool NeedTSRMoireLuma(const FViewInfo& View)
 {
 	return GetMainTAAPassConfig(View) == EMainTAAPassConfig::TSR;
@@ -1346,8 +1396,13 @@ FDefaultTemporalUpscaler::FOutputs AddTemporalSuperResolutionPasses(
 	}
 	check(HistorySliceSequence.Check());
 		
+	const EShaderPlatform ShaderPlatform = View.GetShaderPlatform();
+
+	// Whether lens distortion support is compiled in the shaders.
+	const bool bSupportsLensDistortion = IsTSRLensDistortionSupported(ShaderPlatform);
+
 	// Whether to use 16bit VALU
-	bool bUse16BitVALU = Use16BitVALU(View.GetShaderPlatform());
+	bool bUse16BitVALU = Use16BitVALU(ShaderPlatform);
 
 	// Whether alpha channel is supported.
 	const bool bSupportsAlpha = CVarTSRAlphaChannel.GetValueOnRenderThread() >= 0 ? (CVarTSRAlphaChannel.GetValueOnRenderThread() > 0) : IsPostProcessingWithAlphaChannelSupported();
@@ -1454,6 +1509,26 @@ FDefaultTemporalUpscaler::FOutputs AddTemporalSuperResolutionPasses(
 		}
 	}
 
+	// Whether to use camera cut.
+	const bool bCameraCut =
+		!InputHistory.IsValid() ||
+		View.bCameraCut ||
+		ETSRHistoryFormatBits(InputHistory.FormatBit) != HistoryFormatBits ||
+		false;
+
+	// Whether to apply lens distortion
+	bool bLensDistortion = false;
+	if (bSupportsLensDistortion)
+	{
+		bLensDistortion = PassInputs.LensDistortionLUT.IsEnabled();
+
+		// Still apply lens distortion if the history has been distorted before to ensure smooth transition from distorted -> undistorted.
+		for (int32 i = 0; i < InputHistory.DistortingDisplacementTextures.Num() && !bLensDistortion; i++)
+		{
+			bLensDistortion = bLensDistortion || InputHistory.DistortingDisplacementTextures[i] != nullptr;
+		}
+	}
+
 	FIntPoint HistoryExtent;
 	FIntPoint HistorySize;
 	{
@@ -1492,13 +1567,6 @@ FDefaultTemporalUpscaler::FOutputs AddTemporalSuperResolutionPasses(
 	float OutputToInputResolutionFraction = float(InputRect.Width()) / float(OutputRect.Width());
 	float OutputToInputResolutionFractionSquare = OutputToInputResolutionFraction * OutputToInputResolutionFraction;
 
-	// Whether to use camera cut shader permutation or not.
-	const bool bCameraCut =
-		!InputHistory.IsValid() ||
-		View.bCameraCut ||
-		ETSRHistoryFormatBits(InputHistory.FormatBit) != HistoryFormatBits ||
-		false;
-
 	static auto CVarAntiAliasingQuality = IConsoleManager::Get().FindConsoleVariable(TEXT("sg.AntiAliasingQuality"));
 	check(CVarAntiAliasingQuality);
 	
@@ -2304,8 +2372,10 @@ FDefaultTemporalUpscaler::FOutputs AddTemporalSuperResolutionPasses(
 		PassParameters->TranslucencyPixelPosMax = SeparateTranslucencyRect.Max - 1;
 
 		FScreenTransform HistoryPixelPosToViewportUV = (FScreenTransform::Identity + 0.5f) * CommonParameters.HistoryInfo.ViewportSizeInverse;
+		PassParameters->HistoryPixelPosToViewportUV = HistoryPixelPosToViewportUV;
+		PassParameters->ViewportUVToInputPPCo = FScreenTransform::Identity * CommonParameters.InputInfo.ViewportSize + CommonParameters.InputJitter + CommonParameters.InputPixelPosMin;
 		PassParameters->HistoryPixelPosToScreenPos = HistoryPixelPosToViewportUV * FScreenTransform::ViewportUVToScreenPos;
-		PassParameters->HistoryPixelPosToInputPPCo = HistoryPixelPosToViewportUV * CommonParameters.InputInfo.ViewportSize + CommonParameters.InputJitter + CommonParameters.InputPixelPosMin;
+		PassParameters->HistoryPixelPosToInputPPCo = HistoryPixelPosToViewportUV * PassParameters->ViewportUVToInputPPCo;
 		PassParameters->HistoryPixelPosToTranslucencyPPCo = HistoryPixelPosToViewportUV * SeparateTranslucencyRect.Size() + CommonParameters.InputJitter * SeparateTranslucencyRect.Size() / CommonParameters.InputInfo.ViewportSize + SeparateTranslucencyRect.Min;
 		PassParameters->HistoryQuantizationError = ComputePixelFormatQuantizationError(HistoryColorFormat);
 
@@ -2318,6 +2388,7 @@ FDefaultTemporalUpscaler::FOutputs AddTemporalSuperResolutionPasses(
 		
 		PassParameters->InputToHistoryFactor = float(HistorySize.X) / float(InputRect.Width());
 		PassParameters->InputContributionMultiplier = OutputToHistoryResolutionFractionSquare; 
+		PassParameters->bLensDistortion = bLensDistortion;
 		PassParameters->bGenerateOutputMip1 = false;
 		PassParameters->bGenerateOutputMip2 = false;
 		PassParameters->bGenerateOutputMip3 = false;
@@ -2363,6 +2434,25 @@ FDefaultTemporalUpscaler::FOutputs AddTemporalSuperResolutionPasses(
 			PassParameters->PrevHistoryMetadataTexture = GraphBuilder.CreateSRV(PrevMetadataSRVDesc);
 		}
 
+		PassParameters->PrevDistortingDisplacementTexture = BlackDummy;
+		PassParameters->ResurrectedDistortingDisplacementTexture = BlackDummy;
+		PassParameters->UndistortingDisplacementTexture = BlackDummy;
+		if (bLensDistortion)
+		{
+			if (!bCameraCut && InputHistory.DistortingDisplacementTextures[PrevFrameSliceIndex].IsValid())
+			{
+				PassParameters->PrevDistortingDisplacementTexture = GraphBuilder.RegisterExternalTexture(InputHistory.DistortingDisplacementTextures[PrevFrameSliceIndex]);
+			}
+			if (!bCameraCut && InputHistory.DistortingDisplacementTextures[ResurrectionFrameSliceIndex].IsValid() && bCanResurrectHistory)
+			{
+				PassParameters->ResurrectedDistortingDisplacementTexture = GraphBuilder.RegisterExternalTexture(InputHistory.DistortingDisplacementTextures[ResurrectionFrameSliceIndex]);
+			}
+			if (PassInputs.LensDistortionLUT.UndistortingDisplacementTexture)
+			{
+				PassParameters->UndistortingDisplacementTexture = PassInputs.LensDistortionLUT.UndistortingDisplacementTexture;
+			}
+		}
+
 		{
 			FRDGTextureUAVDesc ColorUAVDesc(History.ColorArray);
 			ColorUAVDesc.FirstArraySlice = CurrentFrameSliceIndex;
@@ -2449,12 +2539,13 @@ FDefaultTemporalUpscaler::FOutputs AddTemporalSuperResolutionPasses(
 		TShaderMapRef<FTSRUpdateHistoryCS> ComputeShader(View.ShaderMap, PermutationVector);
 		FComputeShaderUtils::AddPass(
 			GraphBuilder,
-			RDG_EVENT_NAME("TSR UpdateHistory(#%d Quality=%s%s%s%s%s) %dx%d",
+			RDG_EVENT_NAME("TSR UpdateHistory(#%d Quality=%s%s%s%s%s%s) %dx%d",
 				PermutationVector.ToDimensionValueId(),
 				kUpdateQualityNames[int32(PermutationVector.Get<FTSRUpdateHistoryCS::FQualityDim>())],
 				PermutationVector.Get<FTSRShader::F16BitVALUDim>() ? TEXT(" 16bit") : TEXT(""),
 				PermutationVector.Get<FTSRShader::FAlphaChannelDim>() ? TEXT(" AlphaChannel") : TEXT(""),
 				HistoryColorFormat == PF_FloatR11G11B10 ? TEXT(" R11G11B10") : TEXT(""),
+				bSupportsLensDistortion ? (bLensDistortion ? TEXT(" ApplyLensDistortion") : TEXT(" SupportLensDistortion")) : TEXT(""),
 				PassParameters->bGenerateOutputMip3 ? TEXT(" OutputMip3") : (PassParameters->bGenerateOutputMip2 ? TEXT(" OutputMip2") : (PassParameters->bGenerateOutputMip1 ? TEXT(" OutputMip1") : TEXT(""))),
 				HistorySize.X, HistorySize.Y),
 			AsyncComputePasses >= 3 ? ERDGPassFlags::AsyncCompute : ERDGPassFlags::Compute,
@@ -2565,16 +2656,19 @@ FDefaultTemporalUpscaler::FOutputs AddTemporalSuperResolutionPasses(
 			OutputHistory.ViewMatrices.SetNum(OutputHistory.FrameStorageCount);
 			OutputHistory.SceneColorPreExposures.SetNum(OutputHistory.FrameStorageCount);
 			OutputHistory.InputViewportRects.SetNum(OutputHistory.FrameStorageCount);
+			OutputHistory.DistortingDisplacementTextures.SetNum(OutputHistory.FrameStorageCount);
 		}
 		else
 		{
-			OutputHistory.ViewMatrices = InputHistory.ViewMatrices;
-			OutputHistory.SceneColorPreExposures = InputHistory.SceneColorPreExposures;
-			OutputHistory.InputViewportRects = InputHistory.InputViewportRects;
+			OutputHistory.ViewMatrices                   = InputHistory.ViewMatrices;
+			OutputHistory.SceneColorPreExposures         = InputHistory.SceneColorPreExposures;
+			OutputHistory.InputViewportRects             = InputHistory.InputViewportRects;
+			OutputHistory.DistortingDisplacementTextures = InputHistory.DistortingDisplacementTextures;
 		}
 		OutputHistory.ViewMatrices[CurrentFrameSliceIndex] = View.ViewMatrices;
 		OutputHistory.SceneColorPreExposures[CurrentFrameSliceIndex] = View.PreExposure;
 		OutputHistory.InputViewportRects[CurrentFrameSliceIndex] = InputRect;
+		OutputHistory.DistortingDisplacementTextures[CurrentFrameSliceIndex] = nullptr;
 
 		// Extract filterable history
 		GraphBuilder.QueueTextureExtraction(History.ColorArray, &OutputHistory.ColorArray);
@@ -2588,6 +2682,11 @@ FDefaultTemporalUpscaler::FOutputs AddTemporalSuperResolutionPasses(
 			GraphBuilder.QueueTextureExtraction(History.MoireArray, &OutputHistory.MoireArray);
 		}
 
+		if (bLensDistortion && PassInputs.LensDistortionLUT.IsEnabled())
+		{
+			GraphBuilder.QueueTextureExtraction(PassInputs.LensDistortionLUT.DistortingDisplacementTexture, &OutputHistory.DistortingDisplacementTextures[CurrentFrameSliceIndex]);
+		}
+
 		// Extract the output for next frame SSR so that separate translucency shows up in SSR.
 		{
 			// Output in TemporalAAHistory and not CustomSSR so Lumen can pick up ScreenSpaceRayTracingInput in priority to ensure consistent behavior between TAA and TSR.
@@ -2670,6 +2769,12 @@ FDefaultTemporalUpscaler::FOutputs AddTemporalSuperResolutionPasses(
 			PassParameters->FlickeringFramePeriod = FlickeringFramePeriod;
 			PassParameters->PerceptionAdd = FMath::Pow(0.5f, CVarTSRShadingExposureOffset.GetValueOnRenderThread());
 
+			PassParameters->UndistortingDisplacementTexture = BlackDummy;
+			if (bLensDistortion && PassInputs.LensDistortionLUT.IsEnabled())
+			{
+				PassParameters->UndistortingDisplacementTexture = PassInputs.LensDistortionLUT.UndistortingDisplacementTexture;
+			}
+
 			PassParameters->InputTexture = PassInputs.SceneColor.Texture;
 			if (PassInputs.FlickeringInputTexture.IsValid())
 			{
diff --git a/Engine/Source/Runtime/Renderer/Private/SceneRendering.cpp b/Engine/Source/Runtime/Renderer/Private/SceneRendering.cpp
--- a/Engine/Source/Runtime/Renderer/Private/SceneRendering.cpp
+++ b/Engine/Source/Runtime/Renderer/Private/SceneRendering.cpp
@@ -423,6 +423,12 @@ static TAutoConsoleVariable<int32> CVarEnableMultiGPUForkAndJoin(
 	ECVF_Default
 	);
 
+static TAutoConsoleVariable<float> CVarLensDistortionAffectScreenPercentage(
+	TEXT("r.LensDistortion.AffectScreenPercentage"),
+	0.0f,
+	TEXT("Whether the screen percentage is automatically increased to avoid any upscaling due to the distortion. Disabled by default as this affect render target sizes, and is dependent of the upscaling factor that migth be animated (different FOV or distortion settings for instance)."),
+	ECVF_RenderThreadSafe);
+
 /*-----------------------------------------------------------------------------
 	FParallelCommandListSet
 -----------------------------------------------------------------------------*/
@@ -3089,6 +3095,16 @@ FIntPoint FSceneRenderer::GetDesiredInternalBufferSize(const FSceneViewFamily& V
 		ResolutionFractionUpperBound = PrimaryResolutionFractionUpperBound * ViewFamily.SecondaryViewFraction;
 	}
 
+	if (ViewFamily.Views[0]->bIsViewInfo)
+	{
+		const FViewInfo& View = static_cast<const FViewInfo&>(*ViewFamily.Views[0]);
+		if (View.LensDistortionLUT.IsEnabled())
+		{
+			float AffectScreenPercentage = CVarLensDistortionAffectScreenPercentage.GetValueOnRenderThread();
+			ResolutionFractionUpperBound *= FMath::Lerp(1.0, View.LensDistortionLUT.ResolutionFraction, AffectScreenPercentage);
+		}
+	}
+
 	FIntPoint FamilySizeUpperBound(0, 0);
 
 	for (const FSceneView* View : ViewFamily.AllViews)
@@ -3238,8 +3254,15 @@ void FSceneRenderer::PrepareViewRectsForRendering(FRHICommandListImmediate& RHIC
 		check(ISceneViewFamilyScreenPercentage::IsValidResolutionFraction(PrimaryResolutionFraction));
 	}
 
+	float LensDistortionResolutionFraction = 1.0f;
+	if (Views[0].LensDistortionLUT.IsEnabled())
+	{
+		float AffectScreenPercentage = CVarLensDistortionAffectScreenPercentage.GetValueOnRenderThread();
+		LensDistortionResolutionFraction = FMath::Lerp(1.0, Views[0].LensDistortionLUT.ResolutionFraction, AffectScreenPercentage);
+	}
+
 	// Compute final resolution fraction.
-	float ResolutionFraction = PrimaryResolutionFraction * ViewFamily.SecondaryViewFraction;
+	float ResolutionFraction = PrimaryResolutionFraction * ViewFamily.SecondaryViewFraction * LensDistortionResolutionFraction;
 
 	// Checks that view rects are correctly initialized.
 	for (int32 i = 0; i < Views.Num(); i++)
@@ -3743,6 +3766,19 @@ IVisibilityTaskData* FSceneRenderer::OnRenderBegin(FRDGBuilder& GraphBuilder)
 				}
 			}
 		}
+
+		if (FPaniniProjectionConfig::IsEnabledByCVars())
+		{
+			const FPaniniProjectionConfig PaniniProjection = FPaniniProjectionConfig::ReadCVars();
+
+			for (FViewInfo& View : Views)
+			{
+				if (View.ViewMatrices.IsPerspectiveProjection())
+				{
+					View.LensDistortionLUT = PaniniProjection.GenerateLUTPasses(GraphBuilder, View);
+				}
+			}
+		}
 	
 		PrepareViewRectsForRendering(GraphBuilder.RHICmdList);
 	
diff --git a/Engine/Source/Runtime/Renderer/Private/SceneRendering.h b/Engine/Source/Runtime/Renderer/Private/SceneRendering.h
--- a/Engine/Source/Runtime/Renderer/Private/SceneRendering.h
+++ b/Engine/Source/Runtime/Renderer/Private/SceneRendering.h
@@ -50,6 +50,7 @@
 #include "Nanite/NaniteShared.h"
 #include "LightFunctionAtlas.h"
 #include "SceneExtensions.h"
+#include "PostProcess/LensDistortion.h"
 
 #if RHI_RAYTRACING
 #include "RayTracingInstanceBufferUtil.h"
@@ -875,6 +876,7 @@ struct FTSRHistory
 	TArray<FViewMatrices> ViewMatrices;
 	TArray<float> SceneColorPreExposures;
 	TArray<FIntRect> InputViewportRects;
+	TArray<TRefCountPtr<IPooledRenderTarget>> DistortingDisplacementTextures;
 
 
 	void SafeRelease()
@@ -1440,6 +1442,8 @@ public:
 
 	FTextureRHIRef PrimitiveSceneDataTextureOverrideRHI;
 
+	FLensDistortionLUT LensDistortionLUT;
+
 	FShaderPrintData ShaderPrintData;
 
 private:
