Only apply VRS to the translucency pass when the translucency buffer matches internal resolution (r.SeparateTranslucencyScreenPercentage=100).
If the buffer is using a different resolution, the shading rate image will not be properly aligned. Shelved change from the Coalition.

diff --git a/Engine/Source/Runtime/Renderer/Private/TranslucentRendering.cpp b/Engine/Source/Runtime/Renderer/Private/TranslucentRendering.cpp
--- a/Engine/Source/Runtime/Renderer/Private/TranslucentRendering.cpp
+++ b/Engine/Source/Runtime/Renderer/Private/TranslucentRendering.cpp
@@ -1250,7 +1250,13 @@ static void RenderTranslucencyViewInner(
 	{
 		PassParameters->RenderTargets.DepthStencil = FDepthStencilBinding(SceneDepthTexture, ERenderTargetLoadAction::ELoad, ERenderTargetLoadAction::ELoad, FExclusiveDepthStencil::DepthRead_StencilWrite);
 	}
-	PassParameters->RenderTargets.ShadingRateTexture = GVRSImageManager.GetVariableRateShadingImage(GraphBuilder, View, FVariableRateShadingImageManager::EVRSPassType::TranslucencyAll);
+
+	// Only apply VRS if the translucency render target is at the same scale as the internal resolution
+	const bool bUseVRS = (ViewportScale == 1.0f);
+	if (bUseVRS)
+	{
+		PassParameters->RenderTargets.ShadingRateTexture = GVRSImageManager.GetVariableRateShadingImage(GraphBuilder, View, FVariableRateShadingImageManager::EVRSPassType::TranslucencyAll);
+	}
 	PassParameters->RenderTargets.ResolveRect = FResolveRect(Viewport.Rect);
 
 	const EMeshPass::Type MeshPass = TranslucencyPassToMeshPass(TranslucencyPass);
