Set bAllowVariableRateShading to true by default in FGraphicsPipelineStateInitializer and FGraphicsMinimalPipelineStateInitializer in order to allow using VRS on non-material passes.
Currently, this defaults to false and can only be set to true by "Allow Variable Rate Shading" being checked in a material. Setting it to true by default will enable VRS for non-material passes without altering material behavior. Shelved change from The Coalition.

diff --git a/Engine/Source/Runtime/RHI/Public/RHIResources.h b/Engine/Source/Runtime/RHI/Public/RHIResources.h
--- a/Engine/Source/Runtime/RHI/Public/RHIResources.h
+++ b/Engine/Source/Runtime/RHI/Public/RHIResources.h
@@ -4083,7 +4083,7 @@ public:
 		, bDepthBounds(false)
 		, MultiViewCount(0)
 		, bHasFragmentDensityAttachment(false)
-		, bAllowVariableRateShading(false)
+		, bAllowVariableRateShading(true)
 		, ShadingRate(EVRSShadingRate::VRSSR_1x1)
 		, Flags(0)
 		, StatePrecachePSOHash(0)
diff --git a/Engine/Source/Runtime/Renderer/Public/MeshPassProcessor.h b/Engine/Source/Runtime/Renderer/Public/MeshPassProcessor.h
--- a/Engine/Source/Runtime/Renderer/Public/MeshPassProcessor.h
+++ b/Engine/Source/Runtime/Renderer/Public/MeshPassProcessor.h
@@ -515,7 +515,7 @@ public:
 	bool							bDepthBounds = false;
 	uint8							MultiViewCount = 0;
 	bool							bHasFragmentDensityAttachment = false;
-	bool							bAllowVariableRateShading = false;
+	bool							bAllowVariableRateShading = true;
 	EVRSShadingRate					DrawShadingRate  = EVRSShadingRate::VRSSR_1x1;
 
 	EPrimitiveType					PrimitiveType;
