Fix for light function disapearing from light with disabled "cast shadow" and invaldi for light function atlas.
preflight 665e2fc6b83cb1ec75de6892

diff --git a/Engine/Source/Runtime/Renderer/Private/LightRendering.cpp b/Engine/Source/Runtime/Renderer/Private/LightRendering.cpp
--- a/Engine/Source/Runtime/Renderer/Private/LightRendering.cpp
+++ b/Engine/Source/Runtime/Renderer/Private/LightRendering.cpp
@@ -597,7 +597,7 @@ FDeferredLightUniformStruct GetDeferredLightParameters(const FSceneView& View, c
 	}
 
 	const bool bDynamicShadows = View.Family->EngineShowFlags.DynamicShadows && GetShadowQuality() > 0;
-	const bool bHasLightFunction = LightSceneInfo.Proxy->GetLightFunctionMaterial() != NULL && !bUseLightFunctionAtlas;
+	const bool bHasLightFunction = LightSceneInfo.Proxy->GetLightFunctionMaterial() != NULL && (!bUseLightFunctionAtlas || !LightSceneInfo.Proxy->HasValidLightFunctionAtlasSlot());
 	Out.ShadowedBits = LightSceneInfo.Proxy->CastsStaticShadow() || bHasLightFunction ? 1 : 0;
 	Out.ShadowedBits |= LightSceneInfo.Proxy->CastsDynamicShadow() && View.Family->EngineShowFlags.DynamicShadows ? 3 : 0;
 
