Mobile: Fixed missing AO when using deferred shading and a static SkyLights
#rnx

diff --git a/Engine/Shaders/Private/MobileDeferredShading.usf b/Engine/Shaders/Private/MobileDeferredShading.usf
--- a/Engine/Shaders/Private/MobileDeferredShading.usf
+++ b/Engine/Shaders/Private/MobileDeferredShading.usf
@@ -177,9 +177,11 @@ void MobileDirectionalLightPS(
 	}
 #endif
 
+	half GatheredAmbientOcclusion = 1;
 #if ENABLE_AMBIENT_OCCLUSION
-	half4 GatheredAmbientOcclusion = Texture2DSample(MobileBasePass.AmbientOcclusionTexture, MobileBasePass.AmbientOcclusionSampler, UVAndScreenPos.xy);
-	GBuffer.GBufferAO *= lerp(1.0, GatheredAmbientOcclusion.r, MobileBasePass.AmbientOcclusionStaticFraction);
+	GatheredAmbientOcclusion = Texture2DSample(MobileBasePass.AmbientOcclusionTexture, MobileBasePass.AmbientOcclusionSampler, UVAndScreenPos.xy).r;
+	GatheredAmbientOcclusion = lerp(1.0, GatheredAmbientOcclusion, MobileBasePass.AmbientOcclusionStaticFraction);
+	GBuffer.GBufferAO *= GatheredAmbientOcclusion;
 #endif
 
 	// Directional light
@@ -203,7 +205,8 @@ void MobileDirectionalLightPS(
 	OutProxyAdditive.rgb = Color;
 #else
 	OutColor.rgb = Color;
-	OutColor.a = 1;
+	// Blend is enabled only for static skylights
+	OutColor.a = GatheredAmbientOcclusion;
 #endif
 }
 
@@ -369,9 +372,11 @@ void MobileReflectionEnvironmentSkyLightingPS(
 	float2 LocalPosition = SvPosition.xy - View.ViewRectMin.xy;
 	uint GridIndex = ComputeLightGridCellIndex(uint2(LocalPosition.x, LocalPosition.y), GBuffer.Depth);
 
+	half GatheredAmbientOcclusion = 1;
 #if ENABLE_AMBIENT_OCCLUSION
-	half4 GatheredAmbientOcclusion = Texture2DSample(MobileBasePass.AmbientOcclusionTexture, MobileBasePass.AmbientOcclusionSampler, UVAndScreenPos.xy);
-	GBuffer.GBufferAO *= lerp(1.0, GatheredAmbientOcclusion.r, MobileBasePass.AmbientOcclusionStaticFraction);
+	GatheredAmbientOcclusion = Texture2DSample(MobileBasePass.AmbientOcclusionTexture, MobileBasePass.AmbientOcclusionSampler, UVAndScreenPos.xy).r;
+	GatheredAmbientOcclusion = lerp(1.0, GatheredAmbientOcclusion, MobileBasePass.AmbientOcclusionStaticFraction);
+	GBuffer.GBufferAO *= GatheredAmbientOcclusion;
 #endif
 	
 	FLightAccumulator DirectLighting = (FLightAccumulator)0;
@@ -384,6 +389,7 @@ void MobileReflectionEnvironmentSkyLightingPS(
 	OutProxyAdditive.rgb = ReflectionAndSky.rgb;
 #else
 	OutColor.rgb = ReflectionAndSky.rgb;
-	OutColor.a = 1;
+	// Blend is enabled only for static skylights
+	OutColor.a = GatheredAmbientOcclusion;
 #endif
 }
\ No newline at end of file
diff --git a/Engine/Source/Runtime/Renderer/Private/MobileDeferredShadingPass.cpp b/Engine/Source/Runtime/Renderer/Private/MobileDeferredShadingPass.cpp
--- a/Engine/Source/Runtime/Renderer/Private/MobileDeferredShadingPass.cpp
+++ b/Engine/Source/Runtime/Renderer/Private/MobileDeferredShadingPass.cpp
@@ -395,7 +395,15 @@ void RenderReflectionEnvironmentSkyLighting(FRHICommandList& RHICmdList, const F
 	FGraphicsPipelineStateInitializer GraphicsPSOInit;
 	RHICmdList.ApplyCachedRenderTargets(GraphicsPSOInit);
 	// Add to emissive in SceneColor
-	GraphicsPSOInit.BlendState = TStaticBlendState<CW_RGB, BO_Add, BF_One, BF_One>::GetRHI();
+	if (!bDynamicSkyLight)
+	{
+		// pre-multiply SceneColor with AO
+		GraphicsPSOInit.BlendState = TStaticBlendState<CW_RGB, BO_Add, BF_One, BF_SourceAlpha>::GetRHI();
+	}
+	else
+	{
+		GraphicsPSOInit.BlendState = TStaticBlendState<CW_RGB, BO_Add, BF_One, BF_One>::GetRHI();
+	}
 	GraphicsPSOInit.RasterizerState = TStaticRasterizerState<>::GetRHI();
 	GraphicsPSOInit.DepthStencilState = TStaticDepthStencilState<
 		false, CF_Always,
@@ -471,8 +479,6 @@ static void RenderDirectionalLight(FRHICommandList& RHICmdList, const FScene& Sc
 
 	FGraphicsPipelineStateInitializer GraphicsPSOInit;
 	RHICmdList.ApplyCachedRenderTargets(GraphicsPSOInit);
-	// Add to emissive in SceneColor
-	GraphicsPSOInit.BlendState = TStaticBlendState<CW_RGB, BO_Add, BF_One, BF_One>::GetRHI();
 	GraphicsPSOInit.RasterizerState = TStaticRasterizerState<>::GetRHI();
 
 	uint8 LightingChannelStencilValue = GetLightingChannelStencilValue(LightingChannel);
@@ -530,6 +536,17 @@ static void RenderDirectionalLight(FRHICommandList& RHICmdList, const FScene& Sc
 	const bool bDynamicShadows = DirectionalLight.Proxy->CastsDynamicShadow() && View.Family->EngineShowFlags.DynamicShadows;
 	const bool bPlanarReflection = Scene.GetForwardPassGlobalPlanarReflection() != nullptr;
 
+	// Add to emissive in SceneColor
+	if (bInlineReflectionAndSky && !bDynamicSkyLight)
+	{
+		// pre-multiply SceneColor with AO
+		GraphicsPSOInit.BlendState = TStaticBlendState<CW_RGB, BO_Add, BF_One, BF_SourceAlpha>::GetRHI();
+	}
+	else
+	{
+		GraphicsPSOInit.BlendState = TStaticBlendState<CW_RGB, BO_Add, BF_One, BF_One>::GetRHI();
+	}
+	
 	// Do two passes, first masking DefautLit, second masking all other shading models
 	const bool bOnlyDefaultLitInView = IsOnlyDefaultLitShadingModel(View.ShadingModelMaskInView);
 	int32 NumPasses = !bOnlyDefaultLitInView && MobileUsesGBufferCustomData(Scene.GetShaderPlatform()) ? 2 : 1;
@@ -593,6 +610,11 @@ static void RenderDirectionalLights(FRHICommandList& RHICmdList, const FScene& S
 	bool bPrimitivesUseLightingChannels = (View.bUsesLightingChannels && GMobileIgnoreDeferredShadingSkyLightChannels == 0);
 	const bool bInlineReflectionAndSky = (NumLights == 1) && !bPrimitivesUseLightingChannels && (Scene.MobileDirectionalLights[0] != nullptr);
 
+	if (!bInlineReflectionAndSky)
+	{
+		RenderReflectionEnvironmentSkyLighting(RHICmdList, Scene, View);
+	}
+
 	for (uint32 ChannelIdx = 0; ChannelIdx < UE_ARRAY_COUNT(Scene.MobileDirectionalLights); ChannelIdx++)
 	{
 		FLightSceneInfo* DirectionalLight = Scene.MobileDirectionalLights[ChannelIdx];
@@ -601,11 +623,6 @@ static void RenderDirectionalLights(FRHICommandList& RHICmdList, const FScene& S
 			RenderDirectionalLight(RHICmdList, Scene, View, DefaultLightMaterial, *DirectionalLight, ChannelIdx, bInlineReflectionAndSky);
 		}
 	}
-
-	if (!bInlineReflectionAndSky)
-	{
-		RenderReflectionEnvironmentSkyLighting(RHICmdList, Scene, View);
-	}
 }
 
 template<uint32 LightingChannel, bool bWithStencilCulling>
