Light function atlas - size and slot resolution can now be tweak using cvars with more freedom.
If requested reslution is too high, then slot resolution will be reduced to fit platform restrictions

diff --git a/Engine/Source/Runtime/Renderer/Private/LightFunctionAtlas.cpp b/Engine/Source/Runtime/Renderer/Private/LightFunctionAtlas.cpp
--- a/Engine/Source/Runtime/Renderer/Private/LightFunctionAtlas.cpp
+++ b/Engine/Source/Runtime/Renderer/Private/LightFunctionAtlas.cpp
@@ -59,14 +59,14 @@ static TAutoConsoleVariable<int32> CVarLightFunctionAtlas(
 static TAutoConsoleVariable<int32> CVarLightFunctionAtlasSlotResolution(
 	TEXT("r.LightFunctionAtlas.SlotResolution"),
 	128,
-	TEXT("Experimental: The resolution of each atlas slot. Maximum value is 256."),
+	TEXT("Experimental: The resolution of each atlas slot. If the resolution is too large for the target platforms, it will be reduced to fit."),
 	ECVF_RenderThreadSafe);
 
 // We do not dynamically scale allocated slot resolution for now.
 static TAutoConsoleVariable<int32> CVarLightFunctionAtlasSize(
 	TEXT("r.LightFunctionAtlas.Size"),
 	4,
-	TEXT("Experimental: The default size in atlas slot count of the edge of the 2D texture atlas. Maximum value is 16."),
+	TEXT("Experimental: The default size (atlas slot count) for each dimension of the the 2D texture atlas. Maximum value is 16 since we only allow up to 16x16=256 light functions."),
 	ECVF_RenderThreadSafe);
 
 static TAutoConsoleVariable<int32> CVarLightFunctionAtlasMaxLightCount(
@@ -117,25 +117,35 @@ FAutoConsoleVariableRef CVarLumenUsesLightFunctionAtlas(
 
 //////////////////////////////////////////////////////////////////////////
 
-static const uint32 MAX_LIGHT_FUNCTION_ATLAS_SLOT_RESOLUTION = 256;
-static const uint32 MAX_LIGHT_FUNCTION_ATLAS_EDGE_SIZE = 16;
-
-static uint32 GetAtlasSlotResolution()
+namespace LightFunctionAtlas
 {
-	const uint32 AtlasSlotResolution = FMath::Clamp(CVarLightFunctionAtlasSlotResolution.GetValueOnRenderThread(), 32, MAX_LIGHT_FUNCTION_ATLAS_SLOT_RESOLUTION);
-	return AtlasSlotResolution;
-}
 
-static uint32 GetAtlasEdgeSize()
+static FLightFunctionAtlasSetup GetLightFunctionAtlasSetup()
 {
-	const uint32 AtlasEdgeSize = FMath::Clamp(CVarLightFunctionAtlasSize.GetValueOnRenderThread(), 4, MAX_LIGHT_FUNCTION_ATLAS_EDGE_SIZE);// 16x16 is the maximum slot count of LIGHT_FUNCTION_ATLAS_MAX_LIGHT_FUNCTION_COUNT=256 we currently allow
-	return AtlasEdgeSize;
-}
 
+	// 16x16 is the maximum slot count of LIGHT_FUNCTION_ATLAS_MAX_LIGHT_FUNCTION_COUNT=256 we currently allow
+	const uint32 MaxEdgeSize = 16;
+	check((MaxEdgeSize * MaxEdgeSize) <= LIGHT_FUNCTION_ATLAS_MAX_LIGHT_FUNCTION_COUNT);
 
+	uint32 RequestedEdgeSize = FMath::Clamp(CVarLightFunctionAtlasSize.GetValueOnRenderThread(), 2, MaxEdgeSize);
+	uint32 RequestedSlotResolution = FMath::Max(CVarLightFunctionAtlasSlotResolution.GetValueOnRenderThread(), 32);
 
-namespace LightFunctionAtlas
-{
+	// The atlas texture resolution should not be larger than the maximum resolution supported by the current platforms.
+	const uint32 MaxAtlasResolution = 32 * 1024;
+	const uint32 Max2DTextureDimension = GetMax2DTextureDimension();
+	if ((RequestedEdgeSize * RequestedSlotResolution) > Max2DTextureDimension)
+	{
+		RequestedSlotResolution = Max2DTextureDimension / RequestedEdgeSize;
+	}
+
+	// Due to the way we pack AtlasSlot.MinU and MinV, we should allow resolution up to 32K only. (UVs are stored as uint16 so 64K but we need half texel precision).
+	check((RequestedSlotResolution * RequestedEdgeSize) <= MaxAtlasResolution);
+
+	FLightFunctionAtlasSetup AtlasSetup;
+	AtlasSetup.EdgeSize = RequestedEdgeSize;
+	AtlasSetup.SlotResolution = RequestedSlotResolution;
+	return AtlasSetup;
+}
 
 IMPLEMENT_GLOBAL_SHADER_PARAMETER_STRUCT(FLightFunctionAtlasGlobalParameters, "LightFunctionAtlas");
 
@@ -259,6 +269,8 @@ void FLightFunctionAtlas::BeginSceneFrame(const FViewFamilyInfo& ViewFamily, TAr
 {
 	ClearEmptySceneFrame(nullptr, 0, &LightFunctionAtlasSceneData);
 
+	AtlasSetup = GetLightFunctionAtlasSetup();
+
 	// Now lets check if we need to generate the atlas for this frame
 	bLightFunctionAtlasEnabled = CVarLightFunctionAtlas.GetValueOnRenderThread() > 0 && ViewFamily.EngineShowFlags.LightFunctions > 0;
 #if !UE_BUILD_SHIPPING
@@ -380,8 +392,8 @@ void FLightFunctionAtlas::AllocateAtlasSlots(const TArray<FViewInfo>& Views)
 	//
 	// Allocate slots until we cannot anymore and set light function slot index on FLightSceneInfo to be send to the GPU later
 	//
-	const uint32 AtlasSlotResolution = GetAtlasSlotResolution();
-	const float AtlasEdgeSize = GetAtlasEdgeSize();
+	const uint32 AtlasSlotResolution = AtlasSetup.SlotResolution;
+	const float AtlasEdgeSize = AtlasSetup.EdgeSize;
 	const float AtlasResolution = AtlasSlotResolution * AtlasEdgeSize;
 	check(AtlasEdgeSize * AtlasEdgeSize <= LIGHT_FUNCTION_ATLAS_MAX_LIGHT_FUNCTION_COUNT);
 	const uint32 AtlasMaxLightFunctionCount = AtlasEdgeSize * AtlasEdgeSize;
@@ -581,8 +593,8 @@ void FLightFunctionAtlas::RenderLightFunctionAtlas(FRDGBuilder& GraphBuilder, TA
 		LightFunctionAtlasGlobalParameters->LightFunctionAtlasTexture = RDGAtlasTexture2D;
 		LightFunctionAtlasGlobalParameters->LightFunctionAtlasSampler = TStaticSamplerState<SF_Bilinear>::GetRHI();
 
-		const uint32 AtlasSlotResolution = GetAtlasSlotResolution();
-		float AtlasEdgeSize = GetAtlasEdgeSize();
+		const uint32 AtlasSlotResolution = AtlasSetup.SlotResolution;
+		const float AtlasEdgeSize = AtlasSetup.EdgeSize;
 		float AtlasResolution = AtlasSlotResolution * AtlasEdgeSize;
 		LightFunctionAtlasGlobalParameters->Slot_UVSize = (AtlasSlotResolution - 1.0f) / AtlasResolution; // -1.0 because we remove a bit more than half a texel at the border.
 
@@ -630,9 +642,6 @@ void FLightFunctionAtlas::RenderLightFunctionAtlas(FRDGBuilder& GraphBuilder, TA
 
 			const EffectiveLightFunctionSlot& AtlasSlot = EffectiveLightFunctionSlotArray[LightFunctionAtlasSlotIndex];
 
-			static_assert(MAX_LIGHT_FUNCTION_ATLAS_SLOT_RESOLUTION * MAX_LIGHT_FUNCTION_ATLAS_EDGE_SIZE <= 32 * 1024, 
-				"Unable to pack slot UVs into uint16 when atlas resolution is larger than 32K");
-
 			const uint32 PackedAtlasSlotMinU = uint32(round(AtlasSlot.MinU * 65536.0f));
 			const uint32 PackedAtlasSlotMinV = uint32(round(AtlasSlot.MinV * 65536.0f));
 			const uint32 PackedAtlasSlotMinUV = (PackedAtlasSlotMinU | (PackedAtlasSlotMinV << 16));
@@ -663,8 +672,8 @@ void FLightFunctionAtlas::RenderLightFunctionAtlas(FRDGBuilder& GraphBuilder, TA
 
 void FLightFunctionAtlas::AllocateTexture2DAtlas(FRDGBuilder& GraphBuilder)
 {
-	uint32 AtlasSlotResolution = GetAtlasSlotResolution();
-	uint32 AtlasEdgeSize = GetAtlasEdgeSize();
+	const uint32 AtlasSlotResolution = AtlasSetup.SlotResolution;
+	const float AtlasEdgeSize = AtlasSetup.EdgeSize;
 	uint32 AtlasResolution = AtlasSlotResolution * AtlasEdgeSize;
 	const uint32 MipCount = 1;
 
@@ -700,8 +709,8 @@ void FLightFunctionAtlas::RenderAtlasSlots(FRDGBuilder& GraphBuilder, const TArr
 		ERDGPassFlags::Raster,
 		[PassParameters, &Views, this](FRHICommandList& RHICmdList)
 		{
-			uint32 AtlasSlotResolution = GetAtlasSlotResolution();
-			uint32 AtlasEdgeSize = GetAtlasEdgeSize();
+			const uint32 AtlasSlotResolution = AtlasSetup.SlotResolution;
+			const float AtlasEdgeSize = AtlasSetup.EdgeSize;
 			uint32 AtlasResolution = AtlasSlotResolution * AtlasEdgeSize;
 
 			// This always work because in this case we do not need anything from any view.
@@ -819,8 +828,8 @@ FScreenPassTexture FLightFunctionAtlas::AddDebugVisualizationPasses(FRDGBuilder&
 		return ScreenPassSceneColor;
 	}
 
-	const uint32 AtlasSlotResolution	= GetAtlasSlotResolution();
-	const uint32 AtlasEdgeSize			= GetAtlasEdgeSize();
+	const uint32 AtlasSlotResolution	= AtlasSetup.SlotResolution;
+	const float AtlasEdgeSize			= AtlasSetup.EdgeSize;
 	const uint32 AtlasResolution		= AtlasSlotResolution * AtlasEdgeSize;
 
 	const FIntPoint SrcPoint			= FIntPoint::ZeroValue;
@@ -871,7 +880,10 @@ FScreenPassTexture FLightFunctionAtlas::AddDebugVisualizationPasses(FRDGBuilder&
 		const float DPIScale = Canvas.GetDPIScale();
 		Canvas.SetBaseTransform(FMatrix(FScaleMatrix(DPIScale) * Canvas.CalcBaseTransform2D(Canvas.GetViewRect().Width(), Canvas.GetViewRect().Height())));
 
-		Canvas.DrawShadowedString(DstPoint.X + 180.0f, DstPoint.Y - 20.0f, TEXT("LIGHT FUNCTION ATLAS"), GEngine->GetLargeFont(), FLinearColor::White);
+		Canvas.DrawShadowedString(DstPoint.X + 180.0f, DstPoint.Y - 40.0f, TEXT("LIGHT FUNCTION ATLAS"), GEngine->GetLargeFont(), FLinearColor::White);
+
+		Text = FString::Printf(TEXT("Slot Resolution = %dx%d  -  Size = %dx%d"), AtlasSetup.SlotResolution, AtlasSetup.SlotResolution, AtlasSetup.EdgeSize, AtlasSetup.EdgeSize);
+		Canvas.DrawShadowedString(DstPoint.X + 130.0f, DstPoint.Y - 20.0f, *Text, GEngine->GetLargeFont(), FLinearColor::White);
 
 		Text = FString::Printf(TEXT("Light Functions in atlas:         %d"), EffectiveLightFunctionSlotArray.Num());
 		Canvas.DrawShadowedString(DrawPosX, DrawPosY, *Text, GEngine->GetLargeFont(), FLinearColor::White);
diff --git a/Engine/Source/Runtime/Renderer/Private/LightFunctionAtlas.h b/Engine/Source/Runtime/Renderer/Private/LightFunctionAtlas.h
--- a/Engine/Source/Runtime/Renderer/Private/LightFunctionAtlas.h
+++ b/Engine/Source/Runtime/Renderer/Private/LightFunctionAtlas.h
@@ -123,6 +123,12 @@ struct FAtlasLightInfoData
 	FMatrix44f Transform;
 };
 
+struct FLightFunctionAtlasSetup
+{
+	uint32 EdgeSize = 2;
+	uint32 SlotResolution = 32;
+};
+
 
 
 // This class holds all data and resources related light function for a single scene, including multiple views.
@@ -206,6 +212,8 @@ private:
 	uint32																LightCountSkippedDueToMissingAtlasSlot = 0;
 	TSet<FLightFunctionSlotKey>											SkippedLightFunctionsSet;
 #endif
+
+	FLightFunctionAtlasSetup AtlasSetup;
 };
 
 
