Light function atlas - improved visualisation accuracy.

diff --git a/Engine/Source/Runtime/Renderer/Private/LightFunctionAtlas.cpp b/Engine/Source/Runtime/Renderer/Private/LightFunctionAtlas.cpp
--- a/Engine/Source/Runtime/Renderer/Private/LightFunctionAtlas.cpp
+++ b/Engine/Source/Runtime/Renderer/Private/LightFunctionAtlas.cpp
@@ -447,7 +447,7 @@ void FLightFunctionAtlas::AllocateAtlasSlots(const TArray<FViewInfo>& Views)
 		return LightSlotIndex;
 	};
 
-	// Add the default invalid light slot at the beginning
+	// Add the default invalid light slot at the beginning. This is done so that we always allocate buffer with at least one element.
 	AddLightSlot(nullptr, 0);
 
 	int32 MaxLightCount = CVarLightFunctionAtlasMaxLightCount.GetValueOnRenderThread();
@@ -889,7 +889,7 @@ FScreenPassTexture FLightFunctionAtlas::AddDebugVisualizationPasses(FRDGBuilder&
 		Canvas.DrawShadowedString(DrawPosX, DrawPosY, *Text, GEngine->GetLargeFont(), FLinearColor::White);
 		DrawPosY += 20.0f;
 
-		Text = FString::Printf(TEXT("Local Lights sampling atlas: %d"), EffectiveLocalLightSlotArray.Num());
+		Text = FString::Printf(TEXT("Local Lights sampling atlas: %d"), EffectiveLocalLightSlotArray.Num() - 1);	// minus one because we add a dummy light at slot 0 using AddLightSlot(nullptr, 0);
 		Canvas.DrawShadowedString(DrawPosX, DrawPosY, *Text, GEngine->GetLargeFont(), FLinearColor::White);
 		DrawPosY += 40.0f;
 
@@ -907,7 +907,7 @@ FScreenPassTexture FLightFunctionAtlas::AddDebugVisualizationPasses(FRDGBuilder&
 			uint32 LightCountUsingThisMaterial = 0;
 			for (auto& LocalLight : EffectiveLocalLightSlotArray)
 			{
-				LightCountUsingThisMaterial += LocalLight.LightFunctionAtlasSlotIndex == LightFunctionAtlasSlotIndex ? 1 : 0;
+				LightCountUsingThisMaterial += ((LocalLight.LightSceneInfo != nullptr) && (LocalLight.LightFunctionAtlasSlotIndex == LightFunctionAtlasSlotIndex)) ? 1 : 0;
 			}
 
 			// Draw the light function material
