Path Tracer: Use 1D dispatches when starting from a compacted set of rays.

This applies to both compacted bounces (the default) and adaptive sampling primrary rays (when enabled). This appears to be slightly faster in most cases as it maintains the order the rays were written in from the previous pass.

#rnx

diff --git a/Engine/Source/Runtime/Renderer/Private/PathTracing.cpp b/Engine/Source/Runtime/Renderer/Private/PathTracing.cpp
--- a/Engine/Source/Runtime/Renderer/Private/PathTracing.cpp
+++ b/Engine/Source/Runtime/Renderer/Private/PathTracing.cpp
@@ -54,10 +54,10 @@ TAutoConsoleVariable<bool> CVarPathTracingCompaction(
 	ECVF_RenderThreadSafe
 );
 
-TAutoConsoleVariable<int32> CVarPathTracingIndirectDispatch(
+TAutoConsoleVariable<bool> CVarPathTracingIndirectDispatch(
 	TEXT("r.PathTracing.IndirectDispatch"),
-	0,
-	TEXT("Enables indirect dispatch (if supported by the hardware) for compacted path tracing (default: 0 (disabled))"),
+	false,
+	TEXT("Enables indirect dispatch (if supported by the hardware) for compacted path tracing (default: false (disabled))"),
 	ECVF_RenderThreadSafe
 );
 
@@ -3406,7 +3406,7 @@ void FDeferredShadingSceneRenderer::RenderPathTracing(
 
 			// should we use path compaction?
 			const bool bUseCompaction = (bUseExperimental == false) || CVarPathTracingCompaction.GetValueOnRenderThread() != 0;
-			const bool bUseIndirectDispatch = GRHISupportsRayTracingDispatchIndirect && CVarPathTracingIndirectDispatch.GetValueOnRenderThread() != 0;
+			const bool bUseIndirectDispatch = GRHISupportsRayTracingDispatchIndirect && CVarPathTracingIndirectDispatch.GetValueOnRenderThread();
 			const int FlushRenderingCommands = CVarPathTracingFlushDispatch.GetValueOnRenderThread();
 
 			FRDGBuffer* ActivePaths[2] = {};
@@ -3606,13 +3606,14 @@ void FDeferredShadingSceneRenderer::RenderPathTracing(
 							}
 							ClearUnusedGraphResources(RayGenShader, PassParameters);
 							const bool bFlushRenderingCommands = FlushRenderingCommands == 1 || (FlushRenderingCommands == 2 && Bounce == MaxBounces);
+							const bool bUse1DDispatch = Config.UseAdaptiveSampling || (bUseCompaction && Bounce > 0);
 							GraphBuilder.AddPass(
 								bUseCompaction
 								? RDG_EVENT_NAME("Path Tracer Sample=%d/%d NumLights=%d (Bounce=%d%s)", PathTracingState->SampleIndex, MaxSPP, PassParameters->SceneLightCount, Bounce, bUseIndirectDispatch && Bounce > 0 ? TEXT(" indirect") : TEXT(""))
 								: RDG_EVENT_NAME("Path Tracer Sample=%d/%d NumLights=%d"              , PathTracingState->SampleIndex, MaxSPP, PassParameters->SceneLightCount),
 								PassParameters,
 								ERDGPassFlags::Compute,
-								[PassParameters, RayGenShader, DispatchSizeX, DispatchSizeYLocal, bUseIndirectDispatch, bFlushRenderingCommands, GPUIndex, &View](FRHICommandList& RHICmdList)
+								[PassParameters, RayGenShader, DispatchSizeX, DispatchSizeYLocal, bUseIndirectDispatch, bUse1DDispatch, bFlushRenderingCommands, GPUIndex, &View](FRHICommandList& RHICmdList)
 								{
 									FRHIBatchedShaderParameters& GlobalResources = RHICmdList.GetScratchShaderParameters();
 									SetShaderParameters(GlobalResources, RayGenShader, *PassParameters);
@@ -3627,6 +3628,15 @@ void FDeferredShadingSceneRenderer::RenderPathTracing(
 											PassParameters->PathTracingIndirectArgs->GetIndirectRHICallBuffer(), 3 * PassParameters->Bounce * sizeof(uint32)
 										);
 									}
+									else if (bUse1DDispatch)
+									{
+										RHICmdList.RayTraceDispatch(
+											View.RayTracingMaterialPipeline,
+											RayGenShader.GetRayTracingShader(),
+											View.RayTracingSBT, GlobalResources,
+											DispatchSizeX * DispatchSizeYLocal, 1
+										);
+									}
 									else
 									{
 										RHICmdList.RayTraceDispatch(
