Fix crash when a view extension registers postprocessing pass callbacks for a disabled pass.

TOverridePassSequence allows the last enabled postprocessing pass to target the backbuffer directly. However, if a subsequent disabled pass has callbacks, those callbacks will run anyway and will try to use the backbuffer as an input, which is not supported (we crash with an assert in FRHITextureSRVCreateInfo::Validate because the backbuffer texture hasn't been created with TexCreate_ShaderResource). This change fixes the logic to take into account callbacks for disabled passes as well.

#rnx

diff --git a/Engine/Source/Runtime/Engine/Public/SceneViewExtension.h b/Engine/Source/Runtime/Engine/Public/SceneViewExtension.h
--- a/Engine/Source/Runtime/Engine/Public/SceneViewExtension.h
+++ b/Engine/Source/Runtime/Engine/Public/SceneViewExtension.h
@@ -182,6 +182,8 @@ public:
 
 	/**
 	* This will be called at the beginning of post processing to make sure that each view extension gets a chance to subscribe to an after pass event.
+	*  - The pass MUST write to the override output texture if it is active (this occurs when the pass is the last in the post processing chain writing to the back buffer).
+	*    For performance reasons it is recommended to only subscribe to a pass when the pass will produce a GPU resource. Calling 
 	*/
 	virtual void SubscribeToPostProcessingPass(EPostProcessingPass Pass, FAfterPassCallbackDelegateArray& InOutPassCallbacks, bool bIsPassEnabled) {};
 
diff --git a/Engine/Source/Runtime/Renderer/Internal/OverridePassSequence.h b/Engine/Source/Runtime/Renderer/Internal/OverridePassSequence.h
--- a/Engine/Source/Runtime/Renderer/Internal/OverridePassSequence.h
+++ b/Engine/Source/Runtime/Renderer/Internal/OverridePassSequence.h
@@ -115,20 +115,23 @@ public:
 
 	bool AcceptOverrideIfLastPass(EPass Pass, FScreenPassRenderTarget& OutTargetToOverride, const TOptional<int32>& AfterPassCallbackIndex = TOptional<int32>())
 	{
-		bool bLastAfterPass = AfterPass[(int32)Pass].Num() == 0;
+		const int32 PassIndex = (int32)Pass;
+		bool bHasMoreAfterPassCallbacks;
 
 		if (AfterPassCallbackIndex)
 		{
-			bLastAfterPass = AfterPassCallbackIndex.GetValue() == AfterPass[(int32)Pass].Num() - 1;
+			bHasMoreAfterPassCallbacks = AfterPassCallbackIndex.GetValue() < AfterPass[PassIndex].Num() - 1;
 		}
 		else
 		{
+			bHasMoreAfterPassCallbacks = AfterPass[PassIndex].Num() > 0;
+
 			// Display debug information for a Pass unless it is an after pass.
 			AcceptPass(Pass);
 		}
 
 		// We need to override output only if this is the last pass and the last after pass.
-		if (IsLastPass(Pass) && bLastAfterPass)
+		if (AcceptOverrideInPassIndex == PassIndex && !bHasMoreAfterPassCallbacks)
 		{
 			OutTargetToOverride = OverrideOutput;
 			return true;
@@ -163,6 +166,14 @@ public:
 					bFirstPass = false;
 				}
 				LastPass = (EPass)PassIndex;
+				AcceptOverrideInPassIndex = PassIndex;
+			}
+
+			// We can have callbacks for disabled passes which come after the last enabled pass. In that case we only
+			// accept the output override for the last callback of that pass, regardless of its state.
+			if (AfterPass[PassIndex].Num() > 0)
+			{
+				AcceptOverrideInPassIndex = PassIndex;
 			}
 		}
 	}
@@ -190,6 +201,7 @@ private:
 	TStaticArray<FPassInfo, PassCountMax> Passes;
 	TStaticArray<FAfterPassCallbackDelegateArray, PassCountMax> AfterPass;
 	EPass LastPass = EPass::MAX;
+	int32 AcceptOverrideInPassIndex = PassCountMax;
 
 #if RDG_ENABLE_DEBUG
 	EPass NextPass = EPass(0);
diff --git a/Engine/Source/Runtime/Renderer/Private/ScreenPass.cpp b/Engine/Source/Runtime/Renderer/Private/ScreenPass.cpp
--- a/Engine/Source/Runtime/Renderer/Private/ScreenPass.cpp
+++ b/Engine/Source/Runtime/Renderer/Private/ScreenPass.cpp
@@ -59,39 +59,52 @@ FRDGTextureRef TryCreateViewFamilyDepthTexture(FRDGBuilder& GraphBuilder, const
 }
 
 // static
-FScreenPassTexture FScreenPassTexture::CopyFromSlice(FRDGBuilder& GraphBuilder, const FScreenPassTextureSlice& ScreenTextureSlice)
+FScreenPassTexture FScreenPassTexture::CopyFromSlice(FRDGBuilder& GraphBuilder, const FScreenPassTextureSlice& ScreenTextureSlice, FScreenPassTexture OverrideOutput)
 {
-	if (!ScreenTextureSlice.TextureSRV)
+	FRDGTextureSRV* InputTextureSRV = ScreenTextureSlice.TextureSRV;
+
+	if (!InputTextureSRV)
 	{
-		return FScreenPassTexture(nullptr, ScreenTextureSlice.ViewRect);
+		return OverrideOutput;
 	}
-	else if (!ScreenTextureSlice.TextureSRV->Desc.Texture->Desc.IsTextureArray())
+
+	FRDGTexture* InputTexture = InputTextureSRV->Desc.Texture;
+
+	// We can avoid the copy if it's a 2D texture and there's no override output.
+	if (InputTexture->Desc.IsTexture2D() && !OverrideOutput.IsValid())
 	{
-		return FScreenPassTexture(ScreenTextureSlice.TextureSRV->Desc.Texture, ScreenTextureSlice.ViewRect);
+		return FScreenPassTexture(InputTexture, ScreenTextureSlice.ViewRect);
 	}
 
-	FRDGTextureDesc Desc = ScreenTextureSlice.TextureSRV->Desc.Texture->Desc;
-	Desc.Dimension = ETextureDimension::Texture2D;
-	Desc.ArraySize = 1;
+	check(InputTexture->Desc.IsTexture2D() || InputTexture->Desc.IsTextureArray());
 
-	// If a pass uses blending to write to this post process texture, it needs to support being a render target, so make sure this flag is included.
-	// Most post processing uses SceneColor or its FRDGTextureDesc, and SceneColor already has the RenderTargetable flag set, but TSR (the input to
-	// the "Before Bloom" stage) writes to texture slices using compute, and its output doesn't have this flag.
-	Desc.Flags |= ETextureCreateFlags::RenderTargetable;
+	FRDGTextureRef OutputTexture = OverrideOutput.Texture;
+
+	if (!OutputTexture)
+	{
+		FRDGTextureDesc Desc = InputTexture->Desc;
+		Desc.Dimension = ETextureDimension::Texture2D;
+		Desc.ArraySize = 1;
 
-	FRDGTextureRef NewTexture = GraphBuilder.CreateTexture(Desc, TEXT("CopyToScreenPassTexture2D"));
+		// If a pass uses blending to write to this post process texture, it needs to support being a render target, so make sure this flag is included.
+		// Most post processing uses SceneColor or its FRDGTextureDesc, and SceneColor already has the RenderTargetable flag set, but TSR (the input to
+		// the "Before Bloom" stage) writes to texture slices using compute, and its output doesn't have this flag.
+		Desc.Flags |= ETextureCreateFlags::RenderTargetable;
+
+		OutputTexture = GraphBuilder.CreateTexture(Desc, TEXT("CopyToScreenPassTexture2D"));
+	}
 
 	FRHICopyTextureInfo CopyInfo;
-	CopyInfo.SourceSliceIndex = ScreenTextureSlice.TextureSRV->Desc.FirstArraySlice;
-	CopyInfo.NumMips = ScreenTextureSlice.TextureSRV->Desc.Texture->Desc.NumMips;
+	CopyInfo.SourceSliceIndex = InputTextureSRV->Desc.FirstArraySlice;
+	CopyInfo.NumMips = InputTexture->Desc.NumMips;
 
 	AddCopyTexturePass(
 		GraphBuilder,
-		ScreenTextureSlice.TextureSRV->Desc.Texture,
-		NewTexture,
+		InputTexture,
+		OutputTexture,
 		CopyInfo);
 
-	return FScreenPassTexture(NewTexture, ScreenTextureSlice.ViewRect);
+	return FScreenPassTexture(OutputTexture, ScreenTextureSlice.ViewRect);
 }
 
 // static
diff --git a/Engine/Source/Runtime/Renderer/Public/PostProcess/PostProcessMaterialInputs.h b/Engine/Source/Runtime/Renderer/Public/PostProcess/PostProcessMaterialInputs.h
--- a/Engine/Source/Runtime/Renderer/Public/PostProcess/PostProcessMaterialInputs.h
+++ b/Engine/Source/Runtime/Renderer/Public/PostProcess/PostProcessMaterialInputs.h
@@ -116,20 +116,11 @@ struct FPostProcessMaterialInputs
 	}
 
 	/**
-	* A helper function that extracts the right scene color texture, untouched, to be used further in post processing.
-	*/
+	 * Returns the input scene color as a 2D texture output. This WILL perform a GPU copy if the override output is active or the input scene color was a 2D texture array slice.
+	 */
 	inline FScreenPassTexture ReturnUntouchedSceneColorForPostProcessing(FRDGBuilder& GraphBuilder) const
 	{
-		if (OverrideOutput.IsValid())
-		{
-			return OverrideOutput;
-		}
-		else
-		{
-			/** We don't want to modify scene texture in any way. We just want it to be passed back onto the next stage. */
-			FScreenPassTextureSlice SceneTexture = const_cast<FScreenPassTextureSlice&>(Textures[(uint32)EPostProcessMaterialInput::SceneColor]);
-			return FScreenPassTexture::CopyFromSlice(GraphBuilder, SceneTexture);
-		}
+		return FScreenPassTexture::CopyFromSlice(GraphBuilder, const_cast<FScreenPassTextureSlice&>(Textures[(uint32)EPostProcessMaterialInput::SceneColor]), OverrideOutput);
 	}
 
 	// [Optional] Render to the specified output. If invalid, a new texture is created and returned.
diff --git a/Engine/Source/Runtime/Renderer/Public/ScreenPass.h b/Engine/Source/Runtime/Renderer/Public/ScreenPass.h
--- a/Engine/Source/Runtime/Renderer/Public/ScreenPass.h
+++ b/Engine/Source/Runtime/Renderer/Public/ScreenPass.h
@@ -39,14 +39,15 @@ public:
 // Describes a texture with a paired viewport rect.
 struct FScreenPassTexture
 {
-	FScreenPassTexture() = default;
+	FScreenPassTexture() {}
 
 	explicit FScreenPassTexture(FRDGTextureRef InTexture);
 	explicit FScreenPassTexture(const FScreenPassTextureSlice& ScreenTexture);
 
 	FScreenPassTexture(FRDGTextureRef InTexture, FIntRect InViewRect);
 
-	RENDERER_API static FScreenPassTexture CopyFromSlice(FRDGBuilder& GraphBuilder, const FScreenPassTextureSlice& ScreenTextureSlice);
+	// Returns a 2D texture from the input slice. A copy is used if the input is a slice of a 2D texture array, or the override output is active.
+	RENDERER_API static FScreenPassTexture CopyFromSlice(FRDGBuilder& GraphBuilder, const FScreenPassTextureSlice& ScreenTextureSlice, FScreenPassTexture OverrideOutput = FScreenPassTexture());
 
 	bool IsValid() const;
 
