Move RayTracing mesh draw command structs and classes into RayTracingMeshDrawCommands file

diff --git a/Engine/Plugins/Experimental/GPULightmass/Source/GPULightmass/Private/Scene/Scene.h b/Engine/Plugins/Experimental/GPULightmass/Source/GPULightmass/Private/Scene/Scene.h
--- a/Engine/Plugins/Experimental/GPULightmass/Source/GPULightmass/Private/Scene/Scene.h
+++ b/Engine/Plugins/Experimental/GPULightmass/Source/GPULightmass/Private/Scene/Scene.h
@@ -7,6 +7,7 @@
 #include "Scene/InstancedStaticMesh.h"
 #include "Scene/Landscape.h"
 #include "MeshPassProcessor.h"
+#include "RayTracingMeshDrawCommands.h"
 #include "IrradianceCaching.h"
 #include "GPULightmassSettings.h"
 #include "Templates/UniquePtr.h"
diff --git a/Engine/Plugins/FX/Niagara/Source/NiagaraShader/Private/NiagaraAsyncGpuTraceProviderHwrt.cpp b/Engine/Plugins/FX/Niagara/Source/NiagaraShader/Private/NiagaraAsyncGpuTraceProviderHwrt.cpp
--- a/Engine/Plugins/FX/Niagara/Source/NiagaraShader/Private/NiagaraAsyncGpuTraceProviderHwrt.cpp
+++ b/Engine/Plugins/FX/Niagara/Source/NiagaraShader/Private/NiagaraAsyncGpuTraceProviderHwrt.cpp
@@ -9,6 +9,7 @@
 #include "FXRenderingUtils.h"
 #include "GlobalShader.h"
 #include "MeshPassProcessor.h"
+#include "RayTracingMeshDrawCommands.h"
 #include "NiagaraSettings.h"
 #include "PipelineStateCache.h"
 #include "ShaderParameterMacros.h"
diff --git a/Engine/Source/Runtime/Renderer/Private/MeshPassProcessor.cpp b/Engine/Source/Runtime/Renderer/Private/MeshPassProcessor.cpp
--- a/Engine/Source/Runtime/Renderer/Private/MeshPassProcessor.cpp
+++ b/Engine/Source/Runtime/Renderer/Private/MeshPassProcessor.cpp
@@ -1023,135 +1023,6 @@ void FGraphicsMinimalPipelineStateInitializer::ComputeStatePrecachePSOHash()
 	}
 }
 
-#if RHI_RAYTRACING
-
-void FRayTracingMeshCommand::SetRayTracingShaderBindingsForHitGroup(
-	FRayTracingLocalShaderBindingWriter* BindingWriter,
-	const TUniformBufferRef<FViewUniformShaderParameters>& ViewUniformBuffer,
-	FRHIUniformBuffer* SceneUniformBuffer,
-	FRHIUniformBuffer* NaniteUniformBuffer,
-	uint32 RecordIndex,
-	const FRHIRayTracingGeometry* RayTracingGeometry,
-	uint32 SegmentIndex,
-	uint32 HitGroupIndexInPipeline) const
-{
-	FRayTracingLocalShaderBindings* Bindings = ShaderBindings.SetRayTracingShaderBindingsForHitGroup(BindingWriter, RecordIndex, RayTracingGeometry, SegmentIndex, HitGroupIndexInPipeline);
-
-	if (ViewUniformBufferParameter.IsBound())
-	{
-		check(ViewUniformBuffer);
-		Bindings->UniformBuffers[ViewUniformBufferParameter.GetBaseIndex()] = ViewUniformBuffer;
-	}
-
-	if (SceneUniformBufferParameter.IsBound())
-	{
-		check(SceneUniformBuffer);
-		Bindings->UniformBuffers[SceneUniformBufferParameter.GetBaseIndex()] = SceneUniformBuffer;
-	}
-
-	if (NaniteUniformBufferParameter.IsBound())
-	{
-		check(NaniteUniformBuffer);
-		Bindings->UniformBuffers[NaniteUniformBufferParameter.GetBaseIndex()] = NaniteUniformBuffer;
-	}
-}
-
-void FRayTracingMeshCommand::SetRayTracingShaderBindingsForHitGroup(
-	FRayTracingLocalShaderBindingWriter* BindingWriter,
-	const TUniformBufferRef<FViewUniformShaderParameters>& ViewUniformBuffer,
-	FRHIUniformBuffer* SceneUniformBuffer,
-	FRHIUniformBuffer* NaniteUniformBuffer,
-	uint32 InstanceIndex,
-	uint32 SegmentIndex,
-	uint32 HitGroupIndexInPipeline,
-	uint32 ShaderSlot) const
-{
-	PRAGMA_DISABLE_DEPRECATION_WARNINGS
-	FRayTracingLocalShaderBindings* Bindings = ShaderBindings.SetRayTracingShaderBindingsForHitGroup(BindingWriter, InstanceIndex, SegmentIndex, HitGroupIndexInPipeline, ShaderSlot);
-	PRAGMA_ENABLE_DEPRECATION_WARNINGS
-
-	if (ViewUniformBufferParameter.IsBound())
-	{
-		check(ViewUniformBuffer);
-		Bindings->UniformBuffers[ViewUniformBufferParameter.GetBaseIndex()] = ViewUniformBuffer;
-	}
-
-	if (SceneUniformBufferParameter.IsBound())
-	{
-		check(SceneUniformBuffer);
-		Bindings->UniformBuffers[SceneUniformBufferParameter.GetBaseIndex()] = SceneUniformBuffer;
-	}
-
-	if (NaniteUniformBufferParameter.IsBound())
-	{
-		check(NaniteUniformBuffer);
-		Bindings->UniformBuffers[NaniteUniformBufferParameter.GetBaseIndex()] = NaniteUniformBuffer;
-	}
-}
-
-void FRayTracingMeshCommand::SetShader(const TShaderRef<FShader>& Shader)
-{
-	check(Shader.IsValid());
-	MaterialShaderIndex = Shader.GetRayTracingHitGroupLibraryIndex();
-	MaterialShader = Shader.GetRayTracingShader();
-	ViewUniformBufferParameter = Shader->GetUniformBufferParameter<FViewUniformShaderParameters>();
-	SceneUniformBufferParameter = Shader->GetUniformBufferParameter<FSceneUniformParameters>();
-	NaniteUniformBufferParameter = Shader->GetUniformBufferParameter<FNaniteRayTracingUniformParameters>();
-	ShaderBindings.Initialize(Shader);
-}
-
-void FRayTracingMeshCommand::SetShaders(const FMeshProcessorShaders& Shaders)
-{
-	SetShader(Shaders.RayTracingShader);
-}
-
-bool FRayTracingMeshCommand::IsUsingNaniteRayTracing() const
-{
-	return NaniteUniformBufferParameter.IsBound();
-}
-
-void FRayTracingShaderCommand::SetRayTracingShaderBindings(
-	FRayTracingLocalShaderBindingWriter* BindingWriter,
-	const TUniformBufferRef<FViewUniformShaderParameters>& ViewUniformBuffer,
-	FRHIUniformBuffer* SceneUniformBuffer,
-	FRHIUniformBuffer* NaniteUniformBuffer,
-	uint32 ShaderIndexInPipeline,
-	uint32 ShaderSlot) const
-{
-	FRayTracingLocalShaderBindings* Bindings = ShaderBindings.SetRayTracingShaderBindings(BindingWriter, ShaderIndexInPipeline, ShaderSlot);
-
-	if (ViewUniformBufferParameter.IsBound())
-	{
-		check(ViewUniformBuffer);
-		Bindings->UniformBuffers[ViewUniformBufferParameter.GetBaseIndex()] = ViewUniformBuffer;
-	}
-
-	if (SceneUniformBufferParameter.IsBound())
-	{
-		check(SceneUniformBuffer);
-		Bindings->UniformBuffers[SceneUniformBufferParameter.GetBaseIndex()] = SceneUniformBuffer;
-	}
-
-	if (NaniteUniformBufferParameter.IsBound())
-	{
-		check(NaniteUniformBuffer);
-		Bindings->UniformBuffers[NaniteUniformBufferParameter.GetBaseIndex()] = NaniteUniformBuffer;
-	}
-}
-
-void FRayTracingShaderCommand::SetShader(const TShaderRef<FShader>& InShader)
-{
-	check(InShader->GetFrequency() == SF_RayCallable || InShader->GetFrequency() == SF_RayMiss);
-	ShaderIndex = InShader.GetRayTracingCallableShaderLibraryIndex();
-	Shader = InShader.GetRayTracingShader();
-	ViewUniformBufferParameter = InShader->GetUniformBufferParameter<FViewUniformShaderParameters>();
-	SceneUniformBufferParameter = InShader->GetUniformBufferParameter<FSceneUniformParameters>();
-	NaniteUniformBufferParameter = InShader->GetUniformBufferParameter<FNaniteRayTracingUniformParameters>();
-
-	ShaderBindings.Initialize(InShader);
-}
-#endif // RHI_RAYTRACING
-
 void FMeshDrawCommand::SetDrawParametersAndFinalize(
 	const FMeshBatch& MeshBatch, 
 	int32 BatchElementIndex,
diff --git a/Engine/Source/Runtime/Renderer/Private/RayTracing/RayTracingMaterialHitShaders.h b/Engine/Source/Runtime/Renderer/Private/RayTracing/RayTracingMaterialHitShaders.h
--- a/Engine/Source/Runtime/Renderer/Private/RayTracing/RayTracingMaterialHitShaders.h
+++ b/Engine/Source/Runtime/Renderer/Private/RayTracing/RayTracingMaterialHitShaders.h
@@ -11,6 +11,7 @@
 #include "MaterialDomain.h"
 #include "MeshMaterialShader.h"
 #include "MeshPassProcessor.inl"
+#include "RayTracingMeshDrawCommands.h"
 #include "RayTracingInstanceMask.h"
 #include "RayTracingPayloadType.h"
 #include "ShaderParameterStruct.h"
diff --git a/Engine/Source/Runtime/Renderer/Private/RayTracing/RayTracingMeshDrawCommands.cpp b/Engine/Source/Runtime/Renderer/Private/RayTracing/RayTracingMeshDrawCommands.cpp
new file mode 100644
--- /dev/null
+++ b/Engine/Source/Runtime/Renderer/Private/RayTracing/RayTracingMeshDrawCommands.cpp
@@ -0,0 +1,138 @@
+// Copyright Epic Games, Inc. All Rights Reserved.
+
+/*=============================================================================
+	MeshPassProcessor.cpp: 
+=============================================================================*/
+
+#include "RayTracingMeshDrawCommands.h"
+#include "SceneUniformBuffer.h"
+#include "Nanite/NaniteShared.h"
+
+#if RHI_RAYTRACING
+
+void FRayTracingMeshCommand::SetRayTracingShaderBindingsForHitGroup(
+	FRayTracingLocalShaderBindingWriter* BindingWriter,
+	const TUniformBufferRef<FViewUniformShaderParameters>& ViewUniformBuffer,
+	FRHIUniformBuffer* SceneUniformBuffer,
+	FRHIUniformBuffer* NaniteUniformBuffer,
+	uint32 RecordIndex,
+	const FRHIRayTracingGeometry* RayTracingGeometry,
+	uint32 SegmentIndex,
+	uint32 HitGroupIndexInPipeline) const
+{
+	FRayTracingLocalShaderBindings* Bindings = ShaderBindings.SetRayTracingShaderBindingsForHitGroup(BindingWriter, RecordIndex, RayTracingGeometry, SegmentIndex, HitGroupIndexInPipeline);
+
+	if (ViewUniformBufferParameter.IsBound())
+	{
+		check(ViewUniformBuffer);
+		Bindings->UniformBuffers[ViewUniformBufferParameter.GetBaseIndex()] = ViewUniformBuffer;
+	}
+
+	if (SceneUniformBufferParameter.IsBound())
+	{
+		check(SceneUniformBuffer);
+		Bindings->UniformBuffers[SceneUniformBufferParameter.GetBaseIndex()] = SceneUniformBuffer;
+	}
+
+	if (NaniteUniformBufferParameter.IsBound())
+	{
+		check(NaniteUniformBuffer);
+		Bindings->UniformBuffers[NaniteUniformBufferParameter.GetBaseIndex()] = NaniteUniformBuffer;
+	}
+}
+
+void FRayTracingMeshCommand::SetRayTracingShaderBindingsForHitGroup(
+	FRayTracingLocalShaderBindingWriter* BindingWriter,
+	const TUniformBufferRef<FViewUniformShaderParameters>& ViewUniformBuffer,
+	FRHIUniformBuffer* SceneUniformBuffer,
+	FRHIUniformBuffer* NaniteUniformBuffer,
+	uint32 InstanceIndex,
+	uint32 SegmentIndex,
+	uint32 HitGroupIndexInPipeline,
+	uint32 ShaderSlot) const
+{
+	PRAGMA_DISABLE_DEPRECATION_WARNINGS
+	FRayTracingLocalShaderBindings* Bindings = ShaderBindings.SetRayTracingShaderBindingsForHitGroup(BindingWriter, InstanceIndex, SegmentIndex, HitGroupIndexInPipeline, ShaderSlot);
+	PRAGMA_ENABLE_DEPRECATION_WARNINGS
+
+	if (ViewUniformBufferParameter.IsBound())
+	{
+		check(ViewUniformBuffer);
+		Bindings->UniformBuffers[ViewUniformBufferParameter.GetBaseIndex()] = ViewUniformBuffer;
+	}
+
+	if (SceneUniformBufferParameter.IsBound())
+	{
+		check(SceneUniformBuffer);
+		Bindings->UniformBuffers[SceneUniformBufferParameter.GetBaseIndex()] = SceneUniformBuffer;
+	}
+
+	if (NaniteUniformBufferParameter.IsBound())
+	{
+		check(NaniteUniformBuffer);
+		Bindings->UniformBuffers[NaniteUniformBufferParameter.GetBaseIndex()] = NaniteUniformBuffer;
+	}
+}
+
+void FRayTracingMeshCommand::SetShader(const TShaderRef<FShader>& Shader)
+{
+	check(Shader.IsValid());
+	MaterialShaderIndex = Shader.GetRayTracingHitGroupLibraryIndex();
+	MaterialShader = Shader.GetRayTracingShader();
+	ViewUniformBufferParameter = Shader->GetUniformBufferParameter<FViewUniformShaderParameters>();
+	SceneUniformBufferParameter = Shader->GetUniformBufferParameter<FSceneUniformParameters>();
+	NaniteUniformBufferParameter = Shader->GetUniformBufferParameter<FNaniteRayTracingUniformParameters>();
+	ShaderBindings.Initialize(Shader);
+}
+
+void FRayTracingMeshCommand::SetShaders(const FMeshProcessorShaders& Shaders)
+{
+	SetShader(Shaders.RayTracingShader);
+}
+
+bool FRayTracingMeshCommand::IsUsingNaniteRayTracing() const
+{
+	return NaniteUniformBufferParameter.IsBound();
+}
+
+void FRayTracingShaderCommand::SetRayTracingShaderBindings(
+	FRayTracingLocalShaderBindingWriter* BindingWriter,
+	const TUniformBufferRef<FViewUniformShaderParameters>& ViewUniformBuffer,
+	FRHIUniformBuffer* SceneUniformBuffer,
+	FRHIUniformBuffer* NaniteUniformBuffer,
+	uint32 ShaderIndexInPipeline,
+	uint32 ShaderSlot) const
+{
+	FRayTracingLocalShaderBindings* Bindings = ShaderBindings.SetRayTracingShaderBindings(BindingWriter, ShaderIndexInPipeline, ShaderSlot);
+
+	if (ViewUniformBufferParameter.IsBound())
+	{
+		check(ViewUniformBuffer);
+		Bindings->UniformBuffers[ViewUniformBufferParameter.GetBaseIndex()] = ViewUniformBuffer;
+	}
+
+	if (SceneUniformBufferParameter.IsBound())
+	{
+		check(SceneUniformBuffer);
+		Bindings->UniformBuffers[SceneUniformBufferParameter.GetBaseIndex()] = SceneUniformBuffer;
+	}
+
+	if (NaniteUniformBufferParameter.IsBound())
+	{
+		check(NaniteUniformBuffer);
+		Bindings->UniformBuffers[NaniteUniformBufferParameter.GetBaseIndex()] = NaniteUniformBuffer;
+	}
+}
+
+void FRayTracingShaderCommand::SetShader(const TShaderRef<FShader>& InShader)
+{
+	check(InShader->GetFrequency() == SF_RayCallable || InShader->GetFrequency() == SF_RayMiss);
+	ShaderIndex = InShader.GetRayTracingCallableShaderLibraryIndex();
+	Shader = InShader.GetRayTracingShader();
+	ViewUniformBufferParameter = InShader->GetUniformBufferParameter<FViewUniformShaderParameters>();
+	SceneUniformBufferParameter = InShader->GetUniformBufferParameter<FSceneUniformParameters>();
+	NaniteUniformBufferParameter = InShader->GetUniformBufferParameter<FNaniteRayTracingUniformParameters>();
+
+	ShaderBindings.Initialize(InShader);
+}
+#endif // RHI_RAYTRACING
\ No newline at end of file
diff --git a/Engine/Source/Runtime/Renderer/Private/SceneRendering.h b/Engine/Source/Runtime/Renderer/Private/SceneRendering.h
--- a/Engine/Source/Runtime/Renderer/Private/SceneRendering.h
+++ b/Engine/Source/Runtime/Renderer/Private/SceneRendering.h
@@ -29,6 +29,7 @@
 #include "Templates/UniquePtr.h"
 #include "MeshDrawCommands.h"
 #include "MeshPassProcessor.h"
+#include "RayTracingMeshDrawCommands.h"
 #include "ShaderPrintParameters.h"
 #include "PostProcess/PostProcessAmbientOcclusionMobile.h"
 #include "VirtualShadowMaps/VirtualShadowMapArray.h"
diff --git a/Engine/Source/Runtime/Renderer/Public/MeshPassProcessor.h b/Engine/Source/Runtime/Renderer/Public/MeshPassProcessor.h
--- a/Engine/Source/Runtime/Renderer/Public/MeshPassProcessor.h
+++ b/Engine/Source/Runtime/Renderer/Public/MeshPassProcessor.h
@@ -2367,267 +2367,3 @@ inline FMeshDrawCommandSortKey CalculateMeshStaticSortKey(const TShaderRef<FMesh
 {
 	return CalculateMeshStaticSortKey(VertexShader.GetShader(), PixelShader.GetShader());
 }
-
-class FRayTracingMeshCommand
-{
-public:
-	FMeshDrawShaderBindings ShaderBindings;
-	FRHIRayTracingShader* MaterialShader = nullptr;
-
-	uint32 MaterialShaderIndex = UINT_MAX;
-	uint32 GeometrySegmentIndex = UINT_MAX;
-	uint8 InstanceMask = 0xFF;
-
-	bool bCastRayTracedShadows = true;
-	bool bOpaque = true;
-	bool bAlphaMasked = false;
-	bool bDecal = false;
-	bool bIsSky = false;
-	bool bIsTranslucent = false;
-	bool bTwoSided = false;
-	bool bReverseCulling = false;
-
-	RENDERER_API void SetRayTracingShaderBindingsForHitGroup(
-		FRayTracingLocalShaderBindingWriter* BindingWriter,
-		const TUniformBufferRef<FViewUniformShaderParameters>& ViewUniformBuffer,
-		FRHIUniformBuffer* SceneUniformBuffer,
-		FRHIUniformBuffer* NaniteUniformBuffer,
-		uint32 RecordIndex,
-		const FRHIRayTracingGeometry* RayTracingGeometry,
-		uint32 SegmentIndex,
-		uint32 HitGroupIndexInPipeline) const;
-
-	UE_DEPRECATED(5.5, "Provide RayTracingGeometry and GlobalSegmentIndex instead of InstanceIndex")
-	RENDERER_API void SetRayTracingShaderBindingsForHitGroup(
-		FRayTracingLocalShaderBindingWriter* BindingWriter,
-		const TUniformBufferRef<FViewUniformShaderParameters>& ViewUniformBuffer,
-		FRHIUniformBuffer* SceneUniformBuffer,
-		FRHIUniformBuffer* NaniteUniformBuffer,
-		uint32 InstanceIndex,
-		uint32 SegmentIndex,
-		uint32 HitGroupIndexInPipeline,
-		uint32 ShaderSlot) const;
-
-	/** Sets ray hit group shaders on the mesh command and allocates room for the shader bindings. */
-	RENDERER_API void SetShader(const TShaderRef<FShader>& Shader);
-
-	UE_DEPRECATED(5.4, "Use SetShader")
-	RENDERER_API void SetShaders(const FMeshProcessorShaders& Shaders);
-
-	RENDERER_API bool IsUsingNaniteRayTracing() const;
-private:
-	FShaderUniformBufferParameter ViewUniformBufferParameter;
-	FShaderUniformBufferParameter SceneUniformBufferParameter;
-	FShaderUniformBufferParameter NaniteUniformBufferParameter;
-};
-
-class FVisibleRayTracingMeshCommand
-{
-public:
-	FVisibleRayTracingMeshCommand(const FRayTracingMeshCommand* InRayTracingMeshCommand, const FRHIRayTracingGeometry* InRayTracingGeometry, uint32 InGlobalSegmentIndex, bool bInHidden = false)
-		: RayTracingMeshCommand(InRayTracingMeshCommand)
-		, RayTracingGeometry(InRayTracingGeometry)
-		, GlobalSegmentIndex(InGlobalSegmentIndex)
-		, InstanceIndex(INDEX_NONE)
-		, bHidden(bInHidden)
-	{
-		check(RayTracingGeometry != nullptr);
-		check(GlobalSegmentIndex != INDEX_NONE);
-	}
-
-	UE_DEPRECATED(5.5, "Provide RayTracingGeometry and GlobalSegmentIndex instead of InstanceIndex")
-	FVisibleRayTracingMeshCommand(const FRayTracingMeshCommand* InRayTracingMeshCommand, uint32 InInstanceIndex, bool bInHidden = false)
-		: RayTracingMeshCommand(InRayTracingMeshCommand)
-		, RayTracingGeometry(nullptr)
-		, GlobalSegmentIndex(INDEX_NONE)
-		, InstanceIndex(InInstanceIndex)
-		, bHidden(bInHidden)
-	{
-PRAGMA_DISABLE_DEPRECATION_WARNINGS
-		check(InstanceIndex != INDEX_NONE);
-PRAGMA_ENABLE_DEPRECATION_WARNINGS
-	}
-
-	const FRayTracingMeshCommand* RayTracingMeshCommand;
-	const FRHIRayTracingGeometry* RayTracingGeometry;
-	uint32 GlobalSegmentIndex;
-	UE_DEPRECATED(5.5, "Provide RayTracingGeometry and GlobalSegmentIndex instead of InstanceIndex")
-	uint32 InstanceIndex;
-	bool bHidden;
-	
-	PRAGMA_DISABLE_DEPRECATION_WARNINGS
-	FVisibleRayTracingMeshCommand(const FVisibleRayTracingMeshCommand&) = default;
-	FVisibleRayTracingMeshCommand& operator=(const FVisibleRayTracingMeshCommand&) = default;
-	FVisibleRayTracingMeshCommand(FVisibleRayTracingMeshCommand&&) = default;
-	FVisibleRayTracingMeshCommand& operator=(FVisibleRayTracingMeshCommand&&) = default;
-	PRAGMA_ENABLE_DEPRECATION_WARNINGS
-};
-
-template <>
-struct TUseBitwiseSwap<FVisibleRayTracingMeshCommand>
-{
-	// Prevent Memcpy call overhead during FVisibleRayTracingMeshCommand sorting
-	enum { Value = false };
-};
-
-typedef TArray<FVisibleRayTracingMeshCommand> FRayTracingMeshCommandOneFrameArray;
-
-class FRayTracingMeshCommandContext
-{
-public:
-
-	virtual ~FRayTracingMeshCommandContext() {}
-
-	virtual FRayTracingMeshCommand& AddCommand(const FRayTracingMeshCommand& Initializer) = 0;
-
-	virtual void FinalizeCommand(FRayTracingMeshCommand& RayTracingMeshCommand) = 0;
-};
-
-using FTempRayTracingMeshCommandStorage = TArray<FRayTracingMeshCommand>;
-
-using FCachedRayTracingMeshCommandStorage = TSparseArray<FRayTracingMeshCommand>;
-
-using FDynamicRayTracingMeshCommandStorage = TChunkedArray<FRayTracingMeshCommand>;
-
-template<class T>
-class FCachedRayTracingMeshCommandContext : public FRayTracingMeshCommandContext
-{
-public:
-	FCachedRayTracingMeshCommandContext(T& InDrawListStorage) : DrawListStorage(InDrawListStorage) {}
-
-	virtual FRayTracingMeshCommand& AddCommand(const FRayTracingMeshCommand& Initializer) override final
-	{
-		CommandIndex = DrawListStorage.Add(Initializer);
-		return DrawListStorage[CommandIndex];
-	}
-
-	virtual void FinalizeCommand(FRayTracingMeshCommand& RayTracingMeshCommand) override final {}
-
-	int32 CommandIndex = -1;
-
-private:
-	T& DrawListStorage;
-};
-
-class FDynamicRayTracingMeshCommandContext : public FRayTracingMeshCommandContext
-{
-public:
-	FDynamicRayTracingMeshCommandContext
-	(
-		FDynamicRayTracingMeshCommandStorage& InDynamicCommandStorage,
-		FRayTracingMeshCommandOneFrameArray& InVisibleCommands,
-		const FRHIRayTracingGeometry* InRayTracingGeometry,
-		uint32 InGeometrySegmentIndex,
-		uint32 InGlobalSegmentIndex,
-		uint32 InDecalGlobalSegmentIndex = INDEX_NONE
-	) :
-		DynamicCommandStorage(InDynamicCommandStorage),
-		VisibleCommands(InVisibleCommands),
-		RayTracingGeometry(InRayTracingGeometry),
-		GeometrySegmentIndex(InGeometrySegmentIndex),
-		GlobalSegmentIndex(InGlobalSegmentIndex),
-		DecalGlobalSegmentIndex(InDecalGlobalSegmentIndex),
-		RayTracingInstanceIndex(INDEX_NONE),
-		RayTracingDecalInstanceIndex(INDEX_NONE)
-	{}
-
-	UE_DEPRECATED(5.5, "Provide RayTracingGeometry and GlobalGeometrySegmentIndex and DecalGlobalGeometrySegmentIndex instead")
-	FDynamicRayTracingMeshCommandContext
-	(
-		FDynamicRayTracingMeshCommandStorage& InDynamicCommandStorage,
-		FRayTracingMeshCommandOneFrameArray& InVisibleCommands,
-		uint32 InGeometrySegmentIndex,
-		uint32 InRayTracingInstanceIndex,
-		uint32 InRayTracingDecalInstanceIndex = INDEX_NONE
-	) :
-		DynamicCommandStorage(InDynamicCommandStorage),
-		VisibleCommands(InVisibleCommands),
-		RayTracingGeometry(nullptr),
-		GeometrySegmentIndex(InGeometrySegmentIndex),
-		GlobalSegmentIndex(INDEX_NONE),
-		DecalGlobalSegmentIndex(INDEX_NONE),
-		RayTracingInstanceIndex(InRayTracingInstanceIndex),
-		RayTracingDecalInstanceIndex(InRayTracingDecalInstanceIndex)
-	{}
-
-	virtual FRayTracingMeshCommand& AddCommand(const FRayTracingMeshCommand& Initializer) override final
-	{
-		const int32 Index = DynamicCommandStorage.AddElement(Initializer);
-		FRayTracingMeshCommand& NewCommand = DynamicCommandStorage[Index];
-		NewCommand.GeometrySegmentIndex = GeometrySegmentIndex;
-		return NewCommand;
-	}
-
-	virtual void FinalizeCommand(FRayTracingMeshCommand& RayTracingMeshCommand) override final
-	{
-		if (GlobalSegmentIndex != INDEX_NONE)
-		{
-			const bool bHidden = RayTracingMeshCommand.bDecal;
-			FVisibleRayTracingMeshCommand NewVisibleMeshCommand(&RayTracingMeshCommand, RayTracingGeometry, GlobalSegmentIndex + GeometrySegmentIndex, bHidden);
-			VisibleCommands.Add(NewVisibleMeshCommand);
-		}
-
-		if (DecalGlobalSegmentIndex != INDEX_NONE)
-		{
-			const bool bHidden = !RayTracingMeshCommand.bDecal;
-			FVisibleRayTracingMeshCommand NewVisibleMeshCommand(&RayTracingMeshCommand, RayTracingGeometry, DecalGlobalSegmentIndex + GeometrySegmentIndex, bHidden);
-			VisibleCommands.Add(NewVisibleMeshCommand);
-		}
-
-		PRAGMA_DISABLE_DEPRECATION_WARNINGS
-		if(RayTracingInstanceIndex != INDEX_NONE)
-		{
-			const bool bHidden = RayTracingMeshCommand.bDecal;
-			FVisibleRayTracingMeshCommand NewVisibleMeshCommand(&RayTracingMeshCommand, RayTracingInstanceIndex, bHidden);
-			VisibleCommands.Add(NewVisibleMeshCommand);
-		}
-
-		if (RayTracingDecalInstanceIndex != INDEX_NONE)
-		{
-			const bool bHidden = !RayTracingMeshCommand.bDecal;
-			FVisibleRayTracingMeshCommand NewVisibleMeshCommand(&RayTracingMeshCommand, RayTracingDecalInstanceIndex, bHidden);
-			VisibleCommands.Add(NewVisibleMeshCommand);
-		}
-		PRAGMA_ENABLE_DEPRECATION_WARNINGS
-	}
-
-private:
-	FDynamicRayTracingMeshCommandStorage& DynamicCommandStorage;
-	FRayTracingMeshCommandOneFrameArray& VisibleCommands;
-
-	const FRHIRayTracingGeometry* RayTracingGeometry;
-	uint32 GeometrySegmentIndex;
-	uint32 GlobalSegmentIndex;
-	uint32 DecalGlobalSegmentIndex;
-
-	UE_DEPRECATED(5.5, "Provide RayTracingGeometry and GlobalSegmentIndex instead of RayTracingInstanceIndex")
-	uint32 RayTracingInstanceIndex;
-	UE_DEPRECATED(5.5, "Provide RayTracingGeometry and DecalGlobalGeometrySegmentIndex instead of RayTracingDecalInstanceIndex")
-	uint32 RayTracingDecalInstanceIndex;
-};
-
-class FRayTracingShaderCommand
-{
-public:
-	FMeshDrawShaderBindings ShaderBindings;
-	FRHIRayTracingShader* Shader = nullptr;
-
-	uint32 ShaderIndex = UINT_MAX;
-	uint32 SlotInScene = UINT_MAX;
-
-	RENDERER_API void SetRayTracingShaderBindings(
-		FRayTracingLocalShaderBindingWriter* BindingWriter,
-		const TUniformBufferRef<FViewUniformShaderParameters>& ViewUniformBuffer,
-		FRHIUniformBuffer* SceneUniformBuffer,
-		FRHIUniformBuffer* NaniteUniformBuffer,
-		uint32 ShaderIndexInPipeline,
-		uint32 ShaderSlot) const;
-
-	/** Sets ray tracing shader on the command and allocates room for the shader bindings. */
-	RENDERER_API void SetShader(const TShaderRef<FShader>& Shader);
-
-private:
-	FShaderUniformBufferParameter ViewUniformBufferParameter;
-	FShaderUniformBufferParameter SceneUniformBufferParameter;
-	FShaderUniformBufferParameter NaniteUniformBufferParameter;
-};
diff --git a/Engine/Source/Runtime/Renderer/Public/RayTracingMeshDrawCommands.h b/Engine/Source/Runtime/Renderer/Public/RayTracingMeshDrawCommands.h
new file mode 100644
--- /dev/null
+++ b/Engine/Source/Runtime/Renderer/Public/RayTracingMeshDrawCommands.h
@@ -0,0 +1,273 @@
+// Copyright Epic Games, Inc. All Rights Reserved.
+
+/*=============================================================================
+	MeshPassProcessor.h
+=============================================================================*/
+
+#pragma once
+
+#include "MeshPassProcessor.h"
+
+class FRayTracingMeshCommand
+{
+public:
+	FMeshDrawShaderBindings ShaderBindings;
+	FRHIRayTracingShader* MaterialShader = nullptr;
+
+	uint32 MaterialShaderIndex = UINT_MAX;
+	uint32 GeometrySegmentIndex = UINT_MAX;
+	uint8 InstanceMask = 0xFF;
+
+	bool bCastRayTracedShadows = true;
+	bool bOpaque = true;
+	bool bAlphaMasked = false;
+	bool bDecal = false;
+	bool bIsSky = false;
+	bool bIsTranslucent = false;
+	bool bTwoSided = false;
+	bool bReverseCulling = false;
+
+	RENDERER_API void SetRayTracingShaderBindingsForHitGroup(
+		FRayTracingLocalShaderBindingWriter* BindingWriter,
+		const TUniformBufferRef<FViewUniformShaderParameters>& ViewUniformBuffer,
+		FRHIUniformBuffer* SceneUniformBuffer,
+		FRHIUniformBuffer* NaniteUniformBuffer,
+		uint32 RecordIndex,
+		const FRHIRayTracingGeometry* RayTracingGeometry,
+		uint32 SegmentIndex,
+		uint32 HitGroupIndexInPipeline) const;
+
+	UE_DEPRECATED(5.5, "Provide RayTracingGeometry and GlobalSegmentIndex instead of InstanceIndex")
+	RENDERER_API void SetRayTracingShaderBindingsForHitGroup(
+		FRayTracingLocalShaderBindingWriter* BindingWriter,
+		const TUniformBufferRef<FViewUniformShaderParameters>& ViewUniformBuffer,
+		FRHIUniformBuffer* SceneUniformBuffer,
+		FRHIUniformBuffer* NaniteUniformBuffer,
+		uint32 InstanceIndex,
+		uint32 SegmentIndex,
+		uint32 HitGroupIndexInPipeline,
+		uint32 ShaderSlot) const;
+
+	/** Sets ray hit group shaders on the mesh command and allocates room for the shader bindings. */
+	RENDERER_API void SetShader(const TShaderRef<FShader>& Shader);
+
+	UE_DEPRECATED(5.4, "Use SetShader")
+	RENDERER_API void SetShaders(const FMeshProcessorShaders& Shaders);
+
+	RENDERER_API bool IsUsingNaniteRayTracing() const;
+private:
+	FShaderUniformBufferParameter ViewUniformBufferParameter;
+	FShaderUniformBufferParameter SceneUniformBufferParameter;
+	FShaderUniformBufferParameter NaniteUniformBufferParameter;
+};
+
+class FVisibleRayTracingMeshCommand
+{
+public:
+	FVisibleRayTracingMeshCommand(const FRayTracingMeshCommand* InRayTracingMeshCommand, const FRHIRayTracingGeometry* InRayTracingGeometry, uint32 InGlobalSegmentIndex, bool bInHidden = false)
+		: RayTracingMeshCommand(InRayTracingMeshCommand)
+		, RayTracingGeometry(InRayTracingGeometry)
+		, GlobalSegmentIndex(InGlobalSegmentIndex)
+		, InstanceIndex(INDEX_NONE)
+		, bHidden(bInHidden)
+	{
+		check(RayTracingGeometry != nullptr);
+		check(GlobalSegmentIndex != INDEX_NONE);
+	}
+
+	UE_DEPRECATED(5.5, "Provide RayTracingGeometry and GlobalSegmentIndex instead of InstanceIndex")
+	FVisibleRayTracingMeshCommand(const FRayTracingMeshCommand* InRayTracingMeshCommand, uint32 InInstanceIndex, bool bInHidden = false)
+		: RayTracingMeshCommand(InRayTracingMeshCommand)
+		, RayTracingGeometry(nullptr)
+		, GlobalSegmentIndex(INDEX_NONE)
+		, InstanceIndex(InInstanceIndex)
+		, bHidden(bInHidden)
+	{
+PRAGMA_DISABLE_DEPRECATION_WARNINGS
+		check(InstanceIndex != INDEX_NONE);
+PRAGMA_ENABLE_DEPRECATION_WARNINGS
+	}
+
+	const FRayTracingMeshCommand* RayTracingMeshCommand;
+	const FRHIRayTracingGeometry* RayTracingGeometry;
+	uint32 GlobalSegmentIndex;
+	UE_DEPRECATED(5.5, "Provide RayTracingGeometry and GlobalSegmentIndex instead of InstanceIndex")
+	uint32 InstanceIndex;
+	bool bHidden;
+	
+	PRAGMA_DISABLE_DEPRECATION_WARNINGS
+	FVisibleRayTracingMeshCommand(const FVisibleRayTracingMeshCommand&) = default;
+	FVisibleRayTracingMeshCommand& operator=(const FVisibleRayTracingMeshCommand&) = default;
+	FVisibleRayTracingMeshCommand(FVisibleRayTracingMeshCommand&&) = default;
+	FVisibleRayTracingMeshCommand& operator=(FVisibleRayTracingMeshCommand&&) = default;
+	PRAGMA_ENABLE_DEPRECATION_WARNINGS
+};
+
+template <>
+struct TUseBitwiseSwap<FVisibleRayTracingMeshCommand>
+{
+	// Prevent Memcpy call overhead during FVisibleRayTracingMeshCommand sorting
+	enum { Value = false };
+};
+
+typedef TArray<FVisibleRayTracingMeshCommand> FRayTracingMeshCommandOneFrameArray;
+
+class FRayTracingMeshCommandContext
+{
+public:
+
+	virtual ~FRayTracingMeshCommandContext() {}
+
+	virtual FRayTracingMeshCommand& AddCommand(const FRayTracingMeshCommand& Initializer) = 0;
+
+	virtual void FinalizeCommand(FRayTracingMeshCommand& RayTracingMeshCommand) = 0;
+};
+
+using FTempRayTracingMeshCommandStorage = TArray<FRayTracingMeshCommand>;
+
+using FCachedRayTracingMeshCommandStorage = TSparseArray<FRayTracingMeshCommand>;
+
+using FDynamicRayTracingMeshCommandStorage = TChunkedArray<FRayTracingMeshCommand>;
+
+template<class T>
+class FCachedRayTracingMeshCommandContext : public FRayTracingMeshCommandContext
+{
+public:
+	FCachedRayTracingMeshCommandContext(T& InDrawListStorage) : DrawListStorage(InDrawListStorage) {}
+
+	virtual FRayTracingMeshCommand& AddCommand(const FRayTracingMeshCommand& Initializer) override final
+	{
+		CommandIndex = DrawListStorage.Add(Initializer);
+		return DrawListStorage[CommandIndex];
+	}
+
+	virtual void FinalizeCommand(FRayTracingMeshCommand& RayTracingMeshCommand) override final {}
+
+	int32 CommandIndex = -1;
+
+private:
+	T& DrawListStorage;
+};
+
+class FDynamicRayTracingMeshCommandContext : public FRayTracingMeshCommandContext
+{
+public:
+	FDynamicRayTracingMeshCommandContext
+	(
+		FDynamicRayTracingMeshCommandStorage& InDynamicCommandStorage,
+		FRayTracingMeshCommandOneFrameArray& InVisibleCommands,
+		const FRHIRayTracingGeometry* InRayTracingGeometry,
+		uint32 InGeometrySegmentIndex,
+		uint32 InGlobalSegmentIndex,
+		uint32 InDecalGlobalSegmentIndex = INDEX_NONE
+	) :
+		DynamicCommandStorage(InDynamicCommandStorage),
+		VisibleCommands(InVisibleCommands),
+		RayTracingGeometry(InRayTracingGeometry),
+		GeometrySegmentIndex(InGeometrySegmentIndex),
+		GlobalSegmentIndex(InGlobalSegmentIndex),
+		DecalGlobalSegmentIndex(InDecalGlobalSegmentIndex),
+		RayTracingInstanceIndex(INDEX_NONE),
+		RayTracingDecalInstanceIndex(INDEX_NONE)
+	{}
+
+	UE_DEPRECATED(5.5, "Provide RayTracingGeometry and GlobalGeometrySegmentIndex and DecalGlobalGeometrySegmentIndex instead")
+	FDynamicRayTracingMeshCommandContext
+	(
+		FDynamicRayTracingMeshCommandStorage& InDynamicCommandStorage,
+		FRayTracingMeshCommandOneFrameArray& InVisibleCommands,
+		uint32 InGeometrySegmentIndex,
+		uint32 InRayTracingInstanceIndex,
+		uint32 InRayTracingDecalInstanceIndex = INDEX_NONE
+	) :
+		DynamicCommandStorage(InDynamicCommandStorage),
+		VisibleCommands(InVisibleCommands),
+		RayTracingGeometry(nullptr),
+		GeometrySegmentIndex(InGeometrySegmentIndex),
+		GlobalSegmentIndex(INDEX_NONE),
+		DecalGlobalSegmentIndex(INDEX_NONE),
+		RayTracingInstanceIndex(InRayTracingInstanceIndex),
+		RayTracingDecalInstanceIndex(InRayTracingDecalInstanceIndex)
+	{}
+
+	virtual FRayTracingMeshCommand& AddCommand(const FRayTracingMeshCommand& Initializer) override final
+	{
+		const int32 Index = DynamicCommandStorage.AddElement(Initializer);
+		FRayTracingMeshCommand& NewCommand = DynamicCommandStorage[Index];
+		NewCommand.GeometrySegmentIndex = GeometrySegmentIndex;
+		return NewCommand;
+	}
+
+	virtual void FinalizeCommand(FRayTracingMeshCommand& RayTracingMeshCommand) override final
+	{
+		if (GlobalSegmentIndex != INDEX_NONE)
+		{
+			const bool bHidden = RayTracingMeshCommand.bDecal;
+			FVisibleRayTracingMeshCommand NewVisibleMeshCommand(&RayTracingMeshCommand, RayTracingGeometry, GlobalSegmentIndex + GeometrySegmentIndex, bHidden);
+			VisibleCommands.Add(NewVisibleMeshCommand);
+		}
+
+		if (DecalGlobalSegmentIndex != INDEX_NONE)
+		{
+			const bool bHidden = !RayTracingMeshCommand.bDecal;
+			FVisibleRayTracingMeshCommand NewVisibleMeshCommand(&RayTracingMeshCommand, RayTracingGeometry, DecalGlobalSegmentIndex + GeometrySegmentIndex, bHidden);
+			VisibleCommands.Add(NewVisibleMeshCommand);
+		}
+
+		PRAGMA_DISABLE_DEPRECATION_WARNINGS
+		if(RayTracingInstanceIndex != INDEX_NONE)
+		{
+			const bool bHidden = RayTracingMeshCommand.bDecal;
+			FVisibleRayTracingMeshCommand NewVisibleMeshCommand(&RayTracingMeshCommand, RayTracingInstanceIndex, bHidden);
+			VisibleCommands.Add(NewVisibleMeshCommand);
+		}
+
+		if (RayTracingDecalInstanceIndex != INDEX_NONE)
+		{
+			const bool bHidden = !RayTracingMeshCommand.bDecal;
+			FVisibleRayTracingMeshCommand NewVisibleMeshCommand(&RayTracingMeshCommand, RayTracingDecalInstanceIndex, bHidden);
+			VisibleCommands.Add(NewVisibleMeshCommand);
+		}
+		PRAGMA_ENABLE_DEPRECATION_WARNINGS
+	}
+
+private:
+	FDynamicRayTracingMeshCommandStorage& DynamicCommandStorage;
+	FRayTracingMeshCommandOneFrameArray& VisibleCommands;
+
+	const FRHIRayTracingGeometry* RayTracingGeometry;
+	uint32 GeometrySegmentIndex;
+	uint32 GlobalSegmentIndex;
+	uint32 DecalGlobalSegmentIndex;
+
+	UE_DEPRECATED(5.5, "Provide RayTracingGeometry and GlobalSegmentIndex instead of RayTracingInstanceIndex")
+	uint32 RayTracingInstanceIndex;
+	UE_DEPRECATED(5.5, "Provide RayTracingGeometry and DecalGlobalGeometrySegmentIndex instead of RayTracingDecalInstanceIndex")
+	uint32 RayTracingDecalInstanceIndex;
+};
+
+class FRayTracingShaderCommand
+{
+public:
+	FMeshDrawShaderBindings ShaderBindings;
+	FRHIRayTracingShader* Shader = nullptr;
+
+	uint32 ShaderIndex = UINT_MAX;
+	uint32 SlotInScene = UINT_MAX;
+
+	RENDERER_API void SetRayTracingShaderBindings(
+		FRayTracingLocalShaderBindingWriter* BindingWriter,
+		const TUniformBufferRef<FViewUniformShaderParameters>& ViewUniformBuffer,
+		FRHIUniformBuffer* SceneUniformBuffer,
+		FRHIUniformBuffer* NaniteUniformBuffer,
+		uint32 ShaderIndexInPipeline,
+		uint32 ShaderSlot) const;
+
+	/** Sets ray tracing shader on the command and allocates room for the shader bindings. */
+	RENDERER_API void SetShader(const TShaderRef<FShader>& Shader);
+
+private:
+	FShaderUniformBufferParameter ViewUniformBufferParameter;
+	FShaderUniformBufferParameter SceneUniformBufferParameter;
+	FShaderUniformBufferParameter NaniteUniformBufferParameter;
+};
