Remove SceneViewFamily parameter from BuildRayTracingInstanceMaskAndFlags(...)/GetSceneProxyRayTracingMaskInfo(...)

diff --git a/Engine/Source/Runtime/Renderer/Private/PrimitiveSceneInfo.cpp b/Engine/Source/Runtime/Renderer/Private/PrimitiveSceneInfo.cpp
--- a/Engine/Source/Runtime/Renderer/Private/PrimitiveSceneInfo.cpp
+++ b/Engine/Source/Runtime/Renderer/Private/PrimitiveSceneInfo.cpp
@@ -1267,7 +1267,7 @@ void FPrimitiveSceneInfo::UpdateCachedRayTracingInstance(FPrimitiveSceneInfo* Sc
 		}
 		else
 		{
-			InstanceMaskAndFlags = BuildRayTracingInstanceMaskAndFlags(CachedRayTracingInstance, *SceneProxy, nullptr);
+			InstanceMaskAndFlags = BuildRayTracingInstanceMaskAndFlags(CachedRayTracingInstance, *SceneProxy);
 		}
 
 		SceneInfo->CachedRayTracingInstance.Mask = InstanceMaskAndFlags.Mask; // When no cached command is found, InstanceMask == 0 and the instance is effectively filtered out
diff --git a/Engine/Source/Runtime/Renderer/Private/RayTracing/RayTracing.cpp b/Engine/Source/Runtime/Renderer/Private/RayTracing/RayTracing.cpp
--- a/Engine/Source/Runtime/Renderer/Private/RayTracing/RayTracing.cpp
+++ b/Engine/Source/Runtime/Renderer/Private/RayTracing/RayTracing.cpp
@@ -1009,7 +1009,7 @@ namespace RayTracing
 							}
 							else
 							{
-								InstanceMaskAndFlags = BuildRayTracingInstanceMaskAndFlags(Instance, *SceneProxy, nullptr);
+								InstanceMaskAndFlags = BuildRayTracingInstanceMaskAndFlags(Instance, *SceneProxy);
 							}
 
 							SceneInfo->CachedRayTracingInstance.Mask = InstanceMaskAndFlags.Mask; // When no cached command is found, InstanceMask == 0 and the instance is effectively filtered out
diff --git a/Engine/Source/Runtime/Renderer/Private/RayTracing/RayTracingInstanceMask.cpp b/Engine/Source/Runtime/Renderer/Private/RayTracing/RayTracingInstanceMask.cpp
--- a/Engine/Source/Runtime/Renderer/Private/RayTracing/RayTracingInstanceMask.cpp
+++ b/Engine/Source/Runtime/Renderer/Private/RayTracing/RayTracingInstanceMask.cpp
@@ -111,31 +111,15 @@ uint8 BlendModeToRayTracingInstanceMask(const EBlendMode BlendMode, bool bCastSh
 	return InstanceMask;
 }
 
-FSceneProxyRayTracingMaskInfo GetSceneProxyRayTracingMaskInfo(const FPrimitiveSceneProxy& PrimitiveSceneProxy, const FSceneViewFamily* SceneViewFamily)
+FSceneProxyRayTracingMaskInfo GetSceneProxyRayTracingMaskInfo(const FPrimitiveSceneProxy& PrimitiveSceneProxy)
 {
-
 	bool bAffectsIndirectLightingOnly = PrimitiveSceneProxy.AffectsIndirectLightingWhileHidden() && !PrimitiveSceneProxy.IsDrawnInGame();
 	bool bCastHiddenShadow = PrimitiveSceneProxy.CastsHiddenShadow() && !PrimitiveSceneProxy.IsDrawnInGame();
 	bool bAffectsDynamicIndirectLighting = PrimitiveSceneProxy.AffectsDynamicIndirectLighting();
 
-	ERayTracingViewMaskMode MaskMode = ERayTracingViewMaskMode::RayTracing;
+	const FScene* RenderScene = PrimitiveSceneProxy.GetScene().GetRenderScene();
 
-	if (SceneViewFamily)
-	{
-		if (SceneViewFamily->EngineShowFlags.PathTracing)
-		{
-			MaskMode = ERayTracingViewMaskMode::PathTracing;
-		}
-	}
-	else
-	{
-		FScene* RenderScene = PrimitiveSceneProxy.GetScene().GetRenderScene();
-
-		if (RenderScene)
-		{
-			MaskMode = static_cast<ERayTracingViewMaskMode>(RenderScene->CachedRayTracingMeshCommandsMode);
-		}
-	}
+	ERayTracingViewMaskMode MaskMode = static_cast<ERayTracingViewMaskMode>(RenderScene->CachedRayTracingMeshCommandsMode);
 
 	return {bAffectsIndirectLightingOnly, bCastHiddenShadow, bAffectsDynamicIndirectLighting, MaskMode};
 }
@@ -240,9 +224,9 @@ FRayTracingMaskAndFlags BuildRayTracingInstanceMaskAndFlags(TArrayView<const FMe
 	return Result;
 }
 
-FRayTracingMaskAndFlags BuildRayTracingInstanceMaskAndFlags(const FRayTracingInstance& Instance, const FPrimitiveSceneProxy& PrimitiveSceneProxy, const FSceneViewFamily* SceneViewFamily)
+FRayTracingMaskAndFlags BuildRayTracingInstanceMaskAndFlags(const FRayTracingInstance& Instance, const FPrimitiveSceneProxy& PrimitiveSceneProxy)
 {
-	FSceneProxyRayTracingMaskInfo MaskInfo = GetSceneProxyRayTracingMaskInfo(PrimitiveSceneProxy, SceneViewFamily);
+	FSceneProxyRayTracingMaskInfo MaskInfo = GetSceneProxyRayTracingMaskInfo(PrimitiveSceneProxy);
 
 	const TArrayView<const FMeshBatch> MeshBatches = Instance.GetMaterials();
 
@@ -279,7 +263,7 @@ void SetupRayTracingMeshCommandMaskAndStatus(FRayTracingMeshCommand& MeshCommand
 	// In the case that that this mesh command is not associated with a primitive, the mesh batch value will still apply.
 	MeshCommand.bReverseCulling = PrimitiveSceneProxy->IsCullingReversedByComponent();
 
-	FSceneProxyRayTracingMaskInfo MaskInfo = GetSceneProxyRayTracingMaskInfo(*PrimitiveSceneProxy, nullptr);
+	FSceneProxyRayTracingMaskInfo MaskInfo = GetSceneProxyRayTracingMaskInfo(*PrimitiveSceneProxy);
 
 	// TODO: This should be done once all mesh commands for a mesh are combined (similar to BuildRayTracingInstanceMaskAndFlags(...) above)
 	if (MaskMode == ERayTracingViewMaskMode::PathTracing || MaskMode == ERayTracingViewMaskMode::LightMapTracing)
diff --git a/Engine/Source/Runtime/Renderer/Private/RayTracing/RayTracingInstanceMask.h b/Engine/Source/Runtime/Renderer/Private/RayTracing/RayTracingInstanceMask.h
--- a/Engine/Source/Runtime/Renderer/Private/RayTracing/RayTracingInstanceMask.h
+++ b/Engine/Source/Runtime/Renderer/Private/RayTracing/RayTracingInstanceMask.h
@@ -48,7 +48,7 @@ struct FSceneProxyRayTracingMaskInfo
 	ERayTracingViewMaskMode MaskMode = ERayTracingViewMaskMode::RayTracing;
 };
 
-FSceneProxyRayTracingMaskInfo GetSceneProxyRayTracingMaskInfo(const FPrimitiveSceneProxy& PrimitiveSceneProxy, const FSceneViewFamily* SceneView);
+FSceneProxyRayTracingMaskInfo GetSceneProxyRayTracingMaskInfo(const FPrimitiveSceneProxy& PrimitiveSceneProxy);
 
 
 //-------------------------------------------------------
@@ -59,7 +59,7 @@ FRayTracingMaskAndFlags BuildRayTracingInstanceMaskAndFlags(TArrayView<const FMe
 	const FSceneProxyRayTracingMaskInfo& SceneProxyRayTracingMaskInfo, ERayTracingInstanceLayer InstanceLayer, uint8 ExtraMask = 0);
 
 // Build mask and flags without modification of RayTracingInstance
-FRayTracingMaskAndFlags BuildRayTracingInstanceMaskAndFlags(const FRayTracingInstance& Instance, const FPrimitiveSceneProxy& PrimitiveSceneProxy, const FSceneViewFamily* SceneViewFamily);
+FRayTracingMaskAndFlags BuildRayTracingInstanceMaskAndFlags(const FRayTracingInstance& Instance, const FPrimitiveSceneProxy& PrimitiveSceneProxy);
 
 //-------------------------------------------------------
 //	FRayTracingMeshCommand related mask setup and update
