Mobile: Added rect lights support to a deferred shading
r.Mobile.Forward.RenderRectLightsAsSpotLights will only affect lights that use LightGrid, translucency or clustered

diff --git a/Engine/Shaders/Private/MobileDeferredShading.usf b/Engine/Shaders/Private/MobileDeferredShading.usf
--- a/Engine/Shaders/Private/MobileDeferredShading.usf
+++ b/Engine/Shaders/Private/MobileDeferredShading.usf
@@ -29,10 +29,6 @@
 #endif
 #include "HairStrands/HairStrandsEnvironmentLightingCommon.ush"
 
-#ifndef IS_SPOT_LIGHT 
-#define IS_SPOT_LIGHT 0 
-#endif 
-
 float4x4 TranslatedWorldToLight;
 float2 LightFunctionParameters2;
 
@@ -277,7 +273,7 @@ void MobileRadialLightPS(
 	half3 CameraVector = normalize(TranslatedWorldPosition);
 			
 	FLightAccumulator DirectLighting = (FLightAccumulator)0;
-	FDeferredLightData LightData = InitDeferredLightFromLightParameters(GetRootLightShaderParameters(), IS_SPOT_LIGHT ? LIGHT_TYPE_SPOT : LIGHT_TYPE_POINT);
+	FDeferredLightData LightData = InitDeferredLightFromLightParameters(GetRootLightShaderParameters(), RADIAL_LIGHT_TYPE);
 		
 	// Affect the light color by any other shadow/transmittance before the lighting is evaluated
 	{
diff --git a/Engine/Source/Runtime/Engine/Private/Components/RectLightComponent.cpp b/Engine/Source/Runtime/Engine/Private/Components/RectLightComponent.cpp
--- a/Engine/Source/Runtime/Engine/Private/Components/RectLightComponent.cpp
+++ b/Engine/Source/Runtime/Engine/Private/Components/RectLightComponent.cpp
@@ -201,16 +201,6 @@ void URectLightComponent::PostEditChangeProperty(FPropertyChangedEvent& Property
 }
 #endif // WITH_EDITOR
 
-static bool RenderRectLightsAsSpotLights(FStaticFeatureLevel FeatureLevel)
-{
-	if (FeatureLevel == ERHIFeatureLevel::ES3_1)
-	{
-		static const auto CVarRenderRectLightAsSpotLight = IConsoleManager::Get().FindTConsoleVariableDataInt(TEXT("r.Mobile.Forward.RenderRectLightsAsSpotLights"));
-		return (CVarRenderRectLightAsSpotLight && CVarRenderRectLightAsSpotLight->GetValueOnAnyThread() != 0);
-	}
-	return false;
-}
-
 FRectLightSceneProxy::FRectLightSceneProxy(const URectLightComponent* Component)
 	: FLocalLightSceneProxy(Component)
 	, SourceWidth(Component->SourceWidth)
@@ -221,7 +211,6 @@ FRectLightSceneProxy::FRectLightSceneProxy(const URectLightComponent* Component)
 	, SourceTexture(Component->SourceTexture)
 {
 	RectAtlasId = ~0u;
-	bRenderAsSpotLight = RenderRectLightsAsSpotLights(SceneInterface->GetFeatureLevel());
 }
 
 FRectLightSceneProxy::~FRectLightSceneProxy() {}
@@ -274,7 +263,7 @@ void FRectLightSceneProxy::GetLightShaderParameters(FLightRenderParameters& Ligh
 	}
 	
 	// Render RectLight approximately as SpotLight if the requester does not support rect light (e.g., translucent light grid or mobile)
-	if (!!(Flags & ELightShaderParameterFlags::RectAsSpotLight) || bRenderAsSpotLight)
+	if (!!(Flags & ELightShaderParameterFlags::RectAsSpotLight))
 	{
 		float ClampedOuterConeAngle = FMath::DegreesToRadians(89.001f);
 		float ClampedInnerConeAngle = FMath::DegreesToRadians(70.0f);
diff --git a/Engine/Source/Runtime/Engine/Private/ShaderCompiler/ShaderCompiler.cpp b/Engine/Source/Runtime/Engine/Private/ShaderCompiler/ShaderCompiler.cpp
--- a/Engine/Source/Runtime/Engine/Private/ShaderCompiler/ShaderCompiler.cpp
+++ b/Engine/Source/Runtime/Engine/Private/ShaderCompiler/ShaderCompiler.cpp
@@ -7806,12 +7806,6 @@ void GlobalBeginCompileShader(
 			SET_SHADER_DEFINE(Input.Environment, USE_GLES_FBF_DEFERRED, bGLESDeferredShading ? 1 : 0);
 			SET_SHADER_DEFINE(Input.Environment, MOBILE_EXTENDED_GBUFFER, MobileUsesExtenedGBuffer((EShaderPlatform)Target.Platform) ? 1 : 0);
 		}
-		else
-		{
-			static const auto CVarRenderRectLightAsSpotLight = IConsoleManager::Get().FindTConsoleVariableDataInt(TEXT("r.Mobile.Forward.RenderRectLightsAsSpotLights"));
-			const bool bRenderAsSpotLight = (CVarRenderRectLightAsSpotLight && CVarRenderRectLightAsSpotLight->GetValueOnAnyThread() != 0);
-			SET_SHADER_DEFINE(Input.Environment, RECT_LIGHT_AS_SPOTLIGHT, bRenderAsSpotLight ? 1 : 0);
-		}
 
 		SET_SHADER_DEFINE(Input.Environment, USE_SCENE_DEPTH_AUX, MobileRequiresSceneDepthAux(ShaderPlatform) ? 1 : 0);
 
@@ -7828,6 +7822,11 @@ void GlobalBeginCompileShader(
 		}
 	}
 
+	if (RenderRectLightsAsSpotLights(GetMaxSupportedFeatureLevel(ShaderPlatform)))
+	{
+		SET_SHADER_DEFINE(Input.Environment, RECT_LIGHT_AS_SPOTLIGHT, 1);
+	}
+
 	if (ShaderPlatform == SP_VULKAN_ES3_1_ANDROID || ShaderPlatform == SP_VULKAN_SM5_ANDROID)
 	{
 		bool bIsStripReflect = true;
diff --git a/Engine/Source/Runtime/Engine/Public/RectLightSceneProxy.h b/Engine/Source/Runtime/Engine/Public/RectLightSceneProxy.h
--- a/Engine/Source/Runtime/Engine/Public/RectLightSceneProxy.h
+++ b/Engine/Source/Runtime/Engine/Public/RectLightSceneProxy.h
@@ -31,7 +31,6 @@ public:
 	FRectLightRayTracingData* RayTracingData; // Render thread data only
 	UTexture*	SourceTexture;
 	uint32		RectAtlasId;
-	bool		bRenderAsSpotLight;
 
 	FRectLightSceneProxy(const URectLightComponent* Component);
 	virtual ~FRectLightSceneProxy();
diff --git a/Engine/Source/Runtime/RenderCore/Private/RenderUtils.cpp b/Engine/Source/Runtime/RenderCore/Private/RenderUtils.cpp
--- a/Engine/Source/Runtime/RenderCore/Private/RenderUtils.cpp
+++ b/Engine/Source/Runtime/RenderCore/Private/RenderUtils.cpp
@@ -1663,6 +1663,16 @@ bool UseSplineMeshSceneResources(const FStaticShaderPlatform Platform)
 	return false;
 }
 
+bool RenderRectLightsAsSpotLights(const FStaticFeatureLevel FeatureLevel)
+{
+	if (FeatureLevel == ERHIFeatureLevel::ES3_1)
+	{
+		static const auto CVarRenderRectLightAsSpotLight = IConsoleManager::Get().FindTConsoleVariableDataInt(TEXT("r.Mobile.Forward.RenderRectLightsAsSpotLights"));
+		return (CVarRenderRectLightAsSpotLight && CVarRenderRectLightAsSpotLight->GetValueOnAnyThread() != 0);
+	}
+	return false;
+}
+
 static TAutoConsoleVariable<int32> CVarLightFunctionAtlasFormat(
 	TEXT("r.LightFunctionAtlas.Format"),
 	0,
diff --git a/Engine/Source/Runtime/RenderCore/Private/Shader.cpp b/Engine/Source/Runtime/RenderCore/Private/Shader.cpp
--- a/Engine/Source/Runtime/RenderCore/Private/Shader.cpp
+++ b/Engine/Source/Runtime/RenderCore/Private/Shader.cpp
@@ -1992,16 +1992,6 @@ void ShaderMapAppendKey(EShaderPlatform Platform, FShaderKeyGenerator& KeyGen)
 			KeyGen.AppendSeparator();
 			KeyGen.Append(TEXT("MobFDP"));
 		}
-
-		{
-			static const auto CVarRenderRectLightAsSpotLight = IConsoleManager::Get().FindTConsoleVariableDataInt(TEXT("r.Mobile.Forward.RenderRectLightsAsSpotLights"));
-			const bool bRenderAsSpotLight = (CVarRenderRectLightAsSpotLight && CVarRenderRectLightAsSpotLight->GetValueOnAnyThread() != 0);
-			if (bRenderAsSpotLight)
-			{
-				KeyGen.AppendSeparator();
-				KeyGen.Append(TEXT("R2S"));
-			}
-		}
 	}
 	else
 	{
@@ -2012,6 +2002,12 @@ void ShaderMapAppendKey(EShaderPlatform Platform, FShaderKeyGenerator& KeyGen)
 		}
 	}
 
+	if (RenderRectLightsAsSpotLights(GetMaxSupportedFeatureLevel(Platform)))
+	{
+		KeyGen.AppendSeparator();
+		KeyGen.Append(TEXT("R2S"));
+	}
+	
 	uint32 PlatformShadingModelsMask = GetPlatformShadingModelsMask(Platform);
 	if (PlatformShadingModelsMask != 0xFFFFFFFF)
 	{
diff --git a/Engine/Source/Runtime/RenderCore/Public/RenderUtils.h b/Engine/Source/Runtime/RenderCore/Public/RenderUtils.h
--- a/Engine/Source/Runtime/RenderCore/Public/RenderUtils.h
+++ b/Engine/Source/Runtime/RenderCore/Public/RenderUtils.h
@@ -552,6 +552,8 @@ extern RENDERCORE_API bool DoesProjectSupportLumenRayTracedTranslucentRefraction
 /** Whether or not the platform supports the scene spline texture for spline meshes */
 RENDERCORE_API bool UseSplineMeshSceneResources(const FStaticShaderPlatform Platform);
 
+RENDERCORE_API bool RenderRectLightsAsSpotLights(const FStaticFeatureLevel FeatureLevel);
+
 namespace Substrate
 {
 	RENDERCORE_API bool IsSubstrateEnabled();
diff --git a/Engine/Source/Runtime/Renderer/Private/GPUScene.cpp b/Engine/Source/Runtime/Renderer/Private/GPUScene.cpp
--- a/Engine/Source/Runtime/Renderer/Private/GPUScene.cpp
+++ b/Engine/Source/Runtime/Renderer/Private/GPUScene.cpp
@@ -691,12 +691,13 @@ void FGPUScene::UpdateGPULights(FRDGBuilder& GraphBuilder, const UE::Tasks::FTas
 	{
 		SCOPED_NAMED_EVENT(UpdateGPUScene_Lights, FColor::Green);
 		const bool bAllowStaticLighting = IsStaticLightingAllowed();
+		const uint32 Flags = RenderRectLightsAsSpotLights(FeatureLevel) ? ELightShaderParameterFlags::RectAsSpotLight : 0u;
 
 		for (int32 Index = 0; Index < Scene.Lights.GetMaxIndex(); ++Index)
 		{
 			if (Scene.Lights.IsAllocated(Index))
 			{
-				InitLightData(Scene.Lights[Index], bAllowStaticLighting, LightData[Index]);
+				InitLightData(Scene.Lights[Index], bAllowStaticLighting, Flags, LightData[Index]);
 			}
 			else
 			{
@@ -739,13 +740,13 @@ void FGPUScene::UpdateGPULights(FRDGBuilder& GraphBuilder, const UE::Tasks::FTas
 	GraphBuilder.QueueBufferUpload<FLightSceneData>(LightDataBufferRDG, LightData, ERDGInitialDataFlags::NoCopy);
 }
 
-void FGPUScene::InitLightData(const FLightSceneInfoCompact& LightInfoCompact, bool bAllowStaticLighting, FLightSceneData& DataOut)
+void FGPUScene::InitLightData(const FLightSceneInfoCompact& LightInfoCompact, bool bAllowStaticLighting, uint32 LightShaderParameterFlags, FLightSceneData& DataOut)
 {
 	const FLightSceneInfo& LightInfo = *LightInfoCompact.LightSceneInfo;
 	const FLightSceneProxy& LightProxy = *LightInfo.Proxy;
 
 	FLightRenderParameters LightParams;
-	LightProxy.GetLightShaderParameters(LightParams);
+	LightProxy.GetLightShaderParameters(LightParams, LightShaderParameterFlags);
 
 	if (LightProxy.IsInverseSquared())
 	{
diff --git a/Engine/Source/Runtime/Renderer/Private/GPUScene.h b/Engine/Source/Runtime/Renderer/Private/GPUScene.h
--- a/Engine/Source/Runtime/Renderer/Private/GPUScene.h
+++ b/Engine/Source/Runtime/Renderer/Private/GPUScene.h
@@ -426,7 +426,7 @@ private:
 	 */
 	void UpdateGPULights(FRDGBuilder& GraphBuilder, const UE::Tasks::FTask& PrerequisiteTask);
 
-	static void InitLightData(const FLightSceneInfoCompact& LightInfoCompact, bool bAllowStaticLighting, FLightSceneData& DataOut);
+	static void InitLightData(const FLightSceneInfoCompact& LightInfoCompact, bool bAllowStaticLighting, uint32 LightShaderParameterFlags, FLightSceneData& DataOut);
 
 	void UploadDynamicPrimitiveShaderDataForViewInternal(FRDGBuilder& GraphBuilder, FViewInfo& View, bool bRayTracing, UE::Renderer::Private::IShadowInvalidatingInstances *ShadowInvalidatingInstances);
 
diff --git a/Engine/Source/Runtime/Renderer/Private/LightGridInjection.cpp b/Engine/Source/Runtime/Renderer/Private/LightGridInjection.cpp
--- a/Engine/Source/Runtime/Renderer/Private/LightGridInjection.cpp
+++ b/Engine/Source/Runtime/Renderer/Private/LightGridInjection.cpp
@@ -524,6 +524,7 @@ FComputeLightGridOutput FSceneRenderer::ComputeLightGrid(FRDGBuilder& GraphBuild
 				}
 			}
 
+			const uint32 LightShaderParameterFlags = RenderRectLightsAsSpotLights(FeatureLevel) ? ELightShaderParameterFlags::RectAsSpotLight : 0u;
 			float SelectedForwardDirectionalLightIntensitySq = 0.0f;
 			int32 SelectedForwardDirectionalLightPriority = -1;
 			const TArray<FSortedLightSceneInfo, SceneRenderingAllocator>& SortedLights = SortedLightSet.SortedLights;
@@ -539,7 +540,7 @@ FComputeLightGridOutput FSceneRenderer::ComputeLightGrid(FRDGBuilder& GraphBuild
 				if (LightSceneInfo->ShouldRenderLight(View))
 				{
 					FLightRenderParameters LightParameters;
-					LightProxy->GetLightShaderParameters(LightParameters);
+					LightProxy->GetLightShaderParameters(LightParameters, LightShaderParameterFlags);
 
 					if (LightProxy->IsInverseSquared())
 					{
diff --git a/Engine/Source/Runtime/Renderer/Private/MobileDeferredShadingPass.cpp b/Engine/Source/Runtime/Renderer/Private/MobileDeferredShadingPass.cpp
--- a/Engine/Source/Runtime/Renderer/Private/MobileDeferredShadingPass.cpp
+++ b/Engine/Source/Runtime/Renderer/Private/MobileDeferredShadingPass.cpp
@@ -185,10 +185,10 @@ public:
 	SHADER_USE_PARAMETER_STRUCT_WITH_LEGACY_BASE(FMobileRadialLightFunctionPS, FMaterialShader)
 
 	class FEnableShadingModelSupport: SHADER_PERMUTATION_BOOL("ENABLE_SHADINGMODEL_SUPPORT_MOBILE_DEFERRED");
-	class FSpotLightDim				: SHADER_PERMUTATION_BOOL("IS_SPOT_LIGHT");
+	class FRadialLightTypeDim		: SHADER_PERMUTATION_RANGE_INT("RADIAL_LIGHT_TYPE", LIGHT_TYPE_POINT, LIGHT_TYPE_RECT);
 	class FIESProfileDim			: SHADER_PERMUTATION_BOOL("USE_IES_PROFILE");
 	class FSpotLightShadowDim		: SHADER_PERMUTATION_BOOL("SUPPORT_SPOTLIGHTS_SHADOW");
-	using FPermutationDomain = TShaderPermutationDomain<FEnableShadingModelSupport, FSpotLightDim, FIESProfileDim, FSpotLightShadowDim>;
+	using FPermutationDomain = TShaderPermutationDomain<FEnableShadingModelSupport, FRadialLightTypeDim, FIESProfileDim, FSpotLightShadowDim>;
 
 	BEGIN_SHADER_PARAMETER_STRUCT(FParameters, )
 		SHADER_PARAMETER_STRUCT(FLightShaderParameters, Light)
@@ -713,7 +713,7 @@ static void RenderLocalLight_StencilMask(FRHICommandList& RHICmdList, const FSce
 	FDeferredLightVS::FParameters ParametersVS = FDeferredLightVS::GetParameters(View, &LightSceneInfo);
 	SetShaderParameters(RHICmdList, VertexShader, VertexShader.GetVertexShader(), ParametersVS);
 
-	if (LightType == LightType_Point)
+	if (LightType == LightType_Point || LightType == LightType_Rect)
 	{
 		StencilingGeometry::DrawSphere(RHICmdList);
 	}
@@ -740,7 +740,8 @@ static void RenderLocalLight(
 	const uint8 LightType = LightSceneInfo.Proxy->GetLightType();
 	const bool bIsSpotLight = LightType == LightType_Spot;
 	const bool bIsPointLight = LightType == LightType_Point;
-	if (!bIsSpotLight && !bIsPointLight)
+	const bool bIsRectLight = LightType == LightType_Rect;
+	if (!bIsSpotLight && !bIsPointLight && !bIsRectLight)
 	{
 		return;
 	}
@@ -819,7 +820,7 @@ static void RenderLocalLight(
 
 		FMobileRadialLightFunctionPS::FPermutationDomain PermutationVector;
 		PermutationVector.Set<FMobileRadialLightFunctionPS::FEnableShadingModelSupport>(bEnableShadingModelSupport);
-		PermutationVector.Set<FMobileRadialLightFunctionPS::FSpotLightDim>(bIsSpotLight);
+		PermutationVector.Set<FMobileRadialLightFunctionPS::FRadialLightTypeDim>(LightType);
 		PermutationVector.Set<FMobileRadialLightFunctionPS::FIESProfileDim>(bUseIESTexture);
 		PermutationVector.Set<FMobileRadialLightFunctionPS::FSpotLightShadowDim>(bShouldCastShadow);
 		FCachedLightMaterial LightMaterial;
@@ -840,7 +841,7 @@ static void RenderLocalLight(
 
 		SetShaderParametersMixedPS(RHICmdList, PixelShader, PassParameters, View, LightMaterial.MaterialProxy, *LightMaterial.Material);
 
-		if (LightType == LightType_Point)
+		if (LightType == LightType_Point || LightType == LightType_Rect)
 		{
 			StencilingGeometry::DrawSphere(RHICmdList);
 		}
@@ -919,7 +920,7 @@ static void RenderSimpleLights(
 		TShaderRef<FMobileRadialLightFunctionPS> PixelShader;
 		FMobileRadialLightFunctionPS::FPermutationDomain PermutationVector;
 		PermutationVector.Set<FMobileRadialLightFunctionPS::FEnableShadingModelSupport>(bEnableShadingModelSupport);
-		PermutationVector.Set<FMobileRadialLightFunctionPS::FSpotLightDim>(false);
+		PermutationVector.Set<FMobileRadialLightFunctionPS::FRadialLightTypeDim>(LightType_Point);
 		PermutationVector.Set<FMobileRadialLightFunctionPS::FIESProfileDim>(false);
 		PassPixelShaders[PassIndex] = MaterialShaderMap->GetShader<FMobileRadialLightFunctionPS>(PermutationVector);
 		GraphicsPSOLight[PassIndex].BoundShaderState.PixelShaderRHI = PassPixelShaders[PassIndex].GetPixelShader();
