Fix missing include

diff --git a/Engine/Source/Runtime/Renderer/Private/RayTracing/RayTracingScene.h b/Engine/Source/Runtime/Renderer/Private/RayTracing/RayTracingScene.h
--- a/Engine/Source/Runtime/Renderer/Private/RayTracing/RayTracingScene.h
+++ b/Engine/Source/Runtime/Renderer/Private/RayTracing/RayTracingScene.h
@@ -16,6 +16,7 @@
 #include "Misc/MemStack.h"
 #include "Containers/ArrayView.h"
 #include "MeshPassProcessor.h"
+#include "RayTracingMeshDrawCommands.h"
 #include "RayTracingInstanceBufferUtil.h"
 #include "RayTracingDebugTypes.h"
 
