Support screen space reflections in the mobile renderer.

Supported in forward and deferred. It works without TAA but no denoising is done, so it looks much worse.

At the moment enabling SSR implies doing a depth prepass.

diff --git a/Engine/Shaders/Private/MobileBasePassPixelShader.usf b/Engine/Shaders/Private/MobileBasePassPixelShader.usf
--- a/Engine/Shaders/Private/MobileBasePassPixelShader.usf
+++ b/Engine/Shaders/Private/MobileBasePassPixelShader.usf
@@ -872,6 +872,7 @@ void Main(
 #if MATERIALBLENDING_MASKED || MATERIALBLENDING_SOLID || TRANSLUCENCY_LIGHTING_SURFACE_FORWARDSHADING || TRANSLUCENCY_LIGHTING_SURFACE_LIGHTINGVOLUME || MATERIAL_SHADINGMODEL_SINGLELAYERWATER
 	// Reflection IBL
 	AccumulateReflection(GBuffer
+		, SvPosition
 		, CameraVector
 		, MaterialParameters.WorldPosition_CamRelative
 		, MaterialParameters.ReflectionVector
diff --git a/Engine/Shaders/Private/MobileDeferredShading.usf b/Engine/Shaders/Private/MobileDeferredShading.usf
--- a/Engine/Shaders/Private/MobileDeferredShading.usf
+++ b/Engine/Shaders/Private/MobileDeferredShading.usf
@@ -99,6 +99,7 @@ half3 SkyLightDiffuseMobile(FGBufferData GBuffer, half3 DiffuseColor, half3 V, i
 
 void ReflectionEnvironmentSkyLighting(
 	FGBufferData GBuffer,
+	half4 SvPosition,
 	half3 CameraVector,
 	float3 TranslatedWorldPosition,
 	half3 ReflectionVector,
@@ -125,6 +126,7 @@ void ReflectionEnvironmentSkyLighting(
 
 	// IBL
 	AccumulateReflection(GBuffer
+		, SvPosition
 		, CameraVector
 		, TranslatedWorldPosition
 		, ReflectionVector
@@ -189,7 +191,7 @@ void MobileDirectionalLightPS(
 
 #if ENABLE_CLUSTERED_REFLECTION || ENABLE_SKY_LIGHT || ENABLE_PLANAR_REFLECTION
 	// If we have a single directional light, apply relfection and sky contrubution here
-	ReflectionEnvironmentSkyLighting(GBuffer, CameraVector, TranslatedWorldPosition, ReflectionVector, GridIndex, DirectLighting);
+	ReflectionEnvironmentSkyLighting(GBuffer, SvPosition, CameraVector, TranslatedWorldPosition, ReflectionVector, GridIndex, DirectLighting);
 #endif
 
 	half3 Color = DirectLighting.TotalLight;
@@ -376,7 +378,7 @@ void MobileReflectionEnvironmentSkyLightingPS(
 #endif
 	
 	FLightAccumulator DirectLighting = (FLightAccumulator)0;
-	ReflectionEnvironmentSkyLighting(GBuffer, CameraVector, TranslatedWorldPosition, ReflectionVector, GridIndex, DirectLighting);
+	ReflectionEnvironmentSkyLighting(GBuffer, SvPosition, CameraVector, TranslatedWorldPosition, ReflectionVector, GridIndex, DirectLighting);
 	half3 ReflectionAndSky = DirectLighting.TotalLight;
 
 	ReflectionAndSky *= View.PreExposure;
diff --git a/Engine/Shaders/Private/MobileLightingCommon.ush b/Engine/Shaders/Private/MobileLightingCommon.ush
--- a/Engine/Shaders/Private/MobileLightingCommon.ush
+++ b/Engine/Shaders/Private/MobileLightingCommon.ush
@@ -13,6 +13,20 @@
 #include "DeferredLightingCommon.ush"
 #include "SHCommon.ush"
 
+#if MOBILE_SSR_ENABLED && !(MOBILE_QL_FORCE_FULLY_ROUGH || MATERIAL_FULLY_ROUGH)
+	#if MOBILE_DEFERRED_SHADING
+		#define MOBILE_USE_SSR (MATERIAL_SSR)
+	#else
+		#define MOBILE_USE_SSR (MATERIAL_SSR || !MATERIALBLENDING_ANY_TRANSLUCENT)
+	#endif
+#else 
+	#define MOBILE_USE_SSR 0
+#endif
+
+#if MOBILE_USE_SSR
+#include "SSRT/SSRTRayCast.ush"
+#endif
+
 #ifndef TRANSLUCENCY_NON_DIRECTIONAL
 #define TRANSLUCENCY_NON_DIRECTIONAL 0
 #endif
@@ -446,6 +460,7 @@ float3 GetImageBasedReflectionSpecular(
 // Version used for legacy shading
 void AccumulateReflection(
 	FGBufferData GBuffer, 
+	half4 SvPosition,
 	half3 CameraVector, 
 	float3 TranslatedWorldPosition, 
 	half3 ReflectionVector, 
@@ -479,6 +494,33 @@ void AccumulateReflection(
 		, GridIndex
 	);
 
+
+#if MOBILE_USE_SSR
+	BRANCH if(View.CameraCut == 0 && MobileBasePass.SSRParams.QualityAndExposureCorrection.x > 0.0 && GBuffer.Roughness <= MobileBasePass.SSRParams.QualityAndExposureCorrection.w)
+	{
+		float StepOffset = InterleavedGradientNoise(SvPosition.xy, View.StateFrameIndexMod8);
+		StepOffset -= 0.5;
+		float3 HitUVz;
+		float Level = 0;
+		bool bHit = RayCast(MobileBasePass.SSRParams.HZB, MobileBasePass.SSRParams.HZBSampler, TranslatedWorldPosition, ReflectionVector, GBuffer.Roughness, GBuffer.Depth,
+			12, StepOffset, MobileBasePass.SSRParams.HZBUvFactorAndInvFactor, /*bDebugPrint=*/false, HitUVz, Level);
+
+		BRANCH if(bHit)
+		{
+			float2 SampleUV;
+			float Vignette;
+			ReprojectHit(MobileBasePass.SSRParams.PrevScreenPositionScaleBias, HitUVz, SampleUV, Vignette);
+
+			SampleUV = clamp(SampleUV, MobileBasePass.SSRParams.PrevSceneColorBilinearUVMinMax.xy,  MobileBasePass.SSRParams.PrevSceneColorBilinearUVMinMax.zw);
+			half4 SSR = SampleScreenColor(MobileBasePass.SSRParams.SceneColor, MobileBasePass.SSRParams.SceneColorSampler, SampleUV);
+			SSR *= Vignette * saturate( 2 - 6.6 * GBuffer.Roughness );
+
+			SSR.rgb *= MobileBasePass.SSRParams.QualityAndExposureCorrection.y;
+			SpecularIBL.rgb = SpecularIBL.rgb * (1 - SSR.a) + SSR.rgb;
+		}
+	}
+#endif // MOBILE_USE_SSR
+
 	// Planar reflection
 	SpecularIBL = GetPlanarReflectionbasedReflectionLighting_Mobile(TranslatedWorldPosition, GBuffer.WorldNormal, GBuffer.Roughness, SpecularIBL);
 
diff --git a/Engine/Shaders/Private/SSRT/SSRTReflections.usf b/Engine/Shaders/Private/SSRT/SSRTReflections.usf
--- a/Engine/Shaders/Private/SSRT/SSRTReflections.usf
+++ b/Engine/Shaders/Private/SSRT/SSRTReflections.usf
@@ -1,5 +1,12 @@
 // Copyright Epic Games, Inc. All Rights Reserved.
 
+#if SHADING_PATH_MOBILE
+	#include "../Common.ush"
+	#define MobileSceneTextures MobileBasePass.SceneTextures
+	#include "../DeferredShadingCommon.ush"
+	#undef MobileSceneTextures
+#endif // SHADING_PATH_MOBILE
+
 #include "SSRTRayCast.ush"
 #include "../Random.ush"
 #include "../BRDF.ush"
@@ -30,6 +37,23 @@ SamplerState HZBSampler;
 
 uint ShouldReflectOnlyWater;
 
+#if SHADING_PATH_MOBILE
+	#define SceneColor MobileBasePass.SSRParams.SceneColor
+	#define SceneColorSampler MobileBasePass.SSRParams.SceneColorSampler
+	#define HZB MobileBasePass.SSRParams.HZB
+	#define HZBSampler MobileBasePass.SSRParams.HZBSampler
+	#define PrevScreenPositionScaleBias MobileBasePass.SSRParams.PrevScreenPositionScaleBias
+	#define HZBUvFactorAndInvFactor MobileBasePass.SSRParams.HZBUvFactorAndInvFactor
+	#define ShouldReflectOnlyWater 0
+	#define PrevSceneColorPreExposureCorrection MobileBasePass.SSRParams.QualityAndExposureCorrection.z
+	#define GBufferVelocityTexture MobileBasePass.SceneTextures.SceneVelocityTexture
+	#define GBufferVelocityTextureSampler MobileBasePass.SceneTextures.SceneVelocityTextureSampler
+#endif
+
+#ifndef SSR_USE_VELOCITY
+	#define SSR_USE_VELOCITY 1
+#endif
+
 float GetRoughness(const in FGBufferData GBuffer)
 {
 	float Roughness = GBuffer.Roughness;
@@ -152,7 +176,11 @@ void ScreenSpaceReflections(
 
 #else // SUBSTRATE_ENABLED
 
-	FGBufferData GBuffer = GetGBufferDataFromSceneTextures(UV);
+	#if SHADING_PATH_MOBILE
+		FGBufferData GBuffer = MobileFetchAndDecodeGBuffer(UV, SvPosition.xy);
+	#else
+		FGBufferData GBuffer = GetGBufferDataFromSceneTextures(UV);
+	#endif
 
 	float3 N = GBuffer.WorldNormal;
 	const float SceneDepth = GBuffer.Depth;
@@ -317,7 +345,11 @@ void ScreenSpaceReflections(
 
 				float2 SampleUV;
 				float Vignette;
-				ReprojectHit(PrevScreenPositionScaleBias, GBufferVelocityTexture, GBufferVelocityTextureSampler, HitUVz, SampleUV, Vignette);
+				#if SSR_USE_VELOCITY
+					ReprojectHit(PrevScreenPositionScaleBias, GBufferVelocityTexture, GBufferVelocityTextureSampler, HitUVz, SampleUV, Vignette);
+				#else
+					ReprojectHit(PrevScreenPositionScaleBias, HitUVz, SampleUV, Vignette);
+				#endif
 
 				float4 SampleColor = SampleScreenColor( SceneColor, SceneColorSampler, SampleUV ) * Vignette;
 
@@ -396,7 +428,12 @@ void ScreenSpaceReflections(
 
 			float2 SampleUV;
 			float Vignette;
-			ReprojectHit(PrevScreenPositionScaleBias, GBufferVelocityTexture, GBufferVelocityTextureSampler, HitUVz, SampleUV, Vignette);
+
+			#if SSR_USE_VELOCITY
+				ReprojectHit(PrevScreenPositionScaleBias, GBufferVelocityTexture, GBufferVelocityTextureSampler, HitUVz, SampleUV, Vignette);
+			#else
+				ReprojectHit(PrevScreenPositionScaleBias, HitUVz, SampleUV, Vignette);
+			#endif
 
 			OutColor = SampleScreenColor(SceneColor, SceneColorSampler, SampleUV) * Vignette;
 		}
@@ -434,6 +471,27 @@ void ScreenSpaceReflectionsPS(
 	);
 }
 
+void MobileScreenSpaceReflectionsPS(float4 SvPosition : SV_POSITION
+#if USE_GLES_FBF_DEFERRED
+, out float4 OutProxyAdditive : SV_Target0 // "OutProxyAdditive" has special handling in the shader compiler and does additive blending
+, out float4 OutGBufferA : SV_Target1
+, out float4 OutGBufferB : SV_Target2
+, out float4 OutGBufferC : SV_Target3
+#else
+, out float4 OutColor : SV_Target0
+#endif
+)
+{
+	float4 Color = 0;
+	ScreenSpaceReflections(SvPosition, Color);
+
+#if USE_GLES_FBF_DEFERRED
+	OutProxyAdditive = Color;
+#else
+	OutColor = Color;
+#endif
+}
+
 #if COMPUTE_SHADER
 
 RWTexture2D<float4> SSRColorOutput;
diff --git a/Engine/Source/Developer/TargetPlatform/Private/TargetPlatformSettingsBase.cpp b/Engine/Source/Developer/TargetPlatform/Private/TargetPlatformSettingsBase.cpp
--- a/Engine/Source/Developer/TargetPlatform/Private/TargetPlatformSettingsBase.cpp
+++ b/Engine/Source/Developer/TargetPlatform/Private/TargetPlatformSettingsBase.cpp
@@ -94,6 +94,12 @@ bool FTargetPlatformSettingsBase::UsesMobileDBuffer() const
 	return CVar ? (CVar->GetInt() != 0) : false;
 }
 
+bool FTargetPlatformSettingsBase::UsesMobileScreenSpaceReflections() const
+{
+	static IConsoleVariable* CVar = IConsoleManager::Get().FindConsoleVariable(TEXT("r.Mobile.ScreenSpaceReflections"));
+	return CVar ? (CVar->GetInt() != 0) : false;
+}
+
 bool FTargetPlatformSettingsBase::UsesASTCHDR() const
 {
 	static IConsoleVariable* CVar = IConsoleManager::Get().FindConsoleVariable(TEXT("cook.ASTCTextureCompressor"));
diff --git a/Engine/Source/Developer/TargetPlatform/Public/Common/TargetPlatformBase.h b/Engine/Source/Developer/TargetPlatform/Public/Common/TargetPlatformBase.h
--- a/Engine/Source/Developer/TargetPlatform/Public/Common/TargetPlatformBase.h
+++ b/Engine/Source/Developer/TargetPlatform/Public/Common/TargetPlatformBase.h
@@ -846,6 +846,10 @@ public:
 	{
 		return TargetPlatformSettings->UsesMobileDBuffer();
 	}
+	virtual bool UsesMobileScreenSpaceReflections() const override
+	{
+		return TargetPlatformSettings->UsesMobileScreenSpaceReflections();
+	}
 	virtual bool UsesASTCHDR() const override 
 	{
 		return TargetPlatformSettings->UsesASTCHDR();
diff --git a/Engine/Source/Developer/TargetPlatform/Public/Common/TargetPlatformSettingsBase.h b/Engine/Source/Developer/TargetPlatform/Public/Common/TargetPlatformSettingsBase.h
--- a/Engine/Source/Developer/TargetPlatform/Public/Common/TargetPlatformSettingsBase.h
+++ b/Engine/Source/Developer/TargetPlatform/Public/Common/TargetPlatformSettingsBase.h
@@ -36,6 +36,7 @@ public:
 	TARGETPLATFORM_API virtual int32 GetHeightFogModeForOpaque() const override;
 	TARGETPLATFORM_API virtual bool UsesMobileAmbientOcclusion() const override;
 	TARGETPLATFORM_API virtual bool UsesMobileDBuffer() const override;
+	TARGETPLATFORM_API virtual bool UsesMobileScreenSpaceReflections() const override;
 	TARGETPLATFORM_API virtual bool UsesASTCHDR() const override;
 	TARGETPLATFORM_API virtual void GetRayTracingShaderFormats(TArray<FName>& OutFormats) const override;
 
diff --git a/Engine/Source/Developer/TargetPlatform/Public/Interfaces/ITargetPlatformSettings.h b/Engine/Source/Developer/TargetPlatform/Public/Interfaces/ITargetPlatformSettings.h
--- a/Engine/Source/Developer/TargetPlatform/Public/Interfaces/ITargetPlatformSettings.h
+++ b/Engine/Source/Developer/TargetPlatform/Public/Interfaces/ITargetPlatformSettings.h
@@ -226,6 +226,11 @@ public:
 	 */
 	virtual bool UsesMobileDBuffer() const = 0;
 
+	/**
+	 * Gets whether the platform should use Screen Space Reflections when using the mobile renderer.
+	 */
+	virtual bool UsesMobileScreenSpaceReflections() const = 0;
+
 	/**
 	 * Gets whether the platform uses ASTC HDR
 	 */
diff --git a/Engine/Source/Runtime/Engine/Classes/Engine/RendererSettings.h b/Engine/Source/Runtime/Engine/Classes/Engine/RendererSettings.h
--- a/Engine/Source/Runtime/Engine/Classes/Engine/RendererSettings.h
+++ b/Engine/Source/Runtime/Engine/Classes/Engine/RendererSettings.h
@@ -321,8 +321,7 @@ class URendererSettings : public UDeveloperSettings
 
 	UPROPERTY(config, EditAnywhere, Category = Mobile, meta = (
 		ConsoleVariable = "r.Mobile.AntiAliasing", DisplayName = "Mobile Anti-Aliasing Method",
-		ToolTip = "The mobile default anti-aliasing method.",
-		EditCondition = "MobileShadingPath == 1"))
+		ToolTip = "The mobile default anti-aliasing method."))
 	TEnumAsByte<EMobileAntiAliasingMethod::Type> MobileAntiAliasing;
 
 	UPROPERTY(config, EditAnywhere, Category = Mobile, meta = (
@@ -1363,6 +1362,12 @@ class URendererSettings : public UDeveloperSettings
 		ConfigRestartRequired = true))
 		TEnumAsByte<EMobilePlanarReflectionMode::Type> MobilePlanarReflectionMode;
 
+	UPROPERTY(config, EditAnywhere, Category = Mobile, meta = (
+		ConsoleVariable = "r.Mobile.ScreenSpaceReflections", DisplayName = "Support Screen Space Reflections on mobile",
+		ToolTip = "Support Screen Space Reflections with mobile rendering. Screen Space Reflections on mobile require TAA Changing this setting requires restarting the editor.",
+		ConfigRestartRequired = true))
+		uint32 bMobileScreenSpaceReflections : 1;
+
 	UPROPERTY(config, EditAnywhere, Category = Mobile, meta = (
 		ConsoleVariable = "r.Mobile.SupportsGen4TAA", DisplayName = "Support desktop Gen4 TAA on mobile",
 		ToolTip = "Support desktop Gen4 TAA with mobile rendering. Changing this setting requires restarting the editor.",
diff --git a/Engine/Source/Runtime/Engine/Private/RendererSettings.cpp b/Engine/Source/Runtime/Engine/Private/RendererSettings.cpp
--- a/Engine/Source/Runtime/Engine/Private/RendererSettings.cpp
+++ b/Engine/Source/Runtime/Engine/Private/RendererSettings.cpp
@@ -164,21 +164,24 @@ void UpdateDependentPropertyInConfigFile(URendererSettings* RendererSettings, FN
 
 void URendererSettings::FixAntiAliasingOnShadingPathChange(FPropertyChangedEvent& PropertyChangedEvent)
 {
-	// Going from deferred to forward we force MSAA. After that, the AA method is blocked by the edit condition in metadata.
 	if (PropertyChangedEvent.GetPropertyName() == GET_MEMBER_NAME_CHECKED(URendererSettings, MobileShadingPath))
 	{
 		if (MobileShadingPath == EMobileShadingPath::Forward)
 		{
-			MobileAntiAliasing = EMobileAntiAliasingMethod::MSAA;
+			// Nop
 		}
 		else // When we disable forward shading we need to change AA method
 		{
-			MobileAntiAliasing = EMobileAntiAliasingMethod::None;
-			UE_LOG(LogTemp, Warning, TEXT("Disabling Mobile AA because MSAA is not compatible with Mobile Deferred Shading"));
+			if (MobileAntiAliasing == EMobileAntiAliasingMethod::MSAA)
+			{
+				MobileAntiAliasing = EMobileAntiAliasingMethod::None;
+				UE_LOG(LogTemp, Warning, TEXT("Disabling Mobile AA because MSAA is not compatible with Mobile Deferred Shading"));
+			}
 		}
 		UpdateDependentPropertyInConfigFile(this, GET_MEMBER_NAME_CHECKED(URendererSettings, MobileAntiAliasing));
 		ExportValuesToConsoleVariables(GetClass()->FindPropertyByName(GET_MEMBER_NAME_CHECKED(URendererSettings, MobileAntiAliasing)));
 	}
+	// On desktop going from deferred to forward we force MSAA. After that, the AA method is blocked by the edit condition in metadata.
 	else if (PropertyChangedEvent.GetPropertyName() == GET_MEMBER_NAME_CHECKED(URendererSettings, bForwardShading))
 	{
 		if (bForwardShading)
diff --git a/Engine/Source/Runtime/RenderCore/Private/RenderUtils.cpp b/Engine/Source/Runtime/RenderCore/Private/RenderUtils.cpp
--- a/Engine/Source/Runtime/RenderCore/Private/RenderUtils.cpp
+++ b/Engine/Source/Runtime/RenderCore/Private/RenderUtils.cpp
@@ -540,6 +540,13 @@ RENDERCORE_API bool IsMobileAmbientOcclusionEnabled(const FStaticShaderPlatform
 	return IsMobilePlatform(Platform) && GMobileAmbientOcclusionPlatformMask[(int)Platform];
 }
 
+RENDERCORE_API ShaderPlatformMaskType GMobileScreenSpaceReflectionsPlatformMask;
+
+RENDERCORE_API bool AreMobileScreenSpaceReflectionsEnabled(const FStaticShaderPlatform Platform)
+{
+	return IsMobilePlatform(Platform) && GMobileScreenSpaceReflectionsPlatformMask[(int)Platform];
+}
+
 RENDERCORE_API bool IsMobileDistanceFieldEnabled(const FStaticShaderPlatform Platform)
 {
 	return IsMobilePlatform(Platform) && FDataDrivenShaderPlatformInfo::GetSupportsDistanceFields(Platform) && IsUsingDistanceFields(Platform);
@@ -597,7 +604,7 @@ RENDERCORE_API bool MobileBasePassAlwaysUsesCSM(const FStaticShaderPlatform Plat
 
 RENDERCORE_API bool MobileUsesFullDepthPrepass(const FStaticShaderPlatform Platform)
 {
-	return MobileUsesShadowMaskTexture(Platform) || IsMobileAmbientOcclusionEnabled(Platform) || IsUsingDBuffers(Platform) || FReadOnlyCVARCache::MobileEarlyZPass(Platform) == 1;
+	return MobileUsesShadowMaskTexture(Platform) || IsMobileAmbientOcclusionEnabled(Platform) || IsUsingDBuffers(Platform) || AreMobileScreenSpaceReflectionsEnabled(Platform) || FReadOnlyCVARCache::MobileEarlyZPass(Platform) == 1;
 }
 
 RENDERCORE_API bool ShouldForceFullDepthPass(const FStaticShaderPlatform Platform)
@@ -751,6 +758,9 @@ RENDERCORE_API void RenderUtilsInit()
 	static IConsoleVariable* MobileAmbientOcclusionCVar = IConsoleManager::Get().FindConsoleVariable(TEXT("r.Mobile.AmbientOcclusion"));
 	GMobileAmbientOcclusionPlatformMask.Init(MobileAmbientOcclusionCVar && MobileAmbientOcclusionCVar->GetInt(), EShaderPlatform::SP_NumPlatforms);
 
+	static IConsoleVariable* MobileScreenSpaceReflectionsCVar = IConsoleManager::Get().FindConsoleVariable(TEXT("r.Mobile.ScreenSpaceReflections"));
+	GMobileScreenSpaceReflectionsPlatformMask.Init(MobileScreenSpaceReflectionsCVar && MobileScreenSpaceReflectionsCVar->GetInt(), EShaderPlatform::SP_NumPlatforms);
+
 #if WITH_EDITOR
 	ITargetPlatformManagerModule* TargetPlatformManager = GetTargetPlatformManager();
 	if (TargetPlatformManager)
@@ -791,6 +801,8 @@ RENDERCORE_API void RenderUtilsInit()
 					GVelocityEncodeDepthPlatformMask[ShaderPlatformIndex] = TargetPlatformSettings->VelocityEncodeDepth();
 
 					GMobileAmbientOcclusionPlatformMask[ShaderPlatformIndex] = TargetPlatformSettings->UsesMobileAmbientOcclusion();
+
+					GMobileScreenSpaceReflectionsPlatformMask[ShaderPlatformIndex] = TargetPlatformSettings->UsesMobileScreenSpaceReflections();
 				}
 			}
 
diff --git a/Engine/Source/Runtime/RenderCore/Public/RenderUtils.h b/Engine/Source/Runtime/RenderCore/Public/RenderUtils.h
--- a/Engine/Source/Runtime/RenderCore/Public/RenderUtils.h
+++ b/Engine/Source/Runtime/RenderCore/Public/RenderUtils.h
@@ -312,6 +312,8 @@ RENDERCORE_API uint32 GetPlatformShadingModelsMask(const FStaticShaderPlatform P
 
 RENDERCORE_API bool IsMobileAmbientOcclusionEnabled(const FStaticShaderPlatform Platform);
 
+RENDERCORE_API bool AreMobileScreenSpaceReflectionsEnabled(const FStaticShaderPlatform Platform);
+
 RENDERCORE_API bool IsMobileDistanceFieldEnabled(const FStaticShaderPlatform Platform);
 
 RENDERCORE_API bool IsMobileMovableSpotlightShadowsEnabled(const FStaticShaderPlatform Platform);
diff --git a/Engine/Source/Runtime/Renderer/Private/DepthRendering.cpp b/Engine/Source/Runtime/Renderer/Private/DepthRendering.cpp
--- a/Engine/Source/Runtime/Renderer/Private/DepthRendering.cpp
+++ b/Engine/Source/Runtime/Renderer/Private/DepthRendering.cpp
@@ -775,11 +775,11 @@ bool FDepthPassMeshProcessor::Process(
 	// Use StencilMask for DecalOutput on mobile
 	if (FeatureLevel == ERHIFeatureLevel::ES3_1 && !bShadowProjection)
 	{
-		extern void SetMobileBasePassDepthState(FMeshPassProcessorRenderState& DrawRenderState, const FPrimitiveSceneProxy* PrimitiveSceneProxy, FMaterialShadingModelField ShadingModels, bool bUsesDeferredShading);
+		extern void SetMobileBasePassDepthState(FMeshPassProcessorRenderState& DrawRenderState, const FPrimitiveSceneProxy* PrimitiveSceneProxy, const FMaterial& Material, FMaterialShadingModelField ShadingModels, bool bUsesDeferredShading);
 		
 		FMaterialShadingModelField ShadingModels = MaterialResource.GetShadingModels();
 		bool bUsesDeferredShading = IsMobileDeferredShadingEnabled(GetFeatureLevelShaderPlatform(FeatureLevel));
-		SetMobileBasePassDepthState(DrawRenderState, PrimitiveSceneProxy, ShadingModels, bUsesDeferredShading);
+		SetMobileBasePassDepthState(DrawRenderState, PrimitiveSceneProxy, MaterialResource, ShadingModels, bUsesDeferredShading);
 	}
 
 	FMeshMaterialShaderElementData ShaderElementData;
diff --git a/Engine/Source/Runtime/Renderer/Private/MobileBasePass.cpp b/Engine/Source/Runtime/Renderer/Private/MobileBasePass.cpp
--- a/Engine/Source/Runtime/Renderer/Private/MobileBasePass.cpp
+++ b/Engine/Source/Runtime/Renderer/Private/MobileBasePass.cpp
@@ -62,22 +62,26 @@ EMobileLocalLightSetting GetMobileForwardLocalLightSetting(EShaderPlatform Shade
 	return EMobileLocalLightSetting::LOCAL_LIGHTS_DISABLED;
 }
 
-uint8 GetMobileShadingModelStencilValue(FMaterialShadingModelField ShadingModel)
+extern const uint8 MobileShadingMobelSupportStencilValue = 0b01u;
+uint8 GetMobileShadingModelStencilValue(FMaterialShadingModelField ShadingModel, bool bFullyRough)
 {
+	// Bit 0 is set for materials that are receive SSR
+	// Bit 1 is set for DefaultLit materials (see MobileDeferredShadingPass.cpp)
+	const uint8 DefaultLitMask = bFullyRough ? 0b10u : 0b11u;
 	if (ShadingModel.HasOnlyShadingModel(MSM_DefaultLit))
 	{
-		return 1u;
+		return DefaultLitMask;
 	}
 	else if (ShadingModel.HasOnlyShadingModel(MSM_Unlit))
 	{
-		return 0u;
+		return 0b00u;
 	}
-	
+
 	// mark everyhing as MSM_DefaultLit if GBuffer CustomData is not supported
-	return MobileUsesGBufferCustomData(GMaxRHIShaderPlatform) ? 2u : 1u;
+	return MobileUsesGBufferCustomData(GMaxRHIShaderPlatform) ? MobileShadingMobelSupportStencilValue : DefaultLitMask;
 }
 
-void SetMobileBasePassDepthState(FMeshPassProcessorRenderState& DrawRenderState, const FPrimitiveSceneProxy* PrimitiveSceneProxy, FMaterialShadingModelField ShadingModels, bool bUsesDeferredShading)
+void SetMobileBasePassDepthState(FMeshPassProcessorRenderState& DrawRenderState, const FPrimitiveSceneProxy* PrimitiveSceneProxy, const FMaterial& Material, FMaterialShadingModelField ShadingModels, bool bUsesDeferredShading)
 {
 	DrawRenderState.SetDepthStencilState(TStaticDepthStencilState<
 		true, CF_DepthNearOrEqual,
@@ -93,8 +97,8 @@ void SetMobileBasePassDepthState(FMeshPassProcessorRenderState& DrawRenderState,
 
 	if (bUsesDeferredShading)
 	{
-		// store into [1-3] bits
-		uint8 ShadingModel = GetMobileShadingModelStencilValue(ShadingModels);
+		// store into [1-2] bits
+		uint8 ShadingModel = GetMobileShadingModelStencilValue(ShadingModels, Material.IsFullyRough());
 		StencilValue |= GET_STENCIL_MOBILE_SM_MASK(ShadingModel);
 		StencilValue |= STENCIL_LIGHTING_CHANNELS_MASK(PrimitiveSceneProxy ? PrimitiveSceneProxy->GetLightingChannelStencilValue() : 0x00);
 	}
@@ -563,7 +567,7 @@ void MobileBasePass::SetOpaqueRenderState(FMeshPassProcessorRenderState& DrawRen
 {
 	if (bCanUseDepthStencil)
 	{
-		SetMobileBasePassDepthState(DrawRenderState, PrimitiveSceneProxy, ShadingModels, bUsesDeferredShading);
+		SetMobileBasePassDepthState(DrawRenderState, PrimitiveSceneProxy, Material, ShadingModels, bUsesDeferredShading);
 	}
 	else
 	{
diff --git a/Engine/Source/Runtime/Renderer/Private/MobileBasePassRendering.cpp b/Engine/Source/Runtime/Renderer/Private/MobileBasePassRendering.cpp
--- a/Engine/Source/Runtime/Renderer/Private/MobileBasePassRendering.cpp
+++ b/Engine/Source/Runtime/Renderer/Private/MobileBasePassRendering.cpp
@@ -364,6 +364,8 @@ void SetupMobileBasePassUniformParameters(
 	BasePassParameters.HalfResLocalFogVolumeViewSampler = TStaticSamplerState<SF_Bilinear, AM_Clamp, AM_Clamp, AM_Clamp>::GetRHI();
 
 	SetupReflectionUniformParameters(GraphBuilder, View, BasePassParameters.ReflectionsParameters);
+
+	SetupMobileSSRParameters(GraphBuilder, View, BasePassParameters.SSRParams);
 }
 
 TRDGUniformBufferRef<FMobileBasePassUniformParameters> CreateMobileBasePassUniformBuffer(
diff --git a/Engine/Source/Runtime/Renderer/Private/MobileBasePassRendering.h b/Engine/Source/Runtime/Renderer/Private/MobileBasePassRendering.h
--- a/Engine/Source/Runtime/Renderer/Private/MobileBasePassRendering.h
+++ b/Engine/Source/Runtime/Renderer/Private/MobileBasePassRendering.h
@@ -28,6 +28,7 @@
 #include "LocalFogVolumeRendering.h"
 #include "DBufferTextures.h"
 #include "CompositionLighting/PostProcessDeferredDecals.h"
+#include "MobileSSR.h"
 
 bool MobileLocalLightsBufferEnabled(const FStaticShaderPlatform Platform);
 bool MobileMergeLocalLightsInPrepassEnabled(const FStaticShaderPlatform Platform);
@@ -63,6 +64,7 @@ BEGIN_GLOBAL_SHADER_PARAMETER_STRUCT(FMobileBasePassUniformParameters, )
 	SHADER_PARAMETER_RDG_TEXTURE(Texture2D, HalfResLocalFogVolumeViewTexture)
 	SHADER_PARAMETER_SAMPLER(SamplerState, HalfResLocalFogVolumeViewSampler)
 	SHADER_PARAMETER_STRUCT_INCLUDE(FDBufferParameters, DBuffer)
+	SHADER_PARAMETER_STRUCT(FMobileScreenSpaceReflectionParams, SSRParams)
 END_GLOBAL_SHADER_PARAMETER_STRUCT()
 
 enum class EMobileBasePass
@@ -450,6 +452,7 @@ public:
 		OutEnvironment.SetDefine(TEXT("MOBILE_TRANSLUCENT_COLOR_TRANSMITTANCE_DUAL_SRC_BLENDING"), TranslucentColorTransmittanceMode == EMobileTranslucentColorTransmittanceMode::DUAL_SRC_BLENDING ? 1u : 0u);
 		OutEnvironment.SetDefine(TEXT("MOBILE_TRANSLUCENT_COLOR_TRANSMITTANCE_PROGRAMMABLE_BLENDING"), TranslucentColorTransmittanceMode == EMobileTranslucentColorTransmittanceMode::PROGRAMMABLE_BLENDING ? 1u : 0u);
 		OutEnvironment.SetDefine(TEXT("MOBILE_TRANSLUCENT_COLOR_TRANSMITTANCE_SINGLE_SRC_BLENDING"), TranslucentColorTransmittanceMode == EMobileTranslucentColorTransmittanceMode::SINGLE_SRC_BLENDING ? 1u : 0u);
+		OutEnvironment.SetDefine(TEXT("MOBILE_SSR_ENABLED"), AreMobileScreenSpaceReflectionsEnabled(Parameters.Platform) ? 1u : 0u);
 	}
 	
 	/** Initialization constructor. */
diff --git a/Engine/Source/Runtime/Renderer/Private/MobileDeferredShadingPass.cpp b/Engine/Source/Runtime/Renderer/Private/MobileDeferredShadingPass.cpp
--- a/Engine/Source/Runtime/Renderer/Private/MobileDeferredShadingPass.cpp
+++ b/Engine/Source/Runtime/Renderer/Private/MobileDeferredShadingPass.cpp
@@ -323,8 +323,6 @@ public:
 
 IMPLEMENT_GLOBAL_SHADER(FMobileReflectionEnvironmentSkyLightingPS, "/Engine/Private/MobileDeferredShading.usf", "MobileReflectionEnvironmentSkyLightingPS", SF_Pixel);
 
-extern uint8 GetMobileShadingModelStencilValue(FMaterialShadingModelField ShadingModel);
-
 constexpr uint32 GetLightingChannel(uint32 LightingChannelMask)
 {
 	return (LightingChannelMask & 0x1) ? 0u : ((LightingChannelMask & 0x2) ? 1u : 2u);
@@ -378,6 +376,18 @@ static void GetLightMaterial(const FCachedLightMaterial& DefaultLightMaterial, c
 	OutShader = MaterialShaderMap->GetShader<ShaderType>(PermutationId);
 }
 
+extern const uint8 MobileShadingMobelSupportStencilValue;
+
+uint8 PassShadingModelStencilValue(bool bEnableShadingModelSupport)
+{
+	return bEnableShadingModelSupport ? GET_STENCIL_MOBILE_SM_MASK(MobileShadingMobelSupportStencilValue) : STENCIL_MOBILE_DEFAULTLIT_MASK;
+}
+
+constexpr uint8 PassShadingModelStencilMask(bool bEnableShadingModelSupport)
+{
+	return bEnableShadingModelSupport ? GET_STENCIL_MOBILE_SM_MASK(0xff) : STENCIL_MOBILE_DEFAULTLIT_MASK;
+}
+
 void RenderReflectionEnvironmentSkyLighting(FRHICommandList& RHICmdList, const FScene& Scene, const FViewInfo& View)
 {
 	// Skylights with static lighting already had their diffuse contribution baked into lightmaps
@@ -409,7 +419,7 @@ void RenderReflectionEnvironmentSkyLighting(FRHICommandList& RHICmdList, const F
 		false, CF_Always,
 		true, CF_Equal, SO_Keep, SO_Keep, SO_Keep,
 		false, CF_Always, SO_Keep, SO_Keep, SO_Keep,
-		GET_STENCIL_MOBILE_SM_MASK(0xff), 0x00>::GetRHI();
+		PassShadingModelStencilMask(/*bEnableShadingModelSupport=*/false), 0x00>::GetRHI();
 
 	FMobileReflectionEnvironmentSkyLightingPS::FParameters PassParameters;
 	PassParameters.View = GetShaderBinding(View.ViewUniformBuffer);
@@ -420,15 +430,18 @@ void RenderReflectionEnvironmentSkyLighting(FRHICommandList& RHICmdList, const F
 	// Do two passes, first masking DefautLit, second masking all other shading models
 	const bool bOnlyDefaultLitInView = IsOnlyDefaultLitShadingModel(View.ShadingModelMaskInView);
 	int32 NumPasses = !bOnlyDefaultLitInView && MobileUsesGBufferCustomData(Scene.GetShaderPlatform()) ? 2 : 1;
-	uint8 PassShadingModelStencilValue[2] =
-	{
-		GetMobileShadingModelStencilValue(MSM_DefaultLit),
-		GetMobileShadingModelStencilValue(FMaterialShadingModelField())
-	};
 			
 	for (int32 PassIndex = 0; PassIndex < NumPasses; PassIndex++)
 	{
 		const bool bEnableShadingModelSupport = (PassIndex > 0);
+		if (bEnableShadingModelSupport)
+		{
+			GraphicsPSOInit.DepthStencilState = TStaticDepthStencilState<
+				false, CF_Always,
+				true, CF_Equal, SO_Keep, SO_Keep, SO_Keep,
+				false, CF_Always, SO_Keep, SO_Keep, SO_Keep,
+				PassShadingModelStencilMask(/*bEnableShadingModelSupport=*/true), 0x00>::GetRHI();
+		}
 		
 		FMobileReflectionEnvironmentSkyLightingPS::FPermutationDomain PermutationVector;
 		PermutationVector.Set<FMobileReflectionEnvironmentSkyLightingPS::FEnableShadingModelSupport>(bEnableShadingModelSupport);
@@ -443,8 +456,7 @@ void RenderReflectionEnvironmentSkyLighting(FRHICommandList& RHICmdList, const F
 		GraphicsPSOInit.BoundShaderState.PixelShaderRHI = PixelShader.GetPixelShader();
 		GraphicsPSOInit.PrimitiveType = PT_TriangleList;
 
-		uint8 StencilRef = GET_STENCIL_MOBILE_SM_MASK(PassShadingModelStencilValue[PassIndex]);
-		SetGraphicsPipelineState(RHICmdList, GraphicsPSOInit, StencilRef);
+		SetGraphicsPipelineState(RHICmdList, GraphicsPSOInit, PassShadingModelStencilValue(bEnableShadingModelSupport));
 		SetShaderParameters(RHICmdList, PixelShader, PixelShader.GetPixelShader(), PassParameters);
 
 		const FIntPoint TargetSize = View.GetSceneTexturesConfig().Extent;
@@ -461,14 +473,43 @@ void RenderReflectionEnvironmentSkyLighting(FRHICommandList& RHICmdList, const F
 	}
 }
 
-template<uint32 LightingChannelIdx>
+template<uint32 LightingChannelIdx, bool bEnableShadingModelSupport>
 static void SetDirectionalLightDepthStencilState(FGraphicsPipelineStateInitializer& GraphicsPSOInit)
 {
 	GraphicsPSOInit.DepthStencilState = TStaticDepthStencilState<
 		false, CF_Always,
 		true, CF_Equal, SO_Keep, SO_Keep, SO_Keep,
 		false, CF_Always, SO_Keep, SO_Keep, SO_Keep,
-		GET_STENCIL_MOBILE_SM_MASK(0xff) | STENCIL_LIGHTING_CHANNELS_MASK(1u << LightingChannelIdx), 0x00>::GetRHI();
+		PassShadingModelStencilMask(bEnableShadingModelSupport) | STENCIL_LIGHTING_CHANNELS_MASK(1u << LightingChannelIdx), 0x00>::GetRHI();
+}
+
+template<bool bEnableShadingModelSupport>
+static void SetDirectionalLightDepthStencilState(FGraphicsPipelineStateInitializer& GraphicsPSOInit, uint32 LightingChannelIdx)
+{
+	switch (LightingChannelIdx)
+	{
+	default:
+		SetDirectionalLightDepthStencilState<0, bEnableShadingModelSupport>(GraphicsPSOInit);
+		break;
+	case 1:
+		SetDirectionalLightDepthStencilState<1, bEnableShadingModelSupport>(GraphicsPSOInit);
+		break;
+	case 2:
+		SetDirectionalLightDepthStencilState<2, bEnableShadingModelSupport>(GraphicsPSOInit);
+		break;
+	}
+}
+
+static void SetDirectionalLightDepthStencilState(FGraphicsPipelineStateInitializer& GraphicsPSOInit, uint32 LightingChannelIdx, bool bEnableShadingModelSupport)
+{
+	if (bEnableShadingModelSupport)
+	{
+		SetDirectionalLightDepthStencilState<true>(GraphicsPSOInit, LightingChannelIdx);
+	}
+	else
+	{
+		SetDirectionalLightDepthStencilState<false>(GraphicsPSOInit, LightingChannelIdx);
+	}
 }
 
 static void RenderDirectionalLight(FRHICommandList& RHICmdList, const FScene& Scene, const FViewInfo& View, const FCachedLightMaterial& DefaultLightMaterial, const FLightSceneInfo& DirectionalLight, uint32 LightingChannel, bool bInlineReflectionAndSky)
@@ -481,19 +522,7 @@ static void RenderDirectionalLight(FRHICommandList& RHICmdList, const FScene& Sc
 	RHICmdList.ApplyCachedRenderTargets(GraphicsPSOInit);
 	GraphicsPSOInit.RasterizerState = TStaticRasterizerState<>::GetRHI();
 
-	uint8 LightingChannelStencilValue = GetLightingChannelStencilValue(LightingChannel);
-	if (LightingChannel == 1u)
-	{
-		SetDirectionalLightDepthStencilState<1u>(GraphicsPSOInit);
-	}
-	else if (LightingChannel == 2u)
-	{
-		SetDirectionalLightDepthStencilState<2u>(GraphicsPSOInit);
-	}
-	else
-	{
-		SetDirectionalLightDepthStencilState<0u>(GraphicsPSOInit);
-	}
+	const uint8 LightingChannelStencilValue = GetLightingChannelStencilValue(LightingChannel);
 	
 	TShaderMapRef<FPostProcessVS> VertexShader(View.ShaderMap);
 	
@@ -550,15 +579,11 @@ static void RenderDirectionalLight(FRHICommandList& RHICmdList, const FScene& Sc
 	// Do two passes, first masking DefautLit, second masking all other shading models
 	const bool bOnlyDefaultLitInView = IsOnlyDefaultLitShadingModel(View.ShadingModelMaskInView);
 	int32 NumPasses = !bOnlyDefaultLitInView && MobileUsesGBufferCustomData(Scene.GetShaderPlatform()) ? 2 : 1;
-	uint8 PassShadingModelStencilValue[2] =
-	{
-		GetMobileShadingModelStencilValue(MSM_DefaultLit),
-		GetMobileShadingModelStencilValue(FMaterialShadingModelField())
-	};
 	
 	for (int32 PassIndex = 0; PassIndex < NumPasses; ++PassIndex)
 	{
 		const bool bEnableShadingModelSupport = (PassIndex > 0);
+		SetDirectionalLightDepthStencilState(GraphicsPSOInit, LightingChannelStencilValue, bEnableShadingModelSupport);
 		
 		FMobileDirectionalLightFunctionPS::FPermutationDomain PermutationVector = FMobileDirectionalLightFunctionPS::BuildPermutationVector(
 			View,
@@ -573,9 +598,7 @@ static void RenderDirectionalLight(FRHICommandList& RHICmdList, const FScene& Sc
 		TShaderRef<FMobileDirectionalLightFunctionPS> PixelShader;
 		GetLightMaterial(DefaultLightMaterial, LightFunctionMaterialProxy, PermutationVector.ToDimensionValueId(), LightMaterial, PixelShader);
 
-		uint8 StencilRef = 
-			GET_STENCIL_MOBILE_SM_MASK(PassShadingModelStencilValue[PassIndex]) | 
-			STENCIL_LIGHTING_CHANNELS_MASK(LightingChannelStencilValue);
+		const uint8 StencilRef = PassShadingModelStencilValue(bEnableShadingModelSupport) | STENCIL_LIGHTING_CHANNELS_MASK(LightingChannelStencilValue);
 
 		GraphicsPSOInit.BoundShaderState.VertexDeclarationRHI = GFilterVertexDeclaration.VertexDeclarationRHI;
 		GraphicsPSOInit.BoundShaderState.VertexShaderRHI = VertexShader.GetVertexShader();
@@ -625,7 +648,7 @@ static void RenderDirectionalLights(FRHICommandList& RHICmdList, const FScene& S
 	}
 }
 
-template<uint32 LightingChannel, bool bWithStencilCulling>
+template<uint32 LightingChannel, bool bWithStencilCulling, bool bEnableShadingModelSupport>
 static void SetLocalLightRasterizerAndDepthState(FGraphicsPipelineStateInitializer& GraphicsPSOInit, const FViewInfo& View, const FSphere& LightBounds)
 {
 	if (bWithStencilCulling)
@@ -637,7 +660,7 @@ static void SetLocalLightRasterizerAndDepthState(FGraphicsPipelineStateInitializ
 			false, CF_LessEqual,
 			false, CF_Equal, SO_Keep, SO_Keep, SO_Keep,
 			true, CF_Equal, SO_Zero, SO_Keep, SO_Zero,
-			GET_STENCIL_MOBILE_SM_MASK(0xff) | STENCIL_LIGHTING_CHANNELS_MASK(1u << LightingChannel) | STENCIL_SANDBOX_MASK,
+			PassShadingModelStencilMask(bEnableShadingModelSupport) | STENCIL_LIGHTING_CHANNELS_MASK(1u << LightingChannel) | STENCIL_SANDBOX_MASK,
 			STENCIL_SANDBOX_MASK
 		>::GetRHI();
 	}
@@ -657,7 +680,7 @@ static void SetLocalLightRasterizerAndDepthState(FGraphicsPipelineStateInitializ
 				false, CF_Always,
 				true, CF_Equal, SO_Keep, SO_Keep, SO_Keep,
 				false, CF_Always, SO_Keep, SO_Keep, SO_Keep,
-				GET_STENCIL_MOBILE_SM_MASK(0xff) | STENCIL_LIGHTING_CHANNELS_MASK(1u << LightingChannel), 0x00>::GetRHI();
+				PassShadingModelStencilMask(bEnableShadingModelSupport) | STENCIL_LIGHTING_CHANNELS_MASK(1u << LightingChannel), 0x00>::GetRHI();
 		}
 		else
 		{
@@ -667,21 +690,51 @@ static void SetLocalLightRasterizerAndDepthState(FGraphicsPipelineStateInitializ
 				false, CF_DepthNearOrEqual,
 				true, CF_Equal, SO_Keep, SO_Keep, SO_Keep,
 				false, CF_Always, SO_Keep, SO_Keep, SO_Keep,
-				GET_STENCIL_MOBILE_SM_MASK(0xff) | STENCIL_LIGHTING_CHANNELS_MASK(1u << LightingChannel), 0x00>::GetRHI();
+				PassShadingModelStencilMask(bEnableShadingModelSupport) | STENCIL_LIGHTING_CHANNELS_MASK(1u << LightingChannel), 0x00>::GetRHI();
 		}
 	}
 }
 
-template<uint32 LightingChannel>
-static void SetLocalLightRasterizerAndDepthState(FGraphicsPipelineStateInitializer& GraphicsPSOInit, const FViewInfo& View, const FSphere& LightBounds)
+template <bool bEnableShadingModelSupport, bool bWithStencilCulling>
+static void SetLocalLightRasterizerAndDepthState(FGraphicsPipelineStateInitializer& GraphicsPSOInit, const FViewInfo& View, const FSphere& LightBounds, uint32 LightingChannel)
+{
+	// TODO: support multi-channel ligths?
+	switch (LightingChannel)
+	{
+	default:
+		SetLocalLightRasterizerAndDepthState<0, bWithStencilCulling, bEnableShadingModelSupport>(GraphicsPSOInit, View, LightBounds);
+		break;
+	case 1:
+		SetLocalLightRasterizerAndDepthState<1, bWithStencilCulling, bEnableShadingModelSupport>(GraphicsPSOInit, View, LightBounds);
+		break;
+	case 2:
+		SetLocalLightRasterizerAndDepthState<2, bWithStencilCulling, bEnableShadingModelSupport>(GraphicsPSOInit, View, LightBounds);
+		break;
+	}
+}
+
+template <bool bEnableShadingModelSupport>
+static void SetLocalLightRasterizerAndDepthState(FGraphicsPipelineStateInitializer& GraphicsPSOInit, const FViewInfo& View, const FSphere& LightBounds, uint32 LightingChannel)
 {
 	if (GMobileUseLightStencilCulling != 0)
 	{
-		SetLocalLightRasterizerAndDepthState<LightingChannel, true>(GraphicsPSOInit, View, LightBounds);
+		SetLocalLightRasterizerAndDepthState<bEnableShadingModelSupport, true>(GraphicsPSOInit, View, LightBounds, LightingChannel);
+	}
+	else
+	{
+		SetLocalLightRasterizerAndDepthState<bEnableShadingModelSupport, false>(GraphicsPSOInit, View, LightBounds, LightingChannel);
+	}
+}
+
+static void SetLocalLightRasterizerAndDepthState(FGraphicsPipelineStateInitializer& GraphicsPSOInit, const FViewInfo& View, const FSphere& LightBounds, uint32 LightingChannel, bool bEnableShadingModelSupport)
+{
+	if (bEnableShadingModelSupport)
+	{
+		SetLocalLightRasterizerAndDepthState<true>(GraphicsPSOInit, View, LightBounds, LightingChannel);
 	}
 	else
 	{
-		SetLocalLightRasterizerAndDepthState<LightingChannel, false>(GraphicsPSOInit, View, LightBounds);
+		SetLocalLightRasterizerAndDepthState<false>(GraphicsPSOInit, View, LightBounds, LightingChannel);
 	}
 }
 
@@ -764,22 +817,8 @@ static void RenderLocalLight(
 	GraphicsPSOInit.PrimitiveType = PT_TriangleList;
 	const FSphere LightBounds = LightSceneInfo.Proxy->GetBoundingSphere();
 
-	uint32 LightingChannel = GetLightingChannel(LightingChannelMask);
-	uint8 LightingChannelStencilValue = GetLightingChannelStencilValue(LightingChannel);
-
-	// TODO: support multi-channel ligths?
-	if (LightingChannel == 1u)
-	{
-		SetLocalLightRasterizerAndDepthState<1u>(GraphicsPSOInit, View, LightBounds);
-	}
-	else if (LightingChannel == 2u)
-	{
-		SetLocalLightRasterizerAndDepthState<2u>(GraphicsPSOInit, View, LightBounds);
-	}
-	else
-	{
-		SetLocalLightRasterizerAndDepthState<0u>(GraphicsPSOInit, View, LightBounds);
-	}
+	const uint32 LightingChannel = GetLightingChannel(LightingChannelMask);
+	const uint8 LightingChannelStencilValue = GetLightingChannelStencilValue(LightingChannel);
 
 	FDeferredLightVS::FPermutationDomain PermutationVectorVS;
 	PermutationVectorVS.Set<FDeferredLightVS::FRadialLight>(true);
@@ -808,15 +847,11 @@ static void RenderLocalLight(
 	// Do two passes, first masking DefautLit, second masking all other shading models
 	const bool bOnlyDefaultLitInView = IsOnlyDefaultLitShadingModel(View.ShadingModelMaskInView);
 	int32 NumPasses = !bOnlyDefaultLitInView && MobileUsesGBufferCustomData(Scene.GetShaderPlatform()) ? 2 : 1;
-	uint8 PassShadingModelStencilValue[2] =
-	{
-		GetMobileShadingModelStencilValue(MSM_DefaultLit),
-		GetMobileShadingModelStencilValue(FMaterialShadingModelField())
-	};
 
 	for (int32 PassIndex = 0; PassIndex < NumPasses; PassIndex++)
 	{
 		const bool bEnableShadingModelSupport = (PassIndex > 0);
+		SetLocalLightRasterizerAndDepthState(GraphicsPSOInit, View, LightBounds, LightingChannel, bEnableShadingModelSupport);
 
 		FMobileRadialLightFunctionPS::FPermutationDomain PermutationVector;
 		PermutationVector.Set<FMobileRadialLightFunctionPS::FEnableShadingModelSupport>(bEnableShadingModelSupport);
@@ -831,9 +866,7 @@ static void RenderLocalLight(
 		GraphicsPSOInit.BoundShaderState.VertexShaderRHI = VertexShader.GetVertexShader();
 		GraphicsPSOInit.BoundShaderState.PixelShaderRHI = PixelShader.GetPixelShader();
 
-		uint8 StencilRef = 
-			GET_STENCIL_MOBILE_SM_MASK(PassShadingModelStencilValue[PassIndex]) | 
-			STENCIL_LIGHTING_CHANNELS_MASK(LightingChannelStencilValue);
+		const uint8 StencilRef = PassShadingModelStencilValue(bEnableShadingModelSupport) | STENCIL_LIGHTING_CHANNELS_MASK(LightingChannelStencilValue);
 
 		SetGraphicsPipelineState(RHICmdList, GraphicsPSOInit, StencilRef);
 
@@ -897,11 +930,6 @@ static void RenderSimpleLights(
 	// Do two passes, first masking DefautLit, second masking all other shading models
 	const bool bOnlyDefaultLitInView = IsOnlyDefaultLitShadingModel(View.ShadingModelMaskInView);
 	int32 NumPasses = !bOnlyDefaultLitInView && MobileUsesGBufferCustomData(Scene.GetShaderPlatform()) ? 2 : 1;
-	uint8 PassShadingModelStencilValue[2] =
-	{
-		GetMobileShadingModelStencilValue(MSM_DefaultLit),
-		GetMobileShadingModelStencilValue(FMaterialShadingModelField())
-	};
 	TShaderRef<FMobileRadialLightFunctionPS> PassPixelShaders[2];
 	FGraphicsPipelineStateInitializer GraphicsPSOLight[2];
 
@@ -915,7 +943,14 @@ static void RenderSimpleLights(
 		GraphicsPSOLight[PassIndex].PrimitiveType = PT_TriangleList;
 		GraphicsPSOLight[PassIndex].BoundShaderState.VertexDeclarationRHI = GetVertexDeclarationFVector4();
 		GraphicsPSOLight[PassIndex].BoundShaderState.VertexShaderRHI = VertexShader.GetVertexShader();
-		SetLocalLightRasterizerAndDepthState<0u, true>(GraphicsPSOLight[PassIndex], View, FSphere());
+		if (bEnableShadingModelSupport)
+		{
+			SetLocalLightRasterizerAndDepthState<0u, true, /*bEnableShadingModelSupport=*/true>(GraphicsPSOLight[PassIndex], View, FSphere());
+		}
+		else
+		{
+			SetLocalLightRasterizerAndDepthState<0u, true, /*bEnableShadingModelSupport=*/false>(GraphicsPSOLight[PassIndex], View, FSphere());
+		}
 
 		TShaderRef<FMobileRadialLightFunctionPS> PixelShader;
 		FMobileRadialLightFunctionPS::FPermutationDomain PermutationVector;
@@ -951,7 +986,7 @@ static void RenderSimpleLights(
 
 		for (int32 PassIndex = 0; PassIndex < NumPasses; ++PassIndex)
 		{
-			uint8 StencilRef = GET_STENCIL_MOBILE_SM_MASK(PassShadingModelStencilValue[PassIndex]);
+			const uint8 StencilRef = PassShadingModelStencilValue(PassIndex > 0);
 			SetGraphicsPipelineState(RHICmdList, GraphicsPSOLight[PassIndex], StencilRef);
 
 			SetShaderParameters(RHICmdList, VertexShader, VertexShader.GetVertexShader(), ParametersVS);
diff --git a/Engine/Source/Runtime/Renderer/Private/MobileSSR.cpp b/Engine/Source/Runtime/Renderer/Private/MobileSSR.cpp
new file mode 100644
--- /dev/null
+++ b/Engine/Source/Runtime/Renderer/Private/MobileSSR.cpp
@@ -0,0 +1,180 @@
+// Copyright Epic Games, Inc. All Rights Reserved.
+
+#include "MobileSSR.h"
+
+#include "RenderGraph.h"
+#include "PixelShaderUtils.h"
+#include "ScreenPass.h"
+#include "PostProcess/SceneRenderTargets.h"
+#include "ScenePrivate.h"
+#include "SystemTextures.h"
+
+int32 GMobileScreenSpaceReflectionsSupported = 0;
+static FAutoConsoleVariableRef CVarMobileScreenSpaceReflections(
+	TEXT("r.Mobile.ScreenSpaceReflections"),
+	GMobileScreenSpaceReflectionsSupported,
+	TEXT("0: Mobile Renderer Screen Space Reflections disabled (default)\n"
+		 "1: Mobile Renderer Screen Space Reflections enabled\n"),
+	ECVF_ReadOnly | ECVF_RenderThreadSafe);
+
+bool ShouldRenderMobileSSR(const FViewInfo& View)
+{
+	return GMobileScreenSpaceReflectionsSupported && ScreenSpaceRayTracing::ShouldRenderScreenSpaceReflections(View) && View.PrevViewInfo.TemporalAAHistory.RT[0] && View.HZB;
+}
+
+bool IsMobileSSREnabled(const FViewInfo& View)
+{
+	return GMobileScreenSpaceReflectionsSupported && ScreenSpaceRayTracing::ShouldRenderScreenSpaceReflections(View);
+}
+
+class FMobileSSRQualityDim : SHADER_PERMUTATION_ENUM_CLASS("SSR_QUALITY", ESSRQuality);
+class FMobileSSRUseVelocity : SHADER_PERMUTATION_BOOL("SSR_USE_VELOCITY");
+
+class FMobileScreenSpaceReflectionsPS : public FGlobalShader
+{
+	DECLARE_GLOBAL_SHADER(FMobileScreenSpaceReflectionsPS);
+	SHADER_USE_PARAMETER_STRUCT(FMobileScreenSpaceReflectionsPS, FGlobalShader);
+
+	using FPermutationDomain = TShaderPermutationDomain<FMobileSSRQualityDim, FMobileSSRUseVelocity>;
+
+	static bool ShouldCompilePermutation(const FGlobalShaderPermutationParameters& Parameters)
+	{
+		return IsMobilePlatform(Parameters.Platform) && IsMobileDeferredShadingEnabled(Parameters.Platform);
+	}
+
+	static void ModifyCompilationEnvironment(const FGlobalShaderPermutationParameters& Parameters, FShaderCompilerEnvironment& OutEnvironment)
+	{
+		FGlobalShader::ModifyCompilationEnvironment(Parameters, OutEnvironment);
+		OutEnvironment.SetDefine(TEXT("SUPPORTS_ANISOTROPIC_MATERIALS"), FDataDrivenShaderPlatformInfo::GetSupportsAnisotropicMaterials(Parameters.Platform));
+
+		check(IsMobileDeferredShadingEnabled(Parameters.Platform));
+		OutEnvironment.SetDefine(TEXT("ENABLE_SHADINGMODEL_SUPPORT_MOBILE_DEFERRED"), MobileUsesGBufferCustomData(Parameters.Platform));
+		OutEnvironment.SetDefine(TEXT("IS_MOBILE_DEFERREDSHADING_SUBPASS"), 1u);
+		OutEnvironment.SetDefine(TEXT("FORCE_DEPTH_TEXTURE_READS"), 1u);
+		
+		OutEnvironment.SetDefine(TEXT("SSR_OUTPUT_FOR_DENOISER"), 0u);
+		OutEnvironment.SetDefine(TEXT("DIM_LIGHTING_TERM"), 0u);
+	}
+
+	BEGIN_SHADER_PARAMETER_STRUCT(FParameters, )
+		SHADER_PARAMETER(FLinearColor, SSRParams)
+		SHADER_PARAMETER_STRUCT_REF(FViewUniformShaderParameters, ViewUniformBuffer)
+	END_SHADER_PARAMETER_STRUCT()
+};
+
+IMPLEMENT_GLOBAL_SHADER(FMobileScreenSpaceReflectionsPS, "/Engine/Private/SSRT/SSRTReflections.usf", "MobileScreenSpaceReflectionsPS", SF_Pixel);
+
+
+using DepthStencilStateMobileSSR = TStaticDepthStencilState<
+	false, CF_Always,
+	true, CF_Equal, SO_Keep, SO_Keep, SO_Keep,
+	false, CF_Always, SO_Keep, SO_Keep, SO_Keep,
+	STENCIL_MOBILE_REFLECTIVE_MASK | STENCIL_MOBILE_SKY_MASK, 0x00>;
+
+const uint8 StencilRefMobileSSR = STENCIL_MOBILE_REFLECTIVE_MASK; // Process SSR where there are potentially reflective materials (and not sky)
+
+void SetupMobileSSRParameters(FRDGBuilder& GraphBuilder,const FViewInfo& View, FMobileScreenSpaceReflectionParams& Params)
+{
+	if (!ShouldRenderMobileSSR(View))
+	{
+		const FRDGSystemTextures& SystemTextures = FRDGSystemTextures::Get(GraphBuilder);
+		Params.HZB = SystemTextures.Black;
+		Params.HZBSampler = TStaticSamplerState<SF_Point, AM_Clamp, AM_Clamp, AM_Clamp>::GetRHI();
+		Params.SceneColor = GraphBuilder.CreateSRV(FRDGTextureSRVDesc(SystemTextures.Black));
+		Params.SceneColorSampler = TStaticSamplerState<SF_Bilinear, AM_Clamp, AM_Clamp, AM_Clamp>::GetRHI();
+		Params.PrevSceneColorBilinearUVMinMax = FVector4f(0.0f, 0.0f, 1.0f, 1.0f);
+		Params.QualityAndExposureCorrection = FVector4f(ForceInitToZero);
+		return;
+	}
+	FRDGTextureRef SceneColor = GraphBuilder.RegisterExternalTexture(View.PrevViewInfo.TemporalAAHistory.RT[0]);
+	Params.SceneColor = GraphBuilder.CreateSRV(SceneColor->Desc.IsTextureArray()
+		? FRDGTextureSRVDesc::CreateForSlice(SceneColor, View.PrevViewInfo.TemporalAAHistory.OutputSliceIndex)
+		: FRDGTextureSRVDesc(SceneColor));
+	Params.SceneColorSampler = TStaticSamplerState<SF_Point>::GetRHI();
+	Params.HZB = View.HZB;
+	Params.HZBSampler = TStaticSamplerState<SF_Point>::GetRHI();
+
+	{
+		const FVector2D HZBUvFactor(
+			float(View.ViewRect.Width()) / float(2 * View.HZBMipmap0Size.X),
+			float(View.ViewRect.Height()) / float(2 * View.HZBMipmap0Size.Y));
+		Params.HZBUvFactorAndInvFactor = FVector4f(
+			HZBUvFactor.X,
+			HZBUvFactor.Y,
+			1.0f / HZBUvFactor.X,
+			1.0f / HZBUvFactor.Y);
+	}
+
+	{
+		ensure(View.PrevViewInfo.TemporalAAHistory.IsValid());
+		FIntPoint ViewportOffset = View.PrevViewInfo.TemporalAAHistory.ViewportRect.Min;
+		FIntPoint ViewportExtent = View.PrevViewInfo.TemporalAAHistory.ViewportRect.Size();
+		FIntPoint BufferSize = View.PrevViewInfo.TemporalAAHistory.ReferenceBufferSize;
+		ensure(ViewportExtent.X > 0 && ViewportExtent.Y > 0);
+		ensure(BufferSize.X > 0 && BufferSize.Y > 0);
+
+		FVector2D InvBufferSize(1.0f / float(BufferSize.X), 1.0f / float(BufferSize.Y));
+		Params.PrevScreenPositionScaleBias = FVector4f(
+			ViewportExtent.X * 0.5f * InvBufferSize.X,
+			-ViewportExtent.Y * 0.5f * InvBufferSize.Y,
+			(ViewportExtent.X * 0.5f + ViewportOffset.X) * InvBufferSize.X,
+			(ViewportExtent.Y * 0.5f + ViewportOffset.Y) * InvBufferSize.Y);
+	}
+
+	{
+		ESSRQuality SSRQuality;
+		IScreenSpaceDenoiser::FReflectionsRayTracingConfig DenoiserConfig;
+		ScreenSpaceRayTracing::GetSSRQualityForView(View, &SSRQuality, &DenoiserConfig);
+
+		Params.QualityAndExposureCorrection.X = (float)SSRQuality;
+		Params.QualityAndExposureCorrection.Y = 1.0f / View.PrevViewInfo.SceneColorPreExposure;
+		Params.QualityAndExposureCorrection.Z = View.PreExposure / View.PrevViewInfo.SceneColorPreExposure;
+		Params.QualityAndExposureCorrection.W = FMath::Clamp(View.FinalPostProcessSettings.ScreenSpaceReflectionMaxRoughness, 0.0f, 1.0f);
+	}
+
+	{
+		FScreenPassTextureViewportParameters PrevSceneColorParameters = GetScreenPassTextureViewportParameters(FScreenPassTextureViewport(SceneColor->Desc.Extent, View.PrevViewInfo.TemporalAAHistory.ViewportRect));
+		Params.PrevSceneColorBilinearUVMinMax.X = PrevSceneColorParameters.UVViewportBilinearMin.X;
+		Params.PrevSceneColorBilinearUVMinMax.Y = PrevSceneColorParameters.UVViewportBilinearMin.Y;
+		Params.PrevSceneColorBilinearUVMinMax.Z = PrevSceneColorParameters.UVViewportBilinearMax.X;
+		Params.PrevSceneColorBilinearUVMinMax.W = PrevSceneColorParameters.UVViewportBilinearMax.Y;
+	}
+}
+
+void FMobileSceneRenderer::RenderSSR(FRHICommandList& RHICmdList, const FViewInfo& View)
+{
+	if (!ShouldRenderMobileSSR(View))
+	{
+		return;
+	}
+
+	FMobileScreenSpaceReflectionsPS::FParameters Parameters;
+
+	ESSRQuality SSRQuality;
+	IScreenSpaceDenoiser::FReflectionsRayTracingConfig DenoiserConfig;
+	ScreenSpaceRayTracing::GetSSRQualityForView(View, &SSRQuality, &DenoiserConfig);
+	if (SSRQuality < ESSRQuality::Low)
+	{
+		return;
+	}
+
+	Parameters.SSRParams = ScreenSpaceRayTracing::ComputeSSRParams(View, SSRQuality, false);
+
+	Parameters.ViewUniformBuffer = View.ViewUniformBuffer;
+
+	FMobileScreenSpaceReflectionsPS::FPermutationDomain PermutationVectorPS;
+	PermutationVectorPS.Set<FMobileSSRQualityDim>(SSRQuality);
+	PermutationVectorPS.Set<FMobileSSRUseVelocity>(bShouldRenderVelocities);
+
+	TShaderMapRef<FMobileScreenSpaceReflectionsPS> PixelShader(View.ShaderMap, PermutationVectorPS);
+	FGraphicsPipelineStateInitializer GraphicsPSOInit;
+	FPixelShaderUtils::InitFullscreenPipelineState(RHICmdList, View.ShaderMap, PixelShader, /* out */ GraphicsPSOInit);
+
+	GraphicsPSOInit.BlendState = TStaticBlendState<CW_RGB, BO_Add, BF_One, BF_One>::GetRHI();
+	GraphicsPSOInit.DepthStencilState = DepthStencilStateMobileSSR::GetRHI();
+
+	SetGraphicsPipelineState(RHICmdList, GraphicsPSOInit, StencilRefMobileSSR);
+	SetShaderParameters(RHICmdList, PixelShader, PixelShader.GetPixelShader(), Parameters);
+
+	FPixelShaderUtils::DrawFullscreenTriangle(RHICmdList);
+}
diff --git a/Engine/Source/Runtime/Renderer/Private/MobileSSR.h b/Engine/Source/Runtime/Renderer/Private/MobileSSR.h
new file mode 100644
--- /dev/null
+++ b/Engine/Source/Runtime/Renderer/Private/MobileSSR.h
@@ -0,0 +1,24 @@
+// Copyright Epic Games, Inc.All Rights Reserved.
+
+#pragma once
+
+#include "CoreMinimal.h"
+#include "RenderGraphDefinitions.h"
+#include "SceneRendering.h"
+#include "ScreenSpaceRayTracing.h"
+
+BEGIN_SHADER_PARAMETER_STRUCT(FMobileScreenSpaceReflectionParams, )
+SHADER_PARAMETER_RDG_TEXTURE_SRV(Texture2D, SceneColor)
+SHADER_PARAMETER_SAMPLER(SamplerState, SceneColorSampler)
+SHADER_PARAMETER_RDG_TEXTURE(Texture2D, HZB)
+SHADER_PARAMETER_SAMPLER(SamplerState, HZBSampler)
+SHADER_PARAMETER(FVector4f, HZBUvFactorAndInvFactor)
+SHADER_PARAMETER(FVector4f, PrevScreenPositionScaleBias)
+SHADER_PARAMETER(FVector4f, PrevSceneColorBilinearUVMinMax)
+SHADER_PARAMETER(FVector4f, QualityAndExposureCorrection) // .x = SSRQuality, .y = PrevSceneColorPreExposureInv, .z = PrevSceneColorPreExposureCorrection, .w = View.FinalPostProcessSettings.ScreenSpaceReflectionMaxRoughness
+END_SHADER_PARAMETER_STRUCT()
+
+bool IsMobileSSREnabled(const FViewInfo& View);
+bool ShouldRenderMobileSSR(const FViewInfo& View);
+
+void SetupMobileSSRParameters(FRDGBuilder& GraphBuilder, const FViewInfo& View, FMobileScreenSpaceReflectionParams& Params);
diff --git a/Engine/Source/Runtime/Renderer/Private/MobileShadingRenderer.cpp b/Engine/Source/Runtime/Renderer/Private/MobileShadingRenderer.cpp
--- a/Engine/Source/Runtime/Renderer/Private/MobileShadingRenderer.cpp
+++ b/Engine/Source/Runtime/Renderer/Private/MobileShadingRenderer.cpp
@@ -76,6 +76,7 @@
 #include "Rendering/CustomRenderPass.h"
 #include "GenerateMips.h"
 #include "StereoRenderUtils.h"
+#include "MobileSSR.h"
 
 uint32 GetShadowQuality();
 
@@ -563,7 +564,7 @@ void FMobileSceneRenderer::InitViews(
 		&& !ViewFamily.UseDebugViewPS()
 		&& bRendererOutputFinalSceneColor;
 
-	bShouldRenderHZB = ShouldRenderHZB() && bRendererOutputFinalSceneColor;
+	bShouldRenderHZB = ShouldRenderHZB(Views) && bRendererOutputFinalSceneColor;
 
 	// Whether we need to store depth for post-processing
 	// On PowerVR we see flickering of shadows and depths not updating correctly if targets are discarded.
@@ -1828,9 +1829,7 @@ void MobileDeferredCopyBuffer(FRHICommandList& RHICmdList, const FViewInfo& View
 	RHICmdList.ApplyCachedRenderTargets(GraphicsPSOInit);
 	GraphicsPSOInit.BlendState = TStaticBlendState<>::GetRHI();
 	GraphicsPSOInit.RasterizerState = TStaticRasterizerState<>::GetRHI();
-	// Shade only MSM_DefaultLit pixels
-	uint8 StencilRef = GET_STENCIL_MOBILE_SM_MASK(MSM_DefaultLit);
-	GraphicsPSOInit.DepthStencilState = TStaticDepthStencilState<false, CF_Always>::GetRHI(); // 4 bits for shading models
+	GraphicsPSOInit.DepthStencilState = TStaticDepthStencilState<false, CF_Always>::GetRHI();
 
 	TShaderMapRef<FPostProcessVS> VertexShader(View.ShaderMap);
 	TShaderMapRef<T> PixelShader(View.ShaderMap);
@@ -1839,7 +1838,7 @@ void MobileDeferredCopyBuffer(FRHICommandList& RHICmdList, const FViewInfo& View
 	GraphicsPSOInit.BoundShaderState.VertexShaderRHI = VertexShader.GetVertexShader();
 	GraphicsPSOInit.BoundShaderState.PixelShaderRHI = PixelShader.GetPixelShader();
 	GraphicsPSOInit.PrimitiveType = PT_TriangleList;
-	SetGraphicsPipelineState(RHICmdList, GraphicsPSOInit, StencilRef);
+	SetGraphicsPipelineState(RHICmdList, GraphicsPSOInit, 0u);
 
 	DrawRectangle(
 		RHICmdList,
@@ -1979,6 +1978,8 @@ void FMobileSceneRenderer::RenderDeferredSinglePass(FRDGBuilder& GraphBuilder, F
 				// SceneColor write, SceneDepth is read only
 				RHICmdList.NextSubpass();
 				MobileDeferredShadingPass(RHICmdList, ViewContext.ViewIndex, Views.Num(), View, *Scene, SortedLightSet, VisibleLightInfos);
+				RenderSSR(RHICmdList, View);
+
 				if (bUsingPixelLocalStorage)
 				{
 					MobileDeferredCopyBuffer<FMobileDeferredCopyPLSPS>(RHICmdList, View);
@@ -2149,6 +2150,7 @@ void FMobileSceneRenderer::RenderDeferredMultiPass(FRDGBuilder& GraphBuilder, FS
 				FViewInfo& View = *ViewContext.ViewInfo;
 
 				MobileDeferredShadingPass(RHICmdList, ViewContext.ViewIndex, Views.Num(), View, *Scene, SortedLightSet, VisibleLightInfos);
+				RenderSSR(RHICmdList, View);
 				RenderFog(RHICmdList, View);
 
 				// Draw translucency.
@@ -2351,7 +2353,7 @@ void FMobileSceneRenderer::PreTonemapMSAA(FRHICommandList& RHICmdList, const FMi
 		EDRF_UseTriangleOptimization);
 }
 
-bool FMobileSceneRenderer::ShouldRenderHZB()
+bool FMobileSceneRenderer::ShouldRenderHZB(TArrayView<FViewInfo> InViews)
 {
 	static const auto MobileAmbientOcclusionTechniqueCVar = IConsoleManager::Get().FindTConsoleVariableDataInt(TEXT("r.Mobile.AmbientOcclusionTechnique"));
 
@@ -2366,6 +2368,18 @@ bool FMobileSceneRenderer::ShouldRenderHZB()
 
 	bool bNeedsHZB = bIsFeatureRequested;
 
+	if (!bNeedsHZB)
+	{
+		for (const FViewInfo& View : InViews)
+		{
+			if (IsMobileSSREnabled(View))
+			{
+				bNeedsHZB = true;
+				break;
+			}
+		}
+	}
+
 	return bNeedsHZB;
 }
 
diff --git a/Engine/Source/Runtime/Renderer/Private/PostProcess/PostProcessing.cpp b/Engine/Source/Runtime/Renderer/Private/PostProcess/PostProcessing.cpp
--- a/Engine/Source/Runtime/Renderer/Private/PostProcess/PostProcessing.cpp
+++ b/Engine/Source/Runtime/Renderer/Private/PostProcess/PostProcessing.cpp
@@ -73,6 +73,7 @@
 #include "IlluminanceMeter.h"
 #include "SparseVolumeTexture/SparseVolumeTextureStreamingVisualize.h"
 #include "CanvasItem.h"
+#include "MobileSSR.h"
 
 bool IsMobileEyeAdaptationEnabled(const FViewInfo& View);
 
@@ -2733,7 +2734,7 @@ void AddMobilePostProcessingPasses(FRDGBuilder& GraphBuilder, FScene* Scene, con
 			EMainTAAPassConfig TAAConfig = GetMainTAAPassConfig(View);
 			checkSlow(TAAConfig != EMainTAAPassConfig::Disabled);
 
-			FDefaultTemporalUpscaler::FInputs UpscalerPassInputs;
+			FDefaultTemporalUpscaler::FInputs UpscalerPassInputs{};
 			UpscalerPassInputs.SceneColor = FScreenPassTexture(SceneColor.Texture, View.ViewRect);
 			UpscalerPassInputs.SceneDepth = FScreenPassTexture(SceneDepth.Texture, View.ViewRect);
 			UpscalerPassInputs.SceneVelocity = FScreenPassTexture(Velocity.Texture, View.ViewRect);
@@ -2759,6 +2760,24 @@ void AddMobilePostProcessingPasses(FRDGBuilder& GraphBuilder, FScene* Scene, con
 			}
 			SceneColor = FScreenPassTexture(Outputs.FullRes);
 		}
+		else if (IsMobileSSREnabled(View))
+		{
+			// If we need SSR, and TAA is enabled, then AddTemporalAAPass() has already handled the scene history.
+			// If we need SSR, and TAA is not enabled, then we just need to extract the history.
+			if (!View.bStatePrevViewInfoIsReadOnly)
+			{
+				check(View.ViewState);
+				FTemporalAAHistory& OutputHistory = View.ViewState->PrevFrameViewInfo.TemporalAAHistory;
+				GraphBuilder.QueueTextureExtraction(SceneColor.Texture, &OutputHistory.RT[0]);
+
+				// For SSR, we still fill up the rest of the OutputHistory data using shared math from FTAAPassParameters.
+				FTAAPassParameters TAAInputs(View);
+				TAAInputs.SceneColorInput = SceneColor.Texture;
+				TAAInputs.SetupViewRect(View);
+				OutputHistory.ViewportRect = TAAInputs.OutputViewRect;
+				OutputHistory.ReferenceBufferSize = TAAInputs.GetOutputExtent() * TAAInputs.ResolutionDivisor;
+			}
+		}
 	}
 	else
 	{
diff --git a/Engine/Source/Runtime/Renderer/Private/PostProcess/SceneRenderTargets.h b/Engine/Source/Runtime/Renderer/Private/PostProcess/SceneRenderTargets.h
--- a/Engine/Source/Runtime/Renderer/Private/PostProcess/SceneRenderTargets.h
+++ b/Engine/Source/Runtime/Renderer/Private/PostProcess/SceneRenderTargets.h
@@ -59,6 +59,8 @@
 // [Mobile Deferred only]
 // Store shading model into stencil [1-2] bits
 #define GET_STENCIL_MOBILE_SM_MASK(Value)			uint8(((Value) & 0x3) << 1)
+#define STENCIL_MOBILE_DEFAULTLIT_MASK				uint8(1 << 2)
+#define STENCIL_MOBILE_REFLECTIVE_MASK				uint8(1 << 1)
 
 // [Mobile Forward only]
 // Cast contact shadow mask - bit 4 / Must match shader (ScreenSpaceShadows.usf)
diff --git a/Engine/Source/Runtime/Renderer/Private/SceneRendering.h b/Engine/Source/Runtime/Renderer/Private/SceneRendering.h
--- a/Engine/Source/Runtime/Renderer/Private/SceneRendering.h
+++ b/Engine/Source/Runtime/Renderer/Private/SceneRendering.h
@@ -2637,7 +2637,7 @@ protected:
 	/** Issues occlusion queries */
 	void RenderOcclusion(FRHICommandList& RHICmdList);
 	
-	bool ShouldRenderHZB();
+	bool ShouldRenderHZB(TArrayView<FViewInfo> InViews);
 
 
 	/** Generate HZB */
@@ -2658,6 +2658,9 @@ protected:
 	/** Renders the base pass for translucency. */
 	void RenderTranslucency(FRHICommandList& RHICmdList, const FViewInfo& View);
 
+	/** Renders screen space reflections. */
+	void RenderSSR(FRHICommandList& RHICmdList, const FViewInfo& View);
+
 	/** On chip pre-tonemap before scene color MSAA resolve (iOS only) */
 	void PreTonemapMSAA(FRHICommandList& RHICmdList, const FMinimalSceneTextures& SceneTextures);
 
diff --git a/Engine/Source/Runtime/Renderer/Private/ScreenSpaceRayTracing.cpp b/Engine/Source/Runtime/Renderer/Private/ScreenSpaceRayTracing.cpp
--- a/Engine/Source/Runtime/Renderer/Private/ScreenSpaceRayTracing.cpp
+++ b/Engine/Source/Runtime/Renderer/Private/ScreenSpaceRayTracing.cpp
@@ -156,7 +156,8 @@ bool ShouldRenderScreenSpaceReflections(const FViewInfo& View)
 		return false;
 	}
 
-	if (IsForwardShadingEnabled(View.GetShaderPlatform()))
+	const EShaderPlatform ShaderPlatform = View.GetShaderPlatform();
+	if (IsForwardShadingEnabled(ShaderPlatform) && !IsMobilePlatform(ShaderPlatform))
 	{
 		return false;
 	}
@@ -287,16 +288,7 @@ void SetupCommonScreenSpaceRayParameters(
 		/* inout */ OutParameters);
 } // SetupCommonScreenSpaceRayParameters()
 
-
-} // namespace ScreenSpaceRayTracing
-
-
-bool UseSingleLayerWaterIndirectDraw(EShaderPlatform ShaderPlatform);
-
-namespace
-{
-
-float ComputeRoughnessMaskScale(const FViewInfo& View, ESSRQuality SSRQuality)
+static float ComputeRoughnessMaskScale(const FViewInfo& View, ESSRQuality SSRQuality)
 {
 	float MaxRoughness = FMath::Clamp(View.FinalPostProcessSettings.ScreenSpaceReflectionMaxRoughness, 0.01f, 1.0f);
 
@@ -314,11 +306,11 @@ FLinearColor ComputeSSRParams(const FViewInfo& View, ESSRQuality SSRQuality, boo
 
 	float FrameRandom = 0;
 
-	if(View.ViewState)
+	if (View.ViewState)
 	{
 		bool bTemporalAAIsOn = IsTemporalAccumulationBasedMethod(View.AntiAliasingMethod);
 
-		if(bTemporalAAIsOn)
+		if (bTemporalAAIsOn)
 		{
 			// usually this number is in the 0..7 range but it depends on the TemporalAA quality
 			FrameRandom = View.ViewState->GetCurrentTemporalAASampleIndex() * 1551;
@@ -331,13 +323,18 @@ FLinearColor ComputeSSRParams(const FViewInfo& View, ESSRQuality SSRQuality, boo
 	}
 
 	return FLinearColor(
-		FMath::Clamp(View.FinalPostProcessSettings.ScreenSpaceReflectionIntensity * 0.01f, 0.0f, 1.0f), 
+		FMath::Clamp(View.FinalPostProcessSettings.ScreenSpaceReflectionIntensity * 0.01f, 0.0f, 1.0f),
 		RoughnessMaskScale,
 		(float)bEnableDiscard,	// TODO 
 		FrameRandom);
 }
 
+} // namespace ScreenSpaceRayTracing
 
+bool UseSingleLayerWaterIndirectDraw(EShaderPlatform ShaderPlatform);
+
+namespace
+{
 
 BEGIN_SHADER_PARAMETER_STRUCT(FSSRTTileClassificationParameters, )
 	SHADER_PARAMETER(FIntPoint, TileBufferExtent)
diff --git a/Engine/Source/Runtime/Renderer/Private/ScreenSpaceRayTracing.h b/Engine/Source/Runtime/Renderer/Private/ScreenSpaceRayTracing.h
--- a/Engine/Source/Runtime/Renderer/Private/ScreenSpaceRayTracing.h
+++ b/Engine/Source/Runtime/Renderer/Private/ScreenSpaceRayTracing.h
@@ -117,4 +117,6 @@ void SetupCommonScreenSpaceRayParameters(
 	const FViewInfo& View,
 	FCommonScreenSpaceRayParameters* OutParameters);
 
+FLinearColor ComputeSSRParams(const FViewInfo& View, ESSRQuality SSRQuality, bool bEnableDiscard);
+
 } // namespace ScreenSpaceRayTracing
diff --git a/Engine/Source/Runtime/Renderer/Private/VelocityRendering.cpp b/Engine/Source/Runtime/Renderer/Private/VelocityRendering.cpp
--- a/Engine/Source/Runtime/Renderer/Private/VelocityRendering.cpp
+++ b/Engine/Source/Runtime/Renderer/Private/VelocityRendering.cpp
@@ -786,11 +786,11 @@ bool FVelocityMeshProcessor::Process(
 	// When velocity is used as a depth pass we need to set a correct stencil state on mobile
 	if (FeatureLevel == ERHIFeatureLevel::ES3_1 && EarlyZPassMode == DDM_AllOpaqueNoVelocity)
 	{
-		extern void SetMobileBasePassDepthState(FMeshPassProcessorRenderState& DrawRenderState, const FPrimitiveSceneProxy* PrimitiveSceneProxy, FMaterialShadingModelField ShadingModels, bool bUsesDeferredShading);
+		extern void SetMobileBasePassDepthState(FMeshPassProcessorRenderState& DrawRenderState, const FPrimitiveSceneProxy* PrimitiveSceneProxy, const FMaterial& Material, FMaterialShadingModelField ShadingModels, bool bUsesDeferredShading);
 
 		FMaterialShadingModelField ShadingModels = MaterialResource.GetShadingModels();
 		bool bUsesDeferredShading = IsMobileDeferredShadingEnabled(GetFeatureLevelShaderPlatform(FeatureLevel));
-		SetMobileBasePassDepthState(PassDrawRenderState, PrimitiveSceneProxy, ShadingModels, bUsesDeferredShading);
+		SetMobileBasePassDepthState(PassDrawRenderState, PrimitiveSceneProxy, MaterialResource, ShadingModels, bUsesDeferredShading);
 	}
 
 	FMeshMaterialShaderElementData ShaderElementData;
