Adds whether lens distortion is being applied in VisualizeTemporalUpscaler

diff --git a/Engine/Source/Runtime/Renderer/Private/PostProcess/VisualizeTemporalUpscaler.cpp b/Engine/Source/Runtime/Renderer/Private/PostProcess/VisualizeTemporalUpscaler.cpp
--- a/Engine/Source/Runtime/Renderer/Private/PostProcess/VisualizeTemporalUpscaler.cpp
+++ b/Engine/Source/Runtime/Renderer/Private/PostProcess/VisualizeTemporalUpscaler.cpp
@@ -325,6 +325,10 @@ FScreenPassTexture AddVisualizeTemporalUpscalerPass(FRDGBuilder& GraphBuilder, c
 				QuickDrawSummary(/* Location = */ 6, Sharpen > 0 ? FString::Printf(TEXT("Tonemapper Sharpen: %f"), Sharpen) : TEXT("Tonemapper Sharpen: Off"));
 			}
 
+			// Display if doing any lens distortion
+			{
+				QuickDrawSummary(/* Location = */ 7, View.LensDistortionLUT.IsEnabled() ? FString::Printf(TEXT("Lens Distortion: Enabled (%3.1f%% Additional Upscale)"), View.LensDistortionLUT.ResolutionFraction * 100.0f) : TEXT("Lens Distortion: Disabled"));
+			}
 		});
 
 	}
