Fixes up editor viewports's selection, outline and gizmo with lens distortion enabled.

diff --git a/Engine/Shaders/Private/HitProxyCopy.usf b/Engine/Shaders/Private/HitProxyCopy.usf
new file mode 100644
--- /dev/null
+++ b/Engine/Shaders/Private/HitProxyCopy.usf
@@ -0,0 +1,27 @@
+// Copyright Epic Games, Inc. All Rights Reserved.
+
+
+#include "LensDistortion.ush"
+#include "ScreenPass.ush"
+
+
+Texture2D<float2> UndistortingDisplacementTexture;
+SamplerState UndistortingDisplacementSampler;
+Texture2D HitProxyTexture;
+
+FScreenTransform PassSvPositionToViewportUV;
+FScreenTransform ViewportUVToHitProxyPixelPos;
+int2 HitProxyPixelPosMin;
+int2 HitProxyPixelPosMax;
+
+
+void MainPS(float4 SvPosition : SV_POSITION, out float4 OutColor : SV_Target0)
+{
+	float2 DistortedViewportUV = ApplyScreenTransform(SvPosition.xy, PassSvPositionToViewportUV);
+	float2 UndistortedViewportUV = ApplyLensDistortionOnViewportUV(UndistortingDisplacementTexture, UndistortingDisplacementSampler, DistortedViewportUV);
+
+	int2 HitProxyPixelPos = int2(ApplyScreenTransform(UndistortedViewportUV, ViewportUVToHitProxyPixelPos));
+	HitProxyPixelPos = clamp(HitProxyPixelPos, HitProxyPixelPosMin, HitProxyPixelPosMax);
+
+	OutColor = HitProxyTexture[HitProxyPixelPos];
+;}
diff --git a/Engine/Shaders/Private/LensDistortion.ush b/Engine/Shaders/Private/LensDistortion.ush
new file mode 100644
--- /dev/null
+++ b/Engine/Shaders/Private/LensDistortion.ush
@@ -0,0 +1,32 @@
+// Copyright Epic Games, Inc. All Rights Reserved.
+
+#pragma once
+
+#include "Common.ush"
+
+
+float2 ApplyLensDistortionOnViewportUV(Texture2D<float2> DisplacementTexture, SamplerState DisplacementSampler, float2 DestViewportUV)
+{
+	return DestViewportUV + DisplacementTexture.SampleLevel(DisplacementSampler, DestViewportUV, 0);
+}
+
+float2 ApplyLensDistortionOnScreenPos(Texture2D<float2> DisplacementTexture, SamplerState DisplacementSampler, float2 DestScreenPos)
+{
+	float2 DestViewportUV = ScreenPosToViewportUV(DestScreenPos);
+
+	return DestScreenPos + DisplacementTexture.SampleLevel(DisplacementSampler, DestViewportUV, 0) * float2(2, -2);
+}
+
+#if SUPPORTS_INDEPENDENT_SAMPLERS
+
+float2 ApplyLensDistortionOnViewportUV(Texture2D<float2> DisplacementTexture, float2 DestViewportUV)
+{
+	return ApplyLensDistortionOnViewportUV(DisplacementTexture, GlobalBilinearClampedSampler, DestViewportUV);
+}
+
+float2 ApplyLensDistortionOnScreenPos(Texture2D<float2> DisplacementTexture, float2 DestScreenPos)
+{
+	return ApplyLensDistortionOnScreenPos(DisplacementTexture, GlobalBilinearClampedSampler, DestScreenPos);
+}
+
+#endif
diff --git a/Engine/Shaders/Private/PostProcessCompositePrimitives.usf b/Engine/Shaders/Private/PostProcessCompositePrimitives.usf
--- a/Engine/Shaders/Private/PostProcessCompositePrimitives.usf
+++ b/Engine/Shaders/Private/PostProcessCompositePrimitives.usf
@@ -3,6 +3,7 @@
 #include "Common.ush"
 #include "ScreenPass.ush"
 #include "PostProcessCommon.ush"
+#include "LensDistortion.ush"
 
 #ifndef MSAA_SAMPLE_COUNT
 #define MSAA_SAMPLE_COUNT 1
@@ -68,6 +69,9 @@ static const uint kSquareIndexes3x3[9] = { 4u, 0u, 1u, 2u, 3u, 5u, 6u, 7u, 8u };
 // Indexes of the offsets to have plus + shape.
 static const uint kPlusIndexes3x3[5] = { 4u, 1u, 3u, 5u, 7u };
 
+Texture2D<float2> UndistortingDisplacementTexture;
+SamplerState UndistortingDisplacementSampler;
+
 #if MSAA_SAMPLE_COUNT > 1
 Texture2DMS<float4, MSAA_SAMPLE_COUNT> EditorPrimitivesColor;
 Texture2DMS<float, MSAA_SAMPLE_COUNT> EditorPrimitivesDepth;
@@ -92,8 +96,10 @@ SCREEN_PASS_TEXTURE_VIEWPORT(Color)
 SCREEN_PASS_TEXTURE_VIEWPORT(Depth)
 SCREEN_PASS_TEXTURE_VIEWPORT(PrevHistory)
 SCREEN_PASS_TEXTURE_VIEWPORT(History)
-SCREEN_PASS_TEXTURE_VIEWPORT(Output)
-FScreenTransform ColorToDepth;
+
+FScreenTransform PassSvPositionToViewportUV;
+FScreenTransform ViewportUVToColorUV;
+FScreenTransform ViewportUVToDepthUV;
 
 float2 DepthTextureJitter;
 uint bCameraCut;
@@ -307,12 +313,16 @@ void MainCompositeEditorPrimitivesPS(
 	out float4 OutColor : SV_Target0,
 	out float OutDepth : SV_Depth)
 {
-	float2 ViewportUV = (SvPosition.xy - Output_ViewportMin) * Output_ViewportSizeInverse;
-	
-	const float2 ColorUV = (ViewportUV * Color_ViewportSize + Color_ViewportMin) * Color_ExtentInverse;
+	const float2 DistortedViewportUV = ApplyScreenTransform(SvPosition.xy, PassSvPositionToViewportUV);
+	const float2 UndistortedViewportUV = ApplyLensDistortionOnViewportUV(UndistortingDisplacementTexture, UndistortingDisplacementSampler, DistortedViewportUV);
+
+	const float2 DistortedColorUV = ApplyScreenTransform(DistortedViewportUV, ViewportUVToColorUV);
+	const float2 ColorUV = ApplyScreenTransform(UndistortedViewportUV, ViewportUVToColorUV);
+	const float2 DepthUV = ApplyScreenTransform(UndistortedViewportUV, ViewportUVToDepthUV);
+
 	const int2 ColorPixelPos = int2(ColorUV * Color_Extent);
 	
-	float4 SceneColor = Texture2DSample(ColorTexture, ColorSampler, ColorUV);
+	float4 SceneColor = Texture2DSample(ColorTexture, ColorSampler, DistortedColorUV);
 
 	// Resolve editor primitive scene color and depth.
 	float4 EditorPrimitiveColor;
@@ -323,7 +333,7 @@ void MainCompositeEditorPrimitivesPS(
 		TargetSampleIndex = min(TargetSampleIndex, MSAA_SAMPLE_COUNT-1);
 		const float2 SampleOffset = SampleOffsetArray[TargetSampleIndex].xy;
 
-		float2 SampleDepthUV = (ViewportUV * Depth_ViewportSize + Depth_ViewportMin + SampleOffset - DepthTextureJitter) * Depth_ExtentInverse;
+		float2 SampleDepthUV = (UndistortedViewportUV * Depth_ViewportSize + Depth_ViewportMin + SampleOffset - DepthTextureJitter) * Depth_ExtentInverse;
 		SampleDepthUV = clamp(SampleDepthUV, Depth_UVViewportBilinearMin, Depth_UVViewportBilinearMax);
 
 		EditorPrimitiveColor = EditorPrimitivesColor.Load(ColorPixelPos, TargetSampleIndex);
@@ -349,7 +359,7 @@ void MainCompositeEditorPrimitivesPS(
 		{
 			const float2 SampleOffset = SampleOffsetArray[SampleIndex].xy;
 
-			float2 SampleDepthUV = (ViewportUV * Depth_ViewportSize + Depth_ViewportMin + SampleOffset - DepthTextureJitter) * Depth_ExtentInverse;
+			float2 SampleDepthUV = (UndistortedViewportUV * Depth_ViewportSize + Depth_ViewportMin + SampleOffset - DepthTextureJitter) * Depth_ExtentInverse;
 			SampleDepthUV = clamp(SampleDepthUV, Depth_UVViewportBilinearMin, Depth_UVViewportBilinearMax);
 			float SampleSceneDeviceZ = Texture2DSampleLevel(DepthTexture, DepthSampler, SampleDepthUV, 0).r;
 
@@ -388,7 +398,7 @@ void MainCompositeEditorPrimitivesPS(
 	#else
 	{
 		// De-jitter the sample position and make a filtered lookup - for planes this allows to reconstruct a much less jittery depth comparison function. It doesn't fix silhouettes, however.
-		float2 SampleDepthUV = (ViewportUV * Depth_ViewportSize + Depth_ViewportMin - DepthTextureJitter) * Depth_ExtentInverse;
+		float2 SampleDepthUV = (UndistortedViewportUV * Depth_ViewportSize + Depth_ViewportMin - DepthTextureJitter) * Depth_ExtentInverse;
 		SampleDepthUV = clamp(SampleDepthUV, Depth_UVViewportBilinearMin, Depth_UVViewportBilinearMax);
 
 		float SceneDeviceZ = Texture2DSampleLevel(DepthTexture, DepthSampler, SampleDepthUV, 0).r;
@@ -415,8 +425,6 @@ void MainCompositeEditorPrimitivesPS(
 		EditorPrimitiveColor.rgb *= EditorPrimitiveColor.a;
 	}
 
-	float2 DepthUV = ApplyScreenTransform(ColorUV, ColorToDepth);
-
 	// Generate 2x2 square tiling pattern for foreground primitive that end up behind scene opaque primitives.
 	float PatternMask = ((ColorPixelPos.x/2 + ColorPixelPos.y/2) % 2) * 0.7f;
 
diff --git a/Engine/Shaders/Private/PostProcessSelectionOutline.usf b/Engine/Shaders/Private/PostProcessSelectionOutline.usf
--- a/Engine/Shaders/Private/PostProcessSelectionOutline.usf
+++ b/Engine/Shaders/Private/PostProcessSelectionOutline.usf
@@ -4,6 +4,11 @@
 #include "ScreenPass.ush"
 #include "GammaCorrectionCommon.ush"
 #include "ACES.ush"
+#include "LensDistortion.ush"
+
+
+Texture2D<float2> UndistortingDisplacementTexture;
+SamplerState UndistortingDisplacementSampler;
 
 #if MSAA_SAMPLE_COUNT > 1
 Texture2DMS<float, MSAA_SAMPLE_COUNT> EditorPrimitivesDepth;
@@ -25,7 +30,9 @@ SamplerState OverlayLookupSampler;
 
 SCREEN_PASS_TEXTURE_VIEWPORT(Color)
 SCREEN_PASS_TEXTURE_VIEWPORT(Depth)
-FScreenTransform ColorToDepth;
+FScreenTransform PassSvPositionToViewportUV;
+FScreenTransform ViewportUVToColorUV;
+FScreenTransform ViewportUVToDepthUV;
 
 float4 OutlineColors[8];
 int OutlineColorIndexBits;
@@ -211,11 +218,16 @@ float4 PerSample(int2 PixelPos, int SampleID, FPixelInfo CenterPixelInfo, float3
 }
 
 void MainPS(
-	noperspective float4 UVAndScreenPos : TEXCOORD0,
 	float4 SvPosition : SV_POSITION,
 	out float4 OutColor : SV_Target0)
 {
-	const float2 ColorUV = UVAndScreenPos.xy;
+	const float2 DistortedViewportUV = ApplyScreenTransform(SvPosition.xy, PassSvPositionToViewportUV);
+	const float2 UndistortedViewportUV = ApplyLensDistortionOnViewportUV(UndistortingDisplacementTexture, UndistortingDisplacementSampler, DistortedViewportUV);
+
+	const float2 DistortedColorUV = ApplyScreenTransform(DistortedViewportUV, ViewportUVToColorUV);
+	const float2 ColorUV = ApplyScreenTransform(UndistortedViewportUV, ViewportUVToColorUV);
+	const float2 DepthUV = ApplyScreenTransform(UndistortedViewportUV, ViewportUVToDepthUV);
+
 	const int2 ColorPixelPos = int2(ColorUV * Color_Extent);
 
 	// Scout outwards from the center pixel to determine if any neighboring pixels are selected
@@ -237,8 +249,6 @@ void MainPS(
 	);
 #endif
 
-	const float2 DepthUV = ApplyScreenTransform(ColorUV, ColorToDepth);
-
 	float Center = Texture2DSampleLevel(DepthTexture, DepthSampler, DepthUV, 0).r;
 	float Left   = Texture2DSampleLevel(DepthTexture, DepthSampler, DepthUV + float2(-1,  0) * Depth_ExtentInverse, 0).r;
 	float Right  = Texture2DSampleLevel(DepthTexture, DepthSampler, DepthUV + float2( 1,  0) * Depth_ExtentInverse, 0).r;
@@ -259,7 +269,7 @@ void MainPS(
 	FPixelInfo CenterPixelInfo = GetPixelInfo(ColorPixelPos, 0, SecondaryViewportOffset, 0, DeviceZPlane, DeviceZMinMax);
 
 	float4 OverlayColor = Texture2DSample(OverlayLookupTexture, OverlayLookupSampler, ColorUV);
-	OutColor = Texture2DSample(ColorTexture, ColorSampler, ColorUV); 
+	OutColor = Texture2DSample(ColorTexture, ColorSampler, DistortedColorUV); 
 	
 	float4 SelectionEffect = 0;
 	const int ScoutSum = dot(saturate(StencilScout), 1);
diff --git a/Engine/Shaders/Private/PostProcessUpscale.usf b/Engine/Shaders/Private/PostProcessUpscale.usf
--- a/Engine/Shaders/Private/PostProcessUpscale.usf
+++ b/Engine/Shaders/Private/PostProcessUpscale.usf
@@ -5,6 +5,7 @@
 #include "PostProcessCommon.ush"
 #include "PaniniProjection.ush"
 #include "TextureSampling.ush"
+#include "LensDistortion.ush"
 
 #define UPSCALE_METHOD_NEAREST 0
 #define UPSCALE_METHOD_BILINEAR 1
@@ -92,7 +93,7 @@ void MainVS(
 	float4 Position = float4(TexCoord.x, TexCoord.y, 0, 1);
 
 	// distort pos
-	Position.xy = TexCoord + DistortingDisplacementTexture.SampleLevel(DistortingDisplacementSampler, TexCoord, 0);
+	Position.xy = ApplyLensDistortionOnViewportUV(DistortingDisplacementTexture, DistortingDisplacementSampler, TexCoord);
 
 	DrawRectangle(Position, TexCoord, OutPosition, OutTexCoord.xy);
 	OutTexCoord.zw =  OutPosition.xy;
diff --git a/Engine/Shaders/Private/TemporalSuperResolution/TSRVisualize.usf b/Engine/Shaders/Private/TemporalSuperResolution/TSRVisualize.usf
--- a/Engine/Shaders/Private/TemporalSuperResolution/TSRVisualize.usf
+++ b/Engine/Shaders/Private/TemporalSuperResolution/TSRVisualize.usf
@@ -1,6 +1,9 @@
 // Copyright Epic Games, Inc. All Rights Reserved.
 
 #include "TSRCommon.ush"
+#if TSR_SUPPORT_LENS_DISTORTION
+	#include "../LensDistortion.ush"
+#endif
 
 
 //------------------------------------------------------- CONFIG
@@ -60,16 +63,8 @@ RWTexture2D<tsr_half4> Output;
 
 //------------------------------------------------------- FUNCTION
 
+#if !TSR_SUPPORT_LENS_DISTORTION
 float2 ApplyLensDistortionOnScreenPos(Texture2D<float2> DisplacementTexture, float2 DestScreenPos)
-#if TSR_SUPPORT_LENS_DISTORTION
-{
-	float2 DestViewportUV = ScreenPosToViewportUV(DestScreenPos);
-
-	float2 SrcViewportUV = DestViewportUV + DisplacementTexture.SampleLevel(GlobalBilinearClampedSampler, DestViewportUV, 0);
-			
-	return ViewportUVToScreenPos(SrcViewportUV);
-}
-#else
 {
 	return DestScreenPos;
 }
diff --git a/Engine/Source/Runtime/Renderer/Private/PostProcess/PostProcessCompositeEditorPrimitives.cpp b/Engine/Source/Runtime/Renderer/Private/PostProcess/PostProcessCompositeEditorPrimitives.cpp
--- a/Engine/Source/Runtime/Renderer/Private/PostProcess/PostProcessCompositeEditorPrimitives.cpp
+++ b/Engine/Source/Runtime/Renderer/Private/PostProcess/PostProcessCompositeEditorPrimitives.cpp
@@ -10,39 +10,47 @@
 #include "PixelShaderUtils.h"
 #include "Substrate/Substrate.h"
 #include "MeshEdgesRendering.h"
+#include "PixelShaderUtils.h"
 
 namespace
 {
-	class FCompositeEditorPrimitivesPS : public FCompositePrimitiveShaderBase
+class FCompositeEditorPrimitivesPS : public FCompositePrimitiveShaderBase
+{
+public:
+	DECLARE_GLOBAL_SHADER(FCompositeEditorPrimitivesPS);
+	SHADER_USE_PARAMETER_STRUCT(FCompositeEditorPrimitivesPS, FCompositePrimitiveShaderBase);
+
+	BEGIN_SHADER_PARAMETER_STRUCT(FParameters, )
+		SHADER_PARAMETER_STRUCT_REF(FViewUniformShaderParameters, View)
+		SHADER_PARAMETER_STRUCT(FScreenPassTextureViewportParameters, Color)
+		SHADER_PARAMETER_STRUCT(FScreenPassTextureViewportParameters, Depth)
+		SHADER_PARAMETER_ARRAY(FVector4f, SampleOffsetArray, [FCompositePrimitiveShaderBase::kMSAASampleCountMax])
+
+		SHADER_PARAMETER_RDG_TEXTURE(Texture2D, UndistortingDisplacementTexture)
+		SHADER_PARAMETER_SAMPLER(SamplerState,  UndistortingDisplacementSampler)
+		SHADER_PARAMETER_RDG_TEXTURE(Texture2D, EditorPrimitivesDepth)
+		SHADER_PARAMETER_RDG_TEXTURE(Texture2D, EditorPrimitivesColor)
+		SHADER_PARAMETER_RDG_TEXTURE(Texture2D, ColorTexture)
+		SHADER_PARAMETER_SAMPLER(SamplerState,  ColorSampler)
+		SHADER_PARAMETER_RDG_TEXTURE(Texture2D, DepthTexture)
+		SHADER_PARAMETER_SAMPLER(SamplerState,  DepthSampler)
+
+
+		SHADER_PARAMETER(FScreenTransform, PassSvPositionToViewportUV)
+		SHADER_PARAMETER(FScreenTransform, ViewportUVToColorUV)
+		SHADER_PARAMETER(FScreenTransform, ViewportUVToDepthUV)
+		SHADER_PARAMETER(uint32, bOpaqueEditorGizmo)
+		SHADER_PARAMETER(uint32, bCompositeAnyNonNullDepth)
+		SHADER_PARAMETER(FVector2f, DepthTextureJitter)
+		RENDER_TARGET_BINDING_SLOTS()
+	END_SHADER_PARAMETER_STRUCT()
+
+	static void ModifyCompilationEnvironment(const FGlobalShaderPermutationParameters& Parameters, FShaderCompilerEnvironment& OutEnvironment)
 	{
-	public:
-		DECLARE_GLOBAL_SHADER(FCompositeEditorPrimitivesPS);
-		SHADER_USE_PARAMETER_STRUCT(FCompositeEditorPrimitivesPS, FCompositePrimitiveShaderBase);
-
-		BEGIN_SHADER_PARAMETER_STRUCT(FParameters, )
-			SHADER_PARAMETER_STRUCT_REF(FViewUniformShaderParameters, View)
-			SHADER_PARAMETER_STRUCT(FScreenPassTextureViewportParameters, Color)
-			SHADER_PARAMETER_STRUCT(FScreenPassTextureViewportParameters, Depth)
-			SHADER_PARAMETER_STRUCT(FScreenPassTextureViewportParameters, Output)
-			SHADER_PARAMETER_ARRAY(FVector4f, SampleOffsetArray, [FCompositePrimitiveShaderBase::kMSAASampleCountMax])
-			SHADER_PARAMETER_RDG_TEXTURE(Texture2D, EditorPrimitivesDepth)
-			SHADER_PARAMETER_RDG_TEXTURE(Texture2D, EditorPrimitivesColor)
-			SHADER_PARAMETER_RDG_TEXTURE(Texture2D, ColorTexture)
-			SHADER_PARAMETER_SAMPLER(SamplerState, ColorSampler)
-			SHADER_PARAMETER_RDG_TEXTURE(Texture2D, DepthTexture)
-			SHADER_PARAMETER_SAMPLER(SamplerState, DepthSampler)
-			SHADER_PARAMETER(uint32, bOpaqueEditorGizmo)
-			SHADER_PARAMETER(uint32, bCompositeAnyNonNullDepth)
-			SHADER_PARAMETER(FVector2f, DepthTextureJitter)
-			RENDER_TARGET_BINDING_SLOTS()
-			END_SHADER_PARAMETER_STRUCT()
-
-			static void ModifyCompilationEnvironment(const FGlobalShaderPermutationParameters& Parameters, FShaderCompilerEnvironment& OutEnvironment)
-		{
-			FGlobalShader::ModifyCompilationEnvironment(Parameters, OutEnvironment);
-		}
-	};
-	IMPLEMENT_GLOBAL_SHADER(FCompositeEditorPrimitivesPS, "/Engine/Private/PostProcessCompositePrimitives.usf", "MainCompositeEditorPrimitivesPS", SF_Pixel);
+		FGlobalShader::ModifyCompilationEnvironment(Parameters, OutEnvironment);
+	}
+};
+IMPLEMENT_GLOBAL_SHADER(FCompositeEditorPrimitivesPS, "/Engine/Private/PostProcessCompositePrimitives.usf", "MainCompositeEditorPrimitivesPS", SF_Pixel);
 
 
 void RenderEditorPrimitives(FRHICommandList& RHICmdList, const FViewInfo& View, FMeshPassProcessorRenderState& DrawRenderState, FInstanceCullingManager& InstanceCullingManager)
@@ -392,7 +400,14 @@ FScreenPassTexture AddEditorPrimitivePass(
 		PassParameters->View = View.ViewUniformBuffer;
 		PassParameters->Color = GetScreenPassTextureViewportParameters(FScreenPassTextureViewport(Inputs.SceneColor));
 		PassParameters->Depth = GetScreenPassTextureViewportParameters(FScreenPassTextureViewport(SceneDepth));
-		PassParameters->Output = GetScreenPassTextureViewportParameters(FScreenPassTextureViewport(Output));
+
+		PassParameters->UndistortingDisplacementTexture = GSystemTextures.GetBlackDummy(GraphBuilder);
+		PassParameters->UndistortingDisplacementSampler = TStaticSamplerState<SF_Bilinear, AM_Clamp, AM_Clamp, AM_Clamp>::GetRHI();
+		if (Inputs.LensDistortionLUT.IsEnabled())
+		{
+			PassParameters->UndistortingDisplacementTexture = Inputs.LensDistortionLUT.UndistortingDisplacementTexture;
+		}
+
 		PassParameters->ColorTexture = Inputs.SceneColor.Texture;
 		PassParameters->ColorSampler = PointClampSampler;
 		if (View.Family->EngineShowFlags.SceneCaptureCopySceneDepth)
@@ -406,6 +421,13 @@ FScreenPassTexture AddEditorPrimitivePass(
 		PassParameters->DepthSampler = PointClampSampler;
 		PassParameters->EditorPrimitivesDepth = EditorPrimitiveDepth;
 		PassParameters->EditorPrimitivesColor = EditorPrimitiveColor;
+		
+		PassParameters->PassSvPositionToViewportUV = FScreenTransform::SvPositionToViewportUV(Output.ViewRect);
+		PassParameters->ViewportUVToColorUV = FScreenTransform::ChangeTextureBasisFromTo(
+			FScreenPassTextureViewport(Inputs.SceneColor), FScreenTransform::ETextureBasis::ViewportUV, FScreenTransform::ETextureBasis::TextureUV);
+		PassParameters->ViewportUVToDepthUV = FScreenTransform::ChangeTextureBasisFromTo(
+			FScreenPassTextureViewport(SceneDepth), FScreenTransform::ETextureBasis::ViewportUV, FScreenTransform::ETextureBasis::TextureUV);
+		
 		PassParameters->bOpaqueEditorGizmo = bOpaqueEditorGizmo;
 		PassParameters->bCompositeAnyNonNullDepth = bProducedByPriorPass && !View.Family->EngineShowFlags.MeshEdges;
 		PassParameters->DepthTextureJitter = SceneDepthJitter;
diff --git a/Engine/Source/Runtime/Renderer/Private/PostProcess/PostProcessCompositePrimitivesCommon.h b/Engine/Source/Runtime/Renderer/Private/PostProcess/PostProcessCompositePrimitivesCommon.h
--- a/Engine/Source/Runtime/Renderer/Private/PostProcess/PostProcessCompositePrimitivesCommon.h
+++ b/Engine/Source/Runtime/Renderer/Private/PostProcess/PostProcessCompositePrimitivesCommon.h
@@ -5,6 +5,7 @@
 #include "ScreenPass.h"
 #include "OverridePassSequence.h"
 #include "DataDrivenShaderPlatformInfo.h"
+#include "PostProcess/LensDistortion.h"
 
 //UE_ENABLE_DEBUG_DRAWING i.e. !(UE_BUILD_SHIPPING || UE_BUILD_TEST) || WITH_EDITOR
 //Only available in Debug/Development/Editor builds due to current use cases, but can be extended in future
@@ -68,6 +69,9 @@ struct FCompositePrimitiveInputs
 	// [Required] The scene depth to composite with editor primitives.
 	FScreenPassTexture SceneDepth;
 
+	// [Optional] Lens distortion applied on the scene color.
+	FLensDistortionLUT LensDistortionLUT;
+
 	bool bUseMetalMSAAHDRDecode = false;
 };
 
diff --git a/Engine/Source/Runtime/Renderer/Private/PostProcess/PostProcessSelectionOutline.cpp b/Engine/Source/Runtime/Renderer/Private/PostProcess/PostProcessSelectionOutline.cpp
--- a/Engine/Source/Runtime/Renderer/Private/PostProcess/PostProcessSelectionOutline.cpp
+++ b/Engine/Source/Runtime/Renderer/Private/PostProcess/PostProcessSelectionOutline.cpp
@@ -10,6 +10,7 @@
 #include "ClearQuad.h"
 #include "ScenePrivate.h"
 #include "PostProcess/SceneRenderTargets.h"
+#include "PixelShaderUtils.h"
 
 namespace
 {
@@ -32,21 +33,28 @@ public:
 		SHADER_PARAMETER_STRUCT_REF(FViewUniformShaderParameters, View)
 		SHADER_PARAMETER_STRUCT(FScreenPassTextureViewportParameters, Color)
 		SHADER_PARAMETER_STRUCT(FScreenPassTextureViewportParameters, Depth)
+
+		SHADER_PARAMETER_RDG_TEXTURE(Texture2D, UndistortingDisplacementTexture)
+		SHADER_PARAMETER_SAMPLER(SamplerState,  UndistortingDisplacementSampler)
 		SHADER_PARAMETER_RDG_TEXTURE(Texture2D, ColorTexture)
-		SHADER_PARAMETER_SAMPLER(SamplerState, ColorSampler)
+		SHADER_PARAMETER_SAMPLER(SamplerState,  ColorSampler)
 		SHADER_PARAMETER_RDG_TEXTURE(Texture2D, DepthTexture)
-		SHADER_PARAMETER_SAMPLER(SamplerState, DepthSampler)
+		SHADER_PARAMETER_SAMPLER(SamplerState,  DepthSampler)
 		SHADER_PARAMETER_RDG_TEXTURE(Texture2D, EditorPrimitivesDepth)
 		SHADER_PARAMETER_RDG_TEXTURE_SRV(Texture2D, EditorPrimitivesStencil)
 		SHADER_PARAMETER_RDG_TEXTURE(Texture2D, OverlayLookupTexture)
-		SHADER_PARAMETER_SAMPLER(SamplerState, OverlayLookupSampler)
-		SHADER_PARAMETER(FScreenTransform, ColorToDepth)
+		SHADER_PARAMETER_SAMPLER(SamplerState,  OverlayLookupSampler)
+
+		SHADER_PARAMETER(FScreenTransform, PassSvPositionToViewportUV)
+		SHADER_PARAMETER(FScreenTransform, ViewportUVToColorUV)
+		SHADER_PARAMETER(FScreenTransform, ViewportUVToDepthUV)
 		SHADER_PARAMETER_ARRAY(FVector4f, OutlineColors, [8])
 		SHADER_PARAMETER(int, OutlineColorIndexBits)
 		SHADER_PARAMETER(float, SelectionHighlightIntensity)
 		SHADER_PARAMETER(float, BSPSelectionIntensity)
 		SHADER_PARAMETER(float, UILuminanceAndIsSCRGB)
 		SHADER_PARAMETER(float, SecondaryViewportOffset)
+
 		RENDER_TARGET_BINDING_SLOTS()
 	END_SHADER_PARAMETER_STRUCT()
 
@@ -274,17 +282,34 @@ FScreenPassTexture AddSelectionOutlinePass(
 		PassParameters->View = View.ViewUniformBuffer;
 		PassParameters->Color = GetScreenPassTextureViewportParameters(ColorViewport);
 		PassParameters->Depth = GetScreenPassTextureViewportParameters(DepthViewport);
-		PassParameters->ColorToDepth = FScreenTransform::ChangeTextureUVCoordinateFromTo(ColorViewport, DepthViewport);
+
+		PassParameters->UndistortingDisplacementTexture = GSystemTextures.GetBlackDummy(GraphBuilder);
+		PassParameters->UndistortingDisplacementSampler = TStaticSamplerState<SF_Bilinear, AM_Clamp, AM_Clamp, AM_Clamp>::GetRHI();
+		if (Inputs.LensDistortionLUT.IsEnabled())
+		{
+			PassParameters->UndistortingDisplacementTexture = Inputs.LensDistortionLUT.UndistortingDisplacementTexture;
+		}
+
 		PassParameters->ColorTexture = Inputs.SceneColor.Texture;
 		PassParameters->ColorSampler = PointClampSampler;
+
 		PassParameters->DepthTexture = Inputs.SceneDepth.Texture;
 		PassParameters->DepthSampler = PointClampSampler;
+
 		PassParameters->EditorPrimitivesDepth = DepthStencilTexture;
 		PassParameters->EditorPrimitivesStencil = GraphBuilder.CreateSRV(FRDGTextureSRVDesc::CreateWithPixelFormat(DepthStencilTexture, PF_X24_G8));
+		
 		PassParameters->OverlayLookupTexture = OverlayColorTexture;
 		PassParameters->OverlayLookupSampler = PointClampSampler;
+		
 		PassParameters->OutlineColors[0] = View.SelectionOutlineColor;
 		PassParameters->OutlineColors[1] = View.SubduedSelectionOutlineColor;
+		
+		PassParameters->PassSvPositionToViewportUV = FScreenTransform::SvPositionToViewportUV(OutputViewport.Rect);
+		PassParameters->ViewportUVToColorUV = FScreenTransform::ChangeTextureBasisFromTo(
+			ColorViewport, FScreenTransform::ETextureBasis::ViewportUV, FScreenTransform::ETextureBasis::TextureUV);
+		PassParameters->ViewportUVToDepthUV = FScreenTransform::ChangeTextureBasisFromTo(
+			DepthViewport, FScreenTransform::ETextureBasis::ViewportUV, FScreenTransform::ETextureBasis::TextureUV);
 
 		if (GEnableSelectionOutlineColors)
 		{
@@ -346,14 +371,13 @@ FScreenPassTexture AddSelectionOutlinePass(
 
 		TShaderMapRef<FSelectionOutlinePS> PixelShader(View.ShaderMap, PermutationVector);
 
-		AddDrawScreenPass(
+		FPixelShaderUtils::AddFullscreenPass(
 			GraphBuilder,
+			View.ShaderMap,
 			RDG_EVENT_NAME("OutlineColor %dx%d", OutputViewport.Rect.Width(), OutputViewport.Rect.Height()),
-			View,
-			OutputViewport,
-			ColorViewport,
 			PixelShader,
-			PassParameters);
+			PassParameters,
+			Output.ViewRect);
 	}
 
 	return MoveTemp(Output);
diff --git a/Engine/Source/Runtime/Renderer/Private/PostProcess/PostProcessSelectionOutline.h b/Engine/Source/Runtime/Renderer/Private/PostProcess/PostProcessSelectionOutline.h
--- a/Engine/Source/Runtime/Renderer/Private/PostProcess/PostProcessSelectionOutline.h
+++ b/Engine/Source/Runtime/Renderer/Private/PostProcess/PostProcessSelectionOutline.h
@@ -6,6 +6,7 @@
 
 #include "ScreenPass.h"
 #include "OverridePassSequence.h"
+#include "PostProcess/LensDistortion.h"
 
 class FSceneUniformBuffer;
 namespace Nanite
@@ -26,6 +27,9 @@ struct FSelectionOutlineInputs
 
 	// [Required] Used when scene textures are required by the material.
 	FSceneTextureShaderParameters SceneTextures;
+
+	// [Optional] Lens distortion applied on the scene color.
+	FLensDistortionLUT LensDistortionLUT;
 };
 
 FScreenPassTexture AddSelectionOutlinePass(
diff --git a/Engine/Source/Runtime/Renderer/Private/PostProcess/PostProcessing.cpp b/Engine/Source/Runtime/Renderer/Private/PostProcess/PostProcessing.cpp
--- a/Engine/Source/Runtime/Renderer/Private/PostProcess/PostProcessing.cpp
+++ b/Engine/Source/Runtime/Renderer/Private/PostProcess/PostProcessing.cpp
@@ -1534,7 +1534,10 @@ void AddPostProcessingPasses(
 		PassInputs.SceneColor = SceneColor;
 		PassInputs.SceneDepth = SceneDepth;
 		PassInputs.SceneTextures.SceneTextures = Inputs.SceneTextures;
-
+		if (bApplyLensDistortionInTSR)
+		{
+			PassInputs.LensDistortionLUT = View.LensDistortionLUT;
+		}
 		SceneColor = AddSelectionOutlinePass(GraphBuilder, View, SceneUniformBuffer, PassInputs, NaniteRasterResults, InstancedEditorDepthTexture);
 	}
 
@@ -1548,7 +1551,10 @@ void AddPostProcessingPasses(
 		PassInputs.SceneColor = SceneColor;
 		PassInputs.SceneDepth = SceneDepth;
 		PassInputs.BasePassType = FCompositePrimitiveInputs::EBasePassType::Deferred;
-
+		if (bApplyLensDistortionInTSR)
+		{
+			PassInputs.LensDistortionLUT = View.LensDistortionLUT;
+		}
 		SceneColor = AddEditorPrimitivePass(GraphBuilder, View, PassInputs, InstanceCullingManager);
 	}
 #endif //WITH_EDITOR
diff --git a/Engine/Source/Runtime/Renderer/Private/SceneHitProxyRendering.cpp b/Engine/Source/Runtime/Renderer/Private/SceneHitProxyRendering.cpp
--- a/Engine/Source/Runtime/Renderer/Private/SceneHitProxyRendering.cpp
+++ b/Engine/Source/Runtime/Renderer/Private/SceneHitProxyRendering.cpp
@@ -35,6 +35,7 @@
 #include "HairStrands/HairStrandsData.h"
 #include "SimpleMeshDrawCommandPass.h"
 #include "StaticMeshSceneProxy.h"
+#include "PixelShaderUtils.h"
 
 class FHitProxyShaderElementData : public FMeshMaterialShaderElementData
 {
@@ -224,11 +225,6 @@ BEGIN_SHADER_PARAMETER_STRUCT(FHitProxyPassParameters, )
 	RENDER_TARGET_BINDING_SLOTS()
 END_SHADER_PARAMETER_STRUCT()
 
-BEGIN_SHADER_PARAMETER_STRUCT(FHitProxyCopyToViewFamilyParameters, )
-	RDG_TEXTURE_ACCESS(HitProxyTexture, ERHIAccess::SRVGraphics)
-	RENDER_TARGET_BINDING_SLOTS()
-END_SHADER_PARAMETER_STRUCT()
-
 static void AddViewMeshElementsPass(const TIndirectArray<FMeshBatch> &MeshElements, FRDGBuilder& GraphBuilder, FHitProxyPassParameters* PassParameters, const FScene* Scene, const FViewInfo& View, const FMeshPassProcessorRenderState& DrawRenderState, FInstanceCullingManager& InstanceCullingManager)
 {
 	AddSimpleMeshPass(GraphBuilder, PassParameters, Scene, View, &InstanceCullingManager, RDG_EVENT_NAME("HitProxy::MeshElementsPass"), View.ViewRect,
@@ -251,6 +247,28 @@ static void AddViewMeshElementsPass(const TIndirectArray<FMeshBatch> &MeshElemen
 	);
 }
 
+class FHitProxyCopyPS : public FGlobalShader
+{
+public:
+	DECLARE_GLOBAL_SHADER(FHitProxyCopyPS);
+	SHADER_USE_PARAMETER_STRUCT(FHitProxyCopyPS, FGlobalShader);
+
+	BEGIN_SHADER_PARAMETER_STRUCT(FParameters, )
+		SHADER_PARAMETER_RDG_TEXTURE(Texture2D, UndistortingDisplacementTexture)
+		SHADER_PARAMETER_SAMPLER(SamplerState,  UndistortingDisplacementSampler)
+		SHADER_PARAMETER_RDG_TEXTURE(Texture2D, HitProxyTexture)
+
+		SHADER_PARAMETER(FScreenTransform, PassSvPositionToViewportUV)
+		SHADER_PARAMETER(FScreenTransform, ViewportUVToHitProxyPixelPos)
+		SHADER_PARAMETER(FIntPoint, HitProxyPixelPosMin)
+		SHADER_PARAMETER(FIntPoint, HitProxyPixelPosMax)
+
+		RENDER_TARGET_BINDING_SLOTS()
+	END_SHADER_PARAMETER_STRUCT()
+};
+IMPLEMENT_GLOBAL_SHADER(FHitProxyCopyPS, "/Engine/Private/HitProxyCopy.usf", "MainPS", SF_Pixel);
+
+
 static void DoRenderHitProxies(
 	FRDGBuilder& GraphBuilder, 
 	const FSceneRenderer* SceneRenderer, 
@@ -476,78 +494,36 @@ static void DoRenderHitProxies(
 	FRDGTextureRef ViewFamilyTexture = TryCreateViewFamilyTexture(GraphBuilder, ViewFamily);
 	check(ViewFamilyTexture);
 
-	//
-	// Copy the hit proxy buffer into the view family's render target.
-	//
-
+	// Copy & Apply lens distortion of the hit proxy buffer into the view family's render target.
+	for (int32 ViewIndex = 0; ViewIndex < Views.Num(); ViewIndex++)
 	{
-		auto* PassParameters = GraphBuilder.AllocParameters<FHitProxyCopyToViewFamilyParameters>();
+		const FViewInfo& View = Views[ViewIndex];
+
+		FHitProxyCopyPS::FParameters* PassParameters = GraphBuilder.AllocParameters<FHitProxyCopyPS::FParameters>();
+		PassParameters->PassSvPositionToViewportUV = FScreenTransform::SvPositionToViewportUV(View.UnscaledViewRect);
+		PassParameters->ViewportUVToHitProxyPixelPos = FScreenTransform::ChangeTextureBasisFromTo(
+			FScreenPassTextureViewport(HitProxyTexture, View.ViewRect), FScreenTransform::ETextureBasis::ViewportUV, FScreenTransform::ETextureBasis::TexelPosition);
+		PassParameters->HitProxyPixelPosMin = View.ViewRect.Min;
+		PassParameters->HitProxyPixelPosMax = View.ViewRect.Max - 1;
+
+		PassParameters->UndistortingDisplacementTexture = GSystemTextures.GetBlackDummy(GraphBuilder);
+		PassParameters->UndistortingDisplacementSampler = TStaticSamplerState<SF_Bilinear, AM_Clamp, AM_Clamp, AM_Clamp>::GetRHI();
+		if (View.LensDistortionLUT.IsEnabled())
+		{
+			PassParameters->UndistortingDisplacementTexture = View.LensDistortionLUT.UndistortingDisplacementTexture;
+		}
 		PassParameters->HitProxyTexture = HitProxyTexture;
 		PassParameters->RenderTargets[0] = FRenderTargetBinding(ViewFamilyTexture, ERenderTargetLoadAction::ELoad);
 
-		GraphBuilder.AddPass(
-			RDG_EVENT_NAME("HitProxies::CopyOutput"),
-			PassParameters,
-			ERDGPassFlags::Raster,
-			[&Views, HitProxyTextureExtent, HitProxyTexture, ViewFamilyTexture, FeatureLevel](FRHICommandListImmediate& RHICmdList)
-		{
-			// Set up a FTexture that is used to draw the hit proxy buffer to the view family's render target.
-			FTexture HitProxyRenderTargetTexture;
-			HitProxyRenderTargetTexture.TextureRHI = HitProxyTexture->GetRHI();
-			HitProxyRenderTargetTexture.SamplerStateRHI = TStaticSamplerState<>::GetRHI();
+		TShaderMapRef<FHitProxyCopyPS> PixelShader(View.ShaderMap);
 
-			// Generate the vertices and triangles mapping the hit proxy RT pixels into the view family's RT pixels.
-			FBatchedElements BatchedElements;
-			for (int32 ViewIndex = 0; ViewIndex < Views.Num(); ViewIndex++)
-			{
-				const FViewInfo& View = Views[ViewIndex];
-
-				float InvBufferSizeX = 1.0f / HitProxyTextureExtent.X;
-				float InvBufferSizeY = 1.0f / HitProxyTextureExtent.Y;
-
-				const float U0 = View.ViewRect.Min.X * InvBufferSizeX;
-				const float V0 = View.ViewRect.Min.Y * InvBufferSizeY;
-				const float U1 = View.ViewRect.Max.X * InvBufferSizeX;
-				const float V1 = View.ViewRect.Max.Y * InvBufferSizeY;
-
-				// Note: High DPI .  We are drawing to the size of the unscaled view rect because that is the size of the views render target
-				// if we do not do this clicking would be off.
-				const int32 V00 = BatchedElements.AddVertexf(FVector4f(View.UnscaledViewRect.Min.X, View.UnscaledViewRect.Min.Y, 0, 1), FVector2f(U0, V0), FLinearColor::White, FHitProxyId());
-				const int32 V10 = BatchedElements.AddVertexf(FVector4f(View.UnscaledViewRect.Max.X, View.UnscaledViewRect.Min.Y, 0, 1), FVector2f(U1, V0), FLinearColor::White, FHitProxyId());
-				const int32 V01 = BatchedElements.AddVertexf(FVector4f(View.UnscaledViewRect.Min.X, View.UnscaledViewRect.Max.Y, 0, 1), FVector2f(U0, V1), FLinearColor::White, FHitProxyId());
-				const int32 V11 = BatchedElements.AddVertexf(FVector4f(View.UnscaledViewRect.Max.X, View.UnscaledViewRect.Max.Y, 0, 1), FVector2f(U1, V1), FLinearColor::White, FHitProxyId());
-
-				BatchedElements.AddTriangle(V00, V10, V11, &HitProxyRenderTargetTexture, BLEND_Opaque);
-				BatchedElements.AddTriangle(V00, V11, V01, &HitProxyRenderTargetTexture, BLEND_Opaque);
-			}
-
-			// Generate a transform which maps from view family RT pixel coordinates to Normalized Device Coordinates.
-			FIntPoint ViewFamilyTextureExtent = ViewFamilyTexture->Desc.Extent;
-
-			const FMatrix PixelToView =
-				FTranslationMatrix(FVector(0, 0, 0)) *
-				FMatrix(
-					FPlane(1.0f / ((float)ViewFamilyTextureExtent.X / 2.0f), 0.0, 0.0f, 0.0f),
-					FPlane(0.0f, -GProjectionSignY / ((float)ViewFamilyTextureExtent.Y / 2.0f), 0.0f, 0.0f),
-					FPlane(0.0f, 0.0f, 1.0f, 0.0f),
-					FPlane(-1.0f, GProjectionSignY, 0.0f, 1.0f)
-				);
-
-			FSceneView SceneView = FBatchedElements::CreateProxySceneView(PixelToView, FIntRect(0, 0, ViewFamilyTextureExtent.X, ViewFamilyTextureExtent.Y));
-			FMeshPassProcessorRenderState DrawRenderState;
-
-			DrawRenderState.SetDepthStencilState(TStaticDepthStencilState<false, CF_Always>::GetRHI());
-			DrawRenderState.SetBlendState(TStaticBlendState<>::GetRHI());
-
-			BatchedElements.Draw(
-				RHICmdList,
-				DrawRenderState,
-				FeatureLevel,
-				SceneView,
-				false,
-				1.0f
-			);
-		});
+		FPixelShaderUtils::AddFullscreenPass(
+			GraphBuilder,
+			View.ShaderMap,
+			RDG_EVENT_NAME("HitProxyCopy %dx%d", View.UnscaledViewRect.Width(), View.UnscaledViewRect.Height()),
+			PixelShader,
+			PassParameters,
+			View.UnscaledViewRect);
 	}
 }
 #endif
