Don't convert scene uniform buffer to external unless explicitly asked to. Reduces render thread time for non-ray traced platforms.

diff --git a/Engine/Source/Runtime/Renderer/Private/SceneUniformBuffer.cpp b/Engine/Source/Runtime/Renderer/Private/SceneUniformBuffer.cpp
--- a/Engine/Source/Runtime/Renderer/Private/SceneUniformBuffer.cpp
+++ b/Engine/Source/Runtime/Renderer/Private/SceneUniformBuffer.cpp
@@ -128,7 +128,6 @@ FSceneUniformParameters::FSceneUniformParameters()
 
 FSceneUniformBuffer::FSceneUniformBuffer() :
 	Buffer(nullptr),
-	RHIBuffer(nullptr),
 	bAnyMemberDirty(false),
 	MemberHasBeenSet(false, FRegistry::FImpl::MemberInfos.Num())
 {
@@ -207,13 +206,9 @@ RENDERER_API TRDGUniformBufferRef<FSceneUniformParameters> FSceneUniformBuffer::
 			}
 		}
 		// Create and copy cached parameters into the RDG-lifetime struct
-		//auto DataCopy = GraphBuilder.Alloc(CachedData.Num(), SHADER_PARAMETER_STRUCT_ALIGNMENT);
-		//FPlatformMemory::Memcpy(DataCopy, CachedData.GetData(), CachedData.Num());
-		//Buffer = GraphBuilder.CreateUniformBuffer(GraphBuilder.AllocObject<FSceneUniformParameters>(DataCopy));
 		FSceneUniformParameters* ParameterStruct = GraphBuilder.AllocObject<FSceneUniformParameters>();
 		FPlatformMemory::Memcpy(ParameterStruct->Data, CachedData.GetData(), CachedData.Num());
 		Buffer = GraphBuilder.CreateUniformBuffer(ParameterStruct);
-		RHIBuffer = GraphBuilder.ConvertToExternalUniformBuffer(Buffer); //RT pipeline can't bind RDG UBs, so we must convert to external
 
 		bAnyMemberDirty = false;
 	}
@@ -222,9 +217,7 @@ RENDERER_API TRDGUniformBufferRef<FSceneUniformParameters> FSceneUniformBuffer::
 
 RENDERER_API FRHIUniformBuffer* FSceneUniformBuffer::GetBufferRHI(FRDGBuilder& GraphBuilder)
 {
-	// Ensure the buffer is prepped.
-	GetBuffer(GraphBuilder);
-	return RHIBuffer;
+	return GraphBuilder.ConvertToExternalUniformBuffer(Buffer);
 }
 
 #if !UE_BUILD_SHIPPING
diff --git a/Engine/Source/Runtime/Renderer/Public/SceneUniformBuffer.h b/Engine/Source/Runtime/Renderer/Public/SceneUniformBuffer.h
--- a/Engine/Source/Runtime/Renderer/Public/SceneUniformBuffer.h
+++ b/Engine/Source/Runtime/Renderer/Public/SceneUniformBuffer.h
@@ -213,7 +213,6 @@ private:
 	RENDERER_API const void* GetOrDefault(const FMemberId MemberId, const int32 ExpectedValueSize, FRDGBuilder& GraphBuilder) UE_LIFETIMEBOUND;
 
 	TRDGUniformBufferRef<FSceneUniformParameters> Buffer;
-	FRHIUniformBuffer* RHIBuffer;
 
 	bool bAnyMemberDirty;
 	TBitArray<> MemberHasBeenSet;
