Fixes motion blur, VisualizeMotionBlur, VisualizeReprojection when using r.TSR.LensDistortion

diff --git a/Engine/Shaders/Private/MotionBlur/MotionBlurVelocityFlatten.usf b/Engine/Shaders/Private/MotionBlur/MotionBlurVelocityFlatten.usf
--- a/Engine/Shaders/Private/MotionBlur/MotionBlurVelocityFlatten.usf
+++ b/Engine/Shaders/Private/MotionBlur/MotionBlurVelocityFlatten.usf
@@ -3,6 +3,7 @@
 #define CONFIG_MAX_RANGE_SIZE DIM_BLUR_DIRECTIONS
 
 #include "MotionBlurVelocityFlatten.ush"
+#include "../LensDistortion.ush"
 
 
 //------------------------------------------------------- CONSTANTS
@@ -18,12 +19,20 @@
 
 //------------------------------------------------------- PARAMETERS
 
+Texture2D<float2> DistortingDisplacementTexture;
+SamplerState      DistortingDisplacementSampler;
+
+Texture2D<float2> UndistortingDisplacementTexture;
+SamplerState      UndistortingDisplacementSampler;
+
 SCREEN_PASS_TEXTURE_VIEWPORT(Velocity)
 
 float4x4 ClipToPrevClipOverride;
-
+FScreenTransform ThreadIdToViewportUV;
+FScreenTransform ViewportUVToPixelPos;
 uint bCancelCameraMotion;
 uint bAddCustomCameraMotion;
+uint bLensDistortion;
 
 Texture2D VelocityTexture;
 Texture2D DepthTexture;
@@ -76,7 +85,17 @@ void MainCS(
 { 
 	float4 Debug = 0;
 
-	int2 PixelPos = min(DispatchThreadId.xy + Velocity_ViewportMin, Velocity_ViewportMax - 1);
+	int2 OutputPixelPos = min(DispatchThreadId.xy + Velocity_ViewportMin, Velocity_ViewportMax - 1);
+	int2 PixelPos = OutputPixelPos;
+
+	BRANCH
+	if (bLensDistortion)
+	{
+		float2 DistortedViewportUV = ApplyScreenTransform(float2(DispatchThreadId), ThreadIdToViewportUV);
+		float2 UndistortedViewportUV = ApplyLensDistortionOnViewportUV(UndistortingDisplacementTexture, UndistortingDisplacementSampler, DistortedViewportUV);
+		PixelPos = int2(ApplyScreenTransform(UndistortedViewportUV, ViewportUVToPixelPos));
+	}
+
 	int2 DepthOffset = 0;
 
 	{
@@ -130,7 +149,7 @@ void MainCS(
 
 	// Output velcity flatten texture.
 	{
-		uint2 OutputPos = select(PixelPos.xy < Velocity_ViewportMax, uint2(PixelPos), uint(~0).xx);
+		uint2 OutputPos = select(OutputPixelPos.xy < Velocity_ViewportMax, uint2(OutputPixelPos), uint(~0).xx);
 		OutVelocityFlatTexture[OutputPos] = EncodeVelocityFlatten(VelocityPolar, DeviceZ);
 	}
 
@@ -141,6 +160,6 @@ void MainCS(
 	}
 	
 #if DEBUG_MOTION_BLUR_OUTPUT
-	DebugOutput[PixelPos] = Debug;
+	DebugOutput[OutputPixelPos] = Debug;
 #endif
 }
diff --git a/Engine/Shaders/Private/MotionBlur/MotionBlurVisualize.usf b/Engine/Shaders/Private/MotionBlur/MotionBlurVisualize.usf
--- a/Engine/Shaders/Private/MotionBlur/MotionBlurVisualize.usf
+++ b/Engine/Shaders/Private/MotionBlur/MotionBlurVisualize.usf
@@ -3,6 +3,7 @@
 #define CONFIG_MAX_RANGE_SIZE 1
 
 #include "MotionBlurCommon.ush"
+#include "../LensDistortion.ush"
 
 
 //------------------------------------------------------- CONSTANTS
@@ -13,16 +14,20 @@
 
 //------------------------------------------------------- PARAMETERS
 
+Texture2D<float2> UndistortingDisplacementTexture;
+SamplerState      UndistortingDisplacementSampler;
+
 SCREEN_PASS_TEXTURE_VIEWPORT(Velocity)
 
-Texture2D ColorTexture;
 Texture2D VelocityTexture;
 Texture2D DepthTexture;
 
-SamplerState ColorSampler;
 SamplerState VelocitySampler;
 SamplerState DepthSampler;
 
+FScreenTransform SvPositionToScreenPos;
+FScreenTransform ScreenPosToColorUV;
+FScreenTransform ScreenPosToVelocityUV;
 int CheckerboardEnabled;
 
 #if USE_POST_MOTION_BLUR_TRANSLUCENCY
@@ -116,27 +121,32 @@ float Compute3DCheckerBoard(float3 Pos)
 //------------------------------------------------------- ENTRY POINT
 
 void MainPS(
-	noperspective float4 UVAndScreenPos : TEXCOORD0,
+	float4 SvPosition : SV_POSITION,
 	out float4 OutColor : SV_Target0)
 {
 	const float TileSize = 48;
 
-	float2 ScreenSpacePos = UVAndScreenPos.zw;
-	float2 PixelPos = ScreenSpacePos * Velocity_ScreenPosToViewportScale + Velocity_ScreenPosToViewportBias;
-	float2 UV = PixelPos * Velocity_ExtentInverse;
+	float2 DistortedScreenPos = ApplyScreenTransform(SvPosition.xy, SvPositionToScreenPos);
+	float2 UndistortedScreenPos = ApplyLensDistortionOnScreenPos(UndistortingDisplacementTexture, UndistortingDisplacementSampler, DistortedScreenPos);
+
+	float2 VelocityUV = ApplyScreenTransform(UndistortedScreenPos, ScreenPosToVelocityUV);
 
+	float2 PixelPos = DistortedScreenPos * Velocity_ScreenPosToViewportScale + Velocity_ScreenPosToViewportBias;
 	float2 PixelPosAtTileCenter = PixelPos - (frac(PixelPos / TileSize) - 0.5f) * TileSize;
-	float2 UVAtTileCenter = PixelPosAtTileCenter * Velocity_ExtentInverse;
-	float2 ScreenPosAtTileCenter = (PixelPosAtTileCenter - Velocity_ScreenPosToViewportBias) / Velocity_ScreenPosToViewportScale;
+
+	float2 DistortedScreenPosAtTileCenter = (PixelPosAtTileCenter - Velocity_ScreenPosToViewportBias) / Velocity_ScreenPosToViewportScale;
+	float2 UndistortedScreenPosAtTileCenter = ApplyLensDistortionOnScreenPos(UndistortingDisplacementTexture, UndistortingDisplacementSampler, DistortedScreenPosAtTileCenter);
+
+	float2 VelocityUVAtTileCenter = ApplyScreenTransform(UndistortedScreenPosAtTileCenter, ScreenPosToVelocityUV);
 
 	// World aligned checkerboards
 	OutColor = float4(0, 0, 0, 1);
 	if (CheckerboardEnabled)
 	{
-		float DeviceZ = DepthTexture.SampleLevel(SharedDepthSampler, UV, 0).r;
+		float DeviceZ = DepthTexture.SampleLevel(SharedDepthSampler, VelocityUV, 0).r;
 		float SceneDepth = ConvertFromDeviceZ(DeviceZ);
 
-		float3 ScreenVector = mul(float3(ScreenSpacePos, 1), DFToFloat3x3(PrimaryView.ScreenToWorld)).xyz;
+		float3 ScreenVector = mul(float3(UndistortedScreenPos, 1), DFToFloat3x3(PrimaryView.ScreenToWorld)).xyz;
 
 		// world space position of the current pixel
 		float3 OffsetWorldPos = ScreenVector * SceneDepth;
@@ -146,26 +156,26 @@ void MainPS(
 		OutColor = float4(lerp(WorldCheckerboard, float3(0,0,0), 0.7f), 1);
 	}
 
-	bool bSelectorOpaqueAtTileCenter = VelocityTexture.SampleLevel(SharedVelocitySampler, UVAtTileCenter, 0).x > 0;
-	bool bSelectorOpaque = VelocityTexture.SampleLevel(SharedVelocitySampler, UV, 0).x > 0;
+	bool bSelectorOpaqueAtTileCenter = VelocityTexture.SampleLevel(SharedVelocitySampler, VelocityUVAtTileCenter, 0).x > 0;
+	bool bSelectorOpaque = VelocityTexture.SampleLevel(SharedVelocitySampler, VelocityUV, 0).x > 0;
 
 	// relative, in screen space -1...1 -1..1, can be even outside of that range, points into the movement direction
 	float2 VelocityAtTileCenter;
 	{
-		VelocityAtTileCenter = DecodeVelocityFromTexture(VelocityTexture.SampleLevel(SharedVelocitySampler, UVAtTileCenter, 0)).xy;
+		VelocityAtTileCenter = DecodeVelocityFromTexture(VelocityTexture.SampleLevel(SharedVelocitySampler, VelocityUVAtTileCenter, 0)).xy;
 
 		// reconstruct from camera motion if we don't have velocity data
 		if(!bSelectorOpaqueAtTileCenter)
 		{
-			float DeviceZ = DepthTexture.SampleLevel(SharedDepthSampler, UVAtTileCenter, 0).r;
+			float DeviceZ = DepthTexture.SampleLevel(SharedDepthSampler, VelocityUVAtTileCenter, 0).r;
 
 			// UV is per tile
-			float4 ThisClip = float4(ScreenPosAtTileCenter, DeviceZ, 1);
+			float4 ThisClip = float4(UndistortedScreenPosAtTileCenter, DeviceZ, 1);
 			float4 PrevClip = mul(ThisClip, View.ClipToPrevClip);
 			float2 PrevScreen = PrevClip.xy / PrevClip.w;
 
 			// points into the movement direction
-			VelocityAtTileCenter = ScreenPosAtTileCenter - PrevScreen;
+			VelocityAtTileCenter = UndistortedScreenPosAtTileCenter - PrevScreen;
 		}
 	}
 
diff --git a/Engine/Shaders/Private/MotionBlur/VisualizeMotionVectors.usf b/Engine/Shaders/Private/MotionBlur/VisualizeMotionVectors.usf
--- a/Engine/Shaders/Private/MotionBlur/VisualizeMotionVectors.usf
+++ b/Engine/Shaders/Private/MotionBlur/VisualizeMotionVectors.usf
@@ -2,10 +2,16 @@
 
 #include "../Common.ush"
 #include "../ScreenPass.ush"
+#include "../LensDistortion.ush"
 
 
 //------------------------------------------------------- PARAMETERS
 
+Texture2D<float2> PrevDistortingDisplacementTexture;
+SamplerState      PrevDistortingDisplacementSampler;
+Texture2D<float2> UndistortingDisplacementTexture;
+SamplerState      UndistortingDisplacementSampler;
+
 Texture2D ColorTexture;
 Texture2D VelocityTexture;
 Texture2D DepthTexture;
@@ -20,9 +26,9 @@ SamplerState VelocitySampler;
 SamplerState DepthSampler;
 SamplerState PrevColorSampler;
 
-FScreenTransform SvPositionToVelocity;
-FScreenTransform SvPositionToColor;
 FScreenTransform SvPositionToScreenPos;
+FScreenTransform ScreenPosToVelocity;
+FScreenTransform ScreenPosToColor;
 FScreenTransform PrevScreenPosToPrevColor;
 uint VisualizeType;
 
@@ -46,29 +52,32 @@ void MainPS(
 	float4 SvPosition : SV_POSITION,
 	out float4 OutColor : SV_Target0)
 {
-	float2 ColorTextureUV = ApplyScreenTransform(SvPosition.xy, SvPositionToColor);
-	float2 VelocityTextureUV = ApplyScreenTransform(SvPosition.xy, SvPositionToVelocity);
-	float2 ScreenPos = ApplyScreenTransform(SvPosition.xy, SvPositionToScreenPos);
-
-	ColorTextureUV = clamp(ColorTextureUV, Color_UVViewportBilinearMin, Color_UVViewportBilinearMax);
-	VelocityTextureUV = clamp(VelocityTextureUV, Velocity_UVViewportBilinearMin, Velocity_UVViewportBilinearMax);
-
-	float4 Color = ColorTexture.SampleLevel(ColorSampler, ColorTextureUV, 0.0);
-	float4 EncodedVelocity = VelocityTexture.SampleLevel(VelocitySampler, VelocityTextureUV, 0.0);
-	float DeviceZ = DepthTexture.SampleLevel(DepthSampler, VelocityTextureUV, 0.0).r;
-	
-	float2 ScreenVelocity = ComputeStaticVelocity(ScreenPos, DeviceZ);
-	bool bIsRenderedVelocity = EncodedVelocity.x > 0.0;
-	if (bIsRenderedVelocity)
-	{
-		ScreenVelocity = DecodeVelocityFromTexture(EncodedVelocity).xy;
-	}
+	float2 DistortedScreenPos = ApplyScreenTransform(SvPosition.xy, SvPositionToScreenPos);
+	float2 UndistortedScreenPos = ApplyLensDistortionOnScreenPos(UndistortingDisplacementTexture, UndistortingDisplacementSampler, DistortedScreenPos);
 
 	BRANCH
 	if (VisualizeType == 0) // EVisualizeMotionVectors::ReprojectionAlignment
 	{
-		float2 PrevScreenPos = ScreenPos - ScreenVelocity;
-		float2 PrevColorTextureUV = ApplyScreenTransform(PrevScreenPos, PrevScreenPosToPrevColor);
+		float2 ColorTextureUV = ApplyScreenTransform(DistortedScreenPos, ScreenPosToColor);
+		float2 VelocityTextureUV = ApplyScreenTransform(UndistortedScreenPos, ScreenPosToVelocity);
+
+		ColorTextureUV = clamp(ColorTextureUV, Color_UVViewportBilinearMin, Color_UVViewportBilinearMax);
+		VelocityTextureUV = clamp(VelocityTextureUV, Velocity_UVViewportBilinearMin, Velocity_UVViewportBilinearMax);
+
+		float4 Color = ColorTexture.SampleLevel(ColorSampler, ColorTextureUV, 0.0);
+		float4 EncodedVelocity = VelocityTexture.SampleLevel(VelocitySampler, VelocityTextureUV, 0.0);
+		float DeviceZ = DepthTexture.SampleLevel(DepthSampler, VelocityTextureUV, 0.0).r;
+	
+		float2 ScreenVelocity = ComputeStaticVelocity(UndistortedScreenPos, DeviceZ);
+		bool bIsRenderedVelocity = EncodedVelocity.x > 0.0;
+		if (bIsRenderedVelocity)
+		{
+			ScreenVelocity = DecodeVelocityFromTexture(EncodedVelocity).xy;
+		}
+
+		float2 UndistortedPrevScreenPos = UndistortedScreenPos - ScreenVelocity;
+		float2 DistortedPrevScreenPos = ApplyLensDistortionOnScreenPos(PrevDistortingDisplacementTexture, PrevDistortingDisplacementSampler, UndistortedPrevScreenPos);
+		float2 PrevColorTextureUV = ApplyScreenTransform(UndistortedPrevScreenPos, PrevScreenPosToPrevColor);
 		PrevColorTextureUV = clamp(PrevColorTextureUV, Color_UVViewportBilinearMin, Color_UVViewportBilinearMax);
 
 		float4 PrevColor = PrevColorTexture.SampleLevel(PrevColorSampler, PrevColorTextureUV, 0.0);
@@ -93,6 +102,11 @@ void MainPS(
 	#if VELOCITY_ENCODE_HAS_PIXEL_ANIMATION
 	else if (VisualizeType == 1) // EVisualizeMotionVectors::HasPixelAnimationFlag
 	{
+		// Don't apply lens distortion on HasPixelAnimation to be consistent with the rest of the VisualizeTemporalUpscaler's inputs.
+		float2 VelocityTextureUV = ApplyScreenTransform(DistortedScreenPos, ScreenPosToVelocity);
+		VelocityTextureUV = clamp(VelocityTextureUV, Velocity_UVViewportBilinearMin, Velocity_UVViewportBilinearMax);
+		
+		float4 EncodedVelocity = VelocityTexture.SampleLevel(VelocitySampler, VelocityTextureUV, 0.0);
 		bool bHasPixelAnimation = DecodeHasPixelAnimationFromVelocityTexture(EncodedVelocity);
 		
 		OutColor = bHasPixelAnimation ? 1.0 : 0.0;
diff --git a/Engine/Source/Runtime/Renderer/Private/PostProcess/PostProcessMotionBlur.cpp b/Engine/Source/Runtime/Renderer/Private/PostProcess/PostProcessMotionBlur.cpp
--- a/Engine/Source/Runtime/Renderer/Private/PostProcess/PostProcessMotionBlur.cpp
+++ b/Engine/Source/Runtime/Renderer/Private/PostProcess/PostProcessMotionBlur.cpp
@@ -10,6 +10,7 @@
 #include "PostProcess/PostProcessing.h"
 #include "VelocityRendering.h"
 #include "UnrealEngine.h"
+#include "PixelShaderUtils.h"
 
 namespace
 {
@@ -322,9 +323,16 @@ public:
 		SHADER_PARAMETER_STRUCT(FScreenPassTextureViewportParameters, Velocity)
 		SHADER_PARAMETER_STRUCT_INCLUDE(FVelocityFlattenParameters, VelocityFlattenParameters)
 		SHADER_PARAMETER(FMatrix44f, ClipToPrevClipOverride)
+		SHADER_PARAMETER(FScreenTransform, ThreadIdToViewportUV)
+		SHADER_PARAMETER(FScreenTransform, ViewportUVToPixelPos)
 		SHADER_PARAMETER(int32, bCancelCameraMotion)
 		SHADER_PARAMETER(int32, bAddCustomCameraMotion)
+		SHADER_PARAMETER(int32, bLensDistortion)
 
+		SHADER_PARAMETER_RDG_TEXTURE(Texture2D, DistortingDisplacementTexture)
+		SHADER_PARAMETER_SAMPLER(SamplerState, DistortingDisplacementSampler)
+		SHADER_PARAMETER_RDG_TEXTURE(Texture2D, UndistortingDisplacementTexture)
+		SHADER_PARAMETER_SAMPLER(SamplerState, UndistortingDisplacementSampler)
 		SHADER_PARAMETER_RDG_TEXTURE(Texture2D, VelocityTexture)
 		SHADER_PARAMETER_RDG_TEXTURE(Texture2D, DepthTexture)
 
@@ -518,17 +526,18 @@ class FMotionBlurVisualizePS : public FMotionBlurShader
 	BEGIN_SHADER_PARAMETER_STRUCT(FParameters, )
 		SHADER_PARAMETER(FMatrix44f, WorldToClipPrev)
 		SHADER_PARAMETER_STRUCT_REF(FViewUniformShaderParameters, View)
-		SHADER_PARAMETER_RDG_TEXTURE_SRV(Texture2D, ColorTexture)
+		SHADER_PARAMETER_RDG_TEXTURE(Texture2D, UndistortingDisplacementTexture)
+		SHADER_PARAMETER_SAMPLER(SamplerState,  UndistortingDisplacementSampler)
 		SHADER_PARAMETER_RDG_TEXTURE(Texture2D, DepthTexture)
 		SHADER_PARAMETER_RDG_TEXTURE(Texture2D, VelocityTexture)
 
-		SHADER_PARAMETER_STRUCT(FScreenPassTextureViewportParameters, Color)
 		SHADER_PARAMETER_STRUCT(FScreenPassTextureViewportParameters, Velocity)
 
-		SHADER_PARAMETER_SAMPLER(SamplerState, ColorSampler)
 		SHADER_PARAMETER_SAMPLER(SamplerState, VelocitySampler)
 		SHADER_PARAMETER_SAMPLER(SamplerState, DepthSampler)
 
+		SHADER_PARAMETER(FScreenTransform, SvPositionToScreenPos)
+		SHADER_PARAMETER(FScreenTransform, ScreenPosToVelocityUV)
 		SHADER_PARAMETER(int, CheckerboardEnabled)
 
 		RENDER_TARGET_BINDING_SLOTS()
@@ -663,6 +672,21 @@ void AddMotionBlurVelocityPass(
 			PassParameters->ClipToPrevClipOverride = FMatrix44f(View.ClipToPrevClipOverride.GetValue());
 		}
 
+		PassParameters->DistortingDisplacementTexture = GSystemTextures.GetBlackDummy(GraphBuilder);
+		PassParameters->DistortingDisplacementSampler = TStaticSamplerState<SF_Bilinear, AM_Clamp, AM_Clamp, AM_Clamp>::GetRHI();
+		PassParameters->UndistortingDisplacementTexture = GSystemTextures.GetBlackDummy(GraphBuilder);
+		PassParameters->UndistortingDisplacementSampler = TStaticSamplerState<SF_Bilinear, AM_Clamp, AM_Clamp, AM_Clamp>::GetRHI();
+		PassParameters->bLensDistortion = Inputs.LensDistortionLUT.IsEnabled();
+		if (Inputs.LensDistortionLUT.IsEnabled())
+		{
+			PassParameters->DistortingDisplacementTexture = Inputs.LensDistortionLUT.DistortingDisplacementTexture;
+			PassParameters->UndistortingDisplacementTexture = Inputs.LensDistortionLUT.UndistortingDisplacementTexture;
+		}
+
+		PassParameters->ThreadIdToViewportUV = FScreenTransform::DispatchThreadIdToViewportUV(Viewports.Velocity.Rect);
+		PassParameters->ViewportUVToPixelPos = FScreenTransform::ChangeTextureBasisFromTo(
+			FScreenPassTextureViewport(Viewports.Velocity), FScreenTransform::ETextureBasis::ViewportUV, FScreenTransform::ETextureBasis::TexelPosition);
+
 		PassParameters->Velocity = Viewports.VelocityParameters;
 		PassParameters->DepthTexture = Inputs.SceneDepth.Texture;
 		PassParameters->VelocityTexture = Inputs.SceneVelocity.Texture;
@@ -1140,23 +1164,42 @@ FScreenPassTextureSlice AddVisualizeMotionBlurPass(FRDGBuilder& GraphBuilder, co
 
 	const bool bVisualizeDebugInfo = CVarVisualizeMotionBlurEnableDebugInformation.GetValueOnRenderThread();
 
-	FMotionBlurVisualizePS::FParameters* PassParameters = GraphBuilder.AllocParameters<FMotionBlurVisualizePS::FParameters>();
-	PassParameters->WorldToClipPrev = FMatrix44f(GetPreviousWorldToClipMatrix(View));		// LWC_TODO: Precision loss
-	PassParameters->View = View.ViewUniformBuffer;
-	PassParameters->ColorTexture = Inputs.SceneColor.TextureSRV;
-	PassParameters->DepthTexture = Inputs.SceneDepth.Texture;
-	PassParameters->VelocityTexture = Inputs.SceneVelocity.Texture;
-	PassParameters->Color = Viewports.ColorParameters;
-	PassParameters->Velocity = Viewports.VelocityParameters;
-	PassParameters->ColorSampler = GetMotionBlurColorSampler();
-	PassParameters->VelocitySampler = GetMotionBlurVelocitySampler();
-	PassParameters->DepthSampler = GetMotionBlurVelocitySampler();
-	PassParameters->RenderTargets[0] = Output.GetRenderTargetBinding();
-	PassParameters->CheckerboardEnabled = bVisualizeDebugInfo;
-
-	TShaderMapRef<FMotionBlurVisualizePS> PixelShader(View.ShaderMap);
-
-	AddDrawScreenPass(GraphBuilder, RDG_EVENT_NAME("VisualizeMotionBlur"), View, Viewports.Color, Viewports.Color, PixelShader, PassParameters);
+	{
+		FMotionBlurVisualizePS::FParameters* PassParameters = GraphBuilder.AllocParameters<FMotionBlurVisualizePS::FParameters>();
+		PassParameters->WorldToClipPrev = FMatrix44f(GetPreviousWorldToClipMatrix(View));		// LWC_TODO: Precision loss
+		PassParameters->View = View.ViewUniformBuffer;
+
+		PassParameters->UndistortingDisplacementTexture = GSystemTextures.GetBlackDummy(GraphBuilder);
+		PassParameters->UndistortingDisplacementSampler = TStaticSamplerState<SF_Bilinear, AM_Clamp, AM_Clamp, AM_Clamp>::GetRHI();
+		if (Inputs.LensDistortionLUT.IsEnabled())
+		{
+			PassParameters->UndistortingDisplacementTexture = Inputs.LensDistortionLUT.UndistortingDisplacementTexture;
+		}
+
+		PassParameters->DepthTexture = Inputs.SceneDepth.Texture;
+		PassParameters->VelocityTexture = Inputs.SceneVelocity.Texture;
+		PassParameters->Velocity = Viewports.VelocityParameters;
+		PassParameters->VelocitySampler = GetMotionBlurVelocitySampler();
+		PassParameters->DepthSampler = GetMotionBlurVelocitySampler();
+		PassParameters->RenderTargets[0] = Output.GetRenderTargetBinding();
+
+		PassParameters->SvPositionToScreenPos = (
+			FScreenTransform::SvPositionToViewportUV(Output.ViewRect) *
+			FScreenTransform::ViewportUVToScreenPos);
+		PassParameters->ScreenPosToVelocityUV = FScreenTransform::ChangeTextureBasisFromTo(
+			Inputs.SceneDepth, FScreenTransform::ETextureBasis::ScreenPosition, FScreenTransform::ETextureBasis::TextureUV);
+
+		PassParameters->CheckerboardEnabled = bVisualizeDebugInfo;
+
+		TShaderMapRef<FMotionBlurVisualizePS> PixelShader(View.ShaderMap);
+		FPixelShaderUtils::AddFullscreenPass(
+			GraphBuilder,
+			View.ShaderMap,
+			RDG_EVENT_NAME("VisualizeMotionVectors %dx%d", Output.ViewRect.Width(), Output.ViewRect.Height()),
+			PixelShader,
+			PassParameters,
+			Output.ViewRect);
+	}
 
 	Output.LoadAction = ERenderTargetLoadAction::ELoad;
 
diff --git a/Engine/Source/Runtime/Renderer/Private/PostProcess/PostProcessMotionBlur.h b/Engine/Source/Runtime/Renderer/Private/PostProcess/PostProcessMotionBlur.h
--- a/Engine/Source/Runtime/Renderer/Private/PostProcess/PostProcessMotionBlur.h
+++ b/Engine/Source/Runtime/Renderer/Private/PostProcess/PostProcessMotionBlur.h
@@ -8,6 +8,7 @@
 
 #include "ScreenPass.h"
 #include "TranslucentRendering.h"
+#include "PostProcess/LensDistortion.h"
 
 
 // Returns whether motion blur is enabled for the requested view.
@@ -98,6 +99,9 @@ struct FMotionBlurInputs
 
 	// [Optional] the velocity flatten is can be generated by a earlier pass.
 	FVelocityFlattenTextures VelocityFlattenTextures;
+
+	// [Optional] Lens distortion applied on the scene color.
+	FLensDistortionLUT LensDistortionLUT;
 };
 
 struct FMotionBlurOutputs
diff --git a/Engine/Source/Runtime/Renderer/Private/PostProcess/PostProcessing.cpp b/Engine/Source/Runtime/Renderer/Private/PostProcess/PostProcessing.cpp
--- a/Engine/Source/Runtime/Renderer/Private/PostProcess/PostProcessing.cpp
+++ b/Engine/Source/Runtime/Renderer/Private/PostProcess/PostProcessing.cpp
@@ -827,7 +827,7 @@ void AddPostProcessingPasses(
 				else
 				{
 					UpscalerPassInputs.bGenerateOutputMip1 = bMotionBlurNeedsHalfResInput;
-					UpscalerPassInputs.bGenerateVelocityFlattenTextures = FVelocityFlattenTextures::AllowExternal(View) && !bVisualizeMotionBlur;
+					UpscalerPassInputs.bGenerateVelocityFlattenTextures = FVelocityFlattenTextures::AllowExternal(View) && !bVisualizeMotionBlur && !bApplyLensDistortionInTSR;
 				}
 			}
 			else if (PostProcessMaterialBeforeBloomChain.Num() > 0)
@@ -969,6 +969,10 @@ void AddPostProcessingPasses(
 			PassInputs.Quality = GetMotionBlurQuality();
 			PassInputs.Filter = GetMotionBlurFilter();
 			PassInputs.VelocityFlattenTextures = VelocityFlattenTextures;
+			if (bApplyLensDistortionInTSR)
+			{
+				PassInputs.LensDistortionLUT = View.LensDistortionLUT;
+			}
 
 			// Motion blur visualization replaces motion blur when enabled.
 			if (bVisualizeMotionBlur)
@@ -1677,6 +1681,10 @@ void AddPostProcessingPasses(
 		PassInputs.SceneColor = SceneColor;
 		PassInputs.SceneDepth = SceneDepth;
 		PassInputs.SceneVelocity = Velocity;
+		if (bApplyLensDistortionInTSR)
+		{
+			PassInputs.LensDistortionLUT = View.LensDistortionLUT;
+		}
 
 		SceneColor = AddVisualizeMotionVectorsPass(GraphBuilder, View, PassInputs, EVisualizeMotionVectors::ReprojectionAlignment);
 	}
diff --git a/Engine/Source/Runtime/Renderer/Private/PostProcess/VisualizeMotionVectors.cpp b/Engine/Source/Runtime/Renderer/Private/PostProcess/VisualizeMotionVectors.cpp
--- a/Engine/Source/Runtime/Renderer/Private/PostProcess/VisualizeMotionVectors.cpp
+++ b/Engine/Source/Runtime/Renderer/Private/PostProcess/VisualizeMotionVectors.cpp
@@ -14,6 +14,10 @@ class FVisualizeMotionVectorsPS : public FGlobalShader
 
 	BEGIN_SHADER_PARAMETER_STRUCT(FParameters, )
 		SHADER_PARAMETER_STRUCT_REF(FViewUniformShaderParameters, View)
+		SHADER_PARAMETER_RDG_TEXTURE(Texture2D, PrevDistortingDisplacementTexture)
+		SHADER_PARAMETER_SAMPLER(SamplerState,  PrevDistortingDisplacementSampler)
+		SHADER_PARAMETER_RDG_TEXTURE(Texture2D, UndistortingDisplacementTexture)
+		SHADER_PARAMETER_SAMPLER(SamplerState,  UndistortingDisplacementSampler)
 		SHADER_PARAMETER_RDG_TEXTURE(Texture2D, ColorTexture)
 		SHADER_PARAMETER_RDG_TEXTURE(Texture2D, DepthTexture)
 		SHADER_PARAMETER_RDG_TEXTURE(Texture2D, VelocityTexture)
@@ -28,9 +32,9 @@ class FVisualizeMotionVectorsPS : public FGlobalShader
 		SHADER_PARAMETER_SAMPLER(SamplerState, DepthSampler)
 		SHADER_PARAMETER_SAMPLER(SamplerState, PrevColorSampler)
 
-		SHADER_PARAMETER(FScreenTransform, SvPositionToVelocity)
-		SHADER_PARAMETER(FScreenTransform, SvPositionToColor)
 		SHADER_PARAMETER(FScreenTransform, SvPositionToScreenPos)
+		SHADER_PARAMETER(FScreenTransform, ScreenPosToVelocity)
+		SHADER_PARAMETER(FScreenTransform, ScreenPosToColor)
 		SHADER_PARAMETER(FScreenTransform, PrevScreenPosToPrevColor)
 		SHADER_PARAMETER(int32, VisualizeType)
 
@@ -69,6 +73,19 @@ FScreenPassTexture AddVisualizeMotionVectorsPass(FRDGBuilder& GraphBuilder, cons
 		FVisualizeMotionVectorsPS::FParameters* PassParameters = GraphBuilder.AllocParameters<FVisualizeMotionVectorsPS::FParameters>();
 		PassParameters->View = View.ViewUniformBuffer;
 
+		PassParameters->PrevDistortingDisplacementTexture = GSystemTextures.GetBlackDummy(GraphBuilder);
+		PassParameters->PrevDistortingDisplacementSampler = TStaticSamplerState<SF_Bilinear, AM_Clamp, AM_Clamp, AM_Clamp>::GetRHI();
+		PassParameters->UndistortingDisplacementTexture = GSystemTextures.GetBlackDummy(GraphBuilder);
+		PassParameters->UndistortingDisplacementSampler = TStaticSamplerState<SF_Bilinear, AM_Clamp, AM_Clamp, AM_Clamp>::GetRHI();
+		if (View.PrevViewInfo.bIsVisualizeMotionVectorsDistorted)
+		{
+			PassParameters->PrevDistortingDisplacementTexture = GraphBuilder.RegisterExternalTexture(View.PrevViewInfo.DistortingDisplacementTexture);
+		}
+		if (Inputs.LensDistortionLUT.IsEnabled())
+		{
+			PassParameters->UndistortingDisplacementTexture = Inputs.LensDistortionLUT.UndistortingDisplacementTexture;
+		}
+
 		PassParameters->ColorTexture = Inputs.SceneColor.Texture;
 		PassParameters->DepthTexture = Inputs.SceneDepth.Texture;
 		PassParameters->VelocityTexture = Inputs.SceneVelocity.Texture;
@@ -83,18 +100,15 @@ FScreenPassTexture AddVisualizeMotionVectorsPass(FRDGBuilder& GraphBuilder, cons
 		PassParameters->DepthSampler     = TStaticSamplerState<SF_Point, AM_Clamp, AM_Clamp, AM_Clamp>::GetRHI();
 		PassParameters->PrevColorSampler = TStaticSamplerState<SF_Bilinear, AM_Clamp, AM_Clamp, AM_Clamp>::GetRHI();
 
-		FScreenTransform SvPositionToViewportUV = FScreenTransform::SvPositionToViewportUV(Output.ViewRect);
-		PassParameters->SvPositionToColor = (
-			SvPositionToViewportUV *
-			FScreenTransform::ChangeTextureBasisFromTo(Inputs.SceneColor, FScreenTransform::ETextureBasis::ViewportUV, FScreenTransform::ETextureBasis::TextureUV));
-		PassParameters->SvPositionToVelocity = (
-			SvPositionToViewportUV *
-			FScreenTransform::ChangeTextureBasisFromTo(Inputs.SceneDepth, FScreenTransform::ETextureBasis::ViewportUV, FScreenTransform::ETextureBasis::TextureUV));
 		PassParameters->SvPositionToScreenPos = (
-			SvPositionToViewportUV *
+			FScreenTransform::SvPositionToViewportUV(Output.ViewRect) *
 			FScreenTransform::ViewportUVToScreenPos);
-		PassParameters->PrevScreenPosToPrevColor = (
-			FScreenTransform::ChangeTextureBasisFromTo(PrevColor, FScreenTransform::ETextureBasis::ScreenPosition, FScreenTransform::ETextureBasis::TextureUV));
+		PassParameters->ScreenPosToColor = FScreenTransform::ChangeTextureBasisFromTo(
+			Inputs.SceneColor, FScreenTransform::ETextureBasis::ScreenPosition, FScreenTransform::ETextureBasis::TextureUV);
+		PassParameters->ScreenPosToVelocity = FScreenTransform::ChangeTextureBasisFromTo(
+			Inputs.SceneDepth, FScreenTransform::ETextureBasis::ScreenPosition, FScreenTransform::ETextureBasis::TextureUV);
+		PassParameters->PrevScreenPosToPrevColor = FScreenTransform::ChangeTextureBasisFromTo(
+			PrevColor, FScreenTransform::ETextureBasis::ScreenPosition, FScreenTransform::ETextureBasis::TextureUV);
 		PassParameters->VisualizeType = int32(Visualize);
 
 		PassParameters->RenderTargets[0] = Output.GetRenderTargetBinding();
@@ -113,6 +127,12 @@ FScreenPassTexture AddVisualizeMotionVectorsPass(FRDGBuilder& GraphBuilder, cons
 	{
 		GraphBuilder.QueueTextureExtraction(Inputs.SceneColor.Texture, &View.ViewState->PrevFrameViewInfo.VisualizeMotionVectors);
 		View.ViewState->PrevFrameViewInfo.VisualizeMotionVectorsRect = Inputs.SceneColor.ViewRect;
+		View.ViewState->PrevFrameViewInfo.bIsVisualizeMotionVectorsDistorted = Inputs.LensDistortionLUT.IsEnabled();
+		
+		if (Inputs.LensDistortionLUT.IsEnabled())
+		{
+			GraphBuilder.QueueTextureExtraction(Inputs.LensDistortionLUT.DistortingDisplacementTexture, &View.ViewState->PrevFrameViewInfo.DistortingDisplacementTexture);
+		}
 	}
 
 	return MoveTemp(Output);
diff --git a/Engine/Source/Runtime/Renderer/Private/PostProcess/VisualizeMotionVectors.h b/Engine/Source/Runtime/Renderer/Private/PostProcess/VisualizeMotionVectors.h
--- a/Engine/Source/Runtime/Renderer/Private/PostProcess/VisualizeMotionVectors.h
+++ b/Engine/Source/Runtime/Renderer/Private/PostProcess/VisualizeMotionVectors.h
@@ -3,6 +3,7 @@
 #pragma once
 
 #include "ScreenPass.h"
+#include "PostProcess/LensDistortion.h"
 
 
 enum class EVisualizeMotionVectors : uint8
@@ -19,6 +20,9 @@ struct FVisualizeMotionVectorsInputs
 	FScreenPassTexture SceneColor;
 	FScreenPassTexture SceneDepth;
 	FScreenPassTexture SceneVelocity;
+
+	// [Optional] Lens distortion applied on the scene color.
+	FLensDistortionLUT LensDistortionLUT;
 };
 
 FScreenPassTexture AddVisualizeMotionVectorsPass(FRDGBuilder& GraphBuilder, const FViewInfo& View, const FVisualizeMotionVectorsInputs& Inputs, EVisualizeMotionVectors Visualize);
\ No newline at end of file
diff --git a/Engine/Source/Runtime/Renderer/Private/PostProcess/VisualizeTemporalUpscaler.cpp b/Engine/Source/Runtime/Renderer/Private/PostProcess/VisualizeTemporalUpscaler.cpp
--- a/Engine/Source/Runtime/Renderer/Private/PostProcess/VisualizeTemporalUpscaler.cpp
+++ b/Engine/Source/Runtime/Renderer/Private/PostProcess/VisualizeTemporalUpscaler.cpp
@@ -134,6 +134,7 @@ FScreenPassTexture AddVisualizeTemporalUpscalerPass(FRDGBuilder& GraphBuilder, c
 			PassInputs.SceneColor = Inputs.Outputs.FullRes;
 			PassInputs.SceneDepth = Inputs.Inputs.SceneDepth;
 			PassInputs.SceneVelocity = Inputs.Inputs.SceneVelocity;
+			PassInputs.LensDistortionLUT = Inputs.Inputs.LensDistortionLUT;
 
 			FVisualizeBufferTile& Tile = Tiles[4 * 1 + 0];
 			Tile.Input = FScreenPassTexture(AddVisualizeMotionBlurPass(GraphBuilder, View, PassInputs));
@@ -147,6 +148,7 @@ FScreenPassTexture AddVisualizeTemporalUpscalerPass(FRDGBuilder& GraphBuilder, c
 			PassInputs.SceneColor = Inputs.SceneColor;
 			PassInputs.SceneDepth = Inputs.Inputs.SceneDepth;
 			PassInputs.SceneVelocity = Inputs.Inputs.SceneVelocity;
+			PassInputs.LensDistortionLUT = Inputs.Inputs.LensDistortionLUT;
 
 			FVisualizeBufferTile& Tile = Tiles[4 * 2 + 0];
 			Tile.Input = AddVisualizeMotionVectorsPass(GraphBuilder, View, PassInputs, EVisualizeMotionVectors::ReprojectionAlignment);
diff --git a/Engine/Source/Runtime/Renderer/Private/RendererScene.cpp b/Engine/Source/Runtime/Renderer/Private/RendererScene.cpp
--- a/Engine/Source/Runtime/Renderer/Private/RendererScene.cpp
+++ b/Engine/Source/Runtime/Renderer/Private/RendererScene.cpp
@@ -849,11 +849,21 @@ uint64 FTemporalAAHistory::GetGPUSizeBytes(bool bLogSizes) const
 
 uint64 FTSRHistory::GetGPUSizeBytes(bool bLogSizes) const
 {
-	return
+	uint64 TotalSize =
 		GetRenderTargetGPUSizeBytes(ColorArray, bLogSizes) +
 		GetRenderTargetGPUSizeBytes(MetadataArray, bLogSizes) +
 		GetRenderTargetGPUSizeBytes(GuideArray, bLogSizes) +
 		GetRenderTargetGPUSizeBytes(MoireArray, bLogSizes);
+
+	for (const TRefCountPtr <IPooledRenderTarget>& Texture : DistortingDisplacementTextures)
+	{
+		if (Texture.IsValid())
+		{
+			TotalSize += GetRenderTargetGPUSizeBytes(Texture, bLogSizes);
+		}
+	}
+
+	return TotalSize;
 }
 
 uint64 FScreenSpaceDenoiserHistory::GetGPUSizeBytes(bool bLogSizes) const
@@ -876,6 +886,7 @@ uint64 FPreviousViewInfo::GetGPUSizeBytes(bool bLogSizes) const
 		GetRenderTargetGPUSizeBytes(GBufferC, bLogSizes) +
 		GetRenderTargetGPUSizeBytes(HZB, bLogSizes) +
 		GetRenderTargetGPUSizeBytes(NaniteHZB, bLogSizes) +
+		GetRenderTargetGPUSizeBytes(DistortingDisplacementTexture, bLogSizes) +
 		GetRenderTargetGPUSizeBytes(CompressedDepthViewNormal, bLogSizes) +
 		GetRenderTargetGPUSizeBytes(CompressedOpaqueDepth, bLogSizes) +
 		GetRenderTargetGPUSizeBytes(CompressedOpaqueShadingModel, bLogSizes) +
diff --git a/Engine/Source/Runtime/Renderer/Private/SceneRendering.h b/Engine/Source/Runtime/Renderer/Private/SceneRendering.h
--- a/Engine/Source/Runtime/Renderer/Private/SceneRendering.h
+++ b/Engine/Source/Runtime/Renderer/Private/SceneRendering.h
@@ -973,6 +973,9 @@ struct FPreviousViewInfo
 	TRefCountPtr<IPooledRenderTarget> HZB;
 	TRefCountPtr<IPooledRenderTarget> NaniteHZB;
 
+	// Distorting displacement texture applied.
+	TRefCountPtr<IPooledRenderTarget> DistortingDisplacementTexture;
+
 	// Bit mask used to interpret per-instance occlusion query results for this view.
 	// Expected to contain a single active bit or zero if instance occlusion query data is not available.
 	// See FInstanceCullingOcclusionQueryRenderer.
@@ -1068,6 +1071,7 @@ struct FPreviousViewInfo
 	// Scene color used for reprojecting next frame to verify the motion vector reprojects correctly.
 	TRefCountPtr<IPooledRenderTarget> VisualizeMotionVectors;
 	FIntRect VisualizeMotionVectorsRect;
+	bool bIsVisualizeMotionVectorsDistorted = false;
 
 	uint64 GetGPUSizeBytes(bool bLogSizes) const;
 };
