Fix atlased shadowed projections being rendered incorrectly in instanced stereo primary views.

diff --git a/Engine/Source/Runtime/Renderer/Private/ShadowRendering.cpp b/Engine/Source/Runtime/Renderer/Private/ShadowRendering.cpp
--- a/Engine/Source/Runtime/Renderer/Private/ShadowRendering.cpp
+++ b/Engine/Source/Runtime/Renderer/Private/ShadowRendering.cpp
@@ -1069,8 +1069,13 @@ void FProjectedShadowInfo::SetupProjectionStencilMask(
 
 		RHICmdList.SetStreamSource(0, GFrustumVertexBuffer.VertexBufferRHI, 0);
 
+		// Shadow projection stenciling is special-cased to run per-view for instanced stereo views.
+		// TODO: Support instanced stereo properly in the projection stenciling pass.
+		const bool bIsInstancedStereoEmulated = View->bIsInstancedStereoEnabled && !View->bIsMobileMultiViewEnabled && IStereoRendering::IsStereoEyeView(*View);
+		const uint32 NumberOfInstances = bIsInstancedStereoEmulated ? 1 : View->GetStereoPassInstanceFactor();
+
 		// Draw the frustum using the stencil buffer to mask just the pixels which are inside the shadow frustum.
-		RHICmdList.DrawIndexedPrimitive(GCubeIndexBuffer.IndexBufferRHI, 0, 0, 8, 0, 12, View->GetStereoPassInstanceFactor());
+		RHICmdList.DrawIndexedPrimitive(GCubeIndexBuffer.IndexBufferRHI, 0, 0, 8, 0, 12, NumberOfInstances);
 
 		// if rendering modulated shadows mask out subject mesh elements to prevent self shadowing.
 		if (bMobileModulatedProjections && !CVarEnableModulatedSelfShadow.GetValueOnRenderThread())
