Fix for VFog voxelisation not working with Substrate OFF.

diff --git a/Engine/Source/Runtime/Renderer/Private/VolumetricFog.cpp b/Engine/Source/Runtime/Renderer/Private/VolumetricFog.cpp
--- a/Engine/Source/Runtime/Renderer/Private/VolumetricFog.cpp
+++ b/Engine/Source/Runtime/Renderer/Private/VolumetricFog.cpp
@@ -1116,8 +1116,10 @@ IMPLEMENT_GLOBAL_SHADER(FVolumetricFogFinalIntegrationCS, "/Engine/Private/Volum
 
 bool DoesPlatformSupportVolumetricFogVoxelization(const FStaticShaderPlatform Platform)
 {
-	const bool bCanRHICompileHlsl2021GeometryShaders = Substrate::IsSubstrateEnabled() && GetMaxSupportedFeatureLevel(Platform) > ERHIFeatureLevel::SM5;
-	return !IsMobilePlatform(Platform) && bCanRHICompileHlsl2021GeometryShaders;
+	// Substrate requires HLSL2021 which must be cross-compiled for D3D11 and Vulkan to be consumed by FXC compiler.
+	// This cross-compilation toolchain does not support geometry shaders.
+	const bool bCanRHICompileGeometryShaders =!Substrate::IsSubstrateEnabled() || GetMaxSupportedFeatureLevel(Platform) > ERHIFeatureLevel::SM5;
+	return !IsMobilePlatform(Platform) && bCanRHICompileGeometryShaders;
 }
 bool ShouldRenderVolumetricFog(const FScene* Scene, const FSceneViewFamily& ViewFamily)
 {
