Fix dual source blending with Substrate on the mobile rendered.

diff --git a/Engine/Shaders/Private/MobileBasePassPixelShader.usf b/Engine/Shaders/Private/MobileBasePassPixelShader.usf
--- a/Engine/Shaders/Private/MobileBasePassPixelShader.usf
+++ b/Engine/Shaders/Private/MobileBasePassPixelShader.usf
@@ -782,11 +782,14 @@ void Main(
 		half3 FoggedColor = lerp(float3(1, 1, 1), DualBlendSurfaceTransmittancePreCoverage, Fogging.aaa* DualBlendSurfaceCoverage);
 		OutColor = half4(FoggedColor, 1.0f);
 	#elif SUBSTRATE_BLENDING_ALPHAHOLDOUT
-		float AdjustedAlpha = saturate(1 - dot(AdjustedDualBlendMul, float3(1.0f, 1.0f, 1.0f) / 3.0f));
+		half AdjustedAlpha = saturate(1 - dot(AdjustedDualBlendMul, half3(1.0f, 1.0f, 1.0f) / 3.0f));
 		OutColor = half4(0.0f.xxx, AdjustedAlpha);
+	#elif MOBILE_TRANSLUCENT_COLOR_TRANSMITTANCE && MOBILE_TRANSLUCENT_COLOR_TRANSMITTANCE_DUAL_SRC_BLENDING
+		OutColor = half4(AdjustedDualBlendAdd, 0.0);
+		OutColor1 = half4(AdjustedDualBlendMul, 1.0);
 	#else
 		// Pre multipled alpha blending
-		float AdjustedAlpha = saturate(1 - dot(AdjustedDualBlendMul, float3(1.0f, 1.0f, 1.0f) / 3.0f));
+		half AdjustedAlpha = saturate(1 - dot(AdjustedDualBlendMul, half3(1.0f, 1.0f, 1.0f) / 3.0f));
 		#if SUBSTRATE_USES_CONVERSION_FROM_LEGACY && MATERIALBLENDING_ADDITIVE
 			AdjustedAlpha = 0.0f;
 		#endif
diff --git a/Engine/Source/Runtime/Renderer/Private/MobileBasePass.cpp b/Engine/Source/Runtime/Renderer/Private/MobileBasePass.cpp
--- a/Engine/Source/Runtime/Renderer/Private/MobileBasePass.cpp
+++ b/Engine/Source/Runtime/Renderer/Private/MobileBasePass.cpp
@@ -161,7 +161,7 @@ bool GetUniformMobileBasePassShaders(
 	const FMaterial& Material, 
 	const FVertexFactoryType* VertexFactoryType, 
 	bool bEnableSkyLight,
-	EMobileTranslucentColorTransmittanceMode ThinTranslucentFallback,
+	EMobileTranslucentColorTransmittanceMode ColoredTransmittanceFallback,
 	TShaderRef<TMobileBasePassVSPolicyParamType<FUniformLightMapPolicy>>& VertexShader,
 	TShaderRef<TMobileBasePassPSPolicyParamType<FUniformLightMapPolicy>>& PixelShader
 	)
@@ -180,7 +180,7 @@ bool GetUniformMobileBasePassShaders(
 		ShaderTypes.AddShaderType<TMobileBasePassVS<TUniformLightMapPolicy<Policy>, LDR_GAMMA_32>>();			
 	}
 
-	switch (ThinTranslucentFallback)
+	switch (ColoredTransmittanceFallback)
 	{
 	default:
 	case EMobileTranslucentColorTransmittanceMode::DEFAULT:
@@ -206,7 +206,7 @@ bool GetMobileBasePassShaders(
 	const FMaterial& Material, 
 	const FVertexFactoryType* VertexFactoryType, 
 	bool bEnableSkyLight,
-	EMobileTranslucentColorTransmittanceMode ThinTranslucentFallback,
+	EMobileTranslucentColorTransmittanceMode ColoredTransmittanceFallback,
 	TShaderRef<TMobileBasePassVSPolicyParamType<FUniformLightMapPolicy>>& VertexShader,
 	TShaderRef<TMobileBasePassPSPolicyParamType<FUniformLightMapPolicy>>& PixelShader
 	)
@@ -214,25 +214,25 @@ bool GetMobileBasePassShaders(
 	switch (LightMapPolicyType)
 	{
 	case LMP_NO_LIGHTMAP:
-		return GetUniformMobileBasePassShaders<LMP_NO_LIGHTMAP, LocalLightSetting>(Material, VertexFactoryType, bEnableSkyLight, ThinTranslucentFallback, VertexShader, PixelShader);
+		return GetUniformMobileBasePassShaders<LMP_NO_LIGHTMAP, LocalLightSetting>(Material, VertexFactoryType, bEnableSkyLight, ColoredTransmittanceFallback, VertexShader, PixelShader);
 	case LMP_LQ_LIGHTMAP:
-		return GetUniformMobileBasePassShaders<LMP_LQ_LIGHTMAP, LocalLightSetting>(Material, VertexFactoryType, bEnableSkyLight, ThinTranslucentFallback, VertexShader, PixelShader);
+		return GetUniformMobileBasePassShaders<LMP_LQ_LIGHTMAP, LocalLightSetting>(Material, VertexFactoryType, bEnableSkyLight, ColoredTransmittanceFallback, VertexShader, PixelShader);
 	case LMP_MOBILE_DISTANCE_FIELD_SHADOWS_AND_LQ_LIGHTMAP:
-		return GetUniformMobileBasePassShaders<LMP_MOBILE_DISTANCE_FIELD_SHADOWS_AND_LQ_LIGHTMAP, LocalLightSetting>(Material, VertexFactoryType, bEnableSkyLight, ThinTranslucentFallback, VertexShader, PixelShader);
+		return GetUniformMobileBasePassShaders<LMP_MOBILE_DISTANCE_FIELD_SHADOWS_AND_LQ_LIGHTMAP, LocalLightSetting>(Material, VertexFactoryType, bEnableSkyLight, ColoredTransmittanceFallback, VertexShader, PixelShader);
 	case LMP_MOBILE_DISTANCE_FIELD_SHADOWS_LIGHTMAP_AND_CSM:
-		return GetUniformMobileBasePassShaders<LMP_MOBILE_DISTANCE_FIELD_SHADOWS_LIGHTMAP_AND_CSM, LocalLightSetting>(Material, VertexFactoryType, bEnableSkyLight, ThinTranslucentFallback, VertexShader, PixelShader);
+		return GetUniformMobileBasePassShaders<LMP_MOBILE_DISTANCE_FIELD_SHADOWS_LIGHTMAP_AND_CSM, LocalLightSetting>(Material, VertexFactoryType, bEnableSkyLight, ColoredTransmittanceFallback, VertexShader, PixelShader);
 	case LMP_MOBILE_DIRECTIONAL_LIGHT_CSM_AND_LIGHTMAP:
-		return GetUniformMobileBasePassShaders<LMP_MOBILE_DIRECTIONAL_LIGHT_CSM_AND_LIGHTMAP, LocalLightSetting>(Material, VertexFactoryType, bEnableSkyLight, ThinTranslucentFallback, VertexShader, PixelShader);
+		return GetUniformMobileBasePassShaders<LMP_MOBILE_DIRECTIONAL_LIGHT_CSM_AND_LIGHTMAP, LocalLightSetting>(Material, VertexFactoryType, bEnableSkyLight, ColoredTransmittanceFallback, VertexShader, PixelShader);
 	case LMP_MOBILE_DIRECTIONAL_LIGHT_AND_SH_INDIRECT:
-		return GetUniformMobileBasePassShaders<LMP_MOBILE_DIRECTIONAL_LIGHT_AND_SH_INDIRECT, LocalLightSetting>(Material, VertexFactoryType, bEnableSkyLight, ThinTranslucentFallback, VertexShader, PixelShader);
+		return GetUniformMobileBasePassShaders<LMP_MOBILE_DIRECTIONAL_LIGHT_AND_SH_INDIRECT, LocalLightSetting>(Material, VertexFactoryType, bEnableSkyLight, ColoredTransmittanceFallback, VertexShader, PixelShader);
 	case LMP_MOBILE_DIRECTIONAL_LIGHT_CSM_AND_SH_INDIRECT:
-		return GetUniformMobileBasePassShaders<LMP_MOBILE_DIRECTIONAL_LIGHT_CSM_AND_SH_INDIRECT, LocalLightSetting>(Material, VertexFactoryType, bEnableSkyLight, ThinTranslucentFallback, VertexShader, PixelShader);
+		return GetUniformMobileBasePassShaders<LMP_MOBILE_DIRECTIONAL_LIGHT_CSM_AND_SH_INDIRECT, LocalLightSetting>(Material, VertexFactoryType, bEnableSkyLight, ColoredTransmittanceFallback, VertexShader, PixelShader);
 	case LMP_MOBILE_MOVABLE_DIRECTIONAL_LIGHT_WITH_LIGHTMAP:
-		return GetUniformMobileBasePassShaders<LMP_MOBILE_MOVABLE_DIRECTIONAL_LIGHT_WITH_LIGHTMAP, LocalLightSetting>(Material, VertexFactoryType, bEnableSkyLight, ThinTranslucentFallback, VertexShader, PixelShader);
+		return GetUniformMobileBasePassShaders<LMP_MOBILE_MOVABLE_DIRECTIONAL_LIGHT_WITH_LIGHTMAP, LocalLightSetting>(Material, VertexFactoryType, bEnableSkyLight, ColoredTransmittanceFallback, VertexShader, PixelShader);
 	case LMP_MOBILE_MOVABLE_DIRECTIONAL_LIGHT_CSM_WITH_LIGHTMAP:
-		return GetUniformMobileBasePassShaders<LMP_MOBILE_MOVABLE_DIRECTIONAL_LIGHT_CSM_WITH_LIGHTMAP, LocalLightSetting>(Material, VertexFactoryType, bEnableSkyLight, ThinTranslucentFallback, VertexShader, PixelShader);
+		return GetUniformMobileBasePassShaders<LMP_MOBILE_MOVABLE_DIRECTIONAL_LIGHT_CSM_WITH_LIGHTMAP, LocalLightSetting>(Material, VertexFactoryType, bEnableSkyLight, ColoredTransmittanceFallback, VertexShader, PixelShader);
 	case LMP_MOBILE_DIRECTIONAL_LIGHT_CSM:
-		return GetUniformMobileBasePassShaders<LMP_MOBILE_DIRECTIONAL_LIGHT_CSM, LocalLightSetting>(Material, VertexFactoryType, bEnableSkyLight, ThinTranslucentFallback, VertexShader, PixelShader);
+		return GetUniformMobileBasePassShaders<LMP_MOBILE_DIRECTIONAL_LIGHT_CSM, LocalLightSetting>(Material, VertexFactoryType, bEnableSkyLight, ColoredTransmittanceFallback, VertexShader, PixelShader);
 	default:										
 		check(false);
 		return true;
@@ -254,11 +254,11 @@ bool MobileBasePass::GetShaders(
 		bEnableSkyLight = !bEnableSkyLight;
 	}
 
-	EMobileTranslucentColorTransmittanceMode ThinTranslucencyFallback = EMobileTranslucentColorTransmittanceMode::DEFAULT;
-	if (MaterialResource.GetShadingModels().HasShadingModel(MSM_ThinTranslucent))
+	EMobileTranslucentColorTransmittanceMode ColoredTransmittanceFallback = EMobileTranslucentColorTransmittanceMode::DEFAULT;
+	if (MaterialRequiresColorTransmittanceBlending(MaterialResource))
 	{
 		const EShaderPlatform ShaderPlatform = GetFeatureLevelShaderPlatform(MaterialResource.GetFeatureLevel());
-		ThinTranslucencyFallback = MobileActiveTranslucentColorTransmittanceMode(ShaderPlatform, false);
+		ColoredTransmittanceFallback = MobileActiveTranslucentColorTransmittanceMode(ShaderPlatform, false);
 	}
 
 	switch (LocalLightSetting)
@@ -270,7 +270,7 @@ bool MobileBasePass::GetShaders(
 				MaterialResource,
 				VertexFactoryType,
 				bEnableSkyLight,
-				ThinTranslucencyFallback,
+				ColoredTransmittanceFallback,
 				VertexShader,
 				PixelShader
 				);
@@ -283,7 +283,7 @@ bool MobileBasePass::GetShaders(
 				MaterialResource,
 				VertexFactoryType,
 				bEnableSkyLight,
-				ThinTranslucencyFallback,
+				ColoredTransmittanceFallback,
 				VertexShader,
 				PixelShader
 				);
@@ -296,7 +296,7 @@ bool MobileBasePass::GetShaders(
 				MaterialResource,
 				VertexFactoryType,
 				bEnableSkyLight,
-				ThinTranslucencyFallback,
+				ColoredTransmittanceFallback,
 				VertexShader,
 				PixelShader
 				);
diff --git a/Engine/Source/Runtime/Renderer/Private/MobileBasePassRendering.cpp b/Engine/Source/Runtime/Renderer/Private/MobileBasePassRendering.cpp
--- a/Engine/Source/Runtime/Renderer/Private/MobileBasePassRendering.cpp
+++ b/Engine/Source/Runtime/Renderer/Private/MobileBasePassRendering.cpp
@@ -169,6 +169,11 @@ IMPLEMENT_MOBILE_SHADING_BASEPASS_LIGHTMAPPED_SHADER_TYPE(TUniformLightMapPolicy
 IMPLEMENT_MOBILE_SHADING_BASEPASS_LIGHTMAPPED_SHADER_TYPE(TUniformLightMapPolicy<LMP_MOBILE_MOVABLE_DIRECTIONAL_LIGHT_CSM_WITH_LIGHTMAP>, FMobileMovableDirectionalLightCSMWithLightmapPolicy);
 IMPLEMENT_MOBILE_SHADING_BASEPASS_LIGHTMAPPED_SHADER_TYPE(TUniformLightMapPolicy<LMP_MOBILE_DIRECTIONAL_LIGHT_CSM>, FMobileDirectionalLightAndCSMPolicy);
 
+bool MaterialRequiresColorTransmittanceBlending(const FMaterial& MaterialResource)
+{
+	return MaterialResource.GetShadingModels().HasShadingModel(MSM_ThinTranslucent) || MaterialResource.GetBlendMode() == BLEND_TranslucentColoredTransmittance;
+}
+
 bool MaterialRequiresColorTransmittanceBlending(const FMaterialShaderParameters& MaterialParameters)
 {
 	return MaterialParameters.ShadingModels.HasShadingModel(MSM_ThinTranslucent) || MaterialParameters.BlendMode == BLEND_TranslucentColoredTransmittance;
diff --git a/Engine/Source/Runtime/Renderer/Private/MobileBasePassRendering.h b/Engine/Source/Runtime/Renderer/Private/MobileBasePassRendering.h
--- a/Engine/Source/Runtime/Renderer/Private/MobileBasePassRendering.h
+++ b/Engine/Source/Runtime/Renderer/Private/MobileBasePassRendering.h
@@ -123,6 +123,7 @@ enum class EMobileTranslucentColorTransmittanceMode
 
 EMobileTranslucentColorTransmittanceMode MobileDefaultTranslucentColorTransmittanceMode(EShaderPlatform Platform);
 EMobileTranslucentColorTransmittanceMode MobileActiveTranslucentColorTransmittanceMode(EShaderPlatform Platform, bool bExplicitDefaultMode);
+bool MaterialRequiresColorTransmittanceBlending(const FMaterial& MaterialResource);
 bool MaterialRequiresColorTransmittanceBlending(const FMaterialShaderParameters& MaterialParameters);
 bool ShouldCacheShaderForColorTransmittanceFallback(const FMaterialShaderPermutationParameters& Parameters, EMobileTranslucentColorTransmittanceMode TranslucentColorTransmittanceFallback);
 
@@ -445,7 +446,7 @@ public:
 		OutEnvironment.SetDefine(TEXT("USE_SHADOWMASKTEXTURE"), bMobileUsesShadowMaskTexture && !bTranslucentMaterial ? 1u : 0u);
 		OutEnvironment.SetDefine(TEXT("ENABLE_DBUFFER_TEXTURES"), Parameters.MaterialParameters.MaterialDomain == MD_Surface ? 1u : 0u);
 		EMobileTranslucentColorTransmittanceMode TranslucentColorTransmittanceMode = EMobileTranslucentColorTransmittanceMode::DEFAULT;
-		if (Parameters.MaterialParameters.ShadingModels.HasShadingModel(MSM_ThinTranslucent))
+		if (MaterialRequiresColorTransmittanceBlending(Parameters.MaterialParameters))
 		{
 			TranslucentColorTransmittanceMode = (TranslucentColorTransmittanceFallback == EMobileTranslucentColorTransmittanceMode::DEFAULT) ? MobileDefaultTranslucentColorTransmittanceMode(Parameters.Platform) : TranslucentColorTransmittanceFallback;
 		}
