Add configuration of which spatial, spatial-temporal denoiser to use for path tracing. All denoiser plugins can now be loaded and dynamically switched.
* Use r.PathTracing.Denoiser.Name to choose which spatial denoiser to use. Default to OIDN. The alternative is NNE_OIDN, NFOR, and OptiX
* Use r.PathTracing.TemporalDenoiser.Name to choose which temporal denoiser to use for offline rendering. Default to NFOR. The alternatives is OptiX
#tests Runs on hillside projects with all the testing.
#rnx

diff --git a/Engine/Plugins/Experimental/NFORDenoise/Source/NFORDenoise/Private/NFORDenoise.cpp b/Engine/Plugins/Experimental/NFORDenoise/Source/NFORDenoise/Private/NFORDenoise.cpp
--- a/Engine/Plugins/Experimental/NFORDenoise/Source/NFORDenoise/Private/NFORDenoise.cpp
+++ b/Engine/Plugins/Experimental/NFORDenoise/Source/NFORDenoise/Private/NFORDenoise.cpp
@@ -262,7 +262,7 @@ void FNFORDenoiseModule::StartupModule()
 		FString ModuleDir = Plugin->GetBaseDir() + TEXT("/Source/NFORDenoise");
 		AddShaderSourceDirectoryMapping(TEXT("/NFORDenoise"), FPaths::Combine(ModuleDir, TEXT("Shaders")));
 
-		GPathTracingSpatialTemporalDenoiserPlugin = MakeUnique<FNFORDenosier>();
+		RegisterSpatialTemporalDenoiser(MakeUnique<FNFORDenosier>(),TEXT("NFOR"));
 	}
 	else
 	{
@@ -273,9 +273,7 @@ void FNFORDenoiseModule::StartupModule()
 void FNFORDenoiseModule::ShutdownModule()
 {
 	UE_LOG(LogNFORDenoise, Log, TEXT("NFORDenoise function shutting down"));
-
-	GPathTracingSpatialTemporalDenoiserPlugin.Reset();
-
+	UnregisterDenoiser(TEXT("NFOR"));
 }
 
 #undef LOCTEXT_NAMESPACE
diff --git a/Engine/Plugins/Experimental/NNEDenoiser/Source/NNEDenoiser/Private/NNEDenoiserViewExtension.cpp b/Engine/Plugins/Experimental/NNEDenoiser/Source/NNEDenoiser/Private/NNEDenoiserViewExtension.cpp
--- a/Engine/Plugins/Experimental/NNEDenoiser/Source/NNEDenoiser/Private/NNEDenoiserViewExtension.cpp
+++ b/Engine/Plugins/Experimental/NNEDenoiser/Source/NNEDenoiser/Private/NNEDenoiserViewExtension.cpp
@@ -304,8 +304,8 @@ FViewExtension::~FViewExtension()
 	}
 #endif
 
-	GPathTracingDenoiserPlugin.Reset();
-	GPathTracingSpatialTemporalDenoiserPlugin.Reset();
+	UnregisterDenoiser(TEXT("NNE_OIDN"));
+	UnregisterDenoiser(TEXT("NNE_KPCN"));
 }
 
 void FViewExtension::BeginRenderViewFamily(FSceneViewFamily& InViewFamily)
@@ -317,21 +317,23 @@ void FViewExtension::PreRenderViewFamily_RenderThread(FRDGBuilder& GraphBuilder,
 {
 	if (!bDenoiserEnabled)
 	{
-		GPathTracingDenoiserPlugin.Reset();
-		GPathTracingSpatialTemporalDenoiserPlugin.Reset();
+		UnregisterDenoiser(TEXT("NNE_OIDN"));
+		UnregisterDenoiser(TEXT("NNE_KPCN"));
 	}
 	else
 	{
 		if (DenoiserToSwap.IsValid())
 		{
-			GPathTracingSpatialTemporalDenoiserPlugin.Reset();
-			GPathTracingDenoiserPlugin = MoveTemp(DenoiserToSwap);
+			UnregisterDenoiser(TEXT("NNE_OIDN"));
+			UnregisterDenoiser(TEXT("NNE_KPCN"));
+			RegisterSpatialDenoiser(MoveTemp(DenoiserToSwap),TEXT("NNE_OIDN"));
 		}
 
 		if (SpatialTemporalDenoiserToSwap.IsValid())
 		{
-			GPathTracingDenoiserPlugin.Reset();
-			GPathTracingSpatialTemporalDenoiserPlugin = MoveTemp(SpatialTemporalDenoiserToSwap);
+			UnregisterDenoiser(TEXT("NNE_OIDN"));
+			UnregisterDenoiser(TEXT("NNE_KPCN"));
+			RegisterSpatialTemporalDenoiser(MoveTemp(SpatialTemporalDenoiserToSwap), TEXT("NNE_KPCN"));
 		}
 	}
 }
diff --git a/Engine/Plugins/Experimental/OpenImageDenoise/Source/OpenImageDenoise/Private/Plugin.cpp b/Engine/Plugins/Experimental/OpenImageDenoise/Source/OpenImageDenoise/Private/Plugin.cpp
--- a/Engine/Plugins/Experimental/OpenImageDenoise/Source/OpenImageDenoise/Private/Plugin.cpp
+++ b/Engine/Plugins/Experimental/OpenImageDenoise/Source/OpenImageDenoise/Private/Plugin.cpp
@@ -300,7 +300,7 @@ void FOpenImageDenoiseModule::StartupModule()
 	UE_LOG(LogOpenImageDenoise, Log, TEXT("OIDN starting up (built with OIDN %hs)"), OIDN_VERSION_STRING);
 #endif
 
-	GPathTracingDenoiserPlugin = MakeUnique<FOIDNDenoiser>();
+	RegisterSpatialDenoiser(MakeUnique<FOIDNDenoiser>(),TEXT("OIDN"));
 }
 
 void FOpenImageDenoiseModule::ShutdownModule()
@@ -310,5 +310,5 @@ void FOpenImageDenoiseModule::ShutdownModule()
 #endif
 
 	// Release scratch memory and destroy the OIDN device and filters
-	GPathTracingDenoiserPlugin.Reset();
+	UnregisterDenoiser(TEXT("OIDN"));
 }
diff --git a/Engine/Plugins/Experimental/OptiXDenoise/Source/OptiXDenoise/Private/Plugin.cpp b/Engine/Plugins/Experimental/OptiXDenoise/Source/OptiXDenoise/Private/Plugin.cpp
--- a/Engine/Plugins/Experimental/OptiXDenoise/Source/OptiXDenoise/Private/Plugin.cpp
+++ b/Engine/Plugins/Experimental/OptiXDenoise/Source/OptiXDenoise/Private/Plugin.cpp
@@ -326,7 +326,7 @@ void FOptiXDenoiseModule::StartupModule()
 	OptiXDenoiseBaseDLLHandle = FPlatformProcess::GetDllHandle(*OptiXDenoiseBaseDllPath);
 	FOptiXCudaFunctionList::Get().RegisterFunctionInstance<FOptiXDenoiserFunctionInstance>();
 
-	GPathTracingSpatialTemporalDenoiserPlugin = MakeUnique<FOptiXDenosier>();
+	RegisterSpatialTemporalDenoiser(MakeUnique<FOptiXDenosier>(), TEXT("OptiX"));
 }
 
 void FOptiXDenoiseModule::ShutdownModule()
@@ -335,7 +335,7 @@ void FOptiXDenoiseModule::ShutdownModule()
 	UE_LOG(LogOptiXDenoise, Log, TEXT("OptiXDenoise shutting down"));
 #endif
 
-	GPathTracingSpatialTemporalDenoiserPlugin.Reset();
+	UnregisterDenoiser(TEXT("OptiX"));
 
 	// Assure resources related to CUDA is released before the releasing of CUDA module.
 	Denoiser.Reset();
diff --git a/Engine/Source/Runtime/Renderer/Private/PathTracingSpatialTemporalDenoising.cpp b/Engine/Source/Runtime/Renderer/Private/PathTracingSpatialTemporalDenoising.cpp
--- a/Engine/Source/Runtime/Renderer/Private/PathTracingSpatialTemporalDenoising.cpp
+++ b/Engine/Source/Runtime/Renderer/Private/PathTracingSpatialTemporalDenoising.cpp
@@ -4,8 +4,106 @@
 #include "PathTracing.h"
 #include "RHI.h"
 
-TUniquePtr<UE::Renderer::Private::IPathTracingDenoiser> GPathTracingDenoiserPlugin;
-TUniquePtr<UE::Renderer::Private::IPathTracingSpatialTemporalDenoiser> GPathTracingSpatialTemporalDenoiserPlugin;
+class FDenoiserManager
+{
+public:
+	static FDenoiserManager& Get()
+	{
+		static FDenoiserManager DenoiserManager;
+		return DenoiserManager;
+	}
+
+	void RegisterSpatialDenoiser(TUniquePtr<UE::Renderer::Private::IPathTracingDenoiser> InDenosier, FString Name)
+	{
+		check(!SpatialDenoiser.Contains(Name));
+
+		bNeedTextureCreateExtraFlags |= InDenosier->NeedTextureCreateExtraFlags();
+		SpatialDenoiser.Add(Name, MoveTemp(InDenosier));
+	}
+	void RegisterSpatialTemporalDenoiser(TUniquePtr<UE::Renderer::Private::IPathTracingSpatialTemporalDenoiser> InDenosier, FString Name)
+	{
+		check(!PathTracingSpatialTemporalDenoisers.Contains(Name));
+		bNeedTextureCreateExtraFlags |= InDenosier->NeedTextureCreateExtraFlags();
+		PathTracingSpatialTemporalDenoisers.Add(Name, MoveTemp(InDenosier));
+	}
+
+	bool HasSpatialDenoiser()const { return SpatialDenoiser.Num() > 0; }
+	bool HasSpatialTemporalDenoiser()const { return PathTracingSpatialTemporalDenoisers.Num() > 0; }
+	bool HasDenoiser()const { return HasSpatialDenoiser() || HasSpatialTemporalDenoiser();}
+
+	UE::Renderer::Private::IPathTracingDenoiser* GetSpatialDenosier(FString Name, bool bMatch)
+	{
+		if (SpatialDenoiser.Contains(Name))
+		{
+			return SpatialDenoiser[Name].Get();
+		}
+		else if (SpatialDenoiser.Num() > 0 && !bMatch)
+		{
+			for (const auto& Pair : SpatialDenoiser)
+			{
+				return Pair.Value.Get();
+			}
+		}
+
+		return nullptr;
+	}
+
+	UE::Renderer::Private::IPathTracingSpatialTemporalDenoiser* GetSpatialTemporalDenosier(FString Name, bool bMatch)
+	{
+		if (PathTracingSpatialTemporalDenoisers.Contains(Name))
+		{
+			return PathTracingSpatialTemporalDenoisers[Name].Get();
+		}
+		else if (PathTracingSpatialTemporalDenoisers.Num() > 0 && !bMatch)
+		{
+			for (const auto& Pair : PathTracingSpatialTemporalDenoisers)
+			{
+				return Pair.Value.Get();
+			}
+		}
+
+		return nullptr;
+	}
+
+	void UnregisterDenoiser(FString Name)
+	{
+		SpatialDenoiser.Remove(Name);
+		PathTracingSpatialTemporalDenoisers.Remove(Name);
+	}
+
+	/** If any plugin needs extra creation flag*/
+	bool NeedTextureCreateExtraFlags()
+	{
+		return bNeedTextureCreateExtraFlags;
+	}
+
+private:
+	FDenoiserManager() = default;
+	FDenoiserManager(const FDenoiserManager&) = delete;
+	FDenoiserManager& operator=(const FDenoiserManager&) = delete;
+
+	TMap<FString, TUniquePtr<UE::Renderer::Private::IPathTracingDenoiser> > SpatialDenoiser;
+
+	/**Spatial-temporal denoiser can work as both spatial and temporal denoiser*/
+	TMap<FString, TUniquePtr<UE::Renderer::Private::IPathTracingSpatialTemporalDenoiser> > PathTracingSpatialTemporalDenoisers;
+
+	/** Since all plugins can be dynamically switched, need to increase the compatibility*/
+	bool bNeedTextureCreateExtraFlags = false;
+};
+
+void RegisterSpatialDenoiser(TUniquePtr<UE::Renderer::Private::IPathTracingDenoiser> PathTracingDenoiser, FString Name)
+{
+	FDenoiserManager::Get().RegisterSpatialDenoiser(MoveTemp(PathTracingDenoiser), Name);
+}
+void RegisterSpatialTemporalDenoiser(TUniquePtr<UE::Renderer::Private::IPathTracingSpatialTemporalDenoiser> PathTracingDenoiser, FString Name)
+{
+	FDenoiserManager::Get().RegisterSpatialTemporalDenoiser(MoveTemp(PathTracingDenoiser), Name);
+}
+
+void UnregisterDenoiser(FString Name)
+{
+	FDenoiserManager::Get().UnregisterDenoiser(Name);
+}
 
 #if RHI_RAYTRACING
 
@@ -71,6 +169,20 @@ namespace {
 		ECVF_RenderThreadSafe
 	);
 
+	TAutoConsoleVariable<FString> CVarPathTracingDenoiserName(
+		TEXT("r.PathTracing.Denoiser.Name"),
+		"OIDN",
+		TEXT("Set the spatial denoiser name. It is the corresponding name registered by a denosier plugin\n")
+		TEXT("Any registered denoiser should be able to denoise a single frame spatially.\n")
+	);
+
+	TAutoConsoleVariable<FString> CVarPathTracingTemporalDenoiserName(
+		TEXT("r.PathTracing.TemporalDenoiser.Name"),
+		"NFOR",
+		TEXT("Set the temporal denoiser name. It is the corresponding name registered by a denosier plugin\n")
+		TEXT("The temporal denoiser usually has better temporal stability when rendering offline.\n")
+	);
+
 	TAutoConsoleVariable<int32> CVarPathTracingSpatialDenoiserType(
 		TEXT("r.PathTracing.SpatialDenoiser.Type"),
 		0,
@@ -256,7 +368,7 @@ int GetPathTracingDenoiserMode(const FViewInfo& View)
 
 bool IsPathTracingDenoiserEnabled(const FViewInfo& View)
 {
-	return GetPathTracingDenoiserMode(View) != 0 && (GPathTracingDenoiserPlugin || GPathTracingSpatialTemporalDenoiserPlugin);
+	return GetPathTracingDenoiserMode(View) != 0 && FDenoiserManager::Get().HasDenoiser();
 }
 
 static bool ShouldDenoiseWithNormalInCameraSpace()
@@ -315,12 +427,12 @@ static bool ShouldApplySpatialDenoiser()
 TArray<ESpatialDenoiserType> GetAvailableSpatialDenoiserTypes()
 {
 	TArray<ESpatialDenoiserType> Types;
-	if (GPathTracingSpatialTemporalDenoiserPlugin)
+	if (FDenoiserManager::Get().HasSpatialTemporalDenoiser())
 	{
 		Types.Add(ESpatialDenoiserType::SPATIAL_TEMPORAL_DENOISER_PLUGIN);
 	}
 
-	if (GPathTracingDenoiserPlugin)
+	if (FDenoiserManager::Get().HasSpatialDenoiser())
 	{
 		Types.Add(ESpatialDenoiserType::SPATIAL_DENOISER_PLUGIN);
 	}
@@ -342,11 +454,11 @@ ESpatialDenoiserType GetSpatialDenosierType()
 
 	ESpatialDenoiserType DenoiserType = static_cast<ESpatialDenoiserType>(Type);
 
-	if (DenoiserType == ESpatialDenoiserType::SPATIAL_DENOISER_PLUGIN && GPathTracingDenoiserPlugin)
+	if (DenoiserType == ESpatialDenoiserType::SPATIAL_DENOISER_PLUGIN && FDenoiserManager::Get().HasSpatialDenoiser())
 	{
 		return DenoiserType;
 	}
-	else if (DenoiserType == ESpatialDenoiserType::SPATIAL_TEMPORAL_DENOISER_PLUGIN && GPathTracingSpatialTemporalDenoiserPlugin)
+	else if (DenoiserType == ESpatialDenoiserType::SPATIAL_TEMPORAL_DENOISER_PLUGIN && FDenoiserManager::Get().HasSpatialTemporalDenoiser())
 	{
 		return DenoiserType;
 	}
@@ -398,15 +510,25 @@ ETemporalDenoiserMotionVectorType GetTemporalDenoiserMotionVectorType()
 	return static_cast<ETemporalDenoiserMotionVectorType>(Type);
 }
 
+UE::Renderer::Private::IPathTracingDenoiser* GetActiveSpatialDenoiser(bool bMatch = false)
+{
+	FString DenoiserName = CVarPathTracingDenoiserName.GetValueOnRenderThread();
+	return FDenoiserManager::Get().GetSpatialDenosier(DenoiserName, bMatch);
+}
+
+UE::Renderer::Private::IPathTracingSpatialTemporalDenoiser* GetActiveSpatialTemporalDenoiser(bool bMatch = false)
+{
+	FString DenoiserName = CVarPathTracingTemporalDenoiserName.GetValueOnRenderThread();
+	return FDenoiserManager::Get().GetSpatialTemporalDenosier(DenoiserName, bMatch);
+}
+
 ETextureCreateFlags GetExtraTextureCreateFlagsForDenoiser()
 {
 	ESpatialDenoiserType SpatialDenoiserType = GetSpatialDenosierType();
 	ETemporalDenoiserType TemporalDenoiserType = GetTemporalDenoiserType();
 	ETextureCreateFlags TextureCreateFlags = ETextureCreateFlags::None;
 
-	bool bThirdPartyDenoiserNeedsTextureCreateExtraFlags = false;
-	bThirdPartyDenoiserNeedsTextureCreateExtraFlags |= GPathTracingDenoiserPlugin && GPathTracingDenoiserPlugin->NeedTextureCreateExtraFlags();
-	bThirdPartyDenoiserNeedsTextureCreateExtraFlags |= GPathTracingSpatialTemporalDenoiserPlugin && GPathTracingSpatialTemporalDenoiserPlugin->NeedTextureCreateExtraFlags();
+	bool bThirdPartyDenoiserNeedsTextureCreateExtraFlags = FDenoiserManager::Get().NeedTextureCreateExtraFlags();
 
 	if (SpatialDenoiserType == ESpatialDenoiserType::SPATIAL_TEMPORAL_DENOISER_PLUGIN ||
 		TemporalDenoiserType == ETemporalDenoiserType::SPATIAL_TEMPORAL_DENOISER_PLUGIN ||
@@ -902,10 +1024,9 @@ static void PathTracingDenoiserPlugin(FRDGBuilder& GraphBuilder,
 	FRDGTextureRef InputTexture, 
 	FRDGTextureRef AlbedoTexture,
 	FRDGTextureRef NormalTexture,
+	FRDGTextureRef VarianceTexture,
 	FRDGTextureRef OutputTexture)
 {
-	check(GPathTracingDenoiserPlugin);
-
 	FRDGTextureRef ProcessedNormalTexture = NormalTexture;
 
 	if (ShouldDenoiseWithNormalInCameraSpace())
@@ -924,7 +1045,51 @@ static void PathTracingDenoiserPlugin(FRDGBuilder& GraphBuilder,
 		ConvertNormalSpace(GraphBuilder, View, ProcessedNormalTexture);
 	}
 
-	GPathTracingDenoiserPlugin->AddPasses(GraphBuilder, View, {InputTexture, AlbedoTexture, NormalTexture, OutputTexture});
+	/**
+	* First try to use the matched denoiser to denoise.
+	* If the requested denoiser is not available, fallback to the spatial denoiser.
+	*/
+	bool bMatch = true;
+	if (UE::Renderer::Private::IPathTracingDenoiser* ActiveSpatialDenoiser = GetActiveSpatialDenoiser(bMatch); ActiveSpatialDenoiser)
+	{
+		ActiveSpatialDenoiser->AddPasses(GraphBuilder, View, { InputTexture, AlbedoTexture, NormalTexture, OutputTexture });
+	}
+	else
+	{
+		using UE::Renderer::Private::IPathTracingSpatialTemporalDenoiser;
+
+		FString DenoiserName = CVarPathTracingDenoiserName.GetValueOnRenderThread();
+		IPathTracingSpatialTemporalDenoiser* ActiveSpatialTemporalDenoiser = FDenoiserManager::Get().GetSpatialTemporalDenosier(DenoiserName,bMatch);
+		
+		if (ActiveSpatialTemporalDenoiser)
+		{
+			/**
+			* Force the temporal denoiser to denoise a single frame.
+			*/
+			IPathTracingSpatialTemporalDenoiser::FInputs Inputs;
+			Inputs.ColorTex = InputTexture;
+			Inputs.AlbedoTex = AlbedoTexture;
+			Inputs.NormalTex = NormalTexture;
+			Inputs.VarianceTex = VarianceTexture;
+			
+			// Set flow texture to black
+			Inputs.FlowTex = GraphBuilder.CreateTexture(VarianceTexture->Desc, TEXT("PathTracing.OpticalFlow"));
+			AddClearRenderTargetPass(GraphBuilder, Inputs.FlowTex, FLinearColor::Black);
+			
+			Inputs.PreviousOutputTex = InputTexture;
+			Inputs.OutputTex = OutputTexture;
+			Inputs.DenoisingFrameId = 0;
+			Inputs.bForceSpatialDenoiserOnly = true; // Force to use spatial denoiser.
+			
+			ActiveSpatialTemporalDenoiser->AddPasses(GraphBuilder, View, Inputs);
+		}
+		else
+		{
+			//fallback to the default spatial denoiser
+			bMatch = false;
+			GetActiveSpatialDenoiser(bMatch)->AddPasses(GraphBuilder, View, { InputTexture, AlbedoTexture, NormalTexture, OutputTexture });
+		}
+	}
 }
 
 static void PathTracingSpatialTemporalDenoiserPlugin(FRDGBuilder& GraphBuilder,
@@ -941,7 +1106,7 @@ static void PathTracingSpatialTemporalDenoiserPlugin(FRDGBuilder& GraphBuilder,
 	bool bForceSpatialDenoiserOnly,
 	FPathTracingSpatialTemporalDenoisingContext& Context)
 {
-	check(GPathTracingSpatialTemporalDenoiserPlugin);
+	check(GetActiveSpatialTemporalDenoiser());
 
 	FRDGTextureRef ProcessedNormalTexture = NormalTexture;
 
@@ -976,12 +1141,12 @@ static void PathTracingSpatialTemporalDenoiserPlugin(FRDGBuilder& GraphBuilder,
 	Inputs.bForceSpatialDenoiserOnly = bForceSpatialDenoiserOnly;
 
 	if (Context.SpatialTemporalDenoiserHistory &&
-		Context.SpatialTemporalDenoiserHistory->GetDebugName() == GPathTracingSpatialTemporalDenoiserPlugin->GetDebugName())
+		Context.SpatialTemporalDenoiserHistory->GetDebugName() == GetActiveSpatialTemporalDenoiser()->GetDebugName())
 	{
 		Inputs.PrevHistory = Context.SpatialTemporalDenoiserHistory;
 	}
 
-	IPathTracingSpatialTemporalDenoiser::FOutputs Outputs = GPathTracingSpatialTemporalDenoiserPlugin->AddPasses(GraphBuilder, View, Inputs);
+	IPathTracingSpatialTemporalDenoiser::FOutputs Outputs = GetActiveSpatialTemporalDenoiser()->AddPasses(GraphBuilder, View, Inputs);
 	if (Outputs.NewHistory)
 	{
 		Context.SpatialTemporalDenoiserHistory = Outputs.NewHistory;
@@ -1005,12 +1170,12 @@ static void PathTracingMotionVectorPlugin(FRDGBuilder& GraphBuilder,
 	FRDGTextureRef ReferenceFrameTexture,
 	FRDGTextureRef OutputTexture)
 {
-	check(GPathTracingSpatialTemporalDenoiserPlugin);
+	check(GetActiveSpatialTemporalDenoiser());
 
 	bool bShouldApplyPreExposure = ShouldApplyPreExposureToMotionVectorEstimation();
 	float PreExposure = bShouldApplyPreExposure ? View.PreExposure : 1.0f;
 
-	GPathTracingSpatialTemporalDenoiserPlugin->AddMotionVectorPass(GraphBuilder, View, {InputFrameTexture, ReferenceFrameTexture, OutputTexture, PreExposure});
+	GetActiveSpatialTemporalDenoiser()->AddMotionVectorPass(GraphBuilder, View, {InputFrameTexture, ReferenceFrameTexture, OutputTexture, PreExposure});
 }
 
 class FMotionVectorEstimationContext
@@ -1823,6 +1988,7 @@ void PathTracingSpatialTemporalDenoising(FRDGBuilder& GraphBuilder,
 				SpatialTemporalDenoisingContext.RadianceTexture,
 				SpatialTemporalDenoisingContext.AlbedoTexture,
 				SpatialTemporalDenoisingContext.NormalTexture,
+				SpatialTemporalDenoisingContext.VarianceTexture,
 				TargetTexture);
 		}
 
diff --git a/Engine/Source/Runtime/Renderer/Public/PathTracingDenoiser.h b/Engine/Source/Runtime/Renderer/Public/PathTracingDenoiser.h
--- a/Engine/Source/Runtime/Renderer/Public/PathTracingDenoiser.h
+++ b/Engine/Source/Runtime/Renderer/Public/PathTracingDenoiser.h
@@ -106,5 +106,17 @@ public:
 
 } // namespace UE::Renderer::Private
 
+UE_DEPRECATED(5.5, "GPathTracingDenoiserPlugin is deprecated. Use RegisterSpatialDenoiser/RegisterSpatialTemporalDenoiser to register")
 extern RENDERER_API TUniquePtr<UE::Renderer::Private::IPathTracingDenoiser> GPathTracingDenoiserPlugin;
-extern RENDERER_API TUniquePtr<UE::Renderer::Private::IPathTracingSpatialTemporalDenoiser> GPathTracingSpatialTemporalDenoiserPlugin;
\ No newline at end of file
+
+UE_DEPRECATED(5.5, "GPathTracingSpatialTemporalDenoiserPlugin is deprecated. Use RegisterSpatialDenoiser/RegisterSpatialTemporalDenoiser to register")
+extern RENDERER_API TUniquePtr<UE::Renderer::Private::IPathTracingSpatialTemporalDenoiser> GPathTracingSpatialTemporalDenoiserPlugin;
+
+/** Register a spatial denoiser*/
+extern RENDERER_API void RegisterSpatialDenoiser(TUniquePtr<UE::Renderer::Private::IPathTracingDenoiser> PathTracingDenoiser, FString Name);
+
+/** Register a spatial temporal denoiser*/
+extern RENDERER_API void RegisterSpatialTemporalDenoiser(TUniquePtr<UE::Renderer::Private::IPathTracingSpatialTemporalDenoiser> PathTracingDenoiser, FString Name);
+
+/** Unregister all denoisers with the name Name*/
+extern RENDERER_API void UnregisterDenoiser(FString Name);
