Fix EnergyConservation tables format not supported on certain platforms.

This CLs adds RG8/R8 format support for the energy conservation tables.

diff --git a/Engine/Source/Runtime/Renderer/Private/ShadingEnergyConservation.cpp b/Engine/Source/Runtime/Renderer/Private/ShadingEnergyConservation.cpp
--- a/Engine/Source/Runtime/Renderer/Private/ShadingEnergyConservation.cpp
+++ b/Engine/Source/Runtime/Renderer/Private/ShadingEnergyConservation.cpp
@@ -204,6 +204,22 @@ IMPLEMENT_SHADER_TYPE(, FBuildShadingEnergyConservationTableCS, TEXT("/Engine/Pr
 
 ////////////////////////////////////////////////////////////////////////////////////////////////
 
+struct FTranscode
+{
+	FTranscode(const uint16* InData, uint32 InElementCount)
+	{
+		const uint32 ElementCount = sizeof(ShadingEnergyConservationData::GGXSpecValues) / sizeof(uint16);
+		Data.SetNum(InElementCount);
+		for (uint32 It = 0; It < InElementCount; ++It)
+		{
+			Data[It] = uint8(FMath::Clamp((float(InData[It]) / 65535.f) * 0xFF, 0u, 0xFF));
+		}
+	}
+	TArray<uint8> Data;
+};
+
+#define TRANSCODE_16_TO_8(T) (FTranscode(T, sizeof(T) / sizeof(uint16)).Data.GetData())
+
 namespace ShadingEnergyConservation
 {
 
@@ -233,9 +249,19 @@ void Init(FRDGBuilder& GraphBuilder, FViewInfo& View)
 		// Change this to true in order to regenerate the energy tables, and manually copy the coefficients into ShadingEnergyConservationData.h
 		const bool bRuntimeGeneration = false;
 
+		EPixelFormat Format = PF_R8G8;
+		// for low roughness we would get banding with PF_R8G8 but for low spec it could be used, for now we don't do this optimization
+		if (GPixelFormats[PF_G16R16].Supported && UE::PixelFormat::HasCapabilities(PF_G16R16, EPixelFormatCapabilities::TextureFilterable))
+		{
+			Format = PF_G16R16;
+		}
+
+		const bool bRG16Supported = GPixelFormats[PF_G16R16].Supported && UE::PixelFormat::HasCapabilities(PF_G16R16, EPixelFormatCapabilities::TextureFilterable);
+		const bool bR16Supported = GPixelFormats[PF_G16].Supported && UE::PixelFormat::HasCapabilities(PF_G16, EPixelFormatCapabilities::TextureFilterable);
+
 		const int Size = SHADING_ENERGY_CONSERVATION_TABLE_RESOLUTION;
-		const EPixelFormat SpecFormat = bRuntimeGeneration && CVarShadingFurnaceTest_TableFormat.GetValueOnRenderThread() > 0 ? PF_G32R32F : PF_G16R16;
-		const EPixelFormat DiffFormat = PF_G16;
+		const EPixelFormat SpecFormat = bRuntimeGeneration && CVarShadingFurnaceTest_TableFormat.GetValueOnRenderThread() > 0 ? PF_G32R32F : (bRG16Supported ? PF_G16R16 : PF_R8G8);
+		const EPixelFormat DiffFormat = bR16Supported ? PF_G16 : PF_R8;
 		const bool bBuildTable = 
 			View.ViewState->ShadingEnergyConservationData.Format != SpecFormat ||
 			View.ViewState->ShadingEnergyConservationData.GGXSpecEnergyTexture == nullptr ||
@@ -328,17 +354,36 @@ void Init(FRDGBuilder& GraphBuilder, FViewInfo& View)
 			else
 			{
 				// Precomputed data are stored as float16
-				check(SpecFormat == PF_G16R16);
+				check(SpecFormat == PF_G16R16 || SpecFormat == PF_R8G8);
+				check(DiffFormat == PF_G16 || DiffFormat == PF_R8);
 
 				View.ViewState->ShadingEnergyConservationData.GGXSpecEnergyTexture  = GRenderTargetPool.FindFreeElement(FRDGTextureDesc::Create2D(FIntPoint(Size, Size),        SpecFormat, FClearValueBinding::None, TexCreate_ShaderResource), TEXT("Shading.GGXSpecEnergy"));
 				View.ViewState->ShadingEnergyConservationData.GGXGlassEnergyTexture = GRenderTargetPool.FindFreeElement(FRDGTextureDesc::Create3D(FIntVector(Size, Size, Size), SpecFormat, FClearValueBinding::None, TexCreate_ShaderResource), TEXT("Shading.GGXGlassEnergy"));
 				View.ViewState->ShadingEnergyConservationData.ClothEnergyTexture    = GRenderTargetPool.FindFreeElement(FRDGTextureDesc::Create2D(FIntPoint(Size, Size),        SpecFormat, FClearValueBinding::None, TexCreate_ShaderResource), TEXT("Shading.ClothSpecEnergy"));
 				View.ViewState->ShadingEnergyConservationData.DiffuseEnergyTexture  = GRenderTargetPool.FindFreeElement(FRDGTextureDesc::Create2D(FIntPoint(Size, Size),        DiffFormat,	FClearValueBinding::None, TexCreate_ShaderResource), TEXT("Shading.DiffuseEnergy"));
 
-				ShadingEnergyConservationData::LockCopyTexture2D(GraphBuilder.RHICmdList, View.ViewState->ShadingEnergyConservationData.GGXSpecEnergyTexture->GetRHI(),  ShadingEnergyConservationData::GGXSpecValues, 2);
-				ShadingEnergyConservationData::LockCopyTexture3D(GraphBuilder.RHICmdList, View.ViewState->ShadingEnergyConservationData.GGXGlassEnergyTexture->GetRHI(), ShadingEnergyConservationData::GGXGlassValues, 2);
-				ShadingEnergyConservationData::LockCopyTexture2D(GraphBuilder.RHICmdList, View.ViewState->ShadingEnergyConservationData.ClothEnergyTexture->GetRHI(),    Substrate::IsSubstrateEnabled() ? ShadingEnergyConservationData::SubstrateClothSpecValues : ShadingEnergyConservationData::ClothSpecValues, 2);
-				ShadingEnergyConservationData::LockCopyTexture2D(GraphBuilder.RHICmdList, View.ViewState->ShadingEnergyConservationData.DiffuseEnergyTexture->GetRHI(),  ShadingEnergyConservationData::DiffuseValues, 1);
+
+				if (bRG16Supported)
+				{
+					ShadingEnergyConservationData::LockCopyTexture2D(GraphBuilder.RHICmdList, View.ViewState->ShadingEnergyConservationData.GGXSpecEnergyTexture->GetRHI(),  ShadingEnergyConservationData::GGXSpecValues, 2);
+					ShadingEnergyConservationData::LockCopyTexture3D(GraphBuilder.RHICmdList, View.ViewState->ShadingEnergyConservationData.GGXGlassEnergyTexture->GetRHI(), ShadingEnergyConservationData::GGXGlassValues, 2);
+					ShadingEnergyConservationData::LockCopyTexture2D(GraphBuilder.RHICmdList, View.ViewState->ShadingEnergyConservationData.ClothEnergyTexture->GetRHI(),    Substrate::IsSubstrateEnabled() ? ShadingEnergyConservationData::SubstrateClothSpecValues : ShadingEnergyConservationData::ClothSpecValues, 2);
+				}
+				else
+				{
+					ShadingEnergyConservationData::LockCopyTexture2D(GraphBuilder.RHICmdList, View.ViewState->ShadingEnergyConservationData.GGXSpecEnergyTexture->GetRHI(),  TRANSCODE_16_TO_8(ShadingEnergyConservationData::GGXSpecValues), 2);
+					ShadingEnergyConservationData::LockCopyTexture3D(GraphBuilder.RHICmdList, View.ViewState->ShadingEnergyConservationData.GGXGlassEnergyTexture->GetRHI(), TRANSCODE_16_TO_8(ShadingEnergyConservationData::GGXGlassValues), 2);
+					ShadingEnergyConservationData::LockCopyTexture2D(GraphBuilder.RHICmdList, View.ViewState->ShadingEnergyConservationData.ClothEnergyTexture->GetRHI(),    Substrate::IsSubstrateEnabled() ? TRANSCODE_16_TO_8(ShadingEnergyConservationData::SubstrateClothSpecValues) : TRANSCODE_16_TO_8(ShadingEnergyConservationData::ClothSpecValues), 2);
+				}
+
+				if (bR16Supported)
+				{
+					ShadingEnergyConservationData::LockCopyTexture2D(GraphBuilder.RHICmdList, View.ViewState->ShadingEnergyConservationData.DiffuseEnergyTexture->GetRHI(),  ShadingEnergyConservationData::DiffuseValues, 1);
+				}
+				else
+				{
+					ShadingEnergyConservationData::LockCopyTexture2D(GraphBuilder.RHICmdList, View.ViewState->ShadingEnergyConservationData.DiffuseEnergyTexture->GetRHI(),  TRANSCODE_16_TO_8(ShadingEnergyConservationData::DiffuseValues), 1);
+				}
 
 				GGXSpecEnergyTexture  = GraphBuilder.RegisterExternalTexture(View.ViewState->ShadingEnergyConservationData.GGXSpecEnergyTexture);
 				GGXGlassEnergyTexture = GraphBuilder.RegisterExternalTexture(View.ViewState->ShadingEnergyConservationData.GGXGlassEnergyTexture);
