Fixes depth encoding in motion blur on orthographic views

diff --git a/Engine/Source/Runtime/Renderer/Private/PostProcess/PostProcessMotionBlur.cpp b/Engine/Source/Runtime/Renderer/Private/PostProcess/PostProcessMotionBlur.cpp
--- a/Engine/Source/Runtime/Renderer/Private/PostProcess/PostProcessMotionBlur.cpp
+++ b/Engine/Source/Runtime/Renderer/Private/PostProcess/PostProcessMotionBlur.cpp
@@ -649,7 +649,7 @@ void AddMotionBlurVelocityPass(
 			// NOTE: Use scene depth's dimensions because velocity can actually be a 1x1 black texture when there are no moving objects in sight.
 			FRDGTextureDesc Desc = FRDGTextureDesc::Create2D(
 				Inputs.SceneDepth.Texture->Desc.Extent,
-				PF_FloatR11G11B10,
+				View.IsPerspectiveProjection() ? PF_FloatR11G11B10 : PF_A32B32G32R32F,
 				FClearValueBinding::None,
 				GFastVRamConfig.VelocityFlat | TexCreate_ShaderResource | TexCreate_UAV);
 
diff --git a/Engine/Source/Runtime/Renderer/Private/PostProcess/TemporalSuperResolution.cpp b/Engine/Source/Runtime/Renderer/Private/PostProcess/TemporalSuperResolution.cpp
--- a/Engine/Source/Runtime/Renderer/Private/PostProcess/TemporalSuperResolution.cpp
+++ b/Engine/Source/Runtime/Renderer/Private/PostProcess/TemporalSuperResolution.cpp
@@ -1954,7 +1954,7 @@ FDefaultTemporalUpscaler::FOutputs AddTemporalSuperResolutionPasses(
 			{
 				FRDGTextureDesc Desc = FRDGTextureDesc::Create2D(
 					InputExtent,
-					PF_FloatR11G11B10,
+					bIsOrthoProjection ? PF_A32B32G32R32F : PF_FloatR11G11B10,
 					FClearValueBinding::None,
 					GFastVRamConfig.VelocityFlat | TexCreate_ShaderResource | TexCreate_UAV);
 
