Fixed: Scene depth capture modes now do not ask for a SceneColor target with an alpha channel.
Fixed: Custom depth capture renderpass to use same subpass hints as main renderer.
These changes help to avoid creating new unecessary PSOs while capturing scene depth.

diff --git a/Engine/Source/Runtime/Renderer/Private/MobileShadingRenderer.cpp b/Engine/Source/Runtime/Renderer/Private/MobileShadingRenderer.cpp
--- a/Engine/Source/Runtime/Renderer/Private/MobileShadingRenderer.cpp
+++ b/Engine/Source/Runtime/Renderer/Private/MobileShadingRenderer.cpp
@@ -873,14 +873,20 @@ void FMobileSceneRenderer::RenderMaskedPrePass(FRHICommandList& RHICmdList, cons
 void FMobileSceneRenderer::RenderCustomRenderPassBasePass(FRDGBuilder& GraphBuilder, TArrayView<FViewInfo> InViews, FRDGTextureRef ViewFamilyTexture, FSceneTextures& SceneTextures)
 {
 	FRenderTargetBindingSlots BasePassRenderTargets;
+	// Use the same subpass hints as main render, to avoid generating new PSOs 
+	int32 NumAdditionalSubpasses = 0;
 	if (bDeferredShading)
 	{
 		FColorTargets ColorTargets = GetColorTargets_Deferred(SceneTextures);
 		BasePassRenderTargets = InitRenderTargetBindings_Deferred(SceneTextures, ColorTargets);
+		BasePassRenderTargets.SubpassHint = ESubpassHint::DeferredShadingSubpass;
+		NumAdditionalSubpasses = 2;
 	}
 	else
 	{
 		BasePassRenderTargets = InitRenderTargetBindings_Forward(ViewFamilyTexture, SceneTextures);
+		BasePassRenderTargets.SubpassHint = ESubpassHint::DepthReadSubpass;
+		NumAdditionalSubpasses = 1;
 	}
 
 	FRenderViewContextArray RenderViews;
@@ -905,10 +911,15 @@ void FMobileSceneRenderer::RenderCustomRenderPassBasePass(FRDGBuilder& GraphBuil
 			RDG_EVENT_NAME("RenderMobileBasePass"),
 			PassParameters,
 			ERDGPassFlags::Raster,
-			[this, PassParameters, ViewContext, &SceneTextures](FRHICommandList& RHICmdList)
+			[this, PassParameters, ViewContext, &SceneTextures, NumAdditionalSubpasses](FRHICommandList& RHICmdList)
 			{
 				FViewInfo& View = *ViewContext.ViewInfo;
 				RenderMobileBasePass(RHICmdList, View, &PassParameters->InstanceCullingDrawParams);
+
+				for (int32 AddSubpass = 0; AddSubpass < NumAdditionalSubpasses; ++AddSubpass)
+				{
+					RHICmdList.NextSubpass();
+				}
 			});
 
 		if (!bIsFullDepthPrepassEnabled)
diff --git a/Engine/Source/Runtime/Renderer/Private/SceneRendering.cpp b/Engine/Source/Runtime/Renderer/Private/SceneRendering.cpp
--- a/Engine/Source/Runtime/Renderer/Private/SceneRendering.cpp
+++ b/Engine/Source/Runtime/Renderer/Private/SceneRendering.cpp
@@ -6502,3 +6502,28 @@ FRDGTextureRef CreateQuarterResolutionDepthMinAndMax(FRDGBuilder& GraphBuilder,
 
 	return SmallTexture;
 }
+
+bool SceneCaptureRequiresAlphaChannel(const FSceneView& View)
+{
+	// Planar reflections and scene captures use scene color alpha to keep track of where content has been rendered, for compositing into a different scene later
+	if (View.bIsPlanarReflection)
+	{
+		return true;
+	}
+
+	if (View.bIsSceneCapture)
+	{
+		// Depth capture modes do not require alpha channel
+		if (View.CustomRenderPass)
+		{
+			return View.CustomRenderPass->GetRenderOutput() != FCustomRenderPassBase::ERenderOutput::SceneDepth
+				&& View.CustomRenderPass->GetRenderOutput() != FCustomRenderPassBase::ERenderOutput::DeviceDepth;
+		}
+		else if(View.Family)
+		{
+			return View.Family->SceneCaptureSource != SCS_SceneDepth 
+				&& View.Family->SceneCaptureSource != SCS_DeviceDepth;
+		}
+	}
+	return false;
+}
diff --git a/Engine/Source/Runtime/Renderer/Private/SceneRendering.h b/Engine/Source/Runtime/Renderer/Private/SceneRendering.h
--- a/Engine/Source/Runtime/Renderer/Private/SceneRendering.h
+++ b/Engine/Source/Runtime/Renderer/Private/SceneRendering.h
@@ -2894,3 +2894,5 @@ inline const FViewFamilyInfo& GetViewFamilyInfo(const TArrayView<const FViewInfo
 	check(Views.Num() == 1 || Views[0].Family == Views.Last().Family);
 	return *(const FViewFamilyInfo*)Views[0].Family;
 }
+
+bool SceneCaptureRequiresAlphaChannel(const FSceneView& View);
diff --git a/Engine/Source/Runtime/Renderer/Private/SceneTextures.cpp b/Engine/Source/Runtime/Renderer/Private/SceneTextures.cpp
--- a/Engine/Source/Runtime/Renderer/Private/SceneTextures.cpp
+++ b/Engine/Source/Runtime/Renderer/Private/SceneTextures.cpp
@@ -379,15 +379,11 @@ void InitializeSceneTexturesConfig(FSceneTexturesConfig& Config, const FSceneVie
 
 	bool bRequiresAlphaChannel = ShadingPath == EShadingPath::Mobile ? IsMobilePropagateAlphaEnabled(ViewFamily.GetShaderPlatform()) : false;
 	int32 NumberOfViewsWithMultiviewEnabled = 0;
-	for (int32 ViewIndex = 0; ViewIndex < ViewFamily.AllViews.Num(); ViewIndex++)
+	
+	for (const FSceneView* View : ViewFamily.AllViews)
 	{
-		// Planar reflections and scene captures use scene color alpha to keep track of where content has been rendered, for compositing into a different scene later
-		if (ViewFamily.AllViews[ViewIndex]->bIsPlanarReflection || ViewFamily.AllViews[ViewIndex]->bIsSceneCapture)
-		{
-			bRequiresAlphaChannel = true;
-		}
-
-		NumberOfViewsWithMultiviewEnabled += (ViewFamily.AllViews[ViewIndex]->bIsMobileMultiViewEnabled) ? 1 : 0;
+		bRequiresAlphaChannel|= SceneCaptureRequiresAlphaChannel(*View);
+		NumberOfViewsWithMultiviewEnabled += (View->bIsMobileMultiViewEnabled) ? 1 : 0;
 	}
 
 	ensureMsgf(NumberOfViewsWithMultiviewEnabled == 0 || NumberOfViewsWithMultiviewEnabled == ViewFamily.AllViews.Num(),
