Adding ECVF_Scalability to r.LUT.Size. Default remains the same = 32.

Based on customer request. Citing the need for increasing the LUT size only on higher scalability settings due to increased cost to CombineLUTs pass.

[RN] Adding ability to drive r.LUT.Size through Scalability system

diff --git a/Engine/Source/Runtime/Renderer/Private/PostProcess/PostProcessCombineLUTs.cpp b/Engine/Source/Runtime/Renderer/Private/PostProcess/PostProcessCombineLUTs.cpp
--- a/Engine/Source/Runtime/Renderer/Private/PostProcess/PostProcessCombineLUTs.cpp
+++ b/Engine/Source/Runtime/Renderer/Private/PostProcess/PostProcessCombineLUTs.cpp
@@ -37,7 +37,7 @@ TAutoConsoleVariable<int32> CVarLUTSize(
 	TEXT("r.LUT.Size"),
 	32,
 	TEXT("Size of film LUT"),
-	ECVF_RenderThreadSafe);
+	ECVF_RenderThreadSafe | ECVF_Scalability);
 
 TAutoConsoleVariable<int32> CVarColorGrading(
 	TEXT("r.Color.Grading"), 1,
