Avoid shader aspects creation at every mobile multi view check.

diff --git a/Engine/Source/Runtime/Engine/Private/SceneView.cpp b/Engine/Source/Runtime/Engine/Private/SceneView.cpp
--- a/Engine/Source/Runtime/Engine/Private/SceneView.cpp
+++ b/Engine/Source/Runtime/Engine/Private/SceneView.cpp
@@ -961,7 +961,7 @@ PRAGMA_ENABLE_DEPRECATION_WARNINGS
 	{
 		// The shader variants that are compiled have ISR or MMV _enabled_ in the shader, even if the current ViewFamily doesn't
 		// require multiple views functionality.
-		const UE::StereoRenderUtils::FStereoShaderAspects Aspects(ShaderPlatform);
+		Aspects = UE::StereoRenderUtils::FStereoShaderAspects(ShaderPlatform);
 		bShouldBindInstancedViewUB = Aspects.IsInstancedStereoEnabled() || Aspects.IsMobileMultiViewEnabled();
 
 		if (Family && Family->bRequireMultiView && !Aspects.IsMobileMultiViewEnabled())
diff --git a/Engine/Source/Runtime/Engine/Public/SceneView.h b/Engine/Source/Runtime/Engine/Public/SceneView.h
--- a/Engine/Source/Runtime/Engine/Public/SceneView.h
+++ b/Engine/Source/Runtime/Engine/Public/SceneView.h
@@ -16,6 +16,7 @@
 #include "RenderResource.h"
 #include "ShowFlags.h"
 #include "StereoRendering.h"
+#include "StereoRenderUtils.h"
 
 #define MAX_PHYSICS_FIELD_TARGETS 32
 
@@ -1625,6 +1626,9 @@ public:
 	/** Use to allow ray tracing on this view. */
 	bool bAllowRayTracing = true;
 
+	/**  Stereo aspects of the shader pipeline based on this view's shader platform */
+	UE::StereoRenderUtils::FStereoShaderAspects Aspects;
+
 protected:
 	friend class FSceneRenderer;
 
diff --git a/Engine/Source/Runtime/RenderCore/Private/StereoRenderUtils.cpp b/Engine/Source/Runtime/RenderCore/Private/StereoRenderUtils.cpp
--- a/Engine/Source/Runtime/RenderCore/Private/StereoRenderUtils.cpp
+++ b/Engine/Source/Runtime/RenderCore/Private/StereoRenderUtils.cpp
@@ -128,6 +128,14 @@ RENDERCORE_API FStereoShaderAspects::FStereoShaderAspects(EShaderPlatform Platfo
 	checkf(!bInstancedStereoEnabled || (bInstancedMultiViewportEnabled || bMobileMultiViewFallback), TEXT("If ISR is enabled, we need either multi-viewport (since we no longer support clip-distance method) or MMV fallback (which uses vertex layer)."));
 }
 	
+RENDERCORE_API FStereoShaderAspects::FStereoShaderAspects() :
+	bInstancedStereoEnabled(false)
+	, bMobileMultiViewEnabled(false)
+	, bInstancedMultiViewportEnabled(false)
+	, bInstancedStereoNative(false)
+	, bMobileMultiViewNative(false)
+	, bMobileMultiViewFallback(false)
+{}
 
 RENDERCORE_API void LogISRInit(const UE::StereoRenderUtils::FStereoShaderAspects& Aspects)
 {
diff --git a/Engine/Source/Runtime/RenderCore/Public/StereoRenderUtils.h b/Engine/Source/Runtime/RenderCore/Public/StereoRenderUtils.h
--- a/Engine/Source/Runtime/RenderCore/Public/StereoRenderUtils.h
+++ b/Engine/Source/Runtime/RenderCore/Public/StereoRenderUtils.h
@@ -24,6 +24,11 @@ public:
 	*/
 	RENDERCORE_API FStereoShaderAspects(EShaderPlatform Platform);
 
+	/**  
+	* Default empty constructor for object in FSceneView. Do not use! 
+	*/
+	RENDERCORE_API FStereoShaderAspects();
+
 	/**
 	* Whether instanced stereo rendering is enabled - i.e. using a single instanced drawcall to render to both stereo views.
 	* The output is redirected via the viewport index.
diff --git a/Engine/Source/Runtime/Renderer/Private/DecalRenderingShared.cpp b/Engine/Source/Runtime/Renderer/Private/DecalRenderingShared.cpp
--- a/Engine/Source/Runtime/Renderer/Private/DecalRenderingShared.cpp
+++ b/Engine/Source/Runtime/Renderer/Private/DecalRenderingShared.cpp
@@ -526,7 +526,7 @@ namespace DecalRendering
 
 	FMatrix ComputeComponentToClipMatrix(const FViewInfo& View, const FMatrix& DecalComponentToWorld)
 	{
-		if (View.bIsMobileMultiViewEnabled || UE::StereoRenderUtils::FStereoShaderAspects(View.GetShaderPlatform()).IsMobileMultiViewEnabled())
+		if (View.bIsMobileMultiViewEnabled || View.Aspects.IsMobileMultiViewEnabled())
 		{
 			// In multi view, the rest of the matrix that is multiplied with DecalComponentToWorld in the non-multi view
 			// case is split out in ViewUniformShaderParameters.MobileMultiviewDecalTransform so we can multiply
diff --git a/Engine/Source/Runtime/Renderer/Private/MobileBasePassRendering.cpp b/Engine/Source/Runtime/Renderer/Private/MobileBasePassRendering.cpp
--- a/Engine/Source/Runtime/Renderer/Private/MobileBasePassRendering.cpp
+++ b/Engine/Source/Runtime/Renderer/Private/MobileBasePassRendering.cpp
@@ -261,7 +261,7 @@ void SetupMobileBasePassUniformParameters(
 	const FViewInfo* InstancedView = View.GetInstancedView();
 
 	SetupFogUniformParameters(GraphBuilder, View, BasePassParameters.Fog);
-	if (InstancedView && (View.bIsMobileMultiViewEnabled || UE::StereoRenderUtils::FStereoShaderAspects(View.GetShaderPlatform()).IsMobileMultiViewEnabled()))
+	if (InstancedView && (View.bIsMobileMultiViewEnabled || View.Aspects.IsMobileMultiViewEnabled()))
 	{
 		SetupFogUniformParameters(GraphBuilder, *InstancedView, BasePassParameters.FogMMV);
 	}
@@ -324,7 +324,7 @@ void SetupMobileBasePassUniformParameters(
 	BasePassParameters.AmbientOcclusionSampler = TStaticSamplerState<SF_Bilinear, AM_Clamp, AM_Clamp, AM_Clamp>::GetRHI();
 	BasePassParameters.AmbientOcclusionStaticFraction = FMath::Clamp(View.FinalPostProcessSettings.AmbientOcclusionStaticFraction, 0.0f, 1.0f);
 	
-	const bool bIsMobileMultiViewEnabled = View.bIsMobileMultiViewEnabled || UE::StereoRenderUtils::FStereoShaderAspects(View.GetShaderPlatform()).IsMobileMultiViewEnabled();
+	const bool bIsMobileMultiViewEnabled = View.bIsMobileMultiViewEnabled || View.Aspects.IsMobileMultiViewEnabled();
 	BasePassParameters.DBuffer = GetDBufferParameters(GraphBuilder, MobileBasePassTextures.DBufferTextures, View.GetShaderPlatform(), bIsMobileMultiViewEnabled);
 
 	const bool bMobileUsesShadowMaskTexture = MobileUsesShadowMaskTexture(View.GetShaderPlatform());
diff --git a/Engine/Source/Runtime/Renderer/Private/MobileShadingRenderer.cpp b/Engine/Source/Runtime/Renderer/Private/MobileShadingRenderer.cpp
--- a/Engine/Source/Runtime/Renderer/Private/MobileShadingRenderer.cpp
+++ b/Engine/Source/Runtime/Renderer/Private/MobileShadingRenderer.cpp
@@ -75,7 +75,6 @@
 #include "WaterInfoTextureRendering.h"
 #include "Rendering/CustomRenderPass.h"
 #include "GenerateMips.h"
-#include "StereoRenderUtils.h"
 #include "MobileSSR.h"
 
 uint32 GetShadowQuality();
@@ -802,7 +801,7 @@ void FMobileSceneRenderer::RenderFullDepthPrepass(FRDGBuilder& GraphBuilder, TAr
 		PassParameters->View = View.GetShaderParameters();
 		PassParameters->MobileBasePass = CreateMobileBasePassUniformBuffer(GraphBuilder, View, EMobileBasePass::DepthPrePass, EMobileSceneTextureSetupMode::None);
 		//if the scenecolor isn't multiview but the app is, need to render as a single-view multiview due to shaders
-		PassParameters->RenderTargets.MultiViewCount = (View.bIsMobileMultiViewEnabled) ? 2 : (UE::StereoRenderUtils::FStereoShaderAspects(View.GetShaderPlatform()).IsMobileMultiViewEnabled() ? 1 : 0);
+		PassParameters->RenderTargets.MultiViewCount = (View.bIsMobileMultiViewEnabled) ? 2 : (View.Aspects.IsMobileMultiViewEnabled() ? 1 : 0);
 		if (ViewContext.bIsLastView)
 		{
 			LastViewMobileBasePassUB = PassParameters->MobileBasePass;
@@ -1555,7 +1554,7 @@ void FMobileSceneRenderer::RenderForward(FRDGBuilder& GraphBuilder, FRDGTextureR
 	BasePassRenderTargets.ShadingRateTexture = (!MainView.bIsSceneCapture && !MainView.bIsReflectionCapture && (NewShadingRateTarget != nullptr)) ? NewShadingRateTarget : nullptr;
 
 	//if the scenecolor isn't multiview but the app is, need to render as a single-view multiview due to shaders
-	BasePassRenderTargets.MultiViewCount = (MainView.bIsMobileMultiViewEnabled) ? 2 : (UE::StereoRenderUtils::FStereoShaderAspects(MainView.GetShaderPlatform()).IsMobileMultiViewEnabled() ? 1 : 0);
+	BasePassRenderTargets.MultiViewCount = (MainView.bIsMobileMultiViewEnabled) ? 2 : (MainView.Aspects.IsMobileMultiViewEnabled() ? 1 : 0);
 
 
 	const FRDGSystemTextures& SystemTextures = FRDGSystemTextures::Get(GraphBuilder);
@@ -1911,7 +1910,7 @@ FRenderTargetBindingSlots FMobileSceneRenderer::InitRenderTargetBindings_Deferre
 	BasePassRenderTargets.ShadingRateTexture = nullptr;
 	
 	//if the scenecolor isn't multiview but the app is, need to render as a single-view multiview due to shaders
-	BasePassRenderTargets.MultiViewCount = (Views[0].bIsMobileMultiViewEnabled) ? 2 : (UE::StereoRenderUtils::FStereoShaderAspects(Views[0].GetShaderPlatform()).IsMobileMultiViewEnabled() ? 1 : 0);
+	BasePassRenderTargets.MultiViewCount = (Views[0].bIsMobileMultiViewEnabled) ? 2 : (Views[0].Aspects.IsMobileMultiViewEnabled() ? 1 : 0);
 
 	return BasePassRenderTargets;
 }
diff --git a/Engine/Source/Runtime/Renderer/Private/ShadowRendering.cpp b/Engine/Source/Runtime/Renderer/Private/ShadowRendering.cpp
--- a/Engine/Source/Runtime/Renderer/Private/ShadowRendering.cpp
+++ b/Engine/Source/Runtime/Renderer/Private/ShadowRendering.cpp
@@ -1739,16 +1739,10 @@ FMatrix FProjectedShadowInfo::GetScreenToShadowMatrix(const FSceneView& View, ui
 				1
 				)
 			);
-
-	// Aspects creation is embedded in lambda so we don't pay the price for it when it's not needed.
-	auto IsMobileMultiViewEnabledInAspects = [](const FSceneView& View) {
-		const UE::StereoRenderUtils::FStereoShaderAspects Aspects(View.GetShaderPlatform());
-		return Aspects.IsMobileMultiViewEnabled();
-	};
 	
 	// Checking the Aspects is a workaround for editor windows or scene captures where View.bIsMobileMultiViewEnabled
 	// is false but shaders have been compiled with MOBILE_MULTI_VIEW enabled.
-	if (View.bIsMobileMultiViewEnabled || IsMobileMultiViewEnabledInAspects(View))
+	if (View.bIsMobileMultiViewEnabled || View.Aspects.IsMobileMultiViewEnabled())
 	{
 		// In Multiview, we split ViewDependentTransform out into ViewUniformShaderParameters.MobileMultiviewShadowTransform
 		// So we can multiply it later in shader.
