Avoid DBufferTexArray texture creation when not using mobile multi view

diff --git a/Engine/Source/Runtime/Renderer/Private/CompositionLighting/PostProcessDeferredDecals.cpp b/Engine/Source/Runtime/Renderer/Private/CompositionLighting/PostProcessDeferredDecals.cpp
--- a/Engine/Source/Runtime/Renderer/Private/CompositionLighting/PostProcessDeferredDecals.cpp
+++ b/Engine/Source/Runtime/Renderer/Private/CompositionLighting/PostProcessDeferredDecals.cpp
@@ -192,7 +192,7 @@ void GetDeferredDecalPassParameters(
 
 	FRenderTargetBindingSlots& RenderTargets = PassParameters.RenderTargets;
 	PassParameters.RenderTargets.ShadingRateTexture = GVRSImageManager.GetVariableRateShadingImage(GraphBuilder, View, FVariableRateShadingImageManager::EVRSPassType::Decals);
-	PassParameters.RenderTargets.MultiViewCount = (View.bIsMobileMultiViewEnabled) ? 2 : (UE::StereoRenderUtils::FStereoShaderAspects(View.GetShaderPlatform()).IsMobileMultiViewEnabled() ? 1 : 0);
+	PassParameters.RenderTargets.MultiViewCount = (View.bIsMobileMultiViewEnabled) ? 2 : (View.Aspects.IsMobileMultiViewEnabled() ? 1 : 0);
 	uint32 ColorTargetIndex = 0;
 
 	const auto AddColorTarget = [&](FRDGTextureRef Texture, ERenderTargetLoadAction LoadAction = ERenderTargetLoadAction::ELoad, FRDGTextureRef TextureArray = nullptr, bool bIsMobileMultiView = false)
@@ -232,14 +232,16 @@ void GetDeferredDecalPassParameters(
 
 		const FDBufferTextures& DBufferTextures = *Textures.DBufferTextures;
 
-		const ERenderTargetLoadAction LoadAction = (DBufferTextures.DBufferA->HasBeenProduced() || DBufferTextures.DBufferATexArray->HasBeenProduced())
+		const bool bDBufferAProduced = DBufferTextures.DBufferA ? DBufferTextures.DBufferA->HasBeenProduced() : false;
+		const bool bDBufferTexArrayAProduced = DBufferTextures.DBufferATexArray ? DBufferTextures.DBufferATexArray->HasBeenProduced() : false;
+		const bool bUseTextureArrays = View.bIsMobileMultiViewEnabled || UE::StereoRenderUtils::FStereoShaderAspects(View.GetShaderPlatform()).IsMobileMultiViewEnabled();
+		const ERenderTargetLoadAction LoadAction = (bUseTextureArrays ? bDBufferTexArrayAProduced : bDBufferAProduced)
 			? ERenderTargetLoadAction::ELoad
 			: ERenderTargetLoadAction::EClear;
 
-		bool bIsMobileMultiView = View.bIsMobileMultiViewEnabled || UE::StereoRenderUtils::FStereoShaderAspects(View.GetShaderPlatform()).IsMobileMultiViewEnabled();
-		AddColorTarget(DBufferTextures.DBufferA, LoadAction, DBufferTextures.DBufferATexArray, bIsMobileMultiView);
-		AddColorTarget(DBufferTextures.DBufferB, LoadAction, DBufferTextures.DBufferBTexArray, bIsMobileMultiView);
-		AddColorTarget(DBufferTextures.DBufferC, LoadAction, DBufferTextures.DBufferCTexArray, bIsMobileMultiView);
+		AddColorTarget(DBufferTextures.DBufferA, LoadAction, DBufferTextures.DBufferATexArray, bUseTextureArrays);
+		AddColorTarget(DBufferTextures.DBufferB, LoadAction, DBufferTextures.DBufferBTexArray, bUseTextureArrays);
+		AddColorTarget(DBufferTextures.DBufferC, LoadAction, DBufferTextures.DBufferCTexArray, bUseTextureArrays);
 
 		if (DBufferTextures.DBufferMask)
 		{
diff --git a/Engine/Source/Runtime/Renderer/Private/DBufferTextures.cpp b/Engine/Source/Runtime/Renderer/Private/DBufferTextures.cpp
--- a/Engine/Source/Runtime/Renderer/Private/DBufferTextures.cpp
+++ b/Engine/Source/Runtime/Renderer/Private/DBufferTextures.cpp
@@ -80,7 +80,7 @@ FDBufferTexturesDesc GetDBufferTexturesDesc(FIntPoint Extent, EShaderPlatform Sh
 	return DBufferTexturesDesc;
 }
 
-FDBufferTextures CreateDBufferTextures(FRDGBuilder& GraphBuilder, FIntPoint Extent, EShaderPlatform ShaderPlatform)
+FDBufferTextures CreateDBufferTextures(FRDGBuilder& GraphBuilder, FIntPoint Extent, EShaderPlatform ShaderPlatform, const bool bIsMobileMultiView)
 {
 	FDBufferTextures DBufferTextures;
 
@@ -92,14 +92,18 @@ FDBufferTextures CreateDBufferTextures(FRDGBuilder& GraphBuilder, FIntPoint Exte
 		const ERDGTextureFlags TextureFlags = DBufferMaskTechnique != EDecalDBufferMaskTechnique::Disabled
 			? ERDGTextureFlags::MaintainCompression
 			: ERDGTextureFlags::None;
-				
-		DBufferTextures.DBufferA = GraphBuilder.CreateTexture(TexturesDesc.DBufferADesc, TEXT("DBufferA"), TextureFlags);
-		DBufferTextures.DBufferB = GraphBuilder.CreateTexture(TexturesDesc.DBufferBDesc, TEXT("DBufferB"), TextureFlags);
-		DBufferTextures.DBufferC = GraphBuilder.CreateTexture(TexturesDesc.DBufferCDesc, TEXT("DBufferC"), TextureFlags);
-
-		DBufferTextures.DBufferATexArray = GraphBuilder.CreateTexture(TexturesDesc.DBufferATexArrayDesc, TEXT("DBufferATexArray"), TextureFlags);
-		DBufferTextures.DBufferBTexArray = GraphBuilder.CreateTexture(TexturesDesc.DBufferBTexArrayDesc, TEXT("DBufferBTexArray"), TextureFlags);
-		DBufferTextures.DBufferCTexArray = GraphBuilder.CreateTexture(TexturesDesc.DBufferCTexArrayDesc, TEXT("DBufferCTexArray"), TextureFlags);
+		if (!bIsMobileMultiView)
+		{
+			DBufferTextures.DBufferA = GraphBuilder.CreateTexture(TexturesDesc.DBufferADesc, TEXT("DBufferA"), TextureFlags);
+			DBufferTextures.DBufferB = GraphBuilder.CreateTexture(TexturesDesc.DBufferBDesc, TEXT("DBufferB"), TextureFlags);
+			DBufferTextures.DBufferC = GraphBuilder.CreateTexture(TexturesDesc.DBufferCDesc, TEXT("DBufferC"), TextureFlags);
+		}
+		else
+		{
+			DBufferTextures.DBufferATexArray = GraphBuilder.CreateTexture(TexturesDesc.DBufferATexArrayDesc, TEXT("DBufferATexArray"), TextureFlags);
+			DBufferTextures.DBufferBTexArray = GraphBuilder.CreateTexture(TexturesDesc.DBufferBTexArrayDesc, TEXT("DBufferBTexArray"), TextureFlags);
+			DBufferTextures.DBufferCTexArray = GraphBuilder.CreateTexture(TexturesDesc.DBufferCTexArrayDesc, TEXT("DBufferCTexArray"), TextureFlags);
+		}
 
 		if (DBufferMaskTechnique == EDecalDBufferMaskTechnique::PerPixel)
 		{
diff --git a/Engine/Source/Runtime/Renderer/Private/DBufferTextures.h b/Engine/Source/Runtime/Renderer/Private/DBufferTextures.h
--- a/Engine/Source/Runtime/Renderer/Private/DBufferTextures.h
+++ b/Engine/Source/Runtime/Renderer/Private/DBufferTextures.h
@@ -45,7 +45,7 @@ struct FDBufferTextures
 };
 
 FDBufferTexturesDesc GetDBufferTexturesDesc(FIntPoint Extent, EShaderPlatform ShaderPlatform);
-FDBufferTextures CreateDBufferTextures(FRDGBuilder& GraphBuilder, FIntPoint Extent, EShaderPlatform ShaderPlatform);
+FDBufferTextures CreateDBufferTextures(FRDGBuilder& GraphBuilder, FIntPoint Extent, EShaderPlatform ShaderPlatform, const bool bIsMobileMultiView = false);
 
 BEGIN_SHADER_PARAMETER_STRUCT(FDBufferParameters, )
 	SHADER_PARAMETER_RDG_TEXTURE(Texture2D, DBufferATexture)
diff --git a/Engine/Source/Runtime/Renderer/Private/MobileShadingRenderer.cpp b/Engine/Source/Runtime/Renderer/Private/MobileShadingRenderer.cpp
--- a/Engine/Source/Runtime/Renderer/Private/MobileShadingRenderer.cpp
+++ b/Engine/Source/Runtime/Renderer/Private/MobileShadingRenderer.cpp
@@ -1322,7 +1322,9 @@ void FMobileSceneRenderer::Render(FRDGBuilder& GraphBuilder)
 		{
 			if (bRequiresDBufferDecals)
 			{
-				DBufferTextures = CreateDBufferTextures(GraphBuilder, SceneTextures.Config.Extent, ShaderPlatform);
+				FViewInfo* MainView = Views.Num() > 0 ? &Views[0] : nullptr;
+				bool bIsMobileMultiView = SceneTextures.Config.bRequireMultiView || (MainView && MainView->Aspects.IsMobileMultiViewEnabled());
+				DBufferTextures = CreateDBufferTextures(GraphBuilder, SceneTextures.Config.Extent, ShaderPlatform, bIsMobileMultiView);
 				RenderDBuffer(GraphBuilder, SceneTextures, DBufferTextures, InstanceCullingManager);
 			}
 		}
