Optimized light loop for volumetric fog.

diff --git a/Engine/Source/Runtime/Renderer/Private/LightSceneInfo.cpp b/Engine/Source/Runtime/Renderer/Private/LightSceneInfo.cpp
--- a/Engine/Source/Runtime/Renderer/Private/LightSceneInfo.cpp
+++ b/Engine/Source/Runtime/Renderer/Private/LightSceneInfo.cpp
@@ -59,6 +59,7 @@ FLightSceneInfo::FLightSceneInfo(FLightSceneProxy* InProxy, bool InbVisible)
 	, DynamicInteractionOftenMovingPrimitiveList(NULL)
 	, DynamicInteractionStaticPrimitiveList(NULL)
 	, Proxy(InProxy)
+	, Type((ELightComponentType)InProxy->GetLightType())
 	, Id(INDEX_NONE)
 	, DynamicShadowMapChannel(-1)
 	, bPrecomputedLightingIsValid(InProxy->GetLightComponent()->IsPrecomputedLightingValid())
@@ -207,8 +208,6 @@ bool FLightSceneInfo::ShouldRenderLight(const FViewInfo& View, bool bOffscreen)
 	bool bLocalVisible = bVisible && (bOffscreen ? View.VisibleLightInfos[Id].bInDrawRange : View.VisibleLightInfos[Id].bInViewFrustum);
 
 #if !UE_BUILD_SHIPPING
-	ELightComponentType Type = (ELightComponentType)Proxy->GetLightType();
-
 	switch(Type)
 	{
 		case LightType_Directional:
diff --git a/Engine/Source/Runtime/Renderer/Private/LightSceneInfo.h b/Engine/Source/Runtime/Renderer/Private/LightSceneInfo.h
--- a/Engine/Source/Runtime/Renderer/Private/LightSceneInfo.h
+++ b/Engine/Source/Runtime/Renderer/Private/LightSceneInfo.h
@@ -208,6 +208,8 @@ public:
 	/** The light's scene proxy. */
 	FLightSceneProxy* Proxy;
 
+	ELightComponentType Type;
+
 	/** If bVisible == true, this is the index of the primitive in Scene->Lights. */
 	int32 Id;
 
diff --git a/Engine/Source/Runtime/Renderer/Private/SceneRendering.h b/Engine/Source/Runtime/Renderer/Private/SceneRendering.h
--- a/Engine/Source/Runtime/Renderer/Private/SceneRendering.h
+++ b/Engine/Source/Runtime/Renderer/Private/SceneRendering.h
@@ -2176,23 +2176,23 @@ public:
 		const FSceneTextures& SceneTextures,
 		FIntVector VolumetricFogGridSize,
 		float VolumetricFogMaxDistance,
+		FLightSceneInfo* DirectionalLightSceneInfo,
 		FMatrix44f& OutLightFunctionTranslatedWorldToShadow,
-		FRDGTexture*& OutLightFunctionTexture,
-		bool& bOutUseDirectionalLightShadowing,
-		FLightSceneInfo*& OutDirectionalLightSceneInfo);
+		FRDGTexture*& OutLightFunctionTexture);
 
 	void RenderLocalLightsForVolumetricFog(
 		FRDGBuilder& GraphBuilder,
 		FViewInfo& View,
-		uint32 ViewIndex,
 		bool bUseTemporalReprojection,
 		const struct FVolumetricFogIntegrationParameterData& IntegrationData,
 		const FExponentialHeightFogSceneInfo& FogInfo,
 		FIntVector VolumetricFogGridSize,
 		FVector GridZParams,
 		const FRDGTextureDesc& VolumeDesc,
-		FRDGTexture*& OutLocalShadowedLightScattering,
-		FRDGTextureRef ConservativeDepthTexture);
+		FRDGTextureRef ConservativeDepthTexture,
+		TConstArrayView<const FLightSceneInfo*> LightsToInject,
+		TConstArrayView<const FLightSceneInfo*> RayTracedLightsToInject,
+		FRDGTexture*& OutLocalShadowedLightScattering);
 
 	void SetupVolumetricFog();
 
diff --git a/Engine/Source/Runtime/Renderer/Private/VolumetricFog.cpp b/Engine/Source/Runtime/Renderer/Private/VolumetricFog.cpp
--- a/Engine/Source/Runtime/Renderer/Private/VolumetricFog.cpp
+++ b/Engine/Source/Runtime/Renderer/Private/VolumetricFog.cpp
@@ -28,6 +28,7 @@ VolumetricFog.cpp
 #include "ProfilingDebugging/CpuProfilerTrace.h"
 #include "DataDrivenShaderPlatformInfo.h"
 #include "LightFunctionAtlas.h"
+#include "Math/UnrealMathUtility.h"
 
 using namespace LightFunctionAtlas;
 
@@ -681,50 +682,18 @@ TGlobalResource<FCircleRasterizeIndexBuffer> GCircleRasterizeIndexBuffer;
 
 void FSceneRenderer::RenderLocalLightsForVolumetricFog(
 	FRDGBuilder& GraphBuilder,
-	FViewInfo& View, 
-	uint32 ViewIndex,
+	FViewInfo& View,
 	bool bUseTemporalReprojection,
 	const FVolumetricFogIntegrationParameterData& IntegrationData,
 	const FExponentialHeightFogSceneInfo& FogInfo,
 	FIntVector VolumetricFogViewGridSize,
 	FVector GridZParams,
 	const FRDGTextureDesc& VolumeDesc,
-	FRDGTexture*& OutLocalShadowedLightScattering,
-	FRDGTextureRef ConservativeDepthTexture)
+	FRDGTextureRef ConservativeDepthTexture,
+	TConstArrayView<const FLightSceneInfo*> LightsToInject,
+	TConstArrayView<const FLightSceneInfo*> RayTracedLightsToInject,
+	FRDGTexture*& OutLocalShadowedLightScattering)
 {
-	// Gather lights that need to be rendered with shadow from opaque or light functions.
-	TArray<const FLightSceneInfo*, SceneRenderingAllocator> LightsToInject;
-	TArray<const FLightSceneInfo*, SceneRenderingAllocator> RayTracedLightsToInject;
-	for (auto LightIt = Scene->Lights.CreateConstIterator(); LightIt; ++LightIt)
-	{
-		const FLightSceneInfoCompact& LightSceneInfoCompact = *LightIt;
-		const FLightSceneInfo* LightSceneInfo = LightSceneInfoCompact.LightSceneInfo;
-
-		bool bIsShadowed = LightNeedsSeparateInjectionIntoVolumetricFogForOpaqueShadow(View, LightSceneInfo, VisibleLightInfos[LightSceneInfo->Id]);
-		bool bUsesRectLightTexture = GVolumetricFogRectLightTexture && LightSceneInfo->Proxy->HasSourceTexture();
-
-		if (LightSceneInfo->ShouldRenderLightViewIndependent()
-			&& LightSceneInfo->ShouldRenderLight(View)
-			&& (bIsShadowed || bUsesRectLightTexture)
-			&& LightSceneInfo->Proxy->GetVolumetricScatteringIntensity() > 0)
-		{
-			const FSphere LightBounds = LightSceneInfo->Proxy->GetBoundingSphere();
-
-			if ((View.ViewMatrices.GetViewOrigin() - LightBounds.Center).SizeSquared() < (FogInfo.VolumetricFogDistance + LightBounds.W) * (FogInfo.VolumetricFogDistance + LightBounds.W))
-			{
-				const bool bRayTracedLight = LightHasRayTracedShadows(LightSceneInfo);
-				if (bRayTracedLight)
-				{
-					RayTracedLightsToInject.Add(LightSceneInfo);
-				}
-				else
-				{
-					LightsToInject.Add(LightSceneInfo);
-				}
-			}
-		}
-	}
-
 	// Setup the light function atlas
 	const bool bUseLightFunctionAtlas = LightFunctionAtlas::IsEnabled(View, ELightFunctionAtlasSystem::VolumetricFog);
 	TRDGUniformBufferRef<FLightFunctionAtlasGlobalParameters> LightFunctionAtlasGlobalParameters = LightFunctionAtlas::BindGlobalParameters(GraphBuilder, View);
@@ -1358,9 +1327,86 @@ void FSceneRenderer::ComputeVolumetricFog(FRDGBuilder& GraphBuilder,
 	RDG_CSV_STAT_EXCLUSIVE_SCOPE(GraphBuilder, VolumetricFog);
 	RDG_GPU_STAT_SCOPE(GraphBuilder, VolumetricFog);
 
+	// Gather lights that need to be rendered with shadow from opaque or light functions.
+
+	struct FLightsToInject
+	{
+		TArray<const FLightSceneInfo*, SceneRenderingAllocator> Lights;
+		TArray<const FLightSceneInfo*, SceneRenderingAllocator> RayTracedLights;
+		FLightSceneInfo* DirectionalLightFunction = nullptr;
+		bool bUseDirectionalLightShadowing = false;
+	};
+
+	TArray<FLightsToInject, TInlineAllocator<2>> LightsToInjectPerView;
+	LightsToInjectPerView.SetNum(Views.Num());
+
+	for (auto LightIt = Scene->Lights.CreateConstIterator(); LightIt; ++LightIt)
+	{
+		const FLightSceneInfoCompact& LightSceneInfoCompact = *LightIt;
+
+		// Nothing to do for black lights.
+		if (LightSceneInfoCompact.Color.IsAlmostBlack())
+		{
+			continue;
+		}
+
+		FLightSceneInfo* LightSceneInfo = LightSceneInfoCompact.LightSceneInfo;
+
+		for (int32 ViewIndex = 0; ViewIndex < Views.Num(); ViewIndex++)
+		{
+			FViewInfo& View = Views[ViewIndex];
+			FLightsToInject& LightsToInject = LightsToInjectPerView[ViewIndex];
+
+			if (!LightSceneInfo->ShouldRenderLight(View))
+			{
+				continue;
+			}
+
+			const FLightSceneProxy* LightSceneProxy = LightSceneInfo->Proxy;
+
+			if (LightSceneProxy->GetVolumetricScatteringIntensity() > SMALL_NUMBER)
+			{
+				bool bIsShadowed = LightNeedsSeparateInjectionIntoVolumetricFogForOpaqueShadow(View, LightSceneInfo, VisibleLightInfos[LightSceneInfo->Id]);
+				bool bUsesRectLightTexture = GVolumetricFogRectLightTexture && LightSceneProxy->HasSourceTexture();
+
+				if (bIsShadowed || bUsesRectLightTexture)
+				{
+					const FSphere LightBounds = LightSceneProxy->GetBoundingSphere();
+
+					if ((View.ViewMatrices.GetViewOrigin() - LightBounds.Center).SizeSquared() < (FogInfo.VolumetricFogDistance + LightBounds.W) * (FogInfo.VolumetricFogDistance + LightBounds.W))
+					{
+						const bool bRayTracedLight = LightHasRayTracedShadows(LightSceneInfo);
+						if (bRayTracedLight)
+						{
+							LightsToInject.RayTracedLights.Add(LightSceneInfo);
+						}
+						else
+						{
+							LightsToInject.Lights.Add(LightSceneInfo);
+						}
+					}
+				}
+			}
+
+			// The only directional light we can accept in the volumetric fog because we use the forward lighting data in the Scattering compute shader.
+			const FLightSceneProxy* SelectedForwardDirectionalLightProxy = View.ForwardLightingResources.SelectedForwardDirectionalLightProxy;
+
+			if (LightSceneProxy == SelectedForwardDirectionalLightProxy && LightSceneProxy->GetLightType() == LightType_Directional)
+			{
+				LightsToInject.bUseDirectionalLightShadowing = true;
+
+				if (CheckForLightFunction(LightSceneInfo) && ViewFamily.EngineShowFlags.LightFunctions)
+				{
+					LightsToInject.DirectionalLightFunction = LightSceneInfo;
+				}
+			}
+		}
+	}
+
 	for (int32 ViewIndex = 0; ViewIndex < Views.Num(); ViewIndex++)
 	{
 		FViewInfo& View = Views[ViewIndex];
+		const FLightsToInject& LightsToInject = LightsToInjectPerView[ViewIndex];
 		RDG_GPU_MASK_SCOPE(GraphBuilder, View.GPUMask);
 
 		// Allocate texture using scene render targets size so we do not reallocate every frame when dynamic resolution is used in order to avoid resources allocation hitches.
@@ -1396,7 +1442,8 @@ void FSceneRenderer::ComputeVolumetricFog(FRDGBuilder& GraphBuilder,
 			&& ViewFamily.bRealtimeUpdate
 			&& View.ViewState->LightScatteringHistory;
 
-		FMatrix44f DirectionalLightFunctionTranslatedWorldToShadow;
+		FMatrix44f DirectionalLightFunctionTranslatedWorldToShadow = FMatrix44f::Identity;
+
 
 		RDG_EVENT_SCOPE(GraphBuilder, "VolumetricFog");
 
@@ -1446,19 +1493,20 @@ void FSceneRenderer::ComputeVolumetricFog(FRDGBuilder& GraphBuilder,
 				AtmosphericDirectionalLightIndex = 1;
 			}
 		}
-		
-		FLightSceneInfo* DirectionalLightSceneInfo = nullptr;
-		RenderLightFunctionForVolumetricFog(
-			GraphBuilder,
-			View,
-			SceneTextures,
-			VolumetricFogViewGridSize,
-			FogInfo.VolumetricFogDistance,
-			DirectionalLightFunctionTranslatedWorldToShadow,
-			DirectionalLightFunctionTexture,
-			bUseDirectionalLightShadowing,
-			DirectionalLightSceneInfo);
-			
+
+		if (LightsToInject.DirectionalLightFunction)
+		{
+			RenderLightFunctionForVolumetricFog(
+				GraphBuilder,
+				View,
+				SceneTextures,
+				VolumetricFogViewGridSize,
+				FogInfo.VolumetricFogDistance,
+				LightsToInject.DirectionalLightFunction,
+				DirectionalLightFunctionTranslatedWorldToShadow,
+				DirectionalLightFunctionTexture);
+		}
+
 		View.VolumetricFogResources.IntegratedLightScatteringTexture = nullptr;
 		TRDGUniformBufferRef<FFogUniformParameters> FogUniformBuffer = CreateFogUniformBuffer(GraphBuilder, View);
 
@@ -1478,8 +1526,8 @@ void FSceneRenderer::ComputeVolumetricFog(FRDGBuilder& GraphBuilder,
 		}
 
 		FRDGTexture* LocalShadowedLightScattering = GraphBuilder.RegisterExternalTexture(GSystemTextures.VolumetricBlackDummy);
-		RenderLocalLightsForVolumetricFog(GraphBuilder, View, uint32(ViewIndex), bUseTemporalReprojection, IntegrationData, FogInfo,
-			VolumetricFogViewGridSize, GridZParams, VolumeDescFastVRAM, LocalShadowedLightScattering, ConservativeDepthTexture);
+		RenderLocalLightsForVolumetricFog(GraphBuilder, View, bUseTemporalReprojection, IntegrationData, FogInfo,
+			VolumetricFogViewGridSize, GridZParams, VolumeDescFastVRAM, ConservativeDepthTexture, LightsToInject.Lights, LightsToInject.RayTracedLights, LocalShadowedLightScattering);
 	
 		FRDGTextureRef RaytracedShadowsVolume = nullptr;
 #if RHI_RAYTRACING		
@@ -1664,8 +1712,8 @@ void FSceneRenderer::ComputeVolumetricFog(FRDGBuilder& GraphBuilder,
 			PassParameters->LightFunctionAtlas = LightFunctionAtlas::BindGlobalParameters(GraphBuilder, View);
 			if (bUseLightFunctionAtlas)
 			{
-				PassParameters->DirectionalApplyLightFunctionFromAtlas = bUseLightFunctionAtlas && DirectionalLightSceneInfo && DirectionalLightSceneInfo->Proxy->HasValidLightFunctionAtlasSlot() ? 1 :0;
-				PassParameters->DirectionalLightFunctionAtlasLightIndex = DirectionalLightSceneInfo ? DirectionalLightSceneInfo->Proxy->GetLightFunctionAtlasLightIndex() : 0;
+				PassParameters->DirectionalApplyLightFunctionFromAtlas = bUseLightFunctionAtlas && LightsToInject.DirectionalLightFunction && LightsToInject.DirectionalLightFunction->Proxy->HasValidLightFunctionAtlasSlot() ? 1 :0;
+				PassParameters->DirectionalLightFunctionAtlasLightIndex = LightsToInject.DirectionalLightFunction ? LightsToInject.DirectionalLightFunction->Proxy->GetLightFunctionAtlasLightIndex() : 0;
 			}
 			else
 			{
diff --git a/Engine/Source/Runtime/Renderer/Private/VolumetricFogLightFunction.cpp b/Engine/Source/Runtime/Renderer/Private/VolumetricFogLightFunction.cpp
--- a/Engine/Source/Runtime/Renderer/Private/VolumetricFogLightFunction.cpp
+++ b/Engine/Source/Runtime/Renderer/Private/VolumetricFogLightFunction.cpp
@@ -114,187 +114,147 @@ void FSceneRenderer::RenderLightFunctionForVolumetricFog(
 	const FSceneTextures& SceneTextures,
 	FIntVector VolumetricFogGridSize,
 	float VolumetricFogMaxDistance,
+	FLightSceneInfo* DirectionalLightSceneInfo,
 	FMatrix44f& OutDirectionalLightFunctionTranslatedWorldToShadow,
-	FRDGTexture*& OutDirectionalLightFunctionTexture,
-	bool& bOutUseDirectionalLightShadowing,
-	FLightSceneInfo*& OutDirectionalLightSceneInfo)
+	FRDGTexture*& OutDirectionalLightFunctionTexture)
 {
-	// The only directional light we can accept in the volumetric fog because we use the forward lighting data in the Scattering compute shader.
-	const FLightSceneProxy* SelectedForwardDirectionalLightProxy = View.ForwardLightingResources.SelectedForwardDirectionalLightProxy;
-	// Default directional light properties
-	OutDirectionalLightFunctionTranslatedWorldToShadow = FMatrix44f::Identity;
-	bOutUseDirectionalLightShadowing = false;
-	OutDirectionalLightSceneInfo = NULL;
-
-	// Gather lights that need to evaluate light functions
-	for (auto LightIt = Scene->Lights.CreateConstIterator(); LightIt; ++LightIt)
+	RDG_EVENT_SCOPE(GraphBuilder, "DirectionalLightFunction");
+
+	// Estimate the resolution and the projection matrix.
+	FProjectedShadowInfo ProjectedShadowInfo;
+	FIntPoint LightFunctionResolution;
 	{
-		const FLightSceneInfoCompact& LightSceneInfoCompact = *LightIt;
-		FLightSceneInfo* LightSceneInfo = LightSceneInfoCompact.LightSceneInfo;
-
-		if (   LightSceneInfo
-			&& LightSceneInfo->Proxy
-			// Band-aid fix for extremely rare case that light scene proxy contains NaNs.
-			&& !LightSceneInfo->Proxy->GetDirection().ContainsNaN()
-			&& LightSceneInfo->ShouldRenderLightViewIndependent()
-			&& LightSceneInfo->ShouldRenderLight(View))
+		const FVector ViewForward = View.ViewMatrices.GetOverriddenTranslatedViewMatrix().GetColumn(2);
+		const FVector ViewUp = View.ViewMatrices.GetOverriddenTranslatedViewMatrix().GetColumn(1);
+		const FVector ViewRight = View.ViewMatrices.GetOverriddenTranslatedViewMatrix().GetColumn(0);
+
+		const FVector LightDirection = DirectionalLightSceneInfo->Proxy->GetDirection().GetSafeNormal();
+		FVector AxisWeights;
+		AxisWeights.X = FMath::Abs(LightDirection | ViewRight) * VolumetricFogGridSize.X;
+		AxisWeights.Y = FMath::Abs(LightDirection | ViewUp) * VolumetricFogGridSize.Y;
+		AxisWeights.Z = FMath::Abs(LightDirection | ViewForward) * VolumetricFogGridSize.Z;
+
+		const float VolumeResolutionEstimate = FMath::Max(AxisWeights.X, FMath::Max(AxisWeights.Y, AxisWeights.Z)) * GVolumetricFogLightFunctionDirectionalLightSupersampleScale;
+		LightFunctionResolution = FIntPoint(FMath::TruncToInt(VolumeResolutionEstimate), FMath::TruncToInt(VolumeResolutionEstimate));
+
+		// Snap the resolution to allow render target pool hits most of the time
+		const int32 ResolutionSnapFactor = 32;
+		LightFunctionResolution.X = FMath::DivideAndRoundUp(LightFunctionResolution.X, ResolutionSnapFactor) * ResolutionSnapFactor;
+		LightFunctionResolution.Y = FMath::DivideAndRoundUp(LightFunctionResolution.Y, ResolutionSnapFactor) * ResolutionSnapFactor;
+
+		// Guard against invalid resolutions
+		const uint32 LiFuncResXU32 = (uint32)LightFunctionResolution.X;
+		const uint32 LiFuncResYU32 = (uint32)LightFunctionResolution.Y;
+		const uint32 MaxTextureResU32 = GMaxTextureDimensions;
+		if (LiFuncResXU32 > MaxTextureResU32 || LiFuncResYU32 > MaxTextureResU32)
 		{
-			if (OutDirectionalLightSceneInfo == NULL && LightSceneInfo->Proxy == SelectedForwardDirectionalLightProxy &&
-				LightSceneInfo->Proxy->GetLightType() == LightType_Directional)
-			{
-				// We only take the first directional light into account
-				bOutUseDirectionalLightShadowing = LightSceneInfo->Proxy->CastsVolumetricShadow();
-
-				if (CheckForLightFunction(LightSceneInfo) && ViewFamily.EngineShowFlags.LightFunctions)
-				{
-					OutDirectionalLightSceneInfo = LightSceneInfo;
-					break;	// we only deal with the directional light function so we directly stop now
-				}
-			}
+#if !(UE_BUILD_SHIPPING)
+			UE_LOG(LogRenderer, Error,
+				TEXT("Invalid LightFunctionResolution %dx%d, View={ %s}, LightDirection={ %s }"),
+				LightFunctionResolution.X,
+				LightFunctionResolution.Y,
+				*View.ViewMatrices.GetOverriddenTranslatedViewMatrix().ToString(),
+				*LightDirection.ToString());
+#endif
+			const uint32 ClampedRes = FMath::Max(FMath::Min3(LiFuncResXU32, LiFuncResYU32, MaxTextureResU32), (uint32)ResolutionSnapFactor);
+			LightFunctionResolution.X = ClampedRes;
+			LightFunctionResolution.Y = ClampedRes;
 		}
+
+		FWholeSceneProjectedShadowInitializer ShadowInitializer;
+
+		check(VolumetricFogMaxDistance > 0);
+		FSphere Bounds = DirectionalLightSceneInfo->Proxy->GetShadowSplitBoundsDepthRange(View, View.ViewMatrices.GetViewOrigin(), 0, VolumetricFogMaxDistance, NULL);
+		check(Bounds.W > 0);
+
+		const float ShadowExtent = Bounds.W / FMath::Sqrt(3.0f);
+		const FBoxSphereBounds SubjectBounds(Bounds.Center, FVector(ShadowExtent, ShadowExtent, ShadowExtent), Bounds.W);
+		ShadowInitializer.PreShadowTranslation = -Bounds.Center;
+		ShadowInitializer.WorldToLight = FInverseRotationMatrix(LightDirection.Rotation());
+		ShadowInitializer.Scales = FVector2D(1.0f / Bounds.W, 1.0f / Bounds.W);
+		ShadowInitializer.SubjectBounds = FBoxSphereBounds(FVector::ZeroVector, SubjectBounds.BoxExtent, SubjectBounds.SphereRadius);
+		ShadowInitializer.WAxis = FVector4(0, 0, 0, 1);
+		ShadowInitializer.MinLightW = -HALF_WORLD_MAX;
+		// Reduce casting distance on a directional light
+		// This is necessary to improve floating point precision in several places, especially when deriving frustum verts from InvReceiverMatrix
+		ShadowInitializer.MaxDistanceToCastInLightW = HALF_WORLD_MAX / 32.0f;
+		ShadowInitializer.bRayTracedDistanceField = false;
+		ShadowInitializer.CascadeSettings.bFarShadowCascade = false;
+
+		ProjectedShadowInfo.SetupWholeSceneProjection(
+			DirectionalLightSceneInfo,
+			&View,
+			ShadowInitializer,
+			LightFunctionResolution.X,
+			LightFunctionResolution.Y,
+			LightFunctionResolution.X,
+			LightFunctionResolution.Y,
+			0
+		);
+
+		OutDirectionalLightFunctionTranslatedWorldToShadow = ProjectedShadowInfo.TranslatedWorldToClipInnerMatrix;
 	}
 
-	// Now bake the directional light function into a 2d transient texture for the special single directional light we have selected
-	if (OutDirectionalLightSceneInfo)
+	// Now render the texture
 	{
-		RDG_EVENT_SCOPE(GraphBuilder, "DirectionalLightFunction");
+		FRDGTextureDesc LightFunctionTextureDesc = FRDGTextureDesc::Create2D(LightFunctionResolution, PF_G8, FClearValueBinding::None, TexCreate_ShaderResource | TexCreate_RenderTargetable);
+		LightFunctionTextureDesc.Flags |= GFastVRamConfig.VolumetricFog;
 
-		// Estimate the resolution and the projection matrix.
-		FProjectedShadowInfo ProjectedShadowInfo;
-		FIntPoint LightFunctionResolution;
-		{
-			const FVector ViewForward = View.ViewMatrices.GetOverriddenTranslatedViewMatrix().GetColumn(2);
-			const FVector ViewUp = View.ViewMatrices.GetOverriddenTranslatedViewMatrix().GetColumn(1);
-			const FVector ViewRight = View.ViewMatrices.GetOverriddenTranslatedViewMatrix().GetColumn(0);
-
-			const FVector LightDirection = OutDirectionalLightSceneInfo->Proxy->GetDirection().GetSafeNormal();
-			FVector AxisWeights;
-			AxisWeights.X = FMath::Abs(LightDirection | ViewRight) * VolumetricFogGridSize.X;
-			AxisWeights.Y = FMath::Abs(LightDirection | ViewUp) * VolumetricFogGridSize.Y;
-			AxisWeights.Z = FMath::Abs(LightDirection | ViewForward) * VolumetricFogGridSize.Z;
-
-			const float VolumeResolutionEstimate = FMath::Max(AxisWeights.X, FMath::Max(AxisWeights.Y, AxisWeights.Z)) * GVolumetricFogLightFunctionDirectionalLightSupersampleScale;
-			LightFunctionResolution = FIntPoint(FMath::TruncToInt(VolumeResolutionEstimate), FMath::TruncToInt(VolumeResolutionEstimate));
-
-			// Snap the resolution to allow render target pool hits most of the time
-			const int32 ResolutionSnapFactor = 32;
-			LightFunctionResolution.X = FMath::DivideAndRoundUp(LightFunctionResolution.X, ResolutionSnapFactor) * ResolutionSnapFactor;
-			LightFunctionResolution.Y = FMath::DivideAndRoundUp(LightFunctionResolution.Y, ResolutionSnapFactor) * ResolutionSnapFactor;
-
-			// Guard against invalid resolutions
-			const uint32 LiFuncResXU32 = (uint32)LightFunctionResolution.X;
-			const uint32 LiFuncResYU32 = (uint32)LightFunctionResolution.Y;
-			const uint32 MaxTextureResU32 = GMaxTextureDimensions;
-			if (LiFuncResXU32 > MaxTextureResU32 || LiFuncResYU32 > MaxTextureResU32)
-			{
-#if !(UE_BUILD_SHIPPING)
-				UE_LOG(LogRenderer, Error,
-					TEXT("Invalid LightFunctionResolution %dx%d, View={ %s}, LightDirection={ %s }"),
-					LightFunctionResolution.X,
-					LightFunctionResolution.Y,
-					*View.ViewMatrices.GetOverriddenTranslatedViewMatrix().ToString(),
-					*LightDirection.ToString());
-#endif
-				const uint32 ClampedRes = FMath::Max(FMath::Min3(LiFuncResXU32, LiFuncResYU32, MaxTextureResU32), (uint32)ResolutionSnapFactor);
-				LightFunctionResolution.X = ClampedRes;
-				LightFunctionResolution.Y = ClampedRes;
-			}
+		FRDGTexture* LightFunctionTexture = GraphBuilder.CreateTexture(LightFunctionTextureDesc, TEXT("VolumetricFog.LightFunction"));
+		OutDirectionalLightFunctionTexture = LightFunctionTexture;
 
-			FWholeSceneProjectedShadowInitializer ShadowInitializer;
-
-			check(VolumetricFogMaxDistance > 0);
-			FSphere Bounds = OutDirectionalLightSceneInfo->Proxy->GetShadowSplitBoundsDepthRange(View, View.ViewMatrices.GetViewOrigin(), 0, VolumetricFogMaxDistance, NULL);
-			check(Bounds.W > 0);
-
-			const float ShadowExtent = Bounds.W / FMath::Sqrt(3.0f);
-			const FBoxSphereBounds SubjectBounds(Bounds.Center, FVector(ShadowExtent, ShadowExtent, ShadowExtent), Bounds.W);
-			ShadowInitializer.PreShadowTranslation = -Bounds.Center;
-			ShadowInitializer.WorldToLight = FInverseRotationMatrix(LightDirection.Rotation());
-			ShadowInitializer.Scales = FVector2D(1.0f / Bounds.W, 1.0f / Bounds.W);
-			ShadowInitializer.SubjectBounds = FBoxSphereBounds(FVector::ZeroVector, SubjectBounds.BoxExtent, SubjectBounds.SphereRadius);
-			ShadowInitializer.WAxis = FVector4(0, 0, 0, 1);
-			ShadowInitializer.MinLightW = -HALF_WORLD_MAX;
-			// Reduce casting distance on a directional light
-			// This is necessary to improve floating point precision in several places, especially when deriving frustum verts from InvReceiverMatrix
-			ShadowInitializer.MaxDistanceToCastInLightW = HALF_WORLD_MAX / 32.0f;
-			ShadowInitializer.bRayTracedDistanceField = false;
-			ShadowInitializer.CascadeSettings.bFarShadowCascade = false;
-
-			ProjectedShadowInfo.SetupWholeSceneProjection(
-				OutDirectionalLightSceneInfo,
-				&View,
-				ShadowInitializer,
-				LightFunctionResolution.X,
-				LightFunctionResolution.Y,
-				LightFunctionResolution.X,
-				LightFunctionResolution.Y,
-				0
-			);
+		const FMaterialRenderProxy* MaterialProxyForRendering = DirectionalLightSceneInfo->Proxy->GetLightFunctionMaterial();
+		const FMaterial& Material = MaterialProxyForRendering->GetMaterialWithFallback(Scene->GetFeatureLevel(), MaterialProxyForRendering);
 
-			OutDirectionalLightFunctionTranslatedWorldToShadow = ProjectedShadowInfo.TranslatedWorldToClipInnerMatrix;
-		}
+		FVolumetricFogLightFunctionParameters* PassParameters = GraphBuilder.AllocParameters<FVolumetricFogLightFunctionParameters>();
+		PassParameters->RenderTargets[0] = FRenderTargetBinding(LightFunctionTexture, ERenderTargetLoadAction::ENoAction);
+		PassParameters->SceneTextures = SceneTextures.UniformBuffer;
 
-		// Now render the texture
-		{
-			FRDGTextureDesc LightFunctionTextureDesc = FRDGTextureDesc::Create2D(LightFunctionResolution, PF_G8, FClearValueBinding::None, TexCreate_ShaderResource | TexCreate_RenderTargetable);
-			LightFunctionTextureDesc.Flags |= GFastVRamConfig.VolumetricFog;
-
-			FRDGTexture* LightFunctionTexture = GraphBuilder.CreateTexture(LightFunctionTextureDesc, TEXT("VolumetricFog.LightFunction"));
-			OutDirectionalLightFunctionTexture = LightFunctionTexture;
-
-			const FMaterialRenderProxy* MaterialProxyForRendering = OutDirectionalLightSceneInfo->Proxy->GetLightFunctionMaterial();
-			const FMaterial& Material = MaterialProxyForRendering->GetMaterialWithFallback(Scene->GetFeatureLevel(), MaterialProxyForRendering);
-
-			FVolumetricFogLightFunctionParameters* PassParameters = GraphBuilder.AllocParameters<FVolumetricFogLightFunctionParameters>();
-			PassParameters->RenderTargets[0] = FRenderTargetBinding(LightFunctionTexture, ERenderTargetLoadAction::ENoAction);
-			PassParameters->SceneTextures = SceneTextures.UniformBuffer;
-
-			FMatrix44f LightFunctionTranslatedWorldToShadowMatrix = OutDirectionalLightFunctionTranslatedWorldToShadow;
-			GraphBuilder.AddPass(
-				RDG_EVENT_NAME("LightFunction %ux%u Material=%s", LightFunctionResolution.X, LightFunctionResolution.Y, *(Material.GetFriendlyName())),
-				PassParameters,
-				ERDGPassFlags::Raster,
-				[PassParameters, &View, MaterialProxyForRendering, &Material, LightFunctionResolution, OutDirectionalLightSceneInfo, LightFunctionTranslatedWorldToShadowMatrix](FRHICommandList& RHICmdList)
-				{
-					RHICmdList.SetViewport(0.f, 0.f, 0.f, LightFunctionResolution.X, LightFunctionResolution.Y, 1.f);
-
-					FGraphicsPipelineStateInitializer GraphicsPSOInit;
-					RHICmdList.ApplyCachedRenderTargets(GraphicsPSOInit);
-					GraphicsPSOInit.RasterizerState = TStaticRasterizerState<FM_Solid, CM_None>::GetRHI();
-					GraphicsPSOInit.BlendState = TStaticBlendState<>::GetRHI();
-					GraphicsPSOInit.DepthStencilState = TStaticDepthStencilState<false, CF_Always>::GetRHI();
-					GraphicsPSOInit.PrimitiveType = PT_TriangleList;
-
-					const FMaterialShaderMap* MaterialShaderMap = Material.GetRenderingThreadShaderMap();
-					TShaderMapRef<FPostProcessVS> VertexShader(View.ShaderMap);
-
-					check(OutDirectionalLightSceneInfo->Proxy->GetLightType() == LightType_Directional)
-					FVolumetricFogLightFunctionPS::FPermutationDomain PermutationVector;
-					PermutationVector.Set<FVolumetricFogLightFunctionPS::FLightType>(3);
-					TShaderRef<FVolumetricFogLightFunctionPS> PixelShader = MaterialShaderMap->GetShader<FVolumetricFogLightFunctionPS>(PermutationVector);
-
-					GraphicsPSOInit.BoundShaderState.VertexDeclarationRHI = GFilterVertexDeclaration.VertexDeclarationRHI;
-					GraphicsPSOInit.BoundShaderState.VertexShaderRHI = VertexShader.GetVertexShader();
-					GraphicsPSOInit.BoundShaderState.PixelShaderRHI = PixelShader.GetPixelShader();
-
-					SetGraphicsPipelineState(RHICmdList, GraphicsPSOInit, 0);
-
-					FVolumetricFogLightFunctionPS::FParameters PS = PixelShader->GetParameters(View, OutDirectionalLightSceneInfo, FVector2D(1.0f / LightFunctionResolution.X, 1.0f / LightFunctionResolution.Y), LightFunctionTranslatedWorldToShadowMatrix.Inverse());
-
-					ClearUnusedGraphResources(PixelShader, &PS);
-					SetShaderParametersMixedPS(RHICmdList, PixelShader, PS, View, MaterialProxyForRendering);
-
-					DrawRectangle(
-						RHICmdList,
-						0, 0,
-						LightFunctionResolution.X, LightFunctionResolution.Y,
-						0, 0,
-						LightFunctionResolution.X, LightFunctionResolution.Y,
-						LightFunctionResolution,
-						LightFunctionResolution,
-						VertexShader);
-				}
-			);
-		}
+		FMatrix44f LightFunctionTranslatedWorldToShadowMatrix = OutDirectionalLightFunctionTranslatedWorldToShadow;
+		GraphBuilder.AddPass(
+			RDG_EVENT_NAME("LightFunction %ux%u Material=%s", LightFunctionResolution.X, LightFunctionResolution.Y, *(Material.GetFriendlyName())),
+			PassParameters,
+			ERDGPassFlags::Raster,
+			[PassParameters, &View, MaterialProxyForRendering, &Material, LightFunctionResolution, DirectionalLightSceneInfo, LightFunctionTranslatedWorldToShadowMatrix](FRHICommandList& RHICmdList)
+			{
+				RHICmdList.SetViewport(0.f, 0.f, 0.f, LightFunctionResolution.X, LightFunctionResolution.Y, 1.f);
+
+				FGraphicsPipelineStateInitializer GraphicsPSOInit;
+				RHICmdList.ApplyCachedRenderTargets(GraphicsPSOInit);
+				GraphicsPSOInit.RasterizerState = TStaticRasterizerState<FM_Solid, CM_None>::GetRHI();
+				GraphicsPSOInit.BlendState = TStaticBlendState<>::GetRHI();
+				GraphicsPSOInit.DepthStencilState = TStaticDepthStencilState<false, CF_Always>::GetRHI();
+				GraphicsPSOInit.PrimitiveType = PT_TriangleList;
+
+				const FMaterialShaderMap* MaterialShaderMap = Material.GetRenderingThreadShaderMap();
+				TShaderMapRef<FPostProcessVS> VertexShader(View.ShaderMap);
+
+				check(DirectionalLightSceneInfo->Proxy->GetLightType() == LightType_Directional)
+				FVolumetricFogLightFunctionPS::FPermutationDomain PermutationVector;
+				PermutationVector.Set<FVolumetricFogLightFunctionPS::FLightType>(3);
+				TShaderRef<FVolumetricFogLightFunctionPS> PixelShader = MaterialShaderMap->GetShader<FVolumetricFogLightFunctionPS>(PermutationVector);
+
+				GraphicsPSOInit.BoundShaderState.VertexDeclarationRHI = GFilterVertexDeclaration.VertexDeclarationRHI;
+				GraphicsPSOInit.BoundShaderState.VertexShaderRHI = VertexShader.GetVertexShader();
+				GraphicsPSOInit.BoundShaderState.PixelShaderRHI = PixelShader.GetPixelShader();
+
+				SetGraphicsPipelineState(RHICmdList, GraphicsPSOInit, 0);
+
+				FVolumetricFogLightFunctionPS::FParameters PS = PixelShader->GetParameters(View, DirectionalLightSceneInfo, FVector2D(1.0f / LightFunctionResolution.X, 1.0f / LightFunctionResolution.Y), LightFunctionTranslatedWorldToShadowMatrix.Inverse());
+
+				ClearUnusedGraphResources(PixelShader, &PS);
+				SetShaderParametersMixedPS(RHICmdList, PixelShader, PS, View, MaterialProxyForRendering);
+
+				DrawRectangle(
+					RHICmdList,
+					0, 0,
+					LightFunctionResolution.X, LightFunctionResolution.Y,
+					0, 0,
+					LightFunctionResolution.X, LightFunctionResolution.Y,
+					LightFunctionResolution,
+					LightFunctionResolution,
+					VertexShader);
+			}
+		);
 	}
 }
