Fixed VSM visualization pass not running outside of editor in non-shipping builds.
Fixed early exit in VSM visualization pass not copying the output texture correctly.
#rnx

diff --git a/Engine/Source/Runtime/Renderer/Private/PostProcess/PostProcessing.cpp b/Engine/Source/Runtime/Renderer/Private/PostProcess/PostProcessing.cpp
--- a/Engine/Source/Runtime/Renderer/Private/PostProcess/PostProcessing.cpp
+++ b/Engine/Source/Runtime/Renderer/Private/PostProcess/PostProcessing.cpp
@@ -498,8 +498,6 @@ void AddPostProcessingPasses(
 	PassSequence.SetEnabled(EPass::VisualizeSkyAtmosphere, Scene&& View.Family && View.Family->EngineShowFlags.VisualizeSkyAtmosphere && ShouldRenderSkyAtmosphereDebugPasses(Scene, View.Family->EngineShowFlags));
 	PassSequence.SetEnabled(EPass::VisualizeSkyLightIlluminanceMeter, Scene&& Scene->SkyLight && View.Family && View.Family->EngineShowFlags.VisualizeSkyLightIlluminance);
 	PassSequence.SetEnabled(EPass::VisualizeLightFunctionAtlas, Scene && Scene->LightFunctionAtlasSceneData.GetLightFunctionAtlasEnabled() && View.Family && View.Family->EngineShowFlags.VisualizeLightFunctionAtlas);
-	PassSequence.SetEnabled(EPass::VisualizeVirtualShadowMaps_PreEditorPrimitives, EngineShowFlags.VisualizeVirtualShadowMap && VirtualShadowMapArray != nullptr);
-	PassSequence.SetEnabled(EPass::VisualizeVirtualShadowMaps_PostEditorPrimitives, EngineShowFlags.VisualizeVirtualShadowMap && VirtualShadowMapArray != nullptr);
 	PassSequence.SetEnabled(EPass::VisualizeLevelInstance, GIsEditor && EngineShowFlags.EditingLevelInstance && EngineShowFlags.VisualizeLevelInstanceEditing && !bVisualizeHDR);
 	PassSequence.SetEnabled(EPass::SelectionOutline, GIsEditor && EngineShowFlags.Selection && EngineShowFlags.SelectionOutline && !EngineShowFlags.Wireframe && !bVisualizeHDR);
 	PassSequence.SetEnabled(EPass::EditorPrimitive, FSceneRenderer::ShouldCompositeEditorPrimitives(View));
@@ -507,12 +505,19 @@ void AddPostProcessingPasses(
 	PassSequence.SetEnabled(EPass::VisualizeSkyAtmosphere, false);
 	PassSequence.SetEnabled(EPass::VisualizeSkyLightIlluminanceMeter, false);
 	PassSequence.SetEnabled(EPass::VisualizeLightFunctionAtlas, false);
-	PassSequence.SetEnabled(EPass::VisualizeVirtualShadowMaps_PreEditorPrimitives, false);
-	PassSequence.SetEnabled(EPass::VisualizeVirtualShadowMaps_PostEditorPrimitives, false);
 	PassSequence.SetEnabled(EPass::VisualizeLevelInstance, false);
 	PassSequence.SetEnabled(EPass::SelectionOutline, false);
 	PassSequence.SetEnabled(EPass::EditorPrimitive, false);
 #endif
+
+#if WITH_EDITOR || !UE_BUILD_SHIPPING
+	PassSequence.SetEnabled(EPass::VisualizeVirtualShadowMaps_PreEditorPrimitives, EngineShowFlags.VisualizeVirtualShadowMap && VirtualShadowMapArray != nullptr);
+	PassSequence.SetEnabled(EPass::VisualizeVirtualShadowMaps_PostEditorPrimitives, EngineShowFlags.VisualizeVirtualShadowMap && VirtualShadowMapArray != nullptr);
+#else
+	PassSequence.SetEnabled(EPass::VisualizeVirtualShadowMaps_PreEditorPrimitives, false);
+	PassSequence.SetEnabled(EPass::VisualizeVirtualShadowMaps_PostEditorPrimitives, false);
+#endif
+
 	PassSequence.SetEnabled(EPass::VisualizeShadingModels, EngineShowFlags.VisualizeShadingModels);
 	PassSequence.SetEnabled(EPass::VisualizeGBufferHints, EngineShowFlags.GBufferHints);
 	PassSequence.SetEnabled(EPass::VisualizeSubsurface, EngineShowFlags.VisualizeSSS);
diff --git a/Engine/Source/Runtime/Renderer/Private/VirtualShadowMaps/VirtualShadowMapArray.cpp b/Engine/Source/Runtime/Renderer/Private/VirtualShadowMaps/VirtualShadowMapArray.cpp
--- a/Engine/Source/Runtime/Renderer/Private/VirtualShadowMaps/VirtualShadowMapArray.cpp
+++ b/Engine/Source/Runtime/Renderer/Private/VirtualShadowMaps/VirtualShadowMapArray.cpp
@@ -3470,19 +3470,30 @@ IMPLEMENT_GLOBAL_SHADER(FDesaturatePS, "/Engine/Private/VirtualShadowMaps/Desatu
 
 FScreenPassTexture FVirtualShadowMapArray::AddVisualizePass(FRDGBuilder& GraphBuilder, const FViewInfo& View, int32 ViewIndex, EVSMVisualizationPostPass Pass, FScreenPassTexture& SceneColor, FScreenPassRenderTarget& OverrideOutput)
 {
-	FScreenPassTexture Output = SceneColor;
+	FScreenPassTexture Output = MoveTemp(SceneColor);
+
+	auto FinalizeOutput = [](FRDGBuilder& GraphBuilder, const FViewInfo& View, FScreenPassTexture& Output, FScreenPassRenderTarget& OverrideOutput) -> FScreenPassTexture&
+	{
+		if (OverrideOutput.IsValid())
+		{
+			AddDrawTexturePass(GraphBuilder, View, Output, OverrideOutput);
+			return OverrideOutput;
+		}
+
+		return Output;
+	};
 
 #if !UE_BUILD_SHIPPING
 
 	if (!IsAllocated() || DebugVisualizationOutput.IsEmpty())
 	{
-		return MoveTemp(SceneColor);
+		return MoveTemp(FinalizeOutput(GraphBuilder, View, Output, OverrideOutput));
 	}
 
 	const FVirtualShadowMapVisualizationData& VisualizationData = GetVirtualShadowMapVisualizationData();
 	if (!VisualizationData.IsActive() || !VisualizeLight[ViewIndex].IsValid())
 	{
-		return MoveTemp(SceneColor);
+		return MoveTemp(FinalizeOutput(GraphBuilder, View, Output, OverrideOutput));
 	}
 
 	RDG_EVENT_SCOPE(GraphBuilder, "VirtualShadowMapsVisualization");
@@ -3588,16 +3599,9 @@ FScreenPassTexture FVirtualShadowMapArray::AddVisualizePass(FRDGBuilder& GraphBu
 			});
 	}
 
-
-	if (OverrideOutput.IsValid())
-	{
-		AddDrawTexturePass(GraphBuilder, View, Output, OverrideOutput);
-		return OverrideOutput;
-	}
-
 #endif //!UE_BUILD_SHIPPING
 
-	return MoveTemp(Output);
+	return MoveTemp(FinalizeOutput(GraphBuilder, View, Output, OverrideOutput));
 }
 
 float FVirtualShadowMapArray::InterpolateResolutionBias(float BiasNonMoving, float BiasMoving, float LightMobilityFactor)
