Fixed custom data upload for dynamic primitives in a mobile GPUScene case

diff --git a/Engine/Source/Runtime/Renderer/Private/GPUScene.cpp b/Engine/Source/Runtime/Renderer/Private/GPUScene.cpp
--- a/Engine/Source/Runtime/Renderer/Private/GPUScene.cpp
+++ b/Engine/Source/Runtime/Renderer/Private/GPUScene.cpp
@@ -2350,7 +2350,7 @@ void FBatchedPrimitiveShaderData::Setup(const FPrimitiveUniformShaderParameters&
 
 	// Set all the custom primitive data float4. This matches the loop in SceneData.ush
 	int32 NumCustomData = FMath::Min<int32>(FCustomPrimitiveData::NumCustomPrimitiveDataFloat4s, DataStrideInFloat4s - i);
-	for (int32 DataIndex = 0; DataIndex < FCustomPrimitiveData::NumCustomPrimitiveDataFloat4s; ++DataIndex)
+	for (int32 DataIndex = 0; DataIndex < NumCustomData; ++DataIndex)
 	{
 		Data[i + DataIndex] = PrimitiveUniformShaderParameters.CustomPrimitiveData[DataIndex];
 	}
