Added UseNanite check to enable skipping Nanite LPIs as otherwise it allocates shadow maps for lights that do not need it (as ShouldSkipNaniteLPIs returned true for non nanite platforms & this causes FProjectedShadowInfo::HasSubjectPrims to return true).

diff --git a/Engine/Source/Runtime/Renderer/Private/PrimitiveSceneInfo.cpp b/Engine/Source/Runtime/Renderer/Private/PrimitiveSceneInfo.cpp
--- a/Engine/Source/Runtime/Renderer/Private/PrimitiveSceneInfo.cpp
+++ b/Engine/Source/Runtime/Renderer/Private/PrimitiveSceneInfo.cpp
@@ -1642,7 +1642,7 @@ void FPrimitiveSceneInfo::AddToScene(FScene* Scene, TArrayView<FPrimitiveSceneIn
 	}
 
 	{
-		const bool bSkipNaniteInOctree = ShouldSkipNaniteLPIs();
+		const bool bSkipNaniteInOctree = ShouldSkipNaniteLPIs(Scene->GetShaderPlatform());
 		SCOPED_NAMED_EVENT(FPrimitiveSceneInfo_AddToScene_AddToPrimitiveOctree, FColor::Red);
 		for (FPrimitiveSceneInfo* SceneInfo : SceneInfos)
 		{
diff --git a/Engine/Source/Runtime/Renderer/Private/RendererPrivateUtils.h b/Engine/Source/Runtime/Renderer/Private/RendererPrivateUtils.h
--- a/Engine/Source/Runtime/Renderer/Private/RendererPrivateUtils.h
+++ b/Engine/Source/Runtime/Renderer/Private/RendererPrivateUtils.h
@@ -10,7 +10,7 @@
 /**
  * Returns true if Nanite primitives are to be skipped in Light Primitive Interactions (LPI) and the primitive acceleration structure (octree).
  */
-bool ShouldSkipNaniteLPIs();
+bool ShouldSkipNaniteLPIs(EShaderPlatform ShaderPlatform);
 
 
 class FTileTexCoordVertexBuffer : public FVertexBuffer
diff --git a/Engine/Source/Runtime/Renderer/Private/RendererScene.cpp b/Engine/Source/Runtime/Renderer/Private/RendererScene.cpp
--- a/Engine/Source/Runtime/Renderer/Private/RendererScene.cpp
+++ b/Engine/Source/Runtime/Renderer/Private/RendererScene.cpp
@@ -4542,7 +4542,7 @@ FRHIUniformBuffer* FScene::GetSpeedTreeUniformBuffer(const FVertexFactory* Verte
  */
 void FScene::GetRelevantLights_RenderThread( const FPrimitiveSceneProxy* PrimitiveSceneProxy, TArray<const FLightSceneProxy*> &OutRelevantLights ) const
 {
-	if (ShouldSkipNaniteLPIs() && PrimitiveSceneProxy->IsNaniteMesh())
+	if (ShouldSkipNaniteLPIs(GetShaderPlatform()) && PrimitiveSceneProxy->IsNaniteMesh())
 	{
 		if (PrimitiveSceneProxy->GetLightingChannelMask() != 0)
 		{
@@ -6707,7 +6707,7 @@ void FScene::Update(FRDGBuilder& GraphBuilder, const FUpdateParameters& Paramete
 	{
 		SCOPED_NAMED_EVENT(CreateLightPrimitiveInteractions, FColor::Emerald);
 
-		bool bSkipNaniteLPIs = ShouldSkipNaniteLPIs();
+		bool bSkipNaniteLPIs = ShouldSkipNaniteLPIs(GetShaderPlatform());
 
 		for (FPrimitiveSceneInfo* PrimitiveSceneInfo : SceneInfosWithAddToScene)
 		{
diff --git a/Engine/Source/Runtime/Renderer/Private/RendererUtils.cpp b/Engine/Source/Runtime/Renderer/Private/RendererUtils.cpp
--- a/Engine/Source/Runtime/Renderer/Private/RendererUtils.cpp
+++ b/Engine/Source/Runtime/Renderer/Private/RendererUtils.cpp
@@ -27,10 +27,11 @@ static FAutoConsoleVariableRef CVarSkipNaniteLPIs(
 	}),
 	ECVF_RenderThreadSafe);
 
-bool ShouldSkipNaniteLPIs()
+bool ShouldSkipNaniteLPIs(EShaderPlatform ShaderPlatform)
 {
-	return GSkipNaniteLPIs > 1 
-		|| ( GSkipNaniteLPIs == 1 && !IsStaticLightingAllowed());
+	return (GSkipNaniteLPIs > 1 
+		|| ( GSkipNaniteLPIs == 1 && !IsStaticLightingAllowed()))
+		&& UseNanite(ShaderPlatform);
 }
 
 class FRTWriteMaskDecodeCS : public FGlobalShader
diff --git a/Engine/Source/Runtime/Renderer/Private/ShadowSetup.cpp b/Engine/Source/Runtime/Renderer/Private/ShadowSetup.cpp
--- a/Engine/Source/Runtime/Renderer/Private/ShadowSetup.cpp
+++ b/Engine/Source/Runtime/Renderer/Private/ShadowSetup.cpp
@@ -2427,7 +2427,7 @@ bool FProjectedShadowInfo::HasSubjectPrims() const
 		|| ShadowDepthPass.HasAnyDraw()
 		|| SubjectMeshCommandBuildRequests.Num() > 0
 		|| ShadowDepthPassVisibleCommands.Num() > 0
-		|| ((bContainsNaniteSubjects || ShouldSkipNaniteLPIs()) && bNaniteGeometry);
+		|| (bContainsNaniteSubjects && bNaniteGeometry);
 }
 
 void FProjectedShadowInfo::AddReceiverPrimitive(FPrimitiveSceneInfo* PrimitiveSceneInfo)
@@ -4271,7 +4271,8 @@ void FSceneRenderer::CreateWholeSceneProjectedShadow(
 					ProjectedShadowInfo->bVSM = true;
 					ProjectedShadowInfo->MeshSelectionMask = EShadowMeshSelection::VSM;
 
-					bool bContainsNaniteSubjects = false;
+					// Initialize to true if we are not creating LPIs for Nanite.
+					bool bContainsNaniteSubjects = ShouldSkipNaniteLPIs(TaskData.Scene->GetShaderPlatform());
 
 					// Skip mesh setup if it won't be rendered anyway
 					if (ProjectedShadowInfo->bShouldRenderVSM)
@@ -4318,7 +4319,7 @@ void FSceneRenderer::CreateWholeSceneProjectedShadow(
 							ProjectedShadowInfo->MeshSelectionMask = EShadowMeshSelection::SM;
 						}
 
-						bool bContainsNaniteSubjects = false;
+						bool bContainsNaniteSubjects = ShouldSkipNaniteLPIs(TaskData.Scene->GetShaderPlatform());
 
 						// Ray traced shadows use the GPU managed distance field object buffers, no CPU culling should be used
 						// Actually we probably want to have fallback even for these? Another performance regresion however (since we'd be adding culling & draw calls for DF lights).
