VolumetricCloud - made permutation naming concistant.

diff --git a/Engine/Shaders/Private/VolumetricRenderTarget.usf b/Engine/Shaders/Private/VolumetricRenderTarget.usf
--- a/Engine/Shaders/Private/VolumetricRenderTarget.usf
+++ b/Engine/Shaders/Private/VolumetricRenderTarget.usf
@@ -478,8 +478,8 @@ float4 SceneWithoutSingleLayerWaterViewRect;
 float2 FullResolutionToWaterBufferScale;
 #endif
 
-#if MSAA_SAMPLE_COUNT > 1
-Texture2DMS<float, MSAA_SAMPLE_COUNT> MSAADepthTexture;
+#if PERMUTATION_MSAA_SAMPLE_COUNT > 1
+Texture2DMS<float, PERMUTATION_MSAA_SAMPLE_COUNT> MSAADepthTexture;
 #endif
 
 #if INSTANCED_STEREO
@@ -613,7 +613,7 @@ float4 ApplyFogToCloudFromUVs(float2 UVWithOffset, float4 CloudRGBT)
 void ComposeVolumetricRTOverScenePS(
 	in float4 SVPos : SV_POSITION,
 	out float4 OutputRt0 : SV_Target0
-#if MSAA_SAMPLE_COUNT > 1
+#if PERMUTATION_MSAA_SAMPLE_COUNT > 1
 	, in uint SampleIndex : SV_SampleIndex
 #endif
 )
@@ -695,7 +695,7 @@ void ComposeVolumetricRTOverScenePS(
 
 #else
 
-#if  MSAA_SAMPLE_COUNT > 1
+#if  PERMUTATION_MSAA_SAMPLE_COUNT > 1
 	float PixelDeviceZ = MSAADepthTexture.Load(int2(SVPos.xy), SampleIndex).x;
 #else
 	float PixelDeviceZ = SceneTexturesStruct.SceneDepthTexture.Load(uint3(SVPos.xy, 0)).r;
diff --git a/Engine/Source/Runtime/Renderer/Private/VolumetricRenderTarget.cpp b/Engine/Source/Runtime/Renderer/Private/VolumetricRenderTarget.cpp
--- a/Engine/Source/Runtime/Renderer/Private/VolumetricRenderTarget.cpp
+++ b/Engine/Source/Runtime/Renderer/Private/VolumetricRenderTarget.cpp
@@ -679,7 +679,7 @@ class FComposeVolumetricRTOverScenePS : public FGlobalShader
 	class FUpsamplingMode : SHADER_PERMUTATION_RANGE_INT("PERMUTATION_UPSAMPLINGMODE", 0, 5);
 	class FRenderUnderWaterBuffer : SHADER_PERMUTATION_BOOL("PERMUTATION_RENDER_UNDERWATER_BUFFER");	// Render into the water scene color buffer (used when rendering from water system)
 	class FRenderCameraComposeWithWater : SHADER_PERMUTATION_BOOL("PERMUTATION_COMPOSE_WITH_WATER");	// When water us used and the camera is under water, use that permutation (to handle camera intersection with water and double cloud composition) 
-	class FMSAASampleCount : SHADER_PERMUTATION_SPARSE_INT("MSAA_SAMPLE_COUNT", 1, 2, 4, 8);
+	class FMSAASampleCount : SHADER_PERMUTATION_SPARSE_INT("PERMUTATION_MSAA_SAMPLE_COUNT", 1, 2, 4, 8);
 	class FApplyFog : SHADER_PERMUTATION_BOOL("PERMUTATION_APPLY_FOG");
 	class FApplyLocalFogVolume : SHADER_PERMUTATION_BOOL("PERMUTATION_APPLY_LOCAL_FOG_VOLUME");
 	using FPermutationDomain = TShaderPermutationDomain<FUpsamplingMode, FRenderUnderWaterBuffer, FRenderCameraComposeWithWater, FMSAASampleCount, FApplyFog, FApplyLocalFogVolume>;
