Fix ManyLights accidentally triggering (useless) SMRT work due to forward light grid encoding.
[RN] Fix ManyLights accidentally triggering (useless) SMRT work due to forward light grid encoding.

diff --git a/Engine/Shaders/Private/LightGridCommon.ush b/Engine/Shaders/Private/LightGridCommon.ush
--- a/Engine/Shaders/Private/LightGridCommon.ush
+++ b/Engine/Shaders/Private/LightGridCommon.ush
@@ -244,13 +244,9 @@ FLocalLightData GetLocalLightData_Internal(
 	Out.SpotAnglesAndSourceRadiusPacked			= InData3;
 	Out.LightTangentAndIESDataAndSpecularScale	= InData4;
 	Out.RectData								= InData5.xyz;
-#if MANY_LIGHTS
-	Out.VirtualShadowMapId						= -1;
-	Out.PrevLocalLightIndex						= int(InData5.w);
-#else
-	Out.VirtualShadowMapId						= int(InData5.w);
-	Out.PrevLocalLightIndex						= -1;
-#endif
+	// See PackVirtualShadowMapIdAndPrevLocalLightIndex
+	Out.VirtualShadowMapId						= int(asuint(InData5.w) >> 16U) - 1;
+	Out.PrevLocalLightIndex						= int(asuint(InData5.w) & 0xFFFF) - 1;
 	Out.LightSceneId							= int(Out.LightColorAndIdAndFalloffExponent.z);
 
 	return Out;
diff --git a/Engine/Source/Runtime/Renderer/Private/LightGridInjection.cpp b/Engine/Source/Runtime/Renderer/Private/LightGridInjection.cpp
--- a/Engine/Source/Runtime/Renderer/Private/LightGridInjection.cpp
+++ b/Engine/Source/Runtime/Renderer/Private/LightGridInjection.cpp
@@ -329,6 +329,19 @@ static FVector2f PackLightColor(const FVector3f& LightColor)
 	return LightColorPacked;
 }
 
+static uint32 PackVirtualShadowMapIdAndPrevLocalLightIndex(int32 VirtualShadowMapId, int32 PrevLocalLightIndex)
+{
+	// NOTE: Both of these could possibly be INDEX_NONE, which needs to be represented
+	// We map all negative numbers to 0, and add one to any positive ones
+	uint32 VSMPacked = VirtualShadowMapId < 0 ? 0 : uint32(VirtualShadowMapId + 1);
+	uint32 PrevPacked = PrevLocalLightIndex < 0 ? 0 : uint32(PrevLocalLightIndex + 1);
+
+	// Pack to 16 bits each
+	check(VSMPacked <= MAX_uint16);
+	check(PrevPacked <= MAX_uint16);
+	return (VSMPacked << 16) | (PrevPacked & 0xFFFF);
+}
+
 static void PackLocalLightData(
 	FForwardLocalLightData& Out,
 	const FViewInfo& View,
@@ -368,12 +381,14 @@ static void PackLocalLightData(
 	const FVector3f LightColor = (FVector3f)SimpleLight.Color * FLightRenderParameters::GetLightExposureScale(View.GetLastEyeAdaptationExposure(), SimpleLight.InverseExposureBlend);
 	const FVector2f LightColorPacked = PackLightColor(LightColor);
 
+	const uint32 VirtualShadowMapIdAndPrevLocalLightIndex = PackVirtualShadowMapIdAndPrevLocalLightIndex(INDEX_NONE, INDEX_NONE);
+
 	Out.LightPositionAndInvRadius							= FVector4f(LightTranslatedWorldPosition, 1.0f / FMath::Max(SimpleLight.Radius, KINDA_SMALL_NUMBER));
 	Out.LightColorAndIdAndFalloffExponent					= FVector4f(LightColorPacked.X, LightColorPacked.Y, INDEX_NONE, SimpleLight.Exponent);
 	Out.LightDirectionAndShadowMapChannelMask				= FVector4f(FVector3f(1, 0, 0), FMath::AsFloat(ShadowMapChannelMask));
 	Out.SpotAnglesAndSourceRadiusPacked						= FVector4f(-2, 1, FMath::AsFloat(PackedZ), FMath::AsFloat(PackedW));
 	Out.LightTangentAndIESDataAndSpecularScale				= FVector4f(1.0f, 0.0f, 0.0f, FMath::AsFloat(SpecularScaleAndIESData));
-	Out.RectDataAndVirtualShadowMapIdOrPrevLocalLightIndex	= FVector4f(FMath::AsFloat(RectPackedX), FMath::AsFloat(RectPackedY), FMath::AsFloat(RectPackedZ), -1);
+	Out.RectDataAndVirtualShadowMapIdOrPrevLocalLightIndex	= FVector4f(FMath::AsFloat(RectPackedX), FMath::AsFloat(RectPackedY), FMath::AsFloat(RectPackedZ), FMath::AsFloat(VirtualShadowMapIdAndPrevLocalLightIndex));
 }
 
 static void PackLocalLightData(
@@ -384,7 +399,6 @@ static void PackLocalLightData(
 	const int32 LightSceneId,
 	const int32 VirtualShadowMapId,
 	const int32 PrevLocalLightIndex,
-	const bool bHandledByManyLights,
 	const float VolumetricScatteringIntensity)
 {
 	const FVector3f LightTranslatedWorldPosition(View.ViewMatrices.GetPreViewTranslation() + LightParameters.WorldPosition);
@@ -408,23 +422,16 @@ static void PackLocalLightData(
 
 	const FVector2f LightColorPacked = PackLightColor(FVector3f(LightParameters.Color));
 
-	// Since lights don't use VSM and Many Lights simultaneously and
-	// currently PrevLocalLightIndex is only accessed by Many Lights shaders
-	// we can store only one of the values to avoid increasing the size of the struct
-	// TODO: Improve packing to avoid this so that PrevLocalLightIndex can be accessed in other shaders as well
-	const int32 VirtualShadowMapIdOrPrevLocalLightIndex = bHandledByManyLights ? PrevLocalLightIndex : VirtualShadowMapId;
+	const uint32 VirtualShadowMapIdAndPrevLocalLightIndex = 
+		PackVirtualShadowMapIdAndPrevLocalLightIndex(VirtualShadowMapId, PrevLocalLightIndex);
 
-	// NOTE: This cast of VirtualShadowMapIdOrPrevLocalLightIndex to float is not ideal, but bitcast has issues here with INDEX_NONE -> NaN
-	// and 32-bit floats have enough mantissa to cover all reasonable numbers here for now.
 	// NOTE: SpotAngles needs full-precision for VSM one pass projection
 	Out.LightPositionAndInvRadius							= FVector4f(LightTranslatedWorldPosition, LightParameters.InvRadius);
 	Out.LightColorAndIdAndFalloffExponent					= FVector4f(LightColorPacked.X, LightColorPacked.Y, LightSceneId, LightParameters.FalloffExponent);
 	Out.LightDirectionAndShadowMapChannelMask				= FVector4f(LightParameters.Direction, FMath::AsFloat(LightTypeAndShadowMapChannelMaskAndLightFunctionIndexPacked));
 	Out.SpotAnglesAndSourceRadiusPacked						= FVector4f(LightParameters.SpotAngles.X, LightParameters.SpotAngles.Y, FMath::AsFloat(PackedZ), FMath::AsFloat(PackedW));
 	Out.LightTangentAndIESDataAndSpecularScale				= FVector4f(LightParameters.Tangent, FMath::AsFloat(SpecularScaleAndIESData));
-	Out.RectDataAndVirtualShadowMapIdOrPrevLocalLightIndex	= FVector4f(FMath::AsFloat(RectPackedX), FMath::AsFloat(RectPackedY), FMath::AsFloat(RectPackedZ), float(VirtualShadowMapIdOrPrevLocalLightIndex));
-
-	checkSlow(int32(Out.RectDataAndVirtualShadowMapIdOrPrevLocalLightIndex.W) == VirtualShadowMapIdOrPrevLocalLightIndex);
+	Out.RectDataAndVirtualShadowMapIdOrPrevLocalLightIndex	= FVector4f(FMath::AsFloat(RectPackedX), FMath::AsFloat(RectPackedY), FMath::AsFloat(RectPackedZ), FMath::AsFloat(VirtualShadowMapIdAndPrevLocalLightIndex));
 }
 
 FComputeLightGridOutput FSceneRenderer::ComputeLightGrid(FRDGBuilder& GraphBuilder, bool bCullLightsToGrid, FSortedLightSetSceneInfo& SortedLightSet)
@@ -602,7 +609,6 @@ FComputeLightGridOutput FSceneRenderer::ComputeLightGrid(FRDGBuilder& GraphBuild
 							LightSceneInfo->Id,
 							VirtualShadowMapId,
 							PrevLocalLightIndex,
-							SortedLightInfo.SortKey.Fields.bHandledByManyLights,
 							VolumetricScatteringIntensity);
 
 						const FSphere BoundingSphere = LightProxy->GetBoundingSphere();
