Path Tracer: Change main radiance buffer to half precision

The goal is to save memory and enable the 16 bit version of the NNE denoiser. To maintain enough dynamic range, we have to apply the PreExposure factor during rendering. However, this factor can change during rendering, so we apply the difference when displaying the accumulated radiance buffer. In non-offline rendering cases, allow the base and current exposure factors to differ by at most 16x (4 stops) which should be sufficient to avoid issues in the case of auto-exposure.

Because we are now using exposure-relative accumulation, change the max path exposure into a relative control (and rename it to max path intensity). This is closer to the behavior of most other path tracers.

diff --git a/Engine/Plugins/Experimental/OpenImageDenoise/Source/OpenImageDenoise/Private/Plugin.cpp b/Engine/Plugins/Experimental/OpenImageDenoise/Source/OpenImageDenoise/Private/Plugin.cpp
--- a/Engine/Plugins/Experimental/OpenImageDenoise/Source/OpenImageDenoise/Private/Plugin.cpp
+++ b/Engine/Plugins/Experimental/OpenImageDenoise/Source/OpenImageDenoise/Private/Plugin.cpp
@@ -82,7 +82,7 @@ FDenoiseSettings GetCurrentSettings()
 struct OIDNState
 {
 	// scratch CPU memory for running the OIDN filter
-	TArray<FLinearColor> RawPixels;
+	TArray<FHalfColor> RawPixels;
 	TArray<FHalfColor> RawAlbedo;
 	TArray<FHalfColor> RawNormal;
 
@@ -142,8 +142,8 @@ struct OIDNState
 				NormalFilter = oidn::FilterRef();
 			}
 			PixelsFilter = OIDNDevice.newFilter("RT");
-			PixelsFilter.setImage("color" , RawPixels.GetData(), oidn::Format::Float3, Size.X, Size.Y, 0, sizeof(FLinearColor), sizeof(FLinearColor) * Size.X);
-			PixelsFilter.setImage("output", RawPixels.GetData(), oidn::Format::Float3, Size.X, Size.Y, 0, sizeof(FLinearColor), sizeof(FLinearColor) * Size.X);
+			PixelsFilter.setImage("color" , RawPixels.GetData(), oidn::Format::Half3, Size.X, Size.Y, 0, sizeof(FHalfColor), sizeof(FHalfColor) * Size.X);
+			PixelsFilter.setImage("output", RawPixels.GetData(), oidn::Format::Half3, Size.X, Size.Y, 0, sizeof(FHalfColor), sizeof(FHalfColor) * Size.X);
 
 			if (CurrentSettings.UseAux)
 			{
@@ -162,8 +162,8 @@ struct OIDNState
 			if (CurrentSettings.DenoiseAlpha)
 			{
 				AlphaFilter = OIDNDevice.newFilter("RT");
-				AlphaFilter.setImage("color" , RawPixels.GetData(), oidn::Format::Float, Size.X, Size.Y, sizeof(float) * 3, sizeof(FLinearColor), sizeof(FLinearColor) * Size.X);
-				AlphaFilter.setImage("output", RawPixels.GetData(), oidn::Format::Float, Size.X, Size.Y, sizeof(float) * 3, sizeof(FLinearColor), sizeof(FLinearColor) * Size.X);
+				AlphaFilter.setImage("color" , RawPixels.GetData(), oidn::Format::Half, Size.X, Size.Y, sizeof(uint16_t) * 3, sizeof(uint16_t), sizeof(FHalfColor) * Size.X);
+				AlphaFilter.setImage("output", RawPixels.GetData(), oidn::Format::Half, Size.X, Size.Y, sizeof(uint16_t) * 3, sizeof(uint16_t), sizeof(FHalfColor) * Size.X);
 				AlphaFilter.set("hdr", true);
 				AlphaFilter.commit();
 			}
@@ -198,9 +198,9 @@ template <typename PixelType>
 static void CopyTextureFromCPUToGPU(FRHICommandListImmediate& RHICmdList, const TArray<PixelType>& SrcArray, FIntPoint Size, FRHITexture* DstTexture)
 {
 	uint32_t DestStride;
-	FLinearColor* DstBuffer = static_cast<PixelType*>(RHICmdList.LockTexture2D(DstTexture, 0, RLM_WriteOnly, DestStride, false));
+	PixelType* DstBuffer = static_cast<PixelType*>(RHICmdList.LockTexture2D(DstTexture, 0, RLM_WriteOnly, DestStride, false));
 	DestStride /= sizeof(PixelType);
-	const FLinearColor* SrcBuffer = SrcArray.GetData();
+	const PixelType* SrcBuffer = SrcArray.GetData();
 	for (int Y = 0; Y < Size.Y; Y++, SrcBuffer += Size.X, DstBuffer += DestStride)
 	{
 		FPlatformMemory::Memcpy(DstBuffer, SrcBuffer, Size.X * sizeof(PixelType));
diff --git a/Engine/Plugins/Interchange/Runtime/Source/Import/Private/Animation/InterchangeLevelSequenceHelper.cpp b/Engine/Plugins/Interchange/Runtime/Source/Import/Private/Animation/InterchangeLevelSequenceHelper.cpp
--- a/Engine/Plugins/Interchange/Runtime/Source/Import/Private/Animation/InterchangeLevelSequenceHelper.cpp
+++ b/Engine/Plugins/Interchange/Runtime/Source/Import/Private/Animation/InterchangeLevelSequenceHelper.cpp
@@ -309,7 +309,7 @@ namespace UE::Interchange::Private
 			{EInterchangePropertyTracks::CameraPostProcessSettingsMotionBlurMax, {UMovieSceneFloatTrack::StaticClass()->GetName(),  GET_MEMBER_NAME_STRING_CHECKED(UCameraComponent, PostProcessSettings) + FString{ TEXT(".") } + GET_MEMBER_NAME_STRING_CHECKED(FPostProcessSettings, MotionBlurMax), TEXT("Max (Post Process Settings)")}},
 			{EInterchangePropertyTracks::CameraPostProcessSettingsMotionBlurPerObjectSize, {UMovieSceneFloatTrack::StaticClass()->GetName(),  GET_MEMBER_NAME_STRING_CHECKED(UCameraComponent, PostProcessSettings) + FString{ TEXT(".") } + GET_MEMBER_NAME_STRING_CHECKED(FPostProcessSettings, MotionBlurPerObjectSize), TEXT("Per Object Size (Post Process Settings)")}},
 			{EInterchangePropertyTracks::CameraPostProcessSettingsPathTracingMaxBounces, {UMovieSceneIntegerTrack::StaticClass()->GetName(),  GET_MEMBER_NAME_STRING_CHECKED(UCameraComponent, PostProcessSettings) + FString{ TEXT(".") } + GET_MEMBER_NAME_STRING_CHECKED(FPostProcessSettings, PathTracingMaxBounces), TEXT("Max. Bounces (Post Process Settings)")}},
-			{EInterchangePropertyTracks::CameraPostProcessSettingsPathTracingMaxPathExposure, {UMovieSceneFloatTrack::StaticClass()->GetName(),  GET_MEMBER_NAME_STRING_CHECKED(UCameraComponent, PostProcessSettings) + FString{ TEXT(".") } + GET_MEMBER_NAME_STRING_CHECKED(FPostProcessSettings, PathTracingMaxPathExposure), TEXT("Max Path Exposure (Post Process Settings)")}},
+			{EInterchangePropertyTracks::CameraPostProcessSettingsPathTracingMaxPathIntensity, {UMovieSceneFloatTrack::StaticClass()->GetName(),  GET_MEMBER_NAME_STRING_CHECKED(UCameraComponent, PostProcessSettings) + FString{ TEXT(".") } + GET_MEMBER_NAME_STRING_CHECKED(FPostProcessSettings, PathTracingMaxPathIntensity), TEXT("Max Path Intensity (Post Process Settings)")}},
 			{EInterchangePropertyTracks::CameraPostProcessSettingsRayTracingAO, {UMovieSceneByteTrack::StaticClass()->GetName(),  GET_MEMBER_NAME_STRING_CHECKED(UCameraComponent, PostProcessSettings) + FString{ TEXT(".") } + GET_MEMBER_NAME_STRING_CHECKED(FPostProcessSettings, RayTracingAO), TEXT("Enabled (Post Process Settings)")}},
 			{EInterchangePropertyTracks::CameraPostProcessSettingsRayTracingAOIntensity, {UMovieSceneFloatTrack::StaticClass()->GetName(),  GET_MEMBER_NAME_STRING_CHECKED(UCameraComponent, PostProcessSettings) + FString{ TEXT(".") } + GET_MEMBER_NAME_STRING_CHECKED(FPostProcessSettings, RayTracingAOIntensity), TEXT("Intensity (Post Process Settings)")}},
 			{EInterchangePropertyTracks::CameraPostProcessSettingsRayTracingAORadius, {UMovieSceneFloatTrack::StaticClass()->GetName(),  GET_MEMBER_NAME_STRING_CHECKED(UCameraComponent, PostProcessSettings) + FString{ TEXT(".") } + GET_MEMBER_NAME_STRING_CHECKED(FPostProcessSettings, RayTracingAORadius), TEXT("Radius (Post Process Settings)")}},
diff --git a/Engine/Plugins/Interchange/Runtime/Source/Nodes/Public/InterchangeAnimationDefinitions.h b/Engine/Plugins/Interchange/Runtime/Source/Nodes/Public/InterchangeAnimationDefinitions.h
--- a/Engine/Plugins/Interchange/Runtime/Source/Nodes/Public/InterchangeAnimationDefinitions.h
+++ b/Engine/Plugins/Interchange/Runtime/Source/Nodes/Public/InterchangeAnimationDefinitions.h
@@ -286,7 +286,7 @@ enum class EInterchangePropertyTracks : int32
 	CameraPostProcessSettingsMotionBlurMax,
 	CameraPostProcessSettingsMotionBlurPerObjectSize,
 	CameraPostProcessSettingsPathTracingMaxBounces,
-	CameraPostProcessSettingsPathTracingMaxPathExposure,
+	CameraPostProcessSettingsPathTracingMaxPathIntensity,
 	CameraPostProcessSettingsRayTracingAO,
 	CameraPostProcessSettingsRayTracingAOIntensity,
 	CameraPostProcessSettingsRayTracingAORadius,
diff --git a/Engine/Shaders/Private/PathTracing/PathTracingCompositingPixelShader.usf b/Engine/Shaders/Private/PathTracing/PathTracingCompositingPixelShader.usf
--- a/Engine/Shaders/Private/PathTracing/PathTracingCompositingPixelShader.usf
+++ b/Engine/Shaders/Private/PathTracing/PathTracingCompositingPixelShader.usf
@@ -13,7 +13,7 @@ uint MaxSamples;
 int ProgressDisplayEnabled;
 Texture2D<float4> RadianceTexture;
 Texture2D<float> DepthTexture;
-
+float PreExposure; // Taken from the view, but adjusted by the base exposure already applied to the Radiance texture
 int AdaptiveSamplingVisualize;
 
 void CompositeMain(
@@ -27,11 +27,12 @@ void CompositeMain(
 	int3 TexCoord = int3(UV * BufferSize - View.ViewRectMin.xy, 0);
 
 	float4 Radiance = RadianceTexture.Load(TexCoord);
+	Radiance = select(IsFinite(Radiance), Radiance, MaxHalfFloat); // if any component overflowed half precision, clamp it now
 
 	// NOTE: UE has the convention of using 1.0 - Alpha
 	// We must saturate the result because the path tracer's alpha is stochastic in the presence of volumetrics
 	// It is important that the clamp happens _after_ averaging, otherwise this can bias the estimate
-	OutColor = float4(Radiance.rgb * View.PreExposure, saturate(1.0 - Radiance.a));
+	OutColor = float4(Radiance.rgb * PreExposure, saturate(1.0 - Radiance.a));
 
 	if (AdaptiveSamplingVisualize == 1)
 	{
diff --git a/Engine/Shaders/Private/PathTracing/PathTracingCore.ush b/Engine/Shaders/Private/PathTracing/PathTracingCore.ush
--- a/Engine/Shaders/Private/PathTracing/PathTracingCore.ush
+++ b/Engine/Shaders/Private/PathTracing/PathTracingCore.ush
@@ -31,6 +31,7 @@
 #include "./Volume/PathTracingVolumeSampling.ush"
 
 float BlendFactor;
+float BaseExposure;
 uint Iteration;
 uint TemporalSeed;
 uint MaxSamples;
@@ -285,13 +286,15 @@ struct FPathState
 
 void AccumulateRadiance(inout float3 TotalRadiance, float3 PathRadiance, bool bIsCameraRay)
 {
+	PathRadiance *= BaseExposure;
 	// User asked for path contributions to be clamped to reduce fireflies.
 	// This puts an upper bound on variance within the pixel at the expense of bias
 	float MaxPathRadiance = max3(PathRadiance.x, PathRadiance.y, PathRadiance.z);
-	if (!bIsCameraRay && MaxPathRadiance > MaxPathIntensity)
+	float MaxIntensity = bIsCameraRay ? MaxHalfFloat : MaxPathIntensity;
+	if (MaxPathRadiance > MaxIntensity)
 	{
 		// adjust sample while keeping color
-		PathRadiance *= MaxPathIntensity / MaxPathRadiance;
+		PathRadiance *= MaxIntensity / MaxPathRadiance;
 	}
 
 	// accumulate into the total
diff --git a/Engine/Source/Runtime/Engine/Classes/Engine/Scene.h b/Engine/Source/Runtime/Engine/Classes/Engine/Scene.h
--- a/Engine/Source/Runtime/Engine/Classes/Engine/Scene.h
+++ b/Engine/Source/Runtime/Engine/Classes/Engine/Scene.h
@@ -1291,7 +1291,7 @@ struct FPostProcessSettings
 	uint32 bOverride_PathTracingSamplesPerPixel : 1;
 
 	UPROPERTY(EditAnywhere, BlueprintReadWrite, Category = Overrides, meta = (PinHiddenByDefault, InlineEditConditionToggle))
-	uint32 bOverride_PathTracingMaxPathExposure : 1;
+	uint32 bOverride_PathTracingMaxPathIntensity : 1;
 
 	UPROPERTY(EditAnywhere, BlueprintReadWrite, Category = Overrides, meta = (PinHiddenByDefault, InlineEditConditionToggle))
 	uint32 bOverride_PathTracingEnableEmissiveMaterials : 1;
@@ -2224,9 +2224,9 @@ struct FPostProcessSettings
 	UPROPERTY(EditAnywhere, BlueprintReadWrite, Category = "Path Tracing", meta = (ClampMin = "1", UIMax = "65536", editcondition = "bOverride_PathTracingSamplesPerPixel", DisplayName = "Samples Per Pixel"))
 	int32 PathTracingSamplesPerPixel;
 
-	/** Sets the maximum exposure allowed in the path tracer to reduce fireflies. This should be set a few stops higher than the scene exposure. */
-	UPROPERTY(interp, EditAnywhere, BlueprintReadWrite, Category = "Path Tracing", meta = (ClampMin = "-10.0", ClampMax = "30.0", editcondition = "bOverride_PathTracingMaxPathExposure", DisplayName = "Max Path Exposure"))
-	float PathTracingMaxPathExposure;
+	/** Sets the maximum intensity of indirect samples to reduce fireflies. Lowering this value reduces noise at the expense of accuracy. Increasing it is more accurate but may lead to more noise. */
+	UPROPERTY(interp, EditAnywhere, BlueprintReadWrite, Category = "Path Tracing", meta = (ClampMin = "1.0", ClampMax = "65504.0", editcondition = "bOverride_PathTracingMaxPathIntensity", DisplayName = "Max Path Exposure"))
+	float PathTracingMaxPathIntensity;
 
 	/** Should emissive materials contribute to scene lighting? */
 	UPROPERTY(EditAnywhere, BlueprintReadWrite, Category = "Path Tracing", meta = (editcondition = "bOverride_PathTracingEnableEmissiveMaterials", DisplayName = "Emissive Materials"))
diff --git a/Engine/Source/Runtime/Engine/Private/Scene.cpp b/Engine/Source/Runtime/Engine/Private/Scene.cpp
--- a/Engine/Source/Runtime/Engine/Private/Scene.cpp
+++ b/Engine/Source/Runtime/Engine/Private/Scene.cpp
@@ -633,7 +633,7 @@ FPostProcessSettings::FPostProcessSettings()
 
 	PathTracingMaxBounces = 32;
 	PathTracingSamplesPerPixel = 2048;
-	PathTracingMaxPathExposure = 30.0f;
+	PathTracingMaxPathIntensity = 24.0f;
 	PathTracingEnableEmissiveMaterials = 1;
 	PathTracingEnableReferenceDOF = 0;
 	PathTracingEnableReferenceAtmosphere = 0;
diff --git a/Engine/Source/Runtime/Engine/Private/SceneView.cpp b/Engine/Source/Runtime/Engine/Private/SceneView.cpp
--- a/Engine/Source/Runtime/Engine/Private/SceneView.cpp
+++ b/Engine/Source/Runtime/Engine/Private/SceneView.cpp
@@ -1741,7 +1741,7 @@ void FSceneView::OverridePostProcessSettings(const FPostProcessSettings& Src, fl
 		// Path Tracing related settings
 		SET_PP(PathTracingMaxBounces);
 		SET_PP(PathTracingSamplesPerPixel);
-		LERP_PP(PathTracingMaxPathExposure);
+		LERP_PP(PathTracingMaxPathIntensity);
 		SET_PP(PathTracingEnableEmissiveMaterials);
 		SET_PP(PathTracingEnableReferenceDOF);
 		SET_PP(PathTracingEnableReferenceAtmosphere);
diff --git a/Engine/Source/Runtime/Renderer/Private/PathTracing.cpp b/Engine/Source/Runtime/Renderer/Private/PathTracing.cpp
--- a/Engine/Source/Runtime/Renderer/Private/PathTracing.cpp
+++ b/Engine/Source/Runtime/Renderer/Private/PathTracing.cpp
@@ -506,6 +506,7 @@ TAutoConsoleVariable<float> CVarPathTracingBackgroundAlpha(
 
 BEGIN_SHADER_PARAMETER_STRUCT(FPathTracingData, )
 	SHADER_PARAMETER(float, BlendFactor)
+	SHADER_PARAMETER(float, BaseExposure)
 	SHADER_PARAMETER(uint32, Iteration)
 	SHADER_PARAMETER(uint32, TemporalSeed)
 	SHADER_PARAMETER(uint32, MaxSamples)
@@ -604,6 +605,13 @@ struct FPathTracingConfig
 			UseMultiGPU != Other.UseMultiGPU;
 	}
 
+	bool IsExposureDifferentEnough(const FPathTracingConfig& Other) const
+	{
+		const float ExposureA = PathTracingData.BaseExposure;
+		const float ExposureB = Other.PathTracingData.BaseExposure;
+		return FMath::Max(ExposureA, ExposureB) > 16.0f * FMath::Min(ExposureA, ExposureB);
+	}
+
 	bool IsDOFDifferent(const FPathTracingConfig& Other) const
 	{
 		return PathTracingData.CameraFocusDistance != Other.PathTracingData.CameraFocusDistance ||
@@ -718,6 +726,9 @@ static void PreparePathTracingData(const FScene* Scene, const FViewInfo& View, F
 	const FFinalPostProcessSettings& PPV = View.FinalPostProcessSettings;
 	const FEngineShowFlags& ShowFlags = View.Family->EngineShowFlags;
 
+	// Capture the current exposure (NOTE: This is overwritten later so we maintain the exposure that was used on the first sample)
+	PathTracingData.BaseExposure = View.PreExposure;
+
 	int32 MaxBounces = CVarPathTracingMaxBounces.GetValueOnRenderThread();
 	if (MaxBounces < 0)
 	{
@@ -734,9 +745,10 @@ static void PreparePathTracingData(const FScene* Scene, const FViewInfo& View, F
 	PathTracingData.MaxPathIntensity = CVarPathTracingMaxPathIntensity.GetValueOnRenderThread();
 	if (PathTracingData.MaxPathIntensity <= 0)
 	{
-		// cvar clamp disabled, use PPV exposure value instad
-		PathTracingData.MaxPathIntensity = FMath::Pow(2.0f, PPV.PathTracingMaxPathExposure);
+		// cvar clamp disabled, use PPV value instad
+		PathTracingData.MaxPathIntensity = PPV.PathTracingMaxPathIntensity;
 	}
+	PathTracingData.MaxPathIntensity = FFloat16(PathTracingData.MaxPathIntensity).GetClampedNonNegativeAndFinite().GetFloat(); // Clip to half precision
 	PathTracingData.ApproximateCaustics = CVarPathTracingApproximateCaustics.GetValueOnRenderThread();
 	PathTracingData.EnableCameraBackfaceCulling = CVarPathTracingEnableCameraBackfaceCulling.GetValueOnRenderThread();
 	PathTracingData.SamplerType = CVarPathTracingSamplerType.GetValueOnRenderThread();
@@ -2758,6 +2770,7 @@ class FPathTracingCompositorPS : public FGlobalShader
 		SHADER_PARAMETER(float, AdaptiveSamplingErrorThreshold)
 		SHADER_PARAMETER(int, AdaptiveSamplingVisualize)
 		SHADER_PARAMETER(FIntVector, VarianceTextureDims)
+		SHADER_PARAMETER(float, PreExposure)
 
 		RENDER_TARGET_BINDING_SLOTS()
 	END_SHADER_PARAMETER_STRUCT()
@@ -3132,6 +3145,7 @@ void FDeferredShadingSceneRenderer::RenderPathTracing(
 	// NOTE: same for DOF changes, these parameters could be animated which should not automatically invalidate a render in progress
 	if (FirstTime ||
 		Config.IsDifferent(PathTracingState->LastConfig) ||
+		(!View.bIsOfflineRender && Config.IsExposureDifferentEnough(PathTracingState->LastConfig)) ||
 		(!View.bIsOfflineRender && Config.IsDOFDifferent(PathTracingState->LastConfig)) ||
 		(!View.bIsOfflineRender && HairStrands::HasPositionsChanged(GraphBuilder, *Scene, View)))
 	{
@@ -3139,6 +3153,8 @@ void FDeferredShadingSceneRenderer::RenderPathTracing(
 		PathTracingState->LastConfig = Config;
 		View.ViewState->PathTracingInvalidate();
 	}
+	// copy the base exposure from last time, so we can have a consistent exposure when we accumulate samples
+	Config.PathTracingData.BaseExposure = PathTracingState->LastConfig.PathTracingData.BaseExposure;
 
 	// Declare heterogeneous volume buffers
 	TRDGUniformBufferRef<FOrthoVoxelGridUniformBufferParameters> OrthoGridUniformBuffer;
@@ -3165,7 +3181,7 @@ void FDeferredShadingSceneRenderer::RenderPathTracing(
 		// First time through, need to make a new texture
 		FRDGTextureDesc RadianceDesc = FRDGTextureDesc::Create2D(
 			View.ViewRect.Size(),
-			PF_A32B32G32R32F,
+			PF_FloatRGBA,
 			FClearValueBinding::None,
 			TexCreate_ShaderResource | TexCreate_UAV | GetExtraTextureCreateFlagsForDenoiser());
 		FRDGTextureDesc AlbedoNormalDesc = FRDGTextureDesc::Create2D(
@@ -3970,6 +3986,7 @@ void FDeferredShadingSceneRenderer::RenderPathTracing(
 	DisplayParameters->RenderTargets[0] = FRenderTargetBinding(SceneColorOutputTexture, ERenderTargetLoadAction::ELoad);
 	DisplayParameters->RenderTargets.DepthStencil = FDepthStencilBinding(SceneDepthOutputTexture,  ERenderTargetLoadAction::ELoad, ERenderTargetLoadAction::ENoAction, FExclusiveDepthStencil::DepthWrite_StencilNop);
 	DisplayParameters->VarianceSampler = TStaticSamplerState<SF_Bilinear>::GetRHI();
+	DisplayParameters->PreExposure = View.PreExposure / Config.PathTracingData.BaseExposure;
 
 	FScreenPassTextureViewport Viewport(SceneColorOutputTexture, View.ViewRect);
 
diff --git a/Engine/Source/Runtime/Renderer/Private/PathTracingSpatialTemporalDenoising.cpp b/Engine/Source/Runtime/Renderer/Private/PathTracingSpatialTemporalDenoising.cpp
--- a/Engine/Source/Runtime/Renderer/Private/PathTracingSpatialTemporalDenoising.cpp
+++ b/Engine/Source/Runtime/Renderer/Private/PathTracingSpatialTemporalDenoising.cpp
@@ -1964,13 +1964,7 @@ void PathTracingSpatialTemporalDenoising(FRDGBuilder& GraphBuilder,
 	RDG_GPU_STAT_SCOPE(GraphBuilder, PathTracingSpatialTemporalDenoising);
 	RDG_EVENT_SCOPE(GraphBuilder, "PathTracingSpatialTemporalDenoising");
 
-	ETextureCreateFlags ExtraFlags = GetExtraTextureCreateFlagsForDenoiser();
-
-	FRDGTextureDesc RadianceTextureDesc = FRDGTextureDesc::Create2D(
-		View.ViewRect.Size(),
-		PF_A32B32G32R32F,
-		FClearValueBinding::None,
-		TexCreate_ShaderResource | TexCreate_UAV | ExtraFlags);
+	FRDGTextureDesc RadianceTextureDesc = SpatialTemporalDenoisingContext.RadianceTexture->Desc;
 
 	const ESpatialDenoiserType SpatialDenoiserType = GetSpatialDenosierType();
 	const ETemporalDenoiserType TemporalDenoiserType = GetTemporalDenoiserType();
