Path Tracer: Enable NNE Denoiser by default

Fix a few typos (denosier should be denoiser)

diff --git a/Engine/Plugins/Experimental/NFORDenoise/Source/NFORDenoise/Private/NFORDenoise.cpp b/Engine/Plugins/Experimental/NFORDenoise/Source/NFORDenoise/Private/NFORDenoise.cpp
--- a/Engine/Plugins/Experimental/NFORDenoise/Source/NFORDenoise/Private/NFORDenoise.cpp
+++ b/Engine/Plugins/Experimental/NFORDenoise/Source/NFORDenoise/Private/NFORDenoise.cpp
@@ -105,7 +105,7 @@ FRDGFrameCache FExternalFrameCache::ToRGDResource(FRDGBuilder& GraphBuilder, con
 }
 
 
-class FNFORDenosier : public IPathTracingSpatialTemporalDenoiser
+class FNFORDenoiser : public IPathTracingSpatialTemporalDenoiser
 {
 public:
 	class FHistory : public IHistory
@@ -182,7 +182,7 @@ public:
 		TArray<FExternalFrameCache> FrameCaches;
 	};
 
-	~FNFORDenosier() {}
+	~FNFORDenoiser() {}
 
 	const TCHAR* GetDebugName() const override { return *DebugName; }
 
@@ -249,7 +249,7 @@ public:
 	}
 
 private:
-	inline static const FString DebugName = TEXT("FNFORDenosier");
+	inline static const FString DebugName = TEXT("FNFORDenoiser");
 };
 
 void FNFORDenoiseModule::StartupModule()
@@ -262,7 +262,7 @@ void FNFORDenoiseModule::StartupModule()
 		FString ModuleDir = Plugin->GetBaseDir() + TEXT("/Source/NFORDenoise");
 		AddShaderSourceDirectoryMapping(TEXT("/NFORDenoise"), FPaths::Combine(ModuleDir, TEXT("Shaders")));
 
-		RegisterSpatialTemporalDenoiser(MakeUnique<FNFORDenosier>(),TEXT("NFOR"));
+		RegisterSpatialTemporalDenoiser(MakeUnique<FNFORDenoiser>(),TEXT("NFOR"));
 	}
 	else
 	{
diff --git a/Engine/Plugins/Experimental/NNEDenoiser/NNEDenoiser.uplugin b/Engine/Plugins/Experimental/NNEDenoiser/NNEDenoiser.uplugin
--- a/Engine/Plugins/Experimental/NNEDenoiser/NNEDenoiser.uplugin
+++ b/Engine/Plugins/Experimental/NNEDenoiser/NNEDenoiser.uplugin
@@ -11,9 +11,13 @@
 	"MarketplaceURL": "",
 	"SupportURL": "",
 	"CanContainContent": true,
+	"EnabledByDefault" : true,
 	"IsBetaVersion": false,
 	"IsExperimentalVersion": true,
 	"Installed": false,
+	"SupportedTargetPlatforms": [
+		"Win64", "Linux", "Mac"
+	],
 	"Modules": [
 		{
 			"Name": "NNEDenoiser",
diff --git a/Engine/Plugins/Experimental/OptiXDenoise/Source/OptiXDenoise/Private/Plugin.cpp b/Engine/Plugins/Experimental/OptiXDenoise/Source/OptiXDenoise/Private/Plugin.cpp
--- a/Engine/Plugins/Experimental/OptiXDenoise/Source/OptiXDenoise/Private/Plugin.cpp
+++ b/Engine/Plugins/Experimental/OptiXDenoise/Source/OptiXDenoise/Private/Plugin.cpp
@@ -231,7 +231,7 @@ static void Denoise(
 
 using namespace UE::Renderer::Private;
 
-class FOptiXDenosier : public IPathTracingSpatialTemporalDenoiser
+class FOptiXDenoiser : public IPathTracingSpatialTemporalDenoiser
 {
 public:
 	class FHistory : public IHistory
@@ -247,7 +247,7 @@ public:
 		const TCHAR* DebugName;
 	};
 
-	~FOptiXDenosier() {}
+	~FOptiXDenoiser() {}
 
 	const TCHAR* GetDebugName() const override { return *DebugName; }
 
@@ -308,7 +308,7 @@ public:
 	}
 
 private:
-	inline static const FString DebugName = TEXT("FOptiXDenosier");
+	inline static const FString DebugName = TEXT("FOptiXDenoiser");
 };
 
 void FOptiXDenoiseModule::StartupModule()
@@ -326,7 +326,7 @@ void FOptiXDenoiseModule::StartupModule()
 	OptiXDenoiseBaseDLLHandle = FPlatformProcess::GetDllHandle(*OptiXDenoiseBaseDllPath);
 	FOptiXCudaFunctionList::Get().RegisterFunctionInstance<FOptiXDenoiserFunctionInstance>();
 
-	RegisterSpatialTemporalDenoiser(MakeUnique<FOptiXDenosier>(), TEXT("OptiX"));
+	RegisterSpatialTemporalDenoiser(MakeUnique<FOptiXDenoiser>(), TEXT("OptiX"));
 }
 
 void FOptiXDenoiseModule::ShutdownModule()
diff --git a/Engine/Plugins/MovieScene/MovieRenderPipeline/Source/MovieRenderPipelineRenderPasses/Private/Graph/Nodes/MovieGraphPathTracerPassNode.cpp b/Engine/Plugins/MovieScene/MovieRenderPipeline/Source/MovieRenderPipelineRenderPasses/Private/Graph/Nodes/MovieGraphPathTracerPassNode.cpp
--- a/Engine/Plugins/MovieScene/MovieRenderPipeline/Source/MovieRenderPipelineRenderPasses/Private/Graph/Nodes/MovieGraphPathTracerPassNode.cpp
+++ b/Engine/Plugins/MovieScene/MovieRenderPipeline/Source/MovieRenderPipelineRenderPasses/Private/Graph/Nodes/MovieGraphPathTracerPassNode.cpp
@@ -89,9 +89,9 @@ void UMovieGraphPathTracerRenderPassNode::SetupImpl(const FMovieGraphRenderPassS
 
 	// Reset the max frame count to 0 if spatial denoiser is in use.
 	// We need this reset because here we use the max of all nodes.
-	const EMovieGraphPathTracerDenoiserType EffectiveDenosierType = static_cast<EMovieGraphPathTracerDenoiserType>(MaxDenoiserType);
+	const EMovieGraphPathTracerDenoiserType EffectiveDenoiserType = static_cast<EMovieGraphPathTracerDenoiserType>(MaxDenoiserType);
 
-	if (EffectiveDenosierType == EMovieGraphPathTracerDenoiserType::Spatial)
+	if (EffectiveDenoiserType == EMovieGraphPathTracerDenoiserType::Spatial)
 	{
 		MaxFrameCount = 0;
 	}
@@ -110,14 +110,14 @@ void UMovieGraphPathTracerRenderPassNode::SetupImpl(const FMovieGraphRenderPassS
 	}
 
 	//If no temporal denoiser is enabled, provide warning if we use temporal denoiser.
-	if (!HasTemporalDenosier())
+	if (!HasTemporalDenoiser())
 	{
-		bool bShouldPerformTemporalDenoising = EMovieGraphPathTracerDenoiserType::Temporal == EffectiveDenosierType;
+		bool bShouldPerformTemporalDenoising = EMovieGraphPathTracerDenoiserType::Temporal == EffectiveDenoiserType;
 
 		if (bShouldPerformTemporalDenoising)
 		{
 			UE_LOG(LogMovieRenderPipeline, Warning, TEXT("The Path Traced Renderer node of Layer `%s` enables temporal denoising but no temporal denoiser plugin is enabled." 
-														 "Fallback to the first available spatial denosier. Please enable at least one denoiser plugin "
+														 "Fallback to the first available spatial denoiser. Please enable at least one denoiser plugin "
 														 "with temporal denoising capability to fully apply the node setting."), *WarningLayerName);
 		}
 	}
@@ -172,7 +172,7 @@ int32 UMovieGraphPathTracerRenderPassNode::GetCoolingDownFrameCount() const
 {
 	return ((DenoiserType == EMovieGraphPathTracerDenoiserType::Temporal)
 		&& GetAllowDenoiser()
-		&& HasTemporalDenosier()) ? FrameCount : 0;
+		&& HasTemporalDenoiser()) ? FrameCount : 0;
 }
 
 EViewModeIndex UMovieGraphPathTracerRenderPassNode::GetViewModeIndex() const
diff --git a/Engine/Plugins/MovieScene/MovieRenderPipeline/Source/MovieRenderPipelineRenderPasses/Public/Graph/Nodes/MovieGraphPathTracerPassNode.h b/Engine/Plugins/MovieScene/MovieRenderPipeline/Source/MovieRenderPipelineRenderPasses/Public/Graph/Nodes/MovieGraphPathTracerPassNode.h
--- a/Engine/Plugins/MovieScene/MovieRenderPipeline/Source/MovieRenderPipelineRenderPasses/Public/Graph/Nodes/MovieGraphPathTracerPassNode.h
+++ b/Engine/Plugins/MovieScene/MovieRenderPipeline/Source/MovieRenderPipelineRenderPasses/Public/Graph/Nodes/MovieGraphPathTracerPassNode.h
@@ -118,7 +118,7 @@ public:
 	bool bEnableDenoiser;
 
 	/**
-	* Select which type of denosier to use when the denoiser is enabled. Temporal denoisers will provide better results when
+	* Select which type of denoiser to use when the denoiser is enabled. Temporal denoisers will provide better results when
 	* denoising animated sequences (the denoising results will look more stable), especially when combined with an appropriate 
 	* Frame Count (non-zero). Denoisers are implemented as plugins so you may need to enable a plugin as well for this to work.
 	*/
diff --git a/Engine/Source/Runtime/Renderer/Private/LightRendering.cpp b/Engine/Source/Runtime/Renderer/Private/LightRendering.cpp
--- a/Engine/Source/Runtime/Renderer/Private/LightRendering.cpp
+++ b/Engine/Source/Runtime/Renderer/Private/LightRendering.cpp
@@ -1124,7 +1124,7 @@ extern int32 GbEnableAsyncComputeTranslucencyLightingVolumeClear;
 
 uint32 GetShadowQuality();
 
-static bool LightRequiresDenosier(const FLightSceneInfo& LightSceneInfo)
+static bool LightRequiresDenoiser(const FLightSceneInfo& LightSceneInfo)
 {
 	ELightComponentType LightType = ELightComponentType(LightSceneInfo.Proxy->GetLightType());
 	if (LightType == LightType_Directional)
@@ -1677,7 +1677,7 @@ void FDeferredShadingSceneRenderer::RenderLights(
 						IScreenSpaceDenoiser::FShadowRayTracingConfig RayTracingConfig;
 						RayTracingConfig.RayCountPerPixel = GShadowRayTracingSamplesPerPixel > -1? GShadowRayTracingSamplesPerPixel : LightSceneProxy.GetSamplesPerPixel();
 
-						const bool bDenoiserCompatible = !LightRequiresDenosier(LightSceneInfo) || IScreenSpaceDenoiser::EShadowRequirements::PenumbraAndClosestOccluder == DenoiserToUse->GetShadowRequirements(View, LightSceneInfo, RayTracingConfig);
+						const bool bDenoiserCompatible = !LightRequiresDenoiser(LightSceneInfo) || IScreenSpaceDenoiser::EShadowRequirements::PenumbraAndClosestOccluder == DenoiserToUse->GetShadowRequirements(View, LightSceneInfo, RayTracingConfig);
 
 						const bool bWantsBatchedShadow = OcclusionType == FLightOcclusionType::Raytraced && 
 							bDoShadowBatching &&
@@ -1755,7 +1755,7 @@ void FDeferredShadingSceneRenderer::RenderLights(
 									continue;
 								}
 
-								const bool bRequiresDenoiser = LightRequiresDenosier(BatchLightSceneInfo) && DenoiserMode > 0;
+								const bool bRequiresDenoiser = LightRequiresDenoiser(BatchLightSceneInfo) && DenoiserMode > 0;
 
 								IScreenSpaceDenoiser::FShadowRayTracingConfig BatchRayTracingConfig;
 								BatchRayTracingConfig.RayCountPerPixel = GShadowRayTracingSamplesPerPixel > -1 ? GShadowRayTracingSamplesPerPixel : BatchLightSceneInfo.Proxy->GetSamplesPerPixel();
@@ -1967,7 +1967,7 @@ void FDeferredShadingSceneRenderer::RenderLights(
 							RayTracingConfig.RayCountPerPixel = GShadowRayTracingSamplesPerPixel > -1 ? GShadowRayTracingSamplesPerPixel : LightSceneProxy.GetSamplesPerPixel();
 
 							IScreenSpaceDenoiser::EShadowRequirements DenoiserRequirements = IScreenSpaceDenoiser::EShadowRequirements::Bailout;
-							if (DenoiserMode != 0 && LightRequiresDenosier(LightSceneInfo))
+							if (DenoiserMode != 0 && LightRequiresDenoiser(LightSceneInfo))
 							{
 								DenoiserRequirements = DenoiserToUse->GetShadowRequirements(View, LightSceneInfo, RayTracingConfig);
 							}
diff --git a/Engine/Source/Runtime/Renderer/Private/PathTracingSpatialTemporalDenoising.cpp b/Engine/Source/Runtime/Renderer/Private/PathTracingSpatialTemporalDenoising.cpp
--- a/Engine/Source/Runtime/Renderer/Private/PathTracingSpatialTemporalDenoising.cpp
+++ b/Engine/Source/Runtime/Renderer/Private/PathTracingSpatialTemporalDenoising.cpp
@@ -13,25 +13,25 @@ public:
 		return DenoiserManager;
 	}
 
-	void RegisterSpatialDenoiser(TUniquePtr<UE::Renderer::Private::IPathTracingDenoiser> InDenosier, FString Name)
+	void RegisterSpatialDenoiser(TUniquePtr<UE::Renderer::Private::IPathTracingDenoiser> InDenoiser, FString Name)
 	{
 		check(!SpatialDenoiser.Contains(Name));
 
-		bNeedTextureCreateExtraFlags |= InDenosier->NeedTextureCreateExtraFlags();
-		SpatialDenoiser.Add(Name, MoveTemp(InDenosier));
+		bNeedTextureCreateExtraFlags |= InDenoiser->NeedTextureCreateExtraFlags();
+		SpatialDenoiser.Add(Name, MoveTemp(InDenoiser));
 	}
-	void RegisterSpatialTemporalDenoiser(TUniquePtr<UE::Renderer::Private::IPathTracingSpatialTemporalDenoiser> InDenosier, FString Name)
+	void RegisterSpatialTemporalDenoiser(TUniquePtr<UE::Renderer::Private::IPathTracingSpatialTemporalDenoiser> InDenoiser, FString Name)
 	{
 		check(!PathTracingSpatialTemporalDenoisers.Contains(Name));
-		bNeedTextureCreateExtraFlags |= InDenosier->NeedTextureCreateExtraFlags();
-		PathTracingSpatialTemporalDenoisers.Add(Name, MoveTemp(InDenosier));
+		bNeedTextureCreateExtraFlags |= InDenoiser->NeedTextureCreateExtraFlags();
+		PathTracingSpatialTemporalDenoisers.Add(Name, MoveTemp(InDenoiser));
 	}
 
 	bool HasSpatialDenoiser()const { return SpatialDenoiser.Num() > 0; }
 	bool HasSpatialTemporalDenoiser()const { return PathTracingSpatialTemporalDenoisers.Num() > 0; }
 	bool HasDenoiser()const { return HasSpatialDenoiser() || HasSpatialTemporalDenoiser();}
 
-	UE::Renderer::Private::IPathTracingDenoiser* GetSpatialDenosier(FString Name, bool bMatch)
+	UE::Renderer::Private::IPathTracingDenoiser* GetSpatialDenoiser(FString Name, bool bMatch)
 	{
 		if (SpatialDenoiser.Contains(Name))
 		{
@@ -48,7 +48,7 @@ public:
 		return nullptr;
 	}
 
-	UE::Renderer::Private::IPathTracingSpatialTemporalDenoiser* GetSpatialTemporalDenosier(FString Name, bool bMatch)
+	UE::Renderer::Private::IPathTracingSpatialTemporalDenoiser* GetSpatialTemporalDenoiser(FString Name, bool bMatch)
 	{
 		if (PathTracingSpatialTemporalDenoisers.Contains(Name))
 		{
@@ -105,7 +105,7 @@ void UnregisterDenoiser(FString Name)
 	FDenoiserManager::Get().UnregisterDenoiser(Name);
 }
 
-bool HasTemporalDenosier()
+bool HasTemporalDenoiser()
 {
 	return FDenoiserManager::Get().HasSpatialTemporalDenoiser();
 }
@@ -176,15 +176,15 @@ namespace {
 
 	TAutoConsoleVariable<FString> CVarPathTracingDenoiserName(
 		TEXT("r.PathTracing.Denoiser.Name"),
-		"OIDN",
-		TEXT("Set the spatial denoiser name. It is the corresponding name registered by a denosier plugin\n")
+		"NNE_OIDN",
+		TEXT("Set the spatial denoiser name. It is the corresponding name registered by a denoiser plugin\n")
 		TEXT("Any registered denoiser should be able to denoise a single frame spatially.\n")
 	);
 
 	TAutoConsoleVariable<FString> CVarPathTracingTemporalDenoiserName(
 		TEXT("r.PathTracing.TemporalDenoiser.Name"),
 		"NFOR",
-		TEXT("Set the temporal denoiser name. It is the corresponding name registered by a denosier plugin\n")
+		TEXT("Set the temporal denoiser name. It is the corresponding name registered by a denoiser plugin\n")
 		TEXT("The temporal denoiser usually has better temporal stability when rendering offline.\n")
 	);
 
@@ -449,7 +449,7 @@ TArray<ESpatialDenoiserType> GetAvailableSpatialDenoiserTypes()
 	return Types;
 }
 
-ESpatialDenoiserType GetSpatialDenosierType()
+ESpatialDenoiserType GetSpatialDenoiserType()
 {
 	int32 Type = CVarPathTracingSpatialDenoiserType.GetValueOnRenderThread();
 
@@ -489,7 +489,7 @@ static bool ShouldApplyTemporalDenoiser(const FPathTracingSpatialTemporalDenoisi
 ETemporalDenoiserType GetTemporalDenoiserType()
 {
 
-	ESpatialDenoiserType SpatialDenoiserType = GetSpatialDenosierType();
+	ESpatialDenoiserType SpatialDenoiserType = GetSpatialDenoiserType();
 
 	if ( SpatialDenoiserType == ESpatialDenoiserType::SPATIAL_DENOISER_PLUGIN)
 	{
@@ -518,18 +518,18 @@ ETemporalDenoiserMotionVectorType GetTemporalDenoiserMotionVectorType()
 UE::Renderer::Private::IPathTracingDenoiser* GetActiveSpatialDenoiser(bool bMatch = false)
 {
 	FString DenoiserName = CVarPathTracingDenoiserName.GetValueOnRenderThread();
-	return FDenoiserManager::Get().GetSpatialDenosier(DenoiserName, bMatch);
+	return FDenoiserManager::Get().GetSpatialDenoiser(DenoiserName, bMatch);
 }
 
 UE::Renderer::Private::IPathTracingSpatialTemporalDenoiser* GetActiveSpatialTemporalDenoiser(bool bMatch = false)
 {
 	FString DenoiserName = CVarPathTracingTemporalDenoiserName.GetValueOnRenderThread();
-	return FDenoiserManager::Get().GetSpatialTemporalDenosier(DenoiserName, bMatch);
+	return FDenoiserManager::Get().GetSpatialTemporalDenoiser(DenoiserName, bMatch);
 }
 
 ETextureCreateFlags GetExtraTextureCreateFlagsForDenoiser()
 {
-	ESpatialDenoiserType SpatialDenoiserType = GetSpatialDenosierType();
+	ESpatialDenoiserType SpatialDenoiserType = GetSpatialDenoiserType();
 	ETemporalDenoiserType TemporalDenoiserType = GetTemporalDenoiserType();
 	ETextureCreateFlags TextureCreateFlags = ETextureCreateFlags::None;
 
@@ -1064,7 +1064,7 @@ static void PathTracingDenoiserPlugin(FRDGBuilder& GraphBuilder,
 		using UE::Renderer::Private::IPathTracingSpatialTemporalDenoiser;
 
 		FString DenoiserName = CVarPathTracingDenoiserName.GetValueOnRenderThread();
-		IPathTracingSpatialTemporalDenoiser* ActiveSpatialTemporalDenoiser = FDenoiserManager::Get().GetSpatialTemporalDenosier(DenoiserName,bMatch);
+		IPathTracingSpatialTemporalDenoiser* ActiveSpatialTemporalDenoiser = FDenoiserManager::Get().GetSpatialTemporalDenoiser(DenoiserName,bMatch);
 		
 		if (ActiveSpatialTemporalDenoiser)
 		{
@@ -1966,7 +1966,7 @@ void PathTracingSpatialTemporalDenoising(FRDGBuilder& GraphBuilder,
 
 	FRDGTextureDesc RadianceTextureDesc = SpatialTemporalDenoisingContext.RadianceTexture->Desc;
 
-	const ESpatialDenoiserType SpatialDenoiserType = GetSpatialDenosierType();
+	const ESpatialDenoiserType SpatialDenoiserType = GetSpatialDenoiserType();
 	const ETemporalDenoiserType TemporalDenoiserType = GetTemporalDenoiserType();
 	const bool bApplySpatialDenoiser = ShouldApplySpatialDenoiser();
 	const bool bApplyTemporalDenoiser = ShouldApplyTemporalDenoiser(SpatialTemporalDenoisingContext, View);
diff --git a/Engine/Source/Runtime/Renderer/Public/PathTracingDenoiser.h b/Engine/Source/Runtime/Renderer/Public/PathTracingDenoiser.h
--- a/Engine/Source/Runtime/Renderer/Public/PathTracingDenoiser.h
+++ b/Engine/Source/Runtime/Renderer/Public/PathTracingDenoiser.h
@@ -122,4 +122,4 @@ extern RENDERER_API void RegisterSpatialTemporalDenoiser(TUniquePtr<UE::Renderer
 extern RENDERER_API void UnregisterDenoiser(FString Name);
 
 /** Return true if there is any temporal denoiser plugin registered*/
-extern RENDERER_API bool HasTemporalDenosier();
+extern RENDERER_API bool HasTemporalDenoiser();
