Fix memory not being freed when ray tracing is toggled off.
- RayTracingSceneBufferSRV which holds a reference to RayTracingSceneRHI was not being freed
  - this resulted in TLAS, and BLAS referenced by TLAS also not being freed.
- InstanceUploadBuffer and TransformUploadBuffer were also not being freed.

diff --git a/Engine/Source/Runtime/Renderer/Private/RayTracing/RayTracingScene.cpp b/Engine/Source/Runtime/Renderer/Private/RayTracing/RayTracingScene.cpp
--- a/Engine/Source/Runtime/Renderer/Private/RayTracing/RayTracingScene.cpp
+++ b/Engine/Source/Runtime/Renderer/Private/RayTracing/RayTracingScene.cpp
@@ -576,7 +576,8 @@ void FRayTracingScene::Reset(bool bInInstanceDebugDataEnabled)
 
 		Layer.RayTracingSceneRHI = nullptr;
 		Layer.RayTracingSceneBufferRDG = nullptr;
-		
+		Layer.RayTracingSceneBufferSRV = nullptr;
+
 		Layer.InstanceBuffer = nullptr;
 		Layer.BuildScratchBuffer = nullptr;
 		Layer.InstanceDebugBuffer = nullptr;
@@ -608,12 +609,7 @@ void FRayTracingScene::EndFrame()
 
 		for (uint32 LayerIndex = 0; LayerIndex < NumLayers; ++LayerIndex)
 		{
-			FLayer& Layer = Layers[LayerIndex];
-
-			Layer.Instances.Empty();
-			Layer.InstancesDebugData.Empty();
-
-			Layer.RayTracingScenePooledBuffer = nullptr;
+			Layers[LayerIndex] = {};
 		}
 
 		CallableCommands.Empty();
