Fix for missing directional light shadow in volumetric fog. Also clarified light function usage.

diff --git a/Engine/Source/Runtime/Renderer/Private/VolumetricFog.cpp b/Engine/Source/Runtime/Renderer/Private/VolumetricFog.cpp
--- a/Engine/Source/Runtime/Renderer/Private/VolumetricFog.cpp
+++ b/Engine/Source/Runtime/Renderer/Private/VolumetricFog.cpp
@@ -1393,7 +1393,7 @@ void FSceneRenderer::ComputeVolumetricFog(FRDGBuilder& GraphBuilder,
 
 			if (LightSceneProxy == SelectedForwardDirectionalLightProxy && LightSceneProxy->GetLightType() == LightType_Directional)
 			{
-				LightsToInject.bUseDirectionalLightShadowing = true;
+				LightsToInject.bUseDirectionalLightShadowing = LightSceneProxy->CastsVolumetricShadow();
 
 				if (CheckForLightFunction(LightSceneInfo) && ViewFamily.EngineShowFlags.LightFunctions)
 				{
@@ -1428,7 +1428,7 @@ void FSceneRenderer::ComputeVolumetricFog(FRDGBuilder& GraphBuilder,
 		}
 
 		// Mobile has limited capacities with SRV binding so do not enable atlas sampling on there.
-		const bool bUseLightFunctionAtlas = LightFunctionAtlas::IsEnabled(*Scene, ELightFunctionAtlasSystem::VolumetricFog) && !IsMobilePlatform(View.GetShaderPlatform());
+		const bool bUseLightFunctionAtlasEnabledAndSupported = LightFunctionAtlas::IsEnabled(*Scene, ELightFunctionAtlasSystem::VolumetricFog) && !IsMobilePlatform(View.GetShaderPlatform());
 
 		const bool bUseTemporalReprojection =
 			GVolumetricFogTemporalReprojection
@@ -1468,7 +1468,6 @@ void FSceneRenderer::ComputeVolumetricFog(FRDGBuilder& GraphBuilder,
 
 		// The potential light function for the main directional light is kept separate to be applied during the main VolumetricFogLightScattering pass (as an optimisation).
 		FRDGTexture* DirectionalLightFunctionTexture = GraphBuilder.RegisterExternalTexture(GSystemTextures.WhiteDummy);
-		bool bUseDirectionalLightShadowing = false;
 
 		// Recover the information about the light use as the forward directional light for cloud shadowing
 		int AtmosphericDirectionalLightIndex = -1;
@@ -1684,7 +1683,7 @@ void FSceneRenderer::ComputeVolumetricFog(FRDGBuilder& GraphBuilder,
 
 			PassParameters->PhaseG = FogInfo.VolumetricFogScatteringDistribution;
 			PassParameters->InverseSquaredLightDistanceBiasScale = GInverseSquaredLightDistanceBiasScale;
-			PassParameters->UseDirectionalLightShadowing = bUseDirectionalLightShadowing ? 1.0f : 0.0f;
+			PassParameters->UseDirectionalLightShadowing = LightsToInject.bUseDirectionalLightShadowing ? 1.0f : 0.0f;
 			PassParameters->LightScatteringSampleJitterMultiplier = GVolumetricFogJitter ? GLightScatteringSampleJitterMultiplier : 0;
 			PassParameters->UseHeightFogColors = FVector2f(
 				OverrideDirectionalLightInScatteringUsingHeightFog(View, FogInfo) ? 1.0f : 0.0f,
@@ -1710,10 +1709,10 @@ void FSceneRenderer::ComputeVolumetricFog(FRDGBuilder& GraphBuilder,
 			PassParameters->RaytracedShadowsVolume = RaytracedShadowsVolume ? GraphBuilder.CreateSRV(RaytracedShadowsVolume) : nullptr;
 
 			PassParameters->LightFunctionAtlas = LightFunctionAtlas::BindGlobalParameters(GraphBuilder, View);
-			if (bUseLightFunctionAtlas)
+			if (LightsToInject.DirectionalLightFunction && bUseLightFunctionAtlasEnabledAndSupported)
 			{
-				PassParameters->DirectionalApplyLightFunctionFromAtlas = bUseLightFunctionAtlas && LightsToInject.DirectionalLightFunction && LightsToInject.DirectionalLightFunction->Proxy->HasValidLightFunctionAtlasSlot() ? 1 :0;
-				PassParameters->DirectionalLightFunctionAtlasLightIndex = LightsToInject.DirectionalLightFunction ? LightsToInject.DirectionalLightFunction->Proxy->GetLightFunctionAtlasLightIndex() : 0;
+				PassParameters->DirectionalApplyLightFunctionFromAtlas = LightsToInject.DirectionalLightFunction->Proxy->HasValidLightFunctionAtlasSlot() ? 1 :0;
+				PassParameters->DirectionalLightFunctionAtlasLightIndex = PassParameters->DirectionalApplyLightFunctionFromAtlas == 1 ? LightsToInject.DirectionalLightFunction->Proxy->GetLightFunctionAtlasLightIndex() : 0;
 			}
 			else
 			{
@@ -1748,7 +1747,7 @@ void FSceneRenderer::ComputeVolumetricFog(FRDGBuilder& GraphBuilder,
 			PermutationVector.Set< FVolumetricFogLightScatteringCS::FVirtualShadowMap >(VirtualShadowMapArray.IsAllocated() );
 			PermutationVector.Set< FVolumetricFogLightScatteringCS::FCloudTransmittance >(AtmosphericDirectionalLightIndex >= 0);
 			PermutationVector.Set< FVolumetricFogLightScatteringCS::FRaytracedShadowsVolume >(bUseRaytracedShadowsVolume);
-			PermutationVector.Set< FVolumetricFogLightScatteringCS::FSampleLightFunctionAtlas >(bUseLightFunctionAtlas);
+			PermutationVector.Set< FVolumetricFogLightScatteringCS::FSampleLightFunctionAtlas >(PassParameters->DirectionalApplyLightFunctionFromAtlas>1);
 
 			auto ComputeShader = View.ShaderMap->GetShader< FVolumetricFogLightScatteringCS >(PermutationVector);
 			ClearUnusedGraphResources(ComputeShader, PassParameters);
