Mobile: Fix lighting for non-DefaultLit Materials when using Deferred

diff --git a/Engine/Source/Runtime/Renderer/Private/DepthRendering.cpp b/Engine/Source/Runtime/Renderer/Private/DepthRendering.cpp
--- a/Engine/Source/Runtime/Renderer/Private/DepthRendering.cpp
+++ b/Engine/Source/Runtime/Renderer/Private/DepthRendering.cpp
@@ -777,7 +777,8 @@ bool FDepthPassMeshProcessor::Process(
 	{
 		extern void SetMobileBasePassDepthState(FMeshPassProcessorRenderState& DrawRenderState, const FPrimitiveSceneProxy* PrimitiveSceneProxy, const FMaterial& Material, FMaterialShadingModelField ShadingModels, bool bUsesDeferredShading);
 		
-		FMaterialShadingModelField ShadingModels = MaterialResource.GetShadingModels();
+		// *Don't* get shading models from MaterialResource since it's for a default material
+		FMaterialShadingModelField ShadingModels = MeshBatch.MaterialRenderProxy->GetIncompleteMaterialWithFallback(ERHIFeatureLevel::ES3_1).GetShadingModels();
 		bool bUsesDeferredShading = IsMobileDeferredShadingEnabled(GetFeatureLevelShaderPlatform(FeatureLevel));
 		SetMobileBasePassDepthState(DrawRenderState, PrimitiveSceneProxy, MaterialResource, ShadingModels, bUsesDeferredShading);
 	}
