Fix occlusion culling on mobile

Pass View.PrevViewInfo.HZB into FInstanceCullingContext as is done in SceneRendering (SceneRendering bypasses the basepass for Mobile and instead the pass is rendered here). This wasn't necessary prior to CL 33623275 because the shader was always occluding. However, with the dynamic branch added, we need to pass this in so MergeBatches() sets the occlusion flag.

The potential for this "bug" originated way back in 17175430, I'm not sure why HZB wasn't passed in originally.  Perhaps it was an oversite because MergeBatches didn't depend on it.

diff --git a/Engine/Source/Runtime/Renderer/Private/MobileShadingRenderer.cpp b/Engine/Source/Runtime/Renderer/Private/MobileShadingRenderer.cpp
--- a/Engine/Source/Runtime/Renderer/Private/MobileShadingRenderer.cpp
+++ b/Engine/Source/Runtime/Renderer/Private/MobileShadingRenderer.cpp
@@ -399,7 +399,7 @@ void FMobileSceneRenderer::SetupMobileBasePassAfterShadowInit(FExclusiveDepthSte
 		Pass.DispatchPassSetup(
 			Scene,
 			View,
-			FInstanceCullingContext(PassName, ShaderPlatform, &InstanceCullingManager, ViewIds, nullptr, InstanceCullingMode),
+			FInstanceCullingContext(PassName, ShaderPlatform, &InstanceCullingManager, ViewIds, View.PrevViewInfo.HZB, InstanceCullingMode),
 			EMeshPass::BasePass,
 			BasePassDepthStencilAccess,
 			MeshPassProcessor,
