Fixes FScreenPassTexture::CopyFromSlice() no longer copying from a Texture2DArray to a Texture2D which FScreenPassTexture can only handle

diff --git a/Engine/Source/Runtime/Renderer/Private/ScreenPass.cpp b/Engine/Source/Runtime/Renderer/Private/ScreenPass.cpp
--- a/Engine/Source/Runtime/Renderer/Private/ScreenPass.cpp
+++ b/Engine/Source/Runtime/Renderer/Private/ScreenPass.cpp
@@ -71,7 +71,7 @@ FScreenPassTexture FScreenPassTexture::CopyFromSlice(FRDGBuilder& GraphBuilder,
 	FRDGTexture* InputTexture = InputTextureSRV->Desc.Texture;
 
 	// We can avoid the copy if it's a 2D texture and there's no override output.
-	if (InputTexture->Desc.IsTexture2D() && !OverrideOutput.IsValid())
+	if (InputTexture->Desc.IsTexture2D() && !InputTexture->Desc.IsTextureArray() && !OverrideOutput.IsValid())
 	{
 		return FScreenPassTexture(InputTexture, ScreenTextureSlice.ViewRect);
 	}
