PostProcessing: Raised the priority of propagate alpha in a non-final tonemap pass, so that subsequent operations can retain full precision and reliably use scene color alpha.

diff --git a/Engine/Source/Runtime/Renderer/Private/PostProcess/PostProcessTonemap.cpp b/Engine/Source/Runtime/Renderer/Private/PostProcess/PostProcessTonemap.cpp
--- a/Engine/Source/Runtime/Renderer/Private/PostProcess/PostProcessTonemap.cpp
+++ b/Engine/Source/Runtime/Renderer/Private/PostProcess/PostProcessTonemap.cpp
@@ -616,12 +616,15 @@ FScreenPassTexture AddTonemapPass(FRDGBuilder& GraphBuilder, const FViewInfo& Vi
 		
 		const FTonemapperOutputDeviceParameters OutputDeviceParameters = GetTonemapperOutputDeviceParameters(*View.Family);
 		const EDisplayOutputFormat OutputDevice = static_cast<EDisplayOutputFormat>(OutputDeviceParameters.OutputDevice);
+		const bool bPostProcessingAlpha = IsPostProcessingWithAlphaChannelSupported();
+		// If scene color is high-precision and alpha is supported, we make sure to preserve at least half precision in the alpha channel until the end of post-processing.
+		const bool bPreserveHalfPrecisionAlpha = bPostProcessingAlpha && IsHDR(OutputDesc.Format);
 
 		if (OutputDevice == EDisplayOutputFormat::HDR_LinearEXR)
 		{
 			OutputDesc.Format = PF_A32B32G32R32F;
 		}
-		else if (OutputDevice == EDisplayOutputFormat::HDR_LinearNoToneCurve || OutputDevice == EDisplayOutputFormat::HDR_LinearWithToneCurve)
+		else if (OutputDevice == EDisplayOutputFormat::HDR_LinearNoToneCurve || OutputDevice == EDisplayOutputFormat::HDR_LinearWithToneCurve || bPreserveHalfPrecisionAlpha)
 		{
 			OutputDesc.Format = PF_FloatRGBA;
 		}
@@ -634,7 +637,7 @@ FScreenPassTexture AddTonemapPass(FRDGBuilder& GraphBuilder, const FViewInfo& Vi
 			// Render into a pixel format that do not loose bit depth precision for the view family.
 			OutputDesc.Format = View.Family->RenderTarget->GetRenderTargetTexture()->GetFormat();
 		}
-		else if (IsPostProcessingWithAlphaChannelSupported())
+		else if (bPostProcessingAlpha)
 		{
 			// Make sure there is no loss for a 10bit bit-depth using the 10bit of mantissa of halfs
 			OutputDesc.Format = PF_FloatRGBA;
