Skip LightHasRayTracedShadows call when raytracing is disabled in runtime

#rnx

diff --git a/Engine/Source/Runtime/Renderer/Private/LightGridInjection.cpp b/Engine/Source/Runtime/Renderer/Private/LightGridInjection.cpp
--- a/Engine/Source/Runtime/Renderer/Private/LightGridInjection.cpp
+++ b/Engine/Source/Runtime/Renderer/Private/LightGridInjection.cpp
@@ -595,7 +595,7 @@ FComputeLightGridOutput FSceneRenderer::ComputeLightGrid(FRDGBuilder& GraphBuild
 						LightParameters.Color *= LightParameters.GetLightExposureScale(Exposure);
 
 						float VolumetricScatteringIntensity = LightProxy->GetVolumetricScatteringIntensity();
-						if (LightNeedsSeparateInjectionIntoVolumetricFogForOpaqueShadow(View, LightSceneInfo, VisibleLightInfos[LightSceneInfo->Id]))
+						if (LightNeedsSeparateInjectionIntoVolumetricFogForOpaqueShadow(View, LightSceneInfo, VisibleLightInfos[LightSceneInfo->Id], Scene->bHasLightsWithRayTracedShadows))
 						{
 							// Disable this lights forward shading volumetric scattering contribution
 							VolumetricScatteringIntensity = 0;
diff --git a/Engine/Source/Runtime/Renderer/Private/VolumetricFog.cpp b/Engine/Source/Runtime/Renderer/Private/VolumetricFog.cpp
--- a/Engine/Source/Runtime/Renderer/Private/VolumetricFog.cpp
+++ b/Engine/Source/Runtime/Renderer/Private/VolumetricFog.cpp
@@ -566,7 +566,7 @@ bool LightHasRayTracedShadows(const FLightSceneInfo* LightSceneInfo)
 	return GetLightOcclusionType(*LightSceneInfo->Proxy) == FLightOcclusionType::Raytraced && GVolumetricFogInjectRaytracedLights;
 }
 
-bool LightNeedsSeparateInjectionIntoVolumetricFogForOpaqueShadow(const FViewInfo& View, const FLightSceneInfo* LightSceneInfo, const FVisibleLightInfo& VisibleLightInfo)
+bool LightNeedsSeparateInjectionIntoVolumetricFogForOpaqueShadow(const FViewInfo& View, const FLightSceneInfo* LightSceneInfo, const FVisibleLightInfo& VisibleLightInfo, bool bTestRayTracedShadows)
 {
 	const FLightSceneProxy* LightProxy = LightSceneInfo->Proxy;
 
@@ -579,7 +579,7 @@ bool LightNeedsSeparateInjectionIntoVolumetricFogForOpaqueShadow(const FViewInfo
 		const FStaticShadowDepthMap* StaticShadowDepthMap = LightProxy->GetStaticShadowDepthMap();
 		const bool bStaticallyShadowed = LightSceneInfo->IsPrecomputedLightingValid() && StaticShadowDepthMap && StaticShadowDepthMap->Data && StaticShadowDepthMap->TextureRHI;
 		const bool bHasVirtualShadowMap = VisibleLightInfo.GetVirtualShadowMapId( &View ) != INDEX_NONE;
-		const bool bHasRayTracedShadows = LightHasRayTracedShadows(LightSceneInfo);
+		const bool bHasRayTracedShadows = (bTestRayTracedShadows) ? LightHasRayTracedShadows(LightSceneInfo) : false;
 
 		return GetShadowForInjectionIntoVolumetricFog(VisibleLightInfo) != NULL || bStaticallyShadowed || bHasVirtualShadowMap || bHasRayTracedShadows;
 	}
@@ -711,7 +711,7 @@ void FSceneRenderer::RenderLocalLightsForVolumetricFog(
 			const FIntPoint VolumeZBounds = CalculateVolumetricFogBoundsForLight(LightBounds, View, VolumetricFogViewGridSize, GridZParams);
 			if (VolumeZBounds.X < VolumeZBounds.Y)
 			{
-				bool bIsShadowed = LightNeedsSeparateInjectionIntoVolumetricFogForOpaqueShadow(View, LightSceneInfo, VisibleLightInfo);
+				bool bIsShadowed = LightNeedsSeparateInjectionIntoVolumetricFogForOpaqueShadow(View, LightSceneInfo, VisibleLightInfo, Scene->bHasLightsWithRayTracedShadows);
 				bool bUsesRectLightTexture = GVolumetricFogRectLightTexture && LightSceneInfo->Proxy->HasSourceTexture();
 
 				int32 VirtualShadowMapId = VisibleLightInfo.GetVirtualShadowMapId(&View);
@@ -1366,7 +1366,7 @@ void FSceneRenderer::ComputeVolumetricFog(FRDGBuilder& GraphBuilder,
 
 			if (LightSceneProxy->GetVolumetricScatteringIntensity() > SMALL_NUMBER)
 			{
-				bool bIsShadowed = LightNeedsSeparateInjectionIntoVolumetricFogForOpaqueShadow(View, LightSceneInfo, VisibleLightInfos[LightSceneInfo->Id]);
+				bool bIsShadowed = LightNeedsSeparateInjectionIntoVolumetricFogForOpaqueShadow(View, LightSceneInfo, VisibleLightInfos[LightSceneInfo->Id], Scene->bHasLightsWithRayTracedShadows);
 				bool bUsesRectLightTexture = GVolumetricFogRectLightTexture && LightSceneProxy->HasSourceTexture();
 
 				if (bIsShadowed || bUsesRectLightTexture)
@@ -1375,7 +1375,7 @@ void FSceneRenderer::ComputeVolumetricFog(FRDGBuilder& GraphBuilder,
 
 					if ((View.ViewMatrices.GetViewOrigin() - LightBounds.Center).SizeSquared() < (FogInfo.VolumetricFogDistance + LightBounds.W) * (FogInfo.VolumetricFogDistance + LightBounds.W))
 					{
-						const bool bRayTracedLight = LightHasRayTracedShadows(LightSceneInfo);
+						const bool bRayTracedLight = (Scene->bHasLightsWithRayTracedShadows) ? LightHasRayTracedShadows(LightSceneInfo) : false;
 						if (bRayTracedLight)
 						{
 							LightsToInject.RayTracedLights.Add(LightSceneInfo);
diff --git a/Engine/Source/Runtime/Renderer/Private/VolumetricFog.h b/Engine/Source/Runtime/Renderer/Private/VolumetricFog.h
--- a/Engine/Source/Runtime/Renderer/Private/VolumetricFog.h
+++ b/Engine/Source/Runtime/Renderer/Private/VolumetricFog.h
@@ -26,4 +26,4 @@ extern bool DoesPlatformSupportVolumetricFogVoxelization(const FStaticShaderPlat
 extern bool ShouldRenderVolumetricFog(const FScene* Scene, const FSceneViewFamily& ViewFamily);
 extern const FProjectedShadowInfo* GetShadowForInjectionIntoVolumetricFog(const FVisibleLightInfo& VisibleLightInfo);
 
-extern bool LightNeedsSeparateInjectionIntoVolumetricFogForOpaqueShadow(const FViewInfo& View, const FLightSceneInfo* LightSceneInfo, const FVisibleLightInfo& VisibleLightInfo);
+extern bool LightNeedsSeparateInjectionIntoVolumetricFogForOpaqueShadow(const FViewInfo& View, const FLightSceneInfo* LightSceneInfo, const FVisibleLightInfo& VisibleLightInfo, bool bTestRayTracedShadows);
