Fixed CIS when RHI_RAYTRACING is disabled

#rnx

diff --git a/Engine/Source/Runtime/Renderer/Private/LightGridInjection.cpp b/Engine/Source/Runtime/Renderer/Private/LightGridInjection.cpp
--- a/Engine/Source/Runtime/Renderer/Private/LightGridInjection.cpp
+++ b/Engine/Source/Runtime/Renderer/Private/LightGridInjection.cpp
@@ -595,7 +595,7 @@ FComputeLightGridOutput FSceneRenderer::ComputeLightGrid(FRDGBuilder& GraphBuild
 						LightParameters.Color *= LightParameters.GetLightExposureScale(Exposure);
 
 						float VolumetricScatteringIntensity = LightProxy->GetVolumetricScatteringIntensity();
-						if (LightNeedsSeparateInjectionIntoVolumetricFogForOpaqueShadow(View, LightSceneInfo, VisibleLightInfos[LightSceneInfo->Id], Scene->bHasLightsWithRayTracedShadows))
+						if (LightNeedsSeparateInjectionIntoVolumetricFogForOpaqueShadow(View, LightSceneInfo, VisibleLightInfos[LightSceneInfo->Id], *Scene))
 						{
 							// Disable this lights forward shading volumetric scattering contribution
 							VolumetricScatteringIntensity = 0;
diff --git a/Engine/Source/Runtime/Renderer/Private/VolumetricFog.cpp b/Engine/Source/Runtime/Renderer/Private/VolumetricFog.cpp
--- a/Engine/Source/Runtime/Renderer/Private/VolumetricFog.cpp
+++ b/Engine/Source/Runtime/Renderer/Private/VolumetricFog.cpp
@@ -566,8 +566,13 @@ bool LightHasRayTracedShadows(const FLightSceneInfo* LightSceneInfo)
 	return GetLightOcclusionType(*LightSceneInfo->Proxy) == FLightOcclusionType::Raytraced && GVolumetricFogInjectRaytracedLights;
 }
 
-bool LightNeedsSeparateInjectionIntoVolumetricFogForOpaqueShadow(const FViewInfo& View, const FLightSceneInfo* LightSceneInfo, const FVisibleLightInfo& VisibleLightInfo, bool bTestRayTracedShadows)
+bool LightNeedsSeparateInjectionIntoVolumetricFogForOpaqueShadow(const FViewInfo& View, const FLightSceneInfo* LightSceneInfo, const FVisibleLightInfo& VisibleLightInfo, const FScene& InScene)
 {
+#if RHI_RAYTRACING
+	bool bTestRayTracedShadows = InScene.bHasLightsWithRayTracedShadows;
+#else
+	bool bTestRayTracedShadows = false;
+#endif
 	const FLightSceneProxy* LightProxy = LightSceneInfo->Proxy;
 
 	if (GVolumetricFogInjectShadowedLightsSeparately
@@ -711,7 +716,7 @@ void FSceneRenderer::RenderLocalLightsForVolumetricFog(
 			const FIntPoint VolumeZBounds = CalculateVolumetricFogBoundsForLight(LightBounds, View, VolumetricFogViewGridSize, GridZParams);
 			if (VolumeZBounds.X < VolumeZBounds.Y)
 			{
-				bool bIsShadowed = LightNeedsSeparateInjectionIntoVolumetricFogForOpaqueShadow(View, LightSceneInfo, VisibleLightInfo, Scene->bHasLightsWithRayTracedShadows);
+				bool bIsShadowed = LightNeedsSeparateInjectionIntoVolumetricFogForOpaqueShadow(View, LightSceneInfo, VisibleLightInfo, *Scene);
 				bool bUsesRectLightTexture = GVolumetricFogRectLightTexture && LightSceneInfo->Proxy->HasSourceTexture();
 
 				int32 VirtualShadowMapId = VisibleLightInfo.GetVirtualShadowMapId(&View);
@@ -1366,7 +1371,7 @@ void FSceneRenderer::ComputeVolumetricFog(FRDGBuilder& GraphBuilder,
 
 			if (LightSceneProxy->GetVolumetricScatteringIntensity() > SMALL_NUMBER)
 			{
-				bool bIsShadowed = LightNeedsSeparateInjectionIntoVolumetricFogForOpaqueShadow(View, LightSceneInfo, VisibleLightInfos[LightSceneInfo->Id], Scene->bHasLightsWithRayTracedShadows);
+				bool bIsShadowed = LightNeedsSeparateInjectionIntoVolumetricFogForOpaqueShadow(View, LightSceneInfo, VisibleLightInfos[LightSceneInfo->Id], *Scene);
 				bool bUsesRectLightTexture = GVolumetricFogRectLightTexture && LightSceneProxy->HasSourceTexture();
 
 				if (bIsShadowed || bUsesRectLightTexture)
@@ -1375,12 +1380,14 @@ void FSceneRenderer::ComputeVolumetricFog(FRDGBuilder& GraphBuilder,
 
 					if ((View.ViewMatrices.GetViewOrigin() - LightBounds.Center).SizeSquared() < (FogInfo.VolumetricFogDistance + LightBounds.W) * (FogInfo.VolumetricFogDistance + LightBounds.W))
 					{
+#if RHI_RAYTRACING
 						const bool bRayTracedLight = (Scene->bHasLightsWithRayTracedShadows) ? LightHasRayTracedShadows(LightSceneInfo) : false;
 						if (bRayTracedLight)
 						{
 							LightsToInject.RayTracedLights.Add(LightSceneInfo);
 						}
 						else
+#endif
 						{
 							LightsToInject.Lights.Add(LightSceneInfo);
 						}
diff --git a/Engine/Source/Runtime/Renderer/Private/VolumetricFog.h b/Engine/Source/Runtime/Renderer/Private/VolumetricFog.h
--- a/Engine/Source/Runtime/Renderer/Private/VolumetricFog.h
+++ b/Engine/Source/Runtime/Renderer/Private/VolumetricFog.h
@@ -26,4 +26,4 @@ extern bool DoesPlatformSupportVolumetricFogVoxelization(const FStaticShaderPlat
 extern bool ShouldRenderVolumetricFog(const FScene* Scene, const FSceneViewFamily& ViewFamily);
 extern const FProjectedShadowInfo* GetShadowForInjectionIntoVolumetricFog(const FVisibleLightInfo& VisibleLightInfo);
 
-extern bool LightNeedsSeparateInjectionIntoVolumetricFogForOpaqueShadow(const FViewInfo& View, const FLightSceneInfo* LightSceneInfo, const FVisibleLightInfo& VisibleLightInfo, bool bTestRayTracedShadows);
+extern bool LightNeedsSeparateInjectionIntoVolumetricFogForOpaqueShadow(const FViewInfo& View, const FLightSceneInfo* LightSceneInfo, const FVisibleLightInfo& VisibleLightInfo, const FScene& InScene);
