Add permutation for Sub depth culling in DFS compute shader. Allows to drive the feature  at runtime by a cvar and have it enable/disable per mode or per project. Disabling that feature, saves 0.3% dynres on some plaforms. Increased premutation count by 188 on DistanceFieldShadowingCS shader, no noticable increase in memory.

#tests 666c333a2c2e2e8405b01433 (Switch)

diff --git a/Engine/Source/Runtime/Renderer/Private/DistanceFieldShadowing.cpp b/Engine/Source/Runtime/Renderer/Private/DistanceFieldShadowing.cpp
--- a/Engine/Source/Runtime/Renderer/Private/DistanceFieldShadowing.cpp
+++ b/Engine/Source/Runtime/Renderer/Private/DistanceFieldShadowing.cpp
@@ -133,13 +133,13 @@ static FAutoConsoleVariableRef CVarCompactCulledObjects(
 	TEXT("Note that each tile can only hold up to r.DFShadowAverageObjectsPerCullTile number of objects when compaction is not used."),
 	ECVF_RenderThreadSafe);
 
-int32 GDFShadowCullingSubsampleDepth = 0;
+bool GDFShadowCullingSubsampleDepth = false;
 static FAutoConsoleVariableRef CVarCullingSubsampleDepth(
 	TEXT("r.DFShadowCullingSubsampleDepth"),
 	GDFShadowCullingSubsampleDepth,
 	TEXT("When deciding whether to cull DF shadows for a pixel, subsample the depthbuffer instead of checking all relevant depth texels. ")
 	TEXT("Decreases bandwidth, but produces artifacts on edges and pixel-sized holes"),
-	ECVF_ReadOnly);
+	ECVF_RenderThreadSafe);
 
 int32 const GDistanceFieldShadowTileSizeX = 8;
 int32 const GDistanceFieldShadowTileSizeY = 8;
@@ -320,7 +320,8 @@ class FDistanceFieldShadowingCS : public FGlobalShader
 	class FHasPreviousOutput : SHADER_PERMUTATION_BOOL("HAS_PREVIOUS_OUTPUT");
 	class FOffsetDataStructure : SHADER_PERMUTATION_INT("OFFSET_DATA_STRUCT", 3);
 	class FCompactCulledObjects : SHADER_PERMUTATION_BOOL("COMPACT_CULLED_SHADOW_OBJECTS");
-	using FPermutationDomain = TShaderPermutationDomain<FCullingType, FShadowQuality, FPrimitiveType, FHasPreviousOutput, FOffsetDataStructure, FCompactCulledObjects>;
+	class FCullingSubSampleDepth : SHADER_PERMUTATION_BOOL("CULLING_SUBSAMPLE_DEPTH");
+	using FPermutationDomain = TShaderPermutationDomain<FCullingType, FShadowQuality, FPrimitiveType, FHasPreviousOutput, FOffsetDataStructure, FCompactCulledObjects, FCullingSubSampleDepth>;
 
 	static FPermutationDomain RemapPermutation(FPermutationDomain PermutationVector)
 	{
@@ -375,7 +376,6 @@ class FDistanceFieldShadowingCS : public FGlobalShader
 		OutEnvironment.SetDefine(TEXT("THREADGROUP_SIZEY"), GDistanceFieldShadowTileSizeY);
 		OutEnvironment.SetDefine(TEXT("FORCE_DEPTH_TEXTURE_READS"), 1);
 		OutEnvironment.SetDefine(TEXT("PLATFORM_SUPPORTS_TYPED_UAV_LOAD"), (int32)RHISupports4ComponentUAVReadWrite(Parameters.Platform));
-		OutEnvironment.SetDefine(TEXT("CULLING_SUBSAMPLE_DEPTH"), GDFShadowCullingSubsampleDepth);
 	}
 };
 
@@ -875,6 +875,7 @@ void RayTraceShadows(
 		extern int32 GDistanceFieldOffsetDataStructure;
 		PermutationVector.Set< FDistanceFieldShadowingCS::FOffsetDataStructure >(GDistanceFieldOffsetDataStructure);
 		PermutationVector.Set<FDistanceFieldShadowingCS::FCompactCulledObjects>(GDFShadowCompactCulledObjects != 0);
+		PermutationVector.Set< FDistanceFieldShadowingCS::FCullingSubSampleDepth>(GDFShadowCullingSubsampleDepth);
 
 		PermutationVector = FDistanceFieldShadowingCS::RemapPermutation(PermutationVector);
 
