HWRT: Instance stats buffer should be cleared on the same queue on which we build Instance buffer otherwise we might have unnecessary syncs between queues.
#tests CitySample
#rnx

diff --git a/Engine/Source/Runtime/Renderer/Private/RayTracing/RayTracingScene.cpp b/Engine/Source/Runtime/Renderer/Private/RayTracing/RayTracingScene.cpp
--- a/Engine/Source/Runtime/Renderer/Private/RayTracing/RayTracingScene.cpp
+++ b/Engine/Source/Runtime/Renderer/Private/RayTracing/RayTracingScene.cpp
@@ -183,7 +183,7 @@ void FRayTracingScene::Create(FRDGBuilder& GraphBuilder, const FViewInfo& View,
 
 		FRDGBufferRef OutputStatsBuffer = GraphBuilder.CreateBuffer(OutputStatsBufferDesc, TEXT("FRayTracingScene::OutputStatsBuffer"));
 		FRDGBufferUAVRef OutputStatsBufferUAV = GraphBuilder.CreateUAV(OutputStatsBuffer);
-		AddClearUAVPass(GraphBuilder, OutputStatsBufferUAV, 0);
+		AddClearUAVPass(GraphBuilder, OutputStatsBufferUAV, 0, ComputePassFlags);
 #endif
 
 		FRDGBufferUAVRef DebugInstanceGPUSceneIndexBufferUAV = nullptr;
