Bloom Multiview and Instanced Stereo Fixes

diff --git a/Engine/Plugins/Runtime/OpenXR/Source/OpenXRHMD/Private/OpenXRHMD.cpp b/Engine/Plugins/Runtime/OpenXR/Source/OpenXRHMD/Private/OpenXRHMD.cpp
--- a/Engine/Plugins/Runtime/OpenXR/Source/OpenXRHMD/Private/OpenXRHMD.cpp
+++ b/Engine/Plugins/Runtime/OpenXR/Source/OpenXRHMD/Private/OpenXRHMD.cpp
@@ -4143,7 +4143,7 @@ bool FOpenXRHMD::HasVisibleAreaMesh() const
 	return VisibleAreaMeshes.Num() > 0;
 }
 
-void FOpenXRHMD::DrawHiddenAreaMesh(class FRHICommandList& RHICmdList, int32 ViewIndex) const
+void FOpenXRHMD::DrawHiddenAreaMesh(class FRHICommandList& RHICmdList, int32 ViewIndex, int32 InstanceCount) const
 {
 	check(ViewIndex != INDEX_NONE);
 
@@ -4154,12 +4154,12 @@ void FOpenXRHMD::DrawHiddenAreaMesh(class FRHICommandList& RHICmdList, int32 Vie
 		if (Mesh.IsValid())
 		{
 			RHICmdList.SetStreamSource(0, Mesh.VertexBufferRHI, 0);
-			RHICmdList.DrawIndexedPrimitive(Mesh.IndexBufferRHI, 0, 0, Mesh.NumVertices, 0, Mesh.NumTriangles, 1);
+			RHICmdList.DrawIndexedPrimitive(Mesh.IndexBufferRHI, 0, 0, Mesh.NumVertices, 0, Mesh.NumTriangles, InstanceCount);
 		}
 	}
 }
 
-void FOpenXRHMD::DrawVisibleAreaMesh(class FRHICommandList& RHICmdList, int32 ViewIndex) const
+void FOpenXRHMD::DrawVisibleAreaMesh(class FRHICommandList& RHICmdList, int32 ViewIndex, int32 InstanceCount) const
 {
 	check(ViewIndex != INDEX_NONE);
 	check(ViewIndex < VisibleAreaMeshes.Num());
@@ -4169,7 +4169,7 @@ void FOpenXRHMD::DrawVisibleAreaMesh(class FRHICommandList& RHICmdList, int32 Vi
 		const FHMDViewMesh& Mesh = VisibleAreaMeshes[ViewIndex];
 
 		RHICmdList.SetStreamSource(0, Mesh.VertexBufferRHI, 0);
-		RHICmdList.DrawIndexedPrimitive(Mesh.IndexBufferRHI, 0, 0, Mesh.NumVertices, 0, Mesh.NumTriangles, 1);
+		RHICmdList.DrawIndexedPrimitive(Mesh.IndexBufferRHI, 0, 0, Mesh.NumVertices, 0, Mesh.NumTriangles, InstanceCount);
 	}
 	else
 	{
@@ -4178,6 +4178,16 @@ void FOpenXRHMD::DrawVisibleAreaMesh(class FRHICommandList& RHICmdList, int32 Vi
 	}
 }
 
+void FOpenXRHMD::DrawHiddenAreaMesh(class FRHICommandList& RHICmdList, int32 ViewIndex) const
+{
+	DrawHiddenAreaMesh(RHICmdList, ViewIndex, 1);
+}
+
+void FOpenXRHMD::DrawVisibleAreaMesh(class FRHICommandList& RHICmdList, int32 ViewIndex) const
+{
+	DrawVisibleAreaMesh(RHICmdList, ViewIndex, 1);
+}
+
 void FOpenXRHMD::UpdateLayer(FOpenXRLayer& ManagerLayer, uint32 LayerId, bool bIsValid)
 {
 	ENQUEUE_RENDER_COMMAND(UpdateLayer)(
diff --git a/Engine/Plugins/Runtime/OpenXR/Source/OpenXRHMD/Private/OpenXRHMD.h b/Engine/Plugins/Runtime/OpenXR/Source/OpenXRHMD/Private/OpenXRHMD.h
--- a/Engine/Plugins/Runtime/OpenXR/Source/OpenXRHMD/Private/OpenXRHMD.h
+++ b/Engine/Plugins/Runtime/OpenXR/Source/OpenXRHMD/Private/OpenXRHMD.h
@@ -379,6 +379,8 @@ public:
 	virtual bool HasVisibleAreaMesh() const override final;
 	virtual void DrawHiddenAreaMesh(class FRHICommandList& RHICmdList, int32 ViewIndex) const override final;
 	virtual void DrawVisibleAreaMesh(class FRHICommandList& RHICmdList, int32 ViewIndex) const override final;
+	virtual void DrawHiddenAreaMesh(class FRHICommandList& RHICmdList, int32 ViewIndex, int32 InstanceCount) const override final;
+	virtual void DrawVisibleAreaMesh(class FRHICommandList& RHICmdList, int32 ViewIndex, int32 InstanceCount) const override final;
 	virtual void OnBeginRendering_RenderThread(FRHICommandListImmediate& RHICmdList, FSceneViewFamily& ViewFamily) override;
 	virtual void OnBeginRendering_GameThread() override;
 	virtual void OnLateUpdateApplied_RenderThread(FRHICommandListImmediate& RHICmdList, const FTransform& NewRelativeTransform) override;
diff --git a/Engine/Shaders/Private/PostProcessMobile.usf b/Engine/Shaders/Private/PostProcessMobile.usf
--- a/Engine/Shaders/Private/PostProcessMobile.usf
+++ b/Engine/Shaders/Private/PostProcessMobile.usf
@@ -9,11 +9,15 @@
 #include "PostprocessHistogramCommon.ush"
 #include "ScreenPass.ush"
 
-#if USE_MULTI_VIEW_ID_SV
-Texture2DArray SceneColorTexture;
+#define MOBILE_MULTI_VIEW_FALLBACK (INSTANCED_STEREO && MOBILE_MULTI_VIEW)
+
+#if MOBILE_MULTI_VIEW || MOBILE_MULTI_VIEW_FALLBACK
+#define MultiViewTexture2D Texture2DArray<half4>
 #else
-Texture2D SceneColorTexture;
+#define MultiViewTexture2D Texture2D<half4>
 #endif
+
+MultiViewTexture2D SceneColorTexture;
 SamplerState SceneColorSampler;
 
 Texture2D LastFrameSceneColorTexture;
@@ -31,13 +35,13 @@ SamplerState DofDownSampler;
 Texture2D DofBlurTexture;
 SamplerState DofBlurSampler;
 
-Texture2D BloomDownSourceTexture;
+MultiViewTexture2D BloomDownSourceTexture;
 SamplerState BloomDownSourceSampler;
 
-Texture2D BloomUpSourceATexture;
+MultiViewTexture2D BloomUpSourceATexture;
 SamplerState BloomUpSourceASampler;
 
-Texture2D BloomUpSourceBTexture;
+MultiViewTexture2D BloomUpSourceBTexture;
 SamplerState BloomUpSourceBSampler;
 
 Texture2D SunAlphaTexture;
@@ -46,10 +50,10 @@ SamplerState SunAlphaSampler;
 Texture2D SunBlurTexture;
 SamplerState SunBlurSampler;
 
-Texture2D BloomSetup_BloomTexture;
+MultiViewTexture2D BloomSetup_BloomTexture;
 SamplerState BloomSetup_BloomSampler;
 
-Texture2D BloomUpTexture;
+MultiViewTexture2D BloomUpTexture;
 SamplerState BloomUpSampler;
 
 Texture2D SunMergeTexture;
@@ -136,15 +140,47 @@ half2 SunConstDepthMaskScaleBias()
 	return half2(Scale,-DepthMin * Scale);
 }
 
-float4 SampleSceneColor(float2 SceneUV, float ArrayIndex)
+half4 SampleMultiViewTexture(MultiViewTexture2D Texture, SamplerState Sampler, float2 SceneUV, float ArrayIndex)
 {
-#if USE_MULTI_VIEW_ID_SV
-	return Texture2DArraySample(SceneColorTexture, SceneColorSampler, float3(SceneUV.xy, ArrayIndex));
+#if MOBILE_MULTI_VIEW || MOBILE_MULTI_VIEW_FALLBACK
+	return Texture.Sample(Sampler, float3(SceneUV.xy,ArrayIndex));
 #else
-	return Texture2DSample(SceneColorTexture, SceneColorSampler, SceneUV);
+	return Texture.Sample(Sampler, SceneUV);
 #endif
 }
 
+float4 SampleSceneColor(float2 SceneUV, float ArrayIndex)
+{
+	return SampleMultiViewTexture(SceneColorTexture, SceneColorSampler, SceneUV.xy, ArrayIndex);
+}
+
+half4 SampleBloomDown(float2 SceneUV, float ArrayIndex)
+{
+	return SampleMultiViewTexture(BloomDownSourceTexture, BloomDownSourceSampler, SceneUV.xy, ArrayIndex);
+}
+
+half4 SampleBloomUpA(float2 SceneUV, float ArrayIndex)
+{
+	return SampleMultiViewTexture(BloomUpSourceATexture, BloomUpSourceASampler, SceneUV.xy, ArrayIndex);
+}
+
+// Sampling BloomUpB requires float return type to avoid precision issues leading to artifacts
+// e.g. blue tint all over the view when looking partially at the sky.
+float4 SampleBloomUpB(float2 SceneUV, float ArrayIndex)
+{
+	return SampleMultiViewTexture(BloomUpSourceBTexture, BloomUpSourceBSampler, SceneUV.xy, ArrayIndex);
+}
+
+half4 SampleBloomSetup(float2 SceneUV, float ArrayIndex)
+{
+	return SampleMultiViewTexture(BloomSetup_BloomTexture, BloomSetup_BloomSampler, SceneUV.xy, ArrayIndex);
+}
+
+half4 SampleBloomUp(float2 SceneUV, float ArrayIndex)
+{
+	return SampleMultiViewTexture(BloomUpTexture, BloomUpSampler, SceneUV.xy, ArrayIndex);
+}
+
 //
 // Convert depth in alpha into combined circle of confusion and sun intensity.
 //
@@ -152,21 +188,24 @@ float4 SampleSceneColor(float2 SceneUV, float ArrayIndex)
 #if SHADER_SUN_MASK
 
 void SunMaskPS_Mobile(
-	float4 InUVPos : TEXCOORD0,
-#if USE_MULTI_VIEW_ID_SV
-	in uint MultiViewId : SV_ViewID,
+	float4 InUVPos : TEXCOORD0
+#if MOBILE_MULTI_VIEW_FALLBACK
+	, in nointerpolation uint InViewId : VIEW_ID
+#elif MOBILE_MULTI_VIEW
+	, in uint InViewId : SV_ViewID
 #endif
-	out HALF_TYPE OutSunShaftAndDof : SV_Target0
+	, out HALF_TYPE OutSunShaftAndDof : SV_Target0
 #if MOBILE_USESUN && METAL_MSAA_HDR_DECODE
 	, out HALF4_TYPE OutColor : SV_Target1
 #endif
 	)
 {
-	#if !USE_MULTI_VIEW_ID_SV
-		uint MultiViewId = 0;
-	#endif
 
-	half4 SceneColor = SampleSceneColor(InUVPos.xy,MultiViewId);
+#if !(MOBILE_MULTI_VIEW || MOBILE_MULTI_VIEW_FALLBACK)
+	const uint InViewId = 0;
+#endif
+
+	half4 SceneColor = SampleSceneColor(InUVPos.xy,InViewId);
 
 	#if MOBILE_USEDEPTHTEXTURE 
 		half InDepth = ConvertFromDeviceZ(Texture2DSampleLevel(MobileSceneTextures.SceneDepthTexture, MobileSceneTextures.SceneDepthTextureSampler, InUVPos.xy, 0).r);
@@ -224,9 +263,22 @@ void PreTonemapMSAA_Mobile(
 #endif
 }
 
-Texture2D InputTexture;
+#if MOBILE_MULTI_VIEW || MOBILE_MULTI_VIEW_FALLBACK
+	Texture2DArray InputTexture;
+#else
+	Texture2D InputTexture;
+#endif
 SamplerState InputSampler;
 
+float4 SampleInput(float2 SceneUV, float ArrayIndex)
+{
+#if MOBILE_MULTI_VIEW || MOBILE_MULTI_VIEW_FALLBACK
+	return Texture2DArraySample(InputTexture, InputSampler, float3(SceneUV.xy,ArrayIndex));
+#else
+	return Texture2DSample(InputTexture, InputSampler, SceneUV);
+#endif
+}
+
 void MSAADecodeAndCopyRectPS(
 	noperspective float4 UVAndScreenPos : TEXCOORD0,
 	out HALF4_TYPE OutColor : SV_Target0
@@ -247,9 +299,19 @@ void BloomVS_Mobile(
 	in float4 InPosition : ATTRIBUTE0,
 	in float2 InTexCoord : ATTRIBUTE1,
 	out float2 OutTexCoords[4] : TEXCOORD0,
+#if MOBILE_MULTI_VIEW_FALLBACK
+	in uint InstanceId : SV_InstanceID,
+	out nointerpolation uint OutViewId : VIEW_ID,
+	out uint LayerIndex : SV_RenderTargetArrayIndex,
+#endif 
 	out float4 OutPosition : SV_POSITION
 	)
-{
+{	
+#if MOBILE_MULTI_VIEW_FALLBACK
+	OutViewId = GetEyeIndex(InstanceId);
+	LayerIndex = OutViewId;
+#endif 
+
 	float2 TransformedUV;
 	DrawRectangle(InPosition, InTexCoord, OutPosition, TransformedUV);
 
@@ -278,8 +340,10 @@ float HdrWeightInv4(float3 Color)
 
 void BloomPS_Mobile(
 	float2 InUVs[4] : TEXCOORD0
-#if USE_MULTI_VIEW_ID_SV
-	, in uint MultiViewId : SV_ViewID
+#if MOBILE_MULTI_VIEW_FALLBACK
+	, in nointerpolation uint InViewId : VIEW_ID
+#elif MOBILE_MULTI_VIEW
+	, in nointerpolation uint InViewId : SV_ViewID
 #endif
 #if MOBILE_USEBLOOM
 
@@ -326,15 +390,15 @@ void BloomPS_Mobile(
 	half3 AverageColor = 0.0f;
 
 #if MOBILE_USEBLOOM || MOBILE_USESUN || MOBILE_USEEYEADAPTATION
-
-#if !USE_MULTI_VIEW_ID_SV
-	uint MultiViewId = 0;
-#endif
 	
-	float4 C0 = SampleSceneColor(InUVs[0],MultiViewId);
-	float4 C1 = SampleSceneColor(InUVs[1],MultiViewId);
-	float4 C2 = SampleSceneColor(InUVs[2],MultiViewId);
-	float4 C3 = SampleSceneColor(InUVs[3],MultiViewId);
+#if !(MOBILE_MULTI_VIEW || MOBILE_MULTI_VIEW_FALLBACK)
+	const uint InViewId = 0;
+#endif 
+	
+	float4 C0 = SampleSceneColor(InUVs[0],InViewId);
+	float4 C1 = SampleSceneColor(InUVs[1],InViewId);
+	float4 C2 = SampleSceneColor(InUVs[2],InViewId);
+	float4 C3 = SampleSceneColor(InUVs[3],InViewId);
 
 	C0 *= HdrWeight4(C0);
 	C1 *= HdrWeight4(C1);
@@ -446,9 +510,19 @@ void BloomDownVS_Mobile(
 	in float4 InPosition : ATTRIBUTE0,
 	in float2 InTexCoord : ATTRIBUTE1,
 	out float4 OutTexCoords[8] : TEXCOORD0,
+#if MOBILE_MULTI_VIEW_FALLBACK
+	in uint InstanceId : SV_InstanceID,
+	out nointerpolation uint EyeIndex : VIEW_ID,
+	out uint LayerIndex : SV_RenderTargetArrayIndex,
+#endif
 	out float4 OutPosition : SV_POSITION
 	)
 {
+#if MOBILE_MULTI_VIEW_FALLBACK
+	EyeIndex = GetEyeIndex(InstanceId);
+	LayerIndex = EyeIndex;
+#endif
+
 	DrawRectangle(InPosition, InTexCoord, OutPosition, InTexCoord);
 
 	float Start = 2.0/14.0;
@@ -473,25 +547,34 @@ void BloomDownVS_Mobile(
 }
 
 void BloomDownPS_Mobile(
-	float4 InUVs[8] : TEXCOORD0,
-	out HALF4_TYPE OutColor : SV_Target0
+	float4 InUVs[8] : TEXCOORD0
+#if MOBILE_MULTI_VIEW_FALLBACK
+	, in nointerpolation uint InViewId : VIEW_ID
+#elif MOBILE_MULTI_VIEW
+	, in nointerpolation uint InViewId : SV_ViewID
+#endif
+	, out HALF4_TYPE OutColor : SV_Target0
 	)
 {
-	half4 N0 = BloomDownSourceTexture.Sample(BloomDownSourceSampler, InUVs[0].xy).rgba;
-	half4 N1 = BloomDownSourceTexture.Sample(BloomDownSourceSampler, InUVs[0].zw).rgba;
-	half4 N2 = BloomDownSourceTexture.Sample(BloomDownSourceSampler, InUVs[1].xy).rgba;
-	half4 N3 = BloomDownSourceTexture.Sample(BloomDownSourceSampler, InUVs[1].zw).rgba;
-	half4 N4 = BloomDownSourceTexture.Sample(BloomDownSourceSampler, InUVs[2].xy).rgba;
-	half4 N5 = BloomDownSourceTexture.Sample(BloomDownSourceSampler, InUVs[2].zw).rgba;
-	half4 N6 = BloomDownSourceTexture.Sample(BloomDownSourceSampler, InUVs[3].xy).rgba;
-	half4 N7 = BloomDownSourceTexture.Sample(BloomDownSourceSampler, InUVs[3].zw).rgba;
-	half4 N8 = BloomDownSourceTexture.Sample(BloomDownSourceSampler, InUVs[4].xy).rgba;
-	half4 N9 = BloomDownSourceTexture.Sample(BloomDownSourceSampler, InUVs[4].zw).rgba;
-	half4 N10 = BloomDownSourceTexture.Sample(BloomDownSourceSampler, InUVs[5].xy).rgba;
-	half4 N11 = BloomDownSourceTexture.Sample(BloomDownSourceSampler, InUVs[5].zw).rgba;
-	half4 N12 = BloomDownSourceTexture.Sample(BloomDownSourceSampler, InUVs[6].xy).rgba;
-	half4 N13 = BloomDownSourceTexture.Sample(BloomDownSourceSampler, InUVs[6].zw).rgba;
-	half4 N14 = BloomDownSourceTexture.Sample(BloomDownSourceSampler, InUVs[7].xy).rgba;
+#if !(MOBILE_MULTI_VIEW || MOBILE_MULTI_VIEW_FALLBACK)
+	const uint InViewId = 0;
+#endif 
+
+	half4 N0 = SampleBloomDown(InUVs[0].xy, InViewId).rgba;
+	half4 N1 = SampleBloomDown(InUVs[0].zw, InViewId).rgba;
+	half4 N2 = SampleBloomDown(InUVs[1].xy, InViewId).rgba;
+	half4 N3 = SampleBloomDown(InUVs[1].zw, InViewId).rgba;
+	half4 N4 = SampleBloomDown(InUVs[2].xy, InViewId).rgba;
+	half4 N5 = SampleBloomDown(InUVs[2].zw, InViewId).rgba;
+	half4 N6 = SampleBloomDown(InUVs[3].xy, InViewId).rgba;
+	half4 N7 = SampleBloomDown(InUVs[3].zw, InViewId).rgba;
+	half4 N8 = SampleBloomDown(InUVs[4].xy, InViewId).rgba;
+	half4 N9 = SampleBloomDown(InUVs[4].zw, InViewId).rgba;
+	half4 N10 = SampleBloomDown(InUVs[5].xy, InViewId).rgba;
+	half4 N11 = SampleBloomDown(InUVs[5].zw, InViewId).rgba;
+	half4 N12 = SampleBloomDown(InUVs[6].xy, InViewId).rgba;
+	half4 N13 = SampleBloomDown(InUVs[6].zw, InViewId).rgba;
+	half4 N14 = SampleBloomDown(InUVs[7].xy, InViewId).rgba;
 	float W = 1.0/15.0;
 	OutColor.rgb = 
 		(N0 * W) +
@@ -523,9 +606,19 @@ void BloomUpVS_Mobile(
 	in float4 InPosition : ATTRIBUTE0,
 	in float2 InTexCoord : ATTRIBUTE1,
 	out float4 OutTexCoords[8] : TEXCOORD0,
+#if MOBILE_MULTI_VIEW_FALLBACK
+	in uint InstanceId : SV_InstanceID,
+	out nointerpolation uint EyeIndex : VIEW_ID,
+	out uint LayerIndex : SV_RenderTargetArrayIndex,
+#endif
 	out float4 OutPosition : SV_POSITION
 	)
 {
+#if MOBILE_MULTI_VIEW_FALLBACK
+	EyeIndex = GetEyeIndex(InstanceId);
+	LayerIndex = EyeIndex;
+#endif
+	
 	DrawRectangle(InPosition, InTexCoord, OutPosition, InTexCoord);
 
 	float Start;
@@ -562,27 +655,36 @@ float4 BloomTintA;
 float4 BloomTintB;
 
 void BloomUpPS_Mobile(
-	float4 InUVs[8] : TEXCOORD0,
-	out HALF4_TYPE OutColor : SV_Target0
+	float4 InUVs[8] : TEXCOORD0
+#if MOBILE_MULTI_VIEW_FALLBACK
+	, in nointerpolation uint InViewId : VIEW_ID
+#elif MOBILE_MULTI_VIEW
+	, in nointerpolation uint InViewId : SV_ViewID
+#endif
+	, out HALF4_TYPE OutColor : SV_Target0
 	)
 {
-	half3 A0 = BloomUpSourceATexture.Sample(BloomUpSourceASampler, InUVs[0].xy).rgb;
-	half3 A1 = BloomUpSourceATexture.Sample(BloomUpSourceASampler, InUVs[0].zw).rgb;
-	half3 A2 = BloomUpSourceATexture.Sample(BloomUpSourceASampler, InUVs[1].xy).rgb;
-	half3 A3 = BloomUpSourceATexture.Sample(BloomUpSourceASampler, InUVs[1].zw).rgb;
-	half3 A4 = BloomUpSourceATexture.Sample(BloomUpSourceASampler, InUVs[2].xy).rgb;
-	half3 A5 = BloomUpSourceATexture.Sample(BloomUpSourceASampler, InUVs[2].zw).rgb;
-	half3 A6 = BloomUpSourceATexture.Sample(BloomUpSourceASampler, InUVs[3].xy).rgb;
-	half3 A7 = BloomUpSourceATexture.Sample(BloomUpSourceASampler, InUVs[3].zw).rgb;
-
-	half3 B0 = BloomUpSourceBTexture.Sample(BloomUpSourceBSampler, InUVs[3].zw).rgb;
-	half3 B1 = BloomUpSourceBTexture.Sample(BloomUpSourceBSampler, InUVs[4].xy).rgb;
-	half3 B2 = BloomUpSourceBTexture.Sample(BloomUpSourceBSampler, InUVs[4].zw).rgb;
-	half3 B3 = BloomUpSourceBTexture.Sample(BloomUpSourceBSampler, InUVs[5].xy).rgb;
-	half3 B4 = BloomUpSourceBTexture.Sample(BloomUpSourceBSampler, InUVs[5].zw).rgb;
-	half3 B5 = BloomUpSourceBTexture.Sample(BloomUpSourceBSampler, InUVs[6].xy).rgb;
-	half3 B6 = BloomUpSourceBTexture.Sample(BloomUpSourceBSampler, InUVs[6].zw).rgb;
-	half3 B7 = BloomUpSourceBTexture.Sample(BloomUpSourceBSampler, InUVs[7].xy).rgb;
+#if !(MOBILE_MULTI_VIEW || MOBILE_MULTI_VIEW_FALLBACK)
+	const uint InViewId = 0;
+#endif 
+
+	half3 A0 = SampleBloomUpA(InUVs[0].xy, InViewId).rgb;
+	half3 A1 = SampleBloomUpA(InUVs[0].zw, InViewId).rgb;
+	half3 A2 = SampleBloomUpA(InUVs[1].xy, InViewId).rgb;
+	half3 A3 = SampleBloomUpA(InUVs[1].zw, InViewId).rgb;
+	half3 A4 = SampleBloomUpA(InUVs[2].xy, InViewId).rgb;
+	half3 A5 = SampleBloomUpA(InUVs[2].zw, InViewId).rgb;
+	half3 A6 = SampleBloomUpA(InUVs[3].xy, InViewId).rgb;
+	half3 A7 = SampleBloomUpA(InUVs[3].zw, InViewId).rgb;
+										   
+	half3 B0 = SampleBloomUpB(InUVs[3].zw, InViewId).rgb;
+	half3 B1 = SampleBloomUpB(InUVs[4].xy, InViewId).rgb;
+	half3 B2 = SampleBloomUpB(InUVs[4].zw, InViewId).rgb;
+	half3 B3 = SampleBloomUpB(InUVs[5].xy, InViewId).rgb;
+	half3 B4 = SampleBloomUpB(InUVs[5].zw, InViewId).rgb;
+	half3 B5 = SampleBloomUpB(InUVs[6].xy, InViewId).rgb;
+	half3 B6 = SampleBloomUpB(InUVs[6].zw, InViewId).rgb;
+	half3 B7 = SampleBloomUpB(InUVs[7].xy, InViewId).rgb;
 
 	// A is the same size source.
 	half3 WA = BloomTintA.rgb;
@@ -619,9 +721,19 @@ void DofNearVS_Mobile(
 	in float2 InTexCoord : ATTRIBUTE1,
 	out float2 OutTexCoords2 : TEXCOORD0,
 	out float4 OutTexCoords4[4] : TEXCOORD1,
+#if MOBILE_MULTI_VIEW_FALLBACK
+	in uint InstanceId : SV_InstanceID,
+	out nointerpolation uint EyeIndex : VIEW_ID,
+	out uint LayerIndex : SV_RenderTargetArrayIndex,
+#endif
 	out float4 OutPosition : SV_POSITION
 	)
 {
+#if MOBILE_MULTI_VIEW_FALLBACK
+	EyeIndex = GetEyeIndex(InstanceId);
+	LayerIndex = EyeIndex;
+#endif
+
 	DrawRectangle(InPosition, InTexCoord, OutPosition, InTexCoord);
 
 	OutTexCoords2 = InTexCoord;
@@ -636,11 +748,20 @@ void DofNearVS_Mobile(
 }
 
 void DofNearPS_Mobile(
-	float2 InUVs2 : TEXCOORD0,
-	float4 InUVs[4] : TEXCOORD1,
-	out HALF_TYPE OutColor : SV_Target0
+	float2 InUVs2 : TEXCOORD0
+	, float4 InUVs[4] : TEXCOORD1
+#if MOBILE_MULTI_VIEW_FALLBACK
+	, in nointerpolation uint InViewId : VIEW_ID
+#elif MOBILE_MULTI_VIEW
+	, in nointerpolation uint InViewId : SV_ViewID
+#endif
+	, out HALF_TYPE OutColor : SV_Target0
 	)
 {
+#if !(MOBILE_MULTI_VIEW || MOBILE_MULTI_VIEW_FALLBACK)
+	const uint InViewId = 0;
+#endif 
+
 	half N0 = SunShaftAndDofTexture.Sample(SunShaftAndDofSampler, InUVs2).r;
 	half N1 = SunShaftAndDofTexture.Sample(SunShaftAndDofSampler, InUVs[0].xy).r;
 	half N2 = SunShaftAndDofTexture.Sample(SunShaftAndDofSampler, InUVs[0].zw).r;
@@ -692,9 +813,19 @@ void DofDownVS_Mobile(
 	in float4 InPosition : ATTRIBUTE0,
 	in float2 InTexCoord : ATTRIBUTE1,
 	out float2 OutTexCoords[5] : TEXCOORD0,
+#if MOBILE_MULTI_VIEW_FALLBACK
+	in uint InstanceId : SV_InstanceID,
+	out nointerpolation uint EyeIndex : VIEW_ID,
+	out uint LayerIndex : SV_RenderTargetArrayIndex,
+#endif
 	out float4 OutPosition : SV_POSITION
 	)
 {
+#if MOBILE_MULTI_VIEW_FALLBACK
+	EyeIndex = GetEyeIndex(InstanceId);
+	LayerIndex = EyeIndex;
+#endif
+
 	DrawRectangle(InPosition, InTexCoord, OutPosition, InTexCoord);
 
 	// Near position fixed to use UV based out output position.
@@ -707,11 +838,13 @@ void DofDownVS_Mobile(
 }
 
 void DofDownPS_Mobile(
-	float2 InUVs[5] : TEXCOORD0,
-#if USE_MULTI_VIEW_ID_SV
-	in uint MultiViewId : SV_ViewID,
+	float2 InUVs[5] : TEXCOORD0
+#if MOBILE_MULTI_VIEW_FALLBACK
+	, in nointerpolation uint InViewId : VIEW_ID
+#elif MOBILE_MULTI_VIEW
+	, in nointerpolation uint InViewId : SV_ViewID
 #endif
-	out HALF4_TYPE OutColor : SV_Target0
+	, out HALF4_TYPE OutColor : SV_Target0
 	)
 {
 	// This shader needs float precision to work.
@@ -719,14 +852,14 @@ void DofDownPS_Mobile(
 	// Fetch near diolation and scale to (0 to 16384.0) range.
 	float N = DofNearTexture.Sample(DofNearSampler, InUVs[0]).r * 16384.0;
 
-#if !USE_MULTI_VIEW_ID_SV
-	uint MultiViewId = 0;
-#endif
-	
-	float4 A = SampleSceneColor(InUVs[1],MultiViewId);
-	float4 B = SampleSceneColor(InUVs[2],MultiViewId);
-	float4 C = SampleSceneColor(InUVs[3],MultiViewId);
-	float4 D = SampleSceneColor(InUVs[4],MultiViewId);
+#if !(MOBILE_MULTI_VIEW || MOBILE_MULTI_VIEW_FALLBACK)
+	const uint InViewId = 0;
+#endif 
+
+	float4 A = SampleSceneColor(InUVs[1], InViewId);
+	float4 B = SampleSceneColor(InUVs[2], InViewId);
+	float4 C = SampleSceneColor(InUVs[3], InViewId);
+	float4 D = SampleSceneColor(InUVs[4], InViewId);
 
 	A.a = SunShaftAndDofTexture.Sample(SunShaftAndDofSampler, InUVs[1]).r;
 	B.a = SunShaftAndDofTexture.Sample(SunShaftAndDofSampler, InUVs[2]).r;
@@ -836,9 +969,19 @@ void DofBlurVS_Mobile(
 	in float4 InPosition : ATTRIBUTE0,
 	in float2 InTexCoord : ATTRIBUTE1,
 	out float2 OutTexCoords[5] : TEXCOORD0,
+#if MOBILE_MULTI_VIEW_FALLBACK
+	in uint InstanceId : SV_InstanceID,
+	out nointerpolation uint EyeIndex : VIEW_ID,
+	out uint LayerIndex : SV_RenderTargetArrayIndex,
+#endif
 	out float4 OutPosition : SV_POSITION
 	)
 {
+#if MOBILE_MULTI_VIEW_FALLBACK
+	EyeIndex = GetEyeIndex(InstanceId);
+	LayerIndex = EyeIndex;
+#endif
+
 	DrawRectangle(InPosition, InTexCoord, OutPosition, InTexCoord);
 
     OutTexCoords[0] = InTexCoord.xy;
@@ -849,10 +992,19 @@ void DofBlurVS_Mobile(
 }
 
 void DofBlurPS_Mobile(
-	float2 InUVs[5] : TEXCOORD0,
-	out HALF4_TYPE OutColor : SV_Target0
+	float2 InUVs[5] : TEXCOORD0
+#if MOBILE_MULTI_VIEW_FALLBACK
+	, in nointerpolation uint InViewId : VIEW_ID
+#elif MOBILE_MULTI_VIEW
+	, in nointerpolation uint InViewId : SV_ViewID
+#endif
+	, out HALF4_TYPE OutColor : SV_Target0
 	)
 {
+#if !(MOBILE_MULTI_VIEW || MOBILE_MULTI_VIEW_FALLBACK)
+	const uint InViewId = 0;
+#endif 
+
 	// Near diolation size is copied into alpha for the tonemapper pass.
 	OutColor.a = DofNearTexture.Sample(DofNearSampler, InUVs[0]).r;
 
@@ -887,9 +1039,19 @@ void IntegrateDOFVS_Mobile(
 	in float2 InTexCoord : ATTRIBUTE1,
 	out float4 OutTexCoords : TEXCOORD0,
 	out float2 OutFineDofGrain : TEXCOORD1,
+#if MOBILE_MULTI_VIEW_FALLBACK
+	in uint InstanceId : SV_InstanceID,
+	out nointerpolation uint EyeIndex : VIEW_ID,
+	out uint LayerIndex : SV_RenderTargetArrayIndex,
+#endif
 	out float4 OutPosition : SV_POSITION
 )
 {
+#if MOBILE_MULTI_VIEW_FALLBACK
+	EyeIndex = GetEyeIndex(InstanceId);
+	LayerIndex = EyeIndex;
+#endif
+
 	DrawRectangle(InPosition, InTexCoord, OutPosition, OutTexCoords.xy);
 
 	// Fine adjustment is inside the possible non-full viewport in the full resolution texture.
@@ -903,21 +1065,23 @@ void IntegrateDOFVS_Mobile(
 }
 
 void IntegrateDOFPS_Mobile(
-	in float4 TexCoords : TEXCOORD0,
-	in float2 FineDofGrain : TEXCOORD1,
-#if USE_MULTI_VIEW_ID_SV
-	in uint MultiViewId : SV_ViewID,
+	in float4 TexCoords : TEXCOORD0
+	, in float2 FineDofGrain : TEXCOORD1
+#if MOBILE_MULTI_VIEW_FALLBACK
+	, in nointerpolation uint InViewId : VIEW_ID
+#elif MOBILE_MULTI_VIEW
+	, in nointerpolation uint InViewId : SV_ViewID
 #endif
-	out HALF4_TYPE OutColor : SV_Target0
+	, out HALF4_TYPE OutColor : SV_Target0
 )
 {
-	#if !USE_MULTI_VIEW_ID_SV
-		uint MultiViewId = 0;;
-	#endif
-	
-	half4 SceneColor = SampleSceneColor(TexCoords.xy,MultiViewId);
-	half4 DofFine = SampleSceneColor(FineDofGrain.xy,MultiViewId);
+#if !(MOBILE_MULTI_VIEW || MOBILE_MULTI_VIEW_FALLBACK)
+	const uint InViewId = 0;
+#endif 
 	
+	half4 SceneColor = SampleSceneColor(TexCoords.xy, InViewId);
+	half4 DofFine = SampleSceneColor(FineDofGrain.xy, InViewId);
+
 	half4 Dof = Texture2DSample(DofBlurTexture, DofBlurSampler, TexCoords.zw);
 	half DofCoc = Texture2DSample(SunShaftAndDofTexture, SunShaftAndDofSampler, TexCoords.xy).r;
 	// Convert alpha back into circle of confusion.
@@ -989,9 +1153,19 @@ void SunAlphaVS_Mobile(
 	in float4 InPosition : ATTRIBUTE0,
 	in float2 InTexCoord : ATTRIBUTE1,
 	out float2 OutTexCoords[8] : TEXCOORD0,
+#if MOBILE_MULTI_VIEW_FALLBACK
+	in uint InstanceId : SV_InstanceID,
+	out nointerpolation uint EyeIndex : VIEW_ID,
+	out uint LayerIndex : SV_RenderTargetArrayIndex,
+#endif
 	out float4 OutPosition : SV_POSITION
 	)
 {
+#if MOBILE_MULTI_VIEW_FALLBACK
+	EyeIndex = GetEyeIndex(InstanceId);
+	LayerIndex = EyeIndex;
+#endif
+
 	DrawRectangle(InPosition, InTexCoord, OutPosition, InTexCoord);
 
     OutTexCoords[0] = SunShaftRect(OutPosition.xy, 1.0 - SUN_P0 * SUN_M);
@@ -1018,10 +1192,19 @@ half SunUnBias(half A)
 }
 
 void SunAlphaPS_Mobile(
-	float2 InUVs[8] : TEXCOORD0,
-	out HALF_TYPE OutColor : SV_Target0
+	float2 InUVs[8] : TEXCOORD0
+#if MOBILE_MULTI_VIEW_FALLBACK
+	, in nointerpolation uint InViewId : VIEW_ID
+#elif MOBILE_MULTI_VIEW
+	, in nointerpolation uint InViewId : SV_ViewID
+#endif
+	, out HALF_TYPE OutColor : SV_Target0
 	)
 {
+#if !(MOBILE_MULTI_VIEW || MOBILE_MULTI_VIEW_FALLBACK)
+	const uint InViewId = 0;
+#endif 
+
 	OutColor = 
 		SunUnBias(SunShaftAndDofTexture.Sample(SunShaftAndDofSampler, InUVs[0]).r) * 0.125 +
 		SunUnBias(SunShaftAndDofTexture.Sample(SunShaftAndDofSampler, InUVs[1]).r) * 0.125 +
@@ -1047,9 +1230,19 @@ void SunBlurVS_Mobile(
 	in float4 InPosition : ATTRIBUTE0,
 	in float2 InTexCoord : ATTRIBUTE1,
 	out float2 OutTexCoords[8] : TEXCOORD0,
+#if MOBILE_MULTI_VIEW_FALLBACK
+	in uint InstanceId : SV_InstanceID,
+	out nointerpolation uint EyeIndex : VIEW_ID,
+	out uint LayerIndex : SV_RenderTargetArrayIndex,
+#endif
 	out float4 OutPosition : SV_POSITION
 	)
 {
+#if MOBILE_MULTI_VIEW_FALLBACK
+	EyeIndex = GetEyeIndex(InstanceId);
+	LayerIndex = EyeIndex;
+#endif
+
 	DrawRectangle(InPosition, InTexCoord, OutPosition, InTexCoord);
 
     OutTexCoords[0] = SunShaftRect(OutPosition.xy, 1.0 - SUN_P0 * SUN_M);
@@ -1065,10 +1258,19 @@ void SunBlurVS_Mobile(
 #undef SUN_M
 
 void SunBlurPS_Mobile(
-	float2 InUVs[8] : TEXCOORD0,
-	out HALF_TYPE OutColor : SV_Target0
+	float2 InUVs[8] : TEXCOORD0
+#if MOBILE_MULTI_VIEW_FALLBACK
+	, in nointerpolation uint InViewId : VIEW_ID
+#elif MOBILE_MULTI_VIEW
+	, in nointerpolation uint InViewId : SV_ViewID
+#endif
+	, out HALF_TYPE OutColor : SV_Target0
 	)
 {
+#if !(MOBILE_MULTI_VIEW || MOBILE_MULTI_VIEW_FALLBACK)
+	const uint InViewId = 0;
+#endif 
+
 	OutColor = 
 		SunAlphaTexture.Sample(SunAlphaSampler, InUVs[0]).r * 0.125 +
 		SunAlphaTexture.Sample(SunAlphaSampler, InUVs[1]).r * 0.125 +
@@ -1092,9 +1294,19 @@ void SunMergeVS_Mobile(
 	in float2 InTexCoord : ATTRIBUTE1,
 	out float4 OutTexCoordVignette : TEXCOORD0,
 	out float4 OutTexCoords[7] : TEXCOORD1,
+#if MOBILE_MULTI_VIEW_FALLBACK
+	in uint InstanceId : SV_InstanceID,
+	out nointerpolation uint EyeIndex : VIEW_ID,
+	out uint LayerIndex : SV_RenderTargetArrayIndex,
+#endif
 	out float4 OutPosition : SV_POSITION
 	)
 {
+#if MOBILE_MULTI_VIEW_FALLBACK
+	EyeIndex = GetEyeIndex(InstanceId);
+	LayerIndex = EyeIndex;
+#endif
+	
 	DrawRectangle(InPosition, InTexCoord, OutPosition, InTexCoord);
 
 	OutTexCoordVignette.xy = InTexCoord.xy;
@@ -1134,11 +1346,20 @@ SamplerState BloomDirtMaskSampler;
 float4 BloomDirtMaskTint;
 
 void SunMergePS_Mobile(
-	float4 InUVVignette : TEXCOORD0,
-	float4 InUVs[7] : TEXCOORD1,
-	out HALF4_TYPE OutColor : SV_Target0
+	float4 InUVVignette : TEXCOORD0
+	, float4 InUVs[7] : TEXCOORD1
+#if MOBILE_MULTI_VIEW_FALLBACK
+	, in nointerpolation uint InViewId : VIEW_ID
+#elif MOBILE_MULTI_VIEW
+	, in nointerpolation uint InViewId : SV_ViewID
+#endif
+	, out HALF4_TYPE OutColor : SV_Target0
 	)
 {
+#if !(MOBILE_MULTI_VIEW || MOBILE_MULTI_VIEW_FALLBACK)
+	const uint InViewId = 0;
+#endif 
+
 	#if MOBILE_USEBLOOM
 
 		float Scale1 = 1.0/7.0;
@@ -1146,15 +1367,15 @@ void SunMergePS_Mobile(
 		float Scale2 = 1.0/7.0;
 
 		half3 Bloom2 = (
-			BloomSetup_BloomTexture.Sample(BloomSetup_BloomSampler, InUVVignette.xy).rgba *Scale1 +
-			BloomSetup_BloomTexture.Sample(BloomSetup_BloomSampler, InUVs[0].xy).rgba * Scale2 +
-			BloomSetup_BloomTexture.Sample(BloomSetup_BloomSampler, InUVs[1].xy).rgba * Scale2 +
-			BloomSetup_BloomTexture.Sample(BloomSetup_BloomSampler, InUVs[2].xy).rgba * Scale2 +
-			BloomSetup_BloomTexture.Sample(BloomSetup_BloomSampler, InUVs[3].xy).rgba * Scale2 +
-			BloomSetup_BloomTexture.Sample(BloomSetup_BloomSampler, InUVs[4].xy).rgba * Scale2 +
-			BloomSetup_BloomTexture.Sample(BloomSetup_BloomSampler, InUVs[5].xy).rgba * Scale2) * rcp(Scale1 * 1.0 + Scale2 * 6.0);
+			SampleBloomSetup(InUVVignette.xy, InViewId).rgba *Scale1 +
+			SampleBloomSetup(InUVs[0].xy, InViewId).rgba * Scale2 +
+			SampleBloomSetup(InUVs[1].xy, InViewId).rgba * Scale2 +
+			SampleBloomSetup(InUVs[2].xy, InViewId).rgba * Scale2 +
+			SampleBloomSetup(InUVs[3].xy, InViewId).rgba * Scale2 +
+			SampleBloomSetup(InUVs[4].xy, InViewId).rgba * Scale2 +
+			SampleBloomSetup(InUVs[5].xy, InViewId).rgba * Scale2) * rcp(Scale1 * 1.0 + Scale2 * 6.0);
 
-		OutColor.rgb = BloomUpTexture.Sample(BloomUpSampler, InUVVignette.xy);
+		OutColor.rgb = SampleBloomUp(InUVVignette.xy, InViewId);
 
 		half3 BloomDirtMaskColor = BloomDirtMaskTexture.Sample(BloomDirtMaskSampler, InUVVignette.xy).rgb * BloomDirtMaskTint.rgb;
 
@@ -1267,7 +1488,7 @@ void AverageLuminance_MainCS(
 
 				Weight = max(AdaptationWeightTexture(BufferUV), 0.05f);
 
-				GroupLuminance.x += Texture2DSample(InputTexture, InputSampler, BufferUV).x * Weight;
+				GroupLuminance.x += SampleInput(BufferUV, 0).x * Weight;
 				GroupLuminance.y += Weight;
 			}
 		}
@@ -1423,7 +1644,7 @@ void Histogram_MainCS(
 				BufferUV = (float2)TexelPos + float2(1.0f, 1.0f); // offset pixel pos by 1 to use bilinear filter
 				BufferUV = BufferUV * SourceSizeAndInvSize.zw;
 
-				LogLuminance = Texture2DSample(InputTexture, InputSampler, BufferUV).x;
+				LogLuminance = SampleInput(BufferUV, 0).x;
 				ScreenWeight = AdaptationWeightTexture(BufferUV);
 
 				// Map the normalized histogram position into texels.
diff --git a/Engine/Shaders/Private/PostProcessTonemap.usf b/Engine/Shaders/Private/PostProcessTonemap.usf
--- a/Engine/Shaders/Private/PostProcessTonemap.usf
+++ b/Engine/Shaders/Private/PostProcessTonemap.usf
@@ -30,7 +30,7 @@ SCREEN_PASS_TEXTURE_VIEWPORT(Output)
 
 #define MOBILE_MULTI_VIEW_FALLBACK (MOBILE_MULTI_VIEW && INSTANCED_STEREO)
 
-#if USE_MULTI_VIEW_ID_SV || MOBILE_MULTI_VIEW_FALLBACK
+#if MOBILE_MULTI_VIEW || MOBILE_MULTI_VIEW_FALLBACK
 Texture2DArray ColorTexture;
 #else
 Texture2D ColorTexture;
@@ -41,7 +41,12 @@ SamplerState ColorSampler;
 FScreenTransform ColorToBloom;
 float2 BloomUVViewportBilinearMin;
 float2 BloomUVViewportBilinearMax;
-Texture2D BloomTexture;
+
+#if MOBILE_MULTI_VIEW || MOBILE_MULTI_VIEW_FALLBACK
+	Texture2DArray BloomTexture;
+#else
+	Texture2D BloomTexture;
+#endif
 SamplerState BloomSampler;
 
 #if SUPPORTS_SCENE_COLOR_APPLY_PARAMETERS
@@ -242,10 +247,20 @@ void MainVS(
 	out noperspective 	float2 OutVignette				: TEXCOORD1,
 	out noperspective 	float4 OutGrainUV 				: TEXCOORD2,
 	out noperspective 	float2 OutScreenPos 			: TEXCOORD3,
-	out noperspective 	float2 OutFullViewUV 			: TEXCOORD4,
-	out 				float4 OutPosition 				: SV_POSITION
+	out noperspective 	float2 OutFullViewUV 			: TEXCOORD4
+#if MOBILE_MULTI_VIEW_FALLBACK
+	, in uint InstanceId : SV_InstanceID
+	, out nointerpolation uint OutViewId : VIEW_ID
+	, out uint LayerIndex : SV_RenderTargetArrayIndex
+#endif
+	, out 				float4 OutPosition 				: SV_POSITION
 	)
 {
+#if MOBILE_MULTI_VIEW_FALLBACK
+	OutViewId = GetEyeIndex(InstanceId);
+	LayerIndex = OutViewId;
+#endif
+
 	DrawRectangle(InPosition, InTexCoord, OutPosition, OutTexCoord);
     TonemapCommonVS(OutPosition, OutTexCoord, OutVignette, OutGrainUV, OutScreenPos, OutFullViewUV);
 }
@@ -258,11 +273,11 @@ float max4(float4 x)
 float4 SampleSceneColor(float2 SceneUV, float ArrayIndex)
 {
 	SceneUV = clamp(SceneUV, Color_UVViewportBilinearMin, Color_UVViewportBilinearMax);
-	#if USE_MULTI_VIEW_ID_SV || MOBILE_MULTI_VIEW_FALLBACK
+#if MOBILE_MULTI_VIEW || MOBILE_MULTI_VIEW_FALLBACK
 	return Texture2DArraySample(ColorTexture, ColorSampler, float3(SceneUV.xy,ArrayIndex));
-	#else
+#else
 	return Texture2DSample(ColorTexture, ColorSampler, SceneUV);
-	#endif
+#endif
 }
 
 half3 LookupSceneColor(float2 SceneUV, float ArrayIndex, float2 PixelOffset)
@@ -347,7 +362,12 @@ float4 TonemapCommonPS(
 		}
 		#endif
 
-		float4 RawBloom = Texture2DSample(BloomTexture, BloomSampler, BloomUV);
+		float4 RawBloom;
+		#if MOBILE_MULTI_VIEW || MOBILE_MULTI_VIEW_FALLBACK
+			RawBloom = Texture2DArraySample(BloomTexture, BloomSampler, float3(BloomUV,ArrayIndex));
+		#else
+			RawBloom = Texture2DSample(BloomTexture, BloomSampler, BloomUV);
+		#endif
 
 		#if FEATURE_LEVEL == FEATURE_LEVEL_ES3_1
 		{
@@ -586,9 +606,11 @@ void MainPS(
 	in noperspective float4 GrainUV : TEXCOORD2,
 	in noperspective float2 ScreenPos : TEXCOORD3,
 	in noperspective float2 FullViewUV : TEXCOORD4,
-	#if USE_MULTI_VIEW_ID_SV
-	in uint MultiViewId : SV_ViewID,
-	#endif
+#if MOBILE_MULTI_VIEW_FALLBACK
+	in nointerpolation uint InViewId : VIEW_ID,
+#elif MOBILE_MULTI_VIEW
+	in uint InViewId : SV_ViewID,
+#endif
 	float4 SvPosition : SV_POSITION,		// after all interpolators
 	out float4 OutColor : SV_Target0
 #if OUTPUT_LUMINANCE
@@ -596,13 +618,13 @@ void MainPS(
 #endif
 	)
 {
+#if !(MOBILE_MULTI_VIEW || MOBILE_MULTI_VIEW_FALLBACK)
+	const uint InViewId = 0;
+#endif
+
 	float Luminance;
-	#if USE_MULTI_VIEW_ID_SV
-		float3 UVI = float3(UV.xy,MultiViewId);
-	#else
-		float3 UVI = float3(UV.xy,0);
-	#endif
-    OutColor = TonemapCommonPS(UVI, InVignette, GrainUV, ScreenPos, FullViewUV, SvPosition, Luminance);
+	float3 UVI = float3(UV.xy, InViewId);
+	OutColor = TonemapCommonPS(UVI, InVignette, GrainUV, ScreenPos, FullViewUV, SvPosition, Luminance);
 #if OUTPUT_LUMINANCE
 	OutLuminance = Luminance;
 #endif
diff --git a/Engine/Source/Runtime/HeadMountedDisplay/Public/IHeadMountedDisplay.h b/Engine/Source/Runtime/HeadMountedDisplay/Public/IHeadMountedDisplay.h
--- a/Engine/Source/Runtime/HeadMountedDisplay/Public/IHeadMountedDisplay.h
+++ b/Engine/Source/Runtime/HeadMountedDisplay/Public/IHeadMountedDisplay.h
@@ -214,12 +214,26 @@ public:
 	*/
 	virtual void DrawHiddenAreaMesh(class FRHICommandList& RHICmdList, int32 ViewIndex) const {};
 
+	/**
+	* Optional method to draw a view's hidden area mesh where supported.
+	* This can be used to avoid rendering pixels which are not included as input into the final distortion pass.
+	* Supports instanced rendering.
+	*/
+	virtual void DrawHiddenAreaMesh(class FRHICommandList& RHICmdList, int32 ViewIndex, int32 InstanceCount) const {};
+
 	/**
 	* Optional method to draw a view's visible area mesh where supported.
 	* This can be used instead of a full screen quad to avoid rendering pixels which are not included as input into the final distortion pass.
 	*/
 	virtual void DrawVisibleAreaMesh(class FRHICommandList& RHICmdList, int32 ViewIndex) const {};
 
+	/**
+	* Optional method to draw a view's visible area mesh where supported.
+	* This can be used instead of a full screen quad to avoid rendering pixels which are not included as input into the final distortion pass.
+	* Supports instanced rendering.
+	*/
+	virtual void DrawVisibleAreaMesh(class FRHICommandList& RHICmdList, int32 ViewIndex, int32 InstanceCount) const {};
+
 	virtual void DrawDistortionMesh_RenderThread(struct FHeadMountedDisplayPassContext& Context, const FIntPoint& TextureSize) {}
 
 	/**
diff --git a/Engine/Source/Runtime/Renderer/Private/PostProcess/PostProcessMobile.cpp b/Engine/Source/Runtime/Renderer/Private/PostProcess/PostProcessMobile.cpp
--- a/Engine/Source/Runtime/Renderer/Private/PostProcess/PostProcessMobile.cpp
+++ b/Engine/Source/Runtime/Renderer/Private/PostProcess/PostProcessMobile.cpp
@@ -183,7 +183,7 @@ public:
 	SHADER_USE_PARAMETER_STRUCT_WITH_LEGACY_BASE(FMobileBloomSetupVS, FGlobalShader);
 
 	BEGIN_SHADER_PARAMETER_STRUCT(FParameters, )
-		SHADER_PARAMETER_STRUCT_REF(FViewUniformShaderParameters, View)
+		SHADER_PARAMETER_STRUCT_INCLUDE(FViewShaderParameters, View)
 		SHADER_PARAMETER(FVector4f, BufferSizeAndInvSize)
 	END_SHADER_PARAMETER_STRUCT()
 
@@ -217,7 +217,7 @@ public:
 
 	BEGIN_SHADER_PARAMETER_STRUCT(FParameters, )
 		SHADER_PARAMETER(float, BloomThreshold)
-		SHADER_PARAMETER_STRUCT_REF(FViewUniformShaderParameters, View)
+		SHADER_PARAMETER_STRUCT_INCLUDE(FViewShaderParameters, View)
 		SHADER_PARAMETER_STRUCT(FEyeAdaptationParameters, EyeAdaptation)
 		SHADER_PARAMETER_RDG_TEXTURE(Texture2D, SceneColorTexture)
 		SHADER_PARAMETER_SAMPLER(SamplerState, SceneColorSampler)
@@ -283,7 +283,9 @@ FMobileBloomSetupOutputs AddMobileBloomSetupPass(FRDGBuilder& GraphBuilder, cons
 
 	TArray<EBloomSetupOutputType> BloomSetupOutputType = GetBloomSetupOutputType(Inputs.bUseBloom, Inputs.bUseSun, Inputs.bUseDof, Inputs.bUseEyeAdaptation);
 
-	const auto GetBloomSetupTarget = [&GraphBuilder, bIsValidVariation, &Inputs, &OutputSize, &BloomSetupOutputType](int32 OutputIndex)
+	bool bUseTextureArrays = View.bIsMobileMultiViewEnabled;
+
+	const auto GetBloomSetupTarget = [&GraphBuilder, bIsValidVariation, &Inputs, &OutputSize, &BloomSetupOutputType, bUseTextureArrays](int32 OutputIndex)
 	{
 		checkSlow(OutputIndex < BloomSetupOutputType.Num());
 
@@ -303,7 +305,9 @@ FMobileBloomSetupOutputs AddMobileBloomSetupPass(FRDGBuilder& GraphBuilder, cons
 			Format = PF_FloatR11G11B10;
 		}
 
-		FRDGTextureDesc BloomSetupDesc = FRDGTextureDesc::Create2D(OutputSize, Format, FClearValueBinding::Black, TargetableFlags);
+		FRDGTextureDesc BloomSetupDesc = bUseTextureArrays ?
+			FRDGTextureDesc::Create2DArray(OutputSize, Format, FClearValueBinding::Black, TargetableFlags, 2) :
+			FRDGTextureDesc::Create2D(OutputSize, Format, FClearValueBinding::Black, TargetableFlags);
 
 		return FScreenPassRenderTarget(GraphBuilder.CreateTexture(BloomSetupDesc, GetBloomSetupOutputTypeName(BloomSetupOutputType[OutputIndex])), ERenderTargetLoadAction::ENoAction);
 	};
@@ -319,7 +323,7 @@ FMobileBloomSetupOutputs AddMobileBloomSetupPass(FRDGBuilder& GraphBuilder, cons
 
 	FMobileBloomSetupVS::FParameters VSShaderParameters;
 
-	VSShaderParameters.View = View.ViewUniformBuffer;
+	VSShaderParameters.View = View.GetShaderParameters();
 	VSShaderParameters.BufferSizeAndInvSize = FVector4f(BufferSize.X, BufferSize.Y, 1.0f / BufferSize.X, 1.0f / BufferSize.Y);
 
 	auto ShaderPermutationVector = FMobileBloomSetupPS::BuildPermutationVector(Inputs.bUseBloom, Inputs.bUseSun, Inputs.bUseDof, Inputs.bUseEyeAdaptation, Inputs.bUseMetalMSAAHDRDecode);
@@ -332,10 +336,12 @@ FMobileBloomSetupOutputs AddMobileBloomSetupPass(FRDGBuilder& GraphBuilder, cons
 	{
 		PSShaderParameters->RenderTargets[i] = DestRenderTargets[i].GetRenderTargetBinding();
 	}
-
+	//if the scenecolor isn't multiview but the app is, need to render as a single-view multiview due to shaders
+	PSShaderParameters->RenderTargets.MultiViewCount = (View.bIsMobileMultiViewEnabled) ? 2 : (UE::StereoRenderUtils::FStereoShaderAspects(View.GetShaderPlatform()).IsMobileMultiViewEnabled() ? 1 : 0);
+	
 	PSShaderParameters->EyeAdaptation = EyeAdaptationParameters;
 	PSShaderParameters->BloomThreshold = View.FinalPostProcessSettings.BloomThreshold;
-	PSShaderParameters->View = View.ViewUniformBuffer;
+	PSShaderParameters->View = View.GetShaderParameters();
 
 	PSShaderParameters->SceneColorTexture = Inputs.SceneColor.Texture;
 	PSShaderParameters->SceneColorSampler = TStaticSamplerState<SF_Bilinear, AM_Clamp, AM_Clamp, AM_Clamp>::GetRHI();
@@ -380,7 +386,8 @@ FMobileBloomSetupOutputs AddMobileBloomSetupPass(FRDGBuilder& GraphBuilder, cons
 			OutputViewport.Extent,
 			InputViewport.Extent,
 			VertexShader,
-			EDRF_UseTriangleOptimization);
+			EDRF_UseTriangleOptimization,
+			View.GetStereoPassInstanceFactor());
 	});
 
 	FMobileBloomSetupOutputs Outputs;
@@ -416,7 +423,7 @@ public:
 	SHADER_USE_PARAMETER_STRUCT(FMobileBloomDownPS, FGlobalShader);
 
 	BEGIN_SHADER_PARAMETER_STRUCT(FParameters, )
-		SHADER_PARAMETER_STRUCT_REF(FViewUniformShaderParameters, View)
+		SHADER_PARAMETER_STRUCT_INCLUDE(FViewShaderParameters, View)
 		SHADER_PARAMETER_RDG_TEXTURE(Texture2D, BloomDownSourceTexture)
 		SHADER_PARAMETER_SAMPLER(SamplerState, BloomDownSourceSampler)
 		RENDER_TARGET_BINDING_SLOTS()
@@ -437,7 +444,7 @@ public:
 	SHADER_USE_PARAMETER_STRUCT_WITH_LEGACY_BASE(FMobileBloomDownVS, FGlobalShader);
 
 	BEGIN_SHADER_PARAMETER_STRUCT(FParameters, )
-		SHADER_PARAMETER_STRUCT_REF(FViewUniformShaderParameters, View)
+		SHADER_PARAMETER_STRUCT_INCLUDE(FViewShaderParameters, View)
 		SHADER_PARAMETER(FVector4f, BufferSizeAndInvSize)
 		SHADER_PARAMETER(float, BloomDownScale)
 	END_SHADER_PARAMETER_STRUCT()
@@ -456,7 +463,10 @@ FScreenPassTexture AddMobileBloomDownPass(FRDGBuilder& GraphBuilder, const FView
 
 	const FIntPoint& BufferSize = Inputs.BloomDownSource.Texture->Desc.Extent;
 
-	FRDGTextureDesc BloomDownDesc = FRDGTextureDesc::Create2D(OutputSize, PF_FloatR11G11B10, FClearValueBinding::Black, TexCreate_RenderTargetable | TexCreate_ShaderResource);
+	bool bUseTextureArrays = View.bIsMobileMultiViewEnabled;
+	FRDGTextureDesc BloomDownDesc = bUseTextureArrays ?
+		FRDGTextureDesc::Create2DArray(OutputSize, PF_FloatR11G11B10, FClearValueBinding::Black, TexCreate_RenderTargetable | TexCreate_ShaderResource, 2) :
+		FRDGTextureDesc::Create2D(OutputSize, PF_FloatR11G11B10, FClearValueBinding::Black, TexCreate_RenderTargetable | TexCreate_ShaderResource);
 
 	FScreenPassRenderTarget BloomDownOutput = FScreenPassRenderTarget(GraphBuilder.CreateTexture(BloomDownDesc, TEXT("BloomDown")), ERenderTargetLoadAction::EClear);
 
@@ -464,7 +474,7 @@ FScreenPassTexture AddMobileBloomDownPass(FRDGBuilder& GraphBuilder, const FView
 
 	FMobileBloomDownVS::FParameters VSShaderParameters;
 
-	VSShaderParameters.View = View.ViewUniformBuffer;
+	VSShaderParameters.View = View.GetShaderParameters();
 	VSShaderParameters.BufferSizeAndInvSize = FVector4f(BufferSize.X, BufferSize.Y, 1.0f / BufferSize.X, 1.0f / BufferSize.Y);
 	VSShaderParameters.BloomDownScale = Inputs.BloomDownScale;
 
@@ -472,7 +482,8 @@ FScreenPassTexture AddMobileBloomDownPass(FRDGBuilder& GraphBuilder, const FView
 
 	FMobileBloomDownPS::FParameters* PSShaderParameters = GraphBuilder.AllocParameters<FMobileBloomDownPS::FParameters>();
 	PSShaderParameters->RenderTargets[0] = BloomDownOutput.GetRenderTargetBinding();
-	PSShaderParameters->View = View.ViewUniformBuffer;
+	PSShaderParameters->RenderTargets.MultiViewCount = (View.bIsMobileMultiViewEnabled) ? 2 : (UE::StereoRenderUtils::FStereoShaderAspects(View.GetShaderPlatform()).IsMobileMultiViewEnabled() ? 1 : 0);
+	PSShaderParameters->View = View.GetShaderParameters();
 
 	PSShaderParameters->BloomDownSourceTexture = Inputs.BloomDownSource.Texture;
 	PSShaderParameters->BloomDownSourceSampler = TStaticSamplerState<SF_Bilinear, AM_Clamp, AM_Clamp, AM_Clamp>::GetRHI();
@@ -484,7 +495,7 @@ FScreenPassTexture AddMobileBloomDownPass(FRDGBuilder& GraphBuilder, const FView
 		RDG_EVENT_NAME("BloomDown %dx%d (PS)", OutputViewport.Extent.X, OutputViewport.Extent.Y),
 		PSShaderParameters,
 		ERDGPassFlags::Raster,
-		[VertexShader, VSShaderParameters, PixelShader, PSShaderParameters, InputViewport, OutputViewport](FRHICommandList& RHICmdList)
+		[VertexShader, VSShaderParameters, PixelShader, PSShaderParameters, InputViewport, OutputViewport, &View](FRHICommandList& RHICmdList)
 	{
 		RHICmdList.SetViewport(OutputViewport.Rect.Min.X, OutputViewport.Rect.Min.Y, 0.0f, OutputViewport.Rect.Max.X, OutputViewport.Rect.Max.Y, 1.0f);
 
@@ -514,7 +525,8 @@ FScreenPassTexture AddMobileBloomDownPass(FRDGBuilder& GraphBuilder, const FView
 			OutputViewport.Extent,
 			InputViewport.Extent,
 			VertexShader,
-			EDRF_UseTriangleOptimization);
+			EDRF_UseTriangleOptimization,
+			View.GetStereoPassInstanceFactor());
 	});
 
 	return MoveTemp(BloomDownOutput);
@@ -533,7 +545,7 @@ public:
 	BEGIN_SHADER_PARAMETER_STRUCT(FParameters, )
 		SHADER_PARAMETER(FVector4f, BloomTintA)
 		SHADER_PARAMETER(FVector4f, BloomTintB)
-		SHADER_PARAMETER_STRUCT_REF(FViewUniformShaderParameters, View)
+		SHADER_PARAMETER_STRUCT_INCLUDE(FViewShaderParameters, View)
 		SHADER_PARAMETER_RDG_TEXTURE(Texture2D, BloomUpSourceATexture)
 		SHADER_PARAMETER_SAMPLER(SamplerState, BloomUpSourceASampler)
 		SHADER_PARAMETER_RDG_TEXTURE(Texture2D, BloomUpSourceBTexture)
@@ -556,7 +568,7 @@ public:
 	SHADER_USE_PARAMETER_STRUCT_WITH_LEGACY_BASE(FMobileBloomUpVS, FGlobalShader);
 
 	BEGIN_SHADER_PARAMETER_STRUCT(FParameters, )
-		SHADER_PARAMETER_STRUCT_REF(FViewUniformShaderParameters, View)
+		SHADER_PARAMETER_STRUCT_INCLUDE(FViewShaderParameters, View)
 		SHADER_PARAMETER(FVector4f, BufferASizeAndInvSize)
 		SHADER_PARAMETER(FVector4f, BufferBSizeAndInvSize)
 		SHADER_PARAMETER(FVector2f, BloomUpScales)
@@ -578,7 +590,10 @@ FScreenPassTexture AddMobileBloomUpPass(FRDGBuilder& GraphBuilder, const FViewIn
 
 	const FIntPoint& BufferSizeB = Inputs.BloomUpSourceB.Texture->Desc.Extent;
 
-	FRDGTextureDesc BloomUpDesc = FRDGTextureDesc::Create2D(OutputSize, PF_FloatR11G11B10, FClearValueBinding::Black, TexCreate_RenderTargetable | TexCreate_ShaderResource);
+	const bool bUseTextureArrays = View.bIsMobileMultiViewEnabled;
+	FRDGTextureDesc BloomUpDesc = bUseTextureArrays ? 
+		FRDGTextureDesc::Create2DArray(OutputSize, PF_FloatR11G11B10, FClearValueBinding::Black, TexCreate_RenderTargetable | TexCreate_ShaderResource, 2) :
+		FRDGTextureDesc::Create2D(OutputSize, PF_FloatR11G11B10, FClearValueBinding::Black, TexCreate_RenderTargetable | TexCreate_ShaderResource);
 
 	FScreenPassRenderTarget BloomUpOutput = FScreenPassRenderTarget(GraphBuilder.CreateTexture(BloomUpDesc, TEXT("BloomUp")), ERenderTargetLoadAction::EClear);
 
@@ -586,7 +601,7 @@ FScreenPassTexture AddMobileBloomUpPass(FRDGBuilder& GraphBuilder, const FViewIn
 
 	FMobileBloomUpVS::FParameters VSShaderParameters;
 
-	VSShaderParameters.View = View.ViewUniformBuffer;
+	VSShaderParameters.View = View.GetShaderParameters();
 	VSShaderParameters.BufferASizeAndInvSize = FVector4f(BufferSizeA.X, BufferSizeA.Y, 1.0f / BufferSizeA.X, 1.0f / BufferSizeA.Y);
 	VSShaderParameters.BufferBSizeAndInvSize = FVector4f(BufferSizeB.X, BufferSizeB.Y, 1.0f / BufferSizeB.X, 1.0f / BufferSizeB.Y);
 	VSShaderParameters.BloomUpScales = FVector2f(Inputs.ScaleAB);	// LWC_TODO: Precision loss
@@ -595,7 +610,8 @@ FScreenPassTexture AddMobileBloomUpPass(FRDGBuilder& GraphBuilder, const FViewIn
 
 	FMobileBloomUpPS::FParameters* PSShaderParameters = GraphBuilder.AllocParameters<FMobileBloomUpPS::FParameters>();
 	PSShaderParameters->RenderTargets[0] = BloomUpOutput.GetRenderTargetBinding();
-	PSShaderParameters->View = View.ViewUniformBuffer;
+	PSShaderParameters->RenderTargets.MultiViewCount = (View.bIsMobileMultiViewEnabled) ? 2 : (UE::StereoRenderUtils::FStereoShaderAspects(View.GetShaderPlatform()).IsMobileMultiViewEnabled() ? 1 : 0);
+	PSShaderParameters->View = View.GetShaderParameters();
 
 	PSShaderParameters->BloomUpSourceATexture = Inputs.BloomUpSourceA.Texture;
 	PSShaderParameters->BloomUpSourceASampler = TStaticSamplerState<SF_Bilinear, AM_Clamp, AM_Clamp, AM_Clamp>::GetRHI();
@@ -612,7 +628,7 @@ FScreenPassTexture AddMobileBloomUpPass(FRDGBuilder& GraphBuilder, const FViewIn
 		RDG_EVENT_NAME("BloomUp %dx%d (PS)", OutputViewport.Extent.X, OutputViewport.Extent.Y),
 		PSShaderParameters,
 		ERDGPassFlags::Raster,
-		[VertexShader, VSShaderParameters, PixelShader, PSShaderParameters, InputViewport, OutputViewport](FRHICommandList& RHICmdList)
+		[VertexShader, VSShaderParameters, PixelShader, PSShaderParameters, InputViewport, OutputViewport, &View](FRHICommandList& RHICmdList)
 	{
 		RHICmdList.SetViewport(OutputViewport.Rect.Min.X, OutputViewport.Rect.Min.Y, 0.0f, OutputViewport.Rect.Max.X, OutputViewport.Rect.Max.Y, 1.0f);
 
@@ -642,7 +658,8 @@ FScreenPassTexture AddMobileBloomUpPass(FRDGBuilder& GraphBuilder, const FViewIn
 			OutputViewport.Extent,
 			InputViewport.Extent,
 			VertexShader,
-			EDRF_UseTriangleOptimization);
+			EDRF_UseTriangleOptimization,
+			View.GetStereoPassInstanceFactor());
 	});
 
 	return MoveTemp(BloomUpOutput);
@@ -671,7 +688,7 @@ public:
 
 	BEGIN_SHADER_PARAMETER_STRUCT(FParameters, RENDERER_API)
 		SHADER_PARAMETER_RDG_UNIFORM_BUFFER(FMobileSceneTextureUniformParameters, SceneTextures)
-		SHADER_PARAMETER_STRUCT_REF(FViewUniformShaderParameters, View)
+		SHADER_PARAMETER_STRUCT_INCLUDE(FViewShaderParameters, View)
 		SHADER_PARAMETER_RDG_TEXTURE(Texture2D, SceneColorTexture)
 		SHADER_PARAMETER_SAMPLER(SamplerState, SceneColorSampler)
 		SHADER_PARAMETER(FVector4f, SunColorApertureDiv2)
@@ -767,7 +784,7 @@ FMobileSunMaskOutputs AddMobileSunMaskPass(FRDGBuilder& GraphBuilder, const FVie
 	PassParameters->RenderTargets[0] = SunMaskOutput.GetRenderTargetBinding();
 	PassParameters->RenderTargets[1] = SceneColorOutput.GetRenderTargetBinding();
 
-	PassParameters->View = View.ViewUniformBuffer;
+	PassParameters->View = View.GetShaderParameters();
 	PassParameters->SceneTextures = Inputs.SceneTextures;
 	PassParameters->SceneColorTexture = Inputs.SceneColor.Texture;
 	PassParameters->SceneColorSampler = TStaticSamplerState<SF_Bilinear, AM_Clamp, AM_Clamp, AM_Clamp>::GetRHI();
@@ -808,7 +825,7 @@ public:
 	using FPermutationDomain = TShaderPermutationDomain<FUseDofDim>;
 
 	BEGIN_SHADER_PARAMETER_STRUCT(FParameters, )
-		SHADER_PARAMETER_STRUCT_REF(FViewUniformShaderParameters, View)
+		SHADER_PARAMETER_STRUCT_INCLUDE(FViewShaderParameters, View)
 		SHADER_PARAMETER_RDG_TEXTURE(Texture2D, SunShaftAndDofTexture)
 		SHADER_PARAMETER_SAMPLER(SamplerState, SunShaftAndDofSampler)
 		RENDER_TARGET_BINDING_SLOTS()
@@ -836,7 +853,7 @@ public:
 	SHADER_USE_PARAMETER_STRUCT_WITH_LEGACY_BASE(FMobileSunAlphaVS, FGlobalShader);
 
 	BEGIN_SHADER_PARAMETER_STRUCT(FParameters, )
-		SHADER_PARAMETER_STRUCT_REF(FViewUniformShaderParameters, View)
+		SHADER_PARAMETER_STRUCT_INCLUDE(FViewShaderParameters, View)
 		SHADER_PARAMETER(FVector2f, LightShaftCenter)
 	END_SHADER_PARAMETER_STRUCT()
 
@@ -860,7 +877,7 @@ FScreenPassTexture AddMobileSunAlphaPass(FRDGBuilder& GraphBuilder, const FViewI
 
 	FMobileSunAlphaVS::FParameters VSShaderParameters;
 
-	VSShaderParameters.View = View.ViewUniformBuffer;
+	VSShaderParameters.View = View.GetShaderParameters();
 	VSShaderParameters.LightShaftCenter = FVector2f(View.MobileLightShaft->Center);	// LWC_TODO: Precision loss
 
 	auto ShaderPermutationVector = FMobileSunAlphaPS::BuildPermutationVector(Inputs.bUseMobileDof);
@@ -869,7 +886,7 @@ FScreenPassTexture AddMobileSunAlphaPass(FRDGBuilder& GraphBuilder, const FViewI
 
 	FMobileSunAlphaPS::FParameters* PSShaderParameters = GraphBuilder.AllocParameters<FMobileSunAlphaPS::FParameters>();
 	PSShaderParameters->RenderTargets[0] = SunAlphaOutput.GetRenderTargetBinding();
-	PSShaderParameters->View = View.ViewUniformBuffer;
+	PSShaderParameters->View = View.GetShaderParameters();
 	PSShaderParameters->SunShaftAndDofTexture = Inputs.BloomSetup_SunShaftAndDof.Texture;
 	PSShaderParameters->SunShaftAndDofSampler = TStaticSamplerState<SF_Bilinear, AM_Clamp, AM_Clamp, AM_Clamp>::GetRHI();
 
@@ -879,7 +896,7 @@ FScreenPassTexture AddMobileSunAlphaPass(FRDGBuilder& GraphBuilder, const FViewI
 		RDG_EVENT_NAME("SunAlpha %dx%d (PS)", OutputViewport.Extent.X, OutputViewport.Extent.Y),
 		PSShaderParameters,
 		ERDGPassFlags::Raster,
-		[VertexShader, VSShaderParameters, PixelShader, PSShaderParameters, OutputViewport](FRHICommandList& RHICmdList)
+		[VertexShader, VSShaderParameters, PixelShader, PSShaderParameters, OutputViewport, &View](FRHICommandList& RHICmdList)
 	{
 		RHICmdList.SetViewport(OutputViewport.Rect.Min.X, OutputViewport.Rect.Min.Y, 0.0f, OutputViewport.Rect.Max.X, OutputViewport.Rect.Max.Y, 1.0f);
 
@@ -909,7 +926,8 @@ FScreenPassTexture AddMobileSunAlphaPass(FRDGBuilder& GraphBuilder, const FViewI
 			OutputViewport.Extent,
 			OutputViewport.Extent,
 			VertexShader,
-			EDRF_UseTriangleOptimization);
+			EDRF_UseTriangleOptimization,
+			View.GetStereoPassInstanceFactor());
 	});
 
 	return MoveTemp(SunAlphaOutput);
@@ -926,7 +944,7 @@ public:
 	SHADER_USE_PARAMETER_STRUCT(FMobileSunBlurPS, FGlobalShader);
 
 	BEGIN_SHADER_PARAMETER_STRUCT(FParameters, )
-		SHADER_PARAMETER_STRUCT_REF(FViewUniformShaderParameters, View)
+		SHADER_PARAMETER_STRUCT_INCLUDE(FViewShaderParameters, View)
 		SHADER_PARAMETER_RDG_TEXTURE(Texture2D, SunAlphaTexture)
 		SHADER_PARAMETER_SAMPLER(SamplerState, SunAlphaSampler)
 		RENDER_TARGET_BINDING_SLOTS()
@@ -947,7 +965,7 @@ public:
 	SHADER_USE_PARAMETER_STRUCT_WITH_LEGACY_BASE(FMobileSunBlurVS, FGlobalShader);
 
 	BEGIN_SHADER_PARAMETER_STRUCT(FParameters, )
-		SHADER_PARAMETER_STRUCT_REF(FViewUniformShaderParameters, View)
+		SHADER_PARAMETER_STRUCT_INCLUDE(FViewShaderParameters, View)
 		SHADER_PARAMETER(FVector2f, LightShaftCenter)
 	END_SHADER_PARAMETER_STRUCT()
 
@@ -971,14 +989,14 @@ FScreenPassTexture AddMobileSunBlurPass(FRDGBuilder& GraphBuilder, const FViewIn
 
 	FMobileSunBlurVS::FParameters VSShaderParameters;
 
-	VSShaderParameters.View = View.ViewUniformBuffer;
+	VSShaderParameters.View = View.GetShaderParameters();
 	VSShaderParameters.LightShaftCenter = FVector2f(View.MobileLightShaft->Center);	// LWC_TODO: Precision loss
 
 	TShaderMapRef<FMobileSunBlurPS> PixelShader(View.ShaderMap);
 
 	FMobileSunBlurPS::FParameters* PSShaderParameters = GraphBuilder.AllocParameters<FMobileSunBlurPS::FParameters>();
 	PSShaderParameters->RenderTargets[0] = SunBlurOutput.GetRenderTargetBinding();
-	PSShaderParameters->View = View.ViewUniformBuffer;
+	PSShaderParameters->View = View.GetShaderParameters();
 	PSShaderParameters->SunAlphaTexture = Inputs.SunAlpha.Texture;
 	PSShaderParameters->SunAlphaSampler = TStaticSamplerState<SF_Bilinear, AM_Clamp, AM_Clamp, AM_Clamp>::GetRHI();
 
@@ -988,7 +1006,7 @@ FScreenPassTexture AddMobileSunBlurPass(FRDGBuilder& GraphBuilder, const FViewIn
 		RDG_EVENT_NAME("SunBlur %dx%d (PS)", OutputViewport.Extent.X, OutputViewport.Extent.Y),
 		PSShaderParameters,
 		ERDGPassFlags::Raster,
-		[VertexShader, VSShaderParameters, PixelShader, PSShaderParameters, OutputViewport](FRHICommandList& RHICmdList)
+		[VertexShader, VSShaderParameters, PixelShader, PSShaderParameters, OutputViewport, &View](FRHICommandList& RHICmdList)
 	{
 		RHICmdList.SetViewport(OutputViewport.Rect.Min.X, OutputViewport.Rect.Min.Y, 0.0f, OutputViewport.Rect.Max.X, OutputViewport.Rect.Max.Y, 1.0f);
 
@@ -1019,7 +1037,8 @@ FScreenPassTexture AddMobileSunBlurPass(FRDGBuilder& GraphBuilder, const FViewIn
 			OutputViewport.Extent,
 			OutputViewport.Extent,
 			VertexShader,
-			EDRF_UseTriangleOptimization);
+			EDRF_UseTriangleOptimization,
+			View.GetStereoPassInstanceFactor());
 	});
 
 	return MoveTemp(SunBlurOutput);
@@ -1045,7 +1064,7 @@ public:
 		SHADER_PARAMETER(FVector4f, BloomDirtMaskTint)
 		SHADER_PARAMETER(FVector4f, SunColorVignetteIntensity)
 		SHADER_PARAMETER(FVector3f, BloomColor)
-		SHADER_PARAMETER_STRUCT_REF(FViewUniformShaderParameters, View)
+		SHADER_PARAMETER_STRUCT_INCLUDE(FViewShaderParameters, View)
 		SHADER_PARAMETER_RDG_TEXTURE(Texture2D, SunBlurTexture)
 		SHADER_PARAMETER_SAMPLER(SamplerState, SunBlurSampler)
 		SHADER_PARAMETER_RDG_TEXTURE(Texture2D, BloomSetup_BloomTexture)
@@ -1080,7 +1099,7 @@ public:
 	SHADER_USE_PARAMETER_STRUCT_WITH_LEGACY_BASE(FMobileSunMergeVS, FGlobalShader);
 
 	BEGIN_SHADER_PARAMETER_STRUCT(FParameters, )
-		SHADER_PARAMETER_STRUCT_REF(FViewUniformShaderParameters, View)
+		SHADER_PARAMETER_STRUCT_INCLUDE(FViewShaderParameters, View)
 		SHADER_PARAMETER(FVector2f, LightShaftCenter)
 		SHADER_PARAMETER(FVector4f, BloomUpSizeAndInvSize)
 		SHADER_PARAMETER(FVector4f, ViewportSize)
@@ -1098,7 +1117,10 @@ FScreenPassTexture AddMobileSunMergePass(FRDGBuilder& GraphBuilder, const FViewI
 {
 	FIntPoint OutputSize = FIntPoint::DivideAndRoundUp(View.ViewRect.Size(), 4);
 	
-	FRDGTextureDesc SunMergeDesc = FRDGTextureDesc::Create2D(OutputSize, PF_FloatR11G11B10, FClearValueBinding::Black, TexCreate_RenderTargetable | TexCreate_ShaderResource);
+	const bool bUseTextureArrays = View.bIsMobileMultiViewEnabled;
+	FRDGTextureDesc SunMergeDesc = bUseTextureArrays ? 
+		FRDGTextureDesc::Create2DArray(OutputSize, PF_FloatR11G11B10, FClearValueBinding::Black, TexCreate_RenderTargetable | TexCreate_ShaderResource, 2) :
+		FRDGTextureDesc::Create2D(OutputSize, PF_FloatR11G11B10, FClearValueBinding::Black, TexCreate_RenderTargetable | TexCreate_ShaderResource);
 
 	FScreenPassRenderTarget SunMergeOutput = FScreenPassRenderTarget(GraphBuilder.CreateTexture(SunMergeDesc, TEXT("SunMerge")), ERenderTargetLoadAction::EClear);
 
@@ -1106,7 +1128,7 @@ FScreenPassTexture AddMobileSunMergePass(FRDGBuilder& GraphBuilder, const FViewI
 
 	FMobileSunMergeVS::FParameters VSShaderParameters;
 
-	VSShaderParameters.View = View.ViewUniformBuffer;
+	VSShaderParameters.View = View.GetShaderParameters();
 	FMobileLightShaftInfo MobileLightShaft;
 	if (View.MobileLightShaft)
 	{
@@ -1148,7 +1170,8 @@ FScreenPassTexture AddMobileSunMergePass(FRDGBuilder& GraphBuilder, const FViewI
 
 	FMobileSunMergePS::FParameters* PSShaderParameters = GraphBuilder.AllocParameters<FMobileSunMergePS::FParameters>();
 	PSShaderParameters->RenderTargets[0] = SunMergeOutput.GetRenderTargetBinding();
-	PSShaderParameters->View = View.ViewUniformBuffer;
+	PSShaderParameters->RenderTargets.MultiViewCount = (View.bIsMobileMultiViewEnabled) ? 2 : (UE::StereoRenderUtils::FStereoShaderAspects(View.GetShaderPlatform()).IsMobileMultiViewEnabled() ? 1 : 0);
+	PSShaderParameters->View = View.GetShaderParameters();
 	PSShaderParameters->BloomDirtMaskTint = Settings.BloomDirtMaskTint * Settings.BloomDirtMaskIntensity;
 	PSShaderParameters->SunColorVignetteIntensity = SunColorVignetteIntensityParam;
 	PSShaderParameters->BloomColor = FVector3f(BloomColor.R, BloomColor.G, BloomColor.B);
@@ -1197,7 +1220,8 @@ FScreenPassTexture AddMobileSunMergePass(FRDGBuilder& GraphBuilder, const FViewI
 			OutputViewport.Extent,
 			OutputViewport.Extent,
 			VertexShader,
-			EDRF_UseTriangleOptimization);
+			EDRF_UseTriangleOptimization,
+			View.GetStereoPassInstanceFactor());
 	});
 
 	return MoveTemp(SunMergeOutput);
@@ -1214,7 +1238,7 @@ public:
 	SHADER_USE_PARAMETER_STRUCT_WITH_LEGACY_BASE(FMobileDofDownVS, FGlobalShader);
 
 	BEGIN_SHADER_PARAMETER_STRUCT(FParameters, )
-		SHADER_PARAMETER_STRUCT_REF(FViewUniformShaderParameters, View)
+		SHADER_PARAMETER_STRUCT_INCLUDE(FViewShaderParameters, View)
 		SHADER_PARAMETER(FVector4f, BufferSizeAndInvSize)
 	END_SHADER_PARAMETER_STRUCT()
 
@@ -1236,7 +1260,7 @@ public:
 	using FPermutationDomain = TShaderPermutationDomain<FUseSunDim>;
 
 	BEGIN_SHADER_PARAMETER_STRUCT(FParameters, )
-		SHADER_PARAMETER_STRUCT_REF(FViewUniformShaderParameters, View)
+		SHADER_PARAMETER_STRUCT_INCLUDE(FViewShaderParameters, View)
 		SHADER_PARAMETER_RDG_TEXTURE(Texture2D, SceneColorTexture)
 		SHADER_PARAMETER_SAMPLER(SamplerState, SceneColorSampler)
 		SHADER_PARAMETER_RDG_TEXTURE(Texture2D, SunShaftAndDofTexture)
@@ -1275,7 +1299,7 @@ FMobileDofDownOutputs AddMobileDofDownPass(FRDGBuilder& GraphBuilder, const FVie
 
 	FMobileDofDownVS::FParameters VSShaderParameters;
 
-	VSShaderParameters.View = View.ViewUniformBuffer;
+	VSShaderParameters.View = View.GetShaderParameters();
 	VSShaderParameters.BufferSizeAndInvSize = FVector4f(BufferSize.X, BufferSize.Y, 1.0f / BufferSize.X, 1.0f / BufferSize.Y);
 
 	auto ShaderPermutationVector = FMobileDofDownPS::BuildPermutationVector(Inputs.bUseSun);
@@ -1284,7 +1308,7 @@ FMobileDofDownOutputs AddMobileDofDownPass(FRDGBuilder& GraphBuilder, const FVie
 
 	FMobileDofDownPS::FParameters* PSShaderParameters = GraphBuilder.AllocParameters<FMobileDofDownPS::FParameters>();
 	PSShaderParameters->RenderTargets[0] = DofDownOutput.GetRenderTargetBinding();
-	PSShaderParameters->View = View.ViewUniformBuffer;
+	PSShaderParameters->View = View.GetShaderParameters();
 
 	PSShaderParameters->SceneColorTexture = Inputs.SceneColor.Texture;
 	PSShaderParameters->SceneColorSampler = TStaticSamplerState<SF_Point, AM_Clamp, AM_Clamp, AM_Clamp>::GetRHI();
@@ -1300,7 +1324,7 @@ FMobileDofDownOutputs AddMobileDofDownPass(FRDGBuilder& GraphBuilder, const FVie
 		RDG_EVENT_NAME("DofDown %dx%d (PS)", OutputViewport.Extent.X, OutputViewport.Extent.Y),
 		PSShaderParameters,
 		ERDGPassFlags::Raster,
-		[VertexShader, VSShaderParameters, PixelShader, PSShaderParameters, InputViewport, OutputViewport](FRHICommandList& RHICmdList)
+		[VertexShader, VSShaderParameters, PixelShader, PSShaderParameters, InputViewport, OutputViewport, &View](FRHICommandList& RHICmdList)
 	{
 		RHICmdList.SetViewport(OutputViewport.Rect.Min.X, OutputViewport.Rect.Min.Y, 0.0f, OutputViewport.Rect.Max.X, OutputViewport.Rect.Max.Y, 1.0f);
 
@@ -1330,7 +1354,8 @@ FMobileDofDownOutputs AddMobileDofDownPass(FRDGBuilder& GraphBuilder, const FVie
 			OutputViewport.Extent,
 			InputViewport.Extent,
 			VertexShader,
-			EDRF_UseTriangleOptimization);
+			EDRF_UseTriangleOptimization,
+			View.GetStereoPassInstanceFactor());
 	});
 
 	FMobileDofDownOutputs Outputs;
@@ -1351,7 +1376,7 @@ public:
 	SHADER_USE_PARAMETER_STRUCT_WITH_LEGACY_BASE(FMobileDofNearVS, FGlobalShader);
 
 	BEGIN_SHADER_PARAMETER_STRUCT(FParameters, )
-		SHADER_PARAMETER_STRUCT_REF(FViewUniformShaderParameters, View)
+		SHADER_PARAMETER_STRUCT_INCLUDE(FViewShaderParameters, View)
 		SHADER_PARAMETER(FVector4f, BufferSizeAndInvSize)
 	END_SHADER_PARAMETER_STRUCT()
 
@@ -1373,7 +1398,7 @@ public:
 	using FPermutationDomain = TShaderPermutationDomain<FUseSunDim>;
 
 	BEGIN_SHADER_PARAMETER_STRUCT(FParameters, )
-		SHADER_PARAMETER_STRUCT_REF(FViewUniformShaderParameters, View)
+		SHADER_PARAMETER_STRUCT_INCLUDE(FViewShaderParameters, View)
 		SHADER_PARAMETER_RDG_TEXTURE(Texture2D, SunShaftAndDofTexture)
 		SHADER_PARAMETER_SAMPLER(SamplerState, SunShaftAndDofSampler)
 		RENDER_TARGET_BINDING_SLOTS()
@@ -1406,7 +1431,7 @@ FMobileDofNearOutputs AddMobileDofNearPass(FRDGBuilder& GraphBuilder, const FVie
 
 	FMobileDofNearVS::FParameters VSShaderParameters;
 
-	VSShaderParameters.View = View.ViewUniformBuffer;
+	VSShaderParameters.View = View.GetShaderParameters();
 	VSShaderParameters.BufferSizeAndInvSize = FVector4f(BufferSize.X, BufferSize.Y, 1.0f / BufferSize.X, 1.0f / BufferSize.Y);
 
 	auto ShaderPermutationVector = FMobileDofNearPS::BuildPermutationVector(Inputs.bUseSun);
@@ -1415,7 +1440,7 @@ FMobileDofNearOutputs AddMobileDofNearPass(FRDGBuilder& GraphBuilder, const FVie
 
 	FMobileDofNearPS::FParameters* PSShaderParameters = GraphBuilder.AllocParameters<FMobileDofNearPS::FParameters>();
 	PSShaderParameters->RenderTargets[0] = DofNearOutput.GetRenderTargetBinding();
-	PSShaderParameters->View = View.ViewUniformBuffer;
+	PSShaderParameters->View = View.GetShaderParameters();
 
 	PSShaderParameters->SunShaftAndDofTexture = Inputs.BloomSetup_SunShaftAndDof.Texture;
 	PSShaderParameters->SunShaftAndDofSampler = TStaticSamplerState<SF_Point, AM_Clamp, AM_Clamp, AM_Clamp>::GetRHI();
@@ -1426,7 +1451,7 @@ FMobileDofNearOutputs AddMobileDofNearPass(FRDGBuilder& GraphBuilder, const FVie
 		RDG_EVENT_NAME("DofNear %dx%d (PS)", OutputViewport.Extent.X, OutputViewport.Extent.Y),
 		PSShaderParameters,
 		ERDGPassFlags::Raster,
-		[VertexShader, VSShaderParameters, PixelShader, PSShaderParameters, OutputViewport](FRHICommandList& RHICmdList)
+		[VertexShader, VSShaderParameters, PixelShader, PSShaderParameters, OutputViewport, &View](FRHICommandList& RHICmdList)
 	{
 		RHICmdList.SetViewport(OutputViewport.Rect.Min.X, OutputViewport.Rect.Min.Y, 0.0f, OutputViewport.Rect.Max.X, OutputViewport.Rect.Max.Y, 1.0f);
 
@@ -1456,7 +1481,8 @@ FMobileDofNearOutputs AddMobileDofNearPass(FRDGBuilder& GraphBuilder, const FVie
 			OutputViewport.Extent,
 			OutputViewport.Extent,
 			VertexShader,
-			EDRF_UseTriangleOptimization);
+			EDRF_UseTriangleOptimization,
+			View.GetStereoPassInstanceFactor());
 	});
 
 	FMobileDofNearOutputs Outputs;
@@ -1477,7 +1503,7 @@ public:
 	SHADER_USE_PARAMETER_STRUCT(FMobileDofBlurPS, FGlobalShader);
 
 	BEGIN_SHADER_PARAMETER_STRUCT(FParameters, )
-		SHADER_PARAMETER_STRUCT_REF(FViewUniformShaderParameters, View)
+		SHADER_PARAMETER_STRUCT_INCLUDE(FViewShaderParameters, View)
 		SHADER_PARAMETER_RDG_TEXTURE(Texture2D, DofNearTexture)
 		SHADER_PARAMETER_SAMPLER(SamplerState, DofNearSampler)
 		SHADER_PARAMETER_RDG_TEXTURE(Texture2D, DofDownTexture)
@@ -1500,7 +1526,7 @@ public:
 	SHADER_USE_PARAMETER_STRUCT_WITH_LEGACY_BASE(FMobileDofBlurVS, FGlobalShader);
 
 	BEGIN_SHADER_PARAMETER_STRUCT(FParameters, )
-		SHADER_PARAMETER_STRUCT_REF(FViewUniformShaderParameters, View)
+		SHADER_PARAMETER_STRUCT_INCLUDE(FViewShaderParameters, View)
 		SHADER_PARAMETER(FVector4f, BufferSizeAndInvSize)
 	END_SHADER_PARAMETER_STRUCT()
 
@@ -1524,14 +1550,14 @@ FMobileDofBlurOutputs AddMobileDofBlurPass(FRDGBuilder& GraphBuilder, const FVie
 
 	FMobileDofBlurVS::FParameters VSShaderParameters;
 
-	VSShaderParameters.View = View.ViewUniformBuffer;
+	VSShaderParameters.View = View.GetShaderParameters();
 	VSShaderParameters.BufferSizeAndInvSize = FVector4f(BufferSize.X, BufferSize.Y, 1.0f / BufferSize.X, 1.0f / BufferSize.Y);
 
 	TShaderMapRef<FMobileDofBlurPS> PixelShader(View.ShaderMap);
 
 	FMobileDofBlurPS::FParameters* PSShaderParameters = GraphBuilder.AllocParameters<FMobileDofBlurPS::FParameters>();
 	PSShaderParameters->RenderTargets[0] = DofBlurOutput.GetRenderTargetBinding();
-	PSShaderParameters->View = View.ViewUniformBuffer;
+	PSShaderParameters->View = View.GetShaderParameters();
 
 	PSShaderParameters->DofDownTexture = Inputs.DofDown.Texture;
 	PSShaderParameters->DofDownSampler = TStaticSamplerState<SF_Bilinear, AM_Clamp, AM_Clamp, AM_Clamp>::GetRHI();
@@ -1544,7 +1570,7 @@ FMobileDofBlurOutputs AddMobileDofBlurPass(FRDGBuilder& GraphBuilder, const FVie
 		RDG_EVENT_NAME("DofBlur %dx%d (PS)", OutputViewport.Extent.X, OutputViewport.Extent.Y),
 		PSShaderParameters,
 		ERDGPassFlags::Raster,
-		[VertexShader, VSShaderParameters, PixelShader, PSShaderParameters, OutputViewport](FRHICommandList& RHICmdList)
+		[VertexShader, VSShaderParameters, PixelShader, PSShaderParameters, OutputViewport, &View](FRHICommandList& RHICmdList)
 	{
 		RHICmdList.SetViewport(OutputViewport.Rect.Min.X, OutputViewport.Rect.Min.Y, 0.0f, OutputViewport.Rect.Max.X, OutputViewport.Rect.Max.Y, 1.0f);
 
@@ -1574,7 +1600,8 @@ FMobileDofBlurOutputs AddMobileDofBlurPass(FRDGBuilder& GraphBuilder, const FVie
 			OutputViewport.Extent,
 			OutputViewport.Extent,
 			VertexShader,
-			EDRF_UseTriangleOptimization);
+			EDRF_UseTriangleOptimization,
+			View.GetStereoPassInstanceFactor());
 	});
 
 	FMobileDofBlurOutputs Outputs;
@@ -1591,7 +1618,7 @@ public:
 	SHADER_USE_PARAMETER_STRUCT(FMobileIntegrateDofPS, FGlobalShader);
 
 	BEGIN_SHADER_PARAMETER_STRUCT(FParameters, )
-		SHADER_PARAMETER_STRUCT_REF(FViewUniformShaderParameters, View)
+		SHADER_PARAMETER_STRUCT_INCLUDE(FViewShaderParameters, View)
 		SHADER_PARAMETER_RDG_TEXTURE(Texture2D, SceneColorTexture)
 		SHADER_PARAMETER_SAMPLER(SamplerState, SceneColorSampler)
 		SHADER_PARAMETER_RDG_TEXTURE(Texture2D, DofBlurTexture)
@@ -1616,7 +1643,7 @@ public:
 	SHADER_USE_PARAMETER_STRUCT_WITH_LEGACY_BASE(FMobileIntegrateDofVS, FGlobalShader);
 
 	BEGIN_SHADER_PARAMETER_STRUCT(FParameters, )
-		SHADER_PARAMETER_STRUCT_REF(FViewUniformShaderParameters, View)
+		SHADER_PARAMETER_STRUCT_INCLUDE(FViewShaderParameters, View)
 		SHADER_PARAMETER(FVector4f, BufferSizeAndInvSize)
 		SHADER_PARAMETER(FVector4f, DofBlurSizeAndInvSize)
 	END_SHADER_PARAMETER_STRUCT()
@@ -1640,7 +1667,7 @@ FScreenPassTexture AddMobileIntegrateDofPass(FRDGBuilder& GraphBuilder, const FV
 
 	FMobileIntegrateDofVS::FParameters VSShaderParameters;
 
-	VSShaderParameters.View = View.ViewUniformBuffer;
+	VSShaderParameters.View = View.GetShaderParameters();
 	VSShaderParameters.BufferSizeAndInvSize = FVector4f(BufferSize.X, BufferSize.Y, 1.0f / BufferSize.X, 1.0f / BufferSize.Y);
 	VSShaderParameters.DofBlurSizeAndInvSize = FVector4f(DofBlurSize.X, DofBlurSize.Y, 1.0f / DofBlurSize.X, 1.0f / DofBlurSize.Y);
 
@@ -1648,7 +1675,7 @@ FScreenPassTexture AddMobileIntegrateDofPass(FRDGBuilder& GraphBuilder, const FV
 
 	FMobileIntegrateDofPS::FParameters* PSShaderParameters = GraphBuilder.AllocParameters<FMobileIntegrateDofPS::FParameters>();
 	PSShaderParameters->RenderTargets[0] = IntegrateDofOutput.GetRenderTargetBinding();
-	PSShaderParameters->View = View.ViewUniformBuffer;
+	PSShaderParameters->View = View.GetShaderParameters();
 
 	PSShaderParameters->SceneColorTexture = Inputs.SceneColor.Texture;
 	PSShaderParameters->SceneColorSampler = TStaticSamplerState<SF_Bilinear, AM_Clamp, AM_Clamp, AM_Clamp>::GetRHI();
@@ -1663,7 +1690,7 @@ FScreenPassTexture AddMobileIntegrateDofPass(FRDGBuilder& GraphBuilder, const FV
 		RDG_EVENT_NAME("IntegrateDof %dx%d (PS)", OutputViewport.Extent.X, OutputViewport.Extent.Y),
 		PSShaderParameters,
 		ERDGPassFlags::Raster,
-		[VertexShader, VSShaderParameters, PixelShader, PSShaderParameters, OutputViewport](FRHICommandList& RHICmdList)
+		[VertexShader, VSShaderParameters, PixelShader, PSShaderParameters, OutputViewport, &View](FRHICommandList& RHICmdList)
 	{
 		RHICmdList.SetViewport(OutputViewport.Rect.Min.X, OutputViewport.Rect.Min.Y, 0.0f, OutputViewport.Rect.Max.X, OutputViewport.Rect.Max.Y, 1.0f);
 
@@ -1693,7 +1720,8 @@ FScreenPassTexture AddMobileIntegrateDofPass(FRDGBuilder& GraphBuilder, const FV
 			OutputViewport.Extent,
 			OutputViewport.Extent,
 			VertexShader,
-			EDRF_UseTriangleOptimization);
+			EDRF_UseTriangleOptimization,
+			View.GetStereoPassInstanceFactor());
 	});
 
 	return MoveTemp(IntegrateDofOutput);
@@ -1766,7 +1794,7 @@ public:
 	static const uint32 HistogramSize = 64; // HistogramSize must be 64 and ThreadGroupSizeX * ThreadGroupSizeY must be larger than 32
 
 	BEGIN_SHADER_PARAMETER_STRUCT(FParameters, )
-		SHADER_PARAMETER_STRUCT_REF(FViewUniformShaderParameters, View)
+		SHADER_PARAMETER_STRUCT_INCLUDE(FViewShaderParameters, View)
 		SHADER_PARAMETER(FVector4f, SourceSizeAndInvSize)
 		SHADER_PARAMETER_STRUCT(FEyeAdaptationParameters, EyeAdaptation)
 		SHADER_PARAMETER_SAMPLER(SamplerState, InputSampler)
@@ -1861,7 +1889,7 @@ FMobileEyeAdaptationSetupOutputs AddMobileEyeAdaptationSetupPass(FRDGBuilder& Gr
 	{
 		FMobileHistogram::FParameters* PassParameters = GraphBuilder.AllocParameters<FMobileHistogram::FParameters>();
 
-		PassParameters->View = View.ViewUniformBuffer;
+		PassParameters->View = View.GetShaderParameters();
 		PassParameters->SourceSizeAndInvSize = FVector4f(BufferSize.X, BufferSize.Y, 1.0f / BufferSize.X, 1.0f / BufferSize.Y);
 		PassParameters->InputTexture = Inputs.BloomSetup_EyeAdaptation.Texture;
 		PassParameters->InputSampler = TStaticSamplerState<SF_Bilinear, AM_Clamp, AM_Clamp, AM_Clamp>::GetRHI();
@@ -1907,7 +1935,7 @@ public:
 	SHADER_USE_PARAMETER_STRUCT(FMobileBasicEyeAdaptationCS, FGlobalShader);
 
 	BEGIN_SHADER_PARAMETER_STRUCT(FParameters, )
-		SHADER_PARAMETER_STRUCT_REF(FViewUniformShaderParameters, View)
+		SHADER_PARAMETER_STRUCT_INCLUDE(FViewShaderParameters, View)
 		SHADER_PARAMETER_STRUCT(FEyeAdaptationParameters, EyeAdaptation)
 		SHADER_PARAMETER_RDG_BUFFER_SRV(StructuredBuffer<float4>, EyeAdaptationBuffer)
 		SHADER_PARAMETER_RDG_BUFFER_SRV(Buffer<uint>, LogLuminanceWeightBuffer)
@@ -1979,7 +2007,7 @@ void AddMobileEyeAdaptationPass(FRDGBuilder& GraphBuilder, const FViewInfo& View
 		FMobileBasicEyeAdaptationCS::FParameters* PassParameters = GraphBuilder.AllocParameters<FMobileBasicEyeAdaptationCS::FParameters>();
 
 		PassParameters->EyeAdaptation = EyeAdaptationParameters;
-		PassParameters->View = View.ViewUniformBuffer;
+		PassParameters->View = View.GetShaderParameters();
 		PassParameters->EyeAdaptationBuffer = EyeAdaptationBufferSRV;
 		PassParameters->LogLuminanceWeightBuffer = Inputs.EyeAdaptationSetupSRV;
 		PassParameters->OutputBuffer = OutputBufferUAV;
diff --git a/Engine/Source/Runtime/Renderer/Private/PostProcess/PostProcessTonemap.cpp b/Engine/Source/Runtime/Renderer/Private/PostProcess/PostProcessTonemap.cpp
--- a/Engine/Source/Runtime/Renderer/Private/PostProcess/PostProcessTonemap.cpp
+++ b/Engine/Source/Runtime/Renderer/Private/PostProcess/PostProcessTonemap.cpp
@@ -309,7 +309,7 @@ BEGIN_SHADER_PARAMETER_STRUCT(FFilmGrainParameters, )
 END_SHADER_PARAMETER_STRUCT()
 
 BEGIN_SHADER_PARAMETER_STRUCT(FTonemapParameters, )
-	SHADER_PARAMETER_STRUCT_REF(FViewUniformShaderParameters, View)
+	SHADER_PARAMETER_STRUCT_INCLUDE(FViewShaderParameters, View)
 	SHADER_PARAMETER_STRUCT_INCLUDE(FFilmGrainParameters, FilmGrain)
 	SHADER_PARAMETER_STRUCT_INCLUDE(FTonemapperOutputDeviceParameters, OutputDevice)
 	SHADER_PARAMETER_STRUCT(FScreenPassTextureViewportParameters, Color)
@@ -707,7 +707,7 @@ FScreenPassTexture AddTonemapPass(FRDGBuilder& GraphBuilder, const FViewInfo& Vi
 	#endif
 
 	FTonemapParameters CommonParameters;
-	CommonParameters.View = View.ViewUniformBuffer;
+	CommonParameters.View = View.GetShaderParameters();
 
 	{
 		uint8 FrameIndexMod8 = 0;
diff --git a/Engine/Source/Runtime/Renderer/Private/PostProcess/SceneFilterRendering.cpp b/Engine/Source/Runtime/Renderer/Private/PostProcess/SceneFilterRendering.cpp
--- a/Engine/Source/Runtime/Renderer/Private/PostProcess/SceneFilterRendering.cpp
+++ b/Engine/Source/Runtime/Renderer/Private/PostProcess/SceneFilterRendering.cpp
@@ -282,7 +282,8 @@ void DrawHmdMesh(
 	FIntPoint TargetSize,
 	FIntPoint TextureSize,
 	int32 StereoView,
-	const TShaderRef<FShader>& VertexShader
+	const TShaderRef<FShader>& VertexShader,
+	int32 InstanceCount
 	)
 {
 	{
@@ -297,7 +298,7 @@ void DrawHmdMesh(
 
 	if (GEngine->XRSystem->GetHMDDevice())
 	{
-		GEngine->XRSystem->GetHMDDevice()->DrawVisibleAreaMesh(RHICmdList, StereoView);
+		GEngine->XRSystem->GetHMDDevice()->DrawVisibleAreaMesh(RHICmdList, StereoView, InstanceCount);
 	}
 }
 
@@ -316,15 +317,16 @@ void DrawPostProcessPass(
 	const TShaderRef<FShader>& VertexShader,
 	int32 StereoViewIndex,
 	bool bHasCustomMesh,
-	EDrawRectangleFlags Flags)
+	EDrawRectangleFlags Flags,
+	int32 InstanceCount)
 {
 	if (bHasCustomMesh && StereoViewIndex != INDEX_NONE)
 	{
-		DrawHmdMesh(RHICmdList, X, Y, SizeX, SizeY, U, V, SizeU, SizeV, TargetSize, TextureSize, StereoViewIndex, VertexShader);
+		DrawHmdMesh(RHICmdList, X, Y, SizeX, SizeY, U, V, SizeU, SizeV, TargetSize, TextureSize, StereoViewIndex, VertexShader, InstanceCount);
 	}
 	else
 	{
-		DrawRectangle(RHICmdList, X, Y, SizeX, SizeY, U, V, SizeU, SizeV, TargetSize, TextureSize, VertexShader, Flags);
+		DrawRectangle(RHICmdList, X, Y, SizeX, SizeY, U, V, SizeU, SizeV, TargetSize, TextureSize, VertexShader, Flags, InstanceCount);
 	}
 }
 
@@ -370,7 +372,8 @@ namespace UE::Renderer::PostProcess
 				View.ViewRect.Width(), View.ViewRect.Height(),
 				View.UnconstrainedViewRect.Size(),
 				View.UnconstrainedViewRect.Size(),
-				Flags);
+				Flags,
+				InstanceCount);
 		}
 	}
 
diff --git a/Engine/Source/Runtime/Renderer/Private/PostProcess/SceneFilterRendering.h b/Engine/Source/Runtime/Renderer/Private/PostProcess/SceneFilterRendering.h
--- a/Engine/Source/Runtime/Renderer/Private/PostProcess/SceneFilterRendering.h
+++ b/Engine/Source/Runtime/Renderer/Private/PostProcess/SceneFilterRendering.h
@@ -89,7 +89,8 @@ extern RENDERER_API void DrawHmdMesh(
 	FIntPoint TargetSize,
 	FIntPoint TextureSize,
 	EStereoscopicPass StereoView,
-	const TShaderRef<FShader>& VertexShader
+	const TShaderRef<FShader>& VertexShader,
+	int32 InstanceCount = 1
 	);
 
 // NOTE: Assumes previously set PSO has PrimitiveType = PT_TriangleList
@@ -108,7 +109,8 @@ extern RENDERER_API void DrawPostProcessPass(
 	const TShaderRef<FShader>& VertexShader,
 	int32 StereoViewIndex,
 	bool bHasCustomMesh,
-	EDrawRectangleFlags Flags = EDRF_Default
+	EDrawRectangleFlags Flags = EDRF_Default,
+	int32 InstanceCount = 1
 	);
 
 class FTesselatedScreenRectangleIndexBuffer : public FIndexBuffer
diff --git a/Engine/Source/Runtime/Renderer/Private/ScreenPass.cpp b/Engine/Source/Runtime/Renderer/Private/ScreenPass.cpp
--- a/Engine/Source/Runtime/Renderer/Private/ScreenPass.cpp
+++ b/Engine/Source/Runtime/Renderer/Private/ScreenPass.cpp
@@ -261,7 +261,8 @@ void DrawScreenPass_PostSetup(
 		PipelineState.VertexShader,
 		ViewInfo.StereoViewIndex,
 		bUseHMDHiddenAreaMask,
-		DrawRectangleFlags);
+		DrawRectangleFlags,
+		ViewInfo.InstanceCount);
 }
 
 void AddDrawTexturePass(
diff --git a/Engine/Source/Runtime/Renderer/Public/ScreenPass.h b/Engine/Source/Runtime/Renderer/Public/ScreenPass.h
--- a/Engine/Source/Runtime/Renderer/Public/ScreenPass.h
+++ b/Engine/Source/Runtime/Renderer/Public/ScreenPass.h
@@ -417,16 +417,19 @@ ENUM_CLASS_FLAGS(EScreenPassDrawFlags);
 struct FScreenPassViewInfo
 {
 	const int32 StereoViewIndex;
+	const int32 InstanceCount;
 	const bool bHMDHiddenAreaMaskActive;
 
 	FScreenPassViewInfo()
 		: StereoViewIndex(INDEX_NONE)
+		, InstanceCount(1)
 		, bHMDHiddenAreaMaskActive(IsHMDHiddenAreaMaskActive())
 	{
 	}
 
 	FScreenPassViewInfo(const FSceneView& View)
 		: StereoViewIndex(View.StereoViewIndex)
+		, InstanceCount(View.GetStereoPassInstanceFactor())
 		, bHMDHiddenAreaMaskActive(View.bHMDHiddenAreaMaskActive)
 	{
 	}
