Fix cached ray tracing state not being updated when BLAS is streamed in by FRayTracingGeometryManager.

diff --git a/Engine/Source/Runtime/Engine/Private/Rendering/RayTracingGeometryManager.cpp b/Engine/Source/Runtime/Engine/Private/Rendering/RayTracingGeometryManager.cpp
--- a/Engine/Source/Runtime/Engine/Private/Rendering/RayTracingGeometryManager.cpp
+++ b/Engine/Source/Runtime/Engine/Private/Rendering/RayTracingGeometryManager.cpp
@@ -744,6 +744,12 @@ void FRayTracingGeometryManager::ProcessCompletedStreamingRequests(FRHICommandLi
 
 				RegisteredGeometry.Status = FRegisteredGeometry::FStatus::StreamedIn;
 
+				if (!RegisteredGeometry.Geometry->GetRequiresBuild())
+				{
+					// only need to request here if no build will be requested since build path already requests update as necessary
+					RequestUpdateCachedRenderState(RegisteredGeometry.Geometry->GroupHandle);
+				}
+
 				RegisteredGeometry.Geometry->RequestBuildIfNeeded(RHICmdList, ERTAccelerationStructureBuildPriority::Normal);
 			}
 		}
diff --git a/Engine/Source/Runtime/Renderer/Private/RayTracing/RayTracing.cpp b/Engine/Source/Runtime/Renderer/Private/RayTracing/RayTracing.cpp
--- a/Engine/Source/Runtime/Renderer/Private/RayTracing/RayTracing.cpp
+++ b/Engine/Source/Runtime/Renderer/Private/RayTracing/RayTracing.cpp
@@ -539,10 +539,12 @@ namespace RayTracing
 							{
 								RelevantPrimitive->CachedRayTracingMeshCommandIndices = SceneInfo->CachedRayTracingMeshCommandIndicesPerLOD[LODIndex];
 							}
-							else
-							{
-								// TODO: check if this ever happens. should probably skip primitive if so
-							}
+
+							// CacheInstances expects to have one ray tracing mesh command per BLAS segment.
+							// If that's not the case in the future, other logic such as NumCachedStaticVisibleMeshCommands calculation needs to be updated.
+							checkf(RelevantPrimitive->CachedRayTracingMeshCommandIndices.Num() == SceneInfo->CachedRayTracingInstance.GeometryRHI->GetNumSegments(),
+								TEXT("Expected to have one ray tracing mesh command per BLAS segment (primitive has %d cached mesh commands but BLAS has %d segments)."),
+								RelevantPrimitive->CachedRayTracingMeshCommandIndices.Num(), SceneInfo->CachedRayTracingInstance.GeometryRHI->GetNumSegments());
 
 							const bool bNeedMainInstance = !RelevantPrimitive->bAllSegmentsDecal;
 
