Add renderer module API support for flushing pages of individual allocated virtual textures.

diff --git a/Engine/Source/Runtime/RenderCore/Public/RendererInterface.h b/Engine/Source/Runtime/RenderCore/Public/RendererInterface.h
--- a/Engine/Source/Runtime/RenderCore/Public/RendererInterface.h
+++ b/Engine/Source/Runtime/RenderCore/Public/RendererInterface.h
@@ -832,8 +832,10 @@ public:
 	/**	Request an array of virtual texture page requests that was captured with SetVirtualTextureRequestRecordBuffer. Note that the array will be moved and ownership is taken. */
 	virtual void RequestVirtualTextureTiles(TArray<uint64>&& InPageRequests) = 0;
 
-	/** Evict all data from virtual texture caches*/
+	/** Evict all data from virtual texture caches. */
 	virtual void FlushVirtualTextureCache() = 0;
+	/** Evict all data from virtual texture cache for a given allocated virtual texture. */
+	virtual void FlushVirtualTextureCache(IAllocatedVirtualTexture* AllocatedVT, const FVector2f& InUV0, const FVector2f& InUV1) = 0;
 
 	/** Allocate a buffer and record all nanite page requests until the next call to either SetNaniteRequestRecordBuffer or GetNaniteRequestRecordBuffer. */
 	virtual void SetNaniteRequestRecordBuffer(uint64 Handle) = 0;
diff --git a/Engine/Source/Runtime/Renderer/Private/RendererModule.h b/Engine/Source/Runtime/Renderer/Private/RendererModule.h
--- a/Engine/Source/Runtime/Renderer/Private/RendererModule.h
+++ b/Engine/Source/Runtime/Renderer/Private/RendererModule.h
@@ -115,6 +115,7 @@ public:
 	virtual void SetVirtualTextureRequestRecordBuffer(uint64 Handle) override;
 	virtual uint64 GetVirtualTextureRequestRecordBuffer(TSet<uint64>& OutPageRequests) override;
 	virtual void FlushVirtualTextureCache() override;
+	virtual void FlushVirtualTextureCache(IAllocatedVirtualTexture* AllocatedVT, const FVector2f& InUV0, const FVector2f& InUV1) override;
 
 	virtual void SetNaniteRequestRecordBuffer(uint64 Handle) override;
 	virtual uint64 GetNaniteRequestRecordBuffer(TArray<uint32>& OutPageRequests) override; 	
diff --git a/Engine/Source/Runtime/Renderer/Private/SceneRendering.cpp b/Engine/Source/Runtime/Renderer/Private/SceneRendering.cpp
--- a/Engine/Source/Runtime/Renderer/Private/SceneRendering.cpp
+++ b/Engine/Source/Runtime/Renderer/Private/SceneRendering.cpp
@@ -101,6 +101,8 @@
 #include "LightFunctionAtlas.h"
 #include "EnvironmentComponentsFlags.h"
 #include "VolumetricCloudProxy.h"
+#include "VT/VirtualTextureScalability.h"
+#include "VT/VirtualTextureSystem.h"
 #include <type_traits>
 
 /*-----------------------------------------------------------------------------
@@ -5831,6 +5833,26 @@ void FRendererModule::FlushVirtualTextureCache()
 	FVirtualTextureSystem::Get().FlushCache();
 }
 
+void FRendererModule::FlushVirtualTextureCache(IAllocatedVirtualTexture* AllocatedVT, const FVector2f& InUV0, const FVector2f& InUV1)
+{
+	if (AllocatedVT != nullptr)
+	{
+		const uint32 NumLayers = AllocatedVT->GetNumTextureLayers();
+		const uint32 SpaceID = AllocatedVT->GetSpaceID();
+		const uint32 Width = AllocatedVT->GetBlockWidthInTiles() * AllocatedVT->GetVirtualTileSize();
+		const uint32 Height = AllocatedVT->GetBlockHeightInTiles() * AllocatedVT->GetVirtualTileSize();
+		const FIntPoint Texel0 = FIntPoint(FMath::FloorToInt32(InUV0.X * Width), FMath::FloorToInt32(InUV0.Y * Height));
+		const FIntPoint Texel1 = FIntPoint(FMath::CeilToInt32(InUV1.X * Width), FMath::CeilToInt32(InUV1.Y * Height));
+		const FIntRect TextureRect(Texel0, Texel1);
+		const uint32 MaxLevel = AllocatedVT->GetMaxLevel();
+		const uint32 MaxAgeToKeepMapped = VirtualTextureScalability::GetKeepDirtyPageMappedFrameThreshold();
+
+		for (uint32 LayerIndex = 0; LayerIndex < NumLayers;  ++LayerIndex)
+		{
+			FVirtualTextureSystem::Get().FlushCache(AllocatedVT->GetProducerHandle(LayerIndex), SpaceID, TextureRect, MaxLevel, MaxAgeToKeepMapped);
+		}
+	}
+}
 
 uint64 FRendererModule::GetNaniteRequestRecordBuffer(TArray<uint32>& OutPageRequests)
 {
