PostProcessing: Allow lens distortion LUT access and application from engine plugins.

We also expose a function to get the pass location of the lens distortion application in post-processing.  Scene view extension callbacks need this information to decided whether they should use distorted UVs or not.

diff --git a/Engine/Source/Runtime/Renderer/Internal/PostProcess/LensDistortion.h b/Engine/Source/Runtime/Renderer/Internal/PostProcess/LensDistortion.h
new file mode 100644
--- /dev/null
+++ b/Engine/Source/Runtime/Renderer/Internal/PostProcess/LensDistortion.h
@@ -0,0 +1,95 @@
+// Copyright Epic Games, Inc. All Rights Reserved.
+
+#pragma once
+
+#include "ScreenPass.h"
+
+struct FLensDistortionLUT
+{
+	/** Expected pixel format of the LUT. 16bit floating point do not work when distorting in TSR. */
+	static constexpr EPixelFormat kFormat = PF_G32R32F;
+
+	/** DistortedViewportUV = UndistortedViewportUV + DistortingDisplacementTexture.Sample(UndistortedViewportUV) */
+	FRDGTextureRef DistortingDisplacementTexture = nullptr;
+
+	/** UndistortedViewportUV = DistortedViewportUV + UndistortingDisplacementTexture.BilinearSample(DistortedViewportUV) */
+	FRDGTextureRef UndistortingDisplacementTexture = nullptr;
+
+	/** Resolution fraction of the upscaling happening due to distortion. */
+	float ResolutionFraction = 1.0f;
+
+	/** Returns whether the displacement is enabled. */
+	bool IsEnabled() const
+	{
+		return DistortingDisplacementTexture && UndistortingDisplacementTexture;
+	}
+};
+
+struct FPaniniProjectionConfig
+{
+	FPaniniProjectionConfig() = default;
+
+	/** Returns whether the panini is enabled by cvar. */
+	static bool IsEnabledByCVars();
+
+	/** Returns teh cvars' configuration. */
+	static FPaniniProjectionConfig ReadCVars();
+
+	bool IsEnabled() const
+	{
+		return D > 0.01f;
+	}
+
+	void Sanitize()
+	{
+		D = FMath::Max(D, 0.0f);
+	}
+
+	// 0=none..1=full, must be >= 0.
+	float D = 0.0f;
+
+	// Panini hard vertical compression lerp (0=no vertical compression, 1=hard compression).
+	float S = 0.0f;
+
+
+	/** Add a RDG pass to generate the lens distortion LUT from the settings. */
+	FLensDistortionLUT GenerateLUTPasses(FRDGBuilder& GraphBuilder, const FViewInfo& View) const;
+};
+
+namespace LensDistortion
+{
+	// Possible pass locations of the lens distortion application in post-processing.
+	enum class EPassLocation : uint8
+	{
+		TSR,
+		PrimaryUpscale
+	};
+
+	/**
+	 * Get the pass location of the lens distortion application in post-processing.
+	 *
+	 * @param InViewInfo Active view info.
+	 * @return Location enum.
+	 */
+	EPassLocation GetPassLocation(const FViewInfo& InViewInfo);
+
+	/**
+	 * Get the pass location of the lens distortion application in post-processing.
+	 * Unsafe to due internal argument cast from "FSceneView" to "FViewInfo".
+	 *
+	 * @param InView Active view info.
+	 * @return EPassLocation enum.
+	 */
+	RENDERER_API EPassLocation GetPassLocationUnsafe(const FSceneView& InView);
+
+	/**
+	 * Get the view lens distortion LUT.
+	 * Unsafe to due internal argument cast from "FSceneView" to "FViewInfo".
+	 *
+	 * @param InView Active view info.
+	 * @return FLensDistortionLUT texture struct.
+	 */
+	RENDERER_API const FLensDistortionLUT& GetLUTUnsafe(const FSceneView& InView);
+
+} // end namespace LensDistortion
+
diff --git a/Engine/Source/Runtime/Renderer/Private/PostProcess/LensDistortion.cpp b/Engine/Source/Runtime/Renderer/Private/PostProcess/LensDistortion.cpp
--- a/Engine/Source/Runtime/Renderer/Private/PostProcess/LensDistortion.cpp
+++ b/Engine/Source/Runtime/Renderer/Private/PostProcess/LensDistortion.cpp
@@ -1,9 +1,12 @@
 // Copyright Epic Games, Inc. All Rights Reserved.
 
 #include "PostProcess/LensDistortion.h"
+
+#include "PostProcessing.h"
 #include "PostProcess/SceneFilterRendering.h"
 #include "DataDrivenShaderPlatformInfo.h"
 #include "SceneRendering.h"
+#include "TemporalAA.h"
 
 namespace
 {
@@ -216,3 +219,32 @@ FLensDistortionLUT FPaniniProjectionConfig::GenerateLUTPasses(FRDGBuilder& Graph
 
 	return LensDistortionLUT;
 }
+
+LensDistortion::EPassLocation LensDistortion::GetPassLocation(const FViewInfo& InViewInfo)
+{
+	if (IsPostProcessingEnabled(InViewInfo)
+		&& GetMainTAAPassConfig(InViewInfo) == EMainTAAPassConfig::TSR
+		&& IsTSRLensDistortionEnabled(InViewInfo.GetShaderPlatform()))
+	{
+		return LensDistortion::EPassLocation::TSR;
+	}
+	else
+	{
+		return LensDistortion::EPassLocation::PrimaryUpscale;
+	}
+}
+
+LensDistortion::EPassLocation LensDistortion::GetPassLocationUnsafe(const FSceneView& InView)
+{
+	check(InView.bIsViewInfo);
+
+	return LensDistortion::GetPassLocation(static_cast<const FViewInfo&>(InView));
+}
+
+const FLensDistortionLUT& LensDistortion::GetLUTUnsafe(const FSceneView& InView)
+{
+	check(InView.bIsViewInfo);
+
+	return static_cast<const FViewInfo&>(InView).LensDistortionLUT;
+}
+
diff --git a/Engine/Source/Runtime/Renderer/Private/PostProcess/LensDistortion.h b/Engine/Source/Runtime/Renderer/Private/PostProcess/LensDistortion.h
deleted file mode 100644
diff --git a/Engine/Source/Runtime/Renderer/Private/PostProcess/PostProcessing.cpp b/Engine/Source/Runtime/Renderer/Private/PostProcess/PostProcessing.cpp
--- a/Engine/Source/Runtime/Renderer/Private/PostProcess/PostProcessing.cpp
+++ b/Engine/Source/Runtime/Renderer/Private/PostProcess/PostProcessing.cpp
@@ -379,8 +379,8 @@ void AddPostProcessingPasses(
 	// Temporal Anti-aliasing. Also may perform a temporal upsample from primary to secondary view rect.
 	const EMainTAAPassConfig TAAConfig = GetMainTAAPassConfig(View);
 
-	bool bApplyLensDistortion = View.LensDistortionLUT.IsEnabled();
-	bool bApplyLensDistortionInTSR = TAAConfig == EMainTAAPassConfig::TSR && bPostProcessingEnabled && IsTSRLensDistortionEnabled(View.GetShaderPlatform());
+	const bool bApplyLensDistortion = View.LensDistortionLUT.IsEnabled();
+	const bool bApplyLensDistortionInTSR = (LensDistortion::GetPassLocation(View) == LensDistortion::EPassLocation::TSR);
 
 	enum class EPass : uint32
 	{
