HWRT: RayTracingDynamicGeometry update:

* AddDynamicGeometryUpdatePass - will update the list of all geometries to build/update and add a RDG pass to build them
* Remove Segments array as it was not needed.
* Add new stats to see the number of dynamic primitives that are being built/updated
* Add r.RayTracing.DynamicGeometry.MaxUpdatePrimitivesPerFrame  to set a budget of how many primitives can be built per frame (all of them by default).
* Modify RayTracingSkinnedGometry to always allow at least one update so it's never stuck.
#rnx
#tests CitySample

diff --git a/Engine/Plugins/Experimental/GPULightmass/Source/GPULightmass/Private/LightmapRenderer.cpp b/Engine/Plugins/Experimental/GPULightmass/Source/GPULightmass/Private/LightmapRenderer.cpp
--- a/Engine/Plugins/Experimental/GPULightmass/Source/GPULightmass/Private/LightmapRenderer.cpp
+++ b/Engine/Plugins/Experimental/GPULightmass/Source/GPULightmass/Private/LightmapRenderer.cpp
@@ -1061,6 +1061,7 @@ bool FSceneRenderState::SetupRayTracingScene(FRDGBuilder& GraphBuilder, FSceneUn
 								StaticMeshInstanceRenderStates.Elements.Num() + InstanceGroupRenderStates.Elements.Num() + LandscapeIndex
 							);
 
+							PRAGMA_DISABLE_DEPRECATION_WARNINGS
 							const uint32 BLASScratchSize = DynamicGeometryCollection.ComputeScratchBufferSize();
 
 							FRHIResourceCreateInfo ScratchBufferCreateInfo(TEXT("RHILandscapeScratchBuffer"));
@@ -1068,6 +1069,7 @@ bool FSceneRenderState::SetupRayTracingScene(FRDGBuilder& GraphBuilder, FSceneUn
 
 							RHICmdList.SetStaticUniformBuffers({ReferenceView->ViewUniformBuffer, SceneUniforms.GetBufferRHI(GraphBuilder)});
 							DynamicGeometryCollection.DispatchUpdates(RHICmdList, ScratchBuffer);
+							PRAGMA_ENABLE_DEPRECATION_WARNINGS
 
 							// Landscape VF doesn't really use the vertex buffer in HitGroupSystemParameters
 							// We can release after all related RHI cmds get dispatched onto the cmd list
diff --git a/Engine/Source/Runtime/Engine/Private/RayTracingSkinnedGeometry.cpp b/Engine/Source/Runtime/Engine/Private/RayTracingSkinnedGeometry.cpp
--- a/Engine/Source/Runtime/Engine/Private/RayTracingSkinnedGeometry.cpp
+++ b/Engine/Source/Runtime/Engine/Private/RayTracingSkinnedGeometry.cpp
@@ -275,12 +275,7 @@ void FRayTracingSkinnedGeometryUpdateQueue::Commit(FRDGBuilder& GraphBuilder, ER
 
 		for (FRayTracingGeometry* RayTracingGeometry : GeometriesToUpdate)
 		{
-			const int32 NumPrimitives = RayTracingGeometry->Initializer.TotalPrimitiveCount;
-
-			if (NumUpdatedPrimitives + NumPrimitives > MaxUpdatePrimitivesPerFrame)
-			{
-				break;
-			}
+			const int32 NumPrimitives = RayTracingGeometry->Initializer.TotalPrimitiveCount;			
 						
 			FRayTracingGeometryBuildParams BuildParams;
 			BuildParams.Geometry = RayTracingGeometry->GetRHI();
@@ -294,6 +289,11 @@ void FRayTracingSkinnedGeometryUpdateQueue::Commit(FRDGBuilder& GraphBuilder, ER
 			BLASScratchSize += RayTracingGeometry->GetRHI()->GetSizeInfo().UpdateScratchSize;
 
 			NumUpdatedPrimitives += NumPrimitives;
+
+			if (NumUpdatedPrimitives > MaxUpdatePrimitivesPerFrame)
+			{
+				break;
+			}
 		}
 	}
 
diff --git a/Engine/Source/Runtime/Renderer/Private/DeferredShadingRenderer.cpp b/Engine/Source/Runtime/Renderer/Private/DeferredShadingRenderer.cpp
--- a/Engine/Source/Runtime/Renderer/Private/DeferredShadingRenderer.cpp
+++ b/Engine/Source/Runtime/Renderer/Private/DeferredShadingRenderer.cpp
@@ -285,7 +285,6 @@ DECLARE_CYCLE_STAT(TEXT("DeferredShadingSceneRenderer RenderFinish"), STAT_FDefe
 
 DECLARE_GPU_STAT(RayTracingScene);
 DECLARE_GPU_STAT(RayTracingGeometry);
-DECLARE_GPU_STAT(RayTracingDynamicGeometry);
 
 DEFINE_GPU_STAT(Postprocessing);
 DECLARE_GPU_STAT(VisibilityCommands);
@@ -861,41 +860,7 @@ bool FDeferredShadingSceneRenderer::DispatchRayTracingWorldUpdates(FRDGBuilder&
 	// Keep mask the same as what's already set (which will be the view mask) if TLAS updates should be masked to the view
 	RDG_GPU_MASK_SCOPE(GraphBuilder, GRayTracingMultiGpuTLASMask ? GraphBuilder.RHICmdList.GetGPUMask() : FRHIGPUMask::All());
 
-	const uint32 BLASScratchSize = Scene->GetRayTracingDynamicGeometryCollection()->ComputeScratchBufferSize();
-	if (BLASScratchSize > 0)
-	{
-		const uint32 ScratchAlignment = GRHIRayTracingScratchBufferAlignment;
-		FRDGBufferDesc ScratchBufferDesc;
-		ScratchBufferDesc.Usage = EBufferUsageFlags::RayTracingScratch | EBufferUsageFlags::StructuredBuffer;
-		ScratchBufferDesc.BytesPerElement = ScratchAlignment;
-		ScratchBufferDesc.NumElements = FMath::DivideAndRoundUp(BLASScratchSize, ScratchAlignment);
-
-		OutDynamicGeometryScratchBuffer = GraphBuilder.CreateBuffer(ScratchBufferDesc, TEXT("DynamicGeometry.BLASSharedScratchBuffer"));
-	}
-
-	{
-		// Dynamic geometry (BLAS) updates must always run on all GPUs.  Other passes may either run on all GPUs or be scoped to the view's GPUs.
-		// See GRayTracingMultiGpuTLASMask.
-		RDG_GPU_MASK_SCOPE(GraphBuilder, FRHIGPUMask::All());
-
-		FBuildAccelerationStructurePassParams* PassParams = GraphBuilder.AllocParameters<FBuildAccelerationStructurePassParams>();
-		PassParams->View = ReferenceView.ViewUniformBuffer;
-		PassParams->Scene = GetSceneUniforms().GetBuffer(GraphBuilder);
-		PassParams->DynamicGeometryScratchBuffer = OutDynamicGeometryScratchBuffer;
-		PassParams->LightGridPacked = nullptr;
-		PassParams->ClusterPageData = nullptr;
-		PassParams->HierarchyBuffer = nullptr;
-		PassParams->RayTracingDataBuffer = nullptr;
-
-		GraphBuilder.AddPass(RDG_EVENT_NAME("RayTracingDynamicUpdate"), PassParams, ComputePassFlags | ERDGPassFlags::NeverCull,
-			[this, PassParams, bRayTracingAsyncBuild](FRHICommandList& RHICmdList)
-			{
-				SCOPED_GPU_STAT(RHICmdList, RayTracingDynamicGeometry);
-				FRHIBuffer* DynamicGeometryScratchBuffer = PassParams->DynamicGeometryScratchBuffer ? PassParams->DynamicGeometryScratchBuffer->GetRHI() : nullptr;
-				Scene->GetRayTracingDynamicGeometryCollection()->DispatchUpdates(RHICmdList, DynamicGeometryScratchBuffer);
-				Scene->GetRayTracingDynamicGeometryCollection()->EndUpdate();
-			});
-	}
+	Scene->GetRayTracingDynamicGeometryCollection()->AddDynamicGeometryUpdatePass(ReferenceView, GraphBuilder, ComputePassFlags, OutDynamicGeometryScratchBuffer);
 
 
 	{
diff --git a/Engine/Source/Runtime/Renderer/Private/RayTracing/RayTracingDynamicGeometry.cpp b/Engine/Source/Runtime/Renderer/Private/RayTracing/RayTracingDynamicGeometry.cpp
--- a/Engine/Source/Runtime/Renderer/Private/RayTracing/RayTracingDynamicGeometry.cpp
+++ b/Engine/Source/Runtime/Renderer/Private/RayTracing/RayTracingDynamicGeometry.cpp
@@ -6,11 +6,17 @@
 #include "RayTracingDynamicGeometryCollection.h"
 #include "RayTracingInstance.h"
 #include "RayTracingGeometry.h"
+#include "RenderGraphBuilder.h"
 
 #if RHI_RAYTRACING
 
 #include "Materials/MaterialRenderProxy.h"
 
+DECLARE_GPU_STAT(RayTracingDynamicGeometry);
+
+DECLARE_DWORD_COUNTER_STAT(TEXT("Ray tracing dynamic build primitives"), STAT_RayTracingDynamicBuildPrimitives, STATGROUP_SceneRendering);
+DECLARE_DWORD_COUNTER_STAT(TEXT("Ray tracing dynamic update primitives"), STAT_RayTracingDynamicUpdatePrimitives, STATGROUP_SceneRendering);
+
 static int32 GRTDynGeomSharedVertexBufferSizeInMB = 4;
 static FAutoConsoleVariableRef CVarRTDynGeomSharedVertexBufferSizeInMB(
 	TEXT("r.RayTracing.DynamicGeometry.SharedVertexBufferSizeInMB"),
@@ -27,6 +33,13 @@ static FAutoConsoleVariableRef CVarRTDynGeomSharedVertexBufferGarbageCollectLate
 	ECVF_RenderThreadSafe
 );
 
+static TAutoConsoleVariable<int32> CVarRTDynGeomMaxUpdatePrimitivesPerFrame(
+	TEXT("r.RayTracing.DynamicGeometry.MaxUpdatePrimitivesPerFrame"),
+	-1,
+	TEXT("Sets the dynamic ray tracing acceleration structure build budget in terms of maximum number of updated triangles per frame (<= 0 then disabled and all acceleration structures are updated - default)"),
+	ECVF_RenderThreadSafe
+);
+
 class FRayTracingDynamicGeometryConverterCS : public FMeshMaterialShader
 {
 	DECLARE_SHADER_TYPE(FRayTracingDynamicGeometryConverterCS, MeshMaterial);
@@ -120,15 +133,18 @@ void FRayTracingDynamicGeometryCollection::Clear()
 	// Clear working arrays - keep max size allocated
 	DispatchCommands.Empty(DispatchCommands.Max());
 	BuildParams.Empty(BuildParams.Max());
-	Segments.Empty(Segments.Max());
+
+	DynamicGeometryBuilds.Empty(DynamicGeometryBuilds.Max());
+	DynamicGeometryUpdates.Empty(DynamicGeometryUpdates.Max());
 }
 
 int64 FRayTracingDynamicGeometryCollection::BeginUpdate()
 {
 	check(DispatchCommands.IsEmpty());
 	check(BuildParams.IsEmpty());
-	check(Segments.IsEmpty());
 	check(ReferencedUniformBuffers.IsEmpty());
+	check(DynamicGeometryBuilds.IsEmpty());
+	check(DynamicGeometryUpdates.IsEmpty());
 
 	// Vertex buffer data can be immediatly reused the next frame, because it's already 'consumed' for building the AccelerationStructure data
 	// Garbage collect unused buffers for n generations
@@ -222,6 +238,9 @@ void FRayTracingDynamicGeometryCollection::AddDynamicMeshBatchForGeometryUpdate(
 		RWBuffer = &VertexPositionBuffer->RWBuffer;
 	}
 
+	FRayTracingDynamicGeometryBuildParams GeometryBuildParams;
+	GeometryBuildParams.DispatchCommands.Reserve(UpdateParams.MeshBatches.Num());
+
 	for (const FMeshBatch& MeshBatch : UpdateParams.MeshBatches)
 	{
 		if (!ensureMsgf(MeshBatch.VertexFactory->GetType()->SupportsRayTracingDynamicGeometry(),
@@ -302,7 +321,7 @@ void FRayTracingDynamicGeometryCollection::AddDynamicMeshBatchForGeometryUpdate(
 					ShaderBindings.Finalize(&MeshProcessorShaders);
 #endif
 
-					DispatchCommands.Add(DispatchCmd);
+					GeometryBuildParams.DispatchCommands.Add(DispatchCmd);
 
 					break;
 				}
@@ -357,23 +376,22 @@ void FRayTracingDynamicGeometryCollection::AddDynamicMeshBatchForGeometryUpdate(
 		Geometry.CreateRayTracingGeometry(RHICmdList, ERTAccelerationStructureBuildPriority::Skip);
 	}
 
-	FRayTracingGeometryBuildParams Params;
-	Params.Geometry = Geometry.GetRHI();
-	Params.BuildMode = Geometry.GetRequiresBuild()
+	EAccelerationStructureBuildMode BuildMode = Geometry.GetRequiresBuild()
 		? EAccelerationStructureBuildMode::Build
 		: EAccelerationStructureBuildMode::Update;
 
+	GeometryBuildParams.Geometry = UpdateParams.Geometry;
+
 	Geometry.SetRequiresBuild(false);
 
-	if (bUseSharedVertexBuffer)
+	if (BuildMode == EAccelerationStructureBuildMode::Build)
 	{
-		Segments.Append(Geometry.Initializer.Segments);
-
-		// Cache the count of segments so final views can be made when all segments are collected (Segments array could still be reallocated)
-		Params.Segments = MakeArrayView((FRayTracingGeometrySegment*)nullptr, Geometry.Initializer.Segments.Num());
+		DynamicGeometryBuilds.Add(GeometryBuildParams);
+	}
+	else
+	{
+		DynamicGeometryUpdates.Add(GeometryBuildParams);
 	}
-
-	BuildParams.Add(Params);
 	
 	if (bUseSharedVertexBuffer)
 	{
@@ -385,6 +403,149 @@ void FRayTracingDynamicGeometryCollection::AddDynamicMeshBatchForGeometryUpdate(
 	}
 }
 
+BEGIN_SHADER_PARAMETER_STRUCT(FRayTracingDynamicGeometryUpdatePassParams, )
+	RDG_BUFFER_ACCESS(DynamicGeometryScratchBuffer, ERHIAccess::UAVCompute)
+
+	SHADER_PARAMETER_STRUCT_REF(FViewUniformShaderParameters, View)
+	SHADER_PARAMETER_RDG_UNIFORM_BUFFER(FSceneUniformParameters, Scene)
+END_SHADER_PARAMETER_STRUCT()
+
+uint32 FRayTracingDynamicGeometryCollection::Update()
+{
+	TRACE_CPUPROFILER_EVENT_SCOPE(FRayTracingDynamicGeometryCollection::Update);
+
+	const int32 TotalNumGeometryBuilds = DynamicGeometryBuilds.Num() + DynamicGeometryUpdates.Num();
+	if (TotalNumGeometryBuilds == 0)
+	{
+		return 0;
+	}
+
+	checkf(DispatchCommands.IsEmpty(), TEXT("DispatchCommands is not empty. Previous frame updates were not dispatched."));
+	checkf(BuildParams.IsEmpty(), TEXT("BuildParams is not empty. Previous frame updates were not dispatched."));
+
+	DispatchCommands.Reserve(TotalNumGeometryBuilds);
+	BuildParams.Reserve(TotalNumGeometryBuilds);
+
+	const uint32 ScratchAlignment = GRHIRayTracingScratchBufferAlignment;
+
+	uint32 BLASScratchSize = 0;
+	int32 NumBuildPrimitives = 0;
+
+	for (const FRayTracingDynamicGeometryBuildParams& Build : DynamicGeometryBuilds)
+	{
+		FRHIRayTracingGeometry* RayTracingGeometry = Build.Geometry->GetRHI();
+
+		NumBuildPrimitives += Build.Geometry->Initializer.TotalPrimitiveCount;
+
+		const uint32 ScratchSize = RayTracingGeometry->GetSizeInfo().BuildScratchSize;
+		BLASScratchSize = Align(BLASScratchSize + ScratchSize, ScratchAlignment);
+
+		FRayTracingGeometryBuildParams BuildParam;
+		BuildParam.Geometry = RayTracingGeometry;
+		BuildParam.BuildMode = EAccelerationStructureBuildMode::Build;
+		BuildParams.Add(MoveTemp(BuildParam));
+
+		DispatchCommands.Append(Build.DispatchCommands);
+	}
+
+	const int32 MaxUpdatePrimitivesPerFrame = CVarRTDynGeomMaxUpdatePrimitivesPerFrame.GetValueOnRenderThread();
+
+	int32 NumUpdatedPrimitives = 0;
+
+	if (MaxUpdatePrimitivesPerFrame <= 0)
+	{
+		for (const FRayTracingDynamicGeometryBuildParams& Update : DynamicGeometryUpdates)
+		{
+			FRHIRayTracingGeometry* RayTracingGeometry = Update.Geometry->GetRHI();
+
+			Update.Geometry->LastUpdatedFrame = GFrameCounterRenderThread;
+
+			NumUpdatedPrimitives += Update.Geometry->Initializer.TotalPrimitiveCount;
+
+			const uint32 ScratchSize = RayTracingGeometry->GetSizeInfo().UpdateScratchSize;
+			BLASScratchSize = Align(BLASScratchSize + ScratchSize, ScratchAlignment);
+
+			FRayTracingGeometryBuildParams BuildParam;
+			BuildParam.Geometry = RayTracingGeometry;
+			BuildParam.BuildMode = EAccelerationStructureBuildMode::Update;
+			BuildParams.Add(MoveTemp(BuildParam));
+
+			DispatchCommands.Append(Update.DispatchCommands);
+		}
+	}
+	else
+	{
+		DynamicGeometryUpdates.Sort([](const FRayTracingDynamicGeometryBuildParams& InLHS, const FRayTracingDynamicGeometryBuildParams& InRHS)
+			{
+				return InLHS.Geometry->LastUpdatedFrame < InRHS.Geometry->LastUpdatedFrame;
+			});
+
+		for (const FRayTracingDynamicGeometryBuildParams& Update : DynamicGeometryUpdates)
+		{
+			FRHIRayTracingGeometry* RayTracingGeometry = Update.Geometry->GetRHI();
+
+			Update.Geometry->LastUpdatedFrame = GFrameCounterRenderThread;
+
+			NumUpdatedPrimitives += Update.Geometry->Initializer.TotalPrimitiveCount;
+
+			const uint32 ScratchSize = RayTracingGeometry->GetSizeInfo().UpdateScratchSize;
+			BLASScratchSize = Align(BLASScratchSize + ScratchSize, ScratchAlignment);
+
+			FRayTracingGeometryBuildParams BuildParam;
+			BuildParam.Geometry = RayTracingGeometry;
+			BuildParam.BuildMode = EAccelerationStructureBuildMode::Update;
+			BuildParams.Add(MoveTemp(BuildParam));
+
+			DispatchCommands.Append(Update.DispatchCommands);
+
+			if (NumUpdatedPrimitives > MaxUpdatePrimitivesPerFrame)
+			{
+				break;
+			}
+		}
+	}
+
+	INC_DWORD_STAT_BY(STAT_RayTracingDynamicUpdatePrimitives, NumUpdatedPrimitives);
+	INC_DWORD_STAT_BY(STAT_RayTracingDynamicBuildPrimitives, NumBuildPrimitives);
+
+	return BLASScratchSize;
+}
+
+void FRayTracingDynamicGeometryCollection::AddDynamicGeometryUpdatePass(const FViewInfo& View, FRDGBuilder& GraphBuilder, ERDGPassFlags ComputePassFlags, FRDGBufferRef& OutDynamicGeometryScratchBuffer)
+{
+	RDG_GPU_MASK_SCOPE(GraphBuilder, FRHIGPUMask::All());
+	
+	const uint32 ScratchAlignment = GRHIRayTracingScratchBufferAlignment;
+	const uint32 BLASScratchSize = Update();
+
+	if (BLASScratchSize > 0)
+	{
+		FRDGBufferDesc ScratchBufferDesc;
+		ScratchBufferDesc.Usage = EBufferUsageFlags::RayTracingScratch | EBufferUsageFlags::StructuredBuffer;
+		ScratchBufferDesc.BytesPerElement = ScratchAlignment;
+		ScratchBufferDesc.NumElements = FMath::DivideAndRoundUp(BLASScratchSize, ScratchAlignment);
+
+		OutDynamicGeometryScratchBuffer = GraphBuilder.CreateBuffer(ScratchBufferDesc, TEXT("DynamicGeometry.BLASSharedScratchBuffer"));
+	}
+
+	FRayTracingDynamicGeometryUpdatePassParams* PassParams = GraphBuilder.AllocParameters<FRayTracingDynamicGeometryUpdatePassParams>();
+	PassParams->View = View.ViewUniformBuffer;
+	PassParams->Scene = View.GetSceneUniforms().GetBuffer(GraphBuilder);
+	PassParams->DynamicGeometryScratchBuffer = OutDynamicGeometryScratchBuffer;	
+
+	GraphBuilder.AddPass(RDG_EVENT_NAME("RayTracingDynamicUpdate"), PassParams, ComputePassFlags | ERDGPassFlags::NeverCull,
+		[this, PassParams](FRHICommandList& RHICmdList)
+		{
+			SCOPED_GPU_STAT(RHICmdList, RayTracingDynamicGeometry);
+			FRHIBuffer* DynamicGeometryScratchBuffer = PassParams->DynamicGeometryScratchBuffer ? PassParams->DynamicGeometryScratchBuffer->GetRHI() : nullptr;
+
+			PRAGMA_DISABLE_DEPRECATION_WARNINGS
+			DispatchUpdates(RHICmdList, DynamicGeometryScratchBuffer);
+			EndUpdate();
+			PRAGMA_ENABLE_DEPRECATION_WARNINGS
+		});
+}
+
 void FRayTracingDynamicGeometryCollection::DispatchUpdates(FRHICommandList& RHICmdList, FRHIBuffer* ScratchBuffer)
 {
 	if (DispatchCommands.Num() > 0)
@@ -406,22 +567,6 @@ void FRayTracingDynamicGeometryCollection::DispatchUpdates(FRHICommandList& RHIC
 				});
 		}
 
-		{
-			TRACE_CPUPROFILER_EVENT_SCOPE(SetupSegmentData);
-
-			// Setup the array views on final allocated segments array
-			FRayTracingGeometrySegment* SegmentData = Segments.GetData();
-			for (FRayTracingGeometryBuildParams& Param : BuildParams)
-			{
-				uint32 SegmentCount = Param.Segments.Num();
-				if (SegmentCount > 0)
-				{
-					Param.Segments = MakeArrayView(SegmentData, SegmentCount);
-					SegmentData += SegmentCount;
-				}
-			}
-		}
-
 		FMemMark Mark(FMemStack::Get());
 
 		TArray<FRHITransitionInfo, TMemStackAllocator<>> TransitionsBefore, TransitionsAfter;
@@ -534,21 +679,8 @@ void FRayTracingDynamicGeometryCollection::EndUpdate()
 }
 
 uint32 FRayTracingDynamicGeometryCollection::ComputeScratchBufferSize()
-{
-	TRACE_CPUPROFILER_EVENT_SCOPE(FRayTracingDynamicGeometryCollection::ComputeScratchBufferSize);
-
-	const uint64 ScratchAlignment = GRHIRayTracingScratchBufferAlignment;
-
-	uint32 BLASScratchSize = 0;
-
-	for (FRayTracingGeometryBuildParams& Params : BuildParams)
-	{
-		const FRayTracingAccelerationStructureSize BLASSizeInfo = Params.Geometry->GetSizeInfo();
-		const uint64 ScratchSize = Params.BuildMode == EAccelerationStructureBuildMode::Build ? BLASSizeInfo.BuildScratchSize : BLASSizeInfo.UpdateScratchSize;
-		BLASScratchSize = Align(BLASScratchSize + ScratchSize, ScratchAlignment);
-	}
-
-	return BLASScratchSize;
+{	
+	return Update();
 }
 
 #undef USE_RAY_TRACING_DYNAMIC_GEOMETRY_PARALLEL_COMMAND_LISTS
diff --git a/Engine/Source/Runtime/Renderer/Public/RayTracingDynamicGeometryCollection.h b/Engine/Source/Runtime/Renderer/Public/RayTracingDynamicGeometryCollection.h
--- a/Engine/Source/Runtime/Renderer/Public/RayTracingDynamicGeometryCollection.h
+++ b/Engine/Source/Runtime/Renderer/Public/RayTracingDynamicGeometryCollection.h
@@ -8,10 +8,14 @@
 
 #include "RHI.h"
 #include "RHIUtilities.h"
+#include "RenderGraphDefinitions.h"
 
 class FPrimitiveSceneProxy;
 class FScene;
 class FSceneView;
+class FViewInfo;
+class FRDGBuilder;
+class FRayTracingGeometry;
 struct FMeshComputeDispatchCommand;
 struct FRayTracingDynamicGeometryUpdateParams;
 
@@ -41,19 +45,37 @@ public:
 
 	// Starts an update batch and returns the current shared buffer generation ID which is used for validation.
 	RENDERER_API int64 BeginUpdate();
+
+	UE_DEPRECATED(5.5, "Use AddDynamicGeometryUpdatePass instead.")
 	RENDERER_API void DispatchUpdates(FRHICommandList& RHICmdList, FRHIBuffer* ScratchBuffer);
+
+	UE_DEPRECATED(5.5, "Use AddDynamicGeometryUpdatePass instead.")
 	RENDERER_API void EndUpdate();
 
+	RENDERER_API void AddDynamicGeometryUpdatePass(const FViewInfo& View, FRDGBuilder& GraphBuilder, ERDGPassFlags ComputePassFlags, FRDGBufferRef& OutDynamicGeometryScratchBuffer);
+
 	// Clears the working arrays to not hold any references.
 	RENDERER_API void Clear();
 
+	UE_DEPRECATED(5.5, "Use AddDynamicGeometryUpdatePass instead which allocates scratch buffer internally.")
 	RENDERER_API uint32 ComputeScratchBufferSize();
 
 private:
 
+	// Prepares DispatchCommands and BuildParams and returns the number of bytes needed for scratch buffer
+	uint32 Update();
+
+	struct FRayTracingDynamicGeometryBuildParams
+	{
+		TArray<FMeshComputeDispatchCommand> DispatchCommands;
+		FRayTracingGeometry* Geometry;
+	};
+
+	TArray<FRayTracingDynamicGeometryBuildParams> DynamicGeometryBuilds;
+	TArray<FRayTracingDynamicGeometryBuildParams> DynamicGeometryUpdates;
+
 	TArray<FMeshComputeDispatchCommand> DispatchCommands;
 	TArray<FRayTracingGeometryBuildParams> BuildParams;
-	TArray<FRayTracingGeometrySegment> Segments;
 
 	struct FVertexPositionBuffer
 	{
