Fix CIS warning due to a typo.

diff --git a/Engine/Source/Runtime/Renderer/Private/LightGridInjection.cpp b/Engine/Source/Runtime/Renderer/Private/LightGridInjection.cpp
--- a/Engine/Source/Runtime/Renderer/Private/LightGridInjection.cpp
+++ b/Engine/Source/Runtime/Renderer/Private/LightGridInjection.cpp
@@ -315,10 +315,8 @@ static uint32 PackRGB10(float In0, float In1, float In2)
 {
 	return 
 		(uint32(FMath::Clamp(In0 * 1023u, 0u, 1023u))      )|
-		(uint32(FMath::Clamp(In0 * 1023u, 0u, 1023u)) << 10)|
-		(uint32(FMath::Clamp(In0 * 1023u, 0u, 1023u)) << 10);
-
-	//return uint32(FFloat16(In0).Encoded) | (uint32(FFloat16(In1).Encoded) << 16);
+		(uint32(FMath::Clamp(In1 * 1023u, 0u, 1023u)) << 10)|
+		(uint32(FMath::Clamp(In2 * 1023u, 0u, 1023u)) << 10);
 }
 
 static FVector2f PackLightColor(const FVector3f& LightColor)
