Revert one of the LFAtlas fix for clustered shading as it creates unexplained GPU regression on certain platforms.

diff --git a/Engine/Source/Runtime/Renderer/Private/LightRendering.cpp b/Engine/Source/Runtime/Renderer/Private/LightRendering.cpp
--- a/Engine/Source/Runtime/Renderer/Private/LightRendering.cpp
+++ b/Engine/Source/Runtime/Renderer/Private/LightRendering.cpp
@@ -1254,7 +1254,7 @@ void FSceneRenderer::GatherAndSortLights(FSortedLightSetSceneInfo& OutSortedLigh
 					// Contact shadow are not supported.
 					const bool bClusteredDeferredSupported =
 						(!SortedLightInfo->SortKey.Fields.bShadowed || bShadowedLightsInClustered) &&
-						(!SortedLightInfo->SortKey.Fields.bLightFunction || (bUseLightFunctionAtlas && SortedLightInfo->bIsCompatibleWithLightFunctionAtlas))
+						(!SortedLightInfo->SortKey.Fields.bLightFunction || (bUseLightFunctionAtlas))
 						&& LightSceneInfoCompact.LightType != LightType_Directional
 						&& LightSceneInfo->Proxy->GetContactShadowLength() == 0
 						&& !bHandledByManyLights;
