Allow all virtual texture physical textures to create a UAV (except block compressed ones on platforms that don't support aliasing).
This removes special case for lightmaps, and also will enable incoming virtual texture adapter.
#rnx

diff --git a/Engine/Source/Runtime/Renderer/Private/VT/VirtualTexturePhysicalSpace.cpp b/Engine/Source/Runtime/Renderer/Private/VT/VirtualTexturePhysicalSpace.cpp
--- a/Engine/Source/Runtime/Renderer/Private/VT/VirtualTexturePhysicalSpace.cpp
+++ b/Engine/Source/Runtime/Renderer/Private/VT/VirtualTexturePhysicalSpace.cpp
@@ -143,7 +143,7 @@ void FVirtualTexturePhysicalSpace::InitRHI(FRHICommandListBase& RHICmdList)
 	{
 		const EPixelFormat FormatSRV = RemapVirtualTexturePhysicalSpaceFormat(Description.Format[Layer]);
 		const EPixelFormat FormatUAV = GetUnorderedAccessViewFormat(FormatSRV);
-		const bool bCreateAliasedUAV = (FormatUAV != PF_Unknown) && (FormatUAV != FormatSRV);
+		const bool bCreateUAV = (FormatUAV != PF_Unknown);
 		
 		// Not all RHIs support sRGB views/aliasing. On those platforms create texture in an expected storage format 
 		const bool bDefaultToSRGB = Description.bHasLayerSrgbView[Layer];
@@ -158,11 +158,10 @@ void FVirtualTexturePhysicalSpace::InitRHI(FRHICommandListBase& RHICmdList)
 			FormatSRV,
 			FClearValueBinding::None,
 			VT_SRGB,
-			// GPULightmass hack: always create UAV for PF_A32B32G32R32F
-			(bCreateAliasedUAV || FormatSRV == PF_A32B32G32R32F) ? TexCreate_ShaderResource | TexCreate_UAV : TexCreate_ShaderResource,
+			bCreateUAV ? TexCreate_ShaderResource | TexCreate_UAV : TexCreate_ShaderResource,
 			false);
 
-		if (bCreateAliasedUAV)
+		if (bCreateUAV)
 		{
 			Desc.UAVFormat = FormatUAV;
 		}
