Do not use shadow stencil culling with a mobile deferred shading. As there is not enough stencil bits to support shading models and shadow stencil mask at the same time

diff --git a/Engine/Source/Runtime/Renderer/Private/ShadowRendering.cpp b/Engine/Source/Runtime/Renderer/Private/ShadowRendering.cpp
--- a/Engine/Source/Runtime/Renderer/Private/ShadowRendering.cpp
+++ b/Engine/Source/Runtime/Renderer/Private/ShadowRendering.cpp
@@ -174,6 +174,16 @@ static FAutoConsoleVariableRef CVarGShadowStencilCulling(
 	ECVF_RenderThreadSafe
 );
 
+static bool UseShadowStencilCulling(FStaticShaderPlatform ShaderPlatform)
+{
+	if (IsMobilePlatform(ShaderPlatform) && IsMobileDeferredShadingEnabled(ShaderPlatform))
+	{
+		// ShadowStencilMask now clashes with a mobile deferred ShadingModels mask
+		return false;
+	}
+	return GShadowStencilCulling != 0;
+}
+
 static TAutoConsoleVariable<int32> CVarFilterMethod(
 	TEXT("r.Shadow.FilterMethod"),
 	0,
@@ -1294,7 +1304,7 @@ void FProjectedShadowInfo::RenderProjectionInternal(
 	SetupFrustumForProjection(View, FrustumVertices, bCameraInsideShadowFrustum, OutFrustmPlanes);
 
 	const bool bSubPixelSupport = HairStrandsUniformBuffer != nullptr;// HairStrands::HasViewHairStrandsData(*View);
-	const bool bStencilTestEnabled = !bSubPixelSupport && GShadowStencilCulling;
+	const bool bStencilTestEnabled = !bSubPixelSupport && UseShadowStencilCulling(View ? View->GetShaderPlatform() : GMaxRHIShaderPlatform);
 	const bool bDepthBoundsTestEnabled = IsWholeSceneDirectionalShadow() && GSupportsDepthBoundsTest && CVarCSMDepthBoundsTest.GetValueOnRenderThread() != 0;// && !bSubPixelSupport;
 	const uint32 StencilRef = bSubPixelSupport && !IsWholeSceneDirectionalShadow() && !bCameraInsideShadowFrustum ? 1u : 0u;
 
