Fixed Static Analysis warning introduced in CL#34522078

diff --git a/Engine/Source/Runtime/Renderer/Private/ShadowRendering.cpp b/Engine/Source/Runtime/Renderer/Private/ShadowRendering.cpp
--- a/Engine/Source/Runtime/Renderer/Private/ShadowRendering.cpp
+++ b/Engine/Source/Runtime/Renderer/Private/ShadowRendering.cpp
@@ -1304,7 +1304,7 @@ void FProjectedShadowInfo::RenderProjectionInternal(
 	SetupFrustumForProjection(View, FrustumVertices, bCameraInsideShadowFrustum, OutFrustmPlanes);
 
 	const bool bSubPixelSupport = HairStrandsUniformBuffer != nullptr;// HairStrands::HasViewHairStrandsData(*View);
-	const bool bStencilTestEnabled = !bSubPixelSupport && UseShadowStencilCulling(View ? View->GetShaderPlatform() : GMaxRHIShaderPlatform);
+	const bool bStencilTestEnabled = !bSubPixelSupport && UseShadowStencilCulling(View->GetShaderPlatform());
 	const bool bDepthBoundsTestEnabled = IsWholeSceneDirectionalShadow() && GSupportsDepthBoundsTest && CVarCSMDepthBoundsTest.GetValueOnRenderThread() != 0;// && !bSubPixelSupport;
 	const uint32 StencilRef = bSubPixelSupport && !IsWholeSceneDirectionalShadow() && !bCameraInsideShadowFrustum ? 1u : 0u;
 
