Implements r.Test.ViewRollAngle

diff --git a/Engine/Source/Runtime/Renderer/Private/SceneRendering.cpp b/Engine/Source/Runtime/Renderer/Private/SceneRendering.cpp
--- a/Engine/Source/Runtime/Renderer/Private/SceneRendering.cpp
+++ b/Engine/Source/Runtime/Renderer/Private/SceneRendering.cpp
@@ -100,6 +100,7 @@
 #include "Rendering/CustomRenderPass.h"
 #include "LightFunctionAtlas.h"
 #include "EnvironmentComponentsFlags.h"
+#include "Math/RotationMatrix.h"
 #include "VolumetricCloudProxy.h"
 #include "VT/VirtualTextureScalability.h"
 #include "VT/VirtualTextureSystem.h"
@@ -470,6 +471,10 @@ static TAutoConsoleVariable<int32> CVarTestCameraCut(
 	TEXT("Force enabling camera cut for testing purposes.\n")
 	TEXT(" 0: disabled (default); 1: enabled."));
 
+static TAutoConsoleVariable<float> CVarTestViewRollAngle(
+	TEXT("r.Test.ViewRollAngle"), 0.0f,
+	TEXT("Roll the camera in degrees, for testing motion vector upscaling precision. (disabled by default)"));
+
 static TAutoConsoleVariable<int32> CVarTestScreenPercentageInterface(
 	TEXT("r.Test.DynamicResolutionHell"),
 	0,
@@ -2695,6 +2700,25 @@ FSceneRenderer::FSceneRenderer(const FSceneViewFamily* InViewFamily, FHitProxyCo
 		ViewInfo->DynamicPrimitiveCollector = FGPUScenePrimitiveCollector(&GPUSceneDynamicContext);
 		ViewInfo->RayTracingDynamicPrimitiveCollector = FGPUScenePrimitiveCollector(&GPUSceneDynamicContext);
 
+#if !UE_BUILD_SHIPPING
+		if (float ViewRollAngle = CVarTestViewRollAngle.GetValueOnGameThread())
+		{
+			FViewMatrices& CurrentMatrices = ViewInfo->ViewMatrices;
+
+			FRotator Rotate(/* InPitch = */ 0.0, /* InYaw = */ ViewRollAngle, /* Roll = */ 0.0);
+			FMatrix Rotation = FRotationMatrix::Make(Rotate);
+
+			FViewMatrices::FMinimalInitializer NewMatrices;
+			NewMatrices.ViewRotationMatrix  = CurrentMatrices.GetViewMatrix().RemoveTranslation() * Rotation;
+			NewMatrices.ProjectionMatrix    = CurrentMatrices.GetProjectionMatrix();
+			NewMatrices.ViewOrigin          = CurrentMatrices.GetViewOrigin();
+			NewMatrices.ConstrainedViewRect = ViewInfo->CameraConstrainedViewRect;
+			NewMatrices.CameraToViewTarget  = CurrentMatrices.GetCameraToViewTarget();
+
+			CurrentMatrices = FViewMatrices(NewMatrices);
+		}
+#endif
+
 		check(ViewInfo->ViewRect.Area() == 0);
 
 #if WITH_EDITOR
