Fixed light's IESAtlasIndex encoding bug from CL 34495601. It was clamping IESAtlasIndex to 0, which was causing INDEX_NONE (-1) to be decoded as 0.

diff --git a/Engine/Shaders/Private/LightData.ush b/Engine/Shaders/Private/LightData.ush
--- a/Engine/Shaders/Private/LightData.ush
+++ b/Engine/Shaders/Private/LightData.ush
@@ -307,7 +307,8 @@ float UnpackLightDiffuseScale(FLocalLightData In)
 
 float UnpackLigthIESAtlasIndex(FLocalLightData In)
 {
-	return asuint(In.LightTangentAndIESDataAndSpecularScale.w) >> 20;
+	// Offset IESAtlasIndex here in order to preserve INDEX_NONE = -1 after encoding
+	return int(asuint(In.LightTangentAndIESDataAndSpecularScale.w) >> 20) - 1;
 }
 
 ///////////////////////////////////////////////////////////////////////////////////////////////////
diff --git a/Engine/Source/Runtime/Renderer/Private/LightGridInjection.cpp b/Engine/Source/Runtime/Renderer/Private/LightGridInjection.cpp
--- a/Engine/Source/Runtime/Renderer/Private/LightGridInjection.cpp
+++ b/Engine/Source/Runtime/Renderer/Private/LightGridInjection.cpp
@@ -385,7 +385,9 @@ static void PackLocalLightData(
 	const float SpecularScale = SimpleLight.SpecularScale;
 	const float DiffuseScale  = SimpleLight.DiffuseScale;
 	const float IESAtlasIndex = INDEX_NONE;
-	const uint32 SpecularScale_DiffuseScale_IESData = PackRGB10(SpecularScale, DiffuseScale, IESAtlasIndex);
+
+	// Offset IESAtlasIndex here in order to preserve INDEX_NONE = -1 after encoding
+	const uint32 SpecularScale_DiffuseScale_IESData = PackRGB10(SpecularScale, DiffuseScale, IESAtlasIndex + 1);
 
 	const FVector3f LightColor = (FVector3f)SimpleLight.Color * FLightRenderParameters::GetLightExposureScale(View.GetLastEyeAdaptationExposure(), SimpleLight.InverseExposureBlend);
 	const FVector2f LightColorPacked = PackLightColor(LightColor);
@@ -427,7 +429,8 @@ static void PackLocalLightData(
 	RectPackedZ |= uint32(FMath::Clamp(LightParameters.RectLightAtlasMaxLevel, 0.f, 63.f)) << 26;			//  6 bits
 
 	// Pack specular scale and IES profile index
-	const uint32 SpecularScale_DiffuseScale_IESData = PackRGB10(LightParameters.SpecularScale, LightParameters.DiffuseScale, LightParameters.IESAtlasIndex); // pack atlas id here? 16bit specular 8bit IES and 8 bit LightFunction
+	 // Offset IESAtlasIndex here in order to preserve INDEX_NONE = -1 after encoding
+	const uint32 SpecularScale_DiffuseScale_IESData = PackRGB10(LightParameters.SpecularScale, LightParameters.DiffuseScale, LightParameters.IESAtlasIndex + 1); // pack atlas id here? 16bit specular 8bit IES and 8 bit LightFunction
 
 	const FVector2f LightColorPacked = PackLightColor(FVector3f(LightParameters.Color));
 
