Implements TSR reprojection field to improve geometric edge reprojection's quality

diff --git a/Engine/Shaders/Private/TemporalSuperResolution/TSRCommon.ush b/Engine/Shaders/Private/TemporalSuperResolution/TSRCommon.ush
--- a/Engine/Shaders/Private/TemporalSuperResolution/TSRCommon.ush
+++ b/Engine/Shaders/Private/TemporalSuperResolution/TSRCommon.ush
@@ -14,7 +14,7 @@
 
 //------------------------------------------------------- RECOMPILE HASH
 
-#pragma message("UESHADERMETADATA_VERSION F3C5FF82-E809-4903-98CD-5CC7E2254C77")
+#pragma message("UESHADERMETADATA_VERSION F3C5FF82-E809-4903-98CD-5DC7E2254C87")
 
 
 //------------------------------------------------------- COMPILER CONFIG
@@ -36,9 +36,11 @@
 #endif
 
 #if PLATFORM_SUPPORTS_WAVE_BROADCAST && 1
-	#define CONFIG_BUTTERFLY_KERNEL 1
+	#define CONFIG_BUTTERFLY_KERNEL  1
+	#define CONFIG_BUTTERFLY_SAMPLES 4
 #else
-	#define CONFIG_BUTTERFLY_KERNEL 0
+	#define CONFIG_BUTTERFLY_KERNEL  0
+	#define CONFIG_BUTTERFLY_SAMPLES 9
 #endif
 
 #define CONFIG_SCENE_COLOR_OVERFLOW 1
@@ -149,6 +151,17 @@ static const uint kSquareIndexes3x3[9] = { 4, 0, 1, 2, 3, 8, 7, 6, 5 };
 // Indexes of the offsets to have plus + shape.
 static const uint kPlusIndexes3x3[5] = { 4, 1, 3, 7, 5 };
 
+// Cardinal pixel directions;
+static const int2 kOffsetC = int2( 0,  0);
+static const int2 kOffsetN = int2( 0, -1);
+static const int2 kOffsetS = int2( 0, +1);
+static const int2 kOffsetE = int2(+1,  0);
+static const int2 kOffsetW = int2(-1,  0);
+static const int2 kOffsetNE = int2(+1, -1);
+static const int2 kOffsetNW = int2(-1, -1);
+static const int2 kOffsetSE = int2(+1, +1);
+static const int2 kOffsetSW = int2(-1, +1);
+
 
 #if CONFIG_COMPILE_FP16
 	#define tsr_half  half
@@ -302,6 +315,33 @@ uint bCameraCut;
 #endif
 
 
+//------------------------------------------------------- UTILITY
+	
+uint sign_bit(float x)
+{
+    return asuint(x) >> 31u;
+}
+
+uint2 sign_bit(float2 x)
+{
+    return uint2(sign_bit(x.x), sign_bit(x.y));
+}
+
+#if PLATFORM_SUPPORTS_REAL_TYPES
+
+uint16_t sign_bit(half x)
+{
+    return asuint16(x) >> uint16_t(15);
+}
+
+uint16_t2 sign_bit(half2 x)
+{
+    return asuint16(x) >> uint16_t(15);
+}
+
+#endif
+
+
 //------------------------------------------------------- FUNCTIONS
 	
 CALL_SITE_DEBUGLOC 
@@ -1087,17 +1127,94 @@ tsr_short2 GetLaneOffsetSign()
 }
 #endif
 
-CALL_SITE_DEBUGLOC
-float ComputePixelVelocityBilateralWeight(float2 PixelVelocityDelta)
+tsr_short2 GetNeighborhoodPixelPos(const uint ArrayIndex, tsr_short2 InputPixelPos)
+{
+	tsr_short2 LaneOffsetSign = GetLaneOffsetSign();
+
+	tsr_short2 SampleInputPixelPos;
+	#if CONFIG_BUTTERFLY_KERNEL
+	{
+		tsr_short2 Offset = tsr_short2(Offsets2x2[ArrayIndex]) * LaneOffsetSign;
+		SampleInputPixelPos = ClampPixelOffset(InputPixelPos + Offset, InputPixelPosMin, InputPixelPosMax);
+	}
+	#else
+	{
+		tsr_short2 Offset = kOffsets3x3[ArrayIndex];
+					
+		SampleInputPixelPos = InputPixelPos + ClampPixelOffset(
+			InputPixelPos,
+			Offset, Offset,
+			InputPixelPosMin, InputPixelPosMax);
+	}
+	#endif
+
+	return SampleInputPixelPos;
+}
+
+template<typename Type>
+Type AccessNeighborhoodStaticOffset(Type Array[CONFIG_BUTTERFLY_SAMPLES], const int2 Offset)
 {
-	return saturate(1.0 - length2(PixelVelocityDelta));
+	const uint OffsetIndex = uint(dot(Offset + 1, int2(1, 3)));
+
+	#if CONFIG_BUTTERFLY_KERNEL
+		const uint ButterflyArray[] = {
+			0x3, 0x2, 0x2,
+			0x1, 0x0, 0x0,
+			0x1, 0x0, 0x0,
+		};
+		const uint IndexArray[] = {
+			0, 0, 1,
+			0, 0, 1,
+			2, 2, 3,
+		};
+
+		const FWaveBroadcastSettings BroadcastSetting = InitWaveXorButterfly(ButterflyArray[OffsetIndex]);
+		
+		return WaveBroadcast(BroadcastSetting, Array[IndexArray[OffsetIndex]]);
+	#else
+		return Array[OffsetIndex];
+	#endif
 }
 
-CALL_SITE_DEBUGLOC
-float ComputeScreenVelocityBilateralWeight(float2 ScreenVelocityDelta)
+template<typename Type>
+Type AccessNeighborhoodDynamicOffset(Type Array[CONFIG_BUTTERFLY_SAMPLES], const tsr_short2 Offset)
+{
+	const uint OffsetIndex = uint(dot(Offset + 1, tsr_short2(1, 3)));
+
+	#if CONFIG_BUTTERFLY_KERNEL
+		const FWaveBroadcastSettings BroadcastVertical   = InitWaveXorButterfly(0x2);
+		const FWaveBroadcastSettings BroadcastHorizontal = InitWaveXorButterfly(0x1);
+		const FWaveBroadcastSettings BroadcastDiagonal   = InitWaveXorButterfly(0x3);
+
+		Type Return = Array[0];
+		Return = select(all(Offset == tsr_short2(-1, -1)), WaveBroadcast(BroadcastDiagonal,   Array[0]), Return);
+		Return = select(all(Offset == tsr_short2( 0, -1)), WaveBroadcast(BroadcastVertical,   Array[0]), Return);
+		Return = select(all(Offset == tsr_short2(+1, -1)), WaveBroadcast(BroadcastVertical,   Array[1]), Return);
+		
+		Return = select(all(Offset == tsr_short2(-1,  0)), WaveBroadcast(BroadcastHorizontal, Array[0]), Return);
+		Return = select(all(Offset == tsr_short2(+1,  0)), Array[1], Return);
+		
+		Return = select(all(Offset == tsr_short2(-1, +1)), WaveBroadcast(BroadcastHorizontal, Array[2]), Return);
+		Return = select(all(Offset == tsr_short2( 0, +1)), Array[2], Return);
+		Return = select(all(Offset == tsr_short2(+1, +1)), Array[3], Return);
+		
+		return Return;
+	#else
+		return Array[OffsetIndex];
+	#endif
+}
+
+template<typename Type>
+Type AccessNeighborhoodCenter(Type Array[CONFIG_BUTTERFLY_SAMPLES])
+#if CONFIG_BUTTERFLY_KERNEL
+{
+	return Array[0];
+}
+#else
 {
-	return ComputePixelVelocityBilateralWeight(ScreenVelocityDelta * ScreenVelocityToInputPixelVelocity);
+	return Array[4];
 }
+#endif
 
 tsr_half2 ComputePrevWeightMultiplier(tsr_half2 BlendFactor)
 {
diff --git a/Engine/Shaders/Private/TemporalSuperResolution/TSRConvolutionNetworkPass.ush b/Engine/Shaders/Private/TemporalSuperResolution/TSRConvolutionNetworkPass.ush
--- a/Engine/Shaders/Private/TemporalSuperResolution/TSRConvolutionNetworkPass.ush
+++ b/Engine/Shaders/Private/TemporalSuperResolution/TSRConvolutionNetworkPass.ush
@@ -110,14 +110,35 @@
 
 #define tsr_tensor_float  TLaneVector2D<float, 1, LANE_STRIDE_X, LANE_STRIDE_Y>
 #define tsr_tensor_float2 TLaneVector2D<float, 2, LANE_STRIDE_X, LANE_STRIDE_Y>
+#define tsr_tensor_float3 TLaneVector2D<float, 3, LANE_STRIDE_X, LANE_STRIDE_Y>
+#define tsr_tensor_float4 TLaneVector2D<float, 4, LANE_STRIDE_X, LANE_STRIDE_Y>
+#define tsr_tensor_floatC TLaneVector2D<float, CONFIG_CHANNEL_COUNT, LANE_STRIDE_X, LANE_STRIDE_Y>
 #define tsr_tensor_half  TLaneVector2D<tsr_half, 1, LANE_STRIDE_X, LANE_STRIDE_Y>
 #define tsr_tensor_half2 TLaneVector2D<tsr_half, 2, LANE_STRIDE_X, LANE_STRIDE_Y>
 #define tsr_tensor_half3 TLaneVector2D<tsr_half, 3, LANE_STRIDE_X, LANE_STRIDE_Y>
 #define tsr_tensor_half4 TLaneVector2D<tsr_half, 4, LANE_STRIDE_X, LANE_STRIDE_Y>
 #define tsr_tensor_halfC TLaneVector2D<tsr_half, CONFIG_CHANNEL_COUNT, LANE_STRIDE_X, LANE_STRIDE_Y>
+
+#define tsr_tensor_uint    TLaneVector2D<uint, 1, LANE_STRIDE_X, LANE_STRIDE_Y>
+#define tsr_tensor_uint2   TLaneVector2D<uint, 2, LANE_STRIDE_X, LANE_STRIDE_Y>
+#define tsr_tensor_uint3   TLaneVector2D<uint, 3, LANE_STRIDE_X, LANE_STRIDE_Y>
+#define tsr_tensor_uint4   TLaneVector2D<uint, 4, LANE_STRIDE_X, LANE_STRIDE_Y>
+#define tsr_tensor_int     TLaneVector2D<int, 1, LANE_STRIDE_X, LANE_STRIDE_Y>
+#define tsr_tensor_int2    TLaneVector2D<int, 2, LANE_STRIDE_X, LANE_STRIDE_Y>
+#define tsr_tensor_int3    TLaneVector2D<int, 3, LANE_STRIDE_X, LANE_STRIDE_Y>
+#define tsr_tensor_int4    TLaneVector2D<int, 4, LANE_STRIDE_X, LANE_STRIDE_Y>
 #define tsr_tensor_ushort  TLaneVector2D<tsr_ushort, 1, LANE_STRIDE_X, LANE_STRIDE_Y>
-#define tsr_tensor_short2  TLaneVector2D<tsr_ushort, 2, LANE_STRIDE_X, LANE_STRIDE_Y>
+#define tsr_tensor_ushort2 TLaneVector2D<tsr_ushort, 2, LANE_STRIDE_X, LANE_STRIDE_Y>
+#define tsr_tensor_ushort3 TLaneVector2D<tsr_ushort, 3, LANE_STRIDE_X, LANE_STRIDE_Y>
+#define tsr_tensor_ushort4 TLaneVector2D<tsr_ushort, 4, LANE_STRIDE_X, LANE_STRIDE_Y>
+#define tsr_tensor_short   TLaneVector2D<tsr_short, 1, LANE_STRIDE_X, LANE_STRIDE_Y>
+#define tsr_tensor_short2  TLaneVector2D<tsr_short, 2, LANE_STRIDE_X, LANE_STRIDE_Y>
+#define tsr_tensor_short3  TLaneVector2D<tsr_short, 3, LANE_STRIDE_X, LANE_STRIDE_Y>
+#define tsr_tensor_short4  TLaneVector2D<tsr_short, 4, LANE_STRIDE_X, LANE_STRIDE_Y>
+
 #define tsr_tensor_bool    TLaneVector2D<bool, 1, LANE_STRIDE_X, LANE_STRIDE_Y>
+#define tsr_tensor_bool2   TLaneVector2D<bool, 2, LANE_STRIDE_X, LANE_STRIDE_Y>
+#define tsr_tensor_bool3   TLaneVector2D<bool, 3, LANE_STRIDE_X, LANE_STRIDE_Y>
 #define tsr_tensor_bool4   TLaneVector2D<bool, 4, LANE_STRIDE_X, LANE_STRIDE_Y>
 #define tsr_tensor_boolC   TLaneVector2D<bool, CONFIG_CHANNEL_COUNT, LANE_STRIDE_X, LANE_STRIDE_Y>
 
diff --git a/Engine/Shaders/Private/TemporalSuperResolution/TSRDecimateHistory.usf b/Engine/Shaders/Private/TemporalSuperResolution/TSRDecimateHistory.usf
--- a/Engine/Shaders/Private/TemporalSuperResolution/TSRDecimateHistory.usf
+++ b/Engine/Shaders/Private/TemporalSuperResolution/TSRDecimateHistory.usf
@@ -1,6 +1,7 @@
 // Copyright Epic Games, Inc. All Rights Reserved.
 
 #include "TSRCommon.ush"
+#include "TSRReprojectionField.ush"
 
 
 //------------------------------------------------------- CONFIG
@@ -12,19 +13,17 @@
 #define CONFIG_MANUAL_LDS_SPILL 1
 
 #define CONFIG_VELOCITY_SAMPLES 9
-#if CONFIG_VELOCITY_SAMPLES != 1
-	#define CONFIG_VELOCITY_SAMPLE_PAIRS ((CONFIG_VELOCITY_SAMPLES - 1) / 2)
+#define CONFIG_VELOCITY_SAMPLE_PAIRS ((CONFIG_VELOCITY_SAMPLES - 1) / 2)
 
-	#if CONFIG_BUTTERFLY_KERNEL
-		#if CONFIG_VELOCITY_SAMPLES == 9
-			#define CONFIG_VELOCITY_PRE_FETCH_SAMPLES 4
-		#else
-			#define CONFIG_VELOCITY_PRE_FETCH_SAMPLES 3
-		#endif
-	#else // !CONFIG_BUTTERFLY_KERNEL
-		#define CONFIG_VELOCITY_PRE_FETCH_SAMPLES (CONFIG_VELOCITY_SAMPLES)
+#if CONFIG_BUTTERFLY_KERNEL
+	#if CONFIG_VELOCITY_SAMPLES == 9
+		#define CONFIG_VELOCITY_PRE_FETCH_SAMPLES 4
+	#else
+		#define CONFIG_VELOCITY_PRE_FETCH_SAMPLES 3
 	#endif
-#endif // CONFIG_VELOCITY_SAMPLES != 1
+#else // !CONFIG_BUTTERFLY_KERNEL
+	#define CONFIG_VELOCITY_PRE_FETCH_SAMPLES (CONFIG_VELOCITY_SAMPLES)
+#endif
 
 /** Whether to compile the moire reprojection. */
 #define CONFIG_MOIRE_REPROJECTION (DIM_MOIRE_REPROJECTION)
@@ -37,7 +36,7 @@
 
 float4x4 RotationalClipToPrevClip;
 
-Texture2D<float2> DilatedVelocityTexture;
+Texture2D<uint> ReprojectionVectorTexture;
 #if CONFIG_RESURRECTION_REPROJECTION
 	Texture2D<float2> ClosestDepthTexture;
 #else
@@ -47,7 +46,6 @@ Texture2DArray<uint> PrevAtomicTextureArray;
 
 Texture2DArray<tsr_half4> PrevHistoryGuide;
 Texture2DArray<tsr_half4> PrevHistoryMoire;
-Texture2D<float2> PrevVelocityTexture;
 float2 PrevGuideInfo_Extent;
 float2 PrevGuideInfo_ExtentInverse;
 uint2  PrevGuideInfo_ViewportMin;
@@ -70,13 +68,13 @@ float4x4 ClipToResurrectionClip;
 RWTexture2DArray<tsr_half4> ReprojectedHistoryGuideOutput;
 RWTexture2DArray<tsr_half4> ReprojectedHistoryMoireOutput;
 
-RWTexture2D<float2> HoleFilledVelocityOutput;
+RWTexture2DArray<uint> ReprojectionFieldOutput;
 RWTexture2D<tsr_half2> DecimateMaskOutput;
 
 
 //------------------------------------------------------- DEBUG
 
-#define DEBUG_ARRAY_SIZE 2
+#define DEBUG_ARRAY_SIZE 8
 
 static float4 Debug[DEBUG_ARRAY_SIZE];
 
@@ -84,17 +82,16 @@ static float4 Debug[DEBUG_ARRAY_SIZE];
 //------------------------------------------------------- LDS
 
 #if CONFIG_MANUAL_LDS_SPILL
-	groupshared float2 SharedArray0[TILE_SIZE * TILE_SIZE];
+	groupshared uint   SharedArray0[TILE_SIZE * TILE_SIZE];
 	groupshared float  SharedArray1[TILE_SIZE * TILE_SIZE];
 #endif
 
 
 //------------------------------------------------------- FUNCTIONS
 
-#if CONFIG_VELOCITY_SAMPLES != 1
-void LoadVelocityNeighborhood(
+void LoadReprojectionNeighborhood(
 	tsr_short2 InputPixelPos,
-	out float2 EncodedVelocityNeighborhood[CONFIG_VELOCITY_PRE_FETCH_SAMPLES])
+	out uint EncodedReprojectionVectorNeighborhood[CONFIG_VELOCITY_PRE_FETCH_SAMPLES])
 {
 	tsr_short2 LaneOffsetSign = GetLaneOffsetSign();
 
@@ -105,7 +102,7 @@ void LoadVelocityNeighborhood(
 		#if CONFIG_BUTTERFLY_KERNEL
 		{
 			tsr_short2 Offset = tsr_short2(Offsets2x2[i]) * LaneOffsetSign;
-				
+
 			SampleInputPixelPos = InputPixelPos + Offset;
 			SampleInputPixelPos = ClampPixelOffset(SampleInputPixelPos, InputPixelPosMin, InputPixelPosMax);
 		}
@@ -117,7 +114,7 @@ void LoadVelocityNeighborhood(
 			#elif CONFIG_VELOCITY_SAMPLES == 5
 				Offset = tsr_short2(kOffsets3x3[kPlusIndexes3x3[i]]);
 			#endif
-					
+
 			SampleInputPixelPos = InputPixelPos + ClampPixelOffset(
 				InputPixelPos,
 				Offset, Offset,
@@ -125,26 +122,27 @@ void LoadVelocityNeighborhood(
 		}
 		#endif
 
-		EncodedVelocityNeighborhood[i] = DilatedVelocityTexture[SampleInputPixelPos];
+		EncodedReprojectionVectorNeighborhood[i] = ReprojectionVectorTexture[SampleInputPixelPos];
 	}
-} // LoadVelocityNeighborhood()
+} // LoadReprojectionNeighborhood()
 
-void ProcessVelocityNeighborhood(
-	float2 EncodedVelocityNeighborhood[CONFIG_VELOCITY_PRE_FETCH_SAMPLES],
-	out float2 EncodedScreenVelocity,
-	out tsr_half VelocityEdge)
+void ComputePixelVelocityNeighborhood(
+	uint EncodedReprojectionVectorNeighborhood[CONFIG_VELOCITY_PRE_FETCH_SAMPLES],
+	out float2 PixelVelocityNeighborhood[CONFIG_VELOCITY_PRE_FETCH_SAMPLES])
 {
-	// Preprocessed the prefetched samples
-	float2 PixelVelocityNeighborhood[CONFIG_VELOCITY_PRE_FETCH_SAMPLES];
+	UNROLL_N(CONFIG_VELOCITY_PRE_FETCH_SAMPLES)
+	for (uint i = 0; i < CONFIG_VELOCITY_PRE_FETCH_SAMPLES; i++)
 	{
-		UNROLL_N(CONFIG_VELOCITY_PRE_FETCH_SAMPLES)
-		for (uint i = 0; i < CONFIG_VELOCITY_PRE_FETCH_SAMPLES; i++)
-		{
-			PixelVelocityNeighborhood[i] = DecodeVelocityFromTexture(float4(EncodedVelocityNeighborhood[i], 0.0, 0.0)).xy * ScreenVelocityToInputPixelVelocity;
-		}
+		PixelVelocityNeighborhood[i] = DecodeReprojectionVector(EncodedReprojectionVectorNeighborhood[i]) * ScreenVelocityToInputPixelVelocity;
 	}
-	
-	tsr_half RejectHighFrequency = tsr_half(0.0);
+} // ComputePixelVelocityNeighborhood()
+
+/** Compute the reprojection edge from the pixel velocities. */
+void ComputeReprojectionEdge(
+	float2 PixelVelocityNeighborhood[CONFIG_VELOCITY_PRE_FETCH_SAMPLES],
+	out tsr_half ReprojectionEdge)
+{
+	ReprojectionEdge = tsr_half(0.0);
 
 	// Sample
 	float2 ScreenPixelVelocity = PixelVelocityNeighborhood[0];
@@ -212,7 +210,7 @@ void ProcessVelocityNeighborhood(
 		tsr_half2 NeighborDeltaPixelVelocity0 = tsr_half2(ScreenPixelVelocity - NeighborPixelVelocity0);
 		tsr_half2 NeighborDeltaPixelVelocity1 = tsr_half2(ScreenPixelVelocity - NeighborPixelVelocity1);
 			
-		// Reject high frequency on velocity edge discontinuity.
+		// Reject high frequency on reprojection edge discontinuity.
 		{
 			tsr_half IsOverlapingEdge0 = tsr_half(1.0) - dot(saturate(abs(NeighborDeltaPixelVelocity0 * Offset0)), tsr_half(1.0).xx); // * saturate(ScreenPixelVelocityLength - NeighborPixelVelocityLength0);
 			tsr_half IsOverlapingEdge1 = tsr_half(1.0) - dot(saturate(abs(NeighborDeltaPixelVelocity1 * Offset1)), tsr_half(1.0).xx); // * saturate(ScreenPixelVelocityLength - NeighborPixelVelocityLength1);
@@ -220,19 +218,15 @@ void ProcessVelocityNeighborhood(
 			tsr_half IsOverlapingEdge = min(IsOverlapingEdge0, IsOverlapingEdge1);
 
 			if (i == 0)
-				RejectHighFrequency = IsOverlapingEdge;
+				ReprojectionEdge = IsOverlapingEdge;
 			else
-				RejectHighFrequency = min(RejectHighFrequency, IsOverlapingEdge);
+				ReprojectionEdge = min(ReprojectionEdge, IsOverlapingEdge);
 		}
-	}
+	} // for (uint i = 0; i < CONFIG_VELOCITY_SAMPLE_PAIRS; i++)
 
 	// Remove noise from high frequency rejections.
-	VelocityEdge = saturate(RejectHighFrequency * tsr_half(1.1));
-
-	EncodedScreenVelocity = EncodedVelocityNeighborhood[0];
-} // ProcessVelocityNeighborhood()
-
-#endif // CONFIG_VELOCITY_SAMPLES != 1
+	ReprojectionEdge = saturate(ReprojectionEdge * tsr_half(1.1));
+} // ComputeReprojectionEdge()
 
 void LoadPrevAtomicTexturesSamples(
 	float2 PrevScreenPos,
@@ -482,7 +476,7 @@ tsr_half4 CorrectMoireExposure(tsr_half4 ReprojectedHistoryMoireColor, const flo
 void ReprojectAllPrevTextures(
 	uint2 GroupId,
 	uint GroupThreadIndex,
-	float2 EncodedScreenVelocity,
+	uint EncodedReprojectionVector,
 	float DeviceZ,
 	out uint HistoryUseCountSamples[4],
 	out uint HistoryClosestDeviceZSamples0[4],
@@ -506,7 +500,7 @@ void ReprojectAllPrevTextures(
 			
 	{
 		float2 ScreenPos = ApplyScreenTransform(float2(InputPixelPos), InputPixelPosToScreenPos);
-		float2 PrevScreenPos = ScreenPos - DecodeVelocityFromTexture(float4(EncodedScreenVelocity, 0.0, 0.0)).xy;
+		float2 PrevScreenPos = ScreenPos - DecodeReprojectionVector(EncodedReprojectionVector);
 		LoadPrevAtomicTexturesSamples(PrevScreenPos,
 			/* out */ HistoryUseCountSamples,
 			/* out */ HistoryClosestDeviceZSamples0,
@@ -516,7 +510,7 @@ void ReprojectAllPrevTextures(
 	{
 		float2 OutputScreenPos = ApplyScreenTransform(float2(InputPixelPos), InputPixelPosToReprojectScreenPos);
 	
-		float2 PrevOutputScreenPos = OutputScreenPos - DecodeVelocityFromTexture(float4(EncodedScreenVelocity, 0.0, 0.0)).xy;
+		float2 PrevOutputScreenPos = OutputScreenPos - DecodeReprojectionVector(EncodedReprojectionVector);
 		bIsOffScreen = (
 			(PrevOutputScreenPos.x < -1.0) ||
 			(PrevOutputScreenPos.y < -1.0) ||
@@ -597,24 +591,21 @@ void MainCS(
 	}
 
 	// Sample current frame informations.
-	float2 EncodedScreenVelocity;
+	uint DilatedEncodedReprojectionVector;
 	float DeviceZ;
 	float PrevDeviceZ;
-	tsr_half VelocityEdge = tsr_half(1.0);
 	bool bHasPixelAnimation;
 	ISOLATE
 	{
-		tsr_ushort2 InputPixelPos = (
-			Map8x8Tile2x2Lane(GroupThreadIndex) +
-			(tsr_ushort2(InputInfo_ViewportMin) + tsr_ushort2(GroupId) * tsr_ushort2(TILE_SIZE, TILE_SIZE)));
-			
-		#if CONFIG_VELOCITY_SAMPLES != 1
-		float2 EncodedVelocityNeighborhood[CONFIG_VELOCITY_PRE_FETCH_SAMPLES];
-		#endif
+		//uint EncodedReprojectionVectorNeighborhood[CONFIG_VELOCITY_PRE_FETCH_SAMPLES];
 
 		// Issue texture fetechs
 		ISOLATE
 		{
+			tsr_ushort2 InputPixelPos = (
+				Map8x8Tile2x2Lane(GroupThreadIndex) +
+				(tsr_ushort2(InputInfo_ViewportMin) + tsr_ushort2(GroupId) * tsr_ushort2(TILE_SIZE, TILE_SIZE)));
+			
 			#if CONFIG_RESURRECTION_REPROJECTION
 				float2 DeviceZAndPrevDeviceZ = ClosestDepthTexture[InputPixelPos];
 				DeviceZ = DeviceZAndPrevDeviceZ.g;
@@ -626,37 +617,30 @@ void MainCS(
 			PrevDeviceZ = abs(DeviceZAndPrevDeviceZ.r);
 			bHasPixelAnimation = DeviceZAndPrevDeviceZ.r < 0.0;
 
-			#if CONFIG_VELOCITY_SAMPLES == 1
-				EncodedScreenVelocity = DilatedVelocityTexture[InputPixelPos];
-			#else // CONFIG_VELOCITY_SAMPLES != 1
-				LoadVelocityNeighborhood(InputPixelPos, /* out */ EncodedVelocityNeighborhood);
-				EncodedScreenVelocity = EncodedVelocityNeighborhood[0];
-			#endif // CONFIG_VELOCITY_SAMPLES != 1
+			//LoadReprojectionNeighborhood(InputPixelPos, /* out */ EncodedReprojectionVectorNeighborhood);
+			//DilatedEncodedReprojectionVector = EncodedReprojectionVectorNeighborhood[0];
+			DilatedEncodedReprojectionVector = ReprojectionFieldOutput[tsr_ushort3(InputPixelPos, 0)];
 		}
 
 		// LDS spill to save VGPR on the history texture fetches.
 		#if CONFIG_MANUAL_LDS_SPILL
 		ISOLATE
 		{
-			#if CONFIG_VELOCITY_SAMPLES == 1
-				SharedArray0[GroupThreadIndex] = EncodedScreenVelocity;
-			#else
-				SharedArray0[GroupThreadIndex] = EncodedVelocityNeighborhood[0];
-			#endif
+			SharedArray0[GroupThreadIndex] = DilatedEncodedReprojectionVector;
 			SharedArray1[GroupThreadIndex] = PrevDeviceZ;
 		}
 		#endif
 		
 		// Detect whether is on velocity edge to clamp the high frequencies
-		#if CONFIG_VELOCITY_SAMPLES != 1
+		#if 0
 		ISOLATE
 		{
-			ProcessVelocityNeighborhood(
-				EncodedVelocityNeighborhood,
-				/* out */ EncodedScreenVelocity,
-				/* out */ VelocityEdge);
+			float2 PixelVelocityNeighborhood[CONFIG_VELOCITY_PRE_FETCH_SAMPLES];
+			ComputePixelVelocityNeighborhood(EncodedReprojectionVectorNeighborhood, /* out */ PixelVelocityNeighborhood);
+
+			ComputeReprojectionEdge(PixelVelocityNeighborhood, /* out */ VelocityEdge);
 		}
-		#endif // CONFIG_VELOCITY_SAMPLES != 1
+		#endif
 	}
 	
 	// Sample the previous frame information
@@ -682,7 +666,7 @@ void MainCS(
 	{
 		ReprojectAllPrevTextures(
 			GroupId, GroupThreadIndex,
-			EncodedScreenVelocity, DeviceZ,
+			DilatedEncodedReprojectionVector, DeviceZ,
 			/* out */ HistoryUseCountSamples,
 			/* out */ HistoryClosestDeviceZSamples0,
 			/* out */ HistoryClosestDeviceZSamples1,
@@ -704,6 +688,7 @@ void MainCS(
 	tsr_halfC ReprojectedHistoryGuideColor;
 	tsr_half  ReprojectedHistoryGuideUncertainty;
 	tsr_half ParallaxRejectionMask;
+	bool bOverwriteReprojectionField = false;
 	//ISOLATE
 	{
 		#if CONFIG_SCENE_COLOR_ALPHA
@@ -736,7 +721,7 @@ void MainCS(
 		// LDS despill
 		#if CONFIG_MANUAL_LDS_SPILL
 		{
-			EncodedScreenVelocity = SharedArray0[GroupThreadIndex];
+			DilatedEncodedReprojectionVector = SharedArray0[GroupThreadIndex];
 			PrevDeviceZ = SharedArray1[GroupThreadIndex];
 		}
 		#endif
@@ -746,7 +731,7 @@ void MainCS(
 			(tsr_ushort2(InputInfo_ViewportMin) + tsr_ushort2(GroupId) * tsr_ushort2(TILE_SIZE, TILE_SIZE)));
 
 		float2 ScreenPos = ApplyScreenTransform(float2(InputPixelPos), InputPixelPosToScreenPos);
-		float2 ScreenVelocity = DecodeVelocityFromTexture(float4(EncodedScreenVelocity, 0.0, 0.0)).xy;
+		float2 ScreenVelocity = DecodeReprojectionVector(DilatedEncodedReprojectionVector);
 		float2 PrevScreenPos = ScreenPos - ScreenVelocity;
 		
 		// Compute the parralax rejection mask
@@ -766,6 +751,13 @@ void MainCS(
 		ParallaxRejectionMask = select(bIsOffScreen, tsr_half(0.0), ParallaxRejectionMask);
 
 		bool bIsOffscreenOrParallaxRejected = ParallaxRejectionMask < tsr_half(PARALLAX_REJECTION_MASK_THRESHOLD);
+
+		// LDS despill
+		#if CONFIG_MANUAL_LDS_SPILL
+		{
+			DilatedEncodedReprojectionVector = SharedArray0[GroupThreadIndex];
+		}
+		#endif
 	
 		// Hole fill the velocity.
 		{
@@ -781,13 +773,11 @@ void MainCS(
 			FLATTEN
 			if (tsr_half(HoleFillingOutOfRange) > 0.0 && bIsOffscreenOrParallaxRejected)
 			{
-				float2 ScreenPixelVelocity = ScreenVelocity * ScreenVelocityToInputPixelVelocity;
-				float VelocityDelta = length(ScreenPixelVelocity - HoleFillingPixelVelocity);
-
-				float2 HoleFillingVelocity = HoleFillingPixelVelocity * InputPixelVelocityToScreenVelocity;
+				// Disable the reprojection's jacobian since this is to extrude disoccluded area we don't know how they look anyway.
+				float2 HoleFillingReprojectionVector = HoleFillingPixelVelocity * abs(InputPixelVelocityToScreenVelocity);
 
-				EncodedScreenVelocity = EncodeVelocityToTexture(float3(HoleFillingVelocity, 0.0)).xy;
-				//HoleFilledVelocityMask = tsr_half(HoleFillingOutOfRange) * tsr_half(saturate(MaxHollFillPixelVelocity - VelocityDelta));
+				DilatedEncodedReprojectionVector = EncodeReprojectionVector(HoleFillingReprojectionVector);
+				bOverwriteReprojectionField = true;
 			}
 		}
 	}
@@ -840,9 +830,18 @@ void MainCS(
 			#endif
 		}
 		#endif
+		
+		DecimateMaskOutput[OutputPixelPos] = tsr_half2(tsr_half(DecimateBitMask) * rcp(tsr_half(255.0)), 0.0);
+
+		// Selectively overwrite the reprojection field
+		{
+			const uint EncodedReprojectionJacobian = EncodeReprojectionJacobian(tsr_half2x2(tsr_half(0.0).xxxx));
 
-		HoleFilledVelocityOutput[OutputPixelPos] = EncodedScreenVelocity;
-		DecimateMaskOutput[InputPixelPos] = tsr_half2(tsr_half(DecimateBitMask) * rcp(tsr_half(255.0)), VelocityEdge);
+			tsr_short2 ReprojectionFieldOutputPixelPos = OutputPixelPos;
+			ReprojectionFieldOutputPixelPos.x = select(bOverwriteReprojectionField, ReprojectionFieldOutputPixelPos.x, -tsr_short(1));
+			ReprojectionFieldOutput[tsr_short3(ReprojectionFieldOutputPixelPos, 0)] = DilatedEncodedReprojectionVector;
+			ReprojectionFieldOutput[tsr_short3(ReprojectionFieldOutputPixelPos, 2)] = EncodedReprojectionJacobian;
+		}
 
 		#if DEBUG_OUTPUT
 		for (uint DebugId = 0; DebugId < DEBUG_ARRAY_SIZE; DebugId++)
diff --git a/Engine/Shaders/Private/TemporalSuperResolution/TSRDilateVelocity.usf b/Engine/Shaders/Private/TemporalSuperResolution/TSRDilateVelocity.usf
--- a/Engine/Shaders/Private/TemporalSuperResolution/TSRDilateVelocity.usf
+++ b/Engine/Shaders/Private/TemporalSuperResolution/TSRDilateVelocity.usf
@@ -1,6 +1,8 @@
 // Copyright Epic Games, Inc. All Rights Reserved.
 
 #include "TSRCommon.ush"
+#include "TSRSpatialAntiAliasing.ush"
+#include "TSRReprojectionField.ush"
 
 #if DIM_MOTION_BLUR_DIRECTIONS > 0
 	#define CONFIG_MAX_RANGE_SIZE (DIM_MOTION_BLUR_DIRECTIONS)
@@ -16,20 +18,16 @@
 	#define TILE_SIZE 8
 #endif
 
-#define CONFIG_AA_CROSS 1
-
-#define CONFIG_AA_CROSS_OPTIMISATION (PLATFORM_SUPPORTS_WAVE_BROADCAST)
-
+/** Whether to use the velocity from SceneVelocityTexture. */
 #define CONFIG_USE_SCENE_VELOCITY 1
 
+/** Whether the dilation should be only on depth edges, improving texture reprojection accuracy on flat surfaces. */
 #define CONFIG_DEPTH_EDGE_DETECTION 1
 
-#define CONFIG_VELOCITY_EXTRAPOLATION 0
-
 
 //------------------------------------------------------- CONSTANTS
 
-static const int2 kOffsets[4] = {
+static const int2 kPairOffsets[4] = {
 	int2( 1,  0),
 	int2( 1,  1),
 	int2( 0,  1),
@@ -42,22 +40,30 @@ static const int2 kOffsets[4] = {
 float4x4 RotationalClipToPrevClip;
 float2 PrevOutputBufferUVMin;
 float2 PrevOutputBufferUVMax;
-float VelocityExtrapolationMultiplier;
 float InvFlickeringMaxParralaxVelocity;
+float ReprojectionFieldAntiAliasVelocityThreshold;
+uint bReprojectionField;
 uint bOutputIsMovingTexture;
 
 Texture2D<float> SceneDepthTexture;
-Texture2D SceneVelocityTexture;
+Texture2D<float4> SceneVelocityTexture;
 
-RWTexture2D<float2> DilatedVelocityOutput;
 RWTexture2D<float2> ClosestDepthOutput;
 globallycoherent RWTexture2DArray<uint> PrevAtomicOutput;
+RWTexture2DArray<uint> ReprojectionFieldOutput;
 RWTexture2DArray<uint> R8Output;
 
 RWTexture2D<float3>	VelocityFlattenOutput;
 RWTexture2DArray<float4> VelocityTileArrayOutput;
 
 
+//------------------------------------------------------- DEBUG
+
+#define DEBUG_ARRAY_SIZE 8
+
+static float4 Debug[DEBUG_ARRAY_SIZE];
+
+
 //------------------------------------------------------- UTILITY
 
 uint PackDepthAndOffset(const int2 Offset, float DeviceZ)
@@ -65,15 +71,432 @@ uint PackDepthAndOffset(const int2 Offset, float DeviceZ)
 	return ((asuint(DeviceZ) << 2) & ~0xF) | (uint(1 + Offset.x) << 0) | (uint(1 + Offset.y) << 2);
 }
 
-void UnpackDepthAndOffset(uint Packed, out int2 Offset, out float DeviceZ)
+void UnpackDepthAndOffset(uint Packed, out tsr_short2 Offset, out float DeviceZ)
 {
 	DeviceZ = asfloat((Packed & ~0xF) >> 2);
-	Offset.x = int((Packed >> 0) & 0x3) - 1;
-	Offset.y = int((Packed >> 2) & 0x3) - 1;
+	Offset.x = tsr_short((Packed >> 0u) & 0x3u) - tsr_short(1);
+	Offset.y = tsr_short((Packed >> 2u) & 0x3u) - tsr_short(1);
+}
+
+
+//------------------------------------------------------- FUNCTIONS
+
+/* Loads all the 3x3's depth. */
+void FetchDepth3x3(
+	tsr_short2 InputPixelPos,
+	out float DevizeZNeighborhood[CONFIG_BUTTERFLY_SAMPLES])
+{
+	UNROLL_N(CONFIG_BUTTERFLY_SAMPLES)
+	for (uint i = 0; i < CONFIG_BUTTERFLY_SAMPLES; i++)
+	{
+		tsr_short2 SampleInputPixelPos = GetNeighborhoodPixelPos(i, InputPixelPos);
+		
+		DevizeZNeighborhood[i] = SceneDepthTexture[SampleInputPixelPos];
+	}
+} // FetchDepthVelocity3x3()
+
+/* Loads all the 3x3's depth and velocities. */
+void FetchDepthVelocity3x3(
+	tsr_short2 InputPixelPos,
+	out float DevizeZNeighborhood[CONFIG_BUTTERFLY_SAMPLES],
+	out float4 VelocityNeighborhood[CONFIG_BUTTERFLY_SAMPLES])
+{
+	UNROLL_N(CONFIG_BUTTERFLY_SAMPLES)
+	for (uint i = 0; i < CONFIG_BUTTERFLY_SAMPLES; i++)
+	{
+		tsr_short2 SampleInputPixelPos = GetNeighborhoodPixelPos(i, InputPixelPos);
+		
+		DevizeZNeighborhood[i] = SceneDepthTexture[SampleInputPixelPos];
+		VelocityNeighborhood[i] = SceneVelocityTexture[SampleInputPixelPos];
+	}
+} // FetchDepthVelocity3x3()
+
+/** Convert the fetches velocities to pixel velocities. */
+void ComputePixelVelocityNeighborhood(
+	tsr_short2 InputPixelPos,
+	float DevizeZNeighborhood[CONFIG_BUTTERFLY_SAMPLES],
+	float4 VelocityNeighborhood[CONFIG_BUTTERFLY_SAMPLES],
+	out float2 PixelVelocityNeighborhood[CONFIG_BUTTERFLY_SAMPLES],
+	out float  DepthVelocityNeighborhood[CONFIG_BUTTERFLY_SAMPLES],
+	out bool   bIsDrawingVelocityNeighborhood[CONFIG_BUTTERFLY_SAMPLES])
+{
+	bool bHasAnyStaticPixels = false;
+	bool bHasAnyDynamicPixels = false;
+	{
+		UNROLL_N(CONFIG_BUTTERFLY_SAMPLES)
+		for (uint i = 0; i < CONFIG_BUTTERFLY_SAMPLES; i++)
+		{
+			#if CONFIG_USE_SCENE_VELOCITY
+				bool bIsDrawingVelocity = VelocityNeighborhood[i].x > 0.0;
+			#else
+				const bool bIsDrawingVelocity = false;
+			#endif
+			bIsDrawingVelocityNeighborhood[i] = bIsDrawingVelocity;
+			bHasAnyStaticPixels  |= !bIsDrawingVelocity;
+			bHasAnyDynamicPixels |= bIsDrawingVelocity;
+		}
+
+		#if COMPILER_SUPPORTS_WAVE_VOTE
+			bHasAnyStaticPixels  = WaveActiveAnyTrue(bHasAnyStaticPixels);
+			bHasAnyDynamicPixels = WaveActiveAnyTrue(bHasAnyDynamicPixels);
+		#else
+			bHasAnyStaticPixels  = true;
+			bHasAnyDynamicPixels = true;
+		#endif
+	}
+	
+	#if CONFIG_USE_SCENE_VELOCITY
+	BRANCH
+	if (bHasAnyStaticPixels && bHasAnyDynamicPixels)
+	{
+		UNROLL_N(CONFIG_BUTTERFLY_SAMPLES)
+		for (uint i = 0; i < CONFIG_BUTTERFLY_SAMPLES; i++)
+		{
+			tsr_short2 SampleInputPixelPos = GetNeighborhoodPixelPos(i, InputPixelPos);
+			float2 ScreenPos = ApplyScreenTransform(float2(int2(SampleInputPixelPos)), InputPixelPosToScreenPos);
+
+			float3 ScreenVelocity = ComputeStaticVelocity(ScreenPos, DevizeZNeighborhood[i]);
+			FLATTEN
+			if (VelocityNeighborhood[i].x > 0.0)
+			{
+				ScreenVelocity = DecodeVelocityFromTexture(VelocityNeighborhood[i]);
+			}
+
+			PixelVelocityNeighborhood[i] = ScreenVelocity.xy * ScreenVelocityToInputPixelVelocity;
+			DepthVelocityNeighborhood[i] = ScreenVelocity.z;
+		}
+	}
+	else if (bHasAnyDynamicPixels) // && !bHasAnyStaticPixels)
+	{
+		UNROLL_N(CONFIG_BUTTERFLY_SAMPLES)
+		for (uint i = 0; i < CONFIG_BUTTERFLY_SAMPLES; i++)
+		{
+			float3 ScreenVelocity = DecodeVelocityFromTexture(VelocityNeighborhood[i]);
+
+			PixelVelocityNeighborhood[i] = ScreenVelocity.xy * ScreenVelocityToInputPixelVelocity;
+			DepthVelocityNeighborhood[i] = ScreenVelocity.z;
+		}
+	}
+	else // if (bHasAnyStaticPixels && !bHasAnyDynamicPixels
+	#endif // !CONFIG_USE_SCENE_VELOCITY
+	{
+		UNROLL_N(CONFIG_BUTTERFLY_SAMPLES)
+		for (uint i = 0; i < CONFIG_BUTTERFLY_SAMPLES; i++)
+		{
+			tsr_short2 SampleInputPixelPos = GetNeighborhoodPixelPos(i, InputPixelPos);
+			float2 ScreenPos = ApplyScreenTransform(float2(int2(SampleInputPixelPos)), InputPixelPosToScreenPos);
+
+			float3 ScreenVelocity = ComputeStaticVelocity(ScreenPos, DevizeZNeighborhood[i]);
+
+			PixelVelocityNeighborhood[i] = ScreenVelocity.xy * ScreenVelocityToInputPixelVelocity;
+			DepthVelocityNeighborhood[i] = ScreenVelocity.z;
+		}
+	}
+} // ComputePixelVelocityNeighborhood()
+
+/** Find the closest DevizeZ and its corresponding offset */
+void FindClosestDepthOffset(
+	float DevizeZNeighborhood[CONFIG_BUTTERFLY_SAMPLES],
+	out float ClosestDeviceZ,
+	out tsr_short2 ReprojectionOffset)
+{
+	float DeviceZ = AccessNeighborhoodCenter(DevizeZNeighborhood);
+	uint PackedDepthOffset = PackDepthAndOffset(/* Offset = */ int2(0, 0), DeviceZ);
+	
+	#if CONFIG_DEPTH_EDGE_DETECTION
+	float DeviceZError;
+	{
+		float Depth = ConvertFromDeviceZ(DeviceZ);
+		float ErrorCorrectedDeviceZ = ConvertToDeviceZ(GetDepthPixelRadiusForProjectionType(Depth) * 2.0f + Depth);
+		DeviceZError = abs(ErrorCorrectedDeviceZ - DeviceZ);
+	}
+	#endif
+	
+	#if 0
+		UNROLL_N(4)
+		for (uint i = 0; i < 4; i++)
+	#elif 1
+		UNROLL_N(2)
+		for (uint i = 0; i < 4; i+= 2)
+	#else
+		UNROLL_N(2)
+		for (uint i = 1; i < 4; i+= 2)
+	#endif
+	{
+		const int2 PairOffset = kPairOffsets[i];
+
+		float SampleDeviceZ0 = AccessNeighborhoodStaticOffset(DevizeZNeighborhood, +PairOffset);
+		float SampleDeviceZ1 = AccessNeighborhoodStaticOffset(DevizeZNeighborhood, -PairOffset);
+
+		#if CONFIG_DEPTH_EDGE_DETECTION
+		{
+			float DepthDiff = abs(SampleDeviceZ1 - SampleDeviceZ0);
+			float DepthVariation = (SampleDeviceZ0 + SampleDeviceZ1) * 0.5 - DeviceZ;
+
+			FLATTEN
+			if (DepthVariation > max(DepthDiff * 0.25, DeviceZError))
+			{
+				PackedDepthOffset = max(PackedDepthOffset, PackDepthAndOffset(+PairOffset, SampleDeviceZ0));
+				PackedDepthOffset = max(PackedDepthOffset, PackDepthAndOffset(-PairOffset, SampleDeviceZ1));
+			}
+		}
+		#else
+		{
+			PackedDepthOffset = max(PackedDepthOffset, PackDepthAndOffset(+PairOffset, SampleDeviceZ0));
+			PackedDepthOffset = max(PackedDepthOffset, PackDepthAndOffset(-PairOffset, SampleDeviceZ1));
+		}
+		#endif
+	}
+
+	UnpackDepthAndOffset(PackedDepthOffset, /* out */ ReprojectionOffset, /* out */ ClosestDeviceZ);
+	
+	#if CONFIG_BUTTERFLY_KERNEL
+		tsr_short2 LaneOffsetSign = GetLaneOffsetSign();
+		ReprojectionOffset *= LaneOffsetSign;
+	#endif
+} // FindClosestDepthOffset()
+
+/** Compute the velocity differential along a screen space axis. */
+bool ComputeReprojectionDifferential(
+	float  DevizeZNeighborhood[CONFIG_BUTTERFLY_SAMPLES],
+	float2 PixelVelocityNeighborhood[CONFIG_BUTTERFLY_SAMPLES],
+	float DeviceZError,
+	const int2 PositiveDirection,
+	bool bCompute,
+	inout tsr_half2 ReprojectionDifferential)
+{
+	float DevizeZC = AccessNeighborhoodCenter(DevizeZNeighborhood);
+	float DevizeZP = AccessNeighborhoodStaticOffset(DevizeZNeighborhood, +PositiveDirection);
+	float DevizeZN = AccessNeighborhoodStaticOffset(DevizeZNeighborhood, -PositiveDirection);
+	
+	float2 PixelVelocityC = AccessNeighborhoodCenter(PixelVelocityNeighborhood);
+	float2 PixelVelocityP = AccessNeighborhoodStaticOffset(PixelVelocityNeighborhood, +PositiveDirection);
+	float2 PixelVelocityN = AccessNeighborhoodStaticOffset(PixelVelocityNeighborhood, -PositiveDirection);
+	
+	bool bSuccess = false;
+	BRANCH
+	if (bCompute)
+	{
+		const tsr_half WeightThreshold = tsr_half(2.0); 
+		const tsr_half WeightThresholdSquare = WeightThreshold * WeightThreshold; 
+	
+		// Compute the velocity differential.
+		tsr_half2 VelocityDeltaP = tsr_half2(PixelVelocityP - PixelVelocityC);
+		tsr_half2 VelocityDeltaN = tsr_half2(PixelVelocityC - PixelVelocityN);
+	
+		float DepthVariation = DevizeZC - (DevizeZP + DevizeZN) * 0.5;
+		float DepthMax = max3(DevizeZP, DevizeZN, DevizeZC);
+		float DepthMin = min3(DevizeZP, DevizeZN, DevizeZC);
+		float Delta = max(2.0 * DepthVariation - (DepthMax - DepthMin), 0.0);
+
+		float DepthDiffPN = abs(DevizeZP - DevizeZN);
+		float DepthDiffP  = abs(DevizeZP - DevizeZC);
+		float DepthDiffN  = abs(DevizeZN - DevizeZC);
+	
+		float MinDepthDerivative = clamp(min3(DepthDiffPN, DepthDiffP, DepthDiffN), tsr_half(0.125) * DeviceZError, tsr_half(8.0) * DeviceZError);
+
+		tsr_half WeightP = tsr_half(1.0);
+		tsr_half WeightN = tsr_half(1.0);
+
+		WeightP *= tsr_half(saturate(1.1 - (DepthDiffP) / MinDepthDerivative));
+		WeightN *= tsr_half(saturate(1.1 - (DepthDiffN) / MinDepthDerivative));
+	
+		// Bilateral based on velicities differentials
+		//WeightP *= saturate(WeightThresholdSquare - dot(VelocityDeltaP, VelocityDeltaP));
+		//WeightN *= saturate(WeightThresholdSquare - dot(VelocityDeltaN, VelocityDeltaN));
+
+		tsr_half TotalWeight = WeightP + WeightN;
+		bSuccess = TotalWeight > tsr_half(0.0);
+		tsr_half InvTotalWeight = select(bSuccess, rcp(TotalWeight), tsr_half(0.0));
+	
+		ReprojectionDifferential = VelocityDeltaP * (WeightP * InvTotalWeight) + VelocityDeltaN * (WeightN * InvTotalWeight);
+	}
+	return bSuccess;
 }
 
+/** Find the closest DevizeZ and its corresponding offset */
+void ComputeReprojectionJacobian(
+	float  DevizeZNeighborhood[CONFIG_BUTTERFLY_SAMPLES],
+	float2 PixelVelocityNeighborhood[CONFIG_BUTTERFLY_SAMPLES],
+	out tsr_half2x2 ReprojectionJacobian)
+{
+	float DeviceZ = AccessNeighborhoodCenter(DevizeZNeighborhood);
+	float DeviceZError;
+	{
+		float Depth = ConvertFromDeviceZ(DeviceZ);
+		float ErrorCorrectedDeviceZ = ConvertToDeviceZ(GetDepthPixelRadiusForProjectionType(Depth) * 2.0f + Depth);
+		DeviceZError = abs(ErrorCorrectedDeviceZ - DeviceZ);
+	}
+	
+	// Horizontal & vertical cardinals
+	bool bSuccess;
+	{
+		ReprojectionJacobian = tsr_half(0.0);
+
+		bool bSuccessX = ComputeReprojectionDifferential(
+			DevizeZNeighborhood,
+			PixelVelocityNeighborhood,
+			DeviceZError,
+			/* PositiveDirection = */ kOffsetE,
+			/* bCompute = */ true,
+			/* inout */ ReprojectionJacobian[0]);
+		
+		bool bSuccessY = ComputeReprojectionDifferential(
+			DevizeZNeighborhood,
+			PixelVelocityNeighborhood,
+			DeviceZError,
+			/* PositiveDirection = */ kOffsetS,
+			/* bCompute = */ true,
+			/* inout */ ReprojectionJacobian[1]);
+	
+		bSuccess = bSuccessX || bSuccessY;
+	}
+	
+	// Diagonal cardinals
+	{
+		tsr_half2 DifferentialSE;
+		ComputeReprojectionDifferential(
+			DevizeZNeighborhood,
+			PixelVelocityNeighborhood,
+			DeviceZError,
+			/* PositiveDirection = */ kOffsetSE,
+			/* bCompute = */ !bSuccess,
+			/* inout */ ReprojectionJacobian[0]);
+			
+		tsr_half2 DifferentialSW;
+		ComputeReprojectionDifferential(
+			DevizeZNeighborhood,
+			PixelVelocityNeighborhood,
+			DeviceZError,
+			/* PositiveDirection = */ kOffsetSW,
+			/* bCompute = */ !bSuccess,
+			/* inout */ ReprojectionJacobian[1]);
+			
+		if (!bSuccess)
+		{
+			tsr_half2 DifferentialSE = ReprojectionJacobian[0];
+			tsr_half2 DifferentialSW = ReprojectionJacobian[1];
+			ReprojectionJacobian[0] = DifferentialSE * tsr_half(+0.5) + DifferentialSW * tsr_half(-0.5);
+			ReprojectionJacobian[1] = DifferentialSE * tsr_half(+0.5) + DifferentialSW * tsr_half(+0.5);
+		}
+	}
+	
+	#if CONFIG_BUTTERFLY_KERNEL && 1
+		tsr_half2 LaneOffsetSign = tsr_half2(GetLaneOffsetSign());
+		ReprojectionJacobian[0] *= LaneOffsetSign.x;
+		ReprojectionJacobian[1] *= LaneOffsetSign.y;
+	#endif
+} // ComputeReprojectionJacobian()
+
+void FetchAndComputeScreenVelocity(
+	tsr_short2 InputPixelPos,
+	float ClosestDeviceZ,
+	tsr_short2 ReprojectionOffset,
+	out float3 ScreenVelocity,
+	out bool bHasPixelAnimation)
+{
+	float4 EncodedVelocity = SceneVelocityTexture[InputPixelPos + ReprojectionOffset];
+
+	#if VELOCITY_ENCODE_HAS_PIXEL_ANIMATION
+		bHasPixelAnimation = DecodeHasPixelAnimationFromVelocityTexture(EncodedVelocity);
+	#endif
+	
+	#if CONFIG_USE_SCENE_VELOCITY
+	BRANCH
+	if (EncodedVelocity.x > 0.0)
+	{
+		ScreenVelocity = DecodeVelocityFromTexture(EncodedVelocity);
+	}
+	else
+	#endif
+	{
+		float2 ScreenPos = ApplyScreenTransform(float2(int2(InputPixelPos)), InputPixelPosToScreenPos);
+		ScreenVelocity = ComputeStaticVelocity(ScreenPos, ClosestDeviceZ);
+	}
+} // FetchAndComputeScreenVelocity()
+
+/** Spatial anti-alias the depth buffer and return the boundary. */
+tsr_half2 ComputeReprojectionBoundary(
+	tsr_short2 InputPixelPos,
+	float2 ScreenVelocity,
+	float2 DilatedScreenVelocity,
+	float DevizeZNeighborhood[CONFIG_BUTTERFLY_SAMPLES])
+{
+	const uint BrowsingIterations = 3;
+
+	float InputC = AccessNeighborhoodCenter(DevizeZNeighborhood);
+		
+	float InputN = AccessNeighborhoodStaticOffset(DevizeZNeighborhood, kOffsetN);
+	float InputS = AccessNeighborhoodStaticOffset(DevizeZNeighborhood, kOffsetS);
+	float InputE = AccessNeighborhoodStaticOffset(DevizeZNeighborhood, kOffsetE);
+	float InputW = AccessNeighborhoodStaticOffset(DevizeZNeighborhood, kOffsetW);
+		
+	float InputNE = AccessNeighborhoodStaticOffset(DevizeZNeighborhood, kOffsetNE);
+	float InputNW = AccessNeighborhoodStaticOffset(DevizeZNeighborhood, kOffsetNW);
+	float InputSE = AccessNeighborhoodStaticOffset(DevizeZNeighborhood, kOffsetSE);
+	float InputSW = AccessNeighborhoodStaticOffset(DevizeZNeighborhood, kOffsetSW);
+	
+	float2 ScreenPos = ApplyScreenTransform(float2(int2(InputPixelPos)), InputPixelPosToScreenPos);
+
+	float2 PrevScreenPos = ScreenPos - DilatedScreenVelocity.xy;
+	bool bOffscreen = any(PrevScreenPos != clamp(PrevScreenPos, -1.0, 1.0));
+
+	bool bDilationHasEnoughVelocityDiff = length2(InputInfo_ViewportSize * (ScreenVelocity.xy - DilatedScreenVelocity.xy)) > ReprojectionFieldAntiAliasVelocityThreshold;
+	bool bAntiAliasReprojectionField = bReprojectionField && bDilationHasEnoughVelocityDiff && !bOffscreen;
+
+	tsr_half2 ReprojectionBoundary;
+	BRANCH
+	if (bAntiAliasReprojectionField)
+	{
+		float NoiseFiltering;
+		tsr_short2 BrowseDirection;
+		tsr_short2 EdgeSide;
+		float EdgeInput;
+		FindBrowsingDirection(
+			InputC,
+			InputN,  InputS,  InputE,  InputW,
+			InputNE, InputNW, InputSE, InputSW,
+			/* out */ NoiseFiltering,
+			/* out */ BrowseDirection,
+			/* out */ EdgeSide,
+			/* out */ EdgeInput);
+				
+		#if CONFIG_BUTTERFLY_KERNEL
+			EdgeSide *= GetLaneOffsetSign();
+		#endif
+		
+		tsr_ushort EdgeLengthP, EdgeLengthN;
+		bool bEdgeStopedByIncrementP, bEdgeStopedByIncrementN;
+		bool bEdgeStopedByDecrementP, bEdgeStopedByDecrementN;
+		BrowseNeighborhoodBilinearOptimized(
+			BrowsingIterations,
+			SceneDepthTexture,
+			InputC,
+			EdgeInput,
+			InputPixelPos,
+			BrowseDirection,
+			EdgeSide,
+			/* out */ EdgeLengthP,
+			/* out */ EdgeLengthN,
+			/* out */ bEdgeStopedByIncrementP,
+			/* out */ bEdgeStopedByIncrementN,
+			/* out */ bEdgeStopedByDecrementP,
+			/* out */ bEdgeStopedByDecrementN);
+			
+		ReprojectionBoundary = ComputeReprojectionBoundary(
+			BrowsingIterations,
+			EdgeSide, EdgeLengthP, EdgeLengthN,
+			bEdgeStopedByIncrementP, bEdgeStopedByIncrementN,
+			bEdgeStopedByDecrementP, bEdgeStopedByDecrementN);
+	}
+	else
+	{
+		// Fully dilate the boundary in the history update to still apply the jacobian on geometric edges.
+		ReprojectionBoundary = kFullDilateBoundary;
+	}
 
-//------------------------------------------------------- FUNCTION
+	return ReprojectionBoundary;
+} // ComputeReprojectionBoundary()
 
 void ScatterClosestOccluder(bool bValidOutputPixel, float2 ScreenPos, float2 FinalScreenVelocity, float PrevClosestDeviceZ)
 {
@@ -154,7 +577,7 @@ void ScatterClosestOccluder(bool bValidOutputPixel, float2 ScreenPos, float2 Fin
 			InterlockedMax(PrevAtomicOutput[tsr_ushort3(PixelPos, 1)], PrevClosestDepth);
 		}
 	}
-}
+} // ScatterClosestOccluder()
 		
 
 //------------------------------------------------------- ENTRY POINT
@@ -164,206 +587,157 @@ void MainCS(
 	uint2 GroupId : SV_GroupID,
 	uint GroupThreadIndex : SV_GroupIndex)
 {
-	float4 Debug = 0.0;
+	for (uint DebugId = 0; DebugId < DEBUG_ARRAY_SIZE; DebugId++)
+	{
+		Debug[DebugId] = 0.0;
+	}
 
 	uint2 DispatchThreadId = (
 		ZOrder2D(GroupThreadIndex, uint(log2(float(TILE_SIZE)))) +
 		GroupId * uint2(TILE_SIZE, TILE_SIZE));
 
-	#if CONFIG_AA_CROSS_OPTIMISATION
-		int2 ThreadOffsetSign = -1 + int2((DispatchThreadId % uint(2)) * uint(2));
-	#endif
-
-	uint2 InputPixelPos = InputInfo_ViewportMin + DispatchThreadId;
+	tsr_short2 InputPixelPos = tsr_short2(InputInfo_ViewportMin + DispatchThreadId);
 	
-	float2 ViewportUV = (float2(DispatchThreadId) + 0.5f) * InputInfo_ViewportSizeInverse;
-	float2 ScreenPos = ViewportUVToScreenPos(ViewportUV);
-	float2 InputBufferUV = (InputPixelPos + 0.5) * InputInfo_ExtentInverse;
-
-	// Do all depth memory fetches.
+	// Find the closest depth and its offset.
 	float DeviceZ;
-	#if CONFIG_AA_CROSS_OPTIMISATION
-		float NeighborDeviceZArray[3];
-	#endif
-	ISOLATE
+	float ClosestDeviceZ;
+	float3 DilatedScreenVelocity;
+	uint EncodedReprojectionBoundary;
+	uint EncodedReprojectionJacobian;
+	bool bHasPixelAnimation;
+	tsr_half VelocityEdge;
+	BRANCH
+	if (bReprojectionField)
 	{
+		// Load the 3x3 Neighborhood.
+		float DevizeZNeighborhood[CONFIG_BUTTERFLY_SAMPLES];
+		float4 VelocityNeighborhood[CONFIG_BUTTERFLY_SAMPLES];
+		ISOLATE
 		{
-			int2 NeighborPixelPos = clamp(int2(InputPixelPos), InputPixelPosMin, InputPixelPosMax);
-				
-			DeviceZ = SceneDepthTexture[NeighborPixelPos];
+			FetchDepthVelocity3x3(InputPixelPos, /* out */ DevizeZNeighborhood, /* out */ VelocityNeighborhood);
 		}
-		#if CONFIG_AA_CROSS_OPTIMISATION
-			UNROLL_N(3)
-			for (uint i = 0; i < 3; i++)
-			{
-				int2 NeighborPixelPos = int2(InputPixelPos) + kOffsets[i] * ThreadOffsetSign;
-				NeighborPixelPos = clamp(NeighborPixelPos, InputPixelPosMin, InputPixelPosMax);
 
-				NeighborDeviceZArray[i] = SceneDepthTexture[NeighborPixelPos];
-			}
+		// Convert velocities to pixel velocities to take down register pressure.
+		float2 PixelVelocityNeighborhood[CONFIG_BUTTERFLY_SAMPLES];
+		float  DepthVelocityNeighborhood[CONFIG_BUTTERFLY_SAMPLES];
+		bool   bIsDrawingVelocityNeighborhood[CONFIG_BUTTERFLY_SAMPLES];
+		ISOLATE
+		{
+			ComputePixelVelocityNeighborhood(
+				InputPixelPos,
+				DevizeZNeighborhood,
+				VelocityNeighborhood,
+				/* out */ PixelVelocityNeighborhood,
+				/* out */ DepthVelocityNeighborhood,
+				/* out */ bIsDrawingVelocityNeighborhood);
+		}
+		#if VELOCITY_ENCODE_HAS_PIXEL_ANIMATION
+			bHasPixelAnimation = DecodeHasPixelAnimationFromVelocityTexture(AccessNeighborhoodCenter(VelocityNeighborhood));
 		#endif
-	}
 	
-	float ClosestDeviceZ = DeviceZ;
-	int2 VelocityPixelOffset = int2(0, 0);
-	#if CONFIG_AA_CROSS
-	ISOLATE
-	{
-		#if CONFIG_AA_CROSS_OPTIMISATION
-			const FWaveBroadcastSettings Horizontal = InitWaveXorButterfly(0x1);
-			const FWaveBroadcastSettings Vertical   = InitWaveXorButterfly(0x2);
-			const FWaveBroadcastSettings Diagonal   = InitWaveXorButterfly(0x3);
-
-			const FWaveBroadcastSettings CenterBroadcastArray[4] = {
-				Horizontal,
-				Diagonal,
-				Vertical,
-				InitNopBroadcast(), // unused
-			};
-		#endif
-
-
-		uint PackedDepthOffset = PackDepthAndOffset(/* Offset = */ int2(0, 0), DeviceZ);
-
-		#if CONFIG_DEPTH_EDGE_DETECTION
-		float DeviceZError;
+		tsr_half2x2 ReprojectionJacobian;
+		ISOLATE
 		{
-			float Depth = ConvertFromDeviceZ(DeviceZ);
-			float ErrorCorrectedDeviceZ = ConvertToDeviceZ(GetDepthPixelRadiusForProjectionType(Depth) * 2.0f + Depth);
-			DeviceZError = abs(ErrorCorrectedDeviceZ - DeviceZ);
+			ComputeReprojectionJacobian(
+				DevizeZNeighborhood,
+				PixelVelocityNeighborhood,
+				/* out */ ReprojectionJacobian);
+			EncodedReprojectionJacobian = EncodeReprojectionJacobian(ReprojectionJacobian);
 		}
-		#endif
-
-		#if 1
-			UNROLL_N(4)
-			for (uint i = 0; i < 4; i++)
-		#else
-			UNROLL_N(2)
-			for (uint i = 1; i < 4; i+= 2)
-		#endif
+		
+		tsr_short2 ReprojectionOffset;
+		ISOLATE
 		{
-			#if CONFIG_AA_CROSS_OPTIMISATION
-				int2 Offset = kOffsets[i] * ThreadOffsetSign;
-
-				float SampleDeviceZ0 = i == 3 ? WaveBroadcast(Horizontal, NeighborDeviceZArray[2]) : NeighborDeviceZArray[i];
-				float SampleDeviceZ1 = i == 3 ? WaveBroadcast(Vertical  , NeighborDeviceZArray[0]) : WaveBroadcast(CenterBroadcastArray[i], DeviceZ);
-			#else // !CONFIG_AA_CROSS_OPTIMISATION
-				const int2 Offset = kOffsets[i];
-				
-				int2 NeighborPixelPos0 = int2(InputPixelPos) + Offset;
-				int2 NeighborPixelPos1 = int2(InputPixelPos) - Offset;
-				NeighborPixelPos0 = clamp(NeighborPixelPos0, InputPixelPosMin, InputPixelPosMax);
-				NeighborPixelPos1 = clamp(NeighborPixelPos1, InputPixelPosMin, InputPixelPosMax);
-
-				float SampleDeviceZ0 = SceneDepthTexture[NeighborPixelPos0];
-				float SampleDeviceZ1 = SceneDepthTexture[NeighborPixelPos1];
-			#endif // !CONFIG_AA_CROSS_OPTIMISATION
-
-			#if CONFIG_DEPTH_EDGE_DETECTION
-			{
-				float DepthDiff = abs(SampleDeviceZ1 - SampleDeviceZ0);
-				float DepthVariation = (SampleDeviceZ0 + SampleDeviceZ1) * 0.5 - DeviceZ;
-
-				FLATTEN
-				if (DepthVariation > max(DepthDiff * 0.25, DeviceZError))
-				{
-					PackedDepthOffset = max(PackedDepthOffset, PackDepthAndOffset(+Offset, SampleDeviceZ0));
-					PackedDepthOffset = max(PackedDepthOffset, PackDepthAndOffset(-Offset, SampleDeviceZ1));
-				}
-			}
-			#else
-			{
-				PackedDepthOffset = max(PackedDepthOffset, PackDepthAndOffset(+Offset, SampleDeviceZ0));
-				PackedDepthOffset = max(PackedDepthOffset, PackDepthAndOffset(-Offset, SampleDeviceZ1));
-			}
-			#endif
+			FindClosestDepthOffset(
+				DevizeZNeighborhood,
+				/* out */ ClosestDeviceZ,
+				/* out */ ReprojectionOffset);
 		}
 
-		UnpackDepthAndOffset(PackedDepthOffset, /* out */ VelocityPixelOffset, /* out */ ClosestDeviceZ);
-	}
-	#endif
+		float2 ScreenVelocity;
+		ISOLATE
+		{
+			// Compute final screen vecocity from pixel velocities to avoid duplicated register.
+			DeviceZ = AccessNeighborhoodCenter(DevizeZNeighborhood);
+			ScreenVelocity = AccessNeighborhoodCenter(PixelVelocityNeighborhood) * InputPixelVelocityToScreenVelocity;
 	
-	// Camera motion for pixel.
-	float3 FinalScreenVelocity = 0.0;
-
-	float4 CenterEncodedVelocity = SceneVelocityTexture[int2(InputBufferUV * InputInfo_Extent)];
-	bool bIsMovingPixel = CenterEncodedVelocity.x > 0.0;
+			// Compute the dilated screen velocity.
+			bool bDildatedHasPixelAnimation;
+			FetchAndComputeScreenVelocity(
+				InputPixelPos,
+				ClosestDeviceZ,
+				ReprojectionOffset,
+				/* out */ DilatedScreenVelocity,
+				/* out */ bDildatedHasPixelAnimation);
+		}
 	
-	bool bHasPixelAnimation = false;
-	#if !VELOCITY_ENCODE_HAS_PIXEL_ANIMATION
-		// UE-198416: there is high mess arround VelocityEncodeDepth() to clean up
-		//#error TSR requires the bHasPixelAnimation encoded in the velocity buffer.
-	#endif
-
-	// Compute final screen vecocity.
-	#if CONFIG_VELOCITY_EXTRAPOLATION
-	BRANCH
-	if (VelocityExtrapolationMultiplier > 0.0)
-	{
-		int2 PixelCoord1 = int2(InputBufferUV * InputInfo_Extent) + VelocityPixelOffset * 1;
-		int2 PixelCoord2 = int2(InputBufferUV * InputInfo_Extent) + VelocityPixelOffset * 2;
-
-		PixelCoord2 = clamp(PixelCoord2, InputPixelPosMin, InputPixelPosMax);
-
-		float ClosestDeviceZ1 = ClosestDeviceZ;
-		float ClosestDeviceZ2 = SceneDepthTexture[PixelCoord2];
-
-		float4 EncodedVelocity1 = SceneVelocityTexture[PixelCoord1];
-		float4 EncodedVelocity2 = SceneVelocityTexture[PixelCoord2];
-
-		float3 ScreenVelocity1 = ComputeStaticVelocity(ScreenPos, ClosestDeviceZ1);
-		float3 ScreenVelocity2 = ComputeStaticVelocity(ScreenPos, ClosestDeviceZ2);
+		tsr_half2 NeighborDeltaPixelVelocity = tsr_half2((DilatedScreenVelocity.xy - ScreenVelocity) * ScreenVelocityToInputPixelVelocity);
+		VelocityEdge = saturate(tsr_half(1.1) - dot(abs(NeighborDeltaPixelVelocity), tsr_half(1.1).xx));
 		
-		#if VELOCITY_ENCODE_HAS_PIXEL_ANIMATION
-			bHasPixelAnimation = DecodeHasPixelAnimationFromVelocityTexture(EncodedVelocity1);
-		#endif
-
-		#if CONFIG_USE_SCENE_VELOCITY
+		// Compute the reprojection boundary using a depth spatial anti-aliaser
+		ISOLATE
 		{
-			if (EncodedVelocity1.x > 0.0)
-			{
-				ScreenVelocity1 = DecodeVelocityFromTexture(EncodedVelocity1);
-			}
-			if (EncodedVelocity2.x > 0.0)
-			{
-				ScreenVelocity2 = DecodeVelocityFromTexture(EncodedVelocity2);
-			}
-		}
-		#endif
-
-		float3 VelocityDelta = ScreenVelocity1 - ScreenVelocity2;
-		float ExtrapolationBilateral = ComputeScreenVelocityBilateralWeight(VelocityDelta.xy);
+			tsr_half2 ReprojectionBoundary = ComputeReprojectionBoundary(
+				InputPixelPos,
+				ScreenVelocity,
+				DilatedScreenVelocity.xy,
+				DevizeZNeighborhood);
 
-		FinalScreenVelocity = ScreenVelocity1 + VelocityDelta * ExtrapolationBilateral * VelocityExtrapolationMultiplier;
+			EncodedReprojectionBoundary = EncodeReprojectionBoundary(ReprojectionOffset, ReprojectionBoundary);
+		}
 	}
-	else
-	#endif
+	else // if (!bReprojectionField)
 	{
-		FinalScreenVelocity = ComputeStaticVelocity(ScreenPos, ClosestDeviceZ);
-		
-		float4 EncodedVelocity = SceneVelocityTexture[uint2(InputBufferUV * InputInfo_Extent) + VelocityPixelOffset];
-
-		#if VELOCITY_ENCODE_HAS_PIXEL_ANIMATION
-			bHasPixelAnimation = DecodeHasPixelAnimationFromVelocityTexture(EncodedVelocity);
-		#endif
-
-		#if CONFIG_USE_SCENE_VELOCITY
+		// Load the 3x3 Neighborhood.
+		float DevizeZNeighborhood[CONFIG_BUTTERFLY_SAMPLES];
+		ISOLATE
 		{
-			if (EncodedVelocity.x > 0.0)
-			{
-				FinalScreenVelocity = DecodeVelocityFromTexture(EncodedVelocity);
-			}
+			FetchDepth3x3(InputPixelPos, /* out */ DevizeZNeighborhood);
 		}
-		#endif
+		
+		tsr_short2 ReprojectionOffset;
+		FindClosestDepthOffset(
+			DevizeZNeighborhood,
+			/* out */ ClosestDeviceZ,
+			/* out */ ReprojectionOffset);
+	
+		// Compute final screen vecocity from pixel velocities to avoid duplicated register.
+		DeviceZ        = AccessNeighborhoodCenter(DevizeZNeighborhood);
+	
+		// Compute the screen vecocity.
+		float3 ScreenVelocity;
+		FetchAndComputeScreenVelocity(
+			InputPixelPos,
+			ClosestDeviceZ,
+			/* ReprojectionOffset = */ tsr_short2(0, 0),
+			/* out */ ScreenVelocity,
+			/* out */ bHasPixelAnimation);
+
+		// Compute the dilated screen vecocity.
+		bool bDildatedHasPixelAnimation;
+		FetchAndComputeScreenVelocity(
+			InputPixelPos,
+			ClosestDeviceZ,
+			ReprojectionOffset,
+			/* out */ DilatedScreenVelocity,
+			/* out */ bDildatedHasPixelAnimation);
+			
+		tsr_half2 NeighborDeltaPixelVelocity = tsr_half2((DilatedScreenVelocity.xy - ScreenVelocity.xy) * ScreenVelocityToInputPixelVelocity);
+		VelocityEdge = saturate(tsr_half(1.1) - dot(abs(NeighborDeltaPixelVelocity), tsr_half(1.1).xx));
+
+		EncodedReprojectionBoundary = 0u;
+		EncodedReprojectionJacobian = 0u;
 	}
 
 	tsr_half IsMovingMask = tsr_half(0.0);
 	BRANCH
 	if (bOutputIsMovingTexture)
 	{
+		float2 ScreenPos = ApplyScreenTransform(float2(int2(InputPixelPos)), InputPixelPosToScreenPos);
 		float Depth = ConvertFromDeviceZ(ClosestDeviceZ);
-		float PrevDepth = ConvertFromDeviceZ(ClosestDeviceZ - FinalScreenVelocity.z);
-		float2 PrevScreenPos = ScreenPos - FinalScreenVelocity.xy;
+
+		float PrevDepth = ConvertFromDeviceZ(ClosestDeviceZ - DilatedScreenVelocity.z);
+		float2 PrevScreenPos = ScreenPos - DilatedScreenVelocity.xy;
 		
 		float4 ThisClip = float4(ScreenPos, DeviceZ, 1);
 		float4 PrevClip = mul(ThisClip, View.ClipToPrevClip); 
@@ -394,26 +768,31 @@ void MainCS(
 	#if DIM_MOTION_BLUR_DIRECTIONS > 0
 		FVelocityRange VelocityPolarRange;
 		float2 VelocityPolar;
-		ReduceVelocityFlattenTile(GroupThreadIndex, FinalScreenVelocity.xy, /* out */ VelocityPolarRange, /* out */ VelocityPolar);
+		ReduceVelocityFlattenTile(GroupThreadIndex, DilatedScreenVelocity.xy, /* out */ VelocityPolarRange, /* out */ VelocityPolar);
 	#endif
 	
-	float PrevClosestDeviceZ = ClosestDeviceZ - FinalScreenVelocity.z;
+	float PrevClosestDeviceZ = ClosestDeviceZ - DilatedScreenVelocity.z;
 
 	{
 		bool bValidOutputPixel = all(InputPixelPos < InputInfo_ViewportMax);
-		
-		float4 EncodedVelocity = SceneVelocityTexture[uint2(InputBufferUV * InputInfo_Extent)];
-		
 		uint2 OutputPixelPos = bValidOutputPixel ? InputPixelPos : uint(~0).xx;
 
-		//DilatedVelocityOutput[OutputPixelPos] = EncodedVelocity.xy; //EncodeVelocityToTexture(float3(FinalScreenVelocity.xy, 0.0)).xy;
-		DilatedVelocityOutput[OutputPixelPos] = EncodeVelocityToTexture(float3(FinalScreenVelocity.xy, 0.0)).xy;
 		ClosestDepthOutput[OutputPixelPos] = float2(select(bHasPixelAnimation, -PrevClosestDeviceZ, PrevClosestDeviceZ), ClosestDeviceZ);
 
+		R8Output[uint3(OutputPixelPos, 0)] = uint(round(VelocityEdge * 255.0));
 		BRANCH
 		if (bOutputIsMovingTexture)
 		{
-			R8Output[uint3(OutputPixelPos, 0)] = uint(round(IsMovingMask * 255.0));
+			R8Output[uint3(OutputPixelPos, 1)] = uint(round(IsMovingMask * 255.0));
+		}
+
+		// Output the reprojection field
+		ReprojectionFieldOutput[uint3(OutputPixelPos, 0)] = EncodeReprojectionVector(DilatedScreenVelocity.xy);
+		BRANCH
+		if (bReprojectionField)
+		{
+			ReprojectionFieldOutput[uint3(OutputPixelPos, 1)] = EncodedReprojectionBoundary;
+			ReprojectionFieldOutput[uint3(OutputPixelPos, 2)] = EncodedReprojectionJacobian;
 		}
 
 		// Output motion blur velocity flatten
@@ -427,11 +806,13 @@ void MainCS(
 		#endif
 		
 		// Scatter parallax rejection
-		ScatterClosestOccluder(bValidOutputPixel, ScreenPos, FinalScreenVelocity.xy, PrevClosestDeviceZ);
+		float2 ScreenPos = ApplyScreenTransform(float2(int2(OutputPixelPos)), InputPixelPosToScreenPos);
+		ScatterClosestOccluder(bValidOutputPixel, ScreenPos, DilatedScreenVelocity.xy, PrevClosestDeviceZ);
 		
-		#if DEBUG_OUTPUT && 0 // running out of UAVs
+		#if DEBUG_OUTPUT
+		for (uint DebugId = 0; DebugId < DEBUG_ARRAY_SIZE; DebugId++)
 		{
-			DebugOutput[tsr_short3(OutputPixelPos, 0)] = Debug;
+			DebugOutput[tsr_short3(OutputPixelPos, DebugId)] = Debug[DebugId];
 		}
 		#endif
 	}
diff --git a/Engine/Shaders/Private/TemporalSuperResolution/TSRRejectShading.usf b/Engine/Shaders/Private/TemporalSuperResolution/TSRRejectShading.usf
--- a/Engine/Shaders/Private/TemporalSuperResolution/TSRRejectShading.usf
+++ b/Engine/Shaders/Private/TemporalSuperResolution/TSRRejectShading.usf
@@ -3,6 +3,7 @@
 //------------------------------------------------------- INCLUDES
 
 #include "TSRShadingAnalysis.ush"
+#include "TSRReprojectionField.ush"
 
 
 //------------------------------------------------------- CONFIG
@@ -31,6 +32,10 @@
 
 FScreenTransform InputPixelPosToTranslucencyTextureUV;
 float4x4 ClipToResurrectionClip;
+float2 ResurrectionJacobianXMul;
+float2 ResurrectionJacobianXAdd;
+float2 ResurrectionJacobianYMul;
+float2 ResurrectionJacobianYAdd;
 float2 TranslucencyTextureUVMin;
 float2 TranslucencyTextureUVMax;
 float3 HistoryGuideQuantizationError;
@@ -50,6 +55,7 @@ Texture2D<tsr_half4> ReprojectedHistoryMoireTexture;
 	Texture2D<tsr_half4> ResurrectedHistoryGuideTexture;
 	Texture2D<tsr_half>  ResurrectedHistoryGuideMetadataTexture;
 #endif
+Texture2D<uint> DilateMaskTexture;
 Texture2D<tsr_half2> DecimateMaskTexture;
 Texture2D<tsr_ushort> IsMovingMaskTexture;
 Texture2D<float2> ClosestDepthTexture;
@@ -57,7 +63,7 @@ Texture2D<float2> ClosestDepthTexture;
 RWTexture2DArray<tsr_half4> HistoryGuideOutput;
 RWTexture2DArray<tsr_half4> HistoryMoireOutput;
 RWTexture2D<tsr_half4>  HistoryRejectionOutput;
-RWTexture2D<float2>     DilatedVelocityOutput;
+RWTexture2DArray<uint>  ReprojectionFieldOutput;
 RWTexture2D<tsr_halfC>  InputSceneColorOutput;
 RWTexture2D<tsr_half>   InputSceneColorLdrLumaOutput;
 RWTexture2D<tsr_ushort> AntiAliasMaskOutput;
@@ -122,11 +128,12 @@ void FetchDecimateMask(
 	{
 		tsr_short2 ClampedFetchPixelPos = PixelPos.GetElement(ElementIndex);
 		
+		uint DilateMask = DilateMaskTexture[ClampedFetchPixelPos];
 		tsr_half2 DecimateMask = DecimateMaskTexture[ClampedFetchPixelPos];
 
 		tsr_ushort BitMask = tsr_ushort(DecimateMask.r * tsr_half(255.0));
 
-		VelocityEdge.SetElement(ElementIndex, DecimateMask.g);
+		VelocityEdge.SetElement(ElementIndex, saturate(tsr_half(tsr_ushort(DilateMask)) * tsr_half(rcp(255.0))));
 		bIsDisoccluded.SetElement(ElementIndex,     (BitMask & tsr_ushort(0x1)) != tsr_ushort(0));
 		bHasPixelAnimation.SetElement(ElementIndex, (BitMask & tsr_ushort(0x2)) != tsr_ushort(0));
 	}
@@ -299,26 +306,52 @@ void FetchDeviceZ(
 	}
 } // FetchMoireInput()
 
-// Compute the velocity used by the reprojection of the resurrection.
+// Overwrite the reprojection field for pixels that needs resurrection.
 CALL_SITE_DEBUGLOC
-tsr_tensor_float2 ComputeResurrectionMotionVector(tsr_tensor_float DeviceZ)
+void OverwriteReprojectionFieldWithResurrection(tsr_tensor_bool bResurrectHistory, tsr_tensor_float DeviceZ)
 {
-	tsr_tensor_float2 ResurrectedEncodedScreenVelocity;
-	
-	UNROLL_N(SIMD_SIZE)
-	for (uint ElementIndex = 0; ElementIndex < SIMD_SIZE; ElementIndex++)
-	{
-		float2 ScreenPos = ApplyScreenTransform(float2(ComputeElementPixelPos(ElementIndex)), InputPixelPosToScreenPos);
+	uint EncodedReprojectionVector[SIMD_SIZE];
+	uint EncodedReprojectionJacobian[SIMD_SIZE];
 
-		float4 ThisClip = float4(ScreenPos, DeviceZ.GetElement(ElementIndex), 1);
-		float4 ResurrectionClip = mul(ThisClip, ClipToResurrectionClip); 
-		float2 ResurrectionScreen = ResurrectionClip.xy / ResurrectionClip.w;
+	PLATFORM_SPECIFIC_ISOLATE
+	{
+		UNROLL_N(SIMD_SIZE)
+		for (uint ElementIndex = 0; ElementIndex < SIMD_SIZE; ElementIndex++)
+		ISOLATE
+		{
+			float LocalDeviceZ = DeviceZ.GetElement(ElementIndex);
+		
+			tsr_short2 InputPixelPos = ComputeElementPixelPos(ElementIndex);
+			float2 ScreenPos = ApplyScreenTransform(float2(InputPixelPos), InputPixelPosToScreenPos);
 			
-		ResurrectedEncodedScreenVelocity.SetElement(ElementIndex, EncodeVelocityToTexture(float3(ScreenPos - ResurrectionScreen, 0.0)).xy);
+			float4 ThisClip = float4(ScreenPos, DeviceZ.GetElement(ElementIndex), 1);
+			float4 ResurrectionClip = mul(ThisClip, ClipToResurrectionClip);
+			float  ResurrectionProjection = rcp(ResurrectionClip.w);
+			float2 ResurrectionReprojectionVector = ScreenPos - ResurrectionClip.xy * ResurrectionProjection;
+			EncodedReprojectionVector[ElementIndex] = EncodeReprojectionVector(ResurrectionReprojectionVector);
+
+			// Compute the Jacobian from the PrevClip. Should technically take into account the depth delta along X and Y, but it's costly to compute, and it's using the dilated DeviceZ anyway, and not even using a 2x3 jacobian.
+			tsr_half2x2 ResurrectionJacobian;
+			ResurrectionJacobian[0] = tsr_half2(ResurrectionJacobianXMul) * tsr_half(ResurrectionProjection) + tsr_half2(ResurrectionJacobianXAdd);
+			ResurrectionJacobian[1] = tsr_half2(ResurrectionJacobianYMul) * tsr_half(ResurrectionProjection) + tsr_half2(ResurrectionJacobianYAdd);
+
+			EncodedReprojectionJacobian[ElementIndex] = EncodeReprojectionJacobian(ResurrectionJacobian);
+		}
 	}
 
-	return ResurrectedEncodedScreenVelocity;
-}
+	PLATFORM_SPECIFIC_ISOLATE
+	{
+		UNROLL_N(SIMD_SIZE)
+		for (uint ElementIndex = 0; ElementIndex < SIMD_SIZE; ElementIndex++)
+		{
+			tsr_short2 OutputPixelPos = ComputeElementOutputPixelPos(ElementIndex);
+			OutputPixelPos.x = select(bResurrectHistory.GetElement(ElementIndex), OutputPixelPos.x, tsr_short(-1));
+
+			ReprojectionFieldOutput[tsr_short3(OutputPixelPos, 0)] = EncodedReprojectionVector[ElementIndex];
+			ReprojectionFieldOutput[tsr_short3(OutputPixelPos, 2)] = EncodedReprojectionJacobian[ElementIndex];
+		}
+	}
+} // OverwriteReprojectionFieldWithResurrection()
 
 #endif // CONFIG_HISTORY_RESURRECTION
 
@@ -718,21 +751,8 @@ void MainCS(
 			History = select(tsr_tensor_boolC::Vectorize(bResurrectHistory), ResurrectedHistory, History);
 			HistoryUncertainty = select(bResurrectHistory, ResurrectedHistoryUncertainty, HistoryUncertainty);
 			
-			PLATFORM_SPECIFIC_ISOLATE
-			{
-				// Compute the motion vector for the resurrection
-				tsr_tensor_float2 ResurrectionMotionVector = ComputeResurrectionMotionVector(DeviceZ);
-			
-				// Amends the motion vector for resurrection areas.
-				UNROLL_N(SIMD_SIZE)
-				for (uint ElementIndex = 0; ElementIndex < SIMD_SIZE; ElementIndex++)
-				{
-					tsr_short2 OutputPixelPos = ComputeElementOutputPixelPos(ElementIndex);
-					OutputPixelPos.x = select(bResurrectHistory.GetElement(ElementIndex), OutputPixelPos.x, tsr_short(-1));
-
-					DilatedVelocityOutput[OutputPixelPos] = ResurrectionMotionVector.GetElement(ElementIndex);
-				}
-			}
+			// Apply the resurrection the reprojection field.
+			OverwriteReprojectionFieldWithResurrection(bResurrectHistory, DeviceZ);
 		}
 		#endif // CONFIG_HISTORY_RESURRECTION
 
diff --git a/Engine/Shaders/Private/TemporalSuperResolution/TSRReprojectionField.ush b/Engine/Shaders/Private/TemporalSuperResolution/TSRReprojectionField.ush
new file mode 100644
--- /dev/null
+++ b/Engine/Shaders/Private/TemporalSuperResolution/TSRReprojectionField.ush
@@ -0,0 +1,153 @@
+// Copyright Epic Games, Inc. All Rights Reserved.
+
+#include "TSRCommon.ush"
+
+
+//------------------------------------------------------- REPROJECTION's DILATING OFFSET & BOUNDARY
+
+static const tsr_half2 kFullDilateBoundary = tsr_half2(0.0, 1.0);
+
+
+/** Encodes the reprojection field's dilating offset [-1;1]^2 in first 4bits, and ReprojectionBoundary.xy in 11bit eachs, 26bits total. */
+uint EncodeReprojectionBoundary(tsr_short2 ReprojectionOffset, tsr_half2 ReprojectionBoundary)
+{
+	uint EncodedReprojectionOffset = uint(dot(tsr_ushort2(ReprojectionOffset + tsr_ushort(1)), tsr_ushort2(1, 4)));
+
+	tsr_short2 BoundaryBits = tsr_short2(round(ReprojectionBoundary * tsr_half(1023))) + 1024;
+	BoundaryBits = fastClamp(BoundaryBits, tsr_short(1).xx, tsr_short(2047).xx);
+
+	uint EncodedReprojectionBoundary = EncodedReprojectionOffset | (uint(BoundaryBits.x) << 4) | (uint(BoundaryBits.y) << 15);
+	return EncodedReprojectionBoundary;
+}
+
+/** Decodes EncodeReprojectionBoundary()'s ReprojectionOffset. */
+tsr_short2 DecodeReprojectionOffset(uint EncodedReprojectionBoundary)
+{
+	tsr_short2 ReprojectionOffset = (tsr_short2(EncodedReprojectionBoundary, EncodedReprojectionBoundary >> 2u) & tsr_short(0x3)) - tsr_short(1);
+	return ReprojectionOffset;
+}
+
+/** Returns any(EncodeReprojectionBoundary() != 0). */
+bool HasEncodedReprojectionOffset(uint EncodedReprojectionBoundary)
+{
+	uint EncodedReprojectionOffset = EncodedReprojectionBoundary & 0xFu;
+	bool bHasReprojectionOffset = EncodedReprojectionOffset != 5;
+	return bHasReprojectionOffset;
+}
+
+/** Decodes EncodeReprojectionBoundary()'s ReprojectionBoundary. */
+tsr_half2 DecodeReprojectionBoundary(uint EncodedReprojectionBoundary)
+{
+	tsr_short2 BoundaryBits = tsr_short2(EncodedReprojectionBoundary >> 4u, EncodedReprojectionBoundary >> 15u) % tsr_short(2048) - tsr_short(1024);
+	tsr_half2 ReprojectionBoundary = fastClamp(tsr_half2(BoundaryBits) * rcp(tsr_half(1023)), tsr_half(-1.0).xx, tsr_half(1.0).xx);
+	return ReprojectionBoundary;
+}
+
+/** Estimate the coverage of the boundary over the rendered pixel. */
+tsr_half ComputeReprojectionBoundaryCoverage(tsr_half2 ReprojectionBoundary)
+{
+	tsr_half ReprojectionBoundaryCoverage = length(ReprojectionBoundary);
+	return ReprojectionBoundaryCoverage;
+}
+
+/** Returns whether the history's pixel has any overlap with the velocity boundary, to apply the ReprojectionOffset. */
+bool IsHistoryPixelWithinOffsetBoundary(tsr_half2 dInputKO, tsr_half2 ReprojectionBoundary, float HistoryToInputFactor)
+{
+	tsr_half2 dInputKC = tsr_half(0.5) - tsr_half2(sign_bit(ReprojectionBoundary));
+	tsr_half2 dInputOR = + dInputKC * tsr_half(HistoryToInputFactor);
+
+	tsr_half2 dInputBC = ReprojectionBoundary;
+	tsr_half2 dInputBR = dInputBC - dInputKC + dInputKO + dInputOR;
+
+	bool bHistoryPixelWithinOffsetBoundary = dot(dInputBR, dInputBC) > tsr_half(0.0);
+	return bHistoryPixelWithinOffsetBoundary;
+}
+
+
+//------------------------------------------------------- REPROJECTION's VECTOR
+
+/** Encodes the reprojection vector into 32bits. */
+uint EncodeReprojectionVector(float2 ReprojectionVector)
+{
+	float2 EncodedScreenVelocity = EncodeVelocityToTexture(float3(ReprojectionVector.xy, 0.0)).xy;
+	
+	float2 UnormEncoding = clamp(EncodedScreenVelocity * 65535.0f + 0.5f, 0.0f, 65535.0f);
+
+	uint EncodedReprojectionVector = 0u;
+	EncodedReprojectionVector |= uint(UnormEncoding.x) <<  0u;
+	EncodedReprojectionVector |= uint(UnormEncoding.y) << 16u;
+	return EncodedReprojectionVector;
+}
+
+/** Inverse of EncodeReprojectionVector() */
+float2 DecodeReprojectionVector(uint EncodedReprojectionVector)
+{
+	float2 UnormEncoding = float2(uint2(EncodedReprojectionVector, EncodedReprojectionVector >> 16u) & 0xFFFFu);
+
+	float2 ReprojectionVector = DecodeVelocityFromTexture(float4(UnormEncoding * rcp(65535.0f), 0.0, 0.0)).xy;
+	return ReprojectionVector;
+}
+
+
+//------------------------------------------------------- REPROJECTION's JACOBIAN
+
+static const tsr_half kVelocityJacobianRange = tsr_half(2.0);
+//static const tsr_half kVelocityJacobianRange = tsr_half(0.5);
+static const uint kVelocityJacobianBits = 8u;
+
+/** Encodes the reprojection's jacobian into 32bits. */
+uint EncodeReprojectionJacobian(tsr_half2x2 ReprojectionJacobian)
+{
+	const tsr_half EncodingRange = tsr_half((1u << (kVelocityJacobianBits - 1u)) - 1u);
+	const tsr_half EncodingOffset = EncodingRange + tsr_half(0.5);
+	const tsr_half EncodingMultiply = EncodingRange / kVelocityJacobianRange;
+
+	ReprojectionJacobian[0] = clamp(ReprojectionJacobian[0], -kVelocityJacobianRange, kVelocityJacobianRange);
+	ReprojectionJacobian[1] = clamp(ReprojectionJacobian[1], -kVelocityJacobianRange, kVelocityJacobianRange);
+
+	#if 1
+		ReprojectionJacobian[0] = tsr_half2(sign(ReprojectionJacobian[0])) * sqrt(abs(ReprojectionJacobian[0])) * sqrt(kVelocityJacobianRange);
+		ReprojectionJacobian[1] = tsr_half2(sign(ReprojectionJacobian[1])) * sqrt(abs(ReprojectionJacobian[1])) * sqrt(kVelocityJacobianRange);
+	#endif
+
+	tsr_ushort2x2 DiscreteJacobian = tsr_ushort2x2(EncodingMultiply * ReprojectionJacobian + EncodingOffset);
+	uint EncodedReprojectionJacobian = 0u;
+	EncodedReprojectionJacobian |= uint(DiscreteJacobian[0][0]) << (0u * kVelocityJacobianBits);
+	EncodedReprojectionJacobian |= uint(DiscreteJacobian[0][1]) << (1u * kVelocityJacobianBits);
+	EncodedReprojectionJacobian |= uint(DiscreteJacobian[1][0]) << (2u * kVelocityJacobianBits);
+	EncodedReprojectionJacobian |= uint(DiscreteJacobian[1][1]) << (3u * kVelocityJacobianBits);
+	return EncodedReprojectionJacobian;
+}
+
+/** Inverse of EncodeReprojectionJacobian() */
+tsr_half2x2 DecodeReprojectionJacobian(uint EncodedReprojectionJacobian)
+{
+	const uint ComponentMask = (1u << kVelocityJacobianBits) - 1u;
+	const tsr_half EncodingRange = tsr_half((1u << (kVelocityJacobianBits - 1u)) - 1u);
+	const tsr_half DecodingMultiply = kVelocityJacobianRange / EncodingRange;
+
+	tsr_ushort2x2 DiscreteJacobian;
+	DiscreteJacobian[0][0] = tsr_ushort((EncodedReprojectionJacobian >> (0u * kVelocityJacobianBits)) & ComponentMask);
+	DiscreteJacobian[0][1] = tsr_ushort((EncodedReprojectionJacobian >> (1u * kVelocityJacobianBits)) & ComponentMask);
+	DiscreteJacobian[1][0] = tsr_ushort((EncodedReprojectionJacobian >> (2u * kVelocityJacobianBits)) & ComponentMask);
+	DiscreteJacobian[1][1] = tsr_ushort((EncodedReprojectionJacobian >> (3u * kVelocityJacobianBits)) & ComponentMask);
+	
+	tsr_half2x2 ReprojectionJacobian = tsr_half2x2(DiscreteJacobian) * DecodingMultiply - tsr_half(kVelocityJacobianRange);
+	
+	#if 1
+		ReprojectionJacobian[0] = ReprojectionJacobian[0] * abs(ReprojectionJacobian[0]) * rcp(kVelocityJacobianRange);
+		ReprojectionJacobian[1] = ReprojectionJacobian[1] * abs(ReprojectionJacobian[1]) * rcp(kVelocityJacobianRange);
+	#endif
+
+	return ReprojectionJacobian;
+}
+
+/** Estimate upscale ratio of the reprojection jsut from the jacobian. */
+tsr_half ComputeReprojectionUpscaleFactorFromJacobian(tsr_half2x2 ReprojectionJacobian)
+{
+	tsr_half2 VectorOffsetE = tsr_half2(1.0, 0.0) + ReprojectionJacobian[0];
+	tsr_half2 VectorOffsetS = tsr_half2(0.0, 1.0) + ReprojectionJacobian[1];
+
+	tsr_half ReprojectionUpscaleFactor = max(length(VectorOffsetE), length(VectorOffsetS));
+	return ReprojectionUpscaleFactor;
+}
diff --git a/Engine/Shaders/Private/TemporalSuperResolution/TSRSpatialAntiAliasing.ush b/Engine/Shaders/Private/TemporalSuperResolution/TSRSpatialAntiAliasing.ush
--- a/Engine/Shaders/Private/TemporalSuperResolution/TSRSpatialAntiAliasing.ush
+++ b/Engine/Shaders/Private/TemporalSuperResolution/TSRSpatialAntiAliasing.ush
@@ -18,6 +18,18 @@ tsr_half2 SampleAAInput(Texture2D<tsr_half> Texture, tsr_short2x2 KernelCenter,
 	return InputLuma;
 }
 
+float BilinearSampleAAInput(Texture2D<float> Texture, float2 KernelCenterUV, float2 UVOffset, const bool bIsNegativeDirection)
+{
+	float2 SampleUV = KernelCenterUV + UVOffset;
+
+	if (bIsNegativeDirection)
+		SampleUV = max(SampleUV, InputInfo_UVViewportBilinearMin);
+	else
+		SampleUV = min(SampleUV, InputInfo_UVViewportBilinearMax);
+
+	return Texture.SampleLevel(GlobalBilinearClampedSampler, SampleUV, 0);
+}
+
 tsr_half2 BilinearSampleAAInput(Texture2D<tsr_half> Texture, float2x2 KernelCenterUV, float2x2 UVOffset, const bool bIsNegativeDirection)
 {
 	float2x2 SampleUV = KernelCenterUV + UVOffset;
@@ -37,6 +49,59 @@ tsr_half2 BilinearSampleAAInput(Texture2D<tsr_half> Texture, float2x2 KernelCent
 
 //------------------------------------------------------- FUNCTIONS
 
+/** Finds the browsing direction of the spatial anti-aliaser out from the neighborhood. */
+void FindBrowsingDirection(
+	float InputC,
+	float InputN,
+	float InputS,
+	float InputE,
+	float InputW,
+	float InputNE,
+	float InputNW,
+	float InputSE,
+	float InputSW,
+	out float NoiseFiltering,
+	out tsr_short2 BrowseDirection,
+	out tsr_short2 EdgeSide,
+	out float EdgeInput)
+{
+	float TotalVarHN = abs(0.5 * (InputNE + InputNW) - InputN);
+	float TotalVarH  = abs(0.5 * (InputE  + InputW ) - InputC);
+	float TotalVarHS = abs(0.5 * (InputSE + InputSW) - InputS);
+			
+	float TotalVarVE = abs(0.5 * (InputNE + InputSE) - InputE);
+	float TotalVarV  = abs(0.5 * (InputN  + InputS ) - InputC);
+	float TotalVarVW = abs(0.5 * (InputNW + InputSW) - InputW);
+		
+	float DiffN = abs(InputN - InputC);
+	float DiffS = abs(InputS - InputC);
+	float DiffE = abs(InputE - InputC);
+	float DiffW = abs(InputW - InputC);
+
+	bool bBrowseVerticaly = (TotalVarHN + TotalVarH + TotalVarHS) > (TotalVarVE + TotalVarV + TotalVarVW);
+	BrowseDirection[0] = select(bBrowseVerticaly, tsr_short(0), tsr_short(1));
+	BrowseDirection[1] = select(bBrowseVerticaly, tsr_short(1), tsr_short(0));
+	
+	{
+		NoiseFiltering = saturate(float(2.0) * TotalVarH - max(DiffE, DiffW));
+		NoiseFiltering = max(NoiseFiltering, saturate(float(2.0) * TotalVarV - max(DiffN, DiffS)));
+	}
+
+	EdgeSide = tsr_short(0);
+
+	FLATTEN
+	if (bBrowseVerticaly)
+	{
+		EdgeSide[0] = DiffW > DiffE ? -1 : 1;
+		EdgeInput = DiffW > DiffE ? InputW : InputE;
+	}
+	else
+	{
+		EdgeSide[1] = DiffN > DiffS ? -1 : 1;
+		EdgeInput = DiffN > DiffS ? InputN : InputS;
+	}
+} // FindBrowsingDirection()
+
 /** Finds the browsing direction of the spatial anti-aliaser out from the neighborhood. */
 void FindBrowsingDirection(
 	tsr_half2 InputC,
@@ -102,6 +167,9 @@ void FindBrowsingDirection(
 	}
 } // FindBrowsingDirection()
 
+
+//------------------------------------------------------- ACCUMULATE
+
 /** Accumulate an edge length */
 void AccumulateEdgeLength(
 	tsr_half2 StartLuma, tsr_half2 StartEdgeLuma,
@@ -230,10 +298,11 @@ void BrowseNeighborhoodBilinearOptimized(
 	bool bStop = false;
 		
 	#if CONFIG_COMPILE_FP16
-		#define asuint(x) asuint16(x)
+		#define asuintX(x) asuint16(x)
 		const tsr_ushort SignShift = 15;
 		const tsr_ushort SignMask = 0x8000;
 	#else
+		#define asuintX(x) asuint(x)
 		const tsr_ushort SignShift = 31;
 		const tsr_ushort SignMask = 0x80000000;
 	#endif
@@ -251,10 +320,10 @@ void BrowseNeighborhoodBilinearOptimized(
 				
 		//bEdgeStopedByMinLuma = bEdgeStopedByMinLuma || bool2(SampleLuma < MinLuma) && !bEdgeStopedByMaxLuma;
 		//bEdgeStopedByMaxLuma = bEdgeStopedByMaxLuma || bool2(SampleLuma > MaxLuma) && !bEdgeStopedByMinLuma;
-		tsr_ushort2 StopByMinLumaP = asuint(SampleLumaP - MinLuma) & ~EdgeStopedByMaxLumaP;
-		tsr_ushort2 StopByMaxLumaP = asuint(MaxLuma - SampleLumaP) & ~EdgeStopedByMinLumaP;
-		tsr_ushort2 StopByMinLumaN = asuint(SampleLumaN - MinLuma) & ~EdgeStopedByMaxLumaN;
-		tsr_ushort2 StopByMaxLumaN = asuint(MaxLuma - SampleLumaN) & ~EdgeStopedByMinLumaN;
+		tsr_ushort2 StopByMinLumaP = asuintX(SampleLumaP - MinLuma) & ~EdgeStopedByMaxLumaP;
+		tsr_ushort2 StopByMaxLumaP = asuintX(MaxLuma - SampleLumaP) & ~EdgeStopedByMinLumaP;
+		tsr_ushort2 StopByMinLumaN = asuintX(SampleLumaN - MinLuma) & ~EdgeStopedByMaxLumaN;
+		tsr_ushort2 StopByMaxLumaN = asuintX(MaxLuma - SampleLumaN) & ~EdgeStopedByMinLumaN;
 
 		EdgeStopedByMinLumaP |= StopByMinLumaP;
 		EdgeStopedByMaxLumaP |= StopByMaxLumaP;
@@ -279,6 +348,94 @@ void BrowseNeighborhoodBilinearOptimized(
 	bEdgeStopedByDecrementN = or(and(bEdgeMoreLuminous, bEdgeStopedByMinLumaN), and(!bEdgeMoreLuminous, bEdgeStopedByMaxLumaN));
 
 
+} // BrowseNeighborhoodBilinearRDNAOptimized
+
+void BrowseNeighborhoodBilinearOptimized(
+	const uint Iterations,
+	Texture2D<float> Texture,
+	float InputC,
+	float EdgeLuma,
+	tsr_short2 InputPixelPos,
+	tsr_short2 BrowseDirection,
+	tsr_short2 EdgeSide,
+	out tsr_ushort EdgeLengthP,
+	out tsr_ushort EdgeLengthN,
+	out bool bEdgeStopedByIncrementP,
+	out bool bEdgeStopedByIncrementN,
+	out bool bEdgeStopedByDecrementP,
+	out bool bEdgeStopedByDecrementN)
+{
+	EdgeLengthP = 0;
+	EdgeLengthN = 0;
+	bEdgeStopedByIncrementP = false;
+	bEdgeStopedByIncrementN = false;
+	bEdgeStopedByDecrementP = false;
+	bEdgeStopedByDecrementN = false;
+
+	const float BilinearInterp = 0.5;
+
+	float2 fBrowseDirection = float2(BrowseDirection) * InputInfo_ExtentInverse;
+	float2 KernelUV = (float2(InputPixelPos) + 0.5 + float2(EdgeSide) * BilinearInterp) * InputInfo_ExtentInverse;
+
+	float MergedEdgeLuma = BilinearInterp * (EdgeLuma + InputC);
+		
+	float MinLuma = MergedEdgeLuma - float(BilinearInterp * 0.5) * abs(EdgeLuma - InputC);
+	float MaxLuma = MergedEdgeLuma + float(BilinearInterp * 0.5) * abs(EdgeLuma - InputC);
+
+	bool bEdgeMoreLuminous = EdgeLuma > InputC;
+		
+	uint EdgeStopedByMinLumaP = 0;
+	uint EdgeStopedByMinLumaN = 0;
+		
+	uint EdgeStopedByMaxLumaP = 0;
+	uint EdgeStopedByMaxLumaN = 0;
+		
+	bool bStop = false;
+		
+	const uint SignShift = 31;
+	const uint SignMask = 0x80000000;
+			
+	EdgeLengthP = tsr_ushort(Iterations);
+	EdgeLengthN = tsr_ushort(Iterations);
+		
+	UNROLL
+	for (uint i = 0; i < Iterations; i++)
+	{
+		float2 NeighborOffset = fBrowseDirection * float(i + 1);
+
+		float SampleLumaP = BilinearSampleAAInput(Texture, KernelUV,  NeighborOffset, /* bIsNegativeDirection = */ false);
+		float SampleLumaN = BilinearSampleAAInput(Texture, KernelUV, -NeighborOffset, /* bIsNegativeDirection = */ true);
+				
+		//bEdgeStopedByMinLuma = bEdgeStopedByMinLuma || bool2(SampleLuma < MinLuma) && !bEdgeStopedByMaxLuma;
+		//bEdgeStopedByMaxLuma = bEdgeStopedByMaxLuma || bool2(SampleLuma > MaxLuma) && !bEdgeStopedByMinLuma;
+		uint StopByMinLumaP = asuint(SampleLumaP - MinLuma) & ~EdgeStopedByMaxLumaP;
+		uint StopByMaxLumaP = asuint(MaxLuma - SampleLumaP) & ~EdgeStopedByMinLumaP;
+		uint StopByMinLumaN = asuint(SampleLumaN - MinLuma) & ~EdgeStopedByMaxLumaN;
+		uint StopByMaxLumaN = asuint(MaxLuma - SampleLumaN) & ~EdgeStopedByMinLumaN;
+
+		EdgeStopedByMinLumaP |= StopByMinLumaP;
+		EdgeStopedByMaxLumaP |= StopByMaxLumaP;
+		EdgeStopedByMinLumaN |= StopByMinLumaN;
+		EdgeStopedByMaxLumaN |= StopByMaxLumaN;
+				
+		//EdgeLength += (bEdgeStopedByMaxLuma || bEdgeStopedByMinLuma) ? 0 : 1;
+		EdgeLengthP -= tsr_ushort((EdgeStopedByMaxLumaP | EdgeStopedByMinLumaP) >> SignShift);
+		EdgeLengthN -= tsr_ushort((EdgeStopedByMaxLumaN | EdgeStopedByMinLumaN) >> SignShift);
+	}
+		
+	bool bEdgeStopedByMinLumaP = (EdgeStopedByMinLumaP & SignMask) != 0;
+	bool bEdgeStopedByMinLumaN = (EdgeStopedByMinLumaN & SignMask) != 0;
+		
+	bool bEdgeStopedByMaxLumaP = (EdgeStopedByMaxLumaP & SignMask) != 0;
+	bool bEdgeStopedByMaxLumaN = (EdgeStopedByMaxLumaN & SignMask) != 0;
+		
+	bEdgeStopedByIncrementP = or(and(bEdgeMoreLuminous, bEdgeStopedByMaxLumaP), and(!bEdgeMoreLuminous, bEdgeStopedByMinLumaP));
+	bEdgeStopedByIncrementN = or(and(bEdgeMoreLuminous, bEdgeStopedByMaxLumaN), and(!bEdgeMoreLuminous, bEdgeStopedByMinLumaN));
+		
+	bEdgeStopedByDecrementP = or(and(bEdgeMoreLuminous, bEdgeStopedByMinLumaP), and(!bEdgeMoreLuminous, bEdgeStopedByMaxLumaP));
+	bEdgeStopedByDecrementN = or(and(bEdgeMoreLuminous, bEdgeStopedByMinLumaN), and(!bEdgeMoreLuminous, bEdgeStopedByMaxLumaN));
+
+
 } // BrowseNeighborhoodBilinearRDNAOptimized
 
 tsr_half ComputeDistanceToEdge(
@@ -328,3 +485,140 @@ tsr_half ComputeDistanceToEdge(
 
 	return TexelCenterOffsetToEdge;
 } // ComputeDistanceToEdge()
+
+tsr_half2 ComputeReprojectionBoundary(
+	const uint BrowsingIterations,
+	tsr_short2 EdgeSide,
+	tsr_ushort EdgeLengthP,
+	tsr_ushort EdgeLengthN,
+	bool bEdgeStopedByIncrementP,
+	bool bEdgeStopedByIncrementN,
+	bool bEdgeStopedByDecrementP,
+	bool bEdgeStopedByDecrementN)
+{
+	const tsr_half fMaxEdgeLength = tsr_half(BrowsingIterations + 1);
+	const tsr_half fMaxEdgeInvLength = rcp(fMaxEdgeLength);
+
+	tsr_half fEdgeLengthN = tsr_half(EdgeLengthN);
+	tsr_half fEdgeLengthP = tsr_half(EdgeLengthP);
+	tsr_half fEdgeLength = tsr_half(1.0) + fEdgeLengthN + fEdgeLengthP;
+	tsr_half fEdgeInvLength = rcp(fEdgeLength);
+
+	tsr_half MaxDistanceToEdge = tsr_half(0.0);
+	tsr_half MaxEdgeInclination = tsr_half(0.0);
+	{
+		FLATTEN
+		if (bEdgeStopedByDecrementP && bEdgeStopedByDecrementN)
+		{
+			MaxDistanceToEdge = (min(fEdgeLengthP, fEdgeLengthN) + tsr_half(0.5)) * fMaxEdgeInvLength;
+			MaxEdgeInclination = select(fEdgeLengthP < fEdgeLengthN, -fMaxEdgeInvLength, tsr_half(0.0));
+			MaxEdgeInclination = select(fEdgeLengthN < fEdgeLengthP, +fMaxEdgeInvLength, MaxEdgeInclination);
+		}
+		else if (bEdgeStopedByDecrementP)
+		{
+			MaxDistanceToEdge = (fEdgeLengthP + tsr_half(0.5)) * max(fEdgeInvLength, fMaxEdgeInvLength);
+			MaxEdgeInclination = -max(fEdgeInvLength, fMaxEdgeInvLength);
+		}
+		else if (bEdgeStopedByDecrementN)
+		{
+			MaxDistanceToEdge = (fEdgeLengthN + tsr_half(0.5)) * max(fEdgeInvLength, fMaxEdgeInvLength);
+			MaxEdgeInclination = +max(fEdgeInvLength, fMaxEdgeInvLength);
+		}
+		else
+		{
+			MaxDistanceToEdge = tsr_half(1.0);
+			MaxEdgeInclination = tsr_half(0.0);
+		}
+	}
+
+	tsr_half MinDistanceToEdge = tsr_half(0.0);
+	tsr_half MinEdgeInclination = tsr_half(0.0);
+	{
+		FLATTEN
+		if (bEdgeStopedByIncrementP && bEdgeStopedByIncrementN)
+		{
+			MinDistanceToEdge = 1.0 - (min(fEdgeLengthP, fEdgeLengthN) + tsr_half(0.5)) * fMaxEdgeInvLength;
+			MinEdgeInclination = select(fEdgeLengthP < fEdgeLengthN, +fMaxEdgeInvLength, tsr_half(0.0));
+			MinEdgeInclination = select(fEdgeLengthN < fEdgeLengthP, -fMaxEdgeInvLength, MinEdgeInclination);
+		}
+		else if (bEdgeStopedByIncrementP)
+		{
+			MinDistanceToEdge = 1.0 - (fEdgeLengthP + tsr_half(0.5)) * max(fMaxEdgeInvLength, fEdgeInvLength);
+			MinEdgeInclination = +max(fEdgeInvLength, fMaxEdgeInvLength);
+		}
+		else if (bEdgeStopedByIncrementN)
+		{
+			MinDistanceToEdge = 1.0 - (fEdgeLengthN + tsr_half(0.5)) * max(fMaxEdgeInvLength, fEdgeInvLength);
+			MinEdgeInclination = -max(fEdgeInvLength, fMaxEdgeInvLength);
+		}
+		else
+		{
+			MinDistanceToEdge = tsr_half(0.0);
+			MinEdgeInclination = tsr_half(0.0);
+		}
+		
+		MinDistanceToEdge = max(MinDistanceToEdge, rcp(tsr_half(1023.0)));
+		//MinDistanceToEdge = max(MinDistanceToEdge, rcp(tsr_half(8.0)));
+	}
+
+	// Merge the distance to edge computation.
+	tsr_half DistanceToEdge;
+	tsr_half EdgeInclination;
+	#if 1
+	{
+		//tsr_half DilateAmount = dot(tsr_half2(EdgeSide), tsr_half2(InputJitter));
+		tsr_half DilateAmount = tsr_half(0.5);
+		//DilateAmount = max(DilateAmount, rcp(tsr_half(8.0)));
+		//DilateAmount += tsr_half(0.5);
+
+		// Absolute mid dilate with spatial anti-aliasing.
+		DistanceToEdge = max(MinDistanceToEdge, min(MaxDistanceToEdge, DilateAmount));
+		EdgeInclination = tsr_half(0.0);
+		EdgeInclination = select(MaxDistanceToEdge <= tsr_half(DilateAmount), MaxEdgeInclination, EdgeInclination);
+		EdgeInclination = select(MinDistanceToEdge >  tsr_half(DilateAmount), MinEdgeInclination, EdgeInclination);
+	}
+	#elif 0
+	{
+		// Absolute minimum dilate with spatial anti-aliasing.
+		DistanceToEdge = MinDistanceToEdge;
+		EdgeInclination = MinEdgeInclination;
+	}
+	#elif 0
+	{
+		// Full dilate at rendering pixel granularity, but chamfer the extremity.
+		DistanceToEdge = MaxDistanceToEdge;
+		EdgeInclination = MaxEdgeInclination;
+	}
+	#else
+	{
+		// Full dilate at rendering pixel granularity.
+		DistanceToEdge = tsr_half(1.0);
+		EdgeInclination = tsr_half(0.0);
+	}
+	#endif
+
+	#if 1
+	{
+		tsr_half2 fEdgeSide = tsr_half2(EdgeSide);
+		tsr_half2 fDirectionP = tsr_half2(abs(fEdgeSide.y), abs(fEdgeSide.x));
+
+		tsr_half2 EdgeDirection = fEdgeSide + fDirectionP * EdgeInclination;
+		EdgeDirection *= rsqrt(dot(EdgeDirection, EdgeDirection));
+
+		tsr_half2 K = saturate(tsr_half2(sign(EdgeDirection)));
+		tsr_half2 E = tsr_half2(0.5, 0.5) + fEdgeSide * (0.5 - DistanceToEdge);
+
+		tsr_half2 EK = K - E;
+		tsr_half KFromEdge = dot(EK, EdgeDirection);
+		tsr_half2 L = K - KFromEdge * EdgeDirection;
+
+		tsr_half2 ReprojectionBoundary = K - L;
+		return ReprojectionBoundary;
+	}
+	#else
+	{
+		tsr_half2 ReprojectionBoundary = tsr_half2(EdgeSide) * DistanceToEdge;
+		return ReprojectionBoundary;
+	}
+	#endif
+} // ComputeReprojectionBoundary()
diff --git a/Engine/Shaders/Private/TemporalSuperResolution/TSRUpdateHistory.usf b/Engine/Shaders/Private/TemporalSuperResolution/TSRUpdateHistory.usf
--- a/Engine/Shaders/Private/TemporalSuperResolution/TSRUpdateHistory.usf
+++ b/Engine/Shaders/Private/TemporalSuperResolution/TSRUpdateHistory.usf
@@ -1,6 +1,7 @@
 // Copyright Epic Games, Inc. All Rights Reserved.
 
 #include "TSRCommon.ush"
+#include "TSRReprojectionField.ush"
 
 
 //------------------------------------------------------- DEBUG
@@ -24,7 +25,6 @@
 
 #elif DIM_UPDATE_QUALITY == 3 // Epic
 	#define CONFIG_SAMPLES 6
-	#define CONFIG_INTERPOLATE_VELOCITY 1
 	#define CONFIG_REJECTION_ANTI_ALIASING 1
 
 #else
@@ -34,6 +34,7 @@
 #define CONFIG_COMPOSE_TRANSLUCENCY 0
 #define CONFIG_REJECT_TRANSLUCENCY_ON_RESPONSIVE_AA_ONLY 0
 #define CONFIG_HISTORY_RESURRECTION 1
+#define CONFIG_REPROJECTION_FIELD 1
 
 /** Whether lens distortion is supported. */
 #define CONFIG_LENS_DISTORTION TSR_SUPPORT_LENS_DISTORTION
@@ -62,10 +63,6 @@
 	#define PLATFORM_SPECIFIC_ISOLATE ISOLATE
 #endif
 
-#ifndef CONFIG_INTERPOLATE_VELOCITY
-	#define CONFIG_INTERPOLATE_VELOCITY 0
-#endif
-
 
 //------------------------------------------------------- CONSTANTS
 
@@ -91,7 +88,11 @@ Texture2D<tsr_half4> InputSceneTranslucencyTexture;
 
 Texture2D<tsr_half4> HistoryRejectionTexture;
 
-Texture2D<float2> DilatedVelocityTexture;
+#if CONFIG_REPROJECTION_FIELD
+	Texture2D<uint> ReprojectionBoundaryTexture;
+	Texture2D<uint> ReprojectionJacobianTexture;
+#endif
+Texture2D<uint> ReprojectionVectorTexture;
 Texture2D<tsr_ushort2> AntiAliasingTexture;
 
 #if CONFIG_LENS_DISTORTION
@@ -120,6 +121,9 @@ float PrevFrameIndex;
 #if CONFIG_LENS_DISTORTION
 	uint bLensDistortion;
 #endif
+#if CONFIG_REPROJECTION_FIELD
+	uint bReprojectionField;
+#endif
 uint bGenerateOutputMip1;
 uint bGenerateOutputMip2;
 uint bGenerateOutputMip3;
@@ -297,14 +301,14 @@ void ComputeInputKernelSamplePosition(
 }
 
 template<typename T>
-T BilinearSampleColorHistory(Texture2D<T> Texture, float2 SampleUV)
+T BilinearSampleColorHistory(Texture2DArray<T> Texture, float3 SampleUV)
 #if 1
 {
 	return Texture.SampleLevel(GlobalBilinearClampedSampler, SampleUV, 0);
 }
 #else
 {
-	FBilinearSampleInfos BilinearInter = GetBilinearSampleLevelInfos(SampleUV, PrevHistoryInfo_Extent, PrevHistoryInfo_ExtentInverse);
+	FBilinearSampleInfos BilinearInter = GetBilinearSampleLevelInfos(SampleUV.xy, PrevHistoryInfo_Extent, PrevHistoryInfo_ExtentInverse);
 
 	T Return = 0.0;
 
@@ -315,16 +319,16 @@ T BilinearSampleColorHistory(Texture2D<T> Texture, float2 SampleUV)
 		uint2 PixelPos = GetSamplePixelCoord(BilinearInter, i);
 		PixelPos = fastClamp(PixelPos, PrevHistoryInfo_ViewportMin, PrevHistoryInfo_ViewportMax - 1);
 
-		T RawSample = Texture[PixelPos];
+		T RawSample = Texture[uint3(PixelPos, SampleUV.z)];
 		
 		if (View.GeneralPurposeTweak == 1.0)
-			RawSample.rgb *= HdrWeight4(RawSample.rgb);
+			RawSample.rgb *= tsr_half(HdrWeight4(RawSample.rgb));
 
-		Return += RawSample * BilinearWeight;
+		Return += RawSample * tsr_half(BilinearWeight);
 	}
 	
 	if (View.GeneralPurposeTweak == 1.0)
-		Return.rgb *= HdrWeightInvY(Luma4(Return.rgb));
+		Return.rgb *= tsr_half(HdrWeightInvY(Luma4(Return.rgb)));
 
 	return Return;
 }
@@ -472,6 +476,46 @@ float2 ApplyDisplacementTextureOnScreenPos(Texture2D<float2> DisplacementTexture
 }
 #endif
 
+#if CONFIG_REPROJECTION_FIELD
+
+void FetchVelocityDilateBoundary(
+	tsr_short2x2 InputPixelPos,
+	out uint EncodedReprojectionBoundary[DPV_PIXEL_PER_LANE],
+	out bool bApplyReprojectionBoundary)
+{
+	BRANCH
+	if (bReprojectionField)
+	PLATFORM_SPECIFIC_ISOLATE
+	{
+		UNROLL_N(DPV_PIXEL_PER_LANE)
+		for (uint PixelId = 0; PixelId < DPV_PIXEL_PER_LANE; PixelId++)
+		{
+			tsr_ushort2 LocalInputPixelPos = dpv_access_pixel(InputPixelPos, PixelId);
+			
+			EncodedReprojectionBoundary[PixelId] = ReprojectionBoundaryTexture[LocalInputPixelPos];
+		}
+		
+		const uint EncodedReprojectionNoBoundary = EncodeReprojectionBoundary(/* ReprojectionOffset = */ tsr_short2(0, 0), /* ReprojectionBoundary = */ kFullDilateBoundary);
+		bApplyReprojectionBoundary = false;
+		UNROLL_N(DPV_PIXEL_PER_LANE)
+		for (uint PixelId = 0; PixelId < DPV_PIXEL_PER_LANE; PixelId++)
+		{
+			bApplyReprojectionBoundary |= EncodedReprojectionBoundary[PixelId] != EncodedReprojectionNoBoundary;
+		}
+	}
+	else
+	{
+		UNROLL_N(DPV_PIXEL_PER_LANE)
+		for (uint PixelId = 0; PixelId < DPV_PIXEL_PER_LANE; PixelId++)
+		{
+			EncodedReprojectionBoundary[PixelId] = 0;
+		}
+		bApplyReprojectionBoundary = false;
+	}
+} // FetchVelocityDilateBoundary
+
+#endif // CONFIG_REPROJECTION_FIELD
+
 
 //------------------------------------------------------- ENTRY POINT
 
@@ -532,7 +576,7 @@ void MainCS(
 	{
 		ScreenPos = ApplyScreenTransform(float2x2(HistoryPixelPos), HistoryPixelPosToScreenPos);
 	
-		// Pixel coordinate of the center of output pixel O in the input viewport.
+	// Pixel coordinate of the center of output pixel O in the input viewport.
 		InputPPCo = ApplyScreenTransform(float2x2(HistoryPixelPos), HistoryPixelPosToInputPPCo);
 	}
 
@@ -542,6 +586,13 @@ void MainCS(
 	tsr_short2x2 InputPixelPos = ClampPixelOffset(
 		tsr_short2x2(InputPPCo),
 		InputPixelPosMin, InputPixelPosMax);
+		
+	// Fetch the reprojection anti-aliasing.
+	#if CONFIG_REPROJECTION_FIELD
+	uint EncodedReprojectionBoundary[DPV_PIXEL_PER_LANE];
+	bool bApplyReprojectionBoundary;
+	FetchVelocityDilateBoundary(InputPixelPos, /* out */ EncodedReprojectionBoundary, /* out */ bApplyReprojectionBoundary);
+	#endif // CONFIG_REPROJECTION_FIELD
 
 	// Fetch reprojection-related information.
 	float2x2 PrevScreenPos = ScreenPos;
@@ -563,31 +614,62 @@ void MainCS(
 
 	PLATFORM_SPECIFIC_ISOLATE
 	{
-		#if CONFIG_INTERPOLATE_VELOCITY
-			float2 RawEncodedVelocityNeighborhood[4][DPV_PIXEL_PER_LANE];
-		#else
-			float2 RawEncodedVelocity[DPV_PIXEL_PER_LANE];
-		#endif
-
+		uint RawEncodedReprojectionVector[DPV_PIXEL_PER_LANE];
 		tsr_half4 RawHistoryRejection[DPV_PIXEL_PER_LANE];
 		
 		#if CONFIG_REJECTION_ANTI_ALIASING
 			tsr_ushort2 RawEncodedInputTexelOffset[DPV_PIXEL_PER_LANE];
 		#endif
 
+		#if CONFIG_REPROJECTION_FIELD
+			uint RawEncodedReprojectionJacobian[DPV_PIXEL_PER_LANE];
+			tsr_half2 JacobianCoordinate[DPV_PIXEL_PER_LANE];
+		#endif
+
 		// Issue overlapped texture fetches
 		PLATFORM_SPECIFIC_ISOLATE
 		{
+			tsr_half2x2 dInputKO = tsr_half2x2(InputPPCo - InputPPCk);
+
 			UNROLL_N(DPV_PIXEL_PER_LANE)
 			for (uint PixelId = 0; PixelId < DPV_PIXEL_PER_LANE; PixelId++)
 			{
 				tsr_ushort2 LocalInputPixelPos = dpv_access_pixel(InputPixelPos, PixelId);
+
+				#if CONFIG_REPROJECTION_FIELD
+				tsr_short2 BoundaryDilateOffset;
+				BRANCH
+				if (bApplyReprojectionBoundary)
+				{
+					tsr_short2 ReprojectionOffset  = DecodeReprojectionOffset(EncodedReprojectionBoundary[PixelId]);
+					tsr_half2 ReprojectionBoundary = DecodeReprojectionBoundary(EncodedReprojectionBoundary[PixelId]);
+
+					bool bHistoryPixelWithinOffsetBoundary = IsHistoryPixelWithinOffsetBoundary(
+						dpv_access_pixel(dInputKO, PixelId),
+						ReprojectionBoundary,
+						1.0 / InputToHistoryFactor);
+					BoundaryDilateOffset = select(bHistoryPixelWithinOffsetBoundary, ReprojectionOffset, -ReprojectionOffset);
+					
+					//Debug[0][0][PixelId] = bHistoryPixelWithinOffsetBoundary && any(ReprojectionOffset != tsr_short(0));
+					//Debug[0][1][PixelId] = ((LocalInputPixelPos.x ^ LocalInputPixelPos.y) & 0x1) == 0;
+				}
+				else
+				{
+					BoundaryDilateOffset = tsr_short2(0, 0);
+				}
 				
-				#if !CONFIG_INTERPOLATE_VELOCITY
-					RawEncodedVelocity[PixelId]   = DilatedVelocityTexture[LocalInputPixelPos];
+				JacobianCoordinate[PixelId] = dpv_access_pixel(dInputKO, PixelId) - tsr_half2(BoundaryDilateOffset);
+				tsr_ushort2 LocalInputPixelPosWithReprojectionAA = LocalInputPixelPos + BoundaryDilateOffset;
+				#else
+				tsr_ushort2 LocalInputPixelPosWithReprojectionAA = LocalInputPixelPos;
+				#endif
+				
+				RawEncodedReprojectionVector[PixelId]   = ReprojectionVectorTexture[LocalInputPixelPosWithReprojectionAA];
+				#if CONFIG_REPROJECTION_FIELD
+					RawEncodedReprojectionJacobian[PixelId] = ReprojectionJacobianTexture[LocalInputPixelPosWithReprojectionAA];
 				#endif
 
-				RawHistoryRejection[PixelId] = HistoryRejectionTexture[LocalInputPixelPos];
+				RawHistoryRejection[PixelId] = HistoryRejectionTexture[LocalInputPixelPosWithReprojectionAA];
 				
 				#if CONFIG_REJECTION_ANTI_ALIASING
 				{
@@ -595,74 +677,46 @@ void MainCS(
 				}
 				#endif
 			}
-
-			// Vector in pixel between pixel K -> O.
-			tsr_half2x2 dInputKO = tsr_half2x2(InputPPCo - InputPPCk);
-		
-			// Fetch the 2x2 neighborhood of the velocity
-			#if CONFIG_INTERPOLATE_VELOCITY
-			UNROLL_N(DPV_PIXEL_PER_LANE)
-			for (uint PixelId2 = 0; PixelId2 < DPV_PIXEL_PER_LANE; PixelId2++)
-			{
-				tsr_short2 LocalInputPixelPos0 = dpv_access_pixel(InputPixelPos, PixelId2);
-				tsr_short2 LocalInputPixelPos1 = LocalInputPixelPos0 + tsr_short2(sign(dpv_access_pixel(dInputKO, PixelId2)));
-				LocalInputPixelPos1 = ClampPixelOffset(LocalInputPixelPos1, InputPixelPosMin, InputPixelPosMax);
-
-				RawEncodedVelocityNeighborhood[0][PixelId2] = DilatedVelocityTexture[tsr_short2(LocalInputPixelPos0.x, LocalInputPixelPos0.y)];
-				RawEncodedVelocityNeighborhood[1][PixelId2] = DilatedVelocityTexture[tsr_short2(LocalInputPixelPos1.x, LocalInputPixelPos0.y)];
-				RawEncodedVelocityNeighborhood[2][PixelId2] = DilatedVelocityTexture[tsr_short2(LocalInputPixelPos0.x, LocalInputPixelPos1.y)];
-				RawEncodedVelocityNeighborhood[3][PixelId2] = DilatedVelocityTexture[tsr_short2(LocalInputPixelPos1.x, LocalInputPixelPos1.y)];
-			}
-			#endif
 		}
 		
 		// Process texture fetches.
 		PLATFORM_SPECIFIC_ISOLATE
 		{
-			// Vector in pixel between pixel K -> O.
-			tsr_half2x2 dInputKO = tsr_half2x2(InputPPCo - InputPPCk);
-		
 			// Process velocity
 			{
-				float2x2 ScreenVelocity;
-				#if CONFIG_INTERPOLATE_VELOCITY
-				UNROLL_N(DPV_PIXEL_PER_LANE)
+				float2x2 ReprojectionVector;
+
 				for (uint PixelId = 0; PixelId < DPV_PIXEL_PER_LANE; PixelId++)
 				{
-					float2 V00 = DecodeVelocityFromTexture(float4(RawEncodedVelocityNeighborhood[0][PixelId], 0.0, 0.0)).xy;
-					float2 V10 = DecodeVelocityFromTexture(float4(RawEncodedVelocityNeighborhood[1][PixelId], 0.0, 0.0)).xy;
-					float2 V01 = DecodeVelocityFromTexture(float4(RawEncodedVelocityNeighborhood[2][PixelId], 0.0, 0.0)).xy;
-					float2 V11 = DecodeVelocityFromTexture(float4(RawEncodedVelocityNeighborhood[3][PixelId], 0.0, 0.0)).xy;
-
-					tsr_half2 BilinearInterp = abs(dpv_access_pixel(dInputKO, PixelId));
-					
-					float Bilinear00 = saturate(1.0 - BilinearInterp.x) * saturate(1.0 - BilinearInterp.y);
-					float Bilinear10 = saturate(0.0 + BilinearInterp.x) * saturate(1.0 - BilinearInterp.y);
-					float Bilinear01 = saturate(1.0 - BilinearInterp.x) * saturate(0.0 + BilinearInterp.y);
-					float Bilinear11 = saturate(0.0 + BilinearInterp.x) * saturate(0.0 + BilinearInterp.y);
+					#if CONFIG_REPROJECTION_FIELD
+					float2 ReprojectionScreenPosCorrection;
+					BRANCH
+					if (bReprojectionField)
+					{
+						uint EncodedReprojectionJacobian = RawEncodedReprojectionJacobian[PixelId];
+						tsr_half2x2 ReprojectionJacobian = DecodeReprojectionJacobian(EncodedReprojectionJacobian);
+						
+						float2 ReprojectionPixelPosCorrection = float2(mul(JacobianCoordinate[PixelId], ReprojectionJacobian));
+						//ReprojectionScreenPosCorrection = ReprojectionPixelPosCorrection * (float2(2, -2) * InputInfo_ViewportSizeInverse);
+						ReprojectionScreenPosCorrection = ReprojectionPixelPosCorrection * InputPixelPosToScreenPos.xy;
+					}
+					else
+					{
+						ReprojectionScreenPosCorrection = float(0.0);
+					}
+					#else // !CONFIG_REPROJECTION_FIELD
+					const float2 ReprojectionScreenPosCorrection = 0.0;
+					#endif
 					
-					float Bilateral00 = Bilinear00;
-					float Bilateral10 = Bilinear10 * ComputeScreenVelocityBilateralWeight(V10 - V00);
-					float Bilateral01 = Bilinear01 * ComputeScreenVelocityBilateralWeight(V01 - V00);
-					float Bilateral11 = Bilinear11 * ComputeScreenVelocityBilateralWeight(V11 - V00);
-
-					float NormalizeBilateral = rcp(Bilateral00 + Bilateral10 + Bilateral01 + Bilateral11);
+					uint EncodedReprojectionVector = RawEncodedReprojectionVector[PixelId];
+					float2 LocalReprojectionVector = DecodeReprojectionVector(EncodedReprojectionVector) + float2(ReprojectionScreenPosCorrection);
 
-					float2 LocalScreenVelocity = (V00 * Bilateral00 + V10 * Bilateral10 + V01 * Bilateral01 + V11 * Bilateral11) * NormalizeBilateral;
-					
-					ScreenVelocity[0][PixelId] = LocalScreenVelocity.x;
-					ScreenVelocity[1][PixelId] = LocalScreenVelocity.y;
+					ReprojectionVector[0][PixelId] = LocalReprojectionVector[0];
+					ReprojectionVector[1][PixelId] = LocalReprojectionVector[1];
 				}
-				#else // !CONFIG_INTERPOLATE_VELOCITY
-				{
-					ScreenVelocity = dpv_interleave_registers(
-						DecodeVelocityFromTexture(float4(RawEncodedVelocity[0], 0.0, 0.0)).xy,
-						DecodeVelocityFromTexture(float4(RawEncodedVelocity[1], 0.0, 0.0)).xy);
-				}
-				#endif
 
-				PrevScreenPos = ScreenPos - ScreenVelocity;
-				OutputPixelVelocity = tsr_half2(dpv_length(dpv_mul(ScreenVelocity, HistoryInfo_ViewportSize)));
+				PrevScreenPos = ScreenPos - ReprojectionVector;
+				OutputPixelVelocity = tsr_half2(dpv_length(dpv_mul(ReprojectionVector, HistoryInfo_ViewportSize)));
 			}
 			
 			// Unpacks bits masks
@@ -727,7 +781,7 @@ void MainCS(
 		bIsParallaxRejected,
 		/* out */ bIsOffScreen,
 		/* out */ bIsDisoccluded);
-				
+		
 	// Final post processing.
 	#if !CONFIG_CLAMP
 	{
@@ -793,7 +847,7 @@ void MainCS(
 						float FrameIndex = PrevFrameIndex;
 					#endif
 
-					RawHighFrequency[i][PixelId] = PrevHistoryColorTexture.SampleLevel(GlobalBilinearClampedSampler, float3(SampleUV, FrameIndex), 0);
+					RawHighFrequency[i][PixelId] = BilinearSampleColorHistory(PrevHistoryColorTexture, float3(SampleUV, FrameIndex));
 					RawMetadata[i][PixelId]      = PrevHistoryMetadataTexture.SampleLevel(GlobalBilinearClampedSampler, float3(SampleUV, FrameIndex), 0);
 				}
 
diff --git a/Engine/Shaders/Private/TemporalSuperResolution/TSRVisualize.usf b/Engine/Shaders/Private/TemporalSuperResolution/TSRVisualize.usf
--- a/Engine/Shaders/Private/TemporalSuperResolution/TSRVisualize.usf
+++ b/Engine/Shaders/Private/TemporalSuperResolution/TSRVisualize.usf
@@ -1,6 +1,7 @@
 // Copyright Epic Games, Inc. All Rights Reserved.
 
 #include "TSRCommon.ush"
+#include "TSRReprojectionField.ush"
 #if TSR_SUPPORT_LENS_DISTORTION
 	#include "../LensDistortion.ush"
 #endif
@@ -19,6 +20,7 @@ static tsr_half3 kRed    = tsr_half3(1.00, 0.10, 0.10);
 static tsr_half3 kOrange = tsr_half3(1.00, 0.75, 0.10);
 static tsr_half3 kGreen  = tsr_half3(0.10, 1.00, 0.10);
 static tsr_half3 kBlue   = tsr_half3(0.10, 0.10, 1.00);
+static tsr_half3 kCyan   = tsr_half3(0.10, 1.00, 1.00);
 static tsr_half3 kPink   = tsr_half3(1.00, 0.10, 0.50);
 
 
@@ -37,6 +39,7 @@ uint2 OutputViewRectMax;
 uint VisualizeId;
 uint bCanResurrectHistory;
 uint bCanSpatialAntiAlias;
+uint bReprojectionField;
 float MaxHistorySampleCount;
 float OutputToHistoryResolutionFractionSquare;
 float FlickeringFramePeriod;
@@ -50,7 +53,9 @@ Texture2D<tsr_half>  InputMoireLumaTexture;
 Texture2D<tsr_half4> InputSceneTranslucencyTexture;
 Texture2D<tsr_half3> SceneColorTexture;
 Texture2D<float2>    ClosestDepthTexture;
-Texture2D<float2>    DilatedVelocityTexture;
+Texture2D<uint>      ReprojectionBoundaryTexture;
+Texture2D<uint>      ReprojectionJacobianTexture;
+Texture2D<uint>      ReprojectionVectorTexture;
 Texture2D<tsr_ushort> IsMovingMaskTexture;
 Texture2D<tsr_half4> HistoryRejectionTexture;
 Texture2D<tsr_half4> MoireHistoryTexture;
@@ -61,6 +66,11 @@ Texture2D<tsr_half3> ResurrectedHistoryColorTexture;
 RWTexture2D<tsr_half4> Output;
 
 
+//------------------------------------------------------- DEBUG
+
+static float4 Debug;
+	
+
 //------------------------------------------------------- FUNCTION
 
 #if !TSR_SUPPORT_LENS_DISTORTION
@@ -95,6 +105,52 @@ tsr_half3 SampleTSROutputGray(FContext Ctx)
 	return GrayScaleColor(SceneColorTexture[Ctx.OutputPixelPos]);
 }
 
+struct FReprojectionFieldSample
+{
+	bool bIsOffScreen;
+	bool bIsFullDilate;
+	bool bAntiAliasReprojectionField;
+	tsr_short2 Offset;
+	tsr_half2 Boundary;
+	float2 Vector;
+	float2 PrevScreenPos;
+	tsr_half2x2 Jacobian;
+};
+
+FReprojectionFieldSample SampleReprojectionField(FContext Ctx)
+{
+	FReprojectionFieldSample Reprojection;
+
+	uint EncodedReprojectionBoundary = ReprojectionBoundaryTexture[Ctx.InputPixelPos];
+	Reprojection.Offset = DecodeReprojectionOffset(EncodedReprojectionBoundary);
+	Reprojection.Boundary = DecodeReprojectionBoundary(EncodedReprojectionBoundary);
+	Reprojection.bIsFullDilate = all(Reprojection.Boundary == kFullDilateBoundary);
+	Reprojection.bAntiAliasReprojectionField = any(Reprojection.Offset != tsr_short(0)) && !Reprojection.bIsFullDilate;
+
+	uint EncodedReprojectionVector   = ReprojectionVectorTexture[Ctx.InputPixelPos + Reprojection.Offset];
+	uint EncodedReprojectionJacobian = ReprojectionJacobianTexture[Ctx.InputPixelPos + Reprojection.Offset];
+
+	Reprojection.Vector   = DecodeReprojectionVector(EncodedReprojectionVector);
+	Reprojection.Jacobian = DecodeReprojectionJacobian(EncodedReprojectionJacobian);
+	
+	Reprojection.PrevScreenPos = Ctx.ScreenPos - Reprojection.Vector;
+
+	// Detect whether the history reprojection is successful
+	bool2 bIsOffScreenLocal;
+	bool2 bIsDisoccluded;
+	IsOffScreenOrDisoccluded(
+		bCameraCut,
+		dpv_interleave_mono_registers(Reprojection.PrevScreenPos),
+		/* bIsParallaxRejected = */ false,
+		/* out */ bIsOffScreenLocal,
+		/* out */ bIsDisoccluded);
+
+	Reprojection.bIsOffScreen = bIsOffScreenLocal[0];
+
+	return Reprojection;
+}
+
+
 //------------------------------------------------------- ENTRY POINT
 
 tsr_half3 DisplayHistorySampleCount(FContext Ctx)
@@ -283,6 +339,112 @@ tsr_half3 DisplayAntiFlickering(FContext Ctx)
 }
 
 
+//------------------------------------------------------- REPROJECTION FIELD
+
+tsr_half3 DisplayReprojectionFieldSummary(FContext Ctx)
+{
+	FReprojectionFieldSample Reprojection = SampleReprojectionField(Ctx);
+	
+	tsr_half JacobianVisibilityMultiplier = tsr_half(64.0);
+	//tsr_half JacobianVisibilityMultiplier = tsr_half(select(Reprojection.bAntiAliasReprojectionField, 2, 1)) * tsr_half(64.0);
+
+	tsr_half MaxJacobianX = max(abs(Reprojection.Jacobian[0][0]), abs(Reprojection.Jacobian[0][1]));
+	tsr_half MaxJacobianY = max(abs(Reprojection.Jacobian[1][0]), abs(Reprojection.Jacobian[1][1]));
+
+	tsr_half3 Color;
+	Color = kWhite;
+	Color.g = lerp(Color.g, tsr_half(1.0), saturate(MaxJacobianX * JacobianVisibilityMultiplier));
+	Color.b = lerp(Color.b, tsr_half(1.0), saturate(MaxJacobianY * JacobianVisibilityMultiplier));
+	Color = select(Reprojection.bIsOffScreen, kPink, Color);
+	//Color = select(Reprojection.bAntiAliasReprojectionField, kGreen, Color);
+	Color = select(!bReprojectionField, kPink, Color);
+	return SampleTSROutputGray(Ctx) * Color;
+}
+
+tsr_half3 DisplayReprojectionFieldOffset(FContext Ctx)
+{
+	FReprojectionFieldSample Reprojection = SampleReprojectionField(Ctx);
+	
+	tsr_half3 Color = kWhite;
+	Color.gb = tsr_half2((float2(Reprojection.Offset) * 0.5 + 0.5) * 0.67 + 0.33);
+	Color = select(all(Reprojection.Offset == tsr_short(0)), select(Reprojection.bIsOffScreen, kPink, kWhite), Color);
+	Color = select(!bReprojectionField, kPink, Color);
+	return SampleTSROutputGray(Ctx) * Color;
+}
+
+tsr_half3 DisplayReprojectionOffsetCoverage(FContext Ctx)
+{
+	FReprojectionFieldSample Reprojection = SampleReprojectionField(Ctx);
+	
+	tsr_half ReprojectionBoundaryCoverage = ComputeReprojectionBoundaryCoverage(Reprojection.Boundary);
+
+	tsr_half3 Color;
+	Color = lerp(kRed, kGreen, ReprojectionBoundaryCoverage);
+	Color = select(Reprojection.bIsFullDilate, kCyan, Color);
+	Color = select(all(Reprojection.Offset == tsr_short(0)), select(Reprojection.bIsOffScreen, kPink, kWhite), Color);
+	Color = select(!bReprojectionField, kPink, Color);
+	return SampleTSROutputGray(Ctx) * Color;
+}
+
+tsr_half3 DisplayReprojectionNullJacobian(FContext Ctx)
+{
+	FReprojectionFieldSample Reprojection = SampleReprojectionField(Ctx);
+	
+	tsr_half MaxJacobianX = max(abs(Reprojection.Jacobian[0][0]), abs(Reprojection.Jacobian[0][1]));
+	tsr_half MaxJacobianY = max(abs(Reprojection.Jacobian[1][0]), abs(Reprojection.Jacobian[1][1]));
+	tsr_half MaxJacobian = max(MaxJacobianX, MaxJacobianY);
+	
+	bool bIsNullVector = all(abs(Reprojection.Vector) < (0.5 / 65535.0));
+	bool bIsNullJacobian = MaxJacobian < pow(tsr_half(1u << kVelocityJacobianBits), tsr_half(-2.0));
+
+	tsr_half3 Color = select(bIsNullJacobian && !bIsNullVector, kOrange, kWhite);
+	Color = select(Reprojection.bIsOffScreen || !bReprojectionField, kPink, Color);
+	return SampleTSROutputGray(Ctx) * Color;
+}
+
+tsr_half3 DisplayReprojectionClampedJacobian(FContext Ctx)
+{
+	FReprojectionFieldSample Reprojection = SampleReprojectionField(Ctx);
+	
+	tsr_half MaxJacobianX = max(abs(Reprojection.Jacobian[0][0]), abs(Reprojection.Jacobian[0][1]));
+	tsr_half MaxJacobianY = max(abs(Reprojection.Jacobian[1][0]), abs(Reprojection.Jacobian[1][1]));
+	tsr_half MaxJacobian = max(MaxJacobianX, MaxJacobianY);
+
+	tsr_half3 Color;
+	Color = select(Reprojection.bIsOffScreen, kPink, kWhite);
+	Color = select(MaxJacobian > kVelocityJacobianRange - rcp(tsr_half(1u << kVelocityJacobianBits)), kRed, Color);
+	Color = select(!bReprojectionField, kPink, Color);
+	return SampleTSROutputGray(Ctx) * Color;
+}
+
+tsr_half3 DisplayReprojectionFieldDilatedJacobian(FContext Ctx)
+{
+	FReprojectionFieldSample Reprojection = SampleReprojectionField(Ctx);
+	
+	tsr_half ReprojectionUpscaleFactor = ComputeReprojectionUpscaleFactorFromJacobian(Reprojection.Jacobian);
+
+	tsr_half3 Color;
+	Color = kWhite;
+	Color = lerp(Color,   kRed, saturate( log2(ReprojectionUpscaleFactor)));
+	Color = lerp(Color, kGreen, saturate(-log2(ReprojectionUpscaleFactor)));
+	Color = select(Reprojection.bIsOffScreen || !bReprojectionField, kPink, Color);
+	return SampleTSROutputGray(Ctx) * Color;
+}
+
+tsr_half3 DisplayReprojectionFieldAA(FContext Ctx)
+{
+	FReprojectionFieldSample Reprojection = SampleReprojectionField(Ctx);
+	
+	tsr_half ReprojectionUpscaleFactor = ComputeReprojectionUpscaleFactorFromJacobian(Reprojection.Jacobian);
+
+	tsr_half3 Color;
+	Color = select(Reprojection.bIsOffScreen, kPink, kWhite);
+	Color = select(Reprojection.bAntiAliasReprojectionField, kGreen, Color);
+	Color = select(!bReprojectionField, kPink, Color);
+	return SampleTSROutputGray(Ctx) * Color;
+}
+
+
 //------------------------------------------------------- ENTRY POINT
 
 [numthreads(TILE_SIZE * TILE_SIZE, 1, 1)]
@@ -292,7 +454,7 @@ void MainCS(
 {
 	uint GroupWaveIndex = GetGroupWaveIndex(GroupThreadIndex, /* GroupSize = */ TILE_SIZE * TILE_SIZE);
 
-	float4 Debug = 0.0;
+	Debug = 0.0;
 	
 	FContext Ctx;
 	Ctx.OutputPixelPos = (
@@ -340,6 +502,34 @@ void MainCS(
 	{
 		OutputColor = DisplayAntiFlickering(Ctx);
 	}
+	else if (VisualizeId == 8)
+	{
+		OutputColor = DisplayReprojectionFieldSummary(Ctx);
+	}
+	else if (VisualizeId == 9)
+	{
+		OutputColor = DisplayReprojectionFieldOffset(Ctx);
+	}
+	else if (VisualizeId == 10)
+	{
+		OutputColor = DisplayReprojectionOffsetCoverage(Ctx);
+	}
+	else if (VisualizeId == 11)
+	{
+		OutputColor = DisplayReprojectionFieldAA(Ctx);
+	}
+	else if (VisualizeId == 12)
+	{
+		OutputColor = DisplayReprojectionNullJacobian(Ctx);
+	}
+	else if (VisualizeId == 13)
+	{
+		OutputColor = DisplayReprojectionClampedJacobian(Ctx);
+	}
+	else if (VisualizeId == 14)
+	{
+		OutputColor = DisplayReprojectionFieldDilatedJacobian(Ctx);
+	}
 
 	{
 		tsr_short2 OutputPixelPos = InvalidateOutputPixelPos(Ctx.OutputPixelPos, OutputViewRectMax);
diff --git a/Engine/Source/Runtime/Renderer/Private/PostProcess/TemporalSuperResolution.cpp b/Engine/Source/Runtime/Renderer/Private/PostProcess/TemporalSuperResolution.cpp
--- a/Engine/Source/Runtime/Renderer/Private/PostProcess/TemporalSuperResolution.cpp
+++ b/Engine/Source/Runtime/Renderer/Private/PostProcess/TemporalSuperResolution.cpp
@@ -276,6 +276,27 @@ TAutoConsoleVariable<int32> CVarTSRAsyncCompute(
 	TEXT(" 3: Run all passes on async compute;"),
 	ECVF_RenderThreadSafe);
 
+TAutoConsoleVariable<int32> CVarTSRReprojectionField(
+	TEXT("r.TSR.ReprojectionField"), 0,
+	TEXT("Enables TSR's reprojection field for higher reprojection vector upscale and dilate quality (Experimental, disabled by default).\n")
+	TEXT("\n")
+	TEXT("When the reprojection fields is enabled, it dilates the reprojection vector by half spatially ")
+	TEXT("anti-aliased rendering pixel from the depth buffer, instead by a full rendering pixel ")
+	TEXT("in dilate velocity pass. This allows hide the rendering resolution due whenever velocity buffer ends up extruding some ")
+	TEXT("object to edges, for instance when rotating. This come at the cost of spatial anti-aliasing in the DilateVelocity pass ")
+	TEXT("as well as an extra dependent texture fetches right at the begining of the HistoryUpdate pass.\n")
+	TEXT("\n")
+	TEXT("The reprojection field also embeds a jacobian 2x2 matrix for each pixel to have more precise reprojection of the history")
+	TEXT("for the display pixels in the rendering pixels. This for instance allows to maintains sharp geometric edges on movements."),
+	ECVF_Scalability | ECVF_RenderThreadSafe);
+
+TAutoConsoleVariable<float> CVarTSRReprojectionFieldAntiAliasPixelSpeed(
+	TEXT("r.TSR.ReprojectionField.AntiAliasPixelSpeed"), 0.75f,
+	TEXT("Defines the output pixel velocity at which point the dilation should be spatial anti-aliased based of the depth buffer ")
+	TEXT("to avoid reprojection aliasing by extrusion on fast geometric edges (Default to 0.75 to make sure there is parralax ")
+	TEXT("disocclusion mask behind moving objects, best tuned with r.TSR.Visualize=11)."),
+	ECVF_RenderThreadSafe);
+
 TAutoConsoleVariable<float> CVarTSRWeightClampingSampleCount(
 	TEXT("r.TSR.Velocity.WeightClampingSampleCount"), 4.0f,
 	TEXT("Number of sample to count to in history pixel to clamp history to when output pixel velocity reach r.TSR.Velocity.WeightClampingPixelSpeed. ")
@@ -298,19 +319,14 @@ TAutoConsoleVariable<float> CVarTSRWeightClampingPixelSpeed(
 	TEXT("(Default = 1.0f)."),
 	ECVF_RenderThreadSafe);
 
-// TODO: improve CONFIG_VELOCITY_EXTRAPOLATION in TSRDilateVelcity.usf that is disabled at the moment.
-//TAutoConsoleVariable<float> CVarTSRVelocityExtrapolation(
-//	TEXT("r.TSR.Velocity.Extrapolation"), 1.0f,
-//	TEXT("Defines how much the velocity should be extrapolated on geometric discontinuities (Default = 1.0f)."),
-//	ECVF_Scalability | ECVF_RenderThreadSafe);
-
 #if !UE_BUILD_OPTIMIZED_SHOWFLAGS
 
 TAutoConsoleVariable<int32> CVarTSRVisualize(
 	TEXT("r.TSR.Visualize"), -1,
 	TEXT("Selects the TSR internal visualization mode.\n")
-	TEXT(" -2: Display an overview grid based regardless of VisualizeTSR show flag;\n")
-	TEXT(" -1: Display an overview grid based on the VisualizeTSR show flag (default, opened with the `show VisualizeTSR` command at runtime or Show > Visualize > TSR in editor viewports);\n")
+	TEXT(" -3: Display the reprojection field's grid based overview;\n")
+	TEXT(" -2: Display an grid based overview regardless of VisualizeTSR show flag;\n")
+	TEXT(" -1: Display an grid based overview on the VisualizeTSR show flag (default, opened with the `show VisualizeTSR` command at runtime or Show > Visualize > TSR in editor viewports);\n")
 	TEXT("  0: Number of accumulated samples in the history, particularily interesting to tune r.TSR.ShadingRejection.SampleCount and r.TSR.Velocity.WeightClampingSampleCount;\n")
 	TEXT("  1: Parallax disocclusion based of depth and velocity buffers;\n")
 	TEXT("  2: Mask where the history is rejected;\n")
@@ -318,7 +334,14 @@ TAutoConsoleVariable<int32> CVarTSRVisualize(
 	TEXT("  4: Mask where the history is resurrected (with r.TSR.Resurrection=1);\n")
 	TEXT("  5: Mask where the history is resurrected in the resurrected frame (with r.TSR.Resurrection=1), particularily interesting to tune r.TSR.Resurrection.PersistentFrameInterval;\n")
 	TEXT("  6: Mask where spatial anti-aliasing is being computed;\n")
-	TEXT("  7: Mask where the flickering temporal analysis heuristic is taking effects (with r.TSR.ShadingRejection.Flickering=1);\n"),
+	TEXT("  7: Mask where the flickering temporal analysis heuristic is taking effects (with r.TSR.ShadingRejection.Flickering=1);\n")
+	TEXT("  8: Summary of the reprojection field, show the the jacobian on X in green and Y in blue;\n")
+	TEXT("  9: Reprojection field's dilating offset to apply in the HistoryUpdate;\n")
+	TEXT(" 10: Coverage of the dilating offset to apply in the HistoryUpdate (red the coverage is close to 0, green is close to 1, blue has been fully dilated to 1 without computing any spatial anti-aliasing from the depth buffer);\n")
+	TEXT(" 11: Mask where the reprojection field is anti-aliased from the depth buffer in green (handy to tune r.TSR.ReprojectionField.AntiAliasPixelSpeed);\n")
+	TEXT(" 12: Mask where the pixel's jacobian is null in the reprojection field in orange;\n")
+	TEXT(" 13: Mask where the pixel's jacobian has reached its encoding limit in the reprojection field in red;\n")
+	TEXT(" 14: Mask where the reprojected history is upscaled (in red) or downscaled (in green) by the reprojection field's jacobian (for instance due to getting closer or further away from camera respectively, or an object is getting scaled dynamicaly);\n"),
 	ECVF_RenderThreadSafe);
 
 #endif
@@ -657,16 +680,17 @@ class FTSRDilateVelocityCS : public FTSRShader
 		SHADER_PARAMETER(FMatrix44f, RotationalClipToPrevClip)
 		SHADER_PARAMETER(FVector2f, PrevOutputBufferUVMin)
 		SHADER_PARAMETER(FVector2f, PrevOutputBufferUVMax)
-		SHADER_PARAMETER(float, VelocityExtrapolationMultiplier)
 		SHADER_PARAMETER(float, InvFlickeringMaxParralaxVelocity)
+		SHADER_PARAMETER(float, ReprojectionFieldAntiAliasVelocityThreshold)
+		SHADER_PARAMETER(int32, bReprojectionField)
 		SHADER_PARAMETER(int32, bOutputIsMovingTexture)
 
 		SHADER_PARAMETER_RDG_TEXTURE(Texture2D, SceneDepthTexture)
 		SHADER_PARAMETER_RDG_TEXTURE(Texture2D, SceneVelocityTexture)
 
-		SHADER_PARAMETER_RDG_TEXTURE_UAV(RWTexture2D, DilatedVelocityOutput)
 		SHADER_PARAMETER_RDG_TEXTURE_UAV(RWTexture2D, ClosestDepthOutput)
 		SHADER_PARAMETER_RDG_TEXTURE_UAV(RWTexture2DArray, PrevAtomicOutput)
+		SHADER_PARAMETER_RDG_TEXTURE_UAV(RWTexture2DArray, ReprojectionFieldOutput)
 		SHADER_PARAMETER_RDG_TEXTURE_UAV(RWTexture2DArray, R8Output)
 
 		SHADER_PARAMETER_RDG_TEXTURE_UAV(RWTexture2D, VelocityFlattenOutput)
@@ -710,7 +734,7 @@ class FTSRDecimateHistoryCS : public FTSRShader
 		SHADER_PARAMETER(FMatrix44f, RotationalClipToPrevClip)
 
 		SHADER_PARAMETER_RDG_TEXTURE(Texture2D, InputSceneColorTexture)
-		SHADER_PARAMETER_RDG_TEXTURE(Texture2D, DilatedVelocityTexture)
+		SHADER_PARAMETER_RDG_TEXTURE_SRV(Texture2D, ReprojectionVectorTexture)
 		SHADER_PARAMETER_RDG_TEXTURE(Texture2D, ClosestDepthTexture)
 		SHADER_PARAMETER_RDG_TEXTURE(Texture2DArray, PrevAtomicTextureArray)
 
@@ -731,7 +755,7 @@ class FTSRDecimateHistoryCS : public FTSRShader
 
 		SHADER_PARAMETER_RDG_TEXTURE_UAV(RWTexture2DArray, ReprojectedHistoryGuideOutput)
 		SHADER_PARAMETER_RDG_TEXTURE_UAV(RWTexture2DArray, ReprojectedHistoryMoireOutput)
-		SHADER_PARAMETER_RDG_TEXTURE_UAV(RWTexture2D, HoleFilledVelocityOutput)
+		SHADER_PARAMETER_RDG_TEXTURE_UAV(RWTexture2DArray, ReprojectionFieldOutput)
 		SHADER_PARAMETER_RDG_TEXTURE_UAV(RWTexture2D, DecimateMaskOutput)
 
 		SHADER_PARAMETER_RDG_TEXTURE_UAV(RWTexture2DArray, DebugOutput)
@@ -785,6 +809,10 @@ class FTSRRejectShadingCS : public FTSRConvolutionNetworkShader
 		SHADER_PARAMETER(FVector2f, TranslucencyTextureUVMin)
 		SHADER_PARAMETER(FVector2f, TranslucencyTextureUVMax)
 		SHADER_PARAMETER(FMatrix44f, ClipToResurrectionClip)
+		SHADER_PARAMETER(FVector2f, ResurrectionJacobianXMul)
+		SHADER_PARAMETER(FVector2f, ResurrectionJacobianXAdd)
+		SHADER_PARAMETER(FVector2f, ResurrectionJacobianYMul)
+		SHADER_PARAMETER(FVector2f, ResurrectionJacobianYAdd)
 		SHADER_PARAMETER(FVector3f, HistoryGuideQuantizationError)
 		SHADER_PARAMETER(float, FlickeringFramePeriod)
 		SHADER_PARAMETER(float, TheoricBlendFactor)
@@ -801,6 +829,7 @@ class FTSRRejectShadingCS : public FTSRConvolutionNetworkShader
 		SHADER_PARAMETER_RDG_TEXTURE_SRV(Texture2D, ReprojectedHistoryMoireTexture)
 		SHADER_PARAMETER_RDG_TEXTURE_SRV(Texture2D, ResurrectedHistoryGuideTexture)
 		SHADER_PARAMETER_RDG_TEXTURE_SRV(Texture2D, ResurrectedHistoryGuideMetadataTexture)
+		SHADER_PARAMETER_RDG_TEXTURE_SRV(Texture2D, DilateMaskTexture)
 		SHADER_PARAMETER_RDG_TEXTURE(Texture2D, DecimateMaskTexture)
 		SHADER_PARAMETER_RDG_TEXTURE_SRV(Texture2D, IsMovingMaskTexture)
 		SHADER_PARAMETER_RDG_TEXTURE(Texture2D, ClosestDepthTexture)
@@ -808,7 +837,7 @@ class FTSRRejectShadingCS : public FTSRConvolutionNetworkShader
 		SHADER_PARAMETER_RDG_TEXTURE_UAV(RWTexture2DArray, HistoryGuideOutput)
 		SHADER_PARAMETER_RDG_TEXTURE_UAV(RWTexture2DArray, HistoryMoireOutput)
 		SHADER_PARAMETER_RDG_TEXTURE_UAV(RWTexture2D, HistoryRejectionOutput)
-		SHADER_PARAMETER_RDG_TEXTURE_UAV(RWTexture2D, DilatedVelocityOutput)
+		SHADER_PARAMETER_RDG_TEXTURE_UAV(RWTexture2DArray, ReprojectionFieldOutput)
 		SHADER_PARAMETER_RDG_TEXTURE_UAV(RWTexture2D, InputSceneColorOutput)
 		SHADER_PARAMETER_RDG_TEXTURE_UAV(RWTexture2D, InputSceneColorLdrLumaOutput)
 		SHADER_PARAMETER_RDG_TEXTURE_UAV(RWTexture2D, AntiAliasMaskOutput)
@@ -931,7 +960,9 @@ class FTSRUpdateHistoryCS : public FTSRShader
 		SHADER_PARAMETER_RDG_TEXTURE(Texture2D, InputSceneTranslucencyTexture)
 
 		SHADER_PARAMETER_RDG_TEXTURE(Texture2D, HistoryRejectionTexture)
-		SHADER_PARAMETER_RDG_TEXTURE(Texture2D, DilatedVelocityTexture)
+		SHADER_PARAMETER_RDG_TEXTURE_SRV(Texture2D, ReprojectionBoundaryTexture)
+		SHADER_PARAMETER_RDG_TEXTURE_SRV(Texture2D, ReprojectionJacobianTexture)
+		SHADER_PARAMETER_RDG_TEXTURE_SRV(Texture2D, ReprojectionVectorTexture)
 		SHADER_PARAMETER_RDG_TEXTURE(Texture2D, AntiAliasingTexture)
 
 		SHADER_PARAMETER_STRUCT(FScreenPassTextureViewportParameters, TranslucencyInfo)
@@ -955,6 +986,7 @@ class FTSRUpdateHistoryCS : public FTSRShader
 		SHADER_PARAMETER(float, ResurrectionFrameIndex)
 		SHADER_PARAMETER(float, PrevFrameIndex)
 		SHADER_PARAMETER(int32, bLensDistortion)
+		SHADER_PARAMETER(int32, bReprojectionField)
 		SHADER_PARAMETER(int32, bGenerateOutputMip1)
 		SHADER_PARAMETER(int32, bGenerateOutputMip2)
 		SHADER_PARAMETER(int32, bGenerateOutputMip3)
@@ -1115,6 +1147,7 @@ class FTSRVisualizeCS : public FTSRShader
 		SHADER_PARAMETER(int32, VisualizeId)
 		SHADER_PARAMETER(int32, bCanResurrectHistory)
 		SHADER_PARAMETER(int32, bCanSpatialAntiAlias)
+		SHADER_PARAMETER(int32, bReprojectionField)
 		SHADER_PARAMETER(float, MaxHistorySampleCount)
 		SHADER_PARAMETER(float, OutputToHistoryResolutionFractionSquare)
 		SHADER_PARAMETER(float, FlickeringFramePeriod)
@@ -1128,7 +1161,9 @@ class FTSRVisualizeCS : public FTSRShader
 		SHADER_PARAMETER_RDG_TEXTURE(Texture2D, InputSceneTranslucencyTexture)
 		SHADER_PARAMETER_RDG_TEXTURE_SRV(Texture2D, SceneColorTexture)
 		SHADER_PARAMETER_RDG_TEXTURE(Texture2D, ClosestDepthTexture)
-		SHADER_PARAMETER_RDG_TEXTURE(Texture2D, DilatedVelocityTexture)
+		SHADER_PARAMETER_RDG_TEXTURE_SRV(Texture2D, ReprojectionBoundaryTexture)
+		SHADER_PARAMETER_RDG_TEXTURE_SRV(Texture2D, ReprojectionJacobianTexture)
+		SHADER_PARAMETER_RDG_TEXTURE_SRV(Texture2D, ReprojectionVectorTexture)
 		SHADER_PARAMETER_RDG_TEXTURE_SRV(Texture2D, IsMovingMaskTexture)
 		SHADER_PARAMETER_RDG_TEXTURE(Texture2D, HistoryRejectionTexture)
 		SHADER_PARAMETER_RDG_TEXTURE_SRV(Texture2D, MoireHistoryTexture)
@@ -1424,6 +1459,9 @@ FDefaultTemporalUpscaler::FOutputs AddTemporalSuperResolutionPasses(
 	// period at which history changes is considered too distracting.
 	const float FlickeringFramePeriod = CVarTSRFlickeringEnable.GetValueOnRenderThread() ? (CVarTSRFlickeringPeriod.GetValueOnRenderThread() / FMath::Max(RefreshRateToFrameRateCap, 1.0f)) : 0.0f;
 
+	// Whether the reprojection field is enabled.
+	const bool bReprojectionField = CVarTSRReprojectionField.GetValueOnRenderThread() != 0;
+
 	ETSRHistoryFormatBits HistoryFormatBits = ETSRHistoryFormatBits::None;
 	{
 		if (FlickeringFramePeriod > 0)
@@ -1624,7 +1662,7 @@ FDefaultTemporalUpscaler::FOutputs AddTemporalSuperResolutionPasses(
 		CommonParameters.InputPixelPosToScreenPos = (FScreenTransform::Identity + 0.5f) * FScreenTransform::ChangeTextureBasisFromTo(FScreenPassTextureViewport(
 			InputExtent, InputRect), FScreenTransform::ETextureBasis::TexelPosition, FScreenTransform::ETextureBasis::ScreenPosition);
 		CommonParameters.ScreenVelocityToInputPixelVelocity = (FScreenTransform::Identity / CommonParameters.InputPixelPosToScreenPos).Scale;
-		CommonParameters.InputPixelVelocityToScreenVelocity = CommonParameters.InputPixelPosToScreenPos.Scale.GetAbs();
+		CommonParameters.InputPixelVelocityToScreenVelocity = CommonParameters.InputPixelPosToScreenPos.Scale;
 
 		CommonParameters.HistoryInfo = GetScreenPassTextureViewportParameters(FScreenPassTextureViewport(
 			HistoryExtent, FIntRect(FIntPoint(0, 0), HistorySize)));
@@ -1868,25 +1906,17 @@ FDefaultTemporalUpscaler::FOutputs AddTemporalSuperResolutionPasses(
 	}
 
 	// Dilate the velocity texture & scatter reprojection into previous frame
-	FRDGTextureRef DilatedVelocityTexture;
+	FRDGTextureRef ReprojectionFieldTexture;
+	FRDGTextureSRVRef ReprojectionVectorTexture;
+	FRDGTextureSRVRef ReprojectionBoundaryTexture = nullptr;
+	FRDGTextureSRVRef ReprojectionJacobianTexture = nullptr;
 	FRDGTextureRef ClosestDepthTexture;
+	FRDGTextureSRVRef DilateMaskTexture = nullptr;
 	FRDGTextureSRVRef IsMovingMaskTexture = nullptr;
 	FVelocityFlattenTextures VelocityFlattenTextures;
 	{
 		const bool bOutputIsMovingTexture = FlickeringFramePeriod > 0.0f;
 
-		FRDGTextureRef R8OutputTexture;
-
-		{
-			FRDGTextureDesc Desc = FRDGTextureDesc::Create2D(
-				InputExtent,
-				PF_G16R16,
-				FClearValueBinding::None,
-				/* InFlags = */ TexCreate_ShaderResource | TexCreate_UAV);
-
-			DilatedVelocityTexture = GraphBuilder.CreateTexture(Desc, TEXT("TSR.Velocity.Dilated"));
-		}
-
 		{
 			EPixelFormat ClosestDepthFormat;
 			if (bIsOrthoProjection)
@@ -1906,18 +1936,38 @@ FDefaultTemporalUpscaler::FOutputs AddTemporalSuperResolutionPasses(
 			ClosestDepthTexture = GraphBuilder.CreateTexture(Desc, TEXT("TSR.ClosestDepthTexture"));
 		}
 
+		FRDGTextureRef R8OutputTexture;
 		{
 			FRDGTextureDesc Desc = FRDGTextureDesc::Create2DArray(
 				InputExtent,
 				PF_R8_UINT,
 				FClearValueBinding::None,
 				/* InFlags = */ TexCreate_ShaderResource | TexCreate_UAV,
-				1);
+				/* ArraySize = */ bOutputIsMovingTexture ? 2 : 1);
 
-			R8OutputTexture = GraphBuilder.CreateTexture(Desc, TEXT("TSR.ParallaxFactor"));
+			R8OutputTexture = GraphBuilder.CreateTexture(Desc, TEXT("TSR.DilateR8"));
+			DilateMaskTexture = GraphBuilder.CreateSRV(FRDGTextureSRVDesc::CreateForSlice(R8OutputTexture, 0));
 			if (bOutputIsMovingTexture)
 			{
-				IsMovingMaskTexture = GraphBuilder.CreateSRV(FRDGTextureSRVDesc::CreateForSlice(R8OutputTexture, 0));
+				IsMovingMaskTexture = GraphBuilder.CreateSRV(FRDGTextureSRVDesc::CreateForSlice(R8OutputTexture, 1));
+			}
+		}
+
+		{
+			FRDGTextureDesc Desc = FRDGTextureDesc::Create2DArray(
+				InputExtent,
+				PF_R32_UINT,
+				FClearValueBinding::None,
+				/* InFlags = */ TexCreate_ShaderResource | TexCreate_UAV,
+				/* ArraySize = */ bReprojectionField ? 3 : 1);
+
+			ReprojectionFieldTexture = GraphBuilder.CreateTexture(Desc, TEXT("TSR.Reprojection.Vector"));
+			ReprojectionVectorTexture = GraphBuilder.CreateSRV(FRDGTextureSRVDesc::CreateForSlice(ReprojectionFieldTexture, 0));
+
+			if (bReprojectionField)
+			{
+				ReprojectionBoundaryTexture = GraphBuilder.CreateSRV(FRDGTextureSRVDesc::CreateForSlice(ReprojectionFieldTexture, 1));
+				ReprojectionJacobianTexture = GraphBuilder.CreateSRV(FRDGTextureSRVDesc::CreateForSlice(ReprojectionFieldTexture, 2));
 			}
 		}
 
@@ -1929,20 +1979,21 @@ FDefaultTemporalUpscaler::FOutputs AddTemporalSuperResolutionPasses(
 		PassParameters->RotationalClipToPrevClip = RotationalClipToPrevClip;
 		PassParameters->PrevOutputBufferUVMin = CommonParameters.InputInfo.UVViewportBilinearMin - CommonParameters.InputInfo.ExtentInverse;
 		PassParameters->PrevOutputBufferUVMax = CommonParameters.InputInfo.UVViewportBilinearMax + CommonParameters.InputInfo.ExtentInverse;
-		PassParameters->VelocityExtrapolationMultiplier = 0.0; //FMath::Clamp(CVarTSRVelocityExtrapolation.GetValueOnRenderThread(), 0.0f, 1.0f);
 		{
 			float FlickeringMaxParralaxVelocity = RefreshRateToFrameRateCap * CVarTSRFlickeringMaxParralaxVelocity.GetValueOnRenderThread() * float(View.ViewRect.Width()) / 1920.0f;
 			PassParameters->InvFlickeringMaxParralaxVelocity = 1.0f / FlickeringMaxParralaxVelocity;
 		}
+		PassParameters->ReprojectionFieldAntiAliasVelocityThreshold = FMath::Square(FMath::Max(CVarTSRReprojectionFieldAntiAliasPixelSpeed.GetValueOnRenderThread() / OutputToInputResolutionFraction, 1.0f / 64.0f));
+		PassParameters->bReprojectionField = bReprojectionField;
 		PassParameters->bOutputIsMovingTexture = bOutputIsMovingTexture;
 		
 		PassParameters->SceneDepthTexture = PassInputs.SceneDepth.Texture;
 		PassParameters->SceneVelocityTexture = PassInputs.SceneVelocity.Texture;
 
-		PassParameters->DilatedVelocityOutput = GraphBuilder.CreateUAV(DilatedVelocityTexture);
 		PassParameters->ClosestDepthOutput = GraphBuilder.CreateUAV(ClosestDepthTexture);
 		PassParameters->PrevAtomicOutput = GraphBuilder.CreateUAV(PrevAtomicTextureArray);
 		PassParameters->R8Output = GraphBuilder.CreateUAV(R8OutputTexture);
+		PassParameters->ReprojectionFieldOutput = GraphBuilder.CreateUAV(ReprojectionFieldTexture);
 
 		// Setup up the motion blur's velocity flatten pass.
 		if (PassInputs.bGenerateVelocityFlattenTextures)
@@ -1985,9 +2036,10 @@ FDefaultTemporalUpscaler::FOutputs AddTemporalSuperResolutionPasses(
 		TShaderMapRef<FTSRDilateVelocityCS> ComputeShader(View.ShaderMap, PermutationVector);
 		FComputeShaderUtils::AddPass(
 			GraphBuilder,
-			RDG_EVENT_NAME("TSR DilateVelocity(#%d MotionBlurDirections=%d%s) %dx%d",
+			RDG_EVENT_NAME("TSR DilateVelocity(#%d MotionBlurDirections=%d%s%s) %dx%d",
 				PermutationVector.ToDimensionValueId(),
 				int32(PermutationVector.Get<FTSRDilateVelocityCS::FMotionBlurDirectionsDim>()),
+				bReprojectionField ? TEXT(" ReprojectionField") : TEXT(""),
 				bOutputIsMovingTexture ? TEXT(" OutputIsMoving") : TEXT(""),
 				InputRect.Width(), InputRect.Height()),
 			AsyncComputePasses >= 2 ? ERDGPassFlags::AsyncCompute : ERDGPassFlags::Compute,
@@ -2001,7 +2053,6 @@ FDefaultTemporalUpscaler::FOutputs AddTemporalSuperResolutionPasses(
 	FRDGTextureRef ReprojectedHistoryMoireTexture = nullptr;
 	FRDGTextureRef DecimateMaskTexture = nullptr;
 	{
-		FRDGTextureRef HoleFilledVelocityTexture;
 		{
 			FRDGTextureDesc Desc = FRDGTextureDesc::Create2D(
 				InputExtent,
@@ -2009,11 +2060,7 @@ FDefaultTemporalUpscaler::FOutputs AddTemporalSuperResolutionPasses(
 				FClearValueBinding::None,
 				/* InFlags = */ TexCreate_ShaderResource | TexCreate_UAV);
 
-			Desc.Format = PF_R8G8;
 			DecimateMaskTexture = GraphBuilder.CreateTexture(Desc, TEXT("TSR.DecimateMask"));
-
-			Desc.Format = PF_G16R16;
-			HoleFilledVelocityTexture = GraphBuilder.CreateTexture(Desc, TEXT("TSR.Velocity.HoleFilled"));
 		}
 
 		{
@@ -2042,7 +2089,7 @@ FDefaultTemporalUpscaler::FOutputs AddTemporalSuperResolutionPasses(
 		PassParameters->RotationalClipToPrevClip = RotationalClipToPrevClip;
 
 		PassParameters->InputSceneColorTexture = PassInputs.SceneColor.Texture;
-		PassParameters->DilatedVelocityTexture = DilatedVelocityTexture;
+		PassParameters->ReprojectionVectorTexture = ReprojectionVectorTexture;
 		PassParameters->ClosestDepthTexture = ClosestDepthTexture;
 		PassParameters->PrevAtomicTextureArray = PrevAtomicTextureArray;
 
@@ -2077,7 +2124,7 @@ FDefaultTemporalUpscaler::FOutputs AddTemporalSuperResolutionPasses(
 		{
 			PassParameters->ReprojectedHistoryMoireOutput = GraphBuilder.CreateUAV(ReprojectedHistoryMoireTexture);
 		}
-		PassParameters->HoleFilledVelocityOutput = GraphBuilder.CreateUAV(HoleFilledVelocityTexture);
+		PassParameters->ReprojectionFieldOutput = GraphBuilder.CreateUAV(ReprojectionFieldTexture);
 		PassParameters->DecimateMaskOutput = GraphBuilder.CreateUAV(DecimateMaskTexture);
 		PassParameters->DebugOutput = CreateDebugUAV(InputExtent, TEXT("Debug.TSR.DecimateHistory"));
 
@@ -2090,9 +2137,9 @@ FDefaultTemporalUpscaler::FOutputs AddTemporalSuperResolutionPasses(
 		TShaderMapRef<FTSRDecimateHistoryCS> ComputeShader(View.ShaderMap, PermutationVector);
 		FComputeShaderUtils::AddPass(
 			GraphBuilder,
-			RDG_EVENT_NAME("TSR DecimateHistory(#%d %s%s%s%s) %dx%d",
+			RDG_EVENT_NAME("TSR DecimateHistory(#%d%s%s%s%s) %dx%d",
 				PermutationVector.ToDimensionValueId(),
-				PermutationVector.Get<FTSRDecimateHistoryCS::FMoireReprojectionDim>() ? TEXT("ReprojectMoire") : TEXT(""),
+				PermutationVector.Get<FTSRDecimateHistoryCS::FMoireReprojectionDim>() ? TEXT(" ReprojectMoire") : TEXT(""),
 				PermutationVector.Get<FTSRDecimateHistoryCS::FResurrectionReprojectionDim>() ? TEXT(" ReprojectResurrection") : TEXT(""),
 				PermutationVector.Get<FTSRShader::F16BitVALUDim>() ? TEXT(" 16bit") : TEXT(""),
 				PermutationVector.Get<FTSRShader::FAlphaChannelDim>() ? TEXT(" AlphaChannel") : TEXT(""),
@@ -2101,8 +2148,6 @@ FDefaultTemporalUpscaler::FOutputs AddTemporalSuperResolutionPasses(
 			ComputeShader,
 			PassParameters,
 			FComputeShaderUtils::GetGroupCount(InputRect.Size(), 8));
-
-		DilatedVelocityTexture = HoleFilledVelocityTexture;
 	}
 
 	// Merge PostDOF translucency within same scene color.
@@ -2188,7 +2233,17 @@ FDefaultTemporalUpscaler::FOutputs AddTemporalSuperResolutionPasses(
 			FScreenTransform::ChangeTextureBasisFromTo(TranslucencyViewport, FScreenTransform::ETextureBasis::ViewportUV, FScreenTransform::ETextureBasis::TextureUV);
 		PassParameters->TranslucencyTextureUVMin = GetScreenPassTextureViewportParameters(TranslucencyViewport).UVViewportBilinearMin;
 		PassParameters->TranslucencyTextureUVMax = GetScreenPassTextureViewportParameters(TranslucencyViewport).UVViewportBilinearMax;
-		PassParameters->ClipToResurrectionClip = ClipToResurrectionClip;
+		{
+			PassParameters->ClipToResurrectionClip = ClipToResurrectionClip;
+
+			FVector2f InputPixelVelocityToScreenVelocity = CommonParameters.InputPixelVelocityToScreenVelocity;
+			FVector2f ScreenVelocityToInputPixelVelocity = CommonParameters.ScreenVelocityToInputPixelVelocity;
+
+			PassParameters->ResurrectionJacobianXMul = - ScreenVelocityToInputPixelVelocity * FVector2f(ClipToResurrectionClip.M[0][0], ClipToResurrectionClip.M[0][1]) * InputPixelVelocityToScreenVelocity.X;
+			PassParameters->ResurrectionJacobianXAdd = ScreenVelocityToInputPixelVelocity * FVector2f(InputPixelVelocityToScreenVelocity.X, 0.0f);
+			PassParameters->ResurrectionJacobianYMul = - ScreenVelocityToInputPixelVelocity * FVector2f(ClipToResurrectionClip.M[1][0], ClipToResurrectionClip.M[1][1]) * InputPixelVelocityToScreenVelocity.Y;
+			PassParameters->ResurrectionJacobianYAdd = ScreenVelocityToInputPixelVelocity * FVector2f(0.0f, InputPixelVelocityToScreenVelocity.Y);
+		}
 		PassParameters->HistoryGuideQuantizationError = ComputePixelFormatQuantizationError(History.GuideArray->Desc.Format);
 		PassParameters->FlickeringFramePeriod = FlickeringFramePeriod;
 		PassParameters->TheoricBlendFactor = 1.0f / (1.0f + MaxHistorySampleCount / OutputToInputResolutionFractionSquare);
@@ -2233,6 +2288,7 @@ FDefaultTemporalUpscaler::FOutputs AddTemporalSuperResolutionPasses(
 			PassParameters->ResurrectedHistoryGuideTexture = GraphBuilder.CreateSRV(FRDGTextureSRVDesc(BlackDummy));
 			PassParameters->ResurrectedHistoryGuideMetadataTexture = GraphBuilder.CreateSRV(FRDGTextureSRVDesc(BlackDummy));
 		}
+		PassParameters->DilateMaskTexture = DilateMaskTexture;
 		PassParameters->DecimateMaskTexture = DecimateMaskTexture;
 		PassParameters->IsMovingMaskTexture = IsMovingMaskTexture ? IsMovingMaskTexture : GraphBuilder.CreateSRV(FRDGTextureSRVDesc(BlackUintDummy));
 		PassParameters->ClosestDepthTexture = ClosestDepthTexture;
@@ -2284,7 +2340,7 @@ FDefaultTemporalUpscaler::FOutputs AddTemporalSuperResolutionPasses(
 			PassParameters->HistoryRejectionOutput = GraphBuilder.CreateUAV(HistoryRejectionTexture);
 
 			// Amends how the history should be reprojected
-			PassParameters->DilatedVelocityOutput = GraphBuilder.CreateUAV(DilatedVelocityTexture);
+			PassParameters->ReprojectionFieldOutput = GraphBuilder.CreateUAV(ReprojectionFieldTexture);
 
 			// Output the composed translucency and opaque scene color to speed up HistoryUpdate
 			PassParameters->InputSceneColorOutput = bComputeInputSceneColorTexture
@@ -2378,7 +2434,9 @@ FDefaultTemporalUpscaler::FOutputs AddTemporalSuperResolutionPasses(
 		PassParameters->InputSceneTranslucencyTexture = SeparateTranslucencyTexture;
 		PassParameters->HistoryRejectionTexture = HistoryRejectionTexture;
 
-		PassParameters->DilatedVelocityTexture = DilatedVelocityTexture;
+		PassParameters->ReprojectionBoundaryTexture = ReprojectionBoundaryTexture ? ReprojectionBoundaryTexture : GraphBuilder.CreateSRV(FRDGTextureSRVDesc(BlackUintDummy));
+		PassParameters->ReprojectionJacobianTexture = ReprojectionJacobianTexture ? ReprojectionJacobianTexture : GraphBuilder.CreateSRV(FRDGTextureSRVDesc(BlackUintDummy));
+		PassParameters->ReprojectionVectorTexture = GraphBuilder.CreateSRV(FRDGTextureSRVDesc(ReprojectionFieldTexture));
 		PassParameters->AntiAliasingTexture = AntiAliasingTexture;
 
 		PassParameters->TranslucencyPixelPosMin = SeparateTranslucencyRect.Min;
@@ -2402,6 +2460,7 @@ FDefaultTemporalUpscaler::FOutputs AddTemporalSuperResolutionPasses(
 		PassParameters->InputToHistoryFactor = float(HistorySize.X) / float(InputRect.Width());
 		PassParameters->InputContributionMultiplier = OutputToHistoryResolutionFractionSquare; 
 		PassParameters->bLensDistortion = bLensDistortion;
+		PassParameters->bReprojectionField = bReprojectionField;
 		PassParameters->bGenerateOutputMip1 = false;
 		PassParameters->bGenerateOutputMip2 = false;
 		PassParameters->bGenerateOutputMip3 = false;
@@ -2541,12 +2600,13 @@ FDefaultTemporalUpscaler::FOutputs AddTemporalSuperResolutionPasses(
 		TShaderMapRef<FTSRUpdateHistoryCS> ComputeShader(View.ShaderMap, PermutationVector);
 		FComputeShaderUtils::AddPass(
 			GraphBuilder,
-			RDG_EVENT_NAME("TSR UpdateHistory(#%d Quality=%s%s%s%s%s%s) %dx%d",
+			RDG_EVENT_NAME("TSR UpdateHistory(#%d Quality=%s%s%s%s%s%s%s) %dx%d",
 				PermutationVector.ToDimensionValueId(),
 				kUpdateQualityNames[int32(PermutationVector.Get<FTSRUpdateHistoryCS::FQualityDim>())],
 				PermutationVector.Get<FTSRShader::F16BitVALUDim>() ? TEXT(" 16bit") : TEXT(""),
 				PermutationVector.Get<FTSRShader::FAlphaChannelDim>() ? TEXT(" AlphaChannel") : TEXT(""),
 				HistoryColorFormat == PF_FloatR11G11B10 ? TEXT(" R11G11B10") : TEXT(""),
+				bReprojectionField ? TEXT(" ReprojectionField") : TEXT(""),
 				bSupportsLensDistortion ? (bLensDistortion ? TEXT(" ApplyLensDistortion") : TEXT(" SupportLensDistortion")) : TEXT(""),
 				PassParameters->bGenerateOutputMip3 ? TEXT(" OutputMip3") : (PassParameters->bGenerateOutputMip2 ? TEXT(" OutputMip2") : (PassParameters->bGenerateOutputMip1 ? TEXT(" OutputMip1") : TEXT(""))),
 				HistorySize.X, HistorySize.Y),
@@ -2707,7 +2767,9 @@ FDefaultTemporalUpscaler::FOutputs AddTemporalSuperResolutionPasses(
 
 		enum class EVisualizeId : int32
 		{
-			Overview = -1,
+			ReprojectionFieldOverview = -3,
+			Overview = -2,
+			ShowFlag = -1,
 			HistorySampleCount = 0,
 			ParallaxDisocclusionMask = 1,
 			HistoryRejection = 2,
@@ -2716,6 +2778,13 @@ FDefaultTemporalUpscaler::FOutputs AddTemporalSuperResolutionPasses(
 			ResurrectedColor = 5,
 			SpatialAntiAliasingMask = 6,
 			AntiFlickering = 7,
+			ReprojectionFieldSummary = 8,
+			ReprojectionFieldOffset = 9,
+			ReprojectionFieldOffsetCoverage = 10,
+			ReprojectionFieldAA = 11,
+			ReprojectionFieldNullJacobian = 12,
+			ReprojectionFieldClampedJacobian = 13,
+			ReprojectionFieldDilatedJacobian = 14,
 			MAX,
 		};
 
@@ -2728,15 +2797,27 @@ FDefaultTemporalUpscaler::FOutputs AddTemporalSuperResolutionPasses(
 			TEXT("ResurrectedColor"),
 			TEXT("SpatialAntiAliasingMask"),
 			TEXT("AntiFlickering"),
+			TEXT("ReprojectionFieldSummary"),
+			TEXT("ReprojectionFieldOffset"),
+			TEXT("ReprojectionFieldOffsetCoverage"),
+			TEXT("ReprojectionFieldAA"),
+			TEXT("ReprojectionFieldNullJacobian"),
+			TEXT("ReprojectionFieldClampedJacobian"),
+			TEXT("ReprojectionFieldDilatedJacobian"),
 		};
 		static_assert(UE_ARRAY_COUNT(kVisualizationName) == int32(EVisualizeId::MAX), "kVisualizationName doesn't match EVisualizeId");
 
-		const EVisualizeId Visualization = EVisualizeId(FMath::Clamp(CVarTSRVisualize.GetValueOnRenderThread(), int32(EVisualizeId::Overview), int32(EVisualizeId::MAX) - 1));
-		FIntRect VisualizeRect = Visualization == EVisualizeId::Overview ? FIntRect(OutputRect.Min + OutputRect.Size() / 4, OutputRect.Min + (OutputRect.Size() * 3) / 4) : OutputRect;
+		const EVisualizeId Visualization = EVisualizeId(FMath::Clamp(CVarTSRVisualize.GetValueOnRenderThread(), int32(EVisualizeId::ReprojectionFieldOverview), int32(EVisualizeId::MAX) - 1));
+		const bool bIsOverviewVisualize =
+			Visualization == EVisualizeId::ShowFlag ||
+			Visualization == EVisualizeId::Overview ||
+			Visualization == EVisualizeId::ReprojectionFieldOverview;
+		
+		FIntRect VisualizeRect = bIsOverviewVisualize ? FIntRect(OutputRect.Min + OutputRect.Size() / 4, OutputRect.Min + (OutputRect.Size() * 3) / 4) : OutputRect;
 
 		auto Visualize = [&](EVisualizeId VisualizeId, FString Label)
 		{
-			check(VisualizeId != EVisualizeId::Overview);
+			check(int32(VisualizeId) >= 0);
 
 			FRDGTextureDesc OutputDesc = FRDGTextureDesc::Create2D(
 				OutputExtent,
@@ -2766,6 +2847,7 @@ FDefaultTemporalUpscaler::FOutputs AddTemporalSuperResolutionPasses(
 			PassParameters->VisualizeId = int32(VisualizeId);
 			PassParameters->bCanResurrectHistory = bCanResurrectHistory;
 			PassParameters->bCanSpatialAntiAlias = RejectionAntiAliasingQuality > 0;
+			PassParameters->bReprojectionField = bReprojectionField;
 			PassParameters->MaxHistorySampleCount = MaxHistorySampleCount;
 			PassParameters->OutputToHistoryResolutionFractionSquare = OutputToHistoryResolutionFractionSquare;
 			PassParameters->FlickeringFramePeriod = FlickeringFramePeriod;
@@ -2792,8 +2874,9 @@ FDefaultTemporalUpscaler::FOutputs AddTemporalSuperResolutionPasses(
 			PassParameters->InputSceneTranslucencyTexture = SeparateTranslucencyTexture;
 			PassParameters->SceneColorTexture = SceneColorOutputTextureSRV;
 			PassParameters->ClosestDepthTexture = ClosestDepthTexture;
-			PassParameters->DilatedVelocityTexture = DilatedVelocityTexture;
-			PassParameters->IsMovingMaskTexture = IsMovingMaskTexture ? IsMovingMaskTexture : GraphBuilder.CreateSRV(FRDGTextureSRVDesc(BlackUintDummy));
+			PassParameters->ReprojectionBoundaryTexture = ReprojectionBoundaryTexture ? ReprojectionBoundaryTexture : GraphBuilder.CreateSRV(FRDGTextureSRVDesc(BlackUintDummy));
+			PassParameters->ReprojectionJacobianTexture = ReprojectionJacobianTexture ? ReprojectionJacobianTexture : GraphBuilder.CreateSRV(FRDGTextureSRVDesc(BlackUintDummy));
+			PassParameters->ReprojectionVectorTexture = GraphBuilder.CreateSRV(FRDGTextureSRVDesc(ReprojectionFieldTexture));			PassParameters->IsMovingMaskTexture = IsMovingMaskTexture ? IsMovingMaskTexture : GraphBuilder.CreateSRV(FRDGTextureSRVDesc(BlackUintDummy));
 			PassParameters->HistoryRejectionTexture = HistoryRejectionTexture;
 			PassParameters->MoireHistoryTexture = MoireHistoryTexture ? MoireHistoryTexture : GraphBuilder.CreateSRV(FRDGTextureSRVDesc(BlackDummy));
 			PassParameters->AntiAliasMaskTexture = AntiAliasMaskTexture ? AntiAliasMaskTexture : BlackUintDummy;
@@ -2826,10 +2909,11 @@ FDefaultTemporalUpscaler::FOutputs AddTemporalSuperResolutionPasses(
 		};
 
 		FRDGTextureRef OutputTexture;
-		if (Visualization == EVisualizeId::Overview)
+		if (bIsOverviewVisualize)
 		{
 			TArray<FVisualizeBufferTile> Tiles;
 			Tiles.SetNum(16);
+			if (Visualization == EVisualizeId::Overview || Visualization == EVisualizeId::ShowFlag)
 			{
 				Tiles[4 * 0 + 0] = Visualize(EVisualizeId::HistorySampleCount, TEXT("Accumulated Sample Count"));
 				Tiles[4 * 0 + 1] = Visualize(EVisualizeId::ParallaxDisocclusionMask, TEXT("Parallax Disocclusion"));
@@ -2841,8 +2925,24 @@ FDefaultTemporalUpscaler::FOutputs AddTemporalSuperResolutionPasses(
 					Tiles[4 * 2 + 0] = Visualize(EVisualizeId::ResurrectedColor, TEXT("Resurrected Frame"));
 				}
 				Tiles[4 * 3 + 0] = Visualize(EVisualizeId::SpatialAntiAliasingMask, TEXT("Spatial Anti-Aliasing"));
+				Tiles[4 * 3 + 1] = Visualize(EVisualizeId::ReprojectionFieldSummary, TEXT("Reprojection Field"));
+				Tiles[4 * 3 + 1].Label = FString::Printf(TEXT("Reprojection Field (r.TSR.Visualize=%d)"), int32(EVisualizeId::ReprojectionFieldOverview));
 				Tiles[4 * 1 + 3] = Visualize(EVisualizeId::AntiFlickering, TEXT("Flickering Temporal Analysis"));
 			}
+			else if (Visualization == EVisualizeId::ReprojectionFieldOverview)
+			{
+				Tiles[4 * 0 + 0] = Visualize(EVisualizeId::ReprojectionFieldSummary, TEXT("Reprojection Field Summary"));
+				Tiles[4 * 0 + 1] = Visualize(EVisualizeId::ReprojectionFieldNullJacobian, TEXT("Reprojection Field's Null Jacobian"));
+				Tiles[4 * 0 + 2] = Visualize(EVisualizeId::ReprojectionFieldClampedJacobian, TEXT("Reprojection Field's Clamped Jacobian"));
+				Tiles[4 * 0 + 3] = Visualize(EVisualizeId::ReprojectionFieldDilatedJacobian, TEXT("Reprojection Field's Dilated Jacobian"));
+				Tiles[4 * 1 + 0] = Visualize(EVisualizeId::ReprojectionFieldOffset, TEXT("Reprojection Field's Offset"));
+				Tiles[4 * 2 + 0] = Visualize(EVisualizeId::ReprojectionFieldOffsetCoverage, TEXT("Reprojection Field's Offset Coverage"));
+				Tiles[4 * 3 + 0] = Visualize(EVisualizeId::ReprojectionFieldAA, TEXT("Reprojection Field's Anti-Aliasing"));
+			}
+			else
+			{
+				unimplemented();
+			}
 
 			{
 				FRDGTextureDesc OutputDesc = FRDGTextureDesc::Create2D(
