Mobile SSR: SSR is now done in deferred shading pass.

diff --git a/Engine/Shaders/Private/MobileDeferredShading.usf b/Engine/Shaders/Private/MobileDeferredShading.usf
--- a/Engine/Shaders/Private/MobileDeferredShading.usf
+++ b/Engine/Shaders/Private/MobileDeferredShading.usf
@@ -189,7 +189,7 @@ void MobileDirectionalLightPS(
 	// LightFunction should only affect direct lighting result
 	DirectLighting.TotalLight *= ComputeLightFunctionMultiplier(TranslatedWorldPosition);
 
-#if ENABLE_CLUSTERED_REFLECTION || ENABLE_SKY_LIGHT || ENABLE_PLANAR_REFLECTION
+#if ENABLE_CLUSTERED_REFLECTION || ENABLE_SKY_LIGHT || ENABLE_PLANAR_REFLECTION || MOBILE_SSR_ENABLED
 	// If we have a single directional light, apply relfection and sky contrubution here
 	ReflectionEnvironmentSkyLighting(GBuffer, SvPosition, CameraVector, TranslatedWorldPosition, ReflectionVector, GridIndex, DirectLighting);
 #endif
diff --git a/Engine/Shaders/Private/MobileLightingCommon.ush b/Engine/Shaders/Private/MobileLightingCommon.ush
--- a/Engine/Shaders/Private/MobileLightingCommon.ush
+++ b/Engine/Shaders/Private/MobileLightingCommon.ush
@@ -15,16 +15,21 @@
 
 #if MOBILE_SSR_ENABLED && !(MOBILE_QL_FORCE_FULLY_ROUGH || MATERIAL_FULLY_ROUGH)
 	#if MOBILE_DEFERRED_SHADING
-		#define MOBILE_USE_SSR (MATERIAL_SSR)
+		#if IS_MOBILE_BASE_PASS
+			#define MOBILE_USE_SSR (MATERIAL_SSR)
+		#else
+			#define MOBILE_USE_SSR 1
+		#endif
 	#else
 		#define MOBILE_USE_SSR (MATERIAL_SSR || !MATERIALBLENDING_ANY_TRANSLUCENT)
 	#endif
+	#define SSRData	MobileBasePass.SSRParams
 #else 
 	#define MOBILE_USE_SSR 0
 #endif
 
 #if MOBILE_USE_SSR
-#include "SSRT/SSRTRayCast.ush"
+#include "MobileSSR.ush"
 #endif
 
 #ifndef TRANSLUCENCY_NON_DIRECTIONAL
@@ -497,28 +502,10 @@ void AccumulateReflection(
 
 
 #if MOBILE_USE_SSR
-	BRANCH if(View.CameraCut == 0 && MobileBasePass.SSRParams.QualityAndExposureCorrection.x > 0.0 && GBuffer.Roughness <= MobileBasePass.SSRParams.QualityAndExposureCorrection.w)
+	half4 SSR = 0;
+	if (MobileSSR(GBuffer, SvPosition, TranslatedWorldPosition, ReflectionVector, SSR))
 	{
-		float StepOffset = InterleavedGradientNoise(SvPosition.xy, View.StateFrameIndexMod8);
-		StepOffset -= 0.5;
-		float3 HitUVz;
-		float Level = 0;
-		bool bHit = RayCast(MobileBasePass.SSRParams.HZB, MobileBasePass.SSRParams.HZBSampler, TranslatedWorldPosition, ReflectionVector, GBuffer.Roughness, GBuffer.Depth,
-			12, StepOffset, MobileBasePass.SSRParams.HZBUvFactorAndInvFactor, /*bDebugPrint=*/false, HitUVz, Level);
-
-		BRANCH if(bHit)
-		{
-			float2 SampleUV;
-			float Vignette;
-			ReprojectHit(MobileBasePass.SSRParams.PrevScreenPositionScaleBias, HitUVz, SampleUV, Vignette);
-
-			SampleUV = clamp(SampleUV, MobileBasePass.SSRParams.PrevSceneColorBilinearUVMinMax.xy,  MobileBasePass.SSRParams.PrevSceneColorBilinearUVMinMax.zw);
-			half4 SSR = SampleScreenColor(MobileBasePass.SSRParams.SceneColor, MobileBasePass.SSRParams.SceneColorSampler, SampleUV);
-			SSR *= Vignette * saturate( 2 - 6.6 * GBuffer.Roughness );
-
-			SSR.rgb *= MobileBasePass.SSRParams.QualityAndExposureCorrection.y;
-			SpecularIBL.rgb = SpecularIBL.rgb * (1 - SSR.a) + SSR.rgb;
-		}
+		SpecularIBL.rgb = SpecularIBL.rgb * (1 - SSR.a) + SSR.rgb;
 	}
 #endif // MOBILE_USE_SSR
 
diff --git a/Engine/Shaders/Private/MobileSSR.ush b/Engine/Shaders/Private/MobileSSR.ush
new file mode 100644
--- /dev/null
+++ b/Engine/Shaders/Private/MobileSSR.ush
@@ -0,0 +1,57 @@
+// Copyright Epic Games, Inc. All Rights Reserved.
+
+#include "Common.ush"
+
+#include "SSRT/SSRTRayCast.ush"
+
+#if MATERIALBLENDING_ANY_TRANSLUCENT || MOBILE_SSR_QUALITY <= 1
+	#define MOBILE_SSR_VELOCITY_TEXTURE_SAMPLER
+	#ifndef MOBILE_SSR_STEPS
+		#define MOBILE_SSR_STEPS 8
+	#endif
+#else
+	#define MOBILE_SSR_VELOCITY_TEXTURE_SAMPLER MobileSceneTextures.SceneVelocityTexture, MobileSceneTextures.SceneVelocityTextureSampler, 
+	#ifndef MOBILE_SSR_STEPS
+		#define MOBILE_SSR_STEPS 12
+	#endif
+#endif
+
+
+#ifndef SSRData
+#define SSRData MobileBasePass.SSRParams
+#endif
+
+#define SSRData_Intensity						((half)SSRData.IntensityAndExposureCorrection.x)
+#define SSRData_PrevSceneColorPreExposureInv	((half)SSRData.IntensityAndExposureCorrection.y)
+#define SSRData_MaxRoughness					((half)SSRData.IntensityAndExposureCorrection.z)
+#define SSRData_2DivMaxRoughness				((half)SSRData.IntensityAndExposureCorrection.w)
+
+bool MobileSSR(FGBufferData GBuffer, float4 SvPosition, float3 TranslatedWorldPosition, half3 ReflectionVector, out half4 SSR)
+{
+	BRANCH if(View.CameraCut == 0 && SSRData_Intensity > 0.0 && GBuffer.Roughness <= SSRData_MaxRoughness && GBuffer.ShadingModelID != 0)
+	{
+		float StepOffset = InterleavedGradientNoise(SvPosition.xy, SSRData.NoiseIndex);
+		StepOffset -= 0.5;
+		float3 HitUVz;
+		float Level = 0;
+		bool bHit = RayCast(SSRData.HZB, SSRData.HZBSampler, TranslatedWorldPosition, ReflectionVector, GBuffer.Roughness, GBuffer.Depth,
+			MOBILE_SSR_STEPS, StepOffset, SSRData.HZBUvFactorAndInvFactor, /*bDebugPrint=*/false, HitUVz, Level);
+
+		BRANCH if(bHit)
+		{
+			float2 SampleUV;
+			float Vignette;
+			ReprojectHit(SSRData.PrevScreenPositionScaleBias, MOBILE_SSR_VELOCITY_TEXTURE_SAMPLER 
+				HitUVz, SampleUV, Vignette);
+
+			SampleUV = clamp(SampleUV, SSRData.PrevSceneColorBilinearUVMinMax.xy,  SSRData.PrevSceneColorBilinearUVMinMax.zw);
+			SSR = SampleScreenColor(SSRData.SceneColor, SSRData.SceneColorSampler, SampleUV);
+
+			half RoughnessFade = RoughnessFade = saturate(2.0 - SSRData_2DivMaxRoughness * GBuffer.Roughness);
+			SSR *= Vignette * SSRData_Intensity * RoughnessFade;
+			SSR.rgb *= SSRData_PrevSceneColorPreExposureInv;
+			return true;
+		}
+	}
+	return false;
+}
diff --git a/Engine/Shaders/Private/SSRT/SSRTReflections.usf b/Engine/Shaders/Private/SSRT/SSRTReflections.usf
--- a/Engine/Shaders/Private/SSRT/SSRTReflections.usf
+++ b/Engine/Shaders/Private/SSRT/SSRTReflections.usf
@@ -1,12 +1,5 @@
 // Copyright Epic Games, Inc. All Rights Reserved.
 
-#if SHADING_PATH_MOBILE
-	#include "../Common.ush"
-	#define MobileSceneTextures MobileBasePass.SceneTextures
-	#include "../DeferredShadingCommon.ush"
-	#undef MobileSceneTextures
-#endif // SHADING_PATH_MOBILE
-
 #include "SSRTRayCast.ush"
 #include "../Random.ush"
 #include "../BRDF.ush"
@@ -37,23 +30,6 @@ SamplerState HZBSampler;
 
 uint ShouldReflectOnlyWater;
 
-#if SHADING_PATH_MOBILE
-	#define SceneColor MobileBasePass.SSRParams.SceneColor
-	#define SceneColorSampler MobileBasePass.SSRParams.SceneColorSampler
-	#define HZB MobileBasePass.SSRParams.HZB
-	#define HZBSampler MobileBasePass.SSRParams.HZBSampler
-	#define PrevScreenPositionScaleBias MobileBasePass.SSRParams.PrevScreenPositionScaleBias
-	#define HZBUvFactorAndInvFactor MobileBasePass.SSRParams.HZBUvFactorAndInvFactor
-	#define ShouldReflectOnlyWater 0
-	#define PrevSceneColorPreExposureCorrection MobileBasePass.SSRParams.QualityAndExposureCorrection.z
-	#define GBufferVelocityTexture MobileBasePass.SceneTextures.SceneVelocityTexture
-	#define GBufferVelocityTextureSampler MobileBasePass.SceneTextures.SceneVelocityTextureSampler
-#endif
-
-#ifndef SSR_USE_VELOCITY
-	#define SSR_USE_VELOCITY 1
-#endif
-
 float GetRoughness(const in FGBufferData GBuffer)
 {
 	float Roughness = GBuffer.Roughness;
@@ -176,11 +152,7 @@ void ScreenSpaceReflections(
 
 #else // SUBSTRATE_ENABLED
 
-	#if SHADING_PATH_MOBILE
-		FGBufferData GBuffer = MobileFetchAndDecodeGBuffer(UV, SvPosition.xy);
-	#else
-		FGBufferData GBuffer = GetGBufferDataFromSceneTextures(UV);
-	#endif
+	FGBufferData GBuffer = GetGBufferDataFromSceneTextures(UV);
 
 	float3 N = GBuffer.WorldNormal;
 	const float SceneDepth = GBuffer.Depth;
@@ -345,11 +317,7 @@ void ScreenSpaceReflections(
 
 				float2 SampleUV;
 				float Vignette;
-				#if SSR_USE_VELOCITY
-					ReprojectHit(PrevScreenPositionScaleBias, GBufferVelocityTexture, GBufferVelocityTextureSampler, HitUVz, SampleUV, Vignette);
-				#else
-					ReprojectHit(PrevScreenPositionScaleBias, HitUVz, SampleUV, Vignette);
-				#endif
+				ReprojectHit(PrevScreenPositionScaleBias, GBufferVelocityTexture, GBufferVelocityTextureSampler, HitUVz, SampleUV, Vignette);
 
 				float4 SampleColor = SampleScreenColor( SceneColor, SceneColorSampler, SampleUV ) * Vignette;
 
@@ -428,12 +396,7 @@ void ScreenSpaceReflections(
 
 			float2 SampleUV;
 			float Vignette;
-
-			#if SSR_USE_VELOCITY
-				ReprojectHit(PrevScreenPositionScaleBias, GBufferVelocityTexture, GBufferVelocityTextureSampler, HitUVz, SampleUV, Vignette);
-			#else
-				ReprojectHit(PrevScreenPositionScaleBias, HitUVz, SampleUV, Vignette);
-			#endif
+			ReprojectHit(PrevScreenPositionScaleBias, GBufferVelocityTexture, GBufferVelocityTextureSampler, HitUVz, SampleUV, Vignette);
 
 			OutColor = SampleScreenColor(SceneColor, SceneColorSampler, SampleUV) * Vignette;
 		}
@@ -471,32 +434,6 @@ void ScreenSpaceReflectionsPS(
 	);
 }
 
-void MobileScreenSpaceReflectionsPS(float4 SvPosition : SV_POSITION
-#if USE_GLES_FBF_DEFERRED
-, out float4 OutProxyAdditive : SV_Target0 // "OutProxyAdditive" has special handling in the shader compiler and does additive blending
-, out float4 OutGBufferA : SV_Target1
-, out float4 OutGBufferB : SV_Target2
-, out float4 OutGBufferC : SV_Target3
-#else
-, out float4 OutColor : SV_Target0
-#endif
-)
-{
-	float4 Color = 0;
-	float OutClosestHitDistance = 0;
-	ScreenSpaceReflections(SvPosition, Color
-#if SSR_OUTPUT_FOR_DENOISER
-		,OutClosestHitDistance
-#endif
-	);
-
-#if USE_GLES_FBF_DEFERRED
-	OutProxyAdditive = Color;
-#else
-	OutColor = Color;
-#endif
-}
-
 #if COMPUTE_SHADER
 
 RWTexture2D<float4> SSRColorOutput;
diff --git a/Engine/Source/Developer/TargetPlatform/Private/TargetPlatformSettingsBase.cpp b/Engine/Source/Developer/TargetPlatform/Private/TargetPlatformSettingsBase.cpp
--- a/Engine/Source/Developer/TargetPlatform/Private/TargetPlatformSettingsBase.cpp
+++ b/Engine/Source/Developer/TargetPlatform/Private/TargetPlatformSettingsBase.cpp
@@ -94,12 +94,6 @@ bool FTargetPlatformSettingsBase::UsesMobileDBuffer() const
 	return CVar ? (CVar->GetInt() != 0) : false;
 }
 
-bool FTargetPlatformSettingsBase::UsesMobileScreenSpaceReflections() const
-{
-	static IConsoleVariable* CVar = IConsoleManager::Get().FindConsoleVariable(TEXT("r.Mobile.ScreenSpaceReflections"));
-	return CVar ? (CVar->GetInt() != 0) : false;
-}
-
 bool FTargetPlatformSettingsBase::UsesASTCHDR() const
 {
 	static IConsoleVariable* CVar = IConsoleManager::Get().FindConsoleVariable(TEXT("cook.ASTCTextureCompressor"));
diff --git a/Engine/Source/Developer/TargetPlatform/Public/Common/TargetPlatformBase.h b/Engine/Source/Developer/TargetPlatform/Public/Common/TargetPlatformBase.h
--- a/Engine/Source/Developer/TargetPlatform/Public/Common/TargetPlatformBase.h
+++ b/Engine/Source/Developer/TargetPlatform/Public/Common/TargetPlatformBase.h
@@ -851,10 +851,6 @@ public:
 	{
 		return TargetPlatformSettings->UsesMobileDBuffer();
 	}
-	virtual bool UsesMobileScreenSpaceReflections() const override
-	{
-		return TargetPlatformSettings->UsesMobileScreenSpaceReflections();
-	}
 	virtual bool UsesASTCHDR() const override 
 	{
 		return TargetPlatformSettings->UsesASTCHDR();
diff --git a/Engine/Source/Developer/TargetPlatform/Public/Common/TargetPlatformSettingsBase.h b/Engine/Source/Developer/TargetPlatform/Public/Common/TargetPlatformSettingsBase.h
--- a/Engine/Source/Developer/TargetPlatform/Public/Common/TargetPlatformSettingsBase.h
+++ b/Engine/Source/Developer/TargetPlatform/Public/Common/TargetPlatformSettingsBase.h
@@ -36,7 +36,6 @@ public:
 	TARGETPLATFORM_API virtual int32 GetHeightFogModeForOpaque() const override;
 	TARGETPLATFORM_API virtual bool UsesMobileAmbientOcclusion() const override;
 	TARGETPLATFORM_API virtual bool UsesMobileDBuffer() const override;
-	TARGETPLATFORM_API virtual bool UsesMobileScreenSpaceReflections() const override;
 	TARGETPLATFORM_API virtual bool UsesASTCHDR() const override;
 	TARGETPLATFORM_API virtual void GetRayTracingShaderFormats(TArray<FName>& OutFormats) const override;
 
diff --git a/Engine/Source/Developer/TargetPlatform/Public/Interfaces/ITargetPlatformSettings.h b/Engine/Source/Developer/TargetPlatform/Public/Interfaces/ITargetPlatformSettings.h
--- a/Engine/Source/Developer/TargetPlatform/Public/Interfaces/ITargetPlatformSettings.h
+++ b/Engine/Source/Developer/TargetPlatform/Public/Interfaces/ITargetPlatformSettings.h
@@ -226,11 +226,6 @@ public:
 	 */
 	virtual bool UsesMobileDBuffer() const = 0;
 
-	/**
-	 * Gets whether the platform should use Screen Space Reflections when using the mobile renderer.
-	 */
-	virtual bool UsesMobileScreenSpaceReflections() const = 0;
-
 	/**
 	 * Gets whether the platform uses ASTC HDR
 	 */
diff --git a/Engine/Source/Runtime/RenderCore/Private/RenderUtils.cpp b/Engine/Source/Runtime/RenderCore/Private/RenderUtils.cpp
--- a/Engine/Source/Runtime/RenderCore/Private/RenderUtils.cpp
+++ b/Engine/Source/Runtime/RenderCore/Private/RenderUtils.cpp
@@ -540,11 +540,10 @@ RENDERCORE_API bool IsMobileAmbientOcclusionEnabled(const FStaticShaderPlatform
 	return IsMobilePlatform(Platform) && GMobileAmbientOcclusionPlatformMask[(int)Platform];
 }
 
-RENDERCORE_API ShaderPlatformMaskType GMobileScreenSpaceReflectionsPlatformMask;
-
 RENDERCORE_API bool AreMobileScreenSpaceReflectionsEnabled(const FStaticShaderPlatform Platform)
 {
-	return IsMobilePlatform(Platform) && GMobileScreenSpaceReflectionsPlatformMask[(int)Platform];
+	static FShaderPlatformCachedIniValue<bool> MobileScreenSpaceReflectionsIniValue(TEXT("r.Mobile.ScreenSpaceReflections"));
+	return IsMobilePlatform(Platform) && MobileScreenSpaceReflectionsIniValue.Get(Platform);
 }
 
 RENDERCORE_API bool IsMobileDistanceFieldEnabled(const FStaticShaderPlatform Platform)
@@ -604,7 +603,7 @@ RENDERCORE_API bool MobileBasePassAlwaysUsesCSM(const FStaticShaderPlatform Plat
 
 RENDERCORE_API bool MobileUsesFullDepthPrepass(const FStaticShaderPlatform Platform)
 {
-	return MobileUsesShadowMaskTexture(Platform) || IsMobileAmbientOcclusionEnabled(Platform) || IsUsingDBuffers(Platform) || AreMobileScreenSpaceReflectionsEnabled(Platform) || FReadOnlyCVARCache::MobileEarlyZPass(Platform) == 1;
+	return MobileUsesShadowMaskTexture(Platform) || IsMobileAmbientOcclusionEnabled(Platform) || IsUsingDBuffers(Platform) || FReadOnlyCVARCache::MobileEarlyZPass(Platform) == 1;
 }
 
 RENDERCORE_API bool ShouldForceFullDepthPass(const FStaticShaderPlatform Platform)
@@ -758,9 +757,6 @@ RENDERCORE_API void RenderUtilsInit()
 	static IConsoleVariable* MobileAmbientOcclusionCVar = IConsoleManager::Get().FindConsoleVariable(TEXT("r.Mobile.AmbientOcclusion"));
 	GMobileAmbientOcclusionPlatformMask.Init(MobileAmbientOcclusionCVar && MobileAmbientOcclusionCVar->GetInt(), EShaderPlatform::SP_NumPlatforms);
 
-	static IConsoleVariable* MobileScreenSpaceReflectionsCVar = IConsoleManager::Get().FindConsoleVariable(TEXT("r.Mobile.ScreenSpaceReflections"));
-	GMobileScreenSpaceReflectionsPlatformMask.Init(MobileScreenSpaceReflectionsCVar && MobileScreenSpaceReflectionsCVar->GetInt(), EShaderPlatform::SP_NumPlatforms);
-
 #if WITH_EDITOR
 	ITargetPlatformManagerModule* TargetPlatformManager = GetTargetPlatformManager();
 	if (TargetPlatformManager)
@@ -801,8 +797,6 @@ RENDERCORE_API void RenderUtilsInit()
 					GVelocityEncodeDepthPlatformMask[ShaderPlatformIndex] = TargetPlatformSettings->VelocityEncodeDepth();
 
 					GMobileAmbientOcclusionPlatformMask[ShaderPlatformIndex] = TargetPlatformSettings->UsesMobileAmbientOcclusion();
-
-					GMobileScreenSpaceReflectionsPlatformMask[ShaderPlatformIndex] = TargetPlatformSettings->UsesMobileScreenSpaceReflections();
 				}
 			}
 
diff --git a/Engine/Source/Runtime/RenderCore/Private/Shader.cpp b/Engine/Source/Runtime/RenderCore/Private/Shader.cpp
--- a/Engine/Source/Runtime/RenderCore/Private/Shader.cpp
+++ b/Engine/Source/Runtime/RenderCore/Private/Shader.cpp
@@ -1987,6 +1987,12 @@ void ShaderMapAppendKey(EShaderPlatform Platform, FShaderKeyGenerator& KeyGen)
 			KeyGen.AppendSeparator();
 			KeyGen.Append(TEXT("MobFDP"));
 		}
+
+		if (AreMobileScreenSpaceReflectionsEnabled(Platform))
+		{
+			KeyGen.AppendSeparator();
+			KeyGen.Append(TEXT("MobSSR"));
+		}
 	}
 	else
 	{
diff --git a/Engine/Source/Runtime/Renderer/Private/MobileBasePass.cpp b/Engine/Source/Runtime/Renderer/Private/MobileBasePass.cpp
--- a/Engine/Source/Runtime/Renderer/Private/MobileBasePass.cpp
+++ b/Engine/Source/Runtime/Renderer/Private/MobileBasePass.cpp
@@ -62,7 +62,7 @@ EMobileLocalLightSetting GetMobileForwardLocalLightSetting(EShaderPlatform Shade
 	return EMobileLocalLightSetting::LOCAL_LIGHTS_DISABLED;
 }
 
-extern const uint8 MobileShadingMobelSupportStencilValue = 0b01u;
+extern const uint8 MobileShadingModelSupportStencilValue = 0b01u;
 uint8 GetMobileShadingModelStencilValue(FMaterialShadingModelField ShadingModel, bool bFullyRough)
 {
 	// Bit 0 is set for materials that are receive SSR
@@ -78,7 +78,7 @@ uint8 GetMobileShadingModelStencilValue(FMaterialShadingModelField ShadingModel,
 	}
 
 	// mark everyhing as MSM_DefaultLit if GBuffer CustomData is not supported
-	return MobileUsesGBufferCustomData(GMaxRHIShaderPlatform) ? MobileShadingMobelSupportStencilValue : DefaultLitMask;
+	return MobileUsesGBufferCustomData(GMaxRHIShaderPlatform) ? MobileShadingModelSupportStencilValue : DefaultLitMask;
 }
 
 void SetMobileBasePassDepthState(FMeshPassProcessorRenderState& DrawRenderState, const FPrimitiveSceneProxy* PrimitiveSceneProxy, const FMaterial& Material, FMaterialShadingModelField ShadingModels, bool bUsesDeferredShading)
diff --git a/Engine/Source/Runtime/Renderer/Private/MobileDeferredShadingPass.cpp b/Engine/Source/Runtime/Renderer/Private/MobileDeferredShadingPass.cpp
--- a/Engine/Source/Runtime/Renderer/Private/MobileDeferredShadingPass.cpp
+++ b/Engine/Source/Runtime/Renderer/Private/MobileDeferredShadingPass.cpp
@@ -11,6 +11,7 @@
 #include "LightRendering.h"
 #include "LocalLightSceneProxy.h"
 #include "Materials/MaterialRenderProxy.h"
+#include "MobileSSR.h"
 
 DECLARE_GPU_STAT(DeferredShading);
 
@@ -63,6 +64,7 @@ class FMobileDirectionalLightFunctionPS : public FMaterialShader
 	class FEnableDynamicSkyLight		: SHADER_PERMUTATION_BOOL("ENABLE_DYNAMIC_SKY_LIGHT");
 	class FEnableCSM					: SHADER_PERMUTATION_BOOL("ENABLE_MOBILE_CSM");
 	class FShadowQuality				: SHADER_PERMUTATION_RANGE_INT("MOBILE_SHADOW_QUALITY", 1, 3); // not using Quality=0
+	class FMobileSSRQuality 			: SHADER_PERMUTATION_ENUM_CLASS("MOBILE_SSR_QUALITY", EMobileSSRQuality);
 	
 	using FPermutationDomain = TShaderPermutationDomain<
 		FEnableShadingModelSupport,
@@ -72,7 +74,8 @@ class FMobileDirectionalLightFunctionPS : public FMaterialShader
 		FEnableSkyLight,
 		FEnableDynamicSkyLight,
 		FEnableCSM, 
-		FShadowQuality>;
+		FShadowQuality,
+		FMobileSSRQuality>;
 
 	BEGIN_SHADER_PARAMETER_STRUCT(FParameters, )
 		SHADER_PARAMETER_STRUCT_REF(FMobileDirectionalLightShaderParameters, MobileDirectionalLight)
@@ -86,6 +89,7 @@ class FMobileDirectionalLightFunctionPS : public FMaterialShader
 
 	static void ModifyCompilationEnvironment(const FMaterialShaderPermutationParameters& Parameters, FShaderCompilerEnvironment& OutEnvironment)
 	{
+		FPermutationDomain PermutationVector(Parameters.PermutationId);
 		FMaterialShader::ModifyCompilationEnvironment(Parameters, OutEnvironment);
 		FForwardLightingParameters::ModifyCompilationEnvironment(Parameters.Platform, OutEnvironment);
 		OutEnvironment.SetDefine(TEXT("USE_LIGHT_FUNCTION"), Parameters.MaterialParameters.bIsDefaultMaterial ? 0 : 1);
@@ -96,6 +100,8 @@ class FMobileDirectionalLightFunctionPS : public FMaterialShader
 
 		const bool bMobileForceDepthRead = MobileUsesFullDepthPrepass(Parameters.Platform);
 		OutEnvironment.SetDefine(TEXT("FORCE_DEPTH_TEXTURE_READS"), bMobileForceDepthRead ? 1u : 0u);
+
+		OutEnvironment.SetDefine(TEXT("MOBILE_SSR_ENABLED"), PermutationVector.Get<FMobileSSRQuality>() != EMobileSSRQuality::Disabled ? 1u : 0u);
 	}
 
 	static FPermutationDomain RemapPermutationVector(FPermutationDomain PermutationVector, EShaderPlatform Platform)
@@ -120,6 +126,11 @@ class FMobileDirectionalLightFunctionPS : public FMaterialShader
 			PermutationVector.Set<FEnableDynamicSkyLight>(false);
 		}
 
+		if (!AreMobileScreenSpaceReflectionsEnabled(Platform))
+		{
+			PermutationVector.Set<FMobileSSRQuality>(EMobileSSRQuality::Disabled);
+		}
+
 		return PermutationVector;
 	}
 
@@ -142,7 +153,7 @@ class FMobileDirectionalLightFunctionPS : public FMaterialShader
 		return true;
 	}
 
-	static FPermutationDomain BuildPermutationVector(const FViewInfo& View, bool bInlineReflectionAndSky, bool bShadingModelSupport, bool bDynamicShadows, bool bSkyLight, bool bDynamicSkyLight, bool bPlanarReflection)
+	static FPermutationDomain BuildPermutationVector(const FViewInfo& View, bool bInlineReflectionAndSky, bool bShadingModelSupport, bool bDynamicShadows, bool bSkyLight, bool bDynamicSkyLight, bool bPlanarReflection, EMobileSSRQuality MobileSSRQuality)
 	{
 		bool bUseClusteredLights = UseClusteredDeferredShading(View.GetShaderPlatform());
 		bool bClustredReflection = bInlineReflectionAndSky && (View.NumBoxReflectionCaptures + View.NumSphereReflectionCaptures) > 0;
@@ -160,6 +171,7 @@ class FMobileDirectionalLightFunctionPS : public FMaterialShader
 		PermutationVector.Set<FMobileDirectionalLightFunctionPS::FEnableDynamicSkyLight>(bEnableDynamicSkyLight);
 		PermutationVector.Set<FMobileDirectionalLightFunctionPS::FEnableCSM>(ShadowQuality > 0);
 		PermutationVector.Set<FMobileDirectionalLightFunctionPS::FShadowQuality>(FMath::Clamp(ShadowQuality, 1, 3));
+		PermutationVector.Set<FMobileDirectionalLightFunctionPS::FMobileSSRQuality>(bInlineReflectionAndSky ? MobileSSRQuality : EMobileSSRQuality::Disabled);
 		return PermutationVector;
 	}
 
@@ -272,14 +284,16 @@ public:
 	class FEnableClustredReflection		: SHADER_PERMUTATION_BOOL("ENABLE_CLUSTERED_REFLECTION");
 	class FEnablePlanarReflection		: SHADER_PERMUTATION_BOOL("ENABLE_PLANAR_REFLECTION");
 	class FEnableSkyLight				: SHADER_PERMUTATION_BOOL("ENABLE_SKY_LIGHT");
-	class FEnableDynamicSkyLight : SHADER_PERMUTATION_BOOL("ENABLE_DYNAMIC_SKY_LIGHT");
+	class FEnableDynamicSkyLight		: SHADER_PERMUTATION_BOOL("ENABLE_DYNAMIC_SKY_LIGHT");
+	class FMobileSSRQuality				: SHADER_PERMUTATION_ENUM_CLASS("MOBILE_SSR_QUALITY", EMobileSSRQuality);
 	
 	using FPermutationDomain = TShaderPermutationDomain<
 		FEnableShadingModelSupport, 
 		FEnableClustredReflection, 
 		FEnablePlanarReflection,
 		FEnableSkyLight,
-		FEnableDynamicSkyLight
+		FEnableDynamicSkyLight,
+		FMobileSSRQuality
 	>;
 
 	BEGIN_SHADER_PARAMETER_STRUCT(FParameters, )
@@ -306,11 +320,17 @@ public:
 			return false;
 		}
 
+		if (PermutationVector.Get<FMobileSSRQuality>() != EMobileSSRQuality::Disabled && !AreMobileScreenSpaceReflectionsEnabled(Parameters.Platform))
+		{
+			return false;
+		}
+
 		return true;
 	}
 
 	static void ModifyCompilationEnvironment(const FGlobalShaderPermutationParameters& Parameters, FShaderCompilerEnvironment& OutEnvironment)
 	{
+		FPermutationDomain PermutationVector(Parameters.PermutationId);
 		FGlobalShader::ModifyCompilationEnvironment(Parameters, OutEnvironment);
 		FForwardLightingParameters::ModifyCompilationEnvironment(Parameters.Platform, OutEnvironment);
 		OutEnvironment.SetDefine(TEXT("IS_MOBILE_DEFERREDSHADING_SUBPASS"), 1u);
@@ -318,6 +338,7 @@ public:
 		const bool bMobileForceDepthRead = MobileUsesFullDepthPrepass(Parameters.Platform);
 		OutEnvironment.SetDefine(TEXT("FORCE_DEPTH_TEXTURE_READS"), bMobileForceDepthRead ? 1u : 0u);
 		OutEnvironment.SetDefine(TEXT("ENABLE_AMBIENT_OCCLUSION"), IsMobileAmbientOcclusionEnabled(Parameters.Platform) ? 1u : 0u);
+		OutEnvironment.SetDefine(TEXT("MOBILE_SSR_ENABLED"), PermutationVector.Get<FMobileSSRQuality>() != EMobileSSRQuality::Disabled ? 1u : 0u);
 	}
 };
 
@@ -376,11 +397,11 @@ static void GetLightMaterial(const FCachedLightMaterial& DefaultLightMaterial, c
 	OutShader = MaterialShaderMap->GetShader<ShaderType>(PermutationId);
 }
 
-extern const uint8 MobileShadingMobelSupportStencilValue;
+extern const uint8 MobileShadingModelSupportStencilValue;
 
 uint8 PassShadingModelStencilValue(bool bEnableShadingModelSupport)
 {
-	return bEnableShadingModelSupport ? GET_STENCIL_MOBILE_SM_MASK(MobileShadingMobelSupportStencilValue) : STENCIL_MOBILE_DEFAULTLIT_MASK;
+	return bEnableShadingModelSupport ? GET_STENCIL_MOBILE_SM_MASK(MobileShadingModelSupportStencilValue) : STENCIL_MOBILE_DEFAULTLIT_MASK;
 }
 
 constexpr uint8 PassShadingModelStencilMask(bool bEnableShadingModelSupport)
@@ -388,18 +409,18 @@ constexpr uint8 PassShadingModelStencilMask(bool bEnableShadingModelSupport)
 	return bEnableShadingModelSupport ? GET_STENCIL_MOBILE_SM_MASK(0xff) : STENCIL_MOBILE_DEFAULTLIT_MASK;
 }
 
-void RenderReflectionEnvironmentSkyLighting(FRHICommandList& RHICmdList, const FScene& Scene, const FViewInfo& View)
+void RenderReflectionEnvironmentSkyLighting(FRHICommandList& RHICmdList, const FScene& Scene, const FViewInfo& View, const EMobileSSRQuality MobileSSRQuality)
 {
 	// Skylights with static lighting already had their diffuse contribution baked into lightmaps
 	const bool bSkyLight = Scene.SkyLight && !Scene.SkyLight->bHasStaticLighting && View.Family->EngineShowFlags.SkyLighting;
 	const bool bDynamicSkyLight = bSkyLight && !Scene.SkyLight->bWantsStaticShadowing;
 	const bool bClustredReflection = (View.NumBoxReflectionCaptures + View.NumSphereReflectionCaptures) > 0;
 	const bool bPlanarReflection = Scene.GetForwardPassGlobalPlanarReflection() != nullptr;
-	if (!(bSkyLight || bClustredReflection || bPlanarReflection))
+	if (!(bSkyLight || bClustredReflection || bPlanarReflection || MobileSSRQuality != EMobileSSRQuality::Disabled))
 	{
 		return;
 	}
-		
+
 	SCOPED_DRAW_EVENT(RHICmdList, ReflectionEnvironmentSkyLighting);
 
 	FGraphicsPipelineStateInitializer GraphicsPSOInit;
@@ -415,40 +436,63 @@ void RenderReflectionEnvironmentSkyLighting(FRHICommandList& RHICmdList, const F
 		GraphicsPSOInit.BlendState = TStaticBlendState<CW_RGB, BO_Add, BF_One, BF_One>::GetRHI();
 	}
 	GraphicsPSOInit.RasterizerState = TStaticRasterizerState<>::GetRHI();
-	GraphicsPSOInit.DepthStencilState = TStaticDepthStencilState<
+
+	int32 NumPasses = 1;
+	FDepthStencilStateRHIRef StencilState[3];
+	StencilState[0] = TStaticDepthStencilState<
 		false, CF_Always,
 		true, CF_Equal, SO_Keep, SO_Keep, SO_Keep,
 		false, CF_Always, SO_Keep, SO_Keep, SO_Keep,
 		PassShadingModelStencilMask(/*bEnableShadingModelSupport=*/false), 0x00>::GetRHI();
+	uint8 StencilRef[3] = {};
+	StencilRef[0] = PassShadingModelStencilValue(/*bEnableShadingModelSupport=*/false);
 
 	FMobileReflectionEnvironmentSkyLightingPS::FParameters PassParameters;
 	PassParameters.View = GetShaderBinding(View.ViewUniformBuffer);
 	PassParameters.MobileReflectionCaptureData = GetShaderBinding(View.MobileReflectionCaptureUniformBuffer);
 
 	TShaderMapRef<FPostProcessVS> VertexShader(View.ShaderMap);
+	
+	uint32 PassEnableShadingModelSupport = 0;
+	uint32 PassEnableSSR = 0;
+	if (MobileSSRQuality != EMobileSSRQuality::Disabled)
+	{
+		// Separate pass for fully rough default lit materials
+		int PassIndex = NumPasses++;
+		StencilState[0] = StencilState[PassIndex] = TStaticDepthStencilState<
+			false, CF_Always,
+			true, CF_Equal, SO_Keep, SO_Keep, SO_Keep,
+			false, CF_Always, SO_Keep, SO_Keep, SO_Keep,
+			GET_STENCIL_MOBILE_SM_MASK(0xff), 0x00>::GetRHI();
+
+		StencilRef[PassIndex] = STENCIL_MOBILE_DEFAULTLIT_MASK | STENCIL_MOBILE_REFLECTIVE_MASK;
+		PassEnableSSR |= (1 << PassIndex);
+	}
+	if (!IsOnlyDefaultLitShadingModel(View.ShadingModelMaskInView) && MobileUsesGBufferCustomData(Scene.GetShaderPlatform()))
+	{
+		// Separate pass for all materials with custom shading models
+		int PassIndex = NumPasses++;
+		StencilState[PassIndex] = TStaticDepthStencilState<
+			false, CF_Always,
+			true, CF_Equal, SO_Keep, SO_Keep, SO_Keep,
+			false, CF_Always, SO_Keep, SO_Keep, SO_Keep,
+			PassShadingModelStencilMask(/*bEnableShadingModelSupport=*/true), 0x00>::GetRHI();
+		StencilRef[PassIndex] = PassShadingModelStencilValue(/*bEnableShadingModelSupport=*/true);
+		PassEnableSSR |= (MobileSSRQuality != EMobileSSRQuality::Disabled) ? (1 << PassIndex) : 0;
+		PassEnableShadingModelSupport |= (1 << PassIndex);
+	}
 
-	// Do two passes, first masking DefautLit, second masking all other shading models
-	const bool bOnlyDefaultLitInView = IsOnlyDefaultLitShadingModel(View.ShadingModelMaskInView);
-	int32 NumPasses = !bOnlyDefaultLitInView && MobileUsesGBufferCustomData(Scene.GetShaderPlatform()) ? 2 : 1;
-			
 	for (int32 PassIndex = 0; PassIndex < NumPasses; PassIndex++)
 	{
-		const bool bEnableShadingModelSupport = (PassIndex > 0);
-		if (bEnableShadingModelSupport)
-		{
-			GraphicsPSOInit.DepthStencilState = TStaticDepthStencilState<
-				false, CF_Always,
-				true, CF_Equal, SO_Keep, SO_Keep, SO_Keep,
-				false, CF_Always, SO_Keep, SO_Keep, SO_Keep,
-				PassShadingModelStencilMask(/*bEnableShadingModelSupport=*/true), 0x00>::GetRHI();
-		}
+		GraphicsPSOInit.DepthStencilState = StencilState[PassIndex];
 		
 		FMobileReflectionEnvironmentSkyLightingPS::FPermutationDomain PermutationVector;
-		PermutationVector.Set<FMobileReflectionEnvironmentSkyLightingPS::FEnableShadingModelSupport>(bEnableShadingModelSupport);
+		PermutationVector.Set<FMobileReflectionEnvironmentSkyLightingPS::FEnableShadingModelSupport>(PassEnableSSR & (1 << PassIndex));
 		PermutationVector.Set<FMobileReflectionEnvironmentSkyLightingPS::FEnableClustredReflection>(bClustredReflection);
 		PermutationVector.Set<FMobileReflectionEnvironmentSkyLightingPS::FEnablePlanarReflection>(bPlanarReflection);
 		PermutationVector.Set<FMobileReflectionEnvironmentSkyLightingPS::FEnableSkyLight>(bSkyLight);
 		PermutationVector.Set<FMobileReflectionEnvironmentSkyLightingPS::FEnableDynamicSkyLight>(bDynamicSkyLight);
+		PermutationVector.Set<FMobileReflectionEnvironmentSkyLightingPS::FMobileSSRQuality>((PassEnableSSR & (1 << PassIndex)) ? MobileSSRQuality : EMobileSSRQuality::Disabled);
 		TShaderMapRef<FMobileReflectionEnvironmentSkyLightingPS> PixelShader(View.ShaderMap, PermutationVector);
 		
 		GraphicsPSOInit.BoundShaderState.VertexDeclarationRHI = GFilterVertexDeclaration.VertexDeclarationRHI;
@@ -456,7 +500,7 @@ void RenderReflectionEnvironmentSkyLighting(FRHICommandList& RHICmdList, const F
 		GraphicsPSOInit.BoundShaderState.PixelShaderRHI = PixelShader.GetPixelShader();
 		GraphicsPSOInit.PrimitiveType = PT_TriangleList;
 
-		SetGraphicsPipelineState(RHICmdList, GraphicsPSOInit, PassShadingModelStencilValue(bEnableShadingModelSupport));
+		SetGraphicsPipelineState(RHICmdList, GraphicsPSOInit, StencilRef[PassIndex]);
 		SetShaderParameters(RHICmdList, PixelShader, PixelShader.GetPixelShader(), PassParameters);
 
 		const FIntPoint TargetSize = View.GetSceneTexturesConfig().Extent;
@@ -512,7 +556,7 @@ static void SetDirectionalLightDepthStencilState(FGraphicsPipelineStateInitializ
 	}
 }
 
-static void RenderDirectionalLight(FRHICommandList& RHICmdList, const FScene& Scene, const FViewInfo& View, const FCachedLightMaterial& DefaultLightMaterial, const FLightSceneInfo& DirectionalLight, uint32 LightingChannel, bool bInlineReflectionAndSky)
+static void RenderDirectionalLight(FRHICommandList& RHICmdList, const FScene& Scene, const FViewInfo& View, const FCachedLightMaterial& DefaultLightMaterial, const FLightSceneInfo& DirectionalLight, uint32 LightingChannel, bool bInlineReflectionAndSky, EMobileSSRQuality MobileSSRQuality)
 {
 	FString LightNameWithLevel;
 	FSceneRenderer::GetLightNameForDrawEvent(DirectionalLight.Proxy, LightNameWithLevel);
@@ -578,27 +622,48 @@ static void RenderDirectionalLight(FRHICommandList& RHICmdList, const FScene& Sc
 	
 	// Do two passes, first masking DefautLit, second masking all other shading models
 	const bool bOnlyDefaultLitInView = IsOnlyDefaultLitShadingModel(View.ShadingModelMaskInView);
-	int32 NumPasses = !bOnlyDefaultLitInView && MobileUsesGBufferCustomData(Scene.GetShaderPlatform()) ? 2 : 1;
+	const bool bUseSSR = bInlineReflectionAndSky && MobileSSRQuality != EMobileSSRQuality::Disabled;
+	int32 NumPasses = 1;
+	uint32 PassEnableSSR = 0;
+	uint32 PassEnableShadingModelSupport = 0;
+	uint32 PassShadingModelStencilMaskAnyDefaultLit = 1;
+	uint32 ShadingModelStencilRef[3] = {};
+	ShadingModelStencilRef[0] = PassShadingModelStencilValue(/*bEnableShadingModelSupport=*/false);
+
+	if (!bOnlyDefaultLitInView && MobileUsesGBufferCustomData(Scene.GetShaderPlatform()))
+	{
+		const int32 PassIndex = NumPasses++;
+		PassEnableShadingModelSupport |= (1 << PassIndex);
+		PassEnableSSR |= bUseSSR ? (1 << PassIndex) : 0;
+		ShadingModelStencilRef[PassIndex] = PassShadingModelStencilValue(/*bEnableShadingModelSupport=*/true);
+	}
+	if (bUseSSR)
+	{
+		PassShadingModelStencilMaskAnyDefaultLit = 0; // Pass 0 only does default lit, non-reflective.
+		const int32 PassIndex = NumPasses++;
+		PassEnableSSR |= (1 << PassIndex);
+		ShadingModelStencilRef[PassIndex] = STENCIL_MOBILE_DEFAULTLIT_MASK | STENCIL_MOBILE_REFLECTIVE_MASK;
+	}
 	
 	for (int32 PassIndex = 0; PassIndex < NumPasses; ++PassIndex)
 	{
-		const bool bEnableShadingModelSupport = (PassIndex > 0);
-		SetDirectionalLightDepthStencilState(GraphicsPSOInit, LightingChannelStencilValue, bEnableShadingModelSupport);
+		SetDirectionalLightDepthStencilState(GraphicsPSOInit, LightingChannelStencilValue, (PassShadingModelStencilMaskAnyDefaultLit & (1 << PassIndex)) == 0);
 		
 		FMobileDirectionalLightFunctionPS::FPermutationDomain PermutationVector = FMobileDirectionalLightFunctionPS::BuildPermutationVector(
 			View,
 			bInlineReflectionAndSky,
-			bEnableShadingModelSupport,
+			PassEnableShadingModelSupport & (1 << PassIndex),
 			bDynamicShadows,
 			bSkyLight,
 			bDynamicSkyLight,
-			bPlanarReflection
+			bPlanarReflection, 
+			(PassEnableSSR & (1<<PassIndex)) ? MobileSSRQuality : EMobileSSRQuality::Disabled
 		);
 		FCachedLightMaterial LightMaterial;
 		TShaderRef<FMobileDirectionalLightFunctionPS> PixelShader;
 		GetLightMaterial(DefaultLightMaterial, LightFunctionMaterialProxy, PermutationVector.ToDimensionValueId(), LightMaterial, PixelShader);
 
-		const uint8 StencilRef = PassShadingModelStencilValue(bEnableShadingModelSupport) | STENCIL_LIGHTING_CHANNELS_MASK(LightingChannelStencilValue);
+		const uint8 StencilRef = ShadingModelStencilRef[PassIndex] | STENCIL_LIGHTING_CHANNELS_MASK(LightingChannelStencilValue);
 
 		GraphicsPSOInit.BoundShaderState.VertexDeclarationRHI = GFilterVertexDeclaration.VertexDeclarationRHI;
 		GraphicsPSOInit.BoundShaderState.VertexShaderRHI = VertexShader.GetVertexShader();
@@ -622,7 +687,7 @@ static void RenderDirectionalLight(FRHICommandList& RHICmdList, const FScene& Sc
 	}
 }
 
-static void RenderDirectionalLights(FRHICommandList& RHICmdList, const FScene& Scene, const FViewInfo& View, const FCachedLightMaterial& DefaultLightMaterial)
+static void RenderDirectionalLights(FRHICommandList& RHICmdList, const FScene& Scene, const FViewInfo& View, const FCachedLightMaterial& DefaultLightMaterial, EMobileSSRQuality MobileSSRQuality)
 {
 	uint32 NumLights = 0;
 	for (uint32 ChannelIdx = 0; ChannelIdx < UE_ARRAY_COUNT(Scene.MobileDirectionalLights); ChannelIdx++)
@@ -635,7 +700,7 @@ static void RenderDirectionalLights(FRHICommandList& RHICmdList, const FScene& S
 
 	if (!bInlineReflectionAndSky)
 	{
-		RenderReflectionEnvironmentSkyLighting(RHICmdList, Scene, View);
+		RenderReflectionEnvironmentSkyLighting(RHICmdList, Scene, View, MobileSSRQuality);
 	}
 
 	for (uint32 ChannelIdx = 0; ChannelIdx < UE_ARRAY_COUNT(Scene.MobileDirectionalLights); ChannelIdx++)
@@ -643,7 +708,7 @@ static void RenderDirectionalLights(FRHICommandList& RHICmdList, const FScene& S
 		FLightSceneInfo* DirectionalLight = Scene.MobileDirectionalLights[ChannelIdx];
 		if (DirectionalLight)
 		{
-			RenderDirectionalLight(RHICmdList, Scene, View, DefaultLightMaterial, *DirectionalLight, ChannelIdx, bInlineReflectionAndSky);
+			RenderDirectionalLight(RHICmdList, Scene, View, DefaultLightMaterial, *DirectionalLight, ChannelIdx, bInlineReflectionAndSky, MobileSSRQuality);
 		}
 	}
 }
@@ -1007,7 +1072,8 @@ void MobileDeferredShadingPass(
 	const FViewInfo& View,
 	const FScene& Scene, 
 	const FSortedLightSetSceneInfo& SortedLightSet,
-	const TArray<FVisibleLightInfo, SceneRenderingAllocator>& VisibleLightInfos)
+	const TArray<FVisibleLightInfo, SceneRenderingAllocator>& VisibleLightInfos,
+	EMobileSSRQuality MobileSSRQuality)
 {
 	SCOPED_DRAW_EVENT(RHICmdList, DeferredShading);
 	SCOPED_GPU_STAT(RHICmdList, DeferredShading);
@@ -1020,7 +1086,7 @@ void MobileDeferredShadingPass(
 	DefaultMaterial.Material = DefaultMaterial.MaterialProxy->GetMaterialNoFallback(ERHIFeatureLevel::ES3_1);
 	check(DefaultMaterial.Material);
 
-	RenderDirectionalLights(RHICmdList, Scene, View, DefaultMaterial);
+	RenderDirectionalLights(RHICmdList, Scene, View, DefaultMaterial, MobileSSRQuality);
 	
 	const bool bMobileUseClusteredDeferredShading = UseClusteredDeferredShading(View.GetShaderPlatform());
 	if (!bMobileUseClusteredDeferredShading)
diff --git a/Engine/Source/Runtime/Renderer/Private/MobileDeferredShadingPass.h b/Engine/Source/Runtime/Renderer/Private/MobileDeferredShadingPass.h
--- a/Engine/Source/Runtime/Renderer/Private/MobileDeferredShadingPass.h
+++ b/Engine/Source/Runtime/Renderer/Private/MobileDeferredShadingPass.h
@@ -6,6 +6,7 @@
 #include "RenderGraphDefinitions.h"
 #include "SceneRendering.h"
 
+enum class EMobileSSRQuality;
 class FMobileSceneTextureUniformParameters;
 class FRDGBuilder;
 class FRHICommandListImmediate;
@@ -23,4 +24,5 @@ void MobileDeferredShadingPass(
 	const FViewInfo& View,
 	const FScene& Scene, 
 	const FSortedLightSetSceneInfo& SortedLightSet,
-	const TArray<FVisibleLightInfo, SceneRenderingAllocator>& VisibleLightInfos);
+	const TArray<FVisibleLightInfo, SceneRenderingAllocator>& VisibleLightInfos,
+	EMobileSSRQuality MobileSSRQuality);
diff --git a/Engine/Source/Runtime/Renderer/Private/MobileSSR.cpp b/Engine/Source/Runtime/Renderer/Private/MobileSSR.cpp
--- a/Engine/Source/Runtime/Renderer/Private/MobileSSR.cpp
+++ b/Engine/Source/Runtime/Renderer/Private/MobileSSR.cpp
@@ -17,65 +17,47 @@ static FAutoConsoleVariableRef CVarMobileScreenSpaceReflections(
 		 "1: Mobile Renderer Screen Space Reflections enabled\n"),
 	ECVF_ReadOnly | ECVF_RenderThreadSafe);
 
-bool ShouldRenderMobileSSR(const FViewInfo& View)
-{
-	return GMobileScreenSpaceReflectionsSupported && ScreenSpaceRayTracing::ShouldRenderScreenSpaceReflections(View) && View.PrevViewInfo.TemporalAAHistory.RT[0] && View.HZB;
-}
+float GMobileMobileSSRIntensity = -1.0f;
+static FAutoConsoleVariableRef CVarMobileMobileSSRIntensity(
+	TEXT("r.Mobile.ScreenSpaceReflections.Intensity"),
+	GMobileMobileSSRIntensity,
+	TEXT("Scale factor to adjust the intensity of mobile screen space reflections in the range [0.0, 1.0] or -1. (default: -1, ignores this setting)\n"),
+	ECVF_RenderThreadSafe);
 
 bool IsMobileSSREnabled(const FViewInfo& View)
 {
-	return GMobileScreenSpaceReflectionsSupported && ScreenSpaceRayTracing::ShouldRenderScreenSpaceReflections(View);
+	return GMobileScreenSpaceReflectionsSupported && ScreenSpaceRayTracing::ShouldRenderScreenSpaceReflections(View) && (GMobileMobileSSRIntensity < 0.0f || GMobileMobileSSRIntensity > 0.0f);
 }
 
-class FMobileSSRQualityDim : SHADER_PERMUTATION_ENUM_CLASS("SSR_QUALITY", ESSRQuality);
-class FMobileSSRUseVelocity : SHADER_PERMUTATION_BOOL("SSR_USE_VELOCITY");
-
-class FMobileScreenSpaceReflectionsPS : public FGlobalShader
+EMobileSSRQuality ActiveMobileSSRQuality(const FViewInfo& View, bool bHasVelocityTexture)
 {
-	DECLARE_GLOBAL_SHADER(FMobileScreenSpaceReflectionsPS);
-	SHADER_USE_PARAMETER_STRUCT(FMobileScreenSpaceReflectionsPS, FGlobalShader);
-
-	using FPermutationDomain = TShaderPermutationDomain<FMobileSSRQualityDim, FMobileSSRUseVelocity>;
-
-	static bool ShouldCompilePermutation(const FGlobalShaderPermutationParameters& Parameters)
+	if (!IsMobileSSREnabled(View))
 	{
-		return IsMobilePlatform(Parameters.Platform) && IsMobileDeferredShadingEnabled(Parameters.Platform);
+		return EMobileSSRQuality::Disabled;
 	}
-
-	static void ModifyCompilationEnvironment(const FGlobalShaderPermutationParameters& Parameters, FShaderCompilerEnvironment& OutEnvironment)
+	ESSRQuality SSRQuality;
+	IScreenSpaceDenoiser::FReflectionsRayTracingConfig DenoiserConfig;
+	ScreenSpaceRayTracing::GetSSRQualityForView(View, &SSRQuality, &DenoiserConfig);
+	if (SSRQuality < ESSRQuality::Low)
 	{
-		FGlobalShader::ModifyCompilationEnvironment(Parameters, OutEnvironment);
-		OutEnvironment.SetDefine(TEXT("SUPPORTS_ANISOTROPIC_MATERIALS"), FDataDrivenShaderPlatformInfo::GetSupportsAnisotropicMaterials(Parameters.Platform));
-
-		check(IsMobileDeferredShadingEnabled(Parameters.Platform));
-		OutEnvironment.SetDefine(TEXT("ENABLE_SHADINGMODEL_SUPPORT_MOBILE_DEFERRED"), MobileUsesGBufferCustomData(Parameters.Platform));
-		OutEnvironment.SetDefine(TEXT("IS_MOBILE_DEFERREDSHADING_SUBPASS"), 1u);
-		OutEnvironment.SetDefine(TEXT("FORCE_DEPTH_TEXTURE_READS"), 1u);
-		
-		OutEnvironment.SetDefine(TEXT("SSR_OUTPUT_FOR_DENOISER"), 0u);
-		OutEnvironment.SetDefine(TEXT("DIM_LIGHTING_TERM"), 0u);
+		return EMobileSSRQuality::Disabled;
+	}
+	if (!View.PrevViewInfo.TemporalAAHistory.RT[0])
+	{
+		return EMobileSSRQuality::Disabled;
 	}
 
-	BEGIN_SHADER_PARAMETER_STRUCT(FParameters, )
-		SHADER_PARAMETER(FLinearColor, SSRParams)
-		SHADER_PARAMETER_STRUCT_REF(FViewUniformShaderParameters, ViewUniformBuffer)
-	END_SHADER_PARAMETER_STRUCT()
-};
-
-IMPLEMENT_GLOBAL_SHADER(FMobileScreenSpaceReflectionsPS, "/Engine/Private/SSRT/SSRTReflections.usf", "MobileScreenSpaceReflectionsPS", SF_Pixel);
-
-
-using DepthStencilStateMobileSSR = TStaticDepthStencilState<
-	false, CF_Always,
-	true, CF_Equal, SO_Keep, SO_Keep, SO_Keep,
-	false, CF_Always, SO_Keep, SO_Keep, SO_Keep,
-	STENCIL_MOBILE_REFLECTIVE_MASK | STENCIL_MOBILE_SKY_MASK, 0x00>;
+	if (SSRQuality >= ESSRQuality::Medium && bHasVelocityTexture)
+	{
+		return EMobileSSRQuality::Medium;
+	}
 
-const uint8 StencilRefMobileSSR = STENCIL_MOBILE_REFLECTIVE_MASK; // Process SSR where there are potentially reflective materials (and not sky)
+	return  EMobileSSRQuality::Low;
+}
 
 void SetupMobileSSRParameters(FRDGBuilder& GraphBuilder,const FViewInfo& View, FMobileScreenSpaceReflectionParams& Params)
 {
-	if (!ShouldRenderMobileSSR(View))
+	if (!IsMobileSSREnabled(View) || !View.PrevViewInfo.TemporalAAHistory.RT[0] || !(View.HZB || View.PrevViewInfo.HZB))
 	{
 		const FRDGSystemTextures& SystemTextures = FRDGSystemTextures::Get(GraphBuilder);
 		Params.HZB = SystemTextures.Black;
@@ -83,7 +65,7 @@ void SetupMobileSSRParameters(FRDGBuilder& GraphBuilder,const FViewInfo& View, F
 		Params.SceneColor = GraphBuilder.CreateSRV(FRDGTextureSRVDesc(SystemTextures.Black));
 		Params.SceneColorSampler = TStaticSamplerState<SF_Bilinear, AM_Clamp, AM_Clamp, AM_Clamp>::GetRHI();
 		Params.PrevSceneColorBilinearUVMinMax = FVector4f(0.0f, 0.0f, 1.0f, 1.0f);
-		Params.QualityAndExposureCorrection = FVector4f(ForceInitToZero);
+		Params.IntensityAndExposureCorrection = FVector4f(ForceInitToZero);
 		return;
 	}
 	FRDGTextureRef SceneColor = GraphBuilder.RegisterExternalTexture(View.PrevViewInfo.TemporalAAHistory.RT[0]);
@@ -91,13 +73,16 @@ void SetupMobileSSRParameters(FRDGBuilder& GraphBuilder,const FViewInfo& View, F
 		? FRDGTextureSRVDesc::CreateForSlice(SceneColor, View.PrevViewInfo.TemporalAAHistory.OutputSliceIndex)
 		: FRDGTextureSRVDesc(SceneColor));
 	Params.SceneColorSampler = TStaticSamplerState<SF_Point>::GetRHI();
-	Params.HZB = View.HZB;
+	FRDGTextureRef HZB = View.HZB ? View.HZB : GraphBuilder.RegisterExternalTexture(View.PrevViewInfo.HZB);
+	Params.HZB = HZB;
 	Params.HZBSampler = TStaticSamplerState<SF_Point>::GetRHI();
 
 	{
+		const FIntPoint HZBMipmap0Size = HZB->Desc.Extent;
+		const FIntRect ViewRect = View.HZB ? View.ViewRect : View.PrevViewInfo.ViewRect;
 		const FVector2D HZBUvFactor(
-			float(View.ViewRect.Width()) / float(2 * View.HZBMipmap0Size.X),
-			float(View.ViewRect.Height()) / float(2 * View.HZBMipmap0Size.Y));
+			float(ViewRect.Width()) / float(2 * HZBMipmap0Size.X),
+			float(ViewRect.Height()) / float(2 * HZBMipmap0Size.Y));
 		Params.HZBUvFactorAndInvFactor = FVector4f(
 			HZBUvFactor.X,
 			HZBUvFactor.Y,
@@ -126,10 +111,13 @@ void SetupMobileSSRParameters(FRDGBuilder& GraphBuilder,const FViewInfo& View, F
 		IScreenSpaceDenoiser::FReflectionsRayTracingConfig DenoiserConfig;
 		ScreenSpaceRayTracing::GetSSRQualityForView(View, &SSRQuality, &DenoiserConfig);
 
-		Params.QualityAndExposureCorrection.X = (float)SSRQuality;
-		Params.QualityAndExposureCorrection.Y = 1.0f / View.PrevViewInfo.SceneColorPreExposure;
-		Params.QualityAndExposureCorrection.Z = View.PreExposure / View.PrevViewInfo.SceneColorPreExposure;
-		Params.QualityAndExposureCorrection.W = FMath::Clamp(View.FinalPostProcessSettings.ScreenSpaceReflectionMaxRoughness, 0.0f, 1.0f);
+		const float MobileMobileSSRIntensity = GMobileMobileSSRIntensity >= 0.0f ? GMobileMobileSSRIntensity : 1.0f;
+		Params.IntensityAndExposureCorrection.X = SSRQuality > ESSRQuality::VisualizeSSR ? FMath::Clamp(View.FinalPostProcessSettings.ScreenSpaceReflectionIntensity * 0.01f * MobileMobileSSRIntensity, 0.0f, 1.0f) : 0.0f;
+		Params.IntensityAndExposureCorrection.Y = 1.f / View.PrevViewInfo.SceneColorPreExposure;
+		float MaxRoughness = FMath::Clamp(View.FinalPostProcessSettings.ScreenSpaceReflectionMaxRoughness, 0.01f, 0.6f);
+		MaxRoughness *= 0.5;
+		Params.IntensityAndExposureCorrection.Z = MaxRoughness;
+		Params.IntensityAndExposureCorrection.W = 2.0 / MaxRoughness;
 	}
 
 	{
@@ -139,42 +127,21 @@ void SetupMobileSSRParameters(FRDGBuilder& GraphBuilder,const FViewInfo& View, F
 		Params.PrevSceneColorBilinearUVMinMax.Z = PrevSceneColorParameters.UVViewportBilinearMax.X;
 		Params.PrevSceneColorBilinearUVMinMax.W = PrevSceneColorParameters.UVViewportBilinearMax.Y;
 	}
-}
-
-void FMobileSceneRenderer::RenderSSR(FRHICommandList& RHICmdList, const FViewInfo& View)
-{
-	if (!ShouldRenderMobileSSR(View))
-	{
-		return;
-	}
 
-	FMobileScreenSpaceReflectionsPS::FParameters Parameters;
-
-	ESSRQuality SSRQuality;
-	IScreenSpaceDenoiser::FReflectionsRayTracingConfig DenoiserConfig;
-	ScreenSpaceRayTracing::GetSSRQualityForView(View, &SSRQuality, &DenoiserConfig);
-	if (SSRQuality < ESSRQuality::Low)
+	switch (View.AntiAliasingMethod)
 	{
-		return;
+	default:
+		// Without TAA disable temporal noise and reduce intensity of SSR to hide the noise.
+		Params.NoiseIndex = 0;
+		if (GMobileMobileSSRIntensity < 0.0f)
+		{
+			Params.IntensityAndExposureCorrection.X = FMath::Min(Params.IntensityAndExposureCorrection.X, 0.4f);
+		}
+		break;
+	case AAM_TemporalAA:
+	case AAM_TSR:
+		Params.NoiseIndex = View.ViewState ? View.ViewState->GetFrameIndex() % 8 : 0;
+		break;
 	}
-
-	Parameters.SSRParams = ScreenSpaceRayTracing::ComputeSSRParams(View, SSRQuality, false);
-
-	Parameters.ViewUniformBuffer = View.ViewUniformBuffer;
-
-	FMobileScreenSpaceReflectionsPS::FPermutationDomain PermutationVectorPS;
-	PermutationVectorPS.Set<FMobileSSRQualityDim>(SSRQuality);
-	PermutationVectorPS.Set<FMobileSSRUseVelocity>(bShouldRenderVelocities);
-
-	TShaderMapRef<FMobileScreenSpaceReflectionsPS> PixelShader(View.ShaderMap, PermutationVectorPS);
-	FGraphicsPipelineStateInitializer GraphicsPSOInit;
-	FPixelShaderUtils::InitFullscreenPipelineState(RHICmdList, View.ShaderMap, PixelShader, /* out */ GraphicsPSOInit);
-
-	GraphicsPSOInit.BlendState = TStaticBlendState<CW_RGB, BO_Add, BF_One, BF_One>::GetRHI();
-	GraphicsPSOInit.DepthStencilState = DepthStencilStateMobileSSR::GetRHI();
-
-	SetGraphicsPipelineState(RHICmdList, GraphicsPSOInit, StencilRefMobileSSR);
-	SetShaderParameters(RHICmdList, PixelShader, PixelShader.GetPixelShader(), Parameters);
-
-	FPixelShaderUtils::DrawFullscreenTriangle(RHICmdList);
 }
+
diff --git a/Engine/Source/Runtime/Renderer/Private/MobileSSR.h b/Engine/Source/Runtime/Renderer/Private/MobileSSR.h
--- a/Engine/Source/Runtime/Renderer/Private/MobileSSR.h
+++ b/Engine/Source/Runtime/Renderer/Private/MobileSSR.h
@@ -15,10 +15,18 @@ SHADER_PARAMETER_SAMPLER(SamplerState, HZBSampler)
 SHADER_PARAMETER(FVector4f, HZBUvFactorAndInvFactor)
 SHADER_PARAMETER(FVector4f, PrevScreenPositionScaleBias)
 SHADER_PARAMETER(FVector4f, PrevSceneColorBilinearUVMinMax)
-SHADER_PARAMETER(FVector4f, QualityAndExposureCorrection) // .x = SSRQuality, .y = PrevSceneColorPreExposureInv, .z = PrevSceneColorPreExposureCorrection, .w = View.FinalPostProcessSettings.ScreenSpaceReflectionMaxRoughness
+SHADER_PARAMETER(FVector4f, IntensityAndExposureCorrection) // .x = Intensity, .y = PrevSceneColorPreExposureInv, .z = View.FinalPostProcessSettings.ScreenSpaceReflectionMaxRoughness, .w = 2DivMaxRoughness
+SHADER_PARAMETER(uint32, NoiseIndex)
 END_SHADER_PARAMETER_STRUCT()
 
-bool IsMobileSSREnabled(const FViewInfo& View);
-bool ShouldRenderMobileSSR(const FViewInfo& View);
+enum class EMobileSSRQuality
+{
+	Disabled,
+	Low,
+	Medium,
+	MAX
+};
 
+bool IsMobileSSREnabled(const FViewInfo& View);
+EMobileSSRQuality ActiveMobileSSRQuality(const FViewInfo& View, bool bHasVelocityTexture);
 void SetupMobileSSRParameters(FRDGBuilder& GraphBuilder, const FViewInfo& View, FMobileScreenSpaceReflectionParams& Params);
diff --git a/Engine/Source/Runtime/Renderer/Private/MobileShadingRenderer.cpp b/Engine/Source/Runtime/Renderer/Private/MobileShadingRenderer.cpp
--- a/Engine/Source/Runtime/Renderer/Private/MobileShadingRenderer.cpp
+++ b/Engine/Source/Runtime/Renderer/Private/MobileShadingRenderer.cpp
@@ -303,6 +303,7 @@ FMobileSceneRenderer::FMobileSceneRenderer(const FSceneViewFamily* InViewFamily,
 	bIsMaskedOnlyDepthPrepassEnabled = Scene->EarlyZPassMode == DDM_MaskedOnly;
 	bEnableClusteredLocalLights = MobileForwardEnableLocalLights(ShaderPlatform);
 	bEnableClusteredReflections = MobileForwardEnableClusteredReflections(ShaderPlatform);
+	bRequiresScreenSpaceReflections = AreMobileScreenSpaceReflectionsEnabled(ShaderPlatform);
 	
 	StandardTranslucencyPass = ViewFamily.AllowTranslucencyAfterDOF() ? ETranslucencyPass::TPT_TranslucencyStandard : ETranslucencyPass::TPT_AllTranslucency;
 	StandardTranslucencyMeshPass = TranslucencyPassToMeshPass(StandardTranslucencyPass);
@@ -572,7 +573,7 @@ void FMobileSceneRenderer::InitViews(
 		&& !ViewFamily.EngineShowFlags.HitProxies
 		&& !ViewFamily.EngineShowFlags.VisualizeLightCulling
 		&& !ViewFamily.UseDebugViewPS()
-		&& bRendererOutputFinalSceneColor;
+		&& bRendererOutputFinalSceneColor; 
 
 	bShouldRenderHZB = ShouldRenderHZB(Views) && bRendererOutputFinalSceneColor;
 
@@ -595,6 +596,7 @@ void FMobileSceneRenderer::InitViews(
 		Views[0].AntiAliasingMethod == AAM_TemporalAA ||
 		bRequireSeparateViewPass ||
 		bIsFullDepthPrepassEnabled ||
+		bShouldRenderHZB ||
 		GraphBuilder.IsDumpingFrame();
 	// never keep MSAA depth if SceneDepthAux is enabled
 	bKeepDepthContent = ((NumMSAASamples > 1) && bRequiresSceneDepthAux) ? false : bKeepDepthContent;
@@ -1323,7 +1325,6 @@ void FMobileSceneRenderer::Render(FRDGBuilder& GraphBuilder)
 		}
 
 		// Local Light prepass
-
 		if (bRendererOutputFinalSceneColor)
 		{
 			RenderMobileLocalLightsBuffer(GraphBuilder, SceneTextures, SortedLightSet);
@@ -1473,6 +1474,11 @@ void FMobileSceneRenderer::Render(FRDGBuilder& GraphBuilder)
 	GEngine->GetPostRenderDelegateEx().Broadcast(GraphBuilder);
 	GetSceneExtensionsRenderers().PostRender(GraphBuilder);
 
+	if (bRendererOutputFinalSceneColor && bShouldRenderHZB && !bRequiresMultiPass)
+	{
+		RenderHZB(GraphBuilder, SceneTextures.Depth.Resolve);
+	}
+
 	OnRenderFinish(GraphBuilder, ViewFamilyTexture);
 
 	if (bRendererOutputFinalSceneColor)
@@ -1748,6 +1754,10 @@ void FMobileSceneRenderer::RenderForwardMultiPass(FRDGBuilder& GraphBuilder, FMo
 	{
 		AddResolveSceneColorPass(GraphBuilder, View, SceneTextures.DepthAux);
 	}
+	if (bShouldRenderHZB && !bIsFullDepthPrepassEnabled)
+	{
+		RenderHZB(GraphBuilder, SceneTextures.Depth.Resolve);
+	}
 
 	FExclusiveDepthStencil::Type ExclusiveDepthStencil = FExclusiveDepthStencil::DepthRead_StencilRead;
 	if (bModulatedShadowsInUse)
@@ -1979,6 +1989,7 @@ void FMobileSceneRenderer::RenderDeferredSinglePass(FRDGBuilder& GraphBuilder, F
 		PassParameters->HalfResLocalFogVolumeDepthSRV = View.LocalFogVolumeViewData.HalfResLocalFogVolumeDepthSRV;
 		PassParameters->RenderTargets = BasePassRenderTargets;
 		PassParameters->RenderTargets.SubpassHint = ESubpassHint::DeferredShadingSubpass;
+		const EMobileSSRQuality MobileSSRQuality = ActiveMobileSSRQuality(View, bShouldRenderVelocities);
 		const bool bDoOcclusionQueires = (!bIsFullDepthPrepassEnabled && ViewContext.bIsLastView && DoOcclusionQueries());
 		PassParameters->RenderTargets.NumOcclusionQueries = bDoOcclusionQueires ? ComputeNumOcclusionQueriesToBatch() : 0u;
 
@@ -1989,7 +2000,7 @@ void FMobileSceneRenderer::RenderDeferredSinglePass(FRDGBuilder& GraphBuilder, F
 			PassParameters,
 			// the second view pass should not be merged with the first view pass on mobile since the subpass would not work properly.
 			ERDGPassFlags::Raster | ERDGPassFlags::NeverMerge,
-			[this, PassParameters, ViewContext, bDoOcclusionQueires, &SceneTextures, &SortedLightSet, bUsingPixelLocalStorage](FRHICommandList& RHICmdList)
+			[this, PassParameters, ViewContext, bDoOcclusionQueires, MobileSSRQuality, &SortedLightSet, bUsingPixelLocalStorage](FRHICommandList& RHICmdList)
 			{
 				FViewInfo& View = *ViewContext.ViewInfo;
 
@@ -2005,8 +2016,7 @@ void FMobileSceneRenderer::RenderDeferredSinglePass(FRDGBuilder& GraphBuilder, F
 				RenderDecals(RHICmdList, View);
 				// SceneColor write, SceneDepth is read only
 				RHICmdList.NextSubpass();
-				MobileDeferredShadingPass(RHICmdList, ViewContext.ViewIndex, Views.Num(), View, *Scene, SortedLightSet, VisibleLightInfos);
-				RenderSSR(RHICmdList, View);
+				MobileDeferredShadingPass(RHICmdList, ViewContext.ViewIndex, Views.Num(), View, *Scene, SortedLightSet, VisibleLightInfos, MobileSSRQuality);
 
 				if (bUsingPixelLocalStorage)
 				{
@@ -2101,6 +2111,11 @@ void FMobileSceneRenderer::RenderDeferredMultiPass(FRDGBuilder& GraphBuilder, FS
 			});
 	}
 
+	if (bShouldRenderHZB && !bIsFullDepthPrepassEnabled)
+	{
+		RenderHZB(GraphBuilder, SceneTextures.Depth.Target);
+	}
+
 	BasePassRenderTargets.Enumerate([](FRenderTargetBinding& RenderTarget) {
 		RenderTarget.SetLoadAction(ERenderTargetLoadAction::ELoad);
 		});
@@ -2168,17 +2183,17 @@ void FMobileSceneRenderer::RenderDeferredMultiPass(FRDGBuilder& GraphBuilder, FS
 
 		View.ParallelMeshDrawCommandPasses[StandardTranslucencyMeshPass].BuildRenderingCommands(GraphBuilder, Scene->GPUScene, TranslucencyInstanceCullingDrawParams);
 		PassParameters->InstanceCullingDrawParams = TranslucencyInstanceCullingDrawParams;
+		const EMobileSSRQuality MobileSSRQuality = ActiveMobileSSRQuality(View, bShouldRenderVelocities);
 
 		GraphBuilder.AddPass(
 			RDG_EVENT_NAME("LightingAndTranslucency"),
 			PassParameters,
 			ERDGPassFlags::Raster,
-			[this, PassParameters, ViewContext, &SceneTextures, &SortedLightSet](FRHICommandList& RHICmdList)
+			[this, PassParameters, ViewContext, MobileSSRQuality, &SceneTextures, &SortedLightSet](FRHICommandList& RHICmdList)
 			{
 				FViewInfo& View = *ViewContext.ViewInfo;
 
-				MobileDeferredShadingPass(RHICmdList, ViewContext.ViewIndex, Views.Num(), View, *Scene, SortedLightSet, VisibleLightInfos);
-				RenderSSR(RHICmdList, View);
+				MobileDeferredShadingPass(RHICmdList, ViewContext.ViewIndex, Views.Num(), View, *Scene, SortedLightSet, VisibleLightInfos, MobileSSRQuality);
 				RenderFog(RHICmdList, View);
 
 				// Draw translucency.
@@ -2454,7 +2469,7 @@ void FMobileSceneRenderer::RenderHZB(FRDGBuilder& GraphBuilder, FRDGTextureRef S
 
 				if (View.ViewState)
 				{
-					if (FInstanceCullingContext::IsOcclusionCullingEnabled())
+					if (FInstanceCullingContext::IsOcclusionCullingEnabled() || (AreMobileScreenSpaceReflectionsEnabled(ShaderPlatform) && !bIsFullDepthPrepassEnabled))
 					{
 						GraphBuilder.QueueTextureExtraction(FurthestHZBTexture, &View.ViewState->PrevFrameViewInfo.HZB);
 					}
diff --git a/Engine/Source/Runtime/Renderer/Private/SceneRendering.h b/Engine/Source/Runtime/Renderer/Private/SceneRendering.h
--- a/Engine/Source/Runtime/Renderer/Private/SceneRendering.h
+++ b/Engine/Source/Runtime/Renderer/Private/SceneRendering.h
@@ -2666,9 +2666,6 @@ protected:
 	/** Renders the base pass for translucency. */
 	void RenderTranslucency(FRHICommandList& RHICmdList, const FViewInfo& View);
 
-	/** Renders screen space reflections. */
-	void RenderSSR(FRHICommandList& RHICmdList, const FViewInfo& View);
-
 	/** On chip pre-tonemap before scene color MSAA resolve (iOS only) */
 	void PreTonemapMSAA(FRHICommandList& RHICmdList, const FMinimalSceneTextures& SceneTextures);
 
@@ -2714,6 +2711,7 @@ private:
 	bool bRequiresAmbientOcclusionPass;
 	bool bShouldRenderVelocities;
 	bool bShouldRenderHZB;
+	bool bRequiresScreenSpaceReflections;
 	bool bIsFullDepthPrepassEnabled;
 	bool bIsMaskedOnlyDepthPrepassEnabled;
 	bool bRequiresSceneDepthAux;
