Possible build fix for StressTestUnity build test.

diff --git a/Engine/Source/Runtime/Renderer/Private/VT/MeshPaintVirtualTextureSceneExtension.cpp b/Engine/Source/Runtime/Renderer/Private/VT/MeshPaintVirtualTextureSceneExtension.cpp
--- a/Engine/Source/Runtime/Renderer/Private/VT/MeshPaintVirtualTextureSceneExtension.cpp
+++ b/Engine/Source/Runtime/Renderer/Private/VT/MeshPaintVirtualTextureSceneExtension.cpp
@@ -32,14 +32,14 @@ END_SHADER_PARAMETER_STRUCT()
 
 DECLARE_SCENE_UB_STRUCT(FMeshPaintTextureParameters, MeshPaint, RENDERER_API)
 
-static void ConstructDefault(FMeshPaintTextureParameters& Parameters, FRDGBuilder& GraphBuilder)
+static void GetDefaultMeshPaintParameters(FMeshPaintTextureParameters& Parameters, FRDGBuilder& GraphBuilder)
 {
 	Parameters.PageTableTexture = GBlackUintTexture->TextureRHI;
 	Parameters.PhysicalTexture = GBlackTextureWithSRV->TextureRHI;
 	Parameters.PackedUniform = FUintVector4(0, 0, 0, 0);
 }
 
-IMPLEMENT_SCENE_UB_STRUCT(FMeshPaintTextureParameters, MeshPaint, ConstructDefault);
+IMPLEMENT_SCENE_UB_STRUCT(FMeshPaintTextureParameters, MeshPaint, GetDefaultMeshPaintParameters);
 
 void FMeshPaintVirtualTextureSceneExtension::FRenderer::UpdateSceneUniformBuffer(FRDGBuilder& GraphBuilder, FSceneUniformBuffer& SceneUniformBuffer)
 {
