Instance occlusion culling now supports multiple HZBs

diff --git a/Engine/Source/Runtime/Renderer/Private/InstanceCulling/InstanceCullingContext.cpp b/Engine/Source/Runtime/Renderer/Private/InstanceCulling/InstanceCullingContext.cpp
--- a/Engine/Source/Runtime/Renderer/Private/InstanceCulling/InstanceCullingContext.cpp
+++ b/Engine/Source/Runtime/Renderer/Private/InstanceCulling/InstanceCullingContext.cpp
@@ -613,8 +613,8 @@ const TRDGUniformBufferRef<FInstanceCullingGlobalUniforms> FInstanceCullingConte
 class FInstanceCullingDeferredContext : public FInstanceCullingMergedContext
 {
 public:
-	FInstanceCullingDeferredContext(EShaderPlatform InShaderPlatform, FInstanceCullingManager* InInstanceCullingManager = nullptr)
-		: FInstanceCullingMergedContext(InShaderPlatform)
+	FInstanceCullingDeferredContext(EShaderPlatform InShaderPlatform, FInstanceCullingManager* InInstanceCullingManager = nullptr, int32 InNumBins=2)
+		: FInstanceCullingMergedContext(InShaderPlatform, false, InNumBins)
 		, InstanceCullingManager(InInstanceCullingManager)
 	{}
 
@@ -627,7 +627,7 @@ public:
 
 	bool bProcessed = false;
 
-	void ProcessBatched(TStaticArray<FBuildInstanceIdBufferAndCommandsFromPrimitiveIdsCs::FParameters*, static_cast<uint32>(EBatchProcessingMode::Num)> PassParameters);
+	void ProcessBatched(TArray<FBuildInstanceIdBufferAndCommandsFromPrimitiveIdsCs::FParameters*> PassParameters);
 
 #if MESH_DRAW_COMMAND_STATS
 	FRHIGPUBufferReadback* MeshDrawCommandStatsIndirectArgsReadbackBuffer = nullptr;
@@ -983,7 +983,7 @@ void FInstanceCullingContext::BuildRenderingCommandsInternal(
 #endif
 }
 
-void FInstanceCullingDeferredContext::ProcessBatched(TStaticArray<FBuildInstanceIdBufferAndCommandsFromPrimitiveIdsCs::FParameters*, static_cast<uint32>(EBatchProcessingMode::Num)> PassParameters)
+void FInstanceCullingDeferredContext::ProcessBatched(TArray<FBuildInstanceIdBufferAndCommandsFromPrimitiveIdsCs::FParameters*> PassParameters)
 {
 	if (bProcessed)
 	{
@@ -1009,22 +1009,13 @@ void FInstanceCullingDeferredContext::ProcessBatched(TStaticArray<FBuildInstance
 
 	bProcessed = true;
 
-
 	// Finalize culling pass parameters
-	for (uint32 Mode = 0U; Mode < uint32(EBatchProcessingMode::Num); ++Mode)
+	for (uint32 BinIndex = 0U; BinIndex < uint32(LoadBalancers.Num()); ++BinIndex)
 	{
-		PassParameters[Mode]->NumViewIds = ViewIds.Num();
-		PassParameters[Mode]->LoadBalancerParameters.NumBatches = LoadBalancers[Mode].GetBatches().Num();
-		PassParameters[Mode]->LoadBalancerParameters.NumItems = LoadBalancers[Mode].GetItems().Num();
-		PassParameters[Mode]->LoadBalancerParameters.NumGroupsPerBatch = 1;
-
-		const bool bOcclusionCullInstances = PrevHZB != nullptr && FInstanceCullingContext::IsOcclusionCullingEnabled();
-		if (bOcclusionCullInstances)
-		{
-			PassParameters[Mode]->HZBTexture = PrevHZB;
-			PassParameters[Mode]->HZBSize = PrevHZB->Desc.Extent;
-			PassParameters[Mode]->HZBSampler = TStaticSamplerState< SF_Point, AM_Clamp, AM_Clamp, AM_Clamp >::GetRHI();
-		}
+		PassParameters[BinIndex]->NumViewIds = ViewIds.Num();
+		PassParameters[BinIndex]->LoadBalancerParameters.NumBatches = LoadBalancers[BinIndex].GetBatches().Num();
+		PassParameters[BinIndex]->LoadBalancerParameters.NumItems = LoadBalancers[BinIndex].GetItems().Num();
+		PassParameters[BinIndex]->LoadBalancerParameters.NumGroupsPerBatch = 1;
 	}
 }
 
@@ -1045,8 +1036,13 @@ FInstanceCullingDeferredContext *FInstanceCullingContext::CreateDeferredContext(
 	FGPUScene& GPUScene,
 	FInstanceCullingManager& InstanceCullingManager)
 {
-#define INST_CULL_CALLBACK_MODE(CustomCode) \
-	[PassParameters, DeferredContext, Mode]() \
+	// Bin 0 is reserved for UnCulled batches, every other bin is for each HZB. So at the very least we must have 2 bins.
+	// Generic batches with a null HZB will go in bin 1, together with the ones associated to the first HZB.
+
+	uint32 NumBins = FMath::Max(2, InstanceCullingManager.ViewPrevHZBs.Num() + 1);
+
+#define INST_CULL_CALLBACK_BIN_INDEX(CustomCode) \
+	[PassParameters, DeferredContext, BinIndex]() \
 	{ \
 		DeferredContext->ProcessBatched(PassParameters); \
 		return CustomCode; \
@@ -1067,18 +1063,18 @@ FInstanceCullingDeferredContext *FInstanceCullingContext::CreateDeferredContext(
 	INST_CULL_CALLBACK(DeferredContext->ArrayName.GetData()), \
 	INST_CULL_CALLBACK(DeferredContext->ArrayName.Num() * DeferredContext->ArrayName.GetTypeSize())
 
-#define INST_CULL_CREATE_STRUCT_BUFF_ARGS_MODE(ArrayName) \
+#define INST_CULL_CREATE_STRUCT_BUFF_ARGS_BIN_INDEX(ArrayName) \
 	GraphBuilder, \
 	TEXT("InstanceCulling.") TEXT(#ArrayName), \
-	DeferredContext->ArrayName[Mode].GetTypeSize(), \
-	INST_CULL_CALLBACK_MODE(DeferredContext->ArrayName[Mode].Num()), \
-	INST_CULL_CALLBACK_MODE(DeferredContext->ArrayName[Mode].GetData()), \
-	INST_CULL_CALLBACK_MODE(DeferredContext->ArrayName[Mode].Num() * DeferredContext->ArrayName[Mode].GetTypeSize())
+	DeferredContext->ArrayName[BinIndex].GetTypeSize(), \
+	INST_CULL_CALLBACK_BIN_INDEX(DeferredContext->ArrayName[BinIndex].Num()), \
+	INST_CULL_CALLBACK_BIN_INDEX(DeferredContext->ArrayName[BinIndex].GetData()), \
+	INST_CULL_CALLBACK_BIN_INDEX(DeferredContext->ArrayName[BinIndex].Num() * DeferredContext->ArrayName[BinIndex].GetTypeSize())
 
 	const ERHIFeatureLevel::Type FeatureLevel = GPUScene.GetFeatureLevel();
 	const EShaderPlatform ShaderPlatform = GPUScene.GetShaderPlatform();
 
-	FInstanceCullingDeferredContext* DeferredContext = GraphBuilder.AllocObject<FInstanceCullingDeferredContext>(ShaderPlatform, &InstanceCullingManager);
+	FInstanceCullingDeferredContext* DeferredContext = GraphBuilder.AllocObject<FInstanceCullingDeferredContext>(ShaderPlatform, &InstanceCullingManager, NumBins);
 
 	const bool bCullInstances = CVarCullInstances.GetValueOnRenderThread() != 0;
 	const bool bAllowWPODisable = true;
@@ -1086,10 +1082,11 @@ FInstanceCullingDeferredContext *FInstanceCullingContext::CreateDeferredContext(
 	RDG_EVENT_SCOPE(GraphBuilder, "BuildRenderingCommandsDeferred(Culling=%s)", bCullInstances ? TEXT("On") : TEXT("Off"));
 	RDG_GPU_STAT_SCOPE(GraphBuilder, BuildRenderingCommandsDeferred);
 
-	TStaticArray<FBuildInstanceIdBufferAndCommandsFromPrimitiveIdsCs::FParameters*, static_cast<uint32>(EBatchProcessingMode::Num)> PassParameters;
-	for (uint32 Mode = 0U; Mode < uint32(EBatchProcessingMode::Num); ++Mode)
+	TArray<FBuildInstanceIdBufferAndCommandsFromPrimitiveIdsCs::FParameters*> PassParameters;
+	PassParameters.SetNum(NumBins);
+	for (uint32 BinIndex = 0U; BinIndex < NumBins; ++BinIndex)
 	{
-		PassParameters[Mode] = GraphBuilder.AllocParameters<FBuildInstanceIdBufferAndCommandsFromPrimitiveIdsCs::FParameters>();
+		PassParameters[BinIndex] = GraphBuilder.AllocParameters<FBuildInstanceIdBufferAndCommandsFromPrimitiveIdsCs::FParameters>();
 	}
 
 	// Create buffers for compacting instances for draw commands that need it
@@ -1203,49 +1200,58 @@ FInstanceCullingDeferredContext *FInstanceCullingContext::CreateDeferredContext(
 	// Record the number of culling views to be able to check that no views referencing out-of bounds views are queued up
 	DeferredContext->NumCullingViews = InstanceCullingManager.CullingIntermediate.NumViews;
 
-	for (uint32 Mode = 0U; Mode < uint32(EBatchProcessingMode::Num); ++Mode)
+	for (uint32 BinIndex = 0U; BinIndex < NumBins; ++BinIndex)
 	{
-		*PassParameters[Mode] = PassParametersTmp;
+		*PassParameters[BinIndex] = PassParametersTmp;
 
-		FRDGBufferRef BatchIndsRDG = CreateStructuredBuffer(INST_CULL_CREATE_STRUCT_BUFF_ARGS_MODE(BatchInds));
-		PassParameters[Mode]->BatchInds = GraphBuilder.CreateSRV(BatchIndsRDG);
+		FRDGBufferRef BatchIndsRDG = CreateStructuredBuffer(INST_CULL_CREATE_STRUCT_BUFF_ARGS_BIN_INDEX(BatchInds));
+		PassParameters[BinIndex]->BatchInds = GraphBuilder.CreateSRV(BatchIndsRDG);
+
+		// The first bin (0) is used for the EBatchProcessingMode::UnCulled batches
+		EBatchProcessingMode CurrentBatchProcessingMode = (BinIndex == 0) ? EBatchProcessingMode::UnCulled : EBatchProcessingMode::Generic;
 
-		// 
 		FInstanceProcessingGPULoadBalancer::FGPUData Result;
 		FRDGBufferRef BatchBuffer = CreateStructuredBuffer(
 			GraphBuilder,
 			TEXT("InstanceCullingLoadBalancer.Batches"),
 			sizeof(FInstanceProcessingGPULoadBalancer::FPackedBatch),
-			INST_CULL_CALLBACK_MODE(DeferredContext->LoadBalancers[Mode].GetBatches().Num()),
-			INST_CULL_CALLBACK_MODE(DeferredContext->LoadBalancers[Mode].GetBatches().GetData()),
-			INST_CULL_CALLBACK_MODE(GetArrayDataSize(DeferredContext->LoadBalancers[Mode].GetBatches())));
+			INST_CULL_CALLBACK_BIN_INDEX(DeferredContext->LoadBalancers[BinIndex].GetBatches().Num()),
+			INST_CULL_CALLBACK_BIN_INDEX(DeferredContext->LoadBalancers[BinIndex].GetBatches().GetData()),
+			INST_CULL_CALLBACK_BIN_INDEX(GetArrayDataSize(DeferredContext->LoadBalancers[BinIndex].GetBatches())));
 
 		FRDGBufferRef ItemBuffer = CreateStructuredBuffer(
 			GraphBuilder,
 			TEXT("InstanceCullingLoadBalancer.Items"),
 			sizeof(FInstanceProcessingGPULoadBalancer::FPackedItem),
-			INST_CULL_CALLBACK_MODE(DeferredContext->LoadBalancers[Mode].GetItems().Num()),
-			INST_CULL_CALLBACK_MODE(DeferredContext->LoadBalancers[Mode].GetItems().GetData()),
-			INST_CULL_CALLBACK_MODE(GetArrayDataSize(DeferredContext->LoadBalancers[Mode].GetItems())));
+			INST_CULL_CALLBACK_BIN_INDEX(DeferredContext->LoadBalancers[BinIndex].GetItems().Num()),
+			INST_CULL_CALLBACK_BIN_INDEX(DeferredContext->LoadBalancers[BinIndex].GetItems().GetData()),
+			INST_CULL_CALLBACK_BIN_INDEX(GetArrayDataSize(DeferredContext->LoadBalancers[BinIndex].GetItems())));
+
+		PassParameters[BinIndex]->LoadBalancerParameters.BatchBuffer = GraphBuilder.CreateSRV(BatchBuffer);
+		PassParameters[BinIndex]->LoadBalancerParameters.ItemBuffer = GraphBuilder.CreateSRV(ItemBuffer);
+		PassParameters[BinIndex]->LoadBalancerParameters.NumGroupsPerBatch = 1;
+		PassParameters[BinIndex]->CurrentBatchProcessingMode = static_cast<uint32>(CurrentBatchProcessingMode);
 
-		PassParameters[Mode]->LoadBalancerParameters.BatchBuffer = GraphBuilder.CreateSRV(BatchBuffer);
-		PassParameters[Mode]->LoadBalancerParameters.ItemBuffer = GraphBuilder.CreateSRV(ItemBuffer);
-		PassParameters[Mode]->LoadBalancerParameters.NumGroupsPerBatch = 1;
-		PassParameters[Mode]->CurrentBatchProcessingMode = Mode;
+		FRDGTextureRef ViewPrevHZB = nullptr;
 
-		const bool bOcclusionCullInstances = FInstanceCullingContext::IsOcclusionCullingEnabled();
+		if (BinIndex > 0 && InstanceCullingManager.ViewPrevHZBs.Num() >= (int32)BinIndex)
+		{
+			ViewPrevHZB = GraphBuilder.RegisterExternalTexture(InstanceCullingManager.ViewPrevHZBs[BinIndex - 1]);
+		}
+		
+		bool bOcclusionCullInstances = ViewPrevHZB != nullptr && FInstanceCullingContext::IsOcclusionCullingEnabled();
+		
 		if (bOcclusionCullInstances)
 		{
-			// Fill with a placeholder as AddPass expects HZBTexture to be valid. ProcessBatched will fill with real HZB textures.
-			PassParameters[Mode]->HZBTexture = GraphBuilder.RegisterExternalTexture(GSystemTextures.BlackDummy);
-			PassParameters[Mode]->HZBSize = PassParameters[Mode]->HZBTexture->Desc.Extent;
-			PassParameters[Mode]->HZBSampler = TStaticSamplerState< SF_Point, AM_Clamp, AM_Clamp, AM_Clamp >::GetRHI();
+			PassParameters[BinIndex]->HZBTexture = ViewPrevHZB;
+			PassParameters[BinIndex]->HZBSize = ViewPrevHZB->Desc.Extent;
+			PassParameters[BinIndex]->HZBSampler = TStaticSamplerState< SF_Point, AM_Clamp, AM_Clamp, AM_Clamp >::GetRHI();
 		}
 
 		FBuildInstanceIdBufferAndCommandsFromPrimitiveIdsCs::FPermutationDomain PermutationVector;
 		PermutationVector.Set<FBuildInstanceIdBufferAndCommandsFromPrimitiveIdsCs::FBatchedDim>(true);
-		PermutationVector.Set<FBuildInstanceIdBufferAndCommandsFromPrimitiveIdsCs::FSingleInstanceModeDim>(EBatchProcessingMode(Mode) == EBatchProcessingMode::UnCulled);
-		PermutationVector.Set<FBuildInstanceIdBufferAndCommandsFromPrimitiveIdsCs::FCullInstancesDim>(bCullInstances && EBatchProcessingMode(Mode) != EBatchProcessingMode::UnCulled);
+		PermutationVector.Set<FBuildInstanceIdBufferAndCommandsFromPrimitiveIdsCs::FSingleInstanceModeDim>(CurrentBatchProcessingMode == EBatchProcessingMode::UnCulled);
+		PermutationVector.Set<FBuildInstanceIdBufferAndCommandsFromPrimitiveIdsCs::FCullInstancesDim>(bCullInstances && CurrentBatchProcessingMode != EBatchProcessingMode::UnCulled);
 		PermutationVector.Set<FBuildInstanceIdBufferAndCommandsFromPrimitiveIdsCs::FAllowWPODisableDim>(bAllowWPODisable);
 		PermutationVector.Set<FBuildInstanceIdBufferAndCommandsFromPrimitiveIdsCs::FOcclusionCullInstancesDim>(bOcclusionCullInstances);
 		PermutationVector.Set<FBuildInstanceIdBufferAndCommandsFromPrimitiveIdsCs::FInstanceCompactionDim>(bEnableInstanceCompaction);
@@ -1254,10 +1260,10 @@ FInstanceCullingDeferredContext *FInstanceCullingContext::CreateDeferredContext(
 
 		FComputeShaderUtils::AddPass(
 			GraphBuilder,
-			RDG_EVENT_NAME("CullInstances(%s)", BatchProcessingModeStr[Mode]),
+			RDG_EVENT_NAME("CullInstances(%s). Bin %d", BatchProcessingModeStr[uint32(CurrentBatchProcessingMode)], BinIndex),
 			ComputeShader,
-			PassParameters[Mode],
-			INST_CULL_CALLBACK_MODE(DeferredContext->LoadBalancers[Mode].GetWrappedCsGroupCount()));
+			PassParameters[BinIndex],
+			INST_CULL_CALLBACK_BIN_INDEX(DeferredContext->LoadBalancers[BinIndex].GetWrappedCsGroupCount()));
 	}
 
 	// TODO: Come up with a way to cull these passes when no compaction is needed. The group count resulting in (0, 0, 0) causes the pass lambdas to not execute,
diff --git a/Engine/Source/Runtime/Renderer/Private/InstanceCulling/InstanceCullingManager.cpp b/Engine/Source/Runtime/Renderer/Private/InstanceCulling/InstanceCullingManager.cpp
--- a/Engine/Source/Runtime/Renderer/Private/InstanceCulling/InstanceCullingManager.cpp
+++ b/Engine/Source/Runtime/Renderer/Private/InstanceCulling/InstanceCullingManager.cpp
@@ -57,6 +57,13 @@ int32 FInstanceCullingManager::RegisterView(const FViewInfo& ViewInfo)
 	Params.MaxPixelsPerEdgeMultipler = 1.0f;
 	Params.InstanceOcclusionQueryMask = ViewInfo.PrevViewInfo.InstanceOcclusionQueryMask;
 
+	if (ViewInfo.PrevViewInfo.HZB != nullptr)
+	{
+		check(IsInRenderingThread());
+
+		ViewPrevHZBs.AddUnique(ViewInfo.PrevViewInfo.HZB);
+	}	
+		
 	return RegisterView(Params);
 }
 
@@ -86,6 +93,32 @@ void FInstanceCullingManager::FlushRegisteredViews(FRDGBuilder& GraphBuilder)
 	}
 }
 
+int32 FInstanceCullingManager::GetBinIndex(EBatchProcessingMode Mode, const TRefCountPtr<IPooledRenderTarget>& HZB)
+{
+	if (Mode == EBatchProcessingMode::UnCulled)
+	{
+		return 0;
+	}
+
+	int32 BinIndex;
+	// all contexts without a valid HZB go in the first bin, together with the first view's HZB
+	if (!HZB.IsValid())
+	{
+		return 1;
+	}
+
+	if (!ViewPrevHZBs.Find(HZB, BinIndex))
+	{
+		// error: the HZB is not registered correctly
+		return -1;
+	}
+
+	// bin 0 is used for EBatchProcessingMode::UnCulled batches
+	BinIndex += 1;
+
+	return BinIndex;
+}
+
 bool FInstanceCullingManager::AllowBatchedBuildRenderingCommands(const FGPUScene& GPUScene)
 {
 	return GPUScene.IsEnabled() && !!GAllowBatchedBuildRenderingCommands && !FRDGBuilder::IsImmediateMode();
diff --git a/Engine/Source/Runtime/Renderer/Private/InstanceCulling/InstanceCullingManager.h b/Engine/Source/Runtime/Renderer/Private/InstanceCulling/InstanceCullingManager.h
--- a/Engine/Source/Runtime/Renderer/Private/InstanceCulling/InstanceCullingManager.h
+++ b/Engine/Source/Runtime/Renderer/Private/InstanceCulling/InstanceCullingManager.h
@@ -101,6 +101,13 @@ public:
 	FRDGBufferRef InstanceOcclusionQueryBuffer = {};
 	EPixelFormat InstanceOcclusionQueryBufferFormat = PF_Unknown;
 
+	TArray<TRefCountPtr<IPooledRenderTarget>> ViewPrevHZBs;
+
+	// to support merging of multiple contexts with different HZBs, we use separate bins (LoadBalancers) in the InstanceCullingDeferredContext
+	// bin 0 is reserved for UnCulled batches. Bins >= 1 are for the different HZBs.
+
+	int32 GetBinIndex(EBatchProcessingMode Mode, const TRefCountPtr<IPooledRenderTarget>& HZB);
+
 private:
 
 	friend class FInstanceCullingContext;
diff --git a/Engine/Source/Runtime/Renderer/Private/InstanceCulling/InstanceCullingMergedContext.cpp b/Engine/Source/Runtime/Renderer/Private/InstanceCulling/InstanceCullingMergedContext.cpp
--- a/Engine/Source/Runtime/Renderer/Private/InstanceCulling/InstanceCullingMergedContext.cpp
+++ b/Engine/Source/Runtime/Renderer/Private/InstanceCulling/InstanceCullingMergedContext.cpp
@@ -4,6 +4,28 @@
 #include "InstanceCulling/InstanceCullingManager.h"
 #include "RenderGraphBuilder.h"
 
+FInstanceCullingMergedContext::FInstanceCullingMergedContext(EShaderPlatform InShaderPlatform, bool bInMustAddAllContexts, int32 InNumBins)
+	: ShaderPlatform(InShaderPlatform)
+	, bMustAddAllContexts(bInMustAddAllContexts)
+{
+	// make sure we always have at least 2 bins (for UnCulled + Generic batching modes)
+	check(InNumBins >= 2);
+
+	LoadBalancers.SetNum(InNumBins);
+	BatchInds.SetNum(InNumBins);
+	TotalBatches.SetNumZeroed(InNumBins);
+	TotalItems.SetNumZeroed(InNumBins);
+}
+
+int32 FInstanceCullingMergedContext::GetLoadBalancerIndex(EBatchProcessingMode Mode, const FBatchItem& BatchItem)
+{
+	int32 BinIndex = (Mode == EBatchProcessingMode::UnCulled) ? 0 : BatchItem.GenericBinIndex;
+
+	check(BinIndex >= 0 && BinIndex < LoadBalancers.Num());
+	
+	return BinIndex;
+}
+
 void FInstanceCullingMergedContext::MergeBatches()
 {
 	for (FBatchItem& AsyncBatchItem : AsyncBatches)
@@ -16,9 +38,9 @@ void FInstanceCullingMergedContext::MergeBatches()
 	}
 	AsyncBatches.Reset();
 
-	for (uint32 Mode = 0U; Mode < uint32(EBatchProcessingMode::Num); ++Mode)
+	for (uint32 BinIndex = 0U; BinIndex < uint32(LoadBalancers.Num()); ++BinIndex)
 	{
-		LoadBalancers[Mode].ReserveStorage(TotalBatches[Mode], TotalItems[Mode]);
+		LoadBalancers[BinIndex].ReserveStorage(TotalBatches[BinIndex], TotalItems[BinIndex]);
 	}
 	// Pre-size all arrays
 	IndirectArgs.Empty(TotalIndirectArgs);
@@ -80,8 +102,10 @@ void FInstanceCullingMergedContext::MergeBatches()
 
 		for (uint32 Mode = 0U; Mode < uint32(EBatchProcessingMode::Num); ++Mode)
 		{
-			int32 StartIndex = BatchInds[Mode].Num();
-			TInstanceCullingLoadBalancer<SceneRenderingAllocator>* MergedLoadBalancer = &LoadBalancers[Mode];
+			int32 BinIndex = GetLoadBalancerIndex(static_cast<EBatchProcessingMode>(Mode), BatchItem);
+
+			int32 StartIndex = BatchInds[BinIndex].Num();
+			TInstanceCullingLoadBalancer<SceneRenderingAllocator>* MergedLoadBalancer = &LoadBalancers[BinIndex];
 
 			BatchInfo.ItemDataOffset[Mode] = MergedLoadBalancer->GetItems().Num();
 			FInstanceProcessingGPULoadBalancer* LoadBalancer = InstanceCullingContext.LoadBalancers[Mode];
@@ -90,12 +114,12 @@ void FInstanceCullingMergedContext::MergeBatches()
 			// UnCulled bucket is used for a single instance mode
 			check(EBatchProcessingMode(Mode) != EBatchProcessingMode::UnCulled || LoadBalancer->HasSingleInstanceItemsOnly());
 
-			BatchInds[Mode].AddDefaulted(LoadBalancer->GetBatches().Num());
+			BatchInds[BinIndex].AddDefaulted(LoadBalancer->GetBatches().Num());
 
 			MergedLoadBalancer->AppendData(*LoadBalancer);
-			for (int32 Index = StartIndex; Index < BatchInds[Mode].Num(); ++Index)
+			for (int32 Index = StartIndex; Index < BatchInds[BinIndex].Num(); ++Index)
 			{
-				BatchInds[Mode][Index] = BatchInfoIndex;
+				BatchInds[BinIndex][Index] = BatchInfoIndex;
 			}
 		}
 		const uint32 BatchTotalInstances = InstanceCullingContext.TotalInstances * InstanceCullingContext.ViewIds.Num();
@@ -136,27 +160,25 @@ void FInstanceCullingMergedContext::AddBatch(FRDGBuilder& GraphBuilder, FInstanc
 
 	const bool bOcclusionCullInstances = Context->PrevHZB.IsValid() && FInstanceCullingContext::IsOcclusionCullingEnabled();
 
-	// Set HZB texture for the merged batches
+	// Resolve the bin index based on the PrevHZB. Bin 0 is reserved for UnCulled batches, every other bin is for each HZB.
+	// Generic batches with a null HZB will go in bin 1, together with the ones associated to the first HZB.
+	int32 BinIndex = 1;
 	if (bOcclusionCullInstances)
 	{
-		// Verify that each batch contains the same HZB if not null as we only support one
-		check(PrevHZB == nullptr || PrevHZB == GraphBuilder.RegisterExternalTexture(Context->PrevHZB));
+		BinIndex = Context->InstanceCullingManager->GetBinIndex(EBatchProcessingMode::Generic, Context->PrevHZB);
 
-		if (PrevHZB == nullptr)
-		{
-			// Note: performing the registration here because the final merge of contexts may happen during RDG execute (in case of the deferred culling) which seems ill-defined
-			PrevHZB = GraphBuilder.RegisterExternalTexture(Context->PrevHZB);
-		}
+		// make sure that this Context's PrevHZB is registered correctly
+		check(BinIndex > 0);
 	}
 
 	if (Context->SyncPrerequisitesFunc)
 	{
-		AsyncBatches.Add(FBatchItem{ Context, InstanceCullingDrawParams });
+		AsyncBatches.Add(FBatchItem{ Context, InstanceCullingDrawParams, BinIndex });
 
 	}
 	else
 	{
-		AddBatchItem(FBatchItem{ Context, InstanceCullingDrawParams });
+		AddBatchItem(FBatchItem{ Context, InstanceCullingDrawParams, BinIndex });
 	}
 
 }
@@ -172,8 +194,11 @@ void FInstanceCullingMergedContext::AddBatchItem(const FBatchItem& BatchItem)
 		for (uint32 Mode = 0U; Mode < uint32(EBatchProcessingMode::Num); ++Mode)
 		{
 			Context->LoadBalancers[Mode]->FinalizeBatches();
-			TotalBatches[Mode] += Context->LoadBalancers[Mode]->GetBatches().Max();
-			TotalItems[Mode] += Context->LoadBalancers[Mode]->GetItems().Max();
+
+			int32 BinIndex = GetLoadBalancerIndex(static_cast<EBatchProcessingMode>(Mode), BatchItem);
+
+			TotalBatches[BinIndex] += Context->LoadBalancers[Mode]->GetBatches().Max();
+			TotalItems[BinIndex] += Context->LoadBalancers[Mode]->GetItems().Max();
 		}
 #if DO_CHECK
 		for (int32 ViewId : Context->ViewIds)
diff --git a/Engine/Source/Runtime/Renderer/Private/InstanceCulling/InstanceCullingMergedContext.h b/Engine/Source/Runtime/Renderer/Private/InstanceCulling/InstanceCullingMergedContext.h
--- a/Engine/Source/Runtime/Renderer/Private/InstanceCulling/InstanceCullingMergedContext.h
+++ b/Engine/Source/Runtime/Renderer/Private/InstanceCulling/InstanceCullingMergedContext.h
@@ -9,15 +9,14 @@
 class FInstanceCullingMergedContext
 {
 public:
-	FInstanceCullingMergedContext(EShaderPlatform InShaderPlatform, bool bInMustAddAllContexts = false)
-		: ShaderPlatform(InShaderPlatform)
-		, bMustAddAllContexts(bInMustAddAllContexts)
-	{}
-
+	FInstanceCullingMergedContext(EShaderPlatform InShaderPlatform, bool bInMustAddAllContexts = false, int32 InNumBins = 2);
+	
 	struct FBatchItem
 	{
 		FInstanceCullingContext* Context = nullptr;
 		FInstanceCullingDrawParams* Result = nullptr;
+		// EBatchProcessingMode::Generic batches are put in bins, based on their view's PrevHZB
+		int32 GenericBinIndex = -1;
 	};
 
 	// Info about a batch of culling work produced by a context, when part of a batched job
@@ -35,6 +34,9 @@ public:
 		uint32 ItemDataOffset[uint32(EBatchProcessingMode::Num)];
 	};
 
+	/** Bin 0 is used for UnCulled batches, Culled ones go in bins >= 1. */
+	static constexpr uint32 FirstGenericBinIndex = 1;
+
 	/** Batches of GPU instance culling input data. */
 	TArray<FBatchItem, SceneRenderingAllocator> Batches;
 
@@ -53,8 +55,9 @@ public:
 	TArray<FInstanceCullingContext::FCompactionData, SceneRenderingAllocator> DrawCommandCompactionData;
 	TArray<uint32, SceneRenderingAllocator> CompactionBlockDataIndices;	
 
-	TStaticArray<TInstanceCullingLoadBalancer<SceneRenderingAllocator>, static_cast<uint32>(EBatchProcessingMode::Num)> LoadBalancers;
-	TStaticArray<TArray<uint32, SceneRenderingAllocator>, static_cast<uint32>(EBatchProcessingMode::Num)> BatchInds;
+	// these are the actual BINS, so we would need one per HZB
+	TArray<TInstanceCullingLoadBalancer<SceneRenderingAllocator>, SceneRenderingAllocator> LoadBalancers;
+	TArray<TArray<uint32, SceneRenderingAllocator>, SceneRenderingAllocator> BatchInds; // TODO: rename to ContextInds
 	TArray<FContextBatchInfoPacked, SceneRenderingAllocator> BatchInfos;
 
 	EShaderPlatform ShaderPlatform = SP_NumPlatforms;
@@ -63,8 +66,9 @@ public:
 	bool bMustAddAllContexts = false;
 	// Counters to sum up all sizes to facilitate pre-sizing
 	uint32 InstanceIdBufferElements = 0U;
-	TStaticArray<int32, uint32(EBatchProcessingMode::Num)> TotalBatches = TStaticArray<int32, uint32(EBatchProcessingMode::Num)>(InPlace, 0);
-	TStaticArray<int32, uint32(EBatchProcessingMode::Num)> TotalItems = TStaticArray<int32, uint32(EBatchProcessingMode::Num)>(InPlace, 0);
+	// preallocate 5 to cover all scenarios up to UnCulled bin + 4 HZBs in case of 4 primary views
+	TArray<int32, TInlineAllocator<5>> TotalBatches;
+	TArray<int32, TInlineAllocator<5>> TotalItems;
 	int32 TotalIndirectArgs = 0;
 	int32 TotalPayloads = 0;
 	int32 TotalViewIds = 0;
@@ -73,8 +77,6 @@ public:
 	int32 TotalCompactionBlocks = 0;
 	int32 TotalCompactionInstances = 0;
 
-	// Single Previous frame HZB which is shared among all batched contexts, thus only one is allowed (but the same can be used in multiple passes). (Needs atlas or bindless to expand).
-	FRDGTextureRef PrevHZB = nullptr;
 	int32 NumCullingViews = 0;
 
 	// Merge the queued batches and populate the derived data.
@@ -86,4 +88,6 @@ public:
 private:
 	void AddBatchItem(const FBatchItem& BatchItem);
 
+	int32 GetLoadBalancerIndex(EBatchProcessingMode Mode, const FBatchItem& BatchItem);
+
 };
\ No newline at end of file
diff --git a/Engine/Source/Runtime/Renderer/Private/VirtualShadowMaps/VirtualShadowMapArray.cpp b/Engine/Source/Runtime/Renderer/Private/VirtualShadowMaps/VirtualShadowMapArray.cpp
--- a/Engine/Source/Runtime/Renderer/Private/VirtualShadowMaps/VirtualShadowMapArray.cpp
+++ b/Engine/Source/Runtime/Renderer/Private/VirtualShadowMaps/VirtualShadowMapArray.cpp
@@ -2943,10 +2943,10 @@ void FVirtualShadowMapArray::RenderVirtualShadowMapsNonNanite(FRDGBuilder& Graph
 			    InstanceCullingMergedContext.IndirectArgs,
 			    InstanceCullingMergedContext.DrawCommandDescs,
 			    InstanceCullingMergedContext.InstanceIdOffsets,
-			    &InstanceCullingMergedContext.LoadBalancers[uint32(EBatchProcessingMode::Generic)],
+			    &InstanceCullingMergedContext.LoadBalancers[FInstanceCullingMergedContext::FirstGenericBinIndex],
 			    InstanceCullingMergedContext.BatchInfos,
 			    VSMCullingBatchInfos,
-			    InstanceCullingMergedContext.BatchInds[uint32(EBatchProcessingMode::Generic)],
+			    InstanceCullingMergedContext.BatchInds[FInstanceCullingMergedContext::FirstGenericBinIndex],
 			    InstanceCullingMergedContext.TotalInstances,
 			    TotalViewScaledInstanceCount,
 			    TotalPrimaryViews,
@@ -3077,7 +3077,7 @@ void FVirtualShadowMapArray::RenderVirtualShadowMapsNonNanite(FRDGBuilder& Graph
 				InstanceCullingContext->IndirectArgs, 
 				InstanceCullingContext->DrawCommandDescs,
 				InstanceCullingContext->InstanceIdOffsets,
-				InstanceCullingContext->LoadBalancers[uint32(EBatchProcessingMode::Generic)],
+				InstanceCullingContext->LoadBalancers[FInstanceCullingMergedContext::FirstGenericBinIndex],
 				MakeArrayView(&CullingBatchInfo, 1),
 				MakeArrayView(&VSMCullingBatchInfo, 1),
 				MakeArrayView<const uint32>(nullptr, 0),
