Fixed missing per-light bias term in vsm resolution for local lights.

#rnx

diff --git a/Engine/Source/Runtime/Renderer/Private/Shadows/ShadowSceneRenderer.cpp b/Engine/Source/Runtime/Renderer/Private/Shadows/ShadowSceneRenderer.cpp
--- a/Engine/Source/Runtime/Renderer/Private/Shadows/ShadowSceneRenderer.cpp
+++ b/Engine/Source/Runtime/Renderer/Private/Shadows/ShadowSceneRenderer.cpp
@@ -239,7 +239,7 @@ TSharedPtr<FVirtualShadowMapPerLightCacheEntry> FShadowSceneRenderer::AddLocalLi
 	//       we can absolutely mirror the page marking calc better, just unclear how much it helps. 
 	//       Also possible to feed back from gpu - which would be more accurate wrt partially visible lights (e.g., a spot going through the ground).
 	//       Of course this creates jumps if visibility changes, which may or may not create unsolvable artifacts.	
-	const float BiasedFootprintThreshold = float(FVirtualShadowMap::PageSize) * FMath::Exp2(ResolutionLODBiasLocal);
+	const float BiasedFootprintThreshold = float(FVirtualShadowMap::PageSize) * FMath::Exp2(ResolutionLODBiasLocal - LightSceneProxy->GetVSMResolutionLodBias());
 	const bool bIsDistantLight = CVarDistantLightMode.GetValueOnRenderThread() != 0
 		&& (MaxScreenRadius <= BiasedFootprintThreshold || CVarDistantLightMode.GetValueOnRenderThread() == 2);
 
