Fixed RHI validation assert on invalid UB layout specific to a mobile GPUScene

diff --git a/Engine/Source/Runtime/Renderer/Private/MeshPassProcessor.cpp b/Engine/Source/Runtime/Renderer/Private/MeshPassProcessor.cpp
--- a/Engine/Source/Runtime/Renderer/Private/MeshPassProcessor.cpp
+++ b/Engine/Source/Runtime/Renderer/Private/MeshPassProcessor.cpp
@@ -1780,6 +1780,7 @@ public:
 		static FName BatchedPrimitiveSlotName = "BatchedPrimitive";
 		FRHIUniformBufferLayoutInitializer Initializer(TEXT("DynamicBatchedPrimitive"), 16u);
 		EnumAddFlags(Initializer.Flags, ERHIUniformBufferFlags::UniformView);
+		EnumAddFlags(Initializer.Flags, ERHIUniformBufferFlags::NoEmulatedUniformBuffer); // implicit for UniformView
 		Initializer.Resources.Add({0, UBMT_RDG_BUFFER_SRV});
 		Initializer.StaticSlot = FUniformBufferStaticSlotRegistry::Get().FindSlotByName(BatchedPrimitiveSlotName);
 		Initializer.BindingFlags = EUniformBufferBindingFlags::StaticAndShader;
