LensDistortion: Fixed lens distortion incorrectly sampling with texture UVs instead of viewport UVs.

#rnx

diff --git a/Engine/Shaders/Private/ComposeSeparateTranslucency.usf b/Engine/Shaders/Private/ComposeSeparateTranslucency.usf
--- a/Engine/Shaders/Private/ComposeSeparateTranslucency.usf
+++ b/Engine/Shaders/Private/ComposeSeparateTranslucency.usf
@@ -32,6 +32,9 @@ SamplerState      UndistortingDisplacementSampler;
 
 FScreenTransform ScreenPosToSceneColorUV;
 FScreenTransform ScreenPosToSeparateTranslucencyUV;
+FScreenTransform SeparateTranslucencyUVToViewportUV;
+FScreenTransform ViewportUVToSeparateTranslucencyUV;
+
 float2 SeparateTranslucencyUVMin;
 float2 SeparateTranslucencyUVMax;
 float2 SeparateTranslucencyExtentInverse;
@@ -50,7 +53,9 @@ void MainPS(
 	BRANCH
 	if (bLensDistortion)
 	{
-		SeparateTranslucencyUV = ApplyLensDistortionOnViewportUV(UndistortingDisplacementTexture, UndistortingDisplacementSampler, SeparateTranslucencyUV);
+		float2 DistortedViewportUV = ApplyScreenTransform(SeparateTranslucencyUV, SeparateTranslucencyUVToViewportUV);
+		float2 UndistortedViewportUV = ApplyLensDistortionOnViewportUV(UndistortingDisplacementTexture, UndistortingDisplacementSampler, DistortedViewportUV);
+		SeparateTranslucencyUV = ApplyScreenTransform(UndistortedViewportUV, ViewportUVToSeparateTranslucencyUV);
 	}
 
 	float4 SceneColorSample = SceneColorTexture.SampleLevel(SceneColorSampler, SceneColorUV, 0);
diff --git a/Engine/Shaders/Private/MotionBlur/MotionBlurApply.usf b/Engine/Shaders/Private/MotionBlur/MotionBlurApply.usf
--- a/Engine/Shaders/Private/MotionBlur/MotionBlurApply.usf
+++ b/Engine/Shaders/Private/MotionBlur/MotionBlurApply.usf
@@ -82,6 +82,8 @@ SCREEN_PASS_TEXTURE_VIEWPORT(Color)
 SCREEN_PASS_TEXTURE_VIEWPORT(Velocity)
 SCREEN_PASS_TEXTURE_VIEWPORT(VelocityTile)
 FScreenTransform ColorToVelocity;
+FScreenTransform SeparateTranslucencyUVToViewportUV;
+FScreenTransform ViewportUVToSeparateTranslucencyUV;
 uint MaxSampleCount;
 uint OutputMip1;
 uint OutputMip2;
@@ -251,6 +253,14 @@ mb_half GetVelocityLengthPixels(mb_half2 EncodedVelocity)
 	return EncodedVelocity.x;
 }
 
+float2 ApplyLensDistortionOnTranslucencyUV(float2 SeparateTranslucencyUV)
+{
+	float2 DistortedViewportUV = ApplyScreenTransform(SeparateTranslucencyUV, SeparateTranslucencyUVToViewportUV);
+	float2 UndistortedViewportUV = ApplyLensDistortionOnViewportUV(UndistortingDisplacementTexture, UndistortingDisplacementSampler, DistortedViewportUV);
+	
+	return ApplyScreenTransform(UndistortedViewportUV, ViewportUVToSeparateTranslucencyUV);
+}
+
 
 //------------------------------------------------------- ENTRY POINT
 
@@ -743,10 +753,10 @@ void MainCS(
 		BRANCH
 		if (bLensDistortion)
 		{
-			PostMotionBlurTranslucencyUV0 = ApplyLensDistortionOnViewportUV(UndistortingDisplacementTexture, UndistortingDisplacementSampler, PostMotionBlurTranslucencyUV0);
-			PostMotionBlurTranslucencyUV1 = ApplyLensDistortionOnViewportUV(UndistortingDisplacementTexture, UndistortingDisplacementSampler, PostMotionBlurTranslucencyUV1);
-			PostMotionBlurTranslucencyUV2 = ApplyLensDistortionOnViewportUV(UndistortingDisplacementTexture, UndistortingDisplacementSampler, PostMotionBlurTranslucencyUV2);
-			PostMotionBlurTranslucencyUV3 = ApplyLensDistortionOnViewportUV(UndistortingDisplacementTexture, UndistortingDisplacementSampler, PostMotionBlurTranslucencyUV3);
+			PostMotionBlurTranslucencyUV0 = ApplyLensDistortionOnTranslucencyUV(PostMotionBlurTranslucencyUV0);
+			PostMotionBlurTranslucencyUV1 = ApplyLensDistortionOnTranslucencyUV(PostMotionBlurTranslucencyUV1);
+			PostMotionBlurTranslucencyUV2 = ApplyLensDistortionOnTranslucencyUV(PostMotionBlurTranslucencyUV2);
+			PostMotionBlurTranslucencyUV3 = ApplyLensDistortionOnTranslucencyUV(PostMotionBlurTranslucencyUV3);
 		}
 			
 		PostMotionBlurTranslucencyUV0 = clamp(PostMotionBlurTranslucencyUV0, TranslucencyUVMin, TranslucencyUVMax);
@@ -996,7 +1006,7 @@ void MainCS(
 		BRANCH
 		if (bLensDistortion)
 		{
-			PostMotionBlurTranslucencyUV = ApplyLensDistortionOnViewportUV(UndistortingDisplacementTexture, UndistortingDisplacementSampler, PostMotionBlurTranslucencyUV);
+			PostMotionBlurTranslucencyUV = ApplyLensDistortionOnTranslucencyUV(PostMotionBlurTranslucencyUV);
 		}
 		PostMotionBlurTranslucencyUV = clamp(PostMotionBlurTranslucencyUV, TranslucencyUVMin, TranslucencyUVMax);
 		
diff --git a/Engine/Source/Runtime/Renderer/Private/PostProcess/PostProcessMotionBlur.cpp b/Engine/Source/Runtime/Renderer/Private/PostProcess/PostProcessMotionBlur.cpp
--- a/Engine/Source/Runtime/Renderer/Private/PostProcess/PostProcessMotionBlur.cpp
+++ b/Engine/Source/Runtime/Renderer/Private/PostProcess/PostProcessMotionBlur.cpp
@@ -461,6 +461,8 @@ class FMotionBlurFilterCS : public FMotionBlurShader
 		SHADER_PARAMETER_STRUCT(FScreenPassTextureViewportParameters, VelocityTile)
 
 		SHADER_PARAMETER(FScreenTransform, ColorToVelocity)
+		SHADER_PARAMETER(FScreenTransform, SeparateTranslucencyUVToViewportUV)
+		SHADER_PARAMETER(FScreenTransform, ViewportUVToSeparateTranslucencyUV)
 		SHADER_PARAMETER(int32, MaxSampleCount)
 		SHADER_PARAMETER(int32, OutputMip1)
 		SHADER_PARAMETER(int32, OutputMip2)
@@ -1025,9 +1027,12 @@ FMotionBlurOutputs AddMotionBlurFilterPass(
 		
 			OriginalPassParameters.TranslucencyTexture = PostMotionBlurTranslucency;
 			OriginalPassParameters.TranslucencySampler = GetPostMotionBlurTranslucencySampler(bScaleTranslucency);
-			OriginalPassParameters.ColorToTranslucency = FScreenTransform::ChangeTextureUVCoordinateFromTo(
-				Viewports.Color,
-				FScreenPassTextureViewport(PostMotionBlurTranslucency->Desc.Extent, FIntRect(FIntPoint::ZeroValue, PostMotionBlurTranslucencySize)));
+			FScreenPassTextureViewport TranslucencyViewport = FScreenPassTextureViewport(PostMotionBlurTranslucency->Desc.Extent, FIntRect(FIntPoint::ZeroValue, PostMotionBlurTranslucencySize));
+			OriginalPassParameters.ColorToTranslucency = FScreenTransform::ChangeTextureUVCoordinateFromTo(Viewports.Color, TranslucencyViewport);
+			OriginalPassParameters.SeparateTranslucencyUVToViewportUV = FScreenTransform::ChangeTextureBasisFromTo(
+				TranslucencyViewport, FScreenTransform::ETextureBasis::TextureUV, FScreenTransform::ETextureBasis::ViewportUV);
+			OriginalPassParameters.ViewportUVToSeparateTranslucencyUV = FScreenTransform::ChangeTextureBasisFromTo(
+				TranslucencyViewport, FScreenTransform::ETextureBasis::ViewportUV, FScreenTransform::ETextureBasis::TextureUV);
 			OriginalPassParameters.TranslucencyUVMin = FVector2f(0.0f, 0.0f);
 			OriginalPassParameters.TranslucencyUVMax = (FVector2f(PostMotionBlurTranslucencySize) - FVector2f(0.5f, 0.5f)) * PostMotionBlurTranslucencyExtentInv;
 			OriginalPassParameters.TranslucencyExtentInverse = PostMotionBlurTranslucencyExtentInv;
@@ -1043,6 +1048,8 @@ FMotionBlurOutputs AddMotionBlurFilterPass(
 			OriginalPassParameters.TranslucencyTexture = GSystemTextures.GetBlackAlphaOneDummy(GraphBuilder);
 			OriginalPassParameters.TranslucencySampler = TStaticSamplerState<SF_Point, AM_Clamp, AM_Clamp, AM_Clamp>::GetRHI();
 			OriginalPassParameters.ColorToTranslucency = FScreenTransform::Identity;
+			OriginalPassParameters.SeparateTranslucencyUVToViewportUV = FScreenTransform::Identity;
+			OriginalPassParameters.ViewportUVToSeparateTranslucencyUV = FScreenTransform::Identity;
 			OriginalPassParameters.TranslucencyUVMin = FVector2f(0.0f, 0.0f);
 			OriginalPassParameters.TranslucencyUVMax = FVector2f(0.0f, 0.0f);
 			OriginalPassParameters.TranslucencyExtentInverse = FVector2f(0.0f, 0.0f);
diff --git a/Engine/Source/Runtime/Renderer/Private/TranslucentRendering.cpp b/Engine/Source/Runtime/Renderer/Private/TranslucentRendering.cpp
--- a/Engine/Source/Runtime/Renderer/Private/TranslucentRendering.cpp
+++ b/Engine/Source/Runtime/Renderer/Private/TranslucentRendering.cpp
@@ -386,6 +386,8 @@ class FComposeSeparateTranslucencyPS : public FGlobalShader
 	BEGIN_SHADER_PARAMETER_STRUCT(FParameters, )
 		SHADER_PARAMETER(FScreenTransform, ScreenPosToSceneColorUV)
 		SHADER_PARAMETER(FScreenTransform, ScreenPosToSeparateTranslucencyUV)
+		SHADER_PARAMETER(FScreenTransform, SeparateTranslucencyUVToViewportUV)
+		SHADER_PARAMETER(FScreenTransform, ViewportUVToSeparateTranslucencyUV)
 		SHADER_PARAMETER(FVector2f, SeparateTranslucencyUVMin)
 		SHADER_PARAMETER(FVector2f, SeparateTranslucencyUVMax)
 		SHADER_PARAMETER(FVector2f, SeparateTranslucencyExtentInverse)
@@ -568,6 +570,10 @@ FScreenPassTexture FTranslucencyComposition::AddPass(
 		SceneColorViewport, FScreenTransform::ETextureBasis::ViewportUV, FScreenTransform::ETextureBasis::TextureUV);
 	PassParameters->ScreenPosToSeparateTranslucencyUV = SvPositionToViewportUV * FScreenTransform::ChangeTextureBasisFromTo(
 		TranslucencyViewport, FScreenTransform::ETextureBasis::ViewportUV, FScreenTransform::ETextureBasis::TextureUV);
+	PassParameters->SeparateTranslucencyUVToViewportUV = FScreenTransform::ChangeTextureBasisFromTo(
+		TranslucencyViewport, FScreenTransform::ETextureBasis::TextureUV, FScreenTransform::ETextureBasis::ViewportUV);
+	PassParameters->ViewportUVToSeparateTranslucencyUV = FScreenTransform::ChangeTextureBasisFromTo(
+		TranslucencyViewport, FScreenTransform::ETextureBasis::ViewportUV, FScreenTransform::ETextureBasis::TextureUV);
 
 	PassParameters->SeparateTranslucencyUVMin = (FVector2f(TranslucencyViewport.Rect.Min) + FVector2f(0.5f, 0.5f)) * SeparateTranslucencyExtentInv;
 	PassParameters->SeparateTranslucencyUVMax = (FVector2f(TranslucencyViewport.Rect.Max) - FVector2f(0.5f, 0.5f)) * SeparateTranslucencyExtentInv;
