Moved GetAccelerationStructuresAddresses to a separate RDG CommandListSetupTask

diff --git a/Engine/Source/Runtime/Renderer/Private/RayTracing/RayTracingScene.cpp b/Engine/Source/Runtime/Renderer/Private/RayTracing/RayTracingScene.cpp
--- a/Engine/Source/Runtime/Renderer/Private/RayTracing/RayTracingScene.cpp
+++ b/Engine/Source/Runtime/Renderer/Private/RayTracing/RayTracingScene.cpp
@@ -284,6 +284,30 @@ void FRayTracingScene::Create(FRDGBuilder& GraphBuilder, const FViewInfo& View,
 						TransformUploadData);
 				});
 
+			GraphBuilder.AddCommandListSetupTask([&Layer, &SceneInitializer](FRHICommandList& RHICmdList)
+				{
+					for (uint32 GPUIndex : RHICmdList.GetGPUMask())
+					{
+						FRayTracingAccelerationStructureAddress* AddressesPtr = (FRayTracingAccelerationStructureAddress*)RHICmdList.LockBufferMGPU(
+							Layer.AccelerationStructureAddressesBuffer.Buffer,
+							GPUIndex,
+							0,
+							SceneInitializer.ReferencedGeometries.Num() * sizeof(FRayTracingAccelerationStructureAddress), RLM_WriteOnly);
+
+						RHICmdList.EnqueueLambda([AddressesPtr, &SceneInitializer, GPUIndex](FRHICommandListBase&)
+							{
+								TRACE_CPUPROFILER_EVENT_SCOPE(GetAccelerationStructuresAddresses);
+
+								for (int32 GeometryIndex = 0; GeometryIndex < SceneInitializer.ReferencedGeometries.Num(); ++GeometryIndex)
+								{
+									AddressesPtr[GeometryIndex] = SceneInitializer.ReferencedGeometries[GeometryIndex]->GetAccelerationStructureAddress(GPUIndex);
+								}
+							});
+
+						RHICmdList.UnlockBufferMGPU(Layer.AccelerationStructureAddressesBuffer.Buffer, GPUIndex);
+					}
+				});
+
 			FBuildInstanceBufferPassParams* PassParams = GraphBuilder.AllocParameters<FBuildInstanceBufferPassParams>();
 			PassParams->InstanceBuffer = GraphBuilder.CreateUAV(Layer.InstanceBuffer);
 			PassParams->DebugInstanceGPUSceneIndexBuffer = DebugInstanceGPUSceneIndexBufferUAV;
@@ -301,7 +325,6 @@ void FRayTracingScene::Create(FRDGBuilder& GraphBuilder, const FViewInfo& View,
 				this,
 				&Layer,
 				GPUScene,
-				&SceneInitializer,
 				NumNativeGPUSceneInstances = Layer.InitializationData.NumNativeGPUSceneInstances,
 				NumNativeCPUInstances = Layer.InitializationData.NumNativeCPUInstances,
 				CullingParameters = View.RayTracingCullingParameters
@@ -314,27 +337,6 @@ void FRayTracingScene::Create(FRDGBuilder& GraphBuilder, const FViewInfo& View,
 						RHICmdList.UnlockBuffer(Layer.TransformUploadBuffer);
 					}
 
-					for (uint32 GPUIndex : RHICmdList.GetGPUMask())
-					{
-						FRayTracingAccelerationStructureAddress* AddressesPtr = (FRayTracingAccelerationStructureAddress*)RHICmdList.LockBufferMGPU(
-							Layer.AccelerationStructureAddressesBuffer.Buffer,
-							GPUIndex,
-							0,
-							SceneInitializer.ReferencedGeometries.Num() * sizeof(FRayTracingAccelerationStructureAddress), RLM_WriteOnly);
-
-						RHICmdList.EnqueueLambda([AddressesPtr, &SceneInitializer, GPUIndex](FRHICommandListBase&)
-							{
-								TRACE_CPUPROFILER_EVENT_SCOPE(GetAccelerationStructuresAddresses);
-
-								for (int32 GeometryIndex = 0; GeometryIndex < SceneInitializer.ReferencedGeometries.Num(); ++GeometryIndex)
-								{
-									AddressesPtr[GeometryIndex] = SceneInitializer.ReferencedGeometries[GeometryIndex]->GetAccelerationStructureAddress(GPUIndex);
-								}
-							});
-
-						RHICmdList.UnlockBufferMGPU(Layer.AccelerationStructureAddressesBuffer.Buffer, GPUIndex);
-					}
-
 					BuildRayTracingInstanceBuffer(
 						RHICmdList,
 						GPUScene,
