Lock/unlock InstanceUploadBuffer/TransformUploadBuffer inside RDG CommandListSetupTask.

diff --git a/Engine/Source/Runtime/Renderer/Private/RayTracing/RayTracingScene.cpp b/Engine/Source/Runtime/Renderer/Private/RayTracing/RayTracingScene.cpp
--- a/Engine/Source/Runtime/Renderer/Private/RayTracing/RayTracingScene.cpp
+++ b/Engine/Source/Runtime/Renderer/Private/RayTracing/RayTracingScene.cpp
@@ -251,16 +251,9 @@ void FRayTracingScene::Create(FRDGBuilder& GraphBuilder, const FViewInfo& View,
 
 		if (NumNativeInstances > 0)
 		{
-			const uint32 InstanceUploadBytes = NumNativeInstances * sizeof(FRayTracingInstanceDescriptorInput);
-			const uint32 TransformUploadBytes = Layer.InitializationData.NumNativeCPUInstances * 3 * sizeof(FVector4f);
-
-			FRayTracingInstanceDescriptorInput* InstanceUploadData = (FRayTracingInstanceDescriptorInput*)RHICmdList.LockBuffer(Layer.InstanceUploadBuffer, 0, InstanceUploadBytes, RLM_WriteOnly);
-			FVector4f* TransformUploadData = (TransformUploadBytes > 0) ? (FVector4f*)RHICmdList.LockBuffer(Layer.TransformUploadBuffer, 0, TransformUploadBytes, RLM_WriteOnly) : nullptr;
-
 			// Fill instance upload buffer on separate thread since results are only needed in RHI thread
-			GraphBuilder.AddSetupTask(
-				[InstanceUploadData = MakeArrayView(InstanceUploadData, NumNativeInstances),
-				TransformUploadData = MakeArrayView(TransformUploadData, Layer.InitializationData.NumNativeCPUInstances * 3),
+			GraphBuilder.AddCommandListSetupTask(
+				[NumNativeInstances,
 				NumNativeGPUSceneInstances = Layer.InitializationData.NumNativeGPUSceneInstances,
 				NumNativeCPUInstances = Layer.InitializationData.NumNativeCPUInstances,
 				Instances = MakeArrayView(Layer.Instances),
@@ -268,9 +261,18 @@ void FRayTracingScene::Create(FRDGBuilder& GraphBuilder, const FViewInfo& View,
 				BaseUploadBufferOffsets = MakeArrayView(Layer.InitializationData.BaseUploadBufferOffsets),
 				BaseInstancePrefixSum = MakeArrayView(Layer.InitializationData.BaseInstancePrefixSum),
 				RayTracingSceneRHI = Layer.RayTracingSceneRHI,
-				PreViewTranslation = this->PreViewTranslation]()
+				InstanceUploadBuffer = Layer.InstanceUploadBuffer,
+				TransformUploadBuffer = Layer.TransformUploadBuffer,
+				PreViewTranslation = this->PreViewTranslation](FRHICommandList& RHICmdList)
 				{
 					FTaskTagScope TaskTagScope(ETaskTag::EParallelRenderingThread);
+
+					const uint32 InstanceUploadBytes = NumNativeInstances * sizeof(FRayTracingInstanceDescriptorInput);
+					const uint32 TransformUploadBytes = NumNativeCPUInstances * 3 * sizeof(FVector4f);
+
+					FRayTracingInstanceDescriptorInput* InstanceUploadData = (FRayTracingInstanceDescriptorInput*)RHICmdList.LockBuffer(InstanceUploadBuffer, 0, InstanceUploadBytes, RLM_WriteOnly);
+					FVector4f* TransformUploadData = (NumNativeCPUInstances > 0) ? (FVector4f*)RHICmdList.LockBuffer(TransformUploadBuffer, 0, TransformUploadBytes, RLM_WriteOnly) : nullptr;
+
 					FillRayTracingInstanceUploadBuffer(
 						RayTracingSceneRHI,
 						PreViewTranslation,
@@ -280,8 +282,15 @@ void FRayTracingScene::Create(FRDGBuilder& GraphBuilder, const FViewInfo& View,
 						BaseInstancePrefixSum,
 						NumNativeGPUSceneInstances,
 						NumNativeCPUInstances,
-						InstanceUploadData,
-						TransformUploadData);
+						MakeArrayView(InstanceUploadData, NumNativeInstances),
+						MakeArrayView(TransformUploadData, NumNativeCPUInstances * 3));
+
+					RHICmdList.UnlockBuffer(InstanceUploadBuffer);
+
+					if (NumNativeCPUInstances > 0)
+					{
+						RHICmdList.UnlockBuffer(TransformUploadBuffer);
+					}
 				});
 
 			GraphBuilder.AddCommandListSetupTask([&Layer, &SceneInitializer](FRHICommandList& RHICmdList)
@@ -330,13 +339,6 @@ void FRayTracingScene::Create(FRDGBuilder& GraphBuilder, const FViewInfo& View,
 				CullingParameters = View.RayTracingCullingParameters
 				](FRHICommandList& RHICmdList)
 				{
-					RHICmdList.UnlockBuffer(Layer.InstanceUploadBuffer);
-
-					if (NumNativeCPUInstances > 0)
-					{
-						RHICmdList.UnlockBuffer(Layer.TransformUploadBuffer);
-					}
-
 					BuildRayTracingInstanceBuffer(
 						RHICmdList,
 						GPUScene,
