Fixes crash on Mac

diff --git a/Engine/Source/Runtime/Renderer/Private/PostProcess/TemporalSuperResolution.cpp b/Engine/Source/Runtime/Renderer/Private/PostProcess/TemporalSuperResolution.cpp
--- a/Engine/Source/Runtime/Renderer/Private/PostProcess/TemporalSuperResolution.cpp
+++ b/Engine/Source/Runtime/Renderer/Private/PostProcess/TemporalSuperResolution.cpp
@@ -2436,7 +2436,7 @@ FDefaultTemporalUpscaler::FOutputs AddTemporalSuperResolutionPasses(
 
 		PassParameters->ReprojectionBoundaryTexture = ReprojectionBoundaryTexture ? ReprojectionBoundaryTexture : GraphBuilder.CreateSRV(FRDGTextureSRVDesc(BlackUintDummy));
 		PassParameters->ReprojectionJacobianTexture = ReprojectionJacobianTexture ? ReprojectionJacobianTexture : GraphBuilder.CreateSRV(FRDGTextureSRVDesc(BlackUintDummy));
-		PassParameters->ReprojectionVectorTexture = GraphBuilder.CreateSRV(FRDGTextureSRVDesc(ReprojectionFieldTexture));
+		PassParameters->ReprojectionVectorTexture = ReprojectionVectorTexture;
 		PassParameters->AntiAliasingTexture = AntiAliasingTexture;
 
 		PassParameters->TranslucencyPixelPosMin = SeparateTranslucencyRect.Min;
@@ -2876,7 +2876,8 @@ FDefaultTemporalUpscaler::FOutputs AddTemporalSuperResolutionPasses(
 			PassParameters->ClosestDepthTexture = ClosestDepthTexture;
 			PassParameters->ReprojectionBoundaryTexture = ReprojectionBoundaryTexture ? ReprojectionBoundaryTexture : GraphBuilder.CreateSRV(FRDGTextureSRVDesc(BlackUintDummy));
 			PassParameters->ReprojectionJacobianTexture = ReprojectionJacobianTexture ? ReprojectionJacobianTexture : GraphBuilder.CreateSRV(FRDGTextureSRVDesc(BlackUintDummy));
-			PassParameters->ReprojectionVectorTexture = GraphBuilder.CreateSRV(FRDGTextureSRVDesc(ReprojectionFieldTexture));			PassParameters->IsMovingMaskTexture = IsMovingMaskTexture ? IsMovingMaskTexture : GraphBuilder.CreateSRV(FRDGTextureSRVDesc(BlackUintDummy));
+			PassParameters->ReprojectionVectorTexture = ReprojectionVectorTexture;
+			PassParameters->IsMovingMaskTexture = IsMovingMaskTexture ? IsMovingMaskTexture : GraphBuilder.CreateSRV(FRDGTextureSRVDesc(BlackUintDummy));
 			PassParameters->HistoryRejectionTexture = HistoryRejectionTexture;
 			PassParameters->MoireHistoryTexture = MoireHistoryTexture ? MoireHistoryTexture : GraphBuilder.CreateSRV(FRDGTextureSRVDesc(BlackDummy));
 			PassParameters->AntiAliasMaskTexture = AntiAliasMaskTexture ? AntiAliasMaskTexture : BlackUintDummy;
