Many Lights - remove dedicated resolve/composite pass and instead composite samples on the fly during shading. This fixed issues with shading atlas overflow and makes everything faster (4.78ms -> 4.40ms).

diff --git a/Engine/Shaders/Private/ManyLights/ManyLights.usf b/Engine/Shaders/Private/ManyLights/ManyLights.usf
--- a/Engine/Shaders/Private/ManyLights/ManyLights.usf
+++ b/Engine/Shaders/Private/ManyLights/ManyLights.usf
@@ -154,25 +154,6 @@ void InitTileIndirectArgsCS(uint3 DispatchThreadId : SV_DispatchThreadID)
 
 #endif
 
-RWBuffer<uint> RWIndirectArgs;
-StructuredBuffer<uint> CompositeTileAllocator;
-uint MaxCompositeTiles;
-
-#ifdef InitCompositeIndirectArgsCS
-
-[numthreads(THREADGROUP_SIZE, 1, 1)]
-void InitCompositeIndirectArgsCS(uint3 DispatchThreadId : SV_DispatchThreadID)
-{
-	if (all(DispatchThreadId == 0))
-	{
-		// Dispatch one group per shadow mask tile
-		uint NumGroups = min(CompositeTileAllocator[0], MaxCompositeTiles);
-		WriteDispatchIndirectArgs(RWIndirectArgs, 0, NumGroups, 1, 1);
-	}
-}
-
-#endif
-
 RWTexture2D<float4> RWCompositeUpsampleWeights;
 
 Texture2D<float> DownsampledSceneDepth;
diff --git a/Engine/Shaders/Private/ManyLights/ManyLights.ush b/Engine/Shaders/Private/ManyLights/ManyLights.ush
--- a/Engine/Shaders/Private/ManyLights/ManyLights.ush
+++ b/Engine/Shaders/Private/ManyLights/ManyLights.ush
@@ -3,7 +3,7 @@
 #pragma once
 
 // Change this to force recompilation of all Many Lights shaders
-#pragma message("UESHADERMETADATA_VERSION DE81D00C-08D2-4DA6-956F-B7566B5AF5A8")
+#pragma message("UESHADERMETADATA_VERSION A48A689C-146E-4051-BAD6-21564B78503F")
 
 #define MANY_LIGHTS 1
 #include "../LightGridCommon.ush"
@@ -24,10 +24,9 @@
 #define DEBUG_MODE_VISUALIZE_TRACING  2
 
 // Keep in sync with ManyLights.cpp
-#define SHADING_ATLAS_SIZE_IN_TILES 512
 #define DOWNSAMPLE_FACTOR			2
 #define TILE_SIZE					8
-// Limited by PackLightSample, PackShadingTile and SharedCandidateLightHiMask
+// Limited by PackLightSample and SharedCandidateLightHiMask
 #define MAX_LOCAL_LIGHT_INDEX		1024
 #define SHARED_LIGHT_MASK_SIZE		(MAX_LOCAL_LIGHT_INDEX / 32)
 
@@ -47,14 +46,6 @@ uint2 UnpackTile(uint PackedTile)
 	return uint2(PackedTile & 0xFFFF, PackedTile >> 16);
 }
 
-uint2 GetShadingTileAtlasCoord(uint ShadingTileIndex, uint2 CoordInTile)
-{
-	uint2 AtlasCoord;
-	AtlasCoord.x = (ShadingTileIndex % SHADING_ATLAS_SIZE_IN_TILES) * TILE_SIZE + CoordInTile.x;
-	AtlasCoord.y = (ShadingTileIndex / SHADING_ATLAS_SIZE_IN_TILES) * TILE_SIZE + CoordInTile.y;
-	return AtlasCoord;
-}
-
 struct FLightSample
 {
 	uint LocalLightIndex;
@@ -122,54 +113,6 @@ float2 DownsampledScreenCoordToScreenUV(uint2 DownsampledScreenCoord)
 	return ScreenUV;
 }
 
-uint2 PackCompositeTile(uint LocalLightIndex, uint2 DownsampledScreenCoord)
-{
-	DownsampledScreenCoord /= TILE_SIZE;
-
-	uint2 PackedTile;
-	PackedTile.x = LocalLightIndex;
-	PackedTile.y = (DownsampledScreenCoord.x & 0xFFFF) | ((DownsampledScreenCoord.y & 0xFFFF) << 16);
-	return PackedTile;
-}
-
-void UnpackCompositeTile(uint2 PackedTile, out uint LocalLightIndex, out uint2 DownsampledScreenCoord)
-{
-	LocalLightIndex = PackedTile.x;
-	DownsampledScreenCoord.x = PackedTile.y & 0xFFFF;
-	DownsampledScreenCoord.y = PackedTile.y >> 16;
-
-	DownsampledScreenCoord *= TILE_SIZE;
-}
-
-struct FShadingTile
-{
-	uint LocalLightIndex;
-	uint ShadingTileIndex;
-};
-
-uint PackShadingTile(uint LocalLightIndex, uint ShadingTileIndex)
-{
-	uint PackedShadingTile = LocalLightIndex & 0xFFF;
-	PackedShadingTile |= ShadingTileIndex << 12;
-	return PackedShadingTile;
-}
-
-FShadingTile UnpackShadingTile(uint PackedShadingTile)
-{
-	FShadingTile ShadingTile;
-	ShadingTile.LocalLightIndex = PackedShadingTile & 0xFFF;
-	ShadingTile.ShadingTileIndex = PackedShadingTile >> 12;
-	return ShadingTile;
-}
-
-uint2 ShadingTileGridSize;
-uint MaxShadingTilesPerGridCell;
-
-uint GetShadingTileGridOffset(uint2 ShadingTileCoord, uint IndexInCell)
-{
-	return (ShadingTileCoord.x + ShadingTileCoord.y * ShadingTileGridSize.x) * MaxShadingTilesPerGridCell + IndexInCell;
-}
-
 // Remove texture detail before denoising
 void DemodulateLighting(FGBufferData GBuffer, float3 TranslatedWorldPosition, inout float3 DiffuseLighting, inout float3 SpecularLighting)
 {
diff --git a/Engine/Shaders/Private/ManyLights/ManyLightsResolve.usf b/Engine/Shaders/Private/ManyLights/ManyLightsResolve.usf
deleted file mode 100644
diff --git a/Engine/Shaders/Private/ManyLights/ManyLightsSampling.usf b/Engine/Shaders/Private/ManyLights/ManyLightsSampling.usf
--- a/Engine/Shaders/Private/ManyLights/ManyLightsSampling.usf
+++ b/Engine/Shaders/Private/ManyLights/ManyLightsSampling.usf
@@ -73,7 +73,6 @@ FLightTargetPDF GetLocalLightTargetPDF(FDeferredLightData LightData, float3 Tran
 
 uint2 DownsampledViewMin;
 uint2 DownsampledViewSize;
-uint MaxCompositeTiles;
 float SamplingMinWeight;
 uint2 NumSamplesPerPixel;
 int DebugMode;
@@ -81,26 +80,14 @@ int DebugMode;
 RWTexture2D<float> RWDownsampledSceneDepth;
 RWTexture2D<UNORM float3> RWDownsampledSceneWorldNormal;
 RWTexture2D<float2> RWSampleLuminanceSum;
-RWStructuredBuffer<uint> RWCompositeTileAllocator;
-RWStructuredBuffer<uint2> RWCompositeTileData;
 RWTexture2D<uint> RWLightSamples;
 
 float4 HistoryScreenPositionScaleBias;
 float4 HistoryUVMinMax;
 
-groupshared uint SharedLightMask[SHARED_LIGHT_MASK_SIZE];
-
 groupshared uint SharedCandidateLightHiMask;
 groupshared uint SharedCandidateLightMask[SHARED_LIGHT_MASK_SIZE];
 
-void MarkLocalLightIndex(uint LocalLightIndex)
-{
-	// Re-distribute nearby ids to different bytes in order to balance gather workload
-	uint DWORDIndex = LocalLightIndex % SHARED_LIGHT_MASK_SIZE;
-	uint BitMask = 1u << (LocalLightIndex / SHARED_LIGHT_MASK_SIZE);
-	InterlockedOr(SharedLightMask[DWORDIndex], BitMask);
-}
-
 StructuredBuffer<uint> DownsampledTileAllocator;
 StructuredBuffer<uint> DownsampledTileData;
 uint DownsampledTileDataStride;
@@ -124,7 +111,6 @@ void GenerateLightSamplesCS(
 
 	if (LinearThreadIndex < SHARED_LIGHT_MASK_SIZE)
 	{
-		SharedLightMask[LinearThreadIndex] = 0;
 		SharedCandidateLightMask[LinearThreadIndex] = 0;
 	}
 
@@ -381,43 +367,6 @@ void GenerateLightSamplesCS(
 				}
 
 				RWLightSamples[DownsampledScreenCoord * NumSamplesPerPixel + LightSampleOffsets[LightSampleIndex]] = PackLightSample(LightSample);
-
-				// Mark tiles which need to be processed
-				if (LightSample.LocalLightIndex != MAX_LOCAL_LIGHT_INDEX)
-				{
-					MarkLocalLightIndex(LightSample.LocalLightIndex);
-				}
-			}
-		}
-
-		GroupMemoryBarrierWithGroupSync();
-
-		// Write out allocated shadow mask tiles
-		// #ml_todo: process part of the list per tread
-		if (LinearThreadIndex < SHARED_LIGHT_MASK_SIZE)
-		{
-			const uint LightIndex = LinearThreadIndex;
-			const uint BitMask = 1u << LightIndex;
-
-			uint2 DownsampledScreenCoord = DownsampledTileCoord * TILE_SIZE;
-	
-			uint LightMask = SharedLightMask[LinearThreadIndex];
-			while (LightMask != 0)
-			{
-				const uint NextBitIndex = firstbitlow(LightMask);
-				const uint NextBitMask = 1u << NextBitIndex;
-				LightMask ^= NextBitMask;
-
-				const uint LocalLightIndex = LinearThreadIndex + NextBitIndex * SHARED_LIGHT_MASK_SIZE;
-
-				// #ml_todo: merge allocations through groupshared and do only one atomic per group, or just replace entire thing with a scalar loop over SelectedSamples?
-				uint CompositeTileIndex = 0;
-				InterlockedAdd(RWCompositeTileAllocator[0], 1, CompositeTileIndex);
-
-				if (CompositeTileIndex < MaxCompositeTiles)
-				{
-					RWCompositeTileData[CompositeTileIndex] = PackCompositeTile(LocalLightIndex, DownsampledScreenCoord);
-				}
 			}
 		}
 	}
diff --git a/Engine/Shaders/Private/ManyLights/ManyLightsShading.usf b/Engine/Shaders/Private/ManyLights/ManyLightsShading.usf
--- a/Engine/Shaders/Private/ManyLights/ManyLightsShading.usf
+++ b/Engine/Shaders/Private/ManyLights/ManyLightsShading.usf
@@ -12,20 +12,43 @@
 RWTexture2D<float3> RWResolvedDiffuseLighting;
 RWTexture2D<float3> RWResolvedSpecularLighting;
 
-Texture2D<float> ShadingTileAtlas;
-Texture2D<uint> ShadingTileGridAllocator;
-StructuredBuffer<uint> ShadingTileGrid;
-StructuredBuffer<uint> CompositeTileAllocator;
-StructuredBuffer<uint> ShadingTileAllocator;
-uint MaxCompositeTiles;
-uint MaxShadingTiles;
+int2 SampleViewMin;
+int2 SampleViewSize;
 
 StructuredBuffer<uint> TileAllocator;
 StructuredBuffer<uint> TileData;
 uint TileDataStride;
 
+Texture2D<uint> LightSamples;
+Texture2D<float4> CompositeUpsampleWeights;
+
+// 4x4 neighborhood with borders
+#define SHARED_TILE_SIZE_X ((4 + 2) * NUM_SAMPLES_PER_PIXEL_2D_X)
+#define SHARED_TILE_SIZE_Y ((4 + 2) * NUM_SAMPLES_PER_PIXEL_2D_Y)
+groupshared uint SharedSamples[SHARED_TILE_SIZE_X][SHARED_TILE_SIZE_Y];
+
+// Light scalarization
+groupshared uint SharedLightHiMask;
+groupshared uint SharedLightMask[SHARED_LIGHT_MASK_SIZE];
+
+void AccumulateLightSample(uint2 SharedSampleCoord, uint LocalLightIndex, float UpsampleWeight, inout float SampleWeightSum)
+{
+	for (uint SampleOffsetY = 0; SampleOffsetY < NUM_SAMPLES_PER_PIXEL_2D_Y; ++SampleOffsetY)
+	{
+		for (uint SampleOffsetX = 0; SampleOffsetX < NUM_SAMPLES_PER_PIXEL_2D_X; ++SampleOffsetX)
+		{
+			FLightSample LightSample = UnpackLightSample(SharedSamples[SharedSampleCoord.x * NUM_SAMPLES_PER_PIXEL_2D_X + SampleOffsetX][SharedSampleCoord.y * NUM_SAMPLES_PER_PIXEL_2D_Y + SampleOffsetY]);
+
+			if (LightSample.LocalLightIndex == LocalLightIndex)
+			{
+				SampleWeightSum += LightSample.Weight * UpsampleWeight;
+			}
+		}
+	}
+}
+
 /**
- * Apply all lights per pixel to affected tiles using sampled weights from upscaled shading tiles
+ * Upsample light weights and apply all lights per pixel to affected tiles
  */
 [numthreads(THREADGROUP_SIZE, THREADGROUP_SIZE, 1)]
 void ShadeLightSamplesCS(
@@ -33,11 +56,69 @@ void ShadeLightSamplesCS(
 	uint3 GroupThreadId : SV_GroupThreadID,
 	uint3 DispatchThreadId : SV_DispatchThreadID)
 {
+	uint LinearThreadIndex = GroupThreadId.x + GroupThreadId.y * THREADGROUP_SIZE;
+
+	if (LinearThreadIndex == 0)
+	{
+		SharedLightHiMask = 0;
+	}
+
+	if (LinearThreadIndex < SHARED_LIGHT_MASK_SIZE)
+	{
+		SharedLightMask[LinearThreadIndex] = 0;
+	}
+
+	GroupMemoryBarrierWithGroupSync();
+
 	uint TileIndex = GroupId.x;
 	if (TileIndex < TileAllocator[TILE_TYPE])
 	{
 		uint2 TileCoord = UnpackTile(TileData[TileIndex + TILE_TYPE * TileDataStride]);
 		uint2 ScreenCoord = TileCoord * TILE_SIZE + GroupThreadId.xy;
+		uint LocalLightHiMask = 0;
+
+		// First load all samples into shared memory
+		for (uint SharedCoordY = GroupThreadId.y; SharedCoordY < SHARED_TILE_SIZE_Y; SharedCoordY += THREADGROUP_SIZE)
+		{
+			for (uint SharedCoordX = GroupThreadId.x; SharedCoordX < SHARED_TILE_SIZE_X; SharedCoordX += THREADGROUP_SIZE)
+			{
+				const uint2 NumSamplesPerPixel = uint2(NUM_SAMPLES_PER_PIXEL_2D_X, NUM_SAMPLES_PER_PIXEL_2D_Y);
+				uint2 SharedSampleCoord = uint2(SharedCoordX, SharedCoordY);
+				int2 DownsampledScreenCoord = int2(TileCoord * TILE_SIZE - 1) / DOWNSAMPLE_FACTOR + SharedSampleCoord / NumSamplesPerPixel;
+				int2 SampleCoord = DownsampledScreenCoord * NumSamplesPerPixel + SharedSampleCoord % NumSamplesPerPixel;
+
+				FLightSample LightSample = InitLightSample();
+
+				if (and(all(SampleCoord >= SampleViewMin), all(SampleCoord < SampleViewMin + SampleViewSize)))
+				{
+					LightSample = UnpackLightSample(LightSamples[SampleCoord]);
+
+					if (LightSample.bVisible && LightSample.LocalLightIndex < MAX_LOCAL_LIGHT_INDEX)
+					{
+						uint DWORDIndex = LightSample.LocalLightIndex / 32;
+						uint BitMask = 1u << (LightSample.LocalLightIndex % 32);
+						InterlockedOr(SharedLightMask[DWORDIndex], BitMask);
+
+						uint HiBitMask = 1u << DWORDIndex;
+						LocalLightHiMask |= HiBitMask;
+					}
+					else
+					{
+						LightSample = InitLightSample();
+					}
+				}
+
+				SharedSamples[SharedSampleCoord.x][SharedSampleCoord.y] = PackLightSample(LightSample);
+			}
+		}
+
+		uint WaveHiMask = WaveActiveBitOr(LocalLightHiMask);
+		if (WaveIsFirstLane())
+		{
+			InterlockedOr(SharedLightHiMask, WaveHiMask);
+		}
+
+		GroupMemoryBarrierWithGroupSync();
 
 		if (all(ScreenCoord < View.ViewRectMinAndSize.xy + View.ViewRectMinAndSize.zw))
 		{
@@ -61,97 +142,109 @@ void ShadeLightSamplesCS(
 				DiffuseLighting = 0.0f;
 				SpecularLighting = 0.0f;
 
-				uint2 ShadingTileCoord = TileCoord;
-				uint2 CoordInShadingTile = GroupThreadId.xy;
-
 				float3 TranslatedWorldPosition = GetTranslatedWorldPositionFromScreenUV(ScreenUV, GBuffer.Depth);		
 				float3 CameraVector = normalize(TranslatedWorldPosition - View.TranslatedWorldCameraOrigin);
 
-				uint NumShadingTilesInCell = ShadingTileGridAllocator.Load(int3(ShadingTileCoord, 0));
-
 				#if DEBUG_MODE
 				if (bDebug)
 				{
 					Print(Context, TEXT("ShadeSamplesCS"));
 					Newline(Context);
+					Print(Context, UnpackLightSample(SharedSamples[0][0]).LocalLightIndex);
+					Print(Context, UnpackLightSample(SharedSamples[1][0]).LocalLightIndex);
+					Print(Context, UnpackLightSample(SharedSamples[2][0]).LocalLightIndex);
+					Print(Context, UnpackLightSample(SharedSamples[3][0]).LocalLightIndex);
+					Print(Context, UnpackLightSample(SharedSamples[4][0]).LocalLightIndex);
+					Print(Context, UnpackLightSample(SharedSamples[5][0]).LocalLightIndex);
+					Newline(Context);
 					Print(Context, TEXT("ScreenCoord: "));
 					Print(Context, ScreenCoord.x);
 					Print(Context, ScreenCoord.y);
 					Newline(Context);
-					Print(Context, TEXT("NumShadingTilesInCell: "));
-					Print(Context, NumShadingTilesInCell);
-					Newline(Context);
-					Print(Context, TEXT("ShadingTiles:    "));
-					Print(Context, ShadingTileAllocator[0]);
-					Print(Context, MaxShadingTiles);
-					Print(Context, ShadingTileAllocator[0] / float(MaxShadingTiles));
-					Newline(Context);
 					Print(Context, TEXT("LightId   | Weight    | Diffuse   | Specular"));
 				}
 				#endif
 
-				for (uint IndexInCell = 0; IndexInCell < NumShadingTilesInCell; ++IndexInCell)
+				uint LightHiMask = SharedLightHiMask;
+				while (LightHiMask != 0)
 				{
-					const uint PackedShadingTile = ShadingTileGrid[GetShadingTileGridOffset(ShadingTileCoord, IndexInCell)];
-					const FShadingTile ShadingTile = UnpackShadingTile(PackedShadingTile);
-					const FLocalLightData LocalLightData = GetLocalLightDataNonStereo(ShadingTile.LocalLightIndex);
-					FDeferredLightData LightData = ConvertToDeferredLight(LocalLightData);
+					const uint NextHiBitIndex = firstbitlow(LightHiMask);
+					const uint NextHiBitMask = 1u << NextHiBitIndex;
+					LightHiMask ^= NextHiBitMask;
 
-					float SampleWeight = 0.0f;
+					const uint MaskIndex = NextHiBitIndex;
+					uint LightMask = SharedLightMask[MaskIndex];
 
-					if (ShadingTile.ShadingTileIndex < MaxShadingTiles)
+					while (LightMask != 0)
 					{
-						uint2 ShadingTileAtlasCoord = GetShadingTileAtlasCoord(ShadingTile.ShadingTileIndex, CoordInShadingTile);
-						SampleWeight = ShadingTileAtlas[ShadingTileAtlasCoord];
-					}
+						const uint NextBitIndex = firstbitlow(LightMask);
+						const uint NextBitMask = 1u << NextBitIndex;
+						LightMask ^= NextBitMask;
 
-					if (SampleWeight > 0.01f && LightData.IESAtlasIndex >= 0)
-					{
-						// #ml_todo: evaluate moving this to CompositeLightSamplesCS
-						SampleWeight *= ComputeLightProfileMultiplier(TranslatedWorldPosition, LightData.TranslatedWorldPosition, -LightData.Direction, LightData.Tangent, LightData.IESAtlasIndex);
-					}
+						const uint LocalLightIndex = MaskIndex * 32 + NextBitIndex;
 
-					#if DEBUG_MODE
-					if (bDebug)
-					{
-						Newline(Context);
-						Print(Context, LocalLightData.LightSceneId);
-						Print(Context, SampleWeight);
-					}
-					#endif
+						const FLocalLightData LocalLightData = GetLocalLightDataNonStereo(LocalLightIndex);
+						FDeferredLightData LightData = ConvertToDeferredLight(LocalLightData);
 
-					if (SampleWeight > 0.0f)
-					{
-						GBuffer.Roughness = max(GBuffer.Roughness, View.MinRoughness);
+						float4 DownsampledScreenWeights = CompositeUpsampleWeights[ScreenCoord];
+						uint2 SharedSampleCoord00 = (GroupThreadId.xy + 1) / DOWNSAMPLE_FACTOR;
 
-						float4 LightAttenuation = 1.0f;
-						float Dither = 0.5f;
-						float SurfaceShadow = 1;
-						float AmbientOcclusion = ScreenSpaceAO;
+						float SampleWeight = 0.0f;
+						AccumulateLightSample(SharedSampleCoord00 + uint2(0, 0), LocalLightIndex, DownsampledScreenWeights.x, SampleWeight);
+						AccumulateLightSample(SharedSampleCoord00 + uint2(1, 0), LocalLightIndex, DownsampledScreenWeights.y, SampleWeight);
+						AccumulateLightSample(SharedSampleCoord00 + uint2(0, 1), LocalLightIndex, DownsampledScreenWeights.z, SampleWeight);
+						AccumulateLightSample(SharedSampleCoord00 + uint2(1, 1), LocalLightIndex, DownsampledScreenWeights.w, SampleWeight);
 
-						// #ml_todo: Move more shading models to TILE_MODE_SIMPLE_SHADING
-						#if TILE_TYPE == TILE_MODE_SIMPLE_SHADING
+						if (SampleWeight > 0.01f && LightData.IESAtlasIndex >= 0)
 						{
-							GBuffer.ShadingModelID = SHADINGMODELID_DEFAULT_LIT;
+							// #ml_todo: evaluate moving this to CompositeLightSamplesCS
+							SampleWeight *= ComputeLightProfileMultiplier(TranslatedWorldPosition, LightData.TranslatedWorldPosition, -LightData.Direction, LightData.Tangent, LightData.IESAtlasIndex);
 						}
-						#endif
-						LightData.ShadowedBits = 0;
-
-						FDeferredLightingSplit SplitLighting = GetDynamicLightingSplit(
-							TranslatedWorldPosition, CameraVector, GBuffer, AmbientOcclusion, GBuffer.ShadingModelID, 
-							LightData, LightAttenuation, Dither, ScreenCoord, 
-							SurfaceShadow);
-
-						DiffuseLighting += SplitLighting.DiffuseLighting.xyz * SampleWeight;
-						SpecularLighting += SplitLighting.SpecularLighting.xyz * SampleWeight;
 
-						#if DEBUG_MODE 
+						#if DEBUG_MODE
 						if (bDebug)
 						{
-							Print(Context, SplitLighting.DiffuseLighting.xyz);
-							Print(Context, SplitLighting.SpecularLighting.xyz);
+							Newline(Context);
+							Print(Context, DownsampledScreenWeights);
+							Newline(Context);
+							Print(Context, LocalLightData.LightSceneId);
+							Print(Context, SampleWeight);
 						}
 						#endif
+
+						if (SampleWeight > 0.0f)
+						{
+							GBuffer.Roughness = max(GBuffer.Roughness, View.MinRoughness);
+
+							float4 LightAttenuation = 1.0f;
+							float Dither = 0.5f;
+							float SurfaceShadow = 1;
+							float AmbientOcclusion = ScreenSpaceAO;
+
+							// #ml_todo: Move more shading models to TILE_MODE_SIMPLE_SHADING
+							#if TILE_TYPE == TILE_MODE_SIMPLE_SHADING
+							{
+								GBuffer.ShadingModelID = SHADINGMODELID_DEFAULT_LIT;
+							}
+							#endif
+							LightData.ShadowedBits = 0;
+
+							FDeferredLightingSplit SplitLighting = GetDynamicLightingSplit(
+								TranslatedWorldPosition, CameraVector, GBuffer, AmbientOcclusion, GBuffer.ShadingModelID, 
+								LightData, LightAttenuation, Dither, ScreenCoord, 
+								SurfaceShadow);
+
+							DiffuseLighting += SplitLighting.DiffuseLighting.xyz * SampleWeight;
+							SpecularLighting += SplitLighting.SpecularLighting.xyz * SampleWeight;
+
+							#if DEBUG_MODE 
+							if (bDebug)
+							{
+								Print(Context, SplitLighting.DiffuseLighting.xyz);
+								Print(Context, SplitLighting.SpecularLighting.xyz);
+							}
+							#endif
+						}
 					}
 				}
 
diff --git a/Engine/Source/Runtime/Renderer/Private/ManyLights/ManyLights.cpp b/Engine/Source/Runtime/Renderer/Private/ManyLights/ManyLights.cpp
--- a/Engine/Source/Runtime/Renderer/Private/ManyLights/ManyLights.cpp
+++ b/Engine/Source/Runtime/Renderer/Private/ManyLights/ManyLights.cpp
@@ -25,13 +25,6 @@ static TAutoConsoleVariable<int32> CVarManyLightsNumSamplesPerPixel(
 	ECVF_Scalability | ECVF_RenderThreadSafe
 );
 
-static TAutoConsoleVariable<int32> CVarManyLightsMaxShadingTilesPerGridCell(
-	TEXT("r.ManyLights.MaxShadingTilesPerGridCell"),
-	32,
-	TEXT("Maximum number of shading tiles per grid cell."),
-	ECVF_Scalability | ECVF_RenderThreadSafe
-);
-
 static TAutoConsoleVariable<float> CVarManyLightsSamplingMinWeight(
 	TEXT("r.ManyLights.Sampling.MinWeight"),
 	0.001f,
@@ -156,13 +149,18 @@ static TAutoConsoleVariable<int32> CVarManyLightsIESProfiles(
 	ECVF_Scalability | ECVF_RenderThreadSafe
 );
 
+static TAutoConsoleVariable<int32> CVarManyLightsTest(
+	TEXT("r.ManyLights.Test"),
+	1,
+	TEXT("Test."),
+	ECVF_Scalability | ECVF_RenderThreadSafe
+);
+
 namespace ManyLights
 {
 	// must match values in ManyLights.ush
 	constexpr int32 TileSize = 8;
 	constexpr int32 MaxLocalLightIndexXY = 16; // 16 * 16 = 256
-	constexpr uint32 ShadingTileIndexUnshadowed = 0xFFFFF; // limited by PackShadingTile()
-	constexpr uint32 ShadingAtlasSizeInTiles = 512;
 
 	bool IsEnabled()
 	{
@@ -208,10 +206,15 @@ namespace ManyLights
 		return StateFrameIndex;
 	}
 
+	FIntPoint GetNumSamplesPerPixel2d(int32 NumSamplesPerPixel1d)
+	{
+		return NumSamplesPerPixel1d == 4 ? FIntPoint(2, 2) : (NumSamplesPerPixel1d == 2 ? FIntPoint(2, 1) : FIntPoint(1, 1));
+	}
+
 	FIntPoint GetNumSamplesPerPixel2d()
 	{
 		const uint32 NumSamplesPerPixel1d = FMath::RoundUpToPowerOfTwo(FMath::Clamp(CVarManyLightsNumSamplesPerPixel.GetValueOnAnyThread(), 1, 4));
-		return NumSamplesPerPixel1d == 4 ? FIntPoint(2, 2) : (NumSamplesPerPixel1d == 2 ? FIntPoint(2, 1) : FIntPoint(1, 1));
+		return GetNumSamplesPerPixel2d(NumSamplesPerPixel1d);
 	}
 
 	int32 GetDebugMode()
@@ -309,36 +312,6 @@ class FInitTileIndirectArgsCS : public FGlobalShader
 
 IMPLEMENT_GLOBAL_SHADER(FInitTileIndirectArgsCS, "/Engine/Private/ManyLights/ManyLights.usf", "InitTileIndirectArgsCS", SF_Compute);
 
-class FInitCompositeIndirectArgsCS : public FGlobalShader
-{
-	DECLARE_GLOBAL_SHADER(FInitCompositeIndirectArgsCS)
-	SHADER_USE_PARAMETER_STRUCT(FInitCompositeIndirectArgsCS, FGlobalShader);
-
-	BEGIN_SHADER_PARAMETER_STRUCT(FParameters, )
-		SHADER_PARAMETER_STRUCT_INCLUDE(FManyLightsParameters, ManyLightsParameters)
-		SHADER_PARAMETER_RDG_BUFFER_UAV(RWBuffer<uint>, RWIndirectArgs)
-		SHADER_PARAMETER_RDG_BUFFER_SRV(StructuredBuffer<uint>, CompositeTileAllocator)
-	END_SHADER_PARAMETER_STRUCT()
-
-	static bool ShouldCompilePermutation(const FGlobalShaderPermutationParameters& Parameters)
-	{
-		return ManyLights::ShouldCompileShaders(Parameters);
-	}
-
-	static void ModifyCompilationEnvironment(const FGlobalShaderPermutationParameters& Parameters, FShaderCompilerEnvironment& OutEnvironment)
-	{
-		FGlobalShader::ModifyCompilationEnvironment(Parameters, OutEnvironment);
-		OutEnvironment.SetDefine(TEXT("THREADGROUP_SIZE"), GetGroupSize());
-	}
-
-	static int32 GetGroupSize()
-	{
-		return 64;
-	}
-};
-
-IMPLEMENT_GLOBAL_SHADER(FInitCompositeIndirectArgsCS, "/Engine/Private/ManyLights/ManyLights.usf", "InitCompositeIndirectArgsCS", SF_Compute);
-
 class FGenerateLightSamplesCS : public FGlobalShader
 {
 	DECLARE_GLOBAL_SHADER(FGenerateLightSamplesCS)
@@ -350,8 +323,6 @@ class FGenerateLightSamplesCS : public FGlobalShader
 		SHADER_PARAMETER_RDG_TEXTURE_UAV(RWTexture2D<float2>, RWSampleLuminanceSum)
 		SHADER_PARAMETER_RDG_TEXTURE_UAV(RWTexture2D<float>, RWDownsampledSceneDepth)
 		SHADER_PARAMETER_RDG_TEXTURE_UAV(RWTexture2D<UNORM float3>, RWDownsampledSceneWorldNormal)
-		SHADER_PARAMETER_RDG_BUFFER_UAV(RWStructuredBuffer<uint>, RWCompositeTileAllocator)
-		SHADER_PARAMETER_RDG_BUFFER_UAV(RWStructuredBuffer<uint2>, RWCompositeTileData)
 		SHADER_PARAMETER_RDG_TEXTURE_UAV(RWTexture2D<uint>, RWLightSamples)
 		SHADER_PARAMETER_RDG_BUFFER_SRV(StructuredBuffer<uint>, DownsampledTileAllocator)
 		SHADER_PARAMETER_RDG_BUFFER_SRV(StructuredBuffer<uint>, DownsampledTileData)
@@ -484,46 +455,6 @@ class FInitCompositeUpsampleWeightsCS : public FGlobalShader
 
 IMPLEMENT_GLOBAL_SHADER(FInitCompositeUpsampleWeightsCS, "/Engine/Private/ManyLights/ManyLights.usf", "InitCompositeUpsampleWeightsCS", SF_Compute);
 
-class FResolveLightSamplesCS : public FGlobalShader
-{
-	DECLARE_GLOBAL_SHADER(FResolveLightSamplesCS)
-	SHADER_USE_PARAMETER_STRUCT(FResolveLightSamplesCS, FGlobalShader)
-
-	BEGIN_SHADER_PARAMETER_STRUCT(FParameters, )
-		SHADER_PARAMETER_STRUCT_INCLUDE(FManyLightsParameters, ManyLightsParameters)
-		RDG_BUFFER_ACCESS(IndirectArgs, ERHIAccess::IndirectArgs)
-		SHADER_PARAMETER_RDG_BUFFER_UAV(RWStructuredBuffer<uint>, RWShadingTileAllocator)
-		SHADER_PARAMETER_RDG_TEXTURE_UAV(RWTexture2D<uint>, RWShadingTileGridAllocator)
-		SHADER_PARAMETER_RDG_BUFFER_UAV(RWStructuredBuffer<uint>, RWShadingTileGrid)
-		SHADER_PARAMETER_RDG_TEXTURE_UAV(RWTexture2D<float>, RWShadingTileAtlas)
-		SHADER_PARAMETER_RDG_BUFFER_SRV(StructuredBuffer<uint2>, CompositeTileData)
-		SHADER_PARAMETER_RDG_TEXTURE(Texture2D<flaot4>, CompositeUpsampleWeights)
-		SHADER_PARAMETER_RDG_TEXTURE(Texture2D<uint>, LightSamples)
-	END_SHADER_PARAMETER_STRUCT()
-
-	static int32 GetGroupSize()
-	{
-		return 8;
-	}
-
-	class FDebugMode : SHADER_PERMUTATION_BOOL("DEBUG_MODE");
-	using FPermutationDomain = TShaderPermutationDomain<FDebugMode>;
-
-	static bool ShouldCompilePermutation(const FGlobalShaderPermutationParameters& Parameters)
-	{
-		return ManyLights::ShouldCompileShaders(Parameters);
-	}
-
-	static void ModifyCompilationEnvironment(const FGlobalShaderPermutationParameters& Parameters, FShaderCompilerEnvironment& OutEnvironment)
-	{
-		FGlobalShader::ModifyCompilationEnvironment(Parameters, OutEnvironment);
-		ManyLights::ModifyCompilationEnvironment(Parameters.Platform, OutEnvironment);
-		OutEnvironment.SetDefine(TEXT("THREADGROUP_SIZE"), GetGroupSize());
-	}
-};
-
-IMPLEMENT_GLOBAL_SHADER(FResolveLightSamplesCS, "/Engine/Private/ManyLights/ManyLightsResolve.usf", "ResolveLightSamplesCS", SF_Compute);
-
 class FShadeLightSamplesCS : public FGlobalShader
 {
 	DECLARE_GLOBAL_SHADER(FShadeLightSamplesCS)
@@ -536,12 +467,8 @@ class FShadeLightSamplesCS : public FGlobalShader
 		SHADER_PARAMETER_RDG_TEXTURE_UAV(RWTexture2D<float3>, RWResolvedSpecularLighting)
 		SHADER_PARAMETER_RDG_BUFFER_SRV(StructuredBuffer<uint>, TileAllocator)
 		SHADER_PARAMETER_RDG_BUFFER_SRV(StructuredBuffer<uint>, TileData)
-		SHADER_PARAMETER_RDG_BUFFER_SRV(StructuredBuffer<uint>, ShadingTileAllocator)
-		SHADER_PARAMETER_RDG_BUFFER_SRV(StructuredBuffer<uint>, CompositeTileAllocator)
-		SHADER_PARAMETER_RDG_TEXTURE(Texture2D<uint>, ShadingTileGridAllocator)
-		SHADER_PARAMETER_RDG_BUFFER_SRV(StructuredBuffer<uint>, ShadingTileGrid)
-		SHADER_PARAMETER_RDG_BUFFER_SRV(StructuredBuffer<uint>, ShadingTiles)
-		SHADER_PARAMETER_RDG_TEXTURE(Texture2D<float>, ShadingTileAtlas)
+		SHADER_PARAMETER_RDG_TEXTURE(Texture2D<flaot4>, CompositeUpsampleWeights)
+		SHADER_PARAMETER_RDG_TEXTURE(Texture2D<uint>, LightSamples)
 	END_SHADER_PARAMETER_STRUCT()
 
 	static int32 GetGroupSize()
@@ -553,8 +480,9 @@ class FShadeLightSamplesCS : public FGlobalShader
 	class FIESProfile : SHADER_PERMUTATION_BOOL("USE_IES_PROFILE");
 	class FLightFunctionAtlas : SHADER_PERMUTATION_BOOL("USE_LIGHT_FUNCTION_ATLAS");
 	class FTexturedRectLights : SHADER_PERMUTATION_BOOL("USE_SOURCE_TEXTURE");
+	class FNumSamplesPerPixel1d : SHADER_PERMUTATION_SPARSE_INT("NUM_SAMPLES_PER_PIXEL_1D", 1, 2, 4);
 	class FDebugMode : SHADER_PERMUTATION_BOOL("DEBUG_MODE");
-	using FPermutationDomain = TShaderPermutationDomain<FTileType, FIESProfile, FLightFunctionAtlas, FTexturedRectLights, FDebugMode>;
+	using FPermutationDomain = TShaderPermutationDomain<FTileType, FIESProfile, FLightFunctionAtlas, FTexturedRectLights, FNumSamplesPerPixel1d, FDebugMode>;
 
 	static bool ShouldCompilePermutation(const FGlobalShaderPermutationParameters& Parameters)
 	{
@@ -588,6 +516,12 @@ class FShadeLightSamplesCS : public FGlobalShader
 		FGlobalShader::ModifyCompilationEnvironment(Parameters, OutEnvironment);
 		ManyLights::ModifyCompilationEnvironment(Parameters.Platform, OutEnvironment);
 		OutEnvironment.SetDefine(TEXT("THREADGROUP_SIZE"), GetGroupSize());
+
+		FPermutationDomain PermutationVector(Parameters.PermutationId);
+		const int32 NumSamplesPerPixel1d = PermutationVector.Get<FNumSamplesPerPixel1d>();
+		const FIntPoint NumSamplesPerPixel2d = ManyLights::GetNumSamplesPerPixel2d(NumSamplesPerPixel1d);
+		OutEnvironment.SetDefine(TEXT("NUM_SAMPLES_PER_PIXEL_2D_X"), NumSamplesPerPixel2d.X);
+		OutEnvironment.SetDefine(TEXT("NUM_SAMPLES_PER_PIXEL_2D_Y"), NumSamplesPerPixel2d.Y);
 	}
 };
 
@@ -764,16 +698,6 @@ void FDeferredShadingSceneRenderer::RenderManyLights(FRDGBuilder& GraphBuilder,
 		const FIntPoint SampleBufferSize = DownsampledBufferSize * NumSamplesPerPixel2d;
 		const FIntPoint DonwnsampledSampleBufferSize = DownsampledBufferSize * NumSamplesPerPixel2d;
 
-		// #ml_todo: make atlas size based on the resolution
-		const FIntPoint ShadingTileGridSize = FIntPoint::DivideAndRoundUp(SceneTextures.Config.Extent, ManyLights::TileSize);
-		const FIntPoint ShadingTileAtlasSize = ManyLights::TileSize * ManyLights::ShadingAtlasSizeInTiles;
-		const int32 MaxShadingTiles = ManyLights::ShadingAtlasSizeInTiles * ManyLights::ShadingAtlasSizeInTiles;
-		check(MaxShadingTiles == (ShadingTileAtlasSize.X * ShadingTileAtlasSize.Y) / (ManyLights::TileSize * ManyLights::TileSize));
-		check(MaxShadingTiles < ManyLights::ShadingTileIndexUnshadowed);
-		const int32 MaxCompositeTiles = MaxShadingTiles;
-
-		const int32 MaxShadingTilesPerGridCell = FMath::Max(CVarManyLightsMaxShadingTilesPerGridCell.GetValueOnRenderThread(), 0);
-
 		FRDGTextureRef DownsampledSceneDepth = GraphBuilder.CreateTexture(
 			FRDGTextureDesc::Create2D(DownsampledBufferSize, PF_R32_FLOAT, FClearValueBinding::Black, TexCreate_ShaderResource | TexCreate_UAV),
 			TEXT("ManyLights.DownsampledSceneDepth"));
@@ -865,10 +789,6 @@ void FDeferredShadingSceneRenderer::RenderManyLights(FRDGBuilder& GraphBuilder,
 			ManyLightsParameters.DownsampledTileMask = DownsampledTileMask;
 			ManyLightsParameters.DownsampledSceneDepth = DownsampledSceneDepth;
 			ManyLightsParameters.DownsampledSceneWorldNormal = DownsampledSceneWorldNormal;
-			ManyLightsParameters.MaxCompositeTiles = MaxCompositeTiles;
-			ManyLightsParameters.MaxShadingTiles = MaxShadingTiles;
-			ManyLightsParameters.MaxShadingTilesPerGridCell = MaxShadingTilesPerGridCell;
-			ManyLightsParameters.ShadingTileGridSize = ShadingTileGridSize;
 			ManyLightsParameters.DownsampledBufferInvSize = FVector2f(1.0f) / DownsampledBufferSize;
 			ManyLightsParameters.SamplingMinWeight = FMath::Max(CVarManyLightsSamplingMinWeight.GetValueOnRenderThread(), 0.0f);
 			ManyLightsParameters.TileDataStride = TileDataStride;
@@ -978,10 +898,6 @@ void FDeferredShadingSceneRenderer::RenderManyLights(FRDGBuilder& GraphBuilder,
 				FIntVector(1, 1, 1));
 		}
 
-		FRDGBufferRef CompositeTileAllocator = GraphBuilder.CreateBuffer(FRDGBufferDesc::CreateStructuredDesc(sizeof(uint32), 1), TEXT("ManyLights.CompositeTileAllocator"));
-		FRDGBufferRef CompositeTileData = GraphBuilder.CreateBuffer(FRDGBufferDesc::CreateStructuredDesc(2 * sizeof(uint32), MaxCompositeTiles), TEXT("ManyLights.CompositeTileData"));
-		AddClearUAVPass(GraphBuilder, GraphBuilder.CreateUAV(CompositeTileAllocator), 0);
-
 		FRDGTextureRef SampleLuminanceSum = GraphBuilder.CreateTexture(
 			FRDGTextureDesc::Create2D(DownsampledBufferSize, PF_G16R16F, FClearValueBinding::Black, TexCreate_ShaderResource | TexCreate_UAV),
 			TEXT("ManyLights.SampleLuminanceSum"));
@@ -991,8 +907,6 @@ void FDeferredShadingSceneRenderer::RenderManyLights(FRDGBuilder& GraphBuilder,
 			FRDGTextureUAVRef SampleLuminanceSumUAV = GraphBuilder.CreateUAV(SampleLuminanceSum, ERDGUnorderedAccessViewFlags::SkipBarrier);
 			FRDGTextureUAVRef DownsampledSceneDepthUAV = GraphBuilder.CreateUAV(DownsampledSceneDepth, ERDGUnorderedAccessViewFlags::SkipBarrier);
 			FRDGTextureUAVRef DownsampledSceneWorldNormalUAV = GraphBuilder.CreateUAV(DownsampledSceneWorldNormal, ERDGUnorderedAccessViewFlags::SkipBarrier);
-			FRDGBufferUAVRef CompositeTileAllocatorUAV = GraphBuilder.CreateUAV(CompositeTileAllocator, ERDGUnorderedAccessViewFlags::SkipBarrier);
-			FRDGBufferUAVRef CompositeTileDataUAV = GraphBuilder.CreateUAV(CompositeTileData, ERDGUnorderedAccessViewFlags::SkipBarrier);
 			FRDGTextureUAVRef LightSamplesUAV = GraphBuilder.CreateUAV(LightSamples, ERDGUnorderedAccessViewFlags::SkipBarrier);
 
 			// Clear tiles which don't contain any lights or geometry
@@ -1028,8 +942,6 @@ void FDeferredShadingSceneRenderer::RenderManyLights(FRDGBuilder& GraphBuilder,
 				PassParameters->RWDownsampledSceneDepth = DownsampledSceneDepthUAV;
 				PassParameters->RWDownsampledSceneWorldNormal = DownsampledSceneWorldNormalUAV;
 				PassParameters->RWLightSamples = LightSamplesUAV;
-				PassParameters->RWCompositeTileAllocator = CompositeTileAllocatorUAV;
-				PassParameters->RWCompositeTileData = CompositeTileDataUAV;
 				PassParameters->DownsampledTileAllocator = GraphBuilder.CreateSRV(DownsampledTileAllocator);
 				PassParameters->DownsampledTileData = GraphBuilder.CreateSRV(DownsampledTileData);
 
@@ -1052,25 +964,6 @@ void FDeferredShadingSceneRenderer::RenderManyLights(FRDGBuilder& GraphBuilder,
 			}
 		}
 
-		FRDGBufferRef CompositeTileIndirectArgs = GraphBuilder.CreateBuffer(FRDGBufferDesc::CreateIndirectDesc<FRHIDispatchIndirectParameters>(1), TEXT("ManyLights.CompositeTileIndirectArgs"));
-
-		// Setup indirect args for shadow mask tile updates
-		{
-			FInitCompositeIndirectArgsCS::FParameters* PassParameters = GraphBuilder.AllocParameters<FInitCompositeIndirectArgsCS::FParameters>();
-			PassParameters->ManyLightsParameters = ManyLightsParameters;
-			PassParameters->RWIndirectArgs = GraphBuilder.CreateUAV(CompositeTileIndirectArgs);
-			PassParameters->CompositeTileAllocator = GraphBuilder.CreateSRV(CompositeTileAllocator);
-
-			auto ComputeShader = View.ShaderMap->GetShader<FInitCompositeIndirectArgsCS>();
-
-			FComputeShaderUtils::AddPass(
-				GraphBuilder,
-				RDG_EVENT_NAME("InitCompositeIndirectArgs"),
-				ComputeShader,
-				PassParameters,
-				FIntVector(1, 1, 1));
-		}
-
 		ManyLights::RayTraceLightSamples(
 			ViewFamily,
 			View,
@@ -1104,51 +997,6 @@ void FDeferredShadingSceneRenderer::RenderManyLights(FRDGBuilder& GraphBuilder,
 				GroupCount);
 		}
 
-		FRDGBufferRef ShadingTileAllocator = GraphBuilder.CreateBuffer(
-			FRDGBufferDesc::CreateStructuredDesc(sizeof(uint32), 1),
-			TEXT("ManyLightsParameters.ShadingTileAllocator"));
-
-		FRDGTextureRef ShadingTileGridAllocator = GraphBuilder.CreateTexture(
-			FRDGTextureDesc::Create2D(ShadingTileGridSize, PF_R32_UINT, FClearValueBinding::Black, TexCreate_ShaderResource | TexCreate_UAV),
-			TEXT("ManyLights.ShadingTileGridAllocator"));
-
-		FRDGBufferRef ShadingTileGrid = GraphBuilder.CreateBuffer(
-			FRDGBufferDesc::CreateStructuredDesc(sizeof(uint32), ShadingTileGridSize.X * ShadingTileGridSize.Y * MaxShadingTilesPerGridCell),
-			TEXT("ManyLightsParameters.ShadingTileGrid"));
-
-		FRDGTextureRef ShadingTileAtlas = GraphBuilder.CreateTexture(
-			FRDGTextureDesc::Create2D(ShadingTileAtlasSize, PF_R16F, FClearValueBinding::Black, TexCreate_ShaderResource | TexCreate_UAV),
-			TEXT("ManyLights	.ShadingTileAtlas"));
-
-		AddClearUAVPass(GraphBuilder, GraphBuilder.CreateUAV(ShadingTileAllocator), 0u);
-		AddClearUAVPass(GraphBuilder, GraphBuilder.CreateUAV(ShadingTileGridAllocator), 0u);
-
-		// Composite shadow masks traces
-		{
-			FResolveLightSamplesCS::FParameters* PassParameters = GraphBuilder.AllocParameters<FResolveLightSamplesCS::FParameters>();
-			PassParameters->ManyLightsParameters = ManyLightsParameters;
-			PassParameters->IndirectArgs = CompositeTileIndirectArgs;
-			PassParameters->RWShadingTileAllocator = GraphBuilder.CreateUAV(ShadingTileAllocator);
-			PassParameters->RWShadingTileGridAllocator = GraphBuilder.CreateUAV(ShadingTileGridAllocator);
-			PassParameters->RWShadingTileGrid = GraphBuilder.CreateUAV(ShadingTileGrid);
-			PassParameters->RWShadingTileAtlas = GraphBuilder.CreateUAV(ShadingTileAtlas);
-			PassParameters->CompositeTileData = GraphBuilder.CreateSRV(CompositeTileData);
-			PassParameters->CompositeUpsampleWeights = CompositeUpsampleWeights;
-			PassParameters->LightSamples = LightSamples;
-
-			FResolveLightSamplesCS::FPermutationDomain PermutationVector;
-			PermutationVector.Set<FResolveLightSamplesCS::FDebugMode>(bDebug);
-			auto ComputeShader = View.ShaderMap->GetShader<FResolveLightSamplesCS>(PermutationVector);
-
-			FComputeShaderUtils::AddPass(
-				GraphBuilder,
-				RDG_EVENT_NAME("CompositeLightSamples"),
-				ComputeShader,
-				PassParameters,
-				CompositeTileIndirectArgs,
-				0);
-		}
-
 		FRDGTextureRef ResolvedDiffuseLighting = GraphBuilder.CreateTexture(
 			FRDGTextureDesc::Create2D(View.GetSceneTexturesConfig().Extent, PF_FloatRGB, FClearValueBinding::Black, TexCreate_ShaderResource | TexCreate_UAV),
 			TEXT("ManyLights.ResolvedDiffuseLighting"));
@@ -1190,19 +1038,17 @@ void FDeferredShadingSceneRenderer::RenderManyLights(FRDGBuilder& GraphBuilder,
 				PassParameters->RWResolvedSpecularLighting = ResolvedSpecularLightingUAV;
 				PassParameters->IndirectArgs = TileIndirectArgs;
 				PassParameters->ManyLightsParameters = ManyLightsParameters;
-				PassParameters->CompositeTileAllocator = GraphBuilder.CreateSRV(CompositeTileAllocator);
 				PassParameters->TileAllocator = GraphBuilder.CreateSRV(TileAllocator);
 				PassParameters->TileData = GraphBuilder.CreateSRV(TileData);
-				PassParameters->ShadingTileAllocator = GraphBuilder.CreateSRV(ShadingTileAllocator);
-				PassParameters->ShadingTileGridAllocator = ShadingTileGridAllocator;
-				PassParameters->ShadingTileGrid = GraphBuilder.CreateSRV(ShadingTileGrid);
-				PassParameters->ShadingTileAtlas = ShadingTileAtlas;
+				PassParameters->CompositeUpsampleWeights = CompositeUpsampleWeights;
+				PassParameters->LightSamples = LightSamples;
 
 				FShadeLightSamplesCS::FPermutationDomain PermutationVector;
 				PermutationVector.Set<FShadeLightSamplesCS::FTileType>(TileType);
 				PermutationVector.Set<FShadeLightSamplesCS::FIESProfile>(CVarManyLightsIESProfiles.GetValueOnRenderThread() != 0);
 				PermutationVector.Set<FShadeLightSamplesCS::FLightFunctionAtlas>(bUseLightFunctionAtlas);
 				PermutationVector.Set<FShadeLightSamplesCS::FTexturedRectLights>(CVarManyLightsTexturedRectLights.GetValueOnRenderThread() != 0);
+				PermutationVector.Set<FShadeLightSamplesCS::FNumSamplesPerPixel1d>(NumSamplesPerPixel2d.X * NumSamplesPerPixel2d.Y);
 				PermutationVector.Set<FShadeLightSamplesCS::FDebugMode>(bDebug);
 				auto ComputeShader = View.ShaderMap->GetShader<FShadeLightSamplesCS>(PermutationVector);
 
diff --git a/Engine/Source/Runtime/Renderer/Private/ManyLights/ManyLightsInternal.h b/Engine/Source/Runtime/Renderer/Private/ManyLights/ManyLightsInternal.h
--- a/Engine/Source/Runtime/Renderer/Private/ManyLights/ManyLightsInternal.h
+++ b/Engine/Source/Runtime/Renderer/Private/ManyLights/ManyLightsInternal.h
@@ -25,10 +25,6 @@ BEGIN_SHADER_PARAMETER_STRUCT(FManyLightsParameters, )
 	SHADER_PARAMETER(FIntPoint, NumSamplesPerPixel)
 	SHADER_PARAMETER(FIntPoint, NumSamplesPerPixelDivideShift)
 	SHADER_PARAMETER(uint32, ManyLightsStateFrameIndex)
-	SHADER_PARAMETER(uint32, MaxCompositeTiles)
-	SHADER_PARAMETER(uint32, MaxShadingTiles)
-	SHADER_PARAMETER(uint32, MaxShadingTilesPerGridCell)
-	SHADER_PARAMETER(FIntPoint, ShadingTileGridSize)
 	SHADER_PARAMETER(FVector2f, DownsampledBufferInvSize)
 	SHADER_PARAMETER(float, SamplingMinWeight)
 	SHADER_PARAMETER(int32, TileDataStride)
