Custom Render Pass:  Fix for various bugs caused by RenderBasePass and calls to "FPrimitiveSceneProxy::GetDynamicMeshElements" not using correct EngineShowFlags from the CRP's view family.  The most severe example bugs involved static meshes disappearing completely from CRPs when the Lighting flag was turned off in the capture, or when various debug visualization features were enabled in the main view.

The first issue is solved by making RenderBasePass static, so it doesn't have access to the main ViewFamily member of the scene renderer, and instead pulls the view family from the first view passed in.  The second issue is solved by dividing calls to the "GetDynamicMeshElements" virtual function by view family, so the correct ViewFamily is sent to each call.  The view visibility bits are masked for the subset of views with a given family.  This approach avoids the need to change the API and refactor many implementations to be aware of multiple view families.  A handful of GetDynamicMeshElements implementations include logic that can't be run more than once.  To handle that case, the flag bSinglePassGDME was added to fall back to the original single call behavior, with the assumption that the implementation must properly handle variations in the Family pointer between Views.  A pass was made over all 100+ GetDynamicMeshElements implementations, with these edge cases found and handled:

* FMaterialSpriteSceneProxy and FEditorWidgetBillboardProxy each share a single set of vertex buffers across all views.
* FLandscapeComponentSceneProxy includes an LOD frame dump feature, which should only run once.
* FModelSceneProxy has a bAnySelectedSurfs flag accumulated across all views.
* UE::XRCreative::Private::FRenderComponentSceneProxy::GetDynamicMeshElements consumes cached geometry generated over the previous frame.
* FWaterMeshSceneProxy::GetDynamicMeshElements uses all Views together to compute visible water tiles.  Running in a single call requires wireframe status to be determined per view.

No edge cases were found in the GetDynamicMeshElements implementations in a large in-house project using the engine, and behavior only changes if a new feature (Custom Render Passes) is in use, so hopefully it should be rare for other third party projects to encounter issues.

As a separate optimization, force the "Lighting" flag to always be set for Custom Render Passes, to avoid causing the "IsRichView" function to inadvertently push static meshes into the much slower dynamic code path.  The "Lighting" flag has no effect on Custom Render Passes, as they don't run downstream visualization logic, and so is a needless performance cost which is silently introduced if a user touches the flag.  Another minor bug noticed and fixed was terrain editing widgets displaying in captures.

diff --git a/Engine/Plugins/Experimental/Water/Source/Runtime/Private/WaterMeshSceneProxy.cpp b/Engine/Plugins/Experimental/Water/Source/Runtime/Private/WaterMeshSceneProxy.cpp
--- a/Engine/Plugins/Experimental/Water/Source/Runtime/Private/WaterMeshSceneProxy.cpp
+++ b/Engine/Plugins/Experimental/Water/Source/Runtime/Private/WaterMeshSceneProxy.cpp
@@ -168,6 +168,9 @@ FWaterMeshSceneProxy::FWaterMeshSceneProxy(UWaterMeshComponent* Component)
 	AWaterZone* WaterZone = CastChecked<AWaterZone>(Component->GetOwner());
 	checkf(WaterZone != nullptr, TEXT("WaterMeshComponent is owned by an actor that is not a WaterZone. This is not supported!"));
 
+	// Need to enable single pass GetDynamicMeshElements, as the function uses all Views together to compute visible water tiles
+	bSinglePassGDME = true;
+
 	if (Component->ShouldBuildQuadTreeInSceneProxy())
 	{
 		// Get the QuadTreeBuilder and build the quadtree
@@ -410,11 +413,30 @@ void FWaterMeshSceneProxy::GetDynamicMeshElements(const TArray<const FSceneView*
 		return;
 	}
 
-	// Set up wireframe material (if needed)
-	const bool bWireframe = AllowDebugViewmodes() && (ViewFamily.EngineShowFlags.Wireframe || CVarWaterMeshShowWireframe.GetValueOnRenderThread() == 1);
+	// Figure out which views have wireframe enabled
+	uint32 WireframeVisibilityMap = 0;
+	if (AllowDebugViewmodes())
+	{
+		if (CVarWaterMeshShowWireframe.GetValueOnRenderThread() == 1)
+		{
+			// Force all views to wireframe
+			WireframeVisibilityMap = VisibilityMap;
+		}
+		else
+		{
+			for (int32 ViewIndex = 0; ViewIndex < Views.Num(); ViewIndex++)
+			{
+				if ((VisibilityMap & (1 << ViewIndex)) && Views[ViewIndex]->Family->EngineShowFlags.Wireframe)
+				{
+					WireframeVisibilityMap |= (1 << ViewIndex);
+				}
+			}
+		}
+	}
 
+	// Set up wireframe material (if needed)
 	FColoredMaterialRenderProxy* WireframeMaterialInstance = nullptr;
-	if (bWireframe && CVarWaterMeshShowWireframeAtBaseHeight.GetValueOnRenderThread() == 1)
+	if (WireframeVisibilityMap && CVarWaterMeshShowWireframeAtBaseHeight.GetValueOnRenderThread() == 1)
 	{
 		WireframeMaterialInstance = new FColoredMaterialRenderProxy(
 			GEngine->WireframeMaterial ? GEngine->WireframeMaterial->GetRenderProxy() : NULL,
@@ -536,6 +558,7 @@ void FWaterMeshSceneProxy::GetDynamicMeshElements(const TArray<const FSceneView*
 				TRACE_CPUPROFILER_EVENT_SCOPE(BucketsPerView);
 
 				const bool bInstancedStereo = Views[ViewIndex]->bIsInstancedStereoEnabled;
+				const bool bViewIsWireframe = WireframeVisibilityMap & (1 << ViewIndex);
 
 				for (int32 MaterialIndex = 0; MaterialIndex < NumWaterMaterials; ++MaterialIndex)
 				{
@@ -548,7 +571,7 @@ void FWaterMeshSceneProxy::GetDynamicMeshElements(const TArray<const FSceneView*
 
 						TRACE_CPUPROFILER_EVENT_SCOPE(DensityBucket);
 
-						const FMaterialRenderProxy* MaterialRenderProxy = (WireframeMaterialInstance != nullptr) ? WireframeMaterialInstance : WaterQuadTree.GetWaterMaterials()[MaterialIndex];
+						const FMaterialRenderProxy* MaterialRenderProxy = bViewIsWireframe && WireframeMaterialInstance ? WireframeMaterialInstance : WaterQuadTree.GetWaterMaterials()[MaterialIndex];
 						check(MaterialRenderProxy != nullptr);
 
 						bool bUseForDepthPass = false;
@@ -565,7 +588,7 @@ void FWaterMeshSceneProxy::GetDynamicMeshElements(const TArray<const FSceneView*
 						{
 							// Set up mesh batch
 							FMeshBatch& Mesh = Collector.AllocateMesh();
-							Mesh.bWireframe = bWireframe;
+							Mesh.bWireframe = bViewIsWireframe;
 							Mesh.VertexFactory = bInstancedStereo ? (FVertexFactory*)WaterVertexFactoryIndirectDrawISR : (FVertexFactory*)WaterVertexFactoryIndirectDraw;
 							Mesh.MaterialRenderProxy = MaterialRenderProxy;
 							Mesh.ReverseCulling = IsLocalToWorldDeterminantNegative();
@@ -731,6 +754,8 @@ void FWaterMeshSceneProxy::GetDynamicMeshElements(const TArray<const FSceneView*
 		{
 			TRACE_CPUPROFILER_EVENT_SCOPE(BucketsPerView);
 
+			const bool bViewIsWireframe = WireframeVisibilityMap & (1 << ViewIndex);
+
 			FWaterQuadTree::FTraversalOutput& WaterInstanceData = WaterInstanceDataPerView[TraversalIndex];
 			TraversalIndex++;
 
@@ -751,7 +776,7 @@ void FWaterMeshSceneProxy::GetDynamicMeshElements(const TArray<const FSceneView*
 
 					TRACE_CPUPROFILER_EVENT_SCOPE(DensityBucket);
 
-					const FMaterialRenderProxy* MaterialRenderProxy = (WireframeMaterialInstance != nullptr) ? WireframeMaterialInstance : WaterQuadTree.GetWaterMaterials()[MaterialIndex];
+					const FMaterialRenderProxy* MaterialRenderProxy = bViewIsWireframe && WireframeMaterialInstance ? WireframeMaterialInstance : WaterQuadTree.GetWaterMaterials()[MaterialIndex];
 					check(MaterialRenderProxy != nullptr);
 
 					bool bUseForDepthPass = false;
@@ -768,7 +793,7 @@ void FWaterMeshSceneProxy::GetDynamicMeshElements(const TArray<const FSceneView*
 					{
 						// Set up mesh batch
 						FMeshBatch& Mesh = Collector.AllocateMesh();
-						Mesh.bWireframe = bWireframe;
+						Mesh.bWireframe = bViewIsWireframe;
 						Mesh.VertexFactory = WaterVertexFactories[DensityIndex];
 						Mesh.MaterialRenderProxy = MaterialRenderProxy;
 						Mesh.ReverseCulling = IsLocalToWorldDeterminantNegative();
diff --git a/Engine/Plugins/Experimental/XRCreativeFramework/Source/XRCreative/Private/ITF/XRCreativeITFRenderComponent.cpp b/Engine/Plugins/Experimental/XRCreativeFramework/Source/XRCreative/Private/ITF/XRCreativeITFRenderComponent.cpp
--- a/Engine/Plugins/Experimental/XRCreativeFramework/Source/XRCreative/Private/ITF/XRCreativeITFRenderComponent.cpp
+++ b/Engine/Plugins/Experimental/XRCreativeFramework/Source/XRCreative/Private/ITF/XRCreativeITFRenderComponent.cpp
@@ -24,6 +24,10 @@ namespace UE::XRCreative::Private
 			: FPrimitiveSceneProxy(InComponent)
 		{
 			GetGeometryQueryFunc = MoveTemp(GeometryQueryFunc);
+
+			// GetDynamicMeshElements needs to run single pass, because the function GetGeometryQueryFunc consumes cached geometry
+			// generated over the previous frame, and the geometry won't exist if called a second time.
+			bSinglePassGDME = true;
 		}
 
 		virtual void GetDynamicMeshElements(const TArray<const FSceneView*>& Views, const FSceneViewFamily& ViewFamily, uint32 VisibilityMap, FMeshElementCollector& Collector) const override
diff --git a/Engine/Plugins/Runtime/nDisplay/Source/DisplayCluster/Private/Game/EngineClasses/Scene/DisplayClusterWidgetComponent.cpp b/Engine/Plugins/Runtime/nDisplay/Source/DisplayCluster/Private/Game/EngineClasses/Scene/DisplayClusterWidgetComponent.cpp
--- a/Engine/Plugins/Runtime/nDisplay/Source/DisplayCluster/Private/Game/EngineClasses/Scene/DisplayClusterWidgetComponent.cpp
+++ b/Engine/Plugins/Runtime/nDisplay/Source/DisplayCluster/Private/Game/EngineClasses/Scene/DisplayClusterWidgetComponent.cpp
@@ -36,6 +36,9 @@ public:
 	{
 		bWillEverBeLit = false;
 
+		// GetDynamicMeshElements needs to run single pass, because a single set of StaticMeshVertexBuffers are shared across all views.
+		bSinglePassGDME = true;
+
 		StaticMeshVertexBuffers.PositionVertexBuffer.Init(1);
 		StaticMeshVertexBuffers.StaticMeshVertexBuffer.Init(1, 1);
 		StaticMeshVertexBuffers.ColorVertexBuffer.Init(1);
diff --git a/Engine/Source/Runtime/Engine/Private/Components/MaterialBillboardComponent.cpp b/Engine/Source/Runtime/Engine/Private/Components/MaterialBillboardComponent.cpp
--- a/Engine/Source/Runtime/Engine/Private/Components/MaterialBillboardComponent.cpp
+++ b/Engine/Source/Runtime/Engine/Private/Components/MaterialBillboardComponent.cpp
@@ -71,6 +71,9 @@ public:
 			}
 		}
 
+		// GetDynamicMeshElements needs to run single pass, because a single set of StaticMeshVertexBuffers are shared across all views.
+		bSinglePassGDME = true;
+
 		StaticMeshVertexBuffers.PositionVertexBuffer.Init(1);
 		StaticMeshVertexBuffers.StaticMeshVertexBuffer.Init(1, 1);
 		StaticMeshVertexBuffers.ColorVertexBuffer.Init(1);
diff --git a/Engine/Source/Runtime/Engine/Private/ModelRender.cpp b/Engine/Source/Runtime/Engine/Private/ModelRender.cpp
--- a/Engine/Source/Runtime/Engine/Private/ModelRender.cpp
+++ b/Engine/Source/Runtime/Engine/Private/ModelRender.cpp
@@ -225,6 +225,9 @@ public:
 			FColor(157,149,223,255))
 #endif
 	{
+		// GetDynamicMeshElements needs to run single pass, due to handling of the bAnySelectedSurfs flag across multiple views
+		bSinglePassGDME = true;
+
 		ENQUEUE_RENDER_COMMAND(InitOrUpdateVertexBufferCmd)([VertexBuffer = &InComponent->GetModel()->VertexBuffer](FRHICommandList& RHICmdList)
 		{
 			if (!VertexBuffer->Buffers.PositionVertexBuffer.IsInitialized())
@@ -303,10 +306,13 @@ public:
 			{
 				const FSceneView* View = Views[ViewIndex];
 
-				bool bShowSelection = GIsEditor && !View->bIsGameView && ViewFamily.EngineShowFlags.Selection;
-				bool bDynamicBSPTriangles = bShowSelection || IsRichView(ViewFamily);
-				bool bShowBSPTriangles = ViewFamily.EngineShowFlags.BSPTriangles;
-				bool bShowBSP = ViewFamily.EngineShowFlags.BSP;
+				// View families can vary per view, with bSinglePassGDME == true being set for this class, so get the specific family for the view
+				const FSceneViewFamily& ViewSpecificFamily = *View->Family;
+
+				bool bShowSelection = GIsEditor && !View->bIsGameView && ViewSpecificFamily.EngineShowFlags.Selection;
+				bool bDynamicBSPTriangles = bShowSelection || IsRichView(ViewSpecificFamily);
+				bool bShowBSPTriangles = ViewSpecificFamily.EngineShowFlags.BSPTriangles;
+				bool bShowBSP = ViewSpecificFamily.EngineShowFlags.BSP;
 
 #if WITH_EDITOR
 				bool bDrawCollision = false;
diff --git a/Engine/Source/Runtime/Engine/Private/PrimitiveSceneProxy.cpp b/Engine/Source/Runtime/Engine/Private/PrimitiveSceneProxy.cpp
--- a/Engine/Source/Runtime/Engine/Private/PrimitiveSceneProxy.cpp
+++ b/Engine/Source/Runtime/Engine/Private/PrimitiveSceneProxy.cpp
@@ -434,6 +434,7 @@ FPrimitiveSceneProxy::FPrimitiveSceneProxy(const FPrimitiveSceneProxyDesc& InPro
 ,	bCollisionEnabled(InProxyDesc.IsCollisionEnabled())
 ,	bTreatAsBackgroundForOcclusion(InProxyDesc.bTreatAsBackgroundForOcclusion)
 ,	bSupportsParallelGDME(true)
+,	bSinglePassGDME(false)
 ,	bVisibleInLumenScene(false)
 ,	bOpaqueOrMasked(true)
 ,	bCanSkipRedundantTransformUpdates(true)
diff --git a/Engine/Source/Runtime/Engine/Public/PrimitiveSceneProxy.h b/Engine/Source/Runtime/Engine/Public/PrimitiveSceneProxy.h
--- a/Engine/Source/Runtime/Engine/Public/PrimitiveSceneProxy.h
+++ b/Engine/Source/Runtime/Engine/Public/PrimitiveSceneProxy.h
@@ -748,6 +748,7 @@ public:
 	inline bool ShouldRenderInMainPass() const { return bRenderInMainPass; }
 	inline bool ShouldRenderInDepthPass() const { return bRenderInMainPass || bRenderInDepthPass; }
 	inline bool SupportsParallelGDME() const { return bSupportsParallelGDME; }
+	inline bool SinglePassGDME() const { return bSinglePassGDME; }
 	inline bool IsCollisionEnabled() const { return bCollisionEnabled; }
 	inline bool IsHovered() const { return bHovered; }
 	inline bool IsOwnedBy(const AActor* Actor) const { return Owners.Find(Actor) != INDEX_NONE; }
@@ -1270,6 +1271,15 @@ protected:
 	/** Whether the proxy supports asynchronously calling GetDynamicMeshElements. If disabled, all calls for various proxies are serialized with respect to each other. */
 	uint8 bSupportsParallelGDME : 1;
 
+	/**
+	 * Whether to call GetDynamicMeshElements a single time, instead of once per unique FSceneViewFamily.  If set to true, the GetDynamicMeshElements
+	 * implementation must take into account that Views may point to different view families, with different EngineShowFlags.  In practice, this means
+	 * using View->Family in the loop over Views, especially for accessing EngineShowFlags, rather than the ViewFamily parameter passed in.  The
+	 * "AllViews" member, plus certain other members like the frame counter and time, are invariant across view families, and safe to access from
+	 * ViewFamily.  Example use cases are features that do global processing or initialize a common shared buffer across views, which are fairly rare.
+	 */
+	uint8 bSinglePassGDME : 1;
+
 	/** Whether this component should be tracked by Lumen Scene. Turning this off will remove it from Lumen Scene and Lumen won't generate surface cache for it. */
 	uint8 bVisibleInLumenScene : 1;
 
diff --git a/Engine/Source/Runtime/Landscape/Private/LandscapeRender.cpp b/Engine/Source/Runtime/Landscape/Private/LandscapeRender.cpp
--- a/Engine/Source/Runtime/Landscape/Private/LandscapeRender.cpp
+++ b/Engine/Source/Runtime/Landscape/Private/LandscapeRender.cpp
@@ -1358,6 +1358,9 @@ FLandscapeComponentSceneProxy::FLandscapeComponentSceneProxy(ULandscapeComponent
 	// They do however have continuous LOD which is problematic, considered static as the LODs (are intended to) represent the same static surface.
 	bHasDeformableMesh = false;
 
+	// Enabled just so the GVarDumpLandscapeLODsCurrentFrame debug feature only runs once per render proxy in GetDynamicMeshElements
+	bSinglePassGDME = true;
+
 	VisibilityHelper.Init(InComponent, this);
 
 	if (!VisibilityHelper.ShouldBeVisible())
@@ -2584,7 +2587,6 @@ void FLandscapeComponentSceneProxy::GetDynamicMeshElements(const TArray<const FS
 	int32 NumPasses = 0;
 	int32 NumTriangles = 0;
 	int32 NumDrawCalls = 0;
-	const bool bIsWireframe = ViewFamily.EngineShowFlags.Wireframe;
 
 	const FLandscapeRenderSystem& RenderSystem = *LandscapeRenderSystems.FindChecked(LandscapeKey);
 
@@ -2597,6 +2599,11 @@ void FLandscapeComponentSceneProxy::GetDynamicMeshElements(const TArray<const FS
 
 			const FSceneView* View = Views[ViewIndex];
 
+			// With bSinglePassGDME == true, there can be a different view family per active view, so grab a reference to the specific view family
+			const FSceneViewFamily& ViewSpecificFamily = *View->Family;
+
+			const bool bIsWireframe = ViewSpecificFamily.EngineShowFlags.Wireframe;
+
 			int32 LODToRender = static_cast<int32>(RenderSystem.GetSectionLODValue(*View, RenderCoord));
 
 			FMeshBatch& Mesh = Collector.AllocateMesh();
@@ -2805,7 +2812,7 @@ void FLandscapeComponentSceneProxy::GetDynamicMeshElements(const TArray<const FS
 				else
 #endif
 					// Regular Landscape rendering. Only use the dynamic path if we're rendering a rich view or we've disabled the static path for debugging.
-					if (IsRichView(ViewFamily) ||
+					if (IsRichView(ViewSpecificFamily) ||
 						GLandscapeDebugOptions.bDisableStatic ||
 						bIsWireframe ||
 #if WITH_EDITOR
@@ -2833,7 +2840,7 @@ void FLandscapeComponentSceneProxy::GetDynamicMeshElements(const TArray<const FS
 
 #if WITH_EDITOR
 			  // Extra render passes for landscape tools
-			if (GLandscapeEditModeActive)
+			if (GLandscapeEditModeActive && ViewSpecificFamily.EngineShowFlags.Editor)
 			{
 				// Region selection
 				if (EditToolRenderData.SelectedType)
@@ -2923,9 +2930,9 @@ void FLandscapeComponentSceneProxy::GetDynamicMeshElements(const TArray<const FS
 				DrawWireBox(Collector.GetPDI(ViewIndex), GetBounds().GetBox(), FColor(255, 255, 0), SDPG_World);
 			}
 
-			if (ViewFamily.EngineShowFlags.Bounds)
+			if (ViewSpecificFamily.EngineShowFlags.Bounds)
 			{
-				RenderBounds(Collector.GetPDI(ViewIndex), ViewFamily.EngineShowFlags, GetBounds(), IsSelected());
+				RenderBounds(Collector.GetPDI(ViewIndex), ViewSpecificFamily.EngineShowFlags, GetBounds(), IsSelected());
 			}
 		}
 	}
diff --git a/Engine/Source/Runtime/Renderer/Private/AnisotropyRendering.cpp b/Engine/Source/Runtime/Renderer/Private/AnisotropyRendering.cpp
--- a/Engine/Source/Runtime/Renderer/Private/AnisotropyRendering.cpp
+++ b/Engine/Source/Runtime/Renderer/Private/AnisotropyRendering.cpp
@@ -332,8 +332,8 @@ void FDeferredShadingSceneRenderer::RenderAnisotropyPass(
 	FRDGBuilder& GraphBuilder, 
 	TArrayView<FViewInfo> InViews,
 	FSceneTextures& SceneTextures,
-	bool bDoParallelPass
-)
+	const FScene* Scene,
+	bool bDoParallelPass)
 {
 	RDG_CSV_STAT_EXCLUSIVE_SCOPE(GraphBuilder, RenderAnisotropyPass);
 	SCOPED_NAMED_EVENT(FDeferredShadingSceneRenderer_RenderAnisotropyPass, FColor::Emerald);
@@ -370,7 +370,7 @@ void FDeferredShadingSceneRenderer::RenderAnisotropyPass(
 					RDG_EVENT_NAME("AnisotropyPassParallel"),
 					PassParameters,
 					ERDGPassFlags::Raster | ERDGPassFlags::SkipRenderPass,
-					[this, &View, &ParallelMeshPass, PassParameters](const FRDGPass* InPass, FRHICommandListImmediate& RHICmdList)
+					[&View, &ParallelMeshPass, PassParameters](const FRDGPass* InPass, FRHICommandListImmediate& RHICmdList)
 				{
 					FRDGParallelCommandListSet ParallelCommandListSet(InPass, RHICmdList, View, FParallelCommandListBindings(PassParameters));
 
@@ -385,7 +385,7 @@ void FDeferredShadingSceneRenderer::RenderAnisotropyPass(
 					RDG_EVENT_NAME("AnisotropyPass"),
 					PassParameters,
 					ERDGPassFlags::Raster,
-					[this, &View, &ParallelMeshPass, PassParameters](FRHICommandList& RHICmdList)
+					[&View, &ParallelMeshPass, PassParameters](FRHICommandList& RHICmdList)
 				{
 					SetStereoViewport(RHICmdList, View);
 
diff --git a/Engine/Source/Runtime/Renderer/Private/BasePassRendering.cpp b/Engine/Source/Runtime/Renderer/Private/BasePassRendering.cpp
--- a/Engine/Source/Runtime/Renderer/Private/BasePassRendering.cpp
+++ b/Engine/Source/Runtime/Renderer/Private/BasePassRendering.cpp
@@ -962,6 +962,7 @@ void ModifyBasePassCSPSCompilationEnvironment(const FMeshMaterialShaderPermutati
 }
 
 void FDeferredShadingSceneRenderer::RenderBasePass(
+	FDeferredShadingSceneRenderer& Renderer,
 	FRDGBuilder& GraphBuilder,
 	TArrayView<FViewInfo> InViews,
 	FSceneTextures& SceneTextures,
@@ -975,6 +976,10 @@ void FDeferredShadingSceneRenderer::RenderBasePass(
 {
 	TRACE_CPUPROFILER_EVENT_SCOPE(FDeferredShadingSceneRenderer::RenderBasePass);
 
+	// This function has been deliberately made static, so it doesn't have access to the ViewFamily member, to support Custom Render Passes,
+	// which have a separate ViewFamily.  Create a local reference to the first view's Family here.
+	FViewFamilyInfo& ViewFamily = *(FViewFamilyInfo*)InViews[0].Family;
+
 	const bool bEnableParallelBasePasses = GRHICommandList.UseParallelAlgorithms() && CVarParallelBasePass.GetValueOnRenderThread();
 
 	static const auto ClearMethodCVar = IConsoleManager::Get().FindTConsoleVariableDataInt(TEXT("r.ClearSceneMethod"));
@@ -1037,7 +1042,7 @@ void FDeferredShadingSceneRenderer::RenderBasePass(
 
 	TStaticArray<FTextureRenderTargetBinding, MaxSimultaneousRenderTargets> BasePassTextures;
 	uint32 BasePassTextureCount = SceneTextures.GetGBufferRenderTargets(BasePassTextures);
-	Substrate::AppendSubstrateMRTs(*this, BasePassTextureCount, BasePassTextures);
+	Substrate::AppendSubstrateMRTs(Renderer, BasePassTextureCount, BasePassTextures);
 	TArrayView<FTextureRenderTargetBinding> BasePassTexturesView = MakeArrayView(BasePassTextures.GetData(), BasePassTextureCount);
 	FRDGTextureRef BasePassDepthTexture = SceneTextures.Depth.Target;
 	FLinearColor SceneColorClearValue = FLinearColor::Black;
@@ -1106,7 +1111,7 @@ void FDeferredShadingSceneRenderer::RenderBasePass(
 		{
 			// We only render this warning text when bRequiresRHIClear==true to make sure the scene color buffer is allocated at this stage.
 			// When false, the option specifies that all pixels must be written to by a sky dome anyway.
-			RenderSkyAtmosphereEditorNotifications(GraphBuilder, InViews, SceneTextures.Color.Target);
+			Renderer.RenderSkyAtmosphereEditorNotifications(GraphBuilder, InViews, SceneTextures.Color.Target);
 		}
 	}
 
@@ -1148,7 +1153,7 @@ void FDeferredShadingSceneRenderer::RenderBasePass(
 	}
 	ForwardBasePassTextures.bIs24BitUnormDepthStencil = ForwardBasePassTextures.SceneDepthIfResolved ? GPixelFormats[ForwardBasePassTextures.SceneDepthIfResolved->Desc.Format].bIs24BitUnormDepthStencil : 1;
 
-	RenderBasePassInternal(GraphBuilder, InViews, SceneTextures, BasePassRenderTargets, BasePassDepthStencilAccess, ForwardBasePassTextures, DBufferTextures, bDoParallelBasePass, bRenderLightmapDensity, InstanceCullingManager, bNaniteEnabled, NaniteBasePassShadingCommands, NaniteRasterResults);
+	RenderBasePassInternal(Renderer, GraphBuilder, InViews, SceneTextures, BasePassRenderTargets, BasePassDepthStencilAccess, ForwardBasePassTextures, DBufferTextures, bDoParallelBasePass, bRenderLightmapDensity, InstanceCullingManager, bNaniteEnabled, NaniteBasePassShadingCommands, NaniteRasterResults);
 	
 	for (const TSharedRef<ISceneViewExtension>& ViewExtension : ViewFamily.ViewExtensions)
 	{
@@ -1164,9 +1169,9 @@ void FDeferredShadingSceneRenderer::RenderBasePass(
 	}
 
 	// The anisotropy GBuffer is used for lighting, and not needed when running a scene capture with ERendererOutput::BasePass (BaseColor or Normal output), which doesn't run lighting.
-	if (ShouldRenderAnisotropyPass(InViews) && GetRendererOutput() == ERendererOutput::FinalSceneColor)
+	if (ShouldRenderAnisotropyPass(InViews) && Renderer.GetRendererOutput() == ERendererOutput::FinalSceneColor)
 	{
-		RenderAnisotropyPass(GraphBuilder, InViews, SceneTextures, bEnableParallelBasePasses);
+		RenderAnisotropyPass(GraphBuilder, InViews, SceneTextures, ViewFamily.Scene->GetRenderScene(), bEnableParallelBasePasses);
 	}
 
 #if !(UE_BUILD_SHIPPING)
@@ -1304,6 +1309,7 @@ static void RenderEditorPrimitives(
 }
 
 void FDeferredShadingSceneRenderer::RenderBasePassInternal(
+	FDeferredShadingSceneRenderer& Renderer,
 	FRDGBuilder& GraphBuilder,
 	TArrayView<FViewInfo> InViews,
 	const FSceneTextures& SceneTextures,
@@ -1321,6 +1327,11 @@ void FDeferredShadingSceneRenderer::RenderBasePassInternal(
 	RDG_CSV_STAT_EXCLUSIVE_SCOPE(GraphBuilder, RenderBasePass);
 	SCOPED_NAMED_EVENT(FDeferredShadingSceneRenderer_RenderBasePass, FColor::Emerald);
 
+	// This function has been deliberately made static, so it doesn't have access to the ViewFamily member, to support Custom Render Passes,
+	// which have a separate ViewFamily.  Create a local reference to the first view's Family here.
+	FViewFamilyInfo& ViewFamily = *(FViewFamilyInfo*)InViews[0].Family;
+	FScene* Scene = ViewFamily.Scene->GetRenderScene();
+
 #if WITH_DEBUG_VIEW_MODES
 	Nanite::EDebugViewMode NaniteDebugViewMode = Nanite::EDebugViewMode::None;
 	if (ViewFamily.EngineShowFlags.Wireframe)
@@ -1394,7 +1405,7 @@ void FDeferredShadingSceneRenderer::RenderBasePassInternal(
 			Nanite::DispatchBasePass(
 				GraphBuilder,
 				NaniteBasePassShadingCommands,
-				*this,
+				Renderer,
 				SceneTextures,
 				BasePassRenderTargets,
 				DBufferTextures,
@@ -1412,7 +1423,7 @@ void FDeferredShadingSceneRenderer::RenderBasePassInternal(
 		if (bNaniteEnabled)
 		{
 			// Should always have a full Z prepass with Nanite
-			check(ShouldRenderPrePass());
+			check(Renderer.ShouldRenderPrePass());
 
 			for (int32 ViewIndex = 0; ViewIndex < InViews.Num(); ++ViewIndex)
 			{
@@ -1451,7 +1462,7 @@ void FDeferredShadingSceneRenderer::RenderBasePassInternal(
 				RDG_EVENT_SCOPE_CONDITIONAL(GraphBuilder, InViews.Num() > 1, "View%d", ViewIndex);
 				View.BeginRenderView();
 
-				const bool bLumenGIEnabled = GetViewPipelineState(View).DiffuseIndirectMethod == EDiffuseIndirectMethod::Lumen;
+				const bool bLumenGIEnabled = Renderer.GetViewPipelineState(View).DiffuseIndirectMethod == EDiffuseIndirectMethod::Lumen;
 
 				FMeshPassProcessorRenderState DrawRenderState;
 				SetupBasePassState(BasePassDepthStencilAccess, ViewFamily.EngineShowFlags.ShaderComplexity, DrawRenderState);
@@ -1472,7 +1483,7 @@ void FDeferredShadingSceneRenderer::RenderBasePassInternal(
 						RDG_EVENT_NAME("BasePassParallel"),
 						PassParameters,
 						ERDGPassFlags::Raster | ERDGPassFlags::SkipRenderPass,
-						[this, &View, PassParameters](const FRDGPass* InPass, FRHICommandListImmediate& RHICmdList)
+						[&View, PassParameters](const FRDGPass* InPass, FRHICommandListImmediate& RHICmdList)
 					{
 						FRDGParallelCommandListSet ParallelCommandListSet(InPass, RHICmdList, View, FParallelCommandListBindings(PassParameters));
 						View.ParallelMeshDrawCommandPasses[EMeshPass::BasePass].DispatchDraw(&ParallelCommandListSet, RHICmdList, &PassParameters->InstanceCullingDrawParams);
@@ -1485,7 +1496,7 @@ void FDeferredShadingSceneRenderer::RenderBasePassInternal(
 				if (bShouldRenderViewForNanite)
 				{
 					// Should always have a full Z prepass with Nanite
-					check(ShouldRenderPrePass());
+					check(Renderer.ShouldRenderPrePass());
 
 					RenderNaniteBasePass(View, ViewIndex);
 				}
@@ -1506,7 +1517,7 @@ void FDeferredShadingSceneRenderer::RenderBasePassInternal(
 						RDG_EVENT_NAME("SkyPassParallel"),
 						SkyPassPassParameters,
 						ERDGPassFlags::Raster | ERDGPassFlags::SkipRenderPass,
-						[this, &View, SkyPassPassParameters](const FRDGPass* InPass, FRHICommandListImmediate& RHICmdList)
+						[&View, SkyPassPassParameters](const FRDGPass* InPass, FRHICommandListImmediate& RHICmdList)
 					{
 						FRDGParallelCommandListSet ParallelCommandListSet(InPass, RHICmdList, View, FParallelCommandListBindings(SkyPassPassParameters));
 						View.ParallelMeshDrawCommandPasses[EMeshPass::SkyPass].DispatchDraw(&ParallelCommandListSet, RHICmdList, &SkyPassPassParameters->InstanceCullingDrawParams);
@@ -1523,7 +1534,7 @@ void FDeferredShadingSceneRenderer::RenderBasePassInternal(
 				RDG_EVENT_SCOPE_CONDITIONAL(GraphBuilder, InViews.Num() > 1, "View%d", ViewIndex);
 				View.BeginRenderView();
 
-				const bool bLumenGIEnabled = GetViewPipelineState(View).DiffuseIndirectMethod == EDiffuseIndirectMethod::Lumen;
+				const bool bLumenGIEnabled = Renderer.GetViewPipelineState(View).DiffuseIndirectMethod == EDiffuseIndirectMethod::Lumen;
 
 				FMeshPassProcessorRenderState DrawRenderState;
 				SetupBasePassState(BasePassDepthStencilAccess, ViewFamily.EngineShowFlags.ShaderComplexity, DrawRenderState);
@@ -1544,7 +1555,7 @@ void FDeferredShadingSceneRenderer::RenderBasePassInternal(
 						RDG_EVENT_NAME("BasePass"),
 						PassParameters,
 						ERDGPassFlags::Raster,
-						[this, &View, PassParameters](FRHICommandList& RHICmdList)
+						[&View, PassParameters](FRHICommandList& RHICmdList)
 						{
 							SetStereoViewport(RHICmdList, View, 1.0f);
 							View.ParallelMeshDrawCommandPasses[EMeshPass::BasePass].DispatchDraw(nullptr, RHICmdList, &PassParameters->InstanceCullingDrawParams);
@@ -1556,7 +1567,7 @@ void FDeferredShadingSceneRenderer::RenderBasePassInternal(
 				if (bShouldRenderViewForNanite)
 				{
 					// Should always have a full Z prepass with Nanite
-					check(ShouldRenderPrePass());
+					check(Renderer.ShouldRenderPrePass());
 
 					RenderNaniteBasePass(View, ViewIndex);
 				}
@@ -1577,7 +1588,7 @@ void FDeferredShadingSceneRenderer::RenderBasePassInternal(
 						RDG_EVENT_NAME("SkyPass"),
 						SkyPassParameters,
 						ERDGPassFlags::Raster,
-						[this, &View, SkyPassParameters](FRHICommandList& RHICmdList)
+						[&View, SkyPassParameters](FRHICommandList& RHICmdList)
 						{
 							SetStereoViewport(RHICmdList, View, 1.0f);
 							View.ParallelMeshDrawCommandPasses[EMeshPass::SkyPass].DispatchDraw(nullptr, RHICmdList, &SkyPassParameters->InstanceCullingDrawParams);
diff --git a/Engine/Source/Runtime/Renderer/Private/DeferredShadingRenderer.cpp b/Engine/Source/Runtime/Renderer/Private/DeferredShadingRenderer.cpp
--- a/Engine/Source/Runtime/Renderer/Private/DeferredShadingRenderer.cpp
+++ b/Engine/Source/Runtime/Renderer/Private/DeferredShadingRenderer.cpp
@@ -2098,7 +2098,7 @@ void FDeferredShadingSceneRenderer::Render(FRDGBuilder& GraphBuilder)
 						Nanite::BuildShadingCommands(GraphBuilder, *Scene, ENaniteMeshPass::BasePass, NaniteBasePassShadingCommands, Nanite::EBuildShadingCommandsMode::Custom);
 					}
 
-					RenderBasePass(GraphBuilder, CustomRenderPassViews, SceneTextures, DBufferTextures, BasePassDepthStencilAccess, /*ForwardScreenSpaceShadowMaskTexture=*/nullptr, InstanceCullingManager, bNaniteEnabled, NaniteBasePassShadingCommands, NaniteRasterResults);
+					RenderBasePass(*this, GraphBuilder, CustomRenderPassViews, SceneTextures, DBufferTextures, BasePassDepthStencilAccess, /*ForwardScreenSpaceShadowMaskTexture=*/nullptr, InstanceCullingManager, bNaniteEnabled, NaniteBasePassShadingCommands, NaniteRasterResults);
 
 					if (ShouldRenderSingleLayerWater(CustomRenderPassViews))
 					{
@@ -2216,7 +2216,7 @@ void FDeferredShadingSceneRenderer::Render(FRDGBuilder& GraphBuilder)
 				ViewExtension->PreRenderBasePass_RenderThread(GraphBuilder, ShouldRenderPrePass() /*bDepthBufferIsPopulated*/);
 			}
 
-			RenderBasePass(GraphBuilder, Views, SceneTextures, DBufferTextures, BasePassDepthStencilAccess, /*ForwardScreenSpaceShadowMaskTexture=*/nullptr, InstanceCullingManager, bNaniteEnabled, Scene->NaniteShadingCommands[ENaniteMeshPass::BasePass], NaniteRasterResults);
+			RenderBasePass(*this, GraphBuilder, Views, SceneTextures, DBufferTextures, BasePassDepthStencilAccess, /*ForwardScreenSpaceShadowMaskTexture=*/nullptr, InstanceCullingManager, bNaniteEnabled, Scene->NaniteShadingCommands[ENaniteMeshPass::BasePass], NaniteRasterResults);
 
 			if (ShouldRenderSingleLayerWater(Views))
 			{
@@ -2468,7 +2468,7 @@ void FDeferredShadingSceneRenderer::Render(FRDGBuilder& GraphBuilder)
 		{
 			if (!bHasRayTracedOverlay)
 			{
-				RenderBasePass(GraphBuilder, Views, SceneTextures, DBufferTextures, BasePassDepthStencilAccess, ForwardScreenSpaceShadowMaskTexture, InstanceCullingManager, bNaniteEnabled, Scene->NaniteShadingCommands[ENaniteMeshPass::BasePass], NaniteRasterResults);
+				RenderBasePass(*this, GraphBuilder, Views, SceneTextures, DBufferTextures, BasePassDepthStencilAccess, ForwardScreenSpaceShadowMaskTexture, InstanceCullingManager, bNaniteEnabled, Scene->NaniteShadingCommands[ENaniteMeshPass::BasePass], NaniteRasterResults);
 				GraphBuilder.AddDispatchHint();
 			}
 
diff --git a/Engine/Source/Runtime/Renderer/Private/DeferredShadingRenderer.h b/Engine/Source/Runtime/Renderer/Private/DeferredShadingRenderer.h
--- a/Engine/Source/Runtime/Renderer/Private/DeferredShadingRenderer.h
+++ b/Engine/Source/Runtime/Renderer/Private/DeferredShadingRenderer.h
@@ -347,7 +347,13 @@ public:
 	 */
 	void DebugLightGrid(FRDGBuilder& GraphBuilder, FSceneTextures& SceneTextures, bool bNeedLightGrid);
 
-	void RenderBasePass(
+	/**
+	 * The following three functions are static, for compile time enforcement related to CustomRenderPass rendering, which uses these functions.
+	 * Custom Render Passes have a separate ViewFamily, and making these functions static prevents the ViewFamily member in the scene renderer
+	 * class from being inadvertently accessed.
+	 */
+	static void RenderBasePass(
+		FDeferredShadingSceneRenderer& Renderer,
 		FRDGBuilder& GraphBuilder,
 		TArrayView<FViewInfo> InViews,
 		FSceneTextures& SceneTextures,
@@ -359,7 +365,8 @@ public:
 		struct FNaniteShadingCommands& NaniteBasePassShadingCommands,
 		const TArrayView<Nanite::FRasterResults>& NaniteRasterResults);
 
-	void RenderBasePassInternal(
+	static void RenderBasePassInternal(
+		FDeferredShadingSceneRenderer& Renderer,
 		FRDGBuilder& GraphBuilder,
 		TArrayView<FViewInfo> InViews,
 		const FSceneTextures& SceneTextures,
@@ -374,10 +381,11 @@ public:
 		struct FNaniteShadingCommands& NaniteBasePassShadingCommands,
 		const TArrayView<Nanite::FRasterResults>& NaniteRasterResults);
 
-	void RenderAnisotropyPass(
+	static void RenderAnisotropyPass(
 		FRDGBuilder& GraphBuilder,
 		TArrayView<FViewInfo> InViews,
 		FSceneTextures& SceneTextures,
+		const FScene* Scene,
 		bool bDoParallelPass);
 	/**
 	 * Runs water pre-pass if enabled and returns an RDG-allocated object with intermediates, or null.
diff --git a/Engine/Source/Runtime/Renderer/Private/SceneRendering.cpp b/Engine/Source/Runtime/Renderer/Private/SceneRendering.cpp
--- a/Engine/Source/Runtime/Renderer/Private/SceneRendering.cpp
+++ b/Engine/Source/Runtime/Renderer/Private/SceneRendering.cpp
@@ -2947,6 +2947,12 @@ FSceneRenderer::FSceneRenderer(const FSceneViewFamily* InViewFamily, FHitProxyCo
 		// Disable sky rendering, which is gated by the Atmosphere flag (unnecessary perf cost during base pass rendering)
 		FamilyCVS.EngineShowFlags.Atmosphere = false;
 
+		// Force enable the Lighting flag.  Custom Render Passes don't run lighting regardless of this flag, but disabling it causes the function
+		// "IsRichView()" to return true, forcing certain render proxy classes to go through a slow dynamic rendering code path used for debug
+		// visualization, which isn't relevant for Custom Render Passes.  Avoids a massive silent perf loss if someone has modified the flag
+		// (which is user facing in the Scene Capture properties).
+		FamilyCVS.EngineShowFlags.Lighting = true;
+
 		FCustomRenderPassInfo& CustomRenderPassInfo = CustomRenderPassInfos.Emplace_GetRef(FamilyCVS, ViewFamily);
 		CustomRenderPassInfo.CustomRenderPass = CustomRenderPass;
 		CustomRenderPassInfo.ViewFamily.Time = ViewFamily.Time;
diff --git a/Engine/Source/Runtime/Renderer/Private/SceneVisibility.cpp b/Engine/Source/Runtime/Renderer/Private/SceneVisibility.cpp
--- a/Engine/Source/Runtime/Renderer/Private/SceneVisibility.cpp
+++ b/Engine/Source/Runtime/Renderer/Private/SceneVisibility.cpp
@@ -3827,7 +3827,7 @@ void FVisibilityViewPacket::BeginInitVisibility()
 ///////////////////////////////////////////////////////////////////////////////
 
 FDynamicMeshElementContext::FDynamicMeshElementContext(FSceneRenderer& SceneRenderer)
-	: ViewFamily(SceneRenderer.ViewFamily)
+	: FirstViewFamily(SceneRenderer.ViewFamily)
 	, Views(SceneRenderer.AllViews)
 	, Primitives(SceneRenderer.Scene->Primitives)
 	// Defer committing GPU scene and Material updates until the mesh collectors are finished. Deferring materials allows VT updates to
@@ -3841,6 +3841,23 @@ FDynamicMeshElementContext::FDynamicMeshElementContext(FSceneRenderer& SceneRend
 	, DynamicVertexBuffer(*RHICmdList)
 	, DynamicIndexBuffer(*RHICmdList)
 {
+	// With Custom Render Passes, it's possible to have Views that point to a different Family.  Divide Views into groups with the same family.
+	uint8 ViewSubsetMask = 0;
+	int32 LastViewIndex = Views.Num() - 1;
+	for (int32 ViewIndex = 0; ViewIndex <= LastViewIndex; ViewIndex++)
+	{
+		ViewSubsetMask |= 1u << ViewIndex;
+
+		// Check if the next element has a different ViewFamily, or this is the last view.  If so, emit the GetDynamicMeshElements call
+		// with visibility mask bits set for the subset of views with the same ViewFamily, and reset the mask for the next family.
+		if (ViewIndex == LastViewIndex || Views[ViewIndex]->Family != Views[ViewIndex + 1]->Family)
+		{
+			// Emit a batch with the given mask, and reset the mask
+			ViewFamilyGroups.Add({ Views[ViewIndex]->Family, ViewSubsetMask });
+			ViewSubsetMask = 0;
+		}
+	}
+
 	RHICmdList->SwitchPipeline(ERHIPipeline::Graphics);
 
 	ViewMeshArraysPerView.SetNum(Views.Num());
@@ -3948,7 +3965,22 @@ void FDynamicMeshElementContext::GatherDynamicMeshElementsForPrimitive(FPrimitiv
 	}
 
 	MeshCollector.SetPrimitive(Primitive->Proxy, Primitive->DefaultDynamicHitProxyId);
-	Primitive->Proxy->GetDynamicMeshElements(ViewFamily.AllViews, ViewFamily, ViewMask, MeshCollector);
+
+	// If Custom Render Passes aren't in use, there will be only one group, which is the common case.
+	if (ViewFamilyGroups.Num() == 1 || Primitive->Proxy->SinglePassGDME())
+	{
+		Primitive->Proxy->GetDynamicMeshElements(FirstViewFamily.AllViews, FirstViewFamily, ViewMask, MeshCollector);
+	}
+	else
+	{
+		for (FViewFamilyGroup& Group : ViewFamilyGroups)
+		{
+			if (uint8 MaskedViewMask = ViewMask & Group.ViewSubsetMask)
+			{
+				Primitive->Proxy->GetDynamicMeshElements(FirstViewFamily.AllViews, *Group.Family, MaskedViewMask, MeshCollector);
+			}
+		}
+	}
 
 	for (int32 ViewIndex = 0; ViewIndex < Views.Num(); ViewIndex++)
 	{
@@ -3969,7 +4001,22 @@ void FDynamicMeshElementContext::GatherDynamicMeshElementsForEditorPrimitive(FPr
 {
 #if WITH_EDITOR
 	EditorMeshCollector.SetPrimitive(Primitive->Proxy, Primitive->DefaultDynamicHitProxyId);
-	Primitive->Proxy->GetDynamicMeshElements(ViewFamily.AllViews, ViewFamily, ViewMask, EditorMeshCollector);
+
+	// If Custom Render Passes aren't in use, there will be only one group, which is the common case.
+	if (ViewFamilyGroups.Num() == 1 || Primitive->Proxy->SinglePassGDME())
+	{
+		Primitive->Proxy->GetDynamicMeshElements(FirstViewFamily.AllViews, FirstViewFamily, ViewMask, EditorMeshCollector);
+	}
+	else
+	{
+		for (FViewFamilyGroup& Group : ViewFamilyGroups)
+		{
+			if (uint8 MaskedViewMask = ViewMask & Group.ViewSubsetMask)
+			{
+				Primitive->Proxy->GetDynamicMeshElements(FirstViewFamily.AllViews, *Group.Family, MaskedViewMask, EditorMeshCollector);
+			}
+		}
+	}
 #endif
 }
 
diff --git a/Engine/Source/Runtime/Renderer/Private/SceneVisibilityPrivate.h b/Engine/Source/Runtime/Renderer/Private/SceneVisibilityPrivate.h
--- a/Engine/Source/Runtime/Renderer/Private/SceneVisibilityPrivate.h
+++ b/Engine/Source/Runtime/Renderer/Private/SceneVisibilityPrivate.h
@@ -261,6 +261,12 @@ public:
 private:
 	void Finish();
 
+	struct FViewFamilyGroup
+	{
+		const FSceneViewFamily* Family;
+		uint8 ViewSubsetMask;				// Subset of views with the given view family
+	};
+
 	struct FViewMeshArrays
 	{
 		TArray<FMeshBatchAndRelevance, SceneRenderingAllocator> DynamicMeshElements;
@@ -276,9 +282,10 @@ private:
 #endif
 	};
 
-	const FSceneViewFamily& ViewFamily;
+	const FSceneViewFamily& FirstViewFamily;
 	TArrayView<FViewInfo*> Views;
 	TArrayView<FPrimitiveSceneInfo*> Primitives;
+	TArray<FViewFamilyGroup, TInlineAllocator<1>> ViewFamilyGroups;
 	TArray<FViewMeshArrays, TInlineAllocator<2>> ViewMeshArraysPerView;
 	TArray<FDynamicPrimitive, SceneRenderingAllocator> DynamicPrimitives;
 	FMeshElementCollector MeshCollector;
