Rendering: Added FSceneView::OutputFrameIndex to go with FrameIndex. Under normal conditions (typical runtime) these two will be identical and there is no difference. When rendered in MRQ, OutputFrameIndex will only increment once per output frame, while FrameIndex increments once per frame.

This allows downstream rendering systems (such as Film Grain) to work correctly when using multiple temporal/spatial samples. Before this change, each sample would generate a unique film grain pattern which would significantly diminish the visible effect of film grain and often left frames feeling "soft".

diff --git a/Engine/Plugins/MovieScene/MovieRenderPipeline/Source/MovieRenderPipelineRenderPasses/Private/Graph/Renderers/MovieGraphImagePassBase.cpp b/Engine/Plugins/MovieScene/MovieRenderPipeline/Source/MovieRenderPipelineRenderPasses/Private/Graph/Renderers/MovieGraphImagePassBase.cpp
--- a/Engine/Plugins/MovieScene/MovieRenderPipeline/Source/MovieRenderPipelineRenderPasses/Private/Graph/Renderers/MovieGraphImagePassBase.cpp
+++ b/Engine/Plugins/MovieScene/MovieRenderPipeline/Source/MovieRenderPipelineRenderPasses/Private/Graph/Renderers/MovieGraphImagePassBase.cpp
@@ -202,6 +202,8 @@ void FMovieGraphImagePassBase::ApplyMovieGraphOverridesToSceneView(TSharedRef<FS
 	// determinism with things like TAA.
 	FSceneView* View = const_cast<FSceneView*>(InOutFamily->Views[0]);
 	View->OverrideFrameIndexValue = InInitData.FrameIndex;
+	View->OverrideOutputFrameIndexValue = InInitData.TimeData.OutputFrameNumber;
+
 	// Each shot should initialize a scene history from scratch so there should be no need to do an extra camera cut flag.
 	View->bCameraCut = false; 
 	View->AntiAliasingMethod = InInitData.AntiAliasingMethod;
diff --git a/Engine/Plugins/MovieScene/MovieRenderPipeline/Source/MovieRenderPipelineRenderPasses/Private/MoviePipelineImagePassBase.cpp b/Engine/Plugins/MovieScene/MovieRenderPipeline/Source/MovieRenderPipelineRenderPasses/Private/MoviePipelineImagePassBase.cpp
--- a/Engine/Plugins/MovieScene/MovieRenderPipeline/Source/MovieRenderPipelineRenderPasses/Private/MoviePipelineImagePassBase.cpp
+++ b/Engine/Plugins/MovieScene/MovieRenderPipeline/Source/MovieRenderPipelineRenderPasses/Private/MoviePipelineImagePassBase.cpp
@@ -251,6 +251,7 @@ TSharedPtr<FSceneViewFamilyContext> UMoviePipelineImagePassBase::CalculateViewFa
 	// Override the view's FrameIndex to be based on our progress through the sequence. This greatly increases
 	// determinism with things like TAA.
 	View->OverrideFrameIndexValue = InOutSampleState.FrameIndex;
+	View->OverrideOutputFrameIndexValue = InOutSampleState.OutputState.OutputFrameNumber;
 	View->bCameraCut = InOutSampleState.bCameraCut;
 	View->bIsOfflineRender = true;
 	View->AntiAliasingMethod = IsAntiAliasingSupported() ? InOutSampleState.AntiAliasingMethod : EAntiAliasingMethod::AAM_None;
diff --git a/Engine/Source/Runtime/Engine/Public/SceneView.h b/Engine/Source/Runtime/Engine/Public/SceneView.h
--- a/Engine/Source/Runtime/Engine/Public/SceneView.h
+++ b/Engine/Source/Runtime/Engine/Public/SceneView.h
@@ -851,6 +851,7 @@ enum ETranslucencyVolumeCascade
 	VIEW_UNIFORM_BUFFER_MEMBER(uint32, FrameCounter) \
 	VIEW_UNIFORM_BUFFER_MEMBER(uint32, StateFrameIndexMod8) \
 	VIEW_UNIFORM_BUFFER_MEMBER(uint32, StateFrameIndex) \
+	VIEW_UNIFORM_BUFFER_MEMBER(uint32, StateOutputFrameIndex) \
 	VIEW_UNIFORM_BUFFER_MEMBER(uint32, DebugViewModeMask) \
 	VIEW_UNIFORM_BUFFER_MEMBER(uint32, WorldIsPaused) \
 	VIEW_UNIFORM_BUFFER_MEMBER_EX(float, CameraCut, EShaderPrecisionModifier::Half) \
@@ -1527,6 +1528,11 @@ public:
 	
 	/** The frame index to override, useful for keeping determinism when rendering sequences. **/
 	TOptional<uint32> OverrideFrameIndexValue;
+	/** 
+	* If set, overrides the Output Frame Index. Useful for Halton-like offsets that should produce the same results in multi-sample accumulation when rendering sequences. 
+	* If unset, renderer defaults to FrameIndex.
+	*/
+	TOptional<uint32> OverrideOutputFrameIndexValue;
 
 	/** In some cases, the principal point of the lens is not at the center of the screen, especially for overlapped tile
 	 *  rendering. So given a UV in [-1,1] viewport space, convert it to the [-1,1] viewport space of the lens using
diff --git a/Engine/Source/Runtime/Renderer/Private/PostProcess/PostProcessTonemap.cpp b/Engine/Source/Runtime/Renderer/Private/PostProcess/PostProcessTonemap.cpp
--- a/Engine/Source/Runtime/Renderer/Private/PostProcess/PostProcessTonemap.cpp
+++ b/Engine/Source/Runtime/Renderer/Private/PostProcess/PostProcessTonemap.cpp
@@ -710,12 +710,14 @@ FScreenPassTexture AddTonemapPass(FRDGBuilder& GraphBuilder, const FViewInfo& Vi
 	CommonParameters.View = View.GetShaderParameters();
 
 	{
-		uint8 FrameIndexMod8 = 0;
+		uint8 OutputFrameIndexMod8 = 0;
 		if (View.ViewState)
 		{
-			FrameIndexMod8 = View.ViewState->GetFrameIndex(8);
+			// Grain should be temporally stable when accumulating multiple samples per output frame, so we use OutputFrameIndex instead of FrameIndex.
+			// Without this, the effect of grain is softened which goes against artistic intent after they tune it to the real-time viewport.
+			OutputFrameIndexMod8 = View.ViewState->GetOutputFrameIndex(8);
 		}
-		GrainRandomFromFrame(&CommonParameters.FilmGrain.GrainRandomFull, FrameIndexMod8);
+		GrainRandomFromFrame(&CommonParameters.FilmGrain.GrainRandomFull, OutputFrameIndexMod8);
 	}
 
 	if (View.FilmGrainTexture)
@@ -754,8 +756,9 @@ FScreenPassTexture AddTonemapPass(FRDGBuilder& GraphBuilder, const FViewInfo& Vi
 		CommonParameters.FilmGrain.FilmGrainTexture = FilmGrainTexture;
 		CommonParameters.FilmGrain.FilmGrainSampler = TStaticSamplerState<SF_Bilinear, AM_Wrap, AM_Wrap>::GetRHI();
 
+		// Grain should be temporally stable when accumulating multiple samples per output frame, so we use OutputFrameIndex instead of FrameIndex.
 		int32 RandomSequenceLength = CVarFilmGrainSequenceLength.GetValueOnRenderThread();
-		int32 RandomSequenceIndex = (View.ViewState ? View.ViewState->FrameIndex : 0) % RandomSequenceLength;
+		int32 RandomSequenceIndex = (View.ViewState ? View.ViewState->GetOutputFrameIndex() : 0) % RandomSequenceLength;
 
 		FVector2f RandomGrainTextureUVOffset;
 		RandomGrainTextureUVOffset.X = Halton(RandomSequenceIndex + 1, 2);
diff --git a/Engine/Source/Runtime/Renderer/Private/RendererScene.cpp b/Engine/Source/Runtime/Renderer/Private/RendererScene.cpp
--- a/Engine/Source/Runtime/Renderer/Private/RendererScene.cpp
+++ b/Engine/Source/Runtime/Renderer/Private/RendererScene.cpp
@@ -374,6 +374,7 @@ FSceneViewState::FSceneViewState(ERHIFeatureLevel::Type FeatureLevel, FSceneView
 	MIDUsedCount = 0;
 	TemporalAASampleIndex = 0;
 	FrameIndex = 0;
+	OutputFrameIndex = 0;
 	DistanceFieldTemporalSampleIndex = 0;
 	bDOFHistory = true;
 	bDOFHistory2 = true;
diff --git a/Engine/Source/Runtime/Renderer/Private/ScenePrivate.h b/Engine/Source/Runtime/Renderer/Private/ScenePrivate.h
--- a/Engine/Source/Runtime/Renderer/Private/ScenePrivate.h
+++ b/Engine/Source/Runtime/Renderer/Private/ScenePrivate.h
@@ -1002,8 +1002,20 @@ public:
 	// if TemporalAA is on this cycles through 0..TemporalAASampleCount-1, ResetViewState() puts it back to 0
 	int32 TemporalAASampleIndex;
 
-	// counts up by one each frame, warped in 0..7 range, ResetViewState() puts it back to 0
+	// Counts up by one each frame, ResetViewState() puts it back to 0. Should be used as the seed for Halton
+	// and other per-render changes. Use OutputFrameIndex as the seed if the effect should be per output frame.
+	// Under normal rendering FrameIndex == OutputFrameIndex and there is no distinction between per-render and
+	// per-frame, but when accumulating multiple samples then FrameIndex will be unique for each accumulation sample
+	// rendered, but OutputFrameIndex will only increment per multi-sample accumulated output frame.
+	// 
+	// Can be overwritten by OverrideFrameIndexValue.
 	uint32 FrameIndex;
+
+	// Counts up by one each frame, ResetViewState() puts it back to zero. See FrameIndex for more details. This should
+	// equal FrameIndex in normal scenarios, but can differ when accumulating multiple samples to produce one output frame.
+	//
+	// Can be overwritten by OverrideOutputFrameIndexValue
+	uint32 OutputFrameIndex;
 	
 	/** Informations of to persist for the next frame's FViewInfo::PrevViewInfo.
 	 *
@@ -1198,11 +1210,25 @@ public:
 		return FrameIndex;
 	}
 
+	// Returns the index of the output frame with a desired power of two modulus.
+	inline uint32 GetOutputFrameIndex(uint32 Pow2Modulus) const
+	{
+		check(FMath::IsPowerOfTwo(Pow2Modulus));
+		return OutputFrameIndex & (Pow2Modulus - 1);
+	}
+
+	// Returns 32bits output frame index. Matches GetFrameIndex unless using multi-sample accumulation.
+	inline uint32 GetOutputFrameIndex() const
+	{
+		return OutputFrameIndex;
+	}
+
 	// to make rendering more deterministic
 	virtual void ResetViewState()
 	{
 		TemporalAASampleIndex = 0;
 		FrameIndex = 0;
+		OutputFrameIndex = 0;
 		DistanceFieldTemporalSampleIndex = 0;
 		PreExposure = 1.f;
 
diff --git a/Engine/Source/Runtime/Renderer/Private/SceneRendering.cpp b/Engine/Source/Runtime/Renderer/Private/SceneRendering.cpp
--- a/Engine/Source/Runtime/Renderer/Private/SceneRendering.cpp
+++ b/Engine/Source/Runtime/Renderer/Private/SceneRendering.cpp
@@ -1734,14 +1734,17 @@ void FViewInfo::SetupUniformBufferParameters(
 	}
 
 	uint32 FrameIndex = 0;
+	uint32 OutputFrameIndex = 0;
 	if (ViewState)
 	{
 		FrameIndex = ViewState->GetFrameIndex();
+		OutputFrameIndex = ViewState->GetOutputFrameIndex();
 	}
 
 	// TODO(GA): kill StateFrameIndexMod8 because this is only a scalar bit mask with StateFrameIndex anyway.
 	ViewUniformShaderParameters.StateFrameIndexMod8 = FrameIndex % 8;
 	ViewUniformShaderParameters.StateFrameIndex = FrameIndex;
+	ViewUniformShaderParameters.StateOutputFrameIndex = OutputFrameIndex;
 
 	{
 		// If rendering in stereo, the other stereo passes uses the left eye's translucency lighting volume.
diff --git a/Engine/Source/Runtime/Renderer/Private/SceneVisibility.cpp b/Engine/Source/Runtime/Renderer/Private/SceneVisibility.cpp
--- a/Engine/Source/Runtime/Renderer/Private/SceneVisibility.cpp
+++ b/Engine/Source/Runtime/Renderer/Private/SceneVisibility.cpp
@@ -4979,6 +4979,17 @@ void FSceneRenderer::PrepareViewStateForVisibility(const FSceneTexturesConfig& S
 			{
 				ViewState->FrameIndex = View.OverrideFrameIndexValue.GetValue();
 			}
+
+			if (View.OverrideOutputFrameIndexValue.IsSet())
+			{
+				ViewState->OutputFrameIndex = View.OverrideOutputFrameIndexValue.GetValue();
+			}
+			else
+			{
+				// If the output frame index isn't being overwritten then we keep it in sync with
+				// FrameIndex so that downstream systems are unchanged.
+				ViewState->OutputFrameIndex = ViewState->FrameIndex;
+			}
 		}
 		
 		// Subpixel jitter for temporal AA
