Fixed r.water.waterinfo.rendermethod 2 being broken due to CRP rendertargets being set to external access before being consumed as RDG inputs by subsequent water CRPs. Fixed by moving UseExternalAccessMode() calls from the end of each CRP to the OnEndPass() function of scene capture CRPs. CRPs are generic and their rendertarget is a RDG pointer that is initialized by each CRP implementation, so it makes sense that the implementation should be responsible for potentially setting external access mode.

#rnx

diff --git a/Engine/Source/Runtime/Renderer/Private/DeferredShadingRenderer.cpp b/Engine/Source/Runtime/Renderer/Private/DeferredShadingRenderer.cpp
--- a/Engine/Source/Runtime/Renderer/Private/DeferredShadingRenderer.cpp
+++ b/Engine/Source/Runtime/Renderer/Private/DeferredShadingRenderer.cpp
@@ -2130,10 +2130,6 @@ void FDeferredShadingSceneRenderer::Render(FRDGBuilder& GraphBuilder)
 			#if WITH_MGPU
 				DoCrossGPUTransfers(GraphBuilder, CustomRenderPass->GetRenderTargetTexture(), CustomRenderPassViews, false, FRHIGPUMask::All());
 			#endif
-
-				// Materials in the main view renderer will be using this render target, so we need RDG to transition it back to SRV now,
-				// rather than at the end of graph execution.
-				GraphBuilder.UseExternalAccessMode(CustomRenderPass->GetRenderTargetTexture(), ERHIAccess::SRVMask);
 			}
 
 			CustomRenderPass->EndPass(GraphBuilder);
diff --git a/Engine/Source/Runtime/Renderer/Private/MobileShadingRenderer.cpp b/Engine/Source/Runtime/Renderer/Private/MobileShadingRenderer.cpp
--- a/Engine/Source/Runtime/Renderer/Private/MobileShadingRenderer.cpp
+++ b/Engine/Source/Runtime/Renderer/Private/MobileShadingRenderer.cpp
@@ -1278,10 +1278,6 @@ void FMobileSceneRenderer::Render(FRDGBuilder& GraphBuilder)
 				{
 					FGenerateMips::Execute(GraphBuilder, FeatureLevel, CustomRenderPass->GetRenderTargetTexture(), FGenerateMipsParams());
 				}
-
-				// Materials in the main view renderer will be using this render target, so we need RDG to transition it back to SRV now,
-				// rather than at the end of graph execution.
-				GraphBuilder.UseExternalAccessMode(CustomRenderPass->GetRenderTargetTexture(), ERHIAccess::SRVMask);
 			}
 
 			CustomRenderPass->EndPass(GraphBuilder);
diff --git a/Engine/Source/Runtime/Renderer/Private/SceneCaptureRendering.cpp b/Engine/Source/Runtime/Renderer/Private/SceneCaptureRendering.cpp
--- a/Engine/Source/Runtime/Renderer/Private/SceneCaptureRendering.cpp
+++ b/Engine/Source/Runtime/Renderer/Private/SceneCaptureRendering.cpp
@@ -959,6 +959,13 @@ public:
 	{
 		RenderTargetTexture = SceneCaptureRenderTarget->GetRenderTargetTexture(GraphBuilder);
 	}
+
+	virtual void OnEndPass(FRDGBuilder& GraphBuilder) override
+	{
+		// Materials in the main view renderer will be using this render target, so we need RDG to transition it back to SRV now,
+		// rather than at the end of graph execution.
+		GraphBuilder.UseExternalAccessMode(RenderTargetTexture, ERHIAccess::SRVMask);
+	}
 	
 	FRenderTarget* SceneCaptureRenderTarget = nullptr;
 };
