Fix for EngineTest crash when Visualizing motion blur in packaged builds.

When visualizing motion blur in a packaged build that may be using an override output, if the assigned texture is not flagged as a ShaderResource, this results in a failure in the subsequent call to CreateSRV. This change works around the issue by creating a separate resource for the necessary texture slice output, but copies the result to the override output too to ensure parity is maintained for all use cases.

This solution is slightly more expensive for this specific circumstance outside the editor, but as it's the only occurance and is just a debug visualize pass it's ok to accept that cost.

diff --git a/Engine/Source/Runtime/Renderer/Private/PostProcess/PostProcessMotionBlur.cpp b/Engine/Source/Runtime/Renderer/Private/PostProcess/PostProcessMotionBlur.cpp
--- a/Engine/Source/Runtime/Renderer/Private/PostProcess/PostProcessMotionBlur.cpp
+++ b/Engine/Source/Runtime/Renderer/Private/PostProcess/PostProcessMotionBlur.cpp
@@ -1167,17 +1167,50 @@ FScreenPassTextureSlice AddVisualizeMotionBlurPass(FRDGBuilder& GraphBuilder, co
 	check(Inputs.SceneVelocity.IsValid());
 	checkf(Inputs.SceneDepth.ViewRect == Inputs.SceneVelocity.ViewRect, TEXT("The implementation requires that depth and velocity have the same viewport."));
 
-	FScreenPassRenderTarget Output = Inputs.OverrideOutput;
+	FScreenPassRenderTarget Output;
+	bool bOverrideMissingSRFlag = false;
+	if (Inputs.OverrideOutput.IsValid())
+	{
+		/**
+		 * If the override texture is used, check whether the output texture supports SRVs before assigning it, as this causes a crash in
+		 * FScreenPassTextureSlice::CreateFromScreenPassTexture at the end of this function if not.
+		 */
+		if(static_cast<uint64>(Inputs.OverrideOutput.Texture->Desc.Flags) & static_cast<uint64>(ETextureCreateFlags::ShaderResource))
+		{
+			Output = Inputs.OverrideOutput;
+		}
+		else
+		{
+			bOverrideMissingSRFlag = true;
+		}
+	}
 
 	if (!Output.IsValid())
 	{
+		FIntPoint OutputExtent;
+		EPixelFormat OutputFormat;
+		FIntRect OutputRect;
+		if (bOverrideMissingSRFlag)
+		{
+			//If the Override texture does not support SRVs, we need to create one that does so we can call/return via FScreenPassTextureSlice::CreateFromScreenPassTexture
+			OutputExtent = Inputs.OverrideOutput.Texture->Desc.Extent;
+			OutputFormat = Inputs.OverrideOutput.Texture->Desc.Format;
+			OutputRect = Inputs.OverrideOutput.ViewRect;
+		}
+		else
+		{
+			OutputExtent = Inputs.SceneColor.TextureSRV->Desc.Texture->Desc.Extent;
+			OutputFormat = Inputs.SceneColor.TextureSRV->Desc.Texture->Desc.Format;
+			OutputRect = Inputs.SceneColor.ViewRect;
+		}
+
 		FRDGTextureDesc OutputDesc = FRDGTextureDesc::Create2D(
-			Inputs.SceneColor.TextureSRV->Desc.Texture->Desc.Extent,
-			Inputs.SceneColor.TextureSRV->Desc.Texture->Desc.Format,
+			OutputExtent,
+			OutputFormat,
 			FClearValueBinding::None,
 			ETextureCreateFlags::ShaderResource | ETextureCreateFlags::RenderTargetable);
 			
-		Output = FScreenPassRenderTarget(GraphBuilder.CreateTexture(OutputDesc, TEXT("MotionBlur.Visualize")), Inputs.SceneColor.ViewRect, View.GetOverwriteLoadAction());
+		Output = FScreenPassRenderTarget(GraphBuilder.CreateTexture(OutputDesc, TEXT("MotionBlur.Visualize")), OutputRect, View.GetOverwriteLoadAction());
 	}
 
 	// NOTE: Scene depth is used as the velocity viewport because velocity can actually be a 1x1 black texture.
@@ -1263,6 +1296,15 @@ FScreenPassTextureSlice AddVisualizeMotionBlurPass(FRDGBuilder& GraphBuilder, co
 			});
 	}
 
+	if (bOverrideMissingSRFlag)
+	{
+		/**
+		 * The OverrideOutput texture did not support SRVs, so we created one that did, but we still need to resolve the output image back
+		 * to the OverrideOutput texture too.
+		 */
+		AddCopyTexturePass(GraphBuilder, Output.Texture, Inputs.OverrideOutput.Texture, FRHICopyTextureInfo());
+	}
+
 	return FScreenPassTextureSlice::CreateFromScreenPassTexture(GraphBuilder, Output);
 }
 
