Mobile: Pre-cache PSOs for TranslucencyAll only when separate translucency is disabled

diff --git a/Engine/Source/Runtime/Renderer/Private/MobileBasePass.cpp b/Engine/Source/Runtime/Renderer/Private/MobileBasePass.cpp
--- a/Engine/Source/Runtime/Renderer/Private/MobileBasePass.cpp
+++ b/Engine/Source/Runtime/Renderer/Private/MobileBasePass.cpp
@@ -1071,11 +1071,23 @@ void FMobileBasePassMeshProcessor::CollectPSOInitializersForLMPolicy(
 
 void FMobileBasePassMeshProcessor::CollectPSOInitializers(const FSceneTexturesConfig& SceneTexturesConfig, const FMaterial& Material, const FPSOPrecacheVertexFactoryData& VertexFactoryData, const FPSOPrecacheParams& PreCacheParams, TArray<FPSOPrecacheData>& PSOInitializers)
 {
-	static IConsoleVariable* PSOPrecacheTranslucencyAllPass = IConsoleManager::Get().FindConsoleVariable(TEXT("r.PSOPrecache.TranslucencyAllPass"));
-	// PSO precaching enabled for TranslucencyAll
-	if (MeshPassType == EMeshPass::TranslucencyAll && PSOPrecacheTranslucencyAllPass->GetInt() == 0)
+	if (bTranslucentBasePass)
 	{
-		return;
+		static IConsoleVariable* PSOPrecacheTranslucencyAllPass = IConsoleManager::Get().FindConsoleVariable(TEXT("r.PSOPrecache.TranslucencyAllPass"));
+		static IConsoleVariable* CVarSeparateTranslucency = IConsoleManager::Get().FindConsoleVariable(TEXT("r.SeparateTranslucency"));
+		if (CVarSeparateTranslucency->GetInt() == 0)
+		{
+			if (MeshPassType != EMeshPass::TranslucencyAll)
+			{
+				// Precache only TranslucencyAll when SeparateTranslucency is not active
+				return;
+			}
+		}
+		else if (MeshPassType == EMeshPass::TranslucencyAll && PSOPrecacheTranslucencyAllPass->GetInt() == 0)
+		{
+			// PSO precaching is disabled for TranslucencyAll while SeparateTranslucency is active
+			return;
+		}
 	}
 	
 	// Check if material should be rendered
