Fixed a check that happens when there are too many cells in the scene culling grid wrt the reduced precision in the internal cache. Made the RLE encoding use a variable number of integers to allow up to 31 bits of range.

diff --git a/Engine/Source/Runtime/Renderer/Private/SceneCulling/SceneCulling.cpp b/Engine/Source/Runtime/Renderer/Private/SceneCulling/SceneCulling.cpp
--- a/Engine/Source/Runtime/Renderer/Private/SceneCulling/SceneCulling.cpp
+++ b/Engine/Source/Runtime/Renderer/Private/SceneCulling/SceneCulling.cpp
@@ -1482,6 +1482,7 @@ public:
 		BUILDER_LOG_LIST("BuildInstanceRange(%d, %d):", InstanceDataOffset, NumInstances);
 
 		constexpr bool bCompressRLE = UpdateFrequencyCategory != EUpdateFrequencyCategory::Dynamic;
+		CellIndexCacheEntry.bSingleInstanceOnly = !bCompressRLE;
 
 		FSceneCulling::FLocation64 PrevInstanceCellLoc;
 		int32 SameInstanceLocRunCount = 0;
@@ -1789,11 +1790,14 @@ public:
 			BuildInstanceRange<UpdateFrequencyCategory>(InstanceDataOffset, NumInstances, HashLocationComputer, CellIndexCacheEntry);
 		}
 
+#if SC_ENABLE_DETAILED_LOGGING
 		BUILDER_LOG_LIST("CellIndexCacheEntry(%d):", CellIndexCacheEntry.Items.Num());
-		for (FCellIndexCacheEntry::FItem Item : CellIndexCacheEntry.Items)
+		for (int32 ItemIndex = NumInstances; ItemIndex < CellIndexCacheEntry.Items.Num(); ++ItemIndex)
 		{
+			FCellIndexCacheEntry::FItem Item = CellIndexCacheEntry.LoadAndStepItem(ItemIndex);
 			BUILDER_LOG_LIST_APPEND("(%d, %d)", Item.CellIndex, Item.NumInstances);
 		}
+#endif
 
 		SceneCulling.TotalCellIndexCacheItems += CellIndexCacheEntry.Items.Num();
 
@@ -1992,8 +1996,9 @@ public:
 			const FCellIndexCacheEntry &CellIndexCacheEntry = GetCacheEntry(CacheIndex);
 
 			BUILDER_LOG_LIST("MarkForRemove(%d):", CellIndexCacheEntry.Items.Num());
-			for (FCellIndexCacheEntry::FItem Item : CellIndexCacheEntry.Items)
+			for (int32 ItemIndex = 0; ItemIndex < CellIndexCacheEntry.Items.Num(); ++ItemIndex)
 			{
+				FCellIndexCacheEntry::FItem Item = CellIndexCacheEntry.LoadAndStepItem(ItemIndex);
 				BUILDER_LOG_LIST_APPEND("(%d, %d)", Item.CellIndex, Item.NumInstances);
 				MarkCellForRemove(Item.CellIndex, Item.NumInstances, (PrimitiveState.State == FPrimitiveState::Dynamic) ? EUpdateFrequencyCategory::Dynamic : EUpdateFrequencyCategory::Static);
 			}
@@ -2008,7 +2013,7 @@ public:
 	}
 
 	template <typename HashLocationComputerType>
-	SC_FORCEINLINE void UpdateProcessDynamicInstances(HashLocationComputerType &HashLocationComputer, int32 InstanceDataOffset, int32 NumInstances, int32 PrevNumInstances, FSceneCulling::FCellIndexCacheEntry &CacheEntry)
+	SC_FORCEINLINE void UpdateProcessDynamicInstances(HashLocationComputerType &HashLocationComputer, int32 InstanceDataOffset, int32 NumInstances, int32 PrevNumInstances, FSceneCulling::FCellIndexCacheEntry &CellIndexCacheEntry)
 	{
 		for (int32 InstanceIndex = 0; InstanceIndex < NumInstances; ++InstanceIndex)
 		{
@@ -2018,7 +2023,7 @@ public:
 			bool bNeedAdd = InstanceIndex >= PrevNumInstances;
 			if (!bNeedAdd)
 			{
-				int32 PrevCellIndex =  CacheEntry.Items[InstanceIndex].CellIndex;
+				int32 PrevCellIndex =  CellIndexCacheEntry.LoadAndStepItem(InstanceIndex).CellIndex;
 				FSceneCulling::FLocation64 PrevCellLoc = SceneCulling.GetCellLoc(PrevCellIndex);
 				if (PrevCellLoc != InstanceCellLoc)
 				{
@@ -2033,7 +2038,7 @@ public:
 			if (bNeedAdd)
 			{
 				int32 CellIndex = AddToCell(InstanceCellLoc, InstanceId, EUpdateFrequencyCategory::Dynamic);
-				CacheEntry.Set(InstanceIndex, CellIndex, 1);
+				CellIndexCacheEntry.Set(InstanceIndex, CellIndex);
 			}
 		}
 	}
@@ -2117,19 +2122,21 @@ public:
 		{
 			SC_SCOPED_NAMED_EVENT_DETAIL(SceneCulling_Post_UpdateInstances_DynamicUpdate, FColor::Red);
 			check(InstanceSceneDataBuffers != nullptr);
+			check(PrevPrimitiveState.bDynamic);
 
 			// For dynamic instance batches we process individual instances since they can then more often be retained
 			// Stored in the same data structure, just guaranteed to be singular instances
 			FCellIndexCacheEntry &CellIndexCacheEntry = GetCacheEntry(PrevPrimitiveState.Payload);
 			// retain previous state.
 			NewPrimitiveState.Payload = PrevPrimitiveState.Payload;
-
-			// Mark overflowing ones for remove (if the number of instances shrank
+			check(CellIndexCacheEntry.bSingleInstanceOnly);
+			// Mark overflowing ones for remove (if the number of instances shrank)
 			for (int32 ItemIndex = NumInstances; ItemIndex < CellIndexCacheEntry.Items.Num(); ++ItemIndex)
 			{
-				FCellIndexCacheEntry::FItem Item = CellIndexCacheEntry.Items[ItemIndex];
+				FCellIndexCacheEntry::FItem Item = CellIndexCacheEntry.LoadAndStepItem(ItemIndex);
+				check(Item.NumInstances == 1);
 				// Assumes 1:1 between index and ID
-				MarkForRemove(Item.CellIndex, PrevPrimitiveState.InstanceDataOffset + ItemIndex, Item.NumInstances, PrevPrimitiveState.bDynamic ? EUpdateFrequencyCategory::Dynamic : EUpdateFrequencyCategory::Static);
+				MarkForRemove(Item.CellIndex, PrevPrimitiveState.InstanceDataOffset + ItemIndex, 1, EUpdateFrequencyCategory::Dynamic);
 			}
 			// Maintain the total accross all entries
 			SceneCulling.TotalCellIndexCacheItems -= CellIndexCacheEntry.Items.Num();
@@ -2170,8 +2177,9 @@ public:
 			FCellIndexCacheEntry &CellIndexCacheEntry = GetCacheEntry(PrevPrimitiveState.Payload);
 
 			RemovedInstanceFlags.SetRange(PrevPrimitiveState.InstanceDataOffset, PrevPrimitiveState.NumInstances, true);
-			for (FCellIndexCacheEntry::FItem Item : CellIndexCacheEntry.Items)
+			for (int32 ItemIndex = NumInstances; ItemIndex < CellIndexCacheEntry.Items.Num(); ++ItemIndex)
 			{
+				FCellIndexCacheEntry::FItem Item = CellIndexCacheEntry.LoadAndStepItem(ItemIndex);
 				MarkCellForRemove(Item.CellIndex, Item.NumInstances, PrevPrimitiveState.bDynamic ? EUpdateFrequencyCategory::Dynamic : EUpdateFrequencyCategory::Static);
 			}
 			
diff --git a/Engine/Source/Runtime/Renderer/Private/SceneCulling/SceneCulling.h b/Engine/Source/Runtime/Renderer/Private/SceneCulling/SceneCulling.h
--- a/Engine/Source/Runtime/Renderer/Private/SceneCulling/SceneCulling.h
+++ b/Engine/Source/Runtime/Renderer/Private/SceneCulling/SceneCulling.h
@@ -161,43 +161,43 @@ private:
 	 */
 	struct FCellIndexCacheEntry
 	{
-		static constexpr int32 InstanceCountNumBits = 12;
-		static constexpr int32 InstanceCountMax = (1 << InstanceCountNumBits);
-		static constexpr int32 CellIndexMax = 1 << (32 - InstanceCountNumBits);
+		static constexpr uint32 SingleInstanceMask = 1u << 31;
+		static constexpr uint32 CellIndexMask = (1u << 31) - 1u;
+		static constexpr uint32 CellIndexMax = 1u << 31;
+
 		struct FItem
 		{
-			uint32 NumInstances : InstanceCountNumBits;
-			uint32 CellIndex : 32 - InstanceCountNumBits;
+			int32 NumInstances;
+			int32 CellIndex;
 		};
 
 		inline void Add(int32 CellIndex, int32 NumInstances)
 		{
-			check(CellIndex < CellIndexMax);
-			
-			FItem Item;
-			Item.CellIndex = CellIndex;
-			// break into blocks...
-			while (NumInstances >= InstanceCountMax)
+			check(uint32(CellIndex) < CellIndexMax);
+			check(!bSingleInstanceOnly || NumInstances == 1);
+
+			if (NumInstances > 1)
 			{
-				Item.NumInstances = InstanceCountMax - 1;
-				NumInstances -= InstanceCountMax - 1;
-				Items.Add(Item);
+				// Add RLE entry
+				Items.Add(CellIndex);
+				Items.Add(NumInstances);
+			}
+			else
+			{
+				// Mark as single-istance
+				Items.Add(uint32(CellIndex) | SingleInstanceMask);
 			}
-
-			check(NumInstances < InstanceCountMax);
-			Item.NumInstances = NumInstances;
-			Items.Add(Item);
 		}
 
-		inline void Set(int32 Index, int32 CellIndex, int32 NumInstances)
+		/**
+		 * Only possible if there is one single instance items in the list, otherwise we don't have a 1:1 mapping.
+		 */
+		inline void Set(int32 Index, int32 CellIndex)
 		{
-			check(CellIndex < CellIndexMax);
-			check(NumInstances < InstanceCountMax);
+			check(bSingleInstanceOnly);
+			check(uint32(CellIndex) < CellIndexMax);
 
-			FItem Item;
-			Item.CellIndex = CellIndex;
-			Item.NumInstances = NumInstances;
-			Items[Index] = Item;
+			Items[Index] = uint32(CellIndex) | SingleInstanceMask;
 		}
 
 		FORCEINLINE void Reset() 
@@ -205,7 +205,25 @@ private:
 			Items.Reset();
 		}
 
-		TArray<FItem> Items;
+
+		/**
+		 * Load and unpack an item at a given ItemIndex. NOTE: advances ItemIndex if the item is RLE'd
+		 */
+		FItem LoadAndStepItem(int32 &InOutItemIndex) const
+		{
+			FItem Result; 
+			uint32 PackedCellIndex = Items[InOutItemIndex];
+			Result.CellIndex = int32(PackedCellIndex & FCellIndexCacheEntry::CellIndexMask);
+			Result.NumInstances = 1;
+			if ((PackedCellIndex & FCellIndexCacheEntry::SingleInstanceMask) == 0u)
+			{
+				Result.NumInstances = int32(Items[++InOutItemIndex]);
+			}
+			return Result;
+		}
+
+		bool bSingleInstanceOnly = false;
+		TArray<uint32> Items;
 	};
 
 	/** 
