Local Fog Volumes - updated render settings tool tips and cvars.

diff --git a/Engine/Source/Runtime/Engine/Classes/Engine/RendererSettings.h b/Engine/Source/Runtime/Engine/Classes/Engine/RendererSettings.h
--- a/Engine/Source/Runtime/Engine/Classes/Engine/RendererSettings.h
+++ b/Engine/Source/Runtime/Engine/Classes/Engine/RendererSettings.h
@@ -691,7 +691,7 @@ class URendererSettings : public UDeveloperSettings
 
 	UPROPERTY(config, EditAnywhere, Category=Translucency, meta=(
 		ConsoleVariable="r.LocalFogVolume.ApplyOnTranslucent",
-		ToolTip="Allow translucency to be rendered to a separate render targeted and composited after depth of field. Prevents translucency from appearing out of focus."))
+		ToolTip="Allow local fog volumes to be combined and rendered over translucent meshes. Only per vertex evaluation is supported as of today. It requires r.SupportLocalFogVolumes to be true"))
 	uint32 bLocalFogVolumeApplyOnTranslucent:1;
 
 	UPROPERTY(config, EditAnywhere, Category = VR, meta = (
@@ -958,13 +958,9 @@ class URendererSettings : public UDeveloperSettings
 		ConfigRestartRequired = true))
 		uint32 bSupportExpFogMatchesVolumetricFog : 1;
 
-	/**
-	"Local fog volume components can will need to be applied on translucent, and opaque in forward, so resources will need to be bound to apply aerial perspective on transparent surfaces (and all surfaces on mobile via per vertex evaluation)."
-	"It requires r.SupportLocalFogVolumes to be true."
-	*/
 	UPROPERTY(config, EditAnywhere, Category = Optimizations, meta = (
 		ConsoleVariable = "r.SupportLocalFogVolumes", DisplayName = "Support Local Fog Volumes",
-		ToolTip = "Local fog volume components can will need to be applied on translucent, and opaque in forward, so resources will need to be bound to apply aerial perspective on transparent surfaces (and all surfaces on mobile via per vertex evaluation). It requires r.SupportLocalFogVolumes to be true.",
+		ToolTip = "Enable local fog volume components rendering support. Disable that setting if local fog volumes are not used in a project for increased performance.",
 		ConfigRestartRequired = true))
 		uint32 bSupportLocalFogVolumes : 1;
 
diff --git a/Engine/Source/Runtime/Renderer/Private/LocalFogVolumeRendering.cpp b/Engine/Source/Runtime/Renderer/Private/LocalFogVolumeRendering.cpp
--- a/Engine/Source/Runtime/Renderer/Private/LocalFogVolumeRendering.cpp
+++ b/Engine/Source/Runtime/Renderer/Private/LocalFogVolumeRendering.cpp
@@ -12,7 +12,7 @@
 // The runtime ON/OFF toggle
 static TAutoConsoleVariable<int32> CVarLocalFogVolume(
 	TEXT("r.LocalFogVolume"), 1,
-	TEXT("LocalFogVolume components are rendered when this is not 0, otherwise ignored."),
+	TEXT("Project settings enabling the rendering of Local Fog Volumes."),
 	ECVF_RenderThreadSafe);
 
 // The project setting (disable runtime and shader code)
@@ -29,7 +29,7 @@ static TAutoConsoleVariable<int32> CVarLocalFogVolumeRenderDuringHeightFogPass(
 
 static TAutoConsoleVariable<int32> CVarLocalFogVolumeRenderIntoVolumetricFog(
 	TEXT("r.LocalFogVolume.RenderIntoVolumetricFog"), 1,
-	TEXT("LocalFogVolume are going to be voxelised into the volumetric fog when this is not 0, otherwise it will remain isolated."),
+	TEXT("Enables the voxelization of local fog volumes into the volumetric fog rendering system. Otherwise, local fog volumes will remain isolated."),
 	ECVF_RenderThreadSafe);
 
 static TAutoConsoleVariable<float> CVarLocalFogVolumeMaxDensityIntoVolumetricFog(
@@ -49,12 +49,12 @@ static TAutoConsoleVariable<int32> CVarLocalFogVolumeTilePixelSize(
 
 static TAutoConsoleVariable<int32> CVarLocalFogVolumeTileMaxInstanceCount(
 	TEXT("r.LocalFogVolume.TileMaxInstanceCount"), 32,
-	TEXT("Maximum number of local fog volume to account for per view (and per tile or consistency)."),
+	TEXT("Maximum number of local fog volumes to account for per view (and per tile or consistency)."),
 	ECVF_RenderThreadSafe);
 
 static TAutoConsoleVariable<int32> CVarLocalFogVolumeTileCullingUseAsync(
 	TEXT("r.LocalFogVolume.TileCullingUseAsync"), 1,
-	TEXT("True if we want to try and use culling on the async pipe."),
+	TEXT("Enables running the culling process on the async compute pipe."),
 	ECVF_RenderThreadSafe);
 
 static TAutoConsoleVariable<int32> CVarLocalFogVolumeTileDebug(
@@ -69,12 +69,12 @@ static TAutoConsoleVariable<float> CVarLocalFogVolumeGlobalStartDistance(
 
 static TAutoConsoleVariable<int32> CVarLocalFogVolumeUseHZB(
 	TEXT("r.LocalFogVolume.UseHZB"), 1,
-	TEXT("Use the HZB to cull loca lfog volume away.\n"),
+	TEXT("Enables the use of the HZB to cull local fog volumes away.\n"),
 	ECVF_RenderThreadSafe);
 
 static TAutoConsoleVariable<int32> CVarLocalFogVolumeHalfResolution(
 	TEXT("r.LocalFogVolume.HalfResolution"), 0,
-	TEXT("Set to one to render local fog volumes at half resoltuion with an upsampling to full resolution later. Only work for the mobile path for now.\n"),
+	TEXT("Enables half resolution rendering of local fog volumes with an upsampling to full resolution later. Only works for the mobile path for now.\n"),
 	ECVF_RenderThreadSafe);
 
 // Example of tile setup
