[Materials] - Custom Buffer visualization now works with Material Instances.

- The code would always get the parent UMaterial from the UMaterialInterface and render with that.
- Now we use the original material interface (which could be a material instance with different parameters) to render with.
- Still query the blendable location and other material settings from the parent material.

diff --git a/Engine/Source/Runtime/Renderer/Private/PostProcess/PostProcessMaterial.cpp b/Engine/Source/Runtime/Renderer/Private/PostProcess/PostProcessMaterial.cpp
--- a/Engine/Source/Runtime/Renderer/Private/PostProcess/PostProcessMaterial.cpp
+++ b/Engine/Source/Runtime/Renderer/Private/PostProcess/PostProcessMaterial.cpp
@@ -985,7 +985,7 @@ FPostProcessMaterialChain GetPostProcessMaterialChain(const FViewInfo& View, EBl
 
 		if (Material && (Material->BlendableLocation == Location || Location == EBlendableLocation::BL_MAX))
 		{
-			Nodes.Add(FPostProcessMaterialNode(Material, VisMaterial->GetBlendableLocation(Material), VisMaterial->GetBlendablePriority(Material), Material->bIsBlendable));
+			Nodes.Add(FPostProcessMaterialNode(VisMaterial, VisMaterial->GetBlendableLocation(Material), VisMaterial->GetBlendablePriority(Material), Material->bIsBlendable));
 		}
 	}
 
