Fix depth texture not set in spatial denoising mode for NOFR.

#tests Does not crash on RenderGallary project

diff --git a/Engine/Source/Runtime/Renderer/Private/PathTracingSpatialTemporalDenoising.cpp b/Engine/Source/Runtime/Renderer/Private/PathTracingSpatialTemporalDenoising.cpp
--- a/Engine/Source/Runtime/Renderer/Private/PathTracingSpatialTemporalDenoising.cpp
+++ b/Engine/Source/Runtime/Renderer/Private/PathTracingSpatialTemporalDenoising.cpp
@@ -1038,6 +1038,7 @@ static void PathTracingDenoiserPlugin(FRDGBuilder& GraphBuilder,
 	FRDGTextureRef InputTexture, 
 	FRDGTextureRef AlbedoTexture,
 	FRDGTextureRef NormalTexture,
+	FRDGTextureRef DepthTexture,
 	FRDGTextureRef VarianceTexture,
 	FRDGTextureRef OutputTexture)
 {
@@ -1084,6 +1085,7 @@ static void PathTracingDenoiserPlugin(FRDGBuilder& GraphBuilder,
 			Inputs.ColorTex = InputTexture;
 			Inputs.AlbedoTex = AlbedoTexture;
 			Inputs.NormalTex = NormalTexture;
+			Inputs.DepthTex = DepthTexture;
 			Inputs.VarianceTex = VarianceTexture;
 			
 			// Set flow texture to black
@@ -2003,6 +2005,7 @@ void PathTracingSpatialTemporalDenoising(FRDGBuilder& GraphBuilder,
 				SpatialTemporalDenoisingContext.RadianceTexture,
 				SpatialTemporalDenoisingContext.AlbedoTexture,
 				SpatialTemporalDenoisingContext.NormalTexture,
+				SpatialTemporalDenoisingContext.DepthTexture,
 				SpatialTemporalDenoisingContext.VarianceTexture,
 				TargetTexture);
 		}
