Mobile: Fixed case where static actors were not rendered into stationary light CSM in a project with a pre-computed lighting

diff --git a/Engine/Source/Runtime/Renderer/Private/ShadowSetup.cpp b/Engine/Source/Runtime/Renderer/Private/ShadowSetup.cpp
--- a/Engine/Source/Runtime/Renderer/Private/ShadowSetup.cpp
+++ b/Engine/Source/Runtime/Renderer/Private/ShadowSetup.cpp
@@ -1348,6 +1348,9 @@ struct FDynamicShadowsTaskData
 	bool bHasRayTracedDistanceFieldShadows = false;
 	bool bFinishedMeshPassSetup = false;
 
+	// Whether static actors with a static shadowing should be rendered into CSM on mobile
+	bool bMobileAllowStaticCSM = false;
+
 	// Generated from prepare task
 	TArray<FProjectedShadowInfo*, SceneRenderingAllocator> PreShadows;
 	TArray<FProjectedShadowInfo*, SceneRenderingAllocator> ViewDependentWholeSceneShadows;
@@ -1415,6 +1418,8 @@ struct FDynamicShadowsTaskData
 			MeshCollectors.Emplace(Allocator.Create<FShadowMeshCollector>(InRHICmdList, *SceneRenderer));
 			BeginGatherDynamicMeshElementsTask.Trigger();
 		}
+
+		bMobileAllowStaticCSM = (FeatureLevel == ERHIFeatureLevel::ES3_1 && FReadOnlyCVARCache::AllowStaticLighting() && FReadOnlyCVARCache::MobileEnableStaticAndCSMShadowReceivers());
 	}
 
 	~FDynamicShadowsTaskData()
@@ -4971,7 +4976,7 @@ struct FGatherShadowPrimitivesPacket
 					// Only render shadows from objects that use static lighting during a reflection capture, since the reflection capture doesn't update at runtime
 					&& (!TaskData.bStaticSceneOnly || PrimitiveProxy->HasStaticLighting())
 					// Render dynamic lit objects if CSMForDynamicObjects is enabled.
-					&& (!LightProxy.UseCSMForDynamicObjects() || !PrimitiveProxy->HasStaticLighting()))
+					&& (!LightProxy.UseCSMForDynamicObjects() || !PrimitiveProxy->HasStaticLighting() || TaskData.bMobileAllowStaticCSM))
 				{
 					FAddSubjectPrimitiveResult Result;
 					Result.Qword = ProjectedShadowInfo->AddSubjectPrimitive_AnyThread(
