Use ShaderBindingLayout on single RayGen RTPSOs as well when generated ad hoc

diff --git a/Engine/Source/Runtime/Renderer/Private/RayTracing/RayTracingAmbientOcclusion.cpp b/Engine/Source/Runtime/Renderer/Private/RayTracing/RayTracingAmbientOcclusion.cpp
--- a/Engine/Source/Runtime/Renderer/Private/RayTracing/RayTracingAmbientOcclusion.cpp
+++ b/Engine/Source/Runtime/Renderer/Private/RayTracing/RayTracingAmbientOcclusion.cpp
@@ -206,6 +206,13 @@ void FDeferredShadingSceneRenderer::RenderRayTracingAmbientOcclusion(
 			// Declare default pipeline
 			FRayTracingPipelineStateInitializer Initializer;
 			Initializer.MaxPayloadSizeInBytes = GetRayTracingPayloadTypeMaxSize(ERayTracingPayloadType::RayTracingMaterial);
+			
+			const FShaderBindingLayout* ShaderBindingLayout = RayTracing::GetShaderBindingLayout(ShaderPlatform);
+			if (ShaderBindingLayout)
+			{
+				Initializer.ShaderBindingLayout = &ShaderBindingLayout->RHILayout;
+			}
+
 			FRHIRayTracingShader* RayGenShaderTable[] = { RayGenerationShader.GetRayTracingShader() };
 			Initializer.SetRayGenShaderTable(RayGenShaderTable);
 
diff --git a/Engine/Source/Runtime/Renderer/Private/RayTracing/RayTracingDebug.cpp b/Engine/Source/Runtime/Renderer/Private/RayTracing/RayTracingDebug.cpp
--- a/Engine/Source/Runtime/Renderer/Private/RayTracing/RayTracingDebug.cpp
+++ b/Engine/Source/Runtime/Renderer/Private/RayTracing/RayTracingDebug.cpp
@@ -829,6 +829,12 @@ static FRDGBufferRef RayTracingPerformPicking(FRDGBuilder& GraphBuilder, const F
 	FRayTracingPipelineStateInitializer Initializer;
 	Initializer.MaxPayloadSizeInBytes = GetRayTracingPayloadTypeMaxSize(ERayTracingPayloadType::RayTracingDebug);
 
+	const FShaderBindingLayout* ShaderBindingLayout = RayTracing::GetShaderBindingLayout(Scene->GetShaderPlatform());
+	if (ShaderBindingLayout)
+	{
+		Initializer.ShaderBindingLayout = &ShaderBindingLayout->RHILayout;
+	}
+
 	FRHIRayTracingShader* RayGenShaderTable[] = { RayGenShader.GetRayTracingShader() };
 	Initializer.SetRayGenShaderTable(RayGenShaderTable);
 
@@ -1085,6 +1091,12 @@ static FRDGBufferRef RayTracingPerformHitStatsPerPrimitive(FRDGBuilder& GraphBui
 
 	FRayTracingPipelineStateInitializer Initializer;
 	Initializer.MaxPayloadSizeInBytes = GetRayTracingPayloadTypeMaxSize(ERayTracingPayloadType::RayTracingDebug);
+	
+	const FShaderBindingLayout* ShaderBindingLayout = RayTracing::GetShaderBindingLayout(Scene->GetShaderPlatform());
+	if (ShaderBindingLayout)
+	{
+		Initializer.ShaderBindingLayout = &ShaderBindingLayout->RHILayout;
+	}
 
 	FRHIRayTracingShader* RayGenShaderTable[] = { RayGenShader.GetRayTracingShader() };
 	Initializer.SetRayGenShaderTable(RayGenShaderTable);
@@ -1531,6 +1543,12 @@ void FDeferredShadingSceneRenderer::RenderRayTracingDebug(FRDGBuilder& GraphBuil
 	{
 		FRayTracingPipelineStateInitializer Initializer;
 		Initializer.MaxPayloadSizeInBytes = GetRayTracingPayloadTypeMaxSize(ERayTracingPayloadType::RayTracingDebug);
+		
+		const FShaderBindingLayout* ShaderBindingLayout = RayTracing::GetShaderBindingLayout(ShaderPlatform);
+		if (ShaderBindingLayout)
+		{
+			Initializer.ShaderBindingLayout = &ShaderBindingLayout->RHILayout;
+		}
 
 		FRHIRayTracingShader* RayGenShaderTable[] = { RayGenShader.GetRayTracingShader() };
 		Initializer.SetRayGenShaderTable(RayGenShaderTable);
diff --git a/Engine/Source/Runtime/Renderer/Private/RayTracing/RayTracingShadows.cpp b/Engine/Source/Runtime/Renderer/Private/RayTracing/RayTracingShadows.cpp
--- a/Engine/Source/Runtime/Renderer/Private/RayTracing/RayTracingShadows.cpp
+++ b/Engine/Source/Runtime/Renderer/Private/RayTracing/RayTracingShadows.cpp
@@ -547,6 +547,12 @@ void FDeferredShadingSceneRenderer::RenderRayTracingShadows(
 						FRayTracingPipelineStateInitializer Initializer;
 
 						Initializer.MaxPayloadSizeInBytes = GetRayTracingPayloadTypeMaxSize(ERayTracingPayloadType::RayTracingMaterial);
+						
+						const FShaderBindingLayout* ShaderBindingLayout = RayTracing::GetShaderBindingLayout(ShaderPlatform);
+						if (ShaderBindingLayout)
+						{
+							Initializer.ShaderBindingLayout = &ShaderBindingLayout->RHILayout;
+						}
 
 						FRHIRayTracingShader* RayGenShaderTable[] = { RayGenerationShader.GetRayTracingShader() };
 						Initializer.SetRayGenShaderTable(RayGenShaderTable);
diff --git a/Engine/Source/Runtime/Renderer/Private/RayTracing/RaytracingSkylight.cpp b/Engine/Source/Runtime/Renderer/Private/RayTracing/RaytracingSkylight.cpp
--- a/Engine/Source/Runtime/Renderer/Private/RayTracing/RaytracingSkylight.cpp
+++ b/Engine/Source/Runtime/Renderer/Private/RayTracing/RaytracingSkylight.cpp
@@ -489,6 +489,13 @@ void FDeferredShadingSceneRenderer::RenderRayTracingSkyLight(
 			{
 				// Declare default pipeline
 				FRayTracingPipelineStateInitializer Initializer;
+
+				const FShaderBindingLayout* ShaderBindingLayout = RayTracing::GetShaderBindingLayout(ShaderPlatform);
+				if (ShaderBindingLayout)
+				{
+					Initializer.ShaderBindingLayout = &ShaderBindingLayout->RHILayout;
+				}
+
 				Initializer.MaxPayloadSizeInBytes = GetRayTracingPayloadTypeMaxSize(ERayTracingPayloadType::RayTracingMaterial);
 				FRHIRayTracingShader* RayGenShaderTable[] = { RayGenerationShader.GetRayTracingShader() };
 				Initializer.SetRayGenShaderTable(RayGenShaderTable);
