Fix incorrect loop start index & a typo.

#rnx

diff --git a/Engine/Source/Runtime/Renderer/Private/SceneCulling/SceneCulling.cpp b/Engine/Source/Runtime/Renderer/Private/SceneCulling/SceneCulling.cpp
--- a/Engine/Source/Runtime/Renderer/Private/SceneCulling/SceneCulling.cpp
+++ b/Engine/Source/Runtime/Renderer/Private/SceneCulling/SceneCulling.cpp
@@ -1293,16 +1293,16 @@ public:
 			int32 TotalItemChunks = 0;
 			// 1. flush the dynamic stuff
 			{
-				int32 NumToremove = RemovedInstanceCount[EUpdateFrequencyCategory::Dynamic];
-				TotalItemChunks += Builders[EUpdateFrequencyCategory::Dynamic].FinalizeChunks(Builder, PrevItemChunksOffset + PrevNumStaticItemChunks, PrevItemChunksOffset + PrevNumItemChunks, NumToremove);
-				check(NumToremove == 0);
+				int32 NumToRemove = RemovedInstanceCount[EUpdateFrequencyCategory::Dynamic];
+				TotalItemChunks += Builders[EUpdateFrequencyCategory::Dynamic].FinalizeChunks(Builder, PrevItemChunksOffset + PrevNumStaticItemChunks, PrevItemChunksOffset + PrevNumItemChunks, NumToRemove);
+				check(NumToRemove == 0);
 			}
 
 			// 2. And then the static.
 			{
-				int32 NumToremove = RemovedInstanceCount[EUpdateFrequencyCategory::Static];
-				TotalItemChunks += Builders[EUpdateFrequencyCategory::Static].FinalizeChunks(Builder, PrevItemChunksOffset, PrevItemChunksOffset + PrevNumStaticItemChunks, NumToremove);
-				check(NumToremove == 0);
+				int32 NumToRemove = RemovedInstanceCount[EUpdateFrequencyCategory::Static];
+				TotalItemChunks += Builders[EUpdateFrequencyCategory::Static].FinalizeChunks(Builder, PrevItemChunksOffset, PrevItemChunksOffset + PrevNumStaticItemChunks, NumToRemove);
+				check(NumToRemove == 0);
 			}
 
 			// Insert retained chunk info first.
@@ -1792,7 +1792,7 @@ public:
 
 #if SC_ENABLE_DETAILED_LOGGING
 		BUILDER_LOG_LIST("CellIndexCacheEntry(%d):", CellIndexCacheEntry.Items.Num());
-		for (int32 ItemIndex = NumInstances; ItemIndex < CellIndexCacheEntry.Items.Num(); ++ItemIndex)
+		for (int32 ItemIndex = 0; ItemIndex < CellIndexCacheEntry.Items.Num(); ++ItemIndex)
 		{
 			FCellIndexCacheEntry::FItem Item = CellIndexCacheEntry.LoadAndStepItem(ItemIndex);
 			BUILDER_LOG_LIST_APPEND("(%d, %d)", Item.CellIndex, Item.NumInstances);
@@ -2177,7 +2177,7 @@ public:
 			FCellIndexCacheEntry &CellIndexCacheEntry = GetCacheEntry(PrevPrimitiveState.Payload);
 
 			RemovedInstanceFlags.SetRange(PrevPrimitiveState.InstanceDataOffset, PrevPrimitiveState.NumInstances, true);
-			for (int32 ItemIndex = NumInstances; ItemIndex < CellIndexCacheEntry.Items.Num(); ++ItemIndex)
+			for (int32 ItemIndex = 0; ItemIndex < CellIndexCacheEntry.Items.Num(); ++ItemIndex)
 			{
 				FCellIndexCacheEntry::FItem Item = CellIndexCacheEntry.LoadAndStepItem(ItemIndex);
 				MarkCellForRemove(Item.CellIndex, Item.NumInstances, PrevPrimitiveState.bDynamic ? EUpdateFrequencyCategory::Dynamic : EUpdateFrequencyCategory::Static);
