Mesh Edges rendering & gizmo: fixed view location/rotaion not being copied

#rnx

diff --git a/Engine/Source/Runtime/Renderer/Private/MeshEdgesRendering.cpp b/Engine/Source/Runtime/Renderer/Private/MeshEdgesRendering.cpp
--- a/Engine/Source/Runtime/Renderer/Private/MeshEdgesRendering.cpp
+++ b/Engine/Source/Runtime/Renderer/Private/MeshEdgesRendering.cpp
@@ -248,14 +248,20 @@ void CopyViewFamily(const FSceneViewFamily& SrcViewFamily, FSceneViewFamily& Vie
 
 	for (int32 ViewIndex = 0; ViewIndex < SrcViewFamily.Views.Num(); ++ViewIndex)
 	{
-		FSceneViewInitOptions ViewInitOptions = SrcViewFamily.Views[ViewIndex]->SceneViewInitOptions;
-		ViewInitOptions.ViewFamily = &ViewFamily;
+		const FSceneView* SrcSceneView = SrcViewFamily.Views[ViewIndex];
+		if (ensure(SrcSceneView))
+		{
+			FSceneViewInitOptions ViewInitOptions = SrcSceneView->SceneViewInitOptions;
+			ViewInitOptions.ViewFamily = &ViewFamily;
+			ViewInitOptions.ViewLocation = SrcSceneView->ViewLocation;
+			ViewInitOptions.ViewRotation = SrcSceneView->ViewRotation;
 
-		// Reset to avoid incorrect culling problems
-		ViewInitOptions.SceneViewStateInterface = FSceneViewInitOptions{}.SceneViewStateInterface;
+			// Reset to avoid incorrect culling problems
+			ViewInitOptions.SceneViewStateInterface = FSceneViewInitOptions{}.SceneViewStateInterface;
 		
-		FSceneView* View = new FSceneView(ViewInitOptions);
-		ViewFamily.Views.Emplace(View);
+			FSceneView* View = new FSceneView(ViewInitOptions);
+			ViewFamily.Views.Emplace(View);
+		}
 	}
 }
 
