[HWRT] Export FRayTracingLightGrid with RENDERER_API like FRaytracingLightDataPacked was.

#rnx

diff --git a/Engine/Source/Runtime/Renderer/Private/RayTracing/RayTracingLighting.h b/Engine/Source/Runtime/Renderer/Private/RayTracing/RayTracingLighting.h
--- a/Engine/Source/Runtime/Renderer/Private/RayTracing/RayTracingLighting.h
+++ b/Engine/Source/Runtime/Renderer/Private/RayTracing/RayTracingLighting.h
@@ -16,7 +16,7 @@
 FRHIRayTracingShader* GetRayTracingLightingMissShader(const FGlobalShaderMap* ShaderMap);
 
 // This struct holds a light grid and list of raytracing lights for both building and rendering
-BEGIN_GLOBAL_SHADER_PARAMETER_STRUCT(FRayTracingLightGrid, )
+BEGIN_GLOBAL_SHADER_PARAMETER_STRUCT(FRayTracingLightGrid, RENDERER_API)
 	SHADER_PARAMETER(uint32, SceneLightCount)
 	SHADER_PARAMETER(uint32, SceneInfiniteLightCount)
 	SHADER_PARAMETER(FVector3f, SceneLightsTranslatedBoundMin)
