Custom Render Pass:  Fix for culling not working properly, due to Custom Render Pass views not being added to the FInstanceCullingManager.

#rnx

diff --git a/Engine/Source/Runtime/Renderer/Private/SceneVisibility.cpp b/Engine/Source/Runtime/Renderer/Private/SceneVisibility.cpp
--- a/Engine/Source/Runtime/Renderer/Private/SceneVisibility.cpp
+++ b/Engine/Source/Runtime/Renderer/Private/SceneVisibility.cpp
@@ -5667,15 +5667,15 @@ void FDeferredShadingSceneRenderer::BeginInitViews(
 	}
 
 	// Create GPU-side representation of the view for instance culling.
-	InstanceCullingManager.AllocateViews(Views.Num());
-	for (FViewInfo& ViewInfo : Views)
+	InstanceCullingManager.AllocateViews(AllViews.Num());
+	for (FViewInfo* ViewInfo : AllViews)
 	{
-		ViewInfo.GPUSceneViewId = InstanceCullingManager.RegisterView(ViewInfo);
+		ViewInfo->GPUSceneViewId = InstanceCullingManager.RegisterView(*ViewInfo);
 
-		uint32 InstanceFactor = ViewInfo.bIsInstancedStereoEnabled && IStereoRendering::IsStereoEyeView(ViewInfo) && GEngine->StereoRenderingDevice.IsValid() ?
+		uint32 InstanceFactor = ViewInfo->bIsInstancedStereoEnabled && IStereoRendering::IsStereoEyeView(*ViewInfo) && GEngine->StereoRenderingDevice.IsValid() ?
 			GEngine->StereoRenderingDevice->GetDesiredNumberOfViews(true) : 1;
 
-		ViewInfo.InstanceFactor = InstanceFactor > 0 ? InstanceFactor : 1;
+		ViewInfo->InstanceFactor = InstanceFactor > 0 ? InstanceFactor : 1;
 	}
 
 	{
