Fix complexity debug view modes failing to properly display with an offset viewrect.

diff --git a/Engine/Shaders/Private/ShaderComplexityApplyPixelShader.usf b/Engine/Shaders/Private/ShaderComplexityApplyPixelShader.usf
--- a/Engine/Shaders/Private/ShaderComplexityApplyPixelShader.usf
+++ b/Engine/Shaders/Private/ShaderComplexityApplyPixelShader.usf
@@ -202,14 +202,18 @@ void Main(
 #endif
 	CenterViewportHDRColor.r *= ComplexityScale;
 
-	const int BorderX = 64;
-
 	// Legend
 	{
-		// left top of the border
+		const int OffsetFromLeft = 64;
+		const int OffsetFromBottom = 52;
 		const int Height = 18;
-		const int2 LeftTop = int2(BorderX, Input_ViewportMax.y - Height - 52);
-		const int2 Size = int2(Input_ViewportMax.x - Input_ViewportMin.x - BorderX * 2, Height);
+
+		const int2 Size = int2(Input_ViewportMax.x - Input_ViewportMin.x - OffsetFromLeft * 2, Height);
+
+		const int2 LeftTop = int2(
+			Input_ViewportMin.x + OffsetFromLeft,
+			Input_ViewportMax.y - OffsetFromBottom - Size.y);
+
 		const int OuterBorder = 1;
 
 		// (0, 0) .. (1, 1)
diff --git a/Engine/Source/Runtime/Renderer/Private/PostProcess/PostProcessVisualizeComplexity.cpp b/Engine/Source/Runtime/Renderer/Private/PostProcess/PostProcessVisualizeComplexity.cpp
--- a/Engine/Source/Runtime/Renderer/Private/PostProcess/PostProcessVisualizeComplexity.cpp
+++ b/Engine/Source/Runtime/Renderer/Private/PostProcess/PostProcessVisualizeComplexity.cpp
@@ -166,38 +166,41 @@ FScreenPassTexture AddVisualizeComplexityPass(FRDGBuilder& GraphBuilder, const F
 			Canvas.DrawShadowedString(X / DPIScale, Y / DPIScale, Text, GetStatsFont(), Color);
 		};
 
+		const FIntPoint CanvasMin(0, 0);
+		const FIntPoint CanvasMax = Output.ViewRect.Max - Output.ViewRect.Min;
+
 		if (View.Family->GetDebugViewShaderMode() == DVSM_QuadComplexity)
 		{
-			int32 StartX = Output.ViewRect.Min.X + 62;
-			int32 EndX = Output.ViewRect.Max.X - 66;
+			int32 StartX = CanvasMin.X + 62;
+			int32 EndX = CanvasMax.X - 66;
 			int32 NumOffset = (EndX - StartX) / (ShaderComplexityColorCount - 1);
 			for (int32 PosX = StartX, Number = 0; PosX <= EndX; PosX += NumOffset, ++Number)
 			{
 				FString Line;
 				Line = FString::Printf(TEXT("%d"), Number);
-				DrawString(PosX, Output.ViewRect.Max.Y - 87, *Line);
+				DrawString(PosX, CanvasMax.Y - 87, *Line);
 			}
 		}
 		else
 		{
-			DrawString(Output.ViewRect.Min.X + 63, Output.ViewRect.Max.Y - 51, TEXT("Good"));
-			DrawString(Output.ViewRect.Min.X + 63 + (int32)(Output.ViewRect.Width() * 107.0f / 397.0f), Output.ViewRect.Max.Y - 51, TEXT("Bad"));
-			DrawString(Output.ViewRect.Max.X - 170, Output.ViewRect.Max.Y - 51, TEXT("Extremely bad"));
+			DrawString(CanvasMin.X + 63, CanvasMax.Y - 51, TEXT("Good"));
+			DrawString(CanvasMin.X + 63 + (int32)(Output.ViewRect.Width() * 107.0f / 397.0f), CanvasMax.Y - 51, TEXT("Bad"));
+			DrawString(CanvasMax.X - 170, CanvasMax.Y - 51, TEXT("Extremely bad"));
 
-			DrawString(Output.ViewRect.Min.X + 62, Output.ViewRect.Max.Y - 87, TEXT("0"));
+			DrawString(CanvasMin.X + 62, CanvasMax.Y - 87, TEXT("0"));
 
 			if (View.Family->GetDebugViewShaderMode() == DVSM_LWCComplexity)
 			{
 #if WITH_DEBUG_VIEW_MODES
 				extern float GMaxLWCComplexity;
 				FString Line = FString::Printf(TEXT("r.ShaderComplexity.MaxLWCComplexity=%d"), (int32)GMaxLWCComplexity);
-				DrawString(Output.ViewRect.Max.X - 430, Output.ViewRect.Max.Y - 88, *Line);
+				DrawString(CanvasMax.X - 430, CanvasMax.Y - 88, *Line);
 #endif
 			}
 			else
 			{
 				FString Line = FString::Printf(TEXT("MaxShaderComplexityCount=%d"), (int32)GetMaxShaderComplexityCount(View.GetFeatureLevel()));
-				DrawString(Output.ViewRect.Max.X - 330, Output.ViewRect.Max.Y - 88, *Line);
+				DrawString(CanvasMax.X - 330, CanvasMax.Y - 88, *Line);
 			}
 		}
 	});
