Added some ScreenPass variations for consistency.

diff --git a/Engine/Source/Runtime/Renderer/Private/ScreenPass.cpp b/Engine/Source/Runtime/Renderer/Private/ScreenPass.cpp
--- a/Engine/Source/Runtime/Renderer/Private/ScreenPass.cpp
+++ b/Engine/Source/Runtime/Renderer/Private/ScreenPass.cpp
@@ -110,7 +110,7 @@ FScreenPassTexture FScreenPassTexture::CopyFromSlice(FRDGBuilder& GraphBuilder,
 // static
 FScreenPassTextureSlice FScreenPassTextureSlice::CreateFromScreenPassTexture(FRDGBuilder& GraphBuilder, const FScreenPassTexture& ScreenTexture)
 {
-	if (!ScreenTexture.Texture)
+	if (!ScreenTexture.Texture || !EnumHasAnyFlags(ScreenTexture.Texture->Desc.Flags, ETextureCreateFlags::ShaderResource))
 	{
 		return FScreenPassTextureSlice(nullptr, ScreenTexture.ViewRect);
 	}
@@ -132,6 +132,24 @@ FScreenPassRenderTarget FScreenPassRenderTarget::CreateFromInput(
 	return FScreenPassRenderTarget(GraphBuilder.CreateTexture(OutputDesc, OutputName), Input.ViewRect, OutputLoadAction);
 }
 
+FScreenPassRenderTarget FScreenPassRenderTarget::CreateFromInput(
+	FRDGBuilder& GraphBuilder,
+	FRDGTexture* InputTexture,
+	FIntPoint Extent,
+	ERenderTargetLoadAction OutputLoadAction,
+	const TCHAR* OutputName)
+{
+	check(InputTexture);
+
+	FRDGTextureDesc OutputDesc = InputTexture->Desc;
+	OutputDesc.Reset();
+	EnumRemoveFlags(OutputDesc.Flags, ETextureCreateFlags::Presentable);
+	OutputDesc.Extent = Extent;
+
+	return FScreenPassRenderTarget(GraphBuilder.CreateTexture(OutputDesc, OutputName), OutputLoadAction);
+}
+
+
 FScreenPassRenderTarget FScreenPassRenderTarget::CreateViewFamilyOutput(FRDGTextureRef ViewFamilyTexture, const FViewInfo& View)
 {
 	if (!ViewFamilyTexture)
@@ -267,7 +285,7 @@ void DrawScreenPass_PostSetup(
 
 void AddDrawTexturePass(
 	FRDGBuilder& GraphBuilder,
-	const FSceneView& View,
+	FScreenPassViewInfo ViewInfo,
 	FRDGTextureRef InputTexture,
 	FRDGTextureRef OutputTexture,
 	FIntPoint InputPosition,
@@ -298,19 +316,19 @@ void AddDrawTexturePass(
 	const FScreenPassTextureViewport InputViewport(InputDesc.Extent, FIntRect(InputPosition, InputPosition + InputSize));
 	const FScreenPassTextureViewport OutputViewport(OutputDesc.Extent, FIntRect(OutputPosition, OutputPosition + OutputSize));
 
-	TShaderMapRef<FCopyRectPS> PixelShader(static_cast<const FViewInfo&>(View).ShaderMap);
+	TShaderMapRef<FCopyRectPS> PixelShader(GetGlobalShaderMap(ViewInfo.FeatureLevel));
 
 	FCopyRectPS::FParameters* Parameters = GraphBuilder.AllocParameters<FCopyRectPS::FParameters>();
 	Parameters->InputTexture = InputTexture;
 	Parameters->InputSampler = TStaticSamplerState<>::GetRHI();
 	Parameters->RenderTargets[0] = FRenderTargetBinding(OutputTexture, LoadAction);
 
-	AddDrawScreenPass(GraphBuilder, RDG_EVENT_NAME("DrawTexture"), View, OutputViewport, InputViewport, PixelShader, Parameters);
+	AddDrawScreenPass(GraphBuilder, RDG_EVENT_NAME("DrawTexture"), ViewInfo, OutputViewport, InputViewport, PixelShader, Parameters);
 }
 
 void AddDrawTexturePass(
 	FRDGBuilder& GraphBuilder,
-	const FSceneView& View,
+	FScreenPassViewInfo ViewInfo,
 	FRDGTextureRef InputTexture,
 	FRDGTextureRef OutputTexture,
 	FIntPoint InputPosition,
@@ -319,7 +337,7 @@ void AddDrawTexturePass(
 {
 	AddDrawTexturePass(
 		GraphBuilder,
-		View,
+		ViewInfo,
 		InputTexture,
 		OutputTexture,
 		InputPosition,
@@ -330,40 +348,40 @@ void AddDrawTexturePass(
 
 void AddDrawTexturePass(
 	FRDGBuilder& GraphBuilder,
-	const FSceneView& View,
+	FScreenPassViewInfo ViewInfo,
 	FScreenPassTexture Input,
 	FScreenPassRenderTarget Output)
 {
 	const FScreenPassTextureViewport InputViewport(Input);
 	const FScreenPassTextureViewport OutputViewport(Output);
 
-	TShaderMapRef<FCopyRectPS> PixelShader(static_cast<const FViewInfo&>(View).ShaderMap);
+	TShaderMapRef<FCopyRectPS> PixelShader(GetGlobalShaderMap(ViewInfo.FeatureLevel));
 
 	FCopyRectPS::FParameters* Parameters = GraphBuilder.AllocParameters<FCopyRectPS::FParameters>();
 	Parameters->InputTexture = Input.Texture;
 	Parameters->InputSampler = TStaticSamplerState<>::GetRHI();
 	Parameters->RenderTargets[0] = Output.GetRenderTargetBinding();
 
-	AddDrawScreenPass(GraphBuilder, RDG_EVENT_NAME("DrawTexture"), View, OutputViewport, InputViewport, PixelShader, Parameters);
+	AddDrawScreenPass(GraphBuilder, RDG_EVENT_NAME("DrawTexture"), ViewInfo, OutputViewport, InputViewport, PixelShader, Parameters);
 }
 
 void AddDrawTexturePass(
 	FRDGBuilder& GraphBuilder,
-	const FSceneView& View,
+	FScreenPassViewInfo ViewInfo,
 	FScreenPassTextureSlice Input,
 	FScreenPassRenderTarget Output)
 {
 	const FScreenPassTextureViewport InputViewport(Input);
 	const FScreenPassTextureViewport OutputViewport(Output);
 
-	TShaderMapRef<FCopyRectSrvPS> PixelShader(static_cast<const FViewInfo&>(View).ShaderMap);
+	TShaderMapRef<FCopyRectSrvPS> PixelShader(GetGlobalShaderMap(ViewInfo.FeatureLevel));
 
 	FCopyRectSrvPS::FParameters* Parameters = GraphBuilder.AllocParameters<FCopyRectSrvPS::FParameters>();
 	Parameters->InputTexture = Input.TextureSRV;
 	Parameters->InputSampler = TStaticSamplerState<>::GetRHI();
 	Parameters->RenderTargets[0] = Output.GetRenderTargetBinding();
 
-	AddDrawScreenPass(GraphBuilder, RDG_EVENT_NAME("DrawTexture"), View, OutputViewport, InputViewport, PixelShader, Parameters);
+	AddDrawScreenPass(GraphBuilder, RDG_EVENT_NAME("DrawTexture"), ViewInfo, OutputViewport, InputViewport, PixelShader, Parameters);
 }
 
 class FDownsampleDepthPS : public FGlobalShader
diff --git a/Engine/Source/Runtime/Renderer/Public/ScreenPass.h b/Engine/Source/Runtime/Renderer/Public/ScreenPass.h
--- a/Engine/Source/Runtime/Renderer/Public/ScreenPass.h
+++ b/Engine/Source/Runtime/Renderer/Public/ScreenPass.h
@@ -87,6 +87,13 @@ struct FScreenPassRenderTarget : public FScreenPassTexture
 		ERenderTargetLoadAction OutputLoadAction,
 		const TCHAR* OutputName);
 
+	RENDERER_API static FScreenPassRenderTarget CreateFromInput(
+		FRDGBuilder& GraphBuilder,
+		FRDGTexture* InputTexture,
+		FIntPoint Extent,
+		ERenderTargetLoadAction OutputLoadAction,
+		const TCHAR* OutputName);
+
 	static FScreenPassRenderTarget CreateViewFamilyOutput(FRDGTextureRef ViewFamilyTexture, const FViewInfo& View);
 
 	FScreenPassRenderTarget() = default;
@@ -118,12 +125,6 @@ struct FScreenPassRenderTarget : public FScreenPassTexture
 class FScreenPassTextureViewport
 {
 public:
-	// Creates a viewport that is downscaled by an integer multiple.
-	static FScreenPassTextureViewport CreateDownscaled(const FScreenPassTextureViewport& Other, FIntPoint ScaleFactor);
-
-	// Creates a viewport scaled by a floating point multiplier.
-	static FScreenPassTextureViewport CreateScaled(const FScreenPassTextureViewport& Other, FVector2D Scale);
-
 	FScreenPassTextureViewport() = default;
 
 	explicit FScreenPassTextureViewport(FIntRect InRect)
@@ -374,7 +375,7 @@ struct FScreenPassPipelineState
 		const TShaderRef<FShader>& InPixelShader,
 		FRHIBlendState* InBlendState = FDefaultBlendState::GetRHI(),
 		FRHIDepthStencilState* InDepthStencilState = FDefaultDepthStencilState::GetRHI(),
-		uint32 InStencilRef = 0,
+		uint8 InStencilRef = 0,
 		FRHIVertexDeclaration* InVertexDeclaration = GFilterVertexDeclaration.VertexDeclarationRHI)
 		: VertexShader(InVertexShader)
 		, PixelShader(InPixelShader)
@@ -398,7 +399,7 @@ struct FScreenPassPipelineState
 	FRHIBlendState* BlendState = nullptr;
 	FRHIDepthStencilState* DepthStencilState = nullptr;
 	FRHIVertexDeclaration* VertexDeclaration = nullptr;
-	uint32 StencilRef{};
+	uint8 StencilRef{};
 };
 
 // Helper function which sets the pipeline state object on the command list prior to invoking a screen pass.
@@ -416,19 +417,27 @@ ENUM_CLASS_FLAGS(EScreenPassDrawFlags);
 /** Type used to carry the limited amount of data we need from a FSceneView. */
 struct FScreenPassViewInfo
 {
+	const ERHIFeatureLevel::Type FeatureLevel;
 	const int32 StereoViewIndex;
 	const int32 InstanceCount;
 	const bool bHMDHiddenAreaMaskActive;
 
 	FScreenPassViewInfo()
-		: StereoViewIndex(INDEX_NONE)
+		: FScreenPassViewInfo(GMaxRHIFeatureLevel)
+	{
+	}
+
+	FScreenPassViewInfo(ERHIFeatureLevel::Type InFeatureLevel)
+		: FeatureLevel(InFeatureLevel)
+		, StereoViewIndex(INDEX_NONE)
 		, InstanceCount(1)
 		, bHMDHiddenAreaMaskActive(IsHMDHiddenAreaMaskActive())
 	{
 	}
 
 	FScreenPassViewInfo(const FSceneView& View)
-		: StereoViewIndex(View.StereoViewIndex)
+		: FeatureLevel(View.FeatureLevel)
+		, StereoViewIndex(View.StereoViewIndex)
 		, InstanceCount(View.GetStereoPassInstanceFactor())
 		, bHMDHiddenAreaMaskActive(View.bHMDHiddenAreaMaskActive)
 	{
@@ -451,7 +460,7 @@ RENDERER_API void DrawScreenPass_PostSetup(
 template<typename TSetupFunction>
 void DrawScreenPass(
 	FRHICommandList& RHICmdList,
-	const FScreenPassViewInfo ViewInfo,
+	FScreenPassViewInfo ViewInfo,
 	const FScreenPassTextureViewport& OutputViewport,
 	const FScreenPassTextureViewport& InputViewport,
 	const FScreenPassPipelineState& PipelineState,
@@ -478,7 +487,7 @@ template <typename PixelShaderType>
 FORCEINLINE void AddDrawScreenPass(
 	FRDGBuilder& GraphBuilder,
 	FRDGEventName&& PassName,
-	const FScreenPassViewInfo ViewInfo,
+	FScreenPassViewInfo ViewInfo,
 	const FScreenPassTextureViewport& OutputViewport,
 	const FScreenPassTextureViewport& InputViewport,
 	const TShaderRef<FShader>& VertexShader,
@@ -513,7 +522,7 @@ template <typename PixelShaderType>
 FORCEINLINE void AddDrawScreenPass(
 	FRDGBuilder& GraphBuilder,
 	FRDGEventName&& PassName,
-	const FScreenPassViewInfo ViewInfo,
+	FScreenPassViewInfo ViewInfo,
 	const FScreenPassTextureViewport& OutputViewport,
 	const FScreenPassTextureViewport& InputViewport,
 	const TShaderRef<FShader>& VertexShader,
@@ -530,7 +539,7 @@ template <typename PixelShaderType>
 FORCEINLINE void AddDrawScreenPass(
 	FRDGBuilder& GraphBuilder,
 	FRDGEventName&& PassName,
-	const FScreenPassViewInfo ViewInfo,
+	FScreenPassViewInfo ViewInfo,
 	const FScreenPassTextureViewport& OutputViewport,
 	const FScreenPassTextureViewport& InputViewport,
 	const TShaderRef<FShader>& VertexShader,
@@ -547,7 +556,7 @@ template <typename PixelShaderType>
 FORCEINLINE void AddDrawScreenPass(
 	FRDGBuilder& GraphBuilder,
 	FRDGEventName&& PassName,
-	const FScreenPassViewInfo ViewInfo,
+	FScreenPassViewInfo ViewInfo,
 	const FScreenPassTextureViewport& OutputViewport,
 	const FScreenPassTextureViewport& InputViewport,
 	const TShaderRef<FShader>& VertexShader,
@@ -564,36 +573,18 @@ template <typename PixelShaderType>
 FORCEINLINE void AddDrawScreenPass(
 	FRDGBuilder& GraphBuilder,
 	FRDGEventName&& PassName,
-	const FSceneView& View,
+	FScreenPassViewInfo ViewInfo,
 	const FScreenPassTextureViewport& OutputViewport,
 	const FScreenPassTextureViewport& InputViewport,
 	const TShaderRef<PixelShaderType>& PixelShader,
 	typename PixelShaderType::FParameters* PixelShaderParameters,
 	EScreenPassDrawFlags Flags = EScreenPassDrawFlags::None)
 {
-	FGlobalShaderMap* ShaderMap = GetGlobalShaderMap(View.GetFeatureLevel());
+	FGlobalShaderMap* ShaderMap = GetGlobalShaderMap(ViewInfo.FeatureLevel);
 	TShaderMapRef<FScreenPassVS> VertexShader(ShaderMap);
 	FRHIBlendState* BlendState = FScreenPassPipelineState::FDefaultBlendState::GetRHI();
 	FRHIDepthStencilState* DepthStencilState = FScreenPassPipelineState::FDefaultDepthStencilState::GetRHI();
-	AddDrawScreenPass(GraphBuilder, Forward<FRDGEventName&&>(PassName), View, OutputViewport, InputViewport, VertexShader, PixelShader, BlendState, DepthStencilState, PixelShaderParameters, Flags);
-}
-
-template <typename PixelShaderType>
-FORCEINLINE void AddDrawScreenPass(
-	FRDGBuilder& GraphBuilder,
-	FRDGEventName&& PassName,
-	ERHIFeatureLevel::Type FeatureLevel,
-	const FScreenPassTextureViewport& OutputViewport,
-	const FScreenPassTextureViewport& InputViewport,
-	const TShaderRef<PixelShaderType>& PixelShader,
-	typename PixelShaderType::FParameters* PixelShaderParameters,
-	EScreenPassDrawFlags Flags = EScreenPassDrawFlags::None)
-{
-	FGlobalShaderMap* ShaderMap = GetGlobalShaderMap(FeatureLevel);
-	TShaderMapRef<FScreenPassVS> VertexShader(ShaderMap);
-	FRHIBlendState* BlendState = FScreenPassPipelineState::FDefaultBlendState::GetRHI();
-	FRHIDepthStencilState* DepthStencilState = FScreenPassPipelineState::FDefaultDepthStencilState::GetRHI();
-	AddDrawScreenPass(GraphBuilder, Forward<FRDGEventName&&>(PassName), FScreenPassViewInfo(), OutputViewport, InputViewport, VertexShader, PixelShader, BlendState, DepthStencilState, PixelShaderParameters, Flags);
+	AddDrawScreenPass(GraphBuilder, Forward<FRDGEventName&&>(PassName), ViewInfo, OutputViewport, InputViewport, VertexShader, PixelShader, BlendState, DepthStencilState, PixelShaderParameters, Flags);
 }
 
 /** Render graph variant of more advanced DrawScreenPass function. Does *not* clear unused graph
@@ -603,7 +594,7 @@ template <typename TSetupFunction, typename TPassParameterStruct>
 FORCEINLINE void AddDrawScreenPass(
 	FRDGBuilder& GraphBuilder,
 	FRDGEventName&& PassName,
-	const FScreenPassViewInfo ViewInfo,
+	FScreenPassViewInfo ViewInfo,
 	const FScreenPassTextureViewport& OutputViewport,
 	const FScreenPassTextureViewport& InputViewport,
 	const FScreenPassPipelineState& PipelineState,
@@ -628,7 +619,7 @@ template <typename TSetupFunction, typename TPassParameterStruct>
 FORCEINLINE void AddDrawScreenPass(
 	FRDGBuilder& GraphBuilder,
 	FRDGEventName&& PassName,
-	const FScreenPassViewInfo ViewInfo,
+	FScreenPassViewInfo ViewInfo,
 	const FScreenPassTextureViewport& OutputViewport,
 	const FScreenPassTextureViewport& InputViewport,
 	const FScreenPassPipelineState& PipelineState,
@@ -645,7 +636,7 @@ FORCEINLINE void AddDrawScreenPass(
  */
 void RENDERER_API AddDrawTexturePass(
 	FRDGBuilder& GraphBuilder,
-	const FSceneView& View,
+	FScreenPassViewInfo ViewInfo,
 	FRDGTextureRef InputTexture,
 	FRDGTextureRef OutputTexture,
 	FIntPoint InputPosition = FIntPoint::ZeroValue,
@@ -654,7 +645,7 @@ void RENDERER_API AddDrawTexturePass(
 
 void RENDERER_API AddDrawTexturePass(
 	FRDGBuilder& GraphBuilder,
-	const FSceneView& View,
+	FScreenPassViewInfo ViewInfo,
 	FRDGTextureRef InputTexture,
 	FRDGTextureRef OutputTexture,
 	FIntPoint InputPosition,
@@ -665,24 +656,38 @@ void RENDERER_API AddDrawTexturePass(
 /** Helper variant which takes a shared viewport instead of unique input / output positions. */
 FORCEINLINE void AddDrawTexturePass(
 	FRDGBuilder& GraphBuilder,
-	const FSceneView& View,
+	FScreenPassViewInfo ViewInfo,
 	FRDGTextureRef InputTexture,
 	FRDGTextureRef OutputTexture,
 	FIntRect ViewportRect)
 {
-	AddDrawTexturePass(GraphBuilder, View, InputTexture, OutputTexture, ViewportRect.Min, ViewportRect.Min, ViewportRect.Size());
+	AddDrawTexturePass(GraphBuilder, ViewInfo, InputTexture, OutputTexture, ViewportRect.Min, ViewportRect.Min, ViewportRect.Size());
 }
 
 void RENDERER_API AddDrawTexturePass(
 	FRDGBuilder& GraphBuilder,
-	const FSceneView& View,
+	FScreenPassViewInfo ViewInfo,
 	FScreenPassTexture Input,
 	FScreenPassRenderTarget Output);
 
+inline void AddDrawTexturePass(
+	FRDGBuilder& GraphBuilder,
+	FScreenPassViewInfo ViewInfo,
+	FScreenPassTexture Input,
+	FScreenPassTexture Output)
+{
+	ERenderTargetLoadAction LoadAction = ERenderTargetLoadAction::ELoad;
+	if (Input.Texture->Desc.Extent == Output.Texture->Desc.Extent && Input.ViewRect == Output.ViewRect)
+	{
+		LoadAction = ERenderTargetLoadAction::ENoAction;
+	}
+	AddDrawTexturePass(GraphBuilder, ViewInfo, Input, FScreenPassRenderTarget(Output, LoadAction));
+}
+
 /** Variation of above that accepts a slice, which contains a texture SRV instead of a texture */
 void RENDERER_API AddDrawTexturePass(
 	FRDGBuilder& GraphBuilder,
-	const FSceneView& View,
+	FScreenPassViewInfo ViewInfo,
 	FScreenPassTextureSlice Input,
 	FScreenPassRenderTarget Output);
 
@@ -733,7 +738,7 @@ enum class EDownsampleDepthFilter
 	MinAndMaxDepth
 };
 
-void AddDownsampleDepthPass(
+RENDERER_API void AddDownsampleDepthPass(
 	FRDGBuilder& GraphBuilder,
 	const FViewInfo& View,
 	FScreenPassTexture Input,
diff --git a/Engine/Source/Runtime/Renderer/Public/ScreenPass.inl b/Engine/Source/Runtime/Renderer/Public/ScreenPass.inl
--- a/Engine/Source/Runtime/Renderer/Public/ScreenPass.inl
+++ b/Engine/Source/Runtime/Renderer/Public/ScreenPass.inl
@@ -99,7 +99,7 @@ inline FScreenPassTexture::FScreenPassTexture(FRDGTextureRef InTexture, FIntRect
 }
 
 inline FScreenPassTexture::FScreenPassTexture(const FScreenPassTextureSlice& ScreenTexture)
-	: Texture(ScreenTexture.TextureSRV->Desc.Texture)
+	: Texture(ScreenTexture.TextureSRV ? ScreenTexture.TextureSRV->Desc.Texture : nullptr)
 	, ViewRect(ScreenTexture.ViewRect)
 {
 	if (Texture && Texture->Desc.IsTextureArray())
