Removed redundant bIsHDR check which is currently always false.

diff --git a/Engine/Source/Runtime/Renderer/Private/PostProcess/PostProcessDeviceEncodingOnly.cpp b/Engine/Source/Runtime/Renderer/Private/PostProcess/PostProcessDeviceEncodingOnly.cpp
--- a/Engine/Source/Runtime/Renderer/Private/PostProcess/PostProcessDeviceEncodingOnly.cpp
+++ b/Engine/Source/Runtime/Renderer/Private/PostProcess/PostProcessDeviceEncodingOnly.cpp
@@ -48,10 +48,6 @@ FDeviceEncodingOnlyOutputDeviceParameters GetDeviceEncodingOnlyOutputDeviceParam
 	{
 		OutputDeviceValue = EDisplayOutputFormat::HDR_LinearWithToneCurve;
 	}
-	else if (Family.bIsHDR)
-	{
-		OutputDeviceValue = EDisplayOutputFormat::HDR_ACES_1000nit_ST2084;
-	}
 	else
 	{
 		OutputDeviceValue = Family.RenderTarget->GetDisplayOutputFormat();
@@ -105,7 +101,6 @@ BEGIN_SHADER_PARAMETER_STRUCT(FDeviceEncodingOnlyParameters, )
 	SHADER_PARAMETER_RDG_TEXTURE(Texture2D, ColorTexture)
 	SHADER_PARAMETER_SAMPLER(SamplerState, ColorSampler)
 	SHADER_PARAMETER(float, EditorNITLevel)
-	SHADER_PARAMETER(uint32, bOutputInHDR)
 END_SHADER_PARAMETER_STRUCT()
 
 class FDeviceEncodingOnlyPS : public FGlobalShader
@@ -227,7 +222,6 @@ FScreenPassTexture AddDeviceEncodingOnlyPass(FRDGBuilder& GraphBuilder, const FV
 	CommonParameters.ColorTexture = Inputs.SceneColor.Texture;
 	CommonParameters.ColorSampler = BilinearClampSampler;
 	CommonParameters.EditorNITLevel = EditorNITLevel;
-	CommonParameters.bOutputInHDR = ViewFamily.bIsHDR;
 
 	// Generate permutation vector for the desktop tonemapper.
 	DeviceEncodingOnlyPermutation::FDesktopDomain DesktopPermutationVector;
diff --git a/Engine/Source/Runtime/Renderer/Private/PostProcess/PostProcessTonemap.cpp b/Engine/Source/Runtime/Renderer/Private/PostProcess/PostProcessTonemap.cpp
--- a/Engine/Source/Runtime/Renderer/Private/PostProcess/PostProcessTonemap.cpp
+++ b/Engine/Source/Runtime/Renderer/Private/PostProcess/PostProcessTonemap.cpp
@@ -258,10 +258,6 @@ FTonemapperOutputDeviceParameters GetTonemapperOutputDeviceParameters(const FSce
 	{
 		OutputDeviceValue = EDisplayOutputFormat::HDR_LinearWithToneCurve;
 	}
-	else if (Family.bIsHDR)
-	{
-		OutputDeviceValue = EDisplayOutputFormat::HDR_ACES_1000nit_ST2084;
-	}
 	else
 	{
 		OutputDeviceValue = Family.RenderTarget->GetDisplayOutputFormat();
@@ -269,7 +265,7 @@ FTonemapperOutputDeviceParameters GetTonemapperOutputDeviceParameters(const FSce
 
 	float Gamma = CVarOutputGamma->GetValueOnRenderThread();
 
-    // In case gamma is unspecified, fall back to 2.2 which is the most common case
+	// In case gamma is unspecified, fall back to 2.2 which is the most common case
 	if ((PLATFORM_APPLE || OutputDeviceValue == EDisplayOutputFormat::SDR_ExplicitGammaMapping) && Gamma == 0.0f)
 	{
 		Gamma = 2.2f;
@@ -351,7 +347,6 @@ BEGIN_SHADER_PARAMETER_STRUCT(FTonemapParameters, )
 	SHADER_PARAMETER(float, LUTOffset)
 	SHADER_PARAMETER(float, EditorNITLevel)
 	SHADER_PARAMETER(float, BackbufferQuantizationDithering)
-	SHADER_PARAMETER(uint32, bOutputInHDR)
 END_SHADER_PARAMETER_STRUCT()
 
 class FFilmGrainReduceCS : public FGlobalShader
@@ -613,7 +608,7 @@ FScreenPassTexture AddTonemapPass(FRDGBuilder& GraphBuilder, const FViewInfo& Vi
 			Inputs.SceneColor.TextureSRV->Desc.Texture->Desc.Format,
 			FClearValueBinding(FLinearColor(0, 0, 0, 0)),
 			GFastVRamConfig.Tonemap | TexCreate_ShaderResource | TexCreate_RenderTargetable | (View.bUseComputePasses ? TexCreate_UAV : TexCreate_None));;
-		
+
 		const FTonemapperOutputDeviceParameters OutputDeviceParameters = GetTonemapperOutputDeviceParameters(*View.Family);
 		const EDisplayOutputFormat OutputDevice = static_cast<EDisplayOutputFormat>(OutputDeviceParameters.OutputDevice);
 		const bool bPostProcessingAlpha = IsPostProcessingWithAlphaChannelSupported();
@@ -628,10 +623,6 @@ FScreenPassTexture AddTonemapPass(FRDGBuilder& GraphBuilder, const FViewInfo& Vi
 		{
 			OutputDesc.Format = PF_FloatRGBA;
 		}
-		else if (Inputs.bOutputInHDR)
-		{
-			OutputDesc.Format = GRHIHDRDisplayOutputFormat;
-		}
 		else if (View.Family->RenderTarget && View.Family->RenderTarget->GetRenderTargetTexture())
 		{
 			// Render into a pixel format that do not loose bit depth precision for the view family.
@@ -858,7 +849,6 @@ FScreenPassTexture AddTonemapPass(FRDGBuilder& GraphBuilder, const FViewInfo& Vi
 			}
 		}
 	}
-	CommonParameters.bOutputInHDR = ViewFamily.bIsHDR;
 	CommonParameters.LUTSize = LUTSize;
 	CommonParameters.InvLUTSize = 1.0f / LUTSize;
 	CommonParameters.LUTScale = (LUTSize - 1.0f) / LUTSize;
