Enable parallel execution of DrawTileMesh passes.
 - Use non-immediate command list for passes.
 - Defer deletion of render resources to destructor of render data (which is itself destroyed during RDG async deletion).

diff --git a/Engine/Source/Runtime/Engine/Private/TileRendering.cpp b/Engine/Source/Runtime/Engine/Private/TileRendering.cpp
--- a/Engine/Source/Runtime/Engine/Private/TileRendering.cpp
+++ b/Engine/Source/Runtime/Engine/Private/TileRendering.cpp
@@ -71,6 +71,11 @@ FCanvasTileRendererItem::FRenderData::FRenderData(
 	, VertexFactory(&StaticMeshVertexBuffers, InFeatureLevel)
 {}
 
+FCanvasTileRendererItem::FRenderData::~FRenderData()
+{
+	ReleaseTileMesh();
+}
+
 uint32 FCanvasTileRendererItem::FRenderData::GetNumVertices() const
 {
 	return Tiles.Num() * CanvasTileVertexCount;
@@ -200,11 +205,6 @@ void FCanvasTileRendererItem::FRenderData::RenderTiles(
 	// Flush the final batch: 
 	check(CurrentMeshBatch != nullptr);
 	GetRendererModule().DrawTileMesh(RenderContext, DrawRenderState, View, *CurrentMeshBatch, bIsHitTesting, CurrentMeshBatch->BatchHitProxyId, bUse128bitRT);
-
-	AddPass(RenderContext.GraphBuilder, RDG_EVENT_NAME("ReleaseTileMesh"), [this](FRHICommandListImmediate&)
-	{
-		ReleaseTileMesh();
-	});
 }
 
 bool FCanvasTileRendererItem::Render_RenderThread(FCanvasRenderContext& RenderContext, FMeshPassProcessorRenderState& DrawRenderState, const FCanvas* Canvas)
diff --git a/Engine/Source/Runtime/Engine/Private/TriangleRendering.cpp b/Engine/Source/Runtime/Engine/Private/TriangleRendering.cpp
--- a/Engine/Source/Runtime/Engine/Private/TriangleRendering.cpp
+++ b/Engine/Source/Runtime/Engine/Private/TriangleRendering.cpp
@@ -36,6 +36,11 @@ void FCanvasTriangleRendererItem::FTriangleVertexFactory::InitResource(FRHIComma
 	FLocalVertexFactory::InitResource(RHICmdList);
 }
 
+FCanvasTriangleRendererItem::FRenderData::~FRenderData()
+{
+	ReleaseTriangleMesh();
+}
+
 FMeshBatch* FCanvasTriangleRendererItem::FRenderData::AllocTriangleMeshBatch(FCanvasRenderContext& InRenderContext, FHitProxyId InHitProxyId)
 {
 	FMeshBatch* MeshBatch = InRenderContext.Alloc<FMeshBatch>();
@@ -176,11 +181,6 @@ void FCanvasTriangleRendererItem::FRenderData::RenderTriangles(
 	// Flush the final batch: 
 	check(CurrentMeshBatch != nullptr);
 	GetRendererModule().DrawTileMesh(RenderContext, DrawRenderState, View, *CurrentMeshBatch, bIsHitTesting, CurrentMeshBatch->BatchHitProxyId);
-
-	AddPass(RenderContext.GraphBuilder, RDG_EVENT_NAME("ReleaseTriangleMesh"), [this](FRHICommandListImmediate&)
-	{
-		ReleaseTriangleMesh();
-	});
 }
 
 bool FCanvasTriangleRendererItem::Render_RenderThread(FCanvasRenderContext& RenderContext, FMeshPassProcessorRenderState& DrawRenderState, const FCanvas* Canvas)
diff --git a/Engine/Source/Runtime/Engine/Private/UserInterface/Canvas.cpp b/Engine/Source/Runtime/Engine/Private/UserInterface/Canvas.cpp
--- a/Engine/Source/Runtime/Engine/Private/UserInterface/Canvas.cpp
+++ b/Engine/Source/Runtime/Engine/Private/UserInterface/Canvas.cpp
@@ -433,7 +433,7 @@ bool FCanvasBatchedElementRenderItem::Render_RenderThread(FCanvasRenderContext&
 
 		RenderContext.AddPass(
 			RDG_EVENT_NAME("CanvasBatchedElements"),
-			[LocalData = Data, DrawRenderState, Canvas](FRHICommandListImmediate& RHICmdList)
+			[LocalData = Data, DrawRenderState, Canvas](FRHICommandList& RHICmdList)
 		{
 			// current render target set for the canvas
 			const FRenderTarget* CanvasRenderTarget = Canvas->GetRenderTarget();
diff --git a/Engine/Source/Runtime/Engine/Public/CanvasRender.h b/Engine/Source/Runtime/Engine/Public/CanvasRender.h
--- a/Engine/Source/Runtime/Engine/Public/CanvasRender.h
+++ b/Engine/Source/Runtime/Engine/Public/CanvasRender.h
@@ -14,7 +14,7 @@ public:
 	void AddPass(FRDGEventName&& PassName, const ParameterStructType* PassParameters, ExecuteLambdaType&& ExecuteLambda)
 	{
 		GraphBuilder.AddPass(Forward<FRDGEventName>(PassName), PassParameters, ERDGPassFlags::Raster,
-			[LocalScissorRect = ScissorRect, LocalViewportRect = ViewportRect, LocalExecuteLambda = Forward<ExecuteLambdaType&&>(ExecuteLambda)](FRHICommandListImmediate& RHICmdList)
+			[LocalScissorRect = ScissorRect, LocalViewportRect = ViewportRect, LocalExecuteLambda = Forward<ExecuteLambdaType&&>(ExecuteLambda)](FRHICommandList& RHICmdList)
 		{
 			RHICmdList.SetViewport(static_cast<float>(LocalViewportRect.Min.X), static_cast<float>(LocalViewportRect.Min.Y), 0.0f, static_cast<float>(LocalViewportRect.Max.X), static_cast<float>(LocalViewportRect.Max.Y), 1.0f);
 
diff --git a/Engine/Source/Runtime/Engine/Public/CanvasRendererItem.h b/Engine/Source/Runtime/Engine/Public/CanvasRendererItem.h
--- a/Engine/Source/Runtime/Engine/Public/CanvasRendererItem.h
+++ b/Engine/Source/Runtime/Engine/Public/CanvasRendererItem.h
@@ -103,6 +103,8 @@ private:
 			const FMaterialRenderProxy* InMaterialRenderProxy,
 			const FCanvas::FTransformEntry& InTransform);
 
+		~FRenderData();
+
 		void RenderTiles(
 			FCanvasRenderContext& RenderContext,
 			FMeshPassProcessorRenderState& DrawRenderState,
@@ -264,6 +266,8 @@ private:
 			, VertexFactory(&StaticMeshVertexBuffers, InFeatureLevel)
 		{}
 
+		~FRenderData();
+
 		FORCEINLINE int32 AddTriangle(const FCanvasUVTri& Tri, FHitProxyId HitProxyId)
 		{
 			FTriangleInst NewTri = { Tri, HitProxyId };
diff --git a/Engine/Source/Runtime/Renderer/Private/Renderer.cpp b/Engine/Source/Runtime/Renderer/Private/Renderer.cpp
--- a/Engine/Source/Runtime/Renderer/Private/Renderer.cpp
+++ b/Engine/Source/Runtime/Renderer/Private/Renderer.cpp
@@ -415,7 +415,7 @@ void FRendererModule::DrawTileMesh(FCanvasRenderContext& RenderContext, FMeshPas
 				PassParameters->DebugViewMode = CreateDebugViewModePassUniformBuffer(GraphBuilder, View, nullptr);
 
 				RenderContext.AddPass(RDG_EVENT_NAME("OutputMaterialTextureScales"), PassParameters,
-					[Scene, &View, &Mesh](FRHICommandListImmediate& RHICmdList)
+					[Scene, &View, &Mesh](FRHICommandList& RHICmdList)
 				{
 					DrawDynamicMeshPass(View, RHICmdList, [&](FMeshPassDrawListContext* InDrawListContext)
 					{
@@ -439,7 +439,7 @@ void FRendererModule::DrawTileMesh(FCanvasRenderContext& RenderContext, FMeshPas
 				PassParameters->TranslucentBasePass = CreateTranslucentBasePassUniformBuffer(GraphBuilder, Scene, View);
 
 				RenderContext.AddPass(RDG_EVENT_NAME("TranslucentDeferred"), PassParameters,
-					[Scene, &View, &Mesh, DrawRenderState, bUse128bitRT](FRHICommandListImmediate& RHICmdList)
+					[Scene, &View, &Mesh, DrawRenderState, bUse128bitRT](FRHICommandList& RHICmdList)
 				{
 					DrawDynamicMeshPass(View, RHICmdList, [&](FDynamicPassMeshDrawListContext* DynamicMeshPassContext)
 					{
@@ -463,7 +463,7 @@ void FRendererModule::DrawTileMesh(FCanvasRenderContext& RenderContext, FMeshPas
 				PassParameters->MobileBasePass = CreateMobileBasePassUniformBuffer(GraphBuilder, View, EMobileBasePass::Translucent, EMobileSceneTextureSetupMode::None);
 
 				RenderContext.AddPass(RDG_EVENT_NAME("TranslucentMobile"), PassParameters,
-					[Scene, &View, DrawRenderState, &Mesh](FRHICommandListImmediate& RHICmdList)
+					[Scene, &View, DrawRenderState, &Mesh](FRHICommandList& RHICmdList)
 				{
 					DrawDynamicMeshPass(View, RHICmdList, [&](FDynamicPassMeshDrawListContext* DynamicMeshPassContext)
 					{
@@ -495,7 +495,7 @@ void FRendererModule::DrawTileMesh(FCanvasRenderContext& RenderContext, FMeshPas
 
 #if WITH_EDITOR
 				RenderContext.AddPass(RDG_EVENT_NAME("HitTesting"), PassParameters,
-					[Scene, &View, DrawRenderState, &Mesh](FRHICommandListImmediate& RHICmdList)
+					[Scene, &View, DrawRenderState, &Mesh](FRHICommandList& RHICmdList)
 				{
 					DrawDynamicMeshPass(View, RHICmdList, [&](FDynamicPassMeshDrawListContext* DynamicMeshPassContext)
 					{
@@ -519,7 +519,7 @@ void FRendererModule::DrawTileMesh(FCanvasRenderContext& RenderContext, FMeshPas
 					PassParameters->OpaqueBasePass = CreateOpaqueBasePassUniformBuffer(GraphBuilder, View);
 
 					RenderContext.AddPass(RDG_EVENT_NAME("OpaqueDeferred"), PassParameters,
-						[Scene, &View, DrawRenderState, &Mesh, bUse128bitRT](FRHICommandListImmediate& RHICmdList)
+						[Scene, &View, DrawRenderState, &Mesh, bUse128bitRT](FRHICommandList& RHICmdList)
 					{
 						DrawDynamicMeshPass(View, RHICmdList,
 							[&](FDynamicPassMeshDrawListContext* DynamicMeshPassContext)
@@ -543,7 +543,7 @@ void FRendererModule::DrawTileMesh(FCanvasRenderContext& RenderContext, FMeshPas
 					PassParameters->MobileBasePass = CreateMobileBasePassUniformBuffer(GraphBuilder, View, EMobileBasePass::Opaque, EMobileSceneTextureSetupMode::None);
 
 					RenderContext.AddPass(RDG_EVENT_NAME("OpaqueMobile"), PassParameters,
-						[Scene, &View, DrawRenderState, &Mesh](FRHICommandListImmediate& RHICmdList)
+						[Scene, &View, DrawRenderState, &Mesh](FRHICommandList& RHICmdList)
 					{
 						DrawDynamicMeshPass(View, RHICmdList, [&](FDynamicPassMeshDrawListContext* DynamicMeshPassContext)
 						{
