Add LumenHardwareRayTracingUniformBuffer to Raytracing ShaderBindingLayout

diff --git a/Engine/Source/Runtime/Renderer/Private/RayTracing/RayTracing.cpp b/Engine/Source/Runtime/Renderer/Private/RayTracing/RayTracing.cpp
--- a/Engine/Source/Runtime/Renderer/Private/RayTracing/RayTracing.cpp
+++ b/Engine/Source/Runtime/Renderer/Private/RayTracing/RayTracing.cpp
@@ -306,6 +306,7 @@ namespace RayTracing
 			StaticUniformBuffers.Add(FindUniformBufferStructByName(TEXT("Scene")));
 			StaticUniformBuffers.Add(FindUniformBufferStructByName(TEXT("View")));
 			StaticUniformBuffers.Add(FindUniformBufferStructByName(TEXT("NaniteRayTracing")));
+			StaticUniformBuffers.Add(FindUniformBufferStructByName(TEXT("LumenHardwareRayTracingUniformBuffer")));
 
 			BuildShaderBindingLayout(StaticUniformBuffers, ShaderBindingLayoutFlags, *this);
 		}
@@ -339,6 +340,7 @@ namespace RayTracing
 			StaticUniformBuffers.AddUniformBuffer(View.ViewUniformBuffer.GetReference());
 			StaticUniformBuffers.AddUniformBuffer(SceneUniformBuffer);
 			StaticUniformBuffers.AddUniformBuffer(Nanite::GRayTracingManager.GetUniformBuffer().GetReference());
+			StaticUniformBuffers.AddUniformBuffer(View.LumenHardwareRayTracingUniformBuffer->GetRHI());
 
 			StaticUniformBufferScope.Emplace(RHICmdList, StaticUniformBuffers);
 		}
