Light Function Atlas - unified material name on vis screen.

#rnx

diff --git a/Engine/Source/Runtime/Renderer/Private/LightFunctionAtlas.cpp b/Engine/Source/Runtime/Renderer/Private/LightFunctionAtlas.cpp
--- a/Engine/Source/Runtime/Renderer/Private/LightFunctionAtlas.cpp
+++ b/Engine/Source/Runtime/Renderer/Private/LightFunctionAtlas.cpp
@@ -955,7 +955,7 @@ FScreenPassTexture FLightFunctionAtlas::AddDebugVisualizationPasses(FRDGBuilder&
 
 			Canvas.DrawTile(DrawPosX - 20.0f, DrawPosY, 15.0f, 15.0f, 0.0f, 0.0f, 1.0f, 1.0f, MaterialColor, nullptr, false);
 
-			Text = FString::Printf(TEXT("lights - %s"), *NotCompatibleMaterial.Value->GetFullName());
+			Text = FString::Printf(TEXT("lights - %s"), *NotCompatibleMaterial.Value->GetRenderProxy()->GetMaterialName());
 			Canvas.DrawShadowedString(DrawPosX + 20, DrawPosY, *Text, GEngine->GetLargeFont(), MaterialColor);
 			DrawPosY += 20.0f;
 		}
