Light Function Atlas - fixed a case where light function would not be visible on a light when a LFMaterial was not compatible with the Atlas and "Cast Shadow" was disabled.
Clarified the meaning of bClusteredDeferredSupported.

diff --git a/Engine/Source/Runtime/Renderer/Private/LightRendering.cpp b/Engine/Source/Runtime/Renderer/Private/LightRendering.cpp
--- a/Engine/Source/Runtime/Renderer/Private/LightRendering.cpp
+++ b/Engine/Source/Runtime/Renderer/Private/LightRendering.cpp
@@ -1247,14 +1247,15 @@ void FSceneRenderer::GatherAndSortLights(FSortedLightSetSceneInfo& OutSortedLigh
 					// Lights handled by Many Lights
 					const bool bHandledByManyLights = ManyLights::IsLightSupported(LightSceneInfoCompact.LightType, LightSceneInfoCompact.CastRaytracedShadow);
 
-					// tiled and clustered deferred lighting only supported for certain lights that don't use any additional features
-					// And also that are not directional (mostly because it doesn't make so much sense to insert them into every grid cell in the universe)
+					// NOTE: bClusteredDeferredSupported==false means "lights cannot be batched" (tiled or clustered). When false, light will go the slower unbatched render path.
+					// Tiled and clustered deferred lighting only support certain lights that don't use any additional features (like shadow or light function not compatible with the atlas.)
+					// And also that are not directional (mostly because it doesn't make so much sense to insert them into every grid cell in the universe).
 					// In the forward case one directional light gets put into its own variables, and in the deferred case it gets a full-screen pass.
 					// Usually it'll have shadows and stuff anyway.
 					// Contact shadow are not supported.
 					const bool bClusteredDeferredSupported =
 						(!SortedLightInfo->SortKey.Fields.bShadowed || bShadowedLightsInClustered) &&
-						(!SortedLightInfo->SortKey.Fields.bLightFunction || (bUseLightFunctionAtlas))
+						(!SortedLightInfo->SortKey.Fields.bLightFunction || (bUseLightFunctionAtlas && SortedLightInfo->bIsCompatibleWithLightFunctionAtlas)) // If not compatible with the LightFunctionAtlas, light with light function materials must go the unbatched route.
 						&& LightSceneInfoCompact.LightType != LightType_Directional
 						&& LightSceneInfo->Proxy->GetContactShadowLength() == 0
 						&& !bHandledByManyLights;
