Scene updates refactor to make it easier to add & maintain light/primitive update types and propagate this to the scene extensions.

* Added TSceneUpdateCommandQueue which is used to queue scene update commands (add/remove/update<Payload>), update payload types are extensible by a statically compiled enum but can be extended with registration if need be.
* Changed scene primitive updates to use this instead of separate hash tables.
* Changed scene light updates to use the same instead of its bespoke queue.
* Exposed the updates in the FScene[Pre|Post]UpdateChangeSet such that all extensions can access any change.
* Refactored relevant scene extensions & scene update to use the new data.
* Sets us up for being able to pass around incremental instance updates to various scene modules (once the delta data is plumbed through, that is).
* Step towards batching the updates on the game thread, though this requires a larger refactor of how and when scene updates can be invoked.

diff --git a/Engine/Source/Runtime/Renderer/Private/GPUScene.cpp b/Engine/Source/Runtime/Renderer/Private/GPUScene.cpp
--- a/Engine/Source/Runtime/Renderer/Private/GPUScene.cpp
+++ b/Engine/Source/Runtime/Renderer/Private/GPUScene.cpp
@@ -2086,14 +2086,11 @@ bool FGPUScene::HasPendingGPUWrite(uint32 PrimitiveId) const
 
 void FGPUScene::OnPreSceneUpdate(FRDGBuilder& GraphBuilder, const FScenePreUpdateChangeSet& ScenePreUpdateData)
 {
-	for (FPersistentPrimitiveIndex PersistentPrimitiveIndex  : ScenePreUpdateData.RemovedPrimitiveIds)
+	ScenePreUpdateData.PrimitiveUpdates.ForEachUpdateCommand(ESceneUpdateCommandFilter::Updated | ESceneUpdateCommandFilter::Deleted, EPrimitiveUpdateDirtyFlags::GPUState, 
+	[&](const FPrimitiveUpdateCommand& Cmd)
 	{
-		AddPrimitiveToUpdate(PersistentPrimitiveIndex, EPrimitiveDirtyState::Removed);
-	}
-	for (FPersistentPrimitiveIndex PersistentPrimitiveIndex  : ScenePreUpdateData.UpdatedPrimitiveIds)
-	{
-		AddPrimitiveToUpdate(PersistentPrimitiveIndex, EPrimitiveDirtyState::ChangedTransform);
-	}
+		AddPrimitiveToUpdate(Cmd.GetPersistentId(), Cmd.IsDelete() ? EPrimitiveDirtyState::Removed : EPrimitiveDirtyState::ChangedAll );
+	});
 }
 
 void FGPUScene::OnPostSceneUpdate(FRDGBuilder& GraphBuilder, const FScenePostUpdateChangeSet& ScenePostUpdateData)
diff --git a/Engine/Source/Runtime/Renderer/Private/LightSceneInfo.cpp b/Engine/Source/Runtime/Renderer/Private/LightSceneInfo.cpp
--- a/Engine/Source/Runtime/Renderer/Private/LightSceneInfo.cpp
+++ b/Engine/Source/Runtime/Renderer/Private/LightSceneInfo.cpp
@@ -125,7 +125,7 @@ void FLightSceneInfo::AddToScene()
 
 bool FLightSceneInfo::ShouldCreateLightPrimitiveInteraction(const FLightSceneInfoCompact& LightSceneInfoCompact, const FPrimitiveSceneInfoCompact& PrimitiveSceneInfoCompact)
 {
-	if (!Scene->IsPrimitiveBeingRemoved(PrimitiveSceneInfoCompact.PrimitiveSceneInfo) && LightSceneInfoCompact.AffectsPrimitive(FBoxSphereBounds(PrimitiveSceneInfoCompact.Bounds), PrimitiveSceneInfoCompact.Proxy))
+	if (LightSceneInfoCompact.AffectsPrimitive(FBoxSphereBounds(PrimitiveSceneInfoCompact.Bounds), PrimitiveSceneInfoCompact.Proxy))
 	{
 		// create light interaction and add to light/primitive lists
 		return FLightPrimitiveInteraction::ShouldCreate(this, PrimitiveSceneInfoCompact.PrimitiveSceneInfo).bShouldCreate;
@@ -135,7 +135,7 @@ bool FLightSceneInfo::ShouldCreateLightPrimitiveInteraction(const FLightSceneInf
 
 void FLightSceneInfo::CreateLightPrimitiveInteraction(const FLightSceneInfoCompact& LightSceneInfoCompact, const FPrimitiveSceneInfoCompact& PrimitiveSceneInfoCompact)
 {
-	if (!Scene->IsPrimitiveBeingRemoved(PrimitiveSceneInfoCompact.PrimitiveSceneInfo) && LightSceneInfoCompact.AffectsPrimitive(FBoxSphereBounds(PrimitiveSceneInfoCompact.Bounds), PrimitiveSceneInfoCompact.Proxy))
+	if (LightSceneInfoCompact.AffectsPrimitive(FBoxSphereBounds(PrimitiveSceneInfoCompact.Bounds), PrimitiveSceneInfoCompact.Proxy))
 	{
 		// create light interaction and add to light/primitive lists
 		FLightPrimitiveInteraction::Create(this,PrimitiveSceneInfoCompact.PrimitiveSceneInfo);
diff --git a/Engine/Source/Runtime/Renderer/Private/LightSceneInfo.h b/Engine/Source/Runtime/Renderer/Private/LightSceneInfo.h
--- a/Engine/Source/Runtime/Renderer/Private/LightSceneInfo.h
+++ b/Engine/Source/Runtime/Renderer/Private/LightSceneInfo.h
@@ -205,13 +205,16 @@ class FLightSceneInfo
 	FLightPrimitiveInteraction* DynamicInteractionStaticPrimitiveList;
 
 public:
+	using FPersistentId = int32;
+
 	/** The light's scene proxy. */
 	FLightSceneProxy* Proxy;
 
 	ELightComponentType Type;
 
 	/** If bVisible == true, this is the index of the primitive in Scene->Lights. */
-	int32 Id;
+	FPersistentId Id;
+	FORCEINLINE FPersistentId GetPersistentIndex() const { return Id; }
 
 	/** The identifier for the primitive in Scene->PrimitiveOctree. */
 	FOctreeElementId2 OctreeId;
diff --git a/Engine/Source/Runtime/Renderer/Private/RendererScene.cpp b/Engine/Source/Runtime/Renderer/Private/RendererScene.cpp
--- a/Engine/Source/Runtime/Renderer/Private/RendererScene.cpp
+++ b/Engine/Source/Runtime/Renderer/Private/RendererScene.cpp
@@ -175,92 +175,6 @@ TGlobalResource< FGlobalDitherUniformBuffer > GDitherFadedInUniformBuffer;
 
 static FThreadSafeCounter FSceneViewState_UniqueID;
 
-struct FUpdateLightTransformParameters
-{
-	FMatrix LightToWorld;
-	FVector4 Position;
-};
-
-struct FUpdateLightCommand
-{
-	enum class EAddOrRemove
-	{
-		Add,
-		Remove,
-		None,
-	};
-
-	struct FColorParameters
-	{
-		FLinearColor NewColor;
-		float NewIndirectLightingScale;
-		float NewVolumetricScatteringIntensity;
-	};
-
-	using FTransformParameters = FUpdateLightTransformParameters;
-
-	inline bool IsAdd() const { return AddOrRemove == EAddOrRemove::Add; }
-	inline bool IsRemove() const { return AddOrRemove == EAddOrRemove::Remove; }
-
-	void Set(const FColorParameters& InColorParameters)
-	{
-		bHasColor = true;
-		ColorParameters = InColorParameters;
-	}
-	void Set(const FTransformParameters& InTransformParameters)
-	{
-		bHasTransform = true;
-		TransformParameters = InTransformParameters;
-	}
-	void Set(const EAddOrRemove& InAddOrRemove)
-	{
-		AddOrRemove = InAddOrRemove;
-	}
-
-	FUpdateLightCommand(const FColorParameters& InColorParameters, FLightSceneInfo* InLightSceneInfo) : LightSceneInfo(InLightSceneInfo), AddOrRemove(EAddOrRemove::None), bHasTransform(false) { Set(InColorParameters); }
-	FUpdateLightCommand(const FTransformParameters& InTransformParameters, FLightSceneInfo* InLightSceneInfo) : LightSceneInfo(InLightSceneInfo), AddOrRemove(EAddOrRemove::None), bHasColor(false) { Set(InTransformParameters); }
-	// Crtor for add/re,pve
-	FUpdateLightCommand(EAddOrRemove InAddOrRemove, FLightSceneInfo* InLightSceneInfo) : LightSceneInfo(InLightSceneInfo), AddOrRemove(InAddOrRemove), bHasTransform(false), bHasColor(false) {}
-
-	FTransformParameters TransformParameters;
-	FColorParameters ColorParameters;
-	FLightSceneInfo* LightSceneInfo;
-	EAddOrRemove AddOrRemove;
-	uint32 bHasTransform : 1;
-	uint32 bHasColor : 1;
-};
-
-class FSceneLightInfoUpdates
-{
-public:
-	Experimental::TRobinHoodHashMap<FLightSceneInfo*, FUpdateLightCommand> Commands;
-
-	int32 NumAdds = 0;
-	int32 NumRemoves = 0;
-	int32 NumUpdates = 0;
-
-	template <typename PayloadT>
-	void Enqueue(const PayloadT& Payload, FLightSceneInfo* LightSceneInfo)
-	{
-		// Allocate a new slot for update data
-		bool bWasAlreadyInMap = false;
-		FUpdateLightCommand* Command = Commands.FindOrAdd(LightSceneInfo, FUpdateLightCommand(Payload, LightSceneInfo), bWasAlreadyInMap);
-		if (bWasAlreadyInMap)
-		{
-			check(LightSceneInfo == Command->LightSceneInfo);
-			Command->Set(Payload);
-		}
-	}
-
-	void Reset()
-	{
-		Commands.Empty();
-		NumAdds = 0;
-		NumRemoves = 0;
-		NumUpdates = 0;
-	}
-};
-
 /**
  * Holds the info to update SpeedTree wind per unique tree object in the scene, instead of per instance
  */
@@ -1591,12 +1505,13 @@ static void TBitArraySwapElements(TBitArray<>& Array, int32 i1, int32 i2)
 
 void FScene::AddPrimitiveSceneInfo_RenderThread(FPrimitiveSceneInfo* PrimitiveSceneInfo, const TOptional<FTransform>& PreviousTransform)
 {
+	// Must always be a novel primitive that is added
 	check(PrimitiveSceneInfo->PackedIndex == INDEX_NONE);
-	check(AddedPrimitiveSceneInfos.Find(PrimitiveSceneInfo) == nullptr);
-	AddedPrimitiveSceneInfos.FindOrAdd(PrimitiveSceneInfo);
+	PrimitiveUpdates.EnqueueAdd(PrimitiveSceneInfo);
+
 	if (PreviousTransform.IsSet())
 	{
-		OverridenPreviousTransforms.Update(PrimitiveSceneInfo, PreviousTransform.GetValue().ToMatrixWithScale());
+		PrimitiveUpdates.Enqueue<FUpdateOverridePreviousTransformData>(PrimitiveSceneInfo, FUpdateOverridePreviousTransformData(PreviousTransform.GetValue().ToMatrixWithScale()));
 	}
 }
 
@@ -1792,8 +1707,7 @@ FScene::~FScene()
 	}
 #endif
 
-	checkf(RemovedPrimitiveSceneInfos.Num() == 0, TEXT("All pending primitive removal operations are expected to be flushed when the scene is destroyed. Remaining operations are likely to cause a memory leak."));
-	checkf(AddedPrimitiveSceneInfos.Num() == 0, TEXT("All pending primitive addition operations are expected to be flushed when the scene is destroyed. Remaining operations are likely to cause a memory leak."));
+	checkf(PrimitiveUpdates.IsEmpty(), TEXT("All pending primitive addition operations are expected to be flushed when the scene is destroyed. Remaining operations are likely to cause a memory leak."));
 	checkf(Primitives.Num() == 0, TEXT("All primitives are expected to be removed before the scene is destroyed. Remaining primitives are likely to cause a memory leak."));
 
 	delete InstanceCullingOcclusionQueryRenderer;
@@ -1846,8 +1760,6 @@ FScene::~FScene()
 	RayTracingSkinnedGeometryUpdateQueue = nullptr;
 #endif // RHI_RAYTRACING
 
-	checkf(RemovedPrimitiveSceneInfos.Num() == 0, TEXT("Leaking %d FPrimitiveSceneInfo instances."), RemovedPrimitiveSceneInfos.Num()); // Ensure UpdateAllPrimitiveSceneInfos() is called before destruction.
-
 	delete SceneLightInfoUpdates;
 }
 
@@ -1999,6 +1911,7 @@ void FScene::BatchAddPrimitivesInternal(TArrayView<T*> InPrimitives)
 
 	if (!CreateCommands.IsEmpty())
 	{
+		// Must enqueue RT commands to ensure SetTransform & CreateRenderThreadResources are executed in-order.
 		ENQUEUE_RENDER_COMMAND(AddPrimitiveCommand)(
 			[this, CreateCommands = MoveTemp(CreateCommands)](FRHICommandListBase& RHICmdList)
 		{
@@ -2040,45 +1953,58 @@ static FAutoConsoleVariableRef CVarSkipRedundantTransformUpdate(
 	ECVF_Default
 );
 
-void FScene::UpdatePrimitiveTransform_RenderThread(FPrimitiveSceneProxy* PrimitiveSceneProxy, const FBoxSphereBounds& WorldBounds, const FBoxSphereBounds& LocalBounds, const FMatrix& LocalToWorld, const FVector& AttachmentRootPosition, const TOptional<FTransform>& PreviousTransform)
-{
 #if VALIDATE_PRIMITIVE_PACKED_INDEX
-	if (AddedPrimitiveSceneInfos.Find(PrimitiveSceneProxy->GetPrimitiveSceneInfo()) != nullptr)
+inline void ValidatePackedPrimitiveIndexForUpdate(FPrimitiveSceneInfo* PrimitiveSceneInfo, const FScenePrimitiveUpdate& PrimitiveUpdates)
+{
+	if (const FPrimitiveUpdateCommand* Cmd = PrimitiveUpdates.FindCommand(PrimitiveSceneInfo))
 	{
-		check(PrimitiveSceneProxy->GetPrimitiveSceneInfo()->PackedIndex == INDEX_NONE);
+		if (Cmd->bAdded)
+		{
+			check(PrimitiveSceneInfo->GetIndex() == INDEX_NONE);
+		}
+		else
+		{
+			check(PrimitiveSceneInfo->GetIndex() != INDEX_NONE);
+		}				
 	}
 	else
 	{
-		check(PrimitiveSceneProxy->GetPrimitiveSceneInfo()->PackedIndex != INDEX_NONE);
+		// sending update, and no queued command - must have a valid index
+		check(PrimitiveSceneInfo->GetIndex() != INDEX_NONE);
 	}
+}
+#endif
 
-	check(RemovedPrimitiveSceneInfos.Find(PrimitiveSceneProxy->GetPrimitiveSceneInfo()) == nullptr);
+void FScene::UpdatePrimitiveTransform_RenderThread(FPrimitiveSceneProxy* PrimitiveSceneProxy, const FBoxSphereBounds& WorldBounds, const FBoxSphereBounds& LocalBounds, const FMatrix& LocalToWorld, const FVector& AttachmentRootPosition, const TOptional<FTransform>& PreviousTransform)
+{
+	FPrimitiveSceneInfo* PrimitiveSceneInfo = PrimitiveSceneProxy->GetPrimitiveSceneInfo();
+#if VALIDATE_PRIMITIVE_PACKED_INDEX
+	ValidatePackedPrimitiveIndexForUpdate(PrimitiveSceneInfo, PrimitiveUpdates);
 #endif
 
-	UpdatedTransforms.Update(PrimitiveSceneProxy, { WorldBounds, LocalBounds, LocalToWorld, AttachmentRootPosition });
+	PrimitiveUpdates.Enqueue(PrimitiveSceneInfo, FUpdateTransformCommand { .WorldBounds = WorldBounds, .LocalBounds = LocalBounds, .LocalToWorld = LocalToWorld, .AttachmentRootPosition = AttachmentRootPosition });
 
 	if (PreviousTransform.IsSet())
 	{
-		OverridenPreviousTransforms.Update(PrimitiveSceneProxy->GetPrimitiveSceneInfo(), PreviousTransform.GetValue().ToMatrixWithScale());
+		PrimitiveUpdates.Enqueue(PrimitiveSceneInfo, FUpdateOverridePreviousTransformData(PreviousTransform.GetValue().ToMatrixWithScale()));
 	}
 }
 
-void FScene::UpdatePrimitiveOcclusionBoundsSlack_RenderThread(const FPrimitiveSceneProxy* PrimitiveSceneProxy, float NewSlack)
+template <typename UpdatePayloadType>
+void FScene::UpdatePrimitiveInternal(FPrimitiveSceneProxy* SceneProxy, UpdatePayloadType&& InUpdatePayload)
 {
-#if VALIDATE_PRIMITIVE_PACKED_INDEX
-	if (AddedPrimitiveSceneInfos.Find(PrimitiveSceneProxy->GetPrimitiveSceneInfo()) != nullptr)
-	{
-		check(PrimitiveSceneProxy->GetPrimitiveSceneInfo()->PackedIndex == INDEX_NONE);
-	}
-	else
+	check(IsInGameThread());
+	if (SceneProxy != nullptr)
 	{
-		check(PrimitiveSceneProxy->GetPrimitiveSceneInfo()->PackedIndex != INDEX_NONE);
-	}
-
-	check(RemovedPrimitiveSceneInfos.Find(PrimitiveSceneProxy->GetPrimitiveSceneInfo()) == nullptr);
+		ENQUEUE_RENDER_COMMAND(UpdatePrimitiveCmd)(
+		[this, PrimitiveSceneInfo = SceneProxy->PrimitiveSceneInfo, UpdatePayload = MoveTemp(InUpdatePayload)] (FRHICommandListBase&) mutable
+		{
+#if VALIDATE_PRIMITIVE_PACKED_INDEX
+			ValidatePackedPrimitiveIndexForUpdate(PrimitiveSceneInfo, PrimitiveUpdates);
 #endif
-
-	UpdatedOcclusionBoundsSlacks.Update(PrimitiveSceneProxy, NewSlack);
+			PrimitiveUpdates.Enqueue<UpdatePayloadType>(PrimitiveSceneInfo, MoveTemp(UpdatePayload));
+		});
+	}
 }
 
 void FScene::UpdatePrimitiveTransform(UPrimitiveComponent* Primitive)
@@ -2091,7 +2017,7 @@ void FScene::UpdatePrimitiveTransform(FPrimitiveSceneDesc* Primitive)
 	UpdatePrimitiveTransformInternal(Primitive);	
 }
 
-template<class T>
+template<class T> 	
 void FScene::UpdatePrimitiveTransformInternal(T* Primitive)
 {
 	SCOPE_CYCLE_COUNTER(STAT_UpdatePrimitiveTransformGT);
@@ -2249,38 +2175,17 @@ void FScene::FinishUpdatePrimitiveTransform()
 
 void FScene::UpdatePrimitiveOcclusionBoundsSlack(UPrimitiveComponent* Primitive, float NewSlack)
 {
-	if (const FPrimitiveSceneProxy* SceneProxy = Primitive->GetSceneProxy())
-	{
-		ENQUEUE_RENDER_COMMAND(UpdateOcclusionBoundsSlackCmd)(
-			[this, SceneProxy, NewSlack] (FRHICommandListBase&)
-			{
-				UpdatePrimitiveOcclusionBoundsSlack_RenderThread(SceneProxy, NewSlack);
-			});
-	}
+	UpdatePrimitiveInternal(Primitive->GetSceneProxy(), FUpdateOcclusionBoundsSlacksData(NewSlack));
 }
 
 void FScene::UpdatePrimitiveDrawDistance(UPrimitiveComponent* Primitive, float MinDrawDistance, float MaxDrawDistance, float VirtualTextureMaxDrawDistance)
 {
-	if (FPrimitiveSceneProxy* SceneProxy = Primitive->GetSceneProxy())
-	{
-		ENQUEUE_RENDER_COMMAND(UpdatePrimitiveDrawDistanceCmd)(
-			[this, SceneProxy, MinDrawDistance, MaxDrawDistance, VirtualTextureMaxDrawDistance] (FRHICommandListBase&)
-			{
-				UpdatedDrawDistance.Update(SceneProxy, FVector3f(MinDrawDistance, MaxDrawDistance, VirtualTextureMaxDrawDistance));
-			});
-	}
+	UpdatePrimitiveInternal(Primitive->GetSceneProxy(), FUpdateDrawDistanceData(FVector3f(MinDrawDistance, MaxDrawDistance, VirtualTextureMaxDrawDistance)));
 }
 
 void FScene::UpdateInstanceCullDistance(UPrimitiveComponent* Primitive, float StartCullDistance, float EndCullDistance)
 {
-	if (FPrimitiveSceneProxy* SceneProxy = Primitive->GetSceneProxy())
-	{
-		ENQUEUE_RENDER_COMMAND(UpdateInstanceCullDistanceCmd)(
-			[this, SceneProxy, StartCullDistance, EndCullDistance] (FRHICommandListBase&)
-			{
-				UpdatedInstanceCullDistance.Update(SceneProxy, FVector2f(StartCullDistance, EndCullDistance));
-			});
-	}
+	UpdatePrimitiveInternal(Primitive->GetSceneProxy(), FUpdateInstanceCullDistanceData(FVector2f(StartCullDistance, EndCullDistance)));
 }
 
 void FScene::UpdatePrimitiveInstances(UInstancedStaticMeshComponent* Primitive)
@@ -2298,12 +2203,12 @@ void FScene::UpdatePrimitiveInstances(UInstancedStaticMeshComponent* Primitive)
 	FUpdateInstanceCommand UpdateParams;
 	UpdateParams.PrimitiveSceneProxy = Primitive->GetSceneProxy();
 	UpdateParams.WorldBounds = Primitive->Bounds;
-	UpdateParams.LocalBounds = ((UPrimitiveComponent*)Primitive)->GetLocalBounds();
+	UpdateParams.LocalBounds = static_cast<UPrimitiveComponent*>(Primitive)->GetLocalBounds();
 
 	// #todo (jnadro) This code should not be dependent on static mesh bounds.		
 	UpdateParams.StaticMeshBounds = Primitive->GetStaticMesh()->GetBounds();
 
-	return UpdatePrimitiveInstances(UpdateParams);
+	UpdatePrimitiveInstances(UpdateParams);
 }
 
 void FScene::UpdatePrimitiveInstances(FInstancedStaticMeshSceneDesc* Primitive)
@@ -2326,30 +2231,12 @@ void FScene::UpdatePrimitiveInstances(FInstancedStaticMeshSceneDesc* Primitive)
 	// #todo (jnadro) This code should not be dependent on static mesh bounds.		
 	UpdateParams.StaticMeshBounds = Primitive->GetStaticMesh()->GetBounds();
 
-	return UpdatePrimitiveInstances(UpdateParams);
+	UpdatePrimitiveInstances(UpdateParams);
 }
 
 void FScene::UpdatePrimitiveInstances(FUpdateInstanceCommand& UpdateParams)
 {	
-		ENQUEUE_RENDER_COMMAND(UpdateInstanceCommand)(
-			[this, UpdateParams = MoveTemp(UpdateParams)] (FRHICommandListBase&)
-		{
-#if VALIDATE_PRIMITIVE_PACKED_INDEX
-			if (AddedPrimitiveSceneInfos.Find(UpdateParams.PrimitiveSceneProxy->GetPrimitiveSceneInfo()) != nullptr)
-			{
-				check(UpdateParams.PrimitiveSceneProxy->GetPrimitiveSceneInfo()->PackedIndex == INDEX_NONE);
-			}
-			else
-			{
-				check(UpdateParams.PrimitiveSceneProxy->GetPrimitiveSceneInfo()->PackedIndex != INDEX_NONE);
-			}
-
-			check(RemovedPrimitiveSceneInfos.Find(UpdateParams.PrimitiveSceneProxy->GetPrimitiveSceneInfo()) == nullptr);
-#endif
-			FScopeCycleCounter Context(UpdateParams.PrimitiveSceneProxy->GetStatId());
-			UpdatedInstances.Update(UpdateParams.PrimitiveSceneProxy, UpdateParams);
-		}
-	);
+	UpdatePrimitiveInternal(UpdateParams.PrimitiveSceneProxy, MoveTemp(UpdateParams));
 }
 
 void FScene::UpdatePrimitiveSelectedState_RenderThread(const FPrimitiveSceneInfo* PrimitiveSceneInfo, bool bIsSelected)
@@ -2373,22 +2260,11 @@ void FScene::UpdatePrimitiveLightingAttachmentRoot(UPrimitiveComponent* Primitiv
 
 	if (NewLightingAttachmentRoot == Primitive)
 	{
-		NewLightingAttachmentRoot = NULL;
+		NewLightingAttachmentRoot = nullptr;
 	}
 
 	FPrimitiveComponentId NewComponentId = NewLightingAttachmentRoot ? NewLightingAttachmentRoot->GetPrimitiveSceneId() : FPrimitiveComponentId();
-
-	if (Primitive->SceneProxy)
-	{
-		FPrimitiveSceneProxy* Proxy = Primitive->SceneProxy;
-		FScene* Scene = this;
-		ENQUEUE_RENDER_COMMAND(UpdatePrimitiveAttachment)(
-			[Scene, Proxy, NewComponentId] (FRHICommandListBase&)
-			{
-				FPrimitiveSceneInfo* PrimitiveInfo = Proxy->GetPrimitiveSceneInfo();
-				Scene->UpdatedAttachmentRoots.Update(PrimitiveInfo, NewComponentId);
-			});
-	}
+	UpdatePrimitiveInternal(Primitive->GetSceneProxy(), FUpdateAttachmentRootData(NewComponentId));
 }
 
 void FScene::UpdatePrimitiveAttachment(UPrimitiveComponent* Primitive)
@@ -2429,44 +2305,14 @@ void FScene::UpdateCustomPrimitiveData(FPrimitiveSceneDesc* Primitive, const FCu
 
 void FScene::UpdateCustomPrimitiveData(FPrimitiveSceneProxy* SceneProxy, const FCustomPrimitiveData& CustomPrimitiveData)
 {
-	// This path updates the primitive data directly in the GPUScene. 
-	if (SceneProxy) 
-	{
-		struct FUpdateParams
-		{
-			FScene* Scene;
-			FPrimitiveSceneProxy* PrimitiveSceneProxy;
-			FCustomPrimitiveData CustomPrimitiveData;
-		};
-
-		FUpdateParams UpdateParams;
-		UpdateParams.Scene = this;
-		UpdateParams.PrimitiveSceneProxy = SceneProxy;
-		UpdateParams.CustomPrimitiveData = CustomPrimitiveData; 
-
-		ENQUEUE_RENDER_COMMAND(UpdateCustomPrimitiveDataCommand)(
-			[UpdateParams] (FRHICommandListBase&)
-			{
-				UpdateParams.Scene->UpdatedCustomPrimitiveParams.Update(UpdateParams.PrimitiveSceneProxy, UpdateParams.CustomPrimitiveData);
-			});
-	}
+	UpdatePrimitiveInternal(SceneProxy, FUpdateCustomPrimitiveData(CustomPrimitiveData));
 }
 
 void FScene::UpdatePrimitiveDistanceFieldSceneData_GameThread(UPrimitiveComponent* Primitive)
 {
-	check(IsInGameThread());
-
 	if (Primitive->SceneProxy)
 	{
-		ENQUEUE_RENDER_COMMAND(UpdatePrimDFSceneDataCmd)(
-			[this, PrimitiveSceneProxy = Primitive->SceneProxy] (FRHICommandListBase&)
-			{
-				if (PrimitiveSceneProxy && PrimitiveSceneProxy->GetPrimitiveSceneInfo())
-				{
-					FPrimitiveSceneInfo* Info = PrimitiveSceneProxy->GetPrimitiveSceneInfo();
-					this->DistanceFieldSceneDataUpdates.FindOrAdd(Info);
-				}
-			});
+		UpdatePrimitiveInternal(Primitive->GetSceneProxy(), FUpdateDistanceFieldSceneData{});
 	}
 }
 
@@ -2490,31 +2336,9 @@ FPrimitiveSceneInfo* FScene::GetPrimitiveSceneInfo(const FPersistentPrimitiveInd
 	return GetPrimitiveSceneInfo(PrimitiveIndex);
 }
 
-bool FScene::RemovePrimitiveSceneInfo_RenderThread(FPrimitiveSceneInfo* PrimitiveSceneInfo)
+void FScene::RemovePrimitiveSceneInfo_RenderThread(FPrimitiveSceneInfo* PrimitiveSceneInfo)
 {
-	const bool bRemovePendingAdd = AddedPrimitiveSceneInfos.Remove(PrimitiveSceneInfo);
-
-	if (bRemovePendingAdd)
-	{
-		check(PrimitiveSceneInfo->PackedIndex == INDEX_NONE);
-		UpdatedTransforms.Remove(PrimitiveSceneInfo->Proxy);
-		UpdatedCustomPrimitiveParams.Remove(PrimitiveSceneInfo->Proxy);
-		OverridenPreviousTransforms.Remove(PrimitiveSceneInfo);
-		UpdatedOcclusionBoundsSlacks.Remove(PrimitiveSceneInfo->Proxy);
-		UpdatedInstanceCullDistance.Remove(PrimitiveSceneInfo->Proxy);
-		UpdatedDrawDistance.Remove(PrimitiveSceneInfo->Proxy);
-		DistanceFieldSceneDataUpdates.Remove(PrimitiveSceneInfo);
-		UpdatedAttachmentRoots.Remove(PrimitiveSceneInfo);
-		DeletedPrimitiveSceneInfos.Emplace(PrimitiveSceneInfo);
-	}
-	else
-	{
-		check(PrimitiveSceneInfo->PackedIndex != INDEX_NONE);
-		check(RemovedPrimitiveSceneInfos.Find(PrimitiveSceneInfo) == nullptr);
-		RemovedPrimitiveSceneInfos.FindOrAdd(PrimitiveSceneInfo);
-	}
-	
-	return !bRemovePendingAdd;
+	PrimitiveUpdates.EnqueueDelete(PrimitiveSceneInfo);
 }
 
 void FScene::RemovePrimitive(UPrimitiveComponent* Primitive)
@@ -2569,6 +2393,7 @@ void FScene::BatchRemovePrimitivesInternal(TArrayView<T*> InPrimitives)
 
 	if (!DestroyCommands.IsEmpty())
 	{
+		// must run RT cmds in order here too
 		ENQUEUE_RENDER_COMMAND(FRemovePrimitiveCommand)(
 			[this, DestroyCommands = MoveTemp(DestroyCommands)](FRHICommandListBase&)
 		{
@@ -2815,10 +2640,10 @@ void FScene::AddLight(ULightComponent* Light)
 			{
 				CSV_SCOPED_TIMING_STAT_EXCLUSIVE(Scene_AddLight);
 				FScopeCycleCounter Context(LightSceneInfo->Proxy->GetStatId());
-				bool bWasAlreadyInMap = false;
-				SceneLightInfoUpdates->Commands.FindOrAdd(LightSceneInfo, FUpdateLightCommand(FUpdateLightCommand::EAddOrRemove::Add, LightSceneInfo), bWasAlreadyInMap);
-				check(!bWasAlreadyInMap);
-				++SceneLightInfoUpdates->NumAdds;
+
+				// It should not have been added to the scene already
+				check(LightSceneInfo->Id == INDEX_NONE);
+				SceneLightInfoUpdates->EnqueueAdd(LightSceneInfo);
 			});
 	}
 }
@@ -3270,7 +3095,7 @@ void FScene::UpdateReflectionCaptureTransform(UReflectionCaptureComponent* Compo
 		FReflectionCaptureProxy* Proxy = Component->SceneProxy;
 		FMatrix Transform = Component->GetComponentTransform().ToMatrixWithScale();
 
-		ENQUEUE_RENDER_COMMAND(FUpdateTransformCommand)(
+		ENQUEUE_RENDER_COMMAND(UpdateReflectionCaptureTransformCommand)(
 			[Scene, Proxy, Transform, bUsingPreviewCaptureData](FRHICommandListBase& RHICmdList)
 		{
 			if (Proxy->bUsingPreviewCaptureData)
@@ -3848,7 +3673,7 @@ bool DoesPlatformNeedLocalLightPrimitiveInteraction(EShaderPlatform ShaderPlatfo
 	return !IsMobilePlatform(ShaderPlatform) || !MobileLocalLightsUseSinglePermutation(ShaderPlatform) || IsMobileMovableSpotlightShadowsEnabled(ShaderPlatform);
 }
 
-void FScene::UpdateLightTransform_RenderThread(int32 LightId, FLightSceneInfo* LightSceneInfo, const FUpdateLightCommand::FTransformParameters& Parameters)
+void FScene::UpdateLightTransform_RenderThread(int32 LightId, FLightSceneInfo* LightSceneInfo, const FUpdateLightTransformParameters& Parameters)
 {
 	SCOPED_NAMED_EVENT(FScene_UpdateLightTransform_RenderThread, FColor::Yellow);
 
@@ -3929,54 +3754,35 @@ void FScene::UpdateLightTransform_RenderThread(int32 LightId, FLightSceneInfo* L
 	}
 }
 
-void FScene::UpdateLightTransform(ULightComponent* Light)
+template <typename UpdatePayloadType>
+void FScene::UpdateLightInternal(FLightSceneProxy* LightSceneProxy, UpdatePayloadType&& InUpdatePayload)
 {
-	if(Light->SceneProxy)
+	if (LightSceneProxy)
 	{
-		FUpdateLightCommand::FTransformParameters Parameters;
-		Parameters.LightToWorld = Light->GetComponentTransform().ToMatrixNoScale();
-		Parameters.Position = Light->GetLightPosition();
-		
-		FLightSceneInfo* LightSceneInfo = Light->SceneProxy->GetLightSceneInfo();
+		FLightSceneInfo* LightSceneInfo = LightSceneProxy->GetLightSceneInfo();
 		if (LightSceneInfo->bVisible)
 		{
 			ENQUEUE_RENDER_COMMAND(UpdateLightTransform)(
-				[this, LightSceneInfo, Parameters] (FRHICommandListBase&)
+				[this, LightSceneInfo, UpdatePayload = MoveTemp(InUpdatePayload)] (FRHICommandListBase&) mutable
 				{
 					FScopeCycleCounter Context(LightSceneInfo->Proxy->GetStatId());
-					SceneLightInfoUpdates->Enqueue(Parameters, LightSceneInfo);
-					++SceneLightInfoUpdates->NumUpdates;
+					SceneLightInfoUpdates->Enqueue(LightSceneInfo, MoveTemp(UpdatePayload));
 				});
 		}
 	}
 }
 
-/** 
- * Updates the color and brightness of a light which has already been added to the scene. 
- *
- * @param Light - light component to update
- */
-void FScene::UpdateLightColorAndBrightness(ULightComponent* Light)
+void FScene::UpdateLightTransform(ULightComponent* Light)
 {
-	if(Light->SceneProxy)
-	{
-		FUpdateLightCommand::FColorParameters NewParameters;
-		NewParameters.NewColor = Light->GetColoredLightBrightness();
-		NewParameters.NewIndirectLightingScale = Light->IndirectLightingIntensity;
-		NewParameters.NewVolumetricScatteringIntensity = Light->VolumetricScatteringIntensity;
+	UpdateLightInternal(Light->SceneProxy, FUpdateLightTransformParameters{.LightToWorld = Light->GetComponentTransform().ToMatrixNoScale(), .Position = Light->GetLightPosition()});
+}
 
-		FScene* Scene = this;
-		FLightSceneInfo* LightSceneInfo = Light->SceneProxy->GetLightSceneInfo();
-		if (LightSceneInfo->bVisible)
-		{
-			ENQUEUE_RENDER_COMMAND(UpdateLightColorAndBrightness)(
-				[this, LightSceneInfo, NewParameters] (FRHICommandListBase&)
-				{
-					SceneLightInfoUpdates->Enqueue(NewParameters, LightSceneInfo);
-					++SceneLightInfoUpdates->NumUpdates;
-				});
-		}
-	}
+void FScene::UpdateLightColorAndBrightness(ULightComponent* Light)
+{
+	UpdateLightInternal(Light->SceneProxy, FUpdateLightColorParameters{
+		.NewColor = Light->GetColoredLightBrightness(),
+		.NewIndirectLightingScale = Light->IndirectLightingIntensity,
+		.NewVolumetricScatteringIntensity = Light->VolumetricScatteringIntensity});
 }
 
 void FScene::RemoveLightSceneInfo_RenderThread(FLightSceneInfo* LightSceneInfo)
@@ -4111,25 +3917,14 @@ void FScene::RemoveLight(ULightComponent* Light)
 
 				if (LightSceneInfo->bVisible)
 				{
-					auto HashElementId = SceneLightInfoUpdates->Commands.FindId(LightSceneInfo);
-					// If an Add is pending, just remove the update and fall through to delete the data.
-					if (HashElementId.IsValid() && SceneLightInfoUpdates->Commands.GetByElementId(HashElementId).Value.IsAdd())
-					{
-						SceneLightInfoUpdates->Commands.RemoveByElementId(HashElementId);
-						--SceneLightInfoUpdates->NumAdds;
-					}
-					else
-					{
-						SceneLightInfoUpdates->Enqueue(FUpdateLightCommand::EAddOrRemove::Remove, LightSceneInfo);
-						++SceneLightInfoUpdates->NumRemoves;
-						// early out to defer the deletion
-						return;
-					}
+					SceneLightInfoUpdates->EnqueueDelete(LightSceneInfo);
+					// Note: all the deletes are deferred now (except the invisible ones) as this simpifies the update buffer.
+					return;
 				}
 				else
 				{
 					// There should never be updates queued for lights that are not visible
-					check(SceneLightInfoUpdates->Commands.Find(LightSceneInfo) == nullptr);
+					check(!SceneLightInfoUpdates->HasCommand(LightSceneInfo));
 					// The "invisible lights" are removed at once.
 					InvisibleLights.RemoveAt(LightSceneInfo->Id);
 				}
@@ -4863,7 +4658,6 @@ void FScene::ApplyWorldOffset_RenderThread(FRHICommandListBase& RHICmdList, cons
 	GPUScene.bUpdateAllPrimitives = true;
 
 	// Primitives
-	checkf(AddedPrimitiveSceneInfos.Num() == 0, TEXT("All primitives found in AddedPrimitiveSceneInfos must have been added to the scene before the world offset is applied"));
 	for (int32 Idx = 0; Idx < Primitives.Num(); ++Idx)
 	{
 		Primitives[Idx]->ApplyWorldOffset(RHICmdList, InOffset);
@@ -5124,7 +4918,8 @@ void FScene::UpdatePrimitiveVelocityState_RenderThread(FPrimitiveSceneInfo* Prim
 }
 
 #if RHI_RAYTRACING
-void FScene::UpdateRayTracingGroupBounds_AddPrimitives(const Experimental::TRobinHoodHashSet<FPrimitiveSceneInfo*>& PrimitiveSceneInfos)
+
+void FScene::UpdateRayTracingGroupBounds_AddPrimitives(const TArray<FPrimitiveSceneInfo*, SceneRenderingAllocator>& PrimitiveSceneInfos)
 {
 	for (FPrimitiveSceneInfo* const PrimitiveSceneInfo : PrimitiveSceneInfos)
 	{
@@ -5169,7 +4964,7 @@ static void UpdateRayTracingGroupBounds(Experimental::TRobinHoodHashSet<FScene::
 	}
 }
 
-void FScene::UpdateRayTracingGroupBounds_RemovePrimitives(const Experimental::TRobinHoodHashSet<FPrimitiveSceneInfo*>& PrimitiveSceneInfos)
+void FScene::UpdateRayTracingGroupBounds_RemovePrimitives(const TArray<FPrimitiveSceneInfo*, SceneRenderingAllocator>& PrimitiveSceneInfos)
 {
 	Experimental::TRobinHoodHashSet<FRayTracingCullingGroup*> GroupsToUpdate;
 	for (FPrimitiveSceneInfo* const PrimitiveSceneInfo : PrimitiveSceneInfos)
@@ -5194,13 +4989,13 @@ void FScene::UpdateRayTracingGroupBounds_RemovePrimitives(const Experimental::TR
 	UpdateRayTracingGroupBounds(GroupsToUpdate);
 }
 
-template<typename ValueType>
-inline void FScene::UpdateRayTracingGroupBounds_UpdatePrimitives(const Experimental::TRobinHoodHashMap<FPrimitiveSceneProxy*, ValueType>& InUpdatedTransforms)
+template<typename RangeType>
+inline void FScene::UpdateRayTracingGroupBounds_UpdatePrimitives(const RangeType& InUpdatedTransforms)
 {
 	Experimental::TRobinHoodHashSet<FRayTracingCullingGroup*> GroupsToUpdate;
 	for (const auto& Transform : InUpdatedTransforms)
 	{
-		FPrimitiveSceneProxy* const PrimitiveSceneProxy = Transform.Key;
+		FPrimitiveSceneProxy* const PrimitiveSceneProxy = Transform.SceneInfo->Proxy;
 		const int32 RayTracingGroupId = PrimitiveSceneProxy->GetRayTracingGroupId();
 		const Experimental::FHashElementId GroupId = (RayTracingGroupId != -1) ? PrimitiveRayTracingGroups.FindId(RayTracingGroupId) : Experimental::FHashElementId();
 		if (GroupId.IsValid())
@@ -5238,50 +5033,42 @@ static inline bool IsPrimitiveRelevantToPathTracing(FPrimitiveSceneInfo* Primiti
 }
 
 
-FLightSceneChangeSet FScene::UpdateAllLightSceneInfos(FRDGBuilder& GraphBuilder)
+void FScene::UpdateAllLightSceneInfos(FRDGBuilder& GraphBuilder)
 {
 	SCOPED_NAMED_EVENT(FScene_UpdateAllLightSceneInfos, FColor::Orange);
 	struct FFLightSceneChangeSetAllocation
 	{
 		TArray<int32, SceneRenderingAllocator> RemovedLightIds;
-		TArray<int32, SceneRenderingAllocator> TransformUpdatedLightIds;
-		TArray<int32, SceneRenderingAllocator> ColorUpdatedLightIds;
 		TArray<int32, SceneRenderingAllocator> AddedLightIds;
+		FFLightSceneChangeSetAllocation(FSceneLightInfoUpdates&& InSceneLightInfoUpdates) : SceneLightInfoUpdates(MoveTemp(InSceneLightInfoUpdates)) {}
+		FSceneLightInfoUpdates SceneLightInfoUpdates;
 	};
 	// Allocate change set storage with graph builder lifetime such that we can safely pass it to async tasks.
-	FFLightSceneChangeSetAllocation& ChangeSet = *GraphBuilder.AllocObject<FFLightSceneChangeSetAllocation>();
+	FFLightSceneChangeSetAllocation& ChangeSet = *GraphBuilder.AllocObject<FFLightSceneChangeSetAllocation>(MoveTemp(*SceneLightInfoUpdates));
 
-	// Conservative pre-size the arrays (some updates are covered by Adds).
-	ChangeSet.RemovedLightIds.Reserve(SceneLightInfoUpdates->NumRemoves);
-	ChangeSet.TransformUpdatedLightIds.Reserve(SceneLightInfoUpdates->NumUpdates);
-	ChangeSet.ColorUpdatedLightIds.Reserve(SceneLightInfoUpdates->NumUpdates);
-	ChangeSet.AddedLightIds.Reserve(SceneLightInfoUpdates->NumAdds);
+#if DO_CHECK
+	FSceneLightInfoUpdates::FReadAccessScope ReadAccessScope(ChangeSet.SceneLightInfoUpdates);
+#endif
 
-	// Filter out removes & updates:
-	for (const auto& Element : SceneLightInfoUpdates->Commands)
+	// Filter out removes & never-adds
+	ChangeSet.SceneLightInfoUpdates.ForEachCommand([&](const FUpdateLightCommand& UpdateLightCommand)
 	{
-		const FUpdateLightCommand& UpdateLightCommand = Element.Value;
-		int32 Id = UpdateLightCommand.LightSceneInfo->Id;
-		if (UpdateLightCommand.IsRemove())
-		{
-			ChangeSet.RemovedLightIds.Add(Id);
-		}
-		else if (!UpdateLightCommand.IsAdd())
+		int32 Id = UpdateLightCommand.GetSceneInfo()->Id;
+		if (UpdateLightCommand.IsDelete())
 		{
-			check(Id != INDEX_NONE);
-			if (UpdateLightCommand.bHasTransform)
+			if (Id != INDEX_NONE)
 			{
-				ChangeSet.TransformUpdatedLightIds.Add(Id);
+				ChangeSet.RemovedLightIds.Add(Id);
 			}
-			if (UpdateLightCommand.bHasColor)
+			else
 			{
-				ChangeSet.ColorUpdatedLightIds.Add(Id);
+				delete UpdateLightCommand.GetSceneInfo()->Proxy;
+				delete UpdateLightCommand.GetSceneInfo();
 			}
 		}
-	}
+	});
 	// This can't access the scene light data if done async since it happens before the actual removals.
-	OnPreLightSceneInfoUpdate.Broadcast(GraphBuilder, FLightSceneChangeSet{ ChangeSet.RemovedLightIds, TConstArrayView<int32>(), ChangeSet.TransformUpdatedLightIds, ChangeSet.ColorUpdatedLightIds });
-
+	OnPreLightSceneInfoUpdate.Broadcast(GraphBuilder, FLightSceneChangeSet{ ChangeSet.RemovedLightIds, TConstArrayView<int32>(), &ChangeSet.SceneLightInfoUpdates });
 	// Batch process all light removes
 	for (int32 LightId : ChangeSet.RemovedLightIds)
 	{
@@ -5291,33 +5078,27 @@ FLightSceneChangeSet FScene::UpdateAllLightSceneInfos(FRDGBuilder& GraphBuilder)
 	}
 
 	// Process all light adds & updates
-	for (const auto &Element : SceneLightInfoUpdates->Commands)
+	ChangeSet.SceneLightInfoUpdates.ForEachCommand(ESceneUpdateCommandFilter::AddedUpdated, [&](const FUpdateLightCommand& UpdateLightCommand)
 	{
-		const FUpdateLightCommand& UpdateLightCommand = Element.Value;
-		if (UpdateLightCommand.IsRemove())
-		{
-			continue;
-		}
-
-		FLightSceneInfo* LightSceneInfo = UpdateLightCommand.LightSceneInfo;
+		FLightSceneInfo* LightSceneInfo = UpdateLightCommand.GetSceneInfo();
 		FScopeCycleCounter Context(LightSceneInfo->Proxy->GetStatId());
 
 		const int32 Id = LightSceneInfo->Id;
 		const bool bHasId = Id != INDEX_NONE;
 		check(bHasId == !UpdateLightCommand.IsAdd());
 		// Directly process updates.
-		if (UpdateLightCommand.bHasTransform)
+		if (FUpdateLightTransformParameters* TransformParameters = ChangeSet.SceneLightInfoUpdates.GetPayloadPtr<FUpdateLightTransformParameters>(UpdateLightCommand))
 		{
 			SCOPE_CYCLE_COUNTER(STAT_UpdateSceneLightTime);
 
-			UpdateLightTransform_RenderThread(Id, LightSceneInfo, UpdateLightCommand.TransformParameters);
+			UpdateLightTransform_RenderThread(Id, LightSceneInfo, *TransformParameters);
 		}
-		if (UpdateLightCommand.bHasColor)
+
+		if (FUpdateLightColorParameters* ColorParameters = ChangeSet.SceneLightInfoUpdates.GetPayloadPtr<FUpdateLightColorParameters>(UpdateLightCommand))
 		{
+			const FUpdateLightColorParameters& NewParameters = *ColorParameters;
 			SCOPE_CYCLE_COUNTER(STAT_UpdateSceneLightTime);
 
-			const FUpdateLightCommand::FColorParameters& NewParameters = UpdateLightCommand.ColorParameters;
-
 			// Mobile renderer:
 			// a light with no color/intensity can cause the light to be ignored when rendering.
 			// thus, lights that change state in this way must update the draw lists.
@@ -5350,14 +5131,10 @@ FLightSceneChangeSet FScene::UpdateAllLightSceneInfos(FRDGBuilder& GraphBuilder)
 			// Note: Id is set in AddLightSceneInfo_RenderThread so we must fetch it again
 			ChangeSet.AddedLightIds.Add(LightSceneInfo->Id);
 		}
-	}
-
-	OnPostLightSceneInfoUpdate.Broadcast(GraphBuilder, FLightSceneChangeSet{ ChangeSet.RemovedLightIds, ChangeSet.AddedLightIds, ChangeSet.TransformUpdatedLightIds, ChangeSet.ColorUpdatedLightIds });
-	GPUScene.OnPostLightSceneInfoUpdate(GraphBuilder, FLightSceneChangeSet{ ChangeSet.RemovedLightIds, ChangeSet.AddedLightIds, ChangeSet.TransformUpdatedLightIds, ChangeSet.ColorUpdatedLightIds });
-
-	SceneLightInfoUpdates->Reset();
+	});
 
-	return FLightSceneChangeSet{ ChangeSet.RemovedLightIds, ChangeSet.AddedLightIds, ChangeSet.TransformUpdatedLightIds, ChangeSet.ColorUpdatedLightIds };
+	OnPostLightSceneInfoUpdate.Broadcast(GraphBuilder, FLightSceneChangeSet{ ChangeSet.RemovedLightIds, ChangeSet.AddedLightIds, &ChangeSet.SceneLightInfoUpdates });
+	GPUScene.OnPostLightSceneInfoUpdate(GraphBuilder, FLightSceneChangeSet{ ChangeSet.RemovedLightIds, ChangeSet.AddedLightIds, &ChangeSet.SceneLightInfoUpdates });
 }
 
 template<class T>
@@ -5524,33 +5301,58 @@ void UpdateReflectionSceneData(FScene* Scene)
  */
 struct FSceneUpdateChangeSetStorage
 {
+
+	FSceneUpdateChangeSetStorage()
+		: PreUpdateChangeSet(FScenePreUpdateChangeSet {
+			TConstArrayView<FPersistentPrimitiveIndex>(RemovedPrimitiveIds),
+			TConstArrayView<FPrimitiveSceneInfo*>(RemovedPrimitiveSceneInfos),
+			PrimitiveUpdates
+			})
+		, PostUpdateChangeSet(FScenePostUpdateChangeSet {
+			TConstArrayView<FPersistentPrimitiveIndex>(AddedPrimitiveIds),
+			TConstArrayView<FPrimitiveSceneInfo*>(AddedPrimitiveSceneInfos),
+			PrimitiveUpdates
+		})
+	{
+	}
+	FSceneUpdateChangeSetStorage(FScenePrimitiveUpdates&& InPrimitiveUpdates)
+		: PrimitiveUpdates(MoveTemp(InPrimitiveUpdates))
+		, PreUpdateChangeSet(FScenePreUpdateChangeSet {
+			TConstArrayView<FPersistentPrimitiveIndex>(RemovedPrimitiveIds),
+			TConstArrayView<FPrimitiveSceneInfo*>(RemovedPrimitiveSceneInfos),
+			PrimitiveUpdates
+			})
+		, PostUpdateChangeSet(FScenePostUpdateChangeSet {
+			TConstArrayView<FPersistentPrimitiveIndex>(AddedPrimitiveIds),
+			TConstArrayView<FPrimitiveSceneInfo*>(AddedPrimitiveSceneInfos),
+			PrimitiveUpdates
+		})
+	{
+	}
+
+
 	TArray<FPersistentPrimitiveIndex, SceneRenderingAllocator> RemovedPrimitiveIds;
 	TArray<FPrimitiveSceneInfo*,SceneRenderingAllocator> RemovedPrimitiveSceneInfos;
 
-	TArray<FPersistentPrimitiveIndex, SceneRenderingAllocator> UpdatedPrimitiveIds;
-	TArray<FPrimitiveSceneInfo*, SceneRenderingAllocator> UpdatedPrimitiveSceneInfos;
-
 	TArray<FPersistentPrimitiveIndex, SceneRenderingAllocator> AddedPrimitiveIds;
 	TArray<FPrimitiveSceneInfo*, SceneRenderingAllocator> AddedPrimitiveSceneInfos;
+	FScenePrimitiveUpdates PrimitiveUpdates;
+
+	FScenePreUpdateChangeSet PreUpdateChangeSet;
+	FScenePostUpdateChangeSet PostUpdateChangeSet;
 
-	FScenePreUpdateChangeSet GetPreUpdateSet() const 
+	const FScenePreUpdateChangeSet& GetPreUpdateSet() 
 	{
-		return FScenePreUpdateChangeSet {
-			TConstArrayView<FPersistentPrimitiveIndex>(RemovedPrimitiveIds),
-			TConstArrayView<FPrimitiveSceneInfo*>(RemovedPrimitiveSceneInfos),
-			TConstArrayView<FPersistentPrimitiveIndex>(UpdatedPrimitiveIds),
-			TConstArrayView<FPrimitiveSceneInfo*>(UpdatedPrimitiveSceneInfos)
-		};
+		PreUpdateChangeSet.RemovedPrimitiveIds = TConstArrayView<FPersistentPrimitiveIndex>(RemovedPrimitiveIds);
+		PreUpdateChangeSet.RemovedPrimitiveSceneInfos =	TConstArrayView<FPrimitiveSceneInfo*>(RemovedPrimitiveSceneInfos);
+		return PreUpdateChangeSet;
 	}
 
-	FScenePostUpdateChangeSet GetPostUpdateSet() const 
+	const FScenePostUpdateChangeSet& GetPostUpdateSet() 
 	{
-		return FScenePostUpdateChangeSet {
-			TConstArrayView<FPersistentPrimitiveIndex>(AddedPrimitiveIds),
-			TConstArrayView<FPrimitiveSceneInfo*>(AddedPrimitiveSceneInfos),
-			TConstArrayView<FPersistentPrimitiveIndex>(UpdatedPrimitiveIds),
-			TConstArrayView<FPrimitiveSceneInfo*>(UpdatedPrimitiveSceneInfos)
-		};
+		PostUpdateChangeSet.AddedPrimitiveIds = TConstArrayView<FPersistentPrimitiveIndex>(AddedPrimitiveIds);
+		PostUpdateChangeSet.AddedPrimitiveSceneInfos = TConstArrayView<FPrimitiveSceneInfo*>(AddedPrimitiveSceneInfos);
+		return PostUpdateChangeSet;
 	}
 };
 
@@ -5595,18 +5397,44 @@ void FScene::Update(FRDGBuilder& GraphBuilder, const FUpdateParameters& Paramete
 	RDG_EVENT_SCOPE(GraphBuilder, "UpdateAllPrimitiveSceneInfos");
 
 	// Allocated with render graph lifetime, safe to reference from RDG tasks.
-	FSceneUpdateChangeSetStorage &SceneUpdateChangeSetStorage = *GraphBuilder.AllocObject<FSceneUpdateChangeSetStorage>();
-	SceneUpdateChangeSetStorage.RemovedPrimitiveIds.Reserve(RemovedPrimitiveSceneInfos.Num());
-	SceneUpdateChangeSetStorage.RemovedPrimitiveSceneInfos.Reserve(RemovedPrimitiveSceneInfos.Num());
+	FSceneUpdateChangeSetStorage& SceneUpdateChangeSetStorage = *GraphBuilder.AllocObject<FSceneUpdateChangeSetStorage>(MoveTemp(PrimitiveUpdates));
+	TArray<FPrimitiveSceneInfo*, SceneRenderingAllocator> DeletedPrimitiveSceneInfos;
+	DeletedPrimitiveSceneInfos.Reserve(SceneUpdateChangeSetStorage.PrimitiveUpdates.NumCommands());
+	SceneUpdateChangeSetStorage.AddedPrimitiveSceneInfos.Reserve(SceneUpdateChangeSetStorage.PrimitiveUpdates.NumCommands());
+	SceneUpdateChangeSetStorage.RemovedPrimitiveSceneInfos.Reserve(SceneUpdateChangeSetStorage.PrimitiveUpdates.NumCommands());
+	SceneUpdateChangeSetStorage.RemovedPrimitiveIds.Reserve(SceneUpdateChangeSetStorage.PrimitiveUpdates.NumCommands());
+
+	// Retrieve the commonly used categories used in scene updates.
+	// 1. RemovedPrimitiveSceneInfos: Primitives that exist in the scene before the update and will be removed.
+	// 2. AddedPrimitiveSceneInfos: Primitives not in the scene before the update that will be added.
+	// 3. DeletedPrimitiveSceneInfos: Primitives that needs to be deleted.
+	SceneUpdateChangeSetStorage.PrimitiveUpdates.ForEachCommand([&](const FPrimitiveUpdateCommand& Cmd)
+	{
+		// Skip those that were added & removed in the same frame
+		if (Cmd.IsAdd() && !Cmd.IsDelete())
+		{
+			SceneUpdateChangeSetStorage.AddedPrimitiveSceneInfos.Add(Cmd.GetSceneInfo());
+		}
+		if (Cmd.IsDelete())
+		{
+			DeletedPrimitiveSceneInfos.Add(Cmd.GetSceneInfo());
+			// Skip those that were added & removed in the same frame
+			if (!Cmd.IsAdd())
+			{
+				SceneUpdateChangeSetStorage.RemovedPrimitiveSceneInfos.Add(Cmd.GetSceneInfo());
+				SceneUpdateChangeSetStorage.RemovedPrimitiveIds.Add(Cmd.GetPersistentId());
+			}
+		}
+	});
+
+	const TArray<FPrimitiveSceneInfo*, SceneRenderingAllocator> &AddedPrimitiveSceneInfos = SceneUpdateChangeSetStorage.AddedPrimitiveSceneInfos;
+	const TArray<FPrimitiveSceneInfo*, SceneRenderingAllocator> &RemovedPrimitiveSceneInfos = SceneUpdateChangeSetStorage.RemovedPrimitiveSceneInfos;
 
 	TArray<FPrimitiveSceneInfo*, SceneRenderingAllocator> RemovedLocalPrimitiveSceneInfos;
 	RemovedLocalPrimitiveSceneInfos.Reserve(RemovedPrimitiveSceneInfos.Num());
 	for (FPrimitiveSceneInfo* PrimitiveSceneInfo : RemovedPrimitiveSceneInfos)
 	{
 		RemovedLocalPrimitiveSceneInfos.Add(PrimitiveSceneInfo);
-
-		SceneUpdateChangeSetStorage.RemovedPrimitiveIds.Add(PrimitiveSceneInfo->GetPersistentIndex());
-		SceneUpdateChangeSetStorage.RemovedPrimitiveSceneInfos.Add(PrimitiveSceneInfo);
 	}
 
 	TArray<FPrimitiveSceneInfo*, SceneRenderingAllocator> AddedLocalPrimitiveSceneInfos;
@@ -5623,61 +5451,10 @@ void FScene::Update(FRDGBuilder& GraphBuilder, const FUpdateParameters& Paramete
 	UpdateRayTracingGroupBounds_AddPrimitives(AddedPrimitiveSceneInfos);
 #endif
 
-	{
-		SceneUpdateChangeSetStorage.UpdatedPrimitiveIds.Reserve(UpdatedInstances.Num() + UpdatedTransforms.Num());
-		SceneUpdateChangeSetStorage.UpdatedPrimitiveSceneInfos.Reserve(UpdatedInstances.Num() + UpdatedTransforms.Num());
-		// All updated instances must also before moving or re-allocating (TODO: filter out only those actually updated)
-		for (const auto& Instance : UpdatedInstances)
-		{
-			FPrimitiveSceneInfo* PrimitiveSceneInfo = Instance.Key->GetPrimitiveSceneInfo();
-
-			if (RemovedPrimitiveSceneInfos.Find(PrimitiveSceneInfo) != nullptr)
-			{
-				continue;
-			}
-
-			FPersistentPrimitiveIndex Id = Instance.Key->GetPrimitiveSceneInfo()->GetPersistentIndex();
-			// Note: may not be valid if it is also being added this in this update, in which case it will get processed in the post-update anyway.
-			if (Id.IsValid())
-			{
-				SceneUpdateChangeSetStorage.UpdatedPrimitiveIds.Add(Id);
-				SceneUpdateChangeSetStorage.UpdatedPrimitiveSceneInfos.Add(PrimitiveSceneInfo);
-			}
-		}
-		// As must all primitive updates, 
-		for (const auto& Transform : UpdatedTransforms)
-		{
-			FPrimitiveSceneInfo* PrimitiveSceneInfo = Transform.Key->GetPrimitiveSceneInfo();
-
-			if (UpdatedInstances.Find(Transform.Key) != nullptr)
-			{
-				continue;
-			}
-
-			if (RemovedPrimitiveSceneInfos.Find(PrimitiveSceneInfo) != nullptr)
-			{
-				continue;
-			}
-
-			FPersistentPrimitiveIndex Id = Transform.Key->GetPrimitiveSceneInfo()->GetPersistentIndex();
-			// Note: may not be valid if it is also being added this in this update, in which case it will get processed in the post-update anyway.
-			if (Id.IsValid())
-			{
-				SceneUpdateChangeSetStorage.UpdatedPrimitiveIds.Add(Id);
-				SceneUpdateChangeSetStorage.UpdatedPrimitiveSceneInfos.Add(PrimitiveSceneInfo);
-			}
-		}
-	}
-
-	TRACE_COUNTER_SET(ScenePrimitivesAdds, AddedPrimitiveSceneInfos.Num());
-	TRACE_COUNTER_SET(ScenePrimitivesRemoves, RemovedPrimitiveSceneInfos.Num());
-	TRACE_COUNTER_SET(ScenePrimitivesUpdatesTransforms, UpdatedInstances.Num());
-	TRACE_COUNTER_SET(ScenePrimitivesUpdatesInstances, UpdatedTransforms.Num());
-
-	// NOTE: We clear this early because IsPrimitiveBeingRemoved gets called from the CreateLightPrimitiveInteraction (to make sure that old primitives are not accessed) 
-	// we cannot safely kick off the AsyncCreateLightPrimitiveInteractionsTask before the RemovedPrimitiveSceneInfos has been cleared.
-	// TODO: this is probably not true anymore!
-	RemovedPrimitiveSceneInfos.Empty();
+	auto UpdatedInstances = SceneUpdateChangeSetStorage.PrimitiveUpdates.GetRangeView<FUpdateInstanceCommand>();
+	auto UpdatedTransforms = SceneUpdateChangeSetStorage.PrimitiveUpdates.GetRangeView<FUpdateTransformCommand>();
+	auto UpdatedInstanceCullDistance =  SceneUpdateChangeSetStorage.PrimitiveUpdates.GetRangeView<FUpdateInstanceCullDistanceData>();
+	auto OverridenPreviousTransforms =  SceneUpdateChangeSetStorage.PrimitiveUpdates.GetRangeView<FUpdateOverridePreviousTransformData>();	
 	bool bAnySceneUpdatesQueued = RemovedLocalPrimitiveSceneInfos.Num() + AddedPrimitiveSceneInfos.Num() + UpdatedTransforms.Num() + UpdatedInstances.Num() > 0;
 	RemovedLocalPrimitiveSceneInfos.Sort(FPrimitiveArraySortKey());
 
@@ -5715,17 +5492,17 @@ void FScene::Update(FRDGBuilder& GraphBuilder, const FUpdateParameters& Paramete
 		// All updated instances must also before moving or re-allocating (TODO: filter out only those actually updated)
 		for (const auto& Instance : UpdatedInstances)
 		{
-			InvalidatingPrimitiveCollector.UpdatedTransform(Instance.Key->GetPrimitiveSceneInfo());
+			InvalidatingPrimitiveCollector.UpdatedTransform(Instance.SceneInfo);
 		}
 		// As must all primitive updates, 
 		for (const auto& Transform : UpdatedTransforms)
 		{
-			InvalidatingPrimitiveCollector.UpdatedTransform(Transform.Key->GetPrimitiveSceneInfo());
+			InvalidatingPrimitiveCollector.UpdatedTransform(Transform.SceneInfo);
 		}
 
 		for (const auto& CullDistance : UpdatedInstanceCullDistance)
 		{
-			InvalidatingPrimitiveCollector.UpdatedTransform(CullDistance.Key->GetPrimitiveSceneInfo());
+			InvalidatingPrimitiveCollector.UpdatedTransform(CullDistance.SceneInfo);
 		}
 
 		TRACE_INT_VALUE(TEXT("Shadow.Virtual.Cache.PreInvalidationInstances"), InvalidatingPrimitiveCollector.Instances.GetTotalNumInstances());
@@ -5733,9 +5510,6 @@ void FScene::Update(FRDGBuilder& GraphBuilder, const FUpdateParameters& Paramete
 	}
 
 	AddedLocalPrimitiveSceneInfos.Sort(FPrimitiveArraySortKey());
-
-	DeletedPrimitiveSceneInfos.Reserve(RemovedLocalPrimitiveSceneInfos.Num());
-
 	TArray<int32> RemovedPrimitiveIndices;
 	RemovedPrimitiveIndices.SetNumUninitialized(RemovedLocalPrimitiveSceneInfos.Num());
 
@@ -5941,8 +5715,6 @@ void FScene::Update(FRDGBuilder& GraphBuilder, const FUpdateParameters& Paramete
 
 				bNeedPathTracedInvalidation = bNeedPathTracedInvalidation || IsPrimitiveRelevantToPathTracing(PrimitiveSceneInfo);
 
-				DeletedPrimitiveSceneInfos.Emplace(PrimitiveSceneInfo);
-
 				const int32 PersistentIndex = PrimitiveSceneInfo->PersistentIndex.Index;
 				PersistentPrimitiveIdAllocator.Free(PersistentIndex);
 				PersistentPrimitiveIdToIndexMap[PersistentIndex] = INDEX_NONE;
@@ -6010,15 +5782,9 @@ void FScene::Update(FRDGBuilder& GraphBuilder, const FUpdateParameters& Paramete
 
 		for (const auto& UpdateInstance : UpdatedInstances)
 		{
-			FPrimitiveSceneProxy* PrimitiveSceneProxy = UpdateInstance.Key;
+			FPrimitiveSceneProxy* PrimitiveSceneProxy = UpdateInstance.SceneInfo->Proxy;
 			FScopeCycleCounter Context(PrimitiveSceneProxy->GetStatId());
-			FPrimitiveSceneInfo* PrimitiveSceneInfo = PrimitiveSceneProxy->GetPrimitiveSceneInfo();
-			
-			// being added or deleted, skip update logic
-			if (DeletedPrimitiveSceneInfos.Contains(PrimitiveSceneInfo))
-			{
-				continue;
-			}
+			FPrimitiveSceneInfo* PrimitiveSceneInfo = UpdateInstance.SceneInfo;
 
 			const FInstanceDataBufferHeader &InstanceDataBufferHeader = PrimitiveSceneInfo->GetInstanceDataHeader();
 			const bool bInstanceCountChanged = PrimitiveSceneInfo->GetNumInstanceSceneDataEntries() != InstanceDataBufferHeader.NumInstances;
@@ -6041,6 +5807,21 @@ void FScene::Update(FRDGBuilder& GraphBuilder, const FUpdateParameters& Paramete
 
 		PersistentPrimitiveIdAllocator.Consolidate();
 
+		SceneUpdateChangeSetStorage.PrimitiveUpdates.ForEachCommand(ESceneUpdateCommandFilter::Added, [&](FPrimitiveUpdateCommand& Cmd)
+		{
+			// Skip those that were added & removed in the same frame
+			if (Cmd.IsAdd() && !Cmd.IsDelete())
+			{
+				check(!Cmd.GetPersistentId().IsValid());
+				check(!Cmd.GetSceneInfo()->PersistentIndex.IsValid());
+				FPersistentPrimitiveIndex PersistentPrimitiveIndex{ PersistentPrimitiveIdAllocator.Allocate() };
+				// Note: propagate to the command such that is it valid for subsequent passes without having to hit the primitive scene info.
+				Cmd.SetPersistentId(PersistentPrimitiveIndex);
+				Cmd.GetSceneInfo()->PersistentIndex = PersistentPrimitiveIndex;
+			}
+		});
+		PersistentPrimitiveIdToIndexMap.SetNumUninitialized(GetMaxPersistentPrimitiveIndex());
+
 		if (AddedLocalPrimitiveSceneInfos.Num())
 		{
 			SCOPED_NAMED_EVENT(FScene_AddPrimitiveSceneInfos_ReserveMemory, FColor::Green);
@@ -6107,15 +5888,8 @@ void FScene::Update(FRDGBuilder& GraphBuilder, const FUpdateParameters& Paramete
 
 					const int32 SourceIndex = PrimitiveSceneProxies.Num() - 1;
 					PrimitiveSceneInfo->PackedIndex = SourceIndex;
-					checkSlow(PrimitiveSceneInfo->PersistentIndex.Index == INDEX_NONE);
-					FPersistentPrimitiveIndex PersistentPrimitiveIndex{ PersistentPrimitiveIdAllocator.Allocate() };
-					PrimitiveSceneInfo->PersistentIndex = PersistentPrimitiveIndex;
-					// Ensure map is large enough
-					if (PersistentPrimitiveIndex.Index >= PersistentPrimitiveIdToIndexMap.Num())
-					{
-						PersistentPrimitiveIdToIndexMap.SetNumUninitialized(PersistentPrimitiveIndex.Index + 1);
-					}
-					PersistentPrimitiveIdToIndexMap[PersistentPrimitiveIndex.Index] = SourceIndex;
+					check(PrimitiveSceneInfo->PersistentIndex.IsValid());
+					PersistentPrimitiveIdToIndexMap[PrimitiveSceneInfo->PersistentIndex.Index] = SourceIndex;
 				}
 			}
 
@@ -6279,19 +6053,18 @@ void FScene::Update(FRDGBuilder& GraphBuilder, const FUpdateParameters& Paramete
 		SCOPED_NAMED_EVENT(FScene_AddPrimitiveSceneInfos, FColor::Yellow);
 		SCOPE_CYCLE_COUNTER(STAT_UpdatePrimitiveTransformRenderThreadTime);
 
-		for (const auto& Transform : UpdatedTransforms)
+
+		for (const auto& Item : UpdatedTransforms)
 		{
-			FPrimitiveSceneProxy* PrimitiveSceneProxy = Transform.Key;
-			if (DeletedPrimitiveSceneInfos.Contains(PrimitiveSceneProxy->GetPrimitiveSceneInfo()))
-			{
-				continue;
-			}
+			FPrimitiveSceneProxy* PrimitiveSceneProxy = Item.SceneInfo->Proxy;
+
+			const FBoxSphereBounds& WorldBounds = Item.Payload.WorldBounds;
+			const FBoxSphereBounds& LocalBounds = Item.Payload.LocalBounds;
+			const FMatrix& LocalToWorld = Item.Payload.LocalToWorld;
+			const FVector& AttachmentRootPosition = Item.Payload.AttachmentRootPosition;
+
 			check(PrimitiveSceneProxy->GetPrimitiveSceneInfo()->PackedIndex != INDEX_NONE);
 
-			const FBoxSphereBounds& WorldBounds = Transform.Value.WorldBounds;
-			const FBoxSphereBounds& LocalBounds = Transform.Value.LocalBounds;
-			const FMatrix& LocalToWorld = Transform.Value.LocalToWorld;
-			const FVector& AttachmentRootPosition = Transform.Value.AttachmentRootPosition;
 			FScopeCycleCounter Context(PrimitiveSceneProxy->GetStatId());
 
 			FPrimitiveSceneInfo* PrimitiveSceneInfo = PrimitiveSceneProxy->GetPrimitiveSceneInfo();
@@ -6344,8 +6117,8 @@ void FScene::Update(FRDGBuilder& GraphBuilder, const FUpdateParameters& Paramete
 
 		for (const auto& Transform : OverridenPreviousTransforms)
 		{
-			FPrimitiveSceneInfo* PrimitiveSceneInfo = Transform.Key;
-			VelocityData.OverridePreviousTransform(PrimitiveSceneInfo->PrimitiveComponentId, Transform.Value);
+			FPrimitiveSceneInfo* PrimitiveSceneInfo = Transform.SceneInfo;
+			VelocityData.OverridePreviousTransform(PrimitiveSceneInfo->PrimitiveComponentId, Transform.Payload.Value);
 		}
 	}
 
@@ -6362,18 +6135,12 @@ void FScene::Update(FRDGBuilder& GraphBuilder, const FUpdateParameters& Paramete
 		SCOPED_NAMED_EVENT(FScene_UpdatePrimitiveInstances, FColor::Emerald);
 		SCOPE_CYCLE_COUNTER(STAT_UpdatePrimitiveInstanceRenderThreadTime);
 
-		for (const auto& UpdateInstance : UpdatedInstances)
+		for (const auto& Item : UpdatedInstances)
 		{
-			FPrimitiveSceneProxy* PrimitiveSceneProxy = UpdateInstance.Key;
+			FPrimitiveSceneProxy* PrimitiveSceneProxy = Item.SceneInfo->Proxy;
+			FPrimitiveSceneInfo* PrimitiveSceneInfo = Item.SceneInfo;
+			const FUpdateInstanceCommand &UpdateInstance = Item.Payload;
 			FScopeCycleCounter Context(PrimitiveSceneProxy->GetStatId());
-			FPrimitiveSceneInfo* PrimitiveSceneInfo = PrimitiveSceneProxy->GetPrimitiveSceneInfo();
-			
-			// being added or deleted, skip update logic
-			if (DeletedPrimitiveSceneInfos.Contains(PrimitiveSceneInfo))
-			{
-				continue;
-			}
-
 			
 			QueueFlushVirtualTexture(PrimitiveSceneInfo);
 
@@ -6407,7 +6174,7 @@ void FScene::Update(FRDGBuilder& GraphBuilder, const FUpdateParameters& Paramete
 #endif
 
 			// Update the Proxy's data.
-			PrimitiveSceneProxy->UpdateInstances_RenderThread(GraphBuilder.RHICmdList, UpdateInstance.Value.WorldBounds, UpdateInstance.Value.LocalBounds, UpdateInstance.Value.StaticMeshBounds);
+			PrimitiveSceneProxy->UpdateInstances_RenderThread(GraphBuilder.RHICmdList, UpdateInstance.WorldBounds, UpdateInstance.LocalBounds, UpdateInstance.StaticMeshBounds);
 
 			if (!RHISupportsVolumeTextures(GetFeatureLevel())
 				&& (PrimitiveSceneProxy->IsMovable() || PrimitiveSceneProxy->NeedsUnbuiltPreviewLighting() || PrimitiveSceneProxy->GetLightmapType() == ELightmapType::ForceVolumetric))
@@ -6570,11 +6337,9 @@ void FScene::Update(FRDGBuilder& GraphBuilder, const FUpdateParameters& Paramete
 	}
 	
 	SceneUpdateChangeSetStorage.AddedPrimitiveIds.Reserve(AddedPrimitiveSceneInfos.Num());
-	SceneUpdateChangeSetStorage.AddedPrimitiveSceneInfos.Reserve(AddedPrimitiveSceneInfos.Num());
 	for (FPrimitiveSceneInfo* PrimitiveSceneInfo : AddedPrimitiveSceneInfos)
 	{
 		SceneUpdateChangeSetStorage.AddedPrimitiveIds.Add(PrimitiveSceneInfo->GetPersistentIndex());
-		SceneUpdateChangeSetStorage.AddedPrimitiveSceneInfos.Add(PrimitiveSceneInfo);
 	}
 
 	SceneCullingUpdater.OnPostSceneUpdate(GraphBuilder, SceneUpdateChangeSetStorage.GetPostUpdateSet());
@@ -6752,19 +6517,11 @@ void FScene::Update(FRDGBuilder& GraphBuilder, const FUpdateParameters& Paramete
 #endif
 	}
 
-	for (const auto& CustomParams : UpdatedCustomPrimitiveParams)
+	for (const auto& Item : SceneUpdateChangeSetStorage.PrimitiveUpdates.GetRangeView<FUpdateCustomPrimitiveData>())
 	{
-		FPrimitiveSceneProxy* PrimitiveSceneProxy = CustomParams.Key;
-		if (DeletedPrimitiveSceneInfos.Contains(PrimitiveSceneProxy->GetPrimitiveSceneInfo()))
-		{
-			continue;
-		}
-
+		FPrimitiveSceneProxy* PrimitiveSceneProxy = Item.SceneInfo->Proxy;
 		FScopeCycleCounter Context(PrimitiveSceneProxy->GetStatId());
-		PrimitiveSceneProxy->CustomPrimitiveData = CustomParams.Value;
-
-		// Ensure an update of primitive data before rendering
-		PrimitiveSceneProxy->GetPrimitiveSceneInfo()->MarkGPUStateDirty(EPrimitiveDirtyState::ChangedOther);
+		PrimitiveSceneProxy->CustomPrimitiveData = Item.Payload.Value;
 	}
 
 	if (auto NaniteMaterialsUpdater = SceneExtensionsUpdaters.GetUpdaterPtr<Nanite::FMaterialsSceneExtension::FUpdater>())
@@ -6781,38 +6538,29 @@ void FScene::Update(FRDGBuilder& GraphBuilder, const FUpdateParameters& Paramete
 		SceneLODHierarchy.UpdateNodeSceneInfo(PrimitiveSceneInfo->PrimitiveComponentId, PrimitiveSceneInfo);
 	}
 
-	for (const auto& Attachments : UpdatedAttachmentRoots)
-	{
-		FPrimitiveSceneInfo* PrimitiveSceneInfo = Attachments.Key;
-		if (DeletedPrimitiveSceneInfos.Contains(PrimitiveSceneInfo))
-		{
-			continue;
-		}
+	auto UpdatedAttachmentRoots = SceneUpdateChangeSetStorage.PrimitiveUpdates.GetRangeView<FUpdateAttachmentRootData>();
+	auto DistanceFieldSceneDataUpdates = SceneUpdateChangeSetStorage.PrimitiveUpdates.GetRangeView<FUpdateDistanceFieldSceneData>();
+	auto UpdatedOcclusionBoundsSlacks = SceneUpdateChangeSetStorage.PrimitiveUpdates.GetRangeView<FUpdateOcclusionBoundsSlacksData>();
+	auto UpdatedDrawDistance =  SceneUpdateChangeSetStorage.PrimitiveUpdates.GetRangeView<FUpdateDrawDistanceData>();
 
+	for (const auto& Item : UpdatedAttachmentRoots)
+	{
+		FPrimitiveSceneInfo* PrimitiveSceneInfo = Item.SceneInfo;
 		PrimitiveSceneInfo->UnlinkAttachmentGroup();
-		PrimitiveSceneInfo->LightingAttachmentRoot = Attachments.Value;
+		PrimitiveSceneInfo->LightingAttachmentRoot = Item.Payload.Value;
 		PrimitiveSceneInfo->LinkAttachmentGroup();
 	}
 
-	for (FPrimitiveSceneInfo* PrimitiveSceneInfo : DistanceFieldSceneDataUpdates)
+	for (const auto& Item : DistanceFieldSceneDataUpdates)
 	{
-		if (DeletedPrimitiveSceneInfos.Contains(PrimitiveSceneInfo))
-		{
-			continue;
-		}
-
-		DistanceFieldSceneData.UpdatePrimitive(PrimitiveSceneInfo);
+		DistanceFieldSceneData.UpdatePrimitive(Item.SceneInfo);
 	}
 
 	for (const auto& OccSlackDelta : UpdatedOcclusionBoundsSlacks)
 	{
-		const FPrimitiveSceneProxy* SceneProxy = OccSlackDelta.Key;
-		const FPrimitiveSceneInfo* SceneInfo = SceneProxy->GetPrimitiveSceneInfo();
-
-		if (DeletedPrimitiveSceneInfos.Contains(SceneInfo))
-		{
-			continue;
-		}
+		const FPrimitiveSceneProxy* SceneProxy = OccSlackDelta.SceneInfo->Proxy;
+		const FPrimitiveSceneInfo* SceneInfo = OccSlackDelta.SceneInfo;
+		float Delta = OccSlackDelta.Payload.Value;
 
 		FBoxSphereBounds NewOccBounds;
 		if (SceneProxy->HasCustomOcclusionBounds())
@@ -6824,39 +6572,25 @@ void FScene::Update(FRDGBuilder& GraphBuilder, const FUpdateParameters& Paramete
 			NewOccBounds = SceneProxy->GetBounds();
 		}
 
-		PrimitiveOcclusionBounds[SceneInfo->PackedIndex] = NewOccBounds.ExpandBy(OCCLUSION_SLOP + OccSlackDelta.Value);
+		PrimitiveOcclusionBounds[SceneInfo->PackedIndex] = NewOccBounds.ExpandBy(OCCLUSION_SLOP + Delta);
 	}
 
-	for (auto& CullDistance : UpdatedInstanceCullDistance)
+	for (const auto& CullDistance : UpdatedInstanceCullDistance)
 	{
-		FPrimitiveSceneProxy* SceneProxy = CullDistance.Key;
-		FPrimitiveSceneInfo* SceneInfo = SceneProxy->GetPrimitiveSceneInfo();
-
-		if (DeletedPrimitiveSceneInfos.Contains(SceneInfo))
-		{
-			continue;
-		}
-
-		float StartCullDistance = CullDistance.Value.X;
-		float EndCullDistance = CullDistance.Value.Y;
+		FPrimitiveSceneProxy* SceneProxy = CullDistance.SceneInfo->Proxy;
+		float StartCullDistance = CullDistance.Payload.Value.X;
+		float EndCullDistance = CullDistance.Payload.Value.Y;
 		
 		SceneProxy->SetInstanceCullDistance_RenderThread(StartCullDistance, EndCullDistance);
-		SceneInfo->MarkGPUStateDirty(EPrimitiveDirtyState::ChangedOther);
 	}
 
-	for (auto& DrawDistance : UpdatedDrawDistance)
+	for (const auto& DrawDistance : UpdatedDrawDistance)
 	{
-		FPrimitiveSceneProxy* SceneProxy = DrawDistance.Key;
-		FPrimitiveSceneInfo* SceneInfo = SceneProxy->GetPrimitiveSceneInfo();
-
-		if (DeletedPrimitiveSceneInfos.Contains(SceneInfo))
-		{
-			continue;
-		}
-
-		float MinDrawDistance = DrawDistance.Value.X;
-		float MaxDrawDistance = DrawDistance.Value.Y;
-		float VirtualTextureMaxDrawDistance = DrawDistance.Value.Z;
+		FPrimitiveSceneProxy* SceneProxy = DrawDistance.SceneInfo->Proxy;
+		FPrimitiveSceneInfo* SceneInfo = DrawDistance.SceneInfo;
+		float MinDrawDistance = DrawDistance.Payload.Value.X;
+		float MaxDrawDistance = DrawDistance.Payload.Value.Y;
+		float VirtualTextureMaxDrawDistance = DrawDistance.Payload.Value.Z;
 
 		SceneProxy->SetDrawDistance_RenderThread(MinDrawDistance, MaxDrawDistance, VirtualTextureMaxDrawDistance);
 
@@ -6947,16 +6681,6 @@ void FScene::Update(FRDGBuilder& GraphBuilder, const FUpdateParameters& Paramete
 	{
 		InvalidatePathTracedOutput();
 	}
-	UpdatedAttachmentRoots.Empty();
-	UpdatedTransforms.Empty();
-	UpdatedInstances.Empty();
-	UpdatedCustomPrimitiveParams.Empty();
-	OverridenPreviousTransforms.Empty();
-	UpdatedOcclusionBoundsSlacks.Empty();
-	UpdatedInstanceCullDistance.Empty();
-	UpdatedDrawDistance.Empty();
-	DistanceFieldSceneDataUpdates.Empty();
-	AddedPrimitiveSceneInfos.Empty();
 	LevelCommands.Empty();
 
 	LastUpdateFrameCounter = GFrameCounterRenderThread;
@@ -6976,12 +6700,6 @@ void FScene::Update(FRDGBuilder& GraphBuilder, const FUpdateParameters& Paramete
 #endif
 }
 
-bool FScene::IsPrimitiveBeingRemoved(FPrimitiveSceneInfo* PrimitiveSceneInfo) const
-{
-	check(IsInParallelRenderingThread() || IsInRenderingThread());
-	return RemovedPrimitiveSceneInfos.Find(PrimitiveSceneInfo) != nullptr;
-}
-
 bool FScene::ShouldRenderSkylightInBasePass(bool bIsTranslucent) const
 {
 	if (IsMobilePlatform(GetShaderPlatform()))
@@ -7308,20 +7026,20 @@ void UpdateStaticMeshesForMaterials(const TArray<const FMaterial*>& MaterialReso
 	UE::RenderCommandPipe::FSyncScope SyncScope;
 
 	ENQUEUE_RENDER_COMMAND(FUpdateStaticMeshesForMaterials)(
-		[UsedPrimitives = MoveTemp(UsedPrimitives)](FRHICommandListImmediate& RHICmdList) mutable
+	[UsedPrimitives = MoveTemp(UsedPrimitives)](FRHICommandListImmediate& RHICmdList) mutable
+	{
+		// Defer the caching until the next render tick, to make sure that all render components queued
+		// for re-creation are processed. Otherwise, we may end up caching mesh commands from stale data.
+		for (auto& SceneInfos: UsedPrimitives)
 		{
-			// Defer the caching until the next render tick, to make sure that all render components queued
-			// for re-creation are processed. Otherwise, we may end up caching mesh commands from stale data.
-			for (auto& SceneInfos: UsedPrimitives)
-			{
-				SceneInfos.Key->UpdateAllPrimitiveSceneInfos(RHICmdList);
-			}
-			for (auto& SceneInfos : UsedPrimitives)
-			{
-				TArray<FPrimitiveSceneInfo*>& SceneInfoArray = SceneInfos.Value;
-				FPrimitiveSceneInfo::UpdateStaticMeshes(SceneInfos.Key, SceneInfoArray, EUpdateStaticMeshFlags::AllCommands, false);
-			}
-		});
+			SceneInfos.Key->UpdateAllPrimitiveSceneInfos(RHICmdList);
+		}
+		for (auto& SceneInfos : UsedPrimitives)
+		{
+			TArray<FPrimitiveSceneInfo*>& SceneInfoArray = SceneInfos.Value;
+			FPrimitiveSceneInfo::UpdateStaticMeshes(SceneInfos.Key, SceneInfoArray, EUpdateStaticMeshFlags::AllCommands, false);
+		}
+	});
 }
 
 void FRendererModule::UpdateStaticDrawListsForMaterials(const TArray<const FMaterial*>& Materials)
diff --git a/Engine/Source/Runtime/Renderer/Private/SceneCulling/SceneCulling.cpp b/Engine/Source/Runtime/Renderer/Private/SceneCulling/SceneCulling.cpp
--- a/Engine/Source/Runtime/Renderer/Private/SceneCulling/SceneCulling.cpp
+++ b/Engine/Source/Runtime/Renderer/Private/SceneCulling/SceneCulling.cpp
@@ -2420,21 +2420,17 @@ public:
 		{
 			SCOPED_NAMED_EVENT(SceneCulling_Post_UpdateInstances, FColor::Emerald);
 
-			// [transform-] updated primitives instances & primitives with updated instances 
-			for (int32 Index = 0; Index < ScenePostUpdateData.UpdatedPrimitiveIds.Num(); ++Index)
+			ScenePostUpdateData.PrimitiveUpdates.ForEachUpdateCommand(ESceneUpdateCommandFilter::AddedUpdated, EPrimitiveUpdateDirtyFlags::AllCulling, [&](const FPrimitiveUpdateCommand& Cmd)
 			{
-				UpdateInstances(ScenePostUpdateData.UpdatedPrimitiveIds[Index], ScenePostUpdateData.UpdatedPrimitiveSceneInfos[Index]);
-			}	
-		}
-
-		{
-			SCOPED_NAMED_EVENT(SceneCulling_Post_AddInstances, FColor::Emerald);
-
-			// Next process all added and added ones.
-			for (int32 Index = 0; Index < ScenePostUpdateData.AddedPrimitiveIds.Num(); ++Index)
-			{
-				AddInstances(ScenePostUpdateData.AddedPrimitiveIds[Index], ScenePostUpdateData.AddedPrimitiveSceneInfos[Index]);
-			}
+				if (Cmd.IsAdd())
+				{
+					AddInstances(Cmd.GetPersistentId(), Cmd.GetSceneInfo());
+				}
+				else
+				{
+					UpdateInstances(Cmd.GetPersistentId(), Cmd.GetSceneInfo());
+				}
+			});
 		}
 		FinalizeTempCellsAndUncullable();
 
diff --git a/Engine/Source/Runtime/Renderer/Private/ScenePrivate.h b/Engine/Source/Runtime/Renderer/Private/ScenePrivate.h
--- a/Engine/Source/Runtime/Renderer/Private/ScenePrivate.h
+++ b/Engine/Source/Runtime/Renderer/Private/ScenePrivate.h
@@ -59,6 +59,7 @@
 #include "LightFunctionAtlas.h"
 #include "SceneExtensions.h"
 #include "HeterogeneousVolumes/HeterogeneousVolumes.h"
+#include "SceneUpdateCommandQueue.h"
 
 /** Factor by which to grow occlusion tests **/
 #define OCCLUSION_SLOP (1.0f)
@@ -94,54 +95,8 @@ class FSparseVolumeTextureViewerSceneProxy;
 class FExponentialHeightFogSceneInfo;
 class FStaticMeshBatch;
 class FShadowScene;
-class FSceneLightInfoUpdates;
 class FSceneCulling;
 
-/**
- * Describes all light modifications to the scene by recording the light scene IDs.
- * TODO: If needed, we could add a reference to the FLightUpdates (which contains the commands) since this would enable systems to consume out the delta updates as they come in.
- *       If this is useful we must ensure FLightUpdates are kept alive as long as any async tasks might require.
- */
-struct FLightSceneChangeSet
-{
-	// IDs of all lights before they were removed, IDs in this array may not be valid at all times when the change-set is used (depends on whether the callback site is before or after the given lights are removed from the scene).
-	TConstArrayView<int32> RemovedLightIds;
-	// IDs of all lights added to the scene, only available after all lights are added to the scene, may contain the same ID's as removed, as they may be reused.
-	TConstArrayView<int32> AddedLightIds;
-	// IDs of updated lights, does not contain any from the above, since 'add' implies the update of all aspects and 'remove' implies cancellation of all updates. 
-	// The updated arrays are not disjoint as a light may have both types of update applied.
-	TConstArrayView<int32> TransformUpdatedLightIds;
-	TConstArrayView<int32> ColorUpdatedLightIds;
-};
-
-/**
- * Change set that is valid before removes are processed and the scene data modified.
- * The referenced arrays have RDG life-time and can be safely used in RDG tasks.
- * However, the referenced data (primitive/proxy) and meaning of the persistent ID is not generally valid past the call in which this is passed. 
- * Thus, care need to be excercised.
- */
-class FScenePreUpdateChangeSet
-{
-public:
-	TConstArrayView<FPersistentPrimitiveIndex> RemovedPrimitiveIds;
-	TConstArrayView<FPrimitiveSceneInfo*> RemovedPrimitiveSceneInfos;
-	TConstArrayView<FPersistentPrimitiveIndex> UpdatedPrimitiveIds;
-	TConstArrayView<FPrimitiveSceneInfo*> UpdatedPrimitiveSceneInfos;
-};
-
-/**
- * Change set that is valid before after adds are processed and the scene data is modified.
- * The referenced arrays have RDG life-time and can be safely used in RDG tasks.
- */
-class FScenePostUpdateChangeSet
-{
-public:
-	TConstArrayView<FPersistentPrimitiveIndex> AddedPrimitiveIds;
-	TConstArrayView<FPrimitiveSceneInfo*> AddedPrimitiveSceneInfos;
-	TConstArrayView<FPersistentPrimitiveIndex> UpdatedPrimitiveIds;
-	TConstArrayView<FPrimitiveSceneInfo*> UpdatedPrimitiveSceneInfos;
-};
-
 /** Holds information about a single primitive's occlusion. */
 class FPrimitiveOcclusionHistory
 {
@@ -2885,6 +2840,173 @@ private:
 	FLumenSceneDataMap::TConstIterator NextSceneData;
 };
 
+/**
+ * Definitions of primitive scene update commands.
+ */
+
+enum class EPrimitiveUpdateDirtyFlags : uint32 
+{
+	None			=	0u,
+	/** The Transform is modified by this command. */
+	Transform		=	1u << 0,
+	/** The (any) instance data is modified by this command. */
+	InstanceData	=	1u << 1,
+	/** 
+	 * The culling bounds are modified by this command. 
+	 * This means the bounds (instance, primitive or both) as used in the culling and should not be updated for any other case. 
+	 * Thus, needs to be set for transform updates of all kinds.
+	 */
+	CullingBounds	=	1u << 2,
+	/** 
+	 * Culling distances or similar, affecting culling logic, but NOT the bounds.
+	 */
+	CullingLogic	=	1u << 3,
+	/** Any state that either makes its way into GPU-Scene or the per primitive UB */
+	GPUState		=	1u << 4,
+	/** All culling-affecting changes */
+	AllCulling		= CullingBounds | CullingLogic,
+	All				=	GPUState |  Transform | InstanceData | CullingBounds | CullingLogic,
+};
+ENUM_CLASS_FLAGS(EPrimitiveUpdateDirtyFlags);
+
+enum class EPrimitiveUpdateId : uint32 
+{
+	UpdateTransform,
+	UpdateInstance,
+	UpdateAttachmentRoot,
+	CustomPrimitiveData,
+	OcclusionBoundsSlacks,  
+	InstanceCullDistance,  
+	DrawDistance,
+	DistanceFieldScene,
+	OverridePreviousTransform,
+	MAX
+};
+
+using FScenePrimitiveUpdates = TSceneUpdateCommandQueue<FPrimitiveSceneInfo, EPrimitiveUpdateDirtyFlags, EPrimitiveUpdateId>;
+
+using FPrimitiveUpdateCommand = FScenePrimitiveUpdates::FUpdateCommand;
+
+template <EPrimitiveUpdateId InId, EPrimitiveUpdateDirtyFlags InDirtyFlags>
+using TPrimitiveUpdatePayloadBase = FScenePrimitiveUpdates::TPayloadBase<InId, InDirtyFlags>;
+
+struct FUpdateTransformCommand : public TPrimitiveUpdatePayloadBase<EPrimitiveUpdateId::UpdateTransform, 
+	EPrimitiveUpdateDirtyFlags::GPUState | EPrimitiveUpdateDirtyFlags::Transform | EPrimitiveUpdateDirtyFlags::CullingBounds>
+{
+	FBoxSphereBounds WorldBounds;
+	FBoxSphereBounds LocalBounds; 
+	FMatrix LocalToWorld; 
+	FVector AttachmentRootPosition;
+};
+
+struct FUpdateInstanceCommand : public TPrimitiveUpdatePayloadBase<EPrimitiveUpdateId::UpdateInstance, 
+	EPrimitiveUpdateDirtyFlags::GPUState | EPrimitiveUpdateDirtyFlags::Transform | EPrimitiveUpdateDirtyFlags::CullingBounds | EPrimitiveUpdateDirtyFlags::InstanceData>
+{
+	FPrimitiveSceneProxy* PrimitiveSceneProxy{ nullptr };
+	FBoxSphereBounds WorldBounds;
+	FBoxSphereBounds LocalBounds;
+	FBoxSphereBounds StaticMeshBounds;
+};
+
+
+/**
+ * Helper for the update payloads that contain a single payload value.
+ */
+template <typename InPayloadDataType, EPrimitiveUpdateId InId, EPrimitiveUpdateDirtyFlags InPrimitiveDirtyFlags>
+struct TSingleValuePrimitiveUpdatePayload : public TPrimitiveUpdatePayloadBase<InId, InPrimitiveDirtyFlags>
+{
+	TSingleValuePrimitiveUpdatePayload(const InPayloadDataType &InValue) : Value(InValue) {}
+	InPayloadDataType Value;
+};
+
+using FUpdateAttachmentRootData = TSingleValuePrimitiveUpdatePayload<FPrimitiveComponentId, EPrimitiveUpdateId::UpdateAttachmentRoot,
+	EPrimitiveUpdateDirtyFlags::None>; // No GPU side effect (?).
+using FUpdateCustomPrimitiveData = TSingleValuePrimitiveUpdatePayload<FCustomPrimitiveData, EPrimitiveUpdateId::CustomPrimitiveData,
+	EPrimitiveUpdateDirtyFlags::GPUState>; // Needs upload
+using FUpdateOcclusionBoundsSlacksData = TSingleValuePrimitiveUpdatePayload<float, EPrimitiveUpdateId::OcclusionBoundsSlacks,  
+	EPrimitiveUpdateDirtyFlags::None>; // Only affects primitive occlusion.
+using FUpdateInstanceCullDistanceData = TSingleValuePrimitiveUpdatePayload<FVector2f, EPrimitiveUpdateId::InstanceCullDistance,  
+	EPrimitiveUpdateDirtyFlags::GPUState | EPrimitiveUpdateDirtyFlags::CullingLogic>; // Affects GPU culling?
+using FUpdateDrawDistanceData = TSingleValuePrimitiveUpdatePayload<FVector3f, EPrimitiveUpdateId::DrawDistance,  
+	EPrimitiveUpdateDirtyFlags::CullingLogic>; // Only affects CPU culling.
+using FUpdateDistanceFieldSceneData = TPrimitiveUpdatePayloadBase<EPrimitiveUpdateId::DistanceFieldScene,  
+	EPrimitiveUpdateDirtyFlags::None>; // Only affects DF scene rep - candidate for using abstract type.
+using FUpdateOverridePreviousTransformData = TSingleValuePrimitiveUpdatePayload<FMatrix, EPrimitiveUpdateId::OverridePreviousTransform,  
+	EPrimitiveUpdateDirtyFlags::GPUState>;// Overrides the previous transform, which needs to be propagated to the GPU, but otherwise does not change anything on its own.
+
+/**
+ * Change set that is valid before removes are processed and the scene data modified.
+ * The referenced arrays have RDG life-time and can be safely used in RDG tasks.
+ * However, the referenced data (primitive/proxy) and meaning of the persistent ID is not generally valid past the call in which this is passed. 
+ * Thus, care need to be excercised.
+ */
+class FScenePreUpdateChangeSet
+{
+public:
+	TConstArrayView<FPersistentPrimitiveIndex> RemovedPrimitiveIds;
+	TConstArrayView<FPrimitiveSceneInfo*> RemovedPrimitiveSceneInfos;
+	const FScenePrimitiveUpdates &PrimitiveUpdates;
+};
+
+/**
+ * Change set that is valid before after adds are processed and the scene data is modified.
+ * The referenced arrays have RDG life-time and can be safely used in RDG tasks.
+ */
+class FScenePostUpdateChangeSet
+{
+public:
+	TConstArrayView<FPersistentPrimitiveIndex> AddedPrimitiveIds;
+	TConstArrayView<FPrimitiveSceneInfo*> AddedPrimitiveSceneInfos;
+	const FScenePrimitiveUpdates &PrimitiveUpdates;
+};
+
+/**
+ * Definitions for light scene updates.
+ */
+
+enum class ELightDirtyFlags : uint32 
+{
+	None			=	0u
+};
+ENUM_CLASS_FLAGS(ELightDirtyFlags);
+
+enum class ELightUpdateId : uint32 
+{
+	Transform,
+	Color,
+	MAX
+};
+
+using FSceneLightInfoUpdates = TSceneUpdateCommandQueue<FLightSceneInfo, ELightDirtyFlags, ELightUpdateId>;
+using FUpdateLightCommand = FSceneLightInfoUpdates::FUpdateCommand;
+
+struct FUpdateLightTransformParameters : public FSceneLightInfoUpdates::TPayloadBase<ELightUpdateId::Transform, ELightDirtyFlags::None>
+{
+	FMatrix LightToWorld;
+	FVector4 Position;
+};
+
+struct FUpdateLightColorParameters: public FSceneLightInfoUpdates::TPayloadBase<ELightUpdateId::Color, ELightDirtyFlags::None>
+{
+	FLinearColor NewColor;
+	float NewIndirectLightingScale;
+	float NewVolumetricScatteringIntensity;
+};
+
+/**
+ * Describes all light modifications to the scene by recording the light scene IDs.
+ * TODO: If needed, we could add a reference to the FLightUpdates (which contains the commands) since this would enable systems to consume out the delta updates as they come in.
+ *       If this is useful we must ensure FLightUpdates are kept alive as long as any async tasks might require.
+ */
+struct FLightSceneChangeSet
+{
+	// IDs of all lights before they were removed, IDs in this array may not be valid at all times when the change-set is used (depends on whether the callback site is before or after the given lights are removed from the scene).
+	TConstArrayView<int32> RemovedLightIds;
+	// IDs of all lights added to the scene, only available after all lights are added to the scene, may contain the same ID's as removed, as they may be reused.
+	TConstArrayView<int32> AddedLightIds;
+	FSceneLightInfoUpdates *SceneLightInfoUpdates = nullptr;
+};
+
 /** 
  * Renderer scene which is private to the renderer module.
  * Ordinarily this is the renderer version of a UWorld, but an FScene can be created for previewing in editors which don't have a UWorld as well.
@@ -3724,8 +3846,6 @@ public:
 		return &CachedShadowMapDatas[ShadowMapIndex];
 	}
 
-	bool IsPrimitiveBeingRemoved(FPrimitiveSceneInfo* PrimitiveSceneInfo) const;
-
 	/**
 	 * Maximum used persistent Primitive Index, use to size arrays that store primitive data indexed by FPrimitiveSceneInfo::PersistentIndex.
 	 * Only changes during UpdateAllPrimitiveSceneInfos.
@@ -3895,7 +4015,7 @@ private:
 	 * Removes a primitive from the scene.  Called in the rendering thread by RemovePrimitive.
 	 * @param PrimitiveSceneInfo - The primitive being removed.
 	 */
-	bool RemovePrimitiveSceneInfo_RenderThread(FPrimitiveSceneInfo* PrimitiveSceneInfo);
+	void RemovePrimitiveSceneInfo_RenderThread(FPrimitiveSceneInfo* PrimitiveSceneInfo);
 
 	/** Updates a primitive's transform, called on the rendering thread. */
 	void UpdatePrimitiveTransform_RenderThread(FPrimitiveSceneProxy* PrimitiveSceneProxy, const FBoxSphereBounds& WorldBounds, const FBoxSphereBounds& LocalBounds, const FMatrix& LocalToWorld, const FVector& OwnerPosition, const TOptional<FTransform>& PreviousTransform);
@@ -3963,11 +4083,16 @@ private:
 	void ProcessAtmosphereLightAddition_RenderThread(FLightSceneInfo* LightSceneInfo);
 
 	/**
-	 * Process all scene updates for lights, returns the change-set, which references arrays allocated with a RDG builder life-time.
+	 * Process all scene updates for lights.
 	 */
-	FLightSceneChangeSet UpdateAllLightSceneInfos(FRDGBuilder& GraphBuilder);
+	void UpdateAllLightSceneInfos(FRDGBuilder& GraphBuilder);
 
 private:
+	template <typename UpdatePayloadType>
+	void UpdatePrimitiveInternal(FPrimitiveSceneProxy* SceneProxy, UpdatePayloadType &&InUpdatePayload);
+
+	template <typename UpdatePayloadType>
+	void UpdateLightInternal(FLightSceneProxy* LightSceneProxy, UpdatePayloadType &&InUpdatePayload);
 
 	/**
 	 * Update tracked scene state for cached CSM shadows
@@ -3976,28 +4101,12 @@ private:
 
 	FString FullWorldName;
 #if RHI_RAYTRACING
-	void UpdateRayTracingGroupBounds_AddPrimitives(const Experimental::TRobinHoodHashSet<FPrimitiveSceneInfo*>& PrimitiveSceneInfos);
-	void UpdateRayTracingGroupBounds_RemovePrimitives(const Experimental::TRobinHoodHashSet<FPrimitiveSceneInfo*>& PrimitiveSceneInfos);
-	template<typename ValueType>
-	inline void UpdateRayTracingGroupBounds_UpdatePrimitives(const Experimental::TRobinHoodHashMap<FPrimitiveSceneProxy*, ValueType>& UpdatedTransforms);
+	void UpdateRayTracingGroupBounds_AddPrimitives(const TArray<FPrimitiveSceneInfo*, SceneRenderingAllocator>& PrimitiveSceneInfos);
+	void UpdateRayTracingGroupBounds_RemovePrimitives(const TArray<FPrimitiveSceneInfo*, SceneRenderingAllocator>& PrimitiveSceneInfos);
+	template<typename RangeType>
+	inline void UpdateRayTracingGroupBounds_UpdatePrimitives(const RangeType& UpdatedTransforms);
 #endif
 
-	struct FUpdateTransformCommand
-	{
-		FBoxSphereBounds WorldBounds;
-		FBoxSphereBounds LocalBounds; 
-		FMatrix LocalToWorld; 
-		FVector AttachmentRootPosition;
-	};
-
-	struct FUpdateInstanceCommand
-	{
-		FPrimitiveSceneProxy* PrimitiveSceneProxy{ nullptr };
-		FBoxSphereBounds WorldBounds;
-		FBoxSphereBounds LocalBounds;
-		FBoxSphereBounds StaticMeshBounds;
-	};
-
 	void UpdatePrimitiveInstances(FUpdateInstanceCommand& UpdateParams);
 
 	struct FLevelCommand
@@ -4012,19 +4121,8 @@ private:
 		EOp Op;
 	};
 
-	Experimental::TRobinHoodHashMap<FPrimitiveSceneInfo*, FPrimitiveComponentId> UpdatedAttachmentRoots;
-	Experimental::TRobinHoodHashMap<FPrimitiveSceneProxy*, FCustomPrimitiveData> UpdatedCustomPrimitiveParams;
-	Experimental::TRobinHoodHashMap<FPrimitiveSceneProxy*, FUpdateTransformCommand> UpdatedTransforms;
-	Experimental::TRobinHoodHashMap<FPrimitiveSceneProxy*, FUpdateInstanceCommand> UpdatedInstances;
-	Experimental::TRobinHoodHashMap<FPrimitiveSceneInfo*, FMatrix> OverridenPreviousTransforms;
-	Experimental::TRobinHoodHashMap<const FPrimitiveSceneProxy*, float> UpdatedOcclusionBoundsSlacks;
-	Experimental::TRobinHoodHashMap<FPrimitiveSceneProxy*, FVector2f> UpdatedInstanceCullDistance;
-	Experimental::TRobinHoodHashMap<FPrimitiveSceneProxy*, FVector3f> UpdatedDrawDistance;
-	Experimental::TRobinHoodHashSet<FPrimitiveSceneInfo*> AddedPrimitiveSceneInfos;
-	Experimental::TRobinHoodHashSet<FPrimitiveSceneInfo*> RemovedPrimitiveSceneInfos;
-	Experimental::TRobinHoodHashSet<FPrimitiveSceneInfo*> DistanceFieldSceneDataUpdates;
+	FScenePrimitiveUpdates PrimitiveUpdates;
 	TArray<FLevelCommand> LevelCommands;
-	TSet<FPrimitiveSceneInfo*> DeletedPrimitiveSceneInfos;
 
 	UE::Tasks::FTask CreateLightPrimitiveInteractionsTask;
 	UE::Tasks::FTask GPUSkinCacheTask;
diff --git a/Engine/Source/Runtime/Renderer/Private/SceneUpdateCommandQueue.h b/Engine/Source/Runtime/Renderer/Private/SceneUpdateCommandQueue.h
new file mode 100644
--- /dev/null
+++ b/Engine/Source/Runtime/Renderer/Private/SceneUpdateCommandQueue.h
@@ -0,0 +1,567 @@
+// Copyright Epic Games, Inc. All Rights Reserved.
+#pragma once
+
+#include "CoreMinimal.h"
+#include "Experimental/Containers/RobinHoodHashTable.h"
+#include "RendererInterface.h"
+
+enum class ESceneUpdateCommandFilter : uint32
+{
+	Added			= 1u << 0,
+	Deleted			= 1u << 1,
+	Updated			= 1u << 2,
+	AddedUpdated	= Added | Updated,
+	All				= Added | Deleted | Updated
+};
+ENUM_CLASS_FLAGS(ESceneUpdateCommandFilter);
+
+/**
+ * An unordered queue for sending scene object updates (agnostic to the object type InSceneInfoType/FSceneInfo). 
+ * Several update payloads can be enqueued for each object but only the last of each type will have effect.
+ * The update payloads are stored in a typed compact array, but are not themselves required to have virtual destructors or even be of any particular type.
+ * TPayloadBase is a helper & example that can be used as a base class to implement payload data types.
+ * Update payload types are identified by an ID that comes from the template argument enum EInId.
+ * While iterating the FUpdateCommand it is possible to access each type of update payload safely, or one can interate the payload types in a continuous fashion.
+ */
+template <typename InSceneInfoType, typename EInDirtyFlagsType, typename EInId>
+class TSceneUpdateCommandQueue
+{
+private:
+	using FAllocator = FDefaultAllocator;
+
+	struct FBasePayloadArray;
+	template <typename PayloadType>
+	struct TPayloadArray;
+
+public:
+	using EDirtyFlags = EInDirtyFlagsType;
+	using FSceneInfo = InSceneInfoType;
+	using FSceneInfoPersistentId = typename InSceneInfoType::FPersistentId;
+
+	using EId = EInId;
+
+	static constexpr int32 MaxId = int32(EId::MAX);
+	static_assert(MaxId <= 64, "The max update ID must fit in the 64 bits we use to store the mask.");
+
+	/**
+	 * Each command represents all the updates for a given scene object. Add/Delete/AttributeUpdate. 
+	 * Associated with a command are zero or more payloads which are arbitrarily typed data packets.
+	 */
+	struct FUpdateCommand
+	{
+		FUpdateCommand(FSceneInfo *InSceneInfo, FSceneInfoPersistentId InPersistentId) : SceneInfo(InSceneInfo), PersistentId(InPersistentId) { }
+
+		template <typename PayloadType>
+		int32 GetPayloadOffset() const
+		{
+			if (PayloadType::IdBit & PayloadMask)
+			{
+				int32 Index = FPlatformMath::CountBits(PayloadMask& PayloadType::ExclusiveIdMask);
+
+				return PayloadDataSlots[Index];
+			}
+			return INDEX_NONE;
+		}
+
+		template <typename PayloadType>
+		void SetOrAddPayloadOffset(int32 PayloadOffset, EDirtyFlags InDirtyFlags)
+		{
+			DirtyFlags |= InDirtyFlags;
+
+			int32 Index = FPlatformMath::CountBits(PayloadMask& PayloadType::ExclusiveIdMask);
+			// previously set, replace pointer
+			if (PayloadType::IdBit& PayloadMask)
+			{
+				PayloadDataSlots[Index] = PayloadOffset;
+			}
+			else
+			{
+				PayloadMask |= PayloadType::IdBit;
+				PayloadDataSlots.Insert(PayloadOffset, Index);
+			}
+		}
+
+		FSceneInfo* GetSceneInfo() const { return SceneInfo; }
+		FSceneInfoPersistentId GetPersistentId() const { return PersistentId; }
+		bool IsDelete() const { return bDeleted; }
+		bool IsAdd() const { return bAdded; }
+
+		// Should only be called for added objects after the ID has been allocated.
+		void SetPersistentId(FSceneInfoPersistentId InId)
+		{
+			check(bAdded);
+			PersistentId = InId;
+		}
+
+	private:
+		FSceneInfo* SceneInfo = nullptr;
+		uint64 PayloadMask = 0ull;
+		FSceneInfoPersistentId PersistentId;
+		EDirtyFlags DirtyFlags = EDirtyFlags::None;
+		bool bDeleted = false;
+		bool bAdded = false;
+		// offsets to the stored data (in bit-order)
+		TArray<int32, TInlineAllocator<8, FAllocator>> PayloadDataSlots;
+		friend TSceneUpdateCommandQueue;
+	};
+
+	TSceneUpdateCommandQueue() 
+		: PayloadArrays(InPlace)
+	{
+	}
+
+	// Prevent copy constructor to avoid issues with the owned payload arrays
+	TSceneUpdateCommandQueue(const TSceneUpdateCommandQueue&) = delete;
+	TSceneUpdateCommandQueue(TSceneUpdateCommandQueue&& Other)
+		: PayloadArrays(MoveTemp(Other.PayloadArrays))
+		, Commands(MoveTemp(Other.Commands))
+	{
+		check(Other.RaceGuard == 0);
+		Other.CommandSlots.Empty();
+	}
+
+	~TSceneUpdateCommandQueue()
+	{
+		check(RaceGuard == 0);
+	}
+
+	bool IsEmpty() const { return Commands.IsEmpty(); }
+
+	int32 NumCommands() const { return Commands.Num(); }
+
+	bool HasCommand(FSceneInfo* SceneInfo) const
+	{
+		return CommandSlots.Find(SceneInfo) != nullptr;
+	}
+
+	const FUpdateCommand* FindCommand(FSceneInfo* SceneInfo) const
+	{
+		Experimental::FHashElementId Id = CommandSlots.FindId(SceneInfo);
+		if (Id.IsValid())
+		{
+			return &Commands[Id.GetIndex()];
+		}
+		return nullptr;
+	}
+
+	/**
+	 * Enqueue a Delete command. This will mark the command for the scene object as deleted, but does not remove the command or associated updates.
+	 * Thus a command may have both add/delete flags and update payloads. It is up to the consumer to handle these appropriately.
+	 */
+	void EnqueueDelete(FSceneInfo* SceneInfo)
+	{
+		check(RaceGuard == 0);
+
+		int32 CommandSlot = GetOrAddCommandSlot(SceneInfo);
+		FUpdateCommand& Command = Commands[CommandSlot];
+
+		Command.bDeleted = true;
+		// We leave any update data in place, as the alternative is doing a remove-swap and then fixing up all the indexes, which seems not-worth it, instead these updates should be skipped.
+	}
+
+	/**
+	 * Enqueue an Add command. This will mark the command for the scene object as added, and must always be the first command for the object.
+	 */
+	void EnqueueAdd(FSceneInfo* SceneInfo)
+	{
+		check(RaceGuard == 0);
+
+		// Add should always be the first command.
+		check(CommandSlots.Find(SceneInfo) == nullptr);
+		// It should not have been added to the scene already
+		int32 CommandSlot = GetOrAddCommandSlot(SceneInfo);
+		FUpdateCommand& Command = Commands[CommandSlot];
+		// It should not even be possible for this to happen
+		check(!Command.bAdded);
+		check(!Command.bDeleted);
+		Command.bAdded = true;
+	}
+
+	/**
+	 * Enqueue an update with a data payload.
+	 */
+	template <typename PayloadType>
+	void Enqueue(FSceneInfo* SceneInfo, PayloadType&& Payload)
+	{
+		check(RaceGuard == 0);
+
+		int32 CommandSlot = GetOrAddCommandSlot(SceneInfo);
+		FUpdateCommand& Command = Commands[CommandSlot];
+
+		if (PayloadArrays[int32(PayloadType::Id)] == nullptr)
+		{
+			PayloadArrays[int32(PayloadType::Id)] = MakeUnique<TPayloadArray<PayloadType>>();
+		}
+		TPayloadArray<PayloadType>& Payloads = GetPayloadArray<PayloadType>();
+		int32 PrevPayloadOffset = Command.template GetPayloadOffset<PayloadType>();
+
+		// Update existing payload (maybe we want to disallow this?)
+		if (PrevPayloadOffset != INDEX_NONE)
+		{
+			check(Payloads.CommandSlots[PrevPayloadOffset] == CommandSlot);
+			Payloads.PayloadData[PrevPayloadOffset] = Payload;
+		}
+		else // New payload for this command
+		{
+			int32 PayloadOffset = Payloads.CommandSlots.Num();
+			Commands[CommandSlot].template SetOrAddPayloadOffset<PayloadType>(PayloadOffset, Payload.GetDirtyFlags());
+			Payloads.CommandSlots.Emplace(CommandSlot);
+			Payloads.PayloadData.Emplace(Payload);
+		}
+	}
+
+	/**
+	 * Retriev a pointer to the PayloadType data for the given command. Returs nullptr if no such data exists.
+	 */
+	template <typename PayloadType>
+	PayloadType* GetPayloadPtr(const FUpdateCommand& Command)
+	{
+		if (PayloadArrays[int32(PayloadType::Id)] == nullptr)
+		{
+			return nullptr;
+		}
+
+		TPayloadArray<PayloadType>& Payloads = GetPayloadArray<PayloadType>();
+		int32 PayloadOffset = Command.template GetPayloadOffset<PayloadType>();
+
+		// Update existing payload (maybe we want to disallow this?)
+		if (PayloadOffset != INDEX_NONE)
+		{
+			// cross check
+			check(Payloads.CommandSlots[PayloadOffset] != INDEX_NONE);
+			check(&Commands[Payloads.CommandSlots[PayloadOffset]] == &Command);
+			return &Payloads.PayloadData[PayloadOffset];
+		}
+		return nullptr;
+	}
+
+	/**
+	 * Reset the command and payload data stored in the buffer, leaving allocations unchanged.
+	 */
+	void Reset()
+	{
+		check(RaceGuard == 0);
+
+		CommandSlots.Empty();
+		Commands.Reset();
+		for (const auto& PayloadArray : PayloadArrays)
+		{
+			if (PayloadArray != nullptr)
+			{
+				PayloadArray->Reset();
+			}
+		}
+	}
+
+	/**
+	 */
+	template <typename CallbackFuncType>
+	void ForEachCommand(ESceneUpdateCommandFilter CommandFilter, CallbackFuncType CallbackFunc)
+	{
+		for (FUpdateCommand& Command : Commands)
+		{
+			if (IsFilterIncludingCommand(Command, CommandFilter))
+			{
+				CallbackFunc(Command);
+			}
+		}
+	}
+
+	template <typename CallbackFuncType>
+	void ForEachCommand(CallbackFuncType CallbackFunc)
+	{
+		return ForEachCommand(ESceneUpdateCommandFilter::All, CallbackFunc);
+	}
+
+	/**
+	 * Filter on ESceneUpdateCommandFilter and _updates_ on payload mask. I.e., the payload mask only matters if the Command is an update.
+	 *   E.g., ForEachUpdateCommand(ESceneUpdateCommandFilter::Added,...) will return _all_ added commands regardless of payload mask.
+	 *   E.g., ForEachUpdateCommand(ESceneUpdateCommandFilter::Added | ESceneUpdateCommandFilter::Updated,...) will return _all_ added commands regardless of UpdatePayloadMask AND all updates that match the UpdatePayloadMask.
+	 */
+	template <typename CallbackFuncType>
+	void ForEachUpdateCommand(ESceneUpdateCommandFilter CommandFilter, uint64 UpdatePayloadMask, CallbackFuncType CallbackFunc) const
+	{
+		for (const FUpdateCommand& Command : Commands)
+		{
+			if (!IsFilterIncludingCommand(Command, CommandFilter))
+			{
+				continue;
+			}
+
+			// Include a command if it is added/deleted OR matches the mask (the former can be excluded in the CommandFilter)
+			if (Command.bAdded || Command.bDeleted || (UpdatePayloadMask& Command.PayloadMask) != 0)
+			{
+				CallbackFunc(Command);
+			}
+		}
+	}
+
+	/**
+	 * Filter on ESceneUpdateCommandFilter and _updates_ on DirtyFlags mask. I.e., the DirtyFlags mask only matters if the Command is an update.
+	 *   E.g., ForEachUpdateCommand(ESceneUpdateCommandFilter::Added,...) will return _all_ added commands regardless of payload mask.
+	 *   E.g., ForEachUpdateCommand(ESceneUpdateCommandFilter::Added | ESceneUpdateCommandFilter::Updated,...) will return _all_ added commands regardless of DirtyFlags AND all updates that match the DirtyFlags.
+	 */
+	template <typename CallbackFuncType>
+	void ForEachUpdateCommand(ESceneUpdateCommandFilter CommandFilter, EDirtyFlags DirtyFlags, CallbackFuncType CallbackFunc) const
+	{
+		for (const FUpdateCommand& Command : Commands)
+		{
+			if (!IsFilterIncludingCommand(Command, CommandFilter))
+			{
+				continue;
+			}
+
+			if (Command.bAdded || Command.bDeleted || EnumHasAnyFlags(Command.DirtyFlags, DirtyFlags))
+			{
+				CallbackFunc(Command);
+			}
+		}
+	}
+
+	/**
+	 * Iterator to loop over a particular type of payload.
+	 * Used to implement the (typically) more convenient TPayloadRangeView, see GetRangeView() below.
+	 * Deleted items are automatically skipped.
+	 */
+	template <typename PayloadType>
+	struct TConstPayloadIterator
+	{
+		struct FItem
+		{
+			const PayloadType& Payload;
+			FSceneInfo* SceneInfo;
+		};
+
+		TConstPayloadIterator(const TPayloadArray<PayloadType>* InPayloads, const TArray<FUpdateCommand, FAllocator>& InCommands, int32 InIndex = 0) 
+			: Payloads(InPayloads)
+			, Commands(InCommands)
+			, Index(InIndex)
+		{
+			SkipDeleted();
+		}
+
+		void SkipDeleted()
+		{
+			while (Payloads && Index < Payloads->CommandSlots.Num() && Commands[Payloads->CommandSlots[Index]].bDeleted)
+			{
+				++Index;
+			}
+		}
+
+		void operator++() 
+		{ 
+			++Index;
+			SkipDeleted();
+		}
+
+		FItem operator*() const 
+		{ 
+			check(Payloads && Index < Payloads->CommandSlots.Num());
+			const FUpdateCommand& Command = Commands[Payloads->CommandSlots[Index]];
+			check(!Command.bDeleted);
+			return { Payloads->PayloadData[Index], Command.SceneInfo };
+		}
+
+		bool operator != (const TConstPayloadIterator& It ) const 
+		{ 
+			check(Payloads == It.Payloads);
+			return Index != It.Index; 
+		}
+
+		explicit operator bool() const { return Payloads!= nullptr && Index < Payloads->CommandSlots.Num(); }
+
+	private:
+		const TPayloadArray<PayloadType>* Payloads;
+		const TArray<FUpdateCommand, FAllocator>& Commands;
+		int32 Index = 0;
+	};
+
+	/**
+	 * Get an iterator to iterate updates with a given payload type.
+	 */
+	template <typename PayloadType>
+	TConstPayloadIterator<PayloadType> GetIterator() const
+	{
+		const TPayloadArray<PayloadType>* Payloads = &GetPayloadArray<PayloadType>();
+		return TConstPayloadIterator<PayloadType>(Payloads, Commands, 0);
+	}
+
+	/**
+	 * Get the number of updates with a given payload type.
+	 */
+	template <typename PayloadType>
+	int32 GetNumItems() const
+	{
+		const TPayloadArray<PayloadType>* Payloads = &GetPayloadArray<PayloadType>();
+		return Payloads != nullptr ? Payloads->PayloadData.Num() : 0;
+	}
+
+	template <typename PayloadType>
+	struct TPayloadRangeView
+	{
+		TPayloadRangeView(TSceneUpdateCommandQueue& InUpdateBuffer, int32 InNumItems) 
+			: UpdateBuffer(InUpdateBuffer)
+			, NumItems(InNumItems)
+		{
+#if DO_CHECK
+			UpdateBuffer.BeginReadAccess();
+#endif
+		}
+		~TPayloadRangeView()
+		{
+#if DO_CHECK
+			UpdateBuffer.EndReadAccess();
+#endif
+		}
+
+		TConstPayloadIterator<PayloadType> begin() const { return UpdateBuffer.GetIterator<PayloadType>(); }
+		TConstPayloadIterator<PayloadType> end() const { return UpdateBuffer.GetEndIterator<PayloadType>(); }
+
+		int32 Num() const { return NumItems; }
+
+	private:
+		TSceneUpdateCommandQueue& UpdateBuffer;
+		int32 NumItems = 0;
+	};
+
+	/**
+	 * Get a "range" that can be used in a range for loop to access updates of a single payload type.
+	 * e.g. for (auto& Item : Buffer.GetRangeView<MyUpdatePayloadType>());
+	 * Deleted items are automatically skipped.
+	 */
+	template <typename PayloadType>
+	TPayloadRangeView<PayloadType> GetRangeView()
+	{
+		return TPayloadRangeView<PayloadType>(*this, GetNumItems<PayloadType>());
+	}
+
+	/**
+	 * Helper payload base templace, that defines the expected id flags & masks.
+	 * Not required, as any struct that has the same interface can be used.
+	 */
+	template <EId InId, EDirtyFlags InDirtyFlags>
+	struct TPayloadBase
+	{
+		// unique ID for the type of update, can be made to support registration like the SceneExtensions to rather than be compile time.
+		static constexpr EId Id = InId;
+		static constexpr uint64 IdBit = 1ull << uint32(Id);
+		static constexpr uint64 ExclusiveIdMask = (1ull << uint32(Id)) - 1ull;
+		static constexpr EDirtyFlags DirtyFlags = InDirtyFlags;
+	
+		// Allow static polymorphism to implement per-payload runtime variable flags.
+		EDirtyFlags GetDirtyFlags() const { return DirtyFlags; }
+	};
+	
+#if DO_CHECK
+	void BeginReadAccess()
+	{
+		++RaceGuard;
+	}
+
+	void EndReadAccess()
+	{
+		--RaceGuard;
+	}
+
+	struct FReadAccessScope
+	{
+		FReadAccessScope(TSceneUpdateCommandQueue& InUpdateQueue) : 
+			UpdateQueue(InUpdateQueue)
+		{
+			UpdateQueue.BeginReadAccess();
+		}
+		~FReadAccessScope()
+		{
+			UpdateQueue.EndReadAccess();
+		}
+	protected:
+		TSceneUpdateCommandQueue& UpdateQueue;
+	};
+#endif
+
+private:
+	/**
+	 * Get an iterator that points to the end of a payload array with a given payload type.
+	 */
+	template <typename PayloadType>
+	TConstPayloadIterator<PayloadType> GetEndIterator() const
+	{
+		const TPayloadArray<PayloadType>* Payloads = &GetPayloadArray<PayloadType>();
+		return TConstPayloadIterator<PayloadType>(Payloads, Commands, Payloads != nullptr ? Payloads->PayloadData.Num() : 0);
+	}
+
+	struct FBasePayloadArray
+	{
+		FBasePayloadArray(int32 InPayloadByteSize) : PayloadByteSize(InPayloadByteSize) {}
+		virtual ~FBasePayloadArray() { };
+		virtual void Reset() = 0;
+		
+		int32 PayloadByteSize = 0;
+	};
+
+	template <typename PayloadType>
+	struct TPayloadArray : public FBasePayloadArray
+	{
+		TArray<PayloadType, FAllocator> PayloadData;
+		TArray<int32, FAllocator> CommandSlots;
+		
+		TPayloadArray() : FBasePayloadArray(sizeof(PayloadType)) {}
+
+		virtual void Reset() override
+		{
+			PayloadData.Reset();
+			CommandSlots.Reset();
+		}
+	};
+
+	int32 GetOrAddCommandSlot(FSceneInfo* SceneInfo)
+	{
+		bool bWasAlreadyInSet = false;
+		Experimental::FHashElementId Id = CommandSlots.FindOrAddId(SceneInfo, bWasAlreadyInSet);
+		int32 CommandSlot = Id.GetIndex();
+		if (!bWasAlreadyInSet)
+		{
+			// Since we only add to this thing the slots should always go at the end.
+			// Note: since we always grow this table it could easily have a much simpler hash table.
+			check(CommandSlot == Commands.Num());
+			Commands.Emplace(SceneInfo, SceneInfo->GetPersistentIndex());
+		}
+		check(Commands.IsValidIndex(CommandSlot));
+		return CommandSlot;
+	}
+
+	bool IsFilterIncludingCommand(const FUpdateCommand& Command, ESceneUpdateCommandFilter CommandFilter) const
+	{
+		if (Command.bDeleted)
+		{
+			return EnumHasAnyFlags(CommandFilter, ESceneUpdateCommandFilter::Deleted);
+		}
+
+		if (Command.bAdded)
+		{
+			return EnumHasAnyFlags(CommandFilter, ESceneUpdateCommandFilter::Added);
+		}
+		return EnumHasAnyFlags(CommandFilter, ESceneUpdateCommandFilter::Updated);
+	}
+
+	template <typename PayloadType>
+	TPayloadArray<PayloadType>& GetPayloadArray()
+	{
+		check(PayloadArrays[int32(PayloadType::Id)] == nullptr || PayloadArrays[int32(PayloadType::Id)]->PayloadByteSize == sizeof(PayloadType));
+		return static_cast<TPayloadArray<PayloadType>&>(*PayloadArrays[int32(PayloadType::Id)]);
+	}
+	template <typename PayloadType>
+	const TPayloadArray<PayloadType>& GetPayloadArray() const
+	{
+		return const_cast<TSceneUpdateCommandQueue*>(this)->GetPayloadArray<PayloadType>();
+	}
+
+	using FPayloadArrays = TStaticArray<TUniquePtr<FBasePayloadArray>, 64>;
+	FPayloadArrays PayloadArrays = FPayloadArrays(InPlace);
+	TArray<FUpdateCommand, FAllocator> Commands;
+	Experimental::TRobinHoodHashSet<FSceneInfo*> CommandSlots;
+
+#if DO_CHECK
+	std::atomic<int> RaceGuard = 0;
+#endif
+};
diff --git a/Engine/Source/Runtime/Renderer/Private/Shadows/ShadowScene.cpp b/Engine/Source/Runtime/Renderer/Private/Shadows/ShadowScene.cpp
--- a/Engine/Source/Runtime/Renderer/Private/Shadows/ShadowScene.cpp
+++ b/Engine/Source/Runtime/Renderer/Private/Shadows/ShadowScene.cpp
@@ -38,7 +38,7 @@ void FShadowScene::PostLightsUpdate(FRDGBuilder& GraphBuilder, const FLightScene
 	// don't spawn async work for no good reason.
 	constexpr int32 kMinWorkSizeForAsync = 64;
 
-	int32 WorkSize = LightSceneChangeSet.AddedLightIds.Num() + LightSceneChangeSet.RemovedLightIds.Num() + LightSceneChangeSet.TransformUpdatedLightIds.Num();
+	int32 WorkSize = LightSceneChangeSet.SceneLightInfoUpdates->NumCommands();
 
 	// Don't sync, or kick off a new job if there is no work to do.
 	if (WorkSize > 0)
@@ -78,7 +78,10 @@ void FShadowScene::PostLightsUpdate(FRDGBuilder& GraphBuilder, const FLightScene
 			};
 
 			Algo::ForEach(LightSceneChangeSet.AddedLightIds, UpdateLight);
-			Algo::ForEach(LightSceneChangeSet.TransformUpdatedLightIds, UpdateLight);
+			for (const auto& Item : LightSceneChangeSet.SceneLightInfoUpdates->GetRangeView<FUpdateLightTransformParameters>())
+			{ 
+				UpdateLight(Item.SceneInfo->Id);
+			}
 
 		}, WorkSize > kMinWorkSizeForAsync);
 	}
diff --git a/Engine/Source/Runtime/Renderer/Private/SplineMeshSceneResources.cpp b/Engine/Source/Runtime/Renderer/Private/SplineMeshSceneResources.cpp
--- a/Engine/Source/Runtime/Renderer/Private/SplineMeshSceneResources.cpp
+++ b/Engine/Source/Runtime/Renderer/Private/SplineMeshSceneResources.cpp
@@ -338,23 +338,20 @@ void FSplineMeshSceneUpdater::PostSceneUpdate(FRDGBuilder& GraphBuilder, const F
 		}
 	};
 
-	// First, register any new primitives
-	for (FPrimitiveSceneInfo* PrimitiveSceneInfo : ChangeSet.AddedPrimitiveSceneInfos)
+	ChangeSet.PrimitiveUpdates.ForEachUpdateCommand(ESceneUpdateCommandFilter::AddedUpdated, FUpdateInstanceCommand::IdBit | FUpdateTransformCommand::IdBit, [&](const FPrimitiveUpdateCommand& Cmd)
 	{
-		if (PrimitiveSceneInfo->Proxy->IsSplineMesh())
-		{
-			RequestUpdate(SceneData->Register(*PrimitiveSceneInfo));
-		}
-	}
-
-	// Request updates from any updated primitives
-	for (FPrimitiveSceneInfo* PrimitiveSceneInfo : ChangeSet.UpdatedPrimitiveSceneInfos)
-	{
-		if (PrimitiveSceneInfo->Proxy->IsSplineMesh())
+		if (Cmd.GetSceneInfo()->Proxy->IsSplineMesh())
 		{
-			RequestUpdate(SceneData->RegisteredPrimitives.FindChecked(PrimitiveSceneInfo));
+			if (Cmd.IsAdd())
+			{
+				RequestUpdate(SceneData->Register(*Cmd.GetSceneInfo()));
+			}
+			else
+			{
+				RequestUpdate(SceneData->RegisteredPrimitives.FindChecked(Cmd.GetSceneInfo()));
+			}
 		}
-	}
+	});
 
 	if (SceneData->NumRegisteredPrimitives() > 0)
 	{
diff --git a/Engine/Source/Runtime/Renderer/Private/VirtualShadowMaps/VirtualShadowMapCacheManager.cpp b/Engine/Source/Runtime/Renderer/Private/VirtualShadowMaps/VirtualShadowMapCacheManager.cpp
--- a/Engine/Source/Runtime/Renderer/Private/VirtualShadowMaps/VirtualShadowMapCacheManager.cpp
+++ b/Engine/Source/Runtime/Renderer/Private/VirtualShadowMaps/VirtualShadowMapCacheManager.cpp
@@ -1688,24 +1688,19 @@ void FVirtualShadowMapInvalidationSceneUpdater::PreSceneUpdate(FRDGBuilder& Grap
 			InvalidatingPrimitiveCollector.UpdatedTransform(PrimitiveSceneInfo);
 		}
 
-		// All removed primitives must invalidate their footprints in the VSM before leaving
-		for (FPrimitiveSceneInfo* PrimitiveSceneInfo : ChangeSet.RemovedPrimitiveSceneInfos)
+		// Note: skips added as they are not fully defined at this point (not primitive ID allocated, 
+		ChangeSet.PrimitiveUpdates.ForEachUpdateCommand(ESceneUpdateCommandFilter::Updated | ESceneUpdateCommandFilter::Deleted, EPrimitiveUpdateDirtyFlags::AllCulling, [&](const FPrimitiveUpdateCommand& Cmd)
 		{
-			InvalidatingPrimitiveCollector.Removed(PrimitiveSceneInfo);
-		}
-		// As must all primitive updates, 
-		for (FPrimitiveSceneInfo* PrimitiveSceneInfo : ChangeSet.UpdatedPrimitiveSceneInfos)
-		{
-			InvalidatingPrimitiveCollector.UpdatedTransform(PrimitiveSceneInfo);
-		}
-
-		// TODO! Where do we get this data from...
-		/*
-		for (const auto& CullDistance : UpdatedInstanceCullDistance)
-		{
-			InvalidatingPrimitiveCollector.UpdatedTransform(CullDistance.Key->GetPrimitiveSceneInfo());
-		}
-		*/
+			if (Cmd.IsDelete())
+			{
+				// All removed primitives must invalidate their footprints in the VSM before leaving.
+				InvalidatingPrimitiveCollector.Removed(Cmd.GetSceneInfo());
+			}
+			else
+			{
+				InvalidatingPrimitiveCollector.UpdatedTransform(Cmd.GetSceneInfo());
+			}
+		});
 
 		TRACE_INT_VALUE(TEXT("Shadow.Virtual.Cache.PreInvalidationInstances"), InvalidatingPrimitiveCollector.Instances.GetTotalNumInstances());
 		CacheManager.ProcessInvalidations(GraphBuilder, SceneUniforms, InvalidatingPrimitiveCollector);
@@ -1717,7 +1712,7 @@ void FVirtualShadowMapInvalidationSceneUpdater::PostSceneUpdate(FRDGBuilder& Gra
 	CacheManager.ReallocatePersistentPrimitiveIndices();
 
 	// Grab a reference, but we currently do all the work in PostGPUSceneUpdate
-	PostUpdateChangeSet = ChangeSet;
+	PostUpdateChangeSet = &ChangeSet;
 }
 
 void FVirtualShadowMapInvalidationSceneUpdater::PostGPUSceneUpdate(FRDGBuilder& GraphBuilder, FSceneUniformBuffer& SceneUniforms)
@@ -1728,17 +1723,21 @@ void FVirtualShadowMapInvalidationSceneUpdater::PostGPUSceneUpdate(FRDGBuilder&
 	{
 		FVirtualShadowMapArrayCacheManager::FInvalidatingPrimitiveCollector InvalidatingPrimitiveCollector(&CacheManager);
 
-		for (FPrimitiveSceneInfo* PrimitiveSceneInfo : PostUpdateChangeSet.AddedPrimitiveSceneInfos)
-		{
-			InvalidatingPrimitiveCollector.Added(PrimitiveSceneInfo);
-		}
-		for (FPrimitiveSceneInfo* PrimitiveSceneInfo : PostUpdateChangeSet.UpdatedPrimitiveSceneInfos)
+		// Filter out all updates that are either "add" or has dirty flags to say they affect the bounds.
+		PostUpdateChangeSet->PrimitiveUpdates.ForEachUpdateCommand(ESceneUpdateCommandFilter::AddedUpdated, EPrimitiveUpdateDirtyFlags::AllCulling, [&](const FPrimitiveUpdateCommand& Cmd)
 		{
-			InvalidatingPrimitiveCollector.UpdatedTransform(PrimitiveSceneInfo);
-		}
+			if (Cmd.IsAdd())
+			{
+				InvalidatingPrimitiveCollector.Added(Cmd.GetSceneInfo());
+			}
+			else
+			{
+				InvalidatingPrimitiveCollector.UpdatedTransform(Cmd.GetSceneInfo());
+			}
+		});
 
 		TRACE_INT_VALUE(TEXT("Shadow.Virtual.Cache.PostInvalidationInstances"), InvalidatingPrimitiveCollector.Instances.GetTotalNumInstances());
 		CacheManager.ProcessInvalidations(GraphBuilder, SceneUniforms, InvalidatingPrimitiveCollector);
 	}
-	PostUpdateChangeSet = FScenePostUpdateChangeSet();
+	PostUpdateChangeSet = nullptr;
 }
diff --git a/Engine/Source/Runtime/Renderer/Private/VirtualShadowMaps/VirtualShadowMapCacheManager.h b/Engine/Source/Runtime/Renderer/Private/VirtualShadowMaps/VirtualShadowMapCacheManager.h
--- a/Engine/Source/Runtime/Renderer/Private/VirtualShadowMaps/VirtualShadowMapCacheManager.h
+++ b/Engine/Source/Runtime/Renderer/Private/VirtualShadowMaps/VirtualShadowMapCacheManager.h
@@ -490,5 +490,5 @@ public:
 private:
 	FVirtualShadowMapArrayCacheManager& CacheManager;
 
-	FScenePostUpdateChangeSet PostUpdateChangeSet;
+	const FScenePostUpdateChangeSet* PostUpdateChangeSet = nullptr;
 };
\ No newline at end of file
diff --git a/Engine/Source/Runtime/Renderer/Public/PrimitiveSceneInfo.h b/Engine/Source/Runtime/Renderer/Public/PrimitiveSceneInfo.h
--- a/Engine/Source/Runtime/Renderer/Public/PrimitiveSceneInfo.h
+++ b/Engine/Source/Runtime/Renderer/Public/PrimitiveSceneInfo.h
@@ -284,6 +284,7 @@ class FPrimitiveSceneInfo : public FDeferredCleanupInterface
 	friend class FSceneRenderer;
 	friend struct FViewDebugInfo;
 public:
+	using FPersistentId = FPersistentPrimitiveIndex;
 
 	/** The render proxy for the primitive. */
 	FPrimitiveSceneProxy* Proxy;
