RVT: Disable support for rendering locked root pages outside of their scene update.
This could cause a crash where we are updating the scene while render commands are in flight.
This path can only happen for DrawTileMesh() style rendering used by material/HLOD baking.
Now the RVT pages won't be updated *unless* they are already warmed up by a standard scene update.
To avoid any issues it is good practice to avoid sampling RVT in material baking, and if it is necessary then ensure that you have streaming mips built and enabled.

diff --git a/Engine/Source/Runtime/Renderer/Private/VT/RuntimeVirtualTextureProducer.cpp b/Engine/Source/Runtime/Renderer/Private/VT/RuntimeVirtualTextureProducer.cpp
--- a/Engine/Source/Runtime/Renderer/Private/VT/RuntimeVirtualTextureProducer.cpp
+++ b/Engine/Source/Runtime/Renderer/Private/VT/RuntimeVirtualTextureProducer.cpp
@@ -156,12 +156,13 @@ FVTRequestPageResult FRuntimeVirtualTextureProducer::RequestPageData(
 	uint64 vAddress,
 	EVTRequestPagePriority Priority)
 {
-	//todo[vt]: 
-	// Possibly throttle rendering according to performance and return Saturated here.
-
+	// Note that when the finalizer is not ready (outside of the Begin/End Scene Render) we return the Saturated status here.
+	// This is to indicate that the RVT can't render at this time (because we require the GPU Scene to be up to date).
+	// This will happen for DrawTileMesh() style rendering used by material/HLOD baking.
+	// It's best to avoid sampling RVT in material baking, but if it is necessary then an option is to have streaming mips built and enabled.
 	FVTRequestPageResult result;
 	result.Handle = 0;
-	result.Status = Finalizer.IsReady() ? EVTRequestPageStatus::Available : EVTRequestPageStatus::Pending;
+	result.Status = Finalizer.IsReady() ? EVTRequestPageStatus::Available : EVTRequestPageStatus::Saturated;
 	return result;
 }
 
diff --git a/Engine/Source/Runtime/Renderer/Private/VT/RuntimeVirtualTextureRender.cpp b/Engine/Source/Runtime/Renderer/Private/VT/RuntimeVirtualTextureRender.cpp
--- a/Engine/Source/Runtime/Renderer/Private/VT/RuntimeVirtualTextureRender.cpp
+++ b/Engine/Source/Runtime/Renderer/Private/VT/RuntimeVirtualTextureRender.cpp
@@ -1707,7 +1707,7 @@ namespace RuntimeVirtualTexture
 	void RenderPages(FRDGBuilder& GraphBuilder, FRenderPageBatchDesc const& InDesc)
 	{
 		check(InDesc.Scene != nullptr);
-		if (InDesc.Scene->GPUScene.IsRendering())
+		if (ensure(InDesc.Scene->GPUScene.IsRendering()))
 		{
 			// TODO: this should be replaced by piping through a reference to the scene renderer rather than just the scene, such that we can get at the already populated scene UB.
 			class FSimpleRVTRenderer : public FSceneRendererBase
@@ -1726,12 +1726,6 @@ namespace RuntimeVirtualTexture
 			const bool bAllowCachedMeshDrawCommands = true;
 			RenderPagesInternal(GraphBuilder, InDesc, &SimpleRenderer, bAllowCachedMeshDrawCommands);
 		}
-		else
-		{
-			// We allow locked root pages to be rendered outside of their scene update.
-			// We expect to hit this path very rarely. (One case is during material baking.)
-			RenderPagesStandAlone(GraphBuilder, InDesc);
-		}
 	}
 
 	uint32 GetRuntimeVirtualTextureSceneIndex_GameThread(class URuntimeVirtualTextureComponent* InComponent)
