* Remove CPU path for deep shadow map
* Extend limit from 16 to 32 DOM slots

diff --git a/Engine/Shaders/Private/HairStrands/HairStrandsDebug.usf b/Engine/Shaders/Private/HairStrands/HairStrandsDebug.usf
--- a/Engine/Shaders/Private/HairStrands/HairStrandsDebug.usf
+++ b/Engine/Shaders/Private/HairStrands/HairStrandsDebug.usf
@@ -838,11 +838,11 @@ void VisualizeDomPS(
 	float2 PixelPosAtTileCenter = PixelPos - (frac(PixelPos / TileSize) - 0.5f) * TileSize;
 	float2 UVAtTileCenter = PixelPosAtTileCenter * InvOutputResolution;
 
-	const float OffsetX = 50;
-	const float OffsetY = 220;
+	const float OffsetX = 450;
+	const float OffsetY = 50;
 	const float RectSize = 300;
 	const float4 DepthRect = float4(OffsetX, OffsetY, RectSize, RectSize);
-	const float4 DomRect = float4(OffsetX * 2 + RectSize, OffsetY, RectSize, RectSize);
+	const float4 DomRect = float4(OffsetX + RectSize, OffsetY, RectSize, RectSize);
 
 	float3 Color = 0;
 
@@ -919,7 +919,6 @@ uint MacroGroupCount;
 uint bViewRectOptimizeEnabled;
 uint bVoxelizationEnabled;
 int2 AtlasResolution;
-uint bIsGPUDriven;
 
 Buffer<int>	MacroGroupAABBBuffer;
 StructuredBuffer<FDeepShadowViewInfo> ShadowViewInfoBuffer;
@@ -932,55 +931,6 @@ void MainCS(uint3 DispatchThreadId : SV_DispatchThreadID)
 		return;
 	}
 
-	for (uint MacroGroupIt = 0; MacroGroupIt < MacroGroupCount; ++MacroGroupIt)
-	{
-		const FHairAABB Bound = ReadHairAABB(MacroGroupIt, MacroGroupAABBBuffer);
-		AddAABBTWS(Bound.Min, Bound.Max, GetGPUAABBMacroGroupBoundVoxelAlignedColor());
-	}
-
-	const float4 FrustumColor = float4(0.5f, 1, 0, 1);
-	LOOP
-	for (uint SlotIt = 0; SlotIt < AllocatedSlotCount; ++SlotIt)
-	{
-		const float4x4 TranslatedWorldToClip = ShadowViewInfoBuffer[SlotIt].TranslatedWorldToClip;
-		const float4x4 ClipToTranslatedWorld = Inverse(TranslatedWorldToClip);
-
-		float4 P01 = float4(-1, -1, 0, 1);
-		float4 P11 = float4(-1,  1, 0, 1);
-		float4 P21 = float4( 1,  1, 0, 1);
-		float4 P31 = float4( 1, -1, 0, 1);
-
-		float4 P00 = float4(-1, -1, 1, 1);
-		float4 P10 = float4(-1,  1, 1, 1);
-		float4 P20 = float4( 1,  1, 1, 1);
-		float4 P30 = float4( 1, -1, 1, 1);
-
-		P00 = mul(P00, ClipToTranslatedWorld); P00 /= P00.w;
-		P10 = mul(P10, ClipToTranslatedWorld); P10 /= P10.w;
-		P20 = mul(P20, ClipToTranslatedWorld); P20 /= P20.w;
-		P30 = mul(P30, ClipToTranslatedWorld); P30 /= P30.w;
-
-		P01 = mul(P01, ClipToTranslatedWorld); P01 /= P01.w;
-		P11 = mul(P11, ClipToTranslatedWorld); P11 /= P11.w;
-		P21 = mul(P21, ClipToTranslatedWorld); P21 /= P21.w;
-		P31 = mul(P31, ClipToTranslatedWorld); P31 /= P31.w;
-
-		AddLineTWS(P00.xyz, P10.xyz, FrustumColor, FrustumColor);
-		AddLineTWS(P10.xyz, P20.xyz, FrustumColor, FrustumColor);
-		AddLineTWS(P20.xyz, P30.xyz, FrustumColor, FrustumColor);
-		AddLineTWS(P30.xyz, P00.xyz, FrustumColor, FrustumColor);
-
-		AddLineTWS(P00.xyz, P01.xyz, FrustumColor, FrustumColor);
-		AddLineTWS(P10.xyz, P11.xyz, FrustumColor, FrustumColor);
-		AddLineTWS(P20.xyz, P21.xyz, FrustumColor, FrustumColor);
-		AddLineTWS(P30.xyz, P31.xyz, FrustumColor, FrustumColor);
-
-		AddLineTWS(P01.xyz, P11.xyz, FrustumColor, FrustumColor);
-		AddLineTWS(P11.xyz, P21.xyz, FrustumColor, FrustumColor);
-		AddLineTWS(P21.xyz, P31.xyz, FrustumColor, FrustumColor);
-		AddLineTWS(P31.xyz, P01.xyz, FrustumColor, FrustumColor);
-	}
-
 
 	FShaderPrintContext Context = InitShaderPrintContext(true, uint2(50, 40));
 
@@ -1001,7 +951,67 @@ void MainCS(uint3 DispatchThreadId : SV_DispatchThreadID)
 	Print(Context, TEXT("DOM "), FontYellow); Newline(Context);
 	Print(Context, TEXT("Atlas resolution      : "), FontColor); Print(Context, AtlasResolution.x, FontColor, 4, 0); Print(Context, TEXT(" x "), FontColor); Print(Context, AtlasResolution.y, FontColor, 4, 0); Newline(Context);
 	Print(Context, TEXT("Atlas slots           : "), FontColor); Print(Context, AllocatedSlotCount, FontColor); Newline(Context);
-	Print(Context, TEXT("GPU driven            : "), FontColor); PrintBool(Context, bIsGPUDriven); Newline(Context);
+	Print(Context, TEXT("GPU driven            : "), FontColor); PrintBool(Context, true); Newline(Context);
+	Newline(Context);
+
+	Print(Context, TEXT("Bounds "), FontYellow); Newline(Context);
+	const bool bMacroAABB  = AddCheckbox(Context, TEXT("Macro  AABB "), false, FontColor); Newline(Context);
+	const bool bShadowAABB = AddCheckbox(Context, TEXT("Shadow AABB "), false, FontColor); Newline(Context);
+
+	if (bMacroAABB)
+	{
+		for (uint MacroGroupIt = 0; MacroGroupIt < MacroGroupCount; ++MacroGroupIt)
+		{
+			const FHairAABB Bound = ReadHairAABB(MacroGroupIt, MacroGroupAABBBuffer);
+			AddAABBTWS(Bound.Min, Bound.Max, GetGPUAABBMacroGroupBoundVoxelAlignedColor());
+		}
+	}
+
+	if (bShadowAABB)
+	{
+		const float4 FrustumColor = float4(0.5f, 1, 0, 1);
+		LOOP
+		for (uint SlotIt = 0; SlotIt < AllocatedSlotCount; ++SlotIt)
+		{
+			const float4x4 TranslatedWorldToClip = ShadowViewInfoBuffer[SlotIt].TranslatedWorldToClip;
+			const float4x4 ClipToTranslatedWorld = Inverse(TranslatedWorldToClip);
+
+			float4 P01 = float4(-1, -1, 0, 1);
+			float4 P11 = float4(-1,  1, 0, 1);
+			float4 P21 = float4( 1,  1, 0, 1);
+			float4 P31 = float4( 1, -1, 0, 1);
+
+			float4 P00 = float4(-1, -1, 1, 1);
+			float4 P10 = float4(-1,  1, 1, 1);
+			float4 P20 = float4( 1,  1, 1, 1);
+			float4 P30 = float4( 1, -1, 1, 1);
+
+			P00 = mul(P00, ClipToTranslatedWorld); P00 /= P00.w;
+			P10 = mul(P10, ClipToTranslatedWorld); P10 /= P10.w;
+			P20 = mul(P20, ClipToTranslatedWorld); P20 /= P20.w;
+			P30 = mul(P30, ClipToTranslatedWorld); P30 /= P30.w;
+
+			P01 = mul(P01, ClipToTranslatedWorld); P01 /= P01.w;
+			P11 = mul(P11, ClipToTranslatedWorld); P11 /= P11.w;
+			P21 = mul(P21, ClipToTranslatedWorld); P21 /= P21.w;
+			P31 = mul(P31, ClipToTranslatedWorld); P31 /= P31.w;
+
+			AddLineTWS(P00.xyz, P10.xyz, FrustumColor, FrustumColor);
+			AddLineTWS(P10.xyz, P20.xyz, FrustumColor, FrustumColor);
+			AddLineTWS(P20.xyz, P30.xyz, FrustumColor, FrustumColor);
+			AddLineTWS(P30.xyz, P00.xyz, FrustumColor, FrustumColor);
+
+			AddLineTWS(P00.xyz, P01.xyz, FrustumColor, FrustumColor);
+			AddLineTWS(P10.xyz, P11.xyz, FrustumColor, FrustumColor);
+			AddLineTWS(P20.xyz, P21.xyz, FrustumColor, FrustumColor);
+			AddLineTWS(P30.xyz, P31.xyz, FrustumColor, FrustumColor);
+
+			AddLineTWS(P01.xyz, P11.xyz, FrustumColor, FrustumColor);
+			AddLineTWS(P11.xyz, P21.xyz, FrustumColor, FrustumColor);
+			AddLineTWS(P21.xyz, P31.xyz, FrustumColor, FrustumColor);
+			AddLineTWS(P31.xyz, P01.xyz, FrustumColor, FrustumColor);
+		}
+	}
 }
 #endif // SHADER_DOMINFO
 
diff --git a/Engine/Shaders/Private/HairStrands/HairStrandsDeepShadowCommon.ush b/Engine/Shaders/Private/HairStrands/HairStrandsDeepShadowCommon.ush
--- a/Engine/Shaders/Private/HairStrands/HairStrandsDeepShadowCommon.ush
+++ b/Engine/Shaders/Private/HairStrands/HairStrandsDeepShadowCommon.ush
@@ -154,7 +154,7 @@ float2 ComputeJitteredSample(float2 Sample, float2 Jitter)
 
 float SampleDOM_PCSS(
 	float3 LightSpacePosition,
-	uint2 DeepShadowResolution,
+	uint2 DeepShadowAtlasResolution,
 	float DepthBias,
 	float4 InLayerDepths,
 	Texture2D<float> FrontDepthTexture,
@@ -188,7 +188,7 @@ float SampleDOM_PCSS(
 	float OccluderDistance = 0; // Occluder distance in clip space
 	for (uint OccluderIt = 0; OccluderIt < OccluderCount; ++OccluderIt)
 	{
-		const float2 Offset = ComputeJitteredSample(PoissonDisk[OccluderIt], Jitter) * DeepShadowResolution;
+		const float2 Offset = ComputeJitteredSample(PoissonDisk[OccluderIt], Jitter) * DeepShadowAtlasResolution;
 		const float2 SamplePosition = LightSpacePosition.xy + Offset;
 		const float2 P = floor(SamplePosition - float(0.5f).xx);
 		const uint2 C0 = P;
@@ -218,7 +218,7 @@ float SampleDOM_PCSS(
 	const float SampleRadius = OccluderDistance * tan(ApexHalfAngleInRad); // Sample radius in clip space
 	for (uint SampleIt = 0; SampleIt < SampleCount; ++SampleIt)
 	{
-		const float2 Offset = ComputeJitteredSample(PoissonDisk[SampleIt], Jitter) * SampleRadius * DeepShadowResolution;
+		const float2 Offset = ComputeJitteredSample(PoissonDisk[SampleIt], Jitter) * SampleRadius * DeepShadowAtlasResolution;
 		HairCount += SampleDOM_PCF2x2(LightSpacePosition + float3(Offset,0), DepthBias, InLayerDepths, FrontDepthTexture, DomTexture);
 	}
 	
diff --git a/Engine/Shaders/Private/HairStrands/HairStrandsDeepShadowMask.usf b/Engine/Shaders/Private/HairStrands/HairStrandsDeepShadowMask.usf
--- a/Engine/Shaders/Private/HairStrands/HairStrandsDeepShadowMask.usf
+++ b/Engine/Shaders/Private/HairStrands/HairStrandsDeepShadowMask.usf
@@ -42,9 +42,7 @@ float3 GetTranslatedWorldPosition(float2 UV, float SceneDepth)
 #if SHADER_SHADOWMASK_DEEPSHADOW
 
 StructuredBuffer<FDeepShadowViewInfo> DeepShadow_ViewInfoBuffer;
-uint				DeepShadow_bIsGPUDriven;
 
-float4x4			DeepShadow_CPUTranslatedWorldToLightTransform;
 int2				DeepShadow_SlotResolution;
 int2				DeepShadow_SlotOffset;
 uint				DeepShadow_SlotIndex;
@@ -67,18 +65,9 @@ SamplerState 		ShadowSampler;
 //  W: Positive if the position is valid, negative otherwise
 float4 ToLightPosition(float3 InTranslatedWorldPosition, out float4 AtlasScaleBias)
 {
-	float4x4 TranslatedWorldToLightTransform = 0;
-	if (DeepShadow_bIsGPUDriven > 0)
-	{
-		FDeepShadowViewInfo ShadowViewInfo = DeepShadow_ViewInfoBuffer[DeepShadow_SlotIndex];
-		TranslatedWorldToLightTransform = ShadowViewInfo.TranslatedWorldToClip;
-		AtlasScaleBias = ShadowViewInfo.AtlasScaleBias;
-	}
-	else
-	{
-		TranslatedWorldToLightTransform = DeepShadow_CPUTranslatedWorldToLightTransform;
-		AtlasScaleBias = float4(1.0f, 1.0f, 0.0f, 0.0f);
-	}
+	FDeepShadowViewInfo ShadowViewInfo = DeepShadow_ViewInfoBuffer[DeepShadow_SlotIndex];
+	float4x4 TranslatedWorldToLightTransform = ShadowViewInfo.TranslatedWorldToClip;
+	AtlasScaleBias = ShadowViewInfo.AtlasScaleBias;
 
 	float4 LightPos = mul(float4(InTranslatedWorldPosition,1), TranslatedWorldToLightTransform);
 	LightPos.xyz/= LightPos.w;
diff --git a/Engine/Shaders/Private/HairStrands/HairStrandsDeepShadowVS.usf b/Engine/Shaders/Private/HairStrands/HairStrandsDeepShadowVS.usf
--- a/Engine/Shaders/Private/HairStrands/HairStrandsDeepShadowVS.usf
+++ b/Engine/Shaders/Private/HairStrands/HairStrandsDeepShadowVS.usf
@@ -9,6 +9,10 @@
 #define MESH_RENDER_DEPTH 0
 #define MESH_RENDER_DOM   1
 
+#if MESH_RENDER_MODE != MESH_RENDER_DEPTH && MESH_RENDER_MODE != MESH_RENDER_DOM
+#error Unknown DeepShadow render mode 
+#endif
+
 struct FDeepShadowAccumulateVSToPS
 {
 #if SUPPORT_TANGENT_PROPERTY == 1 || SUPPORT_MATERIAL_PROPERTY == 1
@@ -35,49 +39,38 @@ void Main(
 
 	FVertexFactoryIntermediates VFIntermediates = GetVertexFactoryIntermediates(Input);
 
-#if MESH_RENDER_MODE == MESH_RENDER_DEPTH || MESH_RENDER_MODE == MESH_RENDER_DOM
-	float4x4 TranslatedWorldToClipMatrix = DeepRasterPass.CPU_TranslatedWorldToClipMatrix;
-#endif
-
-#ifdef HAIR_STRAND_MESH_FACTORY	
+#ifdef HAIR_STRAND_MESH_FACTORY
 	const FHairRenderInfo HairRenderInfo = GetHairRenderInfo(ResolvedView.HairRenderInfo, ResolvedView.HairRenderInfoBits);
+	const FDeepShadowViewInfo DeepShadowViewInfo = DeepRasterPass.DeepShadowViewInfoBuffer[DeepRasterPass.AtlasSlotIndex];
 
 	FHairViewInfo HairViewInfo;
 	HairViewInfo.TranslatedWorldCameraOrigin = ResolvedView.TranslatedWorldCameraOrigin;
-	HairViewInfo.ViewForward = ResolvedView.ViewForward;
-	HairViewInfo.RadiusAtDepth1 = HairRenderInfo.RadiusAtDepth1Primary;
+	HairViewInfo.ViewForward = DeepShadowViewInfo.ViewForward;
+	HairViewInfo.RadiusAtDepth1 = DeepShadowViewInfo.MinRadiusAtDepth1;
 	HairViewInfo.bIsOrthoView = HairRenderInfo.bIsOrthoView;
 
-	#if MESH_RENDER_MODE == MESH_RENDER_DEPTH || MESH_RENDER_MODE == MESH_RENDER_DOM
-	if (HairRenderInfo.bIsGPUDriven)
-	{
-		const uint AtlasSlotIndex = DeepRasterPass.AtlasSlotIndex;
-		FDeepShadowViewInfo DeepShadowViewInfo = DeepRasterPass.DeepShadowViewInfoBuffer[AtlasSlotIndex];
-		TranslatedWorldToClipMatrix = DeepShadowViewInfo.TranslatedWorldToClipScaledBiased;
-		HairViewInfo.ViewForward	= DeepShadowViewInfo.ViewForward;
-		HairViewInfo.RadiusAtDepth1 = DeepShadowViewInfo.MinRadiusAtDepth1;
-	}
-	#endif
-
-	float4 TranslatedWorldPosition = VertexFactoryGetWorldPosition(Input, VFIntermediates, HairViewInfo);
+	const float4x4 TranslatedWorldToClipMatrix = DeepShadowViewInfo.TranslatedWorldToClipScaledBiased;
+	const float4 TranslatedWorldPosition = VertexFactoryGetWorldPosition(Input, VFIntermediates, HairViewInfo);
 #else
-	float4 TranslatedWorldPosition = VertexFactoryGetWorldPosition(Input, VFIntermediates);
+	const float4x4 TranslatedWorldToClipMatrix = 0;
+	const float4 TranslatedWorldPosition = VertexFactoryGetWorldPosition(Input, VFIntermediates);
 #endif
 
 	Output.Position = mul(TranslatedWorldPosition, TranslatedWorldToClipMatrix);
 
-	float HairCoverage = 1;
-  #ifdef HAIR_STRAND_MESH_FACTORY	
+#ifdef HAIR_STRAND_MESH_FACTORY	
 	const float CurrentDepth = Output.Position.z / Output.Position.w;
 	const float PixelRadius = HairRenderInfo.bIsOrthoView ? HairRenderInfo.RadiusAtDepth1Primary : (CurrentDepth * HairRenderInfo.RadiusAtDepth1Primary);
 	const float StrandRealRadius = VFIntermediates.HairDimensions.y;
 	const float Coverage = StrandRealRadius / max(StrandRealRadius, PixelRadius);
-	HairCoverage = Coverage * VFIntermediates.HairDensity;
-  #endif
+	const float HairCoverage = Coverage * VFIntermediates.HairDensity;
+#else
+	const float HairCoverage = 1;
+#endif
 
-  #if MESH_RENDER_MODE == MESH_RENDER_DOM
+#if MESH_RENDER_MODE == MESH_RENDER_DOM
 	Output.HairCoverage = HairCoverage;
-  #endif
+#endif
 }
 
 #endif // VERTEXSHADER
diff --git a/Engine/Shaders/Private/HairStrands/HairStrandsDeepTransmittanceMask.usf b/Engine/Shaders/Private/HairStrands/HairStrandsDeepTransmittanceMask.usf
--- a/Engine/Shaders/Private/HairStrands/HairStrandsDeepTransmittanceMask.usf
+++ b/Engine/Shaders/Private/HairStrands/HairStrandsDeepTransmittanceMask.usf
@@ -23,8 +23,7 @@
 // Utils function & structs
 struct FTransmittanceSettings
 {
-	int2	 DeepShadowResolution;
-	int2	 DeepShadowOffset;
+	int2	 DeepShadowAtlasResolution;
 	float4x4 TranslatedWorldToLightTransform;
 	float4 AtlasScaleBias;
 	float3   LightDirection; // Direction from world position to the light
@@ -72,12 +71,10 @@ float3 GetTranslatedWorldPosition(float2 UV, float SceneDepth)
 
 #if SHADER_TRANSMITTANCE_DEEPSHADOW
 
+StructuredBuffer<uint> DeepShadow_AtlasSlotIndexBuffer;
 StructuredBuffer<FDeepShadowViewInfo> DeepShadow_ViewInfoBuffer;
-uint				DeepShadow_bIsGPUDriven;
 
-uint2				DeepShadow_Resolution;
-float4x4			DeepShadow_CPUTranslatedWorldToLightTransforms[MAX_HAIR_MACROGROUP_COUNT];
-//float3				LightDirection;
+uint2				DeepShadow_AtlasResolution;
 uint				LightChannelMask;
 float4				ShadowChannelMask;
 float4				TranslatedLightPosition_LightDirection;
@@ -88,7 +85,6 @@ float				DeepShadow_KernelAperture;
 float4				DeepShadow_LayerDepths;
 uint				DeepShadow_KernelType;
 uint				DeepShadow_DebugMode;
-int4				DeepShadow_AtlasSlotOffsets_AtlasSlotIndex[MAX_HAIR_MACROGROUP_COUNT];
 Texture2D<float>	DeepShadow_FrontDepthTexture;
 Texture2D<float4>	DeepShadow_DomTexture;
 
@@ -131,7 +127,7 @@ FHairTransmittanceMask ComputeTransmittance(
 			float HairCount = 0;
 			if (Settings.TransmittanceKernelType == 3)
 			{
-				HairCount = SampleDOM_PCSS(LightSpacePosition.xyz, Settings.DeepShadowResolution, DepthBias, InLayerDepths, FrontDepthTexture, DomTexture, Settings.TransmittanceKernelApertureInDegree, Settings.Random.xy);
+				HairCount = SampleDOM_PCSS(LightSpacePosition.xyz, Settings.DeepShadowAtlasResolution, DepthBias, InLayerDepths, FrontDepthTexture, DomTexture, Settings.TransmittanceKernelApertureInDegree, Settings.Random.xy);
 			}
 			else if (Settings.TransmittanceKernelType == 2)
 			{
@@ -184,27 +180,17 @@ void MainCS(uint2 DispatchThreadId : SV_DispatchThreadID)
 		}
 
 		const uint MacroGroupIndex = clamp(Sample.MacroGroupId, 0, MAX_HAIR_MACROGROUP_COUNT - 1);
-		const int3 SlotOffset_SlotIndex = DeepShadow_AtlasSlotOffsets_AtlasSlotIndex[MacroGroupIndex].xyz;
-		float4x4 TranslatedWorldToLightTransform = 0;
-		float4 AtlasScaleBias = float4(1.0f, 1.0f, 0.0f, 0.0f);
-		if (DeepShadow_bIsGPUDriven > 0)
-		{
-			FDeepShadowViewInfo ShadowViewInfo = DeepShadow_ViewInfoBuffer[SlotOffset_SlotIndex.z];
-			TranslatedWorldToLightTransform = ShadowViewInfo.TranslatedWorldToClip;
-			AtlasScaleBias = ShadowViewInfo.AtlasScaleBias * float4(DeepShadow_Resolution, DeepShadow_Resolution);
-		}	
-		else
-		{
-			TranslatedWorldToLightTransform = DeepShadow_CPUTranslatedWorldToLightTransforms[MacroGroupIndex];
-			AtlasScaleBias = float4(DeepShadow_Resolution, SlotOffset_SlotIndex.xy);
-		}
+		const uint AtlasSlotIndex = DeepShadow_AtlasSlotIndexBuffer[MacroGroupIndex];
+
+		const FDeepShadowViewInfo ShadowViewInfo = DeepShadow_ViewInfoBuffer[AtlasSlotIndex];
+		const float4x4 TranslatedWorldToLightTransform = ShadowViewInfo.TranslatedWorldToClip;
+		const float4 AtlasScaleBias = ShadowViewInfo.AtlasScaleBias * float4(DeepShadow_AtlasResolution, DeepShadow_AtlasResolution);
 		
 		const float SceneDepth = ConvertFromDeviceZ(Sample.Depth);
 		const float3 TranslatedWorldPosition = GetTranslatedWorldPosition(UV, SceneDepth);
 
 		FTransmittanceSettings Settings;
-		Settings.DeepShadowOffset = SlotOffset_SlotIndex.xy;
-		Settings.DeepShadowResolution = DeepShadow_Resolution;
+		Settings.DeepShadowAtlasResolution = DeepShadow_AtlasResolution;
 		Settings.TranslatedWorldToLightTransform = TranslatedWorldToLightTransform;
 		Settings.AtlasScaleBias = AtlasScaleBias;
 		Settings.bIsDirectional = IsLightDirectional(TranslatedLightPosition_LightDirection);
diff --git a/Engine/Shaders/Private/HairStrands/HairStrandsVoxelRasterCompute.usf b/Engine/Shaders/Private/HairStrands/HairStrandsVoxelRasterCompute.usf
--- a/Engine/Shaders/Private/HairStrands/HairStrandsVoxelRasterCompute.usf
+++ b/Engine/Shaders/Private/HairStrands/HairStrandsVoxelRasterCompute.usf
@@ -913,8 +913,6 @@ void MainCS(uint3 DispatchThreadId : SV_DispatchThreadID)
 ///////////////////////////////////////////////////////////////////////////
 
 #if SHADER_DEPTH_INJECTION
-float4x4 CPU_TranslatedWorldToClip;
-
 float2 OutputResolution;
 uint MacroGroupId;
 uint AtlasSlotIndex;
@@ -922,7 +920,6 @@ uint AtlasSlotIndex;
 float3 LightDirection;
 uint bIsDirectional;
 float3 TranslatedLightPosition;
-uint bIsGPUDriven;
 
 StructuredBuffer<FDeepShadowViewInfo> DeepShadowViewInfoBuffer;
 
@@ -996,12 +993,8 @@ void MainVS(
 	case 35: TranslatedWorldPosition = Position7; break;
 	}
 
-	float4x4 TranslatedWorldToClipMatrix = CPU_TranslatedWorldToClip;
-	if (bIsGPUDriven)
-	{
-		FDeepShadowViewInfo DeepShadowViewInfo = DeepShadowViewInfoBuffer[AtlasSlotIndex];
-		TranslatedWorldToClipMatrix = DeepShadowViewInfo.TranslatedWorldToClipScaledBiased;
-	}
+	const FDeepShadowViewInfo DeepShadowViewInfo = DeepShadowViewInfoBuffer[AtlasSlotIndex];
+	const float4x4 TranslatedWorldToClipMatrix = DeepShadowViewInfo.TranslatedWorldToClipScaledBiased;
 
 	OutTranslatedWorldPosition = TranslatedWorldPosition;
 	OutPosition = mul(float4(TranslatedWorldPosition, 1), TranslatedWorldToClipMatrix);
@@ -1062,12 +1055,8 @@ void MainPS(
 
 	if (bIsValid)
 	{
-		float4x4 TranslatedWorldToClipMatrix = CPU_TranslatedWorldToClip;
-		if (bIsGPUDriven)
-		{
-			FDeepShadowViewInfo DeepShadowViewInfo = DeepShadowViewInfoBuffer[AtlasSlotIndex];
-			TranslatedWorldToClipMatrix = DeepShadowViewInfo.TranslatedWorldToClipScaledBiased;
-		}
+		const FDeepShadowViewInfo DeepShadowViewInfo = DeepShadowViewInfoBuffer[AtlasSlotIndex];
+		const float4x4 TranslatedWorldToClipMatrix = DeepShadowViewInfo.TranslatedWorldToClipScaledBiased;
 		const float3 HitP = TranslatedWP0 + normalize(TranslatedWP1 - TranslatedWP0) * TraversalResult.HitT;
 		float4 ClipP = mul(float4(HitP, 1), TranslatedWorldToClipMatrix);
 		OutDepth = ClipP.z /= ClipP.w;
diff --git a/Engine/Source/Runtime/Renderer/Private/HairStrands/HairStrandsData.h b/Engine/Source/Runtime/Renderer/Private/HairStrands/HairStrandsData.h
--- a/Engine/Source/Runtime/Renderer/Private/HairStrands/HairStrandsData.h
+++ b/Engine/Source/Runtime/Renderer/Private/HairStrands/HairStrandsData.h
@@ -238,15 +238,12 @@ struct FMinHairRadiusAtDepth1
 /// Hold deep shadow information for a given light.
 struct FHairStrandsDeepShadowData
 {
-	static const uint32 MaxMacroGroupCount = 16u;
-
-	FMatrix CPU_TranslatedWorldToLightTransform;
 	FMinHairRadiusAtDepth1 CPU_MinStrandRadiusAtDepth1;
 	FIntRect AtlasRect;
 	uint32 MacroGroupId = ~0;
 	uint32 AtlasSlotIndex = 0;
 
-	FIntPoint ShadowResolution = FIntPoint::ZeroValue;
+	FIntPoint AtlasResolution = FIntPoint::ZeroValue;
 	uint32 LightId = ~0;
 	bool bIsLightDirectional = false;
 	FVector3f LightDirection;
@@ -267,7 +264,6 @@ struct FHairStrandsDeepShadowResources
 
 	uint32 TotalAtlasSlotCount = 0;
 	FIntPoint AtlasSlotResolution;
-	bool bIsGPUDriven = false;
 
 	FRDGTextureRef DepthAtlasTexture = nullptr;
 	FRDGTextureRef LayersAtlasTexture = nullptr;
diff --git a/Engine/Source/Runtime/Renderer/Private/HairStrands/HairStrandsDebug.cpp b/Engine/Source/Runtime/Renderer/Private/HairStrands/HairStrandsDebug.cpp
--- a/Engine/Source/Runtime/Renderer/Private/HairStrands/HairStrandsDebug.cpp
+++ b/Engine/Source/Runtime/Renderer/Private/HairStrands/HairStrandsDebug.cpp
@@ -697,7 +697,6 @@ class FDeepShadowInfoCS : public FGlobalShader
 		SHADER_PARAMETER(uint32, bViewRectOptimizeEnabled)
 		SHADER_PARAMETER(uint32, bVoxelizationEnabled)
 		SHADER_PARAMETER(FIntPoint, AtlasResolution)
-		SHADER_PARAMETER(uint32, bIsGPUDriven)
 		SHADER_PARAMETER_STRUCT_REF(FViewUniformShaderParameters, ViewUniformBuffer)
 		SHADER_PARAMETER_RDG_BUFFER_SRV(Buffer, MacroGroupAABBBuffer)
 		SHADER_PARAMETER_RDG_BUFFER_SRV(StructuredBuffer, ShadowViewInfoBuffer)
@@ -753,7 +752,6 @@ static void AddDeepShadowInfoPass(
 	Parameters->bViewRectOptimizeEnabled = IsHairStrandsViewRectOptimEnable() ? 1u : 0u;
 	Parameters->bVoxelizationEnabled = View.HairStrandsViewData.VirtualVoxelResources.IsValid() ? 1u : 0u;
 	Parameters->AtlasResolution = View.HairStrandsViewData.DeepShadowResources.DepthAtlasTexture->Desc.Extent;
-	Parameters->bIsGPUDriven = View.HairStrandsViewData.DeepShadowResources.bIsGPUDriven ? 1u : 0u;
 	Parameters->MacroGroupAABBBuffer = GraphBuilder.CreateSRV(MacroGroupResources.MacroGroupAABBsBuffer, PF_R32_SINT);
 	Parameters->ShadowViewInfoBuffer = GraphBuilder.CreateSRV(DeepShadowResources.DeepShadowViewInfoBuffer);
 	ShaderPrint::SetParameters(GraphBuilder, View.ShaderPrintData, Parameters->ShaderPrintParameters);
@@ -1263,10 +1261,7 @@ static void InternalRenderHairStrandsDebugInfo(
 	{
 		if ((ViewMode == EGroomViewMode::LightBounds || ViewMode == EGroomViewMode::DeepOpacityMaps))
 		{
-			if (HairData.DeepShadowResources.bIsGPUDriven)
-			{
-				AddDeepShadowInfoPass(GraphBuilder, View, HairData.DeepShadowResources, HairData.MacroGroupResources, Params.SceneColorTexture);
-			}
+			AddDeepShadowInfoPass(GraphBuilder, View, HairData.DeepShadowResources, HairData.MacroGroupResources, Params.SceneColorTexture);
 		}
 	}
 	
diff --git a/Engine/Source/Runtime/Renderer/Private/HairStrands/HairStrandsDeepShadow.cpp b/Engine/Source/Runtime/Renderer/Private/HairStrands/HairStrandsDeepShadow.cpp
--- a/Engine/Source/Runtime/Renderer/Private/HairStrands/HairStrandsDeepShadow.cpp
+++ b/Engine/Source/Runtime/Renderer/Private/HairStrands/HairStrandsDeepShadow.cpp
@@ -16,9 +16,6 @@ static FAutoConsoleVariableRef CVarDeepShadowResolution(TEXT("r.HairStrands.Deep
 static int32 GDeepShadowMinResolution = 64;
 static FAutoConsoleVariableRef CVarDeepShadowMinResolution(TEXT("r.HairStrands.DeepShadow.MinResolution"), GDeepShadowMinResolution, TEXT("Minimum shadow resolution for shadow atlas tiles for Deep Opacity Map rendering. (default = 64)"));
 
-static int32 GDeepShadowGPUDriven = 1;
-static FAutoConsoleVariableRef CVarDeepShadowGPUDriven(TEXT("r.HairStrands.DeepShadow.GPUDriven"), GDeepShadowGPUDriven, TEXT("Enable deep shadow to be driven by GPU bounding box, rather CPU ones. This allows more robust behavior"));
-
 static int32 GDeepShadowInjectVoxelDepth = 0;
 static FAutoConsoleVariableRef CVarDeepShadowInjectVoxelDepth(TEXT("r.HairStrands.DeepShadow.InjectVoxelDepth"), GDeepShadowInjectVoxelDepth, TEXT("Inject voxel content to generate the deep shadow map instead of rasterizing groom. This is an experimental path"));
 
@@ -26,11 +23,8 @@ static FAutoConsoleVariableRef CVarDeepShadowInjectVoxelDepth(TEXT("r.HairStrand
 // Inject voxel structure into shadow map to amortize the tracing, and rely on look up kernel to 
 // filter limited resolution
 BEGIN_SHADER_PARAMETER_STRUCT(FHairStransShadowDepthInjectionParameters, )
-	SHADER_PARAMETER(FMatrix44f, CPU_TranslatedWorldToClip)
-
 	SHADER_PARAMETER(FVector2f, OutputResolution)
 	SHADER_PARAMETER(uint32, AtlasSlotIndex)
-	SHADER_PARAMETER(uint32, bIsGPUDriven)
 
 	SHADER_PARAMETER(FVector3f, LightDirection)
 	SHADER_PARAMETER(uint32, MacroGroupId)
@@ -88,7 +82,6 @@ void AddInjectHairVoxelShadowCaster(
 	const FViewInfo& View,
 	const bool bClear,
 	const FHairStrandsDeepShadowData& DomData,
-	FMatrix CPU_TranslatedWorldToClipMatrix,
 	FIntRect AtlasRect,
 	uint32 AtlasSlotIndex,
 	FIntPoint AtlasSlotResolution,
@@ -98,7 +91,6 @@ void AddInjectHairVoxelShadowCaster(
 {
 	FHairStransShadowDepthInjectionParameters* Parameters = GraphBuilder.AllocParameters<FHairStransShadowDepthInjectionParameters>();
 	Parameters->OutputResolution = AtlasSlotResolution;
-	Parameters->CPU_TranslatedWorldToClip = FMatrix44f(CPU_TranslatedWorldToClipMatrix);
 	Parameters->ViewUniformBuffer = View.ViewUniformBuffer;
 	Parameters->RenderTargets.DepthStencil = FDepthStencilBinding(OutDepthTexture, bClear ? ERenderTargetLoadAction::EClear : ERenderTargetLoadAction::ELoad, ERenderTargetLoadAction::ENoAction, FExclusiveDepthStencil::DepthWrite_StencilNop);
 	Parameters->VirtualVoxel = VoxelResources.UniformBuffer;
@@ -107,7 +99,6 @@ void AddInjectHairVoxelShadowCaster(
 	Parameters->bIsDirectional = DomData.bIsLightDirectional ? 1 : 0;
 	Parameters->MacroGroupId = DomData.MacroGroupId;
 	Parameters->DeepShadowViewInfoBuffer = DeepShadowViewInfoBufferSRV;
-	Parameters->bIsGPUDriven = GDeepShadowGPUDriven > 0;
 	Parameters->AtlasSlotIndex = AtlasSlotIndex;
 
 	TShaderMapRef<FHairStrandsShadowDepthInjectionVS> VertexShader(View.ShaderMap);
@@ -210,7 +201,7 @@ public:
 	{
 		FGlobalShader::ModifyCompilationEnvironment(Parameters, OutEnvironment);
 		OutEnvironment.SetDefine(TEXT("SHADER_ALLOCATE"), 1);
-		OutEnvironment.SetDefine(TEXT("MAX_SLOT_COUNT"), FHairStrandsDeepShadowData::MaxMacroGroupCount);
+		OutEnvironment.SetDefine(TEXT("MAX_SLOT_COUNT"), FHairStrandsDeepShadowResources::MaxAtlasSlotCount);
 	}
 };
 
@@ -260,7 +251,7 @@ void RenderHairStrandsDeepShadows(
 			return; 
 		}
 
-		// Compute the number of DOM which need to be created and insert default value
+		// 0. Compute the number of DOM which need to be created and insert default value
 		uint32 DOMSlotCount = 0;
 		for (const FHairStrandsMacroGroupData& MacroGroup : MacroGroupDatas)
 		{			
@@ -287,17 +278,9 @@ void RenderHairStrandsDeepShadows(
 			return;
 
 		const uint32 AtlasSlotX = FGenericPlatformMath::CeilToInt(FMath::Sqrt(static_cast<float>(DOMSlotCount)));
-		const FIntPoint AtlasSlotDimension(AtlasSlotX, AtlasSlotX == DOMSlotCount ? 1 : AtlasSlotX);
 		const FIntPoint AtlasSlotResolution(GDeepShadowResolution, GDeepShadowResolution);
-		FIntPoint AtlasResolution(AtlasSlotResolution.X * AtlasSlotDimension.X, AtlasSlotResolution.Y * AtlasSlotDimension.Y);
-
-		const bool bDeepShadowGPUDriven = GDeepShadowGPUDriven > 0;
-		if (bDeepShadowGPUDriven)
-		{
-			AtlasResolution = FIntPoint(GDeepShadowResolution, GDeepShadowResolution);
-		}
-
-		
+		const FIntPoint AtlasResolution = FIntPoint(GDeepShadowResolution, GDeepShadowResolution);
+	
 		DeepShadowResources.TotalAtlasSlotCount = 0;
 
 		// Create Atlas resources for DOM. It is shared for all lights, across all views
@@ -307,6 +290,7 @@ void RenderHairStrandsDeepShadows(
 
 		const FVector& TranslatedWorldOffset = View.ViewMatrices.GetPreViewTranslation();
 
+		// 1. Cull light per macro groups
 		// TODO add support for multiple view: need to deduplicate light which are visible accross several views
 		// Allocate atlas CPU slot
 		uint32 TotalAtlasSlotIndex = 0;
@@ -331,22 +315,20 @@ void RenderHairStrandsDeepShadows(
 				const ELightComponentType LightType = (ELightComponentType)LightProxy->GetLightType();
 				const bool bIsDirectional = LightType == ELightComponentType::LightType_Directional;
 				FMinHairRadiusAtDepth1 MinStrandRadiusAtDepth1;
-				const FIntPoint AtlasRectOffset(
-					(TotalAtlasSlotIndex % AtlasSlotDimension.X) * AtlasSlotResolution.X,
-					(TotalAtlasSlotIndex / AtlasSlotDimension.X) * AtlasSlotResolution.Y);
 
 				// Note: LightPosition.W is used in the transmittance mask shader to differentiate between directional and local lights.
 				FHairStrandsDeepShadowData& DomData = MacroGroup.DeepShadowDatas.AddZeroed_GetRef();
-				ComputeTranslatedWorldToLightClip(TranslatedWorldOffset, DomData.CPU_TranslatedWorldToLightTransform, MinStrandRadiusAtDepth1, MacroGroupBounds, *LightProxy, LightType, AtlasSlotResolution);
+				FMatrix CPU_TranslatedWorldToLightTransform;
+				ComputeTranslatedWorldToLightClip(TranslatedWorldOffset, CPU_TranslatedWorldToLightTransform, MinStrandRadiusAtDepth1, MacroGroupBounds, *LightProxy, LightType, AtlasSlotResolution);
 				DomData.LightDirection = (FVector3f)LightProxy->GetDirection();
 				DomData.TranslatedLightPosition = FVector4f(FVector3f((FVector4f)LightProxy->GetPosition() + (FVector3f)TranslatedWorldOffset), bIsDirectional ? 0 : 1);
 				DomData.LightLuminance = LightProxy->GetColor();
 				DomData.LayerDistribution = LightProxy->GetDeepShadowLayerDistribution();
 				DomData.bIsLightDirectional = bIsDirectional;
 				DomData.LightId = LightInfo->Id;
-				DomData.ShadowResolution = bDeepShadowGPUDriven ? AtlasResolution.X : AtlasSlotResolution;
+				DomData.AtlasResolution = AtlasResolution;
 				DomData.Bounds = MacroGroupBounds;
-				DomData.AtlasRect = bDeepShadowGPUDriven ? FIntRect(0, 0, AtlasResolution.X, AtlasResolution.Y) : FIntRect(AtlasRectOffset, AtlasRectOffset + AtlasSlotResolution);
+				DomData.AtlasRect = FIntRect(0, 0, AtlasResolution.X, AtlasResolution.Y);
 				DomData.MacroGroupId = MacroGroup.MacroGroupId;
 				DomData.CPU_MinStrandRadiusAtDepth1 = MinStrandRadiusAtDepth1;
 				DomData.AtlasSlotIndex = TotalAtlasSlotIndex;
@@ -363,7 +345,7 @@ void RenderHairStrandsDeepShadows(
 		FRDGBufferRef DeepShadowViewInfoBuffer = GraphBuilder.CreateBuffer(FRDGBufferDesc::CreateStructuredDesc((16 + 16 + 4 + 3 + 1) * sizeof(float), FMath::Max(1u, TotalAtlasSlotIndex)), TEXT("Hair.DeepShadowViewInfo"));
 		FRDGBufferSRVRef DeepShadowViewInfoBufferSRV = GraphBuilder.CreateSRV(DeepShadowViewInfoBuffer);
 
-		DeepShadowResources.bIsGPUDriven = bDeepShadowGPUDriven;
+		// 2. Allocate slots
 		{
 			// Allocate and create projection matrix and Min radius
 			// Stored FDeepShadowViewInfo structs
@@ -420,14 +402,14 @@ void RenderHairStrandsDeepShadows(
 				FIntVector(1, 1, 1));
 		}
 
-		// Render deep shadows
+		// 3. Render deep shadows
 		for (FHairStrandsMacroGroupData& MacroGroup : MacroGroupDatas)
 		{
 			for (FHairStrandsDeepShadowData& DomData : MacroGroup.DeepShadowDatas)
 			{
 				const bool bIsOrtho = DomData.bIsLightDirectional;
 				const FVector4f HairRenderInfo = PackHairRenderInfo(DomData.CPU_MinStrandRadiusAtDepth1.Primary, DomData.CPU_MinStrandRadiusAtDepth1.Stable, DomData.CPU_MinStrandRadiusAtDepth1.Primary, 1);
-				const uint32 HairRenderInfoBits = PackHairRenderInfoBits(bIsOrtho, DeepShadowResources.bIsGPUDriven);
+				const uint32 HairRenderInfoBits = PackHairRenderInfoBits(bIsOrtho, true /*bIsGPUDriven*/);
 					
 				const bool bDeepShadow = GDeepShadowInjectVoxelDepth == 0;
 				// Inject voxel result into the deep shadow
@@ -442,7 +424,6 @@ void RenderHairStrandsDeepShadows(
 						View,
 						bClear,
 						DomData,
-						DomData.CPU_TranslatedWorldToLightTransform,
 						DomData.AtlasRect,
 						DomData.AtlasSlotIndex,
 						AtlasSlotResolution,
@@ -470,10 +451,6 @@ void RenderHairStrandsDeepShadows(
 
 					{
 						FHairDeepShadowRasterUniformParameters* UniformParameters = GraphBuilder.AllocParameters<FHairDeepShadowRasterUniformParameters>();
-
-						UniformParameters->CPU_TranslatedWorldToClipMatrix = FMatrix44f(DomData.CPU_TranslatedWorldToLightTransform);	// LWC_TODO: Precision loss
-						UniformParameters->SliceValue = FVector4f(1, 1, 1, 1);
-						UniformParameters->AtlasRect = DomData.AtlasRect;
 						UniformParameters->AtlasSlotIndex = DomData.AtlasSlotIndex;
 						UniformParameters->LayerDepths = LayerDepths;
 						UniformParameters->FrontDepthTexture = SystemTextures.DepthDummy;
@@ -509,10 +486,6 @@ void RenderHairStrandsDeepShadows(
 
 					{
 						FHairDeepShadowRasterUniformParameters* UniformParameters = GraphBuilder.AllocParameters<FHairDeepShadowRasterUniformParameters>();
-
-						UniformParameters->CPU_TranslatedWorldToClipMatrix = FMatrix44f(DomData.CPU_TranslatedWorldToLightTransform);	// LWC_TODO: Precision loss
-						UniformParameters->SliceValue = FVector4f(1, 1, 1, 1);
-						UniformParameters->AtlasRect = DomData.AtlasRect;
 						UniformParameters->AtlasSlotIndex = DomData.AtlasSlotIndex;
 						UniformParameters->LayerDepths = LayerDepths;
 						UniformParameters->FrontDepthTexture = FrontDepthAtlasTexture;
diff --git a/Engine/Source/Runtime/Renderer/Private/HairStrands/HairStrandsRasterCommon.h b/Engine/Source/Runtime/Renderer/Private/HairStrands/HairStrandsRasterCommon.h
--- a/Engine/Source/Runtime/Renderer/Private/HairStrands/HairStrandsRasterCommon.h
+++ b/Engine/Source/Runtime/Renderer/Private/HairStrands/HairStrandsRasterCommon.h
@@ -24,9 +24,6 @@ enum class EHairStrandsRasterPassType : uint8
 // Deep shadow raster pass
 
 BEGIN_GLOBAL_SHADER_PARAMETER_STRUCT(FHairDeepShadowRasterUniformParameters, RENDERER_API)
-	SHADER_PARAMETER(FMatrix44f, CPU_TranslatedWorldToClipMatrix)
-	SHADER_PARAMETER(FVector4f, SliceValue)
-	SHADER_PARAMETER(FIntRect, AtlasRect)
 	SHADER_PARAMETER(uint32, AtlasSlotIndex)
 	SHADER_PARAMETER(FVector4f, LayerDepths)
 	SHADER_PARAMETER_RDG_TEXTURE(Texture2D<float>, FrontDepthTexture)
diff --git a/Engine/Source/Runtime/Renderer/Private/HairStrands/HairStrandsTransmittance.cpp b/Engine/Source/Runtime/Renderer/Private/HairStrands/HairStrandsTransmittance.cpp
--- a/Engine/Source/Runtime/Renderer/Private/HairStrands/HairStrandsTransmittance.cpp
+++ b/Engine/Source/Runtime/Renderer/Private/HairStrands/HairStrandsTransmittance.cpp
@@ -320,9 +320,7 @@ class FHairStrandsDeepShadowTransmittanceMaskCS : public FGlobalShader
 		SHADER_PARAMETER_RDG_UNIFORM_BUFFER(FHairStrandsViewUniformParameters, HairStrands)
 		SHADER_PARAMETER_STRUCT_REF(FViewUniformShaderParameters, ViewUniformBuffer)
 
-		SHADER_PARAMETER_ARRAY(FIntVector4, DeepShadow_AtlasSlotOffsets_AtlasSlotIndex, [FHairStrandsDeepShadowData::MaxMacroGroupCount])
-		SHADER_PARAMETER_ARRAY(FMatrix44f, DeepShadow_CPUTranslatedWorldToLightTransforms, [FHairStrandsDeepShadowData::MaxMacroGroupCount])
-		SHADER_PARAMETER(FIntPoint, DeepShadow_Resolution)
+		SHADER_PARAMETER(FIntPoint, DeepShadow_AtlasResolution)
 		SHADER_PARAMETER(float, LightRadius)
 		SHADER_PARAMETER(FVector4f, TranslatedLightPosition_LightDirection)
 		SHADER_PARAMETER(uint32, LightChannelMask)
@@ -334,12 +332,12 @@ class FHairStrandsDeepShadowTransmittanceMaskCS : public FGlobalShader
 		SHADER_PARAMETER(uint32, DeepShadow_KernelType)
 		SHADER_PARAMETER(uint32, DeepShadow_DebugMode)
 		SHADER_PARAMETER(FMatrix44f, DeepShadow_ShadowToWorld)
-		SHADER_PARAMETER(uint32, DeepShadow_bIsGPUDriven)
 
 		SHADER_PARAMETER_RDG_TEXTURE(Texture2D, RayMarchMaskTexture)
 		SHADER_PARAMETER_RDG_TEXTURE(Texture2D, DeepShadow_FrontDepthTexture)
 		SHADER_PARAMETER_RDG_TEXTURE(Texture2D, DeepShadow_DomTexture)
 		SHADER_PARAMETER_RDG_BUFFER_SRV(StructuredBuffer, DeepShadow_ViewInfoBuffer)
+		SHADER_PARAMETER_RDG_BUFFER_SRV(StructuredBuffer, DeepShadow_AtlasSlotIndexBuffer)
 		SHADER_PARAMETER_RDG_BUFFER_UAV(RWBuffer, OutTransmittanceMask)
 
 		RDG_BUFFER_ACCESS(IndirectArgsBuffer, ERHIAccess::IndirectArgs)
@@ -363,11 +361,10 @@ IMPLEMENT_GLOBAL_SHADER(FHairStrandsDeepShadowTransmittanceMaskCS, "/Engine/Priv
 
 struct FHairStrandsDeepShadowTransmittanceLightParams : FHairStrandsTransmittanceLightParams
 {
-	FIntVector4 DeepShadow_AtlasSlotOffsets_AtlasSlotIndex[FHairStrandsDeepShadowData::MaxMacroGroupCount];
-	FMatrix DeepShadow_CPUTranslatedWorldToLightTransforms[FHairStrandsDeepShadowData::MaxMacroGroupCount];
+	//FIntVector4 DeepShadow_AtlasSlotOffsets_AtlasSlotIndex[FHairStrandsDeepShadowData::MaxMacroGroupCount];
 	FRDGBufferSRVRef DeepShadow_ViewInfoBuffer = nullptr;
-	FIntPoint DeepShadow_Resolution = FIntPoint(0, 0);
-	bool DeepShadow_bIsGPUDriven = false;
+	FRDGBufferSRVRef DeepShadow_AtlasSlotIndexBuffer = nullptr;
+	FIntPoint DeepShadow_AtlasResolution = FIntPoint(0, 0);
 	FVector4f DeepShadow_LayerDepths = FVector4f(0, 0, 0, 0);
 	float DeepShadow_DepthBiasScale = 0;
 	float DeepShadow_DensityScale = 0;
@@ -398,7 +395,7 @@ static FRDGBufferRef AddHairStrandsDeepShadowTransmittanceMaskPass(
 	Parameters->DeepShadow_DomTexture = Params.DeepShadow_DomTexture;
 	Parameters->OutTransmittanceMask = GraphBuilder.CreateUAV(OutBuffer, FHairStrandsTransmittanceMaskData::Format);
 	Parameters->LightChannelMask = Params.LightChannelMask;
-	Parameters->DeepShadow_Resolution = Params.DeepShadow_Resolution;
+	Parameters->DeepShadow_AtlasResolution = Params.DeepShadow_AtlasResolution;
 	Parameters->TranslatedLightPosition_LightDirection = Params.TranslatedLightPosition_LightDirection;
 	Parameters->LightRadius = Params.LightRadius;
 	Parameters->DeepShadow_DepthBiasScale = Params.DeepShadow_DepthBiasScale;
@@ -410,17 +407,11 @@ static FRDGBufferRef AddHairStrandsDeepShadowTransmittanceMaskPass(
 	Parameters->DeepShadow_ShadowToWorld = FMatrix44f(Params.DeepShadow_ShadowToWorld);		// LWC_TODO: Precision loss
 	Parameters->IndirectArgsBuffer = IndirectArgsBuffer;
 	Parameters->HairStrands = HairStrands::BindHairStrandsViewUniformParameters(View);
-	Parameters->DeepShadow_bIsGPUDriven = Params.DeepShadow_bIsGPUDriven ? 1 : 0;;
 	Parameters->DeepShadow_ViewInfoBuffer = Params.DeepShadow_ViewInfoBuffer;
+	Parameters->DeepShadow_AtlasSlotIndexBuffer = Params.DeepShadow_AtlasSlotIndexBuffer;
 	Parameters->RayMarchMaskTexture = ScreenShadowMaskSubPixelTexture ? ScreenShadowMaskSubPixelTexture : GraphBuilder.RegisterExternalTexture(GSystemTextures.WhiteDummy);
 	Parameters->ShadowChannelMask = FVector4f(0, 0, 0, 0);
 	Parameters->ShadowChannelMask[FMath::Clamp<uint32>(Params.ShadowChannelMask, 0, 3)] = 1.0f;
-
-	for (uint32 SlotIndex=0;SlotIndex< FHairStrandsDeepShadowData::MaxMacroGroupCount;++SlotIndex)
-	{
-		Parameters->DeepShadow_AtlasSlotOffsets_AtlasSlotIndex[SlotIndex] = Params.DeepShadow_AtlasSlotOffsets_AtlasSlotIndex[SlotIndex];
-		Parameters->DeepShadow_CPUTranslatedWorldToLightTransforms[SlotIndex] = FMatrix44f(Params.DeepShadow_CPUTranslatedWorldToLightTransforms[SlotIndex]);
-	}
 	
 	check(NodeGroupSize == 64 || NodeGroupSize == 32);
 	FHairStrandsDeepShadowTransmittanceMaskCS::FPermutationDomain PermutationVector;
@@ -612,10 +603,8 @@ class FHairStrandsDeepShadowMaskPS : public FGlobalShader
 		SHADER_PARAMETER(FIntPoint, DeepShadow_SlotOffset)
 		SHADER_PARAMETER(uint32, DeepShadow_SlotIndex)
 		SHADER_PARAMETER(FIntPoint, DeepShadow_SlotResolution)
-		SHADER_PARAMETER(FMatrix44f, DeepShadow_CPUTranslatedWorldToLightTransform)
 		SHADER_PARAMETER(float, DeepShadow_DepthBiasScale)
 		SHADER_PARAMETER(float, DeepShadow_DensityScale)
-		SHADER_PARAMETER(uint32, DeepShadow_bIsGPUDriven)
 		SHADER_PARAMETER(FVector4f, DeepShadow_LayerDepths)
 		SHADER_PARAMETER(float, FadeAlpha)
 		SHADER_PARAMETER(uint32, EncodingType)
@@ -651,11 +640,9 @@ struct FHairStrandsDeepShadowParams
 {
 	uint32			OutputChannel = ~0;
 	FRDGBufferSRVRef DeepShadow_ViewInfoBuffer = nullptr;
-	FMatrix			DeepShadow_CPUTranslatedWorldToLightTransform;
 	FIntRect		DeepShadow_AtlasRect;
 	FRDGTextureRef	DeepShadow_FrontDepthTexture = nullptr;
 	FRDGTextureRef	DeepShadow_LayerTexture = nullptr;
-	bool			DeepShadow_bIsGPUDriven = false;
 	float			DeepShadow_DepthBiasScale = 1;
 	float			DeepShadow_DensityScale = 1;
 	uint32			DeepShadow_AtlasSlotIndex = 0;
@@ -679,7 +666,6 @@ static void AddHairStrandsDeepShadowMaskPass(
 	Parameters->EncodingType = EncodingType;
 	Parameters->ViewUniformBuffer = View.ViewUniformBuffer;
 	Parameters->SceneDepthTexture = SceneDepthTexture;
-	Parameters->DeepShadow_CPUTranslatedWorldToLightTransform = FMatrix44f(Params.DeepShadow_CPUTranslatedWorldToLightTransform);
 	Parameters->DeepShadow_FrontDepthTexture = Params.DeepShadow_FrontDepthTexture;
 	Parameters->DeepShadow_DomTexture = Params.DeepShadow_LayerTexture;
 	Parameters->LinearSampler = TStaticSamplerState<SF_Bilinear, AM_Clamp, AM_Clamp, AM_Clamp>::GetRHI();
@@ -693,7 +679,6 @@ static void AddHairStrandsDeepShadowMaskPass(
 	Parameters->DeepShadow_SlotOffset = FIntPoint(Params.DeepShadow_AtlasRect.Min.X, Params.DeepShadow_AtlasRect.Min.Y);
 	Parameters->DeepShadow_SlotResolution = FIntPoint(Params.DeepShadow_AtlasRect.Max.X - Params.DeepShadow_AtlasRect.Min.X, Params.DeepShadow_AtlasRect.Max.Y - Params.DeepShadow_AtlasRect.Min.Y);
 	Parameters->DeepShadow_ViewInfoBuffer = Params.DeepShadow_ViewInfoBuffer;
-	Parameters->DeepShadow_bIsGPUDriven = Params.DeepShadow_bIsGPUDriven ? 1 : 0;;
 	Parameters->HairStrands = HairStrands::BindHairStrandsViewUniformParameters(View);
 	if (ShaderPrint::IsValid(View.ShaderPrintData))
 	{
@@ -798,42 +783,40 @@ static FHairStrandsTransmittanceMaskData InternalRenderHairStrandsTransmittanceM
 		FHairStrandsDeepShadowTransmittanceLightParams Params;
 		Params.DeepShadow_DensityScale = GetDeepShadowDensityScale();
 		Params.DeepShadow_DepthBiasScale = GetDeepShadowDepthBiasScale();
-		memset(Params.DeepShadow_AtlasSlotOffsets_AtlasSlotIndex, 0, sizeof(Params.DeepShadow_AtlasSlotOffsets_AtlasSlotIndex));
-		memset(Params.DeepShadow_CPUTranslatedWorldToLightTransforms, 0, sizeof(Params.DeepShadow_CPUTranslatedWorldToLightTransforms));
 
-		FRDGBufferSRVRef DeepShadow_ViewInfoBufferSRV = nullptr;
+		TArray<uint32> AtlasSlotIndexBuffer;
+		AtlasSlotIndexBuffer.Reserve(MacroGroupDatas.Num());
 		for (const FHairStrandsMacroGroupData& MacroGroupData : MacroGroupDatas)
 		{
+			uint32 AtlasSlotIndex = 0;
 			for (const FHairStrandsDeepShadowData& DeepShadowData : MacroGroupData.DeepShadowDatas)
 			{
 				if (DeepShadowData.LightId == LightSceneInfo->Id)
 				{
-					if (DeepShadow_ViewInfoBufferSRV == nullptr)
-					{
-						DeepShadow_ViewInfoBufferSRV = GraphBuilder.CreateSRV(DeepShadowResources.DeepShadowViewInfoBuffer);
-					}
-
 					bHasFoundLight = true;
 					Params.TranslatedLightPosition_LightDirection = GetLightTranslatedWorldPositionAndDirection(View, LightSceneInfo);
 					Params.DeepShadow_FrontDepthTexture = DeepShadowResources.DepthAtlasTexture;
 					Params.DeepShadow_DomTexture = DeepShadowResources.LayersAtlasTexture;
-					Params.DeepShadow_Resolution = DeepShadowData.ShadowResolution;
+					Params.DeepShadow_AtlasResolution = DeepShadowData.AtlasResolution;
 					Params.LightRadius = 0;
 					Params.LightChannelMask = LightSceneInfo->Proxy->GetLightingChannelMask();
 					Params.ShadowChannelMask = bProjectingForForwardShading ? LightSceneInfo->GetDynamicShadowMapChannel() : 0;
 					Params.DeepShadow_LayerDepths = ComputeDeepShadowLayerDepths(DeepShadowData.LayerDistribution);
-					Params.DeepShadow_AtlasSlotOffsets_AtlasSlotIndex[DeepShadowData.MacroGroupId] = FIntVector4(DeepShadowData.AtlasRect.Min.X, DeepShadowData.AtlasRect.Min.Y, DeepShadowData.AtlasSlotIndex, 0);
-					Params.DeepShadow_CPUTranslatedWorldToLightTransforms[DeepShadowData.MacroGroupId] = DeepShadowData.CPU_TranslatedWorldToLightTransform;
-					Params.DeepShadow_ViewInfoBuffer = DeepShadow_ViewInfoBufferSRV;
-					Params.DeepShadow_bIsGPUDriven = DeepShadowResources.bIsGPUDriven;
+					AtlasSlotIndex = DeepShadowData.AtlasSlotIndex;
 				}
 			}
+			AtlasSlotIndexBuffer.Add(AtlasSlotIndex);
 		}
 
 		if (bHasFoundLight)
 		{
 			check(Params.DeepShadow_FrontDepthTexture);
 			check(Params.DeepShadow_DomTexture);
+
+			FRDGBufferRef DeepShadowAtlasSlotIndexBuffer = CreateStructuredBuffer(GraphBuilder, TEXT("Hair.DeepShadow.AtlasSlotIndexBuffer"), 4u, AtlasSlotIndexBuffer.Num(), AtlasSlotIndexBuffer.GetData(), 4u * AtlasSlotIndexBuffer.Num());			
+			Params.DeepShadow_ViewInfoBuffer = GraphBuilder.CreateSRV(DeepShadowResources.DeepShadowViewInfoBuffer);
+			Params.DeepShadow_AtlasSlotIndexBuffer = GraphBuilder.CreateSRV(DeepShadowAtlasSlotIndexBuffer);
+
 			Out.TransmittanceMask = AddHairStrandsDeepShadowTransmittanceMaskPass(
 				GraphBuilder,
 				SceneTextures,
@@ -978,9 +961,8 @@ static void InternalRenderHairStrandsShadowMask(
 	bool bHasDeepShadow = false;
 	if (!IsHairStrandsForVoxelTransmittanceAndShadowEnable())
 	{
-		FRDGBufferRef DeepShadow_ViewInfoBuffer = nullptr;
+		// TODO change this into a single pass, instead of iterating onto all macrogroup, as we do for transmittance mask
 		FRDGBufferSRVRef DeepShadow_ViewInfoBufferSRV = nullptr;
-
 		for (const FHairStrandsMacroGroupData& MacroGroupData : InMacroGroupDatas)
 		{
 			for (const FHairStrandsDeepShadowData& DomData : MacroGroupData.DeepShadowDatas)
@@ -988,10 +970,9 @@ static void InternalRenderHairStrandsShadowMask(
 				if (DomData.LightId != LightSceneInfo->Id)
 					continue;
 
-				if (DeepShadow_ViewInfoBuffer == nullptr)
+				if (DeepShadow_ViewInfoBufferSRV == nullptr)
 				{
-					DeepShadow_ViewInfoBuffer = DeepShadowResources.DeepShadowViewInfoBuffer;
-					DeepShadow_ViewInfoBufferSRV = GraphBuilder.CreateSRV(DeepShadow_ViewInfoBuffer);
+					DeepShadow_ViewInfoBufferSRV = GraphBuilder.CreateSRV(DeepShadowResources.DeepShadowViewInfoBuffer);
 				}
 
 				bHasDeepShadow = true;
@@ -999,8 +980,6 @@ static void InternalRenderHairStrandsShadowMask(
 				FHairStrandsDeepShadowParams Params;
 				Params.DeepShadow_AtlasSlotIndex = DomData.AtlasSlotIndex;
 				Params.DeepShadow_ViewInfoBuffer = DeepShadow_ViewInfoBufferSRV;
-				Params.DeepShadow_bIsGPUDriven = DeepShadowResources.bIsGPUDriven ? 1 : 0;
-				Params.DeepShadow_CPUTranslatedWorldToLightTransform = DomData.CPU_TranslatedWorldToLightTransform;
 				Params.DeepShadow_AtlasRect = DomData.AtlasRect;
 				Params.DeepShadow_FrontDepthTexture = DeepShadowResources.DepthAtlasTexture;
 				Params.DeepShadow_LayerTexture = DeepShadowResources.LayersAtlasTexture;
