Fix incorrect PSO precaching on Substrate classification shader with WaveOps permutation.

The PSO was precached while the wave size requirement was not meet by the HW. This was causing the PSO caching to fail with an error.

diff --git a/Engine/Source/Runtime/Renderer/Private/Substrate/Substrate.cpp b/Engine/Source/Runtime/Renderer/Private/Substrate/Substrate.cpp
--- a/Engine/Source/Runtime/Renderer/Private/Substrate/Substrate.cpp
+++ b/Engine/Source/Runtime/Renderer/Private/Substrate/Substrate.cpp
@@ -746,6 +746,11 @@ TRDGUniformBufferRef<FSubstratePublicGlobalUniformParameters> CreatePublicGlobal
 
 ////////////////////////////////////////////////////////////////////////////////////////////////////////////////////////////////////////////////
 
+static bool DoesRuntimeSupportWave64()
+{
+	return GRHISupportsWaveOperations && (GRHIMinimumWaveSize <= 64 && GRHIMaximumWaveSize >= 64);
+}
+
 class FSubstrateClosureTilePassCS : public FGlobalShader
 {
 	DECLARE_GLOBAL_SHADER(FSubstrateClosureTilePassCS);
@@ -785,6 +790,20 @@ class FSubstrateClosureTilePassCS : public FGlobalShader
 		return GetMaxSupportedFeatureLevel(Parameters.Platform) >= ERHIFeatureLevel::SM5 && Substrate::IsSubstrateEnabled();
 	}
 
+	static EShaderPermutationPrecacheRequest ShouldPrecachePermutation(const FGlobalShaderPermutationParameters& Parameters) 
+	{ 
+		bool bUsed = ShouldCompilePermutation(Parameters);
+		if (bUsed)
+		{
+			FPermutationDomain PermutationVector(Parameters.PermutationId);
+			if (PermutationVector.Get<FWaveOps>() && !DoesRuntimeSupportWave64())
+			{
+				bUsed = false;
+			}
+		}
+		return bUsed ? EShaderPermutationPrecacheRequest::Precached : EShaderPermutationPrecacheRequest::NotUsed;
+	}
+
 	static void ModifyCompilationEnvironment(const FGlobalShaderPermutationParameters& Parameters, FShaderCompilerEnvironment& OutEnvironment)
 	{
 		FGlobalShader::ModifyCompilationEnvironment(Parameters, OutEnvironment);
@@ -844,6 +863,20 @@ class FSubstrateMaterialTileClassificationPassCS : public FGlobalShader
 		return GetMaxSupportedFeatureLevel(Parameters.Platform) >= ERHIFeatureLevel::SM5 && Substrate::IsSubstrateEnabled();
 	}
 
+	static EShaderPermutationPrecacheRequest ShouldPrecachePermutation(const FGlobalShaderPermutationParameters& Parameters) 
+	{ 
+		bool bUsed = ShouldCompilePermutation(Parameters);
+		if (bUsed)
+		{
+			FPermutationDomain PermutationVector(Parameters.PermutationId);
+			if (PermutationVector.Get<FWaveOps>() && !DoesRuntimeSupportWave64())
+			{
+				bUsed = false;
+			}
+		}
+		return bUsed ? EShaderPermutationPrecacheRequest::Precached : EShaderPermutationPrecacheRequest::NotUsed;
+	}
+
 	static void ModifyCompilationEnvironment(const FGlobalShaderPermutationParameters& Parameters, FShaderCompilerEnvironment& OutEnvironment)
 	{
 		FGlobalShader::ModifyCompilationEnvironment(Parameters, OutEnvironment);
@@ -1345,7 +1378,8 @@ void AddSubstrateMaterialClassificationPass(FRDGBuilder& GraphBuilder, const FMi
 			continue;
 		}
 
-		const bool bWaveOps = GRHISupportsWaveOperations&& GRHIMaximumWaveSize >= 64 && SubstrateSupportsWaveOps(Platform) != ERHIFeatureSupport::Unsupported;
+		// Our current classification require 64 waves
+		const bool bWaveOps = DoesRuntimeSupportWave64() && SubstrateSupportsWaveOps(Platform) != ERHIFeatureSupport::Unsupported;
 		
 		const FSubstrateViewData* SubstrateViewData = &View.SubstrateViewData;
 		const FSubstrateSceneData* SubstrateSceneData = View.SubstrateViewData.SceneData;
