Mobile: Fixed a several issues with a PSO pre-cache misses in a projects with a static lighting enabled.
Added PSO pre-caching for TextRenderComponent for platfroms that do not support MVF
Enabled PSO pre-caching by default on Vulkan.

diff --git a/Engine/Source/Runtime/Engine/Private/Components/PrimitiveComponent.cpp b/Engine/Source/Runtime/Engine/Private/Components/PrimitiveComponent.cpp
--- a/Engine/Source/Runtime/Engine/Private/Components/PrimitiveComponent.cpp
+++ b/Engine/Source/Runtime/Engine/Private/Components/PrimitiveComponent.cpp
@@ -4741,6 +4741,7 @@ void UPrimitiveComponent::SetupPrecachePSOParams(FPSOPrecacheParams& Params)
 	Params.bRenderInMainPass = bRenderInMainPass;
 	Params.bRenderInDepthPass = bRenderInDepthPass;
 	Params.bStaticLighting = HasStaticLighting();
+	Params.bUsesIndirectLightingCache = Params.bStaticLighting && IndirectLightingCacheQuality != ILCQ_Off && (!IsPrecomputedLightingValid() || GetLightmapType() == ELightmapType::ForceVolumetric);
 	Params.bAffectDynamicIndirectLighting = bAffectDynamicIndirectLighting;
 	Params.bCastShadow = CastShadow;
 	// Custom depth can be toggled at runtime with PSO precache call so assume it might be needed when depth pass is needed
diff --git a/Engine/Source/Runtime/Engine/Private/Components/TextRenderComponent.cpp b/Engine/Source/Runtime/Engine/Private/Components/TextRenderComponent.cpp
--- a/Engine/Source/Runtime/Engine/Private/Components/TextRenderComponent.cpp
+++ b/Engine/Source/Runtime/Engine/Private/Components/TextRenderComponent.cpp
@@ -1534,11 +1534,33 @@ void UTextRenderComponent::PostLoad()
 	Super::PostLoad();
 }
 
+FVertexDeclarationElementList InitDummyVertexDeclarationElementsForText()
+{
+	FStaticMeshVertexBuffers VertexBuffers;
+	VertexBuffers.PositionVertexBuffer.Init(1);
+	VertexBuffers.StaticMeshVertexBuffer.Init(1, 1);
+	VertexBuffers.ColorVertexBuffer.Init(1);
+	
+	VertexBuffers.PositionVertexBuffer.VertexPosition(0) = FVector3f(0, 0, 0);
+	VertexBuffers.StaticMeshVertexBuffer.SetVertexTangents(0, FVector3f(1, 0, 0), FVector3f(0, 1, 0), FVector3f(0, 0, 1));
+	VertexBuffers.StaticMeshVertexBuffer.SetVertexUV(0, 0, FVector2f(0, 0));
+	VertexBuffers.ColorVertexBuffer.VertexColor(0) = FColor(1,1,1,1);
+
+	FLocalVertexFactory::FDataType Data;
+	VertexBuffers.PositionVertexBuffer.BindPositionVertexBuffer(nullptr, Data);
+	VertexBuffers.StaticMeshVertexBuffer.BindTangentVertexBuffer(nullptr, Data);
+	VertexBuffers.StaticMeshVertexBuffer.BindPackedTexCoordVertexBuffer(nullptr, Data);
+	VertexBuffers.StaticMeshVertexBuffer.BindLightMapVertexBuffer(nullptr, Data, 0);
+	VertexBuffers.ColorVertexBuffer.BindColorVertexBuffer(nullptr, Data);
+	
+	FVertexDeclarationElementList Elements;
+	FLocalVertexFactory::GetVertexElements(GMaxRHIFeatureLevel, EVertexInputStreamType::Default, false, Data, Elements);
+	return Elements;
+}
+
 void UTextRenderComponent::PrecachePSOs()
 {
-	if (IsComponentPSOPrecachingEnabled() && TextMaterial
-		// FIXME: need to collect an actual vertex declaration for non-MVF path
-		&& RHISupportsManualVertexFetch(GMaxRHIShaderPlatform))
+	if (IsComponentPSOPrecachingEnabled() && TextMaterial)
 	{
 		FPSOPrecacheParams PrecachePSOParams;
 		SetupPrecachePSOParams(PrecachePSOParams);
@@ -1547,7 +1569,20 @@ void UTextRenderComponent::PrecachePSOs()
 		// and leaves the default CastShadow value which is true
 		PrecachePSOParams.bCastShadow = true;
 
-		TextMaterial->PrecachePSOs(&FLocalVertexFactory::StaticType, PrecachePSOParams);
+		FPSOPrecacheVertexFactoryDataList VertexFactoryDataList;
+
+		if (RHISupportsManualVertexFetch(GMaxRHIShaderPlatform))
+		{
+			VertexFactoryDataList.Add(FPSOPrecacheVertexFactoryData(&FLocalVertexFactory::StaticType));
+		}
+		else
+		{
+			const static FVertexDeclarationElementList Elements = InitDummyVertexDeclarationElementsForText();
+			VertexFactoryDataList.Add(FPSOPrecacheVertexFactoryData(&FLocalVertexFactory::StaticType, Elements));
+		}
+
+		TArray<FMaterialPSOPrecacheRequestID> MaterialPrecacheRequestIDs;
+		TextMaterial->PrecachePSOs(VertexFactoryDataList, PrecachePSOParams, EPSOPrecachePriority::Medium, MaterialPrecacheRequestIDs);
 	}
 }
 
diff --git a/Engine/Source/Runtime/Engine/Public/PSOPrecache.h b/Engine/Source/Runtime/Engine/Public/PSOPrecache.h
--- a/Engine/Source/Runtime/Engine/Public/PSOPrecache.h
+++ b/Engine/Source/Runtime/Engine/Public/PSOPrecache.h
@@ -107,6 +107,7 @@ struct FPSOPrecacheParams
 			uint64 bCastShadow : 1;
 			uint64 bRenderCustomDepth : 1;
 			
+			uint64 bUsesIndirectLightingCache : 1;
 			uint64 bAffectDynamicIndirectLighting : 1;
 			uint64 bReverseCulling : 1;
 			uint64 bDisableBackFaceCulling : 1;
@@ -119,7 +120,7 @@ struct FPSOPrecacheParams
 
 			uint64 BasePassPixelFormat : 16;
 
-			uint64 Unused : 21;
+			uint64 Unused : 20;
 		};
 		uint64 Data;
 	};
diff --git a/Engine/Source/Runtime/OpenGLDrv/Private/OpenGLDevice.cpp b/Engine/Source/Runtime/OpenGLDrv/Private/OpenGLDevice.cpp
--- a/Engine/Source/Runtime/OpenGLDrv/Private/OpenGLDevice.cpp
+++ b/Engine/Source/Runtime/OpenGLDrv/Private/OpenGLDevice.cpp
@@ -1273,12 +1273,7 @@ static void InitRHICapabilitiesForGL()
 	// @TODO revisit this with newer drivers
 	GRHINeedsUnatlasedCSMDepthsWorkaround = true;
 
-	static const auto CVarPSOPrecaching = IConsoleManager::Get().FindConsoleVariable(TEXT("r.PSOPrecaching"));
-	if (CVarPSOPrecaching && CVarPSOPrecaching->GetInt() != 0 && CVarAllowPSOPrecaching.GetValueOnAnyThread())
-	{
-		GRHISupportsPSOPrecaching = true;
-	}
-	
+	GRHISupportsPSOPrecaching = CVarAllowPSOPrecaching.GetValueOnAnyThread();
 	GRHISupportsPipelineFileCache = !GRHISupportsPSOPrecaching || CVarEnablePSOFileCacheWhenPrecachingActive.GetValueOnAnyThread();
 
 	GRHIGlobals.NeedsShaderUnbinds = true;
diff --git a/Engine/Source/Runtime/Renderer/Private/LightMapRendering.cpp b/Engine/Source/Runtime/Renderer/Private/LightMapRendering.cpp
--- a/Engine/Source/Runtime/Renderer/Private/LightMapRendering.cpp
+++ b/Engine/Source/Runtime/Renderer/Private/LightMapRendering.cpp
@@ -177,7 +177,10 @@ bool FMobileDistanceFieldShadowsAndLQLightMapPolicy::ShouldCompilePermutation(co
 {
 	static auto* CVarMobileAllowDistanceFieldShadows = IConsoleManager::Get().FindTConsoleVariableDataInt(TEXT("r.Mobile.AllowDistanceFieldShadows"));
 	const bool bMobileAllowDistanceFieldShadows = CVarMobileAllowDistanceFieldShadows->GetValueOnAnyThread() == 1;
-	return bMobileAllowDistanceFieldShadows && Super::ShouldCompilePermutation(Parameters);
+	
+	return bMobileAllowDistanceFieldShadows && 
+		!IsTranslucentBlendMode(Parameters.MaterialParameters) &&
+		Super::ShouldCompilePermutation(Parameters);
 }
 
 void FMobileDistanceFieldShadowsAndLQLightMapPolicy::ModifyCompilationEnvironment(const FMaterialShaderPermutationParameters& Parameters, FShaderCompilerEnvironment& OutEnvironment)
@@ -192,10 +195,7 @@ bool FMobileDistanceFieldShadowsLightMapAndCSMLightingPolicy::ShouldCompilePermu
 		return false;
 	}
 
-	return FReadOnlyCVARCache::MobileEnableStaticAndCSMShadowReceivers() &&
-		Parameters.MaterialParameters.ShadingModels.IsLit() &&
-		!IsTranslucentBlendMode(Parameters.MaterialParameters) &&
-		Super::ShouldCompilePermutation(Parameters);
+	return FReadOnlyCVARCache::MobileEnableStaticAndCSMShadowReceivers() &&	Super::ShouldCompilePermutation(Parameters);
 }
 
 void FMobileDistanceFieldShadowsLightMapAndCSMLightingPolicy::ModifyCompilationEnvironment(const FMaterialShaderPermutationParameters& Parameters, FShaderCompilerEnvironment& OutEnvironment)
diff --git a/Engine/Source/Runtime/Renderer/Private/MobileBasePass.cpp b/Engine/Source/Runtime/Renderer/Private/MobileBasePass.cpp
--- a/Engine/Source/Runtime/Renderer/Private/MobileBasePass.cpp
+++ b/Engine/Source/Runtime/Renderer/Private/MobileBasePass.cpp
@@ -408,7 +408,7 @@ ELightMapPolicyType MobileBasePass::SelectMeshLightmapPolicy(
 
 			if (LightMapInteraction.GetType() == LMIT_Texture && FReadOnlyCVARCache::EnableLowQualityLightmaps())
 			{
-				const FShadowMapInteraction ShadowMapInteraction = (Mesh.LCI != nullptr)
+				const FShadowMapInteraction ShadowMapInteraction = (Mesh.LCI != nullptr && !bIsTranslucent)
 					? Mesh.LCI->GetShadowMapInteraction(FeatureLevel)
 					: FShadowMapInteraction();
 
@@ -486,14 +486,16 @@ static FMobileLightMapPolicyTypeList GetUniformLightMapPolicyTypeForPSOCollectio
 	
 	if (bLitMaterial)
 	{
-		if (!IsStaticLightingAllowed())
+		const bool bAllowStaticLighting = IsStaticLightingAllowed();
+		
+		if (!bAllowStaticLighting)
 		{
-			if (bUsesDeferredShading || bTranslucent || !MobileUseCSMShaderBranch())
+			if (bUsesDeferredShading || !MobileUseCSMShaderBranch())
 			{
 				Result.Add(LMP_NO_LIGHTMAP);
 			}
 						
-			if (!bTranslucent && !bUsesDeferredShading)
+			if (!bUsesDeferredShading)
 			{
 				// permutation that can receive CSM
 				Result.Add(LMP_MOBILE_DIRECTIONAL_LIGHT_CSM);
@@ -501,11 +503,11 @@ static FMobileLightMapPolicyTypeList GetUniformLightMapPolicyTypeForPSOCollectio
 		}
 		else
 		{
-			if (FReadOnlyCVARCache::EnableLowQualityLightmaps())
+			if (!bMovable && FReadOnlyCVARCache::EnableLowQualityLightmaps())
 			{
-				if (FReadOnlyCVARCache::MobileEnableStaticAndCSMShadowReceivers() && !bUsesDeferredShading)
+				if (FReadOnlyCVARCache::MobileEnableStaticAndCSMShadowReceivers() && !bUsesDeferredShading && bCanReceiveCSM)
 				{
-					if (FReadOnlyCVARCache::MobileAllowDistanceFieldShadows())
+					if (FReadOnlyCVARCache::MobileAllowDistanceFieldShadows() && !bTranslucent)
 					{
 						Result.Add(LMP_MOBILE_DISTANCE_FIELD_SHADOWS_LIGHTMAP_AND_CSM);
 					}
@@ -518,7 +520,7 @@ static FMobileLightMapPolicyTypeList GetUniformLightMapPolicyTypeForPSOCollectio
 					Result.Add(LMP_MOBILE_DIRECTIONAL_LIGHT_CSM_AND_LIGHTMAP);
 				}
 
-				if (FReadOnlyCVARCache::MobileAllowDistanceFieldShadows())
+				if (FReadOnlyCVARCache::MobileAllowDistanceFieldShadows() && !bCanReceiveCSM && !bTranslucent)
 				{
 					Result.Add(LMP_MOBILE_DISTANCE_FIELD_SHADOWS_AND_LQ_LIGHTMAP);
 				}
@@ -534,16 +536,16 @@ static FMobileLightMapPolicyTypeList GetUniformLightMapPolicyTypeForPSOCollectio
 						Result.Add(LMP_MOBILE_MOVABLE_DIRECTIONAL_LIGHT_WITH_LIGHTMAP);
 					}
 				}
-				else
-				{
-					Result.Add(LMP_LQ_LIGHTMAP);
-				}
+				
+				Result.Add(LMP_LQ_LIGHTMAP);
 			}
 						
 			// ILC/LVM
 			if (bMovable)
 			{
-				if (!bUsesDeferredShading && FReadOnlyCVARCache::MobileEnableStaticAndCSMShadowReceivers())
+				Result.Add(LMP_NO_LIGHTMAP); // in case there is no valid ILC/VLM
+				
+				if (!bUsesDeferredShading && FReadOnlyCVARCache::MobileEnableStaticAndCSMShadowReceivers() && bCanReceiveCSM)
 				{
 					Result.Add(LMP_MOBILE_DIRECTIONAL_LIGHT_CSM_AND_SH_INDIRECT);
 				}
@@ -1105,7 +1107,11 @@ void FMobileBasePassMeshProcessor::CollectPSOInitializers(const FSceneTexturesCo
 	const FMaterialShadingModelField ShadingModels = Material.GetShadingModels();
 	const bool bLitMaterial = ShadingModels.IsLit();
 
-	bool bMovable = PreCacheParams.Mobility == EComponentMobility::Movable || PreCacheParams.Mobility == EComponentMobility::Stationary;
+	bool bMovable = 
+		PreCacheParams.Mobility == EComponentMobility::Movable || 
+		PreCacheParams.Mobility == EComponentMobility::Stationary || 
+		PreCacheParams.bUsesIndirectLightingCache; // ILC uses movable path
+
 	bool bDitheredLODTransition = !bMovable && Material.IsDitheredLODTransition() && !PreCacheParams.bForceLODModel;
 
 	// Setup the draw state
diff --git a/Engine/Source/Runtime/Renderer/Public/MeshPassProcessor.inl b/Engine/Source/Runtime/Renderer/Public/MeshPassProcessor.inl
--- a/Engine/Source/Runtime/Renderer/Public/MeshPassProcessor.inl
+++ b/Engine/Source/Runtime/Renderer/Public/MeshPassProcessor.inl
@@ -288,6 +288,7 @@ void FMeshPassProcessor::AddGraphicsPipelineStateInitializer(
 	MinimalPipelineStateInitializer.BlendState = DrawRenderState.GetBlendState();
 	MinimalPipelineStateInitializer.DepthStencilState = DrawRenderState.GetDepthStencilState();
 	MinimalPipelineStateInitializer.DrawShadingRate = GetShadingRateFromMaterial(MaterialResource.GetShadingRate());
+	MinimalPipelineStateInitializer.bAllowVariableRateShading = MaterialResource.IsVariableRateShadingAllowed() && HardwareVariableRateShadingSupportedByPlatform(GMaxRHIShaderPlatform);
 
 	// NOTE: AsGraphicsPipelineStateInitializer will create the RHIShaders internally if they are not cached yet
 	FGraphicsPipelineStateInitializer PipelineStateInitializer = MinimalPipelineStateInitializer.AsGraphicsPipelineStateInitializer(); 
diff --git a/Engine/Source/Runtime/VulkanRHI/Private/VulkanChunkedPipelineCache.cpp b/Engine/Source/Runtime/VulkanRHI/Private/VulkanChunkedPipelineCache.cpp
--- a/Engine/Source/Runtime/VulkanRHI/Private/VulkanChunkedPipelineCache.cpp
+++ b/Engine/Source/Runtime/VulkanRHI/Private/VulkanChunkedPipelineCache.cpp
@@ -64,7 +64,7 @@ namespace UE
 {
 	namespace Vulkan
 	{
-		static int32 GUseChunkedPSOCache = 0;
+		static int32 GUseChunkedPSOCache = PLATFORM_ANDROID;
 		static FAutoConsoleVariableRef GVulkanGUseNewCacheCodeCVar(
 			TEXT("r.Vulkan.UseChunkedPSOCache"),
 			GUseChunkedPSOCache,
diff --git a/Engine/Source/Runtime/VulkanRHI/Private/VulkanRHI.cpp b/Engine/Source/Runtime/VulkanRHI/Private/VulkanRHI.cpp
--- a/Engine/Source/Runtime/VulkanRHI/Private/VulkanRHI.cpp
+++ b/Engine/Source/Runtime/VulkanRHI/Private/VulkanRHI.cpp
@@ -91,9 +91,9 @@ static FAutoConsoleVariableRef GCVarEnableTransientResourceAllocator(
 
 static TAutoConsoleVariable<bool> CVarAllowVulkanPSOPrecache(
 	TEXT("r.Vulkan.AllowPSOPrecaching"),
-	false,
-	TEXT("true: if r.PSOPrecaching=1 Vulkan RHI will use precaching.\n")
-	TEXT("false: Vulkan RHI will disable precaching (even if r.PSOPrecaching=1). (default)"),
+	true,
+	TEXT("true: if r.PSOPrecaching=1 Vulkan RHI will use precaching. (default)\n")
+	TEXT("false: Vulkan RHI will disable precaching (even if r.PSOPrecaching=1)."),
 	ECVF_RenderThreadSafe | ECVF_ReadOnly);
 
 // If precaching is active we should not need the file cache.
@@ -516,9 +516,7 @@ FVulkanDynamicRHI::FVulkanDynamicRHI()
 
 	GRHIGlobals.SupportsBarycentricsSemantic = true;
 
-	static const auto CVarPSOPrecaching = IConsoleManager::Get().FindConsoleVariable(TEXT("r.PSOPrecaching"));
-
-	GRHISupportsPSOPrecaching = FVulkanChunkedPipelineCacheManager::IsEnabled() && (CVarPSOPrecaching && CVarPSOPrecaching->GetInt() != 0) && CVarAllowVulkanPSOPrecache.GetValueOnAnyThread();
+	GRHISupportsPSOPrecaching = CVarAllowVulkanPSOPrecache.GetValueOnAnyThread();
 	GRHISupportsPipelineFileCache = !GRHISupportsPSOPrecaching || CVarEnableVulkanPSOFileCacheWhenPrecachingActive.GetValueOnAnyThread();
 	UE_LOG(LogVulkanRHI, Log, TEXT("Vulkan PSO Precaching = %d, PipelineFileCache = %d"), GRHISupportsPSOPrecaching, GRHISupportsPipelineFileCache);
 
