Mobile: Fixed flicker when rendering primitives with refractive materials in deferred shading mode

diff --git a/Engine/Source/Runtime/Engine/Private/SceneUtils.cpp b/Engine/Source/Runtime/Engine/Private/SceneUtils.cpp
--- a/Engine/Source/Runtime/Engine/Private/SceneUtils.cpp
+++ b/Engine/Source/Runtime/Engine/Private/SceneUtils.cpp
@@ -49,18 +49,6 @@ bool IsMobilePropagateAlphaEnabled(EShaderPlatform Platform)
 	return IsMobilePlatform(Platform) && (FPlatformMisc::GetMobilePropagateAlphaSetting() > 0);
 }
 
-ENGINE_API EMobileHDRMode GetMobileHDRMode()
-{
-	EMobileHDRMode HDRMode = EMobileHDRMode::EnabledFloat16;
-
-	if (!IsMobileHDR())
-	{
-		HDRMode = EMobileHDRMode::Disabled;
-	}
-	
-	return HDRMode;
-}
-
 ENGINE_API bool IsMobileColorsRGB()
 {
 	static auto* MobileUseHWsRGBEncodingCVAR = IConsoleManager::Get().FindTConsoleVariableDataInt(TEXT("r.Mobile.UseHWsRGBEncoding"));
diff --git a/Engine/Source/Runtime/Engine/Public/SceneUtils.h b/Engine/Source/Runtime/Engine/Public/SceneUtils.h
--- a/Engine/Source/Runtime/Engine/Public/SceneUtils.h
+++ b/Engine/Source/Runtime/Engine/Public/SceneUtils.h
@@ -61,8 +61,6 @@ inline bool IsTemporalAccumulationBasedMethod(EAntiAliasingMethod AntiAliasingMe
 /** True if Alpha Propagate is enabled for the mobile renderer. */
 ENGINE_API bool IsMobilePropagateAlphaEnabled(EShaderPlatform Platform);
 
-ENGINE_API EMobileHDRMode GetMobileHDRMode();
-
 ENGINE_API bool IsMobileColorsRGB();
 
 ENGINE_API EAntiAliasingMethod GetDefaultAntiAliasingMethod(const FStaticFeatureLevel InFeatureLevel);
diff --git a/Engine/Source/Runtime/Renderer/Private/MobileDistortionPass.cpp b/Engine/Source/Runtime/Renderer/Private/MobileDistortionPass.cpp
--- a/Engine/Source/Runtime/Renderer/Private/MobileDistortionPass.cpp
+++ b/Engine/Source/Runtime/Renderer/Private/MobileDistortionPass.cpp
@@ -26,14 +26,9 @@ bool IsMobileDistortionActive(const FViewInfo& View)
 	static IConsoleVariable* CVar = IConsoleManager::Get().FindConsoleVariable(TEXT("r.DisableDistortion"));
 	int32 DisableDistortion = CVar->GetInt();
 
-	// Distortion on mobile requires SceneDepth information in SceneColor.A channel
-	const EMobileHDRMode HDRMode = GetMobileHDRMode();
-	const bool bVisiblePrims = View.ParallelMeshDrawCommandPasses[EMeshPass::Distortion].HasAnyDraw();
-
 	return
-		HDRMode == EMobileHDRMode::EnabledFloat16 &&
 		View.Family->EngineShowFlags.Translucency &&
-		bVisiblePrims &&
+		View.bHasDistortionPrimitives &&
 		FSceneRenderer::GetRefractionQuality(*View.Family) > 0 &&
 		DisableDistortion == 0;
 }
