Use primary view to get the LumenHardwareRayTracingUniformBuffer during raytracing SBT setup

diff --git a/Engine/Source/Runtime/Renderer/Private/RayTracing/RayTracing.cpp b/Engine/Source/Runtime/Renderer/Private/RayTracing/RayTracing.cpp
--- a/Engine/Source/Runtime/Renderer/Private/RayTracing/RayTracing.cpp
+++ b/Engine/Source/Runtime/Renderer/Private/RayTracing/RayTracing.cpp
@@ -340,7 +340,7 @@ namespace RayTracing
 			StaticUniformBuffers.AddUniformBuffer(View.ViewUniformBuffer.GetReference());
 			StaticUniformBuffers.AddUniformBuffer(SceneUniformBuffer);
 			StaticUniformBuffers.AddUniformBuffer(Nanite::GRayTracingManager.GetUniformBuffer().GetReference());
-			StaticUniformBuffers.AddUniformBuffer(View.LumenHardwareRayTracingUniformBuffer->GetRHI());
+			StaticUniformBuffers.AddUniformBuffer(View.GetPrimaryView()->LumenHardwareRayTracingUniformBuffer->GetRHI());
 
 			StaticUniformBufferScope.Emplace(RHICmdList, StaticUniformBuffers);
 		}
