[Backout] - CL34897867
Original CL Desc
-----------------------------------------------------------------
Scene updates refactor to make it easier to add & maintain light/primitive update types and propagate this to the scene extensions.

* Added TSceneUpdateCommandQueue which is used to queue scene update commands (add/remove/update<Payload>), update payload types are extensible by a statically compiled enum but can be extended with registration if need be.
* Changed scene primitive updates to use this instead of separate hash tables.
* Changed scene light updates to use the same instead of its bespoke queue.
* Exposed the updates in the FScene[Pre|Post]UpdateChangeSet such that all extensions can access any change.
* Refactored relevant scene extensions & scene update to use the new data.
* Sets us up for being able to pass around incremental instance updates to various scene modules (once the delta data is plumbed through, that is).
* Step towards batching the updates on the game thread, though this requires a larger refactor of how and when scene updates can be invoked.

diff --git a/Engine/Source/Runtime/Renderer/Private/GPUScene.cpp b/Engine/Source/Runtime/Renderer/Private/GPUScene.cpp
--- a/Engine/Source/Runtime/Renderer/Private/GPUScene.cpp
+++ b/Engine/Source/Runtime/Renderer/Private/GPUScene.cpp
@@ -2086,11 +2086,14 @@ bool FGPUScene::HasPendingGPUWrite(uint32 PrimitiveId) const
 
 void FGPUScene::OnPreSceneUpdate(FRDGBuilder& GraphBuilder, const FScenePreUpdateChangeSet& ScenePreUpdateData)
 {
-	ScenePreUpdateData.PrimitiveUpdates.ForEachUpdateCommand(ESceneUpdateCommandFilter::Updated | ESceneUpdateCommandFilter::Deleted, EPrimitiveUpdateDirtyFlags::GPUState, 
-	[&](const FPrimitiveUpdateCommand& Cmd)
+	for (FPersistentPrimitiveIndex PersistentPrimitiveIndex  : ScenePreUpdateData.RemovedPrimitiveIds)
 	{
-		AddPrimitiveToUpdate(Cmd.GetPersistentId(), Cmd.IsDelete() ? EPrimitiveDirtyState::Removed : EPrimitiveDirtyState::ChangedAll );
-	});
+		AddPrimitiveToUpdate(PersistentPrimitiveIndex, EPrimitiveDirtyState::Removed);
+	}
+	for (FPersistentPrimitiveIndex PersistentPrimitiveIndex  : ScenePreUpdateData.UpdatedPrimitiveIds)
+	{
+		AddPrimitiveToUpdate(PersistentPrimitiveIndex, EPrimitiveDirtyState::ChangedTransform);
+	}
 }
 
 void FGPUScene::OnPostSceneUpdate(FRDGBuilder& GraphBuilder, const FScenePostUpdateChangeSet& ScenePostUpdateData)
diff --git a/Engine/Source/Runtime/Renderer/Private/LightSceneInfo.cpp b/Engine/Source/Runtime/Renderer/Private/LightSceneInfo.cpp
--- a/Engine/Source/Runtime/Renderer/Private/LightSceneInfo.cpp
+++ b/Engine/Source/Runtime/Renderer/Private/LightSceneInfo.cpp
@@ -125,7 +125,7 @@ void FLightSceneInfo::AddToScene()
 
 bool FLightSceneInfo::ShouldCreateLightPrimitiveInteraction(const FLightSceneInfoCompact& LightSceneInfoCompact, const FPrimitiveSceneInfoCompact& PrimitiveSceneInfoCompact)
 {
-	if (LightSceneInfoCompact.AffectsPrimitive(FBoxSphereBounds(PrimitiveSceneInfoCompact.Bounds), PrimitiveSceneInfoCompact.Proxy))
+	if (!Scene->IsPrimitiveBeingRemoved(PrimitiveSceneInfoCompact.PrimitiveSceneInfo) && LightSceneInfoCompact.AffectsPrimitive(FBoxSphereBounds(PrimitiveSceneInfoCompact.Bounds), PrimitiveSceneInfoCompact.Proxy))
 	{
 		// create light interaction and add to light/primitive lists
 		return FLightPrimitiveInteraction::ShouldCreate(this, PrimitiveSceneInfoCompact.PrimitiveSceneInfo).bShouldCreate;
@@ -135,7 +135,7 @@ bool FLightSceneInfo::ShouldCreateLightPrimitiveInteraction(const FLightSceneInf
 
 void FLightSceneInfo::CreateLightPrimitiveInteraction(const FLightSceneInfoCompact& LightSceneInfoCompact, const FPrimitiveSceneInfoCompact& PrimitiveSceneInfoCompact)
 {
-	if (LightSceneInfoCompact.AffectsPrimitive(FBoxSphereBounds(PrimitiveSceneInfoCompact.Bounds), PrimitiveSceneInfoCompact.Proxy))
+	if (!Scene->IsPrimitiveBeingRemoved(PrimitiveSceneInfoCompact.PrimitiveSceneInfo) && LightSceneInfoCompact.AffectsPrimitive(FBoxSphereBounds(PrimitiveSceneInfoCompact.Bounds), PrimitiveSceneInfoCompact.Proxy))
 	{
 		// create light interaction and add to light/primitive lists
 		FLightPrimitiveInteraction::Create(this,PrimitiveSceneInfoCompact.PrimitiveSceneInfo);
diff --git a/Engine/Source/Runtime/Renderer/Private/LightSceneInfo.h b/Engine/Source/Runtime/Renderer/Private/LightSceneInfo.h
--- a/Engine/Source/Runtime/Renderer/Private/LightSceneInfo.h
+++ b/Engine/Source/Runtime/Renderer/Private/LightSceneInfo.h
@@ -205,16 +205,13 @@ class FLightSceneInfo
 	FLightPrimitiveInteraction* DynamicInteractionStaticPrimitiveList;
 
 public:
-	using FPersistentId = int32;
-
 	/** The light's scene proxy. */
 	FLightSceneProxy* Proxy;
 
 	ELightComponentType Type;
 
 	/** If bVisible == true, this is the index of the primitive in Scene->Lights. */
-	FPersistentId Id;
-	FORCEINLINE FPersistentId GetPersistentIndex() const { return Id; }
+	int32 Id;
 
 	/** The identifier for the primitive in Scene->PrimitiveOctree. */
 	FOctreeElementId2 OctreeId;
diff --git a/Engine/Source/Runtime/Renderer/Private/RendererScene.cpp b/Engine/Source/Runtime/Renderer/Private/RendererScene.cpp
--- a/Engine/Source/Runtime/Renderer/Private/RendererScene.cpp
+++ b/Engine/Source/Runtime/Renderer/Private/RendererScene.cpp
@@ -175,6 +175,92 @@ TGlobalResource< FGlobalDitherUniformBuffer > GDitherFadedInUniformBuffer;
 
 static FThreadSafeCounter FSceneViewState_UniqueID;
 
+struct FUpdateLightTransformParameters
+{
+	FMatrix LightToWorld;
+	FVector4 Position;
+};
+
+struct FUpdateLightCommand
+{
+	enum class EAddOrRemove
+	{
+		Add,
+		Remove,
+		None,
+	};
+
+	struct FColorParameters
+	{
+		FLinearColor NewColor;
+		float NewIndirectLightingScale;
+		float NewVolumetricScatteringIntensity;
+	};
+
+	using FTransformParameters = FUpdateLightTransformParameters;
+
+	inline bool IsAdd() const { return AddOrRemove == EAddOrRemove::Add; }
+	inline bool IsRemove() const { return AddOrRemove == EAddOrRemove::Remove; }
+
+	void Set(const FColorParameters& InColorParameters)
+	{
+		bHasColor = true;
+		ColorParameters = InColorParameters;
+	}
+	void Set(const FTransformParameters& InTransformParameters)
+	{
+		bHasTransform = true;
+		TransformParameters = InTransformParameters;
+	}
+	void Set(const EAddOrRemove& InAddOrRemove)
+	{
+		AddOrRemove = InAddOrRemove;
+	}
+
+	FUpdateLightCommand(const FColorParameters& InColorParameters, FLightSceneInfo* InLightSceneInfo) : LightSceneInfo(InLightSceneInfo), AddOrRemove(EAddOrRemove::None), bHasTransform(false) { Set(InColorParameters); }
+	FUpdateLightCommand(const FTransformParameters& InTransformParameters, FLightSceneInfo* InLightSceneInfo) : LightSceneInfo(InLightSceneInfo), AddOrRemove(EAddOrRemove::None), bHasColor(false) { Set(InTransformParameters); }
+	// Crtor for add/re,pve
+	FUpdateLightCommand(EAddOrRemove InAddOrRemove, FLightSceneInfo* InLightSceneInfo) : LightSceneInfo(InLightSceneInfo), AddOrRemove(InAddOrRemove), bHasTransform(false), bHasColor(false) {}
+
+	FTransformParameters TransformParameters;
+	FColorParameters ColorParameters;
+	FLightSceneInfo* LightSceneInfo;
+	EAddOrRemove AddOrRemove;
+	uint32 bHasTransform : 1;
+	uint32 bHasColor : 1;
+};
+
+class FSceneLightInfoUpdates
+{
+public:
+	Experimental::TRobinHoodHashMap<FLightSceneInfo*, FUpdateLightCommand> Commands;
+
+	int32 NumAdds = 0;
+	int32 NumRemoves = 0;
+	int32 NumUpdates = 0;
+
+	template <typename PayloadT>
+	void Enqueue(const PayloadT& Payload, FLightSceneInfo* LightSceneInfo)
+	{
+		// Allocate a new slot for update data
+		bool bWasAlreadyInMap = false;
+		FUpdateLightCommand* Command = Commands.FindOrAdd(LightSceneInfo, FUpdateLightCommand(Payload, LightSceneInfo), bWasAlreadyInMap);
+		if (bWasAlreadyInMap)
+		{
+			check(LightSceneInfo == Command->LightSceneInfo);
+			Command->Set(Payload);
+		}
+	}
+
+	void Reset()
+	{
+		Commands.Empty();
+		NumAdds = 0;
+		NumRemoves = 0;
+		NumUpdates = 0;
+	}
+};
+
 /**
  * Holds the info to update SpeedTree wind per unique tree object in the scene, instead of per instance
  */
@@ -1505,13 +1591,12 @@ static void TBitArraySwapElements(TBitArray<>& Array, int32 i1, int32 i2)
 
 void FScene::AddPrimitiveSceneInfo_RenderThread(FPrimitiveSceneInfo* PrimitiveSceneInfo, const TOptional<FTransform>& PreviousTransform)
 {
-	// Must always be a novel primitive that is added
 	check(PrimitiveSceneInfo->PackedIndex == INDEX_NONE);
-	PrimitiveUpdates.EnqueueAdd(PrimitiveSceneInfo);
-
+	check(AddedPrimitiveSceneInfos.Find(PrimitiveSceneInfo) == nullptr);
+	AddedPrimitiveSceneInfos.FindOrAdd(PrimitiveSceneInfo);
 	if (PreviousTransform.IsSet())
 	{
-		PrimitiveUpdates.Enqueue<FUpdateOverridePreviousTransformData>(PrimitiveSceneInfo, FUpdateOverridePreviousTransformData(PreviousTransform.GetValue().ToMatrixWithScale()));
+		OverridenPreviousTransforms.Update(PrimitiveSceneInfo, PreviousTransform.GetValue().ToMatrixWithScale());
 	}
 }
 
@@ -1707,7 +1792,8 @@ FScene::~FScene()
 	}
 #endif
 
-	checkf(PrimitiveUpdates.IsEmpty(), TEXT("All pending primitive addition operations are expected to be flushed when the scene is destroyed. Remaining operations are likely to cause a memory leak."));
+	checkf(RemovedPrimitiveSceneInfos.Num() == 0, TEXT("All pending primitive removal operations are expected to be flushed when the scene is destroyed. Remaining operations are likely to cause a memory leak."));
+	checkf(AddedPrimitiveSceneInfos.Num() == 0, TEXT("All pending primitive addition operations are expected to be flushed when the scene is destroyed. Remaining operations are likely to cause a memory leak."));
 	checkf(Primitives.Num() == 0, TEXT("All primitives are expected to be removed before the scene is destroyed. Remaining primitives are likely to cause a memory leak."));
 
 	delete InstanceCullingOcclusionQueryRenderer;
@@ -1760,6 +1846,8 @@ FScene::~FScene()
 	RayTracingSkinnedGeometryUpdateQueue = nullptr;
 #endif // RHI_RAYTRACING
 
+	checkf(RemovedPrimitiveSceneInfos.Num() == 0, TEXT("Leaking %d FPrimitiveSceneInfo instances."), RemovedPrimitiveSceneInfos.Num()); // Ensure UpdateAllPrimitiveSceneInfos() is called before destruction.
+
 	delete SceneLightInfoUpdates;
 }
 
@@ -1911,7 +1999,6 @@ void FScene::BatchAddPrimitivesInternal(TArrayView<T*> InPrimitives)
 
 	if (!CreateCommands.IsEmpty())
 	{
-		// Must enqueue RT commands to ensure SetTransform & CreateRenderThreadResources are executed in-order.
 		ENQUEUE_RENDER_COMMAND(AddPrimitiveCommand)(
 			[this, CreateCommands = MoveTemp(CreateCommands)](FRHICommandListBase& RHICmdList)
 		{
@@ -1953,57 +2040,45 @@ static FAutoConsoleVariableRef CVarSkipRedundantTransformUpdate(
 	ECVF_Default
 );
 
-#if VALIDATE_PRIMITIVE_PACKED_INDEX
-inline void ValidatePackedPrimitiveIndexForUpdate(FPrimitiveSceneInfo* PrimitiveSceneInfo, const FScenePrimitiveUpdate& PrimitiveUpdates)
+void FScene::UpdatePrimitiveTransform_RenderThread(FPrimitiveSceneProxy* PrimitiveSceneProxy, const FBoxSphereBounds& WorldBounds, const FBoxSphereBounds& LocalBounds, const FMatrix& LocalToWorld, const FVector& AttachmentRootPosition, const TOptional<FTransform>& PreviousTransform)
 {
-	if (const FPrimitiveUpdateCommand* Cmd = PrimitiveUpdates.FindCommand(PrimitiveSceneInfo))
+#if VALIDATE_PRIMITIVE_PACKED_INDEX
+	if (AddedPrimitiveSceneInfos.Find(PrimitiveSceneProxy->GetPrimitiveSceneInfo()) != nullptr)
 	{
-		if (Cmd->bAdded)
-		{
-			check(PrimitiveSceneInfo->GetIndex() == INDEX_NONE);
-		}
-		else
-		{
-			check(PrimitiveSceneInfo->GetIndex() != INDEX_NONE);
-		}				
+		check(PrimitiveSceneProxy->GetPrimitiveSceneInfo()->PackedIndex == INDEX_NONE);
 	}
 	else
 	{
-		// sending update, and no queued command - must have a valid index
-		check(PrimitiveSceneInfo->GetIndex() != INDEX_NONE);
+		check(PrimitiveSceneProxy->GetPrimitiveSceneInfo()->PackedIndex != INDEX_NONE);
 	}
-}
-#endif
 
-void FScene::UpdatePrimitiveTransform_RenderThread(FPrimitiveSceneProxy* PrimitiveSceneProxy, const FBoxSphereBounds& WorldBounds, const FBoxSphereBounds& LocalBounds, const FMatrix& LocalToWorld, const FVector& AttachmentRootPosition, const TOptional<FTransform>& PreviousTransform)
-{
-	FPrimitiveSceneInfo* PrimitiveSceneInfo = PrimitiveSceneProxy->GetPrimitiveSceneInfo();
-#if VALIDATE_PRIMITIVE_PACKED_INDEX
-	ValidatePackedPrimitiveIndexForUpdate(PrimitiveSceneInfo, PrimitiveUpdates);
+	check(RemovedPrimitiveSceneInfos.Find(PrimitiveSceneProxy->GetPrimitiveSceneInfo()) == nullptr);
 #endif
 
-	PrimitiveUpdates.Enqueue(PrimitiveSceneInfo, FUpdateTransformCommand { .WorldBounds = WorldBounds, .LocalBounds = LocalBounds, .LocalToWorld = LocalToWorld, .AttachmentRootPosition = AttachmentRootPosition });
+	UpdatedTransforms.Update(PrimitiveSceneProxy, { WorldBounds, LocalBounds, LocalToWorld, AttachmentRootPosition });
 
 	if (PreviousTransform.IsSet())
 	{
-		PrimitiveUpdates.Enqueue(PrimitiveSceneInfo, FUpdateOverridePreviousTransformData(PreviousTransform.GetValue().ToMatrixWithScale()));
+		OverridenPreviousTransforms.Update(PrimitiveSceneProxy->GetPrimitiveSceneInfo(), PreviousTransform.GetValue().ToMatrixWithScale());
 	}
 }
 
-template <typename UpdatePayloadType>
-void FScene::UpdatePrimitiveInternal(FPrimitiveSceneProxy* SceneProxy, UpdatePayloadType&& InUpdatePayload)
+void FScene::UpdatePrimitiveOcclusionBoundsSlack_RenderThread(const FPrimitiveSceneProxy* PrimitiveSceneProxy, float NewSlack)
 {
-	if (SceneProxy != nullptr)
-	{
-		ENQUEUE_RENDER_COMMAND(UpdatePrimitiveCmd)(
-		[this, PrimitiveSceneInfo = SceneProxy->PrimitiveSceneInfo, UpdatePayload = MoveTemp(InUpdatePayload)] (FRHICommandListBase&) mutable
-		{
 #if VALIDATE_PRIMITIVE_PACKED_INDEX
-			ValidatePackedPrimitiveIndexForUpdate(PrimitiveSceneInfo, PrimitiveUpdates);
-#endif
-			PrimitiveUpdates.Enqueue<UpdatePayloadType>(PrimitiveSceneInfo, MoveTemp(UpdatePayload));
-		});
+	if (AddedPrimitiveSceneInfos.Find(PrimitiveSceneProxy->GetPrimitiveSceneInfo()) != nullptr)
+	{
+		check(PrimitiveSceneProxy->GetPrimitiveSceneInfo()->PackedIndex == INDEX_NONE);
 	}
+	else
+	{
+		check(PrimitiveSceneProxy->GetPrimitiveSceneInfo()->PackedIndex != INDEX_NONE);
+	}
+
+	check(RemovedPrimitiveSceneInfos.Find(PrimitiveSceneProxy->GetPrimitiveSceneInfo()) == nullptr);
+#endif
+
+	UpdatedOcclusionBoundsSlacks.Update(PrimitiveSceneProxy, NewSlack);
 }
 
 void FScene::UpdatePrimitiveTransform(UPrimitiveComponent* Primitive)
@@ -2016,7 +2091,7 @@ void FScene::UpdatePrimitiveTransform(FPrimitiveSceneDesc* Primitive)
 	UpdatePrimitiveTransformInternal(Primitive);	
 }
 
-template<class T> 	
+template<class T>
 void FScene::UpdatePrimitiveTransformInternal(T* Primitive)
 {
 	SCOPE_CYCLE_COUNTER(STAT_UpdatePrimitiveTransformGT);
@@ -2174,17 +2249,38 @@ void FScene::FinishUpdatePrimitiveTransform()
 
 void FScene::UpdatePrimitiveOcclusionBoundsSlack(UPrimitiveComponent* Primitive, float NewSlack)
 {
-	UpdatePrimitiveInternal(Primitive->GetSceneProxy(), FUpdateOcclusionBoundsSlacksData(NewSlack));
+	if (const FPrimitiveSceneProxy* SceneProxy = Primitive->GetSceneProxy())
+	{
+		ENQUEUE_RENDER_COMMAND(UpdateOcclusionBoundsSlackCmd)(
+			[this, SceneProxy, NewSlack] (FRHICommandListBase&)
+			{
+				UpdatePrimitiveOcclusionBoundsSlack_RenderThread(SceneProxy, NewSlack);
+			});
+	}
 }
 
 void FScene::UpdatePrimitiveDrawDistance(UPrimitiveComponent* Primitive, float MinDrawDistance, float MaxDrawDistance, float VirtualTextureMaxDrawDistance)
 {
-	UpdatePrimitiveInternal(Primitive->GetSceneProxy(), FUpdateDrawDistanceData(FVector3f(MinDrawDistance, MaxDrawDistance, VirtualTextureMaxDrawDistance)));
+	if (FPrimitiveSceneProxy* SceneProxy = Primitive->GetSceneProxy())
+	{
+		ENQUEUE_RENDER_COMMAND(UpdatePrimitiveDrawDistanceCmd)(
+			[this, SceneProxy, MinDrawDistance, MaxDrawDistance, VirtualTextureMaxDrawDistance] (FRHICommandListBase&)
+			{
+				UpdatedDrawDistance.Update(SceneProxy, FVector3f(MinDrawDistance, MaxDrawDistance, VirtualTextureMaxDrawDistance));
+			});
+	}
 }
 
 void FScene::UpdateInstanceCullDistance(UPrimitiveComponent* Primitive, float StartCullDistance, float EndCullDistance)
 {
-	UpdatePrimitiveInternal(Primitive->GetSceneProxy(), FUpdateInstanceCullDistanceData(FVector2f(StartCullDistance, EndCullDistance)));
+	if (FPrimitiveSceneProxy* SceneProxy = Primitive->GetSceneProxy())
+	{
+		ENQUEUE_RENDER_COMMAND(UpdateInstanceCullDistanceCmd)(
+			[this, SceneProxy, StartCullDistance, EndCullDistance] (FRHICommandListBase&)
+			{
+				UpdatedInstanceCullDistance.Update(SceneProxy, FVector2f(StartCullDistance, EndCullDistance));
+			});
+	}
 }
 
 void FScene::UpdatePrimitiveInstances(UInstancedStaticMeshComponent* Primitive)
@@ -2202,12 +2298,12 @@ void FScene::UpdatePrimitiveInstances(UInstancedStaticMeshComponent* Primitive)
 	FUpdateInstanceCommand UpdateParams;
 	UpdateParams.PrimitiveSceneProxy = Primitive->GetSceneProxy();
 	UpdateParams.WorldBounds = Primitive->Bounds;
-	UpdateParams.LocalBounds = static_cast<UPrimitiveComponent*>(Primitive)->GetLocalBounds();
+	UpdateParams.LocalBounds = ((UPrimitiveComponent*)Primitive)->GetLocalBounds();
 
 	// #todo (jnadro) This code should not be dependent on static mesh bounds.		
 	UpdateParams.StaticMeshBounds = Primitive->GetStaticMesh()->GetBounds();
 
-	UpdatePrimitiveInstances(UpdateParams);
+	return UpdatePrimitiveInstances(UpdateParams);
 }
 
 void FScene::UpdatePrimitiveInstances(FInstancedStaticMeshSceneDesc* Primitive)
@@ -2230,12 +2326,30 @@ void FScene::UpdatePrimitiveInstances(FInstancedStaticMeshSceneDesc* Primitive)
 	// #todo (jnadro) This code should not be dependent on static mesh bounds.		
 	UpdateParams.StaticMeshBounds = Primitive->GetStaticMesh()->GetBounds();
 
-	UpdatePrimitiveInstances(UpdateParams);
+	return UpdatePrimitiveInstances(UpdateParams);
 }
 
 void FScene::UpdatePrimitiveInstances(FUpdateInstanceCommand& UpdateParams)
 {	
-	UpdatePrimitiveInternal(UpdateParams.PrimitiveSceneProxy, MoveTemp(UpdateParams));
+		ENQUEUE_RENDER_COMMAND(UpdateInstanceCommand)(
+			[this, UpdateParams = MoveTemp(UpdateParams)] (FRHICommandListBase&)
+		{
+#if VALIDATE_PRIMITIVE_PACKED_INDEX
+			if (AddedPrimitiveSceneInfos.Find(UpdateParams.PrimitiveSceneProxy->GetPrimitiveSceneInfo()) != nullptr)
+			{
+				check(UpdateParams.PrimitiveSceneProxy->GetPrimitiveSceneInfo()->PackedIndex == INDEX_NONE);
+			}
+			else
+			{
+				check(UpdateParams.PrimitiveSceneProxy->GetPrimitiveSceneInfo()->PackedIndex != INDEX_NONE);
+			}
+
+			check(RemovedPrimitiveSceneInfos.Find(UpdateParams.PrimitiveSceneProxy->GetPrimitiveSceneInfo()) == nullptr);
+#endif
+			FScopeCycleCounter Context(UpdateParams.PrimitiveSceneProxy->GetStatId());
+			UpdatedInstances.Update(UpdateParams.PrimitiveSceneProxy, UpdateParams);
+		}
+	);
 }
 
 void FScene::UpdatePrimitiveSelectedState_RenderThread(const FPrimitiveSceneInfo* PrimitiveSceneInfo, bool bIsSelected)
@@ -2259,11 +2373,22 @@ void FScene::UpdatePrimitiveLightingAttachmentRoot(UPrimitiveComponent* Primitiv
 
 	if (NewLightingAttachmentRoot == Primitive)
 	{
-		NewLightingAttachmentRoot = nullptr;
+		NewLightingAttachmentRoot = NULL;
 	}
 
 	FPrimitiveComponentId NewComponentId = NewLightingAttachmentRoot ? NewLightingAttachmentRoot->GetPrimitiveSceneId() : FPrimitiveComponentId();
-	UpdatePrimitiveInternal(Primitive->GetSceneProxy(), FUpdateAttachmentRootData(NewComponentId));
+
+	if (Primitive->SceneProxy)
+	{
+		FPrimitiveSceneProxy* Proxy = Primitive->SceneProxy;
+		FScene* Scene = this;
+		ENQUEUE_RENDER_COMMAND(UpdatePrimitiveAttachment)(
+			[Scene, Proxy, NewComponentId] (FRHICommandListBase&)
+			{
+				FPrimitiveSceneInfo* PrimitiveInfo = Proxy->GetPrimitiveSceneInfo();
+				Scene->UpdatedAttachmentRoots.Update(PrimitiveInfo, NewComponentId);
+			});
+	}
 }
 
 void FScene::UpdatePrimitiveAttachment(UPrimitiveComponent* Primitive)
@@ -2304,15 +2429,44 @@ void FScene::UpdateCustomPrimitiveData(FPrimitiveSceneDesc* Primitive, const FCu
 
 void FScene::UpdateCustomPrimitiveData(FPrimitiveSceneProxy* SceneProxy, const FCustomPrimitiveData& CustomPrimitiveData)
 {
-	UpdatePrimitiveInternal(SceneProxy, FUpdateCustomPrimitiveData(CustomPrimitiveData));
+	// This path updates the primitive data directly in the GPUScene. 
+	if (SceneProxy) 
+	{
+		struct FUpdateParams
+		{
+			FScene* Scene;
+			FPrimitiveSceneProxy* PrimitiveSceneProxy;
+			FCustomPrimitiveData CustomPrimitiveData;
+		};
+
+		FUpdateParams UpdateParams;
+		UpdateParams.Scene = this;
+		UpdateParams.PrimitiveSceneProxy = SceneProxy;
+		UpdateParams.CustomPrimitiveData = CustomPrimitiveData; 
+
+		ENQUEUE_RENDER_COMMAND(UpdateCustomPrimitiveDataCommand)(
+			[UpdateParams] (FRHICommandListBase&)
+			{
+				UpdateParams.Scene->UpdatedCustomPrimitiveParams.Update(UpdateParams.PrimitiveSceneProxy, UpdateParams.CustomPrimitiveData);
+			});
+	}
 }
 
 void FScene::UpdatePrimitiveDistanceFieldSceneData_GameThread(UPrimitiveComponent* Primitive)
 {
 	check(IsInGameThread());
+
 	if (Primitive->SceneProxy)
 	{
-		UpdatePrimitiveInternal(Primitive->GetSceneProxy(), FUpdateDistanceFieldSceneData{});
+		ENQUEUE_RENDER_COMMAND(UpdatePrimDFSceneDataCmd)(
+			[this, PrimitiveSceneProxy = Primitive->SceneProxy] (FRHICommandListBase&)
+			{
+				if (PrimitiveSceneProxy && PrimitiveSceneProxy->GetPrimitiveSceneInfo())
+				{
+					FPrimitiveSceneInfo* Info = PrimitiveSceneProxy->GetPrimitiveSceneInfo();
+					this->DistanceFieldSceneDataUpdates.FindOrAdd(Info);
+				}
+			});
 	}
 }
 
@@ -2336,9 +2490,31 @@ FPrimitiveSceneInfo* FScene::GetPrimitiveSceneInfo(const FPersistentPrimitiveInd
 	return GetPrimitiveSceneInfo(PrimitiveIndex);
 }
 
-void FScene::RemovePrimitiveSceneInfo_RenderThread(FPrimitiveSceneInfo* PrimitiveSceneInfo)
+bool FScene::RemovePrimitiveSceneInfo_RenderThread(FPrimitiveSceneInfo* PrimitiveSceneInfo)
 {
-	PrimitiveUpdates.EnqueueDelete(PrimitiveSceneInfo);
+	const bool bRemovePendingAdd = AddedPrimitiveSceneInfos.Remove(PrimitiveSceneInfo);
+
+	if (bRemovePendingAdd)
+	{
+		check(PrimitiveSceneInfo->PackedIndex == INDEX_NONE);
+		UpdatedTransforms.Remove(PrimitiveSceneInfo->Proxy);
+		UpdatedCustomPrimitiveParams.Remove(PrimitiveSceneInfo->Proxy);
+		OverridenPreviousTransforms.Remove(PrimitiveSceneInfo);
+		UpdatedOcclusionBoundsSlacks.Remove(PrimitiveSceneInfo->Proxy);
+		UpdatedInstanceCullDistance.Remove(PrimitiveSceneInfo->Proxy);
+		UpdatedDrawDistance.Remove(PrimitiveSceneInfo->Proxy);
+		DistanceFieldSceneDataUpdates.Remove(PrimitiveSceneInfo);
+		UpdatedAttachmentRoots.Remove(PrimitiveSceneInfo);
+		DeletedPrimitiveSceneInfos.Emplace(PrimitiveSceneInfo);
+	}
+	else
+	{
+		check(PrimitiveSceneInfo->PackedIndex != INDEX_NONE);
+		check(RemovedPrimitiveSceneInfos.Find(PrimitiveSceneInfo) == nullptr);
+		RemovedPrimitiveSceneInfos.FindOrAdd(PrimitiveSceneInfo);
+	}
+	
+	return !bRemovePendingAdd;
 }
 
 void FScene::RemovePrimitive(UPrimitiveComponent* Primitive)
@@ -2393,7 +2569,6 @@ void FScene::BatchRemovePrimitivesInternal(TArrayView<T*> InPrimitives)
 
 	if (!DestroyCommands.IsEmpty())
 	{
-		// must run RT cmds in order here too
 		ENQUEUE_RENDER_COMMAND(FRemovePrimitiveCommand)(
 			[this, DestroyCommands = MoveTemp(DestroyCommands)](FRHICommandListBase&)
 		{
@@ -2640,10 +2815,10 @@ void FScene::AddLight(ULightComponent* Light)
 			{
 				CSV_SCOPED_TIMING_STAT_EXCLUSIVE(Scene_AddLight);
 				FScopeCycleCounter Context(LightSceneInfo->Proxy->GetStatId());
-
-				// It should not have been added to the scene already
-				check(LightSceneInfo->Id == INDEX_NONE);
-				SceneLightInfoUpdates->EnqueueAdd(LightSceneInfo);
+				bool bWasAlreadyInMap = false;
+				SceneLightInfoUpdates->Commands.FindOrAdd(LightSceneInfo, FUpdateLightCommand(FUpdateLightCommand::EAddOrRemove::Add, LightSceneInfo), bWasAlreadyInMap);
+				check(!bWasAlreadyInMap);
+				++SceneLightInfoUpdates->NumAdds;
 			});
 	}
 }
@@ -3095,7 +3270,7 @@ void FScene::UpdateReflectionCaptureTransform(UReflectionCaptureComponent* Compo
 		FReflectionCaptureProxy* Proxy = Component->SceneProxy;
 		FMatrix Transform = Component->GetComponentTransform().ToMatrixWithScale();
 
-		ENQUEUE_RENDER_COMMAND(UpdateReflectionCaptureTransformCommand)(
+		ENQUEUE_RENDER_COMMAND(FUpdateTransformCommand)(
 			[Scene, Proxy, Transform, bUsingPreviewCaptureData](FRHICommandListBase& RHICmdList)
 		{
 			if (Proxy->bUsingPreviewCaptureData)
@@ -3673,7 +3848,7 @@ bool DoesPlatformNeedLocalLightPrimitiveInteraction(EShaderPlatform ShaderPlatfo
 	return !IsMobilePlatform(ShaderPlatform) || !MobileLocalLightsUseSinglePermutation(ShaderPlatform) || IsMobileMovableSpotlightShadowsEnabled(ShaderPlatform);
 }
 
-void FScene::UpdateLightTransform_RenderThread(int32 LightId, FLightSceneInfo* LightSceneInfo, const FUpdateLightTransformParameters& Parameters)
+void FScene::UpdateLightTransform_RenderThread(int32 LightId, FLightSceneInfo* LightSceneInfo, const FUpdateLightCommand::FTransformParameters& Parameters)
 {
 	SCOPED_NAMED_EVENT(FScene_UpdateLightTransform_RenderThread, FColor::Yellow);
 
@@ -3754,35 +3929,54 @@ void FScene::UpdateLightTransform_RenderThread(int32 LightId, FLightSceneInfo* L
 	}
 }
 
-template <typename UpdatePayloadType>
-void FScene::UpdateLightInternal(FLightSceneProxy* LightSceneProxy, UpdatePayloadType&& InUpdatePayload)
+void FScene::UpdateLightTransform(ULightComponent* Light)
 {
-	if (LightSceneProxy)
+	if(Light->SceneProxy)
 	{
-		FLightSceneInfo* LightSceneInfo = LightSceneProxy->GetLightSceneInfo();
+		FUpdateLightCommand::FTransformParameters Parameters;
+		Parameters.LightToWorld = Light->GetComponentTransform().ToMatrixNoScale();
+		Parameters.Position = Light->GetLightPosition();
+		
+		FLightSceneInfo* LightSceneInfo = Light->SceneProxy->GetLightSceneInfo();
 		if (LightSceneInfo->bVisible)
 		{
 			ENQUEUE_RENDER_COMMAND(UpdateLightTransform)(
-				[this, LightSceneInfo, UpdatePayload = MoveTemp(InUpdatePayload)] (FRHICommandListBase&) mutable
+				[this, LightSceneInfo, Parameters] (FRHICommandListBase&)
 				{
 					FScopeCycleCounter Context(LightSceneInfo->Proxy->GetStatId());
-					SceneLightInfoUpdates->Enqueue(LightSceneInfo, MoveTemp(UpdatePayload));
+					SceneLightInfoUpdates->Enqueue(Parameters, LightSceneInfo);
+					++SceneLightInfoUpdates->NumUpdates;
 				});
 		}
 	}
 }
 
-void FScene::UpdateLightTransform(ULightComponent* Light)
-{
-	UpdateLightInternal(Light->SceneProxy, FUpdateLightTransformParameters{.LightToWorld = Light->GetComponentTransform().ToMatrixNoScale(), .Position = Light->GetLightPosition()});
-}
-
+/** 
+ * Updates the color and brightness of a light which has already been added to the scene. 
+ *
+ * @param Light - light component to update
+ */
 void FScene::UpdateLightColorAndBrightness(ULightComponent* Light)
 {
-	UpdateLightInternal(Light->SceneProxy, FUpdateLightColorParameters{
-		.NewColor = Light->GetColoredLightBrightness(),
-		.NewIndirectLightingScale = Light->IndirectLightingIntensity,
-		.NewVolumetricScatteringIntensity = Light->VolumetricScatteringIntensity});
+	if(Light->SceneProxy)
+	{
+		FUpdateLightCommand::FColorParameters NewParameters;
+		NewParameters.NewColor = Light->GetColoredLightBrightness();
+		NewParameters.NewIndirectLightingScale = Light->IndirectLightingIntensity;
+		NewParameters.NewVolumetricScatteringIntensity = Light->VolumetricScatteringIntensity;
+
+		FScene* Scene = this;
+		FLightSceneInfo* LightSceneInfo = Light->SceneProxy->GetLightSceneInfo();
+		if (LightSceneInfo->bVisible)
+		{
+			ENQUEUE_RENDER_COMMAND(UpdateLightColorAndBrightness)(
+				[this, LightSceneInfo, NewParameters] (FRHICommandListBase&)
+				{
+					SceneLightInfoUpdates->Enqueue(NewParameters, LightSceneInfo);
+					++SceneLightInfoUpdates->NumUpdates;
+				});
+		}
+	}
 }
 
 void FScene::RemoveLightSceneInfo_RenderThread(FLightSceneInfo* LightSceneInfo)
@@ -3917,14 +4111,25 @@ void FScene::RemoveLight(ULightComponent* Light)
 
 				if (LightSceneInfo->bVisible)
 				{
-					SceneLightInfoUpdates->EnqueueDelete(LightSceneInfo);
-					// Note: all the deletes are deferred now (except the invisible ones) as this simpifies the update buffer.
-					return;
+					auto HashElementId = SceneLightInfoUpdates->Commands.FindId(LightSceneInfo);
+					// If an Add is pending, just remove the update and fall through to delete the data.
+					if (HashElementId.IsValid() && SceneLightInfoUpdates->Commands.GetByElementId(HashElementId).Value.IsAdd())
+					{
+						SceneLightInfoUpdates->Commands.RemoveByElementId(HashElementId);
+						--SceneLightInfoUpdates->NumAdds;
+					}
+					else
+					{
+						SceneLightInfoUpdates->Enqueue(FUpdateLightCommand::EAddOrRemove::Remove, LightSceneInfo);
+						++SceneLightInfoUpdates->NumRemoves;
+						// early out to defer the deletion
+						return;
+					}
 				}
 				else
 				{
 					// There should never be updates queued for lights that are not visible
-					check(!SceneLightInfoUpdates->HasCommand(LightSceneInfo));
+					check(SceneLightInfoUpdates->Commands.Find(LightSceneInfo) == nullptr);
 					// The "invisible lights" are removed at once.
 					InvisibleLights.RemoveAt(LightSceneInfo->Id);
 				}
@@ -4658,6 +4863,7 @@ void FScene::ApplyWorldOffset_RenderThread(FRHICommandListBase& RHICmdList, cons
 	GPUScene.bUpdateAllPrimitives = true;
 
 	// Primitives
+	checkf(AddedPrimitiveSceneInfos.Num() == 0, TEXT("All primitives found in AddedPrimitiveSceneInfos must have been added to the scene before the world offset is applied"));
 	for (int32 Idx = 0; Idx < Primitives.Num(); ++Idx)
 	{
 		Primitives[Idx]->ApplyWorldOffset(RHICmdList, InOffset);
@@ -4918,8 +5124,7 @@ void FScene::UpdatePrimitiveVelocityState_RenderThread(FPrimitiveSceneInfo* Prim
 }
 
 #if RHI_RAYTRACING
-
-void FScene::UpdateRayTracingGroupBounds_AddPrimitives(const TArray<FPrimitiveSceneInfo*, SceneRenderingAllocator>& PrimitiveSceneInfos)
+void FScene::UpdateRayTracingGroupBounds_AddPrimitives(const Experimental::TRobinHoodHashSet<FPrimitiveSceneInfo*>& PrimitiveSceneInfos)
 {
 	for (FPrimitiveSceneInfo* const PrimitiveSceneInfo : PrimitiveSceneInfos)
 	{
@@ -4964,7 +5169,7 @@ static void UpdateRayTracingGroupBounds(Experimental::TRobinHoodHashSet<FScene::
 	}
 }
 
-void FScene::UpdateRayTracingGroupBounds_RemovePrimitives(const TArray<FPrimitiveSceneInfo*, SceneRenderingAllocator>& PrimitiveSceneInfos)
+void FScene::UpdateRayTracingGroupBounds_RemovePrimitives(const Experimental::TRobinHoodHashSet<FPrimitiveSceneInfo*>& PrimitiveSceneInfos)
 {
 	Experimental::TRobinHoodHashSet<FRayTracingCullingGroup*> GroupsToUpdate;
 	for (FPrimitiveSceneInfo* const PrimitiveSceneInfo : PrimitiveSceneInfos)
@@ -4989,13 +5194,13 @@ void FScene::UpdateRayTracingGroupBounds_RemovePrimitives(const TArray<FPrimitiv
 	UpdateRayTracingGroupBounds(GroupsToUpdate);
 }
 
-template<typename RangeType>
-inline void FScene::UpdateRayTracingGroupBounds_UpdatePrimitives(const RangeType& InUpdatedTransforms)
+template<typename ValueType>
+inline void FScene::UpdateRayTracingGroupBounds_UpdatePrimitives(const Experimental::TRobinHoodHashMap<FPrimitiveSceneProxy*, ValueType>& InUpdatedTransforms)
 {
 	Experimental::TRobinHoodHashSet<FRayTracingCullingGroup*> GroupsToUpdate;
 	for (const auto& Transform : InUpdatedTransforms)
 	{
-		FPrimitiveSceneProxy* const PrimitiveSceneProxy = Transform.SceneInfo->Proxy;
+		FPrimitiveSceneProxy* const PrimitiveSceneProxy = Transform.Key;
 		const int32 RayTracingGroupId = PrimitiveSceneProxy->GetRayTracingGroupId();
 		const Experimental::FHashElementId GroupId = (RayTracingGroupId != -1) ? PrimitiveRayTracingGroups.FindId(RayTracingGroupId) : Experimental::FHashElementId();
 		if (GroupId.IsValid())
@@ -5033,42 +5238,50 @@ static inline bool IsPrimitiveRelevantToPathTracing(FPrimitiveSceneInfo* Primiti
 }
 
 
-void FScene::UpdateAllLightSceneInfos(FRDGBuilder& GraphBuilder)
+FLightSceneChangeSet FScene::UpdateAllLightSceneInfos(FRDGBuilder& GraphBuilder)
 {
 	SCOPED_NAMED_EVENT(FScene_UpdateAllLightSceneInfos, FColor::Orange);
 	struct FFLightSceneChangeSetAllocation
 	{
 		TArray<int32, SceneRenderingAllocator> RemovedLightIds;
+		TArray<int32, SceneRenderingAllocator> TransformUpdatedLightIds;
+		TArray<int32, SceneRenderingAllocator> ColorUpdatedLightIds;
 		TArray<int32, SceneRenderingAllocator> AddedLightIds;
-		FFLightSceneChangeSetAllocation(FSceneLightInfoUpdates&& InSceneLightInfoUpdates) : SceneLightInfoUpdates(MoveTemp(InSceneLightInfoUpdates)) {}
-		FSceneLightInfoUpdates SceneLightInfoUpdates;
 	};
 	// Allocate change set storage with graph builder lifetime such that we can safely pass it to async tasks.
-	FFLightSceneChangeSetAllocation& ChangeSet = *GraphBuilder.AllocObject<FFLightSceneChangeSetAllocation>(MoveTemp(*SceneLightInfoUpdates));
+	FFLightSceneChangeSetAllocation& ChangeSet = *GraphBuilder.AllocObject<FFLightSceneChangeSetAllocation>();
 
-#if DO_CHECK
-	FSceneLightInfoUpdates::FReadAccessScope ReadAccessScope(ChangeSet.SceneLightInfoUpdates);
-#endif
+	// Conservative pre-size the arrays (some updates are covered by Adds).
+	ChangeSet.RemovedLightIds.Reserve(SceneLightInfoUpdates->NumRemoves);
+	ChangeSet.TransformUpdatedLightIds.Reserve(SceneLightInfoUpdates->NumUpdates);
+	ChangeSet.ColorUpdatedLightIds.Reserve(SceneLightInfoUpdates->NumUpdates);
+	ChangeSet.AddedLightIds.Reserve(SceneLightInfoUpdates->NumAdds);
 
-	// Filter out removes & never-adds
-	ChangeSet.SceneLightInfoUpdates.ForEachCommand([&](const FUpdateLightCommand& UpdateLightCommand)
+	// Filter out removes & updates:
+	for (const auto& Element : SceneLightInfoUpdates->Commands)
 	{
-		int32 Id = UpdateLightCommand.GetSceneInfo()->Id;
-		if (UpdateLightCommand.IsDelete())
+		const FUpdateLightCommand& UpdateLightCommand = Element.Value;
+		int32 Id = UpdateLightCommand.LightSceneInfo->Id;
+		if (UpdateLightCommand.IsRemove())
 		{
-			if (Id != INDEX_NONE)
+			ChangeSet.RemovedLightIds.Add(Id);
+		}
+		else if (!UpdateLightCommand.IsAdd())
+		{
+			check(Id != INDEX_NONE);
+			if (UpdateLightCommand.bHasTransform)
 			{
-				ChangeSet.RemovedLightIds.Add(Id);
+				ChangeSet.TransformUpdatedLightIds.Add(Id);
 			}
-			else
+			if (UpdateLightCommand.bHasColor)
 			{
-				delete UpdateLightCommand.GetSceneInfo()->Proxy;
-				delete UpdateLightCommand.GetSceneInfo();
+				ChangeSet.ColorUpdatedLightIds.Add(Id);
 			}
 		}
-	});
+	}
 	// This can't access the scene light data if done async since it happens before the actual removals.
-	OnPreLightSceneInfoUpdate.Broadcast(GraphBuilder, FLightSceneChangeSet{ ChangeSet.RemovedLightIds, TConstArrayView<int32>(), &ChangeSet.SceneLightInfoUpdates });
+	OnPreLightSceneInfoUpdate.Broadcast(GraphBuilder, FLightSceneChangeSet{ ChangeSet.RemovedLightIds, TConstArrayView<int32>(), ChangeSet.TransformUpdatedLightIds, ChangeSet.ColorUpdatedLightIds });
+
 	// Batch process all light removes
 	for (int32 LightId : ChangeSet.RemovedLightIds)
 	{
@@ -5078,27 +5291,33 @@ void FScene::UpdateAllLightSceneInfos(FRDGBuilder& GraphBuilder)
 	}
 
 	// Process all light adds & updates
-	ChangeSet.SceneLightInfoUpdates.ForEachCommand(ESceneUpdateCommandFilter::AddedUpdated, [&](const FUpdateLightCommand& UpdateLightCommand)
+	for (const auto &Element : SceneLightInfoUpdates->Commands)
 	{
-		FLightSceneInfo* LightSceneInfo = UpdateLightCommand.GetSceneInfo();
+		const FUpdateLightCommand& UpdateLightCommand = Element.Value;
+		if (UpdateLightCommand.IsRemove())
+		{
+			continue;
+		}
+
+		FLightSceneInfo* LightSceneInfo = UpdateLightCommand.LightSceneInfo;
 		FScopeCycleCounter Context(LightSceneInfo->Proxy->GetStatId());
 
 		const int32 Id = LightSceneInfo->Id;
 		const bool bHasId = Id != INDEX_NONE;
 		check(bHasId == !UpdateLightCommand.IsAdd());
 		// Directly process updates.
-		if (FUpdateLightTransformParameters* TransformParameters = ChangeSet.SceneLightInfoUpdates.GetPayloadPtr<FUpdateLightTransformParameters>(UpdateLightCommand))
+		if (UpdateLightCommand.bHasTransform)
 		{
 			SCOPE_CYCLE_COUNTER(STAT_UpdateSceneLightTime);
 
-			UpdateLightTransform_RenderThread(Id, LightSceneInfo, *TransformParameters);
+			UpdateLightTransform_RenderThread(Id, LightSceneInfo, UpdateLightCommand.TransformParameters);
 		}
-
-		if (FUpdateLightColorParameters* ColorParameters = ChangeSet.SceneLightInfoUpdates.GetPayloadPtr<FUpdateLightColorParameters>(UpdateLightCommand))
+		if (UpdateLightCommand.bHasColor)
 		{
-			const FUpdateLightColorParameters& NewParameters = *ColorParameters;
 			SCOPE_CYCLE_COUNTER(STAT_UpdateSceneLightTime);
 
+			const FUpdateLightCommand::FColorParameters& NewParameters = UpdateLightCommand.ColorParameters;
+
 			// Mobile renderer:
 			// a light with no color/intensity can cause the light to be ignored when rendering.
 			// thus, lights that change state in this way must update the draw lists.
@@ -5131,10 +5350,14 @@ void FScene::UpdateAllLightSceneInfos(FRDGBuilder& GraphBuilder)
 			// Note: Id is set in AddLightSceneInfo_RenderThread so we must fetch it again
 			ChangeSet.AddedLightIds.Add(LightSceneInfo->Id);
 		}
-	});
+	}
 
-	OnPostLightSceneInfoUpdate.Broadcast(GraphBuilder, FLightSceneChangeSet{ ChangeSet.RemovedLightIds, ChangeSet.AddedLightIds, &ChangeSet.SceneLightInfoUpdates });
-	GPUScene.OnPostLightSceneInfoUpdate(GraphBuilder, FLightSceneChangeSet{ ChangeSet.RemovedLightIds, ChangeSet.AddedLightIds, &ChangeSet.SceneLightInfoUpdates });
+	OnPostLightSceneInfoUpdate.Broadcast(GraphBuilder, FLightSceneChangeSet{ ChangeSet.RemovedLightIds, ChangeSet.AddedLightIds, ChangeSet.TransformUpdatedLightIds, ChangeSet.ColorUpdatedLightIds });
+	GPUScene.OnPostLightSceneInfoUpdate(GraphBuilder, FLightSceneChangeSet{ ChangeSet.RemovedLightIds, ChangeSet.AddedLightIds, ChangeSet.TransformUpdatedLightIds, ChangeSet.ColorUpdatedLightIds });
+
+	SceneLightInfoUpdates->Reset();
+
+	return FLightSceneChangeSet{ ChangeSet.RemovedLightIds, ChangeSet.AddedLightIds, ChangeSet.TransformUpdatedLightIds, ChangeSet.ColorUpdatedLightIds };
 }
 
 template<class T>
@@ -5301,58 +5524,33 @@ void UpdateReflectionSceneData(FScene* Scene)
  */
 struct FSceneUpdateChangeSetStorage
 {
-
-	FSceneUpdateChangeSetStorage()
-		: PreUpdateChangeSet(FScenePreUpdateChangeSet {
-			TConstArrayView<FPersistentPrimitiveIndex>(RemovedPrimitiveIds),
-			TConstArrayView<FPrimitiveSceneInfo*>(RemovedPrimitiveSceneInfos),
-			PrimitiveUpdates
-			})
-		, PostUpdateChangeSet(FScenePostUpdateChangeSet {
-			TConstArrayView<FPersistentPrimitiveIndex>(AddedPrimitiveIds),
-			TConstArrayView<FPrimitiveSceneInfo*>(AddedPrimitiveSceneInfos),
-			PrimitiveUpdates
-		})
-	{
-	}
-	FSceneUpdateChangeSetStorage(FScenePrimitiveUpdates&& InPrimitiveUpdates)
-		: PrimitiveUpdates(MoveTemp(InPrimitiveUpdates))
-		, PreUpdateChangeSet(FScenePreUpdateChangeSet {
-			TConstArrayView<FPersistentPrimitiveIndex>(RemovedPrimitiveIds),
-			TConstArrayView<FPrimitiveSceneInfo*>(RemovedPrimitiveSceneInfos),
-			PrimitiveUpdates
-			})
-		, PostUpdateChangeSet(FScenePostUpdateChangeSet {
-			TConstArrayView<FPersistentPrimitiveIndex>(AddedPrimitiveIds),
-			TConstArrayView<FPrimitiveSceneInfo*>(AddedPrimitiveSceneInfos),
-			PrimitiveUpdates
-		})
-	{
-	}
-
-
 	TArray<FPersistentPrimitiveIndex, SceneRenderingAllocator> RemovedPrimitiveIds;
 	TArray<FPrimitiveSceneInfo*,SceneRenderingAllocator> RemovedPrimitiveSceneInfos;
 
+	TArray<FPersistentPrimitiveIndex, SceneRenderingAllocator> UpdatedPrimitiveIds;
+	TArray<FPrimitiveSceneInfo*, SceneRenderingAllocator> UpdatedPrimitiveSceneInfos;
+
 	TArray<FPersistentPrimitiveIndex, SceneRenderingAllocator> AddedPrimitiveIds;
 	TArray<FPrimitiveSceneInfo*, SceneRenderingAllocator> AddedPrimitiveSceneInfos;
-	FScenePrimitiveUpdates PrimitiveUpdates;
 
-	FScenePreUpdateChangeSet PreUpdateChangeSet;
-	FScenePostUpdateChangeSet PostUpdateChangeSet;
-
-	const FScenePreUpdateChangeSet& GetPreUpdateSet() 
+	FScenePreUpdateChangeSet GetPreUpdateSet() const 
 	{
-		PreUpdateChangeSet.RemovedPrimitiveIds = TConstArrayView<FPersistentPrimitiveIndex>(RemovedPrimitiveIds);
-		PreUpdateChangeSet.RemovedPrimitiveSceneInfos =	TConstArrayView<FPrimitiveSceneInfo*>(RemovedPrimitiveSceneInfos);
-		return PreUpdateChangeSet;
+		return FScenePreUpdateChangeSet {
+			TConstArrayView<FPersistentPrimitiveIndex>(RemovedPrimitiveIds),
+			TConstArrayView<FPrimitiveSceneInfo*>(RemovedPrimitiveSceneInfos),
+			TConstArrayView<FPersistentPrimitiveIndex>(UpdatedPrimitiveIds),
+			TConstArrayView<FPrimitiveSceneInfo*>(UpdatedPrimitiveSceneInfos)
+		};
 	}
 
-	const FScenePostUpdateChangeSet& GetPostUpdateSet() 
+	FScenePostUpdateChangeSet GetPostUpdateSet() const 
 	{
-		PostUpdateChangeSet.AddedPrimitiveIds = TConstArrayView<FPersistentPrimitiveIndex>(AddedPrimitiveIds);
-		PostUpdateChangeSet.AddedPrimitiveSceneInfos = TConstArrayView<FPrimitiveSceneInfo*>(AddedPrimitiveSceneInfos);
-		return PostUpdateChangeSet;
+		return FScenePostUpdateChangeSet {
+			TConstArrayView<FPersistentPrimitiveIndex>(AddedPrimitiveIds),
+			TConstArrayView<FPrimitiveSceneInfo*>(AddedPrimitiveSceneInfos),
+			TConstArrayView<FPersistentPrimitiveIndex>(UpdatedPrimitiveIds),
+			TConstArrayView<FPrimitiveSceneInfo*>(UpdatedPrimitiveSceneInfos)
+		};
 	}
 };
 
@@ -5397,44 +5595,18 @@ void FScene::Update(FRDGBuilder& GraphBuilder, const FUpdateParameters& Paramete
 	RDG_EVENT_SCOPE(GraphBuilder, "UpdateAllPrimitiveSceneInfos");
 
 	// Allocated with render graph lifetime, safe to reference from RDG tasks.
-	FSceneUpdateChangeSetStorage& SceneUpdateChangeSetStorage = *GraphBuilder.AllocObject<FSceneUpdateChangeSetStorage>(MoveTemp(PrimitiveUpdates));
-	TArray<FPrimitiveSceneInfo*, SceneRenderingAllocator> DeletedPrimitiveSceneInfos;
-	DeletedPrimitiveSceneInfos.Reserve(SceneUpdateChangeSetStorage.PrimitiveUpdates.NumCommands());
-	SceneUpdateChangeSetStorage.AddedPrimitiveSceneInfos.Reserve(SceneUpdateChangeSetStorage.PrimitiveUpdates.NumCommands());
-	SceneUpdateChangeSetStorage.RemovedPrimitiveSceneInfos.Reserve(SceneUpdateChangeSetStorage.PrimitiveUpdates.NumCommands());
-	SceneUpdateChangeSetStorage.RemovedPrimitiveIds.Reserve(SceneUpdateChangeSetStorage.PrimitiveUpdates.NumCommands());
-
-	// Retrieve the commonly used categories used in scene updates.
-	// 1. RemovedPrimitiveSceneInfos: Primitives that exist in the scene before the update and will be removed.
-	// 2. AddedPrimitiveSceneInfos: Primitives not in the scene before the update that will be added.
-	// 3. DeletedPrimitiveSceneInfos: Primitives that needs to be deleted.
-	SceneUpdateChangeSetStorage.PrimitiveUpdates.ForEachCommand([&](const FPrimitiveUpdateCommand& Cmd)
-	{
-		// Skip those that were added & removed in the same frame
-		if (Cmd.IsAdd() && !Cmd.IsDelete())
-		{
-			SceneUpdateChangeSetStorage.AddedPrimitiveSceneInfos.Add(Cmd.GetSceneInfo());
-		}
-		if (Cmd.IsDelete())
-		{
-			DeletedPrimitiveSceneInfos.Add(Cmd.GetSceneInfo());
-			// Skip those that were added & removed in the same frame
-			if (!Cmd.IsAdd())
-			{
-				SceneUpdateChangeSetStorage.RemovedPrimitiveSceneInfos.Add(Cmd.GetSceneInfo());
-				SceneUpdateChangeSetStorage.RemovedPrimitiveIds.Add(Cmd.GetPersistentId());
-			}
-		}
-	});
-
-	const TArray<FPrimitiveSceneInfo*, SceneRenderingAllocator> &AddedPrimitiveSceneInfos = SceneUpdateChangeSetStorage.AddedPrimitiveSceneInfos;
-	const TArray<FPrimitiveSceneInfo*, SceneRenderingAllocator> &RemovedPrimitiveSceneInfos = SceneUpdateChangeSetStorage.RemovedPrimitiveSceneInfos;
+	FSceneUpdateChangeSetStorage &SceneUpdateChangeSetStorage = *GraphBuilder.AllocObject<FSceneUpdateChangeSetStorage>();
+	SceneUpdateChangeSetStorage.RemovedPrimitiveIds.Reserve(RemovedPrimitiveSceneInfos.Num());
+	SceneUpdateChangeSetStorage.RemovedPrimitiveSceneInfos.Reserve(RemovedPrimitiveSceneInfos.Num());
 
 	TArray<FPrimitiveSceneInfo*, SceneRenderingAllocator> RemovedLocalPrimitiveSceneInfos;
 	RemovedLocalPrimitiveSceneInfos.Reserve(RemovedPrimitiveSceneInfos.Num());
 	for (FPrimitiveSceneInfo* PrimitiveSceneInfo : RemovedPrimitiveSceneInfos)
 	{
 		RemovedLocalPrimitiveSceneInfos.Add(PrimitiveSceneInfo);
+
+		SceneUpdateChangeSetStorage.RemovedPrimitiveIds.Add(PrimitiveSceneInfo->GetPersistentIndex());
+		SceneUpdateChangeSetStorage.RemovedPrimitiveSceneInfos.Add(PrimitiveSceneInfo);
 	}
 
 	TArray<FPrimitiveSceneInfo*, SceneRenderingAllocator> AddedLocalPrimitiveSceneInfos;
@@ -5451,10 +5623,61 @@ void FScene::Update(FRDGBuilder& GraphBuilder, const FUpdateParameters& Paramete
 	UpdateRayTracingGroupBounds_AddPrimitives(AddedPrimitiveSceneInfos);
 #endif
 
-	auto UpdatedInstances = SceneUpdateChangeSetStorage.PrimitiveUpdates.GetRangeView<FUpdateInstanceCommand>();
-	auto UpdatedTransforms = SceneUpdateChangeSetStorage.PrimitiveUpdates.GetRangeView<FUpdateTransformCommand>();
-	auto UpdatedInstanceCullDistance =  SceneUpdateChangeSetStorage.PrimitiveUpdates.GetRangeView<FUpdateInstanceCullDistanceData>();
-	auto OverridenPreviousTransforms =  SceneUpdateChangeSetStorage.PrimitiveUpdates.GetRangeView<FUpdateOverridePreviousTransformData>();	
+	{
+		SceneUpdateChangeSetStorage.UpdatedPrimitiveIds.Reserve(UpdatedInstances.Num() + UpdatedTransforms.Num());
+		SceneUpdateChangeSetStorage.UpdatedPrimitiveSceneInfos.Reserve(UpdatedInstances.Num() + UpdatedTransforms.Num());
+		// All updated instances must also before moving or re-allocating (TODO: filter out only those actually updated)
+		for (const auto& Instance : UpdatedInstances)
+		{
+			FPrimitiveSceneInfo* PrimitiveSceneInfo = Instance.Key->GetPrimitiveSceneInfo();
+
+			if (RemovedPrimitiveSceneInfos.Find(PrimitiveSceneInfo) != nullptr)
+			{
+				continue;
+			}
+
+			FPersistentPrimitiveIndex Id = Instance.Key->GetPrimitiveSceneInfo()->GetPersistentIndex();
+			// Note: may not be valid if it is also being added this in this update, in which case it will get processed in the post-update anyway.
+			if (Id.IsValid())
+			{
+				SceneUpdateChangeSetStorage.UpdatedPrimitiveIds.Add(Id);
+				SceneUpdateChangeSetStorage.UpdatedPrimitiveSceneInfos.Add(PrimitiveSceneInfo);
+			}
+		}
+		// As must all primitive updates, 
+		for (const auto& Transform : UpdatedTransforms)
+		{
+			FPrimitiveSceneInfo* PrimitiveSceneInfo = Transform.Key->GetPrimitiveSceneInfo();
+
+			if (UpdatedInstances.Find(Transform.Key) != nullptr)
+			{
+				continue;
+			}
+
+			if (RemovedPrimitiveSceneInfos.Find(PrimitiveSceneInfo) != nullptr)
+			{
+				continue;
+			}
+
+			FPersistentPrimitiveIndex Id = Transform.Key->GetPrimitiveSceneInfo()->GetPersistentIndex();
+			// Note: may not be valid if it is also being added this in this update, in which case it will get processed in the post-update anyway.
+			if (Id.IsValid())
+			{
+				SceneUpdateChangeSetStorage.UpdatedPrimitiveIds.Add(Id);
+				SceneUpdateChangeSetStorage.UpdatedPrimitiveSceneInfos.Add(PrimitiveSceneInfo);
+			}
+		}
+	}
+
+	TRACE_COUNTER_SET(ScenePrimitivesAdds, AddedPrimitiveSceneInfos.Num());
+	TRACE_COUNTER_SET(ScenePrimitivesRemoves, RemovedPrimitiveSceneInfos.Num());
+	TRACE_COUNTER_SET(ScenePrimitivesUpdatesTransforms, UpdatedInstances.Num());
+	TRACE_COUNTER_SET(ScenePrimitivesUpdatesInstances, UpdatedTransforms.Num());
+
+	// NOTE: We clear this early because IsPrimitiveBeingRemoved gets called from the CreateLightPrimitiveInteraction (to make sure that old primitives are not accessed) 
+	// we cannot safely kick off the AsyncCreateLightPrimitiveInteractionsTask before the RemovedPrimitiveSceneInfos has been cleared.
+	// TODO: this is probably not true anymore!
+	RemovedPrimitiveSceneInfos.Empty();
 	bool bAnySceneUpdatesQueued = RemovedLocalPrimitiveSceneInfos.Num() + AddedPrimitiveSceneInfos.Num() + UpdatedTransforms.Num() + UpdatedInstances.Num() > 0;
 	RemovedLocalPrimitiveSceneInfos.Sort(FPrimitiveArraySortKey());
 
@@ -5492,17 +5715,17 @@ void FScene::Update(FRDGBuilder& GraphBuilder, const FUpdateParameters& Paramete
 		// All updated instances must also before moving or re-allocating (TODO: filter out only those actually updated)
 		for (const auto& Instance : UpdatedInstances)
 		{
-			InvalidatingPrimitiveCollector.UpdatedTransform(Instance.SceneInfo);
+			InvalidatingPrimitiveCollector.UpdatedTransform(Instance.Key->GetPrimitiveSceneInfo());
 		}
 		// As must all primitive updates, 
 		for (const auto& Transform : UpdatedTransforms)
 		{
-			InvalidatingPrimitiveCollector.UpdatedTransform(Transform.SceneInfo);
+			InvalidatingPrimitiveCollector.UpdatedTransform(Transform.Key->GetPrimitiveSceneInfo());
 		}
 
 		for (const auto& CullDistance : UpdatedInstanceCullDistance)
 		{
-			InvalidatingPrimitiveCollector.UpdatedTransform(CullDistance.SceneInfo);
+			InvalidatingPrimitiveCollector.UpdatedTransform(CullDistance.Key->GetPrimitiveSceneInfo());
 		}
 
 		TRACE_INT_VALUE(TEXT("Shadow.Virtual.Cache.PreInvalidationInstances"), InvalidatingPrimitiveCollector.Instances.GetTotalNumInstances());
@@ -5510,6 +5733,9 @@ void FScene::Update(FRDGBuilder& GraphBuilder, const FUpdateParameters& Paramete
 	}
 
 	AddedLocalPrimitiveSceneInfos.Sort(FPrimitiveArraySortKey());
+
+	DeletedPrimitiveSceneInfos.Reserve(RemovedLocalPrimitiveSceneInfos.Num());
+
 	TArray<int32> RemovedPrimitiveIndices;
 	RemovedPrimitiveIndices.SetNumUninitialized(RemovedLocalPrimitiveSceneInfos.Num());
 
@@ -5715,6 +5941,8 @@ void FScene::Update(FRDGBuilder& GraphBuilder, const FUpdateParameters& Paramete
 
 				bNeedPathTracedInvalidation = bNeedPathTracedInvalidation || IsPrimitiveRelevantToPathTracing(PrimitiveSceneInfo);
 
+				DeletedPrimitiveSceneInfos.Emplace(PrimitiveSceneInfo);
+
 				const int32 PersistentIndex = PrimitiveSceneInfo->PersistentIndex.Index;
 				PersistentPrimitiveIdAllocator.Free(PersistentIndex);
 				PersistentPrimitiveIdToIndexMap[PersistentIndex] = INDEX_NONE;
@@ -5782,9 +6010,15 @@ void FScene::Update(FRDGBuilder& GraphBuilder, const FUpdateParameters& Paramete
 
 		for (const auto& UpdateInstance : UpdatedInstances)
 		{
-			FPrimitiveSceneProxy* PrimitiveSceneProxy = UpdateInstance.SceneInfo->Proxy;
+			FPrimitiveSceneProxy* PrimitiveSceneProxy = UpdateInstance.Key;
 			FScopeCycleCounter Context(PrimitiveSceneProxy->GetStatId());
-			FPrimitiveSceneInfo* PrimitiveSceneInfo = UpdateInstance.SceneInfo;
+			FPrimitiveSceneInfo* PrimitiveSceneInfo = PrimitiveSceneProxy->GetPrimitiveSceneInfo();
+			
+			// being added or deleted, skip update logic
+			if (DeletedPrimitiveSceneInfos.Contains(PrimitiveSceneInfo))
+			{
+				continue;
+			}
 
 			const FInstanceDataBufferHeader &InstanceDataBufferHeader = PrimitiveSceneInfo->GetInstanceDataHeader();
 			const bool bInstanceCountChanged = PrimitiveSceneInfo->GetNumInstanceSceneDataEntries() != InstanceDataBufferHeader.NumInstances;
@@ -5807,21 +6041,6 @@ void FScene::Update(FRDGBuilder& GraphBuilder, const FUpdateParameters& Paramete
 
 		PersistentPrimitiveIdAllocator.Consolidate();
 
-		SceneUpdateChangeSetStorage.PrimitiveUpdates.ForEachCommand(ESceneUpdateCommandFilter::Added, [&](FPrimitiveUpdateCommand& Cmd)
-		{
-			// Skip those that were added & removed in the same frame
-			if (Cmd.IsAdd() && !Cmd.IsDelete())
-			{
-				check(!Cmd.GetPersistentId().IsValid());
-				check(!Cmd.GetSceneInfo()->PersistentIndex.IsValid());
-				FPersistentPrimitiveIndex PersistentPrimitiveIndex{ PersistentPrimitiveIdAllocator.Allocate() };
-				// Note: propagate to the command such that is it valid for subsequent passes without having to hit the primitive scene info.
-				Cmd.SetPersistentId(PersistentPrimitiveIndex);
-				Cmd.GetSceneInfo()->PersistentIndex = PersistentPrimitiveIndex;
-			}
-		});
-		PersistentPrimitiveIdToIndexMap.SetNumUninitialized(GetMaxPersistentPrimitiveIndex());
-
 		if (AddedLocalPrimitiveSceneInfos.Num())
 		{
 			SCOPED_NAMED_EVENT(FScene_AddPrimitiveSceneInfos_ReserveMemory, FColor::Green);
@@ -5888,8 +6107,15 @@ void FScene::Update(FRDGBuilder& GraphBuilder, const FUpdateParameters& Paramete
 
 					const int32 SourceIndex = PrimitiveSceneProxies.Num() - 1;
 					PrimitiveSceneInfo->PackedIndex = SourceIndex;
-					check(PrimitiveSceneInfo->PersistentIndex.IsValid());
-					PersistentPrimitiveIdToIndexMap[PrimitiveSceneInfo->PersistentIndex.Index] = SourceIndex;
+					checkSlow(PrimitiveSceneInfo->PersistentIndex.Index == INDEX_NONE);
+					FPersistentPrimitiveIndex PersistentPrimitiveIndex{ PersistentPrimitiveIdAllocator.Allocate() };
+					PrimitiveSceneInfo->PersistentIndex = PersistentPrimitiveIndex;
+					// Ensure map is large enough
+					if (PersistentPrimitiveIndex.Index >= PersistentPrimitiveIdToIndexMap.Num())
+					{
+						PersistentPrimitiveIdToIndexMap.SetNumUninitialized(PersistentPrimitiveIndex.Index + 1);
+					}
+					PersistentPrimitiveIdToIndexMap[PersistentPrimitiveIndex.Index] = SourceIndex;
 				}
 			}
 
@@ -6053,18 +6279,19 @@ void FScene::Update(FRDGBuilder& GraphBuilder, const FUpdateParameters& Paramete
 		SCOPED_NAMED_EVENT(FScene_AddPrimitiveSceneInfos, FColor::Yellow);
 		SCOPE_CYCLE_COUNTER(STAT_UpdatePrimitiveTransformRenderThreadTime);
 
-
-		for (const auto& Item : UpdatedTransforms)
+		for (const auto& Transform : UpdatedTransforms)
 		{
-			FPrimitiveSceneProxy* PrimitiveSceneProxy = Item.SceneInfo->Proxy;
-
-			const FBoxSphereBounds& WorldBounds = Item.Payload.WorldBounds;
-			const FBoxSphereBounds& LocalBounds = Item.Payload.LocalBounds;
-			const FMatrix& LocalToWorld = Item.Payload.LocalToWorld;
-			const FVector& AttachmentRootPosition = Item.Payload.AttachmentRootPosition;
-
+			FPrimitiveSceneProxy* PrimitiveSceneProxy = Transform.Key;
+			if (DeletedPrimitiveSceneInfos.Contains(PrimitiveSceneProxy->GetPrimitiveSceneInfo()))
+			{
+				continue;
+			}
 			check(PrimitiveSceneProxy->GetPrimitiveSceneInfo()->PackedIndex != INDEX_NONE);
 
+			const FBoxSphereBounds& WorldBounds = Transform.Value.WorldBounds;
+			const FBoxSphereBounds& LocalBounds = Transform.Value.LocalBounds;
+			const FMatrix& LocalToWorld = Transform.Value.LocalToWorld;
+			const FVector& AttachmentRootPosition = Transform.Value.AttachmentRootPosition;
 			FScopeCycleCounter Context(PrimitiveSceneProxy->GetStatId());
 
 			FPrimitiveSceneInfo* PrimitiveSceneInfo = PrimitiveSceneProxy->GetPrimitiveSceneInfo();
@@ -6117,8 +6344,8 @@ void FScene::Update(FRDGBuilder& GraphBuilder, const FUpdateParameters& Paramete
 
 		for (const auto& Transform : OverridenPreviousTransforms)
 		{
-			FPrimitiveSceneInfo* PrimitiveSceneInfo = Transform.SceneInfo;
-			VelocityData.OverridePreviousTransform(PrimitiveSceneInfo->PrimitiveComponentId, Transform.Payload.Value);
+			FPrimitiveSceneInfo* PrimitiveSceneInfo = Transform.Key;
+			VelocityData.OverridePreviousTransform(PrimitiveSceneInfo->PrimitiveComponentId, Transform.Value);
 		}
 	}
 
@@ -6135,12 +6362,18 @@ void FScene::Update(FRDGBuilder& GraphBuilder, const FUpdateParameters& Paramete
 		SCOPED_NAMED_EVENT(FScene_UpdatePrimitiveInstances, FColor::Emerald);
 		SCOPE_CYCLE_COUNTER(STAT_UpdatePrimitiveInstanceRenderThreadTime);
 
-		for (const auto& Item : UpdatedInstances)
+		for (const auto& UpdateInstance : UpdatedInstances)
 		{
-			FPrimitiveSceneProxy* PrimitiveSceneProxy = Item.SceneInfo->Proxy;
-			FPrimitiveSceneInfo* PrimitiveSceneInfo = Item.SceneInfo;
-			const FUpdateInstanceCommand &UpdateInstance = Item.Payload;
+			FPrimitiveSceneProxy* PrimitiveSceneProxy = UpdateInstance.Key;
 			FScopeCycleCounter Context(PrimitiveSceneProxy->GetStatId());
+			FPrimitiveSceneInfo* PrimitiveSceneInfo = PrimitiveSceneProxy->GetPrimitiveSceneInfo();
+			
+			// being added or deleted, skip update logic
+			if (DeletedPrimitiveSceneInfos.Contains(PrimitiveSceneInfo))
+			{
+				continue;
+			}
+
 			
 			QueueFlushVirtualTexture(PrimitiveSceneInfo);
 
@@ -6174,7 +6407,7 @@ void FScene::Update(FRDGBuilder& GraphBuilder, const FUpdateParameters& Paramete
 #endif
 
 			// Update the Proxy's data.
-			PrimitiveSceneProxy->UpdateInstances_RenderThread(GraphBuilder.RHICmdList, UpdateInstance.WorldBounds, UpdateInstance.LocalBounds, UpdateInstance.StaticMeshBounds);
+			PrimitiveSceneProxy->UpdateInstances_RenderThread(GraphBuilder.RHICmdList, UpdateInstance.Value.WorldBounds, UpdateInstance.Value.LocalBounds, UpdateInstance.Value.StaticMeshBounds);
 
 			if (!RHISupportsVolumeTextures(GetFeatureLevel())
 				&& (PrimitiveSceneProxy->IsMovable() || PrimitiveSceneProxy->NeedsUnbuiltPreviewLighting() || PrimitiveSceneProxy->GetLightmapType() == ELightmapType::ForceVolumetric))
@@ -6337,9 +6570,11 @@ void FScene::Update(FRDGBuilder& GraphBuilder, const FUpdateParameters& Paramete
 	}
 	
 	SceneUpdateChangeSetStorage.AddedPrimitiveIds.Reserve(AddedPrimitiveSceneInfos.Num());
+	SceneUpdateChangeSetStorage.AddedPrimitiveSceneInfos.Reserve(AddedPrimitiveSceneInfos.Num());
 	for (FPrimitiveSceneInfo* PrimitiveSceneInfo : AddedPrimitiveSceneInfos)
 	{
 		SceneUpdateChangeSetStorage.AddedPrimitiveIds.Add(PrimitiveSceneInfo->GetPersistentIndex());
+		SceneUpdateChangeSetStorage.AddedPrimitiveSceneInfos.Add(PrimitiveSceneInfo);
 	}
 
 	SceneCullingUpdater.OnPostSceneUpdate(GraphBuilder, SceneUpdateChangeSetStorage.GetPostUpdateSet());
@@ -6517,11 +6752,19 @@ void FScene::Update(FRDGBuilder& GraphBuilder, const FUpdateParameters& Paramete
 #endif
 	}
 
-	for (const auto& Item : SceneUpdateChangeSetStorage.PrimitiveUpdates.GetRangeView<FUpdateCustomPrimitiveData>())
+	for (const auto& CustomParams : UpdatedCustomPrimitiveParams)
 	{
-		FPrimitiveSceneProxy* PrimitiveSceneProxy = Item.SceneInfo->Proxy;
+		FPrimitiveSceneProxy* PrimitiveSceneProxy = CustomParams.Key;
+		if (DeletedPrimitiveSceneInfos.Contains(PrimitiveSceneProxy->GetPrimitiveSceneInfo()))
+		{
+			continue;
+		}
+
 		FScopeCycleCounter Context(PrimitiveSceneProxy->GetStatId());
-		PrimitiveSceneProxy->CustomPrimitiveData = Item.Payload.Value;
+		PrimitiveSceneProxy->CustomPrimitiveData = CustomParams.Value;
+
+		// Ensure an update of primitive data before rendering
+		PrimitiveSceneProxy->GetPrimitiveSceneInfo()->MarkGPUStateDirty(EPrimitiveDirtyState::ChangedOther);
 	}
 
 	if (auto NaniteMaterialsUpdater = SceneExtensionsUpdaters.GetUpdaterPtr<Nanite::FMaterialsSceneExtension::FUpdater>())
@@ -6538,29 +6781,38 @@ void FScene::Update(FRDGBuilder& GraphBuilder, const FUpdateParameters& Paramete
 		SceneLODHierarchy.UpdateNodeSceneInfo(PrimitiveSceneInfo->PrimitiveComponentId, PrimitiveSceneInfo);
 	}
 
-	auto UpdatedAttachmentRoots = SceneUpdateChangeSetStorage.PrimitiveUpdates.GetRangeView<FUpdateAttachmentRootData>();
-	auto DistanceFieldSceneDataUpdates = SceneUpdateChangeSetStorage.PrimitiveUpdates.GetRangeView<FUpdateDistanceFieldSceneData>();
-	auto UpdatedOcclusionBoundsSlacks = SceneUpdateChangeSetStorage.PrimitiveUpdates.GetRangeView<FUpdateOcclusionBoundsSlacksData>();
-	auto UpdatedDrawDistance =  SceneUpdateChangeSetStorage.PrimitiveUpdates.GetRangeView<FUpdateDrawDistanceData>();
-
-	for (const auto& Item : UpdatedAttachmentRoots)
+	for (const auto& Attachments : UpdatedAttachmentRoots)
 	{
-		FPrimitiveSceneInfo* PrimitiveSceneInfo = Item.SceneInfo;
+		FPrimitiveSceneInfo* PrimitiveSceneInfo = Attachments.Key;
+		if (DeletedPrimitiveSceneInfos.Contains(PrimitiveSceneInfo))
+		{
+			continue;
+		}
+
 		PrimitiveSceneInfo->UnlinkAttachmentGroup();
-		PrimitiveSceneInfo->LightingAttachmentRoot = Item.Payload.Value;
+		PrimitiveSceneInfo->LightingAttachmentRoot = Attachments.Value;
 		PrimitiveSceneInfo->LinkAttachmentGroup();
 	}
 
-	for (const auto& Item : DistanceFieldSceneDataUpdates)
+	for (FPrimitiveSceneInfo* PrimitiveSceneInfo : DistanceFieldSceneDataUpdates)
 	{
-		DistanceFieldSceneData.UpdatePrimitive(Item.SceneInfo);
+		if (DeletedPrimitiveSceneInfos.Contains(PrimitiveSceneInfo))
+		{
+			continue;
+		}
+
+		DistanceFieldSceneData.UpdatePrimitive(PrimitiveSceneInfo);
 	}
 
 	for (const auto& OccSlackDelta : UpdatedOcclusionBoundsSlacks)
 	{
-		const FPrimitiveSceneProxy* SceneProxy = OccSlackDelta.SceneInfo->Proxy;
-		const FPrimitiveSceneInfo* SceneInfo = OccSlackDelta.SceneInfo;
-		float Delta = OccSlackDelta.Payload.Value;
+		const FPrimitiveSceneProxy* SceneProxy = OccSlackDelta.Key;
+		const FPrimitiveSceneInfo* SceneInfo = SceneProxy->GetPrimitiveSceneInfo();
+
+		if (DeletedPrimitiveSceneInfos.Contains(SceneInfo))
+		{
+			continue;
+		}
 
 		FBoxSphereBounds NewOccBounds;
 		if (SceneProxy->HasCustomOcclusionBounds())
@@ -6572,25 +6824,39 @@ void FScene::Update(FRDGBuilder& GraphBuilder, const FUpdateParameters& Paramete
 			NewOccBounds = SceneProxy->GetBounds();
 		}
 
-		PrimitiveOcclusionBounds[SceneInfo->PackedIndex] = NewOccBounds.ExpandBy(OCCLUSION_SLOP + Delta);
+		PrimitiveOcclusionBounds[SceneInfo->PackedIndex] = NewOccBounds.ExpandBy(OCCLUSION_SLOP + OccSlackDelta.Value);
 	}
 
-	for (const auto& CullDistance : UpdatedInstanceCullDistance)
+	for (auto& CullDistance : UpdatedInstanceCullDistance)
 	{
-		FPrimitiveSceneProxy* SceneProxy = CullDistance.SceneInfo->Proxy;
-		float StartCullDistance = CullDistance.Payload.Value.X;
-		float EndCullDistance = CullDistance.Payload.Value.Y;
+		FPrimitiveSceneProxy* SceneProxy = CullDistance.Key;
+		FPrimitiveSceneInfo* SceneInfo = SceneProxy->GetPrimitiveSceneInfo();
+
+		if (DeletedPrimitiveSceneInfos.Contains(SceneInfo))
+		{
+			continue;
+		}
+
+		float StartCullDistance = CullDistance.Value.X;
+		float EndCullDistance = CullDistance.Value.Y;
 		
 		SceneProxy->SetInstanceCullDistance_RenderThread(StartCullDistance, EndCullDistance);
+		SceneInfo->MarkGPUStateDirty(EPrimitiveDirtyState::ChangedOther);
 	}
 
-	for (const auto& DrawDistance : UpdatedDrawDistance)
+	for (auto& DrawDistance : UpdatedDrawDistance)
 	{
-		FPrimitiveSceneProxy* SceneProxy = DrawDistance.SceneInfo->Proxy;
-		FPrimitiveSceneInfo* SceneInfo = DrawDistance.SceneInfo;
-		float MinDrawDistance = DrawDistance.Payload.Value.X;
-		float MaxDrawDistance = DrawDistance.Payload.Value.Y;
-		float VirtualTextureMaxDrawDistance = DrawDistance.Payload.Value.Z;
+		FPrimitiveSceneProxy* SceneProxy = DrawDistance.Key;
+		FPrimitiveSceneInfo* SceneInfo = SceneProxy->GetPrimitiveSceneInfo();
+
+		if (DeletedPrimitiveSceneInfos.Contains(SceneInfo))
+		{
+			continue;
+		}
+
+		float MinDrawDistance = DrawDistance.Value.X;
+		float MaxDrawDistance = DrawDistance.Value.Y;
+		float VirtualTextureMaxDrawDistance = DrawDistance.Value.Z;
 
 		SceneProxy->SetDrawDistance_RenderThread(MinDrawDistance, MaxDrawDistance, VirtualTextureMaxDrawDistance);
 
@@ -6681,6 +6947,16 @@ void FScene::Update(FRDGBuilder& GraphBuilder, const FUpdateParameters& Paramete
 	{
 		InvalidatePathTracedOutput();
 	}
+	UpdatedAttachmentRoots.Empty();
+	UpdatedTransforms.Empty();
+	UpdatedInstances.Empty();
+	UpdatedCustomPrimitiveParams.Empty();
+	OverridenPreviousTransforms.Empty();
+	UpdatedOcclusionBoundsSlacks.Empty();
+	UpdatedInstanceCullDistance.Empty();
+	UpdatedDrawDistance.Empty();
+	DistanceFieldSceneDataUpdates.Empty();
+	AddedPrimitiveSceneInfos.Empty();
 	LevelCommands.Empty();
 
 	LastUpdateFrameCounter = GFrameCounterRenderThread;
@@ -6700,6 +6976,12 @@ void FScene::Update(FRDGBuilder& GraphBuilder, const FUpdateParameters& Paramete
 #endif
 }
 
+bool FScene::IsPrimitiveBeingRemoved(FPrimitiveSceneInfo* PrimitiveSceneInfo) const
+{
+	check(IsInParallelRenderingThread() || IsInRenderingThread());
+	return RemovedPrimitiveSceneInfos.Find(PrimitiveSceneInfo) != nullptr;
+}
+
 bool FScene::ShouldRenderSkylightInBasePass(bool bIsTranslucent) const
 {
 	if (IsMobilePlatform(GetShaderPlatform()))
@@ -7026,20 +7308,20 @@ void UpdateStaticMeshesForMaterials(const TArray<const FMaterial*>& MaterialReso
 	UE::RenderCommandPipe::FSyncScope SyncScope;
 
 	ENQUEUE_RENDER_COMMAND(FUpdateStaticMeshesForMaterials)(
-	[UsedPrimitives = MoveTemp(UsedPrimitives)](FRHICommandListImmediate& RHICmdList) mutable
-	{
-		// Defer the caching until the next render tick, to make sure that all render components queued
-		// for re-creation are processed. Otherwise, we may end up caching mesh commands from stale data.
-		for (auto& SceneInfos: UsedPrimitives)
+		[UsedPrimitives = MoveTemp(UsedPrimitives)](FRHICommandListImmediate& RHICmdList) mutable
 		{
-			SceneInfos.Key->UpdateAllPrimitiveSceneInfos(RHICmdList);
-		}
-		for (auto& SceneInfos : UsedPrimitives)
-		{
-			TArray<FPrimitiveSceneInfo*>& SceneInfoArray = SceneInfos.Value;
-			FPrimitiveSceneInfo::UpdateStaticMeshes(SceneInfos.Key, SceneInfoArray, EUpdateStaticMeshFlags::AllCommands, false);
-		}
-	});
+			// Defer the caching until the next render tick, to make sure that all render components queued
+			// for re-creation are processed. Otherwise, we may end up caching mesh commands from stale data.
+			for (auto& SceneInfos: UsedPrimitives)
+			{
+				SceneInfos.Key->UpdateAllPrimitiveSceneInfos(RHICmdList);
+			}
+			for (auto& SceneInfos : UsedPrimitives)
+			{
+				TArray<FPrimitiveSceneInfo*>& SceneInfoArray = SceneInfos.Value;
+				FPrimitiveSceneInfo::UpdateStaticMeshes(SceneInfos.Key, SceneInfoArray, EUpdateStaticMeshFlags::AllCommands, false);
+			}
+		});
 }
 
 void FRendererModule::UpdateStaticDrawListsForMaterials(const TArray<const FMaterial*>& Materials)
diff --git a/Engine/Source/Runtime/Renderer/Private/SceneCulling/SceneCulling.cpp b/Engine/Source/Runtime/Renderer/Private/SceneCulling/SceneCulling.cpp
--- a/Engine/Source/Runtime/Renderer/Private/SceneCulling/SceneCulling.cpp
+++ b/Engine/Source/Runtime/Renderer/Private/SceneCulling/SceneCulling.cpp
@@ -2420,17 +2420,21 @@ public:
 		{
 			SCOPED_NAMED_EVENT(SceneCulling_Post_UpdateInstances, FColor::Emerald);
 
-			ScenePostUpdateData.PrimitiveUpdates.ForEachUpdateCommand(ESceneUpdateCommandFilter::AddedUpdated, EPrimitiveUpdateDirtyFlags::AllCulling, [&](const FPrimitiveUpdateCommand& Cmd)
+			// [transform-] updated primitives instances & primitives with updated instances 
+			for (int32 Index = 0; Index < ScenePostUpdateData.UpdatedPrimitiveIds.Num(); ++Index)
 			{
-				if (Cmd.IsAdd())
-				{
-					AddInstances(Cmd.GetPersistentId(), Cmd.GetSceneInfo());
-				}
-				else
-				{
-					UpdateInstances(Cmd.GetPersistentId(), Cmd.GetSceneInfo());
-				}
-			});
+				UpdateInstances(ScenePostUpdateData.UpdatedPrimitiveIds[Index], ScenePostUpdateData.UpdatedPrimitiveSceneInfos[Index]);
+			}	
+		}
+
+		{
+			SCOPED_NAMED_EVENT(SceneCulling_Post_AddInstances, FColor::Emerald);
+
+			// Next process all added and added ones.
+			for (int32 Index = 0; Index < ScenePostUpdateData.AddedPrimitiveIds.Num(); ++Index)
+			{
+				AddInstances(ScenePostUpdateData.AddedPrimitiveIds[Index], ScenePostUpdateData.AddedPrimitiveSceneInfos[Index]);
+			}
 		}
 		FinalizeTempCellsAndUncullable();
 
diff --git a/Engine/Source/Runtime/Renderer/Private/ScenePrivate.h b/Engine/Source/Runtime/Renderer/Private/ScenePrivate.h
--- a/Engine/Source/Runtime/Renderer/Private/ScenePrivate.h
+++ b/Engine/Source/Runtime/Renderer/Private/ScenePrivate.h
@@ -59,7 +59,6 @@
 #include "LightFunctionAtlas.h"
 #include "SceneExtensions.h"
 #include "HeterogeneousVolumes/HeterogeneousVolumes.h"
-#include "SceneUpdateCommandQueue.h"
 
 /** Factor by which to grow occlusion tests **/
 #define OCCLUSION_SLOP (1.0f)
@@ -95,8 +94,54 @@ class FSparseVolumeTextureViewerSceneProxy;
 class FExponentialHeightFogSceneInfo;
 class FStaticMeshBatch;
 class FShadowScene;
+class FSceneLightInfoUpdates;
 class FSceneCulling;
 
+/**
+ * Describes all light modifications to the scene by recording the light scene IDs.
+ * TODO: If needed, we could add a reference to the FLightUpdates (which contains the commands) since this would enable systems to consume out the delta updates as they come in.
+ *       If this is useful we must ensure FLightUpdates are kept alive as long as any async tasks might require.
+ */
+struct FLightSceneChangeSet
+{
+	// IDs of all lights before they were removed, IDs in this array may not be valid at all times when the change-set is used (depends on whether the callback site is before or after the given lights are removed from the scene).
+	TConstArrayView<int32> RemovedLightIds;
+	// IDs of all lights added to the scene, only available after all lights are added to the scene, may contain the same ID's as removed, as they may be reused.
+	TConstArrayView<int32> AddedLightIds;
+	// IDs of updated lights, does not contain any from the above, since 'add' implies the update of all aspects and 'remove' implies cancellation of all updates. 
+	// The updated arrays are not disjoint as a light may have both types of update applied.
+	TConstArrayView<int32> TransformUpdatedLightIds;
+	TConstArrayView<int32> ColorUpdatedLightIds;
+};
+
+/**
+ * Change set that is valid before removes are processed and the scene data modified.
+ * The referenced arrays have RDG life-time and can be safely used in RDG tasks.
+ * However, the referenced data (primitive/proxy) and meaning of the persistent ID is not generally valid past the call in which this is passed. 
+ * Thus, care need to be excercised.
+ */
+class FScenePreUpdateChangeSet
+{
+public:
+	TConstArrayView<FPersistentPrimitiveIndex> RemovedPrimitiveIds;
+	TConstArrayView<FPrimitiveSceneInfo*> RemovedPrimitiveSceneInfos;
+	TConstArrayView<FPersistentPrimitiveIndex> UpdatedPrimitiveIds;
+	TConstArrayView<FPrimitiveSceneInfo*> UpdatedPrimitiveSceneInfos;
+};
+
+/**
+ * Change set that is valid before after adds are processed and the scene data is modified.
+ * The referenced arrays have RDG life-time and can be safely used in RDG tasks.
+ */
+class FScenePostUpdateChangeSet
+{
+public:
+	TConstArrayView<FPersistentPrimitiveIndex> AddedPrimitiveIds;
+	TConstArrayView<FPrimitiveSceneInfo*> AddedPrimitiveSceneInfos;
+	TConstArrayView<FPersistentPrimitiveIndex> UpdatedPrimitiveIds;
+	TConstArrayView<FPrimitiveSceneInfo*> UpdatedPrimitiveSceneInfos;
+};
+
 /** Holds information about a single primitive's occlusion. */
 class FPrimitiveOcclusionHistory
 {
@@ -2840,173 +2885,6 @@ private:
 	FLumenSceneDataMap::TConstIterator NextSceneData;
 };
 
-/**
- * Definitions of primitive scene update commands.
- */
-
-enum class EPrimitiveUpdateDirtyFlags : uint32 
-{
-	None			=	0u,
-	/** The Transform is modified by this command. */
-	Transform		=	1u << 0,
-	/** The (any) instance data is modified by this command. */
-	InstanceData	=	1u << 1,
-	/** 
-	 * The culling bounds are modified by this command. 
-	 * This means the bounds (instance, primitive or both) as used in the culling and should not be updated for any other case. 
-	 * Thus, needs to be set for transform updates of all kinds.
-	 */
-	CullingBounds	=	1u << 2,
-	/** 
-	 * Culling distances or similar, affecting culling logic, but NOT the bounds.
-	 */
-	CullingLogic	=	1u << 3,
-	/** Any state that either makes its way into GPU-Scene or the per primitive UB */
-	GPUState		=	1u << 4,
-	/** All culling-affecting changes */
-	AllCulling		= CullingBounds | CullingLogic,
-	All				=	GPUState |  Transform | InstanceData | CullingBounds | CullingLogic,
-};
-ENUM_CLASS_FLAGS(EPrimitiveUpdateDirtyFlags);
-
-enum class EPrimitiveUpdateId : uint32 
-{
-	UpdateTransform,
-	UpdateInstance,
-	UpdateAttachmentRoot,
-	CustomPrimitiveData,
-	OcclusionBoundsSlacks,  
-	InstanceCullDistance,  
-	DrawDistance,
-	DistanceFieldScene,
-	OverridePreviousTransform,
-	MAX
-};
-
-using FScenePrimitiveUpdates = TSceneUpdateCommandQueue<FPrimitiveSceneInfo, EPrimitiveUpdateDirtyFlags, EPrimitiveUpdateId>;
-
-using FPrimitiveUpdateCommand = FScenePrimitiveUpdates::FUpdateCommand;
-
-template <EPrimitiveUpdateId InId, EPrimitiveUpdateDirtyFlags InDirtyFlags>
-using TPrimitiveUpdatePayloadBase = FScenePrimitiveUpdates::TPayloadBase<InId, InDirtyFlags>;
-
-struct FUpdateTransformCommand : public TPrimitiveUpdatePayloadBase<EPrimitiveUpdateId::UpdateTransform, 
-	EPrimitiveUpdateDirtyFlags::GPUState | EPrimitiveUpdateDirtyFlags::Transform | EPrimitiveUpdateDirtyFlags::CullingBounds>
-{
-	FBoxSphereBounds WorldBounds;
-	FBoxSphereBounds LocalBounds; 
-	FMatrix LocalToWorld; 
-	FVector AttachmentRootPosition;
-};
-
-struct FUpdateInstanceCommand : public TPrimitiveUpdatePayloadBase<EPrimitiveUpdateId::UpdateInstance, 
-	EPrimitiveUpdateDirtyFlags::GPUState | EPrimitiveUpdateDirtyFlags::Transform | EPrimitiveUpdateDirtyFlags::CullingBounds | EPrimitiveUpdateDirtyFlags::InstanceData>
-{
-	FPrimitiveSceneProxy* PrimitiveSceneProxy{ nullptr };
-	FBoxSphereBounds WorldBounds;
-	FBoxSphereBounds LocalBounds;
-	FBoxSphereBounds StaticMeshBounds;
-};
-
-
-/**
- * Helper for the update payloads that contain a single payload value.
- */
-template <typename InPayloadDataType, EPrimitiveUpdateId InId, EPrimitiveUpdateDirtyFlags InPrimitiveDirtyFlags>
-struct TSingleValuePrimitiveUpdatePayload : public TPrimitiveUpdatePayloadBase<InId, InPrimitiveDirtyFlags>
-{
-	TSingleValuePrimitiveUpdatePayload(const InPayloadDataType &InValue) : Value(InValue) {}
-	InPayloadDataType Value;
-};
-
-using FUpdateAttachmentRootData = TSingleValuePrimitiveUpdatePayload<FPrimitiveComponentId, EPrimitiveUpdateId::UpdateAttachmentRoot,
-	EPrimitiveUpdateDirtyFlags::None>; // No GPU side effect (?).
-using FUpdateCustomPrimitiveData = TSingleValuePrimitiveUpdatePayload<FCustomPrimitiveData, EPrimitiveUpdateId::CustomPrimitiveData,
-	EPrimitiveUpdateDirtyFlags::GPUState>; // Needs upload
-using FUpdateOcclusionBoundsSlacksData = TSingleValuePrimitiveUpdatePayload<float, EPrimitiveUpdateId::OcclusionBoundsSlacks,  
-	EPrimitiveUpdateDirtyFlags::None>; // Only affects primitive occlusion.
-using FUpdateInstanceCullDistanceData = TSingleValuePrimitiveUpdatePayload<FVector2f, EPrimitiveUpdateId::InstanceCullDistance,  
-	EPrimitiveUpdateDirtyFlags::GPUState | EPrimitiveUpdateDirtyFlags::CullingLogic>; // Affects GPU culling?
-using FUpdateDrawDistanceData = TSingleValuePrimitiveUpdatePayload<FVector3f, EPrimitiveUpdateId::DrawDistance,  
-	EPrimitiveUpdateDirtyFlags::CullingLogic>; // Only affects CPU culling.
-using FUpdateDistanceFieldSceneData = TPrimitiveUpdatePayloadBase<EPrimitiveUpdateId::DistanceFieldScene,  
-	EPrimitiveUpdateDirtyFlags::None>; // Only affects DF scene rep - candidate for using abstract type.
-using FUpdateOverridePreviousTransformData = TSingleValuePrimitiveUpdatePayload<FMatrix, EPrimitiveUpdateId::OverridePreviousTransform,  
-	EPrimitiveUpdateDirtyFlags::GPUState>;// Overrides the previous transform, which needs to be propagated to the GPU, but otherwise does not change anything on its own.
-
-/**
- * Change set that is valid before removes are processed and the scene data modified.
- * The referenced arrays have RDG life-time and can be safely used in RDG tasks.
- * However, the referenced data (primitive/proxy) and meaning of the persistent ID is not generally valid past the call in which this is passed. 
- * Thus, care need to be excercised.
- */
-class FScenePreUpdateChangeSet
-{
-public:
-	TConstArrayView<FPersistentPrimitiveIndex> RemovedPrimitiveIds;
-	TConstArrayView<FPrimitiveSceneInfo*> RemovedPrimitiveSceneInfos;
-	const FScenePrimitiveUpdates &PrimitiveUpdates;
-};
-
-/**
- * Change set that is valid before after adds are processed and the scene data is modified.
- * The referenced arrays have RDG life-time and can be safely used in RDG tasks.
- */
-class FScenePostUpdateChangeSet
-{
-public:
-	TConstArrayView<FPersistentPrimitiveIndex> AddedPrimitiveIds;
-	TConstArrayView<FPrimitiveSceneInfo*> AddedPrimitiveSceneInfos;
-	const FScenePrimitiveUpdates &PrimitiveUpdates;
-};
-
-/**
- * Definitions for light scene updates.
- */
-
-enum class ELightDirtyFlags : uint32 
-{
-	None			=	0u
-};
-ENUM_CLASS_FLAGS(ELightDirtyFlags);
-
-enum class ELightUpdateId : uint32 
-{
-	Transform,
-	Color,
-	MAX
-};
-
-using FSceneLightInfoUpdates = TSceneUpdateCommandQueue<FLightSceneInfo, ELightDirtyFlags, ELightUpdateId>;
-using FUpdateLightCommand = FSceneLightInfoUpdates::FUpdateCommand;
-
-struct FUpdateLightTransformParameters : public FSceneLightInfoUpdates::TPayloadBase<ELightUpdateId::Transform, ELightDirtyFlags::None>
-{
-	FMatrix LightToWorld;
-	FVector4 Position;
-};
-
-struct FUpdateLightColorParameters: public FSceneLightInfoUpdates::TPayloadBase<ELightUpdateId::Color, ELightDirtyFlags::None>
-{
-	FLinearColor NewColor;
-	float NewIndirectLightingScale;
-	float NewVolumetricScatteringIntensity;
-};
-
-/**
- * Describes all light modifications to the scene by recording the light scene IDs.
- * TODO: If needed, we could add a reference to the FLightUpdates (which contains the commands) since this would enable systems to consume out the delta updates as they come in.
- *       If this is useful we must ensure FLightUpdates are kept alive as long as any async tasks might require.
- */
-struct FLightSceneChangeSet
-{
-	// IDs of all lights before they were removed, IDs in this array may not be valid at all times when the change-set is used (depends on whether the callback site is before or after the given lights are removed from the scene).
-	TConstArrayView<int32> RemovedLightIds;
-	// IDs of all lights added to the scene, only available after all lights are added to the scene, may contain the same ID's as removed, as they may be reused.
-	TConstArrayView<int32> AddedLightIds;
-	FSceneLightInfoUpdates *SceneLightInfoUpdates = nullptr;
-};
-
 /** 
  * Renderer scene which is private to the renderer module.
  * Ordinarily this is the renderer version of a UWorld, but an FScene can be created for previewing in editors which don't have a UWorld as well.
@@ -3846,6 +3724,8 @@ public:
 		return &CachedShadowMapDatas[ShadowMapIndex];
 	}
 
+	bool IsPrimitiveBeingRemoved(FPrimitiveSceneInfo* PrimitiveSceneInfo) const;
+
 	/**
 	 * Maximum used persistent Primitive Index, use to size arrays that store primitive data indexed by FPrimitiveSceneInfo::PersistentIndex.
 	 * Only changes during UpdateAllPrimitiveSceneInfos.
@@ -4015,7 +3895,7 @@ private:
 	 * Removes a primitive from the scene.  Called in the rendering thread by RemovePrimitive.
 	 * @param PrimitiveSceneInfo - The primitive being removed.
 	 */
-	void RemovePrimitiveSceneInfo_RenderThread(FPrimitiveSceneInfo* PrimitiveSceneInfo);
+	bool RemovePrimitiveSceneInfo_RenderThread(FPrimitiveSceneInfo* PrimitiveSceneInfo);
 
 	/** Updates a primitive's transform, called on the rendering thread. */
 	void UpdatePrimitiveTransform_RenderThread(FPrimitiveSceneProxy* PrimitiveSceneProxy, const FBoxSphereBounds& WorldBounds, const FBoxSphereBounds& LocalBounds, const FMatrix& LocalToWorld, const FVector& OwnerPosition, const TOptional<FTransform>& PreviousTransform);
@@ -4083,16 +3963,11 @@ private:
 	void ProcessAtmosphereLightAddition_RenderThread(FLightSceneInfo* LightSceneInfo);
 
 	/**
-	 * Process all scene updates for lights.
+	 * Process all scene updates for lights, returns the change-set, which references arrays allocated with a RDG builder life-time.
 	 */
-	void UpdateAllLightSceneInfos(FRDGBuilder& GraphBuilder);
+	FLightSceneChangeSet UpdateAllLightSceneInfos(FRDGBuilder& GraphBuilder);
 
 private:
-	template <typename UpdatePayloadType>
-	void UpdatePrimitiveInternal(FPrimitiveSceneProxy* SceneProxy, UpdatePayloadType &&InUpdatePayload);
-
-	template <typename UpdatePayloadType>
-	void UpdateLightInternal(FLightSceneProxy* LightSceneProxy, UpdatePayloadType &&InUpdatePayload);
 
 	/**
 	 * Update tracked scene state for cached CSM shadows
@@ -4101,12 +3976,28 @@ private:
 
 	FString FullWorldName;
 #if RHI_RAYTRACING
-	void UpdateRayTracingGroupBounds_AddPrimitives(const TArray<FPrimitiveSceneInfo*, SceneRenderingAllocator>& PrimitiveSceneInfos);
-	void UpdateRayTracingGroupBounds_RemovePrimitives(const TArray<FPrimitiveSceneInfo*, SceneRenderingAllocator>& PrimitiveSceneInfos);
-	template<typename RangeType>
-	inline void UpdateRayTracingGroupBounds_UpdatePrimitives(const RangeType& UpdatedTransforms);
+	void UpdateRayTracingGroupBounds_AddPrimitives(const Experimental::TRobinHoodHashSet<FPrimitiveSceneInfo*>& PrimitiveSceneInfos);
+	void UpdateRayTracingGroupBounds_RemovePrimitives(const Experimental::TRobinHoodHashSet<FPrimitiveSceneInfo*>& PrimitiveSceneInfos);
+	template<typename ValueType>
+	inline void UpdateRayTracingGroupBounds_UpdatePrimitives(const Experimental::TRobinHoodHashMap<FPrimitiveSceneProxy*, ValueType>& UpdatedTransforms);
 #endif
 
+	struct FUpdateTransformCommand
+	{
+		FBoxSphereBounds WorldBounds;
+		FBoxSphereBounds LocalBounds; 
+		FMatrix LocalToWorld; 
+		FVector AttachmentRootPosition;
+	};
+
+	struct FUpdateInstanceCommand
+	{
+		FPrimitiveSceneProxy* PrimitiveSceneProxy{ nullptr };
+		FBoxSphereBounds WorldBounds;
+		FBoxSphereBounds LocalBounds;
+		FBoxSphereBounds StaticMeshBounds;
+	};
+
 	void UpdatePrimitiveInstances(FUpdateInstanceCommand& UpdateParams);
 
 	struct FLevelCommand
@@ -4121,8 +4012,19 @@ private:
 		EOp Op;
 	};
 
-	FScenePrimitiveUpdates PrimitiveUpdates;
+	Experimental::TRobinHoodHashMap<FPrimitiveSceneInfo*, FPrimitiveComponentId> UpdatedAttachmentRoots;
+	Experimental::TRobinHoodHashMap<FPrimitiveSceneProxy*, FCustomPrimitiveData> UpdatedCustomPrimitiveParams;
+	Experimental::TRobinHoodHashMap<FPrimitiveSceneProxy*, FUpdateTransformCommand> UpdatedTransforms;
+	Experimental::TRobinHoodHashMap<FPrimitiveSceneProxy*, FUpdateInstanceCommand> UpdatedInstances;
+	Experimental::TRobinHoodHashMap<FPrimitiveSceneInfo*, FMatrix> OverridenPreviousTransforms;
+	Experimental::TRobinHoodHashMap<const FPrimitiveSceneProxy*, float> UpdatedOcclusionBoundsSlacks;
+	Experimental::TRobinHoodHashMap<FPrimitiveSceneProxy*, FVector2f> UpdatedInstanceCullDistance;
+	Experimental::TRobinHoodHashMap<FPrimitiveSceneProxy*, FVector3f> UpdatedDrawDistance;
+	Experimental::TRobinHoodHashSet<FPrimitiveSceneInfo*> AddedPrimitiveSceneInfos;
+	Experimental::TRobinHoodHashSet<FPrimitiveSceneInfo*> RemovedPrimitiveSceneInfos;
+	Experimental::TRobinHoodHashSet<FPrimitiveSceneInfo*> DistanceFieldSceneDataUpdates;
 	TArray<FLevelCommand> LevelCommands;
+	TSet<FPrimitiveSceneInfo*> DeletedPrimitiveSceneInfos;
 
 	UE::Tasks::FTask CreateLightPrimitiveInteractionsTask;
 	UE::Tasks::FTask GPUSkinCacheTask;
diff --git a/Engine/Source/Runtime/Renderer/Private/SceneUpdateCommandQueue.h b/Engine/Source/Runtime/Renderer/Private/SceneUpdateCommandQueue.h
deleted file mode 100644
diff --git a/Engine/Source/Runtime/Renderer/Private/Shadows/ShadowScene.cpp b/Engine/Source/Runtime/Renderer/Private/Shadows/ShadowScene.cpp
--- a/Engine/Source/Runtime/Renderer/Private/Shadows/ShadowScene.cpp
+++ b/Engine/Source/Runtime/Renderer/Private/Shadows/ShadowScene.cpp
@@ -38,7 +38,7 @@ void FShadowScene::PostLightsUpdate(FRDGBuilder& GraphBuilder, const FLightScene
 	// don't spawn async work for no good reason.
 	constexpr int32 kMinWorkSizeForAsync = 64;
 
-	int32 WorkSize = LightSceneChangeSet.SceneLightInfoUpdates->NumCommands();
+	int32 WorkSize = LightSceneChangeSet.AddedLightIds.Num() + LightSceneChangeSet.RemovedLightIds.Num() + LightSceneChangeSet.TransformUpdatedLightIds.Num();
 
 	// Don't sync, or kick off a new job if there is no work to do.
 	if (WorkSize > 0)
@@ -78,10 +78,7 @@ void FShadowScene::PostLightsUpdate(FRDGBuilder& GraphBuilder, const FLightScene
 			};
 
 			Algo::ForEach(LightSceneChangeSet.AddedLightIds, UpdateLight);
-			for (const auto& Item : LightSceneChangeSet.SceneLightInfoUpdates->GetRangeView<FUpdateLightTransformParameters>())
-			{ 
-				UpdateLight(Item.SceneInfo->Id);
-			}
+			Algo::ForEach(LightSceneChangeSet.TransformUpdatedLightIds, UpdateLight);
 
 		}, WorkSize > kMinWorkSizeForAsync);
 	}
diff --git a/Engine/Source/Runtime/Renderer/Private/SplineMeshSceneResources.cpp b/Engine/Source/Runtime/Renderer/Private/SplineMeshSceneResources.cpp
--- a/Engine/Source/Runtime/Renderer/Private/SplineMeshSceneResources.cpp
+++ b/Engine/Source/Runtime/Renderer/Private/SplineMeshSceneResources.cpp
@@ -338,20 +338,23 @@ void FSplineMeshSceneUpdater::PostSceneUpdate(FRDGBuilder& GraphBuilder, const F
 		}
 	};
 
-	ChangeSet.PrimitiveUpdates.ForEachUpdateCommand(ESceneUpdateCommandFilter::AddedUpdated, FUpdateInstanceCommand::IdBit | FUpdateTransformCommand::IdBit, [&](const FPrimitiveUpdateCommand& Cmd)
+	// First, register any new primitives
+	for (FPrimitiveSceneInfo* PrimitiveSceneInfo : ChangeSet.AddedPrimitiveSceneInfos)
 	{
-		if (Cmd.GetSceneInfo()->Proxy->IsSplineMesh())
+		if (PrimitiveSceneInfo->Proxy->IsSplineMesh())
+		{
+			RequestUpdate(SceneData->Register(*PrimitiveSceneInfo));
+		}
+	}
+
+	// Request updates from any updated primitives
+	for (FPrimitiveSceneInfo* PrimitiveSceneInfo : ChangeSet.UpdatedPrimitiveSceneInfos)
+	{
+		if (PrimitiveSceneInfo->Proxy->IsSplineMesh())
 		{
-			if (Cmd.IsAdd())
-			{
-				RequestUpdate(SceneData->Register(*Cmd.GetSceneInfo()));
-			}
-			else
-			{
-				RequestUpdate(SceneData->RegisteredPrimitives.FindChecked(Cmd.GetSceneInfo()));
-			}
+			RequestUpdate(SceneData->RegisteredPrimitives.FindChecked(PrimitiveSceneInfo));
 		}
-	});
+	}
 
 	if (SceneData->NumRegisteredPrimitives() > 0)
 	{
diff --git a/Engine/Source/Runtime/Renderer/Private/VirtualShadowMaps/VirtualShadowMapCacheManager.cpp b/Engine/Source/Runtime/Renderer/Private/VirtualShadowMaps/VirtualShadowMapCacheManager.cpp
--- a/Engine/Source/Runtime/Renderer/Private/VirtualShadowMaps/VirtualShadowMapCacheManager.cpp
+++ b/Engine/Source/Runtime/Renderer/Private/VirtualShadowMaps/VirtualShadowMapCacheManager.cpp
@@ -1688,19 +1688,24 @@ void FVirtualShadowMapInvalidationSceneUpdater::PreSceneUpdate(FRDGBuilder& Grap
 			InvalidatingPrimitiveCollector.UpdatedTransform(PrimitiveSceneInfo);
 		}
 
-		// Note: skips added as they are not fully defined at this point (not primitive ID allocated, 
-		ChangeSet.PrimitiveUpdates.ForEachUpdateCommand(ESceneUpdateCommandFilter::Updated | ESceneUpdateCommandFilter::Deleted, EPrimitiveUpdateDirtyFlags::AllCulling, [&](const FPrimitiveUpdateCommand& Cmd)
+		// All removed primitives must invalidate their footprints in the VSM before leaving
+		for (FPrimitiveSceneInfo* PrimitiveSceneInfo : ChangeSet.RemovedPrimitiveSceneInfos)
 		{
-			if (Cmd.IsDelete())
-			{
-				// All removed primitives must invalidate their footprints in the VSM before leaving.
-				InvalidatingPrimitiveCollector.Removed(Cmd.GetSceneInfo());
-			}
-			else
-			{
-				InvalidatingPrimitiveCollector.UpdatedTransform(Cmd.GetSceneInfo());
-			}
-		});
+			InvalidatingPrimitiveCollector.Removed(PrimitiveSceneInfo);
+		}
+		// As must all primitive updates, 
+		for (FPrimitiveSceneInfo* PrimitiveSceneInfo : ChangeSet.UpdatedPrimitiveSceneInfos)
+		{
+			InvalidatingPrimitiveCollector.UpdatedTransform(PrimitiveSceneInfo);
+		}
+
+		// TODO! Where do we get this data from...
+		/*
+		for (const auto& CullDistance : UpdatedInstanceCullDistance)
+		{
+			InvalidatingPrimitiveCollector.UpdatedTransform(CullDistance.Key->GetPrimitiveSceneInfo());
+		}
+		*/
 
 		TRACE_INT_VALUE(TEXT("Shadow.Virtual.Cache.PreInvalidationInstances"), InvalidatingPrimitiveCollector.Instances.GetTotalNumInstances());
 		CacheManager.ProcessInvalidations(GraphBuilder, SceneUniforms, InvalidatingPrimitiveCollector);
@@ -1712,7 +1717,7 @@ void FVirtualShadowMapInvalidationSceneUpdater::PostSceneUpdate(FRDGBuilder& Gra
 	CacheManager.ReallocatePersistentPrimitiveIndices();
 
 	// Grab a reference, but we currently do all the work in PostGPUSceneUpdate
-	PostUpdateChangeSet = &ChangeSet;
+	PostUpdateChangeSet = ChangeSet;
 }
 
 void FVirtualShadowMapInvalidationSceneUpdater::PostGPUSceneUpdate(FRDGBuilder& GraphBuilder, FSceneUniformBuffer& SceneUniforms)
@@ -1723,21 +1728,17 @@ void FVirtualShadowMapInvalidationSceneUpdater::PostGPUSceneUpdate(FRDGBuilder&
 	{
 		FVirtualShadowMapArrayCacheManager::FInvalidatingPrimitiveCollector InvalidatingPrimitiveCollector(&CacheManager);
 
-		// Filter out all updates that are either "add" or has dirty flags to say they affect the bounds.
-		PostUpdateChangeSet->PrimitiveUpdates.ForEachUpdateCommand(ESceneUpdateCommandFilter::AddedUpdated, EPrimitiveUpdateDirtyFlags::AllCulling, [&](const FPrimitiveUpdateCommand& Cmd)
+		for (FPrimitiveSceneInfo* PrimitiveSceneInfo : PostUpdateChangeSet.AddedPrimitiveSceneInfos)
 		{
-			if (Cmd.IsAdd())
-			{
-				InvalidatingPrimitiveCollector.Added(Cmd.GetSceneInfo());
-			}
-			else
-			{
-				InvalidatingPrimitiveCollector.UpdatedTransform(Cmd.GetSceneInfo());
-			}
-		});
+			InvalidatingPrimitiveCollector.Added(PrimitiveSceneInfo);
+		}
+		for (FPrimitiveSceneInfo* PrimitiveSceneInfo : PostUpdateChangeSet.UpdatedPrimitiveSceneInfos)
+		{
+			InvalidatingPrimitiveCollector.UpdatedTransform(PrimitiveSceneInfo);
+		}
 
 		TRACE_INT_VALUE(TEXT("Shadow.Virtual.Cache.PostInvalidationInstances"), InvalidatingPrimitiveCollector.Instances.GetTotalNumInstances());
 		CacheManager.ProcessInvalidations(GraphBuilder, SceneUniforms, InvalidatingPrimitiveCollector);
 	}
-	PostUpdateChangeSet = nullptr;
+	PostUpdateChangeSet = FScenePostUpdateChangeSet();
 }
diff --git a/Engine/Source/Runtime/Renderer/Private/VirtualShadowMaps/VirtualShadowMapCacheManager.h b/Engine/Source/Runtime/Renderer/Private/VirtualShadowMaps/VirtualShadowMapCacheManager.h
--- a/Engine/Source/Runtime/Renderer/Private/VirtualShadowMaps/VirtualShadowMapCacheManager.h
+++ b/Engine/Source/Runtime/Renderer/Private/VirtualShadowMaps/VirtualShadowMapCacheManager.h
@@ -490,5 +490,5 @@ public:
 private:
 	FVirtualShadowMapArrayCacheManager& CacheManager;
 
-	const FScenePostUpdateChangeSet* PostUpdateChangeSet = nullptr;
+	FScenePostUpdateChangeSet PostUpdateChangeSet;
 };
\ No newline at end of file
diff --git a/Engine/Source/Runtime/Renderer/Public/PrimitiveSceneInfo.h b/Engine/Source/Runtime/Renderer/Public/PrimitiveSceneInfo.h
--- a/Engine/Source/Runtime/Renderer/Public/PrimitiveSceneInfo.h
+++ b/Engine/Source/Runtime/Renderer/Public/PrimitiveSceneInfo.h
@@ -284,7 +284,6 @@ class FPrimitiveSceneInfo : public FDeferredCleanupInterface
 	friend class FSceneRenderer;
 	friend struct FViewDebugInfo;
 public:
-	using FPersistentId = FPersistentPrimitiveIndex;
 
 	/** The render proxy for the primitive. */
 	FPrimitiveSceneProxy* Proxy;
