Mobile: Removed some of the complexity in static lighting shader permutation settings
Removed r.Mobile.AllowMovableDirectionalLights option. Movable lights are always allowed. CSM from movable lights will only work in non pre-computed lighting enviroiment or when "Combined Static and CSM" option is enabled.
Removed r.Mobile.EnableNoPrecomputedLightingCSMShader. CSM shaders will be generated when project has  "Combined Static and CSM" option enabled.

diff --git a/Engine/Shaders/Private/LightmapCommon.ush b/Engine/Shaders/Private/LightmapCommon.ush
--- a/Engine/Shaders/Private/LightmapCommon.ush
+++ b/Engine/Shaders/Private/LightmapCommon.ush
@@ -310,9 +310,6 @@ half4 GetPrimaryPrecomputedShadowMask(VTPageTableResult LightmapVTPageTableResul
 		// Compute shadow factors by scaling and biasing the distance
 		half4 ShadowFactor = saturate( DistanceField * InvUniformPenumbraSizes + DistanceFieldBias );
 		return GetLightmapData(LightmapDataIndex).StaticShadowMapMasks * ShadowFactor * ShadowFactor;
-	#elif MOVABLE_DIRECTIONAL_LIGHT
-		// Do this before checking for lightmaps as we might have a lightmap + movable directional light
-		return 1.0f;
 	#elif HQ_TEXTURE_LIGHTMAP || LQ_TEXTURE_LIGHTMAP
 		// Mark as shadowed for lightmapped objects with no shadowmap
 		// This is necessary because objects inside a light's influence that were determined to be completely shadowed won't be rendered with STATICLIGHTING_TEXTUREMASK==1
diff --git a/Engine/Source/Runtime/Core/Private/HAL/ConsoleManager.cpp b/Engine/Source/Runtime/Core/Private/HAL/ConsoleManager.cpp
--- a/Engine/Source/Runtime/Core/Private/HAL/ConsoleManager.cpp
+++ b/Engine/Source/Runtime/Core/Private/HAL/ConsoleManager.cpp
@@ -3692,13 +3692,6 @@ static TAutoConsoleVariable<int32> CVarMobileEnableMovableLightCSMShaderCulling(
 		 "1: Primitives lit by movable directional light render with the CSM shader when determined to be within CSM range. (default)"),
 	ECVF_RenderThreadSafe | ECVF_ReadOnly);
 
-static TAutoConsoleVariable<int32> CVarMobileEnableNoPrecomputedLightingCSMShader(
-	TEXT("r.Mobile.EnableNoPrecomputedLightingCSMShader"),
-	0,
-	TEXT("0: CSM shaders for scenes without any precomputed lighting are not generated unless r.AllowStaticLighting is 0. (default)\n")
-	TEXT("1: CSM shaders for scenes without any precomputed lighting are always generated."),
-	ECVF_RenderThreadSafe | ECVF_ReadOnly);
-
 static TAutoConsoleVariable<int32> CVarMobileUseCSMShaderBranch(
 	TEXT("r.Mobile.UseCSMShaderBranch"),
 	0,
@@ -3713,13 +3706,6 @@ static TAutoConsoleVariable<int32> CVarMobileAllowDistanceFieldShadows(
 		 "1: Generate shader permutations to render distance field shadows from stationary directional lights. (default)"),
 	ECVF_RenderThreadSafe | ECVF_ReadOnly);
 
-static TAutoConsoleVariable<int32> CVarMobileAllowMovableDirectionalLights(
-	TEXT("r.Mobile.AllowMovableDirectionalLights"),
-	1,
-	TEXT("0: Do not generate shader permutations to render movable directional lights.\n"
-		 "1: Generate shader permutations to render movable directional lights. (default)"),
-	ECVF_RenderThreadSafe | ECVF_ReadOnly);
-
 static TAutoConsoleVariable<int32> CVarMobileSkyLightPermutation(
 	TEXT("r.Mobile.SkyLightPermutation"),
 	0,
diff --git a/Engine/Source/Runtime/Engine/Classes/Engine/RendererSettings.h b/Engine/Source/Runtime/Engine/Classes/Engine/RendererSettings.h
--- a/Engine/Source/Runtime/Engine/Classes/Engine/RendererSettings.h
+++ b/Engine/Source/Runtime/Engine/Classes/Engine/RendererSettings.h
@@ -1277,14 +1277,6 @@ class URendererSettings : public UDeveloperSettings
 		ConfigRestartRequired = true))
 		uint32 bMobileForwardEnableClusteredReflections : 1;
 
-	UPROPERTY(config, EditAnywhere, Category = MobileShaderPermutationReduction, meta = (
-		ConsoleVariable = "r.Mobile.EnableNoPrecomputedLightingCSMShader",
-		DisplayName = "Support CSM on levels with Force No Precomputed Lighting enabled",
-		EditCondition = "bAllowStaticLighting",
-		ToolTip = "When Allow Static Lighting is enabled, shaders to support CSM without any precomputed lighting are not normally generated. This setting allows CSM for this case at the cost of extra shader permutations. Changing this setting requires restarting the editor.",
-		ConfigRestartRequired = true))
-		uint32 bMobileEnableNoPrecomputedLightingCSMShader : 1;
-
 	UPROPERTY(config, EditAnywhere, Category = MobileShaderPermutationReduction, meta = (
 		ConsoleVariable = "r.Mobile.AllowDistanceFieldShadows",
 		DisplayName = "Support Pre-baked Distance Field Shadow Maps",
@@ -1292,13 +1284,6 @@ class URendererSettings : public UDeveloperSettings
 		ConfigRestartRequired = true))
 		uint32 bMobileAllowDistanceFieldShadows : 1;
 
-	UPROPERTY(config, EditAnywhere, Category = MobileShaderPermutationReduction, meta = (
-		ConsoleVariable = "r.Mobile.AllowMovableDirectionalLights",
-		DisplayName = "Support Movable Directional Lights",
-		ToolTip = "Generate shaders for primitives to receive movable directional lights. Changing this setting requires restarting the editor.",
-		ConfigRestartRequired = true))
-		uint32 bMobileAllowMovableDirectionalLights : 1;
-
 	UPROPERTY(config, EditAnywhere, Category = MobileShaderPermutationReduction, meta = (
 		ConsoleVariable = "r.Mobile.EnableMovableSpotlightsShadow",
 		DisplayName = "Support Movable SpotlightShadows",
diff --git a/Engine/Source/Runtime/RenderCore/Private/ReadOnlyCVARCache.cpp b/Engine/Source/Runtime/RenderCore/Private/ReadOnlyCVARCache.cpp
--- a/Engine/Source/Runtime/RenderCore/Private/ReadOnlyCVARCache.cpp
+++ b/Engine/Source/Runtime/RenderCore/Private/ReadOnlyCVARCache.cpp
@@ -19,7 +19,6 @@ bool FReadOnlyCVARCache::bSupportSkyAtmosphere = true;
 
 // Mobile specific
 bool FReadOnlyCVARCache::bMobileHDR = true;
-bool FReadOnlyCVARCache::bMobileAllowMovableDirectionalLights = true;
 bool FReadOnlyCVARCache::bMobileAllowDistanceFieldShadows = true;
 bool FReadOnlyCVARCache::bMobileEnableStaticAndCSMShadowReceivers = true;
 bool FReadOnlyCVARCache::bMobileEnableMovableLightCSMShaderCulling = true;
@@ -28,7 +27,6 @@ int32 FReadOnlyCVARCache::MobileSkyLightPermutationValue = 0;
 int32 FReadOnlyCVARCache::MobileEarlyZPassValue = 0;
 int32 FReadOnlyCVARCache::MobileForwardLocalLightsValue = 1;
 bool FReadOnlyCVARCache::bMobileForwardParticleLights = false;
-bool FReadOnlyCVARCache::bMobileEnableNoPrecomputedLightingCSMShader = false;
 bool FReadOnlyCVARCache::bMobileDeferredShadingValue = false;
 bool FReadOnlyCVARCache::bMobileEnableMovableSpotlightsShadowValue = false;
 int32 FReadOnlyCVARCache::MobileForwardDecalLightingValue = 1;
@@ -88,12 +86,10 @@ void FReadOnlyCVARCache::Initialize()
 	const auto CVarSupportSkyAtmosphere = IConsoleManager::Get().FindTConsoleVariableDataInt(TEXT("r.SupportSkyAtmosphere"));
 	
 	const auto CVarMobileHDR = IConsoleManager::Get().FindTConsoleVariableDataInt(TEXT("r.MobileHDR"));
-	const auto CVarMobileAllowMovableDirectionalLights = IConsoleManager::Get().FindTConsoleVariableDataInt(TEXT("r.Mobile.AllowMovableDirectionalLights"));
 	const auto CVarMobileEnableStaticAndCSMShadowReceivers = IConsoleManager::Get().FindTConsoleVariableDataInt(TEXT("r.Mobile.EnableStaticAndCSMShadowReceivers"));
 	const auto CVarMobileEnableMovableLightCSMShaderCulling = IConsoleManager::Get().FindTConsoleVariableDataInt(TEXT("r.Mobile.EnableMovableLightCSMShaderCulling"));
 	const auto CVarMobileAllowDistanceFieldShadows = IConsoleManager::Get().FindTConsoleVariableDataInt(TEXT("r.Mobile.AllowDistanceFieldShadows"));
 	const auto CVarMobileSkyLightPermutation = IConsoleManager::Get().FindTConsoleVariableDataInt(TEXT("r.Mobile.SkyLightPermutation"));
-	const auto CVarMobileEnableNoPrecomputedLightingCSMShader = IConsoleManager::Get().FindTConsoleVariableDataInt(TEXT("r.Mobile.EnableNoPrecomputedLightingCSMShader"));
 	const auto CVarMobileSupportGPUScene = IConsoleManager::Get().FindTConsoleVariableDataInt(TEXT("r.Mobile.SupportGPUScene"));
 	const auto CVarMobileForwardDecalLightingValue = IConsoleManager::Get().FindTConsoleVariableDataInt(TEXT("r.Mobile.Forward.DecalLighting"));
 	
@@ -107,12 +103,10 @@ void FReadOnlyCVARCache::Initialize()
 
 	// mobile
 	bMobileHDR = CVarMobileHDR->GetValueOnAnyThread() == 1;
-	bMobileAllowMovableDirectionalLights = CVarMobileAllowMovableDirectionalLights->GetValueOnAnyThread() != 0;
 	bMobileAllowDistanceFieldShadows = CVarMobileAllowDistanceFieldShadows->GetValueOnAnyThread() != 0;
 	bMobileEnableStaticAndCSMShadowReceivers = CVarMobileEnableStaticAndCSMShadowReceivers->GetValueOnAnyThread() != 0;
 	bMobileEnableMovableLightCSMShaderCulling = CVarMobileEnableMovableLightCSMShaderCulling->GetValueOnAnyThread() != 0;
 	MobileSkyLightPermutationValue = CVarMobileSkyLightPermutation->GetValueOnAnyThread();
-	bMobileEnableNoPrecomputedLightingCSMShader = CVarMobileEnableNoPrecomputedLightingCSMShader->GetValueOnAnyThread() != 0;
 	MobileEarlyZPassValue = MobileEarlyZPassIniValue(GMaxRHIShaderPlatform);
 	MobileForwardLocalLightsValue = MobileForwardLocalLightsIniValue(GMaxRHIShaderPlatform);
 	bMobileDeferredShadingValue = MobileDeferredShadingIniValue(GMaxRHIShaderPlatform);
diff --git a/Engine/Source/Runtime/RenderCore/Public/ReadOnlyCVARCache.h b/Engine/Source/Runtime/RenderCore/Public/ReadOnlyCVARCache.h
--- a/Engine/Source/Runtime/RenderCore/Public/ReadOnlyCVARCache.h
+++ b/Engine/Source/Runtime/RenderCore/Public/ReadOnlyCVARCache.h
@@ -62,11 +62,6 @@ struct FReadOnlyCVARCache
 		#endif
 	}
 
-	static inline bool MobileAllowMovableDirectionalLights()
-	{
-		return bMobileAllowMovableDirectionalLights;
-	}
-
 	static inline bool MobileAllowDistanceFieldShadows()
 	{
 		return bMobileAllowDistanceFieldShadows;
@@ -87,11 +82,6 @@ struct FReadOnlyCVARCache
 		return MobileSkyLightPermutationValue;
 	}
 
-	static inline bool MobileEnableNoPrecomputedLightingCSMShader()
-	{
-		return bMobileEnableNoPrecomputedLightingCSMShader;
-	}
-
 	static inline int32 MobileForwardDecalLighting()
 	{
 		return MobileForwardDecalLightingValue;
@@ -159,7 +149,6 @@ private:
 
 	// Mobile specific
 	RENDERCORE_API static bool bMobileHDR;
-	RENDERCORE_API static bool bMobileAllowMovableDirectionalLights;
 	RENDERCORE_API static bool bMobileAllowDistanceFieldShadows;
 	RENDERCORE_API static bool bMobileEnableStaticAndCSMShadowReceivers;
 	RENDERCORE_API static bool bMobileEnableMovableLightCSMShaderCulling;
@@ -169,7 +158,6 @@ private:
 	RENDERCORE_API static int32 MobileForwardLocalLightsValue;
 	RENDERCORE_API static bool bMobileForwardParticleLights;
 	RENDERCORE_API static int32 MobileForwardDecalLightingValue;
-	RENDERCORE_API static bool bMobileEnableNoPrecomputedLightingCSMShader;
 	RENDERCORE_API static bool bMobileDeferredShadingValue;
 	RENDERCORE_API static bool bMobileEnableMovableSpotlightsShadowValue;
 
diff --git a/Engine/Source/Runtime/Renderer/Private/LightMapRendering.cpp b/Engine/Source/Runtime/Renderer/Private/LightMapRendering.cpp
--- a/Engine/Source/Runtime/Renderer/Private/LightMapRendering.cpp
+++ b/Engine/Source/Runtime/Renderer/Private/LightMapRendering.cpp
@@ -164,13 +164,9 @@ bool FMobileDirectionalLightAndCSMPolicy::ShouldCompilePermutation(const FMeshMa
 	{
 		return false;
 	}
-
-	static auto* CVarEnableNoPrecomputedLightingCSMShader = IConsoleManager::Get().FindTConsoleVariableDataInt(TEXT("r.Mobile.EnableNoPrecomputedLightingCSMShader"));
-	const bool bEnableNoPrecomputedLightingCSMShader = CVarEnableNoPrecomputedLightingCSMShader && CVarEnableNoPrecomputedLightingCSMShader->GetValueOnAnyThread() != 0;
-
-	return (!IsStaticLightingAllowed() || bEnableNoPrecomputedLightingCSMShader) &&
-		Parameters.MaterialParameters.ShadingModels.IsLit() &&
-		!IsTranslucentBlendMode(Parameters.MaterialParameters);
+	
+	return Parameters.MaterialParameters.ShadingModels.IsLit() && !IsTranslucentBlendMode(Parameters.MaterialParameters) 
+		&& (!IsStaticLightingAllowed() || FReadOnlyCVARCache::MobileEnableStaticAndCSMShadowReceivers());
 }
 
 bool FMobileDistanceFieldShadowsAndLQLightMapPolicy::ShouldCompilePermutation(const FMeshMaterialShaderPermutationParameters& Parameters)
@@ -251,37 +247,6 @@ void FMobileDirectionalLightCSMAndSHIndirectPolicy::ModifyCompilationEnvironment
 	Super::ModifyCompilationEnvironment(Parameters, OutEnvironment);
 }
 
-bool FMobileMovableDirectionalLightWithLightmapPolicy::ShouldCompilePermutation(const FMeshMaterialShaderPermutationParameters& Parameters)
-{
-	static auto* CVarMobileAllowMovableDirectionalLights = IConsoleManager::Get().FindTConsoleVariableDataInt(TEXT("r.Mobile.AllowMovableDirectionalLights"));
-	const bool bMobileAllowMovableDirectionalLights = CVarMobileAllowMovableDirectionalLights->GetValueOnAnyThread() != 0;
-
-	return bMobileAllowMovableDirectionalLights && Super::ShouldCompilePermutation(Parameters);
-}
-
-void FMobileMovableDirectionalLightWithLightmapPolicy::ModifyCompilationEnvironment(const FMaterialShaderPermutationParameters& Parameters, FShaderCompilerEnvironment& OutEnvironment)
-{
-	OutEnvironment.SetDefine(TEXT("MOVABLE_DIRECTIONAL_LIGHT"), TEXT("1"));
-	Super::ModifyCompilationEnvironment(Parameters, OutEnvironment);
-}
-
-bool FMobileMovableDirectionalLightCSMWithLightmapPolicy::ShouldCompilePermutation(const FMeshMaterialShaderPermutationParameters& Parameters)
-{
-	if (IsMobileDeferredShadingEnabled(Parameters.Platform))
-	{
-		return false;
-	}
-
-	return Super::ShouldCompilePermutation(Parameters);
-}
-
-void FMobileMovableDirectionalLightCSMWithLightmapPolicy::ModifyCompilationEnvironment(const FMaterialShaderPermutationParameters& Parameters, FShaderCompilerEnvironment& OutEnvironment)
-{
-	OutEnvironment.SetDefine(TEXT("DIRECTIONAL_LIGHT_CSM"), TEXT("1"));
-
-	Super::ModifyCompilationEnvironment(Parameters, OutEnvironment);
-}
-
 bool FSelfShadowedCachedPointIndirectLightingPolicy::ShouldCompilePermutation(const FMeshMaterialShaderPermutationParameters& Parameters)
 {
 	return Parameters.MaterialParameters.ShadingModels.IsLit()
@@ -454,10 +419,6 @@ bool FUniformLightMapPolicy::ShouldCompilePermutation(ELightMapPolicyType Policy
 		return FMobileDirectionalLightAndSHIndirectPolicy::ShouldCompilePermutation(Parameters);
 	case LMP_MOBILE_DIRECTIONAL_LIGHT_CSM_AND_SH_INDIRECT:
 		return FMobileDirectionalLightCSMAndSHIndirectPolicy::ShouldCompilePermutation(Parameters);
-	case LMP_MOBILE_MOVABLE_DIRECTIONAL_LIGHT_WITH_LIGHTMAP:
-		return FMobileMovableDirectionalLightWithLightmapPolicy::ShouldCompilePermutation(Parameters);
-	case LMP_MOBILE_MOVABLE_DIRECTIONAL_LIGHT_CSM_WITH_LIGHTMAP:
-		return FMobileMovableDirectionalLightCSMWithLightmapPolicy::ShouldCompilePermutation(Parameters);
 	case LMP_MOBILE_DIRECTIONAL_LIGHT_CSM:
 		return FMobileDirectionalLightAndCSMPolicy::ShouldCompilePermutation(Parameters);
 
@@ -517,12 +478,6 @@ void FUniformLightMapPolicy::ModifyCompilationEnvironment(ELightMapPolicyType Po
 	case LMP_MOBILE_DIRECTIONAL_LIGHT_CSM_AND_SH_INDIRECT:
 		FMobileDirectionalLightCSMAndSHIndirectPolicy::ModifyCompilationEnvironment(Parameters, OutEnvironment);
 		break;
-	case LMP_MOBILE_MOVABLE_DIRECTIONAL_LIGHT_WITH_LIGHTMAP:
-		FMobileMovableDirectionalLightWithLightmapPolicy::ModifyCompilationEnvironment(Parameters, OutEnvironment);
-		break;
-	case LMP_MOBILE_MOVABLE_DIRECTIONAL_LIGHT_CSM_WITH_LIGHTMAP:
-		FMobileMovableDirectionalLightCSMWithLightmapPolicy::ModifyCompilationEnvironment(Parameters, OutEnvironment);
-		break;
 	case LMP_MOBILE_DIRECTIONAL_LIGHT_CSM:
 		FMobileDirectionalLightAndCSMPolicy::ModifyCompilationEnvironment(Parameters, OutEnvironment);
 		break;
diff --git a/Engine/Source/Runtime/Renderer/Private/LightMapRendering.h b/Engine/Source/Runtime/Renderer/Private/LightMapRendering.h
--- a/Engine/Source/Runtime/Renderer/Private/LightMapRendering.h
+++ b/Engine/Source/Runtime/Renderer/Private/LightMapRendering.h
@@ -298,25 +298,6 @@ public:
 	static void ModifyCompilationEnvironment(const FMaterialShaderPermutationParameters& Parameters, FShaderCompilerEnvironment& OutEnvironment);
 };
 
-/** Mobile Specific: Combines a movable directional light with LQ lightmaps  */
-struct FMobileMovableDirectionalLightWithLightmapPolicy : public TLightMapPolicy<LQ_LIGHTMAP>
-{
-	typedef TLightMapPolicy<LQ_LIGHTMAP> Super;
-
-	static bool ShouldCompilePermutation(const FMeshMaterialShaderPermutationParameters& Parameters);
-	static void ModifyCompilationEnvironment(const FMaterialShaderPermutationParameters& Parameters, FShaderCompilerEnvironment& OutEnvironment);
-};
-
-/** Mobile Specific: Combines a movable directional light with LQ lightmaps and CSM */
-struct FMobileMovableDirectionalLightCSMWithLightmapPolicy : public FMobileMovableDirectionalLightWithLightmapPolicy
-{
-	typedef FMobileMovableDirectionalLightWithLightmapPolicy Super;
-	
-	static bool ShouldCompilePermutation(const FMeshMaterialShaderPermutationParameters& Parameters);
-	static void ModifyCompilationEnvironment(const FMaterialShaderPermutationParameters& Parameters, FShaderCompilerEnvironment& OutEnvironment);
-};
-
-
 enum ELightMapPolicyType
 {
 	LMP_NO_LIGHTMAP,
@@ -332,8 +313,6 @@ enum ELightMapPolicyType
 	LMP_MOBILE_DIRECTIONAL_LIGHT_CSM_AND_LIGHTMAP,
 	LMP_MOBILE_DIRECTIONAL_LIGHT_AND_SH_INDIRECT,
 	LMP_MOBILE_DIRECTIONAL_LIGHT_CSM_AND_SH_INDIRECT,
-	LMP_MOBILE_MOVABLE_DIRECTIONAL_LIGHT_WITH_LIGHTMAP,
-	LMP_MOBILE_MOVABLE_DIRECTIONAL_LIGHT_CSM_WITH_LIGHTMAP,
 	LMP_MOBILE_DIRECTIONAL_LIGHT_CSM,
 
 	// LightMapDensity
diff --git a/Engine/Source/Runtime/Renderer/Private/MobileBasePass.cpp b/Engine/Source/Runtime/Renderer/Private/MobileBasePass.cpp
--- a/Engine/Source/Runtime/Renderer/Private/MobileBasePass.cpp
+++ b/Engine/Source/Runtime/Renderer/Private/MobileBasePass.cpp
@@ -227,10 +227,6 @@ bool GetMobileBasePassShaders(
 		return GetUniformMobileBasePassShaders<LMP_MOBILE_DIRECTIONAL_LIGHT_AND_SH_INDIRECT, LocalLightSetting>(Material, VertexFactoryType, bEnableSkyLight, ColoredTransmittanceFallback, VertexShader, PixelShader);
 	case LMP_MOBILE_DIRECTIONAL_LIGHT_CSM_AND_SH_INDIRECT:
 		return GetUniformMobileBasePassShaders<LMP_MOBILE_DIRECTIONAL_LIGHT_CSM_AND_SH_INDIRECT, LocalLightSetting>(Material, VertexFactoryType, bEnableSkyLight, ColoredTransmittanceFallback, VertexShader, PixelShader);
-	case LMP_MOBILE_MOVABLE_DIRECTIONAL_LIGHT_WITH_LIGHTMAP:
-		return GetUniformMobileBasePassShaders<LMP_MOBILE_MOVABLE_DIRECTIONAL_LIGHT_WITH_LIGHTMAP, LocalLightSetting>(Material, VertexFactoryType, bEnableSkyLight, ColoredTransmittanceFallback, VertexShader, PixelShader);
-	case LMP_MOBILE_MOVABLE_DIRECTIONAL_LIGHT_CSM_WITH_LIGHTMAP:
-		return GetUniformMobileBasePassShaders<LMP_MOBILE_MOVABLE_DIRECTIONAL_LIGHT_CSM_WITH_LIGHTMAP, LocalLightSetting>(Material, VertexFactoryType, bEnableSkyLight, ColoredTransmittanceFallback, VertexShader, PixelShader);
 	case LMP_MOBILE_DIRECTIONAL_LIGHT_CSM:
 		return GetUniformMobileBasePassShaders<LMP_MOBILE_DIRECTIONAL_LIGHT_CSM, LocalLightSetting>(Material, VertexFactoryType, bEnableSkyLight, ColoredTransmittanceFallback, VertexShader, PixelShader);
 	default:										
@@ -341,11 +337,6 @@ bool MobileBasePass::StaticCanReceiveCSM(const FLightSceneInfo* LightSceneInfo,
 	return false; 
 }
 
-bool MobileBasePass::IsUsingDirectionalLightForLighmapPolicySelection(const FScene* Scene)
-{
-	return !(!IsStaticLightingAllowed() || (FReadOnlyCVARCache::MobileEnableNoPrecomputedLightingCSMShader() && Scene && Scene->GetForceNoPrecomputedLighting()));
-}
-
 ELightMapPolicyType MobileBasePass::SelectMeshLightmapPolicy(
 	const FScene* Scene, 
 	const FMeshBatch& Mesh, 
@@ -362,7 +353,7 @@ ELightMapPolicyType MobileBasePass::SelectMeshLightmapPolicy(
 	{
 		constexpr ERHIFeatureLevel::Type FeatureLevel = ERHIFeatureLevel::ES3_1;
 		
-		if (!IsUsingDirectionalLightForLighmapPolicySelection(Scene))
+		if (!IsStaticLightingAllowed())
 		{
  			if (!bIsTranslucent)
 			{
@@ -389,22 +380,21 @@ ELightMapPolicyType MobileBasePass::SelectMeshLightmapPolicy(
 
 			const FLightSceneInfo* MobileDirectionalLight = MobileBasePass::GetDirectionalLightInfo(Scene, PrimitiveSceneProxy);
 		
-			const bool bUseMovableLight = MobileDirectionalLight && !MobileDirectionalLight->Proxy->HasStaticShadowing() && FReadOnlyCVARCache::MobileAllowMovableDirectionalLights();
-			const bool bUseStaticAndCSM = MobileDirectionalLight && MobileDirectionalLight->Proxy->UseCSMForDynamicObjects()
-											&& bPrimReceivesCSM
-											&& FReadOnlyCVARCache::MobileEnableStaticAndCSMShadowReceivers();
-
-			const bool bMovableWithCSM = bUseMovableLight && MobileDirectionalLight->ShouldRenderViewIndependentWholeSceneShadows() && bPrimReceivesCSM;
-
-			const bool bPrimitiveUsesILC = PrimitiveSceneProxy
-										&& (PrimitiveSceneProxy->IsMovable() || PrimitiveSceneProxy->NeedsUnbuiltPreviewLighting() || PrimitiveSceneProxy->GetLightmapType() == ELightmapType::ForceVolumetric)
-										&& PrimitiveSceneProxy->WillEverBeLit()
-										&& PrimitiveSceneProxy->GetIndirectLightingCacheQuality() != ILCQ_Off;
+			// Primitive can receive both pre-computed and CSM shadows
+			const bool bPrimReceivesStaticAndCSM = 
+				MobileDirectionalLight 
+				&& bPrimReceivesCSM
+				&& FReadOnlyCVARCache::MobileEnableStaticAndCSMShadowReceivers()
+				&& MobileDirectionalLight->ShouldRenderViewIndependentWholeSceneShadows();
+
+			const bool bPrimitiveUsesILC = 
+				PrimitiveSceneProxy
+				&& (PrimitiveSceneProxy->IsMovable() || PrimitiveSceneProxy->NeedsUnbuiltPreviewLighting() || PrimitiveSceneProxy->GetLightmapType() == ELightmapType::ForceVolumetric)
+				&& PrimitiveSceneProxy->WillEverBeLit()
+				&& PrimitiveSceneProxy->GetIndirectLightingCacheQuality() != ILCQ_Off;
 
 			const bool bHasValidVLM = Scene && Scene->VolumetricLightmapSceneData.HasData();
-
-			const bool bHasValidILC = Scene && Scene->PrecomputedLightVolumes.Num() > 0
-									&& IsIndirectLightingCacheAllowed(FeatureLevel);
+			const bool bHasValidILC = Scene && Scene->PrecomputedLightVolumes.Num() > 0	&& IsIndirectLightingCacheAllowed(FeatureLevel);
 
 			if (LightMapInteraction.GetType() == LMIT_Texture && FReadOnlyCVARCache::EnableLowQualityLightmaps())
 			{
@@ -412,58 +402,22 @@ ELightMapPolicyType MobileBasePass::SelectMeshLightmapPolicy(
 					? Mesh.LCI->GetShadowMapInteraction(FeatureLevel)
 					: FShadowMapInteraction();
 
-				if ((bUseStaticAndCSM || bMovableWithCSM) && !bUsesDeferredShading)
+				if (ShadowMapInteraction.GetType() == SMIT_Texture && FReadOnlyCVARCache::MobileAllowDistanceFieldShadows())
 				{
-					if (ShadowMapInteraction.GetType() == SMIT_Texture &&
-						MobileDirectionalLight->ShouldRenderViewIndependentWholeSceneShadows() &&
-						FReadOnlyCVARCache::MobileAllowDistanceFieldShadows())
-					{
-						SelectedLightmapPolicy = LMP_MOBILE_DISTANCE_FIELD_SHADOWS_LIGHTMAP_AND_CSM;
-					}
-					else
-					{
-						// Lightmap path
-						if (bMovableWithCSM)
-						{
-							SelectedLightmapPolicy = LMP_MOBILE_MOVABLE_DIRECTIONAL_LIGHT_CSM_WITH_LIGHTMAP;
-						}
-						else
-						{
-							SelectedLightmapPolicy = LMP_MOBILE_DIRECTIONAL_LIGHT_CSM_AND_LIGHTMAP;
-						}
-					}
+					SelectedLightmapPolicy = (bPrimReceivesStaticAndCSM && !bUsesDeferredShading) ?  
+						LMP_MOBILE_DISTANCE_FIELD_SHADOWS_LIGHTMAP_AND_CSM : 
+						LMP_MOBILE_DISTANCE_FIELD_SHADOWS_AND_LQ_LIGHTMAP;
 				}
 				else
 				{
-					if (ShadowMapInteraction.GetType() == SMIT_Texture &&
-						FReadOnlyCVARCache::MobileAllowDistanceFieldShadows())
-					{
-						SelectedLightmapPolicy = LMP_MOBILE_DISTANCE_FIELD_SHADOWS_AND_LQ_LIGHTMAP;
-					}
-					else
-					{
-						// Lightmap path
-						if (bUseMovableLight)
-						{
-							if (bUsesDeferredShading)
-							{
-								SelectedLightmapPolicy = LMP_LQ_LIGHTMAP;
-							}
-							else
-							{
-								SelectedLightmapPolicy = LMP_MOBILE_MOVABLE_DIRECTIONAL_LIGHT_WITH_LIGHTMAP;
-							}
-						}
-						else
-						{
-							SelectedLightmapPolicy = LMP_LQ_LIGHTMAP;
-						}
-					}
+					SelectedLightmapPolicy = (bPrimReceivesStaticAndCSM && !bUsesDeferredShading) ? 
+						LMP_MOBILE_DIRECTIONAL_LIGHT_CSM_AND_LIGHTMAP : 
+						LMP_LQ_LIGHTMAP;
 				}
 			}
 			else if ((bHasValidVLM || bHasValidILC) && bPrimitiveUsesILC)
 			{
-				if ((bUseStaticAndCSM || bMovableWithCSM) && !bUsesDeferredShading && FReadOnlyCVARCache::MobileEnableStaticAndCSMShadowReceivers())
+				if (bPrimReceivesStaticAndCSM && !bUsesDeferredShading)
 				{
 					SelectedLightmapPolicy = LMP_MOBILE_DIRECTIONAL_LIGHT_CSM_AND_SH_INDIRECT;
 				}
@@ -472,6 +426,10 @@ ELightMapPolicyType MobileBasePass::SelectMeshLightmapPolicy(
 					SelectedLightmapPolicy = LMP_MOBILE_DIRECTIONAL_LIGHT_AND_SH_INDIRECT;
 				}
 			}
+			else if (bPrimReceivesStaticAndCSM && !bUsesDeferredShading)
+			{
+				SelectedLightmapPolicy = LMP_MOBILE_DIRECTIONAL_LIGHT_CSM;
+			}
 		}
 	}
 		
@@ -486,9 +444,7 @@ static FMobileLightMapPolicyTypeList GetUniformLightMapPolicyTypeForPSOCollectio
 	
 	if (bLitMaterial)
 	{
-		const bool bAllowStaticLighting = IsStaticLightingAllowed();
-		
-		if (!bAllowStaticLighting)
+		if (!IsStaticLightingAllowed())
 		{
 			if (bUsesDeferredShading || !MobileUseCSMShaderBranch())
 			{
@@ -512,11 +468,6 @@ static FMobileLightMapPolicyTypeList GetUniformLightMapPolicyTypeForPSOCollectio
 						Result.Add(LMP_MOBILE_DISTANCE_FIELD_SHADOWS_LIGHTMAP_AND_CSM);
 					}
 
-					if (FReadOnlyCVARCache::MobileAllowMovableDirectionalLights())
-					{
-						Result.Add(LMP_MOBILE_MOVABLE_DIRECTIONAL_LIGHT_CSM_WITH_LIGHTMAP);
-					}
-
 					Result.Add(LMP_MOBILE_DIRECTIONAL_LIGHT_CSM_AND_LIGHTMAP);
 				}
 
@@ -524,18 +475,6 @@ static FMobileLightMapPolicyTypeList GetUniformLightMapPolicyTypeForPSOCollectio
 				{
 					Result.Add(LMP_MOBILE_DISTANCE_FIELD_SHADOWS_AND_LQ_LIGHTMAP);
 				}
-
-				if (FReadOnlyCVARCache::MobileAllowMovableDirectionalLights())
-				{
-					if (bUsesDeferredShading)
-					{
-						Result.Add(LMP_LQ_LIGHTMAP);
-					}
-					else
-					{
-						Result.Add(LMP_MOBILE_MOVABLE_DIRECTIONAL_LIGHT_WITH_LIGHTMAP);
-					}
-				}
 				
 				Result.Add(LMP_LQ_LIGHTMAP);
 			}
@@ -543,8 +482,6 @@ static FMobileLightMapPolicyTypeList GetUniformLightMapPolicyTypeForPSOCollectio
 			// ILC/LVM
 			if (bMovable)
 			{
-				Result.Add(LMP_NO_LIGHTMAP); // in case there is no valid ILC/VLM
-				
 				if (!bUsesDeferredShading && FReadOnlyCVARCache::MobileEnableStaticAndCSMShadowReceivers() && bCanReceiveCSM)
 				{
 					Result.Add(LMP_MOBILE_DIRECTIONAL_LIGHT_CSM_AND_SH_INDIRECT);
@@ -553,6 +490,16 @@ static FMobileLightMapPolicyTypeList GetUniformLightMapPolicyTypeForPSOCollectio
 				{
 					Result.Add(LMP_MOBILE_DIRECTIONAL_LIGHT_AND_SH_INDIRECT);
 				}
+
+				// in case there is no valid ILC/VLM
+				if (bCanReceiveCSM) 
+				{
+					Result.Add(LMP_MOBILE_DIRECTIONAL_LIGHT_CSM);
+				}
+				else
+				{
+					Result.Add(LMP_NO_LIGHTMAP); 
+				}
 			}
 		}
 	}
diff --git a/Engine/Source/Runtime/Renderer/Private/MobileBasePassRendering.cpp b/Engine/Source/Runtime/Renderer/Private/MobileBasePassRendering.cpp
--- a/Engine/Source/Runtime/Renderer/Private/MobileBasePassRendering.cpp
+++ b/Engine/Source/Runtime/Renderer/Private/MobileBasePassRendering.cpp
@@ -165,8 +165,6 @@ IMPLEMENT_MOBILE_SHADING_BASEPASS_LIGHTMAPPED_SHADER_TYPE(TUniformLightMapPolicy
 IMPLEMENT_MOBILE_SHADING_BASEPASS_LIGHTMAPPED_SHADER_TYPE(TUniformLightMapPolicy<LMP_MOBILE_DIRECTIONAL_LIGHT_CSM_AND_LIGHTMAP>, FMobileDirectionalLightCSMAndLightMapPolicy);
 IMPLEMENT_MOBILE_SHADING_BASEPASS_LIGHTMAPPED_SHADER_TYPE(TUniformLightMapPolicy<LMP_MOBILE_DIRECTIONAL_LIGHT_AND_SH_INDIRECT>, FMobileDirectionalLightAndSHIndirectPolicy);
 IMPLEMENT_MOBILE_SHADING_BASEPASS_LIGHTMAPPED_SHADER_TYPE(TUniformLightMapPolicy<LMP_MOBILE_DIRECTIONAL_LIGHT_CSM_AND_SH_INDIRECT>, FMobileDirectionalLightCSMAndSHIndirectPolicy);
-IMPLEMENT_MOBILE_SHADING_BASEPASS_LIGHTMAPPED_SHADER_TYPE(TUniformLightMapPolicy<LMP_MOBILE_MOVABLE_DIRECTIONAL_LIGHT_WITH_LIGHTMAP>, FMobileMovableDirectionalLightWithLightmapPolicy);
-IMPLEMENT_MOBILE_SHADING_BASEPASS_LIGHTMAPPED_SHADER_TYPE(TUniformLightMapPolicy<LMP_MOBILE_MOVABLE_DIRECTIONAL_LIGHT_CSM_WITH_LIGHTMAP>, FMobileMovableDirectionalLightCSMWithLightmapPolicy);
 IMPLEMENT_MOBILE_SHADING_BASEPASS_LIGHTMAPPED_SHADER_TYPE(TUniformLightMapPolicy<LMP_MOBILE_DIRECTIONAL_LIGHT_CSM>, FMobileDirectionalLightAndCSMPolicy);
 
 bool MaterialRequiresColorTransmittanceBlending(const FMaterial& MaterialResource)
diff --git a/Engine/Source/Runtime/Renderer/Private/MobileBasePassRendering.h b/Engine/Source/Runtime/Renderer/Private/MobileBasePassRendering.h
--- a/Engine/Source/Runtime/Renderer/Private/MobileBasePassRendering.h
+++ b/Engine/Source/Runtime/Renderer/Private/MobileBasePassRendering.h
@@ -323,8 +323,6 @@ public:
 
 namespace MobileBasePass
 {
-	bool IsUsingDirectionalLightForLighmapPolicySelection(const FScene* Scene);
-
 	ELightMapPolicyType SelectMeshLightmapPolicy(
 		const FScene* Scene, 
 		const FMeshBatch& MeshBatch, 
diff --git a/Engine/Source/Runtime/Renderer/Private/RendererScene.cpp b/Engine/Source/Runtime/Renderer/Private/RendererScene.cpp
--- a/Engine/Source/Runtime/Renderer/Private/RendererScene.cpp
+++ b/Engine/Source/Runtime/Renderer/Private/RendererScene.cpp
@@ -2749,7 +2749,8 @@ void FScene::AddLightSceneInfo_RenderThread(FLightSceneInfo* LightSceneInfo)
 			    MobileDirectionalLights[FirstLightingChannel] = LightSceneInfo;
     
 			    // if this light is a dynamic shadowcast then we need to update the static draw lists to pick a new lighting policy:
-			    if (MobileBasePass::IsUsingDirectionalLightForLighmapPolicySelection(this) && (!LightSceneInfo->Proxy->HasStaticShadowing() || bUseCSMForDynamicObjects))
+				const bool bUsesDirectionalLightForLighmapPolicySelection = IsStaticLightingAllowed() && !IsMobileDeferredShadingEnabled(GetShaderPlatform());
+			    if (bUsesDirectionalLightForLighmapPolicySelection && (!LightSceneInfo->Proxy->HasStaticShadowing() || bUseCSMForDynamicObjects))
 				{
 		    		bScenesPrimitivesNeedStaticMeshElementUpdate = true;
 					UE_CLOG(!GIsEditor, LogRenderer, Log, TEXT("Forcing update for all mesh draw commands: Add directional light"));
@@ -4040,7 +4041,8 @@ void FScene::RemoveLightSceneInfo_RenderThread(FLightSceneInfo* LightSceneInfo)
 				}
 
 				// if this light is a dynamic shadowcast then we need to update the static draw lists to pick a new lightingpolicy
-				if (MobileBasePass::IsUsingDirectionalLightForLighmapPolicySelection(this) && (!LightSceneInfo->Proxy->HasStaticShadowing() || bUseCSMForDynamicObjects))
+				const bool bUsesDirectionalLightForLighmapPolicySelection = IsStaticLightingAllowed() && !IsMobileDeferredShadingEnabled(GetShaderPlatform());
+				if (bUsesDirectionalLightForLighmapPolicySelection && (!LightSceneInfo->Proxy->HasStaticShadowing() || bUseCSMForDynamicObjects))
 				{
 					bScenesPrimitivesNeedStaticMeshElementUpdate = true;
 					UE_CLOG(!GIsEditor, LogRenderer, Log, TEXT("Forcing update for all mesh draw commands: Remove directional light"));
diff --git a/Engine/Source/Runtime/Renderer/Private/SceneRendering.cpp b/Engine/Source/Runtime/Renderer/Private/SceneRendering.cpp
--- a/Engine/Source/Runtime/Renderer/Private/SceneRendering.cpp
+++ b/Engine/Source/Runtime/Renderer/Private/SceneRendering.cpp
@@ -3999,7 +3999,6 @@ void FSceneRenderer::OnRenderFinish(FRDGBuilder& GraphBuilder, FRDGTextureRef Vi
 		const bool bMobile = (FeatureLevel <= ERHIFeatureLevel::ES3_1);
 		const bool bShowMobileLowQualityLightmapWarning = bMobile && !FReadOnlyCVARCache::EnableLowQualityLightmaps() && IsStaticLightingAllowed();
 		const bool bShowMobileDynamicCSMWarning = bMobile && Scene->NumMobileStaticAndCSMLights_RenderThread > 0 && !(FReadOnlyCVARCache::MobileEnableStaticAndCSMShadowReceivers() && FReadOnlyCVARCache::MobileAllowDistanceFieldShadows());
-		const bool bShowMobileMovableDirectionalLightWarning = bMobile && Scene->NumMobileMovableDirectionalLights_RenderThread > 0 && !FReadOnlyCVARCache::MobileAllowMovableDirectionalLights();
 		const bool bMobileMissingSkyMaterial = (bMobile && Scene->HasSkyAtmosphere() && (Views.Num() > 0 && !Views[0].bSceneHasSkyMaterial));
 
 		const bool bSingleLayerWaterWarning = ShouldRenderSingleLayerWaterSkippedRenderEditorNotification(Views);
@@ -4019,7 +4018,7 @@ void FSceneRenderer::OnRenderFinish(FRDGBuilder& GraphBuilder, FRDGTextureRef Vi
 		const bool bHasDelegateWarnings = OnGetOnScreenMessages.IsBound();
 
 		const bool bAnyWarning = bShowPrecomputedVisibilityWarning || bShowDemotedLocalMemoryWarning || bShowGlobalClipPlaneWarning || bShowSkylightWarning || bShowPointLightWarning
-			|| bShowDFAODisabledWarning || bShowShadowedLightOverflowWarning || bShowMobileDynamicCSMWarning || bShowMobileLowQualityLightmapWarning || bShowMobileMovableDirectionalLightWarning
+			|| bShowDFAODisabledWarning || bShowShadowedLightOverflowWarning || bShowMobileDynamicCSMWarning || bShowMobileLowQualityLightmapWarning
 			|| bMobileMissingSkyMaterial || bShowSkinCacheOOM || bSingleLayerWaterWarning || bShowDFDisabledWarning || bShowNoSkyAtmosphereComponentWarning || bFxDebugDraw 
 			|| bLumenEnabledButHasNoDataForTracing || bLumenEnabledButDisabledForTheProject || bNaniteEnabledButNoAtomics || bNaniteEnabledButDisabledInProject || bRealTimeSkyCaptureButNothingToCapture || bShowWaitingSkylight
 			|| bShowLocalExposureDisabledWarning || bHasDelegateWarnings || bContactShadowIntensityCvarUsed || bLocalFogVolumeInSceneButProjectDisabled
@@ -4067,7 +4066,7 @@ void FSceneRenderer::OnRenderFinish(FRDGBuilder& GraphBuilder, FRDGTextureRef Vi
 						[this, ViewState, GPUSkinCacheExtraRequiredMemory,
 						bLocked, bShowPrecomputedVisibilityWarning, bShowDemotedLocalMemoryWarning, bShowGlobalClipPlaneWarning, bShowDFAODisabledWarning, bShowDFDisabledWarning,
 						bIsFrozen, bShowSkylightWarning, bShowPointLightWarning, bShowShadowedLightOverflowWarning,
-						bShowMobileLowQualityLightmapWarning, bShowMobileMovableDirectionalLightWarning, bShowMobileDynamicCSMWarning, bMobileMissingSkyMaterial, 
+						bShowMobileLowQualityLightmapWarning, bShowMobileDynamicCSMWarning, bMobileMissingSkyMaterial, 
 						bShowSkinCacheOOM, bSingleLayerWaterWarning, bShowNoSkyAtmosphereComponentWarning, bFxDebugDraw, FXInterface, bShowLocalExposureDisabledWarning,
 						bLumenEnabledButHasNoDataForTracing, bLumenEnabledButDisabledForTheProject, bNaniteEnabledButNoAtomics, bNaniteEnabledButDisabledInProject, 
 						bRealTimeSkyCaptureButNothingToCapture, bShowWaitingSkylight, bShowAmbientCubemapMipGenSettingsWarning, bLocalFogVolumeInSceneButProjectDisabled,
@@ -4142,11 +4141,6 @@ void FSceneRenderer::OnRenderFinish(FRDGBuilder& GraphBuilder, FRDGTextureRef Vi
 							static const FText Message = NSLOCTEXT("Renderer", "MobileLQLightmap", "MOBILE PROJECTS SUPPORTING STATIC LIGHTING MUST HAVE LQ LIGHTMAPS ENABLED");
 							Writer.DrawLine(Message);
 						}
-						if (bShowMobileMovableDirectionalLightWarning)
-						{
-							static const FText Message = NSLOCTEXT("Renderer", "MobileMovableDirectional", "PROJECT HAS MOVABLE DIRECTIONAL LIGHTS ON MOBILE DISABLED");
-							Writer.DrawLine(Message);
-						}
 						if (bShowMobileDynamicCSMWarning)
 						{
 							static const FText Message = (!FReadOnlyCVARCache::MobileEnableStaticAndCSMShadowReceivers())
diff --git a/Engine/Source/Runtime/Renderer/Private/ShadowSetup.cpp b/Engine/Source/Runtime/Renderer/Private/ShadowSetup.cpp
--- a/Engine/Source/Runtime/Renderer/Private/ShadowSetup.cpp
+++ b/Engine/Source/Runtime/Renderer/Private/ShadowSetup.cpp
@@ -6057,7 +6057,6 @@ void FSceneRenderer::CreateDynamicShadows(FDynamicShadowsTaskData& TaskData)
 	const bool bHairStrands = HairStrands::HasHairInstanceInScene(*Scene);
 
 	const bool bProjectEnablePointLightShadows = FReadOnlyCVARCache::EnablePointLightShadows() && !bMobile; // Point light shadow is unsupported on mobile for now.
-	const bool bProjectEnableMovableDirectionLightShadows = !bMobile || FReadOnlyCVARCache::MobileAllowMovableDirectionalLights();
 	const bool bProjectEnableMovableSpotLightShadows = !bMobile || IsMobileMovableSpotlightShadowsEnabled(ShaderPlatform);
 
 	uint32 NumPointShadowCachesUpdatedThisFrame = 0;
@@ -6133,7 +6132,6 @@ void FSceneRenderer::CreateDynamicShadows(FDynamicShadowsTaskData& TaskData)
 						const bool bCreateShadowForMovableLight =
 							bShouldCreateShadowForMovableLight
 							&& (!bPointLightShadow || bProjectEnablePointLightShadows)
-							&& (!bDirectionalLightShadow || bProjectEnableMovableDirectionLightShadows)
 							&& (!bSpotLightShadow || bProjectEnableMovableSpotLightShadows);
 
 						// Also create a whole scene shadow for lights with precomputed shadows that are unbuilt
