PR #12138: Fixed redundant threadgroups in motion blur classification dispatch (`FMotionBlurFilterTileClassifyCS`)

diff --git a/Engine/Source/Runtime/Renderer/Private/PostProcess/PostProcessMotionBlur.cpp b/Engine/Source/Runtime/Renderer/Private/PostProcess/PostProcessMotionBlur.cpp
--- a/Engine/Source/Runtime/Renderer/Private/PostProcess/PostProcessMotionBlur.cpp
+++ b/Engine/Source/Runtime/Renderer/Private/PostProcess/PostProcessMotionBlur.cpp
@@ -955,7 +955,7 @@ FMotionBlurOutputs AddMotionBlurFilterPass(
 			RDG_EVENT_NAME("MotionBlur FilterTileClassify %dx%d", FilterTileCount.X, FilterTileCount.Y),
 			ComputeShader,
 			PassParameters,
-			FComputeShaderUtils::GetGroupCount(FilterTileCount, 8));
+			FComputeShaderUtils::GetGroupCount(FilterTileCount, kMotionBlurFilterTileSize));
 	}
 
 	// Setup the filter's dispatch parameters.
