[HWRT] Use RDG_GPU_STAT_SCOPE for skinned and dynamic BLAS builds.

#rnx
#tests CitySample

diff --git a/Engine/Source/Runtime/Engine/Private/RayTracingSkinnedGeometry.cpp b/Engine/Source/Runtime/Engine/Private/RayTracingSkinnedGeometry.cpp
--- a/Engine/Source/Runtime/Engine/Private/RayTracingSkinnedGeometry.cpp
+++ b/Engine/Source/Runtime/Engine/Private/RayTracingSkinnedGeometry.cpp
@@ -320,15 +320,15 @@ void FRayTracingSkinnedGeometryUpdateQueue::Commit(FRDGBuilder& GraphBuilder, ER
 
 	if (GeometryBuildRequests.Num())
 	{
+		RDG_EVENT_SCOPE(GraphBuilder, "SkinnedGeometryBuildBLAS");
+		RDG_GPU_STAT_SCOPE(GraphBuilder, SkinnedGeometryBuildBLAS);
+
 		GraphBuilder.AddPass(RDG_EVENT_NAME("CommitRayTracingSkinnedGeometryUpdates"), BLASUpdateParams, ComputePassFlags | ERDGPassFlags::NeverCull,
 			[
 				BuildRequests = MoveTemp(GeometryBuildRequests),
 				SharedScratchBuffer
 			](FRHICommandList& RHICmdList)
 			{
-				SCOPED_GPU_STAT(RHICmdList, SkinnedGeometryBuildBLAS);
-				SCOPED_DRAW_EVENT(RHICmdList, SkinnedGeometryBuildBLAS);
-
 				FRHIBufferRange ScratchBufferRange;
 				ScratchBufferRange.Buffer = SharedScratchBuffer->GetRHI();
 				ScratchBufferRange.Offset = 0;
diff --git a/Engine/Source/Runtime/Renderer/Private/RayTracing/RayTracingDynamicGeometry.cpp b/Engine/Source/Runtime/Renderer/Private/RayTracing/RayTracingDynamicGeometry.cpp
--- a/Engine/Source/Runtime/Renderer/Private/RayTracing/RayTracingDynamicGeometry.cpp
+++ b/Engine/Source/Runtime/Renderer/Private/RayTracing/RayTracingDynamicGeometry.cpp
@@ -538,7 +538,9 @@ uint32 FRayTracingDynamicGeometryCollection::Update()
 void FRayTracingDynamicGeometryCollection::AddDynamicGeometryUpdatePass(const FViewInfo& View, FRDGBuilder& GraphBuilder, ERDGPassFlags ComputePassFlags, FRDGBufferRef& OutDynamicGeometryScratchBuffer)
 {
 	RDG_GPU_MASK_SCOPE(GraphBuilder, FRHIGPUMask::All());
-	
+	RDG_EVENT_SCOPE(GraphBuilder, "RayTracingDynamicGeometry");
+	RDG_GPU_STAT_SCOPE(GraphBuilder, RayTracingDynamicGeometry);
+
 	const uint32 ScratchAlignment = GRHIRayTracingScratchBufferAlignment;
 	const uint32 BLASScratchSize = Update();
 
@@ -560,7 +562,6 @@ void FRayTracingDynamicGeometryCollection::AddDynamicGeometryUpdatePass(const FV
 	GraphBuilder.AddPass(RDG_EVENT_NAME("RayTracingDynamicUpdate"), PassParams, ComputePassFlags | ERDGPassFlags::NeverCull,
 		[this, PassParams](FRHICommandList& RHICmdList)
 		{
-			SCOPED_GPU_STAT(RHICmdList, RayTracingDynamicGeometry);
 			FRHIBuffer* DynamicGeometryScratchBuffer = PassParams->DynamicGeometryScratchBuffer ? PassParams->DynamicGeometryScratchBuffer->GetRHI() : nullptr;
 
 			PRAGMA_DISABLE_DEPRECATION_WARNINGS
