Many Lights - HW Ray End Bias in order to prevent proxy geo self-occlusion near the light source

diff --git a/Engine/Shaders/Private/ManyLights/ManyLightsHardwareRayTracing.usf b/Engine/Shaders/Private/ManyLights/ManyLightsHardwareRayTracing.usf
--- a/Engine/Shaders/Private/ManyLights/ManyLightsHardwareRayTracing.usf
+++ b/Engine/Shaders/Private/ManyLights/ManyLightsHardwareRayTracing.usf
@@ -14,6 +14,7 @@
 #endif
 
 float RayTracingBias;
+float RayTracingEndBias;
 float RayTracingNormalBias;
 uint DebugMode;
 uint MaxTraversalIterations;
@@ -64,7 +65,7 @@ RAY_TRACING_ENTRY_RAYGEN_OR_INLINE(HardwareRayTraceLightSamples)
 		Ray.Origin = TranslatedWorldPosition;
 		Ray.Direction = LightSampleTrace.Direction;
 		Ray.TMin = max(LightSampleRayMinT, RayTracingBias);
-		Ray.TMax = LightSampleTrace.Distance;
+		Ray.TMax = max(Ray.TMin, LightSampleTrace.Distance - RayTracingEndBias);
 
 		ApplyPositionBias(Ray.Origin, Ray.Direction, SceneWorldNormal, RayTracingNormalBias);
 
diff --git a/Engine/Source/Runtime/Renderer/Private/ManyLights/ManyLightsRayTracing.cpp b/Engine/Source/Runtime/Renderer/Private/ManyLights/ManyLightsRayTracing.cpp
--- a/Engine/Source/Runtime/Renderer/Private/ManyLights/ManyLightsRayTracing.cpp
+++ b/Engine/Source/Runtime/Renderer/Private/ManyLights/ManyLightsRayTracing.cpp
@@ -78,6 +78,13 @@ static TAutoConsoleVariable<float> CVarManyLightsHardwareRayTracingBias(
 	ECVF_Scalability | ECVF_RenderThreadSafe
 );
 
+static TAutoConsoleVariable<float> CVarManyLightsHardwareRayTracingEndBias(
+	TEXT("r.ManyLights.HardwareRayTracing.EndBias"),
+	1.0f,
+	TEXT("Constant bias for hardware ray traced shadow rays to prevent proxy geo self-occlusion near the lights."),
+	ECVF_Scalability | ECVF_RenderThreadSafe
+);
+
 static TAutoConsoleVariable<float> CVarManyLightsHardwareRayTracingNormalBias(
 	TEXT("r.ManyLights.HardwareRayTracing.NormalBias"),
 	0.1f,
@@ -285,6 +292,7 @@ class FHardwareRayTraceLightSamples : public FLumenHardwareRayTracingShaderBase
 		SHADER_PARAMETER_RDG_TEXTURE_UAV(RWTexture2D<uint>, RWLightSamples)
 		SHADER_PARAMETER_RDG_TEXTURE(Texture2D<float>, LightSampleRayDistance)
 		SHADER_PARAMETER(float, RayTracingBias)
+		SHADER_PARAMETER(float, RayTracingEndBias)
 		SHADER_PARAMETER(float, RayTracingNormalBias)
 		// Ray Tracing
 		SHADER_PARAMETER(uint32, MaxTraversalIterations)
@@ -521,6 +529,7 @@ void ManyLights::SetHardwareRayTracingPassParameters(
 	PassParameters->RWLightSamples = GraphBuilder.CreateUAV(LightSamples);
 	PassParameters->LightSampleRayDistance = LightSampleRayDistance;
 	PassParameters->RayTracingBias = CVarManyLightsHardwareRayTracingBias.GetValueOnRenderThread();
+	PassParameters->RayTracingEndBias = CVarManyLightsHardwareRayTracingEndBias.GetValueOnRenderThread();
 	PassParameters->RayTracingNormalBias = CVarManyLightsHardwareRayTracingNormalBias.GetValueOnRenderThread();
 
 	checkf(View.HasRayTracingScene(), TEXT("TLAS does not exist. Verify that the current pass is represented in Lumen::AnyLumenHardwareRayTracingPassEnabled()."));
