Mobile: Fixed case with a missing velocity primitives in a scene captures

diff --git a/Engine/Source/Runtime/Renderer/Private/MobileShadingRenderer.cpp b/Engine/Source/Runtime/Renderer/Private/MobileShadingRenderer.cpp
--- a/Engine/Source/Runtime/Renderer/Private/MobileShadingRenderer.cpp
+++ b/Engine/Source/Runtime/Renderer/Private/MobileShadingRenderer.cpp
@@ -839,7 +839,8 @@ void FMobileSceneRenderer::RenderFullDepthPrepass(FRDGBuilder& GraphBuilder, TAr
 			});
 	}
 
-	if (bShouldRenderVelocities && Scene->EarlyZPassMode == DDM_AllOpaqueNoVelocity)
+	// Do not check for bShouldRenderVelocities here to avoid dropping velocity primitives from depth pre-pass
+	if (Scene->EarlyZPassMode == DDM_AllOpaqueNoVelocity) 
 	{
 		// Render the velocities and depth of movable objects
 		RenderVelocities(GraphBuilder, InViews, SceneTextures, EVelocityPass::Opaque, false);
diff --git a/Engine/Source/Runtime/Renderer/Private/VelocityRendering.cpp b/Engine/Source/Runtime/Renderer/Private/VelocityRendering.cpp
--- a/Engine/Source/Runtime/Renderer/Private/VelocityRendering.cpp
+++ b/Engine/Source/Runtime/Renderer/Private/VelocityRendering.cpp
@@ -246,11 +246,6 @@ void FSceneRenderer::RenderVelocities(
 	EVelocityPass VelocityPass,
 	bool bForceVelocity)
 {
-	if (!ShouldRenderVelocities())
-	{
-		return;
-	}
-
 	RDG_CSV_STAT_EXCLUSIVE_SCOPE(GraphBuilder, RenderVelocities);
 	SCOPED_NAMED_EVENT(FSceneRenderer_RenderVelocities, FColor::Emerald);
 	SCOPE_CYCLE_COUNTER(STAT_RenderVelocities);
