Changes r.TSR.16BitVALU.Nvidia's to enabled by default

diff --git a/Engine/Source/Runtime/Renderer/Private/PostProcess/TemporalSuperResolution.cpp b/Engine/Source/Runtime/Renderer/Private/PostProcess/TemporalSuperResolution.cpp
--- a/Engine/Source/Runtime/Renderer/Private/PostProcess/TemporalSuperResolution.cpp
+++ b/Engine/Source/Runtime/Renderer/Private/PostProcess/TemporalSuperResolution.cpp
@@ -110,7 +110,7 @@ TAutoConsoleVariable<int32> CVarTSR16BitVALUOnIntel(
 	ECVF_RenderThreadSafe);
 
 TAutoConsoleVariable<int32> CVarTSR16BitVALUOnNvidia(
-	TEXT("r.TSR.16BitVALU.Nvidia"), 0,
+	TEXT("r.TSR.16BitVALU.Nvidia"), 1,
 	TEXT("Overrides whether to use 16bit VALU on Nvidia desktop GPUs"),
 	ECVF_RenderThreadSafe);
 
