Substrate - Ignore complex shader shader compilation on metal to get in editor.
#rnx

diff --git a/Engine/Source/Runtime/Renderer/Private/Substrate/SubstrateVisualize.cpp b/Engine/Source/Runtime/Renderer/Private/Substrate/SubstrateVisualize.cpp
--- a/Engine/Source/Runtime/Renderer/Private/Substrate/SubstrateVisualize.cpp
+++ b/Engine/Source/Runtime/Renderer/Private/Substrate/SubstrateVisualize.cpp
@@ -28,7 +28,8 @@ static bool IsSubstrateDebugVisualizationSupported(EShaderPlatform InPlatform, b
 		Substrate::IsSubstrateEnabled() && 
 		Substrate::UsesSubstrateMaterialBuffer(InPlatform) &&
 		GetMaxSupportedFeatureLevel(InPlatform) >= ERHIFeatureLevel::SM5 &&
-		!IsVulkanPlatform(InPlatform) && // Does not compile and fails to package games. Fixing this would require a deeper investigation.
+		!IsVulkanPlatform(InPlatform) &&    // Does not compile and fails to package games. Fixing this would require a deeper investigation.
+        !IsMetalPlatform(InPlatform) &&     // Fails to launch editor on MAC
 		(bIsEditorOnly ? (IsPCPlatform(InPlatform) || EnumHasAllFlags(Flags, EShaderPermutationFlags::HasEditorOnlyData)) : true);
 }
 
