Fix PSO crashing with hair material shader in PPLL mode.

diff --git a/Engine/Source/Runtime/Renderer/Private/HairStrands/HairStrandsVisibility.cpp b/Engine/Source/Runtime/Renderer/Private/HairStrands/HairStrandsVisibility.cpp
--- a/Engine/Source/Runtime/Renderer/Private/HairStrands/HairStrandsVisibility.cpp
+++ b/Engine/Source/Runtime/Renderer/Private/HairStrands/HairStrandsVisibility.cpp
@@ -1788,6 +1788,11 @@ void FHairVisibilityProcessor::AddPSOInitializer(
 		SetupDepthStencilInfo(PF_DepthStencil, SceneTexturesConfig.DepthCreateFlags, ERenderTargetLoadAction::ELoad,
 			ERenderTargetLoadAction::ELoad, FExclusiveDepthStencil::DepthRead_StencilRead, RenderTargetsInfo);
 	}
+	else if (TRenderMode == HairVisibilityRenderMode_PPLL)
+	{
+		SetupDepthStencilInfo(PF_DepthStencil, SceneTexturesConfig.DepthCreateFlags, ERenderTargetLoadAction::ELoad,
+			ERenderTargetLoadAction::ELoad, FExclusiveDepthStencil::DepthRead_StencilRead, RenderTargetsInfo);
+	}
 
 	AddGraphicsPipelineStateInitializer(
 		VertexFactoryData,
