PLAY-33648

Disable benign VirtualShadowMapCacheManager warning due to it causing some automation spam

diff --git a/Engine/Source/Runtime/Renderer/Private/VirtualShadowMaps/VirtualShadowMapCacheManager.cpp b/Engine/Source/Runtime/Renderer/Private/VirtualShadowMaps/VirtualShadowMapCacheManager.cpp
--- a/Engine/Source/Runtime/Renderer/Private/VirtualShadowMaps/VirtualShadowMapCacheManager.cpp
+++ b/Engine/Source/Runtime/Renderer/Private/VirtualShadowMaps/VirtualShadowMapCacheManager.cpp
@@ -1514,7 +1514,9 @@ void FVirtualShadowMapArrayCacheManager::UpdateCachePrimitiveAsDynamic(FInvalida
 				// something that is no longer valid back to static. This could also potentially mean we incorrect transition a new thing that
 				// grabbed this slot back to static, but that is less likely as the addition would trigger a separate invalidation.
 				// Not much we can do here currently other than ignore it and move on
-				UE_LOG(LogRenderer, Display, TEXT("VirtualShadowMapCacheManager: Invalid persistent primitive index %d, age %u!"), PersistentPrimitiveIndex, InvalidationAge);
+				
+				// Disabling log due to build automation spam
+				// UE_LOG(LogRenderer, Display, TEXT("VirtualShadowMapCacheManager: Invalid persistent primitive index %d, age %u!"), PersistentPrimitiveIndex, InvalidationAge);
 				LastPrimitiveInvalidatedFrame[PersistentPrimitiveIndex] = 0xFFFFFFFF;
 			}
 			// NOTE: This is safe with the current set bit iterator, but should maybe use a temp array for future safety?
