ManyLights - expose option to trace 4 samples per pixel for scaling the quality up

diff --git a/Engine/Shaders/Private/ManyLights/ManyLights.ush b/Engine/Shaders/Private/ManyLights/ManyLights.ush
--- a/Engine/Shaders/Private/ManyLights/ManyLights.ush
+++ b/Engine/Shaders/Private/ManyLights/ManyLights.ush
@@ -3,7 +3,7 @@
 #pragma once
 
 // Change this to force recompilation of all Many Lights shaders
-#pragma message("UESHADERMETADATA_VERSION A48A689C-146E-4051-BAD6-21564B78503F")
+#pragma message("UESHADERMETADATA_VERSION 79F169FE-4FCA-47EA-8539-1ED66DAF713E")
 
 #define MANY_LIGHTS 1
 #include "../LightGridCommon.ush"
diff --git a/Engine/Shaders/Private/ManyLights/ManyLightsSampling.usf b/Engine/Shaders/Private/ManyLights/ManyLightsSampling.usf
--- a/Engine/Shaders/Private/ManyLights/ManyLightsSampling.usf
+++ b/Engine/Shaders/Private/ManyLights/ManyLightsSampling.usf
@@ -2,6 +2,8 @@
 
 #ifndef NUM_SAMPLES_PER_PIXEL_1D
 	#define NUM_SAMPLES_PER_PIXEL_1D 1
+	#define NUM_SAMPLES_PER_PIXEL_2D_X 1
+	#define NUM_SAMPLES_PER_PIXEL_2D_Y 1
 #endif
 
 #ifndef TILE_TYPE
@@ -183,14 +185,6 @@ void GenerateLightSamplesCS(
 				LightSamples[LightSampleIndex] = InitLightSample();
 			}
 
-			const uint2 LightSampleOffsets[4] = 
-			{
-				uint2(0, 0),
-				uint2(1, 0),
-				uint2(0, 1),
-				uint2(1, 1),
-			};
-
 			const float2 ScreenUV = (ScreenCoord + 0.5f) * View.BufferSizeAndInvSize.zw;
 			const FManyLightMaterial Material = LoadMaterial(ScreenUV, ScreenCoord, TILE_TYPE == TILE_MODE_SIMPLE_SHADING/*bForceSimpleShading*/);
 			const float SceneDepth = Material.Depth;
@@ -336,7 +330,7 @@ void GenerateLightSamplesCS(
 
 						if (DebugMode == DEBUG_MODE_VISUALIZE_SAMPLING)
 						{
-							const uint2 SampleCoord = DownsampledScreenCoord * DOWNSAMPLE_FACTOR + LightSampleOffsets[LightSampleIndex];
+							const uint2 SampleCoord = DownsampledScreenCoord * NumSamplesPerPixel + uint2(LightSampleIndex % NUM_SAMPLES_PER_PIXEL_2D_X, LightSampleIndex / NUM_SAMPLES_PER_PIXEL_2D_X);
 
 							const FLocalLightData LocalLightData = GetLocalLightDataNonStereo(LightSample.LocalLightIndex);
 							const FLightSampleTrace LightSampleTrace = GetLightSampleTrace(TranslatedWorldPosition, LightSample.LocalLightIndex, SampleCoord);
@@ -373,7 +367,7 @@ void GenerateLightSamplesCS(
 					LightSample.Weight = WeightSum / (NUM_SAMPLES_PER_PIXEL_1D * LightSample.Weight);
 				}
 
-				RWLightSamples[DownsampledScreenCoord * NumSamplesPerPixel + LightSampleOffsets[LightSampleIndex]] = PackLightSample(LightSample);
+				RWLightSamples[DownsampledScreenCoord * NumSamplesPerPixel + uint2(LightSampleIndex % NUM_SAMPLES_PER_PIXEL_2D_X, LightSampleIndex / NUM_SAMPLES_PER_PIXEL_2D_X)] = PackLightSample(LightSample);
 			}
 		}
 	}
diff --git a/Engine/Source/Runtime/Renderer/Private/ManyLights/ManyLights.cpp b/Engine/Source/Runtime/Renderer/Private/ManyLights/ManyLights.cpp
--- a/Engine/Source/Runtime/Renderer/Private/ManyLights/ManyLights.cpp
+++ b/Engine/Source/Runtime/Renderer/Private/ManyLights/ManyLights.cpp
@@ -19,9 +19,9 @@ static TAutoConsoleVariable<int32> CVarManyLightsNumSamplesPerPixel(
 	TEXT("r.ManyLights.NumSamplesPerPixel"),
 	4,
 	TEXT("Number of samples (shadow rays) per half-res pixel.\n")
-	TEXT("1 - 0.25 trace per pixel\n")
 	TEXT("2 - 0.5 trace per pixel\n")
-	TEXT("4 - 1 trace per pixel"),
+	TEXT("4 - 1 trace per pixel")
+	TEXT("16 - 4 traces per pixel"),
 	ECVF_Scalability | ECVF_RenderThreadSafe
 );
 
@@ -201,13 +201,23 @@ namespace ManyLights
 
 	FIntPoint GetNumSamplesPerPixel2d(int32 NumSamplesPerPixel1d)
 	{
-		return NumSamplesPerPixel1d == 4 ? FIntPoint(2, 2) : (NumSamplesPerPixel1d == 2 ? FIntPoint(2, 1) : FIntPoint(1, 1));
+		if (NumSamplesPerPixel1d >= 16)
+		{
+			return FIntPoint(4, 4);
+		}
+		else if (NumSamplesPerPixel1d >= 4)
+		{
+			return FIntPoint(2, 2);
+		}
+		else
+		{
+			return FIntPoint(2, 1);
+		}
 	}
 
 	FIntPoint GetNumSamplesPerPixel2d()
 	{
-		const uint32 NumSamplesPerPixel1d = FMath::RoundUpToPowerOfTwo(FMath::Clamp(CVarManyLightsNumSamplesPerPixel.GetValueOnAnyThread(), 1, 4));
-		return GetNumSamplesPerPixel2d(NumSamplesPerPixel1d);
+		return GetNumSamplesPerPixel2d(CVarManyLightsNumSamplesPerPixel.GetValueOnAnyThread());
 	}
 
 	int32 GetDebugMode()
@@ -325,7 +335,7 @@ class FGenerateLightSamplesCS : public FGlobalShader
 	class FIESProfile : SHADER_PERMUTATION_BOOL("USE_IES_PROFILE");
 	class FLightFunctionAtlas : SHADER_PERMUTATION_BOOL("USE_LIGHT_FUNCTION_ATLAS");
 	class FTexturedRectLights : SHADER_PERMUTATION_BOOL("USE_SOURCE_TEXTURE");
-	class FNumSamplesPerPixel1d : SHADER_PERMUTATION_SPARSE_INT("NUM_SAMPLES_PER_PIXEL_1D", 1, 2, 4);
+	class FNumSamplesPerPixel1d : SHADER_PERMUTATION_SPARSE_INT("NUM_SAMPLES_PER_PIXEL_1D", 2, 4, 16);
 	class FDebugMode : SHADER_PERMUTATION_BOOL("DEBUG_MODE");
 	using FPermutationDomain = TShaderPermutationDomain<FTileType, FIESProfile, FLightFunctionAtlas, FTexturedRectLights, FNumSamplesPerPixel1d, FDebugMode>;
 
@@ -376,6 +386,12 @@ class FGenerateLightSamplesCS : public FGlobalShader
 		ManyLights::ModifyCompilationEnvironment(Parameters.Platform, OutEnvironment);
 		OutEnvironment.SetDefine(TEXT("THREADGROUP_SIZE"), GetGroupSize());
 		OutEnvironment.CompilerFlags.Add(CFLAG_WaveOperations);
+
+		FPermutationDomain PermutationVector(Parameters.PermutationId);
+		const int32 NumSamplesPerPixel1d = PermutationVector.Get<FNumSamplesPerPixel1d>();
+		const FIntPoint NumSamplesPerPixel2d = ManyLights::GetNumSamplesPerPixel2d(NumSamplesPerPixel1d);
+		OutEnvironment.SetDefine(TEXT("NUM_SAMPLES_PER_PIXEL_2D_X"), NumSamplesPerPixel2d.X);
+		OutEnvironment.SetDefine(TEXT("NUM_SAMPLES_PER_PIXEL_2D_Y"), NumSamplesPerPixel2d.Y);
 	}
 };
 
@@ -473,7 +489,7 @@ class FShadeLightSamplesCS : public FGlobalShader
 	class FIESProfile : SHADER_PERMUTATION_BOOL("USE_IES_PROFILE");
 	class FLightFunctionAtlas : SHADER_PERMUTATION_BOOL("USE_LIGHT_FUNCTION_ATLAS");
 	class FTexturedRectLights : SHADER_PERMUTATION_BOOL("USE_SOURCE_TEXTURE");
-	class FNumSamplesPerPixel1d : SHADER_PERMUTATION_SPARSE_INT("NUM_SAMPLES_PER_PIXEL_1D", 1, 2, 4);
+	class FNumSamplesPerPixel1d : SHADER_PERMUTATION_SPARSE_INT("NUM_SAMPLES_PER_PIXEL_1D", 2, 4, 16);
 	class FDebugMode : SHADER_PERMUTATION_BOOL("DEBUG_MODE");
 	using FPermutationDomain = TShaderPermutationDomain<FTileType, FIESProfile, FLightFunctionAtlas, FTexturedRectLights, FNumSamplesPerPixel1d, FDebugMode>;
 
