Fix crash when visualizing light complexity with a groom component in view.

The remap permutation was missing the complex hair transmittance redirection.

diff --git a/Engine/Source/Runtime/Renderer/Private/LightRendering.cpp b/Engine/Source/Runtime/Renderer/Private/LightRendering.cpp
--- a/Engine/Source/Runtime/Renderer/Private/LightRendering.cpp
+++ b/Engine/Source/Runtime/Renderer/Private/LightRendering.cpp
@@ -994,6 +994,7 @@ class FDeferredLightPS : public FGlobalShader
 			PermutationVector.Set< FCloudTransmittance >(false);
 			PermutationVector.Set< FAnistropicMaterials >(false);
 			PermutationVector.Set< FVirtualShadowMapMask >(false);
+			PermutationVector.Set< FHairComplexTransmittance >(false);
 		}
 
 		return PermutationVector;
