Mobile Multi-view Fixes for Custom Depth and Custom Stencil

diff --git a/Engine/Shaders/Private/DeferredShadingCommon.ush b/Engine/Shaders/Private/DeferredShadingCommon.ush
--- a/Engine/Shaders/Private/DeferredShadingCommon.ush
+++ b/Engine/Shaders/Private/DeferredShadingCommon.ush
@@ -798,7 +798,7 @@ void MobileFetchGBuffer(in float2 UV, out half4 GBufferA, out half4 GBufferB, ou
 }
 #endif //MOBILE_DEFERRED_SHADING
 
-FGBufferData MobileFetchAndDecodeGBuffer(in float2 UV, in float2 PixelPos)
+FGBufferData MobileFetchAndDecodeGBuffer(in float2 UV, in float2 PixelPos, uint ViewId)
 {
 	FGBufferData GBuffer = (FGBufferData)0;
 #if (MOBILE_DEFERRED_SHADING && IS_MOBILE_DEFERREDSHADING_SUBPASS && PIXELSHADER)
@@ -813,12 +813,23 @@ FGBufferData MobileFetchAndDecodeGBuffer(in float2 UV, in float2 PixelPos)
 #else
 	GBuffer.Depth = CalcSceneDepth(UV);
 #endif
-
-	GBuffer.CustomDepth = ConvertFromDeviceZ(Texture2DSample(MobileSceneTextures.CustomDepthTexture, MobileSceneTextures.CustomDepthTextureSampler, UV).r);
+#if	MOBILE_MULTI_VIEW
+	GBuffer.CustomDepth = ConvertFromDeviceZ(Texture2DArraySample(MobileSceneTextures.CustomDepthTextureArray, MobileSceneTextures.CustomDepthTextureSampler, float3(UV, ViewId)).r);
+	GBuffer.CustomStencil = MobileSceneTextures.CustomStencilTextureArray.Load(int4(PixelPos.xy, ViewId, 0)) STENCIL_COMPONENT_SWIZZLE;
+#else
+	GBuffer.CustomDepth = ConvertFromDeviceZ(Texture2DSample(MobileSceneTextures.CustomDepthTexture, MobileSceneTextures.CustomDepthTextureSampler, float3(UV, ViewId)).r);
 	GBuffer.CustomStencil = MobileSceneTextures.CustomStencilTexture.Load(int3(PixelPos.xy, 0)) STENCIL_COMPONENT_SWIZZLE;
+#endif
+
 
 	return GBuffer;
 }
+
+FGBufferData MobileFetchAndDecodeGBuffer(in float2 UV, in float2 PixelPos)
+{
+	return MobileFetchAndDecodeGBuffer(UV, PixelPos, 0);
+}
+
 #endif //SHADING_PATH_MOBILE
 
 /** Populates OutGBufferA, B and C */
diff --git a/Engine/Shaders/Private/DepthOnlyVertexShader.usf b/Engine/Shaders/Private/DepthOnlyVertexShader.usf
--- a/Engine/Shaders/Private/DepthOnlyVertexShader.usf
+++ b/Engine/Shaders/Private/DepthOnlyVertexShader.usf
@@ -37,6 +37,8 @@ void Main(
 #endif
 #if INSTANCED_STEREO && MOBILE_MULTI_VIEW
 	, out uint LayerIndex : SV_RenderTargetArrayIndex
+#elif MOBILE_MULTI_VIEW
+	, in nointerpolation uint ViewId : SV_ViewID
 #elif INSTANCED_STEREO
 	, out uint ViewportIndex : SV_ViewPortArrayIndex
 #endif
@@ -50,7 +52,13 @@ void Main(
 		ViewportIndex = EyeIndex;
 	#endif
 #endif
+#if INSTANCED_STEREO
 	ResolvedView = ResolveViewFromVF(Input);
+#elif MOBILE_MULTI_VIEW
+	ResolvedView = ResolveView(ViewId);
+#else
+	ResolvedView = ResolveView();
+#endif
 
 	FVertexFactoryIntermediates VFIntermediates = GetVertexFactoryIntermediates(Input);
 	float4 WorldPos = VertexFactoryGetWorldPosition(Input, VFIntermediates);
diff --git a/Engine/Shaders/Private/MaterialTemplate.ush b/Engine/Shaders/Private/MaterialTemplate.ush
--- a/Engine/Shaders/Private/MaterialTemplate.ush
+++ b/Engine/Shaders/Private/MaterialTemplate.ush
@@ -505,6 +505,9 @@ struct FMaterialPixelParameters
 #if IS_NANITE_PASS
 	uint InstanceId;
 #endif
+#if MOBILE_MULTI_VIEW
+	uint ViewId;
+#endif
 
 #if USES_PER_INSTANCE_CUSTOM_DATA && VF_USE_PRIMITIVE_SCENE_DATA
 	uint CustomDataOffset;
@@ -1046,6 +1049,15 @@ float4 GetScreenPosition(FMaterialPixelParameters Parameters)
 	return Parameters.ScreenPosition;
 }
 
+uint GetViewId(FMaterialPixelParameters Parameters)
+{
+#if MOBILE_MULTI_VIEW
+	return Parameters.ViewId;
+#else
+	return 0;
+#endif
+}
+
 // Returns the pixel's depth, in world units. Works for both orthographic and perspective projections:
 float GetPixelDepth(FMaterialVertexParameters Parameters)
 {
@@ -2922,7 +2934,7 @@ MaterialFloat4 MobileSceneTextureLookup(inout FMaterialPixelParameters Parameter
 	}
 #else
 	const float2 PixelPos = UV * View.BufferSizeAndInvSize.xy;
-	FGBufferData GBuffer = MobileFetchAndDecodeGBuffer(UV, PixelPos);
+	FGBufferData GBuffer = MobileFetchAndDecodeGBuffer(UV, PixelPos, GetViewId(Parameters));
 
 	switch(SceneTextureId)
 	{
diff --git a/Engine/Shaders/Private/MobileBasePassPixelShader.usf b/Engine/Shaders/Private/MobileBasePassPixelShader.usf
--- a/Engine/Shaders/Private/MobileBasePassPixelShader.usf
+++ b/Engine/Shaders/Private/MobileBasePassPixelShader.usf
@@ -364,6 +364,9 @@ void Main(
 	FMaterialPixelParameters MaterialParameters = GetMaterialPixelParameters(Interpolants, SvPosition);
 	FPixelMaterialInputs PixelMaterialInputs;
 	{
+#if MOBILE_MULTI_VIEW
+		MaterialParameters.ViewId = BasePassInterpolants.MultiViewId;
+#endif
 		float4 ScreenPosition = SvPositionToResolvedScreenPosition(SvPosition);
 		float3 WorldPosition = BasePassInterpolants.PixelPosition.xyz;
 		float3 WorldPositionExcludingWPO = BasePassInterpolants.PixelPosition.xyz;
diff --git a/Engine/Source/Runtime/Engine/Public/SceneTexturesConfig.h b/Engine/Source/Runtime/Engine/Public/SceneTexturesConfig.h
--- a/Engine/Source/Runtime/Engine/Public/SceneTexturesConfig.h
+++ b/Engine/Source/Runtime/Engine/Public/SceneTexturesConfig.h
@@ -46,8 +46,10 @@ BEGIN_GLOBAL_SHADER_PARAMETER_STRUCT(FMobileSceneTextureUniformParameters, ENGIN
 	SHADER_PARAMETER_RDG_TEXTURE(Texture2D, ScenePartialDepthTexture)
 	SHADER_PARAMETER_SAMPLER(SamplerState, ScenePartialDepthTextureSampler)
 	SHADER_PARAMETER_RDG_TEXTURE(Texture2D, CustomDepthTexture)
+	SHADER_PARAMETER_RDG_TEXTURE(Texture2DArray, CustomDepthTextureArray)
 	SHADER_PARAMETER_SAMPLER(SamplerState, CustomDepthTextureSampler)
 	SHADER_PARAMETER_RDG_TEXTURE_SRV(Texture2D<uint2>, CustomStencilTexture)
+	SHADER_PARAMETER_RDG_TEXTURE_SRV(Texture2DArray<uint2>, CustomStencilTextureArray)
 	SHADER_PARAMETER_RDG_TEXTURE(Texture2D, SceneVelocityTexture)
 	SHADER_PARAMETER_SAMPLER(SamplerState, SceneVelocityTextureSampler)
 	// GBuffer
diff --git a/Engine/Source/Runtime/Renderer/Internal/CustomDepthRendering.h b/Engine/Source/Runtime/Renderer/Internal/CustomDepthRendering.h
--- a/Engine/Source/Runtime/Renderer/Internal/CustomDepthRendering.h
+++ b/Engine/Source/Runtime/Renderer/Internal/CustomDepthRendering.h
@@ -21,8 +21,7 @@ extern ECustomDepthPassLocation GetCustomDepthPassLocation(EShaderPlatform Platf
 
 struct FCustomDepthTextures
 {
-	static FCustomDepthTextures Create(FRDGBuilder& GraphBuilder, FIntPoint CustomDepthExtent, EShaderPlatform ShaderPlatform);
-
+	static FCustomDepthTextures Create(FRDGBuilder& GraphBuilder, FIntPoint CustomDepthExtent, EShaderPlatform ShaderPlatform, bool bRequireMultiView);
 	bool IsValid() const
 	{
 		return Depth != nullptr;
diff --git a/Engine/Source/Runtime/Renderer/Private/CustomDepthRendering.cpp b/Engine/Source/Runtime/Renderer/Private/CustomDepthRendering.cpp
--- a/Engine/Source/Runtime/Renderer/Private/CustomDepthRendering.cpp
+++ b/Engine/Source/Runtime/Renderer/Private/CustomDepthRendering.cpp
@@ -60,7 +60,7 @@ bool IsCustomDepthPassWritingStencil()
 	return GetCustomDepthMode() == ECustomDepthMode::EnabledWithStencil;
 }
 
-FCustomDepthTextures FCustomDepthTextures::Create(FRDGBuilder& GraphBuilder, FIntPoint CustomDepthExtent, EShaderPlatform ShaderPlatform)
+FCustomDepthTextures FCustomDepthTextures::Create(FRDGBuilder& GraphBuilder, FIntPoint CustomDepthExtent, EShaderPlatform ShaderPlatform, bool bRequireMultiView)
 {
 	if (!IsCustomDepthPassEnabled())
 	{
@@ -82,8 +82,9 @@ FCustomDepthTextures FCustomDepthTextures::Create(FRDGBuilder& GraphBuilder, FIn
 	{
 		CreateFlags |= TexCreate_NoFastClear;
 	}
-
-	const FRDGTextureDesc CustomDepthDesc = FRDGTextureDesc::Create2D(CustomDepthExtent, PF_DepthStencil, FClearValueBinding::DepthFar, CreateFlags);
+	FRDGTextureDesc CustomDepthDesc(bRequireMultiView ?
+		FRDGTextureDesc::Create2DArray(CustomDepthExtent, PF_DepthStencil, FClearValueBinding::DepthFar, CreateFlags, 2) :
+		FRDGTextureDesc::Create2D(CustomDepthExtent, PF_DepthStencil, FClearValueBinding::DepthFar, CreateFlags));
 
 	CustomDepthTextures.Depth = GraphBuilder.CreateTexture(CustomDepthDesc, TEXT("CustomDepth"));
 
@@ -275,7 +276,7 @@ bool FSceneRenderer::RenderCustomDepthPass(
 				DepthLoadAction,
 				StencilLoadAction,
 				FExclusiveDepthStencil::DepthWrite_StencilWrite);
-
+			PassParameters->RenderTargets.MultiViewCount = (View.bIsMobileMultiViewEnabled) ? 2 : (View.Aspects.IsMobileMultiViewEnabled() ? 1 : 0);
 			View.ParallelMeshDrawCommandPasses[EMeshPass::CustomDepth].BuildRenderingCommands(GraphBuilder, Scene->GPUScene, PassParameters->InstanceCullingDrawParams);
 
 			GraphBuilder.AddPass(
diff --git a/Engine/Source/Runtime/Renderer/Private/SceneTextures.cpp b/Engine/Source/Runtime/Renderer/Private/SceneTextures.cpp
--- a/Engine/Source/Runtime/Renderer/Private/SceneTextures.cpp
+++ b/Engine/Source/Runtime/Renderer/Private/SceneTextures.cpp
@@ -477,7 +477,7 @@ void FMinimalSceneTextures::InitializeViewFamily(FRDGBuilder& GraphBuilder, FVie
 	}
 
 	// Custom Depth
-	SceneTextures.CustomDepth = FCustomDepthTextures::Create(GraphBuilder, Config.Extent, Config.ShaderPlatform);
+	SceneTextures.CustomDepth = FCustomDepthTextures::Create(GraphBuilder, Config.Extent, Config.ShaderPlatform, Config.bRequireMultiView);
 
 	SceneTextures.bIsSceneTexturesInitialized = true;
 }
@@ -1158,8 +1158,10 @@ void SetupMobileSceneTextureUniformParameters(
 	SceneTextureParameters.ScenePartialDepthTextureSampler = TStaticSamplerState<>::GetRHI();
 	// CustomDepthTexture is a color texture on mobile, with DeviceZ values
 	SceneTextureParameters.CustomDepthTexture = SystemTextures.Black;
+	SceneTextureParameters.CustomDepthTextureArray = GSystemTextures.GetDefaultTexture(GraphBuilder, ETextureDimension::Texture2DArray, PF_DepthStencil, FClearValueBinding::Black);
 	SceneTextureParameters.CustomDepthTextureSampler = TStaticSamplerState<>::GetRHI();
 	SceneTextureParameters.CustomStencilTexture = SystemTextures.StencilDummySRV;
+	SceneTextureParameters.CustomStencilTextureArray = SystemTextures.StencilDummySRV;
 	SceneTextureParameters.SceneVelocityTexture = SystemTextures.Black;
 	SceneTextureParameters.SceneVelocityTextureSampler = TStaticSamplerState<>::GetRHI();
 	SceneTextureParameters.GBufferATexture = SystemTextures.Black;
@@ -1247,7 +1249,9 @@ void SetupMobileSceneTextureUniformParameters(
 
 			bool bCustomDepthProduced = HasBeenProduced(CustomDepthTextures.Depth);
 			SceneTextureParameters.CustomDepthTexture = bCustomDepthProduced ? CustomDepthTextures.Depth : SystemTextures.DepthDummy;
+			SceneTextureParameters.CustomDepthTextureArray = bCustomDepthProduced ? CustomDepthTextures.Depth : SystemTextures.DepthDummy;
 			SceneTextureParameters.CustomStencilTexture = bCustomDepthProduced ? CustomDepthTextures.Stencil : SystemTextures.StencilDummySRV;
+			SceneTextureParameters.CustomStencilTextureArray = bCustomDepthProduced ? CustomDepthTextures.Stencil : SystemTextures.StencilDummySRV;
 		}
 
 		if (EnumHasAnyFlags(SetupMode, EMobileSceneTextureSetupMode::SceneVelocity))
