Prevent virtual texture page pool overflow by clamping size of pools

diff --git a/Engine/Source/Runtime/Renderer/Private/VT/VirtualTextureSystem.cpp b/Engine/Source/Runtime/Renderer/Private/VT/VirtualTextureSystem.cpp
--- a/Engine/Source/Runtime/Renderer/Private/VT/VirtualTextureSystem.cpp
+++ b/Engine/Source/Runtime/Renderer/Private/VT/VirtualTextureSystem.cpp
@@ -947,18 +947,11 @@ void GetPhysicalSpaceExtraDescription(FVTPhysicalSpaceDescription const& InDesc,
 	int32 PoolCount = 1;
 	while (1)
 	{
-		const uint32 MaxTiles = FMath::Max((uint32)(PoolSizeInBytes / (PoolCount * TileSizeBytes)), 1u);
-
-		TileWidthHeight = FMath::FloorToInt(FMath::Sqrt((float)MaxTiles));
-
-		if (TileWidthHeight * InDesc.TileSize > GetMax2DTextureDimension())
-		{
-			// A good option to support extremely large caches would be to allow additional slices in an array here for caches...
-			// Just try to use the maximum texture size for now
-			TileWidthHeight = GetMax2DTextureDimension() / InDesc.TileSize;
-			break;
-		}
+		const uint32 NumTiles = FMath::Max((uint32)(PoolSizeInBytes / (PoolCount * TileSizeBytes)), 1u);
+		
+		TileWidthHeight = FMath::FloorToInt(FMath::Sqrt((float)NumTiles));
 
+		// If we allow splitting one pool confing into multiple physical pools by size then keep incrementing pool count until we fit.
 		const int32 SplitPhysicalPoolSize = VirtualTexturePool::GetSplitPhysicalPoolSize();
 		if (SplitPhysicalPoolSize <= 0 || TileWidthHeight <= SplitPhysicalPoolSize)
 		{
@@ -968,6 +961,20 @@ void GetPhysicalSpaceExtraDescription(FVTPhysicalSpaceDescription const& InDesc,
 		PoolCount++;
 	}
 
+	// Need to clamp for maximum texture size.
+	if (TileWidthHeight * InDesc.TileSize > GetMax2DTextureDimension())
+	{
+		// A good option to support extremely large caches would be to allow additional slices in an array here for caches...
+		// Just try to use the maximum texture size for now
+		TileWidthHeight = FMath::DivideAndRoundDown(GetMax2DTextureDimension(), InDesc.TileSize);
+	}
+
+	// Need to clamp for maximum tile count in a pool. 
+	// Page table encoding limits this to 1<<16. But FTexturePagePool::FreeHeap being 16bit and needing an overflow bit reduces us to 1<<15.
+	const int32 MaxTiles = 1 << 15;
+	const int32 MaxTilesSqrt = 181; //sqrt(1<<15)
+	TileWidthHeight = FMath::Min(TileWidthHeight, MaxTilesSqrt);
+
 	OutDescExt.TileWidthHeight = TileWidthHeight;
 	OutDescExt.PoolCount = PoolCount;
 	OutDescExt.bEnableResidencyMipMapBias = Config.bEnableResidencyMipMapBias;
