Heterogeneous Volumes: Add Ambient Occlusion calculation to live-shading pipeline

This submission adds an optional ambient occlusion term when applying indirect lighting from Lumen's Translucency Volume.

r.HeterogeneousVolumes.AmbientOcclusion (default=0): Enables ambient occlusion computation
r.HeterogeneousVolumes.AmbientOcclusion.DownsampleFactor (default=4): Performs downsampling when determining the ambient occlusion voxel resolution
r.HeterogeneousVolumes.AmbientOcclusion.RayCount.X (default=4): With the Y-counterpart, determines the number of AO rays when calculating ambient occlusion
r.HeterogeneousVolumes.AmbientOcclusion.RayCount.Y (default=4): With the X-counterpart, determines the number of AO rays when calculating ambient occlusion
r.HeterogeneousVolumes.AmbientOcclusion.MaxTraceDistance (default=10000): Determines the number of rays when calculating ambient occlusion
r.HeterogeneousVolumes.AmbientOcclusion.MaxStepCount (default=64): Determines the maximum steps when ray marching ambient occlusion

This submission also optionally creates an existence mask, Lighting Cache and Ambient Occlusion operations to voxels with non-zero extinction. To support conservative bounds, the extinction mask is dilated by one voxel of the operating resolution of the algorithm.

r.HeterogeneousVolumes.UseExistenceMask (default=1): Creates an evaluation mask which culls operations to the areas with non-zero extinction

diff --git a/Engine/Shaders/Private/HeterogeneousVolumes/HeterogeneousVolumesAmbientOcclusionPipeline.usf b/Engine/Shaders/Private/HeterogeneousVolumes/HeterogeneousVolumesAmbientOcclusionPipeline.usf
new file mode 100644
--- /dev/null
+++ b/Engine/Shaders/Private/HeterogeneousVolumes/HeterogeneousVolumesAmbientOcclusionPipeline.usf
@@ -0,0 +1,298 @@
+// Copyright Epic Games, Inc. All Rights Reserved.
+
+#include "../Common.ush"
+
+#define SUPPORT_CONTACT_SHADOWS 0
+
+#include "/Engine/Generated/Material.ush"
+#include "HeterogeneousVolumesLiveShadingUtils.ush"
+#include "HeterogeneousVolumesTracingUtils.ush"
+#include "HeterogeneousVolumesLightingUtils.ush"
+#include "../BlueNoise.ush"
+
+#ifndef THREADGROUP_SIZE_1D
+#define THREADGROUP_SIZE_1D 1
+#endif // THREADGROUP_SIZE_1D
+
+#ifndef THREADGROUP_SIZE_2D
+#define THREADGROUP_SIZE_2D 1
+#endif // THREADGROUP_SIZE_2D
+
+#ifndef THREADGROUP_SIZE_3D
+#define THREADGROUP_SIZE_3D 1
+#endif // THREADGROUP_SIZE_3D
+
+#ifndef USE_EXISTENCE_MASK
+#define USE_EXISTENCE_MASK 1
+#endif //
+
+// Object data
+float4x4 LocalToWorld;
+float4x4 WorldToLocal;
+float3 LocalBoundsOrigin;
+float3 LocalBoundsExtent;
+int PrimitiveId;
+
+// Volume data
+int3 VoxelResolution;
+int3 VoxelMin;
+int3 VoxelMax;
+
+// Shadowing
+float ShadowStepSize;
+float ShadowStepFactor;
+
+// Ray data
+float MaxTraceDistance;
+float MaxShadowTraceDistance;
+float StepSize;
+float StepFactor;
+int MaxStepCount;
+int bJitter;
+
+Texture3D ExistenceMaskTexture;
+
+// AO data
+int2 NumRays;
+
+FPrimitiveSceneData GetPrimitiveData(FMaterialVertexParameters Parameters)
+{
+	return GetPrimitiveData(Parameters, LocalToWorld, WorldToLocal, LocalBoundsOrigin, LocalBoundsExtent);
+}
+
+FPrimitiveSceneData GetPrimitiveData(FMaterialPixelParameters Parameters)
+{
+	return GetPrimitiveData(Parameters, LocalToWorld, WorldToLocal, LocalBoundsOrigin, LocalBoundsExtent);
+}
+
+float4x4 GetLocalToWorld()
+{
+	return LocalToWorld;
+}
+
+float4x4 GetWorldToLocal()
+{
+	return WorldToLocal;
+}
+
+float3 GetLocalBoundsOrigin()
+{
+	return LocalBoundsOrigin;
+}
+
+float3 GetLocalBoundsExtent()
+{
+	return LocalBoundsExtent;
+}
+
+int3 GetVolumeResolution()
+{
+	return VoxelResolution;
+}
+
+float GetStepSize()
+{
+	return StepSize;
+}
+
+float GetShadowStepSize()
+{
+	return ShadowStepSize;
+}
+
+float GetStepFactor()
+{
+	return StepFactor;
+}
+
+float GetShadowStepFactor()
+{
+	return ShadowStepFactor;
+}
+
+float GetMaxTraceDistance()
+{
+	return MaxTraceDistance;
+}
+
+float GetMaxShadowTraceDistance()
+{
+	return MaxShadowTraceDistance;
+}
+
+// Material sampling hooks..
+struct FVolumeSampleContext
+{
+	FMaterialPixelParameters MaterialParameters;
+	FPixelMaterialInputs PixelMaterialInputs;
+};
+
+FVolumeSampleContext CreateVolumeSampleContext(float3 LocalPosition, float3 WorldPosition, float MipLevel)
+{
+	FVolumeSampleContext VolumeSampleContext;
+	VolumeSampleContext.MaterialParameters = MakeInitializedMaterialPixelParameters();
+	//VolumeSampleContext.MaterialParameters.PrimitiveId = PrimitiveId;
+	VolumeSampleContext.MaterialParameters.AbsoluteWorldPosition = DFPromote(WorldPosition);
+	VolumeSampleContext.MaterialParameters.LWCData.AbsoluteWorldPosition = WSPromote(WorldPosition);
+	// TODO: Add object centroid to LWC.ObjectWorldPosition
+	VolumeSampleContext.MaterialParameters.LWCData.LocalToWorld = WSPromote(GetLocalToWorld());
+	VolumeSampleContext.MaterialParameters.LWCData.WorldToLocal = WSPromoteInverse(GetWorldToLocal());
+
+	// Evaluate Material graph
+	CalcPixelMaterialInputs(VolumeSampleContext.MaterialParameters, VolumeSampleContext.PixelMaterialInputs);
+
+	return VolumeSampleContext;
+}
+
+float3 SampleExtinction(inout FVolumeSampleContext Context)
+{
+	float3 Extinction = SampleExtinctionCoefficients(Context.PixelMaterialInputs);
+	return Extinction;
+}
+
+float3 SampleEmission(inout FVolumeSampleContext Context)
+{
+	float3 Emission = SampleEmissive(Context.PixelMaterialInputs);
+	return Emission;
+}
+
+float3 SampleAlbedo(inout FVolumeSampleContext Context)
+{
+	float3 Albedo = SampleAlbedo(Context.PixelMaterialInputs);
+	return Albedo;
+}
+
+#include "HeterogeneousVolumesRayMarchingUtils.ush"
+
+// Output
+RWTexture3D<uint> RWAmbientOcclusionUIntTexture;
+
+FRayMarchingContext CreateShadowRayMarchingContext(
+	float3 WorldRayOrigin,
+	float3 WorldToLight,
+	uint MaxStepCount
+)
+{
+	float3 WorldRayEnd = WorldRayOrigin + WorldToLight;
+	float3 WorldRayDirection = normalize(WorldToLight);
+
+	float3 LocalRayOrigin = mul(float4(WorldRayOrigin, 1.0), GetWorldToLocal()).xyz;
+	float3 LocalRayEnd = mul(float4(WorldRayEnd, 1.0), GetWorldToLocal()).xyz;
+	float3 LocalRayDirection = LocalRayEnd - LocalRayOrigin;
+	float LocalRayTMin = 0.0;
+	float LocalRayTMax = length(LocalRayDirection);
+	LocalRayDirection /= LocalRayTMax;
+
+	float ShadowBias = 0.5;
+	float ShadowStepSize = CalcShadowStepSize(LocalRayDirection);
+	int bApplyEmissionAndTransmittance = 0;
+	int bApplyDirectLighting = 0;
+	int bApplyShadowTransmittance = 0;
+
+	FRayMarchingContext ShadowRayMarchingContext = CreateRayMarchingContext(
+		LocalRayOrigin,
+		LocalRayDirection,
+		LocalRayTMin,
+		LocalRayTMax,
+		WorldRayOrigin,
+		WorldRayDirection,
+		ShadowBias,
+		ShadowStepSize,
+		MaxStepCount,
+		bApplyEmissionAndTransmittance,
+		bApplyDirectLighting,
+		bApplyShadowTransmittance
+		//RayMarchingContext.MaxShadowTraceDistance
+	);
+
+	return ShadowRayMarchingContext;
+}
+
+[numthreads(THREADGROUP_SIZE_3D, THREADGROUP_SIZE_3D, THREADGROUP_SIZE_3D)]
+void RenderAmbientOcclusionWithLiveShadingCS(
+	uint3 DispatchThreadId : SV_DispatchThreadID
+)
+{
+	int3 VoxelCount = VoxelMax - VoxelMin + 1;
+	int3 VoxelIndex = float3(DispatchThreadId.x, DispatchThreadId.y, DispatchThreadId.z % VoxelCount.z);
+	if (any(VoxelIndex >= VoxelCount))
+	{
+		return;
+	}
+	VoxelIndex += VoxelMin;
+
+	int LinearRayIndex = DispatchThreadId.z / VoxelCount.z;
+	int2 RayIndex = int2(LinearRayIndex % NumRays.x, LinearRayIndex / NumRays.x);
+	if (any(RayIndex >= NumRays))
+	{
+		return;
+	}
+
+	float3 VoxelJitter = 0.5;
+	if (bJitter)
+	{
+		uint3 Rand32Bits = Rand4DPCG32(uint4(VoxelIndex, View.StateFrameIndexMod8)).xyz;
+		VoxelJitter = float3(Rand32Bits) / float(uint(0xffffffff));
+	}
+	float3 UVW = (VoxelIndex + VoxelJitter) / float3(GetLightingCacheResolution());
+
+	float3 LocalBoundsMin = GetLocalBoundsOrigin() - GetLocalBoundsExtent();
+	float3 LocalBoundsMax = GetLocalBoundsOrigin() + GetLocalBoundsExtent();
+	float3 LocalRayOrigin = UVW * GetLocalBoundsExtent() * 2.0 + LocalBoundsMin;
+	float3 WorldRayOrigin = mul(float4(LocalRayOrigin, 1.0), GetLocalToWorld()).xyz;
+
+	// Construct ray direction
+	float2 RayUV = float2(RayIndex + 0.5) / NumRays;
+	float3 WorldRayDirection = EquiAreaSphericalMapping(RayUV);
+
+	// Evaluate transmittance based on MaxShadowDistance
+	FRayMarchingContext RayMarchingContext = CreateShadowRayMarchingContext(WorldRayOrigin, WorldRayDirection * MaxShadowTraceDistance, MaxStepCount);
+
+	float AmbientOcclusion = 1.0;
+#if USE_EXISTENCE_MASK
+	float Existence = ExistenceMaskTexture.SampleLevel(View.SharedBilinearClampedSampler, UVW, 0).r;
+	if (Existence > 0.0)
+#endif
+	{
+		float3 Transmittance = ComputeTransmittance(RayMarchingContext);
+		AmbientOcclusion = Luminance(Transmittance) * rcp(NumRays.x * NumRays.y);
+	}
+
+	// Convert to fixed-point and accumulate
+	int AmbientOcclusionFixedPoint = AmbientOcclusion * 65535.0;
+	InterlockedAdd(RWAmbientOcclusionUIntTexture[VoxelIndex], AmbientOcclusionFixedPoint);
+}
+
+RWTexture3D<float> RWExistenceMaskTexture;
+
+[numthreads(THREADGROUP_SIZE_3D, THREADGROUP_SIZE_3D, THREADGROUP_SIZE_3D)]
+void RenderExistenceMaskWithLiveShadingCS(
+	uint3 DispatchThreadId : SV_DispatchThreadID
+)
+{
+	int3 VoxelCount = VoxelMax - VoxelMin + 1;
+	int3 VoxelIndex = DispatchThreadId;
+	if (any(VoxelIndex > VoxelCount))
+	{
+		return;
+	}
+	VoxelIndex += VoxelMin;
+
+	float3 VoxelJitter = 0.5;
+	if (bJitter)
+	{
+		uint3 Rand32Bits = Rand4DPCG32(uint4(VoxelIndex, View.StateFrameIndexMod8)).xyz;
+		VoxelJitter = float3(Rand32Bits) / float(uint(0xffffffff));
+	}
+	float3 UVW = (VoxelIndex + VoxelJitter) / float3(VoxelResolution);
+
+	float3 LocalBoundsMin = GetLocalBoundsOrigin() - GetLocalBoundsExtent();
+	float3 LocalBoundsMax = GetLocalBoundsOrigin() + GetLocalBoundsExtent();
+	float3 LocalPosition = UVW * GetLocalBoundsExtent() * 2.0 + LocalBoundsMin;
+	float3 WorldPosition = mul(float4(LocalPosition, 1.0), GetLocalToWorld()).xyz;
+
+	int MipLevel = 0;
+	FVolumeSampleContext SampleContext = CreateVolumeSampleContext(LocalPosition, WorldPosition, MipLevel);
+	float3 Extinction = SampleExtinction(SampleContext);
+	RWExistenceMaskTexture[VoxelIndex] = any(Extinction > 0.0) ? 1.0 : 0.0;
+}
diff --git a/Engine/Shaders/Private/HeterogeneousVolumes/HeterogeneousVolumesHardwareRayTracing.usf b/Engine/Shaders/Private/HeterogeneousVolumes/HeterogeneousVolumesHardwareRayTracing.usf
--- a/Engine/Shaders/Private/HeterogeneousVolumes/HeterogeneousVolumesHardwareRayTracing.usf
+++ b/Engine/Shaders/Private/HeterogeneousVolumes/HeterogeneousVolumesHardwareRayTracing.usf
@@ -406,6 +406,11 @@ float3 ComputeTransmittanceHardwareRayTracing(
 	return Payload.GetTransmittance();
 }
 
+float EvalAmbientOcclusion(float3 WorldPosition)
+{
+	return 1.0;
+}
+
 #define HARDWARE_RAY_TRACING 1
 #include "HeterogeneousVolumesRayMarchingUtils.ush"
 
diff --git a/Engine/Shaders/Private/HeterogeneousVolumes/HeterogeneousVolumesLiveShadingGlobalPipeline.usf b/Engine/Shaders/Private/HeterogeneousVolumes/HeterogeneousVolumesLiveShadingGlobalPipeline.usf
--- a/Engine/Shaders/Private/HeterogeneousVolumes/HeterogeneousVolumesLiveShadingGlobalPipeline.usf
+++ b/Engine/Shaders/Private/HeterogeneousVolumes/HeterogeneousVolumesLiveShadingGlobalPipeline.usf
@@ -108,4 +108,42 @@ void ScreenTileClassificationCS(
 			RWScreenTileBuffer[BufferIndex] = ScreenTile;
 		}
 	}
-}
\ No newline at end of file
+}
+
+uint3 TextureResolution;
+Texture3D<uint> UIntTexture3D;
+RWTexture3D<float> RWFloatTexture3D;
+
+[numthreads(THREADGROUP_SIZE_3D, THREADGROUP_SIZE_3D, THREADGROUP_SIZE_3D)]
+void ConvertTexture3DFixedPointToFloatCS(
+	uint3 DispatchThreadId : SV_DispatchThreadID
+)
+{
+	float Value = float(UIntTexture3D[DispatchThreadId]);
+	RWFloatTexture3D[DispatchThreadId] = Value / 65535.0;
+}
+
+Texture3D ExistenceTexture;
+SamplerState TextureSampler;
+RWTexture3D<float> RWDilatedExistenceTexture;
+
+[numthreads(THREADGROUP_SIZE_3D, THREADGROUP_SIZE_3D, THREADGROUP_SIZE_3D)]
+void DilateExistenceMaskCS(
+	uint3 DispatchThreadId : SV_DispatchThreadID
+)
+{
+	if (any(DispatchThreadId > TextureResolution))
+	{
+		return;
+	}
+
+	float3 UVW = DispatchThreadId / (float3)TextureResolution;
+	int MipLevel = 0;
+	float Sample = Texture3DSampleLevel(
+		ExistenceTexture,
+		TextureSampler,
+		UVW,
+		MipLevel);
+
+	RWDilatedExistenceTexture[DispatchThreadId] = Sample > 0.0 ? 1.0 : 0.0;
+}
diff --git a/Engine/Shaders/Private/HeterogeneousVolumes/HeterogeneousVolumesLiveShadingPipeline.usf b/Engine/Shaders/Private/HeterogeneousVolumes/HeterogeneousVolumesLiveShadingPipeline.usf
--- a/Engine/Shaders/Private/HeterogeneousVolumes/HeterogeneousVolumesLiveShadingPipeline.usf
+++ b/Engine/Shaders/Private/HeterogeneousVolumes/HeterogeneousVolumesLiveShadingPipeline.usf
@@ -23,6 +23,10 @@
 #define REFERENCE_FAST_PATH 0
 #endif //
 
+#ifndef USE_EXISTENCE_MASK
+#define USE_EXISTENCE_MASK 0
+#endif //
+
 #define USE_CAMERA_AVSM (DIM_USE_ADAPTIVE_VOLUMETRIC_SHADOW_MAP && (SUPPORT_OVERLAPPING_VOLUMES != 0))
 
 #if SUPPORT_OVERLAPPING_VOLUMES
@@ -242,8 +246,26 @@ float3 SampleAlbedo(inout FVolumeSampleContext Context)
 	return Albedo;
 }
 
+int3 AmbientOcclusionResolution;
+Texture3D AmbientOcclusionTexture;
+
+float EvalAmbientOcclusion(float3 WorldPosition)
+{
+	float3 LocalPosition = mul(float4(WorldPosition, 1.0), GetWorldToLocal()).xyz;
+	float3 LocalBoundsMin = GetLocalBoundsOrigin() - GetLocalBoundsExtent();
+	float3 UVW = saturate((LocalPosition - LocalBoundsMin) / (2.0 * GetLocalBoundsExtent()));
+	float MipLevel = 0;
+
+	return Texture3DSampleLevel(
+		AmbientOcclusionTexture,
+		View.SharedTrilinearClampedSampler,
+		UVW,
+		MipLevel);
+}
+
 #include "HeterogeneousVolumesRayMarchingUtils.ush"
 
+Texture3D ExistenceMaskTexture;
 RWTexture3D<float3> RWLightingCacheTexture;
 
 [numthreads(THREADGROUP_SIZE_3D, THREADGROUP_SIZE_3D, THREADGROUP_SIZE_3D)]
@@ -275,10 +297,15 @@ void RenderLightingCacheWithLiveShadingCS(
 
 	// Existence culling
 	float MipLevel = 0.0f;
+#if USE_EXISTENCE_MASK
+	float Existence = ExistenceMaskTexture.SampleLevel(View.SharedBilinearClampedSampler, UVW, MipLevel).r;
+	if (Existence <= 0.0)
+#else
 	FVolumeSampleContext SampleContext = CreateVolumeSampleContext(LocalRayOrigin, WorldRayOrigin, MipLevel);
 	float3 Extinction = SampleExtinction(SampleContext);
 	float Epsilon = 1.0e-7;
 	if (all(Extinction < Epsilon))
+#endif
 	{
 #if DIM_LIGHTING_CACHE_MODE == LIGHTING_CACHE_TRANSMITTANCE
 		RWLightingCacheTexture[VoxelIndex] = 1.0f;
diff --git a/Engine/Shaders/Private/HeterogeneousVolumes/HeterogeneousVolumesPreshadingPipeline.usf b/Engine/Shaders/Private/HeterogeneousVolumes/HeterogeneousVolumesPreshadingPipeline.usf
--- a/Engine/Shaders/Private/HeterogeneousVolumes/HeterogeneousVolumesPreshadingPipeline.usf
+++ b/Engine/Shaders/Private/HeterogeneousVolumes/HeterogeneousVolumesPreshadingPipeline.usf
@@ -426,6 +426,12 @@ FTraceVoxelResult TraceVoxels(float3 LocalRayOrigin, float3 LocalRayDirection, f
 	return TraceVoxelResult;
 }
 
+// TODO: Implement ambient occlusion for preshading pipeline
+float EvalAmbientOcclusion(float3 WorldPosition)
+{
+	return 1.0;
+}
+
 #include "HeterogeneousVolumesRayMarchingUtils.ush"
 
 RWTexture3D<float3> RWLightingCacheTexture;
diff --git a/Engine/Shaders/Private/HeterogeneousVolumes/HeterogeneousVolumesRayMarchingUtils.ush b/Engine/Shaders/Private/HeterogeneousVolumes/HeterogeneousVolumesRayMarchingUtils.ush
--- a/Engine/Shaders/Private/HeterogeneousVolumes/HeterogeneousVolumesRayMarchingUtils.ush
+++ b/Engine/Shaders/Private/HeterogeneousVolumes/HeterogeneousVolumesRayMarchingUtils.ush
@@ -368,9 +368,10 @@ void RayMarchSingleScattering(
 			float3 Extinction = SampleExtinction(SampleContext);
 			float3 Albedo = SampleAlbedo(SampleContext);
 			float3 ScatteringCoefficient = Albedo * Extinction;
+			float AmbientOcclusion = EvalAmbientOcclusion(WorldPosition);
 
 			// Note: Phase is accounted for in spherical hamonic calculation
-			Radiance += IndirectInscattering * ScatteringCoefficient * Transmittance * StepSize;
+			Radiance += IndirectInscattering * AmbientOcclusion * ScatteringCoefficient * Transmittance * StepSize;
 #endif // DIM_USE_LUMEN_GI
 		}
 
diff --git a/Engine/Source/Runtime/Renderer/Private/HeterogeneousVolumes/HeterogeneousVolumes.h b/Engine/Source/Runtime/Renderer/Private/HeterogeneousVolumes/HeterogeneousVolumes.h
--- a/Engine/Source/Runtime/Renderer/Private/HeterogeneousVolumes/HeterogeneousVolumes.h
+++ b/Engine/Source/Runtime/Renderer/Private/HeterogeneousVolumes/HeterogeneousVolumes.h
@@ -60,6 +60,7 @@ namespace HeterogeneousVolumes
 	// CVars
 	FIntVector GetVolumeResolution(const IHeterogeneousVolumeInterface*);
 	FIntVector GetLightingCacheResolution(const IHeterogeneousVolumeInterface*, float LODFactor);
+	FIntVector GetAmbientOcclusionResolution(const IHeterogeneousVolumeInterface*, float LODFactor);
 
 	int32 GetDownsampleFactor();
 	FIntPoint GetScaledViewRect(FIntRect ViewRect);
@@ -108,6 +109,8 @@ namespace HeterogeneousVolumes
 	bool ShouldApplyHeightFog();
 	bool ShouldApplyVolumetricFog();
 	bool SupportsOverlappingVolumes();
+	bool EnableAmbientOcclusion();
+	bool UseExistenceMask();
 
 	enum class EFogMode
 	{
@@ -720,6 +723,49 @@ void GenerateRayTracingScene(
 	FRayTracingScene& RayTracingScene
 );
 
+void RenderExistenceMaskWithLiveShading(
+	FRDGBuilder& GraphBuilder,
+	// Scene data
+	const FScene* Scene,
+	const FViewInfo& View,
+	const FSceneTextures& SceneTextures,
+	// Object data
+	const IHeterogeneousVolumeInterface* HeterogeneousVolumeInterface,
+	const FMaterialRenderProxy* DefaultMaterialRenderProxy,
+	FPersistentPrimitiveIndex PersistentPrimitiveIndex,
+	const FBoxSphereBounds LocalBoxSphereBounds,
+	FIntVector ExistenceMaskTextureResolution,
+	// Output
+	FRDGTextureRef& ExistenceMaskTexture
+);
+
+void DilateExistenceMask(
+	FRDGBuilder& GraphBuilder,
+	// Scene data
+	const FScene* Scene,
+	const FViewInfo& View,
+	// Existence texture data
+	FRDGTextureRef ExistenceMaskTexture,
+	FIntVector ExistenceMaskTextureResolution,
+	// Output
+	FRDGTextureRef& DilatedExistenceTexture
+);
+
+void RenderAmbientOcclusionWithLiveShading(
+	FRDGBuilder& GraphBuilder,
+	// Scene data
+	const FScene* Scene,
+	const FViewInfo& View,
+	const FSceneTextures& SceneTextures,
+	// Object data
+	const IHeterogeneousVolumeInterface* HeterogeneousVolumeInterface,
+	const FMaterialRenderProxy* DefaultMaterialRenderProxy,
+	FPersistentPrimitiveIndex PersistentPrimitiveIndex,
+	const FBoxSphereBounds LocalBoxSphereBounds,
+	// Output
+	FRDGTextureRef& AmbientOcclusionTexture
+);
+
 void RenderLightingCacheWithPreshadingHardwareRayTracing(
 	FRDGBuilder& GraphBuilder,
 	// Scene data
diff --git a/Engine/Source/Runtime/Renderer/Private/HeterogeneousVolumes/HeterogeneousVolumesAmbientOcclusionPipeline.cpp b/Engine/Source/Runtime/Renderer/Private/HeterogeneousVolumes/HeterogeneousVolumesAmbientOcclusionPipeline.cpp
new file mode 100644
--- /dev/null
+++ b/Engine/Source/Runtime/Renderer/Private/HeterogeneousVolumes/HeterogeneousVolumesAmbientOcclusionPipeline.cpp
@@ -0,0 +1,821 @@
+// Copyright Epic Games, Inc. All Rights Reserved.
+
+#include "HeterogeneousVolumes.h"
+#include "HeterogeneousVolumeInterface.h"
+
+#include "LightRendering.h"
+#include "LocalVertexFactory.h"
+#include "MeshPassUtils.h"
+#include "PixelShaderUtils.h"
+#include "RayTracingDefinitions.h"
+#include "RayTracingInstance.h"
+#include "RayTracingInstanceBufferUtil.h"
+#include "RendererPrivate.h"
+#include "ScenePrivate.h"
+#include "SceneManagement.h"
+#include "VolumeLighting.h"
+#include "VolumetricFog.h"
+#include "BlueNoise.h"
+
+static TAutoConsoleVariable<int32> CVarHeterogeneousVolumesAmbientOcclusion(
+	TEXT("r.HeterogeneousVolumes.AmbientOcclusion"),
+	0,
+	TEXT("Enables ambient occlusion computation (Default = 0)"),
+	ECVF_RenderThreadSafe
+);
+
+static TAutoConsoleVariable<int32> CVarHeterogeneousVolumesAmbientOcclusionDownsampleFactor(
+	TEXT("r.HeterogeneousVolumes.AmbientOcclusion.DownsampleFactor"),
+	4,
+	TEXT("Performs downsampling when determining the ambient occlusion voxel resolution (Default = 4)"),
+	ECVF_RenderThreadSafe
+);
+
+static TAutoConsoleVariable<int32> CVarHeterogeneousVolumesAmbientOcclusionRayCountX(
+	TEXT("r.HeterogeneousVolumes.AmbientOcclusion.RayCount.X"),
+	4,
+	TEXT("With the Y-counterpart, determines the number of AO rays when calculating ambient occlusion (Default = 4)"),
+	ECVF_RenderThreadSafe
+);
+
+static TAutoConsoleVariable<int32> CVarHeterogeneousVolumesAmbientOcclusionRayCountY(
+	TEXT("r.HeterogeneousVolumes.AmbientOcclusion.RayCount.Y"),
+	4,
+	TEXT(" With the X-counterpart, determines the number of AO rays when calculating ambient occlusion (Default = 4)"),
+	ECVF_RenderThreadSafe
+);
+
+static TAutoConsoleVariable<int32> CVarHeterogeneousVolumesAmbientOcclusionMaxTraceDistance(
+	TEXT("r.HeterogeneousVolumes.AmbientOcclusion.MaxTraceDistance"),
+	1000.0,
+	TEXT("Determines the number of rays when calculating ambient occlusion (Default = 1000.0)"),
+	ECVF_RenderThreadSafe
+);
+
+static TAutoConsoleVariable<int32> CVarHeterogeneousVolumesAmbientOcclusionMaxStepCount(
+	TEXT("r.HeterogeneousVolumes.AmbientOcclusion.MaxStepCount"),
+	64,
+	TEXT("Determines the maximum steps when ray marching ambient occlusion (Default = 64)"),
+	ECVF_RenderThreadSafe
+);
+
+namespace HeterogeneousVolumes
+{
+	bool EnableAmbientOcclusion()
+	{
+		return CVarHeterogeneousVolumesAmbientOcclusion.GetValueOnRenderThread() != 0;
+	}
+
+	float GetAmbientOcclusionDownsampleFactor()
+	{
+		return FMath::Clamp(CVarHeterogeneousVolumesAmbientOcclusionDownsampleFactor.GetValueOnRenderThread(), 0.125, 32);
+	}
+
+	FIntPoint GetAmbientOcclusionRayCount()
+	{
+		return FIntPoint(
+			FMath::Clamp(CVarHeterogeneousVolumesAmbientOcclusionRayCountX.GetValueOnRenderThread(), 1, 8),
+			FMath::Clamp(CVarHeterogeneousVolumesAmbientOcclusionRayCountY.GetValueOnRenderThread(), 1, 8)
+		);
+	}
+
+	float GetAmbientOcclusionMaxTraceDistance()
+	{
+		return FMath::Clamp(CVarHeterogeneousVolumesAmbientOcclusionMaxTraceDistance.GetValueOnRenderThread(), 0.0f, GetMaxShadowTraceDistance());
+	}
+
+	int32 GetAmbientOcclusionMaxStepCount()
+	{
+		return FMath::Clamp(CVarHeterogeneousVolumesAmbientOcclusionMaxStepCount.GetValueOnRenderThread(), 1, GetMaxStepCount());
+	}
+
+	FIntVector GetAmbientOcclusionResolution(const IHeterogeneousVolumeInterface* RenderInterface, float LODFactor)
+	{
+		float DownsampleFactor = GetAmbientOcclusionDownsampleFactor() * LODFactor;
+		DownsampleFactor = FMath::Max(DownsampleFactor, 0.125f);
+
+		FVector VolumeResolution = FVector(GetVolumeResolution(RenderInterface));
+		FIntVector AmbientOcclusionResolution = FIntVector(VolumeResolution / DownsampleFactor);
+		AmbientOcclusionResolution.X = FMath::Clamp(AmbientOcclusionResolution.X, 1, 1024);
+		AmbientOcclusionResolution.Y = FMath::Clamp(AmbientOcclusionResolution.Y, 1, 1024);
+		AmbientOcclusionResolution.Z = FMath::Clamp(AmbientOcclusionResolution.Z, 1, 512);
+		return AmbientOcclusionResolution;
+	}
+
+} // namespace HeterogeneousVolumes
+
+//-OPT: Remove duplicate bindings
+// At the moment we need to bind the mesh draw parameters as they will be applied and on some RHIs this will crash if the texture is nullptr
+// We have the same parameters in the loose FParameters shader structure that are applied after the mesh draw.
+class FRenderAmbientOcclusionLooseBindings
+{
+	DECLARE_TYPE_LAYOUT(FRenderAmbientOcclusionLooseBindings, NonVirtual);
+
+public:
+	void Bind(const FShaderParameterMap& ParameterMap)
+	{
+		SceneDepthTextureBinding.Bind(ParameterMap, TEXT("SceneDepthTexture"));
+	}
+
+	template<typename TPassParameters>
+	void SetParameters(FMeshDrawSingleShaderBindings& ShaderBindings, const TPassParameters* PassParameters)
+	{
+		ShaderBindings.AddTexture(
+			SceneDepthTextureBinding,
+			FShaderResourceParameter(),
+			TStaticSamplerState<SF_Point>::GetRHI(),
+			PassParameters->SceneTextures.SceneDepthTexture->GetRHI()
+		);
+	}
+
+	LAYOUT_FIELD(FShaderResourceParameter, SceneDepthTextureBinding);
+};
+
+IMPLEMENT_TYPE_LAYOUT(FRenderAmbientOcclusionLooseBindings);
+
+class FRenderExistenceMaskWithLiveShadingCS : public FMeshMaterialShader
+{
+	DECLARE_SHADER_TYPE(FRenderExistenceMaskWithLiveShadingCS, MeshMaterial);
+
+	using FPermutationDomain = TShaderPermutationDomain<>;
+
+	BEGIN_SHADER_PARAMETER_STRUCT(FParameters, )
+		// Scene data
+		SHADER_PARAMETER_STRUCT_REF(FViewUniformShaderParameters, View)
+		SHADER_PARAMETER_STRUCT_INCLUDE(FSceneTextureParameters, SceneTextures)
+		SHADER_PARAMETER_RDG_UNIFORM_BUFFER(FSceneUniformParameters, Scene)
+
+		// Shadow data
+		SHADER_PARAMETER(float, ShadowStepSize)
+		SHADER_PARAMETER(float, ShadowStepFactor)
+
+		// Object data
+		SHADER_PARAMETER(FMatrix44f, LocalToWorld)
+		SHADER_PARAMETER(FMatrix44f, WorldToLocal)
+		SHADER_PARAMETER(FVector3f, LocalBoundsOrigin)
+		SHADER_PARAMETER(FVector3f, LocalBoundsExtent)
+		SHADER_PARAMETER(int32, PrimitiveId)
+
+		// Ray data
+		SHADER_PARAMETER(float, MaxTraceDistance)
+		SHADER_PARAMETER(float, MaxShadowTraceDistance)
+		SHADER_PARAMETER(float, StepSize)
+		SHADER_PARAMETER(float, StepFactor)
+		SHADER_PARAMETER(int, MaxStepCount)
+		SHADER_PARAMETER(int, bJitter)
+
+		// Volume data
+		SHADER_PARAMETER(FIntVector, VoxelResolution)
+		SHADER_PARAMETER(FIntVector, VoxelMin)
+		SHADER_PARAMETER(FIntVector, VoxelMax)
+
+		// Output
+		SHADER_PARAMETER_RDG_TEXTURE_UAV(RWTexture3D<uint>, RWExistenceMaskTexture)
+	END_SHADER_PARAMETER_STRUCT()
+
+	FRenderExistenceMaskWithLiveShadingCS() = default;
+
+	FRenderExistenceMaskWithLiveShadingCS(
+		const FMeshMaterialShaderType::CompiledShaderInitializerType& Initializer
+	)
+		: FMeshMaterialShader(Initializer)
+	{
+		Bindings.BindForLegacyShaderParameters(
+			this,
+			Initializer.PermutationId,
+			Initializer.ParameterMap,
+			*FParameters::FTypeInfo::GetStructMetadata(),
+			// Don't require full bindings, we use FMaterialShader::SetParameters
+			false
+		);
+		ShaderLooseBindings.Bind(Initializer.ParameterMap);
+	}
+
+	static bool ShouldCompilePermutation(
+		const FMaterialShaderPermutationParameters& Parameters
+	)
+	{
+		return DoesPlatformSupportHeterogeneousVolumes(Parameters.Platform)
+			&& DoesMaterialShaderSupportHeterogeneousVolumes(Parameters.MaterialParameters);
+	}
+
+	static void ModifyCompilationEnvironment(
+		const FMaterialShaderPermutationParameters& Parameters,
+		FShaderCompilerEnvironment& OutEnvironment
+	)
+	{
+		FMaterialShader::ModifyCompilationEnvironment(Parameters, OutEnvironment);
+		OutEnvironment.SetDefine(TEXT("THREADGROUP_SIZE_1D"), GetThreadGroupSize1D());
+		OutEnvironment.SetDefine(TEXT("THREADGROUP_SIZE_2D"), GetThreadGroupSize2D());
+		OutEnvironment.SetDefine(TEXT("THREADGROUP_SIZE_3D"), GetThreadGroupSize3D());
+
+		// This shader takes a very long time to compile with FXC, so we pre-compile it with DXC first and then forward the optimized HLSL to FXC.
+		//OutEnvironment.CompilerFlags.Add(CFLAG_PrecompileWithDXC); // @lh-todo - Disabled to workaround SPIRV-Cross bug: StructuredBuffer<uint> is translated to ByteAddressBuffer in HLSL backend
+		OutEnvironment.CompilerFlags.Add(CFLAG_AllowTypedUAVLoads);
+
+		OutEnvironment.SetDefine(TEXT("GET_PRIMITIVE_DATA_OVERRIDE"), 1);
+	}
+
+	static int32 GetThreadGroupSize1D() { return GetThreadGroupSize2D() * GetThreadGroupSize2D(); }
+	static int32 GetThreadGroupSize2D() { return 8; }
+	static int32 GetThreadGroupSize3D() { return 4; }
+
+	LAYOUT_FIELD(FRenderAmbientOcclusionLooseBindings, ShaderLooseBindings);
+};
+
+IMPLEMENT_MATERIAL_SHADER_TYPE(, FRenderExistenceMaskWithLiveShadingCS, TEXT("/Engine/Private/HeterogeneousVolumes/HeterogeneousVolumesAmbientOcclusionPipeline.usf"), TEXT("RenderExistenceMaskWithLiveShadingCS"), SF_Compute);
+
+void RenderExistenceMaskWithLiveShading(
+	FRDGBuilder& GraphBuilder,
+	// Scene data
+	const FScene* Scene,
+	const FViewInfo& View,
+	const FSceneTextures& SceneTextures,
+	// Object data
+	const IHeterogeneousVolumeInterface* HeterogeneousVolumeInterface,
+	const FMaterialRenderProxy* DefaultMaterialRenderProxy,
+	FPersistentPrimitiveIndex PersistentPrimitiveIndex,
+	const FBoxSphereBounds LocalBoxSphereBounds,
+	FIntVector ExistenceTextureResolution,
+	// Output
+	FRDGTextureRef& ExistenceMaskTexture
+)
+{
+	if (!HeterogeneousVolumes::UseExistenceMask())
+	{
+		FRDGTextureDesc ExistenceTextureDesc = FRDGTextureDesc::Create3D(
+			FIntVector(1),
+			PF_R8,
+			FClearValueBinding::Black,
+			TexCreate_ShaderResource | TexCreate_UAV | TexCreate_3DTiling
+		);
+
+		ExistenceMaskTexture = GraphBuilder.CreateTexture(ExistenceTextureDesc, TEXT("HeterogeneousVolumes.ExistenceMaskTexture"));
+		AddClearUAVPass(GraphBuilder, GraphBuilder.CreateUAV(ExistenceMaskTexture), 1.0f);
+		return;
+	}
+
+	FRDGTextureDesc ExistenceTextureDesc = FRDGTextureDesc::Create3D(
+		ExistenceTextureResolution,
+		PF_R8,
+		FClearValueBinding::Black,
+		TexCreate_ShaderResource | TexCreate_UAV | TexCreate_3DTiling
+	);
+	ExistenceMaskTexture = GraphBuilder.CreateTexture(ExistenceTextureDesc, TEXT("HeterogeneousVolumes.ExistenceMaskTexture"));
+
+	const FMaterialRenderProxy* MaterialRenderProxy = nullptr;
+	const FMaterial& Material = DefaultMaterialRenderProxy->GetMaterialWithFallback(View.GetFeatureLevel(), MaterialRenderProxy);
+	MaterialRenderProxy = MaterialRenderProxy ? MaterialRenderProxy : DefaultMaterialRenderProxy;
+
+	check(Material.GetMaterialDomain() == MD_Volume);
+
+	FRenderExistenceMaskWithLiveShadingCS::FParameters* PassParameters = GraphBuilder.AllocParameters<FRenderExistenceMaskWithLiveShadingCS::FParameters>();
+	{
+		// Scene data
+		PassParameters->View = View.ViewUniformBuffer;
+		PassParameters->SceneTextures = GetSceneTextureParameters(GraphBuilder, SceneTextures);
+		PassParameters->Scene = View.GetSceneUniforms().GetBuffer(GraphBuilder);
+
+		// Object data
+		// LWC_TODO: Convert to relative-local space
+		//FVector3f ViewOriginHigh = FDFVector3(View.ViewMatrices.GetViewOrigin()).High;
+		//FMatrix44f RelativeLocalToWorld = FDFMatrix::MakeToRelativeWorldMatrix(ViewOriginHigh, HeterogeneousVolumeInterface->GetLocalToWorld()).M;
+		FMatrix InstanceToLocal = HeterogeneousVolumeInterface->GetInstanceToLocal();
+		FMatrix LocalToWorld = HeterogeneousVolumeInterface->GetLocalToWorld();
+		PassParameters->LocalToWorld = FMatrix44f(InstanceToLocal * LocalToWorld);
+		PassParameters->WorldToLocal = PassParameters->LocalToWorld.Inverse();
+
+		// Ray data
+		PassParameters->ShadowStepSize = HeterogeneousVolumes::GetShadowStepSize();
+		PassParameters->ShadowStepFactor = HeterogeneousVolumeInterface->GetShadowStepFactor();
+		PassParameters->MaxTraceDistance = HeterogeneousVolumes::GetAmbientOcclusionMaxTraceDistance();
+		PassParameters->MaxShadowTraceDistance = HeterogeneousVolumes::GetAmbientOcclusionMaxTraceDistance();
+		PassParameters->StepSize = HeterogeneousVolumes::GetStepSize();
+		PassParameters->StepFactor = HeterogeneousVolumeInterface->GetStepFactor();
+		PassParameters->MaxStepCount = HeterogeneousVolumes::GetAmbientOcclusionMaxStepCount();
+		PassParameters->bJitter = 0;
+
+		FMatrix LocalToInstance = InstanceToLocal.Inverse();
+		FBoxSphereBounds InstanceBoxSphereBounds = LocalBoxSphereBounds.TransformBy(LocalToInstance);
+		PassParameters->LocalBoundsOrigin = FVector3f(InstanceBoxSphereBounds.Origin);
+		PassParameters->LocalBoundsExtent = FVector3f(InstanceBoxSphereBounds.BoxExtent);
+		PassParameters->PrimitiveId = PersistentPrimitiveIndex.Index;
+
+		// Ambient occlusion volume
+		PassParameters->VoxelResolution = ExistenceTextureResolution;
+		PassParameters->VoxelMin = FIntVector::ZeroValue;
+		PassParameters->VoxelMax = ExistenceTextureResolution - FIntVector(1);
+
+		// Output
+		PassParameters->RWExistenceMaskTexture = GraphBuilder.CreateUAV(ExistenceMaskTexture);
+	}
+
+	FIntVector GroupCount;
+	GroupCount.X = FMath::DivideAndRoundUp(ExistenceTextureResolution.X, FRenderExistenceMaskWithLiveShadingCS::GetThreadGroupSize3D());
+	GroupCount.Y = FMath::DivideAndRoundUp(ExistenceTextureResolution.Y, FRenderExistenceMaskWithLiveShadingCS::GetThreadGroupSize3D());
+	GroupCount.Z = FMath::DivideAndRoundUp(ExistenceTextureResolution.Z, FRenderExistenceMaskWithLiveShadingCS::GetThreadGroupSize3D());
+
+	FRenderExistenceMaskWithLiveShadingCS::FPermutationDomain PermutationVector;
+	TShaderRef<FRenderExistenceMaskWithLiveShadingCS> ComputeShader = Material.GetShader<FRenderExistenceMaskWithLiveShadingCS>(&FLocalVertexFactory::StaticType, PermutationVector, false);
+	if (!ComputeShader.IsNull())
+	{
+		GraphBuilder.AddPass(
+			RDG_EVENT_NAME("ExistenceMask"),
+			PassParameters,
+			ERDGPassFlags::Compute,
+			[ComputeShader, PassParameters, Scene, MaterialRenderProxy, &Material, GroupCount](FRHIComputeCommandList& RHICmdList)
+			{
+				FMeshMaterialShaderElementData ShaderElementData;
+				ShaderElementData.InitializeMeshMaterialData();
+
+				FMeshProcessorShaders PassShaders;
+				PassShaders.ComputeShader = ComputeShader;
+
+				FMeshDrawShaderBindings ShaderBindings;
+				ShaderBindings.Initialize(PassShaders);
+				{
+					FMeshDrawSingleShaderBindings SingleShaderBindings = ShaderBindings.GetSingleShaderBindings(SF_Compute);
+					ComputeShader->GetShaderBindings(Scene, Scene->GetFeatureLevel(), nullptr, *MaterialRenderProxy, Material, ShaderElementData, SingleShaderBindings);
+					ComputeShader->ShaderLooseBindings.SetParameters(SingleShaderBindings, PassParameters);
+					ShaderBindings.Finalize(&PassShaders);
+				}
+
+				UE::MeshPassUtils::Dispatch(RHICmdList, ComputeShader, ShaderBindings, *PassParameters, GroupCount);
+			}
+		);
+	}
+}
+
+class FDilateExistenceMaskCS : public FGlobalShader
+{
+	DECLARE_GLOBAL_SHADER(FDilateExistenceMaskCS);
+	SHADER_USE_PARAMETER_STRUCT(FDilateExistenceMaskCS, FGlobalShader);
+
+	BEGIN_SHADER_PARAMETER_STRUCT(FParameters, )
+		SHADER_PARAMETER(FIntVector, TextureResolution)
+		SHADER_PARAMETER_RDG_TEXTURE_SRV(Texture3D, ExistenceTexture)
+		SHADER_PARAMETER_SAMPLER(SamplerState, TextureSampler)
+		// Output
+		SHADER_PARAMETER_RDG_TEXTURE_UAV(RWTexture3D<float>, RWDilatedExistenceTexture)
+	END_SHADER_PARAMETER_STRUCT()
+
+	static bool ShouldCompilePermutation(const FGlobalShaderPermutationParameters& Parameters)
+	{
+		return DoesPlatformSupportHeterogeneousVolumes(Parameters.Platform);
+	}
+
+	static void ModifyCompilationEnvironment(
+		const FGlobalShaderPermutationParameters& Parameters,
+		FShaderCompilerEnvironment& OutEnvironment
+	)
+	{
+		FMaterialShader::ModifyCompilationEnvironment(Parameters, OutEnvironment);
+		OutEnvironment.SetDefine(TEXT("THREADGROUP_SIZE_1D"), GetThreadGroupSize1D());
+		OutEnvironment.SetDefine(TEXT("THREADGROUP_SIZE_2D"), GetThreadGroupSize2D());
+		OutEnvironment.SetDefine(TEXT("THREADGROUP_SIZE_3D"), GetThreadGroupSize3D());
+
+		// This shader takes a very long time to compile with FXC, so we pre-compile it with DXC first and then forward the optimized HLSL to FXC.
+		OutEnvironment.CompilerFlags.Add(CFLAG_AllowTypedUAVLoads);
+	}
+
+	static int32 GetThreadGroupSize1D() { return 64; }
+	static int32 GetThreadGroupSize2D() { return 8; }
+	static int32 GetThreadGroupSize3D() { return 4; }
+};
+
+IMPLEMENT_GLOBAL_SHADER(FDilateExistenceMaskCS, "/Engine/Private/HeterogeneousVolumes/HeterogeneousVolumesLiveShadingGlobalPipeline.usf", "DilateExistenceMaskCS", SF_Compute);
+
+void DilateExistenceMask(
+	FRDGBuilder& GraphBuilder,
+	// Scene data
+	const FScene* Scene,
+	const FViewInfo& View,
+	// Existence texture data
+	FRDGTextureRef ExistenceTexture,
+	FIntVector ExistenceTextureResolution,
+	// Output
+	FRDGTextureRef& DilatedExistenceTexture
+)
+{
+	if (!HeterogeneousVolumes::UseExistenceMask())
+	{
+		DilatedExistenceTexture = ExistenceTexture;
+		return;
+	}
+
+	FRDGTextureDesc ExistenceTextureDesc = FRDGTextureDesc::Create3D(
+		ExistenceTextureResolution,
+		PF_R8,
+		FClearValueBinding::Black,
+		TexCreate_ShaderResource | TexCreate_UAV | TexCreate_3DTiling
+	);
+	DilatedExistenceTexture = GraphBuilder.CreateTexture(ExistenceTextureDesc, TEXT("HeterogeneousVolumes.DilatedExistenceTexture"));
+
+	FDilateExistenceMaskCS::FParameters* PassParameters = GraphBuilder.AllocParameters<FDilateExistenceMaskCS::FParameters>();
+	{
+		PassParameters->TextureResolution = ExistenceTextureResolution;
+		PassParameters->ExistenceTexture = GraphBuilder.CreateSRV(ExistenceTexture);
+		PassParameters->TextureSampler = TStaticSamplerState<SF_Trilinear, AM_Clamp, AM_Clamp, AM_Clamp>::GetRHI();
+
+		PassParameters->RWDilatedExistenceTexture = GraphBuilder.CreateUAV(DilatedExistenceTexture);
+	}
+
+	FIntVector GroupCount;
+	GroupCount.X = FMath::DivideAndRoundUp(ExistenceTextureResolution.X, FDilateExistenceMaskCS::GetThreadGroupSize3D());
+	GroupCount.Y = FMath::DivideAndRoundUp(ExistenceTextureResolution.Y, FDilateExistenceMaskCS::GetThreadGroupSize3D());
+	GroupCount.Z = FMath::DivideAndRoundUp(ExistenceTextureResolution.Z, FDilateExistenceMaskCS::GetThreadGroupSize3D());
+
+	TShaderRef<FDilateExistenceMaskCS> ComputeShader = View.ShaderMap->GetShader<FDilateExistenceMaskCS>();
+	FComputeShaderUtils::AddPass(
+		GraphBuilder,
+		RDG_EVENT_NAME("DilateExistenceMaskCS"),
+		ComputeShader,
+		PassParameters,
+		GroupCount
+	);
+}
+
+class FRenderAmbientOcclusionWithLiveShadingCS : public FMeshMaterialShader
+{
+	DECLARE_SHADER_TYPE(FRenderAmbientOcclusionWithLiveShadingCS, MeshMaterial);
+
+	class FUseExistenceMask : SHADER_PERMUTATION_INT("USE_EXISTENCE_MASK", 2);
+	using FPermutationDomain = TShaderPermutationDomain<FUseExistenceMask>;
+
+	BEGIN_SHADER_PARAMETER_STRUCT(FParameters, )
+		// Scene data
+		SHADER_PARAMETER_STRUCT_REF(FViewUniformShaderParameters, View)
+		SHADER_PARAMETER_STRUCT_INCLUDE(FSceneTextureParameters, SceneTextures)
+		SHADER_PARAMETER_RDG_UNIFORM_BUFFER(FSceneUniformParameters, Scene)
+
+		// Shadow data
+		SHADER_PARAMETER(float, ShadowStepSize)
+		SHADER_PARAMETER(float, ShadowStepFactor)
+
+		// Object data
+		SHADER_PARAMETER(FMatrix44f, LocalToWorld)
+		SHADER_PARAMETER(FMatrix44f, WorldToLocal)
+		SHADER_PARAMETER(FVector3f, LocalBoundsOrigin)
+		SHADER_PARAMETER(FVector3f, LocalBoundsExtent)
+		SHADER_PARAMETER(int32, PrimitiveId)
+
+		// Ray data
+		SHADER_PARAMETER(float, MaxTraceDistance)
+		SHADER_PARAMETER(float, MaxShadowTraceDistance)
+		SHADER_PARAMETER(float, StepSize)
+		SHADER_PARAMETER(float, StepFactor)
+		SHADER_PARAMETER(int, MaxStepCount)
+		SHADER_PARAMETER(int, bJitter)
+
+		// Volume data
+		SHADER_PARAMETER(FIntVector, VoxelResolution)
+		SHADER_PARAMETER_STRUCT_INCLUDE(FLightingCacheParameters, AmbientOcclusion)
+		SHADER_PARAMETER(FIntVector, VoxelMin)
+		SHADER_PARAMETER(FIntVector, VoxelMax)
+
+		// AO data
+		SHADER_PARAMETER(FIntPoint, NumRays)
+
+		// Processing Mask
+		SHADER_PARAMETER_RDG_TEXTURE_SRV(Texture3D, ExistenceMaskTexture)
+
+		// Output
+		SHADER_PARAMETER_RDG_TEXTURE_UAV(RWTexture3D<uint>, RWAmbientOcclusionUIntTexture)
+	END_SHADER_PARAMETER_STRUCT()
+
+	FRenderAmbientOcclusionWithLiveShadingCS() = default;
+
+	FRenderAmbientOcclusionWithLiveShadingCS(
+		const FMeshMaterialShaderType::CompiledShaderInitializerType& Initializer
+	)
+		: FMeshMaterialShader(Initializer)
+	{
+		Bindings.BindForLegacyShaderParameters(
+			this,
+			Initializer.PermutationId,
+			Initializer.ParameterMap,
+			*FParameters::FTypeInfo::GetStructMetadata(),
+			// Don't require full bindings, we use FMaterialShader::SetParameters
+			false
+		);
+		ShaderLooseBindings.Bind(Initializer.ParameterMap);
+	}
+
+	static bool ShouldCompilePermutation(
+		const FMaterialShaderPermutationParameters& Parameters
+	)
+	{
+		return DoesPlatformSupportHeterogeneousVolumes(Parameters.Platform)
+			&& DoesMaterialShaderSupportHeterogeneousVolumes(Parameters.MaterialParameters);
+	}
+
+	static void ModifyCompilationEnvironment(
+		const FMaterialShaderPermutationParameters& Parameters,
+		FShaderCompilerEnvironment& OutEnvironment
+	)
+	{
+		FMaterialShader::ModifyCompilationEnvironment(Parameters, OutEnvironment);
+		OutEnvironment.SetDefine(TEXT("THREADGROUP_SIZE_1D"), GetThreadGroupSize1D());
+		OutEnvironment.SetDefine(TEXT("THREADGROUP_SIZE_2D"), GetThreadGroupSize2D());
+		OutEnvironment.SetDefine(TEXT("THREADGROUP_SIZE_3D"), GetThreadGroupSize3D());
+
+		// This shader takes a very long time to compile with FXC, so we pre-compile it with DXC first and then forward the optimized HLSL to FXC.
+		//OutEnvironment.CompilerFlags.Add(CFLAG_PrecompileWithDXC); // @lh-todo - Disabled to workaround SPIRV-Cross bug: StructuredBuffer<uint> is translated to ByteAddressBuffer in HLSL backend
+		OutEnvironment.CompilerFlags.Add(CFLAG_AllowTypedUAVLoads);
+
+		OutEnvironment.SetDefine(TEXT("GET_PRIMITIVE_DATA_OVERRIDE"), 1);
+	}
+
+	static int32 GetThreadGroupSize1D() { return GetThreadGroupSize2D() * GetThreadGroupSize2D(); }
+	static int32 GetThreadGroupSize2D() { return 8; }
+	static int32 GetThreadGroupSize3D() { return 4; }
+
+	LAYOUT_FIELD(FRenderAmbientOcclusionLooseBindings, ShaderLooseBindings);
+};
+
+IMPLEMENT_MATERIAL_SHADER_TYPE(, FRenderAmbientOcclusionWithLiveShadingCS, TEXT("/Engine/Private/HeterogeneousVolumes/HeterogeneousVolumesAmbientOcclusionPipeline.usf"), TEXT("RenderAmbientOcclusionWithLiveShadingCS"), SF_Compute);
+
+void RenderAmbientOcclusionWithLiveShadingAsFixedPoint(
+	FRDGBuilder& GraphBuilder,
+	// Scene data
+	const FScene* Scene,
+	const FViewInfo& View,
+	const FSceneTextures& SceneTextures,
+	// Object data
+	const IHeterogeneousVolumeInterface* HeterogeneousVolumeInterface,
+	const FMaterialRenderProxy* DefaultMaterialRenderProxy,
+	FPersistentPrimitiveIndex PersistentPrimitiveIndex,
+	const FBoxSphereBounds LocalBoxSphereBounds,
+	// Intermediary data
+	FRDGTextureRef ExistenceMaskTexture,
+	FIntVector AmbientOcclusionTextureResolution,
+	// Output
+	FRDGTextureRef& AmbientOcclusionUIntTexture
+)
+{
+	FRDGTextureDesc AmbientOcclusionDesc = FRDGTextureDesc::Create3D(
+		AmbientOcclusionTextureResolution,
+		PF_R32_UINT,
+		FClearValueBinding::Black,
+		TexCreate_ShaderResource | TexCreate_UAV | TexCreate_3DTiling
+	);
+	AmbientOcclusionUIntTexture = GraphBuilder.CreateTexture(AmbientOcclusionDesc, TEXT("HeterogeneousVolumes.AmbientOcclusionUIntTexture"));
+	AddClearUAVPass(GraphBuilder, GraphBuilder.CreateUAV(AmbientOcclusionUIntTexture), 0u);
+
+	// Build Ambient Occlusion score
+	const FMaterialRenderProxy* MaterialRenderProxy = nullptr;
+	const FMaterial& Material = DefaultMaterialRenderProxy->GetMaterialWithFallback(View.GetFeatureLevel(), MaterialRenderProxy);
+	MaterialRenderProxy = MaterialRenderProxy ? MaterialRenderProxy : DefaultMaterialRenderProxy;
+	check(Material.GetMaterialDomain() == MD_Volume);
+
+	FIntPoint RayCount = HeterogeneousVolumes::GetAmbientOcclusionRayCount();
+
+	FRenderAmbientOcclusionWithLiveShadingCS::FParameters* PassParameters = GraphBuilder.AllocParameters<FRenderAmbientOcclusionWithLiveShadingCS::FParameters>();
+	{
+		// Scene data
+		PassParameters->View = View.ViewUniformBuffer;
+		PassParameters->SceneTextures = GetSceneTextureParameters(GraphBuilder, SceneTextures);
+		PassParameters->Scene = View.GetSceneUniforms().GetBuffer(GraphBuilder);
+
+		// Object data
+		// LWC_TODO: Convert to relative-local space
+		//FVector3f ViewOriginHigh = FDFVector3(View.ViewMatrices.GetViewOrigin()).High;
+		//FMatrix44f RelativeLocalToWorld = FDFMatrix::MakeToRelativeWorldMatrix(ViewOriginHigh, HeterogeneousVolumeInterface->GetLocalToWorld()).M;
+		FMatrix InstanceToLocal = HeterogeneousVolumeInterface->GetInstanceToLocal();
+		FMatrix LocalToWorld = HeterogeneousVolumeInterface->GetLocalToWorld();
+		PassParameters->LocalToWorld = FMatrix44f(InstanceToLocal * LocalToWorld);
+		PassParameters->WorldToLocal = PassParameters->LocalToWorld.Inverse();
+
+		// Ray data
+		PassParameters->ShadowStepSize = HeterogeneousVolumes::GetShadowStepSize();
+		PassParameters->ShadowStepFactor = HeterogeneousVolumeInterface->GetShadowStepFactor();
+		PassParameters->MaxTraceDistance = HeterogeneousVolumes::GetAmbientOcclusionMaxTraceDistance();
+		PassParameters->MaxShadowTraceDistance = HeterogeneousVolumes::GetAmbientOcclusionMaxTraceDistance();
+		PassParameters->StepSize = HeterogeneousVolumes::GetStepSize();
+		PassParameters->StepFactor = HeterogeneousVolumeInterface->GetStepFactor();
+		PassParameters->MaxStepCount = HeterogeneousVolumes::GetAmbientOcclusionMaxStepCount();
+		PassParameters->bJitter = 0;
+
+		FMatrix LocalToInstance = InstanceToLocal.Inverse();
+		FBoxSphereBounds InstanceBoxSphereBounds = LocalBoxSphereBounds.TransformBy(LocalToInstance);
+		PassParameters->LocalBoundsOrigin = FVector3f(InstanceBoxSphereBounds.Origin);
+		PassParameters->LocalBoundsExtent = FVector3f(InstanceBoxSphereBounds.BoxExtent);
+		PassParameters->PrimitiveId = PersistentPrimitiveIndex.Index;
+
+		// Ambient occlusion volume
+		PassParameters->VoxelResolution = HeterogeneousVolumeInterface->GetVoxelResolution();
+		PassParameters->AmbientOcclusion.LightingCacheResolution = AmbientOcclusionTextureResolution;
+		PassParameters->AmbientOcclusion.LightingCacheVoxelBias = HeterogeneousVolumeInterface->GetShadowBiasFactor();
+		PassParameters->AmbientOcclusion.LightingCacheTexture = FRDGSystemTextures::Get(GraphBuilder).VolumetricBlack;
+
+		// Ambient occlusion data
+		PassParameters->NumRays = RayCount;
+		PassParameters->ExistenceMaskTexture = GraphBuilder.CreateSRV(ExistenceMaskTexture);
+
+		// Output
+		PassParameters->RWAmbientOcclusionUIntTexture = GraphBuilder.CreateUAV(AmbientOcclusionUIntTexture);
+		PassParameters->VoxelMin = FIntVector::ZeroValue;
+		PassParameters->VoxelMax = AmbientOcclusionTextureResolution - FIntVector(1);
+	}
+
+	FIntVector GroupCount;
+	GroupCount.X = FMath::DivideAndRoundUp(AmbientOcclusionTextureResolution.X, FRenderAmbientOcclusionWithLiveShadingCS::GetThreadGroupSize3D());
+	GroupCount.Y = FMath::DivideAndRoundUp(AmbientOcclusionTextureResolution.Y, FRenderAmbientOcclusionWithLiveShadingCS::GetThreadGroupSize3D());
+	GroupCount.Z = FMath::DivideAndRoundUp(AmbientOcclusionTextureResolution.Z, FRenderAmbientOcclusionWithLiveShadingCS::GetThreadGroupSize3D());
+	GroupCount.Z *= RayCount.X * RayCount.Y;
+
+	FRenderAmbientOcclusionWithLiveShadingCS::FPermutationDomain PermutationVector;
+	PermutationVector.Set<FRenderAmbientOcclusionWithLiveShadingCS::FUseExistenceMask>(HeterogeneousVolumes::UseExistenceMask());
+	TShaderRef<FRenderAmbientOcclusionWithLiveShadingCS> ComputeShader = Material.GetShader<FRenderAmbientOcclusionWithLiveShadingCS>(&FLocalVertexFactory::StaticType, PermutationVector, false);
+	GraphBuilder.AddPass(
+		RDG_EVENT_NAME("AmbientOcclusion"),
+		PassParameters,
+		ERDGPassFlags::Compute,
+		[ComputeShader, PassParameters, Scene, MaterialRenderProxy, &Material, GroupCount](FRHIComputeCommandList& RHICmdList)
+		{
+			FMeshMaterialShaderElementData ShaderElementData;
+			ShaderElementData.InitializeMeshMaterialData();
+
+			FMeshProcessorShaders PassShaders;
+			PassShaders.ComputeShader = ComputeShader;
+
+			FMeshDrawShaderBindings ShaderBindings;
+			ShaderBindings.Initialize(PassShaders);
+			{
+				FMeshDrawSingleShaderBindings SingleShaderBindings = ShaderBindings.GetSingleShaderBindings(SF_Compute);
+				ComputeShader->GetShaderBindings(Scene, Scene->GetFeatureLevel(), nullptr, *MaterialRenderProxy, Material, ShaderElementData, SingleShaderBindings);
+				ComputeShader->ShaderLooseBindings.SetParameters(SingleShaderBindings, PassParameters);
+				ShaderBindings.Finalize(&PassShaders);
+			}
+
+			UE::MeshPassUtils::Dispatch(RHICmdList, ComputeShader, ShaderBindings, *PassParameters, GroupCount);
+		}
+	);
+}
+
+class FConvertTexture3DFixedPointToFloatCS : public FGlobalShader
+{
+	DECLARE_GLOBAL_SHADER(FConvertTexture3DFixedPointToFloatCS);
+	SHADER_USE_PARAMETER_STRUCT(FConvertTexture3DFixedPointToFloatCS, FGlobalShader);
+
+	BEGIN_SHADER_PARAMETER_STRUCT(FParameters, )
+		SHADER_PARAMETER(FIntVector, TextureResolution)
+		SHADER_PARAMETER_RDG_TEXTURE_SRV(Texture3D<uint>, UIntTexture3D)
+		SHADER_PARAMETER_RDG_TEXTURE_UAV(RWTexture3D<float>, RWFloatTexture3D)
+		END_SHADER_PARAMETER_STRUCT()
+
+		static bool ShouldCompilePermutation(const FGlobalShaderPermutationParameters& Parameters)
+	{
+		return DoesPlatformSupportHeterogeneousVolumes(Parameters.Platform);
+	}
+
+	static void ModifyCompilationEnvironment(
+		const FGlobalShaderPermutationParameters& Parameters,
+		FShaderCompilerEnvironment& OutEnvironment
+	)
+	{
+		FMaterialShader::ModifyCompilationEnvironment(Parameters, OutEnvironment);
+		OutEnvironment.SetDefine(TEXT("THREADGROUP_SIZE_1D"), GetThreadGroupSize1D());
+		OutEnvironment.SetDefine(TEXT("THREADGROUP_SIZE_2D"), GetThreadGroupSize2D());
+		OutEnvironment.SetDefine(TEXT("THREADGROUP_SIZE_3D"), GetThreadGroupSize3D());
+
+		// This shader takes a very long time to compile with FXC, so we pre-compile it with DXC first and then forward the optimized HLSL to FXC.
+		OutEnvironment.CompilerFlags.Add(CFLAG_AllowTypedUAVLoads);
+	}
+
+	static int32 GetThreadGroupSize1D() { return 64; }
+	static int32 GetThreadGroupSize2D() { return 8; }
+	static int32 GetThreadGroupSize3D() { return 4; }
+};
+
+IMPLEMENT_GLOBAL_SHADER(FConvertTexture3DFixedPointToFloatCS, "/Engine/Private/HeterogeneousVolumes/HeterogeneousVolumesLiveShadingGlobalPipeline.usf", "ConvertTexture3DFixedPointToFloatCS", SF_Compute);
+
+void ConvertFixedPointToFloatingPoint(
+	FRDGBuilder& GraphBuilder,
+	const FScene* Scene,
+	const FViewInfo& View,
+	FIntVector AmbientOcclusionTextureResolution,
+	FRDGTextureRef AmbientOcclusionUIntTexture,
+	FRDGTextureRef& AmbientOcclusionTexture
+)
+{
+	FConvertTexture3DFixedPointToFloatCS::FParameters* PassParameters = GraphBuilder.AllocParameters<FConvertTexture3DFixedPointToFloatCS::FParameters>();
+	{
+		PassParameters->TextureResolution = AmbientOcclusionTextureResolution;
+		PassParameters->UIntTexture3D = GraphBuilder.CreateSRV(AmbientOcclusionUIntTexture);
+		PassParameters->RWFloatTexture3D = GraphBuilder.CreateUAV(AmbientOcclusionTexture);
+	}
+
+	FIntVector GroupCount;
+	GroupCount.X = FMath::DivideAndRoundUp(AmbientOcclusionTextureResolution.X, FConvertTexture3DFixedPointToFloatCS::GetThreadGroupSize3D());
+	GroupCount.Y = FMath::DivideAndRoundUp(AmbientOcclusionTextureResolution.Y, FConvertTexture3DFixedPointToFloatCS::GetThreadGroupSize3D());
+	GroupCount.Z = FMath::DivideAndRoundUp(AmbientOcclusionTextureResolution.Z, FConvertTexture3DFixedPointToFloatCS::GetThreadGroupSize3D());
+
+	FConvertTexture3DFixedPointToFloatCS::FPermutationDomain PermutationVector;
+	TShaderRef<FConvertTexture3DFixedPointToFloatCS> ComputeShader = View.ShaderMap->GetShader<FConvertTexture3DFixedPointToFloatCS>(PermutationVector);
+	FComputeShaderUtils::AddPass(
+		GraphBuilder,
+		RDG_EVENT_NAME("ConvertTexture3DFixedPointToFloatCS"),
+		ComputeShader,
+		PassParameters,
+		GroupCount
+	);
+}
+
+void RenderAmbientOcclusionWithLiveShading(
+	FRDGBuilder& GraphBuilder,
+	// Scene data
+	const FScene* Scene,
+	const FViewInfo& View,
+	const FSceneTextures& SceneTextures,
+	// Object data
+	const IHeterogeneousVolumeInterface* HeterogeneousVolumeInterface,
+	const FMaterialRenderProxy* DefaultMaterialRenderProxy,
+	FPersistentPrimitiveIndex PersistentPrimitiveIndex,
+	const FBoxSphereBounds LocalBoxSphereBounds,
+	// Output
+	FRDGTextureRef& AmbientOcclusionTexture
+)
+{
+	if (!HeterogeneousVolumes::EnableAmbientOcclusion() || !HeterogeneousVolumes::UseIndirectLighting())
+	{
+		FRDGTextureDesc AmbientOcclusionDesc = FRDGTextureDesc::Create3D(
+			FIntVector(1),
+			PF_R8,
+			FClearValueBinding::Black,
+			TexCreate_ShaderResource | TexCreate_UAV | TexCreate_3DTiling
+		);
+
+		AmbientOcclusionTexture = GraphBuilder.CreateTexture(AmbientOcclusionDesc, TEXT("HeterogeneousVolumes.AmbientOcclusionTexture"));
+		AddClearUAVPass(GraphBuilder, GraphBuilder.CreateUAV(AmbientOcclusionTexture), 1.0f);
+		return;
+	}
+
+	float LODFactor = HeterogeneousVolumes::CalcLODFactor(View, HeterogeneousVolumeInterface);
+	FIntVector AmbientOcclusionTextureResolution = HeterogeneousVolumes::GetAmbientOcclusionResolution(HeterogeneousVolumeInterface, LODFactor);
+
+	// Build existence mask
+	FRDGTextureRef ExistenceMaskTexture;
+	RenderExistenceMaskWithLiveShading(
+		GraphBuilder,
+		// Scene
+		Scene,
+		View,
+		SceneTextures,
+		// Object
+		HeterogeneousVolumeInterface,
+		DefaultMaterialRenderProxy,
+		PersistentPrimitiveIndex,
+		LocalBoxSphereBounds,
+		// Mask
+		AmbientOcclusionTextureResolution,
+		ExistenceMaskTexture
+	);
+
+	// Dilate by a pixel
+	FRDGTextureRef DilatedExistenceMaskTexture;
+	DilateExistenceMask(
+		GraphBuilder,
+		Scene,
+		View,
+		ExistenceMaskTexture,
+		AmbientOcclusionTextureResolution,
+		DilatedExistenceMaskTexture
+	);
+
+	// Calculate ambient occlusion
+	FRDGTextureRef AmbientOcclusionTextureAsFixedPoint;
+	RenderAmbientOcclusionWithLiveShadingAsFixedPoint(
+		GraphBuilder,
+		// Scene data
+		Scene,
+		View,
+		SceneTextures,
+		// Object data
+		HeterogeneousVolumeInterface,
+		DefaultMaterialRenderProxy,
+		PersistentPrimitiveIndex,
+		LocalBoxSphereBounds,
+		DilatedExistenceMaskTexture,
+		AmbientOcclusionTextureResolution,
+		// Output
+		AmbientOcclusionTextureAsFixedPoint
+	);
+
+	FRDGTextureDesc AmbientOcclusionDesc = FRDGTextureDesc::Create3D(
+		AmbientOcclusionTextureResolution,
+		PF_R8,
+		FClearValueBinding::Black,
+		TexCreate_ShaderResource | TexCreate_UAV | TexCreate_3DTiling
+	);
+	AmbientOcclusionTexture = GraphBuilder.CreateTexture(AmbientOcclusionDesc, TEXT("HeterogeneousVolumes.AmbientOcclusionTexture"));
+
+	ConvertFixedPointToFloatingPoint(
+		GraphBuilder,
+		Scene,
+		View,
+		AmbientOcclusionTextureResolution,
+		AmbientOcclusionTextureAsFixedPoint,
+		AmbientOcclusionTexture
+	);
+}
diff --git a/Engine/Source/Runtime/Renderer/Private/HeterogeneousVolumes/HeterogeneousVolumesLiveShadingPipeline.cpp b/Engine/Source/Runtime/Renderer/Private/HeterogeneousVolumes/HeterogeneousVolumesLiveShadingPipeline.cpp
--- a/Engine/Source/Runtime/Renderer/Private/HeterogeneousVolumes/HeterogeneousVolumesLiveShadingPipeline.cpp
+++ b/Engine/Source/Runtime/Renderer/Private/HeterogeneousVolumes/HeterogeneousVolumesLiveShadingPipeline.cpp
@@ -38,6 +38,13 @@ static TAutoConsoleVariable<int32> CVarHeterogeneousVolumesSupportOverlappingVol
 	ECVF_RenderThreadSafe
 );
 
+static TAutoConsoleVariable<int32> CVarHeterogeneousVolumesUseExistenceMask(
+	TEXT("r.HeterogeneousVolumes.UseExistenceMask"),
+	1,
+	TEXT("Creates an evaluation mask which culls operations to the areas with non-zero extinction (Default = 1)"),
+	ECVF_RenderThreadSafe
+);
+
 #if 0
 static TAutoConsoleVariable<int32> CVarHeterogeneousVolumesBilinearInterpolation(
 	TEXT("r.HeterogeneousVolumes.BilinearInterpolation"),
@@ -71,6 +78,11 @@ namespace HeterogeneousVolumes
 		return ShouldCompositeHeterogeneousVolumesWithTranslucency() &&
 			(CVarHeterogeneousVolumesSupportOverlappingVolumes.GetValueOnRenderThread() != 0);
 	}
+
+	bool UseExistenceMask()
+	{
+		return CVarHeterogeneousVolumesUseExistenceMask.GetValueOnRenderThread() != 0;
+	}
 #if 0
 	bool UseBilinearInterpolation()
 	{
@@ -164,7 +176,8 @@ class FRenderLightingCacheWithLiveShadingCS : public FMeshMaterialShader
 	class FLightingCacheMode : SHADER_PERMUTATION_INT("DIM_LIGHTING_CACHE_MODE", 2);
 	class FUseAdaptiveVolumetricShadowMap : SHADER_PERMUTATION_BOOL("DIM_USE_ADAPTIVE_VOLUMETRIC_SHADOW_MAP");
 	class FAVSMSampleMode : SHADER_PERMUTATION_INT("AVSM_SAMPLE_MODE", 2);
-	using FPermutationDomain = TShaderPermutationDomain<FLightingCacheMode, FUseAdaptiveVolumetricShadowMap, FAVSMSampleMode>;
+	class FUseExistenceMask : SHADER_PERMUTATION_INT("USE_EXISTENCE_MASK", 2);
+	using FPermutationDomain = TShaderPermutationDomain<FLightingCacheMode, FUseAdaptiveVolumetricShadowMap, FAVSMSampleMode, FUseExistenceMask>;
 
 	BEGIN_SHADER_PARAMETER_STRUCT(FParameters, )
 		// Scene data
@@ -209,6 +222,7 @@ class FRenderLightingCacheWithLiveShadingCS : public FMeshMaterialShader
 		SHADER_PARAMETER_STRUCT_INCLUDE(FLightingCacheParameters, LightingCache)
 		SHADER_PARAMETER(FIntVector, VoxelMin)
 		SHADER_PARAMETER(FIntVector, VoxelMax)
+		SHADER_PARAMETER_RDG_TEXTURE_SRV(Texture3D<float>, ExistenceMaskTexture)
 
 		// Output
 		SHADER_PARAMETER_RDG_TEXTURE_UAV(RWTexture3D<float>, RWLightingCacheTexture)
@@ -367,6 +381,9 @@ class FRenderSingleScatteringWithLiveShadingCS : public FMeshMaterialShader
 		SHADER_PARAMETER(FIntVector, VoxelResolution)
 		SHADER_PARAMETER_STRUCT_INCLUDE(FLightingCacheParameters, LightingCache)
 
+		SHADER_PARAMETER(FIntVector, AmbientOcclusionResolution)
+		SHADER_PARAMETER_RDG_TEXTURE(Texture3D, AmbientOcclusionTexture)
+
 		// Ray data
 		SHADER_PARAMETER(float, MaxTraceDistance)
 		SHADER_PARAMETER(float, StepSize)
@@ -614,6 +631,38 @@ static void RenderLightingCacheWithLiveShading(
 	TUniformBufferRef<FDeferredLightUniformStruct> DeferredLightUB = CreateUniformBufferImmediate(DeferredLightUniform, UniformBuffer_SingleDraw);
 	float LODFactor = HeterogeneousVolumes::CalcLODFactor(View, HeterogeneousVolumeInterface);
 
+	FIntVector LightingCacheResolution = HeterogeneousVolumes::GetLightingCacheResolution(HeterogeneousVolumeInterface, LODFactor);
+
+	FRDGTextureRef ExistenceMaskTexture = nullptr;
+	RenderExistenceMaskWithLiveShading(
+		GraphBuilder,
+		// Scene data
+		Scene,
+		View,
+		SceneTextures,
+		// Object data
+		HeterogeneousVolumeInterface,
+		DefaultMaterialRenderProxy,
+		PersistentPrimitiveIndex,
+		LocalBoxSphereBounds,
+		LightingCacheResolution,
+		// Output
+		ExistenceMaskTexture
+	);
+
+	FRDGTextureRef DilatedExistenceTexture;
+	DilateExistenceMask(
+		GraphBuilder,
+		// Scene data
+		Scene,
+		View,
+		// Existence texture data
+		ExistenceMaskTexture,
+		LightingCacheResolution,
+		// Output
+		DilatedExistenceTexture
+	);
+
 	FRenderLightingCacheWithLiveShadingCS::FParameters* PassParameters = GraphBuilder.AllocParameters<FRenderLightingCacheWithLiveShadingCS::FParameters>();
 	{
 		// Scene data
@@ -647,10 +696,11 @@ static void RenderLightingCacheWithLiveShading(
 
 		// Transmittance volume
 		PassParameters->VoxelResolution = HeterogeneousVolumeInterface->GetVoxelResolution();
-		PassParameters->LightingCache.LightingCacheResolution = HeterogeneousVolumes::GetLightingCacheResolution(HeterogeneousVolumeInterface, LODFactor);
+		PassParameters->LightingCache.LightingCacheResolution = LightingCacheResolution;
 		PassParameters->LightingCache.LightingCacheVoxelBias = HeterogeneousVolumeInterface->GetShadowBiasFactor();
 		//PassParameters->LightingCache.LightingCacheTexture = GraphBuilder.CreateSRV(LightingCacheTexture);
 		PassParameters->LightingCache.LightingCacheTexture = FRDGSystemTextures::Get(GraphBuilder).VolumetricBlack;
+		PassParameters->ExistenceMaskTexture = GraphBuilder.CreateSRV(DilatedExistenceTexture);
 
 		// Ray data
 		PassParameters->MaxTraceDistance = HeterogeneousVolumes::GetMaxTraceDistance();
@@ -750,9 +800,9 @@ static void RenderLightingCacheWithLiveShading(
 
 		FVector LocalLightMinUV = (LocalLightMin - LightingCacheMin) / (LightingCacheMax - LightingCacheMin);
 		FVector LocalLightMaxUV = (LocalLightMax - LightingCacheMin) / (LightingCacheMax - LightingCacheMin);
-		FVector LightingCacheResolution = FVector(PassParameters->LightingCache.LightingCacheResolution);
-		PassParameters->VoxelMin = ClampVector(FloorVector(LocalLightMinUV * LightingCacheResolution), FIntVector::ZeroValue, PassParameters->VoxelMax);
-		PassParameters->VoxelMax = ClampVector(CeilVector(LocalLightMaxUV * LightingCacheResolution), FIntVector::ZeroValue, PassParameters->VoxelMax);
+		FVector LightingCacheResolutionVector = FVector(PassParameters->LightingCache.LightingCacheResolution);
+		PassParameters->VoxelMin = ClampVector(FloorVector(LocalLightMinUV * LightingCacheResolutionVector), FIntVector::ZeroValue, PassParameters->VoxelMax);
+		PassParameters->VoxelMax = ClampVector(CeilVector(LocalLightMaxUV * LightingCacheResolutionVector), FIntVector::ZeroValue, PassParameters->VoxelMax);
 	}
 
 	FIntVector GroupCount = PassParameters->VoxelMax - PassParameters->VoxelMin + FIntVector(1);
@@ -768,6 +818,7 @@ static void RenderLightingCacheWithLiveShading(
 	PermutationVector.Set<FRenderLightingCacheWithLiveShadingCS::FLightingCacheMode>(HeterogeneousVolumes::GetLightingCacheMode() - 1);
 	PermutationVector.Set<FRenderLightingCacheWithLiveShadingCS::FUseAdaptiveVolumetricShadowMap>(bUseAVSM);
 	PermutationVector.Set<FRenderLightingCacheWithLiveShadingCS::FAVSMSampleMode>(AVSMSampleMode);
+	PermutationVector.Set<FRenderLightingCacheWithLiveShadingCS::FUseExistenceMask>(HeterogeneousVolumes::UseExistenceMask());
 	TShaderRef<FRenderLightingCacheWithLiveShadingCS> ComputeShader = Material.GetShader<FRenderLightingCacheWithLiveShadingCS>(&FLocalVertexFactory::StaticType, PermutationVector, false);
 	if (!ComputeShader.IsNull())
 	{
@@ -972,6 +1023,7 @@ void RenderSingleScatteringWithLiveShading(
 	const FBoxSphereBounds LocalBoxSphereBounds,
 	// Transmittance acceleration
 	FRDGTextureRef LightingCacheTexture,
+	FRDGTextureRef AmbientOcclusionTexture,
 	// Output
 	FRDGTextureRef& HeterogeneousVolumeTexture
 )
@@ -1130,6 +1182,9 @@ void RenderSingleScatteringWithLiveShading(
 			PassParameters->LightingCache.LightingCacheTexture = FRDGSystemTextures::Get(GraphBuilder).VolumetricBlack;
 		}
 
+		PassParameters->AmbientOcclusionResolution = HeterogeneousVolumes::GetAmbientOcclusionResolution(HeterogeneousVolumeInterface, LODFactor);
+		PassParameters->AmbientOcclusionTexture = AmbientOcclusionTexture;
+
 		// Dispatch data
 		PassParameters->GroupCount = GroupCount;
 		PassParameters->DownsampleFactor = HeterogeneousVolumes::GetDownsampleFactor();
@@ -1197,6 +1252,7 @@ static void RenderWithTransmittanceVolumePipeline(
 	const FBoxSphereBounds LocalBoxSphereBounds,
 	// Transmittance acceleration
 	FRDGTextureRef LightingCacheTexture,
+	FRDGTextureRef AmbientOcclusionTexture,
 	// Output
 	FRDGTextureRef& HeterogeneousVolumeRadiance
 )
@@ -1288,6 +1344,7 @@ static void RenderWithTransmittanceVolumePipeline(
 				LocalBoxSphereBounds,
 				// Transmittance acceleration
 				LightingCacheTexture,
+				AmbientOcclusionTexture,
 				// Output
 				HeterogeneousVolumeRadiance
 			);
@@ -1316,6 +1373,7 @@ static void RenderWithTransmittanceVolumePipeline(
 				LocalBoxSphereBounds,
 				// Transmittance acceleration
 				LightingCacheTexture,
+				AmbientOcclusionTexture,
 				// Output
 				HeterogeneousVolumeRadiance
 			);
@@ -1338,6 +1396,7 @@ static void RenderWithInscatteringVolumePipeline(
 	const FBoxSphereBounds LocalBoxSphereBounds,
 	// Transmittance acceleration
 	FRDGTextureRef LightingCacheTexture,
+	FRDGTextureRef AmbientOcclusionTexture,
 	// Output
 	FRDGTextureRef& HeterogeneousVolumeRadiance
 )
@@ -1436,6 +1495,7 @@ static void RenderWithInscatteringVolumePipeline(
 				LocalBoxSphereBounds,
 				// Transmittance acceleration
 				LightingCacheTexture,
+				AmbientOcclusionTexture,
 				// Output
 				HeterogeneousVolumeRadiance
 			);
@@ -1464,6 +1524,7 @@ static void RenderWithInscatteringVolumePipeline(
 				LocalBoxSphereBounds,
 				// Transmittance acceleration
 				LightingCacheTexture,
+				AmbientOcclusionTexture,
 				// Output
 				HeterogeneousVolumeRadiance
 			);
@@ -1490,6 +1551,22 @@ void RenderWithLiveShading(
 	FRDGTextureRef& HeterogeneousVolumeRadiance
 )
 {
+	FRDGTextureRef AmbientOcclusionTexture;
+	RenderAmbientOcclusionWithLiveShading(
+		GraphBuilder,
+		// Scene data
+		Scene,
+		View,
+		SceneTextures,
+		// Object data
+		HeterogeneousVolumeInterface,
+		MaterialRenderProxy,
+		PersistentPrimitiveIndex,
+		LocalBoxSphereBounds,
+		// Output
+		AmbientOcclusionTexture
+	);
+
 	if (HeterogeneousVolumes::UseLightingCacheForInscattering())
 	{
 		RenderWithInscatteringVolumePipeline(
@@ -1507,6 +1584,7 @@ void RenderWithLiveShading(
 			LocalBoxSphereBounds,
 			// Transmittance acceleration
 			LightingCacheTexture,
+			AmbientOcclusionTexture,
 			// Output
 			HeterogeneousVolumeRadiance
 		);
@@ -1528,6 +1606,7 @@ void RenderWithLiveShading(
 			LocalBoxSphereBounds,
 			// Transmittance acceleration
 			LightingCacheTexture,
+			AmbientOcclusionTexture,
 			// Output
 			HeterogeneousVolumeRadiance
 		);
