Fix for ensure due to calling GetShadowQuality from a Task Thread.

- Change GetValueOnRenderThread() to GetValueOnAnyThread().
- PSOs are collected on a task thread and it will call GetShadowQuality().
- GetShadowQuality() was assumed to only be called from the Render Thread .

diff --git a/Engine/Source/Runtime/Renderer/Private/ShadowRendering.cpp b/Engine/Source/Runtime/Renderer/Private/ShadowRendering.cpp
--- a/Engine/Source/Runtime/Renderer/Private/ShadowRendering.cpp
+++ b/Engine/Source/Runtime/Renderer/Private/ShadowRendering.cpp
@@ -262,13 +262,13 @@ uint32 GetShadowQuality()
 {
 	static const auto ICVarQuality = IConsoleManager::Get().FindTConsoleVariableDataInt(TEXT("r.ShadowQuality"));
 
-	int Ret = ICVarQuality->GetValueOnRenderThread();
+	int Ret = ICVarQuality->GetValueOnAnyThread();
 
 #if !(UE_BUILD_SHIPPING || UE_BUILD_TEST)
 	static const auto ICVarLimit = IConsoleManager::Get().FindTConsoleVariableDataInt(TEXT("r.LimitRenderingFeatures"));
 	if(ICVarLimit)
 	{
-		int32 Limit = ICVarLimit->GetValueOnRenderThread();
+		int32 Limit = ICVarLimit->GetValueOnAnyThread();
 
 		if(Limit > 2)
 		{
