Scene Capture:  Ability to force a capture to use the same resolution and/or camera as the main view, including matching screen percentage and FOV jitter, allowing a scene capture or custom render pass to be pixel accurate, and benefit from screen percentage for performance scaling.  Custom Render Passes pull resolution and jitter outputs from the main view during the main view's scene render, while Scene Captures copy the screen percentage and jitter inputs, to be able to match the main view's calculations.  The render thread copy of the output render target (FTextureRenderTarget2DResource) is resized to screen resolution, without modifying the original asset (UTextureRenderTarget2D), so there's no effect on the source content, or need for the user to modify the source content to fit potential screen sizes.

Option added to generate UserSceneTexture outputs for custom render passes rendering at main view resolution, allowing access to multiple outputs simultaneously in post process materials, from a single custom render pass.  Supports base color (albedo), normal, and scene color (emissive / unlit pixels), plus the custom render pass can output depth to the render target, allowing 4 total outputs.

diff --git a/Engine/Source/Editor/UnrealEd/Private/EditorViewportClient.cpp b/Engine/Source/Editor/UnrealEd/Private/EditorViewportClient.cpp
--- a/Engine/Source/Editor/UnrealEd/Private/EditorViewportClient.cpp
+++ b/Engine/Source/Editor/UnrealEd/Private/EditorViewportClient.cpp
@@ -4324,6 +4324,7 @@ void FEditorViewportClient::Draw(FViewport* InViewport, FCanvas* Canvas)
 	}
 
 	ViewFamily.bSplitScreenDebugAllowed = true;
+	ViewFamily.bIsMainViewFamily = !ViewFamily.EngineShowFlags.Wireframe;			// Prefer non-wireframe views as "main" view family
 
 	// Draw the 3D scene
 	GetRendererModule().BeginRenderingViewFamily(Canvas,&ViewFamily);
diff --git a/Engine/Source/Runtime/Engine/Classes/Components/SceneCaptureComponent.h b/Engine/Source/Runtime/Engine/Classes/Components/SceneCaptureComponent.h
--- a/Engine/Source/Runtime/Engine/Classes/Components/SceneCaptureComponent.h
+++ b/Engine/Source/Runtime/Engine/Classes/Components/SceneCaptureComponent.h
@@ -234,12 +234,17 @@ public:
 
 	static ENGINE_API void UpdateDeferredCaptures(FSceneInterface* Scene);
 
+	/** Whether this component is a USceneCaptureComponent2D */
+	virtual bool Is2D() const { return false; }
+
 	/** Whether this component is a USceneCaptureComponentCube */
 	virtual bool IsCube() const { return false; }
 
 	/** Updates "FrameUpdated" field, returns "true" if this is one of multiple scene captures this frame.  Called from FScene::UpdateSceneCaptureContents. */
 	ENGINE_API bool SetFrameUpdated();
 
+	virtual void UpdateSceneCaptureContents(FSceneInterface* Scene) {};
+
 protected:
 	/** Update the show flags from our show flags settings (ideally, you'd be able to set this more directly, but currently unable to make FEngineShowFlags a UStruct to use it as a FProperty...) */
 	ENGINE_API void UpdateShowFlags();
@@ -250,8 +255,6 @@ protected:
 
 	ENGINE_API bool IsCulledByDetailMode() const;
 
-	virtual void UpdateSceneCaptureContents(FSceneInterface* Scene) {};
-
 	/**
 	 * The view state holds persistent scene rendering state and enables occlusion culling in scene captures.
 	 * NOTE: This object is used by the rendering thread. When the game thread attempts to destroy it, FDeferredCleanupInterface will keep the object around until the RT is done accessing it.
diff --git a/Engine/Source/Runtime/Engine/Classes/Components/SceneCaptureComponent2D.h b/Engine/Source/Runtime/Engine/Classes/Components/SceneCaptureComponent2D.h
--- a/Engine/Source/Runtime/Engine/Classes/Components/SceneCaptureComponent2D.h
+++ b/Engine/Source/Runtime/Engine/Classes/Components/SceneCaptureComponent2D.h
@@ -118,7 +118,7 @@ public:
 	FVector ClipPlaneNormal;
 	
 	/** Render scene capture as additional render passes of the main renderer rather than as an independent renderer. Applies to scene depth, device depth, base color, and normal modes. */
-	UPROPERTY(EditAnywhere, BlueprintReadWrite, AdvancedDisplay, Category = SceneCapture, meta = (EditCondition = "CaptureSource == ESceneCaptureSource::SCS_SceneDepth || CaptureSource == ESceneCaptureSource::SCS_DeviceDepth || CaptureSource == ESceneCaptureSource::SCS_BaseColor || CaptureSource == ESceneCaptureSource::SCS_Normal"))
+	UPROPERTY(EditAnywhere, BlueprintReadWrite, Category=PassInfo, meta = (EditCondition = "CaptureSource == ESceneCaptureSource::SCS_SceneDepth || CaptureSource == ESceneCaptureSource::SCS_DeviceDepth || CaptureSource == ESceneCaptureSource::SCS_BaseColor || CaptureSource == ESceneCaptureSource::SCS_Normal"))
 	bool bRenderInMainRenderer = false;
 
 	/** 
@@ -133,12 +133,73 @@ public:
 	UPROPERTY(EditAnywhere, BlueprintReadWrite, Category = SceneCapture, meta = (DisplayName = "Fog only on rendered pixels"))
 	uint32 bConsiderUnrenderedOpaquePixelAsFullyTranslucent : 1;
 
+	/** Render with main view family (where bIsMainViewFamily == true) */
+	UPROPERTY(EditAnywhere, BlueprintReadWrite, Category=PassInfo)
+	uint32 bMainViewFamily : 1;
+
+	/** Render with main view resolution, ignoring the dimensions in the resource.  Enables Main View Family. */
+	UPROPERTY(EditAnywhere, BlueprintReadWrite, Category = PassInfo)
+	uint32 bMainViewResolution : 1;
+
+	/** Render with main view camera.  Enables Main View Family and Resolution.  Temporal AA jitter is matched with main view. */
+	UPROPERTY(EditAnywhere, BlueprintReadWrite, Category=PassInfo)
+	uint32 bMainViewCamera : 1;
+
+	/** When rendering with main view resolution, ignore screen percentage scale and render at full resolution.  Temporal AA jitter is also disabled. */
+	UPROPERTY(EditAnywhere, BlueprintReadWrite, Category = PassInfo)
+	uint32 bIgnoreScreenPercentage : 1;
+
+	/** Divisor when rendering at Main View Resolution. */
+	UPROPERTY(EditAnywhere, BlueprintReadWrite, Category=PassInfo, meta = (EditCondition = "bMainViewResolution || bMainViewCamera"))
+	FIntPoint MainViewResolutionDivisor = { 1, 1 };
+
+	/** Expose BaseColor as a UserSceneTexture.  Requires "Render In Main Renderer".  Enables Main View Family and Resolution, disables "Ignore Screen Percentage".  Useful to get multiple outputs from a Custom Render Pass. */
+	UPROPERTY(EditAnywhere, BlueprintReadWrite, Category=PassInfo, meta = (EditCondition = "bRenderInMainRenderer && (CaptureSource == ESceneCaptureSource::SCS_SceneDepth || CaptureSource == ESceneCaptureSource::SCS_DeviceDepth || CaptureSource == ESceneCaptureSource::SCS_BaseColor || CaptureSource == ESceneCaptureSource::SCS_Normal)"))
+	FName UserSceneTextureBaseColor;
+
+	/** Expose Normal as a UserSceneTexture.  Requires "Render In Main Renderer".  Enables Main View Family and Resolution, disables "Ignore Screen Percentage".  Useful to get multiple outputs from a Custom Render Pass. */
+	UPROPERTY(EditAnywhere, BlueprintReadWrite, Category=PassInfo, meta = (EditCondition = "bRenderInMainRenderer && (CaptureSource == ESceneCaptureSource::SCS_SceneDepth || CaptureSource == ESceneCaptureSource::SCS_DeviceDepth || CaptureSource == ESceneCaptureSource::SCS_BaseColor || CaptureSource == ESceneCaptureSource::SCS_Normal)"))
+	FName UserSceneTextureNormal;
+
+	/** Expose SceneColor (emissive/unlit) as a UserSceneTexture.  Requires "Render In Main Renderer".  Enables Main View Family and Resolution, disables "Ignore Screen Percentage".  Useful to get multiple outputs from a Custom Render Pass. */
+	UPROPERTY(EditAnywhere, BlueprintReadWrite, Category=PassInfo, meta = (EditCondition = "bRenderInMainRenderer && (CaptureSource == ESceneCaptureSource::SCS_SceneDepth || CaptureSource == ESceneCaptureSource::SCS_DeviceDepth || CaptureSource == ESceneCaptureSource::SCS_BaseColor || CaptureSource == ESceneCaptureSource::SCS_Normal)"))
+	FName UserSceneTextureSceneColor;
+
+	inline bool ShouldRenderInMainRenderer() const
+	{
+		return bRenderInMainRenderer && (CaptureSource == ESceneCaptureSource::SCS_SceneDepth || CaptureSource == ESceneCaptureSource::SCS_DeviceDepth || CaptureSource == ESceneCaptureSource::SCS_BaseColor || CaptureSource == ESceneCaptureSource::SCS_Normal);
+	}
+
+	inline bool ShouldRenderWithMainViewResolution() const
+	{
+		return bMainViewResolution || bMainViewCamera || (ShouldRenderInMainRenderer() && (!UserSceneTextureBaseColor.IsNone() || !UserSceneTextureNormal.IsNone() || !UserSceneTextureSceneColor.IsNone()));
+	}
+
+	inline bool ShouldRenderWithMainViewFamily() const
+	{
+		return bMainViewFamily || ShouldRenderWithMainViewResolution();
+	}
+
+	inline bool ShouldRenderWithMainViewCamera() const
+	{
+		return bMainViewCamera;
+	}
+
+	inline bool ShouldIgnoreScreenPercentage() const
+	{
+		// User Scene Texture outputs are always in the scaled view resolution, so ignore the bIgnoreScreenPercentage flag if either is set
+		return ShouldRenderInMainRenderer() ? bIgnoreScreenPercentage && UserSceneTextureBaseColor.IsNone() && UserSceneTextureNormal.IsNone() && UserSceneTextureSceneColor.IsNone() : bIgnoreScreenPercentage;
+	}
+
 	/** Array of scene view extensions specifically to apply to this scene capture */
 	TArray< TWeakPtr<ISceneViewExtension, ESPMode::ThreadSafe> > SceneViewExtensions;
 
 	/** Which tile to render of the orthographic view (ignored in Perspective mode) */
 	int32 TileID = 0;
 
+	/** Transient pointer set during calls to UpdateSceneCaptureContents from UpdateDeferredCaptures */
+	const FSceneViewFamily* MainViewFamily = nullptr;
+
 	//~ Begin UActorComponent Interface
 	ENGINE_API virtual void OnComponentDestroyed(bool bDestroyingHierarchy) override;
 	ENGINE_API virtual void OnRegister() override;
@@ -206,6 +267,9 @@ public:
 	/* Return number of Y tiles to render (to be used when orthographic tiling rendering is enabled) */
 	ENGINE_API int32 GetNumYTiles() const;
 
+	/** Whether this component is a USceneCaptureComponent2D */
+	virtual bool Is2D() const override { return true; }
+
 #if WITH_EDITORONLY_DATA
 	ENGINE_API void UpdateDrawFrustum();
 
diff --git a/Engine/Source/Runtime/Engine/Internal/SceneCapture/SceneCaptureInternal.h b/Engine/Source/Runtime/Engine/Internal/SceneCapture/SceneCaptureInternal.h
new file mode 100644
--- /dev/null
+++ b/Engine/Source/Runtime/Engine/Internal/SceneCapture/SceneCaptureInternal.h
@@ -0,0 +1,14 @@
+// Copyright Epic Games, Inc. All Rights Reserved.
+
+#pragma once
+
+#include "CoreMinimal.h"
+
+class FSceneInterface;
+class FSceneViewFamily;
+
+/**
+ * Internal function called from FRendererModule::BeginRenderingViewFamilies.  List of view families is required to handle captures where
+ * USceneCaptureComponent2D::ShouldRenderWithMainViewFamily() is true.
+ */
+ENGINE_API void SceneCaptureUpdateDeferredCapturesInternal(FSceneInterface* Scene, TArray<const FSceneViewFamily*>& InOutViewFamilies);
diff --git a/Engine/Source/Runtime/Engine/Private/Components/SceneCaptureComponent.cpp b/Engine/Source/Runtime/Engine/Private/Components/SceneCaptureComponent.cpp
--- a/Engine/Source/Runtime/Engine/Private/Components/SceneCaptureComponent.cpp
+++ b/Engine/Source/Runtime/Engine/Private/Components/SceneCaptureComponent.cpp
@@ -24,6 +24,7 @@
 #include "Components/SceneCaptureComponent2D.h"
 #include "Engine/SceneCaptureCube.h"
 #include "Components/SceneCaptureComponentCube.h"
+#include "SceneCapture/SceneCaptureInternal.h"
 #include "Components/DrawFrustumComponent.h"
 #include "Engine/PlanarReflection.h"
 #include "Components/PlanarReflectionComponent.h"
@@ -517,6 +518,12 @@ void USceneCaptureComponent::Serialize(FArchive& Ar)
 }
 
 void USceneCaptureComponent::UpdateDeferredCaptures(FSceneInterface* Scene)
+{
+	TArray<const FSceneViewFamily*> EmptyViewFamilies;
+	SceneCaptureUpdateDeferredCapturesInternal(Scene, EmptyViewFamilies);
+}
+
+void SceneCaptureUpdateDeferredCapturesInternal(FSceneInterface* Scene, TArray<const FSceneViewFamily*>& InOutViewFamilies)
 {
 	FScopeLock ScopeLock(&SceneCapturesToUpdateMapCS);
 	UWorld* World = Scene->GetWorld();
@@ -529,6 +536,37 @@ void USceneCaptureComponent::UpdateDeferredCaptures(FSceneInterface* Scene)
 	// Updating others not associated with the scene would cause invalid data to be rendered into the target
 	TArray< TWeakObjectPtr<USceneCaptureComponent> > SceneCapturesToUpdate;
 	SceneCapturesToUpdateMap.MultiFind(World, SceneCapturesToUpdate);
+
+	const FSceneViewFamily* MainViewFamily = nullptr;
+	for (const FSceneViewFamily* Family : InOutViewFamilies)
+	{
+		if (Family->bIsMainViewFamily)
+		{
+			MainViewFamily = Family;
+			break;
+		}
+	}
+
+	int32 IncrementIfNotRemoved;
+	for (int32 CaptureIndex = 0; CaptureIndex < SceneCapturesToUpdate.Num(); CaptureIndex += IncrementIfNotRemoved)
+	{
+		IncrementIfNotRemoved = 1;
+
+		if (SceneCapturesToUpdate[CaptureIndex].IsValid())
+		{
+			if (USceneCaptureComponent2D* Component = Cast<USceneCaptureComponent2D>(SceneCapturesToUpdate[CaptureIndex].Get()))
+			{
+				// If this is intended to render with the main view family, and we don't have a main view family,
+				// remove it from the array so it can render with the main view family later.
+				if (Component->ShouldRenderWithMainViewFamily() && !MainViewFamily)
+				{
+					SceneCapturesToUpdate.RemoveAtSwap(CaptureIndex);
+					IncrementIfNotRemoved = 0;
+				}
+			}
+		}
+	}
+
 	SceneCapturesToUpdate.Sort([](const TWeakObjectPtr<USceneCaptureComponent>& A, const TWeakObjectPtr<USceneCaptureComponent>& B)
 	{
 		if (!A.IsValid())
@@ -546,7 +584,18 @@ void USceneCaptureComponent::UpdateDeferredCaptures(FSceneInterface* Scene)
 	{
 		if (Component.IsValid())
 		{
+			bool bIsCaptureComponent2D = Component->Is2D();
+			if (bIsCaptureComponent2D)
+			{
+				((USceneCaptureComponent2D*)Component.Get())->MainViewFamily = MainViewFamily;
+			}
+
 			Component->UpdateSceneCaptureContents(Scene);
+
+			if (bIsCaptureComponent2D)
+			{
+				((USceneCaptureComponent2D*)Component.Get())->MainViewFamily = nullptr;
+			}
 		}
 	}
 
@@ -612,6 +661,11 @@ USceneCaptureComponent2D::USceneCaptureComponent2D(const FObjectInitializer& Obj
 	ClipPlaneNormal = FVector(0, 0, 1);
 	bCameraCutThisFrame = false;
 	bConsiderUnrenderedOpaquePixelAsFullyTranslucent = false;
+
+	bMainViewFamily = false;
+	bMainViewResolution = false;
+	bMainViewCamera = false;
+	bIgnoreScreenPercentage = false;
 	
 	TileID = 0;
 
diff --git a/Engine/Source/Runtime/Engine/Private/GameViewportClient.cpp b/Engine/Source/Runtime/Engine/Private/GameViewportClient.cpp
--- a/Engine/Source/Runtime/Engine/Private/GameViewportClient.cpp
+++ b/Engine/Source/Runtime/Engine/Private/GameViewportClient.cpp
@@ -1840,6 +1840,7 @@ void UGameViewportClient::Draw(FViewport* InViewport, FCanvas* SceneCanvas)
 	ViewFamily.bIsHDR = GetWindow().IsValid() ? GetWindow().Get()->GetIsHDR() : false;
 
 	ViewFamily.bSplitScreenDebugAllowed = true;
+	ViewFamily.bIsMainViewFamily = true;
 
 	// Draw the player views.
 	if (!bDisableWorldRendering && PlayerViewMap.Num() > 0 && FSlateApplication::Get().GetPlatformApplication()->IsAllowedToRender()) //-V560
diff --git a/Engine/Source/Runtime/Engine/Private/TextureRenderTarget2D.cpp b/Engine/Source/Runtime/Engine/Private/TextureRenderTarget2D.cpp
--- a/Engine/Source/Runtime/Engine/Private/TextureRenderTarget2D.cpp
+++ b/Engine/Source/Runtime/Engine/Private/TextureRenderTarget2D.cpp
@@ -190,9 +190,9 @@ void UTextureRenderTarget2D::ResizeTarget(uint32 InSizeX, uint32 InSizeY)
 			int32 NewSizeX = SizeX;
 			int32 NewSizeY = SizeY;
 			ENQUEUE_RENDER_COMMAND(ResizeRenderTarget)(
-				[InResource, NewSizeX, NewSizeY](FRHICommandListImmediate& RHICmdList)
+				[InResource, NewSizeX, NewSizeY, LocalNumMips = NumMips](FRHICommandListImmediate& RHICmdList)
 				{
-					InResource->Resize(NewSizeX, NewSizeY);
+					InResource->Resize(NewSizeX, NewSizeY, LocalNumMips);
 					InResource->UpdateDeferredResource(RHICmdList, true);
 				}
 			);
@@ -458,6 +458,7 @@ FTextureRenderTarget2DResource::FTextureRenderTarget2DResource(const class UText
 	,	Format(InOwner->GetFormat())
 	,	TargetSizeX(Owner->SizeX)
 	,	TargetSizeY(Owner->SizeY)
+	,	TargetNumMips(Owner->GetNumMips())
 {
 	// note: Resource has a bSRGB field which is not set or checked in the RenderTarget code
 }
@@ -534,9 +535,9 @@ void FTextureRenderTarget2DResource::InitRHI(FRHICommandListBase& RHICmdList)
 
 		FRHITextureCreateDesc Desc =
 			FRHITextureCreateDesc::Create2D(*ResourceName)
-			.SetExtent(Owner->SizeX, Owner->SizeY)
+			.SetExtent(TargetSizeX, TargetSizeY)
 			.SetFormat(Format)
-			.SetNumMips(Owner->GetNumMips())
+			.SetNumMips(TargetNumMips)
 			.SetNumSamples(NumSamples)
 			.SetFlags(TexCreateFlags)
 			.SetInitialState(ERHIAccess::SRVMask)
@@ -659,12 +660,13 @@ void FTextureRenderTarget2DResource::UpdateDeferredResource( FRHICommandListImme
 	GraphBuilder.Execute();
 }
 
-void FTextureRenderTarget2DResource::Resize(int32 NewSizeX, int32 NewSizeY)
+void FTextureRenderTarget2DResource::Resize(int32 NewSizeX, int32 NewSizeY, int32 NewNumMips)
 {
-	if (TargetSizeX != NewSizeX || TargetSizeY != NewSizeY)
+	if (TargetSizeX != NewSizeX || TargetSizeY != NewSizeY || TargetNumMips != NewNumMips)
 	{
 		TargetSizeX = NewSizeX;
 		TargetSizeY = NewSizeY;
+		TargetNumMips = NewNumMips;
 		UpdateRHI(FRHICommandListImmediate::Get());
 	}
 }
diff --git a/Engine/Source/Runtime/Engine/Public/Rendering/CustomRenderPass.h b/Engine/Source/Runtime/Engine/Public/Rendering/CustomRenderPass.h
--- a/Engine/Source/Runtime/Engine/Public/Rendering/CustomRenderPass.h
+++ b/Engine/Source/Runtime/Engine/Public/Rendering/CustomRenderPass.h
@@ -50,6 +50,7 @@ class FCustomRenderPassBase : public ICustomRenderPass
 {
 public:
 	friend class FSceneRenderer;
+	friend class FDeferredShadingSceneRenderer;
 
 	/** Which render passes are needed for the custom render pass. */
 	enum class ERenderMode
@@ -71,7 +72,8 @@ public:
 		SceneColorAndAlpha, // The alpha is inverted opacity (throughput).
 		SceneColorNoAlpha,
 		BaseColor,
-		Normal
+		Normal,
+		MAX
 	};
 
 	enum class ERenderCaptureType
@@ -132,6 +134,12 @@ protected:
 	virtual void OnPostRender(FRDGBuilder& GraphBuilder) {}
 	virtual void OnEndPass(FRDGBuilder& GraphBuilder) {}
 
+	/** Used internally during scene render to generate multiple outputs from a single custom render pass */
+	void OverrideRenderOutput(ERenderOutput InRenderOutput)
+	{
+		RenderOutput = InRenderOutput;
+	}
+
 protected:
 	FString DebugName;
 	ERenderMode RenderMode = ERenderMode::DepthPass;
diff --git a/Engine/Source/Runtime/Engine/Public/SceneView.h b/Engine/Source/Runtime/Engine/Public/SceneView.h
--- a/Engine/Source/Runtime/Engine/Public/SceneView.h
+++ b/Engine/Source/Runtime/Engine/Public/SceneView.h
@@ -2226,6 +2226,12 @@ public:
 	/** Whether this view allows split screen debug -- we only want it for editor and game viewports, not other random scene renders */
 	bool bSplitScreenDebugAllowed = false;
 
+	/**
+	 * Whether this is the "main" view family.  Should be set to true for the view family of the main editor or game viewport.  Affects
+	 * where rendering occurs for Scene Captures with the "bRenderWithMainViewFamily" flag set.
+	 */
+	bool bIsMainViewFamily = false;
+
 	/**
 	 * Which component of the scene rendering should be output to the final render target.
 	 * If SCS_FinalColorLDR this indicates do nothing.
diff --git a/Engine/Source/Runtime/Engine/Public/TextureResource.h b/Engine/Source/Runtime/Engine/Public/TextureResource.h
--- a/Engine/Source/Runtime/Engine/Public/TextureResource.h
+++ b/Engine/Source/Runtime/Engine/Public/TextureResource.h
@@ -522,7 +522,22 @@ protected:
 	 */
 	friend class UTextureRenderTarget2D;
 	virtual void UpdateDeferredResource(FRHICommandListImmediate& RHICmdList, bool bClearRenderTarget=true) override;
-	void Resize(int32 NewSizeX, int32 NewSizeY);
+	void Resize(int32 NewSizeX, int32 NewSizeY, int32 NewNumMips);
+
+	/** Utility function used for resizing of scene texture sized scene capture render targets in the render thread */
+	friend class FScene;
+	friend class FSceneCapturePass;
+	FORCEINLINE void Resize(FRHICommandListBase& RHICmdList, int32 NewSizeX, int32 NewSizeY, bool bAutoGenerateMips)
+	{
+		int32 NewNumMips = bAutoGenerateMips ? FMath::FloorLog2(FMath::Max(NewSizeX, NewSizeY)) + 1 : 1;
+		if (TargetSizeX != NewSizeX || TargetSizeY != NewSizeY || TargetNumMips != NewNumMips)
+		{
+			TargetSizeX = NewSizeX;
+			TargetSizeY = NewSizeY;
+			TargetNumMips = NewNumMips;
+			UpdateRHI(RHICmdList);
+		}
+	}
 
 private:
 	/** The UTextureRenderTarget2D which this resource represents. */
@@ -536,6 +551,7 @@ private:
 	FLinearColor ClearColor;
 	EPixelFormat Format;
 	int32 TargetSizeX,TargetSizeY;
+	int32 TargetNumMips;
 	TRefCountPtr<IPooledRenderTarget> MipGenerationCache;
 };
 
diff --git a/Engine/Source/Runtime/Renderer/Internal/SceneTextures.h b/Engine/Source/Runtime/Renderer/Internal/SceneTextures.h
--- a/Engine/Source/Runtime/Renderer/Internal/SceneTextures.h
+++ b/Engine/Source/Runtime/Renderer/Internal/SceneTextures.h
@@ -32,16 +32,17 @@ enum class EUserSceneTextureEvent
 	CollidingInput,		// Input matches the output, and has been unbound as a result
 	FoundInput,
 	Output,
-	Pass				// Marker for the end of events for a given material pass with UserSceneTexture inputs or outputs
+	Pass,				// Marker for the end of events for a given material pass with UserSceneTexture inputs or outputs
+	CustomRenderPass	// Marker for a custom render pass that writes to a UserSceneTexture -- AllocationOrder contains ERenderOutput enum, MaterialInterface contains CustomRenderPassBase pointer
 };
 
 struct FUserSceneTextureEventData
 {
 	EUserSceneTextureEvent Event;
 	FName Name;
-	uint16 AllocationOrder;
+	uint16 AllocationOrder;							// ERenderOutput stored here for EUserSceneTextureEvent::CustomRenderPass
 	uint16 ViewIndex;								// Necessary to differentiate events from multiple views in split screen
-	const UMaterialInterface* MaterialInterface;
+	const UMaterialInterface* MaterialInterface;	// FCustomRenderPassBase stored here for EUserSceneTextureEvent::CustomRenderPass
 	FIntPoint RectSize;								// Only filled in for EUserSceneTextureEvent::Output
 };
 #endif
diff --git a/Engine/Source/Runtime/Renderer/Private/CustomRenderPassSceneCapture.h b/Engine/Source/Runtime/Renderer/Private/CustomRenderPassSceneCapture.h
new file mode 100644
--- /dev/null
+++ b/Engine/Source/Runtime/Renderer/Private/CustomRenderPassSceneCapture.h
@@ -0,0 +1,40 @@
+// Copyright Epic Games, Inc. All Rights Reserved.
+
+#pragma once
+
+/*=============================================================================
+	CustomRenderPassSceneCapture.h: Internal user data for scene capture custom render passes
+=============================================================================*/
+
+#include "Rendering/CustomRenderPass.h"
+
+//
+// FSceneCaptureCustomRenderPassUserData
+//
+class FSceneCaptureCustomRenderPassUserData : public ICustomRenderPassUserData
+{
+public:
+	IMPLEMENT_CUSTOM_RENDER_PASS_USER_DATA(FSceneCaptureCustomRenderPassUserData);
+
+	bool bMainViewFamily = false;
+	bool bMainViewResolution = false;
+	bool bMainViewCamera = false;
+	bool bIgnoreScreenPercentage = false;
+	FIntPoint SceneTextureDivisor = FIntPoint(1, 1);
+	FName UserSceneTextureBaseColor;
+	FName UserSceneTextureNormal;
+	FName UserSceneTextureSceneColor;
+#if !UE_BUILD_SHIPPING
+	FString CaptureActorName;
+#endif
+
+	static const FSceneCaptureCustomRenderPassUserData& Get(const FCustomRenderPassBase* CustomRenderPass)
+	{
+		const FSceneCaptureCustomRenderPassUserData* UserData = CustomRenderPass->GetUserDataTyped<FSceneCaptureCustomRenderPassUserData>();
+		return UserData ? *UserData : GDefaultData;
+	}
+
+private:
+	// Default data to return if data isn't present, to simplify renderer logic using the data (no need for conditional that data exists everywhere the structure is used)
+	static FSceneCaptureCustomRenderPassUserData GDefaultData;
+};
diff --git a/Engine/Source/Runtime/Renderer/Private/DeferredShadingRenderer.cpp b/Engine/Source/Runtime/Renderer/Private/DeferredShadingRenderer.cpp
--- a/Engine/Source/Runtime/Renderer/Private/DeferredShadingRenderer.cpp
+++ b/Engine/Source/Runtime/Renderer/Private/DeferredShadingRenderer.cpp
@@ -100,6 +100,7 @@
 #include "PostProcess/DebugAlphaChannel.h"
 #include "MegaLights/MegaLights.h"
 #include "Rendering/CustomRenderPass.h"
+#include "CustomRenderPassSceneCapture.h"
 #include "EnvironmentComponentsFlags.h"
 #include "GenerateMips.h"
 
@@ -2050,6 +2051,8 @@ void FDeferredShadingSceneRenderer::Render(FRDGBuilder& GraphBuilder)
 					SingleLayerWaterPrePassResult = RenderSingleLayerWaterDepthPrepass(GraphBuilder, CustomRenderPassViews, SceneTextures);
 				}
 
+				const FSceneCaptureCustomRenderPassUserData& SceneCaptureUserData = FSceneCaptureCustomRenderPassUserData::Get(CustomRenderPass);
+
 				if (CustomRenderPass->GetRenderMode() == FCustomRenderPassBase::ERenderMode::DepthAndBasePass)
 				{
 					SceneTextures.SetupMode |= ESceneTextureSetupMode::SceneColor;
@@ -2071,7 +2074,8 @@ void FDeferredShadingSceneRenderer::Render(FRDGBuilder& GraphBuilder)
 					}
 
 					FCustomRenderPassBase::ERenderOutput RenderOutput = CustomRenderPass->GetRenderOutput();
-					if (RenderOutput == FCustomRenderPassBase::ERenderOutput::BaseColor || RenderOutput == FCustomRenderPassBase::ERenderOutput::Normal)
+					if (RenderOutput == FCustomRenderPassBase::ERenderOutput::BaseColor || RenderOutput == FCustomRenderPassBase::ERenderOutput::Normal ||
+						!SceneCaptureUserData.UserSceneTextureBaseColor.IsNone() || !SceneCaptureUserData.UserSceneTextureNormal.IsNone() || !SceneCaptureUserData.UserSceneTextureSceneColor.IsNone())
 					{
 						// CopySceneCaptureComponentToTarget uses scene texture uniforms
 						SceneTextures.SetupMode |= ESceneTextureSetupMode::GBuffers;
@@ -2080,6 +2084,43 @@ void FDeferredShadingSceneRenderer::Render(FRDGBuilder& GraphBuilder)
 				}
 
 				CopySceneCaptureComponentToTarget(GraphBuilder, SceneTextures, CustomRenderPass->GetRenderTargetTexture(), ViewFamily, CustomRenderPassViews);
+
+				if (!SceneCaptureUserData.UserSceneTextureBaseColor.IsNone())
+				{
+					bool bFirstRender;
+					FRDGTextureRef BaseColorSceneTexture = SceneTextures.FindOrAddUserSceneTexture(GraphBuilder, 0, SceneCaptureUserData.UserSceneTextureBaseColor, SceneCaptureUserData.SceneTextureDivisor, bFirstRender, nullptr, CustomRenderPassViews[0].ViewRect);
+#if !(UE_BUILD_SHIPPING)
+					SceneTextures.UserSceneTextureEvents.Add({ EUserSceneTextureEvent::CustomRenderPass, NAME_None, (uint16)FCustomRenderPassBase::ERenderOutput::BaseColor, 0, (const UMaterialInterface*)CustomRenderPass });
+#endif
+
+					CustomRenderPass->OverrideRenderOutput(FCustomRenderPassBase::ERenderOutput::BaseColor);
+					CopySceneCaptureComponentToTarget(GraphBuilder, SceneTextures, BaseColorSceneTexture, ViewFamily, CustomRenderPassViews);
+				}
+
+				if (!SceneCaptureUserData.UserSceneTextureNormal.IsNone())
+				{
+					bool bFirstRender;
+					FRDGTextureRef NormalSceneTexture = SceneTextures.FindOrAddUserSceneTexture(GraphBuilder, 0, SceneCaptureUserData.UserSceneTextureNormal, SceneCaptureUserData.SceneTextureDivisor, bFirstRender, nullptr, CustomRenderPassViews[0].ViewRect);
+#if !(UE_BUILD_SHIPPING)
+					SceneTextures.UserSceneTextureEvents.Add({ EUserSceneTextureEvent::CustomRenderPass, NAME_None, (uint16)FCustomRenderPassBase::ERenderOutput::Normal, 0, (const UMaterialInterface*)CustomRenderPass });
+#endif
+
+					CustomRenderPass->OverrideRenderOutput(FCustomRenderPassBase::ERenderOutput::Normal);
+					CopySceneCaptureComponentToTarget(GraphBuilder, SceneTextures, NormalSceneTexture, ViewFamily, CustomRenderPassViews);
+				}
+
+				if (!SceneCaptureUserData.UserSceneTextureSceneColor.IsNone())
+				{
+					bool bFirstRender;
+					FRDGTextureRef SceneColorSceneTexture = SceneTextures.FindOrAddUserSceneTexture(GraphBuilder, 0, SceneCaptureUserData.UserSceneTextureSceneColor, SceneCaptureUserData.SceneTextureDivisor, bFirstRender, nullptr, CustomRenderPassViews[0].ViewRect);
+#if !(UE_BUILD_SHIPPING)
+					SceneTextures.UserSceneTextureEvents.Add({ EUserSceneTextureEvent::CustomRenderPass, NAME_None, (uint16)FCustomRenderPassBase::ERenderOutput::SceneColorAndAlpha, 0, (const UMaterialInterface*)CustomRenderPass });
+#endif
+
+					CustomRenderPass->OverrideRenderOutput(FCustomRenderPassBase::ERenderOutput::SceneColorAndAlpha);
+					CopySceneCaptureComponentToTarget(GraphBuilder, SceneTextures, SceneColorSceneTexture, ViewFamily, CustomRenderPassViews);
+				}
+
 				CustomRenderPass->PostRender(GraphBuilder);
 
 				// Mips are normally generated in UpdateSceneCaptureContentDeferred_RenderThread, but that doesn't run when the
diff --git a/Engine/Source/Runtime/Renderer/Private/PostProcess/PostProcessing.cpp b/Engine/Source/Runtime/Renderer/Private/PostProcess/PostProcessing.cpp
--- a/Engine/Source/Runtime/Renderer/Private/PostProcess/PostProcessing.cpp
+++ b/Engine/Source/Runtime/Renderer/Private/PostProcess/PostProcessing.cpp
@@ -75,6 +75,7 @@
 #include "CanvasItem.h"
 #include "MobileSSR.h"
 #include "Materials/MaterialRenderProxy.h"
+#include "CustomRenderPassSceneCapture.h"
 
 bool IsMobileEyeAdaptationEnabled(const FViewInfo& View);
 
@@ -3282,6 +3283,18 @@ static void AddUserSceneTextureDebugPass(FRDGBuilder& GraphBuilder, const FViewI
 				GBlendableLocationShortNames[BL_MAX] = TEXT("MAX");
 			}
 
+			static_assert((int32)FCustomRenderPassBase::ERenderOutput::MAX == 7);
+			static const TCHAR* GCustomRenderPassOutputShortNames[(int32)FCustomRenderPassBase::ERenderOutput::MAX] =
+			{
+				TEXT("Depth"),			// SceneDepth
+				TEXT("Devdepth"),		// DeviceDepth
+				TEXT("ColDepth"),		// SceneColorAndDepth
+				TEXT("Color"),			// SceneColorAndAlpha
+				TEXT("ColorNoA"),		// SceneColorNoAlpha
+				TEXT("Base"),			// BaseColor
+				TEXT("Norm"),			// Normal
+			};
+
 			float MaxBlendableInfoWidth = 0.0f;
 			for (const FUserSceneTextureEventData& EventData : SceneTextures.UserSceneTextureEvents)
 			{
@@ -3292,6 +3305,18 @@ static void AddUserSceneTextureDebugPass(FRDGBuilder& GraphBuilder, const FViewI
 
 					Text = FString::Printf(TEXT("%s [%d]"), GBlendableLocationShortNames[FMath::Min((uint32)RenderProxy->GetBlendableLocation(Material), (uint32)BL_MAX)], RenderProxy->GetBlendablePriority(Material));
 
+					float BlendableInfoWidth = CanvasDrawShadowedStringReturnWidth(Canvas, PrintX, PrintY, Text, Font, TextColor);
+					MaxBlendableInfoWidth = FMath::Max(MaxBlendableInfoWidth, BlendableInfoWidth);
+					PrintY += CRHeight;
+				}
+				else if (EventData.ViewIndex == ViewIndex && EventData.Event == EUserSceneTextureEvent::CustomRenderPass)
+				{
+					const FCustomRenderPassBase* RenderPass = (const FCustomRenderPassBase*)EventData.MaterialInterface;
+					const FSceneCaptureCustomRenderPassUserData& UserData = FSceneCaptureCustomRenderPassUserData::Get(RenderPass);
+
+					// EventData.AllocationOrder stores the ERenderOutput enum
+					Text = FString::Printf(TEXT("%s (CRP:%s)"), *UserData.CaptureActorName, GCustomRenderPassOutputShortNames[EventData.AllocationOrder]);
+
 					float BlendableInfoWidth = CanvasDrawShadowedStringReturnWidth(Canvas, PrintX, PrintY, Text, Font, TextColor);
 					MaxBlendableInfoWidth = FMath::Max(MaxBlendableInfoWidth, BlendableInfoWidth);
 					PrintY += CRHeight;
@@ -3362,13 +3387,15 @@ static void AddUserSceneTextureDebugPass(FRDGBuilder& GraphBuilder, const FViewI
 							PrintX += CanvasDrawShadowedStringReturnWidth(Canvas, PrintX, PrintY, *Text, Font, UserTexture->bUsed ? GrayTextColor : YellowTextColor);
 							bAnyUnused = bAnyUnused || !UserTexture->bUsed;
 
-							if (EventData.MaterialInterface->GetBlendMode() != BLEND_Opaque)
+							// MaterialInterface can be null if this output was generated by a CustomRenderPass
+							if (EventData.MaterialInterface && EventData.MaterialInterface->GetBlendMode() != BLEND_Opaque)
 							{
 								PrintX += CanvasDrawShadowedStringReturnWidth(Canvas, PrintX, PrintY, TEXT("  Blend"), Font, GrayTextColor);
 							}
 						}
 						break;
 					case EUserSceneTextureEvent::Pass:
+					case EUserSceneTextureEvent::CustomRenderPass:
 						// End of line
 						PrintY += CRHeight;
 						PrintX = PrintX_CR;
diff --git a/Engine/Source/Runtime/Renderer/Private/PostProcess/TemporalAA.cpp b/Engine/Source/Runtime/Renderer/Private/PostProcess/TemporalAA.cpp
--- a/Engine/Source/Runtime/Renderer/Private/PostProcess/TemporalAA.cpp
+++ b/Engine/Source/Runtime/Renderer/Private/PostProcess/TemporalAA.cpp
@@ -1122,7 +1122,7 @@ FDefaultTemporalUpscaler::FOutputs AddGen4MainTemporalAAPasses(
 
 EMainTAAPassConfig GetMainTAAPassConfig(const FViewInfo& View)
 {
-	if (!IsPostProcessingEnabled(View))
+	if (!IsPostProcessingEnabled(View) && !View.bSceneCaptureMainViewJitter)
 	{
 		return EMainTAAPassConfig::Disabled;
 	}
diff --git a/Engine/Source/Runtime/Renderer/Private/SceneCaptureRendering.cpp b/Engine/Source/Runtime/Renderer/Private/SceneCaptureRendering.cpp
--- a/Engine/Source/Runtime/Renderer/Private/SceneCaptureRendering.cpp
+++ b/Engine/Source/Runtime/Renderer/Private/SceneCaptureRendering.cpp
@@ -46,6 +46,7 @@
 #include "DumpGPU.h"
 #include "IRenderCaptureProvider.h"
 #include "RenderCaptureInterface.h"
+#include "CustomRenderPassSceneCapture.h"
 
 bool GSceneCaptureAllowRenderInMainRenderer = true;
 static FAutoConsoleVariableRef CVarSceneCaptureAllowRenderInMainRenderer(
@@ -282,9 +283,19 @@ void CopySceneCaptureComponentToTarget(
 		GraphicsPSOInit.BoundShaderState.PixelShaderRHI = PixelShader.GetPixelShader();
 		GraphicsPSOInit.PrimitiveType = PT_TriangleList;
 
-		// Need to use the extent from the actual target texture for cube captures.  Although perhaps we should use the actual texture
-		// extent across the board?  Would it ever be incorrect to do so?
-		FIntPoint TargetSize = View.bIsSceneCaptureCube && NumViews == 6 ? ViewFamilyTexture->Desc.Extent : View.UnconstrainedViewRect.Size();
+		FIntPoint TargetSize;
+		if (((const FViewFamilyInfo*)View.Family)->bIsSceneTextureSizedCapture)
+		{
+			// Scene texture sized target, use actual target extent for copy, and set correct extent for visualization debug feature
+			TargetSize = ViewFamilyTexture->Desc.Extent;
+			ViewFamilyTexture->EncloseVisualizeExtent(View.UnconstrainedViewRect.Max);
+		}
+		else
+		{
+			// Need to use the extent from the actual target texture for cube captures.  Although perhaps we should use the actual texture
+			// extent across the board?  Would it ever be incorrect to do so?
+			TargetSize = (View.bIsSceneCaptureCube && NumViews == 6) ? ViewFamilyTexture->Desc.Extent : View.UnconstrainedViewRect.Size();
+		}
 
 		GraphBuilder.AddPass(
 			RDG_EVENT_NAME("View(%d)", ViewIndex),
@@ -903,6 +914,8 @@ static FSceneRenderer* CreateSceneRendererForSceneCapture(
 	SceneCaptureViewInfo.ViewRect = FIntRect(0, 0, RenderTargetSize.X, RenderTargetSize.Y);
 	SceneCaptureViewInfo.FOV = InFOV;
 
+	bool bInheritMainViewScreenPercentage = false;
+
 	// Use camera position correction for ortho scene captures
 	if(USceneCaptureComponent2D * SceneCaptureComponent2D = Cast<USceneCaptureComponent2D>(SceneCaptureComponent))
 	{
@@ -910,6 +923,11 @@ static FSceneRenderer* CreateSceneRendererForSceneCapture(
 		{
 			SceneCaptureViewInfo.UpdateOrthoPlanes(SceneCaptureComponent2D->bUseCameraHeightAsViewTarget);
 		}
+
+		if (SceneCaptureComponent2D->ShouldRenderWithMainViewResolution() && SceneCaptureComponent2D->MainViewFamily && !SceneCaptureComponent2D->ShouldIgnoreScreenPercentage())
+		{
+			bInheritMainViewScreenPercentage = true;
+		}
 	}
 
 	FSceneViewFamilyContext ViewFamily(FSceneViewFamily::ConstructionValues(
@@ -937,26 +955,57 @@ static FSceneRenderer* CreateSceneRendererForSceneCapture(
 	// Scene capture source is used to determine whether to disable occlusion queries inside FSceneRenderer constructor
 	ViewFamily.SceneCaptureSource = SceneCaptureComponent->CaptureSource;
 
-	// Screen percentage is still not supported in scene capture.
-	ViewFamily.EngineShowFlags.ScreenPercentage = false;
-	ViewFamily.SetScreenPercentageInterface(new FLegacyScreenPercentageDriver(
-		ViewFamily, /* GlobalResolutionFraction = */ 1.0f));
+	if (bInheritMainViewScreenPercentage)
+	{
+		USceneCaptureComponent2D* SceneCaptureComponent2D = Cast<USceneCaptureComponent2D>(SceneCaptureComponent);
+		ViewFamily.EngineShowFlags.ScreenPercentage = SceneCaptureComponent2D->MainViewFamily->EngineShowFlags.ScreenPercentage;
+		ViewFamily.SetScreenPercentageInterface(SceneCaptureComponent2D->MainViewFamily->GetScreenPercentageInterface()->Fork_GameThread(ViewFamily));
+	}
+	else
+	{
+		// Screen percentage is still not supported in scene capture.
+		ViewFamily.EngineShowFlags.ScreenPercentage = false;
+		ViewFamily.SetScreenPercentageInterface(new FLegacyScreenPercentageDriver(
+			ViewFamily, /* GlobalResolutionFraction = */ 1.0f));
+	}
 
 	return FSceneRenderer::CreateSceneRenderer(&ViewFamily, nullptr);
 }
 
+FSceneCaptureCustomRenderPassUserData FSceneCaptureCustomRenderPassUserData::GDefaultData;
+
 class FSceneCapturePass final : public FCustomRenderPassBase
 {
 public:
 	IMPLEMENT_CUSTOM_RENDER_PASS(FSceneCapturePass);
 
-	FSceneCapturePass(const FString& InDebugName, ERenderMode InRenderMode, ERenderOutput InRenderOutput, UTextureRenderTarget2D* InRenderTarget)
-		: FCustomRenderPassBase(InDebugName, InRenderMode, InRenderOutput, FIntPoint(InRenderTarget->GetSurfaceWidth(), InRenderTarget->GetSurfaceHeight()))
+	FSceneCapturePass(const FString& InDebugName, ERenderMode InRenderMode, ERenderOutput InRenderOutput, UTextureRenderTarget2D* InRenderTarget, USceneCaptureComponent2D* CaptureComponent, FIntPoint InRenderTargetSize)
+		: FCustomRenderPassBase(InDebugName, InRenderMode, InRenderOutput, InRenderTargetSize)
 		, SceneCaptureRenderTarget(InRenderTarget->GameThread_GetRenderTargetResource())
-	{}
+		, bAutoGenerateMips(InRenderTarget->bAutoGenerateMips)
+	{
+		FSceneCaptureCustomRenderPassUserData* UserData = new FSceneCaptureCustomRenderPassUserData();
+		UserData->bMainViewFamily = CaptureComponent->ShouldRenderWithMainViewFamily();
+		UserData->bMainViewResolution = CaptureComponent->ShouldRenderWithMainViewResolution();
+		UserData->bMainViewCamera = CaptureComponent->ShouldRenderWithMainViewCamera();
+		UserData->bIgnoreScreenPercentage = CaptureComponent->ShouldIgnoreScreenPercentage();
+		UserData->SceneTextureDivisor = CaptureComponent->MainViewResolutionDivisor.ComponentMax(FIntPoint(1,1));
+		UserData->UserSceneTextureBaseColor = CaptureComponent->UserSceneTextureBaseColor;
+		UserData->UserSceneTextureNormal = CaptureComponent->UserSceneTextureNormal;
+		UserData->UserSceneTextureSceneColor = CaptureComponent->UserSceneTextureSceneColor;
+#if !UE_BUILD_SHIPPING
+		CaptureComponent->GetOuter()->GetName(UserData->CaptureActorName);
+#endif
+
+		SetUserData(TUniquePtr<FSceneCaptureCustomRenderPassUserData>(UserData));
+	}
 
 	virtual void OnPreRender(FRDGBuilder& GraphBuilder) override
 	{
+		// Resize the render resource if necessary -- render target size may have been overridden to the main view resolution, or later be changed back
+		// to the resource resolution.  The resize call does nothing if the size already matches.
+		((FTextureRenderTarget2DResource*)SceneCaptureRenderTarget)->Resize(GraphBuilder.RHICmdList, RenderTargetSize.X, RenderTargetSize.Y, bAutoGenerateMips);
+
 		RenderTargetTexture = SceneCaptureRenderTarget->GetRenderTargetTexture(GraphBuilder);
 	}
 
@@ -968,6 +1017,7 @@ public:
 	}
 	
 	FRenderTarget* SceneCaptureRenderTarget = nullptr;
+	bool bAutoGenerateMips = false;
 };
 
 static void BeginGpuCaptureOrDump(USceneCaptureComponent* CaptureComponent, bool& bCapturingGPU, bool& bDumpingGPU)
@@ -1040,57 +1090,87 @@ void FScene::UpdateSceneCaptureContents(USceneCaptureComponent2D* CaptureCompone
 
 	if (UTextureRenderTarget2D* TextureRenderTarget = CaptureComponent->TextureTarget)
 	{
-		FTransform Transform = CaptureComponent->GetComponentToWorld();
-		FVector ViewLocation = Transform.GetTranslation();
-
-		// Remove the translation from Transform because we only need rotation.
-		Transform.SetTranslation(FVector::ZeroVector);
-		Transform.SetScale3D(FVector::OneVector);
-		FMatrix ViewRotationMatrix = Transform.ToInverseMatrixWithScale();
-
-		// swap axis st. x=z,y=x,z=y (unreal coord space) so that z is up
-		ViewRotationMatrix = ViewRotationMatrix * FMatrix(
-			FPlane(0, 0, 1, 0),
-			FPlane(1, 0, 0, 0),
-			FPlane(0, 1, 0, 0),
-			FPlane(0, 0, 0, 1));
-		const float UnscaledFOV = CaptureComponent->FOVAngle * (float)PI / 360.0f;
-		const float FOV = FMath::Atan((1.0f + CaptureComponent->Overscan) * FMath::Tan(UnscaledFOV));
-		FIntPoint CaptureSize(TextureRenderTarget->GetSurfaceWidth(), TextureRenderTarget->GetSurfaceHeight());
+		FIntPoint CaptureSize;
+		FVector ViewLocation;
+		FMatrix ViewRotationMatrix;
+		FMatrix ProjectionMatrix;
+		bool bEnableOrthographicTiling;
 
 		const bool bUseSceneColorTexture = CaptureNeedsSceneColor(CaptureComponent->CaptureSource);
-		const bool bEnableOrthographicTiling = (CaptureComponent->GetEnableOrthographicTiling() && CaptureComponent->ProjectionType == ECameraProjectionMode::Orthographic && bUseSceneColorTexture);
-		if (CaptureComponent->GetEnableOrthographicTiling() && CaptureComponent->ProjectionType == ECameraProjectionMode::Orthographic && !bUseSceneColorTexture)
-		{
-			UE_LOG(LogRenderer, Warning, TEXT("SceneCapture - Orthographic and tiling with CaptureSource not using SceneColor (i.e FinalColor) not compatible. SceneCapture render will not be tiled"));
-		}
-		
+
 		const int32 TileID = CaptureComponent->TileID;
 		const int32 NumXTiles = CaptureComponent->GetNumXTiles();
 		const int32 NumYTiles = CaptureComponent->GetNumYTiles();
 
-		FMatrix ProjectionMatrix;
-		if (CaptureComponent->bUseCustomProjectionMatrix)
+		if (CaptureComponent->ShouldRenderWithMainViewResolution() && CaptureComponent->MainViewFamily)
 		{
-			ProjectionMatrix = CaptureComponent->CustomProjectionMatrix;
+			CaptureSize = CaptureComponent->MainViewFamily->Views[0]->UnscaledViewRect.Size();
+			CaptureSize = FIntPoint::DivideAndRoundUp(CaptureSize, CaptureComponent->MainViewResolutionDivisor.ComponentMax(FIntPoint(1,1)));
+
+			// Main view resolution rendering doesn't support orthographic tiling
+			bEnableOrthographicTiling = false;
 		}
 		else
 		{
-			if (CaptureComponent->ProjectionType == ECameraProjectionMode::Perspective)
+			CaptureSize = FIntPoint(TextureRenderTarget->GetSurfaceWidth(), TextureRenderTarget->GetSurfaceHeight());
+
+			bEnableOrthographicTiling = (CaptureComponent->GetEnableOrthographicTiling() && CaptureComponent->ProjectionType == ECameraProjectionMode::Orthographic && bUseSceneColorTexture);
+
+			if (CaptureComponent->GetEnableOrthographicTiling() && CaptureComponent->ProjectionType == ECameraProjectionMode::Orthographic && !bUseSceneColorTexture)
 			{
-				const float ClippingPlane = (CaptureComponent->bOverride_CustomNearClippingPlane) ? CaptureComponent->CustomNearClippingPlane : GNearClippingPlane;
-				BuildProjectionMatrix(CaptureSize, FOV, ClippingPlane, ProjectionMatrix);
+				UE_LOG(LogRenderer, Warning, TEXT("SceneCapture - Orthographic and tiling with CaptureSource not using SceneColor (i.e FinalColor) not compatible. SceneCapture render will not be tiled"));
+			}
+		}
+
+		if (CaptureComponent->ShouldRenderWithMainViewCamera() && CaptureComponent->MainViewFamily)
+		{
+			const FSceneView* MainView = CaptureComponent->MainViewFamily->Views[0];
+
+			ViewLocation = MainView->ViewMatrices.GetViewOrigin();
+			ViewRotationMatrix = MainView->ViewMatrices.GetViewMatrix().RemoveTranslation();
+			ProjectionMatrix = MainView->ViewMatrices.GetProjectionMatrix();
+		}
+		else
+		{
+			FTransform Transform = CaptureComponent->GetComponentToWorld();
+			ViewLocation = Transform.GetTranslation();
+
+			// Remove the translation from Transform because we only need rotation.
+			Transform.SetTranslation(FVector::ZeroVector);
+			Transform.SetScale3D(FVector::OneVector);
+			ViewRotationMatrix = Transform.ToInverseMatrixWithScale();
+
+			// swap axis st. x=z,y=x,z=y (unreal coord space) so that z is up
+			ViewRotationMatrix = ViewRotationMatrix * FMatrix(
+				FPlane(0, 0, 1, 0),
+				FPlane(1, 0, 0, 0),
+				FPlane(0, 1, 0, 0),
+				FPlane(0, 0, 0, 1));
+			const float UnscaledFOV = CaptureComponent->FOVAngle * (float)PI / 360.0f;
+			const float FOV = FMath::Atan((1.0f + CaptureComponent->Overscan) * FMath::Tan(UnscaledFOV));
+
+			if (CaptureComponent->bUseCustomProjectionMatrix)
+			{
+				ProjectionMatrix = CaptureComponent->CustomProjectionMatrix;
 			}
 			else
 			{
-				if (bEnableOrthographicTiling)
+				if (CaptureComponent->ProjectionType == ECameraProjectionMode::Perspective)
 				{
-					BuildOrthoMatrix(CaptureSize, CaptureComponent->OrthoWidth, CaptureComponent->TileID, NumXTiles, NumYTiles, ProjectionMatrix);
-					CaptureSize /= FIntPoint(NumXTiles, NumYTiles);
+					const float ClippingPlane = (CaptureComponent->bOverride_CustomNearClippingPlane) ? CaptureComponent->CustomNearClippingPlane : GNearClippingPlane;
+					BuildProjectionMatrix(CaptureSize, FOV, ClippingPlane, ProjectionMatrix);
 				}
 				else
 				{
-					BuildOrthoMatrix(CaptureSize, CaptureComponent->OrthoWidth, -1, 0, 0, ProjectionMatrix);
+					if (bEnableOrthographicTiling)
+					{
+						BuildOrthoMatrix(CaptureSize, CaptureComponent->OrthoWidth, CaptureComponent->TileID, NumXTiles, NumYTiles, ProjectionMatrix);
+						CaptureSize /= FIntPoint(NumXTiles, NumYTiles);
+					}
+					else
+					{
+						BuildOrthoMatrix(CaptureSize, CaptureComponent->OrthoWidth, -1, 0, 0, ProjectionMatrix);
+					}
 				}
 			}
 		}
@@ -1113,15 +1193,32 @@ void FScene::UpdateSceneCaptureContents(USceneCaptureComponent2D* CaptureCompone
 			FCustomRenderPassBase::ERenderOutput RenderOutput;
 			const TCHAR* DebugName;
 
+			bool bHasUserSceneTextureOutput = !CaptureComponent->UserSceneTextureBaseColor.IsNone() || !CaptureComponent->UserSceneTextureNormal.IsNone() || !CaptureComponent->UserSceneTextureSceneColor.IsNone();
+
 			switch (CaptureComponent->CaptureSource)
 			{
 			case ESceneCaptureSource::SCS_SceneDepth:
-				RenderMode = FCustomRenderPassBase::ERenderMode::DepthPass;
+				if (bHasUserSceneTextureOutput)
+				{
+					// If a UserSceneTexture output is specified, the base pass needs to run to generate it.
+					RenderMode = FCustomRenderPassBase::ERenderMode::DepthAndBasePass;
+				}
+				else
+				{
+					RenderMode = FCustomRenderPassBase::ERenderMode::DepthPass;
+				}
 				RenderOutput = FCustomRenderPassBase::ERenderOutput::SceneDepth;
 				DebugName = TEXT("SceneCapturePass_SceneDepth");
 				break;
 			case ESceneCaptureSource::SCS_DeviceDepth:
-				RenderMode = FCustomRenderPassBase::ERenderMode::DepthPass;
+				if (bHasUserSceneTextureOutput)
+				{
+					RenderMode = FCustomRenderPassBase::ERenderMode::DepthAndBasePass;
+				}
+				else
+				{
+					RenderMode = FCustomRenderPassBase::ERenderMode::DepthPass;
+				}
 				RenderOutput = FCustomRenderPassBase::ERenderOutput::DeviceDepth;
 				DebugName = TEXT("SceneCapturePass_DeviceDepth");
 				break;
@@ -1138,7 +1235,7 @@ void FScene::UpdateSceneCaptureContents(USceneCaptureComponent2D* CaptureCompone
 				break;
 			}
 
-			FSceneCapturePass* CustomPass = new FSceneCapturePass(DebugName, RenderMode, RenderOutput, TextureRenderTarget);
+			FSceneCapturePass* CustomPass = new FSceneCapturePass(DebugName, RenderMode, RenderOutput, TextureRenderTarget, CaptureComponent, CaptureSize);
 			PassInput.CustomRenderPass = CustomPass;
 
 			GetShowOnlyAndHiddenComponents(CaptureComponent, PassInput.HiddenPrimitives, PassInput.ShowOnlyPrimitives);
@@ -1174,6 +1271,31 @@ void FScene::UpdateSceneCaptureContents(USceneCaptureComponent2D* CaptureCompone
 
 		check(SceneRenderer != nullptr);
 
+		// Copy temporal AA related settings for main view camera scene capture, to match jitter.  Don't match if the resolution divisor is set,
+		// if it's set to ignore screen percentage, or if it's final color, which will run its own AA.  For custom render passes (handled above),
+		// computed jitter results are copied from the main view later in FSceneRenderer::PrepareViewStateForVisibility, but this doesn't work
+		// for regular scene captures, because they run in a separate scene renderer before the main view, where the main view's results haven't
+		// been computed yet.
+		if (CaptureComponent->ShouldRenderWithMainViewCamera() && CaptureComponent->MainViewFamily && SceneRenderer->Views[0].ViewState && CaptureComponent->MainViewFamily->Views[0]->State &&
+			CaptureComponent->MainViewResolutionDivisor.X <= 1 && CaptureComponent->MainViewResolutionDivisor.Y <= 1 && !CaptureComponent->ShouldIgnoreScreenPercentage() &&
+			CaptureComponent->CaptureSource != ESceneCaptureSource::SCS_FinalColorLDR &&
+			CaptureComponent->CaptureSource != ESceneCaptureSource::SCS_FinalColorHDR &&
+			CaptureComponent->CaptureSource != ESceneCaptureSource::SCS_FinalToneCurveHDR)
+		{
+			const FSceneViewFamily* MainViewFamily = CaptureComponent->MainViewFamily;
+			const FSceneView& SourceView = *MainViewFamily->Views[0];
+			const FSceneViewState& SourceViewState = *(const FSceneViewState*)SourceView.State;
+
+			FViewInfo& DestView = SceneRenderer->Views[0];
+
+			DestView.AntiAliasingMethod = SourceView.AntiAliasingMethod;
+			DestView.PrimaryScreenPercentageMethod = SourceView.PrimaryScreenPercentageMethod;
+			DestView.ViewState->TemporalAASampleIndex = SourceViewState.TemporalAASampleIndex;
+			DestView.bSceneCaptureMainViewJitter = true;
+
+			CaptureComponent->bCameraCutThisFrame = SourceView.bCameraCut;
+		}
+
 		// When bIsMultipleSceneCapture is true, set bIsFirstSceneRenderer to false, which tells the scene renderer it can skip RHI resource flush, saving performance
 		bool bIsMultipleSceneCapture = CaptureComponent->SetFrameUpdated();
 		SceneRenderer->bIsFirstSceneRenderer = !bIsMultipleSceneCapture;
@@ -1319,8 +1441,12 @@ void FScene::UpdateSceneCaptureContents(USceneCaptureComponent2D* CaptureCompone
 		UE::RenderCommandPipe::FSyncScope SyncScope;
 
 		ENQUEUE_RENDER_COMMAND(CaptureCommand)(
-			[SceneRenderer, TextureRenderTargetResource, TexturePtrNotDeferenced, EventName, TargetName, bGenerateMips, GenerateMipsParams, GameViewportRT, bEnableOrthographicTiling, bIsCompositing, bOrthographicCamera, NumXTiles, NumYTiles, TileID, CaptureMemorySize, bCapturingGPU](FRHICommandListImmediate& RHICmdList)
+			[SceneRenderer, TextureRenderTargetResource, TexturePtrNotDeferenced, EventName, TargetName, bGenerateMips, GenerateMipsParams, GameViewportRT, bEnableOrthographicTiling, bIsCompositing, bOrthographicCamera, NumXTiles, NumYTiles, TileID, CaptureMemorySize, bCapturingGPU, CaptureSize](FRHICommandListImmediate& RHICmdList)
 			{
+				// Resize the render resource if necessary, either to the main viewport size overridden above (see ShouldRenderWithMainViewResolution()),
+				// or the original size if we are changing back to that (the resize call does nothing if the size already matches).
+				TextureRenderTargetResource->GetTextureRenderTarget2DResource()->Resize(RHICmdList, CaptureSize.X, CaptureSize.Y, bGenerateMips);
+
 				RenderCaptureInterface::FScopedCapture RenderCapture(bCapturingGPU, &RHICmdList, *FString::Format(TEXT("Scene Capture : {0}"), { EventName }));
 
 				if (GameViewportRT != nullptr)
diff --git a/Engine/Source/Runtime/Renderer/Private/SceneRendering.cpp b/Engine/Source/Runtime/Renderer/Private/SceneRendering.cpp
--- a/Engine/Source/Runtime/Renderer/Private/SceneRendering.cpp
+++ b/Engine/Source/Runtime/Renderer/Private/SceneRendering.cpp
@@ -14,6 +14,7 @@
 #include "Components/ReflectionCaptureComponent.h"
 #include "Components/SceneCaptureComponent2D.h"
 #include "Components/SceneCaptureComponentCube.h"
+#include "SceneCapture/SceneCaptureInternal.h"
 #include "DeferredShadingRenderer.h"
 #include "DumpGPU.h"
 #include "DynamicPrimitiveDrawing.h"
@@ -98,6 +99,7 @@
 #include "LocalFogVolumeRendering.h"
 #include "OIT/OIT.h"
 #include "Rendering/CustomRenderPass.h"
+#include "CustomRenderPassSceneCapture.h"
 #include "LightFunctionAtlas.h"
 #include "EnvironmentComponentsFlags.h"
 #include "Math/RotationMatrix.h"
@@ -946,6 +948,7 @@ void FViewInfo::Init()
 	bHasSingleLayerWaterMaterial = 0;
 	AutoBeforeDOFTranslucencyBoundary = 0.0f;
 	bUsesSecondStageDepthPass = 0;
+	bSceneCaptureMainViewJitter = 0;
 
 	NumVisibleStaticMeshElements = 0;
 	PrecomputedVisibilityData = 0;
@@ -1020,6 +1023,7 @@ void FViewInfo::Init()
 	NumSphereReflectionCaptures = 0;
 	FurthestReflectionCaptureDistance = 0;
 
+	TemporalSourceView = nullptr;
 	TemporalJitterSequenceLength = 1;
 	TemporalJitterIndex = 0;
 	TemporalJitterPixels = FVector2D::ZeroVector;
@@ -1715,7 +1719,7 @@ void FViewInfo::SetupUniformBufferParameters(
 	}
 
 	{
-		ensureMsgf(TemporalJitterSequenceLength == 1 || IsTemporalAccumulationBasedMethod(AntiAliasingMethod),
+		ensureMsgf(TemporalJitterSequenceLength == 1 || IsTemporalAccumulationBasedMethod(AntiAliasingMethod) || (CustomRenderPass && FSceneCaptureCustomRenderPassUserData::Get(CustomRenderPass).bMainViewResolution),
 			TEXT("TemporalJitterSequenceLength = %i is invalid"), TemporalJitterSequenceLength);
 		ensureMsgf(TemporalJitterIndex >= 0 && TemporalJitterIndex < TemporalJitterSequenceLength,
 			TEXT("TemporalJitterIndex = %i is invalid (TemporalJitterSequenceLength = %i)"), TemporalJitterIndex, TemporalJitterSequenceLength);
@@ -2251,6 +2255,25 @@ FIntRect FViewInfo::GetFamilyViewRect() const
 	return FamilyRect;
 }
 
+FIntRect FViewInfo::GetUnscaledFamilyViewRect() const
+{
+	FIntRect FamilyRect = {};
+	for (uint64 ViewIdx = 0, NumViews = (uint64)Family->Views.Num(); ViewIdx < NumViews; ++ViewIdx)
+	{
+		FamilyRect.Union(static_cast<const FViewInfo*>(Family->Views[ViewIdx])->UnscaledViewRect);
+
+		if (bIsMultiViewportEnabled)
+		{
+			for (const FSceneView* SecondaryView : GetSecondaryViews())
+			{
+				const FViewInfo& InstancedView = static_cast<const FViewInfo&>(*SecondaryView);
+				FamilyRect.Union(InstancedView.UnscaledViewRect);
+			}
+		}
+	}
+	return FamilyRect;
+}
+
 void FViewInfo::BeginRenderView() const
 {
 	const bool bShouldWaitForPersistentViewUniformBufferExtensionsJobs = true;
@@ -2936,12 +2959,26 @@ FSceneRenderer::FSceneRenderer(const FSceneViewFamily* InViewFamily, FHitProxyCo
 	CustomRenderPassInfos.Empty(Scene->CustomRenderPassRendererInputs.Num());
 
 	int32 NumAdditionalViews = 0;
-	for (int32 i = 0; i < Scene->CustomRenderPassRendererInputs.Num(); i++)
+	int32 IncrementIfNotRemoved = 1;
+	for (int32 i = 0; i < Scene->CustomRenderPassRendererInputs.Num(); i+=IncrementIfNotRemoved)
 	{
 		const FScene::FCustomRenderPassRendererInput& PassInput = Scene->CustomRenderPassRendererInputs[i];
 		FCustomRenderPassBase* CustomRenderPass = PassInput.CustomRenderPass;
 		check(CustomRenderPass);
 
+		const FSceneCaptureCustomRenderPassUserData& SceneCaptureUserData = FSceneCaptureCustomRenderPassUserData::Get(CustomRenderPass);
+
+		if (SceneCaptureUserData.bMainViewFamily && !ViewFamily.bIsMainViewFamily)
+		{
+			// If the custom render pass is flagged as rendering with the main view family, and this isn't the main view family, skip it.
+			IncrementIfNotRemoved = 1;
+			continue;
+		}
+		else
+		{
+			IncrementIfNotRemoved = 0;
+		}
+
 		// We construct from scratch, rather than copying, as we don't want to copy interfaces attached to the view family
 		// (ScreenPercentageInterface, TemporalUpscalerInterface, etc), which can assert or double free if copied.  Those aren't
 		// relevant for custom render passes anyway.
@@ -2960,6 +2997,7 @@ FSceneRenderer::FSceneRenderer(const FSceneViewFamily* InViewFamily, FHitProxyCo
 		CustomRenderPassInfo.CustomRenderPass = CustomRenderPass;
 		CustomRenderPassInfo.ViewFamily.Time = ViewFamily.Time;
 		CustomRenderPassInfo.ViewFamily.SetSceneRenderer(this);
+		CustomRenderPassInfo.ViewFamily.bIsSceneTextureSizedCapture = SceneCaptureUserData.bMainViewResolution;
 
 		FSceneViewInitOptions ViewInitOptions;
 		ViewInitOptions.SceneViewStateInterface = PassInput.ViewStateInterface;
@@ -2995,6 +3033,8 @@ FSceneRenderer::FSceneRenderer(const FSceneViewFamily* InViewFamily, FHitProxyCo
 		CustomRenderPass->Views.Add(ViewInfo);
 
 		NumAdditionalViews++;
+
+		Scene->CustomRenderPassRendererInputs.RemoveAt(i, EAllowShrinking::No);
 	}
 
 	AllViews.Empty(Views.Num() + NumAdditionalViews);
@@ -3066,8 +3106,6 @@ FSceneRenderer::FSceneRenderer(const FSceneViewFamily* InViewFamily, FHitProxyCo
 		PassInfo.ViewFamily.AllViews = ViewFamily.AllViews;
 	}
 
-	Scene->CustomRenderPassRendererInputs.Reset();
-
 	FeatureLevel = Scene->GetFeatureLevel();
 	ShaderPlatform = Scene->GetShaderPlatform();
 
@@ -3262,7 +3300,7 @@ void FSceneRenderer::PrepareViewRectsForRendering(FRHICommandListImmediate& RHIC
 
 		// Fallback to no anti aliasing.
 		{
-			const bool bWillApplyTemporalAA = (IsPostProcessingEnabled(View) || View.bIsPlanarReflection)
+			const bool bWillApplyTemporalAA = (IsPostProcessingEnabled(View) || View.bIsPlanarReflection || View.bSceneCaptureMainViewJitter)
 #if RHI_RAYTRACING
 				// path tracer does its own anti-aliasing
 				&& (!ViewFamily.EngineShowFlags.PathTracing)
@@ -3432,7 +3470,30 @@ void FSceneRenderer::PrepareViewRectsForRendering(FRHICommandListImmediate& RHIC
 	{
 		for (FViewInfo& View : PassInfo.Views)
 		{
-			View.ViewRect = View.UnscaledViewRect;
+			const FSceneCaptureCustomRenderPassUserData& SceneCaptureUserData = FSceneCaptureCustomRenderPassUserData::Get(PassInfo.CustomRenderPass);
+
+			if (SceneCaptureUserData.bMainViewResolution)
+			{
+				if (SceneCaptureUserData.bIgnoreScreenPercentage)
+				{
+					View.ViewRect = GetDownscaledViewRect(Views[0].UnscaledViewRect, Views[0].GetUnscaledFamilyViewRect().Max, SceneCaptureUserData.SceneTextureDivisor);
+				}
+				else
+				{
+					View.ViewRect = GetDownscaledViewRect(Views[0].ViewRect, Views[0].GetFamilyViewRect().Max, SceneCaptureUserData.SceneTextureDivisor);
+
+					// Share temporal AA offset if this is coincident with main view camera
+					if (SceneCaptureUserData.bMainViewCamera && SceneCaptureUserData.SceneTextureDivisor == FIntPoint(1,1))
+					{
+						View.TemporalSourceView = &Views[0];
+					}
+				}
+				View.UnconstrainedViewRect = View.ViewRect;
+			}
+			else
+			{
+				View.ViewRect = View.UnscaledViewRect;
+			}
 		}	
 	}
 }
@@ -5485,7 +5546,7 @@ void FRendererModule::BeginRenderingViewFamilies(FCanvas* Canvas, TArrayView<FSc
 		bool bShowHitProxies = (Canvas->GetHitProxyConsumer() != nullptr);
 		if (!bShowHitProxies)
 		{
-			USceneCaptureComponent::UpdateDeferredCaptures(Scene);
+			SceneCaptureUpdateDeferredCapturesInternal(Scene, ViewFamiliesConst);
 		}
 
 		FSceneRenderer::CreateSceneRenderers(ViewFamiliesConst, Canvas->GetHitProxyConsumer(), SceneRenderers);
diff --git a/Engine/Source/Runtime/Renderer/Private/SceneRendering.h b/Engine/Source/Runtime/Renderer/Private/SceneRendering.h
--- a/Engine/Source/Runtime/Renderer/Private/SceneRendering.h
+++ b/Engine/Source/Runtime/Renderer/Private/SceneRendering.h
@@ -1319,6 +1319,9 @@ public:
 	float TranslucencyVolumeVoxelSize[TVC_MAX];
 	FVector TranslucencyLightingVolumeSize[TVC_MAX];
 
+	/** Optional source view for temporal AA, to handle custom render passes and scene captures sharing the main view's camera (jitter needs to match) */
+	FViewInfo* TemporalSourceView;
+
 	/** Number of samples in the temporal AA sequqnce */
 	int32 TemporalJitterSequenceLength;
 
@@ -1369,6 +1372,12 @@ public:
 	 */
 	uint32 bUsesSecondStageDepthPass : 1;
 
+	/**
+	 * Set to true if this is a scene capture sharing temporal AA jitter with the main view camera.  Needed to force temporal jitter
+	 * logic to run when post processing is disabled for the scene capture, which otherwise disables jitter.
+	 */
+	uint32 bSceneCaptureMainViewJitter : 1;
+
 	/** Whether post DOF translucency should be rendered before DOF if primitive bounds behind DOF's focus distance. */
 	float AutoBeforeDOFTranslucencyBoundary;
 
@@ -1687,6 +1696,9 @@ public:
 	/** Returns view rect for all views in the family (usable for atlased views only). Normally a union of their view rects */
 	FIntRect GetFamilyViewRect() const;
 
+	/** Similar to above, but returns the unscaled rect, ignoring dynamic resolution scaling */
+	FIntRect GetUnscaledFamilyViewRect() const;
+
 	/** Prepares the view shader parameters for rendering and calls the persistent uniform buffer hooks. */
 	void BeginRenderView() const;
 
@@ -1884,6 +1896,9 @@ public:
 
 	FSceneTexturesConfig SceneTexturesConfig;
 
+	/** Set to true if this is a scene capture sized to scene texture size */
+	bool bIsSceneTextureSizedCapture = false;
+
 	/** Get scene textures associated with this view family -- asserts or checks that they have been initialized */
 	inline FSceneTextures& GetSceneTextures()
 	{
diff --git a/Engine/Source/Runtime/Renderer/Private/SceneVisibility.cpp b/Engine/Source/Runtime/Renderer/Private/SceneVisibility.cpp
--- a/Engine/Source/Runtime/Renderer/Private/SceneVisibility.cpp
+++ b/Engine/Source/Runtime/Renderer/Private/SceneVisibility.cpp
@@ -4996,14 +4996,18 @@ void FSceneRenderer::PrepareViewStateForVisibility(const FSceneTexturesConfig& S
 		int32 CVarTemporalAASamplesValue = CVarTemporalAASamples.GetValueOnRenderThread();
 		const bool bShouldScaleTemporalAASampleCount = (CVarTemporalAAScaleSamples.GetValueOnRenderThread() != 0);
 
-		EMainTAAPassConfig TAAConfig = GetMainTAAPassConfig(View);
+		// Custom render passes support sharing temporal state with the main view
+		FViewInfo& TemporalSourceView = View.TemporalSourceView ? *View.TemporalSourceView : View;
+		FSceneViewState* TemporalViewState = TemporalSourceView.ViewState;
 
-		bool bTemporalUpsampling = View.PrimaryScreenPercentageMethod == EPrimaryScreenPercentageMethod::TemporalUpscale;
+		EMainTAAPassConfig TAAConfig = GetMainTAAPassConfig(TemporalSourceView);
+
+		bool bTemporalUpsampling = TemporalSourceView.PrimaryScreenPercentageMethod == EPrimaryScreenPercentageMethod::TemporalUpscale;
 		
 		// Apply a sub pixel offset to the view.
-		if (IsTemporalAccumulationBasedMethod(View.AntiAliasingMethod) && ViewState && (CVarTemporalAASamplesValue > 0 || bTemporalUpsampling) && View.bAllowTemporalJitter)
+		if (IsTemporalAccumulationBasedMethod(TemporalSourceView.AntiAliasingMethod) && TemporalViewState && (CVarTemporalAASamplesValue > 0 || bTemporalUpsampling) && TemporalSourceView.bAllowTemporalJitter)
 		{
-			float EffectivePrimaryResolutionFraction = float(View.ViewRect.Width()) / float(View.GetSecondaryViewRectSize().X);
+			float EffectivePrimaryResolutionFraction = float(View.ViewRect.Width()) / float(TemporalSourceView.GetSecondaryViewRectSize().X);
 
 			// Compute number of TAA samples.
 			int32 TemporalAASamples;
@@ -5052,28 +5056,38 @@ void FSceneRenderer::PrepareViewStateForVisibility(const FSceneTexturesConfig& S
 			}
 
 			// Compute the new sample index in the temporal sequence.
-			int32 TemporalSampleIndex = ViewState->TemporalAASampleIndex + 1;
-			if (TemporalSampleIndex >= TemporalAASamples || View.bCameraCut)
+			int32 TemporalSampleIndex;
+			if (View.TemporalSourceView)
 			{
-				TemporalSampleIndex = 0;
+				// Where a TemporalSourceView is specified, the temporal sample index is pulled directly from the source view's state, and not incremented,
+				// assuming the	source view already will have updated that.
+				TemporalSampleIndex = TemporalViewState->TemporalAASampleIndex;
 			}
-
-			#if !UE_BUILD_SHIPPING
-			if (CVarTAADebugOverrideTemporalIndex.GetValueOnRenderThread() >= 0)
+			else
 			{
-				TemporalSampleIndex = CVarTAADebugOverrideTemporalIndex.GetValueOnRenderThread();
-			}
-			#endif
+				TemporalSampleIndex = TemporalViewState->TemporalAASampleIndex + 1;
+				if (TemporalSampleIndex >= TemporalAASamples || View.bCameraCut)
+				{
+					TemporalSampleIndex = 0;
+				}
 
-			// Updates view state.
-			if (!View.bStatePrevViewInfoIsReadOnly && !bFreezeTemporalSequences)
-			{
-				ViewState->TemporalAASampleIndex = TemporalSampleIndex;
+				#if !UE_BUILD_SHIPPING
+				if (CVarTAADebugOverrideTemporalIndex.GetValueOnRenderThread() >= 0)
+				{
+					TemporalSampleIndex = CVarTAADebugOverrideTemporalIndex.GetValueOnRenderThread();
+				}
+				#endif
+
+				// Updates view state.
+				if (!View.bStatePrevViewInfoIsReadOnly && !bFreezeTemporalSequences)
+				{
+					TemporalViewState->TemporalAASampleIndex = TemporalSampleIndex;
+				}
 			}
 
 			// Choose sub pixel sample coordinate in the temporal sequence.
 			float SampleX = 0.0f, SampleY = 0.0f;
-			if (View.PrimaryScreenPercentageMethod == EPrimaryScreenPercentageMethod::TemporalUpscale)
+			if (TemporalSourceView.PrimaryScreenPercentageMethod == EPrimaryScreenPercentageMethod::TemporalUpscale)
 			{
 				// Uniformly distribute temporal jittering in [-.5; .5], because there is no longer any alignement of input and output pixels.
 				SampleX = Halton(TemporalSampleIndex + 1, 2) - 0.5f;
