Path Tracer:  Multi-GPU fix for CPU crash, and change from push to pull for cross GPU transfers.  The multi-GPU path tracer still hangs on the GPU due to issues with bindless ray tracing, but on the last good CL, push transfers were broken, so assuming the hang was fixed, this is necessary to get correct visual results.

#rnx

diff --git a/Engine/Source/Runtime/Renderer/Private/PathTracing.cpp b/Engine/Source/Runtime/Renderer/Private/PathTracing.cpp
--- a/Engine/Source/Runtime/Renderer/Private/PathTracing.cpp
+++ b/Engine/Source/Runtime/Renderer/Private/PathTracing.cpp
@@ -3763,10 +3763,10 @@ void FDeferredShadingSceneRenderer::RenderPathTracing(
 									TileToCopy.Min.Y = TileY + CurrentGPU * DispatchSizeYSplit;
 									TileToCopy.Max.X = TileX + DispatchSizeX;
 									TileToCopy.Max.Y = TileToCopy.Min.Y + DispatchSizeYLocal;
-									TransferParams.Emplace(Parameters->InputTexture->GetRHI(), TileToCopy, GPUIndex, FirstGPUIndex, false, false);
-									TransferParams.Emplace(Parameters->InputAlbedo->GetRHI(), TileToCopy, GPUIndex, FirstGPUIndex, false, false);
-									TransferParams.Emplace(Parameters->InputNormal->GetRHI(), TileToCopy, GPUIndex, FirstGPUIndex, false, false);
-									TransferParams.Emplace(Parameters->InputDepth->GetRHI(), TileToCopy, GPUIndex, FirstGPUIndex, false, false);
+									TransferParams.Emplace(Parameters->InputTexture->GetRHI(), TileToCopy, GPUIndex, FirstGPUIndex, true, true);
+									TransferParams.Emplace(Parameters->InputAlbedo->GetRHI(), TileToCopy, GPUIndex, FirstGPUIndex, true, true);
+									TransferParams.Emplace(Parameters->InputNormal->GetRHI(), TileToCopy, GPUIndex, FirstGPUIndex, true, true);
+									TransferParams.Emplace(Parameters->InputDepth->GetRHI(), TileToCopy, GPUIndex, FirstGPUIndex, true, true);
 								}
 								++CurrentGPU;
 							}
@@ -3826,8 +3826,8 @@ void FDeferredShadingSceneRenderer::RenderPathTracing(
 		TShaderMapRef<FPathTracingSwizzleScanlinesCS> ComputeShader(GetGlobalShaderMap(View.FeatureLevel));
 		FRDGTexture* NewRadianceTexture = GraphBuilder.CreateTexture(RadianceTexture->Desc, TEXT("PathTracer.RadianceUnswizzled"));
 		FRDGTexture* NewDepthTexture = GraphBuilder.CreateTexture(DepthTexture->Desc, TEXT("PathTracer.DepthUnswizzled"));
-		FRDGTexture* NewNormalTexture = NeedsDenoise ? GraphBuilder.CreateTexture(NormalTexture->Desc, TEXT("PathTracer.NormalUnswizzled")) : nullptr;
-		FRDGTexture* NewAlbedoTexture = NeedsDenoise ? GraphBuilder.CreateTexture(AlbedoTexture->Desc, TEXT("PathTracer.AlbedoUnswizzled")) : nullptr;
+		FRDGTexture* NewNormalTexture = IsDenoiserEnabled ? GraphBuilder.CreateTexture(NormalTexture->Desc, TEXT("PathTracer.NormalUnswizzled")) : nullptr;
+		FRDGTexture* NewAlbedoTexture = IsDenoiserEnabled ? GraphBuilder.CreateTexture(AlbedoTexture->Desc, TEXT("PathTracer.AlbedoUnswizzled")) : nullptr;
 
 		FRDGTexture* InputTextures[4] = { RadianceTexture, NormalTexture, DepthTexture, AlbedoTexture};
 		FRDGTexture* OutputTextures[4] = { NewRadianceTexture, NewNormalTexture, NewDepthTexture, NewAlbedoTexture};
