Fix bug where if any calls to FVirtualTextureSystem::LoadPendingTiles() fail to lock any outstanding VT root pages then there will not be another attempt to lock the tiles in the main render loop.
Before the fix the world partition tab could try and lock RVT tiles, and would fail. Then the main viewport wouldn't be able to lock the tiles after this, leaving the RVT waiting full initialization.
The new bug was triggered as a side effect of fixing UE-219361 after which root pages could only be locked in the viewport main render loop.
#rnx

diff --git a/Engine/Source/Runtime/Renderer/Private/VT/VirtualTextureSystem.cpp b/Engine/Source/Runtime/Renderer/Private/VT/VirtualTextureSystem.cpp
--- a/Engine/Source/Runtime/Renderer/Private/VT/VirtualTextureSystem.cpp
+++ b/Engine/Source/Runtime/Renderer/Private/VT/VirtualTextureSystem.cpp
@@ -1254,6 +1254,9 @@ void FVirtualTextureSystem::LoadPendingTiles(FRDGBuilder& GraphBuilder, ERHIFeat
 		AllocateResources(GraphBuilder);
 		SubmitRequests(GraphBuilder.RHICmdList, FeatureLevel, Allocator, Settings, RequestList, false);
 		FinalizeRequests(GraphBuilder);
+
+		// Swap any failed tile locks to be picked up by next BeginUpdate()
+		TilesToLock = MoveTemp(TilesToLockForNextFrame);
 	}
 }
 
