Fixed PSO pre-caching for a mobile VR projects
Correctly pre-cache subpass state for MSAA custom resolve case
Correctly pre-cache mobile multiview and shading rate attachment cases

diff --git a/Engine/Source/Runtime/D3D12RHI/Private/D3D12State.cpp b/Engine/Source/Runtime/D3D12RHI/Private/D3D12State.cpp
--- a/Engine/Source/Runtime/D3D12RHI/Private/D3D12State.cpp
+++ b/Engine/Source/Runtime/D3D12RHI/Private/D3D12State.cpp
@@ -505,8 +505,6 @@ uint64 FD3D12DynamicRHI::RHIComputePrecachePSOHash(const FGraphicsPipelineStateI
 		uint16							NumSamples;
 		EConservativeRasterization		ConservativeRasterization;
 		bool							bDepthBounds;
-		uint8							MultiViewCount;
-		bool							bHasFragmentDensityAttachment;
 		EVRSShadingRate					ShadingRate;
 	} HashKey;
 
@@ -521,8 +519,6 @@ uint64 FD3D12DynamicRHI::RHIComputePrecachePSOHash(const FGraphicsPipelineStateI
 	HashKey.NumSamples						= Initializer.NumSamples;
 	HashKey.ConservativeRasterization		= Initializer.ConservativeRasterization;
 	HashKey.bDepthBounds					= Initializer.bDepthBounds;
-	HashKey.MultiViewCount					= Initializer.MultiViewCount;
-	HashKey.bHasFragmentDensityAttachment	= Initializer.bHasFragmentDensityAttachment;
 	HashKey.ShadingRate						= Initializer.ShadingRate;
 
 	return CityHash64((const char*)&HashKey, sizeof(FNonStateHashKey));
diff --git a/Engine/Source/Runtime/Engine/Private/SceneTexturesConfig.cpp b/Engine/Source/Runtime/Engine/Private/SceneTexturesConfig.cpp
--- a/Engine/Source/Runtime/Engine/Private/SceneTexturesConfig.cpp
+++ b/Engine/Source/Runtime/Engine/Private/SceneTexturesConfig.cpp
@@ -7,6 +7,7 @@
 #include "RenderUtils.h"
 #include "SceneInterface.h"
 #include "StereoRenderTargetManager.h"
+#include "DataDrivenShaderPlatformInfo.h"
 
 FSceneTexturesConfig FSceneTexturesConfig::GlobalInstance;
 
@@ -336,8 +337,9 @@ void FSceneTexturesConfig::Init(const FSceneTexturesConfigInitSettings& InitSett
 
 	if (ShadingPath == EShadingPath::Mobile)
 	{
-		bRequiresDepthAux = MobileRequiresSceneDepthAux(ShaderPlatform);
+		bRequiresDepthAux = MobileRequiresSceneDepthAux(ShaderPlatform) && !IsMobileTonemapSubpassEnabled(ShaderPlatform);
 		bPreciseDepthAux = bPreciseDepthAux || MobileRequiresPreciseSceneDepthAux(ShaderPlatform);
+		bCustomResolveSubpass = IsMobileTonemapSubpassEnabledInline(ShaderPlatform, NumSamples);
 	}
 }
 
@@ -354,8 +356,7 @@ void FSceneTexturesConfig::BuildSceneColorAndDepthFlags()
 
 uint32 FSceneTexturesConfig::GetGBufferRenderTargetsInfo(FGraphicsPipelineRenderTargetsInfo& RenderTargetsInfo, EGBufferLayout Layout) const 
 {
-	// Assume 1 sample for now
-	RenderTargetsInfo.NumSamples = 1;
+	RenderTargetsInfo.NumSamples = NumSamples;
 
 	uint32 RenderTargetCount = 0;
 
@@ -397,6 +398,14 @@ uint32 FSceneTexturesConfig::GetGBufferRenderTargetsInfo(FGraphicsPipelineRender
 		RenderTargetsInfo.RenderTargetFormats[RenderTargetCount] = bPreciseDepthAux ? PF_R32_FLOAT : PF_R16F;
 		RenderTargetsInfo.RenderTargetFlags[RenderTargetCount++] = TexCreate_RenderTargetable | TexCreate_ShaderResource | TexCreate_InputAttachmentRead;
 	}
+	
+	if (bCustomResolveSubpass)
+	{
+		// resolve target as an additional color attachment
+		// this is supposed to be be a swapchain pixel format, but atm there is no way to query it
+		RenderTargetsInfo.RenderTargetFormats[RenderTargetCount] = IsAndroidPlatform(ShaderPlatform) ? PF_R8G8B8A8 : PF_B8G8R8A8;
+		RenderTargetsInfo.RenderTargetFlags[RenderTargetCount++] = TexCreate_RenderTargetable | TexCreate_ShaderResource;
+	}
 
 	// Store final number of render targets
 	RenderTargetsInfo.RenderTargetsEnabled = RenderTargetCount;
diff --git a/Engine/Source/Runtime/Engine/Private/SceneUtils.cpp b/Engine/Source/Runtime/Engine/Private/SceneUtils.cpp
--- a/Engine/Source/Runtime/Engine/Private/SceneUtils.cpp
+++ b/Engine/Source/Runtime/Engine/Private/SceneUtils.cpp
@@ -49,6 +49,18 @@ bool IsMobilePropagateAlphaEnabled(EShaderPlatform Platform)
 	return IsMobilePlatform(Platform) && (FPlatformMisc::GetMobilePropagateAlphaSetting() > 0);
 }
 
+ENGINE_API bool IsMobileTonemapSubpassEnabled(EShaderPlatform Platform)
+{
+	static auto* MobileTonemapSubpassPathCvar = IConsoleManager::Get().FindTConsoleVariableDataInt(TEXT("r.Mobile.TonemapSubpass"));
+	return (MobileTonemapSubpassPathCvar && (MobileTonemapSubpassPathCvar->GetValueOnAnyThread() == 1)) && IsMobileHDR() && !IsMobileDeferredShadingEnabled(Platform);
+}
+
+ENGINE_API bool IsMobileTonemapSubpassEnabledInline(EShaderPlatform Platform, uint32 NumMSAASamples)
+{
+	// As of UE 5.4 only vulkan supports inline (single pass) tonemap
+	return IsMobileTonemapSubpassEnabled(Platform) && IsVulkanPlatform(Platform) && (GRHISupportsMSAAShaderResolve || NumMSAASamples <= 1u);
+}
+
 ENGINE_API bool IsMobileColorsRGB()
 {
 	static auto* MobileUseHWsRGBEncodingCVAR = IConsoleManager::Get().FindTConsoleVariableDataInt(TEXT("r.Mobile.UseHWsRGBEncoding"));
diff --git a/Engine/Source/Runtime/Engine/Public/SceneTexturesConfig.h b/Engine/Source/Runtime/Engine/Public/SceneTexturesConfig.h
--- a/Engine/Source/Runtime/Engine/Public/SceneTexturesConfig.h
+++ b/Engine/Source/Runtime/Engine/Public/SceneTexturesConfig.h
@@ -129,6 +129,7 @@ struct FSceneTexturesConfig
 		, bKeepDepthContent{ 1 }
 		, bRequiresDepthAux{}
 		, bPreciseDepthAux{}
+		, bCustomResolveSubpass{}
 		, bSamplesCustomStencil{}
 		, bMemorylessMSAA{}
 		, bSupportsXRTargetManagerDepthAlloc{}
@@ -203,6 +204,9 @@ struct FSceneTexturesConfig
 	
 	// (Mobile) True if SceneDepthAux should use a precise pixel format
 	uint32 bPreciseDepthAux : 1;
+	
+	// (Mobile) True if Custom MSAA resolve subpass is enabled
+	uint32 bCustomResolveSubpass : 1;
 
 	// (Mobile) True if CustomStencil are sampled in a shader
 	uint32 bSamplesCustomStencil : 1;
diff --git a/Engine/Source/Runtime/Engine/Public/SceneUtils.h b/Engine/Source/Runtime/Engine/Public/SceneUtils.h
--- a/Engine/Source/Runtime/Engine/Public/SceneUtils.h
+++ b/Engine/Source/Runtime/Engine/Public/SceneUtils.h
@@ -60,7 +60,8 @@ inline bool IsTemporalAccumulationBasedMethod(EAntiAliasingMethod AntiAliasingMe
 
 /** True if Alpha Propagate is enabled for the mobile renderer. */
 ENGINE_API bool IsMobilePropagateAlphaEnabled(EShaderPlatform Platform);
-
+ENGINE_API bool IsMobileTonemapSubpassEnabled(EShaderPlatform Platform);
+ENGINE_API bool IsMobileTonemapSubpassEnabledInline(EShaderPlatform Platform, uint32 NumMSAASamples);
 ENGINE_API bool IsMobileColorsRGB();
 
 ENGINE_API EAntiAliasingMethod GetDefaultAntiAliasingMethod(const FStaticFeatureLevel InFeatureLevel);
diff --git a/Engine/Source/Runtime/RHI/Private/DynamicRHI.cpp b/Engine/Source/Runtime/RHI/Private/DynamicRHI.cpp
--- a/Engine/Source/Runtime/RHI/Private/DynamicRHI.cpp
+++ b/Engine/Source/Runtime/RHI/Private/DynamicRHI.cpp
@@ -588,12 +588,10 @@ uint64 FDynamicRHI::RHIComputeStatePrecachePSOHash(const FGraphicsPipelineStateI
 		uint32 DepthStencilState;
 		uint32 ImmutableSamplerState;
 
-		uint32 MultiViewCount : 8;
 		uint32 DrawShadingRate : 8;
 		uint32 PrimitiveType : 8;
 		uint32 bDepthBounds : 1;
-		uint32 bHasFragmentDensityAttachment : 1;
-		uint32 Unused : 6;
+		uint32 Unused : 15;
 	} HashKey;
 
 	FMemory::Memzero(&HashKey, sizeof(FHashKey));
@@ -627,11 +625,9 @@ uint64 FDynamicRHI::RHIComputeStatePrecachePSOHash(const FGraphicsPipelineStateI
 	// Ignore immutable samplers for now
 	//HashKey.ImmutableSamplerState = GetTypeHash(ImmutableSamplerState);
 
-	HashKey.MultiViewCount = Initializer.MultiViewCount;
 	HashKey.DrawShadingRate = Initializer.ShadingRate;
 	HashKey.PrimitiveType = Initializer.PrimitiveType;
 	HashKey.bDepthBounds = Initializer.bDepthBounds;
-	HashKey.bHasFragmentDensityAttachment = Initializer.bHasFragmentDensityAttachment;
 
 	uint64 PrecachePSOHash = CityHash64((const char*)&HashKey, sizeof(FHashKey));
 
diff --git a/Engine/Source/Runtime/Renderer/Private/CompositionLighting/PostProcessDeferredDecals.cpp b/Engine/Source/Runtime/Renderer/Private/CompositionLighting/PostProcessDeferredDecals.cpp
--- a/Engine/Source/Runtime/Renderer/Private/CompositionLighting/PostProcessDeferredDecals.cpp
+++ b/Engine/Source/Runtime/Renderer/Private/CompositionLighting/PostProcessDeferredDecals.cpp
@@ -114,7 +114,6 @@ FDeferredDecalPassTextures GetDeferredDecalPassTextures(
 
 void GetDeferredDecalRenderTargetsInfo(
 	const FSceneTexturesConfig& Config,
-	EShaderPlatform ShaderPlatform,
 	EDecalRenderTargetMode RenderTargetMode,
 	FGraphicsPipelineRenderTargetsInfo& RenderTargetsInfo)
 {
@@ -138,7 +137,7 @@ void GetDeferredDecalRenderTargetsInfo(
 
 	case EDecalRenderTargetMode::DBuffer:
 	{
-		const FDBufferTexturesDesc DBufferTexturesDesc = GetDBufferTexturesDesc(Config.Extent, ShaderPlatform);
+		const FDBufferTexturesDesc DBufferTexturesDesc = GetDBufferTexturesDesc(Config.Extent, Config.ShaderPlatform);
 
 		AddRenderTargetInfo(DBufferTexturesDesc.DBufferADesc.Format, DBufferTexturesDesc.DBufferADesc.Flags, RenderTargetsInfo);
 		AddRenderTargetInfo(DBufferTexturesDesc.DBufferBDesc.Format, DBufferTexturesDesc.DBufferBDesc.Flags, RenderTargetsInfo);
@@ -171,6 +170,13 @@ void GetDeferredDecalRenderTargetsInfo(
 			AddRenderTargetInfo(Config.bPreciseDepthAux ? PF_R32_FLOAT : PF_R16F, TexCreate_RenderTargetable | TexCreate_ShaderResource | TexCreate_InputAttachmentRead, RenderTargetsInfo);	
 		};
 	}
+	if (Config.bCustomResolveSubpass)
+	{
+		// resolve target as an additional color attachment
+		AddRenderTargetInfo(IsAndroidPlatform(Config.ShaderPlatform) ? PF_R8G8B8A8 : PF_B8G8R8A8, TexCreate_RenderTargetable | TexCreate_ShaderResource, RenderTargetsInfo);
+	}
+
+	RenderTargetsInfo.NumSamples = Config.NumSamples;
 
 	SetupDepthStencilInfo(PF_DepthStencil, Config.DepthCreateFlags, ERenderTargetLoadAction::ELoad,
 		ERenderTargetLoadAction::ELoad, FExclusiveDepthStencil::DepthRead_StencilWrite, RenderTargetsInfo);
@@ -537,21 +543,15 @@ void CollectDeferredDecalPassPSOInitializers(
 	}
 
 	FGraphicsPipelineRenderTargetsInfo RenderTargetsInfo;
-	RenderTargetsInfo.NumSamples = 1;
+	GetDeferredDecalRenderTargetsInfo(SceneTexturesConfig, DecalRenderTargetMode, RenderTargetsInfo);
+	ApplyTargetsInfo(GraphicsPSOInit, RenderTargetsInfo);
 
 	if (FeatureLevel == ERHIFeatureLevel::ES3_1)
 	{
 		// subpass info set during the submission of the draws in a mobile renderer
 		GraphicsPSOInit.SubpassIndex = 1; // all decals use second sub-pass on mobile
-		GraphicsPSOInit.SubpassHint = SceneTexturesConfig.bIsUsingGBuffers ? ESubpassHint::DeferredShadingSubpass : ESubpassHint::DepthReadSubpass;
-		if (DecalRenderTargetMode == EDecalRenderTargetMode::SceneColor)
-		{
-			RenderTargetsInfo.NumSamples = SceneTexturesConfig.NumSamples;
-		}
+		GraphicsPSOInit.SubpassHint = GetSubpassHint(SceneTexturesConfig.ShaderPlatform, SceneTexturesConfig.bIsUsingGBuffers, SceneTexturesConfig.NumSamples);
 	}
-
-	GetDeferredDecalRenderTargetsInfo(SceneTexturesConfig, ShaderPlatform, DecalRenderTargetMode, RenderTargetsInfo);
-	ApplyTargetsInfo(GraphicsPSOInit, RenderTargetsInfo);
 		
 	const auto AddDeferredDecalPSO = [&](bool bInsideDecal,	bool bReverseHanded, bool bReverseCulling, bool bDecalUsesStencil)
 	{
diff --git a/Engine/Source/Runtime/Renderer/Private/CompositionLighting/PostProcessDeferredDecals.h b/Engine/Source/Runtime/Renderer/Private/CompositionLighting/PostProcessDeferredDecals.h
--- a/Engine/Source/Runtime/Renderer/Private/CompositionLighting/PostProcessDeferredDecals.h
+++ b/Engine/Source/Runtime/Renderer/Private/CompositionLighting/PostProcessDeferredDecals.h
@@ -76,7 +76,6 @@ END_SHADER_PARAMETER_STRUCT()
 
 void GetDeferredDecalRenderTargetsInfo(
 	const FSceneTexturesConfig& Config,
-	EShaderPlatform ShaderPlatform,
 	EDecalRenderTargetMode RenderTargetMode,
 	FGraphicsPipelineRenderTargetsInfo& RenderTargetsInfo);
 
diff --git a/Engine/Source/Runtime/Renderer/Private/CompositionLighting/PostProcessMeshDecals.cpp b/Engine/Source/Runtime/Renderer/Private/CompositionLighting/PostProcessMeshDecals.cpp
--- a/Engine/Source/Runtime/Renderer/Private/CompositionLighting/PostProcessMeshDecals.cpp
+++ b/Engine/Source/Runtime/Renderer/Private/CompositionLighting/PostProcessMeshDecals.cpp
@@ -384,10 +384,8 @@ void FMeshDecalMeshProcessor::CollectDeferredDecalMeshPSOInitializers(
 	Shaders.TryGetVertexShader(MeshDecalPassShaders.VertexShader);
 	Shaders.TryGetPixelShader(MeshDecalPassShaders.PixelShader);
 
-	const EShaderPlatform ShaderPlatform = GetFeatureLevelShaderPlatform(FeatureLevel);
 	FGraphicsPipelineRenderTargetsInfo RenderTargetsInfo;
-	RenderTargetsInfo.NumSamples = 1;
-	GetDeferredDecalRenderTargetsInfo(SceneTexturesConfig, ShaderPlatform, LocalRenderTargetMode, RenderTargetsInfo);
+	GetDeferredDecalRenderTargetsInfo(SceneTexturesConfig, LocalRenderTargetMode, RenderTargetsInfo);
 
 	uint8 SubpassIndex = 0;
 	ESubpassHint SubpassHint = ESubpassHint::None;
@@ -395,7 +393,7 @@ void FMeshDecalMeshProcessor::CollectDeferredDecalMeshPSOInitializers(
 	{
 		// subpass info set during the submission of the draws in a mobile renderer
 		SubpassIndex = 1; // all decals use second sub-pass on mobile
-		SubpassHint = SceneTexturesConfig.bIsUsingGBuffers ? ESubpassHint::DeferredShadingSubpass : ESubpassHint::DepthReadSubpass;
+		SubpassHint = GetSubpassHint(SceneTexturesConfig.ShaderPlatform, SceneTexturesConfig.bIsUsingGBuffers, SceneTexturesConfig.NumSamples);
 	}
 
 	AddGraphicsPipelineStateInitializer(
diff --git a/Engine/Source/Runtime/Renderer/Private/MeshPassProcessor.cpp b/Engine/Source/Runtime/Renderer/Private/MeshPassProcessor.cpp
--- a/Engine/Source/Runtime/Renderer/Private/MeshPassProcessor.cpp
+++ b/Engine/Source/Runtime/Renderer/Private/MeshPassProcessor.cpp
@@ -1588,6 +1588,25 @@ void ApplyTargetsInfo(FGraphicsPipelineStateInitializer& GraphicsPSOInit, const
 	GraphicsPSOInit.bHasFragmentDensityAttachment = RenderTargetsInfo.bHasFragmentDensityAttachment;
 }
 
+ESubpassHint GetSubpassHint(const FStaticShaderPlatform Platform, bool bIsUsingGBuffers, uint32 NumSamples)
+{
+	ESubpassHint SubpassHint = ESubpassHint::None;
+	
+	if (IsMobilePlatform(Platform))
+	{
+		if (bIsUsingGBuffers)
+		{
+			SubpassHint = ESubpassHint::DeferredShadingSubpass;
+		}
+		else
+		{
+			SubpassHint = IsMobileTonemapSubpassEnabledInline(Platform, NumSamples) ? ESubpassHint::CustomResolveSubpass : ESubpassHint::DepthReadSubpass;
+		}
+	}
+
+	return SubpassHint;
+}
+
 uint64 FMeshDrawCommand::GetPipelineStateSortingKey(FRHICommandList& RHICmdList, const FGraphicsPipelineRenderTargetsInfo& RenderTargetsInfo) const
 {
 	// Default fallback sort key
diff --git a/Engine/Source/Runtime/Renderer/Private/MobileBasePass.cpp b/Engine/Source/Runtime/Renderer/Private/MobileBasePass.cpp
--- a/Engine/Source/Runtime/Renderer/Private/MobileBasePass.cpp
+++ b/Engine/Source/Runtime/Renderer/Private/MobileBasePass.cpp
@@ -14,6 +14,8 @@
 #include "MeshPassProcessor.inl"
 #include "Engine/TextureCube.h"
 #include "ShaderPlatformCachedIniValue.h"
+#include "StereoRenderUtils.h"
+#include "VariableRateShadingImageManager.h"
 
 bool MobileLocalLightsBufferEnabled(const FStaticShaderPlatform Platform)
 {
@@ -999,7 +1001,7 @@ void FMobileBasePassMeshProcessor::CollectPSOInitializersForLMPolicy(
 
 	// subpass info set during the submission of the draws in mobile deferred renderer.
 	uint8 SubpassIndex = bTranslucentBasePass ? (bDeferredShading ? 2 : 1) : 0;
-	ESubpassHint SubpassHint = bDeferredShading ? ESubpassHint::DeferredShadingSubpass : ESubpassHint::DepthReadSubpass;
+	ESubpassHint SubpassHint = GetSubpassHint(GMaxRHIShaderPlatform, bDeferredShading, RenderTargetsInfo.NumSamples);
 
 	AddGraphicsPipelineStateInitializer(
 		VertexFactoryData,
@@ -1018,6 +1020,15 @@ void FMobileBasePassMeshProcessor::CollectPSOInitializersForLMPolicy(
 		PSOInitializers);
 }
 
+static void SetupMultiViewInfo(FGraphicsPipelineRenderTargetsInfo& RenderTargetsInfo)
+{
+	const static UE::StereoRenderUtils::FStereoShaderAspects Aspects(GMaxRHIShaderPlatform);
+	// If mobile multiview is enabled we expect it will be used with a native MMV, no pre-caching for fallbacks 
+	RenderTargetsInfo.MultiViewCount = Aspects.IsMobileMultiViewEnabled() ? (GSupportsMobileMultiView ? 2 : 1) : 0;
+	// FIXME: Need to figure out if renderer will use shading rate texture or not
+	RenderTargetsInfo.bHasFragmentDensityAttachment = GVRSImageManager.IsAttachmentVRSEnabled();
+}
+
 void FMobileBasePassMeshProcessor::CollectPSOInitializers(const FSceneTexturesConfig& SceneTexturesConfig, const FMaterial& Material, const FPSOPrecacheVertexFactoryData& VertexFactoryData, const FPSOPrecacheParams& PreCacheParams, TArray<FPSOPrecacheData>& PSOInitializers)
 {
 	if (bTranslucentBasePass)
@@ -1059,26 +1070,20 @@ void FMobileBasePassMeshProcessor::CollectPSOInitializers(const FSceneTexturesCo
 		PreCacheParams.Mobility == EComponentMobility::Stationary || 
 		PreCacheParams.bUsesIndirectLightingCache; // ILC uses movable path
 
-	bool bDitheredLODTransition = !bMovable && Material.IsDitheredLODTransition() && !PreCacheParams.bForceLODModel;
-
 	// Setup the draw state
-	FGraphicsPipelineRenderTargetsInfo RenderTargetsInfo;
-	RenderTargetsInfo.NumSamples = SceneTexturesConfig.NumSamples;
-
 	FMeshPassProcessorRenderState DrawRenderState(PassDrawRenderState);
-	EPixelFormat SceneColorFormat = SceneTexturesConfig.ColorFormat;
-	ETextureCreateFlags SceneColorCreateFlags = SceneTexturesConfig.ColorCreateFlags;
-
+	
 	const bool bMaskedInEarlyPass = MaskedInEarlyPass(ShaderPlatform);
-
 	FExclusiveDepthStencil ExclusiveDepthStencil = (bTranslucentBasePass || bMaskedInEarlyPass) ? 
 		FExclusiveDepthStencil::DepthRead_StencilRead : 
 		FExclusiveDepthStencil::DepthWrite_StencilWrite;
 
+	FGraphicsPipelineRenderTargetsInfo RenderTargetsInfo;
 	SetupGBufferRenderTargetInfo(SceneTexturesConfig, RenderTargetsInfo, false /*bSetupDepthStencil*/);
 	SetupDepthStencilInfo(PF_DepthStencil, SceneTexturesConfig.DepthCreateFlags, ERenderTargetLoadAction::ELoad,
 		ERenderTargetLoadAction::ELoad, ExclusiveDepthStencil, RenderTargetsInfo);
-
+	SetupMultiViewInfo(RenderTargetsInfo);
+					
 	if (bTranslucentBasePass)
 	{
 		MobileBasePass::SetTranslucentRenderState(DrawRenderState, Material, ShadingModels);
diff --git a/Engine/Source/Runtime/Renderer/Private/MobileBasePassRendering.cpp b/Engine/Source/Runtime/Renderer/Private/MobileBasePassRendering.cpp
--- a/Engine/Source/Runtime/Renderer/Private/MobileBasePassRendering.cpp
+++ b/Engine/Source/Runtime/Renderer/Private/MobileBasePassRendering.cpp
@@ -56,6 +56,13 @@ static TAutoConsoleVariable<int32> CVarMobilePropagateAlpha(
 	TEXT("1: Propagate Full Alpha Propagate"),
 	ECVF_ReadOnly | ECVF_RenderThreadSafe);
 
+static TAutoConsoleVariable<int32> CVarMobileTonemapSubpass(
+	TEXT("r.Mobile.TonemapSubpass"),
+	0,
+	TEXT(" Whether to enable mobile tonemap subpass \n")
+	TEXT(" 0 = Off [default]\n")
+	TEXT(" 1 = On"),
+	ECVF_Scalability | ECVF_RenderThreadSafe);
 
 IMPLEMENT_STATIC_UNIFORM_BUFFER_STRUCT(FMobileBasePassUniformParameters, "MobileBasePass", SceneTextures);
 
diff --git a/Engine/Source/Runtime/Renderer/Private/MobileShadingRenderer.cpp b/Engine/Source/Runtime/Renderer/Private/MobileShadingRenderer.cpp
--- a/Engine/Source/Runtime/Renderer/Private/MobileShadingRenderer.cpp
+++ b/Engine/Source/Runtime/Renderer/Private/MobileShadingRenderer.cpp
@@ -101,15 +101,6 @@ static TAutoConsoleVariable<int32> CVarMobileCustomDepthForTranslucency(
 	TEXT(" 1 = On [default]"),
 	ECVF_Scalability | ECVF_RenderThreadSafe);
 
-static TAutoConsoleVariable<int32> CVarMobileTonemapSubpass(
-	TEXT("r.Mobile.TonemapSubpass"),
-	0,
-	TEXT(" Whether to enable mobile tonemap subpass \n")
-	TEXT(" 0 = Off [default]\n")
-	TEXT(" 1 = On"),
-	ECVF_Scalability | ECVF_RenderThreadSafe);
-
-
 static TAutoConsoleVariable<int32> CVarMobileXRMSAAMode(
 	TEXT("r.Mobile.XRMSAAMode"),
 	0,
@@ -119,13 +110,6 @@ static TAutoConsoleVariable<int32> CVarMobileXRMSAAMode(
 	TEXT(" 2 = Make the depth swap chain be MSAA and use it directly as scene depth"),
 	ECVF_ReadOnly | ECVF_RenderThreadSafe);
 
-
-static bool IsMobileTonemapSubpassEnabled(const FStaticShaderPlatform Platform)
-{
-	static auto* MobileTonemapSubpassPathCvar = IConsoleManager::Get().FindTConsoleVariableDataInt(TEXT("r.Mobile.TonemapSubpass"));
-	return (MobileTonemapSubpassPathCvar && (MobileTonemapSubpassPathCvar->GetValueOnAnyThread() == 1)) && IsMobileHDR() && !IsMobileDeferredShadingEnabled(Platform);
-}
-
 DECLARE_GPU_STAT_NAMED(MobileSceneRender, TEXT("Mobile Scene Render"));
 
 extern bool IsMobileEyeAdaptationEnabled(const FViewInfo& View);
@@ -321,7 +305,7 @@ FMobileSceneRenderer::FMobileSceneRenderer(const FSceneViewFamily* InViewFamily,
 	NumMSAASamples = GetDefaultMSAACount(ERHIFeatureLevel::ES3_1);
 	// As of UE 5.4 only vulkan supports inline (single pass) tonemap
 	bTonemapSubpass = IsMobileTonemapSubpassEnabled(ShaderPlatform) && ViewFamily.bResolveScene && GetRendererOutput() != FSceneRenderer::ERendererOutput::DepthPrepassOnly;
-	bTonemapSubpassInline = bTonemapSubpass && IsVulkanPlatform(ShaderPlatform) && (GRHISupportsMSAAShaderResolve || NumMSAASamples == 1);
+	bTonemapSubpassInline = IsMobileTonemapSubpassEnabledInline(ShaderPlatform, NumMSAASamples) && bTonemapSubpass;
 	bRequiresSceneDepthAux = MobileRequiresSceneDepthAux(ShaderPlatform) && !bTonemapSubpass;
 }
 
diff --git a/Engine/Source/Runtime/Renderer/Public/MeshPassProcessor.h b/Engine/Source/Runtime/Renderer/Public/MeshPassProcessor.h
--- a/Engine/Source/Runtime/Renderer/Public/MeshPassProcessor.h
+++ b/Engine/Source/Runtime/Renderer/Public/MeshPassProcessor.h
@@ -366,8 +366,8 @@ public:
 			, EConservativeRasterization::Disabled
 			, 0
 			, bDepthBounds
-			, MultiViewCount
-			, bHasFragmentDensityAttachment
+			, 0
+			, false
 			, bAllowVariableRateShading
 			, DrawShadingRate
 		);
@@ -397,8 +397,6 @@ public:
 			DepthStencilState != rhs.DepthStencilState ||
 			ImmutableSamplerState != rhs.ImmutableSamplerState ||
 			bDepthBounds != rhs.bDepthBounds ||
-			MultiViewCount != rhs.MultiViewCount ||
-			bHasFragmentDensityAttachment != rhs.bHasFragmentDensityAttachment ||
 			bAllowVariableRateShading != rhs.bAllowVariableRateShading ||
 			DrawShadingRate != rhs.DrawShadingRate ||
 			PrimitiveType != rhs.PrimitiveType)
@@ -460,8 +458,6 @@ public:
 			COMPARE_FIELD(RasterizerState)
 			COMPARE_FIELD(DepthStencilState)
 			COMPARE_FIELD(bDepthBounds)
-			COMPARE_FIELD(MultiViewCount)
-			COMPARE_FIELD(bHasFragmentDensityAttachment)
 			COMPARE_FIELD(bAllowVariableRateShading)
 			COMPARE_FIELD(DrawShadingRate)
 			COMPARE_FIELD(PrimitiveType)
@@ -487,8 +483,6 @@ public:
 			COMPARE_FIELD(RasterizerState)
 			COMPARE_FIELD(DepthStencilState)
 			COMPARE_FIELD(bDepthBounds)
-			COMPARE_FIELD(MultiViewCount)
-			COMPARE_FIELD(bHasFragmentDensityAttachment)
 			COMPARE_FIELD(bAllowVariableRateShading)
 			COMPARE_FIELD(DrawShadingRate)
 			COMPARE_FIELD(PrimitiveType)
@@ -513,8 +507,6 @@ public:
 	// as it is sometimes hashed and compared as raw bytes. Explicit padding is therefore required between
 	// all data members and at the end of the structure.
 	bool							bDepthBounds = false;
-	uint8							MultiViewCount = 0;
-	bool							bHasFragmentDensityAttachment = false;
 	bool							bAllowVariableRateShading = true;
 	EVRSShadingRate					DrawShadingRate  = EVRSShadingRate::VRSSR_1x1;
 
@@ -2362,6 +2354,7 @@ RENDERER_API extern void AddRenderTargetInfo(EPixelFormat PixelFormat, ETextureC
 RENDERER_API extern void SetupDepthStencilInfo(EPixelFormat DepthStencilFormat, ETextureCreateFlags DepthStencilCreateFlags, ERenderTargetLoadAction DepthTargetLoadAction, ERenderTargetLoadAction StencilTargetLoadAction, FExclusiveDepthStencil DepthStencilAccess, FGraphicsPipelineRenderTargetsInfo& RenderTargetsInfo);
 RENDERER_API extern void SetupGBufferRenderTargetInfo(const FSceneTexturesConfig& SceneTexturesConfig, FGraphicsPipelineRenderTargetsInfo& RenderTargetsInfo, bool bSetupDepthStencil);
 RENDERER_API extern void ApplyTargetsInfo(FGraphicsPipelineStateInitializer& GraphicsPSOInit, const FGraphicsPipelineRenderTargetsInfo& RenderTargetsInfo);
+RENDERER_API extern ESubpassHint GetSubpassHint(const FStaticShaderPlatform Platform, bool bIsUsingGBuffers, uint32 NumSamples);
 
 inline FMeshDrawCommandSortKey CalculateMeshStaticSortKey(const TShaderRef<FMeshMaterialShader>& VertexShader, const TShaderRef<FMeshMaterialShader>& PixelShader)
 {
diff --git a/Engine/Source/Runtime/VulkanRHI/Private/VulkanRHI.cpp b/Engine/Source/Runtime/VulkanRHI/Private/VulkanRHI.cpp
--- a/Engine/Source/Runtime/VulkanRHI/Private/VulkanRHI.cpp
+++ b/Engine/Source/Runtime/VulkanRHI/Private/VulkanRHI.cpp
@@ -2066,7 +2066,7 @@ IRHITransientResourceAllocator* FVulkanDynamicRHI::RHICreateTransientResourceAll
 
 uint32 FVulkanDynamicRHI::GetPrecachePSOHashVersion()
 {
-	static const uint32 PrecacheHashVersion = 2;
+	static const uint32 PrecacheHashVersion = 3;
 	return PrecacheHashVersion;
 }
 
@@ -2090,13 +2090,11 @@ uint64 FVulkanDynamicRHI::RHIComputeStatePrecachePSOHash(const FGraphicsPipeline
 		uint32 DepthStencilState;
 		uint32 ImmutableSamplerState;
 
-		uint32 MultiViewCount : 8;
 		uint32 DrawShadingRate : 8;
 		uint32 PrimitiveType : 8;
 		uint32 bDepthBounds : 1;
-		uint32 bHasFragmentDensityAttachment : 1;
 		uint32 bAllowVariableRateShading : 1;
-		uint32 Unused : 5;
+		uint32 Unused : 14;
 	} HashKey;
 
 	FMemory::Memzero(&HashKey, sizeof(FHashKey));
@@ -2135,11 +2133,9 @@ uint64 FVulkanDynamicRHI::RHIComputeStatePrecachePSOHash(const FGraphicsPipeline
 	// Ignore immutable samplers for now
 	//HashKey.ImmutableSamplerState = GetTypeHash(ImmutableSamplerState);
 
-	HashKey.MultiViewCount = Initializer.MultiViewCount;
 	HashKey.DrawShadingRate = Initializer.ShadingRate;
 	HashKey.PrimitiveType = Initializer.PrimitiveType;
 	HashKey.bDepthBounds = Initializer.bDepthBounds;
-	HashKey.bHasFragmentDensityAttachment = Initializer.bHasFragmentDensityAttachment;
 	HashKey.bAllowVariableRateShading = Initializer.bAllowVariableRateShading;
 
 	uint64 PrecachePSOHash = CityHash64((const char*)&HashKey, sizeof(FHashKey));
@@ -2173,6 +2169,8 @@ uint64 FVulkanDynamicRHI::RHIComputePrecachePSOHash(const FGraphicsPipelineState
 		uint16							NumSamples;
 		ESubpassHint					SubpassHint;
 		uint8							SubpassIndex;
+		uint8							MultiViewCount;
+		bool							bHasFragmentDensityAttachment;
 		EConservativeRasterization		ConservativeRasterization;
 	} HashKey;
 
@@ -2188,6 +2186,8 @@ uint64 FVulkanDynamicRHI::RHIComputePrecachePSOHash(const FGraphicsPipelineState
 	HashKey.NumSamples = Initializer.NumSamples;
 	HashKey.SubpassHint = Initializer.SubpassHint;
 	HashKey.SubpassIndex = Initializer.SubpassIndex;
+	HashKey.MultiViewCount = Initializer.MultiViewCount;
+	HashKey.bHasFragmentDensityAttachment = Initializer.bHasFragmentDensityAttachment;
 	HashKey.ConservativeRasterization = Initializer.ConservativeRasterization;
 	
 	// TODO: check if any RT flags actually affect PSO in VK
