Move SMRT settings into the VSM global uniforms so they can be shared more easily by SMRT passes being called from different places.

#rnx

diff --git a/Engine/Shaders/Private/ForwardLightingCommon.ush b/Engine/Shaders/Private/ForwardLightingCommon.ush
--- a/Engine/Shaders/Private/ForwardLightingCommon.ush
+++ b/Engine/Shaders/Private/ForwardLightingCommon.ush
@@ -199,10 +199,9 @@ FDeferredLightingSplit GetForwardDirectLightingSplit(
 				#if SUPPORT_VSM_FOWARD_QUALITY == 1
 				{
 					const FLightShaderParameters DirectionalLight = ConvertToLightShaderParameters(LightData);
-					FVirtualShadowMapSMRTSettings SMRTSettings = GetDirectionalLightSMRTSettings();
 
 					const bool bBackfaceCull = !IsSubsurfaceModel(GBufferData.ShadingModelID);
-					const float ScreenRayLengthWorld = SMRTSettings.ScreenRayLength * View.ClipToView[1][1] * GBufferData.Depth;
+					const float ScreenRayLengthWorld = VirtualShadowMap.ScreenRayLength * View.ClipToView[1][1] * GBufferData.Depth;
 					float SMRTRayOffset = ScreenRayLengthWorld;
 
 					FVirtualShadowMapSampleResult VirtualShadowMapSample = TraceDirectional(
@@ -212,15 +211,9 @@ FDeferredLightingSplit GetForwardDirectLightingSplit(
 						GBufferData.Depth,
 						TranslatedWorldPosition,
 						SMRTRayOffset,
-						SMRTSettings.RayCount,
-						SMRTSettings.SamplesPerRay,
-						SMRTSettings.RayLengthScale,
 						Dither,
 						bBackfaceCull,
-						GBufferData.WorldNormal,
-						SMRTSettings.TexelDitherScale,
-						SMRTSettings.ExtrapolateSlope,
-						SMRTSettings.AdaptiveRayCount);
+						GBufferData.WorldNormal);
 
 					FilterVirtualShadowMapSampleResult(PixelPos, VirtualShadowMapSample);
 
diff --git a/Engine/Shaders/Private/LightGridCommon.ush b/Engine/Shaders/Private/LightGridCommon.ush
--- a/Engine/Shaders/Private/LightGridCommon.ush
+++ b/Engine/Shaders/Private/LightGridCommon.ush
@@ -157,34 +157,6 @@ FCulledLightsGridData GetCulledLightsGrid(uint GridIndex, uint EyeIndex)
 	return Result;
 }
 
-struct FVirtualShadowMapSMRTSettings
-{
-	int   RayCount;
-	int   SamplesPerRay;
-	float RayLengthScale;
-	float ScreenRayLength;
-	float CotMaxRayAngleFromLight;
-	float TexelDitherScale;
-	float ExtrapolateSlope;
-	float MaxSlopeBias;
-	uint  AdaptiveRayCount;
-};
-
-FVirtualShadowMapSMRTSettings GetDirectionalLightSMRTSettings()
-{
-	FVirtualShadowMapSMRTSettings Out;
-	Out.ScreenRayLength			= ForwardLightData.DirectionalLightSMRTSettings.ScreenRayLength;
-	Out.RayCount 				= ForwardLightData.DirectionalLightSMRTSettings.SMRTRayCount;
-	Out.SamplesPerRay 			= ForwardLightData.DirectionalLightSMRTSettings.SMRTSamplesPerRay;
-	Out.RayLengthScale 			= ForwardLightData.DirectionalLightSMRTSettings.SMRTRayLengthScale;
-	Out.CotMaxRayAngleFromLight = ForwardLightData.DirectionalLightSMRTSettings.SMRTCotMaxRayAngleFromLight;
-	Out.TexelDitherScale 		= ForwardLightData.DirectionalLightSMRTSettings.SMRTTexelDitherScale;
-	Out.ExtrapolateSlope 		= ForwardLightData.DirectionalLightSMRTSettings.SMRTExtrapolateSlope;
-	Out.MaxSlopeBias 			= ForwardLightData.DirectionalLightSMRTSettings.SMRTMaxSlopeBias;
-	Out.AdaptiveRayCount 		= ForwardLightData.DirectionalLightSMRTSettings.SMRTAdaptiveRayCount;
-	return Out;
-}
-
 FDirectionalLightData GetDirectionalLightData(uint EyeIndex)
 {
 	FDirectionalLightData Result;
diff --git a/Engine/Shaders/Private/Substrate/SubstrateForwardLighting.ush b/Engine/Shaders/Private/Substrate/SubstrateForwardLighting.ush
--- a/Engine/Shaders/Private/Substrate/SubstrateForwardLighting.ush
+++ b/Engine/Shaders/Private/Substrate/SubstrateForwardLighting.ush
@@ -119,10 +119,9 @@ float3 SubstrateForwardLighting(
 		#if SUPPORT_VSM_FOWARD_QUALITY == 1
 		{
 			const FLightShaderParameters DirectionalLight = ConvertToLightShaderParameters(DirLightData);
-			FVirtualShadowMapSMRTSettings SMRTSettings = GetDirectionalLightSMRTSettings();
 
 			const bool bBackfaceCull = true; // SUBSTRATE_TODO: this needs to be fixed but require to extract hair/SSS/diffusebackface prior to BSDF looping
-			const float ScreenRayLengthWorld = SMRTSettings.ScreenRayLength * View.ClipToView[1][1] * SceneDepth;
+			const float ScreenRayLengthWorld = VirtualShadowMap.ScreenRayLength * View.ClipToView[1][1] * SceneDepth;
 			float SMRTRayOffset = ScreenRayLengthWorld;
 
 			// Use the surface normal instead of the layer's normal, for a single evaluation before runing throught the SubstrateTree
@@ -135,15 +134,9 @@ float3 SubstrateForwardLighting(
 				SceneDepth,
 				TranslatedWorldPosition,
 				SMRTRayOffset,
-				SMRTSettings.RayCount,
-				SMRTSettings.SamplesPerRay,
-				SMRTSettings.RayLengthScale,
 				Dither,
 				bBackfaceCull,
-				WorldNormal,
-				SMRTSettings.TexelDitherScale,
-				SMRTSettings.ExtrapolateSlope,
-				SMRTSettings.AdaptiveRayCount);
+				WorldNormal);
 
 			FilterVirtualShadowMapSampleResult(PixelPos, VirtualShadowMapSample);
 
diff --git a/Engine/Shaders/Private/VirtualShadowMaps/VirtualShadowMapProjection.usf b/Engine/Shaders/Private/VirtualShadowMaps/VirtualShadowMapProjection.usf
--- a/Engine/Shaders/Private/VirtualShadowMaps/VirtualShadowMapProjection.usf
+++ b/Engine/Shaders/Private/VirtualShadowMaps/VirtualShadowMapProjection.usf
@@ -80,6 +80,14 @@ FProjectionShadingInfo GetProjectionShadingInfo(uint2 PixelPos)
 	return Out;
 }
 
+int GetSMRTRayCount()
+{
+#if DIRECTIONAL_LIGHT
+	return VirtualShadowMap.SMRTRayCountDirectional;
+#else
+	return VirtualShadowMap.SMRTRayCountLocal;
+#endif
+}
 
 SCHEDULER_MIN_PRESSURE
 MAX_OCCUPANCY
@@ -173,18 +181,8 @@ float4 ComputeRandom4(uint2 PixelPosition)
 }
 
 int4 ProjectionRect;
-float ScreenRayLength;
 float NormalBias;
 float SubsurfaceMinSourceRadius;
-int SMRTRayCount;					// 0 = off
-int SMRTSamplesPerRay;
-float SMRTRayLengthScale;			// Directional lights
-float SMRTCotMaxRayAngleFromLight;	// Spot/point lights
-float SMRTTexelDitherScale;			// Currently only directional lights
-float SMRTExtrapolateSlope;
-float SMRTMaxSlopeBias;				// Currently only local lights
-uint SMRTAdaptiveRayCount;
-uint SMRTHairRayCount;
 
 // One pass projection
 #if ONE_PASS_PROJECTION
@@ -284,7 +282,7 @@ FVirtualShadowMapSampleResult ProjectLight(
 		}
 
 		#if DIRECTIONAL_LIGHT
-		if (SMRTRayCount > 0)
+		if (GetSMRTRayCount() > 0)
 		{
 			Result = TraceDirectional(
 				VirtualShadowMapId,
@@ -293,15 +291,9 @@ FVirtualShadowMapSampleResult ProjectLight(
 				SceneDepth,
 				TranslatedWorldPosition,
 				SMRTRayOffset,
-				SMRTRayCount,
-				SMRTSamplesPerRay,
-				SMRTRayLengthScale,
 				Noise,
 				bBackfaceCull,
-				WorldNormal,
-				SMRTTexelDitherScale,
-				SMRTExtrapolateSlope,
-				SMRTAdaptiveRayCount);
+				WorldNormal);
 		}
 		else
 		{
@@ -312,7 +304,7 @@ FVirtualShadowMapSampleResult ProjectLight(
 				WorldNormal);
 		}
 		#else // !DIRECTIONAL_LIGHT
-		if (SMRTRayCount > 0)
+		if (GetSMRTRayCount() > 0)
 		{
 			Result = TraceLocalLight(
 				VirtualShadowMapId,
@@ -321,16 +313,9 @@ FVirtualShadowMapSampleResult ProjectLight(
 				SceneDepth,
 				TranslatedWorldPosition,
 				SMRTRayOffset,
-				SMRTRayCount,
-				SMRTSamplesPerRay,
-				SMRTCotMaxRayAngleFromLight,
 				Noise,
 				bBackfaceCull,
-				WorldNormal,
-				SMRTTexelDitherScale,
-				SMRTExtrapolateSlope,
-				SMRTMaxSlopeBias,
-				SMRTAdaptiveRayCount);
+				WorldNormal);
 		}
 		else
 		{
@@ -364,7 +349,8 @@ FVirtualShadowMapSampleResult ProjectLight(
 		if (!bIsHairInput && Result.ShadowFactor > 0)
 		{
 			float AccShadowFactor = 0;
-			uint HairRayCount = min(SMRTHairRayCount, uint(SMRTRayCount));
+			// TODO: Maybe do this min() in the uniform buffer parameter itself?
+			uint HairRayCount = min(VirtualShadowMap.SMRTHairRayCount, uint(GetSMRTRayCount()));
 			uint RayIndex = 0; 
 			for (; RayIndex < HairRayCount; RayIndex++ )
 			{
@@ -382,7 +368,7 @@ FVirtualShadowMapSampleResult ProjectLight(
 					Light,
 					PixelPos,
 					TranslatedWorldPosition,
-					SMRTRayLengthScale * 100,
+					VirtualShadowMap.SMRTRayLengthScale * 100,
 					RayIndex,
 					HairRayCount,
 					RayStart,
@@ -393,7 +379,7 @@ FVirtualShadowMapSampleResult ProjectLight(
 					PixelPos,
 					TranslatedWorldPosition,
 					WorldNormal,
-					SMRTCotMaxRayAngleFromLight,
+					VirtualShadowMap.SMRTCotMaxRayAngleFromLight,
 					RayIndex,
 					HairRayCount,
 					RayStart,
@@ -451,7 +437,7 @@ FVirtualShadowMapSampleResult ProjectLight(
 
 					#define USE_ADAPTATIVE 1
 					#if COMPUTESHADER && USE_ADAPTATIVE
-					if (SMRTAdaptiveRayCount > 0)
+					if (VirtualShadowMap.SMRTAdaptiveRayCount > 0)
 					{
 						// TODO: Adapt this heuristic based on SMRTAdaptiveRayCount as well?
 						if( RayIndex == 0 )
@@ -465,7 +451,7 @@ FVirtualShadowMapSampleResult ProjectLight(
 								break;
 							}
 						}
-						else if( RayIndex >= SMRTAdaptiveRayCount )
+						else if( RayIndex >= VirtualShadowMap.SMRTAdaptiveRayCount )
 						{
 							bool bHit = IterationShadowFactor == 0;
 
@@ -539,7 +525,7 @@ void OutputVisualize(
 		}		
 		else if ( VisualizeModeId == VIRTUAL_SHADOW_MAP_VISUALIZE_SMRT_RAY_COUNT )
 		{
-			Output = GreenToRedTurbo( float( Result.RayCount ) / float( SMRTRayCount ) );
+			Output = GreenToRedTurbo(float(Result.RayCount) / float(GetSMRTRayCount()));
 		}
 		else if ( VisualizeModeId == VIRTUAL_SHADOW_MAP_VISUALIZE_CACHED_PAGE && bValidPageAddress )
 		{
@@ -678,7 +664,7 @@ void VirtualShadowMapProjection(
 	const float4 SvPosition = float4( float2( PixelPos ) + 0.5, DeviceZ, 1.0f );
 	const float3 TranslatedWorldPosition = SvPositionToTranslatedWorld( SvPosition );
 
-	const float ScreenRayLengthWorld = GetScreenRayLengthMultiplierForProjectionType(ScreenRayLength * SceneDepth).y;
+	const float ScreenRayLengthWorld = GetScreenRayLengthMultiplierForProjectionType(VirtualShadowMap.ScreenRayLength * SceneDepth).y;
 	const float Noise = BlueNoiseScalar(SvPosition.xy, View.StateFrameIndex);
 
 	const FProjectionShadingInfo ShadingInfo = GetProjectionShadingInfo(PixelPos);
diff --git a/Engine/Shaders/Private/VirtualShadowMaps/VirtualShadowMapProjectionDirectional.ush b/Engine/Shaders/Private/VirtualShadowMaps/VirtualShadowMapProjectionDirectional.ush
--- a/Engine/Shaders/Private/VirtualShadowMaps/VirtualShadowMapProjectionDirectional.ush
+++ b/Engine/Shaders/Private/VirtualShadowMaps/VirtualShadowMapProjectionDirectional.ush
@@ -179,17 +179,10 @@ FVirtualShadowMapSampleResult TraceDirectional(
 	const float SceneDepth,
 	float3 TranslatedWorldPosition,
 	float RayStartOffset,
-	uint MaxRayCount,
-	uint SamplesPerRay,
-	float RayLengthScale,
 	const float Noise,
 	bool bCullBackface,
-	float3 WorldNormal,
-	float TexelDitherScale,
-	float ExtrapolateSlope,
-	uint SMRTAdaptiveRayCount)
+	float3 WorldNormal)
 {
-
 	float3 ViewPosition = mul(float4(TranslatedWorldPosition, 1.0f), View.TranslatedWorldToView).xyz;
 	float DistanceFromViewOrigin = length(ViewPosition);
 
@@ -223,7 +216,7 @@ FVirtualShadowMapSampleResult TraceDirectional(
 	// in some different way.
 	//
 	// 0.5 is arbitrary but makes it consistent with the previous dither scale cvar
-	float PerLightTexelDitherScale = TexelDitherScale * ProjectionData.TexelDitherScale;
+	float PerLightTexelDitherScale = VirtualShadowMap.SMRTTexelDitherScaleDirectional * ProjectionData.TexelDitherScale;
 	float DitherScale = 0.0f;
 	if (PerLightTexelDitherScale > 0.0f)
 	{
@@ -239,7 +232,7 @@ FVirtualShadowMapSampleResult TraceDirectional(
 	// and the maximum size a penumbra can be.
 	// Too high values will cause shadows to detach from their contact points (unless more samples are used).
 	// Too low values will greatly restrict how large penumbras can be in screen space.
-	float RayLength = RayLengthScale * DistanceFromViewOrigin;
+	float RayLength = VirtualShadowMap.SMRTRayLengthScale * DistanceFromViewOrigin;
 
 	// Clipmap origin should be near shaded samples, so shadow translated world should be regular range
 	float3 ViewToShadowTranslation = DFFastLocalSubtractDemote(ProjectionData.PreViewTranslation, PrimaryView.PreViewTranslation);
@@ -249,6 +242,7 @@ FVirtualShadowMapSampleResult TraceDirectional(
 	uint i = 0;
 	float OccluderDistanceSum = 0.0f;
 	float MaxOccluderDistance = -1.0f;
+	const uint MaxRayCount = VirtualShadowMap.SMRTRayCountDirectional;
 	for ( ; i < MaxRayCount; i++)
 	{
 		// One sample for ray, one for texel dither
@@ -264,8 +258,8 @@ FVirtualShadowMapSampleResult TraceDirectional(
 			RayStartOffset,
 			DepthSlopeUV,
 			TexelOffset,
-			ExtrapolateSlope);
-		FSMRTResult SMRTResult = SMRTRayCast(RayState, SamplesPerRay, Noise);
+			VirtualShadowMap.SMRTExtrapolateMaxSlopeDirectional);
+		FSMRTResult SMRTResult = SMRTRayCast(RayState, VirtualShadowMap.SMRTSamplesPerRayDirectional, Noise);
 
 		// Debug output (DCE'd if not used)
 		Result.ClipmapOrMipLevel = GetVirtualShadowMapProjectionData(RayState.SampledClipmapId).ClipmapLevel;
@@ -289,7 +283,7 @@ FVirtualShadowMapSampleResult TraceDirectional(
 		}
 
 		#if COMPUTESHADER
-		if (SMRTAdaptiveRayCount > 0)
+		if (VirtualShadowMap.SMRTAdaptiveRayCount > 0)
 		{
 			// TODO: Adapt this heuristic based on SMRTAdaptiveRayCount as well?
 			if( i == 0 )
@@ -303,7 +297,7 @@ FVirtualShadowMapSampleResult TraceDirectional(
 					break;
 				}
 			}
-			else if( i >= SMRTAdaptiveRayCount )
+			else if( i >= VirtualShadowMap.SMRTAdaptiveRayCount )
 			{
 				// After 2 iterations and all have hit, assume umbra
 				bool bAllLanesHit = WaveActiveAllTrue( RayMissCount == 0 );
diff --git a/Engine/Shaders/Private/VirtualShadowMaps/VirtualShadowMapProjectionSpot.ush b/Engine/Shaders/Private/VirtualShadowMaps/VirtualShadowMapProjectionSpot.ush
--- a/Engine/Shaders/Private/VirtualShadowMaps/VirtualShadowMapProjectionSpot.ush
+++ b/Engine/Shaders/Private/VirtualShadowMaps/VirtualShadowMapProjectionSpot.ush
@@ -320,16 +320,9 @@ FVirtualShadowMapSampleResult TraceLocalLight(
 	const float SceneDepth,
 	float3 TranslatedWorldPosition,
 	float RayStartOffset,
-	uint MaxRayCount,
-	uint SamplesPerRay,
-	float CotMaxRayAngleFromLight,
 	const float Noise,
 	const bool bCullBackface,
-	float3 WorldNormal,
-	float TexelDitherScale,
-	float ExtrapolateSlope,
-	float MaxSlopeBias,
-	uint SMRTAdaptiveRayCount)
+	float3 WorldNormal)
 {
 	const uint LightType = Light.SpotAngles.x > -2.0f ? LIGHT_TYPE_SPOT : LIGHT_TYPE_POINT;
 	const float3 ToLight = Light.TranslatedWorldPosition - TranslatedWorldPosition;
@@ -368,9 +361,9 @@ FVirtualShadowMapSampleResult TraceLocalLight(
 		// We use a radially symmetric mapping here instead of the true function to minimize issues crossing cubemap faces.
 		const float ShadowDepthFunctionScale = abs(ProjectionData.ShadowViewToClipMatrix._43 / length2(ToLight));
 
-		const float DitherScale = (TexelDitherScale * ProjectionData.TexelDitherScale) * ShadowResolutionScale;
-		const float MaxDitherSlope = MaxSlopeBias * DitherScale;
-		const float MaxDepthBias = MaxSlopeBias * ShadowResolutionScale * ShadowDepthFunctionScale;
+		const float DitherScale = (VirtualShadowMap.SMRTTexelDitherScaleLocal * ProjectionData.TexelDitherScale) * ShadowResolutionScale;
+		const float MaxDitherSlope = VirtualShadowMap.SMRTMaxSlopeBiasLocal * DitherScale;
+		const float MaxDepthBias = VirtualShadowMap.SMRTMaxSlopeBiasLocal * ShadowResolutionScale * ShadowDepthFunctionScale;
 
 		// Moving the ray origin means we need to attempt to account for local receiver geometry to avoid
 		// incorrect self-shadowing. To this end we move the point on a plane aligned with the shading normal.
@@ -381,6 +374,7 @@ FVirtualShadowMapSampleResult TraceLocalLight(
 		const float2 DepthSlopeTangent = -NormalPlaneTangent.xy / NormalPlaneTangent.z;
 
 		// If source angle = 0, we don't need multiple samples on the same ray
+		uint SamplesPerRay = VirtualShadowMap.SMRTSamplesPerRayLocal;
 		if (ConeSin == 0.0f)
 		{
 			SamplesPerRay = 0;
@@ -390,6 +384,7 @@ FVirtualShadowMapSampleResult TraceLocalLight(
 		uint RayMissCount = 0;
 		float OccluderDistanceSum = 0.0f;
 		float MaxOccluderDistance = -1.0f;
+		const uint MaxRayCount = VirtualShadowMap.SMRTRayCountLocal;
 		for( ; i < MaxRayCount; i++ )
 		{
 			float4 RandSample = VirtualShadowMapGetRandomSample(PixelPos, View.StateFrameIndex, i, MaxRayCount);
@@ -418,7 +413,7 @@ FVirtualShadowMapSampleResult TraceLocalLight(
 			// where the projection becomes extreme.
 
 			// Clip ray to max angle to limit penumbra sizes
-			float ClampedEnd = DistToLight * 0.75 * saturate( 1.5 / ( CosTheta + CotMaxRayAngleFromLight * SinTheta ) );
+			float ClampedEnd = DistToLight * 0.75 * saturate( 1.5 / ( CosTheta + VirtualShadowMap.SMRTCotMaxRayAngleFromLight * SinTheta ) );
 			float Start = min(RayStartOffset, ClampedEnd - 1e-6f);
 
 			float3 RayStart = RayStartWithDither + Dir * Start;
@@ -442,7 +437,7 @@ FVirtualShadowMapSampleResult TraceLocalLight(
 					DepthBias,
 					SamplesPerRay,
 					StepOffset,
-					ExtrapolateSlope,
+					VirtualShadowMap.SMRTExtrapolateMaxSlopeLocal,
 					RayState);
 			}
 			else
@@ -454,7 +449,7 @@ FVirtualShadowMapSampleResult TraceLocalLight(
 					DepthBias,
 					SamplesPerRay,
 					StepOffset,
-					ExtrapolateSlope,
+					VirtualShadowMap.SMRTExtrapolateMaxSlopeLocal,
 					RayState);
 			}
 
@@ -482,7 +477,7 @@ FVirtualShadowMapSampleResult TraceLocalLight(
 			Result.VirtualTexelAddress = RayState.VirtualTexelAddress;
 			Result.PhysicalTexelAddress = RayState.PhysicalTexelAddress;
 
-			if (SMRTAdaptiveRayCount > 0)
+			if (VirtualShadowMap.SMRTAdaptiveRayCount > 0)
 			{
 				// TODO: Adapt this heuristic based on SMRTAdaptiveRayCount as well?
 				if( i == 0 )
@@ -496,7 +491,7 @@ FVirtualShadowMapSampleResult TraceLocalLight(
 						break;
 					}
 				}
-				else if( i >= SMRTAdaptiveRayCount )
+				else if( i >= VirtualShadowMap.SMRTAdaptiveRayCount )
 				{
 					// After N iterations and all have hit, assume umbra
 					bool bAllLanesHit = WaveActiveAllTrue( RayMissCount == 0 );
diff --git a/Engine/Source/Runtime/Renderer/Private/LightGridInjection.cpp b/Engine/Source/Runtime/Renderer/Private/LightGridInjection.cpp
--- a/Engine/Source/Runtime/Renderer/Private/LightGridInjection.cpp
+++ b/Engine/Source/Runtime/Renderer/Private/LightGridInjection.cpp
@@ -676,7 +676,6 @@ FComputeLightGridOutput FSceneRenderer::ComputeLightGrid(FRDGBuilder& GraphBuild
 							ForwardLightData->DirectionalLightSoftSourceRadius = LightParameters.SoftSourceRadius;
 							ForwardLightData->DirectionalLightShadowMapChannelMask = LightTypeAndShadowMapChannelMaskPacked;
 							ForwardLightData->DirectionalLightVSM = INDEX_NONE;
-							ForwardLightData->DirectionalLightSMRTSettings = GetVirtualShadowMapSMRTSettings(true);
 							ForwardLightData->LightFunctionAtlasLightIndex = LightParameters.LightFunctionAtlasLightIndex;
 							ForwardLightData->bAffectsTranslucentLighting = LightParameters.bAffectsTranslucentLighting;
 
diff --git a/Engine/Source/Runtime/Renderer/Private/SceneRendering.h b/Engine/Source/Runtime/Renderer/Private/SceneRendering.h
--- a/Engine/Source/Runtime/Renderer/Private/SceneRendering.h
+++ b/Engine/Source/Runtime/Renderer/Private/SceneRendering.h
@@ -669,7 +669,6 @@ BEGIN_GLOBAL_SHADER_PARAMETER_STRUCT_WITH_CONSTRUCTOR(FForwardLightData, )
 	SHADER_PARAMETER(uint32, DirectLightingShowFlag)
 	SHADER_PARAMETER(uint32, LightFunctionAtlasLightIndex)
 	SHADER_PARAMETER(uint32, bAffectsTranslucentLighting)
-	SHADER_PARAMETER_STRUCT(FVirtualShadowMapSMRTSettings, DirectionalLightSMRTSettings)
 	SHADER_PARAMETER_RDG_TEXTURE(Texture2D, DirectionalLightShadowmapAtlas)
 	SHADER_PARAMETER_SAMPLER(SamplerState, ShadowmapSampler)
 	SHADER_PARAMETER_TEXTURE(Texture2D, DirectionalLightStaticShadowmap)
diff --git a/Engine/Source/Runtime/Renderer/Private/VirtualShadowMaps/VirtualShadowMapArray.cpp b/Engine/Source/Runtime/Renderer/Private/VirtualShadowMaps/VirtualShadowMapArray.cpp
--- a/Engine/Source/Runtime/Renderer/Private/VirtualShadowMaps/VirtualShadowMapArray.cpp
+++ b/Engine/Source/Runtime/Renderer/Private/VirtualShadowMaps/VirtualShadowMapArray.cpp
@@ -367,6 +367,115 @@ static TAutoConsoleVariable<int32> CVarVirtualShadowMapPageMarkingPixelStrideY(
 	TEXT("Same as PageMarkingPixelStrideX, but on the vertical axis of the screen."),
 	ECVF_RenderThreadSafe);
 
+
+static TAutoConsoleVariable<float> CVarScreenRayLength(
+	TEXT("r.Shadow.Virtual.ScreenRayLength"),
+	0.015f,
+	TEXT("Length of the screen space shadow trace away from receiver surface (smart shadow bias) before the VSM / SMRT lookup."),
+	ECVF_Scalability | ECVF_RenderThreadSafe
+);
+
+static TAutoConsoleVariable<int32> CVarSMRTAdaptiveRayCount(
+	TEXT("r.Shadow.Virtual.SMRT.AdaptiveRayCount"),
+	1,
+	TEXT("Shoot fewer rays in fully shadowed and unshadowed regions. Currently only supported with OnePassProjection. "),
+	ECVF_Scalability | ECVF_RenderThreadSafe
+);
+
+static TAutoConsoleVariable<int32> CVarSMRTRayCountLocal(
+	TEXT("r.Shadow.Virtual.SMRT.RayCountLocal"),
+	7,
+	TEXT("Ray count for shadow map tracing of local lights. 0 = disabled."),
+	ECVF_Scalability | ECVF_RenderThreadSafe
+);
+
+static TAutoConsoleVariable<int32> CVarSMRTSamplesPerRayLocal(
+	TEXT("r.Shadow.Virtual.SMRT.SamplesPerRayLocal"),
+	8,
+	TEXT("Shadow map samples per ray for local lights"),
+	ECVF_Scalability | ECVF_RenderThreadSafe
+);
+
+static TAutoConsoleVariable<float> CVarSMRTMaxRayAngleFromLight(
+	TEXT("r.Shadow.Virtual.SMRT.MaxRayAngleFromLight"),
+	0.03f,
+	TEXT("Max angle (in radians) a ray is allowed to span from the light's perspective for local lights.")
+	TEXT("Smaller angles limit the screen space size of shadow penumbra. ")
+	TEXT("Larger angles lead to more noise. "),
+	ECVF_RenderThreadSafe
+);
+
+TAutoConsoleVariable<float> CVarSMRTExtrapolateMaxSlopeLocal(
+	TEXT("r.Shadow.Virtual.SMRT.ExtrapolateMaxSlopeLocal"),
+	0.05f,
+	TEXT("Maximum depth slope when extrapolating behind occluders for local lights.\n")
+	TEXT("Higher values allow softer penumbra edges but can introduce light leaks behind second occluders.\n")
+	TEXT("Setting to 0 will disable slope extrapolation slightly improving projection performance, at the cost of reduced penumbra quality."),
+	ECVF_Scalability | ECVF_RenderThreadSafe
+);
+
+static TAutoConsoleVariable<float> CVarSMRTTexelDitherScaleLocal(
+	TEXT("r.Shadow.Virtual.SMRT.TexelDitherScaleLocal"),
+	2.0f,
+	TEXT("Applies a dither to the shadow map ray casts for local lights to help hide aliasing due to insufficient shadow resolution.\n")
+	TEXT("Setting this too high can cause shadows light leaks near occluders."),
+	ECVF_Scalability | ECVF_RenderThreadSafe
+);
+
+TAutoConsoleVariable<float> CVarSMRTMaxSlopeBiasLocal(
+	TEXT("r.Shadow.Virtual.SMRT.MaxSlopeBiasLocal"),
+	50.0f,
+	TEXT("Maximum depth slope. Low values produce artifacts if shadow resolution is insufficient. High values can worsen light leaks near occluders and sparkly pixels in shadowed areas."),
+	ECVF_Scalability | ECVF_RenderThreadSafe
+);
+
+static TAutoConsoleVariable<int32> CVarSMRTRayCountDirectional(
+	TEXT("r.Shadow.Virtual.SMRT.RayCountDirectional"),
+	7,
+	TEXT("Ray count for shadow map tracing of directional lights. 0 = disabled."),
+	ECVF_Scalability | ECVF_RenderThreadSafe
+);
+
+static TAutoConsoleVariable<int32> CVarSMRTSamplesPerRayDirectional(
+	TEXT("r.Shadow.Virtual.SMRT.SamplesPerRayDirectional"),
+	8,
+	TEXT("Shadow map samples per ray for directional lights"),
+	ECVF_Scalability | ECVF_RenderThreadSafe
+);
+
+TAutoConsoleVariable<float> CVarSMRTExtrapolateMaxSlopeDirectional(
+	TEXT("r.Shadow.Virtual.SMRT.ExtrapolateMaxSlopeDirectional"),
+	5.0f,
+	TEXT("Maximum depth slope when extrapolating behind occluders for directional lights.\n")
+	TEXT("Higher values allow softer penumbra edges but can introduce light leaks behind second occluders.\n")
+	TEXT("Setting to 0 will disable slope extrapolation slightly improving projection performance, at the cost of reduced penumbra quality."),
+	ECVF_Scalability | ECVF_RenderThreadSafe
+);
+
+static TAutoConsoleVariable<float> CVarSMRTTexelDitherScaleDirectional(
+	TEXT("r.Shadow.Virtual.SMRT.TexelDitherScaleDirectional"),
+	2.0f,
+	TEXT("Applies a dither to the shadow map ray casts for directional lights to help hide aliasing due to insufficient shadow resolution.\n")
+	TEXT("Setting this too high can cause shadows light leaks near occluders."),
+	ECVF_Scalability | ECVF_RenderThreadSafe
+);
+
+static TAutoConsoleVariable<float> CVarSMRTRayLengthScaleDirectional(
+	TEXT("r.Shadow.Virtual.SMRT.RayLengthScaleDirectional"),
+	1.5f,
+	TEXT("Length of ray to shoot for directional lights, scaled by distance to camera.")
+	TEXT("Shorter rays limit the screen space size of shadow penumbra. ")
+	TEXT("Longer rays require more samples to avoid shadows disconnecting from contact points. "),
+	ECVF_RenderThreadSafe
+);
+
+static TAutoConsoleVariable<int32> CVarSMRTRayCountHair(
+	TEXT("r.Shadow.Virtual.SMRT.SamplesPerRayHair"),
+	1,
+	TEXT("Shadow map samples per ray for hair"),
+	ECVF_Scalability | ECVF_RenderThreadSafe
+);
+
 namespace Nanite
 {
 	extern bool IsStatFilterActive(const FString& FilterName);
@@ -459,6 +568,25 @@ void FVirtualShadowMapArray::Initialize(
 	UniformParameters.bClipmapGreedyLevelSelection = CVarClipmapGreedyLevelSelection.GetValueOnRenderThread();
 
 	UniformParameters.SceneFrameNumber = Scene.GetFrameNumberRenderThread();
+	
+	// Global SMRT settings so they can be shared between different passes that call into them
+	UniformParameters.ScreenRayLength = CVarScreenRayLength.GetValueOnRenderThread();
+	UniformParameters.SMRTAdaptiveRayCount = CVarSMRTAdaptiveRayCount.GetValueOnRenderThread();
+
+	UniformParameters.SMRTRayCountLocal = CVarSMRTRayCountLocal.GetValueOnRenderThread();
+	UniformParameters.SMRTSamplesPerRayLocal = CVarSMRTSamplesPerRayLocal.GetValueOnRenderThread();
+	UniformParameters.SMRTExtrapolateMaxSlopeLocal = CVarSMRTExtrapolateMaxSlopeLocal.GetValueOnRenderThread();
+	UniformParameters.SMRTTexelDitherScaleLocal = CVarSMRTTexelDitherScaleLocal.GetValueOnRenderThread();
+	UniformParameters.SMRTMaxSlopeBiasLocal = CVarSMRTMaxSlopeBiasLocal.GetValueOnRenderThread();
+	UniformParameters.SMRTCotMaxRayAngleFromLight = 1.0f / FMath::Tan(CVarSMRTMaxRayAngleFromLight.GetValueOnRenderThread());
+	
+	UniformParameters.SMRTRayCountDirectional = CVarSMRTRayCountDirectional.GetValueOnRenderThread();
+	UniformParameters.SMRTSamplesPerRayDirectional = CVarSMRTSamplesPerRayDirectional.GetValueOnRenderThread();
+	UniformParameters.SMRTExtrapolateMaxSlopeDirectional = CVarSMRTExtrapolateMaxSlopeDirectional.GetValueOnRenderThread();
+	UniformParameters.SMRTTexelDitherScaleDirectional = CVarSMRTTexelDitherScaleDirectional.GetValueOnRenderThread();
+	UniformParameters.SMRTRayLengthScale = CVarSMRTRayLengthScaleDirectional.GetValueOnRenderThread();
+
+	UniformParameters.SMRTHairRayCount = CVarSMRTRayCountHair.GetValueOnRenderThread();
 
 	// Reference dummy data in the UB initially
 	const uint32 DummyPageTableElement = 0xFFFFFFFF;
diff --git a/Engine/Source/Runtime/Renderer/Private/VirtualShadowMaps/VirtualShadowMapArray.h b/Engine/Source/Runtime/Renderer/Private/VirtualShadowMaps/VirtualShadowMapArray.h
--- a/Engine/Source/Runtime/Renderer/Private/VirtualShadowMaps/VirtualShadowMapArray.h
+++ b/Engine/Source/Runtime/Renderer/Private/VirtualShadowMaps/VirtualShadowMapArray.h
@@ -155,6 +155,24 @@ BEGIN_GLOBAL_SHADER_PARAMETER_STRUCT(FVirtualShadowMapUniformParameters, )
 	SHADER_PARAMETER(uint32, bClipmapGreedyLevelSelection)
 	SHADER_PARAMETER(float, GlobalResolutionLodBias)
 
+	// SMRT parameters that are sometimes used globally
+	SHADER_PARAMETER(float, ScreenRayLength)
+	SHADER_PARAMETER(uint32, SMRTAdaptiveRayCount)
+	SHADER_PARAMETER(int32, SMRTRayCountLocal)
+	SHADER_PARAMETER(int32, SMRTSamplesPerRayLocal)
+	SHADER_PARAMETER(float, SMRTExtrapolateMaxSlopeLocal)
+	SHADER_PARAMETER(float, SMRTTexelDitherScaleLocal)
+	SHADER_PARAMETER(float, SMRTMaxSlopeBiasLocal)
+	SHADER_PARAMETER(float, SMRTCotMaxRayAngleFromLight)
+
+	SHADER_PARAMETER(int32, SMRTRayCountDirectional)
+	SHADER_PARAMETER(int32, SMRTSamplesPerRayDirectional)
+	SHADER_PARAMETER(float, SMRTExtrapolateMaxSlopeDirectional)
+	SHADER_PARAMETER(float, SMRTTexelDitherScaleDirectional)
+	SHADER_PARAMETER(float, SMRTRayLengthScale)
+		
+	SHADER_PARAMETER(uint32, SMRTHairRayCount)
+
 	SHADER_PARAMETER_RDG_BUFFER_SRV(ByteAddressBuffer, ProjectionData)
 	SHADER_PARAMETER_RDG_BUFFER_SRV(StructuredBuffer<uint>, PageTable)
 	SHADER_PARAMETER_RDG_BUFFER_SRV(StructuredBuffer<uint>, PageFlags)
diff --git a/Engine/Source/Runtime/Renderer/Private/VirtualShadowMaps/VirtualShadowMapProjection.cpp b/Engine/Source/Runtime/Renderer/Private/VirtualShadowMaps/VirtualShadowMapProjection.cpp
--- a/Engine/Source/Runtime/Renderer/Private/VirtualShadowMaps/VirtualShadowMapProjection.cpp
+++ b/Engine/Source/Runtime/Renderer/Private/VirtualShadowMaps/VirtualShadowMapProjection.cpp
@@ -30,12 +30,6 @@
 
 #define MAX_TEST_PERMUTATION 0
 
-static TAutoConsoleVariable<float> CVarScreenRayLength(
-	TEXT( "r.Shadow.Virtual.ScreenRayLength" ),
-	0.015f,
-	TEXT( "Length of the screen space shadow trace away from receiver surface (smart shadow bias) before the VSM / SMRT lookup." ),
-	ECVF_Scalability | ECVF_RenderThreadSafe
-);
 
 static TAutoConsoleVariable<float> CVarNormalBias(
 	TEXT( "r.Shadow.Virtual.NormalBias" ),
@@ -45,100 +39,6 @@ static TAutoConsoleVariable<float> CVarNormalBias(
 	ECVF_Scalability | ECVF_RenderThreadSafe
 );
 
-static TAutoConsoleVariable<int32> CVarSMRTRayCountLocal(
-	TEXT( "r.Shadow.Virtual.SMRT.RayCountLocal" ),
-	7,
-	TEXT( "Ray count for shadow map tracing of local lights. 0 = disabled." ),
-	ECVF_Scalability | ECVF_RenderThreadSafe
-);
-
-static TAutoConsoleVariable<int32> CVarSMRTSamplesPerRayLocal(
-	TEXT( "r.Shadow.Virtual.SMRT.SamplesPerRayLocal" ),
-	8,
-	TEXT( "Shadow map samples per ray for local lights" ),
-	ECVF_Scalability | ECVF_RenderThreadSafe
-);
-
-static TAutoConsoleVariable<float> CVarSMRTMaxRayAngleFromLight(
-	TEXT( "r.Shadow.Virtual.SMRT.MaxRayAngleFromLight" ),
-	0.03f,
-	TEXT( "Max angle (in radians) a ray is allowed to span from the light's perspective for local lights." )
-	TEXT( "Smaller angles limit the screen space size of shadow penumbra. " )
-	TEXT( "Larger angles lead to more noise. " ),
-	ECVF_RenderThreadSafe
-);
-
-static TAutoConsoleVariable<int32> CVarSMRTRayCountDirectional(
-	TEXT( "r.Shadow.Virtual.SMRT.RayCountDirectional" ),
-	7,
-	TEXT( "Ray count for shadow map tracing of directional lights. 0 = disabled." ),
-	ECVF_Scalability | ECVF_RenderThreadSafe
-);
-
-static TAutoConsoleVariable<int32> CVarSMRTSamplesPerRayDirectional(
-	TEXT( "r.Shadow.Virtual.SMRT.SamplesPerRayDirectional" ),
-	8,
-	TEXT( "Shadow map samples per ray for directional lights" ),
-	ECVF_Scalability | ECVF_RenderThreadSafe
-);
-
-static TAutoConsoleVariable<float> CVarSMRTRayLengthScaleDirectional(
-	TEXT( "r.Shadow.Virtual.SMRT.RayLengthScaleDirectional" ),
-	1.5f,
-	TEXT( "Length of ray to shoot for directional lights, scaled by distance to camera." )
-	TEXT( "Shorter rays limit the screen space size of shadow penumbra. " )
-	TEXT( "Longer rays require more samples to avoid shadows disconnecting from contact points. " ),
-	ECVF_RenderThreadSafe
-);
-
-static TAutoConsoleVariable<int32> CVarSMRTAdaptiveRayCount(
-	TEXT( "r.Shadow.Virtual.SMRT.AdaptiveRayCount" ),
-	1,
-	TEXT( "Shoot fewer rays in fully shadowed and unshadowed regions. Currently only supported with OnePassProjection. " ),
-	ECVF_Scalability | ECVF_RenderThreadSafe
-);
-
-static TAutoConsoleVariable<float> CVarSMRTTexelDitherScaleLocal(
-	TEXT( "r.Shadow.Virtual.SMRT.TexelDitherScaleLocal" ),
-	2.0f,
-	TEXT( "Applies a dither to the shadow map ray casts for local lights to help hide aliasing due to insufficient shadow resolution.\n" )
-	TEXT( "Setting this too high can cause shadows light leaks near occluders." ),
-	ECVF_Scalability | ECVF_RenderThreadSafe
-);
-
-static TAutoConsoleVariable<float> CVarSMRTTexelDitherScaleDirectional(
-	TEXT( "r.Shadow.Virtual.SMRT.TexelDitherScaleDirectional" ),
-	2.0f,
-	TEXT( "Applies a dither to the shadow map ray casts for directional lights to help hide aliasing due to insufficient shadow resolution.\n" )
-	TEXT( "Setting this too high can cause shadows light leaks near occluders." ),
-	ECVF_Scalability | ECVF_RenderThreadSafe
-);
-
-TAutoConsoleVariable<float> CVarSMRTExtrapolateMaxSlopeLocal(
-	TEXT("r.Shadow.Virtual.SMRT.ExtrapolateMaxSlopeLocal"),
-	0.05f,
-	TEXT("Maximum depth slope when extrapolating behind occluders for local lights.\n")
-	TEXT("Higher values allow softer penumbra edges but can introduce light leaks behind second occluders.\n")
-	TEXT("Setting to 0 will disable slope extrapolation slightly improving projection performance, at the cost of reduced penumbra quality."),
-	ECVF_Scalability | ECVF_RenderThreadSafe
-);
-
-TAutoConsoleVariable<float> CVarSMRTExtrapolateMaxSlopeDirectional(
-	TEXT("r.Shadow.Virtual.SMRT.ExtrapolateMaxSlopeDirectional"),
-	5.0f,
-	TEXT("Maximum depth slope when extrapolating behind occluders for directional lights.\n")
-	TEXT("Higher values allow softer penumbra edges but can introduce light leaks behind second occluders.\n")
-	TEXT("Setting to 0 will disable slope extrapolation slightly improving projection performance, at the cost of reduced penumbra quality."),
-	ECVF_Scalability | ECVF_RenderThreadSafe
-);
-
-TAutoConsoleVariable<float> CVarSMRTMaxSlopeBiasLocal(
-	TEXT("r.Shadow.Virtual.SMRT.MaxSlopeBiasLocal"),
-	50.0f,
-	TEXT("Maximum depth slope. Low values produce artifacts if shadow resolution is insufficient. High values can worsen light leaks near occluders and sparkly pixels in shadowed areas."),
-	ECVF_Scalability | ECVF_RenderThreadSafe
-);
-
 static TAutoConsoleVariable<int32> CVarForcePerLightShadowMaskClear(
 	TEXT( "r.Shadow.Virtual.ForcePerLightShadowMaskClear" ),
 	0,
@@ -163,12 +63,6 @@ static TAutoConsoleVariable<int32> CVarSubsurfaceShadowMinSourceAngle(
 	ECVF_RenderThreadSafe
 );
 
-static TAutoConsoleVariable<int32> CVarSMRTRayCountHair(
-	TEXT( "r.Shadow.Virtual.SMRT.SamplesPerRayHair" ),
-	1,
-	TEXT( "Shadow map samples per ray for hair" ),
-	ECVF_Scalability | ECVF_RenderThreadSafe
-);
 
 #if MAX_TEST_PERMUTATION > 0
 static TAutoConsoleVariable<int32> CVarTestPermutation(
@@ -190,36 +84,6 @@ bool IsVSMTranslucentHighQualityEnabled()
 	return CVarVSMTranslucentQuality.GetValueOnRenderThread() > 0;
 }
 
-FVirtualShadowMapSMRTSettings GetVirtualShadowMapSMRTSettings(bool bDirectionalLight)
-{
-	FVirtualShadowMapSMRTSettings Out;
-	Out.ScreenRayLength = CVarScreenRayLength.GetValueOnRenderThread();
-	Out.SMRTAdaptiveRayCount = CVarSMRTAdaptiveRayCount.GetValueOnRenderThread();
-	if (bDirectionalLight)
-	{
-		Out.SMRTRayCount = CVarSMRTRayCountDirectional.GetValueOnRenderThread();
-		Out.SMRTSamplesPerRay = CVarSMRTSamplesPerRayDirectional.GetValueOnRenderThread();
-		Out.SMRTRayLengthScale = CVarSMRTRayLengthScaleDirectional.GetValueOnRenderThread();
-		Out.SMRTCotMaxRayAngleFromLight = 0.0f;	// unused in this path
-		Out.SMRTTexelDitherScale = CVarSMRTTexelDitherScaleDirectional.GetValueOnRenderThread();
-		Out.SMRTExtrapolateSlope = CVarSMRTExtrapolateMaxSlopeDirectional.GetValueOnRenderThread();
-		Out.SMRTMaxSlopeBias = 0.0f; // unused in this path
-		Out.SMRTHairRayCount = CVarSMRTRayCountHair.GetValueOnRenderThread();
-	}
-	else
-	{
-		Out.SMRTRayCount = CVarSMRTRayCountLocal.GetValueOnRenderThread();
-		Out.SMRTSamplesPerRay = CVarSMRTSamplesPerRayLocal.GetValueOnRenderThread();
-		Out.SMRTRayLengthScale = 0.0f;		// unused in this path
-		Out.SMRTCotMaxRayAngleFromLight = 1.0f / FMath::Tan(CVarSMRTMaxRayAngleFromLight.GetValueOnRenderThread());
-		Out.SMRTTexelDitherScale = CVarSMRTTexelDitherScaleLocal.GetValueOnRenderThread();
-		Out.SMRTExtrapolateSlope = CVarSMRTExtrapolateMaxSlopeLocal.GetValueOnRenderThread();
-		Out.SMRTMaxSlopeBias = CVarSMRTMaxSlopeBiasLocal.GetValueOnRenderThread();
-		Out.SMRTHairRayCount = CVarSMRTRayCountHair.GetValueOnRenderThread();
-	}
-	return Out;
-}
-
 const TCHAR* ToString(EVirtualShadowMapProjectionInputType In)
 {
 	switch (In)
@@ -268,7 +132,6 @@ class FVirtualShadowMapProjectionCS : public FGlobalShader
 		SHADER_PARAMETER_RDG_UNIFORM_BUFFER(FSubstrateGlobalUniformParameters, Substrate)
 		SHADER_PARAMETER_STRUCT_REF(FViewUniformShaderParameters, View)
 		SHADER_PARAMETER_STRUCT_REF(FBlueNoise, BlueNoise)
-		SHADER_PARAMETER_STRUCT_INCLUDE(FVirtualShadowMapSMRTSettings, SMRTSettings)
 		SHADER_PARAMETER(FIntVector4, ProjectionRect)
 		SHADER_PARAMETER(float, NormalBias)
 		SHADER_PARAMETER(float, SubsurfaceMinSourceRadius)
@@ -394,7 +257,7 @@ static void RenderVirtualShadowMapProjectionCommon(
 	{
 		// One pass projection
 		PassParameters->ForwardLightData = View.ForwardLightingResources.ForwardLightUniformBuffer;
-		PassParameters->OutShadowMaskBits = GraphBuilder.CreateUAV( OutputTexture );	
+		PassParameters->OutShadowMaskBits = GraphBuilder.CreateUAV( OutputTexture );
 	}
 	else
 	{
@@ -406,9 +269,7 @@ static void RenderVirtualShadowMapProjectionCommon(
 		PassParameters->LightUniformVirtualShadowMapId = VirtualShadowMapId;
 		PassParameters->OutShadowFactor = GraphBuilder.CreateUAV( OutputTexture );
 	}
- 
-	PassParameters->SMRTSettings = GetVirtualShadowMapSMRTSettings(bDirectionalLight);
-	
+ 	
 	bool bDebugOutput = false;
 #if !UE_BUILD_SHIPPING
 	if ( !VirtualShadowMapArray.DebugVisualizationOutput.IsEmpty() && InputType == EVirtualShadowMapProjectionInputType::GBuffer && VirtualShadowMapArray.VisualizeLight[ViewIndex].IsValid())
@@ -426,16 +287,18 @@ static void RenderVirtualShadowMapProjectionCommon(
 
 	// If the requested samples per ray matches one of our static permutations, pick that one
 	// Otherwise use the dynamic samples per ray permutation (-1).
-	int StaticSamplesPerRay = PassParameters->SMRTSettings.SMRTSamplesPerRay == 0 ? PassParameters->SMRTSettings.SMRTSamplesPerRay : -1;
+	int SamplesPerRay = (bDirectionalLight ? VirtualShadowMapArray.UniformParameters.SMRTSamplesPerRayDirectional : VirtualShadowMapArray.UniformParameters.SMRTSamplesPerRayLocal);
+	int StaticSamplesPerRay = SamplesPerRay	== 0 ? SamplesPerRay : -1;
+	float ExtrapolateMaxSlope = (bDirectionalLight ? VirtualShadowMapArray.UniformParameters.SMRTExtrapolateMaxSlopeDirectional : VirtualShadowMapArray.UniformParameters.SMRTExtrapolateMaxSlopeLocal);
 
 	FVirtualShadowMapProjectionCS::FPermutationDomain PermutationVector;
 	PermutationVector.Set< FVirtualShadowMapProjectionCS::FDirectionalLightDim >( bDirectionalLight );
 	PermutationVector.Set< FVirtualShadowMapProjectionCS::FOnePassProjectionDim >( bOnePassProjection );
 	PermutationVector.Set< FVirtualShadowMapProjectionCS::FHairStrandsDim >( bHasHairStrandsData );
 	PermutationVector.Set< FVirtualShadowMapProjectionCS::FVisualizeOutputDim >( bDebugOutput );
-	PermutationVector.Set< FVirtualShadowMapProjectionCS::FExtrapolateSlopeDim >( PassParameters->SMRTSettings.SMRTExtrapolateSlope > 0.0f );
+	PermutationVector.Set< FVirtualShadowMapProjectionCS::FExtrapolateSlopeDim >(ExtrapolateMaxSlope > 0.0f);
 	PermutationVector.Set< FVirtualShadowMapProjectionCS::FUseTileList >(bUseTileList);
-	PermutationVector.Set< FVirtualShadowMapProjectionCS::FSMRTStaticSampleCount >( StaticSamplesPerRay );
+	PermutationVector.Set< FVirtualShadowMapProjectionCS::FSMRTStaticSampleCount >(StaticSamplesPerRay);
 #if MAX_TEST_PERMUTATION > 0
 	{
 		int32 TestPermutation = FMath::Clamp(CVarTestPermutation.GetValueOnRenderThread(), 0, MAX_TEST_PERMUTATION);
@@ -450,10 +313,7 @@ static void RenderVirtualShadowMapProjectionCommon(
 	if (bUseTileList)
 	{
 		GraphBuilder.AddPass(
-			RDG_EVENT_NAME("VirtualShadowMapProjection(RayCount:%u(%s),SamplesPerRay:%u,Input:%s%s,TileList)",
-				PassParameters->SMRTSettings.SMRTRayCount,
-				PassParameters->SMRTSettings.SMRTAdaptiveRayCount > 0 ? TEXT("Adaptive") : TEXT("Static"),
-				PassParameters->SMRTSettings.SMRTSamplesPerRay,
+			RDG_EVENT_NAME("VirtualShadowMapProjection(Input:%s%s,TileList)",
 				ToString(InputType),
 				bDebugOutput ? TEXT(",Debug") : TEXT("")),
 			PassParameters,
@@ -468,10 +328,7 @@ static void RenderVirtualShadowMapProjectionCommon(
 		const FIntPoint GroupCount = FIntPoint::DivideAndRoundUp(ProjectionRect.Size(), VSMProjectionWorkTileSize);
 		FComputeShaderUtils::AddPass(
 			GraphBuilder,
-			RDG_EVENT_NAME("VirtualShadowMapProjection(RayCount:%u(%s),SamplesPerRay:%u,Input:%s%s)",
-				PassParameters->SMRTSettings.SMRTRayCount,
-				PassParameters->SMRTSettings.SMRTAdaptiveRayCount > 0 ? TEXT("Adaptive") : TEXT("Static"),
-				PassParameters->SMRTSettings.SMRTSamplesPerRay,
+			RDG_EVENT_NAME("VirtualShadowMapProjection(Input:%s%s)",
 				ToString(InputType),
 				bDebugOutput ? TEXT(",Debug") : TEXT("")),
 			ComputeShader,
diff --git a/Engine/Source/Runtime/Renderer/Private/VirtualShadowMaps/VirtualShadowMapProjection.h b/Engine/Source/Runtime/Renderer/Private/VirtualShadowMaps/VirtualShadowMapProjection.h
--- a/Engine/Source/Runtime/Renderer/Private/VirtualShadowMaps/VirtualShadowMapProjection.h
+++ b/Engine/Source/Runtime/Renderer/Private/VirtualShadowMaps/VirtualShadowMapProjection.h
@@ -18,21 +18,6 @@ class FVirtualShadowMapArray;
 class FVirtualShadowMapClipmap;
 struct FMinimalSceneTextures;
 
-BEGIN_SHADER_PARAMETER_STRUCT(FVirtualShadowMapSMRTSettings,)
-	SHADER_PARAMETER(float, ScreenRayLength)
-	SHADER_PARAMETER(int32, SMRTRayCount)
-	SHADER_PARAMETER(int32, SMRTSamplesPerRay)
-	SHADER_PARAMETER(float, SMRTRayLengthScale)
-	SHADER_PARAMETER(float, SMRTCotMaxRayAngleFromLight)
-	SHADER_PARAMETER(float, SMRTTexelDitherScale)
-	SHADER_PARAMETER(float, SMRTExtrapolateSlope)
-	SHADER_PARAMETER(float, SMRTMaxSlopeBias)
-	SHADER_PARAMETER(uint32, SMRTAdaptiveRayCount)
-	SHADER_PARAMETER(uint32, SMRTHairRayCount)
-END_SHADER_PARAMETER_STRUCT()
-
-FVirtualShadowMapSMRTSettings GetVirtualShadowMapSMRTSettings(bool bDirectionalLight);
-
 struct FTiledVSMProjection
 {
 	FRDGBufferRef DrawIndirectParametersBuffer;
