Fixes stability regressions on geometric edges under parallax movements

1) Reverts ParallaxRejectionMaskThreshold was adding quiet a bit of instability on geometric edges
2) The measurement of the ReprojectionEdge in the decimate pass after the first velocity dilate is more costly, but less distracting to the eye in high density geometry areas (like forest)

Good bye old friend.

diff --git a/Engine/Shaders/Private/TemporalSuperResolution/TSRClosestOccluder.ush b/Engine/Shaders/Private/TemporalSuperResolution/TSRClosestOccluder.ush
--- a/Engine/Shaders/Private/TemporalSuperResolution/TSRClosestOccluder.ush
+++ b/Engine/Shaders/Private/TemporalSuperResolution/TSRClosestOccluder.ush
@@ -19,7 +19,6 @@
 
 float2 PrevOutputBufferUVMin;
 float2 PrevOutputBufferUVMax;
-float ParallaxRejectionMaskThreshold;
 
 
 //------------------------------------------------------- HOLE FILLING VELOCITY ENCODING
@@ -297,7 +296,6 @@ void ProcessPrevAtomicTexturesSamples(
 
 	HoleFillingPixelVelocity = CartesianHoleFillingVelocity(HoleFillingPixelAngle, HoleFillingPixelLength);
 	bCanHoleFill = bIsValidHoleFillingPixelVelocity && bIsEncodablePixelLength;
-	//bIsParallaxDisocclusion = !bIsOffScreen && ParallaxRejectionMask < tsr_half(ParallaxRejectionMaskThreshold);
 
 	// Dim down the parallax rejection mask if the hole fill velocity matches
 	FLATTEN
@@ -329,5 +327,5 @@ void ProcessPrevAtomicTexturesSamples(
 		ParallaxRejectionMask = max(ParallaxRejectionMask, tsr_half(MinParallaxRejectionMask));
 	}
 
-	bIsParallaxDisocclusion = !bIsOffScreen && ParallaxRejectionMask < tsr_half(ParallaxRejectionMaskThreshold);
+	bIsParallaxDisocclusion = !bIsOffScreen && ParallaxRejectionMask < tsr_half(0.5);
 } // ProcessPrevAtomicTexturesSamples()
diff --git a/Engine/Shaders/Private/TemporalSuperResolution/TSRDecimateHistory.usf b/Engine/Shaders/Private/TemporalSuperResolution/TSRDecimateHistory.usf
--- a/Engine/Shaders/Private/TemporalSuperResolution/TSRDecimateHistory.usf
+++ b/Engine/Shaders/Private/TemporalSuperResolution/TSRDecimateHistory.usf
@@ -25,12 +25,13 @@
 
 float4x4 RotationalClipToPrevClip;
 
-Texture2D<uint> ReprojectionVectorTexture;
+Texture2D<uint> DilatedReprojectionVectorTexture;
 #if CONFIG_RESURRECTION_REPROJECTION
 	Texture2D<float2> ClosestDepthTexture;
 #else
 	Texture2D<float> ClosestDepthTexture;
 #endif
+Texture2D<uint> DilateMaskTexture;
 Texture2D<uint> DepthErrorTexture;
 Texture2DArray<uint> PrevAtomicTextureArray;
 
@@ -245,9 +246,13 @@ void MainCS(
 	float PrevDeviceZ;
 	float DeviceZError;
 	bool bHasPixelAnimation;
+	bool bHasReprojectionOffset;
+	tsr_half ReprojectionEdge;
 	ISOLATE
 	{
-		//uint EncodedReprojectionVectorNeighborhood[CONFIG_VELOCITY_PRE_FETCH_SAMPLES];
+		uint EncodedReprojectionVectorNeighborhood[CONFIG_BUTTERFLY_SAMPLES];
+		
+		tsr_half PremilinaryReprojectionEdge;
 
 		// Issue texture fetechs
 		ISOLATE
@@ -255,7 +260,8 @@ void MainCS(
 			tsr_ushort2 InputPixelPos = (
 				Map8x8Tile2x2Lane(GroupThreadIndex) +
 				(tsr_ushort2(InputInfo_ViewportMin) + tsr_ushort2(GroupId) * tsr_ushort2(TILE_SIZE, TILE_SIZE)));
-			
+
+			uint EncodedDilateMask = DilateMaskTexture[InputPixelPos];
 			uint EncodedDeviceZError = DepthErrorTexture[InputPixelPos];
 			#if CONFIG_RESURRECTION_REPROJECTION
 				float2 DeviceZAndPrevDeviceZ = ClosestDepthTexture[InputPixelPos];
@@ -265,13 +271,14 @@ void MainCS(
 				DeviceZ = 0.0;
 			#endif
 			
+			bHasReprojectionOffset = (EncodedDilateMask & 0x80u) != 0u;
+			PremilinaryReprojectionEdge = tsr_half(tsr_ushort(EncodedDilateMask & 0x7Fu)) * rcp(tsr_half(127.0));
 			PrevDeviceZ = abs(DeviceZAndPrevDeviceZ.r);
 			bHasPixelAnimation = DeviceZAndPrevDeviceZ.r < 0.0;
 			DeviceZError = DecodeDeviceZError(PrevDeviceZ, EncodedDeviceZError);
 
-			//FetchReprojectionNeighborhood(ReprojectionVectorTexture, InputPixelPos, /* out */ EncodedReprojectionVectorNeighborhood);
-			//DilatedEncodedReprojectionVector = EncodedReprojectionVectorNeighborhood[0];
-			DilatedEncodedReprojectionVector = ReprojectionFieldOutput[tsr_ushort3(InputPixelPos, 0)];
+			FetchReprojectionNeighborhood(DilatedReprojectionVectorTexture, InputPixelPos, /* out */ EncodedReprojectionVectorNeighborhood);
+			DilatedEncodedReprojectionVector = AccessNeighborhoodCenter(EncodedReprojectionVectorNeighborhood);
 		}
 
 		// LDS spill to save VGPR on the history texture fetches.
@@ -284,15 +291,15 @@ void MainCS(
 		#endif
 		
 		// Detect whether is on velocity edge to clamp the high frequencies
-		#if 0
 		ISOLATE
 		{
-			float2 PixelVelocityNeighborhood[CONFIG_VELOCITY_PRE_FETCH_SAMPLES];
+			float2 PixelVelocityNeighborhood[CONFIG_BUTTERFLY_SAMPLES];
 			ComputePixelVelocityNeighborhood(EncodedReprojectionVectorNeighborhood, /* out */ PixelVelocityNeighborhood);
 
-			ComputeReprojectionEdge(PixelVelocityNeighborhood, /* out */ VelocityEdge);
+			ReprojectionEdge = ComputeReprojectionEdge(PixelVelocityNeighborhood);
+			ReprojectionEdge = select(PremilinaryReprojectionEdge > tsr_half(0.9), PremilinaryReprojectionEdge, ReprojectionEdge);
+			ReprojectionEdge = select(bHasReprojectionOffset, ReprojectionEdge, tsr_half(1.0));
 		}
-		#endif
 	}
 	
 	// Sample the previous frame information
@@ -469,19 +476,18 @@ void MainCS(
 		}
 		#endif
 		
-		DecimateMaskOutput[OutputPixelPos] = tsr_half2(tsr_half(DecimateBitMask) * rcp(tsr_half(255.0)), 0.0);
+		DecimateMaskOutput[OutputPixelPos] = tsr_half2(tsr_half(DecimateBitMask) * rcp(tsr_half(255.0)), ReprojectionEdge);
 
 		// Selectively overwrite the reprojection field
 		{
-			bool bOverwriteReprojectionField = bReprojectionHollFill;
+			bool bWriteReprojectionVector   = bReprojectionHollFill || true;
+			bool bWriteReprojectionJacobian = bReprojectionHollFill;
 
 			// Disable the reprojection's jacobian since this is to extrude disoccluded area we don't know how they look anyway.
 			const uint EncodedReprojectionJacobian = EncodeReprojectionJacobian(tsr_half2x2(tsr_half(0.0).xxxx));
 
-			tsr_short2 ReprojectionFieldOutputPixelPos = OutputPixelPos;
-			ReprojectionFieldOutputPixelPos.x = select(bOverwriteReprojectionField, ReprojectionFieldOutputPixelPos.x, -tsr_short(1));
-			ReprojectionFieldOutput[tsr_short3(ReprojectionFieldOutputPixelPos, 0)] = DilatedEncodedReprojectionVector;
-			ReprojectionFieldOutput[tsr_short3(ReprojectionFieldOutputPixelPos, 2)] = EncodedReprojectionJacobian;
+			ReprojectionFieldOutput[tsr_short3(select(bWriteReprojectionVector,   OutputPixelPos.x, -tsr_short(1)), OutputPixelPos.y, kReprojectionVectorOutputIndex)]   = DilatedEncodedReprojectionVector;
+			ReprojectionFieldOutput[tsr_short3(select(bWriteReprojectionJacobian, OutputPixelPos.x, -tsr_short(1)), OutputPixelPos.y, kReprojectionJacobianOutputIndex)] = EncodedReprojectionJacobian;
 		}
 
 		#if DEBUG_OUTPUT
diff --git a/Engine/Shaders/Private/TemporalSuperResolution/TSRDilateVelocity.usf b/Engine/Shaders/Private/TemporalSuperResolution/TSRDilateVelocity.usf
--- a/Engine/Shaders/Private/TemporalSuperResolution/TSRDilateVelocity.usf
+++ b/Engine/Shaders/Private/TemporalSuperResolution/TSRDilateVelocity.usf
@@ -27,6 +27,7 @@ float InvFlickeringMaxParralaxVelocity;
 float ReprojectionFieldAntiAliasVelocityThreshold;
 uint bReprojectionField;
 uint bOutputIsMovingTexture;
+uint ReprojectionVectorOutputIndex;
 
 Texture2D<float> SceneDepthTexture;
 Texture2D<float4> SceneVelocityTexture;
@@ -164,6 +165,7 @@ void MainCS(
 	uint EncodedReprojectionJacobian;
 	uint EncodedDeviceZError;
 	bool bHasPixelAnimation;
+	bool bHasReprojectionOffset;
 	tsr_half ReprojectionEdge;
 	BRANCH
 	if (bReprojectionField)
@@ -213,7 +215,8 @@ void MainCS(
 				DevizeZNeighborhood,
 				/* out */ ClosestDeviceZ,
 				/* out */ ReprojectionOffset);
-		}
+			bHasReprojectionOffset = any(ReprojectionOffset != tsr_short(0));
+;		}
 
 		float2 ScreenVelocity;
 		ISOLATE
@@ -232,7 +235,7 @@ void MainCS(
 				/* out */ DilatedScreenVelocity,
 				/* out */ bDildatedHasPixelAnimation);
 		}
-	
+		
 		ReprojectionEdge = ComputeReprojectionEdge(ScreenVelocity, DilatedScreenVelocity.xy);
 		
 		// Compute the reprojection boundary using a depth spatial anti-aliaser
@@ -264,10 +267,13 @@ void MainCS(
 		}
 		
 		tsr_short2 ReprojectionOffset;
-		FindClosestDepthOffset(
-			DevizeZNeighborhood,
-			/* out */ ClosestDeviceZ,
-			/* out */ ReprojectionOffset);
+		{
+			FindClosestDepthOffset(
+				DevizeZNeighborhood,
+				/* out */ ClosestDeviceZ,
+				/* out */ ReprojectionOffset);
+			bHasReprojectionOffset = any(ReprojectionOffset != tsr_short(0));
+		}
 	
 		// Compute the parallax depth error
 		float ParallaxDepthError = ComputeDepthError(DevizeZNeighborhood);
@@ -280,7 +286,7 @@ void MainCS(
 		FetchAndComputeScreenVelocity(
 			SceneVelocityTexture,
 			InputPixelPos,
-			ClosestDeviceZ,
+			DeviceZ,
 			/* ReprojectionOffset = */ tsr_short2(0, 0),
 			/* out */ ScreenVelocity,
 			/* out */ bHasPixelAnimation);
@@ -294,7 +300,7 @@ void MainCS(
 			ReprojectionOffset,
 			/* out */ DilatedScreenVelocity,
 			/* out */ bDilatedHasPixelAnimation);
-			
+		
 		ReprojectionEdge = ComputeReprojectionEdge(ScreenVelocity.xy, DilatedScreenVelocity.xy);
 		
 		// Encode the device z error.
@@ -357,7 +363,7 @@ void MainCS(
 
 		ClosestDepthOutput[OutputPixelPos] = float2(select(bHasPixelAnimation, -PrevClosestDeviceZ, PrevClosestDeviceZ), ClosestDeviceZ);
 
-		R8Output[uint3(OutputPixelPos, 0)] = uint(round(ReprojectionEdge * 255.0));
+		R8Output[uint3(OutputPixelPos, 0)] = uint(round(ReprojectionEdge * 127.0)) | select(bHasReprojectionOffset, 0x80u, 0x00u);
 		R8Output[uint3(OutputPixelPos, 1)] = EncodedDeviceZError;
 		BRANCH
 		if (bOutputIsMovingTexture)
@@ -366,12 +372,12 @@ void MainCS(
 		}
 
 		// Output the reprojection field
-		ReprojectionFieldOutput[uint3(OutputPixelPos, 0)] = EncodeReprojectionVector(DilatedScreenVelocity.xy);
+		ReprojectionFieldOutput[uint3(OutputPixelPos, ReprojectionVectorOutputIndex)] = EncodeReprojectionVector(DilatedScreenVelocity.xy);
 		BRANCH
 		if (bReprojectionField)
 		{
-			ReprojectionFieldOutput[uint3(OutputPixelPos, 1)] = EncodedReprojectionBoundary;
-			ReprojectionFieldOutput[uint3(OutputPixelPos, 2)] = EncodedReprojectionJacobian;
+			ReprojectionFieldOutput[uint3(OutputPixelPos, kReprojectionJacobianOutputIndex)] = EncodedReprojectionJacobian;
+			ReprojectionFieldOutput[uint3(OutputPixelPos, kReprojectionBoundaryOutputIndex)] = EncodedReprojectionBoundary;
 		}
 
 		// Output motion blur velocity flatten
diff --git a/Engine/Shaders/Private/TemporalSuperResolution/TSRRejectShading.usf b/Engine/Shaders/Private/TemporalSuperResolution/TSRRejectShading.usf
--- a/Engine/Shaders/Private/TemporalSuperResolution/TSRRejectShading.usf
+++ b/Engine/Shaders/Private/TemporalSuperResolution/TSRRejectShading.usf
@@ -55,7 +55,6 @@ Texture2D<tsr_half4> ReprojectedHistoryMoireTexture;
 	Texture2D<tsr_half4> ResurrectedHistoryGuideTexture;
 	Texture2D<tsr_half>  ResurrectedHistoryGuideMetadataTexture;
 #endif
-Texture2D<uint> DilateMaskTexture;
 Texture2D<tsr_half2> DecimateMaskTexture;
 Texture2D<tsr_ushort> IsMovingMaskTexture;
 Texture2D<float2> ClosestDepthTexture;
@@ -129,12 +128,11 @@ void FetchDecimateMask(
 	{
 		tsr_short2 ClampedFetchPixelPos = PixelPos.GetElement(ElementIndex);
 		
-		uint DilateMask = DilateMaskTexture[ClampedFetchPixelPos];
 		tsr_half2 DecimateMask = DecimateMaskTexture[ClampedFetchPixelPos];
 
 		tsr_ushort BitMask = tsr_ushort(DecimateMask.r * tsr_half(255.0));
 
-		VelocityEdge.SetElement(ElementIndex, saturate(tsr_half(tsr_ushort(DilateMask)) * tsr_half(rcp(255.0))));
+		VelocityEdge.SetElement(ElementIndex, DecimateMask.g);
 		bIsDisoccluded.SetElement(ElementIndex,     (BitMask & tsr_ushort(0x3)) != tsr_ushort(0));
 		bHasPixelAnimation.SetElement(ElementIndex, (BitMask & tsr_ushort(0x4)) != tsr_ushort(0));
 		bIsResurrectionOffScreen.SetElement(ElementIndex, (BitMask & tsr_ushort(0x10)) != tsr_ushort(0));
@@ -349,8 +347,8 @@ void OverwriteReprojectionFieldWithResurrection(tsr_tensor_bool bResurrectHistor
 			tsr_short2 OutputPixelPos = ComputeElementOutputPixelPos(ElementIndex);
 			OutputPixelPos.x = select(bResurrectHistory.GetElement(ElementIndex), OutputPixelPos.x, tsr_short(-1));
 
-			ReprojectionFieldOutput[tsr_short3(OutputPixelPos, 0)] = EncodedReprojectionVector[ElementIndex];
-			ReprojectionFieldOutput[tsr_short3(OutputPixelPos, 2)] = EncodedReprojectionJacobian[ElementIndex];
+			ReprojectionFieldOutput[tsr_short3(OutputPixelPos, kReprojectionVectorOutputIndex)]   = EncodedReprojectionVector[ElementIndex];
+			ReprojectionFieldOutput[tsr_short3(OutputPixelPos, kReprojectionJacobianOutputIndex)] = EncodedReprojectionJacobian[ElementIndex];
 		}
 	}
 } // OverwriteReprojectionFieldWithResurrection()
diff --git a/Engine/Shaders/Private/TemporalSuperResolution/TSRReprojectionField.ush b/Engine/Shaders/Private/TemporalSuperResolution/TSRReprojectionField.ush
--- a/Engine/Shaders/Private/TemporalSuperResolution/TSRReprojectionField.ush
+++ b/Engine/Shaders/Private/TemporalSuperResolution/TSRReprojectionField.ush
@@ -5,6 +5,13 @@
 #include "TSRCommon.ush"
 
 
+//------------------------------------------------------- CONSTANTS
+
+static const uint kReprojectionVectorOutputIndex = 0;
+static const uint kReprojectionJacobianOutputIndex = 1;
+static const uint kReprojectionBoundaryOutputIndex = 2;
+
+
 //------------------------------------------------------- REPROJECTION's DILATING OFFSET & BOUNDARY
 
 static const tsr_half2 kFullDilateBoundary = tsr_half2(0.0, 1.0);
diff --git a/Engine/Source/Runtime/Renderer/Private/PostProcess/TemporalSuperResolution.cpp b/Engine/Source/Runtime/Renderer/Private/PostProcess/TemporalSuperResolution.cpp
--- a/Engine/Source/Runtime/Renderer/Private/PostProcess/TemporalSuperResolution.cpp
+++ b/Engine/Source/Runtime/Renderer/Private/PostProcess/TemporalSuperResolution.cpp
@@ -666,6 +666,7 @@ class FTSRDilateVelocityCS : public FTSRShader
 		SHADER_PARAMETER(float, ReprojectionFieldAntiAliasVelocityThreshold)
 		SHADER_PARAMETER(int32, bReprojectionField)
 		SHADER_PARAMETER(int32, bOutputIsMovingTexture)
+		SHADER_PARAMETER(int32, ReprojectionVectorOutputIndex)
 
 		SHADER_PARAMETER_RDG_TEXTURE(Texture2D, SceneDepthTexture)
 		SHADER_PARAMETER_RDG_TEXTURE(Texture2D, SceneVelocityTexture)
@@ -715,9 +716,9 @@ class FTSRDecimateHistoryCS : public FTSRShader
 		SHADER_PARAMETER_STRUCT_INCLUDE(FTSRCommonParameters, CommonParameters)
 		SHADER_PARAMETER(FMatrix44f, RotationalClipToPrevClip)
 
-		SHADER_PARAMETER_RDG_TEXTURE(Texture2D, InputSceneColorTexture)
-		SHADER_PARAMETER_RDG_TEXTURE_SRV(Texture2D, ReprojectionVectorTexture)
+		SHADER_PARAMETER_RDG_TEXTURE_SRV(Texture2D, DilatedReprojectionVectorTexture)
 		SHADER_PARAMETER_RDG_TEXTURE(Texture2D, ClosestDepthTexture)
+		SHADER_PARAMETER_RDG_TEXTURE_SRV(Texture2D, DilateMaskTexture)
 		SHADER_PARAMETER_RDG_TEXTURE_SRV(Texture2D, DepthErrorTexture)
 		SHADER_PARAMETER_RDG_TEXTURE(Texture2DArray, PrevAtomicTextureArray)
 
@@ -731,7 +732,6 @@ class FTSRDecimateHistoryCS : public FTSRShader
 		SHADER_PARAMETER(FVector2f, ResurrectionGuideUVViewportBilinearMin)
 		SHADER_PARAMETER(FVector2f, ResurrectionGuideUVViewportBilinearMax)
 		SHADER_PARAMETER(FVector3f, HistoryGuideQuantizationError)
-		SHADER_PARAMETER(float, ParallaxRejectionMaskThreshold)
 		SHADER_PARAMETER(float, ResurrectionFrameIndex)
 		SHADER_PARAMETER(float, PrevFrameIndex)
 		SHADER_PARAMETER(FMatrix44f, ClipToResurrectionClip)
@@ -811,7 +811,6 @@ class FTSRRejectShadingCS : public FTSRConvolutionNetworkShader
 		SHADER_PARAMETER_RDG_TEXTURE_SRV(Texture2D, ReprojectedHistoryMoireTexture)
 		SHADER_PARAMETER_RDG_TEXTURE_SRV(Texture2D, ResurrectedHistoryGuideTexture)
 		SHADER_PARAMETER_RDG_TEXTURE_SRV(Texture2D, ResurrectedHistoryGuideMetadataTexture)
-		SHADER_PARAMETER_RDG_TEXTURE_SRV(Texture2D, DilateMaskTexture)
 		SHADER_PARAMETER_RDG_TEXTURE(Texture2D, DecimateMaskTexture)
 		SHADER_PARAMETER_RDG_TEXTURE_SRV(Texture2D, IsMovingMaskTexture)
 		SHADER_PARAMETER_RDG_TEXTURE(Texture2D, ClosestDepthTexture)
@@ -1881,7 +1880,8 @@ FDefaultTemporalUpscaler::FOutputs AddTemporalSuperResolutionPasses(
 
 	// Dilate the velocity texture & scatter reprojection into previous frame
 	FRDGTextureRef ReprojectionFieldTexture;
-	FRDGTextureSRVRef ReprojectionVectorTexture;
+	FRDGTextureSRVRef DilatedReprojectionVectorTexture;
+	FRDGTextureSRVRef ReprojectionVectorTexture   = nullptr;
 	FRDGTextureSRVRef ReprojectionBoundaryTexture = nullptr;
 	FRDGTextureSRVRef ReprojectionJacobianTexture = nullptr;
 	FRDGTextureRef ClosestDepthTexture;
@@ -1929,22 +1929,34 @@ FDefaultTemporalUpscaler::FOutputs AddTemporalSuperResolutionPasses(
 			}
 		}
 
+		if (bReprojectionField)
 		{
 			FRDGTextureDesc Desc = FRDGTextureDesc::Create2DArray(
 				InputExtent,
 				PF_R32_UINT,
 				FClearValueBinding::None,
 				/* InFlags = */ TexCreate_ShaderResource | TexCreate_UAV,
-				/* ArraySize = */ bReprojectionField ? 3 : 1);
+				/* ArraySize = */ 4);
+
+			ReprojectionFieldTexture = GraphBuilder.CreateTexture(Desc, bReprojectionField ? TEXT("TSR.ReprojectionField") : TEXT("TSR.Reprojection.DilatedVector"));
 
-			ReprojectionFieldTexture = GraphBuilder.CreateTexture(Desc, TEXT("TSR.Reprojection.Vector"));
 			ReprojectionVectorTexture = GraphBuilder.CreateSRV(FRDGTextureSRVDesc::CreateForSlice(ReprojectionFieldTexture, 0));
+			ReprojectionJacobianTexture = GraphBuilder.CreateSRV(FRDGTextureSRVDesc::CreateForSlice(ReprojectionFieldTexture, 1));
+			ReprojectionBoundaryTexture = GraphBuilder.CreateSRV(FRDGTextureSRVDesc::CreateForSlice(ReprojectionFieldTexture, 2));
+			DilatedReprojectionVectorTexture = GraphBuilder.CreateSRV(FRDGTextureSRVDesc::CreateForSlice(ReprojectionFieldTexture, 3));
+		}
+		else
+		{
+			FRDGTextureDesc Desc = FRDGTextureDesc::Create2DArray(
+				InputExtent,
+				PF_R32_UINT,
+				FClearValueBinding::None,
+				/* InFlags = */ TexCreate_ShaderResource | TexCreate_UAV,
+				/* ArraySize = */ 1);
 
-			if (bReprojectionField)
-			{
-				ReprojectionBoundaryTexture = GraphBuilder.CreateSRV(FRDGTextureSRVDesc::CreateForSlice(ReprojectionFieldTexture, 1));
-				ReprojectionJacobianTexture = GraphBuilder.CreateSRV(FRDGTextureSRVDesc::CreateForSlice(ReprojectionFieldTexture, 2));
-			}
+			ReprojectionFieldTexture = GraphBuilder.CreateTexture(Desc, TEXT("TSR.Reprojection.DilatedVector"));
+
+			DilatedReprojectionVectorTexture = GraphBuilder.CreateSRV(FRDGTextureSRVDesc::CreateForSlice(ReprojectionFieldTexture, 0));
 		}
 
 		int32 TileSize = 8;
@@ -1965,6 +1977,7 @@ FDefaultTemporalUpscaler::FOutputs AddTemporalSuperResolutionPasses(
 		
 		PassParameters->SceneDepthTexture = PassInputs.SceneDepth.Texture;
 		PassParameters->SceneVelocityTexture = PassInputs.SceneVelocity.Texture;
+		PassParameters->ReprojectionVectorOutputIndex = DilatedReprojectionVectorTexture->Desc.FirstArraySlice;
 
 		PassParameters->ClosestDepthOutput = GraphBuilder.CreateUAV(ClosestDepthTexture);
 		PassParameters->PrevAtomicOutput = GraphBuilder.CreateUAV(PrevAtomicTextureArray);
@@ -2064,9 +2077,9 @@ FDefaultTemporalUpscaler::FOutputs AddTemporalSuperResolutionPasses(
 		PassParameters->CommonParameters = CommonParameters;
 		PassParameters->RotationalClipToPrevClip = RotationalClipToPrevClip;
 
-		PassParameters->InputSceneColorTexture = PassInputs.SceneColor.Texture;
-		PassParameters->ReprojectionVectorTexture = ReprojectionVectorTexture;
+		PassParameters->DilatedReprojectionVectorTexture = DilatedReprojectionVectorTexture;
 		PassParameters->ClosestDepthTexture = ClosestDepthTexture;
+		PassParameters->DilateMaskTexture = DilateMaskTexture;
 		PassParameters->DepthErrorTexture = DepthErrorTexture;
 		PassParameters->PrevAtomicTextureArray = PrevAtomicTextureArray;
 
@@ -2089,7 +2102,6 @@ FDefaultTemporalUpscaler::FOutputs AddTemporalSuperResolutionPasses(
 			PassParameters->ResurrectionGuideUVViewportBilinearMin = GetScreenPassTextureViewportParameters(ResurrectionGuideViewport).UVViewportBilinearMin;
 			PassParameters->ResurrectionGuideUVViewportBilinearMax = GetScreenPassTextureViewportParameters(ResurrectionGuideViewport).UVViewportBilinearMax;
 			PassParameters->HistoryGuideQuantizationError = ComputePixelFormatQuantizationError(ReprojectedHistoryGuideTexture->Desc.Format);
-			PassParameters->ParallaxRejectionMaskThreshold = 1.0f - 0.25f * OutputToInputResolutionFraction;
 		}
 
 		PassParameters->ResurrectionFrameIndex = ResurrectionFrameSliceIndex;
@@ -2101,7 +2113,26 @@ FDefaultTemporalUpscaler::FOutputs AddTemporalSuperResolutionPasses(
 		{
 			PassParameters->ReprojectedHistoryMoireOutput = GraphBuilder.CreateUAV(ReprojectedHistoryMoireTexture);
 		}
-		PassParameters->ReprojectionFieldOutput = GraphBuilder.CreateUAV(ReprojectionFieldTexture);
+		if (bReprojectionField)
+		{
+			FRDGTextureUAVDesc ReprojectionFieldUAVDesc(ReprojectionFieldTexture);
+			ReprojectionFieldUAVDesc.NumArraySlices = 2;
+			PassParameters->ReprojectionFieldOutput = GraphBuilder.CreateUAV(ReprojectionFieldUAVDesc);
+		}
+		else
+		{
+			// Create a new reprojection vector texture
+			FRDGTextureDesc Desc = FRDGTextureDesc::Create2DArray(
+				InputExtent,
+				PF_R32_UINT,
+				FClearValueBinding::None,
+				/* InFlags = */ TexCreate_ShaderResource | TexCreate_UAV,
+				/* ArraySize = */ 1);
+
+			ReprojectionFieldTexture = GraphBuilder.CreateTexture(Desc, TEXT("TSR.Reprojection.HollFilledVector"));
+			ReprojectionVectorTexture = GraphBuilder.CreateSRV(FRDGTextureSRVDesc::CreateForSlice(ReprojectionFieldTexture, 0));
+			PassParameters->ReprojectionFieldOutput = GraphBuilder.CreateUAV(ReprojectionFieldTexture);
+		}
 		PassParameters->DecimateMaskOutput = GraphBuilder.CreateUAV(DecimateMaskTexture);
 		PassParameters->DebugOutput = CreateDebugUAV(InputExtent, TEXT("Debug.TSR.DecimateHistory"));
 
@@ -2264,7 +2295,6 @@ FDefaultTemporalUpscaler::FOutputs AddTemporalSuperResolutionPasses(
 			PassParameters->ResurrectedHistoryGuideTexture = GraphBuilder.CreateSRV(FRDGTextureSRVDesc(BlackDummy));
 			PassParameters->ResurrectedHistoryGuideMetadataTexture = GraphBuilder.CreateSRV(FRDGTextureSRVDesc(BlackDummy));
 		}
-		PassParameters->DilateMaskTexture = DilateMaskTexture;
 		PassParameters->DecimateMaskTexture = DecimateMaskTexture;
 		PassParameters->IsMovingMaskTexture = IsMovingMaskTexture ? IsMovingMaskTexture : GraphBuilder.CreateSRV(FRDGTextureSRVDesc(BlackUintDummy));
 		PassParameters->ClosestDepthTexture = ClosestDepthTexture;
@@ -2316,7 +2346,16 @@ FDefaultTemporalUpscaler::FOutputs AddTemporalSuperResolutionPasses(
 			PassParameters->HistoryRejectionOutput = GraphBuilder.CreateUAV(HistoryRejectionTexture);
 
 			// Amends how the history should be reprojected
-			PassParameters->ReprojectionFieldOutput = GraphBuilder.CreateUAV(ReprojectionFieldTexture);
+			if (bReprojectionField)
+			{
+				FRDGTextureUAVDesc ReprojectionFieldUAVDesc(ReprojectionFieldTexture);
+				ReprojectionFieldUAVDesc.NumArraySlices = 2;
+				PassParameters->ReprojectionFieldOutput = GraphBuilder.CreateUAV(ReprojectionFieldUAVDesc);
+			}
+			else
+			{
+				PassParameters->ReprojectionFieldOutput = GraphBuilder.CreateUAV(ReprojectionFieldTexture);
+			}
 
 			// Output the composed translucency and opaque scene color to speed up HistoryUpdate
 			PassParameters->InputSceneColorOutput = bComputeInputSceneColorTexture
