Move hardcoded LUT into engine resources to reduce binary size.

Remove about 300Kb of binary size.

#virtualized

diff --git a/Engine/Config/BaseEngine.ini b/Engine/Config/BaseEngine.ini
--- a/Engine/Config/BaseEngine.ini
+++ b/Engine/Config/BaseEngine.ini
@@ -160,6 +160,14 @@ EmissiveMeshMaterialName=/Engine/EngineMaterials/EmissiveMeshMaterial.EmissiveMe
 PreIntegratedSkinBRDFTextureName=/Engine/EngineMaterials/PreintegratedSkinBRDF.PreintegratedSkinBRDF
 BlueNoiseScalarTextureName=/Engine/EngineMaterials/STBlueNoise_scalar_128x128x64.STBlueNoise_scalar_128x128x64
 BlueNoiseVec2TextureName=/Engine/EngineMaterials/STBlueNoise_vec2_128x128x64.STBlueNoise_vec2_128x128x64
+GGXLTCAmpTextureName=/Engine/EngineMaterials/LTC/GGX_LTCAmp.GGX_LTCAmp
+GGXLTCMatTextureName=/Engine/EngineMaterials/LTC/GGX_LTCMat.GGX_LTCMat
+SheenLTCTextureName=/Engine/EngineMaterials/LTC/Sheen_LTC.Sheen_LTC
+GGXReflectionEnergyTextureName=/Engine/EngineMaterials/EnergyConservation/GGX_ReflectionEnergy.GGX_ReflectionEnergy
+GGXTransmissionEnergyTextureName=/Engine/EngineMaterials/EnergyConservation/GGX_TransmissionEnergyArray.GGX_TransmissionEnergyArray
+SheenEnergyTextureName=/Engine/EngineMaterials/EnergyConservation/Sheen_Energy.Sheen_Energy
+SheenLegacyEnergyTextureName=/Engine/EngineMaterials/EnergyConservation/Sheen_LegacyEnergy.Sheen_LegacyEnergy
+DiffuseEnergyTextureName=/Engine/EngineMaterials/EnergyConservation/Diffuse_Energy.Diffuse_Energy
 GlintTextureName=/Engine/EngineMaterials/Substrate/GLints/AGlintTex2dArray.AGlintTex2dArray
 GlintTexture2Name=/Engine/EngineMaterials/Substrate/GLints2/AGlintTex2dArray2.AGlintTex2dArray2
 SimpleVolumeTextureName=/Engine/EngineMaterials/Substrate/Volume/SubstrateSimpleVolumeLUT3D.SubstrateSimpleVolumeLUT3D
diff --git a/Engine/Source/Runtime/Engine/Classes/Engine/Engine.h b/Engine/Source/Runtime/Engine/Classes/Engine/Engine.h
--- a/Engine/Source/Runtime/Engine/Classes/Engine/Engine.h
+++ b/Engine/Source/Runtime/Engine/Classes/Engine/Engine.h
@@ -1349,11 +1349,62 @@ public:
 	/** Path of the tiled blue-noise texture */
 	UPROPERTY(globalconfig)
 	FSoftObjectPath BlueNoiseVec2TextureName;
-	
+
+	/** Texture used for GGX LTC integration (Amplitude Texture) */
+	UPROPERTY()
+	TObjectPtr<class UTexture2D> GGXLTCAmpTexture;
+	/** Path of the texture used for GGX LTC integration (Amplitude Texture) */
+	UPROPERTY(globalconfig)
+	FSoftObjectPath GGXLTCAmpTextureName;
+
+	/** Texture used for GGX LTC integration (Matrix Texture) */
+	UPROPERTY()
+	TObjectPtr<class UTexture2D> GGXLTCMatTexture;
+	/** Path of the texture used for GGX LTC integration (Matrix Texture) */
+	UPROPERTY(globalconfig)
+	FSoftObjectPath GGXLTCMatTextureName;
+
+	/** Texture used for Sheen LTC integration (Matrix Texture) */
+	UPROPERTY()
+	TObjectPtr<class UTexture2D> SheenLTCTexture;	
+	/** Path of the texture used for Sheen LTC integration (Matrix Texture) */
+	UPROPERTY(globalconfig)
+	FSoftObjectPath SheenLTCTextureName;
+
+	/** Texture used for specular reflection energy conservation */
+	UPROPERTY()
+	TObjectPtr<class UTexture2D> GGXReflectionEnergyTexture;
+	/** Path of the texture used for specular reflection energy conservation */
+	UPROPERTY(globalconfig)
+	FSoftObjectPath GGXReflectionEnergyTextureName;
+
+	/** Texture used for specular transmission energy conservation */
+	UPROPERTY()
+	TObjectPtr<class UTexture2D> GGXTransmissionEnergyTexture;
+	/** Path of the texture used for specular transmission energy conservation */
+	UPROPERTY(globalconfig)
+	FSoftObjectPath GGXTransmissionEnergyTextureName;
+		
+	/** Texture used for sheen energy conservation */
+	UPROPERTY()
+	TObjectPtr<class UTexture2D> SheenEnergyTexture;
+	/** Path of the texture used for sheen energy conservation */
+	UPROPERTY(globalconfig)
+	FSoftObjectPath SheenLegacyEnergyTextureName;		
+	/** Path of the texture used for sheen energy conservation */
+	UPROPERTY(globalconfig)
+	FSoftObjectPath SheenEnergyTextureName;
+		
+	/** Texture used for rough diffuse energy conservation */
+	UPROPERTY()
+	TObjectPtr<class UTexture2D> DiffuseEnergyTexture;
+	/** Path of the texture used for rough diffuse energy conservation */
+	UPROPERTY(globalconfig)
+	FSoftObjectPath DiffuseEnergyTextureName;
+
 	/** Stable glint BSDF texture */
 	UPROPERTY()
 	TObjectPtr<class UTexture2DArray> GlintTexture;
-	
 	/** Stable glint BSDF texture with more variety to cover slope space and avoid circular artifact */
 	UPROPERTY()
 	TObjectPtr<class UTexture2DArray> GlintTexture2;
@@ -1931,6 +1982,12 @@ public:
 	/** Conditionally load this texture for a platform. Always loaded in Editor */
 	ENGINE_API void ConditionallyLoadPreIntegratedSkinBRDFTexture();
 
+	/** Delay loading the LTC texture until it is needed by the renderer. */
+	ENGINE_API void LoadLTCTextures();
+
+	/** Delay loading the energy shading texture until it is needed by the renderer. */
+	ENGINE_API void LoadEnergyTextures();
+
 	/** Delay loading the glint texture until it is needed by the renderer.
 	* This texture is not going to be streamed to be available right away.
 	*/
diff --git a/Engine/Source/Runtime/Engine/Private/UnrealEngine.cpp b/Engine/Source/Runtime/Engine/Private/UnrealEngine.cpp
--- a/Engine/Source/Runtime/Engine/Private/UnrealEngine.cpp
+++ b/Engine/Source/Runtime/Engine/Private/UnrealEngine.cpp
@@ -3235,6 +3235,8 @@ void UEngine::InitializeObjectReferences()
 	LoadEngineTexture(WeightMapArrayPlaceholderTexture,  *WeightMapArrayPlaceholderTextureName.ToString());
 	LoadEngineTexture(LightMapDensityTexture, *LightMapDensityTextureName.ToString());
 	ConditionallyLoadPreIntegratedSkinBRDFTexture();
+	LoadLTCTextures();
+	LoadEnergyTextures();
 
 #if WITH_EDITOR
 	// Avoid breaking some engine textures that might be cached very early (i.e. BlueNoise)
@@ -3438,6 +3440,68 @@ void UEngine::ConditionallyLoadPreIntegratedSkinBRDFTexture()
 	}
 }
 
+void UEngine::LoadLTCTextures()
+{
+	static const auto CVar = IConsoleManager::Get().FindTConsoleVariableDataInt(TEXT("r.Substrate"));
+	const bool bSubstrateEnabled = CVar && CVar->GetValueOnAnyThread() > 0;
+
+	const bool bGGX = true;
+	const bool bSheen = bSubstrateEnabled;
+	if (bGGX && GGXLTCAmpTexture == nullptr && GGXLTCAmpTextureName.IsValid())
+	{
+		LoadEngineTexture(GGXLTCAmpTexture, *GGXLTCAmpTextureName.ToString());
+	}
+	if (bGGX && GGXLTCMatTexture == nullptr && GGXLTCMatTextureName.IsValid())
+	{
+		LoadEngineTexture(GGXLTCMatTexture, *GGXLTCMatTextureName.ToString());
+	}
+	if (bSheen && SheenLTCTexture == nullptr && SheenLTCTextureName.IsValid())
+	{
+		LoadEngineTexture(SheenLTCTexture, *SheenLTCTextureName.ToString());
+	}
+}
+
+void UEngine::LoadEnergyTextures()
+{
+	#if RHI_RAYTRACING
+	static const auto CVarPT = IConsoleManager::Get().FindTConsoleVariableDataInt(TEXT("r.PathTracing"));
+	const bool bPathTracingEnabled = CVarPT && CVarPT->GetValueOnAnyThread() > 0;
+	#else
+	const bool bPathTracingEnabled = false;
+	#endif
+
+	static const auto CVarSubstrate = IConsoleManager::Get().FindTConsoleVariableDataInt(TEXT("r.Substrate"));
+	const bool bSubstrateEnabled = CVarSubstrate && CVarSubstrate->GetValueOnAnyThread() > 0;
+
+	if ((bPathTracingEnabled || bSubstrateEnabled) && GGXReflectionEnergyTexture == nullptr && GGXReflectionEnergyTextureName.IsValid())
+	{
+		LoadEngineTexture(GGXReflectionEnergyTexture, *GGXReflectionEnergyTextureName.ToString());
+	}
+
+	if (bPathTracingEnabled)
+	{
+		if (GGXTransmissionEnergyTexture == nullptr && GGXTransmissionEnergyTextureName.IsValid())
+		{
+			LoadEngineTexture(GGXTransmissionEnergyTexture, *GGXTransmissionEnergyTextureName.ToString());
+		}
+
+		if (!bSubstrateEnabled && SheenEnergyTexture == nullptr && SheenLegacyEnergyTextureName.IsValid())
+		{
+			LoadEngineTexture(SheenEnergyTexture, *SheenLegacyEnergyTextureName.ToString());
+		}
+
+		if (bSubstrateEnabled && SheenEnergyTexture == nullptr && SheenEnergyTextureName.IsValid())
+		{
+			LoadEngineTexture(SheenEnergyTexture, *SheenEnergyTextureName.ToString());
+		}
+
+		if (bSubstrateEnabled && DiffuseEnergyTexture == nullptr && DiffuseEnergyTextureName.IsValid())
+		{
+			LoadEngineTexture(DiffuseEnergyTexture, *DiffuseEnergyTextureName.ToString());
+		}
+	}
+}
+
 void UEngine::LoadGlintTextures()
 {
 	if (GlintTexture == nullptr && GlintTextureName.IsValid())
diff --git a/Engine/Source/Runtime/Renderer/Private/ShadingEnergyConservation.cpp b/Engine/Source/Runtime/Renderer/Private/ShadingEnergyConservation.cpp
--- a/Engine/Source/Runtime/Renderer/Private/ShadingEnergyConservation.cpp
+++ b/Engine/Source/Runtime/Renderer/Private/ShadingEnergyConservation.cpp
@@ -11,6 +11,8 @@
 #include "ScenePrivate.h"
 #include "PixelShaderUtils.h"
 #include "Substrate/Substrate.h"
+#include "Engine/VolumeTexture.h"
+#include "Engine/Engine.h"
 
 static TAutoConsoleVariable<int32> CVarShadingEnergyConservation(
 	TEXT("r.Shading.EnergyConservation"),
@@ -60,38 +62,153 @@ static TAutoConsoleVariable<int32> CVarMaterialEnergyConservation(
 
 namespace ShadingEnergyConservationData
 {
-	#include "ShadingEnergyConservationData.h"
-
-	template<typename TDataType>
-	void LockCopyTexture2D(FRHICommandListImmediate& RHICmdList, FRHITexture* Texture, const TDataType* InSrcBuffer, uint32 NumComponents)
+	static TRefCountPtr<IPooledRenderTarget> CreateTexture2D(FRHICommandListImmediate& RHICmdList, TObjectPtr<class UTexture2D>& InCPUTexture, EPixelFormat InFormat, const TCHAR* InName)
 	{
-		const uint8* SrcBuffer = (const uint8*)InSrcBuffer;
-		const uint32 SrcBytesPerPixel = sizeof(TDataType) * NumComponents;
+		if (!InCPUTexture || !InCPUTexture->GetPlatformData())
+		{
+			return nullptr;
+		}
+
+		check(InCPUTexture->Availability == ETextureAvailability::CPU);
+		FSharedImageConstRef Data = InCPUTexture->GetCPUCopy();
+		check(Data && Data->Format == ERawImageFormat::RGBA16F);
+		const TArrayView64<const FFloat16Color> DataView = Data->AsRGBA16F();
+
+		const FIntPoint DataSize(Data->SizeX, Data->SizeY);
+		TRefCountPtr<IPooledRenderTarget> OutTexture = GRenderTargetPool.FindFreeElement(FRDGTextureDesc::Create2D(DataSize, InFormat, FClearValueBinding::None, TexCreate_ShaderResource), InName);
+
+		FTextureRHIRef Texture = OutTexture->GetRHI();
+
+		check(InFormat == PF_G16R16 || InFormat == PF_R8G8 || InFormat == PF_G16 || InFormat == PF_R8);
+		const uint32 ComponentCount = InFormat == PF_G16R16 || InFormat == PF_R8G8 ? 2u: 1u;
+		const bool b8bit = InFormat == PF_R8G8 || InFormat == PF_R8;
+		const uint32 DstBytesPerPixel = (b8bit? 1u : 2u) * ComponentCount;
 
-		uint32 SrcStride = SrcBytesPerPixel * SHADING_ENERGY_CONSERVATION_TABLE_RESOLUTION;
-		uint32 DstStride;
-		uint8* DstBuffer = (uint8*)RHICmdList.LockTexture2D(Texture, 0, RLM_WriteOnly, DstStride, false);
-		for (uint32 y = 0; y < SHADING_ENERGY_CONSERVATION_TABLE_RESOLUTION; ++y)
+		// Write the contents of the texture with transcoding
+		uint32 DestStride;
+		uint8* DestBuffer = (uint8*)RHICmdList.LockTexture2D(Texture, 0, RLM_WriteOnly, DestStride, false);
+		for (int32 y = 0; y < DataSize.Y; ++y)
 		{
-			const uint8* Src = SrcBuffer + y * SrcStride;
-			uint8* Dst = DstBuffer + y * DstStride;
-			FMemory::Memcpy(Dst, Src, SrcStride);
+			for (int32 x = 0; x < DataSize.X; ++x)
+			{		
+				const FFloat16Color Value = DataView[x + y * DataSize.X];
+				if (b8bit)
+				{
+					uint8* Dest  = (uint8*)(DestBuffer + x * DstBytesPerPixel + y * DestStride);
+					{
+						Dest[0] = uint8(FMath::Clamp(float(Value.R) * 0xFF, 0u, 0xFF));
+					}
+					if (ComponentCount > 1)
+					{
+						Dest[1] = uint8(FMath::Clamp(float(Value.G) * 0xFF, 0u, 0xFF));
+					}
+				}
+				else
+				{
+					uint16* Dest = (uint16*)(DestBuffer + x * DstBytesPerPixel + y * DestStride);
+					{
+						Dest[0] = uint16(FMath::Clamp(float(Value.R) * 0xFFFF, 0u, 0xFFFF));
+					}
+					if (ComponentCount > 1)
+					{
+						Dest[1] = uint16(FMath::Clamp(float(Value.G) * 0xFFFF, 0u, 0xFFFF));
+					}
+				}
+			}
 		}
 		RHICmdList.UnlockTexture2D(Texture, 0, false);
+
+		// Release CPU data which are no longer needed
+		#if !WITH_EDITORONLY_DATA
+		InCPUTexture->RemoveFromRoot();
+		InCPUTexture = nullptr;
+		#endif
+
+		return OutTexture;
 	}
 
-	template<typename TDataType>
-	void LockCopyTexture3D(FRHICommandListImmediate& RHICmdList, FRHITexture* Texture, const TDataType* InSrcBuffer, uint32 NumComponents)
+	static TRefCountPtr<IPooledRenderTarget> CreateTexture3D(FRHICommandListImmediate& RHICmdList, TObjectPtr<class UTexture2D>& InCPUTexture, EPixelFormat InFormat, const TCHAR* InName)
 	{
-		const uint8* SrcBuffer = (const uint8*)InSrcBuffer;
-		const uint32 SrcBytesPerPixel = sizeof(TDataType) * NumComponents;
+		if (!InCPUTexture || !InCPUTexture->GetPlatformData())
+		{
+			return nullptr;
+		}
+
+		check(InCPUTexture->Availability == ETextureAvailability::CPU);
+		FSharedImageConstRef Data = InCPUTexture->GetCPUCopy();
+		#if 0
+		typedef FFloat16Color TColorType;
+		check(Data && Data->Format == ERawImageFormat::RGBA16F);
+		const TArrayView64<const FFloat16Color> DataView = Data->AsRGBA16F();
+		#else
+		typedef FLinearColor TColorType;
+		check(Data && Data->Format == ERawImageFormat::RGBA32F);
+		const TArrayView64<const FLinearColor> DataView = Data->AsRGBA32F();
+		#endif
+
+		// Stored as an array of 2D slices
+		const FIntVector DataSize(Data->SizeX, Data->SizeX, Data->SizeY / Data->SizeX);
+		TRefCountPtr<IPooledRenderTarget> OutTexture = GRenderTargetPool.FindFreeElement(FRDGTextureDesc::Create3D(DataSize, InFormat, FClearValueBinding::None, TexCreate_ShaderResource), InName);
+
+		FTextureRHIRef Texture = OutTexture->GetRHI();
+
+		check(InFormat == PF_G16R16 || InFormat == PF_R8G8 || InFormat == PF_G16 || InFormat == PF_R8);
+		const uint32 ComponentCount = InFormat == PF_G16R16 || InFormat == PF_R8G8 ? 2u: 1u;
+		const bool b8bit = InFormat == PF_R8G8 || InFormat == PF_R8;
+		const uint32 DstBytesPerPixel = (b8bit? 1u : 2u) * ComponentCount;
+
+		// Transcoded data before uploading data to the GPU
+		TArray<uint8> TranscodedData;
+		TranscodedData.SetNum(DataSize.X * DataSize.Y * DataSize.Z * DstBytesPerPixel);
+		uint32 DestStrideY = DataSize.X * DstBytesPerPixel;
+		uint32 DestStrideZ = DataSize.X * DataSize.Y * DstBytesPerPixel;
+		uint8* DestBuffer = (uint8*)TranscodedData.GetData();
+		for (int32 z = 0; z < DataSize.Z; ++z)
+		{
+			for (int32 y = 0; y < DataSize.Y; ++y)
+			{
+				for (int32 x = 0; x < DataSize.X; ++x)
+				{		
+					const TColorType Value = DataView[x + y * DataSize.X + z * DataSize.X * DataSize.Y];
+					if (b8bit)
+					{
+						uint8* Dest  = (uint8*)(DestBuffer + x * DstBytesPerPixel + y * DestStrideY + z * DestStrideZ);
+						{
+							Dest[0] = uint8(FMath::Clamp(float(Value.R) * 0xFF, 0u, 0xFF));
+						}
+						if (ComponentCount > 1)
+						{
+							Dest[1] = uint8(FMath::Clamp(float(Value.G) * 0xFF, 0u, 0xFF));
+						}
+					}
+					else
+					{
+						uint16* Dest = (uint16*)(DestBuffer + x * DstBytesPerPixel + y * DestStrideY + z * DestStrideZ);
+						{
+							Dest[0] = uint16(FMath::Clamp(float(Value.R) * 0xFFFF, 0u, 0xFFFF));
+						}
+						if (ComponentCount > 1)
+						{
+							Dest[1] = uint16(FMath::Clamp(float(Value.G) * 0xFFFF, 0u, 0xFFFF));
+						}
+					}
+				}
+			}
+		}
 
-		const FIntVector Extent(SHADING_ENERGY_CONSERVATION_TABLE_RESOLUTION, SHADING_ENERGY_CONSERVATION_TABLE_RESOLUTION, SHADING_ENERGY_CONSERVATION_TABLE_RESOLUTION);
-		FUpdateTextureRegion3D Region(0, 0, 0, 0, 0, 0, Extent.X, Extent.Y, Extent.Z);
-		RHICmdList.UpdateTexture3D(Texture, 0, Region, Extent.X * SrcBytesPerPixel, Extent.X * Extent.Y * SrcBytesPerPixel, SrcBuffer);
+		FUpdateTextureRegion3D Region(0, 0, 0, 0, 0, 0, DataSize.X, DataSize.Y, DataSize.Z);
+		RHICmdList.UpdateTexture3D(Texture, 0, Region, DataSize.X * DstBytesPerPixel, DataSize.X * DataSize.Y * DstBytesPerPixel, TranscodedData.GetData());
 
 		// UpdateTexture3D before and after state is currently undefined
 		RHICmdList.Transition(FRHITransitionInfo(Texture, ERHIAccess::Unknown, ERHIAccess::SRVMask));
+
+		// Release CPU data which are no longer needed
+		#if !WITH_EDITORONLY_DATA
+		InCPUTexture->RemoveFromRoot();
+		InCPUTexture = nullptr;
+		#endif
+
+		return OutTexture;
 	}
 }
 
@@ -204,22 +321,6 @@ IMPLEMENT_SHADER_TYPE(, FBuildShadingEnergyConservationTableCS, TEXT("/Engine/Pr
 
 ////////////////////////////////////////////////////////////////////////////////////////////////
 
-struct FTranscode
-{
-	FTranscode(const uint16* InData, uint32 InElementCount)
-	{
-		const uint32 ElementCount = sizeof(ShadingEnergyConservationData::GGXSpecValues) / sizeof(uint16);
-		Data.SetNum(InElementCount);
-		for (uint32 It = 0; It < InElementCount; ++It)
-		{
-			Data[It] = uint8(FMath::Clamp((float(InData[It]) / 65535.f) * 0xFF, 0u, 0xFF));
-		}
-	}
-	TArray<uint8> Data;
-};
-
-#define TRANSCODE_16_TO_8(T) (FTranscode(T, sizeof(T) / sizeof(uint16)).Data.GetData())
-
 namespace ShadingEnergyConservation
 {
 
@@ -230,11 +331,6 @@ bool IsEnable()
 
 void Init(FRDGBuilder& GraphBuilder, FViewInfo& View)
 {
-	FRDGTextureRef GGXSpecEnergyTexture = nullptr;
-	FRDGTextureRef GGXGlassEnergyTexture = nullptr;
-	FRDGTextureRef ClothEnergyTexture = nullptr;
-	FRDGTextureRef DiffuseEnergyTexture = nullptr;
-
 	// Enabled based on settings
 	const bool bMaterialEnergyConservationEnabled = CVarMaterialEnergyConservation.GetValueOnRenderThread() > 0;
 	const bool bIsEnergyConservationEnabled = CVarShadingEnergyConservation.GetValueOnRenderThread() > 0;
@@ -275,10 +371,10 @@ void Init(FRDGBuilder& GraphBuilder, FViewInfo& View)
 
 			if (bRuntimeGeneration)
 			{			
-				GGXSpecEnergyTexture	= GraphBuilder.CreateTexture(FRDGTextureDesc::Create2D(FIntPoint(Size, Size),        SpecFormat, FClearValueBinding::None, TexCreate_ShaderResource | TexCreate_UAV), TEXT("Shading.GGXSpecEnergy"),   ERDGTextureFlags::MultiFrame);
-				GGXGlassEnergyTexture	= GraphBuilder.CreateTexture(FRDGTextureDesc::Create3D(FIntVector(Size, Size, Size), SpecFormat, FClearValueBinding::None, TexCreate_ShaderResource | TexCreate_UAV), TEXT("Shading.GGXGlassEnergy"),  ERDGTextureFlags::MultiFrame);
-				ClothEnergyTexture		= GraphBuilder.CreateTexture(FRDGTextureDesc::Create2D(FIntPoint(Size, Size),        SpecFormat, FClearValueBinding::None, TexCreate_ShaderResource | TexCreate_UAV), TEXT("Shading.ClothSpecEnergy"), ERDGTextureFlags::MultiFrame);
-				DiffuseEnergyTexture	= GraphBuilder.CreateTexture(FRDGTextureDesc::Create2D(FIntPoint(Size, Size),        DiffFormat, FClearValueBinding::None, TexCreate_ShaderResource | TexCreate_UAV), TEXT("Shading.DiffuseEnergy"),   ERDGTextureFlags::MultiFrame);
+				FRDGTextureRef GGXSpecEnergyTexture	= GraphBuilder.CreateTexture(FRDGTextureDesc::Create2D(FIntPoint(Size, Size),        SpecFormat, FClearValueBinding::None, TexCreate_ShaderResource | TexCreate_UAV), TEXT("Shading.GGXSpecEnergy"),   ERDGTextureFlags::MultiFrame);
+				FRDGTextureRef GGXGlassEnergyTexture= GraphBuilder.CreateTexture(FRDGTextureDesc::Create3D(FIntVector(Size, Size, Size), SpecFormat, FClearValueBinding::None, TexCreate_ShaderResource | TexCreate_UAV), TEXT("Shading.GGXGlassEnergy"),  ERDGTextureFlags::MultiFrame);
+				FRDGTextureRef ClothEnergyTexture	= GraphBuilder.CreateTexture(FRDGTextureDesc::Create2D(FIntPoint(Size, Size),        SpecFormat, FClearValueBinding::None, TexCreate_ShaderResource | TexCreate_UAV), TEXT("Shading.ClothSpecEnergy"), ERDGTextureFlags::MultiFrame);
+				FRDGTextureRef DiffuseEnergyTexture	= GraphBuilder.CreateTexture(FRDGTextureDesc::Create2D(FIntPoint(Size, Size),        DiffFormat, FClearValueBinding::None, TexCreate_ShaderResource | TexCreate_UAV), TEXT("Shading.DiffuseEnergy"),   ERDGTextureFlags::MultiFrame);
 			
 				const uint32 NumSamples = 1u << 14u;
 
@@ -357,47 +453,18 @@ void Init(FRDGBuilder& GraphBuilder, FViewInfo& View)
 				check(SpecFormat == PF_G16R16 || SpecFormat == PF_R8G8);
 				check(DiffFormat == PF_G16 || DiffFormat == PF_R8);
 
-				View.ViewState->ShadingEnergyConservationData.GGXSpecEnergyTexture  = GRenderTargetPool.FindFreeElement(FRDGTextureDesc::Create2D(FIntPoint(Size, Size),        SpecFormat, FClearValueBinding::None, TexCreate_ShaderResource), TEXT("Shading.GGXSpecEnergy"));
-				View.ViewState->ShadingEnergyConservationData.GGXGlassEnergyTexture = GRenderTargetPool.FindFreeElement(FRDGTextureDesc::Create3D(FIntVector(Size, Size, Size), SpecFormat, FClearValueBinding::None, TexCreate_ShaderResource), TEXT("Shading.GGXGlassEnergy"));
-				View.ViewState->ShadingEnergyConservationData.ClothEnergyTexture    = GRenderTargetPool.FindFreeElement(FRDGTextureDesc::Create2D(FIntPoint(Size, Size),        SpecFormat, FClearValueBinding::None, TexCreate_ShaderResource), TEXT("Shading.ClothSpecEnergy"));
-				View.ViewState->ShadingEnergyConservationData.DiffuseEnergyTexture  = GRenderTargetPool.FindFreeElement(FRDGTextureDesc::Create2D(FIntPoint(Size, Size),        DiffFormat,	FClearValueBinding::None, TexCreate_ShaderResource), TEXT("Shading.DiffuseEnergy"));
-
-
-				if (bRG16Supported)
-				{
-					ShadingEnergyConservationData::LockCopyTexture2D(GraphBuilder.RHICmdList, View.ViewState->ShadingEnergyConservationData.GGXSpecEnergyTexture->GetRHI(),  ShadingEnergyConservationData::GGXSpecValues, 2);
-					ShadingEnergyConservationData::LockCopyTexture3D(GraphBuilder.RHICmdList, View.ViewState->ShadingEnergyConservationData.GGXGlassEnergyTexture->GetRHI(), ShadingEnergyConservationData::GGXGlassValues, 2);
-					ShadingEnergyConservationData::LockCopyTexture2D(GraphBuilder.RHICmdList, View.ViewState->ShadingEnergyConservationData.ClothEnergyTexture->GetRHI(),    Substrate::IsSubstrateEnabled() ? ShadingEnergyConservationData::SubstrateClothSpecValues : ShadingEnergyConservationData::ClothSpecValues, 2);
-				}
-				else
-				{
-					ShadingEnergyConservationData::LockCopyTexture2D(GraphBuilder.RHICmdList, View.ViewState->ShadingEnergyConservationData.GGXSpecEnergyTexture->GetRHI(),  TRANSCODE_16_TO_8(ShadingEnergyConservationData::GGXSpecValues), 2);
-					ShadingEnergyConservationData::LockCopyTexture3D(GraphBuilder.RHICmdList, View.ViewState->ShadingEnergyConservationData.GGXGlassEnergyTexture->GetRHI(), TRANSCODE_16_TO_8(ShadingEnergyConservationData::GGXGlassValues), 2);
-					ShadingEnergyConservationData::LockCopyTexture2D(GraphBuilder.RHICmdList, View.ViewState->ShadingEnergyConservationData.ClothEnergyTexture->GetRHI(),    Substrate::IsSubstrateEnabled() ? TRANSCODE_16_TO_8(ShadingEnergyConservationData::SubstrateClothSpecValues) : TRANSCODE_16_TO_8(ShadingEnergyConservationData::ClothSpecValues), 2);
-				}
-
-				if (bR16Supported)
-				{
-					ShadingEnergyConservationData::LockCopyTexture2D(GraphBuilder.RHICmdList, View.ViewState->ShadingEnergyConservationData.DiffuseEnergyTexture->GetRHI(),  ShadingEnergyConservationData::DiffuseValues, 1);
-				}
-				else
-				{
-					ShadingEnergyConservationData::LockCopyTexture2D(GraphBuilder.RHICmdList, View.ViewState->ShadingEnergyConservationData.DiffuseEnergyTexture->GetRHI(),  TRANSCODE_16_TO_8(ShadingEnergyConservationData::DiffuseValues), 1);
-				}
+				View.ViewState->ShadingEnergyConservationData.GGXSpecEnergyTexture  = ShadingEnergyConservationData::CreateTexture2D(GraphBuilder.RHICmdList, GEngine->GGXReflectionEnergyTexture,   SpecFormat, TEXT("Shading.GGXReflectionEnergy"));
+				View.ViewState->ShadingEnergyConservationData.GGXGlassEnergyTexture = ShadingEnergyConservationData::CreateTexture3D(GraphBuilder.RHICmdList, GEngine->GGXTransmissionEnergyTexture, SpecFormat, TEXT("Shading.GGXTransmissionEnergy"));
+				View.ViewState->ShadingEnergyConservationData.ClothEnergyTexture    = ShadingEnergyConservationData::CreateTexture2D(GraphBuilder.RHICmdList, GEngine->SheenEnergyTexture,           SpecFormat, TEXT("Shading.SheenEnergy"));
+				View.ViewState->ShadingEnergyConservationData.DiffuseEnergyTexture  = ShadingEnergyConservationData::CreateTexture2D(GraphBuilder.RHICmdList, GEngine->DiffuseEnergyTexture,         DiffFormat, TEXT("Shading.DiffuseEnergy"));
 
-				GGXSpecEnergyTexture  = GraphBuilder.RegisterExternalTexture(View.ViewState->ShadingEnergyConservationData.GGXSpecEnergyTexture);
-				GGXGlassEnergyTexture = GraphBuilder.RegisterExternalTexture(View.ViewState->ShadingEnergyConservationData.GGXGlassEnergyTexture);
-				ClothEnergyTexture    = GraphBuilder.RegisterExternalTexture(View.ViewState->ShadingEnergyConservationData.ClothEnergyTexture);
-				DiffuseEnergyTexture  = GraphBuilder.RegisterExternalTexture(View.ViewState->ShadingEnergyConservationData.DiffuseEnergyTexture);
+				// Fallback
+				if (!View.ViewState->ShadingEnergyConservationData.GGXSpecEnergyTexture ) { View.ViewState->ShadingEnergyConservationData.GGXSpecEnergyTexture  = GSystemTextures.BlackDummy; }
+				if (!View.ViewState->ShadingEnergyConservationData.GGXGlassEnergyTexture) { View.ViewState->ShadingEnergyConservationData.GGXGlassEnergyTexture = GSystemTextures.VolumetricBlackDummy; }
+				if (!View.ViewState->ShadingEnergyConservationData.ClothEnergyTexture   ) { View.ViewState->ShadingEnergyConservationData.ClothEnergyTexture    = GSystemTextures.BlackDummy; }
+				if (!View.ViewState->ShadingEnergyConservationData.DiffuseEnergyTexture ) { View.ViewState->ShadingEnergyConservationData.DiffuseEnergyTexture  = GSystemTextures.BlackDummy; }
 			}
 		}
-		else
-		{
-			GGXSpecEnergyTexture	= GraphBuilder.RegisterExternalTexture(View.ViewState->ShadingEnergyConservationData.GGXSpecEnergyTexture);
-			GGXGlassEnergyTexture	= GraphBuilder.RegisterExternalTexture(View.ViewState->ShadingEnergyConservationData.GGXGlassEnergyTexture);
-			ClothEnergyTexture		= GraphBuilder.RegisterExternalTexture(View.ViewState->ShadingEnergyConservationData.ClothEnergyTexture);
-			DiffuseEnergyTexture	= GraphBuilder.RegisterExternalTexture(View.ViewState->ShadingEnergyConservationData.DiffuseEnergyTexture);
-		}
 
 		View.ViewState->ShadingEnergyConservationData.bEnergyConservation = bIsEnergyConservationEnabled;
 		View.ViewState->ShadingEnergyConservationData.bEnergyPreservation = bIsEnergyPreservationEnabled;
diff --git a/Engine/Source/Runtime/Renderer/Private/ShadingEnergyConservationData.h b/Engine/Source/Runtime/Renderer/Private/ShadingEnergyConservationData.h
deleted file mode 100644
diff --git a/Engine/Source/Runtime/Renderer/Private/SystemTextures.cpp b/Engine/Source/Runtime/Renderer/Private/SystemTextures.cpp
--- a/Engine/Source/Runtime/Renderer/Private/SystemTextures.cpp
+++ b/Engine/Source/Runtime/Renderer/Private/SystemTextures.cpp
@@ -10,12 +10,93 @@
 #include "Math/Float16.h"
 #include "RenderTargetPool.h"
 #include "ClearQuad.h"
-#include "LTC.h"
 #include "Math/PackedVector.h"
 #include "GlobalRenderResources.h"
+#include "Engine/Texture2D.h"
+#include "Engine/Engine.h"
 
 /*-----------------------------------------------------------------------------
-SystemTextures
+ Export System textures (LTC / EnergyConservation / Hair LUT / ...)
+-----------------------------------------------------------------------------*/
+#define EXPORT_SYSTEM_TEXTURES 0
+
+#if EXPORT_SYSTEM_TEXTURES
+#include "LTC.h"
+#include "ImageUtils.h"
+
+static void SaveEXR(const FString& Filename, const FIntPoint& DataSize, const TArray<FFloat16Color>& Data)
+{
+	const FString RelativeFilepath = FPaths::EngineSavedDir() + Filename + TEXT(".exr");
+	const FString AbsoluteFilepath = FPaths::ConvertRelativePathToFull(RelativeFilepath);
+	FImageView ImageView(Data.GetData(), DataSize.X, DataSize.Y);
+	const bool bSucceed = FImageUtils::SaveImageByExtension(*AbsoluteFilepath, ImageView);
+	check(bSucceed)
+}
+
+static void ExportSystemTextures()
+{
+	// LTC Textures	(used by Rect Lights and/or BSDF evaluation)
+	// GGX - LTC matrix coefficients (4-coefficients)
+	{ 
+		const FIntPoint DataSize(GGX_LTC_Size, GGX_LTC_Size);
+		TArray<FFloat16Color> RawPixels;
+		RawPixels.SetNum(DataSize.X * DataSize.Y);
+		for (int32 y = 0; y < GGX_LTC_Size; ++y)
+		{
+			for (int32 x = 0; x < GGX_LTC_Size; ++x)
+			{
+				FFloat16Color& Dest = RawPixels[x + y * DataSize.X];
+				Dest.R = FFloat16(GGX_LTC_Mat[4 * (x + y * GGX_LTC_Size) + 0]);
+				Dest.G = FFloat16(GGX_LTC_Mat[4 * (x + y * GGX_LTC_Size) + 1]);
+				Dest.B = FFloat16(GGX_LTC_Mat[4 * (x + y * GGX_LTC_Size) + 2]);
+				Dest.A = FFloat16(GGX_LTC_Mat[4 * (x + y * GGX_LTC_Size) + 3]);
+			}
+		}
+		SaveEXR(TEXT("GGX_LTCMat"), DataSize, RawPixels);
+	}
+
+	// GGX - Split-Sum Amplitude coefficients (2-components)
+	{
+		const FIntPoint DataSize(GGX_LTC_Size, GGX_LTC_Size);
+		TArray<FFloat16Color> RawPixels;
+		RawPixels.SetNum(DataSize.X * DataSize.Y);
+		for (int32 y = 0; y < GGX_LTC_Size; ++y)
+		{
+			for (int32 x = 0; x < GGX_LTC_Size; ++x)
+			{
+				FFloat16Color& Dest = RawPixels[x + y * DataSize.X];
+				Dest.R = FFloat16(GGX_LTC_Amp[4 * (x + y * GGX_LTC_Size) + 0]);
+				Dest.G = FFloat16(GGX_LTC_Amp[4 * (x + y * GGX_LTC_Size) + 1]);
+				Dest.B = 0;
+				Dest.A = 0;
+			}
+		}
+		SaveEXR(TEXT("GGX_LTCAmp"), DataSize, RawPixels);
+	}
+
+	// Sheen - Matrix & directional albedo (3-components)
+	{
+		const FIntPoint DataSize(Sheen_LTC_Size, Sheen_LTC_Size);
+		TArray<FFloat16Color> RawPixels;
+		RawPixels.SetNum(DataSize.X * DataSize.Y);
+		for (int32 y = 0; y < DataSize.Y; ++y)
+		{
+			for (int32 x = 0; x < DataSize.X; ++x)
+			{
+				FFloat16Color& Dest = RawPixels[x + y * DataSize.X];
+				Dest.R = FFloat16(Sheen_LTC_Volume[x][y][0]);
+				Dest.G = FFloat16(Sheen_LTC_Volume[x][y][1]);
+				Dest.B = FFloat16(Sheen_LTC_Volume[x][y][2]);
+				Dest.A = FFloat16(0);
+			}
+		}
+		SaveEXR(TEXT("Sheen_LTC"), DataSize, RawPixels);
+	}
+}
+#endif // EXPORT_SYSTEM_TEXTURES
+
+/*-----------------------------------------------------------------------------
+ SystemTextures
 -----------------------------------------------------------------------------*/
 
 RDG_REGISTER_BLACKBOARD_STRUCT(FRDGSystemTextures);
@@ -81,6 +162,12 @@ void FSystemTextures::InitializeTextures(FRHICommandListImmediate& RHICmdList, c
 		InitializeFeatureLevelDependentTextures(RHICmdList, InFeatureLevel);
 	}
 	// there's no needed setup for those feature levels lower or identical to the current one
+
+	// Initialized resources depending on GEngine data
+	if (!bEngineDependentTexturesInitialized && GEngine)
+	{
+		InitializeEngineDependentTextures(RHICmdList);
+	}
 }
 
 template <typename DataType>
@@ -103,6 +190,57 @@ void SetDummyTextureArrayData(FRHITexture* Texture, const DataType& DummyData)
 
 const static FLazyName SystemTexturesName(TEXT("FSystemTextures"));
 
+static TRefCountPtr<IPooledRenderTarget> CreateTexture(FRHICommandListImmediate& RHICmdList, TObjectPtr<class UTexture2D>& InCPUTexture, EPixelFormat InFormat, const TCHAR* InName)
+{
+	check(InCPUTexture->Availability == ETextureAvailability::CPU);
+	FSharedImageConstRef Data = InCPUTexture->GetCPUCopy();
+	check(Data && Data->Format == ERawImageFormat::RGBA16F);
+	const TArrayView64<const FFloat16Color> DataView = Data->AsRGBA16F();
+
+	const FIntPoint DataSize(Data->SizeX, Data->SizeY);
+
+	const FRHITextureCreateDesc Desc =
+	FRHITextureCreateDesc::Create2D(InName, DataSize.X, DataSize.Y, InFormat)
+	.SetFlags(ETextureCreateFlags::ShaderResource | ETextureCreateFlags::FastVRAM)
+	.SetClassName(SystemTexturesName);
+
+	FTextureRHIRef Texture = RHICreateTexture(Desc);
+
+	check(InFormat == PF_FloatRGBA || InFormat == PF_G16R16F);
+	const uint32 ComponentCount = InFormat == PF_G16R16F ? 2u: 4u;
+
+	// Write the contents of the texture.
+	uint32 DestStride;
+	uint8* DestBuffer = (uint8*)RHICmdList.LockTexture2D(Texture, 0, RLM_WriteOnly, DestStride, false);
+	for (int32 y = 0; y < DataSize.Y; ++y)
+	{
+		for (int32 x = 0; x < DataSize.X; ++x)
+		{
+			uint16* Dest = (uint16*)(DestBuffer + x * ComponentCount * sizeof(uint16) + y * DestStride);
+
+			const FFloat16Color Value = DataView[x + y * DataSize.X];
+			Dest[0] = Value.R.Encoded;
+			if (ComponentCount > 1) { Dest[1] = Value.G.Encoded; }
+			if (ComponentCount > 2) { Dest[2] = Value.B.Encoded; }
+			if (ComponentCount > 3) { Dest[3] = Value.A.Encoded; }
+		}
+	}
+	RHICmdList.UnlockTexture2D(Texture, 0, false);
+
+	// Release CPU data which are no longer needed
+	#if !WITH_EDITORONLY_DATA
+	InCPUTexture->RemoveFromRoot();
+	InCPUTexture = nullptr;
+	#endif
+
+	return CreateRenderTarget(Texture, Desc.DebugName);
+}
+
+static bool IsTextureDataValid(UTexture2D* In)
+{
+	return In && In->GetPlatformData();
+}
+
 void FSystemTextures::InitializeCommonTextures(FRHICommandListImmediate& RHICmdList)
 {
 	// First initialize textures that are common to all feature levels. This is always done the first time we come into this function, as doesn't care about the
@@ -669,97 +807,6 @@ void FSystemTextures::InitializeFeatureLevelDependentTextures(FRHICommandListImm
 			PerlinNoise3D = CreateRenderTarget(Texture3D, Desc.DebugName);
 
 		} // end Create the PerlinNoise3D texture
-
-		// LTC Textures	(used by Rect Lights and/or BSDF evaluation)
-		{
-			// GGX - LTC matrix coefficients (4-coefficients)
-			{
-				const FRHITextureCreateDesc Desc =
-					FRHITextureCreateDesc::Create2D(TEXT("GGX.LTCMat"), GGX_LTC_Size, GGX_LTC_Size, PF_FloatRGBA)
-					.SetFlags(ETextureCreateFlags::ShaderResource | ETextureCreateFlags::FastVRAM)
-					.SetClassName(SystemTexturesName);
-
-				FTextureRHIRef Texture = RHICreateTexture(Desc);
-
-				// Write the contents of the texture.
-				uint32 DestStride;
-				uint8* DestBuffer = (uint8*)RHICmdList.LockTexture2D(Texture, 0, RLM_WriteOnly, DestStride, false);
-    
-				for (int32 y = 0; y < GGX_LTC_Size; ++y)
-				{
-					for (int32 x = 0; x < GGX_LTC_Size; ++x)
-					{
-						uint16* Dest = (uint16*)(DestBuffer + x * 4 * sizeof(uint16) + y * DestStride);
-    
-						for (int k = 0; k < 4; k++)
-						{
-							Dest[k] = FFloat16(GGX_LTC_Mat[4 * (x + y * GGX_LTC_Size) + k]).Encoded;
-						}
-					}
-				}
-				RHICmdList.UnlockTexture2D(Texture, 0, false);
-
-				GGXLTCMat = CreateRenderTarget(Texture, Desc.DebugName);
-			}
-
-			// GGX - Split-Sum Amplitude coefficients (2-components)
-			{
-				const FRHITextureCreateDesc Desc =
-					FRHITextureCreateDesc::Create2D(TEXT("GGX.LTCAmp"), GGX_LTC_Size, GGX_LTC_Size, PF_G16R16F)
-					.SetFlags(ETextureCreateFlags::ShaderResource | ETextureCreateFlags::FastVRAM)
-					.SetClassName(SystemTexturesName);
-
-				FTextureRHIRef Texture = RHICreateTexture(Desc);
-
-				// Write the contents of the texture.
-				uint32 DestStride;
-				uint8* DestBuffer = (uint8*)RHICmdList.LockTexture2D(Texture, 0, RLM_WriteOnly, DestStride, false);
-
-				for (int32 y = 0; y < GGX_LTC_Size; ++y)
-				{
-					for (int32 x = 0; x < GGX_LTC_Size; ++x)
-					{
-						uint16* Dest = (uint16*)(DestBuffer + x * 2 * sizeof(uint16) + y * DestStride);
-
-						for (int k = 0; k < 2; k++)
-						{
-							Dest[k] = FFloat16(GGX_LTC_Amp[4 * (x + y * GGX_LTC_Size) + k]).Encoded;
-						}
-					}
-				}
-				RHICmdList.UnlockTexture2D(Texture, 0, false);
-
-				GGXLTCAmp = CreateRenderTarget(Texture, Desc.DebugName);
-			}
-
-			// Sheen - Matrix & directional albedo (3-components)
-			{
-				const FRHITextureCreateDesc Desc =
-					FRHITextureCreateDesc::Create2D(TEXT("Sheen.LTC"), Sheen_LTC_Size, Sheen_LTC_Size, PF_FloatRGBA)
-					.SetFlags(ETextureCreateFlags::ShaderResource | ETextureCreateFlags::FastVRAM);
-
-				FTextureRHIRef Texture = RHICreateTexture(Desc);
-
-				// Write the contents of the texture.
-				uint32 DestStride;
-				uint8* DestBuffer = (uint8*)RHICmdList.LockTexture2D(Texture, 0, RLM_WriteOnly, DestStride, false);
-
-				for (int32 y = 0; y < Sheen_LTC_Size; ++y)
-				{
-					for (int32 x = 0; x < Sheen_LTC_Size; ++x)
-					{
-						uint16* Dest = (uint16*)(DestBuffer + x * 4 * sizeof(uint16) + y * DestStride);
-						Dest[0] = FFloat16(Sheen_LTC_Volume[x][y][0]).Encoded;
-						Dest[1] = FFloat16(Sheen_LTC_Volume[x][y][1]).Encoded;
-						Dest[2] = FFloat16(Sheen_LTC_Volume[x][y][2]).Encoded;
-						Dest[3] = 0;
-					}
-				}
-				RHICmdList.UnlockTexture2D(Texture, 0, false);
-
-				SheenLTC = CreateRenderTarget(Texture, Desc.DebugName);
-			}
-		}
 	}
 
 	// Create the SSAO randomization texture
@@ -942,10 +989,41 @@ void FSystemTextures::InitializeFeatureLevelDependentTextures(FRHICommandListImm
         ZeroUIntArrayAtomicCompatDummy = CreateRenderTarget(Texture, Desc.DebugName);
     }
     
+#if EXPORT_SYSTEM_TEXTURES
+	ExportSystemTextures();
+#endif
+
 	// Initialize textures only once.
 	FeatureLevelInitializedTo = InFeatureLevel;
 }
 
+void FSystemTextures::InitializeEngineDependentTextures(FRHICommandListImmediate& RHICmdList)
+{	
+	if (GEngine)
+	{
+		// LTC Textures	(used by Rect Lights and/or BSDF evaluation)
+
+		// GGX - LTC matrix coefficients (4-coefficients)
+		if (IsTextureDataValid(GEngine->GGXLTCMatTexture))
+		{
+			GGXLTCMat = CreateTexture(RHICmdList, GEngine->GGXLTCMatTexture, PF_FloatRGBA, TEXT("GGX.LTCMat"));
+		}
+		// GGX - Split-Sum Amplitude coefficients (2-components)
+		if (IsTextureDataValid(GEngine->GGXLTCAmpTexture))
+		{
+			GGXLTCAmp = CreateTexture(RHICmdList, GEngine->GGXLTCAmpTexture, PF_G16R16F, TEXT("GGX.LTCAmp"));
+		}
+		// Sheen - Matrix & directional albedo (3-components)
+		if (IsTextureDataValid(GEngine->SheenLTCTexture))
+		{
+			SheenLTC  = CreateTexture(RHICmdList, GEngine->SheenLTCTexture,  PF_FloatRGBA, TEXT("Sheen.LTC"));
+		}
+	}
+
+	// Initialize textures only once.
+	bEngineDependentTexturesInitialized = true;
+}
+
 void FSystemTextures::ReleaseRHI()
 {
 	WhiteDummySRV.SafeRelease();
diff --git a/Engine/Source/Runtime/Renderer/Public/SystemTextures.h b/Engine/Source/Runtime/Renderer/Public/SystemTextures.h
--- a/Engine/Source/Runtime/Renderer/Public/SystemTextures.h
+++ b/Engine/Source/Runtime/Renderer/Public/SystemTextures.h
@@ -240,6 +240,7 @@ public:
 protected:
 	/** Maximum feature level that the textures have been initialized up to */
 	ERHIFeatureLevel::Type FeatureLevelInitializedTo;
+	bool bEngineDependentTexturesInitialized = false;
 
 	/** Default textures allocated on-demand */
 	TArray<FDefaultTexture> DefaultTextures;
@@ -249,6 +250,7 @@ protected:
 
 	void InitializeCommonTextures(FRHICommandListImmediate& RHICmdList);
 	void InitializeFeatureLevelDependentTextures(FRHICommandListImmediate& RHICmdList, const ERHIFeatureLevel::Type InFeatureLevel);
+	void InitializeEngineDependentTextures(FRHICommandListImmediate& RHICmdList);
 };
 
 /** The global system textures used for scene rendering. */
