Add (optional) RGBA16 bit format support for rect light atlas in order to address hue shift.

diff --git a/Engine/Shaders/Private/RectLightAtlas.usf b/Engine/Shaders/Private/RectLightAtlas.usf
--- a/Engine/Shaders/Private/RectLightAtlas.usf
+++ b/Engine/Shaders/Private/RectLightAtlas.usf
@@ -220,6 +220,7 @@ uint  FreeCount;
 int2  OutputResolution;
 uint  AtlasMIPIndex;
 uint  AtlasSourceTextureMIPBias;
+uint  AtlasFormat;
 
 SamplerState AtlasSampler;
 Texture2D<float4> AtlasTexture;
@@ -280,7 +281,12 @@ void MainCS(uint3 DispatchThreadId : SV_DispatchThreadID)
 		Newline(Context);
 
 		Print(Context, TEXT("Atlas format     : "), FontSilver);
-		Print(Context, TEXT("PF_FloatR11G11B10"), FontOrange);
+		if (AtlasFormat == 0) 
+		{ Print(Context, TEXT("PF_FloatR11G11B10"), FontOrange); }
+		else if (AtlasFormat == 1) 
+		{ Print(Context, TEXT("PF_FloatRGBA"), FontOrange); }
+		else  
+		{ Print(Context, TEXT("Unknown"), FontOrange); }
 		Newline(Context);
 	}
 
diff --git a/Engine/Source/Runtime/Renderer/Private/RectLightTextureManager.cpp b/Engine/Source/Runtime/Renderer/Private/RectLightTextureManager.cpp
--- a/Engine/Source/Runtime/Renderer/Private/RectLightTextureManager.cpp
+++ b/Engine/Source/Runtime/Renderer/Private/RectLightTextureManager.cpp
@@ -64,6 +64,12 @@ static TAutoConsoleVariable<int32> CVarRectLighMaxTextureRatio(
 	TEXT("Define the max Width/Height or Height/Width ratio that a texture can have."),
 	ECVF_RenderThreadSafe);
 
+static TAutoConsoleVariable<int32> CVarRectLighAtlasFormat(
+	TEXT("r.RectLightAtlas.Format"),
+	0,
+	TEXT("Define the atlas format used.\n - 0: R11G11B10 (faster but can caused hue shift). \n - 1: RGA16F (more accurate but slower)"),
+	ECVF_RenderThreadSafe);
+
 namespace RectLightAtlas
 {
 
@@ -226,6 +232,11 @@ static int32 GetSlotMaxMIPLevel(const FAtlasSlot& In)
 	return MaxMIP > 0u ? MaxMIP-1u : 0u;
 }
 
+static EPixelFormat GetRectLightAtlasFormat()
+{
+	return CVarRectLighAtlasFormat.GetValueOnRenderThread() > 0 ? PF_FloatRGBA : PF_FloatR11G11B10;
+}
+
 static bool Traits_IsValid(const FAtlasRect& In)					{ return true; }
 static bool Traits_IsValid(const FAtlasHorizon& In)					{ return true; }
 static bool Traits_IsValid(const FAtlasSlot& In)					{ return In.IsValid(); }
@@ -285,6 +296,7 @@ class FRectLightAtlasDebugInfoCS : public FGlobalShader
 		SHADER_PARAMETER(FIntPoint, OutputResolution)
 		SHADER_PARAMETER(uint32, AtlasMIPIndex)
 		SHADER_PARAMETER(uint32, AtlasSourceTextureMIPBias)
+		SHADER_PARAMETER(uint32, AtlasFormat)
 		SHADER_PARAMETER_SAMPLER(SamplerState, AtlasSampler)
 		SHADER_PARAMETER_RDG_TEXTURE(Texture2D, AtlasTexture)
 		SHADER_PARAMETER_RDG_TEXTURE_UAV(RWTexture2D, OutputTexture)
@@ -348,11 +360,20 @@ static FRDGTextureRef AddRectLightDebugInfoPass(
 	FRDGBufferRef HorizonBuffer = CreateSlotBuffer(GraphBuilder, Horizons, TEXT("RectLight.HorizonBuffer"));
 	FRDGBufferRef FreeBuffer = CreateSlotBuffer(GraphBuilder, FreeRects, TEXT("RectLight.FreeBuffer"));
 
+	uint32 AtlasFormat = 2;
+	switch (AtlasTexture->Desc.Format)
+	{
+		case PF_FloatR11G11B10 : AtlasFormat = 0; break;
+		case PF_FloatRGBA : AtlasFormat = 1; break;
+		default: AtlasFormat = 2; break;
+	}	
+
 	const FIntPoint OutputResolution(OutputTexture->Desc.Extent);
 	FRectLightAtlasDebugInfoCS::FParameters* Parameters = GraphBuilder.AllocParameters<FRectLightAtlasDebugInfoCS::FParameters>();
 	Parameters->AtlasResolution = AtlasTexture->Desc.Extent;
 	Parameters->AtlasMaxMipLevel = AtlasTexture->Desc.NumMips;
 	Parameters->AtlasSourceTextureMIPBias = GRectLightTextureManager.AtlasLayout.SourceTextureMIPBias;
+	Parameters->AtlasFormat = AtlasFormat;
 	Parameters->Occupancy = OccupiedPixels / float(AtlasTexture->Desc.Extent.X * AtlasTexture->Desc.Extent.Y);
 	Parameters->SlotCount = ValidSlots.Num();
 	Parameters->HorizonCount = Horizons.Num();
@@ -1516,12 +1537,12 @@ FAtlasSlotDesc GetAtlasSlot(uint32 InSlotIndex)
 	return Out;
 }
 
-static FRDGTextureRef CreateRectLightAtlasTexture(FRDGBuilder& GraphBuilder, const FIntPoint& Resolution)
+static FRDGTextureRef CreateRectLightAtlasTexture(FRDGBuilder& GraphBuilder, const FIntPoint& Resolution, EPixelFormat AtlasFormat)
 {
 	const uint32 MipCount = FMath::Log2(float(FMath::Min(Resolution.X, Resolution.Y)));
 	return GraphBuilder.CreateTexture(FRDGTextureDesc::Create2D(
 		Resolution,
-		PF_FloatR11G11B10,
+		AtlasFormat,
 		FClearValueBinding::Transparent,
 		ETextureCreateFlags::UAV | ETextureCreateFlags::ShaderResource | ETextureCreateFlags::RenderTargetable,
 		MipCount),
@@ -1538,10 +1559,13 @@ void UpdateAtlasTexture(FRDGBuilder& GraphBuilder, const ERHIFeatureLevel::Type
 
 	// Force update by resetting the atlas layout
 	static int32 CachedMaxAtlasResolution = CVarRectLightTextureResolution.GetValueOnRenderThread();
-	const bool bForceUpdate = CVarRectLighForceUpdate.GetValueOnRenderThread() > 0 || CachedMaxAtlasResolution != CVarRectLightTextureResolution.GetValueOnRenderThread();
+	static EPixelFormat CachedAtlasFormat = GetRectLightAtlasFormat();
+	const EPixelFormat AtlasFormat = GetRectLightAtlasFormat();
+	const bool bForceUpdate = CVarRectLighForceUpdate.GetValueOnRenderThread() > 0 || CachedMaxAtlasResolution != CVarRectLightTextureResolution.GetValueOnRenderThread() || CachedAtlasFormat != AtlasFormat;
 	if (bForceUpdate)
 	{
 		CachedMaxAtlasResolution = CVarRectLightTextureResolution.GetValueOnRenderThread();
+		CachedAtlasFormat = AtlasFormat;
 		GRectLightTextureManager.bHasPendingAdds = true;
 		GRectLightTextureManager.AtlasLayout = FAtlasLayout(FIntPoint(CachedMaxAtlasResolution, CachedMaxAtlasResolution));
 		for (FAtlasSlot& Slot : GRectLightTextureManager.AtlasSlots)
@@ -1604,7 +1628,7 @@ void UpdateAtlasTexture(FRDGBuilder& GraphBuilder, const ERHIFeatureLevel::Type
 			const bool bRecreateAtlasTexture = GRectLightTextureManager.AtlasTexture == nullptr || (CopySlots.Num() == 0 && GRectLightTextureManager.AtlasTexture->GetDesc().Extent != GRectLightTextureManager.AtlasLayout.AtlasResolution);
 			if (bRecreateAtlasTexture)
 			{
-				AtlasTexture = CreateRectLightAtlasTexture(GraphBuilder, GRectLightTextureManager.AtlasLayout.AtlasResolution);
+				AtlasTexture = CreateRectLightAtlasTexture(GraphBuilder, GRectLightTextureManager.AtlasLayout.AtlasResolution, AtlasFormat);
 				bNeedExtraction = true;
 				// Sanity check
 				check(CopySlots.Num() == 0);
@@ -1617,7 +1641,7 @@ void UpdateAtlasTexture(FRDGBuilder& GraphBuilder, const ERHIFeatureLevel::Type
 			// 2.1 Copy slots from previous to new atlas texture
 			if (CopySlots.Num() > 0)
 			{
-				FRDGTextureRef NewAtlasTexture = CreateRectLightAtlasTexture(GraphBuilder, GRectLightTextureManager.AtlasLayout.AtlasResolution);
+				FRDGTextureRef NewAtlasTexture = CreateRectLightAtlasTexture(GraphBuilder, GRectLightTextureManager.AtlasLayout.AtlasResolution, AtlasFormat);
 				bNeedExtraction = true;
 
 				CopySlotsPass(GraphBuilder, ShaderMap, CopySlots, AtlasTexture, NewAtlasTexture);
