HeterogeneousVolumes: Initial holdout implementation

When the "Enable alpha channel support in post processing" project setting  is selected, the Heterogeneous Volume Component "holdout" parameter will be respected in the alpha channel.

diff --git a/Engine/Shaders/Private/HeterogeneousVolumes/HeterogeneousVolumesComposite.usf b/Engine/Shaders/Private/HeterogeneousVolumes/HeterogeneousVolumesComposite.usf
--- a/Engine/Shaders/Private/HeterogeneousVolumes/HeterogeneousVolumesComposite.usf
+++ b/Engine/Shaders/Private/HeterogeneousVolumes/HeterogeneousVolumesComposite.usf
@@ -3,6 +3,7 @@
 #include "../Common.ush"
 
 Texture2D<float4> HeterogeneousVolumeRadiance;
+Texture2D<float> HeterogeneousVolumeHoldout;
 RWTexture2D<float4> RWColorTexture;
 int DownsampleFactor;
 
@@ -22,7 +23,16 @@ void HeterogeneousVolumesCompositeCS(
 	float4 Result;
 	float HeterogeneousVolumeTransmittance = HeterogeneousVolumeRadiance[PixelCoord].a;
 	Result.rgb = HeterogeneousVolumeRadiance[PixelCoord].rgb + RWColorTexture[ViewPixelCoord].rgb * HeterogeneousVolumeTransmittance;
+#if POST_PROCESS_ALPHA
+	if (View.RenderingReflectionCaptureMask == 0.0f)
+	{
+		float SurfaceOpacity = RWColorTexture[ViewPixelCoord].a * HeterogeneousVolumeTransmittance;
+		float VolumeOpacity = HeterogeneousVolumeHoldout[PixelCoord];
+		Result.a = SurfaceOpacity + VolumeOpacity;
+	}
+#else
 	Result.a = RWColorTexture[ViewPixelCoord].a * HeterogeneousVolumeTransmittance;
+#endif
 
 	RWColorTexture[ViewPixelCoord] = Result;
 }
diff --git a/Engine/Shaders/Private/HeterogeneousVolumes/HeterogeneousVolumesLiveShadingPipeline.usf b/Engine/Shaders/Private/HeterogeneousVolumes/HeterogeneousVolumesLiveShadingPipeline.usf
--- a/Engine/Shaders/Private/HeterogeneousVolumes/HeterogeneousVolumesLiveShadingPipeline.usf
+++ b/Engine/Shaders/Private/HeterogeneousVolumes/HeterogeneousVolumesLiveShadingPipeline.usf
@@ -70,6 +70,7 @@ int3 VoxelMin;
 int3 VoxelMax;
 
 // Lighting
+int bHoldout;
 int bApplyEmissionAndTransmittance;
 int bApplyDirectLighting;
 int bApplyShadowTransmittance;
@@ -100,6 +101,7 @@ int DownsampleFactor;
 // Output
 RWTexture2D<float4> RWLightingTexture;
 RWTexture2D<float4> RWVelocityTexture;
+RWTexture2D<float> RWHoldoutTexture;
 
 struct FDebugOutput
 {
@@ -595,12 +597,26 @@ void RenderSingleScatteringWithLiveShading(
 #endif
 
 		// Radiance
+#if POST_PROCESS_ALPHA
+		if (bHoldout && View.RenderingReflectionCaptureMask == 0.0f)
+		{
+			Radiance *= 0;
+		}
+#endif
 		RWLightingTexture[PixelCoord].rgb += Radiance * View.PreExposure;
+
 		if (bApplyEmissionAndTransmittance)
 		{
 			float Opacity = Luminance(1.0 - Transmittance);
 			RWLightingTexture[PixelCoord].a = Luminance(Transmittance);
 
+#if POST_PROCESS_ALPHA
+			if (bHoldout && View.RenderingReflectionCaptureMask == 0.0f)
+			{
+				RWHoldoutTexture[PixelCoord] += Luminance(PrevTransmittance) - Luminance(Transmittance);
+			}
+#endif
+
 #if DIM_WRITE_VELOCITY
 			bool bWriteVelocity =
 				(ScatterHitT < HitT.y) &&
diff --git a/Engine/Source/Runtime/Engine/Private/Components/HeterogeneousVolumeComponent.cpp b/Engine/Source/Runtime/Engine/Private/Components/HeterogeneousVolumeComponent.cpp
--- a/Engine/Source/Runtime/Engine/Private/Components/HeterogeneousVolumeComponent.cpp
+++ b/Engine/Source/Runtime/Engine/Private/Components/HeterogeneousVolumeComponent.cpp
@@ -107,6 +107,7 @@ FHeterogeneousVolumeSceneProxy::FHeterogeneousVolumeSceneProxy(UHeterogeneousVol
 	HeterogeneousVolumeData.LightingDownsampleFactor = InComponent->LightingDownsampleFactor;
 	HeterogeneousVolumeData.MipBias = InComponent->StreamingMipBias;
 	HeterogeneousVolumeData.bPivotAtCentroid = InComponent->bPivotAtCentroid;
+	HeterogeneousVolumeData.bHoldout = InComponent->bHoldout;
 
 	// Initialize vertex buffer data for a quad
 	StaticMeshVertexBuffers.PositionVertexBuffer.Init(4);
diff --git a/Engine/Source/Runtime/Engine/Public/HeterogeneousVolumeInterface.h b/Engine/Source/Runtime/Engine/Public/HeterogeneousVolumeInterface.h
--- a/Engine/Source/Runtime/Engine/Public/HeterogeneousVolumeInterface.h
+++ b/Engine/Source/Runtime/Engine/Public/HeterogeneousVolumeInterface.h
@@ -34,6 +34,9 @@ public:
 	virtual float GetLightingDownsampleFactor() const = 0;
 	virtual float GetMipBias() const = 0;
 
+	// Rendering
+	virtual bool IsHoldout() const = 0;
+
 	// Debug
 	virtual FString GetReadableName() const = 0;
 };
@@ -53,6 +56,8 @@ public:
 		, ShadowBiasFactor(0.0)
 		, LightingDownsampleFactor(1.0)
 		, MipBias(0.0)
+		, bPivotAtCentroid(false)
+		, bHoldout(false)
 	{}
 
 	FHeterogeneousVolumeData(const FPrimitiveSceneProxy* SceneProxy, FString Name)
@@ -66,6 +71,7 @@ public:
 		, LightingDownsampleFactor(1.0)
 		, MipBias(0.0)
 		, bPivotAtCentroid(false)
+		, bHoldout(false)
 #if ACTOR_HAS_LABELS
 		, ReadableName(Name)
 #endif // ACTOR_HAS_LABELS
@@ -93,6 +99,9 @@ public:
 	virtual float GetLightingDownsampleFactor() const { return LightingDownsampleFactor; }
 	virtual float GetMipBias() const { return MipBias; }
 
+	// Rendering
+	virtual bool IsHoldout() const { return bHoldout; }
+
 	const FPrimitiveSceneProxy* PrimitiveSceneProxy;
 	FMatrix InstanceToLocal;
 	FIntVector VoxelResolution;
@@ -103,6 +112,7 @@ public:
 	float LightingDownsampleFactor;
 	float MipBias;
 	bool bPivotAtCentroid;
+	bool bHoldout;
 
 #if ACTOR_HAS_LABELS
 	FString ReadableName;
diff --git a/Engine/Source/Runtime/Renderer/Private/HeterogeneousVolumes/HeterogeneousVolumes.cpp b/Engine/Source/Runtime/Renderer/Private/HeterogeneousVolumes/HeterogeneousVolumes.cpp
--- a/Engine/Source/Runtime/Renderer/Private/HeterogeneousVolumes/HeterogeneousVolumes.cpp
+++ b/Engine/Source/Runtime/Renderer/Private/HeterogeneousVolumes/HeterogeneousVolumes.cpp
@@ -4,6 +4,7 @@
 #include "HeterogeneousVolumeInterface.h"
 
 #include "PixelShaderUtils.h"
+#include "PostProcess/PostProcessing.h"
 #include "RayTracingDefinitions.h"
 #include "RayTracingInstance.h"
 #include "RayTracingInstanceBufferUtil.h"
@@ -612,6 +613,11 @@ namespace HeterogeneousVolumes
 		LightingCacheResolution.Z = FMath::Clamp(LightingCacheResolution.Z, 1, 512);
 		return LightingCacheResolution;
 	}
+
+	bool IsHoldout(const IHeterogeneousVolumeInterface* HeterogeneousVolumeInterface)
+	{
+		return IsPostProcessingWithAlphaChannelSupported() && HeterogeneousVolumeInterface->IsHoldout();
+	}
 }
 
 bool ShouldBuildVoxelGrids(const FScene* Scene)
@@ -776,7 +782,8 @@ void FDeferredShadingSceneRenderer::RenderHeterogeneousVolumes(
 	TRDGUniformBufferRef<FOrthoVoxelGridUniformBufferParameters> OrthoGridUniformBuffer = HeterogeneousVolumes::GetOrthoVoxelGridUniformBuffer(GraphBuilder, Views[0].ViewState);
 	TRDGUniformBufferRef<FFrustumVoxelGridUniformBufferParameters> FrustumGridUniformBuffer = HeterogeneousVolumes::GetFrustumVoxelGridUniformBuffer(GraphBuilder, Views[0].ViewState);
 
-	FRDGTextureRef HeterogeneousVolumeRadiance = nullptr;
+	FRDGTextureRef HeterogeneousVolumeRadiance = GSystemTextures.GetBlackDummy(GraphBuilder);
+	FRDGTextureRef HeterogeneousVolumeHoldout = GSystemTextures.GetBlackDummy(GraphBuilder);
 	if (ShouldRenderHeterogeneousVolumesForAnyView(Views))
 	{
 		FRDGTextureDesc Desc = SceneTextures.Color.Target->Desc;
@@ -784,6 +791,13 @@ void FDeferredShadingSceneRenderer::RenderHeterogeneousVolumes(
 		Desc.Flags &= ~(TexCreate_FastVRAM);
 		HeterogeneousVolumeRadiance = GraphBuilder.CreateTexture(Desc, TEXT("HeterogeneousVolumes"));
 		AddClearUAVPass(GraphBuilder, GraphBuilder.CreateUAV(HeterogeneousVolumeRadiance), FLinearColor::Black);
+
+		if (IsPostProcessingWithAlphaChannelSupported())
+		{
+			Desc.Format = PF_R8;
+			HeterogeneousVolumeHoldout = GraphBuilder.CreateTexture(Desc, TEXT("HeterogeneousVolume.Holdout"));
+			AddClearUAVPass(GraphBuilder, GraphBuilder.CreateUAV(HeterogeneousVolumeHoldout), FLinearColor::Black);
+		}
 	}
 
 	for (int32 ViewIndex = 0; ViewIndex < Views.Num(); ViewIndex++)
@@ -939,13 +953,15 @@ void FDeferredShadingSceneRenderer::RenderHeterogeneousVolumes(
 							// Transmittance accleration
 							LightingCacheTexture,
 							// Output
-							HeterogeneousVolumeRadiance
+							HeterogeneousVolumeRadiance,
+							HeterogeneousVolumeHoldout
 						);
 					}
 				}
 			}
 
 			View.HeterogeneousVolumeRadiance = HeterogeneousVolumeRadiance;
+			View.HeterogeneousVolumeHoldout = HeterogeneousVolumeHoldout;
 		}
 	}
 }
@@ -961,6 +977,7 @@ class FHeterogeneousVolumesCompositeCS : public FGlobalShader
 
 		// Volume data
 		SHADER_PARAMETER_RDG_TEXTURE(Texture2D<float4>, HeterogeneousVolumeRadiance)
+		SHADER_PARAMETER_RDG_TEXTURE(Texture2D<float>, HeterogeneousVolumeHoldout)
 
 		// Dispatch data
 		SHADER_PARAMETER(FIntVector, GroupCount)
@@ -1021,6 +1038,7 @@ void FDeferredShadingSceneRenderer::CompositeHeterogeneousVolumes(
 				PassParameters->View = View.ViewUniformBuffer;
 				// Volume data
 				PassParameters->HeterogeneousVolumeRadiance = View.HeterogeneousVolumeRadiance;
+				PassParameters->HeterogeneousVolumeHoldout = View.HeterogeneousVolumeHoldout;
 				// Dispatch data
 				PassParameters->GroupCount = GroupCount;
 				PassParameters->DownsampleFactor = HeterogeneousVolumes::GetDownsampleFactor();
diff --git a/Engine/Source/Runtime/Renderer/Private/HeterogeneousVolumes/HeterogeneousVolumes.h b/Engine/Source/Runtime/Renderer/Private/HeterogeneousVolumes/HeterogeneousVolumes.h
--- a/Engine/Source/Runtime/Renderer/Private/HeterogeneousVolumes/HeterogeneousVolumes.h
+++ b/Engine/Source/Runtime/Renderer/Private/HeterogeneousVolumes/HeterogeneousVolumes.h
@@ -131,6 +131,7 @@ namespace HeterogeneousVolumes
 	float CalcLOD(const FSceneView& View, const IHeterogeneousVolumeInterface* HeterogeneousVolume);
 	float CalcLODFactor(const FSceneView& View, const IHeterogeneousVolumeInterface* HeterogeneousVolume);
 	float CalcLODFactor(float LOD);
+	bool IsHoldout(const IHeterogeneousVolumeInterface* HeterogeneousVolumeInterface);
 
 	const FProjectedShadowInfo* GetProjectedShadowInfo(const FVisibleLightInfo* VisibleLightInfo, int32 ShadowIndex);
 	bool IsDynamicShadow(const FVisibleLightInfo* VisibleLightInfo);
@@ -512,7 +513,8 @@ void RenderWithLiveShading(
 	// Transmittance acceleration
 	FRDGTextureRef LightingCacheTexture,
 	// Output
-	FRDGTextureRef& HeterogeneousVolumeRadiance
+	FRDGTextureRef& HeterogeneousVolumeRadiance,
+	FRDGTextureRef& HeterogeneousVolumeHoldout
 );
 
 void RenderWithPreshading(
diff --git a/Engine/Source/Runtime/Renderer/Private/HeterogeneousVolumes/HeterogeneousVolumesAmbientOcclusionPipeline.cpp b/Engine/Source/Runtime/Renderer/Private/HeterogeneousVolumes/HeterogeneousVolumesAmbientOcclusionPipeline.cpp
--- a/Engine/Source/Runtime/Renderer/Private/HeterogeneousVolumes/HeterogeneousVolumesAmbientOcclusionPipeline.cpp
+++ b/Engine/Source/Runtime/Renderer/Private/HeterogeneousVolumes/HeterogeneousVolumesAmbientOcclusionPipeline.cpp
@@ -737,7 +737,8 @@ void RenderAmbientOcclusionWithLiveShading(
 	FRDGTextureRef& AmbientOcclusionTexture
 )
 {
-	if (!HeterogeneousVolumes::EnableAmbientOcclusion() || !HeterogeneousVolumes::UseIndirectLighting())
+	bool bRenderAmbientOcclusion = HeterogeneousVolumes::EnableAmbientOcclusion() && HeterogeneousVolumes::UseIndirectLighting();
+	if (!bRenderAmbientOcclusion || HeterogeneousVolumeInterface->IsHoldout())
 	{
 		FRDGTextureDesc AmbientOcclusionDesc = FRDGTextureDesc::Create3D(
 			FIntVector(1),
diff --git a/Engine/Source/Runtime/Renderer/Private/HeterogeneousVolumes/HeterogeneousVolumesLiveShadingPipeline.cpp b/Engine/Source/Runtime/Renderer/Private/HeterogeneousVolumes/HeterogeneousVolumesLiveShadingPipeline.cpp
--- a/Engine/Source/Runtime/Renderer/Private/HeterogeneousVolumes/HeterogeneousVolumesLiveShadingPipeline.cpp
+++ b/Engine/Source/Runtime/Renderer/Private/HeterogeneousVolumes/HeterogeneousVolumesLiveShadingPipeline.cpp
@@ -6,6 +6,7 @@
 #include "LightRendering.h"
 #include "LocalVertexFactory.h"
 #include "MeshPassUtils.h"
+#include "PostProcess/PostProcessing.h"
 #include "PixelShaderUtils.h"
 #include "RayTracingDefinitions.h"
 #include "RayTracingInstance.h"
@@ -344,6 +345,7 @@ class FRenderSingleScatteringWithLiveShadingCS : public FMeshMaterialShader
 		SHADER_PARAMETER_STRUCT_REF(FBlueNoise, BlueNoise)
 
 		// Light data
+		SHADER_PARAMETER(int, bHoldout)
 		SHADER_PARAMETER(int, bApplyEmissionAndTransmittance)
 		SHADER_PARAMETER(int, bApplyDirectLighting)
 		SHADER_PARAMETER(int, bApplyShadowTransmittance)
@@ -401,6 +403,7 @@ class FRenderSingleScatteringWithLiveShadingCS : public FMeshMaterialShader
 
 		// Output
 		SHADER_PARAMETER_RDG_TEXTURE_UAV(RWTexture2D<float4>, RWLightingTexture)
+		SHADER_PARAMETER_RDG_TEXTURE_UAV(RWTexture2D<float>, RWHoldoutTexture)
 		SHADER_PARAMETER_RDG_TEXTURE_UAV(RWTexture2D<float4>, RWVelocityTexture)
 		//SHADER_PARAMETER_RDG_BUFFER_UAV(RWStructuredBuffer<Volumes::FDebugOutput>, RWDebugOutputBuffer)
 	END_SHADER_PARAMETER_STRUCT()
@@ -671,7 +674,7 @@ static void RenderLightingCacheWithLiveShading(
 		PassParameters->Scene = View.GetSceneUniforms().GetBuffer(GraphBuilder);
 
 		// Light data
-		check(LightSceneInfo != nullptr)
+		check(LightSceneInfo != nullptr);
 		PassParameters->bApplyEmissionAndTransmittance = bApplyEmissionAndTransmittance;
 		PassParameters->bApplyDirectLighting = bApplyDirectLighting;
 		PassParameters->bApplyShadowTransmittance = bApplyShadowTransmittance;
@@ -1025,7 +1028,8 @@ void RenderSingleScatteringWithLiveShading(
 	FRDGTextureRef LightingCacheTexture,
 	FRDGTextureRef AmbientOcclusionTexture,
 	// Output
-	FRDGTextureRef& HeterogeneousVolumeTexture
+	FRDGTextureRef& HeterogeneousVolumeRadianceTexture,
+	FRDGTextureRef& HeterogeneousVolumeHoldoutTexture
 )
 {
 	typedef FRenderSingleScatteringWithLiveShadingCS<DispatchMode> FRenderSingleScatteringWithLiveShadingDispatchTypeCS;
@@ -1073,6 +1077,7 @@ void RenderSingleScatteringWithLiveShading(
 
 		// Light data
 		float LODFactor = HeterogeneousVolumes::CalcLODFactor(View, HeterogeneousVolumeInterface);
+		PassParameters->bHoldout = HeterogeneousVolumes::IsHoldout(HeterogeneousVolumeInterface);
 		PassParameters->bApplyEmissionAndTransmittance = bApplyEmissionAndTransmittance;
 		PassParameters->bApplyDirectLighting = bApplyDirectLighting;
 		PassParameters->bApplyShadowTransmittance = bApplyShadowTransmittance;
@@ -1195,7 +1200,11 @@ void RenderSingleScatteringWithLiveShading(
 		}
 
 		// Output
-		PassParameters->RWLightingTexture = GraphBuilder.CreateUAV(HeterogeneousVolumeTexture);
+		PassParameters->RWLightingTexture = GraphBuilder.CreateUAV(HeterogeneousVolumeRadianceTexture);
+		if (IsPostProcessingWithAlphaChannelSupported())
+		{
+			PassParameters->RWHoldoutTexture = GraphBuilder.CreateUAV(HeterogeneousVolumeHoldoutTexture);
+		}
 		if (bWriteVelocity)
 		{
 			PassParameters->RWVelocityTexture = GraphBuilder.CreateUAV(SceneTextures.Velocity);
@@ -1254,7 +1263,8 @@ static void RenderWithTransmittanceVolumePipeline(
 	FRDGTextureRef LightingCacheTexture,
 	FRDGTextureRef AmbientOcclusionTexture,
 	// Output
-	FRDGTextureRef& HeterogeneousVolumeRadiance
+	FRDGTextureRef& HeterogeneousVolumeRadiance,
+	FRDGTextureRef& HeterogeneousVolumeHoldout
 )
 {
 	// Light culling
@@ -1346,7 +1356,8 @@ static void RenderWithTransmittanceVolumePipeline(
 				LightingCacheTexture,
 				AmbientOcclusionTexture,
 				// Output
-				HeterogeneousVolumeRadiance
+				HeterogeneousVolumeRadiance,
+				HeterogeneousVolumeHoldout
 			);
 		}
 		else
@@ -1375,7 +1386,8 @@ static void RenderWithTransmittanceVolumePipeline(
 				LightingCacheTexture,
 				AmbientOcclusionTexture,
 				// Output
-				HeterogeneousVolumeRadiance
+				HeterogeneousVolumeRadiance,
+				HeterogeneousVolumeHoldout
 			);
 		}
 	}
@@ -1398,67 +1410,72 @@ static void RenderWithInscatteringVolumePipeline(
 	FRDGTextureRef LightingCacheTexture,
 	FRDGTextureRef AmbientOcclusionTexture,
 	// Output
-	FRDGTextureRef& HeterogeneousVolumeRadiance
+	FRDGTextureRef& HeterogeneousVolumeRadiance,
+	FRDGTextureRef& HeterogeneousVolumeHoldout
 )
 {
-	// Light culling
-	TArray<FLightSceneInfoCompact, TInlineAllocator<64>> LightSceneInfoCompact;
-	for (auto LightIt = Scene->Lights.CreateConstIterator(); LightIt; ++LightIt)
+	bool bRenderLightingCache = !HeterogeneousVolumes::IsHoldout(HeterogeneousVolumeInterface);
+	if (bRenderLightingCache)
 	{
-		if (LightIt->AffectsPrimitive(HeterogeneousVolumeInterface->GetBounds(), HeterogeneousVolumeInterface->GetPrimitiveSceneProxy()))
+		// Light culling
+		TArray<FLightSceneInfoCompact, TInlineAllocator<64>> LightSceneInfoCompact;
+		for (auto LightIt = Scene->Lights.CreateConstIterator(); LightIt; ++LightIt)
 		{
-			LightSceneInfoCompact.Add(*LightIt);
+			if (LightIt->AffectsPrimitive(HeterogeneousVolumeInterface->GetBounds(), HeterogeneousVolumeInterface->GetPrimitiveSceneProxy()))
+			{
+				LightSceneInfoCompact.Add(*LightIt);
+			}
 		}
-	}
 
-	// Light loop:
-	int32 NumPasses = LightSceneInfoCompact.Num();
-	for (int32 PassIndex = 0; PassIndex < NumPasses; ++PassIndex)
-	{
-		bool bApplyEmissionAndTransmittance = (PassIndex == (NumPasses - 1));
-		bool bApplyDirectLighting = !LightSceneInfoCompact.IsEmpty();
-		bool bApplyShadowTransmittance = false;
-
-		uint32 LightType = 0;
-		FLightSceneInfo* LightSceneInfo = nullptr;
-		const FVisibleLightInfo* VisibleLightInfo = nullptr;
-		if (bApplyDirectLighting)
+		// Light loop:
+		int32 NumPasses = LightSceneInfoCompact.Num();
+		for (int32 PassIndex = 0; PassIndex < NumPasses; ++PassIndex)
 		{
-			LightType = LightSceneInfoCompact[PassIndex].LightType;
-			LightSceneInfo = LightSceneInfoCompact[PassIndex].LightSceneInfo;
-			check(LightSceneInfo != nullptr);
-
-			bApplyDirectLighting = (LightSceneInfo != nullptr);
-			if (LightSceneInfo)
+			bool bApplyEmissionAndTransmittance = (PassIndex == (NumPasses - 1));
+			bool bApplyDirectLighting = !LightSceneInfoCompact.IsEmpty();
+			bool bApplyShadowTransmittance = false;
+
+			uint32 LightType = 0;
+			FLightSceneInfo* LightSceneInfo = nullptr;
+			const FVisibleLightInfo* VisibleLightInfo = nullptr;
+			if (bApplyDirectLighting)
 			{
-				VisibleLightInfo = &VisibleLightInfos[LightSceneInfo->Id];
-				bApplyShadowTransmittance = LightSceneInfo->Proxy->CastsVolumetricShadow();
+				LightType = LightSceneInfoCompact[PassIndex].LightType;
+				LightSceneInfo = LightSceneInfoCompact[PassIndex].LightSceneInfo;
+				check(LightSceneInfo != nullptr);
+
+				bApplyDirectLighting = (LightSceneInfo != nullptr);
+				if (LightSceneInfo)
+				{
+					VisibleLightInfo = &VisibleLightInfos[LightSceneInfo->Id];
+					bApplyShadowTransmittance = LightSceneInfo->Proxy->CastsVolumetricShadow();
+				}
 			}
-		}
 
-		RenderLightingCacheWithLiveShading(
-			GraphBuilder,
-			// Scene data
-			Scene,
-			View,
-			SceneTextures,
-			// Light data
-			bApplyEmissionAndTransmittance,
-			bApplyDirectLighting,
-			bApplyShadowTransmittance,
-			LightType,
-			LightSceneInfo,
-			// Shadow data
-			VisibleLightInfo,
-			VirtualShadowMapArray,
-			// Object data
-			HeterogeneousVolumeInterface,
-			MaterialRenderProxy,
-			PersistentPrimitiveIndex,
-			LocalBoxSphereBounds,
-			// Output
-			LightingCacheTexture
-		);
+			RenderLightingCacheWithLiveShading(
+				GraphBuilder,
+				// Scene data
+				Scene,
+				View,
+				SceneTextures,
+				// Light data
+				bApplyEmissionAndTransmittance,
+				bApplyDirectLighting,
+				bApplyShadowTransmittance,
+				LightType,
+				LightSceneInfo,
+				// Shadow data
+				VisibleLightInfo,
+				VirtualShadowMapArray,
+				// Object data
+				HeterogeneousVolumeInterface,
+				MaterialRenderProxy,
+				PersistentPrimitiveIndex,
+				LocalBoxSphereBounds,
+				// Output
+				LightingCacheTexture
+			);
+		}
 	}
 
 	// Direct volume integrator
@@ -1497,7 +1514,8 @@ static void RenderWithInscatteringVolumePipeline(
 				LightingCacheTexture,
 				AmbientOcclusionTexture,
 				// Output
-				HeterogeneousVolumeRadiance
+				HeterogeneousVolumeRadiance,
+				HeterogeneousVolumeHoldout
 			);
 		}
 		else
@@ -1526,7 +1544,8 @@ static void RenderWithInscatteringVolumePipeline(
 				LightingCacheTexture,
 				AmbientOcclusionTexture,
 				// Output
-				HeterogeneousVolumeRadiance
+				HeterogeneousVolumeRadiance,
+				HeterogeneousVolumeHoldout
 			);
 		}
 	}
@@ -1548,7 +1567,8 @@ void RenderWithLiveShading(
 	// Transmittance acceleration
 	FRDGTextureRef LightingCacheTexture,
 	// Output
-	FRDGTextureRef& HeterogeneousVolumeRadiance
+	FRDGTextureRef& HeterogeneousVolumeRadiance,
+	FRDGTextureRef& HeterogeneousVolumeHoldout
 )
 {
 	FRDGTextureRef AmbientOcclusionTexture;
@@ -1586,7 +1606,8 @@ void RenderWithLiveShading(
 			LightingCacheTexture,
 			AmbientOcclusionTexture,
 			// Output
-			HeterogeneousVolumeRadiance
+			HeterogeneousVolumeRadiance,
+			HeterogeneousVolumeHoldout
 		);
 	}
 	else
@@ -1608,7 +1629,8 @@ void RenderWithLiveShading(
 			LightingCacheTexture,
 			AmbientOcclusionTexture,
 			// Output
-			HeterogeneousVolumeRadiance
+			HeterogeneousVolumeRadiance,
+			HeterogeneousVolumeHoldout
 		);
 	}
 }
diff --git a/Engine/Source/Runtime/Renderer/Private/SceneRendering.h b/Engine/Source/Runtime/Renderer/Private/SceneRendering.h
--- a/Engine/Source/Runtime/Renderer/Private/SceneRendering.h
+++ b/Engine/Source/Runtime/Renderer/Private/SceneRendering.h
@@ -1426,6 +1426,7 @@ public:
 	FVolumetricFogViewResources VolumetricFogResources;
 
 	FRDGTextureRef HeterogeneousVolumeRadiance = nullptr;
+	FRDGTextureRef HeterogeneousVolumeHoldout = nullptr;
 
 	// Size of the HZB's mipmap 0
 	// NOTE: the mipmap 0 is downsampled version of the depth buffer
