Add IES Textures in Mobile Forward

diff --git a/Engine/Shaders/Private/MobileLightingCommon.ush b/Engine/Shaders/Private/MobileLightingCommon.ush
--- a/Engine/Shaders/Private/MobileLightingCommon.ush
+++ b/Engine/Shaders/Private/MobileLightingCommon.ush
@@ -624,6 +624,12 @@ void AccumulateLightGridLocalLighting(const FCulledLightsGridData InLightGridDat
 
 		FDeferredLightData LightData = ConvertToDeferredLight_Mobile(LocalLight);
 
+		if (LightType == LIGHT_TYPE_SPOT || LightType == LIGHT_TYPE_POINT)
+		{
+			half Attenuation = ComputeLightProfileMultiplier(TranslatedWorldPosition, LightData.TranslatedWorldPosition, -LightData.Direction, LightData.Tangent, LightData.IESAtlasIndex);
+			LightData.Color *= Attenuation;
+		}
+
 		half4 LightAttenuation = half4(1, 1, DynamicShadowing, DynamicShadowing);
 
 		float SurfaceShadow = 0;
diff --git a/Engine/Source/Runtime/Core/Private/HAL/ConsoleManager.cpp b/Engine/Source/Runtime/Core/Private/HAL/ConsoleManager.cpp
--- a/Engine/Source/Runtime/Core/Private/HAL/ConsoleManager.cpp
+++ b/Engine/Source/Runtime/Core/Private/HAL/ConsoleManager.cpp
@@ -3765,6 +3765,13 @@ static TAutoConsoleVariable<int32> CVarMobileRenderRectLightsAsSpotLights(
 		"1: Rect Lights are rendered as Spot Lights (default)\n"),
 	ECVF_ReadOnly | ECVF_MobileShaderChange);
 
+static TAutoConsoleVariable<int32> CVarMobileForwardEnableIESProfiles(
+	TEXT("r.Mobile.Forward.EnableIESProfiles"),
+	0,
+	TEXT("0: IES profiles disabled in mobile forward (default)\n"
+		"1: IES profiles enabled in mobile forward \n"),
+	ECVF_ReadOnly | ECVF_MobileShaderChange);
+
 static TAutoConsoleVariable<int32> CVarMobileForwardEnableParticleLights(
 	TEXT("r.Mobile.Forward.EnableParticleLights"),
 	0,
diff --git a/Engine/Source/Runtime/Engine/Private/ShaderCompiler/ShaderCompiler.cpp b/Engine/Source/Runtime/Engine/Private/ShaderCompiler/ShaderCompiler.cpp
--- a/Engine/Source/Runtime/Engine/Private/ShaderCompiler/ShaderCompiler.cpp
+++ b/Engine/Source/Runtime/Engine/Private/ShaderCompiler/ShaderCompiler.cpp
@@ -3170,6 +3170,12 @@ void GlobalBeginCompileShader(
 			SET_SHADER_DEFINE(Input.Environment, USE_GLES_FBF_DEFERRED, bGLESDeferredShading ? 1 : 0);
 			SET_SHADER_DEFINE(Input.Environment, MOBILE_EXTENDED_GBUFFER, MobileUsesExtenedGBuffer((EShaderPlatform)Target.Platform) ? 1 : 0);
 		}
+		else
+		{
+			static const auto CVarEnableIESProfilesMobileForward = IConsoleManager::Get().FindTConsoleVariableDataInt(TEXT("r.Mobile.Forward.EnableIESProfiles"));
+			const int32 IESProfilesEnabled = CVarEnableIESProfilesMobileForward ? CVarEnableIESProfilesMobileForward->GetValueOnAnyThread() : 0;
+			SET_SHADER_DEFINE(Input.Environment, USE_IES_PROFILE, IESProfilesEnabled);
+		}
 
 		SET_SHADER_DEFINE(Input.Environment, USE_SCENE_DEPTH_AUX, MobileRequiresSceneDepthAux(ShaderPlatform) ? 1 : 0);
 
diff --git a/Engine/Source/Runtime/Renderer/Private/LightGridInjection.cpp b/Engine/Source/Runtime/Renderer/Private/LightGridInjection.cpp
--- a/Engine/Source/Runtime/Renderer/Private/LightGridInjection.cpp
+++ b/Engine/Source/Runtime/Renderer/Private/LightGridInjection.cpp
@@ -316,7 +316,7 @@ static uint32 PackRGB10(float In0, float In1, float In2)
 	return 
 		(uint32(FMath::Clamp(In0 * 1023u, 0u, 1023u))      )|
 		(uint32(FMath::Clamp(In1 * 1023u, 0u, 1023u)) << 10)|
-		(uint32(FMath::Clamp(In2 * 1023u, 0u, 1023u)) << 10);
+		(uint32(FMath::Clamp(In2 * 1023u, 0u, 1023u)) << 20);
 }
 
 static FVector2f PackLightColor(const FVector3f& LightColor)
@@ -387,7 +387,7 @@ static void PackLocalLightData(
 	const float IESAtlasIndex = INDEX_NONE;
 
 	// Offset IESAtlasIndex here in order to preserve INDEX_NONE = -1 after encoding
-	const uint32 SpecularScale_DiffuseScale_IESData = PackRGB10(SpecularScale, DiffuseScale, IESAtlasIndex + 1);
+	const uint32 SpecularScale_DiffuseScale_IESData = PackRGB10(SpecularScale, DiffuseScale, (IESAtlasIndex + 1) * (1.f / 1023.f));
 
 	const FVector3f LightColor = (FVector3f)SimpleLight.Color * FLightRenderParameters::GetLightExposureScale(View.GetLastEyeAdaptationExposure(), SimpleLight.InverseExposureBlend);
 	const FVector2f LightColorPacked = PackLightColor(LightColor);
@@ -429,8 +429,9 @@ static void PackLocalLightData(
 	RectPackedZ |= uint32(FMath::Clamp(LightParameters.RectLightAtlasMaxLevel, 0.f, 63.f)) << 26;			//  6 bits
 
 	// Pack specular scale and IES profile index
-	 // Offset IESAtlasIndex here in order to preserve INDEX_NONE = -1 after encoding
-	const uint32 SpecularScale_DiffuseScale_IESData = PackRGB10(LightParameters.SpecularScale, LightParameters.DiffuseScale, LightParameters.IESAtlasIndex + 1); // pack atlas id here? 16bit specular 8bit IES and 8 bit LightFunction
+	// Offset IESAtlasIndex here in order to preserve INDEX_NONE = -1 after encoding
+	// IESAtlasIndex requires scaling because PackRGB10 expects inputs to be [0:1]
+	const uint32 SpecularScale_DiffuseScale_IESData = PackRGB10(LightParameters.SpecularScale, LightParameters.DiffuseScale, (LightParameters.IESAtlasIndex + 1) * (1.f / 1023.f)); // pack atlas id here? 16bit specular 8bit IES and 8 bit LightFunction
 
 	const FVector2f LightColorPacked = PackLightColor(FVector3f(LightParameters.Color));
 
