Tonemap: Cleanup/reduce shader permutation count.

* Replaced DIM_OUTPUT_DEVICE enum with OUTPUT_DEVICE_SRGB boolean, matching CombineLUTs logic.
* Removed USE_BLOOM since it was always active.
* Removed USE_VIGNETTE and replaced with dynamic branch.

The active permutations on windows desktop now total (44 out of 256) down from (280 out of 5120) for FTonemapPS before this change. METAL_MSAA_HDR_DECODE could also be removed in a subsequent CL which would bring the active count to (22 out of 128), though abstaining here since I have no means of profiling on iOS. Further changes (removal of) POST_PROCESS_ALPHA should follow shortly.

diff --git a/Engine/Shaders/Private/PostProcessCombineLUTs.usf b/Engine/Shaders/Private/PostProcessCombineLUTs.usf
--- a/Engine/Shaders/Private/PostProcessCombineLUTs.usf
+++ b/Engine/Shaders/Private/PostProcessCombineLUTs.usf
@@ -189,6 +189,7 @@ float4 CombineLUTsCommon(float2 InUV, uint InLayerIndex)
 	float3 LUTEncodedColor = Neutral.rgb;
 	float3 LinearColor;
 	// Decode texture values as ST-2084 (Dolby PQ)
+	BRANCH
 	if (GetOutputDevice() >= TONEMAPPER_OUTPUT_ACES1000nitST2084)
 	{
 		// Since ST2084 returns linear values in nits, divide by a scale factor to convert
diff --git a/Engine/Shaders/Private/PostProcessDeviceEncodingOnly.usf b/Engine/Shaders/Private/PostProcessDeviceEncodingOnly.usf
--- a/Engine/Shaders/Private/PostProcessDeviceEncodingOnly.usf
+++ b/Engine/Shaders/Private/PostProcessDeviceEncodingOnly.usf
@@ -15,10 +15,6 @@
 #include "TonemapCommon.ush"
 #include "ScreenPass.ush"
 
-#ifndef DIM_OUTPUT_DEVICE
-	#define DIM_OUTPUT_DEVICE (TONEMAPPER_OUTPUT_sRGB)
-#endif
-
 SCREEN_PASS_TEXTURE_VIEWPORT(Color)
 SCREEN_PASS_TEXTURE_VIEWPORT(Output)
 
@@ -42,7 +38,11 @@ float ACESGamutCompression;
 
 uint GetOutputDevice()
 {
+#if OUTPUT_DEVICE_SRGB
+	return TONEMAPPER_OUTPUT_sRGB;
+#else
 	return OutputDevice;
+#endif
 }
 
 float4 SampleSceneColor(float2 SceneUV)
diff --git a/Engine/Shaders/Private/PostProcessTonemap.usf b/Engine/Shaders/Private/PostProcessTonemap.usf
--- a/Engine/Shaders/Private/PostProcessTonemap.usf
+++ b/Engine/Shaders/Private/PostProcessTonemap.usf
@@ -19,10 +19,6 @@
 #include "PixelQuadMessagePassing.ush"
 #include "ScreenPass.ush"
 
-#ifndef DIM_OUTPUT_DEVICE
-	#define DIM_OUTPUT_DEVICE (TONEMAPPER_OUTPUT_sRGB)
-#endif
-
 #define DEBUG_GRAIN_QUANTIZATION 0
 
 SCREEN_PASS_TEXTURE_VIEWPORT(Color)
@@ -87,6 +83,8 @@ float4 ScreenPosToFilmGrainTextureUV;
 float BackbufferQuantizationDithering;
 float EditorNITLevel;
 
+uint OutputDevice;
+uint OutputGamut;
 float OutputMaxLuminance;
 
 float4 BloomDirtMaskTint;
@@ -96,6 +94,15 @@ SamplerState BloomDirtMaskSampler;
 float4 LensPrincipalPointOffsetScale;
 float4 LensPrincipalPointOffsetScaleInverse;
 
+uint GetOutputDevice()
+{
+#if OUTPUT_DEVICE_SRGB
+	return TONEMAPPER_OUTPUT_sRGB;
+#else
+	return OutputDevice;
+#endif
+}
+
 half GrainFromUV(float2 GrainUV) 
 {
 	half Grain = frac(sin(GrainUV.x + GrainUV.y * 543.31) *  493013.0);
@@ -170,14 +177,17 @@ half3 ColorLookupTable( half3 LinearColor )
 {
 	float3 LUTEncodedColor;
 	// Encode as ST-2084 (Dolby PQ) values
-	#if (DIM_OUTPUT_DEVICE == TONEMAPPER_OUTPUT_ACES1000nitST2084 || DIM_OUTPUT_DEVICE == TONEMAPPER_OUTPUT_ACES2000nitST2084 || DIM_OUTPUT_DEVICE == TONEMAPPER_OUTPUT_ACES1000nitScRGB || DIM_OUTPUT_DEVICE == TONEMAPPER_OUTPUT_ACES2000nitScRGB || DIM_OUTPUT_DEVICE == TONEMAPPER_OUTPUT_LinearEXR || DIM_OUTPUT_DEVICE == TONEMAPPER_OUTPUT_NoToneCurve || DIM_OUTPUT_DEVICE == TONEMAPPER_OUTPUT_WithToneCurve)
+	BRANCH
+	if(GetOutputDevice() >= TONEMAPPER_OUTPUT_ACES1000nitST2084)
+	{
 		// ST2084 expects to receive linear values 0-10000 in nits.
 		// So the linear value must be multiplied by a scale factor to convert to nits.
 		LUTEncodedColor = LinearToST2084(LinearColor * LinearToNitsScale);
-	#else
+	}
+	else
+	{
 		LUTEncodedColor = LinToLog( LinearColor + LogToLin( 0 ) );
-
-	#endif
+	}
 
 	float3 UVW = LUTEncodedColor * LUTScale + LUTOffset;
 
@@ -341,12 +351,12 @@ float4 TonemapCommonPS(
 	}
 
 	// Sample bloom & mobile's vignette.
-	#if USE_BLOOM && FEATURE_LEVEL == FEATURE_LEVEL_ES3_1
+	#if FEATURE_LEVEL == FEATURE_LEVEL_ES3_1
 		float MobileVignette;
 	#else
 		const float MobileVignette = 1.0;
 	#endif
-	#if USE_BLOOM
+	
 	half3 Bloom;
 	{
 		float2 BloomUV;
@@ -384,14 +394,6 @@ float4 TonemapCommonPS(
 		}
 		#endif
 	}
-	#endif // !USE_BLOOM
-	
-	// Compute vignette.
-	#if USE_VIGNETTE
-		const float VignetteMask = MobileVignette * ComputeVignetteMask(Vignette, TonemapperParams.x);
-	#else
-		const float VignetteMask = MobileVignette;
-	#endif
 	
 	// Compute scene color tint.
 	#if USE_GAMMA_ONLY
@@ -474,14 +476,11 @@ float4 TonemapCommonPS(
 #if USE_GAMMA_ONLY
 	{
 		#if USE_LOCAL_EXPOSURE
-			#error Local exposure is not supporte with gamma only
+			#error Local exposure is not supported with gamma only
 		#endif
 		#if USE_SHARPEN
 			#error Sharpen is not supported with gamma only
 		#endif
-		#if USE_VIGNETTE
-			#error Vignette is not supported with gamma only
-		#endif
 
 		OutColor.rgb = pow(SceneColor.rgb * (OneOverPreExposure * GlobalExposure), InverseGamma.x);
 
@@ -489,15 +488,20 @@ float4 TonemapCommonPS(
 	}
 #else // !USE_GAMMA_ONLY
 	{
+		// Compute vignette.
+		float VignetteMask = MobileVignette;
+		
+		BRANCH
+		if (TonemapperParams.x > 0.0f)
+		{
+			VignetteMask *= ComputeVignetteMask(Vignette, TonemapperParams.x);
+		}
+		
 		// Compute the final linear color space with explicitly all the terms that should affect the pre exposure.
 		half3 FinalLinearColor = SceneColor.rgb * SceneColorTint * (OneOverPreExposure * GlobalExposure * VignetteMask * LocalExposure);
 		
 		// Apply bloom.
-		#if USE_BLOOM
-		{
-			FinalLinearColor += Bloom * (OneOverPreExposure * GlobalExposure * VignetteMask);
-		}
-		#endif
+		FinalLinearColor += Bloom * (OneOverPreExposure * GlobalExposure * VignetteMask);
 		
 		// Film grain is energy conservative so doesn't need to be taken into PreExposure
 		#if USE_FILM_GRAIN
@@ -527,23 +531,22 @@ float4 TonemapCommonPS(
 			OutColor.rgb = sRGBToLinear( OutColor.rgb );
 		#endif
 
-		#if DIM_OUTPUT_DEVICE == TONEMAPPER_OUTPUT_ACES1000nitScRGB || DIM_OUTPUT_DEVICE == TONEMAPPER_OUTPUT_ACES2000nitScRGB
-			OutColor.xyz = ST2084ToScRGB(OutColor.xyz, DIM_OUTPUT_DEVICE, OutputMaxLuminance);
-
-		#elif DIM_OUTPUT_DEVICE == TONEMAPPER_OUTPUT_LinearEXR
+		BRANCH
+		if(GetOutputDevice() == TONEMAPPER_OUTPUT_ACES1000nitScRGB || GetOutputDevice() == TONEMAPPER_OUTPUT_ACES2000nitScRGB)
+		{
+			OutColor.xyz = ST2084ToScRGB(OutColor.xyz, GetOutputDevice(), OutputMaxLuminance);
+		}
+		if(GetOutputDevice() == TONEMAPPER_OUTPUT_LinearEXR)
+		{
 			OutColor.xyz = ST2084ToLinear(OutColor.xyz);
-		#elif DIM_OUTPUT_DEVICE == TONEMAPPER_OUTPUT_NoToneCurve
-			OutColor.xyz = OutDeviceColor;
-		#elif DIM_OUTPUT_DEVICE == TONEMAPPER_OUTPUT_WithToneCurve
+		}
+		else
+		{
 			OutColor.xyz = OutDeviceColor;
-		#endif
+		}
 	}
 #endif // !USE_GAMMA_ONLY
 	
-	#if POST_PROCESS_ALPHA == 2 || (POST_PROCESS_ALPHA == 1 && DIM_OUTPUT_DEVICE == TONEMAPPER_OUTPUT_NoToneCurve)  || (POST_PROCESS_ALPHA == 1 && DIM_OUTPUT_DEVICE == TONEMAPPER_OUTPUT_WithToneCurve)
-		OutColor.a = SceneColor.a;
-	#endif
-
 	#if !USE_GAMMA_ONLY
 	{
 		// Compute uniform grain on [0;1]
@@ -584,6 +587,16 @@ float4 TonemapCommonPS(
 	}
 	#endif
 
+	#if POST_PROCESS_ALPHA
+	#if POST_PROCESS_ALPHA == 1
+		BRANCH
+		if(GetOutputDevice() == TONEMAPPER_OUTPUT_NoToneCurve || GetOutputDevice() == TONEMAPPER_OUTPUT_WithToneCurve)
+	#endif
+		{
+			OutColor.a = SceneColor.a;
+		}
+	#endif
+
 	return OutColor;
 }
 
diff --git a/Engine/Shaders/Private/TonemapCommon.ush b/Engine/Shaders/Private/TonemapCommon.ush
--- a/Engine/Shaders/Private/TonemapCommon.ush
+++ b/Engine/Shaders/Private/TonemapCommon.ush
@@ -7,7 +7,7 @@
 #include "ACES.ush"
 #include "GammaCorrectionCommon.ush"
 
-/** Values of DIM_OUTPUT_DEVICE (matching C++'s EDisplayOutputFormat) */
+/** Values of output device (matching C++'s EDisplayOutputFormat) */
 #define TONEMAPPER_OUTPUT_sRGB 0
 #define TONEMAPPER_OUTPUT_Rec709 1
 #define TONEMAPPER_OUTPUT_ExplicitGammaMapping 2
@@ -19,7 +19,7 @@
 #define TONEMAPPER_OUTPUT_NoToneCurve 8
 #define TONEMAPPER_OUTPUT_WithToneCurve 9
 
-/** Values of DIM_OUTPUT_DEVICE (matching C++'s EDisplayColorGamut) */
+/** Values of output device (matching C++'s EDisplayColorGamut) */
 #define TONEMAPPER_GAMUT_sRGB_D65 0
 #define TONEMAPPER_GAMUT_DCIP3_D65 1
 #define TONEMAPPER_GAMUT_Rec2020_D65 2
diff --git a/Engine/Source/Runtime/Renderer/Private/PostProcess/PostProcessDeviceEncodingOnly.cpp b/Engine/Source/Runtime/Renderer/Private/PostProcess/PostProcessDeviceEncodingOnly.cpp
--- a/Engine/Source/Runtime/Renderer/Private/PostProcess/PostProcessDeviceEncodingOnly.cpp
+++ b/Engine/Source/Runtime/Renderer/Private/PostProcess/PostProcessDeviceEncodingOnly.cpp
@@ -26,9 +26,9 @@ namespace
 	namespace DeviceEncodingOnlyPermutation
 	{
 		// Desktop renderer permutation dimensions.
-		class FDeviceEncodingOnlyOutputDeviceDim : SHADER_PERMUTATION_ENUM_CLASS("DIM_OUTPUT_DEVICE", EDisplayOutputFormat);
+		class FDeviceEncodingOnlyOutputDeviceSRGB : SHADER_PERMUTATION_BOOL("OUTPUT_DEVICE_SRGB");
 
-		using FDesktopDomain = TShaderPermutationDomain<FDeviceEncodingOnlyOutputDeviceDim>;
+		using FDesktopDomain = TShaderPermutationDomain<FDeviceEncodingOnlyOutputDeviceSRGB>;
 	
 	} // namespace DeviceEncodingOnlyPermutation
 } // namespace
@@ -225,7 +225,9 @@ FScreenPassTexture AddDeviceEncodingOnlyPass(FRDGBuilder& GraphBuilder, const FV
 
 	// Generate permutation vector for the desktop tonemapper.
 	DeviceEncodingOnlyPermutation::FDesktopDomain DesktopPermutationVector;
-	DesktopPermutationVector.Set<DeviceEncodingOnlyPermutation::FDeviceEncodingOnlyOutputDeviceDim>(EDisplayOutputFormat(CommonParameters.OutputDevice.OutputDevice));
+
+	const bool bOutputDeviceSRGB = (CommonParameters.OutputDevice.OutputDevice == (uint32)EDisplayOutputFormat::SDR_sRGB);
+	DesktopPermutationVector.Set<DeviceEncodingOnlyPermutation::FDeviceEncodingOnlyOutputDeviceSRGB>(bOutputDeviceSRGB);
 
 	// Override output might not support UAVs.
 	const bool bComputePass = (Output.Texture->Desc.Flags & TexCreate_UAV) == TexCreate_UAV ? View.bUseComputePasses : false;
diff --git a/Engine/Source/Runtime/Renderer/Private/PostProcess/PostProcessTonemap.cpp b/Engine/Source/Runtime/Renderer/Private/PostProcess/PostProcessTonemap.cpp
--- a/Engine/Source/Runtime/Renderer/Private/PostProcess/PostProcessTonemap.cpp
+++ b/Engine/Source/Runtime/Renderer/Private/PostProcess/PostProcessTonemap.cpp
@@ -83,18 +83,14 @@ struct FOutputLuminance
 namespace TonemapperPermutation
 {
 // Shared permutation dimensions between deferred and mobile renderer.
-class FTonemapperBloomDim          : SHADER_PERMUTATION_BOOL("USE_BLOOM");
 class FTonemapperGammaOnlyDim      : SHADER_PERMUTATION_BOOL("USE_GAMMA_ONLY");
 class FTonemapperLocalExposureDim  : SHADER_PERMUTATION_BOOL("USE_LOCAL_EXPOSURE");
-class FTonemapperVignetteDim       : SHADER_PERMUTATION_BOOL("USE_VIGNETTE");
 class FTonemapperSharpenDim        : SHADER_PERMUTATION_BOOL("USE_SHARPEN");
 class FTonemapperFilmGrainDim      : SHADER_PERMUTATION_BOOL("USE_FILM_GRAIN");
 class FTonemapperMsaaDim           : SHADER_PERMUTATION_BOOL("METAL_MSAA_HDR_DECODE");
 using FCommonDomain = TShaderPermutationDomain<
-	FTonemapperBloomDim,
 	FTonemapperGammaOnlyDim,
 	FTonemapperLocalExposureDim,
-	FTonemapperVignetteDim,
 	FTonemapperSharpenDim,
 	FTonemapperFilmGrainDim,
 	FTonemapperMsaaDim>;
@@ -110,9 +106,7 @@ bool ShouldCompileCommonPermutation(const FGlobalShaderPermutationParameters& Pa
 	// If GammaOnly, don't compile any other dimmension == true.
 	if (PermutationVector.Get<FTonemapperGammaOnlyDim>())
 	{
-		return !PermutationVector.Get<FTonemapperBloomDim>() &&
-			!PermutationVector.Get<FTonemapperLocalExposureDim>() &&
-			!PermutationVector.Get<FTonemapperVignetteDim>() &&
+		return !PermutationVector.Get<FTonemapperLocalExposureDim>() &&
 			!PermutationVector.Get<FTonemapperSharpenDim>() &&
 			!PermutationVector.Get<FTonemapperFilmGrainDim>() &&
 			!PermutationVector.Get<FTonemapperMsaaDim>();
@@ -144,8 +138,6 @@ FCommonDomain BuildCommonPermutationDomain(const FViewInfo& View, bool bGammaOnl
 	}
 
 	const FPostProcessSettings& Settings = View.FinalPostProcessSettings;
-	PermutationVector.Set<FTonemapperVignetteDim>(Settings.VignetteIntensity > 0.0f);
-	PermutationVector.Set<FTonemapperBloomDim>(Settings.BloomIntensity > 0.0);
 	PermutationVector.Set<FTonemapperLocalExposureDim>(bLocalExposure);
 	PermutationVector.Set<FTonemapperFilmGrainDim>(View.FilmGrainTexture != nullptr);
 	PermutationVector.Set<FTonemapperSharpenDim>(GetSharpenSetting(Settings) > 0.0f);
@@ -155,14 +147,14 @@ FCommonDomain BuildCommonPermutationDomain(const FViewInfo& View, bool bGammaOnl
 
 // Desktop renderer permutation dimensions.
 class FTonemapperColorFringeDim       : SHADER_PERMUTATION_BOOL("USE_COLOR_FRINGE");
-class FTonemapperOutputDeviceDim      : SHADER_PERMUTATION_ENUM_CLASS("DIM_OUTPUT_DEVICE", EDisplayOutputFormat);
+class FTonemapperOutputDeviceSRGB     : SHADER_PERMUTATION_BOOL("OUTPUT_DEVICE_SRGB");
 class FTonemapperOutputLuminance	  : SHADER_PERMUTATION_BOOL("OUTPUT_LUMINANCE");
 
 using FDesktopDomain = TShaderPermutationDomain<
 	FCommonDomain,
 	FTonemapperColorFringeDim,
 	FTonemapperOutputLuminance,
-	FTonemapperOutputDeviceDim>;
+	FTonemapperOutputDeviceSRGB>;
 
 FDesktopDomain RemapPermutation(FDesktopDomain PermutationVector, ERHIFeatureLevel::Type FeatureLevel)
 {
@@ -174,10 +166,7 @@ FDesktopDomain RemapPermutation(FDesktopDomain PermutationVector, ERHIFeatureLev
 		return PermutationVector;
 	}
 
-	// Grain jitter or intensity looks bad anyway.
-	bool bFallbackToSlowest = false;
-	bFallbackToSlowest = bFallbackToSlowest || CommonPermutationVector.Get<FTonemapperFilmGrainDim>();
-
+	bool bFallbackToSlowest = CommonPermutationVector.Get<FTonemapperFilmGrainDim>();
 	if (bFallbackToSlowest)
 	{
 		CommonPermutationVector.Set<FTonemapperFilmGrainDim>(true);
@@ -186,29 +175,8 @@ FDesktopDomain RemapPermutation(FDesktopDomain PermutationVector, ERHIFeatureLev
 		PermutationVector.Set<FTonemapperColorFringeDim>(true);
 	}
 
-	// Luminance output is only used for Contrast Adaptive Shading.
-	// Even if HW VRS support is disabled, we may still need CAS for Nanite SW VRS.
-	if (!FVariableRateShadingImageManager::IsVRSCompatibleWithOutputType(PermutationVector.Get<FTonemapperOutputDeviceDim>()))
+	if (FeatureLevel < ERHIFeatureLevel::SM5)
 	{
-		PermutationVector.Set<FTonemapperOutputLuminance>(false);
-	}
-
-	// You most likely need Bloom anyway.
-	CommonPermutationVector.Set<FTonemapperBloomDim>(true);
-
-	if (FeatureLevel >= ERHIFeatureLevel::SM5)
-	{
-		// Disabling bloom on desktop renderer is very rare, not worth compiling shader permutation without.
-		CommonPermutationVector.Set<FTonemapperBloomDim>(true);
-	}
-	else
-	{
-		// Mobile supports only sRGB and LinearNoToneCurve output
-		if (PermutationVector.Get<FTonemapperOutputDeviceDim>() != EDisplayOutputFormat::HDR_LinearNoToneCurve)
-		{
-			PermutationVector.Set<FTonemapperOutputDeviceDim>(EDisplayOutputFormat::SDR_sRGB);
-		}
-
 		// Mobile doesn't support film grain.
 		CommonPermutationVector.Set<FTonemapperFilmGrainDim>(false);
 	}
@@ -929,10 +897,11 @@ FScreenPassTexture AddTonemapPass(FRDGBuilder& GraphBuilder, const FViewInfo& Vi
 			DesktopPermutationVector.Set<TonemapperPermutation::FTonemapperColorFringeDim>(PostProcessSettings.SceneFringeIntensity > 0.01f);
 		}
 
-		DesktopPermutationVector.Set<TonemapperPermutation::FTonemapperOutputDeviceDim>(EDisplayOutputFormat(CommonParameters.OutputDevice.OutputDevice));
-
 		DesktopPermutationVector.Set<TonemapperPermutation::FTonemapperOutputLuminance>(!View.bIsMobileMultiViewEnabled && GVRSImageManager.IsVRSEnabledForFrame() && FVariableRateShadingImageManager::IsVRSCompatibleWithView(View));
 
+		const bool bOutputDeviceSRGB = (CommonParameters.OutputDevice.OutputDevice == (uint32)EDisplayOutputFormat::SDR_sRGB);
+		DesktopPermutationVector.Set<TonemapperPermutation::FTonemapperOutputDeviceSRGB>(bOutputDeviceSRGB);
+
 		DesktopPermutationVector = TonemapperPermutation::RemapPermutation(DesktopPermutationVector, View.GetFeatureLevel());
 	}
 
@@ -1074,6 +1043,7 @@ class FMobileCustomResolvePS : public FGlobalShader
 	{
 		const int UseVolumeLut = PipelineVolumeTextureLUTSupportGuaranteedAtRuntime(Parameters.Platform) ? 1 : 0;
 		OutEnvironment.SetDefine(TEXT("USE_VOLUME_LUT"), UseVolumeLut);
+		OutEnvironment.SetDefine(TEXT("OUTPUT_DEVICE_SRGB"), 1);
 	}
 };
 
