HeterogeneousVolumes: Overlapping volumes do not correctly respect SceneDepth

diff --git a/Engine/Shaders/Private/HeterogeneousVolumes/HeterogeneousVolumesLiveShadingShadows.usf b/Engine/Shaders/Private/HeterogeneousVolumes/HeterogeneousVolumesLiveShadingShadows.usf
--- a/Engine/Shaders/Private/HeterogeneousVolumes/HeterogeneousVolumesLiveShadingShadows.usf
+++ b/Engine/Shaders/Private/HeterogeneousVolumes/HeterogeneousVolumesLiveShadingShadows.usf
@@ -293,6 +293,8 @@ void AddOrReplaceSample(FAVSMSampleData Data, int PixelIndex, inout int PrevSamp
 	}
 }
 
+int CameraDownsampleFactor;
+
 [numthreads(THREADGROUP_SIZE_2D, THREADGROUP_SIZE_2D, 1)]
 void RenderVolumetricShadowMapForLightWithLiveShadingCS(
 	uint2 GroupThreadId : SV_GroupThreadID,
@@ -313,6 +315,24 @@ void RenderVolumetricShadowMapForLightWithLiveShadingCS(
 
 	// Generate a light ray
 	float2 ShadowUV = float2(PixelCoord + 0.5f) / float2(ShadowResolution);
+
+	float WorldRayLength = GetMaxTraceDistance();
+#if USE_CAMERA_SCENE_DEPTH
+	float DeviceZ = SampleDeviceZFromSceneTextures(ShadowUV);
+	#if HAS_INVERTED_Z_BUFFER
+		DeviceZ = max(0.000000000001, DeviceZ);
+	#endif // HAS_INVERTED_Z_BUFFER
+
+	// Project ray through view transform to trim ray-length by opaque depth
+	uint2 ViewPixelCoord = PixelCoord * CameraDownsampleFactor + View.ViewRectMin.xy;
+	FDFVector3 WorldViewRayOrigin = DFFastSubtract(GetTranslatedWorldCameraPosFromView(ViewPixelCoord + View.TemporalAAJitter.xy), PrimaryView.PreViewTranslation);
+	FDFVector3 WorldViewRayEnd = SvPositionToWorld(float4(ViewPixelCoord + View.TemporalAAJitter.xy, DeviceZ, 1));
+	float3 WorldViewRayDirection = DFFastSubtractDemote(WorldViewRayEnd, WorldViewRayOrigin);
+	float WorldViewRayLength = length(WorldViewRayDirection);
+
+	WorldRayLength = min(WorldRayLength, WorldViewRayLength);
+#endif
+
 	const float NearScreenZ = 1.0;
 	const float FarScreenZ = 0.0;
 
@@ -321,14 +341,10 @@ void RenderVolumetricShadowMapForLightWithLiveShadingCS(
 	float3 LightRayStart = LightRayStartAsFloat4.xyz * rcp(LightRayStartAsFloat4.w);
 	float3 LightRayEnd = LightRayEndAsFloat4.xyz * rcp(LightRayEndAsFloat4.w);
 
-	// Intersect ray with bounding volume
-	// TODO: LWC integration..
 	float3 WorldRayOrigin = DFFastSubtractDemote(LightRayStart, PrimaryView.PreViewTranslation);
 	float3 WorldRayEnd = DFFastSubtractDemote(LightRayEnd, PrimaryView.PreViewTranslation);
-
-	float3 WorldRayDirection = WorldRayEnd - WorldRayOrigin;
-	float WorldRayLength = length(WorldRayDirection);
-	WorldRayDirection /= WorldRayLength;
+	float3 WorldRayDirection = normalize(WorldRayEnd - WorldRayOrigin);
+	WorldRayEnd = WorldRayOrigin + WorldRayDirection * WorldRayLength;
 
 	float3 LocalRayOrigin = mul(float4(WorldRayOrigin, 1.0), WorldToLocal).xyz;
 	float3 LocalRayEnd = mul(float4(WorldRayEnd, 1.0), WorldToLocal).xyz;
diff --git a/Engine/Source/Runtime/Renderer/Private/HeterogeneousVolumes/HeterogeneousVolumesLiveShadingPipeline.cpp b/Engine/Source/Runtime/Renderer/Private/HeterogeneousVolumes/HeterogeneousVolumesLiveShadingPipeline.cpp
--- a/Engine/Source/Runtime/Renderer/Private/HeterogeneousVolumes/HeterogeneousVolumesLiveShadingPipeline.cpp
+++ b/Engine/Source/Runtime/Renderer/Private/HeterogeneousVolumes/HeterogeneousVolumesLiveShadingPipeline.cpp
@@ -1665,7 +1665,8 @@ class FRenderVolumetricShadowMapForLightWithLiveShadingCS : public FMeshMaterial
 	DECLARE_SHADER_TYPE(FRenderVolumetricShadowMapForLightWithLiveShadingCS, MeshMaterial);
 
 	class FUseAVSMCompression : SHADER_PERMUTATION_BOOL("USE_AVSM_COMPRESSION");
-	using FPermutationDomain = TShaderPermutationDomain<FUseAVSMCompression>;
+	class FUseCameraSceneDepth : SHADER_PERMUTATION_BOOL("USE_CAMERA_SCENE_DEPTH");
+	using FPermutationDomain = TShaderPermutationDomain<FUseAVSMCompression, FUseCameraSceneDepth>;
 
 
 	BEGIN_SHADER_PARAMETER_STRUCT(FParameters, )
@@ -1705,6 +1706,7 @@ class FRenderVolumetricShadowMapForLightWithLiveShadingCS : public FMeshMaterial
 		// Dispatch data
 		SHADER_PARAMETER(FIntVector, GroupCount)
 		SHADER_PARAMETER(int, ShadowDebugTweak)
+		SHADER_PARAMETER(int, CameraDownsampleFactor)
 
 		// Output
 		SHADER_PARAMETER_RDG_BUFFER_UAV(RWBuffer<int>, RWVolumetricShadowLinkedListAllocatorBuffer)
@@ -1873,6 +1875,7 @@ bool RenderVolumetricShadowMapForLightForHeterogeneousVolumeWithLiveShading(
 	FMatrix44f* ShadowToTranslatedWorld,
 	FIntPoint ShadowMapResolution,
 	uint32 MaxSampleCount,
+	bool bUseCameraSceneDepth,
 	// Volume
 	const FVolumetricMeshBatch& VolumetricMeshBatch,
 	// Dispatch
@@ -1954,6 +1957,7 @@ bool RenderVolumetricShadowMapForLightForHeterogeneousVolumeWithLiveShading(
 		PassParameters->GroupCount = GroupCount;
 		//PassParameters->ShadowDebugTweak = CVarHeterogeneousVolumesShadowDebugTweak.GetValueOnRenderThread();
 		PassParameters->ShadowDebugTweak = 0;
+		PassParameters->CameraDownsampleFactor = HeterogeneousVolumes::GetCameraDownsampleFactor();
 
 		// Output
 		PassParameters->RWVolumetricShadowLinkedListAllocatorBuffer = GraphBuilder.CreateUAV(VolumetricShadowLinkedListAllocatorBuffer, PF_R32_UINT);
@@ -1977,6 +1981,7 @@ bool RenderVolumetricShadowMapForLightForHeterogeneousVolumeWithLiveShading(
 
 	FRenderVolumetricShadowMapForLightWithLiveShadingCS::FPermutationDomain PermutationVector;
 	PermutationVector.Set<FRenderVolumetricShadowMapForLightWithLiveShadingCS::FUseAVSMCompression>(HeterogeneousVolumes::UseAVSMCompression());
+	PermutationVector.Set<FRenderVolumetricShadowMapForLightWithLiveShadingCS::FUseCameraSceneDepth>(bUseCameraSceneDepth);
 	TShaderRef<FRenderVolumetricShadowMapForLightWithLiveShadingCS> ComputeShader = Material.GetShader<FRenderVolumetricShadowMapForLightWithLiveShadingCS>(&FLocalVertexFactory::StaticType, PermutationVector, false);
 	if (!ComputeShader.IsNull())
 	{
@@ -2180,6 +2185,7 @@ bool RenderVolumetricShadowMapForLightWithLiveShading(
 			TEXT("HeterogeneousVolume.VolumetricShadowLinkedListBuffer")
 		);
 
+		bool bUseCameraSceneDepth = false;
 		RenderVolumetricShadowMapForLightForHeterogeneousVolumeWithLiveShading(
 			GraphBuilder,
 			SceneTextures,
@@ -2197,6 +2203,7 @@ bool RenderVolumetricShadowMapForLightWithLiveShading(
 			ShadowToTranslatedWorld,
 			ShadowMapResolution,
 			MaxSampleCount,
+			bUseCameraSceneDepth,
 			// Volume
 			*VolumeMeshBatchItr,
 			// Dispatch
@@ -2231,6 +2238,7 @@ bool RenderVolumetricShadowMapForLightWithLiveShading(
 				ShadowToTranslatedWorld,
 				ShadowMapResolution,
 				MaxSampleCount,
+				bUseCameraSceneDepth,
 				// Volume
 				*VolumeMeshBatchItr,
 				// Dispatch
@@ -2497,6 +2505,7 @@ void RenderAdaptiveVolumetricCameraMapWithLiveShading(
 		int32 LightType = 0;
 		FLightSceneInfo* LightSceneInfo = nullptr;
 		FVisibleLightInfo* VisibleLightInfo = nullptr;
+		bool bUseCameraSceneDepth = true;
 		RenderVolumetricShadowMapForLightForHeterogeneousVolumeWithLiveShading(
 			GraphBuilder,
 			SceneTextures,
@@ -2514,6 +2523,7 @@ void RenderAdaptiveVolumetricCameraMapWithLiveShading(
 			ShadowToTranslatedWorld,
 			ShadowMapResolution,
 			MaxSampleCount,
+			bUseCameraSceneDepth,
 			// Volume
 			*VolumeMeshBatchItr,
 			// Dispatch
@@ -2549,6 +2559,7 @@ void RenderAdaptiveVolumetricCameraMapWithLiveShading(
 				ShadowToTranslatedWorld,
 				ShadowMapResolution,
 				MaxSampleCount,
+				bUseCameraSceneDepth,
 				// Volume
 				*VolumeMeshBatchItr,
 				// Dispatch
