Fix sampling of uninitialized sky texture when we're doing a reflection capture.

#rnx

diff --git a/Engine/Source/Runtime/Renderer/Private/MobileShadingRenderer.cpp b/Engine/Source/Runtime/Renderer/Private/MobileShadingRenderer.cpp
--- a/Engine/Source/Runtime/Renderer/Private/MobileShadingRenderer.cpp
+++ b/Engine/Source/Runtime/Renderer/Private/MobileShadingRenderer.cpp
@@ -2309,6 +2309,12 @@ void FMobileSceneRenderer::UpdateSkyReflectionUniformBuffer(FRHICommandListBase&
 		SkyLight = Scene->SkyLight;
 	}
 
+    // Make sure we don't try to use the skylight when doing a scene capture since it might contain uninitialized data
+	if (ViewFamily.EngineShowFlags.SkyLighting == 0 && Views.Num() > 0 && Views[0].bIsReflectionCapture)
+	{
+		SkyLight = nullptr;
+	}
+
 	FMobileReflectionCaptureShaderParameters Parameters;
 	SetupMobileSkyReflectionUniformParameters(SkyLight, Parameters);
 	Scene->UniformBuffers.MobileSkyReflectionUniformBuffer.UpdateUniformBufferImmediate(RHICmdList, Parameters);
