Heterogeneous Volumes: Prevent crash when enabling overlapping volume support and using Niagara baker

diff --git a/Engine/Source/Runtime/Renderer/Private/HeterogeneousVolumes/HeterogeneousVolumesLiveShadingPipeline.cpp b/Engine/Source/Runtime/Renderer/Private/HeterogeneousVolumes/HeterogeneousVolumesLiveShadingPipeline.cpp
--- a/Engine/Source/Runtime/Renderer/Private/HeterogeneousVolumes/HeterogeneousVolumesLiveShadingPipeline.cpp
+++ b/Engine/Source/Runtime/Renderer/Private/HeterogeneousVolumes/HeterogeneousVolumesLiveShadingPipeline.cpp
@@ -865,13 +865,7 @@ TRDGUniformBufferRef<FAdaptiveVolumetricShadowMaps> CreateAdaptiveVolumetricShad
 	FAdaptiveVolumetricShadowMaps* UniformBufferParameters = GraphBuilder.AllocParameters<FAdaptiveVolumetricShadowMaps>();
 	{
 		UniformBufferParameters->AVSM = GetAdaptiveVolumetricShadowMapParametersFromUniformBuffer(HeterogeneousVolumes::GetAdaptiveVolumetricShadowMapUniformBuffer(GraphBuilder, ViewState, LightSceneInfo));
-
-		TRDGUniformBufferRef<FAdaptiveVolumetricShadowMapUniformBufferParameters> CameraAVSM = ViewState->AdaptiveVolumetricCameraMapUniformBuffer;
-		if (!HeterogeneousVolumes::SupportsOverlappingVolumes())
-		{
-			CameraAVSM = HeterogeneousVolumes::CreateEmptyAdaptiveVolumetricShadowMapUniformBuffer(GraphBuilder);
-		}
-		UniformBufferParameters->CameraAVSM = GetAdaptiveVolumetricShadowMapParametersFromUniformBuffer(CameraAVSM);
+		UniformBufferParameters->CameraAVSM = GetAdaptiveVolumetricShadowMapParametersFromUniformBuffer(HeterogeneousVolumes::GetAdaptiveVolumetricCameraMapUniformBuffer(GraphBuilder, ViewState));
 	}
 
 	return GraphBuilder.CreateUniformBuffer(UniformBufferParameters);
