Mobile: Removed SkyLight specific shader permutation.
Static SkyLights do not require it. For stationary and movable skylights r.SupportStationarySkylight has to be enabled or pre-computed lighting disabled in a project.
In a scenes without SkyLight actor - shaders will still include skylight diffuse condtribution code, however will apply a zeroed color.

diff --git a/Engine/Shaders/Private/MobileDeferredShading.usf b/Engine/Shaders/Private/MobileDeferredShading.usf
--- a/Engine/Shaders/Private/MobileDeferredShading.usf
+++ b/Engine/Shaders/Private/MobileDeferredShading.usf
@@ -108,7 +108,7 @@ void ReflectionEnvironmentSkyLighting(
 {
 	half IndirectIrradiance = GBuffer.IndirectIrradiance;
 
-#if ENABLE_DYNAMIC_SKY_LIGHT
+#if ENABLE_SKY_LIGHT
 
 	half3 N = GBuffer.WorldNormal;
 	half3 V = -CameraVector;
@@ -122,7 +122,7 @@ void ReflectionEnvironmentSkyLighting(
 
 	const bool bNeedsSeparateSubsurfaceLightAccumulation = UseSubsurfaceProfile(GBuffer.ShadingModelID);
 	LightAccumulator_Add(DirectLighting, SkyLighting, SkyLighting, 1.0f, bNeedsSeparateSubsurfaceLightAccumulation);
-#endif // ENABLE_DYNAMIC_SKY_LIGHT 
+#endif // ENABLE_SKY_LIGHT 
 
 	// IBL
 	AccumulateReflection(GBuffer
diff --git a/Engine/Source/Editor/MaterialEditor/Private/MaterialStatsCommon.cpp b/Engine/Source/Editor/MaterialEditor/Private/MaterialStatsCommon.cpp
--- a/Engine/Source/Editor/MaterialEditor/Private/MaterialStatsCommon.cpp
+++ b/Engine/Source/Editor/MaterialEditor/Private/MaterialStatsCommon.cpp
@@ -395,23 +395,16 @@ void FMaterialStatsUtils::GetRepresentativeShaderTypesAndDescriptions(TMap<FName
 	else
 	{
 		static const FName TMobileBasePassVSFNoLightMapPolicyName = TEXT("TMobileBasePassVSFNoLightMapPolicy");
-		static const FName TMobileBasePassPSFNoLightMapPolicyName[2] = {
-			TEXT("TMobileBasePassPSFNoLightMapPolicyLOCAL_LIGHTS_DISABLED"),
-			TEXT("TMobileBasePassPSFNoLightMapPolicySkyLightLOCAL_LIGHTS_DISABLED")
-		};
+		static const FName TMobileBasePassPSFNoLightMapPolicyName = TEXT("TMobileBasePassPSFNoLightMapPolicyLOCAL_LIGHTS_DISABLED");
 		
 		if (TargetMaterial->GetShadingModels().IsUnlit())
 		{
 			ShaderTypeNamesAndDescriptions.Add(FLocalVertexFactoryName)
-				.Add(FRepresentativeShaderInfo(ERepresentativeShader::StationarySurface, TMobileBasePassPSFNoLightMapPolicyName[0],
+				.Add(FRepresentativeShaderInfo(ERepresentativeShader::StationarySurface, TMobileBasePassPSFNoLightMapPolicyName,
 					TEXT("Mobile base pass shader without light map")));
 		}
 		else
 		{			
-			// Use SkyLight permutation of projects settings force it
-			static auto* CVarMobileSkyLightPermutation = IConsoleManager::Get().FindTConsoleVariableDataInt(TEXT("r.Mobile.SkyLightPermutation"));
-			const bool SkyLightId = (CVarMobileSkyLightPermutation->GetValueOnAnyThread() == 2) ? 1 : 0;
-						
 			if (IsStaticLightingAllowed() && TargetMaterial->IsUsedWithStaticLighting())
 			{
 				static auto* CVarAllowDistanceFieldShadows = IConsoleManager::Get().FindTConsoleVariableDataInt(TEXT("r.Mobile.AllowDistanceFieldShadows"));
@@ -419,30 +412,26 @@ void FMaterialStatsUtils::GetRepresentativeShaderTypesAndDescriptions(TMap<FName
 				if (bAllowDistanceFieldShadows)
 				{
 					// distance field shadows only shaders
-					static const FName TMobileBasePassPSFMobileDistanceFieldShadowsAndLQLightMapPolicyName[2] = {
-						TEXT("TMobileBasePassPSFMobileDistanceFieldShadowsAndLQLightMapPolicyLOCAL_LIGHTS_DISABLED"),
-						TEXT("TMobileBasePassPSFMobileDistanceFieldShadowsAndLQLightMapPolicySkyLightLOCAL_LIGHTS_DISABLED"),
-					};
+					static const FName TMobileBasePassPSFMobileDistanceFieldShadowsAndLQLightMapPolicyName =
+						TEXT("TMobileBasePassPSFMobileDistanceFieldShadowsAndLQLightMapPolicyLOCAL_LIGHTS_DISABLED");
 					
 					ShaderTypeNamesAndDescriptions.FindOrAdd(FLocalVertexFactoryName)
-							.Add(FRepresentativeShaderInfo(ERepresentativeShader::StationarySurface, TMobileBasePassPSFMobileDistanceFieldShadowsAndLQLightMapPolicyName[SkyLightId],
+							.Add(FRepresentativeShaderInfo(ERepresentativeShader::StationarySurface, TMobileBasePassPSFMobileDistanceFieldShadowsAndLQLightMapPolicyName,
 								TEXT("Mobile base pass shader with distance field shadows")));
 				}
 				else //no shadows & lightmapped
 				{
-					static const FName TMobileBasePassPSTLightMapPolicyLQName[2] = {
-						TEXT("TMobileBasePassPSTLightMapPolicyLQLOCAL_LIGHTS_DISABLED"),
-						TEXT("TMobileBasePassPSTLightMapPolicyLQSkyLightLOCAL_LIGHTS_DISABLED"),
-					};
+					static const FName TMobileBasePassPSTLightMapPolicyLQName =
+						TEXT("TMobileBasePassPSTLightMapPolicyLQLOCAL_LIGHTS_DISABLED");
 										
 					ShaderTypeNamesAndDescriptions.FindOrAdd(FLocalVertexFactoryName)
-						.Add(FRepresentativeShaderInfo(ERepresentativeShader::StationarySurface, TMobileBasePassPSTLightMapPolicyLQName[SkyLightId],
+						.Add(FRepresentativeShaderInfo(ERepresentativeShader::StationarySurface, TMobileBasePassPSTLightMapPolicyLQName,
 							TEXT("Mobile base pass shader with static lighting")));
 				}
 			}
 			// dynamically lit shader NoLightmapPolicy
 			ShaderTypeNamesAndDescriptions.FindOrAdd(FLocalVertexFactoryName)
-				.Add(FRepresentativeShaderInfo(ERepresentativeShader::DynamicallyLitObject, TMobileBasePassPSFNoLightMapPolicyName[SkyLightId],
+				.Add(FRepresentativeShaderInfo(ERepresentativeShader::DynamicallyLitObject, TMobileBasePassPSFNoLightMapPolicyName,
 					TEXT("Mobile base pass shader with only dynamic lighting")));
 		}
 
diff --git a/Engine/Source/Runtime/Core/Private/HAL/ConsoleManager.cpp b/Engine/Source/Runtime/Core/Private/HAL/ConsoleManager.cpp
--- a/Engine/Source/Runtime/Core/Private/HAL/ConsoleManager.cpp
+++ b/Engine/Source/Runtime/Core/Private/HAL/ConsoleManager.cpp
@@ -3706,14 +3706,6 @@ static TAutoConsoleVariable<int32> CVarMobileAllowDistanceFieldShadows(
 		 "1: Generate shader permutations to render distance field shadows from stationary directional lights. (default)"),
 	ECVF_RenderThreadSafe | ECVF_ReadOnly);
 
-static TAutoConsoleVariable<int32> CVarMobileSkyLightPermutation(
-	TEXT("r.Mobile.SkyLightPermutation"),
-	0,
-	TEXT("0: Generate both sky-light and non-skylight permutations. (default)\n"
-		 "1: Generate only non-skylight permutations.\n"
-		 "2: Generate only skylight permutations"),
-	ECVF_RenderThreadSafe | ECVF_ReadOnly);
-
 static TAutoConsoleVariable<int32> CVarMobileFloatPrecisionMode(
 	TEXT("r.Mobile.FloatPrecisionMode"),
 	0,
diff --git a/Engine/Source/Runtime/Engine/Classes/Engine/RendererSettings.h b/Engine/Source/Runtime/Engine/Classes/Engine/RendererSettings.h
--- a/Engine/Source/Runtime/Engine/Classes/Engine/RendererSettings.h
+++ b/Engine/Source/Runtime/Engine/Classes/Engine/RendererSettings.h
@@ -1114,6 +1114,7 @@ class URendererSettings : public UDeveloperSettings
 
 	/**
 	"Stationary skylight requires permutations of the basepass shaders.  Disabling will reduce the number of shader permutations required per material. Changing this setting requires restarting the editor."
+	"For a mobile renderer disabling this option does not reduce shader permutations, but may increse GPU performance. It's recommended to disable it if your project does not requiure Stationary or Movable skylights."
 	*/
 	UPROPERTY(config, EditAnywhere, Category = ShaderPermutationReduction, meta = (
 		ConsoleVariable = "r.SupportStationarySkylight", DisplayName = "Support Stationary Skylight",
diff --git a/Engine/Source/Runtime/Landscape/Private/LandscapeRender.cpp b/Engine/Source/Runtime/Landscape/Private/LandscapeRender.cpp
--- a/Engine/Source/Runtime/Landscape/Private/LandscapeRender.cpp
+++ b/Engine/Source/Runtime/Landscape/Private/LandscapeRender.cpp
@@ -3966,13 +3966,8 @@ public:
 			FName(TEXT("TLightMapDensityPSFNoLightMapPolicy")),
 
 			// Mobile
-			FName(TEXT("TMobileBasePassPSFMobileDirectionalLightAndSHIndirectPolicySkylightLOCAL_LIGHTS_DISABLED")),
 			FName(TEXT("TMobileBasePassPSFMobileDirectionalLightAndSHIndirectPolicyLOCAL_LIGHTS_DISABLED")),
 			FName(TEXT("TMobileBasePassVSFMobileDirectionalLightAndSHIndirectPolicy")),
-			FName(TEXT("TMobileBasePassPSFMobileDirectionalLightAndCSMPolicySkylightLOCAL_LIGHTS_DISABLED")),
-			FName(TEXT("TMobileBasePassPSFMobileDirectionalLightAndCSMPolicyLOCAL_LIGHTS_DISABLED")),
-			FName(TEXT("TMobileBasePassVSFMobileDirectionalLightAndCSMPolicy")),
-			FName(TEXT("TMobileBasePassPSFNoLightMapPolicySkylightLOCAL_LIGHTS_DISABLED")),
 			FName(TEXT("TMobileBasePassPSFNoLightMapPolicyLOCAL_LIGHTS_DISABLED")),
 			FName(TEXT("TMobileBasePassVSFNoLightMapPolicy")),
 
@@ -4028,21 +4023,18 @@ public:
 			FName(TEXT("FDebugViewModePS")),
 
 			// Mobile
-			FName(TEXT("TMobileBasePassPSFMobileMovableDirectionalLightCSMWithLightmapPolicySkylightLOCAL_LIGHTS_DISABLED")),
 			FName(TEXT("TMobileBasePassPSFMobileMovableDirectionalLightCSMWithLightmapPolicyLOCAL_LIGHTS_DISABLED")),
 			FName(TEXT("TMobileBasePassVSFMobileMovableDirectionalLightCSMWithLightmapPolicy")),
-			FName(TEXT("TMobileBasePassPSFMobileMovableDirectionalLightWithLightmapPolicySkylightLOCAL_LIGHTS_DISABLED")),
 			FName(TEXT("TMobileBasePassPSFMobileMovableDirectionalLightWithLightmapPolicyLOCAL_LIGHTS_DISABLED")),
 			FName(TEXT("TMobileBasePassVSFMobileMovableDirectionalLightWithLightmapPolicy")),
-			FName(TEXT("TMobileBasePassPSFMobileDistanceFieldShadowsLightMapAndCSMLightingPolicySkylightLOCAL_LIGHTS_DISABLED")),
 			FName(TEXT("TMobileBasePassPSFMobileDistanceFieldShadowsLightMapAndCSMLightingPolicyLOCAL_LIGHTS_DISABLED")),
 			FName(TEXT("TMobileBasePassVSFMobileDistanceFieldShadowsLightMapAndCSMLightingPolicy")),
-			FName(TEXT("TMobileBasePassPSFMobileDistanceFieldShadowsAndLQLightMapPolicySkylightLOCAL_LIGHTS_DISABLED")),
 			FName(TEXT("TMobileBasePassPSFMobileDistanceFieldShadowsAndLQLightMapPolicyLOCAL_LIGHTS_DISABLED")),
 			FName(TEXT("TMobileBasePassVSFMobileDistanceFieldShadowsAndLQLightMapPolicy")),
-			FName(TEXT("TMobileBasePassPSTLightMapPolicyLQSkylightLOCAL_LIGHTS_DISABLED")),
 			FName(TEXT("TMobileBasePassPSTLightMapPolicyLQLOCAL_LIGHTS_DISABLED")),
 			FName(TEXT("TMobileBasePassVSTLightMapPolicyLQ")),
+			FName(TEXT("TMobileBasePassPSFMobileDirectionalLightAndCSMPolicyLOCAL_LIGHTS_DISABLED")),
+			FName(TEXT("TMobileBasePassVSFMobileDirectionalLightAndCSMPolicy")),
 
 			FName(TEXT("TBasePassVSFCachedVolumeIndirectLightingPolicy")),
 			FName(TEXT("TBasePassPSFCachedVolumeIndirectLightingPolicy")),
diff --git a/Engine/Source/Runtime/RenderCore/Private/ReadOnlyCVARCache.cpp b/Engine/Source/Runtime/RenderCore/Private/ReadOnlyCVARCache.cpp
--- a/Engine/Source/Runtime/RenderCore/Private/ReadOnlyCVARCache.cpp
+++ b/Engine/Source/Runtime/RenderCore/Private/ReadOnlyCVARCache.cpp
@@ -23,7 +23,6 @@ bool FReadOnlyCVARCache::bMobileAllowDistanceFieldShadows = true;
 bool FReadOnlyCVARCache::bMobileEnableStaticAndCSMShadowReceivers = true;
 bool FReadOnlyCVARCache::bMobileEnableMovableLightCSMShaderCulling = true;
 bool FReadOnlyCVARCache::bMobileSupportsGPUScene = false;
-int32 FReadOnlyCVARCache::MobileSkyLightPermutationValue = 0;
 int32 FReadOnlyCVARCache::MobileEarlyZPassValue = 0;
 int32 FReadOnlyCVARCache::MobileForwardLocalLightsValue = 1;
 bool FReadOnlyCVARCache::bMobileForwardParticleLights = false;
@@ -89,7 +88,6 @@ void FReadOnlyCVARCache::Initialize()
 	const auto CVarMobileEnableStaticAndCSMShadowReceivers = IConsoleManager::Get().FindTConsoleVariableDataInt(TEXT("r.Mobile.EnableStaticAndCSMShadowReceivers"));
 	const auto CVarMobileEnableMovableLightCSMShaderCulling = IConsoleManager::Get().FindTConsoleVariableDataInt(TEXT("r.Mobile.EnableMovableLightCSMShaderCulling"));
 	const auto CVarMobileAllowDistanceFieldShadows = IConsoleManager::Get().FindTConsoleVariableDataInt(TEXT("r.Mobile.AllowDistanceFieldShadows"));
-	const auto CVarMobileSkyLightPermutation = IConsoleManager::Get().FindTConsoleVariableDataInt(TEXT("r.Mobile.SkyLightPermutation"));
 	const auto CVarMobileSupportGPUScene = IConsoleManager::Get().FindTConsoleVariableDataInt(TEXT("r.Mobile.SupportGPUScene"));
 	const auto CVarMobileForwardDecalLightingValue = IConsoleManager::Get().FindTConsoleVariableDataInt(TEXT("r.Mobile.Forward.DecalLighting"));
 	
@@ -106,7 +104,6 @@ void FReadOnlyCVARCache::Initialize()
 	bMobileAllowDistanceFieldShadows = CVarMobileAllowDistanceFieldShadows->GetValueOnAnyThread() != 0;
 	bMobileEnableStaticAndCSMShadowReceivers = CVarMobileEnableStaticAndCSMShadowReceivers->GetValueOnAnyThread() != 0;
 	bMobileEnableMovableLightCSMShaderCulling = CVarMobileEnableMovableLightCSMShaderCulling->GetValueOnAnyThread() != 0;
-	MobileSkyLightPermutationValue = CVarMobileSkyLightPermutation->GetValueOnAnyThread();
 	MobileEarlyZPassValue = MobileEarlyZPassIniValue(GMaxRHIShaderPlatform);
 	MobileForwardLocalLightsValue = MobileForwardLocalLightsIniValue(GMaxRHIShaderPlatform);
 	bMobileDeferredShadingValue = MobileDeferredShadingIniValue(GMaxRHIShaderPlatform);
diff --git a/Engine/Source/Runtime/RenderCore/Public/ReadOnlyCVARCache.h b/Engine/Source/Runtime/RenderCore/Public/ReadOnlyCVARCache.h
--- a/Engine/Source/Runtime/RenderCore/Public/ReadOnlyCVARCache.h
+++ b/Engine/Source/Runtime/RenderCore/Public/ReadOnlyCVARCache.h
@@ -77,11 +77,6 @@ struct FReadOnlyCVARCache
 		return bMobileEnableMovableLightCSMShaderCulling;
 	}
 	
-	static inline int32 MobileSkyLightPermutation()
-	{
-		return MobileSkyLightPermutationValue;
-	}
-
 	static inline int32 MobileForwardDecalLighting()
 	{
 		return MobileForwardDecalLightingValue;
@@ -153,7 +148,6 @@ private:
 	RENDERCORE_API static bool bMobileEnableStaticAndCSMShadowReceivers;
 	RENDERCORE_API static bool bMobileEnableMovableLightCSMShaderCulling;
 	RENDERCORE_API static bool bMobileSupportsGPUScene;
-	RENDERCORE_API static int32 MobileSkyLightPermutationValue;
 	RENDERCORE_API static int32 MobileEarlyZPassValue;
 	RENDERCORE_API static int32 MobileForwardLocalLightsValue;
 	RENDERCORE_API static bool bMobileForwardParticleLights;
diff --git a/Engine/Source/Runtime/Renderer/Private/DebugViewModeRendering.cpp b/Engine/Source/Runtime/Renderer/Private/DebugViewModeRendering.cpp
--- a/Engine/Source/Runtime/Renderer/Private/DebugViewModeRendering.cpp
+++ b/Engine/Source/Runtime/Renderer/Private/DebugViewModeRendering.cpp
@@ -474,7 +474,7 @@ void FDebugViewModeMeshProcessor::UpdateInstructionCount(FDebugViewModeShaderEle
 		{
 			TShaderRef<TMobileBasePassVSPolicyParamType<FUniformLightMapPolicy>> MobileVS;
 			TShaderRef<TMobileBasePassPSPolicyParamType<FUniformLightMapPolicy>> MobilePS;
-			if (MobileBasePass::GetShaders(LMP_NO_LIGHTMAP, EMobileLocalLightSetting::LOCAL_LIGHTS_DISABLED, *InBatchMaterial, InVertexFactoryType, false, MobileVS, MobilePS))
+			if (MobileBasePass::GetShaders(LMP_NO_LIGHTMAP, EMobileLocalLightSetting::LOCAL_LIGHTS_DISABLED, *InBatchMaterial, InVertexFactoryType, MobileVS, MobilePS))
 			{
 				OutShaderElementData.NumVSInstructions = MobileVS.IsValid() ? MobileVS->GetNumInstructions() : 0;
 				OutShaderElementData.NumPSInstructions = MobilePS.IsValid() ? MobilePS->GetNumInstructions() : 0;
diff --git a/Engine/Source/Runtime/Renderer/Private/EditorPrimitivesRendering.cpp b/Engine/Source/Runtime/Renderer/Private/EditorPrimitivesRendering.cpp
--- a/Engine/Source/Runtime/Renderer/Private/EditorPrimitivesRendering.cpp
+++ b/Engine/Source/Runtime/Renderer/Private/EditorPrimitivesRendering.cpp
@@ -129,7 +129,6 @@ bool FEditorPrimitivesBasePassMeshProcessor::ProcessMobileShadingPath(const FMes
 
 	const FVertexFactory* VertexFactory = MeshBatch.VertexFactory;
 	const EMobileLocalLightSetting LocalLightSetting = EMobileLocalLightSetting::LOCAL_LIGHTS_DISABLED;
-	const bool bEnableSkyLight = false;
 
 	TMeshProcessorShaders<
 		TMobileBasePassVSPolicyParamType<FUniformLightMapPolicy>,
@@ -139,7 +138,6 @@ bool FEditorPrimitivesBasePassMeshProcessor::ProcessMobileShadingPath(const FMes
 		LocalLightSetting,
 		Material,
 		VertexFactory->GetType(),
-		bEnableSkyLight,
 		BasePassShaders.VertexShader,
 		BasePassShaders.PixelShader))
 	{
diff --git a/Engine/Source/Runtime/Renderer/Private/MobileBasePass.cpp b/Engine/Source/Runtime/Renderer/Private/MobileBasePass.cpp
--- a/Engine/Source/Runtime/Renderer/Private/MobileBasePass.cpp
+++ b/Engine/Source/Runtime/Renderer/Private/MobileBasePass.cpp
@@ -114,24 +114,10 @@ void SetMobileBasePassDepthState(FMeshPassProcessorRenderState& DrawRenderState,
 	DrawRenderState.SetStencilRef(StencilValue); 
 }
 
-template <ELightMapPolicyType Policy, EMobileLocalLightSetting LocalLightSetting, EMobileTranslucentColorTransmittanceMode ThinTranslucencyFallback>
-static void AddMobileBasePassPixelShaderTypes(FMaterialShaderTypes& ShaderTypes, const bool bEnableSkyLight)
-{
-	if (bEnableSkyLight)
-	{
-		ShaderTypes.AddShaderType<TMobileBasePassPS<TUniformLightMapPolicy<Policy>, true, LocalLightSetting, ThinTranslucencyFallback>>();
-	}
-	else
-	{
-		ShaderTypes.AddShaderType<TMobileBasePassPS<TUniformLightMapPolicy<Policy>, false, LocalLightSetting, ThinTranslucencyFallback>>();
-	}
-}
-
 template <ELightMapPolicyType Policy, EMobileLocalLightSetting LocalLightSetting>
 bool GetUniformMobileBasePassShaders(
 	const FMaterial& Material, 
 	const FVertexFactoryType* VertexFactoryType, 
-	bool bEnableSkyLight,
 	EMobileTranslucentColorTransmittanceMode ColoredTransmittanceFallback,
 	TShaderRef<TMobileBasePassVSPolicyParamType<FUniformLightMapPolicy>>& VertexShader,
 	TShaderRef<TMobileBasePassPSPolicyParamType<FUniformLightMapPolicy>>& PixelShader
@@ -147,9 +133,11 @@ bool GetUniformMobileBasePassShaders(
 	{
 	default:
 	case EMobileTranslucentColorTransmittanceMode::DEFAULT:
-		AddMobileBasePassPixelShaderTypes<Policy, LocalLightSetting, EMobileTranslucentColorTransmittanceMode::DEFAULT>(ShaderTypes, bEnableSkyLight); break;
+		ShaderTypes.AddShaderType<TMobileBasePassPS<TUniformLightMapPolicy<Policy>, LocalLightSetting, EMobileTranslucentColorTransmittanceMode::DEFAULT>>();
+		break;
 	case EMobileTranslucentColorTransmittanceMode::SINGLE_SRC_BLENDING:
-		AddMobileBasePassPixelShaderTypes<Policy, LocalLightSetting, EMobileTranslucentColorTransmittanceMode::SINGLE_SRC_BLENDING>(ShaderTypes, bEnableSkyLight); break;
+		ShaderTypes.AddShaderType<TMobileBasePassPS<TUniformLightMapPolicy<Policy>, LocalLightSetting, EMobileTranslucentColorTransmittanceMode::SINGLE_SRC_BLENDING>>();
+		break;
 	}
 
 	FMaterialShaders Shaders;
@@ -168,7 +156,6 @@ bool GetMobileBasePassShaders(
 	ELightMapPolicyType LightMapPolicyType, 
 	const FMaterial& Material, 
 	const FVertexFactoryType* VertexFactoryType, 
-	bool bEnableSkyLight,
 	EMobileTranslucentColorTransmittanceMode ColoredTransmittanceFallback,
 	TShaderRef<TMobileBasePassVSPolicyParamType<FUniformLightMapPolicy>>& VertexShader,
 	TShaderRef<TMobileBasePassPSPolicyParamType<FUniformLightMapPolicy>>& PixelShader
@@ -177,21 +164,21 @@ bool GetMobileBasePassShaders(
 	switch (LightMapPolicyType)
 	{
 	case LMP_NO_LIGHTMAP:
-		return GetUniformMobileBasePassShaders<LMP_NO_LIGHTMAP, LocalLightSetting>(Material, VertexFactoryType, bEnableSkyLight, ColoredTransmittanceFallback, VertexShader, PixelShader);
+		return GetUniformMobileBasePassShaders<LMP_NO_LIGHTMAP, LocalLightSetting>(Material, VertexFactoryType, ColoredTransmittanceFallback, VertexShader, PixelShader);
 	case LMP_LQ_LIGHTMAP:
-		return GetUniformMobileBasePassShaders<LMP_LQ_LIGHTMAP, LocalLightSetting>(Material, VertexFactoryType, bEnableSkyLight, ColoredTransmittanceFallback, VertexShader, PixelShader);
+		return GetUniformMobileBasePassShaders<LMP_LQ_LIGHTMAP, LocalLightSetting>(Material, VertexFactoryType, ColoredTransmittanceFallback, VertexShader, PixelShader);
 	case LMP_MOBILE_DISTANCE_FIELD_SHADOWS_AND_LQ_LIGHTMAP:
-		return GetUniformMobileBasePassShaders<LMP_MOBILE_DISTANCE_FIELD_SHADOWS_AND_LQ_LIGHTMAP, LocalLightSetting>(Material, VertexFactoryType, bEnableSkyLight, ColoredTransmittanceFallback, VertexShader, PixelShader);
+		return GetUniformMobileBasePassShaders<LMP_MOBILE_DISTANCE_FIELD_SHADOWS_AND_LQ_LIGHTMAP, LocalLightSetting>(Material, VertexFactoryType, ColoredTransmittanceFallback, VertexShader, PixelShader);
 	case LMP_MOBILE_DISTANCE_FIELD_SHADOWS_LIGHTMAP_AND_CSM:
-		return GetUniformMobileBasePassShaders<LMP_MOBILE_DISTANCE_FIELD_SHADOWS_LIGHTMAP_AND_CSM, LocalLightSetting>(Material, VertexFactoryType, bEnableSkyLight, ColoredTransmittanceFallback, VertexShader, PixelShader);
+		return GetUniformMobileBasePassShaders<LMP_MOBILE_DISTANCE_FIELD_SHADOWS_LIGHTMAP_AND_CSM, LocalLightSetting>(Material, VertexFactoryType, ColoredTransmittanceFallback, VertexShader, PixelShader);
 	case LMP_MOBILE_DIRECTIONAL_LIGHT_CSM_AND_LIGHTMAP:
-		return GetUniformMobileBasePassShaders<LMP_MOBILE_DIRECTIONAL_LIGHT_CSM_AND_LIGHTMAP, LocalLightSetting>(Material, VertexFactoryType, bEnableSkyLight, ColoredTransmittanceFallback, VertexShader, PixelShader);
+		return GetUniformMobileBasePassShaders<LMP_MOBILE_DIRECTIONAL_LIGHT_CSM_AND_LIGHTMAP, LocalLightSetting>(Material, VertexFactoryType, ColoredTransmittanceFallback, VertexShader, PixelShader);
 	case LMP_MOBILE_DIRECTIONAL_LIGHT_AND_SH_INDIRECT:
-		return GetUniformMobileBasePassShaders<LMP_MOBILE_DIRECTIONAL_LIGHT_AND_SH_INDIRECT, LocalLightSetting>(Material, VertexFactoryType, bEnableSkyLight, ColoredTransmittanceFallback, VertexShader, PixelShader);
+		return GetUniformMobileBasePassShaders<LMP_MOBILE_DIRECTIONAL_LIGHT_AND_SH_INDIRECT, LocalLightSetting>(Material, VertexFactoryType, ColoredTransmittanceFallback, VertexShader, PixelShader);
 	case LMP_MOBILE_DIRECTIONAL_LIGHT_CSM_AND_SH_INDIRECT:
-		return GetUniformMobileBasePassShaders<LMP_MOBILE_DIRECTIONAL_LIGHT_CSM_AND_SH_INDIRECT, LocalLightSetting>(Material, VertexFactoryType, bEnableSkyLight, ColoredTransmittanceFallback, VertexShader, PixelShader);
+		return GetUniformMobileBasePassShaders<LMP_MOBILE_DIRECTIONAL_LIGHT_CSM_AND_SH_INDIRECT, LocalLightSetting>(Material, VertexFactoryType, ColoredTransmittanceFallback, VertexShader, PixelShader);
 	case LMP_MOBILE_DIRECTIONAL_LIGHT_CSM:
-		return GetUniformMobileBasePassShaders<LMP_MOBILE_DIRECTIONAL_LIGHT_CSM, LocalLightSetting>(Material, VertexFactoryType, bEnableSkyLight, ColoredTransmittanceFallback, VertexShader, PixelShader);
+		return GetUniformMobileBasePassShaders<LMP_MOBILE_DIRECTIONAL_LIGHT_CSM, LocalLightSetting>(Material, VertexFactoryType, ColoredTransmittanceFallback, VertexShader, PixelShader);
 	default:										
 		check(false);
 		return true;
@@ -203,16 +190,9 @@ bool MobileBasePass::GetShaders(
 	EMobileLocalLightSetting LocalLightSetting,
 	const FMaterial& MaterialResource,
 	const FVertexFactoryType* VertexFactoryType,
-	bool bEnableSkyLight, 
 	TShaderRef<TMobileBasePassVSPolicyParamType<FUniformLightMapPolicy>>& VertexShader,
 	TShaderRef<TMobileBasePassPSPolicyParamType<FUniformLightMapPolicy>>& PixelShader)
 {
-	bool bIsLit = (MaterialResource.GetShadingModels().IsLit());
-	if (bIsLit && !UseSkylightPermutation(bEnableSkyLight, FReadOnlyCVARCache::MobileSkyLightPermutation()))	
-	{
-		bEnableSkyLight = !bEnableSkyLight;
-	}
-
 	EMobileTranslucentColorTransmittanceMode ColoredTransmittanceFallback = EMobileTranslucentColorTransmittanceMode::DEFAULT;
 	if (MaterialRequiresColorTransmittanceBlending(MaterialResource))
 	{
@@ -228,7 +208,6 @@ bool MobileBasePass::GetShaders(
 				LightMapPolicyType,
 				MaterialResource,
 				VertexFactoryType,
-				bEnableSkyLight,
 				ColoredTransmittanceFallback,
 				VertexShader,
 				PixelShader
@@ -241,7 +220,6 @@ bool MobileBasePass::GetShaders(
 				LightMapPolicyType,
 				MaterialResource,
 				VertexFactoryType,
-				bEnableSkyLight,
 				ColoredTransmittanceFallback,
 				VertexShader,
 				PixelShader
@@ -254,7 +232,6 @@ bool MobileBasePass::GetShaders(
 				LightMapPolicyType,
 				MaterialResource,
 				VertexFactoryType,
-				bEnableSkyLight,
 				ColoredTransmittanceFallback,
 				VertexShader,
 				PixelShader
@@ -831,14 +808,6 @@ bool FMobileBasePassMeshProcessor::Process(
 	TMeshProcessorShaders<
 		TMobileBasePassVSPolicyParamType<FUniformLightMapPolicy>,
 		TMobileBasePassPSPolicyParamType<FUniformLightMapPolicy>> BasePassShaders;
-	
-	bool bEnableSkyLight = false;
-	
-	if (Scene && Scene->SkyLight)
-	{
-		// Uses bTranslucentBasePass instead of BlendMode to handle single layer water meshes.
-		bEnableSkyLight = ShadingModels.IsLit() && Scene->ShouldRenderSkylightInBasePass(bTranslucentBasePass);
-	}
 
 	EMobileLocalLightSetting LocalLightSetting = EMobileLocalLightSetting::LOCAL_LIGHTS_DISABLED;
 	if (Scene && PrimitiveSceneProxy && ShadingModels.IsLit())
@@ -857,7 +826,6 @@ bool FMobileBasePassMeshProcessor::Process(
 		LocalLightSetting,
 		MaterialResource,
 		MeshBatch.VertexFactory->GetType(),
-		bEnableSkyLight,
 		BasePassShaders.VertexShader,
 		BasePassShaders.PixelShader))
 	{
@@ -934,7 +902,6 @@ void FMobileBasePassMeshProcessor::CollectPSOInitializersForLMPolicy(
 	const FMeshPassProcessorRenderState& RESTRICT DrawRenderState,
 	const FGraphicsPipelineRenderTargetsInfo& RESTRICT RenderTargetsInfo,
 	const FMaterial& RESTRICT MaterialResource,
-	const bool bEnableSkyLight,
 	EMobileLocalLightSetting LocalLightSetting,
 	const ELightMapPolicyType LightMapPolicyType,
 	ERasterizerFillMode MeshFillMode,
@@ -951,7 +918,6 @@ void FMobileBasePassMeshProcessor::CollectPSOInitializersForLMPolicy(
 		LocalLightSetting,
 		MaterialResource,
 		VertexFactoryData.VertexFactoryType,
-		bEnableSkyLight,
 		BasePassShaders.VertexShader,
 		BasePassShaders.PixelShader))
 	{
@@ -1069,26 +1035,10 @@ void FMobileBasePassMeshProcessor::CollectPSOInitializers(const FSceneTexturesCo
 	
 	for (ELightMapPolicyType LightMapPolicyType : UniformLightMapPolicyTypes)
 	{
-		// SkyLight OFF
-		bool bEnableSkyLight = false;
-		if (MobileBasePass::UseSkylightPermutation(bEnableSkyLight, FReadOnlyCVARCache::MobileSkyLightPermutation()) || !bLitMaterial)
+		CollectPSOInitializersForLMPolicy(VertexFactoryData, DrawRenderState, RenderTargetsInfo, Material, EMobileLocalLightSetting::LOCAL_LIGHTS_DISABLED, LightMapPolicyType, MeshFillMode, MeshCullMode, (EPrimitiveType)PreCacheParams.PrimitiveType, PSOInitializers);
+		if (bUseLocalLightPermutation)
 		{
-			CollectPSOInitializersForLMPolicy(VertexFactoryData, DrawRenderState, RenderTargetsInfo, Material, bEnableSkyLight, EMobileLocalLightSetting::LOCAL_LIGHTS_DISABLED, LightMapPolicyType, MeshFillMode, MeshCullMode, (EPrimitiveType)PreCacheParams.PrimitiveType, PSOInitializers);
-			if (bUseLocalLightPermutation)
-			{
-				CollectPSOInitializersForLMPolicy(VertexFactoryData, DrawRenderState, RenderTargetsInfo, Material, bEnableSkyLight, LocalLightSetting, LightMapPolicyType, MeshFillMode, MeshCullMode, (EPrimitiveType)PreCacheParams.PrimitiveType, PSOInitializers);
-			}
-		}
-
-		// SkyLight ON
-		bEnableSkyLight = true;
-		if (MobileBasePass::UseSkylightPermutation(bEnableSkyLight, FReadOnlyCVARCache::MobileSkyLightPermutation()) && bLitMaterial)
-		{
-			CollectPSOInitializersForLMPolicy(VertexFactoryData, DrawRenderState, RenderTargetsInfo, Material, bEnableSkyLight, EMobileLocalLightSetting::LOCAL_LIGHTS_DISABLED, LightMapPolicyType, MeshFillMode, MeshCullMode, (EPrimitiveType)PreCacheParams.PrimitiveType, PSOInitializers);
-			if (bUseLocalLightPermutation)
-			{
-				CollectPSOInitializersForLMPolicy(VertexFactoryData, DrawRenderState, RenderTargetsInfo, Material, bEnableSkyLight, LocalLightSetting, LightMapPolicyType, MeshFillMode, MeshCullMode, (EPrimitiveType)PreCacheParams.PrimitiveType, PSOInitializers);
-			}
+			CollectPSOInitializersForLMPolicy(VertexFactoryData, DrawRenderState, RenderTargetsInfo, Material, LocalLightSetting, LightMapPolicyType, MeshFillMode, MeshCullMode, (EPrimitiveType)PreCacheParams.PrimitiveType, PSOInitializers);
 		}
 	}
 }
diff --git a/Engine/Source/Runtime/Renderer/Private/MobileBasePassRendering.cpp b/Engine/Source/Runtime/Renderer/Private/MobileBasePassRendering.cpp
--- a/Engine/Source/Runtime/Renderer/Private/MobileBasePassRendering.cpp
+++ b/Engine/Source/Runtime/Renderer/Private/MobileBasePassRendering.cpp
@@ -132,15 +132,11 @@ EMobileTranslucentColorTransmittanceMode MobileActiveTranslucentColorTransmittan
 	typedef TMobileBasePassVS< LightMapPolicyType > TMobileBasePassVS##LightMapPolicyName; \
 	IMPLEMENT_MATERIAL_SHADER_TYPE(template<>, TMobileBasePassVS##LightMapPolicyName, TEXT("/Engine/Private/MobileBasePassVertexShader.usf"), TEXT("Main"), SF_Vertex); \
 
-#define IMPLEMENT_MOBILE_SHADING_BASEPASS_LIGHTMAPPED_PIXEL_SHADER_TYPE2(LightMapPolicyType, LightMapPolicyName, LocalLightSetting, ThinTranslucencyEnum, ThinTranslucencyName) \
-	typedef TMobileBasePassPS< LightMapPolicyType, false, LocalLightSetting, EMobileTranslucentColorTransmittanceMode::ThinTranslucencyEnum > TMobileBasePassPS##LightMapPolicyName##LocalLightSetting##ThinTranslucencyName; \
-	typedef TMobileBasePassPS< LightMapPolicyType, true, LocalLightSetting, EMobileTranslucentColorTransmittanceMode::ThinTranslucencyEnum > TMobileBasePassPS##LightMapPolicyName##Skylight##LocalLightSetting##ThinTranslucencyName; \
-	IMPLEMENT_MATERIAL_SHADER_TYPE(template<>, TMobileBasePassPS##LightMapPolicyName##LocalLightSetting##ThinTranslucencyName, TEXT("/Engine/Private/MobileBasePassPixelShader.usf"), TEXT("Main"), SF_Pixel); \
-	IMPLEMENT_MATERIAL_SHADER_TYPE(template<>, TMobileBasePassPS##LightMapPolicyName##Skylight##LocalLightSetting##ThinTranslucencyName, TEXT("/Engine/Private/MobileBasePassPixelShader.usf"), TEXT("Main"), SF_Pixel);
-
 #define IMPLEMENT_MOBILE_SHADING_BASEPASS_LIGHTMAPPED_PIXEL_SHADER_TYPE(LightMapPolicyType, LightMapPolicyName, LocalLightSetting) \
-	IMPLEMENT_MOBILE_SHADING_BASEPASS_LIGHTMAPPED_PIXEL_SHADER_TYPE2(LightMapPolicyType, LightMapPolicyName, LocalLightSetting, DEFAULT,) \
-	IMPLEMENT_MOBILE_SHADING_BASEPASS_LIGHTMAPPED_PIXEL_SHADER_TYPE2(LightMapPolicyType, LightMapPolicyName, LocalLightSetting, SINGLE_SRC_BLENDING, ThinTranslGrey)
+	typedef TMobileBasePassPS< LightMapPolicyType, LocalLightSetting, EMobileTranslucentColorTransmittanceMode::DEFAULT > TMobileBasePassPS##LightMapPolicyName##LocalLightSetting; \
+	IMPLEMENT_MATERIAL_SHADER_TYPE(template<>, TMobileBasePassPS##LightMapPolicyName##LocalLightSetting, TEXT("/Engine/Private/MobileBasePassPixelShader.usf"), TEXT("Main"), SF_Pixel); \
+	typedef TMobileBasePassPS< LightMapPolicyType, LocalLightSetting, EMobileTranslucentColorTransmittanceMode::SINGLE_SRC_BLENDING > TMobileBasePassPS##LightMapPolicyName##LocalLightSetting##ThinTranslGrey; \
+	IMPLEMENT_MATERIAL_SHADER_TYPE(template<>, TMobileBasePassPS##LightMapPolicyName##LocalLightSetting##ThinTranslGrey, TEXT("/Engine/Private/MobileBasePassPixelShader.usf"), TEXT("Main"), SF_Pixel); \
 
 #define IMPLEMENT_MOBILE_SHADING_BASEPASS_LIGHTMAPPED_SHADER_TYPE(LightMapPolicyType, LightMapPolicyName) \
 	IMPLEMENT_MOBILE_SHADING_BASEPASS_LIGHTMAPPED_VERTEX_SHADER_TYPE(LightMapPolicyType, LightMapPolicyName) \
diff --git a/Engine/Source/Runtime/Renderer/Private/MobileBasePassRendering.h b/Engine/Source/Runtime/Renderer/Private/MobileBasePassRendering.h
--- a/Engine/Source/Runtime/Renderer/Private/MobileBasePassRendering.h
+++ b/Engine/Source/Runtime/Renderer/Private/MobileBasePassRendering.h
@@ -330,7 +330,6 @@ namespace MobileBasePass
 		EMobileLocalLightSetting LocalLightSetting,
 		const FMaterial& MaterialResource,
 		const FVertexFactoryType* VertexFactoryType,
-		bool bEnableSkyLight, 
 		TShaderRef<TMobileBasePassVSPolicyParamType<FUniformLightMapPolicy>>& VertexShader,
 		TShaderRef<TMobileBasePassPSPolicyParamType<FUniformLightMapPolicy>>& PixelShader);
 
@@ -340,21 +339,9 @@ namespace MobileBasePass
 
 	void SetOpaqueRenderState(FMeshPassProcessorRenderState& DrawRenderState, const FPrimitiveSceneProxy* PrimitiveSceneProxy, const FMaterial& Material, FMaterialShadingModelField ShadingModels, bool bEnableReceiveDecalOutput, bool bUsesDeferredShading);
 	void SetTranslucentRenderState(FMeshPassProcessorRenderState& DrawRenderState, const FMaterial& Material, FMaterialShadingModelField ShadingModels);
-
-	inline bool UseSkylightPermutation(bool bEnableSkyLight, int32 MobileSkyLightPermutationOptions)
-	{
-		if (bEnableSkyLight)
-		{
-			return MobileSkyLightPermutationOptions == 0 || MobileSkyLightPermutationOptions == 2;
-		}
-		else
-		{
-			return MobileSkyLightPermutationOptions == 0 || MobileSkyLightPermutationOptions == 1;
-		}
-	}
 };
 
-template< typename LightMapPolicyType, bool bEnableSkyLight, EMobileLocalLightSetting LocalLightSetting, EMobileTranslucentColorTransmittanceMode TranslucentColorTransmittanceFallback = EMobileTranslucentColorTransmittanceMode::DEFAULT>
+template< typename LightMapPolicyType, EMobileLocalLightSetting LocalLightSetting, EMobileTranslucentColorTransmittanceMode TranslucentColorTransmittanceFallback = EMobileTranslucentColorTransmittanceMode::DEFAULT>
 class TMobileBasePassPS : public TMobileBasePassPSBaseType<LightMapPolicyType>
 {
 	DECLARE_SHADER_TYPE(TMobileBasePassPS,MeshMaterial);
@@ -363,17 +350,7 @@ public:
 	static bool ShouldCompilePermutation(const FMeshMaterialShaderPermutationParameters& Parameters)
 	{		
 		// We compile the point light shader combinations based on the project settings
-		static auto* MobileSkyLightPermutationCVar = IConsoleManager::Get().FindTConsoleVariableDataInt(TEXT("r.Mobile.SkyLightPermutation"));
-		const int32 MobileSkyLightPermutationOptions = MobileSkyLightPermutationCVar->GetValueOnAnyThread();
 		const bool bIsLit = Parameters.MaterialParameters.ShadingModels.IsLit();
-		// Only compile skylight version for lit materials
-		const bool bShouldCacheBySkylight = !bEnableSkyLight || bIsLit;
-		// Only compile skylight permutations when they are enabled
-		if (bIsLit && !MobileBasePass::UseSkylightPermutation(bEnableSkyLight, MobileSkyLightPermutationOptions))
-		{
-			return false;
-		}
-		
 		const bool bDeferredShadingEnabled = IsMobileDeferredShadingEnabled(Parameters.Platform);
 		const bool bIsTranslucent = IsTranslucentBlendMode(Parameters.MaterialParameters) || Parameters.MaterialParameters.ShadingModels.HasShadingModel(MSM_SingleLayerWater);
 		const bool bMaterialUsesForwardShading = bIsLit && bIsTranslucent;
@@ -391,7 +368,6 @@ public:
 		const bool bShouldCacheByLocalLights = !bEnableLocalLights || (bIsLit && (SupportedLocalLightsType == LocalLightSetting));
 
 		return TMobileBasePassPSBaseType<LightMapPolicyType>::ShouldCompilePermutation(Parameters) && 
-				bShouldCacheBySkylight && 
 				bShouldCacheByLocalLights && 
 				ShouldCacheShaderForColorTransmittanceFallback(Parameters, TranslucentColorTransmittanceFallback);
 	}
@@ -408,11 +384,11 @@ public:
 		const bool bMaterialUsesForwardShading = bIsLit && bTranslucentMaterial;
 		// Translucent materials always support clustered shading on mobile deferred
 		const bool bForwardShading = !bDeferredShadingEnabled || bMaterialUsesForwardShading;
-		// Only stationary skylights contribute into basepass with deferred shading materials. Has to do this test as on some project configurations we dont have a separate permutation for no-skylight
-		const bool bEnableSkylightInBasePass = bEnableSkyLight && (bForwardShading || FReadOnlyCVARCache::EnableStationarySkylight());
-
+				
 		TMobileBasePassPSBaseType<LightMapPolicyType>::ModifyCompilationEnvironment(Parameters, OutEnvironment);
-		OutEnvironment.SetDefine(TEXT("ENABLE_SKY_LIGHT"), bEnableSkylightInBasePass);
+		// Only non-static skylights contribute into forward basepass
+		const bool bProjectSupportsNonStaticSkyLights = FReadOnlyCVARCache::EnableStationarySkylight() || !IsStaticLightingAllowed();
+		OutEnvironment.SetDefine(TEXT("ENABLE_SKY_LIGHT"), bIsLit && bForwardShading && bProjectSupportsNonStaticSkyLights);
 		OutEnvironment.SetDefine(TEXT("ENABLE_AMBIENT_OCCLUSION"), bForwardShading && IsMobileAmbientOcclusionEnabled(Parameters.Platform) ? 1u : 0u);
 		
 		FForwardLightingParameters::ModifyCompilationEnvironment(Parameters.Platform, OutEnvironment);
@@ -494,7 +470,6 @@ public:
 		const FMeshPassProcessorRenderState& RESTRICT DrawRenderState,
 		const FGraphicsPipelineRenderTargetsInfo& RESTRICT RenderTargetsInfo,
 		const FMaterial& RESTRICT MaterialResource,
-		const bool bRenderSkylight,
 		EMobileLocalLightSetting LocalLightSetting,
 		const ELightMapPolicyType LightMapPolicyType,
 		ERasterizerFillMode MeshFillMode,
diff --git a/Engine/Source/Runtime/Renderer/Private/MobileDeferredShadingPass.cpp b/Engine/Source/Runtime/Renderer/Private/MobileDeferredShadingPass.cpp
--- a/Engine/Source/Runtime/Renderer/Private/MobileDeferredShadingPass.cpp
+++ b/Engine/Source/Runtime/Renderer/Private/MobileDeferredShadingPass.cpp
@@ -61,7 +61,6 @@ class FMobileDirectionalLightFunctionPS : public FMaterialShader
 	class FEnableClustredReflection		: SHADER_PERMUTATION_BOOL("ENABLE_CLUSTERED_REFLECTION");
 	class FEnablePlanarReflection		: SHADER_PERMUTATION_BOOL("ENABLE_PLANAR_REFLECTION");
 	class FEnableSkyLight				: SHADER_PERMUTATION_BOOL("ENABLE_SKY_LIGHT");
-	class FEnableDynamicSkyLight		: SHADER_PERMUTATION_BOOL("ENABLE_DYNAMIC_SKY_LIGHT");
 	class FEnableCSM					: SHADER_PERMUTATION_BOOL("ENABLE_MOBILE_CSM");
 	class FShadowQuality				: SHADER_PERMUTATION_RANGE_INT("MOBILE_SHADOW_QUALITY", 1, 3); // not using Quality=0
 	class FMobileSSRQuality 			: SHADER_PERMUTATION_ENUM_CLASS("MOBILE_SSR_QUALITY", EMobileSSRQuality);
@@ -72,7 +71,6 @@ class FMobileDirectionalLightFunctionPS : public FMaterialShader
 		FEnableClustredReflection, 
 		FEnablePlanarReflection,
 		FEnableSkyLight,
-		FEnableDynamicSkyLight,
 		FEnableCSM, 
 		FShadowQuality,
 		FMobileSSRQuality>;
@@ -121,11 +119,6 @@ class FMobileDirectionalLightFunctionPS : public FMaterialShader
 			PermutationVector.Set<FEnableShadingModelSupport>(false);
 		}
 
-		if (!PermutationVector.Get<FEnableSkyLight>())
-		{
-			PermutationVector.Set<FEnableDynamicSkyLight>(false);
-		}
-
 		if (!AreMobileScreenSpaceReflectionsEnabled(Platform))
 		{
 			PermutationVector.Set<FMobileSSRQuality>(EMobileSSRQuality::Disabled);
@@ -153,12 +146,11 @@ class FMobileDirectionalLightFunctionPS : public FMaterialShader
 		return true;
 	}
 
-	static FPermutationDomain BuildPermutationVector(const FViewInfo& View, bool bInlineReflectionAndSky, bool bShadingModelSupport, bool bDynamicShadows, bool bSkyLight, bool bDynamicSkyLight, bool bPlanarReflection, EMobileSSRQuality MobileSSRQuality)
+	static FPermutationDomain BuildPermutationVector(const FViewInfo& View, bool bInlineReflectionAndSky, bool bShadingModelSupport, bool bDynamicShadows, bool bSkyLight, bool bPlanarReflection, EMobileSSRQuality MobileSSRQuality)
 	{
 		bool bUseClusteredLights = UseClusteredDeferredShading(View.GetShaderPlatform());
 		bool bClustredReflection = bInlineReflectionAndSky && (View.NumBoxReflectionCaptures + View.NumSphereReflectionCaptures) > 0;
 		bool bEnableSkyLight = bInlineReflectionAndSky && bSkyLight;
-		bool bEnableDynamicSkyLight = bInlineReflectionAndSky && bDynamicSkyLight;
 		const bool bMobileUsesShadowMaskTexture = MobileUsesShadowMaskTexture(View.GetShaderPlatform());
 		int32 ShadowQuality = bDynamicShadows && !bMobileUsesShadowMaskTexture ? (int32)GetShadowQuality() : 0;
 				
@@ -168,7 +160,6 @@ class FMobileDirectionalLightFunctionPS : public FMaterialShader
 		PermutationVector.Set<FMobileDirectionalLightFunctionPS::FEnableClustredReflection>(bClustredReflection);
 		PermutationVector.Set<FMobileDirectionalLightFunctionPS::FEnablePlanarReflection>(bPlanarReflection);
 		PermutationVector.Set<FMobileDirectionalLightFunctionPS::FEnableSkyLight>(bEnableSkyLight);
-		PermutationVector.Set<FMobileDirectionalLightFunctionPS::FEnableDynamicSkyLight>(bEnableDynamicSkyLight);
 		PermutationVector.Set<FMobileDirectionalLightFunctionPS::FEnableCSM>(ShadowQuality > 0);
 		PermutationVector.Set<FMobileDirectionalLightFunctionPS::FShadowQuality>(FMath::Clamp(ShadowQuality, 1, 3));
 		PermutationVector.Set<FMobileDirectionalLightFunctionPS::FMobileSSRQuality>(bInlineReflectionAndSky ? MobileSSRQuality : EMobileSSRQuality::Disabled);
@@ -284,7 +275,6 @@ public:
 	class FEnableClustredReflection		: SHADER_PERMUTATION_BOOL("ENABLE_CLUSTERED_REFLECTION");
 	class FEnablePlanarReflection		: SHADER_PERMUTATION_BOOL("ENABLE_PLANAR_REFLECTION");
 	class FEnableSkyLight				: SHADER_PERMUTATION_BOOL("ENABLE_SKY_LIGHT");
-	class FEnableDynamicSkyLight		: SHADER_PERMUTATION_BOOL("ENABLE_DYNAMIC_SKY_LIGHT");
 	class FMobileSSRQuality				: SHADER_PERMUTATION_ENUM_CLASS("MOBILE_SSR_QUALITY", EMobileSSRQuality);
 	
 	using FPermutationDomain = TShaderPermutationDomain<
@@ -292,7 +282,6 @@ public:
 		FEnableClustredReflection, 
 		FEnablePlanarReflection,
 		FEnableSkyLight,
-		FEnableDynamicSkyLight,
 		FMobileSSRQuality
 	>;
 
@@ -315,11 +304,6 @@ public:
 			return false;
 		}
 
-		if (!PermutationVector.Get<FEnableSkyLight>() && PermutationVector.Get<FEnableDynamicSkyLight>())
-		{
-			return false;
-		}
-
 		if (PermutationVector.Get<FMobileSSRQuality>() != EMobileSSRQuality::Disabled && !AreMobileScreenSpaceReflectionsEnabled(Parameters.Platform))
 		{
 			return false;
@@ -412,11 +396,11 @@ constexpr uint8 PassShadingModelStencilMask(bool bEnableShadingModelSupport)
 void RenderReflectionEnvironmentSkyLighting(FRHICommandList& RHICmdList, const FScene& Scene, const FViewInfo& View, const EMobileSSRQuality MobileSSRQuality)
 {
 	// Skylights with static lighting already had their diffuse contribution baked into lightmaps
-	const bool bSkyLight = Scene.SkyLight && !Scene.SkyLight->bHasStaticLighting && View.Family->EngineShowFlags.SkyLighting;
-	const bool bDynamicSkyLight = bSkyLight && !Scene.SkyLight->bWantsStaticShadowing;
+	const bool bDynamicSkyLight = Scene.SkyLight && (!Scene.SkyLight->bHasStaticLighting || !IsStaticLightingAllowed());
+	const bool bEnableSkyLight = bDynamicSkyLight && View.Family->EngineShowFlags.SkyLighting;
 	const bool bClustredReflection = (View.NumBoxReflectionCaptures + View.NumSphereReflectionCaptures) > 0;
 	const bool bPlanarReflection = Scene.GetForwardPassGlobalPlanarReflection() != nullptr;
-	if (!(bSkyLight || bClustredReflection || bPlanarReflection || MobileSSRQuality != EMobileSSRQuality::Disabled))
+	if (!(bEnableSkyLight || bClustredReflection || bPlanarReflection || MobileSSRQuality != EMobileSSRQuality::Disabled))
 	{
 		return;
 	}
@@ -428,7 +412,7 @@ void RenderReflectionEnvironmentSkyLighting(FRHICommandList& RHICmdList, const F
 	// Add to emissive in SceneColor
 	if (!bDynamicSkyLight)
 	{
-		// pre-multiply SceneColor with AO
+		// pre-multiply SceneColor with AO. Only need it for a static skylights
 		GraphicsPSOInit.BlendState = TStaticBlendState<CW_RGB, BO_Add, BF_One, BF_SourceAlpha>::GetRHI();
 	}
 	else
@@ -490,8 +474,7 @@ void RenderReflectionEnvironmentSkyLighting(FRHICommandList& RHICmdList, const F
 		PermutationVector.Set<FMobileReflectionEnvironmentSkyLightingPS::FEnableShadingModelSupport>(PassEnableSSR & (1 << PassIndex));
 		PermutationVector.Set<FMobileReflectionEnvironmentSkyLightingPS::FEnableClustredReflection>(bClustredReflection);
 		PermutationVector.Set<FMobileReflectionEnvironmentSkyLightingPS::FEnablePlanarReflection>(bPlanarReflection);
-		PermutationVector.Set<FMobileReflectionEnvironmentSkyLightingPS::FEnableSkyLight>(bSkyLight);
-		PermutationVector.Set<FMobileReflectionEnvironmentSkyLightingPS::FEnableDynamicSkyLight>(bDynamicSkyLight);
+		PermutationVector.Set<FMobileReflectionEnvironmentSkyLightingPS::FEnableSkyLight>(bEnableSkyLight);
 		PermutationVector.Set<FMobileReflectionEnvironmentSkyLightingPS::FMobileSSRQuality>((PassEnableSSR & (1 << PassIndex)) ? MobileSSRQuality : EMobileSSRQuality::Disabled);
 		TShaderMapRef<FMobileReflectionEnvironmentSkyLightingPS> PixelShader(View.ShaderMap, PermutationVector);
 		
@@ -604,8 +587,8 @@ static void RenderDirectionalLight(FRHICommandList& RHICmdList, const FScene& Sc
 	}
 
 	// Skylights with static lighting already had their diffuse contribution baked into lightmaps
-	const bool bSkyLight = Scene.SkyLight && !Scene.SkyLight->bHasStaticLighting && View.Family->EngineShowFlags.SkyLighting;
-	const bool bDynamicSkyLight = bSkyLight && !Scene.SkyLight->bWantsStaticShadowing;
+	const bool bDynamicSkyLight = Scene.SkyLight && (!Scene.SkyLight->bHasStaticLighting || !IsStaticLightingAllowed());
+	const bool bEnableSkyLight = bDynamicSkyLight && View.Family->EngineShowFlags.SkyLighting;
 	const bool bDynamicShadows = DirectionalLight.Proxy->CastsDynamicShadow() && View.Family->EngineShowFlags.DynamicShadows;
 	const bool bPlanarReflection = Scene.GetForwardPassGlobalPlanarReflection() != nullptr;
 
@@ -654,8 +637,7 @@ static void RenderDirectionalLight(FRHICommandList& RHICmdList, const FScene& Sc
 			bInlineReflectionAndSky,
 			PassEnableShadingModelSupport & (1 << PassIndex),
 			bDynamicShadows,
-			bSkyLight,
-			bDynamicSkyLight,
+			bEnableSkyLight,
 			bPlanarReflection, 
 			(PassEnableSSR & (1<<PassIndex)) ? MobileSSRQuality : EMobileSSRQuality::Disabled
 		);
diff --git a/Engine/Source/Runtime/Renderer/Private/RendererScene.cpp b/Engine/Source/Runtime/Renderer/Private/RendererScene.cpp
--- a/Engine/Source/Runtime/Renderer/Private/RendererScene.cpp
+++ b/Engine/Source/Runtime/Renderer/Private/RendererScene.cpp
@@ -6656,21 +6656,7 @@ bool FScene::ShouldRenderSkylightInBasePass(bool bIsTranslucent) const
 {
 	if (IsMobilePlatform(GetShaderPlatform()))
 	{
-		bool bRenderSkyLight = SkyLight && !SkyLight->bHasStaticLighting;
-		const bool bIsForwardShading = bIsTranslucent || !IsMobileDeferredShadingEnabled(GetShaderPlatform());
-
-		if (bIsForwardShading)
-		{
-			// Both stationary and movable skylights are applied in base pass for forward shading
-			bRenderSkyLight = bRenderSkyLight && (FReadOnlyCVARCache::EnableStationarySkylight() || !SkyLight->bWantsStaticShadowing);
-		}
-		else
-		{
-			// Only stationary skylights are applied in base pass for deferred
-			bRenderSkyLight = bRenderSkyLight && (FReadOnlyCVARCache::EnableStationarySkylight() && SkyLight->bWantsStaticShadowing);
-		}
-
-		return bRenderSkyLight;
+		return false;
 	}
 	else
 	{
diff --git a/Engine/Source/Runtime/Renderer/Private/SceneRendering.cpp b/Engine/Source/Runtime/Renderer/Private/SceneRendering.cpp
--- a/Engine/Source/Runtime/Renderer/Private/SceneRendering.cpp
+++ b/Engine/Source/Runtime/Renderer/Private/SceneRendering.cpp
@@ -1914,7 +1914,7 @@ void FViewInfo::SetupUniformBufferParameters(
 		const bool bSetupSkyIrradiance = Scene
 			&& Scene->SkyLight
 			// Skylights with static lighting already had their diffuse contribution baked into lightmaps
-			&& !Scene->SkyLight->bHasStaticLighting
+			&& (!Scene->SkyLight->bHasStaticLighting || !IsStaticLightingAllowed())
 			&& Family->EngineShowFlags.SkyLighting;
 
 		if (bSetupSkyIrradiance)
@@ -3988,8 +3988,15 @@ void FSceneRenderer::OnRenderFinish(FRDGBuilder& GraphBuilder, FRDGTextureRef Vi
 
 		const bool bShowNoSkyAtmosphereComponentWarning = !Scene->HasSkyAtmosphere() && ViewFamily.EngineShowFlags.VisualizeSkyAtmosphere;
 
+		const bool bMobile = (FeatureLevel <= ERHIFeatureLevel::ES3_1);
 		const bool bStationarySkylight = Scene->SkyLight && Scene->SkyLight->bWantsStaticShadowing;
-		const bool bShowSkylightWarning = bStationarySkylight && !FReadOnlyCVARCache::EnableStationarySkylight();
+		bool bShowSkylightWarning = bStationarySkylight && !FReadOnlyCVARCache::EnableStationarySkylight();
+		if (bMobile)
+		{
+			// For mobile EnableStationarySkylight has to be enabled in a projects with StaticLighting to support Stationary or Movable skylights
+			bShowSkylightWarning = IsStaticLightingAllowed() && !FReadOnlyCVARCache::EnableStationarySkylight() && (bStationarySkylight || (Scene->SkyLight && Scene->SkyLight->IsMovable()));
+		}
+
 		const bool bRealTimeSkyCaptureButNothingToCapture = Scene->SkyLight && Scene->SkyLight->bRealTimeCaptureEnabled && (!Scene->HasSkyAtmosphere() && !Scene->HasVolumetricCloud() && (Views.Num() > 0 && !Views[0].bSceneHasSkyMaterial));
 
 		const bool bShowPointLightWarning = UsedWholeScenePointLightNames.Num() > 0 && !FReadOnlyCVARCache::EnablePointLightShadows();
@@ -4057,7 +4064,6 @@ void FSceneRenderer::OnRenderFinish(FRDGBuilder& GraphBuilder, FRDGTextureRef Vi
 		const bool bContactShadowIntensityCvarUsed = ContactShadowNonCastingIntensityCVar && ContactShadowNonCastingIntensityCVar->GetFloat() != 0.0f;
 
 		// Mobile-specific warnings
-		const bool bMobile = (FeatureLevel <= ERHIFeatureLevel::ES3_1);
 		const bool bShowMobileLowQualityLightmapWarning = bMobile && !FReadOnlyCVARCache::EnableLowQualityLightmaps() && IsStaticLightingAllowed();
 		const bool bShowMobileDynamicCSMWarning = bMobile && Scene->NumMobileStaticAndCSMLights_RenderThread > 0 && !(FReadOnlyCVARCache::MobileEnableStaticAndCSMShadowReceivers() && FReadOnlyCVARCache::MobileAllowDistanceFieldShadows());
 		const bool bMobileMissingSkyMaterial = (bMobile && Scene->HasSkyAtmosphere() && (Views.Num() > 0 && !Views[0].bSceneHasSkyMaterial));
diff --git a/Engine/Source/Runtime/Renderer/Private/SkyPassRendering.cpp b/Engine/Source/Runtime/Renderer/Private/SkyPassRendering.cpp
--- a/Engine/Source/Runtime/Renderer/Private/SkyPassRendering.cpp
+++ b/Engine/Source/Runtime/Renderer/Private/SkyPassRendering.cpp
@@ -118,7 +118,6 @@ bool FSkyPassMeshProcessor::Process(
 			EMobileLocalLightSetting::LOCAL_LIGHTS_DISABLED,
 			MaterialResource,
 			VertexFactory->GetType(),
-			false,
 			SkyPassShaders.VertexShader,
 			SkyPassShaders.PixelShader
 		))
diff --git a/Engine/Source/Runtime/Renderer/Private/WaterInfoTextureRendering.cpp b/Engine/Source/Runtime/Renderer/Private/WaterInfoTextureRendering.cpp
--- a/Engine/Source/Runtime/Renderer/Private/WaterInfoTextureRendering.cpp
+++ b/Engine/Source/Runtime/Renderer/Private/WaterInfoTextureRendering.cpp
@@ -235,7 +235,7 @@ void FWaterInfoTexturePassMeshProcessor::CollectPSOInitializers(const FSceneText
 
 			FMaterialShaderTypes ShaderTypes;
 			ShaderTypes.AddShaderType<TMobileBasePassVS<TUniformLightMapPolicy<LMP_NO_LIGHTMAP>>>();
-			ShaderTypes.AddShaderType<TMobileBasePassPS<TUniformLightMapPolicy<LMP_NO_LIGHTMAP>, false, LOCAL_LIGHTS_DISABLED>>();
+			ShaderTypes.AddShaderType<TMobileBasePassPS<TUniformLightMapPolicy<LMP_NO_LIGHTMAP>, LOCAL_LIGHTS_DISABLED>>();
 			FMaterialShaders Shaders;
 			if (!Material.TryGetShaders(ShaderTypes, VertexFactoryData.VertexFactoryType, Shaders))
 			{
@@ -323,7 +323,7 @@ bool FWaterInfoTexturePassMeshProcessor::TryAddMeshBatch(const FMeshBatch& RESTR
 
 			FMaterialShaderTypes ShaderTypes;
 			ShaderTypes.AddShaderType<TMobileBasePassVS<TUniformLightMapPolicy<LMP_NO_LIGHTMAP>>>();
-			ShaderTypes.AddShaderType<TMobileBasePassPS<TUniformLightMapPolicy<LMP_NO_LIGHTMAP>, false, LOCAL_LIGHTS_DISABLED>>();
+			ShaderTypes.AddShaderType<TMobileBasePassPS<TUniformLightMapPolicy<LMP_NO_LIGHTMAP>, LOCAL_LIGHTS_DISABLED>>();
 			FMaterialShaders Shaders;
 			if (!Material.TryGetShaders(ShaderTypes, VertexFactory->GetType(), Shaders))
 			{
