Fixed assert in parallel GDME pipes due to race condition when not calling Pipe::WaitUntilEmpty and relying on task dependencies instead.
 - Pipe.Launch returns a task that is used to track dependencies, but the pipe's HasWork() method relies on an atomic refcount that is cleared after the task has finished. This causes a race if the user relies purely on the task instead of explicitly emptying the pipe.
 - Removed an unused atomic from the visiblity Pipe.

diff --git a/Engine/Source/Runtime/Engine/Private/PrimitiveSceneProxy.cpp b/Engine/Source/Runtime/Engine/Private/PrimitiveSceneProxy.cpp
--- a/Engine/Source/Runtime/Engine/Private/PrimitiveSceneProxy.cpp
+++ b/Engine/Source/Runtime/Engine/Private/PrimitiveSceneProxy.cpp
@@ -96,8 +96,7 @@ static TAutoConsoleVariable<int32> CVarApproximateOcclusionQueries(
 	ECVF_RenderThreadSafe
 );
 
-// Temporarily off by default as a workaround for UE-220106.
-bool GParallelGatherDynamicMeshElements = false;
+bool GParallelGatherDynamicMeshElements = true;
 static FAutoConsoleVariableRef CVarParallelGatherDynamicMeshElements(
 	TEXT("r.Visibility.DynamicMeshElements.Parallel"),
 	GParallelGatherDynamicMeshElements,
diff --git a/Engine/Source/Runtime/Renderer/Private/SceneVisibility.cpp b/Engine/Source/Runtime/Renderer/Private/SceneVisibility.cpp
--- a/Engine/Source/Runtime/Renderer/Private/SceneVisibility.cpp
+++ b/Engine/Source/Runtime/Renderer/Private/SceneVisibility.cpp
@@ -4034,6 +4034,10 @@ void FDynamicMeshElementContext::Finish()
 	DynamicVertexBuffer.Commit();
 	DynamicIndexBuffer.Commit();
 	RHICmdList->FinishRecording();
+
+	// Even though task dependencies are setup so all work is done by this point, we still have to wait on the
+	// pipe to clear out its internal state. Otherwise it can assert that it still has work at shutdown.
+	Pipe.WaitUntilEmpty();
 }
 
 FDynamicMeshElementContextContainer::~FDynamicMeshElementContextContainer()
diff --git a/Engine/Source/Runtime/Renderer/Private/SceneVisibilityPrivate.h b/Engine/Source/Runtime/Renderer/Private/SceneVisibilityPrivate.h
--- a/Engine/Source/Runtime/Renderer/Private/SceneVisibilityPrivate.h
+++ b/Engine/Source/Runtime/Renderer/Private/SceneVisibilityPrivate.h
@@ -73,8 +73,6 @@ public:
 			{
 				EmptyFunction();
 			}
-
-			NumEmptyEvents.fetch_add(1, std::memory_order_relaxed);
 		}
 	}
 
@@ -136,7 +134,6 @@ private:
 	TArray<CommandType, SceneRenderingAllocator> Queue;
 	UE::Tasks::FPipe Pipe;
 	std::atomic_int32_t NumCommands{ 0 };
-	std::atomic_int32_t NumEmptyEvents{ 0 };
 };
 
 ///////////////////////////////////////////////////////////////////////////////
