Fixed incorrect scene uniform buffer referenced by ray tracing shaders. Uniform buffer was not being recreated when marked as dirty.

diff --git a/Engine/Source/Runtime/Renderer/Private/SceneUniformBuffer.cpp b/Engine/Source/Runtime/Renderer/Private/SceneUniformBuffer.cpp
--- a/Engine/Source/Runtime/Renderer/Private/SceneUniformBuffer.cpp
+++ b/Engine/Source/Runtime/Renderer/Private/SceneUniformBuffer.cpp
@@ -217,6 +217,8 @@ RENDERER_API TRDGUniformBufferRef<FSceneUniformParameters> FSceneUniformBuffer::
 
 RENDERER_API FRHIUniformBuffer* FSceneUniformBuffer::GetBufferRHI(FRDGBuilder& GraphBuilder)
 {
+	// Ensure the buffer is prepped.
+	GetBuffer(GraphBuilder);
 	return GraphBuilder.ConvertToExternalUniformBuffer(Buffer);
 }
 
