Fixed build break.

diff --git a/Engine/Source/Runtime/Renderer/Public/ScreenPass.h b/Engine/Source/Runtime/Renderer/Public/ScreenPass.h
--- a/Engine/Source/Runtime/Renderer/Public/ScreenPass.h
+++ b/Engine/Source/Runtime/Renderer/Public/ScreenPass.h
@@ -102,7 +102,7 @@ struct FScreenPassRenderTarget : public FScreenPassTexture
 		: FScreenPassTexture(InTexture)
 	{
 		// Derive the load action based on whether we are rendering to the entire texture or not.
-		if (ViewRect.Min != FIntPoint::ZeroValue || Texture && Texture->Desc.Extent != ViewRect.Max)
+		if (ViewRect.Min != FIntPoint::ZeroValue || (Texture && Texture->Desc.Extent != ViewRect.Max))
 		{
 			LoadAction = ERenderTargetLoadAction::ELoad;
 		}
