Sanitize the bounds of primitives after detecting and reporting NaNs through ensure to protect the renderer from bad data.

diff --git a/Engine/Source/Runtime/Renderer/Private/RendererScene.cpp b/Engine/Source/Runtime/Renderer/Private/RendererScene.cpp
--- a/Engine/Source/Runtime/Renderer/Private/RendererScene.cpp
+++ b/Engine/Source/Runtime/Renderer/Private/RendererScene.cpp
@@ -1793,6 +1793,25 @@ void FScene::AddPrimitive(FPrimitiveSceneDesc* Primitive)
 	BatchAddPrimitivesInternal(MakeArrayView(&Primitive, 1));
 }
 
+template<typename PrimitiveType>
+static void CheckAndSanitizePrimitiveBounds(FBoxSphereBounds &InOutWorldBounds, PrimitiveType *Primitive)
+{
+	// Help track down primitive with bad bounds way before the it gets to the Renderer
+	if (!ensureMsgf(!InOutWorldBounds.ContainsNaN(),
+		TEXT("NaNs found on Bounds for Primitive %s: Owner: %s, Resource: %s, Level: %s, Origin: %s, BoxExtent: %s, SphereRadius: %f"),
+		*Primitive->GetName(),
+		*Primitive->GetSceneProxy()->GetOwnerName().ToString(),
+		*Primitive->GetSceneProxy()->GetResourceName().ToString(),
+		*Primitive->GetSceneProxy()->GetLevelName().ToString(),
+		*InOutWorldBounds.Origin.ToString(),
+		*InOutWorldBounds.BoxExtent.ToString(),
+		InOutWorldBounds.SphereRadius
+		))
+	{
+		InOutWorldBounds = FBoxSphereBounds(ForceInit);
+	}
+}
+
 template<class T> 	
 void FScene::BatchAddPrimitivesInternal(TArrayView<T*> InPrimitives)
 {
@@ -1881,6 +1900,9 @@ void FScene::BatchAddPrimitivesInternal(TArrayView<T*> InPrimitives)
 		// Cache the primitives initial transform.
 		FMatrix RenderMatrix = Primitive->GetRenderMatrix();
 		FVector AttachmentRootPosition = Primitive->GetActorPositionForRenderer();
+		FBoxSphereBounds WorldBounds = Primitive->Bounds;
+
+		CheckAndSanitizePrimitiveBounds(WorldBounds, Primitive);
 
 		CreateCommands.Emplace(
 			PrimitiveSceneInfo,
@@ -1888,14 +1910,11 @@ void FScene::BatchAddPrimitivesInternal(TArrayView<T*> InPrimitives)
 			// If this primitive has a simulated previous transform, ensure that the velocity data for the scene representation is correct.
 			FMotionVectorSimulation::Get().GetPreviousTransform(ToUObject(Primitive)),
 			RenderMatrix,
-			Primitive->Bounds,
+			WorldBounds,
 			AttachmentRootPosition,
 			Primitive->GetLocalBounds()
 		);
 
-		// Help track down primitive with bad bounds way before the it gets to the Renderer
-		ensureMsgf(!Primitive->Bounds.ContainsNaN(),
-				TEXT("Nans found on Bounds for Primitive %s: Origin %s, BoxExtent %s, SphereRadius %f"), *Primitive->GetName(), *Primitive->Bounds.Origin.ToString(), *Primitive->Bounds.BoxExtent.ToString(), Primitive->Bounds.SphereRadius);
 
 		INC_DWORD_STAT_BY( STAT_GameToRendererMallocTotal, PrimitiveSceneProxy->GetMemoryFootprint() + PrimitiveSceneInfo->GetMemoryFootprint() );
 
@@ -2044,17 +2063,7 @@ void FScene::UpdatePrimitiveTransformInternal(T* Primitive)
 			UpdateParams.LocalBounds = Primitive->GetLocalBounds();
 			UpdateParams.PreviousTransform = FMotionVectorSimulation::Get().GetPreviousTransform(ToUObject(Primitive));
 
-			// Help track down primitive with bad bounds way before the it gets to the renderer.
-			ensureMsgf(!UpdateParams.WorldBounds.BoxExtent.ContainsNaN() && !UpdateParams.WorldBounds.Origin.ContainsNaN() && !FMath::IsNaN(UpdateParams.WorldBounds.SphereRadius) && FMath::IsFinite(UpdateParams.WorldBounds.SphereRadius),
-				TEXT("NaNs found on Bounds for Primitive %s: Owner: %s, Resource: %s, Level: %s, Origin: %s, BoxExtent: %s, SphereRadius: %f"),
-				*Primitive->GetName(),
-				*Primitive->GetSceneProxy()->GetOwnerName().ToString(),
-				*Primitive->GetSceneProxy()->GetResourceName().ToString(),
-				*Primitive->GetSceneProxy()->GetLevelName().ToString(),
-				*UpdateParams.WorldBounds.Origin.ToString(),
-				*UpdateParams.WorldBounds.BoxExtent.ToString(),
-				UpdateParams.WorldBounds.SphereRadius
-			);
+			CheckAndSanitizePrimitiveBounds(UpdateParams.WorldBounds, Primitive);
 
 			bool bPerformUpdate = true;
 
