Improved handling of downscaled exposure illuminance texture.
- moved point sample to center of output texel instead of top-left.
- use FScreenTransform instead of simple downscale factor.

diff --git a/Engine/Shaders/Private/PostProcessEyeAdaptation.usf b/Engine/Shaders/Private/PostProcessEyeAdaptation.usf
--- a/Engine/Shaders/Private/PostProcessEyeAdaptation.usf
+++ b/Engine/Shaders/Private/PostProcessEyeAdaptation.usf
@@ -366,7 +366,7 @@ void BasicEyeAdaptationCS(uint GIndex : SV_GroupIndex, uint2 GTId : SV_GroupThre
 
 RWTexture2D<float> RWIlluminanceTexture;
 SCREEN_PASS_TEXTURE_VIEWPORT(Illuminance)
-uint IllumiananceDownscaleFactor;
+FScreenTransform IlluminanceRectToColorRect;
 
 // same weights as CalculateEyeAdaptationLuminance
 // but skip the min since we apply EyeAdaptation_IgnoreMaterialsMinBaseColorLuminance here
@@ -379,7 +379,7 @@ float CalculateEyeAdaptationLuminanceWithoutMin(float3 Color)
 void SetupExposureIlluminance(uint2 ThreadId : SV_DispatchThreadID)
 {
 	const uint2 OutputTexelPos = ThreadId + Illuminance_ViewportMin;
-	const uint2 TexelPos = ThreadId * IllumiananceDownscaleFactor + View.ViewRectMinAndSize.xy;
+	const uint2 TexelPos = ApplyScreenTransform(float2(OutputTexelPos + 0.5f), IlluminanceRectToColorRect);
 
 	const float3 Emissive = ColorTexture.Load(int3(TexelPos, 0)).rgb;
 
@@ -507,7 +507,7 @@ FExposureMaterial GetExposureMaterial(uint2 InPixelPos)
 void CalculateExposureIlluminance(uint2 ThreadId : SV_DispatchThreadID)
 {
 	const uint2 OutputTexelPos = ThreadId + Illuminance_ViewportMin;
-	const uint2 TexelPos = ThreadId * IllumiananceDownscaleFactor + View.ViewRectMinAndSize.xy;
+	const uint2 TexelPos = ApplyScreenTransform(float2(OutputTexelPos + 0.5f), IlluminanceRectToColorRect);
 
 	const float3 SceneColor = ColorTexture.Load(int3(TexelPos, 0)).rgb * View.OneOverPreExposure;
 
diff --git a/Engine/Source/Runtime/Renderer/Private/PostProcess/PostProcessEyeAdaptation.cpp b/Engine/Source/Runtime/Renderer/Private/PostProcess/PostProcessEyeAdaptation.cpp
--- a/Engine/Source/Runtime/Renderer/Private/PostProcess/PostProcessEyeAdaptation.cpp
+++ b/Engine/Source/Runtime/Renderer/Private/PostProcess/PostProcessEyeAdaptation.cpp
@@ -641,7 +641,7 @@ class FSetupExposureIlluminanceCS : public FGlobalShader
 		SHADER_PARAMETER_RDG_TEXTURE_UAV(RWTexture2D<float4>, RWIlluminanceTexture)
 		SHADER_PARAMETER_STRUCT(FScreenPassTextureViewportParameters, Illuminance)
 		SHADER_PARAMETER_STRUCT(FEyeAdaptationParameters, EyeAdaptation)
-		SHADER_PARAMETER(uint32, IllumiananceDownscaleFactor)
+		SHADER_PARAMETER(FScreenTransform, IlluminanceRectToColorRect)
 	END_SHADER_PARAMETER_STRUCT()
 
 	static bool ShouldCompilePermutation(const FGlobalShaderPermutationParameters& Parameters)
@@ -692,7 +692,7 @@ FRDGTextureRef AddSetupExposureIlluminancePass(
 			PassParameters->RWIlluminanceTexture = GraphBuilder.CreateUAV(OutputTexture);
 			PassParameters->Illuminance = GetScreenPassTextureViewportParameters(OutputViewport);
 			PassParameters->EyeAdaptation = GetEyeAdaptationParameters(View);
-			PassParameters->IllumiananceDownscaleFactor = GetAutoExposureIlluminanceDownscaleFactor();
+			PassParameters->IlluminanceRectToColorRect = FScreenTransform::ChangeRectFromTo(OutputViewport.Rect, SceneViewport.Rect);
 
 			auto ComputeShader = View.ShaderMap->GetShader<FSetupExposureIlluminanceCS>();
 
@@ -719,7 +719,7 @@ class FCalculateExposureIlluminanceCS : public FGlobalShader
 		SHADER_PARAMETER_RDG_TEXTURE(Texture2D, ColorTexture)
 		SHADER_PARAMETER_RDG_TEXTURE_UAV(RWTexture2D<float4>, RWIlluminanceTexture)
 		SHADER_PARAMETER_STRUCT(FScreenPassTextureViewportParameters, Illuminance)
-		SHADER_PARAMETER(uint32, IllumiananceDownscaleFactor)
+		SHADER_PARAMETER(FScreenTransform, IlluminanceRectToColorRect)
 
 		SHADER_PARAMETER_STRUCT(FEyeAdaptationParameters, EyeAdaptation)
 
@@ -770,7 +770,7 @@ FRDGTextureRef AddCalculateExposureIlluminancePass(
 			PassParameters->ColorTexture = SceneTextures.Color.Resolve;
 			PassParameters->RWIlluminanceTexture = GraphBuilder.CreateUAV(ExposureIlluminanceSetup);
 			PassParameters->Illuminance = GetScreenPassTextureViewportParameters(OutputViewport);
-			PassParameters->IllumiananceDownscaleFactor = GetAutoExposureIlluminanceDownscaleFactor();
+			PassParameters->IlluminanceRectToColorRect = FScreenTransform::ChangeRectFromTo(OutputViewport.Rect, SceneViewport.Rect);
 			PassParameters->EyeAdaptation = GetEyeAdaptationParameters(View);
 
 			PassParameters->PreIntegratedGF = GSystemTextures.PreintegratedGF->GetRHI();
