Support for visualizing luminance used for auto exposure calculation.
- Enable Show -> Visualize -> HDR mode + r.EyeAdaptation.VisualizeDebugType 2.
- r.EyeAdaptation.VisualizeLuminanceScale can be used to scale the output.

diff --git a/Engine/Shaders/Private/PostProcessVisualizeHDR.usf b/Engine/Shaders/Private/PostProcessVisualizeHDR.usf
--- a/Engine/Shaders/Private/PostProcessVisualizeHDR.usf
+++ b/Engine/Shaders/Private/PostProcessVisualizeHDR.usf
@@ -18,12 +18,18 @@ SamplerState SceneColorSampler;
 Texture2D HDRSceneColorTexture;
 SamplerState HDRSceneColorSampler;
 
+FScreenTransform ColorUVToLuminanceUV;
+Texture2D LuminanceTexture;
+SamplerState LuminanceSampler;
+
 Texture2D HistogramTexture;
 
 SCREEN_PASS_TEXTURE_VIEWPORT(Input)
 SCREEN_PASS_TEXTURE_VIEWPORT(Output)
 
 float IlluminanceMeterEnabled;
+float UsingIlluminance;
+float LuminanceScale;
 
 // only needed for nice visualization
 float ComputeHistogramMax(Texture2D HistogramTexture)
@@ -413,6 +419,27 @@ void MainPS(noperspective float4 UVAndScreenPos : TEXCOORD0, float4 SvPosition :
 			OutColor.rgb = float3(0,0,.5);
 		}
 	}
+	else if (EyeAdaptation_VisualizeDebugType == 2)
+	{
+		float Luminance;
+
+		if (UsingIlluminance > 0.5f)
+		{
+			float2 LuminanceUV = ApplyScreenTransform(UV, ColorUVToLuminanceUV);
+
+			Luminance = Texture2DSample(LuminanceTexture, LuminanceSampler, LuminanceUV).r;
+			Luminance *= View.OneOverPreExposure;
+		}
+		else
+		{
+			float4 SceneColor = Texture2DSample(HDRSceneColorTexture, HDRSceneColorSampler, UV);
+			SceneColor.xyz *= View.OneOverPreExposure;
+
+			Luminance = CalculateEyeAdaptationLuminance(SceneColor);
+		}
+
+		OutColor.rgb = Luminance * LuminanceScale;
+	}
 	else // should never happen?
 	{
 		OutColor.rgb = float3(0,0,0);
diff --git a/Engine/Source/Runtime/Renderer/Private/PostProcess/PostProcessEyeAdaptation.cpp b/Engine/Source/Runtime/Renderer/Private/PostProcess/PostProcessEyeAdaptation.cpp
--- a/Engine/Source/Runtime/Renderer/Private/PostProcess/PostProcessEyeAdaptation.cpp
+++ b/Engine/Source/Runtime/Renderer/Private/PostProcess/PostProcessEyeAdaptation.cpp
@@ -54,8 +54,9 @@ namespace
 		TEXT("r.EyeAdaptation.VisualizeDebugType"),
 		0,
 		TEXT("When enabling Show->Visualize->HDR (Eye Adaptation) is enabled, this flag controls the scene color.\n")
-		TEXT("    0: Scene Color after tonemapping (default).\n")
-		TEXT("    1: Histogram Debug\n"),
+		TEXT("    0: Scene Color after tonemapping (default)\n")
+		TEXT("    1: Histogram Debug\n")
+		TEXT("    2: Luminance\n"),
 		ECVF_RenderThreadSafe);
 
 	TAutoConsoleVariable<float> CVarEyeAdaptationLensAttenuation(
diff --git a/Engine/Source/Runtime/Renderer/Private/PostProcess/PostProcessVisualizeHDR.cpp b/Engine/Source/Runtime/Renderer/Private/PostProcess/PostProcessVisualizeHDR.cpp
--- a/Engine/Source/Runtime/Renderer/Private/PostProcess/PostProcessVisualizeHDR.cpp
+++ b/Engine/Source/Runtime/Renderer/Private/PostProcess/PostProcessVisualizeHDR.cpp
@@ -7,6 +7,12 @@
 #include "DataDrivenShaderPlatformInfo.h"
 #include "SceneRendering.h"
 
+TAutoConsoleVariable<float> CVarEyeAdaptationVisualizeLuminanceScale(
+	TEXT("r.EyeAdaptation.VisualizeLuminanceScale"),
+	1.0f,
+	TEXT("Scale applied to output when visualizing luminance.\n"),
+	ECVF_RenderThreadSafe);
+
 extern bool IsExtendLuminanceRangeEnabled();
 
 class FVisualizeHDRPS : public FGlobalShader
@@ -23,9 +29,11 @@ public:
 		SHADER_PARAMETER_RDG_TEXTURE(Texture2D, HDRSceneColorTexture)
 		SHADER_PARAMETER_RDG_TEXTURE(Texture2D, SceneColorTexture)
 		SHADER_PARAMETER_RDG_TEXTURE(Texture2D, HistogramTexture)
+		SHADER_PARAMETER_RDG_TEXTURE(Texture2D, LuminanceTexture)
 		SHADER_PARAMETER_RDG_BUFFER_SRV(StructuredBuffer<float4>, EyeAdaptationBuffer)
 		SHADER_PARAMETER_SAMPLER(SamplerState, HDRSceneColorSampler)
 		SHADER_PARAMETER_SAMPLER(SamplerState, SceneColorSampler)
+		SHADER_PARAMETER_SAMPLER(SamplerState, LuminanceSampler)
 		SHADER_PARAMETER_TEXTURE(Texture2D, MiniFontTexture)
 		SHADER_PARAMETER_STRUCT(FEyeAdaptationParameters, EyeAdaptation)
 		SHADER_PARAMETER(float, FilmSlope)
@@ -33,7 +41,10 @@ public:
 		SHADER_PARAMETER(float, FilmShoulder)
 		SHADER_PARAMETER(float, FilmBlackClip)
 		SHADER_PARAMETER(float, FilmWhiteClip)
+		SHADER_PARAMETER(FScreenTransform, ColorUVToLuminanceUV)
 		SHADER_PARAMETER(float, IlluminanceMeterEnabled)
+		SHADER_PARAMETER(float, UsingIlluminance)
+		SHADER_PARAMETER(float, LuminanceScale)
 		RENDER_TARGET_BINDING_SLOTS()
 	END_SHADER_PARAMETER_STRUCT()
 
@@ -79,8 +90,10 @@ FScreenPassTexture AddVisualizeHDRPass(FRDGBuilder& GraphBuilder, const FViewInf
 
 	const FScreenPassTextureViewport InputViewport(Inputs.SceneColor);
 	const FScreenPassTextureViewport OutputViewport(Output);
+	const FScreenPassTextureViewport LuminanceViewport = Inputs.Luminance.IsValid() ? FScreenPassTextureViewport(Inputs.Luminance) : FScreenPassTextureViewport();
 
 	FRHISamplerState* BilinearClampSampler = TStaticSamplerState<SF_Bilinear, AM_Clamp, AM_Clamp, AM_Clamp>::GetRHI();
+	FRHISamplerState* PointSampler = TStaticSamplerState<SF_Point>::GetRHI();
 
 	const FPostProcessSettings& Settings = View.FinalPostProcessSettings;
 
@@ -93,6 +106,8 @@ FScreenPassTexture AddVisualizeHDRPass(FRDGBuilder& GraphBuilder, const FViewInf
 	PassParameters->SceneColorSampler = BilinearClampSampler;
 	PassParameters->HDRSceneColorTexture = Inputs.SceneColorBeforeTonemap.Texture;
 	PassParameters->HDRSceneColorSampler = BilinearClampSampler;
+	PassParameters->LuminanceTexture = Inputs.Luminance.IsValid() ? Inputs.Luminance.Texture : GSystemTextures.GetBlackDummy(GraphBuilder);
+	PassParameters->LuminanceSampler = PointSampler;
 	PassParameters->HistogramTexture = Inputs.HistogramTexture;
 	PassParameters->EyeAdaptationBuffer = GraphBuilder.CreateSRV(Inputs.EyeAdaptationBuffer);
 	PassParameters->EyeAdaptation = *Inputs.EyeAdaptationParameters;
@@ -103,8 +118,10 @@ FScreenPassTexture AddVisualizeHDRPass(FRDGBuilder& GraphBuilder, const FViewInf
 	PassParameters->FilmShoulder = Settings.FilmShoulder;
 	PassParameters->FilmBlackClip = Settings.FilmBlackClip;
 	PassParameters->FilmWhiteClip = Settings.FilmWhiteClip;
+	PassParameters->ColorUVToLuminanceUV = FScreenTransform::ChangeTextureUVCoordinateFromTo(InputViewport, LuminanceViewport);
 	PassParameters->IlluminanceMeterEnabled = IsIlluminanceMeterSupportedByView(View) ? 1.0f : 0.0f;
-	
+	PassParameters->UsingIlluminance = IsAutoExposureUsingIlluminanceEnabled(View) ? 1.0f : 0.0f;
+	PassParameters->LuminanceScale = CVarEyeAdaptationVisualizeLuminanceScale.GetValueOnRenderThread();
 
 	TShaderMapRef<FVisualizeHDRPS> PixelShader(View.ShaderMap);
 
diff --git a/Engine/Source/Runtime/Renderer/Private/PostProcess/PostProcessVisualizeHDR.h b/Engine/Source/Runtime/Renderer/Private/PostProcess/PostProcessVisualizeHDR.h
--- a/Engine/Source/Runtime/Renderer/Private/PostProcess/PostProcessVisualizeHDR.h
+++ b/Engine/Source/Runtime/Renderer/Private/PostProcess/PostProcessVisualizeHDR.h
@@ -13,6 +13,7 @@ struct FVisualizeHDRInputs
 
 	FScreenPassTexture SceneColor; 
 	FScreenPassTexture SceneColorBeforeTonemap;
+	FScreenPassTexture Luminance;
 
 	FRDGTextureRef HistogramTexture = nullptr;
 	FRDGBufferRef EyeAdaptationBuffer = nullptr;
diff --git a/Engine/Source/Runtime/Renderer/Private/PostProcess/PostProcessing.cpp b/Engine/Source/Runtime/Renderer/Private/PostProcess/PostProcessing.cpp
--- a/Engine/Source/Runtime/Renderer/Private/PostProcess/PostProcessing.cpp
+++ b/Engine/Source/Runtime/Renderer/Private/PostProcess/PostProcessing.cpp
@@ -357,6 +357,9 @@ void AddPostProcessingPasses(
 	FRDGBufferRef LastEyeAdaptationBuffer = GetEyeAdaptationBuffer(GraphBuilder, View);
 	FRDGBufferRef EyeAdaptationBuffer = LastEyeAdaptationBuffer;
 
+	const FIntRect ExposureIlluminanceRect = GetDownscaledRect(PrimaryViewRect, GetAutoExposureIlluminanceDownscaleFactor());
+	FScreenPassTexture ExposureIlluminance = FScreenPassTexture(Inputs.ExposureIlluminance, ExposureIlluminanceRect);
+
 	FLocalExposureParameters LocalExposureParameters;
 
 	// Histogram defaults to black because the histogram eye adaptation pass is used for the manual metering mode.
@@ -1102,11 +1105,9 @@ void AddPostProcessingPasses(
 
 			if (IsAutoExposureUsingIlluminanceEnabled(View))
 			{
-				if (Inputs.ExposureIlluminance)
+				if (ExposureIlluminance.IsValid())
 				{
-					const FIntRect IlluminanceRect = GetDownscaledRect(PrimaryViewRect, GetAutoExposureIlluminanceDownscaleFactor());
-
-					HistogramSceneColor = FScreenPassTextureSlice::CreateFromScreenPassTexture(GraphBuilder, FScreenPassTexture(Inputs.ExposureIlluminance, IlluminanceRect));
+					HistogramSceneColor = FScreenPassTextureSlice::CreateFromScreenPassTexture(GraphBuilder, ExposureIlluminance);
 				}
 				else
 				{
@@ -1664,6 +1665,7 @@ void AddPostProcessingPasses(
 		PassSequence.AcceptOverrideIfLastPass(EPass::VisualizeHDR, PassInputs.OverrideOutput);
 		PassInputs.SceneColor = SceneColor;
 		PassInputs.SceneColorBeforeTonemap = FScreenPassTexture::CopyFromSlice(GraphBuilder, SceneColorBeforeTonemapSlice);
+		PassInputs.Luminance = ExposureIlluminance;
 		PassInputs.HistogramTexture = HistogramTexture;
 		PassInputs.EyeAdaptationBuffer = GetEyeAdaptationBuffer(GraphBuilder, View);
 		PassInputs.EyeAdaptationParameters = &EyeAdaptationParameters;
