Fixed race condition with FPipe HasWork() being true at destruction for shadow GDME.

diff --git a/Engine/Source/Runtime/Renderer/Private/ShadowSetup.cpp b/Engine/Source/Runtime/Renderer/Private/ShadowSetup.cpp
--- a/Engine/Source/Runtime/Renderer/Private/ShadowSetup.cpp
+++ b/Engine/Source/Runtime/Renderer/Private/ShadowSetup.cpp
@@ -1302,6 +1302,7 @@ public:
 		Collector.Finish();
 		DynamicVertexBuffer.Commit();
 		DynamicIndexBuffer.Commit();
+		Pipe.WaitUntilEmpty();
 	}
 
 	UE::Tasks::FPipe& GetPipe()
