Basic Exposure Fusion implementation.

diff --git a/Engine/Shaders/Private/PostProcessHistogramCommon.ush b/Engine/Shaders/Private/PostProcessHistogramCommon.ush
--- a/Engine/Shaders/Private/PostProcessHistogramCommon.ush
+++ b/Engine/Shaders/Private/PostProcessHistogramCommon.ush
@@ -274,3 +274,13 @@ float CalculateLocalExposure(float LogLum, float BaseLogLum, float LogMiddleGrey
 {
 	return exp2(CalculateLogLocalExposure(LogLum, BaseLogLum, LogMiddleGrey, HighlightContrastScale, ShadowContrastScale, DetailStrength));
 }
+
+float ExposureFusionTonemap(float Lum)
+{
+	return Lum / (1 + Lum);
+}
+
+float ExposureFusionInverseTonemap(float Lum)
+{
+	return Lum / (1 - Lum);
+}
\ No newline at end of file
diff --git a/Engine/Shaders/Private/PostProcessLocalExposure.usf b/Engine/Shaders/Private/PostProcessLocalExposure.usf
--- a/Engine/Shaders/Private/PostProcessLocalExposure.usf
+++ b/Engine/Shaders/Private/PostProcessLocalExposure.usf
@@ -1,5 +1,7 @@
 // Copyright Epic Games, Inc. All Rights Reserved.
 
+#define EYE_ADAPTATION_LOOSE_PARAMETERS 1
+
 #include "Common.ush"
 #include "ScreenPass.ush"
 #include "PostProcessHistogramCommon.ush"
@@ -18,6 +20,7 @@ SamplerState TextureSampler;
 
 RWTexture2D<float> OutputFloat;
 RWTexture2D<float4> OutputFloat4;
+RWTexture2D<float4> OutputFloat4_1;
 
 [numthreads(THREADGROUP_SIZEX, THREADGROUP_SIZEY, 1)]
 void SetupLogLuminanceCS(uint3 DispatchThreadId : SV_DispatchThreadID)
@@ -53,3 +56,83 @@ void ApplyLocalExposureCS(uint2 DispatchThreadId : SV_DispatchThreadID)
 		OutputFloat4[OutputPixelPos] = SceneColor;
 	}
 }
+
+float ExposureWeight(float V, float Target)
+{
+	float Sigma = 0.2f;
+
+	float A = (V - Target);
+	float B = 2 * Sigma * Sigma;
+
+	return exp2(-A * A / B);
+}
+
+[numthreads(THREADGROUP_SIZEX, THREADGROUP_SIZEY, 1)]
+void FusionSetupCS(uint2 DispatchThreadId : SV_DispatchThreadID)
+{
+	const uint2 OutputPixelPos = DispatchThreadId + Output_ViewportMin;
+
+	BRANCH
+	if (any(OutputPixelPos >= (uint2)Output_ViewportMax))
+	{
+		return;
+	}
+
+	const float2 ExposureScaleMiddleGrey = EyeAdaptationBuffer[0].xw;
+	const float ExposureScale = ExposureScaleMiddleGrey.x;
+	float MiddleGreyLumValue = log2(0.18 * ExposureScaleMiddleGrey.y * LocalExposure_MiddleGreyExposureCompensation);
+
+	float3 SceneColor = InputTexture.Load(uint3(Input_ViewportMin + DispatchThreadId, 0)).rgb * View.OneOverPreExposure * ExposureScale;
+	float LuminanceVal = CalculateEyeAdaptationLuminance(SceneColor);
+
+	float HighlightLuminanceVal = LuminanceVal * LocalExposure_HighlightContrastScale;
+	float ShadowLuminanceVal = LuminanceVal * LocalExposure_ShadowContrastScale;
+
+	float LuminanceValTonemapped = ExposureFusionTonemap(LuminanceVal);
+	float HighlightLuminanceValTonemapped = ExposureFusionTonemap(HighlightLuminanceVal);
+	float ShadowLuminanceValTonemapped = ExposureFusionTonemap(ShadowLuminanceVal);
+
+	const float Target = 0.50f;
+
+	float3 Weights = float3(ExposureWeight(LuminanceValTonemapped, Target), ExposureWeight(HighlightLuminanceValTonemapped, Target), ExposureWeight(ShadowLuminanceValTonemapped, Target));
+	Weights /= dot(Weights, 1.0f); // normalize weights
+
+	OutputFloat4[OutputPixelPos] = float4((LuminanceValTonemapped), (HighlightLuminanceValTonemapped), (ShadowLuminanceValTonemapped), 0.0f);
+	OutputFloat4_1[OutputPixelPos] = float4(Weights, 0.0f);
+}
+
+Texture2D WeightTexture;
+
+SCREEN_PASS_TEXTURE_VIEWPORT(CoarserMip)
+Texture2D CoarserMipTexture;
+Texture2D PrevResultTexture;
+
+FScreenTransform DispatchThreadToCoarseMipUV;
+
+[numthreads(THREADGROUP_SIZEX, THREADGROUP_SIZEY, 1)]
+void FusionBlendCS(uint2 DispatchThreadId : SV_DispatchThreadID)
+{
+	const uint2 OutputPixelPos = DispatchThreadId + Output_ViewportMin;
+
+	BRANCH
+	if (any(OutputPixelPos >= (uint2)Output_ViewportMax))
+	{
+		return;
+	}
+
+	float3 Weights = WeightTexture.Load(uint3(OutputPixelPos, 0)).rgb;
+	Weights /= dot(Weights, 1.0f); // normalize weights
+
+	float3 Lums = InputTexture.Load(uint3(OutputPixelPos, 0)).rgb;
+	float PrevResultVal = 0.0f;
+
+#if LAPLACIAN
+	float2 CoarserMipUV = ApplyScreenTransform(DispatchThreadId, DispatchThreadToCoarseMipUV);
+	CoarserMipUV = clamp(CoarserMipUV, CoarserMip_UVViewportBilinearMin, CoarserMip_UVViewportBilinearMax);
+
+	Lums -= Texture2DSample(CoarserMipTexture, TextureSampler, CoarserMipUV).rgb;
+	PrevResultVal = Texture2DSample(PrevResultTexture, TextureSampler, CoarserMipUV).r;
+#endif
+
+	OutputFloat[OutputPixelPos] = PrevResultVal + dot(Lums, Weights);
+}
diff --git a/Engine/Shaders/Private/PostProcessTonemap.usf b/Engine/Shaders/Private/PostProcessTonemap.usf
--- a/Engine/Shaders/Private/PostProcessTonemap.usf
+++ b/Engine/Shaders/Private/PostProcessTonemap.usf
@@ -10,6 +10,9 @@
 	#define Undefined SUPPORTS_SCENE_COLOR_APPLY_PARAMETERS 0
 #endif
 
+#define LOCAL_EXPOSURE_MODE_NONE 0
+#define LOCAL_EXPOSURE_MODE_BILATERAL 1
+#define LOCAL_EXPOSURE_MODE_FUSION 2
 
 #include "Common.ush"
 #include "PostProcessCommon.ush"
@@ -49,10 +52,17 @@ SamplerState BloomSampler;
 	StructuredBuffer<float4> SceneColorApplyParamaters;
 #endif
 
+#if LOCAL_EXPOSURE_MODE == LOCAL_EXPOSURE_MODE_BILATERAL
 Texture3D LumBilateralGrid;
 Texture2D BlurredLogLum;
 SamplerState LumBilateralGridSampler;
 SamplerState BlurredLogLumSampler;
+#elif LOCAL_EXPOSURE_MODE == LOCAL_EXPOSURE_MODE_FUSION
+FScreenTransform ColorToExposureFusion;
+SCREEN_PASS_TEXTURE_VIEWPORT(ExposureFusion)
+Texture2D ExposureFusionTexture;
+SamplerState ExposureFusionSampler;
+#endif
 
 // xyz:SceneColorTint.rgb, w:unused
 float4 ColorScale0;
@@ -406,7 +416,7 @@ float4 TonemapCommonPS(
 	#endif
 
 	// Compute the local exposure based of the scene color.
-	#if USE_LOCAL_EXPOSURE
+	#if LOCAL_EXPOSURE_MODE == LOCAL_EXPOSURE_MODE_BILATERAL
 	float LocalExposure = 1.0;
 	{
 		float LuminanceVal = CalculateEyeAdaptationLuminance(SceneColor.rgb * OneOverPreExposure);
@@ -418,7 +428,24 @@ float4 TonemapCommonPS(
 		float BaseLogLum = CalculateBaseLogLuminance(LogLuminance, LocalExposure_BlurredLuminanceBlend, GlobalExposure, FullViewUV, LumBilateralGrid, BlurredLogLum, LumBilateralGridSampler, BlurredLogLumSampler);
 		LocalExposure = CalculateLocalExposure(LogLuminance + log2(GlobalExposure), BaseLogLum, MiddleGreyLumValue, LocalExposure_HighlightContrastScale, LocalExposure_ShadowContrastScale, LocalExposure_DetailStrength);
 	}
-	#else // !USE_LOCAL_EXPOSURE
+	#elif LOCAL_EXPOSURE_MODE == LOCAL_EXPOSURE_MODE_FUSION
+	float LocalExposure = 1.0;
+	{
+		float2 ExposureFusionUV = ApplyScreenTransform(UV.xy, ColorToExposureFusion);
+		ExposureFusionUV = clamp(ExposureFusionUV, ExposureFusion_UVViewportBilinearMin, ExposureFusion_UVViewportBilinearMax);
+
+		const float FusionVal = Texture2DSample(ExposureFusionTexture, ExposureFusionSampler, ExposureFusionUV);
+
+		const float LuminanceVal = CalculateEyeAdaptationLuminance(SceneColor.rgb * OneOverPreExposure * GlobalExposure);
+		
+#if 1
+		LocalExposure = ExposureFusionInverseTonemap(FusionVal) / LuminanceVal;
+#else
+		const float LuminanceValTonemapped = ExposureFusionTonemap(LuminanceVal);
+		LocalExposure = FusionVal / LuminanceValTonemapped;
+#endif
+	}
+	#else // LOCAL_EXPOSURE_MODE_NONE
 	const float LocalExposure = 1.0;
 	#endif
 
@@ -475,7 +502,7 @@ float4 TonemapCommonPS(
 	float4 OutColor = 0;
 #if USE_GAMMA_ONLY
 	{
-		#if USE_LOCAL_EXPOSURE
+		#if LOCAL_EXPOSURE_MODE != LOCAL_EXPOSURE_MODE_NONE
 			#error Local exposure is not supported with gamma only
 		#endif
 		#if USE_SHARPEN
diff --git a/Engine/Source/Runtime/Engine/Classes/Engine/Scene.h b/Engine/Source/Runtime/Engine/Classes/Engine/Scene.h
--- a/Engine/Source/Runtime/Engine/Classes/Engine/Scene.h
+++ b/Engine/Source/Runtime/Engine/Classes/Engine/Scene.h
@@ -41,6 +41,15 @@ enum EAutoExposureMethod : int
 	AEM_MAX,
 };
 
+UENUM()
+enum class ELocalExposureMethod : uint8
+{
+	/** Decompose image into base and detail layers using a bilateral blur. Only the base layer contrast is reduced in order to preserve details. */
+	Bilateral UMETA(DisplayName = "Bilateral"),
+	/** Fuse multiple exposures according to quality measures. Local Exposure is calculated using the fused exposures. */
+	Fusion UMETA(DisplayName = "Fusion (Experimental)"),
+};
+
 UENUM()
 enum EBloomMethod : int
 {
@@ -887,6 +896,9 @@ struct FPostProcessSettings
 	UPROPERTY(EditAnywhere, BlueprintReadWrite, Category=Overrides, meta=(PinHiddenByDefault, InlineEditConditionToggle))
 	uint8 bOverride_HistogramLogMax:1;
 
+	UPROPERTY(EditAnywhere, BlueprintReadWrite, Category=Overrides, meta=(PinHiddenByDefault, InlineEditConditionToggle))
+	uint8 bOverride_LocalExposureMethod : 1;
+
 	UPROPERTY()
 	uint8 bOverride_LocalExposureContrastScale_DEPRECATED:1;
 
@@ -1842,6 +1854,10 @@ struct FPostProcessSettings
 	UPROPERTY()
 	float AutoExposureCalibrationConstant_DEPRECATED;
 
+	/** Local Exposure algorithm */
+	UPROPERTY(EditAnywhere, BlueprintReadWrite, Category = "Lens|Local Exposure", meta = (editcondition = "bOverride_LocalExposureMethod", DisplayName = "Method"))
+	ELocalExposureMethod LocalExposureMethod;
+
 	UPROPERTY()
 	float LocalExposureContrastScale_DEPRECATED;
 
diff --git a/Engine/Source/Runtime/Engine/Private/Scene.cpp b/Engine/Source/Runtime/Engine/Private/Scene.cpp
--- a/Engine/Source/Runtime/Engine/Private/Scene.cpp
+++ b/Engine/Source/Runtime/Engine/Private/Scene.cpp
@@ -509,6 +509,7 @@ FPostProcessSettings::FPostProcessSettings()
 
 	AutoExposureApplyPhysicalCameraExposure = 1;
 
+	LocalExposureMethod = ELocalExposureMethod::Bilateral;
 	LocalExposureContrastScale_DEPRECATED = 1.0f;
 	LocalExposureHighlightContrastScale = 1.0f;
 	LocalExposureShadowContrastScale = 1.0f;
diff --git a/Engine/Source/Runtime/Engine/Private/SceneView.cpp b/Engine/Source/Runtime/Engine/Private/SceneView.cpp
--- a/Engine/Source/Runtime/Engine/Private/SceneView.cpp
+++ b/Engine/Source/Runtime/Engine/Private/SceneView.cpp
@@ -1644,6 +1644,7 @@ void FSceneView::OverridePostProcessSettings(const FPostProcessSettings& Src, fl
 		LERP_PP(AutoExposureBias);
 		LERP_PP(HistogramLogMin);
 		LERP_PP(HistogramLogMax);
+		SET_PP(LocalExposureMethod);
 		LERP_PP(LocalExposureContrastScale_DEPRECATED);
 		LERP_PP(LocalExposureHighlightContrastScale);
 		LERP_PP(LocalExposureShadowContrastScale);
diff --git a/Engine/Source/Runtime/Renderer/Private/PostProcess/PostProcessLocalExposure.cpp b/Engine/Source/Runtime/Renderer/Private/PostProcess/PostProcessLocalExposure.cpp
--- a/Engine/Source/Runtime/Renderer/Private/PostProcess/PostProcessLocalExposure.cpp
+++ b/Engine/Source/Runtime/Renderer/Private/PostProcess/PostProcessLocalExposure.cpp
@@ -7,6 +7,7 @@
 #include "PostProcess/PostProcessLocalExposure.h"
 #include "PostProcess/PostProcessEyeAdaptation.h"
 #include "PostProcess/PostProcessWeightedSampleSum.h"
+#include "PostProcess/PostProcessDownsample.h"
 #include "Curves/CurveFloat.h"
 #include "SceneRendering.h"
 #include "ShaderCompilerCore.h"
@@ -91,6 +92,95 @@ public:
 
 IMPLEMENT_GLOBAL_SHADER(FApplyLocalExposureCS, "/Engine/Private/PostProcessLocalExposure.usf", "ApplyLocalExposureCS", SF_Compute);
 
+class FFusionSetupCS : public FGlobalShader
+{
+public:
+	// Changing these numbers requires LocalExposure.usf to be recompiled
+	static const uint32 ThreadGroupSizeX = 8;
+	static const uint32 ThreadGroupSizeY = 8;
+
+	DECLARE_GLOBAL_SHADER(FFusionSetupCS);
+	SHADER_USE_PARAMETER_STRUCT(FFusionSetupCS, FGlobalShader);
+
+	BEGIN_SHADER_PARAMETER_STRUCT(FParameters, )
+		SHADER_PARAMETER_STRUCT_REF(FViewUniformShaderParameters, View)
+
+		SHADER_PARAMETER_STRUCT(FScreenPassTextureViewportParameters, Input)
+		SHADER_PARAMETER_RDG_TEXTURE_SRV(Texture2D, InputTexture)
+
+		SHADER_PARAMETER_STRUCT(FEyeAdaptationParameters, EyeAdaptation)
+		SHADER_PARAMETER_RDG_BUFFER_SRV(StructuredBuffer<float4>, EyeAdaptationBuffer)
+
+		SHADER_PARAMETER_STRUCT(FLocalExposureParameters, LocalExposure)
+
+		SHADER_PARAMETER_STRUCT(FScreenPassTextureViewportParameters, Output)
+		SHADER_PARAMETER_RDG_TEXTURE_UAV(RWTexture2D<float4>, OutputFloat4)
+		SHADER_PARAMETER_RDG_TEXTURE_UAV(RWTexture2D<float4>, OutputFloat4_1)
+	END_SHADER_PARAMETER_STRUCT()
+
+	static bool ShouldCompilePermutation(const FGlobalShaderPermutationParameters& Parameters)
+	{
+		return IsFeatureLevelSupported(Parameters.Platform, ERHIFeatureLevel::SM5);
+	}
+
+	static void ModifyCompilationEnvironment(const FGlobalShaderPermutationParameters& Parameters, FShaderCompilerEnvironment& OutEnvironment)
+	{
+		FGlobalShader::ModifyCompilationEnvironment(Parameters, OutEnvironment);
+		OutEnvironment.SetDefine(TEXT("THREADGROUP_SIZEX"), ThreadGroupSizeX);
+		OutEnvironment.SetDefine(TEXT("THREADGROUP_SIZEY"), ThreadGroupSizeY);
+	}
+};
+
+IMPLEMENT_GLOBAL_SHADER(FFusionSetupCS, "/Engine/Private/PostProcessLocalExposure.usf", "FusionSetupCS", SF_Compute);
+
+class FFusionBlendCS : public FGlobalShader
+{
+public:
+	// Changing these numbers requires LocalExposure.usf to be recompiled
+	static const uint32 ThreadGroupSizeX = 8;
+	static const uint32 ThreadGroupSizeY = 8;
+
+	DECLARE_GLOBAL_SHADER(FFusionBlendCS);
+	SHADER_USE_PARAMETER_STRUCT(FFusionBlendCS, FGlobalShader);
+
+	class FLaplacianDim : SHADER_PERMUTATION_BOOL("LAPLACIAN");
+	using FPermutationDomain = TShaderPermutationDomain<FLaplacianDim>;
+
+	BEGIN_SHADER_PARAMETER_STRUCT(FParameters, )
+		SHADER_PARAMETER_STRUCT_REF(FViewUniformShaderParameters, View)
+		SHADER_PARAMETER_STRUCT(FEyeAdaptationParameters, EyeAdaptation)
+
+		SHADER_PARAMETER_STRUCT(FScreenPassTextureViewportParameters, Input)
+		SHADER_PARAMETER_RDG_TEXTURE_SRV(Texture2D, InputTexture)
+		SHADER_PARAMETER_RDG_TEXTURE_SRV(Texture2D, WeightTexture)
+
+		SHADER_PARAMETER_STRUCT(FScreenPassTextureViewportParameters, CoarserMip)
+		SHADER_PARAMETER_RDG_TEXTURE_SRV(Texture2D, CoarserMipTexture)
+		SHADER_PARAMETER_RDG_TEXTURE_SRV(Texture2D, PrevResultTexture)
+
+		SHADER_PARAMETER_STRUCT(FScreenPassTextureViewportParameters, Output)
+		SHADER_PARAMETER_RDG_TEXTURE_UAV(RWTexture2D<float>, OutputFloat)
+
+		SHADER_PARAMETER_SAMPLER(SamplerState, TextureSampler)
+
+		SHADER_PARAMETER(FScreenTransform, DispatchThreadToCoarseMipUV)
+	END_SHADER_PARAMETER_STRUCT()
+
+	static bool ShouldCompilePermutation(const FGlobalShaderPermutationParameters& Parameters)
+	{
+		return IsFeatureLevelSupported(Parameters.Platform, ERHIFeatureLevel::SM5);
+	}
+
+	static void ModifyCompilationEnvironment(const FGlobalShaderPermutationParameters& Parameters, FShaderCompilerEnvironment& OutEnvironment)
+	{
+		FGlobalShader::ModifyCompilationEnvironment(Parameters, OutEnvironment);
+		OutEnvironment.SetDefine(TEXT("THREADGROUP_SIZEX"), ThreadGroupSizeX);
+		OutEnvironment.SetDefine(TEXT("THREADGROUP_SIZEY"), ThreadGroupSizeY);
+	}
+};
+
+IMPLEMENT_GLOBAL_SHADER(FFusionBlendCS, "/Engine/Private/PostProcessLocalExposure.usf", "FusionBlendCS", SF_Compute);
+
 } //! namespace
 
 FVector2f GetLocalExposureBilateralGridUVScale(const FIntPoint ViewRectSize);
@@ -132,6 +222,15 @@ FLocalExposureParameters GetLocalExposureParameters(const FViewInfo& View, FIntP
 		}
 	}
 
+	if (View.FinalPostProcessSettings.LocalExposureMethod == ELocalExposureMethod::Fusion)
+	{
+		const float HighlightEV = FMath::Lerp<float>(10, 0, HighlightContrast);
+		HighlightContrast = FMath::Pow(2, -HighlightEV);
+
+		const float ShadowEV = FMath::Lerp<float>(10, 0, ShadowContrast);
+		ShadowContrast = FMath::Pow(2, ShadowEV);
+	}
+
 	FLocalExposureParameters Parameters;
 	Parameters.HighlightContrastScale = HighlightContrast;
 	Parameters.ShadowContrastScale = ShadowContrast;
@@ -242,7 +341,7 @@ void AddApplyLocalExposurePass(
 	PassParameters->LumBilateralGrid = LocalExposureTexture;
 	PassParameters->BlurredLogLum = BlurredLogLuminanceTexture;
 
-	PassParameters->TextureSampler = TStaticSamplerState<SF_Bilinear, AM_Clamp, AM_Clamp, AM_Clamp>::GetRHI();;
+	PassParameters->TextureSampler = TStaticSamplerState<SF_Bilinear, AM_Clamp, AM_Clamp, AM_Clamp>::GetRHI();
 
 	FComputeShaderUtils::AddPass(
 		GraphBuilder,
@@ -251,4 +350,127 @@ void AddApplyLocalExposurePass(
 		View.ShaderMap->GetShader<FApplyLocalExposureCS>(),
 		PassParameters,
 		FComputeShaderUtils::GetGroupCount(Output.ViewRect.Size(), FIntPoint(FApplyLocalExposureCS::ThreadGroupSizeX, FApplyLocalExposureCS::ThreadGroupSizeY)));
+}
+
+FScreenPassTexture AddLocalExposureFusionPass(
+	FRDGBuilder& GraphBuilder,
+	const FViewInfo& View,
+	const FEyeAdaptationParameters& EyeAdaptationParameters,
+	FRDGBufferRef EyeAdaptationBuffer,
+	const FLocalExposureParameters& LocalExposureParamaters,
+	FScreenPassTextureSlice Input)
+{
+	check(Input.IsValid());
+
+	RDG_EVENT_SCOPE(GraphBuilder, "LocalExposure - Fusion");
+
+	FScreenPassTexture LumTexture;
+	FScreenPassTexture WeightTexture;
+
+	{
+		const FRDGTextureDesc& InputDesc = Input.TextureSRV->GetParent()->Desc;
+
+		const FRDGTextureDesc TextureDesc = FRDGTextureDesc::Create2D(
+			InputDesc.Extent,
+			PF_FloatRGB,
+			FClearValueBinding::None,
+			TexCreate_UAV | TexCreate_ShaderResource);
+
+		// output uses same viewport as input
+		LumTexture = FScreenPassTexture(GraphBuilder.CreateTexture(TextureDesc, TEXT("LocalExposureLumTexture")), Input.ViewRect);
+		WeightTexture = FScreenPassTexture(GraphBuilder.CreateTexture(TextureDesc, TEXT("LocalExposureWeightTexture")), Input.ViewRect);
+
+		auto* PassParameters = GraphBuilder.AllocParameters<FFusionSetupCS::FParameters>();
+		PassParameters->View = View.ViewUniformBuffer;
+		PassParameters->EyeAdaptation = EyeAdaptationParameters;
+		PassParameters->EyeAdaptationBuffer = GraphBuilder.CreateSRV(EyeAdaptationBuffer);
+		PassParameters->LocalExposure = LocalExposureParamaters;
+		PassParameters->Input = GetScreenPassTextureViewportParameters(FScreenPassTextureViewport(Input));
+		PassParameters->InputTexture = Input.TextureSRV;
+		PassParameters->Output = GetScreenPassTextureViewportParameters(FScreenPassTextureViewport(LumTexture));
+		PassParameters->OutputFloat4 = GraphBuilder.CreateUAV(LumTexture.Texture);
+		PassParameters->OutputFloat4_1 = GraphBuilder.CreateUAV(WeightTexture.Texture);
+
+		FComputeShaderUtils::AddPass(
+			GraphBuilder,
+			RDG_EVENT_NAME("FusionSetup %dx%d", Input.ViewRect.Width(), Input.ViewRect.Height()),
+			ERDGPassFlags::Compute,
+			View.ShaderMap->GetShader<FFusionSetupCS>(),
+			PassParameters,
+			FComputeShaderUtils::GetGroupCount(Input.ViewRect.Size(), FIntPoint(FFusionSetupCS::ThreadGroupSizeX, FFusionSetupCS::ThreadGroupSizeY)));
+	}
+
+	const bool bLogLumaInAlpha = false;
+
+	FSceneDownsampleChain LumChain;
+	LumChain.Init(
+		GraphBuilder, View,
+		EyeAdaptationParameters,
+		FScreenPassTextureSlice::CreateFromScreenPassTexture(GraphBuilder, LumTexture),
+		EDownsampleQuality::High,
+		bLogLumaInAlpha);
+
+	FSceneDownsampleChain WeightChain;
+	WeightChain.Init(
+		GraphBuilder, View,
+		EyeAdaptationParameters,
+		FScreenPassTextureSlice::CreateFromScreenPassTexture(GraphBuilder, WeightTexture),
+		EDownsampleQuality::High,
+		bLogLumaInAlpha);
+
+	FScreenPassTexture Output;
+
+	FScreenPassTextureSlice CoarserMip;
+
+	for(int32 Index = FSceneDownsampleChain::StageCount - 1; Index >= 0; --Index)
+	{
+		FScreenPassTextureSlice CurrentLum = LumChain.GetTexture(Index);
+		FScreenPassTextureSlice CurrentWeight = WeightChain.GetTexture(Index);
+
+		FRDGTextureSRVRef PrevResult = Output.IsValid() ? GraphBuilder.CreateSRV(Output.Texture) : nullptr;
+
+		{
+			FRDGTextureDesc OutputDesc = CurrentLum.TextureSRV->GetParent()->Desc;
+			OutputDesc.Reset();
+			OutputDesc.Flags |= TexCreate_UAV;
+
+			// output uses same viewport as mip
+			Output = FScreenPassTexture(GraphBuilder.CreateTexture(OutputDesc, TEXT("LocalExposureResult")), CurrentLum.ViewRect);
+		}
+
+		auto* PassParameters = GraphBuilder.AllocParameters<FFusionBlendCS::FParameters>();
+		PassParameters->View = View.ViewUniformBuffer;
+		PassParameters->EyeAdaptation = EyeAdaptationParameters;
+		PassParameters->InputTexture = CurrentLum.TextureSRV;
+		PassParameters->WeightTexture = CurrentWeight.TextureSRV;
+		if (CoarserMip.IsValid())
+		{
+			PassParameters->DispatchThreadToCoarseMipUV =
+				FScreenTransform::DispatchThreadIdToViewportUV(Output.ViewRect) *
+				FScreenTransform::ChangeTextureBasisFromTo(FScreenPassTextureViewport(CoarserMip), FScreenTransform::ETextureBasis::ViewportUV, FScreenTransform::ETextureBasis::TextureUV);
+			PassParameters->CoarserMip = GetScreenPassTextureViewportParameters(FScreenPassTextureViewport(CoarserMip));
+			PassParameters->CoarserMipTexture = CoarserMip.TextureSRV;
+		}
+		PassParameters->PrevResultTexture = PrevResult;
+		PassParameters->Output = GetScreenPassTextureViewportParameters(FScreenPassTextureViewport(Output));
+		PassParameters->OutputFloat = GraphBuilder.CreateUAV(Output.Texture);
+		PassParameters->TextureSampler = TStaticSamplerState<SF_Bilinear, AM_Clamp, AM_Clamp, AM_Clamp>::GetRHI();
+
+		FFusionBlendCS::FPermutationDomain PermutationVector;
+		PermutationVector.Set<FFusionBlendCS::FLaplacianDim>(PrevResult != nullptr);
+
+		auto ComputeShader = View.ShaderMap->GetShader<FFusionBlendCS>(PermutationVector);
+
+		FComputeShaderUtils::AddPass(
+			GraphBuilder,
+			RDG_EVENT_NAME("FusionBlend %dx%d", Output.ViewRect.Width(), Output.ViewRect.Height()),
+			ERDGPassFlags::Compute,
+			ComputeShader,
+			PassParameters,
+			FComputeShaderUtils::GetGroupCount(Output.ViewRect.Size(), FIntPoint(FFusionBlendCS::ThreadGroupSizeX, FFusionBlendCS::ThreadGroupSizeY)));
+
+		CoarserMip = CurrentLum;
+	}
+
+	return MoveTemp(Output);
 }
\ No newline at end of file
diff --git a/Engine/Source/Runtime/Renderer/Private/PostProcess/PostProcessLocalExposure.h b/Engine/Source/Runtime/Renderer/Private/PostProcess/PostProcessLocalExposure.h
--- a/Engine/Source/Runtime/Renderer/Private/PostProcess/PostProcessLocalExposure.h
+++ b/Engine/Source/Runtime/Renderer/Private/PostProcess/PostProcessLocalExposure.h
@@ -36,3 +36,11 @@ void AddApplyLocalExposurePass(
 	FScreenPassTextureSlice Input,
 	FScreenPassTextureSlice Output,
 	ERDGPassFlags PassFlags);
+
+FScreenPassTexture AddLocalExposureFusionPass(
+	FRDGBuilder& GraphBuilder,
+	const FViewInfo& View,
+	const FEyeAdaptationParameters& EyeAdaptationParameters,
+	FRDGBufferRef EyeAdaptationBuffer,
+	const FLocalExposureParameters& LocalExposureParamaters,
+	FScreenPassTextureSlice Input);
\ No newline at end of file
diff --git a/Engine/Source/Runtime/Renderer/Private/PostProcess/PostProcessTonemap.cpp b/Engine/Source/Runtime/Renderer/Private/PostProcess/PostProcessTonemap.cpp
--- a/Engine/Source/Runtime/Renderer/Private/PostProcess/PostProcessTonemap.cpp
+++ b/Engine/Source/Runtime/Renderer/Private/PostProcess/PostProcessTonemap.cpp
@@ -84,7 +84,7 @@ namespace TonemapperPermutation
 {
 // Shared permutation dimensions between deferred and mobile renderer.
 class FTonemapperGammaOnlyDim      : SHADER_PERMUTATION_BOOL("USE_GAMMA_ONLY");
-class FTonemapperLocalExposureDim  : SHADER_PERMUTATION_BOOL("USE_LOCAL_EXPOSURE");
+class FTonemapperLocalExposureDim  : SHADER_PERMUTATION_INT("LOCAL_EXPOSURE_MODE", 3);
 class FTonemapperSharpenDim        : SHADER_PERMUTATION_BOOL("USE_SHARPEN");
 class FTonemapperFilmGrainDim      : SHADER_PERMUTATION_BOOL("USE_FILM_GRAIN");
 class FTonemapperMsaaDim           : SHADER_PERMUTATION_BOOL("METAL_MSAA_HDR_DECODE");
@@ -138,10 +138,19 @@ FCommonDomain BuildCommonPermutationDomain(const FViewInfo& View, bool bGammaOnl
 	}
 
 	const FPostProcessSettings& Settings = View.FinalPostProcessSettings;
-	PermutationVector.Set<FTonemapperLocalExposureDim>(bLocalExposure);
 	PermutationVector.Set<FTonemapperFilmGrainDim>(View.FilmGrainTexture != nullptr);
 	PermutationVector.Set<FTonemapperSharpenDim>(GetSharpenSetting(Settings) > 0.0f);
 	PermutationVector.Set<FTonemapperMsaaDim>(bMetalMSAAHDRDecode);
+
+	if (bLocalExposure)
+	{
+		PermutationVector.Set<FTonemapperLocalExposureDim>(View.FinalPostProcessSettings.LocalExposureMethod == ELocalExposureMethod::Bilateral ? 1 : 2);
+	}
+	else
+	{
+		PermutationVector.Set<FTonemapperLocalExposureDim>(0);
+	}
+
 	return PermutationVector;
 }
 
@@ -297,6 +306,11 @@ BEGIN_SHADER_PARAMETER_STRUCT(FTonemapParameters, )
 	SHADER_PARAMETER_SAMPLER(SamplerState, LumBilateralGridSampler)
 	SHADER_PARAMETER_SAMPLER(SamplerState, BlurredLogLumSampler)
 
+	SHADER_PARAMETER(FScreenTransform, ColorToExposureFusion)
+	SHADER_PARAMETER_STRUCT(FScreenPassTextureViewportParameters, ExposureFusion)
+	SHADER_PARAMETER_RDG_TEXTURE_SRV(Texture2D, ExposureFusionTexture)
+	SHADER_PARAMETER_SAMPLER(SamplerState, ExposureFusionSampler)
+
 	SHADER_PARAMETER_RDG_BUFFER_SRV(StructuredBuffer<float4>, EyeAdaptationBuffer)
 	SHADER_PARAMETER_RDG_TEXTURE(Texture2D, ColorGradingLUT)
 	SHADER_PARAMETER_TEXTURE(Texture2D, BloomDirtMaskTexture)
@@ -766,7 +780,9 @@ FScreenPassTexture AddTonemapPass(FRDGBuilder& GraphBuilder, const FViewInfo& Vi
 	
 	const float LUTSize = Inputs.ColorGradingTexture ? (float)Inputs.ColorGradingTexture->Desc.Extent.Y : /* unused (default): */ 32.0f;
 
-	const bool bLocalExposureEnabled = Inputs.LocalExposureTexture != nullptr;
+	const bool bBilateralLocalExposureEnabled = Inputs.LocalExposureBilateralGridTexture != nullptr;
+	const bool bExposureFusionEnabled = Inputs.ExposureFusion.IsValid();
+	const bool bLocalExposureEnabled = bBilateralLocalExposureEnabled || bExposureFusionEnabled;
 
 	if (bLocalExposureEnabled)
 	{
@@ -778,10 +794,18 @@ FScreenPassTexture AddTonemapPass(FRDGBuilder& GraphBuilder, const FViewInfo& Vi
 	CommonParameters.Color = GetScreenPassTextureViewportParameters(SceneColorViewport);
 	CommonParameters.Output = GetScreenPassTextureViewportParameters(OutputViewport);
 	CommonParameters.ColorTexture = Inputs.SceneColor.TextureSRV;
-	CommonParameters.LumBilateralGrid = Inputs.LocalExposureTexture;
+	CommonParameters.LumBilateralGrid = Inputs.LocalExposureBilateralGridTexture;
 	CommonParameters.BlurredLogLum = Inputs.BlurredLogLuminanceTexture;
 	CommonParameters.LumBilateralGridSampler = BilinearClampSampler;
 	CommonParameters.BlurredLogLumSampler = BilinearClampSampler;
+	if(Inputs.ExposureFusion.IsValid())
+	{
+		const FScreenPassTextureViewport ExposureFusionViewport(Inputs.ExposureFusion);
+		CommonParameters.ColorToExposureFusion = FScreenTransform::ChangeTextureUVCoordinateFromTo(SceneColorViewport, ExposureFusionViewport);
+		CommonParameters.ExposureFusion = GetScreenPassTextureViewportParameters(ExposureFusionViewport);
+		CommonParameters.ExposureFusionTexture = Inputs.ExposureFusion.TextureSRV;
+		CommonParameters.ExposureFusionSampler = BilinearClampSampler;
+	}
 	CommonParameters.EyeAdaptationBuffer = GraphBuilder.CreateSRV(EyeAdaptationBuffer);
 	CommonParameters.EyeAdaptation = *Inputs.EyeAdaptationParameters;
 	CommonParameters.ColorGradingLUT = Inputs.ColorGradingTexture;
diff --git a/Engine/Source/Runtime/Renderer/Private/PostProcess/PostProcessTonemap.h b/Engine/Source/Runtime/Renderer/Private/PostProcess/PostProcessTonemap.h
--- a/Engine/Source/Runtime/Renderer/Private/PostProcess/PostProcessTonemap.h
+++ b/Engine/Source/Runtime/Renderer/Private/PostProcess/PostProcessTonemap.h
@@ -42,11 +42,14 @@ struct FTonemapInputs
 	FRDGBufferRef SceneColorApplyParamaters = nullptr;
 
 	// [Optional] Luminance bilateral grid. If this is null, local exposure is disabled.
-	FRDGTextureRef LocalExposureTexture = nullptr;
+	FRDGTextureRef LocalExposureBilateralGridTexture = nullptr;
 
 	// [Optional] Blurred luminance texture used to calculate local exposure.
 	FRDGTextureRef BlurredLogLuminanceTexture = nullptr;
 
+	// [Optional] Exposure Fusion. If this is null, local exposure is disabled.
+	FScreenPassTextureSlice ExposureFusion;
+
 	// [Optional] Local exposure parameters.
 	const FLocalExposureParameters* LocalExposureParameters = nullptr;
 
diff --git a/Engine/Source/Runtime/Renderer/Private/PostProcess/PostProcessing.cpp b/Engine/Source/Runtime/Renderer/Private/PostProcess/PostProcessing.cpp
--- a/Engine/Source/Runtime/Renderer/Private/PostProcess/PostProcessing.cpp
+++ b/Engine/Source/Runtime/Renderer/Private/PostProcess/PostProcessing.cpp
@@ -365,8 +365,9 @@ void AddPostProcessingPasses(
 	// Histogram defaults to black because the histogram eye adaptation pass is used for the manual metering mode.
 	FRDGTextureRef HistogramTexture = BlackDummy.Texture;
 
-	FRDGTextureRef LocalExposureTexture = nullptr;
+	FRDGTextureRef LocalExposureBilateralGridTexture = nullptr;
 	FRDGTextureRef LocalExposureBlurredLogLumTexture = BlackDummy.Texture;
+	FScreenPassTexture ExposureFusion;
 
 	FVisualizeTemporalUpscalerInputs VisualizeTemporalUpscalerInputs;
 
@@ -707,7 +708,7 @@ void AddPostProcessingPasses(
 		const bool bFFTBloomEnabled = bBloomEnabled && IsFFTBloomEnabled(View);
 
 		const bool bBasicEyeAdaptationEnabled = bEyeAdaptationEnabled && (AutoExposureMethod == EAutoExposureMethod::AEM_Basic);
-		const bool bLocalExposureBlurredLum = bLocalExposureEnabled && View.FinalPostProcessSettings.LocalExposureBlurredLuminanceBlend > 0.0f;
+		const bool bLocalExposureBlurredLum = bLocalExposureEnabled && View.FinalPostProcessSettings.LocalExposureMethod == ELocalExposureMethod::Bilateral && View.FinalPostProcessSettings.LocalExposureBlurredLuminanceBlend > 0.0f;
 
 		const bool bProcessQuarterResolution = CVarPostProcessingQuarterResolutionDownsample.GetValueOnRenderThread() == 1;
 		const bool bProcessEighthResolution = CVarPostProcessingQuarterResolutionDownsample.GetValueOnRenderThread() == 2;
@@ -1088,9 +1089,9 @@ void AddPostProcessingPasses(
 		{
 			FScreenPassTextureSlice LocalExposureSceneColor = bProcessEighthResolution ? EighthResSceneColor : (bProcessQuarterResolution ? QuarterResSceneColor : HalfResSceneColor);
 
-			if (bLocalExposureEnabled)
+			if (bLocalExposureEnabled && View.FinalPostProcessSettings.LocalExposureMethod == ELocalExposureMethod::Bilateral)
 			{
-				LocalExposureTexture = AddLocalExposurePass(
+				LocalExposureBilateralGridTexture = AddLocalExposurePass(
 					GraphBuilder, View,
 					EyeAdaptationParameters,
 					LocalExposureSceneColor);
@@ -1152,7 +1153,7 @@ void AddPostProcessingPasses(
 				LocalExposureParameters,
 				SceneDownsampleChain.GetLastTexture(),
 				LastEyeAdaptationBuffer,
-				bLocalExposureEnabled);
+				bLocalExposureEnabled && View.FinalPostProcessSettings.LocalExposureMethod == ELocalExposureMethod::Bilateral);
 		}
 		// Add histogram eye adaptation pass even if no histogram exists to support the manual clamping mode.
 		else if (bEyeAdaptationEnabled)
@@ -1162,7 +1163,18 @@ void AddPostProcessingPasses(
 				EyeAdaptationParameters,
 				LocalExposureParameters,
 				HistogramTexture,
-				bLocalExposureEnabled);
+				bLocalExposureEnabled && View.FinalPostProcessSettings.LocalExposureMethod == ELocalExposureMethod::Bilateral);
+		}
+
+		if (bLocalExposureEnabled && View.FinalPostProcessSettings.LocalExposureMethod == ELocalExposureMethod::Fusion)
+		{
+			ExposureFusion = AddLocalExposureFusionPass(
+				GraphBuilder, View,
+				EyeAdaptationParameters,
+				EyeAdaptationBuffer,
+				LocalExposureParameters,
+				//bProcessQuarterResolution ? QuarterResSceneColor : HalfResSceneColor);
+				SceneColorSlice);
 		}
 
 		FScreenPassTexture Bloom;
@@ -1209,7 +1221,7 @@ void AddPostProcessingPasses(
 					EyeAdaptationParameters,
 					EyeAdaptationBuffer,
 					LocalExposureParameters,
-					CVarBloomApplyLocalExposure.GetValueOnRenderThread() ? LocalExposureTexture : nullptr,
+					CVarBloomApplyLocalExposure.GetValueOnRenderThread() ? LocalExposureBilateralGridTexture : nullptr,
 					LocalExposureBlurredLogLumTexture);
 
 				Bloom = Outputs.BloomTexture;
@@ -1217,7 +1229,7 @@ void AddPostProcessingPasses(
 			}
 			else
 			{
-				const bool bBloomSetupRequiredEnabled = View.FinalPostProcessSettings.BloomThreshold > -1.0f || LocalExposureTexture != nullptr;
+				const bool bBloomSetupRequiredEnabled = View.FinalPostProcessSettings.BloomThreshold > -1.0f || (LocalExposureBilateralGridTexture != nullptr && View.FinalPostProcessSettings.LocalExposureMethod == ELocalExposureMethod::Bilateral);
 
 				// Reuse the main scene downsample chain if setup isn't required for gaussian bloom.
 				if (SceneDownsampleChain.IsInitialized() && !bBloomSetupRequiredEnabled)
@@ -1237,7 +1249,7 @@ void AddPostProcessingPasses(
 						SetupPassInputs.EyeAdaptationBuffer = EyeAdaptationBuffer;
 						SetupPassInputs.EyeAdaptationParameters = &EyeAdaptationParameters;
 						SetupPassInputs.LocalExposureParameters = &LocalExposureParameters;
-						SetupPassInputs.LocalExposureTexture = CVarBloomApplyLocalExposure.GetValueOnRenderThread() ? LocalExposureTexture : nullptr;
+						SetupPassInputs.LocalExposureTexture = CVarBloomApplyLocalExposure.GetValueOnRenderThread() && View.FinalPostProcessSettings.LocalExposureMethod == ELocalExposureMethod::Bilateral ? LocalExposureBilateralGridTexture : nullptr;
 						SetupPassInputs.BlurredLogLuminanceTexture = LocalExposureBlurredLogLumTexture;
 						SetupPassInputs.Threshold = BloomThreshold;
 
@@ -1308,8 +1320,9 @@ void AddPostProcessingPasses(
 				PassInputs.SceneColor = SceneColorSlice;
 				PassInputs.Bloom = Bloom;
 				PassInputs.SceneColorApplyParamaters = SceneColorApplyParameters;
-				PassInputs.LocalExposureTexture = LocalExposureTexture;
+				PassInputs.LocalExposureBilateralGridTexture = LocalExposureBilateralGridTexture;
 				PassInputs.BlurredLogLuminanceTexture = LocalExposureBlurredLogLumTexture;
+				PassInputs.ExposureFusion = FScreenPassTextureSlice::CreateFromScreenPassTexture(GraphBuilder, ExposureFusion);
 				PassInputs.LocalExposureParameters = &LocalExposureParameters;
 				PassInputs.EyeAdaptationParameters = &EyeAdaptationParameters;
 				PassInputs.EyeAdaptationBuffer = EyeAdaptationBuffer;
@@ -1679,7 +1692,7 @@ void AddPostProcessingPasses(
 		PassSequence.AcceptOverrideIfLastPass(EPass::VisualizeLocalExposure, PassInputs.OverrideOutput);
 		PassInputs.SceneColor = SceneColor;
 		PassInputs.HDRSceneColor = FScreenPassTexture::CopyFromSlice(GraphBuilder, SceneColorBeforeTonemapSlice);
-		PassInputs.LumBilateralGridTexture = LocalExposureTexture;
+		PassInputs.LumBilateralGridTexture = LocalExposureBilateralGridTexture;
 		PassInputs.BlurredLumTexture = LocalExposureBlurredLogLumTexture;
 		PassInputs.LocalExposureParameters = &LocalExposureParameters;
 		PassInputs.EyeAdaptationBuffer = GetEyeAdaptationBuffer(GraphBuilder, View);
