Expose Exposure Fusion Target Luminance as a cvar.

diff --git a/Engine/Shaders/Private/PostProcessLocalExposure.usf b/Engine/Shaders/Private/PostProcessLocalExposure.usf
--- a/Engine/Shaders/Private/PostProcessLocalExposure.usf
+++ b/Engine/Shaders/Private/PostProcessLocalExposure.usf
@@ -67,6 +67,8 @@ float ExposureWeight(float V, float Target)
 	return exp2(-A * A / B);
 }
 
+float TargetLuminance;
+
 [numthreads(THREADGROUP_SIZEX, THREADGROUP_SIZEY, 1)]
 void FusionSetupCS(uint2 DispatchThreadId : SV_DispatchThreadID)
 {
@@ -92,12 +94,10 @@ void FusionSetupCS(uint2 DispatchThreadId : SV_DispatchThreadID)
 	float HighlightLuminanceValTonemapped = ExposureFusionTonemap(HighlightLuminanceVal);
 	float ShadowLuminanceValTonemapped = ExposureFusionTonemap(ShadowLuminanceVal);
 
-	const float Target = 0.50f;
-
-	float3 Weights = float3(ExposureWeight(LuminanceValTonemapped, Target), ExposureWeight(HighlightLuminanceValTonemapped, Target), ExposureWeight(ShadowLuminanceValTonemapped, Target));
+	float3 Weights = float3(ExposureWeight(LuminanceValTonemapped, TargetLuminance), ExposureWeight(HighlightLuminanceValTonemapped, TargetLuminance), ExposureWeight(ShadowLuminanceValTonemapped, TargetLuminance));
 	Weights /= dot(Weights, 1.0f); // normalize weights
 
-	OutputFloat4[OutputPixelPos] = float4((LuminanceValTonemapped), (HighlightLuminanceValTonemapped), (ShadowLuminanceValTonemapped), 0.0f);
+	OutputFloat4[OutputPixelPos] = float4(LuminanceValTonemapped, HighlightLuminanceValTonemapped, ShadowLuminanceValTonemapped, 0.0f);
 	OutputFloat4_1[OutputPixelPos] = float4(Weights, 0.0f);
 }
 
diff --git a/Engine/Source/Runtime/Renderer/Private/PostProcess/PostProcessLocalExposure.cpp b/Engine/Source/Runtime/Renderer/Private/PostProcess/PostProcessLocalExposure.cpp
--- a/Engine/Source/Runtime/Renderer/Private/PostProcess/PostProcessLocalExposure.cpp
+++ b/Engine/Source/Runtime/Renderer/Private/PostProcess/PostProcessLocalExposure.cpp
@@ -13,6 +13,12 @@
 #include "ShaderCompilerCore.h"
 #include "DataDrivenShaderPlatformInfo.h"
 
+static TAutoConsoleVariable<float> CVarExposureFusionTargetLuminance(
+	TEXT("r.LocalExposure.ExposureFusion.TargetLuminance"),
+	0.5f,
+	TEXT("Target Luminance used to determine the weight of each exposure."),
+	ECVF_RenderThreadSafe);
+
 namespace
 {
 
@@ -116,6 +122,8 @@ public:
 		SHADER_PARAMETER_STRUCT(FScreenPassTextureViewportParameters, Output)
 		SHADER_PARAMETER_RDG_TEXTURE_UAV(RWTexture2D<float4>, OutputFloat4)
 		SHADER_PARAMETER_RDG_TEXTURE_UAV(RWTexture2D<float4>, OutputFloat4_1)
+
+		SHADER_PARAMETER(float, TargetLuminance)
 	END_SHADER_PARAMETER_STRUCT()
 
 	static bool ShouldCompilePermutation(const FGlobalShaderPermutationParameters& Parameters)
@@ -390,6 +398,7 @@ FScreenPassTexture AddLocalExposureFusionPass(
 		PassParameters->Output = GetScreenPassTextureViewportParameters(FScreenPassTextureViewport(LumTexture));
 		PassParameters->OutputFloat4 = GraphBuilder.CreateUAV(LumTexture.Texture);
 		PassParameters->OutputFloat4_1 = GraphBuilder.CreateUAV(WeightTexture.Texture);
+		PassParameters->TargetLuminance = CVarExposureFusionTargetLuminance.GetValueOnRenderThread();
 
 		FComputeShaderUtils::AddPass(
 			GraphBuilder,
