Fix Distance Field Shadows rendering incorrectly in views with off-center projection, such as stereo views.
Intersecting object culling used optimized manual FOV calculations that did not account for off-center projection/asymmetrical FOVs like those used for HMDs. Added an additional path to account for this when necessary.

diff --git a/Engine/Shaders/Private/DistanceFieldShadowing.usf b/Engine/Shaders/Private/DistanceFieldShadowing.usf
--- a/Engine/Shaders/Private/DistanceFieldShadowing.usf
+++ b/Engine/Shaders/Private/DistanceFieldShadowing.usf
@@ -495,6 +495,7 @@ void CullObjectsToTileWithGather(
 
 	float ExpandRadius = 0;
 
+#if !STEREO_RENDERING
 	// Note: this code is assuming a centered projection, aka no translation present in ViewToClip
 	// Stereo rendering uses an off center projection
 	float2 TanViewFOV = GetTanHalfFieldOfView();
@@ -513,7 +514,20 @@ void CullObjectsToTileWithGather(
 	ViewTileMax2.xy = max(MinTileZ2 * UnitPlaneTileMax, MaxTileZ * UnitPlaneTileMax) + ExpandRadius;
 	ViewTileMin2.z = MinTileZ2 - ExpandRadius;
 	ViewTileMax2.z = MaxTileZ + ExpandRadius;
+#else
+	// For stereo, we need to do a more expensive transform to account for non-centered projection
+	float2 ScreenTileMin = ((GroupId.xy / NumGroups) - 0.5) * float2(2, -2);
+	float2 ScreenTileMax = (((GroupId.xy + 1) / NumGroups) - 0.5) * float2(2, -2);
 
+	ViewTileMin.xy = min(mul(float4(ScreenTileMin * MinTileZ, MinTileZ, 1), ScreenToView).xy, mul(float4(ScreenTileMin * MaxTileZ2, MaxTileZ2, 1), ScreenToView).xy) - ExpandRadius;
+	ViewTileMax.xy = max(mul(float4(ScreenTileMax * MinTileZ, MinTileZ, 1), ScreenToView).xy, mul(float4(ScreenTileMax * MaxTileZ2, MaxTileZ2, 1), ScreenToView).xy) + ExpandRadius;
+	ViewTileMin.z = MinTileZ - ExpandRadius;
+	ViewTileMax.z = MaxTileZ2 + ExpandRadius;
+	ViewTileMin2.xy = min(mul(float4(ScreenTileMin * MinTileZ2, MinTileZ2, 1), ScreenToView).xy, mul(float4(ScreenTileMin * MaxTileZ, MaxTileZ, 1), ScreenToView).xy) - ExpandRadius;
+	ViewTileMax2.xy = max(mul(float4(ScreenTileMax * MinTileZ2, MinTileZ2, 1), ScreenToView).xy, mul(float4(ScreenTileMax * MaxTileZ, MaxTileZ, 1), ScreenToView).xy) + ExpandRadius;
+	ViewTileMin2.z = MinTileZ2 - ExpandRadius;
+	ViewTileMax2.z = MaxTileZ + ExpandRadius;
+#endif
 	float3 ViewGroup0Center = (ViewTileMax + ViewTileMin) / 2;
 	float3 TranslatedWorldGroup0Center = mul(float4(ViewGroup0Center, 1), View.ViewToTranslatedWorld).xyz;
 	float Group0BoundingRadius = length(ViewGroup0Center - ViewTileMax);
@@ -595,6 +609,8 @@ float2 InvOutputBufferSize;
 Texture2D ShadowFactorsTexture;
 SamplerState ShadowFactorsSampler;
 
+float4x4 ScreenToView;
+
 [numthreads(THREADGROUP_SIZEX, THREADGROUP_SIZEY, 1)]
 void DistanceFieldShadowingCS(
 	uint3 GroupId : SV_GroupID,
diff --git a/Engine/Source/Runtime/Renderer/Private/DistanceFieldShadowing.cpp b/Engine/Source/Runtime/Renderer/Private/DistanceFieldShadowing.cpp
--- a/Engine/Source/Runtime/Renderer/Private/DistanceFieldShadowing.cpp
+++ b/Engine/Source/Runtime/Renderer/Private/DistanceFieldShadowing.cpp
@@ -31,6 +31,7 @@
 #include "ClearQuad.h"
 #include "Substrate/Substrate.h"
 #include "PixelShaderUtils.h"
+#include "Engine/Engine.h"
 
 int32 GDistanceFieldShadowing = 1;
 FAutoConsoleVariableRef CVarDistanceFieldShadowing(
@@ -312,6 +313,7 @@ class FDistanceFieldShadowingCS : public FGlobalShader
 		SHADER_PARAMETER(FVector2f, InvOutputBufferSize)
 		SHADER_PARAMETER_RDG_TEXTURE(Texture2D, ShadowFactorsTexture)
 		SHADER_PARAMETER_SAMPLER(SamplerState, ShadowFactorsSampler)
+		SHADER_PARAMETER(FMatrix44f, ScreenToView)
 	END_SHADER_PARAMETER_STRUCT()
 
 	class FCullingType : SHADER_PERMUTATION_INT("CULLING_TYPE", 3);
@@ -321,7 +323,8 @@ class FDistanceFieldShadowingCS : public FGlobalShader
 	class FOffsetDataStructure : SHADER_PERMUTATION_INT("OFFSET_DATA_STRUCT", 3);
 	class FCompactCulledObjects : SHADER_PERMUTATION_BOOL("COMPACT_CULLED_SHADOW_OBJECTS");
 	class FCullingSubSampleDepth : SHADER_PERMUTATION_BOOL("CULLING_SUBSAMPLE_DEPTH");
-	using FPermutationDomain = TShaderPermutationDomain<FCullingType, FShadowQuality, FPrimitiveType, FHasPreviousOutput, FOffsetDataStructure, FCompactCulledObjects, FCullingSubSampleDepth>;
+	class FStereoRendering : SHADER_PERMUTATION_BOOL("STEREO_RENDERING");
+	using FPermutationDomain = TShaderPermutationDomain<FCullingType, FShadowQuality, FPrimitiveType, FHasPreviousOutput, FOffsetDataStructure, FCompactCulledObjects, FCullingSubSampleDepth, FStereoRendering>;
 
 	static FPermutationDomain RemapPermutation(FPermutationDomain PermutationVector)
 	{
@@ -866,16 +869,20 @@ void RayTraceShadows(
 		PassParameters->InvOutputBufferSize = FVector2f(1.f / OutputBufferSize.X, 1.f / OutputBufferSize.Y);
 		PassParameters->ShadowFactorsTexture = PrevOutputTexture;
 		PassParameters->ShadowFactorsSampler = TStaticSamplerState<>::GetRHI();
+		PassParameters->ScreenToView = FMatrix44f(View.ViewMatrices.GetScreenToClipMatrix() * View.ViewMatrices.GetInvTranslatedViewProjectionMatrix() * View.ViewMatrices.GetOverriddenTranslatedViewMatrix());
 		
+
+
 		FDistanceFieldShadowingCS::FPermutationDomain PermutationVector;
-		PermutationVector.Set< FDistanceFieldShadowingCS::FCullingType >((uint32)DistanceFieldShadowingType);
-		PermutationVector.Set< FDistanceFieldShadowingCS::FShadowQuality >(DFShadowQuality);
-		PermutationVector.Set< FDistanceFieldShadowingCS::FPrimitiveType >(PrimitiveType);
-		PermutationVector.Set< FDistanceFieldShadowingCS::FHasPreviousOutput >(bHasPrevOutput);
+		PermutationVector.Set<FDistanceFieldShadowingCS::FCullingType>((uint32)DistanceFieldShadowingType);
+		PermutationVector.Set<FDistanceFieldShadowingCS::FShadowQuality>(DFShadowQuality);
+		PermutationVector.Set<FDistanceFieldShadowingCS::FPrimitiveType>(PrimitiveType);
+		PermutationVector.Set<FDistanceFieldShadowingCS::FHasPreviousOutput>(bHasPrevOutput);
 		extern int32 GDistanceFieldOffsetDataStructure;
-		PermutationVector.Set< FDistanceFieldShadowingCS::FOffsetDataStructure >(GDistanceFieldOffsetDataStructure);
+		PermutationVector.Set<FDistanceFieldShadowingCS::FOffsetDataStructure>(GDistanceFieldOffsetDataStructure);
 		PermutationVector.Set<FDistanceFieldShadowingCS::FCompactCulledObjects>(GDFShadowCompactCulledObjects != 0);
-		PermutationVector.Set< FDistanceFieldShadowingCS::FCullingSubSampleDepth>(GDFShadowCullingSubsampleDepth);
+		PermutationVector.Set<FDistanceFieldShadowingCS::FCullingSubSampleDepth>(GDFShadowCullingSubsampleDepth);
+		PermutationVector.Set<FDistanceFieldShadowingCS::FStereoRendering>(GEngine && GEngine->IsStereoscopic3D());
 
 		PermutationVector = FDistanceFieldShadowingCS::RemapPermutation(PermutationVector);
 
