Tonemap: Replacing POST_PROCESS_ALPHA with DIM_ALPHA_CHANNEL permutation for future support of runtime alpha toggle.

This continues the work done in CLs 35097313, 35395015, 35501865. Permutation is preferred over additional register use and dependent (alpha channel) texture fetch that a branch would incur, especially on consoles.

diff --git a/Engine/Shaders/Private/PostProcessTonemap.usf b/Engine/Shaders/Private/PostProcessTonemap.usf
--- a/Engine/Shaders/Private/PostProcessTonemap.usf
+++ b/Engine/Shaders/Private/PostProcessTonemap.usf
@@ -614,7 +614,7 @@ float4 TonemapCommonPS(
 	}
 	#endif
 
-	#if POST_PROCESS_ALPHA
+	#if DIM_ALPHA_CHANNEL
 		OutColor.a = SceneColor.a;
 	#endif
 
diff --git a/Engine/Source/Runtime/Renderer/Private/PostProcess/PostProcessTonemap.cpp b/Engine/Source/Runtime/Renderer/Private/PostProcess/PostProcessTonemap.cpp
--- a/Engine/Source/Runtime/Renderer/Private/PostProcess/PostProcessTonemap.cpp
+++ b/Engine/Source/Runtime/Renderer/Private/PostProcess/PostProcessTonemap.cpp
@@ -88,12 +88,14 @@ class FTonemapperLocalExposureDim  : SHADER_PERMUTATION_INT("LOCAL_EXPOSURE_MODE
 class FTonemapperSharpenDim        : SHADER_PERMUTATION_BOOL("USE_SHARPEN");
 class FTonemapperFilmGrainDim      : SHADER_PERMUTATION_BOOL("USE_FILM_GRAIN");
 class FTonemapperMsaaDim           : SHADER_PERMUTATION_BOOL("METAL_MSAA_HDR_DECODE");
+class FTonemapperAlphaChannelDim   : SHADER_PERMUTATION_BOOL("DIM_ALPHA_CHANNEL");
 using FCommonDomain = TShaderPermutationDomain<
 	FTonemapperGammaOnlyDim,
 	FTonemapperLocalExposureDim,
 	FTonemapperSharpenDim,
 	FTonemapperFilmGrainDim,
-	FTonemapperMsaaDim>;
+	FTonemapperMsaaDim,
+	FTonemapperAlphaChannelDim>;
 
 bool ShouldCompileCommonPermutation(const FGlobalShaderPermutationParameters& Parameters, const FCommonDomain& PermutationVector)
 {
@@ -127,6 +129,7 @@ FCommonDomain BuildCommonPermutationDomain(const FViewInfo& View, bool bGammaOnl
 	const FSceneViewFamily* Family = View.Family;
 
 	FCommonDomain PermutationVector;
+	PermutationVector.Set<FTonemapperAlphaChannelDim>(IsPostProcessingWithAlphaChannelSupported());
 
 	// Gamma
 	if (bGammaOnly ||
