SceneCapture: Allow scene color formats without alpha channel on scene captures that don't require it.

diff --git a/Engine/Source/Runtime/Renderer/Private/SceneRendering.cpp b/Engine/Source/Runtime/Renderer/Private/SceneRendering.cpp
--- a/Engine/Source/Runtime/Renderer/Private/SceneRendering.cpp
+++ b/Engine/Source/Runtime/Renderer/Private/SceneRendering.cpp
@@ -6643,12 +6643,14 @@ bool SceneCaptureRequiresAlphaChannel(const FSceneView& View)
 		if (View.CustomRenderPass)
 		{
 			return View.CustomRenderPass->GetRenderOutput() != FCustomRenderPassBase::ERenderOutput::SceneDepth
-				&& View.CustomRenderPass->GetRenderOutput() != FCustomRenderPassBase::ERenderOutput::DeviceDepth;
+				&& View.CustomRenderPass->GetRenderOutput() != FCustomRenderPassBase::ERenderOutput::DeviceDepth
+				&& View.CustomRenderPass->GetRenderOutput() != FCustomRenderPassBase::ERenderOutput::SceneColorNoAlpha;
 		}
 		else if(View.Family)
 		{
 			return View.Family->SceneCaptureSource != SCS_SceneDepth 
-				&& View.Family->SceneCaptureSource != SCS_DeviceDepth;
+				&& View.Family->SceneCaptureSource != SCS_DeviceDepth
+				&& View.Family->SceneCaptureSource != SCS_SceneColorHDRNoAlpha;
 		}
 	}
 	return false;
