[RHI] Replace GSupportsVolumeTextureRendering in conditions with RHISupportsVolumeTextureRendering() function calls to account for preview platform.

When the preview platform does not support geometry shaders (such as Metal) and the runtime RHI does not support vertex shader layers as alternative (such as D3D) we cannot support volume texture rendering.
This change moves this condition from the RHI global GSupportsVolumeTextureRendering to a function call RHISupportsVolumeTextureRendering() to include this difference between preview and runtime RHIs.
This is an appendix to CL 35396864.

#rnx

diff --git a/Engine/Source/Runtime/RHI/Private/RHI.cpp b/Engine/Source/Runtime/RHI/Private/RHI.cpp
--- a/Engine/Source/Runtime/RHI/Private/RHI.cpp
+++ b/Engine/Source/Runtime/RHI/Private/RHI.cpp
@@ -1493,6 +1493,18 @@ int32 RHIGetPreferredClearUAVRectPSResourceType(const FStaticShaderPlatform Plat
 	return 1; // TEXTURE_2D
 }
 
+RHI_API bool RHISupportsVolumeTextureRendering()
+{
+#if WITH_EDITOR
+	// When preview platforms are supported (when building with editor) we might be previewing an RHI that doesn't support geometry shaders (such as Metal)
+	// while rendering with a runtime RHI that doesn't support vertex shader layers as an alternative (such as D3D), so take these DDPI entries into account, too.
+	const EShaderPlatform Platform = GShaderPlatformForFeatureLevel[GMaxRHIFeatureLevel];
+	return GSupportsVolumeTextureRendering && (RHISupportsGeometryShaders(Platform) || RHISupportsVertexShaderLayer(Platform));
+#else
+	return GSupportsVolumeTextureRendering;
+#endif
+}
+
 void FRHIRenderPassInfo::ConvertToRenderTargetsInfo(FRHISetRenderTargetsInfo& OutRTInfo) const
 {
 	for (int32 Index = 0; Index < MaxSimultaneousRenderTargets; ++Index)
diff --git a/Engine/Source/Runtime/RHI/Public/RHI.h b/Engine/Source/Runtime/RHI/Public/RHI.h
--- a/Engine/Source/Runtime/RHI/Public/RHI.h
+++ b/Engine/Source/Runtime/RHI/Public/RHI.h
@@ -83,6 +83,9 @@ RHI_API int32 RHIGetPreferredClearUAVRectPSResourceType(const FStaticShaderPlatf
 // helper to force dump all RHI resource to CSV file
 RHI_API void RHIDumpResourceMemoryToCSV();
 
+// helper to check whether rendering to volume textures is supported. This takes both GSupportsVolumeTextureRendering and DDPI entries into account as this is dependent on the preview platform.
+RHI_API bool RHISupportsVolumeTextureRendering();
+
 struct FRHIResourceStats
 {
 	FName Name;
diff --git a/Engine/Source/Runtime/Renderer/Private/IndirectLightingCache.cpp b/Engine/Source/Runtime/Renderer/Private/IndirectLightingCache.cpp
--- a/Engine/Source/Runtime/Renderer/Private/IndirectLightingCache.cpp
+++ b/Engine/Source/Runtime/Renderer/Private/IndirectLightingCache.cpp
@@ -142,7 +142,7 @@ bool IsIndirectLightingCacheAllowed(ERHIFeatureLevel::Type InFeatureLevel)
 bool CanIndirectLightingCacheUseVolumeTexture(ERHIFeatureLevel::Type InFeatureLevel)
 {
 	// @todo Mac OS X/OpenGL: For OpenGL devices which don't support volume-texture rendering we need to use the simpler point indirect lighting shaders.
-	return InFeatureLevel >= ERHIFeatureLevel::SM5 && GSupportsVolumeTextureRendering;
+	return InFeatureLevel >= ERHIFeatureLevel::SM5 && RHISupportsVolumeTextureRendering();
 }
 
 FIndirectLightingCache::FIndirectLightingCache(ERHIFeatureLevel::Type InFeatureLevel)
diff --git a/Engine/Source/Runtime/Renderer/Private/LightMapRendering.cpp b/Engine/Source/Runtime/Renderer/Private/LightMapRendering.cpp
--- a/Engine/Source/Runtime/Renderer/Private/LightMapRendering.cpp
+++ b/Engine/Source/Runtime/Renderer/Private/LightMapRendering.cpp
@@ -760,7 +760,7 @@ void GetIndirectLightingCacheParameters(
 		// If we are using FCachedVolumeIndirectLightingPolicy then InitViews should have updated the lighting cache which would have initialized it
 		// However the conditions for updating the lighting cache are complex and fail very occasionally in non-reproducible ways
 		// Silently skipping setting the cache texture under failure for now
-		if (FeatureLevel >= ERHIFeatureLevel::SM5 && LightingCache && LightingCache->IsInitialized() && GSupportsVolumeTextureRendering)
+		if (FeatureLevel >= ERHIFeatureLevel::SM5 && LightingCache && LightingCache->IsInitialized() && RHISupportsVolumeTextureRendering())
 		{
 			Parameters.IndirectLightingCacheTexture0 = const_cast<FIndirectLightingCache*>(LightingCache)->GetTexture0();
 			Parameters.IndirectLightingCacheTexture1 = const_cast<FIndirectLightingCache*>(LightingCache)->GetTexture1();
diff --git a/Engine/Source/Runtime/Renderer/Private/LightRendering.cpp b/Engine/Source/Runtime/Renderer/Private/LightRendering.cpp
--- a/Engine/Source/Runtime/Renderer/Private/LightRendering.cpp
+++ b/Engine/Source/Runtime/Renderer/Private/LightRendering.cpp
@@ -1408,7 +1408,7 @@ void FDeferredShadingSceneRenderer::RenderLights(
 			Substrate::AddSubstrateStencilPass(GraphBuilder, Views, SceneTextures);
 		}
 
-		if (ViewFamily.EngineShowFlags.DirectLighting && GUseTranslucentLightingVolumes && GSupportsVolumeTextureRendering)
+		if (ViewFamily.EngineShowFlags.DirectLighting && GUseTranslucentLightingVolumes && RHISupportsVolumeTextureRendering())
 		{
 			RDG_EVENT_SCOPE(GraphBuilder, "InjectTranslucencyLightingVolume");
 
diff --git a/Engine/Source/Runtime/Renderer/Private/TranslucentLighting.cpp b/Engine/Source/Runtime/Renderer/Private/TranslucentLighting.cpp
--- a/Engine/Source/Runtime/Renderer/Private/TranslucentLighting.cpp
+++ b/Engine/Source/Runtime/Renderer/Private/TranslucentLighting.cpp
@@ -922,7 +922,7 @@ void InjectTranslucencyLightingVolumeAmbientCubemap(
 	const TArrayView<const FViewInfo> Views,
 	const FTranslucencyLightingVolumeTextures& Textures)
 {
-	if (!GUseTranslucentLightingVolumes || !GSupportsVolumeTextureRendering)
+	if (!GUseTranslucentLightingVolumes || !RHISupportsVolumeTextureRendering())
 	{
 		return;
 	}
@@ -1091,7 +1091,7 @@ void InjectTranslucencyLightingVolume(
 	const FTranslucencyLightingVolumeTextures& Textures,
 	const FTranslucentLightInjectionCollector& Collector)
 {
-	if (!GUseTranslucentLightingVolumes || !GSupportsVolumeTextureRendering)
+	if (!GUseTranslucentLightingVolumes || !RHISupportsVolumeTextureRendering())
 	{
 		return;
 	}
@@ -1382,7 +1382,7 @@ void FilterTranslucencyLightingVolume(
 	const TArrayView<const FViewInfo> Views,
 	FTranslucencyLightingVolumeTextures& Textures)
 {
-	if (!GUseTranslucentLightingVolumes || !GSupportsVolumeTextureRendering || !GUseTranslucencyVolumeBlur)
+	if (!GUseTranslucentLightingVolumes || !RHISupportsVolumeTextureRendering() || !GUseTranslucencyVolumeBlur)
 	{
 		return;
 	}
diff --git a/Engine/Source/Runtime/SlateRHIRenderer/Private/SlateRHIRenderer.cpp b/Engine/Source/Runtime/SlateRHIRenderer/Private/SlateRHIRenderer.cpp
--- a/Engine/Source/Runtime/SlateRHIRenderer/Private/SlateRHIRenderer.cpp
+++ b/Engine/Source/Runtime/SlateRHIRenderer/Private/SlateRHIRenderer.cpp
@@ -560,7 +560,7 @@ bool CompositeUIWithSceneHDR()
 	static const auto CVarCompositeMode = IConsoleManager::Get().FindTConsoleVariableDataInt(TEXT("r.HDR.UI.CompositeMode"));
 
 	return GRHISupportsHDROutput
-		&& GSupportsVolumeTextureRendering
+		&& RHISupportsVolumeTextureRendering()
 		&& SupportsCompositeUIWithSceneHDR(GetFeatureLevelShaderPlatform(GMaxRHIFeatureLevel))
 		&& CVarCompositeMode
 		&& CVarCompositeMode->GetValueOnAnyThread() != 0;
