Substrate - fixed visualisation when non real time update.

#rnx

diff --git a/Engine/Shaders/Private/Substrate/SubstrateVisualize.usf b/Engine/Shaders/Private/Substrate/SubstrateVisualize.usf
--- a/Engine/Shaders/Private/Substrate/SubstrateVisualize.usf
+++ b/Engine/Shaders/Private/Substrate/SubstrateVisualize.usf
@@ -679,6 +679,7 @@ void MaterialPrintInfoCS(uint3 DispatchThreadId : SV_DispatchThreadID)
 #if SHADER_MATERIALCOUNT
 
 uint ViewMode;
+uint bRealTimeUpdate;
 
 float3 BytePerPixelToColor(uint In)
 {
@@ -706,7 +707,7 @@ void VisualizeMaterialPS(
 	float4 SVPos : SV_POSITION,
 	out float4 OutColor : SV_Target0)
 {
-	const uint2 PixelPosDynRes = uint2(float2(SVPos.xy) * View.ViewResolutionFraction);
+	const uint2 PixelPosDynRes = uint2(float2(SVPos.xy) * (bRealTimeUpdate > 0 ? View.ViewResolutionFraction : 1.0f));
 	float2 BufferUV = SvPositionToBufferUV(float4(PixelPosDynRes, SVPos.zw));
 
 	const float2 ScreenPosition = SvPositionToScreenPosition(SVPos).xy;
diff --git a/Engine/Source/Runtime/Renderer/Private/Substrate/SubstrateVisualize.cpp b/Engine/Source/Runtime/Renderer/Private/Substrate/SubstrateVisualize.cpp
--- a/Engine/Source/Runtime/Renderer/Private/Substrate/SubstrateVisualize.cpp
+++ b/Engine/Source/Runtime/Renderer/Private/Substrate/SubstrateVisualize.cpp
@@ -91,6 +91,7 @@ class FVisualizeMaterialCountPS : public FGlobalShader
 
 	BEGIN_SHADER_PARAMETER_STRUCT(FParameters, )
 		SHADER_PARAMETER(uint32, ViewMode)
+		SHADER_PARAMETER(uint32, bRealTimeUpdate)
 		SHADER_PARAMETER(uint32, bOverrideCursorPosition)
 		SHADER_PARAMETER_STRUCT_REF(FViewUniformShaderParameters, ViewUniformBuffer)
 		SHADER_PARAMETER_RDG_UNIFORM_BUFFER(FSubstrateGlobalUniformParameters, Substrate)
@@ -315,6 +316,7 @@ static void AddVisualizeMaterialCountPasses(FRDGBuilder & GraphBuilder, const FV
 	FVisualizeMaterialCountPS::FParameters* PassParameters = GraphBuilder.AllocParameters<FVisualizeMaterialCountPS::FParameters>();
 	PassParameters->ViewUniformBuffer = View.ViewUniformBuffer;
 	PassParameters->ViewMode = FMath::Clamp(ViewMode, 2, 3);
+	PassParameters->bRealTimeUpdate = View.Family->bRealtimeUpdate ? 1 : 0;
 	PassParameters->bOverrideCursorPosition = WITH_EDITOR ? 0u : 1u;
 	PassParameters->Substrate = Substrate::BindSubstrateGlobalUniformParameters(View);
 	PassParameters->SceneTextures = GetSceneTextureParameters(GraphBuilder, View);
