Alpha: Renamed post-processing alpha setting to "Enable Alpha Output", and enforce alpha presence in scene color format when enabled like mobile.

diff --git a/Engine/Plugins/Experimental/Avalanche/Source/AvalancheMediaEditor/Private/Rundown/Preview/SAvaRundownPagePreview.cpp b/Engine/Plugins/Experimental/Avalanche/Source/AvalancheMediaEditor/Private/Rundown/Preview/SAvaRundownPagePreview.cpp
--- a/Engine/Plugins/Experimental/Avalanche/Source/AvalancheMediaEditor/Private/Rundown/Preview/SAvaRundownPagePreview.cpp
+++ b/Engine/Plugins/Experimental/Avalanche/Source/AvalancheMediaEditor/Private/Rundown/Preview/SAvaRundownPagePreview.cpp
@@ -308,7 +308,7 @@ void SAvaRundownPagePreview::HandleCheckerboardActionExecute() const
 		{
 			const FText NotificationText = LOCTEXT("AlphaSupport",
 				"An output requested Alpha Support but the required project setting is not enabled!\n"
-				"Go to Project Settings > Rendering > PostProcessing > 'Enable Alpha Channel Support in Post Processing' and enable it.");
+				"Go to Project Settings > Rendering > PostProcessing > 'Enable Alpha Output' and enable it.");
 
 			FNotificationInfo Info(NotificationText);
 			Info.ExpireDuration = 5.0f;	// The message is long, need more time to read it.
diff --git a/Engine/Plugins/MovieScene/MovieRenderPipeline/Source/MovieRenderPipelineCore/Private/Graph/MovieGraphRenderLayerSubsystem.cpp b/Engine/Plugins/MovieScene/MovieRenderPipeline/Source/MovieRenderPipelineCore/Private/Graph/MovieGraphRenderLayerSubsystem.cpp
--- a/Engine/Plugins/MovieScene/MovieRenderPipeline/Source/MovieRenderPipelineCore/Private/Graph/MovieGraphRenderLayerSubsystem.cpp
+++ b/Engine/Plugins/MovieScene/MovieRenderPipeline/Source/MovieRenderPipelineCore/Private/Graph/MovieGraphRenderLayerSubsystem.cpp
@@ -195,7 +195,7 @@ void UMovieGraphRenderPropertyModifier::ApplyModifier(const UWorld* World)
 	if (bHoldout && (RendererSettings->bEnableAlphaChannelInPostProcessing == false))
 	{
 		// TODO: Ideally this is called in a general-purpose validation step instead, but that framework does not exist yet.
-		UE_LOG(LogMovieRenderPipeline, Warning, TEXT("A modifier with 'Holdout' is active, but 'Enable alpha channel support in post processing' in "
+		UE_LOG(LogMovieRenderPipeline, Warning, TEXT("A modifier with 'Holdout' is active, but 'Enable Alpha Output' in "
 													 "the project's Rendering settings is disabled: holdout will not work properly."));
 	}
 	
diff --git a/Engine/Plugins/MovieScene/MovieRenderPipeline/Source/MovieRenderPipelineCore/Private/MoviePipeline.cpp b/Engine/Plugins/MovieScene/MovieRenderPipeline/Source/MovieRenderPipelineCore/Private/MoviePipeline.cpp
--- a/Engine/Plugins/MovieScene/MovieRenderPipeline/Source/MovieRenderPipelineCore/Private/MoviePipeline.cpp
+++ b/Engine/Plugins/MovieScene/MovieRenderPipeline/Source/MovieRenderPipelineCore/Private/MoviePipeline.cpp
@@ -122,7 +122,7 @@ void UMoviePipeline::ValidateSequenceAndSettings() const
 
 		if (bAnyOutputWantsAlpha && TonemapAlphaCVar->GetBool() == false)
 		{
-			UE_LOG(LogMovieRenderPipeline, Warning, TEXT("An output requested Alpha Support but the required project setting is not enabled! Go to Project Settings > Rendering > PostProcessing > 'Enable Alpha Channel Support in Post Processing' and set it to 'True'."));
+			UE_LOG(LogMovieRenderPipeline, Warning, TEXT("An output requested Alpha Support but the required project setting is not enabled! Go to Project Settings > Rendering > PostProcessing > 'Enable Alpha Output' and set it to 'True'."));
 		}
 	}
 }
diff --git a/Engine/Plugins/MovieScene/MovieRenderPipeline/Source/MovieRenderPipelineRenderPasses/Public/MoviePipelineDeferredPasses.h b/Engine/Plugins/MovieScene/MovieRenderPipeline/Source/MovieRenderPipelineRenderPasses/Public/MoviePipelineDeferredPasses.h
--- a/Engine/Plugins/MovieScene/MovieRenderPipeline/Source/MovieRenderPipelineRenderPasses/Public/MoviePipelineDeferredPasses.h
+++ b/Engine/Plugins/MovieScene/MovieRenderPipeline/Source/MovieRenderPipelineRenderPasses/Public/MoviePipelineDeferredPasses.h
@@ -101,7 +101,7 @@ protected:
 public:
 	/**
 	* Should multiple temporal/spatial samples accumulate the alpha channel? This requires r.PostProcessing.PropagateAlpha
-	* to be enabled (see "Enable Alpha Channel Support in Post Processing" under Project Settings > Rendering). This adds
+	* to be enabled (see "Enable Alpha Output" under Project Settings > Rendering). This adds
 	* ~30% cost to the accumulation so you should not enable it unless necessary. You must delete both the sky and fog to ensure
 	* that they do not make all pixels opaque.
 	*/
diff --git a/Engine/Source/Runtime/Engine/Classes/Components/ExponentialHeightFogComponent.h b/Engine/Source/Runtime/Engine/Classes/Components/ExponentialHeightFogComponent.h
--- a/Engine/Source/Runtime/Engine/Classes/Components/ExponentialHeightFogComponent.h
+++ b/Engine/Source/Runtime/Engine/Classes/Components/ExponentialHeightFogComponent.h
@@ -213,7 +213,7 @@ class UExponentialHeightFogComponent : public USceneComponent
 	UPROPERTY(EditAnywhere, BlueprintReadOnly, Category = VolumetricFog, AdvancedDisplay)
 	bool bOverrideLightColorsWithFogInscatteringColors;
 
-	/** If this is True, this primitive will render black with an alpha of 0, but all secondary effects (shadows, reflections, indirect lighting) remain. This feature required the project setting "Enable alpha channel support in post processing". */
+	/** If this is True, this primitive will render black with an alpha of 0, but all secondary effects (shadows, reflections, indirect lighting) remain. This feature required the project setting "Enable Alpha Output". */
 	UPROPERTY(EditAnywhere, AdvancedDisplay, BlueprintReadOnly, Category = Rendering, Interp)
 	uint8 bHoldout : 1;
 
diff --git a/Engine/Source/Runtime/Engine/Classes/Components/PrimitiveComponent.h b/Engine/Source/Runtime/Engine/Classes/Components/PrimitiveComponent.h
--- a/Engine/Source/Runtime/Engine/Classes/Components/PrimitiveComponent.h
+++ b/Engine/Source/Runtime/Engine/Classes/Components/PrimitiveComponent.h
@@ -428,7 +428,7 @@ public:
 	UPROPERTY(EditAnywhere, AdvancedDisplay, BlueprintReadOnly, Category=Rendering)
 	uint8 bReceivesDecals:1;
 
-	/** If this is True, this primitive will render black with an alpha of 0, but all secondary effects (shadows, reflections, indirect lighting) remain. This feature required the project setting "Enable alpha channel support in post processing". */
+	/** If this is True, this primitive will render black with an alpha of 0, but all secondary effects (shadows, reflections, indirect lighting) remain. This feature required the project setting "Enable Alpha Output". */
 	UPROPERTY(EditAnywhere, AdvancedDisplay, BlueprintReadOnly, Category = Rendering, Interp)
 	uint8 bHoldout : 1;
 
diff --git a/Engine/Source/Runtime/Engine/Classes/Components/SkyAtmosphereComponent.h b/Engine/Source/Runtime/Engine/Classes/Components/SkyAtmosphereComponent.h
--- a/Engine/Source/Runtime/Engine/Classes/Components/SkyAtmosphereComponent.h
+++ b/Engine/Source/Runtime/Engine/Classes/Components/SkyAtmosphereComponent.h
@@ -160,7 +160,7 @@ class USkyAtmosphereComponent : public USceneComponent
 	UPROPERTY(EditAnywhere, BlueprintReadOnly, interp, Category = "Art Direction", meta = (UIMin = 0.001f, UIMax = 10.0f, ClampMin = 0.001f))
 	float AerialPerspectiveStartDepth;
 
-	/** If this is True, this primitive will render black with an alpha of 0, but all secondary effects (shadows, reflections, indirect lighting) remain. This feature required the project setting "Enable alpha channel support in post processing". */
+	/** If this is True, this primitive will render black with an alpha of 0, but all secondary effects (shadows, reflections, indirect lighting) remain. This feature required the project setting "Enable Alpha Output". */
 	UPROPERTY(EditAnywhere, AdvancedDisplay, BlueprintReadOnly, Category = Rendering, Interp)
 	uint8 bHoldout : 1;
 
diff --git a/Engine/Source/Runtime/Engine/Classes/Components/VolumetricCloudComponent.h b/Engine/Source/Runtime/Engine/Classes/Components/VolumetricCloudComponent.h
--- a/Engine/Source/Runtime/Engine/Classes/Components/VolumetricCloudComponent.h
+++ b/Engine/Source/Runtime/Engine/Classes/Components/VolumetricCloudComponent.h
@@ -139,7 +139,7 @@ class UVolumetricCloudComponent : public USceneComponent
 	UPROPERTY(EditAnywhere, BlueprintReadOnly, interp, Category = "Art Direction", meta = (UIMin = 0.0, UIMax = 100.0, ClampMin = 0.0, SliderExponent = 2.0))
 	float AerialPespectiveMieScatteringFadeDistance;
 
-	/** If this is True, this primitive will render black with an alpha of 0, but all secondary effects (shadows, reflections, indirect lighting) remain. This feature required the project setting "Enable alpha channel support in post processing". */
+	/** If this is True, this primitive will render black with an alpha of 0, but all secondary effects (shadows, reflections, indirect lighting) remain. This feature required the project setting "Enable Alpha Output". */
 	UPROPERTY(EditAnywhere, AdvancedDisplay, BlueprintReadOnly, Category = Rendering, Interp)
 	uint8 bHoldout : 1;
 
diff --git a/Engine/Source/Runtime/Engine/Classes/Engine/RendererSettings.h b/Engine/Source/Runtime/Engine/Classes/Engine/RendererSettings.h
--- a/Engine/Source/Runtime/Engine/Classes/Engine/RendererSettings.h
+++ b/Engine/Source/Runtime/Engine/Classes/Engine/RendererSettings.h
@@ -726,8 +726,8 @@ class URendererSettings : public UDeveloperSettings
 	uint32 bCustomDepthTaaJitter : 1;
 
 	UPROPERTY(config, EditAnywhere, Category = Postprocessing, meta = (
-		ConsoleVariable = "r.PostProcessing.PropagateAlpha", DisplayName = "Enable alpha channel support in post processing",
-		ToolTip = "Propagates the alpha channel through the renderer's post-processing chain.",
+		ConsoleVariable = "r.PostProcessing.PropagateAlpha", DisplayName = "Enable Alpha Output",
+		ToolTip = "Enable r.PostProcessing.PropagateAlpha to enforce alpha in scene color (overriding r.SceneColorFormat if necessary) and propagate it through the renderer's post-processing chain.  The legacy \"Linear color space only\" and \"Allow through tonemapper\" options now map to true and the engine keeps full alpha precision throughout post-processing.",
 		ConfigRestartRequired = true))
 	bool bEnableAlphaChannelInPostProcessing;
 
diff --git a/Engine/Source/Runtime/Engine/Private/SceneTexturesConfig.cpp b/Engine/Source/Runtime/Engine/Private/SceneTexturesConfig.cpp
--- a/Engine/Source/Runtime/Engine/Private/SceneTexturesConfig.cpp
+++ b/Engine/Source/Runtime/Engine/Private/SceneTexturesConfig.cpp
@@ -11,6 +11,7 @@
 
 FSceneTexturesConfig FSceneTexturesConfig::GlobalInstance;
 
+DEFINE_LOG_CATEGORY_STATIC(LogSceneTextures, Log, All);
 IMPLEMENT_STATIC_UNIFORM_BUFFER_SLOT(SceneTextures);
 IMPLEMENT_STATIC_UNIFORM_BUFFER_STRUCT(FSceneTextureUniformParameters, "SceneTexturesStruct", SceneTextures);
 IMPLEMENT_STATIC_UNIFORM_BUFFER_STRUCT(FMobileSceneTextureUniformParameters, "MobileSceneTextures", SceneTextures);
@@ -123,8 +124,13 @@ static EPixelFormat GetSceneColorFormat(bool bRequiresAlphaChannel)
 		Format = PF_FloatRGBA;
 	}
 
-	if (bRequiresAlphaChannel)
+	if (bRequiresAlphaChannel && Format != PF_FloatRGBA)
 	{
+		UE_CALL_ONCE([]()
+			{
+				UE_LOG(LogSceneTextures, Warning, TEXT("Enforcing FloatRGBA scene color format due to alpha channel requirement."));
+			}
+		);
 		Format = PF_FloatRGBA;
 	}
 
diff --git a/Engine/Source/Runtime/Landscape/Classes/LandscapeProxy.h b/Engine/Source/Runtime/Landscape/Classes/LandscapeProxy.h
--- a/Engine/Source/Runtime/Landscape/Classes/LandscapeProxy.h
+++ b/Engine/Source/Runtime/Landscape/Classes/LandscapeProxy.h
@@ -794,7 +794,7 @@ public:
 	UPROPERTY(EditAnywhere, AdvancedDisplay, Category = Lighting, meta = (LandscapeOverridable, ClampMin = "0.0", ClampMax = "10.0", UIMin = "0.0", UIMax = "10.0"))
 	float NonNaniteVirtualShadowMapInvalidationScreenSizeLimit = 0.2f;
 
-	/** If this is True, this primitive will render black with an alpha of 0, but all secondary effects (shadows, reflections, indirect lighting) remain. This feature required the project setting "Enable alpha channel support in post processing". */
+	/** If this is True, this primitive will render black with an alpha of 0, but all secondary effects (shadows, reflections, indirect lighting) remain. This feature required the project setting "Enable Alpha Output". */
 	UPROPERTY(EditAnywhere, AdvancedDisplay, Category = Rendering, meta=(LandscapeOverridable), Interp)
 	uint8 bHoldout : 1;
 
diff --git a/Engine/Source/Runtime/Renderer/Private/PostProcess/PostProcessing.cpp b/Engine/Source/Runtime/Renderer/Private/PostProcess/PostProcessing.cpp
--- a/Engine/Source/Runtime/Renderer/Private/PostProcess/PostProcessing.cpp
+++ b/Engine/Source/Runtime/Renderer/Private/PostProcess/PostProcessing.cpp
@@ -117,7 +117,7 @@ static bool GPostProcessingPropagateAlpha = false;
 FAutoConsoleVariableRef CVarPostProcessingPropagateAlpha(
 	TEXT("r.PostProcessing.PropagateAlpha"),
 	GPostProcessingPropagateAlpha,
-	TEXT("Whether to propagate alpha through post-processing, default: false"),
+	TEXT("Enforce alpha in scene color (overriding r.SceneColorFormat if necessary) and propagate it through the renderer's post-processing chain, default: false"),
 	ECVF_ReadOnly | ECVF_RenderThreadSafe);
 
 TAutoConsoleVariable<int32> CVarPostProcessingPreferCompute(
diff --git a/Engine/Source/Runtime/Renderer/Private/SceneTextures.cpp b/Engine/Source/Runtime/Renderer/Private/SceneTextures.cpp
--- a/Engine/Source/Runtime/Renderer/Private/SceneTextures.cpp
+++ b/Engine/Source/Runtime/Renderer/Private/SceneTextures.cpp
@@ -17,6 +17,7 @@
 #include "PostProcess/PostProcessCompositeEditorPrimitives.h"
 #include "ShaderCompiler.h"
 #include "SystemTextures.h"
+#include "PostProcess/PostProcessing.h"
 #include "PostProcess/PostProcessAmbientOcclusionMobile.h"
 #include "PostProcess/PostProcessPixelProjectedReflectionMobile.h"
 #include "IHeadMountedDisplayModule.h"
@@ -376,7 +377,7 @@ void InitializeSceneTexturesConfig(FSceneTexturesConfig& Config, const FSceneVie
 	FIntPoint Extent = FSceneTextureExtentState::Get().Compute(ViewFamily);
 	EShadingPath ShadingPath = GetFeatureLevelShadingPath(ViewFamily.GetFeatureLevel());
 
-	bool bRequiresAlphaChannel = ShadingPath == EShadingPath::Mobile ? IsMobilePropagateAlphaEnabled(ViewFamily.GetShaderPlatform()) : false;
+	bool bRequiresAlphaChannel = ShadingPath == EShadingPath::Mobile ? IsMobilePropagateAlphaEnabled(ViewFamily.GetShaderPlatform()) : IsPostProcessingWithAlphaChannelSupported();
 	int32 NumberOfViewsWithMultiviewEnabled = 0;
 	
 	for (const FSceneView* View : ViewFamily.AllViews)
