Substrate - Fixed substrate compilation with SM5 and debug.

#rnx

diff --git a/Engine/Source/Runtime/Renderer/Private/DebugProbeRendering.cpp b/Engine/Source/Runtime/Renderer/Private/DebugProbeRendering.cpp
--- a/Engine/Source/Runtime/Renderer/Private/DebugProbeRendering.cpp
+++ b/Engine/Source/Runtime/Renderer/Private/DebugProbeRendering.cpp
@@ -67,7 +67,8 @@ public:
 
 	static bool ShouldCompilePermutation(const FGlobalShaderPermutationParameters& Parameters)
 	{
-		return IsFeatureLevelSupported(Parameters.Platform, ERHIFeatureLevel::SM5) && EnumHasAllFlags(Parameters.Flags, EShaderPermutationFlags::HasEditorOnlyData);
+		// Only SM6 since SM5 is limited to 8 UAV.
+		return IsFeatureLevelSupported(Parameters.Platform, ERHIFeatureLevel::SM6) && EnumHasAllFlags(Parameters.Flags, EShaderPermutationFlags::HasEditorOnlyData);
 	}
 	static void ModifyCompilationEnvironment(const FGlobalShaderPermutationParameters& Parameters, FShaderCompilerEnvironment& OutEnvironment)
 	{
