Substrate - Fixed shadow castin on the ground not matchign each light when two atmospheric lights are used.

#rnx

diff --git a/Engine/Shaders/Private/VolumetricCloud.usf b/Engine/Shaders/Private/VolumetricCloud.usf
--- a/Engine/Shaders/Private/VolumetricCloud.usf
+++ b/Engine/Shaders/Private/VolumetricCloud.usf
@@ -6,10 +6,7 @@
 
 
  // Change this to force recompilation of all volumetric cloud material shaders
-#pragma message("UESHADERMETADATA_VERSION BFF66DE1-56ED-4BFA-AE26-D7FCF6E790CE")
-
-// Change this to force recompilation of all shaders referencing this file.
-#pragma message("UESHADERMETADATA_VERSION 58372ED6-B830-41C4-829E-A21186A68401")
+#pragma message("UESHADERMETADATA_VERSION 89190FC1-44B2-474A-8CB6-05332B45FE36")
 
 #include "Common.ush"
 #include "ColorMap.ush"
@@ -2042,9 +2039,9 @@ void MainPS(
 	ResolvedView = ResolveView();
 
 	CloudShadowTraceContext TraceContext;
-	if (RenderVolumetricCloudParameters.TraceShadowmap>=1)
+	if (RenderVolumetricCloudParameters.TraceShadowmapMode >=1)
 	{
-		const uint LightIndex = RenderVolumetricCloudParameters.TraceShadowmap - 1;
+		const uint LightIndex = clamp(RenderVolumetricCloudParameters.TraceShadowmapMode - 1, 0, 1);
 		TraceContext.FarDepthKm					= RenderVolumetricCloudParameters.CloudShadowmapFarDepthKm[LightIndex].x;
 		TraceContext.Strength					= RenderVolumetricCloudParameters.CloudShadowmapStrength[LightIndex].x;
 		TraceContext.DepthBias					= RenderVolumetricCloudParameters.CloudShadowmapDepthBias[LightIndex].x;
diff --git a/Engine/Source/Runtime/Renderer/Private/VolumetricCloudRendering.cpp b/Engine/Source/Runtime/Renderer/Private/VolumetricCloudRendering.cpp
--- a/Engine/Source/Runtime/Renderer/Private/VolumetricCloudRendering.cpp
+++ b/Engine/Source/Runtime/Renderer/Private/VolumetricCloudRendering.cpp
@@ -627,7 +627,7 @@ BEGIN_GLOBAL_SHADER_PARAMETER_STRUCT(FRenderVolumetricCloudGlobalParameters, )
 	SHADER_PARAMETER(uint32, VolumetricRenderTargetMode)
 	SHADER_PARAMETER(uint32, CloudDebugViewMode)
 	SHADER_PARAMETER(uint32, IsReflectionRendering)
-	SHADER_PARAMETER(uint32, TraceShadowmap)
+	SHADER_PARAMETER(uint32, TraceShadowmapMode)
 	SHADER_PARAMETER(float, LocalLightsShadowSampleCount)
 	SHADER_PARAMETER(FVector4f, OutputSizeInvSize)
 	SHADER_PARAMETER(int32, StepSizeOnZeroConservativeDensity)
@@ -1957,7 +1957,7 @@ void FSceneRenderer::InitVolumetricCloudsForViews(FRDGBuilder& GraphBuilder, boo
 						// We need to make a copy of the parameters on CPU to morph them because the creation is deferred.
 						FRenderVolumetricCloudGlobalParameters& VolumetricCloudParamsAO = *GraphBuilder.AllocParameters<FRenderVolumetricCloudGlobalParameters>();
 						VolumetricCloudParamsAO = VolumetricCloudParams;	// Use the same parameter as for the directional light shadow
-						VolumetricCloudParamsAO.TraceShadowmap = 0;			// Notify that this pass is for SkyAO (avoid to use another shader permutation)
+						VolumetricCloudParamsAO.TraceShadowmapMode = 0;		// Notify that this pass is for SkyAO (avoid to use another shader permutation)
 						TRDGUniformBufferRef<FRenderVolumetricCloudGlobalParameters> TraceVolumetricCloudSkyAOParamsUB = GraphBuilder.CreateUniformBuffer(&VolumetricCloudParamsAO);
 						TraceCloudTexture(CloudSkyAOTexture, true, TraceVolumetricCloudSkyAOParamsUB);
 
@@ -1993,10 +1993,12 @@ void FSceneRenderer::InitVolumetricCloudsForViews(FRDGBuilder& GraphBuilder, boo
 
 							FRDGTextureRef NewCloudShadowTexture = GraphBuilder.CreateTexture(
 								FRDGTextureDesc::Create2D(TracingResolution2D, CloudShadowPixelFormat,
-									FClearValueBinding::None, TexCreate_ShaderResource | TexCreate_UAV | TexCreate_RenderTargetable), TEXT("Cloud.ShadowTexture"));
+									FClearValueBinding::None, TexCreate_ShaderResource | TexCreate_UAV | TexCreate_RenderTargetable), LightIndex == 0 ? TEXT("Cloud.ShadowTexture0") : TEXT("Cloud.ShadowTexture1"));
 
-							VolumetricCloudParams.TraceShadowmap = 1 + LightIndex;
-							TRDGUniformBufferRef<FRenderVolumetricCloudGlobalParameters> TraceVolumetricCloudShadowParamsUB = GraphBuilder.CreateUniformBuffer(&VolumetricCloudParams);
+							FRenderVolumetricCloudGlobalParameters& VolumetricCloudParamsShadow = *GraphBuilder.AllocParameters<FRenderVolumetricCloudGlobalParameters>();
+							VolumetricCloudParamsShadow = VolumetricCloudParams;			// Use the same parameters generated above for the directional light shadow
+							VolumetricCloudParamsShadow.TraceShadowmapMode = 1 + LightIndex;// But change the shadow mode
+							TRDGUniformBufferRef<FRenderVolumetricCloudGlobalParameters> TraceVolumetricCloudShadowParamsUB = GraphBuilder.CreateUniformBuffer(&VolumetricCloudParamsShadow);
 							TraceCloudTexture(NewCloudShadowTexture, false, TraceVolumetricCloudShadowParamsUB);
 
 							// Directional light shadow temporal filter only if the view has a ViewState (not a sky light capture view for instance)
