Added RENDERER_API to FPrimitiveSceneInfo::GetInstanceSceneDataBuffers() so it can be called by plugins.

diff --git a/Engine/Source/Runtime/Renderer/Public/PrimitiveSceneInfo.h b/Engine/Source/Runtime/Renderer/Public/PrimitiveSceneInfo.h
--- a/Engine/Source/Runtime/Renderer/Public/PrimitiveSceneInfo.h
+++ b/Engine/Source/Runtime/Renderer/Public/PrimitiveSceneInfo.h
@@ -604,7 +604,7 @@ public:
 	/**
 	 * Waits for (potential) instance update to produce the data, to avoid a sync, use GetInstanceDataHeader().
 	 */
-	const FInstanceSceneDataBuffers *GetInstanceSceneDataBuffers() const;
+	RENDERER_API const FInstanceSceneDataBuffers* GetInstanceSceneDataBuffers() const;
 
 	/**
 	 * Returns the updated header data in the InstanceDataUpdateTaskInfo without blocking. 
