Add capsule direct shadows to mobile forward renderer

diff --git a/Engine/Shaders/Private/CapsuleShadowShaders.usf b/Engine/Shaders/Private/CapsuleShadowShaders.usf
--- a/Engine/Shaders/Private/CapsuleShadowShaders.usf
+++ b/Engine/Shaders/Private/CapsuleShadowShaders.usf
@@ -963,7 +963,8 @@ void CapsuleShadowingCS(
 	float2 ScreenUV = float2((DispatchThreadId.xy * DownsampleFactor + ScissorRectMinAndSize.xy + .5f) * View.BufferSizeAndInvSize.zw);
 	float2 ScreenPosition = (ScreenUV.xy - View.ScreenPositionScaleBias.wz) / View.ScreenPositionScaleBias.xy;
 
-#if LIGHT_SOURCE_MODE == LIGHT_SOURCE_FROM_RECEIVER
+// Mobile does not support bent normals
+#if LIGHT_SOURCE_MODE == LIGHT_SOURCE_FROM_RECEIVER && !SHADING_PATH_MOBILE
 	float4 ReceiverTextureValue = ReceiverBentNormalTexture.Load(DispatchThreadId.xyz);
 	float3 ReceiverBentNormal = ReceiverTextureValue.xyz;
 	float SceneDepth = ReceiverTextureValue.w;
@@ -1016,7 +1017,7 @@ void CapsuleShadowingCS(
 				float LightAngle = LightAngleAndNormalThreshold.x;
 
 			#endif
-		#elif LIGHT_SOURCE_MODE == LIGHT_SOURCE_FROM_RECEIVER
+		#elif LIGHT_SOURCE_MODE == LIGHT_SOURCE_FROM_RECEIVER && !SHADING_PATH_MOBILE
 			float3 TranslatedWorldRayStart = OpaqueTranslatedWorldPosition;
 			float BentNormalLength = length(ReceiverBentNormal);
 			float3 UnitRayDirection = ReceiverBentNormal / max(BentNormalLength, .00001f);
diff --git a/Engine/Source/Runtime/RenderCore/Private/RenderUtils.cpp b/Engine/Source/Runtime/RenderCore/Private/RenderUtils.cpp
--- a/Engine/Source/Runtime/RenderCore/Private/RenderUtils.cpp
+++ b/Engine/Source/Runtime/RenderCore/Private/RenderUtils.cpp
@@ -558,6 +558,18 @@ RENDERCORE_API bool IsMobileMovableSpotlightShadowsEnabled(const FStaticShaderPl
 	return FReadOnlyCVARCache::MobileEnableMovableSpotlightsShadow(Platform);
 }
 
+RENDERCORE_API bool IsMobileCapsuleShadowsEnabled(const FStaticShaderPlatform Platform)
+{
+	static FShaderPlatformCachedIniValue<bool> MobileCapsuleShadowsIniValue(TEXT("r.Mobile.EnableCapsuleShadows"));
+	return IsMobilePlatform(Platform) && MobileCapsuleShadowsIniValue.Get(Platform);
+}
+
+RENDERCORE_API bool IsMobileCapsuleDirectShadowsEnabled(const FStaticShaderPlatform Platform)
+{
+	static FShaderPlatformCachedIniValue<bool> MobileCapsuleDirectShadowsIniValue(TEXT("r.Mobile.EnableCapsuleDirectShadows"));
+	return IsMobilePlatform(Platform) && IsMobileCapsuleShadowsEnabled(Platform) && MobileCapsuleDirectShadowsIniValue.Get(Platform);
+}
+
 RENDERCORE_API bool MobileForwardEnableClusteredReflections(const FStaticShaderPlatform Platform)
 {
 	static FShaderPlatformCachedIniValue<bool> MobileForwardEnableClusteredReflectionsIniValue(TEXT("r.Mobile.Forward.EnableClusteredReflections"));
@@ -568,7 +580,9 @@ RENDERCORE_API bool MobileUsesShadowMaskTexture(const FStaticShaderPlatform Plat
 {
 	// Only distance field shadow needs to render shadow mask texture on mobile deferred, normal shadows need to be rendered separately because of handling lighting channels.
 	// Besides distance field shadow, with clustered lighting and shadow of local light enabled, shadows will render to shadow mask texture on mobile forward, lighting channels are handled in base pass shader.
-	return IsMobileDistanceFieldEnabled(Platform) || (!IsMobileDeferredShadingEnabled(Platform) && IsMobileMovableSpotlightShadowsEnabled(Platform) && MobileForwardEnableLocalLights(Platform));
+	const bool bMobileCapsuleShadowsEnabled = IsMobileCapsuleDirectShadowsEnabled(Platform);
+	const bool bMovableSpotlightShadowsEnabled = (IsMobileMovableSpotlightShadowsEnabled(Platform) && MobileForwardEnableLocalLights(Platform));
+	return IsMobileDistanceFieldEnabled(Platform) || (!IsMobileDeferredShadingEnabled(Platform) && (bMobileCapsuleShadowsEnabled || bMovableSpotlightShadowsEnabled));
 }
 
 // Whether to support more than 4 color attachments for GBuffer 
diff --git a/Engine/Source/Runtime/RenderCore/Public/RenderUtils.h b/Engine/Source/Runtime/RenderCore/Public/RenderUtils.h
--- a/Engine/Source/Runtime/RenderCore/Public/RenderUtils.h
+++ b/Engine/Source/Runtime/RenderCore/Public/RenderUtils.h
@@ -318,6 +318,10 @@ RENDERCORE_API bool IsMobileDistanceFieldEnabled(const FStaticShaderPlatform Pla
 
 RENDERCORE_API bool IsMobileMovableSpotlightShadowsEnabled(const FStaticShaderPlatform Platform);
 
+RENDERCORE_API bool IsMobileCapsuleShadowsEnabled(const FStaticShaderPlatform Platform);
+
+RENDERCORE_API bool IsMobileCapsuleDirectShadowsEnabled(const FStaticShaderPlatform Platform);
+
 RENDERCORE_API bool MobileForwardEnableClusteredReflections(const FStaticShaderPlatform Platform);
 
 RENDERCORE_API bool MobileUsesShadowMaskTexture(const FStaticShaderPlatform Platform);
diff --git a/Engine/Source/Runtime/Renderer/Private/CapsuleShadowRendering.cpp b/Engine/Source/Runtime/Renderer/Private/CapsuleShadowRendering.cpp
--- a/Engine/Source/Runtime/Renderer/Private/CapsuleShadowRendering.cpp
+++ b/Engine/Source/Runtime/Renderer/Private/CapsuleShadowRendering.cpp
@@ -193,7 +193,7 @@ public:
 	SHADER_USE_PARAMETER_STRUCT(FCapsuleShadowingCS, FGlobalShader);
 
 	BEGIN_SHADER_PARAMETER_STRUCT(FParameters, )
-		SHADER_PARAMETER_RDG_UNIFORM_BUFFER(FSceneTextureUniformParameters, SceneTextures)
+		SHADER_PARAMETER_STRUCT_INCLUDE(FSceneTextureShaderParameters, SceneTextures)
 		SHADER_PARAMETER_RDG_UNIFORM_BUFFER(FSubstrateGlobalUniformParameters, Substrate)
 		SHADER_PARAMETER_STRUCT_REF(FViewUniformShaderParameters, View)
 
@@ -235,11 +235,6 @@ public:
 	class FIndirectPrimitiveType : SHADER_PERMUTATION_ENUM_CLASS("INDIRECT_PRIMITIVE_TYPE", EIndirectShadowingPrimitiveTypes);
 	using FPermutationDomain = TShaderPermutationDomain<FShapeShadow, FIndirectPrimitiveType>;
 
-	static bool ShouldCompilePermutation(const FGlobalShaderPermutationParameters& Parameters)
-	{
-		return FDataDrivenShaderPlatformInfo::GetSupportsCapsuleShadows(Parameters.Platform);
-	}
-
 	static void ModifyCompilationEnvironment(const FGlobalShaderPermutationParameters& Parameters, FShaderCompilerEnvironment& OutEnvironment)
 	{
 		OutEnvironment.CompilerFlags.Add(CFLAG_StandardOptimization);
@@ -273,7 +268,7 @@ public:
 		}
 
 		OutEnvironment.SetDefine(TEXT("LIGHT_SOURCE_MODE"), (uint32)LightSourceMode);
-		const bool bApplyToBentNormal = ShadowingType == ECapsuleShadowingType::MovableSkylightTiledCulling || ShadowingType == ECapsuleShadowingType::MovableSkylightTiledCullingGatherFromReceiverBentNormal;
+		const bool bApplyToBentNormal = !IsMobilePlatform(Parameters.Platform) ? (ShadowingType == ECapsuleShadowingType::MovableSkylightTiledCulling || ShadowingType == ECapsuleShadowingType::MovableSkylightTiledCullingGatherFromReceiverBentNormal) : false;
 		OutEnvironment.SetDefine(TEXT("APPLY_TO_BENT_NORMAL"), bApplyToBentNormal);
 
 		EIndirectShadowingPrimitiveTypes PrimitiveTypes = PermutationVector.Get<FIndirectPrimitiveType>();
@@ -283,7 +278,7 @@ public:
 			OutEnvironment.SetDefine(TEXT("SUPPORT_CAPSULE_SHAPES"), 1);
 		}
 
-		if (PrimitiveTypes == EIndirectShadowingPrimitiveTypes::MeshDistanceFields || PrimitiveTypes == EIndirectShadowingPrimitiveTypes::CapsuleShapesAndMeshDistanceFields)
+		if (!IsMobilePlatform(Parameters.Platform) && (PrimitiveTypes == EIndirectShadowingPrimitiveTypes::MeshDistanceFields || PrimitiveTypes == EIndirectShadowingPrimitiveTypes::CapsuleShapesAndMeshDistanceFields))
 		{
 			OutEnvironment.SetDefine(TEXT("SUPPORT_MESH_DISTANCE_FIELDS"), 1);
 		}
@@ -305,11 +300,6 @@ class FCapsuleShadowingUpsampleVS : public FGlobalShader
 		SHADER_PARAMETER(FIntRect, ScissorRectMinAndSize)
 	END_SHADER_PARAMETER_STRUCT()
 
-	static bool ShouldCompilePermutation(const FGlobalShaderPermutationParameters& Parameters)
-	{
-		return FDataDrivenShaderPlatformInfo::GetSupportsCapsuleShadows(Parameters.Platform);
-	}
-
 	static void ModifyCompilationEnvironment(const FGlobalShaderPermutationParameters& Parameters, FShaderCompilerEnvironment& OutEnvironment)
 	{
 		OutEnvironment.SetDefine(TEXT("TILES_PER_INSTANCE"), 1);
@@ -337,11 +327,6 @@ public:
 	class FApplySSAO : SHADER_PERMUTATION_BOOL("APPLY_TO_SSAO");
 	using FPermutationDomain = TShaderPermutationDomain<FUpsampleRequired, FApplySSAO>;
 
-	static bool ShouldCompilePermutation(const FGlobalShaderPermutationParameters& Parameters)
-	{
-		return FDataDrivenShaderPlatformInfo::GetSupportsCapsuleShadows(Parameters.Platform);
-	}
-
 	static void ModifyCompilationEnvironment(const FGlobalShaderPermutationParameters& Parameters, FShaderCompilerEnvironment& OutEnvironment)
 	{
 		OutEnvironment.SetDefine(TEXT("UPSAMPLE_PASS"), 1);
@@ -354,7 +339,7 @@ IMPLEMENT_GLOBAL_SHADER(FCapsuleShadowingUpsamplePS, "/Engine/Private/CapsuleSha
 BEGIN_SHADER_PARAMETER_STRUCT(FUpsampleCapsuleShadowParameters, )
 	SHADER_PARAMETER_STRUCT_INCLUDE(FCapsuleShadowingUpsampleVS::FParameters, VS)
 	SHADER_PARAMETER_STRUCT_INCLUDE(FCapsuleShadowingUpsamplePS::FParameters, PS)
-	SHADER_PARAMETER_RDG_UNIFORM_BUFFER(FSceneTextureUniformParameters, SceneTextures)
+	SHADER_PARAMETER_STRUCT_INCLUDE(FSceneTextureShaderParameters, SceneTextures)
 	RENDER_TARGET_BINDING_SLOTS()
 END_SHADER_PARAMETER_STRUCT()
 
@@ -372,7 +357,6 @@ void SetupCapsuleShadowingParameters(
 	float MaxOcclusionDistance,
 	const FScene* Scene,
 	const FViewInfo& View,
-	TRDGUniformBufferRef<FSceneTextureUniformParameters> SceneTexturesUniformBuffer,
 
 	uint32 NumShadowCapsules,
 	FRDGBufferSRVRef ShadowCapsuleShapesBuffer,
@@ -384,7 +368,7 @@ void SetupCapsuleShadowingParameters(
 	FRDGBufferUAVRef CapsuleTileIntersectionCountsUAV
 )
 {
-	Parameters.SceneTextures = SceneTexturesUniformBuffer;
+	Parameters.SceneTextures = GetSceneTextureShaderParameters(View);
 	Parameters.View = View.ViewUniformBuffer;
 	Parameters.Substrate = Substrate::BindSubstrateGlobalUniformParameters(View);
 
@@ -463,9 +447,8 @@ void SetupCapsuleShadowingParameters(
 	Parameters.DFAtlasParameters = DistanceField::SetupAtlasParameters(GraphBuilder, Scene->DistanceFieldSceneData);
 }
 
-bool FDeferredShadingSceneRenderer::RenderCapsuleDirectShadows(
+bool FSceneRenderer::RenderCapsuleDirectShadows(
 	FRDGBuilder& GraphBuilder,
-	TRDGUniformBufferRef<FSceneTextureUniformParameters> SceneTexturesUniformBuffer,
 	const FLightSceneInfo& LightSceneInfo,
 	FRDGTextureRef ScreenShadowMaskTexture,
 	TArrayView<const FProjectedShadowInfo* const> CapsuleShadows,
@@ -483,7 +466,7 @@ bool FDeferredShadingSceneRenderer::RenderCapsuleDirectShadows(
 		}
 	}
 
-	if (!SupportsCapsuleDirectShadows(ShaderPlatform)
+	if (!IsCapsuleDirectShadowsEnabled(ShaderPlatform)
 		|| CapsuleShadows.Num() == 0
 		|| !ViewFamily.EngineShowFlags.CapsuleShadows
 		|| !bAllViewsHaveViewState)
@@ -583,7 +566,6 @@ bool FDeferredShadingSceneRenderer::RenderCapsuleDirectShadows(
 					GCapsuleMaxDirectOcclusionDistance,
 					Scene,
 					View,
-					SceneTexturesUniformBuffer,
 					CapsuleShapeData.Num(),
 					GraphBuilder.CreateSRV(ShadowCapsuleShapesBuffer),
 					0,
@@ -616,8 +598,7 @@ bool FDeferredShadingSceneRenderer::RenderCapsuleDirectShadows(
 
 				FUpsampleCapsuleShadowParameters* PassParameters = GraphBuilder.AllocParameters<FUpsampleCapsuleShadowParameters>();
 				PassParameters->RenderTargets[0] = FRenderTargetBinding(ScreenShadowMaskTexture, ERenderTargetLoadAction::ELoad);
-				PassParameters->SceneTextures = SceneTexturesUniformBuffer;
-
+				PassParameters->SceneTextures = GetSceneTextureShaderParameters(View);
 				PassParameters->VS.View = GetShaderBinding(View.ViewUniformBuffer);
 				PassParameters->VS.TileIntersectionCounts = CapsuleTileIntersectionCountsSRV;
 				PassParameters->VS.TileDimensions = GroupSize;
@@ -1048,7 +1029,6 @@ void FDeferredShadingSceneRenderer::RenderIndirectCapsuleShadows(FRDGBuilder& Gr
 					GCapsuleMaxIndirectOcclusionDistance,
 					Scene,
 					View,
-					SceneTextures.UniformBuffer,
 					NumCapsuleShapes,
 					Resources.IndirectShadowCapsuleShapesSRV,
 					NumMeshDistanceFieldCasters,
@@ -1105,7 +1085,8 @@ void FDeferredShadingSceneRenderer::RenderIndirectCapsuleShadows(FRDGBuilder& Gr
 					// Only allow clears for the first use of the render target.
 					RenderTargets[Index].SetLoadAction(ERenderTargetLoadAction::ELoad);
 				}
-				PassParameters->SceneTextures = SceneTextures.UniformBuffer;
+
+				PassParameters->SceneTextures = GetSceneTextureShaderParameters(View);
 
 				PassParameters->VS.View = GetShaderBinding(View.ViewUniformBuffer);
 				PassParameters->VS.TileIntersectionCounts = CapsuleTileIntersectionCountsSRV;
@@ -1234,7 +1215,6 @@ void FDeferredShadingSceneRenderer::RenderCapsuleShadowsForMovableSkylight(
 						GCapsuleMaxIndirectOcclusionDistance,
 						Scene,
 						View,
-						SceneTexturesUniformBuffer,
 						NumCapsuleShapes,
 						Resources.IndirectShadowCapsuleShapesSRV,
 						NumMeshDistanceFieldCasters,
diff --git a/Engine/Source/Runtime/Renderer/Private/CapsuleShadowRendering.h b/Engine/Source/Runtime/Renderer/Private/CapsuleShadowRendering.h
--- a/Engine/Source/Runtime/Renderer/Private/CapsuleShadowRendering.h
+++ b/Engine/Source/Runtime/Renderer/Private/CapsuleShadowRendering.h
@@ -8,23 +8,23 @@
 
 #include "CoreMinimal.h"
 #include "DataDrivenShaderPlatformInfo.h"
+#include "RenderUtils.h"
 
 extern int32 GCapsuleShadows;
 extern int32 GCapsuleDirectShadows;
 extern int32 GCapsuleIndirectShadows;
 
-inline bool SupportsCapsuleShadows(FStaticShaderPlatform ShaderPlatform)
+inline bool IsCapsuleShadowsEnabled(FStaticShaderPlatform ShaderPlatform)
 {
-	return GCapsuleShadows
-		&& FDataDrivenShaderPlatformInfo::GetSupportsCapsuleShadows(ShaderPlatform);
+	return GCapsuleShadows && (!IsMobilePlatform(ShaderPlatform) || IsMobileCapsuleShadowsEnabled(ShaderPlatform));
 }
 
-inline bool SupportsCapsuleDirectShadows(FStaticShaderPlatform ShaderPlatform)
+inline bool IsCapsuleDirectShadowsEnabled(FStaticShaderPlatform ShaderPlatform)
 {
-	return GCapsuleDirectShadows && SupportsCapsuleShadows(ShaderPlatform);
+	return GCapsuleDirectShadows && IsCapsuleShadowsEnabled(ShaderPlatform) && (!IsMobilePlatform(ShaderPlatform) || IsMobileCapsuleDirectShadowsEnabled(ShaderPlatform));
 }
 
 inline bool SupportsCapsuleIndirectShadows(FStaticShaderPlatform ShaderPlatform)
 {
-	return GCapsuleIndirectShadows && SupportsCapsuleShadows(ShaderPlatform);
+	return GCapsuleIndirectShadows && IsCapsuleShadowsEnabled(ShaderPlatform);
 }
diff --git a/Engine/Source/Runtime/Renderer/Private/DeferredShadingRenderer.h b/Engine/Source/Runtime/Renderer/Private/DeferredShadingRenderer.h
--- a/Engine/Source/Runtime/Renderer/Private/DeferredShadingRenderer.h
+++ b/Engine/Source/Runtime/Renderer/Private/DeferredShadingRenderer.h
@@ -910,15 +910,6 @@ private:
 		bool bSupportShadowMaps,
 		FTranslucentLightInjectionCollector& Collector);
 
-	/** Renders capsule shadows for all per-object shadows using it for the given light. */
-	bool RenderCapsuleDirectShadows(
-		FRDGBuilder& GraphBuilder,
-		TRDGUniformBufferRef<FSceneTextureUniformParameters> SceneTexturesUniformBuffer,
-		const FLightSceneInfo& LightSceneInfo,
-		FRDGTextureRef ScreenShadowMaskTexture,
-		TArrayView<const FProjectedShadowInfo* const> CapsuleShadows,
-		bool bProjectingForForwardShading) const;
-
 	/** Renders indirect shadows from capsules modulated onto scene color. */
 	void RenderIndirectCapsuleShadows(FRDGBuilder& GraphBuilder, const FSceneTextures& SceneTextures) const;
 
diff --git a/Engine/Source/Runtime/Renderer/Private/LightGridInjection.cpp b/Engine/Source/Runtime/Renderer/Private/LightGridInjection.cpp
--- a/Engine/Source/Runtime/Renderer/Private/LightGridInjection.cpp
+++ b/Engine/Source/Runtime/Renderer/Private/LightGridInjection.cpp
@@ -1084,7 +1084,7 @@ void FDeferredShadingSceneRenderer::RenderForwardShadowProjections(
 				}
 			}
 
-			RenderCapsuleDirectShadows(GraphBuilder, SceneTextures.UniformBuffer, *LightSceneInfo, ForwardScreenSpaceShadowMask.Target, VisibleLightInfo.CapsuleShadowsToProject, bProjectingForForwardShading);
+			RenderCapsuleDirectShadows(GraphBuilder, *LightSceneInfo, ForwardScreenSpaceShadowMask.Target, VisibleLightInfo.CapsuleShadowsToProject, bProjectingForForwardShading);
 
 			if (LightSceneInfo->GetDynamicShadowMapChannel() >= 0 && LightSceneInfo->GetDynamicShadowMapChannel() < 4)
 			{
diff --git a/Engine/Source/Runtime/Renderer/Private/MobileBasePassRendering.cpp b/Engine/Source/Runtime/Renderer/Private/MobileBasePassRendering.cpp
--- a/Engine/Source/Runtime/Renderer/Private/MobileBasePassRendering.cpp
+++ b/Engine/Source/Runtime/Renderer/Private/MobileBasePassRendering.cpp
@@ -64,6 +64,21 @@ static TAutoConsoleVariable<int32> CVarMobileTonemapSubpass(
 	TEXT(" 1 = On"),
 	ECVF_Scalability | ECVF_RenderThreadSafe);
 
+static TAutoConsoleVariable<int32> CVarMobileEnableCapsuleShadows(
+	TEXT("r.Mobile.EnableCapsuleShadows"),
+	0,
+	TEXT("0: Capsule shadows are disabled in the mobile renderer")
+	TEXT("1: Enables capsule shadowing on skinned components with bCastCapsuleDirectShadow or bCastCapsuleIndirectShadow enabled with the mobile renderer"),
+	ECVF_ReadOnly | ECVF_RenderThreadSafe);
+
+static TAutoConsoleVariable<int32> CVarMobileEnableCapsuleDirectShadows(
+	TEXT("r.Mobile.EnableCapsuleDirectShadows"),
+	0,
+	TEXT("0: Capsule direct shadows are disabled in the mobile renderer")
+	TEXT("1: Enables capsule direct shadowing on skinned components with bCastCapsuleDirectShadow enabled with the mobile renderer"),
+	ECVF_ReadOnly | ECVF_RenderThreadSafe);
+
+
 IMPLEMENT_STATIC_UNIFORM_BUFFER_STRUCT(FMobileBasePassUniformParameters, "MobileBasePass", SceneTextures);
 
 static TAutoConsoleVariable<int32> CVarMobileUseHWsRGBEncoding(
diff --git a/Engine/Source/Runtime/Renderer/Private/SceneRendering.h b/Engine/Source/Runtime/Renderer/Private/SceneRendering.h
--- a/Engine/Source/Runtime/Renderer/Private/SceneRendering.h
+++ b/Engine/Source/Runtime/Renderer/Private/SceneRendering.h
@@ -2411,6 +2411,14 @@ protected:
 	/** Gathers dynamic mesh elements for all shadows. */
 	void GatherShadowDynamicMeshElements(FDynamicShadowsTaskData& TaskData);
 
+	/** Renders capsule shadows for all per-object shadows using it for the given light. */
+	bool RenderCapsuleDirectShadows(
+		FRDGBuilder& GraphBuilder,
+		const FLightSceneInfo& LightSceneInfo,
+		FRDGTextureRef ScreenShadowMaskTexture,
+		TArrayView<const FProjectedShadowInfo* const> CapsuleShadows,
+		bool bProjectingForForwardShading) const;
+
 	/** Performs once per frame setup prior to visibility determination. */
 	void PreVisibilityFrameSetup(FRDGBuilder& GraphBuilder);
 
diff --git a/Engine/Source/Runtime/Renderer/Private/ShadowRendering.cpp b/Engine/Source/Runtime/Renderer/Private/ShadowRendering.cpp
--- a/Engine/Source/Runtime/Renderer/Private/ShadowRendering.cpp
+++ b/Engine/Source/Runtime/Renderer/Private/ShadowRendering.cpp
@@ -2314,7 +2314,7 @@ void FDeferredShadingSceneRenderer::RenderDeferredShadowProjections(
 	RenderShadowProjections(GraphBuilder, SceneTextures, ScreenShadowMaskTexture, ScreenShadowMaskSubPixelTexture, LightSceneInfo, bProjectingForForwardShading);
 	ShadowSceneRenderer->ApplyVirtualShadowMapProjectionForLight(GraphBuilder, SceneTextures, LightSceneInfo, ScreenShadowMaskTexture, ScreenShadowMaskSubPixelTexture);
 
-	RenderCapsuleDirectShadows(GraphBuilder, SceneTextures.UniformBuffer, *LightSceneInfo, ScreenShadowMaskTexture, VisibleLightInfo.CapsuleShadowsToProject, bProjectingForForwardShading);
+	RenderCapsuleDirectShadows(GraphBuilder, *LightSceneInfo, ScreenShadowMaskTexture, VisibleLightInfo.CapsuleShadowsToProject, bProjectingForForwardShading);
 
 	// Inject deep shadow mask for regular shadow map. When using virtual shadow map, it is directly handled in the shadow kernel.
 	if (HairStrands::HasViewHairStrandsData(Views))
@@ -2422,7 +2422,8 @@ void FMobileSceneRenderer::RenderMobileShadowProjections(
 				LightSceneInfo,
 				bProjectingForForwardShading);
 		}
-		
+		RenderCapsuleDirectShadows(GraphBuilder, *LightSceneInfo, ScreenShadowMaskTexture, VisibleLightInfo.CapsuleShadowsToProject, bProjectingForForwardShading);
+
 		if (LightSceneProxy->GetLightType() == LightType_Directional && LightSceneInfo->GetDynamicShadowMapChannel() != -1)
 		{
 			// Dynamic shadows are projected into channels of the light attenuation texture based on their assigned DynamicShadowMapChannel
diff --git a/Engine/Source/Runtime/Renderer/Private/ShadowSetup.cpp b/Engine/Source/Runtime/Renderer/Private/ShadowSetup.cpp
--- a/Engine/Source/Runtime/Renderer/Private/ShadowSetup.cpp
+++ b/Engine/Source/Runtime/Renderer/Private/ShadowSetup.cpp
@@ -815,7 +815,7 @@ bool FProjectedShadowInfo::SetupPerObjectProjection(
 	MaxScreenPercent = InMaxScreenPercent;
 	bDirectionalLight = InLightSceneInfo->Proxy->GetLightType() == LightType_Directional;
 	const EShaderPlatform ShaderPlatform = LightSceneInfo->Scene->GetShaderPlatform();
-	bCapsuleShadow = InParentSceneInfo->Proxy->CastsCapsuleDirectShadow() && !bInPreShadow && SupportsCapsuleDirectShadows(ShaderPlatform);
+	bCapsuleShadow = InParentSceneInfo->Proxy->CastsCapsuleDirectShadow() && !bInPreShadow && IsCapsuleDirectShadowsEnabled(ShaderPlatform);
 	bTranslucentShadow = bInTranslucentShadow;
 	bPreShadow = bInPreShadow;
 	bSelfShadowOnly = InParentSceneInfo->Proxy->CastsSelfShadowOnly();
@@ -6215,7 +6215,7 @@ void FSceneRenderer::CreateDynamicShadows(FDynamicShadowsTaskData& TaskData)
 #endif
 							bGMobileInsetShadows = CVarMobileInsetShadows.GetValueOnRenderThread();
 
-							if (!bMobile || bGMobileInsetShadows || (LightSceneInfo->Proxy->CastsModulatedShadows() && !LightSceneInfo->Proxy->UseCSMForDynamicObjects() && LightSceneInfo->Proxy->HasStaticShadowing()))
+							if (!bMobile || bGMobileInsetShadows || IsMobileCapsuleShadowsEnabled(ShaderPlatform) || (LightSceneInfo->Proxy->CastsModulatedShadows() && !LightSceneInfo->Proxy->UseCSMForDynamicObjects() && LightSceneInfo->Proxy->HasStaticShadowing()))
 							{
 								const TArray<FLightPrimitiveInteraction*>* InteractionShadowPrimitives = LightSceneInfo->GetInteractionShadowPrimitives();
 
@@ -6419,8 +6419,8 @@ void FSceneRenderer::FilterDynamicShadows(FDynamicShadowsTaskData& TaskData)
 				}
 
 				bool bNeedsProjection = ProjectedShadowInfo->CacheMode != SDCM_StaticPrimitivesOnly
-					//// Filter out everything but PerObjectOpaqueShadows & Distance field shadows for ES31
-					&& (!bMobile || ProjectedShadowInfo->bPerObjectOpaqueShadow || ProjectedShadowInfo->bRayTracedDistanceField || ProjectedShadowInfo->bWholeSceneShadow);
+					//// Filter out everything but PerObjectOpaqueShadows, Distance field shadows and capsule shadows for ES31
+					&& (!bMobile || ProjectedShadowInfo->bPerObjectOpaqueShadow || ProjectedShadowInfo->bRayTracedDistanceField || ProjectedShadowInfo->bWholeSceneShadow || ProjectedShadowInfo->bCapsuleShadow);
 
 				if (bNeedsProjection)
 				{
