Fix Warning in Compile PGO Windows: Map.h(811) : warning C4702: unreachable code

diff --git a/Engine/Source/Runtime/Renderer/Private/PathTracingSpatialTemporalDenoising.cpp b/Engine/Source/Runtime/Renderer/Private/PathTracingSpatialTemporalDenoising.cpp
--- a/Engine/Source/Runtime/Renderer/Private/PathTracingSpatialTemporalDenoising.cpp
+++ b/Engine/Source/Runtime/Renderer/Private/PathTracingSpatialTemporalDenoising.cpp
@@ -39,10 +39,7 @@ public:
 		}
 		else if (SpatialDenoiser.Num() > 0 && !bMatch)
 		{
-			for (const auto& Pair : SpatialDenoiser)
-			{
-				return Pair.Value.Get();
-			}
+			return SpatialDenoiser.CreateConstIterator()->Value.Get();
 		}
 
 		return nullptr;
@@ -56,10 +53,7 @@ public:
 		}
 		else if (PathTracingSpatialTemporalDenoisers.Num() > 0 && !bMatch)
 		{
-			for (const auto& Pair : PathTracingSpatialTemporalDenoisers)
-			{
-				return Pair.Value.Get();
-			}
+			return PathTracingSpatialTemporalDenoisers.CreateConstIterator()->Value.Get();
 		}
 
 		return nullptr;
