enabling r.SplitScreenDebug.Enable doesn't implicitly set "r.InstanceCulling.AllowBatchedBuildRenderingCommands 0" anymore

diff --git a/Engine/Source/Runtime/Renderer/Private/SceneRendering.cpp b/Engine/Source/Runtime/Renderer/Private/SceneRendering.cpp
--- a/Engine/Source/Runtime/Renderer/Private/SceneRendering.cpp
+++ b/Engine/Source/Runtime/Renderer/Private/SceneRendering.cpp
@@ -162,27 +162,6 @@ static TAutoConsoleVariable<int32> CVarSplitScreenDebugEnable(
 	TEXT("r.SplitScreenDebug.Enable"),
 	0,
 	TEXT("Debug feature to replace the main view with a pair of split screen views for testing purposes."),
-	FConsoleVariableDelegate::CreateLambda([](IConsoleVariable* CVar)
-	{
-		// We hit this assert on console in FInstanceCullingMergedContext::AddBatch when enabling split screen:
-		//
-		//		// Verify that each batch contains the same HZB if not null as we only support one
-		//		check(PrevHZB == nullptr || PrevHZB == GraphBuilder.RegisterExternalTexture(Context->PrevHZB));
-		//
-		// HZBs are not shared, so there is definitely a different one per view, making batched instance culling incompatible
-		// with split screen at the moment.  I'm not sure what the fix would be (Create an FInstanceCullingManager per view?
-		// Flush the instance culling manager and swap the HZB when switching views?  Reference multiple HZB and choose the
-		// correct one per batch?  Merge HZBs into a single resource?  Something else?), but I don't want the debug feature
-		// to just crash when used.  So disable this CVar as a workaround.
-		if (CVar->GetBool())
-		{
-			IConsoleVariable* InstanceCullBatchCVar = IConsoleManager::Get().FindConsoleVariable(TEXT("r.InstanceCulling.AllowBatchedBuildRenderingCommands"), false);
-			if (InstanceCullBatchCVar)
-			{
-				InstanceCullBatchCVar->Set(TEXT("0"));
-			}
-		}
-	}),
 	ECVF_Default);
 
 static TAutoConsoleVariable<int32> CVarSplitScreenDebugVertical(
