Heterogeneous Volumes: Initial holdout implementation for path tracing

diff --git a/Engine/Source/Runtime/Renderer/Private/HeterogeneousVolumes/HeterogeneousVolumes.cpp b/Engine/Source/Runtime/Renderer/Private/HeterogeneousVolumes/HeterogeneousVolumes.cpp
--- a/Engine/Source/Runtime/Renderer/Private/HeterogeneousVolumes/HeterogeneousVolumes.cpp
+++ b/Engine/Source/Runtime/Renderer/Private/HeterogeneousVolumes/HeterogeneousVolumes.cpp
@@ -343,6 +343,34 @@ bool ShouldRenderHeterogeneousVolumesForView(
 		&& !View.bIsReflectionCapture;
 }
 
+bool ShouldRenderHeterogeneousVolumesAsHoldoutForView(
+	const FViewInfo& View
+)
+{
+	// This query returns true if any volume is marked as a holdout; otherwise, the query returns false
+	if (ShouldRenderHeterogeneousVolumesForView(View))
+	{
+		for (int32 MeshBatchIndex = 0; MeshBatchIndex < View.HeterogeneousVolumesMeshBatches.Num(); ++MeshBatchIndex)
+		{
+			const FMeshBatch* Mesh = View.HeterogeneousVolumesMeshBatches[MeshBatchIndex].Mesh;
+			const FPrimitiveSceneProxy* PrimitiveSceneProxy = View.HeterogeneousVolumesMeshBatches[MeshBatchIndex].Proxy;
+			if (ShouldRenderMeshBatchWithHeterogeneousVolumes(Mesh, PrimitiveSceneProxy, View.GetFeatureLevel()))
+			{
+				for (int32 VolumeIndex = 0; VolumeIndex < Mesh->Elements.Num(); ++VolumeIndex)
+				{
+					const IHeterogeneousVolumeInterface* HeterogeneousVolume = (IHeterogeneousVolumeInterface*)Mesh->Elements[VolumeIndex].UserData;
+					if (HeterogeneousVolumes::IsHoldout(HeterogeneousVolume))
+					{
+						return true;
+					}
+				}
+			}
+		}
+	}
+
+	return false;
+}
+
 bool DoesMaterialShaderSupportHeterogeneousVolumes(const FMaterialShaderParameters& MaterialShaderParameters)
 {
 	return (MaterialShaderParameters.MaterialDomain == MD_Volume)
diff --git a/Engine/Source/Runtime/Renderer/Private/HeterogeneousVolumes/HeterogeneousVolumes.h b/Engine/Source/Runtime/Renderer/Private/HeterogeneousVolumes/HeterogeneousVolumes.h
--- a/Engine/Source/Runtime/Renderer/Private/HeterogeneousVolumes/HeterogeneousVolumes.h
+++ b/Engine/Source/Runtime/Renderer/Private/HeterogeneousVolumes/HeterogeneousVolumes.h
@@ -33,6 +33,7 @@ struct FVolumetricMeshBatch;
 bool ShouldRenderHeterogeneousVolumes(const FScene* Scene);
 bool ShouldRenderHeterogeneousVolumesForAnyView(const TArrayView<FViewInfo>& Views);
 bool ShouldRenderHeterogeneousVolumesForView(const FViewInfo& View);
+bool ShouldRenderHeterogeneousVolumesAsHoldoutForView(const FViewInfo& View);
 bool DoesMaterialShaderSupportHeterogeneousVolumes(const FMaterialShaderParameters& Parameters);
 bool DoesMaterialShaderSupportHeterogeneousVolumes(const FMaterial& Material);
 bool ShouldRenderMeshBatchWithHeterogeneousVolumes(
diff --git a/Engine/Source/Runtime/Renderer/Private/PathTracing.cpp b/Engine/Source/Runtime/Renderer/Private/PathTracing.cpp
--- a/Engine/Source/Runtime/Renderer/Private/PathTracing.cpp
+++ b/Engine/Source/Runtime/Renderer/Private/PathTracing.cpp
@@ -800,6 +800,7 @@ static void PreparePathTracingData(const FScene* Scene, const FViewInfo& View, F
 		&& (Scene->ExponentialFogs[0].FogData[0].Density > 0 ||
 			Scene->ExponentialFogs[0].FogData[1].Density > 0) ? PATH_TRACER_VOLUME_ENABLE_FOG : 0;
 	PathTracingData.VolumeFlags |= ShouldRenderHeterogeneousVolumesForView(View) ? PATH_TRACER_VOLUME_ENABLE_HETEROGENEOUS_VOLUMES : 0;
+	PathTracingData.VolumeFlags |= ShouldRenderHeterogeneousVolumesAsHoldoutForView(View) ? PATH_TRACER_VOLUME_HOLDOUT_HETEROGENEOUS_VOLUMES : 0;
 	PathTracingData.VolumeFlags |= View.SkyAtmosphereUniformShaderParameters != nullptr && IsSkyAtmosphereHoldout(View.CachedViewUniformShaderParameters->EnvironmentComponentsFlags) ? PATH_TRACER_VOLUME_HOLDOUT_ATMOSPHERE : 0;
 	PathTracingData.VolumeFlags |= bVolumeCloudsVisible && IsVolumetricCloudHoldout(View.CachedViewUniformShaderParameters->EnvironmentComponentsFlags) ? PATH_TRACER_VOLUME_HOLDOUT_CLOUDS : 0;
 	PathTracingData.VolumeFlags |= Scene->ExponentialFogs.Num() > 0 && IsExponentialFogHoldout(View.CachedViewUniformShaderParameters->EnvironmentComponentsFlags) ? PATH_TRACER_VOLUME_HOLDOUT_FOG : 0;
