Make sure the LumenHardwareRayTracingUniformBuffer is created and available in all view info's (shared object)

diff --git a/Engine/Source/Runtime/Renderer/Private/DeferredShadingRenderer.cpp b/Engine/Source/Runtime/Renderer/Private/DeferredShadingRenderer.cpp
--- a/Engine/Source/Runtime/Renderer/Private/DeferredShadingRenderer.cpp
+++ b/Engine/Source/Runtime/Renderer/Private/DeferredShadingRenderer.cpp
@@ -898,6 +898,8 @@ void FDeferredShadingSceneRenderer::WaitForRayTracingScene(FRDGBuilder& GraphBui
 	const int32 ReferenceViewIndex = 0;
 	FViewInfo& ReferenceView = Views[ReferenceViewIndex];
 	
+	SetupLumenHardwareRayTracingUniformBuffer(ReferenceView);
+
 	// Send ray tracing resources from reference view to all others.
 	for (int32 ViewIndex = 0; ViewIndex < AllFamilyViews.Num(); ++ViewIndex)
 	{
@@ -906,11 +908,10 @@ void FDeferredShadingSceneRenderer::WaitForRayTracingScene(FRDGBuilder& GraphBui
 		if (View->bHasAnyRayTracingPass && View != &ReferenceView)
 		{
 			View->LumenHardwareRayTracingMaterialPipeline = ReferenceView.LumenHardwareRayTracingMaterialPipeline;
+			View->LumenHardwareRayTracingUniformBuffer = ReferenceView.LumenHardwareRayTracingUniformBuffer;
 		}
 	}
 
-	SetupLumenHardwareRayTracingUniformBuffer(ReferenceView);
-
 	SetupRayTracingPipelineStates(GraphBuilder);
 
 	bool bAnyLumenHardwareInlineRayTracingPassEnabled = false;
diff --git a/Engine/Source/Runtime/Renderer/Private/RayTracing/RayTracing.cpp b/Engine/Source/Runtime/Renderer/Private/RayTracing/RayTracing.cpp
--- a/Engine/Source/Runtime/Renderer/Private/RayTracing/RayTracing.cpp
+++ b/Engine/Source/Runtime/Renderer/Private/RayTracing/RayTracing.cpp
@@ -340,7 +340,7 @@ namespace RayTracing
 			StaticUniformBuffers.AddUniformBuffer(View.ViewUniformBuffer.GetReference());
 			StaticUniformBuffers.AddUniformBuffer(SceneUniformBuffer);
 			StaticUniformBuffers.AddUniformBuffer(Nanite::GRayTracingManager.GetUniformBuffer().GetReference());
-			StaticUniformBuffers.AddUniformBuffer(View.GetPrimaryView()->LumenHardwareRayTracingUniformBuffer.GetReference());
+			StaticUniformBuffers.AddUniformBuffer(View.LumenHardwareRayTracingUniformBuffer.GetReference());
 
 			StaticUniformBufferScope.Emplace(RHICmdList, StaticUniformBuffers);
 		}
