Support inset shadows without requiring stencil culling

diff --git a/Engine/Source/Runtime/Renderer/Private/ShadowRendering.cpp b/Engine/Source/Runtime/Renderer/Private/ShadowRendering.cpp
--- a/Engine/Source/Runtime/Renderer/Private/ShadowRendering.cpp
+++ b/Engine/Source/Runtime/Renderer/Private/ShadowRendering.cpp
@@ -174,16 +174,6 @@ static FAutoConsoleVariableRef CVarGShadowStencilCulling(
 	ECVF_RenderThreadSafe
 );
 
-static bool UseShadowStencilCulling(FStaticShaderPlatform ShaderPlatform)
-{
-	if (IsMobilePlatform(ShaderPlatform) && IsMobileDeferredShadingEnabled(ShaderPlatform))
-	{
-		// ShadowStencilMask now clashes with a mobile deferred ShadingModels mask
-		return false;
-	}
-	return GShadowStencilCulling != 0;
-}
-
 static TAutoConsoleVariable<int32> CVarFilterMethod(
 	TEXT("r.Shadow.FilterMethod"),
 	0,
@@ -668,6 +658,16 @@ static void BindShadowProjectionShaders(int32 Quality, FRHICommandList& RHICmdLi
 	check(GraphicsPSOInit.BoundShaderState.PixelShaderRHI);
 }
 
+bool FProjectedShadowInfo::HasShadowStencilCulling(FStaticShaderPlatform ShaderPlatform)
+{
+	if (IsMobilePlatform(ShaderPlatform) && IsMobileDeferredShadingEnabled(ShaderPlatform))
+	{
+		// ShadowStencilMask now clashes with a mobile deferred ShadingModels mask
+		return false;
+	}
+	return GShadowStencilCulling != 0;
+}
+
 FRHIBlendState* FProjectedShadowInfo::GetBlendStateForProjection(
 	int32 ShadowMapChannel, 
 	bool bIsWholeSceneDirectionalShadow,
@@ -1304,7 +1304,7 @@ void FProjectedShadowInfo::RenderProjectionInternal(
 	SetupFrustumForProjection(View, FrustumVertices, bCameraInsideShadowFrustum, OutFrustmPlanes);
 
 	const bool bSubPixelSupport = HairStrandsUniformBuffer != nullptr;// HairStrands::HasViewHairStrandsData(*View);
-	const bool bStencilTestEnabled = !bSubPixelSupport && UseShadowStencilCulling(View->GetShaderPlatform());
+	const bool bStencilTestEnabled = !bSubPixelSupport && HasShadowStencilCulling(View->GetShaderPlatform());
 	const bool bDepthBoundsTestEnabled = IsWholeSceneDirectionalShadow() && GSupportsDepthBoundsTest && CVarCSMDepthBoundsTest.GetValueOnRenderThread() != 0;// && !bSubPixelSupport;
 	const uint32 StencilRef = bSubPixelSupport && !IsWholeSceneDirectionalShadow() && !bCameraInsideShadowFrustum ? 1u : 0u;
 
diff --git a/Engine/Source/Runtime/Renderer/Private/ShadowRendering.h b/Engine/Source/Runtime/Renderer/Private/ShadowRendering.h
--- a/Engine/Source/Runtime/Renderer/Private/ShadowRendering.h
+++ b/Engine/Source/Runtime/Renderer/Private/ShadowRendering.h
@@ -553,6 +553,8 @@ public:
 	/** Renders shadow maps for translucent primitives. */
 	void RenderTranslucencyDepths(FRDGBuilder& GraphBuilder, class FSceneRenderer* SceneRenderer, const FRenderTargetBindingSlots& RenderTargets, FInstanceCullingManager& InstanceCullingManager);
 
+	static bool HasShadowStencilCulling(FStaticShaderPlatform ShaderPlatform);
+
 	static FRHIBlendState* GetBlendStateForProjection(
 		int32 ShadowMapChannel,
 		bool bIsWholeSceneDirectionalShadow,
diff --git a/Engine/Source/Runtime/Renderer/Private/ShadowSetup.cpp b/Engine/Source/Runtime/Renderer/Private/ShadowSetup.cpp
--- a/Engine/Source/Runtime/Renderer/Private/ShadowSetup.cpp
+++ b/Engine/Source/Runtime/Renderer/Private/ShadowSetup.cpp
@@ -5733,7 +5733,15 @@ void FSceneRenderer::AllocateAtlasedShadowDepthTargets(
 	const FIntPoint MaxTextureSize = GetShadowDepthTextureResolution(FeatureLevel);
 
 	TArray<FLayoutAndAssignedShadows, SceneRenderingAllocator> Layouts;
-	Layouts.Add(FLayoutAndAssignedShadows(MaxTextureSize));
+
+	// If shadow casters use stencil culling they can share a depth atlas
+	bool bCanUseShadowDepthAtlas = FProjectedShadowInfo::HasShadowStencilCulling(ShaderPlatform);
+
+	// Allocate shared atlas
+	if (bCanUseShadowDepthAtlas)
+	{
+		Layouts.Add(FLayoutAndAssignedShadows(MaxTextureSize));
+	}
 
 	for (int32 ShadowIndex = 0; ShadowIndex < Shadows.Num(); ShadowIndex++)
 	{
@@ -5753,11 +5761,9 @@ void FSceneRenderer::AllocateAtlasedShadowDepthTargets(
 		}
 		else
 		{
-			// Avoid infinite loop if texture cannot be allocated even on a fresh atlas
-			// This should not occur, but good to have a safeguard and will still trigger the check() below
-			for (int32 Attempt = 0; Attempt < 2; ++Attempt)
+			auto addCaster = [&](FLayoutAndAssignedShadows& Layout, FProjectedShadowInfo* ProjectedShadowInfo) -> bool
 			{
-				if (Layouts.Last().TextureLayout.AddElement(
+				if (Layout.TextureLayout.AddElement(
 					ProjectedShadowInfo->X,
 					ProjectedShadowInfo->Y,
 					ProjectedShadowInfo->ResolutionX + ProjectedShadowInfo->BorderSize * 2,
@@ -5765,18 +5771,42 @@ void FSceneRenderer::AllocateAtlasedShadowDepthTargets(
 					)
 				{
 					ProjectedShadowInfo->bAllocated = true;
-					Layouts.Last().Shadows.Add(ProjectedShadowInfo);
-					break;
+					Layout.Shadows.Add(ProjectedShadowInfo);
+					return true;
 				}
 
-				// Out of space, add a new atlas and try again
+				return false;
+			};
+
+			if (bCanUseShadowDepthAtlas)
+			{
+				// Avoid infinite loop if texture cannot be allocated even on a fresh atlas
+				// This should not occur, but good to have a safeguard and will still trigger the check() below
+				for (int32 Attempt = 0; Attempt < 2; ++Attempt)
+				{
+					if (!addCaster(Layouts.Last(), ProjectedShadowInfo))
+					{
+						// Out of space, add a new atlas and try again
+						Layouts.Add(FLayoutAndAssignedShadows(MaxTextureSize));
+						continue;
+					}
+
+					break;
+				}
+			}
+			else
+			{
+				// Unique shadow map per caster
 				Layouts.Add(FLayoutAndAssignedShadows(MaxTextureSize));
+				addCaster(Layouts.Last(), ProjectedShadowInfo);
 			}
 
 			check(ProjectedShadowInfo->bAllocated);
 		}
 	}
 
+	const TCHAR *ShadowMapName = bCanUseShadowDepthAtlas ? TEXT("ShadowDepthAtlas") : TEXT("ShadowDepth");
+
 	for (int32 LayoutIndex = 0; LayoutIndex < Layouts.Num(); LayoutIndex++)
 	{
 		const FLayoutAndAssignedShadows& CurrentLayout = Layouts[LayoutIndex];
@@ -5796,7 +5826,7 @@ void FSceneRenderer::AllocateAtlasedShadowDepthTargets(
 
 		FPooledRenderTargetDesc ShadowMapDesc2D = FPooledRenderTargetDesc::Create2DDesc(AtlasSize, PF_ShadowDepth, FClearValueBinding::DepthOne, TexCreate_None, TexCreate_DepthStencilTargetable | TexCreate_ShaderResource, false);
 		ShadowMapDesc2D.Flags |= GFastVRamConfig.ShadowPerObject;
-		GRenderTargetPool.FindFreeElement(RHICmdList, ShadowMapDesc2D, ShadowMapAtlas.RenderTargets.DepthTarget, TEXT("ShadowDepthAtlas"));
+		GRenderTargetPool.FindFreeElement(RHICmdList, ShadowMapDesc2D, ShadowMapAtlas.RenderTargets.DepthTarget, ShadowMapName);
 
 		for (int32 ShadowIndex = 0; ShadowIndex < CurrentLayout.Shadows.Num(); ShadowIndex++)
 		{
