Assign volumetric shadow property during clipmap projection setup.

diff --git a/Engine/Source/Runtime/Renderer/Private/ShadowSetup.cpp b/Engine/Source/Runtime/Renderer/Private/ShadowSetup.cpp
--- a/Engine/Source/Runtime/Renderer/Private/ShadowSetup.cpp
+++ b/Engine/Source/Runtime/Renderer/Private/ShadowSetup.cpp
@@ -1122,6 +1122,7 @@ void FProjectedShadowInfo::SetupClipmapProjection(FLightSceneInfo* InLightSceneI
 	ResolutionY = FVirtualShadowMap::VirtualMaxResolutionXY;
 	bDirectionalLight = true;
 	bWholeSceneShadow = true;
+	bVolumetricShadow = InLightSceneInfo->Proxy->CastsVolumetricShadow();
 	BorderSize = 0;
 	MaxNonFarCascadeDistance = InMaxNonFarCascadeDistance;
 	bVSM = true;
