Heterogeneous Volumes: Temporarily disable continuous LOD during shadow map creation.

diff --git a/Engine/Source/Runtime/Renderer/Private/HeterogeneousVolumes/HeterogeneousVolumesLiveShadingPipeline.cpp b/Engine/Source/Runtime/Renderer/Private/HeterogeneousVolumes/HeterogeneousVolumesLiveShadingPipeline.cpp
--- a/Engine/Source/Runtime/Renderer/Private/HeterogeneousVolumes/HeterogeneousVolumesLiveShadingPipeline.cpp
+++ b/Engine/Source/Runtime/Renderer/Private/HeterogeneousVolumes/HeterogeneousVolumesLiveShadingPipeline.cpp
@@ -1952,7 +1952,8 @@ bool RenderVolumetricShadowMapForLightForHeterogeneousVolumeWithLiveShading(
 		PassParameters->VoxelResolution = HeterogeneousVolumeInterface->GetVoxelResolution();
 
 		// Ray Data
-		float LODFactor = HeterogeneousVolumes::CalcLODFactor(View, HeterogeneousVolumeInterface);
+		//float LODFactor = HeterogeneousVolumes::CalcLODFactor(View, HeterogeneousVolumeInterface);
+		float LODFactor = 1.0f;
 		PassParameters->ShadowStepSize = HeterogeneousVolumes::GetShadowStepSize();
 		PassParameters->ShadowStepFactor = HeterogeneousVolumeInterface->GetShadowStepFactor() * LODFactor;
 		PassParameters->MaxTraceDistance = HeterogeneousVolumes::GetMaxTraceDistance();
@@ -2075,7 +2076,8 @@ bool RenderVolumetricShadowMapForLightWithLiveShading(
 		for (int32 VolumeIndex = 0; VolumeIndex < VolumeCount; ++VolumeIndex)
 		{
 			const IHeterogeneousVolumeInterface* HeterogeneousVolumeInterface = (IHeterogeneousVolumeInterface*)VolumetricMeshBatch.Mesh->Elements[VolumeIndex].UserData;
-			LODValue = FMath::Min(LODValue, HeterogeneousVolumes::CalcLOD(View, HeterogeneousVolumeInterface));
+			//LODValue = FMath::Min(LODValue, HeterogeneousVolumes::CalcLOD(View, HeterogeneousVolumeInterface));
+			LODValue = FMath::Min(LODValue, 1.0f);
 		}
 	}
 	float LODFactor = HeterogeneousVolumes::CalcLODFactor(LODValue);
