Fixed missing SetViewState call on shadow depths.

diff --git a/Engine/Source/Runtime/Renderer/Private/MeshDrawCommands.cpp b/Engine/Source/Runtime/Renderer/Private/MeshDrawCommands.cpp
--- a/Engine/Source/Runtime/Renderer/Private/MeshDrawCommands.cpp
+++ b/Engine/Source/Runtime/Renderer/Private/MeshDrawCommands.cpp
@@ -1737,6 +1737,14 @@ void FParallelMeshDrawCommandPass::DispatchDraw(FParallelCommandListSet* Paralle
 }
 
 void FParallelMeshDrawCommandPass::Dispatch(FRDGDispatchPassBuilder& DispatchPassBuilder, const FInstanceCullingDrawParams* InstanceCullingDrawParams, float ViewportScale) const
+{
+	Dispatch(DispatchPassBuilder, InstanceCullingDrawParams, [View = TaskContext.View, ViewportScale] (FRHICommandList& RHICmdList)
+	{
+		FSceneRenderer::SetStereoViewport(RHICmdList, *View, ViewportScale);
+	});
+}
+
+void FParallelMeshDrawCommandPass::Dispatch(FRDGDispatchPassBuilder& DispatchPassBuilder, const FInstanceCullingDrawParams* InstanceCullingDrawParams, TFunctionRef<void(FRHICommandList&)> SetupCommandListFunction) const
 {
 	extern TAutoConsoleVariable<int32> CVarRHICmdMinDrawsPerParallelCmdList;
 
@@ -1773,7 +1781,7 @@ void FParallelMeshDrawCommandPass::Dispatch(FRDGDispatchPassBuilder& DispatchPas
 		checkSlow(NumDraws > 0);
 
 		FRHICommandList* RHICmdList = DispatchPassBuilder.CreateCommandList();
-		FSceneRenderer::SetStereoViewport(*RHICmdList, *TaskContext.View, ViewportScale);
+		SetupCommandListFunction(*RHICmdList);
 
 		TGraphTask<FDrawVisibleMeshCommandsAnyThreadTask>::CreateTask(&Prereqs)
 			.ConstructAndDispatchWhenReady(*RHICmdList, TaskContext.InstanceCullingContext, TaskContext.MeshDrawCommands, TaskContext.MinimalPipelineStatePassSet,
diff --git a/Engine/Source/Runtime/Renderer/Private/MeshDrawCommands.h b/Engine/Source/Runtime/Renderer/Private/MeshDrawCommands.h
--- a/Engine/Source/Runtime/Renderer/Private/MeshDrawCommands.h
+++ b/Engine/Source/Runtime/Renderer/Private/MeshDrawCommands.h
@@ -164,6 +164,7 @@ public:
 	 */
 	void DispatchDraw(FParallelCommandListSet* ParallelCommandListSet, FRHICommandList& RHICmdList, const FInstanceCullingDrawParams* InstanceCullingDrawParams = nullptr) const;
 	void Dispatch(FRDGDispatchPassBuilder& DispatchPassBuilder, const FInstanceCullingDrawParams* InstanceCullingDrawParams = nullptr, float ViewportScale = 1.0f) const;
+	void Dispatch(FRDGDispatchPassBuilder& DispatchPassBuilder, const FInstanceCullingDrawParams* InstanceCullingDrawParams, TFunctionRef<void(FRHICommandList&)> SetupCommandListFunction) const;
 
 	void WaitForTasksAndEmpty();
 	void SetDumpInstancingStats(const FString& InPassName);
diff --git a/Engine/Source/Runtime/Renderer/Private/ShadowDepthRendering.cpp b/Engine/Source/Runtime/Renderer/Private/ShadowDepthRendering.cpp
--- a/Engine/Source/Runtime/Renderer/Private/ShadowDepthRendering.cpp
+++ b/Engine/Source/Runtime/Renderer/Private/ShadowDepthRendering.cpp
@@ -742,37 +742,6 @@ static TAutoConsoleVariable<int32> CVarParallelShadowsNonWholeScene(
 	ECVF_RenderThreadSafe
 );
 
-class FShadowParallelCommandListSet final : public FParallelCommandListSet
-{
-public:
-	FShadowParallelCommandListSet(
-		const FRDGPass* InPass,
-		FRHICommandListImmediate& InParentCmdList,
-		const FViewInfo& InView,
-		const FProjectedShadowInfo& InProjectedShadowInfo,
-		const FParallelCommandListBindings& InBindings)
-		: FParallelCommandListSet(InPass, InView, InParentCmdList)
-		, ProjectedShadowInfo(InProjectedShadowInfo)
-		, Bindings(InBindings)
-	{}
-
-	~FShadowParallelCommandListSet() override
-	{
-		Dispatch();
-	}
-
-	void SetStateOnCommandList(FRHICommandList& RHICmdList) override
-	{
-		FParallelCommandListSet::SetStateOnCommandList(RHICmdList);
-		Bindings.SetOnCommandList(RHICmdList);
-		ProjectedShadowInfo.SetStateForView(RHICmdList);
-	}
-
-private:
-	const FProjectedShadowInfo& ProjectedShadowInfo;
-	FParallelCommandListBindings Bindings;
-};
-
 class FCopyShadowMapsCubeGS : public FGlobalShader
 {
 public:
@@ -1162,7 +1131,10 @@ void FProjectedShadowInfo::RenderDepth(
 			ERDGPassFlags::Raster,
 			[this, PassParameters](FRDGDispatchPassBuilder& DispatchPassBuilder)
 		{
-			ShadowDepthPass.Dispatch(DispatchPassBuilder, &PassParameters->InstanceCullingDrawParams);
+			ShadowDepthPass.Dispatch(DispatchPassBuilder, &PassParameters->InstanceCullingDrawParams, [this] (FRHICommandList& RHICmdList)
+			{
+				SetStateForView(RHICmdList);
+			});
 		});
 	}
 	else
