Added resolution scaling and cropping with overscan to UCameraComponent. The camera's render resolution can now be scaled proportionally to the amount of overscan, allowing all overscanned pixels to be extra pixels, as well as allowing the overscanned regions to be cropped down at the end of the render pass.

diff --git a/Engine/Source/Editor/UnrealEd/Private/EditorViewportClient.cpp b/Engine/Source/Editor/UnrealEd/Private/EditorViewportClient.cpp
--- a/Engine/Source/Editor/UnrealEd/Private/EditorViewportClient.cpp
+++ b/Engine/Source/Editor/UnrealEd/Private/EditorViewportClient.cpp
@@ -1284,6 +1284,8 @@ FSceneView* FEditorViewportClient::CalcSceneView(FSceneViewFamily* ViewFamily, c
 		ViewInitOptions.bUseFieldOfViewForLOD = ControllingActorViewInfo.bUseFieldOfViewForLOD;
 		ViewInitOptions.FOV = ControllingActorViewInfo.FOV;
 		ViewInitOptions.FirstPersonParams = FFirstPersonParameters(ControllingActorViewInfo.CalculateFirstPersonFOVCorrectionFactor(), ControllingActorViewInfo.FirstPersonScale, ControllingActorViewInfo.bUseFirstPersonParameters);
+		ViewInitOptions.OverscanResolutionFraction = ControllingActorViewInfo.OverscanResolutionFraction;
+		ViewInitOptions.CropFraction = ControllingActorViewInfo.CropFraction;
 	}
 
 	ViewInitOptions.OverrideFarClippingPlaneDistance = FarPlane;
diff --git a/Engine/Source/Runtime/Engine/Classes/Camera/CameraComponent.h b/Engine/Source/Runtime/Engine/Classes/Camera/CameraComponent.h
--- a/Engine/Source/Runtime/Engine/Classes/Camera/CameraComponent.h
+++ b/Engine/Source/Runtime/Engine/Classes/Camera/CameraComponent.h
@@ -130,12 +130,43 @@ class UCameraComponent : public USceneComponent
 	uint8 bUseFieldOfViewForLOD : 1;
 	UFUNCTION(BlueprintCallable, Category = Camera)
 	void SetUseFieldOfViewForLOD(bool bInUseFieldOfViewForLOD) { bUseFieldOfViewForLOD = bInUseFieldOfViewForLOD; }
+	
 	/** Amount to increase the view frustum by, from 0.0 for no increase to 1.0 for 100% increase */
 	UPROPERTY(Interp, EditAnywhere, BlueprintReadWrite, Category = CameraOptions, meta = (UIMin="0.0", ClampMin="0.0", UIMax="1.0", ClampMax="1.0"))
 	float Overscan;
 	UFUNCTION(BlueprintCallable, Category = Camera)
 	void SetOverscan(float InOverscan) { Overscan = InOverscan; }
 	
+	/**
+	 * Indicates that the resolution should be scaled by the overscan amount so that the original view frustum remains the same resolution.
+	 * Note that when enabled, increasing overscan will result in increased rendering workload, potentially decreasing performance as resolution increases
+	 */
+	UPROPERTY(Interp, EditAnywhere, BlueprintReadWrite, Category = CameraOptions)
+	bool bScaleResolutionWithOverscan;
+
+	/**
+	 * Sets whether to scale the resolution by the amount of overscan so that the original view frustum remains the same resolution.
+	 * Note that when enabled, increasing overscan will result in increased rendering workload, potentially decreasing performance as resolution increases
+	 */
+	UFUNCTION(BlueprintCallable, Category = Camera)
+	void SetScaleResolutionWithOverscan(bool bInScaleResolutionWithOverscan) { bScaleResolutionWithOverscan = bInScaleResolutionWithOverscan; }
+
+	/**
+	 * Indicates that the overscanned pixels should be cropped at the end of the rendering pipeline, allowing the overscanned pixels to be used in post process effects
+	 * that need extra pixels beyond the view frustum (e.g. lens distortion) without having to render those pixels to the screen. When bScaleResolutionWithOverscan is enabled,
+	 * the cropped image will have the same resolution as the original non-overscanned image, but when disabled, the cropped image will be a lower resolution.
+	 */
+	UPROPERTY(Interp, EditAnywhere, BlueprintReadWrite, Category = CameraOptions)
+	bool bCropOverscan;
+
+	/**
+	 * Sets whether to crop the overscanned pixels at the end of the rendering pipeline, allowing the overscanned pixels to be used in post process effects
+	 * that need extra pixels beyond the view frustum (e.g. lens distortion) without having to render those pixels to the screen. When bScaleResolutionWithOverscan is enabled,
+	 * the cropped image will have the same resolution as the original non-overscanned image, but when disabled, the cropped image will be a lower resolution.
+	 */
+	UFUNCTION(BlueprintCallable, Category = Camera)
+	void SetCropOverscan(bool bInCropOverscan) { bCropOverscan = bInCropOverscan; }
+
 #if WITH_EDITOR
 	// Returns the filmback text used for burnins on preview viewports
 	UFUNCTION()
diff --git a/Engine/Source/Runtime/Engine/Classes/Camera/CameraTypes.h b/Engine/Source/Runtime/Engine/Classes/Camera/CameraTypes.h
--- a/Engine/Source/Runtime/Engine/Classes/Camera/CameraTypes.h
+++ b/Engine/Source/Runtime/Engine/Classes/Camera/CameraTypes.h
@@ -131,6 +131,14 @@ struct FMinimalViewInfo
 	/** Optional transform to be considered as this view's previous transform */
 	TOptional<FTransform> PreviousViewTransform;
 
+	/** Resolution fraction that scales with the amount of overscan added to the view */
+	UPROPERTY(EditAnywhere, BlueprintReadWrite, Category=Camera)
+	float OverscanResolutionFraction;
+	
+	/** The fraction between 0.0 and 1.0 of the view to crop to during the final post process upscale, with 1.0 meaning no crop */
+	UPROPERTY(EditAnywhere, BlueprintReadWrite, Category=Camera, meta=(ClampMin=0.0, UIMin=0.0, ClampMax=1.0, UIMax=1.0))
+	float CropFraction;
+	
 private:
 	// Only used for Ortho camera auto plane calculations, tells the Near plane of the extra distance that needs to be added.
 	FVector CameraToViewTarget;
@@ -159,6 +167,8 @@ public:
 		, ProjectionMode(ECameraProjectionMode::Perspective)
 		, PostProcessBlendWeight(0.0f)
 		, OffCenterProjectionOffset(ForceInitToZero)
+		, OverscanResolutionFraction(1.0f)
+		, CropFraction(1.0f)
 		, CameraToViewTarget(FVector::ZeroVector)
 	{
 	}
@@ -218,5 +228,5 @@ public:
 	/**
 	 * Apply overscan to the view info, which scales the field of view and ortho width to simulate expanding the view frustum.
 	 */
-	ENGINE_API void ApplyOverscan(float InOverscan);
+	ENGINE_API void ApplyOverscan(float InOverscan, bool bScaleResolutionWithOverscan = false, bool bCropOverscan = false);
 };
diff --git a/Engine/Source/Runtime/Engine/Private/Camera/CameraComponent.cpp b/Engine/Source/Runtime/Engine/Private/Camera/CameraComponent.cpp
--- a/Engine/Source/Runtime/Engine/Private/Camera/CameraComponent.cpp
+++ b/Engine/Source/Runtime/Engine/Private/Camera/CameraComponent.cpp
@@ -464,7 +464,7 @@ void UCameraComponent::GetCameraView(float DeltaTime, FMinimalViewInfo& DesiredV
 	DesiredView.FirstPersonScale = bEnableFirstPersonScale ? FirstPersonScale : 1.0f;
 	DesiredView.bUseFirstPersonParameters = bEnableFirstPersonFieldOfView || bEnableFirstPersonScale;
 	
-	DesiredView.ApplyOverscan(Overscan);
+	DesiredView.ApplyOverscan(Overscan, bScaleResolutionWithOverscan, bCropOverscan);
 	
 	if (bAutoCalculateOrthoPlanes)
 	{
diff --git a/Engine/Source/Runtime/Engine/Private/Camera/CameraStackTypes.cpp b/Engine/Source/Runtime/Engine/Private/Camera/CameraStackTypes.cpp
--- a/Engine/Source/Runtime/Engine/Private/Camera/CameraStackTypes.cpp
+++ b/Engine/Source/Runtime/Engine/Private/Camera/CameraStackTypes.cpp
@@ -493,16 +493,22 @@ float FMinimalViewInfo::CalculateFirstPersonFOVCorrectionFactor() const
 	return FOVCorrectionFactor;
 }
 
-void FMinimalViewInfo::ApplyOverscan(float InOverscan)
+void FMinimalViewInfo::ApplyOverscan(float InOverscan, bool bScaleResolutionWithOverscan, bool bCropOverscan)
 {
 	if (InOverscan > 0.0)
 	{
+		// By convention, 0.0 means no overscan, so add 1 to compute the scalar needed for altering projection values
+		const float OverscanScalar = 1.0f + InOverscan;
+		
 		// Overscan directly scales the view frustum, but can be accomplished by scaling the FOV.
 		// However, must scale the tangent of the half-FOV to accomplish the same mathematical transform.
 		const float HalfFOVInRadians = FMath::DegreesToRadians(0.5f * FOV);
-		const float OverscannedFOV = FMath::Atan((1.0f + InOverscan) * FMath::Tan(HalfFOVInRadians));
+		const float OverscannedFOV = FMath::Atan(OverscanScalar * FMath::Tan(HalfFOVInRadians));
 		FOV = 2.0f * FMath::RadiansToDegrees(OverscannedFOV);
 		
-		OrthoWidth *= 1.0f + InOverscan;
+		OrthoWidth *= OverscanScalar;
+
+		OverscanResolutionFraction *= bScaleResolutionWithOverscan ? OverscanScalar : 1.0;
+		CropFraction *= bCropOverscan ? 1.0f / OverscanScalar : 1.0f;
 	}
 }
diff --git a/Engine/Source/Runtime/Engine/Private/LocalPlayer.cpp b/Engine/Source/Runtime/Engine/Private/LocalPlayer.cpp
--- a/Engine/Source/Runtime/Engine/Private/LocalPlayer.cpp
+++ b/Engine/Source/Runtime/Engine/Private/LocalPlayer.cpp
@@ -837,7 +837,9 @@ FSceneView* ULocalPlayer::CalcSceneView( class FSceneViewFamily* ViewFamily,
 	ViewInitOptions.FOV = ViewInfo.FOV;
 	ViewInitOptions.DesiredFOV = ViewInfo.DesiredFOV;
 	ViewInitOptions.FirstPersonParams = FFirstPersonParameters(ViewInfo.CalculateFirstPersonFOVCorrectionFactor(), ViewInfo.FirstPersonScale, ViewInfo.bUseFirstPersonParameters);
-
+	ViewInitOptions.OverscanResolutionFraction = ViewInfo.OverscanResolutionFraction;
+	ViewInitOptions.CropFraction = ViewInfo.CropFraction;
+	
 	// Fill out the rest of the view init options
 	ViewInitOptions.ViewFamily = ViewFamily;
 
diff --git a/Engine/Source/Runtime/Engine/Public/SceneView.h b/Engine/Source/Runtime/Engine/Public/SceneView.h
--- a/Engine/Source/Runtime/Engine/Public/SceneView.h
+++ b/Engine/Source/Runtime/Engine/Public/SceneView.h
@@ -242,6 +242,12 @@ struct FSceneViewInitOptions : public FSceneViewProjectionData
 	/** If > 0, overrides the view's resolution fraction. */
 	float OverridePrimaryResolutionFraction;
 
+	/** Resolution fraction that scales with the amount of overscan in the view */
+	float OverscanResolutionFraction = 1.0f;
+
+	/** Fraction of the view to crop to during the secondary upscale pass, with 1.0 meaning no crop */
+	float CropFraction = 1.0f;
+	
 #if WITH_EDITOR
 	/** default to 0'th view index, which is a bitfield of 1 */
 	uint64 EditorViewBitflag;
diff --git a/Engine/Source/Runtime/Renderer/Private/PostProcess/PostProcessUpscale.cpp b/Engine/Source/Runtime/Renderer/Private/PostProcess/PostProcessUpscale.cpp
--- a/Engine/Source/Runtime/Renderer/Private/PostProcess/PostProcessUpscale.cpp
+++ b/Engine/Source/Runtime/Renderer/Private/PostProcess/PostProcessUpscale.cpp
@@ -138,7 +138,8 @@ FScreenPassTexture ISpatialUpscaler::AddDefaultUpscalePass(
 		Output.UpdateVisualizeTextureExtent();
 	}
 
-	const FScreenPassTextureViewport InputViewport(Inputs.SceneColor);
+	const FIntRect InputRect = Inputs.Stage == EUpscaleStage::SecondaryToOutput ? View.GetSecondaryViewCropRect() : Inputs.SceneColor.ViewRect;
+	const FScreenPassTextureViewport InputViewport(Inputs.SceneColor.Texture, InputRect);
 	const FScreenPassTextureViewport OutputViewport(Output);
 
 	const bool bApplyLensDistortion = LensDistortionLUT.IsEnabled();
diff --git a/Engine/Source/Runtime/Renderer/Private/PostProcess/TemporalSuperResolution.cpp b/Engine/Source/Runtime/Renderer/Private/PostProcess/TemporalSuperResolution.cpp
--- a/Engine/Source/Runtime/Renderer/Private/PostProcess/TemporalSuperResolution.cpp
+++ b/Engine/Source/Runtime/Renderer/Private/PostProcess/TemporalSuperResolution.cpp
@@ -1502,7 +1502,7 @@ FDefaultTemporalUpscaler::FOutputs AddTemporalSuperResolutionPasses(
 		{
 			DynamicRenderScaling::TMap<float> DynamicResolutionUpperBounds = ScreenPercentageInterface->GetResolutionFractionsUpperBound();
 			const float PrimaryResolutionFractionUpperBound = DynamicResolutionUpperBounds[GDynamicPrimaryResolutionFraction];
-			ResolutionFractionUpperBound = PrimaryResolutionFractionUpperBound * View.Family->SecondaryViewFraction;
+			ResolutionFractionUpperBound = PrimaryResolutionFractionUpperBound * View.Family->SecondaryViewFraction * View.SceneViewInitOptions.OverscanResolutionFraction;
 		}
 
 		FIntPoint MaxRenderingViewSize = FSceneRenderer::ApplyResolutionFraction(*View.Family, View.UnconstrainedViewRect.Size(), ResolutionFractionUpperBound);
diff --git a/Engine/Source/Runtime/Renderer/Private/SceneRendering.cpp b/Engine/Source/Runtime/Renderer/Private/SceneRendering.cpp
--- a/Engine/Source/Runtime/Renderer/Private/SceneRendering.cpp
+++ b/Engine/Source/Runtime/Renderer/Private/SceneRendering.cpp
@@ -1248,8 +1248,27 @@ void SetupPhysicsFieldUniformBufferParameters(const FScene* Scene, FEngineShowFl
 FIntPoint FViewInfo::GetSecondaryViewRectSize() const
 {
 	return FIntPoint(
-		FMath::CeilToInt(UnscaledViewRect.Width() * Family->SecondaryViewFraction),
-		FMath::CeilToInt(UnscaledViewRect.Height() * Family->SecondaryViewFraction));
+		FMath::CeilToInt(UnscaledViewRect.Width() * Family->SecondaryViewFraction * SceneViewInitOptions.OverscanResolutionFraction),
+		FMath::CeilToInt(UnscaledViewRect.Height() * Family->SecondaryViewFraction * SceneViewInitOptions.OverscanResolutionFraction));
+}
+
+FIntRect FViewInfo::GetSecondaryViewCropRect() const
+{
+	const FIntPoint SecondaryVewRectSize = GetSecondaryViewRectSize();
+
+	// Clamp the crop fraction to sensible values to ensure crop rect is always a valid rectangle
+	const float CropFrac = FMath::Clamp(SceneViewInitOptions.CropFraction, 0.0f, 1.0f);
+	
+	FIntRect CropRect;
+	CropRect.Min = FIntPoint(
+			FMath::FloorToInt(0.5f * (1.0f - CropFrac) * SecondaryVewRectSize.X),
+			FMath::FloorToInt(0.5f * (1.0f - CropFrac) * SecondaryVewRectSize.Y));
+		
+	CropRect.Max = FIntPoint(
+		FMath::CeilToInt(0.5f * (1.0f + CropFrac) * SecondaryVewRectSize.X),
+		FMath::CeilToInt(0.5f * (1.0f + CropFrac) * SecondaryVewRectSize.Y));
+
+	return CropRect;
 }
 
 /** Creates the view's uniform buffers given a set of view transforms. */
@@ -3159,11 +3178,21 @@ FIntPoint FSceneRenderer::GetDesiredInternalBufferSize(const FSceneViewFamily& V
 
 	FIntPoint FamilySizeUpperBound(0, 0);
 
+	// For multiple views, use the maximum overscan fraction to ensure that enough space is allocated so that any overscanned views
+	// do not encroach into the space of other views
+	float MaxOverscanResolutionFraction = 1.0f;
+	for (const FSceneView* View : ViewFamily.AllViews)
+	{
+		MaxOverscanResolutionFraction = FMath::Max(MaxOverscanResolutionFraction, View->SceneViewInitOptions.OverscanResolutionFraction);
+	}
+
+	ResolutionFractionUpperBound *= MaxOverscanResolutionFraction;
+	
 	for (const FSceneView* View : ViewFamily.AllViews)
 	{
 		// Note: This ensures that custom passes (rendered with the main renderer) ignore screen percentage, like regular scene captures.
 		const float AdjustedResolutionFractionUpperBounds = View->CustomRenderPass ? 1.0f : (View->SceneViewInitOptions.OverridePrimaryResolutionFraction > 0.0 ? (View->SceneViewInitOptions.OverridePrimaryResolutionFraction * ViewFamily.SecondaryViewFraction)  : ResolutionFractionUpperBound);
-
+		
 		FIntPoint ViewSize = ApplyResolutionFraction(ViewFamily, View->UnconstrainedViewRect.Size(), AdjustedResolutionFractionUpperBounds);
 		FIntPoint ViewRectMin = QuantizeViewRectMin(FIntPoint(
 			FMath::CeilToInt(View->UnconstrainedViewRect.Min.X * AdjustedResolutionFractionUpperBounds),
@@ -3317,6 +3346,14 @@ void FSceneRenderer::PrepareViewRectsForRendering(FRHICommandListImmediate& RHIC
 		LensDistortionResolutionFraction = FMath::Lerp(1.0, Views[0].LensDistortionLUT.ResolutionFraction, AffectScreenPercentage);
 	}
 
+	// For multiple views, we must find the maximum overscan resolution so that views can be offset appropriately to avoid overscanned
+	// views encroaching into other views' buffer space
+	float MaxOverscanResolutionFraction = 1.0f;
+	for (const FSceneView* View : ViewFamily.AllViews)
+	{
+		MaxOverscanResolutionFraction = FMath::Max(MaxOverscanResolutionFraction, View->SceneViewInitOptions.OverscanResolutionFraction);
+	}
+	
 	// Compute final resolution fraction.
 	float ResolutionFraction = PrimaryResolutionFraction * ViewFamily.SecondaryViewFraction * LensDistortionResolutionFraction;
 
@@ -3327,10 +3364,10 @@ void FSceneRenderer::PrepareViewRectsForRendering(FRHICommandListImmediate& RHIC
 
 		float ViewResolutionFraction = View.SceneViewInitOptions.OverridePrimaryResolutionFraction > 0.0 ? (View.SceneViewInitOptions.OverridePrimaryResolutionFraction * ViewFamily.SecondaryViewFraction) : ResolutionFraction;
 
-		FIntPoint ViewSize = ApplyResolutionFraction(ViewFamily, View.UnscaledViewRect.Size(), ViewResolutionFraction);
+		FIntPoint ViewSize = ApplyResolutionFraction(ViewFamily, View.UnscaledViewRect.Size(), ViewResolutionFraction * View.SceneViewInitOptions.OverscanResolutionFraction);
 		FIntPoint ViewRectMin = QuantizeViewRectMin(FIntPoint(
-			FMath::CeilToInt(View.UnscaledViewRect.Min.X * ViewResolutionFraction),
-			FMath::CeilToInt(View.UnscaledViewRect.Min.Y * ViewResolutionFraction)));
+			FMath::CeilToInt(View.UnscaledViewRect.Min.X * ViewResolutionFraction * MaxOverscanResolutionFraction),
+			FMath::CeilToInt(View.UnscaledViewRect.Min.Y * ViewResolutionFraction * MaxOverscanResolutionFraction)));
 
 		// Use the bottom-left view rect if requested, instead of top-left
 		if (CVarViewRectUseScreenBottom.GetValueOnRenderThread())
diff --git a/Engine/Source/Runtime/Renderer/Private/SceneRendering.h b/Engine/Source/Runtime/Renderer/Private/SceneRendering.h
--- a/Engine/Source/Runtime/Renderer/Private/SceneRendering.h
+++ b/Engine/Source/Runtime/Renderer/Private/SceneRendering.h
@@ -1582,11 +1582,14 @@ public:
 
 	/** Returns the size of view rect after primary upscale ( == only with secondary screen percentage). */
 	RENDERER_API FIntPoint GetSecondaryViewRectSize() const;
+
+	/** Returns the rectangle to crop to during the secondary upscale */
+	RENDERER_API FIntRect GetSecondaryViewCropRect() const;
 	
 	/** Returns whether the view requires a secondary upscale. */
 	bool RequiresSecondaryUpscale() const
 	{
-		return UnscaledViewRect.Size() != GetSecondaryViewRectSize();
+		return UnscaledViewRect.Size() != GetSecondaryViewRectSize() || GetSecondaryViewCropRect().Size() != GetSecondaryViewRectSize();
 	}
 
 	/** Compute the pre-exposure of internal renderer resources. */
