Fixed crash due to missing test for null scene proxy in FScene::GetRelevantLights

diff --git a/Engine/Source/Runtime/Renderer/Private/RendererScene.cpp b/Engine/Source/Runtime/Renderer/Private/RendererScene.cpp
--- a/Engine/Source/Runtime/Renderer/Private/RendererScene.cpp
+++ b/Engine/Source/Runtime/Renderer/Private/RendererScene.cpp
@@ -4362,7 +4362,7 @@ void FScene::GetRelevantLights_RenderThread( const FPrimitiveSceneProxy* Primiti
  */
 void FScene::GetRelevantLights( UPrimitiveComponent* Primitive, TArray<const ULightComponent*>* RelevantLights ) const
 {
-	if( Primitive && RelevantLights )
+	if( Primitive && RelevantLights && Primitive->GetSceneProxy() != nullptr)
 	{
 		// Add interacting lights to the array.
 		ENQUEUE_RENDER_COMMAND(FGetRelevantLightsCommand)(
