Add VSM counters to Unreal Insights tracing. Requires both VSM stats and Insights counters to be enabled.  (`r.Shadow.Virtual.ShowStats 1` and `trace.enable counters`)

diff --git a/Engine/Source/Runtime/Renderer/Private/VirtualShadowMaps/VirtualShadowMapCacheManager.cpp b/Engine/Source/Runtime/Renderer/Private/VirtualShadowMaps/VirtualShadowMapCacheManager.cpp
--- a/Engine/Source/Runtime/Renderer/Private/VirtualShadowMaps/VirtualShadowMapCacheManager.cpp
+++ b/Engine/Source/Runtime/Renderer/Private/VirtualShadowMaps/VirtualShadowMapCacheManager.cpp
@@ -686,6 +686,28 @@ void FVirtualShadowMapArrayCacheManager::InitExtension(FScene& InScene)
 #if CSV_PROFILER
 		CSV_CUSTOM_STAT(VSM, NonNanitePostCullInstanceCount, int32(Stats[VSM_STAT_NON_NANITE_INSTANCES_DRAWN]), ECsvCustomStatOp::Set);
 
+		// requires 'trace.enable counters' and vsm stats to be enabled to see this in ue insights
+		TRACE_INT_VALUE(TEXT("Shadow.Virtual.PagesRequested"), Stats[VSM_STAT_REQUESTED_THIS_FRAME_PAGES]);
+		TRACE_INT_VALUE(TEXT("Shadow.Virtual.PagesCachedStatic"), Stats[VSM_STAT_STATIC_CACHED_PAGES]);
+		TRACE_INT_VALUE(TEXT("Shadow.Virtual.PagesInvalidatedStatic"), Stats[VSM_STAT_STATIC_INVALIDATED_PAGES]);
+		TRACE_INT_VALUE(TEXT("Shadow.Virtual.PagesCachedDynamic"), Stats[VSM_STAT_DYNAMIC_CACHED_PAGES]);
+		TRACE_INT_VALUE(TEXT("Shadow.Virtual.PagesInvalidatedDynamic"), Stats[VSM_STAT_DYNAMIC_INVALIDATED_PAGES]);
+		TRACE_INT_VALUE(TEXT("Shadow.Virtual.PagesEmpty"), Stats[VSM_STAT_EMPTY_PAGES]);
+		TRACE_INT_VALUE(TEXT("Shadow.Virtual.NonNanite.InstancesTotal"), Stats[VSM_STAT_NON_NANITE_INSTANCES_TOTAL]);
+		TRACE_INT_VALUE(TEXT("Shadow.Virtual.NonNanite.InstancesDrawn"), Stats[VSM_STAT_NON_NANITE_INSTANCES_DRAWN]);
+		TRACE_INT_VALUE(TEXT("Shadow.Virtual.NonNanite.InstancesHZBCulled"), Stats[VSM_STAT_NON_NANITE_INSTANCES_HZB_CULLED]);
+		TRACE_INT_VALUE(TEXT("Shadow.Virtual.NonNanite.InstancesPageMaskCulled"), Stats[VSM_STAT_NON_NANITE_INSTANCES_PAGE_MASK_CULLED]);
+		TRACE_INT_VALUE(TEXT("Shadow.Virtual.NonNanite.InstancesEmptyRectCulled"), Stats[VSM_STAT_NON_NANITE_INSTANCES_EMPTY_RECT_CULLED]);
+		TRACE_INT_VALUE(TEXT("Shadow.Virtual.NonNanite.InstancesFrustumCulled"), Stats[VSM_STAT_NON_NANITE_INSTANCES_FRUSTUM_CULLED]);
+		TRACE_INT_VALUE(TEXT("Shadow.Virtual.PagesToMerge"), Stats[VSM_STAT_NUM_PAGES_TO_MERGE]);
+		TRACE_INT_VALUE(TEXT("Shadow.Virtual.PagesToClear"), Stats[VSM_STAT_NUM_PAGES_TO_CLEAR]);
+		TRACE_INT_VALUE(TEXT("Shadow.Virtual.HZBPagesBuilt"), Stats[VSM_STAT_NUM_HZB_PAGES_BUILT]);
+		TRACE_INT_VALUE(TEXT("Shadow.Virtual.PagesAllocatedNew"), Stats[VSM_STAT_ALLOCATED_NEW]);
+		TRACE_INT_VALUE(TEXT("Shadow.Virtual.Nanite.Triangles"), Stats[VSM_STAT_NANITE_TRIANGLES]);
+		TRACE_INT_VALUE(TEXT("Shadow.Virtual.Nanite.InstancesMain"), Stats[VSM_STAT_NANITE_INSTANCES_MAIN]);
+		TRACE_INT_VALUE(TEXT("Shadow.Virtual.Nanite.InstancesPost"), Stats[VSM_STAT_NANITE_INSTANCES_POST]);
+		TRACE_INT_VALUE(TEXT("Shadow.Virtual.PagesWPOConsidered"), Stats[VSM_STAT_WPO_CONSIDERED_PAGES]);
+
 		if (FCsvProfiler::Get()->IsCapturing_Renderthread())
 		{
 			static bool bRegisteredInlineStats = false;
