[VSM] Handle non-nanite job marking queue overflow more gracefully.
Track overflows of the non-nanite job marking queue, visible instances buffer, one-pass-projection lights-per-pixel and physical page pool. Produce a warning when overflows occur.

diff --git a/Engine/Shaders/Private/VirtualShadowMaps/VirtualShadowMapBuildPerPageDrawCommands.usf b/Engine/Shaders/Private/VirtualShadowMaps/VirtualShadowMapBuildPerPageDrawCommands.usf
--- a/Engine/Shaders/Private/VirtualShadowMaps/VirtualShadowMapBuildPerPageDrawCommands.usf
+++ b/Engine/Shaders/Private/VirtualShadowMaps/VirtualShadowMapBuildPerPageDrawCommands.usf
@@ -56,7 +56,7 @@ uint GetInclusiveRectArea(uint4 Rect)
 	return 0;
 }
 
-void WriteCmd(uint MipViewId, uint InstanceId, uint IndirectArgIndex, uint CullingFlags, bool bStaticPage)
+bool WriteCmd(uint MipViewId, uint InstanceId, uint IndirectArgIndex, uint CullingFlags, bool bStaticPage)
 {
 	FPageInfo PageInfo;
 	PageInfo.ViewId = MipViewId;
@@ -71,8 +71,12 @@ void WriteCmd(uint MipViewId, uint InstanceId, uint IndirectArgIndex, uint Culli
 	WaveInterlockedAddScalar_(VisibleInstanceCountBufferOut[0], 1U, VisibleInstanceOutputOffset);
 	if (VisibleInstanceOutputOffset < VisibleInstancesBufferNum)
 	{
-		// TODO: Flag & feedback on overflow.
 		VisibleInstancesOut[VisibleInstanceOutputOffset] = VisibleInstanceCmd;
+		return true;
+	}
+	else
+	{
+		return false;
 	}
 }
 
@@ -210,6 +214,9 @@ void CullPerPageDrawCommandsCs(uint3 GroupId : SV_GroupID, int GroupThreadIndex
 		uint TotalPageArea = 0U;
 		float NumVisibleLevels = 0.0f;
 #endif
+		bool bMarkingJobQueueOverflow = false;
+		bool bVisibleInstancesOverflow = false;
+
 		// Loop over views and output visible instance (i.e., those that overlap a valid page)
 		for (uint PrimaryViewId = VSMCullingBatchInfo.FirstPrimaryView; PrimaryViewId < VSMCullingBatchInfo.FirstPrimaryView + VSMCullingBatchInfo.NumPrimaryViews; ++PrimaryViewId)
 		{
@@ -319,9 +326,19 @@ void CullPerPageDrawCommandsCs(uint3 GroupId : SV_GroupID, int GroupThreadIndex
 
 						uint NumMappedPages = 0U;
 						{
-
 							uint2 RectPagesSize = (RectPages.zw + 1u) - RectPages.xy;
-							if (RectPagesSize.x * RectPagesSize.y <= MAX_SINGLE_THREAD_MARKING_AREA)
+							bool bIsSmallJob = RectPagesSize.x * RectPagesSize.y <= MAX_SINGLE_THREAD_MARKING_AREA;
+
+							uint LargeJobIndex = 0U;
+							if (!bIsSmallJob)
+							{
+								// allocate slot to defer the work and use all threads in group
+								InterlockedAdd(NumSharedMarkingJobs, 1U, LargeJobIndex);
+							}
+
+							bMarkingJobQueueOverflow = bMarkingJobQueueOverflow || LargeJobIndex >= MARKING_JOB_QUEUE_SIZE;
+
+							if (bIsSmallJob || bMarkingJobQueueOverflow)
 							{
 								FVirtualSMLevelOffset PageTableLevelOffset = CalcPageTableLevelOffset(VirtualShadowMapId, MipLevel);
 
@@ -340,8 +357,6 @@ void CullPerPageDrawCommandsCs(uint3 GroupId : SV_GroupID, int GroupThreadIndex
 							else
 							{
 								// defer the work and use all threads in group
-								uint JobIndex = 0U;
-								InterlockedAdd(NumSharedMarkingJobs, 1U, JobIndex);
 								FSharedMarkingJob Job;
 								Job.RectPages = RectPages;
 								Job.VirtualShadowMapId = VirtualShadowMapId;
@@ -350,21 +365,16 @@ void CullPerPageDrawCommandsCs(uint3 GroupId : SV_GroupID, int GroupThreadIndex
 								Job.bMaterialInvalidates = bMaterialInvalidates;
 								Job.bShouldCacheAsStatic = bCacheAsStatic;
 								Job.bIsViewUncached = Cull.bIsViewUncached;
-								if (JobIndex < MARKING_JOB_QUEUE_SIZE)
-								{
-									SharedMarkingJobs[JobIndex] = PackJob(Job);
-								}
+								SharedMarkingJobs[LargeJobIndex] = PackJob(Job);
 								// Must assume we have mapped pages (or defer the cmd write also...)
 								NumMappedPages = 1U;
 							}
-
-							// 
 						}
 
 						if (NumMappedPages > 0U)
 						{
 							++ThreadTotalForAllViews;
-							WriteCmd(MipViewId, InstanceId, Payload.IndirectArgIndex, CullingFlags, bCacheAsStatic);
+							bVisibleInstancesOverflow |= !WriteCmd(MipViewId, InstanceId, Payload.IndirectArgIndex, CullingFlags, bCacheAsStatic);
 						}
 					}
 					else
@@ -378,8 +388,10 @@ void CullPerPageDrawCommandsCs(uint3 GroupId : SV_GroupID, int GroupThreadIndex
 				StatsBufferInterlockedAdd(VSM_STAT_NON_NANITE_INSTANCES_FRUSTUM_CULLED, NaniteView.TargetNumMipLevels);
 			}
 		}
-
+		
 		StatsBufferInterlockedAdd(VSM_STAT_NON_NANITE_INSTANCES_DRAWN, ThreadTotalForAllViews);
+		StatsBufferInterlockedEnableFlags<true>(VSM_STAT_OVERFLOW_FLAGS, VSM_STAT_OVERFLOW_FLAG_MARKING_JOB_QUEUE, bMarkingJobQueueOverflow);
+		StatsBufferInterlockedEnableFlags<true>(VSM_STAT_OVERFLOW_FLAGS, VSM_STAT_OVERFLOW_FLAG_VISIBLE_INSTANCES, bVisibleInstancesOverflow);
 
 		// Accumulate total number of instances for each indirect argument, is also used to allocate space and output compact range of instances later
 		InterlockedAdd(DrawIndirectArgsBufferOut[Payload.IndirectArgIndex * INDIRECT_ARGS_NUM_WORDS + 1], ThreadTotalForAllViews);
diff --git a/Engine/Shaders/Private/VirtualShadowMaps/VirtualShadowMapPhysicalPageManagement.usf b/Engine/Shaders/Private/VirtualShadowMaps/VirtualShadowMapPhysicalPageManagement.usf
--- a/Engine/Shaders/Private/VirtualShadowMaps/VirtualShadowMapPhysicalPageManagement.usf
+++ b/Engine/Shaders/Private/VirtualShadowMaps/VirtualShadowMapPhysicalPageManagement.usf
@@ -1197,16 +1197,16 @@ StructuredBuffer<int> PhysicalPageLists;
 [numthreads(1, 1, 1)]
 void FeedbackStatusCS()
 {
-	{
-		FGPUMessageWriter Mw = GPUMessageBegin(StatusMessageId, 2U);
+	FGPUMessageWriter Mw = GPUMessageBegin(StatusMessageId, 3U);
 
-		// Write out how many pages are still available
-		int CountIndex = GetPhysicalPageListStart(PHYSICAL_PAGE_LIST_AVAILABLE) + VirtualShadowMap.MaxPhysicalPages;
-		GPUMessageWriteItem(Mw, PhysicalPageLists[CountIndex]);
+	GPUMessageWriteItem(Mw, VSM_STATUS_MSG_PAGE_MANAGEMENT);
 
-		// Write out the resolution lod bias from this frame
-		GPUMessageWriteItem(Mw, VirtualShadowMap.GlobalResolutionLodBias);
-	}
+	// Write out how many pages are still available
+	int CountIndex = GetPhysicalPageListStart(PHYSICAL_PAGE_LIST_AVAILABLE) + VirtualShadowMap.MaxPhysicalPages;
+	GPUMessageWriteItem(Mw, PhysicalPageLists[CountIndex]);
+
+	// Write out the resolution lod bias from this frame
+	GPUMessageWriteItem(Mw, VirtualShadowMap.GlobalResolutionLodBias);
 }
 
 
diff --git a/Engine/Shaders/Private/VirtualShadowMaps/VirtualShadowMapPrintStats.usf b/Engine/Shaders/Private/VirtualShadowMaps/VirtualShadowMapPrintStats.usf
--- a/Engine/Shaders/Private/VirtualShadowMaps/VirtualShadowMapPrintStats.usf
+++ b/Engine/Shaders/Private/VirtualShadowMaps/VirtualShadowMapPrintStats.usf
@@ -55,12 +55,22 @@ void PrintUnits(inout FShaderPrintContext Context, uint Value)
 		Print(Context, Value);
 	}
 }
+
+uint StatusMessageId;
+
 uint StatsMessageId;
 uint StatsMessageTimestampHigh;
 uint StatsMessageTimestampLow;
 
 StructuredBuffer<FNaniteStats> NaniteStats;
 
+void SendStatusMessage()
+{
+	FGPUMessageWriter Mw = GPUMessageBegin(StatusMessageId, 2U);
+	GPUMessageWriteItem(Mw, VSM_STATUS_MSG_OVERFLOW);
+	GPUMessageWriteItem(Mw, InStatsBuffer[VSM_STAT_OVERFLOW_FLAGS]);
+}
+
 void SendStatsMessage()
 {
 	FGPUMessageWriter Mw = GPUMessageBegin(StatsMessageId, 2U  + VSM_STAT_NUM + (MAX_PAGE_AREA_DIAGNOSTIC_SLOTS * 2U));
@@ -81,19 +91,8 @@ void SendStatsMessage()
 	}
 }
 
-[numthreads(1, 1, 1)]
-void LogVirtualSmStatsCS()
+void PrintStats()
 {
-	if (StatsMessageId != INDEX_NONE)
-	{
-		SendStatsMessage();
-	}
-
-	if (ShowStatsValue == 0U)
-	{
-		return;
-	}
-
 	const bool bShowAll                = ShowStatsValue == 1;
 	// Work around potential compiler bug...
 	const bool bShowPhysicalPages      = (ShowStatsValue == 1) || ShowStatsValue == 2;
@@ -258,6 +257,26 @@ void LogVirtualSmStatsCS()
 	}
 }
 
+[numthreads(1, 1, 1)]
+void LogVirtualSmStatsCS()
+{
+	if (StatusMessageId != INDEX_NONE)
+	{
+		SendStatusMessage();
+	}
+	
+#if VSM_GENERATE_STATS
+	if (StatsMessageId != INDEX_NONE)
+	{
+		SendStatsMessage();
+	}
+
+	if (ShowStatsValue > 0U)
+	{
+		PrintStats();
+	}
+#endif
+}
 
 uint ShadowMapIdRangeStart;
 uint ShadowMapIdRangeEnd;
diff --git a/Engine/Shaders/Private/VirtualShadowMaps/VirtualShadowMapProjection.usf b/Engine/Shaders/Private/VirtualShadowMaps/VirtualShadowMapProjection.usf
--- a/Engine/Shaders/Private/VirtualShadowMaps/VirtualShadowMapProjection.usf
+++ b/Engine/Shaders/Private/VirtualShadowMaps/VirtualShadowMapProjection.usf
@@ -19,6 +19,7 @@
 #include "VirtualShadowMapTransmissionCommon.ush"
 #include "VirtualShadowMapScreenRayTrace.ush"
 #include "VirtualShadowMapLightGrid.ush"
+#include "VirtualShadowMapStats.ush"
 #include "/Engine/Shared/VirtualShadowMapDefinitions.h"
 
 #if HAS_HAIR_STRANDS
@@ -641,6 +642,8 @@ void VirtualShadowMapProjection(
 	const FCulledLightsGridHeader LightGridHeader = VirtualShadowMapGetLightsGrid(LocalPixelPos, SceneDepth);	
 	// We can only handle so many lights in our output encoding right now, so no purpose in computing more
 	uint LightCount = min(GetPackedShadowMaskMaxLightCount(), LightGridHeader.NumLights);
+	bool bLightCountOverflow = GetPackedShadowMaskMaxLightCount() < LightGridHeader.NumLights;
+	StatsBufferInterlockedEnableFlags<true>(VSM_STAT_OVERFLOW_FLAGS, VSM_STAT_OVERFLOW_FLAG_OPP_MAX_LIGHTS, bLightCountOverflow);
 
 	uint4 ShadowMask = InitializePackedShadowMask();
 
diff --git a/Engine/Shaders/Private/VirtualShadowMaps/VirtualShadowMapStats.ush b/Engine/Shaders/Private/VirtualShadowMaps/VirtualShadowMapStats.ush
--- a/Engine/Shaders/Private/VirtualShadowMaps/VirtualShadowMapStats.ush
+++ b/Engine/Shaders/Private/VirtualShadowMaps/VirtualShadowMapStats.ush
@@ -15,27 +15,43 @@
 #define VSM_GENERATE_STATS 0
 #endif // VSM_GENERATE_STATS
 
-#if VSM_GENERATE_STATS
 RWStructuredBuffer<uint> OutStatsBuffer;
-#endif // VSM_GENERATE_STATS
 
+template<bool bWriteEnabled = VSM_GENERATE_STATS>
 void StatsBufferWrite(uint Index, uint Value)
 {
-#if VSM_GENERATE_STATS
-	OutStatsBuffer[Index] = Value;
-#endif
+	if (bWriteEnabled)
+	{
+		OutStatsBuffer[Index] = Value;
+	}
 }
 
+template<bool bWriteEnabled = VSM_GENERATE_STATS>
 void StatsBufferInterlockedAdd(uint Index, uint Value)
 {
-#if VSM_GENERATE_STATS
-	InterlockedAdd(OutStatsBuffer[Index], Value);
-#endif
+	if (bWriteEnabled)
+	{
+		InterlockedAdd(OutStatsBuffer[Index], Value);
+	}
 }
 
+template<bool bWriteEnabled = VSM_GENERATE_STATS>
 void StatsBufferInterlockedInc(uint Index)
 {
-#if VSM_GENERATE_STATS
-	WaveInterlockedAddScalar(OutStatsBuffer[Index], 1U);
-#endif
+	if (bWriteEnabled)
+	{
+		WaveInterlockedAddScalar(OutStatsBuffer[Index], 1U);
+	}
+}
+
+template<bool bWriteEnabled = VSM_GENERATE_STATS>
+void StatsBufferInterlockedEnableFlags(uint Index, uint Flags, bool bEnabled)
+{
+	if (bWriteEnabled)
+	{
+		if (bEnabled)
+		{
+			WaveInterlockedOr(OutStatsBuffer[Index], Flags * (uint)bEnabled);
+		}
+	}
 }
diff --git a/Engine/Shaders/Shared/VirtualShadowMapDefinitions.h b/Engine/Shaders/Shared/VirtualShadowMapDefinitions.h
--- a/Engine/Shaders/Shared/VirtualShadowMapDefinitions.h
+++ b/Engine/Shaders/Shared/VirtualShadowMapDefinitions.h
@@ -54,10 +54,20 @@ VirtualShadowMapDefinitions.h: used in virtual shadow map shaders and C++ code t
 #define VSM_STAT_NANITE_INSTANCES_MAIN					17
 #define VSM_STAT_NANITE_INSTANCES_POST					18
 #define VSM_STAT_WPO_CONSIDERED_PAGES					19
-#define VSM_STAT_TMP_1									20
-#define VSM_STAT_TMP_2									21
-#define VSM_STAT_TMP_3									22
-#define VSM_STAT_NUM									23
+#define VSM_STAT_OVERFLOW_FLAGS							20
+#define VSM_STAT_TMP_1									21
+#define VSM_STAT_TMP_2									22
+#define VSM_STAT_TMP_3									23
+#define VSM_STAT_NUM									24
+
+#define VSM_STAT_OVERFLOW_FLAG_MARKING_JOB_QUEUE 		(1<<0)
+#define VSM_STAT_OVERFLOW_FLAG_OPP_MAX_LIGHTS	 		(1<<1)
+#define VSM_STAT_OVERFLOW_FLAG_PAGE_POOL		 		(1<<2)
+#define VSM_STAT_OVERFLOW_FLAG_VISIBLE_INSTANCES 		(1<<3)
+#define VSM_STAT_OVERFLOW_FLAG_NUM						4
+
+#define VSM_STATUS_MSG_PAGE_MANAGEMENT					0
+#define VSM_STATUS_MSG_OVERFLOW							1
 
 #ifdef __cplusplus
 #include "HLSLTypeAliases.h"
diff --git a/Engine/Source/Runtime/Renderer/Private/VirtualShadowMaps/VirtualShadowMapArray.cpp b/Engine/Source/Runtime/Renderer/Private/VirtualShadowMaps/VirtualShadowMapArray.cpp
--- a/Engine/Source/Runtime/Renderer/Private/VirtualShadowMaps/VirtualShadowMapArray.cpp
+++ b/Engine/Source/Runtime/Renderer/Private/VirtualShadowMaps/VirtualShadowMapArray.cpp
@@ -515,18 +515,6 @@ static float GetNormalBiasForShader()
 	return CVarNormalBias.GetValueOnRenderThread() / 1000.0f;
 }
 
-template <typename ShaderType>
-static bool SetStatsArgsAndPermutation(FRDGBufferUAVRef StatsBufferUAV, typename ShaderType::FParameters *OutPassParameters, typename ShaderType::FPermutationDomain& OutPermutationVector)
-{
-	bool bGenerateStats = StatsBufferUAV != nullptr;
-	if (bGenerateStats)
-	{
-		OutPassParameters->OutStatsBuffer = StatsBufferUAV;
-	}
-	OutPermutationVector.template Set<typename ShaderType::FGenerateStatsDim>(bGenerateStats);
-	return bGenerateStats;
-}
-
 FVirtualShadowMapArray::FVirtualShadowMapArray(FScene& InScene) 
 	: Scene(InScene)
 {
@@ -1177,7 +1165,7 @@ void FVirtualShadowMapArray::MergeStaticPhysicalPages(FRDGBuilder& GraphBuilder)
 		PassParameters->OutPhysicalPagesToMerge = GraphBuilder.CreateUAV(PhysicalPagesToMergeRDG);
 
 		FSelectPagesToMergeCS::FPermutationDomain PermutationVector;
-		SetStatsArgsAndPermutation<FSelectPagesToMergeCS>(StatsBufferUAV, PassParameters, PermutationVector);
+		SetStatsArgsAndPermutation<FSelectPagesToMergeCS>(ShouldGenerateStats(), StatsBufferUAV, PassParameters, PermutationVector);
 
 		auto ComputeShader = GetGlobalShaderMap(Scene.GetFeatureLevel())->GetShader<FSelectPagesToMergeCS>(PermutationVector);
 
@@ -1426,7 +1414,7 @@ void FVirtualShadowMapArray::UpdatePhysicalPageAddresses(FRDGBuilder& GraphBuild
 	}
 	PermutationVector.Set<FUpdatePhysicalPageAddresses::FHasCacheDataDim>(PrevPageRequestFlags != nullptr);
 
-	SetStatsArgsAndPermutation<FUpdatePhysicalPageAddresses>(StatsBufferUAV, PassParameters, PermutationVector);
+	SetStatsArgsAndPermutation<FUpdatePhysicalPageAddresses>(ShouldGenerateStats(), StatsBufferUAV, PassParameters, PermutationVector);
 
 	auto ComputeShader = GetGlobalShaderMap(Scene.GetFeatureLevel())->GetShader<FUpdatePhysicalPageAddresses>(PermutationVector);
 
@@ -1439,6 +1427,25 @@ void FVirtualShadowMapArray::UpdatePhysicalPageAddresses(FRDGBuilder& GraphBuild
 	);
 }
 
+bool FVirtualShadowMapArray::ShouldGenerateStats() const
+{
+#if !UE_BUILD_SHIPPING
+	const bool bRunPageAreaDiagnostics = CVarNumPageAreaDiagSlots.GetValueOnRenderThread() != 0;
+#else
+	constexpr bool bRunPageAreaDiagnostics = false;
+#endif
+	return CVarShowStats.GetValueOnRenderThread() || CacheManager->IsAccumulatingStats() || bRunPageAreaDiagnostics || IsCsvLogEnabled();
+}
+
+bool FVirtualShadowMapArray::IsCsvLogEnabled() const
+{
+#if CSV_PROFILER_STATS
+	return FCsvProfiler::Get()->IsCapturing_Renderthread() && FCsvProfiler::Get()->IsCategoryEnabled(CSV_CATEGORY_INDEX(VSM));
+#else
+	return false;
+#endif
+}
+
 void FVirtualShadowMapArray::BuildPageAllocations(
 	FRDGBuilder& GraphBuilder,
 	const FMinimalSceneTextures& SceneTextures,
@@ -1507,17 +1514,9 @@ void FVirtualShadowMapArray::BuildPageAllocations(
 	UniformParameters.CoarsePagePixelThresholdStatic = CVarCoarsePagePixelThresholdStatic.GetValueOnRenderThread();
 	UniformParameters.CoarsePagePixelThresholdDynamicNanite = CVarCoarsePagePixelThresholdDynamicNanite.GetValueOnRenderThread();
 
-	bool bCsvLogEnabled = false;
-#if !UE_BUILD_SHIPPING
-	const bool bRunPageAreaDiagnostics = CVarNumPageAreaDiagSlots.GetValueOnRenderThread() != 0;
-#if CSV_PROFILER_STATS
-	bCsvLogEnabled = FCsvProfiler::Get()->IsCapturing_Renderthread() && FCsvProfiler::Get()->IsCategoryEnabled(CSV_CATEGORY_INDEX(VSM));
-#endif
-#else
-	constexpr bool bRunPageAreaDiagnostics = false;
-#endif
+	bool bCsvLogEnabled = IsCsvLogEnabled();
 
-	if (CVarShowStats.GetValueOnRenderThread() || CacheManager->IsAccumulatingStats() || bRunPageAreaDiagnostics || bCsvLogEnabled)
+	// Stats buffer
 	{
 		StatsBufferRDG = GraphBuilder.CreateBuffer(FRDGBufferDesc::CreateStructuredDesc(sizeof(uint32), VSM_STAT_NUM + MaxPageAreaDiagnosticSlots * 2), TEXT("Shadow.Virtual.StatsBuffer"));
 		AddClearUAVPass(GraphBuilder, GraphBuilder.CreateUAV(StatsBufferRDG), 0);
@@ -1803,7 +1802,7 @@ void FVirtualShadowMapArray::BuildPageAllocations(
 
 		FUpdatePhysicalPages::FPermutationDomain PermutationVector;
 		PermutationVector.Set<FUpdatePhysicalPages::FHasCacheDataDim>(bCacheDataAvailable);
-		SetStatsArgsAndPermutation<FUpdatePhysicalPages>(StatsBufferUAV, PassParameters, PermutationVector);
+		SetStatsArgsAndPermutation<FUpdatePhysicalPages>(ShouldGenerateStats(), StatsBufferUAV, PassParameters, PermutationVector);
 		auto ComputeShader = Views[0].ShaderMap->GetShader<FUpdatePhysicalPages>(PermutationVector);
 
 		FComputeShaderUtils::AddPass(
@@ -1845,7 +1844,7 @@ void FVirtualShadowMapArray::BuildPageAllocations(
 		PassParameters->OutPhysicalPageMetaData = GraphBuilder.CreateUAV(PhysicalPageMetaDataRDG);
 
 		FAllocateNewPageMappingsCS::FPermutationDomain PermutationVector;
-		SetStatsArgsAndPermutation<FAllocateNewPageMappingsCS>(StatsBufferUAV, PassParameters, PermutationVector);
+		SetStatsArgsAndPermutation<FAllocateNewPageMappingsCS>(ShouldGenerateStats(), StatsBufferUAV, PassParameters, PermutationVector);
 		auto ComputeShader = Views[0].ShaderMap->GetShader<FAllocateNewPageMappingsCS>(PermutationVector);
 
 		FComputeShaderUtils::AddPass(
@@ -1912,7 +1911,7 @@ void FVirtualShadowMapArray::BuildPageAllocations(
 			PassParameters->OutInitializePagesIndirectArgsBuffer = GraphBuilder.CreateUAV(InitializePagesIndirectArgsRDG);
 			PassParameters->OutPhysicalPagesToInitialize = GraphBuilder.CreateUAV(PhysicalPagesToInitializeRDG);
 			FSelectPagesToInitializeCS::FPermutationDomain PermutationVector;
-			SetStatsArgsAndPermutation<FSelectPagesToInitializeCS>(StatsBufferUAV, PassParameters, PermutationVector);
+			SetStatsArgsAndPermutation<FSelectPagesToInitializeCS>(ShouldGenerateStats(), StatsBufferUAV, PassParameters, PermutationVector);
 
 			auto ComputeShader = GetGlobalShaderMap(Scene.GetFeatureLevel())->GetShader<FSelectPagesToInitializeCS>(PermutationVector);
 
@@ -2075,6 +2074,9 @@ class FVirtualSmLogStatsCS : public FVirtualShadowMapPageManagementShader
 	DECLARE_GLOBAL_SHADER(FVirtualSmLogStatsCS);
 	SHADER_USE_PARAMETER_STRUCT(FVirtualSmLogStatsCS, FVirtualShadowMapPageManagementShader)
 
+	class FGenerateStatsDim : SHADER_PERMUTATION_BOOL("VSM_GENERATE_STATS"); 
+	using FPermutationDomain = TShaderPermutationDomain<FGenerateStatsDim>;
+
 	BEGIN_SHADER_PARAMETER_STRUCT(FParameters, )
 		SHADER_PARAMETER_RDG_UNIFORM_BUFFER(FVirtualShadowMapUniformParameters, VirtualShadowMap)
 		SHADER_PARAMETER_STRUCT_INCLUDE(GPUMessage::FParameters, GPUMessageParams)
@@ -2082,6 +2084,7 @@ class FVirtualSmLogStatsCS : public FVirtualShadowMapPageManagementShader
 		SHADER_PARAMETER_RDG_BUFFER_SRV(StructuredBuffer<uint>, InStatsBuffer)
 		SHADER_PARAMETER_RDG_BUFFER_SRV(StructuredBuffer<FNaniteStats>, NaniteStats)
 		SHADER_PARAMETER(int, ShowStatsValue)
+		SHADER_PARAMETER(uint32, StatusMessageId)
 		SHADER_PARAMETER(uint32, StatsMessageId)
 	END_SHADER_PARAMETER_STRUCT()
 
@@ -2100,20 +2103,33 @@ void FVirtualShadowMapArray::LogStats(FRDGBuilder& GraphBuilder, const FViewInfo
 	check(IsEnabled());
 	LLM_SCOPE_BYTAG(Nanite);
 
-	if (StatsBufferRDG)
+	if (!StatsBufferRDG)
 	{
-		// Convenience, enable shader print automatically
-		ShaderPrint::SetEnabled(true);
+		return;
+	}
 
-		int ShowStatsValue = CVarShowStats.GetValueOnRenderThread();
+	FVirtualSmLogStatsCS::FParameters* PassParameters = GraphBuilder.AllocParameters<FVirtualSmLogStatsCS::FParameters>();
+	PassParameters->InStatsBuffer = GraphBuilder.CreateSRV(StatsBufferRDG);
+	PassParameters->VirtualShadowMap = GetUncachedUniformBuffer(GraphBuilder);
+	PassParameters->GPUMessageParams = GPUMessage::GetShaderParameters(GraphBuilder);
+	PassParameters->StatusMessageId = CacheManager->GetStatusFeedbackMessageId();
 
-		FVirtualSmLogStatsCS::FParameters* PassParameters = GraphBuilder.AllocParameters<FVirtualSmLogStatsCS::FParameters>();
+	// If ShouldGenerateStats() is false, the stats buffer will only have data needed for Status messages (e.g. for overflow tracking)
+	bool bGenerateStats = ShouldGenerateStats();
+
+	FVirtualSmLogStatsCS::FPermutationDomain PermutationVector;
+	PermutationVector.Set<FVirtualSmLogStatsCS::FGenerateStatsDim>(bGenerateStats);
+	
+	if (bGenerateStats)
+	{
+		// Convenience, enable shader print automatically
+		ShaderPrint::SetEnabled(true);
 
 		ShaderPrint::SetParameters(GraphBuilder, View.ShaderPrintData, PassParameters->ShaderPrintStruct);
-		PassParameters->InStatsBuffer = GraphBuilder.CreateSRV(StatsBufferRDG);
-		PassParameters->VirtualShadowMap = GetUncachedUniformBuffer(GraphBuilder);
+
+		int ShowStatsValue = CVarShowStats.GetValueOnRenderThread();
 		PassParameters->ShowStatsValue = ShowStatsValue;
-		
+
 #if !UE_BUILD_SHIPPING
 		PassParameters->StatsMessageId = CacheManager->GetStatsFeedbackMessageId();
 		bool bBindNaniteStatsBuffer = StatsNaniteBufferRDG != nullptr;
@@ -2121,6 +2137,7 @@ void FVirtualShadowMapArray::LogStats(FRDGBuilder& GraphBuilder, const FViewInfo
 		PassParameters->StatsMessageId = INDEX_NONE;
 		constexpr bool bBindNaniteStatsBuffer = false;
 #endif
+
 		if (bBindNaniteStatsBuffer)
 		{
 			PassParameters->NaniteStats = GraphBuilder.CreateSRV(StatsNaniteBufferRDG);
@@ -2129,18 +2146,17 @@ void FVirtualShadowMapArray::LogStats(FRDGBuilder& GraphBuilder, const FViewInfo
 		{
 			PassParameters->NaniteStats = GraphBuilder.CreateSRV(GSystemTextures.GetDefaultStructuredBuffer(GraphBuilder, sizeof(FNaniteStats)));
 		}
-		PassParameters->GPUMessageParams = GPUMessage::GetShaderParameters(GraphBuilder);
+	}
 	
-		auto ComputeShader = View.ShaderMap->GetShader<FVirtualSmLogStatsCS>();
+	auto ComputeShader = View.ShaderMap->GetShader<FVirtualSmLogStatsCS>(PermutationVector);
 
-		FComputeShaderUtils::AddPass(
-			GraphBuilder,
-			RDG_EVENT_NAME("VSM Log Stats"),
-			ComputeShader,
-			PassParameters,
-			FIntVector(1, 1, 1)
-		);
-	}
+	FComputeShaderUtils::AddPass(
+		GraphBuilder,
+		RDG_EVENT_NAME("VSM Log Stats And Status"),
+		ComputeShader,
+		PassParameters,
+		FIntVector(1, 1, 1)
+	);
 }
 
 
@@ -2546,20 +2562,24 @@ static FCullingResult AddCullingPasses(FRDGBuilder& GraphBuilder,
 
 		PassParameters->HZBShaderParameters = HZBShaderParameters;
 
-		bool bGenerateStats = VirtualShadowMapArray.StatsBufferUAV != nullptr;
+		FCullPerPageDrawCommandsCs::FPermutationDomain PermutationVector;
+
+		bool bGenerateStats = SetStatsArgsAndPermutation<FCullPerPageDrawCommandsCs>(
+			VirtualShadowMapArray.ShouldGenerateStats(), 
+			VirtualShadowMapArray.StatsBufferUAV, 
+			PassParameters, 
+			PermutationVector);
+
 		if (bGenerateStats)
 		{
-			PassParameters->OutStatsBuffer = VirtualShadowMapArray.StatsBufferUAV;
 #if !UE_BUILD_SHIPPING
 			PassParameters->NumPageAreaDiagnosticSlots = CVarNumPageAreaDiagSlots.GetValueOnRenderThread() < 0 ? FVirtualShadowMapArray::MaxPageAreaDiagnosticSlots : FMath::Min(FVirtualShadowMapArray::MaxPageAreaDiagnosticSlots, uint32(CVarNumPageAreaDiagSlots.GetValueOnRenderThread()));
 			PassParameters->LargeInstancePageAreaThreshold = CVarLargeInstancePageAreaThreshold.GetValueOnRenderThread() >= 0 ? CVarLargeInstancePageAreaThreshold.GetValueOnRenderThread() : (VirtualShadowMapArray.GetMaxPhysicalPages() / 8);
 #endif
 		}
 
-		FCullPerPageDrawCommandsCs::FPermutationDomain PermutationVector;
 		PermutationVector.Set< FCullPerPageDrawCommandsCs::FBatchedDim >(bUseBatchMode);
 		PermutationVector.Set< FCullPerPageDrawCommandsCs::FUseHzbDim >(HZBShaderParameters.HZBTextureArray != nullptr);
-		PermutationVector.Set< FCullPerPageDrawCommandsCs::FGenerateStatsDim >(bGenerateStats);
 
 		auto ComputeShader = ShaderMap->GetShader<FCullPerPageDrawCommandsCs>(PermutationVector);
 
@@ -2774,10 +2794,7 @@ Nanite::FPackedViewArray* FVirtualShadowMapArray::CreateVirtualShadowMapNaniteVi
 
 void FVirtualShadowMapArray::RenderVirtualShadowMapsNanite(FRDGBuilder& GraphBuilder, FSceneRenderer& SceneRenderer, bool bUpdateNaniteStreaming, const FNaniteVisibilityQuery* VisibilityQuery, Nanite::FPackedViewArray* VirtualShadowMapViews, FSceneInstanceCullingQuery* SceneInstanceCullingQuery)
 {
-	bool bCsvLogEnabled = false;
-#if CSV_PROFILER_STATS
-	bCsvLogEnabled = FCsvProfiler::Get()->IsCapturing_Renderthread() && FCsvProfiler::Get()->IsCategoryEnabled(CSV_CATEGORY_INDEX(VSM));
-#endif
+	bool bCsvLogEnabled = IsCsvLogEnabled();
 
 	TRACE_CPUPROFILER_EVENT_SCOPE(FVirtualShadowMapArray::RenderVirtualShadowMapsNanite);
 	RDG_EVENT_SCOPE(GraphBuilder, "RenderVirtualShadowMaps(Nanite)");
@@ -3355,7 +3372,7 @@ void FVirtualShadowMapArray::UpdateHZB(FRDGBuilder& GraphBuilder)
 		PassParameters->bFirstBuildThisFrame = !bHZBBuiltThisFrame;
 		PassParameters->bForceFullHZBUpdate = CVarShadowsVirtualForceFullHZBUpdate.GetValueOnRenderThread();
 		FSelectPagesForHZBAndUpdateDirtyFlagsCS::FPermutationDomain PermutationVector;
-		SetStatsArgsAndPermutation<FSelectPagesForHZBAndUpdateDirtyFlagsCS>(StatsBufferUAV, PassParameters, PermutationVector);
+		SetStatsArgsAndPermutation<FSelectPagesForHZBAndUpdateDirtyFlagsCS>(ShouldGenerateStats(), StatsBufferUAV, PassParameters, PermutationVector);
 		auto ComputeShader = GetGlobalShaderMap(Scene.GetFeatureLevel())->GetShader<FSelectPagesForHZBAndUpdateDirtyFlagsCS>(PermutationVector);
 
 		FComputeShaderUtils::AddPass(
diff --git a/Engine/Source/Runtime/Renderer/Private/VirtualShadowMaps/VirtualShadowMapArray.h b/Engine/Source/Runtime/Renderer/Private/VirtualShadowMaps/VirtualShadowMapArray.h
--- a/Engine/Source/Runtime/Renderer/Private/VirtualShadowMaps/VirtualShadowMapArray.h
+++ b/Engine/Source/Runtime/Renderer/Private/VirtualShadowMaps/VirtualShadowMapArray.h
@@ -359,6 +359,9 @@ public:
 
 	void RenderDebugInfo(FRDGBuilder& GraphBuilder, TArrayView<FViewInfo> Views);
 	
+	bool ShouldGenerateStats() const;
+	bool IsCsvLogEnabled() const;
+
 	void LogStats(FRDGBuilder& GraphBuilder, const FViewInfo& View);
 
 	// Get shader parameters necessary to sample virtual shadow maps
diff --git a/Engine/Source/Runtime/Renderer/Private/VirtualShadowMaps/VirtualShadowMapCacheManager.cpp b/Engine/Source/Runtime/Renderer/Private/VirtualShadowMaps/VirtualShadowMapCacheManager.cpp
--- a/Engine/Source/Runtime/Renderer/Private/VirtualShadowMaps/VirtualShadowMapCacheManager.cpp
+++ b/Engine/Source/Runtime/Renderer/Private/VirtualShadowMaps/VirtualShadowMapCacheManager.cpp
@@ -17,6 +17,7 @@
 #include "Shadows/ShadowScene.h"
 #include "ProfilingDebugging/CountersTrace.h"
 
+#define LOCTEXT_NAMESPACE "VirtualShadowMapCacheManager"
 CSV_DECLARE_CATEGORY_EXTERN(VSM);
 
 static int32 GVSMAccumulateStats = 0;
@@ -156,6 +157,7 @@ static const char *VirtualShadowMap_StatNames[] =
 	"NANITE_INSTANCES_MAIN",
 	"NANITE_INSTANCES_POST",
 	"WPO_CONSIDERED_PAGES",
+	"OVERFLOW_FLAGS",
 	"TMP_1",
 	"TMP_2",
 	"TMP_3",
@@ -590,6 +592,9 @@ ISceneExtensionUpdater* FVirtualShadowMapArrayCacheManager::CreateUpdater()
 FVirtualShadowMapArrayCacheManager::FVirtualShadowMapArrayCacheManager()
 	: ShadowInvalidatingInstancesImplementation(*this)
 {
+#if !UE_BUILD_SHIPPING
+	LastOverflowTimes.Init(-10.0f, VSM_STAT_OVERFLOW_FLAG_NUM);
+#endif
 }
 
 void FVirtualShadowMapArrayCacheManager::InitExtension(FScene& InScene)
@@ -599,76 +604,103 @@ void FVirtualShadowMapArrayCacheManager::InitExtension(FScene& InScene)
 	// Handle message with status sent back from GPU
 	StatusFeedbackSocket = GPUMessage::RegisterHandler(TEXT("Shadow.Virtual.StatusFeedback"), [this](GPUMessage::FReader Message)
 	{
-		// Goes negative on underflow
-		int32 LastFreePhysicalPages = Message.Read<int32>(0);
-		const float LastGlobalResolutionLodBias = FMath::AsFloat(Message.Read<uint32>(0U));
+		int32 MessageType = Message.Read<int32>();
+		if(MessageType == VSM_STATUS_MSG_PAGE_MANAGEMENT)
+		{
+			// Goes negative on underflow
+			int32 LastFreePhysicalPages = Message.Read<int32>(0);
+			const float LastGlobalResolutionLodBias = FMath::AsFloat(Message.Read<uint32>(0U));
 		
-		CSV_CUSTOM_STAT(VSM, FreePages, LastFreePhysicalPages, ECsvCustomStatOp::Set);
+			CSV_CUSTOM_STAT(VSM, FreePages, LastFreePhysicalPages, ECsvCustomStatOp::Set);
 
-		// Dynamic resolution
-		{
-			// Could be cvars if needed, but not clearly something that needs to be tweaked currently
-			// NOTE: Should react more quickly when reducing resolution than when increasing again
-			// TODO: Possibly something smarter/PID-like rather than simple exponential decay
-			const float ResolutionDownExpLerpFactor = 0.5f;
-			const float ResolutionUpExpLerpFactor = 0.1f;
-			const uint32 FramesBeforeResolutionUp = 10;
-
-			const float MaxPageAllocation = CVarVSMDynamicResolutionMaxPagePoolLoadFactor.GetValueOnRenderThread();
-			const float MaxLodBias = CVarVSMDynamicResolutionMaxLodBias.GetValueOnRenderThread();
-			
-			if (MaxPageAllocation > 0.0f)
+			// Dynamic resolution
 			{
-				const uint32 SceneFrameNumber = Scene->GetFrameNumberRenderThread();
+				// Could be cvars if needed, but not clearly something that needs to be tweaked currently
+				// NOTE: Should react more quickly when reducing resolution than when increasing again
+				// TODO: Possibly something smarter/PID-like rather than simple exponential decay
+				const float ResolutionDownExpLerpFactor = 0.5f;
+				const float ResolutionUpExpLerpFactor = 0.1f;
+				const uint32 FramesBeforeResolutionUp = 10;
+
+				const float MaxPageAllocation = CVarVSMDynamicResolutionMaxPagePoolLoadFactor.GetValueOnRenderThread();
+				const float MaxLodBias = CVarVSMDynamicResolutionMaxLodBias.GetValueOnRenderThread();
+			
+				if (MaxPageAllocation > 0.0f)
+				{
+					const uint32 SceneFrameNumber = Scene->GetFrameNumberRenderThread();
 
-				// Dynamically bias shadow resolution when we get too near the maximum pool capacity
-				// NB: In a perfect world each +1 of resolution bias will drop the allocation in half
-				float CurrentAllocation = 1.0f - (LastFreePhysicalPages / static_cast<float>(MaxPhysicalPages));
-				float AllocationRatio = CurrentAllocation / MaxPageAllocation;
-				float TargetLodBias = FMath::Max(0.0f, LastGlobalResolutionLodBias + FMath::Log2(AllocationRatio));
+					// Dynamically bias shadow resolution when we get too near the maximum pool capacity
+					// NB: In a perfect world each +1 of resolution bias will drop the allocation in half
+					float CurrentAllocation = 1.0f - (LastFreePhysicalPages / static_cast<float>(MaxPhysicalPages));
+					float AllocationRatio = CurrentAllocation / MaxPageAllocation;
+					float TargetLodBias = FMath::Max(0.0f, LastGlobalResolutionLodBias + FMath::Log2(AllocationRatio));
 
-				if (CurrentAllocation <= MaxPageAllocation &&
-					(SceneFrameNumber - LastFrameOverPageAllocationBudget) > FramesBeforeResolutionUp)
-				{
-					GlobalResolutionLodBias = FMath::Lerp(GlobalResolutionLodBias, TargetLodBias, ResolutionUpExpLerpFactor);
+					if (CurrentAllocation <= MaxPageAllocation &&
+						(SceneFrameNumber - LastFrameOverPageAllocationBudget) > FramesBeforeResolutionUp)
+					{
+						GlobalResolutionLodBias = FMath::Lerp(GlobalResolutionLodBias, TargetLodBias, ResolutionUpExpLerpFactor);
+					}
+					else if (CurrentAllocation > MaxPageAllocation)
+					{
+						LastFrameOverPageAllocationBudget = SceneFrameNumber;
+						GlobalResolutionLodBias = FMath::Lerp(GlobalResolutionLodBias, TargetLodBias, ResolutionDownExpLerpFactor);
+					}
 				}
-				else if (CurrentAllocation > MaxPageAllocation)
+
+				GlobalResolutionLodBias = FMath::Clamp(GlobalResolutionLodBias, 0.0f, MaxLodBias);
+			}
+
+#if !UE_BUILD_SHIPPING
+			if (LastFreePhysicalPages < 0)
+			{
+				uint32 PagePoolOverflowTypeIndex = (uint32)FMath::Log2((double)VSM_STAT_OVERFLOW_FLAG_PAGE_POOL);
+				LastOverflowTimes[PagePoolOverflowTypeIndex] = float(FGameTime::GetTimeSinceAppStart().GetRealTimeSeconds());
+				if ((LoggedOverflowFlags & VSM_STAT_OVERFLOW_FLAG_PAGE_POOL) == 0)
 				{
-					LastFrameOverPageAllocationBudget = SceneFrameNumber;
-					GlobalResolutionLodBias = FMath::Lerp(GlobalResolutionLodBias, TargetLodBias, ResolutionDownExpLerpFactor);
+					static const auto* CVarResolutionLodBiasLocalPtr = IConsoleManager::Get().FindTConsoleVariableDataFloat(TEXT("r.Shadow.Virtual.ResolutionLodBiasLocal"));
+					static const auto* CVarResolutionLodBiasDirectionalPtr = IConsoleManager::Get().FindTConsoleVariableDataFloat(TEXT("r.Shadow.Virtual.ResolutionLodBiasDirectional"));
+
+					UE_LOG(LogRenderer, Warning, TEXT("Virtual Shadow Map Page Pool overflow (%d page allocations were not served), this will produce visual artifacts (missing shadow), increase the page pool limit or reduce resolution bias to avoid.\n")
+						TEXT(" See r.Shadow.Virtual.MaxPhysicalPages (%d), r.Shadow.Virtual.ResolutionLodBiasLocal (%.2f), r.Shadow.Virtual.ResolutionLodBiasDirectional (%.2f), Global Resolution Lod Bias (%.2f)"),
+						-LastFreePhysicalPages,
+						MaxPhysicalPages,
+						CVarResolutionLodBiasLocalPtr->GetValueOnRenderThread(),
+						CVarResolutionLodBiasDirectionalPtr->GetValueOnRenderThread(),
+						GlobalResolutionLodBias);
+
+					LoggedOverflowFlags |= VSM_STAT_OVERFLOW_FLAG_PAGE_POOL;
 				}
 			}
-
-			GlobalResolutionLodBias = FMath::Clamp(GlobalResolutionLodBias, 0.0f, MaxLodBias);
+			else
+			{
+				LoggedOverflowFlags &= ~VSM_STAT_OVERFLOW_FLAG_PAGE_POOL;
+			}
+#endif
 		}
-
-		if (LastFreePhysicalPages < 0)
+		else if(MessageType == VSM_STATUS_MSG_OVERFLOW)
 		{
 #if !UE_BUILD_SHIPPING
-			if (!bLoggedPageOverflow)
+			uint32 OverflowFlags = Message.Read<int32>();
+			if (OverflowFlags)
 			{
-				static const auto* CVarResolutionLodBiasLocalPtr = IConsoleManager::Get().FindTConsoleVariableDataFloat(TEXT("r.Shadow.Virtual.ResolutionLodBiasLocal"));
-				static const auto* CVarResolutionLodBiasDirectionalPtr = IConsoleManager::Get().FindTConsoleVariableDataFloat(TEXT("r.Shadow.Virtual.ResolutionLodBiasDirectional"));
-
-				UE_LOG(LogRenderer, Warning, TEXT("Virtual Shadow Map Page Pool overflow (%d page allocations were not served), this will produce visual artifacts (missing shadow), increase the page pool limit or reduce resolution bias to avoid.\n")
-					TEXT(" See r.Shadow.Virtual.MaxPhysicalPages (%d), r.Shadow.Virtual.ResolutionLodBiasLocal (%.2f), r.Shadow.Virtual.ResolutionLodBiasDirectional (%.2f), Global Resolution Lod Bias (%.2f)"),
-					-LastFreePhysicalPages,
-					MaxPhysicalPages,
-					CVarResolutionLodBiasLocalPtr->GetValueOnRenderThread(),
-					CVarResolutionLodBiasDirectionalPtr->GetValueOnRenderThread(),
-					GlobalResolutionLodBias);
-
-				bLoggedPageOverflow = true;
+				float CurrentTime = float(FGameTime::GetTimeSinceAppStart().GetRealTimeSeconds());
+				for (uint32 OverflowTypeIndex = 0; OverflowTypeIndex < VSM_STAT_OVERFLOW_FLAG_NUM; OverflowTypeIndex++)
+				{
+					uint32 OverflowTypeFlag = 1 << OverflowTypeIndex;
+					if (OverflowFlags & OverflowTypeFlag)
+					{
+						LastOverflowTimes[OverflowTypeIndex] = CurrentTime;
+
+						if ((LoggedOverflowFlags & OverflowTypeFlag) == 0)
+						{
+							UE_LOG(LogRenderer, Warning, TEXT("%s"), *(GetOverflowMessage(OverflowTypeIndex).ToString()));
+							LoggedOverflowFlags |= OverflowTypeFlag;
+						}
+					}
+				}
 			}
-			LastOverflowTime = float(FGameTime::GetTimeSinceAppStart().GetRealTimeSeconds());
 #endif
 		}
-#if !UE_BUILD_SHIPPING
-		else
-		{
-			bLoggedPageOverflow = false;
-		}
-#endif
 	});
 
 #if !UE_BUILD_SHIPPING
@@ -760,10 +792,15 @@ void FVirtualShadowMapArrayCacheManager::InitExtension(FScene& InScene)
 	{
 		float RealTimeSeconds = float(FGameTime::GetTimeSinceAppStart().GetRealTimeSeconds());
 
-		// Show for ~5s after last overflow
-		if (LastOverflowTime >= 0.0f && RealTimeSeconds - LastOverflowTime < 5.0f)
+		for (uint32 OverflowTypeIndex = 0; OverflowTypeIndex < VSM_STAT_OVERFLOW_FLAG_NUM; OverflowTypeIndex++)
 		{
-			OutMessages.Add(FCoreDelegates::EOnScreenMessageSeverity::Warning, FText::FromString(FString::Printf(TEXT("Virtual Shadow Map Page Pool overflow detected (%0.0f seconds ago)"), RealTimeSeconds - LastOverflowTime)));
+			// Show for ~10s after last overflow
+			float LastOverflowTime = LastOverflowTimes[OverflowTypeIndex];
+			if (LastOverflowTime >= 0.0f && RealTimeSeconds - LastOverflowTime < 10.0f)
+			{
+				FText OverflowMessage = GetOverflowMessage(OverflowTypeIndex);
+				OutMessages.Add(FCoreDelegates::EOnScreenMessageSeverity::Warning, FText::FromString(FString::Printf(TEXT("%s (%0.0f seconds ago)"), *(OverflowMessage.ToString()), RealTimeSeconds - LastOverflowTime)));
+			}
 		}
 
 		for (const auto& Item : LargePageAreaItems)
@@ -796,6 +833,20 @@ FVirtualShadowMapArrayCacheManager::~FVirtualShadowMapArrayCacheManager()
 #endif
 }
 
+#if !UE_BUILD_SHIPPING
+FText FVirtualShadowMapArrayCacheManager::GetOverflowMessage(uint32 OverflowTypeIndex) const
+{
+	uint32 OverflowTypeFlag = 1 << OverflowTypeIndex;
+	switch (OverflowTypeFlag)
+	{
+		case VSM_STAT_OVERFLOW_FLAG_MARKING_JOB_QUEUE: return LOCTEXT("VSM_MarkingJobQueueOverflow", "[VSM] Non-Nanite Marking Job Queue overflow. Performance may be affected. This occurs when many non-nanite meshes cover a large area of the shadow map.");
+		case VSM_STAT_OVERFLOW_FLAG_OPP_MAX_LIGHTS: return LOCTEXT("VSM_OPPMaxLightsOverflow", "[VSM] One Pass Projection max lights overflow. If you see shadow artifacts, decrease the amount of local lights per pixel, or increase r.Shadow.Virtual.OnePassProjection.MaxLightsPerPixel.");
+		case VSM_STAT_OVERFLOW_FLAG_PAGE_POOL: return LOCTEXT("VSM_PagePoolOverflow", "[VSM] Page Pool overflow detected, this will produce visual artifacts (missing shadow). Increase the page pool limit or reduce resolution bias to avoid.");
+		case VSM_STAT_OVERFLOW_FLAG_VISIBLE_INSTANCES: return LOCTEXT("VSM_VisibleInstancesOverflow", "[VSM] Non-Nanite visible instances buffer overflow detected, this will produce visual artifacts (missing shadow).");
+		default: return LOCTEXT("VSM_UnknownOverflow", "[VSM] Unknown overflow");
+	}
+}
+#endif
 
 void FVirtualShadowMapArrayCacheManager::SetPhysicalPoolSize(FRDGBuilder& GraphBuilder, FIntPoint RequestedSize, int RequestedArraySize, uint32 RequestedMaxPhysicalPages)
 {
@@ -1704,3 +1755,5 @@ void FVirtualShadowMapInvalidationSceneUpdater::PostGPUSceneUpdate(FRDGBuilder&
 	}
 	PostUpdateChangeSet = nullptr;
 }
+
+#undef LOCTEXT_NAMESPACE
\ No newline at end of file
diff --git a/Engine/Source/Runtime/Renderer/Private/VirtualShadowMaps/VirtualShadowMapCacheManager.h b/Engine/Source/Runtime/Renderer/Private/VirtualShadowMaps/VirtualShadowMapCacheManager.h
--- a/Engine/Source/Runtime/Renderer/Private/VirtualShadowMaps/VirtualShadowMapCacheManager.h
+++ b/Engine/Source/Runtime/Renderer/Private/VirtualShadowMaps/VirtualShadowMapCacheManager.h
@@ -449,8 +449,10 @@ private:
 	// Debug stuff
 #if !UE_BUILD_SHIPPING
 	FDelegateHandle ScreenMessageDelegate;
-	float LastOverflowTime = -1.0f;
-	bool bLoggedPageOverflow = false;
+	uint32 LoggedOverflowFlags = 0;
+	TArray<float, TInlineAllocator<VSM_STAT_OVERFLOW_FLAG_NUM>> LastOverflowTimes;
+
+	FText GetOverflowMessage(uint32 OverflowTypeIndex) const;
 	
 	// Socket for optional stats that are only sent back if enabled
 	GPUMessage::FSocket StatsFeedbackSocket;
diff --git a/Engine/Source/Runtime/Renderer/Private/VirtualShadowMaps/VirtualShadowMapProjection.cpp b/Engine/Source/Runtime/Renderer/Private/VirtualShadowMaps/VirtualShadowMapProjection.cpp
--- a/Engine/Source/Runtime/Renderer/Private/VirtualShadowMaps/VirtualShadowMapProjection.cpp
+++ b/Engine/Source/Runtime/Renderer/Private/VirtualShadowMaps/VirtualShadowMapProjection.cpp
@@ -143,6 +143,7 @@ class FVirtualShadowMapProjectionCS : public FGlobalShader
 		// Optional tile list
 		SHADER_PARAMETER_RDG_BUFFER_SRV(Buffer, TileListData)
 		RDG_BUFFER_ACCESS(IndirectDispatchArgs, ERHIAccess::IndirectArgs)
+		SHADER_PARAMETER_RDG_BUFFER_UAV(RWStructuredBuffer<uint>, OutStatsBuffer)
 	END_SHADER_PARAMETER_STRUCT()
 
 	static void ModifyCompilationEnvironment( const FGlobalShaderPermutationParameters& Parameters, FShaderCompilerEnvironment& OutEnvironment )
@@ -269,6 +270,8 @@ static void RenderVirtualShadowMapProjectionCommon(
 	}
 #endif
 
+	PassParameters->OutStatsBuffer = VirtualShadowMapArray.StatsBufferUAV;
+
 	// If the requested samples per ray matches one of our static permutations, pick that one
 	// Otherwise use the dynamic samples per ray permutation (-1).
 	int SamplesPerRay = (bDirectionalLight ? VirtualShadowMapArray.UniformParameters.SMRTSamplesPerRayDirectional : VirtualShadowMapArray.UniformParameters.SMRTSamplesPerRayLocal);
diff --git a/Engine/Source/Runtime/Renderer/Private/VirtualShadowMaps/VirtualShadowMapShaders.h b/Engine/Source/Runtime/Renderer/Private/VirtualShadowMaps/VirtualShadowMapShaders.h
--- a/Engine/Source/Runtime/Renderer/Private/VirtualShadowMaps/VirtualShadowMapShaders.h
+++ b/Engine/Source/Runtime/Renderer/Private/VirtualShadowMaps/VirtualShadowMapShaders.h
@@ -47,3 +47,11 @@ public:
 		OutEnvironment.SetDefine(TEXT("VF_SUPPORTS_PRIMITIVE_SCENE_DATA"), 1);
 	}
 };
+
+template <typename ShaderType>
+static bool SetStatsArgsAndPermutation(bool bGenerateStats, FRDGBufferUAVRef StatsBufferUAV, typename ShaderType::FParameters *OutPassParameters, typename ShaderType::FPermutationDomain& OutPermutationVector)
+{
+	OutPassParameters->OutStatsBuffer = StatsBufferUAV;
+	OutPermutationVector.template Set<typename ShaderType::FGenerateStatsDim>(bGenerateStats);
+	return bGenerateStats;
+}
