Fix for MSAA on outlines causing a mismatch in FSelectionOutlinePS

diff --git a/Engine/Source/Runtime/Renderer/Private/PostProcess/PostProcessSelectionOutline.cpp b/Engine/Source/Runtime/Renderer/Private/PostProcess/PostProcessSelectionOutline.cpp
--- a/Engine/Source/Runtime/Renderer/Private/PostProcess/PostProcessSelectionOutline.cpp
+++ b/Engine/Source/Runtime/Renderer/Private/PostProcess/PostProcessSelectionOutline.cpp
@@ -111,7 +111,7 @@ FScreenPassTexture AddSelectionOutlinePass(
 	// Patch uniform buffers with updated state for rendering the outline mesh draw commands.
 	const bool bIsInstancedStereoView = View.bIsInstancedStereoEnabled && IStereoRendering::IsStereoEyePass(View.StereoPass);
 	const FViewInfo* EditorView = CreateCompositePrimitiveView(View, bIsInstancedStereoView ? View.ViewRect : Inputs.SceneColor.ViewRect, NumSamples);
-	FRDGTextureRef OverlayColorTexture = nullptr;
+	FRDGTextureMSAA OverlayColorTexture;
 
 	// Generate custom depth / stencil for outline shapes.
 	{
@@ -128,7 +128,7 @@ FScreenPassTexture AddSelectionOutlinePass(
 			OverlayColorDesc.Flags = TexCreate_RenderTargetable | TexCreate_ShaderResource;
 			OverlayColorDesc.NumSamples = NumSamples;
 
-			OverlayColorTexture = GraphBuilder.CreateTexture(OverlayColorDesc, TEXT("Editor.ColorOverlay"));
+			OverlayColorTexture = CreateTextureMSAA(GraphBuilder, OverlayColorDesc, TEXT("Editor.ColorOverlayMSAA"), TEXT("Editor.ColorOverlay"));
 		}
 		
 		if (View.ShouldRenderView() || !DepthStencilTexture)
@@ -159,7 +159,7 @@ FScreenPassTexture AddSelectionOutlinePass(
 
 				PassParameters->View = EditorView->GetShaderParameters();
 				PassParameters->SceneTextures = Inputs.SceneTextures;
-				PassParameters->RenderTargets[0] = FRenderTargetBinding(OverlayColorTexture, ERenderTargetLoadAction::EClear);
+				PassParameters->RenderTargets[0] = FRenderTargetBinding(OverlayColorTexture.Target, OverlayColorTexture.Resolve, ERenderTargetLoadAction::EClear);
 				PassParameters->RenderTargets.DepthStencil = FDepthStencilBinding(
 					DepthStencilTexture,
 					ERenderTargetLoadAction::EClear,
@@ -203,7 +203,7 @@ FScreenPassTexture AddSelectionOutlinePass(
 		if (bNaniteEnabled)
 		{
 			// Update editor view to true target view rect
-			Nanite::DrawEditorSelection(GraphBuilder, DepthStencilTexture, OverlayColorTexture, *Scene, View, *EditorView, SceneUniformBuffer, NaniteRasterResults);
+			Nanite::DrawEditorSelection(GraphBuilder, DepthStencilTexture, OverlayColorTexture.Target, *Scene, View, *EditorView, SceneUniformBuffer, NaniteRasterResults);
 		}
 
 		// Render HairStrands outlines
@@ -299,7 +299,7 @@ FScreenPassTexture AddSelectionOutlinePass(
 		PassParameters->EditorPrimitivesDepth = DepthStencilTexture;
 		PassParameters->EditorPrimitivesStencil = GraphBuilder.CreateSRV(FRDGTextureSRVDesc::CreateWithPixelFormat(DepthStencilTexture, PF_X24_G8));
 		
-		PassParameters->OverlayLookupTexture = OverlayColorTexture;
+		PassParameters->OverlayLookupTexture = OverlayColorTexture.Resolve;
 		PassParameters->OverlayLookupSampler = PointClampSampler;
 		
 		PassParameters->OutlineColors[0] = View.SelectionOutlineColor;
