Fix incorrect task tag in FRayTracingScene::Create causing ensure failures and deadlocks when running with -onethread.

#rnx

diff --git a/Engine/Source/Runtime/Renderer/Private/RayTracing/RayTracingScene.cpp b/Engine/Source/Runtime/Renderer/Private/RayTracing/RayTracingScene.cpp
--- a/Engine/Source/Runtime/Renderer/Private/RayTracing/RayTracingScene.cpp
+++ b/Engine/Source/Runtime/Renderer/Private/RayTracing/RayTracingScene.cpp
@@ -273,7 +273,7 @@ void FRayTracingScene::Create(FRDGBuilder& GraphBuilder, const FViewInfo& View,
 				TransformUploadBuffer = Layer.TransformUploadBuffer,
 				PreViewTranslation = this->PreViewTranslation](FRHICommandList& RHICmdList)
 				{
-					FTaskTagScope TaskTagScope(ETaskTag::EParallelRenderingThread);
+					FOptionalTaskTagScope TaskTagScope(ETaskTag::EParallelRenderingThread);
 
 					const uint32 InstanceUploadBytes = NumNativeInstances * sizeof(FRayTracingInstanceDescriptorInput);
 					const uint32 TransformUploadBytes = NumNativeCPUInstances * 3 * sizeof(FVector4f);
