Fix race condition when loading LTC textures.

diff --git a/Engine/Source/Runtime/Renderer/Private/SystemTextures.cpp b/Engine/Source/Runtime/Renderer/Private/SystemTextures.cpp
--- a/Engine/Source/Runtime/Renderer/Private/SystemTextures.cpp
+++ b/Engine/Source/Runtime/Renderer/Private/SystemTextures.cpp
@@ -238,7 +238,7 @@ static TRefCountPtr<IPooledRenderTarget> CreateTexture(FRHICommandListImmediate&
 
 static bool IsTextureDataValid(UTexture2D* In)
 {
-	return In && In->GetPlatformData();
+	return In && In->GetPlatformData() && In->GetCPUCopy();
 }
 
 void FSystemTextures::InitializeCommonTextures(FRHICommandListImmediate& RHICmdList)
@@ -1004,24 +1004,24 @@ void FSystemTextures::InitializeEngineDependentTextures(FRHICommandListImmediate
 		// LTC Textures	(used by Rect Lights and/or BSDF evaluation)
 
 		// GGX - LTC matrix coefficients (4-coefficients)
-		if (IsTextureDataValid(GEngine->GGXLTCMatTexture))
+		if (GGXLTCMat == nullptr && IsTextureDataValid(GEngine->GGXLTCMatTexture))
 		{
 			GGXLTCMat = CreateTexture(RHICmdList, GEngine->GGXLTCMatTexture, PF_FloatRGBA, TEXT("GGX.LTCMat"));
 		}
 		// GGX - Split-Sum Amplitude coefficients (2-components)
-		if (IsTextureDataValid(GEngine->GGXLTCAmpTexture))
+		if (GGXLTCAmp == nullptr && IsTextureDataValid(GEngine->GGXLTCAmpTexture))
 		{
 			GGXLTCAmp = CreateTexture(RHICmdList, GEngine->GGXLTCAmpTexture, PF_G16R16F, TEXT("GGX.LTCAmp"));
 		}
 		// Sheen - Matrix & directional albedo (3-components)
-		if (IsTextureDataValid(GEngine->SheenLTCTexture))
+		if (SheenLTC == nullptr && IsTextureDataValid(GEngine->SheenLTCTexture))
 		{
 			SheenLTC  = CreateTexture(RHICmdList, GEngine->SheenLTCTexture,  PF_FloatRGBA, TEXT("Sheen.LTC"));
 		}
 	}
 
-	// Initialize textures only once.
-	bEngineDependentTexturesInitialized = true;
+	// Initialize textures only once (SheenLTC is only used with Substrate)
+	bEngineDependentTexturesInitialized = GGXLTCMat && GGXLTCAmp && (!Substrate::IsSubstrateEnabled() || SheenLTC);
 }
 
 void FSystemTextures::ReleaseRHI()
