Minor refactors of mesh pass bindings / materials / shaders in preparation for Slate shader binding refactor.

diff --git a/Engine/Source/Runtime/Engine/Private/Materials/MaterialRenderProxy.cpp b/Engine/Source/Runtime/Engine/Private/Materials/MaterialRenderProxy.cpp
--- a/Engine/Source/Runtime/Engine/Private/Materials/MaterialRenderProxy.cpp
+++ b/Engine/Source/Runtime/Engine/Private/Materials/MaterialRenderProxy.cpp
@@ -578,7 +578,7 @@ void FMaterialRenderProxy::InvalidateUniformExpressionCache(bool bRecreateUnifor
 	}
 }
 
-void FMaterialRenderProxy::UpdateUniformExpressionCacheIfNeeded(FRHICommandListBase& RHICmdList, ERHIFeatureLevel::Type InFeatureLevel) const
+const FMaterial* FMaterialRenderProxy::UpdateUniformExpressionCacheIfNeeded(FRHICommandListBase& RHICmdList, ERHIFeatureLevel::Type InFeatureLevel) const
 {
 	// Don't cache uniform expressions if an entirely different FMaterialRenderProxy is going to be used for rendering
 	const FMaterial* Material = GetMaterialNoFallback(InFeatureLevel);
@@ -591,6 +591,8 @@ void FMaterialRenderProxy::UpdateUniformExpressionCacheIfNeeded(FRHICommandListB
 		MaterialRenderContext.bShowSelection = GIsEditor;
 		EvaluateUniformExpressions(RHICmdList, UniformExpressionCache[InFeatureLevel], MaterialRenderContext, nullptr);
 	}
+
+	return Material;
 }
 
 void FMaterialRenderProxy::UpdateUniformExpressionCacheIfNeeded(ERHIFeatureLevel::Type InFeatureLevel) const
diff --git a/Engine/Source/Runtime/Engine/Public/Materials/MaterialRenderProxy.h b/Engine/Source/Runtime/Engine/Public/Materials/MaterialRenderProxy.h
--- a/Engine/Source/Runtime/Engine/Public/Materials/MaterialRenderProxy.h
+++ b/Engine/Source/Runtime/Engine/Public/Materials/MaterialRenderProxy.h
@@ -144,7 +144,7 @@ public:
 	ENGINE_API void InvalidateUniformExpressionCache(bool bRecreateUniformBuffer);
 
 	ENGINE_API void UpdateUniformExpressionCacheIfNeeded(ERHIFeatureLevel::Type InFeatureLevel) const;
-	ENGINE_API void UpdateUniformExpressionCacheIfNeeded(FRHICommandListBase& RHICmdList, ERHIFeatureLevel::Type InFeatureLevel) const;
+	ENGINE_API const FMaterial* UpdateUniformExpressionCacheIfNeeded(FRHICommandListBase& RHICmdList, ERHIFeatureLevel::Type InFeatureLevel) const;
 
 	/** Returns the FMaterial, without using a fallback if the FMaterial doesn't have a valid shader map. Can return NULL. */
 	virtual const FMaterial* GetMaterialNoFallback(ERHIFeatureLevel::Type InFeatureLevel) const = 0;
diff --git a/Engine/Source/Runtime/RenderCore/Public/Shader.h b/Engine/Source/Runtime/RenderCore/Public/Shader.h
--- a/Engine/Source/Runtime/RenderCore/Public/Shader.h
+++ b/Engine/Source/Runtime/RenderCore/Public/Shader.h
@@ -1128,6 +1128,23 @@ public:
 		return RHIShader;
 	}
 
+	inline FRHIGraphicsShader* GetGraphicsShader(bool bRequired = true) const
+	{
+		FRHIGraphicsShader* RHIShader = nullptr;
+		if(ShaderContent)
+		{
+			checkSlow(IsValidGraphicsFrequency(ShaderContent->GetFrequency()));
+			RHIShader = static_cast<FRHIGraphicsShader*>(GetResourceChecked().GetShader(ShaderContent->GetResourceIndex(), bRequired));
+			if (RHIShader == nullptr)
+			{
+				UE_LOG(LogShaders, Log, TEXT("Failed to create shader for type %s with resource index %d."), GetType()->GetName(), ShaderContent->GetResourceIndex());
+				return nullptr;
+			}
+			checkSlow(IsValidGraphicsFrequency(RHIShader->GetFrequency()));
+		}
+		return RHIShader;
+	}
+
 	/** @return the shader's vertex shader */
 	inline FRHIVertexShader* GetVertexShader() const
 	{
diff --git a/Engine/Source/Runtime/Renderer/Private/MeshPassProcessor.cpp b/Engine/Source/Runtime/Renderer/Private/MeshPassProcessor.cpp
--- a/Engine/Source/Runtime/Renderer/Private/MeshPassProcessor.cpp
+++ b/Engine/Source/Runtime/Renderer/Private/MeshPassProcessor.cpp
@@ -77,48 +77,6 @@ static FAutoConsoleVariableRef CVarNaniteIsolateInvalidCoarseMesh(
 
 #endif
 
-class FReadOnlyMeshDrawSingleShaderBindings : public FMeshDrawShaderBindingsLayout
-{
-public:
-	FReadOnlyMeshDrawSingleShaderBindings(const FMeshDrawShaderBindingsLayout& InLayout, const uint8* InData) :
-		FMeshDrawShaderBindingsLayout(InLayout)
-	{
-		Data = InData;
-	}
-
-	inline FRHIUniformBuffer*const* GetUniformBufferStart() const
-	{
-		return (FRHIUniformBuffer**)(Data + GetUniformBufferOffset());
-	}
-
-	inline FRHISamplerState** GetSamplerStart() const
-	{
-		const uint8* SamplerDataStart = Data + GetSamplerOffset();
-		return (FRHISamplerState**)SamplerDataStart;
-	}
-
-	inline FRHIResource** GetSRVStart() const
-	{
-		const uint8* SRVDataStart = Data + GetSRVOffset();
-		return (FRHIResource**)SRVDataStart;
-	}
-
-	inline const uint8* GetSRVTypeStart() const
-	{
-		const uint8* SRVTypeDataStart = Data + GetSRVTypeOffset();
-		return SRVTypeDataStart;
-	}
-
-	inline const uint8* GetLooseDataStart() const
-	{
-		const uint8* LooseDataStart = Data + GetLooseDataOffset();
-		return LooseDataStart;
-	}
-
-private:
-	const uint8* Data;
-};
-
 inline void SetTextureParameter(FRHIBatchedShaderParameters& BatchedParameters, const FShaderResourceParameterInfo& Parameter, FRHITexture* TextureRHI)
 {
 #if PLATFORM_SUPPORTS_BINDLESS_RENDERING
@@ -181,7 +139,7 @@ inline void SetLooseParameters(FRHIBatchedShaderParameters& BatchedParameters, c
 	}
 }
 
-void FMeshDrawShaderBindings::SetShaderBindings(
+void FReadOnlyMeshDrawSingleShaderBindings::SetShaderBindings(
 	FRHIBatchedShaderParameters& BatchedParameters,
 	const FReadOnlyMeshDrawSingleShaderBindings& RESTRICT SingleShaderBindings,
 	FShaderBindingState& RESTRICT ShaderBindingState)
@@ -243,7 +201,7 @@ void FMeshDrawShaderBindings::SetShaderBindings(
 	SetLooseParameters(BatchedParameters, SingleShaderBindings.ParameterMapInfo, SingleShaderBindings.GetLooseDataStart());
 }
 
-void FMeshDrawShaderBindings::SetShaderBindings(
+void FReadOnlyMeshDrawSingleShaderBindings::SetShaderBindings(
 	FRHIBatchedShaderParameters& BatchedParameters,
 	const FReadOnlyMeshDrawSingleShaderBindings& RESTRICT SingleShaderBindings)
 {
@@ -1025,7 +983,7 @@ void FGraphicsMinimalPipelineStateInitializer::ComputeStatePrecachePSOHash()
 }
 
 void FMeshDrawCommand::SetDrawParametersAndFinalize(
-	const FMeshBatch& MeshBatch, 
+	const FMeshBatch& MeshBatch,
 	int32 BatchElementIndex,
 	FGraphicsMinimalPipelineStateId PipelineId,
 	const FMeshProcessorShaders* ShadersForDebugging)
@@ -1087,19 +1045,19 @@ void FMeshDrawShaderBindings::SetOnCommandList(FRHICommandList& RHICmdList, cons
 		if (Frequency == SF_Vertex)
 		{
 			FRHIBatchedShaderParameters& BatchedParameters = RHICmdList.GetScratchShaderParameters();
-			SetShaderBindings(BatchedParameters, SingleShaderBindings, ShaderBindingState);
+			FReadOnlyMeshDrawSingleShaderBindings::SetShaderBindings(BatchedParameters, SingleShaderBindings, ShaderBindingState);
 			RHICmdList.SetBatchedShaderParameters(Shaders.VertexShaderRHI, BatchedParameters);
 		} 
 		else if (Frequency == SF_Pixel)
 		{
 			FRHIBatchedShaderParameters& BatchedParameters = RHICmdList.GetScratchShaderParameters();
-			SetShaderBindings(BatchedParameters, SingleShaderBindings, ShaderBindingState);
+			FReadOnlyMeshDrawSingleShaderBindings::SetShaderBindings(BatchedParameters, SingleShaderBindings, ShaderBindingState);
 			RHICmdList.SetBatchedShaderParameters(Shaders.PixelShaderRHI, BatchedParameters);
 		}
 		else if (Frequency == SF_Geometry)
 		{
 			FRHIBatchedShaderParameters& BatchedParameters = RHICmdList.GetScratchShaderParameters();
-			SetShaderBindings(BatchedParameters, SingleShaderBindings, ShaderBindingState);
+			FReadOnlyMeshDrawSingleShaderBindings::SetShaderBindings(BatchedParameters, SingleShaderBindings, ShaderBindingState);
 			RHICmdList.SetBatchedShaderParameters(Shaders.GetGeometryShader(), BatchedParameters);
 		}
 		else
@@ -1118,11 +1076,11 @@ void FMeshDrawShaderBindings::SetParameters(FRHIBatchedShaderParameters& Batched
 
 	if (StateCacheShaderBindings != nullptr)
 	{
-		SetShaderBindings(BatchedParameters, SingleShaderBindings, *StateCacheShaderBindings);
+		FReadOnlyMeshDrawSingleShaderBindings::SetShaderBindings(BatchedParameters, SingleShaderBindings, *StateCacheShaderBindings);
 	}
 	else
 	{
-		SetShaderBindings(BatchedParameters, SingleShaderBindings);
+		FReadOnlyMeshDrawSingleShaderBindings::SetShaderBindings(BatchedParameters, SingleShaderBindings);
 	}
 }
 
diff --git a/Engine/Source/Runtime/Renderer/Private/ShaderBaseClasses.cpp b/Engine/Source/Runtime/Renderer/Private/ShaderBaseClasses.cpp
--- a/Engine/Source/Runtime/Renderer/Private/ShaderBaseClasses.cpp
+++ b/Engine/Source/Runtime/Renderer/Private/ShaderBaseClasses.cpp
@@ -329,7 +329,7 @@ void FMaterialShader::SetParameters(
 }
 
 void FMaterialShader::GetShaderBindings(
-	const FScene* Scene,
+	const FSceneInterface* Scene,
 	const ERHIFeatureLevel::Type FeatureLevel,
 	const FMaterialRenderProxy& MaterialRenderProxy,
 	const FMaterial& Material,
diff --git a/Engine/Source/Runtime/Renderer/Public/MaterialShader.h b/Engine/Source/Runtime/Renderer/Public/MaterialShader.h
--- a/Engine/Source/Runtime/Renderer/Public/MaterialShader.h
+++ b/Engine/Source/Runtime/Renderer/Public/MaterialShader.h
@@ -119,7 +119,7 @@ public:
 	}
 
 	RENDERER_API void GetShaderBindings(
-		const FScene* Scene,
+		const FSceneInterface* Scene,
 		ERHIFeatureLevel::Type FeatureLevel,
 		const FMaterialRenderProxy& MaterialRenderProxy,
 		const FMaterial& Material,
diff --git a/Engine/Source/Runtime/Renderer/Public/MeshDrawShaderBindings.h b/Engine/Source/Runtime/Renderer/Public/MeshDrawShaderBindings.h
--- a/Engine/Source/Runtime/Renderer/Public/MeshDrawShaderBindings.h
+++ b/Engine/Source/Runtime/Renderer/Public/MeshDrawShaderBindings.h
@@ -102,6 +102,7 @@ protected:
 
 class FMeshDrawSingleShaderBindings : public FMeshDrawShaderBindingsLayout
 {
+	friend class FReadOnlyMeshDrawSingleShaderBindings;
 public:
 	FMeshDrawSingleShaderBindings(const FMeshDrawShaderBindingsLayout& InLayout, uint8* InData) :
 		FMeshDrawShaderBindingsLayout(InLayout)
@@ -155,26 +156,36 @@ public:
 		}
 	}
 
-	void AddTexture(
-		FShaderResourceParameter TextureParameter,
-		FShaderResourceParameter SamplerParameter,
-		FRHISamplerState* SamplerStateRHI,
-		FRHITexture* TextureRHI)
+	void Add(FShaderResourceParameter SamplerParameter, FRHISamplerState* SamplerStateRHI)
 	{
-		checkfSlow(TextureParameter.IsInitialized(), TEXT("Parameter was not serialized"));
 		checkfSlow(SamplerParameter.IsInitialized(), TEXT("Parameter was not serialized"));
 
+		if (SamplerParameter.IsBound())
+		{
+			checkf(SamplerStateRHI, TEXT("Attempted to set null Sampler on slot %u"), SamplerParameter.GetBaseIndex());
+			WriteBindingSampler(SamplerStateRHI, SamplerParameter.GetBaseIndex());
+		}
+	}
+
+	void Add(FShaderResourceParameter TextureParameter, FRHITexture* TextureRHI)
+	{
+		checkfSlow(TextureParameter.IsInitialized(), TEXT("Parameter was not serialized"));
+
 		if (TextureParameter.IsBound())
 		{
 			checkf(TextureRHI, TEXT("Attempted to set null Texture on slot %u"), TextureParameter.GetBaseIndex());
 			WriteBindingTexture(TextureRHI, TextureParameter.GetBaseIndex());
 		}
+	}
 
-		if (SamplerParameter.IsBound())
-		{
-			checkf(SamplerStateRHI, TEXT("Attempted to set null Sampler on slot %u"), SamplerParameter.GetBaseIndex());
-			WriteBindingSampler(SamplerStateRHI, SamplerParameter.GetBaseIndex());
-		}
+	void AddTexture(
+		FShaderResourceParameter TextureParameter,
+		FShaderResourceParameter SamplerParameter,
+		FRHISamplerState* SamplerStateRHI,
+		FRHITexture* TextureRHI)
+	{
+		Add(TextureParameter, TextureRHI);
+		Add(SamplerParameter, SamplerStateRHI);
 	}
 
 	template<class ParameterType>
@@ -346,6 +357,4 @@ private:
 
 		checkfSlow(FoundIndex >= 0, TEXT("Attempted to set Texture at BaseIndex %u which was never in the shader's parameter map."), BaseIndex);
 	}
-
-	friend class FMeshDrawShaderBindings;
-};
+};
\ No newline at end of file
diff --git a/Engine/Source/Runtime/Renderer/Public/MeshPassProcessor.h b/Engine/Source/Runtime/Renderer/Public/MeshPassProcessor.h
--- a/Engine/Source/Runtime/Renderer/Public/MeshPassProcessor.h
+++ b/Engine/Source/Runtime/Renderer/Public/MeshPassProcessor.h
@@ -823,17 +823,17 @@ public:
 };
 
 
-/** 
+/**
  * Encapsulates shader bindings for a single FMeshDrawCommand.
  */
 class FMeshDrawShaderBindings
 {
 public:
-
-	FMeshDrawShaderBindings() 
+	FMeshDrawShaderBindings()
 	{
 		static_assert(sizeof(ShaderFrequencyBits) * 8 > SF_NumFrequencies, "Please increase ShaderFrequencyBits size");
 	}
+
 	FMeshDrawShaderBindings(FMeshDrawShaderBindings&& Other)
 	{
 		if (!UsesInlineStorage())
@@ -848,17 +848,18 @@ public:
 			Data = MoveTemp(Other.Data);
 		}
 		else
-		{		
+		{
 			Data.SetHeapData(Other.Data.GetHeapData());
 			Other.Data.SetHeapData(nullptr);
 		}
-		Other.Size = 0;	
+		Other.Size = 0;
 	}
 
 	FMeshDrawShaderBindings(const FMeshDrawShaderBindings& Other)
 	{
 		CopyFrom(Other);
 	}
+
 	RENDERER_API ~FMeshDrawShaderBindings();
 
 	FMeshDrawShaderBindings& operator=(const FMeshDrawShaderBindings& Other)
@@ -1044,15 +1045,63 @@ private:
 	RENDERER_API void CopyFrom(const FMeshDrawShaderBindings& Other);
 
 	RENDERER_API void Release();
+};
 
-	static void SetShaderBindings(
+class FReadOnlyMeshDrawSingleShaderBindings : public FMeshDrawShaderBindingsLayout
+{
+public:
+	static RENDERER_API void SetShaderBindings(
 		FRHIBatchedShaderParameters& BatchedParameters,
 		const class FReadOnlyMeshDrawSingleShaderBindings& RESTRICT SingleShaderBindings,
 		FShaderBindingState& RESTRICT ShaderBindingState);
 
-	static void SetShaderBindings(
+	static RENDERER_API void SetShaderBindings(
 		FRHIBatchedShaderParameters& BatchedParameters,
 		const class FReadOnlyMeshDrawSingleShaderBindings& RESTRICT SingleShaderBindings);
+	
+	explicit FReadOnlyMeshDrawSingleShaderBindings(const FMeshDrawSingleShaderBindings& Bindings) :
+		FMeshDrawShaderBindingsLayout(Bindings)
+	{
+		Data = Bindings.Data;
+	}
+
+	FReadOnlyMeshDrawSingleShaderBindings(const FMeshDrawShaderBindingsLayout& InLayout, const uint8* InData) :
+		FMeshDrawShaderBindingsLayout(InLayout)
+	{
+		Data = InData;
+	}
+
+	inline FRHIUniformBuffer*const* GetUniformBufferStart() const
+	{
+		return (FRHIUniformBuffer**)(Data + GetUniformBufferOffset());
+	}
+
+	inline FRHISamplerState** GetSamplerStart() const
+	{
+		const uint8* SamplerDataStart = Data + GetSamplerOffset();
+		return (FRHISamplerState**)SamplerDataStart;
+	}
+
+	inline FRHIResource** GetSRVStart() const
+	{
+		const uint8* SRVDataStart = Data + GetSRVOffset();
+		return (FRHIResource**)SRVDataStart;
+	}
+
+	inline const uint8* GetSRVTypeStart() const
+	{
+		const uint8* SRVTypeDataStart = Data + GetSRVTypeOffset();
+		return SRVTypeDataStart;
+	}
+
+	inline const uint8* GetLooseDataStart() const
+	{
+		const uint8* LooseDataStart = Data + GetLooseDataOffset();
+		return LooseDataStart;
+	}
+
+private:
+	const uint8* Data;
 };
 
 struct FMeshDrawCommandOverrideArgs
