Cleanup r.MinScreenRadiusForCSMDepth which is not used anymore, r.Shadow.RadiusThreshold is now used for culling shadow casters.

diff --git a/Engine/Source/Runtime/Engine/Classes/Engine/RendererSettings.h b/Engine/Source/Runtime/Engine/Classes/Engine/RendererSettings.h
--- a/Engine/Source/Runtime/Engine/Classes/Engine/RendererSettings.h
+++ b/Engine/Source/Runtime/Engine/Classes/Engine/RendererSettings.h
@@ -363,11 +363,6 @@ class URendererSettings : public UDeveloperSettings
 		ToolTip="Screen radius at which objects are culled for the early Z pass. Larger values can improve performance but very large values can degrade performance if large occluders are not rendered."))
 	float MinScreenRadiusForEarlyZPass;
 
-	UPROPERTY(config, EditAnywhere, Category=Culling, meta=(
-		ConsoleVariable="r.MinScreenRadiusForCSMDepth",DisplayName="Min Screen Radius for Cascaded Shadow Maps",
-		ToolTip="Screen radius at which objects are culled for cascaded shadow map depth passes. Larger values can improve performance but can cause artifacts as objects stop casting shadows."))
-	float MinScreenRadiusForCSMdepth;
-
 	UPROPERTY(config, EditAnywhere, Category=Culling, meta=(
 		ConsoleVariable="r.PrecomputedVisibilityWarning",DisplayName="Warn about no precomputed visibility",
 		ToolTip="Displays a warning when no precomputed visibility data is available for the current camera location. This can be helpful if you are making a game that relies on precomputed visibility, e.g. a first person mobile game."))
diff --git a/Engine/Source/Runtime/Renderer/Private/SceneVisibility.cpp b/Engine/Source/Runtime/Renderer/Private/SceneVisibility.cpp
--- a/Engine/Source/Runtime/Renderer/Private/SceneVisibility.cpp
+++ b/Engine/Source/Runtime/Renderer/Private/SceneVisibility.cpp
@@ -88,14 +88,6 @@ static FAutoConsoleVariableRef CVarMinScreenRadiusForDepthPrepass(
 	ECVF_RenderThreadSafe
 	);
 
-float GMinScreenRadiusForCSMDepth = 0.01f;
-static FAutoConsoleVariableRef CVarMinScreenRadiusForCSMDepth(
-	TEXT("r.MinScreenRadiusForCSMDepth"),
-	GMinScreenRadiusForCSMDepth,
-	TEXT("Threshold below which meshes will be culled from CSM depth pass."),
-	ECVF_RenderThreadSafe
-	);
-
 static TAutoConsoleVariable<int32> CVarTemporalAASamples(
 	TEXT("r.TemporalAASamples"),
 	8,
@@ -1101,7 +1093,6 @@ FFilterStaticMeshesForViewData::FFilterStaticMeshesForViewData(FViewInfo& View)
 
 	LODScale = CVarStaticMeshLODDistanceScale.GetValueOnRenderThread() * View.LODDistanceFactor;
 
-	MinScreenRadiusForCSMDepthSquared = GMinScreenRadiusForCSMDepth * GMinScreenRadiusForCSMDepth;
 	MinScreenRadiusForDepthPrepassSquared = GMinScreenRadiusForDepthPrepass * GMinScreenRadiusForDepthPrepass;
 
 	bFullEarlyZPass = ShouldForceFullDepthPass(View.GetShaderPlatform());
@@ -1480,7 +1471,6 @@ void FRelevancePacket::ComputeRelevance(FDynamicPrimitiveIndexList& DynamicPrimi
 
 			float DistanceSquared = (Bounds.BoxSphereBounds.Origin - ViewData.ViewOrigin).SizeSquared();
 			const float LODFactorDistanceSquared = DistanceSquared * FMath::Square(ViewData.LODScale);
-			const bool bDrawShadowDepth = FMath::Square(Bounds.BoxSphereBounds.SphereRadius) > ViewData.MinScreenRadiusForCSMDepthSquared * LODFactorDistanceSquared;
 			const bool bDrawDepthOnly = ViewData.bFullEarlyZPass || ((ShadingPath != EShadingPath::Mobile) && (FMath::Square(Bounds.BoxSphereBounds.SphereRadius) > GMinScreenRadiusForDepthPrepass * GMinScreenRadiusForDepthPrepass * LODFactorDistanceSquared));
 
 			const int32 NumStaticMeshes = PrimitiveSceneInfo->StaticMeshRelevances.Num();
diff --git a/Engine/Source/Runtime/Renderer/Private/SceneVisibilityPrivate.h b/Engine/Source/Runtime/Renderer/Private/SceneVisibilityPrivate.h
--- a/Engine/Source/Runtime/Renderer/Private/SceneVisibilityPrivate.h
+++ b/Engine/Source/Runtime/Renderer/Private/SceneVisibilityPrivate.h
@@ -659,7 +659,6 @@ struct FFilterStaticMeshesForViewData
 	FVector ViewOrigin;
 	int32 ForcedLODLevel;
 	float LODScale;
-	float MinScreenRadiusForCSMDepthSquared;
 	float MinScreenRadiusForDepthPrepassSquared;
 	bool bFullEarlyZPass;
 
