Removed !ENABLE_LIGHT_CULLING_VIEW_SPACE_BUILD_DATA code paths.

diff --git a/Engine/Shaders/Private/LightGridInjection.usf b/Engine/Shaders/Private/LightGridInjection.usf
--- a/Engine/Shaders/Private/LightGridInjection.usf
+++ b/Engine/Shaders/Private/LightGridInjection.usf
@@ -36,10 +36,8 @@ uint LightGridPixelSizeShift;
 
 uint MegaLightsSupportedStartIndex;
 
-#if ENABLE_LIGHT_CULLING_VIEW_SPACE_BUILD_DATA
 StructuredBuffer<float4> LightViewSpacePositionAndRadius;
 StructuredBuffer<float4> LightViewSpaceDirAndPreprocAngle;
-#endif // ENABLE_LIGHT_CULLING_VIEW_SPACE_BUILD_DATA
 
 #define REFINE_SPOTLIGHT_BOUNDS 1
 
@@ -204,16 +202,9 @@ void LightGridInjectionCS(
 		{
 			uint LocalLightBaseIndex = LocalLightIndex * LOCAL_LIGHT_DATA_STRIDE;
 
-#if ENABLE_LIGHT_CULLING_VIEW_SPACE_BUILD_DATA
 			float4 LightPositionAndRadius = LightViewSpacePositionAndRadius[LocalLightIndex];
 			float3 ViewSpaceLightPosition = LightPositionAndRadius.xyz;
 			float LightRadius = LightPositionAndRadius.w;
-#else // !ENABLE_LIGHT_CULLING_VIEW_SPACE_BUILD_DATA
-			float4 LightPositionAndInvRadius = ForwardLocalLightBuffer[LocalLightBaseIndex + 0];
-			float LightRadius = 1.0f / LightPositionAndInvRadius.w;
-
-			float3 ViewSpaceLightPosition = mul(float4(LightPositionAndInvRadius.xyz, 1), View.TranslatedWorldToView).xyz;
-#endif // ENABLE_LIGHT_CULLING_VIEW_SPACE_BUILD_DATA
 
 			float BoxDistanceSq = ComputeSquaredDistanceFromBoxToPoint(ViewTileCenter, ViewTileExtent, ViewSpaceLightPosition);
 
diff --git a/Engine/Source/Runtime/Renderer/Private/LightGridInjection.cpp b/Engine/Source/Runtime/Renderer/Private/LightGridInjection.cpp
--- a/Engine/Source/Runtime/Renderer/Private/LightGridInjection.cpp
+++ b/Engine/Source/Runtime/Renderer/Private/LightGridInjection.cpp
@@ -227,7 +227,6 @@ public:
 		OutEnvironment.SetDefine(TEXT("THREADGROUP_SIZE"), LightGridInjectionGroupSize);
 		FForwardLightingParameters::ModifyCompilationEnvironment(Parameters.Platform, OutEnvironment);
 		OutEnvironment.SetDefine(TEXT("LIGHT_LINK_STRIDE"), LightLinkStride);
-		OutEnvironment.SetDefine(TEXT("ENABLE_LIGHT_CULLING_VIEW_SPACE_BUILD_DATA"), ENABLE_LIGHT_CULLING_VIEW_SPACE_BUILD_DATA);
 	}
 };
 
@@ -268,7 +267,6 @@ public:
 		FForwardLightingParameters::ModifyCompilationEnvironment(Parameters.Platform, OutEnvironment);
 		OutEnvironment.SetDefine(TEXT("LIGHT_LINK_STRIDE"), LightLinkStride);
 		OutEnvironment.SetDefine(TEXT("MAX_CAPTURES"), GetMaxNumReflectionCaptures(Parameters.Platform));
-		OutEnvironment.SetDefine(TEXT("ENABLE_LIGHT_CULLING_VIEW_SPACE_BUILD_DATA"), ENABLE_LIGHT_CULLING_VIEW_SPACE_BUILD_DATA);
 	}
 };
 
@@ -479,10 +477,9 @@ FComputeLightGridOutput FSceneRenderer::ComputeLightGrid(FRDGBuilder& GraphBuild
 
 		TArray<FForwardLocalLightData, SceneRenderingAllocator> ForwardLocalLightData;
 		TArray<int32, SceneRenderingAllocator>  LocalLightVisibleLightInfosIndex;
-#if ENABLE_LIGHT_CULLING_VIEW_SPACE_BUILD_DATA
+
 		TArray<FVector4f, SceneRenderingAllocator> ViewSpacePosAndRadiusData;
 		TArray<FVector4f, SceneRenderingAllocator> ViewSpaceDirAndPreprocAngleData;
-#endif // ENABLE_LIGHT_CULLING_VIEW_SPACE_BUILD_DATA
 
 		float FurthestLight = 1000;
 
@@ -504,10 +501,9 @@ FComputeLightGridOutput FSceneRenderer::ComputeLightGrid(FRDGBuilder& GraphBuild
 			{
 				ForwardLocalLightData.Reserve(SimpleLightsEnd);
 				LocalLightVisibleLightInfosIndex.Reserve(SimpleLightsEnd);
-#if ENABLE_LIGHT_CULLING_VIEW_SPACE_BUILD_DATA
+
 				ViewSpacePosAndRadiusData.Reserve(SimpleLightsEnd);
 				ViewSpaceDirAndPreprocAngleData.Reserve(SimpleLightsEnd);
-#endif // ENABLE_LIGHT_CULLING_VIEW_SPACE_BUILD_DATA
 
 				const FSimpleLightArray& SimpleLights = SortedLightSet.SimpleLights;
 
@@ -538,11 +534,9 @@ FComputeLightGridOutput FSceneRenderer::ComputeLightGrid(FRDGBuilder& GraphBuild
 					const FSimpleLightPerViewEntry& SimpleLightPerViewData = SimpleLights.GetViewDependentData(SimpleLightIndex, ViewIndex, Views.Num());
 					PackLocalLightData(LightData, View, SimpleLight, SimpleLightPerViewData);
 
-				#if ENABLE_LIGHT_CULLING_VIEW_SPACE_BUILD_DATA
 					FVector4f ViewSpacePosAndRadius(FVector4f(View.ViewMatrices.GetViewMatrix().TransformPosition(SimpleLightPerViewData.Position)), SimpleLight.Radius);
 					ViewSpacePosAndRadiusData.Add(ViewSpacePosAndRadius);
 					ViewSpaceDirAndPreprocAngleData.AddZeroed();
-				#endif // ENABLE_LIGHT_CULLING_VIEW_SPACE_BUILD_DATA
 				}
 			}
 
@@ -630,7 +624,6 @@ FComputeLightGridOutput FSceneRenderer::ComputeLightGrid(FRDGBuilder& GraphBuild
 						const float Distance = View.ViewMatrices.GetViewMatrix().TransformPosition(BoundingSphere.Center).Z + BoundingSphere.W;
 						FurthestLight = FMath::Max(FurthestLight, Distance);
 
-					#if ENABLE_LIGHT_CULLING_VIEW_SPACE_BUILD_DATA
 						// Note: inverting radius twice seems stupid (but done in shader anyway otherwise)
 						const FVector3f LightViewPosition = FVector4f(View.ViewMatrices.GetViewMatrix().TransformPosition(LightParameters.WorldPosition)); // LWC_TODO: precision loss
 						FVector4f ViewSpacePosAndRadius(LightViewPosition, 1.0f / LightParameters.InvRadius);
@@ -639,7 +632,6 @@ FComputeLightGridOutput FSceneRenderer::ComputeLightGrid(FRDGBuilder& GraphBuild
 						const float PreProcAngle = SortedLightInfo.SortKey.Fields.LightType == LightType_Spot ? GetTanRadAngleOrZero(LightSceneInfo->Proxy->GetOuterConeAngle()) : 0.0f;
 						FVector4f ViewSpaceDirAndPreprocAngle(FVector4f(View.ViewMatrices.GetViewMatrix().TransformVector((FVector)LightParameters.Direction)), PreProcAngle); // LWC_TODO: precision loss
 						ViewSpaceDirAndPreprocAngleData.Add(ViewSpaceDirAndPreprocAngle);
-					#endif // ENABLE_LIGHT_CULLING_VIEW_SPACE_BUILD_DATA
 					}
 					// On mobile there is a separate FMobileDirectionalLightShaderParameters UB which holds all directional light data.
 					else if (SortedLightInfo.SortKey.Fields.LightType == LightType_Directional && ViewFamily.EngineShowFlags.DirectionalLights && !IsMobilePlatform(View.GetShaderPlatform()))
@@ -808,23 +800,6 @@ FComputeLightGridOutput FSceneRenderer::ComputeLightGrid(FRDGBuilder& GraphBuild
 			}
 		}
 
-#if ENABLE_LIGHT_CULLING_VIEW_SPACE_BUILD_DATA
-		// Fuse these loops as I see no reason why not and we build some temporary data that is needed in the build pass and is 
-		// not needed to be stored permanently.
-#else // !ENABLE_LIGHT_CULLING_VIEW_SPACE_BUILD_DATA
-
-		ForwardLightDataPerView.Emplace(ForwardLightData);
-	}
-
-	for (int32 ViewIndex = 0; ViewIndex < Views.Num(); ViewIndex++)
-	{
-		const FViewInfo& View = Views[ViewIndex];
-		FForwardLightData* ForwardLightData = ForwardLightDataPerView[ViewIndex];
-
-		const FIntPoint LightGridSizeXY = FIntPoint::DivideAndRoundUp(View.ViewRect.Size(), GLightGridPixelSize);
-
-#endif // ENABLE_LIGHT_CULLING_VIEW_SPACE_BUILD_DATA
-
 		// Allocate buffers using the scene render targets size so we won't reallocate every frame with dynamic resolution
 		const FIntPoint MaxLightGridSizeXY = FIntPoint::DivideAndRoundUp(View.GetSceneTexturesConfig().Extent, GLightGridPixelSize);
 
@@ -899,7 +874,6 @@ FComputeLightGridOutput FSceneRenderer::ComputeLightGrid(FRDGBuilder& GraphBuild
 			PassParameters->LightGridPixelSizeShift = ForwardLightData->LightGridPixelSizeShift;
 			PassParameters->MegaLightsSupportedStartIndex = ForwardLightData->MegaLightsSupportedStartIndex;
 
-#if ENABLE_LIGHT_CULLING_VIEW_SPACE_BUILD_DATA
 			check(ViewSpacePosAndRadiusData.Num() == ForwardLocalLightData.Num());
 			check(ViewSpaceDirAndPreprocAngleData.Num() == ForwardLocalLightData.Num());
 
@@ -908,7 +882,6 @@ FComputeLightGridOutput FSceneRenderer::ComputeLightGrid(FRDGBuilder& GraphBuild
 
 			PassParameters->LightViewSpacePositionAndRadius  = GraphBuilder.CreateSRV(FRDGBufferSRVDesc(LightViewSpacePositionAndRadius));
 			PassParameters->LightViewSpaceDirAndPreprocAngle = GraphBuilder.CreateSRV(FRDGBufferSRVDesc(LightViewSpaceDirAndPreprocAngle));
-#endif // ENABLE_LIGHT_CULLING_VIEW_SPACE_BUILD_DATA
 
 			FLightGridInjectionCS::FPermutationDomain PermutationVector;
 			PermutationVector.Set<FLightGridInjectionCS::FUseLinkedListDim>(GLightLinkedListCulling != 0);
diff --git a/Engine/Source/Runtime/Renderer/Private/SceneRendering.h b/Engine/Source/Runtime/Renderer/Private/SceneRendering.h
--- a/Engine/Source/Runtime/Renderer/Private/SceneRendering.h
+++ b/Engine/Source/Runtime/Renderer/Private/SceneRendering.h
@@ -704,8 +704,6 @@ public:
 	TArray<int32> LocalLightVisibleLightInfosIndex;
 };
 
-#define ENABLE_LIGHT_CULLING_VIEW_SPACE_BUILD_DATA 1
-
 BEGIN_GLOBAL_SHADER_PARAMETER_STRUCT_WITH_CONSTRUCTOR(FVolumetricFogGlobalData,) 
 	SHADER_PARAMETER(FIntVector,ViewGridSizeInt)
 	SHADER_PARAMETER(FVector3f, ViewGridSize)
