Fix for decals in orthographic mode drawing over other scene elements. The Z buffer offset for avoiding z fighting for decals was set on a scale appropriate for perspective, but this difference is too large for Orthographic depths, so the chosen value is now scaled by the PerProjectionDepthThicknessScale so the offset also works for Ortho.

diff --git a/Engine/Source/Runtime/Renderer/Private/SceneRendering.cpp b/Engine/Source/Runtime/Renderer/Private/SceneRendering.cpp
--- a/Engine/Source/Runtime/Renderer/Private/SceneRendering.cpp
+++ b/Engine/Source/Runtime/Renderer/Private/SceneRendering.cpp
@@ -262,7 +262,7 @@ static TAutoConsoleVariable<float> CVarDemosaicVposOffset(
 static TAutoConsoleVariable<float> CVarDecalDepthBias(
 	TEXT("r.DecalDepthBias"),
 	0.005f,
-	TEXT("Global depth bias used by mesh decals. Default is 0.005"),
+	TEXT("Global depth bias used by mesh decals. Default is 0.005 for perspective. Scaled by the PerProjectionDepthThicknessScale for Ortho"),
 	ECVF_RenderThreadSafe);
 
 static TAutoConsoleVariable<int32> CVarRefractionQuality(
@@ -1835,7 +1835,7 @@ void FViewInfo::SetupUniformBufferParameters(
 		ViewUniformShaderParameters.DemosaicVposOffset = CVarDemosaicVposOffset.GetValueOnRenderThread();
 	}
 
-	ViewUniformShaderParameters.DecalDepthBias = CVarDecalDepthBias.GetValueOnRenderThread();
+	ViewUniformShaderParameters.DecalDepthBias = CVarDecalDepthBias.GetValueOnRenderThread() * InViewMatrices.GetPerProjectionDepthThicknessScale();
 
 	ViewUniformShaderParameters.IndirectLightingColorScale = FVector3f(FinalPostProcessSettings.IndirectLightingColor.R * FinalPostProcessSettings.IndirectLightingIntensity,
 		FinalPostProcessSettings.IndirectLightingColor.G * FinalPostProcessSettings.IndirectLightingIntensity,
