Refine rect light culling in Light Grid using half space test.
- stole a bit from PreProcAngle to store a flag identifying which lights are rect lights.

Peformance on level with 15 rect lights with ~half screen behind the rect lights (best case):
- ComputeLightGrid: 0.21ms -> 0.14ms
- MegaLights.GenerateSamples: 0.65ms -> 0.45ms

diff --git a/Engine/Shaders/Private/LightGridInjection.usf b/Engine/Shaders/Private/LightGridInjection.usf
--- a/Engine/Shaders/Private/LightGridInjection.usf
+++ b/Engine/Shaders/Private/LightGridInjection.usf
@@ -40,6 +40,7 @@ StructuredBuffer<float4> LightViewSpacePositionAndRadius;
 StructuredBuffer<float4> LightViewSpaceDirAndPreprocAngle;
 
 #define REFINE_SPOTLIGHT_BOUNDS 1
+#define REFINE_RECTLIGHT_BOUNDS 1
 
 #ifndef USE_LOCALSTORE_FOR_ATOMICS
 #define USE_LOCALSTORE_FOR_ATOMICS 0
@@ -210,21 +211,37 @@ void LightGridInjectionCS(
 
 			if (BoxDistanceSq < LightRadius * LightRadius)
 			{
+				float4 ViewSpaceDirAndPreprocAngle = LightViewSpaceDirAndPreprocAngle[LocalLightIndex];
+				bool bIsRectLight = asuint(ViewSpaceDirAndPreprocAngle.w) & 0x1;
+
+				bool bPassSpotLightTest = true;
 #if REFINE_SPOTLIGHT_BOUNDS
-				bool bPassSpotlightTest = true;
 				{
-					float4 ViewSpaceDirAndPreprocAngle = LightViewSpaceDirAndPreprocAngle[LocalLightIndex];
-					float TanConeAngle = ViewSpaceDirAndPreprocAngle.w;
+					float TanConeAngle = asfloat(asuint(ViewSpaceDirAndPreprocAngle.w) & 0xFFFFFFFE);
 
 					// Set to 0 for non-acute cones, or non-spot lights.
 					if (TanConeAngle > 0.0f)
 					{
 						float3 ViewSpaceLightDirection = -ViewSpaceDirAndPreprocAngle.xyz;
-						bPassSpotlightTest = !IsAabbOutsideInfiniteAcuteConeApprox(ViewSpaceLightPosition, ViewSpaceLightDirection, TanConeAngle, ViewTileCenter, ViewTileExtent);
+						bPassSpotLightTest = !IsAabbOutsideInfiniteAcuteConeApprox(ViewSpaceLightPosition, ViewSpaceLightDirection, TanConeAngle, ViewTileCenter, ViewTileExtent);
 					}
 				}
-				if (bPassSpotlightTest)
 #endif // REFINE_SPOTLIGHT_BOUNDS
+				bool bPassRectLightTest = true;
+#if REFINE_RECTLIGHT_BOUNDS
+				{
+					if (bIsRectLight)
+					{
+						float3 D = ViewTileCenter - ViewSpaceLightPosition;
+						float4 Plane = float4(ViewSpaceDirAndPreprocAngle.xyz, 0.0);
+
+						bPassRectLightTest = !AabbOutsidePlane(D, ViewTileExtent, Plane);
+
+						// TODO: Refine using spot light test when BarnDoorAngle < 90 degrees and BarnDoorLength > 0
+					}
+				}
+#endif // REFINE_RECTLIGHT_BOUNDS
+				if (bPassSpotLightTest && bPassRectLightTest)
 				{
 
 #if USE_LINKED_CULL_LIST
diff --git a/Engine/Source/Runtime/Renderer/Private/LightGridInjection.cpp b/Engine/Source/Runtime/Renderer/Private/LightGridInjection.cpp
--- a/Engine/Source/Runtime/Renderer/Private/LightGridInjection.cpp
+++ b/Engine/Source/Runtime/Renderer/Private/LightGridInjection.cpp
@@ -629,9 +629,11 @@ FComputeLightGridOutput FSceneRenderer::ComputeLightGrid(FRDGBuilder& GraphBuild
 						FVector4f ViewSpacePosAndRadius(LightViewPosition, 1.0f / LightParameters.InvRadius);
 						ViewSpacePosAndRadiusData.Add(ViewSpacePosAndRadius);
 
+						// Pack flags in the LSB of PreProcAngle
 						const float PreProcAngle = SortedLightInfo.SortKey.Fields.LightType == LightType_Spot ? GetTanRadAngleOrZero(LightSceneInfo->Proxy->GetOuterConeAngle()) : 0.0f;
-						FVector4f ViewSpaceDirAndPreprocAngle(FVector4f(View.ViewMatrices.GetViewMatrix().TransformVector((FVector)LightParameters.Direction)), PreProcAngle); // LWC_TODO: precision loss
-						ViewSpaceDirAndPreprocAngleData.Add(ViewSpaceDirAndPreprocAngle);
+						const uint32 PackedPreProcAngleAndFlags = (FMath::AsUInt(PreProcAngle) & 0xFFFFFFFE) | (LightProxy->IsRectLight() ? 0x1 : 0);
+						FVector4f ViewSpaceDirAndPreprocAngleAndFlags(FVector4f(View.ViewMatrices.GetViewMatrix().TransformVector((FVector)LightParameters.Direction)), FMath::AsFloat(PackedPreProcAngleAndFlags)); // LWC_TODO: precision loss
+						ViewSpaceDirAndPreprocAngleData.Add(ViewSpaceDirAndPreprocAngleAndFlags);
 					}
 					// On mobile there is a separate FMobileDirectionalLightShaderParameters UB which holds all directional light data.
 					else if (SortedLightInfo.SortKey.Fields.LightType == LightType_Directional && ViewFamily.EngineShowFlags.DirectionalLights && !IsMobilePlatform(View.GetShaderPlatform()))
@@ -1114,6 +1116,11 @@ class FDebugLightGridPS : public FGlobalShader
 		return EnumHasAllFlags(Parameters.Flags, EShaderPermutationFlags::HasEditorOnlyData) && ShaderPrint::IsSupported(Parameters.Platform);
 	}
 
+	static EShaderPermutationPrecacheRequest ShouldPrecachePermutation(const FGlobalShaderPermutationParameters& Parameters)
+	{
+		return EShaderPermutationPrecacheRequest::NotPrecached;
+	}
+
 	static void ModifyCompilationEnvironment(const FGlobalShaderPermutationParameters& Parameters, FShaderCompilerEnvironment& OutEnvironment)
 	{
 		FGlobalShader::ModifyCompilationEnvironment(Parameters, OutEnvironment);
