Removed extraneous dispatch hints, but fixed the one after Occlusion.

diff --git a/Engine/Source/Runtime/Renderer/Private/DeferredShadingRenderer.cpp b/Engine/Source/Runtime/Renderer/Private/DeferredShadingRenderer.cpp
--- a/Engine/Source/Runtime/Renderer/Private/DeferredShadingRenderer.cpp
+++ b/Engine/Source/Runtime/Renderer/Private/DeferredShadingRenderer.cpp
@@ -536,8 +536,6 @@ static void RenderOpaqueFX(
 		{
 			GPUSortManager->OnPostRenderOpaque(GraphBuilder);
 		}
-
-		GraphBuilder.AddDispatchHint();
 	}
 }
 
@@ -2472,7 +2470,6 @@ void FDeferredShadingSceneRenderer::Render(FRDGBuilder& GraphBuilder)
 			if (!bHasRayTracedOverlay)
 			{
 				RenderBasePass(*this, GraphBuilder, Views, SceneTextures, DBufferTextures, BasePassDepthStencilAccess, ForwardScreenSpaceShadowMaskTexture, InstanceCullingManager, bNaniteEnabled, Scene->NaniteShadingCommands[ENaniteMeshPass::BasePass], NaniteRasterResults);
-				GraphBuilder.AddDispatchHint();
 			}
 
 			if (!bAllowReadOnlyDepthBasePass)
diff --git a/Engine/Source/Runtime/Renderer/Private/HairStrands/HairStrandsRendering.cpp b/Engine/Source/Runtime/Renderer/Private/HairStrands/HairStrandsRendering.cpp
--- a/Engine/Source/Runtime/Renderer/Private/HairStrands/HairStrandsRendering.cpp
+++ b/Engine/Source/Runtime/Renderer/Private/HairStrands/HairStrandsRendering.cpp
@@ -142,7 +142,6 @@ void RenderHairPrePass(
 				// Render DeepShadow for the second view, as for now the computations are view dependent. 
 				// This needs to be view independent to share result between eyes.
 				RenderHairStrandsDeepShadows(GraphBuilder, Scene, View, InstanceCullingManager);
-				GraphBuilder.AddDispatchHint();
 				return;
 			}
 		}
@@ -157,7 +156,6 @@ void RenderHairPrePass(
 
 		//SCOPED_GPU_STAT(RHICmdList, HairRendering);
 		CreateHairStrandsMacroGroups(GraphBuilder, Scene, View, InstancesVisibilityType, View.HairStrandsViewData);
-		GraphBuilder.AddDispatchHint();
 
 		// Voxelization and Deep Opacity Maps
 		VoxelizeHairStrands(GraphBuilder, Scene, View, InstanceCullingManager, PreViewStereoCorrection);
@@ -166,7 +164,6 @@ void RenderHairPrePass(
 			AddMeshDrawTransitionPass(GraphBuilder, View, View.HairStrandsViewData.MacroGroupDatas);
 		}
 		RenderHairStrandsDeepShadows(GraphBuilder, Scene, View, InstanceCullingManager);
-		GraphBuilder.AddDispatchHint();
 	}
 }
 
diff --git a/Engine/Source/Runtime/Renderer/Private/SceneOcclusion.cpp b/Engine/Source/Runtime/Renderer/Private/SceneOcclusion.cpp
--- a/Engine/Source/Runtime/Renderer/Private/SceneOcclusion.cpp
+++ b/Engine/Source/Runtime/Renderer/Private/SceneOcclusion.cpp
@@ -1474,12 +1474,7 @@ void FDeferredShadingSceneRenderer::RenderOcclusion(
 
 	if (bUseHzbOcclusion || bIsOcclusionTesting)
 	{
-		// Hint to the RHI to submit commands up to this point to the GPU if possible.  Can help avoid CPU stalls next frame waiting
-		// for these query results on some platforms.
-		AddPass(GraphBuilder, RDG_EVENT_NAME("SubmitCommands"), [](FRHICommandList& RHICmdList)
-		{
-			RHICmdList.SubmitCommandsHint();
-		});
+		GraphBuilder.AddDispatchHint();
 	}
 
 	if (bIsOcclusionTesting)
