HeterogeneousVolumes: Reinstating LOD support for shadow-casting passes.

diff --git a/Engine/Source/Runtime/Renderer/Private/HeterogeneousVolumes/HeterogeneousVolumes.cpp b/Engine/Source/Runtime/Renderer/Private/HeterogeneousVolumes/HeterogeneousVolumes.cpp
--- a/Engine/Source/Runtime/Renderer/Private/HeterogeneousVolumes/HeterogeneousVolumes.cpp
+++ b/Engine/Source/Runtime/Renderer/Private/HeterogeneousVolumes/HeterogeneousVolumes.cpp
@@ -582,7 +582,10 @@ namespace HeterogeneousVolumes
 		);
 	}
 	
-	float CalcLOD(const FSceneView& View, const IHeterogeneousVolumeInterface* HeterogeneousVolume)
+	float CalcLOD(
+		const HeterogeneousVolumes::FLODInfo& LODInfo,
+		const IHeterogeneousVolumeInterface* HeterogeneousVolume
+	)
 	{
 		if (!HeterogeneousVolumes::UseContinuousLOD())
 		{
@@ -594,21 +597,27 @@ namespace HeterogeneousVolumes
 		float VoxelResolutionMin = VoxelResolution.GetMin();
 
 		float LODValue = FMath::Floor(FMath::Log2(VoxelResolutionMin));
-		if (View.ViewFrustum.IntersectBox(WorldBounds.Origin, WorldBounds.BoxExtent))
+		if (!LODInfo.bIsPerspective)
+		{
+			float VolumeRatio = FVector(LODInfo.WorldSceneBounds.BoxExtent / WorldBounds.BoxExtent).Length();
+			float ViewLODValue = FMath::Log2(VolumeRatio) + HeterogeneousVolume->GetMipBias() + HeterogeneousVolumes::GetCLODBias();
+			ViewLODValue = FMath::Max(ViewLODValue, 0);
+
+			LODValue = FMath::Min(ViewLODValue, LODValue);
+		}
+		else if (LODInfo.WorldShadowFrustum.IntersectBox(WorldBounds.Origin, WorldBounds.BoxExtent))
 		{
 			// Determine the pixel-width at the near-plane
-			float TanHalfFOV = FMath::Tan(FMath::DegreesToRadians(View.FOV * 0.5));
-			float HalfWidth = View.UnconstrainedViewRect.Width() / HeterogeneousVolumes::GetDownsampleFactor() * 0.5;
-			float PixelWidth = TanHalfFOV / HalfWidth;
+			float TanHalfFOV = FMath::Tan(LODInfo.FOV * 0.5);
+			float HalfViewWidth = LODInfo.ViewRect.Width() * 0.5 / LODInfo.DownsampleFactor;
+			float PixelWidth = TanHalfFOV / HalfViewWidth;
 
 			// Project to nearest distance of volume bounds
-			FVector WorldCameraOrigin = View.ViewMatrices.GetViewOrigin();
-			float Distance = FMath::Max((WorldBounds.Origin - WorldCameraOrigin).Length() - WorldBounds.SphereRadius, View.NearClippingDistance);
-			float VoxelWidth = Distance * PixelWidth;
+			float Distance = FMath::Max(FVector::Dist(WorldBounds.Origin, LODInfo.WorldOrigin) - WorldBounds.SphereRadius, LODInfo.NearClippingDistance);
+			float ProjectedPixelWidth = Distance * PixelWidth;
 
 			// MIP is defined as the log of the ratio of native voxel resolution to pixel-coverage of volume bounds
-			//float PixelWidthCoverage = (2.0 * WorldBounds.SphereRadius) / VoxelWidth;
-			float PixelWidthCoverage = (2.0 * WorldBounds.BoxExtent.GetMax()) / VoxelWidth;
+			float PixelWidthCoverage = (2.0 * WorldBounds.BoxExtent.GetMax()) / ProjectedPixelWidth;
 			float ViewLODValue = FMath::Log2(VoxelResolutionMin / PixelWidthCoverage) + HeterogeneousVolume->GetMipBias() + HeterogeneousVolumes::GetCLODBias();
 			ViewLODValue = FMath::Max(ViewLODValue, 0);
 
@@ -617,12 +626,37 @@ namespace HeterogeneousVolumes
 
 		return LODValue;
 	}
+	
+	float CalcLOD(const FSceneView& View, const IHeterogeneousVolumeInterface* HeterogeneousVolume)
+	{
+		FLODInfo LODInfo;
+		// TODO: Not supporting orthographic projection for now
+		LODInfo.bIsPerspective = true;
+		LODInfo.WorldSceneBounds = FBoxSphereBounds(EForceInit::ForceInitToZero);
+
+		LODInfo.WorldOrigin = View.ViewMatrices.GetViewOrigin();
+		LODInfo.ViewRect = View.UnconstrainedViewRect;
+		LODInfo.WorldShadowFrustum = View.ViewFrustum;
+		LODInfo.FOV = FMath::DegreesToRadians(View.FOV);
+		LODInfo.NearClippingDistance = View.NearClippingDistance;
+		LODInfo.DownsampleFactor = HeterogeneousVolumes::GetDownsampleFactor();
+
+		return CalcLOD(LODInfo, HeterogeneousVolume);
+	}
 
 	float CalcLODFactor(float LODValue)
 	{
 		return FMath::Pow(2, LODValue);
 	}
 
+	float CalcLODFactor(
+		const HeterogeneousVolumes::FLODInfo& LODInfo,
+		const IHeterogeneousVolumeInterface* HeterogeneousVolume
+	)
+	{
+		return CalcLODFactor(CalcLOD(LODInfo, HeterogeneousVolume));
+	}
+
 	float CalcLODFactor(const FSceneView& View, const IHeterogeneousVolumeInterface* HeterogeneousVolume)
 	{
 		return CalcLODFactor(CalcLOD(View, HeterogeneousVolume));
diff --git a/Engine/Source/Runtime/Renderer/Private/HeterogeneousVolumes/HeterogeneousVolumes.h b/Engine/Source/Runtime/Renderer/Private/HeterogeneousVolumes/HeterogeneousVolumes.h
--- a/Engine/Source/Runtime/Renderer/Private/HeterogeneousVolumes/HeterogeneousVolumes.h
+++ b/Engine/Source/Runtime/Renderer/Private/HeterogeneousVolumes/HeterogeneousVolumes.h
@@ -3,6 +3,7 @@
 #pragma once
 
 #include "CoreTypes.h"
+#include "ConvexVolume.h"
 #include "RendererInterface.h"
 #include "ShaderParameterMacros.h"
 
@@ -129,9 +130,30 @@ namespace HeterogeneousVolumes
 	int GetVoxelCount(FIntVector VolumeResolution);
 	int GetVoxelCount(const FRDGTextureDesc& TextureDesc);
 	FIntVector GetMipVolumeResolution(FIntVector VolumeResolution, uint32 MipLevel);
-	float CalcLOD(const FSceneView& View, const IHeterogeneousVolumeInterface* HeterogeneousVolume);
-	float CalcLODFactor(const FSceneView& View, const IHeterogeneousVolumeInterface* HeterogeneousVolume);
+
+	struct FLODInfo
+	{
+		// Orthographic projection
+		FBoxSphereBounds WorldSceneBounds = FBoxSphereBounds(EForceInit::ForceInit);
+
+		// Perspective projection
+		FVector WorldOrigin = FVector::ZeroVector;
+		FIntRect ViewRect;
+
+		FConvexVolume WorldShadowFrustum;
+		float FOV = PI / 4.0f;
+		float NearClippingDistance = 1.0f;
+		float DownsampleFactor = 1.0f;
+
+		// Projection type
+		bool bIsPerspective = false;
+	};
+
 	float CalcLODFactor(float LOD);
+	float CalcLODFactor(const HeterogeneousVolumes::FLODInfo& LODInfo, const IHeterogeneousVolumeInterface* HeterogeneousVolume);
+	float CalcLODFactor(const FSceneView& View, const IHeterogeneousVolumeInterface* HeterogeneousVolume);
+	float CalcLOD(const HeterogeneousVolumes::FLODInfo& LODInfo, const IHeterogeneousVolumeInterface* HeterogeneousVolume);
+	float CalcLOD(const FSceneView& View, const IHeterogeneousVolumeInterface* HeterogeneousVolume);
 	bool IsHoldout(const IHeterogeneousVolumeInterface* HeterogeneousVolumeInterface);
 
 	const FProjectedShadowInfo* GetProjectedShadowInfo(const FVisibleLightInfo* VisibleLightInfo, int32 ShadowIndex);
diff --git a/Engine/Source/Runtime/Renderer/Private/HeterogeneousVolumes/HeterogeneousVolumesLiveShadingPipeline.cpp b/Engine/Source/Runtime/Renderer/Private/HeterogeneousVolumes/HeterogeneousVolumesLiveShadingPipeline.cpp
--- a/Engine/Source/Runtime/Renderer/Private/HeterogeneousVolumes/HeterogeneousVolumesLiveShadingPipeline.cpp
+++ b/Engine/Source/Runtime/Renderer/Private/HeterogeneousVolumes/HeterogeneousVolumesLiveShadingPipeline.cpp
@@ -53,6 +53,13 @@ static TAutoConsoleVariable<int32> CVarHeterogeneousVolumesShadowsUseCameraScene
 	ECVF_RenderThreadSafe
 );
 
+static TAutoConsoleVariable<int32> CVarHeterogeneousVolumesShadowsNearClippingDistance(
+	TEXT("r.HeterogeneousVolumes.Shadows.NearClippingDistance"),
+	1.0,
+	TEXT("Near clipping plane distance for shadow projection (Default = 1.0)"),
+	ECVF_RenderThreadSafe
+);
+
 #if 0
 static TAutoConsoleVariable<int32> CVarHeterogeneousVolumesBilinearInterpolation(
 	TEXT("r.HeterogeneousVolumes.BilinearInterpolation"),
@@ -106,6 +113,11 @@ namespace HeterogeneousVolumes
 	{
 		return CVarHeterogeneousVolumesShadowsUseCameraSceneDepth.GetValueOnRenderThread() != 0;
 	}
+
+	float GetShadowNearClippingDistance()
+	{
+		return FMath::Max(CVarHeterogeneousVolumesShadowsNearClippingDistance.GetValueOnRenderThread(), 0.1);
+	}
 }
 
 //-OPT: Remove duplicate bindings
@@ -1875,6 +1887,7 @@ bool RenderVolumetricShadowMapForLightForHeterogeneousVolumeWithLiveShading(
 	const FLightSceneInfo* LightSceneInfo,
 	const FVisibleLightInfo* VisibleLightInfo,
 	// Shadow data
+	HeterogeneousVolumes::FLODInfo LODInfo,
 	const FVector3f& TranslatedWorldOrigin,
 	int32 NumShadowMatrices,
 	FMatrix44f* TranslatedWorldToShadow,
@@ -1952,8 +1965,7 @@ bool RenderVolumetricShadowMapForLightForHeterogeneousVolumeWithLiveShading(
 		PassParameters->VoxelResolution = HeterogeneousVolumeInterface->GetVoxelResolution();
 
 		// Ray Data
-		//float LODFactor = HeterogeneousVolumes::CalcLODFactor(View, HeterogeneousVolumeInterface);
-		float LODFactor = 1.0f;
+		float LODFactor = HeterogeneousVolumes::CalcLODFactor(LODInfo, HeterogeneousVolumeInterface);
 		PassParameters->ShadowStepSize = HeterogeneousVolumes::GetShadowStepSize();
 		PassParameters->ShadowStepFactor = HeterogeneousVolumeInterface->GetShadowStepFactor() * LODFactor;
 		PassParameters->MaxTraceDistance = HeterogeneousVolumes::GetMaxTraceDistance();
@@ -2068,21 +2080,6 @@ bool RenderVolumetricShadowMapForLightWithLiveShading(
 		return false;
 	}
 
-	// Adjust shadow resolution based on minimum MipLevel
-	float LODValue = FMath::CeilLogTwo(ShadowMapResolution.X);
-	for (auto VolumetricMeshBatch : HeterogeneousVolumesMeshBatches)
-	{
-		int32 VolumeCount = VolumetricMeshBatch.Mesh->Elements.Num();
-		for (int32 VolumeIndex = 0; VolumeIndex < VolumeCount; ++VolumeIndex)
-		{
-			const IHeterogeneousVolumeInterface* HeterogeneousVolumeInterface = (IHeterogeneousVolumeInterface*)VolumetricMeshBatch.Mesh->Elements[VolumeIndex].UserData;
-			//LODValue = FMath::Min(LODValue, HeterogeneousVolumes::CalcLOD(View, HeterogeneousVolumeInterface));
-			LODValue = FMath::Min(LODValue, 1.0f);
-		}
-	}
-	float LODFactor = HeterogeneousVolumes::CalcLODFactor(LODValue);
-	ShadowMapResolution /= LODFactor;
-
 	// Build shadow transform
 	NumShadowMatrices = ProjectedShadowInfo->OnePassShadowViewProjectionMatrices.Num();
 	FMatrix44f ShadowToTranslatedWorld[6];
@@ -2115,7 +2112,7 @@ bool RenderVolumetricShadowMapForLightWithLiveShading(
 			PI / 4.0f,
 			ShadowMapResolution.X,
 			ShadowMapResolution.Y,
-			1.0,
+			HeterogeneousVolumes::GetShadowNearClippingDistance(),
 			LightSceneInfo->Proxy->GetRadius()
 		);
 
@@ -2176,6 +2173,34 @@ bool RenderVolumetricShadowMapForLightWithLiveShading(
 	float W = -FVector3f::DotProduct(TranslatedWorldOrigin, FVector3f(LightDirection));
 	TranslatedWorldPlane = FVector4f(LightDirection.X, LightDirection.Y, LightDirection.Z, W);
 
+	FVector PreViewTranslation = View.ViewMatrices.GetPreViewTranslation();
+	HeterogeneousVolumes::FLODInfo LODInfo;
+	LODInfo.WorldSceneBounds = WorldVolumeBounds;
+	LODInfo.WorldOrigin = FVector(TranslatedWorldOrigin) - PreViewTranslation;
+	LODInfo.ViewRect = FIntRect(0, 0, ShadowMapResolution.X, ShadowMapResolution.Y);
+
+	FMatrix WorldToTranslatedWorldMatrix = FTranslationMatrix(PreViewTranslation);
+	FMatrix WorldToShadowMatrix = WorldToTranslatedWorldMatrix * FMatrix(TranslatedWorldToShadow[0]);
+	GetViewFrustumBounds(LODInfo.WorldShadowFrustum, WorldToShadowMatrix, true);
+	LODInfo.FOV = PI / 4.0f;
+	LODInfo.NearClippingDistance = HeterogeneousVolumes::GetShadowNearClippingDistance();
+	LODInfo.DownsampleFactor = 1.0;
+	LODInfo.bIsPerspective = (LightType != LightType_Directional);
+
+	// Adjust shadow resolution based on minimum MipLevel
+	float LODValue = FMath::CeilLogTwo(ShadowMapResolution.X);
+	for (auto VolumetricMeshBatch : HeterogeneousVolumesMeshBatches)
+	{
+		int32 VolumeCount = VolumetricMeshBatch.Mesh->Elements.Num();
+		for (int32 VolumeIndex = 0; VolumeIndex < VolumeCount; ++VolumeIndex)
+		{
+			const IHeterogeneousVolumeInterface* HeterogeneousVolumeInterface = (IHeterogeneousVolumeInterface*)VolumetricMeshBatch.Mesh->Elements[VolumeIndex].UserData;
+			LODValue = FMath::Min(LODValue, HeterogeneousVolumes::CalcLOD(LODInfo, HeterogeneousVolumeInterface));
+		}
+	}
+	float LODFactor = HeterogeneousVolumes::CalcLODFactor(LODValue);
+	ShadowMapResolution /= LODFactor;
+
 	// Iterate over shadow-casting volumes
 	bool bHasShadowCastingVolume = false;
 	if (!HeterogeneousVolumesMeshBatches.IsEmpty())
@@ -2205,6 +2230,7 @@ bool RenderVolumetricShadowMapForLightWithLiveShading(
 			LightSceneInfo,
 			VisibleLightInfo,
 			// Shadow Info
+			LODInfo,
 			TranslatedWorldOrigin,
 			NumShadowMatrices,
 			TranslatedWorldToShadow,
@@ -2240,6 +2266,7 @@ bool RenderVolumetricShadowMapForLightWithLiveShading(
 				LightSceneInfo,
 				VisibleLightInfo,
 				// Shadow Info
+				LODInfo,
 				TranslatedWorldOrigin,
 				NumShadowMatrices,
 				TranslatedWorldToShadow,
@@ -2509,6 +2536,18 @@ void RenderAdaptiveVolumetricCameraMapWithLiveShading(
 			TEXT("HeterogeneousVolume.VolumetricShadowLinkedListBuffer")
 		);
 
+		HeterogeneousVolumes::FLODInfo LODInfo;
+		LODInfo.WorldSceneBounds = WorldVolumeBounds;
+		LODInfo.WorldOrigin = FVector(TranslatedWorldOrigin);
+		LODInfo.ViewRect = FIntRect(0, 0, ShadowMapResolution.X, ShadowMapResolution.Y);
+		FVector PreViewTranslation = View.ViewMatrices.GetPreViewTranslation();
+		FMatrix WorldToTranslatedWorldMatrix = FTranslationMatrix(PreViewTranslation);
+		GetViewFrustumBounds(LODInfo.WorldShadowFrustum, WorldToTranslatedWorldMatrix * FMatrix(TranslatedWorldToShadow[0]), true);
+		LODInfo.FOV = FOV;
+		LODInfo.NearClippingDistance = HeterogeneousVolumes::GetShadowNearClippingDistance();
+		LODInfo.DownsampleFactor = HeterogeneousVolumes::GetCameraDownsampleFactor();
+		LODInfo.bIsPerspective = true;
+
 		// Build a camera shadow for one volume
 		int32 LightType = 0;
 		FLightSceneInfo* LightSceneInfo = nullptr;
@@ -2525,6 +2564,7 @@ void RenderAdaptiveVolumetricCameraMapWithLiveShading(
 			LightSceneInfo,
 			VisibleLightInfo,
 			// Shadow Info
+			LODInfo,
 			TranslatedWorldOrigin,
 			NumShadowMatrices,
 			TranslatedWorldToShadow,
@@ -2561,6 +2601,7 @@ void RenderAdaptiveVolumetricCameraMapWithLiveShading(
 				LightSceneInfo,
 				VisibleLightInfo,
 				// Shadow Info
+				LODInfo,
 				TranslatedWorldOrigin,
 				NumShadowMatrices,
 				TranslatedWorldToShadow,
