Fixed crash in certain RHI's when mesh shader bindings contain a null texture / SRV pointer. This can occur when mesh shader bindings are incomplete and combined with RDG pass parameter bindings.

diff --git a/Engine/Source/Runtime/Renderer/Private/MeshPassProcessor.cpp b/Engine/Source/Runtime/Renderer/Private/MeshPassProcessor.cpp
--- a/Engine/Source/Runtime/Renderer/Private/MeshPassProcessor.cpp
+++ b/Engine/Source/Runtime/Renderer/Private/MeshPassProcessor.cpp
@@ -250,12 +250,20 @@ void FReadOnlyMeshDrawSingleShaderBindings::SetShaderBindings(
 		if (SRVType[TypeByteIndex] & (1 << TypeBitIndex))
 		{
 			FRHIShaderResourceView* SRV = (FRHIShaderResourceView*)SRVBindings[SRVIndex];
-			SetSrvParameter(BatchedParameters, Parameter, SRV);
+
+			if (SRV)
+			{
+				SetSrvParameter(BatchedParameters, Parameter, SRV);
+			}
 		}
 		else
 		{
 			FRHITexture* Texture = (FRHITexture*)SRVBindings[SRVIndex];
-			SetTextureParameter(BatchedParameters, Parameter, Texture);
+
+			if (Texture)
+			{
+				SetTextureParameter(BatchedParameters, Parameter, Texture);
+			}
 		}
 	}
 
