[VSM] Adapt debug viewmodes to better show local lights. Visualization modes now show a composite of all lights by default, and change to showing individual lights when one is selected in editor or by name.
Two commands were added to visualize the previous/next light (r.shadow.virtual.visualize.nextlight).
When VSM visualization is enabled, one pass projection is now turned off, as it is incompatible with the debug output in the projection shader (for now)

diff --git a/Engine/Shaders/Private/VirtualShadowMaps/VirtualShadowMapDebug.usf b/Engine/Shaders/Private/VirtualShadowMaps/VirtualShadowMapDebug.usf
--- a/Engine/Shaders/Private/VirtualShadowMaps/VirtualShadowMapDebug.usf
+++ b/Engine/Shaders/Private/VirtualShadowMaps/VirtualShadowMapDebug.usf
@@ -154,4 +154,35 @@ void DebugVisualizeVirtualSmCS(uint2 PixelPos : SV_DispatchThreadID)
 	}
 }
 
+int VisualizeVirtualShadowMapId;
+float VisualizeNaniteOverdrawScale;
+Texture2D DebugTexture;
+
+float4 TonemapProjectionDebugTexturePS(
+	noperspective float4 UVAndScreenPos : TEXCOORD0,
+	float4 Position : SV_POSITION
+) : SV_Target0
+{
+	bool bOneLightVisible = VisualizeVirtualShadowMapId >= 0;
 
+	float4 Output = DebugTexture.Load(uint3(Position.xy, 0));
+	if (VisualizeModeId == VIRTUAL_SHADOW_MAP_VISUALIZE_SMRT_RAY_COUNT)
+	{
+		if (Output.a > 0)
+		{
+			float TotalRayCount = Output.r;
+			float TotalMaxRayCount = Output.g;
+			float IntensityMultiplier = bOneLightVisible ? 1.0 : smoothstep(0, 64, TotalMaxRayCount)+0.5;
+			Output.rgb = GreenToRedTurbo(IntensityMultiplier * TotalRayCount / TotalMaxRayCount);
+		}
+	}
+	else if (VisualizeModeId == VIRTUAL_SHADOW_MAP_VISUALIZE_NANITE_OVERDRAW)
+	{
+		const float OverdrawScale = clamp(VisualizeNaniteOverdrawScale / 100.0f, 0.0f, 1.0f);
+		const float OverdrawCount = Output.r;
+		const float OverdrawColor = 1 - exp2(-OverdrawCount * OverdrawScale);
+		Output = float4(ColorMapInferno(OverdrawColor), 0.0);
+		Output.a = smoothstep(-0.1, 0.1, dot(Output.rgb, 1.0/3.0));
+	}
+	return Output;
+}
diff --git a/Engine/Shaders/Private/VirtualShadowMaps/VirtualShadowMapProjection.usf b/Engine/Shaders/Private/VirtualShadowMaps/VirtualShadowMapProjection.usf
--- a/Engine/Shaders/Private/VirtualShadowMaps/VirtualShadowMapProjection.usf
+++ b/Engine/Shaders/Private/VirtualShadowMaps/VirtualShadowMapProjection.usf
@@ -146,8 +146,8 @@ RWTexture2D< float2 > OutShadowFactor;
 StructuredBuffer< FPhysicalPageMetaData > PhysicalPageMetaData;
 RWTexture2D< float4 > OutVisualize;
 int VisualizeModeId;
+int bVisualizeCachedPagesOnly;
 int VisualizeVirtualShadowMapId;
-float VisualizeNaniteOverdrawScale;
 
 // Type of input data consume by the page allocation (i.e., data read from the source buffer: Gbuffer, HairStrands data, ...)
 #define INPUT_TYPE_GBUFFER 0
@@ -443,20 +443,29 @@ float3 GetVirtualPageColor(uint2 VirtualTexelAddress)
 	return IntToColor(vPage.x + (vPage.y << 10U));
 }
 
+float Grayscale(float3 SourceClr)
+{
+	return 0.299 * SourceClr.r + 0.587 * SourceClr.g + 0.114 * SourceClr.b;
+}
+
 void OutputVisualize(
+	int VirtualShadowMapId,
 	uint2 PixelPos,
 	float3 TranslatedWorldPosition,	// Unmodified, raw position before any biasing
 	FVirtualShadowMapSampleResult Result)
 {
+	bool bOneLightVisible = VisualizeVirtualShadowMapId >= 0;
+	if (bOneLightVisible && VisualizeVirtualShadowMapId != VirtualShadowMapId)
+	{
+		return;
+	}
+
 	// NOTE: Important to *not* write output if it isn't a recognized mode, as a different
 	// pass may write that output instead.
-	float3 Output = float3( 1, 0, 1 );
-	bool bWriteOutput = false;
+	float4 Output = float4( 1, 0, 1, 0 );
 
 	if ( Result.bValid )
 	{
-		bWriteOutput = true;
-
 		const float3 vPageColor = GetVirtualPageColor(Result.VirtualTexelAddress);
 
 		uint2 pPage = Result.PhysicalTexelAddress.xy >> VSM_LOG2_PAGE_SIZE;
@@ -468,49 +477,105 @@ void OutputVisualize(
 			pPageMeta = PhysicalPageMetaData[VSMPhysicalPageAddressToIndex(pPage)];
 		}
 
+		float4 Accumulator = OutVisualize[PixelPos];
+
 		if ( VisualizeModeId == VIRTUAL_SHADOW_MAP_VISUALIZE_SHADOW_FACTOR )
 		{
-			float3 Lit = float3(1,1,0);
-			float3 Shadow = float3(0,0,1);
+			float4 Lit = float4(1,1,0,1);
+			float4 Shadow = float4(0,0,0.5,1);
 			Output = lerp(Shadow, Lit, Result.ShadowFactor);
+
+			if (!bOneLightVisible)
+			{
+				// todo: figure out which is the least confusing blend mode
+				Output = Accumulator + Output*0.25;
+				Output.b = min(Output.b, 0.12); // a bit of blue helps retain contrast in shadowed areas
+			}
 		}
 		else if ( VisualizeModeId == VIRTUAL_SHADOW_MAP_VISUALIZE_CLIPMAP_OR_MIP )
 		{
+			// todo: consider separating out clipmap
+			// todo: consider changing color mapping so near/far is easier to tell apart
 			float3 Color = IntToColor( Result.ClipmapOrMipLevel );
-			Output = 0.8f*Color + 0.2f*Result.ShadowFactor.xxx;
+			Output = float4(0.8f*Color + 0.2f*Result.ShadowFactor.xxx, 1.0);
+
+			if (!bOneLightVisible)
+			{
+				Output = Accumulator + Output*0.25;
+				Output.a = 1.0;
+			}
 		}
 		else if ( VisualizeModeId == VIRTUAL_SHADOW_MAP_VISUALIZE_VIRTUAL_PAGE && bValidPageAddress )
 		{
 			float3 MipColor = IntToColor( Result.ClipmapOrMipLevel );
-			Output = 0.4f*vPageColor + 0.4f*MipColor + 0.2f*Result.ShadowFactor.xxx;
+			Output = float4(0.4f*vPageColor + 0.4f*MipColor + 0.2f*Result.ShadowFactor.xxx, 1.0);
+
+			if (!bOneLightVisible)
+			{
+				Output = Accumulator + Output*0.25;
+				Output.a = 1.0;
+			}
 		}		
 		else if ( VisualizeModeId == VIRTUAL_SHADOW_MAP_VISUALIZE_SMRT_RAY_COUNT )
 		{
-			Output = GreenToRedTurbo(float(Result.RayCount) / float(GetSMRTRayCount()));
+			Output.xy = Accumulator.xy + float2(Result.RayCount, GetSMRTRayCount());
+			Output.zw = float2(0, 1.0);
 		}
-		else if ( VisualizeModeId == VIRTUAL_SHADOW_MAP_VISUALIZE_CACHED_PAGE && bValidPageAddress )
+		else if ( VisualizeModeId == VIRTUAL_SHADOW_MAP_VISUALIZE_CACHED_PAGE && bValidPageAddress)
 		{
 			bool bUncachedDynamic = (pPageMeta.Flags & VSM_FLAG_DYNAMIC_UNCACHED) != 0;
 			bool bUncachedStatic = (pPageMeta.Flags & VSM_FLAG_STATIC_UNCACHED) != 0;
 			bool bForceCached = (pPageMeta.Flags & VSM_EXTENDED_FLAG_FORCE_CACHED) != 0;
-			// Red = both uncached, blue = only static cached, green otherwise
-			float3 CacheColor = float3(0, 1, 0);
+			
+			float CacheAlpha = bOneLightVisible ? 0.9 : 0.25;
+			float VPageAlpha = 0;
+			float ShadowAlpha = 0.1;
+
+			float4 CacheColor = float4(0, 0, 0, 0);
+			
+			// Red = both uncached, blue = only static cached
 			if (bUncachedDynamic)
 			{
-				CacheColor = bUncachedStatic ? float3(1, 0, 0) : float3(0, 0, 1);
+				if (!bVisualizeCachedPagesOnly)
+				{
+					CacheColor = bUncachedStatic ? float4(1, 0, 0, CacheAlpha) : float4(0, 0, 1, CacheAlpha);
+				}
 			}
 			else if (bForceCached)
 			{
-				// TODO: Maybe only enable this debug color when "advanced" visualizations are enabled?
-				// Could be confusing to content developers.
-				CacheColor = float3(0.75, 1, 0);
+				if (!bVisualizeCachedPagesOnly)
+				{
+					// TODO: Maybe only enable this debug color when "advanced" visualizations are enabled?
+					// Could be confusing to content developers.
+					CacheColor = float4(0.75, 1, 0, CacheAlpha);
+				}
+			}
+			else if (bOneLightVisible || bVisualizeCachedPagesOnly)
+			{
+				CacheColor = float4(0, 1, 0, CacheAlpha); // Green if cached
+				VPageAlpha = 0.1;
+			}
+
+			if (CacheColor.a > 0)
+			{
+				Output.rgb = CacheColor.a*CacheColor.rgb + VPageAlpha*vPageColor + ShadowAlpha*Result.ShadowFactor.xxx;
+				Output = float4(Output + Accumulator.rgb, 0.9f);
+			}
+			else
+			{
+				Output = float4(0, 0, 0, 0);
 			}
-			Output = 0.55f*CacheColor + 0.25f*vPageColor + 0.2f*Result.ShadowFactor.xxx;
 		}
 		else if (VisualizeModeId == VIRTUAL_SHADOW_MAP_VISUALIZE_DIRTY_PAGE)
 		{
 			float3 PageColor = vPageColor * 0.4f;
 			float3 DirtyColor = PageColor;
+
+			if (!bOneLightVisible)
+			{
+				PageColor = Grayscale(PageColor);
+			}
+
 			if (bValidPageAddress)
 			{
 				bool bDirty = (pPageMeta.Flags & VSM_EXTENDED_FLAG_DIRTY) != 0U;
@@ -518,10 +583,17 @@ void OutputVisualize(
 				if (bDirty)
 				{
 					DirtyColor.x = 1.0f;
+
+					PageColor = lerp(PageColor, DirtyColor, 0.8f);
 				}
-				PageColor = lerp(PageColor, DirtyColor, 0.8f);
 			}
-			Output = 0.8f * PageColor + 0.2f * Result.ShadowFactor.xxx;
+			Output = float4(PageColor, 1.0);
+
+			if (!bOneLightVisible)
+			{
+				Output = Accumulator + Output*0.25;
+				Output.a = 1.0;
+			}
 		}
 		else if (VisualizeModeId == VIRTUAL_SHADOW_MAP_VISUALIZE_GPU_INVALIDATED_PAGE)
 		{
@@ -540,32 +612,67 @@ void OutputVisualize(
 					CacheColor.y = 1.0f;
 				}
 			}
-			Output = 0.55f*CacheColor + 0.25f*vPageColor + 0.2f*Result.ShadowFactor.xxx;
+
+			if (bOneLightVisible)
+			{
+				Output = float4(0.55f*CacheColor + 0.25f*vPageColor + 0.2f*Result.ShadowFactor.xxx, 1.0);
+			}
+			else
+			{
+				Output = float4(0.55f*CacheColor + 0.25f * Grayscale(vPageColor) + 0.2f*Result.ShadowFactor.xxx, 1.0);
+
+				Output = Accumulator + Output*0.25;
+				Output.a = 1.0;
+			}
 		}
 		else if ( VisualizeModeId == VIRTUAL_SHADOW_MAP_VISUALIZE_MERGED_PAGE)
 		{
 			bool bMerged = ((pPageMeta.Flags & VSM_EXTENDED_FLAG_DIRTY) != 0U) && (pPageMeta.Flags & VSM_EXTENDED_FLAG_VIEW_UNCACHED) == 0U;
-			float3 Color = float3(0, 1, 0);
-			if (bMerged)
+			if (bOneLightVisible)
+			{
+				float3 Color = float3(0, 1, 0);
+				if (bMerged)
+				{
+					Color = float3(1, 0, 0);
+				}
+				Output = float4(0.55f * Color + 0.25f * vPageColor + 0.2f * Result.ShadowFactor.xxx, 1.0);
+			}
+			else
 			{
-				Color = float3(1, 0, 0);
+				float3 Color = float3(0, 0, 0);
+				if (bMerged)
+				{
+					Color = float3(1, 0, 0);
+				}
+				Output = float4(0.55f * Color + 0.25f * Grayscale(vPageColor) + 0.2f * Result.ShadowFactor.xxx, 1.0);
+
+				Output = Accumulator + Output*0.25;
+				Output.a = 1.0;
 			}
-			Output = 0.55f * Color + 0.25f * vPageColor + 0.2f * Result.ShadowFactor.xxx;
 		}
 		else if ( VisualizeModeId == VIRTUAL_SHADOW_MAP_VISUALIZE_GENERAL_DEBUG )
 		{
-			Output = 0.8f*Result.GeneralDebug + 0.2f*Result.ShadowFactor.xxx;
+			Output = float4(0.8f*Result.GeneralDebug + 0.2f*Result.ShadowFactor.xxx, 1.0);
+			
+			if (!bOneLightVisible)
+			{
+				Output = Accumulator + Output*0.25;
+				Output.a = 1.0;
+			}
 		}
 		else if ( VisualizeModeId == VIRTUAL_SHADOW_MAP_VISUALIZE_NANITE_OVERDRAW )
 		{
 			// This is just to get the raw physical address
 			FVirtualShadowMapSampleResult RawResult = SampleVirtualShadowMapTranslatedWorld(
-				VisualizeVirtualShadowMapId, TranslatedWorldPosition);
+				VirtualShadowMapId, TranslatedWorldPosition);
+
+			const float OverdrawCount = float(VirtualShadowMap.PhysicalPagePool.Load(uint4(RawResult.PhysicalTexelAddress.xy, 2, 0))); // See FVisBufferPixel::WriteOverdraw
+			Output = float4(OverdrawCount, 0, 0, 1.0);
 
-			const float OverdrawScale = clamp(VisualizeNaniteOverdrawScale / 100.0f, 0.0f, 1.0f);
-			const float OverdrawCount = float(VirtualShadowMap.PhysicalPagePool.Load(uint4(RawResult.PhysicalTexelAddress.xy, 2, 0)));
-			const float OverdrawColor = 1 - exp2(-OverdrawCount * OverdrawScale);
-			Output = ColorMapInferno(OverdrawColor);
+			if (!bOneLightVisible)
+			{
+				Output.r += Accumulator.r;
+			}
 		}
 		/* TODO: Enable this and add to menu
 		else if ( VisualizeModeId == VIRTUAL_SHADOW_MAP_VISUALIZE_OCCLUDER_DISTANCE )
@@ -573,15 +680,11 @@ void OutputVisualize(
 			Output = (Result.OccluderDistance / 1000.0f).xxx;
 		}
 		*/
-		else
-		{
-			bWriteOutput = false;
-		}
 	}
 
-	if (bWriteOutput)
+	if (Output.a > 0)
 	{
-		OutVisualize[PixelPos] = float4(Output, 1.0f);
+		OutVisualize[PixelPos] = Output;
 	}
 }
 
@@ -636,8 +739,6 @@ void VirtualShadowMapProjection(
 		SubsurfaceOpacityMFP = GetSubsurfaceOpacityFromGbuffer(PixelPos);
 	}
 
-	FVirtualShadowMapSampleResult VisualizeResult = InitVirtualShadowMapSampleResult();
-
 #if ONE_PASS_PROJECTION
 	const FCulledLightsGridHeader LightGridHeader = VirtualShadowMapGetLightsGrid(LocalPixelPos, SceneDepth);	
 	// We can only handle so many lights in our output encoding right now, so no purpose in computing more
@@ -671,10 +772,9 @@ void VirtualShadowMapProjection(
 			FilterVirtualShadowMapSampleResult(PixelPos, Result);
 			PackShadowMask(ShadowMask, Result.ShadowFactor, Index);
 
-			if (VisualizeVirtualShadowMapId == VirtualShadowMapId)
-			{
-				VisualizeResult = Result;
-			}
+#if VISUALIZE_OUTPUT
+			OutputVisualize( VirtualShadowMapId, PixelPos - View.ViewRectMin.xy, TranslatedWorldPosition, Result );
+#endif // VISUALIZE_OUTPUT
 		}
 	}
 
@@ -709,14 +809,10 @@ void VirtualShadowMapProjection(
 
 		OutShadowFactor[ PixelPos ] = Result.ShadowFactor.xx;
 
-		if (VisualizeVirtualShadowMapId == VirtualShadowMapId)
-		{
-			VisualizeResult = Result;
-		}
-	}
-#endif // ONE_PASS_PROJECTION
-
 #if VISUALIZE_OUTPUT
-	OutputVisualize( PixelPos - View.ViewRectMin.xy, TranslatedWorldPosition, VisualizeResult );
+		OutputVisualize( VirtualShadowMapId, PixelPos - View.ViewRectMin.xy, TranslatedWorldPosition, Result );
 #endif // VISUALIZE_OUTPUT
+
+	}
+#endif // ONE_PASS_PROJECTION
 }
diff --git a/Engine/Source/Editor/CommonMenuExtensions/Private/VirtualShadowMapVisualizationMenuCommands.cpp b/Engine/Source/Editor/CommonMenuExtensions/Private/VirtualShadowMapVisualizationMenuCommands.cpp
--- a/Engine/Source/Editor/CommonMenuExtensions/Private/VirtualShadowMapVisualizationMenuCommands.cpp
+++ b/Engine/Source/Editor/CommonMenuExtensions/Private/VirtualShadowMapVisualizationMenuCommands.cpp
@@ -69,6 +69,9 @@ void FVirtualShadowMapVisualizationMenuCommands::BuildCommandMap()
 			break;
 		}
 	}
+
+	UI_COMMAND(VisualizeNextLightCommand, "Visualize next light", "r.Shadow.Virtual.Visualize.NextLight", EUserInterfaceActionType::Button, FInputChord());
+	UI_COMMAND(VisualizePrevLightCommand, "Visualize previous light", "r.Shadow.Virtual.Visualize.PrevLight", EUserInterfaceActionType::Button, FInputChord());
 }
 
 void FVirtualShadowMapVisualizationMenuCommands::BuildVisualisationSubMenu(FMenuBuilder& Menu)
@@ -135,6 +138,29 @@ void FVirtualShadowMapVisualizationMenuCommands::BindCommands(FUICommandList& Co
 			FIsActionChecked::CreateStatic(&FVirtualShadowMapVisualizationMenuCommands::IsVirtualShadowMapVisualizationModeSelected, Client.ToWeakPtr(), Record.Name)
 		);
 	}
+
+	CommandList.MapAction(
+		VisualizeNextLightCommand,
+		FExecuteAction::CreateLambda([]()
+		{
+			IConsoleCommand* Command = IConsoleManager::Get().FindConsoleObject(TEXT("r.Shadow.Virtual.Visualize.NextLight"))->AsCommand();
+			TArray<FString> Args;
+			Command->Execute(Args, nullptr, *GLog);
+		}),
+		FCanExecuteAction(),
+		EUIActionRepeatMode::RepeatEnabled
+	);
+	CommandList.MapAction(
+		VisualizePrevLightCommand,
+		FExecuteAction::CreateLambda([]()
+		{
+			IConsoleCommand* Command = IConsoleManager::Get().FindConsoleObject(TEXT("r.Shadow.Virtual.Visualize.PrevLight"))->AsCommand();
+			TArray<FString> Args;
+			Command->Execute(Args, nullptr, *GLog);
+		}),
+		FCanExecuteAction(),
+		EUIActionRepeatMode::RepeatEnabled
+	);
 }
 
 void FVirtualShadowMapVisualizationMenuCommands::ChangeVirtualShadowMapVisualizationMode(TWeakPtr<FEditorViewportClient> WeakClient, FName InName)
diff --git a/Engine/Source/Editor/CommonMenuExtensions/Public/VirtualShadowMapVisualizationMenuCommands.h b/Engine/Source/Editor/CommonMenuExtensions/Public/VirtualShadowMapVisualizationMenuCommands.h
--- a/Engine/Source/Editor/CommonMenuExtensions/Public/VirtualShadowMapVisualizationMenuCommands.h
+++ b/Engine/Source/Editor/CommonMenuExtensions/Public/VirtualShadowMapVisualizationMenuCommands.h
@@ -66,4 +66,6 @@ private:
 
 private:
 	TVirtualShadowMapVisualizationModeCommandMap CommandMap;
+	TSharedPtr<FUICommandInfo> VisualizeNextLightCommand;
+	TSharedPtr<FUICommandInfo> VisualizePrevLightCommand;
 };
diff --git a/Engine/Source/Runtime/Engine/Private/VirtualShadowMapVisualizationData.cpp b/Engine/Source/Runtime/Engine/Private/VirtualShadowMapVisualizationData.cpp
--- a/Engine/Source/Runtime/Engine/Private/VirtualShadowMapVisualizationData.cpp
+++ b/Engine/Source/Runtime/Engine/Private/VirtualShadowMapVisualizationData.cpp
@@ -114,6 +114,26 @@ void FVirtualShadowMapVisualizationData::ConfigureConsoleCommand()
 		TEXT(""),
 		*ConsoleDocumentationVisualizationMode,
 		ECVF_Cheat);
+
+	IConsoleManager::Get().RegisterConsoleCommand(
+		TEXT("r.Shadow.Virtual.Visualize.NextLight"),
+		TEXT("Set `r.Shadow.Virtual.Visualize.LightName` to the next light."),
+		FConsoleCommandWithArgsAndOutputDeviceDelegate::CreateLambda(
+			[this](const TArray<FString>& Arguments, FOutputDevice& Output)
+			{
+				SelectAdjacentVisualizeLight = 1;
+			}),
+		ECVF_Cheat);
+
+	IConsoleManager::Get().RegisterConsoleCommand(
+		TEXT("r.Shadow.Virtual.Visualize.PrevLight"),
+		TEXT("Set `r.Shadow.Virtual.Visualize.LightName` to the previous light."),
+		FConsoleCommandWithArgsAndOutputDeviceDelegate::CreateLambda(
+			[this](const TArray<FString>& Arguments, FOutputDevice& Output)
+			{
+				SelectAdjacentVisualizeLight = -1;
+			}),
+		ECVF_Cheat);
 }
 
 void FVirtualShadowMapVisualizationData::SetActiveMode(int32 ModeID, const FName& ModeName)
diff --git a/Engine/Source/Runtime/Engine/Public/VirtualShadowMapVisualizationData.h b/Engine/Source/Runtime/Engine/Public/VirtualShadowMapVisualizationData.h
--- a/Engine/Source/Runtime/Engine/Public/VirtualShadowMapVisualizationData.h
+++ b/Engine/Source/Runtime/Engine/Public/VirtualShadowMapVisualizationData.h
@@ -71,6 +71,8 @@ public:
 		return TEXT("r.Shadow.Virtual.Visualize");
 	}
 
+	int SelectAdjacentVisualizeLight = 0;
+
 private:
 	/** Internal helper function for creating the VirtualShadowMap visualization system console commands. */
 	void ConfigureConsoleCommand();
diff --git a/Engine/Source/Runtime/Renderer/Private/LightGridInjection.cpp b/Engine/Source/Runtime/Renderer/Private/LightGridInjection.cpp
--- a/Engine/Source/Runtime/Renderer/Private/LightGridInjection.cpp
+++ b/Engine/Source/Runtime/Renderer/Private/LightGridInjection.cpp
@@ -98,7 +98,7 @@ FAutoConsoleVariableRef CVarLightCullingMaxDistanceOverride(
 	ECVF_RenderThreadSafe
 );
 
-extern TAutoConsoleVariable<int32> CVarVirtualShadowOnePassProjection;
+extern bool IsVSMOnePassProjectionEnabled(const FEngineShowFlags& ShowFlags);
 
 bool ShouldVisualizeLightGrid()
 {
@@ -950,7 +950,7 @@ FComputeLightGridOutput FDeferredShadingSceneRenderer::GatherLightsAndComputeLig
 	FComputeLightGridOutput Result = {};
 
 	bool bShadowedLightsInClustered = ShouldUseClusteredDeferredShading()
-		&& CVarVirtualShadowOnePassProjection.GetValueOnRenderThread()
+		&& IsVSMOnePassProjectionEnabled(ViewFamily.EngineShowFlags)
 		&& VirtualShadowMapArray.IsEnabled();
 
 	GatherAndSortLights(SortedLightSet, bShadowedLightsInClustered);
diff --git a/Engine/Source/Runtime/Renderer/Private/Shadows/ShadowSceneRenderer.cpp b/Engine/Source/Runtime/Renderer/Private/Shadows/ShadowSceneRenderer.cpp
--- a/Engine/Source/Runtime/Renderer/Private/Shadows/ShadowSceneRenderer.cpp
+++ b/Engine/Source/Runtime/Renderer/Private/Shadows/ShadowSceneRenderer.cpp
@@ -78,6 +78,13 @@ static TAutoConsoleVariable<float> CVarResolutionLodBiasLocalMoving(
 	ECVF_Scalability | ECVF_RenderThreadSafe
 );
 
+bool IsVSMOnePassProjectionEnabled(const FEngineShowFlags& ShowFlags)
+{
+	return CVarVirtualShadowOnePassProjection.GetValueOnAnyThread() 
+		// Debug outputs from projection pass do not support one pass projection
+		&& (ShowFlags.VisualizeVirtualShadowMap == 0);
+}
+
 DECLARE_DWORD_COUNTER_STAT(TEXT("VSM Total Raster Bins"), STAT_VSMNaniteBasePassTotalRasterBins, STATGROUP_ShadowRendering);
 DECLARE_DWORD_COUNTER_STAT(TEXT("VSM Visible Raster Bins"), STAT_VSMNaniteBasePassVisibleRasterBins, STATGROUP_ShadowRendering);
 
@@ -477,7 +484,7 @@ void FShadowSceneRenderer::RenderVirtualShadowMapProjectionMaskBits(
 	// VSM one pass projection (done first as it may be needed by clustered shading)
 	bShouldUseVirtualShadowMapOnePassProjection =
 		VirtualShadowMapArray.IsAllocated() &&
-		CVarVirtualShadowOnePassProjection.GetValueOnRenderThread();
+		IsVSMOnePassProjectionEnabled(SceneRenderer.ViewFamily.EngineShowFlags);
 
 	if (!VirtualShadowMapArray.HasAnyShadowData())
 	{
diff --git a/Engine/Source/Runtime/Renderer/Private/VirtualShadowMaps/VirtualShadowMapArray.cpp b/Engine/Source/Runtime/Renderer/Private/VirtualShadowMaps/VirtualShadowMapArray.cpp
--- a/Engine/Source/Runtime/Renderer/Private/VirtualShadowMaps/VirtualShadowMapArray.cpp
+++ b/Engine/Source/Runtime/Renderer/Private/VirtualShadowMaps/VirtualShadowMapArray.cpp
@@ -282,11 +282,19 @@ FAutoConsoleCommand CmdDumpVSMLightNames(
 	FConsoleCommandDelegate::CreateStatic(DumpVSMLightNames)
 );
 
+UPTRINT GVirtualShadowMapLastSelectedVisualizeLightId = 0;
+UPTRINT GVirtualShadowMapVisualizeLightId = 0;
+bool GVirtualShadowMapVisualizeByLightId = false;
+
 FString GVirtualShadowMapVisualizeLightName;
 FAutoConsoleVariableRef CVarVisualizeLightName(
 	TEXT("r.Shadow.Virtual.Visualize.LightName"),
 	GVirtualShadowMapVisualizeLightName,
 	TEXT("Sets the name of a specific light to visualize (for developer use in non-shipping builds)"),
+	FConsoleVariableDelegate::CreateLambda([](IConsoleVariable* Variable)
+	{
+		GVirtualShadowMapVisualizeByLightId = false;
+	}),
 	ECVF_RenderThreadSafe
 );
 
@@ -1244,20 +1252,82 @@ void FVirtualShadowMapVisualizeLightSearch::CheckLight(const FLightSceneProxy* C
 		UE_LOG(LogRenderer, Display, TEXT("%s"), *CheckLightName);
 	}
 
-	// Fill out new sort key and compare to our best found so far
-	SortKey CheckKey;
-	CheckKey.Packed = 0;
-	CheckKey.Fields.bExactNameMatch = (CheckLightName == GVirtualShadowMapVisualizeLightName);
-	CheckKey.Fields.bPartialNameMatch = CheckKey.Fields.bExactNameMatch || CheckLightName.Contains(GVirtualShadowMapVisualizeLightName);
-	CheckKey.Fields.bSelected = CheckProxy->IsSelected();
-	CheckKey.Fields.bDirectionalLight = CheckProxy->GetLightType() == LightType_Directional;
-	CheckKey.Fields.bExists = 1;
+	const UPTRINT CheckProxyId = (UPTRINT)(CheckProxy->GetLightComponent());
+
+	const int SelectAdjacentVisualizeLight = GetVirtualShadowMapVisualizationData().SelectAdjacentVisualizeLight;
+	const bool SelectNextLight = SelectAdjacentVisualizeLight > 0;
+	const bool SelectPrevLight = SelectAdjacentVisualizeLight < 0;
+	GVirtualShadowMapVisualizeByLightId = GVirtualShadowMapVisualizeByLightId || SelectAdjacentVisualizeLight != 0;
+
+	// When the user clicks a light, visualize the selected light.
+	if (CheckProxy->IsSelected() && CheckProxyId != GVirtualShadowMapLastSelectedVisualizeLightId)
+	{
+		GVirtualShadowMapVisualizeByLightId = false;
+		GVirtualShadowMapLastSelectedVisualizeLightId = CheckProxyId;
+	}
+
+	if (GVirtualShadowMapVisualizeByLightId)
+	{
+		UPTRINT FoundAdjacentLightId = FoundProxy ? (UPTRINT)(FoundProxy->GetLightComponent()) : 0;
+		bool bIsFoundAdjacentLightInvalid = FoundAdjacentLightId == 0;
+
+		if (SelectPrevLight
+			// Light comes before current selection
+			&& CheckProxyId < GVirtualShadowMapVisualizeLightId
+			// Light comes after best match so far
+			&& FoundAdjacentLightId < CheckProxyId)
+		{
+			FoundProxy = CheckProxy;
+			FoundVirtualShadowMapId = CheckVirtualShadowMapId;
+		}
+		else if (SelectNextLight
+			// Light comes after current selection
+			&& GVirtualShadowMapVisualizeLightId < CheckProxyId
+			// Light comes before best match so far
+			&& (CheckProxyId < FoundAdjacentLightId || bIsFoundAdjacentLightInvalid))
+		{
+			FoundProxy = CheckProxy;
+			FoundVirtualShadowMapId = CheckVirtualShadowMapId;
+		}
+		else if (SelectAdjacentVisualizeLight == 0 
+			&& CheckProxyId == GVirtualShadowMapVisualizeLightId)
+		{
+			FoundProxy = CheckProxy;
+			FoundVirtualShadowMapId = CheckVirtualShadowMapId;
+		}
+	}
+	else
+	{
+		// Fill out new sort key and compare to our best found so far
+		SortKey CheckKey;
+		CheckKey.Packed = 0;
+		CheckKey.Fields.bExactNameMatch = (CheckLightName == GVirtualShadowMapVisualizeLightName);
+		CheckKey.Fields.bPartialNameMatch = CheckKey.Fields.bExactNameMatch || (!GVirtualShadowMapVisualizeLightName.IsEmpty() && CheckLightName.Contains(GVirtualShadowMapVisualizeLightName));
+		CheckKey.Fields.bSelected = CheckProxy->IsSelected();
+
+		if (CheckKey.Packed > FoundKey.Packed)		//-V547
+		{
+			FoundKey = CheckKey;
+			FoundProxy = CheckProxy;
+			FoundVirtualShadowMapId = CheckVirtualShadowMapId;
+		}
+	}
+#endif
+}
+
+void FVirtualShadowMapVisualizeLightSearch::ChooseLight()
+{
+#if !UE_BUILD_SHIPPING
+	const int SelectAdjacentVisualizeLight = GetVirtualShadowMapVisualizationData().SelectAdjacentVisualizeLight;
+	GetVirtualShadowMapVisualizationData().SelectAdjacentVisualizeLight = 0;
 
-	if (CheckKey.Packed > FoundKey.Packed)		//-V547
+	if (FoundProxy)
 	{
-		FoundKey = CheckKey;
-		FoundProxy = CheckProxy;
-		FoundVirtualShadowMapId = CheckVirtualShadowMapId;
+		GVirtualShadowMapVisualizeLightId = (UPTRINT)(FoundProxy->GetLightComponent());
+	}
+	else if (SelectAdjacentVisualizeLight) // selected past first or last light, select none
+	{
+		GVirtualShadowMapVisualizeLightId = 0;
 	}
 #endif
 }
@@ -1273,7 +1343,7 @@ static FRDGTextureRef CreateDebugVisualizationTexture(FRDGBuilder& GraphBuilder,
 
 	FRDGTextureDesc Desc = FRDGTextureDesc::Create2D(
 		Extent,
-		PF_R8G8B8A8,
+		PF_FloatRGBA,
 		FClearValueBinding(ClearColor),
 		TexCreate_ShaderResource | TexCreate_UAV);
 
@@ -1328,6 +1398,11 @@ void FVirtualShadowMapArray::UpdateVisualizeLight(
 			}
 		}
 	}
+
+	for (int32 ViewIndex = 0; ViewIndex < Views.Num(); ViewIndex++)
+	{
+		VisualizeLight[ViewIndex].ChooseLight();
+	}
 #endif
 }
 
@@ -3626,6 +3701,23 @@ class FDesaturatePS : public FGlobalShader
 	END_SHADER_PARAMETER_STRUCT()
 };
 IMPLEMENT_GLOBAL_SHADER(FDesaturatePS, "/Engine/Private/VirtualShadowMaps/Desaturate.usf", "DesaturatePS", SF_Pixel);
+
+class FTonemapProjectionDebugTexturePS : public FVirtualShadowMapPageManagementShader
+{
+	DECLARE_GLOBAL_SHADER(FTonemapProjectionDebugTexturePS);
+	SHADER_USE_PARAMETER_STRUCT(FTonemapProjectionDebugTexturePS, FVirtualShadowMapPageManagementShader);
+
+	BEGIN_SHADER_PARAMETER_STRUCT(FParameters, )
+		SHADER_PARAMETER_RDG_TEXTURE(Texture2D, DebugTexture)
+		SHADER_PARAMETER(uint32, VisualizeModeId)
+		SHADER_PARAMETER(int32, VisualizeVirtualShadowMapId)
+		SHADER_PARAMETER(float, VisualizeNaniteOverdrawScale)
+		RENDER_TARGET_BINDING_SLOTS()
+	END_SHADER_PARAMETER_STRUCT()
+};
+IMPLEMENT_GLOBAL_SHADER(FTonemapProjectionDebugTexturePS, "/Engine/Private/VirtualShadowMaps/VirtualShadowMapDebug.usf", "TonemapProjectionDebugTexturePS", SF_Pixel);
+
+extern int32 GNaniteVisualizeOverdrawScale;
 #endif //!UE_BUILD_SHIPPING
 
 FScreenPassTexture FVirtualShadowMapArray::AddVisualizePass(FRDGBuilder& GraphBuilder, const FViewInfo& View, int32 ViewIndex, EVSMVisualizationPostPass Pass, FScreenPassTexture& SceneColor, FScreenPassRenderTarget& OverrideOutput)
@@ -3651,7 +3743,7 @@ FScreenPassTexture FVirtualShadowMapArray::AddVisualizePass(FRDGBuilder& GraphBu
 	}
 
 	const FVirtualShadowMapVisualizationData& VisualizationData = GetVirtualShadowMapVisualizationData();
-	if (!VisualizationData.IsActive() || !VisualizeLight[ViewIndex].IsValid())
+	if (!VisualizationData.IsActive())
 	{
 		return MoveTemp(FinalizeOutput(GraphBuilder, View, Output, OverrideOutput));
 	}
@@ -3663,6 +3755,7 @@ FScreenPassTexture FVirtualShadowMapArray::AddVisualizePass(FRDGBuilder& GraphBu
 	FScreenPassRenderTarget OutputTarget(Output.Texture, OutputViewport.Rect, ERenderTargetLoadAction::ELoad);
 
 	int ActiveModeId = VisualizationData.GetActiveModeID();
+	int VisualizeVirtualShadowMapId = VisualizeLight[ViewIndex].GetVirtualShadowMapId();
 
 	// Resize viewport for layout
 	const int32 VisualizeLayout = CVarVisualizeLayout.GetValueOnRenderThread();
@@ -3685,11 +3778,13 @@ FScreenPassTexture FVirtualShadowMapArray::AddVisualizePass(FRDGBuilder& GraphBu
 	auto DrawDebugVisualizationOutput = [&]()
 	{
 		TShaderMapRef<FScreenPassVS> VertexShader(View.ShaderMap);
-		TShaderMapRef<FCopyRectPS> PixelShader(View.ShaderMap);
+		TShaderMapRef<FTonemapProjectionDebugTexturePS> PixelShader(View.ShaderMap);
 
-		FCopyRectPS::FParameters* Parameters = GraphBuilder.AllocParameters<FCopyRectPS::FParameters>();
-		Parameters->InputTexture = DebugVisualizationOutput[ViewIndex];
-		Parameters->InputSampler = TStaticSamplerState<SF_Bilinear, AM_Clamp, AM_Clamp, AM_Clamp>::GetRHI();
+		FTonemapProjectionDebugTexturePS::FParameters* Parameters = GraphBuilder.AllocParameters<FTonemapProjectionDebugTexturePS::FParameters>();
+		Parameters->DebugTexture = DebugVisualizationOutput[ViewIndex];
+		Parameters->VisualizeModeId = ActiveModeId;
+		Parameters->VisualizeVirtualShadowMapId = VisualizeVirtualShadowMapId;
+		Parameters->VisualizeNaniteOverdrawScale = GNaniteVisualizeOverdrawScale;
 		Parameters->RenderTargets[0] = FRenderTargetBinding(Output.Texture, ERenderTargetLoadAction::ENoAction);
 
 		// Blend with scene color if fullscreen, use black background otherwise
@@ -3744,18 +3839,66 @@ FScreenPassTexture FVirtualShadowMapArray::AddVisualizePass(FRDGBuilder& GraphBu
 		}
 
 		AddDrawCanvasPass(GraphBuilder, RDG_EVENT_NAME("Labels"), View, OutputTarget,
-			[&VisualizeLight = VisualizeLight[ViewIndex], &OutputViewport = OutputViewport](FCanvas& Canvas)
+			[&VisualizeLight = VisualizeLight[ViewIndex], &OutputViewport = OutputViewport, VisualizeVirtualShadowMapId=VisualizeVirtualShadowMapId](FCanvas& Canvas)
 			{
 				const float DPIScale = Canvas.GetDPIScale();
 				Canvas.SetBaseTransform(FMatrix(FScaleMatrix(DPIScale) * Canvas.CalcBaseTransform2D(Canvas.GetViewRect().Width(), Canvas.GetViewRect().Height())));
 
-				const FLinearColor LabelColor(1, 1, 0);
-				Canvas.DrawShadowedString(
-					(OutputViewport.Rect.Min.X + 8) / DPIScale,
-					(OutputViewport.Rect.Max.Y - 19) / DPIScale,
-					*VisualizeLight.GetLightName(),
-					GetStatsFont(),
-					LabelColor);
+				auto DrawColorTile = [&](float X, float Y, float Width, float Height, const FLinearColor& Color)
+				{
+					Canvas.DrawTile(X/DPIScale, Y/DPIScale, Width/DPIScale, Height/DPIScale, 0, 0, 0, 0, Color);
+				};
+
+				auto DrawShadowedString = [&](float X, float Y, FStringView Text, const FLinearColor& Color = FLinearColor::White)
+				{
+					Canvas.DrawShadowedString(X / DPIScale, Y / DPIScale, Text, GetStatsFont(), Color);
+				};
+
+				const FVirtualShadowMapVisualizationData& VisualizationData = GetVirtualShadowMapVisualizationData();
+				FString ModeName = VisualizationData.GetActiveModeName().ToString();
+				int32 ActiveModeId = VisualizationData.GetActiveModeID();
+
+				// Legend
+				FVector2D LegendSize = FVector2D(300, 30);
+				if (ActiveModeId == VIRTUAL_SHADOW_MAP_VISUALIZE_SHADOW_FACTOR)		LegendSize.Y = 60;
+				else if (ActiveModeId == VIRTUAL_SHADOW_MAP_VISUALIZE_CACHED_PAGE)	LegendSize.Y = 100;
+
+				FVector2D LegendPosition = FVector2D(OutputViewport.Rect.Min.X + 8, OutputViewport.Rect.Max.Y - LegendSize.Y - 100);
+				DrawColorTile(LegendPosition.X, LegendPosition.Y, LegendSize.X, LegendSize.Y, FLinearColor(0.1f,0.1f,0.1f,0.8f));
+
+				FString HeaderLabel = ModeName;
+				if (VisualizeLight.IsValid())
+				{
+					HeaderLabel = FString::Printf(TEXT("%s (%s)"), *ModeName, *VisualizeLight.GetLightName());
+				}
+
+				DrawShadowedString(LegendPosition.X + 5, LegendPosition.Y + 5, HeaderLabel);
+
+				if (ActiveModeId == VIRTUAL_SHADOW_MAP_VISUALIZE_SHADOW_FACTOR)
+				{
+					DrawColorTile(LegendPosition.X + 5, LegendPosition.Y + 25, 10, 10, FLinearColor(1,1,0));
+					DrawShadowedString(LegendPosition.X + 20, LegendPosition.Y + 22, TEXT("Lit"));
+					DrawColorTile(LegendPosition.X + 5, LegendPosition.Y + 45, 10, 10, FLinearColor(0,0,1));
+					DrawShadowedString(LegendPosition.X + 20, LegendPosition.Y + 42, TEXT("Shadow"));
+				}
+				else if (ActiveModeId == VIRTUAL_SHADOW_MAP_VISUALIZE_CACHED_PAGE) 
+				{
+					extern int32 GVisualizeCachedPagesOnly;
+					if (VisualizeVirtualShadowMapId != INDEX_NONE || GVisualizeCachedPagesOnly)
+					{
+						DrawColorTile(LegendPosition.X + 5, LegendPosition.Y + 25, 10, 10, FLinearColor(0,1,0));
+						DrawShadowedString(LegendPosition.X + 20, LegendPosition.Y + 22, TEXT("Cached"));
+					}
+					if (!GVisualizeCachedPagesOnly)
+					{
+						DrawColorTile(LegendPosition.X + 5, LegendPosition.Y + 45, 10, 10, FLinearColor(1,0,0));
+						DrawShadowedString(LegendPosition.X + 20, LegendPosition.Y + 42, TEXT("All Invalidated"));
+						DrawColorTile(LegendPosition.X + 5, LegendPosition.Y + 65, 10, 10, FLinearColor(0,0,1));
+						DrawShadowedString(LegendPosition.X + 20, LegendPosition.Y + 62, TEXT("Dynamic Invalidated"));
+						DrawColorTile(LegendPosition.X + 5, LegendPosition.Y + 85, 10, 10, FLinearColor(0.75,1,0));
+						DrawShadowedString(LegendPosition.X + 20, LegendPosition.Y + 82, TEXT("Force cached"));
+					}
+				}
 			});
 	}
 
diff --git a/Engine/Source/Runtime/Renderer/Private/VirtualShadowMaps/VirtualShadowMapArray.h b/Engine/Source/Runtime/Renderer/Private/VirtualShadowMaps/VirtualShadowMapArray.h
--- a/Engine/Source/Runtime/Renderer/Private/VirtualShadowMaps/VirtualShadowMapArray.h
+++ b/Engine/Source/Runtime/Renderer/Private/VirtualShadowMaps/VirtualShadowMapArray.h
@@ -214,6 +214,7 @@ public:
 	}
 
 	void CheckLight(const FLightSceneProxy* CheckProxy, int CheckVirtualShadowMapId);
+	void ChooseLight();
 
 	bool IsValid() const { return FoundProxy != nullptr; }
 
@@ -227,8 +228,6 @@ private:
 		struct
 		{
 			// NOTE: Lowest to highest priority
-			uint32 bExists : 1;				// Catch-all
-			uint32 bDirectionalLight : 1;
 			uint32 bSelected : 1;			// In editor
 			uint32 bPartialNameMatch : 1;
 			uint32 bExactNameMatch : 1;
diff --git a/Engine/Source/Runtime/Renderer/Private/VirtualShadowMaps/VirtualShadowMapProjection.cpp b/Engine/Source/Runtime/Renderer/Private/VirtualShadowMaps/VirtualShadowMapProjection.cpp
--- a/Engine/Source/Runtime/Renderer/Private/VirtualShadowMaps/VirtualShadowMapProjection.cpp
+++ b/Engine/Source/Runtime/Renderer/Private/VirtualShadowMaps/VirtualShadowMapProjection.cpp
@@ -64,7 +64,13 @@ static TAutoConsoleVariable<int32> CVarTestPermutation(
 );
 #endif
 
-extern int32 GNaniteVisualizeOverdrawScale;
+int32 GVisualizeCachedPagesOnly = 0;
+FAutoConsoleVariableRef CVarVisualizeCachedPagesOnly(
+	TEXT("r.Shadow.Virtual.Visualize.ShowCachedPagesOnly"),
+	GVisualizeCachedPagesOnly,
+	TEXT("When true, shows the cached pages for all lights and hides uncached pages."),
+	ECVF_RenderThreadSafe
+);
 
 // The tile size in pixels for VSM projection with tile list.
 // Is also used as the workgroup size for the CS without tile list.
@@ -137,8 +143,8 @@ class FVirtualShadowMapProjectionCS : public FGlobalShader
 		// Visualization output
 		SHADER_PARAMETER_RDG_BUFFER_SRV(StructuredBuffer< FPhysicalPageMetaData >, PhysicalPageMetaData)
 		SHADER_PARAMETER(int32, VisualizeModeId)
+		SHADER_PARAMETER(int32, bVisualizeCachedPagesOnly)
 		SHADER_PARAMETER(int32, VisualizeVirtualShadowMapId)
-		SHADER_PARAMETER(float, VisualizeNaniteOverdrawScale)
 		SHADER_PARAMETER_RDG_TEXTURE_UAV(RWTexture2D, OutVisualize)
 		// Optional tile list
 		SHADER_PARAMETER_RDG_BUFFER_SRV(Buffer, TileListData)
@@ -257,14 +263,14 @@ static void RenderVirtualShadowMapProjectionCommon(
  	
 	bool bDebugOutput = false;
 #if !UE_BUILD_SHIPPING
-	if ( !VirtualShadowMapArray.DebugVisualizationOutput.IsEmpty() && InputType == EVirtualShadowMapProjectionInputType::GBuffer && VirtualShadowMapArray.VisualizeLight[ViewIndex].IsValid())
+	if ( !VirtualShadowMapArray.DebugVisualizationOutput.IsEmpty() && InputType == EVirtualShadowMapProjectionInputType::GBuffer)
 	{
 		const FVirtualShadowMapVisualizationData& VisualizationData = GetVirtualShadowMapVisualizationData();
 
 		bDebugOutput = true;
 		PassParameters->VisualizeModeId = VisualizationData.GetActiveModeID();
+		PassParameters->bVisualizeCachedPagesOnly = GVisualizeCachedPagesOnly;
 		PassParameters->VisualizeVirtualShadowMapId = VirtualShadowMapArray.VisualizeLight[ViewIndex].GetVirtualShadowMapId();
-		PassParameters->VisualizeNaniteOverdrawScale = GNaniteVisualizeOverdrawScale;
 		PassParameters->PhysicalPageMetaData = GraphBuilder.CreateSRV( VirtualShadowMapArray.PhysicalPageMetaDataRDG );
 		PassParameters->OutVisualize = GraphBuilder.CreateUAV( VirtualShadowMapArray.DebugVisualizationOutput[ViewIndex] );
 	}
