Remove auto-scaling shadow map dimensions based on LOD.

diff --git a/Engine/Source/Runtime/Renderer/Private/HeterogeneousVolumes/HeterogeneousVolumesLiveShadingPipeline.cpp b/Engine/Source/Runtime/Renderer/Private/HeterogeneousVolumes/HeterogeneousVolumesLiveShadingPipeline.cpp
--- a/Engine/Source/Runtime/Renderer/Private/HeterogeneousVolumes/HeterogeneousVolumesLiveShadingPipeline.cpp
+++ b/Engine/Source/Runtime/Renderer/Private/HeterogeneousVolumes/HeterogeneousVolumesLiveShadingPipeline.cpp
@@ -53,7 +53,7 @@ static TAutoConsoleVariable<int32> CVarHeterogeneousVolumesShadowsUseCameraScene
 	ECVF_RenderThreadSafe
 );
 
-static TAutoConsoleVariable<int32> CVarHeterogeneousVolumesShadowsNearClippingDistance(
+static TAutoConsoleVariable<float> CVarHeterogeneousVolumesShadowsNearClippingDistance(
 	TEXT("r.HeterogeneousVolumes.Shadows.NearClippingDistance"),
 	1.0,
 	TEXT("Near clipping plane distance for shadow projection (Default = 1.0)"),
@@ -2187,20 +2187,6 @@ bool RenderVolumetricShadowMapForLightWithLiveShading(
 	LODInfo.DownsampleFactor = 1.0;
 	LODInfo.bIsPerspective = (LightType != LightType_Directional);
 
-	// Adjust shadow resolution based on minimum MipLevel
-	float LODValue = FMath::CeilLogTwo(ShadowMapResolution.X);
-	for (auto VolumetricMeshBatch : HeterogeneousVolumesMeshBatches)
-	{
-		int32 VolumeCount = VolumetricMeshBatch.Mesh->Elements.Num();
-		for (int32 VolumeIndex = 0; VolumeIndex < VolumeCount; ++VolumeIndex)
-		{
-			const IHeterogeneousVolumeInterface* HeterogeneousVolumeInterface = (IHeterogeneousVolumeInterface*)VolumetricMeshBatch.Mesh->Elements[VolumeIndex].UserData;
-			LODValue = FMath::Min(LODValue, HeterogeneousVolumes::CalcLOD(LODInfo, HeterogeneousVolumeInterface));
-		}
-	}
-	float LODFactor = HeterogeneousVolumes::CalcLODFactor(LODValue);
-	ShadowMapResolution /= LODFactor;
-
 	// Iterate over shadow-casting volumes
 	bool bHasShadowCastingVolume = false;
 	if (!HeterogeneousVolumesMeshBatches.IsEmpty())
