Restored occlusion fence wait back to immediate command list. Fixes RHI thread stalls on mobile.

diff --git a/Engine/Source/Runtime/Renderer/Private/SceneOcclusion.cpp b/Engine/Source/Runtime/Renderer/Private/SceneOcclusion.cpp
--- a/Engine/Source/Runtime/Renderer/Private/SceneOcclusion.cpp
+++ b/Engine/Source/Runtime/Renderer/Private/SceneOcclusion.cpp
@@ -1470,17 +1470,17 @@ void FDeferredShadingSceneRenderer::RenderOcclusion(
 		TrimAllOcclusionHistory(Views);
 	}
 
+	if (bIsOcclusionTesting)
+	{
+		FenceOcclusionTests(GraphBuilder);
+	}
+
 	const bool bUseHzbOcclusion = RenderHzb(GraphBuilder, SceneTextures.Depth.Resolve, BuildHZBAsyncComputeParams);
 
 	if (bUseHzbOcclusion || bIsOcclusionTesting)
 	{
 		GraphBuilder.AddDispatchHint();
 	}
-
-	if (bIsOcclusionTesting)
-	{
-		FenceOcclusionTests(GraphBuilder);
-	}
 }
 
 void FMobileSceneRenderer::RenderOcclusion(FRHICommandList& RHICmdList)
@@ -1510,7 +1510,7 @@ static uint32 GetViewStateUniqueID(const FSceneRenderer* SceneRenderer)
 	return SceneRenderer->Views.Num() && SceneRenderer->Views[0].ViewState ? SceneRenderer->Views[0].ViewState->UniqueID : 0;
 }
 
-void FSceneRenderer::FenceOcclusionTestsInternal(FRHICommandList& RHICmdList)
+void FSceneRenderer::FenceOcclusionTestsInternal(FRHICommandListImmediate& RHICmdList)
 {
 	SCOPE_CYCLE_COUNTER(STAT_OcclusionSubmittedFence_Dispatch);
 
@@ -1547,10 +1547,12 @@ void FSceneRenderer::FenceOcclusionTests(FRDGBuilder& GraphBuilder)
 {
 	if (DoOcclusionQueries() && IsRunningRHIInSeparateThread())
 	{
-		AddPass(GraphBuilder, RDG_EVENT_NAME("FenceOcclusionTests"), [this](FRHICommandList& RHICmdList)
+		AddPass(GraphBuilder, RDG_EVENT_NAME("FenceOcclusionTests"), [this](FRHICommandListImmediate& RHICmdList)
 		{
 			FenceOcclusionTestsInternal(RHICmdList);
 		});
+
+		GraphBuilder.AddDispatchHint();
 	}
 }
 
diff --git a/Engine/Source/Runtime/Renderer/Private/SceneRendering.h b/Engine/Source/Runtime/Renderer/Private/SceneRendering.h
--- a/Engine/Source/Runtime/Renderer/Private/SceneRendering.h
+++ b/Engine/Source/Runtime/Renderer/Private/SceneRendering.h
@@ -2117,7 +2117,7 @@ public:
 	bool DoOcclusionQueries() const;
 
 	void FenceOcclusionTests(FRDGBuilder& GraphBuilder);
-	void FenceOcclusionTestsInternal(FRHICommandList& RHICmdList);
+	void FenceOcclusionTestsInternal(FRHICommandListImmediate& RHICmdList);
 	void WaitOcclusionTests(FRHICommandListImmediate& RHICmdList);
 
 	// fences to make sure the rhi thread has digested the occlusion query renders before we attempt to read them back async
