Fix artifacts due to UpsampleShadowFactors(...) sampling out of ViewRect.
- calculate UVBilinearMax and used it to clamp UV

diff --git a/Engine/Shaders/Private/CapsuleShadowShaders.usf b/Engine/Shaders/Private/CapsuleShadowShaders.usf
--- a/Engine/Shaders/Private/CapsuleShadowShaders.usf
+++ b/Engine/Shaders/Private/CapsuleShadowShaders.usf
@@ -1146,6 +1146,7 @@ void CapsuleShadowingUpsampleVS(
 
 Texture2D ShadowFactorsTexture;
 SamplerState ShadowFactorsSampler;
+float2 ShadowFactorsUVBilinearMax;
 
 float OutputtingToLightAttenuation;
 
@@ -1161,7 +1162,7 @@ void CapsuleShadowingUpsamplePS(
 
 	float Output;
 	float SceneDepth;
-	UpsampleShadowFactors(SVPos, ScissorRectMinAndSize, 1.0f / DOWNSAMPLE_FACTOR, 0, POSITIVE_INFINITY, ShadowFactorsTexture, ShadowFactorsSampler, Output, SceneDepth);
+	UpsampleShadowFactors(SVPos, ScissorRectMinAndSize, 1.0f / DOWNSAMPLE_FACTOR, 0, POSITIVE_INFINITY, ShadowFactorsTexture, ShadowFactorsSampler, ShadowFactorsUVBilinearMax, Output, SceneDepth);
 
 	if (OutputtingToLightAttenuation > 0)
 	{
diff --git a/Engine/Shaders/Private/DistanceFieldShadowing.usf b/Engine/Shaders/Private/DistanceFieldShadowing.usf
--- a/Engine/Shaders/Private/DistanceFieldShadowing.usf
+++ b/Engine/Shaders/Private/DistanceFieldShadowing.usf
@@ -820,6 +820,7 @@ float NearFadePlaneOffset;
 float InvNearFadePlaneLength;
 
 float OneOverDownsampleFactor;
+float2 ShadowFactorsUVBilinearMax;
 
 void DistanceFieldShadowingUpsamplePS(
 	in float4 SVPos : SV_POSITION,
@@ -827,7 +828,7 @@ void DistanceFieldShadowingUpsamplePS(
 {
 	float Output;
 	float SceneDepth;
-	UpsampleShadowFactors(SVPos, ScissorRectMinAndSize, OneOverDownsampleFactor, MinDepth, MaxDepth, ShadowFactorsTexture, ShadowFactorsSampler, Output, SceneDepth);
+	UpsampleShadowFactors(SVPos, ScissorRectMinAndSize, OneOverDownsampleFactor, MinDepth, MaxDepth, ShadowFactorsTexture, ShadowFactorsSampler, ShadowFactorsUVBilinearMax, Output, SceneDepth);
 
 	float FarBlendFactor = 1.0f - saturate((SceneDepth - FadePlaneOffset) * InvFadePlaneLength);
 	Output = lerp(1, Output, FarBlendFactor);
diff --git a/Engine/Shaders/Private/ScreenSpaceShadows.usf b/Engine/Shaders/Private/ScreenSpaceShadows.usf
--- a/Engine/Shaders/Private/ScreenSpaceShadows.usf
+++ b/Engine/Shaders/Private/ScreenSpaceShadows.usf
@@ -193,6 +193,7 @@ void ScreenSpaceShadowsBendCS(
 
 Texture2D ShadowFactorsTexture;
 SamplerState ShadowFactorsSampler;
+float2 ShadowFactorsUVBilinearMax;
 
 float OneOverDownsampleFactor;
 
@@ -202,7 +203,7 @@ void ScreenSpaceShadowsUpsamplePS(
 {
 	float Output;
 	float SceneDepth;
-	UpsampleShadowFactors(SVPos, ScissorRectMinAndSize, OneOverDownsampleFactor, 0, POSITIVE_INFINITY, ShadowFactorsTexture, ShadowFactorsSampler, Output, SceneDepth);
+	UpsampleShadowFactors(SVPos, ScissorRectMinAndSize, OneOverDownsampleFactor, 0, POSITIVE_INFINITY, ShadowFactorsTexture, ShadowFactorsSampler, ShadowFactorsUVBilinearMax, Output, SceneDepth);
 
 	OutColor = EncodeLightAttenuation(half4(Output, Output, Output, Output));
 }
diff --git a/Engine/Shaders/Private/ShadowFactorsUpsampleCommon.ush b/Engine/Shaders/Private/ShadowFactorsUpsampleCommon.ush
--- a/Engine/Shaders/Private/ShadowFactorsUpsampleCommon.ush
+++ b/Engine/Shaders/Private/ShadowFactorsUpsampleCommon.ush
@@ -13,12 +13,13 @@
 
 void UpsampleShadowFactors(
 	float4 SVPos, 
-	int4 ScissorRectMinAndSize, 
+	int4 ScissorRectMinAndSize,
 	float OneOverDownsampleFactor, 
 	float ShadowFactorsMinDepth, 
 	float ShadowFactorsMaxDepth, 
 	Texture2D ShadowFactorsTexture, 
-	SamplerState ShadowFactorsSampler, 
+	SamplerState ShadowFactorsSampler,
+	float2 ShadowFactorsUVBilinearMax,
 	out float ShadowFactor, 
 	out float SceneDepth)
 {
@@ -46,10 +47,10 @@ void UpsampleShadowFactors(
 	// Sample low res texture using UVs at center of low-res texels since we are interpolating manually
 	float2 Corner00SampleUV = Corner00UVBase + 0.5f * LowResTexelSize;
 
-	float2 TextureValues00 = Texture2DSampleLevel(ShadowFactorsTexture, ShadowFactorsSampler, Corner00SampleUV, 0).xy;
-	float2 TextureValues10 = Texture2DSampleLevel(ShadowFactorsTexture, ShadowFactorsSampler, Corner00SampleUV + float2(LowResTexelSize.x, 0), 0).xy;
-	float2 TextureValues01 = Texture2DSampleLevel(ShadowFactorsTexture, ShadowFactorsSampler, Corner00SampleUV + float2(0, LowResTexelSize.y), 0).xy;
-	float2 TextureValues11 = Texture2DSampleLevel(ShadowFactorsTexture, ShadowFactorsSampler, Corner00SampleUV + LowResTexelSize, 0).xy;
+	float2 TextureValues00 = Texture2DSampleLevel(ShadowFactorsTexture, ShadowFactorsSampler, min(Corner00SampleUV, ShadowFactorsUVBilinearMax), 0).xy;
+	float2 TextureValues10 = Texture2DSampleLevel(ShadowFactorsTexture, ShadowFactorsSampler, min(Corner00SampleUV + float2(LowResTexelSize.x, 0), ShadowFactorsUVBilinearMax), 0).xy;
+	float2 TextureValues01 = Texture2DSampleLevel(ShadowFactorsTexture, ShadowFactorsSampler, min(Corner00SampleUV + float2(0, LowResTexelSize.y), ShadowFactorsUVBilinearMax), 0).xy;
+	float2 TextureValues11 = Texture2DSampleLevel(ShadowFactorsTexture, ShadowFactorsSampler, min(Corner00SampleUV + LowResTexelSize, 0), ShadowFactorsUVBilinearMax).xy;
 
 	float4 CornerWeights = float4(
 		(1 - BilinearWeights.y) * (1 - BilinearWeights.x), 
diff --git a/Engine/Source/Runtime/Renderer/Private/CapsuleShadowRendering.cpp b/Engine/Source/Runtime/Renderer/Private/CapsuleShadowRendering.cpp
--- a/Engine/Source/Runtime/Renderer/Private/CapsuleShadowRendering.cpp
+++ b/Engine/Source/Runtime/Renderer/Private/CapsuleShadowRendering.cpp
@@ -330,6 +330,7 @@ public:
 		SHADER_PARAMETER_RDG_TEXTURE(Texture2D, ShadowFactorsTexture)
 		SHADER_PARAMETER_SAMPLER(SamplerState, ShadowFactorsSampler)
 		SHADER_PARAMETER(FIntRect, ScissorRectMinAndSize)
+		SHADER_PARAMETER(FVector2f, ShadowFactorsUVBilinearMax)
 		SHADER_PARAMETER(float, OutputtingToLightAttenuation)
 	END_SHADER_PARAMETER_STRUCT()
 
@@ -493,10 +494,10 @@ bool FDeferredShadingSceneRenderer::RenderCapsuleDirectShadows(
 
 	QUICK_SCOPE_CYCLE_COUNTER(STAT_RenderCapsuleShadows);
 
-	FRDGTextureRef RayTracedShadowsRT = nullptr;
+	const FIntPoint BufferSize = GetBufferSizeForCapsuleShadows(Views[0]);
 
+	FRDGTextureRef RayTracedShadowsRT = nullptr;
 	{
-		const FIntPoint BufferSize = GetBufferSizeForCapsuleShadows(Views[0]);
 		const FRDGTextureDesc Desc(FRDGTextureDesc::Create2D(BufferSize, PF_G16R16F, FClearValueBinding::None, TexCreate_ShaderResource | TexCreate_UAV));
 		RayTracedShadowsRT = GraphBuilder.CreateTexture(Desc, TEXT("CapsuleShadows.ShadowFactors"));
 	}
@@ -549,9 +550,12 @@ bool FDeferredShadingSceneRenderer::RenderCapsuleDirectShadows(
 				ScissorRect = View.ViewRect;
 			}
 
+			// CapsuleShadowingCS always outputs at rect with min = (0,0)
+			const FIntRect DownsampledViewRect(0, 0, ScissorRect.Width() / GetCapsuleShadowDownsampleFactor(), ScissorRect.Height() / GetCapsuleShadowDownsampleFactor());
+
 			const FIntPoint GroupSize = FIntPoint(
-				FMath::DivideAndRoundUp(ScissorRect.Size().X / GetCapsuleShadowDownsampleFactor(), GShadowShapeTileSize),
-				FMath::DivideAndRoundUp(ScissorRect.Size().Y / GetCapsuleShadowDownsampleFactor(), GShadowShapeTileSize))
+				FMath::DivideAndRoundUp(DownsampledViewRect.Width(), GShadowShapeTileSize),
+				FMath::DivideAndRoundUp(DownsampledViewRect.Height(), GShadowShapeTileSize))
 				.ComponentMax(FIntPoint(1, 1));
 
 			FRDGBufferRef CapsuleTileIntersectionCountsBuffer = GraphBuilder.CreateBuffer(FRDGBufferDesc::CreateStructuredDesc(sizeof(uint32), GroupSize.X * GroupSize.Y), TEXT("CapsuleTileIntersectionCountsBuffer"));
@@ -607,6 +611,10 @@ bool FDeferredShadingSceneRenderer::RenderCapsuleDirectShadows(
 			}
 
 			{
+				const FVector2f RayTracedShadowsTextureExtent(RayTracedShadowsRT->Desc.Extent);
+				const FVector2f RayTracedShadowsTextureExtentInverse(1.0f / RayTracedShadowsTextureExtent.X, 1.0f / RayTracedShadowsTextureExtent.Y);
+				const FVector2f RayTracedShadowsTextureViewportMax(DownsampledViewRect.Max);
+
 				auto VertexShader = View.ShaderMap->GetShader<FCapsuleShadowingUpsampleVS>();
 
 				FCapsuleShadowingUpsamplePS::FPermutationDomain PermutationVector;
@@ -629,6 +637,7 @@ bool FDeferredShadingSceneRenderer::RenderCapsuleDirectShadows(
 				PassParameters->PS.ShadowFactorsTexture = RayTracedShadowsRT;
 				PassParameters->PS.ShadowFactorsSampler = TStaticSamplerState<SF_Bilinear>::GetRHI();
 				PassParameters->PS.ScissorRectMinAndSize = FIntRect(ScissorRect.Min, ScissorRect.Size());
+				PassParameters->PS.ShadowFactorsUVBilinearMax = (RayTracedShadowsTextureViewportMax - 0.5f) * RayTracedShadowsTextureExtentInverse;
 				PassParameters->PS.OutputtingToLightAttenuation = 1.0f;
 
 				ClearUnusedGraphResources(VertexShader, &PassParameters->VS);
@@ -980,10 +989,10 @@ void FDeferredShadingSceneRenderer::RenderIndirectCapsuleShadows(FRDGBuilder& Gr
 
 	RDG_EVENT_SCOPE(GraphBuilder, "IndirectCapsuleShadows");
 
-	FRDGTextureRef RayTracedShadowsRT = nullptr;
+	const FIntPoint BufferSize = GetBufferSizeForCapsuleShadows(Views[0]);
 
+	FRDGTextureRef RayTracedShadowsRT = nullptr;
 	{
-		const FIntPoint BufferSize = GetBufferSizeForCapsuleShadows(Views[0]);
 		const FRDGTextureDesc Desc(FRDGTextureDesc::Create2D(BufferSize, PF_G16R16F, FClearValueBinding::None, TexCreate_ShaderResource | TexCreate_UAV));
 		RayTracedShadowsRT = GraphBuilder.CreateTexture(Desc, TEXT("CapsuleShadows.ShadowFactors"));
 	}
@@ -1021,9 +1030,12 @@ void FDeferredShadingSceneRenderer::RenderIndirectCapsuleShadows(FRDGBuilder& Gr
 
 			const FIntRect ScissorRect = View.ViewRect;
 
+			// CapsuleShadowingCS always outputs at rect with min = (0,0)
+			const FIntRect DownsampledViewRect(0, 0, ScissorRect.Width() / GetCapsuleShadowDownsampleFactor(), ScissorRect.Height() / GetCapsuleShadowDownsampleFactor());
+
 			const FIntPoint GroupSize(
-				FMath::DivideAndRoundUp(ScissorRect.Size().X / GetCapsuleShadowDownsampleFactor(), GShadowShapeTileSize),
-				FMath::DivideAndRoundUp(ScissorRect.Size().Y / GetCapsuleShadowDownsampleFactor(), GShadowShapeTileSize));
+				FMath::DivideAndRoundUp(DownsampledViewRect.Width(), GShadowShapeTileSize),
+				FMath::DivideAndRoundUp(DownsampledViewRect.Height(), GShadowShapeTileSize));
 
 			FRDGBufferRef CapsuleTileIntersectionCountsBuffer = GraphBuilder.CreateBuffer(FRDGBufferDesc::CreateStructuredDesc(sizeof(uint32), GroupSize.X * GroupSize.Y), TEXT("CapsuleTileIntersectionCountsBuffer"));
 			FRDGBufferUAVRef CapsuleTileIntersectionCountsUAV = GraphBuilder.CreateUAV(CapsuleTileIntersectionCountsBuffer);
@@ -1088,6 +1100,10 @@ void FDeferredShadingSceneRenderer::RenderIndirectCapsuleShadows(FRDGBuilder& Gr
 			}
 
 			{
+				const FVector2f RayTracedShadowsTextureExtent(RayTracedShadowsRT->Desc.Extent);
+				const FVector2f RayTracedShadowsTextureExtentInverse(1.0f / RayTracedShadowsTextureExtent.X, 1.0f / RayTracedShadowsTextureExtent.Y);
+				const FVector2f RayTracedShadowsTextureViewportMax(DownsampledViewRect.Max);
+
 				const int32 RenderTargetCount = RenderTargets.Num();
 
 				auto VertexShader = View.ShaderMap->GetShader<FCapsuleShadowingUpsampleVS>();
@@ -1118,6 +1134,7 @@ void FDeferredShadingSceneRenderer::RenderIndirectCapsuleShadows(FRDGBuilder& Gr
 				PassParameters->PS.ShadowFactorsTexture = RayTracedShadowsRT;
 				PassParameters->PS.ShadowFactorsSampler = TStaticSamplerState<SF_Bilinear>::GetRHI();
 				PassParameters->PS.ScissorRectMinAndSize = FIntRect(ScissorRect.Min, ScissorRect.Size());
+				PassParameters->PS.ShadowFactorsUVBilinearMax = (RayTracedShadowsTextureViewportMax - 0.5f) * RayTracedShadowsTextureExtentInverse;
 				PassParameters->PS.OutputtingToLightAttenuation = 0.0f;
 
 				ClearUnusedGraphResources(VertexShader, &PassParameters->VS);
diff --git a/Engine/Source/Runtime/Renderer/Private/DistanceFieldShadowing.cpp b/Engine/Source/Runtime/Renderer/Private/DistanceFieldShadowing.cpp
--- a/Engine/Source/Runtime/Renderer/Private/DistanceFieldShadowing.cpp
+++ b/Engine/Source/Runtime/Renderer/Private/DistanceFieldShadowing.cpp
@@ -152,11 +152,13 @@ int32 GetDFShadowDownsampleFactor()
 
 FIntPoint GetBufferSizeForDFShadows(const FViewInfo& View)
 {
+	// TODO: Use GetDownscaledExtent(...) but need to check different rounding doesn't cause issues
 	return FIntPoint::DivideAndRoundDown(View.GetSceneTexturesConfig().Extent, GetDFShadowDownsampleFactor());
 }
 
 FIntRect GetScissorRectForDFShadows(FIntRect ScissorRect)
 {
+	// TODO: Use GetDownscaledRect(...) but need to check different rounding doesn't cause issues
 	return ScissorRect / GetDFShadowDownsampleFactor();
 }
 
@@ -395,6 +397,7 @@ class FDistanceFieldShadowingUpsamplePS : public FGlobalShader
 		SHADER_PARAMETER_RDG_TEXTURE(Texture2D, ShadowFactorsTexture)
 		SHADER_PARAMETER_SAMPLER(SamplerState, ShadowFactorsSampler)
 		SHADER_PARAMETER(FIntRect, ScissorRectMinAndSize)
+		SHADER_PARAMETER(FVector2f, ShadowFactorsUVBilinearMax)
 		SHADER_PARAMETER(float, FadePlaneOffset)
 		SHADER_PARAMETER(float, InvFadePlaneLength)
 		SHADER_PARAMETER(float, NearFadePlaneOffset)
@@ -902,7 +905,7 @@ void RayTraceShadows(
 	}
 }
 
-FRDGTextureRef FProjectedShadowInfo::RenderRayTracedDistanceFieldProjection(
+FScreenPassTexture FProjectedShadowInfo::RenderRayTracedDistanceFieldProjection(
 	FRDGBuilder& GraphBuilder,
 	bool bAsyncCompute, 
 	const FMinimalSceneTextures& SceneTextures,
@@ -911,18 +914,19 @@ FRDGTextureRef FProjectedShadowInfo::RenderRayTracedDistanceFieldProjection(
 {
 	DistanceFieldShadowViewGPUData& SDFShadowViewGPUData = CachedDistanceFieldShadowViewGPUData.FindOrAdd(&View);
 
+	FIntRect DownsampledScissorRect = GetScissorRectForDFShadows(ScissorRect);
+	DownsampledScissorRect -= DownsampledScissorRect.Min; // DistanceFieldShadowingCS always outputs at rect with min = (0,0)
+
 	if (SDFShadowViewGPUData.RayTracedShadowsTexture)
 	{
 		// Ray traced distance field shadows were already calculated, simply return previous result.
-		return SDFShadowViewGPUData.RayTracedShadowsTexture;
+		return FScreenPassTexture(SDFShadowViewGPUData.RayTracedShadowsTexture, DownsampledScissorRect);
 	}
 
-	const FIntRect DownsampledScissorRect = GetScissorRectForDFShadows(ScissorRect);
-
 	if (DownsampledScissorRect.Area() <= 0)
 	{
 		// skip calculating DF shadows
-		return nullptr;
+		return FScreenPassTexture();
 	}
 
 	const bool bDFShadowSupported = SupportsDistanceFieldShadows(View.GetFeatureLevel(), View.GetShaderPlatform());
@@ -1100,7 +1104,7 @@ FRDGTextureRef FProjectedShadowInfo::RenderRayTracedDistanceFieldProjection(
 			*SDFShadowViewGPUData.HeightFieldLightTileIntersectionParameters);
 	}
 
-	return SDFShadowViewGPUData.RayTracedShadowsTexture;
+	return FScreenPassTexture(SDFShadowViewGPUData.RayTracedShadowsTexture, DownsampledScissorRect);
 }
 
 BEGIN_SHADER_PARAMETER_STRUCT(FDistanceFieldShadowingUpsample, )
@@ -1123,19 +1127,24 @@ void FProjectedShadowInfo::RenderRayTracedDistanceFieldProjection(
 	check(ScissorRect.Area() > 0);
 	const bool bRunTiled = UseShadowIndirectDraw(View.GetShaderPlatform()) && TiledShadowRendering != nullptr;
 
-	FRDGTextureRef RayTracedShadowsTexture = RenderRayTracedDistanceFieldProjection(GraphBuilder, false, SceneTextures, View, ScissorRect);
+	FScreenPassTexture RayTracedShadowsTexture = RenderRayTracedDistanceFieldProjection(GraphBuilder, false, SceneTextures, View, ScissorRect);
 
-	if (RayTracedShadowsTexture)
+	if (RayTracedShadowsTexture.IsValid())
 	{
+		const FVector2f RayTracedShadowsTextureExtent(RayTracedShadowsTexture.Texture->Desc.Extent);
+		const FVector2f RayTracedShadowsTextureExtentInverse(1.0f / RayTracedShadowsTextureExtent.X, 1.0f / RayTracedShadowsTextureExtent.Y);
+		const FVector2f RayTracedShadowsTextureViewportMax(RayTracedShadowsTexture.ViewRect.Max);
+
 		FDistanceFieldShadowingUpsample* PassParameters = GraphBuilder.AllocParameters<FDistanceFieldShadowingUpsample>();
 		PassParameters->RenderTargets[0] = FRenderTargetBinding(ScreenShadowMaskTexture, ERenderTargetLoadAction::ELoad);
 		PassParameters->RenderTargets.DepthStencil = FDepthStencilBinding(SceneTextures.Depth.Target, ERenderTargetLoadAction::ELoad, ERenderTargetLoadAction::ELoad, FExclusiveDepthStencil::DepthRead_StencilRead);
 		
 		PassParameters->PS.View = GetShaderBinding(View.ViewUniformBuffer);
 		PassParameters->PS.SceneTextures = SceneTextures.GetSceneTextureShaderParameters(View.GetFeatureLevel());
-		PassParameters->PS.ShadowFactorsTexture = RayTracedShadowsTexture;
+		PassParameters->PS.ShadowFactorsTexture = RayTracedShadowsTexture.Texture;
 		PassParameters->PS.ShadowFactorsSampler = TStaticSamplerState<SF_Bilinear>::GetRHI();
 		PassParameters->PS.ScissorRectMinAndSize = FIntRect(ScissorRect.Min, ScissorRect.Size());
+		PassParameters->PS.ShadowFactorsUVBilinearMax = (RayTracedShadowsTextureViewportMax - 0.5f) * RayTracedShadowsTextureExtentInverse;
 		PassParameters->PS.OneOverDownsampleFactor = 1.0f / GetDFShadowDownsampleFactor();
 
 		GetDistanceFieldShadowRange(this, DFPT_SignedDistanceField, PassParameters->PS.MinDepth, PassParameters->PS.MaxDepth);
diff --git a/Engine/Source/Runtime/Renderer/Private/ShadowRendering.h b/Engine/Source/Runtime/Renderer/Private/ShadowRendering.h
--- a/Engine/Source/Runtime/Renderer/Private/ShadowRendering.h
+++ b/Engine/Source/Runtime/Renderer/Private/ShadowRendering.h
@@ -620,7 +620,7 @@ public:
 	* Output texture is cached per view. This is useful to support async compute.
 	* (ie: kick off distance field shadows early in the frame using async compute, and then combine result into shadow mask texture when necessary)
 	*/
-	FRDGTextureRef RenderRayTracedDistanceFieldProjection(
+	FScreenPassTexture RenderRayTracedDistanceFieldProjection(
 		FRDGBuilder& GraphBuilder,
 		bool bAsyncCompute,
 		const FMinimalSceneTextures& SceneTextures,
diff --git a/Engine/Source/Runtime/Renderer/Private/Shadows/ScreenSpaceShadows.cpp b/Engine/Source/Runtime/Renderer/Private/Shadows/ScreenSpaceShadows.cpp
--- a/Engine/Source/Runtime/Renderer/Private/Shadows/ScreenSpaceShadows.cpp
+++ b/Engine/Source/Runtime/Renderer/Private/Shadows/ScreenSpaceShadows.cpp
@@ -188,6 +188,7 @@ class FScreenSpaceShadowsUpsamplePS : public FGlobalShader
 		SHADER_PARAMETER_RDG_TEXTURE(Texture2D, ShadowFactorsTexture)
 		SHADER_PARAMETER_SAMPLER(SamplerState, ShadowFactorsSampler)
 		SHADER_PARAMETER(FIntRect, ScissorRectMinAndSize)
+		SHADER_PARAMETER(FVector2f, ShadowFactorsUVBilinearMax)
 		SHADER_PARAMETER(float, OneOverDownsampleFactor)
 	END_SHADER_PARAMETER_STRUCT()
 
@@ -220,20 +221,25 @@ void UpsampleScreenSpaceShadows(
 	FIntRect ScissorRect,
 	bool bProjectingForForwardShading,
 	const FLightSceneInfo* LightSceneInfo,
-	FRDGTextureRef ShadowsTexture,
+	FScreenPassTexture ShadowsTexture,
 	int32 DownsampleFactor,
 	FRDGTextureRef ScreenShadowMaskTexture
 )
 {
+	const FVector2f ShadowsTextureExtent(ShadowsTexture.Texture->Desc.Extent);
+	const FVector2f ShadowsTextureExtentInverse(1.0f / ShadowsTextureExtent.X, 1.0f / ShadowsTextureExtent.Y);
+	const FVector2f ShadowsTextureViewportMax(ShadowsTexture.ViewRect.Max);
+
 	FScreenSpaceShadowsUpsample* PassParameters = GraphBuilder.AllocParameters<FScreenSpaceShadowsUpsample>();
 	PassParameters->RenderTargets[0] = FRenderTargetBinding(ScreenShadowMaskTexture, ERenderTargetLoadAction::ELoad);
 	PassParameters->RenderTargets.DepthStencil = FDepthStencilBinding(SceneTextures.Depth.Target, ERenderTargetLoadAction::ELoad, ERenderTargetLoadAction::ELoad, FExclusiveDepthStencil::DepthRead_StencilRead);
 
 	PassParameters->PS.View = GetShaderBinding(View.ViewUniformBuffer);
 	PassParameters->PS.SceneTextures = SceneTextures.GetSceneTextureShaderParameters(View.GetFeatureLevel());
-	PassParameters->PS.ShadowFactorsTexture = ShadowsTexture;
+	PassParameters->PS.ShadowFactorsTexture = ShadowsTexture.Texture;
 	PassParameters->PS.ShadowFactorsSampler = TStaticSamplerState<SF_Bilinear>::GetRHI();
 	PassParameters->PS.ScissorRectMinAndSize = FIntRect(ScissorRect.Min, ScissorRect.Size());
+	PassParameters->PS.ShadowFactorsUVBilinearMax = (ShadowsTextureViewportMax - 0.5f) * ShadowsTextureExtentInverse;
 	PassParameters->PS.OneOverDownsampleFactor = 1.0f / DownsampleFactor;
 
 	FScreenSpaceShadowsUpsamplePS::FPermutationDomain PermutationVector;
@@ -293,6 +299,9 @@ void RenderScreenSpaceShadows(
 	FRDGTextureDesc Desc(FRDGTextureDesc::Create2D(BufferSize, PF_G16R16F, FClearValueBinding::None, TexCreate_UAV | TexCreate_ShaderResource));
 	FRDGTextureRef ShadowsTexture = GraphBuilder.CreateTexture(Desc, TEXT("ScreenSpaceShadows"));
 
+	// ScreenSpaceShadowsCS always outputs at rect with min = (0,0)
+	FIntRect ShadowsTextureViewRect(0, 0, ScissorRect.Width() / GetScreenSpaceShadowDownsampleFactor(), ScissorRect.Height() / GetScreenSpaceShadowDownsampleFactor());
+
 	{
 		FScreenSpaceShadowsCS::FParameters* PassParameters = GraphBuilder.AllocParameters<FScreenSpaceShadowsCS::FParameters>();
 
@@ -323,8 +332,8 @@ void RenderScreenSpaceShadows(
 
 		auto ComputeShader = View.ShaderMap->GetShader<FScreenSpaceShadowsCS>();
 
-		uint32 GroupSizeX = FMath::DivideAndRoundUp(ScissorRect.Size().X / GetScreenSpaceShadowDownsampleFactor(), GScreenSpaceShadowsTileSizeX);
-		uint32 GroupSizeY = FMath::DivideAndRoundUp(ScissorRect.Size().Y / GetScreenSpaceShadowDownsampleFactor(), GScreenSpaceShadowsTileSizeY);
+		uint32 GroupSizeX = FMath::DivideAndRoundUp(ShadowsTextureViewRect.Width(), GScreenSpaceShadowsTileSizeX);
+		uint32 GroupSizeY = FMath::DivideAndRoundUp(ShadowsTextureViewRect.Height(), GScreenSpaceShadowsTileSizeY);
 
 		FComputeShaderUtils::AddPass(
 			GraphBuilder,
@@ -342,7 +351,7 @@ void RenderScreenSpaceShadows(
 		ScissorRect,
 		bProjectingForForwardShading,
 		LightSceneInfo,
-		ShadowsTexture,
+		FScreenPassTexture(ShadowsTexture, ShadowsTextureViewRect),
 		GetScreenSpaceShadowDownsampleFactor(),
 		ScreenShadowMaskTexture);
 }
@@ -445,7 +454,7 @@ void RenderScreenSpaceShadowsBend(
 		ScissorRect,
 		bProjectingForForwardShading,
 		LightSceneInfo,
-		ShadowsTexture,
+		FScreenPassTexture(ShadowsTexture, ScissorRect),
 		DownsampleFactor,
 		ScreenShadowMaskTexture);
 }
