Water: Fixed WaterInfoTexture RenderMethod 1 only checking the very first view for queued WIT updates. Also fixed RDG internal/external access mode issue when rendering to two slices of the WIT in the same frame.

diff --git a/Engine/Source/Runtime/Renderer/Private/WaterInfoTextureRendering.cpp b/Engine/Source/Runtime/Renderer/Private/WaterInfoTextureRendering.cpp
--- a/Engine/Source/Runtime/Renderer/Private/WaterInfoTextureRendering.cpp
+++ b/Engine/Source/Runtime/Renderer/Private/WaterInfoTextureRendering.cpp
@@ -778,12 +778,21 @@ void RenderWaterInfoTexture(
 	const FScene* Scene
 )
 {
-	const FViewInfo& MainView = SceneRenderer.Views[0];
-	if (MainView.WaterInfoTextureRenderingParams.IsEmpty())
+	bool bAnyWaterInfoTextureUpdates = false;
+	for (const FViewInfo& View : SceneRenderer.Views)
+	{
+		bAnyWaterInfoTextureUpdates = bAnyWaterInfoTextureUpdates || !View.WaterInfoTextureRenderingParams.IsEmpty();
+	}
+
+	if (!bAnyWaterInfoTextureUpdates)
 	{
 		return;
 	}
 
+	TRACE_CPUPROFILER_EVENT_SCOPE(WaterInfo::RenderWaterInfoTexture);
+	RDG_EVENT_SCOPE(GraphBuilder, "WaterInfoTexture");
+	RDG_GPU_STAT_SCOPE(GraphBuilder, WaterInfoTexture);
+
 	int32 RenderCaptureNextWaterInfoDraws = CVarRenderCaptureNextWaterInfoDraws.GetValueOnRenderThread();
 	RenderCaptureInterface::FScopedCapture RenderCapture((RenderCaptureNextWaterInfoDraws != 0), GraphBuilder, TEXT("RenderWaterInfo"));
 	if (RenderCaptureNextWaterInfoDraws != 0)
@@ -791,236 +800,248 @@ void RenderWaterInfoTexture(
 		RenderCaptureNextWaterInfoDraws = FMath::Max(0, RenderCaptureNextWaterInfoDraws - 1);
 		CVarRenderCaptureNextWaterInfoDraws->SetWithCurrentPriority(RenderCaptureNextWaterInfoDraws);
 	}
-	
 
-	TRACE_CPUPROFILER_EVENT_SCOPE(WaterInfo::RenderWaterInfoTexture);
-	RDG_EVENT_SCOPE(GraphBuilder, "WaterInfoTexture");
-	RDG_GPU_STAT_SCOPE(GraphBuilder, WaterInfoTexture);
-	
-	for (const FViewInfo::FWaterInfoTextureRenderingParams& RenderingParams : MainView.WaterInfoTextureRenderingParams)
-	{
-		const FIntPoint TextureSize = RenderingParams.RenderTarget->GetSizeXY();
-		const FIntRect Viewport(0, 0, TextureSize.X, TextureSize.Y);
+	// Keep track of all the result textures so we can batch call UseExternalAccessMode() on them.
+	TArray<FRDGViewableResource*, TInlineAllocator<4>> OutputTextures;
 
-		FViewInfo& WaterView = *MainView.CreateSnapshot();
+	// Check all views for WaterInfoTexture updates that we might need to execute.
+	for (const FViewInfo& View : SceneRenderer.Views)
+	{
+		for (const FViewInfo::FWaterInfoTextureRenderingParams& RenderingParams : View.WaterInfoTextureRenderingParams)
 		{
-			WaterView.ViewState = nullptr;
-			WaterView.DynamicPrimitiveCollector = FGPUScenePrimitiveCollector(&SceneRenderer.GetGPUSceneDynamicContext());
-			WaterView.StereoPass = EStereoscopicPass::eSSP_FULL;
-			WaterView.DrawDynamicFlags = EDrawDynamicFlags::ForceLowestLOD;
-			WaterView.MaterialTextureMipBias = 0;
-			WaterView.PreExposure = 1.0f;
-			WaterView.ViewRect = Viewport;
-
-			WaterView.CachedViewUniformShaderParameters = MakeUnique<FViewUniformShaderParameters>();
-			FBox VolumeBounds[TVC_MAX];
-			WaterView.SetupUniformBufferParameters(VolumeBounds, TVC_MAX, *WaterView.CachedViewUniformShaderParameters);
-
-			WaterView.UpdateProjectionMatrix(RenderingParams.ProjectionMatrix);
-
-			FViewMatrices::FMinimalInitializer Initializer;
-			Initializer.ViewRotationMatrix = RenderingParams.ViewRotationMatrix;
-			Initializer.ViewOrigin = RenderingParams.ViewLocation;
-			Initializer.ProjectionMatrix = RenderingParams.ProjectionMatrix;
-			Initializer.ConstrainedViewRect = Viewport;
-			WaterView.ViewMatrices = FViewMatrices(Initializer);
-
-			TRefCountPtr<IPooledRenderTarget> NullRef;
-			FPlatformMemory::Memcpy(&WaterView.PrevViewInfo.HZB, &NullRef, sizeof(WaterView.PrevViewInfo.HZB));
-
-			WaterView.SetupCommonViewUniformBufferParameters(
-				*WaterView.CachedViewUniformShaderParameters, 
-				TextureSize, 
-				1, 
-				Initializer.ConstrainedViewRect, 
-				WaterView.ViewMatrices, 
-				WaterView.ViewMatrices);
-
-			WaterView.CreateViewUniformBuffers(*WaterView.CachedViewUniformShaderParameters);
-		}
+			const FIntPoint TextureSize = RenderingParams.RenderTarget->GetSizeXY();
+			const FIntRect Viewport(0, 0, TextureSize.X, TextureSize.Y);
+
+			FViewInfo& WaterView = *View.CreateSnapshot();
+			{
+				WaterView.ViewState = nullptr;
+				WaterView.DynamicPrimitiveCollector = FGPUScenePrimitiveCollector(&SceneRenderer.GetGPUSceneDynamicContext());
+				WaterView.StereoPass = EStereoscopicPass::eSSP_FULL;
+				WaterView.DrawDynamicFlags = EDrawDynamicFlags::ForceLowestLOD;
+				WaterView.MaterialTextureMipBias = 0;
+				WaterView.PreExposure = 1.0f;
+				WaterView.ViewRect = Viewport;
+
+				WaterView.CachedViewUniformShaderParameters = MakeUnique<FViewUniformShaderParameters>();
+				FBox VolumeBounds[TVC_MAX];
+				WaterView.SetupUniformBufferParameters(VolumeBounds, TVC_MAX, *WaterView.CachedViewUniformShaderParameters);
+
+				WaterView.UpdateProjectionMatrix(RenderingParams.ProjectionMatrix);
+
+				FViewMatrices::FMinimalInitializer Initializer;
+				Initializer.ViewRotationMatrix = RenderingParams.ViewRotationMatrix;
+				Initializer.ViewOrigin = RenderingParams.ViewLocation;
+				Initializer.ProjectionMatrix = RenderingParams.ProjectionMatrix;
+				Initializer.ConstrainedViewRect = Viewport;
+				WaterView.ViewMatrices = FViewMatrices(Initializer);
+
+				TRefCountPtr<IPooledRenderTarget> NullRef;
+				FPlatformMemory::Memcpy(&WaterView.PrevViewInfo.HZB, &NullRef, sizeof(WaterView.PrevViewInfo.HZB));
+
+				WaterView.SetupCommonViewUniformBufferParameters(
+					*WaterView.CachedViewUniformShaderParameters, 
+					TextureSize, 
+					1, 
+					Initializer.ConstrainedViewRect, 
+					WaterView.ViewMatrices, 
+					WaterView.ViewMatrices);
+
+				WaterView.CreateViewUniformBuffers(*WaterView.CachedViewUniformShaderParameters);
+			}
 
-		const FTextureRHIRef OutputRenderTarget = RenderingParams.RenderTarget->GetRenderTargetTexture();
-		FRDGTexture* OutputTexture = GraphBuilder.RegisterExternalTexture(CreateRenderTarget(OutputRenderTarget, TEXT("WaterInfo.OutputTexture")));
+			const FTextureRHIRef OutputRenderTarget = RenderingParams.RenderTarget->GetRenderTargetTexture();
+			FRDGTexture* OutputTexture = GraphBuilder.RegisterExternalTexture(CreateRenderTarget(OutputRenderTarget, TEXT("WaterInfo.OutputTexture")));
+			GraphBuilder.UseInternalAccessMode(OutputTexture); // Make sure the texture is actually managed by RDG as a call to RegisterExternalTexture does not guarantee this if the texture was already registered and then had UseExternalAccessMode() called on it.
 
-		FRDGTextureDesc TerrainDepthBufferDesc = FRDGTextureDesc::Create2D(TextureSize, PF_DepthStencil, FClearValueBinding::DepthFar, TexCreate_DepthStencilTargetable | TexCreate_ShaderResource);
-		FRDGTexture* TerrainDepthBuffer = GraphBuilder.CreateTexture(TerrainDepthBufferDesc, TEXT("WaterInfo.TerrainDepthTexture"));
+			FRDGTextureDesc TerrainDepthBufferDesc = FRDGTextureDesc::Create2D(TextureSize, PF_DepthStencil, FClearValueBinding::DepthFar, TexCreate_DepthStencilTargetable | TexCreate_ShaderResource);
+			FRDGTexture* TerrainDepthBuffer = GraphBuilder.CreateTexture(TerrainDepthBufferDesc, TEXT("WaterInfo.TerrainDepthTexture"));
 
-		FRDGTextureDesc WaterInfoTextureDesc = FRDGTextureDesc::Create2D(TextureSize, PF_FloatRGBA, FClearValueBinding::Black, TexCreate_RenderTargetable | TexCreate_ShaderResource);
-		FRDGTexture* WaterInfoColorTexture = GraphBuilder.CreateTexture(WaterInfoTextureDesc, TEXT("WaterInfo.ColorTexture"));
+			FRDGTextureDesc WaterInfoTextureDesc = FRDGTextureDesc::Create2D(TextureSize, PF_FloatRGBA, FClearValueBinding::Black, TexCreate_RenderTargetable | TexCreate_ShaderResource);
+			FRDGTexture* WaterInfoColorTexture = GraphBuilder.CreateTexture(WaterInfoTextureDesc, TEXT("WaterInfo.ColorTexture"));
 
-		FRDGTextureDesc WaterInfoDepthBufferDesc = FRDGTextureDesc::Create2D(TextureSize, PF_DepthStencil, FClearValueBinding::DepthFar, TexCreate_DepthStencilTargetable | TexCreate_ShaderResource);
-		FRDGTexture* WaterInfoDepthBuffer = GraphBuilder.CreateTexture(WaterInfoDepthBufferDesc, TEXT("WaterInfo.DepthTexture"));
+			FRDGTextureDesc WaterInfoDepthBufferDesc = FRDGTextureDesc::Create2D(TextureSize, PF_DepthStencil, FClearValueBinding::DepthFar, TexCreate_DepthStencilTargetable | TexCreate_ShaderResource);
+			FRDGTexture* WaterInfoDepthBuffer = GraphBuilder.CreateTexture(WaterInfoDepthBufferDesc, TEXT("WaterInfo.DepthTexture"));
 
-		FRDGTextureDesc DilatedDepthBufferDesc = FRDGTextureDesc::Create2D(TextureSize, PF_DepthStencil, FClearValueBinding::DepthFar, TexCreate_DepthStencilTargetable | TexCreate_ShaderResource);
-		FRDGTexture* DilatedDepthBuffer = GraphBuilder.CreateTexture(DilatedDepthBufferDesc, TEXT("WaterInfo.DilatedDepthTexture"));
+			FRDGTextureDesc DilatedDepthBufferDesc = FRDGTextureDesc::Create2D(TextureSize, PF_DepthStencil, FClearValueBinding::DepthFar, TexCreate_DepthStencilTargetable | TexCreate_ShaderResource);
+			FRDGTexture* DilatedDepthBuffer = GraphBuilder.CreateTexture(DilatedDepthBufferDesc, TEXT("WaterInfo.DilatedDepthTexture"));
 
-		FRDGTextureDesc MergedTextureDesc = FRDGTextureDesc::Create2D(TextureSize, OutputTexture->Desc.Format, FClearValueBinding::Black, TexCreate_RenderTargetable | TexCreate_ShaderResource);
-		FRDGTexture* MergedTexture = GraphBuilder.CreateTexture(MergedTextureDesc, TEXT("WaterInfo.MergedTexture"));
+			FRDGTextureDesc MergedTextureDesc = FRDGTextureDesc::Create2D(TextureSize, OutputTexture->Desc.Format, FClearValueBinding::Black, TexCreate_RenderTargetable | TexCreate_ShaderResource);
+			FRDGTexture* MergedTexture = GraphBuilder.CreateTexture(MergedTextureDesc, TEXT("WaterInfo.MergedTexture"));
 
-		FGlobalShaderMap* ShaderMap = GetGlobalShaderMap(WaterView.GetFeatureLevel());
+			FGlobalShaderMap* ShaderMap = GetGlobalShaderMap(WaterView.GetFeatureLevel());
 
-		TStaticArray<FWaterInfoTextureDraws, static_cast<int32>(EWaterInfoTextureMeshPass::Num)> PassDraws;
-		PassDraws[static_cast<int32>(EWaterInfoTextureMeshPass::TerrainDepth)] =			FWaterInfoTextureDraws(EMeshPass::WaterInfoTextureDepthPass, &RenderingParams.TerrainComponentIds, 0);
-		PassDraws[static_cast<int32>(EWaterInfoTextureMeshPass::DilatedWaterBodyDepth)] =	FWaterInfoTextureDraws(EMeshPass::WaterInfoTextureDepthPass, &RenderingParams.DilatedWaterBodyComponentIds, 0);
-		PassDraws[static_cast<int32>(EWaterInfoTextureMeshPass::WaterBody)] =				FWaterInfoTextureDraws(EMeshPass::WaterInfoTexturePass, &RenderingParams.WaterBodyComponentIds, 0);
-		
-		AddWaterInfoTextureDraws(Scene, PassDraws);
+			TStaticArray<FWaterInfoTextureDraws, static_cast<int32>(EWaterInfoTextureMeshPass::Num)> PassDraws;
+			PassDraws[static_cast<int32>(EWaterInfoTextureMeshPass::TerrainDepth)] =			FWaterInfoTextureDraws(EMeshPass::WaterInfoTextureDepthPass, &RenderingParams.TerrainComponentIds, 0);
+			PassDraws[static_cast<int32>(EWaterInfoTextureMeshPass::DilatedWaterBodyDepth)] =	FWaterInfoTextureDraws(EMeshPass::WaterInfoTextureDepthPass, &RenderingParams.DilatedWaterBodyComponentIds, 0);
+			PassDraws[static_cast<int32>(EWaterInfoTextureMeshPass::WaterBody)] =				FWaterInfoTextureDraws(EMeshPass::WaterInfoTexturePass, &RenderingParams.WaterBodyComponentIds, 0);
+			
+			AddWaterInfoTextureDraws(Scene, PassDraws);
 
-		// Build rendering commands or prepare primitive ID buffer for the mesh draw commands
-		for (int32 PassIdx = 0; PassIdx < 3; ++PassIdx)
-		{
-			if (Scene->GPUScene.IsEnabled())
+			// Build rendering commands or prepare primitive ID buffer for the mesh draw commands
+			for (int32 PassIdx = 0; PassIdx < 3; ++PassIdx)
 			{
-				int32 MaxInstances = 0;
-				int32 VisibleMeshDrawCommandsNum = 0;
-				int32 NewPassVisibleMeshDrawCommandsNum = 0;
-
-				static FName NAME_WaterInfoTexturePass("WaterInfoTexture");
-				PassDraws[PassIdx].InstanceCullingContext = GraphBuilder.AllocObject<FInstanceCullingContext>(NAME_WaterInfoTexturePass, WaterView.GetShaderPlatform(), nullptr, TArrayView<const int32>(&WaterView.GPUSceneViewId, 1), nullptr);
+				if (Scene->GPUScene.IsEnabled())
+				{
+					int32 MaxInstances = 0;
+					int32 VisibleMeshDrawCommandsNum = 0;
+					int32 NewPassVisibleMeshDrawCommandsNum = 0;
 
-				PassDraws[PassIdx].InstanceCullingContext->SetupDrawCommands(PassDraws[PassIdx].VisibleMeshCommands, false, Scene, MaxInstances, VisibleMeshDrawCommandsNum, NewPassVisibleMeshDrawCommandsNum);
-				// Not supposed to do any compaction here.
-				ensure(VisibleMeshDrawCommandsNum == PassDraws[PassIdx].VisibleMeshCommands.Num());
+					static FName NAME_WaterInfoTexturePass("WaterInfoTexture");
+					PassDraws[PassIdx].InstanceCullingContext = GraphBuilder.AllocObject<FInstanceCullingContext>(NAME_WaterInfoTexturePass, WaterView.GetShaderPlatform(), nullptr, TArrayView<const int32>(&WaterView.GPUSceneViewId, 1), nullptr);
 
-				PassDraws[PassIdx].InstanceCullingContext->BuildRenderingCommands(GraphBuilder, Scene->GPUScene, WaterView.DynamicPrimitiveCollector.GetInstanceSceneDataOffset(), WaterView.DynamicPrimitiveCollector.NumInstances(), PassDraws[PassIdx].InstanceCullingResult);
-			}
+					PassDraws[PassIdx].InstanceCullingContext->SetupDrawCommands(PassDraws[PassIdx].VisibleMeshCommands, false, Scene, MaxInstances, VisibleMeshDrawCommandsNum, NewPassVisibleMeshDrawCommandsNum);
+					// Not supposed to do any compaction here.
+					ensure(VisibleMeshDrawCommandsNum == PassDraws[PassIdx].VisibleMeshCommands.Num());
 
-			PassDraws[PassIdx].InstanceCullingResult.Parameters.Scene = SceneRenderer.GetSceneUniforms().GetBuffer(GraphBuilder);
-		}
+					PassDraws[PassIdx].InstanceCullingContext->BuildRenderingCommands(GraphBuilder, Scene->GPUScene, WaterView.DynamicPrimitiveCollector.GetInstanceSceneDataOffset(), WaterView.DynamicPrimitiveCollector.NumInstances(), PassDraws[PassIdx].InstanceCullingResult);
+				}
 
-		// Execute the three mesh passes involved in generating the water info texture: terrain depth-only, dilated water body depth-only and finally regular water body depth + river velocity
-		for (uint32 PassIndex = 0; PassIndex < static_cast<int32>(EWaterInfoTextureMeshPass::Num); ++PassIndex)
-		{
-			const EWaterInfoTextureMeshPass PassType = static_cast<EWaterInfoTextureMeshPass>(PassIndex);
+				PassDraws[PassIdx].InstanceCullingResult.Parameters.Scene = SceneRenderer.GetSceneUniforms().GetBuffer(GraphBuilder);
+			}
 
-			// Terrain depth and dilated water body depth share the same setup, they just render a different set of meshes into different depth buffers
-			if (PassType == EWaterInfoTextureMeshPass::TerrainDepth || PassType == EWaterInfoTextureMeshPass::DilatedWaterBodyDepth)
+			// Execute the three mesh passes involved in generating the water info texture: terrain depth-only, dilated water body depth-only and finally regular water body depth + river velocity
+			for (uint32 PassIndex = 0; PassIndex < static_cast<int32>(EWaterInfoTextureMeshPass::Num); ++PassIndex)
 			{
-				const bool bIsTerrainDepthPass = PassType == EWaterInfoTextureMeshPass::TerrainDepth;
-				FWaterInfoTextureDepthPassParameters* PassParameters = GraphBuilder.AllocParameters<FWaterInfoTextureDepthPassParameters>();
-				PassParameters->View = WaterView.GetShaderParameters();
-				PassParameters->RenderTargets.DepthStencil = FDepthStencilBinding(bIsTerrainDepthPass ? TerrainDepthBuffer : DilatedDepthBuffer, ERenderTargetLoadAction::EClear, ERenderTargetLoadAction::ENoAction, FExclusiveDepthStencil::DepthWrite_StencilNop);
-				PassDraws[PassIndex].InstanceCullingResult.GetDrawParameters(PassParameters->InstanceCullingDrawParams);
-
-				GraphBuilder.AddPass(
-					bIsTerrainDepthPass ? RDG_EVENT_NAME("WaterInfoTexture(Terrain)") : RDG_EVENT_NAME("WaterInfoTexture(DilatedWaterBodies)"),
-					PassParameters,
-					ERDGPassFlags::Raster,
-					[MeshDrawCmds = MoveTemp(PassDraws[PassIndex].VisibleMeshCommands), Scene = Scene, Viewport,
-					PassParameters, InstanceCullingContext = PassDraws[PassIndex].InstanceCullingContext](FRHICommandList& RHICmdList)
-					{
-						QUICK_SCOPE_CYCLE_COUNTER(MeshPass);
+				const EWaterInfoTextureMeshPass PassType = static_cast<EWaterInfoTextureMeshPass>(PassIndex);
 
-						RHICmdList.SetViewport(Viewport.Min.X, Viewport.Min.Y, 0.0f, Viewport.Max.X, Viewport.Max.Y, 1.0f);
-						FGraphicsMinimalPipelineStateSet PipelineStateSet;
-						if (Scene->GPUScene.IsEnabled())
+				// Terrain depth and dilated water body depth share the same setup, they just render a different set of meshes into different depth buffers
+				if (PassType == EWaterInfoTextureMeshPass::TerrainDepth || PassType == EWaterInfoTextureMeshPass::DilatedWaterBodyDepth)
+				{
+					const bool bIsTerrainDepthPass = PassType == EWaterInfoTextureMeshPass::TerrainDepth;
+					FWaterInfoTextureDepthPassParameters* PassParameters = GraphBuilder.AllocParameters<FWaterInfoTextureDepthPassParameters>();
+					PassParameters->View = WaterView.GetShaderParameters();
+					PassParameters->RenderTargets.DepthStencil = FDepthStencilBinding(bIsTerrainDepthPass ? TerrainDepthBuffer : DilatedDepthBuffer, ERenderTargetLoadAction::EClear, ERenderTargetLoadAction::ENoAction, FExclusiveDepthStencil::DepthWrite_StencilNop);
+					PassDraws[PassIndex].InstanceCullingResult.GetDrawParameters(PassParameters->InstanceCullingDrawParams);
+
+					GraphBuilder.AddPass(
+						bIsTerrainDepthPass ? RDG_EVENT_NAME("WaterInfoTexture(Terrain)") : RDG_EVENT_NAME("WaterInfoTexture(DilatedWaterBodies)"),
+						PassParameters,
+						ERDGPassFlags::Raster,
+						[MeshDrawCmds = MoveTemp(PassDraws[PassIndex].VisibleMeshCommands), Scene = Scene, Viewport,
+						PassParameters, InstanceCullingContext = PassDraws[PassIndex].InstanceCullingContext](FRHICommandList& RHICmdList)
 						{
-							InstanceCullingContext->SubmitDrawCommands(MeshDrawCmds, PipelineStateSet, GetMeshDrawCommandOverrideArgs(PassParameters->InstanceCullingDrawParams), 0, MeshDrawCmds.Num(), 1, RHICmdList);
-						}
-						else
+							QUICK_SCOPE_CYCLE_COUNTER(MeshPass);
+
+							RHICmdList.SetViewport(Viewport.Min.X, Viewport.Min.Y, 0.0f, Viewport.Max.X, Viewport.Max.Y, 1.0f);
+							FGraphicsMinimalPipelineStateSet PipelineStateSet;
+							if (Scene->GPUScene.IsEnabled())
+							{
+								InstanceCullingContext->SubmitDrawCommands(MeshDrawCmds, PipelineStateSet, GetMeshDrawCommandOverrideArgs(PassParameters->InstanceCullingDrawParams), 0, MeshDrawCmds.Num(), 1, RHICmdList);
+							}
+							else
+							{
+								FMeshDrawCommandSceneArgs SceneArgs;
+								SubmitMeshDrawCommandsRange(MeshDrawCmds, PipelineStateSet, SceneArgs, FInstanceCullingContext::GetInstanceIdBufferStride(Scene->GetShaderPlatform()), false, 0, MeshDrawCmds.Num(), 1, RHICmdList);
+							}
+						});
+				}
+				else if (PassType == EWaterInfoTextureMeshPass::WaterBody)
+				{
+					FWaterInfoTexturePassParameters* PassParameters = GraphBuilder.AllocParameters<FWaterInfoTexturePassParameters>();
+					PassParameters->View = WaterView.GetShaderParameters();
+					PassParameters->RenderTargets = GetRenderTargetBindings(ERenderTargetLoadAction::EClear, MakeArrayView(&WaterInfoColorTexture, 1));
+					PassParameters->RenderTargets.DepthStencil = FDepthStencilBinding(WaterInfoDepthBuffer, ERenderTargetLoadAction::EClear, ERenderTargetLoadAction::ENoAction, FExclusiveDepthStencil::DepthWrite_StencilNop);
+					PassDraws[PassIndex].InstanceCullingResult.GetDrawParameters(PassParameters->InstanceCullingDrawParams);
+
+					GraphBuilder.AddPass(
+						RDG_EVENT_NAME("WaterInfoTexture(WaterBodies)"),
+						PassParameters,
+						ERDGPassFlags::Raster,
+						[MeshDrawCmds = MoveTemp(PassDraws[PassIndex].VisibleMeshCommands), Scene = Scene, Viewport,
+						PassParameters, InstanceCullingContext = PassDraws[PassIndex].InstanceCullingContext](FRHICommandList& RHICmdList)
 						{
-							FMeshDrawCommandSceneArgs SceneArgs;
-							SubmitMeshDrawCommandsRange(MeshDrawCmds, PipelineStateSet, SceneArgs, FInstanceCullingContext::GetInstanceIdBufferStride(Scene->GetShaderPlatform()), false, 0, MeshDrawCmds.Num(), 1, RHICmdList);
-						}
-					});
+							QUICK_SCOPE_CYCLE_COUNTER(MeshPass);
+
+							RHICmdList.SetViewport(Viewport.Min.X, Viewport.Min.Y, 0.0f, Viewport.Max.X, Viewport.Max.Y, 1.0f);
+							FGraphicsMinimalPipelineStateSet PipelineStateSet;
+							if (Scene->GPUScene.IsEnabled())
+							{
+								InstanceCullingContext->SubmitDrawCommands(MeshDrawCmds, PipelineStateSet, GetMeshDrawCommandOverrideArgs(PassParameters->InstanceCullingDrawParams), 0, MeshDrawCmds.Num(), 1, RHICmdList);
+							}
+							else
+							{
+								FMeshDrawCommandSceneArgs SceneArgs;
+								SubmitMeshDrawCommandsRange(MeshDrawCmds, PipelineStateSet, SceneArgs, FInstanceCullingContext::GetInstanceIdBufferStride(Scene->GetShaderPlatform()), false, 0, MeshDrawCmds.Num(), 1, RHICmdList);
+							}
+						});
+				}
+				
 			}
-			else if (PassType == EWaterInfoTextureMeshPass::WaterBody)
+
+			// Merge terrain depth, dilated water body depth, water body depth and velocity into a single texture
 			{
-				FWaterInfoTexturePassParameters* PassParameters = GraphBuilder.AllocParameters<FWaterInfoTexturePassParameters>();
-				PassParameters->View = WaterView.GetShaderParameters();
-				PassParameters->RenderTargets = GetRenderTargetBindings(ERenderTargetLoadAction::EClear, MakeArrayView(&WaterInfoColorTexture, 1));
-				PassParameters->RenderTargets.DepthStencil = FDepthStencilBinding(WaterInfoDepthBuffer, ERenderTargetLoadAction::EClear, ERenderTargetLoadAction::ENoAction, FExclusiveDepthStencil::DepthWrite_StencilNop);
-				PassDraws[PassIndex].InstanceCullingResult.GetDrawParameters(PassParameters->InstanceCullingDrawParams);
+				FWaterInfoTextureMergePS::FPermutationDomain PixelPermutationVector;
+				// The output format depends on the user configurable format of the passed in render target. Since we store depth data in the texture,
+				// it is sometimes desirable to have full 32bit float precision.
+				PixelPermutationVector.Set<FWaterInfoTextureMergePS::FEnable128BitRT>(OutputTexture->Desc.Format == PF_A32B32G32R32F);
+				TShaderMapRef<FWaterInfoTextureMergePS> PixelShader(ShaderMap, PixelPermutationVector);
 
-				GraphBuilder.AddPass(
-					RDG_EVENT_NAME("WaterInfoTexture(WaterBodies)"),
+				FWaterInfoTextureMergePS::FParameters* PassParameters = GraphBuilder.AllocParameters<FWaterInfoTextureMergePS::FParameters>();
+				PassParameters->View = WaterView.GetShaderParameters();
+				PassParameters->RenderTargets[0] = FRenderTargetBinding(MergedTexture, ERenderTargetLoadAction::ENoAction);
+				PassParameters->SceneTextures = GetSceneTextureShaderParameters(WaterView);
+				PassParameters->GroundDepthTexture = TerrainDepthBuffer;
+				PassParameters->GroundDepthTextureSampler = TStaticSamplerState<SF_Point, AM_Clamp, AM_Clamp, AM_Clamp>::GetRHI();
+				PassParameters->WaterBodyTexture = WaterInfoColorTexture;
+				PassParameters->WaterBodyTextureSampler = TStaticSamplerState<SF_Point, AM_Clamp, AM_Clamp, AM_Clamp>::GetRHI();
+				PassParameters->WaterBodyDepthTexture = WaterInfoDepthBuffer;
+				PassParameters->WaterBodyDepthTextureSampler = TStaticSamplerState<SF_Point, AM_Clamp, AM_Clamp, AM_Clamp>::GetRHI();
+				PassParameters->DilatedWaterBodyDepthTexture = DilatedDepthBuffer;
+				PassParameters->DilatedWaterBodyDepthTextureSampler = TStaticSamplerState<SF_Point, AM_Clamp, AM_Clamp, AM_Clamp>::GetRHI();
+				PassParameters->CaptureZ = RenderingParams.CaptureZ;
+				PassParameters->WaterHeightExtents = RenderingParams.WaterHeightExtents;
+				PassParameters->GroundZMin = RenderingParams.GroundZMin;
+				PassParameters->DilationOverwriteMinimumDistance = CVarWaterInfoDilationOverwriteMinimumDistance.GetValueOnRenderThread();
+				PassParameters->UndergroundDilationDepthOffset = CVarWaterInfoUndergroundDilationDepthOffset.GetValueOnRenderThread();
+
+				FPixelShaderUtils::AddFullscreenPass(
+					GraphBuilder,
+					ShaderMap,
+					RDG_EVENT_NAME("WaterInfoTextureMerge"),
+					PixelShader,
 					PassParameters,
-					ERDGPassFlags::Raster,
-					[MeshDrawCmds = MoveTemp(PassDraws[PassIndex].VisibleMeshCommands), Scene = Scene, Viewport,
-					PassParameters, InstanceCullingContext = PassDraws[PassIndex].InstanceCullingContext](FRHICommandList& RHICmdList)
-					{
-						QUICK_SCOPE_CYCLE_COUNTER(MeshPass);
+					Viewport);
+			}
 
-						RHICmdList.SetViewport(Viewport.Min.X, Viewport.Min.Y, 0.0f, Viewport.Max.X, Viewport.Max.Y, 1.0f);
-						FGraphicsMinimalPipelineStateSet PipelineStateSet;
-						if (Scene->GPUScene.IsEnabled())
-						{
-							InstanceCullingContext->SubmitDrawCommands(MeshDrawCmds, PipelineStateSet, GetMeshDrawCommandOverrideArgs(PassParameters->InstanceCullingDrawParams), 0, MeshDrawCmds.Num(), 1, RHICmdList);
-						}
-						else
-						{
-							FMeshDrawCommandSceneArgs SceneArgs;
-							SubmitMeshDrawCommandsRange(MeshDrawCmds, PipelineStateSet, SceneArgs, FInstanceCullingContext::GetInstanceIdBufferStride(Scene->GetShaderPlatform()), false, 0, MeshDrawCmds.Num(), 1, RHICmdList);
-						}
-					});
+			// Blur the velocity component in the water info texture to get a smooth gradient between water body transitions
+			{
+				FWaterInfoTextureBlurPS::FPermutationDomain PixelPermutationVector;
+				// The output format depends on the user configurable format of the passed in render target. Since we store depth data in the texture,
+				// it is sometimes desirable to have full 32bit float precision.
+				PixelPermutationVector.Set<FWaterInfoTextureBlurPS::FEnable128BitRT>(OutputTexture->Desc.Format == PF_A32B32G32R32F);
+				TShaderMapRef<FWaterInfoTextureBlurPS> PixelShader(ShaderMap, PixelPermutationVector);
+
+				FWaterInfoTextureBlurPS::FParameters* PassParameters = GraphBuilder.AllocParameters<FWaterInfoTextureBlurPS::FParameters>();
+				PassParameters->View = WaterView.GetShaderParameters();
+				PassParameters->RenderTargets[0] = FRenderTargetBinding(OutputTexture, ERenderTargetLoadAction::ENoAction, 0, RenderingParams.RenderTargetArrayLayer);
+				PassParameters->SceneTextures = GetSceneTextureShaderParameters(WaterView);
+				PassParameters->WaterInfoTexture = MergedTexture;
+				PassParameters->WaterZMin = RenderingParams.WaterHeightExtents.X;
+				PassParameters->WaterZMax = RenderingParams.WaterHeightExtents.Y;
+				PassParameters->GroundZMin = RenderingParams.GroundZMin;
+				PassParameters->CaptureZ = RenderingParams.CaptureZ;
+				PassParameters->BlurRadius = RenderingParams.VelocityBlurRadius;
+
+				FPixelShaderUtils::AddFullscreenPass(
+					GraphBuilder,
+					ShaderMap,
+					RDG_EVENT_NAME("WaterInfoTextureBlur"),
+					PixelShader,
+					PassParameters,
+					Viewport);
 			}
-			
-		}
 
-		// Merge terrain depth, dilated water body depth, water body depth and velocity into a single texture
-		{
-			FWaterInfoTextureMergePS::FPermutationDomain PixelPermutationVector;
-			// The output format depends on the user configurable format of the passed in render target. Since we store depth data in the texture,
-			// it is sometimes desirable to have full 32bit float precision.
-			PixelPermutationVector.Set<FWaterInfoTextureMergePS::FEnable128BitRT>(OutputTexture->Desc.Format == PF_A32B32G32R32F);
-			TShaderMapRef<FWaterInfoTextureMergePS> PixelShader(ShaderMap, PixelPermutationVector);
-
-			FWaterInfoTextureMergePS::FParameters* PassParameters = GraphBuilder.AllocParameters<FWaterInfoTextureMergePS::FParameters>();
-			PassParameters->View = WaterView.GetShaderParameters();
-			PassParameters->RenderTargets[0] = FRenderTargetBinding(MergedTexture, ERenderTargetLoadAction::ENoAction);
-			PassParameters->SceneTextures = GetSceneTextureShaderParameters(WaterView);
-			PassParameters->GroundDepthTexture = TerrainDepthBuffer;
-			PassParameters->GroundDepthTextureSampler = TStaticSamplerState<SF_Point, AM_Clamp, AM_Clamp, AM_Clamp>::GetRHI();
-			PassParameters->WaterBodyTexture = WaterInfoColorTexture;
-			PassParameters->WaterBodyTextureSampler = TStaticSamplerState<SF_Point, AM_Clamp, AM_Clamp, AM_Clamp>::GetRHI();
-			PassParameters->WaterBodyDepthTexture = WaterInfoDepthBuffer;
-			PassParameters->WaterBodyDepthTextureSampler = TStaticSamplerState<SF_Point, AM_Clamp, AM_Clamp, AM_Clamp>::GetRHI();
-			PassParameters->DilatedWaterBodyDepthTexture = DilatedDepthBuffer;
-			PassParameters->DilatedWaterBodyDepthTextureSampler = TStaticSamplerState<SF_Point, AM_Clamp, AM_Clamp, AM_Clamp>::GetRHI();
-			PassParameters->CaptureZ = RenderingParams.CaptureZ;
-			PassParameters->WaterHeightExtents = RenderingParams.WaterHeightExtents;
-			PassParameters->GroundZMin = RenderingParams.GroundZMin;
-			PassParameters->DilationOverwriteMinimumDistance = CVarWaterInfoDilationOverwriteMinimumDistance.GetValueOnRenderThread();
-			PassParameters->UndergroundDilationDepthOffset = CVarWaterInfoUndergroundDilationDepthOffset.GetValueOnRenderThread();
-
-			FPixelShaderUtils::AddFullscreenPass(
-				GraphBuilder,
-				ShaderMap,
-				RDG_EVENT_NAME("WaterInfoTextureMerge"),
-				PixelShader,
-				PassParameters,
-				Viewport);
+			// Add the output texture to an array to batch call UseExternalAccessMode() on all of them.
+			OutputTextures.AddUnique(OutputTexture);
 		}
+	}
 
-		// Blur the velocity component in the water info texture to get a smooth gradient between water body transitions
-		{
-			FWaterInfoTextureBlurPS::FPermutationDomain PixelPermutationVector;
-			// The output format depends on the user configurable format of the passed in render target. Since we store depth data in the texture,
-			// it is sometimes desirable to have full 32bit float precision.
-			PixelPermutationVector.Set<FWaterInfoTextureBlurPS::FEnable128BitRT>(OutputTexture->Desc.Format == PF_A32B32G32R32F);
-			TShaderMapRef<FWaterInfoTextureBlurPS> PixelShader(ShaderMap, PixelPermutationVector);
-
-			FWaterInfoTextureBlurPS::FParameters* PassParameters = GraphBuilder.AllocParameters<FWaterInfoTextureBlurPS::FParameters>();
-			PassParameters->View = WaterView.GetShaderParameters();
-			PassParameters->RenderTargets[0] = FRenderTargetBinding(OutputTexture, ERenderTargetLoadAction::ENoAction, 0, RenderingParams.RenderTargetArrayLayer);
-			PassParameters->SceneTextures = GetSceneTextureShaderParameters(WaterView);
-			PassParameters->WaterInfoTexture = MergedTexture;
-			PassParameters->WaterZMin = RenderingParams.WaterHeightExtents.X;
-			PassParameters->WaterZMax = RenderingParams.WaterHeightExtents.Y;
-			PassParameters->GroundZMin = RenderingParams.GroundZMin;
-			PassParameters->CaptureZ = RenderingParams.CaptureZ;
-			PassParameters->BlurRadius = RenderingParams.VelocityBlurRadius;
-
-			FPixelShaderUtils::AddFullscreenPass(
-				GraphBuilder,
-				ShaderMap,
-				RDG_EVENT_NAME("WaterInfoTextureBlur"),
-				PixelShader,
-				PassParameters,
-				Viewport);
-		}
+	// Make sure the textures are in the SRV state by the time they are used in water draws (referenced outside of RDG).
+	GraphBuilder.UseExternalAccessMode(OutputTextures, ERHIAccess::SRVMask);
 
-		// Make sure the texture is in the SRV state by the time it is used in water draws (referenced outside of RDG).
-		GraphBuilder.UseExternalAccessMode(OutputTexture, ERHIAccess::SRVMask);
+	// Reset the WIT updates on all the views to ensure we don't accidentally keep executing the same updates.
+	for (FViewInfo& View : SceneRenderer.Views)
+	{
+		View.WaterInfoTextureRenderingParams.Reset();
 	}
 }
