HeterogeneousVolumes: Move live-shading shadow map creation into ShadowDepthRendering.cpp

This prevents a crash when collecting HeterogeneousVolume subjects later in the frame.

diff --git a/Engine/Source/Runtime/Renderer/Private/HeterogeneousVolumes/HeterogeneousVolumes.cpp b/Engine/Source/Runtime/Renderer/Private/HeterogeneousVolumes/HeterogeneousVolumes.cpp
--- a/Engine/Source/Runtime/Renderer/Private/HeterogeneousVolumes/HeterogeneousVolumes.cpp
+++ b/Engine/Source/Runtime/Renderer/Private/HeterogeneousVolumes/HeterogeneousVolumes.cpp
@@ -727,6 +727,32 @@ bool ShouldCacheVoxelGrids(const FScene* Scene, FSceneViewState* ViewState)
 	return false;
 }
 
+void RenderHeterogeneousVolumeShadows(
+	FRDGBuilder& GraphBuilder,
+	FScene* Scene,
+	const FSceneTextures& SceneTextures,
+	FViewInfo& View,
+	const FSceneViewFamily& ViewFamily,
+	TArray<FVisibleLightInfo, SceneRenderingAllocator>& VisibleLightInfos
+)
+{
+	RDG_EVENT_SCOPE(GraphBuilder, "HeterogeneousVolumeShadows");
+	RDG_GPU_STAT_SCOPE(GraphBuilder, HeterogeneousVolumeShadowsStat);
+	SCOPED_NAMED_EVENT(HeterogeneousVolumes, FColor::Emerald);
+
+	if (HeterogeneousVolumes::GetShadowMode() == HeterogeneousVolumes::EShadowMode::LiveShading)
+	{
+		RenderAdaptiveVolumetricShadowMapWithLiveShading(
+			GraphBuilder,
+			SceneTextures,
+			Scene,
+			ViewFamily,
+			View,
+			VisibleLightInfos
+		);
+	}
+}
+
 void FDeferredShadingSceneRenderer::RenderHeterogeneousVolumeShadows(
 	FRDGBuilder& GraphBuilder,
 	const FSceneTextures& SceneTextures
@@ -793,6 +819,8 @@ void FDeferredShadingSceneRenderer::RenderHeterogeneousVolumeShadows(
 		{
 			if (HeterogeneousVolumes::GetShadowMode() == HeterogeneousVolumes::EShadowMode::LiveShading)
 			{
+				// This path is taken care of now in ShadowDepthRendering
+#if 0
 				RenderAdaptiveVolumetricShadowMapWithLiveShading(
 					GraphBuilder,
 					// Scene data
@@ -803,6 +831,7 @@ void FDeferredShadingSceneRenderer::RenderHeterogeneousVolumeShadows(
 					// Light data
 					VisibleLightInfos
 				);
+#endif
 			}
 			else
 			{
diff --git a/Engine/Source/Runtime/Renderer/Private/ShadowDepthRendering.cpp b/Engine/Source/Runtime/Renderer/Private/ShadowDepthRendering.cpp
--- a/Engine/Source/Runtime/Renderer/Private/ShadowDepthRendering.cpp
+++ b/Engine/Source/Runtime/Renderer/Private/ShadowDepthRendering.cpp
@@ -47,6 +47,7 @@
 #include "Rendering/NaniteResources.h"
 #include "Rendering/NaniteStreamingManager.h"
 #include "Shadows/ShadowSceneRenderer.h"
+#include "HeterogeneousVolumes/HeterogeneousVolumes.h"
 #include "RenderCore.h"
 #include "ShaderPlatformCachedIniValue.h"
 #include "UnrealEngine.h"
@@ -1931,6 +1932,21 @@ void FSceneRenderer::RenderShadowDepthMaps(FRDGBuilder& GraphBuilder, FDynamicSh
 		InstanceCullingManager.BeginDeferredCulling(GraphBuilder, Scene->GPUScene);
 	}
 
+	if (ShouldRenderHeterogeneousVolumes(Scene) && ShouldHeterogeneousVolumesCastShadows())
+	{
+		for (int32 ViewIndex = 0; ViewIndex < Views.Num(); ViewIndex++)
+		{
+			FViewInfo& View = Views[ViewIndex];
+			RenderAdaptiveVolumetricShadowMapWithLiveShading(
+				GraphBuilder,
+				View.GetSceneTextures(),
+				Scene,
+				ViewFamily,
+				View,
+				VisibleLightInfos
+			);
+		}
+	}
 
 	bShadowDepthRenderCompleted = true;
 }
