Perform PresizeSubjectPrimitiveArrays of whole scene shadows once per task instead of redundantly per packet. Thanks to CDPR for this contribution.

diff --git a/Engine/Source/Runtime/Renderer/Private/ShadowSetup.cpp b/Engine/Source/Runtime/Renderer/Private/ShadowSetup.cpp
--- a/Engine/Source/Runtime/Renderer/Private/ShadowSetup.cpp
+++ b/Engine/Source/Runtime/Renderer/Private/ShadowSetup.cpp
@@ -1384,6 +1384,14 @@ struct FDynamicShadowsTaskData
 		return *MeshCollectors.Last();
 	}
 
+	void PresizeViewDependentWholeSceneShadows()
+	{
+		for (int32 ShadowIdx = 0, Num = ViewDependentWholeSceneShadows.Num(); ShadowIdx < Num; ++ShadowIdx)
+		{
+			ViewDependentWholeSceneShadows[ShadowIdx]->PresizeSubjectPrimitiveArrays(GatherStats[ShadowIdx]);
+		}
+	}
+	
 	FDynamicShadowsTaskData(const FDynamicShadowsTaskData&) = delete;
 
 	FDynamicShadowsTaskData(FRHICommandListImmediate& InRHICmdList, FSceneRenderer* InSceneRenderer, FInstanceCullingManager& InInstanceCullingManager, bool bInRunningEarly)
@@ -5013,11 +5021,6 @@ struct FGatherShadowPrimitivesPacket
 
 	void AnyThreadFinalize(FDynamicShadowsTaskData& TaskData)
 	{
-		for (int32 ShadowIdx = 0, Num = TaskData.ViewDependentWholeSceneShadows.Num(); ShadowIdx < Num; ++ShadowIdx)
-		{
-			TaskData.ViewDependentWholeSceneShadows[ShadowIdx]->PresizeSubjectPrimitiveArrays(TaskData.GatherStats[ShadowIdx]);
-		}
-
 		for (int32 ShadowIndex = 0; ShadowIndex < PreShadowSubjectPrimitives.Num(); ShadowIndex++)
 		{
 			FProjectedShadowInfo* ProjectedShadowInfo = TaskData.PreShadows[ShadowIndex];
@@ -5135,6 +5138,9 @@ struct FGatherShadowPrimitivesPrepareTask
 					SCOPED_NAMED_EVENT_TEXT("FGatherShadowPrimitivesFinalizeTask", FColor::Green);
 					FOptionalTaskTagScope Scope(ETaskTag::EParallelRenderingThread);
 
+					// Reserve only once up to total number of elements gathered from stats to avoid reallocations, especially important for really dense scenes
+					TaskData->PresizeViewDependentWholeSceneShadows();
+					
 					for (FGatherShadowPrimitivesPacket* Packet : TaskData->Packets)
 					{
 						Packet->AnyThreadFinalize(*TaskData);
@@ -5322,6 +5328,10 @@ void FSceneRenderer::FinishGatherShadowPrimitives(FDynamicShadowsTaskData* TaskD
 	if (!TaskData->bMultithreadedCreateAndFilterShadows)
 	{
 		SCOPED_NAMED_EVENT_TEXT("FGatherShadowPrimitivesTask", FColor::Green);
+		
+		// Reserve only once up to total number of elements gathered from stats to avoid reallocations, especially important for really dense scenes
+		TaskData->PresizeViewDependentWholeSceneShadows();
+		
 		for (FGatherShadowPrimitivesPacket* Packet : TaskData->Packets)
 		{
 			Packet->AnyThreadFinalize(*TaskData);
