Added experimental froxel generation and VSM page marking from froxels. Saving some 100us in many cases. Enable using r.Shadow.Virtual.MarkPagesUsingFroxels.
 - Froxels are generated from HZB (for free) as well as water depth prepass (not free, but no cost increase).
 - Froxels are 8x8 pixels and as cube-like as possible.
 - Not compatible with page marking from translucent (falls back to per-pixel).

diff --git a/Engine/Shaders/Private/Froxel/Froxel.ush b/Engine/Shaders/Private/Froxel/Froxel.ush
new file mode 100644
--- /dev/null
+++ b/Engine/Shaders/Private/Froxel/Froxel.ush
@@ -0,0 +1,163 @@
+// Copyright Epic Games, Inc. All Rights Reserved.
+
+#include "../Common.ush"
+#include "/Engine/Shared/FroxelDefinitions.h"
+
+// Froxel grid constants
+float4 FroxelToViewScaleBias;
+float4 FroxelToClipScaleBias;
+float2 FroxelClipToViewScale;
+float FroxelRecLog2DepthScale1; // 1 / log2(1.0 + DepthScale)
+float FroxelRecNearScale;
+float FroxelDepthScale1;
+float FroxelNear;
+float FroxelViewToClipTransformScale;
+float FroxelClipToViewTransformScale;
+float FroxelClipToViewTransformBias;
+float FroxelRadius;
+float FroxelInvRadius;
+
+uint bFroxelIsOrtho;
+uint FroxelArgsStride;
+uint FroxelArgsOffset;
+
+int CalcFroxelSliceFromDeviceZ(float DeviceZ)
+{
+	float DeviceZScaleBias = DeviceZ * FroxelClipToViewTransformScale + FroxelClipToViewTransformBias;
+
+	if (bFroxelIsOrtho)
+	{
+		int K = int(floor((DeviceZScaleBias - FroxelNear) * FroxelInvRadius));
+		return K;
+	}
+
+	// Special case for infinity
+	if (DeviceZ == 0.0f)
+	{
+		return FROXEL_INVALID_SLICE;
+	}
+
+	int K = int(floor(log2(rcp(DeviceZScaleBias) * FroxelRecNearScale) * FroxelRecLog2DepthScale1));
+	return K;
+}
+
+float GetViewDepthFromFroxelSlice(float SliceZ)
+{
+	if (bFroxelIsOrtho)
+	{
+		return SliceZ * FroxelRadius + FroxelNear;
+	}
+	return FroxelNear * pow(FroxelDepthScale1, SliceZ);
+}
+
+float GetViewDepthFromFroxelSlice(int SliceZ)
+{
+	return GetViewDepthFromFroxelSlice(float(SliceZ));
+}
+
+float GetClipDepthFromFroxelSlice(float SliceZ)
+{
+	float ViewDepth = GetViewDepthFromFroxelSlice(SliceZ);
+
+	float TmpViewDepth = ViewDepth - FroxelClipToViewTransformBias;
+	if (!bFroxelIsOrtho)
+	{
+		TmpViewDepth = rcp(TmpViewDepth);
+	}
+	return TmpViewDepth * FroxelViewToClipTransformScale;
+}
+
+FPackedFroxel PackFroxel(int3 Froxel)
+{
+	FPackedFroxel Result;
+	Result.XY = uint(Froxel.x) | (uint(Froxel.y) << 16u);
+	Result.Z = Froxel.z;
+	return Result;
+}
+
+int3 UnpackFroxel(FPackedFroxel PackedFroxel)
+{
+	int3 Result;
+	Result.x = int(BitFieldExtractU32(PackedFroxel.XY, 16, 0));
+	Result.y = int(BitFieldExtractU32(PackedFroxel.XY, 16, 16));
+	Result.z = PackedFroxel.Z;
+
+	return Result;
+}
+
+float3 GetFroxelViewSpaceCenter(int3 Froxel)
+{
+	float ViewSpaceCenterZ = GetViewDepthFromFroxelSlice(float(Froxel.z) + 0.5f);
+	float2 TmpXY = (float2(Froxel.xy) + 0.5f) * FroxelToViewScaleBias.xy + FroxelToViewScaleBias.zw;
+	if (bFroxelIsOrtho)
+	{
+		return float3(TmpXY, ViewSpaceCenterZ);
+	}
+	return float3(TmpXY * ViewSpaceCenterZ, ViewSpaceCenterZ);
+}
+
+struct FFroxelViewBounds
+{
+	float3 Min;
+	float3 Max;
+};
+
+FFroxelViewBounds GetFroxelViewSpaceAABB(int3 Froxel)
+{
+	float2 MinXY = ((float2(Froxel.xy)) * FroxelToViewScaleBias.xy + FroxelToViewScaleBias.zw);
+	float2 MaxXY = MinXY + FroxelToViewScaleBias.xy;
+
+	float MinViewZ = GetViewDepthFromFroxelSlice(Froxel.z);
+	FFroxelViewBounds Result;
+	if (bFroxelIsOrtho)
+	{
+		Result.Min.xy = MinXY;
+		Result.Min.z = MinViewZ;
+		Result.Max.xy = MaxXY;
+		Result.Max.z = MinViewZ + FroxelRadius; 
+		return Result;
+	}
+
+	if (abs(MinXY.x) > abs(MaxXY.x))
+	{
+		Swap(MinXY.x, MaxXY.x);
+	}
+	if (abs(MinXY.y) > abs(MaxXY.y))
+	{
+		Swap(MinXY.y, MaxXY.y);
+	}
+
+	float MaxViewZ = MinViewZ * FroxelDepthScale1;
+
+	// Extrema corners in X plane near and far
+	float2 ExtNear = MinXY * MinViewZ;
+	float2 ExtFar = MaxXY * MaxViewZ; // == 
+
+	Result.Min.xy = min(ExtNear, ExtFar);
+	Result.Min.z = MinViewZ;
+	Result.Max.xy = max(ExtNear, ExtFar);
+	Result.Max.z = MaxViewZ; 
+
+	return Result;
+}
+
+struct FFroxelClipBounds
+{
+	float3 ClipMin;
+	float3 ClipMax;
+};
+
+FFroxelClipBounds GetFroxelClipBounds(int3 Froxel)
+{
+	float MinDeviceZ = GetClipDepthFromFroxelSlice(Froxel.z);
+	float MaxDeviceZ = GetClipDepthFromFroxelSlice(Froxel.z + 1.0f);
+
+	float2 MinClipSpaceXY = ((float2(Froxel.xy)) * FroxelToClipScaleBias.xy + FroxelToClipScaleBias.zw);
+	float2 MaxClipSpaceXY = MinClipSpaceXY + FroxelToClipScaleBias.xy;
+
+	FFroxelClipBounds Result;
+	Result.ClipMin = float3(MinClipSpaceXY, MinDeviceZ);
+	Result.ClipMax = float3(MaxClipSpaceXY, MaxDeviceZ);
+	
+	return Result;
+}
\ No newline at end of file
diff --git a/Engine/Shaders/Private/Froxel/FroxelBuild.ush b/Engine/Shaders/Private/Froxel/FroxelBuild.ush
new file mode 100644
--- /dev/null
+++ b/Engine/Shaders/Private/Froxel/FroxelBuild.ush
@@ -0,0 +1,230 @@
+// Copyright Epic Games, Inc. All Rights Reserved.
+
+#include "../Common.ush"
+#include "Froxel.ush"
+#include "../WaveOpUtil.ush"
+
+#if !defined(FROXEL_HASH_BUFFER_SIZE) || !defined(FROXEL_SHARED_HASH_BUFFER_VAR) || !defined(FROXEL_HASH_THREAD_GROUP_SIZE)
+#error "Must define FROXEL_HASH_BUFFER_SIZE and FROXEL_SHARED_HASH_BUFFER_VAR and FROXEL_HASH_THREAD_GROUP_SIZE to use hash froxel builder"
+#endif
+
+#if FROXEL_HASH_THREAD_GROUP_SIZE != FROXEL_HASH_BUFFER_SIZE
+#error "FROXEL_HASH_THREAD_GROUP_SIZE and FROXEL_HASH_BUFFER_SIZE must be the same size"
+#endif
+
+RWStructuredBuffer<FPackedFroxel> OutFroxels;
+RWBuffer<uint> OutFroxelArgs;
+
+groupshared int GroupShared_FroxelAnyTrue;
+
+bool GroupActiveAnyTrue(bool bValue)
+{
+#if FEATURE_LEVEL >= FEATURE_LEVEL_SM6 || PLATFORM_SUPPORTS_SM6_0_WAVE_OPERATIONS
+	bool bWaveValue = WaveActiveAnyTrue(bValue);
+	const uint LaneCount = WaveGetLaneCount();
+	// enough to just use the vote, skip further reductions
+	if (LaneCount >= FROXEL_HASH_THREAD_GROUP_SIZE)
+	{
+		return bWaveValue;
+	}
+	// no need to write to group shared from any other lanes.
+	if (!WaveIsFirstLane())
+	{
+		bWaveValue = false;
+	}
+#else
+	bool bWaveValue = bValue;
+#endif
+	
+	// use group shared variable to broadcast result.
+	GroupShared_FroxelAnyTrue = 0;
+	GroupMemoryBarrierWithGroupSync();
+	if (bWaveValue)
+	{
+		GroupShared_FroxelAnyTrue = 1;
+	}
+	GroupMemoryBarrierWithGroupSync();
+	bool bResult = GroupShared_FroxelAnyTrue != 0;
+	GroupMemoryBarrierWithGroupSync();
+	return bResult;
+}
+
+void EmitFroxel(int3 Froxel)
+{
+	// write directly to global buffer
+	uint WriteOffset = 0;
+	WaveInterlockedAddScalarInGroups(OutFroxelArgs[FroxelArgsOffset * FroxelArgsStride + 3], OutFroxelArgs[FroxelArgsOffset * FroxelArgsStride], 64, 1, WriteOffset);
+	OutFroxels[WriteOffset] = PackFroxel(Froxel);
+}
+
+void EmitValidFroxel(int3 Froxel)
+{
+	if (Froxel.z != FROXEL_INVALID_SLICE)
+	{
+		EmitFroxel(Froxel);
+	}
+}
+
+void WriteHashedSlice(uint PackedSliceId)
+{
+	if (PackedSliceId < FROXEL_INVALID_PACKED_SLICE)
+	{
+		FROXEL_SHARED_HASH_BUFFER_VAR[PackedSliceId % FROXEL_HASH_BUFFER_SIZE] = PackedSliceId;
+	}
+}
+
+void CheckHashedSlice(inout uint PackedSliceId, inout bool bAnyValidLeft)
+{
+	if (PackedSliceId < FROXEL_INVALID_PACKED_SLICE)
+	{
+		if (FROXEL_SHARED_HASH_BUFFER_VAR[PackedSliceId % FROXEL_HASH_BUFFER_SIZE] == PackedSliceId)
+		{
+			PackedSliceId = FROXEL_INVALID_PACKED_SLICE;
+		}
+	}
+	if (PackedSliceId < FROXEL_INVALID_PACKED_SLICE)
+	{
+		bAnyValidLeft = true;
+	}
+}
+
+groupshared int GroupShared_FroxelWorkLeft;
+void HashReduceFroxels(int MaxValid, uint GroupThreadIndex, uint LocalLinearTileId, uint2 GroupTileOffset, uint4 SortedSlices)
+{
+	// pack the slice IDs
+	SortedSlices.x =	 (uint(SortedSlices.x) << 2u) | LocalLinearTileId;
+	if (MaxValid > 1)
+	{
+		SortedSlices.y = (uint(SortedSlices.y) << 2u) | LocalLinearTileId;
+	}
+	if (MaxValid > 2)
+	{
+		SortedSlices.z = (uint(SortedSlices.z) << 2u) | LocalLinearTileId;
+	}
+	if (MaxValid > 3)
+	{
+		SortedSlices.w = (uint(SortedSlices.w) << 2u) | LocalLinearTileId;
+	}
+	
+	LOOP
+	for (int Iter = 0; Iter < 4; ++Iter)
+	{
+		FROXEL_SHARED_HASH_BUFFER_VAR[GroupThreadIndex] = FROXEL_INVALID_PACKED_SLICE;
+		GroupMemoryBarrierWithGroupSync();
+		GroupShared_FroxelWorkLeft = 0;
+
+		// Insert the unique depth slices into the hash. We do all at once because if they don't collide we can reduce the number of barriers and do more in each iteration.
+		WriteHashedSlice(SortedSlices.x);
+		if (MaxValid > 1)
+		{
+			WriteHashedSlice(SortedSlices.y);
+		}
+		if (MaxValid > 2)
+		{
+			WriteHashedSlice(SortedSlices.z);
+		}
+		if (MaxValid > 3)
+		{
+			WriteHashedSlice(SortedSlices.w);
+		}
+		GroupMemoryBarrierWithGroupSync();
+
+		// 2. Look up and clear own ID
+		bool bAnyValidLeft = false;
+		CheckHashedSlice(SortedSlices.x, bAnyValidLeft);
+		if (MaxValid > 1)
+		{
+			CheckHashedSlice(SortedSlices.y, bAnyValidLeft);
+		}
+		if (MaxValid > 2)
+		{
+			CheckHashedSlice(SortedSlices.z, bAnyValidLeft);
+		}
+		if (MaxValid > 3)
+		{
+			CheckHashedSlice(SortedSlices.w, bAnyValidLeft);
+		}
+		// Load thread coherently and emit
+		uint PackedId = FROXEL_SHARED_HASH_BUFFER_VAR[GroupThreadIndex];
+		if (PackedId < FROXEL_INVALID_PACKED_SLICE)
+		{
+			uint2 FinishedLocalTileId = uint2(PackedId & 1u, (PackedId & 2u) >> 1u);
+			int3 FinishedFroxel = int3(GroupTileOffset + FinishedLocalTileId, int(PackedId >> 2u) - int(FROXEL_PACKED_SLICE_BIAS));
+
+			EmitFroxel(FinishedFroxel);
+		}
+
+		if (GroupActiveAnyTrue(bAnyValidLeft))
+		{
+			break;
+		}
+	}
+}
+
+void CompSwap(inout uint A, inout uint B)
+{
+	if (A > B)
+	{
+		Swap(A, B);
+	}
+}
+
+void HashBuildFroxelsFromDeviceZ(float4 DeviceZ, uint GroupThreadIndex, uint LocalLinearTileId, uint2 GroupTileOffset)
+{
+	int4 Slice2x2;
+	Slice2x2.x = CalcFroxelSliceFromDeviceZ(DeviceZ.x);
+	Slice2x2.y = CalcFroxelSliceFromDeviceZ(DeviceZ.y);
+	Slice2x2.z = CalcFroxelSliceFromDeviceZ(DeviceZ.z);
+	Slice2x2.w = CalcFroxelSliceFromDeviceZ(DeviceZ.w);
+
+	uint4 SortedSlices = uint4(Slice2x2 + FROXEL_PACKED_SLICE_BIAS);
+
+	// Sort 
+	CompSwap(SortedSlices.x, SortedSlices.z);
+	CompSwap(SortedSlices.y, SortedSlices.w);
+	CompSwap(SortedSlices.x, SortedSlices.y);
+	CompSwap(SortedSlices.z, SortedSlices.w);
+	CompSwap(SortedSlices.y, SortedSlices.z);
+
+	// make invalid those that are the same as neighbor (must go in descending order)
+	SortedSlices.w = select(SortedSlices.w != SortedSlices.z, SortedSlices.w, FROXEL_INVALID_SLICE);
+	SortedSlices.z = select(SortedSlices.z != SortedSlices.y, SortedSlices.z, FROXEL_INVALID_SLICE);
+	SortedSlices.y = select(SortedSlices.y != SortedSlices.x, SortedSlices.y, FROXEL_INVALID_SLICE);
+
+	// re-sort the three last ones (x does not change as per the above)
+	CompSwap(SortedSlices.y, SortedSlices.w);
+	CompSwap(SortedSlices.y, SortedSlices.z);
+	CompSwap(SortedSlices.z, SortedSlices.w);
+
+	// Skip to wave-optimized version for remaining number of items per thread.
+    if (GroupActiveAnyTrue(SortedSlices.w < FROXEL_INVALID_SLICE))
+	{
+		HashReduceFroxels(4, GroupThreadIndex, LocalLinearTileId, GroupTileOffset, SortedSlices);
+	}
+	// TODO: Wave vote not safe since we need whole wave active in the reduction phase, need full group, so only safe if group size == wave size.
+	else if (GroupActiveAnyTrue(SortedSlices.z < FROXEL_INVALID_SLICE))
+	{
+		HashReduceFroxels(3, GroupThreadIndex, LocalLinearTileId, GroupTileOffset, SortedSlices);
+	}
+	else if (GroupActiveAnyTrue(SortedSlices.y < FROXEL_INVALID_SLICE))
+	{
+		HashReduceFroxels(2, GroupThreadIndex, LocalLinearTileId, GroupTileOffset, SortedSlices);
+	}
+	else if (GroupActiveAnyTrue(SortedSlices.x < FROXEL_INVALID_SLICE))
+	{
+		HashReduceFroxels(1, GroupThreadIndex, LocalLinearTileId, GroupTileOffset, SortedSlices);
+	}
+}
+
+void HashBuildFroxelsFromDeviceZ(float DeviceZ, uint GroupThreadIndex, uint LocalLinearTileId, uint2 GroupTileOffset)
+{
+	int4 SortedSlices;
+	// Sorted by definition so no need to sort
+	SortedSlices.x = CalcFroxelSliceFromDeviceZ(DeviceZ) + FROXEL_PACKED_SLICE_BIAS;
+	SortedSlices.y = SortedSlices.z = SortedSlices.w = FROXEL_INVALID_SLICE;
+	// TODO: Wave vote not safe since we need whole wave active in the reduction phase, need full group, so only safe if group size == wave size.
+    if (GroupActiveAnyTrue(SortedSlices.x < FROXEL_INVALID_SLICE))
+	{
+		HashReduceFroxels(1, GroupThreadIndex, LocalLinearTileId, GroupTileOffset, SortedSlices);
+	}
+}
diff --git a/Engine/Shaders/Private/HZB.usf b/Engine/Shaders/Private/HZB.usf
--- a/Engine/Shaders/Private/HZB.usf
+++ b/Engine/Shaders/Private/HZB.usf
@@ -44,7 +44,7 @@ RWTexture2D<float> ClosestHZBOutput_1;
 RWTexture2D<float> ClosestHZBOutput_2;
 RWTexture2D<float> ClosestHZBOutput_3;
 
-groupshared float SharedMinDeviceZ[GROUP_TILE_SIZE * GROUP_TILE_SIZE];
+groupshared uint SharedMinDeviceZ[GROUP_TILE_SIZE * GROUP_TILE_SIZE];
 groupshared float SharedMaxDeviceZ[GROUP_TILE_SIZE * GROUP_TILE_SIZE];
 
 float RoundUpF16(float DeviceZ)
@@ -155,6 +155,23 @@ float4 Gather4VisZ(float2 BufferUV, uint2 PixelCoord)
 
 #endif
 
+#if DIM_FROXELS
+
+#define FROXEL_HASH_BUFFER_SIZE (GROUP_TILE_SIZE * GROUP_TILE_SIZE)
+#define FROXEL_HASH_THREAD_GROUP_SIZE (GROUP_TILE_SIZE * GROUP_TILE_SIZE)
+
+#define FROXEL_SHARED_HASH_BUFFER_VAR SharedMinDeviceZ
+
+#include "Froxel/FroxelBuild.ush"
+
+
+#if GROUP_TILE_SIZE != FROXEL_TILE_SIZE
+	#error "The froxel construction algorithm assumes the group size matches the tile size, if this is changed this code needs to be updated"
+#endif
+
+#endif // DIM_FROXELS
+
+
 [numthreads(GROUP_TILE_SIZE, GROUP_TILE_SIZE, 1)]
 void HZBBuildCS(
 	uint2 GroupId : SV_GroupID,
@@ -166,7 +183,12 @@ void HZBBuildCS(
 		uint2 GroupThreadId = InitialTilePixelPositionForReduction2x2(MAX_MIP_BATCH_SIZE - 1, GroupThreadIndex);
 	#endif
 
-	uint2 DispatchThreadId = GROUP_TILE_SIZE * GroupId + GroupThreadId;
+	uint2 GroupOffset = GROUP_TILE_SIZE * GroupId;
+	uint2 DispatchThreadId = GroupOffset + GroupThreadId;
+
+	// whether or not any of the threads in the group is going to look at valid pixels, a lot don't in the HZB build because of POT rounding.
+	uint2 MinPixelCoord = (GroupOffset << 1) + PixelViewPortMinMax.xy;
+	bool bValidGroup = all(MinPixelCoord < PixelViewPortMinMax.zw);
 
 	float2 BufferUV = (DispatchThreadId + 0.5) * DispatchThreadIdToBufferUV.xy + DispatchThreadIdToBufferUV.zw;
 	float4 DeviceZ = Gather4(ParentTextureMip, ParentTextureMipSampler, BufferUV);
@@ -180,6 +202,20 @@ void HZBBuildCS(
 
 	float MinDeviceZ = min(min3(DeviceZ.x, DeviceZ.y, DeviceZ.z), DeviceZ.w);
 	float MaxDeviceZ = max(max3(DeviceZ.x, DeviceZ.y, DeviceZ.z), DeviceZ.w);
+
+#if DIM_FROXELS
+	// 2x as each thread does 2x2 through Gather
+	uint2 LocalTileId = (2 * GroupThreadId) / (FROXEL_TILE_SIZE);
+	uint LocalLinearTileId = LocalTileId.y * 2 + LocalTileId.x;
+	 // 2x2 tiles per group
+	uint2 GroupTileOffset = GroupId * 2;
+	uint2 TileId = GroupTileOffset + LocalTileId;
+
+	if (bValidGroup)
+	{
+		HashBuildFroxelsFromDeviceZ(DeviceZ, GroupThreadIndex, LocalLinearTileId, GroupTileOffset);
+	}
+#endif
 	
 	uint2 OutputPixelPos = DispatchThreadId;
 	
@@ -197,7 +233,7 @@ void HZBBuildCS(
 	}
 	#else
 	{
-		SharedMinDeviceZ[GroupThreadIndex] = MinDeviceZ;
+		SharedMinDeviceZ[GroupThreadIndex] = asuint(MinDeviceZ);
 		SharedMaxDeviceZ[GroupThreadIndex] = MaxDeviceZ;
 #if FEATURE_LEVEL >= FEATURE_LEVEL_SM6 || PLATFORM_SUPPORTS_SM6_0_WAVE_OPERATIONS
 		const uint LaneCount = WaveGetLaneCount();
@@ -230,7 +266,7 @@ void HZBBuildCS(
 				for (uint i = 1; i < 4; i++)
 				{
 					uint LDSIndex = GroupThreadIndex + i * ReduceBankSize;
-					ParentMinDeviceZ[i] = SharedMinDeviceZ[LDSIndex];
+					ParentMinDeviceZ[i] = asfloat(SharedMinDeviceZ[LDSIndex]);
 					ParentMaxDeviceZ[i] = SharedMaxDeviceZ[LDSIndex];
 				}
 				
@@ -240,7 +276,7 @@ void HZBBuildCS(
 				OutputPixelPos = OutputPixelPos >> 1;
 				OutputMipLevel(MipLevel, OutputPixelPos, MinDeviceZ, MaxDeviceZ);
 				
-				SharedMinDeviceZ[GroupThreadIndex] = MinDeviceZ;
+				SharedMinDeviceZ[GroupThreadIndex] = asuint(MinDeviceZ);
 				SharedMaxDeviceZ[GroupThreadIndex] = MaxDeviceZ;
 			}
 		}
diff --git a/Engine/Shaders/Private/RendererUtils.usf b/Engine/Shaders/Private/RendererUtils.usf
new file mode 100644
--- /dev/null
+++ b/Engine/Shaders/Private/RendererUtils.usf
@@ -0,0 +1,25 @@
+// Copyright Epic Games, Inc. All Rights Reserved.
+#include "Common.ush"
+
+RWBuffer<uint> OutIndirectArgsBuffer;
+uint NumIndirectArgs;
+uint IndirectArgStride;
+int3 DimClearValue;
+
+// Set dispatch Args to DimClearValue.XYZ (3D grid dim) and any other fields to 0 ready for atomic adding on GPU
+[numthreads(64, 1, 1)]
+void ClearIndirectDispatchArgsCS(uint IndirectArgIndex : SV_DispatchThreadID)
+{
+	if (IndirectArgIndex < NumIndirectArgs)
+	{
+		OutIndirectArgsBuffer[IndirectArgIndex * IndirectArgStride + 0] = DimClearValue.x; // X-group count
+		OutIndirectArgsBuffer[IndirectArgIndex * IndirectArgStride + 1] = DimClearValue.y; // Y-group count
+		OutIndirectArgsBuffer[IndirectArgIndex * IndirectArgStride + 2] = DimClearValue.z; // Z-group count
+
+		// Clear remaining auxiliary or padding fields to zero
+		for (int Index = 3; Index < IndirectArgStride; ++Index)
+		{
+			OutIndirectArgsBuffer[IndirectArgIndex * IndirectArgStride + Index] = 0;
+		}
+	}
+}
diff --git a/Engine/Shaders/Private/SingleLayerWaterComposite.usf b/Engine/Shaders/Private/SingleLayerWaterComposite.usf
--- a/Engine/Shaders/Private/SingleLayerWaterComposite.usf
+++ b/Engine/Shaders/Private/SingleLayerWaterComposite.usf
@@ -21,6 +21,20 @@
 #include "ReflectionEnvironmentComposite.ush"
 #include "Substrate/SubstrateEvaluation.ush"
 
+#if GENERATE_FROXELS
+
+
+#define FROXEL_HASH_BUFFER_SIZE 64
+#define FROXEL_HASH_THREAD_GROUP_SIZE 64
+#define FROXEL_SHARED_HASH_BUFFER_VAR GroupShared_FroxelHashBuffer
+
+groupshared uint GroupShared_FroxelHashBuffer[FROXEL_HASH_BUFFER_SIZE];
+
+#include "Froxel/FroxelBuild.ush"
+
+
+#endif
+
 #ifndef THREADGROUP_SIZE
 	#define THREADGROUP_SIZE 1
 #endif
@@ -277,6 +291,7 @@ SingleLayerWaterCompositeOutput SingleLayerWaterComposite(float2 BufferUV, float
 
 #if USE_WATER_PRE_PASS_STENCIL
 Texture2D<uint2> WaterDepthStencilTexture;
+Texture2D<float> WaterDepthTexture;
 #endif
 
 int2 TiledViewRes;
@@ -312,10 +327,10 @@ bool DoesPixelContainWater(uint2 PixelPos)
 [numthreads(SLW_TILE_SIZE_XY, SLW_TILE_SIZE_XY, 1)]
 void WaterTileCatergorisationMarkCS(uint3 ThreadId : SV_DispatchThreadID, uint3 GroupId : SV_GroupID, uint3 GroupThreadId : SV_GroupThreadID)
 {
+	const uint GroupThreadIndex = GroupThreadId.y * SLW_TILE_SIZE_XY + GroupThreadId.x;
 #if 0
 	// Slow reference path
-	const int LinearIndex = GroupThreadId.y * SLW_TILE_SIZE_XY + GroupThreadId.x;
-	if (LinearIndex < 1)
+	if (GroupThreadIndex < 1u)
 	{
 		bool bContainsWater = false;
 		for (uint i = 0; i < SLW_TILE_SIZE_XY; ++i)
@@ -337,10 +352,11 @@ void WaterTileCatergorisationMarkCS(uint3 ThreadId : SV_DispatchThreadID, uint3
 
 #else
 
+	uint2 PixelPos = ThreadId.xy + View.ViewRectMin.xy;
 	bool bContainsWater = false;
 	if(ThreadId.x < uint(View.BufferSizeAndInvSize.x) && ThreadId.y < uint(View.BufferSizeAndInvSize.y))
 	{
-		bContainsWater = bContainsWater || DoesPixelContainWater(ThreadId.xy + View.ViewRectMin.xy);
+		bContainsWater = DoesPixelContainWater(PixelPos);
 	}
 
 	bool bWriteTile = false;
@@ -367,42 +383,41 @@ void WaterTileCatergorisationMarkCS(uint3 ThreadId : SV_DispatchThreadID, uint3
 
 #else
 
-	const int LinearIndex = GroupThreadId.y * SLW_TILE_SIZE_XY + GroupThreadId.x;
-	ContainsWater[LinearIndex] = bContainsWater;
+	ContainsWater[GroupThreadIndex] = bContainsWater;
 	GroupMemoryBarrierWithGroupSync();
 
-	if (LinearIndex < 32) // 8*8 = 64 elements to merge
+	if (GroupThreadIndex < 32) // 8*8 = 64 elements to merge
 	{
-		ContainsWater[LinearIndex] = ContainsWater[LinearIndex] || ContainsWater[LinearIndex + 32];
+		ContainsWater[GroupThreadIndex] = ContainsWater[GroupThreadIndex] || ContainsWater[GroupThreadIndex + 32];
 	}
 	GroupMemoryBarrierWithGroupSync();
-	if (LinearIndex < 16)
+	if (GroupThreadIndex < 16)
 	{
-		ContainsWater[LinearIndex] = ContainsWater[LinearIndex] || ContainsWater[LinearIndex + 16];
+		ContainsWater[GroupThreadIndex] = ContainsWater[GroupThreadIndex] || ContainsWater[GroupThreadIndex + 16];
 	}
 	GroupMemoryBarrierWithGroupSync();
 
 	// The smallest wave size is 16 on Intel hardware. So now we can do simple math operations without group sync.
 	// EDIT: for some reason group sync is needed until the end, otherwise some pixels are missing...
 
-	if (LinearIndex < 8)
+	if (GroupThreadIndex < 8)
 	{
-		ContainsWater[LinearIndex] = ContainsWater[LinearIndex] || ContainsWater[LinearIndex + 8];
+		ContainsWater[GroupThreadIndex] = ContainsWater[GroupThreadIndex] || ContainsWater[GroupThreadIndex + 8];
 	}
 	GroupMemoryBarrierWithGroupSync();
-	if (LinearIndex < 4)
+	if (GroupThreadIndex < 4)
 	{
-		ContainsWater[LinearIndex] = ContainsWater[LinearIndex] || ContainsWater[LinearIndex + 4];
+		ContainsWater[GroupThreadIndex] = ContainsWater[GroupThreadIndex] || ContainsWater[GroupThreadIndex + 4];
 	}
 	GroupMemoryBarrierWithGroupSync();
-	if (LinearIndex < 2)
+	if (GroupThreadIndex < 2)
 	{
-		ContainsWater[LinearIndex] = ContainsWater[LinearIndex] || ContainsWater[LinearIndex + 2];
+		ContainsWater[GroupThreadIndex] = ContainsWater[GroupThreadIndex] || ContainsWater[GroupThreadIndex + 2];
 	}
 	GroupMemoryBarrierWithGroupSync();
-	if (LinearIndex < 1)
+	if (GroupThreadIndex < 1)
 	{
-		bWriteTile = ContainsWater[LinearIndex] || ContainsWater[LinearIndex + 1];
+		bWriteTile = ContainsWater[GroupThreadIndex] || ContainsWater[GroupThreadIndex + 1];
 	}
 
 #endif
@@ -414,6 +429,27 @@ void WaterTileCatergorisationMarkCS(uint3 ThreadId : SV_DispatchThreadID, uint3
 		InterlockedOr(TileMaskBufferOut[MaskLinearIndex / 32U], 1U << (MaskLinearIndex % 32U));
 	}
 #endif
+
+#if GENERATE_FROXELS
+	// Need to get value that's valid for the whole group
+#if COMPILER_SUPPORTS_WAVE_VOTE
+	if (bAnyWaterPixels != 0)
+#endif
+	{
+		float DeviceZ = 0.0f;
+		BRANCH
+		if (bContainsWater)
+		{
+			DeviceZ = WaterDepthTexture.Load(uint3(PixelPos, 0)).x;
+		}
+		// One group per 8x8 tile
+		const uint2 GroupTileOffset = GroupId.xy;
+		// Only one 8x8 tile per group
+		const uint LocalLinearTileId = 0;
+		HashBuildFroxelsFromDeviceZ(DeviceZ, GroupThreadIndex, LocalLinearTileId, GroupTileOffset);
+	}
+#endif
+
 } 
 
 uint VertexCountPerInstanceIndirect;
diff --git a/Engine/Shaders/Private/VirtualShadowMaps/VirtualShadowMapPageMarking.usf b/Engine/Shaders/Private/VirtualShadowMaps/VirtualShadowMapPageMarking.usf
--- a/Engine/Shaders/Private/VirtualShadowMaps/VirtualShadowMapPageMarking.usf
+++ b/Engine/Shaders/Private/VirtualShadowMaps/VirtualShadowMapPageMarking.usf
@@ -19,6 +19,8 @@
 #include "VirtualShadowMapProjectionCommon.ush"
 #include "VirtualShadowMapStats.ush"
 #include "VirtualShadowMapLightGrid.ush"
+#include "../ShaderPrint.ush"
+#include "../Froxel/Froxel.ush"
 
 // Type of input data consume by the page allocation (i.e., data read from the source buffer: Gbuffer, HairStrands data, ...)
 #define INPUT_TYPE_GBUFFER 0
@@ -507,6 +509,171 @@ void GeneratePageFlagsFromPixels(
 	}
 }
 
+#ifdef GeneratePageFlagsFromFroxelsCS
+
+StructuredBuffer<FPackedFroxel>	Froxels;
+Buffer<uint>			FroxelArgs;
+
+int3 LoadFroxel(uint Index)
+{
+	return UnpackFroxel(Froxels[Index]);
+}
+
+FCulledLightsGridHeader VirtualShadowMapGetLightsGridForFroxel(int3 FroxelCoord, float FroxelCenterViewZ)
+{
+	// Has no effect unless INSTANCED_STEREO is enabled. Used in GetLightGridData to select ForwardLightData or ForwardLightDataISR, but then only uses redundant parameters anyway.
+	uint EyeIndex = 0U;
+
+	// Pick center point in froxel in pixels
+	uint2 PixelPos = FroxelCoord.xy * FROXEL_TILE_SIZE + FROXEL_TILE_SIZE / 2;
+
+	// TODO: calculate mapping directly?
+	uint GridLinearIndex = ComputeLightGridCellIndex(PixelPos, FroxelCenterViewZ, EyeIndex);
+	FCulledLightsGridHeader CulledLightGridData = GetCulledLightsGridHeader(GridLinearIndex, EyeIndex);
+
+	// Replace light count with our pruned count
+	CulledLightGridData.NumLights = VirtualShadowMap.NumCulledLightsGrid[GridLinearIndex];
+	return CulledLightGridData;
+}
+
+uint bShouldMarkLocaLights;
+int PassId;
+float DebugRange;
+
+[numthreads(FROXEL_INDIRECT_ARG_WORKGROUP_SIZE, 1, 1)]
+void GeneratePageFlagsFromFroxelsCS(uint FroxelIndex : SV_DispatchThreadID)
+{
+#if DEBUG_DRAW_GENERATE_FROM_FROXELS
+	FShaderPrintContext Context = InitShaderPrintContext();
+#endif
+	uint FroxelCount = FroxelArgs[3];
+
+	if (FroxelIndex >= FroxelCount)
+	{
+		return;
+	}
+	int3 FroxelCoord = LoadFroxel(FroxelIndex);
+
+	float3 FroxelCenterView = GetFroxelViewSpaceCenter(FroxelCoord);
+	float3 FroxelCenterTranslatedWorld = mul(float4(FroxelCenterView, 1.0f), View.ViewToTranslatedWorld).xyz;
+
+#if DEBUG_DRAW_GENERATE_FROM_FROXELS
+	if (FroxelIndex == 0)
+	{
+		float XLoc = 0.2;
+		Context.Pos.y = 0.13f + 0.02f * float(PassId);
+		Context.Pos.x = XLoc;
+		Print(Context, float(FroxelCount) / 1000, FontGreen);
+	}
+	if (!Context.IsDrawLocked())
+	{
+		float Range = length(float2(View.CursorPosition / FROXEL_TILE_SIZE) - float2(FroxelCoord.xy));
+		if (Range < DebugRange)
+		{
+			FFroxelClipBounds ClipBounds = GetFroxelClipBounds(FroxelCoord);
+			float Fade = saturate(8.0f * (DebugRange - Range) / DebugRange);
+			float4 Color = float4(IntToColor(FroxelCoord.x + FroxelCoord.y * 13791 + FroxelCoord.z * 93113791), Fade);
+			AddOBBTWS(Context, ClipBounds.ClipMin, ClipBounds.ClipMax, Color, View.ClipToTranslatedWorld, true);
+
+			// Draw the center as well
+			// AddSphereTWS(Context, FroxelCenterTranslatedWorld, 0.1f, Color);
+			// And the view space bounds.
+			//FFroxelViewBounds ViewBounds = GetFroxelViewSpaceAABB(FroxelCoord);
+			//AddOBBTWS(Context, ViewBounds.Min, ViewBounds.Max, Color, DFHackToFloat(View.ViewToTranslatedWorld));
+		}
+	}
+#endif
+	// Dither pattern for page dilation
+	// We don't need to to check all 8 adjacent pages; as long as there's at least a single pixel near the edge
+	// the adjacent one will get mapped. In practice only checking one diagonal seems to work fine and have minimal
+	// overhead.
+	const float2 PageDilationDither = float2(
+		(FroxelIndex & 1) ? 1.0f : -1.0f,
+		(FroxelIndex & 2) ? 1.0f : -1.0f);
+		
+	// Directional lights
+	{
+		const bool bUsePageDilation = PageDilationBorderSizeDirectional > 0.0f;
+		const float2 PageDilationOffset = PageDilationBorderSizeDirectional * PageDilationDither;
+
+		for (uint Index = 0; Index < NumDirectionalLightSmInds; ++Index)
+		{
+			int VirtualShadowMapId = DirectionalLightIds[Index];
+			FVirtualShadowMapProjectionShaderData ProjectionData = GetVirtualShadowMapProjectionData(VirtualShadowMapId);
+
+			MarkPageClipmap(ProjectionData, bUsePageDilation, PageDilationOffset, FroxelCenterTranslatedWorld);
+		}
+	}
+
+	// Local lights
+	if (bShouldMarkLocaLights != 0)
+	{
+		const bool bUsePageDilation = PageDilationBorderSizeLocal > 0.0f;
+		const float2 PageDilationOffset = PageDilationBorderSizeLocal * PageDilationDither;
+				
+		const FCulledLightsGridHeader LightGridData = VirtualShadowMapGetLightsGridForFroxel(FroxelCoord, FroxelCenterView.z);
+
+		LOOP
+		for (uint Index = 0; Index < LightGridData.NumLights; ++Index)
+		{
+			const FLocalLightData LightData = VirtualShadowMapGetLocalLightData(LightGridData, Index);
+
+			// If we hit a distant light, we're done as we sorted those to the end. See PruneLightGridCS.
+			if (IsSinglePageVirtualShadowMap(LightData.VirtualShadowMapId))
+			{
+				break;
+			}
+
+
+			// Relative to PrimaryView
+			const float3 LightTranslatedWorldPosition = LightData.LightPositionAndInvRadius.xyz;
+			float3 LocToLight = LightTranslatedWorldPosition - FroxelCenterTranslatedWorld;
+			float LengthSquared = length2(LocToLight);
+
+			if (LengthSquared > Square(1.0f / LightData.LightPositionAndInvRadius.w))
+			{
+				continue;
+			}
+
+			float3 ToLight = normalize(LocToLight);
+
+			// TODO: Can optimize these tests by fusing them together
+			// Also do precise cone test, since froxels are pretty coarse at times.
+			if (dot(ToLight, LightData.LightDirectionAndShadowMask.xyz) < GetLightSpotAngles(LightData).x)
+			{
+				continue;
+			}
+
+			const uint LightTypeAndPackedShadowMapChannelMask = asuint(LightData.LightDirectionAndShadowMask.w);
+			const uint LightType = UnpackLightType(LightTypeAndPackedShadowMapChannelMask);
+			
+			const bool bIsRectLight = LightType == LIGHT_TYPE_RECT;
+			if (bIsRectLight)
+			{
+				const bool bIsBehindRectLight = dot(ToLight, LightData.LightDirectionAndShadowMask.xyz) < 0;
+				if (bIsBehindRectLight)
+				{
+					continue;
+				}
+			}
+
+			// TODO: Precise radius test necessary?
+
+			int VirtualShadowMapId = LightData.VirtualShadowMapId;
+			bool bSpotLight = GetLightSpotAngles(LightData).x > -2.0f;
+			if( !bSpotLight )
+			{
+				VirtualShadowMapId += VirtualShadowMapGetCubeFace(-ToLight);
+			}
+
+			uint MipLevel = GetMipLevelLocal(VirtualShadowMapId, FroxelCenterTranslatedWorld, FroxelCenterView.z);
+			MarkPage(VirtualShadowMapId, MipLevel, FroxelCenterTranslatedWorld, bUsePageDilation, PageDilationOffset);
+		}
+	}
+}
+
+#endif 
+
 uint bMarkCoarsePagesLocal;
 int ClipmapFirstLevel;
 uint ClipmapIndexMask;
diff --git a/Engine/Shaders/Private/VirtualShadowMaps/VirtualShadowMapPhysicalPageManagement.usf b/Engine/Shaders/Private/VirtualShadowMaps/VirtualShadowMapPhysicalPageManagement.usf
--- a/Engine/Shaders/Private/VirtualShadowMaps/VirtualShadowMapPhysicalPageManagement.usf
+++ b/Engine/Shaders/Private/VirtualShadowMaps/VirtualShadowMapPhysicalPageManagement.usf
@@ -665,22 +665,6 @@ void MergePhysicalPixel(uint2 PixelCoord)
 		OutPhysicalPagePool[uint3(PixelCoord, GetVirtualShadowMapStaticArrayIndex())]);
 }
 
-RWBuffer<uint> OutIndirectArgsBuffer;
-uint NumIndirectArgs;
-uint IndirectArgStride;
-
-// Set dispatch Args to 0,1,1 (3D grid dim) ready for atomic adding on GPU
-[numthreads(64, 1, 1)]
-void ClearIndirectDispatchArgs1DCS(uint IndirectArgIndex : SV_DispatchThreadID)
-{
-	if (IndirectArgIndex < NumIndirectArgs)
-	{
-		OutIndirectArgsBuffer[IndirectArgIndex * IndirectArgStride + 0] = 0;
-		OutIndirectArgsBuffer[IndirectArgIndex * IndirectArgStride + 1] = 1;
-		OutIndirectArgsBuffer[IndirectArgIndex * IndirectArgStride + 2] = 1;
-	}
-}
-
 // Log2 2D dimension of thread group size, 2^4 == 16, 
 #define LOG2_TILE_THREAD_GROUP_SIZE_XY 4U
 #define TILE_THREAD_GROUP_SIZE_XY (1U << LOG2_TILE_THREAD_GROUP_SIZE_XY)
diff --git a/Engine/Shaders/Shared/FroxelDefinitions.h b/Engine/Shaders/Shared/FroxelDefinitions.h
new file mode 100644
--- /dev/null
+++ b/Engine/Shaders/Shared/FroxelDefinitions.h
@@ -0,0 +1,39 @@
+// Copyright Epic Games, Inc. All Rights Reserved.
+
+/*================================================================================================
+VirtualShadowMapDefinitions.h: used in virtual shadow map shaders and C++ code to define common constants
+!!! Changing this file requires recompilation of the engine !!!
+=================================================================================================*/
+
+#pragma once
+
+#define FROXEL_TILE_SIZE (8)
+#define FROXEL_INVALID_SLICE (1 << 28)
+#define FROXEL_INVALID_PACKED_SLICE (FROXEL_INVALID_SLICE << 2)
+#define FROXEL_PACKED_SLICE_BIAS (1 << 20) // used to represent negative slices (TODO: dont)
+#define FROXEL_INDIRECT_ARG_WORKGROUP_SIZE (64)
+
+
+#ifdef __cplusplus
+#include "HLSLTypeAliases.h"
+
+namespace UE::HLSL
+{
+#endif
+
+struct FPackedFroxel
+{
+	uint XY;
+	int Z;
+};
+
+
+#ifdef __cplusplus
+} // namespace UE::HLSL
+
+namespace Froxel
+{
+	using FPackedFroxel = UE::HLSL::FPackedFroxel;
+}
+
+#endif
\ No newline at end of file
diff --git a/Engine/Source/Runtime/RenderCore/Private/RenderUtils.cpp b/Engine/Source/Runtime/RenderCore/Private/RenderUtils.cpp
--- a/Engine/Source/Runtime/RenderCore/Private/RenderUtils.cpp
+++ b/Engine/Source/Runtime/RenderCore/Private/RenderUtils.cpp
@@ -1345,13 +1345,18 @@ bool UseNanite(EShaderPlatform ShaderPlatform, bool bCheckForAtomicSupport /*= t
 	return bNaniteEnabled && DoesRuntimeSupportNanite(ShaderPlatform, bCheckForAtomicSupport, bCheckForProjectSetting);
 }
 
-bool UseVirtualShadowMaps(EShaderPlatform ShaderPlatform, const FStaticFeatureLevel FeatureLevel)
+bool UseVirtualShadowMaps(EShaderPlatform ShaderPlatform)
 {
 	static const auto EnableVirtualSMCVar = IConsoleManager::Get().FindConsoleVariable(TEXT("r.Shadow.Virtual.Enable"));
 	const bool bVirtualShadowMapsEnabled = EnableVirtualSMCVar ? (EnableVirtualSMCVar->GetInt() != 0) : false;
 	return bVirtualShadowMapsEnabled && DoesRuntimeSupportNanite(ShaderPlatform, true /* check for atomics */, false /* check project setting */);
 }
 
+bool UseVirtualShadowMaps(EShaderPlatform ShaderPlatform, const FStaticFeatureLevel FeatureLevel)
+{
+	return UseVirtualShadowMaps(ShaderPlatform);
+}
+
 bool DoesPlatformSupportVirtualShadowMaps(EShaderPlatform Platform)
 {
 	return DoesPlatformSupportNanite(Platform, false /* check project setting */);
diff --git a/Engine/Source/Runtime/RenderCore/Public/RenderUtils.h b/Engine/Source/Runtime/RenderCore/Public/RenderUtils.h
--- a/Engine/Source/Runtime/RenderCore/Public/RenderUtils.h
+++ b/Engine/Source/Runtime/RenderCore/Public/RenderUtils.h
@@ -415,6 +415,12 @@ RENDERCORE_API bool DoesTargetPlatformSupportNanite(const ITargetPlatform* Targe
  */
 RENDERCORE_API bool UseNanite(EShaderPlatform ShaderPlatform, bool bCheckForAtomicSupport = true, bool bCheckForProjectSetting = true);
 
+/**
+ * Returns true if Virtual Shadow Maps should be used for the given shader platform.
+ * Note: Virtual Shadow Maps require Nanite support.
+ */
+RENDERCORE_API bool UseVirtualShadowMaps(EShaderPlatform ShaderPlatform);
+
 /**
  * Returns true if Virtual Shadow Maps should be used for the given shader platform.
  * Note: Virtual Shadow Maps require Nanite support.
diff --git a/Engine/Source/Runtime/Renderer/Private/DeferredShadingRenderer.cpp b/Engine/Source/Runtime/Renderer/Private/DeferredShadingRenderer.cpp
--- a/Engine/Source/Runtime/Renderer/Private/DeferredShadingRenderer.cpp
+++ b/Engine/Source/Runtime/Renderer/Private/DeferredShadingRenderer.cpp
@@ -103,6 +103,7 @@
 #include "CustomRenderPassSceneCapture.h"
 #include "EnvironmentComponentsFlags.h"
 #include "GenerateMips.h"
+#include "Froxel/Froxel.h"
 
 #if !UE_BUILD_SHIPPING
 #include "RenderCaptureInterface.h"
@@ -412,7 +413,7 @@ bool FDeferredShadingSceneRenderer::ShouldRenderNanite() const
 	return UseNanite(ShaderPlatform) && ViewFamily.EngineShowFlags.NaniteMeshes && Nanite::GStreamingManager.HasResourceEntries();
 }
 
-bool FDeferredShadingSceneRenderer::RenderHzb(FRDGBuilder& GraphBuilder, FRDGTextureRef SceneDepthTexture, const FBuildHZBAsyncComputeParams* AsyncComputeParams)
+bool FDeferredShadingSceneRenderer::RenderHzb(FRDGBuilder& GraphBuilder, FRDGTextureRef SceneDepthTexture, const FBuildHZBAsyncComputeParams* AsyncComputeParams, Froxel::FRenderer& FroxelRenderer)
 {
 	RDG_GPU_STAT_SCOPE(GraphBuilder, HZB);
 
@@ -445,7 +446,8 @@ bool FDeferredShadingSceneRenderer::RenderHzb(FRDGBuilder& GraphBuilder, FRDGTex
 				TEXT("HZBFurthest"),
 				/* OutFurthestHZBTexture = */ &FurthestHZBTexture,
 				BuildHZBDefaultPixelFormat,
-				AsyncComputeParams);
+				AsyncComputeParams,
+				FroxelRenderer.GetView(ViewIndex));
 
 			// Update the view.
 			{
@@ -2146,7 +2148,7 @@ void FDeferredShadingSceneRenderer::Render(FRDGBuilder& GraphBuilder)
 		return GetViewPipelineState(Views[ViewIndex]).AmbientOcclusionMethod == EAmbientOcclusionMethod::SSAO;
 	});
 
-	const auto RenderOcclusionLambda = [&]()
+	const auto RenderOcclusionLambda = [&]() -> Froxel::FRenderer 
 	{
 		const int32 AsyncComputeMode = CVarSceneDepthHZBAsyncCompute.GetValueOnRenderThread();
 		bool bAsyncCompute = AsyncComputeMode != 0;
@@ -2157,10 +2159,16 @@ void FDeferredShadingSceneRenderer::Render(FRDGBuilder& GraphBuilder)
 			AsyncComputeParams.Prerequisite = ComputeLightGridOutput.CompactLinksPass;
 		}
 
+		bool bShouldGenerateFroxels = DoesVSMWantFroxels(ShaderPlatform);
+
+		Froxel::FRenderer FroxelRenderer(bShouldGenerateFroxels, GraphBuilder, Views);
+
 		RenderOcclusion(GraphBuilder, SceneTextures, bIsOcclusionTesting,
-			bAsyncCompute ? &AsyncComputeParams : nullptr);
+			bAsyncCompute ? &AsyncComputeParams : nullptr, FroxelRenderer);
 
 		CompositionLighting.ProcessAfterOcclusion(GraphBuilder);
+
+		return FroxelRenderer;
 	};
 
 	const bool bShouldRenderVolumetricCloudBase = ShouldRenderVolumetricCloud(Scene, ViewFamily.EngineShowFlags);
@@ -2169,6 +2177,8 @@ void FDeferredShadingSceneRenderer::Render(FRDGBuilder& GraphBuilder)
 	const bool bAsyncComputeVolumetricCloud = IsVolumetricRenderTargetEnabled() && IsVolumetricRenderTargetAsyncCompute();
 	const bool bVolumetricRenderTargetRequired = bShouldRenderVolumetricCloud && !bHasRayTracedOverlay;
 
+	Froxel::FRenderer FroxelRenderer;
+
 	FRDGTextureRef ViewFamilyTexture = TryCreateViewFamilyTexture(GraphBuilder, ViewFamily);
 	FRDGTextureRef ViewFamilyDepthTexture = TryCreateViewFamilyDepthTexture(GraphBuilder, ViewFamily);
 	if (RendererOutput == ERendererOutput::DepthPrepassOnly || RendererOutput == ERendererOutput::BasePass)
@@ -2186,7 +2196,7 @@ void FDeferredShadingSceneRenderer::Render(FRDGBuilder& GraphBuilder)
 			bOcclusionBeforeBasePass = ((DepthPass.EarlyZPassMode == EDepthDrawingMode::DDM_AllOccluders) || bIsEarlyDepthComplete);
 			if (bOcclusionBeforeBasePass)
 			{
-				RenderOcclusionLambda();
+				FroxelRenderer = RenderOcclusionLambda();
 			}
 
 			SceneTextures.SetupMode |= ESceneTextureSetupMode::SceneColor;
@@ -2221,7 +2231,7 @@ void FDeferredShadingSceneRenderer::Render(FRDGBuilder& GraphBuilder)
 
 		if (!bOcclusionBeforeBasePass)
 		{
-			RenderOcclusionLambda();
+			FroxelRenderer = RenderOcclusionLambda();
 		}
 
 		if (bUpdateNaniteStreaming)
@@ -2285,7 +2295,7 @@ void FDeferredShadingSceneRenderer::Render(FRDGBuilder& GraphBuilder)
 
 		if (bOcclusionBeforeBasePass)
 		{
-			RenderOcclusionLambda();
+			FroxelRenderer = RenderOcclusionLambda();
 		}
 
 		// End early occlusion queries
@@ -2540,7 +2550,7 @@ void FDeferredShadingSceneRenderer::Render(FRDGBuilder& GraphBuilder)
 		// Occlusion after base pass
 		if (!bOcclusionBeforeBasePass)
 		{
-			RenderOcclusionLambda();
+			FroxelRenderer = RenderOcclusionLambda();
 		}
 
 		// End occlusion after base
@@ -2649,7 +2659,7 @@ void FDeferredShadingSceneRenderer::Render(FRDGBuilder& GraphBuilder)
 
 					FFrontLayerTranslucencyData FrontLayerTranslucencyData = RenderFrontLayerTranslucency(GraphBuilder, Views, SceneTextures, true /*VSM page marking*/);
 
-					VirtualShadowMapArray.BuildPageAllocations(GraphBuilder, GetActiveSceneTextures(), Views, SortedLightSet, VisibleLightInfos, SingleLayerWaterPrePassResult, FrontLayerTranslucencyData);
+					VirtualShadowMapArray.BuildPageAllocations(GraphBuilder, GetActiveSceneTextures(), Views, SortedLightSet, VisibleLightInfos, SingleLayerWaterPrePassResult, FrontLayerTranslucencyData, FroxelRenderer);
 				}
 
 				RenderShadowDepthMaps(GraphBuilder, InitViewTaskDatas.DynamicShadows, InstanceCullingManager, ExternalAccessQueue);
diff --git a/Engine/Source/Runtime/Renderer/Private/DeferredShadingRenderer.h b/Engine/Source/Runtime/Renderer/Private/DeferredShadingRenderer.h
--- a/Engine/Source/Runtime/Renderer/Private/DeferredShadingRenderer.h
+++ b/Engine/Source/Runtime/Renderer/Private/DeferredShadingRenderer.h
@@ -65,6 +65,11 @@ struct FDBufferTextures;
 struct FILCUpdatePrimTaskData;
 struct FLumenDirectLightingTaskData;
 
+namespace Froxel
+{
+	class FFroxelRenderer;
+}
+
 class IVisibilityTaskData;
 
 #if RHI_RAYTRACING
@@ -424,9 +429,10 @@ public:
 		FRDGBuilder& GraphBuilder,
 		const FSceneTextures& SceneTextures,
 		bool bIsOcclusionTesting,
-		const FBuildHZBAsyncComputeParams* BuildHZBAsyncComputeParams = nullptr);
+		const FBuildHZBAsyncComputeParams* BuildHZBAsyncComputeParams,
+		Froxel::FRenderer& FroxelRenderer);
 
-	bool RenderHzb(FRDGBuilder& GraphBuilder, FRDGTextureRef SceneDepthTexture, const FBuildHZBAsyncComputeParams* AsyncComputeParams);
+	bool RenderHzb(FRDGBuilder& GraphBuilder, FRDGTextureRef SceneDepthTexture, const FBuildHZBAsyncComputeParams* AsyncComputeParams, Froxel::FRenderer& FroxelRenderer);
 
 	/** Renders the view family. */
 	virtual void Render(FRDGBuilder& GraphBuilder) override;
diff --git a/Engine/Source/Runtime/Renderer/Private/Froxel/Froxel.cpp b/Engine/Source/Runtime/Renderer/Private/Froxel/Froxel.cpp
new file mode 100644
--- /dev/null
+++ b/Engine/Source/Runtime/Renderer/Private/Froxel/Froxel.cpp
@@ -0,0 +1,115 @@
+// Copyright Epic Games, Inc. All Rights Reserved.
+
+#include "Froxel.h"
+#include "RenderGraphUtils.h"
+#include "RendererPrivateUtils.h"
+#include "SceneRendering.h"
+
+namespace Froxel
+{
+
+FSharedParameters MakeSharedParameters(const FViewInfo &View)
+{
+	const float FroxelTileSize = float(FRenderer::TileSize);
+
+	FVector2f ViewSize = FVector2f(View.ViewRect.Size());// FVector2f(View.ViewSizeAndInvSize.X, View.ViewSizeAndInvSize.Y);
+				
+	// How cubeish the froxels should be
+	const float DepthStretchFactor = 1.0f;
+
+	const FMatrix &ViewToClip = View.ViewMatrices.GetProjectionMatrix();
+	// 2x for, er, probably because clipspace is +-1
+	FVector2f AbsClipTileSize = FVector2f(FroxelTileSize * 2.0f) / ViewSize;
+	FVector2f ProjScaleXY = FVector2f(static_cast<float>(ViewToClip.M[0][0]), static_cast<float>(ViewToClip.M[1][1]));
+	FVector2f RadiusXY = AbsClipTileSize / ProjScaleXY;
+	float RadiusScreen = FMath::Min(RadiusXY.X, RadiusXY.Y) * DepthStretchFactor; 
+
+
+	const bool bIsOrtho = !View.ViewMatrices.IsPerspectiveProjection();
+
+	const float FroxelNear = static_cast<float>(ViewToClip.M[3][3] - ViewToClip.M[3][2]) / (ViewToClip.M[2][2] - ViewToClip.M[2][3]);
+	float FarPlane = FroxelNear - 1.0f / static_cast<float>(ViewToClip.M[2][2]);
+
+	float DepthScale = ViewToClip.M[2][3] * RadiusScreen; // + ViewToClip[3].W ignore the ortho fix unsure how to change, needs to somehow move outside the log, probably does factor out
+	FSharedParameters FroxelParameters;
+	FroxelParameters.FroxelNear = FroxelNear;
+	FroxelParameters.FroxelDepthScale1 = DepthScale + 1.0f;
+	FroxelParameters.FroxelRecNearScale = 1.0f / FroxelNear;
+	FroxelParameters.FroxelRecLog2DepthScale1 = 1.0f / FMath::Log2(1.0f + DepthScale);
+	FroxelParameters.FroxelRadius = RadiusScreen;
+	FroxelParameters.FroxelInvRadius = 1.0f / RadiusScreen;
+	FroxelParameters.bFroxelIsOrtho = bIsOrtho;
+
+	// needs to be modified per view if/when moved to shared buffer.
+	FroxelParameters.FroxelArgsStride = FRenderer::ArgsStride;
+	FroxelParameters.FroxelArgsOffset = 0;
+
+	if (bIsOrtho)
+	{
+		FroxelParameters.FroxelViewToClipTransformScale = ViewToClip.M[2][2];
+		FroxelParameters.FroxelClipToViewTransformScale = 1.0f / ViewToClip.M[2][2];
+		FroxelParameters.FroxelClipToViewTransformBias = -ViewToClip.M[3][2] / ViewToClip.M[2][2]; 
+	}
+	else
+	{
+		FroxelParameters.FroxelViewToClipTransformScale = ViewToClip.M[3][2];
+		FroxelParameters.FroxelClipToViewTransformScale = 1.0f / ViewToClip.M[3][2];
+		FroxelParameters.FroxelClipToViewTransformBias = -ViewToClip.M[2][2] / ViewToClip.M[3][2]; 
+	}
+
+	const FVector2f ClipTileSize = FVector2f(AbsClipTileSize.X, -AbsClipTileSize.Y);
+	const FVector2f ClipSpaceMin = FVector2f(-1.0f, 1.0f);
+
+	const FMatrix &ClipToView = View.ViewMatrices.GetInvProjectionMatrix();
+	const FVector2f ClipToViewScale = FVector2f(ClipToView.M[0][0], ClipToView.M[1][1]);
+	FVector2f FroxelToViewScale = ClipTileSize * ClipToViewScale;
+	FVector2f FroxelToViewBias = ClipSpaceMin * ClipToViewScale;
+	FroxelParameters.FroxelToViewScaleBias = FVector4f(FroxelToViewScale, FroxelToViewBias);
+
+	FroxelParameters.FroxelToClipScaleBias = FVector4f(ClipTileSize, ClipSpaceMin);
+	FroxelParameters.FroxelClipToViewScale = ClipToViewScale;
+
+	return FroxelParameters;
+}
+
+FRenderer::FRenderer(bool bIsEnabled, FRDGBuilder& GraphBuilder, const TArray<FViewInfo>& InViews)
+{
+	if (!bIsEnabled)
+	{
+		return;
+	}
+
+	for (const FViewInfo &View : InViews)
+	{
+		FSharedParameters FroxelParameters = MakeSharedParameters(View);
+
+		// TODO: over-conservative(!) 
+		int32 MaxNumFroxels = View.ViewRect.Area();
+		FRDGBuffer* FroxelsRDG = GraphBuilder.CreateBuffer(FRDGBufferDesc::CreateStructuredDesc(sizeof(FPackedFroxel), MaxNumFroxels), TEXT("r.Froxels"));
+		// TODO: make shared buffer & initialize once
+		FRDGBuffer* FroxelArgsRDG = FroxelArgsRDG = GraphBuilder.CreateBuffer(FRDGBufferDesc::CreateIndirectDesc(4), TEXT("r.FroxelArgs"));
+		AddClearIndirectDispatchArgs1DPass(GraphBuilder, View.GetFeatureLevel(), FroxelArgsRDG, 1u, ArgsStride);
+
+		Views.Emplace(FViewData{this, FroxelsRDG, FroxelArgsRDG, FroxelParameters});
+	}
+}
+
+FBuilderParameters FViewData::GetBuilderParameters(FRDGBuilder& GraphBuilder) const
+{
+	FBuilderParameters BuilderParameters;
+	BuilderParameters.FroxelParameters = SharedParameters;
+	BuilderParameters.OutFroxels = GraphBuilder.CreateUAV(FroxelsRDG);
+	BuilderParameters.OutFroxelArgs = GraphBuilder.CreateUAV(FroxelArgsRDG);
+	return BuilderParameters;
+}
+
+FParameters FViewData::GetShaderParameters(FRDGBuilder& GraphBuilder) const
+{
+	FParameters Parameters;
+	Parameters.FroxelParameters = SharedParameters;
+	Parameters.Froxels = GraphBuilder.CreateSRV(FroxelsRDG);
+	Parameters.FroxelArgs = GraphBuilder.CreateSRV(FroxelArgsRDG);
+	return Parameters;
+}
+
+}
\ No newline at end of file
diff --git a/Engine/Source/Runtime/Renderer/Private/Froxel/Froxel.h b/Engine/Source/Runtime/Renderer/Private/Froxel/Froxel.h
new file mode 100644
--- /dev/null
+++ b/Engine/Source/Runtime/Renderer/Private/Froxel/Froxel.h
@@ -0,0 +1,88 @@
+// Copyright Epic Games, Inc. All Rights Reserved.
+
+#pragma once
+
+#include "CoreMinimal.h"
+#include "ShaderParameterMacros.h"
+#include "FroxelDefinitions.h"
+
+class FRDGBuilder;
+class FViewInfo;
+
+/**
+ * Functions and data structures to manage lists of froxels that represent some geometry, e.g., the depth buffer.
+ * The froxels are spaced such that they are as deep as they are wide (at the near plane of the froxel slice). This means they provide good bounds for the samples they represent.
+ */
+namespace Froxel
+{
+
+BEGIN_SHADER_PARAMETER_STRUCT(FSharedParameters, )
+	SHADER_PARAMETER(FVector4f, FroxelToViewScaleBias)
+	SHADER_PARAMETER(FVector4f, FroxelToClipScaleBias)
+	SHADER_PARAMETER(FVector2f, FroxelClipToViewScale)
+	SHADER_PARAMETER(float, FroxelRecLog2DepthScale1)
+	SHADER_PARAMETER(float, FroxelRecNearScale)
+	SHADER_PARAMETER(float, FroxelDepthScale1)
+	SHADER_PARAMETER(float, FroxelNear)
+	SHADER_PARAMETER(float, FroxelViewToClipTransformScale)
+	SHADER_PARAMETER(float, FroxelClipToViewTransformScale)
+	SHADER_PARAMETER(float, FroxelClipToViewTransformBias)
+	SHADER_PARAMETER(float, FroxelRadius)
+	SHADER_PARAMETER(float, FroxelInvRadius)
+	SHADER_PARAMETER(uint32, bFroxelIsOrtho)
+	SHADER_PARAMETER(uint32, FroxelArgsOffset)
+	SHADER_PARAMETER(uint32, FroxelArgsStride)
+END_SHADER_PARAMETER_STRUCT()
+
+BEGIN_SHADER_PARAMETER_STRUCT(FBuilderParameters, )
+	SHADER_PARAMETER_STRUCT_INCLUDE(FSharedParameters, FroxelParameters)
+	SHADER_PARAMETER_RDG_BUFFER_UAV(RWStructuredBuffer< FPackedFroxel >, OutFroxels)
+	SHADER_PARAMETER_RDG_BUFFER_UAV(RWBuffer< uint >, OutFroxelArgs)
+END_SHADER_PARAMETER_STRUCT()
+
+BEGIN_SHADER_PARAMETER_STRUCT(FParameters, )
+	SHADER_PARAMETER_STRUCT_INCLUDE(FSharedParameters, FroxelParameters)
+	SHADER_PARAMETER_RDG_BUFFER_SRV(StructuredBuffer< FPackedFroxel >, Froxels)
+	SHADER_PARAMETER_RDG_BUFFER_SRV(Buffer< uint >, FroxelArgs)
+END_SHADER_PARAMETER_STRUCT()
+
+
+FSharedParameters MakeSharedParameters(const FViewInfo &View);
+
+// per view froxel data.
+struct FViewData
+{
+	class FRenderer *FroxelData = nullptr;
+	FRDGBuffer* FroxelsRDG = nullptr;
+	// Represents an argument to use with an indirect dispacth to perform some processing on the froxels. 
+	FRDGBuffer* FroxelArgsRDG = nullptr;
+	FSharedParameters SharedParameters;
+	int32 ArgsOffset = 0;
+
+	FBuilderParameters GetBuilderParameters(FRDGBuilder& GraphBuilder) const;
+
+	FParameters GetShaderParameters(FRDGBuilder& GraphBuilder) const;
+
+};
+
+class FRenderer
+{
+public:
+	// Stride in the indirect argument buffer, the 4th slot is used to store the atomic counter of individual froxels (rather than the group).
+	static constexpr int32 ArgsStride = 4;
+	// Work group size that should be used on an indirect dispatch using the argument produced
+	static constexpr int32 IndirectWorkGroupSize = FROXEL_INDIRECT_ARG_WORKGROUP_SIZE;
+	static constexpr int32 TileSize = FROXEL_TILE_SIZE; // 8x8 tiles
+
+	FRenderer() {}
+	FRenderer(bool bIsEnabled, FRDGBuilder& GraphBuilder, const TArray<FViewInfo>& Views);
+
+	const FViewData *GetView(int32 ViewIndex) const { return Views.IsValidIndex(ViewIndex) ? &Views[ViewIndex] : nullptr; }
+
+	bool IsEnabled() const { return !Views.IsEmpty(); }
+
+private:
+	TArray<FViewData> Views;
+};
+
+} // namespace froxel
diff --git a/Engine/Source/Runtime/Renderer/Private/RendererPrivateUtils.h b/Engine/Source/Runtime/Renderer/Private/RendererPrivateUtils.h
--- a/Engine/Source/Runtime/Renderer/Private/RendererPrivateUtils.h
+++ b/Engine/Source/Runtime/Renderer/Private/RendererPrivateUtils.h
@@ -368,3 +368,30 @@ using TPersistentByteAddressBuffer = UE::RendererPrivateUtils::Implementation::T
 
 template <typename InValueType, int32 InNumValuesPerScatter = 1>
 using TByteAddressBufferScatterUploader = UE::RendererPrivateUtils::Implementation::TBufferScatterUploader<InValueType, UE::RendererPrivateUtils::Implementation::FByteAddressBufferTraits, InNumValuesPerScatter>;
+
+/**
+ * Clear indirect args for GPU-side building, initializes the XYZ dimensions of the dispatch arg to DimClearValue and everything else in the strided range to zero. 
+ * Can clear one or more args (NumIndirectArgs). 
+ * IndirectArgStride gives the size (in uints) of each argument, this may be larger than the required size. E.g., to make space for an atomic counter.
+ */
+void AddClearIndirectDispatchArgsPass(FRDGBuilder& GraphBuilder, ERHIFeatureLevel::Type FeatureLevel, FRDGBufferRef IndirectArgsRDG, const FIntVector3 &DimClearValue, uint32 NumIndirectArgs, uint32 IndirectArgStride);
+
+/**
+ * Clear indirect for 1D kernel launch, sets YZ grid dimensions to 1.
+ */
+inline void AddClearIndirectDispatchArgs1DPass(FRDGBuilder& GraphBuilder, ERHIFeatureLevel::Type FeatureLevel, FRDGBufferRef IndirectArgsRDG, uint32 NumIndirectArgs, uint32 IndirectArgStride)
+{
+	AddClearIndirectDispatchArgsPass(GraphBuilder, FeatureLevel, IndirectArgsRDG, FIntVector3(0,1,1), NumIndirectArgs, IndirectArgStride);
+}
+/**
+ * Create an indirect args buffer and clear it using AddClearIndirectDispatchArgsPass.
+ */
+FRDGBufferRef CreateAndClearIndirectDispatchArgs(FRDGBuilder& GraphBuilder, ERHIFeatureLevel::Type FeatureLevel, const TCHAR* Name, const FIntVector3 &DimClearValue, uint32 NumIndirectArgs, uint32 IndirectArgStride);
+
+/**
+ * Create and initialize a singular indirect args buffer with size derived from FRHIDispatchIndirectParameters.
+ */
+inline FRDGBufferRef CreateAndClearIndirectDispatchArgs1D(FRDGBuilder& GraphBuilder, ERHIFeatureLevel::Type FeatureLevel, const TCHAR* Name)
+{
+	return CreateAndClearIndirectDispatchArgs(GraphBuilder, FeatureLevel, Name, FIntVector3(0,1,1), 1u, sizeof(FRHIDispatchIndirectParameters) / 4u);
+}
diff --git a/Engine/Source/Runtime/Renderer/Private/RendererUtils.cpp b/Engine/Source/Runtime/Renderer/Private/RendererUtils.cpp
--- a/Engine/Source/Runtime/Renderer/Private/RendererUtils.cpp
+++ b/Engine/Source/Runtime/Renderer/Private/RendererUtils.cpp
@@ -296,3 +296,46 @@ RENDERER_API FBufferRHIRef& GetOneTileQuadIndexBuffer()
 {
 	return GOneTileQuadIndexBuffer.IndexBufferRHI;
 }
+
+class FClearIndirectDispatchArgsCS : public FGlobalShader
+{
+	DECLARE_GLOBAL_SHADER(FClearIndirectDispatchArgsCS);
+	SHADER_USE_PARAMETER_STRUCT(FClearIndirectDispatchArgsCS, FGlobalShader)
+
+	BEGIN_SHADER_PARAMETER_STRUCT(FParameters, )
+		SHADER_PARAMETER(uint32, NumIndirectArgs)
+		SHADER_PARAMETER(uint32, IndirectArgStride)
+		SHADER_PARAMETER(FIntVector3, DimClearValue)
+		SHADER_PARAMETER_RDG_BUFFER_UAV(RWBuffer<uint>, OutIndirectArgsBuffer)
+	END_SHADER_PARAMETER_STRUCT()
+};
+IMPLEMENT_GLOBAL_SHADER(FClearIndirectDispatchArgsCS, "/Engine/Private/RendererUtils.usf", "ClearIndirectDispatchArgsCS", SF_Compute);
+
+void AddClearIndirectDispatchArgsPass(FRDGBuilder& GraphBuilder, ERHIFeatureLevel::Type FeatureLevel, FRDGBufferRef IndirectArgsRDG, const FIntVector3 &DimClearValue, uint32 NumIndirectArgs, uint32 IndirectArgStride)
+{
+	// Need room for XYZ dims at least.
+	check(IndirectArgStride >= 3);
+
+	FClearIndirectDispatchArgsCS::FParameters* PassParameters = GraphBuilder.AllocParameters<FClearIndirectDispatchArgsCS::FParameters>();
+	PassParameters->NumIndirectArgs = NumIndirectArgs;
+	PassParameters->IndirectArgStride = IndirectArgStride;
+	PassParameters->DimClearValue = DimClearValue; 
+	PassParameters->OutIndirectArgsBuffer = GraphBuilder.CreateUAV(IndirectArgsRDG);
+
+	auto ComputeShader = GetGlobalShaderMap(FeatureLevel)->GetShader<FClearIndirectDispatchArgsCS>();
+
+	FComputeShaderUtils::AddPass(
+		GraphBuilder,
+		RDG_EVENT_NAME("ClearIndirectDispatchArgs"),
+		ComputeShader,
+		PassParameters,
+		FComputeShaderUtils::GetGroupCount(NumIndirectArgs, 64)
+	);
+}
+
+FRDGBufferRef CreateAndClearIndirectDispatchArgs(FRDGBuilder& GraphBuilder, ERHIFeatureLevel::Type FeatureLevel, const TCHAR* Name, const FIntVector3& DimClearValue, uint32 NumIndirectArgs, uint32 IndirectArgStride)
+{
+	FRDGBufferRef IndirectArgsRDG = GraphBuilder.CreateBuffer(FRDGBufferDesc::CreateIndirectDesc(NumIndirectArgs * IndirectArgStride), Name);
+	AddClearIndirectDispatchArgsPass(GraphBuilder, FeatureLevel, IndirectArgsRDG, DimClearValue, NumIndirectArgs, IndirectArgStride);
+	return IndirectArgsRDG;
+}
diff --git a/Engine/Source/Runtime/Renderer/Private/SceneOcclusion.cpp b/Engine/Source/Runtime/Renderer/Private/SceneOcclusion.cpp
--- a/Engine/Source/Runtime/Renderer/Private/SceneOcclusion.cpp
+++ b/Engine/Source/Runtime/Renderer/Private/SceneOcclusion.cpp
@@ -1383,7 +1383,8 @@ void FDeferredShadingSceneRenderer::RenderOcclusion(
 	FRDGBuilder& GraphBuilder,
 	const FSceneTextures& SceneTextures,
 	bool bIsOcclusionTesting,
-	const FBuildHZBAsyncComputeParams* BuildHZBAsyncComputeParams)
+	const FBuildHZBAsyncComputeParams* BuildHZBAsyncComputeParams,
+	Froxel::FRenderer& FroxelRenderer)
 {
 	TRACE_CPUPROFILER_EVENT_SCOPE(FDeferredShadingSceneRenderer::RenderOcclusion);
 
@@ -1475,7 +1476,7 @@ void FDeferredShadingSceneRenderer::RenderOcclusion(
 		FenceOcclusionTests(GraphBuilder);
 	}
 
-	const bool bUseHzbOcclusion = RenderHzb(GraphBuilder, SceneTextures.Depth.Resolve, BuildHZBAsyncComputeParams);
+	const bool bUseHzbOcclusion = RenderHzb(GraphBuilder, SceneTextures.Depth.Resolve, BuildHZBAsyncComputeParams, FroxelRenderer);
 
 	if (bUseHzbOcclusion || bIsOcclusionTesting)
 	{
diff --git a/Engine/Source/Runtime/Renderer/Private/SceneTextureReductions.cpp b/Engine/Source/Runtime/Renderer/Private/SceneTextureReductions.cpp
--- a/Engine/Source/Runtime/Renderer/Private/SceneTextureReductions.cpp
+++ b/Engine/Source/Runtime/Renderer/Private/SceneTextureReductions.cpp
@@ -11,7 +11,7 @@
 #include "PixelShaderUtils.h"
 #include "SceneTextureParameters.h"
 #include "PostProcess/PostProcessing.h" // for FPostProcessVS
-
+#include "Froxel/Froxel.h"
 
 static TAutoConsoleVariable<int32> CVarHZBBuildUseCompute(
 	TEXT("r.HZB.BuildUseCompute"), 1,
@@ -65,12 +65,13 @@ class FHZBBuildCS : public FGlobalShader
 	class FDimVisBufferFormat : SHADER_PERMUTATION_INT("VIS_BUFFER_FORMAT", 4);
 	class FDimFurthest : SHADER_PERMUTATION_BOOL("DIM_FURTHEST");
 	class FDimClosest : SHADER_PERMUTATION_BOOL("DIM_CLOSEST");
+	class FDimFroxels : SHADER_PERMUTATION_BOOL("DIM_FROXELS");
 	class FDimMipLevelCount : SHADER_PERMUTATION_RANGE_INT("DIM_MIP_LEVEL_COUNT", 1, kMaxMipBatchSize);
-	using FPermutationDomain = TShaderPermutationDomain<FDimFurthest, FDimClosest, FDimMipLevelCount, FDimVisBufferFormat>;
+	using FPermutationDomain = TShaderPermutationDomain<FDimFurthest, FDimClosest, FDimFroxels, FDimMipLevelCount, FDimVisBufferFormat>;
 
 	BEGIN_SHADER_PARAMETER_STRUCT(FParameters, )
 		SHADER_PARAMETER_STRUCT_INCLUDE(FSharedHZBParameters, Shared)
-
+		SHADER_PARAMETER_STRUCT_INCLUDE(Froxel::FBuilderParameters, FroxelBuilder)		
 		SHADER_PARAMETER_RDG_TEXTURE_UAV_ARRAY(RWTexture2D<float>, FurthestHZBOutput, [kMaxMipBatchSize])
 		SHADER_PARAMETER_RDG_TEXTURE_UAV_ARRAY(RWTexture2D<float>, ClosestHZBOutput, [kMaxMipBatchSize])
 	END_SHADER_PARAMETER_STRUCT()
@@ -78,7 +79,12 @@ class FHZBBuildCS : public FGlobalShader
 	static bool ShouldCompilePermutation(const FGlobalShaderPermutationParameters& Parameters)
 	{
 		FPermutationDomain PermutationVector(Parameters.PermutationId);
-
+		
+		// No need for froxels on non-VSM platforms
+		if (PermutationVector.Get<FDimFroxels>() && !DoesPlatformSupportVirtualShadowMaps(Parameters.Platform))
+		{
+			return false;
+		}
 		// Necessarily reduce at least closest of furthest.
 		if (!PermutationVector.Get<FDimFurthest>() && !PermutationVector.Get<FDimClosest>())
 		{
@@ -116,7 +122,8 @@ void BuildHZB(
 	const TCHAR* FurthestHZBName,
 	FRDGTextureRef* OutFurthestHZBTexture,
 	EPixelFormat Format,
-	const FBuildHZBAsyncComputeParams* AsyncComputeParams)
+	const FBuildHZBAsyncComputeParams* AsyncComputeParams,
+	const Froxel::FViewData* OutFroxelViewData)
 {
 	RDG_EVENT_SCOPE(GraphBuilder, "BuildHZB");
 	QUICK_SCOPE_CYCLE_COUNTER(STAT_BuildHZB);
@@ -153,7 +160,7 @@ void BuildHZB(
 	auto ReduceMips = [&](
 		FRDGTextureSRVRef ParentTextureMip, FIntPoint SrcSize,
 		int32 StartDestMip, FVector4f DispatchThreadIdToBufferUV, FVector2D InputViewportMaxBound,
-		FIntVector4 PixelViewPortMinMax, bool bOutputClosest, bool bOutputFurthest)
+		FIntVector4 PixelViewPortMinMax, bool bOutputClosest, bool bOutputFurthest, bool bProduceFroxelData)
 	{
 
 		FSharedHZBParameters ShaderParameters;
@@ -177,9 +184,13 @@ void BuildHZB(
 			for (int32 DestMip = StartDestMip; DestMip < EndDestMip; DestMip++)
 			{
 				if (bOutputFurthest)
+				{
 					PassParameters->FurthestHZBOutput[DestMip - StartDestMip] = GraphBuilder.CreateUAV(FRDGTextureUAVDesc(FurthestHZBTexture, DestMip));
+				}
 				if (bOutputClosest)
+				{
 					PassParameters->ClosestHZBOutput[DestMip - StartDestMip] = GraphBuilder.CreateUAV(FRDGTextureUAVDesc(ClosestHZBTexture, DestMip));
+				}
 			}
 
 			int32 VisBufferFormat = 0;
@@ -204,10 +215,16 @@ void BuildHZB(
 				}
 			}
 
+			if (bProduceFroxelData)
+			{
+				PassParameters->FroxelBuilder = OutFroxelViewData->GetBuilderParameters(GraphBuilder);
+			}
+
 			FHZBBuildCS::FPermutationDomain PermutationVector;
 			PermutationVector.Set<FHZBBuildCS::FDimMipLevelCount>(EndDestMip - StartDestMip);
 			PermutationVector.Set<FHZBBuildCS::FDimFurthest>(bOutputFurthest);
 			PermutationVector.Set<FHZBBuildCS::FDimClosest>(bOutputClosest);
+			PermutationVector.Set<FHZBBuildCS::FDimFroxels>(bProduceFroxelData);
 			PermutationVector.Set<FHZBBuildCS::FDimVisBufferFormat>(VisBufferFormat);
 
 			const bool bAsyncCompute = AsyncComputeParams != nullptr;
@@ -216,10 +233,11 @@ void BuildHZB(
 			TShaderMapRef<FHZBBuildCS> ComputeShader(GetGlobalShaderMap(FeatureLevel), PermutationVector);
 			FRDGPassRef ReduceHZBPass = FComputeShaderUtils::AddPass(
 				GraphBuilder,
-				RDG_EVENT_NAME("ReduceHZB(mips=[%d;%d]%s%s) %dx%d",
+				RDG_EVENT_NAME("ReduceHZB(mips=[%d;%d]%s%s%s) %dx%d",
 					StartDestMip, EndDestMip - 1,
 					bOutputClosest ? TEXT(" Closest") : TEXT(""),
 					bOutputFurthest ? TEXT(" Furthest") : TEXT(""),
+					bProduceFroxelData ? TEXT(" Froxels") : TEXT(""),
 					DstSize.X, DstSize.Y),
 				PassFlags,
 				ComputeShader,
@@ -233,6 +251,7 @@ void BuildHZB(
 		}
 		else
 		{
+			check(!bProduceFroxelData);
 			check(bOutputFurthest);
 			check(!bOutputClosest);
 
@@ -254,7 +273,7 @@ void BuildHZB(
 		}
 	};
 
-	// Reduce first mips Closesy and furtherest are done at same time.
+	// Reduce first mips Closest and furtherest are done at same time.
 	{
 		FRDGTextureSRVRef ParentTextureMip = GraphBuilder.CreateSRV(FRDGTextureSRVDesc::Create(SceneDepth));
 		FIntPoint SrcSize = VisBufferTexture ? VisBufferTexture->Desc.Extent : ParentTextureMip->Desc.Texture->Desc.Extent;
@@ -275,7 +294,8 @@ void BuildHZB(
 		ReduceMips(
 			ParentTextureMip, SrcSize,
 			/* StartDestMip = */ 0, DispatchThreadIdToBufferUV, InputViewportMaxBound, PixelViewPortMinMax,
-			/* bOutputClosest = */ bReduceClosestDepth, /* bOutputFurthest = */ true);
+			/* bOutputClosest = */ bReduceClosestDepth, /* bOutputFurthest = */ true,
+			OutFroxelViewData != nullptr);
 	}
 
 	// Reduce the next mips
@@ -296,7 +316,7 @@ void BuildHZB(
 			FRDGTextureSRVRef ParentTextureMip = GraphBuilder.CreateSRV(GRHISupportsTextureViews ? FRDGTextureSRVDesc::CreateForMipLevel(FurthestHZBTexture, StartDestMip - 1) : FRDGTextureSRVDesc::Create(FurthestHZBTexture));
 			ReduceMips(ParentTextureMip, SrcSize,
 				StartDestMip, DispatchThreadIdToBufferUV, InputViewportMaxBound, PixelViewPortMinMax,
-				/* bOutputClosest = */ false, /* bOutputFurthest = */ true);
+				/* bOutputClosest = */ false, /* bOutputFurthest = */ true, false);
 		}
 
 		if (bReduceClosestDepth)
@@ -305,7 +325,7 @@ void BuildHZB(
 			FRDGTextureSRVRef ParentTextureMip = GraphBuilder.CreateSRV(GRHISupportsTextureViews ? FRDGTextureSRVDesc::CreateForMipLevel(ClosestHZBTexture, StartDestMip - 1) : FRDGTextureSRVDesc::Create(ClosestHZBTexture));
 			ReduceMips(ParentTextureMip, SrcSize,
 				StartDestMip, DispatchThreadIdToBufferUV, InputViewportMaxBound, PixelViewPortMinMax,
-				/* bOutputClosest = */ true, /* bOutputFurthest = */ false);
+				/* bOutputClosest = */ true, /* bOutputFurthest = */ false, false);
 		}
 	}
 
diff --git a/Engine/Source/Runtime/Renderer/Private/SceneTextureReductions.h b/Engine/Source/Runtime/Renderer/Private/SceneTextureReductions.h
--- a/Engine/Source/Runtime/Renderer/Private/SceneTextureReductions.h
+++ b/Engine/Source/Runtime/Renderer/Private/SceneTextureReductions.h
@@ -6,6 +6,11 @@
 #include "RenderGraphFwd.h"
 #include "RHIFwd.h"
 
+namespace Froxel
+{
+	struct FViewData;
+}
+
 struct FBuildHZBAsyncComputeParams
 {
 	FRDGPassRef Prerequisite = nullptr;
@@ -25,7 +30,8 @@ void BuildHZB(
 	const TCHAR* FurthestHZBName,
 	FRDGTextureRef* OutFurthestHZBTexture,
 	EPixelFormat Format = BuildHZBDefaultPixelFormat,
-	const FBuildHZBAsyncComputeParams* AsyncComputeParams = nullptr);
+	const FBuildHZBAsyncComputeParams* AsyncComputeParams = nullptr,
+	const Froxel::FViewData* OutFroxelData = nullptr);
 
 // Build only the furthest HZB
 void BuildHZBFurthest(
diff --git a/Engine/Source/Runtime/Renderer/Private/SingleLayerWaterRendering.cpp b/Engine/Source/Runtime/Renderer/Private/SingleLayerWaterRendering.cpp
--- a/Engine/Source/Runtime/Renderer/Private/SingleLayerWaterRendering.cpp
+++ b/Engine/Source/Runtime/Renderer/Private/SingleLayerWaterRendering.cpp
@@ -298,15 +298,18 @@ class FWaterTileCategorisationMarkCS : public FGlobalShader
 	SHADER_USE_PARAMETER_STRUCT(FWaterTileCategorisationMarkCS, FGlobalShader)
 
 	class FUsePrepassStencil : SHADER_PERMUTATION_BOOL("USE_WATER_PRE_PASS_STENCIL");
-	using FPermutationDomain = TShaderPermutationDomain<FUsePrepassStencil>;
+	class FBuildFroxels : SHADER_PERMUTATION_BOOL("GENERATE_FROXELS");
+	using FPermutationDomain = TShaderPermutationDomain<FUsePrepassStencil, FBuildFroxels>;
 
 	BEGIN_SHADER_PARAMETER_STRUCT(FParameters, )
 		SHADER_PARAMETER_STRUCT_INCLUDE(FSceneTextureParameters, SceneTextures)	// Water scene texture
 		SHADER_PARAMETER_RDG_UNIFORM_BUFFER(FSubstrateGlobalUniformParameters, Substrate)
 		SHADER_PARAMETER_STRUCT_INCLUDE(FViewShaderParameters, View)
 		SHADER_PARAMETER_RDG_TEXTURE_SRV(Texture2D, WaterDepthStencilTexture)
+		SHADER_PARAMETER_RDG_TEXTURE(Texture2D<float>, WaterDepthTexture)
 		SHADER_PARAMETER(FIntPoint, TiledViewRes)
 		SHADER_PARAMETER_RDG_BUFFER_UAV(RWStructuredBuffer<uint>, TileMaskBufferOut)
+		SHADER_PARAMETER_STRUCT_INCLUDE(Froxel::FBuilderParameters, FroxelBuilder)
 	END_SHADER_PARAMETER_STRUCT()
 
 	static FPermutationDomain RemapPermutation(FPermutationDomain PermutationVector)
@@ -316,6 +319,15 @@ class FWaterTileCategorisationMarkCS : public FGlobalShader
 
 	static bool ShouldCompilePermutation(const FGlobalShaderPermutationParameters& Parameters)
 	{
+		FPermutationDomain PermutationVector(Parameters.PermutationId);
+		// No need for froxels on non-VSM platforms
+		if (PermutationVector.Get<FBuildFroxels>() 
+			// only compile if on a supported platform & we have depth stencil avaliable
+			&& (!DoesPlatformSupportVirtualShadowMaps(Parameters.Platform) || !PermutationVector.Get<FUsePrepassStencil>()))
+		{
+			return false;
+		}
+
 		return UseSingleLayerWaterIndirectDraw(Parameters.Platform);
 	}
 
@@ -450,7 +462,7 @@ static FSingleLayerWaterDepthPassParameters* GetSingleLayerWaterDepthPassParamet
  * Build lists of 8x8 tiles used by water pixels
  * Mark and build list steps are separated in order to build a more coherent list (z-ordered over a larger region), which is important for the performance of future passes like ray traced Lumen reflections
  */
-static FSingleLayerWaterTileClassification ClassifyTiles(FRDGBuilder& GraphBuilder, const FViewInfo &View, const FSceneTextures& SceneTextures, const FRDGTextureRef& DepthPrepassTexture)
+static FSingleLayerWaterTileClassification ClassifyTiles(FRDGBuilder& GraphBuilder, const FViewInfo &View, const FSceneTextures& SceneTextures, const FRDGTextureRef& DepthPrepassTexture, const Froxel::FViewData* OutFroxelViewData)
 {
 	FSingleLayerWaterTileClassification Result;
 	const bool bRunTiled = UseSingleLayerWaterIndirectDraw(View.GetShaderPlatform()) && CVarWaterSingleLayerTiledComposite.GetValueOnRenderThread();
@@ -477,10 +489,15 @@ static FSingleLayerWaterTileClassification ClassifyTiles(FRDGBuilder& GraphBuild
 		AddClearUAVPass(GraphBuilder, DrawIndirectParametersBufferUAV, 0);
 		AddClearUAVPass(GraphBuilder, DispatchIndirectParametersBufferUAV, 0);
 
+		// Can't produce froxels unless we have depth data
+		bool bProduceFroxelData = OutFroxelViewData != nullptr && DepthPrepassTexture != nullptr;
+
 		// Mark used tiles based on SHADING_MODEL_ID
 		{
 			FWaterTileCategorisationMarkCS::FPermutationDomain PermutationVector;
 			PermutationVector.Set<FWaterTileCategorisationMarkCS::FUsePrepassStencil>(DepthPrepassTexture != nullptr);
+			PermutationVector.Set<FWaterTileCategorisationMarkCS::FBuildFroxels>(bProduceFroxelData);
+
 			TShaderMapRef<FWaterTileCategorisationMarkCS> ComputeShader(View.ShaderMap, PermutationVector);
 
 			FWaterTileCategorisationMarkCS::FParameters* PassParameters = GraphBuilder.AllocParameters<FWaterTileCategorisationMarkCS::FParameters>();
@@ -490,8 +507,12 @@ static FSingleLayerWaterTileClassification ClassifyTiles(FRDGBuilder& GraphBuild
 			PassParameters->Substrate = Substrate::BindSubstrateGlobalUniformParameters(View);
 			PassParameters->TiledViewRes = Result.TiledViewRes;
 			PassParameters->WaterDepthStencilTexture = DepthPrepassTexture ? GraphBuilder.CreateSRV(FRDGTextureSRVDesc::CreateWithPixelFormat(DepthPrepassTexture, PF_X24_G8)) : nullptr;
+			PassParameters->WaterDepthTexture = DepthPrepassTexture ? DepthPrepassTexture : nullptr;
 			PassParameters->TileMaskBufferOut = TileMaskBufferUAV;
-
+			if (bProduceFroxelData)
+			{
+				PassParameters->FroxelBuilder = OutFroxelViewData->GetBuilderParameters(GraphBuilder);
+			}
 			FComputeShaderUtils::AddPass(
 				GraphBuilder,
 				RDG_EVENT_NAME("SLW::TileCategorisationMarkTiles"),
@@ -658,12 +679,14 @@ FSingleLayerWaterPrePassResult* FDeferredShadingSceneRenderer::RenderSingleLayer
 	AddResolveSceneDepthPass(GraphBuilder, InViews, OutDepthPrepassTexture);
 
 	// Run classification pass.
-	for (int32 ViewIndex = 0; ViewIndex < InViews.Num(); ++ViewIndex)
+	if (UseSingleLayerWaterIndirectDraw(ShaderPlatform) && CVarWaterSingleLayerTiledComposite.GetValueOnRenderThread())
 	{
-		FViewInfo& View = InViews[ViewIndex];
-		if (UseSingleLayerWaterIndirectDraw(View.GetShaderPlatform()) && CVarWaterSingleLayerTiledComposite.GetValueOnRenderThread())
+		Result->Froxels = Froxel::FRenderer(DoesVSMWantFroxels(ShaderPlatform), GraphBuilder, Views);
+
+		for (int32 ViewIndex = 0; ViewIndex < InViews.Num(); ++ViewIndex)
 		{
-			Result->ViewTileClassification[ViewIndex] = ClassifyTiles(GraphBuilder, View, SceneTextures, OutDepthPrepassTexture.Resolve);
+			FViewInfo& View = InViews[ViewIndex];
+			Result->ViewTileClassification[ViewIndex] = ClassifyTiles(GraphBuilder, View, SceneTextures, OutDepthPrepassTexture.Resolve, Result->Froxels.GetView(ViewIndex));
 		}
 	}
 
@@ -877,7 +900,7 @@ void FDeferredShadingSceneRenderer::RenderSingleLayerWaterReflections(
 			}
 			else
 			{
-				SingleLayerWaterTileClassification = ClassifyTiles(GraphBuilder, View, SceneTextures, nullptr);
+				SingleLayerWaterTileClassification = ClassifyTiles(GraphBuilder, View, SceneTextures, nullptr, nullptr);
 			}
 		}
 		FTiledReflection& TiledScreenSpaceReflection = SingleLayerWaterTileClassification.TiledReflection;
diff --git a/Engine/Source/Runtime/Renderer/Private/SingleLayerWaterRendering.h b/Engine/Source/Runtime/Renderer/Private/SingleLayerWaterRendering.h
--- a/Engine/Source/Runtime/Renderer/Private/SingleLayerWaterRendering.h
+++ b/Engine/Source/Runtime/Renderer/Private/SingleLayerWaterRendering.h
@@ -6,6 +6,7 @@
 #include "PixelShaderUtils.h"
 #include "ScreenSpaceRayTracing.h"
 #include "SingleLayerWaterDefinitions.h"
+#include "Froxel/Froxel.h"
 
 class FViewInfo;
 
@@ -18,6 +19,7 @@ struct FSingleLayerWaterTileClassification
 
 struct FSingleLayerWaterPrePassResult
 {
+	Froxel::FRenderer Froxels;
 	FRDGTextureMSAA DepthPrepassTexture;
 	TArray<FSingleLayerWaterTileClassification> ViewTileClassification;
 };
diff --git a/Engine/Source/Runtime/Renderer/Private/VirtualShadowMaps/VirtualShadowMapArray.cpp b/Engine/Source/Runtime/Renderer/Private/VirtualShadowMaps/VirtualShadowMapArray.cpp
--- a/Engine/Source/Runtime/Renderer/Private/VirtualShadowMaps/VirtualShadowMapArray.cpp
+++ b/Engine/Source/Runtime/Renderer/Private/VirtualShadowMaps/VirtualShadowMapArray.cpp
@@ -23,6 +23,7 @@
 #include "SingleLayerWaterRendering.h"
 #include "RenderUtils.h"
 #include "SceneCulling/SceneCullingRenderer.h"
+#include "Froxel/Froxel.h"
 
 #define DEBUG_ALLOW_STATIC_SEPARATE_WITHOUT_CACHING 0
 
@@ -51,6 +52,28 @@ FAutoConsoleVariableRef CVarVSMShowLightDrawEvents(
 	ECVF_RenderThreadSafe
 );
 
+TAutoConsoleVariable<int32> CVarMarkPagesUseFroxels(
+	TEXT("r.Shadow.Virtual.MarkPagesUsingFroxels"),
+	0,
+	TEXT("Experimental: If enabled the virtual shadow map pages are marked using froxels that are generated during HZB build.\n")
+	TEXT("  Higher throughput as it is not bandwidth limited. Is approximate as it only marks the center of each froxel representing 8x8 pixels."),
+	ECVF_RenderThreadSafe
+);
+
+TAutoConsoleVariable<int32> CVarDebugDrawFroxels(
+	TEXT("r.Shadow.Virtual.DebugDrawFroxels"),
+	0,
+	TEXT("Render the froxels using shaderprint (which needs to be enabled) r.ShaderPrint.MaxLine also needs to be set to a high value as this produces many lines."),
+	ECVF_RenderThreadSafe
+);
+
+TAutoConsoleVariable<float> CVarDebugDrawFroxelRange(
+	TEXT("r.Shadow.Virtual.DebugDrawFroxelRange"),
+	20.0f,
+	TEXT("Range in froxel tiles from the mouse cursor which to draw debug froxels in."),
+	ECVF_RenderThreadSafe
+);
+
 TAutoConsoleVariable<int32> CVarEnableVirtualShadowMaps(
 	TEXT("r.Shadow.Virtual.Enable"),
 	0,
@@ -501,6 +524,14 @@ bool IsVSMTranslucentHighQualityEnabled();
 bool IsLumenFrontLayerHistoryValid(const FViewInfo& View);
 extern bool LightGridUses16BitBuffers(EShaderPlatform Platform);
 
+bool DoesVSMWantFroxels(EShaderPlatform ShaderPlatform)
+{
+	return UseVirtualShadowMaps(ShaderPlatform)
+		&& CVarMarkPagesUseFroxels.GetValueOnRenderThread() != 0
+		// fall back to per-pixel marking if the front layer translucency path is enabled as it doesn't generate froxels.
+		&& !IsVSMTranslucentHighQualityEnabled();
+}
+
 FMatrix CalcTranslatedWorldToShadowUVMatrix(
 	const FMatrix& TranslatedWorldToShadowView,
 	const FMatrix& ViewToClip)
@@ -840,6 +871,21 @@ class FPruneLightGridCS : public FVirtualShadowMapPageManagementShader
 };
 IMPLEMENT_GLOBAL_SHADER(FPruneLightGridCS, "/Engine/Private/VirtualShadowMaps/VirtualShadowMapPageMarking.usf", "PruneLightGridCS", SF_Compute);
 
+BEGIN_SHADER_PARAMETER_STRUCT(FBaseGeneratePageFlagsParameters, )
+	SHADER_PARAMETER_RDG_UNIFORM_BUFFER(FVirtualShadowMapUniformParameters, VirtualShadowMap)
+	SHADER_PARAMETER_RDG_UNIFORM_BUFFER(FSceneTextureUniformParameters, SceneTexturesStruct)
+	SHADER_PARAMETER_STRUCT_REF(FViewUniformShaderParameters, View)
+	SHADER_PARAMETER_RDG_UNIFORM_BUFFER(FForwardLightData, ForwardLightData)
+	SHADER_PARAMETER_RDG_BUFFER_UAV(RWStructuredBuffer<uint>, OutPageRequestFlags)
+	SHADER_PARAMETER_RDG_BUFFER_SRV(StructuredBuffer< uint >, DirectionalLightIds)
+	SHADER_PARAMETER_RDG_UNIFORM_BUFFER(FSubstrateGlobalUniformParameters, Substrate)
+	SHADER_PARAMETER(float, PageDilationBorderSizeDirectional)
+	SHADER_PARAMETER(float, PageDilationBorderSizeLocal)
+	SHADER_PARAMETER(uint32, bCullBackfacingPixels)
+	SHADER_PARAMETER(uint32, NumDirectionalLightSmInds)
+	SHADER_PARAMETER(uint32, MipModeLocal)
+END_SHADER_PARAMETER_STRUCT()
+
 class FGeneratePageFlagsFromPixelsCS : public FVirtualShadowMapPageManagementShader
 {
 	DECLARE_GLOBAL_SHADER(FGeneratePageFlagsFromPixelsCS);
@@ -861,38 +907,49 @@ class FGeneratePageFlagsFromPixelsCS : public FVirtualShadowMapPageManagementSha
 	}
 
 	BEGIN_SHADER_PARAMETER_STRUCT(FParameters, )
-		SHADER_PARAMETER_RDG_UNIFORM_BUFFER(FVirtualShadowMapUniformParameters, VirtualShadowMap)
-		SHADER_PARAMETER_RDG_UNIFORM_BUFFER(FSceneTextureUniformParameters, SceneTexturesStruct)
+		SHADER_PARAMETER_STRUCT_INCLUDE(FBaseGeneratePageFlagsParameters, Base)
 		SHADER_PARAMETER_RDG_UNIFORM_BUFFER(FHairStrandsViewUniformParameters, HairStrands)
-		SHADER_PARAMETER_STRUCT_REF(FViewUniformShaderParameters, View)
-		SHADER_PARAMETER_RDG_UNIFORM_BUFFER(FForwardLightData, ForwardLightData)
-		SHADER_PARAMETER_RDG_TEXTURE(Texture2D<uint2>, VisBuffer64)
-		SHADER_PARAMETER_RDG_BUFFER_UAV(RWStructuredBuffer<uint>, OutPageRequestFlags)
-		SHADER_PARAMETER_RDG_BUFFER_SRV(StructuredBuffer< uint >, DirectionalLightIds)
 		// PERMUTATION_WATER_DEPTH
 		SHADER_PARAMETER_RDG_TEXTURE(Texture2D<float>, SingleLayerWaterDepthTexture)
 		SHADER_PARAMETER_RDG_BUFFER_SRV(StructuredBuffer< uint >, SingleLayerWaterTileMask)
 		SHADER_PARAMETER(FIntPoint, SingleLayerWaterTileViewRes)
 		// PERMUTATION_TRANSLUCENCY_DEPTH
+		// FRONT LAYER 
 		SHADER_PARAMETER(uint32, FrontLayerMode)
 		SHADER_PARAMETER(FVector4f, FrontLayerHistoryUVMinMax)
 		SHADER_PARAMETER(FVector4f, FrontLayerHistoryScreenPositionScaleBias)
 		SHADER_PARAMETER(FVector4f, FrontLayerHistoryBufferSizeAndInvSize)
 		SHADER_PARAMETER_RDG_TEXTURE(Texture2D<float>, FrontLayerTranslucencyDepthTexture)
 		SHADER_PARAMETER_RDG_TEXTURE(Texture2D<float>, FrontLayerTranslucencyNormalTexture)
-		SHADER_PARAMETER_RDG_UNIFORM_BUFFER(FSubstrateGlobalUniformParameters, Substrate)
+
 		RDG_BUFFER_ACCESS(IndirectBufferArgs, ERHIAccess::IndirectArgs)
-		SHADER_PARAMETER(float, PageDilationBorderSizeDirectional)
-		SHADER_PARAMETER(float, PageDilationBorderSizeLocal)
 		SHADER_PARAMETER(uint32, InputType)
-		SHADER_PARAMETER(uint32, bCullBackfacingPixels)
-		SHADER_PARAMETER(uint32, NumDirectionalLightSmInds)
-		SHADER_PARAMETER(uint32, MipModeLocal)
 		SHADER_PARAMETER(FIntPoint, PixelStride)
 	END_SHADER_PARAMETER_STRUCT()
 };
 IMPLEMENT_GLOBAL_SHADER(FGeneratePageFlagsFromPixelsCS, "/Engine/Private/VirtualShadowMaps/VirtualShadowMapPageMarking.usf", "GeneratePageFlagsFromPixels", SF_Compute);
 
+class FGeneratePageFlagsFromFroxelsCS : public FVirtualShadowMapPageManagementShader
+{
+	DECLARE_GLOBAL_SHADER(FGeneratePageFlagsFromFroxelsCS);
+	SHADER_USE_PARAMETER_STRUCT(FGeneratePageFlagsFromFroxelsCS, FVirtualShadowMapPageManagementShader)
+
+	class FDebugRenderDim : SHADER_PERMUTATION_BOOL("DEBUG_DRAW_GENERATE_FROM_FROXELS"); 
+	using FPermutationDomain = TShaderPermutationDomain<FDebugRenderDim>;
+
+	BEGIN_SHADER_PARAMETER_STRUCT(FParameters, )
+		SHADER_PARAMETER_STRUCT_INCLUDE(FBaseGeneratePageFlagsParameters, Base)
+		SHADER_PARAMETER(uint32, bShouldMarkLocaLights)
+		SHADER_PARAMETER(int32, PassId)
+		SHADER_PARAMETER(float, DebugRange)
+		SHADER_PARAMETER_STRUCT_INCLUDE(Froxel::FParameters, FroxelParameters)
+		SHADER_PARAMETER_STRUCT_INCLUDE( ShaderPrint::FShaderParameters, ShaderPrintStruct )
+		RDG_BUFFER_ACCESS(IndirectBufferArgs, ERHIAccess::IndirectArgs)
+	END_SHADER_PARAMETER_STRUCT()
+};
+IMPLEMENT_GLOBAL_SHADER(FGeneratePageFlagsFromFroxelsCS, "/Engine/Private/VirtualShadowMaps/VirtualShadowMapPageMarking.usf", "GeneratePageFlagsFromFroxelsCS", SF_Compute);
+
+
 class FMarkCoarsePagesCS : public FVirtualShadowMapPageManagementShader
 {
 	DECLARE_GLOBAL_SHADER(FMarkCoarsePagesCS);
@@ -1069,44 +1126,6 @@ class FInitializePhysicalPagesIndirectCS : public FVirtualShadowMapPageManagemen
 };
 IMPLEMENT_GLOBAL_SHADER(FInitializePhysicalPagesIndirectCS, "/Engine/Private/VirtualShadowMaps/VirtualShadowMapPhysicalPageManagement.usf", "InitializePhysicalPagesIndirectCS", SF_Compute);
 
-class FClearIndirectDispatchArgs1DCS : public FVirtualShadowMapPageManagementShader
-{
-	DECLARE_GLOBAL_SHADER(FClearIndirectDispatchArgs1DCS);
-	SHADER_USE_PARAMETER_STRUCT(FClearIndirectDispatchArgs1DCS, FVirtualShadowMapPageManagementShader)
-
-	BEGIN_SHADER_PARAMETER_STRUCT(FParameters, )
-		SHADER_PARAMETER(uint32, NumIndirectArgs)
-		SHADER_PARAMETER(uint32, IndirectArgStride)
-		SHADER_PARAMETER_RDG_BUFFER_UAV(RWBuffer<uint>, OutIndirectArgsBuffer)
-	END_SHADER_PARAMETER_STRUCT()
-};
-IMPLEMENT_GLOBAL_SHADER(FClearIndirectDispatchArgs1DCS, "/Engine/Private/VirtualShadowMaps/VirtualShadowMapPhysicalPageManagement.usf", "ClearIndirectDispatchArgs1DCS", SF_Compute);
-
-static void AddClearIndirectDispatchArgs1DPass(FRDGBuilder& GraphBuilder, ERHIFeatureLevel::Type FeatureLevel, FRDGBufferRef IndirectArgsRDG, uint32 NumIndirectArgs = 1U, uint32 IndirectArgStride = 4U)
-{
-	FClearIndirectDispatchArgs1DCS::FParameters* PassParameters = GraphBuilder.AllocParameters<FClearIndirectDispatchArgs1DCS::FParameters>();
-	PassParameters->NumIndirectArgs = NumIndirectArgs;
-	PassParameters->IndirectArgStride = IndirectArgStride;
-	PassParameters->OutIndirectArgsBuffer = GraphBuilder.CreateUAV(IndirectArgsRDG);
-
-	auto ComputeShader = GetGlobalShaderMap(FeatureLevel)->GetShader<FClearIndirectDispatchArgs1DCS>();
-
-	FComputeShaderUtils::AddPass(
-		GraphBuilder,
-		RDG_EVENT_NAME("ClearIndirectDispatchArgs"),
-		ComputeShader,
-		PassParameters,
-		FComputeShaderUtils::GetGroupCount(NumIndirectArgs, 64)
-	);
-}
-
-FRDGBufferRef FVirtualShadowMapArray::CreateAndInitializeDispatchIndirectArgs1D(FRDGBuilder& GraphBuilder, ERHIFeatureLevel::Type FeatureLevel, const TCHAR* Name)
-{
-	FRDGBufferRef IndirectArgsRDG = GraphBuilder.CreateBuffer(FRDGBufferDesc::CreateIndirectDesc<FRHIDispatchIndirectParameters>(), Name);
-	AddClearIndirectDispatchArgs1DPass(GraphBuilder, FeatureLevel, IndirectArgsRDG);
-	return IndirectArgsRDG;
-}
-
 class FSelectPagesToMergeCS : public FVirtualShadowMapPageManagementShader
 {
 	DECLARE_GLOBAL_SHADER(FSelectPagesToMergeCS);
@@ -1162,7 +1181,7 @@ void FVirtualShadowMapArray::MergeStaticPhysicalPages(FRDGBuilder& GraphBuilder)
 	FRDGBufferRef PhysicalPagesToMergeRDG = GraphBuilder.CreateBuffer(FRDGBufferDesc::CreateStructuredDesc(sizeof(int32), GetTotalAllocatedPhysicalPages() + 1), TEXT("Shadow.Virtual.PhysicalPagesToMerge"));
 
 	// 1. Initialize the indirect args buffer
-	FRDGBufferRef MergePagesIndirectArgsRDG = CreateAndInitializeDispatchIndirectArgs1D(GraphBuilder, Scene.GetFeatureLevel(), TEXT("Shadow.Virtual.MergePagesIndirectArgs"));
+	FRDGBufferRef MergePagesIndirectArgsRDG = CreateAndClearIndirectDispatchArgs1D(GraphBuilder, Scene.GetFeatureLevel(), TEXT("Shadow.Virtual.MergePagesIndirectArgs"));
 
 	// 2. Filter the relevant physical pages and set up the indirect args
 	{
@@ -1528,7 +1547,8 @@ void FVirtualShadowMapArray::BuildPageAllocations(
 	const FSortedLightSetSceneInfo& SortedLightsInfo,
 	const TConstArrayView<FVisibleLightInfo>& VisibleLightInfos,
 	const FSingleLayerWaterPrePassResult* SingleLayerWaterPrePassResult,
-	const FFrontLayerTranslucencyData& FrontLayerTranslucencyData)
+	const FFrontLayerTranslucencyData& FrontLayerTranslucencyData,
+	const Froxel::FRenderer& FroxelRenderer)
 {
 	check(IsEnabled());
 
@@ -1735,12 +1755,24 @@ void FVirtualShadowMapArray::BuildPageAllocations(
 				// It's currently safe to overlap these passes that all write to same page request flags
 				FRDGBufferUAVRef PageRequestFlagsUAV = GraphBuilder.CreateUAV(PageRequestFlagsRDG, ERDGUnorderedAccessViewFlags::SkipBarrier);
 
+				FBaseGeneratePageFlagsParameters BaseParameters;
+				BaseParameters.VirtualShadowMap = GetUncachedUniformBuffer(GraphBuilder);
+				BaseParameters.SceneTexturesStruct = SceneTextures.UniformBuffer;
+				BaseParameters.View = View.ViewUniformBuffer;
+				BaseParameters.OutPageRequestFlags = PageRequestFlagsUAV;
+				BaseParameters.ForwardLightData = View.ForwardLightingResources.ForwardLightUniformBuffer;
+				BaseParameters.DirectionalLightIds = GraphBuilder.CreateSRV(DirectionalLightIdsRDG);
+				BaseParameters.NumDirectionalLightSmInds = uint32(DirectionalLightIds.Num());
+				BaseParameters.PageDilationBorderSizeLocal = CVarPageDilationBorderSizeLocal.GetValueOnRenderThread();
+				BaseParameters.PageDilationBorderSizeDirectional = CVarPageDilationBorderSizeDirectional.GetValueOnRenderThread();
+				BaseParameters.bCullBackfacingPixels = ShouldCullBackfacingPixels() ? 1 : 0;
+				BaseParameters.Substrate = Substrate::BindSubstrateGlobalUniformParameters(View);
+				BaseParameters.MipModeLocal = CVarMarkPixelPagesMipModeLocal.GetValueOnRenderThread();
+
 				auto GeneratePageFlags = [&](const EVirtualShadowMapProjectionInputType InputType)
 				{
 					FGeneratePageFlagsFromPixelsCS::FParameters* PassParameters = GraphBuilder.AllocParameters< FGeneratePageFlagsFromPixelsCS::FParameters >();
-					PassParameters->VirtualShadowMap = GetUncachedUniformBuffer(GraphBuilder);
-
-					PassParameters->SceneTexturesStruct = SceneTextures.UniformBuffer;
+					PassParameters->Base = BaseParameters;
 
 					const FIntPoint PixelStride(
 						FMath::Clamp(CVarVirtualShadowMapPageMarkingPixelStrideX.GetValueOnRenderThread(), 1, 128),
@@ -1777,10 +1809,6 @@ void FVirtualShadowMapArray::BuildPageAllocations(
 					}
 
 					PassParameters->HairStrands = HairStrands::BindHairStrandsViewUniformParameters(View);
-					PassParameters->View = View.ViewUniformBuffer;
-					PassParameters->OutPageRequestFlags = PageRequestFlagsUAV;
-					PassParameters->ForwardLightData = View.ForwardLightingResources.ForwardLightUniformBuffer;
-					PassParameters->DirectionalLightIds = GraphBuilder.CreateSRV(DirectionalLightIdsRDG);
 					bool bWaterEnabled = false;
 					if (SingleLayerWaterPrePassResult && InputType == EVirtualShadowMapProjectionInputType::GBuffer)
 					{
@@ -1791,13 +1819,7 @@ void FVirtualShadowMapArray::BuildPageAllocations(
 						PassParameters->SingleLayerWaterTileViewRes = SingleLayerWaterPrePassResult->ViewTileClassification[ViewIndex].TiledViewRes;
 						bWaterEnabled = true;
 					}
-					PassParameters->NumDirectionalLightSmInds = uint32(DirectionalLightIds.Num());
-					PassParameters->PageDilationBorderSizeLocal = CVarPageDilationBorderSizeLocal.GetValueOnRenderThread();
-					PassParameters->PageDilationBorderSizeDirectional = CVarPageDilationBorderSizeDirectional.GetValueOnRenderThread();
-					PassParameters->bCullBackfacingPixels = ShouldCullBackfacingPixels() ? 1 : 0;
-					PassParameters->Substrate = Substrate::BindSubstrateGlobalUniformParameters(View);
 					PassParameters->PixelStride = PixelStride;
-					PassParameters->MipModeLocal = CVarMarkPixelPagesMipModeLocal.GetValueOnRenderThread();
 					
 					const FIntPoint StridedPixelSize = FIntPoint::DivideAndRoundUp(View.ViewRect.Size(), PixelStride);
 					// Note: we use the tile size defined by the water as the group-size - this is needed because the tile mask testing code relies on the size being the same to scalarize efficiently.
@@ -1834,8 +1856,53 @@ void FVirtualShadowMapArray::BuildPageAllocations(
 							FIntVector(GridSize.X, GridSize.Y, 1));
 					}
 				};
+				if (FroxelRenderer.IsEnabled() && CVarMarkPagesUseFroxels.GetValueOnRenderThread() != 0)
+				{
+					auto AddPass_FroxelBuild = [&](const Froxel::FViewData* ViewFroxelData, int32 PassId, bool bShouldMarkLocaLights)
+					{
+						if (ViewFroxelData == nullptr)
+						{
+							return;
+						}
+
+						FGeneratePageFlagsFromFroxelsCS::FParameters* PassParameters = GraphBuilder.AllocParameters< FGeneratePageFlagsFromFroxelsCS::FParameters >();
+						PassParameters->Base = BaseParameters;
+						PassParameters->PassId = PassId;
+						PassParameters->bShouldMarkLocaLights = bShouldMarkLocaLights ? 1u : 0u;
+						PassParameters->DebugRange = CVarDebugDrawFroxelRange.GetValueOnRenderThread();
+
+						PassParameters->FroxelParameters = ViewFroxelData->GetShaderParameters(GraphBuilder);
+						PassParameters->IndirectBufferArgs = ViewFroxelData->FroxelArgsRDG;
 
-				GeneratePageFlags(EVirtualShadowMapProjectionInputType::GBuffer);
+						FGeneratePageFlagsFromFroxelsCS::FPermutationDomain PermutationVector;
+						bool bDebugRender = CVarDebugDrawFroxels.GetValueOnRenderThread() < 0 || CVarDebugDrawFroxels.GetValueOnRenderThread() == (PassId + 1);
+						if (bDebugRender)
+						{
+							PermutationVector.Set<FGeneratePageFlagsFromFroxelsCS::FDebugRenderDim>(true);
+							ShaderPrint::SetParameters(GraphBuilder, View.ShaderPrintData, PassParameters->ShaderPrintStruct);
+						}
+
+						auto ComputeShader = View.ShaderMap->GetShader<FGeneratePageFlagsFromFroxelsCS>(PermutationVector);
+						FComputeShaderUtils::AddPass(
+							GraphBuilder,
+							RDG_EVENT_NAME("GeneratePageFlagsFromFroxels(NumShadowMaps=%d)", GetNumFullShadowMaps()),
+							ComputeShader,
+							PassParameters,
+							ViewFroxelData->FroxelArgsRDG,
+							ViewFroxelData->ArgsOffset);
+					};
+
+					AddPass_FroxelBuild(FroxelRenderer.GetView(ViewIndex), 0, true);
+					if (SingleLayerWaterPrePassResult && SingleLayerWaterPrePassResult->Froxels.GetView(ViewIndex))
+					{
+						// Not marking local lights since SLW does not support sampling shadow for these.
+						AddPass_FroxelBuild(SingleLayerWaterPrePassResult->Froxels.GetView(ViewIndex), 1, false);
+					}
+				}
+				else
+				{
+					GeneratePageFlags(EVirtualShadowMapProjectionInputType::GBuffer);
+				}
 				if (HairStrands::HasViewHairStrandsData(View))
 				{
 					GeneratePageFlags(EVirtualShadowMapProjectionInputType::HairStrands);
@@ -1977,7 +2044,7 @@ void FVirtualShadowMapArray::BuildPageAllocations(
 		FRDGBufferRef PhysicalPagesToInitializeRDG = GraphBuilder.CreateBuffer(FRDGBufferDesc::CreateStructuredDesc(sizeof(int32), GetTotalAllocatedPhysicalPages() + 1), TEXT("Shadow.Virtual.PhysicalPagesToInitialize"));
 
 		// 1. Initialize the indirect args buffer
-		FRDGBufferRef InitializePagesIndirectArgsRDG = CreateAndInitializeDispatchIndirectArgs1D(GraphBuilder, Scene.GetFeatureLevel(), TEXT("Shadow.Virtual.InitializePagesIndirectArgs"));
+		FRDGBufferRef InitializePagesIndirectArgsRDG = CreateAndClearIndirectDispatchArgs1D(GraphBuilder, Scene.GetFeatureLevel(), TEXT("Shadow.Virtual.InitializePagesIndirectArgs"));
 		// 2. Filter the relevant physical pages and set up the indirect args
 		{
 			FSelectPagesToInitializeCS::FParameters* PassParameters = GraphBuilder.AllocParameters<FSelectPagesToInitializeCS::FParameters>();
@@ -3434,7 +3501,7 @@ void FVirtualShadowMapArray::UpdateHZB(FRDGBuilder& GraphBuilder)
 	FRDGBufferRef PhysicalPagesForHZBRDG = GraphBuilder.CreateBuffer(FRDGBufferDesc::CreateStructuredDesc(sizeof(int32), GetTotalAllocatedPhysicalPages() + 1), TEXT("Shadow.Virtual.PhysicalPagesForHZB"));
 
 	// 1. Clear the indirect args buffer (note 2x args)
-	AddClearIndirectDispatchArgs1DPass(GraphBuilder, Scene.GetFeatureLevel(), PagesForHZBIndirectArgsRDG, 2U);
+	AddClearIndirectDispatchArgs1DPass(GraphBuilder, Scene.GetFeatureLevel(), PagesForHZBIndirectArgsRDG, 2u, 4u);
 
 	// 2. Filter the relevant physical pages and set up the indirect args
 	{
diff --git a/Engine/Source/Runtime/Renderer/Private/VirtualShadowMaps/VirtualShadowMapArray.h b/Engine/Source/Runtime/Renderer/Private/VirtualShadowMaps/VirtualShadowMapArray.h
--- a/Engine/Source/Runtime/Renderer/Private/VirtualShadowMaps/VirtualShadowMapArray.h
+++ b/Engine/Source/Runtime/Renderer/Private/VirtualShadowMaps/VirtualShadowMapArray.h
@@ -26,6 +26,11 @@ struct FShaderCompilerEnvironment;
 struct FFrontLayerTranslucencyData;
 class FSceneInstanceCullingQuery;
 
+namespace Froxel
+{
+	class FRenderer;
+}
+
 namespace Nanite
 {
 	struct FPackedView;
@@ -48,6 +53,8 @@ constexpr uint32 CalcVirtualShadowMapLevelOffsets(uint32 Level, uint32 Log2Level
 	return 0x55555555U & Mask;
 }
 
+bool DoesVSMWantFroxels(EShaderPlatform ShaderPlatform);
+
 class FVirtualShadowMap
 {
 public:
@@ -325,7 +332,8 @@ public:
 		const FSortedLightSetSceneInfo& SortedLights, 
 		const TConstArrayView<FVisibleLightInfo>& VisibleLightInfos,
 		const FSingleLayerWaterPrePassResult* SingleLayerWaterPrePassResult,
-		const FFrontLayerTranslucencyData& FrontLayerTranslucencyData);
+		const FFrontLayerTranslucencyData& FrontLayerTranslucencyData,
+		const Froxel::FRenderer& FroxelRenderer);
 
 	bool IsAllocated() const
 	{
@@ -394,11 +402,6 @@ public:
 	 */
 	static float InterpolateResolutionBias(float BiasNonMoving, float BiasMoving, float LightMobilityFactor);
 
-	/**
-	* Helper function to create and clear an indirect args buffer
-	*/
-	static FRDGBufferRef CreateAndInitializeDispatchIndirectArgs1D(FRDGBuilder& GraphBuilder, ERHIFeatureLevel::Type FeatureLevel, const TCHAR* Name);
-
 	// We keep a reference to the cache manager that was used to initialize this frame as it owns some of the buffers
 	FVirtualShadowMapArrayCacheManager* CacheManager = nullptr;
 
