Fixed typo

#rnx

diff --git a/Engine/Source/Runtime/Renderer/Private/Froxel/Froxel.cpp b/Engine/Source/Runtime/Renderer/Private/Froxel/Froxel.cpp
--- a/Engine/Source/Runtime/Renderer/Private/Froxel/Froxel.cpp
+++ b/Engine/Source/Runtime/Renderer/Private/Froxel/Froxel.cpp
@@ -8,7 +8,7 @@
 namespace Froxel
 {
 
-FSharedParameters MakeSharedParameters(const FViewInfo &View)
+FSharedParameters MakeSharedParameters(const FViewInfo& View)
 {
 	const float FroxelTileSize = float(FRenderer::TileSize);
 
@@ -17,7 +17,7 @@ FSharedParameters MakeSharedParameters(const FViewInfo &View)
 	// How cubeish the froxels should be
 	const float DepthStretchFactor = 1.0f;
 
-	const FMatrix &ViewToClip = View.ViewMatrices.GetProjectionMatrix();
+	const FMatrix& ViewToClip = View.ViewMatrices.GetProjectionMatrix();
 	// 2x for, er, probably because clipspace is +-1
 	FVector2f AbsClipTileSize = FVector2f(FroxelTileSize * 2.0f) / ViewSize;
 	FVector2f ProjScaleXY = FVector2f(static_cast<float>(ViewToClip.M[0][0]), static_cast<float>(ViewToClip.M[1][1]));
@@ -60,7 +60,7 @@ FSharedParameters MakeSharedParameters(const FViewInfo &View)
 	const FVector2f ClipTileSize = FVector2f(AbsClipTileSize.X, -AbsClipTileSize.Y);
 	const FVector2f ClipSpaceMin = FVector2f(-1.0f, 1.0f);
 
-	const FMatrix &ClipToView = View.ViewMatrices.GetInvProjectionMatrix();
+	const FMatrix& ClipToView = View.ViewMatrices.GetInvProjectionMatrix();
 	const FVector2f ClipToViewScale = FVector2f(ClipToView.M[0][0], ClipToView.M[1][1]);
 	FVector2f FroxelToViewScale = ClipTileSize * ClipToViewScale;
 	FVector2f FroxelToViewBias = ClipSpaceMin * ClipToViewScale;
@@ -79,7 +79,7 @@ FRenderer::FRenderer(bool bIsEnabled, FRDGBuilder& GraphBuilder, const TArray<FV
 		return;
 	}
 
-	for (const FViewInfo &View : InViews)
+	for (const FViewInfo& View : InViews)
 	{
 		FSharedParameters FroxelParameters = MakeSharedParameters(View);
 
@@ -87,7 +87,7 @@ FRenderer::FRenderer(bool bIsEnabled, FRDGBuilder& GraphBuilder, const TArray<FV
 		int32 MaxNumFroxels = View.ViewRect.Area();
 		FRDGBuffer* FroxelsRDG = GraphBuilder.CreateBuffer(FRDGBufferDesc::CreateStructuredDesc(sizeof(FPackedFroxel), MaxNumFroxels), TEXT("r.Froxels"));
 		// TODO: make shared buffer & initialize once
-		FRDGBuffer* FroxelArgsRDG = FroxelArgsRDG = GraphBuilder.CreateBuffer(FRDGBufferDesc::CreateIndirectDesc(4), TEXT("r.FroxelArgs"));
+		FRDGBuffer* FroxelArgsRDG = GraphBuilder.CreateBuffer(FRDGBufferDesc::CreateIndirectDesc(4), TEXT("r.FroxelArgs"));
 		AddClearIndirectDispatchArgs1DPass(GraphBuilder, View.GetFeatureLevel(), FroxelArgsRDG, 1u, ArgsStride);
 
 		Views.Emplace(FViewData{this, FroxelsRDG, FroxelArgsRDG, FroxelParameters});
