Add initial support to translucent holdout.

We added a dedicated path for the alpha channel when alpha propagation is enabled (r.PostProcessing.PropagateAlpha != 0). The rendering logic is as below:
1. Create a new mesh processing batch for all translucent materials.
2. In the translucent path.
   2.1. Get a copy of the alpha before all translucent passes to retain the holdout for opaque objects.
   2.2. After all existing translucent passes are done. Run the dedicated translucent holdout pass to accumualte background visibility and path throughput from far to close.
   2.3. Compose the translucent pass alpha with the retained alpha before all translucent passes back to scene color.

Notes:
1. MODULATE cannot be correctly implemented so we retain the old behavior. Comments about an alternative is in the shader code. However, we can only have one modulate material in the scene to get a close match to path tracing.
2. Substrate has minimal support. Only material types in the demo level is supported.
3. Set r.RenderTranslucentHoldout=0 to fallback to the old behavior.
4. Fog is currently not supported on translucent holdout.
5. Fully backward compatible for alphaholdout blending mode.

#tests Test on custom level with all translucent material pair with/without holdout. It matches close to path tracing except for fog.

diff --git a/Engine/Shaders/Private/BasePassPixelShader.usf b/Engine/Shaders/Private/BasePassPixelShader.usf
--- a/Engine/Shaders/Private/BasePassPixelShader.usf
+++ b/Engine/Shaders/Private/BasePassPixelShader.usf
@@ -130,6 +130,10 @@
 // Reroute for LumenRadianceCacheInterpolation.ush
 #define RadianceCacheInterpolation TranslucentBasePass
 #include "Lumen/LumenTranslucencyVolumeShared.ush"
+
+#define IsTranslucentHoldoutPass (TranslucentBasePass.TranslucencyPass == 1)
+#else
+#define IsTranslucentHoldoutPass false
 #endif
 
 #if TRANSLUCENT_SELF_SHADOWING
@@ -2177,8 +2181,9 @@ void FPixelShaderInOut_MainPS(
 #endif
 
 #if POST_PROCESS_ALPHA != 0 && MATERIAL_DOMAIN_SURFACE
+	const bool bIsHoldout = GetPrimitiveData(MaterialParameters).Flags & PRIMITIVE_SCENE_DATA_FLAG_HOLDOUT;
 	// Only include the holdout logic if the Renderer is forwarding its alpha channel to post-processing
-	if (GetPrimitiveData(MaterialParameters).Flags & PRIMITIVE_SCENE_DATA_FLAG_HOLDOUT)
+	if (bIsHoldout)
 	{
 		// Ignore primitive holdout during reflection captures
 		if(ResolvedView.RenderingReflectionCaptureMask == 0)
@@ -2188,8 +2193,23 @@ void FPixelShaderInOut_MainPS(
 #if MATERIALBLENDING_SOLID || MATERIALBLENDING_MASKED
 			// This catches the common case as well as handles cases where Opacity has a different meaning, such as for Subsurface materials
 			Out.MRT[0] = half4(0.0, 0.0, 0.0, 1.0);
+#elif MATERIALBLENDING_MODULATE
+			Out.MRT[0] = half4(0.0, 0.0, 0.0, Opacity);
+			// Alternative method. Keep MRT[0] as it is to get a closer look to path tracing:
+			//	Out.MRT[0] = half4(FoggedColor, Opacity);
+			// MATERIALBLENDING_MODULATE comments in the holdout pass.
+#elif SUBSTRATE_TRANSLUCENT_MATERIAL
+	#if SUBSTRATE_BLENDING_ALPHAHOLDOUT
+			//This is not the translucent holdout pass, do not reset Out.MRT[0] so holdout can accumulate color contribution.
+	#elif SUBSTRATE_BLENDING_TRANSLUCENT_COLOREDTRANSMITTANCE
+			Out.MRT[0] = half4(0.0, 0.0, 0.0, DualBlendSurfaceCoverage);
+	#elif SUBSTRATE_BLENDING_TRANSLUCENT_GREYTRANSMITTANCE
+			Out.MRT[0] = half4(0.0, 0.0, 0.0, DualBlendSurfaceCoverage);
+	#endif
+			
 #else
-			// TODO: The translucent case doesn't seem to get the right alpha (compared to the path traced reference)
+			// TODO: Not all translucent cases get the right alpha (compared to the path traced reference)
+
 			Out.MRT[0] = half4(0.0, 0.0, 0.0, Opacity);
 #endif
 
@@ -2200,7 +2220,73 @@ void FPixelShaderInOut_MainPS(
 #endif
 		}
 	}
+
 #endif
+
+#if (MATERIALBLENDING_ANY_TRANSLUCENT || MATERIALBLENDING_ALPHACOMPOSITE || MATERIALBLENDING_ALPHAHOLDOUT) && MATERIAL_DOMAIN_SURFACE && POST_PROCESS_ALPHA!= 0
+	// If this is translucency holdout pass, calculate the background visibility and path throughput.
+	if (IsTranslucentHoldoutPass &&
+        ResolvedView.RenderingReflectionCaptureMask == 0)
+	{
+		// Default behavior is to preserve the destination alpha.
+		half BackgroundVisibilityAdd = 0.0f;
+		half PathThroughputMul = 1.0f;
+
+#if SUBSTRATE_TRANSLUCENT_MATERIAL
+	#if SUBSTRATE_BLENDING_TRANSLUCENT_GREYTRANSMITTANCE
+		BackgroundVisibilityAdd = bIsHoldout ? DualBlendSurfaceCoverage : 0;
+		PathThroughputMul = 1 - DualBlendSurfaceCoverage;
+	#elif SUBSTRATE_BLENDING_TRANSLUCENT_COLOREDTRANSMITTANCE
+		BackgroundVisibilityAdd = bIsHoldout ? DualBlendSurfaceCoverage : 0;
+		PathThroughputMul = 1 - DualBlendSurfaceCoverage;
+	#elif SUBSTRATE_BLENDING_COLOREDTRANSMITTANCEONLY
+		// TODO
+	#elif MATERIALBLENDING_ALPHAHOLDOUT
+		//alpha holdout should always contribute to the background visibility.
+		BackgroundVisibilityAdd = 1 - AdjustedAlpha;
+		PathThroughputMul = AdjustedAlpha;
+	#endif 
+#else
+
+	#if MATERIAL_SHADINGMODEL_THIN_TRANSLUCENT
+		//thin translucent shading model uses surface coverage as opacity.
+		BackgroundVisibilityAdd = bIsHoldout ? DualBlendSurfaceCoverage : 0;
+		PathThroughputMul = 1 - DualBlendSurfaceCoverage;
+	#elif MATERIALBLENDING_ALPHAHOLDOUT
+		//alpha holdout should always contribute to the background visibility.
+		BackgroundVisibilityAdd = Opacity;
+		PathThroughputMul = (1 - Opacity);
+	#elif MATERIALBLENDING_TRANSLUCENT
+		// Use Opacity as it is.
+		BackgroundVisibilityAdd = bIsHoldout ? Opacity : 0;
+		PathThroughputMul = 1 - Opacity;
+	#endif
+#endif
+
+#if MATERIALBLENDING_ADDITIVE
+		// Not affecting alpha
+		BackgroundVisibilityAdd = 0.0f;
+		PathThroughputMul = 1.0f;
+#elif MATERIALBLENDING_MODULATE
+
+		BackgroundVisibilityAdd = bIsHoldout ? Opacity : 0;
+		PathThroughputMul = (1 - Opacity);
+		//---------------------------------------------------------------------
+		// Another solution: use fogged color to calcualte throughput and background visibility
+		// C = C * L + C * (1-L).
+		// BackgroundVisibilityAdd = bIsHoldout ? (1-Luminance(FoggedColor)) : 0;
+		// PathThroughputMul = Luminance(FoggedColor);
+		// Issue: Need to fix modulate layer stacking does not match patch tracing.
+		// Otherwise, we can only have one modulate in a scene that can be holdout as well.
+#endif
+
+		// Render far to close logic
+		Out.MRT[0].xyz = BackgroundVisibilityAdd;
+		Out.MRT[0].w = PathThroughputMul;
+		
+	}
+#endif
+
 }
 
 #if SUBSTRATE_OPAQUE_DEFERRED
diff --git a/Engine/Shaders/Private/ComposeSeparateTranslucency.usf b/Engine/Shaders/Private/ComposeSeparateTranslucency.usf
--- a/Engine/Shaders/Private/ComposeSeparateTranslucency.usf
+++ b/Engine/Shaders/Private/ComposeSeparateTranslucency.usf
@@ -39,6 +39,7 @@ float2 SeparateTranslucencyUVMin;
 float2 SeparateTranslucencyUVMax;
 float2 SeparateTranslucencyExtentInverse;
 uint bLensDistortion;
+uint bPassthroughAlpha;
 
 //------------------------------------------------------- ENTRY POINT
 
@@ -95,7 +96,43 @@ void MainPS(
 
 	// Final composition
 	OutColor.rgb = SceneColorSample.rgb * SeparateTranslucencySample.a * SeparateModulationSample.rgb + SeparateTranslucencySample.rgb;
+	
+	BRANCH
+	if (bPassthroughAlpha)
+	{
+		OutColor.a = SceneColorSample.a;
+	}
+	else
+	{
+		float GreyScaleModulateColorBackgroundVisibility = dot(SeparateModulationSample.rgb, float3(1.0f / 3.0f, 1.0f / 3.0f, 1.0f / 3.0f));
+		OutColor.a = SceneColorSample.a * SeparateTranslucencySample.a * GreyScaleModulateColorBackgroundVisibility;
+	}
+}
 
-	float GreyScaleModulateColorBackgroundVisibility = dot(SeparateModulationSample.rgb, float3(1.0f / 3.0f, 1.0f / 3.0f, 1.0f / 3.0f));
-	OutColor.a = SceneColorSample.a * SeparateTranslucencySample.a * GreyScaleModulateColorBackgroundVisibility;
+//------------------------------------------------------- PARAMETERS
+#define VISIBILITY_COPY_FROM_SCENECOLOR 0
+#define VISIBILITY_COPY_TO_SCENECOLOR	1
+
+#ifndef VISIBILITY_COPY_TYPE
+#define VISIBILITY_COPY_TYPE VISIBILITY_COPY_TO_SCENECOLOR
+#endif
+
+Texture2D    TranslucentHoldoutPointTexture;
+SamplerState TranslucentHoldoutPointSampler;
+
+//------------------------------------------------------- ENTRY POINT
+void CopyBackgroundVisibilityPS(
+	float4 SvPosition : SV_POSITION,
+	out float4 OutColor : SV_Target0)
+{
+	float2 PixelPos = SvPosition.xy;
+	
+#if VISIBILITY_COPY_TYPE == VISIBILITY_COPY_FROM_SCENECOLOR
+	float4 SceneColorWithBackgroundVisibilityInAlpha = TranslucentHoldoutPointTexture.Load(uint3(PixelPos, 0));
+	OutColor.rgb = SceneColorWithBackgroundVisibilityInAlpha.a;
+	OutColor.a = 1.0f;
+#elif VISIBILITY_COPY_TYPE == VISIBILITY_COPY_TO_SCENECOLOR
+	float4 BackgroundVisibilityWithAlpha = TranslucentHoldoutPointTexture.Load(uint3(PixelPos, 0));
+	OutColor.a =  saturate(BackgroundVisibilityWithAlpha.r);// We can copy any rgb channel.
+#endif 
 }
diff --git a/Engine/Shaders/Private/TemporalSuperResolution/TSRShadingAnalysis.ush b/Engine/Shaders/Private/TemporalSuperResolution/TSRShadingAnalysis.ush
--- a/Engine/Shaders/Private/TemporalSuperResolution/TSRShadingAnalysis.ush
+++ b/Engine/Shaders/Private/TemporalSuperResolution/TSRShadingAnalysis.ush
@@ -20,7 +20,7 @@ static const tsr_half FilteringWeight = rcp(tsr_half(1.0 + 4 * 0.5 + 4 * 0.25));
 //------------------------------------------------------- PARAMETERS
 
 float FlickeringFramePeriod;
-
+uint bPassthroughAlpha;
 
 //------------------------------------------------------- COMPOSE SEPARATE TRANSLUCENCY
 
@@ -31,7 +31,7 @@ tsr_tensor_halfC ComposeTranslucency(tsr_tensor_halfC Color, tsr_tensor_half4 Tr
 	tsr_tensor_halfC ComposedColor;
 	ComposedColor = Color * tsr_tensor_halfC::Vectorize(Translucency[3]) + ResizeChannels<CONFIG_CHANNEL_COUNT>(Translucency);
 	#if CONFIG_SCENE_COLOR_ALPHA
-		ComposedColor.SetComponent(3, Color[3] * Translucency[3]);
+		ComposedColor.SetComponent(3, select(tsr_tensor_bool::Const(bPassthroughAlpha > 0), Color[3], Color[3] * Translucency[3]));
 	#endif
 	#if CONFIG_SCENE_COLOR_ALPHA
 		ComposedColor = min(ComposedColor, tsr_tensor_halfC::Const(LargestSceneColorRGBA));
diff --git a/Engine/Shaders/Private/TranslucentLightingShaders.usf b/Engine/Shaders/Private/TranslucentLightingShaders.usf
--- a/Engine/Shaders/Private/TranslucentLightingShaders.usf
+++ b/Engine/Shaders/Private/TranslucentLightingShaders.usf
@@ -58,7 +58,7 @@ Texture2D SceneColorTexture;
 void CopySceneColorMain(FScreenVertexOutput Input, out float4 OutColor : SV_Target0)
 {
 	float4 LinearColor = Texture2DSample(SceneColorTexture, GlobalPointClampedSampler, Input.UV);
-	OutColor = float4(EncodeSceneColorForMaterialNode(LinearColor.rgb), 0);
+	OutColor = float4(EncodeSceneColorForMaterialNode(LinearColor.rgb), LinearColor.a);
 }
 
 /** Geometry shader that writes to a range of slices of a volume texture. */
diff --git a/Engine/Source/Runtime/Renderer/Internal/TranslucentPassResource.h b/Engine/Source/Runtime/Renderer/Internal/TranslucentPassResource.h
--- a/Engine/Source/Runtime/Renderer/Internal/TranslucentPassResource.h
+++ b/Engine/Source/Runtime/Renderer/Internal/TranslucentPassResource.h
@@ -16,7 +16,7 @@ namespace ETranslucencyPass
 		TPT_TranslucencyAfterDOF,
 		TPT_TranslucencyAfterDOFModulate,
 		TPT_TranslucencyAfterMotionBlur,
-
+		TPT_TranslucencyHoldout,
 		/** Drawing all translucency, regardless of separate or standard.  Used when drawing translucency outside of the main renderer, eg FRendererModule::DrawTile. */
 		TPT_AllTranslucency,
 		TPT_MAX
diff --git a/Engine/Source/Runtime/Renderer/Private/BasePassRendering.cpp b/Engine/Source/Runtime/Renderer/Private/BasePassRendering.cpp
--- a/Engine/Source/Runtime/Renderer/Private/BasePassRendering.cpp
+++ b/Engine/Source/Runtime/Renderer/Private/BasePassRendering.cpp
@@ -271,6 +271,18 @@ void SetTranslucentRenderState(FMeshPassProcessorRenderState& DrawRenderState, c
 				DrawRenderState.SetBlendState(TStaticBlendState<CW_RGBA, BO_Add, BF_One, BF_InverseSourceAlpha, BO_Add, BF_Zero, BF_InverseSourceAlpha>::GetRHI());
 			}
 		}
+
+		if (InTranslucencyPassType == ETranslucencyPass::TPT_TranslucencyHoldout)
+		{
+			// We render translucent primitives from far to close. For each translucent primitive, 
+			// we output the background visibility (bv) contribution in rgb channel, and the path throughput (pt) contribution in the alpha channel
+			// E.g., alpha holdout blending has:
+			// Src0 = (Opacity, 1 - Opacity)
+			// The blend state allows the accumulation of
+			// bv = Src0.rgb + (1 - Opacity) * bv
+			// pt = pt * (1 - Opacity)
+			DrawRenderState.SetBlendState(TStaticBlendState<CW_RGBA, BO_Add, BF_One, BF_SourceAlpha, BO_Add, BF_Zero, BF_SourceAlpha>::GetRHI());
+		}
 	}
 	else if (Material.GetShadingModels().HasShadingModel(MSM_ThinTranslucent))
 	{
@@ -310,6 +322,11 @@ void SetTranslucentRenderState(FMeshPassProcessorRenderState& DrawRenderState, c
 			// If unsupported, we still use premultipled alpha but the shader will use the variation converting color transmittance to a grey scale transmittance.
 			DrawRenderState.SetBlendState(TStaticBlendState<CW_RGBA, BO_Add, BF_One, BF_InverseSourceAlpha, BO_Add, BF_Zero, BF_InverseSourceAlpha>::GetRHI());
 		}
+
+		if (InTranslucencyPassType == ETranslucencyPass::TPT_TranslucencyHoldout)
+		{
+			DrawRenderState.SetBlendState(TStaticBlendState<CW_RGBA, BO_Add, BF_One, BF_SourceAlpha, BO_Add, BF_Zero, BF_SourceAlpha>::GetRHI());
+		}
 	}
 	else
 	{
@@ -324,9 +341,16 @@ void SetTranslucentRenderState(FMeshPassProcessorRenderState& DrawRenderState, c
 			break;
 		case BLEND_Translucent:
 		case BLEND_TranslucentColoredTransmittance:	// When Substrate is disabled, this falls back to simple Translucency.
-			// Note: alpha channel used by separate translucency, storing how much of the background should be added when doing the final composite
-			// The Alpha channel is also used by non-separate translucency when rendering to scene captures, which store the final opacity
-			DrawRenderState.SetBlendState(TStaticBlendState<CW_RGBA, BO_Add, BF_SourceAlpha, BF_InverseSourceAlpha, BO_Add, BF_Zero, BF_InverseSourceAlpha>::GetRHI());
+			if (InTranslucencyPassType != ETranslucencyPass::TPT_TranslucencyHoldout)
+			{
+				// Note: alpha channel used by separate translucency, storing how much of the background should be added when doing the final composite
+				// The Alpha channel is also used by non-separate translucency when rendering to scene captures, which store the final opacity
+				DrawRenderState.SetBlendState(TStaticBlendState<CW_RGBA, BO_Add, BF_SourceAlpha, BF_InverseSourceAlpha, BO_Add, BF_Zero, BF_InverseSourceAlpha>::GetRHI());
+			}
+			else
+			{
+				DrawRenderState.SetBlendState(TStaticBlendState<CW_RGBA, BO_Add, BF_One, BF_SourceAlpha, BO_Add, BF_Zero, BF_SourceAlpha>::GetRHI());
+			}
 			break;
 		case BLEND_Additive:
 			// Add to the existing scene color
@@ -335,16 +359,37 @@ void SetTranslucentRenderState(FMeshPassProcessorRenderState& DrawRenderState, c
 			DrawRenderState.SetBlendState(TStaticBlendState<CW_RGBA, BO_Add, BF_One, BF_One, BO_Add, BF_Zero, BF_InverseSourceAlpha>::GetRHI());
 			break;
 		case BLEND_Modulate:
-			// Modulate with the existing scene color, preserve destination alpha.
-			DrawRenderState.SetBlendState(TStaticBlendState<CW_RGB, BO_Add, BF_DestColor, BF_Zero>::GetRHI());
+			if (InTranslucencyPassType != ETranslucencyPass::TPT_TranslucencyHoldout)
+			{
+				// Modulate with the existing scene color, preserve destination alpha.
+				DrawRenderState.SetBlendState(TStaticBlendState<CW_RGB, BO_Add, BF_DestColor, BF_Zero>::GetRHI());
+			}
+			else
+			{
+				DrawRenderState.SetBlendState(TStaticBlendState<CW_RGBA, BO_Add, BF_One, BF_SourceAlpha, BO_Add, BF_Zero, BF_SourceAlpha>::GetRHI());
+			}
 			break;
 		case BLEND_AlphaComposite:
-			// Blend with existing scene color. New color is already pre-multiplied by alpha.
-			DrawRenderState.SetBlendState(TStaticBlendState<CW_RGBA, BO_Add, BF_One, BF_InverseSourceAlpha, BO_Add, BF_Zero, BF_InverseSourceAlpha>::GetRHI());
+			if (InTranslucencyPassType != ETranslucencyPass::TPT_TranslucencyHoldout)
+			{
+				// Blend with existing scene color. New color is already pre-multiplied by alpha.
+				DrawRenderState.SetBlendState(TStaticBlendState<CW_RGBA, BO_Add, BF_One, BF_InverseSourceAlpha, BO_Add, BF_Zero, BF_InverseSourceAlpha>::GetRHI());
+			}
+			else
+			{
+				DrawRenderState.SetBlendState(TStaticBlendState<CW_RGBA, BO_Add, BF_One, BF_SourceAlpha, BO_Add, BF_Zero, BF_SourceAlpha>::GetRHI());
+			}
 			break;
 		case BLEND_AlphaHoldout:
-			// Blend by holding out the matte shape of the source alpha
-			DrawRenderState.SetBlendState(TStaticBlendState<CW_RGBA, BO_Add, BF_Zero, BF_InverseSourceAlpha, BO_Add, BF_One, BF_InverseSourceAlpha>::GetRHI());
+			if (InTranslucencyPassType != ETranslucencyPass::TPT_TranslucencyHoldout)
+			{
+				// Blend by holding out the matte shape of the source alpha
+				DrawRenderState.SetBlendState(TStaticBlendState<CW_RGBA, BO_Add, BF_Zero, BF_InverseSourceAlpha, BO_Add, BF_One, BF_InverseSourceAlpha>::GetRHI());
+			}
+			else
+			{
+				DrawRenderState.SetBlendState(TStaticBlendState<CW_RGBA, BO_Add, BF_One, BF_SourceAlpha, BO_Add, BF_Zero, BF_SourceAlpha>::GetRHI());
+			}
 			break;
 		};
 	}
@@ -1955,6 +2000,10 @@ bool FBasePassMeshProcessor::ShouldDraw(const FMaterial& Material)
 				bShouldDraw = Material.IsTranslucencyAfterMotionBlurEnabled();
 				break;
 
+			case ETranslucencyPass::TPT_TranslucencyHoldout:
+				bShouldDraw = true;
+                break;
+
 			case ETranslucencyPass::TPT_AllTranslucency:
 				bShouldDraw = true;
 				break;
@@ -1979,10 +2028,11 @@ bool FBasePassMeshProcessor::TryAddMeshBatch(const FMeshBatch& RESTRICT MeshBatc
 	const FMeshDrawingPolicyOverrideSettings OverrideSettings = ComputeMeshOverrideSettings(MeshBatch);
 	const ERasterizerFillMode MeshFillMode = ComputeMeshFillMode(Material, OverrideSettings);
 	const ERasterizerCullMode MeshCullMode = ComputeMeshCullMode(Material, OverrideSettings);
+	const bool bIsValidTranslucentNonDeferredDecal = PrimitiveSceneProxy && bIsTranslucent && !Material.IsDeferredDecal();
 
 
 	bool bShouldDraw = false;
-	if (AutoBeforeDOFTranslucencyBoundary > 0.0f && PrimitiveSceneProxy && bIsTranslucent && !Material.IsDeferredDecal())
+	if (AutoBeforeDOFTranslucencyBoundary > 0.0f && bIsValidTranslucentNonDeferredDecal)
 	{
 		check(ViewIfDynamicMeshCommand);
 		check(TranslucencyPassType != ETranslucencyPass::TPT_MAX);
@@ -2021,6 +2071,10 @@ bool FBasePassMeshProcessor::TryAddMeshBatch(const FMeshBatch& RESTRICT MeshBatc
 		{
 			bShouldDraw = bIsAfterDOFModulate && !bIsInDOFBackground;
 		}
+		else if (TranslucencyPassType == ETranslucencyPass::TPT_TranslucencyHoldout)
+		{
+			bShouldDraw = true;
+		}
 		else
 		{
 			unimplemented();
@@ -2029,6 +2083,12 @@ bool FBasePassMeshProcessor::TryAddMeshBatch(const FMeshBatch& RESTRICT MeshBatc
 	else
 	{
 		bShouldDraw = ShouldDraw(Material);
+
+		// Only translucent materials should be drawn in the holdout pass.
+		if (TranslucencyPassType == ETranslucencyPass::TPT_TranslucencyHoldout)
+		{
+			bShouldDraw = bIsValidTranslucentNonDeferredDecal;
+		}
 	}
 
 	// Only draw opaque materials.
@@ -2517,6 +2577,16 @@ FMeshPassProcessor* CreateTranslucencyAfterMotionBlurProcessor(ERHIFeatureLevel:
 	return new FBasePassMeshProcessor(EMeshPass::TranslucencyAfterMotionBlur, Scene, FeatureLevel, InViewIfDynamicMeshCommand, PassDrawRenderState, InDrawListContext, Flags, ETranslucencyPass::TPT_TranslucencyAfterMotionBlur);
 }
 
+FMeshPassProcessor* CreateTranslucencyHoldoutPassProcessor(ERHIFeatureLevel::Type FeatureLevel, const FScene* Scene, const FSceneView* InViewIfDynamicMeshCommand, FMeshPassDrawListContext* InDrawListContext)
+{
+	FMeshPassProcessorRenderState PassDrawRenderState;
+	PassDrawRenderState.SetDepthStencilState(TStaticDepthStencilState<false, CF_DepthNearOrEqual>::GetRHI());
+
+	const FBasePassMeshProcessor::EFlags Flags = FBasePassMeshProcessor::EFlags::CanUseDepthStencil;
+
+	return new FBasePassMeshProcessor(EMeshPass::TranslucencyHoldout, Scene, FeatureLevel, InViewIfDynamicMeshCommand, PassDrawRenderState, InDrawListContext, Flags, ETranslucencyPass::TPT_TranslucencyHoldout);
+}
+
 FMeshPassProcessor* CreateTranslucencyAllPassProcessor(ERHIFeatureLevel::Type FeatureLevel, const FScene* Scene, const FSceneView* InViewIfDynamicMeshCommand, FMeshPassDrawListContext* InDrawListContext)
 {
 	FMeshPassProcessorRenderState PassDrawRenderState;
@@ -2533,4 +2603,5 @@ REGISTER_MESHPASSPROCESSOR_AND_PSOCOLLECTOR(TranslucencyStandardModulatePass, Cr
 REGISTER_MESHPASSPROCESSOR_AND_PSOCOLLECTOR(TranslucencyAfterDOFPass, CreateTranslucencyAfterDOFProcessor, EShadingPath::Deferred, EMeshPass::TranslucencyAfterDOF, EMeshPassFlags::MainView);
 REGISTER_MESHPASSPROCESSOR_AND_PSOCOLLECTOR(TranslucencyAfterDOFModulatePass, CreateTranslucencyAfterDOFModulateProcessor, EShadingPath::Deferred, EMeshPass::TranslucencyAfterDOFModulate, EMeshPassFlags::MainView);
 REGISTER_MESHPASSPROCESSOR_AND_PSOCOLLECTOR(TranslucencyAfterMotionBlurPass, CreateTranslucencyAfterMotionBlurProcessor, EShadingPath::Deferred, EMeshPass::TranslucencyAfterMotionBlur, EMeshPassFlags::MainView);
+REGISTER_MESHPASSPROCESSOR_AND_PSOCOLLECTOR(TranslucencyHoldoutPass, CreateTranslucencyHoldoutPassProcessor, EShadingPath::Deferred, EMeshPass::TranslucencyHoldout, EMeshPassFlags::MainView);
 REGISTER_MESHPASSPROCESSOR_AND_PSOCOLLECTOR(TranslucencyAllPass, CreateTranslucencyAllPassProcessor, EShadingPath::Deferred, EMeshPass::TranslucencyAll, EMeshPassFlags::MainView);
diff --git a/Engine/Source/Runtime/Renderer/Private/BasePassRendering.h b/Engine/Source/Runtime/Renderer/Private/BasePassRendering.h
--- a/Engine/Source/Runtime/Renderer/Private/BasePassRendering.h
+++ b/Engine/Source/Runtime/Renderer/Private/BasePassRendering.h
@@ -139,6 +139,7 @@ BEGIN_GLOBAL_SHADER_PARAMETER_STRUCT(FTranslucentBasePassUniformParameters,)
 	SHADER_PARAMETER_SAMPLER(SamplerState, SceneColorCopySampler)
 	SHADER_PARAMETER_STRUCT(FBlueNoiseParameters, BlueNoise)
 	SHADER_PARAMETER_STRUCT(FAdaptiveVolumetricShadowMapUniformBufferParameters, AVSM)
+	SHADER_PARAMETER(int, TranslucencyPass)
 END_GLOBAL_SHADER_PARAMETER_STRUCT()
 
 DECLARE_GPU_DRAWCALL_STAT_EXTERN(Basepass);
@@ -167,7 +168,8 @@ extern TRDGUniformBufferRef<FTranslucentBasePassUniformParameters> CreateTranslu
 	const FTranslucencyLightingVolumeTextures& TranslucencyLightingVolumeTextures = {},
 	FRDGTextureRef SceneColorCopyTexture = nullptr,
 	const ESceneTextureSetupMode SceneTextureSetupMode = ESceneTextureSetupMode::None,
-	bool bLumenGIEnabled = false);
+	bool bLumenGIEnabled = false,
+	ETranslucencyPass::Type TranslucencyPass = ETranslucencyPass::TPT_MAX);
 
 extern bool IsGBufferLayoutSupportedForMaterial(EGBufferLayout Layout, const FMeshMaterialShaderPermutationParameters& Params);
 extern void ModifyBasePassCSPSCompilationEnvironment(const FMeshMaterialShaderPermutationParameters& Params, EGBufferLayout GBufferLayout, bool bEnableSkyLight, FShaderCompilerEnvironment& OutEnvironment);
diff --git a/Engine/Source/Runtime/Renderer/Private/MeshDrawCommands.cpp b/Engine/Source/Runtime/Renderer/Private/MeshDrawCommands.cpp
--- a/Engine/Source/Runtime/Renderer/Private/MeshDrawCommands.cpp
+++ b/Engine/Source/Runtime/Renderer/Private/MeshDrawCommands.cpp
@@ -1413,6 +1413,7 @@ void FParallelMeshDrawCommandPass::DispatchPassSetup(
 		case EMeshPass::TranslucencyAfterDOF: TaskContext.TranslucencyPass			= ETranslucencyPass::TPT_TranslucencyAfterDOF; break;
 		case EMeshPass::TranslucencyAfterDOFModulate: TaskContext.TranslucencyPass	= ETranslucencyPass::TPT_TranslucencyAfterDOFModulate; break;
 		case EMeshPass::TranslucencyAfterMotionBlur: TaskContext.TranslucencyPass	= ETranslucencyPass::TPT_TranslucencyAfterMotionBlur; break;
+		case EMeshPass::TranslucencyHoldout: TaskContext.TranslucencyPass			= ETranslucencyPass::TPT_TranslucencyHoldout; break;
 		case EMeshPass::TranslucencyAll: TaskContext.TranslucencyPass				= ETranslucencyPass::TPT_AllTranslucency; break;
 	}
 
diff --git a/Engine/Source/Runtime/Renderer/Private/MobileBasePass.cpp b/Engine/Source/Runtime/Renderer/Private/MobileBasePass.cpp
--- a/Engine/Source/Runtime/Renderer/Private/MobileBasePass.cpp
+++ b/Engine/Source/Runtime/Renderer/Private/MobileBasePass.cpp
@@ -1116,3 +1116,4 @@ REGISTER_MESHPASSPROCESSOR_AND_PSOCOLLECTOR(MobileTranslucencyAllPass,		CreateMo
 REGISTER_MESHPASSPROCESSOR_AND_PSOCOLLECTOR(MobileTranslucencyStandardPass,	CreateMobileTranslucencyStandardPassProcessor,	EShadingPath::Mobile, EMeshPass::TranslucencyStandard, 	EMeshPassFlags::MainView);
 REGISTER_MESHPASSPROCESSOR_AND_PSOCOLLECTOR(MobileTranslucencyAfterDOFPass,	CreateMobileTranslucencyAfterDOFProcessor,	EShadingPath::Mobile, EMeshPass::TranslucencyAfterDOF, 	EMeshPassFlags::MainView);
 // Skipping EMeshPass::TranslucencyAfterDOFModulate because dual blending is not supported on mobile
+// Skipping EMeshPass::TranslucencyHoldout, it is not supported on mobile.
diff --git a/Engine/Source/Runtime/Renderer/Private/PostProcess/TemporalSuperResolution.cpp b/Engine/Source/Runtime/Renderer/Private/PostProcess/TemporalSuperResolution.cpp
--- a/Engine/Source/Runtime/Renderer/Private/PostProcess/TemporalSuperResolution.cpp
+++ b/Engine/Source/Runtime/Renderer/Private/PostProcess/TemporalSuperResolution.cpp
@@ -802,6 +802,7 @@ class FTSRRejectShadingCS : public FTSRConvolutionNetworkShader
 		SHADER_PARAMETER(int32, TileOverscan)
 		SHADER_PARAMETER(int32, bEnableResurrection)
 		SHADER_PARAMETER(int32, bEnableFlickeringHeuristic)
+		SHADER_PARAMETER(int32, bPassthroughAlpha)
 
 		SHADER_PARAMETER_RDG_TEXTURE(Texture2D, InputTexture)
 		SHADER_PARAMETER_RDG_TEXTURE(Texture2D, InputMoireLumaTexture)
@@ -2258,6 +2259,7 @@ FDefaultTemporalUpscaler::FOutputs AddTemporalSuperResolutionPasses(
 		PassParameters->TileOverscan = TileOverscan;
 		PassParameters->bEnableResurrection = bCanResurrectHistory;
 		PassParameters->bEnableFlickeringHeuristic = FlickeringFramePeriod > 0.0f;
+		PassParameters->bPassthroughAlpha = IsTranslucentHoldoutEnabled(GetFeatureLevelShadingPath(View.GetFeatureLevel()));
 
 		PassParameters->InputTexture = PassInputs.SceneColor.Texture;
 		if (PassInputs.FlickeringInputTexture.IsValid())
diff --git a/Engine/Source/Runtime/Renderer/Private/SceneVisibility.cpp b/Engine/Source/Runtime/Renderer/Private/SceneVisibility.cpp
--- a/Engine/Source/Runtime/Renderer/Private/SceneVisibility.cpp
+++ b/Engine/Source/Runtime/Renderer/Private/SceneVisibility.cpp
@@ -1361,6 +1361,7 @@ void FRelevancePacket::ComputeRelevance(FDynamicPrimitiveIndexList& DynamicPrimi
 	const bool bMobileMaskedInEarlyPass = (ShadingPath == EShadingPath::Mobile) && Scene.EarlyZPassMode == DDM_MaskedOnly;
 	const bool bMobileBasePassAlwaysUsesCSM = (ShadingPath == EShadingPath::Mobile) && MobileBasePassAlwaysUsesCSM(Scene.GetShaderPlatform());
 	const bool bVelocityPassWritesDepth = Scene.EarlyZPassMode == DDM_AllOpaqueNoVelocity;
+	const bool bIsTranslucentHoldoutEnabled = IsTranslucentHoldoutEnabled(ShadingPath);
 	const bool bHLODActive = Scene.SceneLODHierarchy.IsActive();
 	const FHLODVisibilityState* const HLODState = bHLODActive && ViewState ? &ViewState->HLODVisibilityState : nullptr;
 	float MaxDrawDistanceScale = GetCachedScalabilityCVars().ViewDistanceScale;
@@ -1722,6 +1723,11 @@ void FRelevancePacket::ComputeRelevance(FDynamicPrimitiveIndexList& DynamicPrimi
 								DrawCommandPacket.AddCommandsForMesh(PrimitiveIndex, PrimitiveSceneInfo, StaticMeshRelevance, StaticMesh, CullingPayloadFlags, Scene, bCanCache, EMeshPass::TranslucencyAll);
 							}
 
+							if (bIsTranslucentHoldoutEnabled)
+							{
+								DrawCommandPacket.AddCommandsForMesh(PrimitiveIndex, PrimitiveSceneInfo, StaticMeshRelevance, StaticMesh, CullingPayloadFlags, Scene, bCanCache, EMeshPass::TranslucencyHoldout);
+							}
+
 							if (ViewRelevance.bTranslucentSurfaceLighting)
 							{
 								DrawCommandPacket.AddCommandsForMesh(PrimitiveIndex, PrimitiveSceneInfo, StaticMeshRelevance, StaticMesh, CullingPayloadFlags, Scene, bCanCache, EMeshPass::LumenTranslucencyRadianceCacheMark);
@@ -1953,6 +1959,11 @@ void FRelevancePacket::ComputeRelevance(FDynamicPrimitiveIndexList& DynamicPrimi
 				// When using all translucency, Standard and AfterDOF are sorted together instead of being rendered like 2 buckets.
 				TranslucentPrimCount.Add(ETranslucencyPass::TPT_AllTranslucency, ViewRelevance.bUsesSceneColorCopy);
 			}
+			
+			if (bIsTranslucentHoldoutEnabled)
+			{
+				TranslucentPrimCount.Add(ETranslucencyPass::TPT_TranslucencyHoldout, true); // use scene copy
+			}
 
 			if (ViewRelevance.bDistortion)
 			{
@@ -2392,6 +2403,11 @@ static void ComputeDynamicMeshRelevance(
 			View.NumVisibleDynamicMeshElements[EMeshPass::TranslucencyAll] += NumElements;
 		}
 
+		{
+			PassMask.Set(EMeshPass::TranslucencyHoldout);
+			View.NumVisibleDynamicMeshElements[EMeshPass::TranslucencyHoldout] += NumElements;
+		}
+
 		if (ViewRelevance.bTranslucentSurfaceLighting)
 		{
 			PassMask.Set(EMeshPass::LumenTranslucencyRadianceCacheMark);
diff --git a/Engine/Source/Runtime/Renderer/Private/TranslucentRendering.cpp b/Engine/Source/Runtime/Renderer/Private/TranslucentRendering.cpp
--- a/Engine/Source/Runtime/Renderer/Private/TranslucentRendering.cpp
+++ b/Engine/Source/Runtime/Renderer/Private/TranslucentRendering.cpp
@@ -105,6 +105,45 @@ static TAutoConsoleVariable<int32> CVarParallelTranslucency(
 	TEXT("Toggles parallel translucency rendering. Parallel rendering must be enabled for this to have an effect."),
 	ECVF_RenderThreadSafe);
 
+static TAutoConsoleVariable<int32> CVarRenderTranslucentHoldout(
+	TEXT("r.RenderTranslucentHoldout"),
+	-1,
+	TEXT("Toggles translucency holdout rendering.\n")
+	TEXT("   0: Disable the dedicated translucent holdout pass (legacy behavior). \n")
+	TEXT("Else: Automatically turn on/off the dedicated holdout pass based on whether Propagate Alpha is on/off\n"),
+	ECVF_RenderThreadSafe);
+
+bool IsTranslucentHoldoutEnabled(EShadingPath ShadingPath)
+{
+	const int32 RenderTranslucentHoldout = CVarRenderTranslucentHoldout.GetValueOnRenderThread();
+	bool bPropagateAlpha = false;
+
+	if (RenderTranslucentHoldout != 0)
+	{
+		static IConsoleVariable* CVar = IConsoleManager::Get().FindConsoleVariable(TEXT("r.PostProcessing.PropagateAlpha"));
+		// TODO: add support to Mobile renderer 
+		bPropagateAlpha = CVar->GetBool() && (ShadingPath != EShadingPath::Mobile);
+	}
+
+	return bPropagateAlpha;
+}
+
+static bool IsTranslucentHoldoutEnabled(TArrayView<const FViewInfo> Views)
+{
+	bool bPropagateAlpha = true;
+
+	for (int32 ViewId = 0; ViewId < Views.Num(); ++ViewId)
+	{
+		if (!IsTranslucentHoldoutEnabled(GetFeatureLevelShadingPath(Views[ViewId].GetFeatureLevel())))
+		{
+			bPropagateAlpha = false;
+			break;
+		}
+	}
+	
+	return bPropagateAlpha;
+}
+
 DynamicRenderScaling::FHeuristicSettings GetDynamicTranslucencyResolutionSettings()
 {
 	DynamicRenderScaling::FHeuristicSettings BucketSetting;
@@ -128,6 +167,7 @@ static const TCHAR* kTranslucencyPassName[] = {
 	TEXT("AfterDOF"),
 	TEXT("AfterDOFModulate"),
 	TEXT("AfterMotionBlur"),
+	TEXT("Holdout"),
 	TEXT("All"),
 };
 static_assert(UE_ARRAY_COUNT(kTranslucencyPassName) == int32(ETranslucencyPass::TPT_MAX), "Fix me");
@@ -138,6 +178,7 @@ static const TCHAR* kTranslucencyColorTextureName[] = {
 	TEXT("Translucency.AfterDOF.Color"),
 	TEXT("Translucency.AfterDOF.Modulate"),
 	TEXT("Translucency.AfterMotionBlur.Color"),
+	TEXT("Translucency.Holdout.Visibility"),
 	TEXT("Translucency.All.Color"),
 };
 static_assert(UE_ARRAY_COUNT(kTranslucencyColorTextureName) == int32(ETranslucencyPass::TPT_MAX), "Fix me");
@@ -148,6 +189,7 @@ static const TCHAR* kTranslucencyColorTextureMultisampledName[] = {
 	TEXT("Translucency.AfterDOF.ColorMS"),
 	TEXT("Translucency.AfterDOF.ModulateMS"),
 	TEXT("Translucency.AfterMotionBlur.ColorMS"),
+	TEXT("Translucency.Holdout.VisibilityMS"),
 	TEXT("Translucency.All.ColorMS"),
 };
 static_assert(UE_ARRAY_COUNT(kTranslucencyColorTextureMultisampledName) == UE_ARRAY_COUNT(kTranslucencyColorTextureName), "Fix me");
@@ -168,6 +210,7 @@ EMeshPass::Type TranslucencyPassToMeshPass(ETranslucencyPass::Type TranslucencyP
 	case ETranslucencyPass::TPT_TranslucencyAfterDOF:			TranslucencyMeshPass = EMeshPass::TranslucencyAfterDOF; break;
 	case ETranslucencyPass::TPT_TranslucencyAfterDOFModulate:	TranslucencyMeshPass = EMeshPass::TranslucencyAfterDOFModulate; break;
 	case ETranslucencyPass::TPT_TranslucencyAfterMotionBlur:	TranslucencyMeshPass = EMeshPass::TranslucencyAfterMotionBlur; break;
+	case ETranslucencyPass::TPT_TranslucencyHoldout:			TranslucencyMeshPass = EMeshPass::TranslucencyHoldout; break;
 	case ETranslucencyPass::TPT_AllTranslucency:				TranslucencyMeshPass = EMeshPass::TranslucencyAll; break;
 	}
 
@@ -313,7 +356,7 @@ public:
 
 IMPLEMENT_GLOBAL_SHADER(FCopySceneColorPS, "/Engine/Private/TranslucentLightingShaders.usf", "CopySceneColorMain", SF_Pixel);
 
-static FRDGTextureRef AddCopySceneColorPass(FRDGBuilder& GraphBuilder, TArrayView<const FViewInfo> Views, FRDGTextureMSAA SceneColor)
+static FRDGTextureRef AddCopySceneColorPass(FRDGBuilder& GraphBuilder, TArrayView<const FViewInfo> Views, FRDGTextureMSAA SceneColor,bool WithAlpha = false)
 {
 	FRDGTextureRef SceneColorCopyTexture = nullptr;
 	ERenderTargetLoadAction LoadAction = ERenderTargetLoadAction::ENoAction;
@@ -368,6 +411,18 @@ static FRDGTextureRef AddCopySceneColorPass(FRDGBuilder& GraphBuilder, TArrayVie
 				LoadAction = ERenderTargetLoadAction::ELoad;
 			}
 
+			FRHIBlendState* BlendState = nullptr;
+
+			if (WithAlpha)
+			{
+				BlendState = FScreenPassPipelineState::FDefaultBlendState::GetRHI();
+			}
+			else
+			{
+				// The original behavior sets alpha to zero.
+				BlendState = TStaticBlendState<CW_RGBA, BO_Add, BF_Zero, BF_One, BO_Add, BF_Zero, BF_Zero>::GetRHI();
+			}
+
 			AddDrawScreenPass(GraphBuilder, {}, View, Viewport, Viewport, VertexShader, PixelShader, PassParameters);
 		}
 	}
@@ -375,6 +430,75 @@ static FRDGTextureRef AddCopySceneColorPass(FRDGBuilder& GraphBuilder, TArrayVie
 	return SceneColorCopyTexture;
 }
 
+static void AddCopySceneColorAlphaPass(FRDGBuilder& GraphBuilder, TArrayView<const FViewInfo> Views, FRDGTextureRef SourceTexture, FRDGTextureRef TargetTexture)
+{
+	ERenderTargetLoadAction LoadAction = ERenderTargetLoadAction::ENoAction;
+
+	RDG_EVENT_SCOPE(GraphBuilder, "CopySceneColorAlpha");
+
+	for (int32 ViewIndex = 0; ViewIndex < Views.Num(); ++ViewIndex)
+	{
+		const FViewInfo& View = Views[ViewIndex];
+
+		if (View.IsUnderwater())
+		{
+			continue;
+		}
+
+		RDG_GPU_MASK_SCOPE(GraphBuilder, View.GPUMask);
+		RDG_EVENT_SCOPE_CONDITIONAL(GraphBuilder, Views.Num() > 1, "View%d", ViewIndex);
+
+		const FIntPoint SceneColorExtent = SourceTexture->Desc.Extent;
+
+		const FScreenPassTextureViewport Viewport(TargetTexture, View.ViewRect);
+
+		TShaderMapRef<FScreenVS> VertexShader(View.ShaderMap);
+		TShaderMapRef<FCopySceneColorPS> PixelShader(View.ShaderMap);
+
+		auto* PassParameters = GraphBuilder.AllocParameters<FCopySceneColorPS::FParameters>();
+		PassParameters->View = View.ViewUniformBuffer;
+		PassParameters->SceneColorTexture = SourceTexture;
+		PassParameters->RenderTargets[0] = FRenderTargetBinding(TargetTexture, LoadAction);
+
+		if (!View.Family->bMultiGPUForkAndJoin)
+		{
+			LoadAction = ERenderTargetLoadAction::ELoad;
+		}
+
+		FRHIBlendState* BlendState = TStaticBlendState<CW_RGBA, BO_Add, BF_Zero, BF_One, BO_Add, BF_One, BF_Zero>::GetRHI();
+		AddDrawScreenPass(GraphBuilder, {}, View, Viewport, Viewport, VertexShader, PixelShader, BlendState, PassParameters);
+	}
+}
+
+class FCopyBackgroundVisibilityPS : public FGlobalShader
+{
+	DECLARE_GLOBAL_SHADER(FCopyBackgroundVisibilityPS);
+	SHADER_USE_PARAMETER_STRUCT(FCopyBackgroundVisibilityPS, FGlobalShader);
+
+	BEGIN_SHADER_PARAMETER_STRUCT(FParameters, )
+		SHADER_PARAMETER_RDG_TEXTURE(Texture2D, TranslucentHoldoutPointTexture)
+		SHADER_PARAMETER_SAMPLER(SamplerState, TranslucentHoldoutPointSampler)
+
+		SHADER_PARAMETER_STRUCT_REF(FViewUniformShaderParameters, ViewUniformBuffer)
+		RENDER_TARGET_BINDING_SLOTS()
+	END_SHADER_PARAMETER_STRUCT()
+
+	enum class EVisibilityCopyType :int32
+	{
+		FromSceneColor, // Copy alpha (background visibility) of scene color to any of the rgb
+		ToSceneColor, // Copy background visibility from any of the rgb to the alpha of an image for scene color composition.
+		MAX
+	};
+
+	static bool ShouldCompilePermutation(const FGlobalShaderPermutationParameters& Parameters)
+	{
+		return IsFeatureLevelSupported(Parameters.Platform, ERHIFeatureLevel::SM5);
+	}
+
+	class FDimVisibilityCopyType :SHADER_PERMUTATION_ENUM_CLASS("VISIBILITY_COPY_TYPE", EVisibilityCopyType);
+	using FPermutationDomain = TShaderPermutationDomain<FDimVisibilityCopyType>;
+};
+
 class FComposeSeparateTranslucencyPS : public FGlobalShader
 {
 	DECLARE_GLOBAL_SHADER(FComposeSeparateTranslucencyPS);
@@ -392,6 +516,7 @@ class FComposeSeparateTranslucencyPS : public FGlobalShader
 		SHADER_PARAMETER(FVector2f, SeparateTranslucencyUVMax)
 		SHADER_PARAMETER(FVector2f, SeparateTranslucencyExtentInverse)
 		SHADER_PARAMETER(int32, bLensDistortion)
+		SHADER_PARAMETER(int32, bPassthroughAlpha)
 
 		SHADER_PARAMETER_RDG_TEXTURE_SRV(Texture2D, SceneColorTexture)
 		SHADER_PARAMETER_SAMPLER(SamplerState,  SceneColorSampler)
@@ -445,9 +570,47 @@ class FTranslucencyUpsampleResponsiveAAPS : public FGlobalShader
 	}
 };
 
+IMPLEMENT_GLOBAL_SHADER(FCopyBackgroundVisibilityPS, "/Engine/Private/ComposeSeparateTranslucency.usf", "CopyBackgroundVisibilityPS", SF_Pixel);
 IMPLEMENT_GLOBAL_SHADER(FComposeSeparateTranslucencyPS, "/Engine/Private/ComposeSeparateTranslucency.usf", "MainPS", SF_Pixel);
 IMPLEMENT_GLOBAL_SHADER(FTranslucencyUpsampleResponsiveAAPS, "/Engine/Private/TranslucencyUpsampling.usf", "UpsampleResponsiveAAPS", SF_Pixel);
 
+static void AddCopyBackgroundVisibilityPass(
+	FRDGBuilder& GraphBuilder,
+	const FViewInfo& View,
+	FRDGTextureRef CopySource,
+	FRDGTextureRef CopyDestination,
+	FIntRect CopyRect,
+	FCopyBackgroundVisibilityPS::EVisibilityCopyType VisibilityCopyType)
+{
+	typedef FCopyBackgroundVisibilityPS SHADER;
+	SHADER::FParameters* PassParameters = GraphBuilder.AllocParameters<SHADER::FParameters>();
+	{
+		PassParameters->ViewUniformBuffer = View.ViewUniformBuffer;
+		PassParameters->TranslucentHoldoutPointTexture = CopySource;
+		PassParameters->TranslucentHoldoutPointSampler = TStaticSamplerState<SF_Point>::GetRHI();
+
+		PassParameters->RenderTargets[0] = FRenderTargetBinding(CopyDestination, ERenderTargetLoadAction::ENoAction);
+	}
+
+	SHADER::FPermutationDomain PixelShaderPermutationVector;
+	PixelShaderPermutationVector.Set<SHADER::FDimVisibilityCopyType>(VisibilityCopyType);
+	TShaderMapRef<SHADER> PixelShader(View.ShaderMap, PixelShaderPermutationVector);
+
+	FRHIBlendState* BlendState = TStaticBlendState<CW_RGBA, BO_Add, BF_One, BF_Zero, BO_Add, BF_One, BF_Zero>::GetRHI();
+
+	FPixelShaderUtils::AddFullscreenPass(
+		GraphBuilder,
+		View.ShaderMap,
+		RDG_EVENT_NAME(
+			"%s(%s) %dx%d",
+			TEXT("CopyBackgroundVisibility"),
+			TEXT("Holdout"),
+			CopyRect.Width(), CopyRect.Height()),
+		PixelShader,
+		PassParameters,
+		CopyRect,
+		BlendState);
+}
 
 FScreenPassTexture FTranslucencyComposition::AddPass(
 	FRDGBuilder& GraphBuilder,
@@ -500,6 +663,8 @@ FScreenPassTexture FTranslucencyComposition::AddPass(
 	RDG_GPU_STAT_SCOPE(GraphBuilder, Translucency);
 	DynamicRenderScaling::FRDGScope DynamicTranslucencyResolutionScope(GraphBuilder, GDynamicTranslucencyResolution);
 
+	bool bPassthroughAlpha = IsTranslucentHoldoutEnabled(GetFeatureLevelShadingPath(View.GetFeatureLevel()));
+
 	const TCHAR* OpName = nullptr;
 	FRHIBlendState* BlendState = nullptr;
 	FRDGTextureRef NewSceneColor = nullptr;
@@ -547,6 +712,25 @@ FScreenPassTexture FTranslucencyComposition::AddPass(
 			OutputDesc,
 			bPostMotionBlur ? TEXT("PostMotionBlurTranslucency.SceneColor") : TEXT("PostDOFTranslucency.SceneColor"));
 	}
+	else if (Operation == EOperation::ComposeToSceneColorAlpha)
+	{
+		check(SceneColor.IsValid());
+
+		// Now we copy any of RGB channel(background visibility) to alpha channel so we can compose against the holdout value of the background.
+		FRDGTextureRef ResolvedTranslucentHoldout = GraphBuilder.CreateTexture(SeparateTranslucencyTexture->Desc,TEXT("Translucency.Holdout.Resolved"));
+		AddCopyBackgroundVisibilityPass(GraphBuilder, View, SeparateTranslucencyTexture, ResolvedTranslucentHoldout,
+			TranslucencyTextures.ViewRect, FCopyBackgroundVisibilityPS::EVisibilityCopyType::ToSceneColor);
+		SeparateTranslucencyTexture = ResolvedTranslucentHoldout;
+
+		OpName = TEXT("ComposeToSceneColorAlpha");
+		// Keep the color of the target, but override the alpha channel
+		// Req: Alpha stores the background visibility of the holdout.
+		// If the background is opaque, it is 1, otherwise 0, use Max operator to clamp.
+		BlendState = TStaticBlendState<CW_RGBA, BO_Add, BF_Zero, BF_One, BO_Add, BF_One, BF_Zero>::GetRHI();
+		ensure(SceneColor.TextureSRV->Desc.Texture->Desc.Flags & TexCreate_RenderTargetable);
+		NewSceneColor = SceneColor.TextureSRV->Desc.Texture;
+		bPassthroughAlpha = false;
+	}
 	else
 	{
 		unimplemented();
@@ -599,6 +783,8 @@ FScreenPassTexture FTranslucencyComposition::AddPass(
 	PassParameters->UndistortingDisplacementTexture = GSystemTextures.GetBlackDummy(GraphBuilder);
 	PassParameters->UndistortingDisplacementSampler = TStaticSamplerState<SF_Bilinear, AM_Clamp, AM_Clamp, AM_Clamp>::GetRHI();
 	PassParameters->bLensDistortion = LensDistortionLUT.IsEnabled();
+	PassParameters->bPassthroughAlpha = bPassthroughAlpha;
+
 	if (LensDistortionLUT.IsEnabled())
 	{
 		PassParameters->UndistortingDisplacementTexture = LensDistortionLUT.UndistortingDisplacementTexture;
@@ -606,7 +792,7 @@ FScreenPassTexture FTranslucencyComposition::AddPass(
 
 	PassParameters->ViewUniformBuffer = View.ViewUniformBuffer;
 
-	if (Operation == EOperation::ComposeToExistingSceneColor)
+	if (Operation == EOperation::ComposeToExistingSceneColor || Operation == EOperation::ComposeToSceneColorAlpha)
 	{
 		PassParameters->RenderTargets[0] = FRenderTargetBinding(NewSceneColor, ERenderTargetLoadAction::ELoad);
 	}
@@ -753,6 +939,13 @@ const FRDGTextureDesc GetPostDOFTranslucentTextureDesc(
 	EShaderPlatform ShaderPlatform)
 {
 	const bool bNeedUAV = SeparateTranslucencyDimensions.NumSamples == 1 && OIT::IsSortedPixelsEnabled(ShaderPlatform);
+	
+	FClearValueBinding ClearValueBinding = bIsModulate ? FClearValueBinding::White : FClearValueBinding::Black;
+	if (TranslucencyPass == ETranslucencyPass::TPT_TranslucencyHoldout)
+	{
+		ClearValueBinding = FClearValueBinding::Black;
+	}
+
 	return FRDGTextureDesc::Create2D(
 		SeparateTranslucencyDimensions.Extent,
 		bIsModulate ? PF_FloatR11G11B10 : PF_FloatRGBA,
@@ -823,7 +1016,8 @@ TRDGUniformBufferRef<FTranslucentBasePassUniformParameters> CreateTranslucentBas
 	FRDGTextureRef SceneColorCopyTexture,
 	const ESceneTextureSetupMode SceneTextureSetupMode,
 	bool bLumenGIEnabled,
-	const FOITData& OITData)
+	const FOITData& OITData,
+	ETranslucencyPass::Type TranslucencyPass)
 {
 	FTranslucentBasePassUniformParameters& BasePassParameters = *GraphBuilder.AllocParameters<FTranslucentBasePassUniformParameters>();
 
@@ -1017,6 +1211,9 @@ TRDGUniformBufferRef<FTranslucentBasePassUniformParameters> CreateTranslucentBas
 
 	BasePassParameters.AVSM = HeterogeneousVolumes::GetAdaptiveVolumetricCameraMapParameters(GraphBuilder, View.ViewState);
 
+	// Translucency pass for holdout
+	BasePassParameters.TranslucencyPass = TranslucencyPass == ETranslucencyPass::TPT_TranslucencyHoldout ? 1 : 0;
+
 	return GraphBuilder.CreateUniformBuffer(&BasePassParameters);
 }
 
@@ -1028,7 +1225,8 @@ TRDGUniformBufferRef<FTranslucentBasePassUniformParameters> CreateTranslucentBas
 	const FTranslucencyLightingVolumeTextures& TranslucencyLightingVolumeTextures,
 	FRDGTextureRef SceneColorCopyTexture,
 	const ESceneTextureSetupMode SceneTextureSetupMode,
-	bool bLumenGIEnabled)
+	bool bLumenGIEnabled,
+	ETranslucencyPass::Type TranslucencyPass)
 {
 	FOITData OITData = OIT::CreateOITData(GraphBuilder, View, OITPass_None);
 	return CreateTranslucentBasePassUniformBuffer(
@@ -1040,7 +1238,8 @@ TRDGUniformBufferRef<FTranslucentBasePassUniformParameters> CreateTranslucentBas
 		SceneColorCopyTexture,
 		SceneTextureSetupMode,
 		bLumenGIEnabled,
-		OITData);
+		OITData,
+		TranslucencyPass);
 }
 
 static FViewShaderParameters GetSeparateTranslucencyViewParameters(const FViewInfo& View, FIntPoint TextureExtent, float ViewportScale, ETranslucencyPass::Type TranslucencyPass)
@@ -1315,6 +1514,9 @@ void FDeferredShadingSceneRenderer::RenderTranslucencyInner(
 	const bool bIsScalingTranslucency = SeparateTranslucencyDimensions.Scale < 1.0f;
 	const bool bIsStandardSeparatedTranslucency = bStandardTranslucentCanRenderSeparate && TranslucencyPass == ETranslucencyPass::TPT_TranslucencyStandard && ViewFamily.AllowStandardTranslucencySeparated();
 	const bool bRenderInSeparateTranslucency = IsSeparateTranslucencyEnabled(TranslucencyPass, SeparateTranslucencyDimensions.Scale) || bIsStandardSeparatedTranslucency;
+	
+	// Holdout rendering 
+	const bool bRenderTranslucencyHold = TranslucencyPass == ETranslucencyPass::TPT_TranslucencyHoldout;
 
 	// Can't reference scene color in scene textures. Scene color copy is used instead.
 	ESceneTextureSetupMode SceneTextureSetupMode = ESceneTextureSetupMode::All;
@@ -1376,7 +1578,7 @@ void FDeferredShadingSceneRenderer::RenderTranslucencyInner(
 				SeparateTranslucencyColorTexture,
 				SeparateTranslucencyColorLoadAction,
 				SeparateTranslucencyDepthTexture.Target,
-				CreateTranslucentBasePassUniformBuffer(GraphBuilder, Scene, View, ViewIndex, TranslucentLightingVolumeTextures, SceneColorCopyTexture, SceneTextureSetupMode, bLumenGIEnabled, OITData),
+				CreateTranslucentBasePassUniformBuffer(GraphBuilder, Scene, View, ViewIndex, TranslucentLightingVolumeTextures, SceneColorCopyTexture, SceneTextureSetupMode, bLumenGIEnabled, OITData, TranslucencyPass),
 				TranslucencyPass,
 				!bCompositeBackToSceneColor,
 				bRenderInParallel,
@@ -1448,6 +1650,116 @@ void FDeferredShadingSceneRenderer::RenderTranslucencyInner(
 
 			++NumProcessedViews;
 		}
+		else if (bRenderTranslucencyHold)
+		{
+			// Main logic:
+			// 1. Copy the holdout background visibility before any translucent pass to the translucent holdout texture.
+			// 2. Run the TranslucencyViewInner logic to accumulate the background visibility (bv) and path throughput (pt) from back to close.
+			// 3. Copy back the alpha channel of the SharedColorTexture.
+			// Notes:
+			// Since alpha will be polluted by alpha holdout blending mode, we cannot directly compose onto the scene color alpha.
+			// with SceneColor.A = SceneColor.A * pt + bv.
+			// E.g., alphaholdout material in front of a translucent material.
+			// Two step direct compose: bv = HoldoutOpacity*(1 - TranslucentOpacity)
+			// This three step copy compose: bv = HoldoutOpacity
+
+			RDG_GPU_MASK_SCOPE(GraphBuilder, View.GPUMask);
+			RDG_EVENT_SCOPE_CONDITIONAL(GraphBuilder, Views.Num() > 1, "View%d", ViewIndex);
+			
+			const FScreenPassTextureViewport SeparateTranslucencyViewport = SeparateTranslucencyDimensions.GetInstancedStereoViewport(View);
+			const bool bCompositeBackToSceneColor = true;
+			const bool bLumenGIEnabled = GetViewPipelineState(View).DiffuseIndirectMethod == EDiffuseIndirectMethod::Lumen;
+
+			/** Separate translucency color is either composited immediately or later during post processing. If done immediately, it's because the view doesn't support
+			 *  compositing (e.g. we're rendering an underwater view) or because we're downsampling the main translucency pass. In this case, we use a local set of
+			 *  textures instead of the external ones passed in.
+			 */
+			FRDGTextureMSAA SeparateTranslucencyColorTexture = SharedColorTexture;
+
+			if (NumProcessedViews == 0 || View.Family->bMultiGPUForkAndJoin)
+			{
+				FTranslucencyPassResources& TranslucencyPassResources = OutTranslucencyResourceMap->Get(ViewIndex, TranslucencyPass);
+				AddCopyBackgroundVisibilityPass(GraphBuilder, View, SceneColorCopyTexture/*SceneTextures.Color.Resolve*/, SeparateTranslucencyColorTexture.Target,
+					TranslucencyPassResources.ViewRect, FCopyBackgroundVisibilityPS::EVisibilityCopyType::FromSceneColor);
+			}
+
+			// NOTE: We need to read the depth texture for final fog accumulation, but
+			// in the current phase, we don't need depth texture.
+			const bool bAlphaHoldoutNeedDepthTest = true;
+			FRDGTextureMSAA SeparateTranslucencyDepthTexture;
+			if (bAlphaHoldoutNeedDepthTest)
+			{
+				SeparateTranslucencyDepthTexture = SharedDepthTexture;
+			}
+
+			// No scale is needed
+			const float ViewportScale = 1.0f;
+
+			// Simply load as we have already initialized the texture.
+			const ERenderTargetLoadAction SeparateTranslucencyColorLoadAction = ERenderTargetLoadAction::ELoad;
+
+			FOITData OITData = OIT::CreateOITData(GraphBuilder, View, OITPass_SeperateTranslucency);
+
+			RenderTranslucencyViewInner(
+				GraphBuilder,
+				*this,
+				View,
+				SeparateTranslucencyViewport,
+				ViewportScale,
+				SeparateTranslucencyColorTexture,
+				SeparateTranslucencyColorLoadAction,
+				SeparateTranslucencyDepthTexture.Target,
+				CreateTranslucentBasePassUniformBuffer(GraphBuilder, Scene, View, ViewIndex, TranslucentLightingVolumeTextures, SceneColorCopyTexture, SceneTextureSetupMode, bLumenGIEnabled, OITData, TranslucencyPass),
+				TranslucencyPass,
+				!bCompositeBackToSceneColor,
+				bRenderInParallel,
+				InstanceCullingManager);
+
+			{
+				FTranslucencyPassResources& TranslucencyPassResources = OutTranslucencyResourceMap->Get(ViewIndex, TranslucencyPass);
+				TranslucencyPassResources.ViewRect = View.ViewRect;
+				TranslucencyPassResources.ColorTexture = SharedColorTexture;
+				TranslucencyPassResources.DepthTexture = SharedDepthTexture;
+			}
+
+			if (OITData.PassType & OITPass_SeperateTranslucency)
+			{
+				OIT::AddOITComposePass(GraphBuilder, View, OITData, SeparateTranslucencyColorTexture.Target);
+			}
+
+			if (bCompositeBackToSceneColor)
+			{
+				FRDGTextureRef SeparateTranslucencyDepthResolve = nullptr;
+				FRDGTextureRef SceneDepthResolve = nullptr;
+				if (TranslucencyPass != ETranslucencyPass::TPT_TranslucencyAfterMotionBlur)
+				{
+					::AddResolveSceneDepthPass(GraphBuilder, View, SeparateTranslucencyDepthTexture);
+
+					SeparateTranslucencyDepthResolve = SeparateTranslucencyDepthTexture.Resolve;
+					SceneDepthResolve = SceneTextures.Depth.Resolve;
+				}
+
+				FTranslucencyPassResources& TranslucencyPassResources = OutTranslucencyResourceMap->Get(ViewIndex, TranslucencyPass);
+
+				FTranslucencyComposition TranslucencyComposition;
+				TranslucencyComposition.Operation = FTranslucencyComposition::EOperation::ComposeToSceneColorAlpha;
+				TranslucencyComposition.SceneColor = FScreenPassTextureSlice::CreateFromScreenPassTexture(GraphBuilder, FScreenPassTexture(SceneTextures.Color.Target, View.ViewRect));
+				TranslucencyComposition.SceneDepth = FScreenPassTexture(SceneTextures.Depth.Resolve, View.ViewRect);
+				TranslucencyComposition.OutputViewport = FScreenPassTextureViewport(SceneTextures.Depth.Resolve, View.ViewRect);
+
+				FScreenPassTexture UpscaledTranslucency = TranslucencyComposition.AddPass(
+					GraphBuilder, View, TranslucencyPassResources);
+
+				ensure(View.ViewRect == UpscaledTranslucency.ViewRect);
+				ensure(UpscaledTranslucency.Texture == SceneTextures.Color.Target);
+
+				//Invalidate.
+				TranslucencyPassResources = FTranslucencyPassResources();
+				TranslucencyPassResources.Pass = TranslucencyPass;
+			}
+
+			++NumProcessedViews;
+		}
 		else
 		{
 			// When rendering translucent meshes under water, we skip modulate passes which are only required when compositing separate translucency passes from render target.
@@ -1476,7 +1788,7 @@ void FDeferredShadingSceneRenderer::RenderTranslucencyInner(
 				SceneTextures.Color,
 				SceneColorLoadAction,
 				SceneTextures.Depth.Target,
-				CreateTranslucentBasePassUniformBuffer(GraphBuilder, Scene, View, ViewIndex, TranslucentLightingVolumeTextures, SceneColorCopyTexture, SceneTextureSetupMode, bLumenGIEnabled, OITData),
+				CreateTranslucentBasePassUniformBuffer(GraphBuilder, Scene, View, ViewIndex, TranslucentLightingVolumeTextures, SceneColorCopyTexture, SceneTextureSetupMode, bLumenGIEnabled, OITData, TranslucencyPass),
 				TranslucencyPass,
 				bResolveColorTexture,
 				bRenderInParallel,
@@ -1514,10 +1826,11 @@ void FDeferredShadingSceneRenderer::RenderTranslucency(
 	DynamicRenderScaling::FRDGScope DynamicTranslucencyResolutionScope(GraphBuilder, GDynamicTranslucencyResolution);
 
 	FRDGTextureRef SceneColorCopyTexture = nullptr;
+	const bool bIsTranslucentHoldoutEnabled = IsTranslucentHoldoutEnabled(Views);
 
 	if (EnumHasAnyFlags(ViewsToRender, ETranslucencyView::AboveWater))
 	{
-		SceneColorCopyTexture = AddCopySceneColorPass(GraphBuilder, Views, SceneTextures.Color);
+		SceneColorCopyTexture = AddCopySceneColorPass(GraphBuilder, Views, SceneTextures.Color, /*WithAlpha*/ bIsTranslucentHoldoutEnabled);
 	}
 
 	// Create a shared depth texture at the correct resolution.
@@ -1583,6 +1896,12 @@ void FDeferredShadingSceneRenderer::RenderTranslucency(
 	{
 		RenderTranslucencyInner(GraphBuilder, SceneTextures, TranslucentLightingVolumeTextures, OutTranslucencyResourceMap, OutSharedDepthTexture, ViewsToRender, SceneColorCopyTexture, ETranslucencyPass::TPT_AllTranslucency, InstanceCullingManager, bStandardTranslucentCanRenderSeparate);
 	}
+
+	if (bIsTranslucentHoldoutEnabled && SceneColorCopyTexture)
+	{
+		// Render the translucent holdout background visibility to the alpha channel of SceneColor.
+		RenderTranslucencyInner(GraphBuilder, SceneTextures, TranslucentLightingVolumeTextures, OutTranslucencyResourceMap, OutSharedDepthTexture, ViewsToRender, SceneColorCopyTexture, ETranslucencyPass::TPT_TranslucencyHoldout, InstanceCullingManager, bStandardTranslucentCanRenderSeparate);
+	}
 }
 
 
diff --git a/Engine/Source/Runtime/Renderer/Private/TranslucentRendering.h b/Engine/Source/Runtime/Renderer/Private/TranslucentRendering.h
--- a/Engine/Source/Runtime/Renderer/Private/TranslucentRendering.h
+++ b/Engine/Source/Runtime/Renderer/Private/TranslucentRendering.h
@@ -50,6 +50,9 @@ FSeparateTranslucencyDimensions UpdateSeparateTranslucencyDimensions(const FScen
 /** Returns whether the view family is requesting to render translucency. */
 bool ShouldRenderTranslucency(const FSceneViewFamily& ViewFamily);
 
+/** Check whether translucent holdout is enabled for a specific shading path. */
+bool IsTranslucentHoldoutEnabled(EShadingPath ShadingPath);
+
 /** Check if separate translucency pass is needed for given pass and downsample scale */
 bool IsSeparateTranslucencyEnabled(ETranslucencyPass::Type TranslucencyPass, float DownsampleScale);
 
@@ -68,6 +71,7 @@ struct FTranslucencyComposition
 		UpscaleOnly,
 		ComposeToExistingSceneColor,
 		ComposeToNewSceneColor,
+		ComposeToSceneColorAlpha
 	};
 
 	EOperation Operation = EOperation::UpscaleOnly;
diff --git a/Engine/Source/Runtime/Renderer/Public/MeshPassProcessor.h b/Engine/Source/Runtime/Renderer/Public/MeshPassProcessor.h
--- a/Engine/Source/Runtime/Renderer/Public/MeshPassProcessor.h
+++ b/Engine/Source/Runtime/Renderer/Public/MeshPassProcessor.h
@@ -51,6 +51,7 @@ namespace EMeshPass
 		TranslucencyAfterDOF,
 		TranslucencyAfterDOFModulate,
 		TranslucencyAfterMotionBlur,
+		TranslucencyHoldout, /** A standalone pass to render all translucency for holdout, inferring the background visibility*/
 		TranslucencyAll, /** Drawing all translucency, regardless of separate or standard.  Used when drawing translucency outside of the main renderer, eg FRendererModule::DrawTile. */
 		LightmapDensity,
 		DebugViewMode, /** Any of EDebugViewShaderMode */
@@ -107,6 +108,7 @@ inline const TCHAR* GetMeshPassName(EMeshPass::Type MeshPass)
 	case EMeshPass::TranslucencyAfterDOF: return TEXT("TranslucencyAfterDOF");
 	case EMeshPass::TranslucencyAfterDOFModulate: return TEXT("TranslucencyAfterDOFModulate");
 	case EMeshPass::TranslucencyAfterMotionBlur: return TEXT("TranslucencyAfterMotionBlur");
+	case EMeshPass::TranslucencyHoldout: return TEXT("TranslucencyHoldout");
 	case EMeshPass::TranslucencyAll: return TEXT("TranslucencyAll");
 	case EMeshPass::LightmapDensity: return TEXT("LightmapDensity");
 	case EMeshPass::DebugViewMode: return TEXT("DebugViewMode");
@@ -135,9 +137,9 @@ inline const TCHAR* GetMeshPassName(EMeshPass::Type MeshPass)
 	}
 
 #if WITH_EDITOR
-	static_assert(EMeshPass::Num == 37 + 4, "Need to update switch(MeshPass) after changing EMeshPass"); // GUID to prevent incorrect auto-resolves, please change when changing the expression: {674D7D62-CFD8-4971-9A8D-CD91E5612CD8}
+	static_assert(EMeshPass::Num == 38 + 4, "Need to update switch(MeshPass) after changing EMeshPass"); // GUID to prevent incorrect auto-resolves, please change when changing the expression: {674D7D62-CFD8-4971-9A8D-CD91E5612CD8}
 #else
-	static_assert(EMeshPass::Num == 37, "Need to update switch(MeshPass) after changing EMeshPass"); // GUID to prevent incorrect auto-resolves, please change when changing the expression: {674D7D62-CFD8-4971-9A8D-CD91E5612CD8}
+	static_assert(EMeshPass::Num == 38, "Need to update switch(MeshPass) after changing EMeshPass"); // GUID to prevent incorrect auto-resolves, please change when changing the expression: {674D7D62-CFD8-4971-9A8D-CD91E5612CD8}
 #endif
 
 	checkf(0, TEXT("Missing case for EMeshPass %u"), (uint32)MeshPass);
