Fix incremental build warning of expression being always true.

diff --git a/Engine/Source/Runtime/Renderer/Private/TranslucentRendering.cpp b/Engine/Source/Runtime/Renderer/Private/TranslucentRendering.cpp
--- a/Engine/Source/Runtime/Renderer/Private/TranslucentRendering.cpp
+++ b/Engine/Source/Runtime/Renderer/Private/TranslucentRendering.cpp
@@ -1731,7 +1731,7 @@ void FDeferredShadingSceneRenderer::RenderTranslucencyInner(
 			{
 				FRDGTextureRef SeparateTranslucencyDepthResolve = nullptr;
 				FRDGTextureRef SceneDepthResolve = nullptr;
-				if (TranslucencyPass != ETranslucencyPass::TPT_TranslucencyAfterMotionBlur)
+				
 				{
 					::AddResolveSceneDepthPass(GraphBuilder, View, SeparateTranslucencyDepthTexture);
 
