Remove some unused global resources from FDeferredShadingSceneRenderer.

diff --git a/Engine/Source/Runtime/Renderer/Private/DeferredShadingRenderer.cpp b/Engine/Source/Runtime/Renderer/Private/DeferredShadingRenderer.cpp
--- a/Engine/Source/Runtime/Renderer/Private/DeferredShadingRenderer.cpp
+++ b/Engine/Source/Runtime/Renderer/Private/DeferredShadingRenderer.cpp
@@ -394,9 +394,6 @@ FDeferredShadingSceneRenderer::FDeferredShadingSceneRenderer(const FSceneViewFam
 DECLARE_CYCLE_STAT(TEXT("Wait RayTracing Add Mesh Batch"), STAT_WaitRayTracingAddMesh, STATGROUP_SceneRendering);
 DECLARE_CYCLE_STAT(TEXT("Wait Ray Tracing Scene Initialization"), STAT_WaitRayTracingSceneInitTask, STATGROUP_SceneRendering);
 
-FGlobalDynamicIndexBuffer FDeferredShadingSceneRenderer::DynamicIndexBufferForInitShadows;
-FGlobalDynamicVertexBuffer FDeferredShadingSceneRenderer::DynamicVertexBufferForInitShadows;
-
 /**
  * Returns true if the depth Prepass needs to run
  */
diff --git a/Engine/Source/Runtime/Renderer/Private/DeferredShadingRenderer.h b/Engine/Source/Runtime/Renderer/Private/DeferredShadingRenderer.h
--- a/Engine/Source/Runtime/Renderer/Private/DeferredShadingRenderer.h
+++ b/Engine/Source/Runtime/Renderer/Private/DeferredShadingRenderer.h
@@ -541,12 +541,6 @@ private:
 		return bAnyViewHasGIMethodSupportingDFAO;
 	}
 
-	static FGlobalDynamicIndexBuffer DynamicIndexBufferForInitViews;
-	static FGlobalDynamicIndexBuffer DynamicIndexBufferForInitShadows;
-	static FGlobalDynamicVertexBuffer DynamicVertexBufferForInitViews;
-	static FGlobalDynamicVertexBuffer DynamicVertexBufferForInitShadows;
-	static TGlobalResource<FGlobalDynamicReadBuffer> DynamicReadBufferForInitShadows;
-
 	FSeparateTranslucencyDimensions SeparateTranslucencyDimensions;
 
 	/** Creates a per object projected shadow for the given interaction. */
