Moved structs and functions related to FRayTracingMeshBatchTaskPage out of GatherWorldInstancesForView(...)

diff --git a/Engine/Source/Runtime/Renderer/Private/RayTracing/RayTracing.cpp b/Engine/Source/Runtime/Renderer/Private/RayTracing/RayTracing.cpp
--- a/Engine/Source/Runtime/Renderer/Private/RayTracing/RayTracing.cpp
+++ b/Engine/Source/Runtime/Renderer/Private/RayTracing/RayTracing.cpp
@@ -147,6 +147,48 @@ static FAutoConsoleVariableSink CVarRefreshRayTracingInstancesSink(FConsoleComma
 
 namespace RayTracing
 {
+	static void AddDebugRayTracingInstanceFlags(ERayTracingInstanceFlags& InOutFlags)
+	{
+		if (GRayTracingDebugForceOpaque)
+		{
+			InOutFlags |= ERayTracingInstanceFlags::ForceOpaque;
+		}
+		if (GRayTracingDebugDisableTriangleCull)
+		{
+			InOutFlags |= ERayTracingInstanceFlags::TriangleCullDisable;
+		}
+	}
+
+	// Configure ray tracing scene options based on currently enabled features and their needs
+	FSceneOptions::FSceneOptions(
+		FScene& Scene,
+		const FViewFamilyInfo& ViewFamily,
+		FViewInfo& View,
+		EDiffuseIndirectMethod DiffuseIndirectMethod,
+		EReflectionsMethod ReflectionsMethod,
+		RayTracing::FSceneOptions& SceneOptions)
+	{
+		SceneOptions.bTranslucentGeometry = false;
+		LumenHardwareRayTracing::SetRayTracingSceneOptions(View, DiffuseIndirectMethod, ReflectionsMethod, SceneOptions);
+		RayTracingShadows::SetRayTracingSceneOptions(Scene.bHasLightsWithRayTracedShadows, SceneOptions);
+
+		if (ShouldRenderRayTracingTranslucency(View))
+		{
+			SceneOptions.bTranslucentGeometry = true;
+		}
+
+		if (ViewFamily.EngineShowFlags.PathTracing
+			&& FDataDrivenShaderPlatformInfo::GetSupportsPathTracing(Scene.GetShaderPlatform()))
+		{
+			SceneOptions.bTranslucentGeometry = true;
+		}
+
+		if (GRayTracingExcludeTranslucent != 0)
+		{
+			SceneOptions.bTranslucentGeometry = false;
+		}
+	}
+
 	struct FRelevantPrimitive
 	{
 		FRHIRayTracingGeometry* RayTracingGeometryRHI = nullptr;
@@ -330,6 +372,77 @@ namespace RayTracing
 		return StaticUniformBufferScope;
 	}
 
+	struct FRayTracingMeshBatchWorkItem
+	{
+		const FPrimitiveSceneProxy* SceneProxy = nullptr;
+		const FRHIRayTracingGeometry* RayTracingGeometry = nullptr;
+		TArray<FMeshBatch> MeshBatchesOwned;
+		TArrayView<const FMeshBatch> MeshBatchesView;
+		FRayTracingSBTAllocation* SBTAllocation;
+
+		TArrayView<const FMeshBatch> GetMeshBatches() const
+		{
+			if (MeshBatchesOwned.Num())
+			{
+				check(MeshBatchesView.Num() == 0);
+				return TArrayView<const FMeshBatch>(MeshBatchesOwned);
+			}
+			else
+			{
+				check(MeshBatchesOwned.Num() == 0);
+				return MeshBatchesView;
+			}
+		}
+	};
+
+	struct FRayTracingMeshBatchTaskPage
+	{
+		static constexpr uint32 MaxWorkItems = 128; // Try to keep individual pages small to avoid slow-path memory allocations
+
+		FRayTracingMeshBatchWorkItem WorkItems[MaxWorkItems];
+		uint32 NumWorkItems = 0;
+		FRayTracingMeshBatchTaskPage* Next = nullptr;
+	};
+
+	void DispatchRayTracingMeshBatchTask(FSceneRenderingBulkObjectAllocator& InBulkAllocator, FScene& Scene, FViewInfo& View, FRayTracingMeshBatchTaskPage* MeshBatchTaskHead, uint32 NumPendingMeshBatches)
+	{
+		FDynamicRayTracingMeshCommandStorage* TaskDynamicCommandStorage = InBulkAllocator.Create<FDynamicRayTracingMeshCommandStorage>();
+		View.DynamicRayTracingMeshCommandStoragePerTask.Add(TaskDynamicCommandStorage);
+
+		FRayTracingShaderBindingDataOneFrameArray* TaskDirtyShaderBindings = InBulkAllocator.Create<FRayTracingShaderBindingDataOneFrameArray>();
+		TaskDirtyShaderBindings->Reserve(NumPendingMeshBatches);
+		View.DirtyRayTracingShaderBindingsPerTask.Add(TaskDirtyShaderBindings);
+
+		View.AddRayTracingMeshBatchTaskList.Add(FFunctionGraphTask::CreateAndDispatchWhenReady(
+			[TaskDataHead = MeshBatchTaskHead, &View, &Scene, TaskDynamicCommandStorage, TaskDirtyShaderBindings]()
+			{
+				FTaskTagScope TaskTagScope(ETaskTag::EParallelRenderingThread);
+				TRACE_CPUPROFILER_EVENT_SCOPE(RayTracingMeshBatchTask);
+				FRayTracingMeshBatchTaskPage* Page = TaskDataHead;
+				const int32 ExpectedMaxVisibleCommands = TaskDirtyShaderBindings->Max();
+				while (Page)
+				{
+					for (uint32 ItemIndex = 0; ItemIndex < Page->NumWorkItems; ++ItemIndex)
+					{
+						const FRayTracingMeshBatchWorkItem& WorkItem = Page->WorkItems[ItemIndex];
+						TArrayView<const FMeshBatch> MeshBatches = WorkItem.GetMeshBatches();
+						for (int32 SegmentIndex = 0; SegmentIndex < MeshBatches.Num(); SegmentIndex++)
+						{
+							const FMeshBatch& MeshBatch = MeshBatches[SegmentIndex];
+							FDynamicRayTracingMeshCommandContext CommandContext(
+								*TaskDynamicCommandStorage, *TaskDirtyShaderBindings,
+								WorkItem.RayTracingGeometry, SegmentIndex, WorkItem.SBTAllocation);
+							FRayTracingMeshProcessor RayTracingMeshProcessor(&CommandContext, &Scene, &View, Scene.CachedRayTracingMeshCommandsMode);
+							RayTracingMeshProcessor.AddMeshBatch(MeshBatch, 1, WorkItem.SceneProxy);
+						}
+					}
+					FRayTracingMeshBatchTaskPage* NextPage = Page->Next;
+					Page = NextPage;
+				}
+				check(ExpectedMaxVisibleCommands <= TaskDirtyShaderBindings->Max());
+			}, TStatId(), nullptr, ENamedThreads::AnyThread));
+	};
+
 	void GatherRelevantPrimitives(FScene& Scene, const FViewInfo& View, FRelevantPrimitiveList& Result)
 	{
 		TArray<int32> StaticPrimitives;
@@ -721,48 +834,6 @@ namespace RayTracing
 		Result.bValid = true;
 	}
 
-	static void AddDebugRayTracingInstanceFlags(ERayTracingInstanceFlags& InOutFlags)
-	{
-		if (GRayTracingDebugForceOpaque)
-		{
-			InOutFlags |= ERayTracingInstanceFlags::ForceOpaque;
-		}
-		if (GRayTracingDebugDisableTriangleCull)
-		{
-			InOutFlags |= ERayTracingInstanceFlags::TriangleCullDisable;
-		}
-	}
-
-	// Configure ray tracing scene options based on currently enabled features and their needs
-	FSceneOptions::FSceneOptions(
-		FScene& Scene,
-		const FViewFamilyInfo& ViewFamily,
-		FViewInfo& View,
-		EDiffuseIndirectMethod DiffuseIndirectMethod,
-		EReflectionsMethod ReflectionsMethod,
-		RayTracing::FSceneOptions& SceneOptions)
-	{
-		SceneOptions.bTranslucentGeometry = false;
-		LumenHardwareRayTracing::SetRayTracingSceneOptions(View, DiffuseIndirectMethod, ReflectionsMethod, SceneOptions);
-		RayTracingShadows::SetRayTracingSceneOptions(Scene.bHasLightsWithRayTracedShadows, SceneOptions);
-
-		if (ShouldRenderRayTracingTranslucency(View))
-		{
-			SceneOptions.bTranslucentGeometry = true;
-		}
-
-		if (ViewFamily.EngineShowFlags.PathTracing
-			&& FDataDrivenShaderPlatformInfo::GetSupportsPathTracing(Scene.GetShaderPlatform()))
-		{
-			SceneOptions.bTranslucentGeometry = true;
-		}
-
-		if (GRayTracingExcludeTranslucent != 0)
-		{
-			SceneOptions.bTranslucentGeometry = false;
-		}
-	}
-
 	bool GatherWorldInstancesForView(
 		FRDGBuilder& GraphBuilder,
 		FScene& Scene,
@@ -822,37 +893,6 @@ namespace RayTracing
 
 			const int64 SharedBufferGenerationID = Scene.GetRayTracingDynamicGeometryCollection()->BeginUpdate();
 
-			struct FRayTracingMeshBatchWorkItem
-			{
-				const FPrimitiveSceneProxy* SceneProxy = nullptr;
-				const FRHIRayTracingGeometry* RayTracingGeometry = nullptr;
-				TArray<FMeshBatch> MeshBatchesOwned;
-				TArrayView<const FMeshBatch> MeshBatchesView;
-				FRayTracingSBTAllocation* SBTAllocation;
-
-				TArrayView<const FMeshBatch> GetMeshBatches() const
-				{
-					if (MeshBatchesOwned.Num())
-					{
-						check(MeshBatchesView.Num() == 0);
-						return TArrayView<const FMeshBatch>(MeshBatchesOwned);
-					}
-					else
-					{
-						check(MeshBatchesOwned.Num() == 0);
-						return MeshBatchesView;
-					}
-				}
-			};
-
-			static constexpr uint32 MaxWorkItemsPerPage = 128; // Try to keep individual pages small to avoid slow-path memory allocations
-			struct FRayTracingMeshBatchTaskPage
-			{
-				FRayTracingMeshBatchWorkItem WorkItems[MaxWorkItemsPerPage];
-				uint32 NumWorkItems = 0;
-				FRayTracingMeshBatchTaskPage* Next = nullptr;
-			};
-
 			FRayTracingMeshBatchTaskPage* MeshBatchTaskHead = nullptr;
 			FRayTracingMeshBatchTaskPage* MeshBatchTaskPage = nullptr;
 			uint32 NumPendingMeshBatches = 0;
@@ -862,41 +902,7 @@ namespace RayTracing
 				{
 					if (MeshBatchTaskHead)
 					{
-						FDynamicRayTracingMeshCommandStorage* TaskDynamicCommandStorage = InBulkAllocator.Create<FDynamicRayTracingMeshCommandStorage>();
-						View.DynamicRayTracingMeshCommandStoragePerTask.Add(TaskDynamicCommandStorage);
-
-						FRayTracingShaderBindingDataOneFrameArray* TaskDirtyShaderBindings = InBulkAllocator.Create<FRayTracingShaderBindingDataOneFrameArray>();
-						TaskDirtyShaderBindings->Reserve(NumPendingMeshBatches);
-						View.DirtyRayTracingShaderBindingsPerTask.Add(TaskDirtyShaderBindings);
-
-						View.AddRayTracingMeshBatchTaskList.Add(FFunctionGraphTask::CreateAndDispatchWhenReady(
-							[TaskDataHead = MeshBatchTaskHead, &View, &Scene, TaskDynamicCommandStorage, TaskDirtyShaderBindings]()
-							{
-								FTaskTagScope TaskTagScope(ETaskTag::EParallelRenderingThread);
-								TRACE_CPUPROFILER_EVENT_SCOPE(RayTracingMeshBatchTask);
-								FRayTracingMeshBatchTaskPage* Page = TaskDataHead;
-								const int32 ExpectedMaxVisibieCommands = TaskDirtyShaderBindings->Max();
-								while (Page)
-								{
-									for (uint32 ItemIndex = 0; ItemIndex < Page->NumWorkItems; ++ItemIndex)
-									{
-										const FRayTracingMeshBatchWorkItem& WorkItem = Page->WorkItems[ItemIndex];
-										TArrayView<const FMeshBatch> MeshBatches = WorkItem.GetMeshBatches();
-										for (int32 SegmentIndex = 0; SegmentIndex < MeshBatches.Num(); SegmentIndex++)
-										{
-											const FMeshBatch& MeshBatch = MeshBatches[SegmentIndex];
-											FDynamicRayTracingMeshCommandContext CommandContext(
-												*TaskDynamicCommandStorage, *TaskDirtyShaderBindings,
-												WorkItem.RayTracingGeometry, SegmentIndex, WorkItem.SBTAllocation);
-											FRayTracingMeshProcessor RayTracingMeshProcessor(&CommandContext, &Scene, &View, Scene.CachedRayTracingMeshCommandsMode);
-											RayTracingMeshProcessor.AddMeshBatch(MeshBatch, 1, WorkItem.SceneProxy);
-										}
-									}
-									FRayTracingMeshBatchTaskPage* NextPage = Page->Next;
-									Page = NextPage;
-								}
-								check(ExpectedMaxVisibieCommands <= TaskDirtyShaderBindings->Max());
-							}, TStatId(), nullptr, ENamedThreads::AnyThread));
+						DispatchRayTracingMeshBatchTask(InBulkAllocator, Scene, View, MeshBatchTaskHead, NumPendingMeshBatches);
 					}
 
 					MeshBatchTaskHead = nullptr;
@@ -1176,7 +1182,7 @@ namespace RayTracing
 								KickRayTracingMeshBatchTask();
 							}
 
-							if (MeshBatchTaskPage == nullptr || MeshBatchTaskPage->NumWorkItems == MaxWorkItemsPerPage)
+							if (MeshBatchTaskPage == nullptr || MeshBatchTaskPage->NumWorkItems == FRayTracingMeshBatchTaskPage::MaxWorkItems)
 							{
 								FRayTracingMeshBatchTaskPage* NextPage = InBulkAllocator.Create<FRayTracingMeshBatchTaskPage>();
 								if (MeshBatchTaskHead == nullptr)
