SkyAtmosphere: Cleanup-rename of bLightDiskEnabled into bIsReflectionCapture.

diff --git a/Engine/Source/Runtime/Renderer/Private/ReflectionEnvironmentRealTimeCapture.cpp b/Engine/Source/Runtime/Renderer/Private/ReflectionEnvironmentRealTimeCapture.cpp
--- a/Engine/Source/Runtime/Renderer/Private/ReflectionEnvironmentRealTimeCapture.cpp
+++ b/Engine/Source/Runtime/Renderer/Private/ReflectionEnvironmentRealTimeCapture.cpp
@@ -583,7 +583,7 @@ void FScene::AllocateAndCaptureFrameSkyEnvMap(
 				SkyRC.SkyAtmosphereCameraAerialPerspectiveVolumeRayOnly = BlackDummy3dTex;
 
 				SkyRC.Viewport = FIntRect(FIntPoint(0, 0), FIntPoint(CubeWidth, CubeWidth));
-				SkyRC.bLightDiskEnabled = false;
+				SkyRC.bIsReflectionCapture = true;
 				SkyRC.bRenderSkyPixel = true;
 				SkyRC.AerialPerspectiveStartDepthInCm = 0.01f;
 				SkyRC.NearClippingDistance = 0.01f;
diff --git a/Engine/Source/Runtime/Renderer/Private/SkyAtmosphereRendering.cpp b/Engine/Source/Runtime/Renderer/Private/SkyAtmosphereRendering.cpp
--- a/Engine/Source/Runtime/Renderer/Private/SkyAtmosphereRendering.cpp
+++ b/Engine/Source/Runtime/Renderer/Private/SkyAtmosphereRendering.cpp
@@ -1869,9 +1869,8 @@ void FSceneRenderer::RenderSkyAtmosphereInternal(
 		PsPassParameters->VolumetricCloudDepthTexture = SkyRC.VolumetricCloudDepthTexture;
 		PsPassParameters->InputCloudLuminanceTransmittanceTexture = SkyRC.InputCloudLuminanceTransmittanceTexture;
 		PsPassParameters->AerialPerspectiveStartDepthKm = AerialPerspectiveStartDepthInCm * CM_TO_KM;
-		// NOTE: bLightDiskEnabled == !View.bIsReflectionCapture
-		PsPassParameters->bPropagateAlphaNonReflection = (IsPostProcessingWithAlphaChannelSupported() && SkyRC.bLightDiskEnabled) ? 1 : 0;
-		PsPassParameters->SourceDiskEnabled = SkyRC.bLightDiskEnabled ? 1 : 0;
+		PsPassParameters->bPropagateAlphaNonReflection = (IsPostProcessingWithAlphaChannelSupported() && !SkyRC.bIsReflectionCapture) ? 1 : 0;
+		PsPassParameters->SourceDiskEnabled = SkyRC.bIsReflectionCapture ? 0 : 1;
 		PsPassParameters->DepthReadDisabled = SkyRC.bDepthReadDisabled ? 1 : 0;
 		if (bShouldSampleCloudShadow || SkyRC.bShouldSampleCloudSkyAO)
 		{
@@ -2034,7 +2033,7 @@ void FSceneRenderer::RenderSkyAtmosphere(FRDGBuilder& GraphBuilder, const FMinim
 		SkyRC.bSceneHasSkyMaterial = View.bSceneHasSkyMaterial;
 
 		SkyRC.Viewport = View.ViewRect;
-		SkyRC.bLightDiskEnabled = !View.bIsReflectionCapture;
+		SkyRC.bIsReflectionCapture = View.bIsReflectionCapture;
 		SkyRC.AerialPerspectiveStartDepthInCm = GetValidAerialPerspectiveStartDepthInCm(View, SkyAtmosphereSceneProxy);
 		SkyRC.NearClippingDistance = View.NearClippingDistance;
 		SkyRC.FeatureLevel = View.FeatureLevel;
diff --git a/Engine/Source/Runtime/Renderer/Private/SkyAtmosphereRendering.h b/Engine/Source/Runtime/Renderer/Private/SkyAtmosphereRendering.h
--- a/Engine/Source/Runtime/Renderer/Private/SkyAtmosphereRendering.h
+++ b/Engine/Source/Runtime/Renderer/Private/SkyAtmosphereRendering.h
@@ -102,7 +102,7 @@ struct FSkyAtmosphereRenderContext
 
 	FIntRect Viewport;
 
-	bool bLightDiskEnabled;
+	bool bIsReflectionCapture;
 	bool bRenderSkyPixel;
 	float AerialPerspectiveStartDepthInCm;
 	float NearClippingDistance;
diff --git a/Engine/Source/Runtime/Renderer/Private/VolumetricCloudRendering.cpp b/Engine/Source/Runtime/Renderer/Private/VolumetricCloudRendering.cpp
--- a/Engine/Source/Runtime/Renderer/Private/VolumetricCloudRendering.cpp
+++ b/Engine/Source/Runtime/Renderer/Private/VolumetricCloudRendering.cpp
@@ -2894,7 +2894,7 @@ bool FSceneRenderer::RenderVolumetricCloud(
 					SkyRC.Viewport = bShouldViewRenderVolumetricCloudRenderTarget ?
 						FIntRect(FIntPoint(0, 0), FIntPoint(DestinationRT->Desc.GetSize().X, DestinationRT->Desc.GetSize().Y)) // this texture is per view, so we don't need to use viewrect inside it
 						: ViewInfo.ViewRect;
-					SkyRC.bLightDiskEnabled = !ViewInfo.bIsReflectionCapture;
+					SkyRC.bIsReflectionCapture = ViewInfo.bIsReflectionCapture;
 					SkyRC.AerialPerspectiveStartDepthInCm = GetValidAerialPerspectiveStartDepthInCm(ViewInfo, SkyAtmosphereSceneProxy);
 					SkyRC.NearClippingDistance = ViewInfo.NearClippingDistance;
 					SkyRC.FeatureLevel = ViewInfo.FeatureLevel;
