Check RenderRectLightsAsSpotLights(...) when determining culling method for rect lights in light grid.
- also added new permutation to FLightGridInjectionCS to avoid overhead when rect lights are not being used.

diff --git a/Engine/Shaders/Private/LightGridInjection.usf b/Engine/Shaders/Private/LightGridInjection.usf
--- a/Engine/Shaders/Private/LightGridInjection.usf
+++ b/Engine/Shaders/Private/LightGridInjection.usf
@@ -40,7 +40,10 @@ StructuredBuffer<float4> LightViewSpacePositionAndRadius;
 StructuredBuffer<float4> LightViewSpaceDirAndPreprocAngle;
 
 #define REFINE_SPOTLIGHT_BOUNDS 1
+
+#ifndef REFINE_RECTLIGHT_BOUNDS
 #define REFINE_RECTLIGHT_BOUNDS 1
+#endif
 
 #ifndef USE_LOCALSTORE_FOR_ATOMICS
 #define USE_LOCALSTORE_FOR_ATOMICS 0
diff --git a/Engine/Source/Runtime/Renderer/Private/LightGridInjection.cpp b/Engine/Source/Runtime/Renderer/Private/LightGridInjection.cpp
--- a/Engine/Source/Runtime/Renderer/Private/LightGridInjection.cpp
+++ b/Engine/Source/Runtime/Renderer/Private/LightGridInjection.cpp
@@ -189,7 +189,8 @@ class FLightGridInjectionCS : public FGlobalShader
 	SHADER_USE_PARAMETER_STRUCT(FLightGridInjectionCS, FGlobalShader)
 public:
 	class FUseLinkedListDim : SHADER_PERMUTATION_BOOL("USE_LINKED_CULL_LIST");
-	using FPermutationDomain = TShaderPermutationDomain<FUseLinkedListDim>;
+	class FRefineRectLightBoundsDim : SHADER_PERMUTATION_BOOL("REFINE_RECTLIGHT_BOUNDS");
+	using FPermutationDomain = TShaderPermutationDomain<FUseLinkedListDim, FRefineRectLightBoundsDim>;
 
 	BEGIN_SHADER_PARAMETER_STRUCT(FParameters, )
 		SHADER_PARAMETER_STRUCT_REF(FReflectionCaptureShaderData, ReflectionCapture)
@@ -457,6 +458,7 @@ FComputeLightGridOutput FSceneRenderer::ComputeLightGrid(FRDGBuilder& GraphBuild
 
 	const bool bAllowStaticLighting = IsStaticLightingAllowed();
 	const bool bLightGridUses16BitBuffers = LightGridUses16BitBuffers(ShaderPlatform);
+	const bool bRenderRectLightsAsSpotLights = RenderRectLightsAsSpotLights(FeatureLevel);
 
 	const FRDGSystemTextures& SystemTextures = FRDGSystemTextures::Get(GraphBuilder);
 
@@ -543,7 +545,7 @@ FComputeLightGridOutput FSceneRenderer::ComputeLightGrid(FRDGBuilder& GraphBuild
 				}
 			}
 
-			const uint32 LightShaderParameterFlags = RenderRectLightsAsSpotLights(FeatureLevel) ? ELightShaderParameterFlags::RectAsSpotLight : 0u;
+			const uint32 LightShaderParameterFlags = bRenderRectLightsAsSpotLights ? ELightShaderParameterFlags::RectAsSpotLight : 0u;
 			float SelectedForwardDirectionalLightIntensitySq = 0.0f;
 			int32 SelectedForwardDirectionalLightPriority = -1;
 			const TArray<FSortedLightSceneInfo, SceneRenderingAllocator>& SortedLights = SortedLightSet.SortedLights;
@@ -632,13 +634,15 @@ FComputeLightGridOutput FSceneRenderer::ComputeLightGrid(FRDGBuilder& GraphBuild
 						FVector4f ViewSpacePosAndRadius(LightViewPosition, 1.0f / LightParameters.InvRadius);
 						ViewSpacePosAndRadiusData.Add(ViewSpacePosAndRadius);
 
+						const bool bIsRectLight = !bRenderRectLightsAsSpotLights && LightProxy->IsRectLight();
+
 						// Pack flags in the LSB of PreProcAngle
-						const float PreProcAngle = SortedLightInfo.SortKey.Fields.LightType == LightType_Spot ? GetTanRadAngleOrZero(LightSceneInfo->Proxy->GetOuterConeAngle()) : 0.0f;
-						const uint32 PackedPreProcAngleAndFlags = (FMath::AsUInt(PreProcAngle) & 0xFFFFFFFC) | (LightProxy->HasSourceTexture() ? 0x2 : 0) | (LightProxy->IsRectLight() ? 0x1 : 0);
+						const float PreProcAngle = SortedLightInfo.SortKey.Fields.LightType == LightType_Spot ? GetTanRadAngleOrZero(LightProxy->GetOuterConeAngle()) : 0.0f;
+						const uint32 PackedPreProcAngleAndFlags = (FMath::AsUInt(PreProcAngle) & 0xFFFFFFFC) | (LightProxy->HasSourceTexture() ? 0x2 : 0) | (bIsRectLight ? 0x1 : 0);
 						FVector4f ViewSpaceDirAndPreprocAngleAndFlags(FVector4f(View.ViewMatrices.GetViewMatrix().TransformVector((FVector)LightParameters.Direction)), FMath::AsFloat(PackedPreProcAngleAndFlags)); // LWC_TODO: precision loss
 						ViewSpaceDirAndPreprocAngleData.Add(ViewSpaceDirAndPreprocAngleAndFlags);
 
-						bHasRectLights |= LightProxy->IsRectLight();
+						bHasRectLights |= bIsRectLight;
 						bHasTexturedLights |= LightProxy->HasSourceTexture();
 					}
 					// On mobile there is a separate FMobileDirectionalLightShaderParameters UB which holds all directional light data.
@@ -896,6 +900,7 @@ FComputeLightGridOutput FSceneRenderer::ComputeLightGrid(FRDGBuilder& GraphBuild
 
 			FLightGridInjectionCS::FPermutationDomain PermutationVector;
 			PermutationVector.Set<FLightGridInjectionCS::FUseLinkedListDim>(GLightLinkedListCulling != 0);
+			PermutationVector.Set<FLightGridInjectionCS::FRefineRectLightBoundsDim>(bHasRectLights);
 			TShaderMapRef<FLightGridInjectionCS> ComputeShader(View.ShaderMap, PermutationVector);
 
 			if (GLightLinkedListCulling != 0)
