Implement a faster batched path for translucency lighting volume injection.
- r.TranslucencyLightingVolume.Batch 1 is enabled by default
- For lights that don't require per-light CPU bindings, we put them in a path that loops on the GPU to cut down on draw calls and CPU overhead
- Currently loops over the whole light grid, so the GPU performance is roughly the same as the old path. There are opportunities to accelerate this more in the future, but right now the big win is on the CPU.
- Added a more accurate RectLight integration (to both paths). This is required for rect lights to be batched since the approximate logic is CPU-based. Enabled by default since the quality is superior and the performance difference is minimal in practice even with many rect lights. Can be disabled with r.TranslucencyLightingVolume.AccurateRectLights 0.
- Various refactors to increase code sharing and reduce duplication between the two paths and other lighting code in the engine.

Equivalent or much faster on the CPU. Up to ~3ms in scenes with hundreds of lights.
Roughly equivalent speed on the GPU dependig on content and platform.

diff --git a/Engine/Shaders/Private/DynamicLightingCommon.ush b/Engine/Shaders/Private/DynamicLightingCommon.ush
--- a/Engine/Shaders/Private/DynamicLightingCommon.ush
+++ b/Engine/Shaders/Private/DynamicLightingCommon.ush
@@ -4,6 +4,8 @@
 	DynamicLightingCommon.usf: Contains functions shared by dynamic light shaders.
 =============================================================================*/
 
+#pragma once
+
 #include "LargeWorldCoordinates.ush"
 
 float3 GetDeferredLightTranslatedWorldPosition()
@@ -63,37 +65,6 @@ float SpotAttenuation(float3 L, float3 SpotDirection, float2 SpotAngles)
 	return ConeAngleFalloff;
 }
 
-/** Calculates radial and spot attenuation. */
-float CalcLightAttenuation(float3 TranslatedWorldPosition, out float3 WorldLightVector)
-{
-	WorldLightVector = DeferredLightUniforms.Direction;
-	float DistanceAttenuation = 1;
-
-#if RADIAL_ATTENUATION
-	WorldLightVector = GetDeferredLightTranslatedWorldPosition() - TranslatedWorldPosition;
-	float DistanceSqr = dot( WorldLightVector, WorldLightVector );
-
-	// TODO Line segment falloff
-
-	// Sphere falloff (technically just 1/d2 but this avoids inf)
-	DistanceAttenuation = 1 / ( DistanceSqr + 1 );
-	
-	float LightRadiusMask = Square( saturate( 1 - Square( DistanceSqr * DeferredLightUniforms.InvRadius * DeferredLightUniforms.InvRadius ) ) );
-	DistanceAttenuation *= LightRadiusMask;
-	if (DeferredLightUniforms.FalloffExponent > 0.f)
-	{
-		DistanceAttenuation = RadialAttenuation(WorldLightVector * DeferredLightUniforms.InvRadius, DeferredLightUniforms.FalloffExponent);
-	}
-#endif
-
-	float SpotFalloff = 1;
-	#if RADIAL_ATTENUATION
-		SpotFalloff = SpotAttenuation( normalize(WorldLightVector), -DeferredLightUniforms.Direction, DeferredLightUniforms.SpotAngles);
-	#endif
-
-	return SpotFalloff * DistanceAttenuation;
-}
-
 float3 GetNormalizedLightVector(float3 TranslatedWorldPosition)
 {
 	// assumed to be normalized
diff --git a/Engine/Shaders/Private/TranslucentLightInjectionShaders.usf b/Engine/Shaders/Private/TranslucentLightInjectionShaders.usf
--- a/Engine/Shaders/Private/TranslucentLightInjectionShaders.usf
+++ b/Engine/Shaders/Private/TranslucentLightInjectionShaders.usf
@@ -16,6 +16,9 @@
 
 #include "DynamicLightingCommon.ush"
 
+#define SUPPORT_CONTACT_SHADOWS 0
+#include "DeferredLightingCommon.ush"
+
 uint VolumeCascadeIndex;
 float4 SimpleLightPositionAndRadius;
 float4 SimpleLightColorAndExponent;
@@ -41,6 +44,32 @@ uint AtmospherePerPixelTransmittanceEnabled;
 #include "HeterogeneousVolumes/HeterogeneousVolumesAdaptiveVolumetricShadowMapSampling.ush"
 #endif // ADAPTIVE_VOLUMETRIC_SHADOW_MAP
 
+float3 GetTranslatedWorldPosition(float2 UV, uint LayerIndex)
+{
+	float ZPosition = View.TranslucencyLightingVolumeMin[VolumeCascadeIndex].z + (LayerIndex + .5f) * View.TranslucencyLightingVolumeInvSize[VolumeCascadeIndex].w;
+	float3 TranslatedWorldPosition = float3(View.TranslucencyLightingVolumeMin[VolumeCascadeIndex].xy + UV / View.TranslucencyLightingVolumeInvSize[VolumeCascadeIndex].xy - .5f * View.TranslucencyLightingVolumeInvSize[VolumeCascadeIndex].w, ZPosition);
+	return TranslatedWorldPosition;
+}
+
+void AccumulateSHLighting(
+	float3 NormalizedLightVector,
+	float3 Lighting,
+	float DirectionalLightContribution,
+	inout float4 OutColor0,
+	inout float4 OutColor1)
+{
+	FTwoBandSHVectorRGB SHLighting = MulSH(SHBasisFunction(NormalizedLightVector), Lighting);
+
+	// Directional light contribution in w
+	OutColor0 += float4(SHLighting.R.V.x, SHLighting.G.V.x, SHLighting.B.V.x, DirectionalLightContribution);
+
+	float3 LuminanceWeights = LuminanceFactors();
+	float3 Coefficient0 = float3(SHLighting.R.V.y, SHLighting.G.V.y, SHLighting.B.V.y);
+	float3 Coefficient1 = float3(SHLighting.R.V.z, SHLighting.G.V.z, SHLighting.B.V.z);
+	float3 Coefficient2 = float3(SHLighting.R.V.w, SHLighting.G.V.w, SHLighting.B.V.w);
+	OutColor1 += float4(dot(Coefficient0, LuminanceWeights), dot(Coefficient1, LuminanceWeights), dot(Coefficient2, LuminanceWeights), 0);
+}
+
 float CalcSimpleLightAttenuation(float3 TranslatedWorldPosition)
 {
 	float3 WorldLightVector = SimpleLightPositionAndRadius.xyz - TranslatedWorldPosition;
@@ -75,35 +104,70 @@ void SimpleLightInjectMainPS(
 	OutColor0 = 0;
 	OutColor1 = 0;
 
-	// compute XYZ of the position we shader
-	float3 TranslatedWorldPosition;
-	{
-		float ZPosition = View.TranslucencyLightingVolumeMin[VolumeCascadeIndex].z + (Input.LayerIndex + .5f) * View.TranslucencyLightingVolumeInvSize[VolumeCascadeIndex].w;
-
-		TranslatedWorldPosition = float3(View.TranslucencyLightingVolumeMin[VolumeCascadeIndex].xy + Input.Vertex.UV / View.TranslucencyLightingVolumeInvSize[VolumeCascadeIndex].xy - .5f * View.TranslucencyLightingVolumeInvSize[VolumeCascadeIndex].w, ZPosition);
-	}
-
-	// compute UVW of the position we shade in the volume
-	float3 VolumeUVs = float3(Input.Vertex.UV, (Input.LayerIndex + .5f) * View.TranslucencyLightingVolumeMin[VolumeCascadeIndex].w);
+	float3 TranslatedWorldPosition = GetTranslatedWorldPosition(Input.Vertex.UV, Input.LayerIndex);
 	float3 NormalizedLightVector = normalize(SimpleLightPositionAndRadius.xyz - TranslatedWorldPosition);
+
 	float VoxelSize = View.TranslucencyLightingVolumeInvSize[VolumeCascadeIndex].w;
 	float3 TranslatedWorldPositionForLighting = TranslatedWorldPosition + 1 * GetBoxPushout(NormalizedLightVector, .5f * VoxelSize);
 
 	float Attenuation = CalcSimpleLightAttenuation(TranslatedWorldPositionForLighting);
 	float3 Lighting = SimpleLightColorAndExponent.rgb / PI * Attenuation;
 
-	FTwoBandSHVectorRGB SHLighting = MulSH(SHBasisFunction(NormalizedLightVector), Lighting);
-	OutColor0 = float4(SHLighting.R.V.x, SHLighting.G.V.x, SHLighting.B.V.x, 0);
+	float DirectionalLightContribution = 0.0f;
+	AccumulateSHLighting(NormalizedLightVector, Lighting, DirectionalLightContribution,
+		OutColor0, OutColor1);
+}
 
-	float3 LuminanceWeights = LuminanceFactors();
-	float3 Coefficient0 = float3(SHLighting.R.V.y, SHLighting.G.V.y, SHLighting.B.V.y);
-	float3 Coefficient1 = float3(SHLighting.R.V.z, SHLighting.G.V.z, SHLighting.B.V.z);
-	float3 Coefficient2 = float3(SHLighting.R.V.w, SHLighting.G.V.w, SHLighting.B.V.w);
-	OutColor1 = float4(dot(Coefficient0, LuminanceWeights), dot(Coefficient1, LuminanceWeights), dot(Coefficient2, LuminanceWeights), 0);
+float GetBorderLerpFactor(float2 UV, uint LayerIndex)
+{
+	// Larger values result in a shorter transition distance
+	float TransitionScale = 10;
+
+	float3 VolumeUVs = float3(UV, (LayerIndex + .5f) * View.TranslucencyLightingVolumeMin[VolumeCascadeIndex].w);
+
+	// Rescale the UVs to make the fade go to 0 before the edge of the volume
+	float3 FadeUVs = VolumeUVs * (1 + 4 * View.TranslucencyLightingVolumeMin[VolumeCascadeIndex].w) - 2 * View.TranslucencyLightingVolumeMin[VolumeCascadeIndex].w;
+	// Setup a 3d lerp factor going to 0 near the edge of the outer volume
+	float3 LerpFactors = saturate((.5f - abs(FadeUVs - .5f)) * TransitionScale);
+	float FinalLerpFactor = LerpFactors.x * LerpFactors.y * LerpFactors.z;
+
+	return FinalLerpFactor;
+}
+
+float3 GetTranslatedWorldPositionForLighting(float3 TranslatedWorldPosition, float3 NormalizedLightVector)
+{
+	float VoxelSize = View.TranslucencyLightingVolumeInvSize[VolumeCascadeIndex].w;
+	float3 LightingOffset = 1 * GetBoxPushout(NormalizedLightVector, .5f * VoxelSize);
+	return TranslatedWorldPosition + LightingOffset;
+}
+
+float GetTranslucencyVolumeLocalLightAttenuation(float3 TranslatedWorldPosition, FDeferredLightData LightData)
+{
+	float3 ToLight, L;
+	float Attenuation = GetLocalLightAttenuation(TranslatedWorldPosition, LightData, ToLight, L);
+
+	if (LightData.bRectLight)
+	{
+		if (Attenuation > 0.0f)
+		{
+			// Use a more accurate rect light approximation
+			FRect Rect = GetRect(ToLight, LightData);
+			Attenuation *= IntegrateLight(Rect);
+		}
+	}
+	else if (LightData.bInverseSquared)
+	{
+		// This path still uses the legacy local light attenuation, so add that back in
+		float DistanceSqr = dot(ToLight, ToLight);
+		Attenuation *= 1.0f / (DistanceSqr + 1.0f);
+	}
+
+	return Attenuation;
 }
 
 #if INJECTION_PIXEL_SHADER
 
+#include "LightDataUniforms.ush"
 #include "LightFunctionCommon.ush"
 
 // WorldSpace planes to clip the cascade for ShadoewMethod1
@@ -140,26 +204,18 @@ void InjectMainPS(
 	OutColor0 = 0;
 	OutColor1 = 0;
 
-	// compute XYZ of the position we want to shade
-	float3 TranslatedWorldPosition;
-	{
-		float ZPosition = View.TranslucencyLightingVolumeMin[VolumeCascadeIndex].z + (Input.LayerIndex + .5f) * View.TranslucencyLightingVolumeInvSize[VolumeCascadeIndex].w;
-
-		TranslatedWorldPosition = float3(View.TranslucencyLightingVolumeMin[VolumeCascadeIndex].xy + Input.Vertex.UV / View.TranslucencyLightingVolumeInvSize[VolumeCascadeIndex].xy - .5f * View.TranslucencyLightingVolumeInvSize[VolumeCascadeIndex].w, ZPosition);
-	}
-
-	// compute UVW of the position we shade in the volume
-	float3 VolumeUVs = float3(Input.Vertex.UV, (Input.LayerIndex + .5f) * View.TranslucencyLightingVolumeMin[VolumeCascadeIndex].w);
+	float3 TranslatedWorldPosition = GetTranslatedWorldPosition(Input.Vertex.UV, Input.LayerIndex);
+	
+	FDeferredLightData LightData = InitDeferredLightFromUniforms();
+	float3 LightVector = LightData.Direction;
 
 	// 0: no contribution, 1:full contribution
 	float Masking = 1.0f;
-
-	
 	#if RADIAL_ATTENUATION 
 	{
 		// cull voxels outside the light radius (value is < 0)
-		float3 LightVector = DeferredLightUniforms.TranslatedWorldPosition - TranslatedWorldPosition;
-		clip(1.0f / (DeferredLightUniforms.InvRadius * DeferredLightUniforms.InvRadius) - dot(LightVector, LightVector));
+		LightVector = LightData.TranslatedWorldPosition - TranslatedWorldPosition;
+		clip(1.0f / (LightData.InvRadius * LightData.InvRadius) - dot(LightVector, LightVector));
 	}
 	#else
 	{
@@ -177,32 +233,26 @@ void InjectMainPS(
 	}
 	#endif
 
-	float3 NormalizedLightVector = GetNormalizedLightVector(TranslatedWorldPosition);
-
-	float3 TranslatedWorldPositionForLighting;
-	{
-		float VoxelSize = View.TranslucencyLightingVolumeInvSize[VolumeCascadeIndex].w;
-		float3 LightingOffset = 1 * GetBoxPushout(NormalizedLightVector, .5f * VoxelSize);
-
-		TranslatedWorldPositionForLighting = TranslatedWorldPosition + LightingOffset;
-	}
+	float3 NormalizedLightVector = normalize(LightVector);
+	float3 TranslatedWorldPositionForLighting = GetTranslatedWorldPositionForLighting(TranslatedWorldPosition, NormalizedLightVector);
 
 	{
-		float3 WorldLightVector;
-		// Calculate radial attenuation using the same biased position used for shadowing
-		// Anything else would cause the extents of the shadowmap to not match up with the cone falloff on a spotlight
-		float Attenuation = CalcLightAttenuation(TranslatedWorldPositionForLighting, WorldLightVector);
-
-		bool bPointLight = false;
-		bool bSpotLight = false;
+		float Attenuation = 1.0f;
+		float ShadowFactor = 1.0f;
 
-		#if RADIAL_ATTENUATION 
-			bPointLight = SpotlightMask < 1;
-			bSpotLight = SpotlightMask >= 1;
-		#endif
+#if RADIAL_ATTENUATION
+		Attenuation *= GetTranslucencyVolumeLocalLightAttenuation(TranslatedWorldPositionForLighting, LightData);
+#endif
 
-		bool bUnused = false;
-		float ShadowFactor = ComputeVolumeShadowing(TranslatedWorldPositionForLighting, bPointLight, bSpotLight, bUnused);
+		{
+			bool bUnused = false;
+			ShadowFactor *= ComputeVolumeShadowing(
+				TranslatedWorldPositionForLighting,
+				// Maintain the legacy logic of "spotlight" vs "everything else"
+				!LightData.bSpotLight,
+				LightData.bSpotLight,
+				bUnused);
+		}
 
 #if VIRTUAL_SHADOW_MAP
 		FVirtualShadowMapSampleResult VirtualShadowMapSample = SampleVirtualShadowMapTranslatedWorld(VirtualShadowMapId, TranslatedWorldPositionForLighting);
@@ -210,7 +260,7 @@ void InjectMainPS(
 #endif
 
 #if ADAPTIVE_VOLUMETRIC_SHADOW_MAP
-		ShadowFactor *= AVSM_SampleTransmittance(TranslatedWorldPositionForLighting, DeferredLightUniforms.TranslatedWorldPosition);
+		ShadowFactor *= AVSM_SampleTransmittance(TranslatedWorldPositionForLighting, LightData.TranslatedWorldPosition);
 #endif // ADAPTIVE_VOLUMETRIC_SHADOW_MAP
 
 		// Apply light function (it will also fade out at the edge of the volume so overall dimming of light using light function is not a recommended workflow)
@@ -227,13 +277,7 @@ void InjectMainPS(
 		
 		if (VolumeCascadeIndex == 1)
 		{
-			// Larger values result in a shorter transition distance
-			float TransitionScale = 10;
-			// Rescale the UVs to make the fade go to 0 before the edge of the volume
-			float3 FadeUVs = VolumeUVs * (1 + 4 * View.TranslucencyLightingVolumeMin[VolumeCascadeIndex].w) - 2 * View.TranslucencyLightingVolumeMin[VolumeCascadeIndex].w;
-			// Setup a 3d lerp factor going to 0 near the edge of the outer volume
-			float3 LerpFactors = saturate((.5f - abs(FadeUVs - .5f)) * TransitionScale);
-			float FinalLerpFactor = LerpFactors.x * LerpFactors.y * LerpFactors.z;
+			float FinalLerpFactor = GetBorderLerpFactor(Input.Vertex.UV, Input.LayerIndex);
 
 			#if RADIAL_ATTENUATION 
 			// For local lights, fade attenuation to 0 for the border voxels
@@ -246,7 +290,7 @@ void InjectMainPS(
 			#endif
 		}
 
-		float3 Lighting = DeferredLightUniforms.Color / PI * Attenuation * ShadowFactor;
+		float3 Lighting = LightData.Color / PI * Attenuation * ShadowFactor;
 
 		#if RADIAL_ATTENUATION==0
 		// Apply color atmosphere transmittance for atmosphere directional lights only if ndeeded
@@ -257,22 +301,12 @@ void InjectMainPS(
 		}
 		#endif
 
-		FTwoBandSHVectorRGB SHLighting = MulSH(SHBasisFunction(NormalizedLightVector), Lighting);
-
 		float DirectionalLightContribution = 0;
-
-		#if !RADIAL_ATTENUATION 
+#if !RADIAL_ATTENUATION 
 		DirectionalLightContribution = Attenuation * ShadowFactor;
-		#endif
-
-		// Directional light contribution in w
-		OutColor0 = float4(SHLighting.R.V.x, SHLighting.G.V.x, SHLighting.B.V.x, DirectionalLightContribution);
-
-		float3 LuminanceWeights = LuminanceFactors();
-		float3 Coefficient0 = float3(SHLighting.R.V.y, SHLighting.G.V.y, SHLighting.B.V.y);
-		float3 Coefficient1 = float3(SHLighting.R.V.z, SHLighting.G.V.z, SHLighting.B.V.z);
-		float3 Coefficient2 = float3(SHLighting.R.V.w, SHLighting.G.V.w, SHLighting.B.V.w);
-		OutColor1 = float4(dot(Coefficient0, LuminanceWeights), dot(Coefficient1, LuminanceWeights), dot(Coefficient2, LuminanceWeights), 0);
+#endif
+		AccumulateSHLighting(NormalizedLightVector, Lighting, DirectionalLightContribution,
+			OutColor0, OutColor1);
 	}
 
 	// debug, make inner cascase green
@@ -284,6 +318,85 @@ void InjectMainPS(
 
 #endif // #if INJECTION_PIXEL_SHADER
 
+#if INJECTION_BATCH_PIXEL_SHADER
+
+#include "LightGridCommon.ush"
+#include "LightData.ush"
+
+//int VirtualShadowMapId;
+StructuredBuffer<uint> BatchedLocalLights;
+uint MaxBatchedLocalLights;
+
+/** Pixel shader that calculates direct lighting for all the affected voxels of a volume texture. */
+void InjectBatchMainPS(
+	FWriteToSliceGeometryOutput Input,
+	out float4 RTOutColor0 : SV_Target0,
+	out float4 RTOutColor1 : SV_Target1
+)
+{
+	float4 OutColor0 = 0;
+	float4 OutColor1 = 0;
+
+	float3 TranslatedWorldPosition = GetTranslatedWorldPosition(Input.Vertex.UV, Input.LayerIndex);
+	
+	const uint NumUints = (MaxBatchedLocalLights + 31U) / 32U;
+	for (uint UintIndex = 0; UintIndex < NumUints; ++UintIndex)
+	{
+		uint CurrentMask = BatchedLocalLights[UintIndex];
+		while (CurrentMask != 0)
+		{
+			int BitIndex = firstbitlow(CurrentMask);
+			CurrentMask = CurrentMask & ~(1U << BitIndex);
+			uint LocalLightIndex = 32U * UintIndex + BitIndex;
+			
+			{
+				FLocalLightData LocalLightData = GetLocalLightDataNonStereo(LocalLightIndex);
+				FDeferredLightData LightData = ConvertToDeferredLight(LocalLightData);
+
+				float3 LightVector = LightData.TranslatedWorldPosition - TranslatedWorldPosition;
+				if (rcp(Square(LightData.InvRadius)) > dot(LightVector, LightVector))
+				{
+					float3 NormalizedLightVector = normalize(LightVector);
+					float3 TranslatedWorldPositionForLighting = GetTranslatedWorldPositionForLighting(TranslatedWorldPosition, NormalizedLightVector);
+
+					float Attenuation = GetTranslucencyVolumeLocalLightAttenuation(TranslatedWorldPositionForLighting, LightData);
+					float ShadowFactor = 1.0f;
+
+					if (Attenuation > 0.0f)
+					{
+#if VIRTUAL_SHADOW_MAP
+						if (LocalLightData.VirtualShadowMapId != INDEX_NONE)
+						{
+							FVirtualShadowMapSampleResult VirtualShadowMapSample = SampleVirtualShadowMapTranslatedWorld(LocalLightData.VirtualShadowMapId, TranslatedWorldPositionForLighting);
+							ShadowFactor *= VirtualShadowMapSample.ShadowFactor;
+						}
+#endif
+						if (VolumeCascadeIndex == 1)
+						{
+							float FinalLerpFactor = GetBorderLerpFactor(Input.Vertex.UV, Input.LayerIndex);
+							// For local lights, fade attenuation to 0 for the border voxels
+							Attenuation = lerp(0, Attenuation, FinalLerpFactor);
+							ShadowFactor = lerp(0.0f, ShadowFactor, FinalLerpFactor);
+						}
+
+						float3 Lighting = LightData.Color / PI * Attenuation * ShadowFactor;
+
+						float DirectionalLightContribution = 0;
+						AccumulateSHLighting(NormalizedLightVector, Lighting, DirectionalLightContribution,
+							OutColor0, OutColor1);
+					}
+				}
+			}
+		}
+	}
+
+	RTOutColor0 = OutColor0;
+	RTOutColor1 = OutColor1;
+}
+
+#endif // #if INJECTION_BATCH_PIXEL_SHADER
+
+
 #if CLEAR_COMPUTE_SHADER
 RWTexture3D<float4> RWAmbient0;
 RWTexture3D<float4> RWDirectional0;
diff --git a/Engine/Source/Runtime/RenderCore/Private/RenderCore.cpp b/Engine/Source/Runtime/RenderCore/Private/RenderCore.cpp
--- a/Engine/Source/Runtime/RenderCore/Private/RenderCore.cpp
+++ b/Engine/Source/Runtime/RenderCore/Private/RenderCore.cpp
@@ -172,6 +172,7 @@ DEFINE_STAT(STAT_NumShadowedLights);
 DEFINE_STAT(STAT_NumLightFunctionOnlyLights);
 DEFINE_STAT(STAT_NumBatchedLights);
 DEFINE_STAT(STAT_NumLightsInjectedIntoTranslucency);
+DEFINE_STAT(STAT_NumLightsInjectedIntoTranslucencyBatched);
 DEFINE_STAT(STAT_NumLightsUsingStandardDeferred);
 
 DEFINE_STAT(STAT_LightShaftsLights);
diff --git a/Engine/Source/Runtime/RenderCore/Public/RenderCore.h b/Engine/Source/Runtime/RenderCore/Public/RenderCore.h
--- a/Engine/Source/Runtime/RenderCore/Public/RenderCore.h
+++ b/Engine/Source/Runtime/RenderCore/Public/RenderCore.h
@@ -121,7 +121,8 @@ DECLARE_CYCLE_STAT_EXTERN(TEXT("Render Lights and Shadows"),STAT_LightRendering,
 DECLARE_CYCLE_STAT_EXTERN(TEXT("Direct lighting"),STAT_DirectLightRenderingTime,STATGROUP_LightRendering, RENDERCORE_API);
 DECLARE_CYCLE_STAT_EXTERN(TEXT("Translucent injection"),STAT_TranslucentInjectTime,STATGROUP_LightRendering, RENDERCORE_API);
 DECLARE_DWORD_COUNTER_STAT_EXTERN(TEXT("Num using standard deferred"),STAT_NumLightsUsingStandardDeferred,STATGROUP_LightRendering, RENDERCORE_API);
-DECLARE_DWORD_COUNTER_STAT_EXTERN(TEXT("Num injected into translucency"),STAT_NumLightsInjectedIntoTranslucency,STATGROUP_LightRendering, RENDERCORE_API);
+DECLARE_DWORD_COUNTER_STAT_EXTERN(TEXT("Num injected into translucency (unbatched)"),STAT_NumLightsInjectedIntoTranslucency,STATGROUP_LightRendering, RENDERCORE_API);
+DECLARE_DWORD_COUNTER_STAT_EXTERN(TEXT("Num injected into translucency (batched)"), STAT_NumLightsInjectedIntoTranslucencyBatched, STATGROUP_LightRendering, RENDERCORE_API);
 DECLARE_DWORD_COUNTER_STAT_EXTERN(TEXT("Num batched"),STAT_NumBatchedLights,STATGROUP_LightRendering, RENDERCORE_API);
 DECLARE_DWORD_COUNTER_STAT_EXTERN(TEXT("Num light function only"),STAT_NumLightFunctionOnlyLights,STATGROUP_LightRendering, RENDERCORE_API);
 DECLARE_DWORD_COUNTER_STAT_EXTERN(TEXT("Num shadowed"),STAT_NumShadowedLights,STATGROUP_LightRendering, RENDERCORE_API);
diff --git a/Engine/Source/Runtime/Renderer/Private/LightRendering.cpp b/Engine/Source/Runtime/Renderer/Private/LightRendering.cpp
--- a/Engine/Source/Runtime/Renderer/Private/LightRendering.cpp
+++ b/Engine/Source/Runtime/Renderer/Private/LightRendering.cpp
@@ -1409,6 +1409,7 @@ void FDeferredShadingSceneRenderer::RenderLights(
 		if (ViewFamily.EngineShowFlags.DirectLighting && GUseTranslucentLightingVolumes && RHISupportsVolumeTextureRendering())
 		{
 			RDG_EVENT_SCOPE(GraphBuilder, "InjectTranslucencyLightingVolume");
+			SCOPE_CYCLE_COUNTER(STAT_TranslucentInjectTime);
 
 			if (SimpleLights.InstanceData.Num() > 0)
 			{
diff --git a/Engine/Source/Runtime/Renderer/Private/TranslucentLighting.cpp b/Engine/Source/Runtime/Renderer/Private/TranslucentLighting.cpp
--- a/Engine/Source/Runtime/Renderer/Private/TranslucentLighting.cpp
+++ b/Engine/Source/Runtime/Renderer/Private/TranslucentLighting.cpp
@@ -52,6 +52,8 @@
 #include "StaticMeshBatch.h"
 #include "LightFunctionAtlas.h"
 #include "HeterogeneousVolumes/HeterogeneousVolumes.h"
+#include "DeferredShadingRenderer.h"
+#include "BasePassRendering.h"
 
 class FMaterial;
 
@@ -114,6 +116,19 @@ static TAutoConsoleVariable<float> CVarTranslucencyLightingVolumeOuterDistance(
 	TEXT("Distance from the camera that the second volume cascade should end"),
 	ECVF_RenderThreadSafe | ECVF_Scalability);
 
+static TAutoConsoleVariable<int32> CVarTranslucencyLightingVolumeBatch(
+	TEXT("r.TranslucencyLightingVolume.Batch"),
+	1,
+	TEXT("When enabled, batches supported lights into a single draw call for efficiency"),
+	ECVF_RenderThreadSafe | ECVF_Scalability);
+
+static TAutoConsoleVariable<int32> CVarTranslucencyLightingVolumeAccurateRectLights(
+	TEXT("r.TranslucencyLightingVolume.AccurateRectLights"),
+	1,
+	TEXT("When disabled rect lights are approximated as spot lights in the translucency volume.\n")
+	TEXT("Only accurate rect lights are included in batching, so it is recommended to enable this when batching is used."),
+	ECVF_RenderThreadSafe | ECVF_Scalability);
+
 /** Function returning current translucency lighting volume dimensions. */
 int32 GetTranslucencyLightingVolumeDim()
 {
@@ -752,6 +767,54 @@ public:
 
 IMPLEMENT_MATERIAL_SHADER_TYPE(,FTranslucentLightingInjectPS, TEXT("/Engine/Private/TranslucentLightInjectionShaders.usf"), TEXT("InjectMainPS"), SF_Pixel);
 
+/** Shader that adds direct lighting contribution from multiple lights. */
+class FTranslucentLightingInjectBatchPS : public FGlobalShader
+{
+	DECLARE_GLOBAL_SHADER(FTranslucentLightingInjectBatchPS);
+	SHADER_USE_PARAMETER_STRUCT(FTranslucentLightingInjectBatchPS, FGlobalShader);
+
+	BEGIN_SHADER_PARAMETER_STRUCT(FParameters, )
+		SHADER_PARAMETER_STRUCT_REF(FViewUniformShaderParameters, ViewUniformBuffer)
+		SHADER_PARAMETER_RDG_UNIFORM_BUFFER(FForwardLightData, ForwardLightData)
+		SHADER_PARAMETER_STRUCT_INCLUDE(FVirtualShadowMapSamplingParameters, VirtualShadowMapSamplingParameters)
+		SHADER_PARAMETER_RDG_BUFFER_SRV(StructuredBuffer<uint>, BatchedLocalLights)
+		SHADER_PARAMETER(uint32, MaxBatchedLocalLights)
+		SHADER_PARAMETER(uint32, VolumeCascadeIndex)
+		RENDER_TARGET_BINDING_SLOTS()
+	END_SHADER_PARAMETER_STRUCT()
+
+	class FVirtualShadowMap : SHADER_PERMUTATION_BOOL("VIRTUAL_SHADOW_MAP");
+	using FPermutationDomain = TShaderPermutationDomain<FVirtualShadowMap>;
+
+public:
+
+	static void ModifyCompilationEnvironment(const FGlobalShaderPermutationParameters& Parameters, FShaderCompilerEnvironment& OutEnvironment)
+	{
+		FGlobalShader::ModifyCompilationEnvironment(Parameters, OutEnvironment);
+		FVirtualShadowMapArray::SetShaderDefines(OutEnvironment);
+		FForwardLightingParameters::ModifyCompilationEnvironment(Parameters.Platform, OutEnvironment);
+
+		OutEnvironment.SetDefine(TEXT("INJECTION_BATCH_PIXEL_SHADER"), 1);
+	}
+
+	static bool ShouldCompilePermutation(const FGlobalShaderPermutationParameters& Parameters)
+	{
+		FPermutationDomain PermutationVector(Parameters.PermutationId);
+
+		if (!DoesPlatformSupportVirtualShadowMaps(Parameters.Platform) && PermutationVector.Get<FVirtualShadowMap>() != 0)
+		{
+			return false;
+		}
+
+		return (IsFeatureLevelSupported(Parameters.Platform, ERHIFeatureLevel::SM5) &&
+			(RHISupportsGeometryShaders(Parameters.Platform) || RHISupportsVertexShaderLayer(Parameters.Platform)));
+	}
+};
+
+IMPLEMENT_MATERIAL_SHADER_TYPE(,FTranslucentLightingInjectBatchPS, TEXT("/Engine/Private/TranslucentLightInjectionShaders.usf"), TEXT("InjectBatchMainPS"), SF_Pixel);
+
+
+
 class FClearTranslucentLightingVolumeCS : public FGlobalShader
 {
 public:
@@ -1018,7 +1081,7 @@ FTranslucentLightInjectionCollector::FTranslucentLightInjectionCollector(
 	FRDGBuilder& GraphBuilder,
 	TArrayView<const FViewInfo> Views)
 	// NOTE: This data is directly referenced inside the render pass lamba, so must be allocated in the graph
-	: InjectionDataPerView(*GraphBuilder.AllocObject<TArray<FInjectionDataArray, SceneRenderingAllocator>>())
+	: InjectionDataPerView(*GraphBuilder.AllocObject<TArray<FPerViewData, SceneRenderingAllocator>>())
 {
 	InjectionDataPerView.SetNum(Views.Num());
 }
@@ -1036,6 +1099,20 @@ void FTranslucentLightInjectionCollector::AddLightForInjection(
 {
 	if (LightSceneInfo.Proxy->AffectsTranslucentLighting())
 	{
+		const uint8 LightType = LightSceneInfo.Proxy->GetLightType();
+
+		FVolumeBounds VolumeBounds[TVC_MAX];
+		bool bAnyBoundsValid = false;
+		for (uint32 VolumeCascadeIndex = 0; VolumeCascadeIndex < TVC_MAX; ++VolumeCascadeIndex)
+		{
+			VolumeBounds[VolumeCascadeIndex] = CalculateLightVolumeBounds(LightSceneInfo.Proxy->GetBoundingSphere(), View, VolumeCascadeIndex, LightType == LightType_Directional);
+			bAnyBoundsValid = bAnyBoundsValid || VolumeBounds->IsValid();
+		}
+		if (!bAnyBoundsValid)
+		{
+			return;
+		}
+
 		const FVisibleLightInfo& VisibleLightInfo = VisibleLightInfos[LightSceneInfo.Id];
 		const ERHIFeatureLevel::Type FeatureLevel = View.FeatureLevel;
 
@@ -1043,19 +1120,87 @@ void FTranslucentLightInjectionCollector::AddLightForInjection(
 			LightSceneInfo.Proxy->GetLightFunctionMaterial() && 
 			LightSceneInfo.Proxy->GetLightFunctionMaterial()->GetIncompleteMaterialWithFallback(FeatureLevel).IsLightFunction());
 
-		const FMaterialRenderProxy* MaterialProxy = bApplyLightFunction ? 
-			LightSceneInfo.Proxy->GetLightFunctionMaterial() : 
-			UMaterial::GetDefaultMaterial(MD_LightFunction)->GetRenderProxy();
+		const int32 VirtualShadowMapId = VisibleLightInfo.GetVirtualShadowMapId(&View);
+		
+		const bool bUseAdaptiveVolumetricShadowMap =
+			LightSceneInfo.Proxy->CastsVolumetricShadow() &&
+			ShouldRenderHeterogeneousVolumesForView(View) &&
+			ShouldHeterogeneousVolumesCastShadows();
+
+		const bool bStaticShadowing =
+			LightSceneInfo.Proxy->GetStaticShadowDepthMap() &&
+			LightSceneInfo.Proxy->GetStaticShadowDepthMap()->Data;
+
+		// We only support the "accurate" rect light model through this path since the LightGrid stores them that way
+		// Thus if the spot light approximation is used we have to send them through the unbatched path.
+		const bool bSupportRectLights = CVarTranslucencyLightingVolumeAccurateRectLights.GetValueOnRenderThread() != 0;
+
+		// Lights without certain features can be batched into a single draw (loop in shader) which is more efficient
+		bool bSupportsBatching =
+			CVarTranslucencyLightingVolumeBatch.GetValueOnRenderThread() != 0 &&
+			LightType != LightType_Directional &&
+			(bSupportRectLights || LightType != LightType_Rect) &&
+			!bStaticShadowing &&
+			InProjectedShadowInfo == nullptr &&
+			!bApplyLightFunction &&												// TODO: Support LF atlas
+			!bUseAdaptiveVolumetricShadowMap;
+
+		// If it would otherwise be supported, see if we can find the local light index in the light grid
+		if (bSupportsBatching && View.ForwardLightingResources.ForwardLightUniformBuffer && View.ViewState)
+		{
+			const int32* Value = View.ViewState->LightSceneIdToLocalLightIndex.Find(LightSceneInfo.Id);
+			if (Value)
+			{
+				const int32 LocalLightIndex = *Value;
+
+				FPerViewData& InjectionData = InjectionDataPerView[ViewIndex];
+				// There shouldn't be any duplication here
+				++InjectionData.BatchedLocalLightCount;
 
-		// Skip rendering if the DefaultLightFunctionMaterial isn't compiled yet
-		if (MaterialProxy->GetIncompleteMaterialWithFallback(FeatureLevel).IsLightFunction())
+				for (uint32 VolumeCascadeIndex = 0; VolumeCascadeIndex < TVC_MAX; ++VolumeCascadeIndex)
+				{
+					if (VolumeBounds[VolumeCascadeIndex].IsValid())
+					{
+						TBitArray<SceneRenderingAllocator>& Bits = InjectionData.BatchedLocalLights[VolumeCascadeIndex];
+						Bits.PadToNum(LocalLightIndex + 1, false);
+						Bits[LocalLightIndex] = true;
+					}
+				}
+
+				if (VirtualShadowMapId != INDEX_NONE)
+				{
+					// Note if there are any batched lights with VSM for permutation selection later
+					InjectionData.bAnyBatchedLightsWithVirtualShadowMaps = true;
+				}
+			}
+			else
+			{
+				bSupportsBatching = false;
+			}
+		}
+		
+		if (!bSupportsBatching)
 		{
-			FTranslucentLightInjectionCollector::FInjectionData Data;
-			Data.LightSceneInfo = &LightSceneInfo;
-			Data.ProjectedShadowInfo = InProjectedShadowInfo;
-			Data.bApplyLightFunction = bApplyLightFunction;
-			Data.LightFunctionMaterialProxy = MaterialProxy;
-			InjectionDataPerView[ViewIndex].Add(Data);
+			// Skip rendering if the DefaultLightFunctionMaterial isn't compiled yet
+			const FMaterialRenderProxy* MaterialProxy = bApplyLightFunction ?
+				LightSceneInfo.Proxy->GetLightFunctionMaterial() :
+				UMaterial::GetDefaultMaterial(MD_LightFunction)->GetRenderProxy();
+
+			if (MaterialProxy->GetIncompleteMaterialWithFallback(FeatureLevel).IsLightFunction())
+			{
+				FTranslucentLightInjectionCollector::FInjectionData Data;
+				Data.LightSceneInfo = &LightSceneInfo;
+				Data.ProjectedShadowInfo = InProjectedShadowInfo;
+				Data.bApplyLightFunction = bApplyLightFunction;
+				Data.LightFunctionMaterialProxy = MaterialProxy;
+				Data.VirtualShadowMapId = VirtualShadowMapId;
+				Data.bUseAdaptiveVolumetricShadowMap = bUseAdaptiveVolumetricShadowMap;
+				for (uint32 VolumeCascadeIndex = 0; VolumeCascadeIndex < TVC_MAX; ++VolumeCascadeIndex)
+				{
+					Data.VolumeBounds[VolumeCascadeIndex] = VolumeBounds[VolumeCascadeIndex];
+				}
+				InjectionDataPerView[ViewIndex].Unbatched.Add(Data);
+			}
 		}
 	}
 }
@@ -1073,6 +1218,110 @@ static FRDGTextureRef GetSkyTransmittanceLutTexture(FRDGBuilder& GraphBuilder, c
 	return TransmittanceLutTexture;
 }
 
+
+static void SetPSOStateForVolumeInjection(
+	TShaderMapRef<FWriteToSliceVS> VertexShader,
+	TOptionalShaderMapRef<FWriteToSliceGS> GeometryShader,
+	FRHIPixelShader* PixelShaderRHI,
+	FVolumeBounds VolumeBounds,
+	bool bDirectionalLight,
+	FRHICommandList& RHICmdList,
+	FGraphicsPipelineStateInitializer &OutGraphicsPSOInit)
+{
+	RHICmdList.ApplyCachedRenderTargets(OutGraphicsPSOInit);
+
+	OutGraphicsPSOInit.RasterizerState = TStaticRasterizerState<FM_Solid, CM_None>::GetRHI();
+	OutGraphicsPSOInit.DepthStencilState = TStaticDepthStencilState<false, CF_Always>::GetRHI();
+	OutGraphicsPSOInit.PrimitiveType = PT_TriangleStrip;
+
+	// Accumulate the contribution of multiple lights
+	if (bDirectionalLight)
+	{		
+		// Directional lights write their shadowing into alpha of the ambient texture
+		OutGraphicsPSOInit.BlendState = TStaticBlendState<
+			CW_RGBA, BO_Add, BF_One, BF_One, BO_Add, BF_One, BF_One,
+			CW_RGB, BO_Add, BF_One, BF_One, BO_Add, BF_One, BF_One>::GetRHI();
+	}
+	else
+	{
+		OutGraphicsPSOInit.BlendState = TStaticBlendState<
+			CW_RGB, BO_Add, BF_One, BF_One, BO_Add, BF_Zero, BF_One,
+			CW_RGB, BO_Add, BF_One, BF_One, BO_Add, BF_Zero, BF_One>::GetRHI();
+	}
+
+	OutGraphicsPSOInit.BoundShaderState.VertexDeclarationRHI = GScreenVertexDeclaration.VertexDeclarationRHI;
+	OutGraphicsPSOInit.BoundShaderState.VertexShaderRHI = VertexShader.GetVertexShader();
+	OutGraphicsPSOInit.BoundShaderState.SetGeometryShader(GeometryShader.GetGeometryShader());
+	OutGraphicsPSOInit.BoundShaderState.PixelShaderRHI = PixelShaderRHI;
+	SetGraphicsPipelineState(RHICmdList, OutGraphicsPSOInit, 0);
+
+	const int32 TranslucencyLightingVolumeDim = GetTranslucencyLightingVolumeDim();
+
+	SetShaderParametersLegacyVS(RHICmdList, VertexShader, VolumeBounds, FIntVector(TranslucencyLightingVolumeDim));
+	if (GeometryShader.IsValid())
+	{
+		SetShaderParametersLegacyGS(RHICmdList, GeometryShader, VolumeBounds.MinZ);
+	}
+}
+
+
+static void InjectTranslucencyLightingVolumeBatch(
+	FRDGBuilder& GraphBuilder,
+	const FViewInfo& View,
+	const uint32 ViewIndex,
+	uint32 MaxBatchedLocalLights,
+	FRDGBufferRef BatchedLocalLightsRDG,
+	bool bSupportVirtualShadowMaps,
+	uint32 VolumeCascadeIndex,
+	FRDGTextureRef VolumeAmbientTexture,
+	FRDGTextureRef VolumeDirectionalTexture,
+	const FSceneRenderer& Renderer)
+{
+	check(MaxBatchedLocalLights > 0);
+	check(BatchedLocalLightsRDG);
+
+	TShaderMapRef<FWriteToSliceVS> VertexShader(View.ShaderMap);
+	TOptionalShaderMapRef<FWriteToSliceGS> GeometryShader(View.ShaderMap);
+
+	FTranslucentLightingInjectBatchPS::FPermutationDomain PermutationVector;
+	PermutationVector.Set< FTranslucentLightingInjectBatchPS::FVirtualShadowMap >(bSupportVirtualShadowMaps);
+	TShaderMapRef<FTranslucentLightingInjectBatchPS> PixelShader(View.ShaderMap, PermutationVector);
+
+	auto* PassParameters = GraphBuilder.AllocParameters<FTranslucentLightingInjectBatchPS::FParameters>();
+	if (bSupportVirtualShadowMaps)
+	{
+		PassParameters->VirtualShadowMapSamplingParameters = Renderer.VirtualShadowMapArray.GetSamplingParameters(GraphBuilder);
+	}
+	PassParameters->ViewUniformBuffer = View.ViewUniformBuffer;
+	PassParameters->ForwardLightData = View.ForwardLightingResources.ForwardLightUniformBuffer;
+	PassParameters->BatchedLocalLights = GraphBuilder.CreateSRV(BatchedLocalLightsRDG);
+	PassParameters->MaxBatchedLocalLights = MaxBatchedLocalLights;
+	PassParameters->VolumeCascadeIndex = VolumeCascadeIndex;
+
+	PassParameters->RenderTargets[0] = FRenderTargetBinding(VolumeAmbientTexture, ERenderTargetLoadAction::ELoad);
+	PassParameters->RenderTargets[1] = FRenderTargetBinding(VolumeDirectionalTexture, ERenderTargetLoadAction::ELoad);
+			
+	GraphBuilder.AddPass(
+		RDG_EVENT_NAME("InjectTranslucencyLightingVolumeBatch(VolumeCascade=%d,Max=%d%s)",
+			VolumeCascadeIndex, MaxBatchedLocalLights,
+			bSupportVirtualShadowMaps ? TEXT(",VirtualShadowMap") : TEXT("")),
+		PassParameters,
+		ERDGPassFlags::Raster,
+		[PassParameters, VertexShader, GeometryShader, PixelShader, VolumeCascadeIndex](FRHICommandList& RHICmdList)
+	{
+		FGraphicsPipelineStateInitializer GraphicsPSOInit;
+
+		const FVolumeBounds VolumeBounds = FVolumeBounds(GetTranslucencyLightingVolumeDim());
+		SetPSOStateForVolumeInjection(VertexShader, GeometryShader, PixelShader.GetPixelShader(),
+			VolumeBounds, false, RHICmdList, GraphicsPSOInit);
+				
+		SetShaderParameters(RHICmdList, PixelShader, PixelShader.GetPixelShader(), *PassParameters);
+
+		RasterizeToVolumeTexture(RHICmdList, VolumeBounds);
+	});
+}
+
+
 BEGIN_SHADER_PARAMETER_STRUCT(FInjectTranslucentLightArrayParameters, )
 	SHADER_PARAMETER_STRUCT_INCLUDE(FTranslucentLightingInjectPS::FParameters, PS)
 	SHADER_PARAMETER_STRUCT_INCLUDE(FVolumetricCloudShadowAOParameters, CloudShadowAO)
@@ -1096,16 +1345,20 @@ void InjectTranslucencyLightingVolume(
 		return;
 	}
 
-	const FTranslucentLightInjectionCollector::FInjectionDataArray& LightInjectionData = Collector.InjectionDataPerView[ViewIndex];
+	const FTranslucentLightInjectionCollector::FPerViewData& LightInjectionData = Collector.InjectionDataPerView[ViewIndex];
 
-	SCOPE_CYCLE_COUNTER(STAT_TranslucentInjectTime);
-	INC_DWORD_STAT_BY(STAT_NumLightsInjectedIntoTranslucency, LightInjectionData.Num());
+	INC_DWORD_STAT_BY(STAT_NumLightsInjectedIntoTranslucencyBatched, LightInjectionData.BatchedLocalLightCount);
+	INC_DWORD_STAT_BY(STAT_NumLightsInjectedIntoTranslucency, LightInjectionData.Unbatched.Num());
 
 	const FVolumetricCloudShadowAOParameters CloudShadowAOParameters = GetCloudShadowAOParameters(GraphBuilder, View, Scene->GetVolumetricCloudSceneInfo());
 	const bool bUseLightFunctionAtlas = View.LightFunctionAtlasViewData.UsesLightFunctionAtlas(LightFunctionAtlas::ELightFunctionAtlasSystem::DeferredLighting);
 
 	FRDGTextureRef TransmittanceLutTexture = GetSkyTransmittanceLutTexture(GraphBuilder, Scene, View);
 
+	// When accurate rect lights is disabled we approximate rect lights as spotlights
+	uint32 DeferredLightParameterFlags = CVarTranslucencyLightingVolumeAccurateRectLights.GetValueOnRenderThread() != 0 ?
+		0U : ELightShaderParameterFlags::RectAsSpotLight;
+
 	// Inject into each volume cascade. Operate on one cascade at a time to reduce render target switches.
 	for (uint32 VolumeCascadeIndex = 0; VolumeCascadeIndex < TVC_MAX; VolumeCascadeIndex++)
 	{
@@ -1114,17 +1367,40 @@ void InjectTranslucencyLightingVolume(
 		FRDGTextureRef VolumeAmbientTexture = Textures.Ambient[TextureIndex];
 		FRDGTextureRef VolumeDirectionalTexture = Textures.Directional[TextureIndex];
 
-		for (int32 LightIndex = 0; LightIndex < LightInjectionData.Num(); LightIndex++)
+		// Batched lights
 		{
-			const FTranslucentLightInjectionCollector::FInjectionData& InjectionData = LightInjectionData[LightIndex];
+			const TBitArray<SceneRenderingAllocator>& BatchedLocalLights = Collector.InjectionDataPerView[ViewIndex].BatchedLocalLights[VolumeCascadeIndex];
+			FRDGBufferRef BatchedLocalLightsRDG = nullptr;
+			if (BatchedLocalLights.Num() > 0)
+			{
+				const uint32 NumElements = FMath::Max(1, FMath::DivideAndRoundUp(BatchedLocalLights.Num(), 32));
+				BatchedLocalLightsRDG = CreateStructuredBuffer(
+					GraphBuilder,
+					TEXT("TranslucencyLightingVolume.BatchedLocalLights"),
+					sizeof(uint32),
+					NumElements,
+					BatchedLocalLights.GetData(),
+					FMath::DivideAndRoundUp(BatchedLocalLights.Num(), 8)	// Size in bytes of initial data
+				);
+
+				InjectTranslucencyLightingVolumeBatch(GraphBuilder, View, ViewIndex,
+					BatchedLocalLights.Num(), BatchedLocalLightsRDG,
+					Collector.InjectionDataPerView[ViewIndex].bAnyBatchedLightsWithVirtualShadowMaps,
+					VolumeCascadeIndex, VolumeAmbientTexture, VolumeDirectionalTexture,
+					Renderer);
+			}
+		}
+
+		// Unbatched lights
+		for (int32 LightIndex = 0; LightIndex < LightInjectionData.Unbatched.Num(); LightIndex++)
+		{
+			const FTranslucentLightInjectionCollector::FInjectionData& InjectionData = LightInjectionData.Unbatched[LightIndex];
 			const FLightSceneInfo* const LightSceneInfo = InjectionData.LightSceneInfo;
 			const FVisibleLightInfo& VisibleLightInfo = Renderer.VisibleLightInfos[LightSceneInfo->Id];
 			const bool bInverseSquared = LightSceneInfo->Proxy->IsInverseSquared();
 			const bool bDirectionalLight = LightSceneInfo->Proxy->GetLightType() == LightType_Directional;
-			bool bUseVSM = Renderer.VirtualShadowMapArray.IsAllocated();
-			const bool bUseAdaptiveVolumetricShadowMap = LightSceneInfo->Proxy->CastsVolumetricShadow() && ShouldRenderHeterogeneousVolumes(Scene) && ShouldHeterogeneousVolumesCastShadows();
 
-			const FVolumeBounds VolumeBounds = CalculateLightVolumeBounds(LightSceneInfo->Proxy->GetBoundingSphere(), View, VolumeCascadeIndex, bDirectionalLight);
+			const FVolumeBounds VolumeBounds = InjectionData.VolumeBounds[VolumeCascadeIndex];
 			if (VolumeBounds.IsValid())
 			{
 				TShaderMapRef<FWriteToSliceVS> VertexShader(View.ShaderMap);
@@ -1148,17 +1424,12 @@ void InjectTranslucencyLightingVolume(
 				PassParameters->PS.ViewUniformBuffer = View.ViewUniformBuffer;
 
 				FDeferredLightUniformStruct* DeferredLightStruct = GraphBuilder.AllocParameters<FDeferredLightUniformStruct>();
-				*DeferredLightStruct = GetDeferredLightParameters(View, *LightSceneInfo, bUseLightFunctionAtlas, ELightShaderParameterFlags::RectAsSpotLight);
+				*DeferredLightStruct = GetDeferredLightParameters(View, *LightSceneInfo, bUseLightFunctionAtlas, DeferredLightParameterFlags);
 				PassParameters->PS.DeferredLight = GraphBuilder.CreateUniformBuffer(DeferredLightStruct);
 
 				GetVolumeShadowingShaderParameters(GraphBuilder, View, LightSceneInfo, InjectionData.ProjectedShadowInfo, PassParameters->PS.VolumeShadowingParameters);
 
-				const int32 VirtualShadowMapId = bUseVSM ? Renderer.VisibleLightInfos[LightSceneInfo->Id].GetVirtualShadowMapId(&View) : INDEX_NONE;
-				
-				// Switch it back off if there's no ID to avoid the FVirtualShadowMap permutation if we don't need it
-				bUseVSM = (VirtualShadowMapId != INDEX_NONE);
-
-				PassParameters->PS.VirtualShadowMapId = VirtualShadowMapId;
+				PassParameters->PS.VirtualShadowMapId = InjectionData.VirtualShadowMapId;
 				PassParameters->PS.LightFunctionParameters = FLightFunctionSharedParameters::GetLightFunctionSharedParameters(LightSceneInfo, 1.0f);
 				PassParameters->PS.VolumeCascadeIndex = VolumeCascadeIndex;
 				PassParameters->PS.AVSM = HeterogeneousVolumes::GetAdaptiveVolumetricShadowMapUniformBuffer(GraphBuilder, View.ViewState, LightSceneInfo);
@@ -1184,62 +1455,30 @@ void InjectTranslucencyLightingVolume(
 				GraphBuilder.AddPass(
 					RDG_EVENT_NAME("InjectTranslucencyLightingVolume(VolumeCascade=%d%s%s%s)",
 						VolumeCascadeIndex,
-						VirtualShadowMapId != INDEX_NONE ? TEXT(",VirtualShadowMap") : TEXT(""),
+						InjectionData.VirtualShadowMapId != INDEX_NONE ? TEXT(",VirtualShadowMap") : TEXT(""),
 						InjectionData.ProjectedShadowInfo != nullptr ? TEXT(",ShadowMap") : TEXT(""),
 						InjectionData.bApplyLightFunction ? TEXT(",LightFunction") : TEXT("")),
 					PassParameters,
 					ERDGPassFlags::Raster,
-					[PassParameters, VertexShader, GeometryShader, &View, &Renderer, &InjectionData, LightSceneInfo, bDirectionalLight, bUseVSM, VolumeBounds, VolumeCascadeIndex, bUseAdaptiveVolumetricShadowMap](FRHICommandList& RHICmdList)
+					[PassParameters, VertexShader, GeometryShader, &View, &Renderer, &InjectionData, LightSceneInfo, bDirectionalLight, VolumeBounds, VolumeCascadeIndex](FRHICommandList& RHICmdList)
 				{
 					FGraphicsPipelineStateInitializer GraphicsPSOInit;
-					RHICmdList.ApplyCachedRenderTargets(GraphicsPSOInit);
-
-					GraphicsPSOInit.RasterizerState = TStaticRasterizerState<FM_Solid, CM_None>::GetRHI();
-					GraphicsPSOInit.DepthStencilState = TStaticDepthStencilState<false, CF_Always>::GetRHI();
-					GraphicsPSOInit.PrimitiveType = PT_TriangleStrip;
-
-					if (bDirectionalLight)
-					{
-						// Accumulate the contribution of multiple lights
-						// Directional lights write their shadowing into alpha of the ambient texture
-						GraphicsPSOInit.BlendState = TStaticBlendState<
-							CW_RGBA, BO_Add, BF_One, BF_One, BO_Add, BF_One, BF_One,
-							CW_RGB, BO_Add, BF_One, BF_One, BO_Add, BF_One, BF_One>::GetRHI();
-					}
-					else
-					{
-						// Accumulate the contribution of multiple lights
-						GraphicsPSOInit.BlendState = TStaticBlendState<
-							CW_RGB, BO_Add, BF_One, BF_One, BO_Add, BF_Zero, BF_One,
-							CW_RGB, BO_Add, BF_One, BF_One, BO_Add, BF_Zero, BF_One>::GetRHI();
-					}
 
 					const FMaterialRenderProxy* MaterialProxy = InjectionData.LightFunctionMaterialProxy;
 					const FMaterial& Material = MaterialProxy->GetMaterialWithFallback( View.GetFeatureLevel(), MaterialProxy );
 					const FMaterialShaderMap* MaterialShaderMap = Material.GetRenderingThreadShaderMap();
 
 					FTranslucentLightingInjectPS::FPermutationDomain PermutationVector;
-					PermutationVector.Set< FTranslucentLightingInjectPS::FRadialAttenuation >( !bDirectionalLight );
-					PermutationVector.Set< FTranslucentLightingInjectPS::FDynamicallyShadowed >( InjectionData.ProjectedShadowInfo != nullptr );
-					PermutationVector.Set< FTranslucentLightingInjectPS::FLightFunction >( InjectionData.bApplyLightFunction );
-					PermutationVector.Set< FTranslucentLightingInjectPS::FVirtualShadowMap >( bUseVSM );
-					PermutationVector.Set< FTranslucentLightingInjectPS::FAdaptiveVolumetricShadowMap >(bUseAdaptiveVolumetricShadowMap);
+					PermutationVector.Set< FTranslucentLightingInjectPS::FRadialAttenuation >(!bDirectionalLight);
+					PermutationVector.Set< FTranslucentLightingInjectPS::FDynamicallyShadowed >(InjectionData.ProjectedShadowInfo != nullptr);
+					PermutationVector.Set< FTranslucentLightingInjectPS::FLightFunction >(InjectionData.bApplyLightFunction);
+					PermutationVector.Set< FTranslucentLightingInjectPS::FVirtualShadowMap >(InjectionData.VirtualShadowMapId != INDEX_NONE);
+					PermutationVector.Set< FTranslucentLightingInjectPS::FAdaptiveVolumetricShadowMap >(InjectionData.bUseAdaptiveVolumetricShadowMap);
 
 					auto PixelShader = MaterialShaderMap->GetShader< FTranslucentLightingInjectPS >( PermutationVector );
-	
-					GraphicsPSOInit.BoundShaderState.VertexDeclarationRHI = GScreenVertexDeclaration.VertexDeclarationRHI;
-					GraphicsPSOInit.BoundShaderState.VertexShaderRHI = VertexShader.GetVertexShader();
-					GraphicsPSOInit.BoundShaderState.SetGeometryShader(GeometryShader.GetGeometryShader());
-					GraphicsPSOInit.BoundShaderState.PixelShaderRHI = PixelShader.GetPixelShader();
-					SetGraphicsPipelineState(RHICmdList, GraphicsPSOInit, 0);
 
-					const int32 TranslucencyLightingVolumeDim = GetTranslucencyLightingVolumeDim();
-
-					SetShaderParametersLegacyVS(RHICmdList, VertexShader, VolumeBounds, FIntVector(TranslucencyLightingVolumeDim));
-					if (GeometryShader.IsValid())
-					{
-						SetShaderParametersLegacyGS(RHICmdList, GeometryShader, VolumeBounds.MinZ);
-					}
+					SetPSOStateForVolumeInjection(VertexShader, GeometryShader, PixelShader.GetPixelShader(),
+						VolumeBounds, bDirectionalLight, RHICmdList, GraphicsPSOInit);
 
 					SetShaderParametersLegacyPS(RHICmdList, PixelShader, View, InjectionData.LightFunctionMaterialProxy);
 					SetShaderParameters(RHICmdList, PixelShader, PixelShader.GetPixelShader(), PassParameters->PS);
@@ -1281,8 +1520,6 @@ void InjectSimpleTranslucencyLightingVolumeArray(
 	const FTranslucencyLightingVolumeTextures& Textures,
 	const FSimpleLightArray& SimpleLights)
 {
-	SCOPE_CYCLE_COUNTER(STAT_TranslucentInjectTime);
-
 	int32 NumLightsToInject = 0;
 
 	for (int32 LightIndex = 0; LightIndex < SimpleLights.InstanceData.Num(); LightIndex++)
@@ -1346,27 +1583,10 @@ void InjectSimpleTranslucencyLightingVolumeArray(
 							[VertexShader, GeometryShader, PixelShader, PassParameters, VolumeBounds, TranslucencyLightingVolumeDim](FRHICommandList& RHICmdList)
 						{
 							FGraphicsPipelineStateInitializer GraphicsPSOInit;
-							RHICmdList.ApplyCachedRenderTargets(GraphicsPSOInit);
-
-							GraphicsPSOInit.RasterizerState = TStaticRasterizerState<FM_Solid, CM_None>::GetRHI();
-							GraphicsPSOInit.DepthStencilState = TStaticDepthStencilState<false, CF_Always>::GetRHI();
-							// Accumulate the contribution of multiple lights
-							GraphicsPSOInit.BlendState = TStaticBlendState<
-								CW_RGB, BO_Add, BF_One, BF_One, BO_Add, BF_Zero, BF_One,
-								CW_RGB, BO_Add, BF_One, BF_One, BO_Add, BF_Zero, BF_One>::GetRHI();
-							GraphicsPSOInit.PrimitiveType = PT_TriangleStrip;
-
-							GraphicsPSOInit.BoundShaderState.VertexDeclarationRHI = GScreenVertexDeclaration.VertexDeclarationRHI;
-							GraphicsPSOInit.BoundShaderState.VertexShaderRHI = VertexShader.GetVertexShader();
-							GraphicsPSOInit.BoundShaderState.SetGeometryShader(GeometryShader.GetGeometryShader());
-							GraphicsPSOInit.BoundShaderState.PixelShaderRHI = PixelShader.GetPixelShader();
-							SetGraphicsPipelineState(RHICmdList, GraphicsPSOInit, 0);
-
-							SetShaderParametersLegacyVS(RHICmdList, VertexShader, VolumeBounds, FIntVector(TranslucencyLightingVolumeDim));
-							if (GeometryShader.IsValid())
-							{
-								SetShaderParametersLegacyGS(RHICmdList, GeometryShader, VolumeBounds.MinZ);
-							}
+							
+							SetPSOStateForVolumeInjection(VertexShader, GeometryShader, PixelShader.GetPixelShader(),
+								VolumeBounds, false, RHICmdList, GraphicsPSOInit);
+
 							SetShaderParameters(RHICmdList, PixelShader, PixelShader.GetPixelShader(), *PassParameters);
 							RasterizeToVolumeTexture(RHICmdList, VolumeBounds);
 						});
diff --git a/Engine/Source/Runtime/Renderer/Private/TranslucentLighting.h b/Engine/Source/Runtime/Renderer/Private/TranslucentLighting.h
--- a/Engine/Source/Runtime/Renderer/Private/TranslucentLighting.h
+++ b/Engine/Source/Runtime/Renderer/Private/TranslucentLighting.h
@@ -4,6 +4,7 @@
 
 #include "SceneRendering.h"
 #include "RenderGraph.h"
+#include "VolumeRendering.h"
 
 struct FSortedLightSceneInfo;
 struct FCloudShadowAOData;
@@ -72,7 +73,28 @@ public:
 		bool bApplyLightFunction;
 		// must not be 0
 		const FMaterialRenderProxy* LightFunctionMaterialProxy;
+		// can be INDEX_NONE
+		int32 VirtualShadowMapId;
+		//
+		bool bUseAdaptiveVolumetricShadowMap;
+		FVolumeBounds VolumeBounds[TVC_MAX];
 	};
+	
+	typedef TArray<FInjectionData, SceneRenderingAllocator> FInjectionDataArray;
+
+	struct FPerViewData
+	{
+		// Array of lights that will be injected individually
+		FInjectionDataArray Unbatched;
+
+		// Lights that will be injected collectively
+		// For these we store a bitmask of light indices in the light grid
+		TBitArray<SceneRenderingAllocator> BatchedLocalLights[TVC_MAX];
+		uint32 BatchedLocalLightCount = 0;			// Maintain a count for stats
+		bool bAnyBatchedLightsWithVirtualShadowMaps = false;
+	};
+
+	TArray<FPerViewData, SceneRenderingAllocator>& InjectionDataPerView;
 
 	void AddLightForInjection(
 		const FViewInfo& View,
@@ -80,9 +102,6 @@ public:
 		TArrayView<const FVisibleLightInfo> VisibleLightInfos,
 		const FLightSceneInfo& LightSceneInfo,
 		const FProjectedShadowInfo* InProjectedShadowInfo = nullptr);
-
-	typedef TArray<FInjectionData, SceneRenderingAllocator> FInjectionDataArray;
-	TArray<FInjectionDataArray, SceneRenderingAllocator>& InjectionDataPerView;
 };
 
 /** Initializes translucency volume lighting shader parameters from an optional textures struct. If null or uninitialized, fallback textures are used. */
