Fix a crash if there are more than 8 lights cast shadow on mobile GL platform. UniformBuffer can't be created on worker thread on GL.

diff --git a/Engine/Source/Runtime/Renderer/Private/ShadowSetup.cpp b/Engine/Source/Runtime/Renderer/Private/ShadowSetup.cpp
--- a/Engine/Source/Runtime/Renderer/Private/ShadowSetup.cpp
+++ b/Engine/Source/Runtime/Renderer/Private/ShadowSetup.cpp
@@ -4605,7 +4605,7 @@ void FSceneRenderer::GatherShadowDynamicMeshElements(FDynamicShadowsTaskData& Ta
 			FOptionalTaskTagScope Scope(ETaskTag::EParallelRenderingThread);
 			ShadowsToSetupViews[Index]->SetupShadowDepthView(this);
 		},
-		EParallelForFlags::None
+		GSupportsParallelRenderingTasksWithSeparateRHIThread ? EParallelForFlags::None : EParallelForFlags::ForceSingleThread
 	);
 
 	TaskData.ShadowsToGather.Reserve(128);
