Fix translucent holdout is always enabled even when r.Deferred.SupportPrimitiveAlphaHoldout is not turned on.

diff --git a/Engine/Source/Runtime/Renderer/Private/TranslucentRendering.cpp b/Engine/Source/Runtime/Renderer/Private/TranslucentRendering.cpp
--- a/Engine/Source/Runtime/Renderer/Private/TranslucentRendering.cpp
+++ b/Engine/Source/Runtime/Renderer/Private/TranslucentRendering.cpp
@@ -120,9 +120,12 @@ bool IsTranslucentHoldoutEnabled(EShadingPath ShadingPath)
 
 	if (RenderTranslucentHoldout != 0)
 	{
-		static IConsoleVariable* CVar = IConsoleManager::Get().FindConsoleVariable(TEXT("r.PostProcessing.PropagateAlpha"));
+		static IConsoleVariable* CVarPropagateAlpha = IConsoleManager::Get().FindConsoleVariable(TEXT("r.PostProcessing.PropagateAlpha"));
+		static IConsoleVariable* CVarSupportAlphaHoldout = IConsoleManager::Get().FindConsoleVariable(TEXT("r.Deferred.SupportPrimitiveAlphaHoldout"));
+		const bool bDeferredPropagateAlpha = CVarPropagateAlpha ? CVarPropagateAlpha->GetBool() : false;
+		const bool bDeferredSupportPrimitiveAlphaHoldout = CVarSupportAlphaHoldout ? CVarSupportAlphaHoldout->GetBool() : false;
 		// TODO: add support to Mobile renderer 
-		bPropagateAlpha = CVar->GetBool() && (ShadingPath != EShadingPath::Mobile);
+		bPropagateAlpha = bDeferredPropagateAlpha && bDeferredSupportPrimitiveAlphaHoldout && (ShadingPath != EShadingPath::Mobile);
 	}
 
 	return bPropagateAlpha;
@@ -130,18 +133,8 @@ bool IsTranslucentHoldoutEnabled(EShadingPath ShadingPath)
 
 static bool IsTranslucentHoldoutEnabled(TArrayView<const FViewInfo> Views)
 {
-	bool bPropagateAlpha = true;
-
-	for (int32 ViewId = 0; ViewId < Views.Num(); ++ViewId)
-	{
-		if (!IsTranslucentHoldoutEnabled(GetFeatureLevelShadingPath(Views[ViewId].GetFeatureLevel())))
-		{
-			bPropagateAlpha = false;
-			break;
-		}
-	}
-	
-	return bPropagateAlpha;
+	ensure(Views.Num() > 0);
+	return IsTranslucentHoldoutEnabled(GetFeatureLevelShadingPath(Views[0].GetFeatureLevel()));
 }
 
 DynamicRenderScaling::FHeuristicSettings GetDynamicTranslucencyResolutionSettings()
