PathTracing: Implement new debug visualization mode

This mode is similar to the raytracing debug visualization, but using the path tracer's payload. In primary ray mode, a basic light loop is run without sample accumulation which can be used to test performance of the path tracer's light sampling. This mode also allows TSR to run, which can act as a rudimentary denoiser of the lighting signal if enough samples are used.

Remove light grid and decal grid visualization cvars from the main path tracer as these are now well covered by the new debug mode.

#rnx

diff --git a/Engine/Plugins/Experimental/GPULightmass/Shaders/Private/LightmapPathTracing.usf b/Engine/Plugins/Experimental/GPULightmass/Shaders/Private/LightmapPathTracing.usf
--- a/Engine/Plugins/Experimental/GPULightmass/Shaders/Private/LightmapPathTracing.usf
+++ b/Engine/Plugins/Experimental/GPULightmass/Shaders/Private/LightmapPathTracing.usf
@@ -136,24 +136,6 @@ static const float MaxNormalBias = 0.1;
 RaytracingAccelerationStructure TLAS;
 uint SceneVisibleLightCount;
 
-void ApplyRayBias(inout FRayDesc Ray, float HitT, float3 Normal, bool ShadowTerminatorFix = false)
-{
-#if 1
-	// Take maximum of position or hit distance to determine "radius" of hit avoidance
-	const float RefValue = max(max(abs(Ray.Origin.x), abs(Ray.Origin.y)), max(abs(Ray.Origin.z), HitT));
-	const uint UlpRadius = 16; // number of floating point ulps to skip around an intersection
-	const float RelativeOffset = asfloat(asuint(RefValue) + UlpRadius) - RefValue;
-	const float ShadowOffset = 0.1; // 1mm
-	const float BaseOffset = 0.001; // 0.01mm (avoid precision issues very close to the origin where ulps become tiny)
-	// for shadow rays, push out the point a bit more so we can reduce shadow terminator issues from poor tesselations
-	// other rays are less likely to arrive at grazing angles, so don't need the extra offset
-	Ray.Origin += max(ShadowTerminatorFix ? ShadowOffset : BaseOffset, RelativeOffset * 100) * Normal;
-	Ray.TMax -= length(max(ShadowTerminatorFix ? ShadowOffset : BaseOffset, RelativeOffset * 100) * Normal);
-#else
-	ApplyPositionBias(Ray.Origin, Ray.Direction, Normal, MaxNormalBias);
-#endif
-}
-
 void AccumulateRadiance(inout float3 TotalRadiance, float3 PathRadiance)
 {
 	float MaxPathIntensity = 0;
@@ -584,7 +566,7 @@ void PathTracingKernel(
 						LightRay.TMin = 0;
 						LightRay.Direction = LightSample.Direction;
 						LightRay.TMax = LightSample.Distance;
-						ApplyRayBias(LightRay, Payload.HitT, SignedPositionBias * Payload.WorldGeoNormal);
+						ApplyRayBias(LightRay.Origin, Payload.HitT, SignedPositionBias * Payload.WorldGeoNormal);
 					
 						float AvgRoughness = ApproximateCaustics ? GetAverageRoughness(Payload) : 0.0;
 
@@ -716,7 +698,7 @@ void PathTracingKernel(
 		Ray.Direction = MaterialSample.Direction;
 		Ray.TMin = 0;
 		Ray.TMax = POSITIVE_INFINITY;
-		ApplyRayBias(Ray, Payload.HitT, MaterialSample.PositionBiasSign * Payload.WorldGeoNormal);
+		ApplyRayBias(Ray.Origin, Payload.HitT, MaterialSample.PositionBiasSign * Payload.WorldGeoNormal);
 
 		if (ApproximateCaustics)
 		{
@@ -1478,7 +1460,7 @@ void StationaryLightShadowTracingMainRG()
 				/* out */ Ray.TMin,
 				/* out */ Ray.TMax);
 				
-			ApplyRayBias(Ray, LightType == LIGHT_TYPE_DIRECTIONAL ? 1.0f : max(Ray.TMax - Ray.TMin, 0.0f), EffectiveWorldNormal);
+			ApplyRayBias(Ray.Origin, LightType == LIGHT_TYPE_DIRECTIONAL ? 1.0f : max(Ray.TMax - Ray.TMin, 0.0f), EffectiveWorldNormal);
 			
 			if (bRaySampleValid)
 			{
diff --git a/Engine/Shaders/Private/PathTracing/PathTracingCommon.ush b/Engine/Shaders/Private/PathTracing/PathTracingCommon.ush
--- a/Engine/Shaders/Private/PathTracing/PathTracingCommon.ush
+++ b/Engine/Shaders/Private/PathTracing/PathTracingCommon.ush
@@ -243,8 +243,8 @@ struct FPackedPathTracingPayload : FMinimalPayload
 
 	float3 UnpackDiffuseColor()       { return PayloadDecodeLDRColor(PackedData[5]); }
 	float3 UnpackSpecularColor()      { return PayloadDecodeLDRColor(PackedData[6]); }
-	float3 UnpackSpecularEdgeColor()  { return (GetShadingModelID() == SHADINGMODELID_UNLIT) ? 0.0 : PayloadDecodeLDRColor(PackedData[7]); }
-	float4 UnpackRoughnessAniso()     { return (GetShadingModelID() == SHADINGMODELID_UNLIT) ? 0.0 : PayloadDecodeRoughnessAniso(PackedData[8]); }
+	float3 UnpackSpecularEdgeColor()  { return PayloadDecodeLDRColor(PackedData[7]); }
+	float4 UnpackRoughnessAniso()     { return PayloadDecodeRoughnessAniso(PackedData[8]); }
 	float3 UnpackWorldTangent()       { return PayloadDecodeUnitVector(PackedData[9]); }
 	float3 UnpackTransparencyColor()  { return PayloadDecodeLDRColor(PackedData[10]); }
 	float3 UnpackMeanFreePath()       { return float3(f16tof32(PackedData[11]),
@@ -266,10 +266,10 @@ struct FPackedPathTracingPayload : FMinimalPayload
 #else
 	float3 UnpackBaseColor()         { return PayloadDecodeLDRColor(PackedData[5]); }
 	float3 UnpackTransparencyColor() { return PayloadDecodeLDRColor(PackedData[6]); }
-	float UnpackMetallic()           { return (GetShadingModelID() == SHADINGMODELID_UNLIT) ? 0.0 : f16tof32(PackedData[7]); }
-	float UnpackSpecular()           { return (GetShadingModelID() == SHADINGMODELID_UNLIT) ? 0.0 : f16tof32(PackedData[7] >> 16); }
-	float UnpackRoughness()          { return (GetShadingModelID() == SHADINGMODELID_UNLIT) ? 0.0 :  f16tof32(PackedData[8]); }
-	float UnpackAnisotropy()         { return (GetShadingModelID() == SHADINGMODELID_UNLIT) ? 0.0 : f16tof32(PackedData[8] >> 16); }
+	float UnpackMetallic()           { return f16tof32(PackedData[7]); }
+	float UnpackSpecular()           { return f16tof32(PackedData[7] >> 16); }
+	float UnpackRoughness()          { return f16tof32(PackedData[8]); }
+	float UnpackAnisotropy()         { return f16tof32(PackedData[8] >> 16); }
 	float3 UnpackWorldTangent()      { return PayloadDecodeUnitVector(PackedData[9]); }
 	float4 UnpackCustomData0()       { return float4(f16tof32(PackedData[10]),
 	                                                 f16tof32(PackedData[10] >> 16),
diff --git a/Engine/Shaders/Private/PathTracing/PathTracingCopyDepth.usf b/Engine/Shaders/Private/PathTracing/PathTracingCopyDepth.usf
new file mode 100644
--- /dev/null
+++ b/Engine/Shaders/Private/PathTracing/PathTracingCopyDepth.usf
@@ -0,0 +1,20 @@
+// Copyright Epic Games, Inc. All Rights Reserved.
+
+#include "../Common.ush"
+
+Texture2D<float> DepthTexture;
+float PreExposure; // Taken from the view, but adjusted by the base exposure already applied to the Radiance texture
+int AdaptiveSamplingVisualize;
+
+void CopyDepth(
+	in noperspective float4 UVAndScreenPos : TEXCOORD0,
+	out float OutDepth : SV_DEPTH,
+	out float4 OutColor : SV_Target0
+)
+{
+	float2 UV = UVAndScreenPos.xy;
+	float2 BufferSize = View.BufferSizeAndInvSize.xy;
+	int3 TexCoord = int3(UV * BufferSize - View.ViewRectMin.xy, 0);
+	OutDepth = DepthTexture.Load(TexCoord);
+	OutColor = 0;
+}
diff --git a/Engine/Shaders/Private/PathTracing/PathTracingCore.ush b/Engine/Shaders/Private/PathTracing/PathTracingCore.ush
--- a/Engine/Shaders/Private/PathTracing/PathTracingCore.ush
+++ b/Engine/Shaders/Private/PathTracing/PathTracingCore.ush
@@ -62,8 +62,6 @@ uint VolumeMISMode;
 uint ApproximateCaustics;
 uint EnableCameraBackfaceCulling;
 uint SamplerType;
-uint VisualizeLightGrid;
-uint VisualizeDecalGrid;
 uint EnableDBuffer;
 uint ApplyDiffuseSpecularOverrides;
 uint EnabledDirectLightingContributions;
@@ -958,16 +956,6 @@ FProbeResult TraceSSSProbeRay(FRayDesc Ray, inout int InterfaceCounter)
 	}
 }
 
-void ApplyRayBias(inout FRayDesc Ray, float HitT, float3 Normal)
-{
-	// Take maximum of position or hit distance to determine "radius" of hit avoidance
-	const float RefValue = max(max(abs(Ray.Origin.x), abs(Ray.Origin.y)), max(abs(Ray.Origin.z), HitT));
-	const uint UlpRadius = 16; // number of floating point ulps to skip around an intersection
-	const float RelativeOffset = asfloat(asuint(RefValue) + UlpRadius) - RefValue;
-	const float BaseOffset = 0.00390625; // 2^-8cm (roughly 39micro meters) (This avoids precision issues very close to the origin where the ulps become too small)
-	Ray.Origin += max(BaseOffset, RelativeOffset) * Normal;
-}
-
 float3 ComputeDwivediScale(float3 Albedo)
 {
 	// "Zero-Variance Theory for Efficient Subsurface Scattering"
@@ -1187,7 +1175,7 @@ bool ProcessSubsurfaceRandomWalk(inout FPathTracingPayload Payload, inout float3
 	Ray.Origin = Payload.TranslatedWorldPos;
 	Ray.Direction = TangentToWorld(-CosineSampleHemisphere(RandSample.xy).xyz, Payload.WorldNormal);
 	Ray.TMin = 0;
-	ApplyRayBias(Ray, Payload.HitT, -Payload.WorldGeoNormal);
+	ApplyRayBias(Ray.Origin, Payload.HitT, -Payload.WorldGeoNormal);
 
 	// Support SSS for mfp smaller than 1mm.
 	SSS.Radius = max(SSS.Radius, 0.0009); // 0.009mm minimum to make sure random walk can move forward and to keep SigmaT finite
@@ -1352,7 +1340,7 @@ bool ProcessSubsurfaceRandomWalk(inout FPathTracingPayload Payload, inout float3
 				bDoSlabSearch = GuidedRatio > 0;
 #endif
 
-				ApplyRayBias(Ray, ProbeResult.HitT, ProbeResult.WorldGeoNormal);
+				ApplyRayBias(Ray.Origin, ProbeResult.HitT, ProbeResult.WorldGeoNormal);
 
 				// Reset interface counter as if we had just entered the model (taking into account that we are seeing it from the opposite side now)
 				InterfaceCounter = ProbeResult.FrontFace ? -1 : +1;
@@ -1408,12 +1396,16 @@ FPathState CreatePathState(int2 PixelIndex, int2 TextureIndex)
 	// Initialize ray and payload
 	{
 		float2 AAJitter = RandomSequence_GenerateSample2D(Output.RandSequence);
-
+#if 0
+		// Tent filter
+		float2 P = 2 * AAJitter - 0.99999994f;
+		AAJitter = 0.5 + sign(P) * (1.0 - sqrt(1.0 - abs(P)));
+#elif 1
 		// importance sample a gaussian kernel with variable sigma
 		float3 Disk = ConcentricDiskSamplingHelper(AAJitter);
 		float Sigma = FilterWidth / 6.0; // user-provided width covers +/-3*Sigma
 		AAJitter = 0.5 + Sigma * Disk.xy * sqrt(-2.0 * log(1.0 - Disk.z * Disk.z));
-
+#endif
 		float2 ViewportUV = (LaunchIndex + AAJitter) * View.BufferSizeAndInvSize.zw;
 		Output.Ray = CreatePrimaryRay(ViewportUV);
 	}
@@ -1625,25 +1617,6 @@ bool PathTracingKernel(inout FPathState PathState, int Bounce)
 		LightLoopCount.NumMISLights = SceneInfiniteLightCount;
 	}
 
-	// visualize the # of lights
-	if (VisualizeLightGrid)
-	{
-		PathState.Radiance = LightGridVisualize(LightLoopCount, VisualizeLightGrid);
-		PathState.Radiance *= abs(dot(Payload.WorldNormal, PathState.Ray.Direction));
-		PathState.Radiance /= View.PreExposure; // cancel out exposure
-		return false;
-	}
-
-	// visualize the # of decals
-	if (VisualizeDecalGrid)
-	{
-		FDecalLoopCount DecalLoopCount = DecalGridLookup(Payload.TranslatedWorldPos);
-		PathState.Radiance = DecalGridVisualize(DecalLoopCount, VisualizeDecalGrid);
-		PathState.Radiance *= abs(dot(Payload.WorldNormal, PathState.Ray.Direction));
-		PathState.Radiance /= View.PreExposure; // cancel out exposure
-		return false;
-	}
-
 	// Choose a random number for both Light sampling and BxDF sampling
 	float4 RandSample = RandomSequence_GenerateSample4D(PathState.RandSequence);
 
@@ -1719,7 +1692,7 @@ bool PathTracingKernel(inout FPathState PathState, int Bounce)
 				LightRay.TMin = 0;
 				LightRay.Direction = LightSample.Direction;
 				LightRay.TMax = LightSample.Distance;
-				ApplyRayBias(LightRay, Payload.HitT, SignedPositionBias * Payload.WorldGeoNormal);
+				ApplyRayBias(LightRay.Origin, Payload.HitT, SignedPositionBias * Payload.WorldGeoNormal);
 
 				float AvgRoughness = ApproximateCaustics ? GetAverageRoughness(Payload) : 0.0;
 
@@ -1803,7 +1776,7 @@ bool PathTracingKernel(inout FPathState PathState, int Bounce)
 	PathState.Ray.Direction = MaterialSample.Direction;
 	PathState.Ray.TMin = 0;
 	PathState.Ray.TMax = POSITIVE_INFINITY;
-	ApplyRayBias(PathState.Ray, Payload.HitT, MaterialSample.PositionBiasSign * Payload.WorldGeoNormal);
+	ApplyRayBias(PathState.Ray.Origin, Payload.HitT, MaterialSample.PositionBiasSign * Payload.WorldGeoNormal);
 
 	// enlarge roughness based on the chosen lobe roughness
 	float PreviousPathRoughness = PathState.GetPathRoughness();
diff --git a/Engine/Shaders/Private/PathTracing/PathTracingDebug.usf b/Engine/Shaders/Private/PathTracing/PathTracingDebug.usf
new file mode 100644
--- /dev/null
+++ b/Engine/Shaders/Private/PathTracing/PathTracingDebug.usf
@@ -0,0 +1,286 @@
+// Copyright Epic Games, Inc. All Rights Reserved.
+
+#define LIGHT_LOOP_METHOD 	2 // 0: stochastic (like ref PT) 1: loop over all lights 2: RIS (take N samples but trace a single shadow ray)
+
+#define USE_PATH_TRACING_LIGHT_GRID	1
+#define USE_RAY_TRACING_DECAL_GRID	1
+
+
+#include "../Common.ush"
+#include "../BlueNoise.ush"
+#include "../RayTracing/RayTracingCommon.ush"
+#include "../RayTracing/RayTracingDecalGrid.ush"
+#include "PathTracingCommon.ush"
+#include "Utilities/PathTracingRandomSequence.ush"
+#include "Utilities/PathTracingRIS.ush"
+#include "Light/PathTracingLightSampling.ush"
+#include "Light/PathTracingLightGrid.ush"
+
+
+RWTexture2D<float4> RWSceneColor;
+RWTexture2D<float>  RWSceneDepth;
+RaytracingAccelerationStructure TLAS;
+
+int DebugMode;
+int NumLightSamples;
+
+float3 TraceShadowRay(FRayDesc Ray, float PathRoughness, float PreviousPathRoughness, uint MissShaderIndex, bool bCastShadows)
+{
+	if (!bCastShadows && MissShaderIndex == 0)
+	{
+		// no work to do
+		return 1.0;
+	}
+
+	const uint RayFlags = RAY_FLAG_ACCEPT_FIRST_HIT_AND_END_SEARCH | RAY_FLAG_SKIP_CLOSEST_HIT_SHADER;
+	const uint InstanceInclusionMask = PATHTRACER_MASK_SHADOW;
+	const uint RayContributionToHitGroupIndex = RAY_TRACING_SHADER_SLOT_SHADOW;
+	const uint MultiplierForGeometryContributionToShaderIndex = RAY_TRACING_NUM_SHADER_SLOTS;
+
+	FPackedPathTracingPayload PackedPayload = InitPathTracingVisibilityPayload(PathRoughness);
+
+	if (!bCastShadows)
+	{
+		// ray should not cast shadows, make it degenerate so we can still run the miss shader
+		Ray.TMin = POSITIVE_INFINITY;
+		Ray.TMax = POSITIVE_INFINITY;
+	}
+	TraceRay(
+		TLAS,
+		RayFlags,
+		InstanceInclusionMask,
+		RayContributionToHitGroupIndex,
+		MultiplierForGeometryContributionToShaderIndex,
+		MissShaderIndex,
+		Ray.GetNativeDesc(),
+		PackedPayload);
+
+	if (PackedPayload.IsHit())
+	{
+		// We didn't run the miss shader, therefore we must have hit something opaque (or reached full opacity)
+		return 0.0;
+	}
+
+	float3 Throughput = PackedPayload.GetRayThroughput();
+
+	return Throughput;
+}
+
+struct PathTracingResult {
+	float4 Radiance;
+	float Z;
+};
+
+PathTracingResult MakeResult(float3 Radiance, float Alpha, float Z)
+{
+	PathTracingResult Result;
+	Result.Radiance.rgb = Radiance;
+	Result.Radiance.a = Alpha;
+	Result.Z = Z;
+	return Result;
+}
+
+float3 GetLightingChannelMaskDebugColor(int LightingChannelMask)
+{
+	float r = LightingChannelMask & 1 ? 1.0 : 0.1;
+	float g = LightingChannelMask & 2 ? 1.0 : 0.1;
+	float b = LightingChannelMask & 4 ? 1.0 : 0.1;
+	return float3(r, g, b);
+}
+
+float3 AnisoDebugColor(float InAniso)
+{
+	// Follow same mapping than raster buffer visualization (BP located in Engine/Contents/Anisotropy)
+	float AniG = saturate( InAniso);
+	float AniB = saturate(-InAniso);
+	return float3(0.0, pow(AniG, 2.2), pow(AniB, 2.2));
+}
+
+PathTracingResult PathTracingDebugCore(FRayDesc Ray, uint2 LaunchIndex)
+{
+	const uint RayFlags = RAY_FLAG_CULL_BACK_FACING_TRIANGLES;
+	const uint InstanceInclusionMask = PATHTRACER_MASK_CAMERA;
+	const uint MissShaderIndex = 0;
+
+	FPackedPathTracingPayload PackedPayload = InitPathTracingPayload(PATHTRACER_SCATTER_CAMERA, 0.0);
+	TraceRay(
+		TLAS,
+		RayFlags,
+		InstanceInclusionMask,
+		RAY_TRACING_SHADER_SLOT_MATERIAL,
+		RAY_TRACING_NUM_SHADER_SLOTS,
+		MissShaderIndex,
+		Ray.GetNativeDesc(),
+		PackedPayload);
+	float4 Result = float4(0, 0, 0, 1);
+	float Z = 0;
+
+	float LightPickingCdf[RAY_TRACING_LIGHT_COUNT_MAXIMUM];
+	float LightPickingCdfSum = 0;
+	float3 PathThroughput = 1;
+	if (PackedPayload.IsHit())
+	{
+		const float3 ViewZ = View.ViewToTranslatedWorld[2].xyz;
+		const float3 TranslatedWorldPos = Ray.Origin + PackedPayload.HitT * Ray.Direction;
+		const float3 WorldNormal = PackedPayload.UnpackWorldNormal();
+		const uint PrimitiveLightingChannelMask = PackedPayload.GetPrimitiveLightingChannelMask();
+
+		Z = ConvertToDeviceZ(dot(TranslatedWorldPos, ViewZ));
+
+		FLightLoopCount LightLoopCount = LightGridLookup(TranslatedWorldPos);
+
+		switch (DebugMode)
+		{
+			case PATH_TRACER_DEBUG_VIZ_RADIANCE: 					return MakeResult(PackedPayload.UnpackRadiance() * View.PreExposure, 0.0, Z);
+			case PATH_TRACER_DEBUG_VIZ_WORLD_NORMAL: 				return MakeResult(WorldNormal * 0.5 + 0.5, 0.0, Z);
+			case PATH_TRACER_DEBUG_VIZ_WORLD_SMOOTH_NORMAL:			return MakeResult(PackedPayload.UnpackWorldSmoothNormal() * 0.5 + 0.5, 0.0, Z);
+			case PATH_TRACER_DEBUG_VIZ_WORLD_GEO_NORMAL:			return MakeResult(PackedPayload.UnpackWorldGeoNormal() * 0.5 + 0.5, 0.0, Z);
+			default:
+			case PATH_TRACER_DEBUG_VIZ_BASE_COLOR:					return MakeResult(PackedPayload.UnpackBaseColor(), 0.0, Z);
+			case PATH_TRACER_DEBUG_VIZ_DIFFUSE_COLOR:				return MakeResult(PackedPayload.UnpackBaseColor() * (1 - PackedPayload.UnpackMetallic()), 0.0, Z);
+			case PATH_TRACER_DEBUG_VIZ_SPECULAR_COLOR:				return MakeResult(lerp(0.08 * PackedPayload.UnpackSpecular(), PackedPayload.UnpackBaseColor(), PackedPayload.UnpackMetallic()), 0.0, Z);
+			case PATH_TRACER_DEBUG_VIZ_OPACITY:						return MakeResult(PackedPayload.UnpackTransparencyColor(), 0.0, Z);
+			case PATH_TRACER_DEBUG_VIZ_METALLIC:					return MakeResult(PackedPayload.UnpackMetallic(), 0.0, Z);
+			case PATH_TRACER_DEBUG_VIZ_SPECULAR:					return MakeResult(PackedPayload.UnpackSpecular(), 0.0, Z);
+			case PATH_TRACER_DEBUG_VIZ_ROUGHNESS:					return MakeResult(PackedPayload.UnpackRoughness(), 0.0, Z);
+			case PATH_TRACER_DEBUG_VIZ_IOR:							return MakeResult(PackedPayload.UnpackIor(), 0.0, Z);
+			case PATH_TRACER_DEBUG_VIZ_SHADING_MODEL:				return MakeResult(GetShadingModelColor(PackedPayload.GetShadingModelID()), 0.0, Z);
+			case PATH_TRACER_DEBUG_VIZ_LIGHTING_CHANNEL_MASK:		return MakeResult(GetLightingChannelMaskDebugColor(PackedPayload.GetPrimitiveLightingChannelMask()), 0.0, Z);
+			case PATH_TRACER_DEBUG_VIZ_CUSTOM_DATA0:				return MakeResult(PackedPayload.UnpackCustomData0().rgb * View.PreExposure, PackedPayload.UnpackCustomData0().a, Z);
+			case PATH_TRACER_DEBUG_VIZ_CUSTOM_DATA1:				return MakeResult(PackedPayload.UnpackCustomData1().rgb * View.PreExposure, PackedPayload.UnpackCustomData1().a, Z);
+			case PATH_TRACER_DEBUG_VIZ_WORLD_POSITION:				return MakeResult(TranslatedWorldPos * View.PreExposure, 0.0, Z);
+			case PATH_TRACER_DEBUG_VIZ_PRIMARY_RAYS: break; // fallthrough to implementation below
+			case PATH_TRACER_DEBUG_VIZ_WORLD_TANGENT: 				return MakeResult(PackedPayload.UnpackWorldTangent() * 0.5 + 0.5, 0.0, Z);
+			case PATH_TRACER_DEBUG_VIZ_ANISOTROPY:					return MakeResult(AnisoDebugColor(PackedPayload.UnpackAnisotropy()), 0.0, Z);
+			case PATH_TRACER_DEBUG_VIZ_LIGHT_GRID_COUNT:
+			{
+				float3 Color = LightGridVisualize(LightLoopCount, 1);
+				Color *= abs(dot(WorldNormal, Ray.Direction));
+				return MakeResult(Color, 0.0, Z);
+			}
+			case PATH_TRACER_DEBUG_VIZ_DECAL_GRID_COUNT:
+			{
+				FDecalLoopCount DecalLoopCount = DecalGridLookup(TranslatedWorldPos);
+				float3 Color = DecalGridVisualize(DecalLoopCount, 1);
+				Color *= abs(dot(WorldNormal, Ray.Direction));
+				return MakeResult(Color, 0.0, Z);
+			}
+		}
+
+#if LIGHT_LOOP_METHOD == 0
+		for (uint Index = 0, Num = LightLoopCount.NumLights; Index < Num; ++Index)
+		{
+			uint LightIndex = GetLightId(Index, LightLoopCount);
+			float LightEstimate = EstimateLight(LightIndex, TranslatedWorldPos, WorldNormal, PrimitiveLightingChannelMask, false);
+			LightPickingCdfSum += LightEstimate;
+			LightPickingCdf[Index] = LightPickingCdfSum;
+		}
+
+		if (LightPickingCdfSum > 0)
+		{
+			// init worked
+			int LightId;
+			float LightPickPdf = 0;
+			
+			for (int Sample = 0; Sample < NumLightSamples; Sample++)
+			{
+				RandomSequence RandSequence;
+				RandomSequence_Initialize(RandSequence, LaunchIndex, Sample, View.StateFrameIndex, NumLightSamples);
+				float4 RandSample = RandomSequence_GenerateSample4D(RandSequence);
+
+				SelectLight(RandSample.x * LightPickingCdfSum, LightLoopCount.NumLights, LightPickingCdf, LightId, LightPickPdf);
+				LightPickPdf /= LightPickingCdfSum;
+				float2 LightRandSample = RandSample.yz;
+				LightId = GetLightId(LightId, LightLoopCount);
+#else
+		for (int Sample = 0; Sample < NumLightSamples; Sample++)
+		{
+			RandomSequence RandSequence;
+			RandomSequence_Initialize(RandSequence, LaunchIndex, Sample, View.StateFrameIndex, NumLightSamples);
+			float3 RandSample = RandomSequence_GenerateSample3D(RandSequence);			
+			float2 LightRandSample = RandSample.xy;
+
+#if LIGHT_LOOP_METHOD == 2
+			FRISContext TraceSample = InitRISContext(RandSample.z);
+			FRayDesc ShadowRaySample;
+			float3 ShadowRayContrib = 0;
+			uint ShadowRayLightId;
+#endif
+			for (uint Index = 0, Num = LightLoopCount.NumLights; Index < Num; ++Index)
+			{
+				uint LightId = GetLightId(Index, LightLoopCount);
+				float LightPickPdf = 1;
+#endif
+
+				FLightSample LightSample = SampleLight(LightId, LightRandSample, TranslatedWorldPos, WorldNormal);
+
+				if (LightSample.Pdf > 0)
+				{
+					LightSample.RadianceOverPdf /= LightPickPdf;
+					LightSample.Pdf *= LightPickPdf;
+
+					float3 MaterialEvalWeight = PackedPayload.UnpackBaseColor();
+					float  MaterialEvalPdf = saturate(dot(WorldNormal, LightSample.Direction)) / PI;
+					float3 LightContrib = PathThroughput * LightSample.RadianceOverPdf * MaterialEvalWeight * MaterialEvalPdf;
+					{
+						FRayDesc ShadowRay;
+						ShadowRay.Origin = TranslatedWorldPos;
+						ShadowRay.Direction = LightSample.Direction;
+						ShadowRay.TMin = 0;
+						ShadowRay.TMax = LightSample.Distance;
+
+						ApplyRayBias(ShadowRay.Origin, PackedPayload.HitT, PackedPayload.UnpackWorldGeoNormal());
+
+						const bool bCastShadows = CastsShadow(LightId);
+						const uint MissShaderIndex = GetLightMissShaderIndex(LightId);
+
+#if LIGHT_LOOP_METHOD == 2
+						float SelectionWeight = LightContrib.x + LightContrib.y + LightContrib.z;
+						if (TraceSample.Accept(SelectionWeight))
+						{
+							ShadowRaySample = ShadowRay;
+							ShadowRayContrib = LightContrib / SelectionWeight;
+							ShadowRayLightId = LightId;
+						}
+#else
+						Result.xyz += LightContrib * TraceShadowRay(ShadowRay, 1.0, 0.0, MissShaderIndex, bCastShadows);
+#endif
+					}
+				}
+			}
+#if LIGHT_LOOP_METHOD == 2
+			if (TraceSample.HasSample())
+			{
+				const bool bCastShadows = CastsShadow(ShadowRayLightId);
+				const uint MissShaderIndex = GetLightMissShaderIndex(ShadowRayLightId);
+				ShadowRayContrib *= TraceSample.GetNormalization();
+				Result.xyz += ShadowRayContrib * TraceShadowRay(ShadowRaySample, 1.0, 0.0, MissShaderIndex, bCastShadows);
+			}
+#endif
+		}
+		Result.rgb /= NumLightSamples;
+		Result.rgb += PackedPayload.UnpackRadiance();
+		Result.a = 0;
+	}
+	return MakeResult(Result.rgb * View.PreExposure, Result.a, Z);
+}
+
+
+RAY_TRACING_ENTRY_RAYGEN(PathTracingDebugRG)
+{
+	uint2 DispatchIdx = DispatchRaysIndex().xy;
+	const uint2 DispatchDim = DispatchRaysDimensions().xy;
+
+
+	uint2 LaunchIndex = DispatchIdx + View.ViewRectMin.xy;
+
+
+	const float2 AAJitter = 0.5; // trace through pixel center
+	const float2 ViewportUV = (LaunchIndex + AAJitter) * View.BufferSizeAndInvSize.zw;
+
+	
+	PathTracingResult Result = PathTracingDebugCore(CreatePrimaryRay(ViewportUV), LaunchIndex);
+
+	RWSceneColor[DispatchIdx] = Result.Radiance;
+	RWSceneDepth[DispatchIdx] = Result.Z;
+}
diff --git a/Engine/Shaders/Private/RayTracing/RayTracingCommon.ush b/Engine/Shaders/Private/RayTracing/RayTracingCommon.ush
--- a/Engine/Shaders/Private/RayTracing/RayTracingCommon.ush
+++ b/Engine/Shaders/Private/RayTracing/RayTracingCommon.ush
@@ -420,6 +420,17 @@ void ApplyCameraRelativeDepthBias(inout float3 RayOrigin, float3 RayDirection, u
 	ApplyPositionBias(RayOrigin, RayDirection, WorldNormal, RayBias);
 }
 
+// This strategy is similar to ApplyPositionBias, but is the version used in the path tracer
+void ApplyRayBias(inout float3 RayOrigin, float HitT, float3 GeoNormal)
+{
+	// Take maximum of position or hit distance to determine "radius" of hit avoidance
+	const float RefValue = max(max(abs(RayOrigin.x), abs(RayOrigin.y)), max(abs(RayOrigin.z), HitT));
+	const uint UlpRadius = 16; // number of floating point ulps to skip around an intersection
+	const float RelativeOffset = asfloat(asuint(RefValue) + UlpRadius) - RefValue;
+	const float BaseOffset = 0.00390625; // 2^-8cm (roughly 39micro meters) (This avoids precision issues very close to the origin where the ulps become too small)
+	RayOrigin += max(BaseOffset, RelativeOffset) * GeoNormal;
+}
+
 float3 TranslatedWorldToTLASWorld(float3 TranslatedWorldPosition)
 {
 	const FDFVector3 TLASPreViewTranslation = MakeDFVector3(View.TLASPreViewTranslationHigh, View.TLASPreViewTranslationLow);
diff --git a/Engine/Shaders/Private/RayTracing/RayTracingOcclusionRGS.usf b/Engine/Shaders/Private/RayTracing/RayTracingOcclusionRGS.usf
--- a/Engine/Shaders/Private/RayTracing/RayTracingOcclusionRGS.usf
+++ b/Engine/Shaders/Private/RayTracing/RayTracingOcclusionRGS.usf
@@ -709,10 +709,6 @@ RAY_TRACING_ENTRY_RAYGEN_OR_INLINE(Occlusion)
 
 	if (all(PixelCoord >= LightScissor.xy) && all(PixelCoord <= LightScissor.zw))
 	{
-		RandomSequence RandSequence;
-		uint LinearIndex = CalcLinearIndex(PixelCoord);
-		RandomSequence_Initialize(RandSequence, LinearIndex, View.StateFrameIndex);
-
 		FLightShaderParameters LightParameters = GetRootLightShaderParameters();
 
 		// Read material data
diff --git a/Engine/Shaders/Private/ShadingFurnaceTest.usf b/Engine/Shaders/Private/ShadingFurnaceTest.usf
--- a/Engine/Shaders/Private/ShadingFurnaceTest.usf
+++ b/Engine/Shaders/Private/ShadingFurnaceTest.usf
@@ -124,7 +124,7 @@ void MainPS(
 				const float2 E = Hammersley(SampleIt, NumSamples[Set], 0);
 
 				// Sample generation
-				float3 L, H;
+				float3 L = 0, H = 0;
 				if (Set == 0)
 				{
 					H = TangentToWorld(ImportanceSampleGGX(E, Pow4(Roughness)).xyz, N);
diff --git a/Engine/Shaders/Shared/PathTracingDefinitions.h b/Engine/Shaders/Shared/PathTracingDefinitions.h
--- a/Engine/Shaders/Shared/PathTracingDefinitions.h
+++ b/Engine/Shaders/Shared/PathTracingDefinitions.h
@@ -72,3 +72,28 @@
 #define PATH_TRACER_VOLUME_USE_ANALYTIC_TRANSMITTANCE       (PATH_TRACER_VOLUME_HOLDOUT_BIT << PATH_TRACER_MAX_VOLUMES)
 #define PATH_TRACER_VOLUME_USE_CLOUD_MAP                    (PATH_TRACER_VOLUME_USE_ANALYTIC_TRANSMITTANCE << 1)
 #define PATH_TRACER_VOLUME_VIEW_CLOUD_MAP_ACCEL             (PATH_TRACER_VOLUME_USE_ANALYTIC_TRANSMITTANCE << 2)
+
+
+// Constants related to debugging
+#define PATH_TRACER_DEBUG_VIZ_RADIANCE					0 
+#define PATH_TRACER_DEBUG_VIZ_WORLD_NORMAL				1 
+#define PATH_TRACER_DEBUG_VIZ_WORLD_SMOOTH_NORMAL		2
+#define PATH_TRACER_DEBUG_VIZ_WORLD_GEO_NORMAL			3
+#define PATH_TRACER_DEBUG_VIZ_BASE_COLOR				4 
+#define PATH_TRACER_DEBUG_VIZ_DIFFUSE_COLOR				5 
+#define PATH_TRACER_DEBUG_VIZ_SPECULAR_COLOR			6 
+#define PATH_TRACER_DEBUG_VIZ_OPACITY					7 
+#define PATH_TRACER_DEBUG_VIZ_METALLIC					8 
+#define PATH_TRACER_DEBUG_VIZ_SPECULAR					9 
+#define PATH_TRACER_DEBUG_VIZ_ROUGHNESS					10 
+#define PATH_TRACER_DEBUG_VIZ_IOR						11
+#define PATH_TRACER_DEBUG_VIZ_SHADING_MODEL				12 
+#define PATH_TRACER_DEBUG_VIZ_LIGHTING_CHANNEL_MASK		13  
+#define PATH_TRACER_DEBUG_VIZ_CUSTOM_DATA0				14  
+#define PATH_TRACER_DEBUG_VIZ_CUSTOM_DATA1				15
+#define PATH_TRACER_DEBUG_VIZ_WORLD_POSITION			16  
+#define PATH_TRACER_DEBUG_VIZ_PRIMARY_RAYS				17  
+#define PATH_TRACER_DEBUG_VIZ_WORLD_TANGENT				18 
+#define PATH_TRACER_DEBUG_VIZ_ANISOTROPY				19 
+#define PATH_TRACER_DEBUG_VIZ_LIGHT_GRID_COUNT          20
+#define PATH_TRACER_DEBUG_VIZ_DECAL_GRID_COUNT 			21
diff --git a/Engine/Source/Runtime/Renderer/Private/PathTracing.cpp b/Engine/Source/Runtime/Renderer/Private/PathTracing.cpp
--- a/Engine/Source/Runtime/Renderer/Private/PathTracing.cpp
+++ b/Engine/Source/Runtime/Renderer/Private/PathTracing.cpp
@@ -46,6 +46,19 @@ TAutoConsoleVariable<bool> CVarPathTracingExperimental(
 	ECVF_RenderThreadSafe | ECVF_ReadOnly
 );
 
+TAutoConsoleVariable<FString> CVarPathTracingDebugVisualize(
+	TEXT("r.PathTracing.DebugVisualize"),
+	TEXT(""),
+	TEXT("Enables a debug mode for the path tracer which can visualize payload fields and perform basic lighting without sample accumulation. (default: "")")
+);
+
+TAutoConsoleVariable<int32> CVarPathTracingDebugVisualizeLightSamples(
+	TEXT("r.PathTracing.DebugVisualizeLightSamples"),
+	1,
+	TEXT("Number of light samples for the path tracing light sampling debug mode (default: 1)"),
+	ECVF_RenderThreadSafe
+);
+
 
 TAutoConsoleVariable<bool> CVarPathTracingCompaction(
 	TEXT("r.PathTracing.Compaction"),
@@ -356,27 +369,6 @@ TAutoConsoleVariable<int32> CVarPathTracingLightGridMaxCount(
 	ECVF_RenderThreadSafe
 );
 
-TAutoConsoleVariable<int32> CVarPathTracingLightGridVisualize(
-	TEXT("r.PathTracing.LightGridVisualize"),
-	0,
-	TEXT("Enables a visualization mode of the light grid density where red indicates the maximum light count has been reached (default = 0)\n")
-	TEXT("0: off (default)\n")
-	TEXT("1: light count heatmap (red - close to overflow, increase r.PathTracing.LightGridMaxCount)\n")
-	TEXT("2: unique light lists (colors are a function of which lights occupy each cell)\n")
-	TEXT("3: area light visualization (green: point light sources only, blue: some area light sources)\n"),
-	ECVF_RenderThreadSafe
-);
-
-TAutoConsoleVariable<int32> CVarPathTracingDecalGridVisualize(
-	TEXT("r.PathTracing.DecalGrid.Visualize"),
-	0,
-	TEXT("Enables a visualization mode of the decal grid density where red indicates the maximum decal count has been reached (default = 0)\n")
-	TEXT("0: off (default)\n")
-	TEXT("1: decal count heatmap (red - close to overflow, increase r.RayTracing.DecalGrid.MaxCount)\n")
-	TEXT("2: unique decal lists (colors are a function of which decals occupy each cell)\n"),
-	ECVF_RenderThreadSafe
-);
-
 TAutoConsoleVariable<bool> CVarPathTracingUseDBuffer(
 	TEXT("r.PathTracing.UseDBuffer"),
 	true,
@@ -525,8 +517,6 @@ BEGIN_SHADER_PARAMETER_STRUCT(FPathTracingData, )
 	SHADER_PARAMETER(uint32, ApproximateCaustics)
 	SHADER_PARAMETER(uint32, EnableCameraBackfaceCulling)
 	SHADER_PARAMETER(uint32, SamplerType)
-	SHADER_PARAMETER(uint32, VisualizeLightGrid)
-	SHADER_PARAMETER(uint32, VisualizeDecalGrid)
 	SHADER_PARAMETER(uint32, EnableDBuffer)
 	SHADER_PARAMETER(uint32, VolumeFlags)
 	SHADER_PARAMETER(uint32, EnabledDirectLightingContributions)   // PATHTRACER_CONTRIBUTION_*
@@ -580,8 +570,6 @@ struct FPathTracingConfig
 			PathTracingData.SamplerType != Other.PathTracingData.SamplerType ||
 			PathTracingData.ApproximateCaustics != Other.PathTracingData.ApproximateCaustics ||
 			PathTracingData.EnableCameraBackfaceCulling != Other.PathTracingData.EnableCameraBackfaceCulling ||
-			PathTracingData.VisualizeLightGrid != Other.PathTracingData.VisualizeLightGrid ||
-			PathTracingData.VisualizeDecalGrid != Other.PathTracingData.VisualizeDecalGrid ||
 			PathTracingData.EnableDBuffer != Other.PathTracingData.EnableDBuffer ||
 			PathTracingData.MaxPathIntensity != Other.PathTracingData.MaxPathIntensity ||
 			PathTracingData.FilterWidth != Other.PathTracingData.FilterWidth ||
@@ -707,7 +695,41 @@ struct FPathTracingState
 	uint32_t FrameIndex = 0;
 };
 
-uint32 PackRG16(float In0, float In1);
+int GetPathTracingDebugVisualizationMode() {
+	static TMap<FName, uint32> PathTracingDebugModes;
+	if (PathTracingDebugModes.Num() == 0)
+	{
+		PathTracingDebugModes.Emplace(FName(TEXT("Radiance"))			, PATH_TRACER_DEBUG_VIZ_RADIANCE				);
+		PathTracingDebugModes.Emplace(FName(TEXT("WorldNormal"))		, PATH_TRACER_DEBUG_VIZ_WORLD_NORMAL			);
+		PathTracingDebugModes.Emplace(FName(TEXT("WorldSmoothNormal"))	, PATH_TRACER_DEBUG_VIZ_WORLD_SMOOTH_NORMAL		);
+		PathTracingDebugModes.Emplace(FName(TEXT("WorldGeoNormal"))		, PATH_TRACER_DEBUG_VIZ_WORLD_GEO_NORMAL		);
+		PathTracingDebugModes.Emplace(FName(TEXT("BaseColor"))	    	, PATH_TRACER_DEBUG_VIZ_BASE_COLOR				);
+		PathTracingDebugModes.Emplace(FName(TEXT("DiffuseColor"))	   	, PATH_TRACER_DEBUG_VIZ_DIFFUSE_COLOR			);
+		PathTracingDebugModes.Emplace(FName(TEXT("SpecularColor"))	   	, PATH_TRACER_DEBUG_VIZ_SPECULAR_COLOR			);
+		PathTracingDebugModes.Emplace(FName(TEXT("Opacity"))	     	, PATH_TRACER_DEBUG_VIZ_OPACITY					);
+		PathTracingDebugModes.Emplace(FName(TEXT("Metallic"))	     	, PATH_TRACER_DEBUG_VIZ_METALLIC				);
+		PathTracingDebugModes.Emplace(FName(TEXT("Specular"))	     	, PATH_TRACER_DEBUG_VIZ_SPECULAR				);
+		PathTracingDebugModes.Emplace(FName(TEXT("Roughness"))	     	, PATH_TRACER_DEBUG_VIZ_ROUGHNESS				);
+		PathTracingDebugModes.Emplace(FName(TEXT("IOR"))	    	 	, PATH_TRACER_DEBUG_VIZ_IOR						);
+		PathTracingDebugModes.Emplace(FName(TEXT("ShadingModel"))  	 	, PATH_TRACER_DEBUG_VIZ_SHADING_MODEL			);
+		PathTracingDebugModes.Emplace(FName(TEXT("LightingChannelMask")), PATH_TRACER_DEBUG_VIZ_LIGHTING_CHANNEL_MASK 	);
+		PathTracingDebugModes.Emplace(FName(TEXT("CustomData0"))		, PATH_TRACER_DEBUG_VIZ_CUSTOM_DATA0 			);
+		PathTracingDebugModes.Emplace(FName(TEXT("CustomData1"))		, PATH_TRACER_DEBUG_VIZ_CUSTOM_DATA1 			);
+		PathTracingDebugModes.Emplace(FName(TEXT("WorldPosition"))		, PATH_TRACER_DEBUG_VIZ_WORLD_POSITION 			);
+		PathTracingDebugModes.Emplace(FName(TEXT("PrimaryRays"))		, PATH_TRACER_DEBUG_VIZ_PRIMARY_RAYS 			);
+		PathTracingDebugModes.Emplace(FName(TEXT("WorldTangent"))		, PATH_TRACER_DEBUG_VIZ_WORLD_TANGENT 			);
+		PathTracingDebugModes.Emplace(FName(TEXT("Anisotropy"))	    	, PATH_TRACER_DEBUG_VIZ_ANISOTROPY 				);
+		PathTracingDebugModes.Emplace(FName(TEXT("LightGridCount"))	   	, PATH_TRACER_DEBUG_VIZ_LIGHT_GRID_COUNT		);
+		PathTracingDebugModes.Emplace(FName(TEXT("DecalGridCount"))	   	, PATH_TRACER_DEBUG_VIZ_DECAL_GRID_COUNT		);
+	}
+	FString DebugMode = CVarPathTracingDebugVisualize.GetValueOnRenderThread();
+	if (!DebugMode.IsEmpty())
+	{
+		return PathTracingDebugModes.FindRef(FName(*DebugMode, FNAME_Find), -1);
+	}
+	// not doing debugging
+	return -1;
+}
 
 namespace PathTracing
 {
@@ -726,6 +748,11 @@ namespace PathTracing
 	{
 		return GEnableReferenceDOF < 0 ? View.FinalPostProcessSettings.PathTracingEnableReferenceDOF != 0 : GEnableReferenceDOF != 0;
 	}
+
+	bool NeedsAntiAliasing(const FViewInfo& View)
+	{
+		return GetPathTracingDebugVisualizationMode() >= 0;
+	}
 }
 
 // This function prepares the portion of shader arguments that may involve invalidating the path traced state
@@ -760,8 +787,6 @@ static void PreparePathTracingData(const FScene* Scene, const FViewInfo& View, F
 	PathTracingData.ApproximateCaustics = CVarPathTracingApproximateCaustics.GetValueOnRenderThread();
 	PathTracingData.EnableCameraBackfaceCulling = CVarPathTracingEnableCameraBackfaceCulling.GetValueOnRenderThread();
 	PathTracingData.SamplerType = CVarPathTracingSamplerType.GetValueOnRenderThread();
-	PathTracingData.VisualizeLightGrid = CVarPathTracingLightGridVisualize.GetValueOnRenderThread();
-	PathTracingData.VisualizeDecalGrid = CVarPathTracingDecalGridVisualize.GetValueOnRenderThread();
 	PathTracingData.FilterWidth = CVarPathTracingFilterWidth.GetValueOnRenderThread();
 	PathTracingData.CameraFocusDistance = 0;
 	PathTracingData.CameraLensRadius = FVector2f::ZeroVector;
@@ -1195,6 +1220,77 @@ class FPathTracingRG : public FGlobalShader
 };
 IMPLEMENT_GLOBAL_SHADER(FPathTracingRG, "/Engine/Private/PathTracing/PathTracing.usf", "PathTracingMainRG", SF_RayGen);
 
+class FPathTracingDebugRG : public FGlobalShader
+{
+	DECLARE_GLOBAL_SHADER(FPathTracingDebugRG)
+	SHADER_USE_ROOT_PARAMETER_STRUCT(FPathTracingDebugRG, FGlobalShader)
+
+
+	static bool ShouldCompilePermutation(const FGlobalShaderPermutationParameters& Parameters)
+	{
+		return ShouldCompilePathTracingShadersForProject(Parameters.Platform);
+	}
+
+	static void ModifyCompilationEnvironment(const FGlobalShaderPermutationParameters& Parameters, FShaderCompilerEnvironment& OutEnvironment)
+	{
+		FGlobalShader::ModifyCompilationEnvironment(Parameters, OutEnvironment);
+		OutEnvironment.SetDefine(TEXT("USE_RECT_LIGHT_TEXTURES"), 1);
+		OutEnvironment.CompilerFlags.Add(CFLAG_WarningsAsErrors);
+	}
+
+	static ERayTracingPayloadType GetRayTracingPayloadType(const int32 PermutationId)
+	{
+		return ERayTracingPayloadType::PathTracingMaterial | ERayTracingPayloadType::Decals;
+	}
+
+	BEGIN_SHADER_PARAMETER_STRUCT(FParameters, )
+		SHADER_PARAMETER_RDG_TEXTURE_UAV(RWTexture2D<float4>, RWSceneColor)
+		SHADER_PARAMETER_RDG_TEXTURE_UAV(RWTexture2D<float>,  RWSceneDepth)
+		SHADER_PARAMETER_RDG_BUFFER_SRV(RaytracingAccelerationStructure, TLAS)
+		SHADER_PARAMETER_STRUCT_REF(FViewUniformShaderParameters, ViewUniformBuffer)
+		SHADER_PARAMETER(int32, DebugMode)
+
+		// scene lights
+		SHADER_PARAMETER_RDG_BUFFER_SRV(StructuredBuffer<FPathTracingLight>, SceneLights)
+		SHADER_PARAMETER(uint32, SceneLightCount)
+		SHADER_PARAMETER(uint32, SceneVisibleLightCount)
+		SHADER_PARAMETER(uint32, NumLightSamples)
+		SHADER_PARAMETER_STRUCT_INCLUDE(FPathTracingLightGrid, LightGridParameters)
+	
+		// scene decals
+		SHADER_PARAMETER_RDG_UNIFORM_BUFFER(FRayTracingDecals, DecalParameters)
+
+		// Skylight
+		SHADER_PARAMETER_STRUCT_INCLUDE(FPathTracingSkylight, SkylightParameters)
+	END_SHADER_PARAMETER_STRUCT()
+};
+IMPLEMENT_GLOBAL_SHADER(FPathTracingDebugRG, "/Engine/Private/PathTracing/PathTracingDebug.usf", "PathTracingDebugRG", SF_RayGen);
+
+
+class FPathTracingCopyDepthPS : public FGlobalShader
+{
+	DECLARE_GLOBAL_SHADER(FPathTracingCopyDepthPS)
+	SHADER_USE_PARAMETER_STRUCT(FPathTracingCopyDepthPS, FGlobalShader)
+
+	static bool ShouldCompilePermutation(const FGlobalShaderPermutationParameters& Parameters)
+	{
+		return ShouldCompileRayTracingShadersForProject(Parameters.Platform);
+	}
+
+	static void ModifyCompilationEnvironment(const FGlobalShaderPermutationParameters& Parameters, FShaderCompilerEnvironment& OutEnvironment)
+	{
+		OutEnvironment.CompilerFlags.Add(CFLAG_WarningsAsErrors);
+	}
+
+	BEGIN_SHADER_PARAMETER_STRUCT(FParameters, )
+		SHADER_PARAMETER_RDG_TEXTURE_SRV(Texture2D<float>, DepthTexture)
+		SHADER_PARAMETER_STRUCT_REF(FViewUniformShaderParameters, ViewUniformBuffer)
+		RENDER_TARGET_BINDING_SLOTS()
+	END_SHADER_PARAMETER_STRUCT()
+};
+IMPLEMENT_SHADER_TYPE(, FPathTracingCopyDepthPS, TEXT("/Engine/Private/PathTracing/PathTracingCopyDepth.usf"), TEXT("CopyDepth"), SF_Pixel);
+
+
 class FPathTracingInitExtinctionCoefficientRG : public FGlobalShader
 {
 	DECLARE_GLOBAL_SHADER(FPathTracingInitExtinctionCoefficientRG)
@@ -2523,9 +2619,27 @@ RENDERER_API void PrepareLightGrid(FRDGBuilder& GraphBuilder, ERHIFeatureLevel::
 	}
 }
 
-void SetLightParameters(FRDGBuilder& GraphBuilder, FPathTracingRG::FParameters* PassParameters, FScene* Scene, const FViewInfo& View, bool UseMISCompensation)
+uint32 PackRG16(float In0, float In1);
+
+void SetLightParameters(
+	FRDGBuilder& GraphBuilder,
+	FScene* Scene,
+	const FViewInfo& View,
+	const bool UseMISCompensation,
+	const bool bUseAtmosphere,
+	// output args
+	FPathTracingSkylight* SkylightParameters,
+	FPathTracingLightGrid* LightGridParameters,
+	uint32* SceneVisibleLightCount,
+	uint32* SceneLightCount,
+	FRDGBufferSRVRef* SceneLights
+)
 {
-	PassParameters->SceneVisibleLightCount = 0;
+	check(SkylightParameters != nullptr);
+	check(SceneVisibleLightCount != nullptr);
+	check(SceneLightCount != nullptr);
+	check(SceneLights != nullptr);
+	*SceneVisibleLightCount = 0;
 
 	// Lights
 	uint32 MaxNumLights = 1 + Scene->Lights.Num(); // upper bound
@@ -2535,9 +2649,8 @@ void SetLightParameters(FRDGBuilder& GraphBuilder, FPathTracingRG::FParameters*
 
 	// Prepend SkyLight to light buffer since it is not part of the regular light list
 	// skylight should be excluded if we are using the reference atmosphere calculation (don't bother checking again if an atmosphere is present)
-	const bool bUseAtmosphere = (PassParameters->PathTracingData.VolumeFlags & PATH_TRACER_VOLUME_ENABLE_ATMOSPHERE) != 0;
 	const bool bEnableSkydome = !bUseAtmosphere;
-	if (PrepareSkyTexture(GraphBuilder, Scene, View, bEnableSkydome, UseMISCompensation, &PassParameters->SkylightParameters))
+	if (PrepareSkyTexture(GraphBuilder, Scene, View, bEnableSkydome, UseMISCompensation, SkylightParameters))
 	{
 		check(Scene->SkyLight != nullptr);
 		FPathTracingLight& DestLight = Lights[NumLights++];
@@ -2558,7 +2671,7 @@ void SetLightParameters(FRDGBuilder& GraphBuilder, FPathTracingRG::FParameters*
 			// The one exception to this case is if the sky atmo has been marked as holdout.
 
 			// Also allow seeing just the sky via a cvar for debugging purposes
-			PassParameters->SceneVisibleLightCount = 1;
+			*SceneVisibleLightCount = 1;
 		}
 	}
 
@@ -2628,7 +2741,7 @@ void SetLightParameters(FRDGBuilder& GraphBuilder, FPathTracingRG::FParameters*
 	{
 		// show directional lights when atmosphere is enabled and not marked as holdout
 		// NOTE: there cannot be any skydome in this case
-		PassParameters->SceneVisibleLightCount = NumLights;
+		*SceneVisibleLightCount = NumLights;
 	}
 
 	uint32 NumInfiniteLights = NumLights;
@@ -2737,21 +2850,21 @@ void SetLightParameters(FRDGBuilder& GraphBuilder, FPathTracingRG::FParameters*
 		}
 	}
 
-	PassParameters->SceneLightCount = NumLights;
+	*SceneLightCount = NumLights;
 	{
 		// Upload the buffer of lights to the GPU
 		uint32 NumCopyLights = FMath::Max(1u, NumLights); // need at least one since zero-sized buffers are not allowed
 		size_t DataSize = sizeof(FPathTracingLight) * NumCopyLights;
-		PassParameters->SceneLights = GraphBuilder.CreateSRV(FRDGBufferSRVDesc(CreateStructuredBuffer(GraphBuilder, TEXT("PathTracer.LightsBuffer"), sizeof(FPathTracingLight), NumCopyLights, Lights, DataSize, ERDGInitialDataFlags::NoCopy)));
+		*SceneLights = GraphBuilder.CreateSRV(FRDGBufferSRVDesc(CreateStructuredBuffer(GraphBuilder, TEXT("PathTracer.LightsBuffer"), sizeof(FPathTracingLight), NumCopyLights, Lights, DataSize, ERDGInitialDataFlags::NoCopy)));
 	}
 
 	if (CVarPathTracingVisibleLights.GetValueOnRenderThread() == 1)
 	{
 		// make all lights in the scene visible
-		PassParameters->SceneVisibleLightCount = PassParameters->SceneLightCount;
+		*SceneVisibleLightCount = *SceneLightCount;
 	}
 
-	PrepareLightGrid(GraphBuilder, View.FeatureLevel, &PassParameters->LightGridParameters, Lights, NumLights, NumInfiniteLights, PassParameters->SceneLights);
+	PrepareLightGrid(GraphBuilder, View.FeatureLevel, LightGridParameters, Lights, NumLights, NumInfiniteLights, *SceneLights);
 }
 
 class FPathTracingCompositorPS : public FGlobalShader
@@ -2809,17 +2922,25 @@ void FDeferredShadingSceneRenderer::PreparePathTracing(const FSceneViewFamily& V
 	if (ViewFamily.EngineShowFlags.PathTracing
 		&& ShouldCompilePathTracingShadersForProject(ViewFamily.GetShaderPlatform()))
 	{
-		// Declare all RayGen shaders that require material closest hit shaders to be bound
-		FPathTracingRG::FPermutationDomain PermutationVector = GetPathTracingRGPermutation(Scene);
+		if (GetPathTracingDebugVisualizationMode() >= 0)
 		{
-			auto RayGenShader = GetGlobalShaderMap(ViewFamily.GetShaderPlatform())->GetShader<FPathTracingRG>(PermutationVector);
+			auto RayGenShader = GetGlobalShaderMap(ViewFamily.GetShaderPlatform())->GetShader<FPathTracingDebugRG>();
 			OutRayGenShaders.Add(RayGenShader.GetRayTracingShader());
 		}
-
+		else
 		{
-			auto RayGenShader = GetGlobalShaderMap(ViewFamily.GetShaderPlatform())->GetShader<FPathTracingInitExtinctionCoefficientRG>();
-			OutRayGenShaders.Add(RayGenShader.GetRayTracingShader());
-		}		
+			// Declare all RayGen shaders that require material closest hit shaders to be bound
+			FPathTracingRG::FPermutationDomain PermutationVector = GetPathTracingRGPermutation(Scene);
+			{
+				auto RayGenShader = GetGlobalShaderMap(ViewFamily.GetShaderPlatform())->GetShader<FPathTracingRG>(PermutationVector);
+				OutRayGenShaders.Add(RayGenShader.GetRayTracingShader());
+			}
+
+			{
+				auto RayGenShader = GetGlobalShaderMap(ViewFamily.GetShaderPlatform())->GetShader<FPathTracingInitExtinctionCoefficientRG>();
+				OutRayGenShaders.Add(RayGenShader.GetRayTracingShader());
+			}
+		}
 	}
 }
 
@@ -2962,6 +3083,92 @@ void FDeferredShadingSceneRenderer::RenderPathTracing(
 		return;
 	}
 
+
+	if (int DebugMode = GetPathTracingDebugVisualizationMode(); DebugMode >= 0)
+	{
+		const int32 DispatchResX = View.ViewRect.Size().X;
+		const int32 DispatchResY = View.ViewRect.Size().Y;
+
+		// simplified pass for debugging purposes
+		RDG_GPU_STAT_SCOPE(GraphBuilder, PathTracing);
+		FRDGTextureDesc DepthDesc = FRDGTextureDesc::Create2D(
+			View.ViewRect.Size(),
+			PF_R32_FLOAT,
+			FClearValueBinding::None,
+			TexCreate_ShaderResource | TexCreate_UAV);
+		FRDGTextureRef DepthTexture    = GraphBuilder.CreateTexture(DepthDesc, TEXT("PathTracer.Depth"));
+
+
+		FPathTracingDebugRG::FParameters* PassParameters = GraphBuilder.AllocParameters<FPathTracingDebugRG::FParameters>();
+		PassParameters->RWSceneColor = GraphBuilder.CreateUAV(SceneColorOutputTexture);
+		PassParameters->RWSceneDepth = GraphBuilder.CreateUAV(DepthTexture);
+		PassParameters->TLAS = Scene->RayTracingScene.GetLayerView(ERayTracingSceneLayer::Base);
+		PassParameters->ViewUniformBuffer = View.ViewUniformBuffer;
+		PassParameters->DebugMode = DebugMode;
+
+		PassParameters->NumLightSamples = FMath::Max(CVarPathTracingDebugVisualizeLightSamples.GetValueOnRenderThread(), 1);
+		SetLightParameters(
+			GraphBuilder,
+			Scene,
+			View,
+			true,
+			false,
+			&PassParameters->SkylightParameters,
+			&PassParameters->LightGridParameters,
+			&PassParameters->SceneVisibleLightCount,
+			&PassParameters->SceneLightCount,
+			&PassParameters->SceneLights
+		);
+
+		PassParameters->DecalParameters = View.RayTracingDecalUniformBuffer;
+
+		TShaderMapRef<FPathTracingDebugRG> RayGenShader(View.ShaderMap);
+		ClearUnusedGraphResources(RayGenShader, PassParameters);
+		GraphBuilder.AddPass(
+			RDG_EVENT_NAME("Path Tracing"),
+		  	PassParameters,
+		  	ERDGPassFlags::Compute,
+			[PassParameters, RayGenShader, &View, DispatchResX, DispatchResY](FRHICommandList& RHICmdList)
+			{
+				FRHIBatchedShaderParameters& GlobalResources = RHICmdList.GetScratchShaderParameters();
+				SetShaderParameters(GlobalResources, RayGenShader, *PassParameters);
+				RHICmdList.RayTraceDispatch(
+					View.RayTracingMaterialPipeline,
+					RayGenShader.GetRayTracingShader(),
+					View.RayTracingSBT, GlobalResources,
+					DispatchResX, DispatchResY
+				);
+			});
+
+		{
+			FPathTracingCopyDepthPS::FParameters* DisplayParameters = GraphBuilder.AllocParameters<FPathTracingCopyDepthPS::FParameters>();
+			DisplayParameters->ViewUniformBuffer = View.ViewUniformBuffer;
+			DisplayParameters->DepthTexture = GraphBuilder.CreateSRV(FRDGTextureSRVDesc::Create(DepthTexture));
+			//DisplayParameters->RenderTargets[0] = FRenderTargetBinding(SceneColorOutputTexture, ERenderTargetLoadAction::ELoad);
+			DisplayParameters->RenderTargets.DepthStencil = FDepthStencilBinding(SceneDepthOutputTexture,  ERenderTargetLoadAction::ELoad, ERenderTargetLoadAction::ENoAction, FExclusiveDepthStencil::DepthWrite_StencilNop);
+			FScreenPassTextureViewport Viewport(SceneColorOutputTexture, View.ViewRect);
+			TShaderMapRef<FPathTracingCopyDepthPS> PixelShader(View.ShaderMap);
+			TShaderMapRef<FScreenPassVS> VertexShader(View.ShaderMap);
+			FRHIBlendState* BlendState = FScreenPassPipelineState::FDefaultBlendState::GetRHI();
+			FRHIDepthStencilState* DepthStencilState = TStaticDepthStencilState<true /* bEnableDepthWrite */, CF_Always>::GetRHI();
+			AddDrawScreenPass(
+				GraphBuilder,
+				RDG_EVENT_NAME("Path Tracer Copy Depth", View.ViewRect.Size().X, View.ViewRect.Size().Y),
+				View,
+				Viewport,
+				Viewport,
+				VertexShader,
+				PixelShader,
+				BlendState,
+				DepthStencilState,
+				DisplayParameters);
+		}
+
+		return;
+	}
+
+
+
 	FPathTracingConfig Config = {};
 
 	// Get current value of MaxSPP and reset render if it has changed
@@ -3525,7 +3732,18 @@ void FDeferredShadingSceneRenderer::RenderPathTracing(
 							{
 								// upload sky/lights data
 								RDG_GPU_MASK_SCOPE(GraphBuilder, GPUMask); // make sure this happens on all GPUs we will be rendering on
-								SetLightParameters(GraphBuilder, PassParameters, Scene, View, Config.UseMISCompensation);
+								SetLightParameters(
+									GraphBuilder,
+									Scene,
+									View,
+									Config.UseMISCompensation,
+									(Config.PathTracingData.VolumeFlags & PATH_TRACER_VOLUME_ENABLE_ATMOSPHERE) != 0,
+									&PassParameters->SkylightParameters,
+									&PassParameters->LightGridParameters,
+									&PassParameters->SceneVisibleLightCount,
+									&PassParameters->SceneLightCount,
+									&PassParameters->SceneLights
+								);
 							}
 							else
 							{
@@ -4103,6 +4321,11 @@ namespace PathTracing
 	{
 		return false;
 	}
+
+	bool NeedsAntiAliasing(const FViewInfo& View)
+	{
+		return false;
+	}
 }
 
 #endif
diff --git a/Engine/Source/Runtime/Renderer/Private/PathTracing.h b/Engine/Source/Runtime/Renderer/Private/PathTracing.h
--- a/Engine/Source/Runtime/Renderer/Private/PathTracing.h
+++ b/Engine/Source/Runtime/Renderer/Private/PathTracing.h
@@ -38,4 +38,5 @@ namespace PathTracing
 	bool UsesDecals(const FSceneViewFamily& ViewFamily);
 	bool UsesReferenceAtmosphere(const FViewInfo& View);
 	bool UsesReferenceDOF(const FViewInfo& View);
+	bool NeedsAntiAliasing(const FViewInfo& View);
 }
diff --git a/Engine/Source/Runtime/Renderer/Private/SceneRendering.cpp b/Engine/Source/Runtime/Renderer/Private/SceneRendering.cpp
--- a/Engine/Source/Runtime/Renderer/Private/SceneRendering.cpp
+++ b/Engine/Source/Runtime/Renderer/Private/SceneRendering.cpp
@@ -3316,8 +3316,8 @@ void FSceneRenderer::PrepareViewRectsForRendering(FRHICommandListImmediate& RHIC
 		{
 			const bool bWillApplyTemporalAA = (IsPostProcessingEnabled(View) || View.bIsPlanarReflection || View.bSceneCaptureMainViewJitter)
 #if RHI_RAYTRACING
-				// path tracer does its own anti-aliasing
-				&& (!ViewFamily.EngineShowFlags.PathTracing)
+				// path tracer does its own anti-aliasing (unless it specifically requests it, such as for the debug mode)
+				&& (!ViewFamily.EngineShowFlags.PathTracing || PathTracing::NeedsAntiAliasing(View))
 #endif
 			;
 
