Refine rect light culling in light grid using planes.

diff --git a/Engine/Shaders/Private/LightGridInjection.usf b/Engine/Shaders/Private/LightGridInjection.usf
--- a/Engine/Shaders/Private/LightGridInjection.usf
+++ b/Engine/Shaders/Private/LightGridInjection.usf
@@ -42,6 +42,9 @@ uint MegaLightsSupportedStartIndex;
 
 StructuredBuffer<float4> LightViewSpacePositionAndRadius;
 StructuredBuffer<float4> LightViewSpaceDirAndPreprocAngle;
+StructuredBuffer<float4> LightViewSpaceRectPlanes;
+
+#define NUM_PLANES_PER_RECT_LIGHT 4
 
 #define REFINE_SPOTLIGHT_BOUNDS 1
 
@@ -243,12 +246,13 @@ void LightGridInjectionCS(
 			{
 				float4 ViewSpaceDirAndPreprocAngle = LightViewSpaceDirAndPreprocAngle[LocalLightIndex];
 				bool bIsRectLight = asuint(ViewSpaceDirAndPreprocAngle.w) & 0x1;
-				bool bIsTexturedLight = asuint(ViewSpaceDirAndPreprocAngle.w) & 0x2;
+				bool bUseTightRectLightCulling = asuint(ViewSpaceDirAndPreprocAngle.w) & 0x2;
+				bool bIsTexturedLight = asuint(ViewSpaceDirAndPreprocAngle.w) & 0x4;
 
 				bool bPassSpotLightTest = true;
 #if REFINE_SPOTLIGHT_BOUNDS
 				{
-					float TanConeAngle = asfloat(asuint(ViewSpaceDirAndPreprocAngle.w) & 0xFFFFFFFC);
+					float TanConeAngle = asfloat(asuint(ViewSpaceDirAndPreprocAngle.w) & 0xFFFFFFF8);
 
 					// Set to 0 for non-acute cones, or non-spot lights.
 					if (TanConeAngle > 0.0f)
@@ -268,7 +272,16 @@ void LightGridInjectionCS(
 
 						bPassRectLightTest = !AabbOutsidePlane(D, ViewTileExtent, Plane);
 
-						// TODO: Refine using spot light test when BarnDoorAngle < 90 degrees and BarnDoorLength > 0
+						if(bPassRectLightTest && bUseTightRectLightCulling)
+						{
+							for (uint PlaneIndex = 0; PlaneIndex < NUM_PLANES_PER_RECT_LIGHT; PlaneIndex++)
+							{
+								if (AabbOutsidePlane(ViewTileCenter, ViewTileExtent, LightViewSpaceRectPlanes[LocalLightIndex * NUM_PLANES_PER_RECT_LIGHT + PlaneIndex]))
+								{
+									bPassRectLightTest = false;
+								}
+							}
+						}
 					}
 				}
 #endif // REFINE_RECTLIGHT_BOUNDS
diff --git a/Engine/Source/Runtime/Renderer/Private/LightGridInjection.cpp b/Engine/Source/Runtime/Renderer/Private/LightGridInjection.cpp
--- a/Engine/Source/Runtime/Renderer/Private/LightGridInjection.cpp
+++ b/Engine/Source/Runtime/Renderer/Private/LightGridInjection.cpp
@@ -17,6 +17,7 @@ LightGridInjection.cpp
 #include "SceneUtils.h"
 #include "PostProcess/SceneRenderTargets.h"
 #include "LightSceneInfo.h"
+#include "RectLightSceneProxy.h"
 #include "GlobalShader.h"
 #include "SceneRendering.h"
 #include "DeferredShadingRenderer.h"
@@ -213,6 +214,7 @@ public:
 		SHADER_PARAMETER_RDG_BUFFER_UAV(RWStructuredBuffer<uint>, RWCulledLightLinks)
 		SHADER_PARAMETER_RDG_BUFFER_SRV(StructuredBuffer<float4>, LightViewSpacePositionAndRadius)
 		SHADER_PARAMETER_RDG_BUFFER_SRV(StructuredBuffer<float4>, LightViewSpaceDirAndPreprocAngle)
+		SHADER_PARAMETER_RDG_BUFFER_SRV(StructuredBuffer<float4>, LightViewSpaceRectPlanes)
 		SHADER_PARAMETER_RDG_BUFFER_SRV(StructuredBuffer<float4>, ForwardLocalLightBuffer)
 
 		SHADER_PARAMETER(FIntVector, CulledGridSize)
@@ -464,6 +466,49 @@ static void PackLocalLightData(
 	Out.RectDataAndVirtualShadowMapIdOrPrevLocalLightIndex	= FVector4f(FMath::AsFloat(RectPackedX), FMath::AsFloat(RectPackedY), FMath::AsFloat(RectPackedZ), FMath::AsFloat(VirtualShadowMapIdAndPrevLocalLightIndex));
 }
 
+static const uint32 NUM_PLANES_PER_RECT_LIGHT = 4;
+
+static void CalculateRectLightCullingPlanes(const FRectLightSceneProxy* RectProxy, TArray<FPlane, TInlineAllocator<NUM_PLANES_PER_RECT_LIGHT>>& OutPlanes)
+{
+	const float BarnMaxAngle = GetRectLightBarnDoorMaxAngle();
+	const float AngleRad = FMath::DegreesToRadians(FMath::Clamp(RectProxy->BarnDoorAngle, 0.f, BarnMaxAngle));
+	const float BarnDepth = FMath::Cos(AngleRad) * RectProxy->BarnDoorLength;
+	const float BarnExtent = FMath::Sin(AngleRad) * RectProxy->BarnDoorLength;
+
+	FVector Corners[8];
+
+	// corners position based on rect source size and barn door parameters (frustum shape)
+	Corners[0] = FVector(0.0f, +0.5f * RectProxy->SourceWidth, +0.5f * RectProxy->SourceHeight);
+	Corners[1] = FVector(0.0f, +0.5f * RectProxy->SourceWidth, -0.5f * RectProxy->SourceHeight);
+	Corners[2] = FVector(BarnDepth, +0.5f * RectProxy->SourceWidth + BarnExtent, +0.5f * RectProxy->SourceHeight + BarnExtent);
+	Corners[3] = FVector(BarnDepth, +0.5f * RectProxy->SourceWidth + BarnExtent, -0.5f * RectProxy->SourceHeight - BarnExtent);
+	Corners[4] = FVector(0.0f, -0.5f * RectProxy->SourceWidth, +0.5f * RectProxy->SourceHeight);
+	Corners[5] = FVector(0.0f, -0.5f * RectProxy->SourceWidth, -0.5f * RectProxy->SourceHeight);
+	Corners[6] = FVector(BarnDepth, -0.5f * RectProxy->SourceWidth - BarnExtent, +0.5f * RectProxy->SourceHeight + BarnExtent);
+	Corners[7] = FVector(BarnDepth, -0.5f * RectProxy->SourceWidth - BarnExtent, -0.5f * RectProxy->SourceHeight - BarnExtent);
+
+	// adjust corners to account for penumbra
+	auto ApplyPenumbraAdjustment = [RectProxy](const FVector& A, FVector& B) -> FVector
+		{
+			FVector Tmp = B - A;
+			double Len = Tmp.Length();
+			Tmp.Normalize();
+			return A + Tmp * FMath::Max<double>(RectProxy->Radius, Len);
+		};
+
+	Corners[7] = ApplyPenumbraAdjustment(Corners[0], Corners[7]);
+	Corners[6] = ApplyPenumbraAdjustment(Corners[1], Corners[6]);
+	Corners[3] = ApplyPenumbraAdjustment(Corners[4], Corners[3]);
+	Corners[2] = ApplyPenumbraAdjustment(Corners[5], Corners[2]);
+
+	OutPlanes.Add(FPlane(Corners[1], Corners[0], Corners[3])); // right
+	OutPlanes.Add(FPlane(Corners[5], Corners[7], Corners[4])); // left
+	OutPlanes.Add(FPlane(Corners[4], Corners[6], Corners[0])); // top
+	OutPlanes.Add(FPlane(Corners[1], Corners[3], Corners[5])); // bottom
+
+	check(OutPlanes.Num() == NUM_PLANES_PER_RECT_LIGHT);
+}
+
 FComputeLightGridOutput FSceneRenderer::ComputeLightGrid(FRDGBuilder& GraphBuilder, bool bCullLightsToGrid, FSortedLightSetSceneInfo& SortedLightSet)
 {
 	FComputeLightGridOutput Result = {};
@@ -498,6 +543,7 @@ FComputeLightGridOutput FSceneRenderer::ComputeLightGrid(FRDGBuilder& GraphBuild
 
 		TArray<FVector4f, SceneRenderingAllocator> ViewSpacePosAndRadiusData;
 		TArray<FVector4f, SceneRenderingAllocator> ViewSpaceDirAndPreprocAngleData;
+		TArray<FVector4f, SceneRenderingAllocator> ViewSpaceRectPlanesData;
 
 		float FurthestLight = 1000;
 
@@ -525,6 +571,7 @@ FComputeLightGridOutput FSceneRenderer::ComputeLightGrid(FRDGBuilder& GraphBuild
 
 				ViewSpacePosAndRadiusData.Reserve(SimpleLightsEnd);
 				ViewSpaceDirAndPreprocAngleData.Reserve(SimpleLightsEnd);
+				ViewSpaceRectPlanesData.Reserve(SimpleLightsEnd * NUM_PLANES_PER_RECT_LIGHT);
 
 				const FSimpleLightArray& SimpleLights = SortedLightSet.SimpleLights;
 
@@ -558,6 +605,7 @@ FComputeLightGridOutput FSceneRenderer::ComputeLightGrid(FRDGBuilder& GraphBuild
 					FVector4f ViewSpacePosAndRadius(FVector4f(View.ViewMatrices.GetViewMatrix().TransformPosition(SimpleLightPerViewData.Position)), SimpleLight.Radius);
 					ViewSpacePosAndRadiusData.Add(ViewSpacePosAndRadius);
 					ViewSpaceDirAndPreprocAngleData.AddZeroed();
+					ViewSpaceRectPlanesData.AddZeroed(NUM_PLANES_PER_RECT_LIGHT);
 				}
 			}
 
@@ -645,19 +693,41 @@ FComputeLightGridOutput FSceneRenderer::ComputeLightGrid(FRDGBuilder& GraphBuild
 						const float Distance = View.ViewMatrices.GetViewMatrix().TransformPosition(BoundingSphere.Center).Z + BoundingSphere.W;
 						FurthestLight = FMath::Max(FurthestLight, Distance);
 
-						// Note: inverting radius twice seems stupid (but done in shader anyway otherwise)
 						const FVector3f LightViewPosition = FVector4f(View.ViewMatrices.GetViewMatrix().TransformPosition(LightParameters.WorldPosition)); // LWC_TODO: precision loss
+						const FVector3f LightViewDirection = FVector4f(View.ViewMatrices.GetViewMatrix().TransformVector((FVector)LightParameters.Direction)); // LWC_TODO: precision loss
+
+						// Note: inverting radius twice seems stupid (but done in shader anyway otherwise)
 						FVector4f ViewSpacePosAndRadius(LightViewPosition, 1.0f / LightParameters.InvRadius);
 						ViewSpacePosAndRadiusData.Add(ViewSpacePosAndRadius);
 
 						const bool bIsRectLight = !bRenderRectLightsAsSpotLights && LightProxy->IsRectLight();
+						const bool bUseTightRectLightCulling = bIsRectLight && LightParameters.RectLightBarnLength > 0.5f && LightParameters.RectLightBarnCosAngle > FMath::Cos(FMath::DegreesToRadians(GetRectLightBarnDoorMaxAngle()));
 
 						// Pack flags in the LSB of PreProcAngle
 						const float PreProcAngle = SortedLightInfo.SortKey.Fields.LightType == LightType_Spot ? GetTanRadAngleOrZero(LightProxy->GetOuterConeAngle()) : 0.0f;
-						const uint32 PackedPreProcAngleAndFlags = (FMath::AsUInt(PreProcAngle) & 0xFFFFFFFC) | (LightProxy->HasSourceTexture() ? 0x2 : 0) | (bIsRectLight ? 0x1 : 0);
-						FVector4f ViewSpaceDirAndPreprocAngleAndFlags(FVector4f(View.ViewMatrices.GetViewMatrix().TransformVector((FVector)LightParameters.Direction)), FMath::AsFloat(PackedPreProcAngleAndFlags)); // LWC_TODO: precision loss
+						const uint32 PackedPreProcAngleAndFlags = (FMath::AsUInt(PreProcAngle) & 0xFFFFFFF8) | (LightProxy->HasSourceTexture() ? 0x4 : 0) | (bUseTightRectLightCulling ? 0x2 : 0) | (bIsRectLight ? 0x1 : 0);
+						FVector4f ViewSpaceDirAndPreprocAngleAndFlags(LightViewDirection, FMath::AsFloat(PackedPreProcAngleAndFlags)); // LWC_TODO: precision loss
 						ViewSpaceDirAndPreprocAngleData.Add(ViewSpaceDirAndPreprocAngleAndFlags);
 
+						if (bUseTightRectLightCulling)
+						{
+							const FRectLightSceneProxy* RectProxy = (const FRectLightSceneProxy*)LightProxy;
+
+							TArray<FPlane, TInlineAllocator<NUM_PLANES_PER_RECT_LIGHT>> Planes;
+
+							CalculateRectLightCullingPlanes(RectProxy, Planes);
+
+							for (FPlane& Plane : Planes)
+							{
+								const FPlane4f ViewPlane(Plane.TransformBy(LightProxy->GetLightToWorld() * View.ViewMatrices.GetViewMatrix()));
+								ViewSpaceRectPlanesData.Add(FVector4f(FVector3f(ViewPlane), -ViewPlane.W));
+							}
+						}
+						else
+						{
+							ViewSpaceRectPlanesData.AddZeroed(NUM_PLANES_PER_RECT_LIGHT);
+						}
+
 						bHasRectLights |= bIsRectLight;
 						bHasTexturedLights |= LightProxy->HasSourceTexture();
 					}
@@ -907,12 +977,15 @@ FComputeLightGridOutput FSceneRenderer::ComputeLightGrid(FRDGBuilder& GraphBuild
 
 			check(ViewSpacePosAndRadiusData.Num() == ForwardLocalLightData.Num());
 			check(ViewSpaceDirAndPreprocAngleData.Num() == ForwardLocalLightData.Num());
+			check(ViewSpaceRectPlanesData.Num() == ForwardLocalLightData.Num() * NUM_PLANES_PER_RECT_LIGHT);
 
 			FRDGBufferRef LightViewSpacePositionAndRadius  = CreateStructuredBuffer(GraphBuilder, TEXT("ViewSpacePosAndRadiusData"), TConstArrayView<FVector4f>(ViewSpacePosAndRadiusData));
 			FRDGBufferRef LightViewSpaceDirAndPreprocAngle = CreateStructuredBuffer(GraphBuilder, TEXT("ViewSpacePosAndRadiusData"), TConstArrayView<FVector4f>(ViewSpaceDirAndPreprocAngleData));
+			FRDGBufferRef LightViewSpaceRectPlanes = CreateStructuredBuffer(GraphBuilder, TEXT("ViewSpaceRectPlanesData"), TConstArrayView<FVector4f>(ViewSpaceRectPlanesData));
 
 			PassParameters->LightViewSpacePositionAndRadius  = GraphBuilder.CreateSRV(FRDGBufferSRVDesc(LightViewSpacePositionAndRadius));
 			PassParameters->LightViewSpaceDirAndPreprocAngle = GraphBuilder.CreateSRV(FRDGBufferSRVDesc(LightViewSpaceDirAndPreprocAngle));
+			PassParameters->LightViewSpaceRectPlanes = GraphBuilder.CreateSRV(FRDGBufferSRVDesc(LightViewSpaceRectPlanes));
 
 			{
 				PassParameters->HZBTexture = View.HZB;
