Refactor RayTracing::FSceneOptions constructor.

diff --git a/Engine/Source/Runtime/Renderer/Private/RayTracing/RayTracing.cpp b/Engine/Source/Runtime/Renderer/Private/RayTracing/RayTracing.cpp
--- a/Engine/Source/Runtime/Renderer/Private/RayTracing/RayTracing.cpp
+++ b/Engine/Source/Runtime/Renderer/Private/RayTracing/RayTracing.cpp
@@ -161,31 +161,30 @@ namespace RayTracing
 
 	// Configure ray tracing scene options based on currently enabled features and their needs
 	FSceneOptions::FSceneOptions(
-		FScene& Scene,
+		const FScene& Scene,
 		const FViewFamilyInfo& ViewFamily,
-		FViewInfo& View,
+		const FViewInfo& View,
 		EDiffuseIndirectMethod DiffuseIndirectMethod,
-		EReflectionsMethod ReflectionsMethod,
-		RayTracing::FSceneOptions& SceneOptions)
+		EReflectionsMethod ReflectionsMethod)
 	{
-		SceneOptions.bTranslucentGeometry = false;
-		LumenHardwareRayTracing::SetRayTracingSceneOptions(View, DiffuseIndirectMethod, ReflectionsMethod, SceneOptions);
-		RayTracingShadows::SetRayTracingSceneOptions(Scene.bHasLightsWithRayTracedShadows, SceneOptions);
+		bTranslucentGeometry = false;
+		LumenHardwareRayTracing::SetRayTracingSceneOptions(View, DiffuseIndirectMethod, ReflectionsMethod, *this);
+		RayTracingShadows::SetRayTracingSceneOptions(Scene.bHasLightsWithRayTracedShadows, *this);
 
 		if (ShouldRenderRayTracingTranslucency(View))
 		{
-			SceneOptions.bTranslucentGeometry = true;
+			bTranslucentGeometry = true;
 		}
 
 		if (ViewFamily.EngineShowFlags.PathTracing
 			&& FDataDrivenShaderPlatformInfo::GetSupportsPathTracing(Scene.GetShaderPlatform()))
 		{
-			SceneOptions.bTranslucentGeometry = true;
+			bTranslucentGeometry = true;
 		}
 
 		if (GRayTracingExcludeTranslucent != 0)
 		{
-			SceneOptions.bTranslucentGeometry = false;
+			bTranslucentGeometry = false;
 		}
 	}
 
@@ -869,7 +868,7 @@ namespace RayTracing
 
 		RayTracingSBT.ResetDynamicAllocationData();
 
-		RayTracing::FSceneOptions SceneOptions(Scene, ViewFamily, View, DiffuseIndirectMethod, ReflectionsMethod, SceneOptions);
+		RayTracing::FSceneOptions SceneOptions(Scene, ViewFamily, View, DiffuseIndirectMethod, ReflectionsMethod);
 
 		const float CurrentWorldTime = View.Family->Time.GetWorldTimeSeconds();
 
diff --git a/Engine/Source/Runtime/Renderer/Private/RayTracing/RayTracing.h b/Engine/Source/Runtime/Renderer/Private/RayTracing/RayTracing.h
--- a/Engine/Source/Runtime/Renderer/Private/RayTracing/RayTracing.h
+++ b/Engine/Source/Runtime/Renderer/Private/RayTracing/RayTracing.h
@@ -24,12 +24,12 @@ namespace RayTracing
 	{
 		bool bTranslucentGeometry = true;
 
-		FSceneOptions(FScene& Scene,
+		FSceneOptions(
+			const FScene& Scene,
 			const FViewFamilyInfo& ViewFamily,
-			FViewInfo& View,
+			const FViewInfo& View,
 			EDiffuseIndirectMethod DiffuseIndirectMethod,
-			EReflectionsMethod ReflectionsMethod,
-			RayTracing::FSceneOptions& SceneOptions);
+			EReflectionsMethod ReflectionsMethod);
 	};
 };
 
