Fix crash in FLightGridInjectionCS when View.HZB is not available.

diff --git a/Engine/Source/Runtime/Renderer/Private/LightGridInjection.cpp b/Engine/Source/Runtime/Renderer/Private/LightGridInjection.cpp
--- a/Engine/Source/Runtime/Renderer/Private/LightGridInjection.cpp
+++ b/Engine/Source/Runtime/Renderer/Private/LightGridInjection.cpp
@@ -998,7 +998,7 @@ FComputeLightGridOutput FSceneRenderer::ComputeLightGrid(FRDGBuilder& GraphBuild
 			FLightGridInjectionCS::FPermutationDomain PermutationVector;
 			PermutationVector.Set<FLightGridInjectionCS::FUseLinkedListDim>(GLightLinkedListCulling != 0);
 			PermutationVector.Set<FLightGridInjectionCS::FRefineRectLightBoundsDim>(bHasRectLights);
-			PermutationVector.Set<FLightGridInjectionCS::FUseHZBCullDim>(GLightGridHZBCull != 0);
+			PermutationVector.Set<FLightGridInjectionCS::FUseHZBCullDim>(GLightGridHZBCull != 0 && View.HZB != nullptr);
 			TShaderMapRef<FLightGridInjectionCS> ComputeShader(View.ShaderMap, PermutationVector);
 
 			if (GLightLinkedListCulling != 0)
