Avoid SetInitializer in RayTracingGeometryRHI:
* Remove SetInitializer from RayTracingGeometryRHI
* Add missing RHI fence after building acceleration structures
* Store FRayTracingGeometry in FRelevantPrimitive and use its initializer to get number of segments because accessing RHI object is unsafe here.

#tests replay run on all platforms
#rnx

diff --git a/Engine/Source/Runtime/Apple/MetalRHI/Private/MetalRayTracing.cpp b/Engine/Source/Runtime/Apple/MetalRHI/Private/MetalRayTracing.cpp
--- a/Engine/Source/Runtime/Apple/MetalRHI/Private/MetalRayTracing.cpp
+++ b/Engine/Source/Runtime/Apple/MetalRHI/Private/MetalRayTracing.cpp
@@ -387,13 +387,6 @@ void FMetalRayTracingGeometry::ReleaseUnderlyingResource()
 	}
 }
 
-void FMetalRayTracingGeometry::SetInitializer(FRHICommandListBase& RHICmdList, const FRayTracingGeometryInitializer& InInitializer)
-{
-	Initializer = InInitializer;
-
-	// HitGroup Parameters Update is handled by the Scene
-}
-
 void FMetalRayTracingGeometry::Swap(FMetalRayTracingGeometry& Other)
 {
 	::Swap(AccelerationStructureDescriptor, Other.AccelerationStructureDescriptor);
@@ -402,6 +395,10 @@ void FMetalRayTracingGeometry::Swap(FMetalRayTracingGeometry& Other)
 		::Swap(AccelerationStructure[i], Other.AccelerationStructure[i]);
 	}
 	::Swap(AccelerationStructureIndex, Other.AccelerationStructureIndex);
+
+	Initializer = Other.Initializer;
+
+	// HitGroup Parameters Update is handled by the Scene
 }
 
 void FMetalRayTracingGeometry::RemoveCompactionRequest()
diff --git a/Engine/Source/Runtime/Apple/MetalRHI/Private/MetalRayTracing.h b/Engine/Source/Runtime/Apple/MetalRHI/Private/MetalRayTracing.h
--- a/Engine/Source/Runtime/Apple/MetalRHI/Private/MetalRayTracing.h
+++ b/Engine/Source/Runtime/Apple/MetalRHI/Private/MetalRayTracing.h
@@ -27,9 +27,7 @@ public:
 	void ReleaseUnderlyingResource();
 
 	/** FRHIRayTracingGeometry Interface */
-	virtual FRayTracingAccelerationStructureAddress GetAccelerationStructureAddress(uint64 GPUIndex) const final override { return (FRayTracingAccelerationStructureAddress)SceneIndex;
-	}
-	virtual void SetInitializer(FRHICommandListBase& RHICmdList, const FRayTracingGeometryInitializer& Initializer) final override;
+	virtual FRayTracingAccelerationStructureAddress GetAccelerationStructureAddress(uint64 GPUIndex) const final override { return (FRayTracingAccelerationStructureAddress)SceneIndex; }
 	/** FRHIRayTracingGeometry Interface */
 
 	void Swap(FMetalRayTracingGeometry& Other);
diff --git a/Engine/Source/Runtime/D3D12RHI/Private/D3D12RayTracing.cpp b/Engine/Source/Runtime/D3D12RHI/Private/D3D12RayTracing.cpp
--- a/Engine/Source/Runtime/D3D12RHI/Private/D3D12RayTracing.cpp
+++ b/Engine/Source/Runtime/D3D12RHI/Private/D3D12RayTracing.cpp
@@ -2921,7 +2921,25 @@ void FD3D12RayTracingGeometry::Swap(FD3D12RayTracingGeometry& Other)
 	});
 	::Swap(AccelerationStructureCompactedSize, Other.AccelerationStructureCompactedSize);
 
-	// The rest of the members should be updated using SetInitializer()
+	for (uint32 GPUIndex = 0; GPUIndex < MAX_NUM_GPUS && GPUIndex < GNumExplicitGPUsForRendering; ++GPUIndex)
+	{
+		UnregisterAsRenameListener(GPUIndex);
+	}
+
+	Initializer = Other.Initializer;
+
+	DebugName = !Initializer.DebugName.IsNone() ? Initializer.DebugName : FName(TEXT("BLAS"));
+
+	checkf(Initializer.Segments.Num() > 0, TEXT("Ray tracing geometry must be initialized with at least one segment."));
+
+	GeometryDescs.SetNumUninitialized(Initializer.Segments.Num());
+	TranslateRayTracingGeometryDescs(Initializer, GeometryDescs);
+
+	for (uint32 GPUIndex = 0; GPUIndex < MAX_NUM_GPUS && GPUIndex < GNumExplicitGPUsForRendering; ++GPUIndex)
+	{
+		RegisterAsRenameListener(GPUIndex);
+		SetupHitGroupSystemParameters(GPUIndex);
+	}
 }
 
 void FD3D12RayTracingGeometry::ReleaseUnderlyingResource()
@@ -3135,33 +3153,6 @@ void FD3D12RayTracingGeometry::UpdateResidency(FD3D12CommandContext& CommandCont
 	CommandContext.UpdateResidency(AccelerationStructureBuffers[GPUIndex]->GetResource());
 }
 
-void FD3D12RayTracingGeometry::SetInitializer(FRHICommandListBase& RHICmdList, const FRayTracingGeometryInitializer& InInitializer)
-{
-	Initializer = InInitializer;
-
-	FOREACH_GPU(GPUIndex < MAX_NUM_GPUS && GPUIndex < GNumExplicitGPUsForRendering,
-	{
-		UnregisterAsRenameListener(GPUIndex);
-	});
-
-	DebugName = !Initializer.DebugName.IsNone() ? Initializer.DebugName : FName(TEXT("BLAS"));
-
-	checkf(Initializer.Segments.Num() > 0, TEXT("Ray tracing geometry must be initialized with at least one segment."));
-
-	GeometryDescs.SetNumUninitialized(Initializer.Segments.Num());
-	TranslateRayTracingGeometryDescs(Initializer, GeometryDescs);
-	
-	RHICmdList.EnqueueLambda([this](FRHICommandListBase&)
-	{
-		FOREACH_GPU(GPUIndex < MAX_NUM_GPUS && GPUIndex < GNumExplicitGPUsForRendering,
-		{				
-			RegisterAsRenameListener(GPUIndex);
-			SetupHitGroupSystemParameters(GPUIndex);		
-		});
-	});
-	RHICmdList.RHIThreadFence(true);
-}
-
 void FD3D12RayTracingGeometry::SetupHitGroupSystemParameters(uint32 InGPUIndex)
 {
 	D3D12_RAYTRACING_GEOMETRY_TYPE GeometryType = TranslateRayTracingGeometryType(Initializer.GeometryType);
diff --git a/Engine/Source/Runtime/D3D12RHI/Private/D3D12RayTracing.h b/Engine/Source/Runtime/D3D12RHI/Private/D3D12RayTracing.h
--- a/Engine/Source/Runtime/D3D12RHI/Private/D3D12RayTracing.h
+++ b/Engine/Source/Runtime/D3D12RHI/Private/D3D12RayTracing.h
@@ -32,7 +32,6 @@ public:
 			TEXT("Trying to get address of acceleration structure '%s' without allocated memory."), *DebugName.ToString());
 		return AccelerationStructureBuffers[GPUIndex]->ResourceLocation.GetGPUVirtualAddress();
 	}
-	virtual void SetInitializer(FRHICommandListBase& RHICmdList, const FRayTracingGeometryInitializer& Initializer) final override;
 
 	void SetupHitGroupSystemParameters(uint32 InGPUIndex);
 	void TransitionBuffers(FD3D12CommandContext& CommandContext);
diff --git a/Engine/Source/Runtime/RHI/Public/RHICommandList.h b/Engine/Source/Runtime/RHI/Public/RHICommandList.h
--- a/Engine/Source/Runtime/RHI/Public/RHICommandList.h
+++ b/Engine/Source/Runtime/RHI/Public/RHICommandList.h
@@ -3380,6 +3380,8 @@ public:
 				}
 			}
 			ALLOC_COMMAND(FRHICommandBuildAccelerationStructures)(MakeArrayView(InlineParams, Params.Num()), ScratchBufferRange);
+
+			RHIThreadFence(true);
 		}
 	}
 
diff --git a/Engine/Source/Runtime/RHI/Public/RHIResources.h b/Engine/Source/Runtime/RHI/Public/RHIResources.h
--- a/Engine/Source/Runtime/RHI/Public/RHIResources.h
+++ b/Engine/Source/Runtime/RHI/Public/RHIResources.h
@@ -3376,7 +3376,6 @@ public:
 	{}
 
 	virtual FRayTracingAccelerationStructureAddress GetAccelerationStructureAddress(uint64 GPUIndex) const = 0;
-	virtual void SetInitializer(FRHICommandListBase& RHICmdList, const FRayTracingGeometryInitializer& Initializer) = 0;
 	virtual bool IsCompressed() const { return false; }
 
 	const FRayTracingGeometryInitializer& GetInitializer() const
diff --git a/Engine/Source/Runtime/RenderCore/Private/RayTracingGeometry.cpp b/Engine/Source/Runtime/RenderCore/Private/RayTracingGeometry.cpp
--- a/Engine/Source/Runtime/RenderCore/Private/RayTracingGeometry.cpp
+++ b/Engine/Source/Runtime/RenderCore/Private/RayTracingGeometry.cpp
@@ -61,8 +61,6 @@ void FRayTracingGeometry::ReleaseRHIForStreaming(FRHIResourceReplaceBatcher& Bat
 
 void FRayTracingGeometry::RequestBuildIfNeeded(FRHICommandListBase& RHICmdList, ERTAccelerationStructureBuildPriority InBuildPriority)
 {
-	RayTracingGeometryRHI->SetInitializer(RHICmdList, Initializer);
-
 	if (GetRequiresBuild())
 	{
 		RayTracingBuildRequestIndex = GRayTracingGeometryManager->RequestBuildAccelerationStructure(this, InBuildPriority);
diff --git a/Engine/Source/Runtime/Renderer/Private/RayTracing/RayTracing.cpp b/Engine/Source/Runtime/Renderer/Private/RayTracing/RayTracing.cpp
--- a/Engine/Source/Runtime/Renderer/Private/RayTracing/RayTracing.cpp
+++ b/Engine/Source/Runtime/Renderer/Private/RayTracing/RayTracing.cpp
@@ -190,7 +190,7 @@ namespace RayTracing
 
 	struct FRelevantPrimitive
 	{
-		FRHIRayTracingGeometry* RayTracingGeometryRHI = nullptr;
+		const FRayTracingGeometry* RayTracingGeometry = nullptr;
 		FRayTracingCachedMeshCommandFlags CachedMeshCommandFlags;
 		int32 PrimitiveIndex = -1;
 		FPersistentPrimitiveIndex PersistentPrimitiveIndex;
@@ -217,7 +217,7 @@ namespace RayTracing
 			Key ^= CachedMeshCommandFlags.bIsSky ? 0x1ull << 46 : 0x0;
 			Key ^= CachedMeshCommandFlags.bAllSegmentsTranslucent ? 0x1ull << 47 : 0x0;
 			Key ^= CachedMeshCommandFlags.bAllSegmentsReverseCulling ? 0x1ull << 48 : 0x0;
-			return Key ^ reinterpret_cast<uint64>(RayTracingGeometryRHI);
+			return Key ^ reinterpret_cast<uint64>(RayTracingGeometry->GetRHI());
 		}
 	};
 
@@ -693,6 +693,7 @@ namespace RayTracing
 							FRelevantPrimitive* RelevantPrimitive = new (Context.CachedStaticPrimitives) FRelevantPrimitive();
 							RelevantPrimitive->PrimitiveIndex = PrimitiveIndex;
 							RelevantPrimitive->PersistentPrimitiveIndex = SceneInfo->GetPersistentIndex();
+							RelevantPrimitive->RayTracingGeometry = SceneInfo->GetCachedRayTracingGeometry();
 
 							ensureMsgf(!SceneInfo->bCachedRaytracingDataDirty, TEXT("Cached ray tracing instances must be up-to-date at this point"));
 
@@ -702,7 +703,7 @@ namespace RayTracing
 
 							// CacheInstances expects to have one ray tracing mesh command per BLAS segment.
 							// If that's not the case in the future, other logic such as NumCachedStaticVisibleMeshCommands calculation needs to be updated.
-							checkf(RTLODData.CachedMeshCommandIndices.Num() == SceneInfo->CachedRayTracingInstance.GeometryRHI->GetNumSegments(),
+							checkf(RTLODData.CachedMeshCommandIndices.Num() == RelevantPrimitive->RayTracingGeometry->Initializer.Segments.Num(),
 								TEXT("Expected to have one ray tracing mesh command per BLAS segment (primitive has %d cached mesh commands but BLAS has %d segments)."),
 								RTLODData.CachedMeshCommandIndices.Num(), SceneInfo->CachedRayTracingInstance.GeometryRHI->GetNumSegments());
 
@@ -779,7 +780,7 @@ namespace RayTracing
 								RelevantPrimitive->PersistentPrimitiveIndex = SceneInfo->GetPersistentIndex();
 
 								RelevantPrimitive->LODIndex = LODIndex;
-								RelevantPrimitive->RayTracingGeometryRHI = RayTracingGeometry->GetRHI();
+								RelevantPrimitive->RayTracingGeometry = RayTracingGeometry;
 
 								const FPrimitiveSceneInfo::FRayTracingLODData& RTLODData = SceneInfo->RayTracingLODData[LODIndex];
 								RelevantPrimitive->CachedMeshCommandFlags = RTLODData.CachedMeshCommandFlags;
@@ -1446,7 +1447,7 @@ namespace RayTracing
 							InstanceBatch.Add(RayTracingScene, SceneInfo->GetInstanceSceneDataOffset());
 
 							FRayTracingGeometryInstance RayTracingInstance;
-							RayTracingInstance.GeometryRHI = RelevantPrimitive.RayTracingGeometryRHI;
+							RayTracingInstance.GeometryRHI = RelevantPrimitive.RayTracingGeometry->GetRHI();
 							checkf(RayTracingInstance.GeometryRHI, TEXT("Ray tracing instance must have a valid geometry."));
 							RayTracingInstance.InstanceSceneDataOffsets = InstanceBatch.InstanceSceneDataOffsets;
 							RayTracingInstance.UserData = InstanceBatch.InstanceSceneDataOffsets;
@@ -1478,7 +1479,7 @@ namespace RayTracing
 								continue;
 							}
 							
-							uint32 SegmentCount = RayTracingInstance.GeometryRHI->GetNumSegments();
+							uint32 SegmentCount = RelevantPrimitive.RayTracingGeometry->Initializer.Segments.Num();
 							if (bNeedMainInstance)
 							{
 								RayTracingScene.NumSegments += SegmentCount;
@@ -1505,14 +1506,14 @@ namespace RayTracing
 									{
 										const bool bHidden = MeshCommand.bDecal;
 										const uint32 RecordIndex = RTLODData.SBTAllocation->GetRecordIndex(ERayTracingSceneLayer::Base, MeshCommand.GeometrySegmentIndex);
-										FRayTracingShaderBindingData RTShaderBindingData(&MeshCommand, RelevantPrimitive.RayTracingGeometryRHI, RecordIndex, bHidden);
+										FRayTracingShaderBindingData RTShaderBindingData(&MeshCommand, RelevantPrimitive.RayTracingGeometry->GetRHI(), RecordIndex, bHidden);
 										DirtyShaderBindingData.Add(RTShaderBindingData);
 									}
 									if (bNeedDecalInstance)
 									{
 										const bool bHidden = !MeshCommand.bDecal;
 										const uint32 RecordIndex = RTLODData.SBTAllocation->GetRecordIndex(ERayTracingSceneLayer::Decals, MeshCommand.GeometrySegmentIndex);
-										FRayTracingShaderBindingData RTShaderBindingData(&MeshCommand, RelevantPrimitive.RayTracingGeometryRHI, RecordIndex, bHidden);
+										FRayTracingShaderBindingData RTShaderBindingData(&MeshCommand, RelevantPrimitive.RayTracingGeometry->GetRHI(), RecordIndex, bHidden);
 										DirtyShaderBindingData.Add(RTShaderBindingData);
 									}
 								}
diff --git a/Engine/Source/Runtime/Renderer/Public/PrimitiveSceneInfo.h b/Engine/Source/Runtime/Renderer/Public/PrimitiveSceneInfo.h
--- a/Engine/Source/Runtime/Renderer/Public/PrimitiveSceneInfo.h
+++ b/Engine/Source/Runtime/Renderer/Public/PrimitiveSceneInfo.h
@@ -564,6 +564,8 @@ public:
 	RENDERER_API FRHIRayTracingGeometry* GetStaticRayTracingGeometryInstance(int LodLevel) const;
 
 	int GetStaticRayTracingGeometryNum() const { return StaticRayTracingGeometries.Num(); }
+
+	const FRayTracingGeometry* GetCachedRayTracingGeometry() const { return CachedRayTracingGeometry; }
 #endif
 
 	/** Return primitive fullname (for debugging only). */
diff --git a/Engine/Source/Runtime/VulkanRHI/Private/VulkanRayTracing.cpp b/Engine/Source/Runtime/VulkanRHI/Private/VulkanRayTracing.cpp
--- a/Engine/Source/Runtime/VulkanRHI/Private/VulkanRayTracing.cpp
+++ b/Engine/Source/Runtime/VulkanRHI/Private/VulkanRayTracing.cpp
@@ -450,13 +450,6 @@ FVulkanRayTracingGeometry::~FVulkanRayTracingGeometry()
 	RemoveCompactionRequest();
 }
 
-void FVulkanRayTracingGeometry::SetInitializer(FRHICommandListBase& RHICmdList, const FRayTracingGeometryInitializer& InInitializer)
-{
-	Initializer = InInitializer;
-
-	// TODO: Update HitGroup Parameters
-}
-
 void FVulkanRayTracingGeometry::Swap(FVulkanRayTracingGeometry& Other)
 {
 	::Swap(Handle, Other.Handle);
@@ -465,7 +458,9 @@ void FVulkanRayTracingGeometry::Swap(FVulkanRayTracingGeometry& Other)
 
 	AccelerationStructureBuffer = Other.AccelerationStructureBuffer;
 
-	// The rest of the members should be updated using SetInitializer()
+	Initializer = Other.Initializer;
+
+	// TODO: Update HitGroup Parameters
 }
 
 void FVulkanRayTracingGeometry::RemoveCompactionRequest()
diff --git a/Engine/Source/Runtime/VulkanRHI/Private/VulkanRayTracing.h b/Engine/Source/Runtime/VulkanRHI/Private/VulkanRayTracing.h
--- a/Engine/Source/Runtime/VulkanRHI/Private/VulkanRayTracing.h
+++ b/Engine/Source/Runtime/VulkanRHI/Private/VulkanRayTracing.h
@@ -163,8 +163,7 @@ public:
 	FVulkanRayTracingGeometry(FRHICommandListBase& RHICmdList, const FRayTracingGeometryInitializer& Initializer, FVulkanDevice* InDevice);
 	~FVulkanRayTracingGeometry();
 
-	virtual FRayTracingAccelerationStructureAddress GetAccelerationStructureAddress(uint64 GPUIndex) const final override { return Address; }	
-	virtual void SetInitializer(FRHICommandListBase& RHICmdList, const FRayTracingGeometryInitializer& Initializer) final override;
+	virtual FRayTracingAccelerationStructureAddress GetAccelerationStructureAddress(uint64 GPUIndex) const final override { return Address; }
 
 	void Swap(FVulkanRayTracingGeometry& Other);
 
