Moved definitions of FScenePrimitiveUpdates and associated types from ScenePrivate.h to (new) ScenePrimitiveUpdates.h to avoid circular include dependencies when using the types from other parts of the renderer.

#rnx

diff --git a/Engine/Source/Runtime/Renderer/Private/ScenePrimitiveUpdates.h b/Engine/Source/Runtime/Renderer/Private/ScenePrimitiveUpdates.h
new file mode 100644
--- /dev/null
+++ b/Engine/Source/Runtime/Renderer/Private/ScenePrimitiveUpdates.h
@@ -0,0 +1,125 @@
+// Copyright Epic Games, Inc. All Rights Reserved.
+#pragma once
+
+#include "SceneUpdateCommandQueue.h"
+#include "PrimitiveSceneInfo.h"
+
+/**
+ * Definitions of primitive scene update commands.
+ */
+
+enum class EPrimitiveUpdateDirtyFlags : uint32 
+{
+	None			=	0u,
+	/** The Transform is modified by this command. */
+	Transform		=	1u << 0,
+	/** The (any) instance data is modified by this command. */
+	InstanceData	=	1u << 1,
+	/** 
+	 * The culling bounds are modified by this command. 
+	 * This means the bounds (instance, primitive or both) as used in the culling and should not be updated for any other case. 
+	 * Thus, needs to be set for transform updates of all kinds.
+	 */
+	CullingBounds	=	1u << 2,
+	/** 
+	 * Culling distances or similar, affecting culling logic, but NOT the bounds.
+	 */
+	CullingLogic	=	1u << 3,
+	/** Any state that either makes its way into GPU-Scene or the per primitive UB */
+	GPUState		=	1u << 4,
+	/** All culling-affecting changes */
+	AllCulling		= CullingBounds | CullingLogic,
+	All				=	GPUState |  Transform | InstanceData | CullingBounds | CullingLogic,
+};
+ENUM_CLASS_FLAGS(EPrimitiveUpdateDirtyFlags);
+
+enum class EPrimitiveUpdateId : uint32 
+{
+	UpdateTransform,
+	UpdateInstance,
+	UpdateAttachmentRoot,
+	CustomPrimitiveData,
+	OcclusionBoundsSlacks,  
+	InstanceCullDistance,  
+	DrawDistance,
+	DistanceFieldScene,
+	OverridePreviousTransform,
+	MAX
+};
+
+using FScenePrimitiveUpdates = TSceneUpdateCommandQueue<FPrimitiveSceneInfo, EPrimitiveUpdateDirtyFlags, EPrimitiveUpdateId>;
+
+using FPrimitiveUpdateCommand = FScenePrimitiveUpdates::FUpdateCommand;
+
+template <EPrimitiveUpdateId InId, EPrimitiveUpdateDirtyFlags InDirtyFlags>
+using TPrimitiveUpdatePayloadBase = FScenePrimitiveUpdates::TPayloadBase<InId, InDirtyFlags>;
+
+struct FUpdateTransformCommand : public TPrimitiveUpdatePayloadBase<EPrimitiveUpdateId::UpdateTransform, 
+	EPrimitiveUpdateDirtyFlags::GPUState | EPrimitiveUpdateDirtyFlags::Transform | EPrimitiveUpdateDirtyFlags::CullingBounds>
+{
+	FBoxSphereBounds WorldBounds;
+	FBoxSphereBounds LocalBounds; 
+	FMatrix LocalToWorld; 
+	FVector AttachmentRootPosition;
+};
+
+struct FUpdateInstanceCommand : public TPrimitiveUpdatePayloadBase<EPrimitiveUpdateId::UpdateInstance, 
+	EPrimitiveUpdateDirtyFlags::GPUState | EPrimitiveUpdateDirtyFlags::Transform | EPrimitiveUpdateDirtyFlags::CullingBounds | EPrimitiveUpdateDirtyFlags::InstanceData>
+{
+	FPrimitiveSceneProxy* PrimitiveSceneProxy{ nullptr };
+	FBoxSphereBounds WorldBounds;
+	FBoxSphereBounds LocalBounds;
+	FBoxSphereBounds StaticMeshBounds;
+};
+
+
+/**
+ * Helper for the update payloads that contain a single payload value.
+ */
+template <typename InPayloadDataType, EPrimitiveUpdateId InId, EPrimitiveUpdateDirtyFlags InPrimitiveDirtyFlags>
+struct TSingleValuePrimitiveUpdatePayload : public TPrimitiveUpdatePayloadBase<InId, InPrimitiveDirtyFlags>
+{
+	TSingleValuePrimitiveUpdatePayload(const InPayloadDataType &InValue) : Value(InValue) {}
+	InPayloadDataType Value;
+};
+
+using FUpdateAttachmentRootData = TSingleValuePrimitiveUpdatePayload<FPrimitiveComponentId, EPrimitiveUpdateId::UpdateAttachmentRoot,
+	EPrimitiveUpdateDirtyFlags::None>; // No GPU side effect (?).
+using FUpdateCustomPrimitiveData = TSingleValuePrimitiveUpdatePayload<FCustomPrimitiveData, EPrimitiveUpdateId::CustomPrimitiveData,
+	EPrimitiveUpdateDirtyFlags::GPUState>; // Needs upload
+using FUpdateOcclusionBoundsSlacksData = TSingleValuePrimitiveUpdatePayload<float, EPrimitiveUpdateId::OcclusionBoundsSlacks,  
+	EPrimitiveUpdateDirtyFlags::None>; // Only affects primitive occlusion.
+using FUpdateInstanceCullDistanceData = TSingleValuePrimitiveUpdatePayload<FVector2f, EPrimitiveUpdateId::InstanceCullDistance,  
+	EPrimitiveUpdateDirtyFlags::GPUState | EPrimitiveUpdateDirtyFlags::CullingLogic>; // Affects GPU culling?
+using FUpdateDrawDistanceData = TSingleValuePrimitiveUpdatePayload<FVector3f, EPrimitiveUpdateId::DrawDistance,  
+	EPrimitiveUpdateDirtyFlags::CullingLogic>; // Only affects CPU culling.
+using FUpdateDistanceFieldSceneData = TPrimitiveUpdatePayloadBase<EPrimitiveUpdateId::DistanceFieldScene,  
+	EPrimitiveUpdateDirtyFlags::None>; // Only affects DF scene rep - candidate for using abstract type.
+using FUpdateOverridePreviousTransformData = TSingleValuePrimitiveUpdatePayload<FMatrix, EPrimitiveUpdateId::OverridePreviousTransform,  
+	EPrimitiveUpdateDirtyFlags::GPUState>;// Overrides the previous transform, which needs to be propagated to the GPU, but otherwise does not change anything on its own.
+
+/**
+ * Change set that is valid before removes are processed and the scene data modified.
+ * The referenced arrays have RDG life-time and can be safely used in RDG tasks.
+ * However, the referenced data (primitive/proxy) and meaning of the persistent ID is not generally valid past the call in which this is passed. 
+ * Thus, care need to be excercised.
+ */
+class FScenePreUpdateChangeSet
+{
+public:
+	TConstArrayView<FPersistentPrimitiveIndex> RemovedPrimitiveIds;
+	TConstArrayView<FPrimitiveSceneInfo*> RemovedPrimitiveSceneInfos;
+	const FScenePrimitiveUpdates &PrimitiveUpdates;
+};
+
+/**
+ * Change set that is valid before after adds are processed and the scene data is modified.
+ * The referenced arrays have RDG life-time and can be safely used in RDG tasks.
+ */
+class FScenePostUpdateChangeSet
+{
+public:
+	TConstArrayView<FPersistentPrimitiveIndex> AddedPrimitiveIds;
+	TConstArrayView<FPrimitiveSceneInfo*> AddedPrimitiveSceneInfos;
+	const FScenePrimitiveUpdates &PrimitiveUpdates;
+};
diff --git a/Engine/Source/Runtime/Renderer/Private/ScenePrivate.h b/Engine/Source/Runtime/Renderer/Private/ScenePrivate.h
--- a/Engine/Source/Runtime/Renderer/Private/ScenePrivate.h
+++ b/Engine/Source/Runtime/Renderer/Private/ScenePrivate.h
@@ -60,7 +60,7 @@
 #include "LightFunctionAtlas.h"
 #include "SceneExtensions.h"
 #include "HeterogeneousVolumes/HeterogeneousVolumes.h"
-#include "SceneUpdateCommandQueue.h"
+#include "ScenePrimitiveUpdates.h"
 
 /** Factor by which to grow occlusion tests **/
 #define OCCLUSION_SLOP (1.0f)
@@ -2841,126 +2841,6 @@ private:
 	FLumenSceneDataMap::TConstIterator NextSceneData;
 };
 
-/**
- * Definitions of primitive scene update commands.
- */
-
-enum class EPrimitiveUpdateDirtyFlags : uint32 
-{
-	None			=	0u,
-	/** The Transform is modified by this command. */
-	Transform		=	1u << 0,
-	/** The (any) instance data is modified by this command. */
-	InstanceData	=	1u << 1,
-	/** 
-	 * The culling bounds are modified by this command. 
-	 * This means the bounds (instance, primitive or both) as used in the culling and should not be updated for any other case. 
-	 * Thus, needs to be set for transform updates of all kinds.
-	 */
-	CullingBounds	=	1u << 2,
-	/** 
-	 * Culling distances or similar, affecting culling logic, but NOT the bounds.
-	 */
-	CullingLogic	=	1u << 3,
-	/** Any state that either makes its way into GPU-Scene or the per primitive UB */
-	GPUState		=	1u << 4,
-	/** All culling-affecting changes */
-	AllCulling		= CullingBounds | CullingLogic,
-	All				=	GPUState |  Transform | InstanceData | CullingBounds | CullingLogic,
-};
-ENUM_CLASS_FLAGS(EPrimitiveUpdateDirtyFlags);
-
-enum class EPrimitiveUpdateId : uint32 
-{
-	UpdateTransform,
-	UpdateInstance,
-	UpdateAttachmentRoot,
-	CustomPrimitiveData,
-	OcclusionBoundsSlacks,  
-	InstanceCullDistance,  
-	DrawDistance,
-	DistanceFieldScene,
-	OverridePreviousTransform,
-	MAX
-};
-
-using FScenePrimitiveUpdates = TSceneUpdateCommandQueue<FPrimitiveSceneInfo, EPrimitiveUpdateDirtyFlags, EPrimitiveUpdateId>;
-
-using FPrimitiveUpdateCommand = FScenePrimitiveUpdates::FUpdateCommand;
-
-template <EPrimitiveUpdateId InId, EPrimitiveUpdateDirtyFlags InDirtyFlags>
-using TPrimitiveUpdatePayloadBase = FScenePrimitiveUpdates::TPayloadBase<InId, InDirtyFlags>;
-
-struct FUpdateTransformCommand : public TPrimitiveUpdatePayloadBase<EPrimitiveUpdateId::UpdateTransform, 
-	EPrimitiveUpdateDirtyFlags::GPUState | EPrimitiveUpdateDirtyFlags::Transform | EPrimitiveUpdateDirtyFlags::CullingBounds>
-{
-	FBoxSphereBounds WorldBounds;
-	FBoxSphereBounds LocalBounds; 
-	FMatrix LocalToWorld; 
-	FVector AttachmentRootPosition;
-};
-
-struct FUpdateInstanceCommand : public TPrimitiveUpdatePayloadBase<EPrimitiveUpdateId::UpdateInstance, 
-	EPrimitiveUpdateDirtyFlags::GPUState | EPrimitiveUpdateDirtyFlags::Transform | EPrimitiveUpdateDirtyFlags::CullingBounds | EPrimitiveUpdateDirtyFlags::InstanceData>
-{
-	FPrimitiveSceneProxy* PrimitiveSceneProxy{ nullptr };
-	FBoxSphereBounds WorldBounds;
-	FBoxSphereBounds LocalBounds;
-	FBoxSphereBounds StaticMeshBounds;
-};
-
-
-/**
- * Helper for the update payloads that contain a single payload value.
- */
-template <typename InPayloadDataType, EPrimitiveUpdateId InId, EPrimitiveUpdateDirtyFlags InPrimitiveDirtyFlags>
-struct TSingleValuePrimitiveUpdatePayload : public TPrimitiveUpdatePayloadBase<InId, InPrimitiveDirtyFlags>
-{
-	TSingleValuePrimitiveUpdatePayload(const InPayloadDataType &InValue) : Value(InValue) {}
-	InPayloadDataType Value;
-};
-
-using FUpdateAttachmentRootData = TSingleValuePrimitiveUpdatePayload<FPrimitiveComponentId, EPrimitiveUpdateId::UpdateAttachmentRoot,
-	EPrimitiveUpdateDirtyFlags::None>; // No GPU side effect (?).
-using FUpdateCustomPrimitiveData = TSingleValuePrimitiveUpdatePayload<FCustomPrimitiveData, EPrimitiveUpdateId::CustomPrimitiveData,
-	EPrimitiveUpdateDirtyFlags::GPUState>; // Needs upload
-using FUpdateOcclusionBoundsSlacksData = TSingleValuePrimitiveUpdatePayload<float, EPrimitiveUpdateId::OcclusionBoundsSlacks,  
-	EPrimitiveUpdateDirtyFlags::None>; // Only affects primitive occlusion.
-using FUpdateInstanceCullDistanceData = TSingleValuePrimitiveUpdatePayload<FVector2f, EPrimitiveUpdateId::InstanceCullDistance,  
-	EPrimitiveUpdateDirtyFlags::GPUState | EPrimitiveUpdateDirtyFlags::CullingLogic>; // Affects GPU culling?
-using FUpdateDrawDistanceData = TSingleValuePrimitiveUpdatePayload<FVector3f, EPrimitiveUpdateId::DrawDistance,  
-	EPrimitiveUpdateDirtyFlags::CullingLogic>; // Only affects CPU culling.
-using FUpdateDistanceFieldSceneData = TPrimitiveUpdatePayloadBase<EPrimitiveUpdateId::DistanceFieldScene,  
-	EPrimitiveUpdateDirtyFlags::None>; // Only affects DF scene rep - candidate for using abstract type.
-using FUpdateOverridePreviousTransformData = TSingleValuePrimitiveUpdatePayload<FMatrix, EPrimitiveUpdateId::OverridePreviousTransform,  
-	EPrimitiveUpdateDirtyFlags::GPUState>;// Overrides the previous transform, which needs to be propagated to the GPU, but otherwise does not change anything on its own.
-
-/**
- * Change set that is valid before removes are processed and the scene data modified.
- * The referenced arrays have RDG life-time and can be safely used in RDG tasks.
- * However, the referenced data (primitive/proxy) and meaning of the persistent ID is not generally valid past the call in which this is passed. 
- * Thus, care need to be excercised.
- */
-class FScenePreUpdateChangeSet
-{
-public:
-	TConstArrayView<FPersistentPrimitiveIndex> RemovedPrimitiveIds;
-	TConstArrayView<FPrimitiveSceneInfo*> RemovedPrimitiveSceneInfos;
-	const FScenePrimitiveUpdates &PrimitiveUpdates;
-};
-
-/**
- * Change set that is valid before after adds are processed and the scene data is modified.
- * The referenced arrays have RDG life-time and can be safely used in RDG tasks.
- */
-class FScenePostUpdateChangeSet
-{
-public:
-	TConstArrayView<FPersistentPrimitiveIndex> AddedPrimitiveIds;
-	TConstArrayView<FPrimitiveSceneInfo*> AddedPrimitiveSceneInfos;
-	const FScenePrimitiveUpdates &PrimitiveUpdates;
-};
-
 /**
  * Definitions for light scene updates.
  */
