Remove call to View.RayTracingCullingParameters.Init() from GatherWorldInstancesForView(...)
- it's already done in BeginInitVisibility()

diff --git a/Engine/Source/Runtime/Renderer/Private/RayTracing/RayTracing.cpp b/Engine/Source/Runtime/Renderer/Private/RayTracing/RayTracing.cpp
--- a/Engine/Source/Runtime/Renderer/Private/RayTracing/RayTracing.cpp
+++ b/Engine/Source/Runtime/Renderer/Private/RayTracing/RayTracing.cpp
@@ -865,8 +865,6 @@ namespace RayTracing
 		
 		View.RayTracingMeshResourceCollector = MakeUnique<FRayTracingMeshResourceCollector>(Scene.GetFeatureLevel(), InBulkAllocator);
 
-		View.RayTracingCullingParameters.Init(View);
-
 		RayTracingSBT.ResetDynamicAllocationData();
 
 		RayTracing::FSceneOptions SceneOptions(Scene, ViewFamily, View, DiffuseIndirectMethod, ReflectionsMethod);
@@ -1260,7 +1258,6 @@ namespace RayTracing
 			TArray<FRelevantPrimitive>& RelevantStaticPrimitives;
 			TArray<FRelevantPrimitive>& RelevantCachedStaticPrimitives;
 			TArray<FRelevantPrimitiveGatherContext>& GatherContexts;
-			const FRayTracingCullingParameters& CullingParameters;
 			const bool bIsPathTracing;
 
 			const int32& NumCachedStaticInstances;
@@ -1278,7 +1275,6 @@ namespace RayTracing
 				TArray<FRelevantPrimitive>& InRelevantStaticPrimitives,
 				TArray<FRelevantPrimitive>& InRelevantCachedStaticPrimitives,
 				TArray<FRelevantPrimitiveGatherContext>& InGatherContexts,
-				const FRayTracingCullingParameters& InCullingParameters,
 				const bool bInIsPathTracing,
 				const int32& InNumCachedStaticInstances,
 				const int32& InNumCachedStaticDecalInstances,
@@ -1290,7 +1286,6 @@ namespace RayTracing
 				, RelevantStaticPrimitives(InRelevantStaticPrimitives)
 				, RelevantCachedStaticPrimitives(InRelevantCachedStaticPrimitives)
 				, GatherContexts(InGatherContexts)
-				, CullingParameters(InCullingParameters)
 				, bIsPathTracing(bInIsPathTracing)
 				, NumCachedStaticInstances(InNumCachedStaticInstances)
 				, NumCachedStaticDecalInstances(InNumCachedStaticDecalInstances)
@@ -1638,7 +1633,6 @@ namespace RayTracing
 			RelevantPrimitiveList.StaticPrimitives,
 			RelevantPrimitiveList.CachedStaticPrimitives,
 			RelevantPrimitiveList.GatherContexts,
-			View.RayTracingCullingParameters,
 			bool(View.Family->EngineShowFlags.PathTracing),
 			RelevantPrimitiveList.NumCachedStaticInstances,
 			RelevantPrimitiveList.NumCachedStaticDecalInstances,
