Fix multiple definitions of asfloat(...) due to non-static definition in HaltonUtilities.cpp

diff --git a/Engine/Source/Runtime/Renderer/Private/HaltonUtilities.cpp b/Engine/Source/Runtime/Renderer/Private/HaltonUtilities.cpp
--- a/Engine/Source/Runtime/Renderer/Private/HaltonUtilities.cpp
+++ b/Engine/Source/Runtime/Renderer/Private/HaltonUtilities.cpp
@@ -31,7 +31,7 @@
 
 #define PERMUTATION_BUFFER TArray
 
-float asfloat(unsigned int integer)
+static float asfloat(unsigned int integer)
 {
 	union {
 		unsigned int u;
