[Ray Tracing Instance Gathering Refactor] Replace FRayTracingRelevantPrimitiveTaskData with RayTracing::FGatherInstancesTaskData.
- Moved launching and waiting for task to run GatherRelevantPrimitives(...) into RayTracing.cpp (new function RayTracing::BeginGatherInstances(...)).
- Store Scene/View references in FGatherInstancesTaskData instead of passing it to individual functions.

diff --git a/Engine/Source/Runtime/RenderCore/Private/RenderCore.cpp b/Engine/Source/Runtime/RenderCore/Private/RenderCore.cpp
--- a/Engine/Source/Runtime/RenderCore/Private/RenderCore.cpp
+++ b/Engine/Source/Runtime/RenderCore/Private/RenderCore.cpp
@@ -77,7 +77,7 @@ DEFINE_STAT(STAT_DynamicShadowSetupTime);
 DEFINE_STAT(STAT_RenderQueryResultTime);
 // Use 'stat initviews' to get more detail
 DEFINE_STAT(STAT_InitViewsTime);
-DEFINE_STAT(STAT_GatherRayTracingWorldInstances);
+DEFINE_STAT(STAT_RayTracing_FinishGatherInstances);
 DEFINE_STAT(STAT_InitViewsPossiblyAfterPrepass);
 // Measures the time spent in RenderViewFamily_RenderThread
 // Note that this is not the total rendering thread time, any other rendering commands will not be counted here
diff --git a/Engine/Source/Runtime/RenderCore/Public/RenderCore.h b/Engine/Source/Runtime/RenderCore/Public/RenderCore.h
--- a/Engine/Source/Runtime/RenderCore/Public/RenderCore.h
+++ b/Engine/Source/Runtime/RenderCore/Public/RenderCore.h
@@ -31,7 +31,7 @@ DECLARE_LOG_CATEGORY_EXTERN(LogRendererCore, Log, All);
 DECLARE_CYCLE_STAT_EXTERN(TEXT("BeginOcclusionTests"),STAT_BeginOcclusionTestsTime,STATGROUP_SceneRendering, RENDERCORE_API);
 DECLARE_CYCLE_STAT_EXTERN(TEXT("Waiting for GPU - check ProfileGPU"),STAT_RenderQueryResultTime,STATGROUP_SceneRendering, RENDERCORE_API);
 DECLARE_CYCLE_STAT_EXTERN(TEXT("InitViews"),STAT_InitViewsTime,STATGROUP_SceneRendering, RENDERCORE_API);
-DECLARE_CYCLE_STAT_EXTERN(TEXT("GatherRayTracingWorldInstances"), STAT_GatherRayTracingWorldInstances, STATGROUP_SceneRendering, RENDERCORE_API);
+DECLARE_CYCLE_STAT_EXTERN(TEXT("Finish Gather Ray Tracing Instances"), STAT_RayTracing_FinishGatherInstances, STATGROUP_SceneRendering, RENDERCORE_API);
 DECLARE_CYCLE_STAT_EXTERN(TEXT("InitViewsPossiblyAfterPrepass"), STAT_InitViewsPossiblyAfterPrepass, STATGROUP_SceneRendering, RENDERCORE_API);
 DECLARE_CYCLE_STAT_EXTERN(TEXT("Dynamic shadow setup"), STAT_DynamicShadowSetupTime, STATGROUP_SceneRendering, RENDERCORE_API);
 DECLARE_CYCLE_STAT_EXTERN(TEXT("Total GPU rendering"),STAT_TotalGPUFrameTime,STATGROUP_SceneRendering, RENDERCORE_API);
diff --git a/Engine/Source/Runtime/Renderer/Private/DeferredShadingRenderer.cpp b/Engine/Source/Runtime/Renderer/Private/DeferredShadingRenderer.cpp
--- a/Engine/Source/Runtime/Renderer/Private/DeferredShadingRenderer.cpp
+++ b/Engine/Source/Runtime/Renderer/Private/DeferredShadingRenderer.cpp
@@ -555,12 +555,6 @@ static bool ShouldPrepareRayTracingDecals(const FScene& Scene, const FSceneViewF
 	return ViewFamily.EngineShowFlags.PathTracing && PathTracing::UsesDecals(ViewFamily);
 }
 
-struct FRayTracingRelevantPrimitiveTaskData
-{
-	RayTracing::FRelevantPrimitiveList* List;
-	FGraphEventRef Task;
-};
-
 static void DeduplicateRayGenerationShaders(TArray< FRHIRayTracingShader*>& RayGenShaders)
 {
 	TSet<FRHIRayTracingShader*> UniqueRayGenShaders;
@@ -1491,16 +1485,16 @@ void FDeferredShadingSceneRenderer::Render(FRDGBuilder& GraphBuilder)
 		{
 			const int32 ReferenceViewIndex = 0;
 			FViewInfo& ReferenceView = Views[ReferenceViewIndex];
-			FGraphEventRef PrereqTask = CreateCompatibilityGraphEvent(MakeArrayView({ Scene->GetCacheRayTracingPrimitivesTask(), InitViewTaskDatas.VisibilityTaskData->GetFrustumCullTask() }));
 
-			InitViewTaskDatas.RayTracingRelevantPrimitives = Allocator.Create<FRayTracingRelevantPrimitiveTaskData>();
-			InitViewTaskDatas.RayTracingRelevantPrimitives->List = RayTracing::CreateRelevantPrimitiveList(Allocator);
-			InitViewTaskDatas.RayTracingRelevantPrimitives->Task = FFunctionGraphTask::CreateAndDispatchWhenReady(
-				[Scene = Scene, &ReferenceView, &RayTracingRelevantPrimitiveList = *InitViewTaskDatas.RayTracingRelevantPrimitives->List]()
-				{
-					FTaskTagScope TaskTagScope(ETaskTag::EParallelRenderingThread);
-					RayTracing::GatherRelevantPrimitives(*Scene, ReferenceView, RayTracingRelevantPrimitiveList);
-				}, TStatId(), PrereqTask, ENamedThreads::AnyNormalThreadHiPriTask);
+			InitViewTaskDatas.RayTracingGatherInstances = RayTracing::CreateGatherInstancesTaskData(
+				Allocator,
+				*Scene,
+				ReferenceView,
+				ViewFamily,
+				GetViewPipelineState(ReferenceView).DiffuseIndirectMethod,
+				GetViewPipelineState(ReferenceView).ReflectionsMethod);
+
+			RayTracing::BeginGatherInstances(*InitViewTaskDatas.RayTracingGatherInstances, InitViewTaskDatas.VisibilityTaskData->GetFrustumCullTask());
 		}
 #endif
 	}
@@ -1781,27 +1775,13 @@ void FDeferredShadingSceneRenderer::Render(FRDGBuilder& GraphBuilder)
 #if RHI_RAYTRACING
 		if (bAnyRayTracingPassEnabled)
 		{
-			// Wait until RayTracingRelevantPrimitiveList is ready
-			if (InitViewTaskDatas.RayTracingRelevantPrimitives->Task.IsValid())
-			{
-				InitViewTaskDatas.RayTracingRelevantPrimitives->Task->Wait();
-				InitViewTaskDatas.RayTracingRelevantPrimitives->Task.SafeRelease();
-			}
-
-			checkf(InitViewTaskDatas.RayTracingRelevantPrimitives->List, TEXT("Should have a valid FRelevantPrimitiveList containing the result of RayTracing::GatherRelevantPrimitives here."));
-
-			RayTracing::GatherWorldInstancesForView(
+			RayTracing::FinishGatherInstances(
 				GraphBuilder,
-				*Scene,
-				ViewFamily,
-				ReferenceView,
-				GetViewPipelineState(ReferenceView).DiffuseIndirectMethod,
-				GetViewPipelineState(ReferenceView).ReflectionsMethod,
+				*InitViewTaskDatas.RayTracingGatherInstances,
 				RayTracingScene,
 				RayTracingSBT,
 				DynamicReadBufferForRayTracing,
-				Allocator, 
-				*InitViewTaskDatas.RayTracingRelevantPrimitives->List);
+				Allocator);
 		}
 #endif // RHI_RAYTRACING
 
diff --git a/Engine/Source/Runtime/Renderer/Private/DeferredShadingRenderer.h b/Engine/Source/Runtime/Renderer/Private/DeferredShadingRenderer.h
--- a/Engine/Source/Runtime/Renderer/Private/DeferredShadingRenderer.h
+++ b/Engine/Source/Runtime/Renderer/Private/DeferredShadingRenderer.h
@@ -72,9 +72,10 @@ namespace Froxel
 
 class IVisibilityTaskData;
 
-#if RHI_RAYTRACING
-struct FRayTracingRelevantPrimitiveTaskData;
-#endif
+namespace RayTracing
+{
+	struct FGatherInstancesTaskData;
+}
 
 /**   
  * Data for rendering meshes into Surface Cache
@@ -560,9 +561,7 @@ private:
 
 		IVisibilityTaskData* VisibilityTaskData;
 		FILCUpdatePrimTaskData* ILCUpdatePrim = nullptr;
-	#if RHI_RAYTRACING
-		FRayTracingRelevantPrimitiveTaskData* RayTracingRelevantPrimitives = nullptr;
-	#endif
+		RayTracing::FGatherInstancesTaskData* RayTracingGatherInstances = nullptr;
 		FDynamicShadowsTaskData* DynamicShadows = nullptr;
 		FLumenDirectLightingTaskData* LumenDirectLighting = nullptr;
 		FLumenSceneFrameTemporaries* LumenFrameTemporaries = nullptr;
diff --git a/Engine/Source/Runtime/Renderer/Private/RayTracing/RayTracing.cpp b/Engine/Source/Runtime/Renderer/Private/RayTracing/RayTracing.cpp
--- a/Engine/Source/Runtime/Renderer/Private/RayTracing/RayTracing.cpp
+++ b/Engine/Source/Runtime/Renderer/Private/RayTracing/RayTracing.cpp
@@ -228,8 +228,22 @@ namespace RayTracing
 		int32 DirtyShaderBindingsOffset = -1;
 	};
 
-	struct FRelevantPrimitiveList
+	struct FGatherInstancesTaskData
 	{
+		UE_NONCOPYABLE(FGatherInstancesTaskData)
+
+		FGatherInstancesTaskData(FScene& InScene, FViewInfo& InView, FSceneOptions InSceneOptions)
+			: Scene(InScene)
+			, View(InView)
+			, SceneOptions(MoveTemp(InSceneOptions))
+		{
+
+		}
+
+		FScene& Scene;
+		FViewInfo& View;
+		FSceneOptions SceneOptions;
+
 		// Filtered lists of relevant primitives
 		TArray<FRelevantPrimitive> StaticPrimitives;
 		TArray<FRelevantPrimitive> CachedStaticPrimitives;
@@ -237,6 +251,8 @@ namespace RayTracing
 
 		TArray<FRelevantPrimitiveGatherContext> GatherContexts;
 
+		UE::Tasks::FTask GatherRelevantPrimitivesTask;
+
 		// Relevant static primitive LODs are computed asynchronously.
 		// This task must complete before accessing StaticPrimitives/CachedStaticPrimitives in FRayTracingSceneAddInstancesTask.
 		FGraphEventRef StaticPrimitiveLODTask;
@@ -256,6 +272,17 @@ namespace RayTracing
 		bool bValid = false;
 	};
 
+	FGatherInstancesTaskData* CreateGatherInstancesTaskData(
+		FSceneRenderingBulkObjectAllocator& InAllocator,
+		FScene& Scene,
+		FViewInfo& View,
+		const FViewFamilyInfo& ViewFamily,
+		EDiffuseIndirectMethod DiffuseIndirectMethod,
+		EReflectionsMethod ReflectionsMethod)
+	{
+		return InAllocator.Create<FGatherInstancesTaskData>(Scene, View, FSceneOptions(Scene, ViewFamily, View, DiffuseIndirectMethod, ReflectionsMethod));
+	}
+
 	void OnRenderBegin(FScene& Scene, TArray<FViewInfo>& Views, const FViewFamilyInfo& ViewFamily)
 	{
 		const ERayTracingMeshCommandsMode CurrentMode = ViewFamily.EngineShowFlags.PathTracing ? ERayTracingMeshCommandsMode::PATH_TRACING : ERayTracingMeshCommandsMode::RAY_TRACING;
@@ -302,11 +329,6 @@ namespace RayTracing
 				}
 			}
 		}
-	}
-
-	FRelevantPrimitiveList* CreateRelevantPrimitiveList(FSceneRenderingBulkObjectAllocator& InAllocator)
-	{
-		return InAllocator.Create<FRelevantPrimitiveList>();
 	}	
 
 	class FRaytracingShaderBindingLayout : public FShaderBindingLayoutContainer
@@ -442,12 +464,14 @@ namespace RayTracing
 			}, TStatId(), nullptr, ENamedThreads::AnyThread));
 	};
 
-	void GatherRelevantPrimitives(FScene& Scene, const FViewInfo& View, FRelevantPrimitiveList& Result)
+	void GatherRelevantPrimitives(FGatherInstancesTaskData& TaskData, bool bUsingReferenceBasedResidency)
 	{
+		FScene& Scene = TaskData.Scene;
+		FViewInfo& View = TaskData.View;
+
 		TArray<int32> StaticPrimitives;
 
 		const bool bGameView = View.bIsGameView || View.Family->EngineShowFlags.Game;
-		const bool bUsingReferenceBasedResidency = IsRayTracingUsingReferenceBasedResidency();
 
 		bool bPerformRayTracing = View.State != nullptr && !View.bIsReflectionCapture && View.bAllowRayTracing;
 		if (bPerformRayTracing)
@@ -568,16 +592,16 @@ namespace RayTracing
 				}
 
 				StaticPrimitives.Reserve(NumStaticPrimitives);
-				Result.DynamicPrimitives.Reserve(NumDynamicPrimitives);
-				Result.UsedCoarseMeshStreamingHandles.Reserve(NumUsedCoarseMeshStreamingHandles);
-				Result.DirtyCachedRayTracingPrimitives.Reserve(NumDirtyCachedRayTracingPrimitives);
+				TaskData.DynamicPrimitives.Reserve(NumDynamicPrimitives);
+				TaskData.UsedCoarseMeshStreamingHandles.Reserve(NumUsedCoarseMeshStreamingHandles);
+				TaskData.DirtyCachedRayTracingPrimitives.Reserve(NumDirtyCachedRayTracingPrimitives);
 
 				for (auto& Context : Contexts)
 				{
 					Context.StaticPrimitives.CopyToLinearArray(StaticPrimitives);
-					Context.DynamicPrimitives.CopyToLinearArray(Result.DynamicPrimitives);
-					Context.UsedCoarseMeshStreamingHandles.CopyToLinearArray(Result.UsedCoarseMeshStreamingHandles);
-					Context.DirtyCachedRayTracingPrimitives.CopyToLinearArray(Result.DirtyCachedRayTracingPrimitives);
+					Context.DynamicPrimitives.CopyToLinearArray(TaskData.DynamicPrimitives);
+					Context.UsedCoarseMeshStreamingHandles.CopyToLinearArray(TaskData.UsedCoarseMeshStreamingHandles);
+					Context.DirtyCachedRayTracingPrimitives.CopyToLinearArray(TaskData.DirtyCachedRayTracingPrimitives);
 
 					if(bUsingReferenceBasedResidency)
 					{
@@ -588,14 +612,14 @@ namespace RayTracing
 		}
 
 		// TODO: check whether it's ok to do this on a parallel task
-		FPrimitiveSceneInfo::UpdateCachedRaytracingData(&Scene, Result.DirtyCachedRayTracingPrimitives);
+		FPrimitiveSceneInfo::UpdateCachedRaytracingData(&TaskData.Scene, TaskData.DirtyCachedRayTracingPrimitives);
 
 		static const auto ICVarStaticMeshLODDistanceScale = IConsoleManager::Get().FindConsoleVariable(TEXT("r.StaticMeshLODDistanceScale"));
 		const float LODScaleCVarValue = ICVarStaticMeshLODDistanceScale->GetFloat();
 		const int32 ForcedLODLevel = GetCVarForceLOD();
 
-		Result.StaticPrimitiveLODTask = FFunctionGraphTask::CreateAndDispatchWhenReady(
-			[&Result, &Scene, &View, LODScaleCVarValue, ForcedLODLevel, StaticPrimitiveIndices = MoveTemp(StaticPrimitives), bUsingReferenceBasedResidency]()
+		TaskData.StaticPrimitiveLODTask = FFunctionGraphTask::CreateAndDispatchWhenReady(
+			[&TaskData, LODScaleCVarValue, ForcedLODLevel, StaticPrimitiveIndices = MoveTemp(StaticPrimitives), bUsingReferenceBasedResidency]()
 			{
 				FTaskTagScope TaskTagScope(ETaskTag::EParallelRenderingThread);
 
@@ -623,7 +647,7 @@ namespace RayTracing
 					Contexts,
 					StaticPrimitiveIndices.Num(),
 					[](int32 ContextIndex, int32 NumContexts) { return ContextIndex; },
-					[&Scene, &View, LODScaleCVarValue, ForcedLODLevel, &StaticPrimitiveIndices, bUsingReferenceBasedResidency](FRelevantStaticPrimitivesContext& Context, int32 ItemIndex)
+					[&Scene = TaskData.Scene, &View = TaskData.View, LODScaleCVarValue, ForcedLODLevel, &StaticPrimitiveIndices, bUsingReferenceBasedResidency](FRelevantStaticPrimitivesContext& Context, int32 ItemIndex)
 					{
 						const int32 PrimitiveIndex = StaticPrimitiveIndices[ItemIndex];
 
@@ -801,27 +825,27 @@ namespace RayTracing
 							NumCachedStaticPrimitives += Context.CachedStaticPrimitives.Num();
 						}
 
-						Result.StaticPrimitives.Reserve(NumStaticPrimitives);
-						Result.CachedStaticPrimitives.Reserve(NumCachedStaticPrimitives);
+						TaskData.StaticPrimitives.Reserve(NumStaticPrimitives);
+						TaskData.CachedStaticPrimitives.Reserve(NumCachedStaticPrimitives);
 
-						Result.GatherContexts.SetNum(Contexts.Num());
+						TaskData.GatherContexts.SetNum(Contexts.Num());
 
 						for (int32 ContextIndex = 0; ContextIndex < Contexts.Num(); ++ContextIndex)
 						{
 							FRelevantStaticPrimitivesContext& Context = Contexts[ContextIndex];
-							FRelevantPrimitiveGatherContext& GatherContext = Result.GatherContexts[ContextIndex];
+							FRelevantPrimitiveGatherContext& GatherContext = TaskData.GatherContexts[ContextIndex];
 
-							Context.StaticPrimitives.CopyToLinearArray(Result.StaticPrimitives);
-							Context.CachedStaticPrimitives.CopyToLinearArray(Result.CachedStaticPrimitives);
+							Context.StaticPrimitives.CopyToLinearArray(TaskData.StaticPrimitives);
+							Context.CachedStaticPrimitives.CopyToLinearArray(TaskData.CachedStaticPrimitives);
 
-							GatherContext.InstanceOffset = Result.NumCachedStaticInstances;
-							GatherContext.DecalInstanceOffset = Result.NumCachedStaticDecalInstances;
-							GatherContext.DirtyShaderBindingsOffset = Result.NumCachedStaticDirtyShaderBindings;
+							GatherContext.InstanceOffset = TaskData.NumCachedStaticInstances;
+							GatherContext.DecalInstanceOffset = TaskData.NumCachedStaticDecalInstances;
+							GatherContext.DirtyShaderBindingsOffset = TaskData.NumCachedStaticDirtyShaderBindings;
 
-							Result.NumCachedStaticInstances += Context.NumCachedStaticInstances;
-							Result.NumCachedStaticDecalInstances += Context.NumCachedStaticDecalInstances;
-							Result.NumCachedStaticSegments += Context.NumCachedStaticSegments;
-							Result.NumCachedStaticDirtyShaderBindings += Context.NumCachedStaticDirtyShaderBindings;
+							TaskData.NumCachedStaticInstances += Context.NumCachedStaticInstances;
+							TaskData.NumCachedStaticDecalInstances += Context.NumCachedStaticDecalInstances;
+							TaskData.NumCachedStaticSegments += Context.NumCachedStaticSegments;
+							TaskData.NumCachedStaticDirtyShaderBindings += Context.NumCachedStaticDirtyShaderBindings;
 
 							for (const FPrimitiveSceneInfo* SceneInfo : Context.VisibleNaniteRayTracingPrimitives)
 							{
@@ -831,27 +855,40 @@ namespace RayTracing
 					}
 			}, TStatId(), nullptr, ENamedThreads::AnyThread);
 
-		Result.bValid = true;
+		TaskData.bValid = true;
 	}
 
-	bool GatherWorldInstancesForView(
+	void BeginGatherInstances(FGatherInstancesTaskData& TaskData, UE::Tasks::FTask FrustumCullTask)
+	{
+		const bool bUsingReferenceBasedResidency = IsRayTracingUsingReferenceBasedResidency();
+
+		UE::Tasks::FTask CacheRayTracingPrimitivesTask = TaskData.Scene.GetCacheRayTracingPrimitivesTask();
+
+		TaskData.GatherRelevantPrimitivesTask = UE::Tasks::Launch(UE_SOURCE_LOCATION, [&TaskData, bUsingReferenceBasedResidency]
+			{
+				FTaskTagScope Scope(ETaskTag::EParallelRenderingThread);
+				GatherRelevantPrimitives(TaskData, bUsingReferenceBasedResidency);
+			}, UE::Tasks::Prerequisites(CacheRayTracingPrimitivesTask, FrustumCullTask), UE::Tasks::ETaskPriority::High);
+	}
+
+	bool FinishGatherInstances(
 		FRDGBuilder& GraphBuilder,
-		FScene& Scene,
-		const FViewFamilyInfo& ViewFamily,
-		FViewInfo& View,
-		EDiffuseIndirectMethod DiffuseIndirectMethod,
-		EReflectionsMethod ReflectionsMethod,
+		FGatherInstancesTaskData& TaskData,
 		FRayTracingScene& RayTracingScene,
 		FRayTracingShaderBindingTable& RayTracingSBT,
 		FGlobalDynamicReadBuffer& InDynamicReadBuffer,
-		FSceneRenderingBulkObjectAllocator& InBulkAllocator,
-		FRelevantPrimitiveList& RelevantPrimitiveList)
+		FSceneRenderingBulkObjectAllocator& InBulkAllocator)
 	{
-		TRACE_CPUPROFILER_EVENT_SCOPE(GatherRayTracingWorldInstances);
-		SCOPE_CYCLE_COUNTER(STAT_GatherRayTracingWorldInstances);
+		TRACE_CPUPROFILER_EVENT_SCOPE(RayTracing_FinishGatherInstances);
+		SCOPE_CYCLE_COUNTER(STAT_RayTracing_FinishGatherInstances);
+
+		TaskData.GatherRelevantPrimitivesTask.Wait();
+
+		FScene& Scene = TaskData.Scene;
+		FViewInfo& View = TaskData.View;
 
 		// Prepare ray tracing scene instance list
-		checkf(RelevantPrimitiveList.bValid, TEXT("Ray tracing relevant primitive list is expected to have been created before GatherRayTracingWorldInstancesForView() is called."));
+		checkf(TaskData.bValid, TEXT("Ray tracing relevant primitive list is expected to have been created before GatherRayTracingWorldInstancesForView() is called."));
 
 		// Check that any invalidated cached uniform expressions have been updated on the rendering thread.
 		// Normally this work is done through FMaterialRenderProxy::UpdateUniformExpressionCacheIfNeeded,
@@ -867,12 +904,10 @@ namespace RayTracing
 
 		RayTracingSBT.ResetDynamicAllocationData();
 
-		RayTracing::FSceneOptions SceneOptions(Scene, ViewFamily, View, DiffuseIndirectMethod, ReflectionsMethod);
-
 		const float CurrentWorldTime = View.Family->Time.GetWorldTimeSeconds();
 
 		// Consume output of the relevant primitive gathering task
-		RayTracingScene.UsedCoarseMeshStreamingHandles = MoveTemp(RelevantPrimitiveList.UsedCoarseMeshStreamingHandles);
+		RayTracingScene.UsedCoarseMeshStreamingHandles = MoveTemp(TaskData.UsedCoarseMeshStreamingHandles);
 
 		// Inform the coarse mesh streaming manager about all the used streamable render assets in the scene
 		Nanite::FCoarseMeshStreamingManager* CoarseMeshSM = IStreamingManager::Get().GetNaniteCoarseMeshStreamingManager();
@@ -881,7 +916,7 @@ namespace RayTracing
 			CoarseMeshSM->AddUsedStreamingHandles(RayTracingScene.UsedCoarseMeshStreamingHandles);
 		}
 
-		INC_DWORD_STAT_BY(STAT_VisibleRayTracingPrimitives, RelevantPrimitiveList.DynamicPrimitives.Num() + RelevantPrimitiveList.StaticPrimitives.Num());
+		INC_DWORD_STAT_BY(STAT_VisibleRayTracingPrimitives, TaskData.DynamicPrimitives.Num() + TaskData.StaticPrimitives.Num());
 
 		{
 			TRACE_CPUPROFILER_EVENT_SCOPE(GatherRayTracingWorldInstances_DynamicElements);
@@ -929,7 +964,7 @@ namespace RayTracing
 					InDynamicReadBuffer
 				);
 
-				for (int32 PrimitiveIndex : RelevantPrimitiveList.DynamicPrimitives)
+				for (int32 PrimitiveIndex : TaskData.DynamicPrimitives)
 				{
 					check(MaterialGatheringContext.DynamicRayTracingGeometriesToUpdate.IsEmpty());
 
@@ -941,7 +976,7 @@ namespace RayTracing
 
 					int32 BaseRayTracingInstance = DynamicRayTracingInstances.Num();
 
-					if (SceneOptions.bTranslucentGeometry || SceneProxy->IsOpaqueOrMasked())
+					if (TaskData.SceneOptions.bTranslucentGeometry || SceneProxy->IsOpaqueOrMasked())
 					{
 						SceneProxy->GetDynamicRayTracingInstances(MaterialGatheringContext, DynamicRayTracingInstances);
 					}
@@ -976,9 +1011,9 @@ namespace RayTracing
 			const int32 ViewDynamicPrimitiveId = View.RayTracingDynamicPrimitiveCollector.GetPrimitiveIdRange().GetLowerBoundValue();
 			const int32 ViewInstanceSceneDataOffset = View.RayTracingDynamicPrimitiveCollector.GetInstanceSceneDataOffset();
 
-			for (int32 Index = 0; Index < RelevantPrimitiveList.DynamicPrimitives.Num(); ++Index)
+			for (int32 Index = 0; Index < TaskData.DynamicPrimitives.Num(); ++Index)
 			{
-				const int32 PrimitiveIndex = RelevantPrimitiveList.DynamicPrimitives[Index];
+				const int32 PrimitiveIndex = TaskData.DynamicPrimitives[Index];
 				FPrimitiveSceneProxy* SceneProxy = Scene.PrimitiveSceneProxies[PrimitiveIndex];
 				FPrimitiveSceneInfo* SceneInfo = Scene.Primitives[PrimitiveIndex];
 				const FPersistentPrimitiveIndex PersistentPrimitiveIndex = SceneInfo->GetPersistentIndex();
@@ -1270,27 +1305,19 @@ namespace RayTracing
 			FRayTracingScene& RayTracingScene; // New instances are added into FRayTracingScene::Instances and FRayTracingScene::Allocator is used for temporary data
 			TArray<FRayTracingShaderBindingData>& DirtyShaderBindingData; // New elements are added here by this task
 
-			FRayTracingSceneAddStaticInstancesTask(const FScene& InScene,
-				const RayTracing::FSceneOptions& InSceneOptions,
-				TArray<FRelevantPrimitive>& InRelevantStaticPrimitives,
-				TArray<FRelevantPrimitive>& InRelevantCachedStaticPrimitives,
-				TArray<FRelevantPrimitiveGatherContext>& InGatherContexts,
-				const bool bInIsPathTracing,
-				const int32& InNumCachedStaticInstances,
-				const int32& InNumCachedStaticDecalInstances,
-				const int32& InNumCachedStaticSegments,
-				const int32& InNumCachedStaticDirtyShaderBindings,
+			FRayTracingSceneAddStaticInstancesTask(
+				FGatherInstancesTaskData& TaskData, const bool bInIsPathTracing,
 				FRayTracingScene& InRayTracingScene, TArray<FRayTracingShaderBindingData>& InDirtyShaderBindingData)
-				: Scene(InScene)
-				, SceneOptions(InSceneOptions)
-				, RelevantStaticPrimitives(InRelevantStaticPrimitives)
-				, RelevantCachedStaticPrimitives(InRelevantCachedStaticPrimitives)
-				, GatherContexts(InGatherContexts)
+				: Scene(TaskData.Scene)
+				, SceneOptions(TaskData.SceneOptions)
+				, RelevantStaticPrimitives(TaskData.StaticPrimitives)
+				, RelevantCachedStaticPrimitives(TaskData.CachedStaticPrimitives)
+				, GatherContexts(TaskData.GatherContexts)
 				, bIsPathTracing(bInIsPathTracing)
-				, NumCachedStaticInstances(InNumCachedStaticInstances)
-				, NumCachedStaticDecalInstances(InNumCachedStaticDecalInstances)
-				, NumCachedStaticSegments(InNumCachedStaticSegments)
-				, NumCachedStaticDirtyShaderBindings(InNumCachedStaticDirtyShaderBindings)
+				, NumCachedStaticInstances(TaskData.NumCachedStaticInstances)
+				, NumCachedStaticDecalInstances(TaskData.NumCachedStaticDecalInstances)
+				, NumCachedStaticSegments(TaskData.NumCachedStaticSegments)
+				, NumCachedStaticDirtyShaderBindings(TaskData.NumCachedStaticDirtyShaderBindings)
 				, RayTracingScene(InRayTracingScene)
 				, DirtyShaderBindingData(InDirtyShaderBindingData)
 			{
@@ -1624,20 +1651,12 @@ namespace RayTracing
 		};
 
 		FGraphEventArray AddStaticInstancesTaskPrerequisites;
-		AddStaticInstancesTaskPrerequisites.Add(RelevantPrimitiveList.StaticPrimitiveLODTask);
+		AddStaticInstancesTaskPrerequisites.Add(TaskData.StaticPrimitiveLODTask);
 
 		FGraphEventRef AddStaticInstancesTask = TGraphTask<FRayTracingSceneAddStaticInstancesTask>::CreateTask(&AddStaticInstancesTaskPrerequisites).ConstructAndDispatchWhenReady(
 			// inputs
-			Scene,
-			SceneOptions,
-			RelevantPrimitiveList.StaticPrimitives,
-			RelevantPrimitiveList.CachedStaticPrimitives,
-			RelevantPrimitiveList.GatherContexts,
+			TaskData,
 			bool(View.Family->EngineShowFlags.PathTracing),
-			RelevantPrimitiveList.NumCachedStaticInstances,
-			RelevantPrimitiveList.NumCachedStaticDecalInstances,
-			RelevantPrimitiveList.NumCachedStaticSegments,
-			RelevantPrimitiveList.NumCachedStaticDirtyShaderBindings,
 			// outputs
 			RayTracingScene,
 			View.DirtyRayTracingShaderBindings
diff --git a/Engine/Source/Runtime/Renderer/Private/RayTracing/RayTracing.h b/Engine/Source/Runtime/Renderer/Private/RayTracing/RayTracing.h
--- a/Engine/Source/Runtime/Renderer/Private/RayTracing/RayTracing.h
+++ b/Engine/Source/Runtime/Renderer/Private/RayTracing/RayTracing.h
@@ -37,36 +37,35 @@ namespace RayTracing
 
 namespace RayTracing
 {
-	struct FRelevantPrimitiveList;
+	struct FGatherInstancesTaskData;
 
 	void OnRenderBegin(FScene& Scene, TArray<FViewInfo>& Views, const FViewFamilyInfo& ViewFamily);
 
-	FRelevantPrimitiveList* CreateRelevantPrimitiveList(FSceneRenderingBulkObjectAllocator& InAllocator);
-
 	// Get shader resource table desc used for all raytracing shaders which is shared between all shaders in the RTPSO
 	const FShaderBindingLayout* GetShaderBindingLayout(EShaderPlatform ShaderPlatform);
 
 	// Setup the runtime static uniform buffer bindings on the command list if enabled
 	TOptional<FScopedUniformBufferStaticBindings> BindStaticUniformBufferBindings(const FViewInfo& View, FRHIUniformBuffer* SceneUniformBuffer, FRHICommandList& RHICmdList);
+
+	FGatherInstancesTaskData* CreateGatherInstancesTaskData(
+		FSceneRenderingBulkObjectAllocator& InAllocator,
+		FScene& Scene,
+		FViewInfo& View,
+		const FViewFamilyInfo& ViewFamily,
+		EDiffuseIndirectMethod DiffuseIndirectMethod,
+		EReflectionsMethod ReflectionsMethod);
 	
-	// Iterates over Scene's PrimitiveSceneProxies and extracts ones that are relevant for ray tracing.
-	// This function can run on any thread.
-	void GatherRelevantPrimitives(FScene& Scene, const FViewInfo& View, FRelevantPrimitiveList& OutRelevantPrimitiveList);
+	void BeginGatherInstances(FGatherInstancesTaskData& TaskData, UE::Tasks::FTask FrustumCullTask);
 
 	// Fills RayTracingScene instance list for the given View and adds relevant ray tracing data to the view. Does not reset previous scene contents.
 	// This function must run on render thread
-	bool GatherWorldInstancesForView(
+	bool FinishGatherInstances(
 		FRDGBuilder& GraphBuilder,
-		FScene& Scene,
-		const FViewFamilyInfo& ViewFamily,
-		FViewInfo& View,
-		EDiffuseIndirectMethod DiffuseIndirectMethod,
-		EReflectionsMethod ReflectionsMethod,
+		FGatherInstancesTaskData& TaskData,
 		FRayTracingScene& RayTracingScene,
 		FRayTracingShaderBindingTable& RayTracingSBT,
 		FGlobalDynamicReadBuffer& InDynamicReadBuffer,
-		FSceneRenderingBulkObjectAllocator& InBulkAllocator,
-		FRelevantPrimitiveList& RelevantPrimitiveList);
+		FSceneRenderingBulkObjectAllocator& InBulkAllocator);
 
 	bool ShouldExcludeDecals();
 }
