Fixed VisibleRayTracingPrimitives stat not counting primitives using cached ray tracing instances.

diff --git a/Engine/Source/Runtime/Renderer/Private/RayTracing/RayTracing.cpp b/Engine/Source/Runtime/Renderer/Private/RayTracing/RayTracing.cpp
--- a/Engine/Source/Runtime/Renderer/Private/RayTracing/RayTracing.cpp
+++ b/Engine/Source/Runtime/Renderer/Private/RayTracing/RayTracing.cpp
@@ -916,7 +916,7 @@ namespace RayTracing
 			CoarseMeshSM->AddUsedStreamingHandles(RayTracingScene.UsedCoarseMeshStreamingHandles);
 		}
 
-		INC_DWORD_STAT_BY(STAT_VisibleRayTracingPrimitives, TaskData.DynamicPrimitives.Num() + TaskData.StaticPrimitives.Num());
+		INC_DWORD_STAT_BY(STAT_VisibleRayTracingPrimitives, TaskData.StaticPrimitives.Num() + TaskData.CachedStaticPrimitives.Num() + TaskData.DynamicPrimitives.Num());
 
 		{
 			TRACE_CPUPROFILER_EVENT_SCOPE(GatherRayTracingWorldInstances_DynamicElements);
