Split most of logic inside FinishGatherInstances(...) into 3 separate functions.
- AddStaticInstancesToRayTracingScene / AddCachedStaticInstancesToRayTracingScene / AddDynamicInstancesToRayTracingScene.

diff --git a/Engine/Source/Runtime/Renderer/Private/RayTracing/RayTracing.cpp b/Engine/Source/Runtime/Renderer/Private/RayTracing/RayTracing.cpp
--- a/Engine/Source/Runtime/Renderer/Private/RayTracing/RayTracing.cpp
+++ b/Engine/Source/Runtime/Renderer/Private/RayTracing/RayTracing.cpp
@@ -858,809 +858,810 @@ namespace RayTracing
 		TaskData.bValid = true;
 	}
 
-	void BeginGatherInstances(FGatherInstancesTaskData& TaskData, UE::Tasks::FTask FrustumCullTask)
+	struct FAutoInstanceBatch
 	{
-		const bool bUsingReferenceBasedResidency = IsRayTracingUsingReferenceBasedResidency();
+		FRayTracingScene::FInstanceHandle InstanceHandle = FRayTracingScene::INVALID_INSTANCE_HANDLE;
+		FRayTracingScene::FInstanceHandle DecalInstanceHandle = FRayTracingScene::INVALID_INSTANCE_HANDLE;
 
-		UE::Tasks::FTask CacheRayTracingPrimitivesTask = TaskData.Scene.GetCacheRayTracingPrimitivesTask();
+		// Copies the next InstanceSceneDataOffset and user data into the current batch, returns true if arrays were re-allocated.
+		bool Add(FRayTracingScene& InRayTracingScene, uint32 InInstanceSceneDataOffset)
+		{
+			// Adhoc TArray-like resize behavior, in lieu of support for using a custom FMemStackBase in TArray.
+			// Idea for future: if batch becomes large enough, we could actually split it into multiple instances to avoid memory waste.
 
-		TaskData.GatherRelevantPrimitivesTask = UE::Tasks::Launch(UE_SOURCE_LOCATION, [&TaskData, bUsingReferenceBasedResidency]
+			const bool bNeedReallocation = Cursor == InstanceSceneDataOffsets.Num();
+
+			if (bNeedReallocation)
 			{
-				FTaskTagScope Scope(ETaskTag::EParallelRenderingThread);
-				GatherRelevantPrimitives(TaskData, bUsingReferenceBasedResidency);
-			}, UE::Tasks::Prerequisites(CacheRayTracingPrimitivesTask, FrustumCullTask), UE::Tasks::ETaskPriority::High);
-	}
+				int32 PrevCount = InstanceSceneDataOffsets.Num();
+				int32 NextCount = FMath::Max(PrevCount * 2, 1);
 
-	bool FinishGatherInstances(
-		FRDGBuilder& GraphBuilder,
-		FGatherInstancesTaskData& TaskData,
-		FRayTracingScene& RayTracingScene,
-		FRayTracingShaderBindingTable& RayTracingSBT,
-		FGlobalDynamicReadBuffer& InDynamicReadBuffer,
-		FSceneRenderingBulkObjectAllocator& InBulkAllocator)
-	{
-		TRACE_CPUPROFILER_EVENT_SCOPE(RayTracing_FinishGatherInstances);
-		SCOPE_CYCLE_COUNTER(STAT_RayTracing_FinishGatherInstances);
+				TArrayView<uint32> NewInstanceSceneDataOffsets = InRayTracingScene.Allocate<uint32>(NextCount);
+				if (PrevCount)
+				{
+					FMemory::Memcpy(NewInstanceSceneDataOffsets.GetData(), InstanceSceneDataOffsets.GetData(), InstanceSceneDataOffsets.GetTypeSize() * InstanceSceneDataOffsets.Num());
+				}
+				InstanceSceneDataOffsets = NewInstanceSceneDataOffsets;
+			}
 
-		TaskData.GatherRelevantPrimitivesTask.Wait();
+			InstanceSceneDataOffsets[Cursor] = InInstanceSceneDataOffset;
 
-		FScene& Scene = TaskData.Scene;
-		FViewInfo& View = TaskData.View;
+			++Cursor;
 
-		// Prepare ray tracing scene instance list
-		checkf(TaskData.bValid, TEXT("Ray tracing relevant primitive list is expected to have been created before GatherRayTracingWorldInstancesForView() is called."));
-
-		// Check that any invalidated cached uniform expressions have been updated on the rendering thread.
-		// Normally this work is done through FMaterialRenderProxy::UpdateUniformExpressionCacheIfNeeded,
-		// however ray tracing material processing (FMaterialShader::GetShaderBindings, which accesses UniformExpressionCache)
-		// is done on task threads, therefore all work must be done here up-front as UpdateUniformExpressionCacheIfNeeded is not free-threaded.
-		check(!FMaterialRenderProxy::HasDeferredUniformExpressionCacheRequests());
+			return bNeedReallocation;
+		}
 
-		FGPUScenePrimitiveCollector DummyDynamicPrimitiveCollector;
+		bool IsValid() const
+		{
+			return InstanceSceneDataOffsets.Num() != 0;
+		}
 
-		View.DynamicRayTracingMeshCommandStorage.Reserve(Scene.Primitives.Num());
-		
-		View.RayTracingMeshResourceCollector = MakeUnique<FRayTracingMeshResourceCollector>(Scene.GetFeatureLevel(), InBulkAllocator);
+		TArrayView<uint32> InstanceSceneDataOffsets;
+		uint32 Cursor = 0;
+	}; 
+	
+	void AddStaticInstancesToRayTracingScene(
+		const FScene& Scene,
+		const RayTracing::FSceneOptions& SceneOptions,
+		TConstArrayView<FRelevantPrimitive> RelevantStaticPrimitives,
+		FRayTracingScene& RayTracingScene,
+		TArray<FRayTracingShaderBindingData>& DirtyShaderBindingData)
+	{
+		TRACE_CPUPROFILER_EVENT_SCOPE(RayTracingScene_AddStaticInstances);
 
-		RayTracingSBT.ResetDynamicAllocationData();
+		DirtyShaderBindingData.Reserve(DirtyShaderBindingData.Num() + RelevantStaticPrimitives.Num());
 
-		const float CurrentWorldTime = View.Family->Time.GetWorldTimeSeconds();
+		const bool bAutoInstance = CVarRayTracingAutoInstance.GetValueOnRenderThread() != 0;
 
-		// Consume output of the relevant primitive gathering task
-		RayTracingScene.UsedCoarseMeshStreamingHandles = MoveTemp(TaskData.UsedCoarseMeshStreamingHandles);
+		// Instance batches by FRelevantPrimitive::InstancingKey()
+		Experimental::TSherwoodMap<uint64, FAutoInstanceBatch> InstanceBatches;
 
-		// Inform the coarse mesh streaming manager about all the used streamable render assets in the scene
-		Nanite::FCoarseMeshStreamingManager* CoarseMeshSM = IStreamingManager::Get().GetNaniteCoarseMeshStreamingManager();
-		if (CoarseMeshSM)
+		// scan relevant primitives computing hash data to look for duplicate instances
+		for (const FRelevantPrimitive& RelevantPrimitive : RelevantStaticPrimitives)
 		{
-			CoarseMeshSM->AddUsedStreamingHandles(RayTracingScene.UsedCoarseMeshStreamingHandles);
-		}
+			const int32 PrimitiveIndex = RelevantPrimitive.PrimitiveIndex;
+			FPrimitiveSceneInfo* SceneInfo = Scene.Primitives[PrimitiveIndex];
+			FPrimitiveSceneProxy* SceneProxy = Scene.PrimitiveSceneProxies[PrimitiveIndex];
+			ERayTracingPrimitiveFlags Flags = Scene.PrimitiveRayTracingFlags[PrimitiveIndex];
+			const FPersistentPrimitiveIndex PersistentPrimitiveIndex = RelevantPrimitive.PersistentPrimitiveIndex;
 
-		INC_DWORD_STAT_BY(STAT_VisibleRayTracingPrimitives, TaskData.StaticPrimitives.Num() + TaskData.CachedStaticPrimitives.Num() + TaskData.DynamicPrimitives.Num());
+			check(!EnumHasAnyFlags(Flags, ERayTracingPrimitiveFlags::CacheInstances));
 
-		{
-			TRACE_CPUPROFILER_EVENT_SCOPE(GatherRayTracingWorldInstances_DynamicElements);
+			const int8 LODIndex = RelevantPrimitive.LODIndex;
 
-			const bool bParallelMeshBatchSetup = GRayTracingParallelMeshBatchSetup && FApp::ShouldUseThreadingForPerformance();
-			const bool bUsingReferenceBasedResidency = IsRayTracingUsingReferenceBasedResidency();
+			if (LODIndex < 0)
+			{
+				// TODO: Filter these primitives earlier
+				continue;
+			}
 
-			const int64 SharedBufferGenerationID = Scene.GetRayTracingDynamicGeometryCollection()->BeginUpdate();
+			const bool bNeedMainInstance = !RelevantPrimitive.CachedMeshCommandFlags.bAllSegmentsDecal;
 
-			FRayTracingMeshBatchTaskPage* MeshBatchTaskHead = nullptr;
-			FRayTracingMeshBatchTaskPage* MeshBatchTaskPage = nullptr;
-			uint32 NumPendingMeshBatches = 0;
-			const uint32 RayTracingParallelMeshBatchSize = GRayTracingParallelMeshBatchSize;
+			// if primitive has mixed decal and non-decal segments we need to have two ray tracing instances
+			// one containing non-decal segments and the other with decal segments
+			// masking of segments is done using "hidden" hitgroups
+			// TODO: Debug Visualization to highlight primitives using this?
+			const bool bNeedDecalInstance = RelevantPrimitive.CachedMeshCommandFlags.bAnySegmentsDecal && !ShouldExcludeDecals();
 
-			auto KickRayTracingMeshBatchTask = [&InBulkAllocator, &View, &Scene, &MeshBatchTaskHead, &MeshBatchTaskPage, &NumPendingMeshBatches]()
-				{
-					if (MeshBatchTaskHead)
-					{
-						DispatchRayTracingMeshBatchTask(InBulkAllocator, Scene, View, MeshBatchTaskHead, NumPendingMeshBatches);
-					}
+			if (ShouldExcludeDecals() && RelevantPrimitive.CachedMeshCommandFlags.bAllSegmentsDecal)
+			{
+				continue;
+			}
 
-					MeshBatchTaskHead = nullptr;
-					MeshBatchTaskPage = nullptr;
-					NumPendingMeshBatches = 0;
-				};
+			if (!SceneOptions.bTranslucentGeometry && RelevantPrimitive.CachedMeshCommandFlags.bAllSegmentsTranslucent)
+			{
+				continue;
+			}
 
-			// Need to process dynamic primitives in 3 passes to support dynamic primitives/instances in GPU Scene
-			// 1 - gather all dynamic ray tracing instances
-			// 2 - upload dynamic primitive/instance data to GPU scene
-			// 3 - process dynamic ray tracing instances (primitive IDs and instance scene data offset of dynamic primitives is now valid)
+			if (GRayTracingExcludeSky && RelevantPrimitive.CachedMeshCommandFlags.bIsSky)
+			{
+				continue;
+			}
 
-			// Local temporary array of instances used for GetDynamicRayTracingInstances()
-			TArray<FRayTracingInstance> DynamicRayTracingInstances;
-			TArray<TRange<int32>> PrimitivesDynamicRayTracingInstances;
+			// location if this is a new entry
+			const uint64 InstanceKey = RelevantPrimitive.InstancingKey();
 
-			{
-				FRayTracingMaterialGatheringContext MaterialGatheringContext
-				(
-					&Scene,
-					&View,
-					*View.Family,
-					GraphBuilder,
-					*View.RayTracingMeshResourceCollector,
-					View.RayTracingDynamicPrimitiveCollector,
-					InDynamicReadBuffer
-				);
-
-				for (int32 PrimitiveIndex : TaskData.DynamicPrimitives)
-				{
-					check(MaterialGatheringContext.DynamicRayTracingGeometriesToUpdate.IsEmpty());
+			FAutoInstanceBatch DummyInstanceBatch = { };
+			FAutoInstanceBatch& InstanceBatch = bAutoInstance ? InstanceBatches.FindOrAdd(InstanceKey, DummyInstanceBatch) : DummyInstanceBatch;
 
-					FPrimitiveSceneProxy* SceneProxy = Scene.PrimitiveSceneProxies[PrimitiveIndex];
-					FPrimitiveSceneInfo* SceneInfo = Scene.Primitives[PrimitiveIndex];
-					const FPersistentPrimitiveIndex PersistentPrimitiveIndex = SceneInfo->GetPersistentIndex();
+			if (InstanceBatch.IsValid())
+			{
+				// Reusing a previous entry, just append to the instance list.
 
-					MaterialGatheringContext.SetPrimitive(SceneProxy);
+				bool bReallocated = InstanceBatch.Add(RayTracingScene, SceneInfo->GetInstanceSceneDataOffset());
 
-					int32 BaseRayTracingInstance = DynamicRayTracingInstances.Num();
+				if (InstanceBatch.InstanceHandle.IsValid())
+				{
+					FRayTracingGeometryInstance& RayTracingInstance = RayTracingScene.GetInstance(InstanceBatch.InstanceHandle);
+					++RayTracingInstance.NumTransforms;
+					check(RayTracingInstance.NumTransforms == InstanceBatch.Cursor); // sanity check
 
-					if (TaskData.SceneOptions.bTranslucentGeometry || SceneProxy->IsOpaqueOrMasked())
+					if (bReallocated)
 					{
-						SceneProxy->GetDynamicRayTracingInstances(MaterialGatheringContext, DynamicRayTracingInstances);
+						RayTracingInstance.InstanceSceneDataOffsets = InstanceBatch.InstanceSceneDataOffsets;
+						RayTracingInstance.UserData = InstanceBatch.InstanceSceneDataOffsets;
 					}
+				}
 
-					for (const FRayTracingDynamicGeometryUpdateParams& DynamicRayTracingGeometryUpdate : MaterialGatheringContext.DynamicRayTracingGeometriesToUpdate)
-					{
-						Scene.GetRayTracingDynamicGeometryCollection()->AddDynamicMeshBatchForGeometryUpdate(
-							GraphBuilder.RHICmdList,
-							&Scene,
-							&View,
-							SceneProxy,
-							DynamicRayTracingGeometryUpdate,
-							PersistentPrimitiveIndex.Index
-						);
-					}
+				if (InstanceBatch.DecalInstanceHandle.IsValid())
+				{
+					FRayTracingGeometryInstance& RayTracingInstance = RayTracingScene.GetInstance(InstanceBatch.DecalInstanceHandle);
+					++RayTracingInstance.NumTransforms;
+					check(RayTracingInstance.NumTransforms == InstanceBatch.Cursor); // sanity check
 
-					if (bUsingReferenceBasedResidency)
+					if (bReallocated)
 					{
-						((FRayTracingGeometryManager*)GRayTracingGeometryManager)->AddReferencedGeometryGroups(MaterialGatheringContext.GetReferencedGeometryGroups());
+						RayTracingInstance.InstanceSceneDataOffsets = InstanceBatch.InstanceSceneDataOffsets;
+						RayTracingInstance.UserData = InstanceBatch.InstanceSceneDataOffsets;
 					}
-
-					MaterialGatheringContext.Reset();
-
-					PrimitivesDynamicRayTracingInstances.Add(TRange<int32>(BaseRayTracingInstance, DynamicRayTracingInstances.Num()));
 				}
-
-				// FRayTracingMaterialGatheringContext destructor handles committing dynamic mesh batches to GPU Scene
 			}
+			else
+			{
+				// Starting new instance batch
 
-			Scene.GPUScene.UploadDynamicPrimitiveShaderDataForView(GraphBuilder, View, true);
+				InstanceBatch.Add(RayTracingScene, SceneInfo->GetInstanceSceneDataOffset());
 
-			const int32 ViewDynamicPrimitiveId = View.RayTracingDynamicPrimitiveCollector.GetPrimitiveIdRange().GetLowerBoundValue();
-			const int32 ViewInstanceSceneDataOffset = View.RayTracingDynamicPrimitiveCollector.GetInstanceSceneDataOffset();
+				FRayTracingGeometryInstance RayTracingInstance;
+				RayTracingInstance.GeometryRHI = RelevantPrimitive.RayTracingGeometry->GetRHI();
+				checkf(RayTracingInstance.GeometryRHI, TEXT("Ray tracing instance must have a valid geometry."));
+				RayTracingInstance.InstanceSceneDataOffsets = InstanceBatch.InstanceSceneDataOffsets;
+				RayTracingInstance.UserData = InstanceBatch.InstanceSceneDataOffsets;
+				RayTracingInstance.NumTransforms = 1;
 
-			for (int32 Index = 0; Index < TaskData.DynamicPrimitives.Num(); ++Index)
-			{
-				const int32 PrimitiveIndex = TaskData.DynamicPrimitives[Index];
-				FPrimitiveSceneProxy* SceneProxy = Scene.PrimitiveSceneProxies[PrimitiveIndex];
-				FPrimitiveSceneInfo* SceneInfo = Scene.Primitives[PrimitiveIndex];
-				const FPersistentPrimitiveIndex PersistentPrimitiveIndex = SceneInfo->GetPersistentIndex();
+				RayTracingInstance.Mask = RelevantPrimitive.CachedMeshCommandFlags.InstanceMask; // When no cached command is found, InstanceMask == 0 and the instance is effectively filtered out
 
-				TArrayView<FRayTracingInstance> TempRayTracingInstances = MakeArrayView(
-					DynamicRayTracingInstances.GetData() + PrimitivesDynamicRayTracingInstances[Index].GetLowerBound().GetValue(),
-					PrimitivesDynamicRayTracingInstances[Index].Size<int32>());
+				// Run AHS for alpha masked and meshes with only some sections casting shadows, which require per mesh section filtering in AHS
+				if (RelevantPrimitive.CachedMeshCommandFlags.bAllSegmentsOpaque && (RelevantPrimitive.CachedMeshCommandFlags.bAllSegmentsCastShadow || !RelevantPrimitive.CachedMeshCommandFlags.bAnySegmentsCastShadow))
+				{
+					RayTracingInstance.Flags |= ERayTracingInstanceFlags::ForceOpaque;
+				}
+				if (RelevantPrimitive.CachedMeshCommandFlags.bTwoSided)
+				{
+					RayTracingInstance.Flags |= ERayTracingInstanceFlags::TriangleCullDisable;
+				}
+				if (RelevantPrimitive.CachedMeshCommandFlags.bAllSegmentsReverseCulling)
+				{
+					RayTracingInstance.Flags |= ERayTracingInstanceFlags::TriangleCullReverse;
+				}
+				AddDebugRayTracingInstanceFlags(RayTracingInstance.Flags);
 
-				if (TempRayTracingInstances.Num() > 0)
+				InstanceBatch.InstanceHandle = FRayTracingScene::INVALID_INSTANCE_HANDLE;
+
+				FPrimitiveSceneInfo::FRayTracingLODData& RTLODData = SceneInfo->RayTracingLODData[LODIndex];
+				ensure(RTLODData.SBTAllocation);
+				if (RTLODData.SBTAllocation == nullptr)
 				{
-					for (FRayTracingInstance& Instance : TempRayTracingInstances)
-					{
-						const FRayTracingGeometry* Geometry = Instance.Geometry;
-
-						if (!ensureMsgf(Geometry->DynamicGeometrySharedBufferGenerationID == FRayTracingGeometry::NonSharedVertexBuffers
-							|| Geometry->DynamicGeometrySharedBufferGenerationID == SharedBufferGenerationID,
-							TEXT("GenerationID %lld, but expected to be %lld or %lld. Geometry debug name: '%s'. ")
-							TEXT("When shared vertex buffers are used, the contents is expected to be written every frame. ")
-							TEXT("Possibly AddDynamicMeshBatchForGeometryUpdate() was not called for this geometry."),
-							Geometry->DynamicGeometrySharedBufferGenerationID, SharedBufferGenerationID, FRayTracingGeometry::NonSharedVertexBuffers,
-							*Geometry->Initializer.DebugName.ToString()))
-						{
-							continue;
-						}
+					continue;
+				}
 
-						((FRayTracingGeometryManager*)GRayTracingGeometryManager)->AddReferencedGeometry(Geometry);
+				uint32 SegmentCount = RelevantPrimitive.RayTracingGeometry->Initializer.Segments.Num();
+				if (bNeedMainInstance)
+				{
+					RayTracingScene.NumSegments += SegmentCount;
+					RayTracingInstance.InstanceContributionToHitGroupIndex = RTLODData.SBTAllocation->GetInstanceContributionToHitGroupIndex(ERayTracingSceneLayer::Base);
+					InstanceBatch.InstanceHandle = RayTracingScene.AddInstance(RayTracingInstance, ERayTracingSceneLayer::Base, SceneProxy, false);
+				}
 
-						if (Geometry->IsEvicted())
-						{
-							continue;
-						}
+				InstanceBatch.DecalInstanceHandle = FRayTracingScene::INVALID_INSTANCE_HANDLE;
+				if (bNeedDecalInstance)
+				{
+					RayTracingScene.NumSegments += SegmentCount;
+					FRayTracingGeometryInstance DecalRayTracingInstance = RayTracingInstance;
+					DecalRayTracingInstance.InstanceContributionToHitGroupIndex = RTLODData.SBTAllocation->GetInstanceContributionToHitGroupIndex(ERayTracingSceneLayer::Decals);
+					InstanceBatch.DecalInstanceHandle = RayTracingScene.AddInstance(MoveTemp(DecalRayTracingInstance), ERayTracingSceneLayer::Decals, SceneProxy, false);
+				}
 
-						// If geometry still has pending build request then add to list which requires a force build
-						if (Geometry->HasPendingBuildRequest())
-						{
-							RayTracingScene.GeometriesToBuild.Add(Geometry);
-						}
+				for (int32 CommandIndex : RTLODData.CachedMeshCommandIndices)
+				{
+					if (CommandIndex >= 0)
+					{
+						const FRayTracingMeshCommand& MeshCommand = Scene.CachedRayTracingMeshCommands[CommandIndex];
 
-						// Validate the material/segment counts
-						if (!ensureMsgf(Instance.GetMaterials().Num() == Geometry->Initializer.Segments.Num() ||
-							(Geometry->Initializer.Segments.Num() == 0 && Instance.GetMaterials().Num() == 1),
-							TEXT("Ray tracing material assignment validation failed for geometry '%s'. "
-								"Instance.GetMaterials().Num() = %d, Geometry->Initializer.Segments.Num() = %d."),
-							*Geometry->Initializer.DebugName.ToString(), Instance.GetMaterials().Num(),
-							Geometry->Initializer.Segments.Num()))
+						if (bNeedMainInstance)
 						{
-							continue;
+							const bool bHidden = MeshCommand.bDecal;
+							const uint32 RecordIndex = RTLODData.SBTAllocation->GetRecordIndex(ERayTracingSceneLayer::Base, MeshCommand.GeometrySegmentIndex);
+							FRayTracingShaderBindingData RTShaderBindingData(&MeshCommand, RelevantPrimitive.RayTracingGeometry->GetRHI(), RecordIndex, bHidden);
+							DirtyShaderBindingData.Add(RTShaderBindingData);
 						}
-
-						if (Instance.bInstanceMaskAndFlagsDirty || SceneInfo->bCachedRayTracingInstanceMaskAndFlagsDirty)
+						if (bNeedDecalInstance)
 						{
-							// Build InstanceMaskAndFlags since the data in SceneInfo is not up to date
-
-							FRayTracingMaskAndFlags InstanceMaskAndFlags;
-
-							if (Instance.GetMaterials().IsEmpty())
-							{
-								// If the material list is empty, explicitly set the mask to 0 so it will not be added in the raytracing scene
-								InstanceMaskAndFlags.Mask = 0;
-							}
-							else
-							{
-								InstanceMaskAndFlags = BuildRayTracingInstanceMaskAndFlags(Instance, *SceneProxy);
-							}
+							const bool bHidden = !MeshCommand.bDecal;
+							const uint32 RecordIndex = RTLODData.SBTAllocation->GetRecordIndex(ERayTracingSceneLayer::Decals, MeshCommand.GeometrySegmentIndex);
+							FRayTracingShaderBindingData RTShaderBindingData(&MeshCommand, RelevantPrimitive.RayTracingGeometry->GetRHI(), RecordIndex, bHidden);
+							DirtyShaderBindingData.Add(RTShaderBindingData);
+						}
+					}
+					else
+					{
+						// CommandIndex == -1 indicates that the mesh batch has been filtered by FRayTracingMeshProcessor (like the shadow depth pass batch)
+						// Do nothing in this case
+					}
+				}
+			}
+		}
+	}
 
-							SceneInfo->CachedRayTracingInstance.Mask = InstanceMaskAndFlags.Mask; // When no cached command is found, InstanceMask == 0 and the instance is effectively filtered out
+	void AddCachedStaticInstancesToRayTracingScene(
+		const FScene& Scene,
+		const RayTracing::FSceneOptions& SceneOptions,
+		TConstArrayView<FRelevantPrimitive> RelevantCachedStaticPrimitives,
+		TConstArrayView<FRelevantPrimitiveGatherContext> GatherContexts,
+		const int32 NumCachedStaticInstances,
+		const int32 NumCachedStaticDecalInstances,
+		const int32 NumCachedStaticSegments,
+		const int32 NumCachedStaticDirtyShaderBindings,
+		FRayTracingScene& RayTracingScene,
+		TArray<FRayTracingShaderBindingData>& DirtyShaderBindingData)
+	{
+		TRACE_CPUPROFILER_EVENT_SCOPE(RayTracingScene_AddCachedStaticInstances);
+
+		const FRayTracingScene::FInstanceRange CachedStaticInstanceRange = RayTracingScene.AllocateInstanceRangeUninitialized(NumCachedStaticInstances, ERayTracingSceneLayer::Base);
+		const FRayTracingScene::FInstanceRange CachedStaticDecalInstanceRange = RayTracingScene.AllocateInstanceRangeUninitialized(NumCachedStaticDecalInstances, ERayTracingSceneLayer::Decals);
+		const uint32 BaseCachedDirtyShaderBindingsIndex = DirtyShaderBindingData.AddUninitialized(NumCachedStaticDirtyShaderBindings);
+		RayTracingScene.NumSegments += NumCachedStaticSegments;
+
+		const int32 MinBatchSize = 128;
+		ParallelFor(
+			TEXT("RayTracingScene_AddCachedStaticInstances_ParallelFor"),
+			RelevantCachedStaticPrimitives.Num(),
+			MinBatchSize,
+			[&Scene, &RayTracingScene, &DirtyShaderBindingData, &RelevantCachedStaticPrimitives, &GatherContexts, CachedStaticInstanceRange, CachedStaticDecalInstanceRange, BaseCachedDirtyShaderBindingsIndex](int32 Index)
+			{
+				const FRelevantPrimitive& RelevantPrimitive = RelevantCachedStaticPrimitives[Index];
+				const int32 PrimitiveIndex = RelevantPrimitive.PrimitiveIndex;
+				FPrimitiveSceneInfo* SceneInfo = Scene.Primitives[PrimitiveIndex];
+				FPrimitiveSceneProxy* SceneProxy = Scene.PrimitiveSceneProxies[PrimitiveIndex];
+				ERayTracingPrimitiveFlags Flags = Scene.PrimitiveRayTracingFlags[PrimitiveIndex];
+				const FPersistentPrimitiveIndex PersistentPrimitiveIndex = RelevantPrimitive.PersistentPrimitiveIndex;
 
-							if (InstanceMaskAndFlags.bForceOpaque)
-							{
-								SceneInfo->CachedRayTracingInstance.Flags |= ERayTracingInstanceFlags::ForceOpaque;
-							}
+				check(EnumHasAnyFlags(Flags, ERayTracingPrimitiveFlags::CacheInstances));
 
-							if (InstanceMaskAndFlags.bDoubleSided)
-							{
-								SceneInfo->CachedRayTracingInstance.Flags |= ERayTracingInstanceFlags::TriangleCullDisable;
-							}
+				const bool bUsingNaniteRayTracing = (Nanite::GetRayTracingMode() != Nanite::ERayTracingMode::Fallback) && SceneProxy->IsNaniteMesh();
 
-							if (InstanceMaskAndFlags.bReverseCulling)
-							{
-								SceneInfo->CachedRayTracingInstance.Flags |= ERayTracingInstanceFlags::TriangleCullReverse;
-							}
+				if (bUsingNaniteRayTracing)
+				{
+					check(RelevantPrimitive.CachedRayTracingInstance->GeometryRHI != nullptr);
+				}
 
-							SceneInfo->bCachedRayTracingInstanceAnySegmentsDecal = InstanceMaskAndFlags.bAnySegmentsDecal;
-							SceneInfo->bCachedRayTracingInstanceAllSegmentsDecal = InstanceMaskAndFlags.bAllSegmentsDecal;
+				const bool bNeedMainInstance = !RelevantPrimitive.CachedMeshCommandFlags.bAllSegmentsDecal;
 
-							SceneInfo->bCachedRayTracingInstanceMaskAndFlagsDirty = false;
-						}
+				// if primitive has mixed decal and non-decal segments we need to have two ray tracing instances
+				// one containing non-decal segments and the other with decal segments
+				// masking of segments is done using "hidden" hitgroups
+				// TODO: Debug Visualization to highlight primitives using this?
+				const bool bNeedDecalInstance = RelevantPrimitive.CachedMeshCommandFlags.bAnySegmentsDecal && !RelevantPrimitive.CachedMeshCommandFlags.bAllSegmentsDecal && !ShouldExcludeDecals();
 
-						const bool bNeedMainInstance = !SceneInfo->bCachedRayTracingInstanceAllSegmentsDecal;
+				check(!ShouldExcludeDecals() || !RelevantPrimitive.CachedMeshCommandFlags.bAllSegmentsDecal);
+				check(RelevantPrimitive.CachedRayTracingInstance);
 
-						// if primitive has mixed decal and non-decal segments we need to have two ray tracing instances
-						// one containing non-decal segments and the other with decal segments
-						// masking of segments is done using "hidden" hitgroups
-						// TODO: Debug Visualization to highlight primitives using this?
-						const bool bNeedDecalInstance = SceneInfo->bCachedRayTracingInstanceAnySegmentsDecal && !ShouldExcludeDecals();
+				// Cached instances always use single LOD 0 - see above
+				const int8 LODIndex = 0;
+				FPrimitiveSceneInfo::FRayTracingLODData& RTLODData = SceneInfo->RayTracingLODData[LODIndex];
 
-						if (ShouldExcludeDecals() && SceneInfo->bCachedRayTracingInstanceAllSegmentsDecal)
-						{
-							continue;
-						}
+				if (bNeedMainInstance)
+				{
+					const int32 InstanceIndexInRange = GatherContexts[RelevantPrimitive.ContextIndex].InstanceOffset + RelevantPrimitive.RelativeInstanceOffset;
 
-						int32 PrimitiveId = PersistentPrimitiveIndex.Index;
-						int32 InstanceSceneDataOffset = SceneInfo->GetInstanceSceneDataOffset();
+					FRayTracingGeometryInstance RayTracingInstance = *RelevantPrimitive.CachedRayTracingInstance;
+					RayTracingInstance.InstanceContributionToHitGroupIndex = RTLODData.SBTAllocation->GetInstanceContributionToHitGroupIndex(ERayTracingSceneLayer::Base);
+					AddDebugRayTracingInstanceFlags(RayTracingInstance.Flags);
 
-						if (Instance.Materials.Num() > 0 && Instance.Materials[0].Elements.Num() > 0 && Instance.Materials[0].Elements[0].DynamicPrimitiveData != nullptr)
-						{
-							check(Instance.NumTransforms == Instance.Materials[0].Elements[0].NumInstances);
-							PrimitiveId = ViewDynamicPrimitiveId + Instance.Materials[0].Elements[0].DynamicPrimitiveIndex;
-							InstanceSceneDataOffset = ViewInstanceSceneDataOffset + Instance.Materials[0].Elements[0].DynamicPrimitiveInstanceSceneDataOffset;
-						}
+					RayTracingScene.SetInstance(CachedStaticInstanceRange, InstanceIndexInRange, MoveTemp(RayTracingInstance), SceneProxy, false);
+				}
 
-						FRayTracingGeometryInstance RayTracingInstance;
-						RayTracingInstance.GeometryRHI = Geometry->GetRHI();
-						checkf(RayTracingInstance.GeometryRHI, TEXT("Ray tracing instance must have a valid geometry."));
-						RayTracingInstance.DefaultUserData = InstanceSceneDataOffset;
-						RayTracingInstance.bIncrementUserDataPerInstance = true;
-						RayTracingInstance.bApplyLocalBoundsTransform = Instance.bApplyLocalBoundsTransform;
-						RayTracingInstance.Mask = SceneInfo->CachedRayTracingInstance.Mask;
-						RayTracingInstance.Flags = SceneInfo->CachedRayTracingInstance.Flags;
-						AddDebugRayTracingInstanceFlags(RayTracingInstance.Flags);
-
-						if (!Instance.GetPrimitiveInstanceIndices().IsEmpty())
-						{
-							TConstArrayView<uint32> PrimitiveInstanceIndices = Instance.GetPrimitiveInstanceIndices();
+				if (bNeedDecalInstance)
+				{
+					const int32 DecalInstanceIndexInRange = GatherContexts[RelevantPrimitive.ContextIndex].DecalInstanceOffset + RelevantPrimitive.RelativeDecalInstanceOffset;
 
-							// Convert from instance indices to InstanceSceneDataOffsets
-							TArrayView<uint32> InstanceSceneDataOffsets = RayTracingScene.Allocate<uint32>(PrimitiveInstanceIndices.Num());
-							for (int32 InstanceIndex = 0; InstanceIndex < PrimitiveInstanceIndices.Num(); ++InstanceIndex)
-							{
-								InstanceSceneDataOffsets[InstanceIndex] = SceneInfo->GetInstanceSceneDataOffset() + PrimitiveInstanceIndices[InstanceIndex];
-							}
+					FRayTracingGeometryInstance DecalRayTracingInstance = *RelevantPrimitive.CachedRayTracingInstance;
+					DecalRayTracingInstance.InstanceContributionToHitGroupIndex = RTLODData.SBTAllocation->GetInstanceContributionToHitGroupIndex(ERayTracingSceneLayer::Decals);
+					AddDebugRayTracingInstanceFlags(DecalRayTracingInstance.Flags);
 
-							RayTracingInstance.InstanceSceneDataOffsets = InstanceSceneDataOffsets;
-							RayTracingInstance.UserData = InstanceSceneDataOffsets;
-							RayTracingInstance.NumTransforms = PrimitiveInstanceIndices.Num();
-						}
-						else if(!Instance.GetTransforms().IsEmpty())
-						{
-							TConstArrayView<FMatrix> TransformsView;
-							if (Instance.OwnsTransforms())
-							{
-								// Slow path: copy transforms to the owned storage
-								checkf(Instance.InstanceTransformsView.Num() == 0, TEXT("InstanceTransformsView is expected to be empty if using InstanceTransforms"));
-								TArrayView<FMatrix> SceneOwnedTransforms = RayTracingScene.Allocate<FMatrix>(Instance.InstanceTransforms.Num());
-								FMemory::Memcpy(SceneOwnedTransforms.GetData(), Instance.InstanceTransforms.GetData(), Instance.InstanceTransforms.Num() * sizeof(RayTracingInstance.Transforms[0]));
-								static_assert(std::is_same_v<decltype(SceneOwnedTransforms[0]), decltype(Instance.InstanceTransforms[0])>, "Unexpected transform type");
+					RayTracingScene.SetInstance(CachedStaticDecalInstanceRange, DecalInstanceIndexInRange, MoveTemp(DecalRayTracingInstance), SceneProxy, false);
+				}
 
-								TransformsView = SceneOwnedTransforms;
-							}
-							else
-							{
-								// Fast path: just reference persistently-allocated transforms and avoid a copy
-								checkf(Instance.InstanceTransforms.Num() == 0, TEXT("InstanceTransforms is expected to be empty if using InstanceTransformsView"));
-								TransformsView = Instance.InstanceTransformsView;
-							}
+				if (RelevantPrimitive.RelativeDirtyShaderBindingsOffset != INDEX_NONE)
+				{
+					const int32 DirtyShadingBindingOffset = BaseCachedDirtyShaderBindingsIndex + GatherContexts[RelevantPrimitive.ContextIndex].DirtyShaderBindingsOffset + RelevantPrimitive.RelativeDirtyShaderBindingsOffset;
 
-							RayTracingInstance.NumTransforms = TransformsView.Num();
-							RayTracingInstance.Transforms = TransformsView;
-						}
-						else
-						{
-							// If array of transforms was not provided, get the instance transforms from GPU Scene
-							RayTracingInstance.NumTransforms = Instance.NumTransforms;
-							RayTracingInstance.BaseInstanceSceneDataOffset = InstanceSceneDataOffset;
-						}
+					int32 MainCommandOffset = DirtyShadingBindingOffset;
+					int32 DecalCommandOffset = DirtyShadingBindingOffset + (bNeedMainInstance ? RTLODData.CachedMeshCommandIndices.Num() : 0);
 
-						ERayTracingSceneLayerMask ActiveLayers = ERayTracingSceneLayerMask::None;
-						if (bNeedMainInstance)
-						{
-							EnumAddFlags(ActiveLayers, ERayTracingSceneLayerMask::Base);
-						}
-						if (bNeedDecalInstance)
-						{
-							EnumAddFlags(ActiveLayers, ERayTracingSceneLayerMask::Decals);
-						}
+					for (int32 CommandIndex : RTLODData.CachedMeshCommandIndices)
+					{
+						const FRayTracingMeshCommand& MeshCommand = Scene.CachedRayTracingMeshCommands[CommandIndex];
 
-						FRayTracingSBTAllocation* SBTAllocation = RayTracingSBT.AllocateDynamicRange(ActiveLayers, Instance.GetMaterials().Num());
 						if (bNeedMainInstance)
 						{
-							RayTracingScene.NumSegments += Instance.GetMaterials().Num();
-
-							RayTracingInstance.InstanceContributionToHitGroupIndex = SBTAllocation->GetInstanceContributionToHitGroupIndex(ERayTracingSceneLayer::Base);
-							RayTracingScene.AddInstance(RayTracingInstance, ERayTracingSceneLayer::Base, SceneProxy, true);
+							const bool bHidden = MeshCommand.bDecal;
+							const uint32 RecordIndex = RTLODData.SBTAllocation->GetRecordIndex(ERayTracingSceneLayer::Base, MeshCommand.GeometrySegmentIndex);
+							DirtyShaderBindingData[MainCommandOffset] = FRayTracingShaderBindingData(&MeshCommand, RelevantPrimitive.CachedRayTracingInstance->GeometryRHI, RecordIndex, bHidden);
+							++MainCommandOffset;
 						}
 
 						if (bNeedDecalInstance)
 						{
-							RayTracingScene.NumSegments += Instance.GetMaterials().Num();
-
-							FRayTracingGeometryInstance DecalRayTracingInstance = RayTracingInstance;
-							DecalRayTracingInstance.InstanceContributionToHitGroupIndex = SBTAllocation->GetInstanceContributionToHitGroupIndex(ERayTracingSceneLayer::Decals);
-							RayTracingScene.AddInstance(MoveTemp(DecalRayTracingInstance), ERayTracingSceneLayer::Decals, SceneProxy, true);
+							const bool bHidden = !MeshCommand.bDecal;
+							const uint32 RecordIndex = RTLODData.SBTAllocation->GetRecordIndex(ERayTracingSceneLayer::Decals, MeshCommand.GeometrySegmentIndex);
+							DirtyShaderBindingData[DecalCommandOffset] = FRayTracingShaderBindingData(&MeshCommand, RelevantPrimitive.CachedRayTracingInstance->GeometryRHI, RecordIndex, bHidden);
+							++DecalCommandOffset;
 						}
+					}
+				}
+			});
+	}
 
-						if (bParallelMeshBatchSetup)
-						{
-							if (NumPendingMeshBatches >= RayTracingParallelMeshBatchSize)
-							{
-								KickRayTracingMeshBatchTask();
-							}
-
-							if (MeshBatchTaskPage == nullptr || MeshBatchTaskPage->NumWorkItems == FRayTracingMeshBatchTaskPage::MaxWorkItems)
-							{
-								FRayTracingMeshBatchTaskPage* NextPage = InBulkAllocator.Create<FRayTracingMeshBatchTaskPage>();
-								if (MeshBatchTaskHead == nullptr)
-								{
-									MeshBatchTaskHead = NextPage;
-								}
-								if (MeshBatchTaskPage)
-								{
-									MeshBatchTaskPage->Next = NextPage;
-								}
-								MeshBatchTaskPage = NextPage;
-							}
+	void AddDynamicInstancesToRayTracingScene(
+		FRDGBuilder& GraphBuilder,
+		FScene& Scene,
+		FViewInfo& View,
+		const RayTracing::FSceneOptions& SceneOptions,
+		TConstArrayView<int32> DynamicPrimitives,
+		FRayTracingScene& RayTracingScene,
+		FRayTracingShaderBindingTable& RayTracingSBT,
+		FGlobalDynamicReadBuffer& InDynamicReadBuffer,
+		FSceneRenderingBulkObjectAllocator& InBulkAllocator)
+	{
+		TRACE_CPUPROFILER_EVENT_SCOPE(RayTracingScene_AddDynamicInstances);
 
-							FRayTracingMeshBatchWorkItem& WorkItem = MeshBatchTaskPage->WorkItems[MeshBatchTaskPage->NumWorkItems];
-							MeshBatchTaskPage->NumWorkItems++;
+		const float CurrentWorldTime = View.Family->Time.GetWorldTimeSeconds();
 
-							NumPendingMeshBatches += Instance.GetMaterials().Num();
+		const bool bParallelMeshBatchSetup = GRayTracingParallelMeshBatchSetup && FApp::ShouldUseThreadingForPerformance();
+		const bool bUsingReferenceBasedResidency = IsRayTracingUsingReferenceBasedResidency();
 
-							if (Instance.OwnsMaterials())
-							{
-								Swap(WorkItem.MeshBatchesOwned, Instance.Materials);
-							}
-							else
-							{
-								WorkItem.MeshBatchesView = Instance.MaterialsView;
-							}
+		const int64 SharedBufferGenerationID = Scene.GetRayTracingDynamicGeometryCollection()->BeginUpdate();
 
-							WorkItem.SceneProxy = SceneProxy;
-							WorkItem.RayTracingGeometry = Geometry->GetRHI();
-							WorkItem.SBTAllocation = SBTAllocation;
-						}
-						else
-						{
-							TArrayView<const FMeshBatch> InstanceMaterials = Instance.GetMaterials();
-							for (int32 SegmentIndex = 0; SegmentIndex < InstanceMaterials.Num(); SegmentIndex++)
-							{
-								const FMeshBatch& MeshBatch = InstanceMaterials[SegmentIndex];
-								FDynamicRayTracingMeshCommandContext CommandContext(View.DynamicRayTracingMeshCommandStorage, View.DirtyRayTracingShaderBindings, Geometry->GetRHI(), SegmentIndex, SBTAllocation);
-								FRayTracingMeshProcessor RayTracingMeshProcessor(&CommandContext, &Scene, &View, Scene.CachedRayTracingMeshCommandsMode);
-								RayTracingMeshProcessor.AddMeshBatch(MeshBatch, 1, SceneProxy);
-							}
-						}
-					}
+		FRayTracingMeshBatchTaskPage* MeshBatchTaskHead = nullptr;
+		FRayTracingMeshBatchTaskPage* MeshBatchTaskPage = nullptr;
+		uint32 NumPendingMeshBatches = 0;
+		const uint32 RayTracingParallelMeshBatchSize = GRayTracingParallelMeshBatchSize;
 
-					if (CVarRayTracingDynamicGeometryLastRenderTimeUpdateDistance.GetValueOnRenderThread() > 0.0f)
-					{
-						if (FVector::Distance(SceneProxy->GetActorPosition(), View.ViewMatrices.GetViewOrigin()) < CVarRayTracingDynamicGeometryLastRenderTimeUpdateDistance.GetValueOnRenderThread())
-						{
-							// Update LastRenderTime for components so that visibility based ticking (like skeletal meshes) can get updated
-							// We are only doing this for dynamic geometries now
-							SceneInfo->LastRenderTime = CurrentWorldTime;
-							SceneInfo->UpdateComponentLastRenderTime(CurrentWorldTime, /*bUpdateLastRenderTimeOnScreen=*/true);
-						}
-					}
+		auto KickRayTracingMeshBatchTask = [&InBulkAllocator, &View, &Scene, &MeshBatchTaskHead, &MeshBatchTaskPage, &NumPendingMeshBatches]()
+			{
+				if (MeshBatchTaskHead)
+				{
+					DispatchRayTracingMeshBatchTask(InBulkAllocator, Scene, View, MeshBatchTaskHead, NumPendingMeshBatches);
 				}
-			}
 
-			KickRayTracingMeshBatchTask();
-		}
+				MeshBatchTaskHead = nullptr;
+				MeshBatchTaskPage = nullptr;
+				NumPendingMeshBatches = 0;
+			};
+
+		// Need to process dynamic primitives in 3 passes to support dynamic primitives/instances in GPU Scene
+		// 1 - gather all dynamic ray tracing instances
+		// 2 - upload dynamic primitive/instance data to GPU scene
+		// 3 - process dynamic ray tracing instances (primitive IDs and instance scene data offset of dynamic primitives is now valid)
+
+		// Local temporary array of instances used for GetDynamicRayTracingInstances()
+		TArray<FRayTracingInstance> DynamicRayTracingInstances;
+		TArray<TRange<int32>> PrimitivesDynamicRayTracingInstances;
 
-		// Task to iterate over static ray tracing instances, perform auto-instancing and culling.
-		// This adds final instances to the ray tracing scene and must be done before FRayTracingScene::BuildInitializationData().
-		struct FRayTracingSceneAddStaticInstancesTask
 		{
-			UE_NONCOPYABLE(FRayTracingSceneAddStaticInstancesTask)
-
-			static ESubsequentsMode::Type GetSubsequentsMode() { return ESubsequentsMode::TrackSubsequents; }
-			TStatId                       GetStatId() const { return TStatId(); }
-			ENamedThreads::Type           GetDesiredThread() { return ENamedThreads::AnyThread; }
-
-			// Inputs
-
-			const FScene& Scene;
-			const RayTracing::FSceneOptions SceneOptions;
-			TArray<FRelevantPrimitive>& RelevantStaticPrimitives;
-			TArray<FRelevantPrimitive>& RelevantCachedStaticPrimitives;
-			TArray<FRelevantPrimitiveGatherContext>& GatherContexts;
-			const bool bIsPathTracing;
-
-			const int32& NumCachedStaticInstances;
-			const int32& NumCachedStaticDecalInstances;
-			const int32& NumCachedStaticSegments;
-			const int32& NumCachedStaticDirtyShaderBindings;
-			
-			// Outputs
-
-			FRayTracingScene& RayTracingScene; // New instances are added into FRayTracingScene::Instances and FRayTracingScene::Allocator is used for temporary data
-			TArray<FRayTracingShaderBindingData>& DirtyShaderBindingData; // New elements are added here by this task
-
-			FRayTracingSceneAddStaticInstancesTask(
-				FGatherInstancesTaskData& TaskData, const bool bInIsPathTracing,
-				FRayTracingScene& InRayTracingScene, TArray<FRayTracingShaderBindingData>& InDirtyShaderBindingData)
-				: Scene(TaskData.Scene)
-				, SceneOptions(TaskData.SceneOptions)
-				, RelevantStaticPrimitives(TaskData.StaticPrimitives)
-				, RelevantCachedStaticPrimitives(TaskData.CachedStaticPrimitives)
-				, GatherContexts(TaskData.GatherContexts)
-				, bIsPathTracing(bInIsPathTracing)
-				, NumCachedStaticInstances(TaskData.NumCachedStaticInstances)
-				, NumCachedStaticDecalInstances(TaskData.NumCachedStaticDecalInstances)
-				, NumCachedStaticSegments(TaskData.NumCachedStaticSegments)
-				, NumCachedStaticDirtyShaderBindings(TaskData.NumCachedStaticDirtyShaderBindings)
-				, RayTracingScene(InRayTracingScene)
-				, DirtyShaderBindingData(InDirtyShaderBindingData)
+			FRayTracingMaterialGatheringContext MaterialGatheringContext
+			(
+				&Scene,
+				&View,
+				*View.Family,
+				GraphBuilder,
+				*View.RayTracingMeshResourceCollector,
+				View.RayTracingDynamicPrimitiveCollector,
+				InDynamicReadBuffer
+			);
+
+			for (int32 PrimitiveIndex : DynamicPrimitives)
 			{
-				DirtyShaderBindingData.Reserve(RelevantStaticPrimitives.Num() + RelevantCachedStaticPrimitives.Num());
-			}
+				check(MaterialGatheringContext.DynamicRayTracingGeometriesToUpdate.IsEmpty());
 
-			// TODO: Consider moving auto instance batching logic into FRayTracingScene
+				FPrimitiveSceneProxy* SceneProxy = Scene.PrimitiveSceneProxies[PrimitiveIndex];
+				FPrimitiveSceneInfo* SceneInfo = Scene.Primitives[PrimitiveIndex];
+				const FPersistentPrimitiveIndex PersistentPrimitiveIndex = SceneInfo->GetPersistentIndex();
 
-			struct FAutoInstanceBatch
-			{
-				FRayTracingScene::FInstanceHandle InstanceHandle = FRayTracingScene::INVALID_INSTANCE_HANDLE;
-				FRayTracingScene::FInstanceHandle DecalInstanceHandle = FRayTracingScene::INVALID_INSTANCE_HANDLE;
+				MaterialGatheringContext.SetPrimitive(SceneProxy);
 
-				// Copies the next InstanceSceneDataOffset and user data into the current batch, returns true if arrays were re-allocated.
-				bool Add(FRayTracingScene& InRayTracingScene, uint32 InInstanceSceneDataOffset)
+				int32 BaseRayTracingInstance = DynamicRayTracingInstances.Num();
+
+				if (SceneOptions.bTranslucentGeometry || SceneProxy->IsOpaqueOrMasked())
 				{
-					// Adhoc TArray-like resize behavior, in lieu of support for using a custom FMemStackBase in TArray.
-					// Idea for future: if batch becomes large enough, we could actually split it into multiple instances to avoid memory waste.
+					SceneProxy->GetDynamicRayTracingInstances(MaterialGatheringContext, DynamicRayTracingInstances);
+				}
 
-					const bool bNeedReallocation = Cursor == InstanceSceneDataOffsets.Num();
+				for (const FRayTracingDynamicGeometryUpdateParams& DynamicRayTracingGeometryUpdate : MaterialGatheringContext.DynamicRayTracingGeometriesToUpdate)
+				{
+					Scene.GetRayTracingDynamicGeometryCollection()->AddDynamicMeshBatchForGeometryUpdate(
+						GraphBuilder.RHICmdList,
+						&Scene,
+						&View,
+						SceneProxy,
+						DynamicRayTracingGeometryUpdate,
+						PersistentPrimitiveIndex.Index
+					);
+				}
 
-					if (bNeedReallocation)
-					{
-						int32 PrevCount = InstanceSceneDataOffsets.Num();
-						int32 NextCount = FMath::Max(PrevCount * 2, 1);
+				if (bUsingReferenceBasedResidency)
+				{
+					((FRayTracingGeometryManager*)GRayTracingGeometryManager)->AddReferencedGeometryGroups(MaterialGatheringContext.GetReferencedGeometryGroups());
+				}
 
-						TArrayView<uint32> NewInstanceSceneDataOffsets = InRayTracingScene.Allocate<uint32>(NextCount);
-						if (PrevCount)
-						{
-							FMemory::Memcpy(NewInstanceSceneDataOffsets.GetData(), InstanceSceneDataOffsets.GetData(), InstanceSceneDataOffsets.GetTypeSize() * InstanceSceneDataOffsets.Num());
-						}
-						InstanceSceneDataOffsets = NewInstanceSceneDataOffsets;
-					}
+				MaterialGatheringContext.Reset();
 
-					InstanceSceneDataOffsets[Cursor] = InInstanceSceneDataOffset;
+				PrimitivesDynamicRayTracingInstances.Add(TRange<int32>(BaseRayTracingInstance, DynamicRayTracingInstances.Num()));
+			}
 
-					++Cursor;
+			// FRayTracingMaterialGatheringContext destructor handles committing dynamic mesh batches to GPU Scene
+		}
 
-					return bNeedReallocation;
-				}
+		Scene.GPUScene.UploadDynamicPrimitiveShaderDataForView(GraphBuilder, View, true);
 
-				bool IsValid() const
-				{
-					return InstanceSceneDataOffsets.Num() != 0;
-				}
+		const int32 ViewDynamicPrimitiveId = View.RayTracingDynamicPrimitiveCollector.GetPrimitiveIdRange().GetLowerBoundValue();
+		const int32 ViewInstanceSceneDataOffset = View.RayTracingDynamicPrimitiveCollector.GetInstanceSceneDataOffset();
 
-				TArrayView<uint32> InstanceSceneDataOffsets;
-				uint32 Cursor = 0;
-			};
+		for (int32 Index = 0; Index < DynamicPrimitives.Num(); ++Index)
+		{
+			const int32 PrimitiveIndex = DynamicPrimitives[Index];
+			FPrimitiveSceneProxy* SceneProxy = Scene.PrimitiveSceneProxies[PrimitiveIndex];
+			FPrimitiveSceneInfo* SceneInfo = Scene.Primitives[PrimitiveIndex];
+			const FPersistentPrimitiveIndex PersistentPrimitiveIndex = SceneInfo->GetPersistentIndex();
 
-			void DoTask(ENamedThreads::Type CurrentThread, const FGraphEventRef& MyCompletionGraphEvent)
-			{
-				FTaskTagScope TaskTagScope(ETaskTag::EParallelRenderingThread);
+			TArrayView<FRayTracingInstance> TempRayTracingInstances = MakeArrayView(
+				DynamicRayTracingInstances.GetData() + PrimitivesDynamicRayTracingInstances[Index].GetLowerBound().GetValue(),
+				PrimitivesDynamicRayTracingInstances[Index].Size<int32>());
 
-				TRACE_CPUPROFILER_EVENT_SCOPE(RayTracingSceneStaticInstanceTask);
-				
+			if (TempRayTracingInstances.Num() > 0)
+			{
+				for (FRayTracingInstance& Instance : TempRayTracingInstances)
 				{
-					TRACE_CPUPROFILER_EVENT_SCOPE(RayTracingScene_AddStaticInstances);
+					const FRayTracingGeometry* Geometry = Instance.Geometry;
+
+					if (!ensureMsgf(Geometry->DynamicGeometrySharedBufferGenerationID == FRayTracingGeometry::NonSharedVertexBuffers
+						|| Geometry->DynamicGeometrySharedBufferGenerationID == SharedBufferGenerationID,
+						TEXT("GenerationID %lld, but expected to be %lld or %lld. Geometry debug name: '%s'. ")
+						TEXT("When shared vertex buffers are used, the contents is expected to be written every frame. ")
+						TEXT("Possibly AddDynamicMeshBatchForGeometryUpdate() was not called for this geometry."),
+						Geometry->DynamicGeometrySharedBufferGenerationID, SharedBufferGenerationID, FRayTracingGeometry::NonSharedVertexBuffers,
+						*Geometry->Initializer.DebugName.ToString()))
+					{
+						continue;
+					}
 
-					const bool bAutoInstance = CVarRayTracingAutoInstance.GetValueOnRenderThread() != 0;
+					((FRayTracingGeometryManager*)GRayTracingGeometryManager)->AddReferencedGeometry(Geometry);
+
+					if (Geometry->IsEvicted())
+					{
+						continue;
+					}
 
-					// Instance batches by FRelevantPrimitive::InstancingKey()
-					Experimental::TSherwoodMap<uint64, FAutoInstanceBatch> InstanceBatches;
+					// If geometry still has pending build request then add to list which requires a force build
+					if (Geometry->HasPendingBuildRequest())
+					{
+						RayTracingScene.GeometriesToBuild.Add(Geometry);
+					}
 
-					// scan relevant primitives computing hash data to look for duplicate instances
-					for (const FRelevantPrimitive& RelevantPrimitive : RelevantStaticPrimitives)
+					// Validate the material/segment counts
+					if (!ensureMsgf(Instance.GetMaterials().Num() == Geometry->Initializer.Segments.Num() ||
+						(Geometry->Initializer.Segments.Num() == 0 && Instance.GetMaterials().Num() == 1),
+						TEXT("Ray tracing material assignment validation failed for geometry '%s'. "
+							"Instance.GetMaterials().Num() = %d, Geometry->Initializer.Segments.Num() = %d."),
+						*Geometry->Initializer.DebugName.ToString(), Instance.GetMaterials().Num(),
+						Geometry->Initializer.Segments.Num()))
 					{
-						const int32 PrimitiveIndex = RelevantPrimitive.PrimitiveIndex;
-						FPrimitiveSceneInfo* SceneInfo = Scene.Primitives[PrimitiveIndex];
-						FPrimitiveSceneProxy* SceneProxy = Scene.PrimitiveSceneProxies[PrimitiveIndex];
-						ERayTracingPrimitiveFlags Flags = Scene.PrimitiveRayTracingFlags[PrimitiveIndex];
-						const FPersistentPrimitiveIndex PersistentPrimitiveIndex = RelevantPrimitive.PersistentPrimitiveIndex;
+						continue;
+					}
 
-						check(!EnumHasAnyFlags(Flags, ERayTracingPrimitiveFlags::CacheInstances));
+					if (Instance.bInstanceMaskAndFlagsDirty || SceneInfo->bCachedRayTracingInstanceMaskAndFlagsDirty)
+					{
+						// Build InstanceMaskAndFlags since the data in SceneInfo is not up to date
 
-						const int8 LODIndex = RelevantPrimitive.LODIndex;
+						FRayTracingMaskAndFlags InstanceMaskAndFlags;
 
-						if (LODIndex < 0)
+						if (Instance.GetMaterials().IsEmpty())
 						{
-							// TODO: Filter these primitives earlier
-							continue; 
+							// If the material list is empty, explicitly set the mask to 0 so it will not be added in the raytracing scene
+							InstanceMaskAndFlags.Mask = 0;
+						}
+						else
+						{
+							InstanceMaskAndFlags = BuildRayTracingInstanceMaskAndFlags(Instance, *SceneProxy);
 						}
 
-						const bool bNeedMainInstance = !RelevantPrimitive.CachedMeshCommandFlags.bAllSegmentsDecal;
-
-						// if primitive has mixed decal and non-decal segments we need to have two ray tracing instances
-						// one containing non-decal segments and the other with decal segments
-						// masking of segments is done using "hidden" hitgroups
-						// TODO: Debug Visualization to highlight primitives using this?
-						const bool bNeedDecalInstance = RelevantPrimitive.CachedMeshCommandFlags.bAnySegmentsDecal && !ShouldExcludeDecals();
+						SceneInfo->CachedRayTracingInstance.Mask = InstanceMaskAndFlags.Mask; // When no cached command is found, InstanceMask == 0 and the instance is effectively filtered out
 
-						if (ShouldExcludeDecals() && RelevantPrimitive.CachedMeshCommandFlags.bAllSegmentsDecal)
+						if (InstanceMaskAndFlags.bForceOpaque)
 						{
-							continue;
+							SceneInfo->CachedRayTracingInstance.Flags |= ERayTracingInstanceFlags::ForceOpaque;
 						}
 
-						if (!SceneOptions.bTranslucentGeometry && RelevantPrimitive.CachedMeshCommandFlags.bAllSegmentsTranslucent)
+						if (InstanceMaskAndFlags.bDoubleSided)
 						{
-							continue;
+							SceneInfo->CachedRayTracingInstance.Flags |= ERayTracingInstanceFlags::TriangleCullDisable;
 						}
 
-						if (GRayTracingExcludeSky && RelevantPrimitive.CachedMeshCommandFlags.bIsSky)
+						if (InstanceMaskAndFlags.bReverseCulling)
 						{
-							continue;
+							SceneInfo->CachedRayTracingInstance.Flags |= ERayTracingInstanceFlags::TriangleCullReverse;
 						}
 
-						// location if this is a new entry
-						const uint64 InstanceKey = RelevantPrimitive.InstancingKey();
+						SceneInfo->bCachedRayTracingInstanceAnySegmentsDecal = InstanceMaskAndFlags.bAnySegmentsDecal;
+						SceneInfo->bCachedRayTracingInstanceAllSegmentsDecal = InstanceMaskAndFlags.bAllSegmentsDecal;
 
-						FAutoInstanceBatch DummyInstanceBatch = { };
-						FAutoInstanceBatch& InstanceBatch = bAutoInstance ? InstanceBatches.FindOrAdd(InstanceKey, DummyInstanceBatch) : DummyInstanceBatch;
+						SceneInfo->bCachedRayTracingInstanceMaskAndFlagsDirty = false;
+					}
 
-						if (InstanceBatch.IsValid())
-						{
-							// Reusing a previous entry, just append to the instance list.
+					const bool bNeedMainInstance = !SceneInfo->bCachedRayTracingInstanceAllSegmentsDecal;
 
-							bool bReallocated = InstanceBatch.Add(RayTracingScene, SceneInfo->GetInstanceSceneDataOffset());
+					// if primitive has mixed decal and non-decal segments we need to have two ray tracing instances
+					// one containing non-decal segments and the other with decal segments
+					// masking of segments is done using "hidden" hitgroups
+					// TODO: Debug Visualization to highlight primitives using this?
+					const bool bNeedDecalInstance = SceneInfo->bCachedRayTracingInstanceAnySegmentsDecal && !ShouldExcludeDecals();
 
-							if(InstanceBatch.InstanceHandle.IsValid())
-							{
-								FRayTracingGeometryInstance& RayTracingInstance = RayTracingScene.GetInstance(InstanceBatch.InstanceHandle);
-								++RayTracingInstance.NumTransforms;
-								check(RayTracingInstance.NumTransforms == InstanceBatch.Cursor); // sanity check
+					if (ShouldExcludeDecals() && SceneInfo->bCachedRayTracingInstanceAllSegmentsDecal)
+					{
+						continue;
+					}
 
-								if (bReallocated)
-								{
-									RayTracingInstance.InstanceSceneDataOffsets = InstanceBatch.InstanceSceneDataOffsets;
-									RayTracingInstance.UserData = InstanceBatch.InstanceSceneDataOffsets;
-								}
-							}
+					int32 PrimitiveId = PersistentPrimitiveIndex.Index;
+					int32 InstanceSceneDataOffset = SceneInfo->GetInstanceSceneDataOffset();
 
-							if (InstanceBatch.DecalInstanceHandle.IsValid())
-							{
-								FRayTracingGeometryInstance& RayTracingInstance = RayTracingScene.GetInstance(InstanceBatch.DecalInstanceHandle);
-								++RayTracingInstance.NumTransforms;
-								check(RayTracingInstance.NumTransforms == InstanceBatch.Cursor); // sanity check
+					if (Instance.Materials.Num() > 0 && Instance.Materials[0].Elements.Num() > 0 && Instance.Materials[0].Elements[0].DynamicPrimitiveData != nullptr)
+					{
+						check(Instance.NumTransforms == Instance.Materials[0].Elements[0].NumInstances);
+						PrimitiveId = ViewDynamicPrimitiveId + Instance.Materials[0].Elements[0].DynamicPrimitiveIndex;
+						InstanceSceneDataOffset = ViewInstanceSceneDataOffset + Instance.Materials[0].Elements[0].DynamicPrimitiveInstanceSceneDataOffset;
+					}
 
-								if (bReallocated)
-								{
-									RayTracingInstance.InstanceSceneDataOffsets = InstanceBatch.InstanceSceneDataOffsets;
-									RayTracingInstance.UserData = InstanceBatch.InstanceSceneDataOffsets;
-								}
-							}
+					FRayTracingGeometryInstance RayTracingInstance;
+					RayTracingInstance.GeometryRHI = Geometry->GetRHI();
+					checkf(RayTracingInstance.GeometryRHI, TEXT("Ray tracing instance must have a valid geometry."));
+					RayTracingInstance.DefaultUserData = InstanceSceneDataOffset;
+					RayTracingInstance.bIncrementUserDataPerInstance = true;
+					RayTracingInstance.bApplyLocalBoundsTransform = Instance.bApplyLocalBoundsTransform;
+					RayTracingInstance.Mask = SceneInfo->CachedRayTracingInstance.Mask;
+					RayTracingInstance.Flags = SceneInfo->CachedRayTracingInstance.Flags;
+					AddDebugRayTracingInstanceFlags(RayTracingInstance.Flags);
+
+					if (!Instance.GetPrimitiveInstanceIndices().IsEmpty())
+					{
+						TConstArrayView<uint32> PrimitiveInstanceIndices = Instance.GetPrimitiveInstanceIndices();
+
+						// Convert from instance indices to InstanceSceneDataOffsets
+						TArrayView<uint32> InstanceSceneDataOffsets = RayTracingScene.Allocate<uint32>(PrimitiveInstanceIndices.Num());
+						for (int32 InstanceIndex = 0; InstanceIndex < PrimitiveInstanceIndices.Num(); ++InstanceIndex)
+						{
+							InstanceSceneDataOffsets[InstanceIndex] = SceneInfo->GetInstanceSceneDataOffset() + PrimitiveInstanceIndices[InstanceIndex];
+						}
+
+						RayTracingInstance.InstanceSceneDataOffsets = InstanceSceneDataOffsets;
+						RayTracingInstance.UserData = InstanceSceneDataOffsets;
+						RayTracingInstance.NumTransforms = PrimitiveInstanceIndices.Num();
+					}
+					else if (!Instance.GetTransforms().IsEmpty())
+					{
+						TConstArrayView<FMatrix> TransformsView;
+						if (Instance.OwnsTransforms())
+						{
+							// Slow path: copy transforms to the owned storage
+							checkf(Instance.InstanceTransformsView.Num() == 0, TEXT("InstanceTransformsView is expected to be empty if using InstanceTransforms"));
+							TArrayView<FMatrix> SceneOwnedTransforms = RayTracingScene.Allocate<FMatrix>(Instance.InstanceTransforms.Num());
+							FMemory::Memcpy(SceneOwnedTransforms.GetData(), Instance.InstanceTransforms.GetData(), Instance.InstanceTransforms.Num() * sizeof(RayTracingInstance.Transforms[0]));
+							static_assert(std::is_same_v<decltype(SceneOwnedTransforms[0]), decltype(Instance.InstanceTransforms[0])>, "Unexpected transform type");
+
+							TransformsView = SceneOwnedTransforms;
 						}
 						else
 						{
-							// Starting new instance batch
+							// Fast path: just reference persistently-allocated transforms and avoid a copy
+							checkf(Instance.InstanceTransforms.Num() == 0, TEXT("InstanceTransforms is expected to be empty if using InstanceTransformsView"));
+							TransformsView = Instance.InstanceTransformsView;
+						}
 
-							InstanceBatch.Add(RayTracingScene, SceneInfo->GetInstanceSceneDataOffset());
+						RayTracingInstance.NumTransforms = TransformsView.Num();
+						RayTracingInstance.Transforms = TransformsView;
+					}
+					else
+					{
+						// If array of transforms was not provided, get the instance transforms from GPU Scene
+						RayTracingInstance.NumTransforms = Instance.NumTransforms;
+						RayTracingInstance.BaseInstanceSceneDataOffset = InstanceSceneDataOffset;
+					}
 
-							FRayTracingGeometryInstance RayTracingInstance;
-							RayTracingInstance.GeometryRHI = RelevantPrimitive.RayTracingGeometry->GetRHI();
-							checkf(RayTracingInstance.GeometryRHI, TEXT("Ray tracing instance must have a valid geometry."));
-							RayTracingInstance.InstanceSceneDataOffsets = InstanceBatch.InstanceSceneDataOffsets;
-							RayTracingInstance.UserData = InstanceBatch.InstanceSceneDataOffsets;
-							RayTracingInstance.NumTransforms = 1;
+					ERayTracingSceneLayerMask ActiveLayers = ERayTracingSceneLayerMask::None;
+					if (bNeedMainInstance)
+					{
+						EnumAddFlags(ActiveLayers, ERayTracingSceneLayerMask::Base);
+					}
+					if (bNeedDecalInstance)
+					{
+						EnumAddFlags(ActiveLayers, ERayTracingSceneLayerMask::Decals);
+					}
 
-							RayTracingInstance.Mask = RelevantPrimitive.CachedMeshCommandFlags.InstanceMask; // When no cached command is found, InstanceMask == 0 and the instance is effectively filtered out
+					FRayTracingSBTAllocation* SBTAllocation = RayTracingSBT.AllocateDynamicRange(ActiveLayers, Instance.GetMaterials().Num());
+					if (bNeedMainInstance)
+					{
+						RayTracingScene.NumSegments += Instance.GetMaterials().Num();
 
-							// Run AHS for alpha masked and meshes with only some sections casting shadows, which require per mesh section filtering in AHS
-							if (RelevantPrimitive.CachedMeshCommandFlags.bAllSegmentsOpaque && (RelevantPrimitive.CachedMeshCommandFlags.bAllSegmentsCastShadow || !RelevantPrimitive.CachedMeshCommandFlags.bAnySegmentsCastShadow))
-							{
-								RayTracingInstance.Flags |= ERayTracingInstanceFlags::ForceOpaque;
-							}
-							if (RelevantPrimitive.CachedMeshCommandFlags.bTwoSided)
-							{
-								RayTracingInstance.Flags |= ERayTracingInstanceFlags::TriangleCullDisable;
-							}
-							if (RelevantPrimitive.CachedMeshCommandFlags.bAllSegmentsReverseCulling)
-							{
-								RayTracingInstance.Flags |= ERayTracingInstanceFlags::TriangleCullReverse;
-							}
-							AddDebugRayTracingInstanceFlags(RayTracingInstance.Flags);
+						RayTracingInstance.InstanceContributionToHitGroupIndex = SBTAllocation->GetInstanceContributionToHitGroupIndex(ERayTracingSceneLayer::Base);
+						RayTracingScene.AddInstance(RayTracingInstance, ERayTracingSceneLayer::Base, SceneProxy, true);
+					}
 
-							InstanceBatch.InstanceHandle = FRayTracingScene::INVALID_INSTANCE_HANDLE;
+					if (bNeedDecalInstance)
+					{
+						RayTracingScene.NumSegments += Instance.GetMaterials().Num();
 
-							FPrimitiveSceneInfo::FRayTracingLODData& RTLODData = SceneInfo->RayTracingLODData[LODIndex];
-							ensure(RTLODData.SBTAllocation);
-							if (RTLODData.SBTAllocation == nullptr)
-							{
-								continue;
-							}
-							
-							uint32 SegmentCount = RelevantPrimitive.RayTracingGeometry->Initializer.Segments.Num();
-							if (bNeedMainInstance)
-							{
-								RayTracingScene.NumSegments += SegmentCount;
-								RayTracingInstance.InstanceContributionToHitGroupIndex = RTLODData.SBTAllocation->GetInstanceContributionToHitGroupIndex(ERayTracingSceneLayer::Base);
-								InstanceBatch.InstanceHandle = RayTracingScene.AddInstance(RayTracingInstance, ERayTracingSceneLayer::Base, SceneProxy, false);
-							}
+						FRayTracingGeometryInstance DecalRayTracingInstance = RayTracingInstance;
+						DecalRayTracingInstance.InstanceContributionToHitGroupIndex = SBTAllocation->GetInstanceContributionToHitGroupIndex(ERayTracingSceneLayer::Decals);
+						RayTracingScene.AddInstance(MoveTemp(DecalRayTracingInstance), ERayTracingSceneLayer::Decals, SceneProxy, true);
+					}
 
-							InstanceBatch.DecalInstanceHandle = FRayTracingScene::INVALID_INSTANCE_HANDLE;
-							if (bNeedDecalInstance)
+					if (bParallelMeshBatchSetup)
+					{
+						if (NumPendingMeshBatches >= RayTracingParallelMeshBatchSize)
+						{
+							KickRayTracingMeshBatchTask();
+						}
+
+						if (MeshBatchTaskPage == nullptr || MeshBatchTaskPage->NumWorkItems == FRayTracingMeshBatchTaskPage::MaxWorkItems)
+						{
+							FRayTracingMeshBatchTaskPage* NextPage = InBulkAllocator.Create<FRayTracingMeshBatchTaskPage>();
+							if (MeshBatchTaskHead == nullptr)
 							{
-								RayTracingScene.NumSegments += SegmentCount;
-								FRayTracingGeometryInstance DecalRayTracingInstance = RayTracingInstance;
-								DecalRayTracingInstance.InstanceContributionToHitGroupIndex = RTLODData.SBTAllocation->GetInstanceContributionToHitGroupIndex(ERayTracingSceneLayer::Decals);
-								InstanceBatch.DecalInstanceHandle = RayTracingScene.AddInstance(MoveTemp(DecalRayTracingInstance), ERayTracingSceneLayer::Decals, SceneProxy, false);
+								MeshBatchTaskHead = NextPage;
 							}
-							
-							for (int32 CommandIndex : RTLODData.CachedMeshCommandIndices)
+							if (MeshBatchTaskPage)
 							{
-								if (CommandIndex >= 0)
-								{
-									const FRayTracingMeshCommand& MeshCommand = Scene.CachedRayTracingMeshCommands[CommandIndex];
-
-									if (bNeedMainInstance)
-									{
-										const bool bHidden = MeshCommand.bDecal;
-										const uint32 RecordIndex = RTLODData.SBTAllocation->GetRecordIndex(ERayTracingSceneLayer::Base, MeshCommand.GeometrySegmentIndex);
-										FRayTracingShaderBindingData RTShaderBindingData(&MeshCommand, RelevantPrimitive.RayTracingGeometry->GetRHI(), RecordIndex, bHidden);
-										DirtyShaderBindingData.Add(RTShaderBindingData);
-									}
-									if (bNeedDecalInstance)
-									{
-										const bool bHidden = !MeshCommand.bDecal;
-										const uint32 RecordIndex = RTLODData.SBTAllocation->GetRecordIndex(ERayTracingSceneLayer::Decals, MeshCommand.GeometrySegmentIndex);
-										FRayTracingShaderBindingData RTShaderBindingData(&MeshCommand, RelevantPrimitive.RayTracingGeometry->GetRHI(), RecordIndex, bHidden);
-										DirtyShaderBindingData.Add(RTShaderBindingData);
-									}
-								}
-								else
-								{
-									// CommandIndex == -1 indicates that the mesh batch has been filtered by FRayTracingMeshProcessor (like the shadow depth pass batch)
-									// Do nothing in this case
-								}
+								MeshBatchTaskPage->Next = NextPage;
 							}
+							MeshBatchTaskPage = NextPage;
+						}
+
+						FRayTracingMeshBatchWorkItem& WorkItem = MeshBatchTaskPage->WorkItems[MeshBatchTaskPage->NumWorkItems];
+						MeshBatchTaskPage->NumWorkItems++;
+
+						NumPendingMeshBatches += Instance.GetMaterials().Num();
+
+						if (Instance.OwnsMaterials())
+						{
+							Swap(WorkItem.MeshBatchesOwned, Instance.Materials);
+						}
+						else
+						{
+							WorkItem.MeshBatchesView = Instance.MaterialsView;
+						}
+
+						WorkItem.SceneProxy = SceneProxy;
+						WorkItem.RayTracingGeometry = Geometry->GetRHI();
+						WorkItem.SBTAllocation = SBTAllocation;
+					}
+					else
+					{
+						TArrayView<const FMeshBatch> InstanceMaterials = Instance.GetMaterials();
+						for (int32 SegmentIndex = 0; SegmentIndex < InstanceMaterials.Num(); SegmentIndex++)
+						{
+							const FMeshBatch& MeshBatch = InstanceMaterials[SegmentIndex];
+							FDynamicRayTracingMeshCommandContext CommandContext(View.DynamicRayTracingMeshCommandStorage, View.DirtyRayTracingShaderBindings, Geometry->GetRHI(), SegmentIndex, SBTAllocation);
+							FRayTracingMeshProcessor RayTracingMeshProcessor(&CommandContext, &Scene, &View, Scene.CachedRayTracingMeshCommandsMode);
+							RayTracingMeshProcessor.AddMeshBatch(MeshBatch, 1, SceneProxy);
 						}
 					}
 				}
 
+				if (CVarRayTracingDynamicGeometryLastRenderTimeUpdateDistance.GetValueOnRenderThread() > 0.0f)
 				{
-					TRACE_CPUPROFILER_EVENT_SCOPE(RayTracingScene_AddCachedStaticInstances);
-					
-					const FRayTracingScene::FInstanceRange CachedStaticInstanceRange = RayTracingScene.AllocateInstanceRangeUninitialized(NumCachedStaticInstances, ERayTracingSceneLayer::Base);
-					const FRayTracingScene::FInstanceRange CachedStaticDecalInstanceRange = RayTracingScene.AllocateInstanceRangeUninitialized(NumCachedStaticDecalInstances, ERayTracingSceneLayer::Decals);
-					const uint32 BaseCachedDirtyShaderBindingsIndex = DirtyShaderBindingData.AddUninitialized(NumCachedStaticDirtyShaderBindings);
-					RayTracingScene.NumSegments += NumCachedStaticSegments;
-
-					const int32 MinBatchSize = 128;
-					ParallelFor(
-						TEXT("RayTracingScene_AddCachedStaticInstances_ParallelFor"),
-						RelevantCachedStaticPrimitives.Num(),
-						MinBatchSize,
-						[this, CachedStaticInstanceRange, CachedStaticDecalInstanceRange, BaseCachedDirtyShaderBindingsIndex](int32 Index)
+					if (FVector::Distance(SceneProxy->GetActorPosition(), View.ViewMatrices.GetViewOrigin()) < CVarRayTracingDynamicGeometryLastRenderTimeUpdateDistance.GetValueOnRenderThread())
 					{
-						const FRelevantPrimitive& RelevantPrimitive = RelevantCachedStaticPrimitives[Index];
-						const int32 PrimitiveIndex = RelevantPrimitive.PrimitiveIndex;
-						FPrimitiveSceneInfo* SceneInfo = Scene.Primitives[PrimitiveIndex];
-						FPrimitiveSceneProxy* SceneProxy = Scene.PrimitiveSceneProxies[PrimitiveIndex];
-						ERayTracingPrimitiveFlags Flags = Scene.PrimitiveRayTracingFlags[PrimitiveIndex];
-						const FPersistentPrimitiveIndex PersistentPrimitiveIndex = RelevantPrimitive.PersistentPrimitiveIndex;
+						// Update LastRenderTime for components so that visibility based ticking (like skeletal meshes) can get updated
+						// We are only doing this for dynamic geometries now
+						SceneInfo->LastRenderTime = CurrentWorldTime;
+						SceneInfo->UpdateComponentLastRenderTime(CurrentWorldTime, /*bUpdateLastRenderTimeOnScreen=*/true);
+					}
+				}
+			}
+		}
 
-						check(EnumHasAnyFlags(Flags, ERayTracingPrimitiveFlags::CacheInstances));
+		KickRayTracingMeshBatchTask();
+	}
 
-						const bool bUsingNaniteRayTracing = (Nanite::GetRayTracingMode() != Nanite::ERayTracingMode::Fallback) && SceneProxy->IsNaniteMesh();
+	void BeginGatherInstances(FGatherInstancesTaskData& TaskData, UE::Tasks::FTask FrustumCullTask)
+	{
+		const bool bUsingReferenceBasedResidency = IsRayTracingUsingReferenceBasedResidency();
 
-						if (bUsingNaniteRayTracing)
-						{
-							check(RelevantPrimitive.CachedRayTracingInstance->GeometryRHI != nullptr);
-						}
+		UE::Tasks::FTask CacheRayTracingPrimitivesTask = TaskData.Scene.GetCacheRayTracingPrimitivesTask();
 
-						const bool bNeedMainInstance = !RelevantPrimitive.CachedMeshCommandFlags.bAllSegmentsDecal;
+		TaskData.GatherRelevantPrimitivesTask = UE::Tasks::Launch(UE_SOURCE_LOCATION, [&TaskData, bUsingReferenceBasedResidency]
+			{
+				FTaskTagScope Scope(ETaskTag::EParallelRenderingThread);
+				GatherRelevantPrimitives(TaskData, bUsingReferenceBasedResidency);
+			}, UE::Tasks::Prerequisites(CacheRayTracingPrimitivesTask, FrustumCullTask), UE::Tasks::ETaskPriority::High);
+	}
 
-						// if primitive has mixed decal and non-decal segments we need to have two ray tracing instances
-						// one containing non-decal segments and the other with decal segments
-						// masking of segments is done using "hidden" hitgroups
-						// TODO: Debug Visualization to highlight primitives using this?
-						const bool bNeedDecalInstance = RelevantPrimitive.CachedMeshCommandFlags.bAnySegmentsDecal && !RelevantPrimitive.CachedMeshCommandFlags.bAllSegmentsDecal && !ShouldExcludeDecals();
+	bool FinishGatherInstances(
+		FRDGBuilder& GraphBuilder,
+		FGatherInstancesTaskData& TaskData,
+		FRayTracingScene& RayTracingScene,
+		FRayTracingShaderBindingTable& RayTracingSBT,
+		FGlobalDynamicReadBuffer& InDynamicReadBuffer,
+		FSceneRenderingBulkObjectAllocator& InBulkAllocator)
+	{
+		TRACE_CPUPROFILER_EVENT_SCOPE(RayTracing_FinishGatherInstances);
+		SCOPE_CYCLE_COUNTER(STAT_RayTracing_FinishGatherInstances);
 
-						check(!ShouldExcludeDecals() || !RelevantPrimitive.CachedMeshCommandFlags.bAllSegmentsDecal);
-						check(RelevantPrimitive.CachedRayTracingInstance);
-						
-						// Cached instances always use single LOD 0 - see above
-						const int8 LODIndex = 0;
-						FPrimitiveSceneInfo::FRayTracingLODData& RTLODData = SceneInfo->RayTracingLODData[LODIndex];
+		TaskData.GatherRelevantPrimitivesTask.Wait();
 
-						if (bNeedMainInstance)
-						{
-							const int32 InstanceIndexInRange = GatherContexts[RelevantPrimitive.ContextIndex].InstanceOffset + RelevantPrimitive.RelativeInstanceOffset;
+		FScene& Scene = TaskData.Scene;
+		FViewInfo& View = TaskData.View;
 
-							FRayTracingGeometryInstance RayTracingInstance = *RelevantPrimitive.CachedRayTracingInstance;
-							RayTracingInstance.InstanceContributionToHitGroupIndex = RTLODData.SBTAllocation->GetInstanceContributionToHitGroupIndex(ERayTracingSceneLayer::Base);
-							AddDebugRayTracingInstanceFlags(RayTracingInstance.Flags);
+		// Prepare ray tracing scene instance list
+		checkf(TaskData.bValid, TEXT("Ray tracing relevant primitive list is expected to have been created before GatherRayTracingWorldInstancesForView() is called."));
 
-							RayTracingScene.SetInstance(CachedStaticInstanceRange, InstanceIndexInRange, MoveTemp(RayTracingInstance), SceneProxy, false);
-						}
+		// Check that any invalidated cached uniform expressions have been updated on the rendering thread.
+		// Normally this work is done through FMaterialRenderProxy::UpdateUniformExpressionCacheIfNeeded,
+		// however ray tracing material processing (FMaterialShader::GetShaderBindings, which accesses UniformExpressionCache)
+		// is done on task threads, therefore all work must be done here up-front as UpdateUniformExpressionCacheIfNeeded is not free-threaded.
+		check(!FMaterialRenderProxy::HasDeferredUniformExpressionCacheRequests());
 
-						if (bNeedDecalInstance)
-						{
-							const int32 DecalInstanceIndexInRange = GatherContexts[RelevantPrimitive.ContextIndex].DecalInstanceOffset + RelevantPrimitive.RelativeDecalInstanceOffset;
+		FGPUScenePrimitiveCollector DummyDynamicPrimitiveCollector;
 
-							FRayTracingGeometryInstance DecalRayTracingInstance = *RelevantPrimitive.CachedRayTracingInstance;
-							DecalRayTracingInstance.InstanceContributionToHitGroupIndex = RTLODData.SBTAllocation->GetInstanceContributionToHitGroupIndex(ERayTracingSceneLayer::Decals);
-							AddDebugRayTracingInstanceFlags(DecalRayTracingInstance.Flags);
+		View.DynamicRayTracingMeshCommandStorage.Reserve(Scene.Primitives.Num());
+		
+		View.RayTracingMeshResourceCollector = MakeUnique<FRayTracingMeshResourceCollector>(Scene.GetFeatureLevel(), InBulkAllocator);
 
-							RayTracingScene.SetInstance(CachedStaticDecalInstanceRange, DecalInstanceIndexInRange, MoveTemp(DecalRayTracingInstance), SceneProxy, false);
-						}
+		RayTracingSBT.ResetDynamicAllocationData();
 
-						if (RelevantPrimitive.RelativeDirtyShaderBindingsOffset != INDEX_NONE)
-						{
-							const int32 DirtyShadingBindingOffset = BaseCachedDirtyShaderBindingsIndex + GatherContexts[RelevantPrimitive.ContextIndex].DirtyShaderBindingsOffset + RelevantPrimitive.RelativeDirtyShaderBindingsOffset;
+		// Consume output of the relevant primitive gathering task
+		RayTracingScene.UsedCoarseMeshStreamingHandles = MoveTemp(TaskData.UsedCoarseMeshStreamingHandles);
 
-							int32 MainCommandOffset = DirtyShadingBindingOffset;
-							int32 DecalCommandOffset = DirtyShadingBindingOffset + (bNeedMainInstance ? RTLODData.CachedMeshCommandIndices.Num() : 0);
+		// Inform the coarse mesh streaming manager about all the used streamable render assets in the scene
+		Nanite::FCoarseMeshStreamingManager* CoarseMeshSM = IStreamingManager::Get().GetNaniteCoarseMeshStreamingManager();
+		if (CoarseMeshSM)
+		{
+			CoarseMeshSM->AddUsedStreamingHandles(RayTracingScene.UsedCoarseMeshStreamingHandles);
+		}
 
-							for (int32 CommandIndex : RTLODData.CachedMeshCommandIndices)
-							{
-								const FRayTracingMeshCommand& MeshCommand = Scene.CachedRayTracingMeshCommands[CommandIndex];
-								
-								if (bNeedMainInstance)
-								{
-									const bool bHidden = MeshCommand.bDecal;									
-									const uint32 RecordIndex = RTLODData.SBTAllocation->GetRecordIndex(ERayTracingSceneLayer::Base, MeshCommand.GeometrySegmentIndex);
-									DirtyShaderBindingData[MainCommandOffset] = FRayTracingShaderBindingData(&MeshCommand, RelevantPrimitive.CachedRayTracingInstance->GeometryRHI, RecordIndex, bHidden);
-									++MainCommandOffset;
-								}
+		INC_DWORD_STAT_BY(STAT_VisibleRayTracingPrimitives, TaskData.StaticPrimitives.Num() + TaskData.CachedStaticPrimitives.Num() + TaskData.DynamicPrimitives.Num());
 
-								if (bNeedDecalInstance)
-								{
-									const bool bHidden = !MeshCommand.bDecal;									
-									const uint32 RecordIndex = RTLODData.SBTAllocation->GetRecordIndex(ERayTracingSceneLayer::Decals, MeshCommand.GeometrySegmentIndex);
-									DirtyShaderBindingData[DecalCommandOffset] = FRayTracingShaderBindingData(&MeshCommand, RelevantPrimitive.CachedRayTracingInstance->GeometryRHI, RecordIndex, bHidden);
-									++DecalCommandOffset;
-								}
-							}
-						}
-					});
-				}
-			}
-		};
+		AddDynamicInstancesToRayTracingScene(
+			GraphBuilder,
+			Scene,
+			View,
+			TaskData.SceneOptions,
+			TaskData.DynamicPrimitives,
+			RayTracingScene,
+			RayTracingSBT,
+			InDynamicReadBuffer,
+			InBulkAllocator);
+
+		// Task to iterate over static ray tracing instances (also perform auto-instancing on non-cached instances).
+		// This adds final instances to the ray tracing scene and must be done before FRayTracingScene::BuildInitializationData().
+		FGraphEventRef AddStaticInstancesTask = FFunctionGraphTask::CreateAndDispatchWhenReady(
+			[&TaskData, &Scene, &View, &RayTracingScene, &RayTracingSBT]()
+			{
+				FTaskTagScope TaskTagScope(ETaskTag::EParallelRenderingThread);
 
-		FGraphEventArray AddStaticInstancesTaskPrerequisites;
-		AddStaticInstancesTaskPrerequisites.Add(TaskData.StaticPrimitiveLODTask);
+				View.DirtyRayTracingShaderBindings.Reserve(TaskData.StaticPrimitives.Num() + TaskData.NumCachedStaticDirtyShaderBindings);
+
+				AddStaticInstancesToRayTracingScene(
+					Scene,
+					TaskData.SceneOptions,
+					TaskData.StaticPrimitives,
+					RayTracingScene,
+					View.DirtyRayTracingShaderBindings);
+
+				AddCachedStaticInstancesToRayTracingScene(
+					Scene,
+					TaskData.SceneOptions,
+					TaskData.CachedStaticPrimitives,
+					TaskData.GatherContexts,
+					TaskData.NumCachedStaticInstances,
+					TaskData.NumCachedStaticDecalInstances,
+					TaskData.NumCachedStaticSegments,
+					TaskData.NumCachedStaticDirtyShaderBindings,
+					RayTracingScene,
+					View.DirtyRayTracingShaderBindings);
 
-		FGraphEventRef AddStaticInstancesTask = TGraphTask<FRayTracingSceneAddStaticInstancesTask>::CreateTask(&AddStaticInstancesTaskPrerequisites).ConstructAndDispatchWhenReady(
-			// inputs
-			TaskData,
-			bool(View.Family->EngineShowFlags.PathTracing),
-			// outputs
-			RayTracingScene,
-			View.DirtyRayTracingShaderBindings
-		);
+			},
+			TStatId(), TaskData.StaticPrimitiveLODTask, ENamedThreads::AnyThread);
 
 		// Scene init task can run only when all pre-init tasks are complete (including culling tasks that are spawned while adding instances)
 		View.RayTracingSceneInitTask = FFunctionGraphTask::CreateAndDispatchWhenReady(
