Fixed crash due to missing virtual shadow map one-pass projection mask texture when no local lights use VSM but still trigger shadow projection (e.g., just distance fied shadows).

#rnx

diff --git a/Engine/Source/Runtime/Renderer/Private/SceneRendering.h b/Engine/Source/Runtime/Renderer/Private/SceneRendering.h
--- a/Engine/Source/Runtime/Renderer/Private/SceneRendering.h
+++ b/Engine/Source/Runtime/Renderer/Private/SceneRendering.h
@@ -225,6 +225,11 @@ public:
 	*/
 	bool ContainsOnlyVirtualShadowMaps() const;
 
+	/**
+	 * Returns true if there are any virtual shadow maps for any views for this light.
+	 */
+	bool HasVirtualShadowMap() const { return VirtualShadowMapId != INDEX_NONE; }
+
 	/**
 	* Prefer this to direct access of the VirtualShadowMapId member when a view is known.
 	* For directional lights this will attempt to find a clipmap associated with the given view,
diff --git a/Engine/Source/Runtime/Renderer/Private/Shadows/ShadowSceneRenderer.cpp b/Engine/Source/Runtime/Renderer/Private/Shadows/ShadowSceneRenderer.cpp
--- a/Engine/Source/Runtime/Renderer/Private/Shadows/ShadowSceneRenderer.cpp
+++ b/Engine/Source/Runtime/Renderer/Private/Shadows/ShadowSceneRenderer.cpp
@@ -549,6 +549,12 @@ void FShadowSceneRenderer::ApplyVirtualShadowMapProjectionForLight(
 	}
 
 	const FVisibleLightInfo& VisibleLightInfo = SceneRenderer.VisibleLightInfos[LightSceneInfo->Id];
+
+	if (!VisibleLightInfo.HasVirtualShadowMap())
+	{
+		return;
+	}
+
 	FSceneTextureParameters SceneTextureParameters = GetSceneTextureParameters(GraphBuilder, SceneTextures.UniformBuffer);
 
 	for (int32 ViewIndex = 0; ViewIndex < SceneRenderer.Views.Num(); ViewIndex++)
