Alpha: Further remove primitive alpha holdout work when the new project setting is disabled.

#rnx

diff --git a/Engine/Source/Runtime/Renderer/Private/DeferredShadingRenderer.cpp b/Engine/Source/Runtime/Renderer/Private/DeferredShadingRenderer.cpp
--- a/Engine/Source/Runtime/Renderer/Private/DeferredShadingRenderer.cpp
+++ b/Engine/Source/Runtime/Renderer/Private/DeferredShadingRenderer.cpp
@@ -2961,7 +2961,7 @@ void FDeferredShadingSceneRenderer::Render(FRDGBuilder& GraphBuilder)
 					bComposeWithWater,
 					SceneWithoutWaterTextures, SceneTextures);
 
-				if (IsPostProcessingWithAlphaChannelSupported())
+				if (IsPrimitiveAlphaHoldoutEnabledForAnyView(Views))
 				{
 					// When alpha is enabled to work with holdout. We need another full screen tracing pass to update the alpha channel containing the "holdout alpha throughput".
 					// Alpha hold out only works when using r.volumetricrendertarget.mode 3 which is the mode use by MRQ.
diff --git a/Engine/Source/Runtime/Renderer/Private/HeterogeneousVolumes/HeterogeneousVolumes.cpp b/Engine/Source/Runtime/Renderer/Private/HeterogeneousVolumes/HeterogeneousVolumes.cpp
--- a/Engine/Source/Runtime/Renderer/Private/HeterogeneousVolumes/HeterogeneousVolumes.cpp
+++ b/Engine/Source/Runtime/Renderer/Private/HeterogeneousVolumes/HeterogeneousVolumes.cpp
@@ -883,7 +883,7 @@ void FDeferredShadingSceneRenderer::RenderHeterogeneousVolumes(
 		HeterogeneousVolumeRadiance = GraphBuilder.CreateTexture(Desc, TEXT("HeterogeneousVolumes"));
 		AddClearUAVPass(GraphBuilder, GraphBuilder.CreateUAV(HeterogeneousVolumeRadiance), FLinearColor::Black);
 
-		if (IsPostProcessingWithAlphaChannelSupported())
+		if (IsPrimitiveAlphaHoldoutEnabledForAnyView(Views))
 		{
 			Desc.Format = PF_R8;
 			HeterogeneousVolumeHoldout = GraphBuilder.CreateTexture(Desc, TEXT("HeterogeneousVolume.Holdout"));
diff --git a/Engine/Source/Runtime/Renderer/Private/HeterogeneousVolumes/HeterogeneousVolumesLiveShadingPipeline.cpp b/Engine/Source/Runtime/Renderer/Private/HeterogeneousVolumes/HeterogeneousVolumesLiveShadingPipeline.cpp
--- a/Engine/Source/Runtime/Renderer/Private/HeterogeneousVolumes/HeterogeneousVolumesLiveShadingPipeline.cpp
+++ b/Engine/Source/Runtime/Renderer/Private/HeterogeneousVolumes/HeterogeneousVolumesLiveShadingPipeline.cpp
@@ -1219,7 +1219,7 @@ void RenderSingleScatteringWithLiveShading(
 
 		// Output
 		PassParameters->RWLightingTexture = GraphBuilder.CreateUAV(HeterogeneousVolumeRadianceTexture);
-		if (IsPostProcessingWithAlphaChannelSupported())
+		if (IsPrimitiveAlphaHoldoutEnabled(GetFeatureLevelShadingPath(View.GetFeatureLevel())))
 		{
 			PassParameters->RWHoldoutTexture = GraphBuilder.CreateUAV(HeterogeneousVolumeHoldoutTexture);
 		}
diff --git a/Engine/Source/Runtime/Renderer/Private/PostProcess/TemporalSuperResolution.cpp b/Engine/Source/Runtime/Renderer/Private/PostProcess/TemporalSuperResolution.cpp
--- a/Engine/Source/Runtime/Renderer/Private/PostProcess/TemporalSuperResolution.cpp
+++ b/Engine/Source/Runtime/Renderer/Private/PostProcess/TemporalSuperResolution.cpp
@@ -2259,7 +2259,7 @@ FDefaultTemporalUpscaler::FOutputs AddTemporalSuperResolutionPasses(
 		PassParameters->TileOverscan = TileOverscan;
 		PassParameters->bEnableResurrection = bCanResurrectHistory;
 		PassParameters->bEnableFlickeringHeuristic = FlickeringFramePeriod > 0.0f;
-		PassParameters->bPassthroughAlpha = IsTranslucentHoldoutEnabled(GetFeatureLevelShadingPath(View.GetFeatureLevel()));
+		PassParameters->bPassthroughAlpha = IsPrimitiveAlphaHoldoutEnabled(GetFeatureLevelShadingPath(View.GetFeatureLevel()));
 
 		PassParameters->InputTexture = PassInputs.SceneColor.Texture;
 		if (PassInputs.FlickeringInputTexture.IsValid())
diff --git a/Engine/Source/Runtime/Renderer/Private/SceneRendering.cpp b/Engine/Source/Runtime/Renderer/Private/SceneRendering.cpp
--- a/Engine/Source/Runtime/Renderer/Private/SceneRendering.cpp
+++ b/Engine/Source/Runtime/Renderer/Private/SceneRendering.cpp
@@ -6651,6 +6651,26 @@ FRDGTextureRef CreateQuarterResolutionDepthMinAndMax(FRDGBuilder& GraphBuilder,
 	return SmallTexture;
 }
 
+bool IsPrimitiveAlphaHoldoutEnabled(EShadingPath ShadingPath)
+{
+	const bool bSupportPrimitiveAlphaHoldout = CVarPrimitiveAlphaHoldoutSupport.GetValueOnRenderThread();
+
+	return bSupportPrimitiveAlphaHoldout && (ShadingPath != EShadingPath::Mobile) && IsPostProcessingWithAlphaChannelSupported();
+}
+
+bool IsPrimitiveAlphaHoldoutEnabledForAnyView(TArrayView<FViewInfo> Views)
+{
+	for (const FViewInfo& View : Views)
+	{
+		if (IsPrimitiveAlphaHoldoutEnabled(GetFeatureLevelShadingPath(View.GetFeatureLevel())))
+		{
+			return true;
+		}
+	}
+	
+	return false;
+}
+
 bool SceneCaptureRequiresAlphaChannel(const FSceneView& View)
 {
 	// Planar reflections and scene captures use scene color alpha to keep track of where content has been rendered, for compositing into a different scene later
diff --git a/Engine/Source/Runtime/Renderer/Private/SceneRendering.h b/Engine/Source/Runtime/Renderer/Private/SceneRendering.h
--- a/Engine/Source/Runtime/Renderer/Private/SceneRendering.h
+++ b/Engine/Source/Runtime/Renderer/Private/SceneRendering.h
@@ -2935,6 +2935,12 @@ inline const FViewFamilyInfo& GetViewFamilyInfo(const TArrayView<const FViewInfo
 	return *(const FViewFamilyInfo*)Views[0].Family;
 }
 
+/** Checks whether primitive alpha holdout is active.*/
+bool IsPrimitiveAlphaHoldoutEnabled(EShadingPath ShadingPath);
+
+/** Checks whether primitive alpha holdout is active for any view.*/
+bool IsPrimitiveAlphaHoldoutEnabledForAnyView(TArrayView<FViewInfo> Views);
+
 bool SceneCaptureRequiresAlphaChannel(const FSceneView& View);
 
 /** Checks whether the material and scene proxy combination will modify the mesh position. Useful for determining whether the material can be substituted with a default material for depth rendering etc. */
diff --git a/Engine/Source/Runtime/Renderer/Private/SceneVisibility.cpp b/Engine/Source/Runtime/Renderer/Private/SceneVisibility.cpp
--- a/Engine/Source/Runtime/Renderer/Private/SceneVisibility.cpp
+++ b/Engine/Source/Runtime/Renderer/Private/SceneVisibility.cpp
@@ -1361,7 +1361,7 @@ void FRelevancePacket::ComputeRelevance(FDynamicPrimitiveIndexList& DynamicPrimi
 	const bool bMobileMaskedInEarlyPass = (ShadingPath == EShadingPath::Mobile) && Scene.EarlyZPassMode == DDM_MaskedOnly;
 	const bool bMobileBasePassAlwaysUsesCSM = (ShadingPath == EShadingPath::Mobile) && MobileBasePassAlwaysUsesCSM(Scene.GetShaderPlatform());
 	const bool bVelocityPassWritesDepth = Scene.EarlyZPassMode == DDM_AllOpaqueNoVelocity;
-	const bool bIsTranslucentHoldoutEnabled = IsTranslucentHoldoutEnabled(ShadingPath);
+	const bool bIsTranslucentHoldoutEnabled = IsPrimitiveAlphaHoldoutEnabled(ShadingPath);
 	const bool bHLODActive = Scene.SceneLODHierarchy.IsActive();
 	const FHLODVisibilityState* const HLODState = bHLODActive && ViewState ? &ViewState->HLODVisibilityState : nullptr;
 	float MaxDrawDistanceScale = GetCachedScalabilityCVars().ViewDistanceScale;
diff --git a/Engine/Source/Runtime/Renderer/Private/TranslucentRendering.cpp b/Engine/Source/Runtime/Renderer/Private/TranslucentRendering.cpp
--- a/Engine/Source/Runtime/Renderer/Private/TranslucentRendering.cpp
+++ b/Engine/Source/Runtime/Renderer/Private/TranslucentRendering.cpp
@@ -105,38 +105,6 @@ static TAutoConsoleVariable<int32> CVarParallelTranslucency(
 	TEXT("Toggles parallel translucency rendering. Parallel rendering must be enabled for this to have an effect."),
 	ECVF_RenderThreadSafe);
 
-static TAutoConsoleVariable<int32> CVarRenderTranslucentHoldout(
-	TEXT("r.RenderTranslucentHoldout"),
-	-1,
-	TEXT("Toggles translucency holdout rendering.\n")
-	TEXT("   0: Disable the dedicated translucent holdout pass (legacy behavior). \n")
-	TEXT("Else: Automatically turn on/off the dedicated holdout pass based on whether Propagate Alpha is on/off\n"),
-	ECVF_RenderThreadSafe);
-
-bool IsTranslucentHoldoutEnabled(EShadingPath ShadingPath)
-{
-	const int32 RenderTranslucentHoldout = CVarRenderTranslucentHoldout.GetValueOnRenderThread();
-	bool bPropagateAlpha = false;
-
-	if (RenderTranslucentHoldout != 0)
-	{
-		static IConsoleVariable* CVarPropagateAlpha = IConsoleManager::Get().FindConsoleVariable(TEXT("r.PostProcessing.PropagateAlpha"));
-		static IConsoleVariable* CVarSupportAlphaHoldout = IConsoleManager::Get().FindConsoleVariable(TEXT("r.Deferred.SupportPrimitiveAlphaHoldout"));
-		const bool bDeferredPropagateAlpha = CVarPropagateAlpha ? CVarPropagateAlpha->GetBool() : false;
-		const bool bDeferredSupportPrimitiveAlphaHoldout = CVarSupportAlphaHoldout ? CVarSupportAlphaHoldout->GetBool() : false;
-		// TODO: add support to Mobile renderer 
-		bPropagateAlpha = bDeferredPropagateAlpha && bDeferredSupportPrimitiveAlphaHoldout && (ShadingPath != EShadingPath::Mobile);
-	}
-
-	return bPropagateAlpha;
-}
-
-static bool IsTranslucentHoldoutEnabled(TArrayView<const FViewInfo> Views)
-{
-	ensure(Views.Num() > 0);
-	return IsTranslucentHoldoutEnabled(GetFeatureLevelShadingPath(Views[0].GetFeatureLevel()));
-}
-
 DynamicRenderScaling::FHeuristicSettings GetDynamicTranslucencyResolutionSettings()
 {
 	DynamicRenderScaling::FHeuristicSettings BucketSetting;
@@ -656,7 +624,7 @@ FScreenPassTexture FTranslucencyComposition::AddPass(
 	RDG_GPU_STAT_SCOPE(GraphBuilder, Translucency);
 	DynamicRenderScaling::FRDGScope DynamicTranslucencyResolutionScope(GraphBuilder, GDynamicTranslucencyResolution);
 
-	bool bPassthroughAlpha = IsTranslucentHoldoutEnabled(GetFeatureLevelShadingPath(View.GetFeatureLevel()));
+	bool bPassthroughAlpha = IsPrimitiveAlphaHoldoutEnabled(GetFeatureLevelShadingPath(View.GetFeatureLevel()));
 
 	const TCHAR* OpName = nullptr;
 	FRHIBlendState* BlendState = nullptr;
@@ -1819,7 +1787,7 @@ void FDeferredShadingSceneRenderer::RenderTranslucency(
 	DynamicRenderScaling::FRDGScope DynamicTranslucencyResolutionScope(GraphBuilder, GDynamicTranslucencyResolution);
 
 	FRDGTextureRef SceneColorCopyTexture = nullptr;
-	const bool bIsTranslucentHoldoutEnabled = IsTranslucentHoldoutEnabled(Views);
+	const bool bIsTranslucentHoldoutEnabled = IsPrimitiveAlphaHoldoutEnabledForAnyView(Views);
 
 	if (EnumHasAnyFlags(ViewsToRender, ETranslucencyView::AboveWater))
 	{
diff --git a/Engine/Source/Runtime/Renderer/Private/TranslucentRendering.h b/Engine/Source/Runtime/Renderer/Private/TranslucentRendering.h
--- a/Engine/Source/Runtime/Renderer/Private/TranslucentRendering.h
+++ b/Engine/Source/Runtime/Renderer/Private/TranslucentRendering.h
@@ -50,9 +50,6 @@ FSeparateTranslucencyDimensions UpdateSeparateTranslucencyDimensions(const FScen
 /** Returns whether the view family is requesting to render translucency. */
 bool ShouldRenderTranslucency(const FSceneViewFamily& ViewFamily);
 
-/** Check whether translucent holdout is enabled for a specific shading path. */
-bool IsTranslucentHoldoutEnabled(EShadingPath ShadingPath);
-
 /** Check if separate translucency pass is needed for given pass and downsample scale */
 bool IsSeparateTranslucencyEnabled(ETranslucencyPass::Type TranslucencyPass, float DownsampleScale);
 
