[Ray Tracing Instance Gathering Refactor] Replaced StaticPrimitiveLODTask with GatherRelevantStaticPrimitivesTask.
- Moved relevant logic out of GatherRelevantPrimitives(...) and into a standalone function GatherRelevantStaticPrimitives(...)

diff --git a/Engine/Source/Runtime/Renderer/Private/RayTracing/RayTracing.cpp b/Engine/Source/Runtime/Renderer/Private/RayTracing/RayTracing.cpp
--- a/Engine/Source/Runtime/Renderer/Private/RayTracing/RayTracing.cpp
+++ b/Engine/Source/Runtime/Renderer/Private/RayTracing/RayTracing.cpp
@@ -245,17 +245,19 @@ namespace RayTracing
 		FSceneOptions SceneOptions;
 
 		// Filtered lists of relevant primitives
+		TArray<int32> StaticPrimitivesIndices;
+		TArray<int32> DynamicPrimitivesIndices;
+
 		TArray<FRelevantPrimitive> StaticPrimitives;
 		TArray<FRelevantPrimitive> CachedStaticPrimitives;
-		TArray<int32> DynamicPrimitives;
 
 		TArray<FRelevantPrimitiveGatherContext> GatherContexts;
 
+		// This task must complete before accessing StaticPrimitivesIndices/DynamicPrimitivesIndices.
 		UE::Tasks::FTask GatherRelevantPrimitivesTask;
 
-		// Relevant static primitive LODs are computed asynchronously.
-		// This task must complete before accessing StaticPrimitives/CachedStaticPrimitives in FRayTracingSceneAddInstancesTask.
-		UE::Tasks::FTask StaticPrimitiveLODTask;
+		// This task must complete before accessing StaticPrimitives/CachedStaticPrimitives.
+		UE::Tasks::FTask GatherRelevantStaticPrimitivesTask;
 
 		// Array of primitives that should update their cached ray tracing instances via FPrimitiveSceneInfo::UpdateCachedRaytracingData()
 		TArray<FPrimitiveSceneInfo*> DirtyCachedRayTracingPrimitives; // TODO: remove this since it seems to be transient
@@ -469,8 +471,6 @@ namespace RayTracing
 		FScene& Scene = TaskData.Scene;
 		FViewInfo& View = TaskData.View;
 
-		TArray<int32> StaticPrimitives;
-
 		const bool bGameView = View.bIsGameView || View.Family->EngineShowFlags.Game;
 
 		bool bPerformRayTracing = View.State != nullptr && !View.bIsReflectionCapture && View.bAllowRayTracing;
@@ -591,15 +591,15 @@ namespace RayTracing
 					NumDirtyCachedRayTracingPrimitives += Context.DirtyCachedRayTracingPrimitives.Num();
 				}
 
-				StaticPrimitives.Reserve(NumStaticPrimitives);
-				TaskData.DynamicPrimitives.Reserve(NumDynamicPrimitives);
+				TaskData.StaticPrimitivesIndices.Reserve(NumStaticPrimitives);
+				TaskData.DynamicPrimitivesIndices.Reserve(NumDynamicPrimitives);
 				TaskData.UsedCoarseMeshStreamingHandles.Reserve(NumUsedCoarseMeshStreamingHandles);
 				TaskData.DirtyCachedRayTracingPrimitives.Reserve(NumDirtyCachedRayTracingPrimitives);
 
 				for (auto& Context : Contexts)
 				{
-					Context.StaticPrimitives.CopyToLinearArray(StaticPrimitives);
-					Context.DynamicPrimitives.CopyToLinearArray(TaskData.DynamicPrimitives);
+					Context.StaticPrimitives.CopyToLinearArray(TaskData.StaticPrimitivesIndices);
+					Context.DynamicPrimitives.CopyToLinearArray(TaskData.DynamicPrimitivesIndices);
 					Context.UsedCoarseMeshStreamingHandles.CopyToLinearArray(TaskData.UsedCoarseMeshStreamingHandles);
 					Context.DirtyCachedRayTracingPrimitives.CopyToLinearArray(TaskData.DirtyCachedRayTracingPrimitives);
 
@@ -613,249 +613,240 @@ namespace RayTracing
 
 		// TODO: check whether it's ok to do this on a parallel task
 		FPrimitiveSceneInfo::UpdateCachedRaytracingData(&TaskData.Scene, TaskData.DirtyCachedRayTracingPrimitives);
+	}
 
-		static const auto ICVarStaticMeshLODDistanceScale = IConsoleManager::Get().FindConsoleVariable(TEXT("r.StaticMeshLODDistanceScale"));
-		const float LODScaleCVarValue = ICVarStaticMeshLODDistanceScale->GetFloat();
-		const int32 ForcedLODLevel = GetCVarForceLOD();
+	void GatherRelevantStaticPrimitives(FGatherInstancesTaskData& TaskData, float GlobalLODScale, int32 ForcedLODLevel, bool bUsingReferenceBasedResidency)
+	{
+		TRACE_CPUPROFILER_EVENT_SCOPE(GatherRayTracingRelevantStaticPrimitives);
 
-		TaskData.StaticPrimitiveLODTask = UE::Tasks::Launch(UE_SOURCE_LOCATION,
-			[&TaskData, LODScaleCVarValue, ForcedLODLevel, StaticPrimitiveIndices = MoveTemp(StaticPrimitives), bUsingReferenceBasedResidency]()
+		struct FRelevantStaticPrimitivesContext
+		{
+			FRelevantStaticPrimitivesContext(int32 InContextIndex) : ContextIndex(InContextIndex) {}
+
+			TChunkedArray<FRelevantPrimitive> StaticPrimitives;
+			TChunkedArray<FRelevantPrimitive> CachedStaticPrimitives;
+			TChunkedArray<const FPrimitiveSceneInfo*> VisibleNaniteRayTracingPrimitives;
+
+			int32 NumCachedStaticInstances = 0;
+			int32 NumCachedStaticDecalInstances = 0;
+			int32 NumCachedStaticSegments = 0;
+			int32 NumCachedStaticDirtyShaderBindings = 0;
+
+			int32 ContextIndex = INDEX_NONE;
+		};
+
+		TArray<FRelevantStaticPrimitivesContext> Contexts;
+		ParallelForWithTaskContext(
+			TEXT("GatherRayTracingRelevantStaticPrimitives_Parallel"),
+			Contexts,
+			TaskData.StaticPrimitivesIndices.Num(),
+			[](int32 ContextIndex, int32 NumContexts) { return ContextIndex; },
+			[&Scene = TaskData.Scene, &View = TaskData.View, GlobalLODScale, ForcedLODLevel, &StaticPrimitivesIndices = TaskData.StaticPrimitivesIndices, bUsingReferenceBasedResidency](FRelevantStaticPrimitivesContext& Context, int32 ItemIndex)
 			{
-				FTaskTagScope TaskTagScope(ETaskTag::EParallelRenderingThread);
+				const int32 PrimitiveIndex = StaticPrimitivesIndices[ItemIndex];
 
-				TRACE_CPUPROFILER_EVENT_SCOPE(GatherRayTracingRelevantPrimitives_ComputeLOD);
-
-				struct FRelevantStaticPrimitivesContext
-				{
-					FRelevantStaticPrimitivesContext(int32 InContextIndex) : ContextIndex(InContextIndex) {}
-
-					TChunkedArray<FRelevantPrimitive> StaticPrimitives;
-					TChunkedArray<FRelevantPrimitive> CachedStaticPrimitives;
-					TChunkedArray<const FPrimitiveSceneInfo*> VisibleNaniteRayTracingPrimitives;
-
-					int32 NumCachedStaticInstances = 0;
-					int32 NumCachedStaticDecalInstances = 0;
-					int32 NumCachedStaticSegments = 0;
-					int32 NumCachedStaticDirtyShaderBindings = 0;
-
-					int32 ContextIndex = INDEX_NONE;
-				};
-
-				TArray<FRelevantStaticPrimitivesContext> Contexts;
-				ParallelForWithTaskContext(
-					TEXT("GatherRayTracingRelevantPrimitives_ComputeLOD_Parallel"),
-					Contexts,
-					StaticPrimitiveIndices.Num(),
-					[](int32 ContextIndex, int32 NumContexts) { return ContextIndex; },
-					[&Scene = TaskData.Scene, &View = TaskData.View, LODScaleCVarValue, ForcedLODLevel, &StaticPrimitiveIndices, bUsingReferenceBasedResidency](FRelevantStaticPrimitivesContext& Context, int32 ItemIndex)
-					{
-						const int32 PrimitiveIndex = StaticPrimitiveIndices[ItemIndex];
-
-						const FPrimitiveSceneProxy* SceneProxy = Scene.PrimitiveSceneProxies[PrimitiveIndex];
-						const FPrimitiveSceneInfo* SceneInfo = Scene.Primitives[PrimitiveIndex];
-						const ERayTracingPrimitiveFlags Flags = Scene.PrimitiveRayTracingFlags[PrimitiveIndex];
+				const FPrimitiveSceneProxy* SceneProxy = Scene.PrimitiveSceneProxies[PrimitiveIndex];
+				const FPrimitiveSceneInfo* SceneInfo = Scene.Primitives[PrimitiveIndex];
+				const ERayTracingPrimitiveFlags Flags = Scene.PrimitiveRayTracingFlags[PrimitiveIndex];
 
-						const bool bUsingNaniteRayTracing = (Nanite::GetRayTracingMode() != Nanite::ERayTracingMode::Fallback) && SceneProxy->IsNaniteMesh();
+				const bool bUsingNaniteRayTracing = (Nanite::GetRayTracingMode() != Nanite::ERayTracingMode::Fallback) && SceneProxy->IsNaniteMesh();
 
-						if (bUsingNaniteRayTracing)
-						{
-							Context.VisibleNaniteRayTracingPrimitives.AddElement(SceneInfo);
-						}
+				if (bUsingNaniteRayTracing)
+				{
+					Context.VisibleNaniteRayTracingPrimitives.AddElement(SceneInfo);
+				}
 
-						int8 LODIndex = 0;
+				int8 LODIndex = 0;
 
-						if (EnumHasAnyFlags(Flags, ERayTracingPrimitiveFlags::ComputeLOD))
-						{
-							const FPrimitiveBounds& Bounds = Scene.PrimitiveBounds[PrimitiveIndex];
+				if (EnumHasAnyFlags(Flags, ERayTracingPrimitiveFlags::ComputeLOD))
+				{
+					const FPrimitiveBounds& Bounds = Scene.PrimitiveBounds[PrimitiveIndex];
 
-							const int8 CurFirstLODIdx = SceneProxy->GetCurrentFirstLODIdx_RenderThread();
-							check(CurFirstLODIdx >= 0);
+					const int8 CurFirstLODIdx = SceneProxy->GetCurrentFirstLODIdx_RenderThread();
+					check(CurFirstLODIdx >= 0);
 
-							float MeshScreenSizeSquared = 0;
-							float LODScale = LODScaleCVarValue * View.LODDistanceFactor;
-							FLODMask LODToRender = ComputeLODForMeshes(SceneInfo->StaticMeshRelevances, View, Bounds.BoxSphereBounds.Origin, Bounds.BoxSphereBounds.SphereRadius, ForcedLODLevel, MeshScreenSizeSquared, CurFirstLODIdx, LODScale, true);
+					float MeshScreenSizeSquared = 0;
+					float LODScale = GlobalLODScale * View.LODDistanceFactor;
+					FLODMask LODToRender = ComputeLODForMeshes(SceneInfo->StaticMeshRelevances, View, Bounds.BoxSphereBounds.Origin, Bounds.BoxSphereBounds.SphereRadius, ForcedLODLevel, MeshScreenSizeSquared, CurFirstLODIdx, LODScale, true);
 
-							LODIndex = LODToRender.GetRayTracedLOD();
-						}
+					LODIndex = LODToRender.GetRayTracedLOD();
+				}
 
-						if (EnumHasAnyFlags(Flags, ERayTracingPrimitiveFlags::CacheInstances))
+				if (EnumHasAnyFlags(Flags, ERayTracingPrimitiveFlags::CacheInstances))
+				{
+					if (bUsingNaniteRayTracing)
+					{
+						if (SceneInfo->CachedRayTracingInstance.GeometryRHI == nullptr)
 						{
-							if (bUsingNaniteRayTracing)
-							{
-								if (SceneInfo->CachedRayTracingInstance.GeometryRHI == nullptr)
-								{
-									// Nanite ray tracing geometry not ready yet, doesn't include primitive in ray tracing scene
-									return;
-								}
-							}
-							else
-							{
-								// Currently IsCachedRayTracingGeometryValid() can only be called for non-nanite geometries
-								checkf(SceneInfo->IsCachedRayTracingGeometryValid(), TEXT("Cached ray tracing instance is expected to be valid. Was mesh LOD streamed but cached data was not invalidated?"));
-								checkf(SceneInfo->CachedRayTracingInstance.GeometryRHI, TEXT("Ray tracing instance must have a valid geometry."));
-							}
+							// Nanite ray tracing geometry not ready yet, doesn't include primitive in ray tracing scene
+							return;
+						}
+					}
+					else
+					{
+						// Currently IsCachedRayTracingGeometryValid() can only be called for non-nanite geometries
+						checkf(SceneInfo->IsCachedRayTracingGeometryValid(), TEXT("Cached ray tracing instance is expected to be valid. Was mesh LOD streamed but cached data was not invalidated?"));
+						checkf(SceneInfo->CachedRayTracingInstance.GeometryRHI, TEXT("Ray tracing instance must have a valid geometry."));
+					}
 
-							if (ShouldExcludeDecals() && SceneInfo->bCachedRayTracingInstanceAllSegmentsDecal)
-							{
-								return;
-							}
+					if (ShouldExcludeDecals() && SceneInfo->bCachedRayTracingInstanceAllSegmentsDecal)
+					{
+						return;
+					}
 
-							checkf(SceneInfo->CachedRayTracingInstance.GeometryRHI, TEXT("Ray tracing instance must have a valid geometry."));
+					checkf(SceneInfo->CachedRayTracingInstance.GeometryRHI, TEXT("Ray tracing instance must have a valid geometry."));
 
-							// For primitives with ERayTracingPrimitiveFlags::CacheInstances flag we only cache the instance/mesh commands of the current LOD
-							// (see FPrimitiveSceneInfo::UpdateCachedRayTracingInstance(...) and CacheRayTracingPrimitive(...))
-							check(!EnumHasAnyFlags(Flags, ERayTracingPrimitiveFlags::ComputeLOD));
-							LODIndex = 0;
+					// For primitives with ERayTracingPrimitiveFlags::CacheInstances flag we only cache the instance/mesh commands of the current LOD
+					// (see FPrimitiveSceneInfo::UpdateCachedRayTracingInstance(...) and CacheRayTracingPrimitive(...))
+					check(!EnumHasAnyFlags(Flags, ERayTracingPrimitiveFlags::ComputeLOD));
+					LODIndex = 0;
 
-							const FPrimitiveSceneInfo::FRayTracingLODData& RTLODData = SceneInfo->RayTracingLODData[LODIndex];
-							ensure(RTLODData.SBTAllocation);
-							if (RTLODData.SBTAllocation == nullptr)
-							{
-								return;
-							}
+					const FPrimitiveSceneInfo::FRayTracingLODData& RTLODData = SceneInfo->RayTracingLODData[LODIndex];
+					ensure(RTLODData.SBTAllocation);
+					if (RTLODData.SBTAllocation == nullptr)
+					{
+						return;
+					}
 
-							FRelevantPrimitive* RelevantPrimitive = new (Context.CachedStaticPrimitives) FRelevantPrimitive();
-							RelevantPrimitive->PrimitiveIndex = PrimitiveIndex;
-							RelevantPrimitive->PersistentPrimitiveIndex = SceneInfo->GetPersistentIndex();
-							RelevantPrimitive->RayTracingGeometry = SceneInfo->GetCachedRayTracingGeometry();
+					FRelevantPrimitive* RelevantPrimitive = new (Context.CachedStaticPrimitives) FRelevantPrimitive();
+					RelevantPrimitive->PrimitiveIndex = PrimitiveIndex;
+					RelevantPrimitive->PersistentPrimitiveIndex = SceneInfo->GetPersistentIndex();
+					RelevantPrimitive->RayTracingGeometry = SceneInfo->GetCachedRayTracingGeometry();
 
-							ensureMsgf(!SceneInfo->bCachedRaytracingDataDirty, TEXT("Cached ray tracing instances must be up-to-date at this point"));
+					ensureMsgf(!SceneInfo->bCachedRaytracingDataDirty, TEXT("Cached ray tracing instances must be up-to-date at this point"));
 
-							RelevantPrimitive->CachedRayTracingInstance = &SceneInfo->CachedRayTracingInstance;
-							RelevantPrimitive->CachedMeshCommandFlags.bAnySegmentsDecal = SceneInfo->bCachedRayTracingInstanceAnySegmentsDecal;
-							RelevantPrimitive->CachedMeshCommandFlags.bAllSegmentsDecal = SceneInfo->bCachedRayTracingInstanceAllSegmentsDecal;								
+					RelevantPrimitive->CachedRayTracingInstance = &SceneInfo->CachedRayTracingInstance;
+					RelevantPrimitive->CachedMeshCommandFlags.bAnySegmentsDecal = SceneInfo->bCachedRayTracingInstanceAnySegmentsDecal;
+					RelevantPrimitive->CachedMeshCommandFlags.bAllSegmentsDecal = SceneInfo->bCachedRayTracingInstanceAllSegmentsDecal;
 
-							// CacheInstances expects to have one ray tracing mesh command per BLAS segment.
-							// If that's not the case in the future, other logic such as NumCachedStaticVisibleMeshCommands calculation needs to be updated.
-							checkf(RTLODData.CachedMeshCommandIndices.Num() == RelevantPrimitive->RayTracingGeometry->Initializer.Segments.Num(),
-								TEXT("Expected to have one ray tracing mesh command per BLAS segment (primitive has %d cached mesh commands but BLAS has %d segments)."),
-								RTLODData.CachedMeshCommandIndices.Num(), SceneInfo->CachedRayTracingInstance.GeometryRHI->GetNumSegments());
+					// CacheInstances expects to have one ray tracing mesh command per BLAS segment.
+					// If that's not the case in the future, other logic such as NumCachedStaticVisibleMeshCommands calculation needs to be updated.
+					checkf(RTLODData.CachedMeshCommandIndices.Num() == RelevantPrimitive->RayTracingGeometry->Initializer.Segments.Num(),
+						TEXT("Expected to have one ray tracing mesh command per BLAS segment (primitive has %d cached mesh commands but BLAS has %d segments)."),
+						RTLODData.CachedMeshCommandIndices.Num(), SceneInfo->CachedRayTracingInstance.GeometryRHI->GetNumSegments());
 
-							const bool bNeedMainInstance = !RelevantPrimitive->CachedMeshCommandFlags.bAllSegmentsDecal;
+					const bool bNeedMainInstance = !RelevantPrimitive->CachedMeshCommandFlags.bAllSegmentsDecal;
 
-							// if primitive has mixed decal and non-decal segments we need to have two ray tracing instances
-							// one containing non-decal segments and the other with decal segments
-							// masking of segments is done using "hidden" hitgroups
-							// TODO: Debug Visualization to highlight primitives using this?
-							const bool bNeedDecalInstance = RelevantPrimitive->CachedMeshCommandFlags.bAnySegmentsDecal && !ShouldExcludeDecals();
-							checkf(bNeedMainInstance || bNeedDecalInstance, TEXT("FRelevantPrimitive is expected to have a main instance, decal instance or both."));
+					// if primitive has mixed decal and non-decal segments we need to have two ray tracing instances
+					// one containing non-decal segments and the other with decal segments
+					// masking of segments is done using "hidden" hitgroups
+					// TODO: Debug Visualization to highlight primitives using this?
+					const bool bNeedDecalInstance = RelevantPrimitive->CachedMeshCommandFlags.bAnySegmentsDecal && !ShouldExcludeDecals();
+					checkf(bNeedMainInstance || bNeedDecalInstance, TEXT("FRelevantPrimitive is expected to have a main instance, decal instance or both."));
 
-							// For now store offsets relative to current context
-							// Will be patched later to be a global offset
-							RelevantPrimitive->RelativeInstanceOffset = Context.NumCachedStaticInstances;
-							RelevantPrimitive->RelativeDecalInstanceOffset = Context.NumCachedStaticDecalInstances;
-							RelevantPrimitive->ContextIndex = Context.ContextIndex;
+					// For now store offsets relative to current context
+					// Will be patched later to be a global offset
+					RelevantPrimitive->RelativeInstanceOffset = Context.NumCachedStaticInstances;
+					RelevantPrimitive->RelativeDecalInstanceOffset = Context.NumCachedStaticDecalInstances;
+					RelevantPrimitive->ContextIndex = Context.ContextIndex;
 
-							if (bNeedMainInstance)
-							{
-								++Context.NumCachedStaticInstances;
-							}
+					if (bNeedMainInstance)
+					{
+						++Context.NumCachedStaticInstances;
+					}
 
-							if (bNeedDecalInstance)
-							{
-								++Context.NumCachedStaticDecalInstances;
-							}
+					if (bNeedDecalInstance)
+					{
+						++Context.NumCachedStaticDecalInstances;
+					}
 
-							const uint32 NumTLASInstances = bNeedMainInstance && bNeedDecalInstance ? 2 : 1;
+					const uint32 NumTLASInstances = bNeedMainInstance && bNeedDecalInstance ? 2 : 1;
 
-							Context.NumCachedStaticSegments += RTLODData.CachedMeshCommandIndices.Num() * NumTLASInstances;
+					Context.NumCachedStaticSegments += RTLODData.CachedMeshCommandIndices.Num() * NumTLASInstances;
 
-							// Setup the shader bindings write offset which is computed here for all cached instances so the writing can happen in parallel to the dirty
-							// bindings array - with persistent SBTs a dirty flag on the binding data can be checked and then RelativeDirtyShaderBindingsOffset will be
-							// invalid and no SBT write operation needed anymore
-							RelevantPrimitive->RelativeDirtyShaderBindingsOffset = Context.NumCachedStaticDirtyShaderBindings;
-							Context.NumCachedStaticDirtyShaderBindings += RTLODData.SBTAllocation->GetSegmentCount();
-						}
-						// - DirtyCachedRayTracingPrimitives are only processed after StaticPrimitiveIndices is filled
-						// so we can end up with primitives that should be skipped here
-						// - once we update flags of primitive with dirty raytracing state before `GatherRayTracingRelevantPrimitives_Parallel`
-						// we should replace this condition with an assert instead
-						else if(!EnumHasAnyFlags(Flags, ERayTracingPrimitiveFlags::Skip))
-						{
+					// Setup the shader bindings write offset which is computed here for all cached instances so the writing can happen in parallel to the dirty
+					// bindings array - with persistent SBTs a dirty flag on the binding data can be checked and then RelativeDirtyShaderBindingsOffset will be
+					// invalid and no SBT write operation needed anymore
+					RelevantPrimitive->RelativeDirtyShaderBindingsOffset = Context.NumCachedStaticDirtyShaderBindings;
+					Context.NumCachedStaticDirtyShaderBindings += RTLODData.SBTAllocation->GetSegmentCount();
+				}
+				// - DirtyCachedRayTracingPrimitives are only processed after StaticPrimitiveIndices is filled
+				// so we can end up with primitives that should be skipped here
+				// - once we update flags of primitive with dirty raytracing state before `GatherRayTracingRelevantPrimitives_Parallel`
+				// we should replace this condition with an assert instead
+				else if (!EnumHasAnyFlags(Flags, ERayTracingPrimitiveFlags::Skip))
+				{
 #if DO_CHECK
-							if(bUsingReferenceBasedResidency)
-							{
-								FRayTracingGeometry* TargetRayTracingGeometry = SceneInfo->GetStaticRayTracingGeometry(LODIndex);
-								if (!ensure(TargetRayTracingGeometry != nullptr))
-								{
-									return;
-								}
+					if (bUsingReferenceBasedResidency)
+					{
+						FRayTracingGeometry* TargetRayTracingGeometry = SceneInfo->GetStaticRayTracingGeometry(LODIndex);
+						if (!ensure(TargetRayTracingGeometry != nullptr))
+						{
+							return;
+						}
 
-								ensure(((FRayTracingGeometryManager*)GRayTracingGeometryManager)->IsGeometryGroupReferenced(TargetRayTracingGeometry->GroupHandle));
+						ensure(((FRayTracingGeometryManager*)GRayTracingGeometryManager)->IsGeometryGroupReferenced(TargetRayTracingGeometry->GroupHandle));
 
-							}
+					}
 #endif
 
-							FRayTracingGeometry* RayTracingGeometry = SceneInfo->GetValidStaticRayTracingGeometry(LODIndex);
-
-							if (RayTracingGeometry == nullptr)
-							{
-								return;
-							}
+					FRayTracingGeometry* RayTracingGeometry = SceneInfo->GetValidStaticRayTracingGeometry(LODIndex);
 
-							check(RayTracingGeometry->LODIndex == LODIndex);
+					if (RayTracingGeometry == nullptr)
+					{
+						return;
+					}
 
-							// Sometimes LODIndex is out of range because it is clamped by ClampToFirstLOD, like the requested LOD is being streamed in and hasn't been available
-							// According to InitViews, we should hide the static mesh instance
-							if (SceneInfo->RayTracingLODData.IsValidIndex(LODIndex))
-							{
-								FRelevantPrimitive* RelevantPrimitive = new (Context.StaticPrimitives) FRelevantPrimitive();
-								RelevantPrimitive->PrimitiveIndex = PrimitiveIndex;
-								RelevantPrimitive->PersistentPrimitiveIndex = SceneInfo->GetPersistentIndex();
+					check(RayTracingGeometry->LODIndex == LODIndex);
 
-								RelevantPrimitive->LODIndex = LODIndex;
-								RelevantPrimitive->RayTracingGeometry = RayTracingGeometry;
+					// Sometimes LODIndex is out of range because it is clamped by ClampToFirstLOD, like the requested LOD is being streamed in and hasn't been available
+					// According to InitViews, we should hide the static mesh instance
+					if (SceneInfo->RayTracingLODData.IsValidIndex(LODIndex))
+					{
+						FRelevantPrimitive* RelevantPrimitive = new (Context.StaticPrimitives) FRelevantPrimitive();
+						RelevantPrimitive->PrimitiveIndex = PrimitiveIndex;
+						RelevantPrimitive->PersistentPrimitiveIndex = SceneInfo->GetPersistentIndex();
 
-								const FPrimitiveSceneInfo::FRayTracingLODData& RTLODData = SceneInfo->RayTracingLODData[LODIndex];
-								RelevantPrimitive->CachedMeshCommandFlags = RTLODData.CachedMeshCommandFlags;
-							}
-						}
-					});
+						RelevantPrimitive->LODIndex = LODIndex;
+						RelevantPrimitive->RayTracingGeometry = RayTracingGeometry;
 
-					if (Contexts.Num() > 0)
-					{
-						SCOPED_NAMED_EVENT(GatherRayTracingRelevantPrimitives_ComputeLOD_Merge, FColor::Emerald);
+						const FPrimitiveSceneInfo::FRayTracingLODData& RTLODData = SceneInfo->RayTracingLODData[LODIndex];
+						RelevantPrimitive->CachedMeshCommandFlags = RTLODData.CachedMeshCommandFlags;
+					}
+				}
+			});
 
-						uint32 NumStaticPrimitives = 0;
-						uint32 NumCachedStaticPrimitives = 0;
+		if (Contexts.Num() > 0)
+		{
+			SCOPED_NAMED_EVENT(GatherRayTracingRelevantStaticPrimitives__Merge, FColor::Emerald);
 
-						for (auto& Context : Contexts)
-						{
-							NumStaticPrimitives += Context.StaticPrimitives.Num();
-							NumCachedStaticPrimitives += Context.CachedStaticPrimitives.Num();
-						}
+			uint32 NumStaticPrimitives = 0;
+			uint32 NumCachedStaticPrimitives = 0;
 
-						TaskData.StaticPrimitives.Reserve(NumStaticPrimitives);
-						TaskData.CachedStaticPrimitives.Reserve(NumCachedStaticPrimitives);
+			for (auto& Context : Contexts)
+			{
+				NumStaticPrimitives += Context.StaticPrimitives.Num();
+				NumCachedStaticPrimitives += Context.CachedStaticPrimitives.Num();
+			}
 
-						TaskData.GatherContexts.SetNum(Contexts.Num());
+			TaskData.StaticPrimitives.Reserve(NumStaticPrimitives);
+			TaskData.CachedStaticPrimitives.Reserve(NumCachedStaticPrimitives);
 
-						for (int32 ContextIndex = 0; ContextIndex < Contexts.Num(); ++ContextIndex)
-						{
-							FRelevantStaticPrimitivesContext& Context = Contexts[ContextIndex];
-							FRelevantPrimitiveGatherContext& GatherContext = TaskData.GatherContexts[ContextIndex];
+			TaskData.GatherContexts.SetNum(Contexts.Num());
 
-							Context.StaticPrimitives.CopyToLinearArray(TaskData.StaticPrimitives);
-							Context.CachedStaticPrimitives.CopyToLinearArray(TaskData.CachedStaticPrimitives);
+			for (int32 ContextIndex = 0; ContextIndex < Contexts.Num(); ++ContextIndex)
+			{
+				FRelevantStaticPrimitivesContext& Context = Contexts[ContextIndex];
+				FRelevantPrimitiveGatherContext& GatherContext = TaskData.GatherContexts[ContextIndex];
 
-							GatherContext.InstanceOffset = TaskData.NumCachedStaticInstances;
-							GatherContext.DecalInstanceOffset = TaskData.NumCachedStaticDecalInstances;
-							GatherContext.DirtyShaderBindingsOffset = TaskData.NumCachedStaticDirtyShaderBindings;
+				Context.StaticPrimitives.CopyToLinearArray(TaskData.StaticPrimitives);
+				Context.CachedStaticPrimitives.CopyToLinearArray(TaskData.CachedStaticPrimitives);
 
-							TaskData.NumCachedStaticInstances += Context.NumCachedStaticInstances;
-							TaskData.NumCachedStaticDecalInstances += Context.NumCachedStaticDecalInstances;
-							TaskData.NumCachedStaticSegments += Context.NumCachedStaticSegments;
-							TaskData.NumCachedStaticDirtyShaderBindings += Context.NumCachedStaticDirtyShaderBindings;
+				GatherContext.InstanceOffset = TaskData.NumCachedStaticInstances;
+				GatherContext.DecalInstanceOffset = TaskData.NumCachedStaticDecalInstances;
+				GatherContext.DirtyShaderBindingsOffset = TaskData.NumCachedStaticDirtyShaderBindings;
 
-							for (const FPrimitiveSceneInfo* SceneInfo : Context.VisibleNaniteRayTracingPrimitives)
-							{
-								Nanite::GRayTracingManager.AddVisiblePrimitive(SceneInfo);
-							}
-						}
-					}
-			});
+				TaskData.NumCachedStaticInstances += Context.NumCachedStaticInstances;
+				TaskData.NumCachedStaticDecalInstances += Context.NumCachedStaticDecalInstances;
+				TaskData.NumCachedStaticSegments += Context.NumCachedStaticSegments;
+				TaskData.NumCachedStaticDirtyShaderBindings += Context.NumCachedStaticDirtyShaderBindings;
 
-		TaskData.bValid = true;
+				for (const FPrimitiveSceneInfo* SceneInfo : Context.VisibleNaniteRayTracingPrimitives)
+				{
+					Nanite::GRayTracingManager.AddVisiblePrimitive(SceneInfo);
+				}
+			}
+		}
 	}
 
 	struct FAutoInstanceBatch
@@ -1574,6 +1565,18 @@ namespace RayTracing
 				FTaskTagScope Scope(ETaskTag::EParallelRenderingThread);
 				GatherRelevantPrimitives(TaskData, bUsingReferenceBasedResidency);
 			}, UE::Tasks::Prerequisites(CacheRayTracingPrimitivesTask, FrustumCullTask), UE::Tasks::ETaskPriority::High);
+
+		static const auto ICVarStaticMeshLODDistanceScale = IConsoleManager::Get().FindConsoleVariable(TEXT("r.StaticMeshLODDistanceScale"));
+		const float LODScaleCVarValue = ICVarStaticMeshLODDistanceScale->GetFloat();
+		const int32 ForcedLODLevel = GetCVarForceLOD();
+
+		TaskData.GatherRelevantStaticPrimitivesTask = UE::Tasks::Launch(UE_SOURCE_LOCATION, [&TaskData, LODScaleCVarValue, ForcedLODLevel, bUsingReferenceBasedResidency]
+			{
+				FTaskTagScope Scope(ETaskTag::EParallelRenderingThread);
+				GatherRelevantStaticPrimitives(TaskData, LODScaleCVarValue, ForcedLODLevel, bUsingReferenceBasedResidency);
+			}, TaskData.GatherRelevantPrimitivesTask);
+
+		TaskData.bValid = true;
 	}
 
 	bool FinishGatherInstances(
@@ -1619,14 +1622,14 @@ namespace RayTracing
 			CoarseMeshSM->AddUsedStreamingHandles(RayTracingScene.UsedCoarseMeshStreamingHandles);
 		}
 
-		INC_DWORD_STAT_BY(STAT_VisibleRayTracingPrimitives, TaskData.StaticPrimitives.Num() + TaskData.CachedStaticPrimitives.Num() + TaskData.DynamicPrimitives.Num());
+		INC_DWORD_STAT_BY(STAT_VisibleRayTracingPrimitives, TaskData.StaticPrimitives.Num() + TaskData.CachedStaticPrimitives.Num() + TaskData.DynamicPrimitivesIndices.Num());
 
 		AddDynamicInstancesToRayTracingScene(
 			GraphBuilder,
 			Scene,
 			View,
 			TaskData.SceneOptions,
-			TaskData.DynamicPrimitives,
+			TaskData.DynamicPrimitivesIndices,
 			RayTracingScene,
 			RayTracingSBT,
 			InDynamicReadBuffer,
@@ -1660,7 +1663,7 @@ namespace RayTracing
 					RayTracingScene,
 					View.DirtyRayTracingShaderBindings);
 
-			}, TaskData.StaticPrimitiveLODTask);
+			}, TaskData.GatherRelevantStaticPrimitivesTask);
 
 		// Scene init task can run only when all pre-init tasks are complete (including culling tasks that are spawned while adding instances)
 		View.RayTracingSceneInitTask = UE::Tasks::Launch(UE_SOURCE_LOCATION,
