Moved light gather / sort and light atlas update to an async task. The light gather depends on the create shadows task, but isn't actually used until the light grid is computed. With the light atlas this can take anywhere from 0.1 - 0.25 ms.

diff --git a/Engine/Source/Runtime/Renderer/Private/DeferredShadingRenderer.cpp b/Engine/Source/Runtime/Renderer/Private/DeferredShadingRenderer.cpp
--- a/Engine/Source/Runtime/Renderer/Private/DeferredShadingRenderer.cpp
+++ b/Engine/Source/Runtime/Renderer/Private/DeferredShadingRenderer.cpp
@@ -1778,6 +1778,8 @@ void FDeferredShadingSceneRenderer::Render(FRDGBuilder& GraphBuilder)
 		}
 	}
 
+	UE::Tasks::TTask<FSortedLightSetSceneInfo*> GatherAndSortLightsTask;
+
 	if (RendererOutput == ERendererOutput::FinalSceneColor)
 	{
 #if RHI_RAYTRACING
@@ -1823,6 +1825,23 @@ void FDeferredShadingSceneRenderer::Render(FRDGBuilder& GraphBuilder)
 				ShouldRenderLocalFogVolume(Scene, ViewFamily)); // Needed when accessing forward light data for the directional light
 			bComputeLightGrid &= !ViewFamily.EngineShowFlags.PathTracing;
 		}
+
+		{
+			extern bool IsVSMOnePassProjectionEnabled(const FEngineShowFlags& ShowFlags);
+			extern UE::Tasks::FTask GetGatherAndSortLightsPrerequisiteTask(const FDynamicShadowsTaskData* TaskData);
+
+			auto* SortedLightSet = GraphBuilder.AllocObject<FSortedLightSetSceneInfo>();
+			const bool bShadowedLightsInClustered = ShouldUseClusteredDeferredShading()
+				&& IsVSMOnePassProjectionEnabled(ViewFamily.EngineShowFlags)
+				&& VirtualShadowMapArray.IsEnabled();
+
+			GatherAndSortLightsTask = LaunchSceneRenderTask<FSortedLightSetSceneInfo*>(UE_SOURCE_LOCATION, [this, SortedLightSet, bShadowedLightsInClustered]
+			{
+				GatherAndSortLights(*SortedLightSet, bShadowedLightsInClustered);
+				return SortedLightSet;
+
+			}, GetGatherAndSortLightsPrerequisiteTask(InitViewTaskDatas.DynamicShadows));
+		}
 	}
 
 	// force using occ queries for wireframe if rendering is parented or frozen in the first view
@@ -2265,23 +2284,24 @@ void FDeferredShadingSceneRenderer::Render(FRDGBuilder& GraphBuilder)
 		{
 			ViewExtension->PreRenderBasePass_RenderThread(GraphBuilder, ShouldRenderPrePass() /*bDepthBufferIsPopulated*/);
 		}
-	
-		// NOTE: The ordering of the lights is used to select sub-sets for different purposes, e.g., those that support clustered deferred.
-		FSortedLightSetSceneInfo& SortedLightSet = *GraphBuilder.AllocObject<FSortedLightSetSceneInfo>();
+
 		{
+			const FSortedLightSetSceneInfo* SortedLightSet = GatherAndSortLightsTask.GetResult();
+
 			RDG_CSV_STAT_EXCLUSIVE_SCOPE(GraphBuilder, SortLights);
 			RDG_EVENT_SCOPE_STAT(GraphBuilder, SortLights, "SortLights");
 			RDG_GPU_STAT_SCOPE(GraphBuilder, SortLights);
 
-			ComputeLightGridOutput = GatherLightsAndComputeLightGrid(GraphBuilder, bComputeLightGrid, SortedLightSet);
+			GatherAndSortLightsTask.Wait();
+			ComputeLightGridOutput = GatherLightsAndComputeLightGrid(GraphBuilder, bComputeLightGrid, *SortedLightSet);
+
+			CSV_CUSTOM_STAT(LightCount, All,  float(SortedLightSet->SortedLights.Num()), ECsvCustomStatOp::Set);
+			CSV_CUSTOM_STAT(LightCount, Batched, float(SortedLightSet->UnbatchedLightStart), ECsvCustomStatOp::Set);
+			CSV_CUSTOM_STAT(LightCount, Unbatched, float(SortedLightSet->SortedLights.Num()) - float(SortedLightSet->UnbatchedLightStart), ECsvCustomStatOp::Set);
 		}
 
 		LightFunctionAtlas.RenderLightFunctionAtlas(GraphBuilder, Views);
 
-		CSV_CUSTOM_STAT(LightCount, All,  float(SortedLightSet.SortedLights.Num()), ECsvCustomStatOp::Set);
-		CSV_CUSTOM_STAT(LightCount, Batched, float(SortedLightSet.UnbatchedLightStart), ECsvCustomStatOp::Set);
-		CSV_CUSTOM_STAT(LightCount, Unbatched, float(SortedLightSet.SortedLights.Num()) - float(SortedLightSet.UnbatchedLightStart), ECsvCustomStatOp::Set);
-
 		// Run before RenderSkyAtmosphereLookUpTables for cloud shadows to be valid.
 		InitVolumetricCloudsForViews(GraphBuilder, bShouldRenderVolumetricCloudBase, InstanceCullingManager);
 
@@ -2649,7 +2669,7 @@ void FDeferredShadingSceneRenderer::Render(FRDGBuilder& GraphBuilder)
 
 					FFrontLayerTranslucencyData FrontLayerTranslucencyData = RenderFrontLayerTranslucency(GraphBuilder, Views, SceneTextures, true /*VSM page marking*/);
 
-					VirtualShadowMapArray.BuildPageAllocations(GraphBuilder, GetActiveSceneTextures(), Views, SortedLightSet, VisibleLightInfos, SingleLayerWaterPrePassResult, FrontLayerTranslucencyData, FroxelRenderer);
+					VirtualShadowMapArray.BuildPageAllocations(GraphBuilder, GetActiveSceneTextures(), Views, *GatherAndSortLightsTask.GetResult(), VisibleLightInfos, SingleLayerWaterPrePassResult, FrontLayerTranslucencyData, FroxelRenderer);
 				}
 
 				RenderShadowDepthMaps(GraphBuilder, InitViewTaskDatas.DynamicShadows, InstanceCullingManager, ExternalAccessQueue);
@@ -2789,6 +2809,8 @@ void FDeferredShadingSceneRenderer::Render(FRDGBuilder& GraphBuilder)
 			}
 #endif
 
+			const FSortedLightSetSceneInfo& SortedLightSet = *GatherAndSortLightsTask.GetResult();
+
 			RenderLights(GraphBuilder, SceneTextures, TranslucencyLightingVolumeTextures, LightingChannelsTexture, SortedLightSet);
 
 			if (SortedLightSet.MegaLightsLightStart < SortedLightSet.SortedLights.Num())
@@ -2852,6 +2874,7 @@ void FDeferredShadingSceneRenderer::Render(FRDGBuilder& GraphBuilder)
 		}
 		else if (HairStrands::HasViewHairStrandsData(Views) && ViewFamily.EngineShowFlags.Lighting)
 		{
+			const FSortedLightSetSceneInfo& SortedLightSet = *GatherAndSortLightsTask.GetResult();
 			RenderLightsForHair(GraphBuilder, SceneTextures, SortedLightSet, ForwardScreenSpaceShadowMaskHairTexture, LightingChannelsTexture);
 			RenderDeferredReflectionsAndSkyLightingHair(GraphBuilder);
 		}
diff --git a/Engine/Source/Runtime/Renderer/Private/DeferredShadingRenderer.h b/Engine/Source/Runtime/Renderer/Private/DeferredShadingRenderer.h
--- a/Engine/Source/Runtime/Renderer/Private/DeferredShadingRenderer.h
+++ b/Engine/Source/Runtime/Renderer/Private/DeferredShadingRenderer.h
@@ -346,7 +346,7 @@ public:
 	 * Culls local lights and reflection probes to a grid in frustum space, builds one light list and grid per view in the current Views.  
 	 * Needed for forward shading or translucency using the Surface lighting mode, and clustered deferred shading. 
 	 */
-	FComputeLightGridOutput GatherLightsAndComputeLightGrid(FRDGBuilder& GraphBuilder, bool bNeedLightGrid, FSortedLightSetSceneInfo &SortedLightSet);
+	FComputeLightGridOutput GatherLightsAndComputeLightGrid(FRDGBuilder& GraphBuilder, bool bNeedLightGrid, const FSortedLightSetSceneInfo &SortedLightSet);
 
 	/** 
 	 * Debug light grid content on screen.
@@ -845,7 +845,7 @@ private:
 		FMinimalSceneTextures& SceneTextures,
 		const FTranslucencyLightingVolumeTextures& TranslucencyLightingVolumeTextures,
 		FRDGTextureRef LightingChannelsTexture,
-		FSortedLightSetSceneInfo& SortedLightSet);
+		const FSortedLightSetSceneInfo& SortedLightSet);
 
 	/** Render stationary light overlap as complexity to scene color. */
 	void RenderStationaryLightOverlap(
@@ -976,7 +976,7 @@ private:
 	void RenderLightsForHair(
 		FRDGBuilder& GraphBuilder,
 		const FMinimalSceneTextures& SceneTextures,
-		FSortedLightSetSceneInfo& SortedLightSet,
+		const FSortedLightSetSceneInfo& SortedLightSet,
 		FRDGTextureRef InScreenShadowMaskSubPixelTexture,
 		FRDGTextureRef LightingChannelsTexture);
 
diff --git a/Engine/Source/Runtime/Renderer/Private/LightGridInjection.cpp b/Engine/Source/Runtime/Renderer/Private/LightGridInjection.cpp
--- a/Engine/Source/Runtime/Renderer/Private/LightGridInjection.cpp
+++ b/Engine/Source/Runtime/Renderer/Private/LightGridInjection.cpp
@@ -107,8 +107,6 @@ FAutoConsoleVariableRef CVarLightCullingMaxDistanceOverride(
 	ECVF_RenderThreadSafe
 );
 
-extern bool IsVSMOnePassProjectionEnabled(const FEngineShowFlags& ShowFlags);
-
 bool ShouldVisualizeLightGrid()
 {
 	return GForwardLightGridDebug > 0;
@@ -509,7 +507,7 @@ static void CalculateRectLightCullingPlanes(const FRectLightSceneProxy* RectProx
 	check(OutPlanes.Num() == NUM_PLANES_PER_RECT_LIGHT);
 }
 
-FComputeLightGridOutput FSceneRenderer::ComputeLightGrid(FRDGBuilder& GraphBuilder, bool bCullLightsToGrid, FSortedLightSetSceneInfo& SortedLightSet)
+FComputeLightGridOutput FSceneRenderer::ComputeLightGrid(FRDGBuilder& GraphBuilder, bool bCullLightsToGrid, const FSortedLightSetSceneInfo& SortedLightSet)
 {
 	FComputeLightGridOutput Result = {};
 
@@ -1056,17 +1054,11 @@ FComputeLightGridOutput FSceneRenderer::ComputeLightGrid(FRDGBuilder& GraphBuild
 	return Result;
 }
 
-FComputeLightGridOutput FDeferredShadingSceneRenderer::GatherLightsAndComputeLightGrid(FRDGBuilder& GraphBuilder, bool bNeedLightGrid, FSortedLightSetSceneInfo& SortedLightSet)
+FComputeLightGridOutput FDeferredShadingSceneRenderer::GatherLightsAndComputeLightGrid(FRDGBuilder& GraphBuilder, bool bNeedLightGrid, const FSortedLightSetSceneInfo& SortedLightSet)
 {
 	SCOPED_NAMED_EVENT(GatherLightsAndComputeLightGrid, FColor::Emerald);
 	FComputeLightGridOutput Result = {};
 
-	bool bShadowedLightsInClustered = ShouldUseClusteredDeferredShading()
-		&& IsVSMOnePassProjectionEnabled(ViewFamily.EngineShowFlags)
-		&& VirtualShadowMapArray.IsEnabled();
-
-	GatherAndSortLights(SortedLightSet, bShadowedLightsInClustered);
-	
 	if (!bNeedLightGrid)
 	{
 		SetDummyForwardLightUniformBufferOnViews(GraphBuilder, ShaderPlatform, Views);
diff --git a/Engine/Source/Runtime/Renderer/Private/LightRendering.cpp b/Engine/Source/Runtime/Renderer/Private/LightRendering.cpp
--- a/Engine/Source/Runtime/Renderer/Private/LightRendering.cpp
+++ b/Engine/Source/Runtime/Renderer/Private/LightRendering.cpp
@@ -1193,6 +1193,8 @@ bool CanLightUsesAtlasForUnbatchedLight(ERHIFeatureLevel::Type FeatureLevel, con
 
 void FSceneRenderer::GatherAndSortLights(FSortedLightSetSceneInfo& OutSortedLights, bool bShadowedLightsInClustered)
 {
+	TRACE_CPUPROFILER_EVENT_SCOPE(GatherAndSortLights);
+
 	if (AllowSimpleLights())
 	{
 		GatherSimpleLights(ViewFamily, Views, OutSortedLights.SimpleLights);
@@ -1385,7 +1387,7 @@ void FDeferredShadingSceneRenderer::RenderLights(
 	FMinimalSceneTextures& SceneTextures,
 	const FTranslucencyLightingVolumeTextures& TranslucencyLightingVolumeTextures,
 	FRDGTextureRef LightingChannelsTexture,
-	FSortedLightSetSceneInfo& SortedLightSet)
+	const FSortedLightSetSceneInfo& SortedLightSet)
 {
 	const bool bUseHairLighting = HairStrands::HasViewHairStrandsData(Views);
 #if RHI_RAYTRACING
@@ -2957,7 +2959,7 @@ void FDeferredShadingSceneRenderer::RenderLightForHair(
 void FDeferredShadingSceneRenderer::RenderLightsForHair(
 	FRDGBuilder& GraphBuilder,
 	const FMinimalSceneTextures& SceneTextures,
-	FSortedLightSetSceneInfo& SortedLightSet,
+	const FSortedLightSetSceneInfo& SortedLightSet,
 	FRDGTextureRef ScreenShadowMaskSubPixelTexture,
 	FRDGTextureRef LightingChannelsTexture)
 {
diff --git a/Engine/Source/Runtime/Renderer/Private/SceneRendering.h b/Engine/Source/Runtime/Renderer/Private/SceneRendering.h
--- a/Engine/Source/Runtime/Renderer/Private/SceneRendering.h
+++ b/Engine/Source/Runtime/Renderer/Private/SceneRendering.h
@@ -2492,7 +2492,7 @@ protected:
 	 * Culls local lights and reflection probes to a grid in frustum space, builds one light list and grid per view in the current Views.  
 	 * Needed for forward shading or translucency using the Surface lighting mode, and clustered deferred shading. 
 	 */
-	FComputeLightGridOutput ComputeLightGrid(FRDGBuilder& GraphBuilder, bool bCullLightsToGrid, FSortedLightSetSceneInfo& SortedLightSet);
+	FComputeLightGridOutput ComputeLightGrid(FRDGBuilder& GraphBuilder, bool bCullLightsToGrid, const FSortedLightSetSceneInfo& SortedLightSet);
 
 	/**
 	* Used by RenderLights to figure out if light functions need to be rendered to the attenuation buffer.
@@ -2573,6 +2573,22 @@ UE::Tasks::FTask LaunchSceneRenderTask(const TCHAR* DebugName, LambdaType&& Lamb
 	);
 }
 
+template <typename ReturnType, typename LambdaType, typename PrerequisiteTaskCollectionType>
+UE::Tasks::TTask<ReturnType> LaunchSceneRenderTask(const TCHAR* DebugName, LambdaType&& Lambda, PrerequisiteTaskCollectionType&& Prerequisites, bool bExecuteInParallelCondition = true, UE::Tasks::ETaskPriority TaskPriority = UE::Tasks::ETaskPriority::High)
+{
+	const bool bExecuteInParallel = bExecuteInParallelCondition && FApp::ShouldUseThreadingForPerformance() && GIsThreadedRendering;
+
+	return UE::Tasks::Launch(DebugName, [Lambda = Forward<LambdaType&&>(Lambda)] () mutable
+		{
+			FOptionalTaskTagScope Scope(ETaskTag::EParallelRenderingThread);
+			return Lambda();
+		},
+		Forward<PrerequisiteTaskCollectionType&&>(Prerequisites),
+		TaskPriority,
+		bExecuteInParallel ? UE::Tasks::EExtendedTaskPriority::None : UE::Tasks::EExtendedTaskPriority::Inline
+	);
+}
+
 // Creates a FGraphEventRef from UE::Task::FTask prerequisites.
 template <typename PrerequisiteTaskCollectionType>
 FGraphEventRef CreateCompatibilityGraphEvent(PrerequisiteTaskCollectionType&& Prerequisites)
diff --git a/Engine/Source/Runtime/Renderer/Private/ShadowSetup.cpp b/Engine/Source/Runtime/Renderer/Private/ShadowSetup.cpp
--- a/Engine/Source/Runtime/Renderer/Private/ShadowSetup.cpp
+++ b/Engine/Source/Runtime/Renderer/Private/ShadowSetup.cpp
@@ -1362,6 +1362,7 @@ struct FDynamicShadowsTaskData
 	TArray<struct FGatherShadowPrimitivesPacket*, SceneRenderingAllocator> Packets;
 	FPerShadowGatherStats GatherStats;
 	FFilteredShadowArrays ShadowArrays;
+	UE::Tasks::FTaskEvent BeginGatherAndSortLightsTask{ UE_SOURCE_LOCATION };
 	UE::Tasks::FTaskEvent FilterDynamicShadowsTask{ UE_SOURCE_LOCATION };
 
 	// Gather Dynamic Mesh Elements state
@@ -1440,6 +1441,15 @@ struct FDynamicShadowsTaskData
 	}
 };
 
+UE::Tasks::FTask GetGatherAndSortLightsPrerequisiteTask(const FDynamicShadowsTaskData* TaskData)
+{
+	if (TaskData)
+	{
+		return TaskData->BeginGatherAndSortLightsTask;
+	}
+	return {};
+}
+
 void BeginShadowGatherDynamicMeshElements(FDynamicShadowsTaskData* TaskData)
 {
 	if (TaskData)
@@ -6327,6 +6337,9 @@ void FSceneRenderer::CreateDynamicShadows(FDynamicShadowsTaskData& TaskData)
 
 		// Calculate visibility of the projected shadows.
 		InitProjectedShadowVisibility(TaskData);
+
+		// At this point the shadows are initialized and we can start gathering lights for shading.
+		TaskData.BeginGatherAndSortLightsTask.Trigger();
 	}
 
 	// Clear old preshadows and attempt to add new ones to the cache
