Remove compilation of anisotropy shader when substrate is enabled.

diff --git a/Engine/Source/Runtime/Renderer/Private/AnisotropyRendering.cpp b/Engine/Source/Runtime/Renderer/Private/AnisotropyRendering.cpp
--- a/Engine/Source/Runtime/Renderer/Private/AnisotropyRendering.cpp
+++ b/Engine/Source/Runtime/Renderer/Private/AnisotropyRendering.cpp
@@ -22,13 +22,15 @@ bool SupportsAnisotropicMaterials(ERHIFeatureLevel::Type FeatureLevel, EShaderPl
 {
 	return GAnisotropicMaterials
 		&& FeatureLevel >= ERHIFeatureLevel::SM5
-		&& FDataDrivenShaderPlatformInfo::GetSupportsAnisotropicMaterials(ShaderPlatform);
+		&& FDataDrivenShaderPlatformInfo::GetSupportsAnisotropicMaterials(ShaderPlatform)
+		&& !Substrate::IsSubstrateEnabled(); // Substrate renders anisotropy surface natively, without extra pass.;
 }
 
 static bool IsAnisotropyPassCompatible(const EShaderPlatform Platform, FMaterialShaderParameters MaterialParameters)
 {
 	return 
 		FDataDrivenShaderPlatformInfo::GetSupportsAnisotropicMaterials(Platform) &&
+		!Substrate::IsSubstrateEnabled() && // Substrate renders anisotropy surface natively, without extra pass.
 		MaterialParameters.bHasAnisotropyConnected &&
 		!IsTranslucentBlendMode(MaterialParameters) &&
 		MaterialParameters.ShadingModels.HasAnyShadingModel({ MSM_DefaultLit, MSM_ClearCoat });
