Substrate - fixed clouds not compositing correctly with SLWater.

#rnx
#tests If we have test for water over cloud then there should not be anything to do.

diff --git a/Engine/Shaders/Private/VolumetricRenderTarget.usf b/Engine/Shaders/Private/VolumetricRenderTarget.usf
--- a/Engine/Shaders/Private/VolumetricRenderTarget.usf
+++ b/Engine/Shaders/Private/VolumetricRenderTarget.usf
@@ -9,6 +9,8 @@
 
 #include "SceneTextureParameters.ush"
 
+#include "Substrate/Substrate.ush"
+
 #ifndef CLOUD_MIN_AND_MAX_DEPTH
 #define CLOUD_MIN_AND_MAX_DEPTH 0
 #endif
@@ -865,8 +867,14 @@ void ComposeVolumetricRTOverScenePS(
 
 	// We only want to run the special "cloud over water" code path on actual water pixels, otherwise depth threshold can cause visual issues.
 	// However, with Forward Shading, we do not have access to the gbuffer: so we still run that simple water depth test to occlude cloud on forrground meshes. Not the perfect but it will work in most cases.
+#if SUBSTRATE_ENABLED
+	FSubstrateAddressing SubstrateAddressing = GetSubstratePixelDataByteOffset(PixelPos, uint2(View.BufferSizeAndInvSize.xy), Substrate.MaxBytesPerPixel);
+	FSubstratePixelHeader SubstratePixelHeader = UnpackSubstrateHeaderIn(Substrate.MaterialTextureArray, SubstrateAddressing, Substrate.TopLayerTexture);
+	const bool bIsWaterPixel = SubstratePixelHeader.IsSingleLayerWater(); 
+#else
 	FScreenSpaceData ScreenSpaceData = GetScreenSpaceData(ScreenUV);
 	const bool bIsWaterPixel = ForwardShadingEnable || ScreenSpaceData.GBuffer.ShadingModelID == SHADINGMODELID_SINGLELAYERWATER;
+#endif
 
 	uint2 VolumeCoord0 = max(0, int2(VolumeCoordUInt) + int2(0, 0));
 	uint2 VolumeCoord1 = max(0, int2(VolumeCoordUInt) + int2(OffsetX, 0));
diff --git a/Engine/Source/Runtime/Renderer/Private/VolumetricRenderTarget.cpp b/Engine/Source/Runtime/Renderer/Private/VolumetricRenderTarget.cpp
--- a/Engine/Source/Runtime/Renderer/Private/VolumetricRenderTarget.cpp
+++ b/Engine/Source/Runtime/Renderer/Private/VolumetricRenderTarget.cpp
@@ -704,6 +704,7 @@ class FComposeVolumetricRTOverScenePS : public FGlobalShader
 		SHADER_PARAMETER(FUintVector4, VolumetricTextureValidCoordRect)
 		SHADER_PARAMETER(FVector4f, VolumetricTextureValidUvRect)
 		SHADER_PARAMETER_RDG_UNIFORM_BUFFER(FSceneTextureUniformParameters, SceneTextures)
+		SHADER_PARAMETER_RDG_UNIFORM_BUFFER(FSubstrateGlobalUniformParameters, Substrate)
 	END_SHADER_PARAMETER_STRUCT()
 
 	static FPermutationDomain RemapPermutation(FPermutationDomain PermutationVector)
@@ -846,6 +847,7 @@ void ComposeVolumetricRenderTargetOverScene(
 		PassParameters->UvOffsetSampleAcceptanceWeight = GetUvNoiseSampleAcceptanceWeight();
 		PassParameters->FullResolutionToVolumetricBufferResolutionScale = GetCompositionFullResolutionToVolumetricBufferResolutionScale(VRTMode) * GetVolumetricBufferResolutionScale(VRTMode);
 		PassParameters->SceneTextures = SceneTextures.UniformBuffer;
+		PassParameters->Substrate = Substrate::BindSubstrateGlobalUniformParameters(ViewInfo);
 		PassParameters->ForwardShadingEnable = bForwardShading ? 1 : 0;
 		GetTextureSafeUvCoordBound(PassParameters->VolumetricTexture, PassParameters->VolumetricTextureValidCoordRect, PassParameters->VolumetricTextureValidUvRect);
 
@@ -930,6 +932,7 @@ void ComposeVolumetricRenderTargetOverSceneUnderWater(
 		PassParameters->SceneWithoutSingleLayerWaterViewRect = FVector4f(WaterPassViewData.ViewRect.Min.X, WaterPassViewData.ViewRect.Min.Y,
 																		WaterPassViewData.ViewRect.Max.X, WaterPassViewData.ViewRect.Max.Y);
 		PassParameters->SceneTextures = SceneTextures.UniformBuffer;
+		PassParameters->Substrate = Substrate::BindSubstrateGlobalUniformParameters(ViewInfo);
 		GetTextureSafeUvCoordBound(PassParameters->VolumetricTexture, PassParameters->VolumetricTextureValidCoordRect, PassParameters->VolumetricTextureValidUvRect);
 
 		FVector2D VolumetricTextureSize = FVector2D(float(VolumetricTexture->Desc.GetSize().X), float(VolumetricTexture->Desc.GetSize().Y));
@@ -991,6 +994,7 @@ void ComposeVolumetricRenderTargetOverSceneForVisualization(
 		PassParameters->UvOffsetSampleAcceptanceWeight = GetUvNoiseSampleAcceptanceWeight();
 		PassParameters->FullResolutionToVolumetricBufferResolutionScale = GetCompositionFullResolutionToVolumetricBufferResolutionScale(VRTMode) * GetVolumetricBufferResolutionScale(VRTMode);
 		PassParameters->SceneTextures = SceneTextures.UniformBuffer;
+		PassParameters->Substrate = Substrate::BindSubstrateGlobalUniformParameters(ViewInfo);
 		GetTextureSafeUvCoordBound(PassParameters->VolumetricTexture, PassParameters->VolumetricTextureValidCoordRect, PassParameters->VolumetricTextureValidUvRect);
 
 		PassParameters->WaterLinearDepthTexture = GSystemTextures.GetBlackDummy(GraphBuilder);
