Added movable point lights dynamic shadows in mobile forward renderer

diff --git a/Engine/Config/Android/AndroidEngine.ini b/Engine/Config/Android/AndroidEngine.ini
--- a/Engine/Config/Android/AndroidEngine.ini
+++ b/Engine/Config/Android/AndroidEngine.ini
@@ -89,6 +89,7 @@ r.DistanceFields=1
 r.TemporalAA.Mobile.UseCompute=0
 ; Disables parallel GDME by default due to not enough threads
 r.Visibility.DynamicMeshElements.Parallel=0
+r.Shadow.DetectVertexShaderLayerAtRuntime=1
 
 [AlternateTextureCompression]
 ;disabling TextureCompressionFormat for Android/IOS no longer needed
diff --git a/Engine/Shaders/Private/MobileLightingCommon.ush b/Engine/Shaders/Private/MobileLightingCommon.ush
--- a/Engine/Shaders/Private/MobileLightingCommon.ush
+++ b/Engine/Shaders/Private/MobileLightingCommon.ush
@@ -613,7 +613,7 @@ void AccumulateLightGridLocalLighting(const FCulledLightsGridHeader InLightGridH
 
 		half DynamicShadowing = 1.0f;
 #if USE_SHADOWMASKTEXTURE
-		if (LightType == LIGHT_TYPE_SPOT)
+		if (LightType == LIGHT_TYPE_SPOT || LightType == LIGHT_TYPE_POINT)
 		{
 			half4 PreviewShadowMapChannelMask = UnpackShadowMapChannelMask(LightTypeAndPackedShadowMapChannelMask >> 4);
 			DynamicShadowing = dot(PreviewShadowMapChannelMask, DynamicShadowFactors);
diff --git a/Engine/Shaders/Private/ScreenVertexShader.usf b/Engine/Shaders/Private/ScreenVertexShader.usf
--- a/Engine/Shaders/Private/ScreenVertexShader.usf
+++ b/Engine/Shaders/Private/ScreenVertexShader.usf
@@ -15,7 +15,7 @@ void Main(
 	DrawRectangle( float4( InPosition, 0, 1 ), InUV, Output.Position, Output.UV);
 }
 
-#if USING_VERTEX_SHADER_LAYER && USING_LAYERS
+#if USING_LAYERS
 struct FScreenVertexLayeredOutput
 {
 	FScreenVertexOutput Output;
@@ -29,7 +29,7 @@ struct FScreenVertexLayeredOutput
 void MainForGS(
 	float2 InPosition : ATTRIBUTE0,
 	float2 InUV       : ATTRIBUTE1,
-#if USING_VERTEX_SHADER_LAYER && USING_LAYERS
+#if USING_LAYERS
 	in uint InstanceId	: SV_InstanceID,
 	out FScreenVertexLayeredOutput GSOutput
 #else
@@ -37,7 +37,7 @@ void MainForGS(
 #endif
 	)
 {
-#if USING_VERTEX_SHADER_LAYER && USING_LAYERS
+#if USING_LAYERS
 	GSOutput.RTIndex = InstanceId;
 	GSOutput.FaceIndex = InstanceId;
 
diff --git a/Engine/Shaders/Private/ShadowProjectionCommon.ush b/Engine/Shaders/Private/ShadowProjectionCommon.ush
--- a/Engine/Shaders/Private/ShadowProjectionCommon.ush
+++ b/Engine/Shaders/Private/ShadowProjectionCommon.ush
@@ -11,8 +11,6 @@ float3 SoftTransitionScale;
 // xy:ShadowTexelSize.xy, zw:1/ShadowTexelSize.xy
 float4 ShadowBufferSize;
 
-#if FEATURE_LEVEL >= FEATURE_LEVEL_SM4
-
 /** Cube map texture. */
 TextureCube ShadowDepthCubeTexture;
 
@@ -262,8 +260,6 @@ float CubemapHardwarePCF(
 		WorldPosition, LightPosition, LightInvRadius, DepthBias, SlopeDepthBias, MaxSlopeDepthBias);
 }
 
-#endif
-
 #ifndef NO_TRANSLUCENCY_AVAILABLE
 
 /** Textures and parameters needed for Fourier opacity maps. */
diff --git a/Engine/Shaders/Private/ShadowProjectionPixelShader.usf b/Engine/Shaders/Private/ShadowProjectionPixelShader.usf
--- a/Engine/Shaders/Private/ShadowProjectionPixelShader.usf
+++ b/Engine/Shaders/Private/ShadowProjectionPixelShader.usf
@@ -420,8 +420,6 @@ void Main(
 #endif
 }
 
-#if FEATURE_LEVEL >= FEATURE_LEVEL_SM4
-
 // .x:DepthBias, y: SlopeDepthBias, z: MaxSlopeDepthBias, 
 float3 PointLightDepthBias;
 
@@ -489,7 +487,7 @@ void MainOnePassPointLightPS(
 	// Unlike spot/rect/direction light, point lights performs depth bias only at projection time. 
 	// This is why the slope bias drives the normal receiver bias here.
 	float SlopeBias = 0;
-#if !SUBPIXEL_SHADOW
+#if !SUBPIXEL_SHADOW && !(FORWARD_SHADING && SHADING_PATH_MOBILE)
 	{
 #if SUBSTRATE_ENABLED
 		const FSubstrateTopLayerData TopLayerData = SubstrateUnpackTopLayerData(Substrate.TopLayerTexture.Load(uint3(SVPos.xy, 0)));
@@ -572,6 +570,4 @@ void MainOnePassPointLightPS(
 	}
 	OutColor.a = IsSubsurfaceProfile ? EncodeLightAttenuation(Thickness) : OutColor.b;
 #endif // USE_TRANSMISSION
-}
- 
-#endif
+}
\ No newline at end of file
diff --git a/Engine/Source/Runtime/Engine/Private/Components/PointLightComponent.cpp b/Engine/Source/Runtime/Engine/Private/Components/PointLightComponent.cpp
--- a/Engine/Source/Runtime/Engine/Private/Components/PointLightComponent.cpp
+++ b/Engine/Source/Runtime/Engine/Private/Components/PointLightComponent.cpp
@@ -71,7 +71,7 @@ void FPointLightSceneProxy::GetLightShaderParameters(FLightRenderParameters& Lig
 */
 bool FPointLightSceneProxy::GetWholeSceneProjectedShadowInitializer(const FSceneViewFamily& ViewFamily, TArray<FWholeSceneProjectedShadowInitializer, TInlineAllocator<6> >& OutInitializers) const
 {
-	if (ViewFamily.GetFeatureLevel() >= ERHIFeatureLevel::SM5
+	if ((ViewFamily.GetFeatureLevel() >= ERHIFeatureLevel::SM5 || DoesRuntimeSupportOnePassPointLightShadows(ViewFamily.GetShaderPlatform()))
 		&& GAllowPointLightCubemapShadows != 0)
 	{
 		FWholeSceneProjectedShadowInitializer& OutInitializer = *new(OutInitializers) FWholeSceneProjectedShadowInitializer;
diff --git a/Engine/Source/Runtime/Engine/Private/LightSceneProxy.cpp b/Engine/Source/Runtime/Engine/Private/LightSceneProxy.cpp
--- a/Engine/Source/Runtime/Engine/Private/LightSceneProxy.cpp
+++ b/Engine/Source/Runtime/Engine/Private/LightSceneProxy.cpp
@@ -131,7 +131,7 @@ FLightSceneProxy::FLightSceneProxy(const ULightComponent* InLightComponent)
 
 	if (bCastDynamicShadow && IsMobilePlatform(SceneInterface->GetShaderPlatform()))
 	{
-		if (GetLightType() == LightType_Point
+		if ((GetLightType() == LightType_Point && !DoesRuntimeSupportOnePassPointLightShadows(SceneInterface->GetShaderPlatform()))
 			|| GetLightType() == LightType_Rect
 			|| (GetLightType() == LightType_Spot && !IsMobileMovableSpotlightShadowsEnabled(SceneInterface->GetShaderPlatform())))
 		{
diff --git a/Engine/Source/Runtime/RenderCore/Private/RenderUtils.cpp b/Engine/Source/Runtime/RenderCore/Private/RenderUtils.cpp
--- a/Engine/Source/Runtime/RenderCore/Private/RenderUtils.cpp
+++ b/Engine/Source/Runtime/RenderCore/Private/RenderUtils.cpp
@@ -566,9 +566,10 @@ RENDERCORE_API bool MobileForwardEnableClusteredReflections(const FStaticShaderP
 
 RENDERCORE_API bool MobileUsesShadowMaskTexture(const FStaticShaderPlatform Platform)
 {
+	const bool bMoveablePointOrSpotLightShadowsEnabled = IsMobileMovableSpotlightShadowsEnabled(Platform) || FReadOnlyCVARCache::EnablePointLightShadows(Platform);
 	// Only distance field shadow needs to render shadow mask texture on mobile deferred, normal shadows need to be rendered separately because of handling lighting channels.
 	// Besides distance field shadow, with clustered lighting and shadow of local light enabled, shadows will render to shadow mask texture on mobile forward, lighting channels are handled in base pass shader.
-	return IsMobileDistanceFieldEnabled(Platform) || (!IsMobileDeferredShadingEnabled(Platform) && IsMobileMovableSpotlightShadowsEnabled(Platform) && MobileForwardEnableLocalLights(Platform));
+	return IsMobileDistanceFieldEnabled(Platform) || (!IsMobileDeferredShadingEnabled(Platform) && bMoveablePointOrSpotLightShadowsEnabled && MobileForwardEnableLocalLights(Platform));
 }
 
 // Whether to support more than 4 color attachments for GBuffer 
@@ -1439,8 +1440,8 @@ RENDERCORE_API bool DoesRuntimeSupportOnePassPointLightShadows(EShaderPlatform P
 {
 	static const auto CVar = IConsoleManager::Get().FindTConsoleVariableDataInt(TEXT("r.Shadow.DetectVertexShaderLayerAtRuntime"));
 
-	return RHISupportsVertexShaderLayer(Platform)
-		|| (CVar->GetValueOnAnyThread() != 0 && GRHISupportsArrayIndexFromAnyShader != 0);
+	return (RHISupportsVertexShaderLayer(Platform)
+		|| (CVar->GetValueOnAnyThread() != 0 && GRHISupportsArrayIndexFromAnyShader != 0));
 }
 
 bool IsForwardShadingEnabled(const FStaticShaderPlatform Platform)
diff --git a/Engine/Source/Runtime/RenderCore/Public/ReadOnlyCVARCache.h b/Engine/Source/Runtime/RenderCore/Public/ReadOnlyCVARCache.h
--- a/Engine/Source/Runtime/RenderCore/Public/ReadOnlyCVARCache.h
+++ b/Engine/Source/Runtime/RenderCore/Public/ReadOnlyCVARCache.h
@@ -6,6 +6,9 @@
 
 #pragma once
 #include "RHIShaderPlatform.h"
+#include "DataDrivenShaderPlatformInfo.h"
+#include "ShaderPlatformCachedIniValue.h"
+
 
 struct FReadOnlyCVARCache
 {
@@ -21,9 +24,17 @@ struct FReadOnlyCVARCache
 		#endif	
 	}
 
-	static inline bool EnablePointLightShadows()
+	static inline bool EnablePointLightShadows(const FStaticShaderPlatform Platform)
 	{
-		return bEnablePointLightShadows;
+		if (!IsMobilePlatform(Platform))
+		{
+			return bEnablePointLightShadows;
+		}
+		else
+		{
+			static FShaderPlatformCachedIniValue<bool> MobileMovablePointLightShadowsIniValue(TEXT("r.Mobile.EnableMovablePointLightsShadows"));
+			return MobileMovablePointLightShadowsIniValue.Get(Platform) && FReadOnlyCVARCache::MobileSupportsGPUScene();
+		}
 	}
 	
 	static inline bool EnableStationarySkylight()
diff --git a/Engine/Source/Runtime/Renderer/Private/MobileBasePassRendering.cpp b/Engine/Source/Runtime/Renderer/Private/MobileBasePassRendering.cpp
--- a/Engine/Source/Runtime/Renderer/Private/MobileBasePassRendering.cpp
+++ b/Engine/Source/Runtime/Renderer/Private/MobileBasePassRendering.cpp
@@ -39,9 +39,15 @@ static TAutoConsoleVariable<int32> CVarMobileEnableMovableSpotLightShadows(
 static TAutoConsoleVariable<int32> CVarMobileMaxVisibleMovableSpotLightShadows(
 	TEXT("r.Mobile.MaxVisibleMovableSpotLightShadows"),
 	8,
-	TEXT("The max number of visible spotlighs can cast shadow sorted by screen size, should be as less as possible for performance reason"),
+	TEXT("The max number of visible spotlights can cast shadow sorted by screen size, should be as less as possible for performance reason"),
 	ECVF_RenderThreadSafe);
 
+static TAutoConsoleVariable<int32> CVarMobileEnableMovablePointLightShadows(
+	TEXT("r.Mobile.EnableMovablePointLightsShadows"),
+	0,
+	TEXT("If 1 then enable movable point light shadow support"),
+	ECVF_ReadOnly | ECVF_RenderThreadSafe);
+
 static TAutoConsoleVariable<int32> CVarMobileSceneDepthAux(
 	TEXT("r.Mobile.SceneDepthAux"),
 	1,
diff --git a/Engine/Source/Runtime/Renderer/Private/SceneRendering.cpp b/Engine/Source/Runtime/Renderer/Private/SceneRendering.cpp
--- a/Engine/Source/Runtime/Renderer/Private/SceneRendering.cpp
+++ b/Engine/Source/Runtime/Renderer/Private/SceneRendering.cpp
@@ -4026,7 +4026,7 @@ void FSceneRenderer::OnRenderFinish(FRDGBuilder& GraphBuilder, FRDGTextureRef Vi
 
 		const bool bRealTimeSkyCaptureButNothingToCapture = Scene->SkyLight && Scene->SkyLight->bRealTimeCaptureEnabled && (!Scene->HasSkyAtmosphere() && !Scene->HasVolumetricCloud() && (Views.Num() > 0 && !Views[0].bSceneHasSkyMaterial));
 
-		const bool bShowPointLightWarning = UsedWholeScenePointLightNames.Num() > 0 && !FReadOnlyCVARCache::EnablePointLightShadows();
+		const bool bShowPointLightWarning = UsedWholeScenePointLightNames.Num() > 0 && !FReadOnlyCVARCache::EnablePointLightShadows(ShaderPlatform);
 		const bool bShowShadowedLightOverflowWarning = Scene->OverflowingDynamicShadowedLights.Num() > 0;
 
 		const bool bLocalFogVolumeInSceneButProjectDisabled = Scene->HasAnyLocalFogVolume() && !ProjectSupportsLocalFogVolumes();
diff --git a/Engine/Source/Runtime/Renderer/Private/ShadowDepthRendering.cpp b/Engine/Source/Runtime/Renderer/Private/ShadowDepthRendering.cpp
--- a/Engine/Source/Runtime/Renderer/Private/ShadowDepthRendering.cpp
+++ b/Engine/Source/Runtime/Renderer/Private/ShadowDepthRendering.cpp
@@ -93,7 +93,11 @@ public:
 
 	static bool ShouldCompilePermutation(const FGlobalShaderPermutationParameters& Parameters)
 	{
-		return IsFeatureLevelSupported(Parameters.Platform, ERHIFeatureLevel::SM5) && (!bUsingVertexLayers || (!RHISupportsGeometryShaders(Parameters.Platform) && RHISupportsVertexShaderLayer(Parameters.Platform)));
+		if (IsMobilePlatform(Parameters.Platform))
+		{
+			return bUsingVertexLayers;
+		}
+		return true;
 	}
 
 	static void ModifyCompilationEnvironment(const FGlobalShaderPermutationParameters& Parameters, FShaderCompilerEnvironment& OutEnvironment)
@@ -212,6 +216,18 @@ void SetupShadowDepthPassUniformBuffer(
 
 	ShadowDepthPassParameters.ShadowParams = FVector4f(ShadowInfo->GetShaderDepthBias(), ShadowInfo->GetShaderSlopeDepthBias(), ShadowInfo->GetShaderMaxSlopeDepthBias(), ShadowInfo->InvMaxSubjectDepth);
 	ShadowDepthPassParameters.bClampToNearPlane = ShadowInfo->ShouldClampToNearPlane() ? 1.0f : 0.0f;
+
+	if (ShadowInfo->bOnePassPointLightShadow)
+	{
+		check(ShadowInfo->BorderSize == 0);
+
+		// offset from translated world space to (pre translated) shadow space 
+		const FMatrix Translation = FTranslationMatrix(ShadowInfo->PreShadowTranslation - View.ViewMatrices.GetPreViewTranslation());
+		for (int32 FaceIndex = 0; FaceIndex < 6; FaceIndex++)
+		{
+			ShadowDepthPassParameters.ShadowViewProjectionMatrices[FaceIndex] = FMatrix44f(Translation * ShadowInfo->OnePassShadowViewProjectionMatrices[FaceIndex]);		// LWC_TODO: Precision loss?
+		}
+	}
 }
 
 void AddClearShadowDepthPass(FRDGBuilder& GraphBuilder, FRDGTextureRef Texture)
@@ -329,10 +345,11 @@ public:
 		const bool bForceAllPermutations = SupportAllShaderPermutationsVar && SupportAllShaderPermutationsVar->GetValueOnAnyThread() != 0;
 		const bool bSupportPointLightWholeSceneShadows = CVarSupportPointLightWholeSceneShadows.GetValueOnAnyThread() != 0 || bForceAllPermutations;
 
-		// Mobile only needs OutputDepth, and optionally PerspectiveCorrect
+		// Mobile only needs OutputDepth, and optionally PerspectiveCorrect and OnePassPointLight
 		if (!IsFeatureLevelSupported(Platform, ERHIFeatureLevel::SM5) &&
 			!(ShaderMode == VertexShadowDepth_OutputDepth || 
-				(ShaderMode == VertexShadowDepth_PerspectiveCorrect && MobileUsesPerspectiveCorrectShadowPermutation(Platform))))
+				(ShaderMode == VertexShadowDepth_PerspectiveCorrect && MobileUsesPerspectiveCorrectShadowPermutation(Platform))
+					 || ((ShaderMode == VertexShadowDepth_OnePassPointLight) && FReadOnlyCVARCache::EnablePointLightShadows(Platform))))
 		{
 			return false;
 		}
@@ -355,8 +372,6 @@ public:
 		return (Parameters.MaterialParameters.bIsSpecialEngineMaterial
 			// Masked and WPO materials need their shaders but cannot be used with a position only stream.
 			|| ((!Parameters.MaterialParameters.bWritesEveryPixelShadowPass || Parameters.MaterialParameters.bMaterialMayModifyMeshPosition) && !bUsePositionOnlyStream))
-			// Only compile one pass point light shaders for feature levels >= SM5
-			&& (ShaderMode != VertexShadowDepth_OnePassPointLight || IsFeatureLevelSupported(Platform, ERHIFeatureLevel::SM5))
 			// Only compile position-only shaders for vertex factories that support it. (Note: this assumes that a vertex factor which supports PositionOnly, supports also PositionAndNormalOnly)
 			&& (!bUsePositionOnlyStream || Parameters.VertexFactoryType->SupportsPositionOnly())
 			// Don't render ShadowDepth for translucent unlit materials
@@ -454,11 +469,12 @@ public:
 	static bool ShouldCompilePermutation(const FMeshMaterialShaderPermutationParameters& Parameters)
 	{
 		const EShaderPlatform Platform = Parameters.Platform;
-		
-		// Mobile only needs NonPerspectiveCorrect, and optionally PerspectiveCorrect
+
+		// Mobile only needs NonPerspectiveCorrect, and optionally PerspectiveCorrect and OnePassPointLight
 		if (!IsFeatureLevelSupported(Platform, ERHIFeatureLevel::SM5) &&
 			!(ShaderMode == PixelShadowDepth_NonPerspectiveCorrect || 
-				(ShaderMode == PixelShadowDepth_PerspectiveCorrect && MobileUsesPerspectiveCorrectShadowPermutation(Platform))))
+				(ShaderMode == PixelShadowDepth_PerspectiveCorrect && MobileUsesPerspectiveCorrectShadowPermutation(Platform))
+					|| ((ShaderMode == PixelShadowDepth_OnePassPointLight) && FReadOnlyCVARCache::EnablePointLightShadows(Platform))))
 		{
 			return false;
 		}
@@ -773,11 +789,6 @@ public:
 		SHADER_PARAMETER_SAMPLER(SamplerState, ShadowDepthSampler)
 		RENDER_TARGET_BINDING_SLOTS()
 	END_SHADER_PARAMETER_STRUCT()
-
-	static bool ShouldCompilePermutation(const FGlobalShaderPermutationParameters& Parameters)
-	{
-		return IsFeatureLevelSupported(Parameters.Platform, ERHIFeatureLevel::SM5);
-	}
 };
 
 IMPLEMENT_GLOBAL_SHADER(FCopyShadowMapsCubePS, "/Engine/Private/CopyShadowMaps.usf", "CopyCubeDepthPS", SF_Pixel);
@@ -870,7 +881,7 @@ void FProjectedShadowInfo::CopyCachedShadowMap(
 				else
 	#endif
 				{
-					check(RHISupportsVertexShaderLayer(GShaderPlatformForFeatureLevel[SceneRenderer->FeatureLevel]));
+					check((RHISupportsVertexShaderLayer(SceneRenderer->ShaderPlatform) && GRHISupportsArrayIndexFromAnyShader));
 					TShaderMapRef<TScreenVSForGS<true>> VertexShader(View.ShaderMap);
 					GraphicsPSOInit.BoundShaderState.VertexShaderRHI = VertexShader.GetVertexShader();
 					ScreenVertexShader = VertexShader;
@@ -2350,6 +2361,7 @@ FMeshPassProcessor* CreateOnePassPointLightShadowDepthPassProcessor(ERHIFeatureL
 }
 
 REGISTER_MESHPASSPROCESSOR_AND_PSOCOLLECTOR(OnePassPointLightShadowDepthPass, CreateOnePassPointLightShadowDepthPassProcessor, EShadingPath::Deferred, EMeshPass::OnePassPointLightShadowDepth, EMeshPassFlags::CachedMeshCommands);
+REGISTER_MESHPASSPROCESSOR_AND_PSOCOLLECTOR(OnePassPointLightMobileShadowDepthPass, CreateOnePassPointLightShadowDepthPassProcessor, EShadingPath::Mobile, EMeshPass::OnePassPointLightShadowDepth, EMeshPassFlags::CachedMeshCommands);
 
 FMeshPassProcessor* CreateVSMShadowDepthPassProcessor(ERHIFeatureLevel::Type FeatureLevel, const FScene* Scene, const FSceneView* InViewIfDynamicMeshCommand, FMeshPassDrawListContext* InDrawListContext)
 {
diff --git a/Engine/Source/Runtime/Renderer/Private/ShadowRendering.h b/Engine/Source/Runtime/Renderer/Private/ShadowRendering.h
--- a/Engine/Source/Runtime/Renderer/Private/ShadowRendering.h
+++ b/Engine/Source/Runtime/Renderer/Private/ShadowRendering.h
@@ -1571,13 +1571,25 @@ public:
 	{
 		FGlobalShader::ModifyCompilationEnvironment(Parameters,OutEnvironment);
 		OutEnvironment.SetDefine(TEXT("SHADOW_QUALITY"), Quality);
-		OutEnvironment.SetDefine(TEXT("USE_TRANSMISSION"), (uint32)(bUseTransmission ? 1 : 0));
+		if (IsMobilePlatform(Parameters.Platform))
+		{
+			OutEnvironment.SetDefine(TEXT("USE_TRANSMISSION"), 0);
+		}
+		else
+		{
+			OutEnvironment.SetDefine(TEXT("USE_TRANSMISSION"), (uint32)(bUseTransmission ? 1 : 0));
+		}
 		OutEnvironment.SetDefine(TEXT("SUBPIXEL_SHADOW"), (uint32)(bUseSubPixel ? 1 : 0));
 	}
 
 	static bool ShouldCompilePermutation(const FGlobalShaderPermutationParameters& Parameters)
 	{
-		return IsFeatureLevelSupported(Parameters.Platform, ERHIFeatureLevel::SM5);
+		if (IsMobilePlatform(Parameters.Platform))
+		{
+			return !bUseTransmission;
+		}
+
+		return true;
 	}
 
 	void SetParameters(
diff --git a/Engine/Source/Runtime/Renderer/Private/ShadowSetup.cpp b/Engine/Source/Runtime/Renderer/Private/ShadowSetup.cpp
--- a/Engine/Source/Runtime/Renderer/Private/ShadowSetup.cpp
+++ b/Engine/Source/Runtime/Renderer/Private/ShadowSetup.cpp
@@ -41,6 +41,7 @@
 #include "Shadows/ShadowScene.h"
 #include "LineTypes.h"
 #include "SceneCulling/SceneCulling.h"
+#include "ReadOnlyCVARCache.h"
 
 using namespace UE::Geometry;
 
@@ -3703,7 +3704,7 @@ void ComputeWholeSceneShadowCacheModes(
 	}
 
 	if (GCacheWholeSceneShadows 
-		&& (!bCubeShadowMap || RHISupportsGeometryShaders(GShaderPlatformForFeatureLevel[Scene->GetFeatureLevel()]) || RHISupportsVertexShaderLayer(GShaderPlatformForFeatureLevel[Scene->GetFeatureLevel()])))
+		&& (!bCubeShadowMap || RHISupportsGeometryShaders(GShaderPlatformForFeatureLevel[Scene->GetFeatureLevel()]) || (RHISupportsVertexShaderLayer(Scene->GetShaderPlatform()) && GRHISupportsArrayIndexFromAnyShader)))
 	{
 		TArray<FCachedShadowMapData>* CachedShadowMapDatas = Scene->GetCachedShadowMapDatas(LightSceneInfo->Id);
 
@@ -5986,7 +5987,7 @@ void FSceneRenderer::AllocateCSMDepthTargets(
 
 void FSceneRenderer::AllocateOnePassPointLightDepthTargets(FRHICommandListBase& RHICmdList, TConstArrayView<FProjectedShadowInfo*> WholeScenePointShadows)
 {
-	if (FeatureLevel >= ERHIFeatureLevel::SM5)
+	if (FeatureLevel >= ERHIFeatureLevel::SM5 || DoesRuntimeSupportOnePassPointLightShadows(ShaderPlatform))
 	{
 		for (int32 ShadowIndex = 0; ShadowIndex < WholeScenePointShadows.Num(); ShadowIndex++)
 		{
@@ -6115,7 +6116,7 @@ void FSceneRenderer::CreateDynamicShadows(FDynamicShadowsTaskData& TaskData)
 	const bool bMobile = FeatureLevel < ERHIFeatureLevel::SM5;
 	const bool bHairStrands = HairStrands::HasHairInstanceInScene(*Scene);
 
-	const bool bProjectEnablePointLightShadows = FReadOnlyCVARCache::EnablePointLightShadows() && !bMobile; // Point light shadow is unsupported on mobile for now.
+	const bool bProjectEnablePointLightShadows = FReadOnlyCVARCache::EnablePointLightShadows(ShaderPlatform);
 	const bool bProjectEnableMovableSpotLightShadows = !bMobile || IsMobileMovableSpotlightShadowsEnabled(ShaderPlatform);
 
 	uint32 NumPointShadowCachesUpdatedThisFrame = 0;
