Fix holdout value for translucent material when refraction is driven by index of refraction.

diff --git a/Engine/Source/Runtime/Renderer/Private/DistortionRendering.cpp b/Engine/Source/Runtime/Renderer/Private/DistortionRendering.cpp
--- a/Engine/Source/Runtime/Renderer/Private/DistortionRendering.cpp
+++ b/Engine/Source/Runtime/Renderer/Private/DistortionRendering.cpp
@@ -910,6 +910,12 @@ void FDeferredShadingSceneRenderer::RenderDistortion(
 			false, CF_Always, SO_Keep, SO_Keep, SO_Keep,
 			DISTORTION_STENCIL_MASK_BIT, DISTORTION_STENCIL_MASK_BIT>::GetRHI();
 		PipelineState.StencilRef = DISTORTION_STENCIL_MASK_BIT;
+		
+		// When holdout is enabled, we retain the alpha to keep the translucent holdout alpha.
+		if (IsPrimitiveAlphaHoldoutEnabledForAnyView(Views))
+		{
+			PipelineState.BlendState = TStaticBlendState<CW_RGBA, BO_Add, BF_One, BF_Zero, BO_Add, BF_Zero, BF_One>::GetRHI();
+		}
 
 		for (int32 ViewIndex = 0, Num = Views.Num(); ViewIndex < Num; ++ViewIndex)
 		{
