Fix bug where the LocalLightVisibleLightInfosIndex can be set on a View even when the light grid is not present.

#rnx
#tests enginetest translucency volumetric tests

diff --git a/Engine/Source/Runtime/Renderer/Private/LightRendering.cpp b/Engine/Source/Runtime/Renderer/Private/LightRendering.cpp
--- a/Engine/Source/Runtime/Renderer/Private/LightRendering.cpp
+++ b/Engine/Source/Runtime/Renderer/Private/LightRendering.cpp
@@ -1454,7 +1454,7 @@ void FDeferredShadingSceneRenderer::RenderLights(
 
 			// Shadowed and light function lights
 			{
-				FTranslucentLightInjectionCollector Collector(GraphBuilder, Views);
+				FTranslucentLightInjectionCollector Collector(GraphBuilder, Views, bAreLightsInLightGrid);
 
 				// Collect all the light injection data
 				for (int32 LightIndex = SimpleLightsEnd; LightIndex < SortedLights.Num(); LightIndex++)
diff --git a/Engine/Source/Runtime/Renderer/Private/TranslucentLighting.cpp b/Engine/Source/Runtime/Renderer/Private/TranslucentLighting.cpp
--- a/Engine/Source/Runtime/Renderer/Private/TranslucentLighting.cpp
+++ b/Engine/Source/Runtime/Renderer/Private/TranslucentLighting.cpp
@@ -1087,11 +1087,17 @@ FVolumeBounds CalculateLightVolumeBounds(const FSphere& LightBounds, const FView
 
 FTranslucentLightInjectionCollector::FTranslucentLightInjectionCollector(
 	FRDGBuilder& GraphBuilder,
-	TArrayView<const FViewInfo> Views)
+	TArrayView<const FViewInfo> Views,
+	bool bAreLightsInLightGrid)
 	// NOTE: This data is directly referenced inside the render pass lamba, so must be allocated in the graph
 	: InjectionDataPerView(*GraphBuilder.AllocObject<TArray<FPerViewData, SceneRenderingAllocator>>())
 {
 	InjectionDataPerView.SetNum(Views.Num());
+
+	// Static conditions for supporting batching
+	bCollectorSupportsBatching =
+		CVarTranslucencyLightingVolumeBatch.GetValueOnRenderThread() != 0 &&
+		bAreLightsInLightGrid;
 }
 
 /**
@@ -1145,7 +1151,7 @@ void FTranslucentLightInjectionCollector::AddLightForInjection(
 
 		// Lights without certain features can be batched into a single draw (loop in shader) which is more efficient
 		bool bSupportsBatching =
-			CVarTranslucencyLightingVolumeBatch.GetValueOnRenderThread() != 0 &&
+			bCollectorSupportsBatching &&
 			LightType != LightType_Directional &&
 			(bSupportRectLights || LightType != LightType_Rect) &&
 			!bStaticShadowing &&
diff --git a/Engine/Source/Runtime/Renderer/Private/TranslucentLighting.h b/Engine/Source/Runtime/Renderer/Private/TranslucentLighting.h
--- a/Engine/Source/Runtime/Renderer/Private/TranslucentLighting.h
+++ b/Engine/Source/Runtime/Renderer/Private/TranslucentLighting.h
@@ -57,7 +57,8 @@ struct FTranslucentLightInjectionCollector
 public:
 	FTranslucentLightInjectionCollector(
 		FRDGBuilder& GraphBuilder,
-		TArrayView<const FViewInfo> Views);
+		TArrayView<const FViewInfo> Views,
+		bool bAreLightsInLightGrid);
 
 	/** 
 	* Information about a light to be injected.
@@ -102,6 +103,8 @@ public:
 		TArrayView<const FVisibleLightInfo> VisibleLightInfos,
 		const FLightSceneInfo& LightSceneInfo,
 		const FProjectedShadowInfo* InProjectedShadowInfo = nullptr);
+
+	bool bCollectorSupportsBatching = false;
 };
 
 /** Initializes translucency volume lighting shader parameters from an optional textures struct. If null or uninitialized, fallback textures are used. */
