Moved occlusion back to non-immediate command list but preserved dispatch hint.

diff --git a/Engine/Source/Runtime/Renderer/Private/SceneOcclusion.cpp b/Engine/Source/Runtime/Renderer/Private/SceneOcclusion.cpp
--- a/Engine/Source/Runtime/Renderer/Private/SceneOcclusion.cpp
+++ b/Engine/Source/Runtime/Renderer/Private/SceneOcclusion.cpp
@@ -1473,13 +1473,13 @@ void FDeferredShadingSceneRenderer::RenderOcclusion(
 		TrimAllOcclusionHistory(Views);
 	}
 
+	const bool bUseHzbOcclusion = RenderHzb(GraphBuilder, SceneTextures.Depth.Resolve, BuildHZBAsyncComputeParams, FroxelRenderer);
+
 	if (bIsOcclusionTesting)
 	{
 		FenceOcclusionTests(GraphBuilder);
 	}
 
-	const bool bUseHzbOcclusion = RenderHzb(GraphBuilder, SceneTextures.Depth.Resolve, BuildHZBAsyncComputeParams, FroxelRenderer);
-
 	if (bUseHzbOcclusion || bIsOcclusionTesting)
 	{
 		GraphBuilder.AddDispatchHint();
@@ -1513,46 +1513,41 @@ static uint32 GetViewStateUniqueID(const FSceneRenderer* SceneRenderer)
 	return SceneRenderer->Views.Num() && SceneRenderer->Views[0].ViewState ? SceneRenderer->Views[0].ViewState->UniqueID : 0;
 }
 
-void FSceneRenderer::FenceOcclusionTestsInternal(FRHICommandListImmediate& RHICmdList)
-{
-	SCOPE_CYCLE_COUNTER(STAT_OcclusionSubmittedFence_Dispatch);
-
-	if (ViewFamily.bIsMultipleViewFamily)
-	{
-		// If there are multiple view families, we implement a queue of buffered fences, so we can avoid waiting on queries for
-		// another view family that's rendering in the same frame.  Here we push a new fence into the queue of buffered fences.
-		// We assume that the queue isn't full, since WaitOcclusionTests (called earlier in the frame) will always pop at least
-		// one fence if the queue is full.
-		check(OcclusionSubmittedFence[FOcclusionQueryHelpers::MaxBufferedOcclusionFrames - 1].Fence == nullptr);
-
-		for (int32 Dest = FOcclusionQueryHelpers::MaxBufferedOcclusionFrames - 1; Dest >= 1; Dest--)
-		{
-			CA_SUPPRESS(6385);
-			OcclusionSubmittedFence[Dest] = OcclusionSubmittedFence[Dest - 1];
-		}
-	}
-	else
-	{
-		// Single view family implementation, fixed number of buffered frames.
-		int32 NumFrames = FOcclusionQueryHelpers::GetNumBufferedFrames(FeatureLevel);
-		for (int32 Dest = NumFrames - 1; Dest >= 1; Dest--)
-		{
-			CA_SUPPRESS(6385);
-			OcclusionSubmittedFence[Dest] = OcclusionSubmittedFence[Dest - 1];
-		}
-	}
-
-	OcclusionSubmittedFence[0].Fence = RHICmdList.RHIThreadFence();
-	OcclusionSubmittedFence[0].ViewStateUniqueID = GetViewStateUniqueID(this);
-}
-
 void FSceneRenderer::FenceOcclusionTests(FRDGBuilder& GraphBuilder)
 {
 	if (DoOcclusionQueries() && IsRunningRHIInSeparateThread())
 	{
-		AddPass(GraphBuilder, RDG_EVENT_NAME("FenceOcclusionTests"), [this](FRHICommandListImmediate& RHICmdList)
+		AddPass(GraphBuilder, RDG_EVENT_NAME("FenceOcclusionTests"), [this](FRDGAsyncTask, FRHICommandList& RHICmdList)
 		{
-			FenceOcclusionTestsInternal(RHICmdList);
+			SCOPE_CYCLE_COUNTER(STAT_OcclusionSubmittedFence_Dispatch);
+
+			if (ViewFamily.bIsMultipleViewFamily)
+			{
+				// If there are multiple view families, we implement a queue of buffered fences, so we can avoid waiting on queries for
+				// another view family that's rendering in the same frame.  Here we push a new fence into the queue of buffered fences.
+				// We assume that the queue isn't full, since WaitOcclusionTests (called earlier in the frame) will always pop at least
+				// one fence if the queue is full.
+				check(OcclusionSubmittedFence[FOcclusionQueryHelpers::MaxBufferedOcclusionFrames - 1].Fence == nullptr);
+
+				for (int32 Dest = FOcclusionQueryHelpers::MaxBufferedOcclusionFrames - 1; Dest >= 1; Dest--)
+				{
+					CA_SUPPRESS(6385);
+					OcclusionSubmittedFence[Dest] = OcclusionSubmittedFence[Dest - 1];
+				}
+			}
+			else
+			{
+				// Single view family implementation, fixed number of buffered frames.
+				int32 NumFrames = FOcclusionQueryHelpers::GetNumBufferedFrames(FeatureLevel);
+				for (int32 Dest = NumFrames - 1; Dest >= 1; Dest--)
+				{
+					CA_SUPPRESS(6385);
+					OcclusionSubmittedFence[Dest] = OcclusionSubmittedFence[Dest - 1];
+				}
+			}
+
+			OcclusionSubmittedFence[0].Fence = RHICmdList.RHIThreadFence();
+			OcclusionSubmittedFence[0].ViewStateUniqueID = GetViewStateUniqueID(this);
 		});
 
 		GraphBuilder.AddDispatchHint();
diff --git a/Engine/Source/Runtime/Renderer/Private/SceneRendering.h b/Engine/Source/Runtime/Renderer/Private/SceneRendering.h
--- a/Engine/Source/Runtime/Renderer/Private/SceneRendering.h
+++ b/Engine/Source/Runtime/Renderer/Private/SceneRendering.h
@@ -2125,7 +2125,6 @@ public:
 	bool DoOcclusionQueries() const;
 
 	void FenceOcclusionTests(FRDGBuilder& GraphBuilder);
-	void FenceOcclusionTestsInternal(FRHICommandListImmediate& RHICmdList);
 	void WaitOcclusionTests(FRHICommandListImmediate& RHICmdList);
 
 	// fences to make sure the rhi thread has digested the occlusion query renders before we attempt to read them back async
