Add an RDG scope around debug drawing to be able to tell when it's causing GPU hangs.

#rnx

diff --git a/Engine/Source/Runtime/Renderer/Private/PostProcess/PostProcessing.cpp b/Engine/Source/Runtime/Renderer/Private/PostProcess/PostProcessing.cpp
--- a/Engine/Source/Runtime/Renderer/Private/PostProcess/PostProcessing.cpp
+++ b/Engine/Source/Runtime/Renderer/Private/PostProcess/PostProcessing.cpp
@@ -1847,6 +1847,8 @@ void AddPostProcessingPasses(
 	{
 		// Draw debug stuff directly onto the back buffer
 
+		RDG_EVENT_SCOPE(GraphBuilder, "Debug Drawing");
+
 		if (EngineShowFlags.TestImage)
 		{
 			AddTestImagePass(GraphBuilder, View, SceneColor);
