[RHI] Pass shader platform into RHISupportsVolumeTextureRendering() instead of always using GMaxRHIFeatureLevel.

This hardens the condition of whether volume texture rendering is supported by the currently selected shader platform and is a follow up to CL 35543745.

#rnx

diff --git a/Engine/Source/Runtime/Engine/Public/SceneView.h b/Engine/Source/Runtime/Engine/Public/SceneView.h
--- a/Engine/Source/Runtime/Engine/Public/SceneView.h
+++ b/Engine/Source/Runtime/Engine/Public/SceneView.h
@@ -2330,7 +2330,7 @@ public:
 
 	ENGINE_API ERHIFeatureLevel::Type GetFeatureLevel() const;
 
-	EShaderPlatform GetShaderPlatform() const { return GShaderPlatformForFeatureLevel[GetFeatureLevel()]; }
+	EShaderPlatform GetShaderPlatform() const { return GetFeatureLevelShaderPlatform(GetFeatureLevel()); }
 
 #if WITH_DEBUG_VIEW_MODES
 	EDebugViewShaderMode DebugViewShaderMode;
diff --git a/Engine/Source/Runtime/RHI/Private/RHI.cpp b/Engine/Source/Runtime/RHI/Private/RHI.cpp
--- a/Engine/Source/Runtime/RHI/Private/RHI.cpp
+++ b/Engine/Source/Runtime/RHI/Private/RHI.cpp
@@ -1477,12 +1477,11 @@ int32 RHIGetPreferredClearUAVRectPSResourceType(const FStaticShaderPlatform Plat
 	return 1; // TEXTURE_2D
 }
 
-RHI_API bool RHISupportsVolumeTextureRendering()
+RHI_API bool RHISupportsVolumeTextureRendering(const FStaticShaderPlatform Platform)
 {
 #if WITH_EDITOR
 	// When preview platforms are supported (when building with editor) we might be previewing an RHI that doesn't support geometry shaders (such as Metal)
 	// while rendering with a runtime RHI that doesn't support vertex shader layers as an alternative (such as D3D), so take these DDPI entries into account, too.
-	const EShaderPlatform Platform = GShaderPlatformForFeatureLevel[GMaxRHIFeatureLevel];
 	return GSupportsVolumeTextureRendering && (RHISupportsGeometryShaders(Platform) || RHISupportsVertexShaderLayer(Platform));
 #else
 	return GSupportsVolumeTextureRendering;
diff --git a/Engine/Source/Runtime/RHI/Public/RHI.h b/Engine/Source/Runtime/RHI/Public/RHI.h
--- a/Engine/Source/Runtime/RHI/Public/RHI.h
+++ b/Engine/Source/Runtime/RHI/Public/RHI.h
@@ -84,7 +84,7 @@ RHI_API int32 RHIGetPreferredClearUAVRectPSResourceType(const FStaticShaderPlatf
 RHI_API void RHIDumpResourceMemoryToCSV();
 
 // helper to check whether rendering to volume textures is supported. This takes both GSupportsVolumeTextureRendering and DDPI entries into account as this is dependent on the preview platform.
-RHI_API bool RHISupportsVolumeTextureRendering();
+RHI_API bool RHISupportsVolumeTextureRendering(const FStaticShaderPlatform Platform);
 
 struct FRHIResourceStats
 {
diff --git a/Engine/Source/Runtime/Renderer/Private/IndirectLightingCache.cpp b/Engine/Source/Runtime/Renderer/Private/IndirectLightingCache.cpp
--- a/Engine/Source/Runtime/Renderer/Private/IndirectLightingCache.cpp
+++ b/Engine/Source/Runtime/Renderer/Private/IndirectLightingCache.cpp
@@ -142,7 +142,7 @@ bool IsIndirectLightingCacheAllowed(ERHIFeatureLevel::Type InFeatureLevel)
 bool CanIndirectLightingCacheUseVolumeTexture(ERHIFeatureLevel::Type InFeatureLevel)
 {
 	// @todo Mac OS X/OpenGL: For OpenGL devices which don't support volume-texture rendering we need to use the simpler point indirect lighting shaders.
-	return InFeatureLevel >= ERHIFeatureLevel::SM5 && RHISupportsVolumeTextureRendering();
+	return InFeatureLevel >= ERHIFeatureLevel::SM5 && RHISupportsVolumeTextureRendering(GetFeatureLevelShaderPlatform(InFeatureLevel));
 }
 
 FIndirectLightingCache::FIndirectLightingCache(ERHIFeatureLevel::Type InFeatureLevel)
diff --git a/Engine/Source/Runtime/Renderer/Private/LightMapRendering.cpp b/Engine/Source/Runtime/Renderer/Private/LightMapRendering.cpp
--- a/Engine/Source/Runtime/Renderer/Private/LightMapRendering.cpp
+++ b/Engine/Source/Runtime/Renderer/Private/LightMapRendering.cpp
@@ -760,7 +760,7 @@ void GetIndirectLightingCacheParameters(
 		// If we are using FCachedVolumeIndirectLightingPolicy then InitViews should have updated the lighting cache which would have initialized it
 		// However the conditions for updating the lighting cache are complex and fail very occasionally in non-reproducible ways
 		// Silently skipping setting the cache texture under failure for now
-		if (FeatureLevel >= ERHIFeatureLevel::SM5 && LightingCache && LightingCache->IsInitialized() && RHISupportsVolumeTextureRendering())
+		if (FeatureLevel >= ERHIFeatureLevel::SM5 && LightingCache && LightingCache->IsInitialized() && RHISupportsVolumeTextureRendering(GetFeatureLevelShaderPlatform(FeatureLevel)))
 		{
 			Parameters.IndirectLightingCacheTexture0 = const_cast<FIndirectLightingCache*>(LightingCache)->GetTexture0();
 			Parameters.IndirectLightingCacheTexture1 = const_cast<FIndirectLightingCache*>(LightingCache)->GetTexture1();
diff --git a/Engine/Source/Runtime/Renderer/Private/LightRendering.cpp b/Engine/Source/Runtime/Renderer/Private/LightRendering.cpp
--- a/Engine/Source/Runtime/Renderer/Private/LightRendering.cpp
+++ b/Engine/Source/Runtime/Renderer/Private/LightRendering.cpp
@@ -1426,7 +1426,7 @@ void FDeferredShadingSceneRenderer::RenderLights(
 			Substrate::AddSubstrateStencilPass(GraphBuilder, Views, SceneTextures);
 		}
 
-		if (ViewFamily.EngineShowFlags.DirectLighting && GUseTranslucentLightingVolumes && RHISupportsVolumeTextureRendering())
+		if (ViewFamily.EngineShowFlags.DirectLighting && GUseTranslucentLightingVolumes && RHISupportsVolumeTextureRendering(ViewFamily.GetShaderPlatform()))
 		{
 			RDG_EVENT_SCOPE(GraphBuilder, "InjectTranslucencyLightingVolume");
 			SCOPE_CYCLE_COUNTER(STAT_TranslucentInjectTime);
diff --git a/Engine/Source/Runtime/Renderer/Private/TranslucentLighting.cpp b/Engine/Source/Runtime/Renderer/Private/TranslucentLighting.cpp
--- a/Engine/Source/Runtime/Renderer/Private/TranslucentLighting.cpp
+++ b/Engine/Source/Runtime/Renderer/Private/TranslucentLighting.cpp
@@ -993,7 +993,7 @@ void InjectTranslucencyLightingVolumeAmbientCubemap(
 	const TArrayView<const FViewInfo> Views,
 	const FTranslucencyLightingVolumeTextures& Textures)
 {
-	if (!GUseTranslucentLightingVolumes || !RHISupportsVolumeTextureRendering())
+	if (!GUseTranslucentLightingVolumes || Views.Num() == 0 || !RHISupportsVolumeTextureRendering(Views[0].GetShaderPlatform()))
 	{
 		return;
 	}
@@ -1354,7 +1354,7 @@ void InjectTranslucencyLightingVolume(
 	const FTranslucencyLightingVolumeTextures& Textures,
 	const FTranslucentLightInjectionCollector& Collector)
 {
-	if (!GUseTranslucentLightingVolumes || !RHISupportsVolumeTextureRendering())
+	if (!GUseTranslucentLightingVolumes || !RHISupportsVolumeTextureRendering(View.GetShaderPlatform()))
 	{
 		return;
 	}
@@ -1616,7 +1616,7 @@ void FilterTranslucencyLightingVolume(
 	const TArrayView<const FViewInfo> Views,
 	FTranslucencyLightingVolumeTextures& Textures)
 {
-	if (!GUseTranslucentLightingVolumes || !RHISupportsVolumeTextureRendering() || !GUseTranslucencyVolumeBlur)
+	if (!GUseTranslucentLightingVolumes || Views.Num() == 0 || !RHISupportsVolumeTextureRendering(Views[0].GetShaderPlatform()) || !GUseTranslucencyVolumeBlur)
 	{
 		return;
 	}
diff --git a/Engine/Source/Runtime/SlateRHIRenderer/Private/SlateRHIRenderer.cpp b/Engine/Source/Runtime/SlateRHIRenderer/Private/SlateRHIRenderer.cpp
--- a/Engine/Source/Runtime/SlateRHIRenderer/Private/SlateRHIRenderer.cpp
+++ b/Engine/Source/Runtime/SlateRHIRenderer/Private/SlateRHIRenderer.cpp
@@ -567,7 +567,7 @@ bool CompositeUIWithSceneHDR()
 	static const auto CVarCompositeMode = IConsoleManager::Get().FindTConsoleVariableDataInt(TEXT("r.HDR.UI.CompositeMode"));
 
 	return GRHISupportsHDROutput
-		&& RHISupportsVolumeTextureRendering()
+		&& RHISupportsVolumeTextureRendering(GetFeatureLevelShaderPlatform(GMaxRHIFeatureLevel))
 		&& SupportsCompositeUIWithSceneHDR(GetFeatureLevelShaderPlatform(GMaxRHIFeatureLevel))
 		&& CVarCompositeMode
 		&& CVarCompositeMode->GetValueOnAnyThread() != 0;
