HeterogeneousVolumes: Code cleanup

Removing ViewFamily dependency on calling API

diff --git a/Engine/Source/Runtime/Renderer/Private/HeterogeneousVolumes/HeterogeneousVolumes.cpp b/Engine/Source/Runtime/Renderer/Private/HeterogeneousVolumes/HeterogeneousVolumes.cpp
--- a/Engine/Source/Runtime/Renderer/Private/HeterogeneousVolumes/HeterogeneousVolumes.cpp
+++ b/Engine/Source/Runtime/Renderer/Private/HeterogeneousVolumes/HeterogeneousVolumes.cpp
@@ -732,7 +732,6 @@ void RenderHeterogeneousVolumeShadows(
 	FScene* Scene,
 	const FSceneTextures& SceneTextures,
 	FViewInfo& View,
-	const FSceneViewFamily& ViewFamily,
 	TArray<FVisibleLightInfo, SceneRenderingAllocator>& VisibleLightInfos
 )
 {
@@ -746,7 +745,6 @@ void RenderHeterogeneousVolumeShadows(
 			GraphBuilder,
 			SceneTextures,
 			Scene,
-			ViewFamily,
 			View,
 			VisibleLightInfos
 		);
@@ -795,7 +793,6 @@ void FDeferredShadingSceneRenderer::RenderHeterogeneousVolumeShadows(
 					GraphBuilder,
 					SceneTextures,
 					Scene,
-					ViewFamily,
 					View
 				);
 			}
@@ -806,7 +803,6 @@ void FDeferredShadingSceneRenderer::RenderHeterogeneousVolumeShadows(
 					// Scene data
 					SceneTextures,
 					Scene,
-					ViewFamily,
 					View,
 					// Volume data
 					OrthoGridUniformBuffer,
@@ -826,7 +822,6 @@ void FDeferredShadingSceneRenderer::RenderHeterogeneousVolumeShadows(
 					// Scene data
 					SceneTextures,
 					Scene,
-					ViewFamily,
 					View,
 					// Light data
 					VisibleLightInfos
@@ -840,7 +835,6 @@ void FDeferredShadingSceneRenderer::RenderHeterogeneousVolumeShadows(
 					// Scene data
 					SceneTextures,
 					Scene,
-					ViewFamily,
 					View,
 					// Shadow Data
 					VisibleLightInfos,
@@ -903,7 +897,6 @@ void FDeferredShadingSceneRenderer::RenderHeterogeneousVolumes(
 					GraphBuilder,
 					SceneTextures,
 					Scene,
-					ViewFamily,
 					View,
 					OrthoGridUniformBuffer,
 					FrustumGridUniformBuffer,
@@ -1009,7 +1002,6 @@ void FDeferredShadingSceneRenderer::RenderHeterogeneousVolumes(
 							GraphBuilder,
 							SceneTextures,
 							Scene,
-							ViewFamily,
 							View,
 							// Shadow Data
 							VisibleLightInfos,
diff --git a/Engine/Source/Runtime/Renderer/Private/HeterogeneousVolumes/HeterogeneousVolumes.h b/Engine/Source/Runtime/Renderer/Private/HeterogeneousVolumes/HeterogeneousVolumes.h
--- a/Engine/Source/Runtime/Renderer/Private/HeterogeneousVolumes/HeterogeneousVolumes.h
+++ b/Engine/Source/Runtime/Renderer/Private/HeterogeneousVolumes/HeterogeneousVolumes.h
@@ -545,7 +545,6 @@ void RenderWithPreshading(
 	// Scene data
 	const FSceneTextures& SceneTextures,
 	FScene* Scene,
-	const FSceneViewFamily& ViewFamily,
 	FViewInfo& View,
 	// Shadow data
 	TArray<FVisibleLightInfo, SceneRenderingAllocator>& VisibleLightInfos,
@@ -566,7 +565,6 @@ void RenderTransmittanceWithVoxelGrid(
 	// Scene data
 	const FSceneTextures& SceneTextures,
 	FScene* Scene,
-	const FSceneViewFamily& ViewFamily,
 	FViewInfo& View,
 	const TRDGUniformBufferRef<FOrthoVoxelGridUniformBufferParameters>& OrthoGridUniformBuffer,
 	const TRDGUniformBufferRef<FFrustumVoxelGridUniformBufferParameters>& FrustumGridUniformBuffer,
@@ -579,7 +577,6 @@ void RenderAdaptiveVolumetricShadowMapWithVoxelGrid(
 	// Scene data
 	const FSceneTextures& SceneTextures,
 	FScene* Scene,
-	const FSceneViewFamily& ViewFamily,
 	FViewInfo& View,
 	// Shadow data
 	TArray<FVisibleLightInfo, SceneRenderingAllocator>& VisibleLightInfos,
@@ -594,7 +591,6 @@ void RenderAdaptiveVolumetricShadowMapWithLiveShading(
 	// Scene data
 	const FSceneTextures& SceneTextures,
 	FScene* Scene,
-	const FSceneViewFamily& ViewFamily,
 	FViewInfo& View,
 	// Light data
 	TArray<FVisibleLightInfo, SceneRenderingAllocator>& VisibleLightInfos
@@ -605,7 +601,6 @@ void RenderAdaptiveVolumetricCameraMapWithVoxelGrid(
 	// Scene data
 	const FSceneTextures& SceneTextures,
 	FScene* Scene,
-	const FSceneViewFamily& ViewFamily,
 	FViewInfo& View,
 	// Volume data
 	const TRDGUniformBufferRef<FOrthoVoxelGridUniformBufferParameters>& OrthoGridUniformBuffer,
@@ -617,7 +612,6 @@ void RenderAdaptiveVolumetricCameraMapWithLiveShading(
 	// Scene data
 	const FSceneTextures& SceneTextures,
 	FScene* Scene,
-	const FSceneViewFamily& ViewFamily,
 	FViewInfo& View
 );
 
@@ -668,7 +662,6 @@ void RenderSingleScatteringWithVoxelGrid(
 	// Scene data
 	const FSceneTextures& SceneTextures,
 	FScene* Scene,
-	const FSceneViewFamily& ViewFamily,
 	FViewInfo& View,
 	// Shadow data
 	TArray<FVisibleLightInfo, SceneRenderingAllocator>& VisibleLightInfos,
diff --git a/Engine/Source/Runtime/Renderer/Private/HeterogeneousVolumes/HeterogeneousVolumesLiveShadingPipeline.cpp b/Engine/Source/Runtime/Renderer/Private/HeterogeneousVolumes/HeterogeneousVolumesLiveShadingPipeline.cpp
--- a/Engine/Source/Runtime/Renderer/Private/HeterogeneousVolumes/HeterogeneousVolumesLiveShadingPipeline.cpp
+++ b/Engine/Source/Runtime/Renderer/Private/HeterogeneousVolumes/HeterogeneousVolumesLiveShadingPipeline.cpp
@@ -1880,7 +1880,6 @@ bool RenderVolumetricShadowMapForLightForHeterogeneousVolumeWithLiveShading(
 	// Scene data
 	const FSceneTextures& SceneTextures,
 	FScene* Scene,
-	const FSceneViewFamily& ViewFamily,
 	FViewInfo& View,
 	// Light data
 	uint32 LightType,
@@ -2038,7 +2037,6 @@ bool RenderVolumetricShadowMapForLightWithLiveShading(
 	// Scene data
 	const FSceneTextures& SceneTextures,
 	FScene* Scene,
-	const FSceneViewFamily& ViewFamily,
 	FViewInfo& View,
 	// Light data
 	uint32 LightType,
@@ -2209,7 +2207,6 @@ bool RenderVolumetricShadowMapForLightWithLiveShading(
 			GraphBuilder,
 			SceneTextures,
 			Scene,
-			ViewFamily,
 			View,
 			// Light Info
 			LightType,
@@ -2245,7 +2242,6 @@ bool RenderVolumetricShadowMapForLightWithLiveShading(
 				GraphBuilder,
 				SceneTextures,
 				Scene,
-				ViewFamily,
 				View,
 				// Light Info
 				LightType,
@@ -2298,7 +2294,6 @@ void RenderAdaptiveVolumetricShadowMapWithLiveShading(
 	// Scene data
 	const FSceneTextures& SceneTextures,
 	FScene* Scene,
-	const FSceneViewFamily& ViewFamily,
 	FViewInfo& View,
 	// Light data
 	TArray<FVisibleLightInfo, SceneRenderingAllocator>& VisibleLightInfos
@@ -2379,7 +2374,6 @@ void RenderAdaptiveVolumetricShadowMapWithLiveShading(
 					// Scene data
 					SceneTextures,
 					Scene,
-					ViewFamily,
 					View,
 					// Light data
 					LightType,
@@ -2452,7 +2446,6 @@ void RenderAdaptiveVolumetricCameraMapWithLiveShading(
 	// Scene data
 	const FSceneTextures& SceneTextures,
 	FScene* Scene,
-	const FSceneViewFamily& ViewFamily,
 	FViewInfo& View
 )
 {
@@ -2543,7 +2536,6 @@ void RenderAdaptiveVolumetricCameraMapWithLiveShading(
 			GraphBuilder,
 			SceneTextures,
 			Scene,
-			ViewFamily,
 			View,
 			// Light Info
 			LightType,
@@ -2580,7 +2572,6 @@ void RenderAdaptiveVolumetricCameraMapWithLiveShading(
 				GraphBuilder,
 				SceneTextures,
 				Scene,
-				ViewFamily,
 				View,
 				// Light Info
 				LightType,
diff --git a/Engine/Source/Runtime/Renderer/Private/HeterogeneousVolumes/HeterogeneousVolumesPreshadingPipeline.cpp b/Engine/Source/Runtime/Renderer/Private/HeterogeneousVolumes/HeterogeneousVolumesPreshadingPipeline.cpp
--- a/Engine/Source/Runtime/Renderer/Private/HeterogeneousVolumes/HeterogeneousVolumesPreshadingPipeline.cpp
+++ b/Engine/Source/Runtime/Renderer/Private/HeterogeneousVolumes/HeterogeneousVolumesPreshadingPipeline.cpp
@@ -712,7 +712,6 @@ void RenderWithInscatteringVolumePipelineWithPreshadingCompute(
 	// Scene data
 	const FSceneTextures& SceneTextures,
 	const FScene* Scene,
-	const FSceneViewFamily& ViewFamily,
 	FViewInfo& View,
 	// Shadow data
 	TArray<FVisibleLightInfo, SceneRenderingAllocator>& VisibleLightInfos,
@@ -846,7 +845,6 @@ void RenderWithTransmittanceVolumePipelineWithPreshadingCompute(
 	// Scene data
 	const FSceneTextures& SceneTextures,
 	const FScene* Scene,
-	const FSceneViewFamily& ViewFamily,
 	FViewInfo& View,
 	// Shadow data
 	TArray<FVisibleLightInfo, SceneRenderingAllocator>& VisibleLightInfos,
@@ -972,7 +970,6 @@ void RenderWithPreshadingCompute(
 	// Scene data
 	const FSceneTextures& SceneTextures,
 	const FScene* Scene,
-	const FSceneViewFamily& ViewFamily,
 	FViewInfo& View,
 	// Shadow data
 	TArray<FVisibleLightInfo, SceneRenderingAllocator>& VisibleLightInfos,
@@ -1021,7 +1018,6 @@ void RenderWithPreshadingCompute(
 			// Scene data
 			SceneTextures,
 			Scene,
-			ViewFamily,
 			View,
 			// Shadow data
 			VisibleLightInfos,
@@ -1048,7 +1044,6 @@ void RenderWithPreshadingCompute(
 			// Scene data
 			SceneTextures,
 			Scene,
-			ViewFamily,
 			View,
 			// Shadow data
 			VisibleLightInfos,
@@ -1075,7 +1070,6 @@ void RenderWithInscatteringVolumePipelineWithPreshadingHardwareRayTracing(
 	// Scene data
 	const FSceneTextures& SceneTextures,
 	FScene* Scene,
-	const FSceneViewFamily& ViewFamily,
 	FViewInfo& View,
 	// Shadow data
 	TArray<FVisibleLightInfo, SceneRenderingAllocator>& VisibleLightInfos,
@@ -1205,7 +1199,6 @@ void RenderWithTransmittanceVolumePipelineWithPreshadingHardwareRayTracing(
 	// Scene data
 	const FSceneTextures& SceneTextures,
 	FScene* Scene,
-	const FSceneViewFamily& ViewFamily,
 	FViewInfo& View,
 	// Shadow data
 	TArray<FVisibleLightInfo, SceneRenderingAllocator>& VisibleLightInfos,
@@ -1327,7 +1320,6 @@ void RenderWithPreshadingHardwareRayTracing(
 	// Scene data
 	const FSceneTextures& SceneTextures,
 	FScene* Scene,
-	const FSceneViewFamily& ViewFamily,
 	FViewInfo& View,
 	// Shadow data
 	TArray<FVisibleLightInfo, SceneRenderingAllocator>& VisibleLightInfos,
@@ -1387,7 +1379,6 @@ void RenderWithPreshadingHardwareRayTracing(
 			GraphBuilder,
 			SceneTextures,
 			Scene,
-			ViewFamily,
 			View,
 			// Shadow data
 			VisibleLightInfos,
@@ -1411,7 +1402,6 @@ void RenderWithPreshadingHardwareRayTracing(
 			GraphBuilder,
 			SceneTextures,
 			Scene,
-			ViewFamily,
 			View,
 			// Shadow data
 			VisibleLightInfos,
@@ -1544,7 +1534,6 @@ void RenderWithPreshading(
 	// Scene data
 	const FSceneTextures& SceneTextures,
 	FScene* Scene,
-	const FSceneViewFamily& ViewFamily,
 	FViewInfo& View,
 	// Shadow data
 	TArray<FVisibleLightInfo, SceneRenderingAllocator>& VisibleLightInfos,
@@ -1700,7 +1689,6 @@ void RenderWithPreshading(
 			// Scene data
 			SceneTextures,
 			Scene,
-			ViewFamily,
 			View,
 			// Shadow data
 			VisibleLightInfos,
@@ -1725,7 +1713,6 @@ void RenderWithPreshading(
 			// Scene data
 			SceneTextures,
 			Scene,
-			ViewFamily,
 			View,
 			// Shadow data
 			VisibleLightInfos,
diff --git a/Engine/Source/Runtime/Renderer/Private/HeterogeneousVolumes/HeterogeneousVolumesVoxelGridPipeline.cpp b/Engine/Source/Runtime/Renderer/Private/HeterogeneousVolumes/HeterogeneousVolumesVoxelGridPipeline.cpp
--- a/Engine/Source/Runtime/Renderer/Private/HeterogeneousVolumes/HeterogeneousVolumesVoxelGridPipeline.cpp
+++ b/Engine/Source/Runtime/Renderer/Private/HeterogeneousVolumes/HeterogeneousVolumesVoxelGridPipeline.cpp
@@ -2877,7 +2877,6 @@ void RenderTransmittanceWithVoxelGrid(
 	// Scene data
 	const FSceneTextures& SceneTextures,
 	FScene* Scene,
-	const FSceneViewFamily& ViewFamily,
 	FViewInfo& View,
 	const TRDGUniformBufferRef<FOrthoVoxelGridUniformBufferParameters>& OrthoGridUniformBuffer,
 	const TRDGUniformBufferRef<FFrustumVoxelGridUniformBufferParameters>& FrustumGridUniformBuffer,
@@ -3400,7 +3399,6 @@ void RenderVolumetricShadowMapForLightWithVoxelGrid(
 	// Scene data
 	const FSceneTextures& SceneTextures,
 	FScene* Scene,
-	const FSceneViewFamily& ViewFamily,
 	FViewInfo& View,
 	// Light data
 	bool bApplyEmissionAndTransmittance,
@@ -3632,7 +3630,6 @@ void RenderVolumetricShadowMapForCameraWithVoxelGrid(
 	// Scene data
 	const FSceneTextures& SceneTextures,
 	FScene* Scene,
-	const FSceneViewFamily& ViewFamily,
 	FViewInfo& View,
 	// Volume data
 	const TRDGUniformBufferRef<FOrthoVoxelGridUniformBufferParameters>& OrthoGridUniformBuffer,
@@ -3754,7 +3751,6 @@ void RenderAdaptiveVolumetricShadowMapWithVoxelGrid(
 	// Scene data
 	const FSceneTextures& SceneTextures,
 	FScene* Scene,
-	const FSceneViewFamily& ViewFamily,
 	FViewInfo& View,
 	// Shadow data
 	TArray<FVisibleLightInfo, SceneRenderingAllocator>& VisibleLightInfos,
@@ -3841,7 +3837,6 @@ void RenderAdaptiveVolumetricShadowMapWithVoxelGrid(
 					// Scene data
 					SceneTextures,
 					Scene,
-					ViewFamily,
 					View,
 					// Light data
 					bApplyEmissionAndTransmittance,
@@ -3915,7 +3910,6 @@ void RenderAdaptiveVolumetricCameraMapWithVoxelGrid(
 	// Scene data
 	const FSceneTextures& SceneTextures,
 	FScene* Scene,
-	const FSceneViewFamily& ViewFamily,
 	FViewInfo& View,
 	// Volume data
 	const TRDGUniformBufferRef<FOrthoVoxelGridUniformBufferParameters>& OrthoGridUniformBuffer,
@@ -3958,7 +3952,6 @@ void RenderAdaptiveVolumetricCameraMapWithVoxelGrid(
 			// Scene data
 			SceneTextures,
 			Scene,
-			ViewFamily,
 			View,
 			// Volume data
 			OrthoGridUniformBuffer,
diff --git a/Engine/Source/Runtime/Renderer/Private/ShadowDepthRendering.cpp b/Engine/Source/Runtime/Renderer/Private/ShadowDepthRendering.cpp
--- a/Engine/Source/Runtime/Renderer/Private/ShadowDepthRendering.cpp
+++ b/Engine/Source/Runtime/Renderer/Private/ShadowDepthRendering.cpp
@@ -1952,7 +1952,6 @@ void FSceneRenderer::RenderShadowDepthMaps(FRDGBuilder& GraphBuilder, FDynamicSh
 				GraphBuilder,
 				View.GetSceneTextures(),
 				Scene,
-				ViewFamily,
 				View,
 				VisibleLightInfos
 			);
