Light Grid - Fix rect light culling due to incorrect culling plane calculation.

Issue:
- adjusting barn door corners for both horizontal and vertical barn doors simultaneously leads to wrong results.
- intersection between shadow boundary lines and bounding sphere was also incorrectly approximated because calculation was using radius as intersection distance, however that is not far enough since it doesn't account for the size of the rect source.

Fix:
- calculate culling planes for horizontal and vertical barn doors separately and properly calculating intersection between shadow boundary line and bounding circle.

Extra debug functionality:
- r.Forward.LightGridDebug.RectLightBounds to disable tight rect light culling
- r.RectLight.VisualizeCullingBarnDoors to visualize culling planes of the selected light

diff --git a/Engine/Source/Editor/ComponentVisualizers/Private/RectLightComponentVisualizer.cpp b/Engine/Source/Editor/ComponentVisualizers/Private/RectLightComponentVisualizer.cpp
--- a/Engine/Source/Editor/ComponentVisualizers/Private/RectLightComponentVisualizer.cpp
+++ b/Engine/Source/Editor/ComponentVisualizers/Private/RectLightComponentVisualizer.cpp
@@ -15,6 +15,13 @@
 #include "ShowFlags.h"
 #include "Templates/Casts.h"
 
+int32 GVisualizeCullingBarnDoors = 0;
+FAutoConsoleVariableRef CVarLightGridHZBCull(
+	TEXT("r.RectLight.VisualizeCullingBarnDoors"),
+	GVisualizeCullingBarnDoors,
+	TEXT("Whether to render a visualization of the barn doors used to cull the rect light."),
+	ECVF_RenderThreadSafe
+);
 
 void FRectLightComponentVisualizer::DrawVisualization( const UActorComponent* Component, const FSceneView* View, FPrimitiveDrawInterface* PDI )
 {
@@ -35,16 +42,17 @@ void FRectLightComponentVisualizer::DrawVisualization( const UActorComponent* Co
 			);
 
 			const FColor ElementColor(231, 239, 0, 255);
-			auto DrawBarnRect = [&](const FVector& P0, const FVector& P1, const FVector& P2, const FVector& P3)
+
+			auto DrawBarnRect = [&](const FVector& P0, const FVector& P1, const FVector& P2, const FVector& P3, const FColor& Color)
 			{
 				FVector TP0 = LightTM.TransformPosition(P0);
 				FVector TP1 = LightTM.TransformPosition(P1);
 				FVector TP2 = LightTM.TransformPosition(P2);
 				FVector TP3 = LightTM.TransformPosition(P3);
-				PDI->DrawLine(TP0, TP1, ElementColor, SDPG_World);
-				PDI->DrawLine(TP1, TP2, ElementColor, SDPG_World);
-				PDI->DrawLine(TP2, TP3, ElementColor, SDPG_World);
-				PDI->DrawLine(TP3, TP0, ElementColor, SDPG_World);
+				PDI->DrawLine(TP0, TP1, Color, SDPG_World);
+				PDI->DrawLine(TP1, TP2, Color, SDPG_World);
+				PDI->DrawLine(TP2, TP3, Color, SDPG_World);
+				PDI->DrawLine(TP3, TP0, Color, SDPG_World);
 			};
 
 			const float BarnMaxAngle = GetRectLightBarnDoorMaxAngle();
@@ -63,7 +71,7 @@ void FRectLightComponentVisualizer::DrawVisualization( const UActorComponent* Co
 				Corners[0] =  P3;
 				Corners[1] =  P2;
 
-				DrawBarnRect(P0, P1, P2, P3);
+				DrawBarnRect(P0, P1, P2, P3, ElementColor);
 			}
 			// +SourceHeight
 			{
@@ -74,7 +82,7 @@ void FRectLightComponentVisualizer::DrawVisualization( const UActorComponent* Co
 				Corners[2] =  P2;
 				Corners[3] =  P3;
 
-				DrawBarnRect(P0, P1, P2, P3);
+				DrawBarnRect(P0, P1, P2, P3, ElementColor);
 			}
 			// -SourceWidth
 			{
@@ -85,7 +93,7 @@ void FRectLightComponentVisualizer::DrawVisualization( const UActorComponent* Co
 				Corners[4] =  P2;
 				Corners[5] =  P3;
 
-				DrawBarnRect(P0, P1, P2, P3);
+				DrawBarnRect(P0, P1, P2, P3, ElementColor);
 			}
 			// -SourceHeight
 			{
@@ -96,12 +104,41 @@ void FRectLightComponentVisualizer::DrawVisualization( const UActorComponent* Co
 				Corners[6] =  P3;
 				Corners[7] =  P2;
 
-				DrawBarnRect(P0, P1, P2, P3);
+				DrawBarnRect(P0, P1, P2, P3, ElementColor);
 			
 			}
 
 			DrawWireBox(PDI, LightTM.ToMatrixNoScale(), Box, ElementColor, SDPG_World);
 
+			if(GVisualizeCullingBarnDoors != 0)
+			{
+				const FColor CullingRectColor(255, 0, 0, 255);
+
+				{
+					float HorizontalBarnExtent;
+					float HorizontalBarnDepth;
+					CalculateRectLightCullingBarnExtentAndDepth(RectLightComp->SourceWidth, RectLightComp->BarnDoorLength, AngleRad, RectLightComp->AttenuationRadius, HorizontalBarnExtent, HorizontalBarnDepth);
+
+					TStaticArray<FVector, 8> HorizontalCorners;
+					CalculateRectLightBarnCorners(RectLightComp->SourceWidth, RectLightComp->SourceHeight, HorizontalBarnExtent, HorizontalBarnDepth, HorizontalCorners);
+
+					DrawBarnRect(HorizontalCorners[0], HorizontalCorners[2], HorizontalCorners[3], HorizontalCorners[1], CullingRectColor);
+					DrawBarnRect(HorizontalCorners[4], HorizontalCorners[6], HorizontalCorners[7], HorizontalCorners[5], CullingRectColor);
+				}
+
+				{
+					float VerticalBarnExtent;
+					float VerticalBarnDepth;
+					CalculateRectLightCullingBarnExtentAndDepth(RectLightComp->SourceHeight, RectLightComp->BarnDoorLength, AngleRad, RectLightComp->AttenuationRadius, VerticalBarnExtent, VerticalBarnDepth);
+
+					TStaticArray<FVector, 8> VerticalCorners;
+					CalculateRectLightBarnCorners(RectLightComp->SourceWidth, RectLightComp->SourceHeight, VerticalBarnExtent, VerticalBarnDepth, VerticalCorners);
+
+					DrawBarnRect(VerticalCorners[0], VerticalCorners[4], VerticalCorners[6], VerticalCorners[2], CullingRectColor);
+					DrawBarnRect(VerticalCorners[5], VerticalCorners[7], VerticalCorners[3], VerticalCorners[1], CullingRectColor);
+				}
+
+			}
 		}
 	}
 }
diff --git a/Engine/Source/Runtime/Engine/Classes/Components/RectLightComponent.h b/Engine/Source/Runtime/Engine/Classes/Components/RectLightComponent.h
--- a/Engine/Source/Runtime/Engine/Classes/Components/RectLightComponent.h
+++ b/Engine/Source/Runtime/Engine/Classes/Components/RectLightComponent.h
@@ -11,6 +11,9 @@
 
 float ENGINE_API GetRectLightBarnDoorMaxAngle();
 
+void ENGINE_API CalculateRectLightCullingBarnExtentAndDepth(float Size, float Length, float AngleRad, float Radius, float& OutExtent, float& OutDepth);
+void ENGINE_API CalculateRectLightBarnCorners(float SourceWidth, float SourceHeight, float BarnExtent, float BarnDepth, TStaticArray<FVector, 8>& OutCorners);
+
 class FLightSceneProxy;
 
 /**
diff --git a/Engine/Source/Runtime/Engine/Private/Components/RectLightComponent.cpp b/Engine/Source/Runtime/Engine/Private/Components/RectLightComponent.cpp
--- a/Engine/Source/Runtime/Engine/Private/Components/RectLightComponent.cpp
+++ b/Engine/Source/Runtime/Engine/Private/Components/RectLightComponent.cpp
@@ -22,6 +22,67 @@ float GetRectLightBarnDoorMaxAngle()
 	return 88.f;
 }
 
+static float SolveQuadraticEq(float A, float B, float C)
+{
+	float Disc = B * B - 4.0f * C * A;
+	if (Disc > UE_KINDA_SMALL_NUMBER)
+	{
+		Disc = FMath::Sqrt(Disc);
+		const float Denom = 1.0f / (2.0f * A);
+		float Root0 = (-B + Disc) * Denom;
+		float Root1 = (-B - Disc) * Denom;
+
+		return Root0;
+	}
+
+	return -1;
+}
+
+void CalculateRectLightCullingBarnExtentAndDepth(float Size, float Length, float AngleRad, float Radius, float& OutExtent, float& OutDepth)
+{
+	float T = Size / 2.0f;
+
+	// 1. calculate opposite side (law of cosines)
+	float A = Size;
+	float B = Length;
+	float C = FMath::Sqrt(A * A + B * B - 2 * A * B * FMath::Cos(AngleRad + UE_HALF_PI));
+
+	// 2. calculate angle between rect plane and shadow boundary (law of sines)
+	float AuxAngleRad = FMath::Asin(B * FMath::Sin(AngleRad + UE_HALF_PI) / C);
+
+	// 3. calculate shadow boundary line
+	float M = FMath::Tan(AuxAngleRad);
+	float K = M * T;
+
+	// 4. intersect shadow boundary line with circle
+	float X = SolveQuadraticEq(M * M + 1, 2 * M * K, K * K - Radius * Radius);
+	float Y = M * X + K;
+
+	if (FMath::Sqrt((X + T) * (X + T) + Y * Y) >= C)
+	{
+		OutExtent = X - T;
+		OutDepth = Y;
+	}
+	else
+	{
+		// if intersection is closer than regular barn doors, fallback to base extent / depth 
+		OutExtent = FMath::Sin(AngleRad) * Length;
+		OutDepth = FMath::Cos(AngleRad) * Length;
+	}
+}
+
+void CalculateRectLightBarnCorners(float SourceWidth, float SourceHeight, float BarnExtent, float BarnDepth, TStaticArray<FVector, 8>& OutCorners)
+{
+	OutCorners[0] = FVector(0.0f, +0.5f * SourceWidth, +0.5f * SourceHeight);
+	OutCorners[1] = FVector(0.0f, +0.5f * SourceWidth, -0.5f * SourceHeight);
+	OutCorners[2] = FVector(BarnDepth, +0.5f * SourceWidth + BarnExtent, +0.5f * SourceHeight + BarnExtent);
+	OutCorners[3] = FVector(BarnDepth, +0.5f * SourceWidth + BarnExtent, -0.5f * SourceHeight - BarnExtent);
+	OutCorners[4] = FVector(0.0f, -0.5f * SourceWidth, +0.5f * SourceHeight);
+	OutCorners[5] = FVector(0.0f, -0.5f * SourceWidth, -0.5f * SourceHeight);
+	OutCorners[6] = FVector(BarnDepth, -0.5f * SourceWidth - BarnExtent, +0.5f * SourceHeight + BarnExtent);
+	OutCorners[7] = FVector(BarnDepth, -0.5f * SourceWidth - BarnExtent, -0.5f * SourceHeight - BarnExtent);
+}
+
 URectLightComponent::URectLightComponent(const FObjectInitializer& ObjectInitializer)
 	: Super(ObjectInitializer)
 	, RayTracingData(new FRectLightRayTracingData())
diff --git a/Engine/Source/Runtime/Renderer/Private/LightGridInjection.cpp b/Engine/Source/Runtime/Renderer/Private/LightGridInjection.cpp
--- a/Engine/Source/Runtime/Renderer/Private/LightGridInjection.cpp
+++ b/Engine/Source/Runtime/Renderer/Private/LightGridInjection.cpp
@@ -71,6 +71,14 @@ FAutoConsoleVariableRef CVarLightGridHZBCull(
 	ECVF_Scalability | ECVF_RenderThreadSafe
 );
 
+int32 GLightGridRefineRectLightBounds = 1;
+FAutoConsoleVariableRef CVarLightGridRefineRectLightBounds(
+	TEXT("r.Forward.LightGridDebug.RectLightBounds"),
+	GLightGridRefineRectLightBounds,
+	TEXT("Whether to refine rect light bounds (should only be disabled for debugging purposes)."),
+	ECVF_RenderThreadSafe
+);
+
 int32 GMaxCulledLightsPerCell = 32;
 FAutoConsoleVariableRef CVarMaxCulledLightsPerCell(
 	TEXT("r.Forward.MaxCulledLightsPerCell"),
@@ -470,39 +478,32 @@ static void CalculateRectLightCullingPlanes(const FRectLightSceneProxy* RectProx
 {
 	const float BarnMaxAngle = GetRectLightBarnDoorMaxAngle();
 	const float AngleRad = FMath::DegreesToRadians(FMath::Clamp(RectProxy->BarnDoorAngle, 0.f, BarnMaxAngle));
-	const float BarnDepth = FMath::Cos(AngleRad) * RectProxy->BarnDoorLength;
-	const float BarnExtent = FMath::Sin(AngleRad) * RectProxy->BarnDoorLength;
-
-	FVector Corners[8];
-
-	// corners position based on rect source size and barn door parameters (frustum shape)
-	Corners[0] = FVector(0.0f, +0.5f * RectProxy->SourceWidth, +0.5f * RectProxy->SourceHeight);
-	Corners[1] = FVector(0.0f, +0.5f * RectProxy->SourceWidth, -0.5f * RectProxy->SourceHeight);
-	Corners[2] = FVector(BarnDepth, +0.5f * RectProxy->SourceWidth + BarnExtent, +0.5f * RectProxy->SourceHeight + BarnExtent);
-	Corners[3] = FVector(BarnDepth, +0.5f * RectProxy->SourceWidth + BarnExtent, -0.5f * RectProxy->SourceHeight - BarnExtent);
-	Corners[4] = FVector(0.0f, -0.5f * RectProxy->SourceWidth, +0.5f * RectProxy->SourceHeight);
-	Corners[5] = FVector(0.0f, -0.5f * RectProxy->SourceWidth, -0.5f * RectProxy->SourceHeight);
-	Corners[6] = FVector(BarnDepth, -0.5f * RectProxy->SourceWidth - BarnExtent, +0.5f * RectProxy->SourceHeight + BarnExtent);
-	Corners[7] = FVector(BarnDepth, -0.5f * RectProxy->SourceWidth - BarnExtent, -0.5f * RectProxy->SourceHeight - BarnExtent);
-
-	// adjust corners to account for penumbra
-	auto ApplyPenumbraAdjustment = [RectProxy](const FVector& A, FVector& B) -> FVector
-		{
-			FVector Tmp = B - A;
-			double Len = Tmp.Length();
-			Tmp.Normalize();
-			return A + Tmp * FMath::Max<double>(RectProxy->Radius, Len);
-		};
-
-	Corners[7] = ApplyPenumbraAdjustment(Corners[0], Corners[7]);
-	Corners[6] = ApplyPenumbraAdjustment(Corners[1], Corners[6]);
-	Corners[3] = ApplyPenumbraAdjustment(Corners[4], Corners[3]);
-	Corners[2] = ApplyPenumbraAdjustment(Corners[5], Corners[2]);
-
-	OutPlanes.Add(FPlane(Corners[1], Corners[0], Corners[3])); // right
-	OutPlanes.Add(FPlane(Corners[5], Corners[7], Corners[4])); // left
-	OutPlanes.Add(FPlane(Corners[4], Corners[6], Corners[0])); // top
-	OutPlanes.Add(FPlane(Corners[1], Corners[3], Corners[5])); // bottom
+
+	// horizontal barn doors
+	{
+		float HorizontalBarnExtent;
+		float HorizontalBarnDepth;
+		CalculateRectLightCullingBarnExtentAndDepth(RectProxy->SourceWidth, RectProxy->BarnDoorLength, AngleRad, RectProxy->Radius, HorizontalBarnExtent, HorizontalBarnDepth);
+
+		TStaticArray<FVector, 8> Corners;
+		CalculateRectLightBarnCorners(RectProxy->SourceWidth, RectProxy->SourceHeight, HorizontalBarnExtent, HorizontalBarnDepth, Corners);
+
+		OutPlanes.Add(FPlane(Corners[1], Corners[0], Corners[3])); // right
+		OutPlanes.Add(FPlane(Corners[5], Corners[7], Corners[4])); // left
+	}
+	
+	// vertical barn doors
+	{
+		float VerticalBarnExtent;
+		float VerticalBarnDepth;
+		CalculateRectLightCullingBarnExtentAndDepth(RectProxy->SourceHeight, RectProxy->BarnDoorLength, AngleRad, RectProxy->Radius, VerticalBarnExtent, VerticalBarnDepth);
+
+		TStaticArray<FVector, 8> Corners;
+		CalculateRectLightBarnCorners(RectProxy->SourceWidth, RectProxy->SourceHeight, VerticalBarnExtent, VerticalBarnDepth, Corners);
+
+		OutPlanes.Add(FPlane(Corners[4], Corners[6], Corners[0])); // top
+		OutPlanes.Add(FPlane(Corners[1], Corners[3], Corners[5])); // bottom
+	}
 
 	check(OutPlanes.Num() == NUM_PLANES_PER_RECT_LIGHT);
 }
@@ -995,7 +996,7 @@ FComputeLightGridOutput FSceneRenderer::ComputeLightGrid(FRDGBuilder& GraphBuild
 
 			FLightGridInjectionCS::FPermutationDomain PermutationVector;
 			PermutationVector.Set<FLightGridInjectionCS::FUseLinkedListDim>(GLightLinkedListCulling != 0);
-			PermutationVector.Set<FLightGridInjectionCS::FRefineRectLightBoundsDim>(bHasRectLights);
+			PermutationVector.Set<FLightGridInjectionCS::FRefineRectLightBoundsDim>(bHasRectLights && (GLightGridRefineRectLightBounds != 0));
 			PermutationVector.Set<FLightGridInjectionCS::FUseHZBCullDim>(GLightGridHZBCull != 0 && View.HZB != nullptr);
 			TShaderMapRef<FLightGridInjectionCS> ComputeShader(View.ShaderMap, PermutationVector);
 
