Fix translucent volume corruption on some platforms due to initial data padding.

#tests enginetest,lyra
#rnx

diff --git a/Engine/Source/Runtime/Renderer/Private/TranslucentLighting.cpp b/Engine/Source/Runtime/Renderer/Private/TranslucentLighting.cpp
--- a/Engine/Source/Runtime/Renderer/Private/TranslucentLighting.cpp
+++ b/Engine/Source/Runtime/Renderer/Private/TranslucentLighting.cpp
@@ -1383,18 +1383,20 @@ void InjectTranslucencyLightingVolume(
 
 		// Batched lights
 		{
-			const TBitArray<SceneRenderingAllocator>& BatchedLocalLights = Collector.InjectionDataPerView[ViewIndex].BatchedLocalLights[VolumeCascadeIndex];
-			FRDGBufferRef BatchedLocalLightsRDG = nullptr;
+			TBitArray<SceneRenderingAllocator>& BatchedLocalLights = Collector.InjectionDataPerView[ViewIndex].BatchedLocalLights[VolumeCascadeIndex];
 			if (BatchedLocalLights.Num() > 0)
 			{
-				const uint32 NumElements = FMath::Max(1, FMath::DivideAndRoundUp(BatchedLocalLights.Num(), 32));
-				BatchedLocalLightsRDG = CreateStructuredBuffer(
+				const uint32 NumUint32Elements = FMath::DivideAndRoundUp(BatchedLocalLights.Num(), 32);
+				const uint32 InitialDataSize = NumUint32Elements * sizeof(uint32);
+				BatchedLocalLights.PadToNum(NumUint32Elements * 32, false);
+
+				FRDGBufferRef BatchedLocalLightsRDG = CreateStructuredBuffer(
 					GraphBuilder,
 					TEXT("TranslucencyLightingVolume.BatchedLocalLights"),
 					sizeof(uint32),
-					NumElements,
+					NumUint32Elements,
 					BatchedLocalLights.GetData(),
-					FMath::DivideAndRoundUp(BatchedLocalLights.Num(), 8)	// Size in bytes of initial data
+					InitialDataSize
 				);
 
 				InjectTranslucencyLightingVolumeBatch(GraphBuilder, View, ViewIndex,
