Moved volumetric fog to use AddDynamicMeshPass.

diff --git a/Engine/Source/Runtime/Renderer/Private/VolumetricFogVoxelization.cpp b/Engine/Source/Runtime/Renderer/Private/VolumetricFogVoxelization.cpp
--- a/Engine/Source/Runtime/Renderer/Private/VolumetricFogVoxelization.cpp
+++ b/Engine/Source/Runtime/Renderer/Private/VolumetricFogVoxelization.cpp
@@ -634,7 +634,6 @@ IPSOCollector* CreatePSOCollectorVoxelizeVolume(ERHIFeatureLevel::Type FeatureLe
 FRegisterPSOCollectorCreateFunction RegisterPSOCollectorVoxelizeVolume(&CreatePSOCollectorVoxelizeVolume, EShadingPath::Deferred, VoxelizeVolumePassName);
 
 void VoxelizeVolumePrimitive(FVoxelizeVolumeMeshProcessor& PassMeshProcessor,
-	FRHICommandListImmediate& RHICmdList,
 	const FViewInfo& View,
 	FIntVector VolumetricFogViewGridSize,
 	FVector GridZParams,
@@ -743,58 +742,50 @@ void FSceneRenderer::VoxelizeFogVolumePrimitives(
 			GQuadMeshVertexFactory->InitResource(GraphBuilder.RHICmdList);
 		}
 
-		GraphBuilder.AddPass(
-			RDG_EVENT_NAME("VoxelizeVolumePrimitives"),
-			PassParameters,
-			ERDGPassFlags::Raster,
-			[PassParameters, Scene = Scene, &View, VolumetricFogViewGridSize, IntegrationData, VolumetricFogDistance, GridZParams](FRHICommandListImmediate& RHICmdList)
+		// Set the sub region of the texture according to the current dynamic resolution scale.
+		const FIntRect ViewRect(0, 0, VolumetricFogViewGridSize.X, VolumetricFogViewGridSize.Y);
+
+		AddDrawDynamicMeshPass(GraphBuilder, RDG_EVENT_NAME("VoxelizeVolumePrimitives"), PassParameters, View, ViewRect,
+			[&View, VolumetricFogDistance, &VolumetricFogViewGridSize, &GridZParams](FDynamicPassMeshDrawListContext* DynamicMeshPassContext)
 			{
+				FVoxelizeVolumeMeshProcessor PassMeshProcessor(
+					View.Family->Scene->GetRenderScene(),
+					View.GetFeatureLevel(),
+					&View,
+					DynamicMeshPassContext);
 
-			DrawDynamicMeshPass(View, RHICmdList,
-				[&View, VolumetricFogDistance, &RHICmdList, &VolumetricFogViewGridSize, &GridZParams](FDynamicPassMeshDrawListContext* DynamicMeshPassContext)
-				{
-					FVoxelizeVolumeMeshProcessor PassMeshProcessor(
-						View.Family->Scene->GetRenderScene(),
-						View.GetFeatureLevel(),
-						&View,
-						DynamicMeshPassContext);
+				const bool bShouldRenderHeterogeneousVolumes = ShouldRenderHeterogeneousVolumesForView(View);
 
-					const bool bShouldRenderHeterogeneousVolumes = ShouldRenderHeterogeneousVolumesForView(View);
+				for (int32 MeshBatchIndex = 0; MeshBatchIndex < View.VolumetricMeshBatches.Num(); ++MeshBatchIndex)
+				{
+					// Skip volumes flagged as rendered with HeterogenousVolumes
+					const FMeshBatch* Mesh = View.VolumetricMeshBatches[MeshBatchIndex].Mesh;
+					const FPrimitiveSceneProxy* PrimitiveSceneProxy = View.VolumetricMeshBatches[MeshBatchIndex].Proxy;
+					if (ShouldRenderMeshBatchWithHeterogeneousVolumes(Mesh, PrimitiveSceneProxy, View.GetFeatureLevel()))
+					{
+						continue;
+					}
 
-					// Set the sub region of the texture according to the current dynamic resolution scale.
-					RHICmdList.SetViewport(0.0f, 0.0f, 0.0f, VolumetricFogViewGridSize.X, VolumetricFogViewGridSize.Y, 1.0f);
+					const FPrimitiveSceneInfo* PrimitiveSceneInfo = PrimitiveSceneProxy->GetPrimitiveSceneInfo();
+					const FBoxSphereBounds Bounds = PrimitiveSceneProxy->GetBounds();
 
-					for (int32 MeshBatchIndex = 0; MeshBatchIndex < View.VolumetricMeshBatches.Num(); ++MeshBatchIndex)
+					if ((View.ViewMatrices.GetViewOrigin() - Bounds.Origin).SizeSquared() < (VolumetricFogDistance + Bounds.SphereRadius) * (VolumetricFogDistance + Bounds.SphereRadius))
 					{
-						// Skip volumes flagged as rendered with HeterogenousVolumes
-						const FMeshBatch* Mesh = View.VolumetricMeshBatches[MeshBatchIndex].Mesh;
-						const FPrimitiveSceneProxy* PrimitiveSceneProxy = View.VolumetricMeshBatches[MeshBatchIndex].Proxy;
-						if (ShouldRenderMeshBatchWithHeterogeneousVolumes(Mesh, PrimitiveSceneProxy, View.GetFeatureLevel()))
-						{
-							continue;
-						}
-
-						const FPrimitiveSceneInfo* PrimitiveSceneInfo = PrimitiveSceneProxy->GetPrimitiveSceneInfo();
-						const FBoxSphereBounds Bounds = PrimitiveSceneProxy->GetBounds();
-
-						if ((View.ViewMatrices.GetViewOrigin() - Bounds.Origin).SizeSquared() < (VolumetricFogDistance + Bounds.SphereRadius) * (VolumetricFogDistance + Bounds.SphereRadius))
-						{
-							VoxelizeVolumePrimitive(PassMeshProcessor, RHICmdList, View, VolumetricFogViewGridSize, GridZParams, PrimitiveSceneProxy, *Mesh);
-						}
+						VoxelizeVolumePrimitive(PassMeshProcessor, View, VolumetricFogViewGridSize, GridZParams, PrimitiveSceneProxy, *Mesh);
 					}
-				},
-
-				// Force off instanced stereo.
-				// With instanced stereo on, primitives were being drawn to the left eye twice,  thickening the fog more in the
-				// left eye.  It seemed better to force off instanced stereo anyway because of cache coherency in the 3d grids,
-				// which are per-eye (far away in cache). The engine is already instancing across slices, which should be nearby
-				// in cache).
-				//
-				// It may be a tradeoff where small primitives do better with instancing (their texture lookups stay cached)
-				// and large ones covering lots of the voxel grid do worse (grid writes use up too much cache?), and could be
-				// decided based on bounds?  GPUs presumably may have separate caches for read-only data in a way that changes
-				// this tradeoff as well.
-				true /*bForceInstanceStereoOff*/);
-		});
+				}
+			},
+
+			// Force off instanced stereo.
+			// With instanced stereo on, primitives were being drawn to the left eye twice,  thickening the fog more in the
+			// left eye.  It seemed better to force off instanced stereo anyway because of cache coherency in the 3d grids,
+			// which are per-eye (far away in cache). The engine is already instancing across slices, which should be nearby
+			// in cache).
+			//
+			// It may be a tradeoff where small primitives do better with instancing (their texture lookups stay cached)
+			// and large ones covering lots of the voxel grid do worse (grid writes use up too much cache?), and could be
+			// decided based on bounds?  GPUs presumably may have separate caches for read-only data in a way that changes
+			// this tradeoff as well.
+			true /*bForceInstanceStereoOff*/);
 	}
 }
