Implemented async task support for Slate.

diff --git a/Engine/Source/Runtime/Engine/Private/UserInterface/Canvas.cpp b/Engine/Source/Runtime/Engine/Private/UserInterface/Canvas.cpp
--- a/Engine/Source/Runtime/Engine/Private/UserInterface/Canvas.cpp
+++ b/Engine/Source/Runtime/Engine/Private/UserInterface/Canvas.cpp
@@ -576,7 +576,7 @@ FCanvasRenderThreadScope::~FCanvasRenderThreadScope()
 				ERenderTargetLoadAction::ELoad
 			);
 
-			GraphBuilder.AddPass(RDG_EVENT_NAME("MSAAResolveCanvas"), PassParameters, ERDGPassFlags::Raster, [](FRHICommandList&) {});
+			GraphBuilder.AddPass(RDG_EVENT_NAME("MSAAResolveCanvas"), PassParameters, ERDGPassFlags::Raster, [](FRDGAsyncTask, FRHICommandList&) {});
 		}
 
 		GraphBuilder.Execute();
diff --git a/Engine/Source/Runtime/Engine/Public/CanvasRender.h b/Engine/Source/Runtime/Engine/Public/CanvasRender.h
--- a/Engine/Source/Runtime/Engine/Public/CanvasRender.h
+++ b/Engine/Source/Runtime/Engine/Public/CanvasRender.h
@@ -14,7 +14,7 @@ public:
 	void AddPass(FRDGEventName&& PassName, const ParameterStructType* PassParameters, ExecuteLambdaType&& ExecuteLambda)
 	{
 		GraphBuilder.AddPass(Forward<FRDGEventName>(PassName), PassParameters, ERDGPassFlags::Raster,
-			[LocalScissorRect = ScissorRect, LocalViewportRect = ViewportRect, LocalExecuteLambda = Forward<ExecuteLambdaType&&>(ExecuteLambda)](FRHICommandList& RHICmdList)
+			[LocalScissorRect = ScissorRect, LocalViewportRect = ViewportRect, LocalExecuteLambda = Forward<ExecuteLambdaType&&>(ExecuteLambda)](FRDGAsyncTask, FRHICommandList& RHICmdList)
 		{
 			RHICmdList.SetViewport(static_cast<float>(LocalViewportRect.Min.X), static_cast<float>(LocalViewportRect.Min.Y), 0.0f, static_cast<float>(LocalViewportRect.Max.X), static_cast<float>(LocalViewportRect.Max.Y), 1.0f);
 
diff --git a/Engine/Source/Runtime/Renderer/Private/Renderer.cpp b/Engine/Source/Runtime/Renderer/Private/Renderer.cpp
--- a/Engine/Source/Runtime/Renderer/Private/Renderer.cpp
+++ b/Engine/Source/Runtime/Renderer/Private/Renderer.cpp
@@ -402,6 +402,10 @@ void FRendererModule::DrawTileMesh(FCanvasRenderContext& RenderContext, FMeshPas
 			PassParameters->RenderTargets[1] = FRenderTargetBinding(DepthAux, ERenderTargetLoadAction::EClear);
 		}
 
+		// Disable parallel setup tasks since we are only processing one mesh (not worth the task launch cost).
+		const bool bForceStereoInstancingOff = false;
+		const bool bForceParallelSetupOff = true;
+
 		// handle translucent material blend modes, not relevant in MaterialTexCoordScalesAnalysis since it outputs the scales.
 		if (ViewFamily->GetDebugViewShaderMode() == DVSM_OutputMaterialTextureScales)
 		{
@@ -414,21 +418,19 @@ void FRendererModule::DrawTileMesh(FCanvasRenderContext& RenderContext, FMeshPas
 			{
 				PassParameters->DebugViewMode = CreateDebugViewModePassUniformBuffer(GraphBuilder, View, nullptr);
 
-				RenderContext.AddPass(RDG_EVENT_NAME("OutputMaterialTextureScales"), PassParameters,
-					[Scene, &View, &Mesh](FRHICommandList& RHICmdList)
+				AddDrawDynamicMeshPass(GraphBuilder, RDG_EVENT_NAME("OutputMaterialTextureScales"), PassParameters, View, View.ViewRect,
+					[Scene, &View, &Mesh](FMeshPassDrawListContext* InDrawListContext)
 				{
-					DrawDynamicMeshPass(View, RHICmdList, [&](FMeshPassDrawListContext* InDrawListContext)
-					{
-						FDebugViewModeMeshProcessor PassMeshProcessor(
-							Scene,
-							View.GetFeatureLevel(),
-							&View,
-							false,
-							InDrawListContext);
-						const uint64 DefaultBatchElementMask = ~0ull;
-						PassMeshProcessor.AddMeshBatch(Mesh, DefaultBatchElementMask, nullptr);
-					});
-				});
+					FDebugViewModeMeshProcessor PassMeshProcessor(
+						Scene,
+						View.GetFeatureLevel(),
+						&View,
+						false,
+						InDrawListContext);
+					const uint64 DefaultBatchElementMask = ~0ull;
+					PassMeshProcessor.AddMeshBatch(Mesh, DefaultBatchElementMask, nullptr);
+
+				}, bForceStereoInstancingOff, bForceParallelSetupOff);
 			}
 #endif // WITH_DEBUG_VIEW_MODES
 		}
@@ -438,48 +440,42 @@ void FRendererModule::DrawTileMesh(FCanvasRenderContext& RenderContext, FMeshPas
 			{
 				PassParameters->TranslucentBasePass = CreateTranslucentBasePassUniformBuffer(GraphBuilder, Scene, View);
 
-				RenderContext.AddPass(RDG_EVENT_NAME("TranslucentDeferred"), PassParameters,
-					[Scene, &View, &Mesh, DrawRenderState, bUse128bitRT](FRHICommandList& RHICmdList)
+				AddDrawDynamicMeshPass(GraphBuilder, RDG_EVENT_NAME("TranslucentDeferred"), PassParameters, View, View.ViewRect,
+					[Scene, &View, &Mesh, DrawRenderState, bUse128bitRT](FMeshPassDrawListContext* DynamicMeshPassContext)
 				{
-					DrawDynamicMeshPass(View, RHICmdList, [&](FDynamicPassMeshDrawListContext* DynamicMeshPassContext)
-					{
-						FBasePassMeshProcessor PassMeshProcessor(
-							EMeshPass::BasePass,
-							Scene,
-							View.GetFeatureLevel(),
-							&View,
-							DrawRenderState,
-							DynamicMeshPassContext,
-							bUse128bitRT ? FBasePassMeshProcessor::EFlags::bRequires128bitRT : FBasePassMeshProcessor::EFlags::None,
-							ETranslucencyPass::TPT_AllTranslucency);
-
-						const uint64 DefaultBatchElementMask = ~0ull;
-						PassMeshProcessor.AddMeshBatch(Mesh, DefaultBatchElementMask, nullptr);
-					});
-				});
+					FBasePassMeshProcessor PassMeshProcessor(
+						EMeshPass::BasePass,
+						Scene,
+						View.GetFeatureLevel(),
+						&View,
+						DrawRenderState,
+						DynamicMeshPassContext,
+						bUse128bitRT ? FBasePassMeshProcessor::EFlags::bRequires128bitRT : FBasePassMeshProcessor::EFlags::None,
+						ETranslucencyPass::TPT_AllTranslucency);
+
+					const uint64 DefaultBatchElementMask = ~0ull;
+					PassMeshProcessor.AddMeshBatch(Mesh, DefaultBatchElementMask, nullptr);
+				}, bForceStereoInstancingOff, bForceParallelSetupOff);
 			}
 			else // Mobile
 			{
 				PassParameters->MobileBasePass = CreateMobileBasePassUniformBuffer(GraphBuilder, View, EMobileBasePass::Translucent, EMobileSceneTextureSetupMode::None);
 
-				RenderContext.AddPass(RDG_EVENT_NAME("TranslucentMobile"), PassParameters,
-					[Scene, &View, DrawRenderState, &Mesh](FRHICommandList& RHICmdList)
+				AddDrawDynamicMeshPass(GraphBuilder, RDG_EVENT_NAME("TranslucentMobile"), PassParameters, View, View.ViewRect,
+					[Scene, &View, DrawRenderState, &Mesh](FDynamicPassMeshDrawListContext* DynamicMeshPassContext)
 				{
-					DrawDynamicMeshPass(View, RHICmdList, [&](FDynamicPassMeshDrawListContext* DynamicMeshPassContext)
-					{
-						FMobileBasePassMeshProcessor PassMeshProcessor(
-							EMeshPass::TranslucencyAll,
-							Scene,
-							&View,
-							DrawRenderState,
-							DynamicMeshPassContext,
-							FMobileBasePassMeshProcessor::EFlags::None,
-							ETranslucencyPass::TPT_AllTranslucency);
-
-						const uint64 DefaultBatchElementMask = ~0ull;
-						PassMeshProcessor.AddMeshBatch(Mesh, DefaultBatchElementMask, nullptr);
-					});
-				});
+					FMobileBasePassMeshProcessor PassMeshProcessor(
+						EMeshPass::TranslucencyAll,
+						Scene,
+						&View,
+						DrawRenderState,
+						DynamicMeshPassContext,
+						FMobileBasePassMeshProcessor::EFlags::None,
+						ETranslucencyPass::TPT_AllTranslucency);
+
+					const uint64 DefaultBatchElementMask = ~0ull;
+					PassMeshProcessor.AddMeshBatch(Mesh, DefaultBatchElementMask, nullptr);
+				}, bForceStereoInstancingOff, bForceParallelSetupOff);
 			}
 		}
 		// handle opaque materials
@@ -494,22 +490,19 @@ void FRendererModule::DrawTileMesh(FCanvasRenderContext& RenderContext, FMeshPas
 				ensureMsgf(HitProxyId == Mesh.BatchHitProxyId, TEXT("Only Mesh.BatchHitProxyId is used for hit testing."));
 
 #if WITH_EDITOR
-				RenderContext.AddPass(RDG_EVENT_NAME("HitTesting"), PassParameters,
-					[Scene, &View, DrawRenderState, &Mesh](FRHICommandList& RHICmdList)
+				AddDrawDynamicMeshPass(GraphBuilder, RDG_EVENT_NAME("HitTesting"), PassParameters, View, View.ViewRect,
+					[Scene, &View, DrawRenderState, &Mesh](FDynamicPassMeshDrawListContext* DynamicMeshPassContext)
 				{
-					DrawDynamicMeshPass(View, RHICmdList, [&](FDynamicPassMeshDrawListContext* DynamicMeshPassContext)
-					{
-						FHitProxyMeshProcessor PassMeshProcessor(
-							Scene,
-							&View,
-							false,
-							DrawRenderState,
-							DynamicMeshPassContext);
-
-						const uint64 DefaultBatchElementMask = ~0ull;
-						PassMeshProcessor.AddMeshBatch(Mesh, DefaultBatchElementMask, nullptr);
-					});
-				});
+					FHitProxyMeshProcessor PassMeshProcessor(
+						Scene,
+						&View,
+						false,
+						DrawRenderState,
+						DynamicMeshPassContext);
+
+					const uint64 DefaultBatchElementMask = ~0ull;
+					PassMeshProcessor.AddMeshBatch(Mesh, DefaultBatchElementMask, nullptr);
+				}, bForceStereoInstancingOff, bForceParallelSetupOff);
 #endif
 			}
 			else
@@ -518,47 +511,40 @@ void FRendererModule::DrawTileMesh(FCanvasRenderContext& RenderContext, FMeshPas
 				{
 					PassParameters->OpaqueBasePass = CreateOpaqueBasePassUniformBuffer(GraphBuilder, View);
 
-					RenderContext.AddPass(RDG_EVENT_NAME("OpaqueDeferred"), PassParameters,
-						[Scene, &View, DrawRenderState, &Mesh, bUse128bitRT](FRHICommandList& RHICmdList)
+					AddDrawDynamicMeshPass(GraphBuilder, RDG_EVENT_NAME("OpaqueDeferred"), PassParameters, View, View.ViewRect,
+						[Scene, &View, DrawRenderState, &Mesh, bUse128bitRT](FDynamicPassMeshDrawListContext* DynamicMeshPassContext)
 					{
-						DrawDynamicMeshPass(View, RHICmdList,
-							[&](FDynamicPassMeshDrawListContext* DynamicMeshPassContext)
-						{
-							FBasePassMeshProcessor PassMeshProcessor(
-								EMeshPass::BasePass,
-								Scene,
-								View.GetFeatureLevel(),
-								&View,
-								DrawRenderState,
-								DynamicMeshPassContext,
-								bUse128bitRT ? FBasePassMeshProcessor::EFlags::bRequires128bitRT : FBasePassMeshProcessor::EFlags::None);
-
-							const uint64 DefaultBatchElementMask = ~0ull;
-							PassMeshProcessor.AddMeshBatch(Mesh, DefaultBatchElementMask, nullptr);
-						});
+						FBasePassMeshProcessor PassMeshProcessor(
+							EMeshPass::BasePass,
+							Scene,
+							View.GetFeatureLevel(),
+							&View,
+							DrawRenderState,
+							DynamicMeshPassContext,
+							bUse128bitRT ? FBasePassMeshProcessor::EFlags::bRequires128bitRT : FBasePassMeshProcessor::EFlags::None);
+
+						const uint64 DefaultBatchElementMask = ~0ull;
+						PassMeshProcessor.AddMeshBatch(Mesh, DefaultBatchElementMask, nullptr);
 					});
 				}
 				else // Mobile
 				{
 					PassParameters->MobileBasePass = CreateMobileBasePassUniformBuffer(GraphBuilder, View, EMobileBasePass::Opaque, EMobileSceneTextureSetupMode::None);
 
-					RenderContext.AddPass(RDG_EVENT_NAME("OpaqueMobile"), PassParameters,
-						[Scene, &View, DrawRenderState, &Mesh](FRHICommandList& RHICmdList)
+					AddDrawDynamicMeshPass(GraphBuilder, RDG_EVENT_NAME("OpaqueMobile"), PassParameters, View, View.ViewRect,
+						[Scene, &View, DrawRenderState, &Mesh](FDynamicPassMeshDrawListContext* DynamicMeshPassContext)
 					{
-						DrawDynamicMeshPass(View, RHICmdList, [&](FDynamicPassMeshDrawListContext* DynamicMeshPassContext)
-						{
-							FMobileBasePassMeshProcessor PassMeshProcessor(
-								EMeshPass::BasePass,
-								Scene,
-								&View,
-								DrawRenderState,
-								DynamicMeshPassContext,
-								FMobileBasePassMeshProcessor::EFlags::CanReceiveCSM | FMobileBasePassMeshProcessor::EFlags::ForcePassDrawRenderState);
-
-							const uint64 DefaultBatchElementMask = ~0ull;
-							PassMeshProcessor.AddMeshBatch(Mesh, DefaultBatchElementMask, nullptr);
-						});
-					});
+						FMobileBasePassMeshProcessor PassMeshProcessor(
+							EMeshPass::BasePass,
+							Scene,
+							&View,
+							DrawRenderState,
+							DynamicMeshPassContext,
+							FMobileBasePassMeshProcessor::EFlags::CanReceiveCSM | FMobileBasePassMeshProcessor::EFlags::ForcePassDrawRenderState);
+
+						const uint64 DefaultBatchElementMask = ~0ull;
+						PassMeshProcessor.AddMeshBatch(Mesh, DefaultBatchElementMask, nullptr);
+					}, bForceStereoInstancingOff, bForceParallelSetupOff);
 				}
 			}
 		}
diff --git a/Engine/Source/Runtime/Renderer/Public/MeshPassProcessor.inl b/Engine/Source/Runtime/Renderer/Public/MeshPassProcessor.inl
--- a/Engine/Source/Runtime/Renderer/Public/MeshPassProcessor.inl
+++ b/Engine/Source/Runtime/Renderer/Public/MeshPassProcessor.inl
@@ -353,7 +353,8 @@ void AddDrawDynamicMeshPass(
 	const FSceneView& View,
 	FIntRect ViewRect,
 	const LambdaType& BuildPassProcessorLambda,
-	bool bForceStereoInstancingOff = false)
+	bool bForceStereoInstancingOff = false,
+	bool bForceParallelSetupOff = false)
 {
 	// We assume all dynamic passes are in stereo if it is enabled in the view, so we apply ISR to them
 	const uint32 InstanceFactor = (!bForceStereoInstancingOff && View.IsInstancedStereoPass()) ? 2 : 1;
@@ -373,13 +374,14 @@ void AddDrawDynamicMeshPass(
 		TRACE_CPUPROFILER_EVENT_SCOPE(SetupDynamicMeshPass);
 		FDynamicPassMeshDrawListContext DynamicMeshPassContext(Context.DynamicMeshDrawCommandStorage, Context.VisibleMeshDrawCommands, Context.GraphicsMinimalPipelineStateSet, Context.NeedsShaderInitialisation);
 		BuildPassProcessorLambda(&DynamicMeshPassContext);
-	});
+
+	}, !bForceParallelSetupOff);
 
 	GraphBuilder.AddPass(
 		MoveTemp(EventName),
 		PassParameters,
 		ERDGPassFlags::Raster,
-		[&Context, &View, ViewRect, InstanceFactor](FRHICommandList& RHICmdList)
+		[&Context, &View, ViewRect, InstanceFactor](FRDGAsyncTask, FRHICommandList& RHICmdList)
 	{
 		TRACE_CPUPROFILER_EVENT_SCOPE(SetupDynamicMeshPass);
 		RHICmdList.SetViewport(ViewRect.Min.X, ViewRect.Min.Y, 0.0f, ViewRect.Max.X, ViewRect.Max.Y, 1.0f);
diff --git a/Engine/Source/Runtime/SlateCore/Private/Rendering/SlateDrawBuffer.cpp b/Engine/Source/Runtime/SlateCore/Private/Rendering/SlateDrawBuffer.cpp
--- a/Engine/Source/Runtime/SlateCore/Private/Rendering/SlateDrawBuffer.cpp
+++ b/Engine/Source/Runtime/SlateCore/Private/Rendering/SlateDrawBuffer.cpp
@@ -75,18 +75,23 @@ bool FSlateDrawBuffer::Lock()
 	return bIsLock;
 }
 
-void FSlateDrawBuffer::Unlock()
+void FSlateDrawBuffer::Unlock(const UE::Tasks::FTask& PrerequisiteTask)
 {
-	ensureMsgf(IsLocked(), TEXT("The SlateDrawBuffer should be lock before modifying it."));
-	FScopeLock ScopeLock(&GCLock);
-
-	// Rendering doesn't need the batch data anymore
-	for (TSharedRef<FSlateWindowElementList>& ExistingElementList : WindowElementLists)
+	// Launch an inline task that will unlock the buffer when the prerequisite tasks have completed.
+	UE::Tasks::Launch(UE_SOURCE_LOCATION, [this]
 	{
-		ExistingElementList->ResetElementList();
-	}
+		ensureMsgf(IsLocked(), TEXT("The SlateDrawBuffer should be lock before modifying it."));
+		FScopeLock ScopeLock(&GCLock);
+
+		// Rendering doesn't need the batch data anymore
+		for (TSharedRef<FSlateWindowElementList>& ExistingElementList : WindowElementLists)
+		{
+			ExistingElementList->ResetElementList();
+		}
+
+		bIsLocked = false;
 
-	bIsLocked = false;
+	}, PrerequisiteTask, UE::Tasks::ETaskPriority::Normal, UE::Tasks::EExtendedTaskPriority::Inline);
 }
 
 void FSlateDrawBuffer::AddReferencedObjects(FReferenceCollector& Collector)
diff --git a/Engine/Source/Runtime/SlateCore/Public/Rendering/SlateDrawBuffer.h b/Engine/Source/Runtime/SlateCore/Public/Rendering/SlateDrawBuffer.h
--- a/Engine/Source/Runtime/SlateCore/Public/Rendering/SlateDrawBuffer.h
+++ b/Engine/Source/Runtime/SlateCore/Public/Rendering/SlateDrawBuffer.h
@@ -6,6 +6,7 @@
 #include "UObject/GCObject.h"
 #include "HAL/CriticalSection.h"
 #include "Types/SlateVector2.h"
+#include "Tasks/Task.h"
 
 #include <atomic>
 
@@ -59,7 +60,12 @@ public:
 	 *
 	 * @see Lock
 	 */
-	SLATECORE_API void Unlock();
+	SLATECORE_API void Unlock(const UE::Tasks::FTask& PrerequisiteTask);
+
+	void Unlock()
+	{
+		Unlock(UE::Tasks::FTask{});
+	}
 
 	/** @return true if the buffer is locked. */
 	bool IsLocked() const
diff --git a/Engine/Source/Runtime/SlateRHIRenderer/Private/Slate3DRenderer.cpp b/Engine/Source/Runtime/SlateRHIRenderer/Private/Slate3DRenderer.cpp
--- a/Engine/Source/Runtime/SlateRHIRenderer/Private/Slate3DRenderer.cpp
+++ b/Engine/Source/Runtime/SlateRHIRenderer/Private/Slate3DRenderer.cpp
@@ -76,7 +76,7 @@ void FSlate3DRenderer::ReleaseDrawBuffer(FSlateDrawBuffer& InWindowDrawBuffer)
 	ENQUEUE_RENDER_COMMAND(SlateReleaseDrawBufferCommand)(
 		[DrawBuffer = &InWindowDrawBuffer](FRHICommandListImmediate& RHICmdList)
 	{
-		DrawBuffer->Unlock();
+		DrawBuffer->Unlock(FRDGBuilder::GetAsyncExecuteTask());
 	});
 }
 
diff --git a/Engine/Source/Runtime/SlateRHIRenderer/Private/SlatePostProcessor.cpp b/Engine/Source/Runtime/SlateRHIRenderer/Private/SlatePostProcessor.cpp
--- a/Engine/Source/Runtime/SlateRHIRenderer/Private/SlatePostProcessor.cpp
+++ b/Engine/Source/Runtime/SlateRHIRenderer/Private/SlatePostProcessor.cpp
@@ -262,7 +262,7 @@ void AddSlatePostProcessUpsamplePass(FRDGBuilder& GraphBuilder, const FSlatePost
 		RDG_EVENT_NAME("Upsample"),
 		PassParameters,
 		ERDGPassFlags::Raster,
-		[OutputViewport, InputViewport, PipelineState, PixelShader, ClippingOp = Inputs.ClippingOp, PassParameters](FRHICommandList& RHICmdList)
+		[OutputViewport, InputViewport, PipelineState, PixelShader, ClippingOp = Inputs.ClippingOp, PassParameters](FRDGAsyncTask, FRHICommandList& RHICmdList)
 	{
 		// Stencil clipping will issue its own draw calls.
 		SetSlateClipping(RHICmdList, ClippingOp, OutputViewport.Rect);
diff --git a/Engine/Source/Runtime/SlateRHIRenderer/Private/SlateRHIRenderer.cpp b/Engine/Source/Runtime/SlateRHIRenderer/Private/SlateRHIRenderer.cpp
--- a/Engine/Source/Runtime/SlateRHIRenderer/Private/SlateRHIRenderer.cpp
+++ b/Engine/Source/Runtime/SlateRHIRenderer/Private/SlateRHIRenderer.cpp
@@ -336,7 +336,7 @@ void FSlateRHIRenderer::ReleaseDrawBuffer(FSlateDrawBuffer& WindowDrawBuffer)
 
 	ENQUEUE_RENDER_COMMAND(SlateReleaseDrawBufferCommand)([&WindowDrawBuffer](FRHICommandList& RHICmdList)
 	{
-		WindowDrawBuffer.Unlock();
+		WindowDrawBuffer.Unlock(FRDGBuilder::GetAsyncExecuteTask());
 	});
 }
 
@@ -1325,7 +1325,7 @@ void FSlateRHIRenderer::DrawWindows_Private(FSlateDrawBuffer& WindowDrawBuffer)
 			DrawWindowPassOutputs.Reserve(DrawWindowsCommand->Windows.Num());
 
 			{
-				FRDGBuilder GraphBuilder(RHICmdList, RDG_EVENT_NAME("Slate"));
+				FRDGBuilder GraphBuilder(RHICmdList, RDG_EVENT_NAME("Slate"), ERDGBuilderFlags::ParallelSetup | ERDGBuilderFlags::ParallelExecute);
 
 				for (const FRenderThreadUpdateContext& DeferredUpdateContext : DrawWindowsCommand->DeferredUpdates)
 				{
@@ -1734,7 +1734,7 @@ void FSlateRHIRenderer::AddWidgetRendererUpdate(const FRenderThreadUpdateContext
 	{
 		ENQUEUE_RENDER_COMMAND(DrawWidgetRendererImmediate)([Context](FRHICommandListImmediate& RHICmdList)
 		{
-			FRDGBuilder GraphBuilder(RHICmdList);
+			FRDGBuilder GraphBuilder(RHICmdList, RDG_EVENT_NAME("SlateWidgetRender"), ERDGBuilderFlags::ParallelSetup | ERDGBuilderFlags::ParallelExecute);
 			Context.Renderer->DrawWindowToTarget_RenderThread(GraphBuilder, Context);
 			GraphBuilder.Execute();
 		});
diff --git a/Engine/Source/Runtime/SlateRHIRenderer/Private/SlateRHIRenderingPolicy.cpp b/Engine/Source/Runtime/SlateRHIRenderer/Private/SlateRHIRenderingPolicy.cpp
--- a/Engine/Source/Runtime/SlateRHIRenderer/Private/SlateRHIRenderingPolicy.cpp
+++ b/Engine/Source/Runtime/SlateRHIRenderer/Private/SlateRHIRenderingPolicy.cpp
@@ -1469,7 +1469,7 @@ void AddSlateDrawElementsPass(
 			LastPassParameters->RenderTargets[0].SetLoadAction(LoadAction);
 		}
 
-		FRDGPass* Pass = GraphBuilder.AddPass(RDG_EVENT_NAME("ElementBatch"), LastPassParameters, ERDGPassFlags::Raster, [Inputs = RenderBatchDrawInputs, RenderBatchHeadOp] (FRHICommandList& RHICmdList)
+		FRDGPass* Pass = GraphBuilder.AddPass(RDG_EVENT_NAME("ElementBatch"), LastPassParameters, ERDGPassFlags::Raster, [Inputs = RenderBatchDrawInputs, RenderBatchHeadOp] (FRDGAsyncTask, FRHICommandList& RHICmdList)
 		{
 			RHICmdList.SetViewport(Inputs->ElementsViewRect.Min.X, Inputs->ElementsViewRect.Min.Y, 0.0f, Inputs->ElementsViewRect.Max.X, Inputs->ElementsViewRect.Max.Y, 1.0f);
 			RHICmdList.SetScissorRect(false, 0, 0, 0, 0);
