Fix movable point light dynamic shadows for XR

diff --git a/Engine/Shaders/Private/MobileBasePassPixelShader.usf b/Engine/Shaders/Private/MobileBasePassPixelShader.usf
--- a/Engine/Shaders/Private/MobileBasePassPixelShader.usf
+++ b/Engine/Shaders/Private/MobileBasePassPixelShader.usf
@@ -874,7 +874,12 @@ void Main(
 
 	// Directional light
 #if !MATERIAL_SHADINGMODEL_UNLIT
-	AccumulateDirectionalLighting(GBuffer, MaterialParameters.WorldPosition_CamRelative, CameraVector, MaterialParameters.ScreenPosition, SvPosition, DynamicShadowFactors, DirectionalLightShadow, DirectLighting);
+#if MOBILE_MULTI_VIEW
+	const uint EyeIndex = BasePassInterpolants.MultiViewId;
+#else 
+	const uint EyeIndex = 0;
+#endif
+	AccumulateDirectionalLighting(GBuffer, MaterialParameters.WorldPosition_CamRelative, CameraVector, MaterialParameters.ScreenPosition, SvPosition, DynamicShadowFactors, DirectionalLightShadow, DirectLighting, EyeIndex);
 #endif
 
 	float2 LocalPosition = SvPosition.xy - ResolvedView.ViewRectMin.xy;
diff --git a/Engine/Shaders/Private/MobileDeferredShading.usf b/Engine/Shaders/Private/MobileDeferredShading.usf
--- a/Engine/Shaders/Private/MobileDeferredShading.usf
+++ b/Engine/Shaders/Private/MobileDeferredShading.usf
@@ -185,7 +185,7 @@ void MobileDirectionalLightPS(
 	// Directional light
 	half4 DynamicShadowFactors = 1.0f;
 	float DirectionalLightShadow = 1.0f;
-	AccumulateDirectionalLighting(GBuffer, TranslatedWorldPosition, CameraVector, ScreenPosition, SvPosition, DynamicShadowFactors, DirectionalLightShadow, DirectLighting);
+	AccumulateDirectionalLighting(GBuffer, TranslatedWorldPosition, CameraVector, ScreenPosition, SvPosition, DynamicShadowFactors, DirectionalLightShadow, DirectLighting, 0);
 	// LightFunction should only affect direct lighting result
 	DirectLighting.TotalLight *= ComputeLightFunctionMultiplier(TranslatedWorldPosition);
 
diff --git a/Engine/Shaders/Private/MobileLightingCommon.ush b/Engine/Shaders/Private/MobileLightingCommon.ush
--- a/Engine/Shaders/Private/MobileLightingCommon.ush
+++ b/Engine/Shaders/Private/MobileLightingCommon.ush
@@ -267,11 +267,15 @@ FLightAccumulator GetTranslucencySHLighting(FGBufferData GBuffer, FDeferredLight
 	Mobile Directional Light.
 ------------------------------------------------------------------------------*/
 
-FDeferredLightData GetDirectionalLightData(float4 ScreenPosition, float4 SvPosition, inout half4 OutDynamicShadowFactors, inout half OutDynamicShadowing)
+FDeferredLightData GetDirectionalLightData(float4 ScreenPosition, float4 SvPosition, inout half4 OutDynamicShadowFactors, inout half OutDynamicShadowing, uint EyeIndex)
 {
 #if USE_SHADOWMASKTEXTURE
 	half4 PreviewShadowMapChannelMask = half4(1.0f, 0.0f, 0.0f, 0.0f);
+#if MOBILE_MULTI_VIEW
+	OutDynamicShadowFactors = Texture2DArraySample(MobileBasePass.ScreenSpaceShadowMaskTextureArray, MobileBasePass.ScreenSpaceShadowMaskSampler, float3(SvPositionToBufferUV(SvPosition), EyeIndex));
+#else
 	OutDynamicShadowFactors = Texture2DSample(MobileBasePass.ScreenSpaceShadowMaskTexture, MobileBasePass.ScreenSpaceShadowMaskSampler, SvPositionToBufferUV(SvPosition));
+#endif
 	OutDynamicShadowFactors = DecodeLightAttenuation(OutDynamicShadowFactors);
 		
 	PreviewShadowMapChannelMask = UnpackShadowMapChannelMask(MobileDirectionalLight.DirectionalLightShadowMapChannelMask >> 4);
@@ -308,10 +312,10 @@ FDeferredLightData GetDirectionalLightData(float4 ScreenPosition, float4 SvPosit
 	return LightData;
 }
 
-void AccumulateDirectionalLighting(FGBufferData GBuffer, float3 TranslatedWorldPosition, half3 CameraVector, float4 ScreenPosition, float4 SvPosition, inout half4 DynamicShadowFactors, inout float OutDirectionalLightShadow, inout FLightAccumulator DirectLighting)
+void AccumulateDirectionalLighting(FGBufferData GBuffer, float3 TranslatedWorldPosition, half3 CameraVector, float4 ScreenPosition, float4 SvPosition, inout half4 DynamicShadowFactors, inout float OutDirectionalLightShadow, inout FLightAccumulator DirectLighting, uint EyeIndex)
 {
 	half DynamicShadowing = 1.0f;
-	FDeferredLightData LightData = GetDirectionalLightData(ScreenPosition, SvPosition, DynamicShadowFactors, DynamicShadowing);
+	FDeferredLightData LightData = GetDirectionalLightData(ScreenPosition, SvPosition, DynamicShadowFactors, DynamicShadowing, EyeIndex);
 	half4 LightAttenuation = half4(1, 1, DynamicShadowing, DynamicShadowing);
 
 #if TRANSLUCENCY_SH_LIGHTING
diff --git a/Engine/Shaders/Private/Substrate/SubstrateMobileForwardLighting.ush b/Engine/Shaders/Private/Substrate/SubstrateMobileForwardLighting.ush
--- a/Engine/Shaders/Private/Substrate/SubstrateMobileForwardLighting.ush
+++ b/Engine/Shaders/Private/Substrate/SubstrateMobileForwardLighting.ush
@@ -224,7 +224,7 @@ float3 SubstrateMobileForwardLighting(
 				{
 					half DynamicShadowing = 1.0f;
 					half4 DynamicShadowFactors = 1.0f;
-					FDeferredLightData DirLightData = GetDirectionalLightData(MaterialParameters.ScreenPosition, SvPosition, DynamicShadowFactors, DynamicShadowing);
+					FDeferredLightData DirLightData = GetDirectionalLightData(MaterialParameters.ScreenPosition, SvPosition, DynamicShadowFactors, DynamicShadowing, EyeIndex);
 					half4 DirLightAttenuation = half4(1, 1, DynamicShadowing, DynamicShadowing);
 					uint DirLightingChannelMask = LIGHTING_CHANNEL_MASK;
 					float3 ToDirLight = DirLightData.Direction;	// Already normalized
diff --git a/Engine/Source/Runtime/Renderer/Private/LightRendering.cpp b/Engine/Source/Runtime/Renderer/Private/LightRendering.cpp
--- a/Engine/Source/Runtime/Renderer/Private/LightRendering.cpp
+++ b/Engine/Source/Runtime/Renderer/Private/LightRendering.cpp
@@ -754,12 +754,17 @@ float GetLightFadeFactor(const FSceneView& View, const FLightSceneProxy* Proxy)
 	return SizeFade * DistanceFade;
 }
 
-void StencilingGeometry::DrawSphere(FRHICommandList& RHICmdList)
+void StencilingGeometry::DrawSphere(FRHICommandList& RHICmdList, uint32 InstanceCount)
 {
 	RHICmdList.SetStreamSource(0, StencilingGeometry::GStencilSphereVertexBuffer.VertexBufferRHI, 0);
 	RHICmdList.DrawIndexedPrimitive(StencilingGeometry::GStencilSphereIndexBuffer.IndexBufferRHI, 0, 0,
 		StencilingGeometry::GStencilSphereVertexBuffer.GetVertexCount(), 0,
-		StencilingGeometry::GStencilSphereIndexBuffer.GetIndexCount() / 3, 1);
+		StencilingGeometry::GStencilSphereIndexBuffer.GetIndexCount() / 3, InstanceCount);
+}
+
+void StencilingGeometry::DrawSphere(FRHICommandList& RHICmdList)
+{
+	DrawSphere(RHICmdList, 1);
 }
 
 void StencilingGeometry::DrawVectorSphere(FRHICommandList& RHICmdList)
diff --git a/Engine/Source/Runtime/Renderer/Private/LightRendering.h b/Engine/Source/Runtime/Renderer/Private/LightRendering.h
--- a/Engine/Source/Runtime/Renderer/Private/LightRendering.h
+++ b/Engine/Source/Runtime/Renderer/Private/LightRendering.h
@@ -97,6 +97,12 @@ namespace StencilingGeometry
 	* Note: The sphere will be of unit size unless transformed by the shader. 
 	*/
 	extern void DrawSphere(FRHICommandList& RHICmdList);
+	/**
+	* Draws a sphere using RHIDrawIndexedPrimitive, useful as approximate bounding geometry for deferred passes.
+	* Compatible with instanced rendering.
+	* Note: The sphere will be of unit size unless transformed by the shader.
+	*/
+	extern void DrawSphere(FRHICommandList& RHICmdList, uint32 InstanceCount);
 	/** Draws exactly the same as above, but uses FVector rather than FVector4f vertex data. */
 	extern void DrawVectorSphere(FRHICommandList& RHICmdList);
 	/** Renders a cone with a spherical cap, used for rendering spot lights in deferred passes. */
diff --git a/Engine/Source/Runtime/Renderer/Private/MobileBasePassRendering.cpp b/Engine/Source/Runtime/Renderer/Private/MobileBasePassRendering.cpp
--- a/Engine/Source/Runtime/Renderer/Private/MobileBasePassRendering.cpp
+++ b/Engine/Source/Runtime/Renderer/Private/MobileBasePassRendering.cpp
@@ -340,11 +340,13 @@ void SetupMobileBasePassUniformParameters(
 	{
 		FRDGTextureRef ScreenShadowMaskTexture = GraphBuilder.RegisterExternalTexture(GScreenSpaceShadowMaskTextureMobileOutputs.ScreenSpaceShadowMaskTextureMobile, TEXT("ScreenSpaceShadowMaskTextureMobile"));
 		BasePassParameters.ScreenSpaceShadowMaskTexture = ScreenShadowMaskTexture;
+		BasePassParameters.ScreenSpaceShadowMaskTextureArray = ScreenShadowMaskTexture;
 		BasePassParameters.ScreenSpaceShadowMaskSampler = TStaticSamplerState<SF_Point, AM_Clamp, AM_Clamp, AM_Clamp>::GetRHI();
 	}
 	else
 	{
 		BasePassParameters.ScreenSpaceShadowMaskTexture = GSystemTextures.GetWhiteDummy(GraphBuilder);
+		BasePassParameters.ScreenSpaceShadowMaskTextureArray = GSystemTextures.GetDefaultTexture(GraphBuilder, ETextureDimension::Texture2DArray, PF_DepthStencil, FClearValueBinding::White);
 		BasePassParameters.ScreenSpaceShadowMaskSampler = TStaticSamplerState<SF_Point, AM_Clamp, AM_Clamp, AM_Clamp>::GetRHI();
 	}
 
diff --git a/Engine/Source/Runtime/Renderer/Private/MobileBasePassRendering.h b/Engine/Source/Runtime/Renderer/Private/MobileBasePassRendering.h
--- a/Engine/Source/Runtime/Renderer/Private/MobileBasePassRendering.h
+++ b/Engine/Source/Runtime/Renderer/Private/MobileBasePassRendering.h
@@ -59,6 +59,7 @@ BEGIN_GLOBAL_SHADER_PARAMETER_STRUCT(FMobileBasePassUniformParameters, )
 	SHADER_PARAMETER_RDG_TEXTURE(Texture2D, AmbientOcclusionTexture)
 	SHADER_PARAMETER_SAMPLER(SamplerState, AmbientOcclusionSampler)
 	SHADER_PARAMETER_RDG_TEXTURE(Texture2D, ScreenSpaceShadowMaskTexture)
+	SHADER_PARAMETER_RDG_TEXTURE(Texture2DArray, ScreenSpaceShadowMaskTextureArray)
 	SHADER_PARAMETER_SAMPLER(SamplerState, ScreenSpaceShadowMaskSampler)
 	SHADER_PARAMETER(uint32, bApplyHalfResLocalFogToSkyMeshes)
 	SHADER_PARAMETER_RDG_TEXTURE(Texture2D, HalfResLocalFogVolumeViewTexture)
diff --git a/Engine/Source/Runtime/Renderer/Private/MobileShadingRenderer.cpp b/Engine/Source/Runtime/Renderer/Private/MobileShadingRenderer.cpp
--- a/Engine/Source/Runtime/Renderer/Private/MobileShadingRenderer.cpp
+++ b/Engine/Source/Runtime/Renderer/Private/MobileShadingRenderer.cpp
@@ -657,7 +657,9 @@ void FMobileSceneRenderer::InitViews(
 
 		if (bRequiresShadowProjections)
 		{
-			InitMobileShadowProjectionOutputs(RHICmdList, SceneTexturesConfig.Extent);
+			FViewInfo* MainView = Views.Num() > 0 ? &Views[0] : nullptr;
+			bool bIsMobileMultiView = SceneTexturesConfig.bRequireMultiView || (MainView && MainView->Aspects.IsMobileMultiViewEnabled());
+			InitMobileShadowProjectionOutputs(RHICmdList, SceneTexturesConfig.Extent, bIsMobileMultiView);
 		}
 		else
 		{
diff --git a/Engine/Source/Runtime/Renderer/Private/ShadowRendering.cpp b/Engine/Source/Runtime/Renderer/Private/ShadowRendering.cpp
--- a/Engine/Source/Runtime/Renderer/Private/ShadowRendering.cpp
+++ b/Engine/Source/Runtime/Renderer/Private/ShadowRendering.cpp
@@ -1651,7 +1651,7 @@ void FProjectedShadowInfo::RenderOnePassPointLightProjection(
 
 		// Project the point light shadow with some approximately bounding geometry, 
 		// So we can get speedups from depth testing and not processing pixels outside of the light's influence.
-		StencilingGeometry::DrawSphere(RHICmdList);
+		StencilingGeometry::DrawSphere(RHICmdList, View.GetStereoPassInstanceFactor());
 		RHICmdList.SetScissorRect(false, 0, 0, 0, 0);
 	});
 }
@@ -2037,12 +2037,11 @@ void FSceneRenderer::RenderShadowProjections(
 	for (int32 ViewIndex = 0; ViewIndex < Views.Num(); ViewIndex++)
 	{
 		const FViewInfo& View = Views[ViewIndex];
-		const FExclusiveDepthStencil ExclusiveDepthStencil = FExclusiveDepthStencil::DepthRead_StencilWrite;
-
-		if (bSubPixelShadow && !HairStrands::HasViewHairStrandsData(View))
+		if (!View.ShouldRenderView() || (bSubPixelShadow && !HairStrands::HasViewHairStrandsData(View)))
 		{
 			continue;
 		}
+		const FExclusiveDepthStencil ExclusiveDepthStencil = FExclusiveDepthStencil::DepthRead_StencilWrite;
 
 		View.BeginRenderView();
 
@@ -2384,14 +2383,17 @@ void FMobileSceneRenderer::RenderModulatedShadowProjections(FRHICommandList& RHI
 	}
 }
 
-void InitMobileShadowProjectionOutputs(FRHICommandListImmediate& RHICmdList, const FIntPoint& Extent)
+void InitMobileShadowProjectionOutputs(FRHICommandListImmediate& RHICmdList, const FIntPoint& Extent, const bool bRequireMultiView)
 {
 	const FIntPoint& BufferSize = Extent;
 
 	if (!GScreenSpaceShadowMaskTextureMobileOutputs.IsValid() || GScreenSpaceShadowMaskTextureMobileOutputs.ScreenSpaceShadowMaskTextureMobile->GetDesc().Extent != BufferSize)
 	{
 		GScreenSpaceShadowMaskTextureMobileOutputs.ScreenSpaceShadowMaskTextureMobile.SafeRelease();
-		GRenderTargetPool.FindFreeElement(RHICmdList, FPooledRenderTargetDesc::Create2DDesc(BufferSize, PF_B8G8R8A8, FClearValueBinding::White, TexCreate_None, TexCreate_RenderTargetable | TexCreate_ShaderResource, false, 1, false), GScreenSpaceShadowMaskTextureMobileOutputs.ScreenSpaceShadowMaskTextureMobile, TEXT("ForwardScreenSpaceShadowMaskTextureTexture"));
+		FPooledRenderTargetDesc Desc = bRequireMultiView ? FPooledRenderTargetDesc::Create2DArrayDesc(BufferSize, PF_B8G8R8A8, FClearValueBinding::White, TexCreate_None, TexCreate_RenderTargetable | TexCreate_ShaderResource, false, 2, 1, false)
+			: FPooledRenderTargetDesc::Create2DDesc(BufferSize, PF_B8G8R8A8, FClearValueBinding::White, TexCreate_None, TexCreate_RenderTargetable | TexCreate_ShaderResource, false, 1, false);
+
+		GRenderTargetPool.FindFreeElement(RHICmdList, Desc, GScreenSpaceShadowMaskTextureMobileOutputs.ScreenSpaceShadowMaskTextureMobile, TEXT("ForwardScreenSpaceShadowMaskTextureTexture"));
 	}
 }
 
diff --git a/Engine/Source/Runtime/Renderer/Private/ShadowRendering.h b/Engine/Source/Runtime/Renderer/Private/ShadowRendering.h
--- a/Engine/Source/Runtime/Renderer/Private/ShadowRendering.h
+++ b/Engine/Source/Runtime/Renderer/Private/ShadowRendering.h
@@ -64,7 +64,7 @@ void OverrideWithDefaultMaterialForShadowDepth(
 	ERHIFeatureLevel::Type InFeatureLevel
 	);
 
-void InitMobileShadowProjectionOutputs(FRHICommandListImmediate& RHICmdList, const FIntPoint& Extent);
+void InitMobileShadowProjectionOutputs(FRHICommandListImmediate& RHICmdList, const FIntPoint& Extent, const bool bRequireMultiView);
 void ReleaseMobileShadowProjectionOutputs();
 
 enum EShadowDepthRenderMode
