Fixed incorrect slate viewport for DrawTileMesh.

diff --git a/Engine/Source/Runtime/Renderer/Private/Renderer.cpp b/Engine/Source/Runtime/Renderer/Private/Renderer.cpp
--- a/Engine/Source/Runtime/Renderer/Private/Renderer.cpp
+++ b/Engine/Source/Runtime/Renderer/Private/Renderer.cpp
@@ -418,7 +418,7 @@ void FRendererModule::DrawTileMesh(FCanvasRenderContext& RenderContext, FMeshPas
 			{
 				PassParameters->DebugViewMode = CreateDebugViewModePassUniformBuffer(GraphBuilder, View, nullptr);
 
-				AddDrawDynamicMeshPass(GraphBuilder, RDG_EVENT_NAME("OutputMaterialTextureScales"), PassParameters, View, View.ViewRect,
+				AddDrawDynamicMeshPass(GraphBuilder, RDG_EVENT_NAME("OutputMaterialTextureScales"), PassParameters, View, RenderContext.GetViewportRect(),
 					[Scene, &View, &Mesh](FMeshPassDrawListContext* InDrawListContext)
 				{
 					FDebugViewModeMeshProcessor PassMeshProcessor(
@@ -440,7 +440,7 @@ void FRendererModule::DrawTileMesh(FCanvasRenderContext& RenderContext, FMeshPas
 			{
 				PassParameters->TranslucentBasePass = CreateTranslucentBasePassUniformBuffer(GraphBuilder, Scene, View);
 
-				AddDrawDynamicMeshPass(GraphBuilder, RDG_EVENT_NAME("TranslucentDeferred"), PassParameters, View, View.ViewRect,
+				AddDrawDynamicMeshPass(GraphBuilder, RDG_EVENT_NAME("TranslucentDeferred"), PassParameters, View, RenderContext.GetViewportRect(),
 					[Scene, &View, &Mesh, DrawRenderState, bUse128bitRT](FMeshPassDrawListContext* DynamicMeshPassContext)
 				{
 					FBasePassMeshProcessor PassMeshProcessor(
@@ -461,7 +461,7 @@ void FRendererModule::DrawTileMesh(FCanvasRenderContext& RenderContext, FMeshPas
 			{
 				PassParameters->MobileBasePass = CreateMobileBasePassUniformBuffer(GraphBuilder, View, EMobileBasePass::Translucent, EMobileSceneTextureSetupMode::None);
 
-				AddDrawDynamicMeshPass(GraphBuilder, RDG_EVENT_NAME("TranslucentMobile"), PassParameters, View, View.ViewRect,
+				AddDrawDynamicMeshPass(GraphBuilder, RDG_EVENT_NAME("TranslucentMobile"), PassParameters, View, RenderContext.GetViewportRect(),
 					[Scene, &View, DrawRenderState, &Mesh](FDynamicPassMeshDrawListContext* DynamicMeshPassContext)
 				{
 					FMobileBasePassMeshProcessor PassMeshProcessor(
@@ -490,7 +490,7 @@ void FRendererModule::DrawTileMesh(FCanvasRenderContext& RenderContext, FMeshPas
 				ensureMsgf(HitProxyId == Mesh.BatchHitProxyId, TEXT("Only Mesh.BatchHitProxyId is used for hit testing."));
 
 #if WITH_EDITOR
-				AddDrawDynamicMeshPass(GraphBuilder, RDG_EVENT_NAME("HitTesting"), PassParameters, View, View.ViewRect,
+				AddDrawDynamicMeshPass(GraphBuilder, RDG_EVENT_NAME("HitTesting"), PassParameters, View, RenderContext.GetViewportRect(),
 					[Scene, &View, DrawRenderState, &Mesh](FDynamicPassMeshDrawListContext* DynamicMeshPassContext)
 				{
 					FHitProxyMeshProcessor PassMeshProcessor(
@@ -511,7 +511,7 @@ void FRendererModule::DrawTileMesh(FCanvasRenderContext& RenderContext, FMeshPas
 				{
 					PassParameters->OpaqueBasePass = CreateOpaqueBasePassUniformBuffer(GraphBuilder, View);
 
-					AddDrawDynamicMeshPass(GraphBuilder, RDG_EVENT_NAME("OpaqueDeferred"), PassParameters, View, View.ViewRect,
+					AddDrawDynamicMeshPass(GraphBuilder, RDG_EVENT_NAME("OpaqueDeferred"), PassParameters, View, RenderContext.GetViewportRect(),
 						[Scene, &View, DrawRenderState, &Mesh, bUse128bitRT](FDynamicPassMeshDrawListContext* DynamicMeshPassContext)
 					{
 						FBasePassMeshProcessor PassMeshProcessor(
@@ -531,7 +531,7 @@ void FRendererModule::DrawTileMesh(FCanvasRenderContext& RenderContext, FMeshPas
 				{
 					PassParameters->MobileBasePass = CreateMobileBasePassUniformBuffer(GraphBuilder, View, EMobileBasePass::Opaque, EMobileSceneTextureSetupMode::None);
 
-					AddDrawDynamicMeshPass(GraphBuilder, RDG_EVENT_NAME("OpaqueMobile"), PassParameters, View, View.ViewRect,
+					AddDrawDynamicMeshPass(GraphBuilder, RDG_EVENT_NAME("OpaqueMobile"), PassParameters, View, RenderContext.GetViewportRect(),
 						[Scene, &View, DrawRenderState, &Mesh](FDynamicPassMeshDrawListContext* DynamicMeshPassContext)
 					{
 						FMobileBasePassMeshProcessor PassMeshProcessor(
