Path Tracer: Add cvar to control light grid projection axis

This is mainly helpful for debugging purposes, the default remains automatic.

#rnx

diff --git a/Engine/Source/Runtime/Renderer/Private/PathTracing.cpp b/Engine/Source/Runtime/Renderer/Private/PathTracing.cpp
--- a/Engine/Source/Runtime/Renderer/Private/PathTracing.cpp
+++ b/Engine/Source/Runtime/Renderer/Private/PathTracing.cpp
@@ -369,6 +369,13 @@ TAutoConsoleVariable<int32> CVarPathTracingLightGridMaxCount(
 	ECVF_RenderThreadSafe
 );
 
+TAutoConsoleVariable<int32> CVarPathTracingLightGridAxis(
+	TEXT("r.PathTracing.LightGridAxis"),
+	-1,
+	TEXT("Choose the coordinate axis along which to project the light grid (default = -1, automatic)\n"),
+	ECVF_RenderThreadSafe
+);
+
 TAutoConsoleVariable<bool> CVarPathTracingUseDBuffer(
 	TEXT("r.PathTracing.UseDBuffer"),
 	true,
@@ -2538,21 +2545,24 @@ RENDERER_API void PrepareLightGrid(FRDGBuilder& GraphBuilder, ERHIFeatureLevel::
 		LightGridParameters->LightGridResolution = Resolution;
 		LightGridParameters->LightGridMaxCount = MaxCount;
 
-		// pick the shortest axis
-		FVector3f Diag = LightGridParameters->SceneLightsTranslatedBoundMax - LightGridParameters->SceneLightsTranslatedBoundMin;
-		if (Diag.X < Diag.Y && Diag.X < Diag.Z)
-		{
-			LightGridParameters->LightGridAxis = 0;
-		}
-		else if (Diag.Y < Diag.Z)
-		{
-			LightGridParameters->LightGridAxis = 1;
-		}
-		else
+		LightGridParameters->LightGridAxis = CVarPathTracingLightGridAxis.GetValueOnRenderThread();
+		if (LightGridParameters->LightGridAxis < 0 || LightGridParameters->LightGridAxis > 2)
 		{
-			LightGridParameters->LightGridAxis = 2;
+			// pick the shortest axis if user did not explicitly pick an axis
+			FVector3f Diag = LightGridParameters->SceneLightsTranslatedBoundMax - LightGridParameters->SceneLightsTranslatedBoundMin;
+			if (Diag.X < Diag.Y && Diag.X < Diag.Z)
+			{
+				LightGridParameters->LightGridAxis = 0;
+			}
+			else if (Diag.Y < Diag.Z)
+			{
+				LightGridParameters->LightGridAxis = 1;
+			}
+			else
+			{
+				LightGridParameters->LightGridAxis = 2;
+			}
 		}
-
 		FPathTracingBuildLightGridCS::FParameters* LightGridPassParameters = GraphBuilder.AllocParameters< FPathTracingBuildLightGridCS::FParameters>();
 
 		FRDGTextureDesc LightGridDesc = FRDGTextureDesc::Create2D(
