Decals: Fix decals always fading in, even when FadeDuration is 0

* Avoid initializing fade parameters in BatchUpdateDecals() when FadeDuration is 0, and reset the InvFadeDuration

diff --git a/Engine/Source/Runtime/Renderer/Private/RendererScene.cpp b/Engine/Source/Runtime/Renderer/Private/RendererScene.cpp
--- a/Engine/Source/Runtime/Renderer/Private/RendererScene.cpp
+++ b/Engine/Source/Runtime/Renderer/Private/RendererScene.cpp
@@ -2985,7 +2985,17 @@ void FScene::BatchUpdateDecals(TArray<FDeferredDecalUpdateParams>&& UpdateParams
 				}
 
 				DecalUpdate.DecalProxy->SetTransformIncludingDecalSize(DecalUpdate.Transform, DecalUpdate.Bounds);
-				DecalUpdate.DecalProxy->InitializeFadingParameters(DecalUpdate.AbsSpawnTime, DecalUpdate.FadeDuration, DecalUpdate.FadeStartDelay, DecalUpdate.FadeInDuration, DecalUpdate.FadeInStartDelay);
+
+				// When FadeDuration is intentionally set to 0 the user expects the decal to not fade automatically
+				if (DecalUpdate.FadeDuration == 0.0f)
+				{
+					DecalUpdate.DecalProxy->InvFadeDuration = -1.0f;
+				}
+				else
+				{
+					DecalUpdate.DecalProxy->InitializeFadingParameters(DecalUpdate.AbsSpawnTime, DecalUpdate.FadeDuration, DecalUpdate.FadeStartDelay, DecalUpdate.FadeInDuration, DecalUpdate.FadeInStartDelay);
+				}
+
 				DecalUpdate.DecalProxy->FadeScreenSize = DecalUpdate.FadeScreenSize;
 				DecalUpdate.DecalProxy->SortOrder = DecalUpdate.SortOrder;
 				DecalUpdate.DecalProxy->DecalColor = DecalUpdate.DecalColor;
