Fix VSM visualizations not scaling well with dynres in builds

diff --git a/Engine/Shaders/Private/VirtualShadowMaps/VirtualShadowMapDebug.usf b/Engine/Shaders/Private/VirtualShadowMaps/VirtualShadowMapDebug.usf
--- a/Engine/Shaders/Private/VirtualShadowMaps/VirtualShadowMapDebug.usf
+++ b/Engine/Shaders/Private/VirtualShadowMaps/VirtualShadowMapDebug.usf
@@ -157,15 +157,17 @@ void DebugVisualizeVirtualSmCS(uint2 PixelPos : SV_DispatchThreadID)
 int VisualizeVirtualShadowMapId;
 float VisualizeNaniteOverdrawScale;
 Texture2D DebugTexture;
+SamplerState DebugTextureSampler;
 
 float4 TonemapProjectionDebugTexturePS(
-	noperspective float4 UVAndScreenPos : TEXCOORD0,
-	float4 Position : SV_POSITION
+	noperspective float4 UVAndScreenPos : TEXCOORD0
 ) : SV_Target0
 {
 	bool bOneLightVisible = VisualizeVirtualShadowMapId >= 0;
 
-	float4 Output = DebugTexture.Load(uint3(Position.xy, 0));
+	float2 UV = UVAndScreenPos.xy;
+	float4 Output = Texture2DSample(DebugTexture, DebugTextureSampler, UV);
+
 	if (VisualizeModeId == VIRTUAL_SHADOW_MAP_VISUALIZE_SMRT_RAY_COUNT)
 	{
 		if (Output.a > 0)
diff --git a/Engine/Source/Runtime/Renderer/Private/VirtualShadowMaps/VirtualShadowMapArray.cpp b/Engine/Source/Runtime/Renderer/Private/VirtualShadowMaps/VirtualShadowMapArray.cpp
--- a/Engine/Source/Runtime/Renderer/Private/VirtualShadowMaps/VirtualShadowMapArray.cpp
+++ b/Engine/Source/Runtime/Renderer/Private/VirtualShadowMaps/VirtualShadowMapArray.cpp
@@ -3776,6 +3776,7 @@ class FTonemapProjectionDebugTexturePS : public FVirtualShadowMapPageManagementS
 
 	BEGIN_SHADER_PARAMETER_STRUCT(FParameters, )
 		SHADER_PARAMETER_RDG_TEXTURE(Texture2D, DebugTexture)
+		SHADER_PARAMETER_SAMPLER(SamplerState, DebugTextureSampler)
 		SHADER_PARAMETER(uint32, VisualizeModeId)
 		SHADER_PARAMETER(int32, VisualizeVirtualShadowMapId)
 		SHADER_PARAMETER(float, VisualizeNaniteOverdrawScale)
@@ -3849,6 +3850,8 @@ FScreenPassTexture FVirtualShadowMapArray::AddVisualizePass(FRDGBuilder& GraphBu
 
 		FTonemapProjectionDebugTexturePS::FParameters* Parameters = GraphBuilder.AllocParameters<FTonemapProjectionDebugTexturePS::FParameters>();
 		Parameters->DebugTexture = DebugVisualizationOutput[ViewIndex];
+		// Point sampling as DebugTexture could have non-linear data
+		Parameters->DebugTextureSampler = TStaticSamplerState<SF_Point, AM_Clamp, AM_Clamp, AM_Clamp>::GetRHI();
 		Parameters->VisualizeModeId = ActiveModeId;
 		Parameters->VisualizeVirtualShadowMapId = VisualizeVirtualShadowMapId;
 		Parameters->VisualizeNaniteOverdrawScale = GNaniteVisualizeOverdrawScale;
@@ -3870,7 +3873,6 @@ FScreenPassTexture FVirtualShadowMapArray::AddVisualizePass(FRDGBuilder& GraphBu
 		{
 			TShaderMapRef<FScreenPassVS> VertexShader(View.ShaderMap);
 			TShaderMapRef<FDesaturatePS> DesaturatePixelShader(View.ShaderMap);
-			TShaderMapRef<FCopyRectPS> PixelShader(View.ShaderMap);
 
 			FRDGTextureRef SceneColorCopy = GraphBuilder.CreateTexture(SceneColor.Texture->Desc, TEXT("SceneColorCopy"));
 			AddCopyTexturePass(GraphBuilder, SceneColor.Texture, SceneColorCopy);
