Heterogeneous Volumes: Improvements to shadow-casting with the world-space voxel grid

1. Shadows use temporal view jittering.
2. Corrected usage of shading rates and minimum voxel sizes
3. Voxel grid shadows use blue noise sequence
4. Voxel grid shadows use projected pixel size
5. Projected pixel score incorporates downsample factor
6. Updated CVar defaults for voxel grid shadows
7. Voxel grid shadows temporarily disable per-light primitive gathering to avoid crash

diff --git a/Engine/Shaders/Private/HeterogeneousVolumes/HeterogeneousVolumesLiveShadingShadows.usf b/Engine/Shaders/Private/HeterogeneousVolumes/HeterogeneousVolumesLiveShadingShadows.usf
--- a/Engine/Shaders/Private/HeterogeneousVolumes/HeterogeneousVolumesLiveShadingShadows.usf
+++ b/Engine/Shaders/Private/HeterogeneousVolumes/HeterogeneousVolumesLiveShadingShadows.usf
@@ -313,8 +313,8 @@ void RenderVolumetricShadowMapForLightWithLiveShadingCS(
 	int PixelIndex = LinearIndex * MaxSampleCount;
 
 	// Generate a light ray
-	float2 ShadowUV = float2(PixelCoord + 0.5f) / float2(ShadowResolution);
-	//float2 ShadowUV = float2(PixelCoord + View.TemporalAAJitter.xy) / float2(ShadowResolution);
+	//float2 ShadowUV = float2(PixelCoord + 0.5f) / float2(ShadowResolution);
+	float2 ShadowUV = float2(PixelCoord + View.TemporalAAJitter.xy) / float2(ShadowResolution);
 
 	float WorldRayLength = GetMaxTraceDistance();
 #if USE_CAMERA_SCENE_DEPTH
diff --git a/Engine/Shaders/Private/HeterogeneousVolumes/HeterogeneousVolumesVoxelGridPipeline.usf b/Engine/Shaders/Private/HeterogeneousVolumes/HeterogeneousVolumesVoxelGridPipeline.usf
--- a/Engine/Shaders/Private/HeterogeneousVolumes/HeterogeneousVolumesVoxelGridPipeline.usf
+++ b/Engine/Shaders/Private/HeterogeneousVolumes/HeterogeneousVolumesVoxelGridPipeline.usf
@@ -38,9 +38,12 @@ StructuredBuffer<FVectorGridData> ScatteringGridBuffer;
 float3 PrimitiveWorldBoundsMin;
 float3 PrimitiveWorldBoundsMax;
 
-float ShadingRate;
+float ShadingRateInFrustum;
+float ShadingRateOutOfFrustum;
 float MinVoxelSizeInFrustum;
 float MinVoxelSizeOutOfFrustum;
+float DownsampleFactor;
+int bUseProjectedPixelSize;
 
 RWStructuredBuffer<FTopLevelGridData> RWTopLevelGridBuffer;
 
@@ -175,20 +178,25 @@ void TopLevelGridCalculateVoxelSize(
 
 	if (PrimitiveIntersectsVoxel(VoxelBoundsMin, VoxelBoundsMax))
 	{
-		// Initialize the voxel-rate based on view distance
-		float2 TanHalfFOV = GetTanHalfFieldOfView();
-		int HalfWidth = View.ViewSizeAndInvSize.x * 0.5;
-		float PixelWidth = TanHalfFOV.x / HalfWidth;
-
-		float3 WorldCameraOrigin = DFHackToFloat(DFFastSubtract(View.TranslatedWorldCameraOrigin, PrimaryView.PreViewTranslation));
-		float Distance = GetDistanceToCameraFromViewVector(WorldVoxelCenter - WorldCameraOrigin) - length(WorldVoxelExtent);
-		float VoxelWidth = Distance * PixelWidth * ShadingRate;
+		float VoxelSize = MinVoxelSizeOutOfFrustum;
+		float ShadingRate = ShadingRateOutOfFrustum;
+		if (BoxIntersectsViewFrustum(WorldVoxelCenter, WorldVoxelExtent))
+		{
+			VoxelSize = MinVoxelSizeInFrustum;
+			ShadingRate = ShadingRateInFrustum;
+		}
 
-		// Frustum intersection determines which min voxel-size to clip by
-		float3 PrimitiveWorldBoundsOrigin = (PrimitiveWorldBoundsMax + PrimitiveWorldBoundsMin) * 0.5;
-		float3 PrimitiveWorldBoundsExtent = (PrimitiveWorldBoundsMax - PrimitiveWorldBoundsMin) * 0.5;
-		float MinVoxelSize = BoxIntersectsViewFrustum(WorldVoxelCenter, WorldVoxelExtent) ? MinVoxelSizeInFrustum : MinVoxelSizeOutOfFrustum;
-		float VoxelSize = max(VoxelWidth, MinVoxelSize);
+		if (bUseProjectedPixelSize)
+		{
+			float2 TanHalfFOV = GetTanHalfFieldOfView();
+			float HalfWidth = View.ViewSizeAndInvSize.x * 0.5 / DownsampleFactor;
+			float PixelWidth = TanHalfFOV.x / HalfWidth;
+
+			float3 WorldCameraOrigin = DFFastSubtractDemote(View.TranslatedWorldCameraOrigin, PrimaryView.PreViewTranslation);
+			float Distance = max(GetDistanceToCameraFromViewVector(WorldVoxelCenter - WorldCameraOrigin) - length(WorldVoxelExtent), 0.0);
+			float VoxelWidth = Distance * PixelWidth * ShadingRate;
+			VoxelSize = max(VoxelWidth, VoxelSize);
+		}
 
 		// Store minimum voxel rate temporarily as the bottom-level index
 		uint LinearIndex = GetLinearIndex(VoxelIndex, TopLevelGridResolution);
@@ -343,7 +351,8 @@ void GenerateRasterTiles(
 		int3 RasterTileCount3 = BottomLevelResolution;
 #else // DIM_ENABLE_INDIRECTION_GRID
 		// NOTE: Always 1, due to limitations of bottom-level grid resolution
-		int3 RasterTileCount3 = ceil(float3(BottomLevelResolution) / RasterTileVoxelResolution);
+		//int3 RasterTileCount3 = ceil(float3(BottomLevelResolution) / RasterTileVoxelResolution);
+		int3 RasterTileCount3 = 1;
 #endif // DIM_ENABLE_INDIRECTION_GRID
 
 		int RasterTileCount = RasterTileCount3.x * RasterTileCount3.y * RasterTileCount3.z;
diff --git a/Engine/Shaders/Private/HeterogeneousVolumes/HeterogeneousVolumesVoxelGridShadows.usf b/Engine/Shaders/Private/HeterogeneousVolumes/HeterogeneousVolumesVoxelGridShadows.usf
--- a/Engine/Shaders/Private/HeterogeneousVolumes/HeterogeneousVolumesVoxelGridShadows.usf
+++ b/Engine/Shaders/Private/HeterogeneousVolumes/HeterogeneousVolumesVoxelGridShadows.usf
@@ -8,6 +8,7 @@
 #include "HeterogeneousVolumesTracingUtils.ush"
 #include "HeterogeneousVolumesVoxelGridTraversal.ush"
 #include "HeterogeneousVolumesVoxelGridRenderingUtils.ush"
+#include "../BlueNoise.ush"
 
 #include "../PathTracing/Utilities/PathTracingRandomSequence.ush"
 
@@ -335,7 +336,8 @@ void RenderVolumetricShadowMapForLightWithVoxelGridCS(
 	uint LinearIndex = GetLinearIndex(Face, PixelCoord);
 
 	// Generate a light ray
-	float2 ShadowUV = float2(PixelCoord + 0.5f) / float2(ShadowResolution);
+	//float2 ShadowUV = float2(PixelCoord + 0.5f) / float2(ShadowResolution);
+	float2 ShadowUV = float2(PixelCoord + View.TemporalAAJitter.xy) / float2(ShadowResolution);
 	const float NearScreenZ = 1.0;
 	const float FarScreenZ = 0.0;
 
@@ -376,10 +378,9 @@ void RenderVolumetricShadowMapForLightWithVoxelGridCS(
 	int PixelIndex = LinearIndex * MaxSampleCount;
 
 	float HitSpan = RayHitT.y - RayHitT.x;
-	float Depth = 1.0;
 	if (HitSpan > 0.0)
 	{
-		float Jitter = bJitter ? InterleavedGradientNoise(PixelCoord, View.StateFrameIndexMod8) : 0.0;
+		float Jitter = bJitter ? BlueNoiseScalar(PixelCoord, View.StateFrameIndexMod8) : 0.0;
 
 		FRayMarchingContext RayMarchingContext = CreateRayMarchingContext(
 			// Local-space
@@ -422,7 +423,7 @@ void RenderVolumetricShadowMapForLightWithVoxelGridCS(
 			Transmittance *= exp(-Extinction * RayMarchingContext.StepSize);
 
 			bool bFirstEntry = (ElementIndex == 0);
-			if (bFirstEntry && any(Extinction > 0.0))
+			if (bFirstEntry)
 			{
 				FAVSMSampleData SampleData = AVSM_CreateSampleData(WorldHitT, Luminance(Transmittance));
 				AddSample(SampleData, PixelIndex, PrevElementIndex, ElementIndex);
@@ -431,7 +432,7 @@ void RenderVolumetricShadowMapForLightWithVoxelGridCS(
 				ExpectedExtinction = Extinction;
 				ExpectedTransmittance = Transmittance;
 			}
-			else if (ElementIndex > 0 && any(Extinction > 0.0))
+			else
 			{
 				// If extrapolated transmittance is beyond relative error tolerance
 				ExpectedTransmittance *= exp(-ExpectedExtinction * RayMarchingContext.StepSize);
diff --git a/Engine/Source/Runtime/Renderer/Private/HeterogeneousVolumes/HeterogeneousVolumes.cpp b/Engine/Source/Runtime/Renderer/Private/HeterogeneousVolumes/HeterogeneousVolumes.cpp
--- a/Engine/Source/Runtime/Renderer/Private/HeterogeneousVolumes/HeterogeneousVolumes.cpp
+++ b/Engine/Source/Runtime/Renderer/Private/HeterogeneousVolumes/HeterogeneousVolumes.cpp
@@ -771,10 +771,11 @@ void FDeferredShadingSceneRenderer::RenderHeterogeneousVolumeShadows(
 	{
 		FVoxelGridBuildOptions BuildOptions;
 		BuildOptions.VoxelGridBuildMode = EVoxelGridBuildMode::Shadows;
-		BuildOptions.MinimumVoxelSizeOutsideFrustum = HeterogeneousVolumes::GetOutOfFrustumShadingRateForShadows();
-		BuildOptions.MinimumVoxelSizeInFrustum = HeterogeneousVolumes::GetShadingRateForShadows();
+		BuildOptions.ShadingRateInFrustum = HeterogeneousVolumes::GetShadingRateForShadows();
+		BuildOptions.ShadingRateOutOfFrustum = HeterogeneousVolumes::GetOutOfFrustumShadingRateForShadows();
 		BuildOptions.bBuildOrthoGrid = true;
 		BuildOptions.bBuildFrustumGrid = false;
+		BuildOptions.bUseProjectedPixelSizeForOrthoGrid = true;
 		BuildOptions.bJitter = HeterogeneousVolumes::EnableJitterForShadows();
 
 		BuildOrthoVoxelGrid(GraphBuilder, Scene, Views, VisibleLightInfos, BuildOptions, OrthoGridUniformBuffer);
diff --git a/Engine/Source/Runtime/Renderer/Private/HeterogeneousVolumes/HeterogeneousVolumes.h b/Engine/Source/Runtime/Renderer/Private/HeterogeneousVolumes/HeterogeneousVolumes.h
--- a/Engine/Source/Runtime/Renderer/Private/HeterogeneousVolumes/HeterogeneousVolumes.h
+++ b/Engine/Source/Runtime/Renderer/Private/HeterogeneousVolumes/HeterogeneousVolumes.h
@@ -73,6 +73,8 @@ namespace HeterogeneousVolumes
 	float GetMaxStepCount();
 	float GetMinimumVoxelSizeInFrustum();
 	float GetMinimumVoxelSizeOutsideFrustum();
+	float GetShadingRateForFrustumGrid();
+	float GetShadingRateForOrthoGrid();
 
 	// Shadow generation
 	enum class EShadowMode
@@ -294,11 +296,12 @@ enum class EVoxelGridBuildMode
 struct FVoxelGridBuildOptions
 {
 	EVoxelGridBuildMode VoxelGridBuildMode = EVoxelGridBuildMode::PathTracing;
-	float MinimumVoxelSizeOutsideFrustum = HeterogeneousVolumes::GetMinimumVoxelSizeOutsideFrustum();
-	float MinimumVoxelSizeInFrustum = HeterogeneousVolumes::GetMinimumVoxelSizeInFrustum();
+	float ShadingRateInFrustum = HeterogeneousVolumes::GetShadingRateForFrustumGrid();
+	float ShadingRateOutOfFrustum = HeterogeneousVolumes::GetShadingRateForOrthoGrid();
 
 	bool bBuildOrthoGrid = true;
 	bool bBuildFrustumGrid = true;
+	bool bUseProjectedPixelSizeForOrthoGrid = false;
 	bool bJitter = HeterogeneousVolumes::ShouldJitter();
 };
 
diff --git a/Engine/Source/Runtime/Renderer/Private/HeterogeneousVolumes/HeterogeneousVolumesVoxelGridPipeline.cpp b/Engine/Source/Runtime/Renderer/Private/HeterogeneousVolumes/HeterogeneousVolumesVoxelGridPipeline.cpp
--- a/Engine/Source/Runtime/Renderer/Private/HeterogeneousVolumes/HeterogeneousVolumesVoxelGridPipeline.cpp
+++ b/Engine/Source/Runtime/Renderer/Private/HeterogeneousVolumes/HeterogeneousVolumesVoxelGridPipeline.cpp
@@ -122,8 +122,8 @@ static TAutoConsoleVariable<int32> CVarHeterogeneousVolumesBottomLevelGridVoxelH
 
 static TAutoConsoleVariable<float> CVarHeterogeneousVolumesMinimumVoxelSizeInFrustum(
 	TEXT("r.HeterogeneousVolumes.Tessellation.MinimumVoxelSizeInFrustum"),
-	1.0,
-	TEXT("The minimum voxel size (Default = 1.0)"),
+	0.1,
+	TEXT("The minimum voxel size (Default = 0.1)"),
 	ECVF_RenderThreadSafe
 );
 
@@ -252,17 +252,17 @@ static TAutoConsoleVariable<int32> CVarHeterogeneousVolumesShadowDebugTweak(
 	ECVF_RenderThreadSafe
 );
 
-static TAutoConsoleVariable<int32> CVarHeterogeneousVolumesShadowShadingRate(
+static TAutoConsoleVariable<float> CVarHeterogeneousVolumesShadowShadingRate(
 	TEXT("r.HeterogeneousVolumes.Shadows.ShadingRate"),
-	2,
-	TEXT("Debug tweak value (Default = 0)\n"),
+	1.0f,
+	TEXT("Shading rate (in-frustum) for computing shadows (Default = 1.0)\n"),
 	ECVF_RenderThreadSafe
 );
 
 static TAutoConsoleVariable<int32> CVarHeterogeneousVolumesShadowOutOfFrustumShadingRate(
 	TEXT("r.HeterogeneousVolumes.Shadows.OutOfFrustumShadingRate"),
 	2,
-	TEXT("Debug tweak value (Default = 0)\n"),
+	TEXT("Shading rate (out-of-frustum) for computing shadows (Default = 2.0)\n"),
 	ECVF_RenderThreadSafe
 );
 
@@ -668,9 +668,12 @@ class FTopLevelGridCalculateVoxelSize : public FGlobalShader
 		SHADER_PARAMETER(FVector3f, PrimitiveWorldBoundsMin)
 		SHADER_PARAMETER(FVector3f, PrimitiveWorldBoundsMax)
 		
-		SHADER_PARAMETER(float, ShadingRate)
+		SHADER_PARAMETER(float, ShadingRateInFrustum)
+		SHADER_PARAMETER(float, ShadingRateOutOfFrustum)
 		SHADER_PARAMETER(float, MinVoxelSizeInFrustum)
 		SHADER_PARAMETER(float, MinVoxelSizeOutOfFrustum)
+		SHADER_PARAMETER(float, DownsampleFactor)
+		SHADER_PARAMETER(int, bUseProjectedPixelSize)
 
 		SHADER_PARAMETER_RDG_BUFFER_UAV(RWStructuredBuffer<FTopLevelGridData>, RWTopLevelGridBuffer)
 	END_SHADER_PARAMETER_STRUCT()
@@ -1366,14 +1369,14 @@ void CalcViewBoundsAndMinimumVoxelSize(
 )
 {
 	TopLevelGridBounds = FBoxSphereBounds(ForceInit);
-	MinimumVoxelSize = BuildOptions.MinimumVoxelSizeOutsideFrustum;
+	MinimumVoxelSize = HeterogeneousVolumes::GetMinimumVoxelSizeOutsideFrustum();
 
 	// Build view bounds
 	FVector WorldCameraOrigin = View.ViewMatrices.GetViewOrigin();
 	FBoxSphereBounds WorldCameraBounds(FSphere(WorldCameraOrigin, HeterogeneousVolumes::GetMaxTraceDistance()));
 
 	float TanHalfFOV = HeterogeneousVolumes::CalcTanHalfFOV(View.FOV);
-	int32 HalfWidth = View.ViewRect.Width() * 0.5;
+	float HalfWidth = View.ViewRect.Width() * 0.5 / HeterogeneousVolumes::GetDownsampleFactor();
 	float PixelWidth = TanHalfFOV / HalfWidth;
 
 	for (int32 MeshBatchIndex = 0; MeshBatchIndex < View.HeterogeneousVolumesMeshBatches.Num(); ++MeshBatchIndex)
@@ -1405,8 +1408,8 @@ void CalcViewBoundsAndMinimumVoxelSize(
 				{
 					// Bandlimit minimum voxel size request with projected voxel size, based on shading rate
 					FVector VoxelCenter = PrimitiveBounds.Origin;
-					float Distance = FMath::Max(FVector(PrimitiveBounds.Origin - WorldCameraOrigin).Length() - TopLevelGridBounds.BoxExtent.Length(), 0.0);
-					float VoxelWidth = Distance * PixelWidth * HeterogeneousVolumes::GetShadingRateForFrustumGrid();
+					float Distance = FMath::Max(FVector(PrimitiveBounds.Origin - WorldCameraOrigin).Length() - PrimitiveBounds.BoxExtent.Length(), 0.0);
+					float VoxelWidth = Distance * PixelWidth * BuildOptions.ShadingRateInFrustum;
 
 					float PerVolumeMinimumVoxelSize = FMath::Max(VoxelWidth, HeterogeneousVolume->GetMinimumVoxelSize());
 					MinimumVoxelSize = FMath::Min(PerVolumeMinimumVoxelSize, MinimumVoxelSize);
@@ -1425,14 +1428,13 @@ void CalcGlobalBoundsAndMinimumVoxelSize(
 	float& GlobalMinimumVoxelSize
 )
 {
-	const bool bConvertToPixelSpace = (CVarHeterogeneousVolumesOrthoGridVoxelizationMode.GetValueOnAnyThread() == 0);
 	TopLevelGridBounds = FBoxSphereBounds(ForceInit);
-	GlobalMinimumVoxelSize = BuildOptions.MinimumVoxelSizeOutsideFrustum;
+	GlobalMinimumVoxelSize = HeterogeneousVolumes::GetMinimumVoxelSizeOutsideFrustum();
 
 	// Cycle through all Volume PrimitiveSceneProxies to collect bounds information and minimum voxel-size
 	for (int32 ViewIndex = 0; ViewIndex < Views.Num(); ViewIndex++)
 	{
-		float ViewMinimumVoxelSize = BuildOptions.MinimumVoxelSizeOutsideFrustum;
+		float ViewMinimumVoxelSize = HeterogeneousVolumes::GetMinimumVoxelSizeOutsideFrustum();
 		FBoxSphereBounds AggregatePrimitiveBounds = FBoxSphereBounds(ForceInit);
 
 		// Build view bounds
@@ -1441,7 +1443,7 @@ void CalcGlobalBoundsAndMinimumVoxelSize(
 		FBoxSphereBounds WorldCameraBounds(FSphere(WorldCameraOrigin, HeterogeneousVolumes::GetMaxTraceDistance()));
 
 		float TanHalfFOV = HeterogeneousVolumes::CalcTanHalfFOV(View.FOV);
-		int32 HalfWidth = View.ViewRect.Width() * 0.5;
+		float HalfWidth = View.ViewRect.Width() * 0.5 / HeterogeneousVolumes::GetDownsampleFactor();
 		float PixelWidth = TanHalfFOV / HalfWidth;
 
 		for (auto MeshBatchIt = HeterogeneousVolumesMeshBatches.begin(); MeshBatchIt != HeterogeneousVolumesMeshBatches.end(); ++MeshBatchIt)
@@ -1463,34 +1465,26 @@ void CalcGlobalBoundsAndMinimumVoxelSize(
 				{
 					continue;
 				}
-				// Only incorporate the primitive if it intersects with the canera bounding sphere where radius=MaxTraceDistance
+
 				const FBoxSphereBounds& PrimitiveBounds = HeterogeneousVolume->GetBounds();
+				AggregatePrimitiveBounds = Union(AggregatePrimitiveBounds, PrimitiveBounds);
 
+				bool bIntersectsViewFrustum = View.ViewFrustum.IntersectBox(PrimitiveBounds.Origin, PrimitiveBounds.BoxExtent);
+				float VoxelWidth = bIntersectsViewFrustum ? BuildOptions.ShadingRateInFrustum : BuildOptions.ShadingRateOutOfFrustum;
+				if (BuildOptions.bUseProjectedPixelSizeForOrthoGrid)
 				{
-					AggregatePrimitiveBounds = Union(AggregatePrimitiveBounds, PrimitiveBounds);
-
-					if (View.ViewFrustum.IntersectBox(PrimitiveBounds.Origin, PrimitiveBounds.BoxExtent))
-					{
-						float VoxelWidth = BuildOptions.MinimumVoxelSizeInFrustum;
-						// Legacy behavior converts world-space units to pixel-space units
-						if (bConvertToPixelSpace)
-						{
-							FVector VoxelCenter = PrimitiveBounds.Origin;
-							float Distance = FMath::Max(FVector(PrimitiveBounds.Origin - WorldCameraOrigin).Length() - AggregatePrimitiveBounds.BoxExtent.Length(), 0.0);
-							VoxelWidth *= Distance * PixelWidth;
-						}
-
-						float PerVolumeMinimumVoxelSize = FMath::Max(VoxelWidth, HeterogeneousVolume->GetMinimumVoxelSize());
-						ViewMinimumVoxelSize = FMath::Min(PerVolumeMinimumVoxelSize, ViewMinimumVoxelSize);
-					}
+					FVector VoxelCenter = PrimitiveBounds.Origin;
+					float Distance = FMath::Max(FVector(PrimitiveBounds.Origin - WorldCameraOrigin).Length() - PrimitiveBounds.SphereRadius, 0.0);
+					VoxelWidth *= Distance * PixelWidth;
 				}
-				// TODO: Out-of-frustum minimum voxel size per-primitive?
-				// else if (FBoxSphereBounds::BoxesIntersect(WorldCameraBounds, PrimitiveBounds))
+
+				VoxelWidth = FMath::Max(VoxelWidth, HeterogeneousVolume->GetMinimumVoxelSize());
+				ViewMinimumVoxelSize = FMath::Min(VoxelWidth, ViewMinimumVoxelSize);
 			}
 		}
 
 		// When converting to pixel-space units, clamp per-view minimum voxel-size to in-frustum minimum
-		if (bConvertToPixelSpace && View.ViewFrustum.IntersectBox(AggregatePrimitiveBounds.Origin, AggregatePrimitiveBounds.BoxExtent))
+		if (BuildOptions.bUseProjectedPixelSizeForOrthoGrid && View.ViewFrustum.IntersectBox(AggregatePrimitiveBounds.Origin, AggregatePrimitiveBounds.BoxExtent))
 		{
 			ViewMinimumVoxelSize = FMath::Max(ViewMinimumVoxelSize, HeterogeneousVolumes::GetMinimumVoxelSizeInFrustum());
 		}
@@ -1713,9 +1707,12 @@ void CalculateVoxelSize(
 					PassParameters->PrimitiveWorldBoundsMin = FVector3f(PrimitiveBounds.Origin - PrimitiveBounds.BoxExtent);
 					PassParameters->PrimitiveWorldBoundsMax = FVector3f(PrimitiveBounds.Origin + PrimitiveBounds.BoxExtent);
 
-					PassParameters->ShadingRate = HeterogeneousVolumes::GetShadingRateForOrthoGrid();
+					PassParameters->ShadingRateInFrustum = BuildOptions.ShadingRateInFrustum;
+					PassParameters->ShadingRateOutOfFrustum = BuildOptions.ShadingRateOutOfFrustum;
 					PassParameters->MinVoxelSizeInFrustum = FMath::Max(HeterogeneousVolume->GetMinimumVoxelSize(), HeterogeneousVolumes::GetMinimumVoxelSizeInFrustum());
-					PassParameters->MinVoxelSizeOutOfFrustum = BuildOptions.MinimumVoxelSizeOutsideFrustum;
+					PassParameters->MinVoxelSizeOutOfFrustum = HeterogeneousVolumes::GetMinimumVoxelSizeOutsideFrustum();
+					PassParameters->DownsampleFactor = HeterogeneousVolumes::GetDownsampleFactor();
+					PassParameters->bUseProjectedPixelSize = BuildOptions.bUseProjectedPixelSizeForOrthoGrid;
 
 					PassParameters->RWTopLevelGridBuffer = GraphBuilder.CreateUAV(TopLevelGridBuffer);
 				}
@@ -1811,7 +1808,7 @@ void GenerateRasterTiles(
 )
 {
 	const uint32 RasterTileVoxelResolution = HeterogeneousVolumes::GetBottomLevelGridResolution();
-	uint32 TileFactor = 64;
+	uint32 TileFactor = HeterogeneousVolumes::EnableIndirectionGrid() ? 64 : 1;
 	uint32 MaxNumRasterTiles = TileFactor * TopLevelGridResolution.X * TopLevelGridResolution.Y * TopLevelGridResolution.Z;
 	RasterTileBuffer = GraphBuilder.CreateBuffer(
 		FRDGBufferDesc::CreateStructuredDesc(sizeof(FRasterTileData), MaxNumRasterTiles),
@@ -2647,6 +2644,7 @@ void RasterizeVolumesIntoOrthoVoxelGrid(
 
 bool ShouldAddToVoxelGrid(const FPrimitiveSceneProxy* Proxy, const FViewInfo& View, const FVoxelGridBuildOptions& BuildOptions)
 {
+	// TODO: Should only add element if it is within the MaxTraceDistance
 	switch (BuildOptions.VoxelGridBuildMode)
 	{
 		default:
@@ -2691,6 +2689,12 @@ void CollectHeterogeneousVolumeMeshBatches(
 		}
 	}
 
+	// TODO: Temporarily disable per-light primitive gathering because it causes a crash with VoxelGrid builder
+	if (HeterogeneousVolumes::GetShadowMode() == HeterogeneousVolumes::EShadowMode::VoxelGrid)
+	{
+		return;
+	}
+
 	int32 NumPasses = LightSceneInfoCompact.Num();
 	for (int32 PassIndex = 0; PassIndex < NumPasses; ++PassIndex)
 	{
@@ -2948,6 +2952,7 @@ class FRenderVolumetricShadowMapForLightWithVoxelGridCS : public FGlobalShader
 		// Scene data
 		SHADER_PARAMETER_STRUCT_REF(FViewUniformShaderParameters, View)
 		SHADER_PARAMETER_STRUCT_INCLUDE(FSceneTextureParameters, SceneTextures)
+		SHADER_PARAMETER_STRUCT_REF(FBlueNoise, BlueNoise)
 
 		// Shadow data
 		SHADER_PARAMETER(float, ShadowStepSize)
@@ -3482,6 +3487,8 @@ void RenderVolumetricShadowMapForLightWithVoxelGrid(
 		// Scene data
 		PassParameters->View = View.ViewUniformBuffer;
 		PassParameters->SceneTextures = GetSceneTextureParameters(GraphBuilder, SceneTextures);
+		FBlueNoise BlueNoise = GetBlueNoiseGlobalParameters();
+		PassParameters->BlueNoise = CreateUniformBufferImmediate(BlueNoise, EUniformBufferUsage::UniformBuffer_SingleDraw);
 
 		// Ray Data
 		PassParameters->ShadowStepSize = HeterogeneousVolumes::GetStepSizeForShadows();
@@ -3699,6 +3706,8 @@ void RenderVolumetricShadowMapForCameraWithVoxelGrid(
 		// Scene data
 		PassParameters->View = View.ViewUniformBuffer;
 		PassParameters->SceneTextures = GetSceneTextureParameters(GraphBuilder, SceneTextures);
+		FBlueNoise BlueNoise = GetBlueNoiseGlobalParameters();
+		PassParameters->BlueNoise = CreateUniformBufferImmediate(BlueNoise, EUniformBufferUsage::UniformBuffer_SingleDraw);
 
 		// Ray Data
 		PassParameters->ShadowStepSize = HeterogeneousVolumes::GetStepSizeForShadows();
diff --git a/Engine/Source/Runtime/Renderer/Private/ShadowDepthRendering.cpp b/Engine/Source/Runtime/Renderer/Private/ShadowDepthRendering.cpp
--- a/Engine/Source/Runtime/Renderer/Private/ShadowDepthRendering.cpp
+++ b/Engine/Source/Runtime/Renderer/Private/ShadowDepthRendering.cpp
@@ -1943,7 +1943,8 @@ void FSceneRenderer::RenderShadowDepthMaps(FRDGBuilder& GraphBuilder, FDynamicSh
 		InstanceCullingManager.BeginDeferredCulling(GraphBuilder, Scene->GPUScene);
 	}
 
-	if (ShouldRenderHeterogeneousVolumes(Scene) && ShouldHeterogeneousVolumesCastShadows())
+	if (ShouldRenderHeterogeneousVolumes(Scene) && ShouldHeterogeneousVolumesCastShadows() &&
+		(HeterogeneousVolumes::GetShadowMode() == HeterogeneousVolumes::EShadowMode::LiveShading))
 	{
 		for (int32 ViewIndex = 0; ViewIndex < Views.Num(); ViewIndex++)
 		{
