Fixed crash when copying to view family texture that has a different dimension than the source texture.

diff --git a/Engine/Source/Runtime/Renderer/Private/ScreenPass.cpp b/Engine/Source/Runtime/Renderer/Private/ScreenPass.cpp
--- a/Engine/Source/Runtime/Renderer/Private/ScreenPass.cpp
+++ b/Engine/Source/Runtime/Renderer/Private/ScreenPass.cpp
@@ -94,9 +94,13 @@ FScreenPassTexture FScreenPassTexture::CopyFromSlice(FRDGBuilder& GraphBuilder,
 		OutputTexture = GraphBuilder.CreateTexture(Desc, TEXT("CopyToScreenPassTexture2D"));
 	}
 
+	const FIntPoint ViewSize = ScreenTextureSlice.ViewRect.Size();
+
 	FRHICopyTextureInfo CopyInfo;
 	CopyInfo.SourceSliceIndex = InputTextureSRV->Desc.FirstArraySlice;
 	CopyInfo.NumMips = InputTexture->Desc.NumMips;
+	CopyInfo.SourcePosition = CopyInfo.DestPosition = FIntVector(ScreenTextureSlice.ViewRect.Min.X, ScreenTextureSlice.ViewRect.Min.Y, 0);
+	CopyInfo.Size = FIntVector(ViewSize.X, ViewSize.Y, 1);
 
 	AddCopyTexturePass(
 		GraphBuilder,
