Fix mesh paint scene extension so that default parameters use the fallback color already defined in MeshPaintVirtualTexture::GetDefaultFallbackColor()
This fixes using black instead of white as the fallback color when mesh paint textures or virtual textures are fully disabled.

diff --git a/Engine/Source/Runtime/Renderer/Private/VT/MeshPaintVirtualTextureSceneExtension.cpp b/Engine/Source/Runtime/Renderer/Private/VT/MeshPaintVirtualTextureSceneExtension.cpp
--- a/Engine/Source/Runtime/Renderer/Private/VT/MeshPaintVirtualTextureSceneExtension.cpp
+++ b/Engine/Source/Runtime/Renderer/Private/VT/MeshPaintVirtualTextureSceneExtension.cpp
@@ -32,23 +32,24 @@ END_SHADER_PARAMETER_STRUCT()
 
 DECLARE_SCENE_UB_STRUCT(FMeshPaintTextureParameters, MeshPaint, RENDERER_API)
 
-static void GetDefaultMeshPaintParameters(FMeshPaintTextureParameters& Parameters, FRDGBuilder& GraphBuilder)
+static void GetMeshPaintParameters(MeshPaintVirtualTexture::FUniformParams const& InParameters, FMeshPaintTextureParameters& OutParameters)
 {
-	Parameters.PageTableTexture = GBlackUintTexture->TextureRHI;
-	Parameters.PhysicalTexture = GBlackTextureWithSRV->TextureRHI;
-	Parameters.PackedUniform = FUintVector4(0, 0, 0, 0);
+	OutParameters.PageTableTexture = InParameters.PageTableTexture ? InParameters.PageTableTexture : GBlackUintTexture->TextureRHI;
+	OutParameters.PhysicalTexture = InParameters.PhysicalTexture ? InParameters.PhysicalTexture : GBlackTextureWithSRV->TextureRHI;
+	OutParameters.PackedUniform = InParameters.PackedUniform;
+}
+
+static void GetDefaultMeshPaintParameters(FMeshPaintTextureParameters& OutParameters, FRDGBuilder& GraphBuilder)
+{
+	MeshPaintVirtualTexture::FUniformParams DefaultParameters;
+	GetMeshPaintParameters(DefaultParameters, OutParameters);
 }
 
 IMPLEMENT_SCENE_UB_STRUCT(FMeshPaintTextureParameters, MeshPaint, GetDefaultMeshPaintParameters);
 
 void FMeshPaintVirtualTextureSceneExtension::FRenderer::UpdateSceneUniformBuffer(FRDGBuilder& GraphBuilder, FSceneUniformBuffer& SceneUniformBuffer)
 {
-	MeshPaintVirtualTexture::FUniformParams MeshPaintParameters = MeshPaintVirtualTexture::GetUniformParams();
-
 	FMeshPaintTextureParameters Parameters;
-	Parameters.PageTableTexture = MeshPaintParameters.PageTableTexture ? MeshPaintParameters.PageTableTexture : GBlackUintTexture->TextureRHI;
-	Parameters.PhysicalTexture = MeshPaintParameters.PhysicalTexture ? MeshPaintParameters.PhysicalTexture : GBlackTextureWithSRV->TextureRHI;
-	Parameters.PackedUniform = MeshPaintParameters.PackedUniform;
-
+	GetMeshPaintParameters(MeshPaintVirtualTexture::GetUniformParams(), Parameters);
 	SceneUniformBuffer.Set(SceneUB::MeshPaint, Parameters);
 }
