Fixed use after scope of variables instantiated on the stack and passed as pointer into lambda captures

#rnx

diff --git a/Engine/Source/Runtime/Renderer/Private/VolumetricFogVoxelization.cpp b/Engine/Source/Runtime/Renderer/Private/VolumetricFogVoxelization.cpp
--- a/Engine/Source/Runtime/Renderer/Private/VolumetricFogVoxelization.cpp
+++ b/Engine/Source/Runtime/Renderer/Private/VolumetricFogVoxelization.cpp
@@ -746,7 +746,7 @@ void FSceneRenderer::VoxelizeFogVolumePrimitives(
 		const FIntRect ViewRect(0, 0, VolumetricFogViewGridSize.X, VolumetricFogViewGridSize.Y);
 
 		AddDrawDynamicMeshPass(GraphBuilder, RDG_EVENT_NAME("VoxelizeVolumePrimitives"), PassParameters, View, ViewRect,
-			[&View, VolumetricFogDistance, &VolumetricFogViewGridSize, &GridZParams](FDynamicPassMeshDrawListContext* DynamicMeshPassContext)
+			[&View, VolumetricFogDistance, VolumetricFogViewGridSize, GridZParams](FDynamicPassMeshDrawListContext* DynamicMeshPassContext)
 			{
 				FVoxelizeVolumeMeshProcessor PassMeshProcessor(
 					View.Family->Scene->GetRenderScene(),
