VisionOS Passthrough
-Requires VisionOS 2.0.
-Currently only supporting the mobile forward renderer.  Desktop deferred is next up.
-cvar r.AlphaInvertPass enables a render pass which inverts the alpha channel from unreal standard to the more common 0 = transparent, 1 = opaque.  This cvar is set to true in BaseVisionOSEngine.ini.
-To use one must be on VisionOS 2.0.  r.Mobile.PropagateAlpha=1 r.PostProcessing.PropagateAlpha=1r.AlphaInvertPass=1 and edit UESwift.swift  to mixed mode per the comments within.

diff --git a/Engine/Platforms/VisionOS/Plugins/Runtime/OpenXRVisionOS/Source/OXRVisionOS/Private/OXRVisionOSInstance.cpp b/Engine/Platforms/VisionOS/Plugins/Runtime/OpenXRVisionOS/Source/OXRVisionOS/Private/OXRVisionOSInstance.cpp
--- a/Engine/Platforms/VisionOS/Plugins/Runtime/OpenXRVisionOS/Source/OXRVisionOS/Private/OXRVisionOSInstance.cpp
+++ b/Engine/Platforms/VisionOS/Plugins/Runtime/OpenXRVisionOS/Source/OXRVisionOS/Private/OXRVisionOSInstance.cpp
@@ -281,7 +281,7 @@ XrResult FOXRVisionOSInstance::XrEnumerateEnvironmentBlendModes(
 		{
 			check(EnvironmentBlendModeCapacityInput >= 1); // Because we only need 1 we can't ever need to return XR_ERROR_SIZE_INSUFFICIENT
 
-			EnvironmentBlendModes[0] = XR_ENVIRONMENT_BLEND_MODE_OPAQUE;
+			EnvironmentBlendModes[0] = XR_ENVIRONMENT_BLEND_MODE_ALPHA_BLEND;
 		}
 		return XrResult::XR_SUCCESS;
 	}
diff --git a/Engine/Platforms/VisionOS/Plugins/Runtime/OpenXRVisionOS/Source/OXRVisionOS/Private/OXRVisionOSSession.cpp b/Engine/Platforms/VisionOS/Plugins/Runtime/OpenXRVisionOS/Source/OXRVisionOS/Private/OXRVisionOSSession.cpp
--- a/Engine/Platforms/VisionOS/Plugins/Runtime/OpenXRVisionOS/Source/OXRVisionOS/Private/OXRVisionOSSession.cpp
+++ b/Engine/Platforms/VisionOS/Plugins/Runtime/OpenXRVisionOS/Source/OXRVisionOS/Private/OXRVisionOSSession.cpp
@@ -1187,7 +1187,7 @@ XrResult FOXRVisionOSSession::XrEndFrame(
 	}
 
 	const XrFrameEndInfo& FrameEndInfo = *InFrameEndInfo;
-	check(FrameEndInfo.environmentBlendMode == XR_ENVIRONMENT_BLEND_MODE_OPAQUE); // only opaque is supported, see XrEnumerateEnvironmentBlendModes
+	//check(FrameEndInfo.environmentBlendMode == XR_ENVIRONMENT_BLEND_MODE_OPAQUE); // only opaque is supported, see XrEnumerateEnvironmentBlendModes
 	check(FrameEndInfo.layerCount <= OXRVisionOSSessionHelpers::MaxOXRVisionOSLayers);
     
     {
diff --git a/Engine/Platforms/VisionOS/Source/Runtime/Launch/Source/UESwift.swift b/Engine/Platforms/VisionOS/Source/Runtime/Launch/Source/UESwift.swift
--- a/Engine/Platforms/VisionOS/Source/Runtime/Launch/Source/UESwift.swift
+++ b/Engine/Platforms/VisionOS/Source/Runtime/Launch/Source/UESwift.swift
@@ -72,9 +72,9 @@ struct UEContentConfiguration: CompositorLayerConfiguration {
 		//configuration.colorFormat = .rgba16Float
 		configuration.colorFormat = .bgra8Unorm_srgb
 		
-		//configuration.depthFormat = .depth32Float  			//PF_R32_FLOAT   			// This is correct for mobile forward
-		configuration.depthFormat = .depth32Float_stencil8 		//PF_DepthStencil   // This is correct for deferred
-//PFSWITCH
+		//configuration.depthFormat = .depth32Float  			//PF_R32_FLOAT   	
+		configuration.depthFormat = .depth32Float_stencil8 		//PF_DepthStencil 
+
 		
 		configuration.defaultDepthRange = [Float.greatestFiniteMagnitude, 0.1]
 	}
@@ -96,6 +96,15 @@ class HostingViewFactory: NSObject
 @main
 struct UESwiftApp: App {
 	
+#if UE_SDK_VERSION_1
+	// VisionOS 1.x only supports full immersion style in Metal rendering apps.
+	@State private var style: ImmersionStyle = .full
+#else
+	// VisonOS 2+ also supported mixed immersion.
+	@State private var style: ImmersionStyle = .full
+	//@State private var style: ImmersionStyle = .mixed
+#endif
+	
 	@UIApplicationDelegateAdaptor(IOSAppDelegate.self) var delegate
 	
 	var body: some Scene
@@ -115,8 +124,16 @@ struct UESwiftApp: App {
 					FSwiftAppBootstrap.KickoffWithCompositingLayer(layerRenderer)
 			}
 		}
+
+#if UE_SDK_VERSION_1
 		//.upperLimbVisibility(.hidden)
 		.upperLimbVisibility(.visible)
+#else
+		//.upperLimbVisibility(.hidden)
+		//.upperLimbVisibility(.visible)
+		.upperLimbVisibility(.automatic) // Depth occlusion based hand fading, VisionOS 2+ only
+#endif
+		.immersionStyle(selection: $style, in: .mixed, .full)
 	}
 }
 
diff --git a/Engine/Shaders/Private/PostProcessTonemap.usf b/Engine/Shaders/Private/PostProcessTonemap.usf
--- a/Engine/Shaders/Private/PostProcessTonemap.usf
+++ b/Engine/Shaders/Private/PostProcessTonemap.usf
@@ -752,4 +752,32 @@ void MobileCustomResolve_MainPS(
 	OutColor.a = 1.0;
 }
 
-
+void AlphaInvert_MainPS(
+	in float4 SvPosition							: SV_POSITION, 
+	in noperspective float2 UV						: TEXCOORD0,
+	out				 float4 OutColor				: SV_Target0
+#if MOBILE_MULTI_VIEW_FALLBACK
+	, in nointerpolation uint EyeIndex				: VIEW_ID
+#elif MOBILE_MULTI_VIEW
+	, in nointerpolation uint MultiViewId : SV_ViewID
+#endif
+)
+{
+	uint ArrayIndex = 0;
+#if MOBILE_MULTI_VIEW
+	ArrayIndex = MultiViewId;
+#elif MOBILE_MULTI_VIEW_FALLBACK
+	ArrayIndex = EyeIndex;
+#else
+#endif
+	
+	half4 SceneColor;
+	#if MOBILE_MULTI_VIEW || MOBILE_MULTI_VIEW_FALLBACK
+		SceneColor = Texture2DArraySample(ColorTexture, ColorSampler, float3(UV.xy,ArrayIndex));
+	#else
+		SceneColor = Texture2DSample(ColorTexture, ColorSampler, UV);
+	#endif
+	
+	OutColor.rgb = SceneColor.rgb;
+	OutColor.a = 1.0 - SceneColor.a;
+}
diff --git a/Engine/Source/Programs/UnrealBuildTool/ToolChain/AppleToolChain.cs b/Engine/Source/Programs/UnrealBuildTool/ToolChain/AppleToolChain.cs
--- a/Engine/Source/Programs/UnrealBuildTool/ToolChain/AppleToolChain.cs
+++ b/Engine/Source/Programs/UnrealBuildTool/ToolChain/AppleToolChain.cs
@@ -629,6 +629,7 @@ namespace UnrealBuildTool
 				}
 			}
 			Arguments.Add("-DUE_USE_SWIFT_UI_MAIN=" + (bUseSwiftUIMain ? "1" : "0"));
+					
 			return Output;
 		}
 
@@ -678,6 +679,10 @@ namespace UnrealBuildTool
 			if (bUseSwiftUIMain)
 			{
 				Arguments.Add("-DUE_USE_SWIFT_UI_MAIN");
+				if (ToolChainSettings.Value.SDKVersionFloat < 2.0)
+				{
+					Arguments.Add("-DUE_SDK_VERSION_1");
+				}
 			}
 
 			Action CompileAction = Graph.CreateAction(ActionType.Compile);
@@ -726,8 +731,13 @@ namespace UnrealBuildTool
 				if (bUseSwiftUIMain)
 				{
 					Arguments.Add("-DUE_USE_SWIFT_UI_MAIN");
+					if (ToolChainSettings.Value.SDKVersionFloat < 2.0)
+					{
+						Arguments.Add("-DUE_SDK_VERSION_1");
+					}
 				}
 
+
 				// now make an action to export the swift code as a header Obj-C bridging
 				Action HeaderAction = Graph.CreateAction(ActionType.CompileModuleInterface);
 				HeaderAction.CommandPath = FileReference.Combine(ToolChainSettings.Value.ToolchainDir, "swiftc");
diff --git a/Engine/Source/Runtime/Apple/MetalRHI/Private/MetalViewport.cpp b/Engine/Source/Runtime/Apple/MetalRHI/Private/MetalViewport.cpp
--- a/Engine/Source/Runtime/Apple/MetalRHI/Private/MetalViewport.cpp
+++ b/Engine/Source/Runtime/Apple/MetalRHI/Private/MetalViewport.cpp
@@ -698,6 +698,11 @@ void FMetalViewport::PresentImmersive(const MetalRHIVisionOS::PresentImmersivePa
 				
 				FMetalCommandBuffer* CurrentCommandBuffer = Context.GetCurrentCommandBuffer();
 				
+#if ENABLE_METAL_GPUPROFILE
+				FMetalProfiler* Profiler = FMetalProfiler::GetProfiler();
+				FMetalCommandBufferStats* Stats = Profiler->AllocateCommandBuffer(CurrentCommandBuffer->GetMTLCmdBuffer(), 0);
+#endif
+				
 				{
 					dispatch_semaphore_t& FrameSemaphore = Device.GetFrameSemaphore();
 					dispatch_retain(FrameSemaphore);
@@ -715,6 +720,7 @@ void FMetalViewport::PresentImmersive(const MetalRHIVisionOS::PresentImmersivePa
 				cp_frame_t CompositorServicesFrame = VisionOSParams.SwiftFrame;
 				
 				Buffer = Context.Finalize();
+				METAL_GPUPROFILE(Stats->End(CurrentCommandBuffer->GetMTLCmdBuffer()));			
 				Context.GetDevice().GetCommandQueue().CommitCommandBuffer(Buffer);
 				Context.ResetContext();
 				
diff --git a/Engine/Source/Runtime/Renderer/Private/MobileShadingRenderer.cpp b/Engine/Source/Runtime/Renderer/Private/MobileShadingRenderer.cpp
--- a/Engine/Source/Runtime/Renderer/Private/MobileShadingRenderer.cpp
+++ b/Engine/Source/Runtime/Renderer/Private/MobileShadingRenderer.cpp
@@ -1628,6 +1628,12 @@ void FMobileSceneRenderer::RenderForward(FRDGBuilder& GraphBuilder, FRDGTextureR
 			RenderForwardSinglePass(GraphBuilder, PassParameters, ViewContext, SceneTextures);
 		}
 	}
+	
+	static auto CVarAlphaInvertPassLocal = IConsoleManager::Get().FindConsoleVariable(TEXT("r.AlphaInvertPass"));
+	if(CVarAlphaInvertPassLocal != nullptr && CVarAlphaInvertPassLocal->GetBool())
+	{
+		AddAlphaInvertPass(GraphBuilder, MainView, SceneTextures);
+	}
 }
 
 void FMobileSceneRenderer::RenderForwardSinglePass(FRDGBuilder& GraphBuilder, FMobileRenderPassParameters* PassParameters, FRenderViewContext& ViewContext, FSceneTextures& SceneTextures)
diff --git a/Engine/Source/Runtime/Renderer/Private/PostProcess/PostProcessTonemap.cpp b/Engine/Source/Runtime/Renderer/Private/PostProcess/PostProcessTonemap.cpp
--- a/Engine/Source/Runtime/Renderer/Private/PostProcess/PostProcessTonemap.cpp
+++ b/Engine/Source/Runtime/Renderer/Private/PostProcess/PostProcessTonemap.cpp
@@ -72,6 +72,12 @@ static TAutoConsoleVariable<int32> CVarBackbufferQuantizationDitheringOverride(
 	TEXT("Disabled by default. Instead is automatically found out by FSceneViewFamily::RenderTarget's pixel format of the backbuffer."),
 	ECVF_RenderThreadSafe);
 
+static TAutoConsoleVariable<int32> CVarAlphaInvertPass(
+	TEXT("r.AlphaInvertPass"),
+	0,
+	TEXT("Whether to run a render pass to un-invert the alpha value from unreal standard to the much more common standard where alpha 0 is fully transparent and alpha 1 is fully opaque."),
+	ECVF_RenderThreadSafe);
+
 const int32 GTonemapComputeTileSizeX = 8;
 const int32 GTonemapComputeTileSizeY = 8;
 
@@ -1166,3 +1172,85 @@ void AddMobileCustomResolvePass(FRDGBuilder& GraphBuilder, const FViewInfo& View
 			RenderMobileCustomResolve(RHICmdList, View, SubpassMSAASamples, SceneTextures);
 		});
 }
+
+// Pixel shader to un-invert the alpha channel for output.
+class FAlphaInvertPS : public FGlobalShader
+{
+	DECLARE_GLOBAL_SHADER(FAlphaInvertPS);
+
+	SHADER_USE_PARAMETER_STRUCT(FAlphaInvertPS, FGlobalShader);
+	BEGIN_SHADER_PARAMETER_STRUCT(FParameters, )
+		SHADER_PARAMETER_STRUCT_INCLUDE(FViewShaderParameters, View)
+		SHADER_PARAMETER_RDG_TEXTURE(Texture2D, ColorTexture)
+		SHADER_PARAMETER_SAMPLER(SamplerState, ColorSampler)
+		SHADER_PARAMETER(FVector4f, ColorScale0)
+	END_SHADER_PARAMETER_STRUCT()
+};
+
+IMPLEMENT_GLOBAL_SHADER(FAlphaInvertPS, "/Engine/Private/PostProcessTonemap.usf", "AlphaInvert_MainPS", SF_Pixel);
+
+void RenderAlphaInvertPass(FRHICommandList& RHICmdList, const FViewInfo& View, FSceneTextures& SceneTextures)
+{
+	// Part of scene rendering pass
+	check(RHICmdList.IsInsideRenderPass());
+	SCOPED_DRAW_EVENT(RHICmdList, MobileTonemapSubpass);
+
+	const FIntPoint TargetSize = SceneTextures.Color.Resolve->Desc.Extent;
+	
+	TShaderMapRef<FScreenVS> VertexShader(View.ShaderMap);
+	TShaderMapRef<FAlphaInvertPS> PixelShader(View.ShaderMap);
+
+	FAlphaInvertPS::FParameters PSShaderParameters;
+	PSShaderParameters.View = View.GetShaderParameters();
+	PSShaderParameters.ColorTexture = SceneTextures.Color.Resolve;
+	PSShaderParameters.ColorSampler = TStaticSamplerState<SF_Point, AM_Clamp, AM_Clamp, AM_Clamp>::GetRHI();
+
+	FGraphicsPipelineStateInitializer GraphicsPSOInit;
+	RHICmdList.ApplyCachedRenderTargets(GraphicsPSOInit);
+	GraphicsPSOInit.BlendState = TStaticBlendState<>::GetRHI();
+	GraphicsPSOInit.RasterizerState = TStaticRasterizerState<>::GetRHI();
+	GraphicsPSOInit.DepthStencilState = TStaticDepthStencilState<false, CF_Always>::GetRHI();
+	GraphicsPSOInit.BoundShaderState.VertexDeclarationRHI = GFilterVertexDeclaration.VertexDeclarationRHI;
+	GraphicsPSOInit.BoundShaderState.VertexShaderRHI = VertexShader.GetVertexShader();
+	GraphicsPSOInit.BoundShaderState.PixelShaderRHI = PixelShader.GetPixelShader();
+	GraphicsPSOInit.PrimitiveType = PT_TriangleList;
+
+	SetGraphicsPipelineState(RHICmdList, GraphicsPSOInit, 0);
+	SetShaderParameters(RHICmdList, PixelShader, PixelShader.GetPixelShader(), PSShaderParameters);
+	RHICmdList.SetViewport(0, 0, 0.0f, TargetSize.X, TargetSize.Y, 1.0f);
+
+	DrawRectangle(
+		RHICmdList,
+		0, 0,
+		TargetSize.X, TargetSize.Y,
+		0, 0,
+		TargetSize.X, TargetSize.Y,
+		TargetSize,
+		TargetSize,
+		VertexShader,
+		EDRF_UseTriangleOptimization,
+		View.GetStereoPassInstanceFactor());
+}
+
+BEGIN_SHADER_PARAMETER_STRUCT(FAlphaInvertParameters, )
+	SHADER_PARAMETER_STRUCT_INCLUDE(FViewShaderParameters, View)
+	SHADER_PARAMETER_RDG_TEXTURE(Texture2D, ColorTexture)
+	RENDER_TARGET_BINDING_SLOTS()
+END_SHADER_PARAMETER_STRUCT()
+
+void AddAlphaInvertPass(FRDGBuilder& GraphBuilder, const FViewInfo& View, FSceneTextures& SceneTextures)
+{
+	FAlphaInvertParameters* PassParameters = GraphBuilder.AllocParameters<FAlphaInvertParameters>();
+	PassParameters->View = View.GetShaderParameters();
+	PassParameters->ColorTexture = SceneTextures.Color.Resolve;
+	PassParameters->RenderTargets[0] = FRenderTargetBinding(SceneTextures.Color.Resolve, ERenderTargetLoadAction::ELoad);
+
+	GraphBuilder.AddPass(
+		RDG_EVENT_NAME("AlphaInvertPass"),
+		PassParameters,
+		ERDGPassFlags::Raster,
+		[&View, &SceneTextures](FRHICommandListImmediate& RHICmdList)
+		{
+			RenderAlphaInvertPass(RHICmdList, View, SceneTextures);
+		});
+}
diff --git a/Engine/Source/Runtime/Renderer/Private/PostProcess/PostProcessTonemap.h b/Engine/Source/Runtime/Renderer/Private/PostProcess/PostProcessTonemap.h
--- a/Engine/Source/Runtime/Renderer/Private/PostProcess/PostProcessTonemap.h
+++ b/Engine/Source/Runtime/Renderer/Private/PostProcess/PostProcessTonemap.h
@@ -80,3 +80,5 @@ struct FTonemapInputs
 FScreenPassTexture AddTonemapPass(FRDGBuilder& GraphBuilder, const FViewInfo& View, const FTonemapInputs& Inputs);
 void RenderMobileCustomResolve(FRHICommandList& RHICmdList, const FViewInfo& View, const int32 SubpassMSAASamples, FSceneTextures& SceneTextures);
 void AddMobileCustomResolvePass(FRDGBuilder& GraphBuilder, const FViewInfo& View, FSceneTextures& SceneTextures, FRDGTextureRef ViewFamilyTexture);
+
+void AddAlphaInvertPass(FRDGBuilder& GraphBuilder, const FViewInfo& View, FSceneTextures& SceneTextures);
diff --git a/Templates/TP_VirtualRealityBP/Config/VisionOS/VisionOSEngine.ini b/Templates/TP_VirtualRealityBP/Config/VisionOS/VisionOSEngine.ini
--- a/Templates/TP_VirtualRealityBP/Config/VisionOS/VisionOSEngine.ini
+++ b/Templates/TP_VirtualRealityBP/Config/VisionOS/VisionOSEngine.ini
@@ -1,4 +1,10 @@
 
 [/Script/Engine.RendererSettings]
+// These settings are currently required for all visionOS projects.
 vr.InstancedStereo=False
-vr.MobileMultiView=False
\ No newline at end of file
+vr.MobileMultiView=False
+
+//Enable these settings for .mixed immesion mode.
+//r.Mobile.PropagateAlpha=1
+//r.PostProcessing.PropagateAlpha=1
+//r.AlphaInvertPass=1
\ No newline at end of file
