Always include whether or not we have a dynamic instance offset in the base payload.
This will allow programs only interested in fetching the InstanceId to not have to bring the extended payload and FContextBatchInfo: calling LoadInstanceCullingSetup is now enough to get the final InstanceId value. We only need FInstanceProcessingGPULoadBalancer::FShaderParameters as dependencies to perform the computation

This means the number of extended payload data we can index in 2^24 = 16M
Replace some hardcoded numbers with macros
Added natvis for FPackedBatch/FPackedItem

#prelight 66d862da49496108c450d1eb

diff --git a/Engine/Shaders/Private/InstanceCulling/BuildInstanceDrawCommands.usf b/Engine/Shaders/Private/InstanceCulling/BuildInstanceDrawCommands.usf
--- a/Engine/Shaders/Private/InstanceCulling/BuildInstanceDrawCommands.usf
+++ b/Engine/Shaders/Private/InstanceCulling/BuildInstanceDrawCommands.usf
@@ -44,6 +44,7 @@
 
 
 #include "InstanceCullingLoadBalancer.ush"
+#include "InstanceCullingSetup.ush"
 
 uint InstanceSceneDataSOAStride;
 uint GPUSceneNumInstances;
@@ -243,25 +244,20 @@ void InstanceCullBuildInstanceIdBufferCS(uint3 GroupId : SV_GroupID, int GroupTh
 	BatchInfo.bAllowOcclusionCulling = true;
 #endif // ENABLE_BATCH_MODE
 
-	FInstanceWorkSetup WorkSetup = InstanceCullingLoadBalancer_Setup(GroupId, GroupThreadIndex, GetItemDataOffset(BatchInfo, CurrentBatchProcessingMode));
+	FInstanceCullingSetup InstanceCullingSetup = LoadInstanceCullingSetup(GroupId, GroupThreadIndex, BatchInfo.DynamicInstanceIdOffset, BatchInfo.DynamicInstanceIdMax, GetItemDataOffset(BatchInfo, CurrentBatchProcessingMode));
+
+	FInstanceWorkSetup WorkSetup = InstanceCullingSetup.InstanceWorkSetup;
 
 	if (!WorkSetup.bValid)
 	{
 		return;
 	}
 
+	uint InstanceId = InstanceCullingSetup.InstanceId;
+
 	// Extract the draw command payload
 	const FInstanceCullingPayload Payload = LoadInstanceCullingPayload(WorkSetup.Item.Payload, BatchInfo);
 
-	uint InstanceId = WorkSetup.Item.InstanceDataOffset + uint(WorkSetup.LocalItemIndex);
-
-	// Check dynamic flag to offset the instance ID
-	if (Payload.bDynamicInstanceDataOffset)
-	{
-		InstanceId += BatchInfo.DynamicInstanceIdOffset;
-		checkSlow(InstanceId < BatchInfo.DynamicInstanceIdMax);
-	}
-
 	// Load auxiliary per-instanced-draw command info
 	const FDrawCommandDesc DrawCommandDesc = UnpackDrawCommandDesc(DrawCommandDescs[Payload.IndirectArgIndex]);
 
diff --git a/Engine/Shaders/Private/InstanceCulling/InstanceCullingCommon.ush b/Engine/Shaders/Private/InstanceCulling/InstanceCullingCommon.ush
--- a/Engine/Shaders/Private/InstanceCulling/InstanceCullingCommon.ush
+++ b/Engine/Shaders/Private/InstanceCulling/InstanceCullingCommon.ush
@@ -3,6 +3,7 @@
 #pragma once
 
 #include "../Common.ush"
+#include "/Engine/Shared/InstanceCullingDefinitions.h"
 
 #ifndef BATCH_PROCESSING_MODE_NUM
 #define BATCH_PROCESSING_MODE_NUM 1
@@ -120,7 +121,12 @@ struct FPackedInstanceCullingPayload
 
 bool IsExtendedInstanceCullingPayload(uint WorkItemPayload)
 {
-	return (WorkItemPayload & 1U) != 0U;
+	return (WorkItemPayload & INSTANCE_CULLING_PRESERVE_INSTANCE_ORDER_BIT_MASK) != 0U;
+}
+
+bool IsDynamicInstanceDataOffset(uint WorkItemPayload)
+{
+	return (WorkItemPayload & INSTANCE_CULLING_DYNAMIC_INSTANCE_DATA_OFFSET_BIT_MASK) != 0U;
 }
 
 FInstanceCullingPayload UnpackInstanceCullingPayload(FPackedInstanceCullingPayload Payload)
@@ -149,11 +155,13 @@ FInstanceCullingPayload LoadInstanceCullingPayload(uint WorkItemPayload, FContex
 {
 	FInstanceCullingPayload Output;
 
+	Output.bDynamicInstanceDataOffset = IsDynamicInstanceDataOffset(WorkItemPayload);
+
 #if ENABLE_EXTENDED_INSTANCE_CULLING_PAYLOADS
 	BRANCH
 	if (IsExtendedInstanceCullingPayload(WorkItemPayload))
 	{
-		Output = UnpackInstanceCullingPayload(InstanceCullingPayloads[BatchInfo.PayloadOffset + (WorkItemPayload >> 1U)]);
+		Output = UnpackInstanceCullingPayload(InstanceCullingPayloads[BatchInfo.PayloadOffset + (WorkItemPayload >> INSTANCE_CULLING_PAYLOAD_NUM_COMMON_BITS)]);
 		Output.IndirectArgIndex += BatchInfo.IndirectArgOffset;
 		Output.CompactionDataIndex += BatchInfo.CompactionDataOffset;
 		return Output;
@@ -163,8 +171,7 @@ FInstanceCullingPayload LoadInstanceCullingPayload(uint WorkItemPayload, FContex
 	checkSlow(!IsExtendedInstanceCullingPayload(WorkItemPayload));
 #endif
 
-	Output.bDynamicInstanceDataOffset = (WorkItemPayload & 2U) != 0U;
-	Output.IndirectArgIndex = BatchInfo.IndirectArgOffset + (WorkItemPayload >> 2U);
+	Output.IndirectArgIndex = BatchInfo.IndirectArgOffset + (WorkItemPayload >> INSTANCE_CULLING_PAYLOAD_NUM_COMMON_BITS);
 	Output.InstanceDataOffset = 0xFFFFFFFFU;
 	Output.RunInstanceOffset = 0xFFFFFFFFU;
 	Output.CompactionDataIndex = 0xFFFFFFFFU;
diff --git a/Engine/Shaders/Private/InstanceCulling/InstanceCullingSetup.ush b/Engine/Shaders/Private/InstanceCulling/InstanceCullingSetup.ush
new file mode 100644
--- /dev/null
+++ b/Engine/Shaders/Private/InstanceCulling/InstanceCullingSetup.ush
@@ -0,0 +1,43 @@
+// Copyright Epic Games, Inc. All Rights Reserved.
+
+#pragma once 
+
+#include "../Common.ush"
+
+#include "InstanceCullingCommon.ush"
+
+#include "InstanceCullingLoadBalancer.ush"
+
+struct FInstanceCullingSetup
+{
+	FInstanceWorkSetup InstanceWorkSetup;
+	bool bValid;
+	bool bIsDynamic;
+	uint InstanceId;
+};
+
+// InDynamicInstanceIdOffset/InDynamicInstanceIdMax are only needed if you want to support dynamic instances, otherwise you can provide 0
+// ItemDataOffset is only needed if you work in batched mode where multiple load balancers are processed at the same time, see ENABLE_BATCH_MODE
+FInstanceCullingSetup LoadInstanceCullingSetup(uint3 GroupId, uint GroupThreadIndex, uint InDynamicInstanceIdOffset, uint InDynamicInstanceIdMax, uint ItemDataOffset)
+{
+	FInstanceCullingSetup Result = (FInstanceCullingSetup)0;
+
+	Result.InstanceWorkSetup = InstanceCullingLoadBalancer_Setup(GroupId, GroupThreadIndex, ItemDataOffset);
+
+	if (!Result.InstanceWorkSetup.bValid)
+	{
+		return Result;
+	}
+
+	Result.bValid = true;
+	Result.InstanceId = Result.InstanceWorkSetup.Item.InstanceDataOffset + uint(Result.InstanceWorkSetup.LocalItemIndex);
+	// Check dynamic flag to offset the instance ID
+	if (IsDynamicInstanceDataOffset(Result.InstanceWorkSetup.Item.Payload))
+	{
+		Result.InstanceId += InDynamicInstanceIdOffset;
+		checkSlow(InstanceId < InDynamicInstanceIdMax || (InDynamicInstanceIdOffset == 0 && InDynamicInstanceIdMax == 0));
+		Result.bIsDynamic = true;
+	}
+
+	return Result;
+}
diff --git a/Engine/Shaders/Shared/InstanceCullingDefinitions.h b/Engine/Shaders/Shared/InstanceCullingDefinitions.h
new file mode 100644
--- /dev/null
+++ b/Engine/Shaders/Shared/InstanceCullingDefinitions.h
@@ -0,0 +1,8 @@
+// Copyright Epic Games, Inc. All Rights Reserved.
+
+#pragma once
+
+#define INSTANCE_CULLING_PRESERVE_INSTANCE_ORDER_BIT_MASK 1U
+#define INSTANCE_CULLING_DYNAMIC_INSTANCE_DATA_OFFSET_BIT_MASK 2U
+
+#define INSTANCE_CULLING_PAYLOAD_NUM_COMMON_BITS 2U
diff --git a/Engine/Source/Runtime/Renderer/Private/InstanceCulling/InstanceCullingContext.cpp b/Engine/Source/Runtime/Renderer/Private/InstanceCulling/InstanceCullingContext.cpp
--- a/Engine/Source/Runtime/Renderer/Private/InstanceCulling/InstanceCullingContext.cpp
+++ b/Engine/Source/Runtime/Renderer/Private/InstanceCulling/InstanceCullingContext.cpp
@@ -13,6 +13,7 @@
 #include "ScenePrivate.h"
 #include "SystemTextures.h"
 #include "InstanceCulling/InstanceCullingManager.h"
+#include "InstanceCullingDefinitions.h"
 #include "InstanceCullingLoadBalancer.h"
 #include "InstanceCullingMergedContext.h"
 #include "InstanceCullingOcclusionQuery.h"
@@ -298,20 +299,20 @@ void FInstanceCullingContext::AddInstancesToDrawCommand(uint32 IndirectArgsOffse
 	const bool bPreserveInstanceOrder = EnumHasAnyFlags(InstanceFlags, EInstanceFlags::PreserveInstanceOrder);
 	const bool bForceInstanceCulling = EnumHasAnyFlags(InstanceFlags, EInstanceFlags::ForceInstanceCulling);	
 
-	uint32 Payload;
+	uint32 Payload = (bDynamicInstanceDataOffset ? INSTANCE_CULLING_DYNAMIC_INSTANCE_DATA_OFFSET_BIT_MASK : 0U);
 	if (bPreserveInstanceOrder)
 	{
 		checkSlow(!EnumHasAnyFlags(Flags, EInstanceCullingFlags::NoInstanceOrderPreservation)); // this should have already been handled
 
 		// We need to provide full payload data for these instances
-		// NOTE: The extended payload data flag is in the lowest bit instead of the highest because the payload is not a full dword		
-		Payload = 1 | (uint32(PayloadData.Num()) << 1U);
+		// NOTE: The extended payload data flag is in the lowest bit instead of the highest because the payload is not a full dword, see FInstanceCullingLoadBalancerBase::PackItem
+		Payload |= (INSTANCE_CULLING_PRESERVE_INSTANCE_ORDER_BIT_MASK | (uint32(PayloadData.Num()) << INSTANCE_CULLING_PAYLOAD_NUM_COMMON_BITS));
 		PayloadData.Emplace(bDynamicInstanceDataOffset, IndirectArgsOffset, InstanceDataOffset, RunOffset, DrawCommandCompactionData.Num());
 	}
 	else
 	{
 		// Conserve space by packing the relevant payload information into the dword
-		Payload = (IndirectArgsOffset << 2U) | (bDynamicInstanceDataOffset ? 2U : 0U);
+		Payload |= (IndirectArgsOffset << INSTANCE_CULLING_PAYLOAD_NUM_COMMON_BITS);
 	}
 
 	// We special-case the single-instance (i.e., regular primitives) as they don't need culling (again), except where explicitly specified.
