Don't allow hit group indexing for basic raytracing (fix fallout from inline SBT changes)

diff --git a/Engine/Source/Runtime/Renderer/Private/RayTracing/RayTracingBasicShaders.cpp b/Engine/Source/Runtime/Renderer/Private/RayTracing/RayTracingBasicShaders.cpp
--- a/Engine/Source/Runtime/Renderer/Private/RayTracing/RayTracingBasicShaders.cpp
+++ b/Engine/Source/Runtime/Renderer/Private/RayTracing/RayTracingBasicShaders.cpp
@@ -46,7 +46,7 @@ FBasicRayTracingPipeline GetBasicRayTracingPipeline(FRHICommandList& RHICmdList,
 		
 	FRayTracingShaderBindingTableInitializer SBTInitializer;
 	SBTInitializer.ShaderBindingMode = ERayTracingShaderBindingMode::RTPSO;
-	SBTInitializer.HitGroupIndexingMode = ERayTracingHitGroupIndexingMode::Allow;
+	SBTInitializer.HitGroupIndexingMode = ERayTracingHitGroupIndexingMode::Disallow;
 	SBTInitializer.NumGeometrySegments = 1;
 	SBTInitializer.NumShaderSlotsPerGeometrySegment = RAY_TRACING_NUM_SHADER_SLOTS;
 	SBTInitializer.NumMissShaderSlots = 1;
