Light Function Atlas - fixed ensure to account for the real precision of the atlas packing instead of a random threashold.

#rnx
#tests none

diff --git a/Engine/Source/Runtime/Renderer/Private/LightFunctionAtlas.cpp b/Engine/Source/Runtime/Renderer/Private/LightFunctionAtlas.cpp
--- a/Engine/Source/Runtime/Renderer/Private/LightFunctionAtlas.cpp
+++ b/Engine/Source/Runtime/Renderer/Private/LightFunctionAtlas.cpp
@@ -650,8 +650,9 @@ void FLightFunctionAtlas::RenderLightFunctionAtlas(FRDGBuilder& GraphBuilder, TA
 			const uint32 PackedAtlasSlotMinV = uint32(round(AtlasSlot.MinV * 65536.0f));
 			const uint32 PackedAtlasSlotMinUV = (PackedAtlasSlotMinU | (PackedAtlasSlotMinV << 16));
 
-			ensure(FMath::IsNearlyEqual((PackedAtlasSlotMinUV & 0xFFFF) / 65536.0f, AtlasSlot.MinU));
-			ensure(FMath::IsNearlyEqual(((PackedAtlasSlotMinUV >> 16) & 0xFFFF) / 65536.0f, AtlasSlot.MinV));
+			const float IntegerUVTolerance = 0.5f / 65536.0f;
+			ensure(FMath::IsNearlyEqual((PackedAtlasSlotMinUV & 0xFFFF) / 65536.0f, AtlasSlot.MinU, IntegerUVTolerance));
+			ensure(FMath::IsNearlyEqual(((PackedAtlasSlotMinUV >> 16) & 0xFFFF) / 65536.0f, AtlasSlot.MinV, IntegerUVTolerance));
 
 			float TanOuterAngle = LightType == LightType_Spot ? FMath::Tan(LightSceneInfo->Proxy->GetOuterConeAngle()) : -1.0f;
 			if (LightType == LightType_Rect)
