PR #12143: fix: Incorrect RenderTarget was used when RTV Batch was broken

diff --git a/Engine/Source/Runtime/Renderer/Private/VT/RuntimeVirtualTextureProducer.cpp b/Engine/Source/Runtime/Renderer/Private/VT/RuntimeVirtualTextureProducer.cpp
--- a/Engine/Source/Runtime/Renderer/Private/VT/RuntimeVirtualTextureProducer.cpp
+++ b/Engine/Source/Runtime/Renderer/Private/VT/RuntimeVirtualTextureProducer.cpp
@@ -81,7 +81,30 @@ void FRuntimeVirtualTextureFinalizer::Finalize(FRDGBuilder& GraphBuilder)
 	int32 BatchSize = 0;
 	for (auto Entry : Tiles)
 	{
-		RuntimeVirtualTexture::FRenderPageDesc& RenderPageDesc = RenderPageBatchDesc.PageDescs[BatchSize];
+		bool bBreakBatchForTextures = false;
+		for (int LayerIndex = 0; LayerIndex < RuntimeVirtualTexture::MaxTextureLayers; ++LayerIndex)
+		{
+			// This should never happen which is why we don't bother sorting to maximize batch size
+			bBreakBatchForTextures |= (RenderPageBatchDesc.Targets[LayerIndex].Texture != Entry.Targets[LayerIndex].TextureRHI);
+		}
+
+		if (BatchSize == RuntimeVirtualTexture::EMaxRenderPageBatch || bBreakBatchForTextures)
+		{
+			RenderPageBatchDesc.NumPageDescs = BatchSize;
+			RuntimeVirtualTexture::RenderPages(GraphBuilder, RenderPageBatchDesc);
+			BatchSize = 0;
+		}
+
+		if (bBreakBatchForTextures)
+		{
+			for (int LayerIndex = 0; LayerIndex < RuntimeVirtualTexture::MaxTextureLayers; ++LayerIndex)
+			{
+				RenderPageBatchDesc.Targets[LayerIndex].Texture = Entry.Targets[LayerIndex].TextureRHI != nullptr ? Entry.Targets[LayerIndex].TextureRHI->GetTexture2D() : nullptr;
+				RenderPageBatchDesc.Targets[LayerIndex].PooledRenderTarget = Entry.Targets[LayerIndex].PooledRenderTarget;
+			}
+		}
+
+		RuntimeVirtualTexture::FRenderPageDesc& RenderPageDesc = RenderPageBatchDesc.PageDescs[BatchSize++];
 
 		const float X = (float)FMath::ReverseMortonCode2_64(Entry.vAddress);
 		const float Y = (float)FMath::ReverseMortonCode2_64(Entry.vAddress >> 1);
@@ -103,28 +126,6 @@ void FRuntimeVirtualTextureFinalizer::Finalize(FRDGBuilder& GraphBuilder)
 			RenderPageDesc.DestBox[LayerIndex] = FBox2D(DestinationBoxStart0, DestinationBoxStart0 + FVector2D(TileSize, TileSize));
 		}
 
-		bool bBreakBatchForTextures = false;
-		for (int LayerIndex = 0; LayerIndex < RuntimeVirtualTexture::MaxTextureLayers; ++LayerIndex)
-		{
-			// This should never happen which is why we don't bother sorting to maximize batch size
-			bBreakBatchForTextures |= (RenderPageBatchDesc.Targets[LayerIndex].Texture != Entry.Targets[LayerIndex].TextureRHI);
-		}
-
-		if (++BatchSize == RuntimeVirtualTexture::EMaxRenderPageBatch || bBreakBatchForTextures)
-		{
-			RenderPageBatchDesc.NumPageDescs = BatchSize;
-			RuntimeVirtualTexture::RenderPages(GraphBuilder, RenderPageBatchDesc);
-			BatchSize = 0;
-		}
-
-		if (bBreakBatchForTextures)
-		{
-			for (int LayerIndex = 0; LayerIndex < RuntimeVirtualTexture::MaxTextureLayers; ++LayerIndex)
-			{
-				RenderPageBatchDesc.Targets[LayerIndex].Texture = Tiles[0].Targets[LayerIndex].TextureRHI != nullptr ? Tiles[0].Targets[LayerIndex].TextureRHI->GetTexture2D() : nullptr;
-				RenderPageBatchDesc.Targets[LayerIndex].PooledRenderTarget = Tiles[0].Targets[LayerIndex].PooledRenderTarget;
-			}
-		}
 	}
 
 	if (BatchSize > 0)
