Restored behavior that updates last component render time on screen if the component passes frustum culling and the occlusion state is known to be valid. This will avoid updating components that are pending an occlusion result. In practice it actually helps to avoid updating skeletal meshes that are occluded more effectively.

diff --git a/Engine/Source/Runtime/Renderer/Private/SceneRendering.h b/Engine/Source/Runtime/Renderer/Private/SceneRendering.h
--- a/Engine/Source/Runtime/Renderer/Private/SceneRendering.h
+++ b/Engine/Source/Runtime/Renderer/Private/SceneRendering.h
@@ -1144,6 +1144,9 @@ public:
 	/** A map from primitive ID to a boolean ray tracing visibility value. */
 	FSceneBitArray PrimitiveRayTracingVisibilityMap;
 
+	/** Bit set when a primitive is known to be un-occluded. */
+	FSceneBitArray PrimitiveDefinitelyUnoccludedMap;
+
 	/** A map from primitive ID to a boolean is fading value. */
 	FSceneBitArray PotentiallyFadingPrimitiveMap;
 
diff --git a/Engine/Source/Runtime/Renderer/Private/SceneVisibility.cpp b/Engine/Source/Runtime/Renderer/Private/SceneVisibility.cpp
--- a/Engine/Source/Runtime/Renderer/Private/SceneVisibility.cpp
+++ b/Engine/Source/Runtime/Renderer/Private/SceneVisibility.cpp
@@ -2011,8 +2011,12 @@ void FRelevancePacket::ComputeRelevance(FDynamicPrimitiveIndexList& DynamicPrimi
 
 		PrimitiveSceneInfo->LastRenderTime = CurrentWorldTime;
 
-		const bool bUpdateLastRenderTimeOnScreen = true;
-		PrimitiveSceneInfo->UpdateComponentLastRenderTime(CurrentWorldTime, bUpdateLastRenderTimeOnScreen);
+		// Update the last component render time only if we know for certain the primitive is un-occluded.
+		if (View.PrimitiveDefinitelyUnoccludedMap[BitIndex] || View.Family->EngineShowFlags.Wireframe)
+		{
+			const bool bUpdateLastRenderTimeOnScreen = true;
+			PrimitiveSceneInfo->UpdateComponentLastRenderTime(CurrentWorldTime, bUpdateLastRenderTimeOnScreen);
+		}
 
 		// Cache the nearest reflection proxy if needed
 		if (PrimitiveSceneInfo->NeedsReflectionCaptureUpdate())
@@ -2549,6 +2553,18 @@ bool FGPUOcclusionPacket::OcclusionCullPrimitive(VisitorType& Visitor, FOcclusio
 	const TArray<FBoxSphereBounds>* SubBounds = nullptr;
 	int32 SubIsOccludedStart = 0;
 
+	const auto SetAtomicBit = [] (FSceneBitArray& Array, int32 Index, bool Value)
+	{
+		if constexpr (bIsParallel)
+		{
+			Array[Index].AtomicSet(Value);
+		}
+		else
+		{
+			Array[Index] = Value;
+		}
+	};
+
 	if ((OcclusionFlags & EOcclusionFlags::HasSubprimitiveQueries) && OcclusionState.bAllowSubQueries)
 	{
 		FPrimitiveSceneProxy* Proxy = Scene.PrimitiveSceneProxies[Index];
@@ -2557,14 +2573,7 @@ bool FGPUOcclusionPacket::OcclusionCullPrimitive(VisitorType& Visitor, FOcclusio
 		bSubQueries = true;
 		if (!NumSubQueries)
 		{
-			if constexpr (bIsParallel)
-			{
-				View.PrimitiveVisibilityMap[Index].AtomicSet(false);
-			}
-			else
-			{
-				View.PrimitiveVisibilityMap[Index] = false;
-			}
+			SetAtomicBit(View.PrimitiveVisibilityMap, Index, false);
 			return false;
 		}
 
@@ -2835,17 +2844,14 @@ bool FGPUOcclusionPacket::OcclusionCullPrimitive(VisitorType& Visitor, FOcclusio
 		{
 			if (bIsOccluded)
 			{
-				if constexpr (bIsParallel)
-				{
-					View.PrimitiveVisibilityMap[Index].AtomicSet(false);
-				}
-				else
-				{
-					View.PrimitiveVisibilityMap[Index] = false;
-				}
+				SetAtomicBit(View.PrimitiveVisibilityMap, Index, false);
 				bIsVisible = false;
 				Result.NumCulledPrimitives++;
 			}
+			else if (bOcclusionStateIsDefinite)
+			{
+				SetAtomicBit(View.PrimitiveDefinitelyUnoccludedMap, Index, true);
+			}
 		}
 	}
 
@@ -2856,17 +2862,14 @@ bool FGPUOcclusionPacket::OcclusionCullPrimitive(VisitorType& Visitor, FOcclusio
 
 		if (bAllSubOccluded)
 		{
-			if constexpr (bIsParallel)
-			{
-				View.PrimitiveVisibilityMap[Index].AtomicSet(false);
-			}
-			else
-			{
-				View.PrimitiveVisibilityMap[Index] = false;
-			}
+			SetAtomicBit(View.PrimitiveVisibilityMap, Index, false);
 			bIsVisible = false;
 			Result.NumCulledPrimitives++;
 		}
+		else if (bAllSubOcclusionStateIsDefinite)
+		{
+			SetAtomicBit(View.PrimitiveDefinitelyUnoccludedMap, Index, true);
+		}
 	}
 
 	return bIsVisible;
@@ -3090,6 +3093,7 @@ bool FGPUOcclusionParallelPacket::AddPrimitive(int32 PrimitiveIndex)
 		return true;
 	}
 
+	View.PrimitiveDefinitelyUnoccludedMap[PrimitiveIndex].AtomicSet(true);
 	return false;
 }
 
@@ -3266,6 +3270,10 @@ void FGPUOcclusionSerial::AddPrimitives(FPrimitiveRange PrimitiveRange)
 		{
 			Packet.OcclusionCullPrimitive<false>(ProcessVisitor, OcclusionCullResult, BitIt.GetIndex());
 		}
+		else
+		{
+			View.PrimitiveDefinitelyUnoccludedMap.AccessCorrespondingBit(BitIt) = true;
+		}
 	}
 }
 
@@ -3617,6 +3625,7 @@ void FVisibilityViewPacket::BeginInitVisibility()
 	// Allocate the view's visibility maps.
 	View.PrimitiveVisibilityMap.Init(!bShouldVisibilityCull, Scene.Primitives.Num());
 	View.PrimitiveRayTracingVisibilityMap.Init(false, Scene.Primitives.Num());
+	View.PrimitiveDefinitelyUnoccludedMap.Init(false, Scene.Primitives.Num());
 	View.PotentiallyFadingPrimitiveMap.Init(false, Scene.Primitives.Num());
 	View.PrimitiveFadeUniformBuffers.AddZeroed(Scene.Primitives.Num());
 	View.PrimitiveFadeUniformBufferMap.Init(false, Scene.Primitives.Num());
