Fix for instance occlusion queries
Iterating on MDCs to retrieve the InstanceId does work until identical commands compaction happens (see FInstanceCullingContext::SetupDrawCommands, bCompactIdenticalCommands).
In that situation, some MDCs get merged but the final MDC doesn't keep track of the instances it will actually issue
The InstanceProcessingGPULoadBalancer actually holds the data when that merging happens, and is what is being used in BuildInstanceDrawCommands.usf
Even when we create an FInstanceCullingMergedContext, the original FInstanceCullingContext are still preserved, so getting the data just for the BasePass is still possible
Using the load balancer is driven by r.InstanceCulling.UseLoadBalancer to keep the existing behaviour in place in case things go wrong. Needs to be cleaned up later

#prelight 66dea08c202e3933e6332ad1

diff --git a/Engine/Shaders/Private/InstanceCulling/InstanceCullingOcclusionQuery.usf b/Engine/Shaders/Private/InstanceCulling/InstanceCullingOcclusionQuery.usf
--- a/Engine/Shaders/Private/InstanceCulling/InstanceCullingOcclusionQuery.usf
+++ b/Engine/Shaders/Private/InstanceCulling/InstanceCullingOcclusionQuery.usf
@@ -5,11 +5,17 @@
 #include "../VertexFactoryCommon.ush" // for TransformLocalToTranslatedWorld
 #include "../Nanite/NaniteHZBCull.ush"
 
+// We need to define NUM_THREADS_PER_GROUP before including InstanceCullingSetup as it includes InstanceCullingLoadBalancer.ush which requires declares groupshared uint ItemIndex[NUM_THREADS_PER_GROUP];
+#ifndef NUM_THREADS_PER_GROUP_DEFAULT
+#	define NUM_THREADS_PER_GROUP_DEFAULT 64
+#endif // NUM_THREADS_PER_GROUP_DEFAULT
+
 #ifndef NUM_THREADS_PER_GROUP
-#	define NUM_THREADS_PER_GROUP 64
+#	define NUM_THREADS_PER_GROUP NUM_THREADS_PER_GROUP_DEFAULT
 #endif // NUM_THREADS_PER_GROUP
 
-static const uint GBumpMeToForceRebuild = 2;
+#define LOAD_BALANCER_SINGLE_INSTANCE_MODE 0
+#include "InstanceCullingSetup.ush"
 
 #ifndef HZB_DESIRED_FOOTPRINT_PIXELS
 #	define HZB_DESIRED_FOOTPRINT_PIXELS 4
@@ -126,6 +132,23 @@ RWBuffer<uint> OutInstanceIdBuffer;
 // Remap linear visible index to a visible GPUScene instance index
 Buffer<uint> InstanceIdBuffer;
 
+bool LoadInstanceId(uint3 GroupId, uint GroupThreadIndex, uint DispatchThreadId, inout uint InstanceId)
+{
+#if USE_LOAD_BALANCER
+    // No need to properly fetch the final instance id for dynamic primitives since we're going to skip them anyway
+    uint DynamicInstanceIdOffset = 0;
+    uint DynamicInstanceIdMax = 0;
+	FInstanceCullingSetup InstanceCullingSetup = LoadInstanceCullingSetup(GroupId, GroupThreadIndex, DynamicInstanceIdOffset, DynamicInstanceIdMax, 0);
+	InstanceId = InstanceCullingSetup.InstanceId;
+    return InstanceCullingSetup.bValid && !InstanceCullingSetup.bIsDynamic;
+#else
+	const bool bValidInvocation = DispatchThreadId < NumInstances;
+	uint ClampedThreadId = min(DispatchThreadId, NumInstances - 1);
+	InstanceId = InstanceIdBuffer[ClampedThreadId];
+    return bValidInvocation;
+#endif
+}
+
 [numthreads(NUM_THREADS_PER_GROUP, 1, 1)]
 void MainCS(uint3 GroupId : SV_GroupID, uint GroupThreadIndex : SV_GroupIndex, uint DispatchThreadId : SV_DispatchThreadID)
 {
@@ -136,10 +159,8 @@ void MainCS(uint3 GroupId : SV_GroupID, uint GroupThreadIndex : SV_GroupIndex, u
 		GroupNumVisibleInstances = 0;
 	}
 
-	const bool bValidInvocation = DispatchThreadId < NumInstances;
-
-	uint ClampedThreadId = min(DispatchThreadId, NumInstances - 1);
-	uint InstanceId = InstanceIdBuffer[ClampedThreadId];
+    uint InstanceId = 0;
+	const bool bValidInvocation = LoadInstanceId(GroupId, GroupThreadIndex, DispatchThreadId, InstanceId);
 
 	GroupMemoryBarrierWithGroupSync();
 
@@ -272,17 +293,17 @@ void DebugMainVS(
 
 		if (VisibilityStatus == IVS_PossiblyVisible)
 		{
-			OutDebugVisibility = bWasVisibleLastFrame ? 2 : 3;
+			OutDebugVisibility = bWasVisibleLastFrame ? IVS_PossiblyVisible : IVS_Incompatible;
 		}
 		else
 		{
-			OutDebugVisibility = bWasVisibleLastFrame ? 1 : 0;
+			OutDebugVisibility = bWasVisibleLastFrame ? IVS_Visible : IVS_Hidden;
 		}
 	}
 	else
 	{
 		OutPosition = asfloat(0xFFFFFFFF);
-		OutDebugVisibility = 0;
+		OutDebugVisibility = IVS_Hidden;
 	}
 }
 
@@ -293,21 +314,21 @@ void DebugMainPS(
 {
 	switch (DebugVisibility)
 	{
-	case 0:
+	case IVS_Hidden:
 		// Definitely invisible due to HZB test.
 		// No occlusion query performed.
 		OutColor = float4(0.1, 0.0, 0.0, 0.0); // red
 		break;
-	case 1:
+	case IVS_Visible:
 		// Definitely visible due to near plane intersection or other heuristics.
 		// No occlusion query performed.
 		OutColor = float4(0.1, 0.1, 0.0, 0.0); // yellow
 		break;
-	case 2:
+	case IVS_PossiblyVisible:
 		// Visible due to occlusion query.
 		OutColor = float4(0.0, 0.025, 0.0, 0.0); // green
 		break;
-	case 3:
+	case IVS_Incompatible:
 		// Invisible due to occlusion query.
 		OutColor = float4(0.25, 0.0, 0.25, 0.0); // magenta
 		break;
diff --git a/Engine/Source/Runtime/Renderer/Private/InstanceCulling/InstanceCullingOcclusionQuery.cpp b/Engine/Source/Runtime/Renderer/Private/InstanceCulling/InstanceCullingOcclusionQuery.cpp
--- a/Engine/Source/Runtime/Renderer/Private/InstanceCulling/InstanceCullingOcclusionQuery.cpp
+++ b/Engine/Source/Runtime/Renderer/Private/InstanceCulling/InstanceCullingOcclusionQuery.cpp
@@ -6,6 +6,7 @@
 #include "Containers/ResourceArray.h"
 #include "GPUScene.h"
 #include "GlobalShader.h"
+#include "InstanceCulling/InstanceCullingManager.h"
 #include "RHIAccess.h"
 #include "RHIFeatureLevel.h"
 #include "RHIGlobals.h"
@@ -29,6 +30,13 @@ static TAutoConsoleVariable<int32> CVarInstanceCullingOcclusionQueries(
 	TEXT("EXPERIMENTAL: Use per-instance software occlusion queries to perform less conservative visibility test than what's possible with HZB alone"),
 	ECVF_RenderThreadSafe | ECVF_Preview);
 
+static int32 GInstanceCullingUseLoadBalancer = 1;
+static FAutoConsoleVariableRef CVarInstanceCullingUseLoadBalancer(
+	TEXT("r.InstanceCulling.UseLoadBalancer"),
+	GInstanceCullingUseLoadBalancer,
+	TEXT("Prefer to use UseLoadBalancer"),
+	ECVF_RenderThreadSafe);
+
 struct FInstanceCullingOcclusionQueryDeferredContext;
 
 namespace
@@ -70,7 +78,8 @@ class FInstanceCullingOcclusionQueryCS : public FGlobalShader
 public:
 
 	class FMultiView : SHADER_PERMUTATION_BOOL("DIM_MULTI_VIEW");
-	using FPermutationDomain = TShaderPermutationDomain<FMultiView>;
+	class FUseLoadBalancerDim : SHADER_PERMUTATION_BOOL("USE_LOAD_BALANCER");
+	using FPermutationDomain = TShaderPermutationDomain<FMultiView, FUseLoadBalancerDim>;
 
 	static bool ShouldCompilePermutation(const FGlobalShaderPermutationParameters& Parameters)
 	{
@@ -82,9 +91,17 @@ public:
 	static void ModifyCompilationEnvironment(const FGlobalShaderPermutationParameters& Parameters, FShaderCompilerEnvironment& OutEnvironment)
 	{
 		FGlobalShader::ModifyCompilationEnvironment(Parameters, OutEnvironment);
+		const FPermutationDomain PermutationVector(Parameters.PermutationId);
+		
+		// Currently, instance compaction is not supported on mobile platforms
+		if (PermutationVector.Get<FUseLoadBalancerDim>())
+		{
+			FInstanceProcessingGPULoadBalancer::SetShaderDefines(OutEnvironment);
+		}
+
 		OutEnvironment.SetDefine(TEXT("USE_GLOBAL_GPU_SCENE_DATA"), 1);
 		OutEnvironment.SetDefine(TEXT("VF_SUPPORTS_PRIMITIVE_SCENE_DATA"), 1);
-		OutEnvironment.SetDefine(TEXT("NUM_THREADS_PER_GROUP"), NumThreadsPerGroup);
+		OutEnvironment.SetDefine(TEXT("NUM_THREADS_PER_GROUP_DEFAULT"), NumThreadsPerGroup);
 	}
 
 	BEGIN_SHADER_PARAMETER_STRUCT(FParameters, )
@@ -97,6 +114,7 @@ public:
 		SHADER_PARAMETER_RDG_BUFFER_SRV(StructuredBuffer<float4>, GPUScenePrimitiveSceneData)
 		SHADER_PARAMETER_RDG_BUFFER_SRV(Buffer<uint32>, InstanceIdBuffer)
 		SHADER_PARAMETER_STRUCT_REF(FViewUniformShaderParameters, View)
+		SHADER_PARAMETER_STRUCT_INCLUDE(FInstanceProcessingGPULoadBalancer::FShaderParameters, LoadBalancerParameters)
 		SHADER_PARAMETER_RDG_TEXTURE(Texture2D, HZBTexture)
 		SHADER_PARAMETER_SAMPLER(SamplerState, HZBSampler)
 		SHADER_PARAMETER(FVector2f, HZBSize)
@@ -285,10 +303,11 @@ static void RenderInstanceOcclusionCulling(
 */
 struct FInstanceCullingOcclusionQueryDeferredContext
 {
-	FInstanceCullingOcclusionQueryDeferredContext(const FViewInfo* InView, int32 InNumGPUSceneInstances, EMeshPass::Type InMeshPass)
+	FInstanceCullingOcclusionQueryDeferredContext(const FViewInfo* InView, int32 InNumGPUSceneInstances, EMeshPass::Type InMeshPass, FInstanceCullingContext* InInstanceCullingContext)
 		: View(InView)
 		, NumGPUSceneInstances(InNumGPUSceneInstances)
 		, MeshPass(InMeshPass)
+		, InstanceCullingContext(InInstanceCullingContext)
 	{
 	}
 
@@ -342,6 +361,28 @@ struct FInstanceCullingOcclusionQueryDeferredContext
 		// If it does happen, then we may have to move the render pass to later point in the frame.
 		MeshDrawCommandPass.WaitForSetupTask(); 
 
+		if (InstanceCullingContext != nullptr)
+		{
+			InstanceProcessingGPULoadBalancer = InstanceCullingContext->LoadBalancers[int(EBatchProcessingMode::Generic)];
+			bValid = (InstanceProcessingGPULoadBalancer != nullptr);
+			static FInstanceProcessingGPULoadBalancer Dummy;
+			// Always provide a load balancer so that CreateLoadBalancerGPUDataDeferred doesn't crash. bValid=false will skip the dispatch
+			if (!bValid)
+			{
+				InstanceProcessingGPULoadBalancer = &Dummy;
+			}
+
+			// in case something goes wrong: we will skip the compute since bValid won't be true and we will fill up the data from VisibleInstanceIds
+			AlignedNumInstances = FInstanceCullingOcclusionQueryCS::NumThreadsPerGroup;
+			VisibleInstanceIds.SetNumZeroed(AlignedNumInstances);
+			InstanceProcessingGPULoadBalancer->FinalizeBatches();
+			FIntVector LoadBalancerNumThreadGroups = InstanceProcessingGPULoadBalancer->GetWrappedCsGroupCount();
+			// Needed to allocate the buffer holding the instance ids after the culling pass, see DeferredAlignedNumInstancesOutputCulling
+			AlignedNumInstances = LoadBalancerNumThreadGroups.X * LoadBalancerNumThreadGroups.Y * LoadBalancerNumThreadGroups.Z * FInstanceCullingOcclusionQueryCS::NumThreadsPerGroup;
+			
+			return;
+		}
+
 		const FMeshCommandOneFrameArray& VisibleMeshDrawCommands = MeshDrawCommandPass.GetMeshDrawCommands();
 
 		const FScene *Scene = View->Family->Scene->GetRenderScene();
@@ -440,13 +481,22 @@ struct FInstanceCullingOcclusionQueryDeferredContext
 		check(ResultCursor == ResultData + AlignedNumInstances);
 	}
 
-	FRDGBufferNumElementsCallback DeferredAlignedNumInstances()
+	FRDGBufferNumElementsCallback DeferredAlignedNumInstancesOutputCulling()
 	{
 		return [Context = this]() -> uint32
-			{
-				Context->Execute();
-				return Context->AlignedNumInstances;
-			};
+		{
+			Context->Execute();
+			return Context->AlignedNumInstances;
+		};
+	}
+
+	FRDGBufferNumElementsCallback DeferredNumInstanceIdData()
+	{
+		return [Context = this]() -> uint32
+		{
+			Context->Execute();
+			return Context->VisibleInstanceIds.Num();
+		};
 	}
 
 	FRDGBufferInitialDataCallback DeferredInstanceIdData()
@@ -476,6 +526,8 @@ struct FInstanceCullingOcclusionQueryDeferredContext
 	const FViewInfo* View = nullptr;
 	int32 NumGPUSceneInstances = 0;
 	EMeshPass::Type MeshPass = EMeshPass::Num;
+	FInstanceCullingContext* InstanceCullingContext = nullptr;
+	FInstanceProcessingGPULoadBalancer* InstanceProcessingGPULoadBalancer = nullptr;
 	int32 NumInstances = 0;
 	int32 AlignedNumInstances = 0;
 	FIntVector NumThreadGroups = FIntVector::ZeroValue;
@@ -483,6 +535,25 @@ struct FInstanceCullingOcclusionQueryDeferredContext
 	TArray<uint32> VisibleInstanceIds;
 };
 
+static void CreateLoadBalancerGPUDataDeferred(FRDGBuilder& GraphBuilder, FInstanceCullingOcclusionQueryCS::FParameters* PassParameters, FInstanceCullingOcclusionQueryDeferredContext* DeferredContext)
+{
+	PassParameters->LoadBalancerParameters.BatchBuffer = GraphBuilder.CreateSRV(
+		CreateStructuredBuffer(GraphBuilder, TEXT("InstanceCullingLoadBalancer.Batches"), [DeferredContext]() 
+			-> const TArray<FInstanceCullingLoadBalancerBase::FPackedBatch>& 
+			{ 
+				DeferredContext->Execute(); 
+				return DeferredContext->InstanceProcessingGPULoadBalancer->GetBatches(); 
+			}));
+
+	PassParameters->LoadBalancerParameters.ItemBuffer = GraphBuilder.CreateSRV(
+		CreateStructuredBuffer(GraphBuilder, TEXT("InstanceCullingLoadBalancer.Items"), [DeferredContext]() 
+			-> const TArray<FInstanceCullingLoadBalancerBase::FPackedItem>& 
+			{ 
+				DeferredContext->Execute(); 
+				return DeferredContext->InstanceProcessingGPULoadBalancer->GetItems(); 
+			}));
+}
+
 uint32 FInstanceCullingOcclusionQueryRenderer::Render(
 	FRDGBuilder& GraphBuilder,
 	FGPUScene& GPUScene,
@@ -508,7 +579,17 @@ uint32 FInstanceCullingOcclusionQueryRenderer::Render(
 
 	const int32 NumGPUSceneInstances = GPUScene.GetNumInstances();
 
-	FInstanceCullingOcclusionQueryDeferredContext* DeferredContext = GraphBuilder.AllocObject<FInstanceCullingOcclusionQueryDeferredContext>(&View, NumGPUSceneInstances, EMeshPass::BasePass);
+	FInstanceCullingContext* InstanceCullingContext = nullptr;
+	if (GInstanceCullingUseLoadBalancer > 0)
+	{
+		FParallelMeshDrawCommandPass& MeshDrawCommandPass = View.ParallelMeshDrawCommandPasses[EMeshPass::BasePass];
+
+		// At this point in time, we don't have the guarantee that MeshDrawCommandPass is done. Only access stable members, not batches/items/mdcs
+		InstanceCullingContext = MeshDrawCommandPass.GetInstanceCullingContext();
+	}
+
+	FInstanceCullingOcclusionQueryDeferredContext* DeferredContext = GraphBuilder.AllocObject<FInstanceCullingOcclusionQueryDeferredContext>(&View, NumGPUSceneInstances, 
+		EMeshPass::BasePass, InstanceCullingContext);
 
 	FRDGTextureRef DepthTexture = View.GetSceneTextures().Depth.Target;
 	FRDGTextureRef HZBTexture = View.HZB;
@@ -556,14 +637,23 @@ uint32 FInstanceCullingOcclusionQueryRenderer::Render(
 	FRDGBufferUAVRef IndirectArgsUAV = GraphBuilder.CreateUAV(IndirectArgsBuffer, PF_R32_UINT);
 
 	// Buffer of GPUScene instance indices to run occlusion queries for (input for setup CS)
-	FRDGBufferRef SetupInstanceIdBuffer = GraphBuilder.CreateBuffer(
-		FRDGBufferDesc::CreateBufferDesc(sizeof(uint32), 1 /*real size is provided via callback later*/),
-		TEXT("FInstanceCullingOcclusionQueryRenderer_SetupInstanceIdBuffer"), 
-		DeferredContext->DeferredAlignedNumInstances());
+	FRDGBufferRef SetupInstanceIdBuffer;
 
-	GraphBuilder.QueueBufferUpload(SetupInstanceIdBuffer,
-		DeferredContext->DeferredInstanceIdData(),
-		DeferredContext->DeferredInstanceIdDataSize());
+	// When using the GPU load balancer, the upload of the data holding instance ids happens below with InstanceProcessingGPULoadBalancer->Upload
+	if (InstanceCullingContext == nullptr)
+	{
+		SetupInstanceIdBuffer = GraphBuilder.CreateBuffer(
+			FRDGBufferDesc::CreateBufferDesc(sizeof(uint32), 1 /*real size is provided via callback later*/),
+			TEXT("FInstanceCullingOcclusionQueryRenderer_SetupInstanceIdBuffer"), 
+			DeferredContext->DeferredNumInstanceIdData());
+		GraphBuilder.QueueBufferUpload(SetupInstanceIdBuffer,
+									   DeferredContext->DeferredInstanceIdData(),
+									   DeferredContext->DeferredInstanceIdDataSize());
+	}
+	else
+	{
+		SetupInstanceIdBuffer = GSystemTextures.GetDefaultBuffer(GraphBuilder, 4);
+	}
 
 	FRDGBufferSRVRef SetupInstanceIdBufferSRV = GraphBuilder.CreateSRV(SetupInstanceIdBuffer, PF_R32_UINT);
 
@@ -571,7 +661,7 @@ uint32 FInstanceCullingOcclusionQueryRenderer::Render(
 	FRDGBufferRef InstanceIdBuffer = GraphBuilder.CreateBuffer(
 		FRDGBufferDesc::CreateBufferDesc(sizeof(uint32), 1 /*real size is provided via callback later*/),
 		TEXT("FInstanceCullingOcclusionQueryRenderer_InstanceIdBuffer"),
-		DeferredContext->DeferredAlignedNumInstances());
+		DeferredContext->DeferredAlignedNumInstancesOutputCulling());
 
 	FRDGBufferUAVRef InstanceIdUAV = GraphBuilder.CreateUAV(InstanceIdBuffer, PF_R32_UINT);
 	FRDGBufferSRVRef InstanceIdSRV = GraphBuilder.CreateSRV(InstanceIdBuffer, PF_R32_UINT);
@@ -582,6 +672,14 @@ uint32 FInstanceCullingOcclusionQueryRenderer::Render(
 	{
 		FInstanceCullingOcclusionQueryCS::FParameters* PassParameters = GraphBuilder.AllocParameters<FInstanceCullingOcclusionQueryCS::FParameters>();
 
+		// FInstanceGPULoadBalancer uses the SceneRenderingAllocator which should keep data alive until the graphbuilder execution
+		bool bUseGPULoadBalancer = false;
+		if (InstanceCullingContext != nullptr)
+		{
+			CreateLoadBalancerGPUDataDeferred(GraphBuilder, PassParameters, DeferredContext);
+			bUseGPULoadBalancer = true;
+		}
+
 		PassParameters->OutIndirectArgsBuffer = IndirectArgsUAV;
 		PassParameters->OutInstanceIdBuffer = InstanceIdUAV;
 		PassParameters->RWVisibilityMask = VisibilityMaskUAV;
@@ -601,6 +699,7 @@ uint32 FInstanceCullingOcclusionQueryRenderer::Render(
 
 		FInstanceCullingOcclusionQueryCS::FPermutationDomain CSPermutationVector;
 		CSPermutationVector.Set<FInstanceCullingOcclusionQueryCS::FMultiView>(bMultiView);
+		CSPermutationVector.Set<FInstanceCullingOcclusionQueryCS::FUseLoadBalancerDim>(bUseGPULoadBalancer);
 		TShaderMapRef<FInstanceCullingOcclusionQueryCS> ComputeShader(View.ShaderMap, CSPermutationVector);
 
 		ClearUnusedGraphResources(ComputeShader, PassParameters);
@@ -618,11 +717,20 @@ uint32 FInstanceCullingOcclusionQueryRenderer::Render(
 
 			PassParameters->NumInstances = DeferredContext->NumInstances;
 
+			FIntVector CullingNumThreadGroups(DeferredContext->NumThreadGroups);
+			FInstanceProcessingGPULoadBalancer* DeferredContextInstanceProcessingGPULoadBalancer = DeferredContext->InstanceProcessingGPULoadBalancer;
+			if (DeferredContextInstanceProcessingGPULoadBalancer)
+			{
+				PassParameters->LoadBalancerParameters.NumBatches = DeferredContextInstanceProcessingGPULoadBalancer->GetBatches().Num();
+				PassParameters->LoadBalancerParameters.NumItems = DeferredContextInstanceProcessingGPULoadBalancer->GetItems().Num();
+				CullingNumThreadGroups = DeferredContextInstanceProcessingGPULoadBalancer->GetWrappedCsGroupCount();
+			}
+
 			FComputeShaderUtils::Dispatch(
 				RHICmdList,
 				ComputeShader,
 				*PassParameters,
-				DeferredContext->NumThreadGroups);
+				CullingNumThreadGroups);
 		});
 	}
 
