Added clamp on the minimum possible mip level for local light VSMs to optimize various passes that process the page tables.
  - Added MinMipLevel added to the FVirtualShadowMapProjectionShaderData
  - Added a helper (FVirtualShadowMapPerPageShader) to bin the allocated VSMs by MinMipLevel and use that to dispatch kernels to only process what is needed
  - Used this clear only page table entries that may be used.
  - same for allocating new pages and propagating mapped mips
  - Only clear used page-rect bounds (clears all mip levels, which prevents rendering to any out of bounds level)
  - Added clamping to shadow VSM sampling functions.
  - Moved VSM defines to shared shader header to make it possible to explicitly share the logic to compute Mip level (avoiding CPU/GPU divergence).

diff --git a/Engine/Shaders/Private/VirtualShadowMaps/VirtualShadowMapCacheInvalidation.ush b/Engine/Shaders/Private/VirtualShadowMaps/VirtualShadowMapCacheInvalidation.ush
--- a/Engine/Shaders/Private/VirtualShadowMaps/VirtualShadowMapCacheInvalidation.ush
+++ b/Engine/Shaders/Private/VirtualShadowMaps/VirtualShadowMapCacheInvalidation.ush
@@ -86,7 +86,7 @@ void InvalidateInstancePages(int VirtualShadowMapId, uint InstanceId, bool bForc
 		// case #1 mip-map VSM - loop all mip levels, case #2 clipmap, just one 'mip level'
 		int NumMipLevels = (ProjectionData.ClipmapLevelCountRemaining <= 0) ? VSM_MAX_MIP_LEVELS : 1;
 		{
-			for (int MipLevel = 0; MipLevel < NumMipLevels; ++MipLevel)
+			for (int MipLevel = ProjectionData.MinMipLevel; MipLevel < NumMipLevels; ++MipLevel)
 			{
 				int ViewDim = int(uint(VSM_VIRTUAL_MAX_RESOLUTION_XY) >> MipLevel);
 				FScreenRect Rect = GetScreenRect(int4(0, 0, ViewDim, ViewDim), Cull, 4);
diff --git a/Engine/Shaders/Private/VirtualShadowMaps/VirtualShadowMapPageAccessCommon.ush b/Engine/Shaders/Private/VirtualShadowMaps/VirtualShadowMapPageAccessCommon.ush
--- a/Engine/Shaders/Private/VirtualShadowMaps/VirtualShadowMapPageAccessCommon.ush
+++ b/Engine/Shaders/Private/VirtualShadowMaps/VirtualShadowMapPageAccessCommon.ush
@@ -298,14 +298,14 @@ struct FShadowPageTranslationResult
 };
 
 // Finds the best-resolution mapped page at the given UV
-FShadowPageTranslationResult ShadowVirtualToPhysicalUV(uint VirtualShadowMapID, float2 ShadowMapUV)
+FShadowPageTranslationResult ShadowVirtualToPhysicalUV(uint VirtualShadowMapID, float2 ShadowMapUV, uint MinMipLevel)
 {
 	uint2 vPage = uint2(ShadowMapUV * VSM_LEVEL0_DIM_PAGES_XY);
-	FShadowPhysicalPage PhysicalPageEntry = ShadowGetPhysicalPage(CalcPageOffset(VirtualShadowMapID, 0, vPage));
+	FShadowPhysicalPage PhysicalPageEntry = ShadowGetPhysicalPage(CalcPageOffset(VirtualShadowMapID, MinMipLevel, vPage >> MinMipLevel));
 
 	FShadowPageTranslationResult Result;
 	Result.bValid = PhysicalPageEntry.bAnyLODValid;
-	Result.LODOffset = IsSinglePageVirtualShadowMap(VirtualShadowMapID) ? (VSM_MAX_MIP_LEVELS - 1U) : PhysicalPageEntry.LODOffset;
+	Result.LODOffset = IsSinglePageVirtualShadowMap(VirtualShadowMapID) ? (VSM_MAX_MIP_LEVELS - 1U) : (PhysicalPageEntry.LODOffset + MinMipLevel);
 	// TODO: Can optimize this slightly based on relative offset
 	Result.VirtualTexelAddressFloat = ShadowMapUV * float(CalcLevelDimsTexels(Result.LODOffset));
 	Result.VirtualTexelAddress = uint2(Result.VirtualTexelAddressFloat);
diff --git a/Engine/Shaders/Private/VirtualShadowMaps/VirtualShadowMapPageManagement.usf b/Engine/Shaders/Private/VirtualShadowMaps/VirtualShadowMapPageManagement.usf
--- a/Engine/Shaders/Private/VirtualShadowMaps/VirtualShadowMapPageManagement.usf
+++ b/Engine/Shaders/Private/VirtualShadowMaps/VirtualShadowMapPageManagement.usf
@@ -10,6 +10,7 @@
 #include "VirtualShadowMapProjectionCommon.ush"
 #include "VirtualShadowMapPageAccessCommon.ush"
 #include "VirtualShadowMapStats.ush"
+#include "VirtualShadowMapPerPageDispatch.ush"
 
 StructuredBuffer<FPhysicalPageMetaData> PhysicalPageMetaData;
 RWStructuredBuffer<uint> OutPageFlags;
@@ -108,108 +109,123 @@ void GenerateHierarchicalPageFlags(uint ThreadId : SV_DispatchThreadID)
 	}
 }
 
-
+#ifdef PropagateMappedMips
 /**
 * One thread per page in level 0, launched as 1d groups, with 2D grid with Y dim ==  NumFullShadowMaps.
 * This is effectively just a big broadcast operation. There are more efficient ways to do this with
 * fewer threads and wave ops, but given the page counts just relying on memory coalescing is
 * good enough for now.
 */
-[numthreads(VSM_DEFAULT_CS_GROUP_X, 1, 1)]
-void PropagateMappedMips(uint2 ThreadId : SV_DispatchThreadID, uint2 GroupId : SV_GroupID)
+[numthreads(PER_PAGE_THREAD_GROUP_SIZE_X, 1, 1)]
+void PropagateMappedMips(uint2 DispatchThreadId : SV_DispatchThreadID)
 {
-	const uint NumLevel0Entries = VSM_LEVEL0_DIM_PAGES_XY * VSM_LEVEL0_DIM_PAGES_XY;
-	if (ThreadId.x >= NumLevel0Entries)
+	FPerPageDispatchSetup Setup; 
+	Setup.Init(DispatchThreadId);
+
+	if (!Setup.bValid || Setup.BaseLevelOffset.bIsSinglePageSM)
 	{
 		return;
 	}
 
-	int VirtualShadowMapID = int(GroupId.y + VSM_MAX_SINGLE_PAGE_SHADOW_MAPS);
-	uint PageTableEntryIndex = ThreadId.x;
-
-	uint2 Level0Page;
-	Level0Page.x = PageTableEntryIndex & ((1U << VSM_LOG2_LEVEL0_DIM_PAGES_XY) - 1U);
-	Level0Page.y = PageTableEntryIndex >> VSM_LOG2_LEVEL0_DIM_PAGES_XY;
-
-	FVirtualShadowMapProjectionShaderData ProjectionData = GetVirtualShadowMapProjectionData(VirtualShadowMapID);
-
+	FVirtualShadowMapProjectionShaderData ProjectionData = GetVirtualShadowMapProjectionData(Setup.VirtualShadowMapId);
+	
 	if (ProjectionData.LightType == LIGHT_TYPE_DIRECTIONAL)
 	{
-		// Directional lights propagate pages to their coarser/larger clipmap levels (and only use mip0 pages)
-		// Each clipmap level is a separate VSM, so we gather any mapped coarser pages as necessary and write only our own page output
-		// There's also technically a race similar to below with other threads writing the PT data we are reading,
-		// but it's still deterministic as long as we only look at pages with "bThisLODValid".
-		// There's some redundant work of course, but this shader is pretty cheap overall
-
-		uint Page0Offset = CalcPageOffset(VirtualShadowMapID, 0, Level0Page);
-		FShadowPhysicalPage pPage0 = ShadowDecodePageTable(OutPageTable[Page0Offset]);
-
-		BRANCH
-		if (!pPage0.bThisLODValid)
+		const uint NumLevel0Entries = VSM_LEVEL0_DIM_PAGES_XY * VSM_LEVEL0_DIM_PAGES_XY;
+		for (uint PageTableEntryIndex = Setup.LoopStart; PageTableEntryIndex < NumLevel0Entries; PageTableEntryIndex += Setup.LoopStride)
 		{
-			const int OffsetScale = (VSM_LEVEL0_DIM_PAGES_XY >> 2);
-			int2 BaseOffset = OffsetScale * ProjectionData.ClipmapCornerRelativeOffset;
-			int2 BasePage   = int2(Level0Page) - BaseOffset;
+			uint2 Level0Page;
+			Level0Page.x = PageTableEntryIndex & ((1U << VSM_LOG2_LEVEL0_DIM_PAGES_XY) - 1U);
+			Level0Page.y = PageTableEntryIndex >> VSM_LOG2_LEVEL0_DIM_PAGES_XY;
+			// Directional lights propagate pages to their coarser/larger clipmap levels (and only use mip0 pages)
+			// Each clipmap level is a separate VSM, so we gather any mapped coarser pages as necessary and write only our own page output
+			// There's also technically a race similar to below with other threads writing the PT data we are reading,
+			// but it's still deterministic as long as we only look at pages with "bThisLODValid".
+			// There's some redundant work of course, but this shader is pretty cheap overall
+
+			uint Page0Offset = CalcPageOffset(Setup.VirtualShadowMapId, 0, Level0Page);
+			FShadowPhysicalPage pPage0 = ShadowDecodePageTable(OutPageTable[Page0Offset]);
 
-			// Search for first mapped page past this one
-			uint RemainingLevels = ProjectionData.ClipmapLevelCountRemaining;
-			for (uint ClipmapOffset = 1; ClipmapOffset < RemainingLevels; ++ClipmapOffset)
+			BRANCH
+			if (!pPage0.bThisLODValid)
 			{
-				const int ClipmapLevelId = VirtualShadowMapID + int(ClipmapOffset);
-
-				FVirtualShadowMapProjectionShaderData LevelProjectionData = GetVirtualShadowMapProjectionData(ClipmapLevelId);								
-				int2 LevelOffset = OffsetScale * LevelProjectionData.ClipmapCornerRelativeOffset;
+				const int OffsetScale = (VSM_LEVEL0_DIM_PAGES_XY >> 2);
+				int2 BaseOffset = OffsetScale * ProjectionData.ClipmapCornerRelativeOffset;
+				int2 BasePage   = int2(Level0Page) - BaseOffset;
 
-				int2 LevelPage = (BasePage + (LevelOffset << ClipmapOffset)) >> ClipmapOffset;
-				if (IsVirtualShadowMapPageAddressValid(LevelPage, 0))
+				// Search for first mapped page past this one
+				uint RemainingLevels = ProjectionData.ClipmapLevelCountRemaining;
+				for (uint ClipmapOffset = 1; ClipmapOffset < RemainingLevels; ++ClipmapOffset)
 				{
-					uint LevelPageOffset = CalcPageOffset(ClipmapLevelId, 0, uint2(LevelPage));
-					FShadowPhysicalPage pPage = ShadowDecodePageTable(OutPageTable[LevelPageOffset]);
-					if (pPage.bThisLODValid)
+					const int ClipmapLevelId = Setup.VirtualShadowMapId + int(ClipmapOffset);
+
+					FVirtualShadowMapProjectionShaderData LevelProjectionData = GetVirtualShadowMapProjectionData(ClipmapLevelId);								
+					int2 LevelOffset = OffsetScale * LevelProjectionData.ClipmapCornerRelativeOffset;
+
+					int2 LevelPage = (BasePage + (LevelOffset << ClipmapOffset)) >> ClipmapOffset;
+					if (IsVirtualShadowMapPageAddressValid(LevelPage, 0))
 					{
-						OutPageTable[Page0Offset] = ShadowEncodePageTable(pPage.PhysicalAddress, ClipmapOffset);
-						break;
+						uint LevelPageOffset = CalcPageOffset(ClipmapLevelId, 0, uint2(LevelPage));
+						FShadowPhysicalPage pPage = ShadowDecodePageTable(OutPageTable[LevelPageOffset]);
+						if (pPage.bThisLODValid)
+						{
+							OutPageTable[Page0Offset] = ShadowEncodePageTable(pPage.PhysicalAddress, ClipmapOffset);
+							break;
+						}
+					}
+					else
+					{
+						// TODO: We're off the edge... can this ever even happen in practice given the construction?
 					}
-				}
-				else
-				{
-					// TODO: We're off the edge... can this ever even happen in practice given the construction?
 				}
 			}
 		}
 	}
 	else
 	{
-		// Local lights propagate pages to their coarser mips
-		int MappedPageLevel = -1;
-		uint2 MappedPhysicalAddress = 0;
+		int MinLevel = ProjectionData.MinMipLevel;
+		const uint MinLevelShift = VSM_LOG2_LEVEL0_DIM_PAGES_XY - uint(MinLevel);
+		const uint NumMinLevelEntries = 1u << (2 * MinLevelShift);
 
-		for (int Level = (VSM_MAX_MIP_LEVELS - 1); Level >= 0; --Level)
+		for (uint PageTableEntryIndex = Setup.LoopStart; PageTableEntryIndex < NumMinLevelEntries; PageTableEntryIndex += Setup.LoopStride)
 		{
-			uint2 vPage = Level0Page >> Level;
-			uint PageOffset = CalcPageOffset(VirtualShadowMapID, Level, vPage);
-			FShadowPhysicalPage pPage = ShadowDecodePageTable(OutPageTable[PageOffset]);
-
-			BRANCH
-			if (pPage.bThisLODValid)
-			{
-				// This page is mapped, so leave it alone and propagate downwards
-				MappedPageLevel = Level;
-				MappedPhysicalAddress = pPage.PhysicalAddress;
-			}
-			else if( MappedPageLevel >= 0 )
+			// Local lights propagate pages to their coarser mips
+			int MappedPageLevel = -1;
+			uint2 MappedPhysicalAddress = 0;
+			uint2 MinLevelPage;
+			MinLevelPage.x = PageTableEntryIndex & ((1U << MinLevelShift) - 1U);
+			MinLevelPage.y = PageTableEntryIndex >> MinLevelShift;
+
+			for (int Level = (VSM_MAX_MIP_LEVELS - 1); Level >= MinLevel; --Level)
 			{
-				// This page is not mapped; replace it with our suitably offset parent mapped page
-				// Ensure only one thread writes each value to avoid races, but we read on all threads as the broadcast
-				// Note that this can race with the other threads reading this value, but since bThisLODValid will
-				// always be false on these updated pages the values will be ignored. As long as the writes to the page
-				// table are atomic (currently a single DWORD), this is safe.
-				if (all((vPage << Level) == Level0Page))
+				int LevelDelta = Level - MinLevel;
+				uint2 vPage = MinLevelPage >> LevelDelta;
+				uint PageOffset = CalcPageOffset(Setup.VirtualShadowMapId, Level, vPage);
+				FShadowPhysicalPage pPage = ShadowDecodePageTable(OutPageTable[PageOffset]);
+
+				BRANCH
+				if (pPage.bThisLODValid)
+				{
+					// This page is mapped, so leave it alone and propagate downwards
+					MappedPageLevel = Level;
+					MappedPhysicalAddress = pPage.PhysicalAddress;
+				}
+				else if( MappedPageLevel >= 0 )
 				{
-					uint MipOffset = MappedPageLevel - Level;
-					OutPageTable[PageOffset] = ShadowEncodePageTable(MappedPhysicalAddress, MipOffset);
+					// This page is not mapped; replace it with our suitably offset parent mapped page
+					// Ensure only one thread writes each value to avoid races, but we read on all threads as the broadcast
+					// Note that this can race with the other threads reading this value, but since bThisLODValid will
+					// always be false on these updated pages the values will be ignored. As long as the writes to the page
+					// table are atomic (currently a single DWORD), this is safe.
+					if (all((vPage << LevelDelta) == MinLevelPage))
+					{
+						uint MipOffset = MappedPageLevel - Level;
+						OutPageTable[PageOffset] = ShadowEncodePageTable(MappedPhysicalAddress, MipOffset);
+					}
 				}
 			}
 		}
 	}
 }
+
+#endif // PropagateMappedMips
diff --git a/Engine/Shaders/Private/VirtualShadowMaps/VirtualShadowMapPerPageDispatch.ush b/Engine/Shaders/Private/VirtualShadowMaps/VirtualShadowMapPerPageDispatch.ush
new file mode 100644
--- /dev/null
+++ b/Engine/Shaders/Private/VirtualShadowMaps/VirtualShadowMapPerPageDispatch.ush
@@ -0,0 +1,72 @@
+// Copyright Epic Games, Inc. All Rights Reserved.
+#pragma once
+
+#include "../Common.ush"
+
+#ifdef PER_PAGE_DISPATCH_SETUP
+
+StructuredBuffer<uint> VirtualShadowMapIds;
+uint NumVirtualShadowMapIds;
+uint VirtualShadowMapIdsOffset;
+uint PerPageDispatchDimX;
+uint bUseThreadPerId;
+
+/**
+ * Helper to set up the per-thread indexing to process all relevant pages in the page table.
+ */
+struct FPerPageDispatchSetup
+{
+	bool bValid;
+	uint VirtualShadowMapId;
+	
+	uint LoopStride;
+	uint LoopStart;
+	
+	FVirtualSMLevelOffset BaseLevelOffset;
+	uint PageStartOffset;
+	uint PageEndOffset;
+
+	void Init(uint2 DispatchThreadId)
+	{
+		VirtualShadowMapId = INDEX_NONE;
+		uint LinearThreadId = DispatchThreadId.y * PerPageDispatchDimX + DispatchThreadId.x;
+		if (bUseThreadPerId)
+		{
+			LoopStart = 0;
+			LoopStride = 1;
+			if (LinearThreadId < NumVirtualShadowMapIds)
+			{
+				bValid = true;
+				VirtualShadowMapId = VirtualShadowMapIds[VirtualShadowMapIdsOffset + LinearThreadId];
+			}
+		}
+		else
+		{
+			bValid = true;
+			VirtualShadowMapId = VirtualShadowMapIds[VirtualShadowMapIdsOffset + DispatchThreadId.y];
+			LoopStart = DispatchThreadId.x;
+			LoopStride = PerPageDispatchDimX;
+		}
+
+		BaseLevelOffset = CalcPageTableLevelOffset(VirtualShadowMapId, 0);
+		PageStartOffset = 0;
+		PageEndOffset = 1;
+;
+		if (!BaseLevelOffset.bIsSinglePageSM)
+		{
+			const FVirtualShadowMapProjectionShaderData ProjectionData = GetVirtualShadowMapProjectionData(VirtualShadowMapId);
+			if (ProjectionData.LightType == LIGHT_TYPE_DIRECTIONAL)
+			{
+				PageStartOffset = CalcLevelOffsets(0);
+				PageEndOffset = VSM_PAGE_TABLE_SIZE;// TODO: HMip build still relies on all the entries being cleared. CalcLevelOffsets(1);
+			}
+			else
+			{
+				PageStartOffset = CalcLevelOffsets(ProjectionData.MinMipLevel);
+				PageEndOffset = VSM_PAGE_TABLE_SIZE;
+			}
+		}
+	}
+};
+
+#endif
diff --git a/Engine/Shaders/Private/VirtualShadowMaps/VirtualShadowMapPhysicalPageManagement.usf b/Engine/Shaders/Private/VirtualShadowMaps/VirtualShadowMapPhysicalPageManagement.usf
--- a/Engine/Shaders/Private/VirtualShadowMaps/VirtualShadowMapPhysicalPageManagement.usf
+++ b/Engine/Shaders/Private/VirtualShadowMaps/VirtualShadowMapPhysicalPageManagement.usf
@@ -9,10 +9,12 @@
 #include "../ReductionCommon.ush"
 #include "../GPUMessaging.ush"
 #include "../ShaderPrint.ush"
+#include "/Engine/Shared/VirtualShadowMapDefinitions.h"
 #include "VirtualShadowMapProjectionStructs.ush"
 #include "VirtualShadowMapProjectionCommon.ush"
 #include "VirtualShadowMapPageAccessCommon.ush"
 #include "VirtualShadowMapStats.ush"
+#include "VirtualShadowMapPerPageDispatch.ush"
 
 #ifndef HAS_CACHE_DATA
 #define HAS_CACHE_DATA 1
@@ -106,6 +108,7 @@ StructuredBuffer<int> PrevPhysicalPageLists;
 
 RWStructuredBuffer<uint4> OutUncachedPageRectBounds;
 RWStructuredBuffer<uint4> OutAllocatedPageRectBounds;
+uint NumPageRectsToClear;
 
 // This is admitadly a weird fusion of several initializations but it is the first thing
 // we run in a given analysis phase so it's more efficient to do it all here rather than
@@ -113,11 +116,19 @@ RWStructuredBuffer<uint4> OutAllocatedPageRectBounds;
 [numthreads(VSM_DEFAULT_CS_GROUP_X, 1, 1)]
 void InitPageRectBounds(uint3 Index : SV_DispatchThreadID)
 {
-	if (Index.x < VSM_MAX_MIP_LEVELS * VirtualShadowMap.NumShadowMapSlots)
+	// The X thread id maps to the range NumPageRectsToClear which is (GetNumFullShadowMaps() + GetNumSinglePageShadowMaps()) * FVirtualShadowMap::MaxMipLevels
+	// this avoids clearing unused slots.
+	if (Index.x < NumPageRectsToClear)
 	{
+		uint RectOffset = Index.x;
+		// The full shadow maps are offset to a distant part of the ID range
+		if (Index.x >= VirtualShadowMap.NumSinglePageShadowMaps * VSM_MAX_MIP_LEVELS )
+		{
+			RectOffset += VSM_MAX_SINGLE_PAGE_SHADOW_MAPS * VSM_MAX_MIP_LEVELS - VirtualShadowMap.NumSinglePageShadowMaps * VSM_MAX_MIP_LEVELS;
+		}
 		uint4 Empty = uint4(VSM_LEVEL0_DIM_PAGES_XY, VSM_LEVEL0_DIM_PAGES_XY, 0, 0);
-		OutUncachedPageRectBounds[Index.x] = Empty;
-		OutAllocatedPageRectBounds[Index.x] = Empty;
+		OutUncachedPageRectBounds[RectOffset] = Empty;
+		OutAllocatedPageRectBounds[RectOffset] = Empty;
 	}
 
 	// Clear the various list counters
@@ -421,8 +432,6 @@ void CalcPageAddressFromIndex(uint Index, inout uint MipLevel, inout uint2 PageA
 	PageAddress.y = OffsetInLevel >> (VSM_LOG2_LEVEL0_DIM_PAGES_XY - MipLevel);
 }
 
-#define EARLY_ALLOCATE 1
-
 void AllocateNewPageMappings(uint VirtualShadowMapId, uint GlobalPageOffset, uint PageOffsetInSM, const bool bSinglePageSM)
 {
 	checkStructuredBufferAccessSlow(PageRequestFlags, GlobalPageOffset);
@@ -495,43 +504,61 @@ void AllocateNewPageMappings(uint VirtualShadowMapId, uint GlobalPageOffset, uin
 	}
 }
 
+#ifdef ClearPageTableCS
+
+RWStructuredBuffer< uint > OutDestBuffer;
+uint ClearValue;
+
 /**
- * X = One thread per virtual page, Y dim == GetNumShadowMaps() (single-page + full).
+ * X = One thread per virtual page = VSM_PAGE_TABLE_SIZE, Y dim == GetNumFullShadowMaps() + GetNumSinglePage() / VSM_PAGE_TABLE_SIZE.
  */
-[numthreads(VSM_DEFAULT_CS_GROUP_X, 1, 1)]
-void AllocateNewPageMappingsCS(uint2 ThreadId : SV_DispatchThreadID)
+[numthreads(PER_PAGE_THREAD_GROUP_SIZE_X, 1, 1)]
+void ClearPageTableCS(uint2 DispatchThreadId : SV_DispatchThreadID)
 {
-	if (ThreadId.x >= VSM_PAGE_TABLE_SIZE)
+	FPerPageDispatchSetup Setup; 
+	Setup.Init(DispatchThreadId);
+
+	if (!Setup.bValid)
 	{
 		return;
 	}
 
-	// Note: driving with an if at the root to try to help compiler preserve scalarization for the non-single-page path.
-	// Single-page lights use the K remaining slices of the 2D dispatch
-	BRANCH
-	if (ThreadId.y >= VirtualShadowMap.NumFullShadowMaps)
+	// Run a loop to do the work
+	for (uint PageOffsetInSM = Setup.PageStartOffset + Setup.LoopStart; PageOffsetInSM < Setup.PageEndOffset; PageOffsetInSM += Setup.LoopStride)
 	{
-		// remap the threads ID to shadow map ID
-		uint VirtualShadowMapId = (ThreadId.y - VirtualShadowMap.NumFullShadowMaps) * VSM_PAGE_TABLE_SIZE + ThreadId.x;
-		uint GlobalPageOffset = VirtualShadowMapId;
-
-		// Avoid oob access
-		if (VirtualShadowMapId < VirtualShadowMap.NumSinglePageShadowMaps)
-		{
-			AllocateNewPageMappings(VirtualShadowMapId, GlobalPageOffset, 0U, true);
-		}
+		uint PageTableEntryOffset = Setup.BaseLevelOffset.LevelOffset + PageOffsetInSM;
+		OutDestBuffer[PageTableEntryOffset] = ClearValue;
 	}
-	else
+}
+
+#endif // ClearPageTableCS
+
+#ifdef AllocateNewPageMappingsCS
+
+/**
+ * X = One thread per virtual page = VSM_PAGE_TABLE_SIZE, Y dim == GetNumFullShadowMaps() + GetNumSinglePage() / VSM_PAGE_TABLE_SIZE.
+ */
+[numthreads(PER_PAGE_THREAD_GROUP_SIZE_X, 1, 1)]
+void AllocateNewPageMappingsCS(uint2 DispatchThreadId : SV_DispatchThreadID)
+{
+	FPerPageDispatchSetup Setup; 
+	Setup.Init(DispatchThreadId);
+
+	if (!Setup.bValid)
 	{
-		uint PageOffsetInSM = ThreadId.x;
-		// Full SMs are allocated after the single-page ones.
-		uint VirtualShadowMapId = VSM_MAX_SINGLE_PAGE_SHADOW_MAPS + ThreadId.y;
-		uint GlobalPageOffset = CalcFullPageTableLevelOffset(VirtualShadowMapId, 0U) + PageOffsetInSM;
+		return;
+	}
 
-		AllocateNewPageMappings(VirtualShadowMapId, GlobalPageOffset, PageOffsetInSM, false);
+	// Run a grid x-stride loop to do all the work
+	for (uint PageOffsetInSM = Setup.PageStartOffset + Setup.LoopStart; PageOffsetInSM < Setup.PageEndOffset; PageOffsetInSM += Setup.LoopStride)
+	{
+		uint GlobalPageOffset = Setup.BaseLevelOffset.LevelOffset + PageOffsetInSM;
+		AllocateNewPageMappings(Setup.VirtualShadowMapId, GlobalPageOffset, PageOffsetInSM, Setup.BaseLevelOffset.bIsSinglePageSM);
 	}
 }
 
+#endif // AllocateNewPageMappingsCS
+
 // NOTE: We only launch a single group here for now to avoid multi-pass so we really want it as large as possible
 // Can optimize this later if needed for larger physical page counts
 #define NUM_THREADS_PER_GROUP 1024
@@ -666,13 +693,13 @@ void MergePhysicalPixel(uint2 PixelCoord)
 }
 
 // Log2 2D dimension of thread group size, 2^4 == 16, 
-#define LOG2_TILE_THREAD_GROUP_SIZE_XY 4U
-#define TILE_THREAD_GROUP_SIZE_XY (1U << LOG2_TILE_THREAD_GROUP_SIZE_XY)
+#define LOG2_TILE_THREAD_GROUP_SIZE_XY 4u
+#define TILE_THREAD_GROUP_SIZE_XY (1u << LOG2_TILE_THREAD_GROUP_SIZE_XY)
 
 // Each thread takes 2x2 samples to work with, so tile size is 2x thread group size
-#define LOG2_TILE_SIZE_XY (LOG2_TILE_THREAD_GROUP_SIZE_XY + 1U)	
+#define LOG2_TILE_SIZE_XY (LOG2_TILE_THREAD_GROUP_SIZE_XY + 1u)
 
-#if VSM_LOG2_PAGE_SIZE < LOG2_TILE_SIZE
+#if VSM_LOG2_PAGE_SIZE < LOG2_TILE_SIZE_XY
 #error "VSM_LOG2_PAGE_SIZE must be larger than LOG2_TILE_SIZE, either increase one or reduce the other"
 #endif 
 
diff --git a/Engine/Shaders/Private/VirtualShadowMaps/VirtualShadowMapProjectionCommon.ush b/Engine/Shaders/Private/VirtualShadowMaps/VirtualShadowMapProjectionCommon.ush
--- a/Engine/Shaders/Private/VirtualShadowMaps/VirtualShadowMapProjectionCommon.ush
+++ b/Engine/Shaders/Private/VirtualShadowMaps/VirtualShadowMapProjectionCommon.ush
@@ -101,14 +101,14 @@ FVirtualShadowMapSample SampleVirtualShadowMapLevel(int VirtualShadowMapId, floa
 	return Result;
 }
 
-FVirtualShadowMapSample SampleVirtualShadowMap(int VirtualShadowMapId, float2 ShadowMapUV)
+FVirtualShadowMapSample SampleVirtualShadowMap(int VirtualShadowMapId, float2 ShadowMapUV, uint MinMipLevel)
 {
-	FShadowPageTranslationResult Page = ShadowVirtualToPhysicalUV(VirtualShadowMapId, ShadowMapUV);
+	FShadowPageTranslationResult Page = ShadowVirtualToPhysicalUV(VirtualShadowMapId, ShadowMapUV, MinMipLevel);
 	FVirtualShadowMapSample Result = InitVirtualShadowMapSample();
 	if (Page.bValid)
 	{
 		Result.bValid = true;
-		Result.MipLevel = Page.LODOffset;
+		Result.MipLevel = Page.LODOffset + MinMipLevel;
 		Result.VirtualShadowMapId = VirtualShadowMapId;
 		Result.VirtualTexelAddress = Page.VirtualTexelAddress;
 		Result.VirtualTexelAddressFloat = Page.VirtualTexelAddressFloat;
@@ -453,7 +453,7 @@ FVirtualShadowMapSampleResult SampleVirtualShadowMapLocal(
 	float4 ShadowUVz = mul(float4(ShadowTranslatedWorldPosition, 1.0f), ProjectionData.TranslatedWorldToShadowUVMatrix);
 	ShadowUVz.xyz /= ShadowUVz.w;
 
-	FVirtualShadowMapSample SmSample = SampleVirtualShadowMap(VirtualShadowMapId, ShadowUVz.xy);
+	FVirtualShadowMapSample SmSample = SampleVirtualShadowMap(VirtualShadowMapId, ShadowUVz.xy, ProjectionData.MinMipLevel);
 	if (SmSample.bValid)
 	{
 		FVirtualShadowMapSampleResult Result = InitVirtualShadowMapSampleResult();
diff --git a/Engine/Shaders/Private/VirtualShadowMaps/VirtualShadowMapProjectionSpot.ush b/Engine/Shaders/Private/VirtualShadowMaps/VirtualShadowMapProjectionSpot.ush
--- a/Engine/Shaders/Private/VirtualShadowMaps/VirtualShadowMapProjectionSpot.ush
+++ b/Engine/Shaders/Private/VirtualShadowMaps/VirtualShadowMapProjectionSpot.ush
@@ -37,7 +37,7 @@ float ComputeOptimalSlopeBiasLocal(
 
 	float4 ShadowUVz = mul(float4(ShadowTranslatedWorldPosition, 1.0f), ProjectionData.TranslatedWorldToShadowUVMatrix);
 	ShadowUVz.xyz /= ShadowUVz.w;
-	FShadowPageTranslationResult Page = ShadowVirtualToPhysicalUV(ProjectionData.VirtualShadowMapId, ShadowUVz.xy);
+	FShadowPageTranslationResult Page = ShadowVirtualToPhysicalUV(ProjectionData.VirtualShadowMapId, ShadowUVz.xy, ProjectionData.MinMipLevel);
 
 	float MipLevelDim = float(CalcLevelDimsTexels(Page.LODOffset));
 	float2 TexelCenter = float2(Page.VirtualTexelAddress) + 0.5f;
@@ -55,6 +55,7 @@ struct FSMRTSingleRayState
 	float3 RayStartUVz;	
 	float3 RayStepUVz;
 	float ExtrapolateSlope;
+	uint MinMipLevel;
 	// Debug output
 	int MipLevel;
 	uint2 VirtualTexelAddress;
@@ -97,6 +98,7 @@ FSMRTSingleRayState SMRTSingleRayInitialize(
 	Result.RayStartUVz = RayStartUVz.xyz;
 	Result.RayStepUVz = RayStepUVz.xyz;
 	Result.ExtrapolateSlope = ExtrapolateSlope * RayStepUVz.z;
+	Result.MinMipLevel = ProjectionData.MinMipLevel;
 	Result.MipLevel = 0;
 	Result.VirtualTexelAddress = uint2(0xFFFFFFFF, 0xFFFFFFFF);
 	Result.PhysicalTexelAddress = uint2(0xFFFFFFFF, 0xFFFFFFFF);
@@ -114,7 +116,7 @@ FSMRTSample SMRTFindSample(inout FSMRTSingleRayState RayState, float SampleTime)
 
 	if (all(UVz.xy == saturate(UVz.xy)))
 	{
-		FVirtualShadowMapSample SmSample = SampleVirtualShadowMap(RayState.VirtualShadowMapId, UVz.xy);
+		FVirtualShadowMapSample SmSample = SampleVirtualShadowMap(RayState.VirtualShadowMapId, UVz.xy, RayState.MinMipLevel);
 		if (SmSample.bValid)
 		{
 			Sample.bValid = true;
diff --git a/Engine/Shaders/Private/VirtualShadowMaps/VirtualShadowMapProjectionStructs.ush b/Engine/Shaders/Private/VirtualShadowMaps/VirtualShadowMapProjectionStructs.ush
--- a/Engine/Shaders/Private/VirtualShadowMaps/VirtualShadowMapProjectionStructs.ush
+++ b/Engine/Shaders/Private/VirtualShadowMaps/VirtualShadowMapProjectionStructs.ush
@@ -28,12 +28,13 @@ struct FVirtualShadowMapProjectionShaderData
 	uint Flags;
 	float ClipmapLevelWPODistanceDisabledThresholdSquared;
 	float LightSourceRadius;
+	float TexelDitherScale; // Per-light additional dither filtering
+	uint MinMipLevel; // Minimum mip level that is used for this mip chain (only valid for local lights)
 
 	// Derived data for convenience when passing the structure around
 	int VirtualShadowMapId;
 	bool bUnCached; // see VSM_PROJ_FLAG_UNCACHED
 	bool bUnreferenced; // see VSM_PROJ_FLAG_UNREFERENCED
-	float TexelDitherScale; // Per-light additional dither filtering
 };
 
 FVirtualShadowMapProjectionShaderData DecodeVirtualShadowMapProjectionData(ByteAddressBuffer ProjectionData, int VirtualShadowMapId)
@@ -82,6 +83,7 @@ FVirtualShadowMapProjectionShaderData DecodeVirtualShadowMapProjectionData(ByteA
 		                                    = asfloat(ProjectionData.Load (Base + (16*17 + 4*1)));
 	
 	Result.TexelDitherScale					= asfloat(ProjectionData.Load (Base + (16*17 + 4*2)));
+	Result.MinMipLevel						=         ProjectionData.Load (Base + (16*17 + 4*3));
 	
 	// Computed data
 	Result.PreViewTranslation				= MakeDFVector3(PreViewTranslationHigh, PreViewTranslationLow);
diff --git a/Engine/Shaders/Shared/HLSLMathAliases.h b/Engine/Shaders/Shared/HLSLMathAliases.h
new file mode 100644
--- /dev/null
+++ b/Engine/Shaders/Shared/HLSLMathAliases.h
@@ -0,0 +1,35 @@
+// Copyright Epic Games, Inc. All Rights Reserved.
+
+#pragma once
+
+/*================================================================================================
+	HLSLMathAliases.h: include this file in headers that are used in both C++ and HLSL
+	to define aliases for common HSLS math functions to C++ UE math functions.
+=================================================================================================*/
+
+#include "Math/UnrealMathUtility.h"
+
+namespace UE::HLSL
+{
+	inline float log2(float Value)
+	{
+		return FMath::Log2(Value);
+	}
+
+	inline float floor(float Value)
+	{
+		return FMath::Floor(Value);
+	}
+
+	template <typename ValueType>
+	inline float max(ValueType A, ValueType B)
+	{
+		return FMath::Max(A, B);
+	}
+	template <typename ValueType>
+	inline float min(ValueType A, ValueType B)
+	{
+		return FMath::Min(A, B);
+	}
+}
+
diff --git a/Engine/Shaders/Shared/VirtualShadowMapDefinitions.h b/Engine/Shaders/Shared/VirtualShadowMapDefinitions.h
--- a/Engine/Shaders/Shared/VirtualShadowMapDefinitions.h
+++ b/Engine/Shaders/Shared/VirtualShadowMapDefinitions.h
@@ -7,6 +7,19 @@ VirtualShadowMapDefinitions.h: used in virtual shadow map shaders and C++ code t
 
 #pragma once
 
+// Page size is 128x128
+#define VSM_LOG2_PAGE_SIZE 7u
+#define VSM_PAGE_SIZE (1u << VSM_LOG2_PAGE_SIZE)
+#define VSM_PAGE_SIZE_MASK (VSM_PAGE_SIZE - 1u)
+// Page table size is 128x128 (total 16k)
+#define VSM_LOG2_LEVEL0_DIM_PAGES_XY 7u
+#define VSM_LEVEL0_DIM_PAGES_XY (1u << VSM_LOG2_LEVEL0_DIM_PAGES_XY)
+#define VSM_MAX_MIP_LEVELS (VSM_LOG2_LEVEL0_DIM_PAGES_XY + 1u)
+#define VSM_VIRTUAL_MAX_RESOLUTION_XY (VSM_LEVEL0_DIM_PAGES_XY * VSM_PAGE_SIZE)
+#define VSM_RASTER_WINDOW_PAGES (4u)
+#define VSM_PAGE_TABLE_SIZE 21845u // (0x55555555u & ((1u << (VSM_MAX_MIP_LEVELS << 1u)) - 1u))
+
+
 #define VIRTUAL_SHADOW_MAP_VISUALIZE_NONE						0
 #define VIRTUAL_SHADOW_MAP_VISUALIZE_SHADOW_FACTOR				(1 << 0)
 #define VIRTUAL_SHADOW_MAP_VISUALIZE_CLIPMAP_OR_MIP				(1 << 1)
@@ -71,6 +84,7 @@ VirtualShadowMapDefinitions.h: used in virtual shadow map shaders and C++ code t
 
 #ifdef __cplusplus
 #include "HLSLTypeAliases.h"
+#include "HLSLMathAliases.h"
 
 namespace UE::HLSL
 {
@@ -96,10 +110,32 @@ struct FNextVirtualShadowMapData
 	int _Padding;
 };
 
+
+inline uint GetMipLevelLocal(float Footprint, uint MipMode, float ShadowMapResolutionLodBias, float GlobalResolutionLodBias)
+{
+	float MipLevelFloat = log2(Footprint) + ShadowMapResolutionLodBias + GlobalResolutionLodBias;
+	uint MipLevel = uint(max(floor(MipLevelFloat), 0.0f));
+	MipLevel = min(MipLevel, (VSM_MAX_MIP_LEVELS - 1U));
+
+	if (MipMode == 1)
+	{
+		// Even mips
+		MipLevel &= ~(1U);
+	}
+	else if (MipMode == 2)
+	{
+		// Odd mips
+		MipLevel |= 1U;
+	}
+
+	return MipLevel;
+}
+
 #ifdef __cplusplus
 } // namespace UE::HLSL
 
 using FVSMVisibleInstanceCmd = UE::HLSL::FVSMVisibleInstanceCmd;
 using FVSMCullingBatchInfo = UE::HLSL::FVSMCullingBatchInfo;
 using FNextVirtualShadowMapData = UE::HLSL::FNextVirtualShadowMapData;
-#endif
\ No newline at end of file
+
+#endif
diff --git a/Engine/Source/Runtime/Renderer/Private/Shadows/ShadowSceneRenderer.cpp b/Engine/Source/Runtime/Renderer/Private/Shadows/ShadowSceneRenderer.cpp
--- a/Engine/Source/Runtime/Renderer/Private/Shadows/ShadowSceneRenderer.cpp
+++ b/Engine/Source/Runtime/Renderer/Private/Shadows/ShadowSceneRenderer.cpp
@@ -78,6 +78,8 @@ static TAutoConsoleVariable<float> CVarResolutionLodBiasLocalMoving(
 	ECVF_Scalability | ECVF_RenderThreadSafe
 );
 
+extern TAutoConsoleVariable<int32> CVarMarkPixelPagesMipModeLocal;
+
 bool IsVSMOnePassProjectionEnabled(const FEngineShowFlags& ShowFlags)
 {
 	return CVarVirtualShadowOnePassProjection.GetValueOnAnyThread() 
@@ -127,6 +129,22 @@ struct FHeapPair
 
 void FShadowSceneRenderer::BeginRender(FRDGBuilder& GraphBuilder)
 {
+	ViewDatas.Reserve(SceneRenderer.Views.Num());
+	for (const FViewInfo& View : SceneRenderer.Views)
+	{
+		const FVector2f ViewSize = FVector2f(View.ViewRect.Size());
+		FVector2f RadiusClipXY = FVector2f(2.0f) / ViewSize;
+
+		const FMatrix &ViewToClip = View.ViewMatrices.GetProjectionMatrix();
+		// TODO: is RadiusXY always symmetrical?
+		FVector2f ProjScaleXY = FVector2f(static_cast<float>(ViewToClip.M[0][0]), static_cast<float>(ViewToClip.M[1][1]));
+		FVector2f RadiusXY = RadiusClipXY / ProjScaleXY;
+		float MinRadiusXY = FMath::Min(RadiusXY.X, RadiusXY.Y);
+		float ClipToViewSizeScale = ViewToClip.M[2][3] * MinRadiusXY;
+		float ClipToViewSizeBias = ViewToClip.M[3][3] * MinRadiusXY;
+		ViewDatas.Emplace(FViewData{ ClipToViewSizeScale, ClipToViewSizeBias });
+	}
+
 	// Kick off shadow scene updates.
 	ShadowScene.UpdateForRenderedFrame(GraphBuilder);
 
@@ -230,6 +248,32 @@ FVirtualShadowMapProjectionShaderData FShadowSceneRenderer::GetLocalLightProject
 	return Data;
 }
 
+
+/**
+ * Calculate the radius in world-space units of a single pixel at a given depth.
+ */
+static float GetWorldSpacePixelFootprint(float ViewSpaceDepth, float ClipToViewSizeScale, float ClipToViewSizeBias)
+{
+	return ViewSpaceDepth * ClipToViewSizeScale + ClipToViewSizeBias;
+}
+
+/**
+ * Compute the lowest (highest res) mip level that might be marked by any pixels inside the light influence radius for a given scene primary view.
+ */
+static uint32 GetConservativeMipLevelLocal(const FViewInfo& View, float ClipToViewSizeScale, float ClipToViewSizeBias, const FVector& LightOrigin, float LightRadius, float WorldToShadowFootprintScale, float ResolutionLodBias, float GlobalResolutionLodBias, uint32 MipModeLocal)
+{
+	// Note: not just a rotation, full world-space DP.
+	FVector ViewSpaceOrigin = View.ShadowViewMatrices.GetViewMatrix().TransformPosition(LightOrigin);
+
+	// Remove radius to arrive at minimum possible z-distance in view space, from primary view.
+	float RadiusWorld = GetWorldSpacePixelFootprint(FMath::Max(0.0f, float(ViewSpaceOrigin.Z) - LightRadius), ClipToViewSizeScale, ClipToViewSizeBias);
+
+	// Radius is the max possible shadow view space Z, which would require the max res.
+	float ShadowFootprint = RadiusWorld * WorldToShadowFootprintScale / LightRadius;
+
+	return UE::HLSL::GetMipLevelLocal(ShadowFootprint, MipModeLocal, ResolutionLodBias, GlobalResolutionLodBias);
+}
+
 TSharedPtr<FVirtualShadowMapPerLightCacheEntry> FShadowSceneRenderer::AddLocalLightShadow(const FWholeSceneProjectedShadowInitializer& ProjectedShadowInitializer, FProjectedShadowInfo* ProjectedShadowInfo, FLightSceneInfo* LightSceneInfo, float MaxScreenRadius)
 {
 	FVirtualShadowMapArrayCacheManager* CacheManager = VirtualShadowMapArray.CacheManager;
@@ -276,6 +320,33 @@ TSharedPtr<FVirtualShadowMapPerLightCacheEntry> FShadowSceneRenderer::AddLocalLi
 	ProjectedShadowInfo->VirtualShadowMapPerLightCacheEntry = PerLightCacheEntry;
 	ProjectedShadowInfo->bShouldRenderVSM = !PerLightCacheEntry->IsFullyCached();
 
+	// Compute conservative mip level estimate based on radius of the bounding sphere.
+	// TODO: can probably do better by finding  closest point on cone for certain scenarios
+
+	const FVector2f ShadowViewSize = FVector2f(FVirtualShadowMap::VirtualMaxResolutionXY, FVirtualShadowMap::VirtualMaxResolutionXY);
+	const FMatrix &ShadowViewToClip = ProjectedShadowInfo->bOnePassPointLightShadow ? ProjectedShadowInfo->OnePassShadowFaceProjectionMatrix : ProjectedShadowInfo->ViewToClipOuter;
+	float ShadowProjScale = ShadowViewToClip.M[0][0]; // always symmetrical
+	const float WorldToShadowFootprintScale = ShadowProjScale * ShadowViewSize.X;
+
+	uint32 MinMipLevel = FVirtualShadowMap::MaxMipLevels;
+	for (int32 ViewIndex = 0; ViewIndex < SceneRenderer.Views.Num(); ++ViewIndex)
+	{
+		const FViewInfo& View = SceneRenderer.Views[ViewIndex];
+		const FViewData& ViewData = ViewDatas[ViewIndex];
+
+		MinMipLevel = FMath::Min(MinMipLevel, GetConservativeMipLevelLocal(
+			View, 
+			ViewData.ClipToViewSizeScale,
+			ViewData.ClipToViewSizeBias,
+			LightSceneProxy->GetOrigin(), 
+			LightSceneProxy->GetRadius(),
+			WorldToShadowFootprintScale,
+			ResolutionLODBiasLocal, 
+			CacheManager->GetGlobalResolutionLodBias(), 
+			CVarMarkPixelPagesMipModeLocal.GetValueOnRenderThread()
+			));
+	}
+
 	for (int32 Index = 0; Index < NumMaps; ++Index)
 	{
 		const int32 FaceVirtualShadowMapId = VirtualShadowMapId + Index;
@@ -283,6 +354,7 @@ TSharedPtr<FVirtualShadowMapPerLightCacheEntry> FShadowSceneRenderer::AddLocalLi
 		VirtualSmCacheEntry.Update(VirtualShadowMapArray, *PerLightCacheEntry, FaceVirtualShadowMapId);
 		// Update projection data
 		VirtualSmCacheEntry.ProjectionData = GetLocalLightProjectionShaderData(ResolutionLODBiasLocal, ProjectedShadowInfo, Index);
+		VirtualSmCacheEntry.ProjectionData.MinMipLevel = MinMipLevel;
 	}
 
 	return PerLightCacheEntry;
diff --git a/Engine/Source/Runtime/Renderer/Private/Shadows/ShadowSceneRenderer.h b/Engine/Source/Runtime/Renderer/Private/Shadows/ShadowSceneRenderer.h
--- a/Engine/Source/Runtime/Renderer/Private/Shadows/ShadowSceneRenderer.h
+++ b/Engine/Source/Runtime/Renderer/Private/Shadows/ShadowSceneRenderer.h
@@ -104,6 +104,12 @@ public:
 	UE::Renderer::Private::IShadowInvalidatingInstances *GetInvalidatingInstancesInterface(const FSceneView *SceneView);
 
 private:
+	struct FViewData
+	{
+		float ClipToViewSizeScale = 0.0f;
+		float ClipToViewSizeBias = 0.0f;
+	};
+
 	UE::Tasks::FTask RendererSetupTask;
 
 	FVirtualShadowMapProjectionShaderData GetLocalLightProjectionShaderData(float ResolutionLODBiasLocal, const FProjectedShadowInfo* ProjectedShadowInfo, int32 MapIndex) const;
@@ -138,4 +144,5 @@ private:
 	FNaniteVisibilityQuery* NaniteVisibilityQuery = nullptr;
 	Nanite::FPackedViewArray* VirtualShadowMapViews = nullptr;
 	FSceneInstanceCullingQuery *SceneInstanceCullingQuery = nullptr;
+	TArray<FViewData, SceneRenderingAllocator> ViewDatas;
 };
diff --git a/Engine/Source/Runtime/Renderer/Private/VirtualShadowMaps/VirtualShadowMapArray.cpp b/Engine/Source/Runtime/Renderer/Private/VirtualShadowMaps/VirtualShadowMapArray.cpp
--- a/Engine/Source/Runtime/Renderer/Private/VirtualShadowMaps/VirtualShadowMapArray.cpp
+++ b/Engine/Source/Runtime/Renderer/Private/VirtualShadowMaps/VirtualShadowMapArray.cpp
@@ -145,7 +145,7 @@ TAutoConsoleVariable<int32> CVarMarkPixelPages(
 	ECVF_RenderThreadSafe
 );
 
-static TAutoConsoleVariable<int32> CVarMarkPixelPagesMipModeLocal(
+TAutoConsoleVariable<int32> CVarMarkPixelPagesMipModeLocal(
 	TEXT("r.Shadow.Virtual.MarkPixelPagesMipModeLocal"),
 	0,
 	TEXT("When enabled, this uses a subset of mips to reduce instance duplication in VSMs. Will result in better performance but a harsher falloff on mip transitions.\n")
@@ -829,15 +829,6 @@ void FVirtualShadowMapArray::SetShaderDefines(FShaderCompilerEnvironment& OutEnv
 {
 	static_assert(FVirtualShadowMap::Log2Level0DimPagesXY * 2U + NANITE_MAX_VIEWS_PER_CULL_RASTERIZE_PASS_BITS <= 32U, "Page indirection plus view index must fit into 32-bits for page-routing storage!");
 	OutEnvironment.SetDefine(TEXT("ENABLE_NON_NANITE_VSM"), GEnableNonNaniteVSM);
-	OutEnvironment.SetDefine(TEXT("VSM_PAGE_SIZE"), FVirtualShadowMap::PageSize);
-	OutEnvironment.SetDefine(TEXT("VSM_PAGE_SIZE_MASK"), FVirtualShadowMap::PageSizeMask);
-	OutEnvironment.SetDefine(TEXT("VSM_LOG2_PAGE_SIZE"), FVirtualShadowMap::Log2PageSize);
-	OutEnvironment.SetDefine(TEXT("VSM_LEVEL0_DIM_PAGES_XY"), FVirtualShadowMap::Level0DimPagesXY);
-	OutEnvironment.SetDefine(TEXT("VSM_LOG2_LEVEL0_DIM_PAGES_XY"), FVirtualShadowMap::Log2Level0DimPagesXY);
-	OutEnvironment.SetDefine(TEXT("VSM_MAX_MIP_LEVELS"), FVirtualShadowMap::MaxMipLevels);
-	OutEnvironment.SetDefine(TEXT("VSM_VIRTUAL_MAX_RESOLUTION_XY"), FVirtualShadowMap::VirtualMaxResolutionXY);
-	OutEnvironment.SetDefine(TEXT("VSM_RASTER_WINDOW_PAGES"), FVirtualShadowMap::RasterWindowPages);
-	OutEnvironment.SetDefine(TEXT("VSM_PAGE_TABLE_SIZE"), FVirtualShadowMap::PageTableSize);
 	OutEnvironment.SetDefine(TEXT("MAX_PAGE_AREA_DIAGNOSTIC_SLOTS"), MaxPageAreaDiagnosticSlots);
 	OutEnvironment.SetDefine(TEXT("INDEX_NONE"), INDEX_NONE);
 }
@@ -1034,16 +1025,194 @@ class FUpdatePhysicalPages : public FVirtualShadowMapPageManagementShader
 };
 IMPLEMENT_GLOBAL_SHADER(FUpdatePhysicalPages, "/Engine/Private/VirtualShadowMaps/VirtualShadowMapPhysicalPageManagement.usf", "UpdatePhysicalPages", SF_Compute );
 
-class FAllocateNewPageMappingsCS : public FVirtualShadowMapPageManagementShader
+/**
+ * Utility to help schedule kernels that do processing for each page to facilitate not processing mip levels that can't be reached.
+ */
+class FVirtualShadowMapPerPageShader : public FVirtualShadowMapPageManagementShader
+{
+public:
+	static constexpr int32 ThreadGroupSizeX = 64;
+
+	BEGIN_SHADER_PARAMETER_STRUCT(FParameters, )
+		SHADER_PARAMETER_RDG_BUFFER_SRV( StructuredBuffer< uint >, VirtualShadowMapIds)
+		SHADER_PARAMETER( uint32, VirtualShadowMapIdsOffset ) 
+		SHADER_PARAMETER( uint32, NumVirtualShadowMapIds ) 
+		SHADER_PARAMETER( uint32, PerPageDispatchDimX ) 
+		SHADER_PARAMETER( uint32, bUseThreadPerId ) 
+	END_SHADER_PARAMETER_STRUCT()
+
+	static void ModifyCompilationEnvironment(const FGlobalShaderPermutationParameters& Parameters, FShaderCompilerEnvironment& OutEnvironment)
+	{
+		OutEnvironment.SetDefine(TEXT("PER_PAGE_THREAD_GROUP_SIZE_X"), ThreadGroupSizeX);
+		OutEnvironment.SetDefine(TEXT("PER_PAGE_DISPATCH_SETUP"), 1);
+		FVirtualShadowMapPageManagementShader::ModifyCompilationEnvironment(Parameters, OutEnvironment);
+	}
+
+	class FDispatcher
+	{
+	public:
+
+		struct FBinConfig
+		{
+			int32 DispatchDimX;
+		};
+
+		struct FBin
+		{
+			int32 VirtualShadowMapIdsOffset = 0;
+			int32 NumVirtualShadowMapIds = 0;
+		};
+
+		// 1, Small, medium, large
+		static constexpr int32 BinCount = 4;
+		static constexpr int32 BinDispatchDimX[BinCount] = 
+		{
+			4096,
+			512,
+			64,
+			1
+		};
+
+		static int32 CalcBin(int32 VirtualShadowMapId, int32 MinMipLevel)
+		{
+			if (IsSinglePageVirtualShadowMap(VirtualShadowMapId))
+			{
+				// last bin, single thread per shadow map
+				return BinCount - 1;
+			}
+
+			if (MinMipLevel < 6)
+			{
+				return MinMipLevel / 2; 
+			}
+
+			// last bin, single thread per shadow map
+			return BinCount - 1;
+		}
+
+		TStaticArray<FBin, BinCount> Bins;
+		FRDGBufferSRV* VirtualShadowMapIdsSRV = nullptr;
+
+		struct FBuilder
+		{
+			struct FIdBinIndex
+			{
+				uint32 VirtualShadowMapId : 28;
+				uint32 BinIndex : 4;
+			};
+
+			void Add(int32 VirtualShadowMapId, uint32 MinMipLevel)
+			{
+				if (VirtualShadowMapId != INDEX_NONE)
+				{
+					FIdBinIndex IdBinIndex;
+					IdBinIndex.BinIndex = CalcBin(VirtualShadowMapId, MinMipLevel);
+					IdBinIndex.VirtualShadowMapId = VirtualShadowMapId;
+					Tmp.Add(IdBinIndex);
+					Bins[IdBinIndex.BinIndex].NumVirtualShadowMapIds += 1;
+				}
+
+			}
+
+			void Reserve(int32 Num)
+			{
+				Tmp.Reserve(Num);
+			}
+
+			TStaticArray<FBin, BinCount> Bins;
+			TArray<FIdBinIndex,SceneRenderingAllocator> Tmp;
+		};
+
+		void Init(FRDGBuilder& GraphBuilder, const FBuilder& Builder)
+		{
+			Bins = Builder.Bins;
+			// counting sort the temp data into buffer of IDs
+			int32 Offset = 0;
+			for (int32 BinIndex = 0; BinIndex < BinCount; ++BinIndex)
+			{
+				Bins[BinIndex].VirtualShadowMapIdsOffset = Offset;
+				Offset += Bins[BinIndex].NumVirtualShadowMapIds;
+				Bins[BinIndex].NumVirtualShadowMapIds = 0;
+			}
+			TArray<uint32,SceneRenderingAllocator> VirtualShadowMapIds;
+			VirtualShadowMapIds.SetNumUninitialized(Builder.Tmp.Num());
+			for (FBuilder::FIdBinIndex IdBinIndex : Builder.Tmp)
+			{
+				FBin &Bin = Bins[IdBinIndex.BinIndex];
+				VirtualShadowMapIds[Bin.VirtualShadowMapIdsOffset + Bin.NumVirtualShadowMapIds++] = IdBinIndex.VirtualShadowMapId;
+			}
+
+			VirtualShadowMapIdsSRV = GraphBuilder.CreateSRV(CreateStructuredBuffer(GraphBuilder, TEXT("Shadow.Virtual.PerPageBinIds"), MoveTemp(VirtualShadowMapIds)));
+		}
+
+		template <typename TShaderClass>
+		void AddPass(FRDGBuilder& GraphBuilder, FRDGEventName&& PassName, const TShaderRef<TShaderClass>& ComputeShader, typename TShaderClass::FParameters *Parameters)
+		{
+			const FShaderParametersMetadata* ParametersMetadata = TShaderClass::FParameters::FTypeInfo::GetStructMetadata();
+			Parameters->VirtualShadowMapPerPageShader.VirtualShadowMapIds = VirtualShadowMapIdsSRV;
+			ClearUnusedGraphResources(ComputeShader, ParametersMetadata, Parameters);
+
+			GraphBuilder.AddPass(Forward<FRDGEventName>(PassName), ParametersMetadata, Parameters, ERDGPassFlags::Compute,
+			[ParametersMetadata, Parameters, ComputeShader, Bins = Bins](FRHIComputeCommandList& RHICmdList)
+			{
+				for (int32 BinIndex = 0; BinIndex < BinCount; ++BinIndex)
+				{
+					const FBin &Bin = Bins[BinIndex];
+					if (Bins[BinIndex].NumVirtualShadowMapIds > 0)
+					{
+						bool bUseThreadPerId = BinDispatchDimX[BinIndex] == 1;
+						int32 ThreadPerIdDispatchDimX = ThreadGroupSizeX * FMath::DivideAndRoundUp(Bin.NumVirtualShadowMapIds, ThreadGroupSizeX);
+						int32 PerPageDispatchDimX = bUseThreadPerId ? ThreadPerIdDispatchDimX : BinDispatchDimX[BinIndex];
+
+						// Set args for each pass
+						Parameters->VirtualShadowMapPerPageShader.PerPageDispatchDimX = PerPageDispatchDimX; 
+						Parameters->VirtualShadowMapPerPageShader.VirtualShadowMapIdsOffset = Bin.VirtualShadowMapIdsOffset;
+						Parameters->VirtualShadowMapPerPageShader.NumVirtualShadowMapIds = Bin.NumVirtualShadowMapIds;
+						Parameters->VirtualShadowMapPerPageShader.bUseThreadPerId = bUseThreadPerId ? 1 : 0;
+
+						FIntVector GroupCount = FIntVector(0,0,1);
+						int32 NumGroupsY = 1;
+						if (!bUseThreadPerId)
+						{
+							NumGroupsY = Bin.NumVirtualShadowMapIds;
+						}
+
+						FComputeShaderUtils::Dispatch(RHICmdList, ComputeShader, ParametersMetadata, *Parameters, FIntVector(FMath::DivideAndRoundUp(PerPageDispatchDimX, ThreadGroupSizeX), NumGroupsY, 1));
+					}
+				}
+
+			});
+
+		}
+	};
+};
+
+class FClearPageTableCS : public FVirtualShadowMapPerPageShader
+{
+	DECLARE_GLOBAL_SHADER(FClearPageTableCS);
+	SHADER_USE_PARAMETER_STRUCT(FClearPageTableCS, FVirtualShadowMapPerPageShader)
+
+	BEGIN_SHADER_PARAMETER_STRUCT(FParameters, )
+		SHADER_PARAMETER_RDG_UNIFORM_BUFFER( FVirtualShadowMapUniformParameters,		VirtualShadowMap )
+		SHADER_PARAMETER_STRUCT_INCLUDE( FVirtualShadowMapPerPageShader::FParameters,	VirtualShadowMapPerPageShader)
+		SHADER_PARAMETER( uint32,														ClearValue )
+		SHADER_PARAMETER_RDG_BUFFER_UAV( RWStructuredBuffer< uint >,					OutDestBuffer )
+	END_SHADER_PARAMETER_STRUCT()
+};
+IMPLEMENT_GLOBAL_SHADER(FClearPageTableCS, "/Engine/Private/VirtualShadowMaps/VirtualShadowMapPhysicalPageManagement.usf", "ClearPageTableCS", SF_Compute);
+
+
+class FAllocateNewPageMappingsCS : public FVirtualShadowMapPerPageShader
 {
 	DECLARE_GLOBAL_SHADER(FAllocateNewPageMappingsCS);
-	SHADER_USE_PARAMETER_STRUCT(FAllocateNewPageMappingsCS, FVirtualShadowMapPageManagementShader)
+	SHADER_USE_PARAMETER_STRUCT(FAllocateNewPageMappingsCS, FVirtualShadowMapPerPageShader)
 
 	class FGenerateStatsDim : SHADER_PERMUTATION_BOOL("VSM_GENERATE_STATS"); 
 	using FPermutationDomain = TShaderPermutationDomain<FGenerateStatsDim>;
 
 	BEGIN_SHADER_PARAMETER_STRUCT(FParameters, )
 		SHADER_PARAMETER_RDG_UNIFORM_BUFFER( FVirtualShadowMapUniformParameters,		VirtualShadowMap )
+		SHADER_PARAMETER_STRUCT_INCLUDE( FVirtualShadowMapPerPageShader::FParameters,	VirtualShadowMapPerPageShader)
 		SHADER_PARAMETER_RDG_BUFFER_SRV( StructuredBuffer< uint >,						PageRequestFlags )
 		SHADER_PARAMETER_RDG_BUFFER_UAV( RWStructuredBuffer< uint >,					OutPageFlags )
 		SHADER_PARAMETER_RDG_BUFFER_UAV( RWStructuredBuffer< uint >,					OutPageTable )
@@ -1081,14 +1250,15 @@ class FAppendPhysicalPageListsCS : public FVirtualShadowMapPageManagementShader
 };
 IMPLEMENT_GLOBAL_SHADER(FAppendPhysicalPageListsCS, "/Engine/Private/VirtualShadowMaps/VirtualShadowMapPhysicalPageManagement.usf", "AppendPhysicalPageLists", SF_Compute );
 
-class FPropagateMappedMipsCS : public FVirtualShadowMapPageManagementShader
+class FPropagateMappedMipsCS : public FVirtualShadowMapPerPageShader
 {
 	DECLARE_GLOBAL_SHADER(FPropagateMappedMipsCS);
-	SHADER_USE_PARAMETER_STRUCT(FPropagateMappedMipsCS, FVirtualShadowMapPageManagementShader)
+	SHADER_USE_PARAMETER_STRUCT(FPropagateMappedMipsCS, FVirtualShadowMapPerPageShader)
 
 	BEGIN_SHADER_PARAMETER_STRUCT(FParameters, )
-		SHADER_PARAMETER_RDG_UNIFORM_BUFFER( FVirtualShadowMapUniformParameters,	VirtualShadowMap )
-		SHADER_PARAMETER_RDG_BUFFER_UAV( RWStructuredBuffer< uint >,				OutPageTable )
+		SHADER_PARAMETER_RDG_UNIFORM_BUFFER( FVirtualShadowMapUniformParameters,		VirtualShadowMap )
+		SHADER_PARAMETER_STRUCT_INCLUDE( FVirtualShadowMapPerPageShader::FParameters,	VirtualShadowMapPerPageShader)
+		SHADER_PARAMETER_RDG_BUFFER_UAV( RWStructuredBuffer< uint >,					OutPageTable )
 	END_SHADER_PARAMETER_STRUCT()
 };
 IMPLEMENT_GLOBAL_SHADER(FPropagateMappedMipsCS, "/Engine/Private/VirtualShadowMaps/VirtualShadowMapPageManagement.usf", "PropagateMappedMips", SF_Compute);
@@ -1235,6 +1405,7 @@ class FInitPageRectBoundsCS : public FVirtualShadowMapPageManagementShader
 		SHADER_PARAMETER_RDG_UNIFORM_BUFFER(FVirtualShadowMapUniformParameters, VirtualShadowMap)
 		SHADER_PARAMETER_RDG_BUFFER_UAV(RWStructuredBuffer<FIntVector4>, OutUncachedPageRectBounds)
 		SHADER_PARAMETER_RDG_BUFFER_UAV(RWStructuredBuffer<FIntVector4>, OutAllocatedPageRectBounds)
+		SHADER_PARAMETER(uint32, NumPageRectsToClear)		
 		SHADER_PARAMETER_RDG_BUFFER_UAV(RWStructuredBuffer<int>, OutPhysicalPageLists)
 		SHADER_PARAMETER_RDG_BUFFER_UAV(RWStructuredBuffer<FPhysicalPageRequest>, OutPhysicalPageAllocationRequests)
 	END_SHADER_PARAMETER_STRUCT()
@@ -1455,25 +1626,6 @@ void FVirtualShadowMapArray::AppendPhysicalPageList(FRDGBuilder& GraphBuilder, b
 	);
 }
 
-void FVirtualShadowMapArray::UploadProjectionData(FRDGBuilder& GraphBuilder)
-{
-	// Create large enough to hold all the unused elements too (wastes GPU memory but allows direct indexing via the ID)
-	uint32 DataSize = sizeof(FVirtualShadowMapProjectionShaderData) * uint32(GetNumShadowMapSlots());
-	FRDGBufferDesc Desc;
-	Desc.Usage = EBufferUsageFlags::UnorderedAccess | EBufferUsageFlags::ShaderResource | EBufferUsageFlags::ByteAddressBuffer | EBufferUsageFlags::StructuredBuffer;
-	Desc.BytesPerElement = 4;
-	Desc.NumElements = DataSize / 4U;
-	ProjectionDataRDG = GraphBuilder.CreateBuffer(Desc, TEXT("Shadow.Virtual.ProjectionData"));
-		
-	FRDGScatterUploadBuffer Uploader;
-	Uploader.Init(GraphBuilder, GetNumShadowMaps(), sizeof(FVirtualShadowMapProjectionShaderData), false, TEXT("Shadow.Virtual.ProjectionData.UploadBuffer"));
-
-	// Upload data for all cached entries, even if they are not referenced this frame
-	CacheManager->UploadProjectionData(Uploader);
-
-	Uploader.ResourceUploadTo(GraphBuilder, ProjectionDataRDG);
-}
-
 void FVirtualShadowMapArray::UpdatePhysicalPageAddresses(FRDGBuilder& GraphBuilder)
 {
 	if (!IsEnabled())
@@ -1590,7 +1742,34 @@ void FVirtualShadowMapArray::BuildPageAllocations(
 	UpdateVisualizeLight(Views, VisibleLightInfos);
 #endif //!UE_BUILD_SHIPPING
 
-	UploadProjectionData(GraphBuilder);
+	FVirtualShadowMapPerPageShader::FDispatcher PerPageShaderDispatcher;
+	{
+		FVirtualShadowMapPerPageShader::FDispatcher::FBuilder Builder;
+
+		// Create large enough to hold all the unused elements too (wastes GPU memory but allows direct indexing via the ID)
+		uint32 DataSize = sizeof(FVirtualShadowMapProjectionShaderData) * uint32(GetNumShadowMapSlots());
+		FRDGBufferDesc Desc;
+		Desc.Usage = EBufferUsageFlags::UnorderedAccess | EBufferUsageFlags::ShaderResource | EBufferUsageFlags::ByteAddressBuffer | EBufferUsageFlags::StructuredBuffer;
+		Desc.BytesPerElement = 4;
+		Desc.NumElements = DataSize / 4U;
+		ProjectionDataRDG = GraphBuilder.CreateBuffer(Desc, TEXT("Shadow.Virtual.ProjectionData"));
+		
+		FRDGScatterUploadBuffer Uploader;
+		Uploader.Init(GraphBuilder, GetNumShadowMaps(), sizeof(FVirtualShadowMapProjectionShaderData), false, TEXT("Shadow.Virtual.ProjectionData.UploadBuffer"));
+
+		for (auto It = CacheManager->CreateConstEntryIterator(); It; ++It)
+		{
+			const TSharedPtr<FVirtualShadowMapPerLightCacheEntry>& CacheEntry = It.Value();
+			for (const auto& Entry : CacheEntry->ShadowMapEntries)
+			{
+				Uploader.Add(Entry.CurrentVirtualShadowMapId, &Entry.ProjectionData);
+				Builder.Add(Entry.CurrentVirtualShadowMapId, Entry.ProjectionData.MinMipLevel);
+			}
+		}
+		Uploader.ResourceUploadTo(GraphBuilder, ProjectionDataRDG);
+
+		PerPageShaderDispatcher.Init(GraphBuilder, Builder);
+	}
 
 	// Stats
 	SET_DWORD_STAT(STAT_VSMSinglePageCount, GetNumSinglePageShadowMaps());
@@ -1624,9 +1803,23 @@ void FVirtualShadowMapArray::BuildPageAllocations(
 	// We potentially over-allocate these to avoid too many different allocation sizes each frame
 	const int32 NumPageFlagsToAllocate = FMath::RoundUpToPowerOfTwo(FMath::Max(128 * 1024, GetNumFullShadowMaps() * int32(FVirtualShadowMap::PageTableSize) + int32(VSM_MAX_SINGLE_PAGE_SHADOW_MAPS)));
 
+	auto AddClearPageTableUAVPass = [&](FRDGBuilder& GraphBuilder, FRDGBufferUAV* DestUAV, uint32 ClearValue)
+	{
+		FClearPageTableCS::FParameters* PassParameters = GraphBuilder.AllocParameters< FClearPageTableCS::FParameters >();
+		PassParameters->VirtualShadowMap = GetUncachedUniformBuffer(GraphBuilder);
+		PassParameters->ClearValue = ClearValue;
+		PassParameters->OutDestBuffer = DestUAV;
+
+		auto ComputeShader = GetGlobalShaderMap(Scene.GetFeatureLevel())->GetShader<FClearPageTableCS>();
+
+		PerPageShaderDispatcher.AddPass(GraphBuilder, RDG_EVENT_NAME("ClearPageTable(%s Size=%ubytes)", DestUAV->GetParent()->Name, DestUAV->GetParent()->Desc.GetSize()), ComputeShader, PassParameters);
+	};
+
+	FRDGBufferRef InitializePagesIndirectArgsRDG = CreateAndClearIndirectDispatchArgs1D(GraphBuilder, Scene.GetFeatureLevel(), TEXT("Shadow.Virtual.InitializePagesIndirectArgs"));
+
 	// Create and clear the requested page flags
 	PageRequestFlagsRDG = GraphBuilder.CreateBuffer(FRDGBufferDesc::CreateStructuredDesc(sizeof(uint32), NumPageFlagsToAllocate), TEXT("Shadow.Virtual.PageRequestFlags"));
-	AddClearUAVPass(GraphBuilder, GraphBuilder.CreateUAV(PageRequestFlagsRDG), 0);
+	AddClearPageTableUAVPass(GraphBuilder, GraphBuilder.CreateUAV(PageRequestFlagsRDG), 0);
 
 	const uint32 DirtyFlagsPerPageCount = 4;
 	DirtyPageFlagsRDG = GraphBuilder.CreateBuffer(FRDGBufferDesc::CreateStructuredDesc(sizeof(uint32), GetMaxPhysicalPages() * DirtyFlagsPerPageCount), TEXT("Shadow.Virtual.DirtyPageFlags"));
@@ -1641,11 +1834,13 @@ void FVirtualShadowMapArray::BuildPageAllocations(
 	const uint32 NumPageRectsToAllocate = FMath::RoundUpToPowerOfTwo(NumPageRects);
 	UncachedPageRectBoundsRDG = GraphBuilder.CreateBuffer(FRDGBufferDesc::CreateStructuredDesc(sizeof(FIntVector4), NumPageRectsToAllocate), TEXT("Shadow.Virtual.PageRectBounds"));
 	AllocatedPageRectBoundsRDG = GraphBuilder.CreateBuffer(FRDGBufferDesc::CreateStructuredDesc(sizeof(FIntVector4), NumPageRectsToAllocate), TEXT("Shadow.Virtual.AllocatedPageRectBounds"));
+	const uint32 NumPageRectsToClear = (GetNumFullShadowMaps() + GetNumSinglePageShadowMaps()) * FVirtualShadowMap::MaxMipLevels;
 	{
 		FInitPageRectBoundsCS::FParameters* PassParameters = GraphBuilder.AllocParameters< FInitPageRectBoundsCS::FParameters >();
 		PassParameters->VirtualShadowMap = GetUncachedUniformBuffer(GraphBuilder);
 		PassParameters->OutUncachedPageRectBounds = GraphBuilder.CreateUAV(UncachedPageRectBoundsRDG);
 		PassParameters->OutAllocatedPageRectBounds = GraphBuilder.CreateUAV(AllocatedPageRectBoundsRDG);
+		PassParameters->NumPageRectsToClear = NumPageRectsToClear;
 		PassParameters->OutPhysicalPageLists = GraphBuilder.CreateUAV(PhysicalPageListsRDG);
 
 		auto ComputeShader = GetGlobalShaderMap(Scene.GetFeatureLevel())->GetShader<FInitPageRectBoundsCS>();
@@ -1656,7 +1851,7 @@ void FVirtualShadowMapArray::BuildPageAllocations(
 			RDG_EVENT_NAME("InitPageRectBounds"),
 			ComputeShader,
 			PassParameters,
-			FIntVector(FMath::DivideAndRoundUp(NumPageRects, FInitPageRectBoundsCS::DefaultCSGroupX), 1, 1)
+			FIntVector(FMath::DivideAndRoundUp(NumPageRectsToClear, FInitPageRectBoundsCS::DefaultCSGroupX), 1, 1)
 		);
 	}
 
@@ -1915,8 +2110,8 @@ void FVirtualShadowMapArray::BuildPageAllocations(
 	PageTableRDG = GraphBuilder.CreateBuffer(FRDGBufferDesc::CreateStructuredDesc(sizeof(uint32), NumPageFlagsToAllocate), TEXT("Shadow.Virtual.PageTable"));		
 	PageFlagsRDG = GraphBuilder.CreateBuffer(FRDGBufferDesc::CreateStructuredDesc(sizeof(uint32), NumPageFlagsToAllocate), TEXT("Shadow.Virtual.PageFlags"));
 
-	AddClearUAVPass(GraphBuilder, GraphBuilder.CreateUAV(PageTableRDG), 0);
-	AddClearUAVPass(GraphBuilder, GraphBuilder.CreateUAV(PageFlagsRDG), 0);
+	AddClearPageTableUAVPass(GraphBuilder, GraphBuilder.CreateUAV(PageTableRDG), 0);
+	AddClearPageTableUAVPass(GraphBuilder, GraphBuilder.CreateUAV(PageFlagsRDG), 0);
 
 	// Update cached or newly invalidated pages with respect to the new requests
 	{	
@@ -1983,18 +2178,17 @@ void FVirtualShadowMapArray::BuildPageAllocations(
 		PassParameters->OutPageTable			= GraphBuilder.CreateUAV(PageTableRDG);
 		PassParameters->OutPageFlags			= GraphBuilder.CreateUAV(PageFlagsRDG);
 		PassParameters->OutPhysicalPageLists	= GraphBuilder.CreateUAV(PhysicalPageListsRDG);
-		PassParameters->OutPhysicalPageMetaData = GraphBuilder.CreateUAV(PhysicalPageMetaDataRDG);
+		PassParameters->OutPhysicalPageMetaData	= GraphBuilder.CreateUAV(PhysicalPageMetaDataRDG);
 
 		FAllocateNewPageMappingsCS::FPermutationDomain PermutationVector;
 		SetStatsArgsAndPermutation<FAllocateNewPageMappingsCS>(ShouldGenerateStats(), StatsBufferUAV, PassParameters, PermutationVector);
 		auto ComputeShader = Views[0].ShaderMap->GetShader<FAllocateNewPageMappingsCS>(PermutationVector);
 
-		FComputeShaderUtils::AddPass(
+		PerPageShaderDispatcher.AddPass(
 			GraphBuilder,
 			RDG_EVENT_NAME("AllocateNewPageMappings"),
 			ComputeShader,
-			PassParameters,
-			FIntVector(FMath::DivideAndRoundUp(FVirtualShadowMap::PageTableSize, FAllocateNewPageMappingsCS::DefaultCSGroupX), GetNumShadowMaps(), 1)
+			PassParameters
 		);
 	}
 
@@ -2026,12 +2220,11 @@ void FVirtualShadowMapArray::BuildPageAllocations(
 		PassParameters->OutPageTable	= GraphBuilder.CreateUAV(PageTableRDG);
 
 		auto ComputeShader = Views[0].ShaderMap->GetShader<FPropagateMappedMipsCS>();
-		FComputeShaderUtils::AddPass(
+		PerPageShaderDispatcher.AddPass(
 			GraphBuilder,
 			RDG_EVENT_NAME("PropagateMappedMips"),
 			ComputeShader,
-			PassParameters,
-			FIntVector(FMath::DivideAndRoundUp(FMath::Square(FVirtualShadowMap::Level0DimPagesXY), FPropagateMappedMipsCS::DefaultCSGroupX), GetNumFullShadowMaps(), 1)
+			PassParameters
 		);
 	}
 
@@ -2043,8 +2236,6 @@ void FVirtualShadowMapArray::BuildPageAllocations(
 		// Note: We use GetTotalAllocatedPhysicalPages() to size the buffer as the selection shader emits both static/dynamic pages separately when enabled.
 		FRDGBufferRef PhysicalPagesToInitializeRDG = GraphBuilder.CreateBuffer(FRDGBufferDesc::CreateStructuredDesc(sizeof(int32), GetTotalAllocatedPhysicalPages() + 1), TEXT("Shadow.Virtual.PhysicalPagesToInitialize"));
 
-		// 1. Initialize the indirect args buffer
-		FRDGBufferRef InitializePagesIndirectArgsRDG = CreateAndClearIndirectDispatchArgs1D(GraphBuilder, Scene.GetFeatureLevel(), TEXT("Shadow.Virtual.InitializePagesIndirectArgs"));
 		// 2. Filter the relevant physical pages and set up the indirect args
 		{
 			FSelectPagesToInitializeCS::FParameters* PassParameters = GraphBuilder.AllocParameters<FSelectPagesToInitializeCS::FParameters>();
diff --git a/Engine/Source/Runtime/Renderer/Private/VirtualShadowMaps/VirtualShadowMapArray.h b/Engine/Source/Runtime/Renderer/Private/VirtualShadowMaps/VirtualShadowMapArray.h
--- a/Engine/Source/Runtime/Renderer/Private/VirtualShadowMaps/VirtualShadowMapArray.h
+++ b/Engine/Source/Runtime/Renderer/Private/VirtualShadowMaps/VirtualShadowMapArray.h
@@ -38,19 +38,9 @@ namespace Nanite
 	struct FRasterResults;
 }
 
-// TODO: does this exist?
-constexpr uint32 ILog2Const(uint32 n)
+inline bool IsSinglePageVirtualShadowMap(int32 VirtualShadowMapId)
 {
-	return (n > 1) ? 1 + ILog2Const(n / 2) : 0;
-}
-
-// See CalcLevelOffsets in PageAccessCommon.ush for some details on this logic
-constexpr uint32 CalcVirtualShadowMapLevelOffsets(uint32 Level, uint32 Log2Level0DimPagesXY)
-{
-	uint32 NumBits = Level << 1;
-	uint32 StartBit = (2U * Log2Level0DimPagesXY + 2U) - NumBits;
-	uint32 Mask = ((1U << NumBits) - 1U) << StartBit;
-	return 0x55555555U & Mask;
+	return VirtualShadowMapId < int32(VSM_MAX_SINGLE_PAGE_SHADOW_MAPS);
 }
 
 bool DoesVSMWantFroxels(EShaderPlatform ShaderPlatform);
@@ -61,25 +51,22 @@ public:
 	// PageSize * Level0DimPagesXY defines the virtual address space, e.g., 128x128 = 16k
 
 	// 128x128 = 16k
-	static constexpr uint32 PageSize = 128U;
-	static constexpr uint32 Level0DimPagesXY = 128U;
-
-	static constexpr uint32 PageSizeMask = PageSize - 1U;
-	static constexpr uint32 Log2PageSize = ILog2Const(PageSize);
-	static constexpr uint32 Log2Level0DimPagesXY = ILog2Const(Level0DimPagesXY);
-	static constexpr uint32 MaxMipLevels = Log2Level0DimPagesXY + 1U;
+	static constexpr uint32 PageSize =  VSM_PAGE_SIZE;
+	static constexpr uint32 PageSizeMask =  VSM_PAGE_SIZE_MASK;
+	static constexpr uint32 Log2PageSize =  VSM_LOG2_PAGE_SIZE;
+	static constexpr uint32 Level0DimPagesXY =  VSM_LEVEL0_DIM_PAGES_XY;
+	static constexpr uint32 Log2Level0DimPagesXY =  VSM_LOG2_LEVEL0_DIM_PAGES_XY;
+	static constexpr uint32 MaxMipLevels =  VSM_MAX_MIP_LEVELS;
+	static constexpr uint32 VirtualMaxResolutionXY =  VSM_VIRTUAL_MAX_RESOLUTION_XY;
+	static constexpr uint32 RasterWindowPages = VSM_RASTER_WINDOW_PAGES;
+	static constexpr uint32 PageTableSize =  VSM_PAGE_TABLE_SIZE;
 
-	static constexpr uint32 PageTableSize = CalcVirtualShadowMapLevelOffsets(MaxMipLevels, Log2Level0DimPagesXY);
-
-	static constexpr uint32 VirtualMaxResolutionXY = Level0DimPagesXY * PageSize;
-	
 	static constexpr uint32 PhysicalPageAddressBits = 16U;
 	static constexpr uint32 MaxPhysicalTextureDimPages = 1U << PhysicalPageAddressBits;
 	static constexpr uint32 MaxPhysicalTextureDimTexels = MaxPhysicalTextureDimPages * PageSize;
 
 	static constexpr uint32 NumHZBLevels = Log2PageSize;
 
-	static constexpr uint32 RasterWindowPages = 4u;
 	
 	static_assert(MaxMipLevels <= 8, ">8 mips requires more PageFlags bits. See VSM_PAGE_FLAGS_BITS_PER_HMIP in PageAccessCommon.ush");
 
@@ -120,8 +107,8 @@ struct FVirtualShadowMapProjectionShaderData
 	float ClipmapLevelWPODistanceDisableThresholdSquared = 0.0f;
 	float TexelDitherScale;
 	
-	// Seems the FMatrix forces 16-byte alignment
-	float Padding;
+	uint32 MinMipLevel = 0u;
+	// Note: Seems the FMatrix forces 16-byte alignment so pad as needed.
 };
 static_assert(sizeof(FVirtualShadowMapProjectionShaderData) == (16*18), "FVirtualShadowMapProjectionShaderData does not match size in shader. See VirtualShadowMapProjectionStructs.ush.");
 
@@ -459,8 +446,6 @@ private:
 		const TConstArrayView<FViewInfo> &Views,
 		const TConstArrayView<FVisibleLightInfo>& VisibleLightInfos);
 
-	void UploadProjectionData(FRDGBuilder& GraphBuilder);
-
 	void AppendPhysicalPageList(FRDGBuilder& GraphBuilder, bool bEmptyToAvailable);
 
 	uint32 AddRenderViews(const TSharedPtr<FVirtualShadowMapClipmap>& Clipmap, const FViewInfo* CullingView, float LODScaleFactor, bool bSetHzbParams, bool bUpdateHZBMetaData, TArray<Nanite::FPackedView, SceneRenderingAllocator>& OutVirtualShadowViews);
diff --git a/Engine/Source/Runtime/Renderer/Private/VirtualShadowMaps/VirtualShadowMapCacheManager.cpp b/Engine/Source/Runtime/Renderer/Private/VirtualShadowMaps/VirtualShadowMapCacheManager.cpp
--- a/Engine/Source/Runtime/Renderer/Private/VirtualShadowMaps/VirtualShadowMapCacheManager.cpp
+++ b/Engine/Source/Runtime/Renderer/Private/VirtualShadowMaps/VirtualShadowMapCacheManager.cpp
@@ -1117,22 +1117,6 @@ void FVirtualShadowMapArrayCacheManager::UpdateUnreferencedCacheEntries(
 	}
 }
 
-void FVirtualShadowMapArrayCacheManager::UploadProjectionData(FRDGScatterUploadBuffer& Uploader) const
-{
-	// If we get here, this should be non-empty, otherwise we will upload nothing and the destination buffer
-	// will end up as not having been written in RDG, which makes it upset even if it will never get referenced on the GPU.
-	check(!CacheEntries.IsEmpty());
-
-	for (const auto& LightEntry : CacheEntries)
-	{
-		TSharedPtr<FVirtualShadowMapPerLightCacheEntry> CacheEntry = LightEntry.Value;
-		for (const auto& Entry : CacheEntry->ShadowMapEntries)
-		{
-			Uploader.Add(Entry.CurrentVirtualShadowMapId, &Entry.ProjectionData);
-		}
-	}
-}
-
 class FVirtualSmCopyStatsCS : public FGlobalShader
 {
 	DECLARE_GLOBAL_SHADER(FVirtualSmCopyStatsCS);
diff --git a/Engine/Source/Runtime/Renderer/Private/VirtualShadowMaps/VirtualShadowMapCacheManager.h b/Engine/Source/Runtime/Renderer/Private/VirtualShadowMaps/VirtualShadowMapCacheManager.h
--- a/Engine/Source/Runtime/Renderer/Private/VirtualShadowMaps/VirtualShadowMapCacheManager.h
+++ b/Engine/Source/Runtime/Renderer/Private/VirtualShadowMaps/VirtualShadowMapCacheManager.h
@@ -259,9 +259,6 @@ public:
 	 */
 	void UpdateUnreferencedCacheEntries(FVirtualShadowMapArray& VirtualShadowMapArray);
 
-	// Must be called *after* calling UpdateUnreferencedCacheEntries - TODO: Perhaps merge the two to enforce
-	void UploadProjectionData(FRDGScatterUploadBuffer& Uploader) const;
-
 	/**
 	* Call at end of frame to extract resouces from the virtual SM array to preserve to next frame.
 	* 
