Fixed Mac only CIS compile error.

#rnx

diff --git a/Engine/Source/Runtime/Renderer/Private/VirtualShadowMaps/VirtualShadowMapArray.cpp b/Engine/Source/Runtime/Renderer/Private/VirtualShadowMaps/VirtualShadowMapArray.cpp
--- a/Engine/Source/Runtime/Renderer/Private/VirtualShadowMaps/VirtualShadowMapArray.cpp
+++ b/Engine/Source/Runtime/Renderer/Private/VirtualShadowMaps/VirtualShadowMapArray.cpp
@@ -1031,6 +1031,15 @@ IMPLEMENT_GLOBAL_SHADER(FUpdatePhysicalPages, "/Engine/Private/VirtualShadowMaps
 class FVirtualShadowMapPerPageShader : public FVirtualShadowMapPageManagementShader
 {
 public:
+	FVirtualShadowMapPerPageShader()
+	{
+	}
+
+	FVirtualShadowMapPerPageShader(const ShaderMetaType::CompiledShaderInitializerType& Initializer)
+		: FVirtualShadowMapPageManagementShader(Initializer)
+	{
+	}
+
 	static constexpr int32 ThreadGroupSizeX = 64;
 
 	BEGIN_SHADER_PARAMETER_STRUCT(FParameters, )
