Fix ShaderComplexity View crash in ES31 in Editor when SecondayUpscale is active

diff --git a/Engine/Source/Runtime/Renderer/Private/PostProcess/PostProcessing.cpp b/Engine/Source/Runtime/Renderer/Private/PostProcess/PostProcessing.cpp
--- a/Engine/Source/Runtime/Renderer/Private/PostProcess/PostProcessing.cpp
+++ b/Engine/Source/Runtime/Renderer/Private/PostProcess/PostProcessing.cpp
@@ -3040,7 +3040,15 @@ void AddMobilePostProcessingPasses(FRDGBuilder& GraphBuilder, FScene* Scene, con
 				View.UnscaledViewRect.Width(), View.UnscaledViewRect.Height());
 
 			SceneColor = CustomUpscaler->AddPasses(GraphBuilder, View, PassInputs);
-			check(SceneColor == ViewFamilyOutput);
+
+			if (PassSequence.IsLastPass(EPass::SecondaryUpscale))
+			{
+				check(SceneColor == ViewFamilyOutput);
+			}
+			else
+			{
+				check(SceneColor.ViewRect.Size() == View.UnscaledViewRect.Size());
+			}
 		}
 	}
 
