Change the default value of r.SceneCulling.TreatInstancedDynamicAsUnCullable to "off" because of licensee feedback. Having the dynamic treated as uncullable can lead to GPU perf problems as the uncullable path is not load balanced. This is far more difficult to diagnose than the potential CPU overhead of updating the grid for the dynamic instances, making it a better default to turn this off.

diff --git a/Engine/Source/Runtime/Renderer/Private/SceneCulling/SceneCulling.cpp b/Engine/Source/Runtime/Renderer/Private/SceneCulling/SceneCulling.cpp
--- a/Engine/Source/Runtime/Renderer/Private/SceneCulling/SceneCulling.cpp
+++ b/Engine/Source/Runtime/Renderer/Private/SceneCulling/SceneCulling.cpp
@@ -150,8 +150,8 @@ static TAutoConsoleVariable<float> CVarSceneCullingMaxCellSize(
 
 static TAutoConsoleVariable<int32> CVarTreatDynamicInstancedAsUncullable(
 	TEXT("r.SceneCulling.TreatInstancedDynamicAsUnCullable"), 
-	1, 
-	TEXT("If this is turned on (default), dynamic primitives with instances are treated as uncullable (not put into the hierarchy and instead brute-forced on the GPU).")
+	0, 
+	TEXT("If this is turned on, dynamic primitives with instances are treated as uncullable (not put into the hierarchy and instead brute-forced on the GPU).")
 	TEXT("  This significantly reduces the hierarchy update cost on the CPU and for scenes with a large proportion of static elements, does not increase the GPU cost."),
 	ECVF_RenderThreadSafe);
 
