Fix shader compilation error with alpha invert pass

#nopreflight

diff --git a/Engine/Shaders/Private/PostProcessTonemap.usf b/Engine/Shaders/Private/PostProcessTonemap.usf
--- a/Engine/Shaders/Private/PostProcessTonemap.usf
+++ b/Engine/Shaders/Private/PostProcessTonemap.usf
@@ -764,10 +764,10 @@ void AlphaInvert_MainPS(
 )
 {
 	uint ArrayIndex = 0;
-#if MOBILE_MULTI_VIEW
-	ArrayIndex = MultiViewId;
-#elif MOBILE_MULTI_VIEW_FALLBACK
+#if MOBILE_MULTI_VIEW_FALLBACK
 	ArrayIndex = EyeIndex;
+#elif MOBILE_MULTI_VIEW
+	ArrayIndex = MultiViewId;
 #else
 #endif
 	
diff --git a/Engine/Source/Runtime/Renderer/Private/PostProcess/PostProcessTonemap.cpp b/Engine/Source/Runtime/Renderer/Private/PostProcess/PostProcessTonemap.cpp
--- a/Engine/Source/Runtime/Renderer/Private/PostProcess/PostProcessTonemap.cpp
+++ b/Engine/Source/Runtime/Renderer/Private/PostProcess/PostProcessTonemap.cpp
@@ -1193,7 +1193,7 @@ void RenderAlphaInvertPass(FRHICommandList& RHICmdList, const FViewInfo& View, F
 {
 	// Part of scene rendering pass
 	check(RHICmdList.IsInsideRenderPass());
-	SCOPED_DRAW_EVENT(RHICmdList, MobileTonemapSubpass);
+	SCOPED_DRAW_EVENT(RHICmdList, AlphaInvert);
 
 	const FIntPoint TargetSize = SceneTextures.Color.Resolve->Desc.Extent;
 	
