Fix runtime generation of energy conservation table.

diff --git a/Engine/Shaders/Private/ShadingEnergyConservationTable.usf b/Engine/Shaders/Private/ShadingEnergyConservationTable.usf
--- a/Engine/Shaders/Private/ShadingEnergyConservationTable.usf
+++ b/Engine/Shaders/Private/ShadingEnergyConservationTable.usf
@@ -17,6 +17,7 @@
 
 // Number of samples used for integrating directional albedo 
 uint NumSamples;
+uint EnergyTableResolution;
 // x: directional albedo for rough diffuse
 RWTexture2D<float>  Output1Texture2D;
 // x: directional albedo of pure GGX lobe
@@ -28,7 +29,7 @@ RWTexture3D<float2> OutputTexture3D;
 [numthreads(THREADGROUPSIZE_X, THREADGROUPSIZE_Y, 1)]
 void BuildEnergyTableCS(uint3 DispatchThreadId : SV_DispatchThreadID)
 {
-	const float3 Index = (float3(DispatchThreadId) + 0.5) / float(ENERGY_TABLE_RESOLUTION);
+	const float3 Index = (float3(DispatchThreadId) + 0.5) / float(EnergyTableResolution);
 
 	// Map from thread ID to the exact roughness/viewing angle we want to compute
 	const float NoV = max(Index.x, 1e-8);
diff --git a/Engine/Shaders/Shared/PathTracingDefinitions.h b/Engine/Shaders/Shared/PathTracingDefinitions.h
--- a/Engine/Shaders/Shared/PathTracingDefinitions.h
+++ b/Engine/Shaders/Shared/PathTracingDefinitions.h
@@ -47,9 +47,6 @@
 #define PATHTRACER_LIGHT_GRID_SINGULAR_MASK					0x80000000u
 #define PATHTRACER_LIGHT_GRID_LIGHT_COUNT_MASK				0x7FFFFFFFu
 
-// Constants for the energy conservation texture sizes
-#define PATHTRACER_ENERGY_TABLE_RESOLUTION				32
-
 // Constants related to volumetric support
 #define VOLUMEID_ATMOSPHERE				0
 #define VOLUMEID_CLOUDS					1
diff --git a/Engine/Source/Runtime/Renderer/Private/ShadingEnergyConservation.cpp b/Engine/Source/Runtime/Renderer/Private/ShadingEnergyConservation.cpp
--- a/Engine/Source/Runtime/Renderer/Private/ShadingEnergyConservation.cpp
+++ b/Engine/Source/Runtime/Renderer/Private/ShadingEnergyConservation.cpp
@@ -42,12 +42,24 @@ static TAutoConsoleVariable<int32> CVarShadingFurnaceTest_SampleCount(
 	TEXT("Number of sampler per pixel used for furnace tests."),
 	ECVF_RenderThreadSafe);
 
-static TAutoConsoleVariable<int32> CVarShadingFurnaceTest_TableFormat(
+static TAutoConsoleVariable<int32> CVarShadingEnergyConservation_TableFormat(
 	TEXT("r.Shading.EnergyConservation.Format"),
 	1,
 	TEXT("Energy conservation table format 0: 16bits, 1: 32bits."),
 	ECVF_RenderThreadSafe);
 
+static TAutoConsoleVariable<int32> CVarShadingEnergyConservation_TableResolution(
+	TEXT("r.Shading.EnergyConservation.Resolution"),
+	32,
+	TEXT("Energy conservation table resolution. Used only when using runtime generated tables."),
+	ECVF_RenderThreadSafe);
+
+static TAutoConsoleVariable<int32> CVarShadingEnergyConservation_RuntimeGeneration(
+	TEXT("r.Shading.EnergyConservation.RuntimeGeneration"),
+	0,
+	TEXT("Enable Energy conservation tables generation at runtime instead of relying on precomputed tables."),
+	ECVF_RenderThreadSafe);
+
 // Transition render settings that will disapear when Substrate gets enabled
 
 static TAutoConsoleVariable<int32> CVarMaterialEnergyConservation(
@@ -306,11 +318,11 @@ class FBuildShadingEnergyConservationTableCS : public FGlobalShader
 	{
 		OutEnvironment.SetDefine(TEXT("THREADGROUPSIZE_X"), FComputeShaderUtils::kGolden2DGroupSize);
 		OutEnvironment.SetDefine(TEXT("THREADGROUPSIZE_Y"), FComputeShaderUtils::kGolden2DGroupSize);
-		OutEnvironment.SetDefine(TEXT("ENERGY_TABLE_RESOLUTION"), SHADING_ENERGY_CONSERVATION_TABLE_RESOLUTION);
 	}
 
 	BEGIN_SHADER_PARAMETER_STRUCT(FParameters, )
 		SHADER_PARAMETER(uint32, NumSamples)
+		SHADER_PARAMETER(uint32, EnergyTableResolution)
 		SHADER_PARAMETER_RDG_TEXTURE_UAV(RWTexture2D, Output1Texture2D)
 		SHADER_PARAMETER_RDG_TEXTURE_UAV(RWTexture2D, Output2Texture2D)
 		SHADER_PARAMETER_RDG_TEXTURE_UAV(RWTexture3D, OutputTexture3D)
@@ -343,7 +355,8 @@ void Init(FRDGBuilder& GraphBuilder, FViewInfo& View)
 	if (bBindEnergyData)
 	{
 		// Change this to true in order to regenerate the energy tables, and manually copy the coefficients into ShadingEnergyConservationData.h
-		const bool bRuntimeGeneration = false;
+		const bool bRuntimeGeneration = CVarShadingEnergyConservation_RuntimeGeneration.GetValueOnRenderThread() > 0;
+		const int Size2D = FMath::Clamp(CVarShadingEnergyConservation_TableResolution.GetValueOnRenderThread(), 16, 512);
 
 		EPixelFormat Format = PF_R8G8;
 		// for low roughness we would get banding with PF_R8G8 but for low spec it could be used, for now we don't do this optimization
@@ -355,26 +368,28 @@ void Init(FRDGBuilder& GraphBuilder, FViewInfo& View)
 		const bool bRG16Supported = GPixelFormats[PF_G16R16].Supported && UE::PixelFormat::HasCapabilities(PF_G16R16, EPixelFormatCapabilities::TextureFilterable);
 		const bool bR16Supported = GPixelFormats[PF_G16].Supported && UE::PixelFormat::HasCapabilities(PF_G16, EPixelFormatCapabilities::TextureFilterable);
 
-		const int Size = SHADING_ENERGY_CONSERVATION_TABLE_RESOLUTION;
-		const EPixelFormat SpecFormat = bRuntimeGeneration && CVarShadingFurnaceTest_TableFormat.GetValueOnRenderThread() > 0 ? PF_G32R32F : (bRG16Supported ? PF_G16R16 : PF_R8G8);
+		const EPixelFormat SpecFormat = bRuntimeGeneration && CVarShadingEnergyConservation_TableFormat.GetValueOnRenderThread() > 0 ? PF_G32R32F : (bRG16Supported ? PF_G16R16 : PF_R8G8);
 		const EPixelFormat DiffFormat = bR16Supported ? PF_G16 : PF_R8;
 		const bool bBuildTable = 
+			bRuntimeGeneration ||
 			View.ViewState->ShadingEnergyConservationData.Format != SpecFormat ||
 			View.ViewState->ShadingEnergyConservationData.GGXSpecEnergyTexture == nullptr ||
 			View.ViewState->ShadingEnergyConservationData.GGXGlassEnergyTexture == nullptr ||
 			View.ViewState->ShadingEnergyConservationData.ClothEnergyTexture ==  nullptr ||
-			View.ViewState->ShadingEnergyConservationData.DiffuseEnergyTexture == nullptr;
+			View.ViewState->ShadingEnergyConservationData.DiffuseEnergyTexture == nullptr ||
+			View.ViewState->ShadingEnergyConservationData.GGXSpecEnergyTexture->GetDesc().Extent.X != Size2D;
 
 		if (bBuildTable)
 		{
 			View.ViewState->ShadingEnergyConservationData.Format = SpecFormat;
 
 			if (bRuntimeGeneration)
-			{			
-				FRDGTextureRef GGXSpecEnergyTexture	= GraphBuilder.CreateTexture(FRDGTextureDesc::Create2D(FIntPoint(Size, Size),        SpecFormat, FClearValueBinding::None, TexCreate_ShaderResource | TexCreate_UAV), TEXT("Shading.GGXSpecEnergy"),   ERDGTextureFlags::MultiFrame);
-				FRDGTextureRef GGXGlassEnergyTexture= GraphBuilder.CreateTexture(FRDGTextureDesc::Create3D(FIntVector(Size, Size, Size), SpecFormat, FClearValueBinding::None, TexCreate_ShaderResource | TexCreate_UAV), TEXT("Shading.GGXGlassEnergy"),  ERDGTextureFlags::MultiFrame);
-				FRDGTextureRef ClothEnergyTexture	= GraphBuilder.CreateTexture(FRDGTextureDesc::Create2D(FIntPoint(Size, Size),        SpecFormat, FClearValueBinding::None, TexCreate_ShaderResource | TexCreate_UAV), TEXT("Shading.ClothSpecEnergy"), ERDGTextureFlags::MultiFrame);
-				FRDGTextureRef DiffuseEnergyTexture	= GraphBuilder.CreateTexture(FRDGTextureDesc::Create2D(FIntPoint(Size, Size),        DiffFormat, FClearValueBinding::None, TexCreate_ShaderResource | TexCreate_UAV), TEXT("Shading.DiffuseEnergy"),   ERDGTextureFlags::MultiFrame);
+			{
+				const int Size3D = SHADING_ENERGY_CONSERVATION_TABLE_RESOLUTION;
+				FRDGTextureRef GGXSpecEnergyTexture	= GraphBuilder.CreateTexture(FRDGTextureDesc::Create2D(FIntPoint(Size2D, Size2D),         SpecFormat, FClearValueBinding::None, TexCreate_ShaderResource | TexCreate_UAV), TEXT("Shading.GGXSpecEnergy"),   ERDGTextureFlags::MultiFrame);
+				FRDGTextureRef GGXGlassEnergyTexture= GraphBuilder.CreateTexture(FRDGTextureDesc::Create3D(FIntVector(Size3D, Size3D, Size3D),SpecFormat, FClearValueBinding::None, TexCreate_ShaderResource | TexCreate_UAV), TEXT("Shading.GGXGlassEnergy"),  ERDGTextureFlags::MultiFrame);
+				FRDGTextureRef ClothEnergyTexture	= GraphBuilder.CreateTexture(FRDGTextureDesc::Create2D(FIntPoint(Size2D, Size2D),         SpecFormat, FClearValueBinding::None, TexCreate_ShaderResource | TexCreate_UAV), TEXT("Shading.ClothSpecEnergy"), ERDGTextureFlags::MultiFrame);
+				FRDGTextureRef DiffuseEnergyTexture	= GraphBuilder.CreateTexture(FRDGTextureDesc::Create2D(FIntPoint(Size2D, Size2D),         DiffFormat, FClearValueBinding::None, TexCreate_ShaderResource | TexCreate_UAV), TEXT("Shading.DiffuseEnergy"),   ERDGTextureFlags::MultiFrame);
 			
 				const uint32 NumSamples = 1u << 14u;
 
@@ -385,13 +400,14 @@ void Init(FRDGBuilder& GraphBuilder, FViewInfo& View)
 					TShaderMapRef<FBuildShadingEnergyConservationTableCS> ComputeShader(View.ShaderMap, PermutationVector);
 					FBuildShadingEnergyConservationTableCS::FParameters* PassParameters = GraphBuilder.AllocParameters<FBuildShadingEnergyConservationTableCS::FParameters>();
 					PassParameters->NumSamples = NumSamples;
+					PassParameters->EnergyTableResolution = Size2D;
 					PassParameters->Output2Texture2D = GraphBuilder.CreateUAV(FRDGTextureUAVDesc(GGXSpecEnergyTexture, 0));
 					FComputeShaderUtils::AddPass(
 						GraphBuilder,
 						RDG_EVENT_NAME("ShadingEnergyConservation::BuildTable(GGXSpec)"),
 						ComputeShader,
 						PassParameters,
-						FComputeShaderUtils::GetGroupCount(FIntPoint(Size, Size), FComputeShaderUtils::kGolden2DGroupSize));
+						FComputeShaderUtils::GetGroupCount(FIntPoint(Size2D, Size2D), FComputeShaderUtils::kGolden2DGroupSize));
 				}
 
 				// GGX (Reflection + Transmission) indexed by IOR
@@ -401,13 +417,14 @@ void Init(FRDGBuilder& GraphBuilder, FViewInfo& View)
 					TShaderMapRef<FBuildShadingEnergyConservationTableCS> ComputeShader(View.ShaderMap, PermutationVector);
 					FBuildShadingEnergyConservationTableCS::FParameters* PassParameters = GraphBuilder.AllocParameters<FBuildShadingEnergyConservationTableCS::FParameters>();
 					PassParameters->NumSamples = NumSamples;
+					PassParameters->EnergyTableResolution = Size3D;
 					PassParameters->OutputTexture3D = GraphBuilder.CreateUAV(FRDGTextureUAVDesc(GGXGlassEnergyTexture, 0));
 					FComputeShaderUtils::AddPass(
 						GraphBuilder,
 						RDG_EVENT_NAME("ShadingEnergyConservation::BuildTable(GGXGlass)"),
 						ComputeShader,
 						PassParameters,
-						FComputeShaderUtils::GetGroupCount(FIntVector(Size, Size, Size), FIntVector(FComputeShaderUtils::kGolden2DGroupSize, FComputeShaderUtils::kGolden2DGroupSize, 1)));
+						FComputeShaderUtils::GetGroupCount(FIntVector(Size3D, Size3D, Size3D), FIntVector(FComputeShaderUtils::kGolden2DGroupSize, FComputeShaderUtils::kGolden2DGroupSize, 1)));
 				}
 
 				// Cloth
@@ -417,13 +434,14 @@ void Init(FRDGBuilder& GraphBuilder, FViewInfo& View)
 					TShaderMapRef<FBuildShadingEnergyConservationTableCS> ComputeShader(View.ShaderMap, PermutationVector);
 					FBuildShadingEnergyConservationTableCS::FParameters* PassParameters = GraphBuilder.AllocParameters<FBuildShadingEnergyConservationTableCS::FParameters>();
 					PassParameters->NumSamples = NumSamples;
+					PassParameters->EnergyTableResolution = Size2D;
 					PassParameters->Output2Texture2D = GraphBuilder.CreateUAV(FRDGTextureUAVDesc(ClothEnergyTexture, 0));
 					FComputeShaderUtils::AddPass(
 						GraphBuilder,
 						RDG_EVENT_NAME("ShadingEnergyConservation::BuildTable(Cloth)"),
 						ComputeShader,
 						PassParameters,
-						FComputeShaderUtils::GetGroupCount(FIntPoint(Size, Size), FComputeShaderUtils::kGolden2DGroupSize));
+						FComputeShaderUtils::GetGroupCount(FIntPoint(Size2D, Size2D), FComputeShaderUtils::kGolden2DGroupSize));
 				}
 
 				// Diffuse
@@ -433,19 +451,20 @@ void Init(FRDGBuilder& GraphBuilder, FViewInfo& View)
 					TShaderMapRef<FBuildShadingEnergyConservationTableCS> ComputeShader(View.ShaderMap, PermutationVector);
 					FBuildShadingEnergyConservationTableCS::FParameters* PassParameters = GraphBuilder.AllocParameters<FBuildShadingEnergyConservationTableCS::FParameters>();
 					PassParameters->NumSamples = NumSamples;
+					PassParameters->EnergyTableResolution = Size2D;
 					PassParameters->Output1Texture2D = GraphBuilder.CreateUAV(FRDGTextureUAVDesc(DiffuseEnergyTexture, 0));
 					FComputeShaderUtils::AddPass(
 						GraphBuilder,
 						RDG_EVENT_NAME("ShadingEnergyConservation::BuildTable(Diffuse)"),
 						ComputeShader,
 						PassParameters,
-						FComputeShaderUtils::GetGroupCount(FIntPoint(Size, Size), FComputeShaderUtils::kGolden2DGroupSize));
+						FComputeShaderUtils::GetGroupCount(FIntPoint(Size2D, Size2D), FComputeShaderUtils::kGolden2DGroupSize));
 				}
 
-				GraphBuilder.QueueTextureExtraction(GGXSpecEnergyTexture,	&View.ViewState->ShadingEnergyConservationData.GGXSpecEnergyTexture);
-				GraphBuilder.QueueTextureExtraction(GGXGlassEnergyTexture,	&View.ViewState->ShadingEnergyConservationData.GGXGlassEnergyTexture);
-				GraphBuilder.QueueTextureExtraction(ClothEnergyTexture,		&View.ViewState->ShadingEnergyConservationData.ClothEnergyTexture);
-				GraphBuilder.QueueTextureExtraction(DiffuseEnergyTexture,	&View.ViewState->ShadingEnergyConservationData.DiffuseEnergyTexture);
+				View.ViewState->ShadingEnergyConservationData.GGXSpecEnergyTexture  = GraphBuilder.ConvertToExternalTexture(GGXSpecEnergyTexture);
+				View.ViewState->ShadingEnergyConservationData.GGXGlassEnergyTexture = GraphBuilder.ConvertToExternalTexture(GGXGlassEnergyTexture);
+				View.ViewState->ShadingEnergyConservationData.ClothEnergyTexture    = GraphBuilder.ConvertToExternalTexture(ClothEnergyTexture);
+				View.ViewState->ShadingEnergyConservationData.DiffuseEnergyTexture  = GraphBuilder.ConvertToExternalTexture(DiffuseEnergyTexture);
 			}
 			else
 			{
