[Ray Tracing Instance Gathering Refactor] Don't store UsedCoarseMeshStreamingHandles in FRayTracingScene.
- Add from TaskData directly to FCoarseMeshStreamingManager instead.

diff --git a/Engine/Source/Runtime/Renderer/Private/RayTracing/RayTracing.cpp b/Engine/Source/Runtime/Renderer/Private/RayTracing/RayTracing.cpp
--- a/Engine/Source/Runtime/Renderer/Private/RayTracing/RayTracing.cpp
+++ b/Engine/Source/Runtime/Renderer/Private/RayTracing/RayTracing.cpp
@@ -1680,14 +1680,11 @@ namespace RayTracing
 
 		RayTracingSBT.ResetDynamicAllocationData();
 
-		// Consume output of the relevant primitive gathering task
-		RayTracingScene.UsedCoarseMeshStreamingHandles = MoveTemp(TaskData.UsedCoarseMeshStreamingHandles);
-
 		// Inform the coarse mesh streaming manager about all the used streamable render assets in the scene
 		Nanite::FCoarseMeshStreamingManager* CoarseMeshSM = IStreamingManager::Get().GetNaniteCoarseMeshStreamingManager();
 		if (CoarseMeshSM)
 		{
-			CoarseMeshSM->AddUsedStreamingHandles(RayTracingScene.UsedCoarseMeshStreamingHandles);
+			CoarseMeshSM->AddUsedStreamingHandles(TaskData.UsedCoarseMeshStreamingHandles);
 		}
 
 		INC_DWORD_STAT_BY(STAT_VisibleRayTracingPrimitives, TaskData.StaticPrimitives.Num() + TaskData.CachedStaticPrimitives.Num() + TaskData.DynamicPrimitivesIndices.Num());
diff --git a/Engine/Source/Runtime/Renderer/Private/RayTracing/RayTracingScene.cpp b/Engine/Source/Runtime/Renderer/Private/RayTracing/RayTracingScene.cpp
--- a/Engine/Source/Runtime/Renderer/Private/RayTracing/RayTracingScene.cpp
+++ b/Engine/Source/Runtime/Renderer/Private/RayTracing/RayTracingScene.cpp
@@ -604,7 +604,6 @@ void FRayTracingScene::Reset(bool bInInstanceDebugDataEnabled)
 	CallableCommands.Reset();
 	UniformBuffers.Reset();
 	GeometriesToBuild.Reset();
-	UsedCoarseMeshStreamingHandles.Reset();
 
 	NumSegments = 0;
 	NumMissShaderSlots = 1;
@@ -632,7 +631,6 @@ void FRayTracingScene::EndFrame()
 		CallableCommands.Empty();
 		UniformBuffers.Empty();
 		GeometriesToBuild.Empty();
-		UsedCoarseMeshStreamingHandles.Empty();
 
 #if STATS
 		ReleaseReadbackBuffers();
diff --git a/Engine/Source/Runtime/Renderer/Private/RayTracing/RayTracingScene.h b/Engine/Source/Runtime/Renderer/Private/RayTracing/RayTracingScene.h
--- a/Engine/Source/Runtime/Renderer/Private/RayTracing/RayTracingScene.h
+++ b/Engine/Source/Runtime/Renderer/Private/RayTracing/RayTracingScene.h
@@ -160,9 +160,6 @@ public:
 	// Geometries which still have a pending build request but are used this frame and require a force build.
 	TArray<const FRayTracingGeometry*> GeometriesToBuild;
 
-	// Used coarse mesh streaming handles during the last TLAS build
-	TArray<Nanite::CoarseMeshStreamingHandle> UsedCoarseMeshStreamingHandles;
-
 	bool bNeedsDebugInstanceGPUSceneIndexBuffer = false;
 
 	// Used for transforming to translated world space in which TLAS was built.
