Fixed two cases of > comparison against ShowFlags which fails to compile if the flag is a bool (in the unlikely event PLATFORM_USE_SHOWFLAGS_ALWAYS_BITFIELD is not 1, or for debugging purposes).

#rnx

diff --git a/Engine/Source/Runtime/Renderer/Private/LightFunctionAtlas.cpp b/Engine/Source/Runtime/Renderer/Private/LightFunctionAtlas.cpp
--- a/Engine/Source/Runtime/Renderer/Private/LightFunctionAtlas.cpp
+++ b/Engine/Source/Runtime/Renderer/Private/LightFunctionAtlas.cpp
@@ -276,7 +276,7 @@ void FLightFunctionAtlas::BeginSceneFrame(const FViewFamilyInfo& ViewFamily, TAr
 	AtlasSetup = GetLightFunctionAtlasSetup();
 
 	// Now lets check if we need to generate the atlas for this frame
-	bLightFunctionAtlasEnabled = CVarLightFunctionAtlas.GetValueOnRenderThread() > 0 && ViewFamily.EngineShowFlags.LightFunctions > 0;
+	bLightFunctionAtlasEnabled = CVarLightFunctionAtlas.GetValueOnRenderThread() > 0 && ViewFamily.EngineShowFlags.LightFunctions;
 #if !UE_BUILD_SHIPPING
 	LightCountSkippedDueToMissingAtlasSlot = 0;
 	SkippedLightFunctionsSet.Empty();
diff --git a/Engine/Source/Runtime/Renderer/Private/SkyAtmosphereRendering.cpp b/Engine/Source/Runtime/Renderer/Private/SkyAtmosphereRendering.cpp
--- a/Engine/Source/Runtime/Renderer/Private/SkyAtmosphereRendering.cpp
+++ b/Engine/Source/Runtime/Renderer/Private/SkyAtmosphereRendering.cpp
@@ -1293,7 +1293,7 @@ static void SetupSkyAtmosphereInternalCommonParameters(
 	const FSkyAtmosphereSceneProxy& SkyAtmosphereSceneProxy = SkyInfo.GetSkyAtmosphereSceneProxy();
 	InternalCommonParameters.SkyLuminanceFactor = FVector3f(SkyAtmosphereSceneProxy.GetSkyLuminanceFactor());
 	InternalCommonParameters.AerialPespectiveViewDistanceScale = SkyAtmosphereSceneProxy.GetAerialPespectiveViewDistanceScale();
-	InternalCommonParameters.FogShowFlagFactor = ViewFamily.EngineShowFlags.Fog > 0 ? 1.0f : 0.0f;
+	InternalCommonParameters.FogShowFlagFactor = ViewFamily.EngineShowFlags.Fog ? 1.0f : 0.0f;
 
 	auto ValidateDistanceValue = [](float& Value)
 	{
