- Fix for Niagara rendering that requires PreRender to be called when not outputting final scene color
- This occurs when using depth / normal optimized rendering modes we would skip a critical path resulting in corruption / GPU hang

#rnx
#tests Created a level with a scene capture, set capture mode to normal enable base pass optimization then use a GPU ribbon renderer, or stateless GPU path

diff --git a/Engine/Source/Runtime/Renderer/Private/DeferredShadingRenderer.cpp b/Engine/Source/Runtime/Renderer/Private/DeferredShadingRenderer.cpp
--- a/Engine/Source/Runtime/Renderer/Private/DeferredShadingRenderer.cpp
+++ b/Engine/Source/Runtime/Renderer/Private/DeferredShadingRenderer.cpp
@@ -1707,18 +1707,15 @@ void FDeferredShadingSceneRenderer::Render(FRDGBuilder& GraphBuilder)
 
 	InitViewTaskDatas.VisibilityTaskData->FinishGatherDynamicMeshElements(BasePassDepthStencilAccess, InstanceCullingManager, VirtualTextureUpdater.Get());
 
-	if (RendererOutput == ERendererOutput::FinalSceneColor)
+	// Notify the FX system that the scene is about to be rendered.
+	// TODO: These should probably be moved to scene extensions
+	if (FXSystem && Views.IsValidIndex(0))
 	{
-		// Notify the FX system that the scene is about to be rendered.
-		// TODO: These should probably be moved to scene extensions
-		if (FXSystem && Views.IsValidIndex(0))
+		SCOPE_CYCLE_COUNTER(STAT_FDeferredShadingSceneRenderer_FXSystem_PreRender);
+		FXSystem->PreRender(GraphBuilder, GetSceneViews(), GetSceneUniforms(), bIsFirstSceneRenderer /*bAllowGPUParticleUpdate*/);
+		if (FGPUSortManager* GPUSortManager = FXSystem->GetGPUSortManager())
 		{
-			SCOPE_CYCLE_COUNTER(STAT_FDeferredShadingSceneRenderer_FXSystem_PreRender);
-			FXSystem->PreRender(GraphBuilder, GetSceneViews(), GetSceneUniforms(), bIsFirstSceneRenderer /*bAllowGPUParticleUpdate*/);
-			if (FGPUSortManager* GPUSortManager = FXSystem->GetGPUSortManager())
-			{
-				GPUSortManager->OnPreRender(GraphBuilder);
-			}
+			GPUSortManager->OnPreRender(GraphBuilder);
 		}
 	}
 
