[HWRT] Discard dynamic ray tracing instance if all segments are translucent, but translucent geometry is excluded from ray tracing.

diff --git a/Engine/Source/Runtime/Renderer/Private/PrimitiveSceneInfo.cpp b/Engine/Source/Runtime/Renderer/Private/PrimitiveSceneInfo.cpp
--- a/Engine/Source/Runtime/Renderer/Private/PrimitiveSceneInfo.cpp
+++ b/Engine/Source/Runtime/Renderer/Private/PrimitiveSceneInfo.cpp
@@ -318,6 +318,7 @@ FPrimitiveSceneInfo::FPrimitiveSceneInfo(const FPrimitiveSceneInfoAdapter& InAda
 	bCachedRayTracingInstanceMaskAndFlagsDirty(true),
 	bCachedRayTracingInstanceAnySegmentsDecal(false),
 	bCachedRayTracingInstanceAllSegmentsDecal(false),
+	bCachedRayTracingInstanceAllSegmentsTranslucent(false),
 	CoarseMeshStreamingHandle(InAdapter.SceneProxy->GetCoarseMeshStreamingHandle()),
 	RayTracingGeometryGroupHandle(INDEX_NONE),
 #endif
@@ -1331,6 +1332,7 @@ void FPrimitiveSceneInfo::UpdateCachedRayTracingInstance(FPrimitiveSceneInfo* Sc
 		SceneInfo->CachedRayTracingInstance.Flags = ERayTracingInstanceFlags::None;
 		SceneInfo->bCachedRayTracingInstanceAnySegmentsDecal = false;
 		SceneInfo->bCachedRayTracingInstanceAllSegmentsDecal = false;
+		SceneInfo->bCachedRayTracingInstanceAllSegmentsTranslucent = false;
 
 		SceneInfo->bCachedRayTracingInstanceMaskAndFlagsDirty = true;
 	}
@@ -1381,6 +1383,7 @@ void FPrimitiveSceneInfo::UpdateCachedRayTracingInstanceMaskAndFlags(FRayTracing
 
 	bCachedRayTracingInstanceAnySegmentsDecal = InstanceMaskAndFlags.bAnySegmentsDecal;
 	bCachedRayTracingInstanceAllSegmentsDecal = InstanceMaskAndFlags.bAllSegmentsDecal;
+	bCachedRayTracingInstanceAllSegmentsTranslucent = InstanceMaskAndFlags.bAllSegmentsTranslucent;
 
 	bCachedRayTracingInstanceMaskAndFlagsDirty = false;
 }
diff --git a/Engine/Source/Runtime/Renderer/Private/RayTracing/RayTracing.cpp b/Engine/Source/Runtime/Renderer/Private/RayTracing/RayTracing.cpp
--- a/Engine/Source/Runtime/Renderer/Private/RayTracing/RayTracing.cpp
+++ b/Engine/Source/Runtime/Renderer/Private/RayTracing/RayTracing.cpp
@@ -705,6 +705,11 @@ namespace RayTracing
 						SceneInfo->UpdateCachedRayTracingInstanceMaskAndFlags(InstanceMaskAndFlags);
 					}
 
+					if (!SceneOptions.bTranslucentGeometry && SceneInfo->bCachedRayTracingInstanceAllSegmentsTranslucent)
+					{
+						continue;
+					}
+
 					const bool bNeedMainInstance = !SceneInfo->bCachedRayTracingInstanceAllSegmentsDecal;
 
 					// if primitive has mixed decal and non-decal segments we need to have two ray tracing instances
diff --git a/Engine/Source/Runtime/Renderer/Private/RayTracing/RayTracingInstanceMask.cpp b/Engine/Source/Runtime/Renderer/Private/RayTracing/RayTracingInstanceMask.cpp
--- a/Engine/Source/Runtime/Renderer/Private/RayTracing/RayTracingInstanceMask.cpp
+++ b/Engine/Source/Runtime/Renderer/Private/RayTracing/RayTracingInstanceMask.cpp
@@ -133,6 +133,7 @@ FRayTracingMaskAndFlags BuildRayTracingInstanceMaskAndFlags(TArrayView<const FMe
 	ensureMsgf(MeshBatches.Num() > 0, TEXT("You need to add MeshBatches first for instance mask and flags to build upon."));
 
 	bool bAllSegmentsOpaque = true;
+	bool bAllSegmentsTranslucent = true;
 	bool bAnySegmentsCastShadow = false;
 	bool bAllSegmentsCastShadow = true;
 	bool bAnySegmentsDecal = false;
@@ -155,6 +156,7 @@ FRayTracingMaskAndFlags BuildRayTracingInstanceMaskAndFlags(TArrayView<const FMe
 
 			Result.Mask |= BlendModeToRayTracingInstanceMask(BlendMode, bSegmentCastsShadow, MaskMode);
 			bAllSegmentsOpaque &= BlendMode == BLEND_Opaque;
+			bAllSegmentsTranslucent &= BlendMode == BLEND_Translucent;
 			bAnySegmentsCastShadow |= bSegmentCastsShadow;
 			bAllSegmentsCastShadow &= bSegmentCastsShadow;
 			bAnySegmentsDecal |= Material.IsDeferredDecal();
@@ -169,6 +171,7 @@ FRayTracingMaskAndFlags BuildRayTracingInstanceMaskAndFlags(TArrayView<const FMe
 	Result.bDoubleSided = bDoubleSided;	
 	Result.bAnySegmentsDecal = bAnySegmentsDecal;
 	Result.bAllSegmentsDecal = bAllSegmentsDecal;
+	Result.bAllSegmentsTranslucent = bAllSegmentsTranslucent;
 	Result.bReverseCulling = bAllSegmentsReverseCulling;
 
 	const bool bIsHairStrands = Result.Mask & ComputeRayTracingInstanceMask(ERayTracingInstanceMaskType::HairStrands, MaskMode);
diff --git a/Engine/Source/Runtime/Renderer/Public/PrimitiveSceneInfo.h b/Engine/Source/Runtime/Renderer/Public/PrimitiveSceneInfo.h
--- a/Engine/Source/Runtime/Renderer/Public/PrimitiveSceneInfo.h
+++ b/Engine/Source/Runtime/Renderer/Public/PrimitiveSceneInfo.h
@@ -703,6 +703,7 @@ public:
 	bool bCachedRayTracingInstanceMaskAndFlagsDirty : 1;
 	bool bCachedRayTracingInstanceAnySegmentsDecal : 1;
 	bool bCachedRayTracingInstanceAllSegmentsDecal : 1;
+	bool bCachedRayTracingInstanceAllSegmentsTranslucent : 1;
 	Nanite::CoarseMeshStreamingHandle CoarseMeshStreamingHandle;
 	RayTracing::GeometryGroupHandle RayTracingGeometryGroupHandle;
 
diff --git a/Engine/Source/Runtime/Renderer/Public/RayTracingInstance.h b/Engine/Source/Runtime/Renderer/Public/RayTracingInstance.h
--- a/Engine/Source/Runtime/Renderer/Public/RayTracingInstance.h
+++ b/Engine/Source/Runtime/Renderer/Public/RayTracingInstance.h
@@ -16,6 +16,7 @@ struct FRayTracingMaskAndFlags
 		, bReverseCulling(false)
 		, bAnySegmentsDecal(false)
 		, bAllSegmentsDecal(false)
+		, bAllSegmentsTranslucent(false)
 	{}
 
 	/** Instance mask that can be used to exclude the instance from specific effects (eg. ray traced shadows). */
@@ -33,6 +34,9 @@ struct FRayTracingMaskAndFlags
 	/** Whether any or all of the segments in the instance are decals. */
 	uint8 bAnySegmentsDecal : 1;
 	uint8 bAllSegmentsDecal : 1;
+
+	/** Whether all of the segments in the instance are translucent. */
+	uint8 bAllSegmentsTranslucent : 1;
 };
 
 enum class ERayTracingInstanceLayer : uint8
