Set r.RayTracing.Shadows.AvoidSelfIntersectionTraceDistance 1 by default so that RT Shadows work with Nanite
* 1.23ms tracing time -> 1.79ms on 2080TI
* Overall performance regression of 20% for RT Shadows, but needed for it to work

#tests Firefly editor

diff --git a/Engine/Source/Runtime/Renderer/Private/RayTracing/RayTracingShadows.cpp b/Engine/Source/Runtime/Renderer/Private/RayTracing/RayTracingShadows.cpp
--- a/Engine/Source/Runtime/Renderer/Private/RayTracing/RayTracingShadows.cpp
+++ b/Engine/Source/Runtime/Renderer/Private/RayTracing/RayTracingShadows.cpp
@@ -39,7 +39,7 @@ static FAutoConsoleVariableRef CVarRayTracingShadowsEnableMaterials(
 	ECVF_RenderThreadSafe
 );
 
-static float GRayTracingShadowsAvoidSelfIntersectionTraceDistance = 0.0f;
+static float GRayTracingShadowsAvoidSelfIntersectionTraceDistance = 1.0f;
 static FAutoConsoleVariableRef CVarRayTracingShadowsAvoidSelfIntersectionTraceDistance(
 	TEXT("r.RayTracing.Shadows.AvoidSelfIntersectionTraceDistance"),
 	GRayTracingShadowsAvoidSelfIntersectionTraceDistance,
