Fix missing env. lighting rendering in thumbnail when Substrate is enabled.

The energy conservation texture where never bound as the view didn't have any ViewState. This CL removes the need for view state, and replace it with a TGlobalResource for energy consevation resources.

diff --git a/Engine/Source/Runtime/Renderer/Private/RendererScene.cpp b/Engine/Source/Runtime/Renderer/Private/RendererScene.cpp
--- a/Engine/Source/Runtime/Renderer/Private/RendererScene.cpp
+++ b/Engine/Source/Runtime/Renderer/Private/RendererScene.cpp
@@ -1006,15 +1006,6 @@ uint64 FTemporalRenderTargetState::GetGPUSizeBytes(bool bLogSizes) const
 	return TotalSize;
 }
 
-uint64 FShadingEnergyConservationStateData::GetGPUSizeBytes(bool bLogSizes) const
-{
-	return
-		GetRenderTargetGPUSizeBytes(GGXSpecEnergyTexture, bLogSizes) +
-		GetRenderTargetGPUSizeBytes(GGXGlassEnergyTexture, bLogSizes) +
-		GetRenderTargetGPUSizeBytes(ClothEnergyTexture, bLogSizes) +
-		GetRenderTargetGPUSizeBytes(DiffuseEnergyTexture, bLogSizes);
-}
-
 uint64 FGlintShadingLUTsStateData::GetGPUSizeBytes(bool bLogSizes) const
 {
 	return GetTextureGPUSizeBytes(RHIGlintShadingLUTs, bLogSizes);
@@ -1091,7 +1082,6 @@ uint64 FSceneViewState::GetGPUSizeBytes(bool bLogSizes) const
 	TotalSize += GetBufferGPUSizeBytes(HairStrandsViewStateData.VoxelFeedbackBuffer, bLogSizes);
 	TotalSize += GetBufferGPUSizeBytes(ShaderPrintStateData.EntryBuffer, bLogSizes);
 	TotalSize += GetBufferGPUSizeBytes(ShaderPrintStateData.StateBuffer, bLogSizes);
-	TotalSize += ShadingEnergyConservationData.GetGPUSizeBytes(bLogSizes);
 	TotalSize += GlintShadingLUTsData.GetGPUSizeBytes(bLogSizes);
 
 	// Per-view Lumen scene data is stored in a map in the FScene
diff --git a/Engine/Source/Runtime/Renderer/Private/ScenePrivate.h b/Engine/Source/Runtime/Renderer/Private/ScenePrivate.h
--- a/Engine/Source/Runtime/Renderer/Private/ScenePrivate.h
+++ b/Engine/Source/Runtime/Renderer/Private/ScenePrivate.h
@@ -1140,8 +1140,6 @@ public:
 
 	FShaderPrintStateData ShaderPrintStateData;
 
-	FShadingEnergyConservationStateData ShadingEnergyConservationData;
-
 	FGlintShadingLUTsStateData GlintShadingLUTsData;
 
 	bool bLumenSceneDataAdded;
diff --git a/Engine/Source/Runtime/Renderer/Private/SceneRendering.cpp b/Engine/Source/Runtime/Renderer/Private/SceneRendering.cpp
--- a/Engine/Source/Runtime/Renderer/Private/SceneRendering.cpp
+++ b/Engine/Source/Runtime/Renderer/Private/SceneRendering.cpp
@@ -2095,22 +2095,20 @@ void FViewInfo::SetupUniformBufferParameters(
 	ViewUniformShaderParameters.HairScatteringLUTSampler = TStaticSamplerState<SF_Bilinear, AM_Clamp, AM_Clamp, AM_Clamp>::GetRHI();
 
 	// Shading energy conservation
-	ViewUniformShaderParameters.bShadingEnergyConservation = 0u;
-	ViewUniformShaderParameters.bShadingEnergyPreservation = 0u;
-	ViewUniformShaderParameters.ShadingEnergySampler = TStaticSamplerState<SF_Bilinear, AM_Clamp, AM_Clamp, AM_Clamp>::GetRHI();
-	if (ViewState)
 	{
-		ViewUniformShaderParameters.bShadingEnergyConservation		= ViewState->ShadingEnergyConservationData.bEnergyConservation ? 1u : 0u;
-		ViewUniformShaderParameters.bShadingEnergyPreservation		= ViewState->ShadingEnergyConservationData.bEnergyPreservation ? 1u : 0u;
-		ViewUniformShaderParameters.ShadingEnergyGGXSpecTexture		= ViewState->ShadingEnergyConservationData.GGXSpecEnergyTexture ? ViewState->ShadingEnergyConservationData.GGXSpecEnergyTexture->GetRHI() : nullptr;
-		ViewUniformShaderParameters.ShadingEnergyGGXGlassTexture	= ViewState->ShadingEnergyConservationData.GGXGlassEnergyTexture ? ViewState->ShadingEnergyConservationData.GGXGlassEnergyTexture->GetRHI() : nullptr;
-		ViewUniformShaderParameters.ShadingEnergyClothSpecTexture	= ViewState->ShadingEnergyConservationData.ClothEnergyTexture ? ViewState->ShadingEnergyConservationData.ClothEnergyTexture->GetRHI() : nullptr;
-		ViewUniformShaderParameters.ShadingEnergyDiffuseTexture		= ViewState->ShadingEnergyConservationData.DiffuseEnergyTexture ? ViewState->ShadingEnergyConservationData.DiffuseEnergyTexture->GetRHI() : nullptr;
+		const FShadingEnergyConservationData ShadingEnergyConservationData = ShadingEnergyConservation::GetData(*this);
+		ViewUniformShaderParameters.bShadingEnergyConservation		= ShadingEnergyConservationData.bEnergyConservation ? 1u : 0u;
+		ViewUniformShaderParameters.bShadingEnergyPreservation		= ShadingEnergyConservationData.bEnergyPreservation ? 1u : 0u;
+		ViewUniformShaderParameters.ShadingEnergyGGXSpecTexture		= ShadingEnergyConservationData.GGXSpecEnergyTexture ? ShadingEnergyConservationData.GGXSpecEnergyTexture->GetRHI() : nullptr;
+		ViewUniformShaderParameters.ShadingEnergyGGXGlassTexture	= ShadingEnergyConservationData.GGXGlassEnergyTexture ?ShadingEnergyConservationData.GGXGlassEnergyTexture->GetRHI() : nullptr;
+		ViewUniformShaderParameters.ShadingEnergyClothSpecTexture	= ShadingEnergyConservationData.ClothEnergyTexture ?   ShadingEnergyConservationData.ClothEnergyTexture->GetRHI() : nullptr;
+		ViewUniformShaderParameters.ShadingEnergyDiffuseTexture		= ShadingEnergyConservationData.DiffuseEnergyTexture ? ShadingEnergyConservationData.DiffuseEnergyTexture->GetRHI() : nullptr;
 	}
 	ViewUniformShaderParameters.ShadingEnergyGGXSpecTexture		 = OrBlack2DIfNull(ViewUniformShaderParameters.ShadingEnergyGGXSpecTexture);
 	ViewUniformShaderParameters.ShadingEnergyGGXGlassTexture	 = OrBlack3DIfNull(ViewUniformShaderParameters.ShadingEnergyGGXGlassTexture);
 	ViewUniformShaderParameters.ShadingEnergyClothSpecTexture	 = OrBlack2DIfNull(ViewUniformShaderParameters.ShadingEnergyClothSpecTexture);
 	ViewUniformShaderParameters.ShadingEnergyDiffuseTexture		 = OrBlack2DIfNull(ViewUniformShaderParameters.ShadingEnergyDiffuseTexture);
+	ViewUniformShaderParameters.ShadingEnergySampler = TStaticSamplerState<SF_Bilinear, AM_Clamp, AM_Clamp, AM_Clamp>::GetRHI();
 
 	// Glint
 	ViewUniformShaderParameters.GlintSampler = TStaticSamplerState<SF_Trilinear, AM_Clamp, AM_Clamp, AM_Clamp>::GetRHI();
diff --git a/Engine/Source/Runtime/Renderer/Private/ShadingEnergyConservation.cpp b/Engine/Source/Runtime/Renderer/Private/ShadingEnergyConservation.cpp
--- a/Engine/Source/Runtime/Renderer/Private/ShadingEnergyConservation.cpp
+++ b/Engine/Source/Runtime/Renderer/Private/ShadingEnergyConservation.cpp
@@ -72,6 +72,54 @@ static TAutoConsoleVariable<int32> CVarMaterialEnergyConservation(
 
 ////////////////////////////////////////////////////////////////////////////////////////////////////////////////////////////////////////////////
 
+struct FShadingEnergyConservationSettings
+{
+	bool bIsEnergyConservationEnabled = false;
+	bool bIsEnergyPreservationEnabled = false;
+	bool bNeedData = false;
+};
+
+static FShadingEnergyConservationSettings GetSettings(const FViewInfo& View)
+{
+	FShadingEnergyConservationSettings Out;
+
+	// Enabled based on settings
+	const bool bMaterialEnergyConservationEnabled = CVarMaterialEnergyConservation.GetValueOnRenderThread() > 0;
+	Out.bIsEnergyConservationEnabled = CVarShadingEnergyConservation.GetValueOnRenderThread() > 0;
+	Out.bIsEnergyPreservationEnabled = CVarShadingEnergyConservation_Preservation.GetValueOnRenderThread() > 0;
+
+	// Build/bind table if energy conservation is enabled or if Substrate is enabled in order to have 
+	// the correct tables built & bound. Even if we are not using energy conservation, we want to 
+	// have access to directional albedo information for env. lighting for instance)
+	Out.bNeedData = (bMaterialEnergyConservationEnabled || Substrate::IsSubstrateEnabled() || (View.Family->EngineShowFlags.PathTracing)) && (Out.bIsEnergyPreservationEnabled || Out.bIsEnergyConservationEnabled);
+
+	return Out;
+}
+
+class FShadingEnergyConservationResources : public FRenderResource
+{
+public:
+	FShadingEnergyConservationResources() : FRenderResource()
+	{}
+
+	virtual void ReleaseRHI()
+	{
+		GGXSpecEnergyTexture.SafeRelease();
+		GGXGlassEnergyTexture.SafeRelease();
+		ClothEnergyTexture.SafeRelease();
+		DiffuseEnergyTexture.SafeRelease();
+	}
+
+	EPixelFormat Format = PF_Unknown;
+	TRefCountPtr<IPooledRenderTarget> GGXSpecEnergyTexture = nullptr;
+	TRefCountPtr<IPooledRenderTarget> GGXGlassEnergyTexture = nullptr;
+	TRefCountPtr<IPooledRenderTarget> ClothEnergyTexture = nullptr;
+	TRefCountPtr<IPooledRenderTarget> DiffuseEnergyTexture = nullptr;
+};
+
+/** The global energy conservation data.textures used for scene rendering. */
+TGlobalResource<FShadingEnergyConservationResources> GShadingEnergyConservationResources;
+
 namespace ShadingEnergyConservationData
 {
 	static TRefCountPtr<IPooledRenderTarget> CreateTexture2D(FRHICommandListImmediate& RHICmdList, TObjectPtr<class UTexture2D>& InCPUTexture, EPixelFormat InFormat, const TCHAR* InName)
@@ -341,19 +389,12 @@ bool IsEnable()
 	return CVarShadingEnergyConservation.GetValueOnAnyThread() > 0;
 }
 
-void Init(FRDGBuilder& GraphBuilder, FViewInfo& View)
+void Init(FRDGBuilder& GraphBuilder, const FViewInfo& View)
 {
-	// Enabled based on settings
-	const bool bMaterialEnergyConservationEnabled = CVarMaterialEnergyConservation.GetValueOnRenderThread() > 0;
-	const bool bIsEnergyConservationEnabled = CVarShadingEnergyConservation.GetValueOnRenderThread() > 0;
-	const bool bIsEnergyPreservationEnabled = CVarShadingEnergyConservation_Preservation.GetValueOnRenderThread() > 0;	
-
-	// Build/bind table if energy conservation is enabled or if Substrate is enabled in order to have 
-	// the correct tables built & bound. Even if we are not using energy conservation, we want to 
-	// have access to directional albedo information for env. lighting for instance)
-	const bool bBindEnergyData = (View.ViewState != nullptr) && (bMaterialEnergyConservationEnabled || Substrate::IsSubstrateEnabled() || (View.Family->EngineShowFlags.PathTracing)) && (bIsEnergyPreservationEnabled || bIsEnergyConservationEnabled);
-	if (bBindEnergyData)
+	if (GetSettings(View).bNeedData)
 	{
+		FShadingEnergyConservationResources& Out = GShadingEnergyConservationResources;
+
 		// Change this to true in order to regenerate the energy tables, and manually copy the coefficients into ShadingEnergyConservationData.h
 		const bool bRuntimeGeneration = CVarShadingEnergyConservation_RuntimeGeneration.GetValueOnRenderThread() > 0;
 		const int Size2D = FMath::Clamp(CVarShadingEnergyConservation_TableResolution.GetValueOnRenderThread(), 16, 512);
@@ -372,16 +413,16 @@ void Init(FRDGBuilder& GraphBuilder, FViewInfo& View)
 		const EPixelFormat DiffFormat = bR16Supported ? PF_G16 : PF_R8;
 		const bool bBuildTable = 
 			bRuntimeGeneration ||
-			View.ViewState->ShadingEnergyConservationData.Format != SpecFormat ||
-			View.ViewState->ShadingEnergyConservationData.GGXSpecEnergyTexture == nullptr ||
-			View.ViewState->ShadingEnergyConservationData.GGXGlassEnergyTexture == nullptr ||
-			View.ViewState->ShadingEnergyConservationData.ClothEnergyTexture ==  nullptr ||
-			View.ViewState->ShadingEnergyConservationData.DiffuseEnergyTexture == nullptr ||
-			View.ViewState->ShadingEnergyConservationData.GGXSpecEnergyTexture->GetDesc().Extent.X != Size2D;
+			Out.Format != SpecFormat ||
+			Out.GGXSpecEnergyTexture == nullptr ||
+			Out.GGXGlassEnergyTexture == nullptr ||
+			Out.ClothEnergyTexture ==  nullptr ||
+			Out.DiffuseEnergyTexture == nullptr ||
+			Out.GGXSpecEnergyTexture->GetDesc().Extent.X != Size2D;
 
 		if (bBuildTable)
 		{
-			View.ViewState->ShadingEnergyConservationData.Format = SpecFormat;
+			Out.Format = SpecFormat;
 
 			if (bRuntimeGeneration)
 			{
@@ -461,10 +502,10 @@ void Init(FRDGBuilder& GraphBuilder, FViewInfo& View)
 						FComputeShaderUtils::GetGroupCount(FIntPoint(Size2D, Size2D), FComputeShaderUtils::kGolden2DGroupSize));
 				}
 
-				View.ViewState->ShadingEnergyConservationData.GGXSpecEnergyTexture  = GraphBuilder.ConvertToExternalTexture(GGXSpecEnergyTexture);
-				View.ViewState->ShadingEnergyConservationData.GGXGlassEnergyTexture = GraphBuilder.ConvertToExternalTexture(GGXGlassEnergyTexture);
-				View.ViewState->ShadingEnergyConservationData.ClothEnergyTexture    = GraphBuilder.ConvertToExternalTexture(ClothEnergyTexture);
-				View.ViewState->ShadingEnergyConservationData.DiffuseEnergyTexture  = GraphBuilder.ConvertToExternalTexture(DiffuseEnergyTexture);
+				Out.GGXSpecEnergyTexture  = GraphBuilder.ConvertToExternalTexture(GGXSpecEnergyTexture);
+				Out.GGXGlassEnergyTexture = GraphBuilder.ConvertToExternalTexture(GGXGlassEnergyTexture);
+				Out.ClothEnergyTexture    = GraphBuilder.ConvertToExternalTexture(ClothEnergyTexture);
+				Out.DiffuseEnergyTexture  = GraphBuilder.ConvertToExternalTexture(DiffuseEnergyTexture);
 			}
 			else
 			{
@@ -472,28 +513,34 @@ void Init(FRDGBuilder& GraphBuilder, FViewInfo& View)
 				check(SpecFormat == PF_G16R16 || SpecFormat == PF_R8G8);
 				check(DiffFormat == PF_G16 || DiffFormat == PF_R8);
 
-				View.ViewState->ShadingEnergyConservationData.GGXSpecEnergyTexture  = ShadingEnergyConservationData::CreateTexture2D(GraphBuilder.RHICmdList, GEngine->GGXReflectionEnergyTexture,   SpecFormat, TEXT("Shading.GGXReflectionEnergy"));
-				View.ViewState->ShadingEnergyConservationData.GGXGlassEnergyTexture = ShadingEnergyConservationData::CreateTexture3D(GraphBuilder.RHICmdList, GEngine->GGXTransmissionEnergyTexture, SpecFormat, TEXT("Shading.GGXTransmissionEnergy"));
-				View.ViewState->ShadingEnergyConservationData.ClothEnergyTexture    = ShadingEnergyConservationData::CreateTexture2D(GraphBuilder.RHICmdList, GEngine->SheenEnergyTexture,           SpecFormat, TEXT("Shading.SheenEnergy"));
-				View.ViewState->ShadingEnergyConservationData.DiffuseEnergyTexture  = ShadingEnergyConservationData::CreateTexture2D(GraphBuilder.RHICmdList, GEngine->DiffuseEnergyTexture,         DiffFormat, TEXT("Shading.DiffuseEnergy"));
+				Out.GGXSpecEnergyTexture  = ShadingEnergyConservationData::CreateTexture2D(GraphBuilder.RHICmdList, GEngine->GGXReflectionEnergyTexture,   SpecFormat, TEXT("Shading.GGXReflectionEnergy"));
+				Out.GGXGlassEnergyTexture = ShadingEnergyConservationData::CreateTexture3D(GraphBuilder.RHICmdList, GEngine->GGXTransmissionEnergyTexture, SpecFormat, TEXT("Shading.GGXTransmissionEnergy"));
+				Out.ClothEnergyTexture    = ShadingEnergyConservationData::CreateTexture2D(GraphBuilder.RHICmdList, GEngine->SheenEnergyTexture,           SpecFormat, TEXT("Shading.SheenEnergy"));
+				Out.DiffuseEnergyTexture  = ShadingEnergyConservationData::CreateTexture2D(GraphBuilder.RHICmdList, GEngine->DiffuseEnergyTexture,         DiffFormat, TEXT("Shading.DiffuseEnergy"));
 
 				// Fallback
-				if (!View.ViewState->ShadingEnergyConservationData.GGXSpecEnergyTexture ) { View.ViewState->ShadingEnergyConservationData.GGXSpecEnergyTexture  = GSystemTextures.BlackDummy; }
-				if (!View.ViewState->ShadingEnergyConservationData.GGXGlassEnergyTexture) { View.ViewState->ShadingEnergyConservationData.GGXGlassEnergyTexture = GSystemTextures.VolumetricBlackDummy; }
-				if (!View.ViewState->ShadingEnergyConservationData.ClothEnergyTexture   ) { View.ViewState->ShadingEnergyConservationData.ClothEnergyTexture    = GSystemTextures.BlackDummy; }
-				if (!View.ViewState->ShadingEnergyConservationData.DiffuseEnergyTexture ) { View.ViewState->ShadingEnergyConservationData.DiffuseEnergyTexture  = GSystemTextures.BlackDummy; }
+				if (!Out.GGXSpecEnergyTexture ) { Out.GGXSpecEnergyTexture  = GSystemTextures.BlackDummy; }
+				if (!Out.GGXGlassEnergyTexture) { Out.GGXGlassEnergyTexture = GSystemTextures.VolumetricBlackDummy; }
+				if (!Out.ClothEnergyTexture   ) { Out.ClothEnergyTexture    = GSystemTextures.BlackDummy; }
+				if (!Out.DiffuseEnergyTexture ) { Out.DiffuseEnergyTexture  = GSystemTextures.BlackDummy; }
 			}
 		}
-
-		View.ViewState->ShadingEnergyConservationData.bEnergyConservation = bIsEnergyConservationEnabled;
-		View.ViewState->ShadingEnergyConservationData.bEnergyPreservation = bIsEnergyPreservationEnabled;
-	}
-	else if (View.ViewState)
-	{
-		View.ViewState->ShadingEnergyConservationData = FShadingEnergyConservationStateData();
 	}
 }
 
+FShadingEnergyConservationData GetData(const FViewInfo& View)
+{
+	const FShadingEnergyConservationSettings Settings = GetSettings(View);
+
+	FShadingEnergyConservationData Out;
+	Out.bEnergyConservation   = Settings.bIsEnergyConservationEnabled;
+	Out.bEnergyPreservation   = Settings.bIsEnergyPreservationEnabled;
+	Out.GGXSpecEnergyTexture  = GShadingEnergyConservationResources.GGXSpecEnergyTexture;
+	Out.GGXGlassEnergyTexture = GShadingEnergyConservationResources.GGXGlassEnergyTexture;
+	Out.ClothEnergyTexture    = GShadingEnergyConservationResources.ClothEnergyTexture;
+	Out.DiffuseEnergyTexture  = GShadingEnergyConservationResources.DiffuseEnergyTexture;
+	return Out;
+}
 
 void Debug(FRDGBuilder& GraphBuilder, const FViewInfo& View, FSceneTextures& SceneTextures)
 {
@@ -504,4 +551,4 @@ void Debug(FRDGBuilder& GraphBuilder, const FViewInfo& View, FSceneTextures& Sce
 	}
 }
 
-}
+}
\ No newline at end of file
diff --git a/Engine/Source/Runtime/Renderer/Private/ShadingEnergyConservation.h b/Engine/Source/Runtime/Renderer/Private/ShadingEnergyConservation.h
--- a/Engine/Source/Runtime/Renderer/Private/ShadingEnergyConservation.h
+++ b/Engine/Source/Runtime/Renderer/Private/ShadingEnergyConservation.h
@@ -11,23 +11,22 @@
 
 class FViewInfo;
 
-struct FShadingEnergyConservationStateData
+struct FShadingEnergyConservationData
 {
 	bool bEnergyConservation = false;
 	bool bEnergyPreservation = false;
-	EPixelFormat Format = PF_Unknown;
+
 	TRefCountPtr<IPooledRenderTarget> GGXSpecEnergyTexture = nullptr;
 	TRefCountPtr<IPooledRenderTarget> GGXGlassEnergyTexture = nullptr;
 	TRefCountPtr<IPooledRenderTarget> ClothEnergyTexture = nullptr;
 	TRefCountPtr<IPooledRenderTarget> DiffuseEnergyTexture = nullptr;
-
-	uint64 GetGPUSizeBytes(bool bLogSizes) const;
 };
 
 namespace ShadingEnergyConservation
 {
-	void Init(FRDGBuilder& GraphBuilder, FViewInfo& View);
+	void Init(FRDGBuilder& GraphBuilder, const FViewInfo& View);
 	void Debug(FRDGBuilder& GraphBuilder, const FViewInfo& View, FSceneTextures& SceneTextures);
+	FShadingEnergyConservationData GetData(const FViewInfo& View);
 
 	RENDERER_API bool IsEnable();
 }
diff --git a/Engine/Source/Runtime/Renderer/Private/Substrate/SubstrateVisualize.cpp b/Engine/Source/Runtime/Renderer/Private/Substrate/SubstrateVisualize.cpp
--- a/Engine/Source/Runtime/Renderer/Private/Substrate/SubstrateVisualize.cpp
+++ b/Engine/Source/Runtime/Renderer/Private/Substrate/SubstrateVisualize.cpp
@@ -345,11 +345,12 @@ static void AddVisualizeSystemInfoPasses(FRDGBuilder& GraphBuilder, const FViewI
 	ShaderPrint::RequestSpaceForCharacters(1024);
 
 	const FRDGTextureDesc MaterialBufferDesc = View.SubstrateViewData.SceneData->MaterialTextureArray->Desc;
+	const FShadingEnergyConservationData ShadingEnergyConservationData = ShadingEnergyConservation::GetData(View);
 
 	FSubstrateSystemInfoCS::FParameters* PassParameters = GraphBuilder.AllocParameters<FSubstrateSystemInfoCS::FParameters>();
 	PassParameters->bAdvancedDebugEnabled = IsAdvancedVisualizationEnabled() ? 1u : 0u;
-	PassParameters->bEnergyConservation = View.ViewState ? View.ViewState->ShadingEnergyConservationData.bEnergyConservation : false;;
-	PassParameters->bEnergyPreservation = View.ViewState ? View.ViewState->ShadingEnergyConservationData.bEnergyPreservation : false;;
+	PassParameters->bEnergyConservation = ShadingEnergyConservationData.bEnergyConservation;
+	PassParameters->bEnergyPreservation = ShadingEnergyConservationData.bEnergyPreservation;
 	PassParameters->bDbufferPass = IsDBufferPassEnabled(View.GetShaderPlatform()) ? 1 : 0;
 	PassParameters->ClassificationCMask = SupportsCMask(View.GetShaderPlatform()) ? 1 : 0;
 	PassParameters->ClassificationAsync = IsClassificationAsync() ? 1 : 0;
