Fix velocity rendering for mobile multi view

diff --git a/Engine/Shaders/Private/VelocityShader.usf b/Engine/Shaders/Private/VelocityShader.usf
--- a/Engine/Shaders/Private/VelocityShader.usf
+++ b/Engine/Shaders/Private/VelocityShader.usf
@@ -48,15 +48,31 @@ void MainVertexShader(
 	, out float OutGlobalClipPlaneDistance : SV_ClipDistance
 #endif
 #if INSTANCED_STEREO
+#if MOBILE_MULTI_VIEW
+	, out uint LayerIndex : SV_RenderTargetArrayIndex
+#else
 	, out uint ViewportIndex : SV_ViewPortArrayIndex
+#endif
+#elif MOBILE_MULTI_VIEW
+	, in nointerpolation uint ViewId : SV_ViewID
 #endif
 	)
 {
 #if INSTANCED_STEREO
 	const uint EyeIndex = GetEyeIndexFromVF(Input);
+#if MOBILE_MULTI_VIEW
+	LayerIndex = EyeIndex;
+#else
 	ViewportIndex = EyeIndex;
 #endif
+#endif
+#if INSTANCED_STEREO
 	ResolvedView = ResolveViewFromVF(Input);
+#elif MOBILE_MULTI_VIEW
+	ResolvedView = ResolveView(ViewId);
+#else
+	ResolvedView = ResolveView();
+#endif
 	
 	FVertexFactoryIntermediates VFIntermediates = GetVertexFactoryIntermediates(Input);
 	float4 TranslatedWorldPosition = VertexFactoryGetWorldPosition( Input, VFIntermediates );
@@ -119,6 +135,9 @@ void MainPixelShader(
 	FVelocityInterpsVSToPS Input
 	OPTIONAL_IsFrontFace
 	OPTIONAL_OutDepthConservative
+#if MOBILE_MULTI_VIEW && !INSTANCED_STEREO
+	, in nointerpolation uint ViewId : SV_ViewID
+#endif
 	,out ENCODED_VELOCITY_TYPE OutColor : SV_Target0
 #if MATERIALBLENDING_MASKED_USING_COVERAGE
 	, out uint OutCoverage : SV_Coverage
@@ -127,6 +146,8 @@ void MainPixelShader(
 {
 #if INSTANCED_STEREO
 	ResolvedView = ResolveView(Input.FactoryInterpolants.EyeIndex);
+#elif MOBILE_MULTI_VIEW
+	ResolvedView = ResolveView(ViewId);
 #else
 	ResolvedView = ResolveView();
 #endif
diff --git a/Engine/Source/Runtime/Renderer/Private/SceneTextures.cpp b/Engine/Source/Runtime/Renderer/Private/SceneTextures.cpp
--- a/Engine/Source/Runtime/Renderer/Private/SceneTextures.cpp
+++ b/Engine/Source/Runtime/Renderer/Private/SceneTextures.cpp
@@ -675,7 +675,7 @@ void FSceneTextures::InitializeViewFamily(FRDGBuilder& GraphBuilder, FViewFamily
 	}
 
 	// Velocity
-	SceneTextures.Velocity = GraphBuilder.CreateTexture(FVelocityRendering::GetRenderTargetDesc(Config.ShaderPlatform, Config.Extent), TEXT("SceneVelocity"));
+	SceneTextures.Velocity = GraphBuilder.CreateTexture(FVelocityRendering::GetRenderTargetDesc(Config.ShaderPlatform, Config.Extent, Config.bRequireMultiView), TEXT("SceneVelocity"));
 
 	if (Config.bIsUsingGBuffers)
 	{
diff --git a/Engine/Source/Runtime/Renderer/Private/VelocityRendering.cpp b/Engine/Source/Runtime/Renderer/Private/VelocityRendering.cpp
--- a/Engine/Source/Runtime/Renderer/Private/VelocityRendering.cpp
+++ b/Engine/Source/Runtime/Renderer/Private/VelocityRendering.cpp
@@ -307,6 +307,7 @@ void FSceneRenderer::RenderVelocities(
 				ExclusiveDepthStencil);
 
 			PassParameters->RenderTargets[0] = FRenderTargetBinding(SceneTextures.Velocity, (bNeedsClearMask & View.GPUMask.GetNative()) ? ERenderTargetLoadAction::EClear : ERenderTargetLoadAction::ELoad);
+			PassParameters->RenderTargets.MultiViewCount = (View.bIsMobileMultiViewEnabled) ? 2 : (View.Aspects.IsMobileMultiViewEnabled() ? 1 : 0);
 			bNeedsClearMask &= ~View.GPUMask.GetNative();
 
 			if (bIsParallelVelocity)
@@ -378,9 +379,11 @@ ETextureCreateFlags FVelocityRendering::GetCreateFlags(EShaderPlatform ShaderPla
 	return TexCreate_RenderTargetable | TexCreate_UAV | TexCreate_ShaderResource | FastVRamFlag;
 }
 
-FRDGTextureDesc FVelocityRendering::GetRenderTargetDesc(EShaderPlatform ShaderPlatform, FIntPoint Extent)
+FRDGTextureDesc FVelocityRendering::GetRenderTargetDesc(EShaderPlatform ShaderPlatform, FIntPoint Extent, const bool bRequireMultiView)
 {
-	return FRDGTextureDesc::Create2D(Extent, GetFormat(ShaderPlatform), FClearValueBinding::Transparent, GetCreateFlags(ShaderPlatform));
+	return bRequireMultiView ?
+		FRDGTextureDesc::Create2DArray(Extent, GetFormat(ShaderPlatform), FClearValueBinding::Transparent, GetCreateFlags(ShaderPlatform), 2) :
+		FRDGTextureDesc::Create2D(Extent, GetFormat(ShaderPlatform), FClearValueBinding::Transparent, GetCreateFlags(ShaderPlatform));
 }
 
 bool FVelocityRendering::IsVelocityPassSupported(EShaderPlatform ShaderPlatform)
diff --git a/Engine/Source/Runtime/Renderer/Private/VelocityRendering.h b/Engine/Source/Runtime/Renderer/Private/VelocityRendering.h
--- a/Engine/Source/Runtime/Renderer/Private/VelocityRendering.h
+++ b/Engine/Source/Runtime/Renderer/Private/VelocityRendering.h
@@ -39,7 +39,7 @@ struct FVelocityRendering
 	static ETextureCreateFlags GetCreateFlags(EShaderPlatform ShaderPlatform);
 	
 	/** Returns the render target description for the velocity buffer. */
-	static FRDGTextureDesc GetRenderTargetDesc(EShaderPlatform ShaderPlatform, FIntPoint Extent);
+	static FRDGTextureDesc GetRenderTargetDesc(EShaderPlatform ShaderPlatform, FIntPoint Extent, const bool bRequireMultiView = false);
 
 	/** Returns true if a velocity pass is supported. */
 	static bool IsVelocityPassSupported(EShaderPlatform ShaderPlatform);
