Fix crash when graph builder blackboard already contains scene compute update interface.

#rnx

diff --git a/Engine/Source/Runtime/Renderer/Internal/SceneComputeUpdates.h b/Engine/Source/Runtime/Renderer/Internal/SceneComputeUpdates.h
--- a/Engine/Source/Runtime/Renderer/Internal/SceneComputeUpdates.h
+++ b/Engine/Source/Runtime/Renderer/Internal/SceneComputeUpdates.h
@@ -7,6 +7,7 @@
 #include "GPUSceneWriter.h"
 
 class FPrimitiveSceneInfo;
+class FScene;
 
 /**
  * Experimental interface that is able to submit updates to modify GPU-scene. 
@@ -24,21 +25,13 @@ public:
 	 */
 	void EnqueueUpdate(FPrimitiveSceneInfo* PrimitiveSceneInfo, FGPUSceneWriteDelegate&& DataWriterGPU)
 	{
-		if (ensure(IsEnabled()))
-		{
-			EnqueueUpdateInternal(PrimitiveSceneInfo, MoveTemp(DataWriterGPU));
-		}
+		EnqueueUpdateInternal(PrimitiveSceneInfo, MoveTemp(DataWriterGPU));
 	}
 
-	/** Whether writes are enabled. */
-	virtual bool IsEnabled() const { return bEnabled; }
-	virtual void SetEnabled(bool bInEnabled) { bEnabled = bInEnabled; }
+	virtual void SetScene(FScene* InScene) = 0;
 
 protected:
 	virtual void EnqueueUpdateInternal(FPrimitiveSceneInfo* PrimitiveSceneInfo, FGPUSceneWriteDelegate&& DataWriterGPU) = 0;
-
-private:
-	bool bEnabled = true;
 };
 
 /**
diff --git a/Engine/Source/Runtime/Renderer/Private/RendererScene.cpp b/Engine/Source/Runtime/Renderer/Private/RendererScene.cpp
--- a/Engine/Source/Runtime/Renderer/Private/RendererScene.cpp
+++ b/Engine/Source/Runtime/Renderer/Private/RendererScene.cpp
@@ -5407,14 +5407,20 @@ void FScene::UpdateAllPrimitiveSceneInfos(FRDGBuilder& GraphBuilder, EUpdateAllP
 class FSceneComputeUpdates : public ISceneComputeUpdates
 {
 public:
-	FSceneComputeUpdates(FScene &InScene) : Scene(InScene) {}
+	FSceneComputeUpdates(FScene* InScene) : Scene(InScene) {}
+
+	virtual void SetScene(FScene* InScene) override { Scene = InScene; }
 
 	virtual void EnqueueUpdateInternal(FPrimitiveSceneInfo* PrimitiveSceneInfo, FGPUSceneWriteDelegate&& DataWriterGPU) override
 	{
-		Scene.UpdatePrimitiveInstancesFromCompute(PrimitiveSceneInfo, MoveTemp(DataWriterGPU));
+		if (ensure(Scene))
+		{
+			Scene->UpdatePrimitiveInstancesFromCompute(PrimitiveSceneInfo, MoveTemp(DataWriterGPU));
+		}
 	}
 
-	FScene &Scene;
+private:
+	FScene* Scene = nullptr;
 };
 
 void FScene::Update(FRDGBuilder& GraphBuilder, const FUpdateParameters& Parameters)
@@ -5436,9 +5442,15 @@ void FScene::Update(FRDGBuilder& GraphBuilder, const FUpdateParameters& Paramete
 	}
 
 	// Add interface to blackboard to receive scene updates from compute.
+	if (const FSceneComputeUpdatesBlackboardEntry* FoundBlackboardEntry = GraphBuilder.Blackboard.Get<FSceneComputeUpdatesBlackboardEntry>())
+	{
+		check(FoundBlackboardEntry->SceneComputeUpdates);
+		FoundBlackboardEntry->SceneComputeUpdates->SetScene(this);
+	}
+	else
 	{
 		FSceneComputeUpdatesBlackboardEntry& BlackboardEntry = GraphBuilder.Blackboard.Create<FSceneComputeUpdatesBlackboardEntry>();
-		BlackboardEntry.SceneComputeUpdates = GraphBuilder.AllocObject<FSceneComputeUpdates>(*this);
+		BlackboardEntry.SceneComputeUpdates = GraphBuilder.AllocObject<FSceneComputeUpdates>(this);
 	}
 
 	for (IComputeTaskWorker* ComputeTaskWorker : ComputeTaskWorkers)
@@ -5452,7 +5464,7 @@ void FScene::Update(FRDGBuilder& GraphBuilder, const FUpdateParameters& Paramete
 	// Disable scene updates on this builder from this point onwards.
 	{
 		const FSceneComputeUpdatesBlackboardEntry& BlackboardEntry = GraphBuilder.Blackboard.GetChecked<FSceneComputeUpdatesBlackboardEntry>();
-		BlackboardEntry.SceneComputeUpdates->SetEnabled(false);
+		BlackboardEntry.SceneComputeUpdates->SetScene(nullptr);
 	}
 
 	// Avoid overlapping prior scene render and async RDG execution tasks to simplify things as there are a lot of moving pieces.
