Removed FAsyncTask tag from SubmitVirtualTextureFeedbackBuffer, because the underlying GVirtualTextureFeedback global object is not synchronized.

diff --git a/Engine/Source/Runtime/Renderer/Private/VT/VirtualTextureFeedbackBuffer.cpp b/Engine/Source/Runtime/Renderer/Private/VT/VirtualTextureFeedbackBuffer.cpp
--- a/Engine/Source/Runtime/Renderer/Private/VT/VirtualTextureFeedbackBuffer.cpp
+++ b/Engine/Source/Runtime/Renderer/Private/VT/VirtualTextureFeedbackBuffer.cpp
@@ -98,7 +98,7 @@ void FVirtualTextureFeedbackBuffer::Begin(FRDGBuilder& GraphBuilder, const FVirt
 
 void FVirtualTextureFeedbackBuffer::End(FRDGBuilder& GraphBuilder)
 {
-	AddPass(GraphBuilder, RDG_EVENT_NAME("VirtualTextureFeedbackCopy"), [UAV = UAV, BufferRHI = PooledBuffer->GetRHI(), Desc = Desc](FRDGAsyncTask, FRHICommandList& RHICmdList)
+	AddPass(GraphBuilder, RDG_EVENT_NAME("VirtualTextureFeedbackCopy"), [UAV = UAV, BufferRHI = PooledBuffer->GetRHI(), Desc = Desc](FRHICommandList& RHICmdList)
 	{
 		RHICmdList.EndUAVOverlap(UAV);
 		RHICmdList.Transition(FRHITransitionInfo(UAV, ERHIAccess::UAVMask, ERHIAccess::CopySrc));
