Moved FDeferredUpdateResource::UpdateResources into a single location, plus a bit of cleanup of the scene rendering arrays.

diff --git a/Engine/Plugins/Experimental/Water/Source/Runtime/Private/WaterInfoRendering.cpp b/Engine/Plugins/Experimental/Water/Source/Runtime/Private/WaterInfoRendering.cpp
--- a/Engine/Plugins/Experimental/Water/Source/Runtime/Private/WaterInfoRendering.cpp
+++ b/Engine/Plugins/Experimental/Water/Source/Runtime/Private/WaterInfoRendering.cpp
@@ -434,7 +434,6 @@ static void UpdateWaterInfoRendering_RenderThread(
 	ON_SCOPE_EXIT {bIsUpdatingWaterInfo = false;};
 
 	FMaterialRenderProxy::UpdateDeferredCachedUniformExpressions();
-	FDeferredUpdateResource::UpdateResources(RHICmdList);
 
 	FRenderTarget* RenderTarget = Params.RenderTarget;
 	FTexture* OutputTexture = Params.OutputTexture;
diff --git a/Engine/Plugins/Runtime/nDisplay/Source/DisplayCluster/Private/Game/EngineClasses/Basics/DisplayClusterViewportClient.cpp b/Engine/Plugins/Runtime/nDisplay/Source/DisplayCluster/Private/Game/EngineClasses/Basics/DisplayClusterViewportClient.cpp
--- a/Engine/Plugins/Runtime/nDisplay/Source/DisplayCluster/Private/Game/EngineClasses/Basics/DisplayClusterViewportClient.cpp
+++ b/Engine/Plugins/Runtime/nDisplay/Source/DisplayCluster/Private/Game/EngineClasses/Basics/DisplayClusterViewportClient.cpp
@@ -888,7 +888,7 @@ void UDisplayClusterViewportClient::Draw(FViewport* InViewport, FCanvas* SceneCa
 		if (GDisplayClusterSingleRender)
 		{
 			GetRendererModule().BeginRenderingViewFamilies(
-				SceneCanvas, TArrayView<FSceneViewFamily*>((FSceneViewFamily**)(ViewFamilies.GetData()), ViewFamilies.Num()));
+				SceneCanvas, MakeArrayView(reinterpret_cast<FSceneViewFamily* const*>(ViewFamilies.GetData()), ViewFamilies.Num()));
 		}
 		else
 		{
diff --git a/Engine/Source/Runtime/RenderCore/Public/RendererInterface.h b/Engine/Source/Runtime/RenderCore/Public/RendererInterface.h
--- a/Engine/Source/Runtime/RenderCore/Public/RendererInterface.h
+++ b/Engine/Source/Runtime/RenderCore/Public/RendererInterface.h
@@ -861,7 +861,7 @@ public:
 	virtual void InvalidatePathTracedOutput() = 0;
 
 	/** Experimental:  Render multiple view families in a single scene render call.  All families must reference the same FScene.  Scene Capture not yet supported. */
-	virtual void BeginRenderingViewFamilies(FCanvas* Canvas, TArrayView<FSceneViewFamily*> ViewFamilies) = 0;
+	virtual void BeginRenderingViewFamilies(FCanvas* Canvas, TConstArrayView<FSceneViewFamily*> ViewFamilies) = 0;
 
 	/** Resets the scene texture extent history. Call this method after rendering with very large render
 	 *  targets. The next scene render will create them at the requested size.
diff --git a/Engine/Source/Runtime/Renderer/Private/MeshEdgesRendering.cpp b/Engine/Source/Runtime/Renderer/Private/MeshEdgesRendering.cpp
--- a/Engine/Source/Runtime/Renderer/Private/MeshEdgesRendering.cpp
+++ b/Engine/Source/Runtime/Renderer/Private/MeshEdgesRendering.cpp
@@ -278,9 +278,6 @@ void RenderMeshEdges_RenderThread(
 
 	FUniformExpressionCacheAsyncUpdateScope AsyncUpdateScope;
 
-	// update any resources that needed a deferred update
-	FDeferredUpdateResource::UpdateResources(RHICmdList);
-
 	FRDGBuilder GraphBuilder(RHICmdList, RDG_EVENT_NAME("MeshEdges"), ERDGBuilderFlags::Parallel);
 	{
 		RDG_EVENT_SCOPE(GraphBuilder, "RenderMeshEdges_RenderThread");
diff --git a/Engine/Source/Runtime/Renderer/Private/PlanarReflectionRendering.cpp b/Engine/Source/Runtime/Renderer/Private/PlanarReflectionRendering.cpp
--- a/Engine/Source/Runtime/Renderer/Private/PlanarReflectionRendering.cpp
+++ b/Engine/Source/Runtime/Renderer/Private/PlanarReflectionRendering.cpp
@@ -347,9 +347,6 @@ static void UpdatePlanarReflectionContents_RenderThread(
 
 	FUniformExpressionCacheAsyncUpdateScope AsyncUpdateScope;
 
-	// update any resources that needed a deferred update
-	FDeferredUpdateResource::UpdateResources(RHICmdList);
-
 	const ERHIFeatureLevel::Type FeatureLevel = SceneRenderer->FeatureLevel;
 	FRDGBuilder GraphBuilder(RHICmdList, RDG_EVENT_NAME("PlanarReflection"), ERDGBuilderFlags::Parallel);
 
diff --git a/Engine/Source/Runtime/Renderer/Private/ReflectionEnvironmentCapture.cpp b/Engine/Source/Runtime/Renderer/Private/ReflectionEnvironmentCapture.cpp
--- a/Engine/Source/Runtime/Renderer/Private/ReflectionEnvironmentCapture.cpp
+++ b/Engine/Source/Runtime/Renderer/Private/ReflectionEnvironmentCapture.cpp
@@ -516,9 +516,6 @@ void CaptureSceneToScratchCubemap(
 
 	FUniformExpressionCacheAsyncUpdateScope AsyncUpdateScope;
 
-	// update any resources that needed a deferred update
-	FDeferredUpdateResource::UpdateResources(RHICmdList);
-
 	FRDGBuilder GraphBuilder(RHICmdList, RDG_EVENT_NAME("CubeMapCapture"), ERDGBuilderFlags::Parallel);
 
 	{
diff --git a/Engine/Source/Runtime/Renderer/Private/RendererModule.h b/Engine/Source/Runtime/Renderer/Private/RendererModule.h
--- a/Engine/Source/Runtime/Renderer/Private/RendererModule.h
+++ b/Engine/Source/Runtime/Renderer/Private/RendererModule.h
@@ -141,7 +141,7 @@ public:
 
 	virtual void InvalidatePathTracedOutput() override;
 
-	virtual void BeginRenderingViewFamilies(FCanvas* Canvas, TArrayView<FSceneViewFamily*> ViewFamilies) override;
+	virtual void BeginRenderingViewFamilies(FCanvas* Canvas, TConstArrayView<FSceneViewFamily*> ViewFamilies) override;
 
 	virtual void ResetSceneTextureExtentHistory() override;
 
diff --git a/Engine/Source/Runtime/Renderer/Private/SceneCaptureRendering.cpp b/Engine/Source/Runtime/Renderer/Private/SceneCaptureRendering.cpp
--- a/Engine/Source/Runtime/Renderer/Private/SceneCaptureRendering.cpp
+++ b/Engine/Source/Runtime/Renderer/Private/SceneCaptureRendering.cpp
@@ -401,9 +401,6 @@ static void UpdateSceneCaptureContentDeferred_RenderThread(
 {
 	SceneRenderer->RenderThreadBegin(RHICmdList);
 
-	// update any resources that needed a deferred update
-	FDeferredUpdateResource::UpdateResources(RHICmdList);
-
 	const ERHIFeatureLevel::Type FeatureLevel = SceneRenderer->FeatureLevel;
 
 #if WANTS_DRAW_MESH_EVENTS
@@ -502,9 +499,6 @@ static void UpdateSceneCaptureContentMobile_RenderThread(
 {
 	SceneRenderer->RenderThreadBegin(RHICmdList);
 
-	// update any resources that needed a deferred update
-	FDeferredUpdateResource::UpdateResources(RHICmdList);
-
 #if WANTS_DRAW_MESH_EVENTS
 	SCOPED_DRAW_EVENTF(RHICmdList, SceneCaptureMobile, TEXT("SceneCaptureMobile %s"), EventName);
 	FRDGBuilder GraphBuilder(RHICmdList, RDG_EVENT_NAME("SceneCaptureMobile %s", *EventName));
diff --git a/Engine/Source/Runtime/Renderer/Private/SceneRendering.cpp b/Engine/Source/Runtime/Renderer/Private/SceneRendering.cpp
--- a/Engine/Source/Runtime/Renderer/Private/SceneRendering.cpp
+++ b/Engine/Source/Runtime/Renderer/Private/SceneRendering.cpp
@@ -3644,7 +3644,7 @@ static void GetCrossGPUTransfers(FSceneRenderer* SceneRenderer, TArray<FCrossGPU
 }
 #endif // WITH_MGPU
 
-void FSceneRenderer::PreallocateCrossGPUFences(const TArray<FSceneRenderer*>& SceneRenderers)
+void FSceneRenderer::PreallocateCrossGPUFences(TConstArrayView<FSceneRenderer*> SceneRenderers)
 {
 #if WITH_MGPU
 	if (SceneRenderers.Num() > 1)
@@ -4966,7 +4966,7 @@ inline ESceneRenderCleanUpMode GetSceneRenderCleanUpMode()
 	return (ESceneRenderCleanUpMode)GSceneRenderCleanUpMode;
 }
 
-static void FinishCleanUp(FRHICommandListImmediate& RHICmdList)
+static void FinishCleanUp()
 {
 	TRACE_CPUPROFILER_EVENT_SCOPE(FinishCleanUp);
 
@@ -4976,7 +4976,7 @@ static void FinishCleanUp(FRHICommandListImmediate& RHICmdList)
 	FGraphicsMinimalPipelineStateId::ResetLocalPipelineIdTableSize();
 }
 
-static void DeleteSceneRenderers(const TArray<FSceneRenderer*>& SceneRenderers)
+static void DeleteSceneRenderers(TConstArrayView<FSceneRenderer*> SceneRenderers)
 {
 	SCOPED_NAMED_EVENT_TEXT("DeleteSceneRenderer", FColor::Red);
 
@@ -5002,7 +5002,7 @@ static void DeleteSceneRenderers(const TArray<FSceneRenderer*>& SceneRenderers)
 	}
 }
 
-static void WaitForTasksAndDeleteSceneRenderers(FRHICommandListImmediate& RHICmdList, const TArray<FSceneRenderer*>& SceneRenderers)
+static void WaitForTasksAndDeleteSceneRenderers(FRHICommandListImmediate& RHICmdList, TConstArrayView<FSceneRenderer*> SceneRenderers)
 {
 	{
 		QUICK_SCOPE_CYCLE_COUNTER(STAT_DeleteSceneRenderer_WaitForTasks);
@@ -5014,23 +5014,15 @@ static void WaitForTasksAndDeleteSceneRenderers(FRHICommandListImmediate& RHICmd
 
 void FSceneRenderer::RenderThreadBegin(FRHICommandListImmediate& RHICmdList)
 {
-	// Pass through to multi-scene-renderer version of function
-	TArray<FSceneRenderer*> SceneRenderers;
-	SceneRenderers.Add(this);
-
-	RenderThreadBegin(RHICmdList, SceneRenderers);
+	RenderThreadBegin(RHICmdList, MakeArrayView({ this }));
 }
 
 void FSceneRenderer::RenderThreadEnd(FRHICommandListImmediate& RHICmdList)
 {
-	// Pass through to multi-scene renderer version of function
-	TArray<FSceneRenderer*> SceneRenderers;
-	SceneRenderers.Add(this);
-
-	RenderThreadEnd(RHICmdList, SceneRenderers);
+	RenderThreadEnd(RHICmdList, MakeArrayView({ this }));
 }
 
-void FSceneRenderer::RenderThreadBegin(FRHICommandListImmediate& RHICmdList, const TArray<FSceneRenderer*>& SceneRenderers)
+void FSceneRenderer::RenderThreadBegin(FRHICommandListImmediate& RHICmdList, TConstArrayView<FSceneRenderer*> SceneRenderers)
 {
 	CleanUp(RHICmdList);
 
@@ -5064,6 +5056,8 @@ void FSceneRenderer::RenderThreadBegin(FRHICommandListImmediate& RHICmdList, con
 	{
 		SceneRenderer->FXSystem = FXSystem;
 	}
+
+	FDeferredUpdateResource::UpdateResources(RHICmdList);
 }
 
 struct FSceneRenderCleanUpState
@@ -5076,7 +5070,7 @@ struct FSceneRenderCleanUpState
 
 static FSceneRenderCleanUpState GSceneRenderCleanUpState;
 
-void FSceneRenderer::RenderThreadEnd(FRHICommandListImmediate& RHICmdList, const TArray<FSceneRenderer*>& SceneRenderers)
+void FSceneRenderer::RenderThreadEnd(FRHICommandListImmediate& RHICmdList, TConstArrayView<FSceneRenderer*> SceneRenderers)
 {
 	// We need to sync async uniform expression cache updates since we're about to start deleting material proxies.
 	FUniformExpressionCacheAsyncUpdateScope::WaitForTask();
@@ -5090,7 +5084,7 @@ void FSceneRenderer::RenderThreadEnd(FRHICommandListImmediate& RHICmdList, const
 	if (GSceneRenderCleanUpState.CompletionMode == ESceneRenderCleanUpMode::Immediate)
 	{
 		WaitForTasksAndDeleteSceneRenderers(RHICmdList, SceneRenderers);
-		FinishCleanUp(RHICmdList);
+		FinishCleanUp();
 	}
 	else
 	{
@@ -5178,7 +5172,7 @@ void FSceneRenderer::CleanUp(FRHICommandListImmediate& RHICmdList)
 		}
 	}
 
-	FinishCleanUp(RHICmdList);
+	FinishCleanUp();
 	GSceneRenderCleanUpState = {};
 }
 
@@ -5220,7 +5214,7 @@ void ResetAndShrinkModifiedBounds(TArray<FBox>& Bounds)
  *
  * @param SceneRenderers	List of scene renderers to use for rendering.
  */
-static void RenderViewFamilies_RenderThread(FRHICommandListImmediate& RHICmdList, const TArray<FSceneRenderer*>& SceneRenderers)
+static void RenderViewFamilies_RenderThread(FRHICommandListImmediate& RHICmdList, TConstArrayView<FSceneRenderer*> SceneRenderers)
 {
 	LLM_SCOPE(ELLMTag::SceneRender);
 
@@ -5236,9 +5230,6 @@ static void RenderViewFamilies_RenderThread(FRHICommandListImmediate& RHICmdList
 	bool bUpdatedGPUSkinCacheVisualization = false;
 #endif
 
-	// update any resources that needed a deferred update
-	FDeferredUpdateResource::UpdateResources(RHICmdList);
-
 	for (FSceneRenderer* SceneRenderer : SceneRenderers)
 	{
 		const ERHIFeatureLevel::Type FeatureLevel = SceneRenderer->FeatureLevel;
@@ -5449,10 +5440,10 @@ extern CORE_API bool GRenderThreadPollingOn;
 
 void FRendererModule::BeginRenderingViewFamily(FCanvas* Canvas, FSceneViewFamily* ViewFamily)
 {
-	BeginRenderingViewFamilies(Canvas, TArrayView<FSceneViewFamily*>(&ViewFamily, 1));
+	BeginRenderingViewFamilies(Canvas, MakeArrayView({ ViewFamily }));
 }
 
-void FRendererModule::BeginRenderingViewFamilies(FCanvas* Canvas, TArrayView<FSceneViewFamily*> ViewFamilies)
+void FRendererModule::BeginRenderingViewFamilies(FCanvas* Canvas, TConstArrayView<FSceneViewFamily*> ViewFamilies)
 {
 	TRACE_CPUPROFILER_EVENT_SCOPE(BeginRenderingViewFamily);
 	check(Canvas);
diff --git a/Engine/Source/Runtime/Renderer/Private/SceneRendering.h b/Engine/Source/Runtime/Renderer/Private/SceneRendering.h
--- a/Engine/Source/Runtime/Renderer/Private/SceneRendering.h
+++ b/Engine/Source/Runtime/Renderer/Private/SceneRendering.h
@@ -2088,8 +2088,8 @@ public:
 	RENDERER_API void RenderThreadBegin(FRHICommandListImmediate& RHICmdList);
 	RENDERER_API void RenderThreadEnd(FRHICommandListImmediate& RHICmdList);
 
-	RENDERER_API static void RenderThreadBegin(FRHICommandListImmediate& RHICmdList, const TArray<FSceneRenderer*>& SceneRenderers);
-	RENDERER_API static void RenderThreadEnd(FRHICommandListImmediate& RHICmdList, const TArray<FSceneRenderer*>& SceneRenderers);
+	RENDERER_API static void RenderThreadBegin(FRHICommandListImmediate& RHICmdList, TConstArrayView<FSceneRenderer*> SceneRenderers);
+	RENDERER_API static void RenderThreadEnd(FRHICommandListImmediate& RHICmdList, TConstArrayView<FSceneRenderer*> SceneRenderers);
 
 	// ISceneRenderer interface
 	virtual UE::Renderer::Private::IShadowInvalidatingInstances *GetShadowInvalidatingInstancesInterface(const FSceneView *SceneView) override;
@@ -2121,7 +2121,7 @@ public:
 #endif
 
 	/** Logic to update render targets across all GPUs */
-	static void PreallocateCrossGPUFences(const TArray<FSceneRenderer*>& SceneRenderers);
+	static void PreallocateCrossGPUFences(TConstArrayView<FSceneRenderer*> SceneRenderers);
 	void DoCrossGPUTransfers(FRDGBuilder& GraphBuilder, FRDGTextureRef ViewFamilyTexture, TArrayView<FViewInfo> InViews, bool bCrossGPUTransferFencesDefer, FRHIGPUMask RenderTargetGPUMask);
 	void FlushCrossGPUFences(FRDGBuilder& GraphBuilder);
 
