[HWRT] Fix the check to not read the number of segments from RHI object as that's unsafe. Only check it when using Nanite.Mode 0.

diff --git a/Engine/Source/Runtime/Renderer/Private/RayTracing/RayTracing.cpp b/Engine/Source/Runtime/Renderer/Private/RayTracing/RayTracing.cpp
--- a/Engine/Source/Runtime/Renderer/Private/RayTracing/RayTracing.cpp
+++ b/Engine/Source/Runtime/Renderer/Private/RayTracing/RayTracing.cpp
@@ -1139,9 +1139,12 @@ namespace RayTracing
 
 					// CacheInstances expects to have one ray tracing mesh command per BLAS segment.
 					// If that's not the case in the future, other logic such as NumCachedStaticVisibleMeshCommands calculation needs to be updated.
-					checkf(RTLODData.CachedMeshCommandIndices.Num() == CachedRayTracingInstance.GeometryRHI->GetNumSegments(),
-						TEXT("Expected to have one ray tracing mesh command per BLAS segment (primitive has %d cached mesh commands but BLAS has %d segments)."),
-						RTLODData.CachedMeshCommandIndices.Num(), CachedRayTracingInstance.GeometryRHI->GetNumSegments());
+					if (!bUsingNaniteRayTracing)
+					{
+						checkf(RTLODData.CachedMeshCommandIndices.Num() == RelevantPrimitive->RayTracingGeometry->Initializer.Segments.Num(),
+							TEXT("Expected to have one ray tracing mesh command per BLAS segment (primitive has %d cached mesh commands but BLAS has %d segments)."),
+							RTLODData.CachedMeshCommandIndices.Num(), RelevantPrimitive->RayTracingGeometry->Initializer.Segments.Num());
+					}
 
 					const bool bNeedMainInstance = !RelevantPrimitive->CachedMeshCommandFlags.bAllSegmentsDecal;
 
