Fixed crash due to VT feedback update happening back to back in the same RDG builder. Moved the feedback pass back to the render thread. Also fixed cases where the page table texture is still in external access mode.

diff --git a/Engine/Source/Runtime/Renderer/Private/VT/VirtualTextureFeedback.cpp b/Engine/Source/Runtime/Renderer/Private/VT/VirtualTextureFeedback.cpp
--- a/Engine/Source/Runtime/Renderer/Private/VT/VirtualTextureFeedback.cpp
+++ b/Engine/Source/Runtime/Renderer/Private/VT/VirtualTextureFeedback.cpp
@@ -221,7 +221,7 @@ void FVirtualTextureFeedback::TransferGPUToCPU(FRDGBuilder& GraphBuilder, FRDGBu
 		RDG_EVENT_NAME("VirtualTextureFeedbackCopy"),
 		Parameters,
 		ERDGPassFlags::Readback,
-		[this, Buffer, Desc](FRDGAsyncTask, FRHICommandList& InRHICmdList)
+		[this, Buffer, Desc](FRHICommandListImmediate& InRHICmdList)
 	{
 		TransferGPUToCPU(InRHICmdList, Buffer->GetRHI(), Desc);
 	});
diff --git a/Engine/Source/Runtime/Renderer/Private/VT/VirtualTextureFeedbackBuffer.cpp b/Engine/Source/Runtime/Renderer/Private/VT/VirtualTextureFeedbackBuffer.cpp
--- a/Engine/Source/Runtime/Renderer/Private/VT/VirtualTextureFeedbackBuffer.cpp
+++ b/Engine/Source/Runtime/Renderer/Private/VT/VirtualTextureFeedbackBuffer.cpp
@@ -98,7 +98,7 @@ void FVirtualTextureFeedbackBuffer::Begin(FRDGBuilder& GraphBuilder, const FVirt
 
 void FVirtualTextureFeedbackBuffer::End(FRDGBuilder& GraphBuilder)
 {
-	AddPass(GraphBuilder, RDG_EVENT_NAME("VirtualTextureFeedbackCopy"), [UAV = UAV, BufferRHI = PooledBuffer->GetRHI(), Desc = Desc](FRHICommandList& RHICmdList)
+	AddPass(GraphBuilder, RDG_EVENT_NAME("VirtualTextureFeedbackCopy"), [UAV = UAV, BufferRHI = PooledBuffer->GetRHI(), Desc = Desc](FRHICommandListImmediate& RHICmdList)
 	{
 		RHICmdList.EndUAVOverlap(UAV);
 		RHICmdList.Transition(FRHITransitionInfo(UAV, ERHIAccess::UAVMask, ERHIAccess::CopySrc));
diff --git a/Engine/Source/Runtime/Renderer/Private/VT/VirtualTextureSpace.cpp b/Engine/Source/Runtime/Renderer/Private/VT/VirtualTextureSpace.cpp
--- a/Engine/Source/Runtime/Renderer/Private/VT/VirtualTextureSpace.cpp
+++ b/Engine/Source/Runtime/Renderer/Private/VT/VirtualTextureSpace.cpp
@@ -332,6 +332,7 @@ void FVirtualTextureSpace::AllocateTextures(FRDGBuilder& GraphBuilder)
 				CopyInfo.Size.Z = 1;
 				CopyInfo.NumMips = FMath::Min(Desc.NumMips, SrcDesc.NumMips);
 
+				GraphBuilder.UseInternalAccessMode(SrcTexture);
 				AddCopyTexturePass(GraphBuilder, SrcTexture, DstTexture, CopyInfo);
 			}
 
@@ -488,6 +489,7 @@ void FVirtualTextureSpace::ApplyUpdates(FVirtualTextureSystem* System, FRDGBuild
 
 		FTextureEntry& PageTableEntry = PageTable[TextureIndex];
 		FRDGTextureRef PageTableTexture = GraphBuilder.RegisterExternalTexture(PageTableEntry.RenderTarget);
+		GraphBuilder.UseInternalAccessMode(PageTableTexture);
 
 		// Use color write mask to update the proper page table entry for this layer
 		FRHIBlendState* BlendStateRHI = nullptr;
