Add source texture region support for rect lights.

Rect light component has now a SourceTextureOffset and a SourceTextureScale property, allowing to target a particular portion of its source texture.

Work done by [REDACTED_PERSON].

diff --git a/Engine/Shaders/Private/RectLightAtlas.usf b/Engine/Shaders/Private/RectLightAtlas.usf
--- a/Engine/Shaders/Private/RectLightAtlas.usf
+++ b/Engine/Shaders/Private/RectLightAtlas.usf
@@ -8,9 +8,11 @@
 // Add texture
 #if SHADER_RECT_VS
 
-uint SlotBufferOffset;
 Buffer<uint4> SlotBuffer;
+Buffer<float4> ScaleOffsetBuffer;
+uint SlotBufferOffset;
 int2 AtlasResolution;
+int  bHasScaleOffset;
 
 void MainVS(
 	uint InVertexId : SV_VertexID,
@@ -23,6 +25,13 @@ void MainVS(
 	const uint4 Rect = SlotBuffer.Load(SlotBufferOffset + InInstanceId);
 	const int2 RectOffset = Rect.xy;
 	const int2 RectResolution = Rect.zw;
+	
+	// Only used for source texture, when copying source texture data into the atlas
+	float4 ScaleOffset = float4(1.0f, 1.0f, 0.0f, 0.0f);
+	if (bHasScaleOffset != 0)
+	{
+		ScaleOffset = ScaleOffsetBuffer.Load(SlotBufferOffset + InInstanceId);;
+	}
 
 	OutId = InInstanceId;
 
@@ -30,7 +39,7 @@ void MainVS(
 	float2 SrcUV = 0;
 	SrcUV.x = InVertexId == 1 || InVertexId == 2 || InVertexId == 4 ? 1.f : 0.f;
 	SrcUV.y = InVertexId == 2 || InVertexId == 4 || InVertexId == 5 ? 1.f : 0.f;
-	OutUV = SrcUV;
+	OutUV = SrcUV * ScaleOffset.xy + ScaleOffset.zw;
 	OutCoord = SrcUV * RectResolution;
 
 	// UV relative to the atlas
@@ -306,8 +315,11 @@ void MainCS(uint3 DispatchThreadId : SV_DispatchThreadID)
 		if (all(AtlasUV < 1.0f))
 		{
 			// Display the atlas, but gray scale it for easing slot visualization
-			const float TintScale = 0.25f;
-			float4 Color = AtlasTexture.SampleLevel(AtlasSampler, AtlasUV, AtlasMIPIndex) * TintScale;
+			const float TintScale_Valid   = 1.0f;
+			const float TintScale_Invalid = 0.25f;
+			const float4 AtlasColor = AtlasTexture.SampleLevel(AtlasSampler, AtlasUV, AtlasMIPIndex);
+
+			float4 OutColor = 0.0f; 
 
 			// Color valid slot with a green border
 			const float BorderSize =  float(AtlasResolution.x) / float(DisplayAtlasResolution.x) * 2.f;
@@ -319,10 +331,12 @@ void MainCS(uint3 DispatchThreadId : SV_DispatchThreadID)
 				const uint4 Slot = SlotBuffer[SlotIt];
 				if (IsInSlot(AtlasCoord, Slot.xy, Slot.xy + Slot.zw))
 				{
+					// only show actual used texture samples ?
+					OutColor = AtlasColor * TintScale_Valid;
 					bScaleDownColor = false;
 					if (IsOnBorder(AtlasCoord, Slot.xy, Slot.xy + Slot.zw, BorderSize))
 					{
-						Color = float4(0, 1, 0, 1);
+						OutColor = float4(0, 1, 0, 1);
 						break;
 					}
 				}
@@ -336,9 +350,8 @@ void MainCS(uint3 DispatchThreadId : SV_DispatchThreadID)
 				{
 					if (IsOnBorder(AtlasCoord, Slot.xy, Slot.xy + Slot.zw, BorderSize))
 					{
-						bScaleDownColor = false;
-						Color = float4(ColorMapViridis(HorizonIt / float(HorizonCount - 1)), 1);
-						Color = float4(1, 0, 0, 1);
+						OutColor = float4(ColorMapViridis(HorizonIt / float(HorizonCount - 1)), 1);
+						OutColor = float4(1, 0, 0, 1);
 						break;
 					}
 				}
@@ -350,22 +363,16 @@ void MainCS(uint3 DispatchThreadId : SV_DispatchThreadID)
 				const uint4 Slot = FreeBuffer[FreeIt];
 				if (IsInSlot(AtlasCoord, Slot.xy, Slot.xy + Slot.zw))
 				{
+					OutColor = AtlasColor * TintScale_Invalid;
 					if (IsOnBorder(AtlasCoord, Slot.xy, Slot.xy + Slot.zw, BorderSize))
 					{
-						bScaleDownColor = false;
-						Color = float4(1, 1, 0, 1);
+						OutColor = float4(1, 1, 0, 1);
 						break;
 					}
 				}
 			}
 
-			// Extra graying if not within a valid slot
-			if (bScaleDownColor)
-			{
-				Color *= TintScale;
-			}
-
-			OutputTexture[DispatchThreadId.xy] = float4(Color.xyz, 1.0f);
+			OutputTexture[DispatchThreadId.xy] = float4(OutColor.xyz, 1.0f);
 		}
 	}
 	else
diff --git a/Engine/Source/Runtime/Engine/Classes/Components/RectLightComponent.h b/Engine/Source/Runtime/Engine/Classes/Components/RectLightComponent.h
--- a/Engine/Source/Runtime/Engine/Classes/Components/RectLightComponent.h
+++ b/Engine/Source/Runtime/Engine/Classes/Components/RectLightComponent.h
@@ -61,6 +61,14 @@ class URectLightComponent : public ULocalLightComponent
 	UPROPERTY(EditAnywhere, BlueprintReadOnly, Category = Light)
 	TObjectPtr<class UTexture> SourceTexture;
 
+	/** Scales the source texture. Value in 0..1. (default=1) */
+	UPROPERTY(EditAnywhere, BlueprintReadOnly, Category = Light, meta = (UIMin = "0.0", UIMax = "1.0"), AdvancedDisplay)
+	FVector2f SourceTextureScale;
+
+	/** Offsets the source texture. Value in 0..1. (default=0) */
+	UPROPERTY(EditAnywhere, BlueprintReadOnly, Category = Light, meta = (UIMin = "0.0", UIMax = "1.0"), AdvancedDisplay)
+	FVector2f SourceTextureOffset;
+
 	UFUNCTION(BlueprintCallable, Category = "Rendering|Lighting")
 	ENGINE_API void SetSourceTexture(UTexture* NewValue);
 
diff --git a/Engine/Source/Runtime/Engine/Private/Components/RectLightComponent.cpp b/Engine/Source/Runtime/Engine/Private/Components/RectLightComponent.cpp
--- a/Engine/Source/Runtime/Engine/Private/Components/RectLightComponent.cpp
+++ b/Engine/Source/Runtime/Engine/Private/Components/RectLightComponent.cpp
@@ -103,6 +103,8 @@ URectLightComponent::URectLightComponent(const FObjectInitializer& ObjectInitial
 	SourceWidth = 64.0f;
 	SourceHeight = 64.0f;
 	SourceTexture = nullptr;
+	SourceTextureOffset = FVector2f(0.0f, 0.0f);
+	SourceTextureScale = FVector2f(1.0f, 1.0f);
 	BarnDoorAngle = GetRectLightBarnDoorMaxAngle();
 	BarnDoorLength = 20.0f;
 	LightFunctionConeAngle = 0.0f;
@@ -291,6 +293,11 @@ FRectLightSceneProxy::FRectLightSceneProxy(const URectLightComponent* Component)
 	, RayTracingData(Component->RayTracingData)
 	, SourceTexture(Component->SourceTexture)
 	, LightFunctionConeAngleTangent(Component->LightFunctionConeAngle > 0 ? FMath::Tan(FMath::Clamp(Component->LightFunctionConeAngle, 0.0f, 89.0f) * (float)UE_PI / 180.0f) : 0.0f)
+	, SourceTextureScaleOffset(FVector4f(
+	  FMath::Clamp(Component->SourceTextureScale.X, 0.f, 1.f), 
+	  FMath::Clamp(Component->SourceTextureScale.Y, 0.f, 1.f), 
+	  FMath::Clamp(Component->SourceTextureOffset.X, 0.f, 1.f), 
+	  FMath::Clamp(Component->SourceTextureOffset.Y, 0.f, 1.f)))
 {
 	RectAtlasId = ~0u;
 }
diff --git a/Engine/Source/Runtime/Engine/Public/RectLightSceneProxy.h b/Engine/Source/Runtime/Engine/Public/RectLightSceneProxy.h
--- a/Engine/Source/Runtime/Engine/Public/RectLightSceneProxy.h
+++ b/Engine/Source/Runtime/Engine/Public/RectLightSceneProxy.h
@@ -32,6 +32,7 @@ public:
 	UTexture*	SourceTexture;
 	uint32		RectAtlasId;
 	float		LightFunctionConeAngleTangent;	// Use Ortho projection if 0
+	FVector4f	SourceTextureScaleOffset;
 
 	FRectLightSceneProxy(const URectLightComponent* Component);
 	virtual ~FRectLightSceneProxy();
diff --git a/Engine/Source/Runtime/Renderer/Private/RectLightTextureManager.cpp b/Engine/Source/Runtime/Renderer/Private/RectLightTextureManager.cpp
--- a/Engine/Source/Runtime/Renderer/Private/RectLightTextureManager.cpp
+++ b/Engine/Source/Runtime/Renderer/Private/RectLightTextureManager.cpp
@@ -94,16 +94,20 @@ namespace RectLightAtlas
 static const uint32 InvalidSlotIndex = ~0u;
 static const FIntPoint InvalidOrigin = FIntPoint(-1, -1);
 
-FIntPoint GetSourceResolution(const FRHITexture* Tex)
+static FIntPoint GetSourceResolution(const FRHITexture* Tex, const FVector4f& ScaleOffset)
 {
 	FIntPoint Out(1, 1);
 	if (Tex)
 	{
 		const FIntVector SourceOriginalResolution = Tex->GetSizeXYZ();
-		const float RatioYX = float(SourceOriginalResolution.Y) / float(SourceOriginalResolution.X);
-		const float RatioXY = float(SourceOriginalResolution.X) / float(SourceOriginalResolution.Y);
 
-		Out = FIntPoint(SourceOriginalResolution.X, SourceOriginalResolution.Y);
+		float ScaledResolutionX = ScaleOffset.X * SourceOriginalResolution.X;
+		float ScaledResolutionY = ScaleOffset.Y * SourceOriginalResolution.Y;
+
+		const float RatioYX = ScaledResolutionY / ScaledResolutionX;
+		const float RatioXY = ScaledResolutionX / ScaledResolutionY;
+
+		Out = FIntPoint(ScaledResolutionX, ScaledResolutionY);
 		// Max Ratio of 2
 		const float RatioThreshold = FMath::Clamp(CVarRectLighMaxTextureRatio.GetValueOnAnyThread(), 1, 16);
 		if (RatioYX > RatioThreshold)
@@ -164,6 +168,7 @@ struct FAtlasSlot
 {
 	uint32 Id = InvalidSlotIndex;
 	FAtlasRect Rect;
+	FVector4f SourceScaleOffset;
 	FTextureReferenceRHIRef SourceTexture = nullptr;
 	uint32 RefCount = 0;
 	bool bForceRefresh = false;
@@ -172,7 +177,7 @@ struct FAtlasSlot
 
 	FIntPoint GetSourceResolution() const 
 	{ 
-		return RectLightAtlas::GetSourceResolution(SourceTexture ? SourceTexture->GetReferencedTexture() : nullptr);
+		return RectLightAtlas::GetSourceResolution(SourceTexture ? SourceTexture->GetReferencedTexture() : nullptr, SourceScaleOffset);
 	}
 };
 
@@ -243,6 +248,7 @@ static bool Traits_IsValid(const FAtlasSlot& In)					{ return In.IsValid(); }
 static const FAtlasRect& Traits_GetRect(const FAtlasRect& In)		{ return In; }
 static const FAtlasRect& Traits_GetRect(const FAtlasHorizon& In)	{ return In.ExtendedLine; }
 static const FAtlasRect& Traits_GetRect(const FAtlasSlot& In)		{ return In.Rect; }
+static const FVector4f&	 Traits_GetScaleOffset(const FAtlasSlot& In) { return In.SourceScaleOffset; }
 
 template<typename T>
 static FRDGBufferRef CreateSlotBuffer(FRDGBuilder& GraphBuilder, const TArray<T>& In, const TCHAR* InBufferName)
@@ -279,6 +285,43 @@ static FRDGBufferRef CreateSlotBuffer(FRDGBuilder& GraphBuilder, const TArray<T>
 	return Buffer;
 }
 
+template<typename T>
+static FRDGBufferRef CreateScaleOffsetBuffer(FRDGBuilder& GraphBuilder, const TArray<T>& In, const TCHAR* InBufferName)
+{
+	struct FAtlasGPUScaleOffset
+	{
+		float  SourceScaleX;
+		float  SourceScaleY;
+		float  SourceOffsetX;
+		float  SourceOffsetY;
+	};
+
+	if (In.IsEmpty())
+	{
+		return GSystemTextures.GetDefaultBuffer(GraphBuilder, sizeof(FAtlasGPUScaleOffset), 0u);
+	}
+
+	TArray<FAtlasGPUScaleOffset> ScaleOffsetData;
+	ScaleOffsetData.Reserve(In.Num());
+	for (const T& I : In)
+	{
+		if (Traits_IsValid(I))
+		{
+			//Traits_GetRect(I);
+			const FVector4f& ScaleOffset = Traits_GetScaleOffset(I);
+			FAtlasGPUScaleOffset& GPUSlot	= ScaleOffsetData.AddDefaulted_GetRef();
+			GPUSlot.SourceScaleX	= ScaleOffset.X;
+			GPUSlot.SourceScaleY	= ScaleOffset.Y;
+			GPUSlot.SourceOffsetX	= ScaleOffset.Z;
+			GPUSlot.SourceOffsetY	= ScaleOffset.W;
+		}
+	}
+	FRDGBufferRef Buffer = GraphBuilder.CreateBuffer(FRDGBufferDesc::CreateBufferDesc(sizeof(FAtlasGPUScaleOffset), ScaleOffsetData.Num()), InBufferName);
+	GraphBuilder.QueueBufferUpload(Buffer, ScaleOffsetData.GetData(), ScaleOffsetData.Num() * sizeof(FAtlasGPUScaleOffset));
+	return Buffer;
+}
+
+
 ///////////////////////////////////////////////////////////////////////////////////////////////////
 // Debug
 class FRectLightAtlasDebugInfoCS : public FGlobalShader
@@ -405,9 +448,11 @@ class FRectLightAtlasVS : public FGlobalShader
 	SHADER_USE_PARAMETER_STRUCT(FRectLightAtlasVS, FGlobalShader);
 	
 	BEGIN_SHADER_PARAMETER_STRUCT(FParameters, )
-		SHADER_PARAMETER(FIntPoint, AtlasResolution)
-		SHADER_PARAMETER(uint32, SlotBufferOffset)
 		SHADER_PARAMETER_RDG_BUFFER_SRV(Buffer, SlotBuffer)
+		SHADER_PARAMETER_RDG_BUFFER_SRV(Buffer, ScaleOffsetBuffer)
+		SHADER_PARAMETER(uint32, SlotBufferOffset)
+		SHADER_PARAMETER(FIntPoint, AtlasResolution)
+		SHADER_PARAMETER(int32, bHasScaleOffset)
 	END_SHADER_PARAMETER_STRUCT()
 
 public:
@@ -473,7 +518,8 @@ static void AddSlotsPass(
 	const FIntRect Viewport(FIntPoint::ZeroValue, OutAtlas->Desc.Extent);
 	const FIntPoint Resolution(Viewport.Width(), Viewport.Height());
 
-	FRDGBufferRef SlotBuffer = CreateSlotBuffer(GraphBuilder, Slots, TEXT("RectLight.AtlasSlotBuffer"));
+	FRDGBufferRef SlotBuffer		= CreateSlotBuffer(GraphBuilder, Slots, TEXT("RectLight.AtlasSlotBuffer"));
+	FRDGBufferRef ScaleOffsetBuffer = CreateScaleOffsetBuffer(GraphBuilder, Slots, TEXT("RectLight.AtlasScaleOffsetBuffer"));
 
 	// Batch new slots into several passes
 	const uint32 SlotCountPerPass = 8u;
@@ -527,6 +573,8 @@ static void AddSlotsPass(
 		Parameters->VS.AtlasResolution = Resolution;
 		Parameters->VS.SlotBufferOffset = SlotOffset;
 		Parameters->VS.SlotBuffer = GraphBuilder.CreateSRV(SlotBuffer, PF_R16G16B16A16_UINT);
+		Parameters->VS.ScaleOffsetBuffer = GraphBuilder.CreateSRV(ScaleOffsetBuffer, PF_A32B32G32R32F);
+		Parameters->VS.bHasScaleOffset   = true;
 		Parameters->RenderTargets[0] = FRenderTargetBinding(OutAtlas, ERenderTargetLoadAction::ELoad, 0);
 
 		TShaderMapRef<FRectLightAtlasVS> VertexShader(ShaderMap);
@@ -595,6 +643,8 @@ static void CopySlotsPass(
 	const FRDGTextureRef& InAtlas,
 	FRDGTextureRef& OutAtlas)
 {
+	FRDGBufferSRVRef DummyScaleOffsetBuffer = GraphBuilder.CreateSRV(GSystemTextures.GetDefaultBuffer(GraphBuilder, 16u, 1u), PF_A32B32G32R32F);
+
 	// For each MIP, copy the slot from the old atlas (InAtlas) to the new atlas (OutAtlas)	
 	const uint32 MipCount = InAtlas->Desc.NumMips; // FMath::Log2(FMath::Min(InAtlas->Desc.Extent.X, InAtlas->Desc.Extent.Y));
 	for (uint32 MipIt = 0; MipIt < MipCount; ++MipIt)
@@ -640,6 +690,9 @@ static void CopySlotsPass(
 			Parameters->VS.AtlasResolution = Resolution;
 			Parameters->VS.SlotBufferOffset = 0;
 			Parameters->VS.SlotBuffer = GraphBuilder.CreateSRV(DstSlotBuffer, PF_R16G16B16A16_UINT);
+			Parameters->VS.ScaleOffsetBuffer = DummyScaleOffsetBuffer;
+			Parameters->VS.bHasScaleOffset   = false;
+		
 			Parameters->RenderTargets[0] = FRenderTargetBinding(OutAtlas, ERenderTargetLoadAction::ELoad, MipIt);
 
 			TShaderMapRef<FRectLightAtlasVS> VertexShader(ShaderMap);
@@ -710,6 +763,8 @@ static void FilterSlotsPass(
 	const TArray<FAtlasSlot>& InSlots,
 	FRDGTextureRef& InAtlas)
 {
+	FRDGBufferSRVRef DummyScaleOffsetBuffer = GraphBuilder.CreateSRV(GSystemTextures.GetDefaultBuffer(GraphBuilder, 16u, 1u), PF_A32B32G32R32F);
+
 	const uint32 MipCount = InAtlas->Desc.NumMips; // FMath::Log2(FMath::Min(InAtlas->Desc.Extent.X, InAtlas->Desc.Extent.Y));
 	for (uint32 DstMip = 1; DstMip < MipCount; ++DstMip)
 	{
@@ -756,6 +811,8 @@ static void FilterSlotsPass(
 			Parameters->VS.AtlasResolution	= ToMIP(InAtlas->Desc.Extent, DstMip);
 			Parameters->VS.SlotBufferOffset = 0;
 			Parameters->VS.SlotBuffer		= GraphBuilder.CreateSRV(DstSlotBuffer, PF_R16G16B16A16_UINT);
+			Parameters->VS.ScaleOffsetBuffer = DummyScaleOffsetBuffer;
+			Parameters->VS.bHasScaleOffset   = false;
 			Parameters->RenderTargets[0] = FRenderTargetBinding(InAtlas, ERenderTargetLoadAction::ELoad, DstMip);
 
 			TShaderMapRef<FRectLightAtlasVS> VertexShader(ShaderMap);
@@ -1437,7 +1494,7 @@ static void PackAtlas(
 ///////////////////////////////////////////////////////////////////////////////////////////////////
 // API
 
-uint32 AddTexture(UTexture* In)
+uint32 AddTexture(UTexture* In, const FVector4f& ScaleOffset)
 {
 	check(IsInRenderingThread());
 
@@ -1448,7 +1505,8 @@ uint32 AddTexture(UTexture* In)
 		bool bFound = false;
 		for (FAtlasSlot& Slot : GRectLightTextureManager.AtlasSlots)
 		{
-			if (Slot.SourceTexture == In->TextureReference.TextureReferenceRHI)
+			if ((Slot.SourceTexture == In->TextureReference.TextureReferenceRHI) &&
+			    (Slot.SourceScaleOffset == ScaleOffset))
 			{
 				Slot.RefCount++;
 				SlotIndex = Slot.Id;
@@ -1475,9 +1533,10 @@ uint32 AddTexture(UTexture* In)
 
 			*Slot = FAtlasSlot();
 			Slot->SourceTexture = In->TextureReference.TextureReferenceRHI;
+			Slot->SourceScaleOffset = ScaleOffset;
 			Slot->Id = SlotIndex;
 			Slot->Rect.Origin = InvalidOrigin;
-			Slot->Rect.Resolution = GetSourceResolution(Tex);
+			Slot->Rect.Resolution = GetSourceResolution(Tex, ScaleOffset);
 			Slot->RefCount = 1;
 
 			GRectLightTextureManager.bHasPendingAdds = true;
diff --git a/Engine/Source/Runtime/Renderer/Private/RectLightTextureManager.h b/Engine/Source/Runtime/Renderer/Private/RectLightTextureManager.h
--- a/Engine/Source/Runtime/Renderer/Private/RectLightTextureManager.h
+++ b/Engine/Source/Runtime/Renderer/Private/RectLightTextureManager.h
@@ -20,7 +20,7 @@ struct FAtlasSlotDesc
 };
 	
 // Add a rect light source texture to the atlas
-RENDERER_API uint32 AddTexture(UTexture* Texture);
+RENDERER_API uint32 AddTexture(UTexture* Texture, const FVector4f& ScaleOffset = FVector4f::One());
 
 // Remove a rect light source texture from the atlas
 RENDERER_API void RemoveTexture(uint32 InSlotId);
diff --git a/Engine/Source/Runtime/Renderer/Private/RendererScene.cpp b/Engine/Source/Runtime/Renderer/Private/RendererScene.cpp
--- a/Engine/Source/Runtime/Renderer/Private/RendererScene.cpp
+++ b/Engine/Source/Runtime/Renderer/Private/RendererScene.cpp
@@ -2616,7 +2616,7 @@ void FScene::AddLightSceneInfo_RenderThread(FLightSceneInfo* LightSceneInfo)
 	if (LightType == LightType_Rect)
 	{
 		FRectLightSceneProxy* RectProxy = (FRectLightSceneProxy*)LightSceneInfo->Proxy;
-		RectProxy->RectAtlasId = RectLightAtlas::AddTexture(RectProxy->SourceTexture);
+		RectProxy->RectAtlasId = RectLightAtlas::AddTexture(RectProxy->SourceTexture, RectProxy->SourceTextureScaleOffset); 
 	}
 
 	// Register IES texture
