Moved virtual shadow map dispatch passes to run in an async task.

diff --git a/Engine/Source/Runtime/Renderer/Private/VirtualShadowMaps/VirtualShadowMapArray.cpp b/Engine/Source/Runtime/Renderer/Private/VirtualShadowMaps/VirtualShadowMapArray.cpp
--- a/Engine/Source/Runtime/Renderer/Private/VirtualShadowMaps/VirtualShadowMapArray.cpp
+++ b/Engine/Source/Runtime/Renderer/Private/VirtualShadowMaps/VirtualShadowMapArray.cpp
@@ -1162,7 +1162,7 @@ public:
 			ClearUnusedGraphResources(ComputeShader, ParametersMetadata, Parameters);
 
 			GraphBuilder.AddPass(Forward<FRDGEventName>(PassName), ParametersMetadata, Parameters, ERDGPassFlags::Compute,
-			[ParametersMetadata, Parameters, ComputeShader, Bins = Bins](FRHIComputeCommandList& RHICmdList)
+			[ParametersMetadata, Parameters, ComputeShader, Bins = Bins](FRDGAsyncTask, FRHIComputeCommandList& RHICmdList)
 			{
 				for (int32 BinIndex = 0; BinIndex < BinCount; ++BinIndex)
 				{
