Fix rect. texture atlas refreshing when several rect light reference to the same animated texture.

diff --git a/Engine/Source/Runtime/Renderer/Private/RectLightTextureManager.cpp b/Engine/Source/Runtime/Renderer/Private/RectLightTextureManager.cpp
--- a/Engine/Source/Runtime/Renderer/Private/RectLightTextureManager.cpp
+++ b/Engine/Source/Runtime/Renderer/Private/RectLightTextureManager.cpp
@@ -1800,6 +1800,9 @@ FAtlasTextureInvalidationScope::FAtlasTextureInvalidationScope(const UTexture* I
 {
 	if (In)
 	{
+		// Several rect lights can reference the same animated texture, but with different offset/scale. 
+		// In this case we need to refresh all rects referencing this texture.
+		bool bHasBeenLockedLocally = false;
 		for (FAtlasSlot& Slot : GRectLightTextureManager.AtlasSlots)
 		{
 			// If the input texture is actually part of the atlas, 
@@ -1808,13 +1811,16 @@ FAtlasTextureInvalidationScope::FAtlasTextureInvalidationScope(const UTexture* I
 			// its data during the next update
 			if (Slot.SourceTexture == In->TextureReference.TextureReferenceRHI)
 			{
-				// Sanity check, allow a single lock/capture refresh at a time.
-				check(GRectLightTextureManager.bLock == false);
+				if (!bHasBeenLockedLocally)
+				{
+					// Sanity check, allow a single lock/capture refresh at a time.
+					check(GRectLightTextureManager.bLock == false);
+					GRectLightTextureManager.bLock = true;
+					bHasBeenLockedLocally = true;
+				}
 
 				bLocked = true;
 				Slot.bForceRefresh = true;
-				GRectLightTextureManager.bLock = true;
-				return;
 			}
 		}
 	}
