Fix CSM shadow casters getting falsely culled due to test with primitive octree using an incorrect culling volume. Added r.shadow.visualizeprimitiveoctree.

diff --git a/Engine/Source/Runtime/Renderer/Private/ShadowSetup.cpp b/Engine/Source/Runtime/Renderer/Private/ShadowSetup.cpp
--- a/Engine/Source/Runtime/Renderer/Private/ShadowSetup.cpp
+++ b/Engine/Source/Runtime/Renderer/Private/ShadowSetup.cpp
@@ -310,6 +310,15 @@ FAutoConsoleVariableRef CVarUseOctreeForShadowCulling(
 	ECVF_Scalability | ECVF_RenderThreadSafe
 	);
 
+#if !(UE_BUILD_SHIPPING || UE_BUILD_TEST)
+static TAutoConsoleVariable<int32> CVarVisualizePrimitiveOctree(
+	TEXT("r.Shadow.VisualizePrimitiveOctree"),
+	0,
+	TEXT("Render a wireframe outline of the nodes in the primitive octree for debugging purposes"),
+	ECVF_Scalability | ECVF_RenderThreadSafe
+);
+#endif
+
 static TAutoConsoleVariable<int32> CVarAlwaysAllocateMaxResolutionAtlases(
 	TEXT("r.Shadow.AlwaysAllocateMaxResolutionAtlases"),
 	0,
@@ -1050,7 +1059,7 @@ void FProjectedShadowInfo::SetupWholeSceneProjection(
 	// Any meshes between the light and the subject can cast shadows, also any meshes inside the subject region
 	const FMatrix CasterOuterMatrix = WorldToViewScaledInner * BorderScaleMatrix * FShadowProjectionMatrix(Initializer.MinLightW, MaxSubjectZ, Initializer.WAxis);
 	const FMatrix ReceiverInnerMatrix = WorldToViewScaledInner * FShadowProjectionMatrix(MinSubjectZ, ClampedMaxLightW, Initializer.WAxis);
-	GetViewFrustumBounds(CasterOuterFrustum, FMatrix(TranslatedWorldToClipOuterMatrix), true);
+	GetViewFrustumBounds(CasterOuterFrustum, CasterOuterMatrix, true);
 	GetViewFrustumBounds(ReceiverInnerFrustum, ReceiverInnerMatrix, true);
 	
 	InvReceiverInnerMatrix = FMatrix44f(ReceiverInnerMatrix.Inverse());
@@ -5322,6 +5331,21 @@ void FSceneRenderer::BeginGatherShadowPrimitives(FDynamicShadowsTaskData* TaskDa
 	}
 }
 
+static void DebugDrawPrimitiveOctree(FViewInfo& View, const FScenePrimitiveOctree& PrimitiveOctree)
+{
+	FViewElementPDI PrimitivePDI(&View, nullptr, nullptr);
+	PrimitiveOctree.FindNodesWithPredicate(
+		[&](FScenePrimitiveOctree::FNodeIndex ParentNodeIndex, FScenePrimitiveOctree::FNodeIndex CurrentNodeIndex, const FBoxCenterAndExtent& Bounds) -> bool
+		{
+			return true;
+		},
+		[&](FScenePrimitiveOctree::FNodeIndex ParentNodeIndex, FScenePrimitiveOctree::FNodeIndex CurrentNodeIndex, const FBoxCenterAndExtent& Bounds)
+		{
+			DrawWireBox(&PrimitivePDI, Bounds.GetBox(), FLinearColor(1,0,1), SDPG_Foreground);
+		}
+	);
+}
+
 void FSceneRenderer::FinishGatherShadowPrimitives(FDynamicShadowsTaskData* TaskData)
 {
 	SCOPED_NAMED_EVENT_TEXT("FSceneRenderer::FinishGatherShadowPrimitives", FColor::Green);
@@ -5365,6 +5389,11 @@ void FSceneRenderer::FinishGatherShadowPrimitives(FDynamicShadowsTaskData* TaskD
 	}
 
 #if !(UE_BUILD_SHIPPING || UE_BUILD_TEST)
+	if (CVarVisualizePrimitiveOctree.GetValueOnRenderThread() != 0)
+	{
+		DebugDrawPrimitiveOctree(TaskData->Views[0], TaskData->Scene->PrimitiveOctree);
+	}
+
 	if (!TaskData->ShadowArrays.OnePassShadowUnsupportedLights.IsEmpty())
 	{
 		OnGetOnScreenMessages.AddLambda([Scene = Scene, OnePassShadowUnsupportedLights = MoveTemp(TaskData->ShadowArrays.OnePassShadowUnsupportedLights)](FScreenMessageWriter& ScreenMessageWriter)->void
