Share animation bank struct definitions between CPU and GPU, and correctly pack anim bank flags.

diff --git a/Engine/Shaders/Private/Skinning/AnimBankEval.usf b/Engine/Shaders/Private/Skinning/AnimBankEval.usf
--- a/Engine/Shaders/Private/Skinning/AnimBankEval.usf
+++ b/Engine/Shaders/Private/Skinning/AnimBankEval.usf
@@ -5,40 +5,7 @@
 #include "/Engine/Private/SceneData.ush"
 #include "/Engine/Private/BoneTransform.ush"
 #include "/Engine/Private/Visualization.ush"
-
-struct FBoneBlockHeader
-{
-	uint BlockLocalIndex;
-	uint BlockBoneCount;
-	uint BlockTransformOffset;
-	uint BankRecordOffset;
-};
-
-struct FScatterBlockHeader
-{
-	uint BlockLocalIndex;
-	uint BlockSrcTransformOffset;
-	uint BlockDstTransformOffset;
-	uint BlockTransformCount : 8;
-	uint TotalTransformCount : 24;
-};
-
-struct FBankRecordHeader
-{
-	uint BoneCount;
-	uint FrameCount;
-	
-	float SampleRate;
-	float PlayRate;
-
-	float CurrentTime;
-	float PreviousTime;
-
-	uint TransformOffset	: 29;
-	uint Playing			: 1;
-	uint Interpolating		: 1;
-	uint HasScale			: 1;
-};
+#include "/Engine/Shared/SkinningDefinitions.h"
 
 ByteAddressBuffer HeaderBuffer;
 ByteAddressBuffer BankBuffer;
@@ -103,12 +70,12 @@ FKeyIndices GetKeyIndicesFromTime(float Time, int KeyCount, float SequenceLength
 }
 
 [numthreads(BONES_PER_GROUP, 1, 1)]
-void BankEvaluateCS(uint3 GroupId : SV_GroupID, uint GroupIndex : SV_GroupIndex, uint3 DispatchThreadId : SV_DispatchThreadID)
+void BankEvaluateCS(uint3 GroupId : SV_GroupID, uint GroupIndex : SV_GroupIndex)
 {
 	const uint BoneBlock		= GroupId.x;
 	const uint LocalBoneIndex	= GroupIndex;
 
-	const FBoneBlockHeader BoneBlockHeader = HeaderBuffer.Load<FBoneBlockHeader>(BoneBlock * uint(sizeof(FBoneBlockHeader)));
+	const FBankBlockHeader BoneBlockHeader = HeaderBuffer.Load<FBankBlockHeader>(BoneBlock * uint(sizeof(FBankBlockHeader)));
 	if (LocalBoneIndex < BoneBlockHeader.BlockBoneCount)
 	{
 		const uint BankRecordHeaderOffset = BoneBlockHeader.BankRecordOffset;
@@ -188,7 +155,7 @@ void BankScatterCS(uint3 GroupId : SV_GroupID, uint GroupIndex : SV_GroupIndex)
 	const uint GlobalBlockIndex		= GroupId.x;
 	const uint LocalTransformIndex	= GroupIndex;
 
-	const FScatterBlockHeader BlockHeader = HeaderBuffer.Load<FScatterBlockHeader>(GlobalBlockIndex * uint(sizeof(FScatterBlockHeader)));
+	const FBankScatterHeader BlockHeader = HeaderBuffer.Load<FBankScatterHeader>(GlobalBlockIndex * uint(sizeof(FBankScatterHeader)));
 	if (LocalTransformIndex < BlockHeader.BlockTransformCount)
 	{
 		const uint DstTransformOffset = BlockHeader.BlockDstTransformOffset + (LocalTransformIndex * (uint)sizeof(float3x4));
diff --git a/Engine/Shaders/Shared/SkinningDefinitions.h b/Engine/Shaders/Shared/SkinningDefinitions.h
--- a/Engine/Shaders/Shared/SkinningDefinitions.h
+++ b/Engine/Shaders/Shared/SkinningDefinitions.h
@@ -2,9 +2,56 @@
 
 #pragma once
 
-#if defined(__cplusplus)
+#ifdef __cplusplus
+
+#include "HLSLTypeAliases.h"
+
+namespace UE::HLSL
+{
+#endif
+
+struct FBankRecordHeader
+{
+	uint BoneCount;
+	uint FrameCount;
+
+	float SampleRate;
+	float PlayRate;
+
+	float CurrentTime;
+	float PreviousTime;
+
+	uint TransformOffset	: 29;
+	uint Playing			: 1;
+	uint Interpolating		: 1;
+	uint HasScale			: 1;
+};
+
+struct FBankBlockHeader
+{
+	uint BlockLocalIndex;
+	uint BlockBoneCount;
+	uint BlockTransformOffset;
+	uint BankRecordOffset;
+};
+
+struct FBankScatterHeader
+{
+	uint BlockLocalIndex;
+	uint BlockSrcTransformOffset;
+	uint BlockDstTransformOffset;
+	uint BlockTransformCount : 8;
+	uint TotalTransformCount : 24;
+};
+
+#define ANIM_BANK_FLAG_NONE			0x0
+#define ANIM_BANK_FLAG_LOOPING		0x1
+#define ANIM_BANK_FLAG_AUTOSTART	0x2
+
+#ifdef __cplusplus
+} // namespace UE::HLSL
 
 #define REF_POSE_TRANSFORM_PROVIDER_GUID 0x665207E7, 0x449A4FB1, 0xA298F7AD, 0x8F989B11
 #define ANIM_RUNTIME_TRANSFORM_PROVIDER_GUID 0xF1508490, 0xFCC24BB9, 0xA9F277B3, 0x1AF766F0
 
-#endif
\ No newline at end of file
+#endif
diff --git a/Engine/Source/Runtime/Engine/Classes/Animation/AnimBank.h b/Engine/Source/Runtime/Engine/Classes/Animation/AnimBank.h
--- a/Engine/Source/Runtime/Engine/Classes/Animation/AnimBank.h
+++ b/Engine/Source/Runtime/Engine/Classes/Animation/AnimBank.h
@@ -8,6 +8,7 @@
 #include "IO/IoHash.h"
 #include "Logging/LogMacros.h"
 #include "RenderCommandFence.h"
+#include "SkinningDefinitions.h"
 #include "UObject/ObjectPtr.h"
 #include "UObject/WeakObjectPtr.h"
 
@@ -68,6 +69,16 @@ struct FAnimBankEntry
 	uint32 FrameCount	= 0u;
 	uint32 KeyCount		= 0u;
 	uint32 Flags		= 0u;
+
+	inline bool IsLooping() const
+	{
+		return (Flags & ANIM_BANK_FLAG_LOOPING) != 0;
+	}
+
+	inline bool IsAutoStart() const
+	{
+		return (Flags & ANIM_BANK_FLAG_AUTOSTART) != 0;
+	}
 };
 
 inline FArchive& operator<<(FArchive& Ar, FAnimBankEntry& BankEntry)
diff --git a/Engine/Source/Runtime/Engine/Private/Animation/AnimBank.cpp b/Engine/Source/Runtime/Engine/Private/Animation/AnimBank.cpp
--- a/Engine/Source/Runtime/Engine/Private/Animation/AnimBank.cpp
+++ b/Engine/Source/Runtime/Engine/Private/Animation/AnimBank.cpp
@@ -72,6 +72,23 @@ static inline void DecomposeTransforms(
 	}
 }
 
+static inline uint32 BankFlagsFromSequence(const FAnimBankSequence& BankSequence)
+{
+	uint32 Flags = ANIM_BANK_FLAG_NONE;
+
+	if (BankSequence.bLooping)
+	{
+		Flags |= ANIM_BANK_FLAG_LOOPING;
+	}
+
+	if (BankSequence.bAutoStart)
+	{
+		Flags |= ANIM_BANK_FLAG_AUTOSTART;
+	}
+
+	return Flags;
+}
+
 class FAnimBankAsyncBuildWorker : public FNonAbandonableTask
 {
 	FAnimBankBuildAsyncCacheTask* Owner;
@@ -449,7 +466,8 @@ bool FAnimBankBuildAsyncCacheTask::BuildData(const UE::FSharedString& Name, cons
 	{
 		FAnimBankEntry& BankEntry = Data->Entries[ItemIndex];
 
-		UAnimSequence* Sequence = AnimBank->Sequences[ItemIndex].Sequence.Get();
+		FAnimBankSequence& BankSequence = AnimBank->Sequences[ItemIndex];
+		UAnimSequence* Sequence = BankSequence.Sequence.Get();
 		if (Sequence == nullptr)
 		{
 			continue;
@@ -559,6 +577,8 @@ bool FAnimBankBuildAsyncCacheTask::BuildData(const UE::FSharedString& Name, cons
 		// Also account for translated root motion
 		BankEntry.ConservativeBounds = AssetBounds;
 
+		BankEntry.Flags = BankFlagsFromSequence(BankSequence);
+
 		check(BankEntry.PositionKeys.Num() > 0);
 		check(BankEntry.RotationKeys.Num() > 0);
 		check(BankEntry.FrameCount > 0);
@@ -911,8 +931,7 @@ FIoHash UAnimBank::CreateDerivedDataKeyHash(const ITargetPlatform* TargetPlatfor
 	}
 
 	FMemoryHasherBlake3 Writer;
-
-	FGuid AnimBankVersionGuid(0x2FA68926, 0x4EB94719, 0x9B928A29, 0xAE522B0F);
+	FGuid AnimBankVersionGuid(0x09961556, 0xCDAF4D9F, 0xB03ABA3C, 0x40EC8690);
 	Writer << AnimBankVersionGuid;
 
 	if (IsValid(Asset))
@@ -932,11 +951,9 @@ FIoHash UAnimBank::CreateDerivedDataKeyHash(const ITargetPlatform* TargetPlatfor
 			Writer << SequenceHash;
 		}
 
-		uint8 bLooping		= BankSequence.bLooping ? 1 : 0;
-		uint8 bAutoStart	= BankSequence.bAutoStart ? 1 : 0;
+		uint32 Flags = BankFlagsFromSequence(BankSequence);
+		Writer << Flags;
 
-		Writer << bLooping;
-		Writer << bAutoStart;
 		Writer << BankSequence.Position;
 		Writer << BankSequence.PlayRate;
 		Writer << BankSequence.BoundsScale;
diff --git a/Engine/Source/Runtime/Renderer/Private/Skinning/AnimBankTransformProvider.cpp b/Engine/Source/Runtime/Renderer/Private/Skinning/AnimBankTransformProvider.cpp
--- a/Engine/Source/Runtime/Renderer/Private/Skinning/AnimBankTransformProvider.cpp
+++ b/Engine/Source/Runtime/Renderer/Private/Skinning/AnimBankTransformProvider.cpp
@@ -19,6 +19,7 @@
 #include "ScenePrivate.h"
 #include "SceneView.h"
 #include "SkeletalRenderPublic.h"
+#include "SkinningDefinitions.h"
 
 static FGuid AnimBankGPUProviderId(0xA5C0027A, 0x8F884C7C, 0x9312F138, 0x71A9300F);
 static FGuid AnimBankCPUProviderId(0xE7D6173D, 0x246F431A, 0x912D384E, 0x156C0D2C);
@@ -160,13 +161,16 @@ FAnimBankRecordHandle FAnimBankTransformProvider::RegisterBank(const FAnimBankDe
 
 		// Calculate total number of sequence rot/pos keys
 		const FAnimBankData& BankData = Desc.BankAsset.Get()->GetData();
-		const uint32 KeyCount = BankData.Entries[Desc.SequenceIndex].KeyCount;
-		BankRecord.KeyOffset		= BankAllocator.Allocate(KeyCount);
-		BankRecord.Playing = 1;// TODO: BankData.Entries[Desc.SequenceIndex].Flags;
-		BankRecord.FrameCount	= BankData.Entries[Desc.SequenceIndex].FrameCount;
-		BankRecord.PositionKeys = BankData.Entries[Desc.SequenceIndex].PositionKeys;
-		BankRecord.RotationKeys = BankData.Entries[Desc.SequenceIndex].RotationKeys;
-		BankRecord.AssetMapping = BankData.Mapping;
+
+		check(Desc.SequenceIndex < uint32(BankData.Entries.Num()));
+		const FAnimBankEntry& BankEntry = BankData.Entries[Desc.SequenceIndex];
+
+		BankRecord.KeyOffset	= BankAllocator.Allocate(BankEntry.KeyCount);
+		BankRecord.Playing		= BankEntry.IsAutoStart() ? 1 : 0;
+		BankRecord.FrameCount	= BankEntry.FrameCount;
+		BankRecord.PositionKeys = BankEntry.PositionKeys; // TODO: Avoid copy
+		BankRecord.RotationKeys = BankEntry.RotationKeys; // TODO: Avoid copy
+		BankRecord.AssetMapping = BankData.Mapping; // TODO: Avoid copy
 	}
 
 	++BankRecord.ReferenceCount;
@@ -220,43 +224,6 @@ struct FAnimBankCPUData
 	bool bValid = false;
 };
 
-// TODO: Share with GPU
-struct FBankRecordHeader
-{
-	uint32 BoneCount;
-	uint32 FrameCount;
-
-	float SampleRate;
-	float PlayRate;
-
-	float CurrentTime;
-	float PreviousTime;
-
-	uint32 TransformOffset	: 29;
-	uint32 Playing			: 1;
-	uint32 Interpolating	: 1;
-	uint32 HasScale			: 1;
-};
-
-// TODO: Share with GPU
-struct FBoneBlockHeader
-{
-	uint32 BlockLocalIndex;
-	uint32 BlockBoneCount;
-	uint32 BlockTransformOffset;
-	uint32 BankRecordOffset;
-};
-
-// TODO: Share with GPU
-struct FScatterBlockHeader
-{
-	uint32 BlockLocalIndex;
-	uint32 BlockSrcTransformOffset;
-	uint32 BlockDstTransformOffset;
-	uint32 BlockTransformCount : 8;
-	uint32 TotalTransformCount : 24;
-};
-
 static FAnimBankGPUData BuildAnimBankGPUData(const FAnimBankRecordMap& BankRecordMap, FRDGBuilder& GraphBuilder)
 {
 	const uint32 BonesPerGroup = FAnimBankEvaluateCS::BonesPerGroup;
@@ -299,14 +266,14 @@ static FAnimBankGPUData BuildAnimBankGPUData(const FAnimBankRecordMap& BankRecor
 		}
 
 		// Bank Record Data
-		const uint32 BankHeaderSize = sizeof(FBankRecordHeader) * BankData.RecordCount;
+		const uint32 BankHeaderSize = sizeof(UE::HLSL::FBankRecordHeader) * BankData.RecordCount;
 		const uint32 BankMappingSize = sizeof(float) * 7 * BankData.TransformCount;
 		const uint32 BankKeySize = sizeof(float) * 7 * BankData.KeyCount;
 		const uint32 BankDataSize = BankHeaderSize + BankMappingSize + BankKeySize;
 		uint8* BankRecordData = GraphBuilder.AllocPODArray<uint8>(BankDataSize);
 
 		// Bone Block Headers
-		FBoneBlockHeader* BlockHeaders = GraphBuilder.AllocPODArray<FBoneBlockHeader>(BankData.BlockCount);
+		UE::HLSL::FBankBlockHeader* BlockHeaders = GraphBuilder.AllocPODArray<UE::HLSL::FBankBlockHeader>(BankData.BlockCount);
 
 		// Construct Headers and Indirections
 		{
@@ -329,8 +296,8 @@ static FAnimBankGPUData BuildAnimBankGPUData(const FAnimBankRecordMap& BankRecor
 
 				const uint32 BankRecordOffset = UE_PTRDIFF_TO_UINT32(BankRecordWrite - BankRecordData);
 
-				FBankRecordHeader& Header = *reinterpret_cast<FBankRecordHeader*>(BankRecordWrite);
-				BankRecordWrite += sizeof(FBankRecordHeader);
+				UE::HLSL::FBankRecordHeader& Header = *reinterpret_cast<UE::HLSL::FBankRecordHeader*>(BankRecordWrite);
+				BankRecordWrite += sizeof(UE::HLSL::FBankRecordHeader);
 
 				Header.BoneCount		= BoneCount;
 				Header.FrameCount		= BankRecord.FrameCount;
@@ -427,9 +394,9 @@ static FAnimBankGPUData BuildAnimBankGPUData(const FAnimBankRecordMap& BankRecor
 		BankData.BoneBlockBuffer = CreateByteAddressBuffer(
 			GraphBuilder,
 			TEXT("AnimBank.BlockHeaders"),
-			FMath::RoundUpToPowerOfTwo(sizeof(FBoneBlockHeader) * BankData.BlockCount),
+			FMath::RoundUpToPowerOfTwo(sizeof(UE::HLSL::FBankBlockHeader) * BankData.BlockCount),
 			BlockHeaders,
-			sizeof(FBoneBlockHeader) * BankData.BlockCount,
+			sizeof(UE::HLSL::FBankBlockHeader) * BankData.BlockCount,
 			// The buffer data is allocated above on the RDG timeline
 			ERDGInitialDataFlags::NoCopy
 		);
@@ -613,7 +580,7 @@ void FAnimBankTransformProvider::ScatterAnimation(
 		BlockCount += FMath::DivideAndRoundUp(TransformCount, BonesPerGroup) * AnimationCount;
 	}
 
-	FScatterBlockHeader* BlockHeaders = GraphBuilder.AllocPODArray<FScatterBlockHeader>(BlockCount);
+	UE::HLSL::FBankScatterHeader* BlockHeaders = GraphBuilder.AllocPODArray<UE::HLSL::FBankScatterHeader>(BlockCount);
 
 	uint32 BlockWrite = 0;
 	for (const FUintVector2& Indirection : Context.Indirections)
@@ -687,9 +654,9 @@ void FAnimBankTransformProvider::ScatterAnimation(
 	FRDGBufferRef ScatterBlockHeaders = CreateByteAddressBuffer(
 		GraphBuilder,
 		TEXT("AnimBank.ScatterHeaders"),
-		FMath::RoundUpToPowerOfTwo(sizeof(FScatterBlockHeader) * BlockCount),
+		FMath::RoundUpToPowerOfTwo(sizeof(UE::HLSL::FBankScatterHeader) * BlockCount),
 		BlockHeaders,
-		sizeof(FScatterBlockHeader) * BlockCount,
+		sizeof(UE::HLSL::FBankScatterHeader) * BlockCount,
 		// The buffer data is allocated above on the RDG timeline
 		ERDGInitialDataFlags::NoCopy
 	);
