Fix off by one error for FTexturePagePool::GetNumPages()
We should take into account the ReservedPages to be consistent with GetNumLockedPages(), GetNumMappedPages() etc
#rnx

diff --git a/Engine/Source/Runtime/Renderer/Private/VT/TexturePagePool.h b/Engine/Source/Runtime/Renderer/Private/VT/TexturePagePool.h
--- a/Engine/Source/Runtime/Renderer/Private/VT/TexturePagePool.h
+++ b/Engine/Source/Runtime/Renderer/Private/VT/TexturePagePool.h
@@ -23,15 +23,15 @@ struct FVTProducerDescription;
 class FTexturePagePool
 {
 public:
-				FTexturePagePool();
-				~FTexturePagePool();
+	FTexturePagePool();
+	~FTexturePagePool();
 
 	void Initialize(uint32 InNumPages);
 
 	FCriticalSection& GetLock() { return CriticalSection; }
 
-	uint32 GetNumPages() const { return NumPages; }
-	uint32 GetNumLockedPages() const { return GetNumPages() - FreeHeap.Num() - NumReservedPages; }
+	uint32 GetNumPages() const { return FMath::Max(NumPages, NumReservedPages) - NumReservedPages; }
+	uint32 GetNumLockedPages() const { return GetNumPages() - FreeHeap.Num(); }
 	uint32 GetNumMappedPages() const { return NumPagesMapped; }
 	uint32 GetNumAllocatedPages() const { return NumPagesAllocated; }
 
