Fixed parallel translucency not using the correct viewport scale.

diff --git a/Engine/Source/Runtime/Renderer/Private/TranslucentRendering.cpp b/Engine/Source/Runtime/Renderer/Private/TranslucentRendering.cpp
--- a/Engine/Source/Runtime/Renderer/Private/TranslucentRendering.cpp
+++ b/Engine/Source/Runtime/Renderer/Private/TranslucentRendering.cpp
@@ -1347,9 +1347,9 @@ static void RenderTranslucencyViewInner(
 					int32(View.ViewRect.Height() * ViewportScale)),
 				PassParameters,
 				ERDGPassFlags::Raster,
-				[&View, PassParameters, MeshPass](FRDGDispatchPassBuilder& DispatchPassBuilder)
+				[&View, PassParameters, MeshPass, ViewportScale](FRDGDispatchPassBuilder& DispatchPassBuilder)
 			{
-				View.ParallelMeshDrawCommandPasses[MeshPass].Dispatch(DispatchPassBuilder, &PassParameters->InstanceCullingDrawParams);
+				View.ParallelMeshDrawCommandPasses[MeshPass].Dispatch(DispatchPassBuilder, &PassParameters->InstanceCullingDrawParams, ViewportScale);
 			});
 		}
 		else
