Ray Tracing: Add control over the decal and light grid axis to override the default choice

#rnx

diff --git a/Engine/Source/Runtime/Renderer/Private/RayTracing/RayTracingDecals.cpp b/Engine/Source/Runtime/Renderer/Private/RayTracing/RayTracingDecals.cpp
--- a/Engine/Source/Runtime/Renderer/Private/RayTracing/RayTracingDecals.cpp
+++ b/Engine/Source/Runtime/Renderer/Private/RayTracing/RayTracingDecals.cpp
@@ -27,6 +27,14 @@ TAutoConsoleVariable<int32> CVarDecalGridMaxCount(
 	ECVF_RenderThreadSafe
 );
 
+TAutoConsoleVariable<int32> CVarDecalGridAxis(
+	TEXT("r.RayTracing.DecalGrid.Axis"),
+	-1,
+	TEXT("Choose the coordinate axis along which to project the decal grid (default = -1, automatic)\n"),
+	ECVF_RenderThreadSafe
+);
+
+
 enum class EDecalsWriteFlags : uint32
 {
 	None						= 0,
@@ -395,20 +403,24 @@ void BuildDecalGrid(FRDGBuilder& GraphBuilder, uint32 NumDecals, FRDGBufferSRVRe
 		const uint32 MaxCount = FMath::Clamp(CVarDecalGridMaxCount.GetValueOnRenderThread(), 1, FMath::Min((int32)NumDecals, RAY_TRACING_DECAL_COUNT_MAXIMUM));
 		OutParameters.GridResolution = Resolution;
 		OutParameters.GridMaxCount = MaxCount;
+		OutParameters.GridAxis = CVarDecalGridAxis.GetValueOnRenderThread();
 
-		// pick the shortest axis
-		FVector3f Diag = OutParameters.TranslatedBoundMax - OutParameters.TranslatedBoundMin;
-		if (Diag.X < Diag.Y && Diag.X < Diag.Z)
-		{
-			OutParameters.GridAxis = 0;
-		}
-		else if (Diag.Y < Diag.Z)
-		{
-			OutParameters.GridAxis = 1;
-		}
-		else
+		if (OutParameters.GridAxis < 0 || OutParameters.GridAxis > 2)
 		{
-			OutParameters.GridAxis = 2;
+			// pick the shortest axis if user did not explicitly pick an axis
+			FVector3f Diag = OutParameters.TranslatedBoundMax - OutParameters.TranslatedBoundMin;
+			if (Diag.X < Diag.Y && Diag.X < Diag.Z)
+			{
+				OutParameters.GridAxis = 0;
+			}
+			else if (Diag.Y < Diag.Z)
+			{
+				OutParameters.GridAxis = 1;
+			}
+			else
+			{
+				OutParameters.GridAxis = 2;
+			}
 		}
 
 		FRDGTextureDesc DecalGridDesc = FRDGTextureDesc::Create2D(
diff --git a/Engine/Source/Runtime/Renderer/Private/RayTracing/RayTracingLighting.cpp b/Engine/Source/Runtime/Renderer/Private/RayTracing/RayTracingLighting.cpp
--- a/Engine/Source/Runtime/Renderer/Private/RayTracing/RayTracingLighting.cpp
+++ b/Engine/Source/Runtime/Renderer/Private/RayTracing/RayTracingLighting.cpp
@@ -32,6 +32,14 @@ TAutoConsoleVariable<int32> CVarRayTracingLightGridMaxCount(
 	ECVF_RenderThreadSafe
 );
 
+TAutoConsoleVariable<int32> CVarRayTracingLightGridAxis(
+	TEXT("r.RayTracing.LightGridAxis"),
+	-1,
+	TEXT("Choose the coordinate axis along which to project the light grid (default = -1, automatic)\n"),
+	ECVF_RenderThreadSafe
+);
+
+
 IMPLEMENT_GLOBAL_SHADER_PARAMETER_STRUCT(FRayTracingLightGrid, "RaytracingLightGridData");
 
 class FRayTracingBuildLightGridCS : public FGlobalShader
@@ -160,19 +168,23 @@ static void PrepareLightGrid(FRDGBuilder& GraphBuilder, FGlobalShaderMap* Shader
 	LightGridParameters->LightGridResolution = Resolution;
 	LightGridParameters->LightGridMaxCount = MaxCount;
 
-	// pick the shortest axis
-	FVector3f Diag = LightGridParameters->SceneLightsTranslatedBoundMax - LightGridParameters->SceneLightsTranslatedBoundMin;
-	if (Diag.X < Diag.Y && Diag.X < Diag.Z)
-	{
-		LightGridParameters->LightGridAxis = 0;
-	}
-	else if (Diag.Y < Diag.Z)
-	{
-		LightGridParameters->LightGridAxis = 1;
-	}
-	else
+	LightGridParameters->LightGridAxis = CVarRayTracingLightGridAxis.GetValueOnRenderThread();
+	if (LightGridParameters->LightGridAxis < 0 || LightGridParameters->LightGridAxis > 2)
 	{
-		LightGridParameters->LightGridAxis = 2;
+		// pick the shortest axis if user did not explicitly pick an axis
+		FVector3f Diag = LightGridParameters->SceneLightsTranslatedBoundMax - LightGridParameters->SceneLightsTranslatedBoundMin;
+		if (Diag.X < Diag.Y && Diag.X < Diag.Z)
+		{
+			LightGridParameters->LightGridAxis = 0;
+		}
+		else if (Diag.Y < Diag.Z)
+		{
+			LightGridParameters->LightGridAxis = 1;
+		}
+		else
+		{
+			LightGridParameters->LightGridAxis = 2;
+		}
 	}
 
 
