Add ReplaceTonemapper BlendLocation to the Mobile PostProcess

diff --git a/Engine/Source/Runtime/Renderer/Private/PostProcess/PostProcessing.cpp b/Engine/Source/Runtime/Renderer/Private/PostProcess/PostProcessing.cpp
--- a/Engine/Source/Runtime/Renderer/Private/PostProcess/PostProcessing.cpp
+++ b/Engine/Source/Runtime/Renderer/Private/PostProcess/PostProcessing.cpp
@@ -2399,6 +2399,11 @@ void AddMobilePostProcessingPasses(FRDGBuilder& GraphBuilder, FScene* Scene, con
 			PassSequence.AcceptOverrideIfLastPass(EPass::PostProcessMaterialAfterTonemapping, PostProcessMaterialInputs.OverrideOutput);
 		}
 
+		if(BlendableLocation == BL_ReplacingTonemapper && PassSequence.IsEnabled(EPass::Tonemap))
+		{
+			PassSequence.AcceptOverrideIfLastPass(EPass::Tonemap, PostProcessMaterialInputs.OverrideOutput);
+		}
+
 		PostProcessMaterialInputs.SetInput(GraphBuilder, EPostProcessMaterialInput::SceneColor, SceneColor);
 
 		PostProcessMaterialInputs.CustomDepthTexture = CustomDepth.Texture;
@@ -2846,65 +2851,74 @@ void AddMobilePostProcessingPasses(FRDGBuilder& GraphBuilder, FScene* Scene, con
 	
 	if (PassSequence.IsEnabled(EPass::Tonemap))
 	{
-		bool bHDRTonemapperOutput = false;
-
-		if (!BloomOutput.IsValid())
-		{
-			BloomOutput = BlackAlphaOneDummy;
-		}
-
-		bool bDoGammaOnly = !IsMobileHDR();
-
-		FRDGTextureRef ColorGradingTexture = nullptr;
+		const FPostProcessMaterialChain MaterialChain = GetPostProcessMaterialChain(View, BL_ReplacingTonemapper);
 
-		if (IStereoRendering::IsAPrimaryView(View) && !bDoGammaOnly)
-		{
-			ColorGradingTexture = AddCombineLUTPass(GraphBuilder, View);
-		}
-		// We can re-use the color grading texture from the primary view.
-		else if (View.GetTonemappingLUT())
+		if (MaterialChain.Num())
 		{
-			ColorGradingTexture = TryRegisterExternalTexture(GraphBuilder, View.GetTonemappingLUT());
+			AddPostProcessMaterialPass(BL_ReplacingTonemapper, false);
 		}
 		else
 		{
-			const FViewInfo* PrimaryView = static_cast<const FViewInfo*>(View.Family->Views[0]);
-			ColorGradingTexture = TryRegisterExternalTexture(GraphBuilder, PrimaryView->GetTonemappingLUT());
-		}
+			bool bHDRTonemapperOutput = false;
 
-		FTonemapInputs TonemapperInputs;
-		PassSequence.AcceptOverrideIfLastPass(EPass::Tonemap, TonemapperInputs.OverrideOutput);
+			if (!BloomOutput.IsValid())
+			{
+				BloomOutput = BlackAlphaOneDummy;
+			}
 
-		// This is the view family render target.
-		if (TonemapperInputs.OverrideOutput.Texture)
-		{
-			FIntRect OutputViewRect;
-			if (View.PrimaryScreenPercentageMethod == EPrimaryScreenPercentageMethod::RawOutput)
+			bool bDoGammaOnly = !IsMobileHDR();
+
+			FRDGTextureRef ColorGradingTexture = nullptr;
+
+			if (IStereoRendering::IsAPrimaryView(View) && !bDoGammaOnly)
+			{
+				ColorGradingTexture = AddCombineLUTPass(GraphBuilder, View);
+			}
+			// We can re-use the color grading texture from the primary view.
+			else if (View.GetTonemappingLUT())
 			{
-				OutputViewRect = View.ViewRect;
+				ColorGradingTexture = TryRegisterExternalTexture(GraphBuilder, View.GetTonemappingLUT());
 			}
 			else
 			{
-				OutputViewRect = View.UnscaledViewRect;
+				const FViewInfo* PrimaryView = static_cast<const FViewInfo*>(View.Family->Views[0]);
+				ColorGradingTexture = TryRegisterExternalTexture(GraphBuilder, PrimaryView->GetTonemappingLUT());
 			}
-			ERenderTargetLoadAction  OutputLoadAction = View.IsFirstInFamily() ? ERenderTargetLoadAction::EClear : ERenderTargetLoadAction::ELoad;
 
-			TonemapperInputs.OverrideOutput.ViewRect = OutputViewRect;
-			TonemapperInputs.OverrideOutput.LoadAction = OutputLoadAction;
-			TonemapperInputs.OverrideOutput.UpdateVisualizeTextureExtent();
+			FTonemapInputs TonemapperInputs;
+			PassSequence.AcceptOverrideIfLastPass(EPass::Tonemap, TonemapperInputs.OverrideOutput);
+
+			// This is the view family render target.
+			if (TonemapperInputs.OverrideOutput.Texture)
+			{
+				FIntRect OutputViewRect;
+				if (View.PrimaryScreenPercentageMethod == EPrimaryScreenPercentageMethod::RawOutput)
+				{
+					OutputViewRect = View.ViewRect;
+				}
+				else
+				{
+					OutputViewRect = View.UnscaledViewRect;
+				}
+				ERenderTargetLoadAction  OutputLoadAction = View.IsFirstInFamily() ? ERenderTargetLoadAction::EClear : ERenderTargetLoadAction::ELoad;
+
+				TonemapperInputs.OverrideOutput.ViewRect = OutputViewRect;
+				TonemapperInputs.OverrideOutput.LoadAction = OutputLoadAction;
+				TonemapperInputs.OverrideOutput.UpdateVisualizeTextureExtent();
+			}
+
+			TonemapperInputs.SceneColor = FScreenPassTextureSlice::CreateFromScreenPassTexture(GraphBuilder, SceneColor);
+			TonemapperInputs.Bloom = BloomOutput;
+			TonemapperInputs.EyeAdaptationParameters = &EyeAdaptationParameters;
+			TonemapperInputs.ColorGradingTexture = ColorGradingTexture;
+			TonemapperInputs.bWriteAlphaChannel = View.AntiAliasingMethod == AAM_FXAA || IsPostProcessingWithAlphaChannelSupported() || bUseMobileDof || IsMobilePropagateAlphaEnabled(View.GetShaderPlatform());
+			TonemapperInputs.bOutputInHDR = bHDRTonemapperOutput;
+			TonemapperInputs.bGammaOnly = bDoGammaOnly;
+			TonemapperInputs.bMetalMSAAHDRDecode = bMetalMSAAHDRDecode;
+			TonemapperInputs.EyeAdaptationBuffer = bUseEyeAdaptation ? LastEyeAdaptationBuffer : nullptr;
+
+			SceneColor = AddTonemapPass(GraphBuilder, View, TonemapperInputs);
 		}
-			
-		TonemapperInputs.SceneColor = FScreenPassTextureSlice::CreateFromScreenPassTexture(GraphBuilder, SceneColor);
-		TonemapperInputs.Bloom = BloomOutput;
-		TonemapperInputs.EyeAdaptationParameters = &EyeAdaptationParameters;
-		TonemapperInputs.ColorGradingTexture = ColorGradingTexture;
-		TonemapperInputs.bWriteAlphaChannel = View.AntiAliasingMethod == AAM_FXAA || IsPostProcessingWithAlphaChannelSupported() || bUseMobileDof || IsMobilePropagateAlphaEnabled(View.GetShaderPlatform());
-		TonemapperInputs.bOutputInHDR = bHDRTonemapperOutput;
-		TonemapperInputs.bGammaOnly = bDoGammaOnly;
-		TonemapperInputs.bMetalMSAAHDRDecode = bMetalMSAAHDRDecode;
-		TonemapperInputs.EyeAdaptationBuffer = bUseEyeAdaptation ? LastEyeAdaptationBuffer : nullptr;
-
-		SceneColor = AddTonemapPass(GraphBuilder, View, TonemapperInputs);
 
 		//The output color should been decoded to linear space after tone mapper apparently
 		bMetalMSAAHDRDecode = false;
