Fix for crash when the cached UBs have not been updated for VSM (fallout from CL 36347378).

#rnx

diff --git a/Engine/Source/Runtime/Renderer/Private/VirtualShadowMaps/VirtualShadowMapArray.cpp b/Engine/Source/Runtime/Renderer/Private/VirtualShadowMaps/VirtualShadowMapArray.cpp
--- a/Engine/Source/Runtime/Renderer/Private/VirtualShadowMaps/VirtualShadowMapArray.cpp
+++ b/Engine/Source/Runtime/Renderer/Private/VirtualShadowMaps/VirtualShadowMapArray.cpp
@@ -1749,6 +1749,12 @@ void FVirtualShadowMapArray::BuildPageAllocations(
 	PerViewParameters.Reset();
 	PerViewParameters.AddDefaulted(Views.Num());
 
+	for (FVirtualShadowMapPerViewParameters& Param : PerViewParameters)
+	{
+		Param.NumCulledLightsGrid = GraphBuilder.CreateSRV(GSystemTextures.GetDefaultStructuredBuffer(GraphBuilder, sizeof(uint32), 0u));
+		Param.LightGridData = GraphBuilder.CreateSRV(GSystemTextures.GetDefaultStructuredBuffer(GraphBuilder, sizeof(uint32), 0u));
+	}
+
 #if !UE_BUILD_SHIPPING
 	if (GDumpVSMLightNames)
 	{
