Expand verbosity of LogRenderer warning regarding BloomConvolutionTexture name.

#tests local test

diff --git a/Engine/Source/Runtime/Renderer/Private/PostProcess/PostProcessFFTBloom.cpp b/Engine/Source/Runtime/Renderer/Private/PostProcess/PostProcessFFTBloom.cpp
--- a/Engine/Source/Runtime/Renderer/Private/PostProcess/PostProcessFFTBloom.cpp
+++ b/Engine/Source/Runtime/Renderer/Private/PostProcess/PostProcessFFTBloom.cpp
@@ -646,7 +646,12 @@ void InitDomainAndGetKernel(
 
 			if (RecommendedKernelDownsscale > 1 && (CVarBloomWarnKernelResolution.GetValueOnRenderThread() > 1 || (CVarBloomWarnKernelResolution.GetValueOnRenderThread() == 1 && FDataDrivenShaderPlatformInfo::GetIsConsole(View.GetShaderPlatform()))))
 			{
-				UE_LOG(LogRenderer, Warning, TEXT("The FPostProcessSettings::BloomConvolutionTexture could have it's resolution lowered by a factor of %d to save memory."), RecommendedKernelDownsscale);
+				FString WarningTexturePathName = "";
+				if (View.FinalPostProcessSettings.BloomConvolutionTexture)
+				{
+					WarningTexturePathName = "(" + View.FinalPostProcessSettings.BloomConvolutionTexture.GetPathName() + ")";
+				}
+				UE_LOG(LogRenderer, Warning, TEXT("The FPostProcessSettings::BloomConvolutionTexture %s could have it's resolution lowered by a factor of %d to save memory."), *WarningTexturePathName, RecommendedKernelDownsscale);
 			}
 		}
 
