Removed awkward double call of setting ray tracing flags, and moved stateful scene call onto the scene renderer in preparation for merging scene renderers together.

diff --git a/Engine/Source/Runtime/Renderer/Private/DeferredShadingRenderer.cpp b/Engine/Source/Runtime/Renderer/Private/DeferredShadingRenderer.cpp
--- a/Engine/Source/Runtime/Renderer/Private/DeferredShadingRenderer.cpp
+++ b/Engine/Source/Runtime/Renderer/Private/DeferredShadingRenderer.cpp
@@ -587,10 +587,10 @@ bool FDeferredShadingSceneRenderer::SetupRayTracingPipelineStatesAndSBT(FRDGBuil
 		return false;
 	}
 
-	if (!bAnyRayTracingPassEnabled)
+	if (!FamilyPipelineState[&FFamilyPipelineState::bRayTracing])
 	{
 		return false;
-	}	
+	}
 
 	TRACE_CPUPROFILER_EVENT_SCOPE(FDeferredShadingSceneRenderer::SetupRayTracingPipelineStatesAndSBT);
 
@@ -748,7 +748,7 @@ bool FDeferredShadingSceneRenderer::SetupRayTracingPipelineStatesAndSBT(FRDGBuil
 
 void FDeferredShadingSceneRenderer::SetupRayTracingLightDataForViews(FRDGBuilder& GraphBuilder)
 {
-	if (!bAnyRayTracingPassEnabled)
+	if (!FamilyPipelineState[&FFamilyPipelineState::bRayTracing])
 	{
 		return;
 	}
@@ -781,15 +781,14 @@ bool FDeferredShadingSceneRenderer::DispatchRayTracingWorldUpdates(FRDGBuilder&
 {
 	OutDynamicGeometryScratchBuffer = nullptr;
 
-	// We only need to update ray tracing scene for the first view family, if multiple are rendered in a single scene render call.
-	if (!bShouldUpdateRayTracingScene)
+	if (!FamilyPipelineState[&FFamilyPipelineState::bRayTracing])
 	{
 		// - Nanite ray tracing instances are already pointing at the new BLASes and RayTracingDataOffsets in GPUScene have been updated
 		Nanite::GRayTracingManager.ProcessBuildRequests(GraphBuilder);
 		return false;
 	}
 
-	check(IsRayTracingEnabled() && bAnyRayTracingPassEnabled && !Views.IsEmpty());
+	check(IsRayTracingEnabled() && !Views.IsEmpty());
 
 	TRACE_CPUPROFILER_EVENT_SCOPE(FDeferredShadingSceneRenderer::DispatchRayTracingWorldUpdates);
 
@@ -856,7 +855,7 @@ bool FDeferredShadingSceneRenderer::DispatchRayTracingWorldUpdates(FRDGBuilder&
 
 void FDeferredShadingSceneRenderer::WaitForRayTracingScene(FRDGBuilder& GraphBuilder)
 {
-	check(bAnyRayTracingPassEnabled);
+	check(FamilyPipelineState[&FFamilyPipelineState::bRayTracing]);
 
 	TRACE_CPUPROFILER_EVENT_SCOPE(FDeferredShadingSceneRenderer::WaitForRayTracingScene);
 
@@ -1069,7 +1068,7 @@ void FDeferredShadingSceneRenderer::CommitFinalPipelineState()
 			GetViewPipelineStateWritable(View).Commit();
 		}
 		FamilyPipelineState.Commit();
-	} 
+	}
 }
 
 void FDeferredShadingSceneRenderer::RenderNanite(FRDGBuilder& GraphBuilder, const TArray<FViewInfo>& InViews, FSceneTextures& SceneTextures, bool bIsEarlyDepthComplete,
@@ -1487,10 +1486,7 @@ void FDeferredShadingSceneRenderer::Render(FRDGBuilder& GraphBuilder)
 		ShadowSceneRenderer->BeginRender(GraphBuilder);
 
 #if RHI_RAYTRACING
-		// Initialize ray tracing flags, in case they weren't initialized in the CreateSceneRenderers code path
-		InitializeRayTracingFlags_RenderThread();
-
-		if (bAnyRayTracingPassEnabled)
+		if (FamilyPipelineState[&FFamilyPipelineState::bRayTracing])
 		{
 			const int32 ReferenceViewIndex = 0;
 			FViewInfo& ReferenceView = Views[ReferenceViewIndex];
@@ -1783,7 +1779,7 @@ void FDeferredShadingSceneRenderer::Render(FRDGBuilder& GraphBuilder)
 	if (RendererOutput == ERendererOutput::FinalSceneColor)
 	{
 #if RHI_RAYTRACING
-		if (bAnyRayTracingPassEnabled)
+		if (FamilyPipelineState[&FFamilyPipelineState::bRayTracing])
 		{
 			RayTracing::FinishGatherInstances(
 				GraphBuilder,
@@ -2803,7 +2799,7 @@ void FDeferredShadingSceneRenderer::Render(FRDGBuilder& GraphBuilder)
 
 #if RHI_RAYTRACING
 			// Only used by ray traced shadows
-			if (IsRayTracingEnabled() && Scene->bHasLightsWithRayTracedShadows && Views[0].IsRayTracingAllowedForView())
+			if (IsRayTracingEnabled() && Views[0].bHasRayTracingShadows && Views[0].IsRayTracingAllowedForView())
 			{
 				RenderDitheredLODFadingOutMask(GraphBuilder, Views[0], SceneTextures.Depth.Target);
 			}
@@ -3499,7 +3495,7 @@ void FDeferredShadingSceneRenderer::Render(FRDGBuilder& GraphBuilder)
 
 #if RHI_RAYTRACING
 
-bool AnyRayTracingPassEnabled(const FScene* Scene, const FViewInfo& View)
+static bool AnyRayTracingPassEnabled(const FScene* Scene, const FViewInfo& View, bool bSceneHasRayTracedShadows)
 {
 	if (!IsRayTracingEnabled(View.GetShaderPlatform()) || Scene == nullptr)
 	{
@@ -3517,11 +3513,11 @@ bool AnyRayTracingPassEnabled(const FScene* Scene, const FViewInfo& View)
 		return false;
 	}
 
-	return ShouldRenderRayTracingAmbientOcclusion(View)
+	return bSceneHasRayTracedShadows
+		|| ShouldRenderRayTracingAmbientOcclusion(View)
 		|| ShouldRenderRayTracingTranslucency(View)
 		|| ShouldRenderRayTracingSkyLight(Scene->SkyLight, View.GetShaderPlatform())
 		|| ShouldRenderRayTracingShadows()
-		|| Scene->bHasLightsWithRayTracedShadows
 		|| ShouldRenderPluginRayTracingGlobalIllumination(View)
         || Lumen::AnyLumenHardwareRayTracingPassEnabled(Scene, View)
 		|| MegaLights::UseHardwareRayTracing(*View.Family);
@@ -3572,24 +3568,30 @@ bool HasRayTracedOverlay(const FSceneViewFamily& ViewFamily)
 
 void FDeferredShadingSceneRenderer::InitializeRayTracingFlags_RenderThread()
 {
-	// The result of this call is used by AnyRayTracingPassEnabled to decide if we have any RT shadows enabled
-	Scene->UpdateRayTracedLights(ViewFamily);
+	bool bRayTracingShadows = false;
+	bool bRayTracing = false;
 
-	// This function may be called twice -- once in CreateSceneRenderers and again in Render.  We deliberately skip the logic
-	// if the flag is already set, because CreateSceneRenderers fills in the correct value for "bShouldUpdateRayTracingScene"
-	// in that case, and we don't want to overwrite it.
-	if (!bAnyRayTracingPassEnabled && !ViewFamily.EngineShowFlags.HitProxies)
+	// We currently don't need a full list of RT lights, only whether there are any RT lights at all.
+	for (const FLightSceneInfoCompact& LightSceneInfoCompact : Scene->Lights)
 	{
-		for (int32 ViewIndex = 0; ViewIndex < Views.Num(); ViewIndex++)
+		if (GetLightOcclusionType(LightSceneInfoCompact, ViewFamily) == FLightOcclusionType::Raytraced)
 		{
-			bool bHasRayTracing = AnyRayTracingPassEnabled(Scene, Views[ViewIndex]);
+			bRayTracingShadows = true;
+			break;
+		}
+	}
 
-			Views[ViewIndex].bHasAnyRayTracingPass = bHasRayTracing;
+	for (FViewInfo& View : Views)
+	{
+		const bool bViewHasRayTracing = AnyRayTracingPassEnabled(Scene, View, bRayTracingShadows);
 
-			bAnyRayTracingPassEnabled |= bHasRayTracing;
-		}
+		View.bHasAnyRayTracingPass = bViewHasRayTracing;
+		View.bHasRayTracingShadows = bRayTracingShadows;
 
-		bShouldUpdateRayTracingScene = bAnyRayTracingPassEnabled;
+		bRayTracing |= bViewHasRayTracing;
 	}
+
+	FamilyPipelineState.Set(&FFamilyPipelineState::bRayTracingShadows, bRayTracingShadows);
+	FamilyPipelineState.Set(&FFamilyPipelineState::bRayTracing, bRayTracing);
 }
 #endif // RHI_RAYTRACING
diff --git a/Engine/Source/Runtime/Renderer/Private/DeferredShadingRenderer.h b/Engine/Source/Runtime/Renderer/Private/DeferredShadingRenderer.h
--- a/Engine/Source/Runtime/Renderer/Private/DeferredShadingRenderer.h
+++ b/Engine/Source/Runtime/Renderer/Private/DeferredShadingRenderer.h
@@ -313,12 +313,6 @@ public:
 	FLumenCardRenderer LumenCardRenderer;
 	FSceneCullingRenderer SceneCullingRenderer;
 
-#if RHI_RAYTRACING
-	bool bShouldUpdateRayTracingScene =  false;
-
-	void InitializeRayTracingFlags_RenderThread();
-#endif
-
 	FDeferredShadingSceneRenderer(const FSceneViewFamily* InViewFamily, FHitProxyConsumer* HitProxyConsumer);
 
 	/** Determine and commit the final state of the pipeline for the view family and views. */
@@ -453,6 +447,10 @@ public:
 
 private:
 
+#if RHI_RAYTRACING
+	void InitializeRayTracingFlags_RenderThread() override;
+#endif
+
 	/** Structure that contains the final state of deferred shading pipeline for a FViewInfo */
 	struct FPerViewPipelineState
 	{
@@ -479,6 +477,14 @@ private:
 	// Structure that contains the final state of deferred shading pipeline for the FSceneViewFamily
 	struct FFamilyPipelineState
 	{
+#if RHI_RAYTRACING
+		// Whether the scene has lights with ray traced shadows.
+		bool bRayTracingShadows = false;
+
+		// Whether any ray tracing passes are enabled.
+		bool bRayTracing = false;
+#endif
+
 		// Whether Nanite is enabled.
 		bool bNanite;
 
diff --git a/Engine/Source/Runtime/Renderer/Private/PrimitiveSceneInfo.cpp b/Engine/Source/Runtime/Renderer/Private/PrimitiveSceneInfo.cpp
--- a/Engine/Source/Runtime/Renderer/Private/PrimitiveSceneInfo.cpp
+++ b/Engine/Source/Runtime/Renderer/Private/PrimitiveSceneInfo.cpp
@@ -962,29 +962,6 @@ void FScene::RefreshRayTracingInstances()
 	FPrimitiveSceneInfo::UpdateCachedRayTracingInstances(this, Primitives);
 }
 
-void FScene::UpdateRayTracedLights(const FSceneViewFamily& ViewFamily)
-{
-	// Whether a light can use ray traced shadows depends on CVars that may be changed at runtime.
-	// It is not enough to check the light shadow mode when light is added to the scene. This must be done during rendering.
-
-	bHasLightsWithRayTracedShadows = false;
-
-	if (!IsRayTracingEnabled())
-	{
-		return;
-	}
-
-	// We currently don't need a full list of RT lights, only whether there are any RT lights at all.
-	for (const FLightSceneInfoCompact& LightSceneInfoCompact : Lights)
-	{
-		if (GetLightOcclusionType(LightSceneInfoCompact, ViewFamily) == FLightOcclusionType::Raytraced)
-		{
-			bHasLightsWithRayTracedShadows = true;
-			break;
-		}
-	}
-}
-
 void FPrimitiveSceneInfo::UpdateCachedRayTracingInstances(FScene* Scene, const TArrayView<FPrimitiveSceneInfo*>& SceneInfos)
 {
 	if (IsRayTracingEnabled())
diff --git a/Engine/Source/Runtime/Renderer/Private/RayTracing/RayTracing.cpp b/Engine/Source/Runtime/Renderer/Private/RayTracing/RayTracing.cpp
--- a/Engine/Source/Runtime/Renderer/Private/RayTracing/RayTracing.cpp
+++ b/Engine/Source/Runtime/Renderer/Private/RayTracing/RayTracing.cpp
@@ -159,7 +159,7 @@ namespace RayTracing
 	{
 		bTranslucentGeometry = false;
 		LumenHardwareRayTracing::SetRayTracingSceneOptions(View, DiffuseIndirectMethod, ReflectionsMethod, *this);
-		RayTracingShadows::SetRayTracingSceneOptions(Scene.bHasLightsWithRayTracedShadows, *this);
+		RayTracingShadows::SetRayTracingSceneOptions(View.bHasRayTracingShadows, *this);
 
 		if (ShouldRenderRayTracingTranslucency(View))
 		{
diff --git a/Engine/Source/Runtime/Renderer/Private/RayTracing/RayTracingShadows.cpp b/Engine/Source/Runtime/Renderer/Private/RayTracing/RayTracingShadows.cpp
--- a/Engine/Source/Runtime/Renderer/Private/RayTracing/RayTracingShadows.cpp
+++ b/Engine/Source/Runtime/Renderer/Private/RayTracing/RayTracingShadows.cpp
@@ -322,7 +322,7 @@ void FDeferredShadingSceneRenderer::PrepareRayTracingShadows(const FViewInfo& Vi
 	// We have to check if ray tracing is enabled on any of the scene lights. The Scene.bHasRayTracedLights is computed using ShouldRenderRayTracingShadowsForLight() helper, 
 	// which handles various override conditions.
 
-	if (Scene.bHasLightsWithRayTracedShadows == false)
+	if (!View.bHasRayTracingShadows)
 	{
 		return;
 	}
diff --git a/Engine/Source/Runtime/Renderer/Private/RayTracing/RaytracingOptions.h b/Engine/Source/Runtime/Renderer/Private/RayTracing/RaytracingOptions.h
--- a/Engine/Source/Runtime/Renderer/Private/RayTracing/RaytracingOptions.h
+++ b/Engine/Source/Runtime/Renderer/Private/RayTracing/RaytracingOptions.h
@@ -76,7 +76,6 @@ namespace RayTracing
 // Whether a particular effect should be used, taking into account debug override
 extern bool ShouldRenderRayTracingEffect(bool bEffectEnabled, ERayTracingPipelineCompatibilityFlags CompatibilityFlags, const FSceneView* View);
 
-extern bool AnyRayTracingPassEnabled(const FScene* Scene, const FViewInfo& View);
 extern bool AnyInlineRayTracingPassEnabled(const FScene* Scene, const FViewInfo& View);
 extern FRayTracingPrimaryRaysOptions GetRayTracingTranslucencyOptions(const FViewInfo& View);
 
@@ -109,11 +108,6 @@ FORCEINLINE bool ShouldRenderRayTracingEffect(bool bEffectEnabled, ERayTracingPi
 	return false;
 }
 
-FORCEINLINE bool AnyRayTracingPassEnabled(const FScene* Scene, const FViewInfo& View)
-{
-	return false;
-}
-
 FORCEINLINE bool ShouldRenderRayTracingSkyLight(const FSkyLightSceneProxy* SkyLightSceneProxy, EShaderPlatform ShaderPlatform)
 {
 	return false;
diff --git a/Engine/Source/Runtime/Renderer/Private/ScenePrivate.h b/Engine/Source/Runtime/Renderer/Private/ScenePrivate.h
--- a/Engine/Source/Runtime/Renderer/Private/ScenePrivate.h
+++ b/Engine/Source/Runtime/Renderer/Private/ScenePrivate.h
@@ -3179,9 +3179,6 @@ public:
 	
 	/** SBT object to be used with hardware ray tracing */
 	FRayTracingShaderBindingTable RayTracingSBT;
-
-	bool bHasLightsWithRayTracedShadows = false;
-	void UpdateRayTracedLights(const FSceneViewFamily& ViewFamily);
 #endif // RHI_RAYTRACING
 
 	/** Distance field object scene data. */
diff --git a/Engine/Source/Runtime/Renderer/Private/SceneRendering.cpp b/Engine/Source/Runtime/Renderer/Private/SceneRendering.cpp
--- a/Engine/Source/Runtime/Renderer/Private/SceneRendering.cpp
+++ b/Engine/Source/Runtime/Renderer/Private/SceneRendering.cpp
@@ -3873,6 +3873,7 @@ IVisibilityTaskData* FSceneRenderer::OnRenderBegin(FRDGBuilder& GraphBuilder)
 	{
 #if RHI_RAYTRACING
 		RayTracing::OnRenderBegin(*Scene, Views, ViewFamily);
+		InitializeRayTracingFlags_RenderThread();
 #endif
 
 		if (!ViewFamily.ViewExtensions.IsEmpty())
@@ -4783,40 +4784,6 @@ void FSceneRenderer::CreateSceneRenderers(TArrayView<const FSceneViewFamily*> In
 		}
 	}
 
-#if RHI_RAYTRACING
-	// Update ray tracing flags.  We need to do this on the render thread, since "AnyRayTracingPassEnabled" accesses CVars which
-	// are only visible on the render thread.
-	if (ShadingPath == EShadingPath::Deferred)
-	{
-		ENQUEUE_RENDER_COMMAND(InitializeMultiRendererRayTracingFlags)(
-			[SceneRenderers = CopyTemp(OutSceneRenderers)](FRHICommandList& RHICmdList)
-		{
-			// For multi-view-family scene rendering, we need to determine which scene renderer will update the ray tracing
-			// scene.  This will be the first view family that uses ray tracing, and subsequent families that use ray
-			// tracing can skip that step.  If the optimization to update the ray tracing scene once is disabled, we'll
-			// update it for all scene renders.
-			bool bShouldUpdateRayTracingScene = true;
-			const bool bRayTracingSceneUpdateOnce = CVarRayTracingSceneUpdateOnce.GetValueOnRenderThread() != 0;
-
-			for (int32 RendererIndex = 0; RendererIndex < SceneRenderers.Num(); RendererIndex++)
-			{
-				FDeferredShadingSceneRenderer* SceneRenderer = (FDeferredShadingSceneRenderer*)SceneRenderers[RendererIndex];
-
-				SceneRenderer->InitializeRayTracingFlags_RenderThread();
-
-				if (SceneRenderer->bAnyRayTracingPassEnabled)
-				{
-					SceneRenderer->bShouldUpdateRayTracingScene = bShouldUpdateRayTracingScene;
-					if (bRayTracingSceneUpdateOnce)
-					{
-						bShouldUpdateRayTracingScene = false;
-					}
-				}
-			}
-		});
-	}
-#endif  // RHI_RAYTRACING
-
 #if !UE_BUILD_SHIPPING
 	if (bSplitScreenDebug)
 	{
diff --git a/Engine/Source/Runtime/Renderer/Private/SceneRendering.h b/Engine/Source/Runtime/Renderer/Private/SceneRendering.h
--- a/Engine/Source/Runtime/Renderer/Private/SceneRendering.h
+++ b/Engine/Source/Runtime/Renderer/Private/SceneRendering.h
@@ -1548,6 +1548,7 @@ public:
 	int32 PathTracingVolumetricCloudCallableShaderIndex = -1;
 
 	bool bHasAnyRayTracingPass = false;
+	bool bHasRayTracingShadows = false;
 #endif // RHI_RAYTRACING
 
 	/**
@@ -2074,11 +2075,6 @@ public:
 	bool bIsFirstSceneRenderer;
 	bool bIsLastSceneRenderer;
 
-#if RHI_RAYTRACING
-	/** Wheter any ray tracing pass has been enabled on the current frame. */
-	bool bAnyRayTracingPassEnabled = false;
-#endif
-
 public:
 
 	FSceneRenderer(const FSceneViewFamily* InViewFamily, FHitProxyConsumer* HitProxyConsumer);
@@ -2549,6 +2545,10 @@ private:
 
 	bool bShadowDepthRenderCompleted;
 
+#if RHI_RAYTRACING
+	virtual void InitializeRayTracingFlags_RenderThread() {}
+#endif
+
 	friend class FRendererModule;
 };
 
@@ -2970,4 +2970,4 @@ bool IsPrimitiveAlphaHoldoutEnabledForAnyView(TArrayView<FViewInfo> Views);
 bool SceneCaptureRequiresAlphaChannel(const FSceneView& View);
 
 /** Checks whether the material and scene proxy combination will modify the mesh position. Useful for determining whether the material can be substituted with a default material for depth rendering etc. */
-bool DoMaterialAndPrimitiveModifyMeshPosition(const FMaterial& Material, const FPrimitiveSceneProxy* PrimitiveSceneProxy);
+bool DoMaterialAndPrimitiveModifyMeshPosition(const FMaterial& Material, const FPrimitiveSceneProxy* PrimitiveSceneProxy);
\ No newline at end of file
diff --git a/Engine/Source/Runtime/Renderer/Private/SceneVisibility.cpp b/Engine/Source/Runtime/Renderer/Private/SceneVisibility.cpp
--- a/Engine/Source/Runtime/Renderer/Private/SceneVisibility.cpp
+++ b/Engine/Source/Runtime/Renderer/Private/SceneVisibility.cpp
@@ -1595,7 +1595,7 @@ void FRelevancePacket::ComputeRelevance(FDynamicPrimitiveIndexList& DynamicPrimi
 								}
 #if RHI_RAYTRACING
 								// Only used by ray traced shadows
-								if (IsRayTracingEnabled() && Scene.bHasLightsWithRayTracedShadows && View.IsRayTracingAllowedForView())
+								if (IsRayTracingEnabled() && View.bHasRayTracingShadows && View.IsRayTracingAllowedForView())
 								{
 									if (MarkMask & EMarkMaskBits::StaticMeshFadeOutDitheredLODMapMask)
 									{
diff --git a/Engine/Source/Runtime/Renderer/Private/VolumetricFog.cpp b/Engine/Source/Runtime/Renderer/Private/VolumetricFog.cpp
--- a/Engine/Source/Runtime/Renderer/Private/VolumetricFog.cpp
+++ b/Engine/Source/Runtime/Renderer/Private/VolumetricFog.cpp
@@ -545,7 +545,7 @@ static void RenderRaytracedDirectionalShadowVolume(
 
 void FDeferredShadingSceneRenderer::PrepareRayTracingVolumetricFogShadows(const FViewInfo& View, const FScene& Scene, TArray<FRHIRayTracingShader*>& OutRayGenShaders)
 {
-	const bool bEnabled = Scene.bHasLightsWithRayTracedShadows && View.IsRayTracingAllowedForView() && ::ShouldRenderVolumetricFog(&Scene, *View.Family) && GVolumetricFogInjectRaytracedLights;
+	const bool bEnabled = View.bHasRayTracingShadows && View.IsRayTracingAllowedForView() && ::ShouldRenderVolumetricFog(&Scene, *View.Family) && GVolumetricFogInjectRaytracedLights;
 	if (!bEnabled)
 	{
 		return;
@@ -600,7 +600,7 @@ bool LightHasRayTracedShadows(const FLightSceneInfo* LightSceneInfo, const FScen
 bool LightNeedsSeparateInjectionIntoVolumetricFogForOpaqueShadow(const FViewInfo& View, const FLightSceneInfo* LightSceneInfo, const FVisibleLightInfo& VisibleLightInfo, const FScene& InScene)
 {
 #if RHI_RAYTRACING
-	bool bTestRayTracedShadows = InScene.bHasLightsWithRayTracedShadows && View.IsRayTracingAllowedForView();
+	bool bTestRayTracedShadows = View.bHasRayTracingShadows && View.IsRayTracingAllowedForView();
 #else
 	bool bTestRayTracedShadows = false;
 #endif
@@ -1430,7 +1430,7 @@ void FSceneRenderer::ComputeVolumetricFog(FRDGBuilder& GraphBuilder,
 					if ((View.ViewMatrices.GetViewOrigin() - LightBounds.Center).SizeSquared() < (FogInfo.VolumetricFogDistance + LightBounds.W) * (FogInfo.VolumetricFogDistance + LightBounds.W))
 					{
 #if RHI_RAYTRACING
-						const bool bRayTracedLight = (Scene->bHasLightsWithRayTracedShadows && View.IsRayTracingAllowedForView()) ? LightHasRayTracedShadows(LightSceneInfo, ViewFamily) : false;
+						const bool bRayTracedLight = (View.bHasRayTracingShadows && View.IsRayTracingAllowedForView()) ? LightHasRayTracedShadows(LightSceneInfo, ViewFamily) : false;
 						if (bRayTracedLight)
 						{
 							LightsToInject.RayTracedLights.Add(LightSceneInfo);
