Implemented animation bank support for negative play rate (animation in reverse), a global time scale debug cvar (r.AnimBank.TimeScale), as well as hooked up initial position (time), and fixed play rate and key count.

diff --git a/Engine/Shaders/Private/Skinning/AnimBankEval.usf b/Engine/Shaders/Private/Skinning/AnimBankEval.usf
--- a/Engine/Shaders/Private/Skinning/AnimBankEval.usf
+++ b/Engine/Shaders/Private/Skinning/AnimBankEval.usf
@@ -89,7 +89,12 @@ void BankEvaluateCS(uint3 GroupId : SV_GroupID, uint GroupIndex : SV_GroupIndex)
 
 			const float SampleInterval = 1.0f / BankRecordHeader.SampleRate;
 			const float TrackLength = (BankRecordHeader.FrameCount - 1) * SampleInterval;
-			const float Time = fmod(BankRecordHeader.CurrentTime, TrackLength);
+
+			float Time = fmod(BankRecordHeader.CurrentTime, TrackLength);
+			if (Time < 0.0f) 
+			{
+				Time += TrackLength;
+			}
 
 			const float FramesPerSecond = -1.0f;
 			FKeyIndices KeyIndices = GetKeyIndicesFromTime(Time, BankRecordHeader.FrameCount, TrackLength, FramesPerSecond);
diff --git a/Engine/Source/Runtime/Engine/Classes/Animation/AnimBank.h b/Engine/Source/Runtime/Engine/Classes/Animation/AnimBank.h
--- a/Engine/Source/Runtime/Engine/Classes/Animation/AnimBank.h
+++ b/Engine/Source/Runtime/Engine/Classes/Animation/AnimBank.h
@@ -329,6 +329,7 @@ struct FAnimBankRecord
 	FAnimBankDesc Desc;
 	int32 RecordId = INDEX_NONE;
 	int32 KeyOffset = INDEX_NONE;
+	uint32 KeyCount = 0;
 	int32 FrameCount = 0;
 	int32 ReferenceCount = 0;
 
diff --git a/Engine/Source/Runtime/Engine/Private/Animation/AnimBank.cpp b/Engine/Source/Runtime/Engine/Private/Animation/AnimBank.cpp
--- a/Engine/Source/Runtime/Engine/Private/Animation/AnimBank.cpp
+++ b/Engine/Source/Runtime/Engine/Private/Animation/AnimBank.cpp
@@ -516,8 +516,6 @@ bool FAnimBankBuildAsyncCacheTask::BuildData(const UE::FSharedString& Name, cons
 
 		const FCompressedAnimSequence& SampledSequence = Sequence->CompressedData;
 
-		BankEntry.KeyCount = Sequence->GetNumberOfSampledKeys();
-
 		// Set up mapping tables for the decompressor to map internal tracks to the pose array (which is in bone order).
 		const int32 NumTracks = Sequence->CompressedData.CompressedTrackToSkeletonMapTable.Num();
 
@@ -550,6 +548,13 @@ bool FAnimBankBuildAsyncCacheTask::BuildData(const UE::FSharedString& Name, cons
 		BankEntry.FrameCount = Sequence->GetSamplingFrameRate().AsFrameTime(Sequence->GetPlayLength()).RoundToFrame().Value;
 		check(BankEntry.FrameCount > 0);
 
+		BankEntry.KeyCount = BankEntry.FrameCount * NumMeshBones;//  Sequence->GetNumberOfSampledKeys();
+		check(BankEntry.KeyCount > 0);
+
+		BankEntry.Flags = BankFlagsFromSequence(BankSequence);
+		BankEntry.Position = BankSequence.Position;
+		BankEntry.PlayRate = BankSequence.PlayRate;
+
 		// The sampled pose transforms from the sequence are in bone (parent) space. We need to convert them
 		// to local space (component space) to retarget to the mesh and create ref -> anim pose transform.  
 		TArray<FTransform> SampledLocalAnimPose = AnimRefSkeleton.GetRawRefBonePose();
@@ -559,8 +564,8 @@ bool FAnimBankBuildAsyncCacheTask::BuildData(const UE::FSharedString& Name, cons
 		SampledLocalMeshPose.SetNumUninitialized(NumMeshBones);
 		SampledGlobalMeshPose.SetNumUninitialized(NumMeshBones);
 
-		BankEntry.PositionKeys.SetNumUninitialized(BankEntry.FrameCount * NumMeshBones);
-		BankEntry.RotationKeys.SetNumUninitialized(BankEntry.FrameCount * NumMeshBones);
+		BankEntry.PositionKeys.SetNumUninitialized(BankEntry.KeyCount);
+		BankEntry.RotationKeys.SetNumUninitialized(BankEntry.KeyCount);
 
 		// Initialize bounds to the mesh vertex positions in reference pose.
 		FVector3f AnimatedBoundsMin(AssetBounds.Origin - AssetBounds.BoxExtent);
@@ -640,8 +645,6 @@ bool FAnimBankBuildAsyncCacheTask::BuildData(const UE::FSharedString& Name, cons
 		BankEntry.SampledBounds.BoxExtent *= BankSequence.BoundsScale;
 		BankEntry.SampledBounds.SphereRadius *= BankSequence.BoundsScale;
 
-		BankEntry.Flags = BankFlagsFromSequence(BankSequence);
-
 		check(BankEntry.PositionKeys.Num() > 0);
 		check(BankEntry.RotationKeys.Num() > 0);
 		check(BankEntry.FrameCount > 0);
@@ -993,7 +996,7 @@ FIoHash UAnimBank::CreateDerivedDataKeyHash(const ITargetPlatform* TargetPlatfor
 	}
 
 	FMemoryHasherBlake3 Writer;
-	FGuid AnimBankVersionGuid(0x40961173, 0xCBAF4D9D, 0xF03ABA3C, 0x30EC8690);
+	FGuid AnimBankVersionGuid(0x50961173, 0xDBAF4D9D, 0xE03ABA3C, 0x40EC8690);
 	Writer << AnimBankVersionGuid;
 
 	if (IsValid(Asset))
diff --git a/Engine/Source/Runtime/Renderer/Private/Skinning/AnimBankTransformProvider.cpp b/Engine/Source/Runtime/Renderer/Private/Skinning/AnimBankTransformProvider.cpp
--- a/Engine/Source/Runtime/Renderer/Private/Skinning/AnimBankTransformProvider.cpp
+++ b/Engine/Source/Runtime/Renderer/Private/Skinning/AnimBankTransformProvider.cpp
@@ -26,14 +26,22 @@ static FGuid AnimBankCPUProviderId(0xE7D6173D, 0x246F431A, 0x912D384E, 0x156C0D2
 
 IMPLEMENT_SCENE_EXTENSION(FAnimBankTransformProvider);
 
+// Animation is always sampled at 30hz
+#define ANIM_BANK_SAMPLE_RATE 30.0f
+
 int32 GAnimBankInterp = 1;
 FAutoConsoleVariableRef CVarAnimBankInterp(
 	TEXT("r.AnimBank.Interpolation"),
 	GAnimBankInterp,
 	TEXT(""),
-	FConsoleVariableDelegate::CreateLambda([](IConsoleVariable* InVariable)
-	{
-	}),
+	ECVF_Scalability | ECVF_RenderThreadSafe
+);
+
+float GAnimBankTimeScale = 1.0f;
+FAutoConsoleVariableRef CVarAnimBankTimeScale(
+	TEXT("r.AnimBank.TimeScale"),
+	GAnimBankTimeScale,
+	TEXT(""),
 	ECVF_Scalability | ECVF_RenderThreadSafe
 );
 
@@ -165,12 +173,20 @@ FAnimBankRecordHandle FAnimBankTransformProvider::RegisterBank(const FAnimBankDe
 		check(Desc.SequenceIndex < uint32(BankData.Entries.Num()));
 		const FAnimBankEntry& BankEntry = BankData.Entries[Desc.SequenceIndex];
 
+		const float SampleInterval = 1.0f / ANIM_BANK_SAMPLE_RATE;
+		const float TrackLength = float(BankEntry.FrameCount - 1) * SampleInterval;
+
 		BankRecord.KeyOffset	= BankAllocator.Allocate(BankEntry.KeyCount);
+		BankRecord.KeyCount		= BankEntry.KeyCount;
+		BankRecord.CurrentTime	= FMath::Clamp<float>(BankEntry.Position, 0.0f, TrackLength);
+		BankRecord.PreviousTime = BankRecord.CurrentTime;
 		BankRecord.Playing		= BankEntry.IsAutoStart() ? 1 : 0;
 		BankRecord.FrameCount	= BankEntry.FrameCount;
 		BankRecord.PositionKeys = BankEntry.PositionKeys; // TODO: Avoid copy
 		BankRecord.RotationKeys = BankEntry.RotationKeys; // TODO: Avoid copy
 		BankRecord.AssetMapping = BankData.Mapping; // TODO: Avoid copy
+
+		check(BankEntry.KeyCount == BankRecord.PositionKeys.Num() && BankEntry.KeyCount == BankRecord.RotationKeys.Num());
 	}
 
 	++BankRecord.ReferenceCount;
@@ -188,10 +204,7 @@ void FAnimBankTransformProvider::UnregisterBank(const FAnimBankRecordHandle& Han
 
 	if (BankRecord.ReferenceCount == 0)
 	{
-		const FAnimBankData& BankData = BankRecord.Desc.BankAsset->GetData();
-		const uint32 KeyCount = BankData.Entries[BankRecord.Desc.SequenceIndex].KeyCount;
-
-		BankAllocator.Free(BankRecord.KeyOffset, KeyCount);
+		BankAllocator.Free(BankRecord.KeyOffset, BankRecord.KeyCount);
 		BankRecordMap.RemoveByElementId(RecordId);
 	}
 }
@@ -285,8 +298,7 @@ static FAnimBankGPUData BuildAnimBankGPUData(const FAnimBankRecordMap& BankRecor
 			for (auto& RecordPair : BankRecordMap)
 			{
 				const FAnimBankRecord& BankRecord = RecordPair.Value;
-				const int32 KeyCount = BankRecord.PositionKeys.Num();
-				if (KeyCount == 0 || !BankRecord.Playing)
+				if (BankRecord.KeyCount == 0 || !BankRecord.Playing)
 				{
 					continue;
 				}
@@ -301,7 +313,7 @@ static FAnimBankGPUData BuildAnimBankGPUData(const FAnimBankRecordMap& BankRecor
 
 				Header.BoneCount		= BoneCount;
 				Header.FrameCount		= BankRecord.FrameCount;
-				Header.SampleRate		= 30.0f; // 30hz
+				Header.SampleRate		= ANIM_BANK_SAMPLE_RATE;
 				Header.PlayRate			= BankRecord.Desc.PlayRate;
 				Header.CurrentTime		= BankRecord.CurrentTime;
 				Header.PreviousTime		= BankRecord.PreviousTime;
@@ -367,7 +379,7 @@ static FAnimBankGPUData BuildAnimBankGPUData(const FAnimBankRecordMap& BankRecor
 					const FVector3f* PositionKeys	= BankRecord.PositionKeys.GetData();
 					const FQuat4f* RotationKeys		= BankRecord.RotationKeys.GetData();
 
-					for (int32 KeyIndex = 0; KeyIndex < KeyCount; ++KeyIndex)
+					for (uint32 KeyIndex = 0; KeyIndex < BankRecord.KeyCount; ++KeyIndex)
 					{
 						float* KeyWrite = reinterpret_cast<float*>(BankRecordWrite);
 
@@ -468,11 +480,14 @@ static FAnimBankCPUData BuildAnimBankCPUData(const FAnimBankRecordMap& BankRecor
 
 			BankData.IdToOffsetMapping[BankRecord.RecordId] = (TransformOffset * sizeof(FMatrix3x4));
 
-			const float SampleRate = 30.0f; // 30hz
-			const float SampleInterval = 1.0f / SampleRate;
+			const float SampleInterval = 1.0f / ANIM_BANK_SAMPLE_RATE;
 			const float TrackLength = float(BankRecord.FrameCount - 1) * SampleInterval;
 
 			double Time = (double)FMath::Fmod(BankRecord.CurrentTime, TrackLength);
+			if (Time < 0.0)
+			{
+				Time += TrackLength;
+			}
 
 			int32 KeyIndex0, KeyIndex1;
 			float Alpha;
@@ -544,6 +559,8 @@ static FAnimBankCPUData BuildAnimBankCPUData(const FAnimBankRecordMap& BankRecor
 
 void FAnimBankTransformProvider::AdvanceAnimation(FSkinningTransformProvider::FProviderContext& Context)
 {
+	const float GlobalTimeScale = GAnimBankTimeScale;
+
 	for (auto& RecordPair : BankRecordMap)
 	{
 		FAnimBankRecord& BankRecord = RecordPair.Value;
@@ -553,10 +570,7 @@ void FAnimBankTransformProvider::AdvanceAnimation(FSkinningTransformProvider::FP
 		}
 
 		BankRecord.PreviousTime = BankRecord.CurrentTime;
-
-		// TODO: Support negative rate to play in reverse
-		const float PlayRate = FMath::Abs(BankRecord.Desc.PlayRate);
-		BankRecord.CurrentTime += (Context.DeltaTime * PlayRate);
+		BankRecord.CurrentTime += (Context.DeltaTime * BankRecord.Desc.PlayRate) * GlobalTimeScale;
 	}
 }
 
