Fixed scissor rect clipping for DrawTileMesh.

diff --git a/Engine/Source/Runtime/Renderer/Private/Renderer.cpp b/Engine/Source/Runtime/Renderer/Private/Renderer.cpp
--- a/Engine/Source/Runtime/Renderer/Private/Renderer.cpp
+++ b/Engine/Source/Runtime/Renderer/Private/Renderer.cpp
@@ -418,7 +418,7 @@ void FRendererModule::DrawTileMesh(FCanvasRenderContext& RenderContext, FMeshPas
 			{
 				PassParameters->DebugViewMode = CreateDebugViewModePassUniformBuffer(GraphBuilder, View, nullptr);
 
-				AddDrawDynamicMeshPass(GraphBuilder, RDG_EVENT_NAME("OutputMaterialTextureScales"), PassParameters, View, RenderContext.GetViewportRect(),
+				AddDrawDynamicMeshPass(GraphBuilder, RDG_EVENT_NAME("OutputMaterialTextureScales"), PassParameters, View, RenderContext.GetViewportRect(), RenderContext.GetScissorRect(),
 					[Scene, &View, &Mesh](FMeshPassDrawListContext* InDrawListContext)
 				{
 					FDebugViewModeMeshProcessor PassMeshProcessor(
@@ -440,7 +440,7 @@ void FRendererModule::DrawTileMesh(FCanvasRenderContext& RenderContext, FMeshPas
 			{
 				PassParameters->TranslucentBasePass = CreateTranslucentBasePassUniformBuffer(GraphBuilder, Scene, View);
 
-				AddDrawDynamicMeshPass(GraphBuilder, RDG_EVENT_NAME("TranslucentDeferred"), PassParameters, View, RenderContext.GetViewportRect(),
+				AddDrawDynamicMeshPass(GraphBuilder, RDG_EVENT_NAME("TranslucentDeferred"), PassParameters, View, RenderContext.GetViewportRect(), RenderContext.GetScissorRect(),
 					[Scene, &View, &Mesh, DrawRenderState, bUse128bitRT](FMeshPassDrawListContext* DynamicMeshPassContext)
 				{
 					FBasePassMeshProcessor PassMeshProcessor(
@@ -461,7 +461,7 @@ void FRendererModule::DrawTileMesh(FCanvasRenderContext& RenderContext, FMeshPas
 			{
 				PassParameters->MobileBasePass = CreateMobileBasePassUniformBuffer(GraphBuilder, View, EMobileBasePass::Translucent, EMobileSceneTextureSetupMode::None);
 
-				AddDrawDynamicMeshPass(GraphBuilder, RDG_EVENT_NAME("TranslucentMobile"), PassParameters, View, RenderContext.GetViewportRect(),
+				AddDrawDynamicMeshPass(GraphBuilder, RDG_EVENT_NAME("TranslucentMobile"), PassParameters, View, RenderContext.GetViewportRect(), RenderContext.GetScissorRect(),
 					[Scene, &View, DrawRenderState, &Mesh](FDynamicPassMeshDrawListContext* DynamicMeshPassContext)
 				{
 					FMobileBasePassMeshProcessor PassMeshProcessor(
@@ -490,7 +490,7 @@ void FRendererModule::DrawTileMesh(FCanvasRenderContext& RenderContext, FMeshPas
 				ensureMsgf(HitProxyId == Mesh.BatchHitProxyId, TEXT("Only Mesh.BatchHitProxyId is used for hit testing."));
 
 #if WITH_EDITOR
-				AddDrawDynamicMeshPass(GraphBuilder, RDG_EVENT_NAME("HitTesting"), PassParameters, View, RenderContext.GetViewportRect(),
+				AddDrawDynamicMeshPass(GraphBuilder, RDG_EVENT_NAME("HitTesting"), PassParameters, View, RenderContext.GetViewportRect(), RenderContext.GetScissorRect(),
 					[Scene, &View, DrawRenderState, &Mesh](FDynamicPassMeshDrawListContext* DynamicMeshPassContext)
 				{
 					FHitProxyMeshProcessor PassMeshProcessor(
@@ -511,7 +511,7 @@ void FRendererModule::DrawTileMesh(FCanvasRenderContext& RenderContext, FMeshPas
 				{
 					PassParameters->OpaqueBasePass = CreateOpaqueBasePassUniformBuffer(GraphBuilder, View);
 
-					AddDrawDynamicMeshPass(GraphBuilder, RDG_EVENT_NAME("OpaqueDeferred"), PassParameters, View, RenderContext.GetViewportRect(),
+					AddDrawDynamicMeshPass(GraphBuilder, RDG_EVENT_NAME("OpaqueDeferred"), PassParameters, View, RenderContext.GetViewportRect(), RenderContext.GetScissorRect(),
 						[Scene, &View, DrawRenderState, &Mesh, bUse128bitRT](FDynamicPassMeshDrawListContext* DynamicMeshPassContext)
 					{
 						FBasePassMeshProcessor PassMeshProcessor(
@@ -531,7 +531,7 @@ void FRendererModule::DrawTileMesh(FCanvasRenderContext& RenderContext, FMeshPas
 				{
 					PassParameters->MobileBasePass = CreateMobileBasePassUniformBuffer(GraphBuilder, View, EMobileBasePass::Opaque, EMobileSceneTextureSetupMode::None);
 
-					AddDrawDynamicMeshPass(GraphBuilder, RDG_EVENT_NAME("OpaqueMobile"), PassParameters, View, RenderContext.GetViewportRect(),
+					AddDrawDynamicMeshPass(GraphBuilder, RDG_EVENT_NAME("OpaqueMobile"), PassParameters, View, RenderContext.GetViewportRect(), RenderContext.GetScissorRect(),
 						[Scene, &View, DrawRenderState, &Mesh](FDynamicPassMeshDrawListContext* DynamicMeshPassContext)
 					{
 						FMobileBasePassMeshProcessor PassMeshProcessor(
diff --git a/Engine/Source/Runtime/Renderer/Public/MeshPassProcessor.inl b/Engine/Source/Runtime/Renderer/Public/MeshPassProcessor.inl
--- a/Engine/Source/Runtime/Renderer/Public/MeshPassProcessor.inl
+++ b/Engine/Source/Runtime/Renderer/Public/MeshPassProcessor.inl
@@ -351,7 +351,22 @@ void AddDrawDynamicMeshPass(
 	FRDGEventName&& EventName,
 	const ParameterStructType* PassParameters,
 	const FSceneView& View,
-	FIntRect ViewRect,
+	FIntRect ViewportRect,
+	const LambdaType& BuildPassProcessorLambda,
+	bool bForceStereoInstancingOff = false,
+	bool bForceParallelSetupOff = false)
+{
+	AddDrawDynamicMeshPass(GraphBuilder, MoveTemp(EventName), PassParameters, View, ViewportRect, ViewportRect, BuildPassProcessorLambda, bForceStereoInstancingOff, bForceParallelSetupOff);
+}
+
+template <typename ParameterStructType, typename LambdaType>
+void AddDrawDynamicMeshPass(
+	FRDGBuilder& GraphBuilder,
+	FRDGEventName&& EventName,
+	const ParameterStructType* PassParameters,
+	const FSceneView& View,
+	FIntRect ViewportRect,
+	FIntRect ScissorRect,
 	const LambdaType& BuildPassProcessorLambda,
 	bool bForceStereoInstancingOff = false,
 	bool bForceParallelSetupOff = false)
@@ -381,10 +396,11 @@ void AddDrawDynamicMeshPass(
 		MoveTemp(EventName),
 		PassParameters,
 		ERDGPassFlags::Raster,
-		[&Context, &View, ViewRect, InstanceFactor](FRDGAsyncTask, FRHICommandList& RHICmdList)
+		[&Context, &View, ViewportRect, ScissorRect, InstanceFactor](FRDGAsyncTask, FRHICommandList& RHICmdList)
 	{
 		TRACE_CPUPROFILER_EVENT_SCOPE(SetupDynamicMeshPass);
-		RHICmdList.SetViewport(ViewRect.Min.X, ViewRect.Min.Y, 0.0f, ViewRect.Max.X, ViewRect.Max.Y, 1.0f);
+		RHICmdList.SetViewport(ViewportRect.Min.X, ViewportRect.Min.Y, 0.0f, ViewportRect.Max.X, ViewportRect.Max.Y, 1.0f);
+		RHICmdList.SetScissorRect(true, ScissorRect.Min.X, ScissorRect.Min.Y, ScissorRect.Max.X, ScissorRect.Max.Y);
 		DrawDynamicMeshPassPrivate(View, RHICmdList, Context.VisibleMeshDrawCommands, Context.DynamicMeshDrawCommandStorage, Context.GraphicsMinimalPipelineStateSet, Context.NeedsShaderInitialisation, InstanceFactor);
 	});
 }
\ No newline at end of file
