Path Tracer: Add support for NVidia SER extensions and implement more flexible compaction

Material rays are optionally sorted before the closest hit shader is executed. This improves performance in very complex scenes by up 10-50%. Simpler scenes can see a net slowdown, so the feature remains optional.

Also implement partial compaction by dynamically switching between the compacted kernel and the non-compacted kernel. The default depth of 4 was experimentally determined to give good performance across a wide range of scenes. In some cases, tuning this default can improve performance though.

diff --git a/Engine/Shaders/Private/PathTracing/PathTracing.usf b/Engine/Shaders/Private/PathTracing/PathTracing.usf
--- a/Engine/Shaders/Private/PathTracing/PathTracing.usf
+++ b/Engine/Shaders/Private/PathTracing/PathTracing.usf
@@ -36,14 +36,15 @@ int2 ComputeTextureIndex(uint2 DispatchIdx)
 	return DispatchIdx + TileTextureOffset;
 }
 
-#if PATH_TRACER_USE_COMPACTION == 1
+int FirstBounce;
 
-int Bounce;
 RWStructuredBuffer<FPathTracingPackedPathState> PathStateData;
 RWBuffer<uint> ActivePaths;
-RWBuffer<uint> NextActivePaths;
 RWBuffer<uint> NumPathStates;
 
+#if PATH_TRACER_USE_COMPACTION
+RWBuffer<uint> NextActivePaths;
+#endif
 
 FPathState LoadPathStateData(uint Index)
 {
@@ -63,6 +64,7 @@ FPathState LoadPathStateData(uint Index)
 	return Output;
 }
 
+#if PATH_TRACER_USE_COMPACTION
 void StorePathStateData(FPathState PathState, uint Index)
 {
 	FPathTracingPackedPathState Data;
@@ -77,6 +79,7 @@ void StorePathStateData(FPathState PathState, uint Index)
 	Data.PackedRoughnessSigma = PathState.PackedRoughnessSigma;
 	PathStateData[Index] = Data;
 }
+#endif
 
 RAY_TRACING_ENTRY_RAYGEN(PathTracingMainRG)
 {
@@ -85,14 +88,16 @@ RAY_TRACING_ENTRY_RAYGEN(PathTracingMainRG)
 
 	uint LinearPixelIndex = DispatchIdx.x + DispatchDim.x * DispatchIdx.y;
 	uint PathIndex = DispatchIdx.x + 65536 * DispatchIdx.y;
+#if PATH_TRACER_USE_COMPACTION == 1
 	if (LinearPixelIndex == 0)
 	{
 		// write other dimensions for indirect dispatch of the next bounce
-		NumPathStates[3 * Bounce + 4] = 1;
-		NumPathStates[3 * Bounce + 5] = 1;
+		NumPathStates[3 * FirstBounce + 4] = 1;
+		NumPathStates[3 * FirstBounce + 5] = 1;
 	}
+#endif
 	FPathState PathState;
-	if (Bounce == 0)
+	if (FirstBounce == 0)
 	{
 #if PATH_TRACER_USE_ADAPTIVE_SAMPLING
 		if (Iteration > 0)
@@ -112,17 +117,17 @@ RAY_TRACING_ENTRY_RAYGEN(PathTracingMainRG)
 	else
 	{
 		// compare to number of active paths from the previous bounce
-		if (LinearPixelIndex >= NumPathStates[3 * Bounce + 0])
+		if (LinearPixelIndex >= NumPathStates[3 * FirstBounce + 0])
 		{
 			return; // nothing left to do on this thread
 		}
 		PathIndex = ActivePaths[LinearPixelIndex];
 		PathState = LoadPathStateData((PathIndex & 65535) + ScanlineWidth * (PathIndex >> 16));
 	}
-
-	const bool KeepGoing = PathTracingKernel(PathState, Bounce);
+#if PATH_TRACER_USE_COMPACTION == 1
+	const bool KeepGoing = PathTracingKernel(PathState, FirstBounce);
 	const uint2 TextureIndex = ComputeTextureIndex(uint2(PathIndex & 65535, PathIndex >> 16));
-	if (Bounce == 0)
+	if (FirstBounce == 0)
 	{
 		// Camera ray must write depth right away because we don't want to have to carry DepthZ in the PackedPathState
 		PathState.WriteDepth(TextureIndex);
@@ -131,7 +136,7 @@ RAY_TRACING_ENTRY_RAYGEN(PathTracingMainRG)
 	{
 	    // Add the path to the set of valid paths for the next bounce
 		uint PathStateIndex;
-		InterlockedAdd(NumPathStates[3 * Bounce + 3], 1, PathStateIndex);
+		InterlockedAdd(NumPathStates[3 * FirstBounce + 3], 1, PathStateIndex);
 		NextActivePaths[PathStateIndex] = PathIndex;
 		StorePathStateData(PathState, (PathIndex & 65535) + ScanlineWidth * (PathIndex >> 16));
 	}
@@ -141,36 +146,8 @@ RAY_TRACING_ENTRY_RAYGEN(PathTracingMainRG)
 		// Accumulate radiance and update pixel variance
 		PathState.WritePixel(TextureIndex);
 	}
-}
-
 #else // PATH_TRACER_USE_COMPACTION == 0
-
-#if PATH_TRACER_USE_ADAPTIVE_SAMPLING
-RWBuffer<int> ActivePaths;
-RWBuffer<uint> NumPathStates;
-#endif
-
-RAY_TRACING_ENTRY_RAYGEN(PathTracingMainRG)
-{
-	uint2 DispatchIdx = DispatchRaysIndex().xy;
-	const uint2 DispatchDim = DispatchRaysDimensions().xy;
-
-#if PATH_TRACER_USE_ADAPTIVE_SAMPLING
-	if (Iteration > 0)
-	{
-		uint LinearPixelIndex = DispatchIdx.x + DispatchDim.x * DispatchIdx.y;
-		if (LinearPixelIndex >= NumPathStates[0])
-		{
-			return;
-		}
-		uint PathIndex = ActivePaths[LinearPixelIndex];
-		DispatchIdx = uint2(PathIndex & 65535, PathIndex >> 16);
-		LinearPixelIndex = DispatchIdx.x + DispatchDim.x * DispatchIdx.y;
-	}
-#endif
-	FPathState PathState = CreatePathState(ComputePixelIndex(DispatchIdx), ComputeTextureIndex(DispatchIdx));
-
-	for (int Bounce = 0; Bounce <= MaxBounces; Bounce++)
+	for (int Bounce = FirstBounce; Bounce <= MaxBounces; Bounce++)
 	{
 		if (!PathTracingKernel(PathState, Bounce))
 		{
@@ -180,9 +157,11 @@ RAY_TRACING_ENTRY_RAYGEN(PathTracingMainRG)
 	}
 
 	// Accumulate radiance and update pixel variance
-	int2 TextureIndex = ComputeTextureIndex(DispatchIdx);
+	const uint2 TextureIndex = ComputeTextureIndex(uint2(PathIndex & 65535, PathIndex >> 16));
 	PathState.WritePixel(TextureIndex);
-	PathState.WriteDepth(TextureIndex);
+	if (FirstBounce == 0)
+	{
+		PathState.WriteDepth(TextureIndex);
+	}
+#endif
 }
-
-#endif
\ No newline at end of file
diff --git a/Engine/Shaders/Private/PathTracing/PathTracingCore.ush b/Engine/Shaders/Private/PathTracing/PathTracingCore.ush
--- a/Engine/Shaders/Private/PathTracing/PathTracingCore.ush
+++ b/Engine/Shaders/Private/PathTracing/PathTracingCore.ush
@@ -39,6 +39,17 @@
 #include "../ShaderPrint.ush"
 #endif
 
+#if PATH_TRACER_USE_SER
+
+#define NV_SHADER_EXTN_SLOT u0
+#define NV_SHADER_EXTN_REGISTER_SPACE space1001
+
+#define NV_HITOBJECT_USE_MACRO_API
+#include "/Engine/Shared/ThirdParty/NVIDIA/nvHLSLExtns.h"
+
+#endif
+
+
 float BlendFactor;
 float BaseExposure;
 uint Iteration;
@@ -488,6 +499,26 @@ FPackedPathTracingPayload TraceTransparentRay(inout FPathState PathState, int Bo
 		PackedPayload.SetDBufferB(float4(0, 0, 0, 1));
 		PackedPayload.SetDBufferC(float4(0, 0, 0, 1));
 		PackedPayload.SetStochasticSlabRand(RISRandSample.z);
+#if PATH_TRACER_USE_SER 
+		{
+			NvHitObject Hit;
+			NvTraceRayHitObject(
+				TLAS,
+				RayFlags,
+				InstanceInclusionMask,
+				RAY_TRACING_SHADER_SLOT_MATERIAL,
+				RAY_TRACING_NUM_SHADER_SLOTS,
+				MissShaderIndex,
+				PathState.Ray.GetNativeDesc(),
+				PackedPayload,
+				Hit);
+			if (!bIsCameraRay)
+			{
+				NvReorderThread(Hit);
+			}
+			NvInvokeHitObject(TLAS, Hit, PackedPayload);
+		}
+#else
 		// Trace the ray
 		TraceRay(
 			TLAS,
@@ -498,7 +529,7 @@ FPackedPathTracingPayload TraceTransparentRay(inout FPathState PathState, int Bo
 			MissShaderIndex,
 			PathState.Ray.GetNativeDesc(),
 			PackedPayload);
-
+#endif
 		FVolumeIntersectionList VolumeIsectList = VolumeIntersect(PathState.Ray.Origin, PathState.Ray.Direction, PathState.Ray.TMin, PathState.Ray.TMax, PathState.GetPathRoughness(), VolumeFlags);
 		if (VolumeIsectList.HitBlocker() && (PackedPayload.IsMiss() || VolumeIsectList.BlockerHitT < PackedPayload.HitT))
 		{
diff --git a/Engine/Source/Runtime/Renderer/Private/PathTracing.cpp b/Engine/Source/Runtime/Renderer/Private/PathTracing.cpp
--- a/Engine/Source/Runtime/Renderer/Private/PathTracing.cpp
+++ b/Engine/Source/Runtime/Renderer/Private/PathTracing.cpp
@@ -60,10 +60,17 @@ TAutoConsoleVariable<int32> CVarPathTracingDebugVisualizeLightSamples(
 );
 
 
-TAutoConsoleVariable<bool> CVarPathTracingCompaction(
-	TEXT("r.PathTracing.Compaction"),
+TAutoConsoleVariable<int32> CVarPathTracingCompactionDepth(
+	TEXT("r.PathTracing.CompactionDepth"),
+	4,
+	TEXT("Enables path compaction to improve GPU occupancy for the path tracer. The value sets the bounce up to which compaction will happen, beyond that point the path tracer will handle bounces within the dispatch."),
+	ECVF_RenderThreadSafe
+);
+
+TAutoConsoleVariable<bool> CVarPathTracingShaderExecutionReordering(
+	TEXT("r.PathTracing.ShaderExecutionReordering"),
 	true,
-	TEXT("Enables path compaction to improve GPU occupancy for the path tracer. Requires r.PathTracing.Experimental=true to modify.  (default: true (enabled))"),
+	TEXT("Enables Shader Execution Reordering to improve shader coherence for the path tracer. This variable only has effect if the underlying hardware supports it. (default: true (enabled))"),
 	ECVF_RenderThreadSafe
 );
 
@@ -1101,7 +1108,8 @@ class FPathTracingRG : public FGlobalShader
 	class FAdaptiveSampling : SHADER_PERMUTATION_BOOL("PATH_TRACER_USE_ADAPTIVE_SAMPLING");
 	class FCloudShader : SHADER_PERMUTATION_BOOL("PATH_TRACER_USE_CLOUD_SHADER");
 	class FSubstrateComplexSpecialMaterial : SHADER_PERMUTATION_BOOL("PATH_TRACER_USE_SUBSTRATE_SPECIAL_COMPLEX_MATERIAL");
-	using FPermutationDomain = TShaderPermutationDomain<FCompactionType, FAdaptiveSampling, FCloudShader, FSubstrateComplexSpecialMaterial>;
+	class FUseSER : SHADER_PERMUTATION_BOOL("PATH_TRACER_USE_SER");
+	using FPermutationDomain = TShaderPermutationDomain<FCompactionType, FAdaptiveSampling, FCloudShader, FSubstrateComplexSpecialMaterial, FUseSER>;
 
 	static bool ShouldCompilePermutation(const FGlobalShaderPermutationParameters& Parameters)
 	{
@@ -1110,11 +1118,6 @@ class FPathTracingRG : public FGlobalShader
 		if (bUseExperimental == false)
 		{
 			
-			if (PermutationVector.Get<FCompactionType>() == 0)
-			{
-				// non-compaction tracing is considered experimental
-				return false;
-			}
 			if (PermutationVector.Get<FAdaptiveSampling>())
 			{
 				// adaptive sampling is experimental
@@ -1216,7 +1219,7 @@ class FPathTracingRG : public FGlobalShader
 		SHADER_PARAMETER(int32, ScanlineWidth)
 
 		// extra parameters required for path compacting kernel
-		SHADER_PARAMETER(int, Bounce)
+		SHADER_PARAMETER(int32, FirstBounce)
 		SHADER_PARAMETER_RDG_BUFFER_UAV(RWStructuredBuffer<FPathTracingPackedPathState>, PathStateData)
 		SHADER_PARAMETER_RDG_BUFFER_UAV(RWBuffer<uint>, ActivePaths)
 		SHADER_PARAMETER_RDG_BUFFER_UAV(RWBuffer<uint>, NextActivePaths)
@@ -2921,19 +2924,20 @@ class FPathTracingCompositorPS : public FGlobalShader
 };
 IMPLEMENT_SHADER_TYPE(, FPathTracingCompositorPS, TEXT("/Engine/Private/PathTracing/PathTracingCompositingPixelShader.usf"), TEXT("CompositeMain"), SF_Pixel);
 
-static FPathTracingRG::FPermutationDomain GetPathTracingRGPermutation(const FScene& Scene)
+static FPathTracingRG::FPermutationDomain GetPathTracingRGPermutation(const FScene& Scene, const bool bUseCompaction)
 {
 	const bool bUseExperimental = CVarPathTracingExperimental.GetValueOnRenderThread();
-	const bool bUseCompaction = (bUseExperimental == false) || CVarPathTracingCompaction.GetValueOnRenderThread();
 	const bool bUseAdaptiveSampling = bUseExperimental && CVarPathTracingAdaptiveSampling.GetValueOnRenderThread() != 0;
 	const bool bUseCloudShader = bUseExperimental && ShouldCompileRayTracingCallableShadersForProject(Scene.GetShaderPlatform()) && !CVarPathTracingCloudMapEnabled.GetValueOnRenderThread();
 	const bool bHasComplexSpecialRenderPath = Substrate::IsSubstrateEnabled() && Scene.SubstrateSceneData.bUsesComplexSpecialRenderPath;
+	const bool bUseSER = GRHIGlobals.SupportsShaderExecutionReordering && CVarPathTracingShaderExecutionReordering.GetValueOnRenderThread();
 
 	FPathTracingRG::FPermutationDomain Out;
 	Out.Set<FPathTracingRG::FCompactionType>(bUseCompaction);
 	Out.Set<FPathTracingRG::FAdaptiveSampling>(bUseAdaptiveSampling);
 	Out.Set<FPathTracingRG::FCloudShader>(bUseCloudShader);
 	Out.Set<FPathTracingRG::FSubstrateComplexSpecialMaterial>(bHasComplexSpecialRenderPath);
+	Out.Set<FPathTracingRG::FUseSER>(bUseSER);
 	return Out;
 }
 
@@ -2950,12 +2954,16 @@ void FDeferredShadingSceneRenderer::PreparePathTracing(const FSceneViewFamily& V
 		else
 		{
 			// Declare all RayGen shaders that require material closest hit shaders to be bound
-			FPathTracingRG::FPermutationDomain PermutationVector = GetPathTracingRGPermutation(Scene);
 			{
+				FPathTracingRG::FPermutationDomain PermutationVector = GetPathTracingRGPermutation(Scene, false);
+				auto RayGenShader = GetGlobalShaderMap(ViewFamily.GetShaderPlatform())->GetShader<FPathTracingRG>(PermutationVector);
+				OutRayGenShaders.Add(RayGenShader.GetRayTracingShader());
+			}
+			{
+				FPathTracingRG::FPermutationDomain PermutationVector = GetPathTracingRGPermutation(Scene, true);
 				auto RayGenShader = GetGlobalShaderMap(ViewFamily.GetShaderPlatform())->GetShader<FPathTracingRG>(PermutationVector);
 				OutRayGenShaders.Add(RayGenShader.GetRayTracingShader());
 			}
-
 			{
 				auto RayGenShader = GetGlobalShaderMap(ViewFamily.GetShaderPlatform())->GetShader<FPathTracingInitExtinctionCoefficientRG>();
 				OutRayGenShaders.Add(RayGenShader.GetRayTracingShader());
@@ -2984,9 +2992,27 @@ void PreparePathTracingRTPSO()
 				FGlobalShaderMap* ShaderMap = GetGlobalShaderMap(GMaxRHIFeatureLevel);
 
 				FPathTracingRG::FPermutationDomain PermutationVector(PermutationId);
+
+				if (PermutationVector.Get<FPathTracingRG::FCompactionType>())
+				{
+					// skip over these so we don't double count them
+					continue;
+				}
+
+				if (PermutationVector.Get<FPathTracingRG::FUseSER>() && !GRHIGlobals.SupportsShaderExecutionReordering)
+				{
+					// we can safely skip these if the current machine does not support SER
+					// ShouldCompilePermutation will not take this into account because it could be called in a cooking context
+					// where the machine that decides what to compile is not the machine that will run the code
+					continue;
+				}
 				
+				FPathTracingRG::FPermutationDomain PermutationVectorCompaction(PermutationId);
+				PermutationVectorCompaction.Set<FPathTracingRG::FCompactionType>(true);
+
 				FRHIRayTracingShader* RayGenShaderTable[] = {
 					ShaderMap->GetShader<FPathTracingRG>(PermutationVector).GetRayTracingShader(),
+					ShaderMap->GetShader<FPathTracingRG>(PermutationVectorCompaction).GetRayTracingShader(),
 					ShaderMap->GetShader<FPathTracingInitExtinctionCoefficientRG>().GetRayTracingShader(),
 				};
 				FRHIRayTracingShader* MissShaderTable[] = {
@@ -2994,9 +3020,10 @@ void PreparePathTracingRTPSO()
 				};
 				FRHIRayTracingShader* HitGroupTable[] = {
 					GetPathTracingDefaultOpaqueHitShader(ShaderMap),
+					GetPathTracingDefaultHiddenHitShader(ShaderMap),
 				};
 				FRayTracingPipelineStateInitializer Initializer;
-				Initializer.bPartial = true; // TODO: getting a crash in nvidia driver when false
+				Initializer.bPartial = true; // TODO: got a crash in some older nvidia drivers when false - need to find out which driver version has the fix
 				Initializer.SetRayGenShaderTable(RayGenShaderTable);
 				Initializer.SetMissShaderTable(MissShaderTable);
 				Initializer.SetHitGroupTable(HitGroupTable);
@@ -3658,15 +3685,18 @@ void FDeferredShadingSceneRenderer::RenderPathTracing(
 			// We are writing to the texture, we'll need to extract it...
 			bNeedsTextureExtract = true;
 
-			// should we use path compaction?
-			const bool bUseCompaction = (bUseExperimental == false) || CVarPathTracingCompaction.GetValueOnRenderThread() != 0;
+			int32 CompactionDepth = CVarPathTracingCompactionDepth.GetValueOnRenderThread();
+			if (CompactionDepth < 0 || CompactionDepth > int32(Config.PathTracingData.MaxBounces))
+			{
+				CompactionDepth = Config.PathTracingData.MaxBounces;
+			}
+
 			const bool bUseIndirectDispatch = GRHISupportsRayTracingDispatchIndirect && CVarPathTracingIndirectDispatch.GetValueOnRenderThread();
 			const int FlushRenderingCommands = CVarPathTracingFlushDispatch.GetValueOnRenderThread();
 
 			FRDGBuffer* ActivePaths[2] = {};
 			FRDGBuffer* NumActivePaths = nullptr;
 			FRDGBuffer* PathStateData = nullptr;
-			if (bUseCompaction || Config.UseAdaptiveSampling)
 			{
 				const int32 NumPaths = FMath::Min(
 					DispatchSize * FMath::DivideAndRoundUp(DispatchSize, NumGPUs),
@@ -3674,11 +3704,10 @@ void FDeferredShadingSceneRenderer::RenderPathTracing(
 				);
 				ActivePaths[0] = GraphBuilder.CreateBuffer(FRDGBufferDesc::CreateBufferDesc(sizeof(int32), NumPaths), TEXT("PathTracer.ActivePaths0"));
 				ActivePaths[1] = GraphBuilder.CreateBuffer(FRDGBufferDesc::CreateBufferDesc(sizeof(int32), NumPaths), TEXT("PathTracer.ActivePaths1"));
-				NumActivePaths = GraphBuilder.CreateBuffer(FRDGBufferDesc::CreateIndirectDesc<uint32>(3 * (Config.PathTracingData.MaxBounces + 1)), TEXT("PathTracer.NumActivePaths"));
+				NumActivePaths = GraphBuilder.CreateBuffer(FRDGBufferDesc::CreateIndirectDesc<uint32>(3 * (CompactionDepth + 1)), TEXT("PathTracer.NumActivePaths"));
 				PathStateData = GraphBuilder.CreateBuffer(FRDGBufferDesc::CreateStructuredDesc(sizeof(FPathTracingPackedPathState), NumPaths), TEXT("PathTracer.PathStateData"));
 			}
 
-			TShaderMapRef<FPathTracingRG> RayGenShader(View.ShaderMap, GetPathTracingRGPermutation(*Scene));
 			FPathTracingRG::FParameters* PreviousPassParameters = nullptr;
 			// Divide each tile among all the active GPUs (interleaving scanlines)
 			// The assumption is that the tiles are as big as possible, hopefully covering the entire screen
@@ -3687,7 +3716,7 @@ void FDeferredShadingSceneRenderer::RenderPathTracing(
 			for (int32 GPUIndex : GPUMask)
 			{
 				RDG_GPU_MASK_SCOPE(GraphBuilder, FRHIGPUMask::FromIndex(GPUIndex));
-				RDG_EVENT_SCOPE_STAT(GraphBuilder, PathTracing, "Path Tracing GPU%d", GPUIndex);
+				RDG_EVENT_SCOPE_CONDITIONAL_STAT(GraphBuilder, NumGPUs > 1, PathTracing, "Path Tracing GPU%d", GPUIndex);
 #if WITH_MGPU
 				RDG_GPU_STAT_SCOPE(GraphBuilder, PathTracing);
 #endif
@@ -3733,14 +3762,11 @@ void FDeferredShadingSceneRenderer::RenderPathTracing(
 								PassParameters,
 								FComputeShaderUtils::GetGroupCount(PassParameters->DispatchDim, FComputeShaderUtils::kGolden2DGroupSize));
 						}
-						else if (bUseCompaction)
-						{
-							AddClearUAVPass(GraphBuilder, GraphBuilder.CreateUAV(NumActivePaths, PF_R32_UINT), 0);
-						}
+						AddClearUAVPass(GraphBuilder, GraphBuilder.CreateUAV(NumActivePaths, PF_R32_UINT), 0);
 
-						// When using path compaction, we need to run the path tracer once per bounce
-						// otherwise, the path tracer is the one doing the bounces
-						for (int Bounce = 0, MaxBounces = bUseCompaction ? Config.PathTracingData.MaxBounces : 0; Bounce <= MaxBounces; Bounce++)
+						// Run a pass per bounce, up until the compaction depth. Beyond that point, the path tracer will handle any remaining bounces.
+						// Generally, since there is some launch overhead - it can be worthwhile to handle the "bounce tail" in a single loop because not too many paths survive
+						for (int Bounce = 0; Bounce <= CompactionDepth; Bounce++)
 						{
 							FPathTracingRG::FParameters* PassParameters = GraphBuilder.AllocParameters<FPathTracingRG::FParameters>();
 							PassParameters->TLAS = Scene->RayTracingScene.GetLayerView(ERayTracingSceneLayer::Base);
@@ -3852,22 +3878,14 @@ void FDeferredShadingSceneRenderer::RenderPathTracing(
 							PassParameters->ScanlineStride = NumGPUs;
 							PassParameters->ScanlineWidth = DispatchSizeX;
 
-							PassParameters->Bounce = Bounce;
-							if (bUseCompaction)
-							{
-								PassParameters->ActivePaths = GraphBuilder.CreateUAV(ActivePaths[Bounce & 1], PF_R32_UINT);
-								PassParameters->NextActivePaths = GraphBuilder.CreateUAV(ActivePaths[(Bounce & 1) ^ 1], PF_R32_UINT);
-								PassParameters->PathStateData = GraphBuilder.CreateUAV(PathStateData);
-								PassParameters->NumPathStates = GraphBuilder.CreateUAV(NumActivePaths, PF_R32_UINT);
-								if (bUseIndirectDispatch)
-								{
-									PassParameters->PathTracingIndirectArgs = NumActivePaths;
-								}
-							}
-							else if (Config.UseAdaptiveSampling)
+							PassParameters->FirstBounce = Bounce;
+							PassParameters->ActivePaths = GraphBuilder.CreateUAV(ActivePaths[Bounce & 1], PF_R32_UINT);
+							PassParameters->NextActivePaths = GraphBuilder.CreateUAV(ActivePaths[(Bounce & 1) ^ 1], PF_R32_UINT);
+							PassParameters->PathStateData = GraphBuilder.CreateUAV(PathStateData);
+							PassParameters->NumPathStates = GraphBuilder.CreateUAV(NumActivePaths, PF_R32_UINT);
+							if (bUseIndirectDispatch)
 							{
-								PassParameters->ActivePaths = GraphBuilder.CreateUAV(ActivePaths[0], PF_R32_UINT);
-								PassParameters->NumPathStates = GraphBuilder.CreateUAV(NumActivePaths, PF_R32_UINT);
+								PassParameters->PathTracingIndirectArgs = NumActivePaths;
 							}
 							const bool bEnableDebug = CVarPathTracingDebug.GetValueOnRenderThread() > 0;
 							if (bEnableDebug)
@@ -3878,20 +3896,21 @@ void FDeferredShadingSceneRenderer::RenderPathTracing(
 								ShaderPrint::SetParameters(GraphBuilder, View.ShaderPrintData, PassParameters->ShaderPrint);
 							}
 
+
+							const bool bUseCompaction = Bounce < CompactionDepth || Bounce == Config.PathTracingData.MaxBounces;
+							TShaderMapRef<FPathTracingRG> RayGenShader(View.ShaderMap, GetPathTracingRGPermutation(*Scene, bUseCompaction));
 							ClearUnusedGraphResources(RayGenShader, PassParameters);
-							const bool bFlushRenderingCommands = FlushRenderingCommands == 1 || (FlushRenderingCommands == 2 && Bounce == MaxBounces);
-							const bool bUse1DDispatch = Config.UseAdaptiveSampling || (bUseCompaction && Bounce > 0);
+							const bool bFlushRenderingCommands = FlushRenderingCommands == 1 || (FlushRenderingCommands == 2 && Bounce == CompactionDepth);
+							const bool bUse1DDispatch = Config.UseAdaptiveSampling || Bounce > 0;
 							GraphBuilder.AddPass(
-								bUseCompaction
-								? RDG_EVENT_NAME("Path Tracer Sample=%d/%d NumLights=%d (Bounce=%d%s)", PathTracingState->SampleIndex, MaxSPP, PassParameters->SceneLightCount, Bounce, bUseIndirectDispatch && Bounce > 0 ? TEXT(" indirect") : TEXT(""))
-								: RDG_EVENT_NAME("Path Tracer Sample=%d/%d NumLights=%d"              , PathTracingState->SampleIndex, MaxSPP, PassParameters->SceneLightCount),
+								RDG_EVENT_NAME("Path Tracer Sample=%d/%d NumLights=%d (Bounce=%d%s%s)", PathTracingState->SampleIndex, MaxSPP, PassParameters->SceneLightCount, PassParameters->FirstBounce, bUseCompaction ? "" : "+", bUseIndirectDispatch && Bounce > 0 ? TEXT(" indirect") : TEXT("")),
 								PassParameters,
 								ERDGPassFlags::Compute,
 								[PassParameters, RayGenShader, DispatchSizeX, DispatchSizeYLocal, bUseIndirectDispatch, bUse1DDispatch, bFlushRenderingCommands, GPUIndex, &View](FRDGAsyncTask, FRHICommandList& RHICmdList)
 								{
 									FRHIBatchedShaderParameters& GlobalResources = RHICmdList.GetScratchShaderParameters();
 									SetShaderParameters(GlobalResources, RayGenShader, *PassParameters);
-									if (bUseIndirectDispatch && PassParameters->Bounce > 0)
+									if (bUseIndirectDispatch && PassParameters->FirstBounce > 0)
 									{
 										PassParameters->PathTracingIndirectArgs->MarkResourceAsUsed();
 
@@ -3899,7 +3918,7 @@ void FDeferredShadingSceneRenderer::RenderPathTracing(
 											View.RayTracingMaterialPipeline,
 											RayGenShader.GetRayTracingShader(),
 											View.RayTracingSBT, GlobalResources,
-											PassParameters->PathTracingIndirectArgs->GetIndirectRHICallBuffer(), 3 * PassParameters->Bounce * sizeof(uint32)
+											PassParameters->PathTracingIndirectArgs->GetIndirectRHICallBuffer(), 3 * PassParameters->FirstBounce * sizeof(uint32)
 										);
 									}
 									else if (bUse1DDispatch)
