Path Tracer: Only compile SER extensions for the D3D12 RHI for now

This is a temporary workaround until we can either implement this on all RHIs that support the path tracer, or make a DDPI entry for SER.

#rnx

diff --git a/Engine/Source/Runtime/Renderer/Private/PathTracing.cpp b/Engine/Source/Runtime/Renderer/Private/PathTracing.cpp
--- a/Engine/Source/Runtime/Renderer/Private/PathTracing.cpp
+++ b/Engine/Source/Runtime/Renderer/Private/PathTracing.cpp
@@ -1140,6 +1140,12 @@ class FPathTracingRG : public FGlobalShader
 				}
 			}
 		}
+		if (PermutationVector.Get<FUseSER>() && Parameters.Platform != SP_PCD3D_SM6)
+		{
+			// only D3D12 supports SER for now
+			return false;
+		}
+
 		if (!Substrate::IsSubstrateEnabled())
 		{
 			// If we aren't using Substrate, no need to compile the complex material path
