Added a conservative path for determining whether a shadow map is to be treated as distant & thus allocate a single page. The new path is based on the MinMipLevel instead of the old shadow resolution value.
 - This fixes problems originating from the inscribed sphere used in the old path, which is apparent for narrow spot lights.
 - Enable with: r.Shadow.Virtual.UseConservativeDistantLightThreshold

diff --git a/Engine/Source/Runtime/Renderer/Private/Shadows/ShadowSceneRenderer.cpp b/Engine/Source/Runtime/Renderer/Private/Shadows/ShadowSceneRenderer.cpp
--- a/Engine/Source/Runtime/Renderer/Private/Shadows/ShadowSceneRenderer.cpp
+++ b/Engine/Source/Runtime/Renderer/Private/Shadows/ShadowSceneRenderer.cpp
@@ -50,6 +50,14 @@ static TAutoConsoleVariable<float> CVarDistantLightForceCacheFootprintFraction(
 	ECVF_Scalability | ECVF_RenderThreadSafe
 );
 
+static TAutoConsoleVariable<bool> CVarUseConservativeDistantLightThreshold(
+	TEXT("r.Shadow.Virtual.UseConservativeDistantLightThreshold"),
+	false,
+	TEXT("Base the distant light cutoff on the minimum mip level instead of the shadow resolution calculated through the old path.\n")
+	TEXT("  This fixes problems around the use of an inscribed sphere."),
+	ECVF_Scalability | ECVF_RenderThreadSafe
+);
+
 static TAutoConsoleVariable<float> CVarNaniteShadowsLODBias(
 	TEXT("r.Shadow.NaniteLODBias"),
 	1.0f,
@@ -105,6 +113,8 @@ FShadowSceneRenderer::FShadowSceneRenderer(FDeferredShadingSceneRenderer& InScen
 	, Scene(*InSceneRenderer.Scene)
 	, ShadowScene(*Scene.ShadowScene)
 	, VirtualShadowMapArray(InSceneRenderer.VirtualShadowMapArray)
+	, bUseConservativeDistantLightThreshold(CVarUseConservativeDistantLightThreshold.GetValueOnAnyThread())
+	, DistantLightMode(CVarDistantLightMode.GetValueOnAnyThread())
 {
 }
 
@@ -286,42 +296,8 @@ TSharedPtr<FVirtualShadowMapPerLightCacheEntry> FShadowSceneRenderer::AddLocalLi
 	const FLightSceneProxy* LightSceneProxy = ProjectedShadowInfo->GetLightSceneInfo().Proxy;
 	const float ResolutionLODBiasLocal = GetResolutionLODBiasLocal(ShadowScene.GetLightMobilityFactor(LightSceneInfo->Id), LightSceneProxy->GetVSMResolutionLodBias());
 
-	// Single page res, at this point we force the VSM to be single page
-	// TODO: this computation does not match up with page marking logic super-well, particularly for long spot lights,
-	//       we can absolutely mirror the page marking calc better, just unclear how much it helps. 
-	//       Also possible to feed back from gpu - which would be more accurate wrt partially visible lights (e.g., a spot going through the ground).
-	//       Of course this creates jumps if visibility changes, which may or may not create unsolvable artifacts.	
-	const float BiasedFootprintThreshold = float(FVirtualShadowMap::PageSize) * FMath::Exp2(ResolutionLODBiasLocal - LightSceneProxy->GetVSMResolutionLodBias());
-	const bool bIsDistantLight = CVarDistantLightMode.GetValueOnRenderThread() != 0
-		&& (MaxScreenRadius <= BiasedFootprintThreshold || CVarDistantLightMode.GetValueOnRenderThread() == 2);
-
-	const int32 NumMaps = ProjectedShadowInitializer.bOnePassPointLightShadow ? 6 : 1;
-	TSharedPtr<FVirtualShadowMapPerLightCacheEntry> PerLightCacheEntry = CacheManager->FindCreateLightCacheEntry(LightSceneInfo->Id, 0, NumMaps);
-			
-	const float DistantLightForceCacheFootprintFraction = FMath::Clamp(CVarDistantLightForceCacheFootprintFraction.GetValueOnRenderThread(), 0.0f, 1.0f);
-	bool bShouldForceTimeSliceDistantUpdate = (bIsDistantLight && MaxScreenRadius <= BiasedFootprintThreshold * DistantLightForceCacheFootprintFraction);
-	LocalLightShadowFrameSetup.PerLightCacheEntry = PerLightCacheEntry;
-	bool bIsCached = PerLightCacheEntry->UpdateLocal(
-		ProjectedShadowInitializer,
-		LightSceneProxy->GetOrigin(),
-		LightSceneProxy->GetRadius(),
-		bIsDistantLight,
-		CacheManager->IsCacheEnabled(),
-		!bShouldForceTimeSliceDistantUpdate);
-
-	if (bIsCached && bIsDistantLight && PerLightCacheEntry->Prev.ScheduledFrameNumber == Scene.GetFrameNumber())
-	{
-		PerLightCacheEntry->Invalidate();
-	}
-
-	// Update info on the ProjectionShadowInfo; eventually this should all move into local data structures here
-	const int32 VirtualShadowMapId = VirtualShadowMapArray.Allocate(bIsDistantLight, NumMaps);
-	ProjectedShadowInfo->VirtualShadowMapId = VirtualShadowMapId;
-	ProjectedShadowInfo->VirtualShadowMapPerLightCacheEntry = PerLightCacheEntry;
-	ProjectedShadowInfo->bShouldRenderVSM = !PerLightCacheEntry->IsFullyCached();
-
 	// Compute conservative mip level estimate based on radius of the bounding sphere.
-	// TODO: can probably do better by finding  closest point on cone for certain scenarios
+	// TODO: can probably do better by finding  closest point on cone for certain scenarios? Not as important as it might seem as the worst case is for a narrow cone, but then the narrow FOV limits the required resolution.
 
 	const FVector2f ShadowViewSize = FVector2f(FVirtualShadowMap::VirtualMaxResolutionXY, FVirtualShadowMap::VirtualMaxResolutionXY);
 	const FMatrix &ShadowViewToClip = ProjectedShadowInfo->bOnePassPointLightShadow ? ProjectedShadowInfo->OnePassShadowFaceProjectionMatrix : ProjectedShadowInfo->ViewToClipOuter;
@@ -346,6 +322,55 @@ TSharedPtr<FVirtualShadowMapPerLightCacheEntry> FShadowSceneRenderer::AddLocalLi
 			CVarMarkPixelPagesMipModeLocal.GetValueOnRenderThread()
 			));
 	}
+	
+	bool bIsDistantLight = DistantLightMode == 2;
+	bool bShouldForceTimeSliceDistantUpdate = false;
+
+	if (DistantLightMode == 1)
+	{
+		if (bUseConservativeDistantLightThreshold)
+		{
+			// use distant light only if we are sure that there's only one mip level.
+			bIsDistantLight = MinMipLevel == (FVirtualShadowMap::MaxMipLevels - 1);
+			bShouldForceTimeSliceDistantUpdate = false;// TODO: (bIsDistantLight && MaxScreenRadius <= BiasedFootprintThreshold * DistantLightForceCacheFootprintFraction); ??
+		}
+		else
+		{
+			// Single page res, at this point we force the VSM to be single page
+			const float BiasedFootprintThreshold = float(FVirtualShadowMap::PageSize) * FMath::Exp2(ResolutionLODBiasLocal - LightSceneProxy->GetVSMResolutionLodBias());
+			bIsDistantLight = MaxScreenRadius <= BiasedFootprintThreshold;
+			
+			const float DistantLightForceCacheFootprintFraction = FMath::Clamp(CVarDistantLightForceCacheFootprintFraction.GetValueOnRenderThread(), 0.0f, 1.0f);
+			bShouldForceTimeSliceDistantUpdate = (bIsDistantLight && MaxScreenRadius <= BiasedFootprintThreshold * DistantLightForceCacheFootprintFraction);
+
+		}
+	}
+
+
+	const int32 NumMaps = ProjectedShadowInitializer.bOnePassPointLightShadow ? 6 : 1;
+	TSharedPtr<FVirtualShadowMapPerLightCacheEntry> PerLightCacheEntry = CacheManager->FindCreateLightCacheEntry(LightSceneInfo->Id, 0, NumMaps);
+			
+	LocalLightShadowFrameSetup.PerLightCacheEntry = PerLightCacheEntry;
+	bool bIsCached = PerLightCacheEntry->UpdateLocal(
+		ProjectedShadowInitializer,
+		LightSceneProxy->GetOrigin(),
+		LightSceneProxy->GetRadius(),
+		bIsDistantLight,
+		CacheManager->IsCacheEnabled(),
+		!bShouldForceTimeSliceDistantUpdate);
+
+	if (bIsCached && bIsDistantLight && PerLightCacheEntry->Prev.ScheduledFrameNumber == Scene.GetFrameNumber())
+	{
+		PerLightCacheEntry->Invalidate();
+	}
+
+	// Update info on the ProjectionShadowInfo; eventually this should all move into local data structures here
+	const int32 VirtualShadowMapId = VirtualShadowMapArray.Allocate(bIsDistantLight, NumMaps);
+	ProjectedShadowInfo->VirtualShadowMapId = VirtualShadowMapId;
+	ProjectedShadowInfo->VirtualShadowMapPerLightCacheEntry = PerLightCacheEntry;
+	ProjectedShadowInfo->bShouldRenderVSM = !PerLightCacheEntry->IsFullyCached();
+
+
 
 	for (int32 Index = 0; Index < NumMaps; ++Index)
 	{
diff --git a/Engine/Source/Runtime/Renderer/Private/Shadows/ShadowSceneRenderer.h b/Engine/Source/Runtime/Renderer/Private/Shadows/ShadowSceneRenderer.h
--- a/Engine/Source/Runtime/Renderer/Private/Shadows/ShadowSceneRenderer.h
+++ b/Engine/Source/Runtime/Renderer/Private/Shadows/ShadowSceneRenderer.h
@@ -132,9 +132,6 @@ private:
 	};
 	TArray<FDirectionalLightShadowFrameSetup, SceneRenderingAllocator> DirectionalLights;
 
-	// One pass projection stuff. Set up in RenderVitualShadowMapProjectionMaskBits
-	bool bShouldUseVirtualShadowMapOnePassProjection = false;
-
 	// Links to other systems etc.
 	FDeferredShadingSceneRenderer& SceneRenderer;
 	FScene& Scene;
@@ -145,4 +142,11 @@ private:
 	Nanite::FPackedViewArray* VirtualShadowMapViews = nullptr;
 	FSceneInstanceCullingQuery *SceneInstanceCullingQuery = nullptr;
 	TArray<FViewData, SceneRenderingAllocator> ViewDatas;
+
+	// One pass projection stuff. Set up in RenderVitualShadowMapProjectionMaskBits
+	bool bShouldUseVirtualShadowMapOnePassProjection = false;
+
+	// Base the distant light cutoff on the minimum mip level instead of the shadow resolution calculated through the old path.
+	bool bUseConservativeDistantLightThreshold = false;
+	int32 DistantLightMode = 0;
 };
