[Light Grid] Control threshold used in heat map visualization using r.Forward.LightGridDebug.MaxThreshold.

diff --git a/Engine/Shaders/Private/LightGridInjection.usf b/Engine/Shaders/Private/LightGridInjection.usf
--- a/Engine/Shaders/Private/LightGridInjection.usf
+++ b/Engine/Shaders/Private/LightGridInjection.usf
@@ -558,6 +558,7 @@ void LightGridCompactCS(
 
 Texture2D<float4> DepthTexture;
 uint DebugMode;
+uint MaxThreshold;
 
 void DebugLightGridPS(
 	in float4 SVPos : SV_POSITION,
@@ -630,7 +631,7 @@ void DebugLightGridPS(
 	if (DebugNumLightsInGridCell > 0)
 	{
 		OutLuminanceTransmittance.a   = 0.0f;
-		OutLuminanceTransmittance.rgb = GetHSVDebugColor(float(DebugNumLightsInGridCell) / 8.0f);
+		OutLuminanceTransmittance.rgb = GetHSVDebugColor(float(DebugNumLightsInGridCell) / MaxThreshold);
 		PrintSmallUint(PixelPos, OutLuminanceTransmittance.rgb, float3(0.1, 0.1, 0.1), (TileTopLeftPixelCoord + TileSize/2) / View.ViewResolutionFraction, float(DebugNumLightsInGridCell));
 	}
 
diff --git a/Engine/Source/Runtime/Renderer/Private/LightGridInjection.cpp b/Engine/Source/Runtime/Renderer/Private/LightGridInjection.cpp
--- a/Engine/Source/Runtime/Renderer/Private/LightGridInjection.cpp
+++ b/Engine/Source/Runtime/Renderer/Private/LightGridInjection.cpp
@@ -63,6 +63,14 @@ FAutoConsoleVariableRef CVarLightGridDebug(
 	ECVF_RenderThreadSafe
 );
 
+int32 GForwardLightGridDebugMaxThreshold = 8;
+FAutoConsoleVariableRef CVarRayTracingDebugHitCountMaxThreshold(
+	TEXT("r.Forward.LightGridDebug.MaxThreshold"),
+	GForwardLightGridDebugMaxThreshold,
+	TEXT("Maximum light threshold for heat map visualization. (default = 8)\n"),
+	ECVF_RenderThreadSafe
+);
+
 int32 GLightGridHZBCull = 1;
 FAutoConsoleVariableRef CVarLightGridHZBCull(
 	TEXT("r.Forward.LightGridHZBCull"),
@@ -1207,6 +1215,7 @@ class FDebugLightGridPS : public FGlobalShader
 		SHADER_PARAMETER_TEXTURE(Texture2D, MiniFontTexture)
 		SHADER_PARAMETER_RDG_TEXTURE(Texture2D, DepthTexture)
 		SHADER_PARAMETER(uint32, DebugMode)
+		SHADER_PARAMETER(uint32, MaxThreshold)
 		RENDER_TARGET_BINDING_SLOTS()
 	END_SHADER_PARAMETER_STRUCT()
 
@@ -1261,6 +1270,7 @@ FScreenPassTexture AddVisualizeLightGridPass(FRDGBuilder& GraphBuilder, const FV
 		PassParameters->MiniFontTexture = GetMiniFontTexture();
 		PassParameters->RenderTargets[0] = FRenderTargetBinding(ScreenPassSceneColor.Texture, ERenderTargetLoadAction::ELoad);
 		PassParameters->DebugMode = GForwardLightGridDebug;
+		PassParameters->MaxThreshold = GForwardLightGridDebugMaxThreshold;
 
 		FRHIBlendState* PreMultipliedColorTransmittanceBlend = TStaticBlendState<CW_RGB, BO_Add, BF_One, BF_SourceAlpha, BO_Add, BF_Zero, BF_One>::GetRHI();
 
