VolumetricFog - fixed ligth function atlas support on light not casting volumetric shadows.

#rnx
#tests none

diff --git a/Engine/Source/Runtime/Renderer/Private/VolumetricFog.cpp b/Engine/Source/Runtime/Renderer/Private/VolumetricFog.cpp
--- a/Engine/Source/Runtime/Renderer/Private/VolumetricFog.cpp
+++ b/Engine/Source/Runtime/Renderer/Private/VolumetricFog.cpp
@@ -1806,7 +1806,7 @@ void FSceneRenderer::ComputeVolumetricFog(FRDGBuilder& GraphBuilder,
 			PermutationVector.Set< FVolumetricFogLightScatteringCS::FLumenGI >(bUseLumenGI);
 			PermutationVector.Set< FVolumetricFogLightScatteringCS::FVirtualShadowMap >(VirtualShadowMapArray.IsAllocated() );
 			PermutationVector.Set< FVolumetricFogLightScatteringCS::FRaytracedShadowsVolume >(bUseRaytracedShadowsVolume);
-			PermutationVector.Set< FVolumetricFogLightScatteringCS::FSampleLightFunctionAtlas >(PassParameters->DirectionalApplyLightFunctionFromAtlas>1);
+			PermutationVector.Set< FVolumetricFogLightScatteringCS::FSampleLightFunctionAtlas >(bUseLightFunctionAtlasEnabledAndSupported);
 			PermutationVector.Set< FVolumetricFogLightScatteringCS::FMegaLights>(bUseMegaLights);
 			PermutationVector.Set< FVolumetricFogLightScatteringCS::FLightSoftFading >(GetVolumetricFogLightSoftFading() > 0.0f);
 
