Light Function Atlas - updated out of slot warning on screen to be more visible.

#rnx
#tests none

diff --git a/Engine/Source/Runtime/Renderer/Private/DeferredShadingRenderer.cpp b/Engine/Source/Runtime/Renderer/Private/DeferredShadingRenderer.cpp
--- a/Engine/Source/Runtime/Renderer/Private/DeferredShadingRenderer.cpp
+++ b/Engine/Source/Runtime/Renderer/Private/DeferredShadingRenderer.cpp
@@ -3191,8 +3191,6 @@ void FDeferredShadingSceneRenderer::Render(FRDGBuilder& GraphBuilder)
 			}
 		}
 
-		LightFunctionAtlas.RenderDebugInfo(GraphBuilder, Views);
-
 		if (VirtualShadowMapArray.IsEnabled())
 		{
 			VirtualShadowMapArray.RenderDebugInfo(GraphBuilder, Views);
diff --git a/Engine/Source/Runtime/Renderer/Private/LightFunctionAtlas.cpp b/Engine/Source/Runtime/Renderer/Private/LightFunctionAtlas.cpp
--- a/Engine/Source/Runtime/Renderer/Private/LightFunctionAtlas.cpp
+++ b/Engine/Source/Runtime/Renderer/Private/LightFunctionAtlas.cpp
@@ -786,51 +786,34 @@ void FLightFunctionAtlas::RenderAtlasSlots(FRDGBuilder& GraphBuilder, const TArr
 	);
 }
 
-void FLightFunctionAtlas::RenderDebugInfo(FRDGBuilder& GraphBuilder, TArray<FViewInfo>& Views)
+bool FLightFunctionAtlas::IsOutOfSlots() 
+{
+#if !UE_BUILD_SHIPPING
+	return LightCountSkippedDueToMissingAtlasSlot > 0;
+#else
+	return false;
+#endif
+}
+
+FString FLightFunctionAtlas::GetOutOfSlotWarningMessage()
 {
 #if !UE_BUILD_SHIPPING
 	if (!IsLightFunctionAtlasEnabled() || !GAreScreenMessagesEnabled)
 	{
-		return;
+		return FString();
 	}
 
 	// In case we became out of budget, let's notify the game developers.
 	if (LightCountSkippedDueToMissingAtlasSlot > 0)
 	{
-		for (int32 ViewIndex = 0; ViewIndex < Views.Num(); ViewIndex++)
-		{
-			const FViewInfo& View = Views[ViewIndex];
-			const FMinimalSceneTextures& SceneTextures = View.GetSceneTextures();
+		FString Text = FString::Printf(TEXT("Light Functions Atlas:\n"));
+		Text += FString::Printf(TEXT("  - %d light(s) will skip light functions due to out of atlas slot (see r.LightFunctionAtlas.Size).\n"), LightCountSkippedDueToMissingAtlasSlot);
+		Text += FString::Printf(TEXT("  - %d light function material(s) have been skipped.\n"), SkippedLightFunctionsSet.Num());
 
-			AddDrawCanvasPass(GraphBuilder, {}, View, FScreenPassRenderTarget(SceneTextures.Color.Target, View.ViewRect, ERenderTargetLoadAction::ELoad),
-				[&View, this](FCanvas& Canvas)
-				{
-					FString Text;
-					const float DPIScale = Canvas.GetDPIScale();
-					Canvas.SetBaseTransform(FMatrix(FScaleMatrix(DPIScale) * Canvas.CalcBaseTransform2D(Canvas.GetViewRect().Width(), Canvas.GetViewRect().Height())));
-
-					const float ViewPortWidth = float(View.ViewRect.Width());
-					const float ViewPortHeight = float(View.ViewRect.Height());
-					FLinearColor TextColor(1.0f, 0.5f, 0.0f);
-
-					float DrawPosX = 40.0f;
-					float DrawPosY = 650.0f;
-
-					Text = FString::Printf(TEXT("Light Functions Atlas:"));
-					Canvas.DrawShadowedString(DrawPosX, DrawPosY, *Text, GetStatsFont(), TextColor);
-					DrawPosY += 20.0f;
-
-					Text = FString::Printf(TEXT("  - %d light(s) will skip light functions due to out of atlas slot (see r.LightFunctionAtlas.Size)."), LightCountSkippedDueToMissingAtlasSlot);
-					Canvas.DrawShadowedString(DrawPosX, DrawPosY, *Text, GetStatsFont(), TextColor);
-					DrawPosY += 20.0f;
-
-					Text = FString::Printf(TEXT("  - %d light function material(s) have been skipped."), SkippedLightFunctionsSet.Num());
-					Canvas.DrawShadowedString(DrawPosX, DrawPosY, *Text, GetStatsFont(), TextColor);
-					DrawPosY += 20.0f;
-				});
-		}
+		return Text;
 	}
 #endif // !UE_BUILD_SHIPPING
+	return FString();
 }
 
 FScreenPassTexture FLightFunctionAtlas::AddDebugVisualizationPasses(FRDGBuilder& GraphBuilder, const FViewInfo& View, FScreenPassTexture& ScreenPassSceneColor) const
diff --git a/Engine/Source/Runtime/Renderer/Private/LightFunctionAtlas.h b/Engine/Source/Runtime/Renderer/Private/LightFunctionAtlas.h
--- a/Engine/Source/Runtime/Renderer/Private/LightFunctionAtlas.h
+++ b/Engine/Source/Runtime/Renderer/Private/LightFunctionAtlas.h
@@ -149,8 +149,6 @@ struct FLightFunctionAtlas
 
 	void RenderLightFunctionAtlas(FRDGBuilder& GraphBuilder, TArray<FViewInfo>& Views);
 
-	void RenderDebugInfo(FRDGBuilder& GraphBuilder, TArray<FViewInfo>& Views);
-
 
 	FScreenPassTexture AddDebugVisualizationPasses(FRDGBuilder& GraphBuilder, const FViewInfo& View, FScreenPassTexture& ScreenPassSceneColor)  const;
 
@@ -160,6 +158,9 @@ struct FLightFunctionAtlas
 	static FLightFunctionAtlasGlobalParameters*					GetDefaultLightFunctionAtlasGlobalParametersStruct(FRDGBuilder& GraphBuilder);
 	TRDGUniformBufferRef<FLightFunctionAtlasGlobalParameters>	GetDefaultLightFunctionAtlasGlobalParameters(FRDGBuilder& GraphBuilder);
 
+	bool IsOutOfSlots();
+	FString GetOutOfSlotWarningMessage();
+
 private:
 
 	void AllocateAtlasSlots(const TArray<FViewInfo>& Views);
diff --git a/Engine/Source/Runtime/Renderer/Private/SceneRendering.cpp b/Engine/Source/Runtime/Renderer/Private/SceneRendering.cpp
--- a/Engine/Source/Runtime/Renderer/Private/SceneRendering.cpp
+++ b/Engine/Source/Runtime/Renderer/Private/SceneRendering.cpp
@@ -4096,6 +4096,8 @@ void FSceneRenderer::OnRenderFinish(FRDGBuilder& GraphBuilder, FRDGTextureRef Vi
 
 		const bool bSingleLayerWaterWarning = ShouldRenderSingleLayerWaterSkippedRenderEditorNotification(Views);
 
+		const bool bLightFunctionAtlasOutOfSlotWarning = LightFunctionAtlas.IsLightFunctionAtlasEnabled() ? LightFunctionAtlas.IsOutOfSlots() : false;
+
 		bool bShowWaitingSkylight = false;
 #if WITH_EDITOR
 		FSkyLightSceneProxy* SkyLight = Scene->SkyLight;
@@ -4112,7 +4114,7 @@ void FSceneRenderer::OnRenderFinish(FRDGBuilder& GraphBuilder, FRDGTextureRef Vi
 
 		const bool bAnyWarning = bShowPrecomputedVisibilityWarning || bShowDemotedLocalMemoryWarning || bShowGlobalClipPlaneWarning || bShowSkylightWarning || bShowPointLightWarning
 			|| bShowDFAODisabledWarning || bShowShadowedLightOverflowWarning || bShowMobileDynamicCSMWarning || bShowMobileLowQualityLightmapWarning
-			|| bMobileMissingSkyMaterial || bShowSkinCacheOOM || bSingleLayerWaterWarning || bShowDFDisabledWarning || bShowNoSkyAtmosphereComponentWarning || bFxDebugDraw 
+			|| bMobileMissingSkyMaterial || bShowSkinCacheOOM || bSingleLayerWaterWarning || bLightFunctionAtlasOutOfSlotWarning || bShowDFDisabledWarning || bShowNoSkyAtmosphereComponentWarning || bFxDebugDraw
 			|| bLumenEnabledButHasNoDataForTracing || bLumenEnabledButDisabledForTheProject || bNaniteEnabledButNoAtomics || bNaniteEnabledButDisabledInProject || bRealTimeSkyCaptureButNothingToCapture || bShowWaitingSkylight
 			|| bShowLocalExposureDisabledWarning || bHasDelegateWarnings || bContactShadowIntensityCvarUsed || bLocalFogVolumeInSceneButProjectDisabled
 			;
@@ -4160,7 +4162,7 @@ void FSceneRenderer::OnRenderFinish(FRDGBuilder& GraphBuilder, FRDGTextureRef Vi
 						bLocked, bShowPrecomputedVisibilityWarning, bShowDemotedLocalMemoryWarning, bShowGlobalClipPlaneWarning, bShowDFAODisabledWarning, bShowDFDisabledWarning,
 						bIsFrozen, bShowSkylightWarning, bShowPointLightWarning, bShowShadowedLightOverflowWarning,
 						bShowMobileLowQualityLightmapWarning, bShowMobileDynamicCSMWarning, bMobileMissingSkyMaterial, 
-						bShowSkinCacheOOM, bSingleLayerWaterWarning, bShowNoSkyAtmosphereComponentWarning, bFxDebugDraw, FXInterface, bShowLocalExposureDisabledWarning,
+						bShowSkinCacheOOM, bSingleLayerWaterWarning, bLightFunctionAtlasOutOfSlotWarning, bShowNoSkyAtmosphereComponentWarning, bFxDebugDraw, FXInterface, bShowLocalExposureDisabledWarning,
 						bLumenEnabledButHasNoDataForTracing, bLumenEnabledButDisabledForTheProject, bNaniteEnabledButNoAtomics, bNaniteEnabledButDisabledInProject, 
 						bRealTimeSkyCaptureButNothingToCapture, bShowWaitingSkylight, bShowAmbientCubemapMipGenSettingsWarning, bLocalFogVolumeInSceneButProjectDisabled,
 						bStereoView, bPrimaryStereoView, bIsInstancedStereoEnabled, bIsMultiViewportEnabled, bIsMobileMultiViewEnabled, bContactShadowIntensityCvarUsed]
@@ -4271,6 +4273,11 @@ void FSceneRenderer::OnRenderFinish(FRDGBuilder& GraphBuilder, FRDGTextureRef Vi
 							Writer.DrawLine(Message);
 						}
 
+						if (bLightFunctionAtlasOutOfSlotWarning)
+						{
+							Writer.DrawLine(FText::FromString(LightFunctionAtlas.GetOutOfSlotWarningMessage()));
+						}
+
 						if (bLumenEnabledButHasNoDataForTracing)
 						{
 							static const FText Message = NSLOCTEXT("Renderer", "LumenCantDisplay", "Lumen is enabled, but has no ray tracing data and won't operate correctly.\nEither configure Lumen to use software distance field ray tracing and enable 'Generate Mesh Distancefields' in project settings\nor configure Lumen to use Hardware Ray Tracing and enable 'Support Hardware Ray Tracing' in project settings.");
