Fix variable redifinition issue due to incorrect name.

diff --git a/Engine/Source/Runtime/Renderer/Private/LightGridInjection.cpp b/Engine/Source/Runtime/Renderer/Private/LightGridInjection.cpp
--- a/Engine/Source/Runtime/Renderer/Private/LightGridInjection.cpp
+++ b/Engine/Source/Runtime/Renderer/Private/LightGridInjection.cpp
@@ -64,7 +64,7 @@ FAutoConsoleVariableRef CVarLightGridDebug(
 );
 
 int32 GForwardLightGridDebugMaxThreshold = 8;
-FAutoConsoleVariableRef CVarRayTracingDebugHitCountMaxThreshold(
+FAutoConsoleVariableRef CVarLightGridDebugMaxThreshold(
 	TEXT("r.Forward.LightGridDebug.MaxThreshold"),
 	GForwardLightGridDebugMaxThreshold,
 	TEXT("Maximum light threshold for heat map visualization. (default = 8)\n"),
