Fix BaseColor,Roughness,Normal RVT issues.
Streaming mips wasn't baking the 2nd layer. The alias texture wasn't set for the baking path, so it didn't copy correctly.
Also when using LQ compressed format, we would hit an error creating UAV.

diff --git a/Engine/Source/Runtime/Renderer/Private/VT/RuntimeVirtualTextureRender.cpp b/Engine/Source/Runtime/Renderer/Private/VT/RuntimeVirtualTextureRender.cpp
--- a/Engine/Source/Runtime/Renderer/Private/VT/RuntimeVirtualTextureRender.cpp
+++ b/Engine/Source/Runtime/Renderer/Private/VT/RuntimeVirtualTextureRender.cpp
@@ -616,7 +616,7 @@ namespace RuntimeVirtualTexture
 				if (bRenderPass)
 				{
 					OutputAlias0 = RenderTexture0 = GraphBuilder.CreateTexture(FRDGTextureDesc::Create2D(Desc.TextureSize, bLQFormat ? PF_R5G6B5_UNORM : PF_B8G8R8A8, FClearValueBinding::Black, RTCreateFlags), TEXT("RenderTexture0"));
-					RenderTexture1 = GraphBuilder.CreateTexture(FRDGTextureDesc::Create2D(Desc.TextureSize, bLQFormat ? PF_R5G6B5_UNORM : PF_B8G8R8A8, FClearValueBinding::Black, RTCreateFlags), TEXT("RenderTexture1"));
+					OutputAlias1 = RenderTexture1 = GraphBuilder.CreateTexture(FRDGTextureDesc::Create2D(Desc.TextureSize, bLQFormat ? PF_R5G6B5_UNORM : PF_B8G8R8A8, FClearValueBinding::Black, RTCreateFlags), TEXT("RenderTexture1"));
 				}
 				if (bCompressPass)
 				{
diff --git a/Engine/Source/Runtime/Renderer/Private/VT/VirtualTexturePhysicalSpace.cpp b/Engine/Source/Runtime/Renderer/Private/VT/VirtualTexturePhysicalSpace.cpp
--- a/Engine/Source/Runtime/Renderer/Private/VT/VirtualTexturePhysicalSpace.cpp
+++ b/Engine/Source/Runtime/Renderer/Private/VT/VirtualTexturePhysicalSpace.cpp
@@ -121,6 +121,11 @@ static EPixelFormat GetUnorderedAccessViewFormat(EPixelFormat InFormat)
 		return GRHISupportsUAVFormatAliasing ? GetBlockCompressedFormatUAVAliasFormat(InFormat) : PF_Unknown;
 	}
 
+	if (InFormat == PF_B5G5R5A1_UNORM || InFormat == PF_R5G6B5_UNORM)
+	{
+		return PF_Unknown;
+	}
+
 	return InFormat;
 }
 
