Fixed virtual shadow map GetUniformBuffer returning nullptr when accessed before shadow rendering. Need to return safe dummy parameters.

#rnx

diff --git a/Engine/Source/Runtime/Renderer/Private/VirtualShadowMaps/VirtualShadowMapArray.h b/Engine/Source/Runtime/Renderer/Private/VirtualShadowMaps/VirtualShadowMapArray.h
--- a/Engine/Source/Runtime/Renderer/Private/VirtualShadowMaps/VirtualShadowMapArray.h
+++ b/Engine/Source/Runtime/Renderer/Private/VirtualShadowMaps/VirtualShadowMapArray.h
@@ -368,7 +368,13 @@ public:
 	FVirtualShadowMapSamplingParameters GetSamplingParameters(FRDGBuilder& GraphBuilder, int32 ViewIndex) const;
 	TRDGUniformBufferRef<FVirtualShadowMapUniformParameters> GetUniformBuffer(int32 ViewIndex) const
 	{
-		return ViewIndex < CachedUniformBuffers.Num() ? CachedUniformBuffers[ViewIndex] : nullptr;
+		if (CachedUniformBuffers.IsEmpty())
+		{
+			return TRDGUniformBufferRef<FVirtualShadowMapUniformParameters>();
+		}
+
+		// If the view index is out of range, then it means VSM has not been set up yet, this is not a great time to access the UB but is actully done in skyatmosphere rendering so we need to return a safe default.
+		return CachedUniformBuffers[FMath::Min(ViewIndex, CachedUniformBuffers.Num() - 1)];
 	}
 
 	bool HasAnyShadowData() const { return PhysicalPagePoolRDG != nullptr;  }
