Added r.Visibility.LocalLightPrimitiveInteraction - this can be used to save lots of Render Thread time if LocalLightPrimitiveInteraction isn't used at runtime (e.g. MegaLights is used for all local lights).

diff --git a/Engine/Source/Runtime/Renderer/Private/RendererScene.cpp b/Engine/Source/Runtime/Renderer/Private/RendererScene.cpp
--- a/Engine/Source/Runtime/Renderer/Private/RendererScene.cpp
+++ b/Engine/Source/Runtime/Renderer/Private/RendererScene.cpp
@@ -152,6 +152,13 @@ static FAutoConsoleVariableRef CVarVisibilitySkipAlwaysVisible(
 	ECVF_RenderThreadSafe
 );
 
+static TAutoConsoleVariable<int32> CVarVisibilityLocalLightPrimitiveInteraction(
+	TEXT("r.Visibility.LocalLightPrimitiveInteraction"),
+	1,
+	TEXT("Whether to allow computing local primitive interactions. May greatly speedup render thread time if not needed."),
+	ECVF_Scalability | ECVF_RenderThreadSafe
+);
+
 DECLARE_CYCLE_STAT(TEXT("DeferredShadingSceneRenderer MotionBlurStartFrame"), STAT_FDeferredShadingSceneRenderer_MotionBlurStartFrame, STATGROUP_SceneRendering);
 
 IMPLEMENT_GLOBAL_SHADER_PARAMETER_STRUCT(FDistanceCullFadeUniformShaderParameters, "PrimitiveFade");
@@ -3712,6 +3719,11 @@ void FScene::GetPrimitiveUniformShaderParameters_RenderThread(const FPrimitiveSc
 
 bool DoesPlatformNeedLocalLightPrimitiveInteraction(EShaderPlatform ShaderPlatform)
 {
+	if (CVarVisibilityLocalLightPrimitiveInteraction.GetValueOnRenderThread() == 0)
+	{
+		return false;
+	}
+
 	extern bool MobileLocalLightsUseSinglePermutation(EShaderPlatform ShaderPlatform);
 	return !IsMobilePlatform(ShaderPlatform) || !MobileLocalLightsUseSinglePermutation(ShaderPlatform) || IsMobileMovableSpotlightShadowsEnabled(ShaderPlatform);
 }
