Fix Lumen with inline RT when using Android/Vulkan

diff --git a/Engine/Source/Runtime/Renderer/Private/DeferredShadingRenderer.cpp b/Engine/Source/Runtime/Renderer/Private/DeferredShadingRenderer.cpp
--- a/Engine/Source/Runtime/Renderer/Private/DeferredShadingRenderer.cpp
+++ b/Engine/Source/Runtime/Renderer/Private/DeferredShadingRenderer.cpp
@@ -611,7 +611,7 @@ bool FDeferredShadingSceneRenderer::SetupRayTracingPipelineStatesAndSBT(FRDGBuil
 		ReferenceView.AddRayTracingMeshBatchTaskList.Empty();
 	}
 
-	if (!GRHISupportsRayTracingShaders)
+	if (!GRHISupportsRayTracingShaders && !GRHISupportsInlineRayTracing)
 	{
 		return false;
 	}
@@ -966,7 +966,7 @@ void FDeferredShadingSceneRenderer::WaitForRayTracingScene(FRDGBuilder& GraphBui
 
 		if (!bIsPathTracing)
 		{
-			if (GRHISupportsRayTracingShaders)
+			if (GRHISupportsRayTracingShaders || GRHISupportsInlineRayTracing)
 			{
 				if (ReferenceView.LumenHardwareRayTracingMaterialPipeline)
 				{
