Heterogeneous Volumes: Add support for Eye Adaptation node

diff --git a/Engine/Shaders/Private/HeterogeneousVolumes/HeterogeneousVolumesLiveShadingPipeline.usf b/Engine/Shaders/Private/HeterogeneousVolumes/HeterogeneousVolumesLiveShadingPipeline.usf
--- a/Engine/Shaders/Private/HeterogeneousVolumes/HeterogeneousVolumesLiveShadingPipeline.usf
+++ b/Engine/Shaders/Private/HeterogeneousVolumes/HeterogeneousVolumesLiveShadingPipeline.usf
@@ -4,6 +4,7 @@
 #include "../Common.ush"
 
 #define SUPPORT_CONTACT_SHADOWS 0
+#define EYE_ADAPTATION_LOOSE_PARAMETERS 1
 
 #define DYNAMICALLY_SHADOWED 1
 #define TREAT_MAXDEPTH_UNSHADOWED 1
diff --git a/Engine/Shaders/Private/HeterogeneousVolumes/HeterogeneousVolumesMaterialBakingPipeline.usf b/Engine/Shaders/Private/HeterogeneousVolumes/HeterogeneousVolumesMaterialBakingPipeline.usf
--- a/Engine/Shaders/Private/HeterogeneousVolumes/HeterogeneousVolumesMaterialBakingPipeline.usf
+++ b/Engine/Shaders/Private/HeterogeneousVolumes/HeterogeneousVolumesMaterialBakingPipeline.usf
@@ -3,6 +3,7 @@
 #include "../Common.ush"
 
 #define SUPPORT_CONTACT_SHADOWS 0
+#define EYE_ADAPTATION_LOOSE_PARAMETERS 1
 
 #include "/Engine/Generated/Material.ush"
 #include "../BlueNoise.ush"
diff --git a/Engine/Source/Runtime/Renderer/Private/HeterogeneousVolumes/HeterogeneousVolumesLiveShadingPipeline.cpp b/Engine/Source/Runtime/Renderer/Private/HeterogeneousVolumes/HeterogeneousVolumesLiveShadingPipeline.cpp
--- a/Engine/Source/Runtime/Renderer/Private/HeterogeneousVolumes/HeterogeneousVolumesLiveShadingPipeline.cpp
+++ b/Engine/Source/Runtime/Renderer/Private/HeterogeneousVolumes/HeterogeneousVolumesLiveShadingPipeline.cpp
@@ -370,6 +370,7 @@ class FRenderSingleScatteringWithLiveShadingCS : public FMeshMaterialShader
 		SHADER_PARAMETER_STRUCT_REF(FViewUniformShaderParameters, View)
 		SHADER_PARAMETER_STRUCT_INCLUDE(FSceneTextureParameters, SceneTextures)
 		SHADER_PARAMETER_RDG_UNIFORM_BUFFER(FSceneUniformParameters, Scene)
+		SHADER_PARAMETER_RDG_BUFFER_SRV(StructuredBuffer<float4>, EyeAdaptationBuffer)
 		SHADER_PARAMETER_STRUCT_REF(FBlueNoise, BlueNoise)
 
 		// Light data
@@ -1097,6 +1098,7 @@ void RenderSingleScatteringWithLiveShading(
 		PassParameters->View = View.ViewUniformBuffer;
 		PassParameters->SceneTextures = GetSceneTextureParameters(GraphBuilder, SceneTextures);
 		PassParameters->Scene = View.GetSceneUniforms().GetBuffer(GraphBuilder);
+		PassParameters->EyeAdaptationBuffer = GraphBuilder.CreateSRV(GetEyeAdaptationBuffer(GraphBuilder, View));
 		FBlueNoise BlueNoise = GetBlueNoiseGlobalParameters();
 		PassParameters->BlueNoise = CreateUniformBufferImmediate(BlueNoise, EUniformBufferUsage::UniformBuffer_SingleDraw);
 
diff --git a/Engine/Source/Runtime/Renderer/Private/HeterogeneousVolumes/HeterogeneousVolumesMaterialBakingPipeline.cpp b/Engine/Source/Runtime/Renderer/Private/HeterogeneousVolumes/HeterogeneousVolumesMaterialBakingPipeline.cpp
--- a/Engine/Source/Runtime/Renderer/Private/HeterogeneousVolumes/HeterogeneousVolumesMaterialBakingPipeline.cpp
+++ b/Engine/Source/Runtime/Renderer/Private/HeterogeneousVolumes/HeterogeneousVolumesMaterialBakingPipeline.cpp
@@ -21,6 +21,7 @@ class FHeterogeneousVolumesBakeMaterialCS : public FMeshMaterialShader
 		// Scene data
 		SHADER_PARAMETER_STRUCT_REF(FViewUniformShaderParameters, View)
 		SHADER_PARAMETER_RDG_UNIFORM_BUFFER(FSceneUniformParameters, Scene)
+		SHADER_PARAMETER_RDG_BUFFER_SRV(StructuredBuffer<float4>, EyeAdaptationBuffer)
 
 		// Object data
 		SHADER_PARAMETER(FMatrix44f, LocalToWorld)
@@ -136,6 +137,7 @@ void ComputeHeterogeneousVolumeBakeMaterial(
 		// Scene data
 		PassParameters->View = View.ViewUniformBuffer;
 		PassParameters->Scene = View.GetSceneUniforms().GetBuffer(GraphBuilder);
+		PassParameters->EyeAdaptationBuffer = GraphBuilder.CreateSRV(GetEyeAdaptationBuffer(GraphBuilder, View));
 
 		// Object data
 		// TODO: Convert to relative-local space
diff --git a/Engine/Source/Runtime/Renderer/Private/HeterogeneousVolumes/HeterogeneousVolumesVoxelGridPipeline.cpp b/Engine/Source/Runtime/Renderer/Private/HeterogeneousVolumes/HeterogeneousVolumesVoxelGridPipeline.cpp
--- a/Engine/Source/Runtime/Renderer/Private/HeterogeneousVolumes/HeterogeneousVolumesVoxelGridPipeline.cpp
+++ b/Engine/Source/Runtime/Renderer/Private/HeterogeneousVolumes/HeterogeneousVolumesVoxelGridPipeline.cpp
@@ -547,6 +547,7 @@ class FRasterizeBottomLevelFrustumGridCS : public FMeshMaterialShader
 		// Scene data
 		SHADER_PARAMETER_STRUCT_REF(FViewUniformShaderParameters, View)
 		SHADER_PARAMETER_RDG_UNIFORM_BUFFER(FSceneUniformParameters, Scene)
+		SHADER_PARAMETER_RDG_BUFFER_SRV(StructuredBuffer<float4>, EyeAdaptationBuffer)
 
 		// Primitive data
 		SHADER_PARAMETER(FVector3f, PrimitiveWorldBoundsMin)
@@ -847,6 +848,7 @@ class FRasterizeBottomLevelOrthoGridCS : public FMeshMaterialShader
 		// Scene data
 		SHADER_PARAMETER_STRUCT_REF(FViewUniformShaderParameters, View)
 		SHADER_PARAMETER_RDG_UNIFORM_BUFFER(FSceneUniformParameters, Scene)
+		SHADER_PARAMETER_RDG_BUFFER_SRV(StructuredBuffer<float4>, EyeAdaptationBuffer)
 
 		// Primitive data
 		SHADER_PARAMETER(FVector3f, PrimitiveWorldBoundsMin)
@@ -2066,6 +2068,7 @@ void RasterizeVolumesIntoFrustumVoxelGrid(
 				// Scene data
 				PassParameters->View = View.ViewUniformBuffer;
 				PassParameters->Scene = View.GetSceneUniforms().GetBuffer(GraphBuilder);
+				PassParameters->EyeAdaptationBuffer = GraphBuilder.CreateSRV(GetEyeAdaptationBuffer(GraphBuilder, View));
 
 				// Primitive data
 				PassParameters->PrimitiveWorldBoundsMin = FVector3f(PrimitiveBounds.Origin - PrimitiveBounds.BoxExtent);
@@ -2545,6 +2548,7 @@ void RasterizeVolumesIntoOrthoVoxelGrid(
 					// Scene data
 					PassParameters->View = View.ViewUniformBuffer;
 					PassParameters->Scene = View.GetSceneUniforms().GetBuffer(GraphBuilder);
+					PassParameters->EyeAdaptationBuffer = GraphBuilder.CreateSRV(GetEyeAdaptationBuffer(GraphBuilder, View));
 
 					// Primitive data
 					PassParameters->PrimitiveWorldBoundsMin = FVector3f(PrimitiveBounds.Origin - PrimitiveBounds.BoxExtent);
