Add variations of for single primitive rendering that do not depend on a view for the feature level and platform.

diff --git a/Engine/Source/Runtime/RenderCore/Public/RendererInterface.h b/Engine/Source/Runtime/RenderCore/Public/RendererInterface.h
--- a/Engine/Source/Runtime/RenderCore/Public/RendererInterface.h
+++ b/Engine/Source/Runtime/RenderCore/Public/RendererInterface.h
@@ -716,8 +716,10 @@ public:
 	virtual void InitializeSystemTextures(FRHICommandListImmediate& RHICmdList) = 0;
 
 	/** Create a Scene Uniform Buffer containing only the scene representation for a single primitive */
+	virtual FSceneUniformBuffer* CreateSinglePrimitiveSceneUniformBuffer(FRDGBuilder& GraphBuilder, ERHIFeatureLevel::Type FeatureLevel, FMeshBatch& Mesh) = 0;
 	virtual FSceneUniformBuffer* CreateSinglePrimitiveSceneUniformBuffer(FRDGBuilder& GraphBuilder, const FViewInfo& SceneView, FMeshBatch& Mesh) = 0;
 	/** Create a Uniform Buffer containing representation for a single primitive. (For platforms that use "UniformView" path) */
+	virtual TRDGUniformBufferRef<FBatchedPrimitiveParameters> CreateSinglePrimitiveUniformView(FRDGBuilder& GraphBuilder, ERHIFeatureLevel::Type FeatureLevel, EShaderPlatform ShaderPlatform, FMeshBatch& Mesh) = 0;
 	virtual TRDGUniformBufferRef<FBatchedPrimitiveParameters> CreateSinglePrimitiveUniformView(FRDGBuilder& GraphBuilder, const FViewInfo& SceneView, FMeshBatch& Mesh) = 0;
 
 	/** Draws a tile mesh element with the specified view. */
diff --git a/Engine/Source/Runtime/Renderer/Private/Renderer.cpp b/Engine/Source/Runtime/Renderer/Private/Renderer.cpp
--- a/Engine/Source/Runtime/Renderer/Private/Renderer.cpp
+++ b/Engine/Source/Runtime/Renderer/Private/Renderer.cpp
@@ -167,8 +167,11 @@ END_SHADER_PARAMETER_STRUCT()
 
 FSceneUniformBuffer* FRendererModule::CreateSinglePrimitiveSceneUniformBuffer(FRDGBuilder& GraphBuilder, const FViewInfo& SceneView, FMeshBatch& Mesh)
 {
-	const auto FeatureLevel = SceneView.GetFeatureLevel();
-	
+	return CreateSinglePrimitiveSceneUniformBuffer(GraphBuilder, SceneView.FeatureLevel, Mesh);
+}
+
+FSceneUniformBuffer* FRendererModule::CreateSinglePrimitiveSceneUniformBuffer(FRDGBuilder& GraphBuilder, ERHIFeatureLevel::Type FeatureLevel, FMeshBatch& Mesh)
+{
 	FSceneUniformBuffer& SceneUniforms = *GraphBuilder.AllocObject<FSceneUniformBuffer>();
 
 	if (Mesh.VertexFactory->GetPrimitiveIdStreamIndex(FeatureLevel, EVertexInputStreamType::PositionOnly) >= 0)
@@ -243,7 +246,12 @@ FSceneUniformBuffer* FRendererModule::CreateSinglePrimitiveSceneUniformBuffer(FR
 
 TRDGUniformBufferRef<FBatchedPrimitiveParameters> FRendererModule::CreateSinglePrimitiveUniformView(FRDGBuilder& GraphBuilder, const FViewInfo& SceneView, FMeshBatch& Mesh)
 {
-	check(PlatformGPUSceneUsesUniformBufferView(SceneView.GetShaderPlatform()));
+	return CreateSinglePrimitiveUniformView(GraphBuilder, SceneView.FeatureLevel, SceneView.GetShaderPlatform(), Mesh);
+}
+
+TRDGUniformBufferRef<FBatchedPrimitiveParameters> FRendererModule::CreateSinglePrimitiveUniformView(FRDGBuilder& GraphBuilder, ERHIFeatureLevel::Type FeatureLevel, EShaderPlatform ShaderPlatform, FMeshBatch& Mesh)
+{
+	check(PlatformGPUSceneUsesUniformBufferView(ShaderPlatform));
 
 	FBatchedPrimitiveParameters* BatchedPrimitiveParameters = GraphBuilder.AllocParameters<FBatchedPrimitiveParameters>();
 
@@ -251,11 +259,10 @@ TRDGUniformBufferRef<FBatchedPrimitiveParameters> FRendererModule::CreateSingleP
 	PrimitiveDataBufferDesc.Usage |= EBufferUsageFlags::UniformBuffer;
 	FRDGBufferRef PrimitiveDataBuffer = nullptr;
 	
-	ERHIFeatureLevel::Type FeatureLevel = SceneView.GetFeatureLevel();
 	if (Mesh.VertexFactory->GetPrimitiveIdStreamIndex(FeatureLevel, EVertexInputStreamType::PositionOnly) >= 0)
 	{
 		FMeshBatchElement& MeshElement = Mesh.Elements[0];
-		checkf(Mesh.Elements.Num() == 1, TEXT("Only 1 batch element currently supported by CreateSinglePrimitiveSceneUniformBuffer"));
+		checkf(Mesh.Elements.Num() == 1, TEXT("Only 1 batch element currently supported by CreateSinglePrimitiveUniformView"));
 		checkf(MeshElement.PrimitiveUniformBuffer == nullptr, TEXT("CreateSinglePrimitiveUniformView does not currently support an explicit primitive uniform buffer on vertex factories which manually fetch primitive data.  Use PrimitiveUniformBufferResource instead."));
 
 		if (MeshElement.PrimitiveUniformBufferResource)
@@ -350,7 +357,7 @@ void FRendererModule::DrawTileMesh(FCanvasRenderContext& RenderContext, FMeshPas
 
 		FRDGBuilder& GraphBuilder = RenderContext.GraphBuilder;
 		
-		FSceneUniformBuffer& SceneUniforms = *CreateSinglePrimitiveSceneUniformBuffer(GraphBuilder, View, Mesh);
+		FSceneUniformBuffer& SceneUniforms = *CreateSinglePrimitiveSceneUniformBuffer(GraphBuilder, FeatureLevel, Mesh);
 
 		if (!FRDGSystemTextures::IsValid(GraphBuilder))
 		{
diff --git a/Engine/Source/Runtime/Renderer/Private/RendererModule.h b/Engine/Source/Runtime/Renderer/Private/RendererModule.h
--- a/Engine/Source/Runtime/Renderer/Private/RendererModule.h
+++ b/Engine/Source/Runtime/Renderer/Private/RendererModule.h
@@ -49,7 +49,9 @@ public:
 	virtual uint32 GetNumDynamicLightsAffectingPrimitive(const FPrimitiveSceneInfo* PrimitiveSceneInfo,const FLightCacheInterface* LCI) override;
 	virtual void OnWorldCleanup(UWorld* World, bool bSessionEnded, bool bCleanupResources, bool bWorldChanged) override;
 	virtual void InitializeSystemTextures(FRHICommandListImmediate& RHICmdList);
+	virtual FSceneUniformBuffer* CreateSinglePrimitiveSceneUniformBuffer(FRDGBuilder& GraphBuilder, ERHIFeatureLevel::Type FeatureLevel, FMeshBatch& Mesh) override;
 	virtual FSceneUniformBuffer* CreateSinglePrimitiveSceneUniformBuffer(FRDGBuilder& GraphBuilder, const FViewInfo& SceneView, FMeshBatch& Mesh) override;
+	virtual TRDGUniformBufferRef<FBatchedPrimitiveParameters> CreateSinglePrimitiveUniformView(FRDGBuilder& GraphBuilder, ERHIFeatureLevel::Type FeatureLevel, EShaderPlatform ShaderPlatform, FMeshBatch& Mesh) override;
 	virtual TRDGUniformBufferRef<FBatchedPrimitiveParameters> CreateSinglePrimitiveUniformView(FRDGBuilder& GraphBuilder, const FViewInfo& SceneView, FMeshBatch& Mesh) override;
 	virtual void DrawTileMesh(FCanvasRenderContext& RenderContext, FMeshPassProcessorRenderState& DrawRenderState, const FSceneView& View, FMeshBatch& Mesh, bool bIsHitTesting, const FHitProxyId& HitProxyId, bool bUse128bitRT = false) override;
 	virtual void DebugLogOnCrash() override;
