SceneCapture: Moved SetupViewFamily & SetupView scene view extension calls before renderer creation to match the main render path.

This allows custom render passes to be registered by scene view extensions before a scene capture renderer is created.

diff --git a/Engine/Source/Runtime/Renderer/Private/PlanarReflectionRendering.cpp b/Engine/Source/Runtime/Renderer/Private/PlanarReflectionRendering.cpp
--- a/Engine/Source/Runtime/Renderer/Private/PlanarReflectionRendering.cpp
+++ b/Engine/Source/Runtime/Renderer/Private/PlanarReflectionRendering.cpp
@@ -492,7 +492,7 @@ static void UpdatePlanarReflectionContentsWithoutRendering_RenderThread(
 
 extern void BuildProjectionMatrix(FIntPoint RenderTargetSize, float FOV, float InNearClippingPlane, FMatrix& ProjectionMatrix);
 
-extern void SetupViewFamilyForSceneCapture(
+extern TArray<FSceneView*> SetupViewFamilyForSceneCapture(
 	FSceneViewFamily& ViewFamily,
 	USceneCaptureComponent* SceneCaptureComponent,
 	const TArrayView<const FSceneCaptureViewInfo> Views,
@@ -504,6 +504,10 @@ extern void SetupViewFamilyForSceneCapture(
 	const AActor* ViewActor,
 	int32 CubemapFaceIndex);
 
+extern void SetupSceneViewExtensionsForSceneCapture(
+	FSceneViewFamily& ViewFamily,
+	TConstArrayView<FSceneView*> Views);
+
 void FScene::UpdatePlanarReflectionContents(UPlanarReflectionComponent* CaptureComponent, FSceneRenderer& MainSceneRenderer)
 {
 	check(CaptureComponent);
@@ -638,7 +642,7 @@ void FScene::UpdatePlanarReflectionContents(UPlanarReflectionComponent* CaptureC
 		ViewExtensionContext.bStereoEnabled = true;
 		ViewFamily.ViewExtensions = GEngine->ViewExtensions->GatherActiveExtensions(ViewExtensionContext);
 
-		SetupViewFamilyForSceneCapture(
+		TArray<FSceneView*> Views = SetupViewFamilyForSceneCapture(
 			ViewFamily,
 			CaptureComponent,
 			SceneCaptureViewInfo, CaptureComponent->MaxViewDistanceOverride,
@@ -653,30 +657,25 @@ void FScene::UpdatePlanarReflectionContents(UPlanarReflectionComponent* CaptureC
 		ViewFamily.SetScreenPercentageInterface(FSceneRenderer::ForkScreenPercentageInterface(
 			MainSceneRenderer.ViewFamily.GetScreenPercentageInterface(), ViewFamily));
 
+		for (FSceneView* View : Views)
+		{
+			View->GlobalClippingPlane = MirrorPlane;
+			// Jitter can't be removed completely due to the clipping plane
+			// Also, this prevents the prefilter pass, which reads from jittered depth, from having to do special handling of it's depth-dependent input
+			View->bAllowTemporalJitter = false;
+			View->bRenderSceneTwoSided = CaptureComponent->bRenderSceneTwoSided;
+		}
+
+		// Call SetupViewFamily & SetupView on scene view extensions before renderer creation
+		SetupSceneViewExtensionsForSceneCapture(ViewFamily, Views);
+
 		FSceneRenderer* SceneRenderer = FSceneRenderer::CreateSceneRenderer(&ViewFamily, nullptr);
 
 		// Disable screen percentage on planar reflection renderer if main one has screen percentage disabled.
 		SceneRenderer->ViewFamily.EngineShowFlags.ScreenPercentage = MainSceneRenderer.ViewFamily.EngineShowFlags.ScreenPercentage;
 
-		for (const FSceneViewExtensionRef& Extension : SceneRenderer->ViewFamily.ViewExtensions)
-		{
-			Extension->SetupViewFamily(SceneRenderer->ViewFamily);
-		}
-
 		for (int32 ViewIndex = 0; ViewIndex < SceneCaptureViewInfo.Num(); ++ViewIndex)
 		{
-			FViewInfo& ViewInfo = SceneRenderer->Views[ViewIndex];
-			ViewInfo.GlobalClippingPlane = MirrorPlane;
-			// Jitter can't be removed completely due to the clipping plane
-			// Also, this prevents the prefilter pass, which reads from jittered depth, from having to do special handling of it's depth-dependent input
-			ViewInfo.bAllowTemporalJitter = false;
-			ViewInfo.bRenderSceneTwoSided = CaptureComponent->bRenderSceneTwoSided;
-
-			for (const FSceneViewExtensionRef& Extension : SceneRenderer->ViewFamily.ViewExtensions)
-			{
-				Extension->SetupView(SceneRenderer->ViewFamily, ViewInfo);
-			}
-
 			CaptureComponent->ProjectionWithExtraFOV[ViewIndex] = SceneCaptureViewInfo[ViewIndex].ProjectionMatrix;
 
 			const bool bIsStereo = IStereoRendering::IsStereoEyeView(MainSceneRenderer.Views[0]);
diff --git a/Engine/Source/Runtime/Renderer/Private/SceneCaptureRendering.cpp b/Engine/Source/Runtime/Renderer/Private/SceneCaptureRendering.cpp
--- a/Engine/Source/Runtime/Renderer/Private/SceneCaptureRendering.cpp
+++ b/Engine/Source/Runtime/Renderer/Private/SceneCaptureRendering.cpp
@@ -792,7 +792,7 @@ void GetShowOnlyAndHiddenComponents(USceneCaptureComponent* SceneCaptureComponen
 	}
 }
 
-void SetupViewFamilyForSceneCapture(
+TArray<FSceneView*> SetupViewFamilyForSceneCapture(
 	FSceneViewFamily& ViewFamily,
 	USceneCaptureComponent* SceneCaptureComponent,
 	const TArrayView<const FSceneCaptureViewInfo> Views,
@@ -814,6 +814,9 @@ void SetupViewFamilyForSceneCapture(
 	ViewFamily.FrameNumber = ViewFamily.Scene->GetFrameNumber();
 	ViewFamily.FrameCounter = GFrameCounter;
 
+	TArray<FSceneView*> ViewPtrArray;
+	ViewPtrArray.Reserve(Views.Num());
+
 	for (int32 ViewIndex = 0; ViewIndex < Views.Num(); ++ViewIndex)
 	{
 		const FSceneCaptureViewInfo& SceneCaptureViewInfo = Views[ViewIndex];
@@ -864,6 +867,7 @@ void SetupViewFamilyForSceneCapture(
 		}
 		
 		ViewFamily.Views.Add(View);
+		ViewPtrArray.Add(View);
 
 		View->StartFinalPostprocessSettings(SceneCaptureViewInfo.ViewOrigin);
 
@@ -891,6 +895,26 @@ void SetupViewFamilyForSceneCapture(
 		}
 		View->EndFinalPostprocessSettings(ViewInitOptions);
 	}
+
+	return ViewPtrArray;
+}
+
+void SetupSceneViewExtensionsForSceneCapture(
+	FSceneViewFamily& ViewFamily,
+	TConstArrayView<FSceneView*> Views)
+{
+	for (const FSceneViewExtensionRef& Extension : ViewFamily.ViewExtensions)
+	{
+		Extension->SetupViewFamily(ViewFamily);
+	}
+
+	for (FSceneView* View : Views)
+	{
+		for (const FSceneViewExtensionRef& Extension : ViewFamily.ViewExtensions)
+		{
+			Extension->SetupView(ViewFamily, *View);
+		}
+	}
 }
 
 static FSceneRenderer* CreateSceneRendererForSceneCapture(
@@ -904,6 +928,8 @@ static FSceneRenderer* CreateSceneRendererForSceneCapture(
 	float MaxViewDistance,
 	float InFOV,
 	bool bCaptureSceneColor,
+	bool bCameraCut2D,
+	bool bCopyMainViewTemporalSettings2D,
 	FPostProcessSettings* PostProcessSettings,
 	float PostProcessBlendWeight,
 	const AActor* ViewActor, 
@@ -919,9 +945,10 @@ static FSceneRenderer* CreateSceneRendererForSceneCapture(
 	SceneCaptureViewInfo.FOV = InFOV;
 
 	bool bInheritMainViewScreenPercentage = false;
+	USceneCaptureComponent2D* SceneCaptureComponent2D = Cast<USceneCaptureComponent2D>(SceneCaptureComponent);
 
 	// Use camera position correction for ortho scene captures
-	if(USceneCaptureComponent2D * SceneCaptureComponent2D = Cast<USceneCaptureComponent2D>(SceneCaptureComponent))
+	if(IsValid(SceneCaptureComponent2D))
 	{
 		if (!SceneCaptureViewInfo.IsPerspectiveProjection() && SceneCaptureComponent2D->bUpdateOrthoPlanes)
 		{
@@ -944,7 +971,7 @@ static FSceneRenderer* CreateSceneRendererForSceneCapture(
 	FSceneViewExtensionContext ViewExtensionContext(Scene);
 	ViewFamily.ViewExtensions = GEngine->ViewExtensions->GatherActiveExtensions(ViewExtensionContext);
 	
-	SetupViewFamilyForSceneCapture(
+	TArray<FSceneView*> Views = SetupViewFamilyForSceneCapture(
 		ViewFamily,
 		SceneCaptureComponent,
 		MakeArrayView(&SceneCaptureViewInfo, 1),
@@ -961,7 +988,6 @@ static FSceneRenderer* CreateSceneRendererForSceneCapture(
 
 	if (bInheritMainViewScreenPercentage)
 	{
-		USceneCaptureComponent2D* SceneCaptureComponent2D = Cast<USceneCaptureComponent2D>(SceneCaptureComponent);
 		ViewFamily.EngineShowFlags.ScreenPercentage = SceneCaptureComponent2D->MainViewFamily->EngineShowFlags.ScreenPercentage;
 		ViewFamily.SetScreenPercentageInterface(SceneCaptureComponent2D->MainViewFamily->GetScreenPercentageInterface()->Fork_GameThread(ViewFamily));
 	}
@@ -973,6 +999,65 @@ static FSceneRenderer* CreateSceneRendererForSceneCapture(
 			ViewFamily, /* GlobalResolutionFraction = */ 1.0f));
 	}
 
+	if (IsValid(SceneCaptureComponent2D))
+	{
+		// Scene capture 2D only support a single view
+		check(Views.Num() == 1);
+
+		// Ensure that the views for this scene capture reflect any simulated camera motion for this frame
+		TOptional<FTransform> PreviousTransform = FMotionVectorSimulation::Get().GetPreviousTransform(SceneCaptureComponent2D);
+
+		// Update views with scene capture 2d specific settings
+		if (PreviousTransform.IsSet())
+		{
+			Views[0]->PreviousViewTransform = PreviousTransform.GetValue();
+		}
+
+		if (SceneCaptureComponent2D->bEnableClipPlane)
+		{
+			Views[0]->GlobalClippingPlane = FPlane(SceneCaptureComponent2D->ClipPlaneBase, SceneCaptureComponent2D->ClipPlaneNormal.GetSafeNormal());
+			// Jitter can't be removed completely due to the clipping plane
+			Views[0]->bAllowTemporalJitter = false;
+		}
+
+		Views[0]->bCameraCut = bCameraCut2D;
+
+		if (bCopyMainViewTemporalSettings2D)
+		{
+			const FSceneViewFamily* MainViewFamily = SceneCaptureComponent2D->MainViewFamily;
+			const FSceneView& SourceView = *MainViewFamily->Views[0];
+
+			Views[0]->AntiAliasingMethod = SourceView.AntiAliasingMethod;
+			Views[0]->PrimaryScreenPercentageMethod = SourceView.PrimaryScreenPercentageMethod;
+
+			if (Views[0]->State && SourceView.State)
+			{
+				((FSceneViewState*)Views[0]->State)->TemporalAASampleIndex = ((const FSceneViewState*)SourceView.State)->TemporalAASampleIndex;
+			}
+		}
+
+		// Append component-local view extensions to the view family
+		for (int32 Index = 0; Index < SceneCaptureComponent2D->SceneViewExtensions.Num(); ++Index)
+		{
+			TSharedPtr<ISceneViewExtension, ESPMode::ThreadSafe> Extension = SceneCaptureComponent2D->SceneViewExtensions[Index].Pin();
+			if (Extension.IsValid())
+			{
+				if (Extension->IsActiveThisFrame(ViewExtensionContext))
+				{
+					ViewFamily.ViewExtensions.Add(Extension.ToSharedRef());
+				}
+			}
+			else
+			{
+				SceneCaptureComponent2D->SceneViewExtensions.RemoveAt(Index, EAllowShrinking::No);
+				--Index;
+			}
+		}
+	}
+
+	// Call SetupViewFamily & SetupView on scene view extensions before renderer creation
+	SetupSceneViewExtensionsForSceneCapture(ViewFamily, Views);
+
 	return FSceneRenderer::CreateSceneRenderer(&ViewFamily, nullptr);
 }
 
@@ -1258,6 +1343,19 @@ void FScene::UpdateSceneCaptureContents(USceneCaptureComponent2D* CaptureCompone
 			return;
 		}
 
+		
+		// Copy temporal AA related settings for main view camera scene capture, to match jitter.  Don't match if the resolution divisor is set,
+		// if it's set to ignore screen percentage, or if it's final color, which will run its own AA.  For custom render passes (handled above),
+		// computed jitter results are copied from the main view later in FSceneRenderer::PrepareViewStateForVisibility, but this doesn't work
+		// for regular scene captures, because they run in a separate scene renderer before the main view, where the main view's results haven't
+		// been computed yet.
+		const bool bCopyMainViewTemporalSettings2D = (CaptureComponent->ShouldRenderWithMainViewCamera() && CaptureComponent->MainViewFamily &&
+			CaptureComponent->MainViewResolutionDivisor.X <= 1 && CaptureComponent->MainViewResolutionDivisor.Y <= 1 && !CaptureComponent->ShouldIgnoreScreenPercentage() &&
+			CaptureComponent->CaptureSource != ESceneCaptureSource::SCS_FinalColorLDR &&
+			CaptureComponent->CaptureSource != ESceneCaptureSource::SCS_FinalColorHDR &&
+			CaptureComponent->CaptureSource != ESceneCaptureSource::SCS_FinalToneCurveHDR);
+		const bool bCameraCut2D = bCopyMainViewTemporalSettings2D ? CaptureComponent->MainViewFamily->Views[0]->bCameraCut : CaptureComponent->bCameraCutThisFrame;
+
 		FSceneRenderer* SceneRenderer = CreateSceneRendererForSceneCapture(
 			this, 
 			CaptureComponent, 
@@ -1269,41 +1367,19 @@ void FScene::UpdateSceneCaptureContents(USceneCaptureComponent2D* CaptureCompone
 			CaptureComponent->MaxViewDistanceOverride, 
 			CaptureComponent->FOVAngle,
 			bUseSceneColorTexture,
+			bCameraCut2D,
+			bCopyMainViewTemporalSettings2D,
 			&CaptureComponent->PostProcessSettings, 
 			CaptureComponent->PostProcessBlendWeight,
 			CaptureComponent->GetViewOwner());
 
 		check(SceneRenderer != nullptr);
 
-		// Copy temporal AA related settings for main view camera scene capture, to match jitter.  Don't match if the resolution divisor is set,
-		// if it's set to ignore screen percentage, or if it's final color, which will run its own AA.  For custom render passes (handled above),
-		// computed jitter results are copied from the main view later in FSceneRenderer::PrepareViewStateForVisibility, but this doesn't work
-		// for regular scene captures, because they run in a separate scene renderer before the main view, where the main view's results haven't
-		// been computed yet.
-		if (CaptureComponent->ShouldRenderWithMainViewCamera() && CaptureComponent->MainViewFamily && SceneRenderer->Views[0].ViewState && CaptureComponent->MainViewFamily->Views[0]->State &&
-			CaptureComponent->MainViewResolutionDivisor.X <= 1 && CaptureComponent->MainViewResolutionDivisor.Y <= 1 && !CaptureComponent->ShouldIgnoreScreenPercentage() &&
-			CaptureComponent->CaptureSource != ESceneCaptureSource::SCS_FinalColorLDR &&
-			CaptureComponent->CaptureSource != ESceneCaptureSource::SCS_FinalColorHDR &&
-			CaptureComponent->CaptureSource != ESceneCaptureSource::SCS_FinalToneCurveHDR)
-		{
-			const FSceneViewFamily* MainViewFamily = CaptureComponent->MainViewFamily;
-			const FSceneView& SourceView = *MainViewFamily->Views[0];
-			const FSceneViewState& SourceViewState = *(const FSceneViewState*)SourceView.State;
-
-			FViewInfo& DestView = SceneRenderer->Views[0];
-
-			DestView.AntiAliasingMethod = SourceView.AntiAliasingMethod;
-			DestView.PrimaryScreenPercentageMethod = SourceView.PrimaryScreenPercentageMethod;
-			DestView.ViewState->TemporalAASampleIndex = SourceViewState.TemporalAASampleIndex;
-			DestView.bSceneCaptureMainViewJitter = true;
-
-			CaptureComponent->bCameraCutThisFrame = SourceView.bCameraCut;
-		}
-
 		// When bIsMultipleSceneCapture is true, set bIsFirstSceneRenderer to false, which tells the scene renderer it can skip RHI resource flush, saving performance
 		bool bIsMultipleSceneCapture = CaptureComponent->SetFrameUpdated();
 		SceneRenderer->bIsFirstSceneRenderer = !bIsMultipleSceneCapture;
 
+		SceneRenderer->Views[0].bSceneCaptureMainViewJitter = bCopyMainViewTemporalSettings2D;
 		SceneRenderer->Views[0].bFogOnlyOnRenderedOpaque = CaptureComponent->bConsiderUnrenderedOpaquePixelAsFullyTranslucent;
 
 		SceneRenderer->ViewFamily.SceneCaptureCompositeMode = CaptureComponent->CompositeMode;
@@ -1324,62 +1400,6 @@ void FScene::UpdateSceneCaptureContents(USceneCaptureComponent2D* CaptureCompone
 			ViewStateInterface->RemoveLumenSceneData(this);
 		}
 
-		// Ensure that the views for this scene capture reflect any simulated camera motion for this frame
-		TOptional<FTransform> PreviousTransform = FMotionVectorSimulation::Get().GetPreviousTransform(CaptureComponent);
-
-		// Process Scene View extensions for the capture component
-		{
-			FSceneViewExtensionContext ViewExtensionContext(SceneRenderer->Scene);
-
-			for (int32 Index = 0; Index < CaptureComponent->SceneViewExtensions.Num(); ++Index)
-			{
-				TSharedPtr<ISceneViewExtension, ESPMode::ThreadSafe> Extension = CaptureComponent->SceneViewExtensions[Index].Pin();
-				if (Extension.IsValid())
-				{
-					if (Extension->IsActiveThisFrame(ViewExtensionContext))
-					{
-						SceneRenderer->ViewFamily.ViewExtensions.Add(Extension.ToSharedRef());
-					}
-				}
-				else
-				{
-					CaptureComponent->SceneViewExtensions.RemoveAt(Index, EAllowShrinking::No);
-					--Index;
-				}
-			}
-		}
-
-		for (const FSceneViewExtensionRef& Extension : SceneRenderer->ViewFamily.ViewExtensions)
-		{
-			Extension->SetupViewFamily(SceneRenderer->ViewFamily);
-		}
-
-		{
-			FPlane ClipPlane = FPlane(CaptureComponent->ClipPlaneBase, CaptureComponent->ClipPlaneNormal.GetSafeNormal());
-
-			for (FSceneView& View : SceneRenderer->Views)
-			{
-				if (PreviousTransform.IsSet())
-				{
-					View.PreviousViewTransform = PreviousTransform.GetValue();
-				}
-
-				View.bCameraCut = CaptureComponent->bCameraCutThisFrame;
-
-				if (CaptureComponent->bEnableClipPlane)
-				{
-					View.GlobalClippingPlane = ClipPlane;
-					// Jitter can't be removed completely due to the clipping plane
-					View.bAllowTemporalJitter = false;
-				}
-
-				for (const FSceneViewExtensionRef& Extension : SceneRenderer->ViewFamily.ViewExtensions)
-				{
-					Extension->SetupView(SceneRenderer->ViewFamily, View);
-				}
-			}
-		}
-
 		// Reset scene capture's camera cut.
 		CaptureComponent->bCameraCutThisFrame = false;
 
@@ -1658,25 +1678,17 @@ void FScene::UpdateSceneCaptureContents(USceneCaptureComponentCube* CaptureCompo
 				FMatrix ProjectionMatrix;
 				ComputeProjectionMatrix(TargetFace, ViewRotationMatrix, ProjectionMatrix);
 
+				constexpr bool bCameraCut2D = false;
+				constexpr bool bCopyMainViewTemporalSettings2D = false;
 				FSceneRenderer* SceneRenderer = CreateSceneRendererForSceneCapture(this, CaptureComponent,
 					TextureTarget->GameThread_GetRenderTargetResource(), CaptureSize, ViewRotationMatrix,
 					Location, ProjectionMatrix, CaptureComponent->MaxViewDistanceOverride, FOVInDegrees,
-					bCaptureSceneColor, &CaptureComponent->PostProcessSettings, CaptureComponent->PostProcessBlendWeight, CaptureComponent->GetViewOwner(), faceidx);
+					bCaptureSceneColor, bCameraCut2D, bCopyMainViewTemporalSettings2D, &CaptureComponent->PostProcessSettings, CaptureComponent->PostProcessBlendWeight, CaptureComponent->GetViewOwner(), faceidx);
 
 				// When bIsMultipleSceneCapture is true, set bIsFirstSceneRenderer to false, which tells the scene renderer it can skip RHI resource flush, saving performance.
 				// We can also skip RHI resource flush on faces after the first.
 				SceneRenderer->bIsFirstSceneRenderer = (faceidx == 0) && !bIsMultipleSceneCapture;
 
-				for (const FSceneViewExtensionRef& Extension : SceneRenderer->ViewFamily.ViewExtensions)
-				{
-					Extension->SetupViewFamily(SceneRenderer->ViewFamily);
-
-					for (FSceneView& View : SceneRenderer->Views)
-					{
-						Extension->SetupView(SceneRenderer->ViewFamily, View);
-					}
-				}
-
 				for (const FSceneViewExtensionRef& Extension : SceneRenderer->ViewFamily.ViewExtensions)
 				{
 					Extension->BeginRenderViewFamily(SceneRenderer->ViewFamily);
@@ -1823,7 +1835,7 @@ void FScene::UpdateSceneCaptureContents(USceneCaptureComponentCube* CaptureCompo
 			FSceneViewExtensionContext ViewExtensionContext(this);
 			ViewFamily.ViewExtensions = GEngine->ViewExtensions->GatherActiveExtensions(ViewExtensionContext);
 
-			SetupViewFamilyForSceneCapture(
+			TArray<FSceneView*> Views = SetupViewFamilyForSceneCapture(
 				ViewFamily,
 				CaptureComponent,
 				SceneCaptureViewInfos,
@@ -1843,6 +1855,9 @@ void FScene::UpdateSceneCaptureContents(USceneCaptureComponentCube* CaptureCompo
 			ViewFamily.SetScreenPercentageInterface(new FLegacyScreenPercentageDriver(
 				ViewFamily, /* GlobalResolutionFraction = */ 1.0f));
 
+			// Call SetupViewFamily & SetupView on scene view extensions before renderer creation
+			SetupSceneViewExtensionsForSceneCapture(ViewFamily, Views);
+
 			FSceneRenderer* SceneRenderer = FSceneRenderer::CreateSceneRenderer(&ViewFamily, nullptr);
 
 			// Need view state interface to be allocated for Lumen, as it requires persistent data.  This means
@@ -1864,16 +1879,6 @@ void FScene::UpdateSceneCaptureContents(USceneCaptureComponentCube* CaptureCompo
 			// When bIsMultipleSceneCapture is true, set bIsFirstSceneRenderer to false, which tells the scene renderer it can skip RHI resource flush, saving performance
 			SceneRenderer->bIsFirstSceneRenderer = !bIsMultipleSceneCapture;
 
-			for (const FSceneViewExtensionRef& Extension : SceneRenderer->ViewFamily.ViewExtensions)
-			{
-				Extension->SetupViewFamily(SceneRenderer->ViewFamily);
-
-				for (FSceneView& View : SceneRenderer->Views)
-				{
-					Extension->SetupView(SceneRenderer->ViewFamily, View);
-				}
-			}
-
 			for (const FSceneViewExtensionRef& Extension : SceneRenderer->ViewFamily.ViewExtensions)
 			{
 				Extension->BeginRenderViewFamily(SceneRenderer->ViewFamily);
