Gate software VRS to only generate shading rate images if r.Nanite.SoftwareVRS is enabled, because that's the only place they're consumed.

diff --git a/Engine/Source/Runtime/Renderer/Private/VariableRateShading/VariableRateShadingImageManager.cpp b/Engine/Source/Runtime/Renderer/Private/VariableRateShading/VariableRateShadingImageManager.cpp
--- a/Engine/Source/Runtime/Renderer/Private/VariableRateShading/VariableRateShadingImageManager.cpp
+++ b/Engine/Source/Runtime/Renderer/Private/VariableRateShading/VariableRateShadingImageManager.cpp
@@ -66,6 +66,7 @@ TAutoConsoleVariable<int32> CVarEnableVRSSoftwareImage(
 	TEXT("r.VRS.EnableSoftware"),
 	0,
 	TEXT("Enables software (2x2 tile size) Shading Rate Image generation for use with Nanite Software VRS. Allows generating iamges even when r.VRS.Enable/r.VRS.Support=0 or Tier 2 VRS is unsupported by the hardware.")
+	TEXT("Image generation will only be enabled if r.Nanite.SoftwareVRS is also set to 1.")
 	TEXT("0: Off, 1: On"),
 	ECVF_RenderThreadSafe);
 
@@ -520,9 +521,16 @@ FRDGTextureRef FVariableRateShadingImageManager::GetVariableRateShadingImage(FRD
 void FVariableRateShadingImageManager::PrepareImageBasedVRS(FRDGBuilder& GraphBuilder, const FSceneViewFamily& ViewFamily, const FMinimalSceneTextures& SceneTextures)
 {
 	EShaderPlatform ShaderPlatform = ViewFamily.Scene->GetShaderPlatform();
+	static const auto CVarNaniteSoftwareVRS = IConsoleManager::Get().FindConsoleVariable(TEXT("r.Nanite.SoftwareVRS"));
 
-	bHardwareVRSEnabledForFrame = IsAttachmentVRSEnabled() && HardwareVariableRateShadingSupportedByPlatform(ShaderPlatform); // Additional check is required here for preview levels
-	bSoftwareVRSEnabledForFrame = CVarEnableVRSSoftwareImage.GetValueOnRenderThread() > 0 && IsFeatureLevelSupported(ShaderPlatform, ERHIFeatureLevel::SM6);
+	bHardwareVRSEnabledForFrame = 
+		IsAttachmentVRSEnabled() &&
+		HardwareVariableRateShadingSupportedByPlatform(ShaderPlatform); // Additional check is required here for preview levels
+
+	bSoftwareVRSEnabledForFrame = 
+		CVarEnableVRSSoftwareImage.GetValueOnRenderThread() > 0 &&
+		CVarNaniteSoftwareVRS->GetInt() > 0 &&
+		IsFeatureLevelSupported(ShaderPlatform, ERHIFeatureLevel::SM6);
 
 	if (!IsVRSEnabledForFrame())
 	{
