Added initialization to safe values and bounds on the virtual shadow map light grid access to prevent any potential issues on GPU.

diff --git a/Engine/Shaders/Private/VirtualShadowMaps/VirtualShadowMapLightGrid.ush b/Engine/Shaders/Private/VirtualShadowMaps/VirtualShadowMapLightGrid.ush
--- a/Engine/Shaders/Private/VirtualShadowMaps/VirtualShadowMapLightGrid.ush
+++ b/Engine/Shaders/Private/VirtualShadowMaps/VirtualShadowMapLightGrid.ush
@@ -13,10 +13,12 @@ FCulledLightsGridHeader VirtualShadowMapGetLightsGrid(uint2 PixelPos, float Scen
 	// Has no effect unless INSTANCED_STEREO is enabled. Used in GetLightGridData to select ForwardLightData or ForwardLightDataISR, but then only uses redundant parameters anyway.
 	uint EyeIndex = 0U;
 
-	uint GridLinearIndex = ComputeLightGridCellIndex(PixelPos, SceneDepth, EyeIndex);
+	// Index 0 is always going to a valid index and the NumCulledLightsGrid is initialized to zero making acess safe.
+	uint GridLinearIndex = min(VirtualShadowMap.MaxLightGridEntryIndex, ComputeLightGridCellIndex(PixelPos, SceneDepth, EyeIndex));
 	FCulledLightsGridHeader CulledLightGridHeader = GetCulledLightsGridHeader(GridLinearIndex, EyeIndex);
 
 	// Replace light count with our pruned count
+	checkStructuredBufferAccessSlow(VirtualShadowMap.NumCulledLightsGrid, GridLinearIndex);
 	CulledLightGridHeader.NumLights = VirtualShadowMap.NumCulledLightsGrid[GridLinearIndex];
 	return CulledLightGridHeader;
 }
diff --git a/Engine/Source/Runtime/Renderer/Private/VirtualShadowMaps/VirtualShadowMapArray.cpp b/Engine/Source/Runtime/Renderer/Private/VirtualShadowMaps/VirtualShadowMapArray.cpp
--- a/Engine/Source/Runtime/Renderer/Private/VirtualShadowMaps/VirtualShadowMapArray.cpp
+++ b/Engine/Source/Runtime/Renderer/Private/VirtualShadowMaps/VirtualShadowMapArray.cpp
@@ -578,6 +578,16 @@ void FVirtualShadowMapArray::UpdateNextData(int32 PrevVirtualShadowMapId, int32
 	NextData[PrevVirtualShadowMapId].PageAddressOffset = FIntVector2(PageOffset.X, PageOffset.Y);
 }
 
+static FVirtualShadowMapPerViewParameters MakeEmptyVirtualShadowMapPerViewParameters(FRDGBuilder& GraphBuilder)
+{
+	FVirtualShadowMapPerViewParameters PerViewData;
+	PerViewData.MaxLightGridEntryIndex = 0u;
+	PerViewData.NumCulledLightsGrid = GraphBuilder.CreateSRV(GSystemTextures.GetDefaultStructuredBuffer(GraphBuilder, sizeof(uint32), 0u));
+	PerViewData.LightGridData = GraphBuilder.CreateSRV(GSystemTextures.GetDefaultStructuredBuffer(GraphBuilder, sizeof(uint32), 0u));
+
+	return PerViewData;
+}
+
 void FVirtualShadowMapArray::Initialize(
 	FRDGBuilder& GraphBuilder,
 	FVirtualShadowMapArrayCacheManager* InCacheManager,
@@ -636,8 +646,7 @@ void FVirtualShadowMapArray::Initialize(
 	UniformParameters.PageFlags = GraphBuilder.CreateSRV(GSystemTextures.GetDefaultStructuredBuffer(GraphBuilder, sizeof(uint32)));
 	UniformParameters.UncachedPageRectBounds = GraphBuilder.CreateSRV(GSystemTextures.GetDefaultStructuredBuffer(GraphBuilder, sizeof(FIntVector4)));
 	UniformParameters.AllocatedPageRectBounds = GraphBuilder.CreateSRV(GSystemTextures.GetDefaultStructuredBuffer(GraphBuilder, sizeof(FIntVector4)));
-	UniformParameters.PerViewData.LightGridData = GraphBuilder.CreateSRV(GSystemTextures.GetDefaultStructuredBuffer(GraphBuilder, sizeof(uint32)));
-	UniformParameters.PerViewData.NumCulledLightsGrid = GraphBuilder.CreateSRV(GSystemTextures.GetDefaultStructuredBuffer(GraphBuilder, sizeof(uint32)));
+	UniformParameters.PerViewData = MakeEmptyVirtualShadowMapPerViewParameters(GraphBuilder);
 	UniformParameters.CachePrimitiveAsDynamic = GraphBuilder.CreateSRV(GSystemTextures.GetDefaultStructuredBuffer(GraphBuilder, sizeof(uint32)));
 
 	if (bEnabled)
@@ -828,20 +837,17 @@ void FVirtualShadowMapArray::UpdateCachedUniformBuffers(FRDGBuilder& GraphBuilde
 	if (PerViewParameters.Num() == 0)
 	{
 		CachedUniformBuffers.Add(GetUncachedUniformBuffer(GraphBuilder));
-	}
+}
 	// If per-view parameters are initialized
 	else
 	{
 		CachedUniformBuffers.SetNum(PerViewParameters.Num());
 
-		for (int i = 0; i < PerViewParameters.Num(); i++)
+		for (int ViewIndex = 0; ViewIndex < PerViewParameters.Num(); ++ViewIndex)
 		{
-			FVirtualShadowMapUniformParameters* VersionedParameters = GraphBuilder.AllocParameters<FVirtualShadowMapUniformParameters>();
-			*VersionedParameters = UniformParameters;
-
-			VersionedParameters->PerViewData = PerViewParameters[i];
-
-			CachedUniformBuffers[i] = GraphBuilder.CreateUniformBuffer(VersionedParameters);
+			FVirtualShadowMapUniformParameters* VersionedParameters = GraphBuilder.AllocParameters<FVirtualShadowMapUniformParameters>(&UniformParameters);
+			VersionedParameters->PerViewData = PerViewParameters[ViewIndex];
+			CachedUniformBuffers[ViewIndex] = GraphBuilder.CreateUniformBuffer(VersionedParameters);
 		}
 	}
 }
@@ -1743,18 +1749,12 @@ void FVirtualShadowMapArray::BuildPageAllocations(
 	RDG_EVENT_SCOPE(GraphBuilder, "FVirtualShadowMapArray::BuildPageAllocation");
 	SCOPED_NAMED_EVENT(FVirtualShadowMapArray_BuildPageAllocation, FColor::Emerald);
 
-	VisualizeLight.Reset();
+	VisualizeLight.Reset(Views.Num());
 	VisualizeLight.AddDefaulted(Views.Num());
 
-	PerViewParameters.Reset();
+	PerViewParameters.Reset(Views.Num());
 	PerViewParameters.AddDefaulted(Views.Num());
 
-	for (FVirtualShadowMapPerViewParameters& Param : PerViewParameters)
-	{
-		Param.NumCulledLightsGrid = GraphBuilder.CreateSRV(GSystemTextures.GetDefaultStructuredBuffer(GraphBuilder, sizeof(uint32), 0u));
-		Param.LightGridData = GraphBuilder.CreateSRV(GSystemTextures.GetDefaultStructuredBuffer(GraphBuilder, sizeof(uint32), 0u));
-	}
-
 #if !UE_BUILD_SHIPPING
 	if (GDumpVSMLightNames)
 	{
@@ -1899,6 +1899,7 @@ void FVirtualShadowMapArray::BuildPageAllocations(
 		RDG_EVENT_SCOPE_CONDITIONAL(GraphBuilder, Views.Num() > 1, "View%d", ViewIndex);
 
 		const FViewInfo &View = Views[ViewIndex];
+		FVirtualShadowMapPerViewParameters& PerViewData = PerViewParameters[ViewIndex];
 
 		// Gather directional light virtual shadow maps
 		TArray<int32, SceneRenderingAllocator> DirectionalLightIds;
@@ -1919,6 +1920,7 @@ void FVirtualShadowMapArray::BuildPageAllocations(
 		// This view contained no local lights (that were stored in the light grid), and no directional lights, so nothing to do.
 		if (View.ForwardLightingResources.LocalLightVisibleLightInfosIndex.Num() + DirectionalLightIds.Num() == 0)
 		{
+			PerViewData = MakeEmptyVirtualShadowMapPerViewParameters(GraphBuilder);
 			continue;
 		}
 
@@ -1979,14 +1981,14 @@ void FVirtualShadowMapArray::BuildPageAllocations(
 					FComputeShaderUtils::AddPass(GraphBuilder, RDG_EVENT_NAME("PruneLightGrid"), ComputeShader, PassParameters, FComputeShaderUtils::GetGroupCount(NumLightGridCells, FPruneLightGridCS::DefaultCSGroupX));
 				};
 
-				PerViewParameters[ViewIndex].LightGridData = GraphBuilder.CreateSRV(PrunedLightGridDataRDG);
-				PerViewParameters[ViewIndex].NumCulledLightsGrid = GraphBuilder.CreateSRV(PrunedNumCulledLightsGridRDG);
+				PerViewData.LightGridData = GraphBuilder.CreateSRV(PrunedLightGridDataRDG);
+				PerViewData.NumCulledLightsGrid = GraphBuilder.CreateSRV(PrunedNumCulledLightsGridRDG);
+				PerViewData.MaxLightGridEntryIndex = NumLightGridCells - 1u;
 
 				// These will be used for subsequent shaders in this view loop
 				// But because they will be overwritten on each loop, any other shaders that want to access
 				// per-view VSM light grid data need to use the cached uniform buffers generated in UpdateCachedUniformBuffers()
-				UniformParameters.PerViewData.LightGridData = PerViewParameters[ViewIndex].LightGridData;
-				UniformParameters.PerViewData.NumCulledLightsGrid = PerViewParameters[ViewIndex].NumCulledLightsGrid;
+				UniformParameters.PerViewData = PerViewData;
 			}
 
 			// Mark pages based on projected depth buffer pixels
diff --git a/Engine/Source/Runtime/Renderer/Private/VirtualShadowMaps/VirtualShadowMapArray.h b/Engine/Source/Runtime/Renderer/Private/VirtualShadowMaps/VirtualShadowMapArray.h
--- a/Engine/Source/Runtime/Renderer/Private/VirtualShadowMaps/VirtualShadowMapArray.h
+++ b/Engine/Source/Runtime/Renderer/Private/VirtualShadowMaps/VirtualShadowMapArray.h
@@ -124,6 +124,7 @@ BEGIN_SHADER_PARAMETER_STRUCT(FVirtualShadowMapPerViewParameters, )
 	// Still references the original indices from the global light grid
 	SHADER_PARAMETER_RDG_BUFFER_SRV(StructuredBuffer<uint>, LightGridData)
 	SHADER_PARAMETER_RDG_BUFFER_SRV(StructuredBuffer<uint>, NumCulledLightsGrid)
+	SHADER_PARAMETER(uint32, MaxLightGridEntryIndex)	
 END_SHADER_PARAMETER_STRUCT()
 
 BEGIN_GLOBAL_SHADER_PARAMETER_STRUCT(FVirtualShadowMapUniformParameters, )
