Skip processing primitives which don't add to main or decal raytracing instances

diff --git a/Engine/Source/Runtime/Renderer/Private/RayTracing/RayTracing.cpp b/Engine/Source/Runtime/Renderer/Private/RayTracing/RayTracing.cpp
--- a/Engine/Source/Runtime/Renderer/Private/RayTracing/RayTracing.cpp
+++ b/Engine/Source/Runtime/Renderer/Private/RayTracing/RayTracing.cpp
@@ -1355,6 +1355,13 @@ namespace RayTracing
 			// TODO: Debug Visualization to highlight primitives using this?
 			const bool bNeedDecalInstance = RelevantPrimitive.CachedMeshCommandFlags.bAnySegmentsDecal && !ShouldExcludeDecals();
 
+			// skip if not needed for main or decal - default values for bAllSegmentsDecal is true because it's updated with & op for added cached segments
+			// but if there are no cached command indices then default value of true is kept but bAnySegmentsDecal will false as well then.
+			if (!bNeedMainInstance && !bNeedDecalInstance)
+			{
+				continue;
+			}
+
 			if (ShouldExcludeDecals() && RelevantPrimitive.CachedMeshCommandFlags.bAllSegmentsDecal)
 			{
 				continue;
