Added missing definition of LIGHT_LINK_STRIDE in FDebugLightGridPS.

diff --git a/Engine/Source/Runtime/Renderer/Private/LightGridInjection.cpp b/Engine/Source/Runtime/Renderer/Private/LightGridInjection.cpp
--- a/Engine/Source/Runtime/Renderer/Private/LightGridInjection.cpp
+++ b/Engine/Source/Runtime/Renderer/Private/LightGridInjection.cpp
@@ -1181,6 +1181,7 @@ class FDebugLightGridPS : public FGlobalShader
 		OutEnvironment.CompilerFlags.Add(CFLAG_Debug);
 		OutEnvironment.CompilerFlags.Add(CFLAG_AllowTypedUAVLoads);
 		OutEnvironment.SetDefine(TEXT("SHADER_DEBUG_LIGHT_GRID_PS"), 1);
+		OutEnvironment.SetDefine(TEXT("LIGHT_LINK_STRIDE"), LightLinkStride);
 		FForwardLightingParameters::ModifyCompilationEnvironment(Parameters.Platform, OutEnvironment);
 	}
 };
