Fix support for third party TAA with mobile renderer

diff --git a/Engine/Source/Runtime/Renderer/Private/PostProcess/PostProcessEyeAdaptation.cpp b/Engine/Source/Runtime/Renderer/Private/PostProcess/PostProcessEyeAdaptation.cpp
--- a/Engine/Source/Runtime/Renderer/Private/PostProcess/PostProcessEyeAdaptation.cpp
+++ b/Engine/Source/Runtime/Renderer/Private/PostProcess/PostProcessEyeAdaptation.cpp
@@ -591,11 +591,6 @@ public:
 		SHADER_PARAMETER_RDG_BUFFER_SRV(StructuredBuffer<float4>, EyeAdaptationBuffer)
 		SHADER_PARAMETER_RDG_TEXTURE_UAV(RWTexture2D<float4>, RWEyeAdaptationTexture)
 		END_SHADER_PARAMETER_STRUCT()
-
-		static bool ShouldCompilePermutation(const FGlobalShaderPermutationParameters& Parameters)
-	{
-		return IsFeatureLevelSupported(Parameters.Platform, ERHIFeatureLevel::SM5);
-	}
 };
 
 IMPLEMENT_GLOBAL_SHADER(FCopyEyeAdaptationToTextureCS, "/Engine/Private/PostProcessEyeAdaptation.usf", "CopyEyeAdaptationToTextureCS", SF_Compute);
