Fix Cascading Shadow Maps failing to render in the right eye when using Instanced Stereo rendering.

diff --git a/Engine/Source/Runtime/Renderer/Private/ShadowRendering.cpp b/Engine/Source/Runtime/Renderer/Private/ShadowRendering.cpp
--- a/Engine/Source/Runtime/Renderer/Private/ShadowRendering.cpp
+++ b/Engine/Source/Runtime/Renderer/Private/ShadowRendering.cpp
@@ -933,8 +933,8 @@ void FProjectedShadowInfo::SetupProjectionStencilMask(
 		//       This means doing 2x renders, but letting the scissor rect kill the undersired half. Drawing the full mask once is easy, but since the outer
 		//       loop is over each view, the stencil mask is not retained when the right view comes around.
 		// TODO: Support instanced stereo properly in the projection stenciling pass.
-		const bool bIsInstancedStereoEmulated = View->bIsInstancedStereoEnabled && !View->bIsMobileMultiViewEnabled && IStereoRendering::IsStereoEyeView(*View);
-		if (bIsInstancedStereoEmulated && ProjectionStencilingPasses.IsValidIndex(View->PrimaryViewIndex))
+		const bool bIsInstancedStereoBypassed = View->bIsInstancedStereoEnabled && !View->bIsMobileMultiViewEnabled&& IStereoRendering::IsStereoEyeView(*View);
+		if (bIsInstancedStereoBypassed && ProjectionStencilingPasses.IsValidIndex(View->PrimaryViewIndex))
 		{
 			ensure(ProjectionStencilingPasses[View->PrimaryViewIndex]->GetInstanceCullingMode() == EInstanceCullingMode::Stereo);
 
@@ -1081,8 +1081,8 @@ void FProjectedShadowInfo::SetupProjectionStencilMask(
 
 		// Shadow projection stenciling is special-cased to run per-view for instanced stereo views.
 		// TODO: Support instanced stereo properly in the projection stenciling pass.
-		const bool bIsInstancedStereoEmulated = View->bIsInstancedStereoEnabled && !View->bIsMobileMultiViewEnabled && IStereoRendering::IsStereoEyeView(*View);
-		const uint32 NumberOfInstances = bIsInstancedStereoEmulated ? 1 : View->GetStereoPassInstanceFactor();
+		const bool bIsInstancedStereoBypassed = View->bIsInstancedStereoEnabled && !View->bIsMobileMultiViewEnabled&& IStereoRendering::IsStereoEyeView(*View);
+		const uint32 NumberOfInstances = bIsInstancedStereoBypassed ? 1 : View->GetStereoPassInstanceFactor();
 
 		// Draw the frustum using the stencil buffer to mask just the pixels which are inside the shadow frustum.
 		RHICmdList.DrawIndexedPrimitive(GCubeIndexBuffer.IndexBufferRHI, 0, 0, 8, 0, 12, NumberOfInstances);
@@ -1248,14 +1248,14 @@ void FProjectedShadowInfo::RenderProjection(
 		PassParameters->ShadowTexture1 = GraphBuilder.RegisterExternalTexture(RenderTargets.ColorTargets[1]);
 	}
 
-	const bool bIsInstancedStereoEmulated = View->bIsInstancedStereoEnabled && !View->bIsMobileMultiViewEnabled && IStereoRendering::IsStereoEyeView(*View);
+	const bool bIsInstancedStereoBypassed = View->bIsInstancedStereoEnabled && !View->bIsMobileMultiViewEnabled&& IStereoRendering::IsStereoEyeView(*View);
 	if (ViewIndex < ProjectionStencilingPasses.Num() && ProjectionStencilingPasses[ViewIndex] != nullptr)
 	{
 		// GPUCULL_TODO: get rid of const cast
 		FSimpleMeshDrawCommandPass& ProjectionStencilingPass = *const_cast<FSimpleMeshDrawCommandPass*>(ProjectionStencilingPasses[ViewIndex]);
 		ProjectionStencilingPass.BuildRenderingCommands(GraphBuilder, *View, *SceneRender->Scene, PassParameters->InstanceCullingDrawParams);
 	}
-	else if (bIsInstancedStereoEmulated && (bPreShadow || bSelfShadowOnly))
+	else if (bIsInstancedStereoBypassed && (bPreShadow || bSelfShadowOnly))
 	{
 		// NOTE: This here is a hack that must match up to the use inside SetupProjectionStencilMask, where we use the Stereo setup but draw each eye independently
 		//       by scissoring the undersired half (while setting the full viewport to get the scaling to match the stereo pathfor base/pre-pass 1:1).
@@ -2037,7 +2037,8 @@ void FSceneRenderer::RenderShadowProjections(
 	for (int32 ViewIndex = 0; ViewIndex < Views.Num(); ViewIndex++)
 	{
 		const FViewInfo& View = Views[ViewIndex];
-		if (!View.ShouldRenderView() || (bSubPixelShadow && !HairStrands::HasViewHairStrandsData(View)))
+		const bool bIsInstancedStereoBypassed = View.bIsInstancedStereoEnabled && !View.bIsMobileMultiViewEnabled && IStereoRendering::IsStereoEyeView(View);
+		if ((!bIsInstancedStereoBypassed && !View.ShouldRenderView()) || (bSubPixelShadow && !HairStrands::HasViewHairStrandsData(View)))
 		{
 			continue;
 		}
