Fix Horde compilation error due to "CVarNaniteSoftwareVRS" being defined elsewhere.

diff --git a/Engine/Source/Runtime/Renderer/Private/VariableRateShading/VariableRateShadingImageManager.cpp b/Engine/Source/Runtime/Renderer/Private/VariableRateShading/VariableRateShadingImageManager.cpp
--- a/Engine/Source/Runtime/Renderer/Private/VariableRateShading/VariableRateShadingImageManager.cpp
+++ b/Engine/Source/Runtime/Renderer/Private/VariableRateShading/VariableRateShadingImageManager.cpp
@@ -521,7 +521,7 @@ FRDGTextureRef FVariableRateShadingImageManager::GetVariableRateShadingImage(FRD
 void FVariableRateShadingImageManager::PrepareImageBasedVRS(FRDGBuilder& GraphBuilder, const FSceneViewFamily& ViewFamily, const FMinimalSceneTextures& SceneTextures)
 {
 	EShaderPlatform ShaderPlatform = ViewFamily.Scene->GetShaderPlatform();
-	static const auto CVarNaniteSoftwareVRS = IConsoleManager::Get().FindConsoleVariable(TEXT("r.Nanite.SoftwareVRS"));
+	static const auto CVarLocalNaniteSoftwareVRS = IConsoleManager::Get().FindConsoleVariable(TEXT("r.Nanite.SoftwareVRS")); // "CVarNaniteSoftwareVRS" would shadow the static declaration in NaniteShading.cpp
 
 	bHardwareVRSEnabledForFrame = 
 		IsAttachmentVRSEnabled() &&
@@ -529,7 +529,7 @@ void FVariableRateShadingImageManager::PrepareImageBasedVRS(FRDGBuilder& GraphBu
 
 	bSoftwareVRSEnabledForFrame = 
 		CVarEnableVRSSoftwareImage.GetValueOnRenderThread() > 0 &&
-		CVarNaniteSoftwareVRS->GetInt() > 0 &&
+		CVarLocalNaniteSoftwareVRS->GetInt() > 0 &&
 		IsFeatureLevelSupported(ShaderPlatform, ERHIFeatureLevel::SM6);
 
 	if (!IsVRSEnabledForFrame())
