DrawMaterialToRenderTarget blueprint function causes flickering issues in other views due to resetting scene texture extent.  The need to reset can be avoided by preventing the function from modifying scene texture extent in the first place.  The code path in question doesn't allocate scene textures -- it renders a pass directly to the output render target -- so there's no need to update the cached scene texture extent.  An option was added to InitializeSceneTexturesConfig to override the extent, ignoring the cached extent state.

#rnx

diff --git a/Engine/Source/Runtime/Engine/Private/KismetRenderingLibrary.cpp b/Engine/Source/Runtime/Engine/Private/KismetRenderingLibrary.cpp
--- a/Engine/Source/Runtime/Engine/Private/KismetRenderingLibrary.cpp
+++ b/Engine/Source/Runtime/Engine/Private/KismetRenderingLibrary.cpp
@@ -223,12 +223,6 @@ void UKismetRenderingLibrary::DrawMaterialToRenderTarget(UObject* WorldContextOb
 
 			//UpdateResourceImmediate must be called here to ensure mips are generated.
 			TextureRenderTarget->UpdateResourceImmediate(false);
-
-			ENQUEUE_RENDER_COMMAND(ResetSceneTextureExtentHistory)(
-				[RenderTargetResource](FRHICommandListImmediate& RHICmdList)
-				{
-					RenderTargetResource->ResetSceneTextureExtentsHistory();
-				});
 		}
 	}
 }
diff --git a/Engine/Source/Runtime/Renderer/Internal/SceneTextures.h b/Engine/Source/Runtime/Renderer/Internal/SceneTextures.h
--- a/Engine/Source/Runtime/Renderer/Internal/SceneTextures.h
+++ b/Engine/Source/Runtime/Renderer/Internal/SceneTextures.h
@@ -14,7 +14,7 @@ class FViewInfo;
 class FViewFamilyInfo;
 
 /** Initializes a scene textures config instance from the view family. */
-extern RENDERER_API void InitializeSceneTexturesConfig(FSceneTexturesConfig& Config, const FSceneViewFamily& ViewFamily);
+extern RENDERER_API void InitializeSceneTexturesConfig(FSceneTexturesConfig& Config, const FSceneViewFamily& ViewFamily, FIntPoint ExtentOverride = FIntPoint(0,0));
 
 struct FTransientUserSceneTexture
 {
diff --git a/Engine/Source/Runtime/Renderer/Private/Renderer.cpp b/Engine/Source/Runtime/Renderer/Private/Renderer.cpp
--- a/Engine/Source/Runtime/Renderer/Private/Renderer.cpp
+++ b/Engine/Source/Runtime/Renderer/Private/Renderer.cpp
@@ -337,10 +337,10 @@ void FRendererModule::DrawTileMesh(FCanvasRenderContext& RenderContext, FMeshPas
 			View.MaterialMaxEmissiveValue = GetEmissiveMaxValueForPixelFormat(RenderContext.GetRenderTarget()->Desc.Format);
 		}
 
-		// Default init of SceneTexturesConfig will take extents from FSceneTextureExtentState.
-		// We want the view extents, so explicitly set that.
-		InitializeSceneTexturesConfig(ViewFamily->SceneTexturesConfig, *ViewFamily);
-		ViewFamily->SceneTexturesConfig.Extent = View.ViewRect.Size();
+		// Default init of SceneTexturesConfig will take extents from FSceneTextureExtentState.  We want the view extents, so explicitly
+		// set that.  This will bypass scene texture extent caching logic, but this code path doesn't allocate scene textures (it renders
+		// directly to RenderContext.GetRenderTarget()), so caching is irrelevant for purposes of avoiding render target pool thrashing.
+		InitializeSceneTexturesConfig(ViewFamily->SceneTexturesConfig, *ViewFamily, View.ViewRect.Size());
 
 		const auto FeatureLevel = View.GetFeatureLevel();
 		const EShadingPath ShadingPath = GetFeatureLevelShadingPath(FeatureLevel);
diff --git a/Engine/Source/Runtime/Renderer/Private/SceneTextures.cpp b/Engine/Source/Runtime/Renderer/Private/SceneTextures.cpp
--- a/Engine/Source/Runtime/Renderer/Private/SceneTextures.cpp
+++ b/Engine/Source/Runtime/Renderer/Private/SceneTextures.cpp
@@ -372,9 +372,23 @@ void ResetSceneTextureExtentHistory()
 
 ENUM_CLASS_FLAGS(FSceneTextureExtentState::ERenderTargetHistory);
 
-void InitializeSceneTexturesConfig(FSceneTexturesConfig& Config, const FSceneViewFamily& ViewFamily)
+void InitializeSceneTexturesConfig(FSceneTexturesConfig& Config, const FSceneViewFamily& ViewFamily, FIntPoint ExtentOverride)
 {
-	FIntPoint Extent = FSceneTextureExtentState::Get().Compute(ViewFamily);
+	FIntPoint Extent;
+	if (ExtentOverride.X > 0)
+	{
+#if DO_CHECK
+		for (const FSceneView* View : ViewFamily.Views)
+		{
+			check(View->UnscaledViewRect.Max.X <= ExtentOverride.X && View->UnscaledViewRect.Max.Y <= ExtentOverride.Y);
+		}
+#endif
+		Extent = ExtentOverride;
+	}
+	else
+	{
+		Extent = FSceneTextureExtentState::Get().Compute(ViewFamily);
+	}
 	EShadingPath ShadingPath = GetFeatureLevelShadingPath(ViewFamily.GetFeatureLevel());
 
 	bool bRequiresAlphaChannel = ShadingPath == EShadingPath::Mobile ? IsMobilePropagateAlphaEnabled(ViewFamily.GetShaderPlatform()) : IsPostProcessingWithAlphaChannelSupported();
