Remove 3 sources of hitching from virtual textures when running with -game.
Each had a fast path for package, and slower, more conservative path for editor.
Now check both WITH_EDITOR and GIsEditor for the editor path, since GIsEditor will be false when using -game.
First hitch is from enabling the virtual texture pool auto grow.
Second hitch is from not supporting throttled release of allocated virtual textures.
Third hitch is from overriding stalling on render thread for virtual texture root page streaming.

diff --git a/Engine/Source/Runtime/Engine/Private/VT/VirtualTexturePoolConfig.cpp b/Engine/Source/Runtime/Engine/Private/VT/VirtualTexturePoolConfig.cpp
--- a/Engine/Source/Runtime/Engine/Private/VT/VirtualTexturePoolConfig.cpp
+++ b/Engine/Source/Runtime/Engine/Private/VT/VirtualTexturePoolConfig.cpp
@@ -284,10 +284,13 @@ float VirtualTexturePool::GetPoolSizeScale()
 bool VirtualTexturePool::GetPoolAutoGrow()
 {
 #if WITH_EDITOR
-	UVirtualTexturePoolConfig const* PoolConfig = GetDefault<UVirtualTexturePoolConfig>();
-	if (PoolConfig->bPoolAutoGrowInEditor)
+	if (GIsEditor)
 	{
-		return true;
+		UVirtualTexturePoolConfig const* PoolConfig = GetDefault<UVirtualTexturePoolConfig>();
+		if (PoolConfig->bPoolAutoGrowInEditor)
+		{
+			return true;
+		}
 	}
 #endif
 	return CVarVTPoolAutoGrow.GetValueOnAnyThread();
diff --git a/Engine/Source/Runtime/Engine/Private/VT/VirtualTextureScalability.cpp b/Engine/Source/Runtime/Engine/Private/VT/VirtualTextureScalability.cpp
--- a/Engine/Source/Runtime/Engine/Private/VT/VirtualTextureScalability.cpp
+++ b/Engine/Source/Runtime/Engine/Private/VT/VirtualTextureScalability.cpp
@@ -205,10 +205,12 @@ namespace VirtualTextureScalability
 	int32 GetMaxAllocatedVTReleasedPerFrame()
 	{
 #if WITH_EDITOR
-		return 0;
-#else
-		return CVarVTMaxReleasedPerFrame.GetValueOnAnyThread();
+		if (GIsEditor)
+		{
+			return 0;
+		}
 #endif
+		return CVarVTMaxReleasedPerFrame.GetValueOnAnyThread();
 	}
 
 	uint32 GetPageFreeThreshold()
diff --git a/Engine/Source/Runtime/Renderer/Private/VT/VirtualTextureSystem.cpp b/Engine/Source/Runtime/Renderer/Private/VT/VirtualTextureSystem.cpp
--- a/Engine/Source/Runtime/Renderer/Private/VT/VirtualTextureSystem.cpp
+++ b/Engine/Source/Runtime/Renderer/Private/VT/VirtualTextureSystem.cpp
@@ -180,6 +180,22 @@ FVirtualTextureUpdateSettings::FVirtualTextureUpdateSettings()
 	MaxContinuousUpdates = VirtualTextureScalability::GetMaxContinuousUpdatesPerFrame();
 }
 
+static bool ShouldSyncProduceLockedTiles()
+{
+#if WITH_EDITOR
+	if (GIsEditor)
+	{
+		// Always force loading of locked root pages in editor to be sure that any build steps give best quality.
+		return true;
+	}
+#endif
+	// If we return false then we may render with a VT before the root pages are mapped.
+	// When that happens the shader switches to using the single color fallback value instead of sampling the VT.
+	// When the root page is finally mapped we will return to normal high quality VT sampling.
+	// todo[VT]: Make root pages always resident so that we never need to load sync the root pages.
+	return CVarVTSyncProduceLockedTiles.GetValueOnRenderThread() != 0;
+}
+
 
 static FORCEINLINE uint32 EncodePage(uint32 ID, uint32 vLevel, uint32 vTileX, uint32 vTileY)
 {
@@ -2149,19 +2165,9 @@ void FVirtualTextureSystem::SubmitRequests(FRHICommandList& RHICmdList, ERHIFeat
 		TArray<FProducePageDataPrepareTask> PrepareTasks;
 		PrepareTasks.Reserve(RequestList->GetNumLoadRequests());
 
-#if WITH_EDITOR
-		// Always force loading of locked root pages in editor to be sure that any build steps give best quality.
-		const bool bSyncProduceLockedTiles = true;
-#else
-		// if bSyncProduceLockedTiles is false then we may render with a VT before the root pages are mapped.
-		// When that happens the shader switches to using the single color fallback value instead of sampling the VT.
-		// When the root page is finally mapped we will return to normal high quality VT sampling.
-		// todo[VT]: Make root pages always resident so that we never need to load sync the root pages.
-		const bool bSyncProduceLockedTiles = CVarVTSyncProduceLockedTiles.GetValueOnRenderThread() != 0;
-#endif
-
 		bool bWaitForProducers = false;
 
+		const bool bSyncProduceLockedTiles = ShouldSyncProduceLockedTiles();
 		const uint32 MaxPagesProduced = Settings.MaxPagesProduced;
 		const uint32 PageFreeThreshold = VirtualTextureScalability::GetPageFreeThreshold();
 		uint32 NumStacksProduced = 0u;
