Make sure to use the max RTPSO local binding data size from the RTPSO which will be used this frame  - if requested RTPSO is non blocking compiling then the RTPSO from the previous could be used with a different max local binding data size

diff --git a/Engine/Source/Runtime/D3D12RHI/Private/D3D12RayTracing.cpp b/Engine/Source/Runtime/D3D12RHI/Private/D3D12RayTracing.cpp
--- a/Engine/Source/Runtime/D3D12RHI/Private/D3D12RayTracing.cpp
+++ b/Engine/Source/Runtime/D3D12RHI/Private/D3D12RayTracing.cpp
@@ -2121,6 +2121,8 @@ public:
 			CallableShaders.Shaders.Add(Shader);
 		}
 
+		check(Initializer.GetMaxLocalBindingDataSize() >= MaxLocalRootSignatureSize);
+
 		// Wait for all compilation tasks to be complete and then gather the compiled collection descriptors
 
 		CompileTime -= FPlatformTime::Cycles64();
diff --git a/Engine/Source/Runtime/RHI/Private/PipelineStateCache.cpp b/Engine/Source/Runtime/RHI/Private/PipelineStateCache.cpp
--- a/Engine/Source/Runtime/RHI/Private/PipelineStateCache.cpp
+++ b/Engine/Source/Runtime/RHI/Private/PipelineStateCache.cpp
@@ -1118,6 +1118,8 @@ public:
 
 	FRayTracingPipelineStateRHIRef RHIPipeline;
 
+	uint32 MaxLocalBindingSize = 0;
+
 	uint64 HitsAcrossFrames = 0;
 	uint64 LastFrameHit = 0;
 
@@ -1141,6 +1143,15 @@ FRHIRayTracingPipelineState* GetRHIRayTracingPipelineState(FRayTracingPipelineSt
 	return nullptr;
 }
 
+uint32 GetRHIRayTracingPipelineStateMaxLocalBindingDataSize(FRayTracingPipelineState* PipelineState)
+{
+	if (PipelineState)
+	{
+		return PipelineState->MaxLocalBindingSize;
+	}
+	return 0;
+}
+
 int32 FindRayTracingHitGroupIndex(FRayTracingPipelineState* Pipeline, FRHIRayTracingShader* HitGroupShader, bool bRequired)
 {
 #if RHI_RAYTRACING
@@ -2633,6 +2644,7 @@ public:
 	FRayTracingPipelineState* Add(const FRayTracingPipelineStateInitializer& Initializer)
 	{
 		FRayTracingPipelineState* Result = new FRayTracingPipelineState(Initializer);
+		Result->MaxLocalBindingSize = Initializer.GetMaxLocalBindingDataSize();
 
 		FScopeLock ScopeLock(&CriticalSection);
 
diff --git a/Engine/Source/Runtime/RHI/Public/RHICommandList.h b/Engine/Source/Runtime/RHI/Public/RHICommandList.h
--- a/Engine/Source/Runtime/RHI/Public/RHICommandList.h
+++ b/Engine/Source/Runtime/RHI/Public/RHICommandList.h
@@ -2600,6 +2600,7 @@ extern RHI_API FGraphicsPipelineState*		GetGraphicsPipelineState(FRHICommandList
 extern RHI_API FRHIComputePipelineState*	GetRHIComputePipelineState(FComputePipelineState*);
 extern RHI_API FRHIWorkGraphPipelineState*	GetRHIWorkGraphPipelineState(FWorkGraphPipelineState*);
 extern RHI_API FRHIRayTracingPipelineState*	GetRHIRayTracingPipelineState(FRayTracingPipelineState*);
+extern RHI_API uint32					    GetRHIRayTracingPipelineStateMaxLocalBindingDataSize(FRayTracingPipelineState*);
 
 class FRHIComputeCommandList : public FRHICommandListBase
 {
@@ -4183,6 +4184,8 @@ public:
 		}
 		else
 		{
+			check(GetRHIRayTracingPipelineStateMaxLocalBindingDataSize(Pipeline) <= SBT->GetInitializer().LocalBindingDataSize);
+
 			FRayTracingLocalShaderBindings* InlineBindings = nullptr;
 
 			// By default all batch binding data is stored in the command list memory.
diff --git a/Engine/Source/Runtime/Renderer/Private/RayTracing/RayTracingMaterialHitShaders.cpp b/Engine/Source/Runtime/Renderer/Private/RayTracing/RayTracingMaterialHitShaders.cpp
--- a/Engine/Source/Runtime/Renderer/Private/RayTracing/RayTracingMaterialHitShaders.cpp
+++ b/Engine/Source/Runtime/Renderer/Private/RayTracing/RayTracingMaterialHitShaders.cpp
@@ -757,8 +757,6 @@ void FDeferredShadingSceneRenderer::CreateRayTracingMaterialPipeline(
 		PipelineCacheFlags |= ERayTracingPipelineCacheFlags::NonBlocking;
 	}
 
-	OutMaxLocalBindingDataSize = Initializer.GetMaxLocalBindingDataSize();
-
 	FRayTracingPipelineState* PipelineState = PipelineStateCache::GetAndOrCreateRayTracingPipelineState(RHICmdList, Initializer, PipelineCacheFlags);
 
 	if (PipelineState)
@@ -776,6 +774,10 @@ void FDeferredShadingSceneRenderer::CreateRayTracingMaterialPipeline(
 		PipelineState = FallbackPipelineState;
 	}
 
+	// Retrieve the binding data size from the actual used RTPSO because the requested RTPSO could still be non blocking async compiling
+	// and then we are using the RTPSO from the previous frame
+	OutMaxLocalBindingDataSize = GetRHIRayTracingPipelineStateMaxLocalBindingDataSize(PipelineState);
+
 	if (FallbackPipelineState != nullptr && PipelineState != FallbackPipelineState && bIsPathTracing && !View.bIsOfflineRender)
 	{
 		// When using path tracing, a change in pipeline state compared to the previous frame means some new materials got added to the RTPSO
