Substrate - fixed missing blending mode when using substrate material on non substrate project.

#rnx

#tests none

diff --git a/Engine/Source/Runtime/Renderer/Private/MobileBasePass.cpp b/Engine/Source/Runtime/Renderer/Private/MobileBasePass.cpp
--- a/Engine/Source/Runtime/Renderer/Private/MobileBasePass.cpp
+++ b/Engine/Source/Runtime/Renderer/Private/MobileBasePass.cpp
@@ -552,6 +552,7 @@ void MobileBasePass::SetTranslucentRenderState(FMeshPassProcessorRenderState& Dr
 		switch (Material.GetBlendMode())
 		{
 		case BLEND_Translucent:
+		case BLEND_TranslucentColoredTransmittance:	// When Substrate is disabled, this falls back to simple Translucency.
 			if (Material.ShouldWriteOnlyAlpha())
 			{
 				DrawRenderState.SetBlendState(TStaticBlendState<CW_ALPHA, BO_Add, BF_Zero, BF_Zero, BO_Add, BF_One, BF_Zero,
