Optimization to skip/simplify VSM light grid pruning when there are no VSM local lights.
Will generalize this to a range in the future.

diff --git a/Engine/Shaders/Private/VirtualShadowMaps/VirtualShadowMapPageMarking.usf b/Engine/Shaders/Private/VirtualShadowMaps/VirtualShadowMapPageMarking.usf
--- a/Engine/Shaders/Private/VirtualShadowMaps/VirtualShadowMapPageMarking.usf
+++ b/Engine/Shaders/Private/VirtualShadowMaps/VirtualShadowMapPageMarking.usf
@@ -125,7 +125,8 @@ void MarkPageClipmap(
 	}
 }
 
-// 
+uint MinLocalLightIndex;
+uint MaxLocalLightIndex;
 RWStructuredBuffer<uint> OutPrunedLightGridData;
 RWStructuredBuffer<uint> OutPrunedNumCulledLightsGrid;
 
@@ -145,9 +146,11 @@ void PruneLightGridCS(uint GridLinearIndex : SV_DispatchThreadID)
 	uint FirstLightWithVSM = CulledLightGridHeader.NumLights;
 	uint LastLightWithVSM = 0U;
 
+	uint MaxIndex = min(MaxLocalLightIndex, CulledLightGridHeader.NumLights);
+
 	// First pass add the non-distant lights with VSMs
 	LOOP
-	for (uint Index = 0; Index < CulledLightGridHeader.NumLights; ++Index)
+	for (uint Index = MinLocalLightIndex; Index < MaxIndex; ++Index)
 	{
 		int VirtualShadowMapId = GetLocalLightData(CulledLightGridHeader.DataStartIndex + Index, EyeIndex).VirtualShadowMapId;
 		if (VirtualShadowMapId != INDEX_NONE)
diff --git a/Engine/Source/Runtime/Renderer/Private/DeferredShadingRenderer.cpp b/Engine/Source/Runtime/Renderer/Private/DeferredShadingRenderer.cpp
--- a/Engine/Source/Runtime/Renderer/Private/DeferredShadingRenderer.cpp
+++ b/Engine/Source/Runtime/Renderer/Private/DeferredShadingRenderer.cpp
@@ -2665,7 +2665,16 @@ void FDeferredShadingSceneRenderer::Render(FRDGBuilder& GraphBuilder)
 
 					FFrontLayerTranslucencyData FrontLayerTranslucencyData = RenderFrontLayerTranslucency(GraphBuilder, Views, SceneTextures, true /*VSM page marking*/);
 
-					VirtualShadowMapArray.BuildPageAllocations(GraphBuilder, GetActiveSceneTextures(), Views, *GatherAndSortLightsTask.GetResult(), VisibleLightInfos, SingleLayerWaterPrePassResult, FrontLayerTranslucencyData, FroxelRenderer);
+					VirtualShadowMapArray.BuildPageAllocations(GraphBuilder,
+						GetActiveSceneTextures(),
+						Views,
+						*GatherAndSortLightsTask.GetResult(),
+						VisibleLightInfos,
+						SingleLayerWaterPrePassResult,
+						FrontLayerTranslucencyData,
+						FroxelRenderer,
+						ShadowSceneRenderer->AreAnyLocalLightsPreset()
+					);
 				}
 
 				RenderShadowDepthMaps(GraphBuilder, InitViewTaskDatas.DynamicShadows, InstanceCullingManager, ExternalAccessQueue);
diff --git a/Engine/Source/Runtime/Renderer/Private/Shadows/ShadowSceneRenderer.h b/Engine/Source/Runtime/Renderer/Private/Shadows/ShadowSceneRenderer.h
--- a/Engine/Source/Runtime/Renderer/Private/Shadows/ShadowSceneRenderer.h
+++ b/Engine/Source/Runtime/Renderer/Private/Shadows/ShadowSceneRenderer.h
@@ -106,6 +106,11 @@ public:
 		return bNeedMegaLightsProjection;
 	}
 
+	bool AreAnyLocalLightsPreset() const
+	{
+		return LocalLights.Num() > 0;
+	}
+
 	UE::Renderer::Private::IShadowInvalidatingInstances *GetInvalidatingInstancesInterface(const FSceneView *SceneView);
 
 private:
diff --git a/Engine/Source/Runtime/Renderer/Private/VirtualShadowMaps/VirtualShadowMapArray.cpp b/Engine/Source/Runtime/Renderer/Private/VirtualShadowMaps/VirtualShadowMapArray.cpp
--- a/Engine/Source/Runtime/Renderer/Private/VirtualShadowMaps/VirtualShadowMapArray.cpp
+++ b/Engine/Source/Runtime/Renderer/Private/VirtualShadowMaps/VirtualShadowMapArray.cpp
@@ -883,6 +883,8 @@ class FPruneLightGridCS : public FVirtualShadowMapPageManagementShader
 		SHADER_PARAMETER_RDG_UNIFORM_BUFFER(FVirtualShadowMapUniformParameters, VirtualShadowMap)
 		SHADER_PARAMETER_STRUCT_REF(FViewUniformShaderParameters, View)
 		SHADER_PARAMETER_RDG_UNIFORM_BUFFER(FForwardLightData, ForwardLightData)
+		SHADER_PARAMETER(uint32, MinLocalLightIndex)
+		SHADER_PARAMETER(uint32, MaxLocalLightIndex)
 		SHADER_PARAMETER_RDG_BUFFER_UAV(RWStructuredBuffer<uint>, OutPrunedLightGridData)
 		SHADER_PARAMETER_RDG_BUFFER_UAV(RWStructuredBuffer<uint>, OutPrunedNumCulledLightsGrid)
 	END_SHADER_PARAMETER_STRUCT()
@@ -1736,7 +1738,8 @@ void FVirtualShadowMapArray::BuildPageAllocations(
 	const TConstArrayView<FVisibleLightInfo>& VisibleLightInfos,
 	const FSingleLayerWaterPrePassResult* SingleLayerWaterPrePassResult,
 	const FFrontLayerTranslucencyData& FrontLayerTranslucencyData,
-	const Froxel::FRenderer& FroxelRenderer)
+	const Froxel::FRenderer& FroxelRenderer,
+	bool bAnyLocalLightsWithVSMs)
 {
 	check(IsEnabled());
 
@@ -1970,15 +1973,29 @@ void FVirtualShadowMapArray::BuildPageAllocations(
 				FRDGBufferRef PrunedNumCulledLightsGridRDG = GraphBuilder.CreateBuffer(FRDGBufferDesc::CreateStructuredDesc(sizeof(uint32), NumLightGridCells), TEXT("Shadow.Virtual.NumCulledLightsGrid"));
 
 				{
+					// TODO: Make this a more dynamic bound rather than just this special case
+					uint32 MinLocalLightIndex = 0;
+					uint32 MaxLocalLightIndex = View.ForwardLightingResources.ForwardLightData->NumLocalLights;
+					if (!bAnyLocalLightsWithVSMs)
+					{
+						MaxLocalLightIndex = 0;
+					}
+
 					FPruneLightGridCS::FParameters* PassParameters = GraphBuilder.AllocParameters< FPruneLightGridCS::FParameters >();
 					PassParameters->VirtualShadowMap = GetUncachedUniformBuffer(GraphBuilder);
 					PassParameters->View = View.ViewUniformBuffer;
 					PassParameters->ForwardLightData = View.ForwardLightingResources.ForwardLightUniformBuffer;
+					PassParameters->MinLocalLightIndex = MinLocalLightIndex;
+					PassParameters->MaxLocalLightIndex = MaxLocalLightIndex;
 					PassParameters->OutPrunedLightGridData = GraphBuilder.CreateUAV(PrunedLightGridDataRDG);
 					PassParameters->OutPrunedNumCulledLightsGrid = GraphBuilder.CreateUAV(PrunedNumCulledLightsGridRDG);
 					auto ComputeShader = View.ShaderMap->GetShader<FPruneLightGridCS>();
 
-					FComputeShaderUtils::AddPass(GraphBuilder, RDG_EVENT_NAME("PruneLightGrid"), ComputeShader, PassParameters, FComputeShaderUtils::GetGroupCount(NumLightGridCells, FPruneLightGridCS::DefaultCSGroupX));
+					FComputeShaderUtils::AddPass(GraphBuilder,
+						RDG_EVENT_NAME("PruneLightGrid(Min=%d,Max=%d)", MinLocalLightIndex, MaxLocalLightIndex),
+						ComputeShader,
+						PassParameters,
+						FComputeShaderUtils::GetGroupCount(NumLightGridCells, FPruneLightGridCS::DefaultCSGroupX));
 				};
 
 				PerViewData.LightGridData = GraphBuilder.CreateSRV(PrunedLightGridDataRDG);
diff --git a/Engine/Source/Runtime/Renderer/Private/VirtualShadowMaps/VirtualShadowMapArray.h b/Engine/Source/Runtime/Renderer/Private/VirtualShadowMaps/VirtualShadowMapArray.h
--- a/Engine/Source/Runtime/Renderer/Private/VirtualShadowMaps/VirtualShadowMapArray.h
+++ b/Engine/Source/Runtime/Renderer/Private/VirtualShadowMaps/VirtualShadowMapArray.h
@@ -325,7 +325,8 @@ public:
 		const TConstArrayView<FVisibleLightInfo>& VisibleLightInfos,
 		const FSingleLayerWaterPrePassResult* SingleLayerWaterPrePassResult,
 		const FFrontLayerTranslucencyData& FrontLayerTranslucencyData,
-		const Froxel::FRenderer& FroxelRenderer);
+		const Froxel::FRenderer& FroxelRenderer,
+		bool bAnyLocalLightsWithVSMs);
 
 	bool IsAllocated() const
 	{
