Fix Groom holdout with fog.

The groom's holdout pass now runs in two steps:
* Attenuate the scene pixel alpha value by the hair coverage
* Add hair's fog attenuation to the alpha pixel value

diff --git a/Engine/Shaders/Private/HairStrands/HairStrandsComposition.usf b/Engine/Shaders/Private/HairStrands/HairStrandsComposition.usf
--- a/Engine/Shaders/Private/HairStrands/HairStrandsComposition.usf
+++ b/Engine/Shaders/Private/HairStrands/HairStrandsComposition.usf
@@ -24,18 +24,11 @@ void FastResolvePS(in FScreenVertexOutput Input)
 
 ////////////////////////////////////////////////////////////////////////////////////////////////
 
-#if SHADER_COMPOSE_SAMPLE
+#if SHADER_COMPOSE_SAMPLE || SHADER_HOLDOUT
 
 #include "../HeightFogCommon.ush"
 #include "../SkyAtmosphereCommon.ush"
 
-int2 OutputResolution;
-uint bComposeDofDepth;
-uint bHasHoldout;
-
-Texture2D<float>	HairDOFDepthTexture;
-Texture2D<float4>	HairLightingSampleBuffer;
-
 float4 EvaluateVolumetric(float3 WorldPosition, float SceneDepth)
 {
 	const float3 CameraRelative_WorldPosition = WorldPosition - DFHackToFloat(PrimaryView.WorldCameraOrigin);
@@ -74,6 +67,18 @@ float4 EvaluateVolumetric(float3 WorldPosition, float SceneDepth)
 
 	return Fogging;
 }
+#endif // SHADER_COMPOSE_SAMPLE || SHADER_HOLDOUT
+
+////////////////////////////////////////////////////////////////////////////////////////////////
+
+#if SHADER_COMPOSE_SAMPLE
+
+int2 OutputResolution;
+uint bComposeDofDepth;
+uint bHasHoldout;
+
+Texture2D<float>	HairDOFDepthTexture;
+Texture2D<float4>	HairLightingSampleBuffer;
 
 void ComposeSamplePS(
 	in FScreenVertexOutput Input,
@@ -121,10 +126,14 @@ void ComposeSamplePS(
 		const uint2 Coord = GetHairSampleCoord(LinearCoord, Resolution);
 		LocalAccColor += HairLightingSampleBuffer.Load(uint3(Coord, 0)).xyz;
 	}
-	OutColor.rgb = LocalAccColor * Volumetric.a + Volumetric.rgb;
 
-	// If any sample support holdout, then the lighting contribution is set to 0, to mask out the current hair pixel
-	OutColor.rgb *= bHoldout ? 0 : PixelCoverage;
+	// If any sample support holdout, then the lighting contribution is set to 0 to mask out the current hair pixel 
+	// but fog from camera to hair pixel needs to be accounted for
+	if (bHoldout)
+	{
+		LocalAccColor = 0;
+	}
+	OutColor.rgb = (LocalAccColor * Volumetric.a + Volumetric.rgb) * PixelCoverage;
 	OutColor.a = PixelCoverage;
 
 	OutDepth = HairDeviceZ;
@@ -173,6 +182,10 @@ void ComposeSamplePS(
 
 #if SHADER_HOLDOUT
 
+// PassType:
+// * 0: Output attenuation for the scene pixel (alpha only)
+// * 1: Output hair contribution to the scene pixel (alpha value only)
+uint PassType;
 uint bComposeDofDepth;
 Texture2D<float> HairDOFDepthTexture;
 
@@ -208,13 +221,18 @@ void HoldoutPS(
 		bHoldout = bHoldout || HasHairFlags(Sample.Flags, HAIR_FLAGS_HOLDOUT);
 	}
 
+	const float3 ClosetPointWorldPosition = DFHackToFloat(SvPositionToWorld(float4(Input.Position.xy, HairDeviceZ, 1.0)));
+	const float4 Volumetric = EvaluateVolumetric(ClosetPointWorldPosition, ConvertFromDeviceZ(HairDeviceZ));
+
 	if (!bHoldout)
 	{
 		discard;
 	}
 
+	//OutColor.rgb = (LocalAccColor * Volumetric.a + Volumetric.rgb) * PixelCoverage;
+
 	OutColor.rgb = 0;
-	OutColor.a = PixelCoverage;
+	OutColor.a   = PassType == 0 ? PixelCoverage : Volumetric.a * PixelCoverage;
 	OutDepth     = HairDeviceZ;
 	if (bComposeDofDepth)
 	{
diff --git a/Engine/Source/Runtime/Renderer/Private/HairStrands/HairStrandsComposition.cpp b/Engine/Source/Runtime/Renderer/Private/HairStrands/HairStrandsComposition.cpp
--- a/Engine/Source/Runtime/Renderer/Private/HairStrands/HairStrandsComposition.cpp
+++ b/Engine/Source/Runtime/Renderer/Private/HairStrands/HairStrandsComposition.cpp
@@ -45,7 +45,8 @@ enum EHairStrandsCommonPassType
 	TAAFastResolve,
 	GBuffer,
 	Blit,
-	Holdout
+	HoldoutScene,
+	HoldoutHair,
 };
 
 template<typename TPassParameter, typename TPixelShader>
@@ -114,9 +115,13 @@ void InternalCommonDrawPass(
 		{
 			GraphicsPSOInit.BlendState = TStaticBlendState<CW_RGBA, BO_Add, BF_One, BF_Zero, BO_Add, BF_One, BF_Zero>::GetRHI();
 		}
-		else if (Type == EHairStrandsCommonPassType::Holdout)
+		else if (Type == EHairStrandsCommonPassType::HoldoutScene)
 		{
-			GraphicsPSOInit.BlendState = TStaticBlendState<CW_ALPHA, BO_Add, BF_Zero, BF_One, BO_Add, BF_One, BF_InverseSourceAlpha>::GetRHI();
+			GraphicsPSOInit.BlendState = TStaticBlendState<CW_ALPHA, BO_Add, BF_Zero, BF_One, BO_Add, BF_Zero, BF_InverseSourceAlpha>::GetRHI();
+		}
+		else if (Type == EHairStrandsCommonPassType::HoldoutHair)
+		{
+			GraphicsPSOInit.BlendState = TStaticBlendState<CW_ALPHA, BO_Add, BF_Zero, BF_One, BO_Add, BF_One, BF_One>::GetRHI();
 		}
 		else
 		{
@@ -125,7 +130,7 @@ void InternalCommonDrawPass(
 
 		GraphicsPSOInit.RasterizerState = TStaticRasterizerState<>::GetRHI();
 
-		if (Type == EHairStrandsCommonPassType::Composition || Type == EHairStrandsCommonPassType::Holdout)
+		if (Type == EHairStrandsCommonPassType::Composition || Type == EHairStrandsCommonPassType::HoldoutScene || Type == EHairStrandsCommonPassType::HoldoutHair)
 		{
 			GraphicsPSOInit.DepthStencilState = TStaticDepthStencilState<true, CF_DepthNearOrEqual>::GetRHI();
 		}
@@ -185,6 +190,7 @@ class FHairHoldoutPS : public FGlobalShader
 		SHADER_PARAMETER_STRUCT_INCLUDE(FHairStrandsTilePassVS::FParameters, TileData)
 		SHADER_PARAMETER(FIntPoint, OutputResolution)
 		SHADER_PARAMETER(uint32, bComposeDofDepth)
+		SHADER_PARAMETER(uint32, PassType)
 		SHADER_PARAMETER_RDG_UNIFORM_BUFFER(FHairStrandsViewUniformParameters, HairStrands)
 		SHADER_PARAMETER_RDG_TEXTURE(Texture2D, HairDOFDepthTexture)
 		RENDER_TARGET_BINDING_SLOTS()
@@ -211,28 +217,59 @@ static void AddHairHoldoutPass(
 {
 	const bool bDOFEnable = HairDOFDepthTexture != nullptr ? 1 : 0;
 
-	FHairHoldoutPS::FParameters* Parameters = GraphBuilder.AllocParameters<FHairHoldoutPS::FParameters>();
-	Parameters->bComposeDofDepth = bDOFEnable ? 1 : 0;
-	Parameters->HairDOFDepthTexture = bDOFEnable ? HairDOFDepthTexture : GSystemTextures.GetBlackDummy(GraphBuilder);
-	Parameters->OutputResolution = OutColorTexture->Desc.Extent;
-	Parameters->ViewUniformBuffer = View.ViewUniformBuffer;
-	Parameters->HairStrands = View.HairStrandsViewData.UniformBuffer;
-	Parameters->RenderTargets[0] = FRenderTargetBinding(OutColorTexture, ERenderTargetLoadAction::ELoad);
-	Parameters->RenderTargets.DepthStencil = FDepthStencilBinding(OutDepthTexture, ERenderTargetLoadAction::ELoad, ERenderTargetLoadAction::ELoad, FExclusiveDepthStencil::DepthWrite_StencilRead);
+	// Attenuate original scene's pixel alpha value
+	{
+		FHairHoldoutPS::FParameters* Parameters = GraphBuilder.AllocParameters<FHairHoldoutPS::FParameters>();
+		Parameters->bComposeDofDepth = bDOFEnable ? 1 : 0;
+		Parameters->PassType = 0;
+		Parameters->HairDOFDepthTexture = bDOFEnable ? HairDOFDepthTexture : GSystemTextures.GetBlackDummy(GraphBuilder);
+		Parameters->OutputResolution = OutColorTexture->Desc.Extent;
+		Parameters->ViewUniformBuffer = View.ViewUniformBuffer;
+		Parameters->HairStrands = View.HairStrandsViewData.UniformBuffer;
+		Parameters->RenderTargets[0] = FRenderTargetBinding(OutColorTexture, ERenderTargetLoadAction::ELoad);
+		Parameters->RenderTargets.DepthStencil = FDepthStencilBinding(OutDepthTexture, ERenderTargetLoadAction::ELoad, ERenderTargetLoadAction::ELoad, FExclusiveDepthStencil::DepthWrite_StencilRead);
+
+		FHairHoldoutPS::FPermutationDomain PermutationVector;
+		TShaderMapRef<FHairHoldoutPS> PixelShader(View.ShaderMap, PermutationVector);
+		InternalCommonDrawPass(
+			GraphBuilder,
+			RDG_EVENT_NAME("HairStrands::Holdout(Scene)"),
+			View,
+			OutColorTexture->Desc.Extent,
+			EHairStrandsCommonPassType::HoldoutScene,
+			false /*bWriteDepth*/,
+			true /*bHasHoldout*/,
+			VisibilityData.TileData,
+			PixelShader,
+			Parameters);
+	}
 
-	FHairHoldoutPS::FPermutationDomain PermutationVector;
-	TShaderMapRef<FHairHoldoutPS> PixelShader(View.ShaderMap, PermutationVector);
-	InternalCommonDrawPass(
-		GraphBuilder,
-		RDG_EVENT_NAME("HairStrands::Holdout"),
-		View,
-		OutColorTexture->Desc.Extent,
-		EHairStrandsCommonPassType::Holdout,
-		false /*bWriteDepth*/,
-		true /*bHasHoldout*/,
-		VisibilityData.TileData,
-		PixelShader,
-		Parameters);
+	// Add hair's fog contribution
+	{
+		FHairHoldoutPS::FParameters* Parameters = GraphBuilder.AllocParameters<FHairHoldoutPS::FParameters>();
+		Parameters->bComposeDofDepth = bDOFEnable ? 1 : 0;
+		Parameters->PassType = 1;
+		Parameters->HairDOFDepthTexture = bDOFEnable ? HairDOFDepthTexture : GSystemTextures.GetBlackDummy(GraphBuilder);
+		Parameters->OutputResolution = OutColorTexture->Desc.Extent;
+		Parameters->ViewUniformBuffer = View.ViewUniformBuffer;
+		Parameters->HairStrands = View.HairStrandsViewData.UniformBuffer;
+		Parameters->RenderTargets[0] = FRenderTargetBinding(OutColorTexture, ERenderTargetLoadAction::ELoad);
+		Parameters->RenderTargets.DepthStencil = FDepthStencilBinding(OutDepthTexture, ERenderTargetLoadAction::ELoad, ERenderTargetLoadAction::ELoad, FExclusiveDepthStencil::DepthWrite_StencilRead);
+
+		FHairHoldoutPS::FPermutationDomain PermutationVector;
+		TShaderMapRef<FHairHoldoutPS> PixelShader(View.ShaderMap, PermutationVector);
+		InternalCommonDrawPass(
+			GraphBuilder,
+			RDG_EVENT_NAME("HairStrands::Holdout(Hair)"),
+			View,
+			OutColorTexture->Desc.Extent,
+			EHairStrandsCommonPassType::HoldoutHair,
+			false /*bWriteDepth*/,
+			true /*bHasHoldout*/,
+			VisibilityData.TileData,
+			PixelShader,
+			Parameters);
+	}
 }
 
 /////////////////////////////////////////////////////////////////////////////////////////
