Fix missing fog binding from CL 36541940.

diff --git a/Engine/Source/Runtime/Renderer/Private/HairStrands/HairStrandsComposition.cpp b/Engine/Source/Runtime/Renderer/Private/HairStrands/HairStrandsComposition.cpp
--- a/Engine/Source/Runtime/Renderer/Private/HairStrands/HairStrandsComposition.cpp
+++ b/Engine/Source/Runtime/Renderer/Private/HairStrands/HairStrandsComposition.cpp
@@ -192,6 +192,7 @@ class FHairHoldoutPS : public FGlobalShader
 		SHADER_PARAMETER(uint32, bComposeDofDepth)
 		SHADER_PARAMETER(uint32, PassType)
 		SHADER_PARAMETER_RDG_UNIFORM_BUFFER(FHairStrandsViewUniformParameters, HairStrands)
+		SHADER_PARAMETER_RDG_UNIFORM_BUFFER(FFogUniformParameters, FogStruct)
 		SHADER_PARAMETER_RDG_TEXTURE(Texture2D, HairDOFDepthTexture)
 		RENDER_TARGET_BINDING_SLOTS()
 	END_SHADER_PARAMETER_STRUCT()
@@ -217,6 +218,8 @@ static void AddHairHoldoutPass(
 {
 	const bool bDOFEnable = HairDOFDepthTexture != nullptr ? 1 : 0;
 
+	TRDGUniformBufferRef<FFogUniformParameters> FogBuffer = CreateFogUniformBuffer(GraphBuilder, View);
+
 	// Attenuate original scene's pixel alpha value
 	{
 		FHairHoldoutPS::FParameters* Parameters = GraphBuilder.AllocParameters<FHairHoldoutPS::FParameters>();
@@ -226,6 +229,7 @@ static void AddHairHoldoutPass(
 		Parameters->OutputResolution = OutColorTexture->Desc.Extent;
 		Parameters->ViewUniformBuffer = View.ViewUniformBuffer;
 		Parameters->HairStrands = View.HairStrandsViewData.UniformBuffer;
+		Parameters->FogStruct = FogBuffer;
 		Parameters->RenderTargets[0] = FRenderTargetBinding(OutColorTexture, ERenderTargetLoadAction::ELoad);
 		Parameters->RenderTargets.DepthStencil = FDepthStencilBinding(OutDepthTexture, ERenderTargetLoadAction::ELoad, ERenderTargetLoadAction::ELoad, FExclusiveDepthStencil::DepthWrite_StencilRead);
 
@@ -253,6 +257,7 @@ static void AddHairHoldoutPass(
 		Parameters->OutputResolution = OutColorTexture->Desc.Extent;
 		Parameters->ViewUniformBuffer = View.ViewUniformBuffer;
 		Parameters->HairStrands = View.HairStrandsViewData.UniformBuffer;
+		Parameters->FogStruct = FogBuffer;
 		Parameters->RenderTargets[0] = FRenderTargetBinding(OutColorTexture, ERenderTargetLoadAction::ELoad);
 		Parameters->RenderTargets.DepthStencil = FDepthStencilBinding(OutDepthTexture, ERenderTargetLoadAction::ELoad, ERenderTargetLoadAction::ELoad, FExclusiveDepthStencil::DepthWrite_StencilRead);
 
