Fixed unbatched non-nanite VSM rendering crash due to missing scene uniform buffer.

diff --git a/Engine/Source/Runtime/Renderer/Private/VirtualShadowMaps/VirtualShadowMapArray.cpp b/Engine/Source/Runtime/Renderer/Private/VirtualShadowMaps/VirtualShadowMapArray.cpp
--- a/Engine/Source/Runtime/Renderer/Private/VirtualShadowMaps/VirtualShadowMapArray.cpp
+++ b/Engine/Source/Runtime/Renderer/Private/VirtualShadowMaps/VirtualShadowMapArray.cpp
@@ -3063,14 +3063,15 @@ FVSMRenderViewCount GetRenderViewCount(const FProjectedShadowInfo* ProjectedShad
 static void AddRasterPass(
 	FRDGBuilder& GraphBuilder, 
 	FRDGEventName&& PassName,
-	const FViewInfo * ShadowDepthView, 
-	const TRDGUniformBufferRef<FShadowDepthPassUniformParameters> &ShadowDepthPassUniformBuffer,
+	const FViewInfo* ShadowDepthView, 
+	const TRDGUniformBufferRef<FShadowDepthPassUniformParameters>& ShadowDepthPassUniformBuffer,
 	FVirtualShadowMapArray& VirtualShadowMapArray,
 	FRDGBufferRef VirtualShadowViewsRDG,
-	const FCullingResult &CullingResult, 
+	const FCullingResult& CullingResult, 
 	FParallelMeshDrawCommandPass& MeshCommandPass,
 	FVirtualShadowDepthPassParameters* PassParameters,
-	TRDGUniformBufferRef<FInstanceCullingGlobalUniforms> InstanceCullingUniformBuffer)
+	TRDGUniformBufferRef<FInstanceCullingGlobalUniforms> InstanceCullingUniformBuffer,
+	TRDGUniformBufferRef<FSceneUniformParameters> SceneUB)
 {
 	PassParameters->View = ShadowDepthView->ViewUniformBuffer;
 	PassParameters->ShadowDepthPass = ShadowDepthPassUniformBuffer;
@@ -3080,6 +3081,7 @@ static void AddRasterPass(
 	PassParameters->InstanceCullingDrawParams.DrawIndirectArgsBuffer = CullingResult.DrawIndirectArgsRDG;
 	PassParameters->InstanceCullingDrawParams.InstanceIdOffsetBuffer = CullingResult.InstanceIdOffsetBufferRDG;
 	PassParameters->InstanceCullingDrawParams.InstanceCulling = InstanceCullingUniformBuffer;
+	PassParameters->InstanceCullingDrawParams.Scene = SceneUB;
 
 	FIntRect ViewRect;
 	ViewRect.Max = FVirtualShadowMap::VirtualMaxResolutionXY;
@@ -3456,6 +3458,8 @@ void FVirtualShadowMapArray::RenderVirtualShadowMapsNonNanite(FRDGBuilder& Graph
 		HZBShaderParameters.HZBTextureArray = nullptr;
 	}
 
+	TRDGUniformBufferRef<FSceneUniformParameters> SceneUB = SceneUniformBuffer.GetBuffer(GraphBuilder);
+
 	// Process batched passes
 	if (!InstanceCullingMergedContext.Batches.IsEmpty())
 	{
@@ -3510,7 +3514,7 @@ void FVirtualShadowMapArray::RenderVirtualShadowMapsNonNanite(FRDGBuilder& Graph
 				PassParameters->InstanceCullingDrawParams.DrawIndirectArgsBuffer = CullingResult.DrawIndirectArgsRDG;
 				PassParameters->InstanceCullingDrawParams.InstanceIdOffsetBuffer = CullingResult.InstanceIdOffsetBufferRDG;
 				PassParameters->InstanceCullingDrawParams.InstanceCulling = InstanceCullingUniformBuffer;
-				PassParameters->InstanceCullingDrawParams.Scene = SceneUniformBuffer.GetBuffer(GraphBuilder);
+				PassParameters->InstanceCullingDrawParams.Scene = SceneUB;
 				PassParameters->InstanceCullingDrawParams.IndirectArgsByteOffset = 0U;
 				PassParameters->InstanceCullingDrawParams.InstanceDataByteOffset = 0U;
 
@@ -3566,7 +3570,7 @@ void FVirtualShadowMapArray::RenderVirtualShadowMapsNonNanite(FRDGBuilder& Graph
 
 					FString LightNameWithLevel;
 					FSceneRenderer::GetLightNameForDrawEvent(ProjectedShadowInfo->GetLightSceneInfo().Proxy, LightNameWithLevel);
-					AddRasterPass(GraphBuilder, RDG_EVENT_NAME("Rasterize[%s]", *LightNameWithLevel), ShadowDepthView, ShadowDepthPassUniformBuffer, *this, VirtualShadowViewsRDG, CullingResult, MeshCommandPass, BatchedPassParameters[Index], InstanceCullingUniformBuffer);
+					AddRasterPass(GraphBuilder, RDG_EVENT_NAME("Rasterize[%s]", *LightNameWithLevel), ShadowDepthView, ShadowDepthPassUniformBuffer, *this, VirtualShadowViewsRDG, CullingResult, MeshCommandPass, BatchedPassParameters[Index], InstanceCullingUniformBuffer, SceneUB);
 			}
 		}
 	}
@@ -3629,7 +3633,7 @@ void FVirtualShadowMapArray::RenderVirtualShadowMapsNonNanite(FRDGBuilder& Graph
 			FVirtualShadowDepthPassParameters* DepthPassParams = GraphBuilder.AllocParameters<FVirtualShadowDepthPassParameters>();
 			DepthPassParams->InstanceCullingDrawParams.IndirectArgsByteOffset = 0;
 			DepthPassParams->InstanceCullingDrawParams.InstanceDataByteOffset = 0;
-			AddRasterPass(GraphBuilder, RDG_EVENT_NAME("Rasterize"), ShadowDepthView, ShadowDepthPassUniformBuffer, *this, VirtualShadowViewsRDG, CullingResult, MeshCommandPass, DepthPassParams, InstanceCullingUniformBuffer);
+			AddRasterPass(GraphBuilder, RDG_EVENT_NAME("Rasterize"), ShadowDepthView, ShadowDepthPassUniformBuffer, *this, VirtualShadowViewsRDG, CullingResult, MeshCommandPass, DepthPassParams, InstanceCullingUniformBuffer, SceneUB);
 		}
 
 
