Fix support for third party TAA with mobile renderer

Last attempt did not work for mobile forward rendering, so I extracted the relevant part to a new file.

diff --git a/Engine/Shaders/Private/PostProcessEyeAdaptation.usf b/Engine/Shaders/Private/PostProcessEyeAdaptation.usf
--- a/Engine/Shaders/Private/PostProcessEyeAdaptation.usf
+++ b/Engine/Shaders/Private/PostProcessEyeAdaptation.usf
@@ -233,14 +233,6 @@ void EyeAdaptationCS(uint2 DispatchThreadId : SV_DispatchThreadID)
 	RWEyeAdaptationBuffer[0] = OutData0;
 	RWEyeAdaptationBuffer[1] = float4(AverageLocalExposure, 0.0f, 0.0f, 0.0f);
 }
-
-RWTexture2D<float4> RWEyeAdaptationTexture;
-
-[numthreads(1, 1, 1)]
-void CopyEyeAdaptationToTextureCS(uint2 DispatchThreadId : SV_DispatchThreadID)
-{
-	RWEyeAdaptationTexture[DispatchThreadId] = EyeAdaptationBuffer[0];
-}
 #endif
 
 void BasicEyeAdaptationSetupPS(
diff --git a/Engine/Shaders/Private/PostProcessEyeAdaptationUtils.usf b/Engine/Shaders/Private/PostProcessEyeAdaptationUtils.usf
new file mode 100644
--- /dev/null
+++ b/Engine/Shaders/Private/PostProcessEyeAdaptationUtils.usf
@@ -0,0 +1,16 @@
+// Copyright Epic Games, Inc. All Rights Reserved.
+
+#include "Common.ush"
+
+StructuredBuffer<float4> EyeAdaptationBuffer;
+
+#if COMPUTESHADER
+RWTexture2D<float4> RWEyeAdaptationTexture;
+
+[numthreads(1, 1, 1)]
+void CopyEyeAdaptationToTextureCS(uint2 DispatchThreadId : SV_DispatchThreadID)
+{
+	RWEyeAdaptationTexture[DispatchThreadId] = EyeAdaptationBuffer[0];
+}
+
+#endif // COMPUTESHADER
diff --git a/Engine/Source/Runtime/Renderer/Private/PostProcess/PostProcessEyeAdaptation.cpp b/Engine/Source/Runtime/Renderer/Private/PostProcess/PostProcessEyeAdaptation.cpp
--- a/Engine/Source/Runtime/Renderer/Private/PostProcess/PostProcessEyeAdaptation.cpp
+++ b/Engine/Source/Runtime/Renderer/Private/PostProcess/PostProcessEyeAdaptation.cpp
@@ -591,14 +591,9 @@ public:
 		SHADER_PARAMETER_RDG_BUFFER_SRV(StructuredBuffer<float4>, EyeAdaptationBuffer)
 		SHADER_PARAMETER_RDG_TEXTURE_UAV(RWTexture2D<float4>, RWEyeAdaptationTexture)
 		END_SHADER_PARAMETER_STRUCT()
-
-		static bool ShouldCompilePermutation(const FGlobalShaderPermutationParameters& Parameters)
-	{
-		return IsFeatureLevelSupported(Parameters.Platform, ERHIFeatureLevel::SM5);
-	}
 };
 
-IMPLEMENT_GLOBAL_SHADER(FCopyEyeAdaptationToTextureCS, "/Engine/Private/PostProcessEyeAdaptation.usf", "CopyEyeAdaptationToTextureCS", SF_Compute);
+IMPLEMENT_GLOBAL_SHADER(FCopyEyeAdaptationToTextureCS, "/Engine/Private/PostProcessEyeAdaptationUtils.usf", "CopyEyeAdaptationToTextureCS", SF_Compute);
 
 void AddCopyEyeAdaptationDataToTexturePass(FRDGBuilder& GraphBuilder, const FGlobalShaderMap* ShaderMap, FRDGBufferRef EyeAdaptationBuffer, FRDGTextureRef OutputTexture)
 {
