Fix Distance Fields not being generated on newer iGPUs

#rnx

diff --git a/Engine/Source/Runtime/Renderer/Private/DistanceFieldAmbientOcclusion.cpp b/Engine/Source/Runtime/Renderer/Private/DistanceFieldAmbientOcclusion.cpp
--- a/Engine/Source/Runtime/Renderer/Private/DistanceFieldAmbientOcclusion.cpp
+++ b/Engine/Source/Runtime/Renderer/Private/DistanceFieldAmbientOcclusion.cpp
@@ -678,8 +678,8 @@ bool SupportsDistanceFieldAO(ERHIFeatureLevel::Type FeatureLevel, EShaderPlatfor
 		// Pre-GCN AMD cards have a driver bug that prevents the global distance field from being generated correctly
 		// Better to disable entirely than to display garbage
 		&& !GRHIDeviceIsAMDPreGCNArchitecture
-		// Intel HD 4000 hangs in the RHICreateTexture3D call to allocate the large distance field atlas, and virtually no Intel cards can afford it anyway
-		&& !GRHIDeviceIsIntegrated
+		// In case of iGPU we use the maximum feature level to differentiate between older and newer more capable iGPU
+		&& (!GRHIDeviceIsIntegrated || GMaxRHIFeatureLevel >= ERHIFeatureLevel::SM6)
 		&& DoesPlatformSupportDistanceFieldAO(ShaderPlatform)
 		&& IsUsingDistanceFields(ShaderPlatform);
 }
@@ -728,9 +728,9 @@ bool FSceneRenderer::ShouldPrepareDistanceFieldScene() const
 		return false;
 	}
 
-	if (GRHIDeviceIsIntegrated)
+	if (GRHIDeviceIsIntegrated && GMaxRHIFeatureLevel < ERHIFeatureLevel::SM6)
 	{
-		// Intel HD 4000 hangs in the RHICreateTexture3D call to allocate the large distance field atlas, and virtually no Intel cards can afford it anyway
+		// In case of iGPU we use the maximum feature level to differentiate between older and newer more capable iGPU
 		return false;
 	}
 
