Fix light grid debug visualization when TSR is disabled.
- apply screen transform based on actual render target / primary view rect instead of directly using ViewResolutionFraction.

diff --git a/Engine/Shaders/Private/LightGridInjection.usf b/Engine/Shaders/Private/LightGridInjection.usf
--- a/Engine/Shaders/Private/LightGridInjection.usf
+++ b/Engine/Shaders/Private/LightGridInjection.usf
@@ -487,8 +487,10 @@ void LightGridInjectionCS(
 #define SUPPORT_CONTACT_SHADOWS 0
 #include "/Engine/Private/ShaderPrint.ush"
 #include "ColorMap.ush"
+#include "ScreenPass.ush"
 
 Texture2D<float4> DepthTexture;
+FScreenTransform ScreenToPrimaryScreenPos;
 uint DebugMode;
 uint MaxThreshold;
 
@@ -502,12 +504,11 @@ void DebugLightGridPS(
 
 	const uint EyeIndex = 0;
 	const uint2 PixelPos = SVPos.xy;
-	const uint2 LocalPixelPos = (SVPos.xy - ResolvedView.ViewRectMin.xy);
 
-	// Using View.ViewResolutionFraction to maintain a good look under any dynamic resolution factor.
-	const uint2 PixelPosDynRes = float2(PixelPos) * View.ViewResolutionFraction;
-	const uint2 LocalPixelPosDynRes = float2(LocalPixelPos) * View.ViewResolutionFraction;
-	const uint2 LocalPixelPosDynRes2 = float2(LocalPixelPos+1) * View.ViewResolutionFraction;
+	const float2 PixelPosDynRes = ApplyScreenTransform(PixelPos, ScreenToPrimaryScreenPos);
+	const float2 PixelPosDynRes2 = ApplyScreenTransform(PixelPos + 1, ScreenToPrimaryScreenPos);
+	const uint2 LocalPixelPosDynRes = (PixelPosDynRes.xy - ResolvedView.ViewRectMin.xy);
+	const uint2 LocalPixelPosDynRes2 = (PixelPosDynRes2.xy - ResolvedView.ViewRectMin.xy);
 
 	if (any(LocalPixelPosDynRes >= uint2(ResolvedView.ViewSizeAndInvSize.xy)))
 	{
@@ -564,7 +565,11 @@ void DebugLightGridPS(
 	{
 		OutLuminanceTransmittance.a   = 0.0f;
 		OutLuminanceTransmittance.rgb = GetHSVDebugColor(float(DebugNumLightsInGridCell) / MaxThreshold);
-		PrintSmallUint(PixelPos, OutLuminanceTransmittance.rgb, float3(0.1, 0.1, 0.1), (TileTopLeftPixelCoord + TileSize/2) / View.ViewResolutionFraction, float(DebugNumLightsInGridCell));
+
+		int2 TextTopLeftPosition = (TileTopLeftPixelCoord + TileSize / 2);
+		TextTopLeftPosition = (TextTopLeftPosition - ScreenToPrimaryScreenPos.zw) / ScreenToPrimaryScreenPos.xy; // invert ScreenToPrimaryScreenPos transform
+
+		PrintSmallUint(PixelPos, OutLuminanceTransmittance.rgb, float3(0.1, 0.1, 0.1), TextTopLeftPosition, float(DebugNumLightsInGridCell));
 	}
 
 	if ((TileTopLeftPixelCoord.x <= LocalPixelPosDynRes.x && TileTopLeftPixelCoord2.x >= LocalPixelPosDynRes.x)
diff --git a/Engine/Source/Runtime/Renderer/Private/LightGridInjection.cpp b/Engine/Source/Runtime/Renderer/Private/LightGridInjection.cpp
--- a/Engine/Source/Runtime/Renderer/Private/LightGridInjection.cpp
+++ b/Engine/Source/Runtime/Renderer/Private/LightGridInjection.cpp
@@ -1152,6 +1152,7 @@ class FDebugLightGridPS : public FGlobalShader
 		SHADER_PARAMETER_STRUCT_INCLUDE(ShaderPrint::FShaderParameters, ShaderPrintParameters)
 		SHADER_PARAMETER_TEXTURE(Texture2D, MiniFontTexture)
 		SHADER_PARAMETER_RDG_TEXTURE(Texture2D, DepthTexture)
+		SHADER_PARAMETER(FScreenTransform, ScreenToPrimaryScreenPos)
 		SHADER_PARAMETER(uint32, DebugMode)
 		SHADER_PARAMETER(uint32, MaxThreshold)
 		RENDER_TARGET_BINDING_SLOTS()
@@ -1187,7 +1188,7 @@ class FDebugLightGridPS : public FGlobalShader
 };
 IMPLEMENT_GLOBAL_SHADER(FDebugLightGridPS, "/Engine/Private/LightGridInjection.usf", "DebugLightGridPS", SF_Pixel);
 
-FScreenPassTexture AddVisualizeLightGridPass(FRDGBuilder& GraphBuilder, const FViewInfo& View, FScreenPassTexture& ScreenPassSceneColor, FRDGTextureRef SceneDepthTexture)
+FScreenPassTexture AddVisualizeLightGridPass(FRDGBuilder& GraphBuilder, const FViewInfo& View, FScreenPassTexture ScreenPassSceneColor, FScreenPassTexture SceneDepthTexture)
 {
 	if (ShaderPrint::IsSupported(View.Family->GetShaderPlatform()))
 	{
@@ -1205,7 +1206,10 @@ FScreenPassTexture AddVisualizeLightGridPass(FRDGBuilder& GraphBuilder, const FV
 		PassParameters->ViewUniformBuffer = View.ViewUniformBuffer;
 		PassParameters->Forward = View.ForwardLightingResources.ForwardLightUniformBuffer;
 		ShaderPrint::SetParameters(GraphBuilder, View.ShaderPrintData, PassParameters->ShaderPrintParameters);
-		PassParameters->DepthTexture = SceneDepthTexture ? SceneDepthTexture : GSystemTextures.GetMaxFP16Depth(GraphBuilder);
+		PassParameters->DepthTexture = SceneDepthTexture.IsValid() ? SceneDepthTexture.Texture : GSystemTextures.GetMaxFP16Depth(GraphBuilder);
+		PassParameters->ScreenToPrimaryScreenPos = (
+			FScreenTransform::ChangeTextureBasisFromTo(ScreenPassSceneColor, FScreenTransform::ETextureBasis::TexelPosition, FScreenTransform::ETextureBasis::ViewportUV) *
+			FScreenTransform::ChangeTextureBasisFromTo(SceneDepthTexture, FScreenTransform::ETextureBasis::ViewportUV, FScreenTransform::ETextureBasis::TexelPosition));
 		PassParameters->MiniFontTexture = GetMiniFontTexture();
 		PassParameters->RenderTargets[0] = FRenderTargetBinding(ScreenPassSceneColor.Texture, ERenderTargetLoadAction::ELoad);
 		PassParameters->DebugMode = GForwardLightGridDebug;
diff --git a/Engine/Source/Runtime/Renderer/Private/PostProcess/PostProcessing.cpp b/Engine/Source/Runtime/Renderer/Private/PostProcess/PostProcessing.cpp
--- a/Engine/Source/Runtime/Renderer/Private/PostProcess/PostProcessing.cpp
+++ b/Engine/Source/Runtime/Renderer/Private/PostProcess/PostProcessing.cpp
@@ -81,7 +81,7 @@ bool IsMobileEyeAdaptationEnabled(const FViewInfo& View);
 
 bool IsValidBloomSetupVariation(bool bUseBloom, bool bUseSun, bool bUseDof, bool bUseEyeAdaptation);
 
-extern FScreenPassTexture AddVisualizeLightGridPass(FRDGBuilder& GraphBuilder, const FViewInfo& View, FScreenPassTexture& ScreenPassSceneColor, FRDGTextureRef SceneDepthTexture);
+extern FScreenPassTexture AddVisualizeLightGridPass(FRDGBuilder& GraphBuilder, const FViewInfo& View, FScreenPassTexture ScreenPassSceneColor, FScreenPassTexture SceneDepthTexture);
 extern bool ShouldVisualizeLightGrid();
 
 namespace
@@ -1514,7 +1514,7 @@ void AddPostProcessingPasses(
 	{
 		FScreenPassRenderTarget OverrideOutput;
 		PassSequence.AcceptOverrideIfLastPass(EPass::VisualizeLightGrid, OverrideOutput);
-		SceneColor = AddVisualizeLightGridPass(GraphBuilder, View, SceneColor, SceneDepth.Texture);
+		SceneColor = AddVisualizeLightGridPass(GraphBuilder, View, SceneColor, SceneDepth);
 	}
 
 #if WITH_EDITOR
