Move post process compositing of nanite visualisation pass to before the editor primitive rendering.
This specifically solves the issue where the mesh paint brush gizmo is getting hidden by mesh paint color visualizations.

diff --git a/Engine/Source/Runtime/Renderer/Private/PostProcess/PostProcessing.cpp b/Engine/Source/Runtime/Renderer/Private/PostProcess/PostProcessing.cpp
--- a/Engine/Source/Runtime/Renderer/Private/PostProcess/PostProcessing.cpp
+++ b/Engine/Source/Runtime/Renderer/Private/PostProcess/PostProcessing.cpp
@@ -1562,6 +1562,13 @@ void AddPostProcessingPasses(
 		SceneColor = VirtualShadowMapArray->AddVisualizePass(GraphBuilder, View, ViewIndex, EVSMVisualizationPostPass::PreEditorPrimitives, SceneColor, OverrideOutput);
 	}
 	
+#if WITH_EDITOR || !UE_BUILD_SHIPPING
+ 	if (EngineShowFlags.VisualizeNanite && NaniteRasterResults != nullptr)
+ 	{
+ 		AddVisualizeNanitePass(GraphBuilder, View, SceneColor, *NaniteRasterResults);
+ 	}
+#endif
+
 #if WITH_EDITOR
 	if (PassSequence.IsEnabled(EPass::SelectionOutline))
 	{
@@ -1854,11 +1861,6 @@ void AddPostProcessingPasses(
 			AddTestImagePass(GraphBuilder, View, SceneColor);
 		}
 
-		if (EngineShowFlags.VisualizeNanite && NaniteRasterResults != nullptr)
-		{
-			AddVisualizeNanitePass(GraphBuilder, View, SceneColor, *NaniteRasterResults);
-		}
-
 		#if WITH_EDITOR
 		if (CVarGBufferPicking.GetValueOnRenderThread())
 		{
