Disable forced invalidation by WPO distance disable scale change by default.
This can produce some unpredictable spikes with dynamic resolution and generally WPO distances need to be set far enough that this is not visible anyways.

#rnx

diff --git a/Engine/Source/Runtime/Renderer/Private/VirtualShadowMaps/VirtualShadowMapCacheManager.cpp b/Engine/Source/Runtime/Renderer/Private/VirtualShadowMaps/VirtualShadowMapCacheManager.cpp
--- a/Engine/Source/Runtime/Renderer/Private/VirtualShadowMaps/VirtualShadowMapCacheManager.cpp
+++ b/Engine/Source/Runtime/Renderer/Private/VirtualShadowMaps/VirtualShadowMapCacheManager.cpp
@@ -127,6 +127,14 @@ static TAutoConsoleVariable<float> CVarVSMDynamicResolutionMaxPagePoolLoadFactor
 	ECVF_RenderThreadSafe
 );
 
+static TAutoConsoleVariable<int32> CVarClipmapWPODisableDistanceInvalidate(
+	TEXT("r.Shadow.Virtual.Clipmap.WPODisableDistance.InvalidateOnScaleChange"),
+	0,
+	TEXT("When enabled, forces an invalidation of clipmap levels when the distance disable clipmap level shifts due to resolution or FOV changes.\n")
+	TEXT("This can sometimes make things visually more consistent but can also introduce unpredictable performance spikes, so it is disabled by default."),
+	ECVF_Scalability | ECVF_RenderThreadSafe
+);
+
 int32 GVSMLightRadiusInvalidationCulling = 1;
 FAutoConsoleVariableRef CVarVSMLightRadiusCulling(
 	TEXT("r.Shadow.Virtual.Cache.CPUCullInvalidationsOutsideLightRadius"),
@@ -215,7 +223,7 @@ void FVirtualShadowMapCacheEntry::UpdateClipmapLevel(
 	bCacheValid = bCacheValid && (ViewRadiusZ == Clipmap.ViewRadiusZ);
 
 	// Not valid if WPO threshold has changed
-	if (bCacheValid && (WPODistanceDisableThresholdSquared != Clipmap.WPODistanceDisableThresholdSquared))
+	if (bCacheValid && CVarClipmapWPODisableDistanceInvalidate.GetValueOnRenderThread() != 0 && (WPODistanceDisableThresholdSquared != Clipmap.WPODistanceDisableThresholdSquared))
 	{
 		bCacheValid = false;
 		// Only warn once per change... when this changes it will hit all of them
