Fix capsule direct shadows for stereo rendering

diff --git a/Engine/Shaders/Private/CapsuleShadowShaders.usf b/Engine/Shaders/Private/CapsuleShadowShaders.usf
--- a/Engine/Shaders/Private/CapsuleShadowShaders.usf
+++ b/Engine/Shaders/Private/CapsuleShadowShaders.usf
@@ -947,6 +947,8 @@ Texture2D ReceiverBentNormalTexture;
 RWTexture2D<float4> RWBentNormalTexture;
 #endif
 
+uint EyeIndex;
+
 float IndirectCapsuleSelfShadowingIntensity;
 uint DownsampleFactor;
 RWTexture2D<float2> RWShadowFactors;
@@ -971,8 +973,12 @@ void CapsuleShadowingCS(
 	float4 ReceiverTextureValue = ReceiverBentNormalTexture.Load(DispatchThreadId.xyz);
 	float3 ReceiverBentNormal = ReceiverTextureValue.xyz;
 	float SceneDepth = ReceiverTextureValue.w;
+#else
+#if SHADING_PATH_MOBILE
+	float SceneDepth = CalcSceneDepth(ScreenUV, EyeIndex);
 #else
 	float SceneDepth = CalcSceneDepth(ScreenUV);
+#endif
 #endif
 	
 	const float3 OpaqueTranslatedWorldPosition = mul(float4(GetScreenPositionForProjectionType(ScreenPosition, SceneDepth), SceneDepth, 1), PrimaryView.ScreenToTranslatedWorld).xyz;
@@ -1129,8 +1135,14 @@ void CapsuleShadowingUpsampleVS(
 	uint VertexId : SV_VertexID,
 	uint InstanceId : SV_InstanceID, 
 	out float4 OutPosition : SV_POSITION
+#if INSTANCED_STEREO && MOBILE_MULTI_VIEW
+	, out uint LayerIndex : SV_RenderTargetArrayIndex
+#endif
 	)
 {
+#if INSTANCED_STEREO && MOBILE_MULTI_VIEW
+	LayerIndex = EyeIndex;
+#endif
 	// Compute the actual instance id for when multiple tiles are packed into the vertex buffer
 	uint EffectiveInstanceId = InstanceId * TILES_PER_INSTANCE + VertexId / 4;
 	uint NumCapsulesAffectingTile = TileIntersectionCounts[EffectiveInstanceId];
@@ -1167,7 +1179,7 @@ void CapsuleShadowingUpsamplePS(
 
 	float Output;
 	float SceneDepth;
-	UpsampleShadowFactors(SVPos, ScissorRectMinAndSize, 1.0f / DOWNSAMPLE_FACTOR, 0, POSITIVE_INFINITY, ShadowFactorsTexture, ShadowFactorsSampler, ShadowFactorsUVBilinearMax, Output, SceneDepth);
+	UpsampleShadowFactors(SVPos, ScissorRectMinAndSize, 1.0f / DOWNSAMPLE_FACTOR, 0, POSITIVE_INFINITY, ShadowFactorsTexture, ShadowFactorsSampler, ShadowFactorsUVBilinearMax, Output, SceneDepth, EyeIndex);
 
 	if (OutputtingToLightAttenuation > 0)
 	{
diff --git a/Engine/Shaders/Private/ShadowFactorsUpsampleCommon.ush b/Engine/Shaders/Private/ShadowFactorsUpsampleCommon.ush
--- a/Engine/Shaders/Private/ShadowFactorsUpsampleCommon.ush
+++ b/Engine/Shaders/Private/ShadowFactorsUpsampleCommon.ush
@@ -21,14 +21,18 @@ void UpsampleShadowFactors(
 	SamplerState ShadowFactorsSampler,
 	float2 ShadowFactorsUVBilinearMax,
 	out float ShadowFactor, 
-	out float SceneDepth)
+	out float SceneDepth,
+	uint EyeIndex)
 {
 	float2 BufferUV = SvPositionToBufferUV(SVPos);
 
 	// ShadowFactors is computed at 0,0 regardless of viewrect min
 	float2 ShadowFactorsUV = BufferUV - ScissorRectMinAndSize.xy * View.BufferSizeAndInvSize.zw;
-
+#if MOBILE_MULTI_VIEW
+	SceneDepth = CalcSceneDepth(BufferUV, EyeIndex);
+#else
 	SceneDepth = CalcSceneDepth(BufferUV);
+#endif
 
 #define BILATERAL_UPSAMPLE 1
 #if BILATERAL_UPSAMPLE && SHADOW_FACTORS_UPSAMPLE_REQUIRED
@@ -85,3 +89,18 @@ void UpsampleShadowFactors(
 	ShadowFactor = Texture2DSampleLevel(ShadowFactorsTexture, ShadowFactorsSampler, ShadowFactorsUV, 0).x;
 #endif
 }
+
+void UpsampleShadowFactors(
+	float4 SVPos,
+	int4 ScissorRectMinAndSize,
+	float OneOverDownsampleFactor,
+	float ShadowFactorsMinDepth,
+	float ShadowFactorsMaxDepth,
+	Texture2D ShadowFactorsTexture,
+	SamplerState ShadowFactorsSampler,
+	float2 ShadowFactorsUVBilinearMax,
+	out float ShadowFactor,
+	out float SceneDepth)
+{
+	UpsampleShadowFactors(SVPos, ScissorRectMinAndSize, OneOverDownsampleFactor, ShadowFactorsMinDepth, ShadowFactorsMaxDepth, ShadowFactorsTexture, ShadowFactorsSampler, ShadowFactorsUVBilinearMax, ShadowFactor, SceneDepth, 0);
+}
diff --git a/Engine/Source/Runtime/Renderer/Private/CapsuleShadowRendering.cpp b/Engine/Source/Runtime/Renderer/Private/CapsuleShadowRendering.cpp
--- a/Engine/Source/Runtime/Renderer/Private/CapsuleShadowRendering.cpp
+++ b/Engine/Source/Runtime/Renderer/Private/CapsuleShadowRendering.cpp
@@ -197,6 +197,8 @@ public:
 		SHADER_PARAMETER_RDG_UNIFORM_BUFFER(FSubstrateGlobalUniformParameters, Substrate)
 		SHADER_PARAMETER_STRUCT_REF(FViewUniformShaderParameters, View)
 
+		SHADER_PARAMETER(uint32, EyeIndex)
+
 		SHADER_PARAMETER_RDG_TEXTURE_UAV(RWTexture2D, RWShadowFactors)
 		SHADER_PARAMETER_RDG_TEXTURE_UAV(RWTexture2D, RWBentNormalTexture)
 		SHADER_PARAMETER_RDG_TEXTURE(Texture2D, ReceiverBentNormalTexture)
@@ -299,6 +301,7 @@ class FCapsuleShadowingUpsampleVS : public FGlobalShader
 		SHADER_PARAMETER(FIntPoint, TileDimensions)
 		SHADER_PARAMETER(FVector2f, TileSize)
 		SHADER_PARAMETER(FIntRect, ScissorRectMinAndSize)
+		SHADER_PARAMETER(uint32, EyeIndex)
 	END_SHADER_PARAMETER_STRUCT()
 
 	static void ModifyCompilationEnvironment(const FGlobalShaderPermutationParameters& Parameters, FShaderCompilerEnvironment& OutEnvironment)
@@ -323,6 +326,7 @@ public:
 		SHADER_PARAMETER(FIntRect, ScissorRectMinAndSize)
 		SHADER_PARAMETER(FVector2f, ShadowFactorsUVBilinearMax)
 		SHADER_PARAMETER(float, OutputtingToLightAttenuation)
+		SHADER_PARAMETER(uint32, EyeIndex)
 	END_SHADER_PARAMETER_STRUCT()
 
 	class FUpsampleRequired : SHADER_PERMUTATION_BOOL("SHADOW_FACTORS_UPSAMPLE_REQUIRED");
@@ -359,6 +363,7 @@ void SetupCapsuleShadowingParameters(
 	float MaxOcclusionDistance,
 	const FScene* Scene,
 	const FViewInfo& View,
+	const uint32 ViewIndex,
 
 	uint32 NumShadowCapsules,
 	FRDGBufferSRVRef ShadowCapsuleShapesBuffer,
@@ -429,6 +434,8 @@ void SetupCapsuleShadowingParameters(
 		check(ShadowingType == ECapsuleShadowingType::IndirectTiledCulling || ShadowingType == ECapsuleShadowingType::MovableSkylightTiledCulling || ShadowingType == ECapsuleShadowingType::MovableSkylightTiledCullingGatherFromReceiverBentNormal);
 	}
 
+	Parameters.EyeIndex = ViewIndex;
+
 	Parameters.ScissorRectMinAndSize = FIntRect(ScissorRect.Min, ScissorRect.Size());
 	Parameters.DownsampleFactor = DownsampleFactor;
 
@@ -492,7 +499,7 @@ bool FSceneRenderer::RenderCapsuleDirectShadows(
 		const FVector PreViewTranslation = View.ViewMatrices.GetPreViewTranslation();
 
 		RDG_GPU_MASK_SCOPE(GraphBuilder, View.GPUMask);
-		RDG_EVENT_SCOPE_STAT(GraphBuilder, CapsuleShadows, "CapsuleShadows");
+		RDG_EVENT_SCOPE_STAT(GraphBuilder, CapsuleShadows, "CapsuleShadows View%d", ViewIndex);
 		RDG_GPU_STAT_SCOPE(GraphBuilder, CapsuleShadows);
 
 		TArray<FCapsuleShape3f> CapsuleShapeData;
@@ -571,6 +578,7 @@ bool FSceneRenderer::RenderCapsuleDirectShadows(
 					GCapsuleMaxDirectOcclusionDistance,
 					Scene,
 					View,
+					ViewIndex,
 					CapsuleShapeData.Num(),
 					GraphBuilder.CreateSRV(ShadowCapsuleShapesBuffer),
 					0,
@@ -606,13 +614,26 @@ bool FSceneRenderer::RenderCapsuleDirectShadows(
 				auto PixelShader = View.ShaderMap->GetShader<FCapsuleShadowingUpsamplePS>(PermutationVector);
 
 				FUpsampleCapsuleShadowParameters* PassParameters = GraphBuilder.AllocParameters<FUpsampleCapsuleShadowParameters>();
-				PassParameters->RenderTargets[0] = FRenderTargetBinding(ScreenShadowMaskTexture, ERenderTargetLoadAction::ELoad);
+
+				if (View.Aspects.IsMobileMultiViewEnabled() && !View.Aspects.IsInstancedStereoEnabled()) // Native mobile multi view
+				{
+					PassParameters->RenderTargets[0] = FRenderTargetBinding(ScreenShadowMaskTexture, ERenderTargetLoadAction::ELoad, /*MipIndex*/ 0, /*ArraySlice*/ ViewIndex);
+				}
+				else // Mobile multi view fallback or not stereo
+				{
+					// D3D12 RHI does not support setting slice 1 of a texture 2d array without creating it with TexCreate_TargetArraySlicesIndependently.
+					// Unfortunately, this flag does not allow to target the array with SV_RenderTargetArrayIndex anymore so this is not viable for mobile multi view fallback.
+					// In this case, the PS draw calls are rerouted to the correct slice with SV_RenderTargetArrayIndex as a workaround.
+					PassParameters->RenderTargets[0] = FRenderTargetBinding(ScreenShadowMaskTexture, ERenderTargetLoadAction::ELoad);
+				}
+
 				PassParameters->SceneTextures = GetSceneTextureShaderParameters(View);
 				PassParameters->VS.View = GetShaderBinding(View.ViewUniformBuffer);
 				PassParameters->VS.TileIntersectionCounts = CapsuleTileIntersectionCountsSRV;
 				PassParameters->VS.TileDimensions = GroupSize;
 				PassParameters->VS.TileSize = FVector2f(GShadowShapeTileSize * GetCapsuleShadowDownsampleFactor(), GShadowShapeTileSize * GetCapsuleShadowDownsampleFactor());
 				PassParameters->VS.ScissorRectMinAndSize = FIntRect(ScissorRect.Min, ScissorRect.Size());
+				PassParameters->VS.EyeIndex = ViewIndex;
 
 				PassParameters->PS.View = GetShaderBinding(View.ViewUniformBuffer);
 				PassParameters->PS.Substrate = Substrate::BindSubstrateGlobalUniformParameters(View);
@@ -621,6 +642,7 @@ bool FSceneRenderer::RenderCapsuleDirectShadows(
 				PassParameters->PS.ScissorRectMinAndSize = FIntRect(ScissorRect.Min, ScissorRect.Size());
 				PassParameters->PS.ShadowFactorsUVBilinearMax = (RayTracedShadowsTextureViewportMax - 0.5f) * RayTracedShadowsTextureExtentInverse;
 				PassParameters->PS.OutputtingToLightAttenuation = 1.0f;
+				PassParameters->PS.EyeIndex = ViewIndex;
 
 				ClearUnusedGraphResources(VertexShader, &PassParameters->VS);
 				ClearUnusedGraphResources(PixelShader, &PassParameters->PS);
@@ -1043,6 +1065,7 @@ void FDeferredShadingSceneRenderer::RenderIndirectCapsuleShadows(FRDGBuilder& Gr
 					GCapsuleMaxIndirectOcclusionDistance,
 					Scene,
 					View,
+					ViewIndex,
 					NumCapsuleShapes,
 					Resources.IndirectShadowCapsuleShapesSRV,
 					NumMeshDistanceFieldCasters,
@@ -1234,6 +1257,7 @@ void FDeferredShadingSceneRenderer::RenderCapsuleShadowsForMovableSkylight(
 						GCapsuleMaxIndirectOcclusionDistance,
 						Scene,
 						View,
+						0, // Not used in stereo rendering
 						NumCapsuleShapes,
 						Resources.IndirectShadowCapsuleShapesSRV,
 						NumMeshDistanceFieldCasters,
