RHI: Add missing PixelInspector transitions

#rnx

diff --git a/Engine/Source/Runtime/Renderer/Private/PostProcess/PostProcessBufferInspector.cpp b/Engine/Source/Runtime/Renderer/Private/PostProcess/PostProcessBufferInspector.cpp
--- a/Engine/Source/Runtime/Renderer/Private/PostProcess/PostProcessBufferInspector.cpp
+++ b/Engine/Source/Runtime/Renderer/Private/PostProcess/PostProcessBufferInspector.cpp
@@ -66,7 +66,9 @@ void ProcessPixelInspectorRequests(
 						FRHICopyTextureInfo CopyInfo;
 						CopyInfo.SourcePosition = SourcePoint;
 						CopyInfo.Size = FIntVector(1);
+						RHICmdList.Transition(FRHITransitionInfo(DestinationBufferDepth, ERHIAccess::SRVMask, ERHIAccess::CopyDest));
 						RHICmdList.CopyTexture(SourceBufferSceneDepth, DestinationBufferDepth, CopyInfo);
+						RHICmdList.Transition(FRHITransitionInfo(DestinationBufferDepth, ERHIAccess::CopyDest, ERHIAccess::SRVMask));
 					}
 				}
 			}
@@ -107,7 +109,9 @@ void ProcessPixelInspectorRequests(
 
 					if (CopyInfo.Size.X > 0 && CopyInfo.Size.Y > 0)
 					{
+						RHICmdList.Transition(FRHITransitionInfo(DestinationBufferFinalColor, ERHIAccess::SRVMask, ERHIAccess::CopyDest));
 						RHICmdList.CopyTexture(SourceBufferFinalColor, DestinationBufferFinalColor, CopyInfo);
+						RHICmdList.Transition(FRHITransitionInfo(DestinationBufferFinalColor, ERHIAccess::CopyDest, ERHIAccess::SRVMask));
 					}
 				}
 			}
@@ -129,7 +133,9 @@ void ProcessPixelInspectorRequests(
 					FRHICopyTextureInfo CopyInfo;
 					CopyInfo.SourcePosition = SourcePoint;
 					CopyInfo.Size = FIntVector(1, 1, 1);
+					RHICmdList.Transition(FRHITransitionInfo(DestinationBufferSceneColor, ERHIAccess::SRVMask, ERHIAccess::CopyDest));
 					RHICmdList.CopyTexture(SourceBufferSceneColor, DestinationBufferSceneColor, CopyInfo);
+					RHICmdList.Transition(FRHITransitionInfo(DestinationBufferSceneColor, ERHIAccess::CopyDest, ERHIAccess::SRVMask));
 				}
 			}
 
@@ -152,7 +158,9 @@ void ProcessPixelInspectorRequests(
 						FRHICopyTextureInfo CopyInfo;
 						CopyInfo.SourcePosition = SourcePoint;
 						CopyInfo.Size = FIntVector(1, 1, 1);
+						RHICmdList.Transition(FRHITransitionInfo(DestinationBufferHDR, ERHIAccess::SRVMask, ERHIAccess::CopyDest));
 						RHICmdList.CopyTexture(SourceBufferHDR, DestinationBufferHDR, CopyInfo);
+						RHICmdList.Transition(FRHITransitionInfo(DestinationBufferHDR, ERHIAccess::CopyDest, ERHIAccess::SRVMask));
 					}
 				}
 			}
@@ -176,7 +184,9 @@ void ProcessPixelInspectorRequests(
 						FRHICopyTextureInfo CopyInfo;
 						CopyInfo.SourcePosition = SourcePoint;
 						CopyInfo.Size = FIntVector(1, 1, 1);
+						RHICmdList.Transition(FRHITransitionInfo(DestinationBufferA, ERHIAccess::SRVMask, ERHIAccess::CopyDest));
 						RHICmdList.CopyTexture(SourceBufferA, DestinationBufferA, CopyInfo);
+						RHICmdList.Transition(FRHITransitionInfo(DestinationBufferA, ERHIAccess::CopyDest, ERHIAccess::SRVMask));
 					}
 				}
 			}
@@ -200,7 +210,9 @@ void ProcessPixelInspectorRequests(
 						FRHICopyTextureInfo CopyInfo;
 						CopyInfo.SourcePosition = SourcePoint;
 						CopyInfo.Size = FIntVector(1, 1, 1);
+						RHICmdList.Transition(FRHITransitionInfo(DestinationBufferBCDEF, ERHIAccess::SRVMask, ERHIAccess::CopyDest));
 						RHICmdList.CopyTexture(SourceBufferB, DestinationBufferBCDEF, CopyInfo);
+						RHICmdList.Transition(FRHITransitionInfo(DestinationBufferBCDEF, ERHIAccess::CopyDest, ERHIAccess::SRVMask));
 					}
 				}
 
@@ -213,7 +225,9 @@ void ProcessPixelInspectorRequests(
 						CopyInfo.SourcePosition = SourcePoint;
 						CopyInfo.DestPosition = FIntVector(1, 0, 0);
 						CopyInfo.Size = FIntVector(1, 1, 1);
+						RHICmdList.Transition(FRHITransitionInfo(DestinationBufferBCDEF, ERHIAccess::SRVMask, ERHIAccess::CopyDest));
 						RHICmdList.CopyTexture(SourceBufferC, DestinationBufferBCDEF, CopyInfo);
+						RHICmdList.Transition(FRHITransitionInfo(DestinationBufferBCDEF, ERHIAccess::CopyDest, ERHIAccess::SRVMask));
 					}
 				}
 
@@ -226,7 +240,9 @@ void ProcessPixelInspectorRequests(
 						CopyInfo.SourcePosition = SourcePoint;
 						CopyInfo.DestPosition = FIntVector(2, 0, 0);
 						CopyInfo.Size = FIntVector(1, 1, 1);
+						RHICmdList.Transition(FRHITransitionInfo(DestinationBufferBCDEF, ERHIAccess::SRVMask, ERHIAccess::CopyDest));
 						RHICmdList.CopyTexture(SourceBufferD, DestinationBufferBCDEF, CopyInfo);
+						RHICmdList.Transition(FRHITransitionInfo(DestinationBufferBCDEF, ERHIAccess::CopyDest, ERHIAccess::SRVMask));
 					}
 				}
 
@@ -239,7 +255,9 @@ void ProcessPixelInspectorRequests(
 						CopyInfo.SourcePosition = SourcePoint;
 						CopyInfo.DestPosition = FIntVector(3, 0, 0);
 						CopyInfo.Size = FIntVector(1, 1, 1);
+						RHICmdList.Transition(FRHITransitionInfo(DestinationBufferBCDEF, ERHIAccess::SRVMask, ERHIAccess::CopyDest));
 						RHICmdList.CopyTexture(SourceBufferE, DestinationBufferBCDEF, CopyInfo);
+						RHICmdList.Transition(FRHITransitionInfo(DestinationBufferBCDEF, ERHIAccess::CopyDest, ERHIAccess::SRVMask));
 					}
 				}
 
@@ -251,7 +269,9 @@ void ProcessPixelInspectorRequests(
 						FRHICopyTextureInfo CopyInfo;
 						CopyInfo.SourcePosition = SourcePoint;
 						CopyInfo.Size = FIntVector(1, 1, 1);
+						RHICmdList.Transition(FRHITransitionInfo(DestinationBufferBCDEF, ERHIAccess::SRVMask, ERHIAccess::CopyDest));
 						RHICmdList.CopyTexture(SourceBufferF, DestinationBufferBCDEF, CopyInfo);
+						RHICmdList.Transition(FRHITransitionInfo(DestinationBufferBCDEF, ERHIAccess::CopyDest, ERHIAccess::SRVMask));
 					}
 				}
 			}
