RHI: Fix issue when using forward rendering with bindless RHI

#rnx

diff --git a/Engine/Source/Runtime/Renderer/Private/SceneRendering.cpp b/Engine/Source/Runtime/Renderer/Private/SceneRendering.cpp
--- a/Engine/Source/Runtime/Renderer/Private/SceneRendering.cpp
+++ b/Engine/Source/Runtime/Renderer/Private/SceneRendering.cpp
@@ -6420,10 +6420,10 @@ END_SHADER_PARAMETER_STRUCT()
 namespace
 {
 	template<typename ShaderType>
-	FRHIPixelShader* GetDepthResolveShader(const FViewInfo& View, int32& OutTextureIndex)
+	FRHIPixelShader* GetDepthResolveShader(const FViewInfo& View, FShaderResourceParameter& UnresolvedSurfaceParameter)
 	{
 		TShaderMapRef<ShaderType> ShaderRef(View.ShaderMap);
-		OutTextureIndex = ShaderRef->UnresolvedSurface.GetBaseIndex();
+		UnresolvedSurfaceParameter = ShaderRef->UnresolvedSurface;
 		return ShaderRef.GetPixelShader();
 	};
 
@@ -6480,21 +6480,21 @@ void AddResolveSceneDepthPass(FRDGBuilder& GraphBuilder, const FViewInfo& View,
 			TEXT("Resolving scene depth array requires support for outputting SV_RenderTargetArrayIndex from any shader."));
 
 		/** Chooses one of many ResolvePS variants */
-		auto ChoosePixelShader = [](const FViewInfo& View, bool bIsArrayResolve, int32 NumSamples, int32& OutTextureIndex) -> FRHIPixelShader*
+		auto ChoosePixelShader = [](const FViewInfo& View, bool bIsArrayResolve, int32 NumSamples, FShaderResourceParameter& UnresolvedSurfaceParameter) -> FRHIPixelShader*
 		{
 			if (LIKELY(!bIsArrayResolve))
 			{
 				switch (NumSamples)
 				{
 					case 2:
-						return GetDepthResolveShader<FResolveDepth2XPS>(View, OutTextureIndex);
+						return GetDepthResolveShader<FResolveDepth2XPS>(View, UnresolvedSurfaceParameter);
 					case 4:
-						return GetDepthResolveShader<FResolveDepth4XPS>(View, OutTextureIndex);
+						return GetDepthResolveShader<FResolveDepth4XPS>(View, UnresolvedSurfaceParameter);
 					case 8:
-						return GetDepthResolveShader<FResolveDepth8XPS>(View, OutTextureIndex);
+						return GetDepthResolveShader<FResolveDepth8XPS>(View, UnresolvedSurfaceParameter);
 					default:
 						ensureMsgf(false, TEXT("Unsupported depth resolve for samples: %i.  Dynamic loop method isn't supported on all platforms.  Please add specific case."), NumSamples);
-						return GetDepthResolveShader<FResolveDepthPS>(View, OutTextureIndex);
+						return GetDepthResolveShader<FResolveDepthPS>(View, UnresolvedSurfaceParameter);
 				}
 			}
 			else
@@ -6502,21 +6502,21 @@ void AddResolveSceneDepthPass(FRDGBuilder& GraphBuilder, const FViewInfo& View,
 				switch (NumSamples)
 				{
 					case 2:
-						return GetDepthResolveShader<FResolveDepthArray2XPS>(View, OutTextureIndex);
+						return GetDepthResolveShader<FResolveDepthArray2XPS>(View, UnresolvedSurfaceParameter);
 					case 4:
-						return GetDepthResolveShader<FResolveDepthArray4XPS>(View, OutTextureIndex);
+						return GetDepthResolveShader<FResolveDepthArray4XPS>(View, UnresolvedSurfaceParameter);
 					case 8:
-						return GetDepthResolveShader<FResolveDepthArray8XPS>(View, OutTextureIndex);
+						return GetDepthResolveShader<FResolveDepthArray8XPS>(View, UnresolvedSurfaceParameter);
 					default:
 						ensureMsgf(false, TEXT("Unsupported depth resolve for samples: %i (texture array case).  Dynamic loop method isn't supported on all platforms.  Please add specific case."), NumSamples);
-						return GetDepthResolveShader<FResolveDepthPS>(View, OutTextureIndex);
+						return GetDepthResolveShader<FResolveDepthPS>(View, UnresolvedSurfaceParameter);
 				}
 			}
 			// unreachable
 		};
 
-		int32 TextureIndex = -1;
-		FRHIPixelShader* ResolvePixelShader = ChoosePixelShader(View, bArrayResolve, NumSamples, TextureIndex);
+		FShaderResourceParameter UnresolvedSurfaceParameter;
+		FRHIPixelShader* ResolvePixelShader = ChoosePixelShader(View, bArrayResolve, NumSamples, UnresolvedSurfaceParameter);
 
 		TShaderRef<FResolveVS> ResolveVertexShader;
 		GraphicsPSOInit.BoundShaderState.VertexDeclarationRHI = GEmptyVertexDeclaration.VertexDeclarationRHI;
@@ -6528,7 +6528,7 @@ void AddResolveSceneDepthPass(FRDGBuilder& GraphBuilder, const FViewInfo& View,
 		RHICmdList.SetBlendFactor(FLinearColor::White);
 
 		FRHIBatchedShaderParameters& BatchedParameters = RHICmdList.GetScratchShaderParameters();
-		BatchedParameters.SetShaderTexture(TextureIndex, SourceTextureRHI);
+		SetTextureParameter(BatchedParameters, UnresolvedSurfaceParameter, SourceTextureRHI);
 		RHICmdList.SetBatchedShaderParameters(ResolvePixelShader, BatchedParameters);
 
 		SetShaderParametersLegacyVS(RHICmdList, ResolveVertexShader, ResolveRect, ResolveRect, DepthExtent.X, DepthExtent.Y);
