Path Tracer: Extend light and decal grids to cover all three axes

This lets us pick the optimal projection per point, which results in fewer elements  to consider in general. The same change has been applied to decal grid and both the RT and PT light grids.

Tweak the color map a bit to make it brighter and add the color mapping to the Visualization header

diff --git a/Engine/Plugins/Experimental/GPULightmass/Source/GPULightmass/Private/LightmapRayTracing.h b/Engine/Plugins/Experimental/GPULightmass/Source/GPULightmass/Private/LightmapRayTracing.h
--- a/Engine/Plugins/Experimental/GPULightmass/Source/GPULightmass/Private/LightmapRayTracing.h
+++ b/Engine/Plugins/Experimental/GPULightmass/Source/GPULightmass/Private/LightmapRayTracing.h
@@ -20,7 +20,7 @@ BEGIN_SHADER_PARAMETER_STRUCT(FPathTracingLightGrid, RENDERER_API)
 	SHADER_PARAMETER(uint32, SceneInfiniteLightCount)
 	SHADER_PARAMETER(FVector3f, SceneLightsTranslatedBoundMin)
 	SHADER_PARAMETER(FVector3f, SceneLightsTranslatedBoundMax)
-	SHADER_PARAMETER_RDG_TEXTURE(Texture2D, LightGrid)
+	SHADER_PARAMETER_RDG_TEXTURE(Texture2DArray<uint>, LightGrid)
 	SHADER_PARAMETER_RDG_BUFFER_SRV(Buffer<uint>, LightGridData)
 	SHADER_PARAMETER(unsigned, LightGridResolution)
 	SHADER_PARAMETER(unsigned, LightGridMaxCount)
diff --git a/Engine/Shaders/Private/PathTracing/Light/PathTracingLightGrid.ush b/Engine/Shaders/Private/PathTracing/Light/PathTracingLightGrid.ush
--- a/Engine/Shaders/Private/PathTracing/Light/PathTracingLightGrid.ush
+++ b/Engine/Shaders/Private/PathTracing/Light/PathTracingLightGrid.ush
@@ -12,22 +12,13 @@
 uint SceneInfiniteLightCount;
 float3 SceneLightsTranslatedBoundMin;
 float3 SceneLightsTranslatedBoundMax;
-Texture2D<uint> LightGrid;
+Texture2DArray<uint> LightGrid;
 Buffer<uint> LightGridData;
 uint LightGridResolution;
 uint LightGridMaxCount;
 int LightGridAxis;
 #endif
 
-
-// TODO: unify with similar function in other visualizations
-float3 HashToColor(uint H)
-{
-	return 0.5 + 0.3 * cos(2 * PI * ((H & 0xFFFFFF) * 5.96046447754e-08 + float3(0.0, 0.333, 0.6666)));
-}
-
-
-
 struct FLightLoopCount {
 	uint NumLights;
 	uint NumMISLights;
@@ -49,15 +40,27 @@ FLightLoopCount LightGridLookup(float3 TranslatedWorldPos)
 	{
 		float3 P = TranslatedWorldPos - SceneLightsTranslatedBoundMin;
 		float3 D = SceneLightsTranslatedBoundMax - SceneLightsTranslatedBoundMin;
-		int2 UV = 0;
-		switch (LightGridAxis)
+		int2 UVx = int2(floor(LightGridResolution * P.yz / D.yz));
+		int2 UVy = int2(floor(LightGridResolution * P.xz / D.xz));
+		int2 UVz = int2(floor(LightGridResolution * P.xy / D.xy));
+
+		uint LightGridNx = LightGrid.Load(int4(UVx, 0, 0));
+		uint LightGridNy = LightGrid.Load(int4(UVy, 1, 0));
+		uint LightGridNz = LightGrid.Load(int4(UVz, 2, 0));
+		uint LightGridN = LightGridNx;
+		uint LightGridOffset = LightGridMaxCount * (0 + 3 * (UVx.x + UVx.y * LightGridResolution));
+		if ((LightGridAxis < 0 && (LightGridNy & PATHTRACER_LIGHT_GRID_LIGHT_COUNT_MASK) < (LightGridN & PATHTRACER_LIGHT_GRID_LIGHT_COUNT_MASK)) || LightGridAxis == 1)
 		{
-			case 0: UV = int2(floor(LightGridResolution * P.yz / D.yz)); break;
-			case 1: UV = int2(floor(LightGridResolution * P.xz / D.xz)); break;
-			case 2: UV = int2(floor(LightGridResolution * P.xy / D.xy)); break;
+			LightGridN = LightGridNy;
+			LightGridOffset = LightGridMaxCount * (1 + 3 * (UVy.x + UVy.y * LightGridResolution));
 		}
-		uint LightGridN = LightGrid.Load(int3(UV, 0));
-		Result.LightGridOffset = LightGridMaxCount * (UV.x + UV.y * LightGridResolution);
+		if ((LightGridAxis < 0 && (LightGridNz & PATHTRACER_LIGHT_GRID_LIGHT_COUNT_MASK) < (LightGridN & PATHTRACER_LIGHT_GRID_LIGHT_COUNT_MASK)) || LightGridAxis == 2)
+		{
+			LightGridN = LightGridNz;
+			LightGridOffset = LightGridMaxCount * (2 + 3 * (UVz.x + UVz.y * LightGridResolution));
+		}
+
+		Result.LightGridOffset = LightGridOffset;
 		Result.NumLights += LightGridN & PATHTRACER_LIGHT_GRID_LIGHT_COUNT_MASK;
 		if ((LightGridN & PATHTRACER_LIGHT_GRID_SINGULAR_MASK) == 0)
 		{
@@ -85,7 +88,7 @@ float3 LightGridVisualize(FLightLoopCount LoopCount, int VisualizeMode)
 				{
 					H = StrongIntegerHash(H + LightGridData[LoopCount.LightGridOffset + Index]);
 				}
-				return HashToColor(H);
+				return HUEtoRGB((H & 0xFFFFFF) * 5.96046447754e-08);
 			}
 			case 3:
 			{
@@ -103,6 +106,15 @@ float3 LightGridVisualize(FLightLoopCount LoopCount, int VisualizeMode)
 					return float3(0.2, 1, 0.2); // green - only point lights
 				}
 			}
+			case 4:
+			{
+				// color by chosen axis
+				int Axis = (LoopCount.LightGridOffset / LightGridMaxCount) % 3;
+				if (Axis == 0) return float3(0.9, 0.2, 0.2);
+				if (Axis == 1) return float3(0.2, 0.9, 0.2);
+				if (Axis == 2) return float3(0.2, 0.2, 0.9);
+				return float3(0.8, 0.2, 0.8);
+			}
 			default: 
 			{
 				// default mode - color by light count
@@ -113,7 +125,7 @@ float3 LightGridVisualize(FLightLoopCount LoopCount, int VisualizeMode)
 				}
 				float Max = LightGridMaxCount;
 				float t = saturate(float(N) / Max);
-				return t * saturate(2.0 - abs(float3(4, 2, 0) - 4.0 * t));
+				return BlueGreenRedColorMap(t);
 			}
 		}
 	}
diff --git a/Engine/Shaders/Private/PathTracing/PathTracingBuildLightGrid.usf b/Engine/Shaders/Private/PathTracing/PathTracingBuildLightGrid.usf
--- a/Engine/Shaders/Private/PathTracing/PathTracingBuildLightGrid.usf
+++ b/Engine/Shaders/Private/PathTracing/PathTracingBuildLightGrid.usf
@@ -12,7 +12,7 @@
 uint SceneInfiniteLightCount;
 float3 SceneLightsTranslatedBoundMin;
 float3 SceneLightsTranslatedBoundMax;
-RWTexture2D<uint> RWLightGrid;
+RWTexture2DArray<uint> RWLightGrid;
 RWBuffer<uint> RWLightGridData;
 uint LightGridResolution;
 uint LightGridMaxCount;
@@ -20,19 +20,25 @@ int LightGridAxis;
 
 
 [numthreads(THREADGROUPSIZE_X, THREADGROUPSIZE_Y, 1)]
-void PathTracingBuildLightGridCS(uint2 DispatchThreadId : SV_DispatchThreadID)
+void PathTracingBuildLightGridCS(uint3 DispatchThreadId : SV_DispatchThreadID)
 {
-	if (any(DispatchThreadId >= LightGridResolution))
+	if (any(DispatchThreadId.xy >= LightGridResolution) || DispatchThreadId.z >= 3)
 	{
 		return;
 	}
 	// figure out dimension of the 3d grid and the current ID to be filled
 	uint3 VoxelId = 0, VoxelRes = 1;
-	switch (LightGridAxis)
+	int Axis = DispatchThreadId.z;
+	if (LightGridAxis != Axis && LightGridAxis >= 0 && LightGridAxis < 3)
 	{
-		case 0: VoxelId.yz = DispatchThreadId; VoxelRes.yz = LightGridResolution; break;
-		case 1: VoxelId.xz = DispatchThreadId; VoxelRes.xz = LightGridResolution; break;
-		case 2: VoxelId.xy = DispatchThreadId; VoxelRes.xy = LightGridResolution; break;
+		RWLightGrid[DispatchThreadId] = 0;
+		return;
+	}
+	switch (Axis)
+	{
+		case 0: VoxelId.yz = DispatchThreadId.xy; VoxelRes.yz = LightGridResolution; break;
+		case 1: VoxelId.xz = DispatchThreadId.xy; VoxelRes.xz = LightGridResolution; break;
+		case 2: VoxelId.xy = DispatchThreadId.xy; VoxelRes.xy = LightGridResolution; break;
 	}
 	// get bounding box of current voxel
 	float3 Lo = lerp(SceneLightsTranslatedBoundMin, SceneLightsTranslatedBoundMax, float3(VoxelId + 0) / float3(VoxelRes));
@@ -40,7 +46,7 @@ void PathTracingBuildLightGridCS(uint2 DispatchThreadId : SV_DispatchThreadID)
 
 	uint NumVisible = 0;
 	uint MaxVisible = min(LightGridMaxCount, RAY_TRACING_LIGHT_COUNT_MAXIMUM - SceneInfiniteLightCount);
-	uint Offset = LightGridMaxCount * (DispatchThreadId.x + LightGridResolution * DispatchThreadId.y);
+	uint Offset = LightGridMaxCount * (Axis + 3 * (DispatchThreadId.x + LightGridResolution * DispatchThreadId.y));
 	bool AllSingular = true;
 	for (uint LightIndex = SceneInfiniteLightCount; LightIndex < SceneLightCount; LightIndex++)
 	{
diff --git a/Engine/Shaders/Private/RayTracing/RayTracingBuildDecalGrid.usf b/Engine/Shaders/Private/RayTracing/RayTracingBuildDecalGrid.usf
--- a/Engine/Shaders/Private/RayTracing/RayTracingBuildDecalGrid.usf
+++ b/Engine/Shaders/Private/RayTracing/RayTracingBuildDecalGrid.usf
@@ -14,23 +14,29 @@ uint DecalGridResolution;
 uint DecalGridMaxCount;
 uint DecalGridAxis;
 
-RWTexture2D<uint> RWDecalGrid;
+RWTexture2DArray<uint> RWDecalGrid;
 RWBuffer<uint> RWDecalGridData;
 
 [numthreads(THREADGROUPSIZE_X, THREADGROUPSIZE_Y, 1)]
-void BuildDecalGridCS(uint2 DispatchThreadId : SV_DispatchThreadID)
+void BuildDecalGridCS(uint3 DispatchThreadId : SV_DispatchThreadID)
 {
-	if (any(DispatchThreadId >= DecalGridResolution))
+	if (any(DispatchThreadId >= DecalGridResolution) || DispatchThreadId.z >= 3)
 	{
 		return;
 	}
 	// figure out dimension of the 3d grid and the current ID to be filled
 	uint3 VoxelId = 0, VoxelRes = 1;
-	switch (DecalGridAxis)
+	int Axis = DispatchThreadId.z;
+	if (DecalGridAxis != Axis && DecalGridAxis < 3)
 	{
-	case 0: VoxelId.yz = DispatchThreadId; VoxelRes.yz = DecalGridResolution; break;
-	case 1: VoxelId.xz = DispatchThreadId; VoxelRes.xz = DecalGridResolution; break;
-	case 2: VoxelId.xy = DispatchThreadId; VoxelRes.xy = DecalGridResolution; break;
+		RWDecalGrid[DispatchThreadId] = 0;
+		return;
+	}
+	switch (Axis)
+	{
+		case 0: VoxelId.yz = DispatchThreadId.xy; VoxelRes.yz = DecalGridResolution; break;
+		case 1: VoxelId.xz = DispatchThreadId.xy; VoxelRes.xz = DecalGridResolution; break;
+		case 2: VoxelId.xy = DispatchThreadId.xy; VoxelRes.xy = DecalGridResolution; break;
 	}
 	// get bounding box of current voxel
 	float3 Lo = lerp(SceneDecalsTranslatedBoundMin, SceneDecalsTranslatedBoundMax, float3(VoxelId + 0) / float3(VoxelRes));
@@ -38,7 +44,7 @@ void BuildDecalGridCS(uint2 DispatchThreadId : SV_DispatchThreadID)
 
 	uint NumVisible = 0;
 	uint MaxVisible = min(DecalGridMaxCount, RAY_TRACING_DECAL_COUNT_MAXIMUM);
-	uint Offset = DecalGridMaxCount * (DispatchThreadId.x + DecalGridResolution * DispatchThreadId.y);
+	uint Offset = DecalGridMaxCount * (Axis + 3 * (DispatchThreadId.x + DecalGridResolution * DispatchThreadId.y));
 
 	for (uint DecalIndex = 0; DecalIndex < SceneDecalCount; DecalIndex++)
 	{
diff --git a/Engine/Shaders/Private/RayTracing/RayTracingBuildLightGrid.usf b/Engine/Shaders/Private/RayTracing/RayTracingBuildLightGrid.usf
--- a/Engine/Shaders/Private/RayTracing/RayTracingBuildLightGrid.usf
+++ b/Engine/Shaders/Private/RayTracing/RayTracingBuildLightGrid.usf
@@ -10,26 +10,32 @@ StructuredBuffer<FRTLightingData> SceneLights;
 uint SceneInfiniteLightCount;
 float3 SceneLightsTranslatedBoundMin;
 float3 SceneLightsTranslatedBoundMax;
-RWTexture2D<uint> RWLightGrid;
+RWTexture2DArray<uint> RWLightGrid;
 RWBuffer<uint> RWLightGridData;
 uint LightGridResolution;
 uint LightGridMaxCount;
 int LightGridAxis;
 
 [numthreads(THREADGROUPSIZE_X, THREADGROUPSIZE_Y, 1)]
-void RayTracingBuildLightGridCS(uint2 DispatchThreadId : SV_DispatchThreadID)
+void RayTracingBuildLightGridCS(uint3 DispatchThreadId : SV_DispatchThreadID)
 {
-	if (any(DispatchThreadId >= LightGridResolution))
+	if (any(DispatchThreadId >= LightGridResolution) || DispatchThreadId.z >= 3)
 	{
 		return;
 	}
 	// figure out dimension of the 3d grid and the current ID to be filled
 	uint3 VoxelId = 0, VoxelRes = 1;
-	switch (LightGridAxis)
+	int Axis = DispatchThreadId.z;
+	if (Axis != LightGridAxis && LightGridAxis >= 0 && LightGridAxis < 3)
 	{
-		case 0: VoxelId.yz = DispatchThreadId; VoxelRes.yz = LightGridResolution; break;
-		case 1: VoxelId.xz = DispatchThreadId; VoxelRes.xz = LightGridResolution; break;
-		case 2: VoxelId.xy = DispatchThreadId; VoxelRes.xy = LightGridResolution; break;
+		RWLightGrid[DispatchThreadId] = 0;
+		return;
+	}
+	switch (Axis)
+	{
+		case 0: VoxelId.yz = DispatchThreadId.xy; VoxelRes.yz = LightGridResolution; break;
+		case 1: VoxelId.xz = DispatchThreadId.xy; VoxelRes.xz = LightGridResolution; break;
+		case 2: VoxelId.xy = DispatchThreadId.xy; VoxelRes.xy = LightGridResolution; break;
 	}
 	// get bounding box of current voxel
 	float3 Lo = lerp(SceneLightsTranslatedBoundMin, SceneLightsTranslatedBoundMax, float3(VoxelId + 0) / float3(VoxelRes));
diff --git a/Engine/Shaders/Private/RayTracing/RayTracingDecalGrid.ush b/Engine/Shaders/Private/RayTracing/RayTracingDecalGrid.ush
--- a/Engine/Shaders/Private/RayTracing/RayTracingDecalGrid.ush
+++ b/Engine/Shaders/Private/RayTracing/RayTracingDecalGrid.ush
@@ -27,15 +27,29 @@ FDecalLoopCount DecalGridLookup(float3 TranslatedWorldPos)
 	{
 		float3 P = TranslatedWorldPos - RayTracingDecals.TranslatedBoundMin;
 		float3 D = RayTracingDecals.TranslatedBoundMax - RayTracingDecals.TranslatedBoundMin;
-		int2 UV = 0;
-		switch (RayTracingDecals.GridAxis)
+
+		int2 UVx = int2(floor(RayTracingDecals.GridResolution * P.yz / D.yz));
+		int2 UVy = int2(floor(RayTracingDecals.GridResolution * P.xz / D.xz));
+		int2 UVz = int2(floor(RayTracingDecals.GridResolution * P.xy / D.xy));
+
+		uint DecalNx = RayTracingDecals.Grid.Load(int4(UVx, 0, 0));
+		uint DecalNy = RayTracingDecals.Grid.Load(int4(UVy, 1, 0));
+		uint DecalNz = RayTracingDecals.Grid.Load(int4(UVz, 2, 0));
+		uint DecalN = DecalNx;
+		uint DecalOffset = RayTracingDecals.GridMaxCount * (0 + 3 * (UVx.x + UVx.y * RayTracingDecals.GridResolution));
+		if ((RayTracingDecals.GridAxis > 2 && DecalNy < DecalN) || RayTracingDecals.GridAxis == 1)
+		{
+			DecalN = DecalNy;
+			DecalOffset = RayTracingDecals.GridMaxCount * (1 + 3 * (UVy.x + UVy.y * RayTracingDecals.GridResolution));
+		}
+		if ((RayTracingDecals.GridAxis > 2 && DecalNz < DecalN) || RayTracingDecals.GridAxis == 2)
 		{
-		case 0: UV = int2(floor(RayTracingDecals.GridResolution * P.yz / D.yz)); break;
-		case 1: UV = int2(floor(RayTracingDecals.GridResolution * P.xz / D.xz)); break;
-		case 2: UV = int2(floor(RayTracingDecals.GridResolution * P.xy / D.xy)); break;
+			DecalN = DecalNz;
+			DecalOffset = RayTracingDecals.GridMaxCount * (2 + 3 * (UVz.x + UVz.y * RayTracingDecals.GridResolution));
 		}
-		Result.DecalGridOffset = RayTracingDecals.GridMaxCount * (UV.x + UV.y * RayTracingDecals.GridResolution);
-		Result.NumDecals += RayTracingDecals.Grid.Load(int3(UV, 0));
+
+		Result.DecalGridOffset = DecalOffset;
+		Result.NumDecals += DecalN;
 	}
 	return Result;
 }
@@ -59,6 +73,15 @@ float3 DecalGridVisualize(FDecalLoopCount LoopCount, int VisualizeMode)
 			}
 			return IntToColor(H);
 		}
+		case 3:
+		{
+			// color by chosen axis
+			int Axis = (LoopCount.DecalGridOffset / RayTracingDecals.GridMaxCount) % 3;
+			if (Axis == 0) return float3(0.9, 0.2, 0.2);
+			if (Axis == 1) return float3(0.2, 0.9, 0.2);
+			if (Axis == 2) return float3(0.2, 0.2, 0.9);
+			return float3(0.8, 0.2, 0.8); // error
+		}
 		default:
 		{
 			// default mode - color by decal count
@@ -69,7 +92,7 @@ float3 DecalGridVisualize(FDecalLoopCount LoopCount, int VisualizeMode)
 			}
 			float Max = RayTracingDecals.GridMaxCount;
 			float t = saturate(float(N) / Max);
-			return t * saturate(2.0 - abs(float3(4, 2, 0) - 4.0 * t));
+			return BlueGreenRedColorMap(t);
 		}
 		}
 	}
diff --git a/Engine/Shaders/Private/RayTracing/RayTracingLightCullingCommon.ush b/Engine/Shaders/Private/RayTracing/RayTracingLightCullingCommon.ush
--- a/Engine/Shaders/Private/RayTracing/RayTracingLightCullingCommon.ush
+++ b/Engine/Shaders/Private/RayTracing/RayTracingLightCullingCommon.ush
@@ -6,6 +6,8 @@
 
 #pragma once
 
+#include "../Visualization.ush"
+
 // *************** LIGHT CULLING *********************************
 // Use a 2D grid built along the longest axes of the bounding boxes of all finite light sources.
 
@@ -27,15 +29,26 @@ struct FCulledLightList
 		{
 			float3 P = TranslatedWorldPos - SceneMin;
 			float3 D = SceneMax - SceneMin;
-			int2 UV = 0;
-			switch (RaytracingLightGridData.LightGridAxis)
+			int2 UVx = int2(floor(RaytracingLightGridData.LightGridResolution * P.yz / D.yz));
+			int2 UVy = int2(floor(RaytracingLightGridData.LightGridResolution * P.xz / D.xz));
+			int2 UVz = int2(floor(RaytracingLightGridData.LightGridResolution * P.xy / D.xy));
+
+			uint LightGridNx = RaytracingLightGridData.LightGrid.Load(int4(UVx, 0, 0));
+			uint LightGridNy = RaytracingLightGridData.LightGrid.Load(int4(UVy, 1, 0));
+			uint LightGridNz = RaytracingLightGridData.LightGrid.Load(int4(UVz, 2, 0));
+			uint LightGridN = LightGridNx;
+			uint LightGridOffset = RaytracingLightGridData.LightGridMaxCount * (0 + 3 * (UVx.x + UVx.y * RaytracingLightGridData.LightGridResolution));
+			if ((RaytracingLightGridData.LightGridAxis < 0 && LightGridNy < LightGridN) || RaytracingLightGridData.LightGridAxis == 1)
+			{
+				LightGridN = LightGridNy;
+				LightGridOffset = RaytracingLightGridData.LightGridMaxCount * (1 + 3 * (UVy.x + UVy.y * RaytracingLightGridData.LightGridResolution));
+			}
+			if ((RaytracingLightGridData.LightGridAxis < 0 && LightGridNz < LightGridN) || RaytracingLightGridData.LightGridAxis == 2)
 			{
-				case 0: UV = int2(floor(RaytracingLightGridData.LightGridResolution * P.yz / D.yz)); break;
-				case 1: UV = int2(floor(RaytracingLightGridData.LightGridResolution * P.xz / D.xz)); break;
-				case 2: UV = int2(floor(RaytracingLightGridData.LightGridResolution * P.xy / D.xy)); break;
+				LightGridN = LightGridNz;
+				LightGridOffset = RaytracingLightGridData.LightGridMaxCount * (2 + 3 * (UVz.x + UVz.y * RaytracingLightGridData.LightGridResolution));
 			}
-			uint LightGridN = RaytracingLightGridData.LightGrid.Load(int3(UV, 0));
-			Result.LightGridOffset = RaytracingLightGridData.LightGridMaxCount * (UV.x + UV.y * RaytracingLightGridData.LightGridResolution);
+			Result.LightGridOffset = LightGridOffset;
 			Result.NumLights += LightGridN;
 		}
 		return Result;
@@ -88,7 +101,7 @@ struct FCulledLightList
 							H ^= H >> 17;
 						}
 					}
-					return 0.5 + 0.3 * cos(2 * PI * ((H & 0xFFFFFF) * 5.96046447754e-08 + float3(0.0, 0.333, 0.6666)));
+					return HUEtoRGB((H & 0xFFFFFF) * 5.96046447754e-08);
 				}
 				default: 
 				{
@@ -100,7 +113,7 @@ struct FCulledLightList
 					}
 					float Max = RaytracingLightGridData.LightGridMaxCount;
 					float t = saturate(float(N) / Max);
-					return t * saturate(2.0 - abs(float3(4, 2, 0) - 4.0 * t));
+					return BlueGreenRedColorMap(t);
 				}
 			}
 		}
diff --git a/Engine/Shaders/Private/Visualization.ush b/Engine/Shaders/Private/Visualization.ush
--- a/Engine/Shaders/Private/Visualization.ush
+++ b/Engine/Shaders/Private/Visualization.ush
@@ -48,3 +48,8 @@ float3 GreenToRedTurbo(float s)
 {
 	return ColorMapTurbo(lerp(0.5f, 1.0f, saturate(s)));
 }
+
+float3 BlueGreenRedColorMap(float t)
+{
+	return (0.75f * t + 0.25f) * saturate(2.0 - abs(float3(4, 2, 0) - 4.0 * t));
+}
diff --git a/Engine/Source/Runtime/Renderer/Private/PathTracing.cpp b/Engine/Source/Runtime/Renderer/Private/PathTracing.cpp
--- a/Engine/Source/Runtime/Renderer/Private/PathTracing.cpp
+++ b/Engine/Source/Runtime/Renderer/Private/PathTracing.cpp
@@ -70,7 +70,7 @@ TAutoConsoleVariable<int32> CVarPathTracingCompactionDepth(
 TAutoConsoleVariable<bool> CVarPathTracingTraceOpaqueFirst(
 	TEXT("r.PathTracing.TraceOpaqueFirst"),
 	true,
-	TEXT("Trace opaque geometry before translucent geometry. This allows the path tracer to setup a correct depth for DepthFade based effects (default: 0)"),
+	TEXT("Trace opaque geometry before translucent geometry. This allows the path tracer to setup a correct depth for DepthFade based effects (default: 1)"),
 	ECVF_RenderThreadSafe
 );
 
@@ -420,10 +420,10 @@ TAutoConsoleVariable<float> CVarPathTracingMeshDecalBias(
 
 TAutoConsoleVariable<bool> CVarPathTracingLightFunctionColor(
 	TEXT("r.PathTracing.LightFunctionColor"),
-	false,
-	TEXT("Enables light functions to be colored instead of greyscale (default = false)\n")
-	TEXT("0: off (default)\n")
-	TEXT("1: on (light function material output is used directly instead of converting to greyscale)\n"),
+	true,
+	TEXT("Enables colored light function output (default = true)\n")
+	TEXT("0: off (light function material output is converted to grayscale)\n")
+	TEXT("1: on (light function material output is used directly)\n"),
 	ECVF_RenderThreadSafe
 );
 
@@ -986,7 +986,7 @@ BEGIN_SHADER_PARAMETER_STRUCT(FPathTracingLightGrid, RENDERER_API)
 	SHADER_PARAMETER(uint32, SceneInfiniteLightCount)
 	SHADER_PARAMETER(FVector3f, SceneLightsTranslatedBoundMin)
 	SHADER_PARAMETER(FVector3f, SceneLightsTranslatedBoundMax)
-	SHADER_PARAMETER_RDG_TEXTURE(Texture2D, LightGrid)
+	SHADER_PARAMETER_RDG_TEXTURE(Texture2DArray<uint>, LightGrid)
 	SHADER_PARAMETER_RDG_BUFFER_SRV(Buffer<uint>, LightGridData)
 	SHADER_PARAMETER(unsigned, LightGridResolution)
 	SHADER_PARAMETER(unsigned, LightGridMaxCount)
@@ -1016,7 +1016,7 @@ class FPathTracingBuildLightGridCS : public FGlobalShader
 		SHADER_PARAMETER_RDG_BUFFER_SRV(StructuredBuffer<FPathTracingLight>, SceneLights)
 		SHADER_PARAMETER(uint32, SceneLightCount)
 		SHADER_PARAMETER_STRUCT_INCLUDE(FPathTracingLightGrid, LightGridParameters)
-		SHADER_PARAMETER_RDG_TEXTURE_UAV(RWTexture2D, RWLightGrid)
+		SHADER_PARAMETER_RDG_TEXTURE_UAV(RWTexture2DArray<uint>, RWLightGrid)
 		SHADER_PARAMETER_RDG_BUFFER_UAV(RWBuffer<uint>, RWLightGridData)
 	END_SHADER_PARAMETER_STRUCT()
 };
@@ -2558,7 +2558,7 @@ RENDERER_API void PrepareLightGrid(FRDGBuilder& GraphBuilder, ERHIFeatureLevel::
 			LightGridParameters->SceneLightsTranslatedBoundMax = FVector3f::Max(LightGridParameters->SceneLightsTranslatedBoundMax, Hi);
 		}
 
-		const uint32 Resolution = FMath::RoundUpToPowerOfTwo(CVarPathTracingLightGridResolution.GetValueOnRenderThread());
+		const uint32 Resolution = FMath::Clamp(CVarPathTracingLightGridResolution.GetValueOnRenderThread(), 1, 2048);
 		const uint32 MaxCount = FMath::Clamp(
 			CVarPathTracingLightGridMaxCount.GetValueOnRenderThread(),
 			1,
@@ -2568,30 +2568,13 @@ RENDERER_API void PrepareLightGrid(FRDGBuilder& GraphBuilder, ERHIFeatureLevel::
 		LightGridParameters->LightGridMaxCount = MaxCount;
 
 		LightGridParameters->LightGridAxis = CVarPathTracingLightGridAxis.GetValueOnRenderThread();
-		if (LightGridParameters->LightGridAxis < 0 || LightGridParameters->LightGridAxis > 2)
-		{
-			// pick the shortest axis if user did not explicitly pick an axis
-			FVector3f Diag = LightGridParameters->SceneLightsTranslatedBoundMax - LightGridParameters->SceneLightsTranslatedBoundMin;
-			if (Diag.X < Diag.Y && Diag.X < Diag.Z)
-			{
-				LightGridParameters->LightGridAxis = 0;
-			}
-			else if (Diag.Y < Diag.Z)
-			{
-				LightGridParameters->LightGridAxis = 1;
-			}
-			else
-			{
-				LightGridParameters->LightGridAxis = 2;
-			}
-		}
 		FPathTracingBuildLightGridCS::FParameters* LightGridPassParameters = GraphBuilder.AllocParameters< FPathTracingBuildLightGridCS::FParameters>();
 
-		FRDGTextureDesc LightGridDesc = FRDGTextureDesc::Create2D(
+		FRDGTextureDesc LightGridDesc = FRDGTextureDesc::Create2DArray(
 			FIntPoint(Resolution, Resolution),
 			PF_R32_UINT,
 			FClearValueBinding::None,
-			TexCreate_ShaderResource | TexCreate_UAV);
+			TexCreate_ShaderResource | TexCreate_UAV, 3);
 
 		// jhoerner TODO 9/30/2022: Hack to work around MGPU resource transition architectural bug in RDG.  Mask PathTracer.LightGrid texture
 		// to only be present on current GPU.  The bug is that RDG batches transitions, but the execution of batched transitions uses the
@@ -2616,7 +2599,7 @@ RENDERER_API void PrepareLightGrid(FRDGBuilder& GraphBuilder, ERHIFeatureLevel::
 			LightGridDataFormat = PF_R16_UINT;
 			LightGridDataNumBytes = sizeof(uint16);
 		}
-		FRDGBufferDesc LightGridDataDesc = FRDGBufferDesc::CreateBufferDesc(LightGridDataNumBytes, MaxCount * Resolution * Resolution);
+		FRDGBufferDesc LightGridDataDesc = FRDGBufferDesc::CreateBufferDesc(LightGridDataNumBytes, 3 * MaxCount * Resolution * Resolution);
 		FRDGBuffer* LightGridData = GraphBuilder.CreateBuffer(LightGridDataDesc, TEXT("PathTracer.LightGridData"));
 		LightGridPassParameters->RWLightGridData = GraphBuilder.CreateUAV(LightGridData, LightGridDataFormat);
 		LightGridPassParameters->LightGridParameters = *LightGridParameters;
@@ -2629,7 +2612,7 @@ RENDERER_API void PrepareLightGrid(FRDGBuilder& GraphBuilder, ERHIFeatureLevel::
 			RDG_EVENT_NAME("Light Grid Create (%u lights)", NumFiniteLights),
 			ComputeShader,
 			LightGridPassParameters,
-			FComputeShaderUtils::GetGroupCount(FIntPoint(Resolution, Resolution), FComputeShaderUtils::kGolden2DGroupSize));
+			FComputeShaderUtils::GetGroupCount(FIntVector(Resolution, Resolution, 3), FIntVector(FComputeShaderUtils::kGolden2DGroupSize, FComputeShaderUtils::kGolden2DGroupSize, 1)));
 
 		// hookup to the actual rendering pass
 		LightGridParameters->LightGrid = LightGridTexture;
diff --git a/Engine/Source/Runtime/Renderer/Private/RayTracing/RayTracingDecals.cpp b/Engine/Source/Runtime/Renderer/Private/RayTracing/RayTracingDecals.cpp
--- a/Engine/Source/Runtime/Renderer/Private/RayTracing/RayTracingDecals.cpp
+++ b/Engine/Source/Runtime/Renderer/Private/RayTracing/RayTracingDecals.cpp
@@ -141,7 +141,7 @@ class FRayTracingBuildDecalGridCS : public FGlobalShader
 		SHADER_PARAMETER(unsigned, DecalGridMaxCount)
 		SHADER_PARAMETER(int, DecalGridAxis)
 
-		SHADER_PARAMETER_RDG_TEXTURE_UAV(RWTexture2D, RWDecalGrid)
+		SHADER_PARAMETER_RDG_TEXTURE_UAV(RWTexture2DArray<uint>, RWDecalGrid)
 		SHADER_PARAMETER_RDG_BUFFER_UAV(RWBuffer<uint>, RWDecalGridData)
 	END_SHADER_PARAMETER_STRUCT()
 };
@@ -399,35 +399,17 @@ void BuildDecalGrid(FRDGBuilder& GraphBuilder, uint32 NumDecals, FRDGBufferSRVRe
 	// if we have some decals -- build a decal grid
 	if (NumDecals > 0)
 	{
-		const uint32 Resolution = FMath::RoundUpToPowerOfTwo(CVarDecalGridResolution.GetValueOnRenderThread());
+		const uint32 Resolution = FMath::Clamp(CVarDecalGridResolution.GetValueOnRenderThread(), 1, 2048);
 		const uint32 MaxCount = FMath::Clamp(CVarDecalGridMaxCount.GetValueOnRenderThread(), 1, FMath::Min((int32)NumDecals, RAY_TRACING_DECAL_COUNT_MAXIMUM));
 		OutParameters.GridResolution = Resolution;
 		OutParameters.GridMaxCount = MaxCount;
 		OutParameters.GridAxis = CVarDecalGridAxis.GetValueOnRenderThread();
-
-		if (OutParameters.GridAxis < 0 || OutParameters.GridAxis > 2)
-		{
-			// pick the shortest axis if user did not explicitly pick an axis
-			FVector3f Diag = OutParameters.TranslatedBoundMax - OutParameters.TranslatedBoundMin;
-			if (Diag.X < Diag.Y && Diag.X < Diag.Z)
-			{
-				OutParameters.GridAxis = 0;
-			}
-			else if (Diag.Y < Diag.Z)
-			{
-				OutParameters.GridAxis = 1;
-			}
-			else
-			{
-				OutParameters.GridAxis = 2;
-			}
-		}
-
-		FRDGTextureDesc DecalGridDesc = FRDGTextureDesc::Create2D(
+		
+		FRDGTextureDesc DecalGridDesc = FRDGTextureDesc::Create2DArray(
 			FIntPoint(Resolution, Resolution),
 			PF_R32_UINT,
 			FClearValueBinding::None,
-			TexCreate_ShaderResource | TexCreate_UAV);
+			TexCreate_ShaderResource | TexCreate_UAV, 3);
 		FRDGTexture* DecalGridTexture = GraphBuilder.CreateTexture(DecalGridDesc, TEXT("PathTracer.DecalGrid"), ERDGTextureFlags::None);
 
 		EPixelFormat DecalGridDataFormat = PF_R32_UINT;
@@ -442,7 +424,7 @@ void BuildDecalGrid(FRDGBuilder& GraphBuilder, uint32 NumDecals, FRDGBufferSRVRe
 			DecalGridDataFormat = PF_R16_UINT;
 			DecalGridDataNumBytes = sizeof(uint16);
 		}
-		FRDGBufferDesc DecalGridDataDesc = FRDGBufferDesc::CreateBufferDesc(DecalGridDataNumBytes, MaxCount * Resolution * Resolution);
+		FRDGBufferDesc DecalGridDataDesc = FRDGBufferDesc::CreateBufferDesc(DecalGridDataNumBytes, 3 * MaxCount * Resolution * Resolution);
 		FRDGBuffer* DecalGridData = GraphBuilder.CreateBuffer(DecalGridDataDesc, TEXT("PathTracer.DecalGridData"));
 
 		auto* PassParameters = GraphBuilder.AllocParameters<FRayTracingBuildDecalGridCS::FParameters>();
@@ -463,7 +445,7 @@ void BuildDecalGrid(FRDGBuilder& GraphBuilder, uint32 NumDecals, FRDGBufferSRVRe
 			RDG_EVENT_NAME("Decal Grid Create (%u decals)", NumDecals),
 			ComputeShader,
 			PassParameters,
-			FComputeShaderUtils::GetGroupCount(FIntPoint(Resolution, Resolution), FComputeShaderUtils::kGolden2DGroupSize));
+			FComputeShaderUtils::GetGroupCount(FIntVector(Resolution, Resolution, 3), FIntVector(FComputeShaderUtils::kGolden2DGroupSize, FComputeShaderUtils::kGolden2DGroupSize, 1)));
 
 		// hookup to the actual rendering pass
 		OutParameters.Grid = DecalGridTexture;
diff --git a/Engine/Source/Runtime/Renderer/Private/RayTracing/RayTracingDecals.h b/Engine/Source/Runtime/Renderer/Private/RayTracing/RayTracingDecals.h
--- a/Engine/Source/Runtime/Renderer/Private/RayTracing/RayTracingDecals.h
+++ b/Engine/Source/Runtime/Renderer/Private/RayTracing/RayTracingDecals.h
@@ -16,7 +16,7 @@ class FGlobalShaderMap;
 enum EBlendMode : int;
 
 BEGIN_GLOBAL_SHADER_PARAMETER_STRUCT(FRayTracingDecals, RENDERER_API)
-	SHADER_PARAMETER_RDG_TEXTURE(Texture2D<uint>, Grid)
+	SHADER_PARAMETER_RDG_TEXTURE(Texture2DArray<uint>, Grid)
 	SHADER_PARAMETER_RDG_BUFFER_SRV(Buffer<uint>, GridData)
 	SHADER_PARAMETER(FVector3f, TranslatedBoundMin)
 	SHADER_PARAMETER(FVector3f, TranslatedBoundMax)
diff --git a/Engine/Source/Runtime/Renderer/Private/RayTracing/RayTracingLighting.cpp b/Engine/Source/Runtime/Renderer/Private/RayTracing/RayTracingLighting.cpp
--- a/Engine/Source/Runtime/Renderer/Private/RayTracing/RayTracingLighting.cpp
+++ b/Engine/Source/Runtime/Renderer/Private/RayTracing/RayTracingLighting.cpp
@@ -69,7 +69,7 @@ class FRayTracingBuildLightGridCS : public FGlobalShader
 		SHADER_PARAMETER(unsigned, LightGridResolution)
 		SHADER_PARAMETER(unsigned, LightGridMaxCount)
 		SHADER_PARAMETER(int, LightGridAxis)
-		SHADER_PARAMETER_RDG_TEXTURE_UAV(RWTexture2D, RWLightGrid)
+		SHADER_PARAMETER_RDG_TEXTURE_UAV(RWTexture2DArray<uint>, RWLightGrid)
 		SHADER_PARAMETER_RDG_BUFFER_UAV(RWBuffer<uint>, RWLightGridData)
 	END_SHADER_PARAMETER_STRUCT()
 };
@@ -163,30 +163,12 @@ static void PrepareLightGrid(FRDGBuilder& GraphBuilder, FGlobalShaderMap* Shader
 		LightGridParameters->SceneLightsTranslatedBoundMax = FVector3f::Max(LightGridParameters->SceneLightsTranslatedBoundMax, Hi);
 	}
 
-	const uint32 Resolution = FMath::RoundUpToPowerOfTwo(CVarRayTracingLightGridResolution.GetValueOnRenderThread());
+	const uint32 Resolution = FMath::Clamp(CVarRayTracingLightGridResolution.GetValueOnRenderThread(), 1, 2048);
 	const uint32 MaxCount = FMath::Clamp(CVarRayTracingLightGridMaxCount.GetValueOnRenderThread(), 1, NumFiniteLights);
 	LightGridParameters->LightGridResolution = Resolution;
 	LightGridParameters->LightGridMaxCount = MaxCount;
 
 	LightGridParameters->LightGridAxis = CVarRayTracingLightGridAxis.GetValueOnRenderThread();
-	if (LightGridParameters->LightGridAxis < 0 || LightGridParameters->LightGridAxis > 2)
-	{
-		// pick the shortest axis if user did not explicitly pick an axis
-		FVector3f Diag = LightGridParameters->SceneLightsTranslatedBoundMax - LightGridParameters->SceneLightsTranslatedBoundMin;
-		if (Diag.X < Diag.Y && Diag.X < Diag.Z)
-		{
-			LightGridParameters->LightGridAxis = 0;
-		}
-		else if (Diag.Y < Diag.Z)
-		{
-			LightGridParameters->LightGridAxis = 1;
-		}
-		else
-		{
-			LightGridParameters->LightGridAxis = 2;
-		}
-	}
-
 
 	// The light grid stores indexes in the range [0,NumLights-1]
 	EPixelFormat LightGridDataFormat = PF_R32_UINT;
@@ -212,16 +194,15 @@ static void PrepareLightGrid(FRDGBuilder& GraphBuilder, FGlobalShaderMap* Shader
 		TextureDataFormat = PF_R16_UINT;
 	}
 
-
-	FRDGTextureDesc LightGridDesc = FRDGTextureDesc::Create2D(
+	FRDGTextureDesc LightGridDesc = FRDGTextureDesc::Create2DArray(
 		FIntPoint(Resolution, Resolution),
 		TextureDataFormat,
 		FClearValueBinding::None,
-		TexCreate_ShaderResource | TexCreate_UAV);
+		TexCreate_ShaderResource | TexCreate_UAV, 3);
 
 	// Run the build compute shader
 	FRDGTexture* LightGridTexture = GraphBuilder.CreateTexture(LightGridDesc, TEXT("RayTracing.LightGrid"), ERDGTextureFlags::None);
-	FRDGBuffer* LightGridData = GraphBuilder.CreateBuffer(FRDGBufferDesc::CreateBufferDesc(LightGridDataNumBytes, MaxCount * Resolution * Resolution), TEXT("RayTracing.LightGridData"));
+	FRDGBuffer* LightGridData = GraphBuilder.CreateBuffer(FRDGBufferDesc::CreateBufferDesc(LightGridDataNumBytes, 3 * MaxCount * Resolution * Resolution), TEXT("RayTracing.LightGridData"));
 
 
 	FRayTracingBuildLightGridCS::FParameters* BuilderParams = GraphBuilder.AllocParameters<FRayTracingBuildLightGridCS::FParameters>();
@@ -242,7 +223,7 @@ static void PrepareLightGrid(FRDGBuilder& GraphBuilder, FGlobalShaderMap* Shader
 		RDG_EVENT_NAME("Light Grid Create (%u lights)", NumFiniteLights),
 		ComputeShader,
 		BuilderParams,
-		FComputeShaderUtils::GetGroupCount(FIntPoint(Resolution, Resolution), FComputeShaderUtils::kGolden2DGroupSize));
+		FComputeShaderUtils::GetGroupCount(FIntVector(Resolution, Resolution, 3), FIntVector(FComputeShaderUtils::kGolden2DGroupSize, FComputeShaderUtils::kGolden2DGroupSize, 1)));
 
 	// hookup to the actual rendering pass
 	LightGridParameters->LightGrid = LightGridTexture;
diff --git a/Engine/Source/Runtime/Renderer/Private/RayTracing/RayTracingLighting.h b/Engine/Source/Runtime/Renderer/Private/RayTracing/RayTracingLighting.h
--- a/Engine/Source/Runtime/Renderer/Private/RayTracing/RayTracingLighting.h
+++ b/Engine/Source/Runtime/Renderer/Private/RayTracing/RayTracingLighting.h
@@ -22,7 +22,7 @@ BEGIN_GLOBAL_SHADER_PARAMETER_STRUCT(FRayTracingLightGrid, RENDERER_API)
 	SHADER_PARAMETER(FVector3f, SceneLightsTranslatedBoundMin)
 	SHADER_PARAMETER(FVector3f, SceneLightsTranslatedBoundMax)
 	SHADER_PARAMETER_RDG_BUFFER_SRV(StructuredBuffer<FRTLightingData>, SceneLights)
-	SHADER_PARAMETER_RDG_TEXTURE(Texture2D<uint>, LightGrid)
+	SHADER_PARAMETER_RDG_TEXTURE(Texture2DArray<uint>, LightGrid)
 	SHADER_PARAMETER_RDG_BUFFER_SRV(Buffer<uint>, LightGridData)
 	SHADER_PARAMETER(unsigned, LightGridResolution)
 	SHADER_PARAMETER(unsigned, LightGridMaxCount)
