Fix tonemapper being enabled in XR.
Keep it enabled only in Editor windows as they need the linear to sRGB conversion that is done by the HW in the XR backbuffer case.

diff --git a/Engine/Shaders/Private/ScreenPass.usf b/Engine/Shaders/Private/ScreenPass.usf
--- a/Engine/Shaders/Private/ScreenPass.usf
+++ b/Engine/Shaders/Private/ScreenPass.usf
@@ -2,22 +2,45 @@
 
 #include "Common.ush"
 
+#if MOBILE_MULTI_VIEW
+Texture2DArray InputTexture;
+#else
 Texture2D InputTexture;
+#endif
 SamplerState InputSampler;
 
 void ScreenPassVS(
 	in float4 InPosition : ATTRIBUTE0,
 	in float2 InTexCoord : ATTRIBUTE1,
-	out noperspective float4 OutUVAndScreenPos : TEXCOORD0,
-	out float4 OutPosition : SV_POSITION)
+	out noperspective float4 OutUVAndScreenPos : TEXCOORD0
+#if MOBILE_MULTI_VIEW && INSTANCED_STEREO
+	, in uint InstanceId : SV_InstanceID
+	, out nointerpolation uint EyeIndex : VIEW_ID
+	, out uint LayerIndex : SV_RenderTargetArrayIndex
+#endif
+	, out float4 OutPosition : SV_POSITION
+)
 {
+#if MOBILE_MULTI_VIEW && INSTANCED_STEREO
+	EyeIndex = InstanceId;
+	LayerIndex = InstanceId;
+#endif
 	DrawRectangle(InPosition, InTexCoord, OutPosition, OutUVAndScreenPos);
 }
 
 float4 CopyRectPS(
 	noperspective float4 UVAndScreenPos : TEXCOORD0
+#if MOBILE_MULTI_VIEW && INSTANCED_STEREO
+	, in nointerpolation uint EyeIndex : VIEW_ID
+#elif MOBILE_MULTI_VIEW
+	, in nointerpolation uint EyeIndex : SV_ViewID
+#endif
 	) : SV_Target0
 {
 	float2 UV = UVAndScreenPos.xy;
+#if MOBILE_MULTI_VIEW
+	return Texture2DArraySample(InputTexture, InputSampler, float3(UV, EyeIndex));
+#else
 	return Texture2DSample(InputTexture, InputSampler, UV);
+#endif
 }
\ No newline at end of file
diff --git a/Engine/Source/Runtime/Renderer/Private/PostProcess/PostProcessing.cpp b/Engine/Source/Runtime/Renderer/Private/PostProcess/PostProcessing.cpp
--- a/Engine/Source/Runtime/Renderer/Private/PostProcess/PostProcessing.cpp
+++ b/Engine/Source/Runtime/Renderer/Private/PostProcess/PostProcessing.cpp
@@ -2361,7 +2361,9 @@ void AddMobilePostProcessingPasses(FRDGBuilder& GraphBuilder, FScene* Scene, con
 	bool bUseDof = GetMobileDepthOfFieldScale(View) > 0.0f && View.Family->EngineShowFlags.DepthOfField && !View.Family->EngineShowFlags.VisualizeDOF;
 	bool bUseMobileDof = bUseDof && !View.FinalPostProcessSettings.bMobileHQGaussian;
 
-	bool bUseToneMapper = !View.Family->EngineShowFlags.ShaderComplexity && (IsMobileHDR() || IsMobileColorsRGB());
+	// Do not use the tonemapper if output texture is sRGB cause the conversion will be performed by HW.
+	const bool bIsOutputTexsRGB = EnumHasAnyFlags(Inputs.ViewFamilyTexture->Desc.Flags, TexCreate_SRGB);
+	bool bUseToneMapper = !View.Family->EngineShowFlags.ShaderComplexity && (IsMobileHDR() || (IsMobileColorsRGB() && !bIsOutputTexsRGB));
 
 	bool bUseHighResolutionScreenshotMask = IsHighResolutionScreenshotMaskEnabled(View);
 
@@ -3131,7 +3133,8 @@ void AddMobilePostProcessingPasses(FRDGBuilder& GraphBuilder, FScene* Scene, con
 	// Copy the scene color to back buffer in case there is no post process, such as LDR MSAA.
 	if (SceneColor.Texture != ViewFamilyOutput.Texture)
 	{
-		AddDrawTexturePass(GraphBuilder, View, SceneColor, ViewFamilyOutput);
+		const uint32 RTMultiViewCount = (View.bIsMobileMultiViewEnabled) ? 2 : (View.Aspects.IsMobileMultiViewEnabled() ? 1 : 0);
+		AddDrawTexturePass(GraphBuilder, View, SceneColor, ViewFamilyOutput, RTMultiViewCount);
 	}
 }
 
diff --git a/Engine/Source/Runtime/Renderer/Private/ScreenPass.cpp b/Engine/Source/Runtime/Renderer/Private/ScreenPass.cpp
--- a/Engine/Source/Runtime/Renderer/Private/ScreenPass.cpp
+++ b/Engine/Source/Runtime/Renderer/Private/ScreenPass.cpp
@@ -354,7 +354,8 @@ void AddDrawTexturePass(
 	FRDGBuilder& GraphBuilder,
 	FScreenPassViewInfo ViewInfo,
 	FScreenPassTexture Input,
-	FScreenPassRenderTarget Output)
+	FScreenPassRenderTarget Output,
+	uint32 RTMultiViewCount)
 {
 	const FScreenPassTextureViewport InputViewport(Input);
 	const FScreenPassTextureViewport OutputViewport(Output);
@@ -365,6 +366,7 @@ void AddDrawTexturePass(
 	Parameters->InputTexture = Input.Texture;
 	Parameters->InputSampler = TStaticSamplerState<>::GetRHI();
 	Parameters->RenderTargets[0] = Output.GetRenderTargetBinding();
+	Parameters->RenderTargets.MultiViewCount = RTMultiViewCount;
 
 	AddDrawScreenPass(GraphBuilder, RDG_EVENT_NAME("DrawTexture"), ViewInfo, OutputViewport, InputViewport, PixelShader, Parameters);
 }
diff --git a/Engine/Source/Runtime/Renderer/Public/ScreenPass.h b/Engine/Source/Runtime/Renderer/Public/ScreenPass.h
--- a/Engine/Source/Runtime/Renderer/Public/ScreenPass.h
+++ b/Engine/Source/Runtime/Renderer/Public/ScreenPass.h
@@ -678,7 +678,8 @@ void RENDERER_API AddDrawTexturePass(
 	FRDGBuilder& GraphBuilder,
 	FScreenPassViewInfo ViewInfo,
 	FScreenPassTexture Input,
-	FScreenPassRenderTarget Output);
+	FScreenPassRenderTarget Output,
+	uint32 RTMultiviewCount = 1);
 
 inline void AddDrawTexturePass(
 	FRDGBuilder& GraphBuilder,
