Substrate - fixed crash when reading substrate data like normal from emissive decal. This should be allowed. On any other decal passes (DBuffer) we do not bind substrate data as it should not be used.

#rnx
#tests none

diff --git a/Engine/Shaders/Private/DeferredDecal.usf b/Engine/Shaders/Private/DeferredDecal.usf
--- a/Engine/Shaders/Private/DeferredDecal.usf
+++ b/Engine/Shaders/Private/DeferredDecal.usf
@@ -10,14 +10,17 @@
 #if DECAL_RENDERTARGETMODE == DECAL_RENDERTARGETMODE_DBUFFER
 	#define MobileSceneTextures DecalPass.MobileSceneTextures
 	#define EyeAdaptationStruct DecalPass
+	#define SubstrateStruct DecalPass.SubstratePublic
 #else
 	// Reroute MobileSceneTextures uniform buffer references to the base pass uniform buffer
 	#define MobileSceneTextures MobileBasePass.SceneTextures
 	#define EyeAdaptationStruct MobileBasePass
+	#define SubstrateStruct MobileBasePass.SubstratePublic
 #endif
 #else
 #define SceneTexturesStruct DecalPass.SceneTextures
 #define EyeAdaptationStruct DecalPass
+#define SubstrateStruct DecalPass.SubstratePublic
 #endif
 
 #if PIXELSHADER
diff --git a/Engine/Shaders/Private/MeshDecals.usf b/Engine/Shaders/Private/MeshDecals.usf
--- a/Engine/Shaders/Private/MeshDecals.usf
+++ b/Engine/Shaders/Private/MeshDecals.usf
@@ -10,9 +10,11 @@
 	// Reroute MobileSceneTextures uniform buffer references to the base pass uniform buffer 
 	#define MobileSceneTextures MobileBasePass.SceneTextures
 	#define EyeAdaptationStruct MobileBasePass
+	#define SubstrateStruct MobileBasePass.SubstratePublic
 #else
 	#define SceneTexturesStruct DecalPass.SceneTextures
 	#define EyeAdaptationStruct DecalPass
+	#define SubstrateStruct DecalPass.SubstratePublic
 #endif
 
 #include "Substrate/Substrate.ush"
diff --git a/Engine/Shaders/Public/ShaderVersion.ush b/Engine/Shaders/Public/ShaderVersion.ush
--- a/Engine/Shaders/Public/ShaderVersion.ush
+++ b/Engine/Shaders/Public/ShaderVersion.ush
@@ -3,5 +3,5 @@
 // Each time the console command is executed it generates a new GUID. As this file is included
 // in Platform.ush (which should be included in any shader) it allows to invalidate the shader DDC.
 // 
-#pragma message("UESHADERMETADATA_VERSION 1BE4816C-46DD-407E-97A2-83F1C880E8B1")
+#pragma message("UESHADERMETADATA_VERSION C08D4C8D-F9A6-4473-95C7-99576A30AAAF")
 
diff --git a/Engine/Source/Runtime/Renderer/Private/CompositionLighting/CompositionLighting.cpp b/Engine/Source/Runtime/Renderer/Private/CompositionLighting/CompositionLighting.cpp
--- a/Engine/Source/Runtime/Renderer/Private/CompositionLighting/CompositionLighting.cpp
+++ b/Engine/Source/Runtime/Renderer/Private/CompositionLighting/CompositionLighting.cpp
@@ -562,7 +562,7 @@ void FCompositionLighting::TryInit()
 	bInitialized = true;
 }
 
-void FCompositionLighting::ProcessBeforeBasePass(FRDGBuilder& GraphBuilder, FDBufferTextures& DBufferTextures, FInstanceCullingManager& InstanceCullingManager)
+void FCompositionLighting::ProcessBeforeBasePass(FRDGBuilder& GraphBuilder, FDBufferTextures& DBufferTextures, FInstanceCullingManager& InstanceCullingManager, const FSubstrateSceneData& SubstrateSceneData)
 {
 	if (HasRayTracedOverlay(ViewFamily))
 	{
@@ -594,7 +594,7 @@ void FCompositionLighting::ProcessBeforeBasePass(FRDGBuilder& GraphBuilder, FDBu
 		// decals are before AmbientOcclusion so the decal can output a normal that AO is affected by
 		if (bEnableDBuffer)
 		{
-			FDeferredDecalPassTextures DecalPassTextures = GetDeferredDecalPassTextures(GraphBuilder, View, SceneTextures, &DBufferTextures);
+			FDeferredDecalPassTextures DecalPassTextures = GetDeferredDecalPassTextures(GraphBuilder, View, SubstrateSceneData, SceneTextures, &DBufferTextures, EDecalRenderStage::BeforeBasePass);
 			AddDeferredDecalPass(GraphBuilder, View, VisibleDecals[ViewIndex], DecalPassTextures, InstanceCullingManager, EDecalRenderStage::BeforeBasePass);
 		}
 
@@ -612,7 +612,7 @@ void FCompositionLighting::ProcessBeforeBasePass(FRDGBuilder& GraphBuilder, FDBu
 	}
 }
 
-void FCompositionLighting::ProcessAfterBasePass(FRDGBuilder& GraphBuilder, FInstanceCullingManager& InstanceCullingManager, EProcessAfterBasePassMode Mode)
+void FCompositionLighting::ProcessAfterBasePass(FRDGBuilder& GraphBuilder, FInstanceCullingManager& InstanceCullingManager, EProcessAfterBasePassMode Mode, const FSubstrateSceneData& SubstrateSceneData)
 {
 	if (HasRayTracedOverlay(ViewFamily))
 	{
@@ -634,18 +634,19 @@ void FCompositionLighting::ProcessAfterBasePass(FRDGBuilder& GraphBuilder, FInst
 
 		View.BeginRenderView();
 
-		FDeferredDecalPassTextures DecalPassTextures = GetDeferredDecalPassTextures(GraphBuilder, View, SceneTextures, nullptr);
 
 		if (bEnableDecals && !bEnableDBuffer && IsUsingGBuffers(SceneTextures.Config.ShaderPlatform) && Mode != EProcessAfterBasePassMode::SkipBeforeLightingDecals)
 		{
 			// We can disable this pass if using DBuffer decals
 			// Decals are before AmbientOcclusion so the decal can output a normal that AO is affected by
+			FDeferredDecalPassTextures DecalPassTextures = GetDeferredDecalPassTextures(GraphBuilder, View, SubstrateSceneData, SceneTextures, nullptr, EDecalRenderStage::BeforeLighting);
 			AddDeferredDecalPass(GraphBuilder, View, VisibleDecals[ViewIndex], DecalPassTextures, InstanceCullingManager, EDecalRenderStage::BeforeLighting);
 		}
 
 		if (bEnableDecals && Mode != EProcessAfterBasePassMode::OnlyBeforeLightingDecals)
 		{
 			// DBuffer decals with emissive component
+			FDeferredDecalPassTextures DecalPassTextures = GetDeferredDecalPassTextures(GraphBuilder, View, SubstrateSceneData, SceneTextures, nullptr, EDecalRenderStage::Emissive);
 			AddDeferredDecalPass(GraphBuilder, View, VisibleDecals[ViewIndex], DecalPassTextures, InstanceCullingManager, EDecalRenderStage::Emissive);
 		}
 
@@ -688,6 +689,7 @@ void FCompositionLighting::ProcessAfterBasePass(FRDGBuilder& GraphBuilder, FInst
 
 					if (bEnableDecals)
 					{
+						FDeferredDecalPassTextures DecalPassTextures = GetDeferredDecalPassTextures(GraphBuilder, View, SubstrateSceneData, SceneTextures, nullptr, EDecalRenderStage::AmbientOcclusion);
 						DecalPassTextures.ScreenSpaceAO = AmbientOcclusion.Texture;
 						AddDeferredDecalPass(GraphBuilder, View, VisibleDecals[ViewIndex], DecalPassTextures, InstanceCullingManager, EDecalRenderStage::AmbientOcclusion);
 					}
diff --git a/Engine/Source/Runtime/Renderer/Private/CompositionLighting/CompositionLighting.h b/Engine/Source/Runtime/Renderer/Private/CompositionLighting/CompositionLighting.h
--- a/Engine/Source/Runtime/Renderer/Private/CompositionLighting/CompositionLighting.h
+++ b/Engine/Source/Runtime/Renderer/Private/CompositionLighting/CompositionLighting.h
@@ -22,7 +22,7 @@ public:
 
 	void ProcessAfterOcclusion(FRDGBuilder& GraphBuilder);
 
-	void ProcessBeforeBasePass(FRDGBuilder& GraphBuilder, FDBufferTextures& DBufferTextures, FInstanceCullingManager& InstanceCullingManager);
+	void ProcessBeforeBasePass(FRDGBuilder& GraphBuilder, FDBufferTextures& DBufferTextures, FInstanceCullingManager& InstanceCullingManager, const FSubstrateSceneData& SubstrateSceneData);
 
 	enum class EProcessAfterBasePassMode
 	{
@@ -31,7 +31,7 @@ public:
 		All
 	};
 
-	void ProcessAfterBasePass(FRDGBuilder& GraphBuilder, FInstanceCullingManager& InstanceCullingManager, EProcessAfterBasePassMode Mode);
+	void ProcessAfterBasePass(FRDGBuilder& GraphBuilder, FInstanceCullingManager& InstanceCullingManager, EProcessAfterBasePassMode Mode, const FSubstrateSceneData& SubstrateSceneData);
 
 private:
 	void TryInit();
diff --git a/Engine/Source/Runtime/Renderer/Private/CompositionLighting/PostProcessDeferredDecals.cpp b/Engine/Source/Runtime/Renderer/Private/CompositionLighting/PostProcessDeferredDecals.cpp
--- a/Engine/Source/Runtime/Renderer/Private/CompositionLighting/PostProcessDeferredDecals.cpp
+++ b/Engine/Source/Runtime/Renderer/Private/CompositionLighting/PostProcessDeferredDecals.cpp
@@ -70,14 +70,17 @@ IMPLEMENT_STATIC_UNIFORM_BUFFER_STRUCT(FDecalPassUniformParameters, "DecalPass",
 
 FDeferredDecalPassTextures GetDeferredDecalPassTextures(
 	FRDGBuilder& GraphBuilder, 
-	const FSceneView& View,
+	const FViewInfo& View,
+	const FSubstrateSceneData& SubstrateSceneData,
 	const FSceneTextures& SceneTextures, 
-	FDBufferTextures* DBufferTextures)
+	FDBufferTextures* DBufferTextures,
+	EDecalRenderStage DecalRenderStage)
 {
 	FDeferredDecalPassTextures PassTextures;
 
-	auto* Parameters = GraphBuilder.AllocParameters<FDecalPassUniformParameters>();
-	
+	auto* Parameters = GraphBuilder.AllocParameters<FDecalPassUniformParameters>(); //
+
+
 	const bool bIsMobile = (View.GetFeatureLevel() == ERHIFeatureLevel::ES3_1);
 	ESceneTextureSetupMode TextureReadAccess = ESceneTextureSetupMode::None;
 	EMobileSceneTextureSetupMode MobileTextureReadAccess = EMobileSceneTextureSetupMode::None;
@@ -93,6 +96,14 @@ FDeferredDecalPassTextures GetDeferredDecalPassTextures(
 	SetupSceneTextureUniformParameters(GraphBuilder, &SceneTextures, View.FeatureLevel, TextureReadAccess, Parameters->SceneTextures);
 	SetupMobileSceneTextureUniformParameters(GraphBuilder, &SceneTextures, MobileTextureReadAccess, Parameters->MobileSceneTextures);
 	Parameters->EyeAdaptationBuffer = GraphBuilder.CreateSRV(GetEyeAdaptationBuffer(GraphBuilder, View));
+	if (DecalRenderStage == EDecalRenderStage::Emissive)
+	{
+		 Substrate::BindSubstratePublicGlobalUniformParameters(GraphBuilder, &SubstrateSceneData, Parameters->SubstratePublic);
+	}
+	else
+	{
+		Substrate::BindSubstratePublicGlobalUniformParameters(GraphBuilder, nullptr, Parameters->SubstratePublic); // nullptr for default
+	}
 	PassTextures.DecalPassUniformBuffer = GraphBuilder.CreateUniformBuffer(Parameters);
 
 	PassTextures.Depth = SceneTextures.Depth;
@@ -186,6 +197,7 @@ void GetDeferredDecalPassParameters(
 	FRDGBuilder &GraphBuilder,
 	const FViewInfo& View,
 	const FDeferredDecalPassTextures& Textures,
+	EDecalRenderStage DecalRenderStage,
 	EDecalRenderTargetMode RenderTargetMode,
 	FDeferredDecalPassParameters& PassParameters)
 {
@@ -681,7 +693,7 @@ void AddDeferredDecalPass(
 		}
 
 		auto* PassParameters = GraphBuilder.AllocParameters<FDeferredDecalPassParameters>();
-		GetDeferredDecalPassParameters(GraphBuilder, View, PassTextures, RenderTargetMode, *PassParameters);
+		GetDeferredDecalPassParameters(GraphBuilder, View, PassTextures, DecalRenderStage, RenderTargetMode, *PassParameters);
 
 		GraphBuilder.AddPass(
 			RDG_EVENT_NAME("Batch [%d, %d]", DecalIndexBegin, DecalIndexEnd - 1),
diff --git a/Engine/Source/Runtime/Renderer/Private/CompositionLighting/PostProcessDeferredDecals.h b/Engine/Source/Runtime/Renderer/Private/CompositionLighting/PostProcessDeferredDecals.h
--- a/Engine/Source/Runtime/Renderer/Private/CompositionLighting/PostProcessDeferredDecals.h
+++ b/Engine/Source/Runtime/Renderer/Private/CompositionLighting/PostProcessDeferredDecals.h
@@ -6,6 +6,7 @@
 #include "SceneTextureParameters.h"
 #include "SceneView.h"
 #include "InstanceCulling/InstanceCullingContext.h"
+#include "Substrate/Substrate.h"
 
 enum class EDecalRenderStage : uint8;
 enum class EDecalRenderTargetMode : uint8;
@@ -22,6 +23,7 @@ bool IsDBufferEnabled(const FSceneViewFamily& ViewFamily, EShaderPlatform Shader
 BEGIN_GLOBAL_SHADER_PARAMETER_STRUCT(FDecalPassUniformParameters, )
 	SHADER_PARAMETER_STRUCT(FSceneTextureUniformParameters, SceneTextures)
 	SHADER_PARAMETER_STRUCT(FMobileSceneTextureUniformParameters, MobileSceneTextures)
+	SHADER_PARAMETER_STRUCT(FSubstratePublicParameters, SubstratePublic)
 	SHADER_PARAMETER_RDG_BUFFER_SRV(StructuredBuffer<float4>, EyeAdaptationBuffer)
 END_GLOBAL_SHADER_PARAMETER_STRUCT()
 
@@ -42,9 +44,11 @@ struct FDeferredDecalPassTextures
 
 FDeferredDecalPassTextures GetDeferredDecalPassTextures(
 	FRDGBuilder& GraphBuilder, 
-	const FSceneView& ViewInfo,
+	const FViewInfo& View,
+	const FSubstrateSceneData& SubstrateSceneData,
 	const FSceneTextures& SceneTextures, 
-	FDBufferTextures* DBufferTextures);
+	FDBufferTextures* DBufferTextures,
+	EDecalRenderStage DecalRenderStage);
 
 void AddDeferredDecalPass(
 	FRDGBuilder& GraphBuilder,
@@ -91,6 +95,7 @@ void GetDeferredDecalPassParameters(
 	FRDGBuilder& GraphBuilder,
 	const FViewInfo& View,
 	const FDeferredDecalPassTextures& DecalPassTextures,
+	EDecalRenderStage DecalRenderStage,
 	EDecalRenderTargetMode RenderTargetMode,
 	FDeferredDecalPassParameters& PassParameters);
 
diff --git a/Engine/Source/Runtime/Renderer/Private/CompositionLighting/PostProcessMeshDecals.cpp b/Engine/Source/Runtime/Renderer/Private/CompositionLighting/PostProcessMeshDecals.cpp
--- a/Engine/Source/Runtime/Renderer/Private/CompositionLighting/PostProcessMeshDecals.cpp
+++ b/Engine/Source/Runtime/Renderer/Private/CompositionLighting/PostProcessMeshDecals.cpp
@@ -510,7 +510,7 @@ void DrawDecalMeshCommands(
 	EDecalRenderTargetMode RenderTargetMode)
 {
 	auto* PassParameters = GraphBuilder.AllocParameters<FDeferredDecalPassParameters>();
-	GetDeferredDecalPassParameters(GraphBuilder, View, DecalPassTextures, RenderTargetMode, *PassParameters);
+	GetDeferredDecalPassParameters(GraphBuilder, View, DecalPassTextures, DecalRenderStage, RenderTargetMode, *PassParameters);
 
 	EMeshPass::Type DecalMeshPassType = DecalRendering::GetMeshPassType(RenderTargetMode);
 	if (View.ParallelMeshDrawCommandPasses[DecalMeshPassType].HasAnyDraw())
diff --git a/Engine/Source/Runtime/Renderer/Private/DeferredShadingRenderer.cpp b/Engine/Source/Runtime/Renderer/Private/DeferredShadingRenderer.cpp
--- a/Engine/Source/Runtime/Renderer/Private/DeferredShadingRenderer.cpp
+++ b/Engine/Source/Runtime/Renderer/Private/DeferredShadingRenderer.cpp
@@ -2422,7 +2422,7 @@ void FDeferredShadingSceneRenderer::Render(FRDGBuilder& GraphBuilder)
 		{
 			RDG_CSV_STAT_EXCLUSIVE_SCOPE(GraphBuilder, DeferredShadingSceneRenderer_DBuffer);
 			SCOPE_CYCLE_COUNTER(STAT_FDeferredShadingSceneRenderer_DBuffer);
-			CompositionLighting.ProcessBeforeBasePass(GraphBuilder, DBufferTextures, InstanceCullingManager);
+			CompositionLighting.ProcessBeforeBasePass(GraphBuilder, DBufferTextures, InstanceCullingManager, Scene->SubstrateSceneData);
 		}
 		
 		if (IsForwardShadingEnabled(ShaderPlatform))
@@ -2750,7 +2750,7 @@ void FDeferredShadingSceneRenderer::Render(FRDGBuilder& GraphBuilder)
 			const FCompositionLighting::EProcessAfterBasePassMode Mode = AsyncLumenIndirectLightingOutputs.bHasDrawnBeforeLightingDecals ?
 				FCompositionLighting::EProcessAfterBasePassMode::SkipBeforeLightingDecals : FCompositionLighting::EProcessAfterBasePassMode::All;
 
-			CompositionLighting.ProcessAfterBasePass(GraphBuilder, InstanceCullingManager, Mode);
+			CompositionLighting.ProcessAfterBasePass(GraphBuilder, InstanceCullingManager, Mode, Scene->SubstrateSceneData);
 		}
 
 		// Rebuild scene textures to include velocity, custom depth, and SSAO.
diff --git a/Engine/Source/Runtime/Renderer/Private/IndirectLightRendering.cpp b/Engine/Source/Runtime/Renderer/Private/IndirectLightRendering.cpp
--- a/Engine/Source/Runtime/Renderer/Private/IndirectLightRendering.cpp
+++ b/Engine/Source/Runtime/Renderer/Private/IndirectLightRendering.cpp
@@ -896,7 +896,7 @@ void FDeferredShadingSceneRenderer::DispatchAsyncLumenIndirectLightingWork(
 	}
 
 	// Decals may modify GBuffers so they need to be done first. Can decals read velocities and/or custom depth? If so, they need to be rendered earlier too.
-	CompositionLighting.ProcessAfterBasePass(GraphBuilder, InstanceCullingManager, FCompositionLighting::EProcessAfterBasePassMode::OnlyBeforeLightingDecals);
+	CompositionLighting.ProcessAfterBasePass(GraphBuilder, InstanceCullingManager, FCompositionLighting::EProcessAfterBasePassMode::OnlyBeforeLightingDecals, Scene->SubstrateSceneData);
 	Outputs.bHasDrawnBeforeLightingDecals = true;
 
 	if (DynamicShadowsTaskData)
diff --git a/Engine/Source/Runtime/Renderer/Private/MobileDecalRendering.cpp b/Engine/Source/Runtime/Renderer/Private/MobileDecalRendering.cpp
--- a/Engine/Source/Runtime/Renderer/Private/MobileDecalRendering.cpp
+++ b/Engine/Source/Runtime/Renderer/Private/MobileDecalRendering.cpp
@@ -171,7 +171,7 @@ void FMobileSceneRenderer::RenderDBuffer(FRDGBuilder& GraphBuilder, FSceneTextur
 
 		FTransientDecalRenderDataList VisibleDecals = DecalRendering::BuildVisibleDecalList(Scene->Decals, View);
 
-		FDeferredDecalPassTextures DecalPassTextures = GetDeferredDecalPassTextures(GraphBuilder, View, SceneTextures, &DBufferTextures);
+		FDeferredDecalPassTextures DecalPassTextures = GetDeferredDecalPassTextures(GraphBuilder, View, Scene->SubstrateSceneData, SceneTextures, &DBufferTextures, EDecalRenderStage::BeforeBasePass);
 		AddDeferredDecalPass(GraphBuilder, View, VisibleDecals, DecalPassTextures, InstanceCullingManager, EDecalRenderStage::BeforeBasePass);
 	}
 }
diff --git a/Engine/Source/Runtime/Renderer/Private/Substrate/Substrate.cpp b/Engine/Source/Runtime/Renderer/Private/Substrate/Substrate.cpp
--- a/Engine/Source/Runtime/Renderer/Private/Substrate/Substrate.cpp
+++ b/Engine/Source/Runtime/Renderer/Private/Substrate/Substrate.cpp
@@ -705,26 +705,6 @@ TRDGUniformBufferRef<FSubstrateGlobalUniformParameters> BindSubstrateGlobalUnifo
 	return View.SubstrateViewData.SubstrateGlobalUniformParameters;
 }
 
-static void BindSubstratePublicGlobalUniformParameters(FRDGBuilder& GraphBuilder, FSubstrateSceneData* SubstrateSceneData, FSubstratePublicGlobalUniformParameters& OutSubstrateUniformParameters)
-{
-	if (SubstrateSceneData && SubstrateSceneData->TopLayerTexture)
-	{
-		OutSubstrateUniformParameters.Common = GetSubstrateCommonParameter(*SubstrateSceneData);
-		OutSubstrateUniformParameters.FirstSliceStoringSubstrateSSSData = SubstrateSceneData->FirstSliceStoringSubstrateSSSData;
-		OutSubstrateUniformParameters.MaterialTextureArray = SubstrateSceneData->MaterialTextureArray;
-		OutSubstrateUniformParameters.TopLayerTexture = SubstrateSceneData->TopLayerTexture;
-	}
-	else
-	{
-		const FRDGSystemTextures& SystemTextures = FRDGSystemTextures::Get(GraphBuilder);
-		OutSubstrateUniformParameters.Common = GetSubstrateCommonParameter();
-		OutSubstrateUniformParameters.FirstSliceStoringSubstrateSSSData = -1;
-		OutSubstrateUniformParameters.MaterialTextureArray = GetDefaultSubstrateMaterialTextureArray(GraphBuilder);
-		OutSubstrateUniformParameters.TopLayerTexture = SystemTextures.Black;
-	}
-
-}
-
 static ERHIFeatureSupport SubstrateSupportsWaveOps(EShaderPlatform Platform)
 {
 	// D3D11 / SM5 or preview do not support, or work well with, wave-ops by default (or SM5 preview has issues with wave intrinsics too), that fixes classification and black/wrong tiling.
@@ -736,11 +716,30 @@ static ERHIFeatureSupport SubstrateSupportsWaveOps(EShaderPlatform Platform)
 	return FDataDrivenShaderPlatformInfo::GetSupportsWaveOperations(Platform);
 }
 
+void BindSubstratePublicGlobalUniformParameters(FRDGBuilder& GraphBuilder, const FSubstrateSceneData* SubstrateSceneData, FSubstratePublicParameters& OutSubstrateParameters)
+{
+	if (SubstrateSceneData && SubstrateSceneData->TopLayerTexture)
+	{
+		OutSubstrateParameters.Common = GetSubstrateCommonParameter(*SubstrateSceneData);
+		OutSubstrateParameters.FirstSliceStoringSubstrateSSSData = SubstrateSceneData->FirstSliceStoringSubstrateSSSData;
+		OutSubstrateParameters.MaterialTextureArray = SubstrateSceneData->MaterialTextureArray;
+		OutSubstrateParameters.TopLayerTexture = SubstrateSceneData->TopLayerTexture;
+	}
+	else
+	{
+		const FRDGSystemTextures& SystemTextures = FRDGSystemTextures::Get(GraphBuilder);
+		OutSubstrateParameters.Common = GetSubstrateCommonParameter();
+		OutSubstrateParameters.FirstSliceStoringSubstrateSSSData = -1;
+		OutSubstrateParameters.MaterialTextureArray = GetDefaultSubstrateMaterialTextureArray(GraphBuilder);
+		OutSubstrateParameters.TopLayerTexture = SystemTextures.Black;
+	}
+}
+
 TRDGUniformBufferRef<FSubstratePublicGlobalUniformParameters> CreatePublicGlobalUniformBuffer(FRDGBuilder& GraphBuilder, FSubstrateSceneData* SubstrateScene)
 {
 	FSubstratePublicGlobalUniformParameters* SubstratePublicUniformParameters = GraphBuilder.AllocParameters<FSubstratePublicGlobalUniformParameters>();
 	check(SubstratePublicUniformParameters);
-	BindSubstratePublicGlobalUniformParameters(GraphBuilder, SubstrateScene, *SubstratePublicUniformParameters);
+	BindSubstratePublicGlobalUniformParameters(GraphBuilder, SubstrateScene, SubstratePublicUniformParameters->Public);
 	return GraphBuilder.CreateUniformBuffer(SubstratePublicUniformParameters);
 }
 
diff --git a/Engine/Source/Runtime/Renderer/Private/Substrate/Substrate.h b/Engine/Source/Runtime/Renderer/Private/Substrate/Substrate.h
--- a/Engine/Source/Runtime/Renderer/Private/Substrate/Substrate.h
+++ b/Engine/Source/Runtime/Renderer/Private/Substrate/Substrate.h
@@ -72,11 +72,15 @@ BEGIN_GLOBAL_SHADER_PARAMETER_STRUCT(FSubstrateGlobalUniformParameters, RENDERER
 	SHADER_PARAMETER_RDG_BUFFER_SRV(Buffer<uint>, ClosureTileCountBuffer)
 END_GLOBAL_SHADER_PARAMETER_STRUCT()
 
-BEGIN_GLOBAL_SHADER_PARAMETER_STRUCT(FSubstratePublicGlobalUniformParameters, RENDERER_API)
+BEGIN_SHADER_PARAMETER_STRUCT(FSubstratePublicParameters, RENDERER_API)
 	SHADER_PARAMETER_STRUCT_INCLUDE(FSubstrateCommonParameters, Common)
 	SHADER_PARAMETER(int32, FirstSliceStoringSubstrateSSSData)
 	SHADER_PARAMETER_RDG_TEXTURE(Texture2D<SUBSTRATE_TOP_LAYER_TYPE>, TopLayerTexture)
 	SHADER_PARAMETER_RDG_TEXTURE(Texture2DArray<uint>, MaterialTextureArray)
+END_SHADER_PARAMETER_STRUCT()
+
+BEGIN_GLOBAL_SHADER_PARAMETER_STRUCT(FSubstratePublicGlobalUniformParameters, RENDERER_API)
+	SHADER_PARAMETER_STRUCT_INCLUDE(FSubstratePublicParameters, Public)
 END_GLOBAL_SHADER_PARAMETER_STRUCT()
 
 // This must map to the SUBSTRATE_TILE_TYPE defines.
@@ -203,6 +207,7 @@ TRDGUniformBufferRef<FSubstrateGlobalUniformParameters> BindSubstrateGlobalUnifo
 void AppendSubstrateMRTs(const FSceneRenderer& SceneRenderer, uint32& BasePassTextureCount, TArrayView<FTextureRenderTargetBinding> BasePassTextures);
 void SetBasePassRenderTargetOutputFormat(const EShaderPlatform Platform, const FMaterialShaderParameters& MaterialParameters, FShaderCompilerEnvironment& OutEnvironment, EGBufferLayout GBufferLayout);
 
+void BindSubstratePublicGlobalUniformParameters(FRDGBuilder& GraphBuilder, const FSubstrateSceneData* SubstrateSceneData, FSubstratePublicParameters& OutSubstrateUniformParameters);
 TRDGUniformBufferRef<FSubstratePublicGlobalUniformParameters> CreatePublicGlobalUniformBuffer(FRDGBuilder& GraphBuilder, FSubstrateSceneData* SubstrateScene);
 
 void AddSubstrateMaterialClassificationPass(FRDGBuilder& GraphBuilder, const FMinimalSceneTextures& SceneTextures, const FDBufferTextures& DBufferTextures, const TArray<FViewInfo>& Views);
