Fix material draw markers for batched indrect draws

diff --git a/Engine/Source/Runtime/Renderer/Private/InstanceCulling/InstanceCullingContext.cpp b/Engine/Source/Runtime/Renderer/Private/InstanceCulling/InstanceCullingContext.cpp
--- a/Engine/Source/Runtime/Renderer/Private/InstanceCulling/InstanceCullingContext.cpp
+++ b/Engine/Source/Runtime/Renderer/Private/InstanceCulling/InstanceCullingContext.cpp
@@ -19,6 +19,9 @@
 #include "InstanceCullingOcclusionQuery.h"
 #include "RenderCore.h"
 #include "MeshDrawCommandStats.h"
+#include "UnrealEngine.h"
+#include "RHIBreadCrumbs.h"
+#include "Materials/MaterialRenderProxy.h"
 
 static TAutoConsoleVariable<int32> CVarCullInstances(
 	TEXT("r.CullInstances"),
@@ -1721,8 +1724,19 @@ void FInstanceCullingContext::SubmitDrawCommands(
 			}
 			
 			SceneArgs.PrimitiveIdOffset = OverrideArgs.InstanceDataByteOffset + DrawCommandInfo.InstanceDataByteOffset;
-			if (FMeshDrawCommand::SubmitDraw(*VisibleMeshDrawCommand.MeshDrawCommand, GraphicsMinimalPipelineStateSet, SceneArgs, InstanceFactor, RHICmdList, StateCache))
+
+#if WANTS_DRAW_MESH_EVENTS
+			RHI_BREADCRUMB_EVENT_CONDITIONAL(RHICmdList, GShowMaterialDrawEvents != 0, "%s %s (%u instances)"
+				, VisibleMeshDrawCommand.MeshDrawCommand->GetDebugData().MaterialRenderProxy->GetMaterialName()
+				, VisibleMeshDrawCommand.MeshDrawCommand->GetDebugData().ResourceName
+				, VisibleMeshDrawCommand.MeshDrawCommand->NumInstances * InstanceFactor
+			);
+#endif
+			const bool bAllowSkipDrawCommand = true;
+			if (FMeshDrawCommand::SubmitDrawBegin(*VisibleMeshDrawCommand.MeshDrawCommand, GraphicsMinimalPipelineStateSet, SceneArgs, InstanceFactor, RHICmdList, StateCache, bAllowSkipDrawCommand))
 			{
+				FMeshDrawCommand::SubmitDrawEnd(*VisibleMeshDrawCommand.MeshDrawCommand, SceneArgs, InstanceFactor, RHICmdList);
+
 				// If MDC was split to a more than one batch, submit them without changing state
 				for (uint32 BatchIdx = 1; BatchIdx < DrawCommandInfo.NumBatches; ++BatchIdx)
 				{
