Some VSM cleanup, fixes and refactor
 - Added per-clipmap config to decouple from global cvars such that different clipmaps can have different properties.
 - Made it possible to create more than one VSM per light.
 - Fixed the on light-remove handling to actually remove directional lights correctly
 - Fixed incorrect logic for determining if a light is uncached for nanite material visibility (look at the cache entry rather than the global force invalidation flag).
  - Removed some dead code.

diff --git a/Engine/Shaders/Private/VirtualShadowMaps/VirtualShadowMapPageMarking.usf b/Engine/Shaders/Private/VirtualShadowMaps/VirtualShadowMapPageMarking.usf
--- a/Engine/Shaders/Private/VirtualShadowMaps/VirtualShadowMapPageMarking.usf
+++ b/Engine/Shaders/Private/VirtualShadowMaps/VirtualShadowMapPageMarking.usf
@@ -280,11 +280,11 @@ void GeneratePageFlagsFromPixels(
 		bIsValid = SubstratePixelHeader.IsSubstrateMaterial();
 		bBackfaceCull = bBackfaceCull && !SubstratePixelHeader.HasSubsurface();
 	#else // !SUBSTRATE_ENABLED
-		FGBufferData GBuffer = GetGBufferDataUint(PixelPos, true);
-		WorldNormal = GBuffer.WorldNormal;
+		FGBufferData GBufferData = GetGBufferDataUint(PixelPos, true);
+		WorldNormal = GBufferData.WorldNormal;
 		// Excluding unlit to avoid including processing sky dome
-		bIsValid = GBuffer.ShadingModelID != SHADINGMODELID_UNLIT;
-		bBackfaceCull = bBackfaceCull && !IsSubsurfaceModel(GBuffer.ShadingModelID);
+		bIsValid = GBufferData.ShadingModelID != SHADINGMODELID_UNLIT;
+		bBackfaceCull = bBackfaceCull && !IsSubsurfaceModel(GBufferData.ShadingModelID);
 	#endif // SUBSTRATE_ENABLED
 
 	#if PERMUTATION_WATER_DEPTH
@@ -401,16 +401,16 @@ void GeneratePageFlagsFromPixels(
 			int VirtualShadowMapId = DirectionalLightIds[Index];
 			FVirtualShadowMapProjectionShaderData ProjectionData = GetVirtualShadowMapProjectionData(VirtualShadowMapId);
 
-			bool bBackfaceCulledGBuffer = false;
+			bool bMarkPrimary = true;
 			#if PERMUTATION_INPUT_TYPE != INPUT_TYPE_HAIRSTRANDS
 			// Backface test if requested
 			if (bBackfaceCull && IsBackfaceToDirectionalLight(WorldNormal, ProjectionData.LightDirection, ProjectionData.LightSourceRadius))
 			{
-				bBackfaceCulledGBuffer = true;
+				bMarkPrimary = false;
 			}
 			#endif
 
-			if (!bBackfaceCulledGBuffer)
+			if (bMarkPrimary)
 			{
 				MarkPageClipmap(ProjectionData, bUsePageDilation, PageDilationOffset, Primary.TranslatedWorldPosition);
 			}
@@ -670,8 +670,6 @@ void GeneratePageFlagsFromFroxelsCS(uint FroxelIndex : SV_DispatchThreadID)
 #endif 
 
 uint bMarkCoarsePagesLocal;
-int ClipmapFirstLevel;
-uint ClipmapIndexMask;
 uint bIncludeNonNaniteGeometry;
 
 [numthreads(VSM_DEFAULT_CS_GROUP_X, 1, 1)]
@@ -710,8 +708,7 @@ void MarkCoarsePages(uint DispatchThreadId : SV_DispatchThreadID)
 		// We could of course mark a broader set of pages in the coarses clipmap level, but the effective radius
 		// even from just marking a single one is usually already large enough for the systems that need this
 		// data (volumetric fog, translucent light volume).
-		uint ClipmapIndex = uint(ProjectionData.ClipmapLevel - ClipmapFirstLevel);
-		if (((ClipmapIndexMask >> ClipmapIndex) & 1) != 0)
+		if (ProjectionData.bIsCoarseClipLevel)
 		{
 			// TODO: Optimize this... can be boiled down to be just in terms of the snap offsets
 			float3 OriginTranslatedWorld = ProjectionData.ClipmapWorldOriginOffset;
diff --git a/Engine/Shaders/Private/VirtualShadowMaps/VirtualShadowMapProjectionStructs.ush b/Engine/Shaders/Private/VirtualShadowMaps/VirtualShadowMapProjectionStructs.ush
--- a/Engine/Shaders/Private/VirtualShadowMaps/VirtualShadowMapProjectionStructs.ush
+++ b/Engine/Shaders/Private/VirtualShadowMaps/VirtualShadowMapProjectionStructs.ush
@@ -35,6 +35,7 @@ struct FVirtualShadowMapProjectionShaderData
 	int VirtualShadowMapId;
 	bool bUnCached; // see VSM_PROJ_FLAG_UNCACHED
 	bool bUnreferenced; // see VSM_PROJ_FLAG_UNREFERENCED
+	bool bIsCoarseClipLevel; // see VSM_PROJ_FLAG_IS_COARSE_CLIP_LEVEL
 };
 
 FVirtualShadowMapProjectionShaderData DecodeVirtualShadowMapProjectionData(ByteAddressBuffer ProjectionData, int VirtualShadowMapId)
@@ -91,6 +92,7 @@ FVirtualShadowMapProjectionShaderData DecodeVirtualShadowMapProjectionData(ByteA
 
 	Result.bUnCached = (Result.Flags & VSM_PROJ_FLAG_UNCACHED) != 0U;
 	Result.bUnreferenced = (Result.Flags & VSM_PROJ_FLAG_UNREFERENCED) != 0U;
+	Result.bIsCoarseClipLevel = (Result.Flags & VSM_PROJ_FLAG_IS_COARSE_CLIP_LEVEL) != 0U;
 
 	return Result;
 }
diff --git a/Engine/Shaders/Shared/VirtualShadowMapDefinitions.h b/Engine/Shaders/Shared/VirtualShadowMapDefinitions.h
--- a/Engine/Shaders/Shared/VirtualShadowMapDefinitions.h
+++ b/Engine/Shaders/Shared/VirtualShadowMapDefinitions.h
@@ -37,6 +37,7 @@ VirtualShadowMapDefinitions.h: used in virtual shadow map shaders and C++ code t
 #define VSM_PROJ_FLAG_CURRENT_DISTANT_LIGHT (1U << 0)
 #define VSM_PROJ_FLAG_UNCACHED (1U << 1) // Used to indicate that the light is uncached and should only render to dynamic pages
 #define VSM_PROJ_FLAG_UNREFERENCED (1U << 2) // Used to indicate that the light is not referenced/rendered to this render
+#define VSM_PROJ_FLAG_IS_COARSE_CLIP_LEVEL (1U << 3) // Used to indicate that the clip level is a coarse level
 
 // Hard limit for max distant lights supported 8k for now - we may revise later. We need to keep them in a fixed range for now to make allocation easy and minimize overhead for indexing.
 #define VSM_MAX_SINGLE_PAGE_SHADOW_MAPS (1024U * 8U)
diff --git a/Engine/Source/Runtime/Renderer/Private/RendererScene.cpp b/Engine/Source/Runtime/Renderer/Private/RendererScene.cpp
--- a/Engine/Source/Runtime/Renderer/Private/RendererScene.cpp
+++ b/Engine/Source/Runtime/Renderer/Private/RendererScene.cpp
@@ -3920,12 +3920,6 @@ void FScene::RemoveLightSceneInfo_RenderThread(FLightSceneInfo* LightSceneInfo)
 	// Remove the light from the lights list.
 	Lights.RemoveAt(LightSceneInfo->Id);
 
-	// TODO: move this work to FShadowScene & batch the light removals
-	if (auto CacheManager = GetVirtualShadowMapCache())
-	{
-		CacheManager->OnLightRemoved(LightSceneInfo->Id);
-	}
-
 	if (!LightSceneInfo->Proxy->HasStaticShadowing()
 		&& LightSceneInfo->Proxy->CastsDynamicShadow()
 		&& LightSceneInfo->GetDynamicShadowMapChannel() == -1)
diff --git a/Engine/Source/Runtime/Renderer/Private/SceneRendering.h b/Engine/Source/Runtime/Renderer/Private/SceneRendering.h
--- a/Engine/Source/Runtime/Renderer/Private/SceneRendering.h
+++ b/Engine/Source/Runtime/Renderer/Private/SceneRendering.h
@@ -169,9 +169,6 @@ class FVisibleLightViewInfo
 {
 public:
 
-	/** The dynamic primitives which are both visible and affected by this light. */
-	TArray<FPrimitiveSceneInfo*,SceneRenderingAllocator> VisibleDynamicLitPrimitives;
-	
 	/** Whether each shadow in the corresponding FVisibleLightInfo::AllProjectedShadows array is visible. */
 	FSceneBitArray ProjectedShadowVisibilityMap;
 
diff --git a/Engine/Source/Runtime/Renderer/Private/ShadowRendering.h b/Engine/Source/Runtime/Renderer/Private/ShadowRendering.h
--- a/Engine/Source/Runtime/Renderer/Private/ShadowRendering.h
+++ b/Engine/Source/Runtime/Renderer/Private/ShadowRendering.h
@@ -659,7 +659,7 @@ public:
 	/**
 	 * Adds a primitive to the shadow's subject list.
 	 */
-	bool AddSubjectPrimitive(FDynamicShadowsTaskData& TaskData, FPrimitiveSceneInfo* PrimitiveSceneInfo, TArrayView<FViewInfo> ViewArray, bool bRecordShadowSubjectForMobileShading);
+	bool AddSubjectPrimitive(FPrimitiveSceneInfo* PrimitiveSceneInfo, TArrayView<FViewInfo> ViewArray, bool bRecordShadowSubjectForMobileShading);
 
 	uint64 AddSubjectPrimitive_AnyThread(
 		const FPrimitiveSceneInfoCompact& PrimitiveSceneInfoCompact,
diff --git a/Engine/Source/Runtime/Renderer/Private/ShadowSetup.cpp b/Engine/Source/Runtime/Renderer/Private/ShadowSetup.cpp
--- a/Engine/Source/Runtime/Renderer/Private/ShadowSetup.cpp
+++ b/Engine/Source/Runtime/Renderer/Private/ShadowSetup.cpp
@@ -1944,7 +1944,7 @@ bool FProjectedShadowInfo::ShouldDrawStaticMeshes_AnyThread(
 	return bDrawingStaticMeshes;
 }
 
-bool FProjectedShadowInfo::AddSubjectPrimitive(FDynamicShadowsTaskData& TaskData, FPrimitiveSceneInfo* PrimitiveSceneInfo, TArrayView<FViewInfo> ViewArray, bool bShouldRecordShadowSubjectsForMobileShading)
+bool FProjectedShadowInfo::AddSubjectPrimitive(FPrimitiveSceneInfo* PrimitiveSceneInfo, TArrayView<FViewInfo> ViewArray, bool bShouldRecordShadowSubjectsForMobileShading)
 {
 	// Ray traced shadows use the GPU managed distance field object buffers, no CPU culling should be used
 	check(!bRayTracedDistanceField);
@@ -2357,7 +2357,7 @@ void FProjectedShadowInfo::FinalizeAddSubjectPrimitive(
 
 	if (Result.bAddOnRenderThread)
 	{
-		if (AddSubjectPrimitive(TaskData, PrimitiveSceneInfo, ViewArray, Result.bShouldRecordShadowSubjectsForMobile) && VirtualShadowMapClipmap.IsValid())
+		if (AddSubjectPrimitive(PrimitiveSceneInfo, ViewArray, Result.bShouldRecordShadowSubjectsForMobile) && VirtualShadowMapClipmap.IsValid())
 		{
 			VirtualShadowMapClipmap->OnPrimitiveRendered(PrimitiveSceneInfo);
 		}
@@ -3515,7 +3515,7 @@ void FSceneRenderer::CreatePerObjectProjectedShadow(
 						for (int32 ChildIndex = 0, ChildCount = ShadowGroupPrimitives.Num(); ChildIndex < ChildCount; ChildIndex++)
 						{
 							FPrimitiveSceneInfo* ShadowChild = ShadowGroupPrimitives[ChildIndex];
-							ProjectedShadowInfo->AddSubjectPrimitive(TaskData, ShadowChild, Views, false);
+							ProjectedShadowInfo->AddSubjectPrimitive(ShadowChild, Views, false);
 						}
 					}
 					else if (bShadowIsPotentiallyVisibleNextFrame)
@@ -3556,7 +3556,7 @@ void FSceneRenderer::CreatePerObjectProjectedShadow(
 						for (int32 ChildIndex = 0, ChildCount = ShadowGroupPrimitives.Num(); ChildIndex < ChildCount; ChildIndex++)
 						{
 							FPrimitiveSceneInfo* ShadowChild = ShadowGroupPrimitives[ChildIndex];
-							ProjectedShadowInfo->AddSubjectPrimitive(TaskData, ShadowChild, Views, false);
+							ProjectedShadowInfo->AddSubjectPrimitive(ShadowChild, Views, false);
 						}
 					}
 					else if (bShadowIsPotentiallyVisibleNextFrame)
@@ -4275,7 +4275,7 @@ void FSceneRenderer::CreateWholeSceneProjectedShadow(
 								FBoxSphereBounds const& Bounds = PrimitiveSceneInfo->Proxy->GetBounds();
 								if (IntersectsConvexHulls(LightViewFrustumConvexHulls, Bounds))
 								{
-									if (ProjectedShadowInfo->AddSubjectPrimitive(TaskData, Interaction->GetPrimitiveSceneInfo(), Views, false)
+									if (ProjectedShadowInfo->AddSubjectPrimitive(Interaction->GetPrimitiveSceneInfo(), Views, false)
 										&& VirtualSmCacheEntry.IsValid())
 									{
 										// NOTE: We don't track revealed primitives for local lights
diff --git a/Engine/Source/Runtime/Renderer/Private/Shadows/ShadowScene.cpp b/Engine/Source/Runtime/Renderer/Private/Shadows/ShadowScene.cpp
--- a/Engine/Source/Runtime/Renderer/Private/Shadows/ShadowScene.cpp
+++ b/Engine/Source/Runtime/Renderer/Private/Shadows/ShadowScene.cpp
@@ -3,6 +3,7 @@
 #include "ShadowScene.h"
 #include "ScenePrivate.h"
 #include "Algo/ForEach.h"
+#include "VirtualShadowMaps/VirtualShadowMapCacheManager.h"
 
 #if !(UE_BUILD_SHIPPING || UE_BUILD_TEST)
 #include "DynamicPrimitiveDrawing.h"
@@ -47,12 +48,23 @@ void FShadowScene::PostLightsUpdate(FRDGBuilder& GraphBuilder, const FLightScene
 		SceneChangeUpdateTask.Wait();
 		SceneChangeUpdateTask = GraphBuilder.AddSetupTask([this, LightSceneChangeSet]
 		{
+			// Collect a bit mask of all removed lights so we can prune the cache entries in one sweep
+			TBitArray<> RemovedLightMask;
+			RemovedLightMask.SetNum(Lights.GetMaxIndex(), false);
+
 			// Oust all removed Ids.
 			for (int32 Id : LightSceneChangeSet.RemovedLightIds)
 			{
 				check(Lights.IsValidIndex(Id));
 				ActiveLights[Id] = false;
 				Lights.RemoveAt(Id);
+				RemovedLightMask[Id] = true;
+			}
+
+			// scan the 
+			if (auto CacheManager = Scene.GetVirtualShadowMapCache())
+			{
+				CacheManager->ProcessRemovedLights(RemovedLightMask);
 			}
 	
 			// Track active lights (those that are or were moving, and thus need updating)
diff --git a/Engine/Source/Runtime/Renderer/Private/Shadows/ShadowSceneRenderer.cpp b/Engine/Source/Runtime/Renderer/Private/Shadows/ShadowSceneRenderer.cpp
--- a/Engine/Source/Runtime/Renderer/Private/Shadows/ShadowSceneRenderer.cpp
+++ b/Engine/Source/Runtime/Renderer/Private/Shadows/ShadowSceneRenderer.cpp
@@ -16,8 +16,6 @@
 
 CSV_DECLARE_CATEGORY_EXTERN(VSM);
 
-extern int32 GForceInvalidateDirectionalVSM;
-
 TAutoConsoleVariable<int32> CVarVSMMaterialVisibility(
 	TEXT("r.Shadow.Virtual.Nanite.MaterialVisibility"),
 	0,
@@ -422,7 +420,7 @@ void FShadowSceneRenderer::PostInitDynamicShadowsSetup()
 			FProjectedShadowInfo* ProjectedShadowInfo = DirectionalLightShadowFrameSetup.ProjectedShadowInfo;
 			if (!bUnboundedClipmap && ProjectedShadowInfo->bShouldRenderVSM)
 			{
-				const bool bIsCached = VirtualShadowMapArray.CacheManager->IsCacheEnabled() && GForceInvalidateDirectionalVSM == 0;
+				const bool bIsCached = !ProjectedShadowInfo->VirtualShadowMapPerLightCacheEntry->IsUncached();
 
 				// We can only do this culling if the light is both uncached & it is using the accurate bounds (i.e., r.Shadow.Virtual.Clipmap.UseConservativeCulling is turned off).
 				if (!bIsCached && !ProjectedShadowInfo->CascadeSettings.ShadowBoundsAccurate.Planes.IsEmpty())
@@ -664,9 +662,7 @@ void FShadowSceneRenderer::ApplyVirtualShadowMapProjectionForLight(
 		return;
 	}
 
-	FSceneTextureParameters SceneTextureParameters = GetSceneTextureParameters(GraphBuilder, SceneTextures.UniformBuffer);
-
-	for (int32 ViewIndex = 0; ViewIndex < SceneRenderer.Views.Num(); ViewIndex++)
+	for (int32 ViewIndex = 0; ViewIndex < SceneRenderer.Views.Num(); ++ViewIndex)
 	{
 		RDG_EVENT_SCOPE_CONDITIONAL(GraphBuilder, SceneRenderer.Views.Num() > 1, "View%d", ViewIndex);
 
diff --git a/Engine/Source/Runtime/Renderer/Private/VirtualShadowMaps/VirtualShadowMapArray.cpp b/Engine/Source/Runtime/Renderer/Private/VirtualShadowMaps/VirtualShadowMapArray.cpp
--- a/Engine/Source/Runtime/Renderer/Private/VirtualShadowMaps/VirtualShadowMapArray.cpp
+++ b/Engine/Source/Runtime/Renderer/Private/VirtualShadowMaps/VirtualShadowMapArray.cpp
@@ -38,7 +38,6 @@ DECLARE_DWORD_COUNTER_STAT(TEXT("VSM Nanite Views (Primary)"), STAT_VSMNaniteVie
 DECLARE_DWORD_COUNTER_STAT(TEXT("VSM Single Page Count"), STAT_VSMSinglePageCount, STATGROUP_ShadowRendering);
 DECLARE_DWORD_COUNTER_STAT(TEXT("VSM Full Count"), STAT_VSMFullCount, STATGROUP_ShadowRendering);
 
-extern int32 GForceInvalidateDirectionalVSM;
 extern int32 GVSMMaxPageAgeSinceLastRequest;
 extern TAutoConsoleVariable<float> CVarNaniteMaxPixelsPerEdge;
 extern TAutoConsoleVariable<float> CVarNaniteMinPixelsPerEdgeHW;
@@ -980,8 +979,6 @@ class FMarkCoarsePagesCS : public FVirtualShadowMapPageManagementShader
 		SHADER_PARAMETER_RDG_BUFFER_UAV(RWStructuredBuffer<uint>, OutPageRequestFlags)
 		SHADER_PARAMETER(uint32, bMarkCoarsePagesLocal)
 		SHADER_PARAMETER(uint32, bIncludeNonNaniteGeometry)
-		SHADER_PARAMETER(int32, ClipmapFirstLevel)
-		SHADER_PARAMETER(uint32, ClipmapIndexMask)
 	END_SHADER_PARAMETER_STRUCT()
 };
 IMPLEMENT_GLOBAL_SHADER(FMarkCoarsePagesCS, "/Engine/Private/VirtualShadowMaps/VirtualShadowMapPageMarking.usf", "MarkCoarsePages", SF_Compute);
@@ -1947,8 +1944,6 @@ void FVirtualShadowMapArray::BuildPageAllocations(
 				PassParameters->VirtualShadowMap = GetUncachedUniformBuffer(GraphBuilder);
 				PassParameters->OutPageRequestFlags = GraphBuilder.CreateUAV(PageRequestFlagsRDG);
 				PassParameters->bMarkCoarsePagesLocal = bMarkCoarsePagesLocal ? 1 : 0;
-				PassParameters->ClipmapFirstLevel = FVirtualShadowMapClipmap::GetFirstLevel();
-				PassParameters->ClipmapIndexMask = bMarkCoarsePagesDirectional ? FVirtualShadowMapClipmap::GetCoarsePageClipmapIndexMask() : 0;
 				PassParameters->bIncludeNonNaniteGeometry = CVarCoarsePagesIncludeNonNanite.GetValueOnRenderThread();
 
 				auto ComputeShader = View.ShaderMap->GetShader<FMarkCoarsePagesCS>();
diff --git a/Engine/Source/Runtime/Renderer/Private/VirtualShadowMaps/VirtualShadowMapCacheManager.cpp b/Engine/Source/Runtime/Renderer/Private/VirtualShadowMaps/VirtualShadowMapCacheManager.cpp
--- a/Engine/Source/Runtime/Renderer/Private/VirtualShadowMaps/VirtualShadowMapCacheManager.cpp
+++ b/Engine/Source/Runtime/Renderer/Private/VirtualShadowMaps/VirtualShadowMapCacheManager.cpp
@@ -74,13 +74,6 @@ FAutoConsoleVariableRef CVarCacheDebugSkipRevealedPrimitivesInvalidate(
 	TEXT("Debug skip invalidation of revealed Non-Nanite primitives, i.e. they go from being culled on the CPU to unculled."),
 	ECVF_RenderThreadSafe);
 
-int32 GForceInvalidateDirectionalVSM = 0;
-static FAutoConsoleVariableRef  CVarForceInvalidateDirectionalVSM(
-	TEXT("r.Shadow.Virtual.Cache.ForceInvalidateDirectional"),
-	GForceInvalidateDirectionalVSM,
-	TEXT("Forces the clipmap to always invalidate, useful to emulate a moving sun to avoid misrepresenting cache performance."),
-	ECVF_RenderThreadSafe);
-
 // NOTE: At this point it should be fairly safe and minimal performance impact to have this
 // "functionally unlimited", but we'll leave the default somewhat lower as a small mitigation
 // for unforeseen issues.
@@ -295,12 +288,12 @@ void FVirtualShadowMapCacheEntry::SetHZBViewParams(Nanite::FPackedViewParams& Ou
 	OutParams.Flags |= NANITE_VIEW_FLAG_HZBTEST;
 }
 
-void FVirtualShadowMapPerLightCacheEntry::UpdateClipmap(const FVector& LightDirection, int FirstLevel)
+void FVirtualShadowMapPerLightCacheEntry::UpdateClipmap(const FVector& LightDirection, int FirstLevel, bool bForceInvalidate)
 {
 	Prev.RenderedFrameNumber = FMath::Max(Prev.RenderedFrameNumber, Current.RenderedFrameNumber);
 	Current.RenderedFrameNumber = -1;
 
-	if (GForceInvalidateDirectionalVSM != 0 ||
+	if (bForceInvalidate ||
 		LightDirection != ClipmapCacheKey.LightDirection ||
 		FirstLevel != ClipmapCacheKey.FirstLevel)
 	{
@@ -309,7 +302,7 @@ void FVirtualShadowMapPerLightCacheEntry::UpdateClipmap(const FVector& LightDire
 	ClipmapCacheKey.LightDirection = LightDirection;
 	ClipmapCacheKey.FirstLevel = FirstLevel;
 
-	bool bNewIsUncached = GForceInvalidateDirectionalVSM != 0 || Prev.RenderedFrameNumber < 0;
+	bool bNewIsUncached = bForceInvalidate || Prev.RenderedFrameNumber < 0;
 
 	// On transition between uncached <-> cached we must invalidate since the static pages may not be initialized
 	if (bNewIsUncached != bIsUncached)
@@ -1018,11 +1011,11 @@ FRDGBufferRef FVirtualShadowMapArrayCacheManager::UploadCachePrimitiveAsDynamic(
 }
 
 TSharedPtr<FVirtualShadowMapPerLightCacheEntry> FVirtualShadowMapArrayCacheManager::FindCreateLightCacheEntry(
-	int32 LightSceneId, uint32 ViewUniqueID, uint32 NumShadowMaps)
+	int32 LightSceneId, uint32 ViewUniqueID, uint32 NumShadowMaps, uint32 TypeIdTag)
 {
-	const FVirtualShadowMapCacheKey CacheKey = { ViewUniqueID, LightSceneId };
+	const FVirtualShadowMapCacheKey CacheKey = { ViewUniqueID, LightSceneId, TypeIdTag };
 
-	TSharedPtr<FVirtualShadowMapPerLightCacheEntry> *LightEntryKey = CacheEntries.Find(CacheKey);
+	TSharedPtr<FVirtualShadowMapPerLightCacheEntry>* LightEntryKey = CacheEntries.Find(CacheKey);
 
 	if (LightEntryKey)
 	{
@@ -1567,10 +1560,16 @@ void FVirtualShadowMapArrayCacheManager::ProcessInvalidations(FRDGBuilder& Graph
 	}
 }
 
-void FVirtualShadowMapArrayCacheManager::OnLightRemoved(int32 LightId)
+void FVirtualShadowMapArrayCacheManager::ProcessRemovedLights(const TBitArray<>& RemovedLightMask)
 {
-	const FVirtualShadowMapCacheKey CacheKey = { /* TODO: this is broken for directional lights! ViewUniqueID */0, LightId };
-	CacheEntries.Remove(CacheKey);
+	for (auto It = CacheEntries.CreateIterator(); It; ++It)
+	{
+		int32 LightSceneId = It->Key.LightSceneId;
+		if (RemovedLightMask.IsValidIndex(LightSceneId) && RemovedLightMask[LightSceneId])
+		{
+			It.RemoveCurrent();
+		}
+	}
 }
 
 FVirtualShadowMapArrayCacheManager::FInvalidationPassCommon FVirtualShadowMapArrayCacheManager::GetUniformParametersForInvalidation(
diff --git a/Engine/Source/Runtime/Renderer/Private/VirtualShadowMaps/VirtualShadowMapCacheManager.h b/Engine/Source/Runtime/Renderer/Private/VirtualShadowMaps/VirtualShadowMapCacheManager.h
--- a/Engine/Source/Runtime/Renderer/Private/VirtualShadowMaps/VirtualShadowMapCacheManager.h
+++ b/Engine/Source/Runtime/Renderer/Private/VirtualShadowMaps/VirtualShadowMapCacheManager.h
@@ -96,7 +96,8 @@ public:
 
 	void MarkRendered(int32 FrameIndex) { Current.RenderedFrameNumber = FrameIndex; }
 	int32 GetLastScheduledFrameNumber() const { return Prev.ScheduledFrameNumber; }
-	void UpdateClipmap(const FVector& LightDirection, int FirstLevelx);
+	void UpdateClipmap(const FVector& LightDirection, int FirstLevelx, bool bForceInvalidate);
+
 	/**
 	 * Returns true if the cache entry is valid (has previous state).
 	 */
@@ -212,13 +213,14 @@ struct FVirtualShadowMapCacheKey
 {
 	uint32 ViewUniqueID;
 	uint32 LightSceneId;
+	uint32 ShadowTypeId;
 
-	inline bool operator==(const FVirtualShadowMapCacheKey& Other) const { return ViewUniqueID == Other.ViewUniqueID && Other.LightSceneId == LightSceneId; }
+	inline bool operator==(const FVirtualShadowMapCacheKey& Other) const { return ViewUniqueID == Other.ViewUniqueID && Other.LightSceneId == LightSceneId && Other.ShadowTypeId == ShadowTypeId; }
 };
 
 inline uint32 GetTypeHash(FVirtualShadowMapCacheKey Key)
 {
-	return GetTypeHash(Key.LightSceneId) ^ GetTypeHash(Key.ViewUniqueID);
+	return HashCombineFast(GetTypeHash(Key.LightSceneId), HashCombineFast(GetTypeHash(Key.ViewUniqueID), GetTypeHash(Key.ShadowTypeId)));
 }
 
 class FVirtualShadowMapArrayCacheManager : public ISceneExtension
@@ -273,8 +275,9 @@ public:
 
 	/**
 	 * Finds an existing cache entry and moves to the active set or creates a fresh one.
+	 * TypeIdTag is an arbitrary type ID to make it possible to have more than one shadow map for the same light & view, it is up to the user to make sure there are no collisions.
 	 */
-	TSharedPtr<FVirtualShadowMapPerLightCacheEntry> FindCreateLightCacheEntry(int32 LightSceneId, uint32 ViewUniqueID, uint32 NumShadowMaps);
+	TSharedPtr<FVirtualShadowMapPerLightCacheEntry> FindCreateLightCacheEntry(int32 LightSceneId, uint32 ViewUniqueID, uint32 NumShadowMaps, uint32 TypeIdTag = 0u);
 
 	bool IsCacheEnabled();
 	bool IsCacheDataAvailable();
@@ -339,7 +342,7 @@ public:
 	/**
 	 * Handle light removal, need to clear out cache entries as the ID may be reused after this.
 	 */
-	void OnLightRemoved(int32 LightId);
+	void ProcessRemovedLights(const TBitArray<>& RemovedLightMask);
 
 	uint64 GetGPUSizeBytes(bool bLogSizes) const;
 
diff --git a/Engine/Source/Runtime/Renderer/Private/VirtualShadowMaps/VirtualShadowMapClipmap.cpp b/Engine/Source/Runtime/Renderer/Private/VirtualShadowMaps/VirtualShadowMapClipmap.cpp
--- a/Engine/Source/Runtime/Renderer/Private/VirtualShadowMaps/VirtualShadowMapClipmap.cpp
+++ b/Engine/Source/Runtime/Renderer/Private/VirtualShadowMaps/VirtualShadowMapClipmap.cpp
@@ -19,7 +19,12 @@ VirtualShadowMapClipmap.cpp
 #include "Engine/HitResult.h"
 #include "Engine/World.h"
 
-extern int32 GForceInvalidateDirectionalVSM;
+
+static TAutoConsoleVariable<int32>  CVarForceInvalidateDirectionalVSM(
+	TEXT("r.Shadow.Virtual.Cache.ForceInvalidateDirectional"),
+	0,
+	TEXT("Forces the clipmap to always invalidate, useful to emulate a moving sun to avoid misrepresenting cache performance."),
+	ECVF_RenderThreadSafe);
 
 static TAutoConsoleVariable<float> CVarVirtualShadowMapResolutionLodBiasDirectional(
 	TEXT( "r.Shadow.Virtual.ResolutionLodBiasDirectional" ),
@@ -131,14 +136,29 @@ static float GetLevelRadius(int32 Level)
 	return FMath::Pow(2.0f, static_cast<float>(Level + 1));
 }
 
+FVirtualShadowMapClipmapConfig FVirtualShadowMapClipmapConfig::GetGlobal()
+{
+	return FVirtualShadowMapClipmapConfig 
+	{ 
+		.FirstLevel = CVarVirtualShadowMapClipmapFirstLevel.GetValueOnRenderThread(),
+		.LastLevel = CVarVirtualShadowMapClipmapLastLevel.GetValueOnRenderThread(),
+		.FirstCoarseLevel = CVarVirtualShadowMapClipmapFirstCoarseLevel.GetValueOnRenderThread(),
+		.LastCoarseLevel = CVarVirtualShadowMapClipmapLastCoarseLevel.GetValueOnRenderThread(),
+		.bForceInvalidate = CVarForceInvalidateDirectionalVSM.GetValueOnRenderThread() != 0
+	};
+}
+
+
 FVirtualShadowMapClipmap::FVirtualShadowMapClipmap(
 	FVirtualShadowMapArray& VirtualShadowMapArray,
 	const FLightSceneInfo& InLightSceneInfo,
 	const FViewMatrices& CameraViewMatrices,
 	FIntPoint CameraViewRectSize,
 	const FViewInfo* InDependentView,
-	float LightMobilityFactor)
-	: LightSceneInfo(InLightSceneInfo),
+	float LightMobilityFactor,
+	const FVirtualShadowMapClipmapConfig& InConfig)
+	: Config(InConfig),
+	  LightSceneInfo(InLightSceneInfo),
 	  DependentView(InDependentView)
 {
 	FVirtualShadowMapArrayCacheManager* VirtualShadowMapArrayCacheManager = VirtualShadowMapArray.CacheManager;
@@ -227,8 +247,8 @@ FVirtualShadowMapClipmap::FVirtualShadowMapClipmap(
 		WorldOrigin += CameraToViewTarget;
 	}
 
-	FirstLevel = GetFirstLevel();
-	int32 LastLevel = CVarVirtualShadowMapClipmapLastLevel.GetValueOnRenderThread();	
+	FirstLevel = Config.FirstLevel;
+	int32 LastLevel = Config.LastLevel;	
 	if (bIsOrthographicCamera && CVarOrthoVSMEstimateClipmapLevels.GetValueOnRenderThread())
 	{
 		/**
@@ -261,8 +281,8 @@ FVirtualShadowMapClipmap::FVirtualShadowMapClipmap(
 	// TODO: We need a light/cache entry for every light/VSM now, but scene captures may not have persistent view state for indexing
 	// This is likely to all change as we refactor the multiple cache managers stuff anyways; for now this is hopefully safe since they do have separate copies
 	const uint32 UniqueViewKey = InDependentView->ViewState ? InDependentView->ViewState->GetViewKey() : 0U;
-	PerLightCacheEntry = VirtualShadowMapArrayCacheManager->FindCreateLightCacheEntry(LightSceneInfo.Id, UniqueViewKey, LevelCount);
-	PerLightCacheEntry->UpdateClipmap(LightDirection, FirstLevel);
+	PerLightCacheEntry = VirtualShadowMapArrayCacheManager->FindCreateLightCacheEntry(LightSceneInfo.Id, UniqueViewKey, LevelCount, Config.ShadowTypeId);
+	PerLightCacheEntry->UpdateClipmap(LightDirection, FirstLevel, Config.bForceInvalidate);
 
 	const int RadiusLn = GetLevelRadius(LastLevel);
 	const int RadiiPerLevel = 4;
@@ -453,41 +473,19 @@ FVirtualShadowMapProjectionShaderData FVirtualShadowMapClipmap::ComputeProjectio
 	Data.ClipmapLevelWPODistanceDisableThresholdSquared = static_cast<float>(Level.WPODistanceDisableThresholdSquared);
 	Data.LightSourceRadius = GetLightSceneInfo().Proxy->GetSourceRadius();
 	Data.Flags = PerLightCacheEntry->IsUncached() ? VSM_PROJ_FLAG_UNCACHED : 0U;
-	Data.TexelDitherScale = GetLightSceneInfo().Proxy->GetVSMTexelDitherScale();
-	return Data;
-}
 
-int32 FVirtualShadowMapClipmap::GetFirstLevel()
-{
-	return CVarVirtualShadowMapClipmapFirstLevel.GetValueOnRenderThread();
-}
-
-uint32 FVirtualShadowMapClipmap::GetCoarsePageClipmapIndexMask()
-{
-	uint32 BitMask = 0;
-
-	const int FirstLevel = GetFirstLevel();
-	const int LastLevel  = FMath::Max(FirstLevel, CVarVirtualShadowMapClipmapLastLevel.GetValueOnRenderThread());	
-	int FirstCoarseIndex = CVarVirtualShadowMapClipmapFirstCoarseLevel.GetValueOnRenderThread() - FirstLevel;
-	int LastCoarseIndex  = CVarVirtualShadowMapClipmapLastCoarseLevel.GetValueOnRenderThread() - FirstLevel;	
-
-	ensureMsgf((LastLevel - FirstLevel) < 32, TEXT("Too many clipmap levels for coarse page bitmask."));
-
-	FirstCoarseIndex = FMath::Max(0, FirstCoarseIndex);
-	if (LastCoarseIndex >= FirstCoarseIndex)
+	if ((Data.ClipmapLevel >= Config.FirstCoarseLevel && Data.ClipmapLevel <= Config.LastCoarseLevel)
+		// Always mark coarse pages in the last level for clouds/skyatmosphere
+		|| 	Data.ClipmapLevelCountRemaining == 1)
 	{
-		uint32 BitCount = static_cast<uint32>(LastCoarseIndex - FirstCoarseIndex + 1);
-		uint32 BitRange = (1 << BitCount) - 1;
-		BitMask = BitMask | (BitRange << FirstCoarseIndex);
+		Data.Flags |= VSM_PROJ_FLAG_IS_COARSE_CLIP_LEVEL;
 	}
 
-	// Always mark coarse pages in the last level for clouds/skyatmosphere
-	BitMask = BitMask | (1 << (LastLevel - FirstLevel));
 
-	return BitMask;
+	Data.TexelDitherScale = GetLightSceneInfo().Proxy->GetVSMTexelDitherScale();
+	return Data;
 }
 
-
 static inline void LazyInitAndSetBitArray(TBitArray<>& BitArray, int32 Index, bool Value, int32 MaxNum)
 {
 	if (BitArray.IsEmpty())
diff --git a/Engine/Source/Runtime/Renderer/Private/VirtualShadowMaps/VirtualShadowMapClipmap.h b/Engine/Source/Runtime/Renderer/Private/VirtualShadowMaps/VirtualShadowMapClipmap.h
--- a/Engine/Source/Runtime/Renderer/Private/VirtualShadowMaps/VirtualShadowMapClipmap.h
+++ b/Engine/Source/Runtime/Renderer/Private/VirtualShadowMaps/VirtualShadowMapClipmap.h
@@ -19,6 +19,18 @@ class FVirtualShadowMapArray;
 class FVirtualShadowMapArrayCacheManager;
 class FVirtualShadowMapPerLightCacheEntry;
 
+struct FVirtualShadowMapClipmapConfig
+{
+	static FVirtualShadowMapClipmapConfig GetGlobal();
+
+	int32 FirstLevel = 8;
+	int32 LastLevel = 18;
+	int32 FirstCoarseLevel = -1;
+	int32 LastCoarseLevel = -1;
+	uint32 ShadowTypeId = 0u;
+	bool bForceInvalidate = false;
+};
+
 class FVirtualShadowMapClipmap : FRefCountedObject
 {
 public:	
@@ -28,7 +40,8 @@ public:
 		const FViewMatrices& CameraViewMatrices,
 		FIntPoint CameraViewRectSize,
 		const FViewInfo* InDependentView,
-		float LightMobilityFactor
+		float LightMobilityFactor,
+		const FVirtualShadowMapClipmapConfig& Config = FVirtualShadowMapClipmapConfig::GetGlobal()
 	);
 
 	FViewMatrices GetViewMatrices(int32 ClipmapIndex) const;
@@ -84,11 +97,6 @@ public:
 
 	const FViewInfo* GetDependentView() const { return DependentView; }
 
-	static int32 GetFirstLevel();
-	// Returns a mask with one bit per level of which coarse pages to mark (based on cvars)
-	// Bits relative to FirstLevel (i.e. in terms of ClipmapIndex, not ClipmapLevel)	
-	static uint32 GetCoarsePageClipmapIndexMask();
-
 	/**
 	 * Called when a primitive passes CPU-culling, note that this applies to non-nanite primitives only. Not thread safe in general.
 	 */ 
@@ -105,6 +113,8 @@ private:
 	FVirtualShadowMapProjectionShaderData ComputeProjectionShaderData(int32 ClipmapIndex) const;
 	void ComputeBoundingVolumes(const FVector CameraOrigin);
 
+	FVirtualShadowMapClipmapConfig Config;
+
 	const FLightSceneInfo& LightSceneInfo;
 
 	/**
