Enable LightGrid Debug for Mobile
Keep Light Grid Debug disabled for Preview and D3D ES31 because of FXC restriction

diff --git a/Engine/Source/Runtime/Renderer/Private/PostProcess/PostProcessing.cpp b/Engine/Source/Runtime/Renderer/Private/PostProcess/PostProcessing.cpp
--- a/Engine/Source/Runtime/Renderer/Private/PostProcess/PostProcessing.cpp
+++ b/Engine/Source/Runtime/Renderer/Private/PostProcess/PostProcessing.cpp
@@ -2300,6 +2300,7 @@ void AddMobilePostProcessingPasses(FRDGBuilder& GraphBuilder, FScene* Scene, con
 		PrimaryUpscale,
 		SecondaryUpscale,
 		Visualize,
+		VisualizeLightGrid,
 		HMDDistortion,
 		MAX
 	};
@@ -2327,6 +2328,7 @@ void AddMobilePostProcessingPasses(FRDGBuilder& GraphBuilder, FScene* Scene, con
 		TEXT("PrimaryUpscale"),
 		TEXT("SecondaryUpscale"),
 		TEXT("Visualize"),
+		TEXT("VisualizeLightGrid"),
 		TEXT("HMDDistortion")
 	};
 
@@ -2387,7 +2389,7 @@ void AddMobilePostProcessingPasses(FRDGBuilder& GraphBuilder, FScene* Scene, con
 	PassSequence.SetEnabled(EPass::SecondaryUpscale, View.Family->GetSecondarySpatialUpscalerInterface() != nullptr);
 
 	PassSequence.SetEnabled(EPass::Visualize, View.Family->EngineShowFlags.ShaderComplexity);
-
+	PassSequence.SetEnabled(EPass::VisualizeLightGrid, ShouldVisualizeLightGrid());
 	PassSequence.SetEnabled(EPass::HMDDistortion, View.Family->EngineShowFlags.StereoRendering && View.Family->EngineShowFlags.HMDDistortion);
 
 	// Always evaluate custom post processes
@@ -3115,6 +3117,13 @@ void AddMobilePostProcessingPasses(FRDGBuilder& GraphBuilder, FScene* Scene, con
 		}
 	}
 
+	if (PassSequence.IsEnabled(EPass::VisualizeLightGrid))
+	{
+		FScreenPassRenderTarget OverrideOutput;
+		PassSequence.AcceptOverrideIfLastPass(EPass::VisualizeLightGrid, OverrideOutput);
+		SceneColor = AddVisualizeLightGridPass(GraphBuilder, View, SceneColor, SceneDepth);
+	}
+
 	if (ShaderPrint::IsEnabled(View.ShaderPrintData))
 	{
 		ShaderPrint::DrawView(GraphBuilder, View, SceneColor, SceneDepth);
