Use GetDownscaledViewport(...) in AddDownsamplePass(...) instead of calculating extent and view rect manually.

diff --git a/Engine/Source/Runtime/Renderer/Private/PostProcess/PostProcessDownsample.cpp b/Engine/Source/Runtime/Renderer/Private/PostProcess/PostProcessDownsample.cpp
--- a/Engine/Source/Runtime/Renderer/Private/PostProcess/PostProcessDownsample.cpp
+++ b/Engine/Source/Runtime/Renderer/Private/PostProcess/PostProcessDownsample.cpp
@@ -108,16 +108,16 @@ FScreenPassTexture AddDownsamplePass(
 
 	// Construct the output texture to be half resolution (rounded up to even) with an optional format override.
 	{
+		FScreenPassTextureViewport SceneColorViewport(Inputs.SceneColor);
+		FScreenPassTextureViewport DownscaledViewport = GetDownscaledViewport(SceneColorViewport, 2);
+
 		const FRDGTextureDesc& InputDesc = Inputs.SceneColor.TextureSRV->Desc.Texture->Desc;
 		FRDGTextureDesc Desc = FRDGTextureDesc::Create2D(
-			FIntPoint::DivideAndRoundUp(InputDesc.Extent, 2),
+			DownscaledViewport.Extent,
 			Inputs.FormatOverride != PF_Unknown ? Inputs.FormatOverride  : InputDesc.Format,
 			FClearValueBinding(FLinearColor(0, 0, 0, 0)),
 			/* InFlags = */ TexCreate_ShaderResource | GFastVRamConfig.Downsample | (bIsComputePass ? TexCreate_UAV : (TexCreate_RenderTargetable | TexCreate_NoFastClear)));
 
-		Desc.Extent.X = FMath::Max(1, Desc.Extent.X);
-		Desc.Extent.Y = FMath::Max(1, Desc.Extent.Y);
-
 		if (Inputs.UserSuppliedOutput && Translate(Inputs.UserSuppliedOutput->GetDesc()) == Desc)
 		{
 			Output.Texture = GraphBuilder.RegisterExternalTexture(Inputs.UserSuppliedOutput, Inputs.Name);
@@ -126,7 +126,7 @@ FScreenPassTexture AddDownsamplePass(
 		{
 			Output.Texture = GraphBuilder.CreateTexture(Desc, Inputs.Name);
 		}
-		Output.ViewRect = FIntRect::DivideAndRoundUp(Inputs.SceneColor.ViewRect, 2);
+		Output.ViewRect = DownscaledViewport.Rect;
 		Output.LoadAction = ERenderTargetLoadAction::ENoAction;
 		Output.UpdateVisualizeTextureExtent();
 	}
