[Exposure Fusion] Support higher number of laplacian pyramid levels.
- exposed via r.LocalExposure.ExposureFusion.NumLevels.
- Refactored FSceneDownsampleChain to support variable StageCount and renamed to FTextureDownsampleChain.

diff --git a/Engine/Source/Runtime/Renderer/Private/PostProcess/PostProcessBloomSetup.cpp b/Engine/Source/Runtime/Renderer/Private/PostProcess/PostProcessBloomSetup.cpp
--- a/Engine/Source/Runtime/Renderer/Private/PostProcess/PostProcessBloomSetup.cpp
+++ b/Engine/Source/Runtime/Renderer/Private/PostProcess/PostProcessBloomSetup.cpp
@@ -194,13 +194,9 @@ EBloomQuality GetBloomQuality()
 		static_cast<int32>(EBloomQuality::MAX) - 1));
 }
 
-static_assert(
-	static_cast<uint32>(EBloomQuality::MAX) == FSceneDownsampleChain::StageCount,
-	"The total number of stages in the scene downsample chain and the number of bloom quality levels must match.");
-
-FScreenPassTexture AddGaussianBloomPasses(FRDGBuilder& GraphBuilder, const FViewInfo& View, const FSceneDownsampleChain* SceneDownsampleChain)
+FScreenPassTexture AddGaussianBloomPasses(FRDGBuilder& GraphBuilder, const FViewInfo& View, const FTextureDownsampleChain* SceneDownsampleChain)
 {
-	check(SceneDownsampleChain);
+	checkf(SceneDownsampleChain && SceneDownsampleChain->GetStageCount() == static_cast<uint32>(EBloomQuality::MAX), TEXT("The total number of stages in the scene downsample chain and the number of bloom quality levels must match."));
 	check(!IsFFTBloomEnabled(View));
 
 	const FFinalPostProcessSettings& Settings = View.FinalPostProcessSettings;
diff --git a/Engine/Source/Runtime/Renderer/Private/PostProcess/PostProcessBloomSetup.h b/Engine/Source/Runtime/Renderer/Private/PostProcess/PostProcessBloomSetup.h
--- a/Engine/Source/Runtime/Renderer/Private/PostProcess/PostProcessBloomSetup.h
+++ b/Engine/Source/Runtime/Renderer/Private/PostProcess/PostProcessBloomSetup.h
@@ -4,7 +4,7 @@
 
 #include "ScreenPass.h"
 
-class FSceneDownsampleChain;
+class FTextureDownsampleChain;
 class FEyeAdaptationParameters;
 class FLocalExposureParameters;
 
@@ -21,7 +21,7 @@ enum class EBloomQuality : uint32
 
 EBloomQuality GetBloomQuality();
 
-FScreenPassTexture AddGaussianBloomPasses(FRDGBuilder& GraphBuilder, const FViewInfo& View, const FSceneDownsampleChain* SceneDownsampleChain);
+FScreenPassTexture AddGaussianBloomPasses(FRDGBuilder& GraphBuilder, const FViewInfo& View, const FTextureDownsampleChain* SceneDownsampleChain);
 
 struct FBloomSetupInputs
 {
diff --git a/Engine/Source/Runtime/Renderer/Private/PostProcess/PostProcessDownsample.cpp b/Engine/Source/Runtime/Renderer/Private/PostProcess/PostProcessDownsample.cpp
--- a/Engine/Source/Runtime/Renderer/Private/PostProcess/PostProcessDownsample.cpp
+++ b/Engine/Source/Runtime/Renderer/Private/PostProcess/PostProcessDownsample.cpp
@@ -192,38 +192,34 @@ void AddDownsampleComputePass(FRDGBuilder& GraphBuilder, const FViewInfo& View,
 	return AddDownsampleComputePass(GraphBuilder, View, FScreenPassTextureSlice::CreateFromScreenPassTexture(GraphBuilder, Input), Output, Quality, PassFlags);
 }
 
-void FSceneDownsampleChain::Init(
+void FTextureDownsampleChain::Init(
 	FRDGBuilder& GraphBuilder,
 	const FViewInfo& View,
 	const FEyeAdaptationParameters& EyeAdaptationParameters,
-	FScreenPassTextureSlice HalfResolutionSceneColor,
+	FScreenPassTextureSlice InputTexture,
 	EDownsampleQuality DownsampleQuality,
-	bool bLogLumaInAlpha)
+	uint32 StageCount,
+	bool bLogLumaInAlpha,
+	const TCHAR* TextureName,
+	uint32 InputResolutionExponent)
 {
-	check(HalfResolutionSceneColor.IsValid());
+	check(InputTexture.IsValid());
 
-	RDG_EVENT_SCOPE(GraphBuilder, "SceneDownsample");
+	RDG_EVENT_SCOPE(GraphBuilder, "TextureDownsampleChain");
 
-	static const TCHAR* PassNames[StageCount] =
-	{
-		nullptr,
-		TEXT("Scene(1/4)"),
-		TEXT("Scene(1/8)"),
-		TEXT("Scene(1/16)"),
-		TEXT("Scene(1/32)"),
-		TEXT("Scene(1/64)")
-	};
-	static_assert(UE_ARRAY_COUNT(PassNames) == StageCount, "PassNames size must equal StageCount");
+	Textures.SetNum(StageCount);
 
 	// The first stage is the input.
-	Textures[0] = HalfResolutionSceneColor;
+	Textures[0] = InputTexture;
 
 	for (uint32 StageIndex = 1; StageIndex < StageCount; StageIndex++)
 	{
 		const uint32 PreviousStageIndex = StageIndex - 1;
 
+		FString PassName = FString::Printf(TEXT("%s(1/%u)"), TextureName ? TextureName : InputTexture.TextureSRV->Desc.Texture->Name, 1 << (StageIndex + InputResolutionExponent));
+
 		FDownsamplePassInputs PassInputs;
-		PassInputs.Name = PassNames[StageIndex];
+		PassInputs.Name = *PassName;
 		PassInputs.SceneColor = Textures[PreviousStageIndex];
 		PassInputs.Quality = DownsampleQuality;
 
diff --git a/Engine/Source/Runtime/Renderer/Private/PostProcess/PostProcessDownsample.h b/Engine/Source/Runtime/Renderer/Private/PostProcess/PostProcessDownsample.h
--- a/Engine/Source/Runtime/Renderer/Private/PostProcess/PostProcessDownsample.h
+++ b/Engine/Source/Runtime/Renderer/Private/PostProcess/PostProcessDownsample.h
@@ -56,21 +56,21 @@ void AddDownsampleComputePass(FRDGBuilder& GraphBuilder, const FViewInfo& View,
 void AddDownsampleComputePass(FRDGBuilder& GraphBuilder, const FViewInfo& View, FScreenPassTextureSlice Input, FScreenPassTexture Output, EDownsampleQuality Quality, ERDGPassFlags PassFlags);
 
 
-class FSceneDownsampleChain
+class FTextureDownsampleChain
 {
 public:
-	// The number of total stages in the chain. 1/64 reduction.
-	static const uint32 StageCount = 6;
-
-	FSceneDownsampleChain() = default;
+	FTextureDownsampleChain() = default;
 
 	void Init(
 		FRDGBuilder& GraphBuilder,
 		const FViewInfo& View,
 		const FEyeAdaptationParameters& EyeAdaptationParameters,
-		FScreenPassTextureSlice HalfResolutionSceneColor,
+		FScreenPassTextureSlice InputTexture,
 		EDownsampleQuality DownsampleQuality,
-		bool bLogLumaInAlpha);
+		uint32 StageCount,
+		bool bLogLumaInAlpha,
+		const TCHAR* TextureName = nullptr,
+		uint32 InputResolutionExponent = 0);
 
 	bool IsInitialized() const
 	{
@@ -89,10 +89,15 @@ public:
 
 	FScreenPassTextureSlice GetLastTexture() const
 	{
-		return Textures[StageCount - 1];
+		return Textures.Last();
+	}
+
+	uint32 GetStageCount() const
+	{
+		return Textures.Num();
 	}
 
 private:
-	TStaticArray<FScreenPassTextureSlice, StageCount> Textures;
+	TArray<FScreenPassTextureSlice> Textures;
 	bool bInitialized = false;
 };
\ No newline at end of file
diff --git a/Engine/Source/Runtime/Renderer/Private/PostProcess/PostProcessLensFlares.h b/Engine/Source/Runtime/Renderer/Private/PostProcess/PostProcessLensFlares.h
--- a/Engine/Source/Runtime/Renderer/Private/PostProcess/PostProcessLensFlares.h
+++ b/Engine/Source/Runtime/Renderer/Private/PostProcess/PostProcessLensFlares.h
@@ -4,8 +4,6 @@
 
 #include "ScreenPass.h"
 
-class FSceneDownsampleChain;
-
 enum class ELensFlareQuality : uint32
 {
 	Disabled,
diff --git a/Engine/Source/Runtime/Renderer/Private/PostProcess/PostProcessLocalExposure.cpp b/Engine/Source/Runtime/Renderer/Private/PostProcess/PostProcessLocalExposure.cpp
--- a/Engine/Source/Runtime/Renderer/Private/PostProcess/PostProcessLocalExposure.cpp
+++ b/Engine/Source/Runtime/Renderer/Private/PostProcess/PostProcessLocalExposure.cpp
@@ -19,6 +19,12 @@ static TAutoConsoleVariable<float> CVarExposureFusionTargetLuminance(
 	TEXT("Target Luminance used to determine the weight of each exposure."),
 	ECVF_RenderThreadSafe);
 
+static TAutoConsoleVariable<int32> CVarExposureFusionNumLevels(
+	TEXT("r.LocalExposure.ExposureFusion.NumLevels"),
+	16,
+	TEXT("Number of levels in the Laplacian pyramid used to blend the different exposures."),
+	ECVF_RenderThreadSafe);
+
 namespace
 {
 
@@ -411,27 +417,32 @@ FScreenPassTexture AddLocalExposureFusionPass(
 
 	const bool bLogLumaInAlpha = false;
 
-	FSceneDownsampleChain LumChain;
+	const uint32 MaxMips = FMath::Log2(float(FMath::Min(LumTexture.Texture->Desc.Extent.X, LumTexture.Texture->Desc.Extent.Y))) + 1;
+	const uint32 NumMips = FMath::Clamp((uint32)CVarExposureFusionNumLevels.GetValueOnRenderThread(), 1, MaxMips);
+
+	FTextureDownsampleChain LumChain;
 	LumChain.Init(
 		GraphBuilder, View,
 		EyeAdaptationParameters,
 		FScreenPassTextureSlice::CreateFromScreenPassTexture(GraphBuilder, LumTexture),
 		EDownsampleQuality::High,
+		NumMips,
 		bLogLumaInAlpha);
 
-	FSceneDownsampleChain WeightChain;
+	FTextureDownsampleChain WeightChain;
 	WeightChain.Init(
 		GraphBuilder, View,
 		EyeAdaptationParameters,
 		FScreenPassTextureSlice::CreateFromScreenPassTexture(GraphBuilder, WeightTexture),
 		EDownsampleQuality::High,
+		NumMips,
 		bLogLumaInAlpha);
 
 	FScreenPassTexture Output;
 
 	FScreenPassTextureSlice CoarserMip;
 
-	for(int32 Index = FSceneDownsampleChain::StageCount - 1; Index >= 0; --Index)
+	for(int32 Index = NumMips - 1; Index >= 0; --Index)
 	{
 		FScreenPassTextureSlice CurrentLum = LumChain.GetTexture(Index);
 		FScreenPassTextureSlice CurrentWeight = WeightChain.GetTexture(Index);
diff --git a/Engine/Source/Runtime/Renderer/Private/PostProcess/PostProcessing.cpp b/Engine/Source/Runtime/Renderer/Private/PostProcess/PostProcessing.cpp
--- a/Engine/Source/Runtime/Renderer/Private/PostProcess/PostProcessing.cpp
+++ b/Engine/Source/Runtime/Renderer/Private/PostProcess/PostProcessing.cpp
@@ -1130,7 +1130,7 @@ void AddPostProcessingPasses(
 				LastEyeAdaptationBuffer);
 		}
 
-		FSceneDownsampleChain SceneDownsampleChain;
+		FTextureDownsampleChain SceneDownsampleChain;
 		if (bProduceSceneColorChain)
 		{
 			const bool bLogLumaInAlpha = bBasicEyeAdaptationEnabled;
@@ -1139,7 +1139,10 @@ void AddPostProcessingPasses(
 				EyeAdaptationParameters,
 				bProcessEighthResolution ? EighthResSceneColor : (bProcessQuarterResolution ? QuarterResSceneColor : HalfResSceneColor),
 				DownsampleChainQuality,
-				bLogLumaInAlpha);
+				6,
+				bLogLumaInAlpha,
+				TEXT("Scene"),
+				bProcessEighthResolution ? 3 : (bProcessQuarterResolution ? 2 : 1));
 		}
 
 		if (bLocalExposureBlurredLum)
@@ -1187,9 +1190,9 @@ void AddPostProcessingPasses(
 		FRDGBufferRef SceneColorApplyParameters = nullptr;
 		if (bBloomEnabled)
 		{
-			const FSceneDownsampleChain* LensFlareSceneDownsampleChain;
+			const FTextureDownsampleChain* LensFlareSceneDownsampleChain;
 
-			FSceneDownsampleChain BloomDownsampleChain;
+			FTextureDownsampleChain BloomDownsampleChain;
 
 			if (bFFTBloomEnabled)
 			{
@@ -1263,7 +1266,7 @@ void AddPostProcessingPasses(
 					}
 
 					const bool bLogLumaInAlpha = false;
-					BloomDownsampleChain.Init(GraphBuilder, View, EyeAdaptationParameters, DownsampleInput, DownsampleChainQuality, bLogLumaInAlpha);
+					BloomDownsampleChain.Init(GraphBuilder, View, EyeAdaptationParameters, DownsampleInput, DownsampleChainQuality, (uint32)EBloomQuality::MAX, bLogLumaInAlpha, nullptr, bProcessEighthResolution ? 3 : (bProcessQuarterResolution ? 2 : 1));
 
 					LensFlareSceneDownsampleChain = &BloomDownsampleChain;
 				}
