Heterogeneous Volumes: Disable overlapping volume support by default

The quality of overlapping volume support is highly dependant on the maximum depth samples per AVSM pixel and the features of the volume. Because of this, the support will be disabled by default.

diff --git a/Engine/Source/Runtime/Renderer/Private/HeterogeneousVolumes/HeterogeneousVolumesLiveShadingPipeline.cpp b/Engine/Source/Runtime/Renderer/Private/HeterogeneousVolumes/HeterogeneousVolumesLiveShadingPipeline.cpp
--- a/Engine/Source/Runtime/Renderer/Private/HeterogeneousVolumes/HeterogeneousVolumesLiveShadingPipeline.cpp
+++ b/Engine/Source/Runtime/Renderer/Private/HeterogeneousVolumes/HeterogeneousVolumesLiveShadingPipeline.cpp
@@ -34,8 +34,8 @@ static TAutoConsoleVariable<int32> CVarHeterogeneousLightingLiveShadingScreenTil
 
 static TAutoConsoleVariable<int32> CVarHeterogeneousVolumesSupportOverlappingVolumes(
 	TEXT("r.HeterogeneousVolumes.SupportOverlappingVolumes"),
-	1,
-	TEXT("Enables support for overlapping volumes (Default = 1)"),
+	0,
+	TEXT("Enables support for overlapping volumes (Default = 0)"),
 	ECVF_RenderThreadSafe
 );
 
