Refactored GPUSceneWriter API to hide internals from public interfaces. This is a step towards refactoring the GPU scene data layout to support reserved resources (and other future changes we may want to effect).
The general change is to move the writer parameters into an UB that is forward declared (abstract) in public RENDERER_API and to move GPUSceneWriter.h to the Renderer module.
As FGPUSceneWriteDelegate is referenced in the FMeshBatchDynamicPrimitiveData in the Engine module, this also required an an abstraction.

API breakage:
 * To minimize impact on the mesh batch pipeline, FGPUSceneWriteDelegate::CreateLambda was overloaded to take FMeshElementCollector to allocate from. This will produce a compile error which is simple to fix for any licensee that may use this.

Notes:
 * Suppressed deprecation warnings in both NiagaraRenderMeshes.cpp & PCGInstanceDataInterface.cpp to avoid needing to change public headers in each system. This needs to be addressed at some point as the struct will need to be removed eventually.

Also:
 * Fixed missing MoveTemp in TSceneUpdateCommandQueue.

diff --git a/Engine/Plugins/FX/Niagara/Source/Niagara/Private/NiagaraRendererMeshes.cpp b/Engine/Plugins/FX/Niagara/Source/Niagara/Private/NiagaraRendererMeshes.cpp
--- a/Engine/Plugins/FX/Niagara/Source/Niagara/Private/NiagaraRendererMeshes.cpp
+++ b/Engine/Plugins/FX/Niagara/Source/Niagara/Private/NiagaraRendererMeshes.cpp
@@ -1168,7 +1168,8 @@ void FNiagaraRendererMeshes::SetupElementForGPUScene(
 		// Force it to preserve instance order if we are sorted so that GPU Scene's instance culling doesn't scramble them
 		OutMeshBatchElement.bPreserveInstanceOrder = ParticleMeshRenderData.bNeedsSort;
 
-		GPUSceneRes.DynamicPrimitiveData.DataWriterGPU = FGPUSceneWriteDelegate::CreateLambda(				
+		GPUSceneRes.DynamicPrimitiveData.DataWriterGPU = FGPUSceneWriteDelegate::CreateLambda(		
+			ParticleMeshRenderData.Collector,
 			[&GPUSceneRes,ResouceName=SceneProxy.GetResourceName()](FRDGBuilder& GraphBuilder, const FGPUSceneWriteDelegateParams& Params)
 			{
 			#if !UE_BUILD_SHIPPING
@@ -1179,7 +1180,9 @@ void FNiagaraRendererMeshes::SetupElementForGPUScene(
 				}
 			#endif
 
+				PRAGMA_DISABLE_DEPRECATION_WARNINGS
 				GPUSceneRes.GPUWriteParams.GPUSceneWriterParameters	= Params.GPUWriteParams;
+				PRAGMA_ENABLE_DEPRECATION_WARNINGS
 				GPUSceneRes.GPUWriteParams.View						= Params.View->ViewUniformBuffer; // NOTE: Set here, not outside lambda
 				GPUSceneRes.GPUWriteParams.Scene					= GetSceneUniformBufferRef(GraphBuilder, *Params.View);
 				GPUSceneRes.GPUWriteParams.PrimitiveId 				= Params.PersistentPrimitiveId;
diff --git a/Engine/Plugins/PCG/Source/PCG/Private/Compute/DataInterfaces/PCGInstanceDataInterface.cpp b/Engine/Plugins/PCG/Source/PCG/Private/Compute/DataInterfaces/PCGInstanceDataInterface.cpp
--- a/Engine/Plugins/PCG/Source/PCG/Private/Compute/DataInterfaces/PCGInstanceDataInterface.cpp
+++ b/Engine/Plugins/PCG/Source/PCG/Private/Compute/DataInterfaces/PCGInstanceDataInterface.cpp
@@ -269,7 +269,9 @@ void FPCGInstanceDataProviderProxy::PostGraphDispatch(FRDGBuilder& GraphBuilder)
 				Parameters->InInstanceCustomFloatData = InstanceCustomFloatDataSRV;
 				Parameters->InWriteCounters = WriteCountersSRV;
 				Parameters->InPrimitiveId = Params.PersistentPrimitiveId;
+				PRAGMA_DISABLE_DEPRECATION_WARNINGS
 				Parameters->GPUSceneWriterParameters = Params.GPUWriteParams;
+				PRAGMA_ENABLE_DEPRECATION_WARNINGS
 				Parameters->InCustomDataCount = Params.NumCustomDataFloats;
 				Parameters->InPayloadDataFlags = Params.PackedInstanceSceneDataFlags;
 			
diff --git a/Engine/Shaders/Private/GPUScene/GPUSceneWriter.ush b/Engine/Shaders/Private/GPUScene/GPUSceneWriter.ush
--- a/Engine/Shaders/Private/GPUScene/GPUSceneWriter.ush
+++ b/Engine/Shaders/Private/GPUScene/GPUSceneWriter.ush
@@ -15,31 +15,27 @@
 
 #include "../SceneData.ush"
 
-uint GPUSceneInstanceSceneDataSOAStride;
-uint GPUSceneNumAllocatedInstances;
-uint GPUSceneNumAllocatedPrimitives;
-
 void StoreInstanceSceneDataElement(uint InstanceId, uint ElementIndex, float4 DataElement)
 {
-	checkSlow(InstanceId < GPUSceneNumAllocatedInstances);
-	GPUSceneInstanceSceneDataRW[ElementIndex * GPUSceneInstanceSceneDataSOAStride + InstanceId] = DataElement;
+	checkSlow(InstanceId < GPUSceneWriter.GPUSceneNumAllocatedInstances);
+	GPUSceneWriter.GPUSceneInstanceSceneDataRW[ElementIndex * GPUSceneWriter.GPUSceneInstanceSceneDataSOAStride + InstanceId] = DataElement;
 }
 
 void StoreInstanceSceneDataElementX(uint InstanceId, uint ElementIndex, float Data)
 {
 	// NOTE: Structured buffers allow for DWORD granularity storage, so you don't have to write the whole float4
-	checkSlow(InstanceId < GPUSceneNumAllocatedInstances);
-	GPUSceneInstanceSceneDataRW[ElementIndex * GPUSceneInstanceSceneDataSOAStride + InstanceId].x = Data;
+	checkSlow(InstanceId < GPUSceneWriter.GPUSceneNumAllocatedInstances);
+	GPUSceneWriter.GPUSceneInstanceSceneDataRW[ElementIndex * GPUSceneWriter.GPUSceneInstanceSceneDataSOAStride + InstanceId].x = Data;
 }
 
 void StoreInstancePayloadDataElement(uint PayloadDataOffset, uint ElementIndex, float4 DataElement)
 {
-	GPUSceneInstancePayloadDataRW[PayloadDataOffset + ElementIndex] = DataElement;
+	GPUSceneWriter.GPUSceneInstancePayloadDataRW[PayloadDataOffset + ElementIndex] = DataElement;
 }
 
 void StoreInstancePayloadDataElementZW(uint PayloadDataOffset, uint ElementIndex, float2 DataElement)
 {
-	GPUSceneInstancePayloadDataRW[PayloadDataOffset + ElementIndex].zw = DataElement;
+	GPUSceneWriter.GPUSceneInstancePayloadDataRW[PayloadDataOffset + ElementIndex].zw = DataElement;
 }
 
 /**
@@ -47,7 +43,7 @@ void StoreInstancePayloadDataElementZW(uint PayloadDataOffset, uint ElementIndex
  */
 void LoadInstancePrimitiveIdAndFlags(uint InstanceId, inout uint OutPrimitiveId, inout uint OutInstanceFlags)
 {
-	LoadInstancePrimitiveIdAndFlags(InstanceId, GPUSceneInstanceSceneDataSOAStride, OutPrimitiveId, OutInstanceFlags);
+	LoadInstancePrimitiveIdAndFlags(InstanceId, GPUSceneWriter.GPUSceneInstanceSceneDataSOAStride, OutPrimitiveId, OutInstanceFlags);
 }
 
 void WriteInstancePrimitiveIdAndFlags(uint InstanceId, uint PrimitiveId, uint InstanceFlags)
@@ -172,8 +168,8 @@ void InitializeInstanceSceneDataLS(
  FInstancePayloadDataOffsets GetInstancePayloadDataOffsets(uint InstanceId)
 {
 	uint PrimitiveId, InstanceFlags, InstanceRelativeId, CustomDataCount;
-	LoadInstancePrimitiveIdAndFlags(InstanceId, GPUSceneInstanceSceneDataSOAStride, PrimitiveId, InstanceFlags);
-	LoadInstanceRelativeIdAndCustomDataCount(InstanceId, GPUSceneInstanceSceneDataSOAStride, InstanceRelativeId, CustomDataCount);
+	LoadInstancePrimitiveIdAndFlags(InstanceId, GPUSceneWriter.GPUSceneInstanceSceneDataSOAStride, PrimitiveId, InstanceFlags);
+	LoadInstanceRelativeIdAndCustomDataCount(InstanceId, GPUSceneWriter.GPUSceneInstanceSceneDataSOAStride, InstanceRelativeId, CustomDataCount);
 	return GetInstancePayloadDataOffsets(PrimitiveId, InstanceFlags, InstanceRelativeId);
 }
 
diff --git a/Engine/Shaders/Private/SceneData.ush b/Engine/Shaders/Private/SceneData.ush
--- a/Engine/Shaders/Private/SceneData.ush
+++ b/Engine/Shaders/Private/SceneData.ush
@@ -309,8 +309,6 @@ FInstanceSceneData WaveReadLaneAt(FInstanceSceneData In, uint SrcIndex)
 // no access to global buffers
 #elif USE_GLOBAL_GPU_SCENE_DATA
 StructuredBuffer<float4> GPUScenePrimitiveSceneData;
-#elif USE_GLOBAL_GPU_SCENE_DATA_RW
-RWStructuredBuffer<float4> GPUScenePrimitiveSceneDataRW;
 #endif
 
 float4 LoadPrimitivePrimitiveSceneDataElement(uint PrimitiveIndex, uint ItemIndex)
@@ -323,8 +321,8 @@ float4 LoadPrimitivePrimitiveSceneDataElement(uint PrimitiveIndex, uint ItemInde
 	checkStructuredBufferAccessSlow(GPUScenePrimitiveSceneData, TargetIdx);
 	return GPUScenePrimitiveSceneData[TargetIdx];
 #elif USE_GLOBAL_GPU_SCENE_DATA_RW
-	checkStructuredBufferAccessSlow(GPUScenePrimitiveSceneDataRW, TargetIdx);
-	return GPUScenePrimitiveSceneDataRW[TargetIdx];
+	checkStructuredBufferAccessSlow(GPUSceneWriter.GPUScenePrimitiveSceneDataRW, TargetIdx);
+	return GPUSceneWriter.GPUScenePrimitiveSceneDataRW[TargetIdx];
 #else
 	checkStructuredBufferAccessSlow(Scene.GPUScene.GPUScenePrimitiveSceneData, TargetIdx);
 	return Scene.GPUScene.GPUScenePrimitiveSceneData[TargetIdx];
@@ -559,16 +557,14 @@ struct FInstancePayloadDataOffsets
 StructuredBuffer<float4> GPUSceneInstanceSceneData;
 StructuredBuffer<float4> GPUSceneInstancePayloadData;
 uint GPUSceneFrameNumber;
-#elif USE_GLOBAL_GPU_SCENE_DATA_RW
-RWStructuredBuffer<float4> GPUSceneInstanceSceneDataRW;
-RWStructuredBuffer<float4> GPUSceneInstancePayloadDataRW;
-uint GPUSceneFrameNumber;
 #endif
 
 uint GetGPUSceneFrameNumber()
 {
-#if USE_GLOBAL_GPU_SCENE_DATA || USE_GLOBAL_GPU_SCENE_DATA_RW
+#if USE_GLOBAL_GPU_SCENE_DATA
 	return GPUSceneFrameNumber;
+#elif USE_GLOBAL_GPU_SCENE_DATA_RW
+	return GPUSceneWriter.GPUSceneFrameNumber;
 #else
 	return Scene.GPUScene.GPUSceneFrameNumber;
 #endif
@@ -581,7 +577,7 @@ float4 LoadInstanceSceneDataElement(uint Index)
 #elif USE_GLOBAL_GPU_SCENE_DATA
 	return GPUSceneInstanceSceneData[Index];
 #elif USE_GLOBAL_GPU_SCENE_DATA_RW
-	return GPUSceneInstanceSceneDataRW[Index];
+	return GPUSceneWriter.GPUSceneInstanceSceneDataRW[Index];
 #else
 	return Scene.GPUScene.GPUSceneInstanceSceneData[Index];
 #endif
@@ -594,7 +590,7 @@ float4 LoadInstancePayloadDataElement(uint Index)
 #elif USE_GLOBAL_GPU_SCENE_DATA
 	return GPUSceneInstancePayloadData[Index];
 #elif USE_GLOBAL_GPU_SCENE_DATA_RW
-	return GPUSceneInstancePayloadDataRW[Index];
+	return GPUSceneWriter.GPUSceneInstancePayloadDataRW[Index];
 #else
 	return Scene.GPUScene.GPUSceneInstancePayloadData[Index];
 #endif
diff --git a/Engine/Source/Runtime/Engine/Public/GPUSceneWriter.h b/Engine/Source/Runtime/Engine/Public/GPUSceneWriter.h
deleted file mode 100644
diff --git a/Engine/Source/Runtime/Engine/Public/MeshBatch.h b/Engine/Source/Runtime/Engine/Public/MeshBatch.h
--- a/Engine/Source/Runtime/Engine/Public/MeshBatch.h
+++ b/Engine/Source/Runtime/Engine/Public/MeshBatch.h
@@ -6,11 +6,16 @@
 #include "Containers/ArrayView.h"
 #include "EngineDefines.h"
 #include "Engine/EngineTypes.h"
-#include "GPUSceneWriter.h"
+#if UE_ENABLE_INCLUDE_ORDER_DEPRECATED_IN_5_6
+// previously included by GPUSceneWriter.h
+#include "ShaderParameterMacros.h"
+#include "ShaderParameters.h"
+#endif
 #include "HitProxies.h"
 #include "RHIDefinitions.h"
 #include "SceneDefinitions.h"
 #include "VT/RuntimeVirtualTextureEnum.h"
+#include "RHIResources.h"
 
 #if UE_ENABLE_INCLUDE_ORDER_DEPRECATED_IN_5_3
 #include "InstanceUniformShaderParameters.h"
@@ -30,6 +35,41 @@ struct FInstanceSceneData;
 struct FMaterialShaderParameters;
 struct FRenderBounds;
 template<typename TBufferStruct> class TUniformBuffer;
+class FShaderParametersMetadata;
+class FGPUSceneWriteDelegate;
+
+/**
+ * Abstract reference to a GPUSceneWriteDelegate, such that implementation details can be left to the renderer.
+ * Only contains a raw pointer and performs no life-time management, instead it relies on being constructed by FGPUSceneWriteDelegate
+ * which will use AllocateOneFrameResource to ensure a lifetime compatible with mesh batches.
+ */
+class FGPUSceneWriteDelegateRef
+{
+public:
+	FORCEINLINE bool IsBound() const { return GPUSceneWriteDelegate != nullptr; }
+
+	FORCEINLINE FGPUSceneWriteDelegate* operator->() const
+	{
+		check( IsBound() );
+		return GPUSceneWriteDelegate;
+	}
+
+private:
+	friend class FGPUSceneWriteDelegate;
+	FGPUSceneWriteDelegate* GPUSceneWriteDelegate = nullptr;
+};
+
+/** Specifies the point during scene rendering that GPU Scene can be written to */
+enum class EGPUSceneGPUWritePass : int8
+{
+	/** Invalid GPU write pass. Also used to signify writes that occur during upload. */
+	None = -1,
+
+	/** Writes to the GPU Scene will happen after rendering opaques (as well has after FXSystem post-opaque updates). */
+	PostOpaqueRendering,
+
+	Num
+};
 
 enum EPrimitiveIdMode
 {
@@ -89,7 +129,7 @@ struct FMeshBatchDynamicPrimitiveData
 	TConstArrayView<FInstanceDynamicData> InstanceDynamicData;
 	TConstArrayView<FRenderBounds> InstanceLocalBounds;
 	TConstArrayView<float> InstanceCustomData;
-	FGPUSceneWriteDelegate DataWriterGPU;		
+	FGPUSceneWriteDelegateRef DataWriterGPU;		
 	EGPUSceneGPUWritePass DataWriterGPUPass = EGPUSceneGPUWritePass::None;
 	uint16 PayloadDataFlags = 0;
 	uint32 NumInstanceCustomDataFloats = 0;
diff --git a/Engine/Source/Runtime/Renderer/Private/GPUScene.cpp b/Engine/Source/Runtime/Renderer/Private/GPUScene.cpp
--- a/Engine/Source/Runtime/Renderer/Private/GPUScene.cpp
+++ b/Engine/Source/Runtime/Renderer/Private/GPUScene.cpp
@@ -30,6 +30,7 @@
 #include "PrimitiveUniformShaderParametersBuilder.h"
 #include "InstanceDataSceneProxy.h"
 #include "SceneRendererInterface.h"
+#include "GPUSceneWriter.h"
 
 // Useful for debugging
 #define FORCEINLINE_GPUSCENE FORCEINLINE
@@ -58,6 +59,18 @@ static void ConstructDefault(FGPUSceneResourceParameters& GPUScene, FRDGBuilder&
 }
 IMPLEMENT_SCENE_UB_STRUCT(FGPUSceneResourceParameters, GPUScene, ConstructDefault);
 
+BEGIN_GLOBAL_SHADER_PARAMETER_STRUCT(FGPUSceneWriterUniformParameters, RENDERER_API)
+	SHADER_PARAMETER_RDG_BUFFER_UAV(RWStructuredBuffer<float4>, GPUSceneInstanceSceneDataRW)
+	SHADER_PARAMETER_RDG_BUFFER_UAV(RWStructuredBuffer<float4>, GPUSceneInstancePayloadDataRW)
+	SHADER_PARAMETER_RDG_BUFFER_UAV(RWStructuredBuffer<float4>, GPUScenePrimitiveSceneDataRW)
+	SHADER_PARAMETER(uint32, GPUSceneInstanceSceneDataSOAStride)
+	SHADER_PARAMETER(uint32, GPUSceneFrameNumber)
+	SHADER_PARAMETER(uint32, GPUSceneNumAllocatedInstances)
+	SHADER_PARAMETER(uint32, GPUSceneNumAllocatedPrimitives)	
+END_GLOBAL_SHADER_PARAMETER_STRUCT()
+
+IMPLEMENT_GLOBAL_SHADER_PARAMETER_STRUCT(FGPUSceneWriterUniformParameters, "GPUSceneWriter");
+
 static TAutoConsoleVariable<int32> CVarGPUSceneUploadEveryFrame(
 	TEXT("r.GPUScene.UploadEveryFrame"),
 	0,
@@ -884,7 +897,9 @@ void FGPUScene::UpdateInternal(FRDGBuilder& GraphBuilder, FSceneUniformBuffer& S
 		RDG_EVENT_SCOPE(GraphBuilder, "UpdateGPUScene UpdateInstancesFromCompute");
 
 		FGPUSceneWriteDelegateParams Params;
-		GetWriteParameters(GraphBuilder, Params.GPUWriteParams);
+		PRAGMA_DISABLE_DEPRECATION_WARNINGS
+		Params.GPUSceneWriterUB = GetWriteParameters(GraphBuilder, &Params.GPUWriteParams);
+		PRAGMA_ENABLE_DEPRECATION_WARNINGS
 
 		for (const auto& Command : UpdatesFromCompute)
 		{
@@ -1676,7 +1691,9 @@ void FGPUScene::UploadDynamicPrimitiveShaderDataForViewInternal(FRDGBuilder& Gra
 			FGPUSceneWriteDelegateParams Params;
 			Params.View = &View;
 			Params.GPUWritePass = EGPUSceneGPUWritePass::None;
-			GetWriteParameters(GraphBuilder, Params.GPUWriteParams);
+			PRAGMA_DISABLE_DEPRECATION_WARNINGS
+			Params.GPUSceneWriterUB = GetWriteParameters(GraphBuilder, &Params.GPUWriteParams);
+			PRAGMA_ENABLE_DEPRECATION_WARNINGS
 
 			for (uint32 PrimitiveIndex : ImmediateWrites)
 			{
@@ -1686,7 +1703,7 @@ void FGPUScene::UploadDynamicPrimitiveShaderDataForViewInternal(FRDGBuilder& Gra
 				Params.NumCustomDataFloats = PrimData.SourceData.NumInstanceCustomDataFloats;
 				Params.PackedInstanceSceneDataFlags = PrimData.SourceData.PayloadDataFlags;
 
-				PrimData.SourceData.DataWriterGPU.Execute(GraphBuilder, Params);
+				PrimData.SourceData.DataWriterGPU->Execute(GraphBuilder, Params);
 			}
 		}
 	}
@@ -2085,7 +2102,9 @@ bool FGPUScene::ExecuteDeferredGPUWritePass(FRDGBuilder& GraphBuilder, TArray<FV
 
 	FGPUSceneWriteDelegateParams Params;
 	Params.GPUWritePass = GPUWritePass;
-	GetWriteParameters(GraphBuilder, Params.GPUWriteParams);
+	PRAGMA_DISABLE_DEPRECATION_WARNINGS
+	Params.GPUSceneWriterUB = GetWriteParameters(GraphBuilder, &Params.GPUWriteParams);
+	PRAGMA_ENABLE_DEPRECATION_WARNINGS
 
 	for (const FDeferredGPUWrite& DeferredWrite : DeferredGPUWritePassDelegates[PassIndex])
 	{
@@ -2098,7 +2117,7 @@ bool FGPUScene::ExecuteDeferredGPUWritePass(FRDGBuilder& GraphBuilder, TArray<FV
 		Params.NumCustomDataFloats = DeferredWrite.NumCustomDataFloats;
 		Params.PackedInstanceSceneDataFlags = DeferredWrite.PackedInstanceSceneDataFlags;
 
-		DeferredWrite.DataWriterGPU.Execute(GraphBuilder, Params);
+		DeferredWrite.DataWriterGPU->Execute(GraphBuilder, Params);
 	}
 
 	DeferredGPUWritePassDelegates[PassIndex].Reset();
@@ -2194,10 +2213,7 @@ FGPUScenePrimitiveCollector::FUploadData* FGPUSceneDynamicContext::AllocateDynam
 	return UploadData;
 }
 
-/**
- * Fills in the FGPUSceneWriterParameters to use for read/write access to the GPU Scene.
- */
-void FGPUScene::GetWriteParameters(FRDGBuilder& GraphBuilder, FGPUSceneWriterParameters& GPUSceneWriterParametersOut)
+TRDGUniformBufferRef<FGPUSceneWriterUniformParameters> FGPUScene::GetWriteParameters(FRDGBuilder& GraphBuilder, FGPUSceneWriterParameters* LegacyParameters)
 {
 	FRegisteredBuffers BufferState = CachedRegisteredBuffers;
 	// If we're not in a begin/end block we need to register the buffers here and now.
@@ -2205,13 +2221,28 @@ void FGPUScene::GetWriteParameters(FRDGBuilder& GraphBuilder, FGPUSceneWriterPar
 	{
 		BufferState = RegisterBuffers(GraphBuilder);
 	}
-	GPUSceneWriterParametersOut.GPUSceneFrameNumber = SceneFrameNumber;
-	GPUSceneWriterParametersOut.GPUSceneInstanceSceneDataSOAStride = InstanceSceneDataSOAStride;
-	GPUSceneWriterParametersOut.GPUSceneNumAllocatedInstances = InstanceSceneDataAllocator.GetMaxSize();
-	GPUSceneWriterParametersOut.GPUSceneNumAllocatedPrimitives = DynamicPrimitivesOffset;
-	GPUSceneWriterParametersOut.GPUSceneInstanceSceneDataRW = GraphBuilder.CreateUAV(BufferState.InstanceSceneDataBuffer, ERDGUnorderedAccessViewFlags::SkipBarrier);
-	GPUSceneWriterParametersOut.GPUSceneInstancePayloadDataRW = GraphBuilder.CreateUAV(BufferState.InstancePayloadDataBuffer, ERDGUnorderedAccessViewFlags::SkipBarrier);
-	GPUSceneWriterParametersOut.GPUScenePrimitiveSceneDataRW = GraphBuilder.CreateUAV(BufferState.PrimitiveBuffer, ERDGUnorderedAccessViewFlags::SkipBarrier);
+	FGPUSceneWriterUniformParameters* UBParams = GraphBuilder.AllocParameters<FGPUSceneWriterUniformParameters>();
+	UBParams->GPUSceneFrameNumber = SceneFrameNumber;
+	UBParams->GPUSceneInstanceSceneDataSOAStride = InstanceSceneDataSOAStride;
+	UBParams->GPUSceneNumAllocatedInstances = InstanceSceneDataAllocator.GetMaxSize();
+	UBParams->GPUSceneNumAllocatedPrimitives = DynamicPrimitivesOffset;
+	UBParams->GPUSceneInstanceSceneDataRW = GraphBuilder.CreateUAV(BufferState.InstanceSceneDataBuffer, ERDGUnorderedAccessViewFlags::SkipBarrier);
+	UBParams->GPUSceneInstancePayloadDataRW = GraphBuilder.CreateUAV(BufferState.InstancePayloadDataBuffer, ERDGUnorderedAccessViewFlags::SkipBarrier);
+	UBParams->GPUScenePrimitiveSceneDataRW = GraphBuilder.CreateUAV(BufferState.PrimitiveBuffer, ERDGUnorderedAccessViewFlags::SkipBarrier);
+	TRDGUniformBufferRef<FGPUSceneWriterUniformParameters> Result = GraphBuilder.CreateUniformBuffer(UBParams);
+
+	if (LegacyParameters != nullptr)
+	{
+		LegacyParameters->GPUSceneFrameNumber = UBParams->GPUSceneFrameNumber;
+		LegacyParameters->GPUSceneInstanceSceneDataSOAStride = UBParams->GPUSceneInstanceSceneDataSOAStride;
+		LegacyParameters->GPUSceneNumAllocatedInstances = UBParams->GPUSceneNumAllocatedInstances;
+		LegacyParameters->GPUSceneNumAllocatedPrimitives = UBParams->GPUSceneNumAllocatedPrimitives;
+		LegacyParameters->GPUSceneInstanceSceneDataRW = UBParams->GPUSceneInstanceSceneDataRW;
+		LegacyParameters->GPUSceneInstancePayloadDataRW = UBParams->GPUSceneInstancePayloadDataRW;
+		LegacyParameters->GPUScenePrimitiveSceneDataRW = UBParams->GPUScenePrimitiveSceneDataRW;
+		LegacyParameters->GPUSceneWriterUB = Result;
+	}
+	return Result;
 }
 
 class FGPUSceneSetInstancePrimitiveIdCS : public FGlobalShader
@@ -2221,7 +2252,7 @@ class FGPUSceneSetInstancePrimitiveIdCS : public FGlobalShader
 public:
 	BEGIN_SHADER_PARAMETER_STRUCT(FParameters, )
 		SHADER_PARAMETER_STRUCT_INCLUDE(FGPUScene::FInstanceGPULoadBalancer::FShaderParameters, BatcherParameters)
-		SHADER_PARAMETER_STRUCT_INCLUDE(FGPUSceneWriterParameters, GPUSceneWriterParameters)
+		SHADER_PARAMETER_RDG_UNIFORM_BUFFER(FGPUSceneWriterUniformParameters, GPUSceneWriterUB)
 	END_SHADER_PARAMETER_STRUCT()
 
 	static constexpr int32 NumThreadsPerGroup = FGPUScene::FInstanceGPULoadBalancer::ThreadGroupSize;
@@ -2250,9 +2281,7 @@ void FGPUScene::AddUpdatePrimitiveIdsPass(FRDGBuilder& GraphBuilder, FInstanceGP
 		FGPUSceneSetInstancePrimitiveIdCS::FParameters* PassParameters = GraphBuilder.AllocParameters<FGPUSceneSetInstancePrimitiveIdCS::FParameters>();
 
 		IdOnlyUpdateItems.Upload(GraphBuilder).GetShaderParameters(GraphBuilder, PassParameters->BatcherParameters);
-
-		GetWriteParameters(GraphBuilder, PassParameters->GPUSceneWriterParameters);
-
+		PassParameters->GPUSceneWriterUB = GetWriteParameters(GraphBuilder);
 		auto ComputeShader = GetGlobalShaderMap(FeatureLevel)->GetShader<FGPUSceneSetInstancePrimitiveIdCS>();
 
 		FComputeShaderUtils::AddPass(
@@ -2428,3 +2457,22 @@ void FGPUScene::AddClearInstancesPass(FRDGBuilder& GraphBuilder, FInstanceCullin
 	AddUpdatePrimitiveIdsPass(GraphBuilder, ClearIdData);
 	InstanceRangesToClear.Empty();
 }
+
+FGPUSceneWriteDelegate::FGPUSceneWriteDelegate(FGPUSceneWriteDelegateImpl&& GPUSceneWriteDelegateImpl)
+	: Delegate(MoveTemp(GPUSceneWriteDelegateImpl))
+{
+}
+
+void FGPUSceneWriteDelegate::Execute(FRDGBuilder& GraphBuilder, const FGPUSceneWriteDelegateParams& Params) const
+{
+	check(Delegate.IsBound());
+	Delegate.Execute(GraphBuilder, Params);
+}
+
+FGPUSceneWriteDelegateRef FGPUSceneWriteDelegate::CreateInternal(FMeshElementCollector* MeshElementCollector, FGPUSceneWriteDelegateImpl&& DelegateImpl)
+{
+	FGPUSceneWriteDelegateRef Result;
+	Result.GPUSceneWriteDelegate = &MeshElementCollector->AllocateOneFrameResource<FGPUSceneWriteDelegate>(MoveTemp(DelegateImpl));
+	return Result;
+}
+
diff --git a/Engine/Source/Runtime/Renderer/Private/GPUScene.h b/Engine/Source/Runtime/Renderer/Private/GPUScene.h
--- a/Engine/Source/Runtime/Renderer/Private/GPUScene.h
+++ b/Engine/Source/Runtime/Renderer/Private/GPUScene.h
@@ -28,6 +28,8 @@ class FScenePreUpdateChangeSet;
 class FScenePostUpdateChangeSet;
 class FPrimitiveSceneProxy;
 struct FLightSceneChangeSet;
+class FGPUSceneWriterParameters;
+
 namespace UE::Renderer::Private
 {
 	class IShadowInvalidatingInstances;
@@ -257,11 +259,6 @@ public:
 		return PrimitiveDirtyState[PersistentPrimitiveIndex.Index]; 
 	}
 
-	/**
-	 * Fills in the FGPUSceneWriterParameters to use for read/write access to the GPU Scene.
-	 */
-	void GetWriteParameters(FRDGBuilder& GraphBuilder, FGPUSceneWriterParameters& GPUSceneWriterParametersOut);
-
 	uint32 GetSceneFrameNumber() const { return SceneFrameNumber; }
 
 	int32 GetNumInstances() const { return InstanceSceneDataAllocator.GetMaxSize(); }
@@ -341,6 +338,11 @@ public:
 	SIZE_T GetAllocatedSize() const;
 
 private:
+	/**
+	 * Fills in the FGPUSceneWriterParameters to use for read/write access to the GPU Scene.
+	 */
+	TRDGUniformBufferRef<FGPUSceneWriterUniformParameters> GetWriteParameters(FRDGBuilder& GraphBuilder, FGPUSceneWriterParameters* LegacyParameters = nullptr);
+
 	static constexpr int32 InitialBufferSize = 256;
 
 	TRefCountPtr<FRDGPooledBuffer> LightDataBuffer;
@@ -386,7 +388,7 @@ private:
 
 	struct FDeferredGPUWrite
 	{
-		FGPUSceneWriteDelegate DataWriterGPU;
+		FGPUSceneWriteDelegateRef DataWriterGPU;
 		int32 ViewId = INDEX_NONE;
 		uint32 PrimitiveId = INDEX_NONE;
 		uint32 InstanceSceneDataOffset = INDEX_NONE;
diff --git a/Engine/Source/Runtime/Renderer/Private/ScenePrimitiveUpdates.h b/Engine/Source/Runtime/Renderer/Private/ScenePrimitiveUpdates.h
--- a/Engine/Source/Runtime/Renderer/Private/ScenePrimitiveUpdates.h
+++ b/Engine/Source/Runtime/Renderer/Private/ScenePrimitiveUpdates.h
@@ -3,6 +3,7 @@
 
 #include "SceneUpdateCommandQueue.h"
 #include "PrimitiveSceneInfo.h"
+#include "GPUSceneWriter.h"
 
 /**
  * Definitions of primitive scene update commands.
diff --git a/Engine/Source/Runtime/Renderer/Private/SceneUpdateCommandQueue.h b/Engine/Source/Runtime/Renderer/Private/SceneUpdateCommandQueue.h
--- a/Engine/Source/Runtime/Renderer/Private/SceneUpdateCommandQueue.h
+++ b/Engine/Source/Runtime/Renderer/Private/SceneUpdateCommandQueue.h
@@ -200,14 +200,14 @@ public:
 		if (PrevPayloadOffset != INDEX_NONE)
 		{
 			check(Payloads->CommandSlots[PrevPayloadOffset] == CommandSlot);
-			Payloads->PayloadData[PrevPayloadOffset] = Payload;
+			Payloads->PayloadData[PrevPayloadOffset] = MoveTemp(Payload);
 		}
 		else // New payload for this command
 		{
 			int32 PayloadOffset = Payloads->CommandSlots.Num();
 			Commands[CommandSlot].template SetOrAddPayloadOffset<PayloadType>(PayloadOffset, Payload.GetDirtyFlags());
 			Payloads->CommandSlots.Emplace(CommandSlot);
-			Payloads->PayloadData.Emplace(Payload);
+			Payloads->PayloadData.Emplace(MoveTemp(Payload));
 		}
 	}
 
diff --git a/Engine/Source/Runtime/Renderer/Public/GPUSceneWriter.h b/Engine/Source/Runtime/Renderer/Public/GPUSceneWriter.h
new file mode 100644
--- /dev/null
+++ b/Engine/Source/Runtime/Renderer/Public/GPUSceneWriter.h
@@ -0,0 +1,102 @@
+// Copyright Epic Games, Inc. All Rights Reserved.
+
+#pragma once
+
+#include "CoreMinimal.h"
+#include "ShaderParameterMacros.h"
+#include "ShaderParameters.h"
+#include "MeshBatch.h"
+#include "RenderGraphFwd.h"
+#include "SceneManagement.h"
+
+class FSceneView;
+class FRDGBuilder;
+class FMeshElementCollector;
+
+DECLARE_UNIFORM_BUFFER_STRUCT(FGPUSceneWriterUniformParameters, RENDERER_API)
+
+/**
+ * Deprecated for 5.6!
+ * Use SHADER_PARAMETER_RDG_UNIFORM_BUFFER(FGPUSceneWriterUniformParameters, GPUSceneWriterUB) instead.
+ */
+BEGIN_SHADER_PARAMETER_STRUCT(FGPUSceneWriterParameters, RENDERER_API)
+	SHADER_PARAMETER_RDG_BUFFER_UAV(RWStructuredBuffer<float4>, GPUSceneInstanceSceneDataRW)
+	SHADER_PARAMETER_RDG_BUFFER_UAV(RWStructuredBuffer<float4>, GPUSceneInstancePayloadDataRW)
+	SHADER_PARAMETER_RDG_BUFFER_UAV(RWStructuredBuffer<float4>, GPUScenePrimitiveSceneDataRW)
+	SHADER_PARAMETER(uint32, GPUSceneInstanceSceneDataSOAStride)
+	SHADER_PARAMETER(uint32, GPUSceneFrameNumber)
+	SHADER_PARAMETER(uint32, GPUSceneNumAllocatedInstances)
+	SHADER_PARAMETER(uint32, GPUSceneNumAllocatedPrimitives)
+
+	SHADER_PARAMETER_RDG_UNIFORM_BUFFER(FGPUSceneWriterUniformParameters, GPUSceneWriterUB)
+END_SHADER_PARAMETER_STRUCT()
+
+/** The parameters passed to the GPUScene writer delegate */
+struct FGPUSceneWriteDelegateParams
+{
+	/** The ID of the primitive that writes must be limited to. */
+	uint32 PersistentPrimitiveId = (uint32)INDEX_NONE;
+	/** The ID of the first instance scene data of the primitive */
+	uint32 InstanceSceneDataOffset = (uint32)INDEX_NONE;
+	/** Number of custom data floats in the instance payload data. */
+	uint32 NumCustomDataFloats = (uint32)INDEX_NONE;
+	/** Packed instance scene data flags suitable for writing to instance scene data. */
+	uint32 PackedInstanceSceneDataFlags = 0u;
+	/** The GPU Scene write pass that is currently executing. (NOTE: A value of None specifies that it is occurring on upload) */
+	EGPUSceneGPUWritePass GPUWritePass = EGPUSceneGPUWritePass::None;
+	/** The view for which this primitive belongs (for dynamic primitives) */
+	FSceneView* View = nullptr;
+	/** The shader parameters the delegate can use to perform writes on GPU Scene data */
+	UE_DEPRECATED(5.6, "Only the opaque UB (see GPUSceneWriterUB below) should be referenced")
+	FGPUSceneWriterParameters GPUWriteParams;
+	/** 
+	 * Include a reference to the UB in the compute shader parameter struct, for example:
+	 *   SHADER_PARAMETER_RDG_UNIFORM_BUFFER(FGPUSceneWriterUniformParameters, GPUSceneWriterUB)
+	 */
+	TRDGUniformBufferRef<FGPUSceneWriterUniformParameters> GPUSceneWriterUB;
+};
+
+DECLARE_DELEGATE_TwoParams(FGPUSceneWriteDelegateImpl, FRDGBuilder&, const FGPUSceneWriteDelegateParams&);
+
+/** 
+ * Special implementation of FGPUSceneWriteDelegate to keep API similar, but allow abstracting the internals for the mesh batch path.
+ * Used by GPUScene to stage writing to the GPUScene primitive and instance data buffers via the GPU.
+ */
+class FGPUSceneWriteDelegate : public FOneFrameResource
+{
+public:
+	RENDERER_API FGPUSceneWriteDelegate() = default;
+	RENDERER_API FGPUSceneWriteDelegate(FGPUSceneWriteDelegate&&) = default;
+	RENDERER_API FGPUSceneWriteDelegate &operator=(FGPUSceneWriteDelegate&&) = default;
+	RENDERER_API FGPUSceneWriteDelegate(FGPUSceneWriteDelegateImpl&& GPUSceneWriteDelegateImpl);
+
+	/**
+	 * This version returns a FGPUSceneWriteDelegateRef and this must be used to associate a FGPUSceneWriteDelegate with a mesh batch.
+	 */
+	template <typename LambdaType>
+	static FGPUSceneWriteDelegateRef CreateLambda(FMeshElementCollector* MeshElementCollector, LambdaType&& Lambda)
+	{
+		return CreateInternal(MeshElementCollector, FGPUSceneWriteDelegateImpl::CreateLambda(MoveTemp(Lambda)));
+	}
+
+	/**
+	 */
+	template <typename LambdaType>
+	static FGPUSceneWriteDelegate CreateLambda(LambdaType&& Lambda)
+	{
+		return FGPUSceneWriteDelegate(FGPUSceneWriteDelegateImpl::CreateLambda(MoveTemp(Lambda)));
+	}
+
+	/**
+	 */
+	RENDERER_API void Execute(FRDGBuilder& GraphBuilder, const FGPUSceneWriteDelegateParams& Params) const;
+
+	/**
+	 */
+	bool IsBound() const { return Delegate.IsBound(); }
+
+private:
+	RENDERER_API static FGPUSceneWriteDelegateRef CreateInternal(FMeshElementCollector* MeshElementCollector, FGPUSceneWriteDelegateImpl&& DelegateImpl);
+
+	FGPUSceneWriteDelegateImpl Delegate;
+};
