LightFunctionAtlas - fix for missing call to task after 36854183 for the mobile renderer.
#rnx
#tests none

diff --git a/Engine/Source/Runtime/Renderer/Private/MobileShadingRenderer.cpp b/Engine/Source/Runtime/Renderer/Private/MobileShadingRenderer.cpp
--- a/Engine/Source/Runtime/Renderer/Private/MobileShadingRenderer.cpp
+++ b/Engine/Source/Runtime/Renderer/Private/MobileShadingRenderer.cpp
@@ -1099,6 +1099,10 @@ void FMobileSceneRenderer::Render(FRDGBuilder& GraphBuilder)
 			RDG_CSV_STAT_EXCLUSIVE_SCOPE(GraphBuilder, SortLights);
 			// Shadows are applied in clustered shading on mobile forward and separately on mobile deferred.
 			bool bShadowedLightsInClustered = bRequiresShadowProjections && !bDeferredShading;
+
+			// This task needs to run before any other functions gathering lights for upload on GPU, for light function indices to be assigned to lights.
+			UpdateLightFunctionAtlasTaskFunction();
+
 			GatherAndSortLights(SortedLightSet, bShadowedLightsInClustered);
 			int32 NumReflectionCaptures = Views[0].NumBoxReflectionCaptures + Views[0].NumSphereReflectionCaptures;
 			bool bCullLightsToGrid = (((bEnableClusteredReflections || bDeferredShading) && NumReflectionCaptures > 0) || bEnableClusteredLocalLights);
