Fix crash when rendering a scene with a Planar Reflection in Mobile Multi-View Fallback. Also fix an ensure caused by the editor forcing EngineShowFlags.ScreenPercentage on for HighDPI even when screen percentage is unsupported by the view.

diff --git a/Engine/Source/Runtime/Engine/Private/GameViewportClient.cpp b/Engine/Source/Runtime/Engine/Private/GameViewportClient.cpp
--- a/Engine/Source/Runtime/Engine/Private/GameViewportClient.cpp
+++ b/Engine/Source/Runtime/Engine/Private/GameViewportClient.cpp
@@ -1429,7 +1429,7 @@ void UGameViewportClient::Draw(FViewport* InViewport, FCanvas* SceneCanvas)
 	ViewFamily.DebugDPIScale = GetDPIScale();
 
 #if WITH_EDITOR
-	if (GIsEditor)
+	if (GIsEditor && ViewFamily.SupportsScreenPercentage())
 	{
 		// Force enable view family show flag for HighDPI derived's screen percentage.
 		ViewFamily.EngineShowFlags.ScreenPercentage = true;
diff --git a/Engine/Source/Runtime/Renderer/Private/PlanarReflectionRendering.cpp b/Engine/Source/Runtime/Renderer/Private/PlanarReflectionRendering.cpp
--- a/Engine/Source/Runtime/Renderer/Private/PlanarReflectionRendering.cpp
+++ b/Engine/Source/Runtime/Renderer/Private/PlanarReflectionRendering.cpp
@@ -612,12 +612,19 @@ void FScene::UpdatePlanarReflectionContents(UPlanarReflectionComponent* CaptureC
 		
 		FPostProcessSettings PostProcessSettings;
 
+		bool bIsMobileMultiViewEnabled = false;
+		if (MainSceneRenderer.ViewFamily.Views.Num() > 0)
+		{
+			bIsMobileMultiViewEnabled = MainSceneRenderer.ViewFamily.Views[0]->Aspects.IsMobileMultiViewEnabled();
+		}
+
 		FSceneViewFamilyContext ViewFamily(FSceneViewFamily::ConstructionValues(
 			CaptureComponent->RenderTarget,
 			this,
 			CaptureComponent->ShowFlags)
 			.SetResolveScene(false)
-			.SetRealtimeUpdate(true));
+			.SetRealtimeUpdate(true)
+			.SetRequireMobileMultiView(bIsMobileMultiViewEnabled));
 
 		// Uses the exact same secondary view fraction on the planar reflection as the main viewport.
 		ViewFamily.SecondaryViewFraction = MainSceneRenderer.ViewFamily.SecondaryViewFraction;
