Fixed stereo layers depth not working.
Also fixes stereo layer left texture crash and CAS crash on mobile.

diff --git a/Engine/Plugins/Runtime/OpenXR/Source/OpenXRHMD/Private/OpenXRHMD.cpp b/Engine/Plugins/Runtime/OpenXR/Source/OpenXRHMD/Private/OpenXRHMD.cpp
--- a/Engine/Plugins/Runtime/OpenXR/Source/OpenXRHMD/Private/OpenXRHMD.cpp
+++ b/Engine/Plugins/Runtime/OpenXR/Source/OpenXRHMD/Private/OpenXRHMD.cpp
@@ -1576,11 +1576,6 @@ bool CheckPlatformDepthExtensionSupport(const XrInstanceProperties& InstanceProp
 	{
 		return false;
 	}
-	else if (FCStringAnsi::Strstr(InstanceProps.runtimeName, "Oculus") && RHIGetInterfaceType() == ERHIInterfaceType::D3D12)
-	{
-		// No PF_DepthStencil compatible formats offered yet
-		return false;
-	}
 	return true;
 }
 
diff --git a/Engine/Source/Runtime/Engine/Private/Components/StereoLayerComponent.cpp b/Engine/Source/Runtime/Engine/Private/Components/StereoLayerComponent.cpp
--- a/Engine/Source/Runtime/Engine/Private/Components/StereoLayerComponent.cpp
+++ b/Engine/Source/Runtime/Engine/Private/Components/StereoLayerComponent.cpp
@@ -115,7 +115,7 @@ void UStereoLayerComponent::TickComponent(float DeltaTime, enum ELevelTick TickT
 		}
 		if (LeftTexture)
 		{
-			Texture->SetForceMipLevelsToBeResident(30.0f);
+			LeftTexture->SetForceMipLevelsToBeResident(30.0f);
 			LayerDesc.LeftTexture = LeftTexture->GetResource()->TextureRHI;
 		}
 				
diff --git a/Engine/Source/Runtime/Renderer/Private/VariableRateShading/ContrastAdaptiveImageGenerator.cpp b/Engine/Source/Runtime/Renderer/Private/VariableRateShading/ContrastAdaptiveImageGenerator.cpp
--- a/Engine/Source/Runtime/Renderer/Private/VariableRateShading/ContrastAdaptiveImageGenerator.cpp
+++ b/Engine/Source/Runtime/Renderer/Private/VariableRateShading/ContrastAdaptiveImageGenerator.cpp
@@ -542,7 +542,7 @@ bool FContrastAdaptiveImageGenerator::IsSupportedByView(const FSceneView& View)
 {
 	EDisplayOutputFormat DisplayOutputFormat = GetDisplayOutputFormat(View);
 	const bool bCompatibleWithOutputType = (DisplayOutputFormat == EDisplayOutputFormat::SDR_sRGB) || IsHDR10(DisplayOutputFormat);
-	return !View.bIsSceneCapture && bCompatibleWithOutputType;
+	return !View.bIsSceneCapture && bCompatibleWithOutputType && !IsMobilePlatform(View.GetShaderPlatform());
 }
 
 FRDGTextureRef FContrastAdaptiveImageGenerator::GetDebugImage(FRDGBuilder& GraphBuilder, const FViewInfo& ViewInfo, FVariableRateShadingImageManager::EVRSImageType ImageType, bool bGetSoftwareImage)
