[Engine]
* Integrated StateStream into Engine. Added limited early support for running static meshes and skinned meshes through state stream. This is very much WIP. Code paths can be enabled by setting GlobalDefinitions WITH_STATE_STREAM and WITH_STATE_STREAM_ACTOR

diff --git a/Engine/Source/Runtime/Engine/Classes/Components/SceneComponent.h b/Engine/Source/Runtime/Engine/Classes/Components/SceneComponent.h
--- a/Engine/Source/Runtime/Engine/Classes/Components/SceneComponent.h
+++ b/Engine/Source/Runtime/Engine/Classes/Components/SceneComponent.h
@@ -2,11 +2,13 @@
 
 #pragma once
 
+#include "EngineDefines.h"
 #include "UObject/ObjectMacros.h"
 #include "Math/BoxSphereBounds.h"
 #include "UObject/UObjectGlobals.h"
 #include "ComponentInstanceDataCache.h"
 #include "Components/ActorComponent.h"
+#include "TransformStateStreamHandle.h"
 
 #include "SceneComponent.generated.h"
 
@@ -333,6 +335,15 @@ public:
 	/** Get the RelativeRotationCache.  */
 	FORCEINLINE const FRotationConversionCache& GetRelativeRotationCache() const { return RelativeRotationCache; }
 
+
+#if WITH_STATE_STREAM_ACTOR
+protected:
+	FTransformHandle TransformHandle;
+
+public:
+	ENGINE_API FTransformHandle& GetOrCreateTransformHandle();
+#endif
+
 private:
 	UFUNCTION()
 	ENGINE_API void OnRep_Transform();
@@ -1674,10 +1685,10 @@ struct FSceneComponentInstanceData : public FActorComponentInstanceData
 {
 	GENERATED_BODY()
 
-	FSceneComponentInstanceData() = default;
+	ENGINE_API FSceneComponentInstanceData();
 	ENGINE_API FSceneComponentInstanceData(const USceneComponent* SourceComponent);
 			
-	virtual ~FSceneComponentInstanceData() = default;
+	ENGINE_API virtual ~FSceneComponentInstanceData();
 
 	ENGINE_API virtual bool ContainsData() const override;
 
diff --git a/Engine/Source/Runtime/Engine/Classes/Components/SkinnedMeshComponent.h b/Engine/Source/Runtime/Engine/Classes/Components/SkinnedMeshComponent.h
--- a/Engine/Source/Runtime/Engine/Classes/Components/SkinnedMeshComponent.h
+++ b/Engine/Source/Runtime/Engine/Classes/Components/SkinnedMeshComponent.h
@@ -17,6 +17,7 @@
 #include "BoneContainer.h"
 #include "ClothingSystemRuntimeTypes.h"
 #include "Animation/SkinWeightProfile.h"
+#include "StateStream/SkinnedMeshStateStreamHandle.h"
 #include "SkinnedMeshComponent.generated.h"
 
 enum class ESkinCacheUsage : uint8;
@@ -395,6 +396,10 @@ private:
 	/** Temporary array of of component-space bone matrices, update each frame and used for rendering the mesh. */
 	TArray<FTransform> ComponentSpaceTransformsArray[2];
 
+#if WITH_STATE_STREAM_ACTOR
+	FSkinnedMeshHandle MeshHandle;
+#endif
+
 protected:
 	/** Array of bone visibilities (containing one of the values in EBoneVisibilityStatus for each bone).  A bone is only visible if it is *exactly* 1 (BVS_Visible) */
 	/** Note these are only used if we are NOT double-buffering bone transforms */
diff --git a/Engine/Source/Runtime/Engine/Classes/Components/StaticMeshComponent.h b/Engine/Source/Runtime/Engine/Classes/Components/StaticMeshComponent.h
--- a/Engine/Source/Runtime/Engine/Classes/Components/StaticMeshComponent.h
+++ b/Engine/Source/Runtime/Engine/Classes/Components/StaticMeshComponent.h
@@ -18,6 +18,7 @@
 #include "UObject/RenderingObjectVersion.h"
 #include "SceneTypes.h"
 #include "DrawDebugHelpers.h"
+#include "StateStream/StaticMeshStateStreamHandle.h"
 #include "StaticMeshComponent.generated.h"
 
 class FColorVertexBuffer;
@@ -133,6 +134,10 @@ private:
 #endif
 	ENGINE_API void NotifyIfStaticMeshChanged();
 
+#if WITH_STATE_STREAM_ACTOR
+	FStaticMeshHandle MeshHandle;
+#endif
+
 public:
 	/** Helper function to get the FName of the private static mesh member */
 	static const FName GetMemberNameChecked_StaticMesh() { return GET_MEMBER_NAME_CHECKED(UStaticMeshComponent, StaticMesh); }
diff --git a/Engine/Source/Runtime/Engine/Classes/Engine/World.h b/Engine/Source/Runtime/Engine/Classes/Engine/World.h
--- a/Engine/Source/Runtime/Engine/Classes/Engine/World.h
+++ b/Engine/Source/Runtime/Engine/Classes/Engine/World.h
@@ -51,6 +51,7 @@ class FRegisterComponentContext;
 class FTimerManager;
 class FWorldInGamePerformanceTrackers;
 class IInterface_PostProcessVolume;
+class IStateStreamManager;
 class UAISystemBase;
 class UChaosEventRelay;
 class UCanvas;
@@ -1411,6 +1412,12 @@ private:
 	/** Physics scene for this world. */
 	FPhysScene*									PhysicsScene;
 	// Note that this should be merged with PhysScene going forward but is needed for now.
+
+#if WITH_STATE_STREAM
+	/** Manager for all existing state streams */
+	IStateStreamManager* StateStreamManager = nullptr;
+#endif
+
 public:
 
 	/** Current global physics scene. */
@@ -1426,6 +1433,16 @@ public:
 	UPROPERTY(Transient)
 	TObjectPtr<class UPhysicsFieldComponent> PhysicsField;
 
+
+#if WITH_STATE_STREAM
+	/** Get statestream of type T */
+	template<typename T> T& GetStateStream() { return *static_cast<T*>(InternalGetStateStream(T::Id)); }
+
+	/** Get statestream manager */
+	IStateStreamManager* GetStateStreamManager() { return StateStreamManager; }
+
+#endif
+
 private:
 
 	/** Array of components that need to wait on tasks before end of frame updates */
@@ -3635,6 +3652,11 @@ private:
 	/** Updates world's required streaming levels */
 	void InternalUpdateStreamingState();
 
+#if WITH_STATE_STREAM
+	/** Returns state stream based on Id */
+	void* InternalGetStateStream(uint32 Id);
+#endif
+
 #if WITH_EDITOR
 public:
 	void SetPlayInEditorInitialNetMode(ENetMode InNetMode)
diff --git a/Engine/Source/Runtime/Engine/Engine.Build.cs b/Engine/Source/Runtime/Engine/Engine.Build.cs
--- a/Engine/Source/Runtime/Engine/Engine.Build.cs
+++ b/Engine/Source/Runtime/Engine/Engine.Build.cs
@@ -109,6 +109,7 @@ public class Engine : ModuleRules
 				"NetworkReplayStreaming",
 				"PhysicsCore",
 				"SignalProcessing",
+				"StateStream",
 				"AudioExtensions",
 				"DeveloperSettings",
 				"AudioLinkCore",
diff --git a/Engine/Source/Runtime/Engine/Private/Components/SceneComponent.cpp b/Engine/Source/Runtime/Engine/Private/Components/SceneComponent.cpp
--- a/Engine/Source/Runtime/Engine/Private/Components/SceneComponent.cpp
+++ b/Engine/Source/Runtime/Engine/Private/Components/SceneComponent.cpp
@@ -17,6 +17,8 @@
 #include "Engine/Texture2D.h"
 #include "ComponentReregisterContext.h"
 #include "Physics/Experimental/PhysScene_Chaos.h"
+#include "TransformStateStream.h"
+#include "TransformStateStreamImpl.h"
 #include "UnrealEngine.h"
 #include "Logging/MessageLog.h"
 #include "Net/UnrealNetwork.h"
@@ -852,6 +854,10 @@ void USceneComponent::OnUnregister()
 {
 	CachedLevelCollection = nullptr;
 
+#if WITH_STATE_STREAM_ACTOR
+	TransformHandle = {};
+#endif
+
 	Super::OnUnregister();
 }
 
@@ -925,6 +931,12 @@ void USceneComponent::PropagateTransformUpdate(bool bTransformChanged, EUpdateTr
 
 			// Flag render transform as dirty
 			MarkRenderTransformDirty();
+
+#if WITH_STATE_STREAM_ACTOR
+			FTransformDynamicState Ds;
+			Ds.SetLocalTransform(ComponentToWorld);
+			TransformHandle.Update(Ds);
+#endif
 		}
 		
 		{
@@ -976,6 +988,12 @@ void USceneComponent::PropagateTransformUpdate(bool bTransformChanged, EUpdateTr
 			//QUICK_SCOPE_CYCLE_COUNTER(STAT_USceneComponent_PropagateTransformUpdate_MarkRenderTransformDirty);
 			// Need to flag as dirty so new bounds are sent to render thread
 			MarkRenderTransformDirty();
+
+#if WITH_STATE_STREAM_ACTOR
+			FTransformDynamicState Ds;
+			Ds.SetLocalTransform(ComponentToWorld);
+			TransformHandle.Update(Ds);
+#endif
 		}
 	}
 }
@@ -2253,6 +2271,15 @@ bool USceneComponent::AttachToComponent(USceneComponent* Parent, const FAttachme
 		}
 #endif
 
+#if WITH_STATE_STREAM_ACTOR
+		if (TransformHandle.IsValid())
+		{
+			//FTransformDynamicState Ds;
+			//Ds.SetParent(Parent->TransformHandle);
+			//TransformHandle.Update(Ds);
+		}
+#endif
+
 		switch (AttachmentRules.LocationRule)
 		{
 		case EAttachmentRule::KeepRelative:
@@ -2368,6 +2395,15 @@ void USceneComponent::DetachFromParent(bool bMaintainWorldPosition, bool bCallMo
 		DetachmentRules.ScaleRule = EDetachmentRule::KeepWorld;
 	}
 
+#if WITH_STATE_STREAM_ACTOR
+	if (TransformHandle.IsValid())
+	{
+		//FTransformDynamicState Ds;
+		//Ds.SetParent(FTransformHandle());
+		//TransformHandle.Update(Ds);
+	}
+#endif
+
 	DetachFromComponent(DetachmentRules);
 }
 
@@ -2507,6 +2543,8 @@ bool USceneComponent::IsAttachedTo(const USceneComponent* TestComp) const
 	return false;
 }
 
+FSceneComponentInstanceData::FSceneComponentInstanceData() = default;
+
 FSceneComponentInstanceData::FSceneComponentInstanceData(const USceneComponent* SourceComponent)
 	: FActorComponentInstanceData(SourceComponent)
 {
@@ -2522,6 +2560,8 @@ FSceneComponentInstanceData::FSceneComponentInstanceData(const USceneComponent*
 	}
 }
 
+FSceneComponentInstanceData::~FSceneComponentInstanceData() = default;
+
 bool FSceneComponentInstanceData::ContainsData() const
 {
 	return AttachedInstanceComponents.Num() > 0 || Super::ContainsData();
@@ -3269,6 +3309,15 @@ bool USceneComponent::GetMaterialPropertyPath(int32 ElementIndex, UObject*& OutO
 void USceneComponent::OnVisibilityChanged()
 {
 	MarkRenderStateDirty();
+
+#if WITH_STATE_STREAM_ACTOR
+	if (TransformHandle.IsValid())
+	{
+		FTransformDynamicState Ds;
+		Ds.SetVisible(IsVisible());
+		TransformHandle.Update(Ds);
+	}
+#endif
 }
 
 void USceneComponent::SetVisibility(const bool bNewVisibility, const USceneComponent::EVisibilityPropagation PropagateToChildren)
@@ -3314,6 +3363,15 @@ void USceneComponent::SetVisibility(const bool bNewVisibility, const USceneCompo
 void USceneComponent::OnHiddenInGameChanged()
 {
 	MarkRenderStateDirty();
+
+#if WITH_STATE_STREAM_ACTOR
+	if (TransformHandle.IsValid())
+	{
+		FTransformDynamicState Ds;
+		Ds.SetVisible(IsVisible());
+		TransformHandle.Update(Ds);
+	}
+#endif
 }
 
 
@@ -3413,6 +3471,31 @@ FBoxSphereBounds USceneComponent::GetPlacementExtent() const
 	return CalcBounds( FTransform::Identity );
 }
 
+#if WITH_STATE_STREAM_ACTOR
+FTransformHandle& USceneComponent::GetOrCreateTransformHandle()
+{
+	if (TransformHandle.IsValid())
+	{
+		return TransformHandle;
+	}
+
+	UWorld* World = GetWorld();
+	check(World);
+	FTransformStaticState Ss;
+	FTransformDynamicState Ds;
+
+	Ds.SetLocalTransform(ComponentToWorld);
+	Ds.SetVisible(IsVisible());
+	if (USceneComponent* Parent = GetAttachParent())
+	{
+		//Ds.SetParent(Parent->GetOrCreateTransformHandle());
+	}
+
+	TransformHandle = World->GetStateStream<ITransformStateStream>().Game_CreateInstance(Ss, Ds);
+	return TransformHandle;
+}
+#endif
+
 void USceneComponent::OnRep_Transform()
 {
 	bNetUpdateTransform = true;
diff --git a/Engine/Source/Runtime/Engine/Private/Components/SkeletalMeshComponent.cpp b/Engine/Source/Runtime/Engine/Private/Components/SkeletalMeshComponent.cpp
--- a/Engine/Source/Runtime/Engine/Private/Components/SkeletalMeshComponent.cpp
+++ b/Engine/Source/Runtime/Engine/Private/Components/SkeletalMeshComponent.cpp
@@ -42,6 +42,7 @@
 #include "Animation/AnimTrace.h"
 #include "Animation/BuiltInAttributeTypes.h"
 #include "Animation/AnimSubsystem_SharedLinkedAnimLayers.h"
+#include "TransformStateStream.h"
 #include "UObject/Stack.h"
 #if WITH_EDITOR
 #include "Engine/PoseWatch.h"
@@ -662,7 +663,6 @@ void USkeletalMeshComponent::FinalizePoseEvaluationResult(const USkeletalMesh* I
 	{
 		OutBoneSpaceTransforms = RefBonePose;
 	}
-
 	OutRootBoneTranslation = OutBoneSpaceTransforms[0].GetTranslation() - RefBonePose[0].GetTranslation();
 }
 
@@ -4201,6 +4201,11 @@ void USkeletalMeshComponent::RefreshMorphTargets()
 
 void USkeletalMeshComponent::SendRenderDynamicData_Concurrent()
 {
+#if WITH_STATE_STREAM_ACTOR
+	if (true)
+		return;
+#endif
+
 	Super::SendRenderDynamicData_Concurrent();
 
 #if WITH_EDITOR
diff --git a/Engine/Source/Runtime/Engine/Private/Components/SkinnedMeshComponent.cpp b/Engine/Source/Runtime/Engine/Private/Components/SkinnedMeshComponent.cpp
--- a/Engine/Source/Runtime/Engine/Private/Components/SkinnedMeshComponent.cpp
+++ b/Engine/Source/Runtime/Engine/Private/Components/SkinnedMeshComponent.cpp
@@ -45,6 +45,8 @@
 #include "Rendering/RenderCommandPipes.h"
 #include "Rendering/NaniteResources.h"
 #include "ProfilingDebugging/AssetMetadataTrace.h"
+#include "StateStream/SkinnedMeshStateStream.h"
+#include "TransformStateStream.h"
 #include "ComponentRecreateRenderStateContext.h"
 #include "NaniteSceneProxy.h"
 
@@ -975,6 +977,25 @@ void USkinnedMeshComponent::CreateRenderState_Concurrent(FRegisterComponentConte
 {
 	LLM_SCOPE(ELLMTag::SkeletalMesh);
 
+#if WITH_STATE_STREAM_ACTOR
+	if (true) // This branch is here to easily toggle between statestream and not.. will be removed once we are on par
+	{
+		if (GetSkinnedAsset() && FApp::CanEverRender() && ShouldComponentAddToScene())
+		{
+			FSkinnedMeshStaticState Ss;
+			Ss.SetMaterialRelevance(GetMaterialRelevance(ERHIFeatureLevel::SM6).Raw);
+			FSkinnedMeshDynamicState Ds;
+			Ds.SetTransform(GetOrCreateTransformHandle());
+			Ds.SetSkinnedAsset(GetSkinnedAsset());
+			Ds.SetOverrideMaterials(OverrideMaterials);
+			MeshHandle = GetWorld()->GetStateStream<ISkinnedMeshStateStream>().Game_CreateInstance(Ss, Ds);
+		}
+		bRenderStateCreated = true;
+		return;
+	}
+#endif
+
+
 	if( GetSkinnedAsset() )
 	{
 		LLM_SCOPE_DYNAMIC_STAT_OBJECTPATH(GetSkinnedAsset()->GetPackage(), ELLMTagSet::Assets);
@@ -1143,6 +1164,10 @@ void USkinnedMeshComponent::CreateRenderState_Concurrent(FRegisterComponentConte
 
 void USkinnedMeshComponent::DestroyRenderState_Concurrent()
 {
+#if WITH_STATE_STREAM_ACTOR
+	MeshHandle = {};
+#endif
+
 	Super::DestroyRenderState_Concurrent();
 
 	for (UMeshDeformerInstance* MeshDeformerInstance : MeshDeformerInstances.DeformerInstances)
@@ -1472,7 +1497,11 @@ void USkinnedMeshComponent::TickComponent(float DeltaTime, enum ELevelTick TickT
 	}
 
 	// See if this mesh was rendered recently. This has to happen first because other data will rely on this
+#if WITH_STATE_STREAM_ACTOR
+	bRecentlyRendered = true;
+#else
 	bRecentlyRendered = ((bUseScreenRenderStateForUpdate ? GetLastRenderTimeOnScreen() : GetLastRenderTime()) > GetWorld()->TimeSeconds - 1.0f);
+#endif
 
 	// Update component's LOD settings
 	// This must be done BEFORE animation Update and Evaluate (TickPose and RefreshBoneTransforms respectively)
@@ -4283,6 +4312,16 @@ void USkinnedMeshComponent::FinalizeBoneTransform()
 	FlipEditableSpaceBases();
 	// we finalized bone transform, now we have valid bone buffer
 	bHasValidBoneTransform = true;
+
+#if WITH_STATE_STREAM_ACTOR
+	if (TransformHandle.IsValid())
+	{
+		FTransformDynamicState Ds;
+		Ds.SetBoneTransforms(GetComponentSpaceTransforms());
+		const_cast<FTransformHandle&>(TransformHandle).Update(Ds);
+	}
+#endif
+
 }
 
 void USkinnedMeshComponent::FlipEditableSpaceBases()
diff --git a/Engine/Source/Runtime/Engine/Private/Components/StaticMeshComponent.cpp b/Engine/Source/Runtime/Engine/Private/Components/StaticMeshComponent.cpp
--- a/Engine/Source/Runtime/Engine/Private/Components/StaticMeshComponent.cpp
+++ b/Engine/Source/Runtime/Engine/Private/Components/StaticMeshComponent.cpp
@@ -9,6 +9,7 @@
 #include "Materials/Material.h"
 #include "Misc/ConfigCacheIni.h"
 #include "RenderUtils.h"
+#include "StateStream/StaticMeshStateStream.h"
 #include "UObject/ObjectSaveContext.h"
 #include "SceneInterface.h"
 #include "UObject/Package.h"
@@ -876,6 +877,10 @@ void UStaticMeshComponent::OnRegister()
 
 void UStaticMeshComponent::OnUnregister()
 {
+#if WITH_STATE_STREAM_ACTOR
+	MeshHandle = {};
+#endif
+
 	Super::OnUnregister();
 	bPrecachedPSOs = false;
 }
@@ -899,7 +904,32 @@ bool UStaticMeshComponent::RequiresGameThreadEndOfFrameRecreate() const
 void UStaticMeshComponent::CreateRenderState_Concurrent(FRegisterComponentContext* Context)
 {
 	LLM_SCOPE(ELLMTag::StaticMesh);
+
+#if WITH_STATE_STREAM_ACTOR
+	{
+		FStaticMeshStaticState Ss;
+		FStaticMeshDynamicState Ds;
+		Ds.SetTransform(GetOrCreateTransformHandle());
+		Ds.SetMesh(StaticMesh);
+		Ds.SetOverrideMaterials(OverrideMaterials);
+		Ds.SetOnlyOwnerSee(bOnlyOwnerSee);
+		Ds.SetOwnerNoSee(bOwnerNoSee);
+		if(bOnlyOwnerSee || bOwnerNoSee || bUseViewOwnerDepthPriorityGroup)
+		{
+			TArray<uint32> Owners;
+			// Make a list of the actors which directly or indirectly own the InComponent.
+			for(const AActor* CurrentOwner = GetOwner();CurrentOwner;CurrentOwner = CurrentOwner->GetOwner())
+			{
+				Owners.Add(CurrentOwner->GetUniqueID());
+			}
+			Ds.SetOwners(Owners);
+		}
+		MeshHandle = GetWorld()->GetStateStream<IStaticMeshStateStream>().Game_CreateInstance(Ss, Ds);
+		bRenderStateCreated = true;
+	}
+#else
 	Super::CreateRenderState_Concurrent(Context);
+#endif
 }
 
 void UStaticMeshComponent::OnCreatePhysicsState()
@@ -2288,6 +2318,12 @@ void UStaticMeshComponent::SetStaticMeshInternal(UStaticMesh* NewMesh)
 {
 	StaticMesh = NewMesh;
 
+#if WITH_STATE_STREAM_ACTOR
+	FStaticMeshDynamicState Ds;
+	Ds.SetMesh(StaticMesh);
+	MeshHandle.Update(Ds);
+#endif
+
 	NotifyIfStaticMeshChanged();
 }
 
diff --git a/Engine/Source/Runtime/Engine/Private/GameViewportClient.cpp b/Engine/Source/Runtime/Engine/Private/GameViewportClient.cpp
--- a/Engine/Source/Runtime/Engine/Private/GameViewportClient.cpp
+++ b/Engine/Source/Runtime/Engine/Private/GameViewportClient.cpp
@@ -19,6 +19,8 @@
 #include "AI/NavigationSystemBase.h"
 #include "Engine/Canvas.h"
 #include "Components/SkeletalMeshComponent.h"
+#include "StateStreamManager.h"
+#include "StateStreamDebugRenderer.h"
 #include "UObject/UObjectIterator.h"
 #include "UObject/Package.h"
 #include "Particles/ParticleSystemComponent.h"
@@ -1899,6 +1901,29 @@ void UGameViewportClient::Draw(FViewport* InViewport, FCanvas* SceneCanvas)
 		MyWorld->FXSystem->DrawDebug(SceneCanvas);
 	}
 
+#if WITH_STATE_STREAM
+	if (IStateStreamManager* Manager = MyWorld->GetStateStreamManager())
+	{
+		class FStateStreamDebugRenderer : public IStateStreamDebugRenderer
+		{
+		public:
+			FStateStreamDebugRenderer(FViewport* V, UCanvas* C) : Viewport(V), Canvas(C) {}
+
+			virtual void DrawText(const FStringView& Text) override
+			{
+				Canvas->SetDrawColor(FColor::Orange);
+				Canvas->DrawText(GEngine->GetSmallFont(), Text, 100, Y);
+				Y += 20;
+			}
+
+			uint32 Y = 200;
+			FViewport* Viewport;
+			UCanvas* Canvas;
+		} StateStreamDebugRenderer(InViewport, DebugCanvasObject);
+		Manager->Game_DebugRender(StateStreamDebugRenderer);
+	}
+#endif
+
 	// Render the UI.
 	if (FSlateApplication::Get().GetPlatformApplication()->IsAllowedToRender())
 	{
diff --git a/Engine/Source/Runtime/Engine/Private/LevelTick.cpp b/Engine/Source/Runtime/Engine/Private/LevelTick.cpp
--- a/Engine/Source/Runtime/Engine/Private/LevelTick.cpp
+++ b/Engine/Source/Runtime/Engine/Private/LevelTick.cpp
@@ -8,6 +8,7 @@
 #include "Async/ParallelFor.h"
 #include "Misc/TimeGuard.h"
 #include "GameFramework/WorldSettings.h"
+#include "StateStreamManager.h"
 #include "UObject/UObjectStats.h"
 #include "EngineStats.h"
 #include "RenderingThread.h"
@@ -1801,6 +1802,11 @@ void UWorld::Tick( ELevelTick TickType, float DeltaSeconds )
 
 	FWorldDelegates::OnWorldTickEnd.Broadcast(this, TickType, DeltaSeconds);
 
+#if WITH_STATE_STREAM
+	StateStreamManager->Game_EndTick(RealTimeSeconds);
+	StateStreamManager->Game_BeginTick();
+#endif
+
 	FThreadIdleStats::EndCriticalPath();
 }
 
diff --git a/Engine/Source/Runtime/Engine/Private/World.cpp b/Engine/Source/Runtime/Engine/Private/World.cpp
--- a/Engine/Source/Runtime/Engine/Private/World.cpp
+++ b/Engine/Source/Runtime/Engine/Private/World.cpp
@@ -17,6 +17,7 @@
 #include "Logging/LogScopedCategoryAndVerbosityOverride.h"
 #include "Logging/LogScopedVerbosityOverride.h"
 #include "Particles/ParticleSystemComponent.h"
+#include "StateStreamManager.h"
 #include "Stats/StatsMisc.h"
 #include "Misc/ScopedSlowTask.h"
 #include "SceneInterface.h"
@@ -677,6 +678,10 @@ UWorld::UWorld( const FObjectInitializer& ObjectInitializer )
 
 UWorld::~UWorld()
 {
+#if WITH_STATE_STREAM
+	check(!StateStreamManager);
+#endif
+
 	if (PerfTrackers)
 	{
 		delete PerfTrackers;
@@ -1336,9 +1341,22 @@ void UWorld::FinishDestroy()
 		{
 			Scene->Release();
 			Scene = NULL;
+
+#if WITH_STATE_STREAM
+			StateStreamManager = nullptr;
+#endif
+
 		}
 	}
 
+#if WITH_STATE_STREAM
+	if (StateStreamManager)
+	{
+		GetRendererModule().FreeStateStream(StateStreamManager);
+		StateStreamManager = nullptr;
+	}
+#endif
+
 	// Clear GWorld pointer if it's pointing to this object.
 	if( GWorld == this )
 	{
@@ -2137,6 +2155,11 @@ void UWorld::InitWorld(const InitializationValues IVS)
 		GetRendererModule().AllocateScene(this, bRequiresHitProxies, IVS.bCreateFXSystem, GetFeatureLevel());
 	}
 
+#if WITH_STATE_STREAM
+	StateStreamManager = GetRendererModule().AllocateStateStream(this);
+	StateStreamManager->Game_BeginTick();
+#endif
+
 	// Prepare AI systems
 	if (WorldSettings)
 	{
@@ -4290,6 +4313,14 @@ void UWorld::InternalUpdateStreamingState()
 	});
 }
 
+#if WITH_STATE_STREAM
+void* UWorld::InternalGetStateStream(uint32 Id)
+{
+	check(StateStreamManager);
+	return StateStreamManager->Game_GetStreamPointer(Id);
+}
+#endif
+
 bool UWorld::CanAddLoadedLevelToWorld(ULevel* Level) const
 {
 	if (CurrentLevelPendingVisibility == nullptr)
@@ -5583,6 +5614,10 @@ void UWorld::CleanupWorldInternal(bool bSessionEnded, bool bCleanupResources, bo
 	// Tell actors to remove their components from the scene.
 	ClearWorldComponents();
 
+#if WITH_STATE_STREAM
+	StateStreamManager->Game_EndTick(RealTimeSeconds);
+#endif
+
 	if (bCleanupResources && PersistentLevel)
 	{
 		PersistentLevel->ReleaseRenderingResources();
@@ -6109,7 +6144,7 @@ AWorldSettings* UWorld::K2_GetWorldSettings()
 
 AWorldSettings* UWorld::GetWorldSettings( const bool bCheckStreamingPersistent, const bool bChecked ) const
 {
-	checkSlow(!IsInActualRenderingThread());
+	//checkSlow(!IsInActualRenderingThread());
 	AWorldSettings* WorldSettings = nullptr;
 	if (PersistentLevel)
 	{
diff --git a/Engine/Source/Runtime/Engine/Public/StateStream/SkinnedMeshStateStream.h b/Engine/Source/Runtime/Engine/Public/StateStream/SkinnedMeshStateStream.h
new file mode 100644
--- /dev/null
+++ b/Engine/Source/Runtime/Engine/Public/StateStream/SkinnedMeshStateStream.h
@@ -0,0 +1,64 @@
+// Copyright Epic Games, Inc. All Rights Reserved.
+
+#pragma once
+
+#include "Materials/MaterialRelevance.h"
+#include "SkinnedMeshStateStreamHandle.h"
+#include "StateStreamDefinitions.h"
+#include "TransformStateStreamHandle.h"
+#include "SkinnedMeshStateStream.generated.h"
+
+class USkinnedAsset;
+class UMaterialInterface;
+
+////////////////////////////////////////////////////////////////////////////////////////////////////
+// Skinned state for mesh instance. Can only be set upon creation
+
+USTRUCT(StateStreamStaticState)
+struct FSkinnedMeshStaticState
+{
+	GENERATED_USTRUCT_BODY()
+
+	UPROPERTY()
+	uint64 MaterialRelevance;
+};
+
+
+////////////////////////////////////////////////////////////////////////////////////////////////////
+// Dynamic state for mesh instance. Can be updated inside ticks
+
+USTRUCT(StateStreamDynamicState)
+struct FSkinnedMeshDynamicState
+{
+	GENERATED_USTRUCT_BODY()
+
+private: // Use accessors (SetX/GetX) instead of properties directly
+
+	UPROPERTY()
+	FTransformHandle Transform;
+
+	UPROPERTY()
+	TObjectPtr<USkinnedAsset> SkinnedAsset;
+
+	UPROPERTY()
+	TArray<TObjectPtr<UMaterialInterface>> OverrideMaterials;
+};
+
+
+////////////////////////////////////////////////////////////////////////////////////////////////////
+// Mesh state stream id used for registering dependencies and find statestream
+
+inline constexpr uint32 SkinnedMeshStateStreamId = 3;
+
+
+////////////////////////////////////////////////////////////////////////////////////////////////////
+// Interface for creating mesh instances
+
+class ISkinnedMeshStateStream
+{
+public:
+	DECLARE_STATESTREAM(SkinnedMesh)
+	virtual FSkinnedMeshHandle Game_CreateInstance(const FSkinnedMeshStaticState& Ss, const FSkinnedMeshDynamicState& Ds) = 0;
+};
+
+////////////////////////////////////////////////////////////////////////////////////////////////////
diff --git a/Engine/Source/Runtime/Engine/Public/StateStream/SkinnedMeshStateStreamHandle.h b/Engine/Source/Runtime/Engine/Public/StateStream/SkinnedMeshStateStreamHandle.h
new file mode 100644
--- /dev/null
+++ b/Engine/Source/Runtime/Engine/Public/StateStream/SkinnedMeshStateStreamHandle.h
@@ -0,0 +1,19 @@
+// Copyright Epic Games, Inc. All Rights Reserved.
+
+#pragma once
+
+#include "StateStreamHandle.h"
+#include "SkinnedMeshStateStreamHandle.generated.h"
+
+struct FSkinnedMeshDynamicState;
+
+////////////////////////////////////////////////////////////////////////////////////////////////////
+// Handle for mesh instance
+
+USTRUCT(StateStreamHandle)
+struct FSkinnedMeshHandle : public FStateStreamHandle
+{
+	GENERATED_USTRUCT_BODY()
+};
+
+////////////////////////////////////////////////////////////////////////////////////////////////////
diff --git a/Engine/Source/Runtime/Engine/Public/StateStream/StaticMeshStateStream.h b/Engine/Source/Runtime/Engine/Public/StateStream/StaticMeshStateStream.h
new file mode 100644
--- /dev/null
+++ b/Engine/Source/Runtime/Engine/Public/StateStream/StaticMeshStateStream.h
@@ -0,0 +1,67 @@
+// Copyright Epic Games, Inc. All Rights Reserved.
+
+#pragma once
+
+#include "StaticMeshStateStreamHandle.h"
+#include "StateStreamDefinitions.h"
+#include "TransformStateStreamHandle.h"
+#include "StaticMeshStateStream.generated.h"
+
+class UStaticMesh;
+class UMaterialInterface;
+
+////////////////////////////////////////////////////////////////////////////////////////////////////
+// Static state for mesh instance. Can only be set upon creation
+
+USTRUCT(StateStreamStaticState)
+struct FStaticMeshStaticState
+{
+	GENERATED_USTRUCT_BODY()
+};
+
+
+////////////////////////////////////////////////////////////////////////////////////////////////////
+// Dynamic state for mesh instance. Can be updated inside ticks
+
+USTRUCT(StateStreamDynamicState)
+struct FStaticMeshDynamicState
+{
+	GENERATED_USTRUCT_BODY()
+
+private: // Use accessors (SetX/GetX) instead of properties directly
+
+	UPROPERTY()
+	FTransformHandle Transform;
+
+	UPROPERTY()
+	TObjectPtr<UStaticMesh> Mesh;
+
+	UPROPERTY()
+	TArray<TObjectPtr<UMaterialInterface>> OverrideMaterials;
+
+	UPROPERTY()
+	TArray<uint32> Owners;
+
+	UPROPERTY()
+	uint32 bOnlyOwnerSee : 1;
+
+	UPROPERTY()
+	uint32 bOwnerNoSee : 1;
+};
+
+
+////////////////////////////////////////////////////////////////////////////////////////////////////
+// Mesh state stream id used for registering dependencies and find statestream
+
+inline constexpr uint32 StaticMeshStateStreamId = 2;
+
+
+////////////////////////////////////////////////////////////////////////////////////////////////////
+// Interface for creating mesh instances
+
+class IStaticMeshStateStream
+{
+public:
+	DECLARE_STATESTREAM(StaticMesh)
+	virtual FStaticMeshHandle Game_CreateInstance(const FStaticMeshStaticState& Ss, const FStaticMeshDynamicState& Ds) = 0;
+};
diff --git a/Engine/Source/Runtime/Engine/Public/StateStream/StaticMeshStateStreamHandle.h b/Engine/Source/Runtime/Engine/Public/StateStream/StaticMeshStateStreamHandle.h
new file mode 100644
--- /dev/null
+++ b/Engine/Source/Runtime/Engine/Public/StateStream/StaticMeshStateStreamHandle.h
@@ -0,0 +1,19 @@
+// Copyright Epic Games, Inc. All Rights Reserved.
+
+#pragma once
+
+#include "StateStreamHandle.h"
+#include "StaticMeshStateStreamHandle.generated.h"
+
+struct FStaticMeshDynamicState;
+
+////////////////////////////////////////////////////////////////////////////////////////////////////
+// Handle for mesh instance
+
+USTRUCT(StateStreamHandle)
+struct FStaticMeshHandle : public FStateStreamHandle
+{
+	GENERATED_USTRUCT_BODY()
+};
+
+////////////////////////////////////////////////////////////////////////////////////////////////////
diff --git a/Engine/Source/Runtime/RenderCore/Public/RendererInterface.h b/Engine/Source/Runtime/RenderCore/Public/RendererInterface.h
--- a/Engine/Source/Runtime/RenderCore/Public/RendererInterface.h
+++ b/Engine/Source/Runtime/RenderCore/Public/RendererInterface.h
@@ -41,6 +41,7 @@ template<typename ShaderType, typename PointerTableType> class TShaderRefBase;
 class FSceneUniformBuffer;
 class FBatchedPrimitiveParameters;
 class ISceneRenderer;
+class IStateStreamManager;
 
 namespace Nanite
 {
@@ -694,6 +695,11 @@ public:
 	
 	virtual void RemoveScene(FSceneInterface* Scene) = 0;
 
+#if WITH_STATE_STREAM
+	virtual IStateStreamManager* AllocateStateStream(UWorld* World) = 0;
+	virtual void FreeStateStream(IStateStreamManager* Manager) = 0;
+#endif
+
 	/**
 	* Updates all static draw lists for each allocated scene.
 	*/
diff --git a/Engine/Source/Runtime/RenderCore/Public/RenderingThread.h b/Engine/Source/Runtime/RenderCore/Public/RenderingThread.h
--- a/Engine/Source/Runtime/RenderCore/Public/RenderingThread.h
+++ b/Engine/Source/Runtime/RenderCore/Public/RenderingThread.h
@@ -419,7 +419,11 @@ public:
 	template <typename RenderCommandTag>
 	static void Enqueue(FRenderCommandPipe* Pipe, FCommandListFunction&& Function)
 	{
-		if (GRenderCommandPipeMode == ERenderCommandPipeMode::All && Pipe)
+		if (GRenderCommandPipeMode == ERenderCommandPipeMode::All && Pipe
+#if WITH_STATE_STREAM
+			&& !IsInRenderingThread() // Note, when statestream enabled lots of the commands are happening on the render thread (and should not be commands at all)
+#endif
+			)
 		{
 			UE::TScopeLock Lock(Pipe->Mutex);
 
diff --git a/Engine/Source/Runtime/Renderer/Private/RendererModule.h b/Engine/Source/Runtime/Renderer/Private/RendererModule.h
--- a/Engine/Source/Runtime/Renderer/Private/RendererModule.h
+++ b/Engine/Source/Runtime/Renderer/Private/RendererModule.h
@@ -152,6 +152,11 @@ public:
 	virtual void RequestStaticMeshUpdate(FPrimitiveSceneInfo* Info) override;
 	virtual void AddMeshBatchToGPUScene(FGPUScenePrimitiveCollector* Collector, FMeshBatch& MeshBatch) override;
 
+#if WITH_STATE_STREAM
+	virtual IStateStreamManager* AllocateStateStream(UWorld* World) override;
+	virtual void FreeStateStream(IStateStreamManager* Manager) override;
+#endif
+
 private:
 	TSet<FSceneInterface*> AllocatedScenes;
 	FOnPostOpaqueRender PostOpaqueRenderDelegate;
diff --git a/Engine/Source/Runtime/Renderer/Private/RendererScene.cpp b/Engine/Source/Runtime/Renderer/Private/RendererScene.cpp
--- a/Engine/Source/Runtime/Renderer/Private/RendererScene.cpp
+++ b/Engine/Source/Runtime/Renderer/Private/RendererScene.cpp
@@ -79,7 +79,11 @@
 #include "Materials/MaterialRenderProxy.h"
 #include "ProfilingDebugging/AssetMetadataTrace.h"
 #include "ProfilingDebugging/CountersTrace.h"
+#include "StateStream/SkinnedMeshStateStreamImpl.h"
+#include "StateStream/StaticMeshStateStreamImpl.h"
+#include "StateStreamManagerImpl.h"
 #include "SceneCulling/SceneCulling.h"
+#include "TransformStateStreamImpl.h"
 #include "InstanceCulling/InstanceCullingOcclusionQuery.h"
 #include "ComputeWorkerInterface.h"
 #include "Nanite/NaniteMaterialsSceneExtension.h"
@@ -4522,11 +4526,24 @@ void FScene::Release()
 
 	UE::RenderCommandPipe::FSyncScope SyncScope;
 
+#if WITH_STATE_STREAM
+	FStateStreamManagerImpl* StateStreamManager = static_cast<FStateStreamManagerImpl*>(World->GetStateStreamManager());
+#endif
+
 	// Send a command to the rendering thread to release the scene.
 	FScene* Scene = this;
 	ENQUEUE_RENDER_COMMAND(FReleaseCommand)(
-		[Scene](FRHICommandListImmediate& RHICmdList)
+		[Scene
+#if WITH_STATE_STREAM
+		, StateStreamManager
+#endif
+		](FRHICommandListImmediate& RHICmdList)
 		{
+#if WITH_STATE_STREAM
+			StateStreamManager->Render_Exit();
+			GetRendererModule().FreeStateStream(StateStreamManager);
+#endif
+
 			// Flush any remaining batched primitive update commands before deleting the scene.
 			FUpdateParameters UpdateParameters;
 			UpdateParameters.bDestruction = true;
@@ -6993,6 +7010,31 @@ void FRendererModule::RemoveScene(FSceneInterface* Scene)
 	AllocatedScenes.Remove(Scene);
 }
 
+#if WITH_STATE_STREAM
+IStateStreamManager* FRendererModule::AllocateStateStream(UWorld* World)
+{
+	FScene* Scene = (FScene*)World->Scene; // TODO: Check if really FScene?
+	FStateStreamManagerImpl* Manager = new FStateStreamManagerImpl();
+
+	FTransformStateStreamImpl& TransformStreamImpl = *new FTransformStateStreamImpl();
+	FStaticMeshStateStreamImpl& StaticMeshStreamImpl = *new FStaticMeshStateStreamImpl(Scene);
+	FSkinnedMeshStateStreamImpl& SkinnedMeshStreamImpl = *new FSkinnedMeshStateStreamImpl(Scene);
+	Manager->Render_Register(TransformStreamImpl, true);
+	Manager->Render_Register(StaticMeshStreamImpl, true);
+	Manager->Render_Register(SkinnedMeshStreamImpl, true);
+
+	Manager->Render_RegisterDependency(StaticMeshStreamImpl, TransformStreamImpl);
+	Manager->Render_RegisterDependency(SkinnedMeshStreamImpl, TransformStreamImpl);
+
+	return Manager;
+}
+
+void FRendererModule::FreeStateStream(IStateStreamManager* Manager)
+{
+	delete static_cast<FStateStreamManagerImpl*>(Manager);
+}
+#endif
+
 void FRendererModule::UpdateStaticDrawLists()
 {
 	// Update all static meshes in order to recache cached mesh draw commands.
diff --git a/Engine/Source/Runtime/Renderer/Private/SceneRendering.cpp b/Engine/Source/Runtime/Renderer/Private/SceneRendering.cpp
--- a/Engine/Source/Runtime/Renderer/Private/SceneRendering.cpp
+++ b/Engine/Source/Runtime/Renderer/Private/SceneRendering.cpp
@@ -7,6 +7,7 @@
 #include "SceneRendering.h"
 #include "ProfilingDebugging/ProfilingHelpers.h"
 #include "ProfilingDebugging/CpuProfilerTrace.h"
+#include "StateStreamManagerImpl.h"
 #include "UObject/UObjectHash.h"
 #include "UObject/UObjectIterator.h"
 #include "EngineGlobals.h"
@@ -4645,6 +4646,7 @@ FSceneViewFamily* FSceneRenderer::CreateSplitScreenDebugViewFamily(const FSceneV
 				FMatrix YawRotation = UE::Math::TRotationMatrix<double>::Make(FRotator(0.0, Rotate, 0.0));
 				InitOptions.ViewRotationMatrix = YawRotation * InitOptions.ViewRotationMatrix;
 
+#if !WITH_STATE_STREAM // TODO: Fix. ViewActor does not exist on render side
 				// And optionally orbit the position around the player
 				if (CVarSplitScreenDebugOrbit.GetValueOnGameThread() && InitOptions.ViewActor)
 				{
@@ -4652,6 +4654,7 @@ FSceneViewFamily* FSceneRenderer::CreateSplitScreenDebugViewFamily(const FSceneV
 					InitOptions.ViewOrigin = YawRotation.GetTransposed().TransformVector(InitOptions.ViewOrigin - TargetTranslation) + TargetTranslation;
 					InitOptions.ViewLocation = InitOptions.ViewOrigin;
 				}
+#endif
 			}
 
 			// Convert adjusted matrix back to a rotation
@@ -5204,6 +5207,12 @@ static void RenderViewFamilies_RenderThread(FRHICommandListImmediate& RHICmdList
 	bool bUpdatedGPUSkinCacheVisualization = false;
 #endif
 
+#if WITH_STATE_STREAM
+	FStateStreamManagerImpl* StateStreamManager = static_cast<FStateStreamManagerImpl*>(Scene->GetWorld()->GetStateStreamManager());
+	StateStreamManager->Render_Update(Scene->GetWorld()->RealTimeSeconds);
+	StateStreamManager->Render_GarbageCollect(true);
+#endif
+
 	for (FSceneRenderer* SceneRenderer : SceneRenderers)
 	{
 		const ERHIFeatureLevel::Type FeatureLevel = SceneRenderer->FeatureLevel;
diff --git a/Engine/Source/Runtime/Renderer/Private/StateStream/SkinnedMeshStateStreamImpl.cpp b/Engine/Source/Runtime/Renderer/Private/StateStream/SkinnedMeshStateStreamImpl.cpp
new file mode 100644
--- /dev/null
+++ b/Engine/Source/Runtime/Renderer/Private/StateStream/SkinnedMeshStateStreamImpl.cpp
@@ -0,0 +1,210 @@
+// Copyright Epic Games, Inc. All Rights Reserved.
+
+#include "SkinnedMeshStateStreamImpl.h"
+#include "HAL/PlatformMisc.h"
+#include "Rendering/RenderCommandPipes.h"
+#include "ScenePrivate.h"
+#include "SkeletalMeshSceneProxy.h"
+#include "SkinnedMeshSceneProxyDesc.h"
+#include "../../../Engine/Private/SkeletalRenderGPUSkin.h"
+
+////////////////////////////////////////////////////////////////////////////////////////////////////
+
+FSkinnedMeshObject::~FSkinnedMeshObject()
+{
+	FScene& Scene = static_cast<FScene&>(PrimitiveSceneData.SceneProxy->GetScene());
+	Scene.RemovePrimitive(&PrimitiveSceneDesc);
+
+	if (TransformObject)
+	{
+		TransformObject->RemoveListener(this);
+		TransformObject = nullptr;
+	}
+}
+
+void FSkinnedMeshObject::OnTransformObjectDirty()
+{
+	FPrimitiveSceneProxy& Proxy = *PrimitiveSceneData.SceneProxy;
+	FScene& Scene = static_cast<FScene&>(Proxy.GetScene());
+
+	FTransformObject::Info Info = TransformObject->GetInfo();
+
+	#if UE_DEBUG_OFFSET_SKINNED_MESH
+	FTransform Transform = Info.WorldTransform;
+	Transform.AddToTranslation(FVector(10.0f*DebugIndex, 10.0f*DebugIndex, 0));
+	#else
+	const FTransform& Transform = Info.WorldTransform;
+	#endif
+
+	FMatrix Mat = Transform.ToMatrixWithScale();
+	MeshObject->SetTransform(Mat, Scene.GetFrameNumber());
+	MeshObject->RefreshClothingTransforms(Mat, Scene.GetFrameNumber());
+
+	const FBoxSphereBounds& LocalBounds = Proxy.GetLocalBounds();
+	FBoxSphereBounds Bounds = LocalBounds.TransformBy(Transform);
+	FVector ActorPosForRendering(ForceInitToZero);// Proxy.GetActorPositionForRenderer(); TODO This need to be revisited
+	Scene.UpdatePrimitiveTransform_RenderThread(&Proxy, Bounds, LocalBounds, Mat, ActorPosForRendering, {});
+
+	int32 LODLevel = 0;
+
+	TArray<FExternalMorphSets> ExternalMorphSets;
+	ExternalMorphSets.Add({});
+	FSkinnedMeshSceneProxyDynamicData ProxyData;
+	ProxyData.ExternalMorphSets = ExternalMorphSets;
+	ProxyData.ComponentSpaceTransforms = Info.BoneTransforms;
+	ProxyData.PreviousComponentSpaceTransforms = PrevTransforms;
+	ProxyData.ComponentWorldTransform = Transform;
+	ProxyData.PreviousBoneTransformRevisionNumber = BoneTransformRevisionNumber;
+	ProxyData.CurrentBoneTransformFrame = BoneTransformRevisionNumber;
+	ProxyData.CurrentBoneTransformRevisionNumber = ++BoneTransformRevisionNumber;
+	ProxyData.NumLODs = 1;
+
+	//TArray<uint8> VisibilityStates;
+	//VisibilityStates.SetNumUninitialized(Info.BoneTransforms.Num());
+	//memset(VisibilityStates.GetData(), 1, Info.BoneTransforms.Num());
+
+	//ProxyData.BoneVisibilityStates = VisibilityStates;
+	//ProxyData.PreviousBoneVisibilityStates = VisibilityStates;
+
+	MeshObject->Update(LODLevel, ProxyData, &Proxy, SkinnedAsset, FMorphTargetWeightMap(), TArray<float>(), EPreviousBoneTransformUpdateMode::None, FExternalMorphWeightData());
+	MeshObject->bHasBeenUpdatedAtLeastOnce = true;
+
+	PrevTransforms = Info.BoneTransforms;
+}
+
+////////////////////////////////////////////////////////////////////////////////////////////////////
+
+FSkinnedMeshStateStreamImpl::FSkinnedMeshStateStreamImpl(FScene* InScene)
+:	Scene(InScene)
+{
+}
+
+void FSkinnedMeshStateStreamImpl::SetTransformObject(FSkinnedMeshObject& Object, const FSkinnedMeshDynamicState& Ds)
+{
+	if (Object.TransformObject)
+	{
+		Object.TransformObject->RemoveListener(&Object);
+		Object.TransformObject = nullptr;
+	}
+
+	const FTransformHandle& TransformHandle = Ds.GetTransform();
+	if (TransformHandle.IsValid())
+	{
+		FTransformObject* TransformObject = static_cast<FTransformObject*>(TransformHandle.Render_GetUserData());
+		check(TransformObject);
+		TransformObject->AddListener(&Object);
+		Object.TransformObject = TransformObject;
+	}
+}
+
+void FSkinnedMeshStateStreamImpl::Render_OnCreate(const FSkinnedMeshStaticState& Ss, const FSkinnedMeshDynamicState& Ds, FSkinnedMeshObject*& UserData, bool IsDestroyedInSameFrame)
+{
+	check(!IsDestroyedInSameFrame);
+	check(!UserData);
+
+	USkinnedAsset* SkinnedAsset = Ds.GetSkinnedAsset();
+	if (!SkinnedAsset)
+	{
+		return;
+	}
+
+	FSkinnedMeshObject& Object = *new FSkinnedMeshObject();
+	Object.AddRef();
+	Object.SkinnedAsset = SkinnedAsset;
+
+	SetTransformObject(Object, Ds);
+
+	check(Object.TransformObject);
+	FTransformObject::Info Info = Object.TransformObject->GetInfo();
+
+	#if UE_DEBUG_OFFSET_SKINNED_MESH
+	FTransform Transform = Info.WorldTransform;
+	static uint32 DebugCounter = 0;
+	Object.DebugIndex = 5;
+	Transform.AddToTranslation(FVector(10.0f*Object.DebugIndex, 10.0f*Object.DebugIndex, 0));
+	#else
+	const FTransform& Transform = Info.WorldTransform;
+	#endif
+
+	FBoxSphereBounds LocalBounds(ForceInitToZero);
+	if (SkinnedAsset)
+	{
+		LocalBounds = SkinnedAsset->GetBounds();
+	}	
+
+	UWorld& World = *Scene->World;
+
+	Object.PrimitiveSceneDesc.RenderMatrix = Transform.ToMatrixWithScale();
+	Object.PrimitiveSceneDesc.PrimitiveSceneData = &Object.PrimitiveSceneData;
+	Object.PrimitiveSceneDesc.LocalBounds = LocalBounds;
+	Object.PrimitiveSceneDesc.Bounds = LocalBounds.TransformBy(Transform);
+
+	FSkinnedMeshSceneProxyDesc Desc;
+	Desc.SkinnedAsset = SkinnedAsset;
+	Desc.OverrideMaterials = const_cast<TArray<TObjectPtr<UMaterialInterface>>&>(Ds.GetOverrideMaterials());
+	Desc.CustomPrimitiveData = &Object.CustomPrimitiveData;
+	Desc.Scene = Scene;
+	Desc.World = &World;
+	Desc.FeatureLevel = World.GetFeatureLevel();
+	Desc.bIsVisible = Info.bVisible;
+	Desc.bPerBoneMotionBlur = true;
+	Desc.VisibilityId = -1;
+	Desc.bReceivesDecals = false;
+	Desc.bCollisionEnabled = true;
+	Desc.MaterialRelevance.Raw = Ss.GetMaterialRelevance();
+	Desc.CastShadow = true;
+	Desc.bCastDynamicShadow = true;
+	Desc.bCastStaticShadow = true;
+	Desc.bCastContactShadow = true;
+	Desc.bUseAsOccluder = true;
+	
+	FSkeletalMeshObject* MeshObject = FSkinnedMeshSceneProxyDesc::CreateMeshObject(Desc);
+	Desc.MeshObject = MeshObject;
+	Object.MeshObject = MeshObject;
+
+	#if WITH_EDITOR
+	//Desc.TextureStreamingTransformScale = Transform.GetMaximumAxisScale();
+	#endif
+
+	FPrimitiveSceneProxy* SceneProxy = FSkinnedMeshSceneProxyDesc::CreateSceneProxy(Desc, false, 0);
+	Object.PrimitiveSceneData.SceneProxy = SceneProxy;
+
+	SceneProxy->SetPrimitiveColor(FLinearColor::White);
+
+	int32 LODLevel = Desc.GetPredictedLODLevel();
+
+	TArray<FExternalMorphSets> ExternalMorphSets;
+	ExternalMorphSets.Add({});
+	FSkinnedMeshSceneProxyDynamicData ProxyData;
+	ProxyData.ExternalMorphSets = ExternalMorphSets;
+	ProxyData.ComponentSpaceTransforms = Info.BoneTransforms;
+	ProxyData.ComponentWorldTransform = Info.WorldTransform;
+
+	MeshObject->Update(LODLevel, ProxyData, SceneProxy, SkinnedAsset, FMorphTargetWeightMap(), TArray<float>(), EPreviousBoneTransformUpdateMode::None, FExternalMorphWeightData());
+
+	Scene->AddPrimitive(&Object.PrimitiveSceneDesc);
+
+	UserData = &Object;
+}
+
+void FSkinnedMeshStateStreamImpl::Render_OnUpdate(const FSkinnedMeshStaticState& Ss, const FSkinnedMeshDynamicState& Ds, FSkinnedMeshObject*& Object)
+{
+	if (!Object)
+	{
+		return;
+	}
+	if (Ds.TransformModified())
+	{
+		SetTransformObject(*Object, Ds);
+	}
+}
+
+void FSkinnedMeshStateStreamImpl::Render_OnDestroy(const FSkinnedMeshStaticState& Ss, const FSkinnedMeshDynamicState& Ds, FSkinnedMeshObject*& Object)
+{
+	if (Object)
+	{
+		Object->Release();
+	}
+}
+
+////////////////////////////////////////////////////////////////////////////////////////////////////
diff --git a/Engine/Source/Runtime/Renderer/Private/StateStream/SkinnedMeshStateStreamImpl.h b/Engine/Source/Runtime/Renderer/Private/StateStream/SkinnedMeshStateStreamImpl.h
new file mode 100644
--- /dev/null
+++ b/Engine/Source/Runtime/Renderer/Private/StateStream/SkinnedMeshStateStreamImpl.h
@@ -0,0 +1,72 @@
+// Copyright Epic Games, Inc. All Rights Reserved.
+
+#pragma once
+
+#include "PrimitiveSceneDesc.h"
+#include "SceneTypes.h"
+#include "StateStream/SkinnedMeshStateStream.h"
+#include "TransformStateStreamImpl.h"
+
+#define UE_API RENDERER_API
+
+class FScene;
+class FSkeletalMeshObject;
+struct FPrimitiveSceneDesc;
+
+#define UE_DEBUG_OFFSET_SKINNED_MESH 0 // For debugging purposes
+
+////////////////////////////////////////////////////////////////////////////////////////////////////
+
+class FSkinnedMeshObject : public TRefCountingMixin<FSkinnedMeshObject>, public FTransformObjectListener
+{
+public:
+
+private:
+	virtual ~FSkinnedMeshObject();
+	virtual void OnTransformObjectDirty() override final;
+
+	TRefCountPtr<FTransformObject> TransformObject;
+	USkinnedAsset* SkinnedAsset;
+	FSkeletalMeshObject* MeshObject;
+
+	FCustomPrimitiveData CustomPrimitiveData;
+	FPrimitiveSceneInfoData	PrimitiveSceneData;
+	FPrimitiveSceneDesc PrimitiveSceneDesc;
+	TArray<FTransform> PrevTransforms;
+	uint32 BoneTransformRevisionNumber = 0;
+
+	#if UE_DEBUG_OFFSET_SKINNED_MESH
+	uint32 DebugIndex = 0;
+	#endif
+
+	friend class FSkinnedMeshStateStreamImpl;
+	friend class TRefCountingMixin<FSkinnedMeshObject>;
+};
+
+
+////////////////////////////////////////////////////////////////////////////////////////////////////
+
+struct FSkinnedMeshStateStreamSettings : TStateStreamSettings<ISkinnedMeshStateStream, FSkinnedMeshObject>
+{
+	static inline constexpr bool SkipCreatingDeletes = true;
+};
+
+
+////////////////////////////////////////////////////////////////////////////////////////////////////
+
+class FSkinnedMeshStateStreamImpl : public TStateStream<FSkinnedMeshStateStreamSettings>
+{
+public:
+	FSkinnedMeshStateStreamImpl(FScene* InScene);
+private:
+	UE_API void SetTransformObject(FSkinnedMeshObject& Object, const FSkinnedMeshDynamicState& Ds);
+	UE_API virtual void Render_OnCreate(const FSkinnedMeshStaticState& Ss, const FSkinnedMeshDynamicState& Ds, FSkinnedMeshObject*& UserData, bool IsDestroyedInSameFrame) override;
+	UE_API virtual void Render_OnUpdate(const FSkinnedMeshStaticState& Ss, const FSkinnedMeshDynamicState& Ds, FSkinnedMeshObject*& UserData) override;
+	UE_API virtual void Render_OnDestroy(const FSkinnedMeshStaticState& Ss, const FSkinnedMeshDynamicState& Ds, FSkinnedMeshObject*& UserData) override;
+
+	FScene* Scene;
+};
+
+////////////////////////////////////////////////////////////////////////////////////////////////////
+
+#undef UE_API
diff --git a/Engine/Source/Runtime/Renderer/Private/StateStream/StaticMeshStateStreamImpl.cpp b/Engine/Source/Runtime/Renderer/Private/StateStream/StaticMeshStateStreamImpl.cpp
new file mode 100644
--- /dev/null
+++ b/Engine/Source/Runtime/Renderer/Private/StateStream/StaticMeshStateStreamImpl.cpp
@@ -0,0 +1,132 @@
+// Copyright Epic Games, Inc. All Rights Reserved.
+
+#include "StaticMeshStateStreamImpl.h"
+#include "ScenePrivate.h"
+#include "StaticMeshSceneProxy.h"
+#include "StaticMeshSceneProxyDesc.h"
+
+////////////////////////////////////////////////////////////////////////////////////////////////////
+
+FStaticMeshObject::~FStaticMeshObject()
+{
+	FScene& Scene = (FScene&)PrimitiveSceneData.SceneProxy->GetScene();
+	Scene.RemovePrimitive(&PrimitiveSceneDesc);
+
+	if (TransformObject)
+	{
+		TransformObject->RemoveListener(this);
+		TransformObject = nullptr;
+	}
+}
+
+void FStaticMeshObject::OnTransformObjectDirty()
+{
+	FTransformObject::Info Info = TransformObject->GetInfo();
+	// TODO: Implement this to handle changing visibility and movement
+}
+
+////////////////////////////////////////////////////////////////////////////////////////////////////
+
+FStaticMeshStateStreamImpl::FStaticMeshStateStreamImpl(FScene* InScene)
+:	Scene(InScene)
+{
+}
+
+void FStaticMeshStateStreamImpl::SetTransformObject(FStaticMeshObject& Object, const FStaticMeshDynamicState& Ds)
+{
+	if (Object.TransformObject)
+	{
+		Object.TransformObject->RemoveListener(&Object);
+		Object.TransformObject = nullptr;
+	}
+
+	const FTransformHandle& TransformHandle = Ds.GetTransform();
+	if (TransformHandle.IsValid())
+	{
+		FTransformObject* TransformObject = static_cast<FTransformObject*>(TransformHandle.Render_GetUserData());
+		check(TransformObject);
+		TransformObject->AddListener(&Object);
+		Object.TransformObject = TransformObject;
+	}
+}
+
+void FStaticMeshStateStreamImpl::Render_OnCreate(const FStaticMeshStaticState& Ss, const FStaticMeshDynamicState& Ds, FStaticMeshObject*& UserData, bool IsDestroyedInSameFrame)
+{
+	check(!UserData);
+
+	UStaticMesh* Mesh = Ds.GetMesh();
+	if (!Mesh)
+	{
+		return;
+	}
+
+	FStaticMeshObject& Object = *new FStaticMeshObject();
+	Object.AddRef();
+
+	SetTransformObject(Object, Ds);
+
+	FTransformObject::Info Info = Object.TransformObject->GetInfo();
+	const FTransform& Transform = Info.WorldTransform;
+
+	FBoxSphereBounds LocalBounds(ForceInitToZero);
+	if (Mesh)
+	{
+		LocalBounds = Mesh->GetBounds();
+	}	
+
+	UWorld& World = *Scene->World;
+
+	Object.PrimitiveSceneDesc.RenderMatrix = Transform.ToMatrixWithScale();
+	Object.PrimitiveSceneDesc.AttachmentRootPosition = Transform.GetLocation();
+	Object.PrimitiveSceneDesc.PrimitiveSceneData = &Object.PrimitiveSceneData;
+	Object.PrimitiveSceneDesc.LocalBounds = LocalBounds;
+	Object.PrimitiveSceneDesc.Bounds = LocalBounds.TransformBy(Transform);
+
+	FStaticMeshSceneProxyDesc Desc;
+	Desc.StaticMesh = Mesh;
+	Desc.OverrideMaterials = const_cast<TArray<TObjectPtr<UMaterialInterface>>&>(Ds.GetOverrideMaterials());
+	Desc.CustomPrimitiveData = &Object.CustomPrimitiveData;
+	Desc.Scene = Scene;
+	Desc.World = &World;
+	Desc.FeatureLevel = World.GetFeatureLevel();
+	Desc.bIsVisible = Info.bVisible;
+	Desc.bOnlyOwnerSee = Ds.GetOnlyOwnerSee();
+	Desc.bOwnerNoSee = Ds.GetOwnerNoSee();
+	Desc.ActorOwners = Ds.GetOwners();
+
+
+	#if WITH_EDITOR
+	Desc.TextureStreamingTransformScale = Transform.GetMaximumAxisScale();
+	#endif
+
+	Object.PrimitiveSceneData.SceneProxy = new FStaticMeshSceneProxy(Desc, false);
+
+	Scene->AddPrimitive(&Object.PrimitiveSceneDesc);
+
+	UserData = &Object;
+}
+
+void FStaticMeshStateStreamImpl::Render_OnUpdate(const FStaticMeshStaticState& Ss, const FStaticMeshDynamicState& Ds, FStaticMeshObject*& UserData)
+{
+	if (!UserData)
+	{
+		return;
+	}
+
+	FStaticMeshObject& Object = *UserData;
+
+	if (Ds.TransformModified())
+	{
+		SetTransformObject(Object, Ds);
+	}
+}
+
+void FStaticMeshStateStreamImpl::Render_OnDestroy(const FStaticMeshStaticState& Ss, const FStaticMeshDynamicState& Ds, FStaticMeshObject*& UserData)
+{
+	if (UserData)
+	{
+		UserData->Release();
+	}
+}
+
+////////////////////////////////////////////////////////////////////////////////////////////////////
diff --git a/Engine/Source/Runtime/Renderer/Private/StateStream/StaticMeshStateStreamImpl.h b/Engine/Source/Runtime/Renderer/Private/StateStream/StaticMeshStateStreamImpl.h
new file mode 100644
--- /dev/null
+++ b/Engine/Source/Runtime/Renderer/Private/StateStream/StaticMeshStateStreamImpl.h
@@ -0,0 +1,59 @@
+// Copyright Epic Games, Inc. All Rights Reserved.
+
+#pragma once
+
+#include "PrimitiveSceneDesc.h"
+#include "SceneTypes.h"
+#include "StateStream/StaticMeshStateStream.h"
+#include "TransformStateStreamImpl.h"
+
+#define UE_API RENDERER_API
+
+class FScene;
+struct FPrimitiveSceneDesc;
+
+////////////////////////////////////////////////////////////////////////////////////////////////////
+
+class FStaticMeshObject : public TRefCountingMixin<FStaticMeshObject>, public FTransformObjectListener
+{
+public:
+
+private:
+	virtual ~FStaticMeshObject();
+	virtual void OnTransformObjectDirty() override final;
+
+	TRefCountPtr<FTransformObject> TransformObject;
+
+	FCustomPrimitiveData CustomPrimitiveData;
+	FPrimitiveSceneInfoData	PrimitiveSceneData;
+	FPrimitiveSceneDesc PrimitiveSceneDesc;
+
+	friend class FStaticMeshStateStreamImpl;
+	friend class TRefCountingMixin<FStaticMeshObject>;
+};
+
+////////////////////////////////////////////////////////////////////////////////////////////////////
+
+struct FStaticMeshStateStreamSettings : TStateStreamSettings<IStaticMeshStateStream, FStaticMeshObject>
+{
+	static inline constexpr bool SkipCreatingDeletes = true;
+};
+
+////////////////////////////////////////////////////////////////////////////////////////////////////
+
+class FStaticMeshStateStreamImpl : public TStateStream<FStaticMeshStateStreamSettings>
+{
+public:
+	FStaticMeshStateStreamImpl(FScene* InScene);
+private:
+	UE_API void SetTransformObject(FStaticMeshObject& Object, const FStaticMeshDynamicState& Ds);
+	UE_API virtual void Render_OnCreate(const FStaticMeshStaticState& Ss, const FStaticMeshDynamicState& Ds, FStaticMeshObject*& UserData, bool IsDestroyedInSameFrame) override;
+	UE_API virtual void Render_OnUpdate(const FStaticMeshStaticState& Ss, const FStaticMeshDynamicState& Ds, FStaticMeshObject*& UserData) override;
+	UE_API virtual void Render_OnDestroy(const FStaticMeshStaticState& Ss, const FStaticMeshDynamicState& Ds, FStaticMeshObject*& UserData) override;
+
+	FScene* Scene;
+};
+
+////////////////////////////////////////////////////////////////////////////////////////////////////
+
+#undef UE_API
diff --git a/Engine/Source/Runtime/Renderer/Renderer.Build.cs b/Engine/Source/Runtime/Renderer/Renderer.Build.cs
--- a/Engine/Source/Runtime/Renderer/Renderer.Build.cs
+++ b/Engine/Source/Runtime/Renderer/Renderer.Build.cs
@@ -47,6 +47,7 @@ public class Renderer : ModuleRules
 				"ImageWriteQueue",
 				"RHI",
 				"MaterialShaderQualitySettings",
+				"StateStream",
 				"TraceLog",
 			}
 		);
