Fix incorrect dual source blending on PS4 with Substrate due to incorrect output format for the second RT.

diff --git a/Engine/Source/Runtime/Renderer/Private/Substrate/Substrate.cpp b/Engine/Source/Runtime/Renderer/Private/Substrate/Substrate.cpp
--- a/Engine/Source/Runtime/Renderer/Private/Substrate/Substrate.cpp
+++ b/Engine/Source/Runtime/Renderer/Private/Substrate/Substrate.cpp
@@ -1342,10 +1342,16 @@ void SetBasePassRenderTargetOutputFormat(const EShaderPlatform Platform, const F
 		}
 		const FGBufferInfo BufferInfo = FetchFullGBufferInfo(GBufferParams);
 
-		// Add N uint for Substrate fast path
-		for (int i = 0; i < SUBSTRATE_BASE_PASS_MRT_OUTPUT_COUNT; ++i)
+		// Translucent blend mode do not write material data, and thus don't need output format (default to RGBA16f). 
+		// Dual source blending requires to have both target format set to RGBA16f
+		const bool bIsTranslucent = IsTranslucentBlendMode(MaterialParameters.BlendMode);
+		if (!bIsTranslucent)
 		{
-			OutEnvironment.SetRenderTargetOutputFormat(BufferInfo.NumTargets + i, PF_R32_UINT);
+			// Add N uint for Substrate fast path
+			for (int i = 0; i < SUBSTRATE_BASE_PASS_MRT_OUTPUT_COUNT; ++i)
+			{
+				OutEnvironment.SetRenderTargetOutputFormat(BufferInfo.NumTargets + i, PF_R32_UINT);
+			}
 		}
 
 		// Add another MRT for Substrate top layer information
