Substrate - fixe material data debug for RT.

#tests dx12 editor with RT
#rnx

diff --git a/Engine/Shaders/Private/RayTracing/RayTracingDebug.usf b/Engine/Shaders/Private/RayTracing/RayTracingDebug.usf
--- a/Engine/Shaders/Private/RayTracing/RayTracingDebug.usf
+++ b/Engine/Shaders/Private/RayTracing/RayTracingDebug.usf
@@ -140,10 +140,9 @@ float4 DebugRayTracingMaterial(uint2 PixelCoord, FRayDesc Ray, bool bIsFarField,
 				const float3 V = -normalize(TranslatedWorldPosition - PrimaryView.TranslatedWorldCameraOrigin);
 				const uint ClosureIndex = 0;
 				SubstratePrintMaterialProperties(Ctx, PixelCoord, TranslatedWorldPosition, V, ClosureIndex, Payload.SubstrateData);
-				bWriteOuput = false;
-				return 0; // No scene color writing
 			}
-
+			bWriteOuput = false;
+			return float4(0, 0, 0, 1);
 		}
 		else if (SubstratePixelHeader.ClosureCount > 0)
 		{
diff --git a/Engine/Source/Runtime/Renderer/Private/RayTracing/RayTracingDebug.cpp b/Engine/Source/Runtime/Renderer/Private/RayTracing/RayTracingDebug.cpp
--- a/Engine/Source/Runtime/Renderer/Private/RayTracing/RayTracingDebug.cpp
+++ b/Engine/Source/Runtime/Renderer/Private/RayTracing/RayTracingDebug.cpp
@@ -1416,6 +1416,13 @@ void FDeferredShadingSceneRenderer::RenderRayTracingDebug(FRDGBuilder& GraphBuil
 {
 	const uint32 DebugVisualizationMode = GetRaytracingDebugViewMode(View);
 
+	if (DebugVisualizationMode == RAY_TRACING_DEBUG_VIZ_SUBSTRATE_DATA)
+	{
+		ShaderPrint::SetEnabled(true);
+		ShaderPrint::RequestSpaceForLines(1024);
+		ShaderPrint::RequestSpaceForCharacters(1024);
+	}
+
 	if (DebugVisualizationMode == RAY_TRACING_DEBUG_VIZ_BARYCENTRICS)
 	{
 		return RenderRayTracingBarycentrics(GraphBuilder, View, SceneColorTexture, (bool)GVisualizeProceduralPrimitives);
