Fix decals with ISR

diff --git a/Engine/Source/Runtime/Renderer/Private/CompositionLighting/PostProcessDeferredDecals.cpp b/Engine/Source/Runtime/Renderer/Private/CompositionLighting/PostProcessDeferredDecals.cpp
--- a/Engine/Source/Runtime/Renderer/Private/CompositionLighting/PostProcessDeferredDecals.cpp
+++ b/Engine/Source/Runtime/Renderer/Private/CompositionLighting/PostProcessDeferredDecals.cpp
@@ -391,7 +391,8 @@ static bool RenderPreStencil(FRHICommandList& RHICmdList, const FViewInfo& View,
 	RHICmdList.SetStreamSource(0, GetUnitCubeVertexBuffer(), 0);
 
 	// Render decal mask
-	RHICmdList.DrawIndexedPrimitive(GetUnitCubeIndexBuffer(), 0, 0, 8, 0, UE_ARRAY_COUNT(GCubeIndices) / 3, View.GetStereoPassInstanceFactor());
+	uint32 InstanceCount = View.Aspects.IsInstancedMultiViewportEnabled() ? 1 : View.GetStereoPassInstanceFactor();
+	RHICmdList.DrawIndexedPrimitive(GetUnitCubeIndexBuffer(), 0, 0, 8, 0, UE_ARRAY_COUNT(GCubeIndices) / 3, InstanceCount);
 
 	return true;
 }
@@ -754,7 +755,8 @@ void AddDeferredDecalPass(
 				}
 #endif // PSO_PRECACHING_VALIDATE
 
-				RHICmdList.DrawIndexedPrimitive(GetUnitCubeIndexBuffer(), 0, 0, 8, 0, UE_ARRAY_COUNT(GCubeIndices) / 3, View.GetStereoPassInstanceFactor());
+				uint32 InstanceCount = View.Aspects.IsInstancedMultiViewportEnabled() ? 1 : View.GetStereoPassInstanceFactor();
+				RHICmdList.DrawIndexedPrimitive(GetUnitCubeIndexBuffer(), 0, 0, 8, 0, UE_ARRAY_COUNT(GCubeIndices) / 3, InstanceCount);
 			}
 		});
 	};
