Path Tracer: Improve light grid construction

Take into account barndoors and spot light angles when building the path tracing light grid.

This also fixes small missing regions at the edge of the bounding box for spot and rect lights.

Move the AABB code and light/AABB overlap tests to common headers so they can be shared betwen the path tracing and ray tracing light grids

diff --git a/Engine/Plugins/Experimental/GPULightmass/Source/GPULightmass/Private/LightmapRayTracing.h b/Engine/Plugins/Experimental/GPULightmass/Source/GPULightmass/Private/LightmapRayTracing.h
--- a/Engine/Plugins/Experimental/GPULightmass/Source/GPULightmass/Private/LightmapRayTracing.h
+++ b/Engine/Plugins/Experimental/GPULightmass/Source/GPULightmass/Private/LightmapRayTracing.h
@@ -24,7 +24,7 @@ BEGIN_SHADER_PARAMETER_STRUCT(FPathTracingLightGrid, RENDERER_API)
 	SHADER_PARAMETER_RDG_BUFFER_SRV(Buffer<uint>, LightGridData)
 	SHADER_PARAMETER(unsigned, LightGridResolution)
 	SHADER_PARAMETER(unsigned, LightGridMaxCount)
-	SHADER_PARAMETER(int, LightGridAxis)
+	SHADER_PARAMETER(unsigned, LightGridAxis)
 END_SHADER_PARAMETER_STRUCT()
 
 class FLightmapRayTracingMeshProcessor : public FRayTracingMeshProcessor
diff --git a/Engine/Shaders/Private/IntersectionUtils.ush b/Engine/Shaders/Private/IntersectionUtils.ush
--- a/Engine/Shaders/Private/IntersectionUtils.ush
+++ b/Engine/Shaders/Private/IntersectionUtils.ush
@@ -63,4 +63,200 @@ bool SphereIntersectSphere(float4 SphereCenterAndRadius, float4 OtherSphereCente
 	float CombinedRadii = SphereCenterAndRadius.w + OtherSphereCenterAndRadius.w;
 	float3 VectorBetweenCenters = SphereCenterAndRadius.xyz - OtherSphereCenterAndRadius.xyz;
 	return dot(VectorBetweenCenters, VectorBetweenCenters) < CombinedRadii * CombinedRadii;
+}
+
+// Given two ranges of distances along a ray - check if they overlap
+// NOTE: This check ignores empty intervals and intervals which contain only a single point
+bool CheckOverlap(float2 Ta, float2 Tb) {
+	return max(Ta.x, Tb.x) < min(Ta.y, Tb.y);
+}
+
+bool RayOverlapsAABB(float3 O, float3 D, float TMax, float3 Lo, float3 Hi)
+{
+	float3 InvRayDir = rcp(D);
+	float3 T0 = (Lo - O) * InvRayDir;
+	float3 T1 = (Hi - O) * InvRayDir;
+	float3 TNear = min(T0, T1), TFar = max(T0, T1);
+	return max(max(TNear.x, TNear.y), max(TNear.z, 0.0)) <= min(min(TFar.x, TFar.y), min(TFar.z, TMax));
+}
+
+// Compute the overlap between an infinite ray and an infinite cone volume
+// The Apex is the origin of the cone, the normal is the unit vector , CosTheta is the half-angle describing the cone's opening
+float2 RayInfiniteConeVolumeOverlap(
+	float3 RayOrigin, float3 RayDirection,
+	float3 Apex, float3 Axis,
+	float CosTheta)
+{
+	// Adapted from: https://www.iquilezles.org/www/articles/intersectors/intersectors.htm
+	const float3 ba = Axis * CosTheta;
+	const float3 oa = RayOrigin - Apex;
+	const float3 rd = RayDirection;
+
+	const float m0 = dot(ba, ba);
+	const float m1 = dot(oa, ba);
+	const float m3 = dot(rd, ba);
+
+	const float m4 = dot(rd, oa);
+	const float m5 = dot(oa, oa);
+
+	const float k2 = m0 * m0 - m3 * m3;
+	const float k1 = m0 * m0 * m4 - m1 * m3;
+	const float k0 = m0 * m0 * m5 - m1 * m1;
+
+	const float h = k1 * k1 - k2 * k0;
+	if (h > 0)
+	{
+		const float s = sign(k2);
+		const float2 t = (-k1 + float2(-s, s) * sqrt(h)) / k2;
+		const float2 y = m1 + t * m3;
+
+		// check which hits are valid:
+		const int c = ((y.x > 0) ? 1 : 0) + (y.y > 0 ? 2 : 0);
+		if (c == 1)
+		{
+			// only smaller root is valid, far root is invalid
+			return float2(0.0, t.x);
+		}
+		if (c == 2)
+		{
+			// only larger root is valid, must have started outside the cone
+			return float2(t.y, POSITIVE_INFINITY);
+		}
+		if (c == 3)
+		{
+			// both roots are valid
+			return t;
+		}
+	}
+	return -1.0;
+}
+
+float2 RaySphereOverlap(float3 Ro, float3 Rd, float3 Center, float Radius, float TMin, float TMax)
+{
+	float3 oc = Ro - Center;
+	float b = dot(oc, Rd);
+	float h = Radius * Radius - length2(oc - b * Rd);
+	if (h > 0)
+	{
+		float2 t = -b + float2(-1, +1) * sqrt(h);
+		return float2(max(t.x, TMin), min(t.y, TMax));
+	}
+	return -1.0;
+}
+
+// Check for overlap between an AABB and a Cone (the portion of the cone inside a sphere of the given Radius)
+// NOTE: This test can be expensive, so one should check for AABB vs sphere overlap first
+// NOTE: The last argument can be used to do a cheaper test using the AABB of the rounded cone
+bool AABBOverlapsCurvedCone(float3 Lo, float3 Hi, float3 Apex, float3 Axis, float CosTheta, float Radius, bool bUseApproxTest)
+{
+	if (bUseApproxTest)
+	{
+		// Compute AABB of the cone, and check for overlap of the two AABBs
+		// box around ray from light center to tip of the cone
+		float3 Tip = Apex + Axis * Radius;
+		float3 LightBoundLo = min(Apex, Tip);
+		float3 LightBoundHi = max(Apex, Tip);
+
+		// expand by disc around the farthest part of the cone
+		float SinTheta2 = 1 - CosTheta * CosTheta;
+		float3 Disc = sqrt(saturate(SinTheta2 * (1.0 - Axis * Axis)));
+		LightBoundLo = min(LightBoundLo, Apex + Radius * (Axis * CosTheta - Disc));
+		LightBoundHi = max(LightBoundHi, Apex + Radius * (Axis * CosTheta + Disc));
+
+		float3 F = select(abs(Axis) > CosTheta, Apex + Radius * sign(Axis), Apex); // include far point along axis if it lies inside the cone
+		LightBoundLo = min(LightBoundLo, F);
+		LightBoundHi = max(LightBoundHi, F);
+
+		// intersect bounds and see if we have anything left
+		return all(max(Lo, LightBoundLo) < min(Hi, LightBoundHi));
+	}
+	else
+	{
+		return
+			// Does the Apex line inside the AABB?
+			all((Lo <= Apex) & (Apex <= Hi)) ||
+			// Does the central axis of the cone pass overlap the AABB?
+			RayOverlapsAABB(Apex, Axis, Radius, Lo, Hi) ||
+			// Test each of the 12 AABB edges against cone -- if any of them overlap, we can stop
+			// NOTE: this looks expensive, but there is a huge amount of simplification and common terms between each of these calls
+			// NOTE: This is a volumetric test, so edges entirely _inside_ the cone, will pass the test
+			// Check X edges
+			CheckOverlap(RayInfiniteConeVolumeOverlap(float3(Lo.x, Lo.y, Lo.z), float3(1, 0, 0), Apex, Axis, CosTheta), RaySphereOverlap(float3(Lo.x, Lo.y, Lo.z), float3(1, 0, 0), Apex, Radius, 0, Hi.x - Lo.x)) ||
+			CheckOverlap(RayInfiniteConeVolumeOverlap(float3(Lo.x, Lo.y, Hi.z), float3(1, 0, 0), Apex, Axis, CosTheta), RaySphereOverlap(float3(Lo.x, Lo.y, Hi.z), float3(1, 0, 0), Apex, Radius, 0, Hi.x - Lo.x)) ||
+			CheckOverlap(RayInfiniteConeVolumeOverlap(float3(Lo.x, Hi.y, Lo.z), float3(1, 0, 0), Apex, Axis, CosTheta), RaySphereOverlap(float3(Lo.x, Hi.y, Lo.z), float3(1, 0, 0), Apex, Radius, 0, Hi.x - Lo.x)) ||
+			CheckOverlap(RayInfiniteConeVolumeOverlap(float3(Lo.x, Hi.y, Hi.z), float3(1, 0, 0), Apex, Axis, CosTheta), RaySphereOverlap(float3(Lo.x, Hi.y, Hi.z), float3(1, 0, 0), Apex, Radius, 0, Hi.x - Lo.x)) ||
+			// Check Y edges
+			CheckOverlap(RayInfiniteConeVolumeOverlap(float3(Lo.x, Lo.y, Lo.z), float3(0, 1, 0), Apex, Axis, CosTheta), RaySphereOverlap(float3(Lo.x, Lo.y, Lo.z), float3(0, 1, 0), Apex, Radius, 0, Hi.y - Lo.y)) ||
+			CheckOverlap(RayInfiniteConeVolumeOverlap(float3(Lo.x, Lo.y, Hi.z), float3(0, 1, 0), Apex, Axis, CosTheta), RaySphereOverlap(float3(Lo.x, Lo.y, Hi.z), float3(0, 1, 0), Apex, Radius, 0, Hi.y - Lo.y)) ||
+			CheckOverlap(RayInfiniteConeVolumeOverlap(float3(Hi.x, Lo.y, Lo.z), float3(0, 1, 0), Apex, Axis, CosTheta), RaySphereOverlap(float3(Hi.x, Lo.y, Lo.z), float3(0, 1, 0), Apex, Radius, 0, Hi.y - Lo.y)) ||
+			CheckOverlap(RayInfiniteConeVolumeOverlap(float3(Hi.x, Lo.y, Hi.z), float3(0, 1, 0), Apex, Axis, CosTheta), RaySphereOverlap(float3(Hi.x, Lo.y, Hi.z), float3(0, 1, 0), Apex, Radius, 0, Hi.y - Lo.y)) ||
+			// Check Z edges
+			CheckOverlap(RayInfiniteConeVolumeOverlap(float3(Lo.x, Lo.y, Lo.z), float3(0, 0, 1), Apex, Axis, CosTheta), RaySphereOverlap(float3(Lo.x, Lo.y, Lo.z), float3(0, 0, 1), Apex, Radius, 0, Hi.z - Lo.z)) ||
+			CheckOverlap(RayInfiniteConeVolumeOverlap(float3(Lo.x, Hi.y, Lo.z), float3(0, 0, 1), Apex, Axis, CosTheta), RaySphereOverlap(float3(Lo.x, Hi.y, Lo.z), float3(0, 0, 1), Apex, Radius, 0, Hi.z - Lo.z)) ||
+			CheckOverlap(RayInfiniteConeVolumeOverlap(float3(Hi.x, Lo.y, Lo.z), float3(0, 0, 1), Apex, Axis, CosTheta), RaySphereOverlap(float3(Hi.x, Lo.y, Lo.z), float3(0, 0, 1), Apex, Radius, 0, Hi.z - Lo.z)) ||
+			CheckOverlap(RayInfiniteConeVolumeOverlap(float3(Hi.x, Hi.y, Lo.z), float3(0, 0, 1), Apex, Axis, CosTheta), RaySphereOverlap(float3(Hi.x, Hi.y, Lo.z), float3(0, 0, 1), Apex, Radius, 0, Hi.z - Lo.z));
+	}
+}
+
+// Test AABB against solid sphere
+bool AABBOverlapsSphere(float3 Lo, float3 Hi, float3 Center, float Radius)
+{
+	// James Arvo - Graphics Gems (p. 335)
+	return length2(max(Lo - Center, 0) + max(Center - Hi, 0)) < Radius * Radius;
+}
+
+// Does the AABB lie fully "behind" the plane defined by Center and Normal
+// NOTE: Normal vector does not need to be normalized
+bool IsBoxBehindPlane(float3 Lo, float3 Hi, float3 Normal, float3 Center)
+{
+	return dot(select(Normal > 0, Hi, Lo) - Center, Normal) < 0;
+}
+
+// Test AABB against the penumbra frustum defined by a rectlight's barndoors
+// This test is expensive, so should only be used if the barndoors are active
+// It is assumed the AABB already overlaps the sphere of influence of the rect light and is at least partially "in front" of the light
+// The last portion of the test is a bit more expensive, so skip it optionally (at the risk of more false-positives)
+bool AABBOverlapsRectLightFrustum(float3 Lo, float3 Hi, float3 Center, float3 Normal, float3 Tangent, float HalfWidth, float HalfHeight, float Radius, float BarnCos, float BarnLen, bool bImprovedFrustumTest)
+{
+	const float3 dPdv = Tangent;
+	const float3 dPdu = cross(Normal, dPdv);
+	const float BarnSin = sqrt(1 - BarnCos * BarnCos);
+
+	float3 BoundingPlane = float3(
+		2 * HalfWidth  + BarnLen * BarnSin,
+		2 * HalfHeight + BarnLen * BarnSin,
+		BarnLen * BarnCos
+	);
+	// Test if box is outside any of the 4 barndoor penumbra planes:
+	if (IsBoxBehindPlane(Lo, Hi, BoundingPlane.x * Normal + BoundingPlane.z * dPdu, Center - dPdu * HalfWidth) ||
+		IsBoxBehindPlane(Lo, Hi, BoundingPlane.x * Normal - BoundingPlane.z * dPdu, Center + dPdu * HalfWidth) ||
+		IsBoxBehindPlane(Lo, Hi, BoundingPlane.y * Normal + BoundingPlane.z * dPdv, Center - dPdv * HalfHeight) ||
+		IsBoxBehindPlane(Lo, Hi, BoundingPlane.y * Normal - BoundingPlane.z * dPdv, Center + dPdv * HalfHeight))
+	{
+		return false;
+	}
+	
+	if (!bImprovedFrustumTest)
+	{
+		return true;
+	}
+
+	// Checking the AABB against the 4 planes is not sufficient, and some cases which still lie outside could still pass through
+	// See description here: https://iquilezles.org/articles/frustumcorrect/
+
+	// Normalize the penumbra vector and scale it to the back of the sphere of influence
+	BoundingPlane = normalize(BoundingPlane);
+	BoundingPlane.xy *= Radius * rcp(BoundingPlane.z);
+	// Get 8 corners of the penumbra frustum in world space
+	const float3 C00 = Center - HalfWidth * dPdu - HalfHeight * dPdv, F00 = C00 - BoundingPlane.x * dPdu - BoundingPlane.y * dPdv + Radius * Normal;
+	const float3 C01 = Center + HalfWidth * dPdu + HalfHeight * dPdv, F01 = C01 - BoundingPlane.x * dPdu + BoundingPlane.y * dPdv + Radius * Normal;
+	const float3 C10 = Center - HalfWidth * dPdu - HalfHeight * dPdv, F10 = C10 + BoundingPlane.x * dPdu - BoundingPlane.y * dPdv + Radius * Normal;
+	const float3 C11 = Center + HalfWidth * dPdu + HalfHeight * dPdv, F11 = C11 + BoundingPlane.x * dPdu + BoundingPlane.y * dPdv + Radius * Normal;
+	// Get an AABB around the 8 frustum points and check if it overlaps the query AABB
+	if (any(max(max(max(C00, F00), max(C01, F01)), max(max(C10, F10), max(C11, F11))) < Lo) ||
+		any(min(min(min(C00, F00), min(C01, F01)), min(min(C10, F10), min(C11, F11))) > Hi))
+	{
+		return false;
+	}
+	return true;
 }
\ No newline at end of file
diff --git a/Engine/Shaders/Private/PathTracing/Light/PathTracingLightGrid.ush b/Engine/Shaders/Private/PathTracing/Light/PathTracingLightGrid.ush
--- a/Engine/Shaders/Private/PathTracing/Light/PathTracingLightGrid.ush
+++ b/Engine/Shaders/Private/PathTracing/Light/PathTracingLightGrid.ush
@@ -19,7 +19,7 @@ Texture2DArray<uint> LightGrid;
 Buffer<uint> LightGridData;
 uint LightGridResolution;
 uint LightGridMaxCount;
-int LightGridAxis;
+uint LightGridAxis;
 #endif
 
 struct FLightLoopCount {
@@ -52,12 +52,12 @@ FLightLoopCount LightGridLookup(float3 TranslatedWorldPos)
 		uint LightGridNz = LightGrid.Load(int4(UVz, 2, 0));
 		uint LightGridN = LightGridNx;
 		uint LightGridOffset = LightGridMaxCount * (0 + 3 * (UVx.x + UVx.y * LightGridResolution));
-		if ((LightGridAxis < 0 && (LightGridNy & PATHTRACER_LIGHT_GRID_LIGHT_COUNT_MASK) < (LightGridN & PATHTRACER_LIGHT_GRID_LIGHT_COUNT_MASK)) || LightGridAxis == 1)
+		if ((LightGridAxis >= 3 && (LightGridNy & PATHTRACER_LIGHT_GRID_LIGHT_COUNT_MASK) < (LightGridN & PATHTRACER_LIGHT_GRID_LIGHT_COUNT_MASK)) || LightGridAxis == 1)
 		{
 			LightGridN = LightGridNy;
 			LightGridOffset = LightGridMaxCount * (1 + 3 * (UVy.x + UVy.y * LightGridResolution));
 		}
-		if ((LightGridAxis < 0 && (LightGridNz & PATHTRACER_LIGHT_GRID_LIGHT_COUNT_MASK) < (LightGridN & PATHTRACER_LIGHT_GRID_LIGHT_COUNT_MASK)) || LightGridAxis == 2)
+		if ((LightGridAxis >= 3 && (LightGridNz & PATHTRACER_LIGHT_GRID_LIGHT_COUNT_MASK) < (LightGridN & PATHTRACER_LIGHT_GRID_LIGHT_COUNT_MASK)) || LightGridAxis == 2)
 		{
 			LightGridN = LightGridNz;
 			LightGridOffset = LightGridMaxCount * (2 + 3 * (UVz.x + UVz.y * LightGridResolution));
diff --git a/Engine/Shaders/Private/PathTracing/Light/PathTracingPointLight.ush b/Engine/Shaders/Private/PathTracing/Light/PathTracingPointLight.ush
--- a/Engine/Shaders/Private/PathTracing/Light/PathTracingPointLight.ush
+++ b/Engine/Shaders/Private/PathTracing/Light/PathTracingPointLight.ush
@@ -8,7 +8,7 @@ PathTracingPointLight.ush: Light sampling functions for Point lights
 
 #include "PathTracingLightCommon.ush"
 #include "PathTracingCapsuleLight.ush"
-
+#include "../../IntersectionUtils.ush"
 
 
 // #dxr_todo: These are stolen from DynamicLightingCommon.ush but I was having trouble getting that header to include cleanly
@@ -164,57 +164,6 @@ float PointLight_EstimateLight(
 	return OutIrradiance;
 }
 
-
-float2 IntersectInfiniteConeVolume(
-	float3  ro, float3  rd,
-	float3  apex, float3  normal,
-	float radius)
-{
-	// Adapted from: https://www.iquilezles.org/www/articles/intersectors/intersectors.htm
-	const float3 ba = normal;
-	const float3 oa = ro - apex;
-
-	const float m0 = dot(ba, ba);
-	const float m1 = dot(oa, ba);
-	const float m3 = dot(rd, ba);
-
-	const float m4 = dot(rd, oa);
-	const float m5 = dot(oa, oa);
-	const float hy = m0 + radius * radius;
-
-	const float k2 = m0 * m0 - m3 * m3 * hy;
-	const float k1 = m0 * m0 * m4 - m1 * m3 * hy;
-	const float k0 = m0 * m0 * m5 - m1 * m1 * hy;
-
-	const float h = k1 * k1 - k2 * k0;
-	if (h > 0)
-	{
-		const float s = sign(k2);
-		const float2 t = (-k1 + float2(-s, s) * sqrt(h)) / k2;
-		const float2 y = m1 + t * m3;
-
-		// check which hits are valid:
-		const int c = ((y.x > 0) ? 1 : 0) + (y.y > 0 ? 2 : 0);
-		if (c == 1)
-		{
-			// only smaller root is valid, far root is invalid
-			return float2(0.0, t.x);
-		}
-		if (c == 2)
-		{
-			// only larger root is valid, must have started outside the cone
-			return float2(t.y, POSITIVE_INFINITY);
-		}
-		if (c == 3)
-		{
-			// both roots are valid
-			return t;
-		}
-	}
-	return -1.0;
-}
-
-
 FVolumeLightSampleSetup PointLight_PrepareLightVolumeSample(
 	int LightId,
 	float3 RayOrigin,
@@ -225,27 +174,20 @@ FVolumeLightSampleSetup PointLight_PrepareLightVolumeSample(
 {
 	float3 Center = GetTranslatedPosition(LightId);
 	float AttenuationRadius = rcp(GetAttenuation(LightId));
-	float AttenuationRadius2 = Pow2(AttenuationRadius);
-	float3 oc = RayOrigin - Center;
-	float b = dot(oc, RayDirection);
-	float h = AttenuationRadius2 - length2(oc - b * RayDirection);
-	if (h > 0)
+	float2 T = RaySphereOverlap(RayOrigin, RayDirection, Center, AttenuationRadius, TMin, TMax);
+	if (T.x < T.y)
 	{
-		float2 t = -b + float2(-1, +1) * sqrt(h);
-		TMin = max(t.x, TMin);
-		TMax = min(t.y, TMax);
 		if (IsSpotLight(LightId))
 		{
 			float3 N = GetNormal(LightId);
 			// now clip against the cone of the spot
 			const float CosOuterConeAngle = GetCosConeAngles(LightId).x;
-			const float SinOuterConeAngle = sqrt(1 - CosOuterConeAngle * CosOuterConeAngle);
 
-			float2 TCone = IntersectInfiniteConeVolume(RayOrigin, RayDirection, Center, N * CosOuterConeAngle, SinOuterConeAngle);
+			float2 TCone = RayInfiniteConeVolumeOverlap(RayOrigin, RayDirection, Center, N, CosOuterConeAngle);
 			if (TCone.x < TCone.y)
 			{
-				TMin = max(TCone.x, TMin);
-				TMax = min(TCone.y, TMax);
+				T.x = max(TCone.x, T.x);
+				T.y = min(TCone.y, T.y);
 			}
 			else
 			{
@@ -258,10 +200,10 @@ FVolumeLightSampleSetup PointLight_PrepareLightVolumeSample(
 		if ((SceneLights[LightId].Flags & PATHTRACER_FLAG_NON_INVERSE_SQUARE_FALLOFF_MASK) == 0)
 		{
 			// inverse square falloff requires equi-angular sampling
-			return CreateEquiAngularSampler(LightImportance, Center, RayOrigin, RayDirection, TMin, TMax);
+			return CreateEquiAngularSampler(LightImportance, Center, RayOrigin, RayDirection, T.x, T.y);
 		}
 		// otherwise, default to uniform sampling
-		return CreateUniformSampler(LightImportance, TMin, TMax);
+		return CreateUniformSampler(LightImportance, T.x, T.y);
 	}
 	return NullVolumeLightSetup();
 }
diff --git a/Engine/Shaders/Private/PathTracing/Light/PathTracingRectLight.ush b/Engine/Shaders/Private/PathTracing/Light/PathTracingRectLight.ush
--- a/Engine/Shaders/Private/PathTracing/Light/PathTracingRectLight.ush
+++ b/Engine/Shaders/Private/PathTracing/Light/PathTracingRectLight.ush
@@ -221,20 +221,13 @@ FVolumeLightSampleSetup RectLight_PrepareLightVolumeSample(
 {
 	float3 Center = GetTranslatedPosition(LightId);
 	float AttenuationRadius = rcp(GetAttenuation(LightId));
-	float AttenuationRadius2 = Pow2(AttenuationRadius);
-	float3 oc = RayOrigin - Center;
-	float b = dot(oc, RayDirection);
-	float h = AttenuationRadius2 - length2(oc - b * RayDirection);
-	if (h > 0)
+	float2 T = RaySphereOverlap(RayOrigin, RayDirection, Center, AttenuationRadius, TMin, TMax);
+	if (T.x < T.y)
 	{
-		float2 t = -b + float2(-1, +1) * sqrt(h);
-		TMin = max(t.x, TMin);
-		TMax = min(t.y, TMax);
-
 		// now clip the sphere of influence against the plane of the quad and adjust the near/far distance depending on the sign
 		float3 N = GetNormal(LightId);
 
-		ClipRayByPlane(RayOrigin, RayDirection, N, Center, TMin, TMax);
+		ClipRayByPlane(RayOrigin, RayDirection, N, Center, T.x, T.y);
 
 		const float3 C = GetColor(LightId);
 		const float Width = GetWidth(LightId);
@@ -252,10 +245,10 @@ FVolumeLightSampleSetup RectLight_PrepareLightVolumeSample(
 			const float3 LightdPdu = GetdPdu(LightId);
 			const float3 LightdPdv = GetdPdv(LightId);
 
-			ClipRayByPlane(RayOrigin, RayDirection, BoundingPlaneX.x * N + BoundingPlaneX.y * LightdPdu, Center - LightdPdu * Width * 0.5, TMin, TMax);
-			ClipRayByPlane(RayOrigin, RayDirection, BoundingPlaneX.x * N - BoundingPlaneX.y * LightdPdu, Center + LightdPdu * Width * 0.5, TMin, TMax);
-			ClipRayByPlane(RayOrigin, RayDirection, BoundingPlaneY.x * N + BoundingPlaneY.y * LightdPdv, Center - LightdPdv * Height * 0.5, TMin, TMax);
-			ClipRayByPlane(RayOrigin, RayDirection, BoundingPlaneY.x * N - BoundingPlaneY.y * LightdPdv, Center + LightdPdv * Height * 0.5, TMin, TMax);
+			ClipRayByPlane(RayOrigin, RayDirection, BoundingPlaneX.x * N + BoundingPlaneX.y * LightdPdu, Center - LightdPdu * Width * 0.5, T.x, T.y);
+			ClipRayByPlane(RayOrigin, RayDirection, BoundingPlaneX.x * N - BoundingPlaneX.y * LightdPdu, Center + LightdPdu * Width * 0.5, T.x, T.y);
+			ClipRayByPlane(RayOrigin, RayDirection, BoundingPlaneY.x * N + BoundingPlaneY.y * LightdPdv, Center - LightdPdv * Height * 0.5, T.x, T.y);
+			ClipRayByPlane(RayOrigin, RayDirection, BoundingPlaneY.x * N - BoundingPlaneY.y * LightdPdv, Center + LightdPdv * Height * 0.5, T.x, T.y);
 		}
 #endif
 
@@ -264,10 +257,10 @@ FVolumeLightSampleSetup RectLight_PrepareLightVolumeSample(
 		if ((SceneLights[LightId].Flags & PATHTRACER_FLAG_NON_INVERSE_SQUARE_FALLOFF_MASK) == 0)
 		{
 			// inverse square falloff requires equi-angular sampling
-			return CreateEquiAngularSampler(LightImportance, Center, RayOrigin, RayDirection, TMin, TMax);
+			return CreateEquiAngularSampler(LightImportance, Center, RayOrigin, RayDirection, T.x, T.y);
 		}
 		// otherwise, default to uniform sampling
-		return CreateUniformSampler(LightImportance, TMin, TMax);
+		return CreateUniformSampler(LightImportance, T.x, T.y);
 	}
 	return NullVolumeLightSetup();
 }
diff --git a/Engine/Shaders/Private/PathTracing/PathTracingBuildLightGrid.usf b/Engine/Shaders/Private/PathTracing/PathTracingBuildLightGrid.usf
--- a/Engine/Shaders/Private/PathTracing/PathTracingBuildLightGrid.usf
+++ b/Engine/Shaders/Private/PathTracing/PathTracingBuildLightGrid.usf
@@ -4,6 +4,7 @@
 #include "/Engine/Shared/RayTracingDefinitions.h"
 #include "/Engine/Shared/RayTracingTypes.h"
 #include "/Engine/Shared/PathTracingDefinitions.h"
+#include "../IntersectionUtils.ush"
 
 #define USE_ATTENUATION_TERM 0
 #define USE_IES_TERM 0
@@ -16,8 +17,14 @@ RWTexture2DArray<uint> RWLightGrid;
 RWBuffer<uint> RWLightGridData;
 uint LightGridResolution;
 uint LightGridMaxCount;
-int LightGridAxis;
+uint LightGridAxis;
 
+#define RECT_TEST 2 // 0: half-space only
+                    // 1: barndoors (planes only)
+					// 2: barndoors+aabb culling
+
+#define CONE_TEST 1 // 0: test against cone AABB
+					// 1: accurate cone/aabb test
 
 [numthreads(THREADGROUPSIZE_X, THREADGROUPSIZE_Y, 1)]
 void PathTracingBuildLightGridCS(uint3 DispatchThreadId : SV_DispatchThreadID)
@@ -29,7 +36,7 @@ void PathTracingBuildLightGridCS(uint3 DispatchThreadId : SV_DispatchThreadID)
 	// figure out dimension of the 3d grid and the current ID to be filled
 	uint3 VoxelId = 0, VoxelRes = 1;
 	int Axis = DispatchThreadId.z;
-	if (LightGridAxis != Axis && LightGridAxis >= 0 && LightGridAxis < 3)
+	if (LightGridAxis != Axis && LightGridAxis < 3)
 	{
 		RWLightGrid[DispatchThreadId] = 0;
 		return;
@@ -55,40 +62,47 @@ void PathTracingBuildLightGridCS(uint3 DispatchThreadId : SV_DispatchThreadID)
 		float Radius = 1.0 / GetAttenuation(LightIndex);
 		float3 Normal = GetNormal(LightIndex);
 
-		float3 E = max(Lo - Center, 0) + max(Center - Hi, 0);
-		float D2 = dot(E, E);
+		if (!AABBOverlapsSphere(Lo, Hi, Center, Radius))
+		{
+			// bounds don't overlap there sphere of influence
+			continue;
+		}
 
 		// reject against normal
 		if (LightType != PATHTRACING_LIGHT_POINT)
 		{
-			if (dot(select(Normal > 0, Hi, Lo) - Center, Normal) < 0)
+			if (IsBoxBehindPlane(Lo, Hi, Normal, Center))
 			{
 				continue;
 			}
 		}
+#if RECT_TEST > 0
+		if (LightType == PATHTRACING_LIGHT_RECT)
+		{
+			// check barndoors
+			const float BarnCos = GetRectLightBarnCosAngle(LightIndex);
+			if (BarnCos > 0.035)
+			{
+				const float BarnLen = GetRectLightBarnLength(LightIndex);
+				const float HalfWidth = GetWidth(LightIndex) * 0.5;
+				const float HalfHeight = GetHeight(LightIndex) * 0.5;
 
+				if (!AABBOverlapsRectLightFrustum(Lo, Hi, Center, Normal, GetdPdv(LightIndex), HalfWidth, HalfHeight, Radius, BarnCos, BarnLen, RECT_TEST == 2))
+				{
+					continue;
+				}
+			}
+		}
+#endif
 		if (LightType == PATHTRACING_LIGHT_SPOT)
 		{
-			float3 Disc = sqrt(saturate(1.0 - Normal * Normal));
-			// box around ray from light center to tip of the cone
-			float3 Tip = Center + Normal * Radius;
-			float3 LightBoundLo = min(Center, Tip);
-			float3 LightBoundHi = max(Center, Tip);
-			// expand by disc around the farthest part of the cone
-
-			float CosOuter = GetCosConeAngles(LightIndex).x;
-			float SinOuter = sqrt(1.0f - CosOuter * CosOuter);
-
-			LightBoundLo = min(LightBoundLo, Center + Radius * (Normal * CosOuter - Disc * SinOuter));
-			LightBoundHi = max(LightBoundHi, Center + Radius * (Normal * CosOuter + Disc * SinOuter));
-
-			if (any(Hi < LightBoundLo) || any(Lo > LightBoundHi))
+			const float CosOuter = GetCosConeAngles(LightIndex).x;
+			if (!AABBOverlapsCurvedCone(Lo, Hi, Center, Normal, CosOuter, Radius, CONE_TEST == 0))
 			{
 				continue;
 			}
 		}
 
-		if (D2 < Square(Radius))
 		{
 			AllSingular = AllSingular && all(SceneLights[LightIndex].Dimensions == 0);
 			RWLightGridData[Offset + NumVisible] = LightIndex;
diff --git a/Engine/Shaders/Private/RayTracing/RayTracingBuildLightGrid.usf b/Engine/Shaders/Private/RayTracing/RayTracingBuildLightGrid.usf
--- a/Engine/Shaders/Private/RayTracing/RayTracingBuildLightGrid.usf
+++ b/Engine/Shaders/Private/RayTracing/RayTracingBuildLightGrid.usf
@@ -4,6 +4,7 @@
 #include "/Engine/Shared/LightDefinitions.h"
 #include "/Engine/Shared/RayTracingDefinitions.h"
 #include "/Engine/Shared/RayTracingTypes.h"
+#include "../IntersectionUtils.ush"
 
 uint SceneLightCount;
 StructuredBuffer<FRTLightingData> SceneLights;
@@ -14,7 +15,16 @@ RWTexture2DArray<uint> RWLightGrid;
 RWBuffer<uint> RWLightGridData;
 uint LightGridResolution;
 uint LightGridMaxCount;
-int LightGridAxis;
+uint LightGridAxis;
+
+
+#define RECT_TEST 2 // 0: half-space only
+                    // 1: barndoors (planes only)
+					// 2: barndoors+aabb culling
+
+#define CONE_TEST 1 // 0: test against cone AABB
+					// 1: accurate cone/aabb test
+
 
 [numthreads(THREADGROUPSIZE_X, THREADGROUPSIZE_Y, 1)]
 void RayTracingBuildLightGridCS(uint3 DispatchThreadId : SV_DispatchThreadID)
@@ -26,7 +36,7 @@ void RayTracingBuildLightGridCS(uint3 DispatchThreadId : SV_DispatchThreadID)
 	// figure out dimension of the 3d grid and the current ID to be filled
 	uint3 VoxelId = 0, VoxelRes = 1;
 	int Axis = DispatchThreadId.z;
-	if (Axis != LightGridAxis && LightGridAxis >= 0 && LightGridAxis < 3)
+	if (Axis != LightGridAxis && LightGridAxis < 3)
 	{
 		RWLightGrid[DispatchThreadId] = 0;
 		return;
@@ -51,40 +61,48 @@ void RayTracingBuildLightGridCS(uint3 DispatchThreadId : SV_DispatchThreadID)
 		float Radius = 1.0 / SceneLights[LightIndex].InvRadius;
 		float3 Normal = -SceneLights[LightIndex].Direction;
 
-		float3 E = max(Lo - Center, 0) + max(Center - Hi, 0);
-		float D2 = dot(E, E);
+		if (!AABBOverlapsSphere(Lo, Hi, Center, Radius))
+		{
+			// bounds don't overlap there sphere of influence
+			continue;
+		}
 
 		// reject against normal
 		if (LightType != LIGHT_TYPE_POINT)
 		{
-			if (dot(select(Normal > 0, Hi, Lo) - Center, Normal) < 0)
+			if (IsBoxBehindPlane(Lo, Hi, Normal, Center))
 			{
 				continue;
 			}
 		}
-
-		if (LightType == LIGHT_TYPE_SPOT)
+#if RECT_TEST > 0
+		if (LightType == LIGHT_TYPE_RECT)
 		{
-			float3 Disc = sqrt(saturate(1.0 - Normal * Normal));
-			// box around ray from light center to tip of the cone
-			float3 Tip = Center + Normal * Radius;
-			float3 LightBoundLo = min(Center, Tip);
-			float3 LightBoundHi = max(Center, Tip);
-			// expand by disc around the farthest part of the cone
-
-			float CosOuter = SceneLights[LightIndex].SpotAngles.x;
-			float SinOuter = sqrt(1.0f - CosOuter * CosOuter);
+			// check barndoors
+			const float BarnCos = SceneLights[LightIndex].RectLightBarnCosAngle;
+			if (BarnCos > 0.035)
+			{
+				const float BarnLen = SceneLights[LightIndex].RectLightBarnLength;
+				const float HalfWidth = SceneLights[LightIndex].SourceRadius;
+				const float HalfHeight = SceneLights[LightIndex].SourceLength;
 
-			LightBoundLo = min(LightBoundLo, Center + Radius * (Normal * CosOuter - Disc * SinOuter));
-			LightBoundHi = max(LightBoundHi, Center + Radius * (Normal * CosOuter + Disc * SinOuter));
 
-			if (any(Hi < LightBoundLo) || any(Lo > LightBoundHi))
+				if (!AABBOverlapsRectLightFrustum(Lo, Hi, Center, Normal, SceneLights[LightIndex].Tangent, HalfWidth, HalfHeight, Radius, BarnCos, BarnLen, RECT_TEST == 2))
+				{
+					continue;
+				}
+			}
+		}
+#endif
+		if (LightType == LIGHT_TYPE_SPOT)
+		{
+			const float CosOuter = SceneLights[LightIndex].SpotAngles.x;
+			if (!AABBOverlapsCurvedCone(Lo, Hi, Center, Normal, CosOuter, Radius, CONE_TEST == 0))
 			{
 				continue;
 			}
 		}
 
-		if (D2 < Square(Radius))
 		{
 			RWLightGridData[Offset + NumVisible] = LightIndex;
 			NumVisible++;
diff --git a/Engine/Shaders/Private/RayTracing/RayTracingLightCullingCommon.ush b/Engine/Shaders/Private/RayTracing/RayTracingLightCullingCommon.ush
--- a/Engine/Shaders/Private/RayTracing/RayTracingLightCullingCommon.ush
+++ b/Engine/Shaders/Private/RayTracing/RayTracingLightCullingCommon.ush
@@ -38,12 +38,12 @@ struct FCulledLightList
 			uint LightGridNz = RaytracingLightGridData.LightGrid.Load(int4(UVz, 2, 0));
 			uint LightGridN = LightGridNx;
 			uint LightGridOffset = RaytracingLightGridData.LightGridMaxCount * (0 + 3 * (UVx.x + UVx.y * RaytracingLightGridData.LightGridResolution));
-			if ((RaytracingLightGridData.LightGridAxis < 0 && LightGridNy < LightGridN) || RaytracingLightGridData.LightGridAxis == 1)
+			if ((RaytracingLightGridData.LightGridAxis >= 3 && LightGridNy < LightGridN) || RaytracingLightGridData.LightGridAxis == 1)
 			{
 				LightGridN = LightGridNy;
 				LightGridOffset = RaytracingLightGridData.LightGridMaxCount * (1 + 3 * (UVy.x + UVy.y * RaytracingLightGridData.LightGridResolution));
 			}
-			if ((RaytracingLightGridData.LightGridAxis < 0 && LightGridNz < LightGridN) || RaytracingLightGridData.LightGridAxis == 2)
+			if ((RaytracingLightGridData.LightGridAxis >= 3 && LightGridNz < LightGridN) || RaytracingLightGridData.LightGridAxis == 2)
 			{
 				LightGridN = LightGridNz;
 				LightGridOffset = RaytracingLightGridData.LightGridMaxCount * (2 + 3 * (UVz.x + UVz.y * RaytracingLightGridData.LightGridResolution));
diff --git a/Engine/Source/Runtime/Renderer/Private/PathTracing.cpp b/Engine/Source/Runtime/Renderer/Private/PathTracing.cpp
--- a/Engine/Source/Runtime/Renderer/Private/PathTracing.cpp
+++ b/Engine/Source/Runtime/Renderer/Private/PathTracing.cpp
@@ -2521,59 +2521,37 @@ RENDERER_API void PrepareLightGrid(FRDGBuilder& GraphBuilder, ERHIFeatureLevel::
 		for (int Index = 0; Index < NumFiniteLights; Index++)
 		{
 			const FPathTracingLight& Light = FiniteLights[Index];
-			FVector3f Lo = FVector3f(-Inf, -Inf, -Inf);
-			FVector3f Hi = FVector3f( Inf,  Inf,  Inf);
+			FBox3f Box;
 
 			const float Radius = 1.0f / Light.Attenuation;
 			const FVector3f Center = Light.TranslatedWorldPosition;
 			const FVector3f Normal = Light.Normal;
-			const FVector3f Disc = FVector3f(
-				FMath::Sqrt(FMath::Clamp(1 - Normal.X * Normal.X, 0.0f, 1.0f)),
-				FMath::Sqrt(FMath::Clamp(1 - Normal.Y * Normal.Y, 0.0f, 1.0f)),
-				FMath::Sqrt(FMath::Clamp(1 - Normal.Z * Normal.Z, 0.0f, 1.0f))
-			);
 			switch (Light.Flags & PATHTRACER_FLAG_TYPE_MASK)
 			{
 				case PATHTRACING_LIGHT_POINT:
 				{
-					// simple sphere of influence
-					Lo = Center - FVector3f(Radius, Radius, Radius);
-					Hi = Center + FVector3f(Radius, Radius, Radius);
+					Box = GetPointLightBounds(Center, Radius);
 					break;
 				}
 				case PATHTRACING_LIGHT_SPOT:
 				{
-					// box around ray from light center to tip of the cone
-					const FVector3f Tip = Center + Normal * Radius;
-					Lo = FVector3f::Min(Center, Tip);
-					Hi = FVector3f::Max(Center, Tip);
-
-					// expand by disc around the farthest part of the cone
-					const float CosOuter = Light.Shaping.X;
-					const float SinOuter = FMath::Sqrt(1.0f - CosOuter * CosOuter);
-
-					Lo = FVector3f::Min(Lo, Center + Radius * (Normal * CosOuter - Disc * SinOuter));
-					Hi = FVector3f::Max(Hi, Center + Radius * (Normal * CosOuter + Disc * SinOuter));
+					Box = GetSpotLightBounds(Center, Normal, Radius, Light.Shaping.X);
 					break;
 				}
 				case PATHTRACING_LIGHT_RECT:
 				{
-					// quad bbox is the bbox of the disc +  the tip of the hemisphere
-					// TODO: is it worth trying to account for barndoors? seems unlikely to cut much empty space since the volume _inside_ the barndoor receives light
-					const FVector3f Tip = Center + Normal * Radius;
-					Lo = FVector3f::Min(Tip, Center - Radius * Disc);
-					Hi = FVector3f::Max(Tip, Center + Radius * Disc);
+					Box = GetRectLightBounds(Center, Normal, Light.Tangent, Light.Dimensions.X * 0.5f, Light.Dimensions.Y * 0.5f, Radius, Light.Shaping.X, Light.Shaping.Y);
 					break;
 				}
 				default:
 				{
 					// non-finite lights should not appear in this case
-					checkSlow(false);
+					checkNoEntry();
 					break;
 				}
 			}
-			LightGridParameters->SceneLightsTranslatedBoundMin = FVector3f::Min(LightGridParameters->SceneLightsTranslatedBoundMin, Lo);
-			LightGridParameters->SceneLightsTranslatedBoundMax = FVector3f::Max(LightGridParameters->SceneLightsTranslatedBoundMax, Hi);
+			LightGridParameters->SceneLightsTranslatedBoundMin = FVector3f::Min(LightGridParameters->SceneLightsTranslatedBoundMin, Box.Min);
+			LightGridParameters->SceneLightsTranslatedBoundMax = FVector3f::Max(LightGridParameters->SceneLightsTranslatedBoundMax, Box.Max);
 		}
 
 		const uint32 Resolution = FMath::Clamp(CVarPathTracingLightGridResolution.GetValueOnRenderThread(), 1, 2048);
diff --git a/Engine/Source/Runtime/Renderer/Private/RayTracing/RayTracingLighting.cpp b/Engine/Source/Runtime/Renderer/Private/RayTracing/RayTracingLighting.cpp
--- a/Engine/Source/Runtime/Renderer/Private/RayTracing/RayTracingLighting.cpp
+++ b/Engine/Source/Runtime/Renderer/Private/RayTracing/RayTracingLighting.cpp
@@ -68,13 +68,132 @@ class FRayTracingBuildLightGridCS : public FGlobalShader
 		SHADER_PARAMETER_RDG_BUFFER_SRV(StructuredBuffer<FRTLightingData>, SceneLights)
 		SHADER_PARAMETER(unsigned, LightGridResolution)
 		SHADER_PARAMETER(unsigned, LightGridMaxCount)
-		SHADER_PARAMETER(int, LightGridAxis)
+		SHADER_PARAMETER(unsigned, LightGridAxis)
 		SHADER_PARAMETER_RDG_TEXTURE_UAV(RWTexture2DArray<uint>, RWLightGrid)
 		SHADER_PARAMETER_RDG_BUFFER_UAV(RWBuffer<uint>, RWLightGridData)
 	END_SHADER_PARAMETER_STRUCT()
 };
 IMPLEMENT_SHADER_TYPE(, FRayTracingBuildLightGridCS, TEXT("/Engine/Private/RayTracing/RayTracingBuildLightGrid.usf"), TEXT("RayTracingBuildLightGridCS"), SF_Compute);
 
+
+FBox3f GetPointLightBounds(const FVector3f& Center, float Radius)
+{
+	const FVector3f R(Radius, Radius, Radius);
+	return FBox3f(Center - R, Center + R);
+}
+
+FBox3f GetSpotLightBounds(const FVector3f& Center, const FVector3f& Normal, float Radius, float CosOuter)
+{
+	// box around ray from light center to tip of the cone
+	const FVector3f Tip = Center + Normal * Radius;
+	FVector3f Lo = FVector3f::Min(Center, Tip);
+	FVector3f Hi = FVector3f::Max(Center, Tip);
+
+	const float SinOuter = FMath::Sqrt(1.0f - CosOuter * CosOuter);
+
+	// expand by disc around the farthest part of the cone
+	const FVector3f Disc = FVector3f(
+		FMath::Sqrt(FMath::Clamp(1.0f - Normal.X * Normal.X, 0.0f, 1.0f)),
+		FMath::Sqrt(FMath::Clamp(1.0f - Normal.Y * Normal.Y, 0.0f, 1.0f)),
+		FMath::Sqrt(FMath::Clamp(1.0f - Normal.Z * Normal.Z, 0.0f, 1.0f))
+	);
+	Lo = FVector3f::Min(Lo, Center + Radius * (Normal * CosOuter - Disc * SinOuter));
+	Hi = FVector3f::Max(Hi, Center + Radius * (Normal * CosOuter + Disc * SinOuter));
+
+	// Check if any of the coordinate axis points lie inside the cone and include them if they do
+	// This is the only case which is not captured by the AABB above
+	const FVector3f E = FVector3f(
+		FMath::Abs(Normal.X) > CosOuter ? Center.X + copysignf(Radius, Normal.X) : Center.X,
+		FMath::Abs(Normal.Y) > CosOuter ? Center.Y + copysignf(Radius, Normal.Y) : Center.Y,
+		FMath::Abs(Normal.Z) > CosOuter ? Center.Z + copysignf(Radius, Normal.Z) : Center.Z
+	);
+	Lo = FVector3f::Min(Lo, E);
+	Hi = FVector3f::Max(Hi, E);
+	return FBox3f(Lo, Hi);
+}
+
+FBox3f GetRectLightBounds(const FVector3f& Center, const FVector3f& Normal, const FVector3f& Tangent, float HalfWidth, float HalfHeight, float Radius, float BarnCos, float BarnLen)
+{
+	const FVector3f Corner = FVector3f(
+		copysignf(Radius, Normal.X),
+		copysignf(Radius, Normal.Y),
+		copysignf(Radius, Normal.Z)
+	);
+	const FVector3f Disc = FVector3f(
+		FMath::Sqrt(FMath::Clamp(1.0f - Normal.X * Normal.X, 0.0f, 1.0f)),
+		FMath::Sqrt(FMath::Clamp(1.0f - Normal.Y * Normal.Y, 0.0f, 1.0f)),
+		FMath::Sqrt(FMath::Clamp(1.0f - Normal.Z * Normal.Z, 0.0f, 1.0f))
+	);
+
+	// rect bbox is the bbox of the disc + furthest corner of the Radius sized box in the direction of the normal
+	FVector3f Lo = FVector3f::Min(Center + Corner, Center - Radius * Disc);
+	FVector3f Hi = FVector3f::Max(Center + Corner, Center + Radius * Disc);
+
+	// Take into account barndoor frustum if enabled
+	if (BarnCos > 0.035f)
+	{
+		const FVector3f dPdv = Tangent;
+		const FVector3f dPdu = Normal.Cross(Tangent);
+		const float BarnSin = FMath::Sqrt(1 - BarnCos * BarnCos);
+
+		FVector3f BoundingPlane = FVector3f(
+			2 * HalfWidth  + BarnLen * BarnSin,
+			2 * HalfHeight + BarnLen * BarnSin,
+			BarnLen * BarnCos
+		);
+		FVector3f BLo = Center, BHi = Center;
+		// loop through 9 points to get extremes of the "rounded" pyramid defined by the barndoor penumbra + radius
+		for (int Dy = -1; Dy <= 1; Dy++)
+		for (int Dx = -1; Dx <= 1; Dx++)
+		{
+			// Get point on rectangle
+			const FVector3f Rxy = Center + Dx * HalfWidth * dPdu + Dy * HalfHeight * dPdv;
+			// Get penumbra plane vector, normalize it and scale to edge of the sphere (roughly, since we aren't starting from the center)
+			const FVector3f Bxy = Rxy + (Dx * BoundingPlane.X * dPdu + Dy * BoundingPlane.Y * dPdv + BoundingPlane.Z * Normal).GetUnsafeNormal() * Radius;
+			BLo = FVector3f::Min(BLo, Rxy); BHi = FVector3f::Max(BHi, Rxy);
+			BLo = FVector3f::Min(BLo, Bxy); BHi = FVector3f::Max(BHi, Bxy);
+		}
+
+		// Include "axis" points if they lie inside the barndoor penumbra (similar to spot light test above, but the apex is behind the Center and different in X and Y due to the rectangle size)
+		// A 2D visualization of this is here: https://www.desmos.com/calculator/15zh9boeqz
+		const float TanAlphaX = BoundingPlane.X / BoundingPlane.Z, CosAlphaX = FMath::InvSqrt(1 + TanAlphaX * TanAlphaX), ApexX = HalfWidth  / TanAlphaX;
+		const float TanAlphaY = BoundingPlane.Y / BoundingPlane.Z, CosAlphaY = FMath::InvSqrt(1 + TanAlphaY * TanAlphaY), ApexY = HalfHeight / TanAlphaY;
+
+		// Take the extreme point along each axis, then rotate it to local space
+		const FVector3f Px = FVector3f(dPdu.X * Corner.X, dPdv.X * Corner.X, Normal.X * Corner.X);
+		const FVector3f Py = FVector3f(dPdu.Y * Corner.Y, dPdv.Y * Corner.Y, Normal.Y * Corner.Y);
+		const FVector3f Pz = FVector3f(dPdu.Z * Corner.Z, dPdv.Z * Corner.Z, Normal.Z * Corner.Z);
+
+		// Now - check if dot product betwen this point and normal lies within the cone
+		// We do two 2D cone tests as we have a different cone in X and Y
+		// This is the analog of the simpler implementation possible for spotlights above
+		if ((Px.Z - ApexX) > FVector2f(Px.X, Px.Z - ApexX).Length() * CosAlphaX &&
+			(Px.Z - ApexY) > FVector2f(Px.Y, Px.Z - ApexY).Length() * CosAlphaY)
+		{
+			BLo.X = FMath::Min(BLo.X, Center.X + Corner.X);
+			BHi.X = FMath::Max(BHi.X, Center.X + Corner.X);
+		}
+		if ((Py.Z - ApexX) > FVector2f(Py.X, Py.Z - ApexX).Length() * CosAlphaX &&
+			(Py.Z - ApexY) > FVector2f(Py.Y, Py.Z - ApexY).Length() * CosAlphaY)
+		{
+			BLo.Y = FMath::Min(BLo.Y, Center.Y + Corner.Y);
+			BHi.Y = FMath::Max(BHi.Y, Center.Y + Corner.Y);
+		}
+		if ((Pz.Z - ApexX) > FVector2f(Pz.X, Pz.Z - ApexX).Length() * CosAlphaX &&
+			(Pz.Z - ApexY) > FVector2f(Pz.Y, Pz.Z - ApexY).Length() * CosAlphaY)
+		{
+			BLo.Z = FMath::Min(BLo.Z, Center.Z + Corner.Z);
+			BHi.Z = FMath::Max(BHi.Z, Center.Z + Corner.Z);
+		}
+		// Now clip the new BBox against the old (conservative one) for the entire half-space
+		Lo = FVector3f::Max(Lo, BLo);
+		Hi = FVector3f::Min(Hi, BHi);
+	}
+	return FBox3f(Lo, Hi);
+}
+
+
+
 static void PrepareLightGrid(FRDGBuilder& GraphBuilder, FGlobalShaderMap* ShaderMap, FRayTracingLightGrid* LightGridParameters, const FRTLightingData* Lights)
 {
 	// number of lights and infinite lights is provided by caller
@@ -108,59 +227,39 @@ static void PrepareLightGrid(FRDGBuilder& GraphBuilder, FGlobalShaderMap* Shader
 	for (int Index = 0; Index < NumFiniteLights; Index++)
 	{
 		const FRTLightingData& Light = FiniteLights[Index];
-		FVector3f Lo = FVector3f(-Inf, -Inf, -Inf);
-		FVector3f Hi = FVector3f( Inf,  Inf,  Inf);
+		FBox3f Box;
 
-		const float Radius = 1.0f / Light.InvRadius;
 		const FVector3f Center = Light.TranslatedLightPosition;
 		const FVector3f Normal = -Light.Direction;
-		const FVector3f Disc = FVector3f(
-			FMath::Sqrt(FMath::Clamp(1 - Normal.X * Normal.X, 0.0f, 1.0f)),
-			FMath::Sqrt(FMath::Clamp(1 - Normal.Y * Normal.Y, 0.0f, 1.0f)),
-			FMath::Sqrt(FMath::Clamp(1 - Normal.Z * Normal.Z, 0.0f, 1.0f))
-		);
+		const float Radius = 1.0f / Light.InvRadius;
 		switch (Light.Type)
 		{
 			case LightType_Point:
 			{
-				// simple sphere of influence
-				Lo = Center - FVector3f(Radius, Radius, Radius);
-				Hi = Center + FVector3f(Radius, Radius, Radius);
+				Box = GetPointLightBounds(Center, Radius);
 				break;
 			}
 			case LightType_Spot:
 			{
-				// box around ray from light center to tip of the cone
-				const FVector3f Tip = Center + Normal * Radius;
-				Lo = FVector3f::Min(Center, Tip);
-				Hi = FVector3f::Max(Center, Tip);
-
-				// expand by disc around the farthest part of the cone
-				const float CosOuter = Light.SpotAngles.X;
-				const float SinOuter = FMath::Sqrt(1.0f - CosOuter * CosOuter);
-
-				Lo = FVector3f::Min(Lo, Center + Radius * (Normal * CosOuter - Disc * SinOuter));
-				Hi = FVector3f::Max(Hi, Center + Radius * (Normal * CosOuter + Disc * SinOuter));
+				Box = GetSpotLightBounds(Center, Normal, Radius, Light.SpotAngles.X);
 				break;
 			}
 			case LightType_Rect:
 			{
-				// quad bbox is the bbox of the disc +  the tip of the hemisphere
-				// TODO: is it worth trying to account for barndoors? seems unlikely to cut much empty space since the volume _inside_ the barndoor receives light
-				const FVector3f Tip = Center + Normal * Radius;
-				Lo = FVector3f::Min(Tip, Center - Radius * Disc);
-				Hi = FVector3f::Max(Tip, Center + Radius * Disc);
+				const float HalfWidth = Light.SourceRadius;
+				const float HalfHeight = Light.SourceLength;
+				Box = GetRectLightBounds(Center, Normal, Light.Tangent, HalfWidth, HalfHeight, Radius, Light.RectLightBarnCosAngle, Light.RectLightBarnLength);
 				break;
 			}
 			default:
 			{
 				// non-finite lights should not appear in this case
-				checkSlow(false);
+				checkNoEntry();
 				break;
 			}
 		}
-		LightGridParameters->SceneLightsTranslatedBoundMin = FVector3f::Min(LightGridParameters->SceneLightsTranslatedBoundMin, Lo);
-		LightGridParameters->SceneLightsTranslatedBoundMax = FVector3f::Max(LightGridParameters->SceneLightsTranslatedBoundMax, Hi);
+		LightGridParameters->SceneLightsTranslatedBoundMin = FVector3f::Min(LightGridParameters->SceneLightsTranslatedBoundMin, Box.Min);
+		LightGridParameters->SceneLightsTranslatedBoundMax = FVector3f::Max(LightGridParameters->SceneLightsTranslatedBoundMax, Box.Max);
 	}
 
 	const uint32 Resolution = FMath::Clamp(CVarRayTracingLightGridResolution.GetValueOnRenderThread(), 1, 2048);
@@ -230,8 +329,12 @@ static void PrepareLightGrid(FRDGBuilder& GraphBuilder, FGlobalShaderMap* Shader
 	LightGridParameters->LightGridData = GraphBuilder.CreateSRV(LightGridData, LightGridDataFormat);
 }
 
-static bool ShouldIncludeRayTracingLight(const FLightSceneInfoCompact& Light)
+static bool ShouldIncludeRayTracingLight(const FLightSceneInfoCompact& Light, const FEngineShowFlags& ShowFlags)
 {
+	if (Light.LightType == LightType_Directional && !ShowFlags.DirectionalLights) return false;
+	if (Light.LightType == LightType_Spot        && !ShowFlags.SpotLights       ) return false;
+	if (Light.LightType == LightType_Point       && !ShowFlags.PointLights      ) return false;
+	if (Light.LightType == LightType_Rect        && !ShowFlags.RectLights       ) return false;
 	const bool bHasStaticLighting = Light.LightSceneInfo->Proxy->HasStaticLighting() && Light.LightSceneInfo->IsPrecomputedLightingValid();
 	const bool bAffectReflection = Light.LightSceneInfo->Proxy->AffectReflection();
 	return !bHasStaticLighting && bAffectReflection;
@@ -256,7 +359,7 @@ TRDGUniformBufferRef<FRayTracingLightGrid> CreateRayTracingLightData(
 		int NumLightsByType[LightType_MAX] = {};
 		for (const FLightSceneInfoCompact& Light : Lights)
 		{
-			if (!ShouldIncludeRayTracingLight(Light))
+			if (!ShouldIncludeRayTracingLight(Light, View.Family->EngineShowFlags))
 				continue;
 			check(Light.LightType < LightType_MAX);
 			NumLightsByType[Light.LightType]++;
@@ -277,7 +380,7 @@ TRDGUniformBufferRef<FRayTracingLightGrid> CreateRayTracingLightData(
 		const FRayTracingLightFunctionMap* RayTracingLightFunctionMap = GraphBuilder.Blackboard.Get<FRayTracingLightFunctionMap>();
 		for (const FLightSceneInfoCompact& Light : Lights)
 		{
-			if (!ShouldIncludeRayTracingLight(Light))
+			if (!ShouldIncludeRayTracingLight(Light, View.Family->EngineShowFlags))
 				continue;
 
 			FLightRenderParameters LightParameters;
diff --git a/Engine/Source/Runtime/Renderer/Private/RayTracing/RayTracingLighting.h b/Engine/Source/Runtime/Renderer/Private/RayTracing/RayTracingLighting.h
--- a/Engine/Source/Runtime/Renderer/Private/RayTracing/RayTracingLighting.h
+++ b/Engine/Source/Runtime/Renderer/Private/RayTracing/RayTracingLighting.h
@@ -26,7 +26,7 @@ BEGIN_GLOBAL_SHADER_PARAMETER_STRUCT(FRayTracingLightGrid, RENDERER_API)
 	SHADER_PARAMETER_RDG_BUFFER_SRV(Buffer<uint>, LightGridData)
 	SHADER_PARAMETER(unsigned, LightGridResolution)
 	SHADER_PARAMETER(unsigned, LightGridMaxCount)
-	SHADER_PARAMETER(int, LightGridAxis)
+	SHADER_PARAMETER(unsigned, LightGridAxis)
 END_GLOBAL_SHADER_PARAMETER_STRUCT()
 
 using FRayTracingLightFunctionMap = TMap<const FLightSceneInfo*, int32>;
@@ -56,4 +56,9 @@ void BindLightFunctionShadersPathTracing(
 	const FRayTracingLightFunctionMap* RayTracingLightFunctionMap,
 	const class FViewInfo& View);
 
+// Helper routines, useful when building a light grid
+FBox3f GetPointLightBounds(const FVector3f& Center, float Radius);
+FBox3f GetSpotLightBounds(const FVector3f& Center, const FVector3f& Normal, float Radius, float CosOuter);
+FBox3f GetRectLightBounds(const FVector3f& Center, const FVector3f& Normal, const FVector3f& Tangent, float HalfWidth, float HalfHeight, float Radius, float BarnCos, float BarnLen);
+
 #endif
