MGPU:  Fixes for QAVirtualProduction test map.  Change from push to pull cross GPU transfers.  Push transfers seem to have been broken across the board by a recent driver change, as all the way back to Release 5.1 (where push transfers were first implemented and tested) is now broken as well.  To make pull transfers relatively performant, it is necessary to gather them from each scene renderer, and do them all at once on the last scene renderer, otherwise the GPUs get serialized.  Other fix is to make global RHI resource initialization masked to run on all GPUs -- the uninitialized resources caused calls to AddDrawScreenPass to render nothing.

#rnx

diff --git a/Engine/Source/Runtime/RenderCore/Private/RenderResource.cpp b/Engine/Source/Runtime/RenderCore/Private/RenderResource.cpp
--- a/Engine/Source/Runtime/RenderCore/Private/RenderResource.cpp
+++ b/Engine/Source/Runtime/RenderCore/Private/RenderResource.cpp
@@ -123,6 +123,8 @@ void FRenderResource::ReleaseRHIForAllResources()
 void FRenderResource::InitPreRHIResources()
 {
 	FRHICommandListImmediate& RHICmdList = FRHICommandListImmediate::Get();
+	SCOPED_GPU_MASK(RHICmdList, FRHIGPUMask::All());
+
 	RHICmdList.InitializeImmediateContexts();
 
 	FRenderResourceList& PreResourceList = FRenderResourceList::Get<FRenderResource::EInitPhase::Pre>();
diff --git a/Engine/Source/Runtime/Renderer/Private/DeferredShadingRenderer.cpp b/Engine/Source/Runtime/Renderer/Private/DeferredShadingRenderer.cpp
--- a/Engine/Source/Runtime/Renderer/Private/DeferredShadingRenderer.cpp
+++ b/Engine/Source/Runtime/Renderer/Private/DeferredShadingRenderer.cpp
@@ -2127,7 +2127,7 @@ void FDeferredShadingSceneRenderer::Render(FRDGBuilder& GraphBuilder)
 				}
 
 			#if WITH_MGPU
-				DoCrossGPUTransfers(GraphBuilder, CustomRenderPass->GetRenderTargetTexture(), CustomRenderPassViews, false, FRHIGPUMask::All());
+				DoCrossGPUTransfers(GraphBuilder, CustomRenderPass->GetRenderTargetTexture(), CustomRenderPassViews, false, FRHIGPUMask::All(), nullptr);
 			#endif
 			}
 
@@ -3474,8 +3474,9 @@ void FDeferredShadingSceneRenderer::Render(FRDGBuilder& GraphBuilder)
 #if WITH_MGPU
 	if (ViewFamily.bMultiGPUForkAndJoin)
 	{
-		DoCrossGPUTransfers(GraphBuilder, ViewFamilyTexture, Views, CrossGPUTransferFencesDefer.Num() > 0, RenderTargetGPUMask);
+		DoCrossGPUTransfers(GraphBuilder, ViewFamilyTexture, Views, CrossGPUTransferFencesDefer.Num() > 0, RenderTargetGPUMask, CrossGPUTransferDeferred.GetReference());
 	}
+	FlushCrossGPUTransfers(GraphBuilder);
 #endif
 
 	{
diff --git a/Engine/Source/Runtime/Renderer/Private/SceneRendering.cpp b/Engine/Source/Runtime/Renderer/Private/SceneRendering.cpp
--- a/Engine/Source/Runtime/Renderer/Private/SceneRendering.cpp
+++ b/Engine/Source/Runtime/Renderer/Private/SceneRendering.cpp
@@ -506,6 +506,15 @@ static TAutoConsoleVariable<float> CVarTranslucencyAutoBeforeDOF(
 	TEXT("Automatically bin After DOF translucency before DOF if behind focus distance (Experimental)"),
 	ECVF_Default);
 
+static TAutoConsoleVariable<int32> CVarCrossGPUTransferOption(
+	TEXT("r.MultiGPU.Transfer"),
+	2,
+	TEXT("Mode to use for cross GPU transfers when multiple nDisplay views are active\n")
+	TEXT(" 0: immediate pull transfer\n")
+	TEXT(" 1: optimized push transfer (source GPU runs copy, with deferred fence wait on destination GPU)\n")
+	TEXT(" 2: optimized pull transfer (destination GPU runs copy, with transfers delayed to last view's render); (default)"),
+	ECVF_Default);
+
 
 FOcclusionSubmittedFenceState FSceneRenderer::OcclusionSubmittedFence[FOcclusionQueryHelpers::MaxBufferedOcclusionFrames];
 
@@ -3619,6 +3628,19 @@ struct FCrossGPUTransfer
 	}
 };
 
+struct FCrossGPUTarget
+{
+	const FRenderTarget* RenderTarget = nullptr;
+	FRDGTextureRef RenderTargetTexture = nullptr;
+	TArray<FCrossGPUTransfer> Transfers;
+};
+
+class FCrossGPUTransfersDeferred : public FRefCountBase
+{
+public:
+	TArray<FCrossGPUTarget> Targets;
+};
+
 static void GetCrossGPUTransfers(FSceneRenderer* SceneRenderer, TArray<FCrossGPUTransfer>& OutTransfers, TArrayView<FViewInfo> InViews, const FIntPoint RenderTargetSize, FRHIGPUMask RenderTargetGPUMask)
 {
 	check(SceneRenderer->bGPUMasksComputed);
@@ -3648,47 +3670,60 @@ static void GetCrossGPUTransfers(FSceneRenderer* SceneRenderer, TArray<FCrossGPU
 void FSceneRenderer::PreallocateCrossGPUFences(TConstArrayView<FSceneRenderer*> SceneRenderers)
 {
 #if WITH_MGPU
-	if (SceneRenderers.Num() > 1)
+	if (SceneRenderers.Num() > 1 && GNumExplicitGPUsForRendering > 1)
 	{
-		// Allocated fences to wait on are placed in the last scene renderer
-		TArray<FCrossGPUTransferFence*>& LastRendererFencesWait = SceneRenderers.Last()->CrossGPUTransferFencesWait;
+		int32 CrossGPUOption = CVarCrossGPUTransferOption.GetValueOnAnyThread();
+		if (CrossGPUOption == 1)
+		{
+			// Allocated fences to wait on are placed in the last scene renderer
+			TArray<FCrossGPUTransferFence*>& LastRendererFencesWait = SceneRenderers.Last()->CrossGPUTransferFencesWait;
 
-		check(LastRendererFencesWait.IsEmpty());
+			check(LastRendererFencesWait.IsEmpty());
 
-		// Each prior renderer allocates fences and also adds them to last renderer
-		for (int32 RendererIndex = 0; RendererIndex < SceneRenderers.Num() - 1; RendererIndex++)
-		{
-			FSceneRenderer* SceneRenderer = SceneRenderers[RendererIndex];
+			// Each prior renderer allocates fences and also adds them to last renderer
+			for (int32 RendererIndex = 0; RendererIndex < SceneRenderers.Num() - 1; RendererIndex++)
+			{
+				FSceneRenderer* SceneRenderer = SceneRenderers[RendererIndex];
 
-			check(SceneRenderer->CrossGPUTransferFencesDefer.IsEmpty());
+				check(SceneRenderer->CrossGPUTransferFencesDefer.IsEmpty());
 
-			SceneRenderer->ComputeGPUMasks(nullptr);
+				SceneRenderer->ComputeGPUMasks(nullptr);
 
-			if (SceneRenderer->ViewFamily.bMultiGPUForkAndJoin)
-			{
-				// Check if we can do optimized transfers, which requires a single index
-				if (SceneRenderer->AllViewsGPUMask.HasSingleIndex())
+				if (SceneRenderer->ViewFamily.bMultiGPUForkAndJoin)
 				{
-					TArray<FCrossGPUTransfer> Transfers;
-					GetCrossGPUTransfers(SceneRenderer, Transfers, SceneRenderer->Views, SceneRenderer->ViewFamily.RenderTarget->GetSizeXY(), SceneRenderer->RenderTargetGPUMask);
+					// Check if we can do optimized transfers, which requires a single index
+					if (SceneRenderer->AllViewsGPUMask.HasSingleIndex())
+					{
+						TArray<FCrossGPUTransfer> Transfers;
+						GetCrossGPUTransfers(SceneRenderer, Transfers, SceneRenderer->Views, SceneRenderer->ViewFamily.RenderTarget->GetSizeXY(), SceneRenderer->RenderTargetGPUMask);
 
-					SceneRenderer->CrossGPUTransferFencesDefer.SetNumUninitialized(Transfers.Num());
+						SceneRenderer->CrossGPUTransferFencesDefer.SetNumUninitialized(Transfers.Num());
 
-					for (int32 TransferIndex = 0; TransferIndex < Transfers.Num(); TransferIndex++)
-					{
-						FCrossGPUTransferFence* FenceData = RHICreateCrossGPUTransferFence();
+						for (int32 TransferIndex = 0; TransferIndex < Transfers.Num(); TransferIndex++)
+						{
+							FCrossGPUTransferFence* FenceData = RHICreateCrossGPUTransferFence();
 
-						SceneRenderer->CrossGPUTransferFencesDefer[TransferIndex] = FenceData;
-						LastRendererFencesWait.Add(FenceData);
+							SceneRenderer->CrossGPUTransferFencesDefer[TransferIndex] = FenceData;
+							LastRendererFencesWait.Add(FenceData);
+						}
 					}
 				}
 			}
 		}
+		else if (CrossGPUOption == 2)
+		{
+			TRefCountPtr<FCrossGPUTransfersDeferred> TransfersDeferred = new FCrossGPUTransfersDeferred;
+			for (FSceneRenderer* SceneRenderer : SceneRenderers)
+			{
+				// Each scene renderer will add transfers to the shared structure, then the last will emit the transfers
+				SceneRenderer->CrossGPUTransferDeferred = TransfersDeferred;
+			}
+		}
 	}
 #endif
 }
 
-void FSceneRenderer::DoCrossGPUTransfers(FRDGBuilder& GraphBuilder, FRDGTextureRef RenderTargetTexture, TArrayView<FViewInfo> InViews, bool bCrossGPUTransferFencesDefer, FRHIGPUMask InRenderTargetGPUMask)
+void FSceneRenderer::DoCrossGPUTransfers(FRDGBuilder& GraphBuilder, FRDGTextureRef RenderTargetTexture, TArrayView<FViewInfo> InViews, bool bCrossGPUTransferFencesDefer, FRHIGPUMask InRenderTargetGPUMask, class FCrossGPUTransfersDeferred* TransfersDeferred)
 {
 #if WITH_MGPU
 	// Must be all GPUs because context redirector only supports single or all GPUs
@@ -3701,11 +3736,19 @@ void FSceneRenderer::DoCrossGPUTransfers(FRDGBuilder& GraphBuilder, FRDGTextureR
 	TArray<FCrossGPUTransfer> Transfers;
 	GetCrossGPUTransfers(this, Transfers, InViews, RenderTargetTexture->Desc.Extent, InRenderTargetGPUMask);
 
-	if (Transfers.Num() > 0)
+	if (TransfersDeferred)
+	{
+		// Accumulate transfers from each scene renderer
+		if (Transfers.Num() > 0)
+		{
+			TransfersDeferred->Targets.Add({ ViewFamily.RenderTarget, nullptr, MoveTemp(Transfers) });
+		}
+	}
+	else if (Transfers.Num() > 0)
 	{
-		// Check if we can go through the optimized code path, with delay for the cross GPU transfer fence wait
 		if (bCrossGPUTransferFencesDefer)
 		{
+			// Optimized push transfer code path, with delay for the cross GPU transfer fence wait
 			// A readback pass is the closest analog to what this is doing. There isn't a way to express cross-GPU transfers via the RHI barrier API.
 			AddReadbackTexturePass(GraphBuilder, RDG_EVENT_NAME("CrossGPUTransfers"), RenderTargetTexture,
 				[this, RenderTargetTexture, LocalTransfers = MoveTemp(Transfers), PostTransferFences = MoveTemp(CrossGPUTransferFencesDefer)](FRHICommandListImmediate& RHICmdList)
@@ -3734,7 +3777,7 @@ void FSceneRenderer::DoCrossGPUTransfers(FRDGBuilder& GraphBuilder, FRDGTextureR
 				TArray<FTransferResourceParams> TransferParams;
 				for (const FCrossGPUTransfer& Transfer : LocalTransfers)
 				{
-					TransferParams.Add(FTransferResourceParams(RenderTargetTexture->GetRHI(), Transfer.SrcGPUIndex, Transfer.DestGPUIndex, false, false));
+					TransferParams.Add(FTransferResourceParams(RenderTargetTexture->GetRHI(), Transfer.SrcGPUIndex, Transfer.DestGPUIndex, true, false));
 				}
 
 				RHICmdList.TransferResources(TransferParams);
@@ -3744,6 +3787,44 @@ void FSceneRenderer::DoCrossGPUTransfers(FRDGBuilder& GraphBuilder, FRDGTextureR
 #endif // WITH_MGPU
 }
 
+void FSceneRenderer::FlushCrossGPUTransfers(FRDGBuilder& GraphBuilder)
+{
+#if WITH_MGPU
+	if (CrossGPUTransferDeferred)
+	{
+		// If this is the last scene renderer, flush the transfers
+		if (CrossGPUTransferDeferred->GetRefCount() == 1 && CrossGPUTransferDeferred->Targets.Num())
+		{
+			// Create RDG textures for each render target
+			for (FCrossGPUTarget& Target : CrossGPUTransferDeferred->Targets)
+			{
+				FRHITexture* TextureRHI = Target.RenderTarget->GetRenderTargetTexture();
+				check(TextureRHI);
+				Target.RenderTargetTexture = RegisterExternalTexture(GraphBuilder, TextureRHI, TEXT("CrossGPUTexture"));
+			}
+
+			AddPass(GraphBuilder, RDG_EVENT_NAME("CrossGPUTransfers"),
+				[LocalTransfers = CrossGPUTransferDeferred](FRHICommandList& RHICmdList)
+			{
+				TArray<FTransferResourceParams> TransferParams;
+				for (const FCrossGPUTarget& Target : LocalTransfers->Targets)
+				{
+					for (const FCrossGPUTransfer& Transfer : Target.Transfers)
+					{
+						TransferParams.Add(FTransferResourceParams(Target.RenderTargetTexture->GetRHI(), Transfer.SrcGPUIndex, Transfer.DestGPUIndex, true, true));
+					}
+				}
+
+				RHICmdList.TransferResources(TransferParams);
+			});
+		}
+
+		// Remove reference to the deferred transfers in the flush for each scene
+		CrossGPUTransferDeferred = nullptr;
+	}
+#endif // WITH_MGPU
+}
+
 void FSceneRenderer::FlushCrossGPUFences(FRDGBuilder& GraphBuilder)
 {
 #if WITH_MGPU
diff --git a/Engine/Source/Runtime/Renderer/Private/SceneRendering.h b/Engine/Source/Runtime/Renderer/Private/SceneRendering.h
--- a/Engine/Source/Runtime/Renderer/Private/SceneRendering.h
+++ b/Engine/Source/Runtime/Renderer/Private/SceneRendering.h
@@ -2115,7 +2115,8 @@ public:
 
 	/** Logic to update render targets across all GPUs */
 	static void PreallocateCrossGPUFences(TConstArrayView<FSceneRenderer*> SceneRenderers);
-	void DoCrossGPUTransfers(FRDGBuilder& GraphBuilder, FRDGTextureRef ViewFamilyTexture, TArrayView<FViewInfo> InViews, bool bCrossGPUTransferFencesDefer, FRHIGPUMask RenderTargetGPUMask);
+	void DoCrossGPUTransfers(FRDGBuilder& GraphBuilder, FRDGTextureRef ViewFamilyTexture, TArrayView<FViewInfo> InViews, bool bCrossGPUTransferFencesDefer, FRHIGPUMask RenderTargetGPUMask, class FCrossGPUTransfersDeferred* TransfersDeferred);
+	void FlushCrossGPUTransfers(FRDGBuilder& GraphBuilder);
 	void FlushCrossGPUFences(FRDGBuilder& GraphBuilder);
 
 	bool DoOcclusionQueries() const;
@@ -2276,6 +2277,9 @@ protected:
 	TArray<FCrossGPUTransferFence*> CrossGPUTransferFencesDefer;
 	TArray<FCrossGPUTransferFence*> CrossGPUTransferFencesWait;
 
+	/** Deferred transfers to be executed in the last scene renderer */
+	TRefCountPtr<class FCrossGPUTransfersDeferred> CrossGPUTransferDeferred;
+
 	FRHIGPUMask AllViewsGPUMask;
 	bool IsShadowCached(FProjectedShadowInfo* ProjectedShadowInfo) const;
 	FRHIGPUMask GetGPUMaskForShadow(FProjectedShadowInfo* ProjectedShadowInfo) const;
