First Person: Fixed static mesh first person primitives with materials without WPO or PDO or masked materials being rendered with an incompatible position-only vertex shader, resulting in a mismatch between the depth prepass and then basepass outputs.

#tests Tested with with the steps described in the jira.

diff --git a/Engine/Source/Runtime/Renderer/Private/SceneRendering.cpp b/Engine/Source/Runtime/Renderer/Private/SceneRendering.cpp
--- a/Engine/Source/Runtime/Renderer/Private/SceneRendering.cpp
+++ b/Engine/Source/Runtime/Renderer/Private/SceneRendering.cpp
@@ -6757,6 +6757,9 @@ bool DoMaterialAndPrimitiveModifyMeshPosition(const FMaterial& Material, const F
 {
 	const bool bMaterialModifiesMeshPosition = Material.MaterialModifiesMeshPosition_RenderThread();
 	const bool bPrimitiveAllowsWPOEvaluation = !ShouldOptimizedWPOAffectNonNaniteShaderSelection() || (PrimitiveSceneProxy && PrimitiveSceneProxy->EvaluateWorldPositionOffset());
-	const bool bIsFirstPerson = Material.HasFirstPersonOutput() && PrimitiveSceneProxy && PrimitiveSceneProxy->IsFirstPerson(); // This check is necessary for first person rendering to function in case of ShouldOptimizedWPOAffectNonNaniteShaderSelection && !EvaluateWorldPositionOffset.
-	return bMaterialModifiesMeshPosition && (bPrimitiveAllowsWPOEvaluation || bIsFirstPerson);
+	const bool bIsFirstPerson = PrimitiveSceneProxy && PrimitiveSceneProxy->IsFirstPerson();
+	// First person primitives have special logic that modifies vertex positions after WPO has been applied. If Material.HasFirstPersonOutput() is true, then the material has custom logic controlling this,
+	// but even without that custom logic, optimized position-only shaders for depth-only rendering do not support first person, so by checking the bIsFirstPerson flag here, we prevent first person primitives
+	// from being rendered with such incompatible shaders.
+	return (bMaterialModifiesMeshPosition && bPrimitiveAllowsWPOEvaluation) || bIsFirstPerson;
 }
