[Volumetric Fog] Fix grid extending past "VolumetricFogDistance"

diff --git a/Engine/Source/Runtime/Renderer/Private/VolumetricFog.cpp b/Engine/Source/Runtime/Renderer/Private/VolumetricFog.cpp
--- a/Engine/Source/Runtime/Renderer/Private/VolumetricFog.cpp
+++ b/Engine/Source/Runtime/Renderer/Private/VolumetricFog.cpp
@@ -1169,20 +1169,20 @@ FVector GetVolumetricFogGridZParams(float VolumetricFogStartDistance, float Near
 	double N = NearPlane + NearOffset;
 	double F = FarPlane;
 
-	double O = (F - N * FMath::Exp2((GridSizeZ - 1) / S)) / (F - N);
+	double O = (F - N * FMath::Exp2(GridSizeZ / S)) / (F - N);
 	double B = (1 - O) / N;
 
-	double O2 = (FMath::Exp2((GridSizeZ - 1) / S) - F / N) / (-F / N + 1);
-
-	float FloatN = (float)N;
-	float FloatF = (float)F;
-	float FloatB = (float)B;
-	float FloatO = (float)O;
-	float FloatS = (float)S;
-
-	float NSlice = FMath::Log2(FloatN*FloatB + FloatO) * FloatS;
-	float NearPlaneSlice = FMath::Log2(NearPlane*FloatB + FloatO) * FloatS;
-	float FSlice = FMath::Log2(FloatF*FloatB + FloatO) * FloatS;
+// 	double O2 = (FMath::Exp2((GridSizeZ - 1) / S) - F / N) / (-F / N + 1);
+// 
+// 	float FloatN = (float)N;
+// 	float FloatF = (float)F;
+// 	float FloatB = (float)B;
+// 	float FloatO = (float)O;
+// 	float FloatS = (float)S;
+// 
+// 	float NSlice = FMath::Log2(FloatN*FloatB + FloatO) * FloatS;
+// 	float NearPlaneSlice = FMath::Log2(NearPlane*FloatB + FloatO) * FloatS;
+// 	float FSlice = FMath::Log2(FloatF*FloatB + FloatO) * FloatS;
 	// y = log2(z*B + O) * S
 	// f(N) = 0 = log2(N*B + O) * S
 	// 1 = N*B + O
