Restored the behavior of r.Shadow.Virtual.MarkCoarsePagesDirectional which was broken in a recent change.

#rnx

diff --git a/Engine/Source/Runtime/Renderer/Private/VirtualShadowMaps/VirtualShadowMapArray.cpp b/Engine/Source/Runtime/Renderer/Private/VirtualShadowMaps/VirtualShadowMapArray.cpp
--- a/Engine/Source/Runtime/Renderer/Private/VirtualShadowMaps/VirtualShadowMapArray.cpp
+++ b/Engine/Source/Runtime/Renderer/Private/VirtualShadowMaps/VirtualShadowMapArray.cpp
@@ -154,14 +154,6 @@ TAutoConsoleVariable<int32> CVarMarkPixelPagesMipModeLocal(
 	ECVF_RenderThreadSafe | ECVF_Scalability
 );
 
-TAutoConsoleVariable<int32> CVarMarkCoarsePagesDirectional(
-	TEXT("r.Shadow.Virtual.MarkCoarsePagesDirectional"),
-	1,
-	TEXT("Marks coarse pages in directional light virtual shadow maps so that low resolution data is available everywhere.")
-	TEXT("Ability to disable is primarily for profiling and debugging."),
-	ECVF_Scalability | ECVF_RenderThreadSafe
-);
-
 TAutoConsoleVariable<int32> CVarMarkCoarsePagesLocal(
 	TEXT("r.Shadow.Virtual.MarkCoarsePagesLocal"),
 	1,
@@ -1936,7 +1928,6 @@ void FVirtualShadowMapArray::BuildPageAllocations(
 			// can be over-written by any pixel pages that *do* want Non-Nanite geometry. We avoid writing with atomics since that
 			// is much slower.
 			// Because of this we also cannot overlap this pass with the following ones.
-			bool bMarkCoarsePagesDirectional = CVarMarkCoarsePagesDirectional.GetValueOnRenderThread() != 0;
 			bool bMarkCoarsePagesLocal = CVarMarkCoarsePagesLocal.GetValueOnRenderThread() != 0;
 			// Note: always run this pass such that the distant lights may be marked if need be
 			{
diff --git a/Engine/Source/Runtime/Renderer/Private/VirtualShadowMaps/VirtualShadowMapClipmap.cpp b/Engine/Source/Runtime/Renderer/Private/VirtualShadowMaps/VirtualShadowMapClipmap.cpp
--- a/Engine/Source/Runtime/Renderer/Private/VirtualShadowMaps/VirtualShadowMapClipmap.cpp
+++ b/Engine/Source/Runtime/Renderer/Private/VirtualShadowMaps/VirtualShadowMapClipmap.cpp
@@ -126,6 +126,14 @@ static TAutoConsoleVariable<bool> CVarOrthoVSMRayCastViewOrigin(
 	ECVF_Scalability | ECVF_RenderThreadSafe
 );
 
+TAutoConsoleVariable<int32> CVarMarkCoarsePagesDirectional(
+	TEXT("r.Shadow.Virtual.MarkCoarsePagesDirectional"),
+	1,
+	TEXT("Marks coarse pages in directional light virtual shadow maps so that low resolution data is available everywhere.")
+	TEXT("Ability to disable is primarily for profiling and debugging."),
+	ECVF_Scalability | ECVF_RenderThreadSafe
+);
+
 // "Virtual" clipmap level to clipmap radius
 // NOTE: This is the radius of around the clipmap origin that this level must cover
 // The actual clipmap dimensions will be larger due to snapping and other accomodations
@@ -138,12 +146,13 @@ static float GetLevelRadius(int32 Level)
 
 FVirtualShadowMapClipmapConfig FVirtualShadowMapClipmapConfig::GetGlobal()
 {
+	bool bMarkCoarsePagesDirectional = CVarMarkCoarsePagesDirectional.GetValueOnRenderThread() != 0;
 	return FVirtualShadowMapClipmapConfig 
 	{ 
 		.FirstLevel = CVarVirtualShadowMapClipmapFirstLevel.GetValueOnRenderThread(),
 		.LastLevel = CVarVirtualShadowMapClipmapLastLevel.GetValueOnRenderThread(),
-		.FirstCoarseLevel = CVarVirtualShadowMapClipmapFirstCoarseLevel.GetValueOnRenderThread(),
-		.LastCoarseLevel = CVarVirtualShadowMapClipmapLastCoarseLevel.GetValueOnRenderThread(),
+		.FirstCoarseLevel = bMarkCoarsePagesDirectional ? CVarVirtualShadowMapClipmapFirstCoarseLevel.GetValueOnRenderThread() : -1,
+		.LastCoarseLevel = bMarkCoarsePagesDirectional ? CVarVirtualShadowMapClipmapLastCoarseLevel.GetValueOnRenderThread() : -1,
 		.bForceInvalidate = CVarForceInvalidateDirectionalVSM.GetValueOnRenderThread() != 0
 	};
 }
@@ -474,11 +483,14 @@ FVirtualShadowMapProjectionShaderData FVirtualShadowMapClipmap::ComputeProjectio
 	Data.LightSourceRadius = GetLightSceneInfo().Proxy->GetSourceRadius();
 	Data.Flags = PerLightCacheEntry->IsUncached() ? VSM_PROJ_FLAG_UNCACHED : 0U;
 
-	if ((Data.ClipmapLevel >= Config.FirstCoarseLevel && Data.ClipmapLevel <= Config.LastCoarseLevel)
-		// Always mark coarse pages in the last level for clouds/skyatmosphere
-		|| 	Data.ClipmapLevelCountRemaining == 1)
+	if (Config.FirstCoarseLevel >= 0)
 	{
-		Data.Flags |= VSM_PROJ_FLAG_IS_COARSE_CLIP_LEVEL;
+		if ((Data.ClipmapLevel >= Config.FirstCoarseLevel && Data.ClipmapLevel <= Config.LastCoarseLevel)
+			// Always mark coarse pages in the last level for clouds/skyatmosphere
+			|| 	Data.ClipmapLevelCountRemaining == 1)
+		{
+			Data.Flags |= VSM_PROJ_FLAG_IS_COARSE_CLIP_LEVEL;
+		}
 	}
 
 
