- Set StaticShaderBindingLayoutSupport to Unsupported for SM5
- Disable runtime bindless checks when shader binding layout is enabled

diff --git a/Engine/Config/Windows/DataDrivenPlatformInfo.ini b/Engine/Config/Windows/DataDrivenPlatformInfo.ini
--- a/Engine/Config/Windows/DataDrivenPlatformInfo.ini
+++ b/Engine/Config/Windows/DataDrivenPlatformInfo.ini
@@ -62,7 +62,7 @@ bSupportsIndexBufferUAVs=true
 bSupportsROV=true
 bSupportsOIT=true
 BindlessSupport=Unsupported
-StaticShaderBindingLayoutSupport=RayTracingOnly
+StaticShaderBindingLayoutSupport=Unsupported
 SupportsMultiViewport=RuntimeDependent
 FriendlyName=LOCTEXT("FriendlyShaderPlatformName_D3D_SM5", "PC D3D SM5")
 bSupportsDualSourceBlending=true
diff --git a/Engine/Source/Runtime/Renderer/Private/RayTracing/RayTracing.cpp b/Engine/Source/Runtime/Renderer/Private/RayTracing/RayTracing.cpp
--- a/Engine/Source/Runtime/Renderer/Private/RayTracing/RayTracing.cpp
+++ b/Engine/Source/Runtime/Renderer/Private/RayTracing/RayTracing.cpp
@@ -354,8 +354,10 @@ namespace RayTracing
 		if (RHIGetStaticShaderBindingLayoutSupport(ShaderPlatform) != ERHIStaticShaderBindingLayoutSupport::Unsupported)
 		{
 			// Should support bindless for raytracing at least
-			check(RHIGetRuntimeBindlessResourcesConfiguration(ShaderPlatform) != ERHIBindlessConfiguration::Disabled);
-			check(RHIGetRuntimeBindlessSamplersConfiguration(ShaderPlatform) != ERHIBindlessConfiguration::Disabled);
+			// NOTE: checks disable checks because GConfig which is used to check 
+			// runtime binding config can be modified in another thread at the same time
+			//check(RHIGetRuntimeBindlessResourcesConfiguration(ShaderPlatform) != ERHIBindlessConfiguration::Disabled);
+			//check(RHIGetRuntimeBindlessSamplersConfiguration(ShaderPlatform) != ERHIBindlessConfiguration::Disabled);
 
 			// Retrieve the bindless shader binding table
 			return &FRaytracingShaderBindingLayout::GetInstance(FShaderBindingLayoutContainer::EBindingType::Bindless);
