HeterogeneousVolumes: Disable use of analytic derivatives by default

diff --git a/Engine/Source/Runtime/Renderer/Private/HeterogeneousVolumes/HeterogeneousVolumes.cpp b/Engine/Source/Runtime/Renderer/Private/HeterogeneousVolumes/HeterogeneousVolumes.cpp
--- a/Engine/Source/Runtime/Renderer/Private/HeterogeneousVolumes/HeterogeneousVolumes.cpp
+++ b/Engine/Source/Runtime/Renderer/Private/HeterogeneousVolumes/HeterogeneousVolumes.cpp
@@ -260,8 +260,8 @@ static TAutoConsoleVariable<int32> CVarHeterogeneousVolumesApplyFogInscatteringM
 
 static TAutoConsoleVariable<int32> CVarHeterogeneousVolumesUseAnalyticDerivatives(
 	TEXT("r.HeterogeneousVolumes.UseAnalyticDerivatives"),
-	1,
-	TEXT("Writes Heterogeneous Volumes velocity to the feature buffer (Default = 1)"),
+	0,
+	TEXT("Enables support for analytic derivatives (Default = 0)"),
 	ECVF_RenderThreadSafe
 );
 
