Fix Distance Field Shadows being corrupted in split-screen or stereo views past the first.

diff --git a/Engine/Source/Runtime/Renderer/Private/DistanceFieldShadowing.cpp b/Engine/Source/Runtime/Renderer/Private/DistanceFieldShadowing.cpp
--- a/Engine/Source/Runtime/Renderer/Private/DistanceFieldShadowing.cpp
+++ b/Engine/Source/Runtime/Renderer/Private/DistanceFieldShadowing.cpp
@@ -915,7 +915,9 @@ FScreenPassTexture FProjectedShadowInfo::RenderRayTracedDistanceFieldProjection(
 	DistanceFieldShadowViewGPUData& SDFShadowViewGPUData = CachedDistanceFieldShadowViewGPUData.FindOrAdd(&View);
 
 	FIntRect DownsampledScissorRect = GetScissorRectForDFShadows(ScissorRect);
-	DownsampledScissorRect -= DownsampledScissorRect.Min; // DistanceFieldShadowingCS always outputs at rect with min = (0,0)
+	
+	// DistanceFieldShadowingCS always outputs at rect with min = (0,0)
+	DownsampledScissorRect = DownsampledScissorRect - DownsampledScissorRect.Min; // Beware - this gives different results than using '-='!
 
 	if (SDFShadowViewGPUData.RayTracedShadowsTexture)
 	{
