HeterogeneousVolumes: Lower material permutation count

1. Remove FUseTransmittanceVolume
2. Combine FUseAdaptiveVolumetricShadowMap and FAVSMSampleMode into one enum

diff --git a/Engine/Shaders/Private/HeterogeneousVolumes/HeterogeneousVolumesAdaptiveVolumetricShadowMapSampling.ush b/Engine/Shaders/Private/HeterogeneousVolumes/HeterogeneousVolumesAdaptiveVolumetricShadowMapSampling.ush
--- a/Engine/Shaders/Private/HeterogeneousVolumes/HeterogeneousVolumesAdaptiveVolumetricShadowMapSampling.ush
+++ b/Engine/Shaders/Private/HeterogeneousVolumes/HeterogeneousVolumesAdaptiveVolumetricShadowMapSampling.ush
@@ -4,11 +4,12 @@
 
 #include "HeterogeneousVolumesAdaptiveVolumetricShadowMapUtils.ush"
 
-#define AVSM_SAMPLE_MODE_TWO_LEVEL_VECTORIZED 0
-#define AVSM_SAMPLE_MODE_TWO_LEVEL_VECTORIZED_LOOP 1
-//#define AVSM_SAMPLE_MODE_TWO_LEVEL_REFERENCE 2
-//#define AVSM_SAMPLE_MODE_BINARY_SEARCH 3
-//#define AVSM_SAMPLE_MODE_REFERENCE 4
+#define AVSM_SAMPLE_MODE_DISABLED 0
+#define AVSM_SAMPLE_MODE_TWO_LEVEL_VECTORIZED 1
+#define AVSM_SAMPLE_MODE_TWO_LEVEL_VECTORIZED_LOOP 2
+//#define AVSM_SAMPLE_MODE_TWO_LEVEL_REFERENCE 3
+//#define AVSM_SAMPLE_MODE_BINARY_SEARCH 4
+//#define AVSM_SAMPLE_MODE_REFERENCE 5
 
 #ifndef AVSM_SAMPLE_MODE
 #define AVSM_SAMPLE_MODE AVSM_SAMPLE_MODE_TWO_LEVEL_VECTORIZED_LOOP
diff --git a/Engine/Shaders/Private/HeterogeneousVolumes/HeterogeneousVolumesLiveShadingPipeline.usf b/Engine/Shaders/Private/HeterogeneousVolumes/HeterogeneousVolumesLiveShadingPipeline.usf
--- a/Engine/Shaders/Private/HeterogeneousVolumes/HeterogeneousVolumesLiveShadingPipeline.usf
+++ b/Engine/Shaders/Private/HeterogeneousVolumes/HeterogeneousVolumesLiveShadingPipeline.usf
@@ -24,6 +24,10 @@
 #define REFERENCE_FAST_PATH 0
 #endif //
 
+#ifndef DIM_USE_ADAPTIVE_VOLUMETRIC_SHADOW_MAP
+#define DIM_USE_ADAPTIVE_VOLUMETRIC_SHADOW_MAP (AVSM_SAMPLE_MODE != 0)
+#endif //
+
 #ifndef USE_EXISTENCE_MASK
 #define USE_EXISTENCE_MASK 0
 #endif //
diff --git a/Engine/Source/Runtime/Renderer/Private/HeterogeneousVolumes/HeterogeneousVolumes.cpp b/Engine/Source/Runtime/Renderer/Private/HeterogeneousVolumes/HeterogeneousVolumes.cpp
--- a/Engine/Source/Runtime/Renderer/Private/HeterogeneousVolumes/HeterogeneousVolumes.cpp
+++ b/Engine/Source/Runtime/Renderer/Private/HeterogeneousVolumes/HeterogeneousVolumes.cpp
@@ -206,7 +206,7 @@ static TAutoConsoleVariable<int32> CVarHeterogeneousVolumesLightingCache(
 	2,
 	TEXT("Enables an optimized pre-pass, caching certain volumetric rendering lighting quantities (Default = 2)\n")
 	TEXT("0: Disabled\n")
-	TEXT("1: Cache transmittance\n")
+	TEXT("1: Cache transmittance (deprecated)\n")
 	TEXT("2: Cache in-scattering\n"),
 	ECVF_RenderThreadSafe
 );
diff --git a/Engine/Source/Runtime/Renderer/Private/HeterogeneousVolumes/HeterogeneousVolumesHardwareRayTracing.cpp b/Engine/Source/Runtime/Renderer/Private/HeterogeneousVolumes/HeterogeneousVolumesHardwareRayTracing.cpp
--- a/Engine/Source/Runtime/Renderer/Private/HeterogeneousVolumes/HeterogeneousVolumesHardwareRayTracing.cpp
+++ b/Engine/Source/Runtime/Renderer/Private/HeterogeneousVolumes/HeterogeneousVolumesHardwareRayTracing.cpp
@@ -399,10 +399,10 @@ class FRenderSingleScatteringWithPreshadingRGS : public FGlobalShader
 	SHADER_USE_ROOT_PARAMETER_STRUCT(FRenderSingleScatteringWithPreshadingRGS, FGlobalShader)
 
 	class FApplyShadowTransmittanceDim : SHADER_PERMUTATION_BOOL("DIM_APPLY_SHADOW_TRANSMITTANCE");
-	class FUseTransmittanceVolume : SHADER_PERMUTATION_BOOL("DIM_USE_TRANSMITTANCE_VOLUME");
+	//class FUseTransmittanceVolume : SHADER_PERMUTATION_BOOL("DIM_USE_TRANSMITTANCE_VOLUME");
 	class FUseInscatteringVolume : SHADER_PERMUTATION_BOOL("DIM_USE_INSCATTERING_VOLUME");
 	class FUseLumenGI : SHADER_PERMUTATION_BOOL("DIM_USE_LUMEN_GI");
-	using FPermutationDomain = TShaderPermutationDomain<FApplyShadowTransmittanceDim, FUseTransmittanceVolume, FUseInscatteringVolume, FUseLumenGI>;
+	using FPermutationDomain = TShaderPermutationDomain<FApplyShadowTransmittanceDim, FUseInscatteringVolume, FUseLumenGI>;
 
 	BEGIN_SHADER_PARAMETER_STRUCT(FParameters, )
 		// Scene 
@@ -824,7 +824,7 @@ void RenderSingleScatteringWithPreshadingHardwareRayTracing(
 
 	FRenderSingleScatteringWithPreshadingRGS::FPermutationDomain PermutationVector;
 	PermutationVector.Set<FRenderSingleScatteringWithPreshadingRGS::FApplyShadowTransmittanceDim>(bApplyShadowTransmittance);
-	PermutationVector.Set<FRenderSingleScatteringWithPreshadingRGS::FUseTransmittanceVolume>(HeterogeneousVolumes::UseLightingCacheForTransmittance());
+	//PermutationVector.Set<FRenderSingleScatteringWithPreshadingRGS::FUseTransmittanceVolume>(HeterogeneousVolumes::UseLightingCacheForTransmittance());
 	PermutationVector.Set<FRenderSingleScatteringWithPreshadingRGS::FUseInscatteringVolume>(HeterogeneousVolumes::UseLightingCacheForInscattering());
 	PermutationVector.Set<FRenderSingleScatteringWithPreshadingRGS::FUseLumenGI>(HeterogeneousVolumes::UseIndirectLighting() && View.GetLumenTranslucencyGIVolume().Texture0 != nullptr);
 	TShaderRef<FRenderSingleScatteringWithPreshadingRGS> RayGenerationShader = View.ShaderMap->GetShader<FRenderSingleScatteringWithPreshadingRGS>(PermutationVector);
diff --git a/Engine/Source/Runtime/Renderer/Private/HeterogeneousVolumes/HeterogeneousVolumesLiveShadingPipeline.cpp b/Engine/Source/Runtime/Renderer/Private/HeterogeneousVolumes/HeterogeneousVolumesLiveShadingPipeline.cpp
--- a/Engine/Source/Runtime/Renderer/Private/HeterogeneousVolumes/HeterogeneousVolumesLiveShadingPipeline.cpp
+++ b/Engine/Source/Runtime/Renderer/Private/HeterogeneousVolumes/HeterogeneousVolumesLiveShadingPipeline.cpp
@@ -118,6 +118,24 @@ namespace HeterogeneousVolumes
 	{
 		return FMath::Max(CVarHeterogeneousVolumesShadowsNearClippingDistance.GetValueOnRenderThread(), 0.1);
 	}
+
+	enum class EAVSMSampleMode
+	{
+		Disabled,
+		Performance,
+		Quality
+	};
+
+	EAVSMSampleMode GetAVSMSampleMode(bool bEnabled)
+	{
+		EAVSMSampleMode SampleMode = EAVSMSampleMode::Disabled;
+		if (bEnabled)
+		{
+			SampleMode = HeterogeneousVolumes::GetShadowMaxSampleCount() > 16 ? EAVSMSampleMode::Quality : EAVSMSampleMode::Performance;
+		}
+
+		return SampleMode;
+	}
 }
 
 //-OPT: Remove duplicate bindings
@@ -199,11 +217,10 @@ class FRenderLightingCacheWithLiveShadingCS : public FMeshMaterialShader
 	DECLARE_SHADER_TYPE(FRenderLightingCacheWithLiveShadingCS, MeshMaterial);
 
 	class FLightingCacheMode : SHADER_PERMUTATION_INT("DIM_LIGHTING_CACHE_MODE", 2);
-	class FUseAdaptiveVolumetricShadowMap : SHADER_PERMUTATION_BOOL("DIM_USE_ADAPTIVE_VOLUMETRIC_SHADOW_MAP");
-	class FAVSMSampleMode : SHADER_PERMUTATION_INT("AVSM_SAMPLE_MODE", 2);
+	class FAVSMSampleMode : SHADER_PERMUTATION_INT("AVSM_SAMPLE_MODE", 3);
 	class FUseExistenceMask : SHADER_PERMUTATION_INT("USE_EXISTENCE_MASK", 2);
 	class FIsOfflineRender : SHADER_PERMUTATION_INT("IS_OFFLINE_RENDER", 2);
-	using FPermutationDomain = TShaderPermutationDomain<FLightingCacheMode, FUseAdaptiveVolumetricShadowMap, FAVSMSampleMode, FUseExistenceMask, FIsOfflineRender>;
+	using FPermutationDomain = TShaderPermutationDomain<FLightingCacheMode, FAVSMSampleMode, FUseExistenceMask, FIsOfflineRender>;
 
 	BEGIN_SHADER_PARAMETER_STRUCT(FParameters, )
 		// Scene data
@@ -340,12 +357,11 @@ class FRenderSingleScatteringWithLiveShadingCS : public FMeshMaterialShader
 {
 	DECLARE_SHADER_TYPE(FRenderSingleScatteringWithLiveShadingCS, MeshMaterial);
 
-	class FUseTransmittanceVolume : SHADER_PERMUTATION_BOOL("DIM_USE_TRANSMITTANCE_VOLUME");
+	//class FUseTransmittanceVolume : SHADER_PERMUTATION_BOOL("DIM_USE_TRANSMITTANCE_VOLUME");
 	class FUseInscatteringVolume : SHADER_PERMUTATION_BOOL("DIM_USE_INSCATTERING_VOLUME");
 	class FUseLumenGI : SHADER_PERMUTATION_BOOL("DIM_USE_LUMEN_GI");
 	class FWriteVelocity : SHADER_PERMUTATION_BOOL("DIM_WRITE_VELOCITY");
-	class FUseAdaptiveVolumetricShadowMap : SHADER_PERMUTATION_BOOL("DIM_USE_ADAPTIVE_VOLUMETRIC_SHADOW_MAP");
-	class FAVSMSampleMode : SHADER_PERMUTATION_INT("AVSM_SAMPLE_MODE", 2);
+	class FAVSMSampleMode : SHADER_PERMUTATION_INT("AVSM_SAMPLE_MODE", 3);
 	class FSupportOverlappingVolumes : SHADER_PERMUTATION_BOOL("SUPPORT_OVERLAPPING_VOLUMES");
 	//class FAdaptiveMarch : SHADER_PERMUTATION_INT("ADAPTIVE_MARCH", 2);
 	//class FBilinearInterpolation : SHADER_PERMUTATION_INT("AVSM_BILINEAR_INTERPOLATION", 2);
@@ -353,11 +369,10 @@ class FRenderSingleScatteringWithLiveShadingCS : public FMeshMaterialShader
 	class FApplyFogInscattering : SHADER_PERMUTATION_INT("APPLY_FOG_INSCATTERING", 3);
 	class FUseAnalyticDerivatives : SHADER_PERMUTATION_BOOL("USE_ANALYTIC_DERIVATIVES");
 	using FPermutationDomain = TShaderPermutationDomain<
-		FUseTransmittanceVolume, 
+		//FUseTransmittanceVolume, 
 		FUseInscatteringVolume, 
 		FUseLumenGI, 
 		FWriteVelocity, 
-		FUseAdaptiveVolumetricShadowMap, 
 		FApplyFogInscattering, 
 		FAVSMSampleMode, 
 		FSupportOverlappingVolumes,
@@ -848,12 +863,10 @@ static void RenderLightingCacheWithLiveShading(
 	GroupCount.Z = FMath::DivideAndRoundUp(GroupCount.Z, FRenderLightingCacheWithLiveShadingCS::GetThreadGroupSize3D());
 
 	bool bUseAVSM = HeterogeneousVolumes::UseAdaptiveVolumetricShadowMapForSelfShadowing(HeterogeneousVolumeInterface->GetPrimitiveSceneProxy());
-	int32 AVSMSampleMode = HeterogeneousVolumes::GetShadowMaxSampleCount() > 16;
 
 	FRenderLightingCacheWithLiveShadingCS::FPermutationDomain PermutationVector;
 	PermutationVector.Set<FRenderLightingCacheWithLiveShadingCS::FLightingCacheMode>(HeterogeneousVolumes::GetLightingCacheMode() - 1);
-	PermutationVector.Set<FRenderLightingCacheWithLiveShadingCS::FUseAdaptiveVolumetricShadowMap>(bUseAVSM);
-	PermutationVector.Set<FRenderLightingCacheWithLiveShadingCS::FAVSMSampleMode>(AVSMSampleMode);
+	PermutationVector.Set<FRenderLightingCacheWithLiveShadingCS::FAVSMSampleMode>(static_cast<int32>(HeterogeneousVolumes::GetAVSMSampleMode(bUseAVSM)));
 	PermutationVector.Set<FRenderLightingCacheWithLiveShadingCS::FUseExistenceMask>(HeterogeneousVolumes::UseExistenceMask());
 	PermutationVector.Set<FRenderLightingCacheWithLiveShadingCS::FIsOfflineRender>(View.bIsOfflineRender);
 	TShaderRef<FRenderLightingCacheWithLiveShadingCS> ComputeShader = Material.GetShader<FRenderLightingCacheWithLiveShadingCS>(&FLocalVertexFactory::StaticType, PermutationVector, false);
@@ -1255,15 +1268,12 @@ void RenderSingleScatteringWithLiveShading(
 	}
 #endif // WANTS_DRAW_MESH_EVENTS
 
-	int32 AVSMSampleMode = HeterogeneousVolumes::GetShadowMaxSampleCount() > 16;
-
 	typename FRenderSingleScatteringWithLiveShadingDispatchTypeCS::FPermutationDomain PermutationVector;
-	PermutationVector.template Set<typename FRenderSingleScatteringWithLiveShadingDispatchTypeCS::FUseTransmittanceVolume>(HeterogeneousVolumes::UseLightingCacheForTransmittance() && PassParameters->bApplyShadowTransmittance);
+	//PermutationVector.template Set<typename FRenderSingleScatteringWithLiveShadingDispatchTypeCS::FUseTransmittanceVolume>(HeterogeneousVolumes::UseLightingCacheForTransmittance() && PassParameters->bApplyShadowTransmittance);
 	PermutationVector.template Set<typename FRenderSingleScatteringWithLiveShadingDispatchTypeCS::FUseInscatteringVolume>(HeterogeneousVolumes::UseLightingCacheForInscattering());
 	PermutationVector.template Set<typename FRenderSingleScatteringWithLiveShadingDispatchTypeCS::FUseLumenGI>(HeterogeneousVolumes::UseIndirectLighting() && View.GetLumenTranslucencyGIVolume().Texture0 != nullptr);
 	PermutationVector.template Set<typename FRenderSingleScatteringWithLiveShadingDispatchTypeCS::FWriteVelocity>(bWriteVelocity);
-	PermutationVector.template Set<typename FRenderSingleScatteringWithLiveShadingDispatchTypeCS::FUseAdaptiveVolumetricShadowMap>(bUseAVSM);
-	PermutationVector.template Set<typename FRenderSingleScatteringWithLiveShadingDispatchTypeCS::FAVSMSampleMode>(AVSMSampleMode);
+	PermutationVector.template Set<typename FRenderSingleScatteringWithLiveShadingDispatchTypeCS::FAVSMSampleMode>(static_cast<int32>(HeterogeneousVolumes::GetAVSMSampleMode(bUseAVSM)));
 	PermutationVector.template Set<typename FRenderSingleScatteringWithLiveShadingDispatchTypeCS::FSupportOverlappingVolumes>(HeterogeneousVolumes::SupportsOverlappingVolumes());
 	//PermutationVector.template Set<typename FRenderSingleScatteringWithLiveShadingDispatchTypeCS::FBilinearInterpolation>(static_cast<int32>(HeterogeneousVolumes::UseBilinearInterpolation()));
 	//PermutationVector.template Set<typename FRenderSingleScatteringWithLiveShadingDispatchTypeCS::FAdaptiveMarch>(static_cast<int32>(HeterogeneousVolumes::ShouldAdaptiveMarch()));
diff --git a/Engine/Source/Runtime/Renderer/Private/HeterogeneousVolumes/HeterogeneousVolumesPreshadingPipeline.cpp b/Engine/Source/Runtime/Renderer/Private/HeterogeneousVolumes/HeterogeneousVolumesPreshadingPipeline.cpp
--- a/Engine/Source/Runtime/Renderer/Private/HeterogeneousVolumes/HeterogeneousVolumesPreshadingPipeline.cpp
+++ b/Engine/Source/Runtime/Renderer/Private/HeterogeneousVolumes/HeterogeneousVolumesPreshadingPipeline.cpp
@@ -293,13 +293,13 @@ class FRenderSingleScatteringWithPreshadingCS : public FGlobalShader
 	class FApplyShadowTransmittanceDim : SHADER_PERMUTATION_BOOL("DIM_APPLY_SHADOW_TRANSMITTANCE");
 	class FVoxelCullingDim : SHADER_PERMUTATION_BOOL("DIM_VOXEL_CULLING");
 	class FSparseVoxelTracingDim : SHADER_PERMUTATION_BOOL("DIM_SPARSE_VOXEL_TRACING");
-	class FUseTransmittanceVolume : SHADER_PERMUTATION_BOOL("DIM_USE_TRANSMITTANCE_VOLUME");
+	//class FUseTransmittanceVolume : SHADER_PERMUTATION_BOOL("DIM_USE_TRANSMITTANCE_VOLUME");
 	class FUseInscatteringVolume : SHADER_PERMUTATION_BOOL("DIM_USE_INSCATTERING_VOLUME");
 	class FUseLumenGI : SHADER_PERMUTATION_BOOL("DIM_USE_LUMEN_GI");
 	class FWriteVelocity : SHADER_PERMUTATION_BOOL("DIM_WRITE_VELOCITY");
 	class FUseAdaptiveVolumetricShadowMap : SHADER_PERMUTATION_BOOL("DIM_USE_ADAPTIVE_VOLUMETRIC_SHADOW_MAP");
 	class FDebugDim : SHADER_PERMUTATION_BOOL("DIM_DEBUG");
-	using FPermutationDomain = TShaderPermutationDomain<FApplyShadowTransmittanceDim, FVoxelCullingDim, FSparseVoxelTracingDim, FUseTransmittanceVolume, FUseInscatteringVolume, FUseLumenGI, FWriteVelocity, FUseAdaptiveVolumetricShadowMap, FDebugDim>;
+	using FPermutationDomain = TShaderPermutationDomain<FApplyShadowTransmittanceDim, FVoxelCullingDim, FSparseVoxelTracingDim, FUseInscatteringVolume, FUseLumenGI, FWriteVelocity, FUseAdaptiveVolumetricShadowMap, FDebugDim>;
 
 	BEGIN_SHADER_PARAMETER_STRUCT(FParameters, )
 		// Scene data
@@ -377,12 +377,6 @@ class FRenderSingleScatteringWithPreshadingCS : public FGlobalShader
 		{
 			return EShaderPermutationPrecacheRequest::NotUsed;
 		}
-
-		if (PermutationVector.Get<FUseTransmittanceVolume>() != HeterogeneousVolumes::UseLightingCacheForTransmittance())
-		{
-			return EShaderPermutationPrecacheRequest::NotUsed;
-		}
-
 		if (PermutationVector.Get<FUseInscatteringVolume>() != HeterogeneousVolumes::UseLightingCacheForInscattering())
 		{
 			return EShaderPermutationPrecacheRequest::NotUsed;
@@ -691,7 +685,7 @@ void RenderSingleScatteringWithPreshadingCompute(
 	PermutationVector.Set<FRenderSingleScatteringWithPreshadingCS::FApplyShadowTransmittanceDim>(bApplyShadowTransmittance);
 	PermutationVector.Set<FRenderSingleScatteringWithPreshadingCS::FVoxelCullingDim>(HeterogeneousVolumes::UseSparseVoxelPerTileCulling());
 	PermutationVector.Set<FRenderSingleScatteringWithPreshadingCS::FSparseVoxelTracingDim>(HeterogeneousVolumes::UseSparseVoxelPipeline());
-	PermutationVector.Set<FRenderSingleScatteringWithPreshadingCS::FUseTransmittanceVolume>(HeterogeneousVolumes::UseLightingCacheForTransmittance());
+	//PermutationVector.Set<FRenderSingleScatteringWithPreshadingCS::FUseTransmittanceVolume>(HeterogeneousVolumes::UseLightingCacheForTransmittance());
 	PermutationVector.Set<FRenderSingleScatteringWithPreshadingCS::FUseInscatteringVolume>(HeterogeneousVolumes::UseLightingCacheForInscattering());
 	PermutationVector.Set<FRenderSingleScatteringWithPreshadingCS::FUseLumenGI>(HeterogeneousVolumes::UseIndirectLighting() && View.GetLumenTranslucencyGIVolume().Texture0 != nullptr);
 	PermutationVector.Set<FRenderSingleScatteringWithPreshadingCS::FWriteVelocity>(bWriteVelocity);
