Mobile: Allow shadows on a single layer water

diff --git a/Engine/Shaders/Private/MobileLightingCommon.ush b/Engine/Shaders/Private/MobileLightingCommon.ush
--- a/Engine/Shaders/Private/MobileLightingCommon.ush
+++ b/Engine/Shaders/Private/MobileLightingCommon.ush
@@ -287,7 +287,7 @@ FDeferredLightData GetDirectionalLightData(float4 ScreenPosition, float4 SvPosit
 	half ShadowMap = MobileDirectionalLightCSM(ScreenPosition.xy, ScreenPosition.w, ShadowPositionZ);
 	OutDynamicShadowing = 1.0f;
 
-   #if !MOBILE_DEFERRED_SHADING && DIRECTIONAL_LIGHT_CSM && !MATERIAL_SHADINGMODEL_SINGLELAYERWATER
+   #if !MOBILE_DEFERRED_SHADING && DIRECTIONAL_LIGHT_CSM
 	// Cascaded Shadow Map
 	if (IsCSMEnabled())
 	{		
diff --git a/Engine/Source/Runtime/Renderer/Private/MobileBasePass.cpp b/Engine/Source/Runtime/Renderer/Private/MobileBasePass.cpp
--- a/Engine/Source/Runtime/Renderer/Private/MobileBasePass.cpp
+++ b/Engine/Source/Runtime/Renderer/Private/MobileBasePass.cpp
@@ -756,10 +756,11 @@ bool FMobileBasePassMeshProcessor::TryAddMeshBatch(const FMeshBatch& RESTRICT Me
 	if (ShouldDraw(Material))
 	{
 		const FMaterialShadingModelField ShadingModels = Material.GetShadingModels();
-		const bool bCanReceiveCSM = ((Flags & EFlags::CanReceiveCSM) == EFlags::CanReceiveCSM);
+		const bool bSingleLayerWater = ShadingModels.HasShadingModel(MSM_SingleLayerWater);
+		const bool bCanReceiveCSM = bSingleLayerWater || ((Flags & EFlags::CanReceiveCSM) == EFlags::CanReceiveCSM);
 		const EBlendMode BlendMode = Material.GetBlendMode();
 		const bool bIsLitMaterial = ShadingModels.IsLit();
-		const bool bIsTranslucent = IsTranslucentBlendMode(BlendMode) || ShadingModels.HasShadingModel(MSM_SingleLayerWater); // Water goes into the translucent pass;
+		const bool bIsTranslucent = IsTranslucentBlendMode(BlendMode) || bSingleLayerWater; // Water goes into the translucent pass;
 		const bool bIsMasked = IsMaskedBlendMode(Material);		
 		ELightMapPolicyType LightmapPolicyType = MobileBasePass::SelectMeshLightmapPolicy(Scene, MeshBatch, PrimitiveSceneProxy, bCanReceiveCSM, bPassUsesDeferredShading, bIsLitMaterial, bIsTranslucent);
 		return Process(MeshBatch, BatchElementMask, StaticMeshId, PrimitiveSceneProxy, MaterialRenderProxy, Material, bIsMasked, bIsTranslucent, ShadingModels, LightmapPolicyType, bCanReceiveCSM, MeshBatch.LCI);
