Fix missing hair shadow on instance having both 'hidden in game' and 'hidden shadow' enabled with light having RT shadow.

diff --git a/Engine/Source/Runtime/Renderer/Private/ShadowSetup.cpp b/Engine/Source/Runtime/Renderer/Private/ShadowSetup.cpp
--- a/Engine/Source/Runtime/Renderer/Private/ShadowSetup.cpp
+++ b/Engine/Source/Runtime/Renderer/Private/ShadowSetup.cpp
@@ -4513,11 +4513,17 @@ void FSceneRenderer::InitProjectedShadowVisibility(FDynamicShadowsTaskData& Task
 			}
 
 			// Register visible lights for allowing hair strands to cast shadow (directional light)
-			if (bHairStrands && LightIt->LightType == ELightComponentType::LightType_Directional)
+			if (!LightIt->CastRaytracedShadow && bHairStrands && LightIt->LightType == ELightComponentType::LightType_Directional)
 			{
 				HairStrands::AddVisibleShadowCastingLight(*Scene, Views, LightIt->LightSceneInfo);
 			}
 		}
+
+		// Register visible lights for allowing hair strands to cast shadow (directional light)
+		if (LightIt->CastRaytracedShadow && bHairStrands && LightIt->LightType == ELightComponentType::LightType_Directional)
+		{
+			HairStrands::AddVisibleShadowCastingLight(*Scene, Views, LightIt->LightSceneInfo);
+		}
 	}
 
 #if !UE_BUILD_SHIPPING
@@ -6183,6 +6189,25 @@ void FSceneRenderer::CreateDynamicShadows(FDynamicShadowsTaskData& TaskData)
 				FVisibleLightInfo& VisibleLightInfo = VisibleLightInfos[LightSceneInfo->Id];
 
 				const FLightOcclusionType OcclusionType = GetLightOcclusionType(LightSceneInfoCompact, ViewFamily);
+
+				// Register visible lights for allowing hair strands to cast shadow (non-directional light)
+				if (bHairStrands && OcclusionType == FLightOcclusionType::Raytraced && LightSceneInfo->Proxy->GetLightType() != LightType_Directional)
+				{
+					if (LightSceneInfoCompact.bCastStaticShadow || LightSceneInfoCompact.bCastDynamicShadow)
+					{
+						for (int32 ViewIndex = 0; ViewIndex < Views.Num(); ViewIndex++)
+						{
+							// View frustums are only checked when lights have visible primitives or have modulated shadows,
+							// so we don't need to check for that again here
+							if (LightSceneInfo->ShouldRenderLight(Views[ViewIndex]))
+							{
+								HairStrands::AddVisibleShadowCastingLight(*Scene, Views, LightSceneInfo);
+								break;
+							}
+						}
+					}
+				}
+
 				if (OcclusionType != FLightOcclusionType::Shadowmap &&
 					OcclusionType != FLightOcclusionType::MegaLightsVSM)
 				{
