Fix dangling pointer to texture names in FTextureDownsampleChain.
- store texture names in static TSet

diff --git a/Engine/Source/Runtime/Renderer/Private/PostProcess/PostProcessDownsample.cpp b/Engine/Source/Runtime/Renderer/Private/PostProcess/PostProcessDownsample.cpp
--- a/Engine/Source/Runtime/Renderer/Private/PostProcess/PostProcessDownsample.cpp
+++ b/Engine/Source/Runtime/Renderer/Private/PostProcess/PostProcessDownsample.cpp
@@ -207,6 +207,11 @@ void FTextureDownsampleChain::Init(
 
 	RDG_EVENT_SCOPE(GraphBuilder, "TextureDownsampleChain");
 
+#if !UE_BUILD_SHIPPING
+	checkf(IsInRenderingThread(), TEXT("Potential race condition accessing DebugNames?"));
+	static TSet<FString> DebugNames;
+#endif
+
 	Textures.SetNum(StageCount);
 
 	// The first stage is the input.
@@ -216,10 +221,14 @@ void FTextureDownsampleChain::Init(
 	{
 		const uint32 PreviousStageIndex = StageIndex - 1;
 
-		FString PassName = FString::Printf(TEXT("%s(1/%u)"), TextureName ? TextureName : InputTexture.TextureSRV->Desc.Texture->Name, 1 << (StageIndex + InputResolutionExponent));
+		const TCHAR* DebugNamePtr = TEXT("DownsampledTexture");
+#if !UE_BUILD_SHIPPING
+		const FString TextureDebugName = FString::Printf(TEXT("%s(1/%u)"), TextureName ? TextureName : InputTexture.TextureSRV->Desc.Texture->Name, 1 << (StageIndex + InputResolutionExponent));
+		DebugNamePtr = *DebugNames.FindOrAdd(TextureDebugName);
+#endif
 
 		FDownsamplePassInputs PassInputs;
-		PassInputs.Name = *PassName;
+		PassInputs.Name = DebugNamePtr;
 		PassInputs.SceneColor = Textures[PreviousStageIndex];
 		PassInputs.Quality = DownsampleQuality;
 
