[Local Exposure] Overview visualization mode and support for exposure fusion visualizations.

diff --git a/Engine/Shaders/Private/PostProcessVisualizeLocalExposure.usf b/Engine/Shaders/Private/PostProcessVisualizeLocalExposure.usf
--- a/Engine/Shaders/Private/PostProcessVisualizeLocalExposure.usf
+++ b/Engine/Shaders/Private/PostProcessVisualizeLocalExposure.usf
@@ -12,18 +12,27 @@ StructuredBuffer<float4> EyeAdaptationBuffer;
 Texture2D HDRSceneColorTexture;
 Texture2D SceneColorTexture;
 
+#if !EXPOSURE_FUSION
 Texture3D LumBilateralGrid;
 Texture2D BlurredLogLum;
+#else
+Texture2D ExposuresTexture;
+Texture2D WeightsTexture;
+#endif
 
-SamplerState TextureSampler;
+SamplerState BilinearClampSampler;
 
 uint DebugMode;
 
-#define VISUALIZE_LOCAL_EXPOSURE_MODE_DEFAULT			(0)
-#define VISUALIZE_LOCAL_EXPOSURE_MODE_THRESHOLDS		(1)
-#define VISUALIZE_LOCAL_EXPOSURE_MODE_BASE_LUMINANCE	(2)
-#define VISUALIZE_LOCAL_EXPOSURE_MODE_DETAIL_LUMINANCE	(3)
-#define VISUALIZE_LOCAL_EXPOSURE_MODE_VALID_LOOKUP		(4)
+#define VISUALIZE_LOCAL_EXPOSURE_MODE_DEFAULT			(1)
+#define VISUALIZE_LOCAL_EXPOSURE_MODE_THRESHOLDS		(2)
+#define VISUALIZE_LOCAL_EXPOSURE_MODE_BASE_LUMINANCE	(3)
+#define VISUALIZE_LOCAL_EXPOSURE_MODE_DETAIL_LUMINANCE	(4)
+#define VISUALIZE_LOCAL_EXPOSURE_MODE_VALID_LOOKUP		(5)
+#define VISUALIZE_LOCAL_EXPOSURE_MODE_FUSION_BASE		(6)
+#define VISUALIZE_LOCAL_EXPOSURE_MODE_FUSION_SHADOWS	(7)
+#define VISUALIZE_LOCAL_EXPOSURE_MODE_FUSION_HIGHLIGHTS	(8)
+#define VISUALIZE_LOCAL_EXPOSURE_MODE_FUSION_WEIGHTS	(9)
 
 SCREEN_PASS_TEXTURE_VIEWPORT(Input)
 SCREEN_PASS_TEXTURE_VIEWPORT(Output)
@@ -41,16 +50,25 @@ float4 MainPS(noperspective float4 UVAndScreenPos : TEXCOORD0, float4 SvPosition
 
 	float2 ExposureScaleMiddleGrey = EyeAdaptationBuffer[0].xw;
 
-	float4 HDRSceneColor = Texture2DSample(HDRSceneColorTexture, TextureSampler, UV) * View.OneOverPreExposure;
-	float4 SceneColor = Texture2DSample(SceneColorTexture, TextureSampler, UV);
+	float4 HDRSceneColor = Texture2DSample(HDRSceneColorTexture, BilinearClampSampler, UV) * View.OneOverPreExposure;
+	float4 SceneColor = Texture2DSample(SceneColorTexture, BilinearClampSampler, UV);
 
 	float LuminanceVal = CalculateEyeAdaptationLuminance(HDRSceneColor.rgb);
 	float LogLuminance = log2(LuminanceVal);
 	float MiddleGreyLumValue = log2(0.18 * ExposureScaleMiddleGrey.y * LocalExposure_MiddleGreyExposureCompensation);
 
-	float BaseLogLum = CalculateBaseLogLuminance(LogLuminance, LocalExposure_BlurredLuminanceBlend, ExposureScaleMiddleGrey.x, ViewportUV, LumBilateralGrid, BlurredLogLum, TextureSampler, TextureSampler);
+#if !EXPOSURE_FUSION
+	float BaseLogLum = CalculateBaseLogLuminance(LogLuminance, LocalExposure_BlurredLuminanceBlend, ExposureScaleMiddleGrey.x, ViewportUV, LumBilateralGrid, BlurredLogLum, BilinearClampSampler, BilinearClampSampler);
 	float LocalExposure = CalculateLocalExposure(LogLuminance + log2(ExposureScaleMiddleGrey.x), BaseLogLum, MiddleGreyLumValue, LocalExposure_HighlightContrastScale, LocalExposure_ShadowContrastScale, LocalExposure_DetailStrength);
 	float DetailLogLum = (LogLuminance + log2(ExposureScaleMiddleGrey.x)) - BaseLogLum;
+#else
+	float BaseLogLum = LogLuminance;
+	float LocalExposure = 1.0f;
+	float4 Exposures = Texture2DSample(ExposuresTexture, BilinearClampSampler, UV);
+	float4 Weights = Texture2DSample(WeightsTexture, BilinearClampSampler, UV);
+#endif
+
+	float BaseCentered = (BaseLogLum - MiddleGreyLumValue);
 
 	float3 OutValue = 0;
 
@@ -62,7 +80,6 @@ float4 MainPS(noperspective float4 UVAndScreenPos : TEXCOORD0, float4 SvPosition
 	{
 		OutValue = SceneColor.rgb;
 
-		float BaseCentered = (BaseLogLum - MiddleGreyLumValue);
 		if (BaseCentered > 0)
 		{
 			if (BaseCentered - LocalExposure_HighlightThreshold > 0)
@@ -78,13 +95,31 @@ float4 MainPS(noperspective float4 UVAndScreenPos : TEXCOORD0, float4 SvPosition
 			}
 		}
 	}
+#if !EXPOSURE_FUSION
 	else if (DebugMode == VISUALIZE_LOCAL_EXPOSURE_MODE_BASE_LUMINANCE)
 	{
-		OutValue = exp2(BaseLogLum);
+		float ContrastScale = BaseCentered > 0 ? LocalExposure_HighlightContrastScale : LocalExposure_ShadowContrastScale;
+
+		// Apply threshold
+		float ThresholdOffset;
+		{
+			if (BaseCentered > 0)
+			{
+				ThresholdOffset = BaseCentered - max(0.0f, BaseCentered - LocalExposure_HighlightThreshold);
+			}
+			else
+			{
+				ThresholdOffset = BaseCentered - min(0.0f, BaseCentered + LocalExposure_ShadowThreshold);
+			}
+
+			BaseCentered -= ThresholdOffset;
+		}
+
+		OutValue = exp2(MiddleGreyLumValue + ThresholdOffset + BaseCentered * ContrastScale);
 	}
 	else if (DebugMode == VISUALIZE_LOCAL_EXPOSURE_MODE_DETAIL_LUMINANCE)
 	{
-		OutValue = exp2(DetailLogLum);
+		OutValue = exp2(DetailLogLum * LocalExposure_DetailStrength);
 	}
 	else if (DebugMode == VISUALIZE_LOCAL_EXPOSURE_MODE_VALID_LOOKUP)
 	{
@@ -94,7 +129,7 @@ float4 MainPS(noperspective float4 UVAndScreenPos : TEXCOORD0, float4 SvPosition
 		BilateralGridUVW.xy = ViewportUV * LocalExposure_BilateralGridUVScale;
 		BilateralGridUVW.z = (ComputeHistogramPositionFromLogLuminance(LogLuminance) * (BILATERAL_GRID_DEPTH - 1) + 0.5f) / BILATERAL_GRID_DEPTH;
 
-		float2 BilateralGridLum = Texture3DSample(LumBilateralGrid, TextureSampler, BilateralGridUVW).xy;
+		float2 BilateralGridLum = Texture3DSample(LumBilateralGrid, BilinearClampSampler, BilateralGridUVW).xy;
 
 		if (BilateralGridLum.y < 0.001)
 		{
@@ -106,6 +141,30 @@ float4 MainPS(noperspective float4 UVAndScreenPos : TEXCOORD0, float4 SvPosition
 			OutValue = lerp(OutValue, float3(0,1,0), 0.5f);
 		}
 	}
+#else
+	else if (DebugMode == VISUALIZE_LOCAL_EXPOSURE_MODE_FUSION_BASE)
+	{
+		OutValue = Exposures.r;
+	}
+	else if (DebugMode == VISUALIZE_LOCAL_EXPOSURE_MODE_FUSION_HIGHLIGHTS)
+	{
+		OutValue = Exposures.g;
+	}
+	else if (DebugMode == VISUALIZE_LOCAL_EXPOSURE_MODE_FUSION_SHADOWS)
+	{
+		OutValue = Exposures.b;
+	}
+	else if (DebugMode == VISUALIZE_LOCAL_EXPOSURE_MODE_FUSION_WEIGHTS)
+	{
+		// swap rb channels so colors match other visualizations
+		// red = shadows
+		// green = highlights
+		float T = Weights.r;
+		Weights.r = Weights.b;
+		Weights.b = T;
+		OutValue = Weights.rgb;
+	}
+#endif
 
 	return float4(OutValue, 1.0f);
 }
\ No newline at end of file
diff --git a/Engine/Source/Runtime/Renderer/Private/PostProcess/PostProcessLocalExposure.cpp b/Engine/Source/Runtime/Renderer/Private/PostProcess/PostProcessLocalExposure.cpp
--- a/Engine/Source/Runtime/Renderer/Private/PostProcess/PostProcessLocalExposure.cpp
+++ b/Engine/Source/Runtime/Renderer/Private/PostProcess/PostProcessLocalExposure.cpp
@@ -366,7 +366,7 @@ void AddApplyLocalExposurePass(
 		FComputeShaderUtils::GetGroupCount(Output.ViewRect.Size(), FIntPoint(FApplyLocalExposureCS::ThreadGroupSizeX, FApplyLocalExposureCS::ThreadGroupSizeY)));
 }
 
-FScreenPassTexture AddLocalExposureFusionPass(
+FExposureFusionData AddLocalExposureFusionPass(
 	FRDGBuilder& GraphBuilder,
 	const FViewInfo& View,
 	const FEyeAdaptationParameters& EyeAdaptationParameters,
@@ -492,5 +492,10 @@ FScreenPassTexture AddLocalExposureFusionPass(
 		CoarserMip = CurrentLum;
 	}
 
-	return MoveTemp(Output);
+	FExposureFusionData OutputData;
+	OutputData.Result = Output;
+	OutputData.Exposures = LumTexture;
+	OutputData.Weights = WeightTexture;
+
+	return MoveTemp(OutputData);
 }
\ No newline at end of file
diff --git a/Engine/Source/Runtime/Renderer/Private/PostProcess/PostProcessLocalExposure.h b/Engine/Source/Runtime/Renderer/Private/PostProcess/PostProcessLocalExposure.h
--- a/Engine/Source/Runtime/Renderer/Private/PostProcess/PostProcessLocalExposure.h
+++ b/Engine/Source/Runtime/Renderer/Private/PostProcess/PostProcessLocalExposure.h
@@ -37,7 +37,14 @@ void AddApplyLocalExposurePass(
 	FScreenPassTextureSlice Output,
 	ERDGPassFlags PassFlags);
 
-FScreenPassTexture AddLocalExposureFusionPass(
+struct FExposureFusionData
+{
+	FScreenPassTexture Result;
+	FScreenPassTexture Exposures;
+	FScreenPassTexture Weights;
+};
+
+FExposureFusionData AddLocalExposureFusionPass(
 	FRDGBuilder& GraphBuilder,
 	const FViewInfo& View,
 	const FEyeAdaptationParameters& EyeAdaptationParameters,
diff --git a/Engine/Source/Runtime/Renderer/Private/PostProcess/PostProcessVisualizeLocalExposure.cpp b/Engine/Source/Runtime/Renderer/Private/PostProcess/PostProcessVisualizeLocalExposure.cpp
--- a/Engine/Source/Runtime/Renderer/Private/PostProcess/PostProcessVisualizeLocalExposure.cpp
+++ b/Engine/Source/Runtime/Renderer/Private/PostProcess/PostProcessVisualizeLocalExposure.cpp
@@ -3,19 +3,25 @@
 #include "PostProcess/PostProcessVisualizeLocalExposure.h"
 #include "PostProcess/PostProcessTonemap.h"
 #include "PostProcess/PostProcessLocalExposure.h"
+#include "PostProcess/PostProcessVisualizeBuffer.h"
 #include "UnrealEngine.h"
 #include "DataDrivenShaderPlatformInfo.h"
 #include "SceneRendering.h"
 
-TAutoConsoleVariable<int> CVarLocalExposureVisualizeDebugMode(
-	TEXT("r.LocalExposure.VisualizeDebugMode"),
+TAutoConsoleVariable<int> CVarLocalExposureVisualizationMode(
+	TEXT("r.LocalExposure.VisualizationMode"),
 	0,
-	TEXT("When enabling Show->Visualize->Local Exposure is enabled, this flag controls which mode to use.\n")
-	TEXT("    0: Local Exposure\n")
-	TEXT("    1: Thresholds\n")
-	TEXT("    2: Base Luminance\n")
-	TEXT("    3: Detail Luminance\n")
-	TEXT("    4: Valid Bilateral Grid Lookup\n"),
+	TEXT("When enabling Show->Visualize->Local Exposure is enabled, this cvar controls which mode to use.\n")
+	TEXT("    0: Overview\n")
+	TEXT("    1: Local Exposure\n")
+	TEXT("    2: Thresholds\n")
+	TEXT("    3: Base Luminance\n")
+	TEXT("    4: Detail Luminance\n")
+	TEXT("    5: Valid Bilateral Grid Lookup\n")
+	TEXT("    6: Fusion - Base Exposure\n")
+	TEXT("    7: Fusion - Shadows Exposure\n")
+	TEXT("    8: Fusion - Highlights Exposure\n")
+	TEXT("    9: Fusion - Weights\n"),
 	ECVF_RenderThreadSafe);
 
 class FVisualizeLocalExposurePS : public FGlobalShader
@@ -37,13 +43,20 @@ public:
 		SHADER_PARAMETER_RDG_TEXTURE(Texture3D, LumBilateralGrid)
 		SHADER_PARAMETER_RDG_TEXTURE(Texture2D, BlurredLogLum)
 
-		SHADER_PARAMETER_SAMPLER(SamplerState, TextureSampler)
+		SHADER_PARAMETER_RDG_TEXTURE(Texture2D, ExposuresTexture)
+		SHADER_PARAMETER_RDG_TEXTURE(Texture2D, WeightsTexture)
+
+		SHADER_PARAMETER_SAMPLER(SamplerState, BilinearClampSampler)
 
 		SHADER_PARAMETER(uint32, DebugMode)
 
 		RENDER_TARGET_BINDING_SLOTS()
 	END_SHADER_PARAMETER_STRUCT()
 
+	class FExposureFusion : SHADER_PERMUTATION_BOOL("EXPOSURE_FUSION");
+
+	using FPermutationDomain = TShaderPermutationDomain<FExposureFusion>;
+
 	static bool ShouldCompilePermutation(const FGlobalShaderPermutationParameters& Parameters)
 	{
 		return IsFeatureLevelSupported(Parameters.Platform, ERHIFeatureLevel::SM5);
@@ -56,11 +69,44 @@ FScreenPassTexture AddVisualizeLocalExposurePass(FRDGBuilder& GraphBuilder, cons
 {
 	check(Inputs.SceneColor.IsValid());
 	check(Inputs.HDRSceneColor.IsValid());
-	check(Inputs.LumBilateralGridTexture);
-	check(Inputs.BlurredLumTexture);
 	check(Inputs.EyeAdaptationBuffer);
 	check(Inputs.EyeAdaptationParameters);
 
+	RDG_EVENT_SCOPE(GraphBuilder, "VisualizeLocalExposure");
+
+	enum class EVisualizeId : int32
+	{
+		Overview = 0,
+		LocalExposure = 1,
+		Thresholds = 2,
+		BaseLuminance = 3,
+		DetailLuminance = 4,
+		ValidLookup = 5,
+		FusionBase = 6,
+		FusionShadows = 7,
+		FusionHighlights = 8,
+		FusionWeights = 9,
+		MAX,
+	};
+
+	static const TCHAR* kVisualizationName[] = {
+		TEXT(""),
+		TEXT("LocalExposure"),
+		TEXT("Thresholds"),
+		TEXT("BaseLuminance"),
+		TEXT("DetailLuminance"),
+		TEXT("ValidLookup"),
+		TEXT("FusionBase"),
+		TEXT("FusionShadows"),
+		TEXT("FusionHighlights"),
+		TEXT("FusionWeights"),
+	};
+	static_assert(UE_ARRAY_COUNT(kVisualizationName) == int32(EVisualizeId::MAX), "kVisualizationName doesn't match EVisualizeId");
+
+	const EVisualizeId Visualization = EVisualizeId(FMath::Clamp(CVarLocalExposureVisualizationMode.GetValueOnRenderThread(), 0, int32(EVisualizeId::MAX) - 1));
+
+	const bool bIsOverviewVisualize = Visualization == EVisualizeId::Overview;
+
 	FScreenPassRenderTarget Output = Inputs.OverrideOutput;
 
 	if (!Output.IsValid())
@@ -68,33 +114,100 @@ FScreenPassTexture AddVisualizeLocalExposurePass(FRDGBuilder& GraphBuilder, cons
 		Output = FScreenPassRenderTarget::CreateFromInput(GraphBuilder, Inputs.SceneColor, View.GetOverwriteLoadAction(), TEXT("VisualizeLocalExposure"));
 	}
 
-	const FScreenPassTextureViewport InputViewport(Inputs.SceneColor);
-	const FScreenPassTextureViewport OutputViewport(Output);
-
-	FRHISamplerState* BilinearClampSampler = TStaticSamplerState<SF_Bilinear, AM_Clamp, AM_Clamp, AM_Clamp>::GetRHI();
-
-	const FPostProcessSettings& Settings = View.FinalPostProcessSettings;
-
-	auto PassParameters = GraphBuilder.AllocParameters<FVisualizeLocalExposurePS::FParameters>();
-	PassParameters->RenderTargets[0] = Output.GetRenderTargetBinding();
-	PassParameters->View = View.ViewUniformBuffer;
-	PassParameters->EyeAdaptation = *Inputs.EyeAdaptationParameters;
-	PassParameters->LocalExposure = *Inputs.LocalExposureParameters;
-	PassParameters->Input = GetScreenPassTextureViewportParameters(InputViewport);
-	PassParameters->Output = GetScreenPassTextureViewportParameters(OutputViewport);
-	PassParameters->HDRSceneColorTexture = Inputs.HDRSceneColor.Texture;
-	PassParameters->SceneColorTexture = Inputs.SceneColor.Texture;
-	PassParameters->EyeAdaptationBuffer = GraphBuilder.CreateSRV(Inputs.EyeAdaptationBuffer);
-	PassParameters->LumBilateralGrid = Inputs.LumBilateralGridTexture;
-	PassParameters->BlurredLogLum = Inputs.BlurredLumTexture;
-	PassParameters->TextureSampler = BilinearClampSampler;
-	PassParameters->DebugMode = CVarLocalExposureVisualizeDebugMode.GetValueOnRenderThread();
-
-	TShaderMapRef<FVisualizeLocalExposurePS> PixelShader(View.ShaderMap);
-
-	RDG_EVENT_SCOPE(GraphBuilder, "VisualizeLocalExposure");
-
-	AddDrawScreenPass(GraphBuilder, RDG_EVENT_NAME("Visualizer"), View, OutputViewport, InputViewport, PixelShader, PassParameters);
+	auto Visualize = [&](EVisualizeId VisualizeId, FString Label, FScreenPassRenderTarget OverrideOutput = FScreenPassRenderTarget())
+		{
+			FScreenPassRenderTarget TmpOutput = OverrideOutput;
+
+			if (!TmpOutput.IsValid())
+			{
+				FIntPoint TmpOutputExtent = FIntPoint::DivideAndRoundUp(Inputs.SceneColor.ViewRect.Size(), 4);
+				FIntRect TmpOutputViewRect = FIntRect(FIntPoint(0, 0), TmpOutputExtent);
+
+				FRDGTextureDesc TmpOutputDesc = Inputs.SceneColor.Texture->Desc;
+				TmpOutputDesc.Extent = TmpOutputExtent;
+				TmpOutputDesc.Flags |= TexCreate_UAV | TexCreate_RenderTargetable;
+				TmpOutput = FScreenPassRenderTarget(GraphBuilder.CreateTexture(TmpOutputDesc, TEXT("VisualizeLumenScene")), TmpOutputViewRect, ERenderTargetLoadAction::ENoAction);
+			}
+
+			const FScreenPassTextureViewport InputViewport(Inputs.SceneColor);
+			const FScreenPassTextureViewport OutputViewport(TmpOutput);
+
+			const FPostProcessSettings& Settings = View.FinalPostProcessSettings;
+
+			auto PassParameters = GraphBuilder.AllocParameters<FVisualizeLocalExposurePS::FParameters>();
+			PassParameters->RenderTargets[0] = TmpOutput.GetRenderTargetBinding();
+			PassParameters->View = View.ViewUniformBuffer;
+			PassParameters->EyeAdaptation = *Inputs.EyeAdaptationParameters;
+			PassParameters->LocalExposure = *Inputs.LocalExposureParameters;
+			PassParameters->Input = GetScreenPassTextureViewportParameters(InputViewport);
+			PassParameters->Output = GetScreenPassTextureViewportParameters(OutputViewport);
+			PassParameters->HDRSceneColorTexture = Inputs.HDRSceneColor.Texture;
+			PassParameters->SceneColorTexture = Inputs.SceneColor.Texture;
+			PassParameters->EyeAdaptationBuffer = GraphBuilder.CreateSRV(Inputs.EyeAdaptationBuffer);
+			PassParameters->LumBilateralGrid = Inputs.LumBilateralGridTexture;
+			PassParameters->BlurredLogLum = Inputs.BlurredLumTexture;
+
+			if (Inputs.ExposureFusionData != nullptr)
+			{
+				PassParameters->ExposuresTexture = Inputs.ExposureFusionData->Exposures.Texture;
+				PassParameters->WeightsTexture = Inputs.ExposureFusionData->Weights.Texture;
+			}
+
+			PassParameters->BilinearClampSampler = TStaticSamplerState<SF_Bilinear, AM_Clamp, AM_Clamp, AM_Clamp>::GetRHI();
+			PassParameters->DebugMode = uint32(VisualizeId);
+
+			FVisualizeLocalExposurePS::FPermutationDomain PermutationVector;
+			PermutationVector.Set<FVisualizeLocalExposurePS::FExposureFusion>(Inputs.ExposureFusionData != nullptr);
+			auto PixelShader = View.ShaderMap->GetShader<FVisualizeLocalExposurePS>(PermutationVector);
+
+			FRDGEventName PassName = RDG_EVENT_NAME("LocalExposure Visualize(%s) %dx%d", kVisualizationName[int32(VisualizeId)], OutputViewport.Rect.Width(), OutputViewport.Rect.Height());
+			AddDrawScreenPass(GraphBuilder, MoveTemp(PassName), View, OutputViewport, InputViewport, PixelShader, PassParameters);
+
+			FVisualizeBufferTile Tile;
+			Tile.Input = TmpOutput;
+			Tile.Label = Label;
+			return Tile;
+		};
+
+	if (bIsOverviewVisualize)
+	{
+		TArray<FVisualizeBufferTile> Tiles;
+		Tiles.SetNum(16);
+		if (Visualization == EVisualizeId::Overview)
+		{
+			Tiles[4 * 0 + 0] = Visualize(EVisualizeId::LocalExposure, TEXT("Local Exposure"));
+			Tiles[4 * 0 + 3] = Visualize(EVisualizeId::Thresholds, TEXT("Thresholds"));
+
+			if (Inputs.ExposureFusionData != nullptr)
+			{
+				Tiles[4 * 3 + 0] = Visualize(EVisualizeId::FusionShadows, TEXT("Shadows"));
+				Tiles[4 * 3 + 1] = Visualize(EVisualizeId::FusionBase, TEXT("Base"));
+				Tiles[4 * 3 + 2] = Visualize(EVisualizeId::FusionHighlights, TEXT("Highlights"));
+				Tiles[4 * 3 + 3] = Visualize(EVisualizeId::FusionWeights, TEXT("Weights"));
+			}
+			else
+			{
+				Tiles[4 * 3 + 0] = Visualize(EVisualizeId::BaseLuminance, TEXT("Base Luminance"));
+				Tiles[4 * 3 + 1] = Visualize(EVisualizeId::DetailLuminance, TEXT("Detail Luminance"));
+			}
+		}
+		else
+		{
+			unimplemented();
+		}
+
+		{
+			FVisualizeBufferInputs VisualizeBufferInputs;
+			VisualizeBufferInputs.OverrideOutput = Output;
+			VisualizeBufferInputs.SceneColor = FScreenPassTexture::CopyFromSlice(GraphBuilder, FScreenPassTextureSlice::CreateFromScreenPassTexture(GraphBuilder, Inputs.SceneColor));
+			VisualizeBufferInputs.Tiles = Tiles;
+			AddVisualizeBufferPass(GraphBuilder, View, VisualizeBufferInputs);
+		}
+	}
+	else
+	{
+		Visualize(Visualization, TEXT(""), Output);
+	}
 
 	return MoveTemp(Output);
 }
\ No newline at end of file
diff --git a/Engine/Source/Runtime/Renderer/Private/PostProcess/PostProcessVisualizeLocalExposure.h b/Engine/Source/Runtime/Renderer/Private/PostProcess/PostProcessVisualizeLocalExposure.h
--- a/Engine/Source/Runtime/Renderer/Private/PostProcess/PostProcessVisualizeLocalExposure.h
+++ b/Engine/Source/Runtime/Renderer/Private/PostProcess/PostProcessVisualizeLocalExposure.h
@@ -7,6 +7,8 @@
 class FEyeAdaptationParameters;
 class FLocalExposureParameters;
 
+struct FExposureFusionData;
+
 struct FVisualizeLocalExposureInputs
 {
 	// [Optional] Render to the specified output. If invalid, a new texture is created and returned.
@@ -17,11 +19,13 @@ struct FVisualizeLocalExposureInputs
 
 	FRDGBufferRef EyeAdaptationBuffer = nullptr;
 
+	const FEyeAdaptationParameters* EyeAdaptationParameters = nullptr;
+	const FLocalExposureParameters* LocalExposureParameters = nullptr;
+
 	FRDGTextureRef LumBilateralGridTexture = nullptr;
 	FRDGTextureRef BlurredLumTexture = nullptr;
-	const FLocalExposureParameters* LocalExposureParameters = nullptr;
 
-	const FEyeAdaptationParameters* EyeAdaptationParameters = nullptr;
+	FExposureFusionData* ExposureFusionData = nullptr;
 };
 
 FScreenPassTexture AddVisualizeLocalExposurePass(FRDGBuilder& GraphBuilder, const FViewInfo& View, const FVisualizeLocalExposureInputs& Inputs);
\ No newline at end of file
diff --git a/Engine/Source/Runtime/Renderer/Private/PostProcess/PostProcessing.cpp b/Engine/Source/Runtime/Renderer/Private/PostProcess/PostProcessing.cpp
--- a/Engine/Source/Runtime/Renderer/Private/PostProcess/PostProcessing.cpp
+++ b/Engine/Source/Runtime/Renderer/Private/PostProcess/PostProcessing.cpp
@@ -373,7 +373,7 @@ void AddPostProcessingPasses(
 
 	FRDGTextureRef LocalExposureBilateralGridTexture = nullptr;
 	FRDGTextureRef LocalExposureBlurredLogLumTexture = BlackDummy.Texture;
-	FScreenPassTexture ExposureFusion;
+	FExposureFusionData ExposureFusionData;
 
 	FVisualizeTemporalUpscalerInputs VisualizeTemporalUpscalerInputs;
 
@@ -1177,7 +1177,7 @@ void AddPostProcessingPasses(
 
 		if (bLocalExposureEnabled && View.FinalPostProcessSettings.LocalExposureMethod == ELocalExposureMethod::Fusion)
 		{
-			ExposureFusion = AddLocalExposureFusionPass(
+			ExposureFusionData = AddLocalExposureFusionPass(
 				GraphBuilder, View,
 				EyeAdaptationParameters,
 				EyeAdaptationBuffer,
@@ -1331,7 +1331,7 @@ void AddPostProcessingPasses(
 				PassInputs.SceneColorApplyParamaters = SceneColorApplyParameters;
 				PassInputs.LocalExposureBilateralGridTexture = LocalExposureBilateralGridTexture;
 				PassInputs.BlurredLogLuminanceTexture = LocalExposureBlurredLogLumTexture;
-				PassInputs.ExposureFusion = FScreenPassTextureSlice::CreateFromScreenPassTexture(GraphBuilder, ExposureFusion);
+				PassInputs.ExposureFusion = FScreenPassTextureSlice::CreateFromScreenPassTexture(GraphBuilder, ExposureFusionData.Result);
 				PassInputs.LocalExposureParameters = &LocalExposureParameters;
 				PassInputs.EyeAdaptationParameters = &EyeAdaptationParameters;
 				PassInputs.EyeAdaptationBuffer = EyeAdaptationBuffer;
@@ -1710,6 +1710,7 @@ void AddPostProcessingPasses(
 		PassInputs.HDRSceneColor = FScreenPassTexture::CopyFromSlice(GraphBuilder, SceneColorBeforeTonemapSlice);
 		PassInputs.LumBilateralGridTexture = LocalExposureBilateralGridTexture;
 		PassInputs.BlurredLumTexture = LocalExposureBlurredLogLumTexture;
+		PassInputs.ExposureFusionData = View.FinalPostProcessSettings.LocalExposureMethod == ELocalExposureMethod::Fusion ? &ExposureFusionData : nullptr;
 		PassInputs.LocalExposureParameters = &LocalExposureParameters;
 		PassInputs.EyeAdaptationBuffer = GetEyeAdaptationBuffer(GraphBuilder, View);
 		PassInputs.EyeAdaptationParameters = &EyeAdaptationParameters;
