Workaround race condition in deferred decals with async RDG. The deferred decal pass accesses the decal proxy which is modified by render commands inside of FScene. The short term fix is to remove the async tag and run awaited by the RDG builder. A longer term fix is to rework the visible decal system to avoid requiring access to those proxies inside of the RDG lambda.

diff --git a/Engine/Source/Runtime/Renderer/Private/CompositionLighting/PostProcessDeferredDecals.cpp b/Engine/Source/Runtime/Renderer/Private/CompositionLighting/PostProcessDeferredDecals.cpp
--- a/Engine/Source/Runtime/Renderer/Private/CompositionLighting/PostProcessDeferredDecals.cpp
+++ b/Engine/Source/Runtime/Renderer/Private/CompositionLighting/PostProcessDeferredDecals.cpp
@@ -700,7 +700,7 @@ void AddDeferredDecalPass(
 			RDG_EVENT_NAME("Batch [%d, %d]", DecalIndexBegin, DecalIndexEnd - 1),
 			PassParameters,
 			ERDGPassFlags::Raster,
-			[&View, FeatureLevel, ShaderPlatform, DecalIndexBegin, DecalIndexEnd, SortedDecals, DecalRenderStage, RenderTargetMode, bStencilSizeThreshold, bShaderComplexity](FRDGAsyncTask, FRHICommandList& RHICmdList)
+			[&View, FeatureLevel, ShaderPlatform, DecalIndexBegin, DecalIndexEnd, SortedDecals, DecalRenderStage, RenderTargetMode, bStencilSizeThreshold, bShaderComplexity](FRHICommandList& RHICmdList)
 		{
 			RHICmdList.SetViewport(View.ViewRect.Min.X, View.ViewRect.Min.Y, 0.0f, View.ViewRect.Max.X, View.ViewRect.Max.Y, 1.0f);
 
