Substrate - rework pixel merial debug to work on all platforms.
BAsically it is written into a buffer and brought back to CPU using readback buffers.
Then canvas is used to display.
Tested as working on dx11, dx12 and Vulkan. Matal to be tested.

#tests Tested editor PC, Vulkan, DX11 and dx12.
[RN] Substrate GBuffer material visuaolisation is now working on Vulkan and Dx11.

diff --git a/Engine/Shaders/Private/RayTracing/RayTracingDebug.usf b/Engine/Shaders/Private/RayTracing/RayTracingDebug.usf
--- a/Engine/Shaders/Private/RayTracing/RayTracingDebug.usf
+++ b/Engine/Shaders/Private/RayTracing/RayTracingDebug.usf
@@ -16,6 +16,7 @@
 #include "../ShaderPrint.ush"
 #include "../Substrate/Substrate.ush"
 #include "../Substrate/SubstrateVisualizeCommon.ush"
+#include "/Engine/Shared/SubstrateVisualizeDefinitions.h"
 
 #include "RayTracingDebugUtils.ush"
 #include "RayTracingLightCullingCommon.ush"
@@ -140,6 +141,9 @@ float4 DebugRayTracingMaterial(uint2 PixelCoord, FRayDesc Ray, bool bIsFarField,
 				const float3 V = -normalize(TranslatedWorldPosition - PrimaryView.TranslatedWorldCameraOrigin);
 				const uint ClosureIndex = 0;
 				SubstratePrintMaterialProperties(Ctx, PixelCoord, TranslatedWorldPosition, V, ClosureIndex, Payload.SubstrateData);
+
+				SerializeSubstratePixelDebugDataEntry(PixelCoord, Payload.SubstrateData);
+
 			}
 			bWriteOuput = false;
 			return float4(0, 0, 0, 1);
diff --git a/Engine/Shaders/Private/Substrate/Substrate.ush b/Engine/Shaders/Private/Substrate/Substrate.ush
--- a/Engine/Shaders/Private/Substrate/Substrate.ush
+++ b/Engine/Shaders/Private/Substrate/Substrate.ush
@@ -366,15 +366,6 @@ struct FSubstrateBSDF
 #define HEADER_SETCOMPLEXSPECIALPATH(X, NewValue)	WRITE_BITS(X, HEADER_BIT_COUNT_COMPLEXPATH_SPECIAL, HEADER_BIT_OFFSET_COMPLEXPATH_SPECIAL, NewValue)
 #define HEADER_COMPLEXSPECIALPATH_MASK				TO_OFFSET_BITMASK(HEADER_BIT_COUNT_COMPLEXPATH_SPECIAL, HEADER_BIT_OFFSET_COMPLEXPATH_SPECIAL)
 
-#define HEADER_MATERIALMODE_NONE				0
-#define HEADER_MATERIALMODE_SLAB_SIMPLE			1
-#define HEADER_MATERIALMODE_SLAB_SINGLE			2
-#define HEADER_MATERIALMODE_SLAB_COMPLEX		3
-#define HEADER_MATERIALMODE_SLWATER				4
-#define HEADER_MATERIALMODE_HAIR				5
-#define HEADER_MATERIALMODE_EYE					6
-// Up to 8 material modes: 1 remains.
-
 ///////////////////////////////////////////////////////////////////////////////
 // BSDF State
  
@@ -413,17 +404,6 @@ struct FSubstrateBSDF
 #define BSDF_BIT_OFFSET_HASSPECPROFILE			(BSDF_BIT_OFFSET_HASGLINT			+ BSDF_BIT_COUNT_HASGLINT)
 #define BSDF_BIT_OFFSET_WEIGHT10F				22
 
-// Special optmised modes for the single BSDF path.
-#define HEADER_SINGLE_OPTLEGACYMODE_BIT_COUNT	3
-#define HEADER_SINGLE_OPTLEGACYMODE_BIT_MASK	((1u<<HEADER_SINGLE_OPTLEGACYMODE_BIT_COUNT)-1u)
-#define SINGLE_OPTLEGACYMODE_NONE				0
-#define SINGLE_OPTLEGACYMODE_CLEARCOAT			1
-#define SINGLE_OPTLEGACYMODE_CLOTH				2
-#define SINGLE_OPTLEGACYMODE_SSSWRAP			3
-#define SINGLE_OPTLEGACYMODE_SSSPROFILE			4
-#define SINGLE_OPTLEGACYMODE_SSSWRAP_THIN		5
-// 5 more remains: Also how to deal with Aniso?
-
 // Layout 0 (Fast)
 #define BSDF_FASTENCODING_BIT_COUNT				(0u)
 #define BSDF_FASTENCODING_MASK					(0u)
diff --git a/Engine/Shaders/Private/Substrate/SubstrateSubsurface.ush b/Engine/Shaders/Private/Substrate/SubstrateSubsurface.ush
--- a/Engine/Shaders/Private/Substrate/SubstrateSubsurface.ush
+++ b/Engine/Shaders/Private/Substrate/SubstrateSubsurface.ush
@@ -3,6 +3,7 @@
 #pragma once
 
 #include "../TransmissionCommon.ush"
+#include "/Engine/Shared/SubstrateDefinitions.h"
 
 // Forward declarations
 uint PackColorLinearToGamma2AlphaLinear(float4 In);
@@ -28,14 +29,6 @@ struct FSubstrateSubsurfaceHeader
 	uint Bytes;
 };
 
-// Sub-surface scattering types
-#define SSS_TYPE_INVALID						(0u)
-#define SSS_TYPE_WRAP							(1u)
-#define SSS_TYPE_DIFFUSION						(2u)
-#define SSS_TYPE_DIFFUSION_PROFILE				(3u)
-#define SSS_TYPE_SIMPLEVOLUME					(4u)
-#define SSS_TYPE_COUNT							(5u)
-
 // Note: SSSHeader only stores INVALID, WRAP, DIFFUSION, DIFFUSION_PROFILE type. This is because the type is represented onto 2bits and 
 //       other SSS types don't have/need side payload. This means SIMPLEVOLUME is trunked.
 #if SSS_TYPE_COUNT != 5
diff --git a/Engine/Shaders/Private/Substrate/SubstrateVisualize.usf b/Engine/Shaders/Private/Substrate/SubstrateVisualize.usf
--- a/Engine/Shaders/Private/Substrate/SubstrateVisualize.usf
+++ b/Engine/Shaders/Private/Substrate/SubstrateVisualize.usf
@@ -15,6 +15,7 @@
 #define SUBSTRATE_COMPLEXSPECIALPATH 1				// Allow to read data from complex special path
 #include "/Engine/Private/Substrate/SubstrateEvaluation.ush"
 #include "/Engine/Private/Substrate/SubstrateTile.ush"
+#include "/Engine/Shared/SubstrateVisualizeDefinitions.h"
 
 #if SUBSTRATE_INLINE_SHADING == 0
 #include "/Engine/Private/Substrate/SubstrateVisualizeCommon.ush"
@@ -35,9 +36,6 @@ int2 GetCursorPosition()
 
 #if SHADER_MATERIALPRINT
 
-uint ClosureIndex;
-RWBuffer<uint> RWPositionOffsetBuffer;
-
 float3 GetWorldPositionFromPixelPos(uint2 PixelPos)
 {
 	const float2 BufferUV = SvPositionToBufferUV(float4(PixelPos + 0.5f, 0, 0));
@@ -55,24 +53,23 @@ void MaterialPrintInfoCS(uint3 DispatchThreadId : SV_DispatchThreadID)
 	const float3 WorldPosition = GetWorldPositionFromPixelPos(PixelPos);
 	const float3 V = -normalize(WorldPosition - PrimaryView.TranslatedWorldCameraOrigin);
 
-	FShaderPrintContext Context;
-	if (ClosureIndex > 0)
-	{
-		float2 PrintOffset;
-		PrintOffset.x = asfloat(RWPositionOffsetBuffer[0]);
-		PrintOffset.y = asfloat(RWPositionOffsetBuffer[1]);
-		Context = InitShaderPrintContext(true, PrintOffset);
-	}
-	else
-	{
-		Context = InitShaderPrintContext(true, uint2(50, 50));
-	}
+	FShaderPrintContext Context = InitShaderPrintContext(true, uint2(50, 50));
 
-	// Print pixel material properties
-	SubstratePrintMaterialProperties(Context, PixelPos, WorldPosition, V, ClosureIndex);
-
-	RWPositionOffsetBuffer[0] = asuint(Context.Pos.x);
-	RWPositionOffsetBuffer[1] = asuint(Context.Pos.y);
+	// Print pixel material properties. Simple print for simple compilers such as dx11/fxc which cannot handle loop...
+#if SUBSTRATE_MAX_CLOSURE_COUNT!= 8
+#error Unsupported SUBSTRATE_MAX_CLOSURE_COUNT in visualisation shader.
+#endif
+	SubstratePrintMaterialProperties(Context, PixelPos, WorldPosition, V, 0);
+	SubstratePrintMaterialProperties(Context, PixelPos, WorldPosition, V, 1);
+	SubstratePrintMaterialProperties(Context, PixelPos, WorldPosition, V, 2);
+	SubstratePrintMaterialProperties(Context, PixelPos, WorldPosition, V, 3);
+	SubstratePrintMaterialProperties(Context, PixelPos, WorldPosition, V, 4);
+	SubstratePrintMaterialProperties(Context, PixelPos, WorldPosition, V, 5);
+	SubstratePrintMaterialProperties(Context, PixelPos, WorldPosition, V, 6);
+	SubstratePrintMaterialProperties(Context, PixelPos, WorldPosition, V, 7);
+
+	// Serialise out to buffer for print.
+	SerializeSubstratePixelDebugDataEntry(PixelPos);
 
 	// Print pixel footprint
 	const float3 WorldPositionX = GetWorldPositionFromPixelPos(PixelPos + uint2(1, 0));
diff --git a/Engine/Shaders/Private/Substrate/SubstrateVisualizeCommon.ush b/Engine/Shaders/Private/Substrate/SubstrateVisualizeCommon.ush
--- a/Engine/Shaders/Private/Substrate/SubstrateVisualizeCommon.ush
+++ b/Engine/Shaders/Private/Substrate/SubstrateVisualizeCommon.ush
@@ -108,13 +108,6 @@ void PrintBackgroundRect(inout FShaderPrintContext Ctx, float2 StartRect, float2
 		RectBackgroundColor);
 }
 
-void Newline(inout FShaderPrintContext Context, inout float2 MaxRect)
-{
-	MaxRect.x = max(MaxRect.x, Context.Pos.x);
-	MaxRect.y = Context.Pos.y;
-	Newline(Context);
-}
-
 void PrintAddress(inout FShaderPrintContext Context, inout FSubstrateAddressing SubstrateAddressing)
 {
 	FFontColor Font;
@@ -134,14 +127,6 @@ void SubstratePrintBSDF(
 	inout float2 RectMax, 
 	FSubstrateRaytracingPayload Payload)
 {
-	const FFontColor FontBSDF         = FontSilver;
-	const FFontColor FontBSDFType     = FontEmerald;
-	const FFontColor FontBSDFStateName= FontWhite;
-	const FFontColor FontBSDFPropName = FontOrange;
-	const FFontColor FontBSDFPropValu = FontWhite;
-
-	Newline(Context, RectMax);
-
 	#if SUBSTRATE_MATERIALCONTAINER_IS_VIEWRESOURCE
 	const FSubstrateBSDF BSDF = UnpackSubstrateBSDFIn(Substrate.MaterialTextureArray, SubstrateAddressing, Header);
 	#else
@@ -167,276 +152,16 @@ void SubstratePrintBSDF(
 	{
 		case SUBSTRATE_BSDF_TYPE_SLAB:
 		{
-			const bool bWeightL = BSDF_GETHASTRANSABOVE(BSDF);
-			const bool bGreyWeightV = BSDF_GETHASGREYWEIGHT_V(BSDF);
-			const bool bGreyTopTrans = BSDF.TransmittanceAboveAlongN.x == BSDF.TransmittanceAboveAlongN.y && BSDF.TransmittanceAboveAlongN.x == BSDF.TransmittanceAboveAlongN.z;
-
-			Print(Context, TEXT("Slab "), FontBSDFType);
-			PrintAddress(Context, SubstrateAddressing);
-			Newline(Context, RectMax);
-			Print(Context, TEXT(" NormalID BasisType Aniso TopLayer Scattering IsThin WeightV "), FontBSDFStateName);
-			if (!bGreyWeightV)
-			{
-				Print(Context, TEXT("            "), FontBSDFStateName);
-			}
-			if (bWeightL)
-			{
-				Print(Context, TEXT("TransToL     "), FontBSDFStateName);
-				if (!bGreyTopTrans)
-				{
-					Print(Context, TEXT("       "), FontBSDFStateName);
-				}
-				Print(Context, TEXT("CoverToL"), FontBSDFStateName);
-			}
-			Newline(Context, RectMax);
-
-			const uint BasisType  = SubstrateGetSharedLocalBasisType(HEADER_GETSHAREDLOCALBASESTYPE(Header.PackedHeader), BSDF_GETSHAREDLOCALBASISID(BSDF));
-
-			Print(Context, TEXT(" "), FontSilver);
-			Print(Context, BSDF_GETSHAREDLOCALBASISID(BSDF), FontBSDF, 2, 1);	Print(Context, TEXT("       "), FontSilver);
-			if (BasisType == 0)
-			{
-				Print(Context, TEXT("Normal"), FontBSDF);	Print(Context, TEXT("    "), FontSilver);
-			}
-			else
-			{
-				Print(Context, TEXT("Tangent"), FontBSDF);	Print(Context, TEXT("   "), FontSilver);
-			}
-			PrintBool(Context, BSDF_GETHASANISOTROPY(BSDF));					Print(Context, TEXT("  "));
-			PrintBool(Context, BSDF_GETISTOPLAYER(BSDF));						Print(Context, TEXT("     "));
-			PrintBool(Context, BSDF_GETSSSTYPE(BSDF) != SSS_TYPE_INVALID);		Print(Context, TEXT("       "));
-			PrintBool(Context, BSDF_GETISTHIN(BSDF));							Print(Context, TEXT("   "));
-
-			// View Weight
-			{
-				if (bGreyWeightV)
-				{
-					Print(Context, BSDF.LuminanceWeightV.x, FontSilver, 6, 4);
-				}
-				else
-				{
-					Print(Context, BSDF.LuminanceWeightV.r, FontLightRed,   6, 4);
-					Print(Context, BSDF.LuminanceWeightV.g, FontLightGreen, 6, 4);
-					Print(Context, BSDF.LuminanceWeightV.b, FontLightBlue,  6, 4);
-				}
-				Print(Context, TEXT(" "), FontWhite);
-			}
-
-			// Light Weight
-			if (bWeightL)
-			{
-				Print(Context, TEXT(" "), FontSilver);
-				if (bGreyTopTrans)
-				{
-					Print(Context, BSDF.TransmittanceAboveAlongN.x, FontSilver);
-				}
-				else	
-				{
-					Print(Context, BSDF.TransmittanceAboveAlongN.r, FontLightRed,   6, 4);
-					Print(Context, BSDF.TransmittanceAboveAlongN.g, FontLightGreen, 6, 4);
-					Print(Context, BSDF.TransmittanceAboveAlongN.b, FontLightBlue,  6, 4);
-				}
-
-				Print(Context, TEXT(" "), FontSilver);
-				Print(Context, BSDF.CoverageAboveAlongN, FontSilver, 6, 4);
-			}
-			Newline(Context, RectMax);
-
-			Print(Context, TEXT(" Diffuse     "), FontBSDFPropName);	Print(Context, SLAB_DIFFUSEALBEDO(BSDF), FontBSDFPropValu);			Newline(Context, RectMax);
-
-			const bool bGreyF0 = SLAB_F0(BSDF).x == SLAB_F0(BSDF).y && SLAB_F0(BSDF).x == SLAB_F0(BSDF).z;
-			if (bGreyF0)
-			{
-				Print(Context, TEXT(" F0          "), FontBSDFPropName);	Print(Context, SLAB_F0(BSDF).x, FontBSDFPropValu);				Newline(Context, RectMax);
-			}
-			else
-			{
-				Print(Context, TEXT(" F0          "), FontBSDFPropName);	Print(Context, SLAB_F0(BSDF), FontBSDFPropValu);				Newline(Context, RectMax);
-			}
-
-
-			if (BSDF_GETHASF90(BSDF))
-			{
-				const bool bGreyF90 = SLAB_F90(BSDF).x == SLAB_F90(BSDF).y && SLAB_F90(BSDF).x == SLAB_F90(BSDF).z;
-				if (bGreyF90)
-				{
-					Print(Context, TEXT(" F90         "), FontBSDFPropName);	Print(Context, SLAB_F90(BSDF).x, FontBSDFPropValu);			Newline(Context, RectMax);
-				}
-				else
-				{
-					Print(Context, TEXT(" F90         "), FontBSDFPropName);	Print(Context, SLAB_F90(BSDF), FontBSDFPropValu);			Newline(Context, RectMax);
-				}
-			}
-
-			Print(Context, TEXT(" Roughness   "), FontBSDFPropName);	Print(Context, SLAB_ROUGHNESS(BSDF), FontBSDFPropValu);				Newline(Context, RectMax);
-
-			if (BSDF_GETHASANISOTROPY(BSDF))
-			{
-				Print(Context, TEXT(" Anisotropy  "), FontBSDFPropName);	Print(Context, SLAB_ANISOTROPY(BSDF), FontBSDFPropValu);		Newline(Context, RectMax);
-			}
-
-			if (BSDF_GETHASHAZINESS(BSDF))
-			{
-				FHaziness Haziness = UnpackHaziness(SLAB_HAZINESS(BSDF));
-
-				Print(Context, TEXT(" Haziness    "),			FontBSDFPropName);																		Newline(Context, RectMax);
-				Print(Context, TEXT("  - Roughness       "),	FontBSDFPropName);	Print(Context, Haziness.Roughness,				FontBSDFPropValu);	Newline(Context, RectMax);
-				Print(Context, TEXT("  - Weight          "),	FontBSDFPropName);	Print(Context, Haziness.Weight,					FontBSDFPropValu);	Newline(Context, RectMax);
-				Print(Context, TEXT("  - SimpleClearCoat "),	FontBSDFPropName);	Print(Context, Haziness.bSimpleClearCoat?1u:0u,	FontBSDFPropValu);	Newline(Context, RectMax);
-			}
-
-			if (BSDF_GETSSSTYPE(BSDF) != SSS_TYPE_INVALID || BSDF_GETISTHIN(BSDF))
-			{
-				const bool bIsThin = BSDF_GETISTHIN(BSDF);
-				Print(Context, TEXT(" SSS Type    "), FontBSDFPropName);	PrintSSSType(Context, BSDF_GETSSSTYPE(BSDF), BSDF_GETISTHIN(BSDF), FontYellow);	Newline(Context, RectMax);
-				if (BSDF_GETSSSTYPE(BSDF) == SSS_TYPE_WRAP)
-				{
-					const float Opacity = SubstrateSubSurfaceHeaderGetWrapOpacity(SSSHeader);
-					Print(Context, TEXT(" SSS Opacity "), FontBSDFPropName);	Print(Context, Opacity, FontBSDFPropValu);						Newline(Context, RectMax);
-					Print(Context, TEXT(" SSS MFP     "), FontBSDFPropName);	Print(Context, SLAB_SSSMFP(BSDF), FontBSDFPropValu);			Newline(Context, RectMax);
-					Print(Context, TEXT(" SSS Phase   "), FontBSDFPropName);	Print(Context, SLAB_SSSPHASEANISOTROPY(BSDF), FontBSDFPropValu);Newline(Context, RectMax);
-					Print(Context, TEXT(" SSS Thicknes"), FontBSDFPropName);	Print(Context, BSDF_GETTHICKNESSCM(BSDF), FontBSDFPropValu);	Newline(Context, RectMax);
-				}
-				else if (BSDF_GETSSSTYPE(BSDF) == SSS_TYPE_SIMPLEVOLUME)
-				{
-					const float Opacity = SubstrateSubSurfaceHeaderGetWrapOpacity(SSSHeader);
-					Print(Context, TEXT(" SSS Opacity "), FontBSDFPropName);	Print(Context, Opacity, FontBSDFPropValu);						Newline(Context, RectMax);
-					Print(Context, TEXT(" SSS MFP     "), FontBSDFPropName);	Print(Context, SLAB_SSSMFP(BSDF), FontBSDFPropValu);			Newline(Context, RectMax);
-					Print(Context, TEXT(" SSS Phase   "), FontBSDFPropName);	Print(Context, SLAB_SSSPHASEANISOTROPY(BSDF), FontBSDFPropValu);Newline(Context, RectMax);
-					Print(Context, TEXT(" SSS Thicknes"), FontBSDFPropName);	Print(Context, BSDF_GETTHICKNESSCM(BSDF), FontBSDFPropValu);	Newline(Context, RectMax);
-				}
-				else if (BSDF_GETSSSTYPE(BSDF) == SSS_TYPE_DIFFUSION_PROFILE)
-				{
-					const uint ProfileId = SubstrateSubSurfaceHeaderGetProfileId(SSSHeader);
-					const float RadiusScale = SubstrateSubSurfaceHeaderGetProfileRadiusScale(SSSHeader);
-					const float3 DiffuseMFP = GetSubsurfaceProfileMFPInCm(ProfileId).xyz * RadiusScale;
-					const float Phase = GetTransmissionProfileParams(ProfileId).ScatteringDistribution;
-
-					Print(Context, TEXT(" SSS ID.     "), FontBSDFPropName);	Print(Context, ProfileId, FontBSDFPropValu);	Newline(Context, RectMax);
-					Print(Context, TEXT(" SSS Radius  "), FontBSDFPropName);	Print(Context, RadiusScale, FontBSDFPropValu);	Newline(Context, RectMax);
-					Print(Context, TEXT(" SSS MFP     "), FontBSDFPropName);	Print(Context, DiffuseMFP, FontBSDFPropValu);	Newline(Context, RectMax);
-					Print(Context, TEXT(" SSS Phase   "), FontBSDFPropName);	Print(Context, Phase, FontBSDFPropValu);		Newline(Context, RectMax);
-					if (bIsThin)
-					{
-					Print(Context, TEXT(" SSS Thicknes"), FontBSDFPropName);	Print(Context, SLAB_SSSPROFILETHICKNESSCM(BSDF), FontBSDFPropValu);		Newline(Context, RectMax);
-					}
-				}
-				else if (BSDF_GETSSSTYPE(BSDF) == SSS_TYPE_DIFFUSION)
-				{
-					const float3 OriginalMFP = SubstrateSubSurfaceHeaderGetMFP(SSSHeader);
-					const float3 RescaledMFP = SLAB_SSSMFP(BSDF);
-
-					if (bIsThin)
-					{
-						Print(Context, TEXT(" SSS MFP     "), FontBSDFPropName);	Print(Context, OriginalMFP, FontBSDFPropValu);				Newline(Context, RectMax);
-						Print(Context, TEXT(" SSS MFP Norm"), FontBSDFPropName);	Print(Context, RescaledMFP, FontBSDFPropValu);				Newline(Context, RectMax);
-						Print(Context, TEXT(" SSS Thicknes"), FontBSDFPropName);	Print(Context, BSDF_GETTHICKNESSCM(BSDF), FontBSDFPropValu);Newline(Context, RectMax);
-					}
-					else
-					{
-						Print(Context, TEXT(" SSS MFP     "), FontBSDFPropName);	Print(Context, OriginalMFP, FontBSDFPropValu);				Newline(Context, RectMax);
-					}
-
-					Print(Context, TEXT(" SSS Phase   "), FontBSDFPropName);	Print(Context, SLAB_SSSPHASEANISOTROPY(BSDF), FontBSDFPropValu);Newline(Context, RectMax);
-				}
-			} // SSS
-
-			if (BSDF_GETHASFUZZ(BSDF))
-			{
-				Print(Context, TEXT(" Fuzz Amount "), FontBSDFPropName);	Print(Context, SLAB_FUZZ_AMOUNT(BSDF), FontBSDFPropValu);				Newline(Context, RectMax);
-				Print(Context, TEXT(" Fuzz Color  "), FontBSDFPropName);	Print(Context, SLAB_FUZZ_COLOR(BSDF), FontBSDFPropValu);				Newline(Context, RectMax);
-				Print(Context, TEXT(" Fuzz Rough. "), FontBSDFPropName);	Print(Context, SLAB_FUZZ_ROUGHNESS(BSDF), FontBSDFPropValu);			Newline(Context, RectMax);
-			}
-		#if SUBSTRATE_COMPLEXSPECIALPATH
-			if (BSDF_GETHASGLINT(BSDF))
-			{
-				Print(Context, TEXT(" Glint Value "), FontBSDFPropName);	Print(Context, SLAB_GLINT_VALUE(BSDF), FontBSDFPropValu);				Newline(Context, RectMax);
-				Print(Context, TEXT(" Glint UVs   "), FontBSDFPropName);	Print(Context, SLAB_GLINT_UV(BSDF), FontBSDFPropValu);					Newline(Context, RectMax);
-				#if 1
-				Print(Context, TEXT(" Glint UVsDx "), FontBSDFPropName);	Print(Context, SLAB_GLINT_UVDDX(BSDF), FontBSDFPropValu);				Newline(Context, RectMax);
-				Print(Context, TEXT(" Glint UVsDy "), FontBSDFPropName);	Print(Context, SLAB_GLINT_UVDDY(BSDF), FontBSDFPropValu);				Newline(Context, RectMax);
-				#endif
-			}
-			if (BSDF_GETHASSPECPROFILE(BSDF))
-			{
-				Print(Context, TEXT(" SpecID      "), FontBSDFPropName);	Print(Context, GetSpecularProfileId(SLAB_SPECPROFILEID(BSDF)), FontBSDFPropValu);Newline(Context, RectMax);
-				if (GetSpecularProfileParameterization(SLAB_SPECPROFILEID(BSDF)) == 0)
-				{
-				Print(Context, TEXT(" Spec Param. "), FontBSDFPropName);	Print(Context, TEXT("View/Light angles"), FontBSDFPropValu);			Newline(Context, RectMax);
-				}
-				else
-				{
-				Print(Context, TEXT(" Spec Param. "), FontBSDFPropName);	Print(Context, TEXT("Half angles"), FontBSDFPropValu);					Newline(Context, RectMax);
-				}
-			}
-		#endif
 		}
 		break;
 
 		case SUBSTRATE_BSDF_TYPE_HAIR:
 		{
-			Print(Context, TEXT("Hair"), FontBSDFType);
-			PrintAddress(Context, SubstrateAddressing);
-			Newline(Context, RectMax);
-			Print(Context, TEXT(" NormalID Aniso IsTopLayer LuminanceWeight"), FontBSDFStateName);
-			Newline(Context, RectMax);
-
-			Print(Context, TEXT(" "), FontSilver);
-			Print(Context, BSDF_GETSHAREDLOCALBASISID(BSDF), FontBSDF);
-			PrintBool(Context, BSDF_GETHASANISOTROPY(BSDF));
-			PrintBool(Context, BSDF_GETISTOPLAYER(BSDF));
-			if (BSDF_GETHASGREYWEIGHT_V(BSDF))
-			{
-				Print(Context, BSDF.LuminanceWeightV.x, FontSilver);
-			}
-			else
-			{
-				Print(Context, BSDF.LuminanceWeightV.r, FontLightRed);
-				Print(Context, BSDF.LuminanceWeightV.g, FontLightGreen);
-				Print(Context, BSDF.LuminanceWeightV.b, FontLightBlue);
-			}
-			Newline(Context, RectMax);
-
-			Print(Context, TEXT(" BaseColor         "), FontBSDFPropName);	Print(Context, HAIR_BASECOLOR(BSDF), FontBSDFPropValu);				Newline(Context, RectMax);
-			Print(Context, TEXT(" Specular          "), FontBSDFPropName);	Print(Context, HAIR_SPECULAR(BSDF), FontBSDFPropValu);				Newline(Context, RectMax);
-			Print(Context, TEXT(" Roughness         "), FontBSDFPropName);	Print(Context, HAIR_ROUGHNESS(BSDF), FontBSDFPropValu);				Newline(Context, RectMax);
-			Print(Context, TEXT(" Scatter           "), FontBSDFPropName);	Print(Context, HAIR_SCATTER(BSDF), FontBSDFPropValu);				Newline(Context, RectMax);
-			Print(Context, TEXT(" Backlit           "), FontBSDFPropName);	Print(Context, HAIR_BACKLIT(BSDF), FontBSDFPropValu);				Newline(Context, RectMax);
-			Print(Context, TEXT(" Has Transmittance "), FontBSDFPropName);	Print(Context, HAIR_COMPLEXTRANSMITTANCE(BSDF), FontBSDFPropValu);	Newline(Context, RectMax);
 		}
 		break;
 
 		case SUBSTRATE_BSDF_TYPE_EYE:
 		{
-			Print(Context, TEXT("Eye"), FontBSDFType);
-			PrintAddress(Context, SubstrateAddressing);
-			Newline(Context, RectMax);
-			Print(Context, TEXT(" NormalID IsTopLayer LuminanceWeight"), FontBSDFStateName);
-			Newline(Context, RectMax);
-
-			Print(Context, TEXT(" "), FontSilver);
-			Print(Context, BSDF_GETSHAREDLOCALBASISID(BSDF), FontBSDF);
-			PrintBool(Context, BSDF_GETISTOPLAYER(BSDF));
-			if (BSDF_GETHASGREYWEIGHT_V(BSDF))
-			{
-				Print(Context, BSDF.LuminanceWeightV.x, FontSilver);
-			}
-			else
-			{
-				Print(Context, BSDF.LuminanceWeightV.r, FontLightRed);
-				Print(Context, BSDF.LuminanceWeightV.g, FontLightGreen);
-				Print(Context, BSDF.LuminanceWeightV.b, FontLightBlue);
-			}
-			Newline(Context, RectMax);
-
-			Print(Context, TEXT(" Diffuse Albedo    "), FontBSDFPropName);	Print(Context, EYE_DIFFUSEALBEDO(BSDF), FontBSDFPropValu);			Newline(Context, RectMax);
-			Print(Context, TEXT(" F0                "), FontBSDFPropName);	Print(Context, EYE_F0(BSDF), FontBSDFPropValu);						Newline(Context, RectMax);
-			Print(Context, TEXT(" Roughness         "), FontBSDFPropName);	Print(Context, EYE_ROUGHNESS(BSDF), FontBSDFPropValu);				Newline(Context, RectMax);
-			Print(Context, TEXT(" Iris Mask         "), FontBSDFPropName);	Print(Context, EYE_IRISMASK(BSDF), FontBSDFPropValu);			Newline(Context, RectMax);
-			Print(Context, TEXT(" Iris Distance     "), FontBSDFPropName);	Print(Context, EYE_IRISDISTANCE(BSDF), FontBSDFPropValu);				Newline(Context, RectMax);
-
-			const float LineSize = 10.f;
-
 			// Tangent basis for the iris
 			const float3x3 IrisTangentBasis = GetTangentBasis(EYE_IRISNORMAL(BSDF));
 			//DrawReferentialTWS(WorldPosition, IrisTangentBasis[0], IrisTangentBasis[1], IrisTangentBasis[2], float3(0, 1, 1));
@@ -447,34 +172,11 @@ void SubstratePrintBSDF(
 			//DrawReferentialTWS(WorldPosition, IrisPlaneTangentBasis[0], IrisPlaneTangentBasis[1], IrisPlaneTangentBasis[2], float3(1, 0, 1));
 			AddLineTWS(WorldPosition, WorldPosition + EYE_IRISPLANENORMAL(BSDF) * LineSize, float4(1, 0, 1, 1));
 
-
-			Print(Context, TEXT(" Iris N            "), FontBSDFPropName);	Print(Context, EYE_IRISNORMAL(BSDF), FontBSDFPropValu);			Newline(Context, RectMax);
-			Print(Context, TEXT(" PlaneN            "), FontBSDFPropName);	Print(Context, EYE_IRISPLANENORMAL(BSDF), FontBSDFPropValu);			Newline(Context, RectMax);
-
-		#if SUBSTRATE_MATERIALCONTAINER_IS_VIEWRESOURCE
-			const bool bHasSSS = Header.HasSubsurface();
-			const bool bIsValid = SubstrateSubSurfaceHeaderGetIsValid(SSSHeader);
-			const bool bIsProfile = SubstrateSubSurfaceHeaderGetIsProfile(SSSHeader);
-			if (bHasSSS && bIsValid && bIsProfile)
-			{
-				const uint SubsurfaceProfileInt = SubstrateSubSurfaceHeaderGetProfileId(SSSHeader);
-				Print(Context, TEXT(" SSS ID.           "), FontBSDFPropName);	Print(Context, SubsurfaceProfileInt, FontBSDFPropValu);			Newline(Context, RectMax);
-			}
-		#endif
 		}
 		break;
 
 		case SUBSTRATE_BSDF_TYPE_SINGLELAYERWATER:
 		{
-			Print(Context, TEXT("Water"), FontBSDFType);
-			PrintAddress(Context, SubstrateAddressing);
-			Newline(Context);
-
-			Print(Context, TEXT(" BaseColor            "), FontBSDFPropName);	Print(Context, SLW_BASECOLOR(BSDF), FontBSDFPropValu);			Newline(Context, RectMax);
-			Print(Context, TEXT(" Metallic             "), FontBSDFPropName);	Print(Context, SLW_METALLIC(BSDF), FontBSDFPropValu);			Newline(Context, RectMax);
-			Print(Context, TEXT(" Specular             "), FontBSDFPropName);	Print(Context, SLW_SPECULAR(BSDF), FontBSDFPropValu);			Newline(Context, RectMax);
-			Print(Context, TEXT(" Roughness            "), FontBSDFPropName);	Print(Context, SLW_ROUGHNESS(BSDF), FontBSDFPropValu);			Newline(Context, RectMax);
-			Print(Context, TEXT(" Top Material Opacity "), FontBSDFPropName);	Print(Context, SLW_TOPMATERIALOPACITY(BSDF), FontBSDFPropValu);	Newline(Context, RectMax);
 		}
 		break;
 
@@ -484,8 +186,6 @@ void SubstratePrintBSDF(
 		}
 		break;
 	}
-
-	Newline(Context);
 }
 
 void SubstratePrintMaterialProperties_Internal(
@@ -511,67 +211,6 @@ void SubstratePrintMaterialProperties_Internal(
 	float2 RectMin = Context.Pos;
 	float2 RectMax = Context.StartPos;
 
-	// Header
-	if (InClosureIndex < Header.ClosureCount && InClosureIndex == 0)
-	{	
-		const bool bSubstrateMaterial = Header.ClosureCount > 0;
-		const bool bIsSimpleMaterial = Header.IsSimpleMaterial() || Header.ClosureCount == 0;
-		const bool bIsSingleMaterial = !Header.IsSimpleMaterial() && Header.IsSingleMaterial();
-
-		FSubstrateIrradianceAndOcclusion IrradianceAndOcclusion = SubstrateGetIrradianceAndAO(Header);
-
-		RectMin = Context.Pos;
-		Print(Context, Header.ClosureCount, FontBSDFCount, 2, 0);
-		Print(Context, TEXT("BSDF"), FontBSDFCount);
-		Print(Context, TEXT(" - "), FontWhite);
-		PrintPixelType(Context, Header, FontLightRed);
-		if (bIsSingleMaterial)
-		{
-		#if SUBSTRATE_DEFERRED_SHADING
-			uint OptimisedLegacyMode = (Header.PackedHeader >> (HEADER_SINGLEENCODING_BIT_COUNT)) & HEADER_SINGLE_OPTLEGACYMODE_BIT_MASK;
-			if (OptimisedLegacyMode ==SINGLE_OPTLEGACYMODE_CLEARCOAT)
-			{
-				Print(Context, TEXT(" - Legacy Clear-Coat"), FontLightRed);
-			}
-			else if (OptimisedLegacyMode ==SINGLE_OPTLEGACYMODE_CLOTH)
-			{
-				Print(Context, TEXT(" - Legacy Cloth"), FontLightRed);
-			}
-			else if (OptimisedLegacyMode ==SINGLE_OPTLEGACYMODE_SSSWRAP)
-			{
-				Print(Context, TEXT(" - Legacy SSS-Wrap"), FontLightRed);
-			}
-			else if (OptimisedLegacyMode == SINGLE_OPTLEGACYMODE_SSSWRAP_THIN)
-			{
-				Print(Context, TEXT(" - Legacy Foliage"), FontLightRed);
-			}
-			else if (OptimisedLegacyMode == SINGLE_OPTLEGACYMODE_SSSPROFILE)
-			{
-				Print(Context, TEXT(" - Legacy SSS-Profile"), FontLightRed);
-			}
-		#endif
-		}
-		Newline(Context, RectMax);
-
-		Print(Context, TEXT(" AO    IndIrr TopRoughness PreShadow ZeroShadow ContacShadow Ind.Occluder HasSSS BasisCnt"), FontHeaderStateName);
-		Newline(Context, RectMax);
-
-		Print(Context, TEXT(" "));
-		Print(Context, IrradianceAndOcclusion.MaterialAO, FontHeaderPropValue, 5, 3);					Print(Context, TEXT(" "));
-		Print(Context, IrradianceAndOcclusion.IndirectIrradiance, FontHeaderPropValue, 6, 3);			Print(Context, TEXT(" "));
-		Print(Context, TopLayerData.Roughness, FontHeaderPropValue, 5, 3);								Print(Context, TEXT("        "));
-		PrintBool(Context, Header.HasPrecShadowMask());													Print(Context, TEXT("      "));
-		PrintBool(Context, Header.HasZeroPrecShadowMask());												Print(Context, TEXT("       "));
-		PrintBool(Context, Header.DoesCastContactShadow());												Print(Context, TEXT("         "));
-		PrintBool(Context, Header.HasDynamicIndirectShadowCasterRepresentation());						Print(Context, TEXT("         "));
-		PrintBool(Context, Header.HasSubsurface());														Print(Context, TEXT("          "));
-		Print(Context, Header.GetLocalBasesCount());
-		Newline(Context, RectMax);
-		Newline(Context);
-		
-		PrintBackgroundRect(Context, RectMin, RectMax);
-	}
-
 	// BSDFs
 	if (InClosureIndex < Header.ClosureCount)
 	{
@@ -586,54 +225,6 @@ void SubstratePrintMaterialProperties_Internal(
 		RectMin = Context.Pos;
 		RectMax = Context.StartPos;
 		SubstratePrintBSDF(Context, SubstrateAddressing, Header, SSSHeader, InWorldPosition, V, RectMax, Payload);
-		PrintBackgroundRect(Context, RectMin, RectMax);
-	}
-
-	// Memory footprint
-	if (InClosureIndex+1 == Header.ClosureCount)
-	{
-		const bool bSubstrateMaterial = Header.ClosureCount > 0;
-		const bool bIsSimpleMaterial = Header.IsSimpleMaterial() || Header.ClosureCount == 0;
-		const bool bIsSingleMaterial = !Header.IsSimpleMaterial() && Header.IsSingleMaterial();
-
-		const uint FootPrint_PostBSDFs = SubstrateAddressing.ReadBytes;
-
-		const bool bHasSSSData = SubstrateSubSurfaceHeaderGetIsValid(SSSHeader);
-
-		const uint TopLayerDataBytes = 4;	// bytes 
-		const uint SSSDataBytes = 8;
-
-		const uint HeaderSize =  FootPrint_PostHeader - FootPrint_Start;
-		const uint BSDFsSize  =  FootPrint_PostBSDFs - FootPrint_PostHeader;
-		const uint TotalSize  = (FootPrint_PostBSDFs - FootPrint_Start) + (bHasSSSData ? SSSDataBytes : 0);
-
-		RectMin = Context.Pos;
-		RectMax = Context.StartPos;
-		Print(Context, TEXT("Memory"), FontEmerald); Newline(Context, RectMax);
-		if (bIsSingleMaterial || Header.IsEye() || Header.IsHair())
-		{
-			Print(Context, TEXT(" Header       "), FontSilver); Print(Context, HeaderSize - TopLayerDataBytes,				FontSilver, 3, 0); Print(Context, TEXT("bytes"), FontSilver); Newline(Context, RectMax);
-			Print(Context, TEXT(" TopNormalTex "), FontSilver); Print(Context, TopLayerDataBytes,							FontSilver, 3, 0); Print(Context, TEXT("bytes"), FontSilver); Newline(Context, RectMax);
-			Print(Context, TEXT(" BSDF         "), FontSilver); Print(Context, BSDFsSize,									FontSilver, 3, 0); Print(Context, TEXT("bytes"), FontSilver); Newline(Context, RectMax);
-		}
-		else if (bIsSimpleMaterial || Header.IsSingleLayerWater())
-		{
-			Print(Context, TEXT(" Header+BSDF  "), FontSilver); Print(Context, HeaderSize + BSDFsSize - TopLayerDataBytes,	FontSilver, 3, 0); Print(Context, TEXT("bytes"), FontSilver); Newline(Context, RectMax);
-			Print(Context, TEXT(" TopNormalTex "), FontSilver); Print(Context, TopLayerDataBytes,							FontSilver, 3, 0); Print(Context, TEXT("bytes"), FontSilver); Newline(Context, RectMax);
-		}
-		else
-		{
-			Print(Context, TEXT(" Header+Norm  "), FontSilver); Print(Context, HeaderSize,									FontSilver, 3, 0); Print(Context, TEXT("bytes"), FontSilver); Newline(Context, RectMax);
-			Print(Context, TEXT(" BSDFs        "), FontSilver); Print(Context, BSDFsSize,									FontSilver, 3, 0); Print(Context, TEXT("bytes"), FontSilver); Newline(Context, RectMax);
-		}
-
-		if (bHasSSSData)
-		{
-			Print(Context, TEXT(" SSS Data     "), FontSilver); Print(Context, SSSDataBytes, FontSilver, 3, 0); Print(Context, TEXT("bytes"), FontSilver); Newline(Context, RectMax);
-		}
-
-		Print(Context, TEXT(" Total        "), FontOrange); Print(Context, TotalSize,  FontOrange, 3, 0); Print(Context, TEXT("bytes"), FontOrange); Newline(Context, RectMax);
-		PrintBackgroundRect(Context, RectMin, RectMax);
 	}
 }
 
diff --git a/Engine/Shaders/Shared/SubstrateDefinitions.h b/Engine/Shaders/Shared/SubstrateDefinitions.h
--- a/Engine/Shaders/Shared/SubstrateDefinitions.h
+++ b/Engine/Shaders/Shared/SubstrateDefinitions.h
@@ -9,7 +9,7 @@
 
 #ifndef __cplusplus
 // Change this to force recompilation of all Substrate dependent shaders (for instance https://guidgenerator.com/online-guid-generator.aspx)
-#pragma message("UESHADERMETADATA_VERSION 1F256706-39AF-450F-A2C1-6A10D9030F4E")
+#pragma message("UESHADERMETADATA_VERSION ADA0DF8F-BCE7-4921-89DF-E00C06DA6844")
 #endif
 
 // Closure offsets are packed into 32bits, each entry using SUBSTRATE_CLOSURE_OFFSET_BIT_COUNT bits
@@ -102,4 +102,32 @@
 #else
 	#define GetSubstrateTileTypeDispatchIndirectArgOffset_Byte(x)  (x * DISPATCH_INDIRECT_UINT_COUNT * 4)
 	#define GetSubstrateTileTypeDispatchIndirectArgOffset_DWord(x) (x * DISPATCH_INDIRECT_UINT_COUNT)
-#endif
\ No newline at end of file
+#endif
+
+#define HEADER_MATERIALMODE_NONE				0
+#define HEADER_MATERIALMODE_SLAB_SIMPLE			1
+#define HEADER_MATERIALMODE_SLAB_SINGLE			2
+#define HEADER_MATERIALMODE_SLAB_COMPLEX		3
+#define HEADER_MATERIALMODE_SLWATER				4
+#define HEADER_MATERIALMODE_HAIR				5
+#define HEADER_MATERIALMODE_EYE					6
+// Up to 8 material modes: 1 remains.
+
+// Special optmised modes for the single BSDF path.
+#define HEADER_SINGLE_OPTLEGACYMODE_BIT_COUNT	3
+#define HEADER_SINGLE_OPTLEGACYMODE_BIT_MASK	((1u<<HEADER_SINGLE_OPTLEGACYMODE_BIT_COUNT)-1u)
+#define SINGLE_OPTLEGACYMODE_NONE				0
+#define SINGLE_OPTLEGACYMODE_CLEARCOAT			1
+#define SINGLE_OPTLEGACYMODE_CLOTH				2
+#define SINGLE_OPTLEGACYMODE_SSSWRAP			3
+#define SINGLE_OPTLEGACYMODE_SSSPROFILE			4
+#define SINGLE_OPTLEGACYMODE_SSSWRAP_THIN		5
+// 5 more remains: Also how to deal with Aniso?
+
+// Sub-surface scattering types
+#define SSS_TYPE_INVALID						(0u)
+#define SSS_TYPE_WRAP							(1u)
+#define SSS_TYPE_DIFFUSION						(2u)
+#define SSS_TYPE_DIFFUSION_PROFILE				(3u)
+#define SSS_TYPE_SIMPLEVOLUME					(4u)
+#define SSS_TYPE_COUNT							(5u)
diff --git a/Engine/Shaders/Shared/SubstrateVisualizeDefinitions.h b/Engine/Shaders/Shared/SubstrateVisualizeDefinitions.h
new file mode 100644
--- /dev/null
+++ b/Engine/Shaders/Shared/SubstrateVisualizeDefinitions.h
@@ -0,0 +1,590 @@
+// Copyright Epic Games, Inc. All Rights Reserved.
+
+/*================================================================================================
+	SubstrateVisualise.h: used in ray tracing shaders and C++ code to define common constants
+	!!! Changing this file requires recompilation of the engine !!!
+=================================================================================================*/
+
+#pragma once
+
+
+#ifndef __cplusplus
+// Change this to force recompilation of all Substrate dependent shaders (for instance https://guidgenerator.com/online-guid-generator.aspx)
+#pragma message("UESHADERMETADATA_VERSION 8FFAABAF-705E-481F-8234-D4D42C7A4001")
+#include "/Engine/Shared/SubstrateDefinitions.h"
+#else
+#include "SubstrateDefinitions.h"
+#endif
+
+struct FSubsterateDebugClosure
+{
+	int Type;
+
+	int bHasWeightL;// = BSDF_GETHASTRANSABOVE(BSDF);
+	int bHasGreyWeightV;// = BSDF_GETHASGREYWEIGHT_V(BSDF);
+	int Address; // SubstrateAddressing.CurrentIndex
+
+	int NormalID; //BSDF_GETSHAREDLOCALBASISID(BSDF)
+	int BasisType; // SubstrateGetSharedLocalBasisType(HEADER_GETSHAREDLOCALBASESTYPE(Header.PackedHeader), BSDF_GETSHAREDLOCALBASISID(BSDF));
+
+	//////////
+	// Slabs
+	int bIsTopLayer; // BSDF_GETISTOPLAYER(BSDF)
+	int SSSType; // BSDF_GETSSSTYPE(BSDF) != SSS_TYPE_INVALID
+	int bIsThin; // BSDF_GETISTHIN(BSDF)
+	float LuminanceWeightR, LuminanceWeightG, LuminanceWeightB; // bHasGreyWeightV BSDF.LuminanceWeightV.x or rgb
+	float TransmittanceAboveAlongNR, TransmittanceAboveAlongNG, TransmittanceAboveAlongNB; // bHasGreyTopTrans BSDF.TransmittanceAboveAlongN.x or rgb
+	float CoverageAboveAlongN;
+
+	float DiffuseR,	DiffuseG,	DiffuseB;
+	float F0R,		F0G,		F0B;
+	float Roughness;
+
+	int bHasF90;
+	float F90R, F90G, F90B;
+
+	int bHasAnisotropy;
+	float Anisotropy;
+
+	int bHasHaziness;
+	float HazeRoughness;
+	float HazeWeight;
+	int HazeSimpleClearCoatMode;
+
+	float  SSSOpacity;
+	float SSSMFPR, SSSMFPG, SSSMFPB;
+	float SSSRescaledMFPR, SSSRescaledMFPG, SSSRescaledMFPB;
+	float  SSSPhase;
+	float  SSSThickness;
+	float  SSSProfileRadius;
+	int SSSPRofileID;
+
+	float FuzzAmount;
+	float FuzzColorR, FuzzColorG, FuzzColorB;
+	float FuzzRoughness;
+
+	float GlintValue;
+
+	int SpecProfileID;
+	int SpecProfileParameterization;
+};
+
+struct FSubstratePixelDebugData
+{
+	int ClosureCount;
+	int MaterialMode;
+	int bIsComplexSpecialMaterial;
+	int OptimisedLegacyMode;
+
+	float MaterialAO;
+	float IndirectIrradiance;
+	float TopLayerRoughness;
+	int HasPrecShadowMask;
+	int HasZeroPrecShadowMask;
+	int DoesCastContactShadow;
+	int HasDynamicIndirectShadowCasterRepresentation;
+	int HasSubsurface;
+	int LocalBasesCount;
+
+	FSubsterateDebugClosure Closures[SUBSTRATE_MAX_CLOSURE_COUNT];
+
+	int MemoryDisplayMode;
+	int MemorySlotA;
+	int MemorySlotB;
+	int MemorySlotC;
+	int MemorySSSData;
+	int MemoryTotal;
+};
+
+#ifndef __cplusplus
+RWStructuredBuffer<int> PixelMaterialDebugDataUAV;
+struct FSubstratePixelDebugDataSerializer
+{
+	int WriteIndex;
+
+	void Serialize(int Data)
+	{
+		PixelMaterialDebugDataUAV[WriteIndex++] = Data;
+	}
+
+	void Serialize(float Data)
+	{
+		PixelMaterialDebugDataUAV[WriteIndex++] = asuint(Data);
+	}
+};
+#else
+struct FSubstratePixelDebugDataSerializer
+{
+	int ReadIndex = 0;
+	int* SubstratePixelDebugData = nullptr;
+
+	void Serialize(int& Data)
+	{
+		Data = SubstratePixelDebugData[ReadIndex++];
+	}
+
+	void Serialize(float& Data)
+	{
+		Data = ((float*)SubstratePixelDebugData)[ReadIndex++];
+	}
+};
+#endif
+
+#ifndef __cplusplus
+void SerializeSubstratePixelDebugData(in FSubstratePixelDebugDataSerializer S, in FSubstratePixelDebugData D)
+#else
+void SerializeSubstratePixelDebugData(FSubstratePixelDebugDataSerializer& S, FSubstratePixelDebugData& D)
+#endif
+{
+	S.Serialize(D.ClosureCount);
+#ifdef __cplusplus
+	// Safe guard to avoid crash in case of readback problem.
+	D.ClosureCount = FMath::Min(D.ClosureCount, int(SUBSTRATE_MAX_CLOSURE_COUNT));
+#endif
+
+	S.Serialize(D.MaterialMode);
+	S.Serialize(D.bIsComplexSpecialMaterial);
+	S.Serialize(D.OptimisedLegacyMode);
+
+	S.Serialize(D.MaterialAO);
+	S.Serialize(D.IndirectIrradiance);
+	S.Serialize(D.TopLayerRoughness);
+	S.Serialize(D.HasPrecShadowMask);
+	S.Serialize(D.HasZeroPrecShadowMask);
+	S.Serialize(D.DoesCastContactShadow);
+	S.Serialize(D.HasDynamicIndirectShadowCasterRepresentation);
+	S.Serialize(D.HasSubsurface);
+	S.Serialize(D.LocalBasesCount);
+
+	for (int i = 0; i < D.ClosureCount; ++i)
+	{
+		S.Serialize(D.Closures[i].Type);
+
+		S.Serialize(D.Closures[i].bHasWeightL);
+		S.Serialize(D.Closures[i].bHasGreyWeightV);
+		S.Serialize(D.Closures[i].Address);
+
+		S.Serialize(D.Closures[i].NormalID);
+		S.Serialize(D.Closures[i].BasisType);
+
+		S.Serialize(D.Closures[i].bIsTopLayer);
+		S.Serialize(D.Closures[i].SSSType);
+		S.Serialize(D.Closures[i].bIsThin);
+		S.Serialize(D.Closures[i].LuminanceWeightR);
+		S.Serialize(D.Closures[i].LuminanceWeightG);
+		S.Serialize(D.Closures[i].LuminanceWeightB);
+		S.Serialize(D.Closures[i].TransmittanceAboveAlongNR);
+		S.Serialize(D.Closures[i].TransmittanceAboveAlongNG);
+		S.Serialize(D.Closures[i].TransmittanceAboveAlongNB);
+		S.Serialize(D.Closures[i].CoverageAboveAlongN);
+
+		S.Serialize(D.Closures[i].DiffuseR);
+		S.Serialize(D.Closures[i].DiffuseG);
+		S.Serialize(D.Closures[i].DiffuseB);
+		S.Serialize(D.Closures[i].F0R);
+		S.Serialize(D.Closures[i].F0G);
+		S.Serialize(D.Closures[i].F0B);
+		S.Serialize(D.Closures[i].Roughness);
+
+		S.Serialize(D.Closures[i].bHasF90);
+		S.Serialize(D.Closures[i].F90R);
+		S.Serialize(D.Closures[i].F90G);
+		S.Serialize(D.Closures[i].F90B);
+
+		S.Serialize(D.Closures[i].bHasAnisotropy);
+		S.Serialize(D.Closures[i].Anisotropy);
+
+		S.Serialize(D.Closures[i].bHasHaziness);
+		S.Serialize(D.Closures[i].HazeRoughness);
+		S.Serialize(D.Closures[i].HazeWeight);
+		S.Serialize(D.Closures[i].HazeSimpleClearCoatMode);
+
+		S.Serialize(D.Closures[i].SSSOpacity);
+		S.Serialize(D.Closures[i].SSSMFPR);
+		S.Serialize(D.Closures[i].SSSMFPG);
+		S.Serialize(D.Closures[i].SSSMFPB);
+		S.Serialize(D.Closures[i].SSSRescaledMFPR);
+		S.Serialize(D.Closures[i].SSSRescaledMFPG);
+		S.Serialize(D.Closures[i].SSSRescaledMFPB);
+		S.Serialize(D.Closures[i].SSSPhase);
+		S.Serialize(D.Closures[i].SSSThickness);
+		S.Serialize(D.Closures[i].SSSProfileRadius);
+		S.Serialize(D.Closures[i].SSSPRofileID);
+
+		S.Serialize(D.Closures[i].FuzzAmount);
+		S.Serialize(D.Closures[i].FuzzColorR);
+		S.Serialize(D.Closures[i].FuzzColorG);
+		S.Serialize(D.Closures[i].FuzzColorB);
+		S.Serialize(D.Closures[i].FuzzRoughness);
+
+		S.Serialize(D.Closures[i].GlintValue);
+
+		S.Serialize(D.Closures[i].SpecProfileID);
+		S.Serialize(D.Closures[i].SpecProfileParameterization);
+	}
+
+	S.Serialize(D.MemoryDisplayMode);
+	S.Serialize(D.MemorySlotA);
+	S.Serialize(D.MemorySlotB);
+	S.Serialize(D.MemorySlotC);
+	S.Serialize(D.MemorySSSData);
+	S.Serialize(D.MemoryTotal);
+
+}
+
+#ifndef __cplusplus
+
+#if SUBSTRATE_MATERIALCONTAINER_IS_VIEWRESOURCE
+FSubstratePixelDebugData ConvertToSerializableSubstratePixelDebugData(uint2 InCoord)
+#else
+FSubstratePixelDebugData ConvertToSerializableSubstratePixelDebugData(uint2 InCoord, FSubstrateRaytracingPayload Payload)
+#endif
+{
+	FSubstratePixelDebugData Data = (FSubstratePixelDebugData)0;
+
+	FSubstrateAddressing SubstrateAddressing = GetSubstratePixelDataByteOffset(InCoord, uint2(View.BufferSizeAndInvSize.xy), Substrate.MaxBytesPerPixel);
+	const uint FootPrint_Start = SubstrateAddressing.ReadBytes;
+
+#if SUBSTRATE_MATERIALCONTAINER_IS_VIEWRESOURCE
+	FSubstratePixelHeader Header = UnpackSubstrateHeaderIn(Substrate.MaterialTextureArray, SubstrateAddressing, Substrate.TopLayerTexture);
+#else
+	FSubstratePixelHeader Header = UnpackSubstrateHeaderIn(Payload, SubstrateAddressing, Payload);
+#endif
+	const uint FootPrint_PostHeader = SubstrateAddressing.ReadBytes;
+
+#if SUBSTRATE_MATERIALCONTAINER_IS_VIEWRESOURCE
+	FSubstrateSubsurfaceHeader SSSHeader = SubstrateLoadSubsurfaceHeader(Substrate.MaterialTextureArray, Substrate.FirstSliceStoringSubstrateSSSData, SubstrateAddressing.PixelCoords);
+	FSubstrateTopLayerData TopLayerData = SubstrateUnpackTopLayerData(Substrate.TopLayerTexture.Load(uint3(InCoord, 0)));
+#else
+	FSubstrateSubsurfaceHeader SSSHeader = (FSubstrateSubsurfaceHeader)0;
+	FSubstrateTopLayerData TopLayerData = SubstrateUnpackTopLayerData(Payload.PackedTopLayerData);
+#endif
+	Data.TopLayerRoughness = TopLayerData.Roughness;
+
+	const bool bSubstrateMaterial = Header.ClosureCount > 0;
+	const bool bIsSimpleMaterial = Header.IsSimpleMaterial() || Header.ClosureCount == 0;
+	const bool bIsSingleMaterial = !Header.IsSimpleMaterial() && Header.IsSingleMaterial();
+
+	FSubstrateIrradianceAndOcclusion IrradianceAndOcclusion = SubstrateGetIrradianceAndAO(Header);
+
+	Data.ClosureCount = Header.ClosureCount;
+
+	Data.MaterialMode = Header.GetMaterialMode();
+	Data.OptimisedLegacyMode = HEADER_MATERIALMODE_NONE;
+	if (bIsSingleMaterial)
+	{
+#if SUBSTRATE_DEFERRED_SHADING
+		Data.OptimisedLegacyMode = (Header.PackedHeader >> (HEADER_SINGLEENCODING_BIT_COUNT)) & HEADER_SINGLE_OPTLEGACYMODE_BIT_MASK;
+#endif
+	}
+	Data.bIsComplexSpecialMaterial = Header.IsComplexSpecialMaterial();
+
+	Data.MaterialAO = IrradianceAndOcclusion.MaterialAO;
+	Data.IndirectIrradiance = IrradianceAndOcclusion.IndirectIrradiance;
+	Data.HasPrecShadowMask = Header.HasPrecShadowMask() ? 1 : 0;
+	Data.HasZeroPrecShadowMask = Header.HasZeroPrecShadowMask() ? 1 : 0;
+	Data.DoesCastContactShadow = Header.DoesCastContactShadow() ? 1 : 0;
+	Data.HasDynamicIndirectShadowCasterRepresentation = Header.HasDynamicIndirectShadowCasterRepresentation() ? 1 : 0;
+	Data.HasSubsurface = Header.HasSubsurface() ? 1 : 0;
+	Data.LocalBasesCount = Header.GetLocalBasesCount();
+
+	LOOP
+	for(uint i = 0; i < Header.ClosureCount; ++i)
+	{
+		// Unpack BSDF data
+	#if SUBSTRATE_MATERIALCONTAINER_IS_VIEWRESOURCE
+		FSubstrateBSDF BSDF = UnpackSubstrateBSDF(Substrate.MaterialTextureArray, SubstrateAddressing, Header);
+	#else
+		FSubstrateBSDF BSDF = UnpackSubstrateBSDFIn(Payload, SubstrateAddressing, Header);
+	#endif
+
+		Data.Closures[i].Address = SubstrateAddressing.CurrentIndex;
+		Data.Closures[i].Type = BSDF_GETTYPE(BSDF);
+
+		Data.Closures[i].NormalID = BSDF_GETSHAREDLOCALBASISID(BSDF);
+		Data.Closures[i].BasisType = SubstrateGetSharedLocalBasisType(HEADER_GETSHAREDLOCALBASESTYPE(Header.PackedHeader), BSDF_GETSHAREDLOCALBASISID(BSDF));
+
+
+		Data.Closures[i].bIsTopLayer = BSDF_GETISTOPLAYER(BSDF);
+		Data.Closures[i].bHasWeightL = BSDF_GETHASTRANSABOVE(BSDF);
+		Data.Closures[i].bHasGreyWeightV = BSDF_GETHASGREYWEIGHT_V(BSDF);
+
+		Data.Closures[i].LuminanceWeightR = BSDF.LuminanceWeightV.r;
+		Data.Closures[i].LuminanceWeightG = BSDF.LuminanceWeightV.g;
+		Data.Closures[i].LuminanceWeightB = BSDF.LuminanceWeightV.b;
+		Data.Closures[i].TransmittanceAboveAlongNR = BSDF.TransmittanceAboveAlongN.r;
+		Data.Closures[i].TransmittanceAboveAlongNG = BSDF.TransmittanceAboveAlongN.g;
+		Data.Closures[i].TransmittanceAboveAlongNB = BSDF.TransmittanceAboveAlongN.b;
+		Data.Closures[i].CoverageAboveAlongN = BSDF.CoverageAboveAlongN;
+
+		Data.Closures[i].SSSType = BSDF_GETSSSTYPE(BSDF);
+		Data.Closures[i].bIsThin = BSDF_GETISTHIN(BSDF);
+
+		switch (Data.Closures[i].Type)
+		{
+		case SUBSTRATE_BSDF_TYPE_SLAB:
+		{
+			Data.Closures[i].DiffuseR = SLAB_DIFFUSEALBEDO(BSDF).r;
+			Data.Closures[i].DiffuseG = SLAB_DIFFUSEALBEDO(BSDF).g;
+			Data.Closures[i].DiffuseB = SLAB_DIFFUSEALBEDO(BSDF).b;
+			Data.Closures[i].F0R = SLAB_F0(BSDF).r;
+			Data.Closures[i].F0G = SLAB_F0(BSDF).g;
+			Data.Closures[i].F0B = SLAB_F0(BSDF).b;
+			Data.Closures[i].Roughness = SLAB_ROUGHNESS(BSDF);
+
+			Data.Closures[i].bHasF90 = BSDF_GETHASF90(BSDF);
+			if (Data.Closures[i].bHasF90)
+			{
+				Data.Closures[i].F90R = SLAB_F90(BSDF).r;
+				Data.Closures[i].F90G = SLAB_F90(BSDF).g;
+				Data.Closures[i].F90B = SLAB_F90(BSDF).b;
+			}
+			else
+			{
+				Data.Closures[i].F90R = 1.0f;
+				Data.Closures[i].F90G = 1.0f;
+				Data.Closures[i].F90B = 1.0f;
+			}
+
+			Data.Closures[i].bHasAnisotropy = BSDF_GETHASANISOTROPY(BSDF);
+			if (Data.Closures[i].bHasAnisotropy)
+			{
+				Data.Closures[i].Anisotropy = SLAB_ANISOTROPY(BSDF);
+			}
+
+			Data.Closures[i].bHasHaziness = BSDF_GETHASHAZINESS(BSDF);
+			if (Data.Closures[i].bHasHaziness)
+			{
+				FHaziness Haziness = UnpackHaziness(SLAB_HAZINESS(BSDF));
+				Data.Closures[i].HazeRoughness = Haziness.Roughness;
+				Data.Closures[i].HazeWeight = Haziness.Weight;
+				Data.Closures[i].HazeSimpleClearCoatMode = Haziness.bSimpleClearCoat ? 1u : 0u;
+			}
+
+			if (Data.Closures[i].SSSType != SSS_TYPE_INVALID || Data.Closures[i].bIsThin)
+			{
+				if (Data.Closures[i].SSSType == SSS_TYPE_WRAP)
+				{
+				#if SUBSTRATE_MATERIALCONTAINER_IS_VIEWRESOURCE
+					Data.Closures[i].SSSOpacity = SubstrateSubSurfaceHeaderGetWrapOpacity(SSSHeader);
+				#endif
+					Data.Closures[i].SSSMFPR = SLAB_SSSMFP(BSDF).r;
+					Data.Closures[i].SSSMFPG = SLAB_SSSMFP(BSDF).g;
+					Data.Closures[i].SSSMFPB = SLAB_SSSMFP(BSDF).b;
+					Data.Closures[i].SSSPhase = SLAB_SSSPHASEANISOTROPY(BSDF);
+					Data.Closures[i].SSSThickness = BSDF_GETTHICKNESSCM(BSDF);
+				}
+				else if (Data.Closures[i].SSSType == SSS_TYPE_SIMPLEVOLUME)
+				{
+				#if SUBSTRATE_MATERIALCONTAINER_IS_VIEWRESOURCE
+					Data.Closures[i].SSSOpacity = SubstrateSubSurfaceHeaderGetWrapOpacity(SSSHeader);
+				#endif
+					Data.Closures[i].SSSMFPR = SLAB_SSSMFP(BSDF).r;
+					Data.Closures[i].SSSMFPG = SLAB_SSSMFP(BSDF).g;
+					Data.Closures[i].SSSMFPB = SLAB_SSSMFP(BSDF).b;
+					Data.Closures[i].SSSPhase = SLAB_SSSPHASEANISOTROPY(BSDF);
+					Data.Closures[i].SSSThickness = BSDF_GETTHICKNESSCM(BSDF);
+				}
+				else if (Data.Closures[i].SSSType == SSS_TYPE_DIFFUSION_PROFILE)
+				{
+				#if SUBSTRATE_MATERIALCONTAINER_IS_VIEWRESOURCE
+					const uint ProfileId = SubstrateSubSurfaceHeaderGetProfileId(SSSHeader);
+					const float RadiusScale = SubstrateSubSurfaceHeaderGetProfileRadiusScale(SSSHeader);
+					const float3 DiffuseMFP = GetSubsurfaceProfileMFPInCm(ProfileId).xyz * RadiusScale;
+					const float Phase = GetTransmissionProfileParams(ProfileId).ScatteringDistribution;
+
+					Data.Closures[i].SSSPRofileID = ProfileId;
+					Data.Closures[i].SSSProfileRadius = RadiusScale;
+					Data.Closures[i].SSSMFPR = DiffuseMFP.r;
+					Data.Closures[i].SSSMFPG = DiffuseMFP.g;
+					Data.Closures[i].SSSMFPB = DiffuseMFP.b;
+					Data.Closures[i].SSSPhase = Phase;
+				#endif
+					if (Data.Closures[i].bIsThin)
+					{
+						Data.Closures[i].SSSThickness = SLAB_SSSPROFILETHICKNESSCM(BSDF);
+					}
+				}
+				else if (Data.Closures[i].SSSType == SSS_TYPE_DIFFUSION)
+				{
+				#if SUBSTRATE_MATERIALCONTAINER_IS_VIEWRESOURCE
+					const float3 OriginalMFP = SubstrateSubSurfaceHeaderGetMFP(SSSHeader);
+					const float3 RescaledMFP = SLAB_SSSMFP(BSDF);
+
+
+					Data.Closures[i].SSSMFPR = OriginalMFP.r;
+					Data.Closures[i].SSSMFPG = OriginalMFP.g;
+					Data.Closures[i].SSSMFPB = OriginalMFP.b;
+
+					if (Data.Closures[i].bIsThin)
+					{
+						Data.Closures[i].SSSRescaledMFPR = RescaledMFP.r;
+						Data.Closures[i].SSSRescaledMFPG = RescaledMFP.g;
+						Data.Closures[i].SSSRescaledMFPB = RescaledMFP.b;
+						Data.Closures[i].SSSThickness = BSDF_GETTHICKNESSCM(BSDF);
+					}
+				#endif
+
+					Data.Closures[i].SSSPhase = SLAB_SSSPHASEANISOTROPY(BSDF);
+				}
+			}
+
+			if (BSDF_GETHASFUZZ(BSDF))
+			{
+				Data.Closures[i].FuzzAmount = SLAB_FUZZ_AMOUNT(BSDF);
+				Data.Closures[i].FuzzColorR = SLAB_FUZZ_COLOR(BSDF).r;
+				Data.Closures[i].FuzzColorG = SLAB_FUZZ_COLOR(BSDF).g;
+				Data.Closures[i].FuzzColorB = SLAB_FUZZ_COLOR(BSDF).b;
+				Data.Closures[i].FuzzRoughness = SLAB_FUZZ_ROUGHNESS(BSDF);
+			}
+
+			if (BSDF_GETHASGLINT(BSDF))
+			{
+				Data.Closures[i].GlintValue = SLAB_GLINT_VALUE(BSDF);
+			}
+			else
+			{
+				Data.Closures[i].GlintValue = 1.0;	// No glint
+			}
+
+			Data.Closures[i].SpecProfileID = -1.0f;
+			if (BSDF_GETHASSPECPROFILE(BSDF))
+			{
+				Data.Closures[i].SpecProfileID = GetSpecularProfileId(SLAB_SPECPROFILEID(BSDF));
+				Data.Closures[i].SpecProfileParameterization = GetSpecularProfileParameterization(SLAB_SPECPROFILEID(BSDF));
+			}
+		}
+		break;
+
+		case SUBSTRATE_BSDF_TYPE_HAIR:
+		{
+			// Reusing the slab data to avoid using too many registers in case the compiler has trouble optimising
+			Data.Closures[i].DiffuseR	= HAIR_BASECOLOR(BSDF).r;
+			Data.Closures[i].DiffuseG	= HAIR_BASECOLOR(BSDF).g;
+			Data.Closures[i].DiffuseB	= HAIR_BASECOLOR(BSDF).b;
+			Data.Closures[i].F0R		= HAIR_SPECULAR(BSDF);
+			Data.Closures[i].Roughness	= HAIR_ROUGHNESS(BSDF);
+			Data.Closures[i].F90R		= HAIR_SCATTER(BSDF);
+			Data.Closures[i].F90G		= HAIR_BACKLIT(BSDF);
+			Data.Closures[i].F90B		= HAIR_COMPLEXTRANSMITTANCE(BSDF);
+		}
+		break;
+
+		case SUBSTRATE_BSDF_TYPE_EYE:
+		{
+			// Reusing the slab data to avoid using too many registers in case the compiler has trouble optimising
+			Data.Closures[i].DiffuseR	= EYE_DIFFUSEALBEDO(BSDF).r;
+			Data.Closures[i].DiffuseG	= EYE_DIFFUSEALBEDO(BSDF).g;
+			Data.Closures[i].DiffuseB	= EYE_DIFFUSEALBEDO(BSDF).b;
+			Data.Closures[i].F0R		= EYE_F0(BSDF);
+			Data.Closures[i].Roughness	= EYE_ROUGHNESS(BSDF);
+			Data.Closures[i].F90R		= EYE_IRISMASK(BSDF);
+			Data.Closures[i].F90G		= EYE_IRISDISTANCE(BSDF);
+
+			Data.Closures[i].SSSMFPR	= EYE_IRISNORMAL(BSDF).x;
+			Data.Closures[i].SSSMFPG	= EYE_IRISNORMAL(BSDF).y;
+			Data.Closures[i].SSSMFPB	= EYE_IRISNORMAL(BSDF).z;
+			Data.Closures[i].SSSRescaledMFPR = EYE_IRISPLANENORMAL(BSDF).x;
+			Data.Closures[i].SSSRescaledMFPG = EYE_IRISPLANENORMAL(BSDF).y;
+			Data.Closures[i].SSSRescaledMFPB = EYE_IRISPLANENORMAL(BSDF).z;
+			
+		#if SUBSTRATE_MATERIALCONTAINER_IS_VIEWRESOURCE
+			const bool bHasSSS = Header.HasSubsurface();
+			const bool bIsValid = SubstrateSubSurfaceHeaderGetIsValid(SSSHeader);
+			const bool bIsProfile = SubstrateSubSurfaceHeaderGetIsProfile(SSSHeader);
+			Data.Closures[i].SSSPRofileID = -1;
+			if (bHasSSS && bIsValid && bIsProfile)
+			{
+				Data.Closures[i].SSSPRofileID = SubstrateSubSurfaceHeaderGetProfileId(SSSHeader);
+			}
+		#endif
+		}
+		break;
+
+		case SUBSTRATE_BSDF_TYPE_SINGLELAYERWATER:
+		{
+			// Reusing the slab data to avoid using too many registers in case the compiler has trouble optimising
+			Data.Closures[i].DiffuseR	= SLW_BASECOLOR(BSDF).r;
+			Data.Closures[i].DiffuseG	= SLW_BASECOLOR(BSDF).g;
+			Data.Closures[i].DiffuseB	= SLW_BASECOLOR(BSDF).b;
+			Data.Closures[i].F0R		= SLW_SPECULAR(BSDF).r;
+			Data.Closures[i].F0G		= SLW_METALLIC(BSDF).r;
+			Data.Closures[i].Roughness	= SLW_ROUGHNESS(BSDF);
+			Data.Closures[i].SSSOpacity = SLW_TOPMATERIALOPACITY(BSDF);
+		}
+		break;
+
+		default:
+		{
+			// Error
+		}
+		break;
+		}
+	}
+
+	const uint FootPrint_PostBSDFs = SubstrateAddressing.ReadBytes;
+
+	// Output memory reads
+	if (Header.ClosureCount > 0)
+	{
+		const bool bHasSSSData = SubstrateSubSurfaceHeaderGetIsValid(SSSHeader);
+		const uint TopLayerDataBytes = 4;
+		const uint SSSDataBytes = 8;
+
+		const bool bSubstrateMaterial = Header.ClosureCount > 0;
+		const bool bIsSimpleMaterial  = Header.IsSimpleMaterial() || Header.ClosureCount == 0;
+		const bool bIsSingleMaterial  = !Header.IsSimpleMaterial() && Header.IsSingleMaterial();
+
+		const uint HeaderSize = FootPrint_PostHeader - FootPrint_Start;
+		const uint BSDFsSize  = FootPrint_PostBSDFs  - FootPrint_PostHeader;
+		const uint TotalSize  = (FootPrint_PostBSDFs - FootPrint_Start) + (bHasSSSData ? SSSDataBytes : 0);
+
+		if (bIsSimpleMaterial || Header.IsSingleLayerWater())
+		{
+			Data.MemoryDisplayMode = 0;
+			Data.MemorySlotA = HeaderSize + BSDFsSize - TopLayerDataBytes;	// Header+BSDF
+			Data.MemorySlotB = TopLayerDataBytes;							// TopNormalTex
+		}
+		else if (bIsSingleMaterial || Header.IsEye() || Header.IsHair())
+		{
+			Data.MemoryDisplayMode = 1;
+			Data.MemorySlotA = HeaderSize - TopLayerDataBytes;				// Header
+			Data.MemorySlotB = TopLayerDataBytes;							// TopNormalTex
+			Data.MemorySlotC = BSDFsSize;									// BSDF
+		}
+		else
+		{
+			Data.MemoryDisplayMode = 2;
+			Data.MemorySlotA = HeaderSize;									// Header+Norm
+			Data.MemorySlotB = BSDFsSize;									// BSDFs
+		}
+		
+		if (bHasSSSData)
+		{
+			Data.MemorySSSData = SSSDataBytes;								// SSS Data
+		}
+
+		Data.MemoryTotal = TotalSize;
+	}
+
+	return Data;
+}
+
+#if SUBSTRATE_MATERIALCONTAINER_IS_VIEWRESOURCE
+void SerializeSubstratePixelDebugDataEntry(uint2 InCoord)
+#else
+void SerializeSubstratePixelDebugDataEntry(uint2 InCoord, FSubstrateRaytracingPayload Payload)
+#endif
+{
+	// Convert Substrate GBuffer into the CPU/GPU common data to serialize.
+#if SUBSTRATE_MATERIALCONTAINER_IS_VIEWRESOURCE
+	FSubstratePixelDebugData Data = ConvertToSerializableSubstratePixelDebugData(InCoord);
+#else
+	FSubstratePixelDebugData Data = ConvertToSerializableSubstratePixelDebugData(InCoord, Payload);
+#endif
+	
+	// And serialise it.
+	FSubstratePixelDebugDataSerializer S = (FSubstratePixelDebugDataSerializer)0;
+	SerializeSubstratePixelDebugData(S, Data);
+}
+#endif
\ No newline at end of file
diff --git a/Engine/Source/Runtime/Renderer/Private/RayTracing/RayTracingDebug.cpp b/Engine/Source/Runtime/Renderer/Private/RayTracing/RayTracingDebug.cpp
--- a/Engine/Source/Runtime/Renderer/Private/RayTracing/RayTracingDebug.cpp
+++ b/Engine/Source/Runtime/Renderer/Private/RayTracing/RayTracingDebug.cpp
@@ -160,6 +160,7 @@ class FRayTracingDebugRGS : public FGlobalShader
 		SHADER_PARAMETER_RDG_UNIFORM_BUFFER(FRayTracingDebugHitStatsUniformBufferParameters, RayTracingDebugHitStatsUniformBuffer)
 		SHADER_PARAMETER_RDG_UNIFORM_BUFFER(FRayTracingLightGrid, LightGridPacked)
 		SHADER_PARAMETER_STRUCT_INCLUDE(ShaderPrint::FShaderParameters, ShaderPrintUniformBuffer)
+		SHADER_PARAMETER_RDG_BUFFER_UAV(RWBuffer<int>, PixelMaterialDebugDataUAV)
 	END_SHADER_PARAMETER_STRUCT()
 
 	static bool ShouldCompilePermutation(const FGlobalShaderPermutationParameters& Parameters)
@@ -1587,7 +1588,6 @@ void FDeferredShadingSceneRenderer::RenderRayTracingDebug(FRDGBuilder& GraphBuil
 	const uint32 NumInstances = RayTracingScene.GetNumNativeInstances(ERayTracingSceneLayer::Base);
 
 	FRayTracingDebugRGS::FParameters* RayGenParameters = GraphBuilder.AllocParameters<FRayTracingDebugRGS::FParameters>();
-
 	RayGenParameters->VisualizationMode = DebugVisualizationMode;
 	RayGenParameters->PickerDomain = CVarRayTracingDebugPickerDomain.GetValueOnRenderThread();
 	RayGenParameters->ShouldUsePreExposure = View.Family->EngineShowFlags.Tonemapper;
@@ -1600,6 +1600,10 @@ void FDeferredShadingSceneRenderer::RenderRayTracingDebug(FRDGBuilder& GraphBuil
 	RayGenParameters->TopKMostHitInstances = CVarRayTracingDebugHitCountTopKHits.GetValueOnRenderThread();
 	RayGenParameters->NumTotalInstances = NumInstances;
 	ShaderPrint::SetParameters(GraphBuilder, View.ShaderPrintData, RayGenParameters->ShaderPrintUniformBuffer);
+
+	FSubstrateViewDebugData& SubstrateViewDebugData = View.ViewState->GetSubstrateViewDebugData();
+	FSubstrateViewDebugData::FTransientPixelDebugBuffer NewSubstratePixelDebugBuffer = SubstrateViewDebugData.CreateTransientPixelDebugBuffer(GraphBuilder);
+	RayGenParameters->PixelMaterialDebugDataUAV = NewSubstratePixelDebugBuffer.PixelMaterialDebugDataUAV;
 	
 	// If we don't output depth, create dummy 1x1 texture
 	const bool bOutputDepth = DebugVisualizationMode == RAY_TRACING_DEBUG_VIZ_INSTANCE_OVERLAP;
@@ -1665,6 +1669,11 @@ void FDeferredShadingSceneRenderer::RenderRayTracingDebug(FRDGBuilder& GraphBuil
 	{
 		DrawInstanceOverlap(GraphBuilder, Scene, View, SceneColorTexture, OutputDepthTexture);
 	}
+
+	if (DebugVisualizationMode == RAY_TRACING_DEBUG_VIZ_SUBSTRATE_DATA)
+	{
+		Substrate::AddProcessAndPrintSubstrateMaterialPropertiesPasses(GraphBuilder, View, SceneColorTexture, ShaderPlatform, NewSubstratePixelDebugBuffer);
+	}
 }
 
 void FDeferredShadingSceneRenderer::RayTracingDisplayPicking(const FRayTracingPickingFeedback& PickingFeedback, FScreenMessageWriter& Writer)
diff --git a/Engine/Source/Runtime/Renderer/Private/ScenePrivate.h b/Engine/Source/Runtime/Renderer/Private/ScenePrivate.h
--- a/Engine/Source/Runtime/Renderer/Private/ScenePrivate.h
+++ b/Engine/Source/Runtime/Renderer/Private/ScenePrivate.h
@@ -51,6 +51,7 @@
 #include "DynamicBVH.h"
 #include "OIT/OIT.h"
 #include "ShadingEnergyConservation.h"
+#include "Substrate/Substrate.h"
 #include "Substrate/Glint/GlintShadingLUTs.h"
 #include "Experimental/Containers/RobinHoodHashTable.h"
 #include "SpanAllocator.h"
@@ -933,6 +934,8 @@ private:
 
 	} EyeAdaptationManager;
 
+	FSubstrateViewDebugData SubstrateViewDebugData;
+
 	// The LUT used by tonemapping.  In stereo this is only computed and stored by the Left Eye.
 	TRefCountPtr<IPooledRenderTarget> CombinedLUTRenderTarget;
 
@@ -1410,6 +1413,11 @@ PRAGMA_ENABLE_DEPRECATION_WARNINGS
 		return CombinedLUTRenderTarget.GetReference();
 	}
 
+	FSubstrateViewDebugData& GetSubstrateViewDebugData()
+	{
+		return SubstrateViewDebugData;
+	}
+
 	// FRenderResource interface.
 	virtual void InitRHI(FRHICommandListBase& RHICmdList) override
 	{
@@ -1419,7 +1427,8 @@ PRAGMA_ENABLE_DEPRECATION_WARNINGS
 	virtual void ReleaseRHI() override
 	{
 		HZBOcclusionTests.ReleaseRHI();
-		EyeAdaptationManager.SafeRelease();
+		EyeAdaptationManager.SafeRelease(); 
+		SubstrateViewDebugData.SafeRelease();
 		OcclusionFeedback.ReleaseResource();
 PRAGMA_DISABLE_DEPRECATION_WARNINGS
 		bValidEyeAdaptationTexture = false;
diff --git a/Engine/Source/Runtime/Renderer/Private/Substrate/Substrate.h b/Engine/Source/Runtime/Renderer/Private/Substrate/Substrate.h
--- a/Engine/Source/Runtime/Renderer/Private/Substrate/Substrate.h
+++ b/Engine/Source/Runtime/Renderer/Private/Substrate/Substrate.h
@@ -185,6 +185,24 @@ struct FSubstrateViewData
 	void Reset();
 };
 
+// The substrate debug data for each view
+struct FSubstrateViewDebugData
+{
+	uint32 PixelMaterialDebugDataSizeBytes = 0;
+	TQueue<TSharedPtr<FRHIGPUBufferReadback>> ReadbackQueries;
+
+	FSubstrateViewDebugData();
+
+	struct FTransientPixelDebugBuffer
+	{
+		FRDGBufferRef PixelMaterialDebugData = nullptr;
+		FRDGBufferUAVRef PixelMaterialDebugDataUAV = nullptr;
+	};
+	FTransientPixelDebugBuffer CreateTransientPixelDebugBuffer(FRDGBuilder& GraphBuilder);
+
+	void SafeRelease();
+};
+
 namespace Substrate
 {
 constexpr uint32 StencilBit_Fast			= 0x10; // In sync with SceneRenderTargets.h - GET_STENCIL_BIT_MASK(STENCIL_SUBSTRATE_FASTPATH)
@@ -223,6 +241,8 @@ void AddSubstrateOpaqueRoughRefractionPasses(
 bool ShouldRenderSubstrateDebugPasses(const FViewInfo& View);
 FScreenPassTexture AddSubstrateDebugPasses(FRDGBuilder& GraphBuilder, const FViewInfo& View, FScreenPassTexture& ScreenPassSceneColor);
 
+void AddProcessAndPrintSubstrateMaterialPropertiesPasses(FRDGBuilder& GraphBuilder, const FViewInfo& View, FRDGTextureRef SceneColorTexture, EShaderPlatform Platform, FSubstrateViewDebugData::FTransientPixelDebugBuffer& NewSubstratePixelDebugBuffer);
+
 class FSubstrateTilePassVS : public FGlobalShader
 {
 	DECLARE_GLOBAL_SHADER(FSubstrateTilePassVS);
diff --git a/Engine/Source/Runtime/Renderer/Private/Substrate/SubstrateVisualize.cpp b/Engine/Source/Runtime/Renderer/Private/Substrate/SubstrateVisualize.cpp
--- a/Engine/Source/Runtime/Renderer/Private/Substrate/SubstrateVisualize.cpp
+++ b/Engine/Source/Runtime/Renderer/Private/Substrate/SubstrateVisualize.cpp
@@ -10,6 +10,8 @@
 #include "BasePassRendering.h"
 #include "IndirectLightRendering.h"
 #include "SubstrateVisualizationData.h"
+#include "SubstrateVisualizeDefinitions.h"
+#include "CanvasItem.h"
 
 namespace Substrate
 {
@@ -28,7 +30,6 @@ static bool IsSubstrateDebugVisualizationSupported(EShaderPlatform InPlatform, b
 		Substrate::IsSubstrateEnabled() && 
 		Substrate::UsesSubstrateMaterialBuffer(InPlatform) &&
 		GetMaxSupportedFeatureLevel(InPlatform) >= ERHIFeatureLevel::SM5 &&
-		!IsVulkanPlatform(InPlatform) &&    // Does not compile and fails to package games. Fixing this would require a deeper investigation.
         !IsMetalPlatform(InPlatform) &&     // Fails to launch editor on MAC
 		(bIsEditorOnly ? (IsPCPlatform(InPlatform) || EnumHasAllFlags(Flags, EShaderPermutationFlags::HasEditorOnlyData)) : true);
 }
@@ -46,13 +47,12 @@ class FMaterialPrintInfoCS : public FGlobalShader
 	using FPermutationDomain = TShaderPermutationDomain<>;
 
 	BEGIN_SHADER_PARAMETER_STRUCT(FParameters, )
-		SHADER_PARAMETER(uint32, ClosureIndex)
 		SHADER_PARAMETER(uint32, bOverrideCursorPosition)
 		SHADER_PARAMETER_STRUCT_REF(FViewUniformShaderParameters, ViewUniformBuffer)
 		SHADER_PARAMETER_RDG_UNIFORM_BUFFER(FSubstrateGlobalUniformParameters, Substrate)
 		SHADER_PARAMETER_STRUCT_INCLUDE(FSceneTextureParameters, SceneTextures)
 		SHADER_PARAMETER_STRUCT_INCLUDE(ShaderPrint::FShaderParameters, ShaderPrintParameters)
-		SHADER_PARAMETER_RDG_BUFFER_UAV(RWBuffer<uint>, RWPositionOffsetBuffer)
+		SHADER_PARAMETER_RDG_BUFFER_UAV(RWStructuredBuffer<int>, PixelMaterialDebugDataUAV)
 	END_SHADER_PARAMETER_STRUCT()
 
 	static bool IsSupported(EShaderPlatform InPlatform, EShaderPermutationFlags InFlags=EShaderPermutationFlags::None)
@@ -202,7 +202,8 @@ class FMaterialDebugSubstrateTreeCS : public FGlobalShader
 
 	static bool IsSupported(EShaderPlatform InPlatform, EShaderPermutationFlags InFlags=EShaderPermutationFlags::None)
 	{
-		return IsSubstrateDebugVisualizationSupported(InPlatform, true, InFlags);
+		return IsSubstrateDebugVisualizationSupported(InPlatform, true, InFlags)
+			&& !IsVulkanPlatform(InPlatform);  // SUBSTRATE_TODO Move to CPU debug visualisation and it should then work on all platforms
 	}
 
 	static FPermutationDomain RemapPermutation(FPermutationDomain PermutationVector)
@@ -247,7 +248,8 @@ class FMaterialDebugSubstrateTreePS : public FGlobalShader
 
 	static bool IsSupported(EShaderPlatform InPlatform, EShaderPermutationFlags InFlags=EShaderPermutationFlags::None)
 	{
-		return IsSubstrateDebugVisualizationSupported(InPlatform, true, InFlags);
+		return IsSubstrateDebugVisualizationSupported(InPlatform, true, InFlags)
+			&& !IsVulkanPlatform(InPlatform);  // SUBSTRATE_TODO Move to CPU debug visualisation and it should then work on all platforms
 	}
 	
 	static FPermutationDomain RemapPermutation(FPermutationDomain PermutationVector)
@@ -272,34 +274,669 @@ class FMaterialDebugSubstrateTreePS : public FGlobalShader
 };
 IMPLEMENT_GLOBAL_SHADER(FMaterialDebugSubstrateTreePS, "/Engine/Private/Substrate/SubstrateVisualize.usf", "MaterialDebugSubstrateTreePS", SF_Pixel);
 
+void AddProcessAndPrintSubstrateMaterialPropertiesPasses(FRDGBuilder& GraphBuilder, const FViewInfo& View, FRDGTextureRef SceneColorTexture, EShaderPlatform Platform, FSubstrateViewDebugData::FTransientPixelDebugBuffer& NewSubstratePixelDebugBuffer)
+{
+	if (!FMaterialPrintInfoCS::IsSupported(Platform)) return;
+	FSubstrateViewDebugData& SubstrateViewDebugData = View.ViewState->GetSubstrateViewDebugData();
+
+	////////////////////////////////////////////////////////////////////////////////////////////////////
+	// Check if the latest readback query is ready and display the data on screen.
+	TSharedPtr<FRHIGPUBufferReadback> AvailableReadback;
+	SubstrateViewDebugData.ReadbackQueries.Peek(AvailableReadback);
+	if (AvailableReadback.IsValid() && AvailableReadback->IsReady())
+	{
+		SubstrateViewDebugData.ReadbackQueries.Dequeue(AvailableReadback);
+
+		// Access the data and copy to a frame transient buffer for rendering pass.
+		void* PixelDebugData = AvailableReadback->Lock(SubstrateViewDebugData.PixelMaterialDebugDataSizeBytes);
+		if (PixelDebugData == nullptr)
+		{
+			return;
+		}
+		void* PixelDebugDataToPrint = GraphBuilder.Alloc(SubstrateViewDebugData.PixelMaterialDebugDataSizeBytes);
+		memcpy(PixelDebugDataToPrint, PixelDebugData, SubstrateViewDebugData.PixelMaterialDebugDataSizeBytes);
+
+		////////////////////////////////////////////////////////////////////////////////////////////////////
+		AddDrawCanvasPass(GraphBuilder, {}, View, FScreenPassRenderTarget(SceneColorTexture, View.ViewRect, ERenderTargetLoadAction::ELoad),
+			[&View, PixelDebugDataToPrint](FCanvas& Canvas)
+			{
+				FString Text;
+				Canvas.SetScaledToRenderTarget(true);
+
+				const float ViewPortWidth = float(View.ViewRect.Width());
+				const float ViewPortHeight = float(View.ViewRect.Height());
+				FLinearColor TextColor(1.0f, 0.5f, 0.0f);
+
+				const FLinearColor Grey = FLinearColor(0.5f, 0.5f, 0.5f, 1.0f);
+				const FLinearColor Orange = FLinearColor(243.f / 255.f, 156.f / 255.f, 18.f / 255.f, 1.0);
+				const FLinearColor CompR = FLinearColor(0.8f, 0.2f, 0.2f, 1.0);
+				const FLinearColor CompG = FLinearColor(0.2f, 0.8f, 0.2f, 1.0);
+				const FLinearColor CompB = FLinearColor(0.2f, 0.2f, 0.8f, 1.0);
+				const float DrawPosXLeft = 40.0f;
+				float DrawPosX = DrawPosXLeft;
+				float DrawPosY = 50.0f;
+
+				FSubstratePixelDebugDataSerializer S;
+				S.SubstratePixelDebugData = (int*)PixelDebugDataToPrint;
+				FSubstratePixelDebugData Data;
+				SerializeSubstratePixelDebugData(S, Data);
+
+				FString MaterialMode;
+				switch (Data.MaterialMode)
+				{
+				case HEADER_MATERIALMODE_NONE:
+					MaterialMode = TEXT("None");
+					break;
+				case HEADER_MATERIALMODE_SLAB_SIMPLE:
+					MaterialMode = TEXT("Simple BSDF");
+					break;
+				case HEADER_MATERIALMODE_SLAB_SINGLE:
+					MaterialMode = TEXT("Single Slab");
+					break;
+				case HEADER_MATERIALMODE_SLAB_COMPLEX:
+					MaterialMode = TEXT("Complex Slab(s)");
+					break;
+				case HEADER_MATERIALMODE_SLWATER:
+					MaterialMode = TEXT("Single Layer Water");
+					break;
+				case HEADER_MATERIALMODE_HAIR:
+					MaterialMode = TEXT("Hair");
+					break;
+				case HEADER_MATERIALMODE_EYE:
+					MaterialMode = TEXT("Eye");
+					break;
+				default:
+					MaterialMode = TEXT("Unkown material mode, please update visualization code.");
+				}
+				if (Data.bIsComplexSpecialMaterial)
+				{
+					MaterialMode = TEXT("Complex Special Slab(s)");
+				}
+
+				FString OptimisedLegacyMode;
+				switch (Data.OptimisedLegacyMode)
+				{
+				case SINGLE_OPTLEGACYMODE_NONE:
+					break;
+				case SINGLE_OPTLEGACYMODE_CLEARCOAT:
+					MaterialMode += TEXT(" - Legacy Clear Coat");
+					break;
+				case SINGLE_OPTLEGACYMODE_CLOTH:
+					MaterialMode += TEXT(" - Legacy Cloth");
+					break;
+				case SINGLE_OPTLEGACYMODE_SSSWRAP:
+					MaterialMode += TEXT(" - Legacy SSS-Wrap");
+					break;
+				case SINGLE_OPTLEGACYMODE_SSSPROFILE:
+					MaterialMode += TEXT(" - Legacy SSS-Profile");
+					break;
+				case SINGLE_OPTLEGACYMODE_SSSWRAP_THIN:
+					MaterialMode += TEXT(" - Legacy Foliage");
+					break;
+				default:
+					MaterialMode += TEXT(" - Unkown optimised legacy material mode, please update visualization code.");
+				}
+
+				TArray<FCanvasTextItem> TextItemsBatch; // Not very efficient but this is for debug purpose
+				FIntPoint BatchMinMaxY = FIntPoint(99999999, -99999999);
+
+				auto BeginBatch = [&]()
+				{
+					TextItemsBatch.Empty();
+					TextItemsBatch.Reserve(64);
+					BatchMinMaxY = FIntPoint(99999999, -99999999);
+				};
+				auto ExpendBatchBounds = [&](float X, float Y)
+				{
+					BatchMinMaxY = FIntPoint(FMath::Min(BatchMinMaxY.X, Y - 5), FMath::Max(BatchMinMaxY.Y, Y + 20));
+				};
+				auto DispatchBatch = [&]()
+				{
+					Canvas.DrawTile(DrawPosXLeft - 10.0f, BatchMinMaxY.X, 1000.0f, BatchMinMaxY.Y - BatchMinMaxY.X, 0, 0, 1, 1, FLinearColor(0.0, 0.0, 0.0, 0.3));
+					for (auto& Item : TextItemsBatch)
+					{
+						Canvas.DrawItem(Item);
+					}
+				};
+
+				auto PrintS = [&](float X, float Y, FString Text, FLinearColor Color = FLinearColor::White)
+				{
+					TextItemsBatch.Push(FCanvasTextItem(FVector2D(X, Y), FText::FromString(*Text), GEngine->GetSmallFont(), Color));
+					ExpendBatchBounds(X, Y);
+				};
+				auto PrintI = [&](float X, float Y, int Value, FLinearColor Color = FLinearColor::White)
+				{
+					FString String = FString::Printf(TEXT("%i"), Value);
+					TextItemsBatch.Push(FCanvasTextItem(FVector2D(X, Y), FText::FromString(*String), GetStatsFont(), Color));
+					ExpendBatchBounds(X, Y);
+				};
+				auto PrintFSmall = [&](float X, float Y, float Value, FLinearColor Color = FLinearColor::White)
+				{
+					FString String = FString::Printf(TEXT("%1.3f"), Value);
+					TextItemsBatch.Push(FCanvasTextItem(FVector2D(X, Y), FText::FromString(*String), GetStatsFont(), Color));
+					ExpendBatchBounds(X, Y);
+				};
+				auto PrintFAdapt = [&](float X, float Y, float Value, FLinearColor Color = FLinearColor::White)
+				{
+					FString String;
+					if (Value > 100)
+					{
+						String = FString::Printf(TEXT("%.0f"), Value);
+					}
+					else if (Value > 10)
+					{
+						String = FString::Printf(TEXT("%2.2f"), Value);
+					}
+					else if (Value > 1)
+					{
+						String = FString::Printf(TEXT("%1.3f"), Value);
+					}
+					else
+					{
+						String = FString::Printf(TEXT("%0.4f"), Value);
+					}
+
+					TextItemsBatch.Push(FCanvasTextItem(FVector2D(X, Y), FText::FromString(*String), GetStatsFont(), Color));
+					ExpendBatchBounds(X, Y);
+				};
+				auto PrintFColorSmall = [&](float X, float Y, float R, float G, float B)
+				{
+					PrintFSmall(X + 0.0f, Y, R, CompR);
+					PrintFSmall(X + 40.0f, Y, G, CompG);
+					PrintFSmall(X + 80.0f, Y, B, CompB);
+				};
+				auto PrintFColorAdapt = [&](float X, float Y, float R, float G, float B)
+				{
+					PrintFAdapt(X + 0.0f, Y, R, CompR);
+					PrintFAdapt(X + 50.0f, Y, G, CompG);
+					PrintFAdapt(X + 100.0f, Y, B, CompB);
+				};
+				auto PrintFBool = [&](float X, float Y, float Value)
+				{
+					FString String = Value > 0.0f ? TEXT("Yes") : TEXT("No");
+					TextItemsBatch.Push(FCanvasTextItem(FVector2D(X, Y), FText::FromString(*String), GetStatsFont(), Value > 0.0f ? FLinearColor::Green : FLinearColor::Red));
+					ExpendBatchBounds(X, Y);
+				};
+				auto NewLine = [&]()
+				{
+					DrawPosY += 17;
+				};
+
+				// Header 
+				{
+					BeginBatch();
+
+					PrintS(DrawPosX, DrawPosY, FString::Printf(TEXT("Closure Count = %i"), Data.ClosureCount), FLinearColor(0.2f, 0.8f, 0.2f, 1.0f));
+					PrintS(DrawPosX + 130.0f, DrawPosY, FString::Printf(TEXT("Material Mode = %s"), *MaterialMode), FLinearColor(0.7f, 0.7f, 0.0f, 1.0f));
+					NewLine();
+
+					PrintS(DrawPosX + 0.0f, DrawPosY, TEXT("AO"));
+					PrintS(DrawPosX + 50.0f, DrawPosY, TEXT("Ind.Irradiance"));
+					PrintS(DrawPosX + 150.0f, DrawPosY, TEXT("TopRoughness"));
+					PrintS(DrawPosX + 250.0f, DrawPosY, TEXT("PreShadow"));
+					PrintS(DrawPosX + 350.0f, DrawPosY, TEXT("ZeroShadow"));
+					PrintS(DrawPosX + 450.0f, DrawPosY, TEXT("ContactShadow"));
+					PrintS(DrawPosX + 550.0f, DrawPosY, TEXT("Ind.Occluder"));
+					PrintS(DrawPosX + 650.0f, DrawPosY, TEXT("HasSSS"));
+					PrintS(DrawPosX + 720.0f, DrawPosY, TEXT("BasisCount"));
+					NewLine();
+					PrintFSmall(DrawPosX + 0.0f, DrawPosY, Data.MaterialAO, Grey);
+					PrintFSmall(DrawPosX + 50.0f, DrawPosY, Data.IndirectIrradiance, Grey);
+					PrintFSmall(DrawPosX + 150.0f, DrawPosY, Data.TopLayerRoughness, Grey);
+					PrintFBool(DrawPosX + 250.0f, DrawPosY, Data.HasPrecShadowMask);
+					PrintFBool(DrawPosX + 350.0f, DrawPosY, Data.HasZeroPrecShadowMask);
+					PrintFBool(DrawPosX + 450.0f, DrawPosY, Data.DoesCastContactShadow);
+					PrintFBool(DrawPosX + 550.0f, DrawPosY, Data.HasDynamicIndirectShadowCasterRepresentation);
+					PrintFBool(DrawPosX + 650.0f, DrawPosY, Data.HasSubsurface);
+					PrintI(DrawPosX + 720.0f, DrawPosY, Data.LocalBasesCount, Grey);
+
+					DispatchBatch();
+				}
+
+				NewLine();
+				NewLine();
+
+				// Each Closure
+				for (int i = 0; i < Data.ClosureCount; ++i)
+				{
+					FSubsterateDebugClosure& Closure = Data.Closures[i];
+					BeginBatch();
+
+					FString Type;
+					switch (Closure.Type)
+					{
+					case SUBSTRATE_BSDF_TYPE_SLAB:
+						Type = TEXT("Slab");
+						break;
+					case SUBSTRATE_BSDF_TYPE_VOLUMETRICFOGCLOUD:
+						Type = TEXT("VolumetricFogCloud");
+						break;
+					case SUBSTRATE_BSDF_TYPE_UNLIT:
+						Type = TEXT("Unlit");
+						break;
+					case SUBSTRATE_BSDF_TYPE_HAIR:
+						Type = TEXT("Hair");
+						break;
+					case SUBSTRATE_BSDF_TYPE_SINGLELAYERWATER:
+						Type = TEXT("SingleLayerWater");
+						break;
+					case SUBSTRATE_BSDF_TYPE_EYE:
+						Type = TEXT("Eye");
+						break;
+					default:
+						MaterialMode = TEXT("Unkown closure type, please update visualization code.");
+					}
+
+					PrintS(DrawPosX, DrawPosY, Type, FLinearColor(0.2f, 1.0f, 0.2f, 1.0f));
+					PrintS(DrawPosX + 100.0f, DrawPosY, FString::Printf(TEXT("[Address=%i]"), Closure.Address), FLinearColor(0.2f, 0.5f, 0.5f, 1.0f));
+					NewLine();
+
+
+					PrintS(DrawPosX + 0.0f, DrawPosY, TEXT("NomralID"));
+					PrintS(DrawPosX + 100.0f, DrawPosY, TEXT("BasisType"));
+					PrintS(DrawPosX + 200.0f, DrawPosY, TEXT("Anisotropy"));
+					PrintS(DrawPosX + 300.0f, DrawPosY, TEXT("TopLayer"));
+					PrintS(DrawPosX + 400.0f, DrawPosY, TEXT("Scattering"));
+					PrintS(DrawPosX + 500.0f, DrawPosY, TEXT("IsThin"));
+					PrintS(DrawPosX + 550.0f, DrawPosY, TEXT("WeightV"));
+					if (Closure.bHasWeightL)
+					{
+						PrintS(DrawPosX + 700.0f, DrawPosY, TEXT("CoverAlongN"));
+						PrintS(DrawPosX + 800.0f, DrawPosY, TEXT("TransAlongN"));
+					}
+					NewLine();
+					PrintI(DrawPosX + 0.0f, DrawPosY, Closure.NormalID, Grey);
+					PrintS(DrawPosX + 100.0f, DrawPosY, Closure.BasisType == 0 ? TEXT("Normal") : TEXT("Nor+Tan"), Grey);
+					PrintFBool(DrawPosX + 200.0f, DrawPosY, Closure.bHasAnisotropy);
+					PrintFBool(DrawPosX + 300.0f, DrawPosY, Closure.bIsTopLayer);
+					PrintFBool(DrawPosX + 400.0f, DrawPosY, Closure.SSSType);		// TODO: print the special type
+					PrintFBool(DrawPosX + 500.0f, DrawPosY, Closure.bIsThin);
+					if (Closure.bHasGreyWeightV)
+					{
+						PrintFSmall(DrawPosX + 550.0f, DrawPosY, Closure.LuminanceWeightR, Grey);
+					}
+					else
+					{
+						PrintFColorSmall(DrawPosX + 550.0f, DrawPosY, Closure.LuminanceWeightR, Closure.LuminanceWeightG, Closure.LuminanceWeightB);
+					}
+					if (Closure.bHasWeightL)
+					{
+						PrintFSmall(DrawPosX + 700.0f, DrawPosY, Closure.CoverageAboveAlongN, Grey);
+						PrintFColorSmall(DrawPosX + 800.0f, DrawPosY, Closure.TransmittanceAboveAlongNR, Closure.TransmittanceAboveAlongNG, Closure.TransmittanceAboveAlongNB);
+					}
+					NewLine();
+					DrawPosY += 5;	// a little bit more space between the header and data
+
+					const float DataOffset = 130.0f;
+					if (Closure.Type == SUBSTRATE_BSDF_TYPE_SLAB)
+					{
+						PrintS(DrawPosX + 0.0f, DrawPosY, TEXT("Diffuse"), Orange);
+						PrintFColorSmall(DrawPosX + DataOffset, DrawPosY, Closure.DiffuseR, Closure.DiffuseG, Closure.DiffuseB);
+						NewLine();
+
+						PrintS(DrawPosX + 0.0f, DrawPosY, TEXT("F0"), Orange);
+						PrintFColorSmall(DrawPosX + DataOffset, DrawPosY, Closure.F0R, Closure.F0G, Closure.F0B);
+						NewLine();
+
+						PrintS(DrawPosX + 0.0f, DrawPosY, TEXT("Roughness"), Orange);
+						PrintFSmall(DrawPosX + DataOffset, DrawPosY, Closure.Roughness);
+						NewLine();
+
+						if (Closure.bHasF90)
+						{
+							PrintS(DrawPosX + 0.0f, DrawPosY, TEXT("F90"), Orange);
+							PrintFColorSmall(DrawPosX + DataOffset, DrawPosY, Closure.F90R, Closure.F90G, Closure.F90B);
+							NewLine();
+						}
+
+						if (Closure.bHasAnisotropy)
+						{
+							PrintS(DrawPosX + 0.0f, DrawPosY, TEXT("Anisotropy"), Orange);
+							PrintFSmall(DrawPosX + DataOffset, DrawPosY, Closure.Anisotropy);
+							NewLine();
+						}
+
+						if (Data.Closures[i].bHasHaziness)
+						{
+							PrintS(DrawPosX + 0.0f, DrawPosY, TEXT("Haziness"), Orange);
+							NewLine();
+
+							PrintS(DrawPosX + 0.0f, DrawPosY, TEXT(" - Weight"), Orange);
+							PrintFSmall(DrawPosX + DataOffset, DrawPosY, Data.Closures[i].HazeWeight);
+							NewLine();
+
+							PrintS(DrawPosX + 0.0f, DrawPosY, TEXT(" - Roughness"), Orange);
+							PrintFSmall(DrawPosX + DataOffset, DrawPosY, Data.Closures[i].HazeRoughness);
+							NewLine();
+
+							PrintS(DrawPosX + 0.0f, DrawPosY, TEXT(" - LegacyClearCoat"), Orange);
+							PrintFBool(DrawPosX + DataOffset, DrawPosY, Data.Closures[i].HazeSimpleClearCoatMode);
+							NewLine();
+						}
+
+						if (Closure.SSSType != SSS_TYPE_INVALID || Closure.bIsThin)
+						{
+							if (Closure.SSSType == SSS_TYPE_WRAP)
+							{
+								PrintS(DrawPosX + 0.0f, DrawPosY, TEXT("SSS Wrap (Legacy Subsurface)"), Orange);
+								NewLine();
+
+								PrintS(DrawPosX + 0.0f, DrawPosY, TEXT(" - MFP"), Orange);
+								PrintFColorAdapt(DrawPosX + DataOffset, DrawPosY, Closure.SSSMFPR, Closure.SSSMFPG, Closure.SSSMFPB);
+								NewLine();
+
+								PrintS(DrawPosX + 0.0f, DrawPosY, TEXT(" - Thickness"), Orange);
+								PrintFAdapt(DrawPosX + DataOffset, DrawPosY, Data.Closures[i].SSSThickness);
+								NewLine();
+
+								PrintS(DrawPosX + 0.0f, DrawPosY, TEXT(" - Phase"), Orange);
+								PrintFSmall(DrawPosX + DataOffset, DrawPosY, Data.Closures[i].SSSPhase);
+								NewLine();
+
+								PrintS(DrawPosX + 0.0f, DrawPosY, TEXT(" - Opacity"), Orange);
+								PrintFSmall(DrawPosX + DataOffset, DrawPosY, Data.Closures[i].SSSOpacity);
+								NewLine();
+							}
+							else if (Closure.SSSType == SSS_TYPE_SIMPLEVOLUME)
+							{
+								PrintS(DrawPosX + 0.0f, DrawPosY, TEXT("SSS Simple Volume"), Orange);
+								NewLine();
+
+								PrintS(DrawPosX + 0.0f, DrawPosY, TEXT(" - MFP"), Orange);
+								PrintFColorAdapt(DrawPosX + DataOffset, DrawPosY, Closure.SSSMFPR, Closure.SSSMFPG, Closure.SSSMFPB);
+								NewLine();
+
+								PrintS(DrawPosX + 0.0f, DrawPosY, TEXT(" - Thickness"), Orange);
+								PrintFAdapt(DrawPosX + DataOffset, DrawPosY, Data.Closures[i].SSSThickness);
+								NewLine();
+
+								PrintS(DrawPosX + 0.0f, DrawPosY, TEXT(" - Phase"), Orange);
+								PrintFSmall(DrawPosX + DataOffset, DrawPosY, Data.Closures[i].SSSPhase);
+								NewLine();
+
+								PrintS(DrawPosX + 0.0f, DrawPosY, TEXT(" - Opacity"), Orange);
+								PrintFSmall(DrawPosX + DataOffset, DrawPosY, Data.Closures[i].SSSOpacity);
+								NewLine();
+							}
+							else if (Closure.SSSType == SSS_TYPE_DIFFUSION_PROFILE)
+							{
+								PrintS(DrawPosX + 0.0f, DrawPosY, TEXT("SSS Profile"), Orange);
+								NewLine();
+
+								PrintS(DrawPosX + 0.0f, DrawPosY, TEXT(" - ProfileID"), Orange);
+								PrintI(DrawPosX + DataOffset, DrawPosY, Data.Closures[i].SSSPRofileID);
+								NewLine();
+
+								PrintS(DrawPosX + 0.0f, DrawPosY, TEXT(" - ProfileRadius"), Orange);
+								PrintFSmall(DrawPosX + DataOffset, DrawPosY, Data.Closures[i].SSSProfileRadius);
+								NewLine();
+
+								PrintS(DrawPosX + 0.0f, DrawPosY, TEXT(" - MFP"), Orange);
+								PrintFColorAdapt(DrawPosX + DataOffset, DrawPosY, Closure.SSSMFPR, Closure.SSSMFPG, Closure.SSSMFPB);
+								NewLine();
+
+								PrintS(DrawPosX + 0.0f, DrawPosY, TEXT(" - Phase"), Orange);
+								PrintFSmall(DrawPosX + DataOffset, DrawPosY, Data.Closures[i].SSSPhase);
+								NewLine();
+
+								if (Closure.bIsThin)
+								{
+									PrintS(DrawPosX + 0.0f, DrawPosY, TEXT(" - Thickness"), Orange);
+									PrintFAdapt(DrawPosX + DataOffset, DrawPosY, Data.Closures[i].SSSThickness);
+									NewLine();
+								}
+							}
+							else if (Closure.SSSType == SSS_TYPE_DIFFUSION)
+							{
+								PrintS(DrawPosX + 0.0f, DrawPosY, TEXT("SSS Substrate Per Pixel Diffusion"), Orange);
+								NewLine();
+
+								PrintS(DrawPosX + 0.0f, DrawPosY, TEXT(" - MFP"), Orange);
+								PrintFColorAdapt(DrawPosX + DataOffset, DrawPosY, Closure.SSSMFPR, Closure.SSSMFPG, Closure.SSSMFPB);
+								NewLine();
+
+								if (Closure.bIsThin)
+								{
+									PrintS(DrawPosX + 0.0f, DrawPosY, TEXT(" - MFPNorm"), Orange);
+									PrintFColorAdapt(DrawPosX + DataOffset, DrawPosY, Closure.SSSRescaledMFPR, Closure.SSSRescaledMFPG, Closure.SSSRescaledMFPB);
+									NewLine();
+
+									PrintS(DrawPosX + 0.0f, DrawPosY, TEXT(" - Thickness"), Orange);
+									PrintFAdapt(DrawPosX + DataOffset, DrawPosY, Closure.SSSThickness);
+									NewLine();
+								}
+
+								PrintS(DrawPosX + 0.0f, DrawPosY, TEXT(" - Phase"), Orange);
+								PrintFSmall(DrawPosX + DataOffset, DrawPosY, Data.Closures[i].SSSPhase);
+								NewLine();
+							}
+						}
+
+						if (Data.Closures[i].FuzzAmount > 0.0f)
+						{
+							PrintS(DrawPosX + 0.0f, DrawPosY, TEXT("Fuzz"), Orange);
+							NewLine();
+
+							PrintS(DrawPosX + 0.0f, DrawPosY, TEXT(" - FuzzAmount"), Orange);
+							PrintFSmall(DrawPosX + DataOffset, DrawPosY, Data.Closures[i].FuzzAmount);
+							NewLine();
+
+							PrintS(DrawPosX + 0.0f, DrawPosY, TEXT(" - FuzzColor"), Orange);
+							PrintFColorSmall(DrawPosX + DataOffset, DrawPosY, Closure.FuzzColorR, Closure.FuzzColorG, Closure.FuzzColorB);
+							NewLine();
+
+							PrintS(DrawPosX + 0.0f, DrawPosY, TEXT(" - FuzzRoughness"), Orange);
+							PrintFSmall(DrawPosX + DataOffset, DrawPosY, Data.Closures[i].FuzzRoughness);
+							NewLine();
+						}
+
+						if (Data.Closures[i].GlintValue < 1.0f)
+						{
+							PrintS(DrawPosX + 0.0f, DrawPosY, TEXT("Glint Density"), Orange);
+							PrintFSmall(DrawPosX + DataOffset, DrawPosY, Data.Closures[i].GlintValue);
+							NewLine();
+						}
+
+						if (Data.Closures[i].SpecProfileID >= 0)
+						{
+							PrintS(DrawPosX + 0.0f, DrawPosY, TEXT("Specular Profile"), Orange);
+							NewLine();
+
+							PrintS(DrawPosX + 0.0f, DrawPosY, TEXT(" - SpecProfileID"), Orange);
+							PrintI(DrawPosX + DataOffset, DrawPosY, Data.Closures[i].SpecProfileID);
+							NewLine();
+
+							PrintS(DrawPosX + 0.0f, DrawPosY, TEXT(" - SpecProfileMode"), Orange);
+							PrintS(DrawPosX + DataOffset, DrawPosY, Data.Closures[i].SpecProfileParameterization == 0 ? TEXT("View/Light angles") : TEXT("Half angles"));
+							NewLine();
+						}
+					}
+					else if (Closure.Type == SUBSTRATE_BSDF_TYPE_HAIR)
+					{
+						PrintS(DrawPosX + 0.0f, DrawPosY, TEXT("BaseColor"), Orange);
+						PrintFColorSmall(DrawPosX + DataOffset, DrawPosY, Closure.DiffuseR, Closure.DiffuseG, Closure.DiffuseB);
+						NewLine();
+
+						PrintS(DrawPosX + 0.0f, DrawPosY, TEXT("Specular"), Orange);
+						PrintFSmall(DrawPosX + DataOffset, DrawPosY, Closure.F0R);// Aliased variable for Hair BSDF
+						NewLine();
+
+						PrintS(DrawPosX + 0.0f, DrawPosY, TEXT("Roughness"), Orange);
+						PrintFSmall(DrawPosX + DataOffset, DrawPosY, Closure.Roughness);
+						NewLine();
+
+						PrintS(DrawPosX + 0.0f, DrawPosY, TEXT("Scattering"), Orange);
+						PrintFSmall(DrawPosX + DataOffset, DrawPosY, Closure.F90R);// Aliased variable for Hair BSDF
+						NewLine();
+
+						PrintS(DrawPosX + 0.0f, DrawPosY, TEXT("BackLit"), Orange);
+						PrintFSmall(DrawPosX + DataOffset, DrawPosY, Closure.F90G);// Aliased variable for Hair BSDF
+						NewLine();
+
+						PrintS(DrawPosX + 0.0f, DrawPosY, TEXT("Transmittance"), Orange);
+						PrintFSmall(DrawPosX + DataOffset, DrawPosY, Closure.F90B);// Aliased variable for Hair BSDF
+						NewLine();
+					}
+					else if (Closure.Type == SUBSTRATE_BSDF_TYPE_EYE)
+					{
+						PrintS(DrawPosX + 0.0f, DrawPosY, TEXT("Diffuse"), Orange);
+						PrintFColorSmall(DrawPosX + DataOffset, DrawPosY, Closure.DiffuseR, Closure.DiffuseG, Closure.DiffuseB);
+						NewLine();
+
+						PrintS(DrawPosX + 0.0f, DrawPosY, TEXT("F0"), Orange);
+						PrintFSmall(DrawPosX + DataOffset, DrawPosY, Closure.F0R);
+						NewLine();
+
+						PrintS(DrawPosX + 0.0f, DrawPosY, TEXT("Roughness"), Orange);
+						PrintFSmall(DrawPosX + DataOffset, DrawPosY, Closure.Roughness);
+						NewLine();
+
+						PrintS(DrawPosX + 0.0f, DrawPosY, TEXT("IrisMask"), Orange);
+						PrintFSmall(DrawPosX + DataOffset, DrawPosY, Closure.F90R);// Aliased variable for Eye BSDF
+						NewLine();
+
+						PrintS(DrawPosX + 0.0f, DrawPosY, TEXT("IrisDistance"), Orange);
+						PrintFSmall(DrawPosX + DataOffset, DrawPosY, Closure.F90G);// Aliased variable for Eye BSDF
+						NewLine();
+
+						PrintS(DrawPosX + 0.0f, DrawPosY, TEXT("IrisNormal"), Orange);
+						PrintFColorSmall(DrawPosX + DataOffset, DrawPosY, Closure.SSSMFPR, Closure.SSSMFPG, Closure.SSSMFPB);// Aliased variable for Eye BSDF
+						NewLine();
+
+						PrintS(DrawPosX + 0.0f, DrawPosY, TEXT("IrisPlaneNormal"), Orange);
+						PrintFColorSmall(DrawPosX + DataOffset, DrawPosY, Closure.SSSRescaledMFPR, Closure.SSSRescaledMFPG, Closure.SSSRescaledMFPB);// Aliased variable for Eye BSDF
+						NewLine();
+
+						PrintS(DrawPosX + 0.0f, DrawPosY, TEXT("SSSPRofileID"), Orange);
+						PrintI(DrawPosX + DataOffset, DrawPosY, Closure.SSSPRofileID);
+						NewLine();
+					}
+					else if (Closure.Type == SUBSTRATE_BSDF_TYPE_SINGLELAYERWATER)
+					{
+						PrintS(DrawPosX + 0.0f, DrawPosY, TEXT("BaseColor"), Orange);
+						PrintFColorSmall(DrawPosX + DataOffset, DrawPosY, Closure.DiffuseR, Closure.DiffuseG, Closure.DiffuseB);
+						NewLine();
+
+						PrintS(DrawPosX + 0.0f, DrawPosY, TEXT("Specular"), Orange);
+						PrintFSmall(DrawPosX + DataOffset, DrawPosY, Closure.F0R);// Aliased variable for SLW BSDF
+						NewLine();
+
+						PrintS(DrawPosX + 0.0f, DrawPosY, TEXT("Metallic"), Orange);
+						PrintFSmall(DrawPosX + DataOffset, DrawPosY, Closure.F0G);// Aliased variable for SLW BSDF
+						NewLine();
+
+						PrintS(DrawPosX + 0.0f, DrawPosY, TEXT("Roughness"), Orange);
+						PrintFSmall(DrawPosX + DataOffset, DrawPosY, Closure.Roughness);
+						NewLine();
+
+						PrintS(DrawPosX + 0.0f, DrawPosY, TEXT("TopMatOpacity"), Orange);
+						PrintFSmall(DrawPosX + DataOffset, DrawPosY, Closure.SSSOpacity);
+						NewLine();
+					}
+
+					DispatchBatch();
+
+					NewLine();
+					NewLine();
+				}
+
+				// Footer
+				{
+					BeginBatch();
+					const float DataOffset = 130.0f;
+
+					PrintS(DrawPosX, DrawPosY, TEXT("Memory Transactions"), FLinearColor(0.2f, 1.0f, 0.2f, 1.0f));
+					NewLine();
+
+					if (Data.MemoryDisplayMode == 0)
+					{
+						PrintS(DrawPosX + 0.0f, DrawPosY, TEXT(" - Header + BSDF"), Orange);
+						PrintI(DrawPosX + DataOffset, DrawPosY, Data.MemorySlotA);
+						NewLine();
+						PrintS(DrawPosX + 0.0f, DrawPosY, TEXT(" - TopNormalTex"), Orange);
+						PrintI(DrawPosX + DataOffset, DrawPosY, Data.MemorySlotB);
+						NewLine();
+					}
+					else if (Data.MemoryDisplayMode == 1)
+					{
+						PrintS(DrawPosX + 0.0f, DrawPosY, TEXT(" - Header"), Orange);
+						PrintI(DrawPosX + DataOffset, DrawPosY, Data.MemorySlotA);
+						NewLine();
+						PrintS(DrawPosX + 0.0f, DrawPosY, TEXT(" - TopNormalTex"), Orange);
+						PrintI(DrawPosX + DataOffset, DrawPosY, Data.MemorySlotB);
+						NewLine();
+						PrintS(DrawPosX + 0.0f, DrawPosY, TEXT(" - BSDF"), Orange);
+						PrintI(DrawPosX + DataOffset, DrawPosY, Data.MemorySlotC);
+						NewLine();
+					}
+					else if (Data.MemoryDisplayMode == 2)
+					{
+						PrintS(DrawPosX + 0.0f, DrawPosY, TEXT(" - Header + Normals"), Orange);
+						PrintI(DrawPosX + DataOffset, DrawPosY, Data.MemorySlotA);
+						NewLine();
+						PrintS(DrawPosX + 0.0f, DrawPosY, TEXT(" - BSDFs"), Orange);
+						PrintI(DrawPosX + DataOffset, DrawPosY, Data.MemorySlotB);
+						NewLine();
+					}
+					else
+					{
+						PrintS(DrawPosX + 0.0f, DrawPosY, TEXT("Unkown Memory Display Mode"), FLinearColor::Red);
+					}
+
+					if (Data.MemorySSSData > 0)
+					{
+						PrintS(DrawPosX + 0.0f, DrawPosY, TEXT(" - SSS Data"), Orange);
+						PrintI(DrawPosX + DataOffset, DrawPosY, Data.MemorySSSData);
+						NewLine();
+					}
+
+					PrintS(DrawPosX + 0.0f, DrawPosY, TEXT(" - Total"), Orange);
+					PrintI(DrawPosX + DataOffset, DrawPosY, Data.MemoryTotal);
+					NewLine();
+
+					DispatchBatch();
+				}
+			});
+		////////////////////////////////////////////////////////////////////////////////////////////////////
+
+		AvailableReadback->Unlock();
+	}
+
+	////////////////////////////////////////////////////////////////////////////////////////////////////
+	// Copy the debug data on GPU into a readback buffer for debug display when available later.
+	TSharedPtr<FRHIGPUBufferReadback> NewReadBack = MakeShared<FRHIGPUBufferReadback>(TEXT("PixelMaterialDebugDataReadback"));
+	AddEnqueueCopyPass(GraphBuilder, NewReadBack.Get(), NewSubstratePixelDebugBuffer.PixelMaterialDebugData, SubstrateViewDebugData.PixelMaterialDebugDataSizeBytes);
+	SubstrateViewDebugData.ReadbackQueries.Enqueue(NewReadBack);
+}
+
 static void AddVisualizeMaterialPropertiesPasses(FRDGBuilder& GraphBuilder, const FViewInfo& View, FScreenPassTexture& ScreenPassSceneColor, EShaderPlatform Platform)
 {
 	if (!FMaterialPrintInfoCS::IsSupported(Platform)) return;
+	FSubstrateViewDebugData& SubstrateViewDebugData = View.ViewState->GetSubstrateViewDebugData();
 
 	// Force ShaderPrint on.
 	ShaderPrint::SetEnabled(true);
 	ShaderPrint::RequestSpaceForLines(1024);
 	ShaderPrint::RequestSpaceForCharacters(1024);
-	FRDGBufferUAVRef PrintOffsetBufferUAV = nullptr;
 
-	FRDGBufferRef PrintOffsetBuffer = GraphBuilder.CreateBuffer(FRDGBufferDesc::CreateBufferDesc(4, 2), TEXT("Substrate.DebugPrintPositionOffset"));
-	PrintOffsetBufferUAV = GraphBuilder.CreateUAV(PrintOffsetBuffer, PF_R32_UINT);
-	AddClearUAVPass(GraphBuilder, PrintOffsetBufferUAV, 50u);
+	////////////////////////////////////////////////////////////////////////////////////////////////////
+	// Get a new pixel material buffer and render debug data into it
+	FSubstrateViewDebugData::FTransientPixelDebugBuffer NewSubstratePixelDebugBuffer = SubstrateViewDebugData.CreateTransientPixelDebugBuffer(GraphBuilder);
 
-	for (uint32 ClosureIndex=0; ClosureIndex < SUBSTRATE_MAX_CLOSURE_COUNT; ++ClosureIndex)
 	{
 		FMaterialPrintInfoCS::FParameters* PassParameters = GraphBuilder.AllocParameters<FMaterialPrintInfoCS::FParameters>();
-		PassParameters->ClosureIndex = ClosureIndex;
 		PassParameters->bOverrideCursorPosition = WITH_EDITOR ? 0u : 1u;
-		PassParameters->RWPositionOffsetBuffer = PrintOffsetBufferUAV;
 		PassParameters->ViewUniformBuffer = View.ViewUniformBuffer;
 		PassParameters->Substrate = Substrate::BindSubstrateGlobalUniformParameters(View);
 		PassParameters->SceneTextures = GetSceneTextureParameters(GraphBuilder, View);
+		PassParameters->PixelMaterialDebugDataUAV = NewSubstratePixelDebugBuffer.PixelMaterialDebugDataUAV;
 		ShaderPrint::SetParameters(GraphBuilder, View.ShaderPrintData, PassParameters->ShaderPrintParameters);
 
 		TShaderMapRef<FMaterialPrintInfoCS> ComputeShader(View.ShaderMap);
-		FComputeShaderUtils::AddPass(GraphBuilder, RDG_EVENT_NAME("Substrate::VisualizeMaterial(Print, Closure=%d)", ClosureIndex), ComputeShader, PassParameters, FIntVector(1, 1, 1));
+		FComputeShaderUtils::AddPass(GraphBuilder, RDG_EVENT_NAME("Substrate::VisualizeMaterial(Print)"), ComputeShader, PassParameters, FIntVector(1, 1, 1));
 	}
+
+	AddProcessAndPrintSubstrateMaterialPropertiesPasses(GraphBuilder, View, ScreenPassSceneColor.Texture, Platform, NewSubstratePixelDebugBuffer);
 }
 
 static void AddVisualizeMaterialCountPasses(FRDGBuilder & GraphBuilder, const FViewInfo & View, FScreenPassTexture & ScreenPassSceneColor, EShaderPlatform Platform, uint32 ViewMode)
@@ -522,3 +1159,35 @@ FScreenPassTexture AddSubstrateDebugPasses(FRDGBuilder& GraphBuilder, const FVie
 }
 
 } // namespace Substrate
+
+FSubstrateViewDebugData::FSubstrateViewDebugData()
+{
+	//
+}
+
+FSubstrateViewDebugData::FTransientPixelDebugBuffer FSubstrateViewDebugData::CreateTransientPixelDebugBuffer(FRDGBuilder& GraphBuilder)
+{
+	FSubstrateViewDebugData::FTransientPixelDebugBuffer PixelDebugBuffer;
+
+	const uint32 SizeOfUint = sizeof(uint32);
+	const uint32 NumUints = FMath::DivideAndRoundUp(uint32(sizeof(FSubstratePixelDebugData)), SizeOfUint);
+	PixelMaterialDebugDataSizeBytes = NumUints * SizeOfUint;	// Arbitrary for now
+
+	FRDGBufferDesc RDGBufferDesc = FRDGBufferDesc::CreateStructuredDesc(SizeOfUint, NumUints);
+	RDGBufferDesc.Usage |= BUF_SourceCopy;
+	PixelDebugBuffer.PixelMaterialDebugData = GraphBuilder.CreateBuffer(RDGBufferDesc, TEXT("PixelMaterialDebugData"));
+
+	PixelDebugBuffer.PixelMaterialDebugDataUAV = GraphBuilder.CreateUAV(FRDGBufferUAVDesc(PixelDebugBuffer.PixelMaterialDebugData, PF_R32_UINT));
+
+	return PixelDebugBuffer;
+}
+
+void FSubstrateViewDebugData::SafeRelease()
+{
+	TSharedPtr<FRHIGPUBufferReadback> It;
+	while (ReadbackQueries.Dequeue(It))
+	{
+		It.Reset();
+	}
+}
+
