MGPU:  fix for flickering artifacts on Virtual Shadow Maps, due to corrupt cache.  Uses flags to independently track whether cache is valid per GPU, rather than assuming the cache is valid when resource pointers exist.  The initialization code path runs if the cache is not yet valid on all the GPUs the view is running on.  PhysicalPageLists also needs to be propagated from the previous to current resource for GPUs the view is NOT running on, so they stay the same, rather than containing uninitialized data.

#rnx

diff --git a/Engine/Source/Runtime/RHI/Public/MultiGPU.h b/Engine/Source/Runtime/RHI/Public/MultiGPU.h
--- a/Engine/Source/Runtime/RHI/Public/MultiGPU.h
+++ b/Engine/Source/Runtime/RHI/Public/MultiGPU.h
@@ -192,6 +192,21 @@ public:
 		return FRHIGPUMask(~((1u << GPUIndex) - 1)) & All();
 	}
 
+	// Inverts a GPU mask, returning true if the inverse succeeded.  If it fails, OutInverse is arbitrarily set to GPU0.
+	FORCEINLINE bool Invert(FRHIGPUMask& OutInverse) const
+	{
+		if (*this == All())
+		{
+			OutInverse = FRHIGPUMask::GPU0();
+			return false;
+		}
+		else
+		{
+			OutInverse = FRHIGPUMask(~GetMask()) & All();
+			return true;
+		}
+	}
+
 	struct FIterator
 	{
 		FORCEINLINE explicit FIterator(const uint32 InGPUMask)
diff --git a/Engine/Source/Runtime/Renderer/Private/VirtualShadowMaps/VirtualShadowMapArray.cpp b/Engine/Source/Runtime/Renderer/Private/VirtualShadowMaps/VirtualShadowMapArray.cpp
--- a/Engine/Source/Runtime/Renderer/Private/VirtualShadowMaps/VirtualShadowMapArray.cpp
+++ b/Engine/Source/Runtime/Renderer/Private/VirtualShadowMaps/VirtualShadowMapArray.cpp
@@ -2165,15 +2165,17 @@ void FVirtualShadowMapArray::BuildPageAllocations(
 
 	// Update cached or newly invalidated pages with respect to the new requests
 	{	
-		// Cached data from previous frames is available and valid
+		// Cached data from previous frames is available and valid.  Note that we currently don't support GPUMask varying within
+		// a view family, so just use the first view's GPU mask.
 		const bool bCacheDataAvailable = CacheManager->IsCacheDataAvailable();
+		const bool bCacheDataValid = bCacheDataAvailable && CacheManager->GetCacheValidGPUMask().ContainsAll(Views[0].GPUMask);
 
 		FUpdatePhysicalPages::FParameters* PassParameters = GraphBuilder.AllocParameters<FUpdatePhysicalPages::FParameters>();
 		PassParameters->VirtualShadowMap		= GetUncachedUniformBuffer(GraphBuilder);
 		PassParameters->OutPhysicalPageMetaData = GraphBuilder.CreateUAV(PhysicalPageMetaDataRDG);
 		PassParameters->OutPhysicalPageLists	= GraphBuilder.CreateUAV(PhysicalPageListsRDG);
 
-		if (bCacheDataAvailable)
+		if (bCacheDataValid)
 		{
 			PassParameters->PageRequestFlags		   = GraphBuilder.CreateSRV(PageRequestFlagsRDG);
 			PassParameters->OutPageTable			   = GraphBuilder.CreateUAV(PageTableRDG);
@@ -2188,7 +2190,7 @@ void FVirtualShadowMapArray::BuildPageAllocations(
 		}
 
 		FUpdatePhysicalPages::FPermutationDomain PermutationVector;
-		PermutationVector.Set<FUpdatePhysicalPages::FHasCacheDataDim>(bCacheDataAvailable);
+		PermutationVector.Set<FUpdatePhysicalPages::FHasCacheDataDim>(bCacheDataValid);
 		SetStatsArgsAndPermutation<FUpdatePhysicalPages>(ShouldGenerateStats(), StatsBufferUAV, PassParameters, PermutationVector);
 		auto ComputeShader = Views[0].ShaderMap->GetShader<FUpdatePhysicalPages>(PermutationVector);
 
@@ -2199,6 +2201,21 @@ void FVirtualShadowMapArray::BuildPageAllocations(
 			PassParameters,
 			FIntVector(FMath::DivideAndRoundUp(GetMaxPhysicalPages(), FUpdatePhysicalPages::DefaultCSGroupX), 1, 1)
 		);
+
+#if WITH_MGPU
+		// Need to propagate previous physical page lists for GPUs we didn't update in the above pass.  This is only a
+		// 64 KB structure, so copying it is cheap.  Ping ponging and independently tracking which is the current buffer
+		// separately per GPU is another option, which avoids the copy, but adds a bunch of complexity.
+		FRHIGPUMask InverseGPUMask;
+		if (CacheManager->GetPrevBuffers().PhysicalPageLists && Views[0].GPUMask.Invert(InverseGPUMask) && InverseGPUMask.Intersects(CacheManager->GetCacheValidGPUMask()))
+		{
+			RDG_GPU_MASK_SCOPE(GraphBuilder, InverseGPUMask);
+			AddCopyBufferPass(GraphBuilder, PhysicalPageListsRDG, GraphBuilder.RegisterExternalBuffer(CacheManager->GetPrevBuffers().PhysicalPageLists));
+		}
+
+		// Track which GPUs the cache has been initialized on.  Merges GPU mask if cache data was already available, otherwise sets mask to initialize it.
+		CacheManager->UpdateCacheValidGPUMask(Views[0].GPUMask, bCacheDataAvailable);
+#endif
 	}
 
 	{
diff --git a/Engine/Source/Runtime/Renderer/Private/VirtualShadowMaps/VirtualShadowMapCacheManager.cpp b/Engine/Source/Runtime/Renderer/Private/VirtualShadowMaps/VirtualShadowMapCacheManager.cpp
--- a/Engine/Source/Runtime/Renderer/Private/VirtualShadowMaps/VirtualShadowMapCacheManager.cpp
+++ b/Engine/Source/Runtime/Renderer/Private/VirtualShadowMaps/VirtualShadowMapCacheManager.cpp
@@ -1635,6 +1635,8 @@ void FVirtualShadowMapArrayCacheManager::ProcessInvalidations(
 	const FInvalidationPassCommon& InvalidationPassCommon,
 	const FInstanceGPULoadBalancer& Instances) const
 {
+	RDG_GPU_MASK_SCOPE(GraphBuilder, CacheValidGPUMask);
+
 	check(InvalidationPassCommon.UniformParameters->NumFullShadowMaps > 0);
 	check(!Instances.IsEmpty());
 
diff --git a/Engine/Source/Runtime/Renderer/Private/VirtualShadowMaps/VirtualShadowMapCacheManager.h b/Engine/Source/Runtime/Renderer/Private/VirtualShadowMaps/VirtualShadowMapCacheManager.h
--- a/Engine/Source/Runtime/Renderer/Private/VirtualShadowMaps/VirtualShadowMapCacheManager.h
+++ b/Engine/Source/Runtime/Renderer/Private/VirtualShadowMaps/VirtualShadowMapCacheManager.h
@@ -283,6 +283,32 @@ public:
 	bool IsCacheDataAvailable();
 	bool IsHZBDataAvailable();
 
+	FRHIGPUMask GetCacheValidGPUMask() const
+	{
+#if WITH_MGPU
+		return CacheValidGPUMask;
+#else
+		return FRHIGPUMask::GPU0();
+#endif
+	}
+
+	void UpdateCacheValidGPUMask(FRHIGPUMask GPUMask, bool bMergeMask)
+	{
+#if WITH_MGPU
+		if (bMergeMask)
+		{
+			CacheValidGPUMask |= GPUMask;
+		}
+		else
+		{
+			// To handle initialization when first allocating cache resources, we overwrite the mask.  This is necessary because the FRHIGPUMask doesn't
+			// support empty masks.  Also, this deals with cases where the cache is cleared -- the cache resources will be missing, and it can use this
+			// code path to set the mask to a known state when they get re-created.
+			CacheValidGPUMask = GPUMask;
+		}
+#endif
+	}
+
 	bool IsAccumulatingStats();
 
 	using FInstanceGPULoadBalancer = TInstanceCullingLoadBalancer<SceneRenderingAllocator>;
@@ -471,6 +497,10 @@ private:
 
 	FScene* Scene;
 	FShadowInvalidatingInstancesImplementation ShadowInvalidatingInstancesImplementation;
+
+#if WITH_MGPU
+	FRHIGPUMask CacheValidGPUMask;
+#endif
 };
 
 
