Substrate - converted system info debugging to canvas draw with gpu readback. Now all debug info should show up for all the render RHI/API.

#rnx

diff --git a/Engine/Shaders/Private/Substrate/SubstrateVisualize.usf b/Engine/Shaders/Private/Substrate/SubstrateVisualize.usf
--- a/Engine/Shaders/Private/Substrate/SubstrateVisualize.usf
+++ b/Engine/Shaders/Private/Substrate/SubstrateVisualize.usf
@@ -965,23 +965,7 @@ void MaterialDebugSubstrateTreePS(
 // Material Print
 
 #if SHADER_SYSTEMINFO
-uint bAdvancedDebugEnabled;
-uint bEnergyConservation;
-uint bEnergyPreservation;
-uint bDbufferPass;
-uint ClassificationCMask;
-uint ClassificationAsync;
-uint Classification8bits;
-uint bRoughRefraction;
-uint bUseClosureCountFromMaterialData;
-uint ViewsMaxBytesPerPixel;
-uint ViewsMaxClosurePerPixel;
-uint ProjectMaxBytesPerPixel;
-uint ProjectMaxClosurePerPixel;
-uint MaterialBufferAllocationInBytes;
-uint MaterialBufferAllocationMode;
-
-uint MaxClosureCount;
+
 Buffer<uint> ClassificationTileDrawIndirectBuffer;
 
 uint GetTileCount(uint InType)
@@ -992,199 +976,9 @@ uint GetTileCount(uint InType)
 [numthreads(1, 1, 1)]
 void MainCS(uint3 DispatchThreadId : SV_DispatchThreadID)
 {
-	// Print Legend
 	if (all(DispatchThreadId == 0))
 	{
-		FFontColor TitleColor  = FontYellow;
-		FFontColor LegendColor = FontWhite;
-		FFontColor ValueColor  = FontOrange;
-		FShaderPrintContext Context = InitShaderPrintContext(true, uint2(50, 50));
-
-		FFontColor SubTitleColor= FontYellow;
-		SubTitleColor.Color *= 0.7f;
-
-		// General
-		Print(Context, TEXT("General"), TitleColor);
-		Newline(Context);
-
-		Print(Context, TEXT(" Byte Per Pixel"), SubTitleColor); Newline(Context);
-		Print(Context, TEXT(" Max   : "), LegendColor); 
-		Print(Context, Substrate.MaxBytesPerPixel, ValueColor, 3, 0);
-		Print(Context, TEXT(" / "), FontWhite); 
-		Print(Context, ProjectMaxBytesPerPixel, ValueColor, 3, 0);
-		Print(Context, TEXT(" ["), FontWhite); 
-		Print(Context, uint(DivideAndRoundUp(MaterialBufferAllocationInBytes, 1000000.f)), ValueColor, 4, 0);
-		Print(Context, TEXT("Mb]"), FontWhite);
-		Newline(Context);
-
-		Print(Context, TEXT(" Views : "), LegendColor); 
-		Print(Context, ViewsMaxBytesPerPixel, ValueColor, 3, 0);
-		Print(Context, TEXT(" / "), FontWhite); 
-		Print(Context, ProjectMaxBytesPerPixel, ValueColor, 3, 0);
-		Newline(Context);
-		Newline(Context);
-
-		Print(Context, TEXT(" Closure Per Pixel"), SubTitleColor); Newline(Context);
-		Print(Context, TEXT(" Max   : "), LegendColor); 
-		Print(Context, Substrate.MaxClosurePerPixel, ValueColor, 3, 0);
-		Print(Context, TEXT(" / "), FontWhite); 
-		Print(Context, ProjectMaxClosurePerPixel, ValueColor, 3, 0);
-		Newline(Context);
-
-		Print(Context, TEXT(" Views : "), LegendColor); 
-		Print(Context, ViewsMaxClosurePerPixel, ValueColor, 3, 0);
-		Print(Context, TEXT(" / "), FontWhite); 
-		Print(Context, ProjectMaxClosurePerPixel, ValueColor, 3, 0);
-		Newline(Context);
-		Newline(Context);
-
-		Print(Context, TEXT(" Allocation Mode    : "), LegendColor);
-		switch (MaterialBufferAllocationMode)
-		{
-			case 0: Print(Context, TEXT("View based"), ValueColor); break;
-			case 1: Print(Context, TEXT("View based | Growing-only"), ValueColor); break;
-			case 2: Print(Context, TEXT("Setting based"), ValueColor); break;
-			default: Print(Context, TEXT("Unknown: "), ValueColor); break;
-		}
-		Newline(Context);
-
-		Print(Context, TEXT(" Shading quality    : "), LegendColor);
-		Print(Context, uint(SUBSTRATE_SHADING_QUALITY), ValueColor);
-		Newline(Context);
-
-		Print(Context, TEXT(" Rough diffuse      : "), LegendColor);
-		PrintBool(Context, Substrate.bRoughDiffuse);
-		Newline(Context);
-
-		Print(Context, TEXT(" Energy conservation: "), LegendColor);
-		PrintBool(Context, bEnergyConservation);
-		Newline(Context);
-
-		Print(Context, TEXT(" Energy preservation: "), LegendColor);
-		PrintBool(Context, bEnergyPreservation);
-		Newline(Context);
-
-		Print(Context, TEXT(" Use #closures mat. : "), LegendColor);
-		PrintBool(Context, bUseClosureCountFromMaterialData);
-		Newline(Context); 
-
-		Print(Context, TEXT(" Tile Closure Count : "), LegendColor);
-		Print(Context, MaxClosureCount, FontOrange);
-		Newline(Context); 
-
-		Print(Context, TEXT(" DBuffer pass       : "), LegendColor);
-		PrintBool(Context, bDbufferPass);
-		Newline(Context);
-
-		Print(Context, TEXT(" Rough refraction   : "), LegendColor);
-		PrintBool(Context, bRoughRefraction);
-		Newline(Context);
-
-		Newline(Context);
-
-		// Classification
-		Print(Context, TEXT("Classification"), TitleColor);
-		Newline(Context);
-
-		Print(Context, TEXT(" Tile Size          : "), LegendColor);
-		Print(Context, Substrate.TileSize, ValueColor, 1, 0);
-		Print(Context, TEXT("x"), LegendColor);
-		Print(Context, Substrate.TileSize, ValueColor, 1, 0);
-		Newline(Context);
-
-		Print(Context, TEXT(" Tile Count         : "), LegendColor);
-		Print(Context, Substrate.TileCount.x, ValueColor, 3, 0);
-		Print(Context, TEXT("x"), LegendColor);
-		Print(Context, Substrate.TileCount.y, ValueColor, 3, 0);
-		Newline(Context);
-
-		Print(Context, TEXT(" CMask              : "), LegendColor);
-		PrintBool(Context, ClassificationCMask);
-		Newline(Context);
-
-		Print(Context, TEXT(" Async              : "), LegendColor);
-		PrintBool(Context, ClassificationAsync);
-		Newline(Context);
-
-		Print(Context, TEXT(" 8-bits coord       : "), LegendColor);
-		PrintBool(Context, Classification8bits);
-		Newline(Context);
-
-		Newline(Context);
-
-		// Material tiles
-		Print(Context, TEXT("Material Tiles"), TitleColor);
-		Newline(Context);
-
-		Print(Context, TEXT(" # Simple tiles     : "), LegendColor);
-		Print(Context, GetTileCount(SUBSTRATE_TILE_TYPE_SIMPLE), ValueColor);
-		Newline(Context);
-
-		Print(Context, TEXT(" # Single tiles     : "), LegendColor);
-		Print(Context, GetTileCount(SUBSTRATE_TILE_TYPE_SINGLE), ValueColor);
-		Newline(Context);
-
-		Print(Context, TEXT(" # Complex tiles    : "), LegendColor);
-		Print(Context, GetTileCount(SUBSTRATE_TILE_TYPE_COMPLEX), ValueColor);
-		Newline(Context);
-
-		Print(Context, TEXT(" # Complex Special tiles: "), LegendColor);
-		Print(Context, GetTileCount(SUBSTRATE_TILE_TYPE_COMPLEX_SPECIAL), ValueColor);
-		Newline(Context);
-
-		Newline(Context);
-
-		// Rough tiles
-		if (bRoughRefraction)
-		{
-			Print(Context, TEXT("Rough Refraction Tiles"), TitleColor);
-			Newline(Context);
-
-			Print(Context, TEXT(" # Rough refract.   : "), LegendColor);
-			Print(Context, GetTileCount(SUBSTRATE_TILE_TYPE_ROUGH_REFRACT), ValueColor);
-			Newline(Context);
-
-			Print(Context, TEXT(" # Rough r. w/o SSS : "), LegendColor);
-			Print(Context, GetTileCount(SUBSTRATE_TILE_TYPE_ROUGH_REFRACT_SSS_WITHOUT), ValueColor);
-			Newline(Context);
-
-			Newline(Context);
-		}
-
-		// DBuffer Tiles
-		if (bDbufferPass)
-		{
-			Print(Context, TEXT("Decal Tiles"), TitleColor);
-			Newline(Context);
-
-			Print(Context, TEXT(" # Simple tiles     : "), LegendColor);
-			Print(Context, GetTileCount(SUBSTRATE_TILE_TYPE_DECAL_SIMPLE), ValueColor);
-			Newline(Context);
-
-			Print(Context, TEXT(" # Single tiles     : "), LegendColor);
-			Print(Context, GetTileCount(SUBSTRATE_TILE_TYPE_DECAL_SINGLE), ValueColor);
-			Newline(Context);
-
-			Print(Context, TEXT(" # Complex tiles    : "), LegendColor);
-			Print(Context, GetTileCount(SUBSTRATE_TILE_TYPE_DECAL_COMPLEX), ValueColor);
-			Newline(Context);
-
-			Newline(Context);
-		}
-
-		// Debug
-		Print(Context, TEXT("Debug"), TitleColor);
-		Newline(Context);
-
-		Print(Context, TEXT(" Adv. debug         : "), LegendColor);
-		PrintBool(Context, bAdvancedDebugEnabled);
-		Newline(Context);
-
-		Print(Context, TEXT(" Layer peel index   : "), LegendColor);
-		Print(Context, Substrate.PeelLayersAboveDepth, ValueColor);
-		Newline(Context);
-
-		Newline(Context);
+		ConvertToSerializableSubstratePixelDebugData();
 	}
 }
 #endif // SHADER_SYSTEMINFO
\ No newline at end of file
diff --git a/Engine/Shaders/Shared/SubstrateVisualizeDefinitions.h b/Engine/Shaders/Shared/SubstrateVisualizeDefinitions.h
--- a/Engine/Shaders/Shared/SubstrateVisualizeDefinitions.h
+++ b/Engine/Shaders/Shared/SubstrateVisualizeDefinitions.h
@@ -103,23 +103,23 @@ struct FSubstratePixelDebugData
 };
 
 #ifndef __cplusplus
-RWStructuredBuffer<int> PixelMaterialDebugDataUAV;
-struct FSubstratePixelDebugDataSerializer
+RWStructuredBuffer<int> SubstrateDebugDataUAV;
+struct FSubstrateDebugDataSerializer
 {
 	int WriteIndex;
 
 	void Serialize(int Data)
 	{
-		PixelMaterialDebugDataUAV[WriteIndex++] = Data;
+		SubstrateDebugDataUAV[WriteIndex++] = Data;
 	}
 
 	void Serialize(float Data)
 	{
-		PixelMaterialDebugDataUAV[WriteIndex++] = asuint(Data);
+		SubstrateDebugDataUAV[WriteIndex++] = asuint(Data);
 	}
 };
 #else
-struct FSubstratePixelDebugDataSerializer
+struct FSubstrateDebugDataSerializer
 {
 	int ReadIndex = 0;
 	int* SubstratePixelDebugData = nullptr;
@@ -137,9 +137,9 @@ struct FSubstratePixelDebugDataSerializer
 #endif
 
 #ifndef __cplusplus
-void SerializeSubstratePixelDebugData(in FSubstratePixelDebugDataSerializer S, in FSubstratePixelDebugData D)
+void SerializeSubstratePixelDebugData(in FSubstrateDebugDataSerializer S, in FSubstratePixelDebugData D)
 #else
-void SerializeSubstratePixelDebugData(FSubstratePixelDebugDataSerializer& S, FSubstratePixelDebugData& D)
+void SerializeSubstratePixelDebugData(FSubstrateDebugDataSerializer& S, FSubstratePixelDebugData& D)
 #endif
 {
 	S.Serialize(D.ClosureCount);
@@ -601,7 +601,52 @@ void SerializeSubstratePixelDebugDataEntry(uint2 InCoord, FSubstrateRaytracingPa
 #endif
 	
 	// And serialise it.
-	FSubstratePixelDebugDataSerializer S = (FSubstratePixelDebugDataSerializer)0;
+	FSubstrateDebugDataSerializer S = (FSubstrateDebugDataSerializer)0;
 	SerializeSubstratePixelDebugData(S, Data);
 }
-#endif
\ No newline at end of file
+#endif
+
+
+////////////////////////////////////////////////////////////////////////////////////////////////////
+////////////////////////////////////////////////////////////////////////////////////////////////////
+////////////////////////////////////////////////////////////////////////////////////////////////////
+
+
+struct FSubstrateSystemInfoData
+{
+	int TileCount[SUBSTRATE_TILE_TYPE_COUNT];
+};
+
+#ifndef __cplusplus
+void SerializeSubstrateSystemInfoDebugData(in FSubstrateDebugDataSerializer S, in FSubstrateSystemInfoData D)
+#else
+void SerializeSubstrateSystemInfoDebugData(FSubstrateDebugDataSerializer& S, FSubstrateSystemInfoData& D)
+#endif
+{
+	for (int i = 0; i < SUBSTRATE_TILE_TYPE_COUNT; ++i)
+	{
+		S.Serialize(D.TileCount[i]);
+	}
+}
+
+#ifndef __cplusplus
+
+uint GetTileCount(uint InType);
+
+void ConvertToSerializableSubstratePixelDebugData()
+{
+	FSubstrateSystemInfoData Data = (FSubstrateSystemInfoData)0;
+
+	// Setup the system info to send to the CPU for debug print
+	for (int i = 0; i < SUBSTRATE_TILE_TYPE_COUNT; ++i)
+	{
+		Data.TileCount[i] = GetTileCount(i);
+	}
+
+	// And serialise it.
+	FSubstrateDebugDataSerializer S = (FSubstrateDebugDataSerializer)0;
+	SerializeSubstrateSystemInfoDebugData(S, Data);
+}
+
+#endif
+
diff --git a/Engine/Source/Runtime/Renderer/Private/RayTracing/RayTracingDebug.cpp b/Engine/Source/Runtime/Renderer/Private/RayTracing/RayTracingDebug.cpp
--- a/Engine/Source/Runtime/Renderer/Private/RayTracing/RayTracingDebug.cpp
+++ b/Engine/Source/Runtime/Renderer/Private/RayTracing/RayTracingDebug.cpp
@@ -160,7 +160,7 @@ class FRayTracingDebugRGS : public FGlobalShader
 		SHADER_PARAMETER_RDG_UNIFORM_BUFFER(FRayTracingDebugHitStatsUniformBufferParameters, RayTracingDebugHitStatsUniformBuffer)
 		SHADER_PARAMETER_RDG_UNIFORM_BUFFER(FRayTracingLightGrid, LightGridPacked)
 		SHADER_PARAMETER_STRUCT_INCLUDE(ShaderPrint::FShaderParameters, ShaderPrintUniformBuffer)
-		SHADER_PARAMETER_RDG_BUFFER_UAV(RWBuffer<int>, PixelMaterialDebugDataUAV)
+		SHADER_PARAMETER_RDG_BUFFER_UAV(RWBuffer<int>, SubstrateDebugDataUAV)
 	END_SHADER_PARAMETER_STRUCT()
 
 	static bool ShouldCompilePermutation(const FGlobalShaderPermutationParameters& Parameters)
@@ -1603,7 +1603,7 @@ void FDeferredShadingSceneRenderer::RenderRayTracingDebug(FRDGBuilder& GraphBuil
 
 	FSubstrateViewDebugData& SubstrateViewDebugData = View.ViewState->GetSubstrateViewDebugData();
 	FSubstrateViewDebugData::FTransientPixelDebugBuffer NewSubstratePixelDebugBuffer = SubstrateViewDebugData.CreateTransientPixelDebugBuffer(GraphBuilder);
-	RayGenParameters->PixelMaterialDebugDataUAV = NewSubstratePixelDebugBuffer.PixelMaterialDebugDataUAV;
+	RayGenParameters->SubstrateDebugDataUAV = NewSubstratePixelDebugBuffer.PixelMaterialDebugDataUAV;
 	
 	// If we don't output depth, create dummy 1x1 texture
 	const bool bOutputDepth = DebugVisualizationMode == RAY_TRACING_DEBUG_VIZ_INSTANCE_OVERLAP;
diff --git a/Engine/Source/Runtime/Renderer/Private/Substrate/Substrate.h b/Engine/Source/Runtime/Renderer/Private/Substrate/Substrate.h
--- a/Engine/Source/Runtime/Renderer/Private/Substrate/Substrate.h
+++ b/Engine/Source/Runtime/Renderer/Private/Substrate/Substrate.h
@@ -189,7 +189,10 @@ struct FSubstrateViewData
 struct FSubstrateViewDebugData
 {
 	uint32 PixelMaterialDebugDataSizeBytes = 0;
-	TQueue<TSharedPtr<FRHIGPUBufferReadback>> ReadbackQueries;
+	TQueue<TSharedPtr<FRHIGPUBufferReadback>> PixelMaterialDebugDataReadbackQueries;
+
+	uint32 SystemInfoDebugDataSizeBytes = 0;
+	TQueue<TSharedPtr<FRHIGPUBufferReadback>> SystemInfoDebugDataReadbackQueries;
 
 	FSubstrateViewDebugData();
 
@@ -200,6 +203,13 @@ struct FSubstrateViewDebugData
 	};
 	FTransientPixelDebugBuffer CreateTransientPixelDebugBuffer(FRDGBuilder& GraphBuilder);
 
+	struct FTransientSystemInfoDebugBuffer
+	{
+		FRDGBufferRef SystemInfoDebugData = nullptr;
+		FRDGBufferUAVRef SystemInfoDebugDataUAV = nullptr;
+	};
+	FTransientSystemInfoDebugBuffer CreateTransientSystemInfoDebugBuffer(FRDGBuilder& GraphBuilder);
+
 	void SafeRelease();
 };
 
diff --git a/Engine/Source/Runtime/Renderer/Private/Substrate/SubstrateVisualize.cpp b/Engine/Source/Runtime/Renderer/Private/Substrate/SubstrateVisualize.cpp
--- a/Engine/Source/Runtime/Renderer/Private/Substrate/SubstrateVisualize.cpp
+++ b/Engine/Source/Runtime/Renderer/Private/Substrate/SubstrateVisualize.cpp
@@ -52,7 +52,7 @@ class FMaterialPrintInfoCS : public FGlobalShader
 		SHADER_PARAMETER_RDG_UNIFORM_BUFFER(FSubstrateGlobalUniformParameters, Substrate)
 		SHADER_PARAMETER_STRUCT_INCLUDE(FSceneTextureParameters, SceneTextures)
 		SHADER_PARAMETER_STRUCT_INCLUDE(ShaderPrint::FShaderParameters, ShaderPrintParameters)
-		SHADER_PARAMETER_RDG_BUFFER_UAV(RWStructuredBuffer<int>, PixelMaterialDebugDataUAV)
+		SHADER_PARAMETER_RDG_BUFFER_UAV(RWStructuredBuffer<int>, SubstrateDebugDataUAV)
 	END_SHADER_PARAMETER_STRUCT()
 
 	static bool IsSupported(EShaderPlatform InPlatform, EShaderPermutationFlags InFlags=EShaderPermutationFlags::None)
@@ -136,27 +136,11 @@ class FSubstrateSystemInfoCS : public FGlobalShader
 	using FPermutationDomain = TShaderPermutationDomain<>;
 
 	BEGIN_SHADER_PARAMETER_STRUCT(FParameters, )
-		SHADER_PARAMETER(uint32, bAdvancedDebugEnabled)
-		SHADER_PARAMETER(uint32, bEnergyConservation)
-		SHADER_PARAMETER(uint32, bEnergyPreservation)
-		SHADER_PARAMETER(uint32, bDbufferPass)
-		SHADER_PARAMETER(uint32, ClassificationCMask)
-		SHADER_PARAMETER(uint32, ClassificationAsync)
-		SHADER_PARAMETER(uint32, Classification8bits)
-		SHADER_PARAMETER(uint32, bRoughRefraction)
-		SHADER_PARAMETER(uint32, bUseClosureCountFromMaterialData)
-		SHADER_PARAMETER(uint32, ProjectMaxBytesPerPixel)
-		SHADER_PARAMETER(uint32, ProjectMaxClosurePerPixel)
-		SHADER_PARAMETER(uint32, ViewsMaxBytesPerPixel)
-		SHADER_PARAMETER(uint32, ViewsMaxClosurePerPixel)
-		SHADER_PARAMETER(uint32, MaterialBufferAllocationInBytes)
-		SHADER_PARAMETER(uint32, MaterialBufferAllocationMode)
-		SHADER_PARAMETER(uint32, MaxClosureCount)
 		SHADER_PARAMETER_RDG_BUFFER_SRV(Buffer<uint>, ClassificationTileDrawIndirectBuffer)
 		SHADER_PARAMETER_STRUCT_REF(FViewUniformShaderParameters, ViewUniformBuffer)
 		SHADER_PARAMETER_RDG_UNIFORM_BUFFER(FSubstrateGlobalUniformParameters, Substrate)
 		SHADER_PARAMETER_STRUCT_INCLUDE(FSceneTextureParameters, SceneTextures)
-		SHADER_PARAMETER_STRUCT_INCLUDE(ShaderPrint::FShaderParameters, ShaderPrintParameters)
+		SHADER_PARAMETER_RDG_BUFFER_UAV(RWStructuredBuffer<int>, SubstrateDebugDataUAV)
 	END_SHADER_PARAMETER_STRUCT()
 
 	static bool IsSupported(EShaderPlatform InPlatform, EShaderPermutationFlags InFlags=EShaderPermutationFlags::None)
@@ -282,10 +266,10 @@ void AddProcessAndPrintSubstrateMaterialPropertiesPasses(FRDGBuilder& GraphBuild
 	////////////////////////////////////////////////////////////////////////////////////////////////////
 	// Check if the latest readback query is ready and display the data on screen.
 	TSharedPtr<FRHIGPUBufferReadback> AvailableReadback;
-	SubstrateViewDebugData.ReadbackQueries.Peek(AvailableReadback);
+	SubstrateViewDebugData.PixelMaterialDebugDataReadbackQueries.Peek(AvailableReadback);
 	if (AvailableReadback.IsValid() && AvailableReadback->IsReady())
 	{
-		SubstrateViewDebugData.ReadbackQueries.Dequeue(AvailableReadback);
+		SubstrateViewDebugData.PixelMaterialDebugDataReadbackQueries.Dequeue(AvailableReadback);
 
 		// Access the data and copy to a frame transient buffer for rendering pass.
 		void* PixelDebugData = AvailableReadback->Lock(SubstrateViewDebugData.PixelMaterialDebugDataSizeBytes);
@@ -300,13 +284,8 @@ void AddProcessAndPrintSubstrateMaterialPropertiesPasses(FRDGBuilder& GraphBuild
 		AddDrawCanvasPass(GraphBuilder, {}, View, FScreenPassRenderTarget(SceneColorTexture, View.ViewRect, ERenderTargetLoadAction::ELoad),
 			[&View, PixelDebugDataToPrint](FCanvas& Canvas)
 			{
-				FString Text;
 				Canvas.SetScaledToRenderTarget(true);
 
-				const float ViewPortWidth = float(View.ViewRect.Width());
-				const float ViewPortHeight = float(View.ViewRect.Height());
-				FLinearColor TextColor(1.0f, 0.5f, 0.0f);
-
 				const FLinearColor Grey = FLinearColor(0.5f, 0.5f, 0.5f, 1.0f);
 				const FLinearColor Orange = FLinearColor(243.f / 255.f, 156.f / 255.f, 18.f / 255.f, 1.0);
 				const FLinearColor CompR = FLinearColor(0.8f, 0.2f, 0.2f, 1.0);
@@ -316,7 +295,7 @@ void AddProcessAndPrintSubstrateMaterialPropertiesPasses(FRDGBuilder& GraphBuild
 				float DrawPosX = DrawPosXLeft;
 				float DrawPosY = 50.0f;
 
-				FSubstratePixelDebugDataSerializer S;
+				FSubstrateDebugDataSerializer S;
 				S.SubstratePixelDebugData = (int*)PixelDebugDataToPrint;
 				FSubstratePixelDebugData Data;
 				SerializeSubstratePixelDebugData(S, Data);
@@ -936,7 +915,7 @@ void AddProcessAndPrintSubstrateMaterialPropertiesPasses(FRDGBuilder& GraphBuild
 	// Copy the debug data on GPU into a readback buffer for debug display when available later.
 	TSharedPtr<FRHIGPUBufferReadback> NewReadBack = MakeShared<FRHIGPUBufferReadback>(TEXT("PixelMaterialDebugDataReadback"));
 	AddEnqueueCopyPass(GraphBuilder, NewReadBack.Get(), NewSubstratePixelDebugBuffer.PixelMaterialDebugData, SubstrateViewDebugData.PixelMaterialDebugDataSizeBytes);
-	SubstrateViewDebugData.ReadbackQueries.Enqueue(NewReadBack);
+	SubstrateViewDebugData.PixelMaterialDebugDataReadbackQueries.Enqueue(NewReadBack);
 }
 
 static void AddVisualizeMaterialPropertiesPasses(FRDGBuilder& GraphBuilder, const FViewInfo& View, FScreenPassTexture& ScreenPassSceneColor, EShaderPlatform Platform)
@@ -959,7 +938,7 @@ static void AddVisualizeMaterialPropertiesPasses(FRDGBuilder& GraphBuilder, cons
 		PassParameters->ViewUniformBuffer = View.ViewUniformBuffer;
 		PassParameters->Substrate = Substrate::BindSubstrateGlobalUniformParameters(View);
 		PassParameters->SceneTextures = GetSceneTextureParameters(GraphBuilder, View);
-		PassParameters->PixelMaterialDebugDataUAV = NewSubstratePixelDebugBuffer.PixelMaterialDebugDataUAV;
+		PassParameters->SubstrateDebugDataUAV = NewSubstratePixelDebugBuffer.PixelMaterialDebugDataUAV;
 		ShaderPrint::SetParameters(GraphBuilder, View.ShaderPrintData, PassParameters->ShaderPrintParameters);
 
 		TShaderMapRef<FMaterialPrintInfoCS> ComputeShader(View.ShaderMap);
@@ -1004,7 +983,10 @@ bool Is8bitTileCoordEnabled();
 
 static void AddVisualizeSystemInfoPasses(FRDGBuilder& GraphBuilder, const FViewInfo& View, FScreenPassTexture& ScreenPassSceneColor, EShaderPlatform Platform)
 {
-	if (!FSubstrateSystemInfoCS::IsSupported(Platform)) return;
+	if (!FSubstrateSystemInfoCS::IsSupported(Platform))
+	{
+		return;
+	}
 
 	// Force ShaderPrint on.
 	ShaderPrint::SetEnabled(true);
@@ -1015,30 +997,289 @@ static void AddVisualizeSystemInfoPasses(FRDGBuilder& GraphBuilder, const FViewI
 	const FShadingEnergyConservationData ShadingEnergyConservationData = ShadingEnergyConservation::GetData(View);
 
 	FSubstrateSystemInfoCS::FParameters* PassParameters = GraphBuilder.AllocParameters<FSubstrateSystemInfoCS::FParameters>();
-	PassParameters->bAdvancedDebugEnabled = IsAdvancedVisualizationEnabled() ? 1u : 0u;
-	PassParameters->bEnergyConservation = ShadingEnergyConservationData.bEnergyConservation;
-	PassParameters->bEnergyPreservation = ShadingEnergyConservationData.bEnergyPreservation;
-	PassParameters->bDbufferPass = IsDBufferPassEnabled(View.GetShaderPlatform()) ? 1 : 0;
-	PassParameters->ClassificationCMask = SupportsCMask(View.GetShaderPlatform()) ? 1 : 0;
-	PassParameters->ClassificationAsync = IsClassificationAsync() ? 1 : 0;
-	PassParameters->Classification8bits = Is8bitTileCoordEnabled() ? 1 : 0;
-	PassParameters->MaxClosureCount = GetSubstrateMaxClosureCount(View);
-	PassParameters->bUseClosureCountFromMaterialData = UsesSubstrateClosureCountFromMaterialData() ? 1 : 0;
-	PassParameters->bRoughRefraction = IsOpaqueRoughRefractionEnabled() ? 1 : 0;
 	PassParameters->ClassificationTileDrawIndirectBuffer = GraphBuilder.CreateSRV(View.SubstrateViewData.ClassificationTileDrawIndirectBuffer, PF_R32_UINT);
 	PassParameters->ViewUniformBuffer = View.ViewUniformBuffer;
 	PassParameters->Substrate = Substrate::BindSubstrateGlobalUniformParameters(View);
 	PassParameters->SceneTextures = GetSceneTextureParameters(GraphBuilder, View);
-	PassParameters->ProjectMaxBytesPerPixel = GetBytePerPixel(View.GetShaderPlatform());
-	PassParameters->ProjectMaxClosurePerPixel = GetClosurePerPixel(View.GetShaderPlatform());
-	PassParameters->ViewsMaxBytesPerPixel = View.SubstrateViewData.SceneData->ViewsMaxBytesPerPixel;
-	PassParameters->ViewsMaxClosurePerPixel = View.SubstrateViewData.SceneData->ViewsMaxClosurePerPixel;
-	PassParameters->MaterialBufferAllocationMode = GetMaterialBufferAllocationMode();
-	PassParameters->MaterialBufferAllocationInBytes = MaterialBufferDesc.Extent.X * MaterialBufferDesc.Extent.Y * MaterialBufferDesc.ArraySize * sizeof(uint32);
-	ShaderPrint::SetParameters(GraphBuilder, View.ShaderPrintData, PassParameters->ShaderPrintParameters);
+
+	////////////////////////////////////////////////////////////////////////////////////////////////////
+	// Get a new pixel material buffer and render debug data into it
+	FSubstrateViewDebugData& SubstrateViewDebugData = View.ViewState->GetSubstrateViewDebugData();
+	FSubstrateViewDebugData::FTransientSystemInfoDebugBuffer NewSubstrateSystemInfoDebugBuffer = SubstrateViewDebugData.CreateTransientSystemInfoDebugBuffer(GraphBuilder);
+	PassParameters->SubstrateDebugDataUAV = NewSubstrateSystemInfoDebugBuffer.SystemInfoDebugDataUAV;
 
 	TShaderMapRef<FSubstrateSystemInfoCS> ComputeShader(View.ShaderMap);
 	FComputeShaderUtils::AddPass(GraphBuilder, RDG_EVENT_NAME("Substrate::VisualizeSystemInfo"), ComputeShader, PassParameters, FIntVector(1, 1, 1));
+
+	struct FSystemInformation
+	{
+		int CurrentBytesPerPixel;
+		int ProjectMaxBytesPerPixel;
+		int MaterialBufferAllocationInBytes;
+		int ViewsMaxBytesPerPixel;
+		int CurrentClosuresPerPixel;
+		int ViewsMaxClosuresPerPixel;
+		int ProjectMaxClosuresPerPixel;
+		int MaterialBufferAllocationMode;
+		int ShadingQuality;
+		int bRoughDiffuse;
+		int bEnergyConservation;
+		int bEnergyPreservation;
+		int bUseClosureCountFromMaterialData;
+		int TileClosureCount;
+		int bDbufferPass;
+		int bRoughRefraction;
+
+		int TileSize;
+		int TileCountX;
+		int TileCountY;
+		int bCMask;
+		int bAsync;
+		int b8BitsCoord;
+
+		int bAdvancedDebugEnabled;
+		int LayerPeelIndex;
+	} SystemInformation;
+	SystemInformation.CurrentBytesPerPixel = View.SubstrateViewData.SceneData->EffectiveMaxBytesPerPixel;
+	SystemInformation.ProjectMaxBytesPerPixel = GetBytePerPixel(View.GetShaderPlatform());
+	SystemInformation.MaterialBufferAllocationInBytes = MaterialBufferDesc.Extent.X * MaterialBufferDesc.Extent.Y * MaterialBufferDesc.ArraySize * sizeof(uint32);
+	SystemInformation.ViewsMaxBytesPerPixel = View.SubstrateViewData.SceneData->ViewsMaxBytesPerPixel;
+	SystemInformation.CurrentClosuresPerPixel = View.SubstrateViewData.SceneData->EffectiveMaxClosurePerPixel;
+	SystemInformation.ViewsMaxClosuresPerPixel = View.SubstrateViewData.SceneData->ViewsMaxClosurePerPixel;
+	SystemInformation.ProjectMaxClosuresPerPixel = GetClosurePerPixel(View.GetShaderPlatform());
+	SystemInformation.MaterialBufferAllocationMode = GetMaterialBufferAllocationMode();
+	SystemInformation.ShadingQuality = Substrate::GetShadingQuality(View.GetShaderPlatform());
+	SystemInformation.bRoughDiffuse; // Substrate.bRoughDiffuse
+	SystemInformation.bEnergyConservation = ShadingEnergyConservationData.bEnergyConservation;
+	SystemInformation.bEnergyPreservation = ShadingEnergyConservationData.bEnergyPreservation;
+	SystemInformation.bUseClosureCountFromMaterialData = UsesSubstrateClosureCountFromMaterialData() ? 1 : 0;
+	SystemInformation.TileClosureCount = GetSubstrateMaxClosureCount(View);
+	SystemInformation.bDbufferPass = IsDBufferPassEnabled(View.GetShaderPlatform()) ? 1 : 0;
+	SystemInformation.bRoughRefraction = IsOpaqueRoughRefractionEnabled() ? 1 : 0;
+	SystemInformation.TileSize = SUBSTRATE_TILE_SIZE;
+	SystemInformation.TileCountX = View.SubstrateViewData.TileCount.X;
+	SystemInformation.TileCountY = View.SubstrateViewData.TileCount.Y;
+	SystemInformation.bCMask = SupportsCMask(View.GetShaderPlatform()) ? 1 : 0;
+	SystemInformation.bAsync = IsClassificationAsync() ? 1 : 0;
+	SystemInformation.b8BitsCoord = Is8bitTileCoordEnabled() ? 1 : 0;
+	SystemInformation.bAdvancedDebugEnabled = IsAdvancedVisualizationEnabled() ? 1 : 0;
+	SystemInformation.LayerPeelIndex = View.SubstrateViewData.SceneData->PeelLayersAboveDepth;
+
+	////////////////////////////////////////////////////////////////////////////////////////////////////
+	// Check if the latest readback query is ready and display the data on screen.
+	TSharedPtr<FRHIGPUBufferReadback> AvailableReadback;
+	SubstrateViewDebugData.SystemInfoDebugDataReadbackQueries.Peek(AvailableReadback);
+	if (AvailableReadback.IsValid() && AvailableReadback->IsReady())
+	{
+		SubstrateViewDebugData.SystemInfoDebugDataReadbackQueries.Dequeue(AvailableReadback);
+
+		// Access the data and copy to a frame transient buffer for rendering pass.
+		void* SystemInfoDebugData = AvailableReadback->Lock(SubstrateViewDebugData.SystemInfoDebugDataSizeBytes);
+		if (SystemInfoDebugData == nullptr)
+		{
+			return;
+		}
+		int* SystemInfoDebugDataToPrint = (int*)GraphBuilder.Alloc(SubstrateViewDebugData.SystemInfoDebugDataSizeBytes);
+		memcpy(SystemInfoDebugDataToPrint, SystemInfoDebugData, SubstrateViewDebugData.SystemInfoDebugDataSizeBytes);
+
+		////////////////////////////////////////////////////////////////////////////////////////////////////
+		AddDrawCanvasPass(GraphBuilder, {}, View, FScreenPassRenderTarget(ScreenPassSceneColor.Texture, View.ViewRect, ERenderTargetLoadAction::ELoad),
+			[&View, SystemInfoDebugDataToPrint, SystemInformation](FCanvas& Canvas)
+		{
+			Canvas.SetScaledToRenderTarget(true);
+
+			FSubstrateDebugDataSerializer S;
+			S.SubstratePixelDebugData = (int*)SystemInfoDebugDataToPrint;
+			FSubstrateSystemInfoData Data;
+			SerializeSubstrateSystemInfoDebugData(S, Data);
+
+			float DrawPosX = 50.0f;
+			float DrawPosY = 50.0f;
+
+			auto PrintI = [&](float X, float Y, int Value, FLinearColor Color = FLinearColor::White)
+			{
+				FString String = FString::Printf(TEXT("%i"), Value);
+				FCanvasTextItem Item = FCanvasTextItem(FVector2D(X, Y), FText::FromString(*String), GetStatsFont(), Color);
+				Canvas.DrawItem(Item);
+			};
+			auto PrintS = [&](float X, float Y, FString Text, FLinearColor Color = FLinearColor::White)
+			{
+				FCanvasTextItem Item = FCanvasTextItem(FVector2D(X, Y), FText::FromString(*Text), GEngine->GetSmallFont(), Color);
+				Canvas.DrawItem(Item);
+			};
+			auto PrintBool = [&](float X, float Y, int Value)
+			{
+				FString String = Value > 0.0f ? TEXT("Yes") : TEXT("No");
+				FCanvasTextItem Item = FCanvasTextItem(FVector2D(X, Y), FText::FromString(*String), GetStatsFont(), Value > 0.0f ? FLinearColor::Green : FLinearColor::Red);
+				Canvas.DrawItem(Item);
+			};
+			auto NewLine = [&]()
+			{
+				DrawPosY += 17;
+			};
+
+			Canvas.DrawTile(DrawPosX - 10.0f, DrawPosY - 10.0f, 350.0f, 700.0f, 0, 0, 1, 1, FLinearColor(0.0, 0.0, 0.0, 0.3));
+
+			const FLinearColor White = FLinearColor::White;
+			const FLinearColor Yellow = FLinearColor(1, 1, 0, 1);
+			const FLinearColor DarkYellow = FLinearColor(0.5, 0.5, 0, 1);
+			const FLinearColor Orange = FLinearColor(243.f / 255.f, 156.f / 255.f, 18.f / 255.f, 1.0);
+			FString Text;
+
+			PrintS(DrawPosX, DrawPosY, TEXT("General"), Yellow);
+			NewLine();
+			{
+				PrintS(DrawPosX, DrawPosY, TEXT("  Bytes Per Pixel"), DarkYellow);
+				NewLine();
+				PrintS(DrawPosX +   0, DrawPosY, TEXT("    Max"), White);
+				PrintI(DrawPosX + 150, DrawPosY, SystemInformation.CurrentBytesPerPixel, Orange);
+				PrintS(DrawPosX + 180, DrawPosY, TEXT("/"), DarkYellow);
+				PrintI(DrawPosX + 190, DrawPosY, SystemInformation.ProjectMaxBytesPerPixel, Orange);
+				Text = FString::Printf(TEXT("[%i MB]"), SystemInformation.MaterialBufferAllocationInBytes/(1024*1024));
+				PrintS(DrawPosX + 220, DrawPosY, *Text, White);
+				NewLine();
+				PrintS(DrawPosX + 0, DrawPosY, TEXT("    Views"), White);
+				PrintI(DrawPosX + 150, DrawPosY, SystemInformation.ViewsMaxBytesPerPixel, Orange);
+				PrintS(DrawPosX + 180, DrawPosY, TEXT("/ "), DarkYellow);
+				PrintI(DrawPosX + 190, DrawPosY, SystemInformation.ProjectMaxBytesPerPixel, Orange);
+				NewLine();
+				NewLine();
+
+				PrintS(DrawPosX, DrawPosY, TEXT("  Closures Per Pixel"), DarkYellow);
+				NewLine();
+				PrintS(DrawPosX + 0, DrawPosY, TEXT("    Max"), White);
+				PrintI(DrawPosX + 150, DrawPosY, SystemInformation.CurrentClosuresPerPixel, Orange);
+				PrintS(DrawPosX + 180, DrawPosY, TEXT("/ "), DarkYellow);
+				PrintI(DrawPosX + 190, DrawPosY, SystemInformation.ProjectMaxClosuresPerPixel, Orange);
+				NewLine();
+				PrintS(DrawPosX + 0, DrawPosY, TEXT("    Views"), White);
+				PrintI(DrawPosX + 150, DrawPosY, SystemInformation.ViewsMaxClosuresPerPixel, Orange);
+				PrintS(DrawPosX + 180, DrawPosY, TEXT("/ "), DarkYellow);
+				PrintI(DrawPosX + 190, DrawPosY, SystemInformation.ProjectMaxClosuresPerPixel, Orange);
+				NewLine();
+				NewLine();
+
+				PrintS(DrawPosX, DrawPosY, TEXT("  Allocation mode"), White);
+				switch (SystemInformation.MaterialBufferAllocationMode)
+				{
+				case 0: PrintS(DrawPosX + 150, DrawPosY, TEXT("View based"), Orange); break;
+				case 1: PrintS(DrawPosX + 150, DrawPosY, TEXT("View based | Growing-only"), Orange); break;
+				case 2: PrintS(DrawPosX + 150, DrawPosY, TEXT("Setting based"), Orange); break;
+				default: PrintS(DrawPosX + 150, DrawPosY, TEXT("Unkown"), Orange); break;
+				}
+				NewLine();
+				PrintS(DrawPosX, DrawPosY, TEXT("  Shading quality"), White);
+				PrintI(DrawPosX + 150, DrawPosY, SystemInformation.ShadingQuality, Orange);
+				NewLine();
+				PrintS(DrawPosX, DrawPosY, TEXT("  Rough diffuse"), White);
+				PrintBool(DrawPosX + 150, DrawPosY, SystemInformation.bRoughDiffuse);
+				NewLine();
+				PrintS(DrawPosX, DrawPosY, TEXT("  Energy conservation"), White);
+				PrintBool(DrawPosX + 150, DrawPosY, SystemInformation.bEnergyConservation);
+				NewLine();
+				PrintS(DrawPosX, DrawPosY, TEXT("  Energy preservation"), White);
+				PrintBool(DrawPosX + 150, DrawPosY, SystemInformation.bEnergyPreservation);
+				NewLine();
+				PrintS(DrawPosX, DrawPosY, TEXT("  Use #closures mat."), White);
+				PrintBool(DrawPosX + 150, DrawPosY, SystemInformation.bUseClosureCountFromMaterialData);
+				NewLine();
+				PrintS(DrawPosX, DrawPosY, TEXT("  Tile closure Count"), White);
+				PrintI(DrawPosX + 150, DrawPosY, SystemInformation.TileClosureCount, Orange);
+				NewLine();
+				PrintS(DrawPosX, DrawPosY, TEXT("  DBuffser pass"), White);
+				PrintBool(DrawPosX + 150, DrawPosY, SystemInformation.bDbufferPass);
+				NewLine();
+				PrintS(DrawPosX, DrawPosY, TEXT("  Rough refraction"), White);
+				PrintBool(DrawPosX + 150, DrawPosY, SystemInformation.bRoughRefraction);
+				NewLine();
+			}
+			NewLine();
+
+			PrintS(DrawPosX, DrawPosY, TEXT("Classification"), Yellow);
+			NewLine();
+			{
+				PrintS(DrawPosX, DrawPosY, TEXT("  Tile size"), White);
+				Text = FString::Printf(TEXT("%ix%i"), SystemInformation.TileSize, SystemInformation.TileSize);
+				PrintS(DrawPosX + 150, DrawPosY, *Text, Orange);
+				NewLine();
+				PrintS(DrawPosX, DrawPosY, TEXT("  Tile Resolution"), White);
+				Text = FString::Printf(TEXT("%ix%i"), SystemInformation.TileCountX, SystemInformation.TileCountY);
+				PrintS(DrawPosX + 150, DrawPosY, *Text, Orange);
+				NewLine();
+				PrintS(DrawPosX, DrawPosY, TEXT("  CMask Classification"), White);
+				PrintBool(DrawPosX + 150, DrawPosY, SystemInformation.bCMask);
+				NewLine();
+				PrintS(DrawPosX, DrawPosY, TEXT("  Async Compute"), White);
+				PrintBool(DrawPosX + 150, DrawPosY, SystemInformation.bAsync);
+				NewLine();
+				PrintS(DrawPosX, DrawPosY, TEXT("  8Bits coord"), White);
+				PrintBool(DrawPosX + 150, DrawPosY, SystemInformation.b8BitsCoord);
+				NewLine();
+			}
+			NewLine();
+
+			PrintS(DrawPosX, DrawPosY, TEXT("Material Tiles"), Yellow);
+			NewLine();
+			{
+				PrintS(DrawPosX, DrawPosY, TEXT("  # Simple tiles"), White);
+				PrintI(DrawPosX + 150, DrawPosY, Data.TileCount[SUBSTRATE_TILE_TYPE_SIMPLE], Orange);
+				NewLine();
+				PrintS(DrawPosX, DrawPosY, TEXT("  # Single tiles"), White);
+				PrintI(DrawPosX + 150, DrawPosY, Data.TileCount[SUBSTRATE_TILE_TYPE_SINGLE], Orange);
+				NewLine();
+				PrintS(DrawPosX, DrawPosY, TEXT("  # Complex tiles"), White);
+				PrintI(DrawPosX + 150, DrawPosY, Data.TileCount[SUBSTRATE_TILE_TYPE_COMPLEX], Orange);
+				NewLine();
+				PrintS(DrawPosX, DrawPosY, TEXT("  # Complex Special tiles"), White);
+				PrintI(DrawPosX + 150, DrawPosY, Data.TileCount[SUBSTRATE_TILE_TYPE_COMPLEX_SPECIAL], Orange);
+				NewLine();
+				if (SystemInformation.bRoughRefraction)
+				{
+					PrintS(DrawPosX, DrawPosY, TEXT("  # Rough refract. tiles"), White);
+					PrintI(DrawPosX + 150, DrawPosY, Data.TileCount[SUBSTRATE_TILE_TYPE_ROUGH_REFRACT], Orange);
+					NewLine();
+					PrintS(DrawPosX, DrawPosY, TEXT("  # Rough r. w/o SSS tiles"), White);
+					PrintI(DrawPosX + 150, DrawPosY, Data.TileCount[SUBSTRATE_TILE_TYPE_ROUGH_REFRACT_SSS_WITHOUT], Orange);
+					NewLine();
+				}
+				if (SystemInformation.bDbufferPass)
+				{
+					PrintS(DrawPosX, DrawPosY, TEXT("  # Decal Simple tiles"), White);
+					PrintI(DrawPosX + 150, DrawPosY, Data.TileCount[SUBSTRATE_TILE_TYPE_DECAL_SIMPLE], Orange);
+					NewLine();
+					PrintS(DrawPosX, DrawPosY, TEXT("  # Decal Single tiles"), White);
+					PrintI(DrawPosX + 150, DrawPosY, Data.TileCount[SUBSTRATE_TILE_TYPE_DECAL_SINGLE], Orange);
+					NewLine();
+					PrintS(DrawPosX, DrawPosY, TEXT("  # Decal Complex tiles"), White);
+					PrintI(DrawPosX + 150, DrawPosY, Data.TileCount[SUBSTRATE_TILE_TYPE_DECAL_COMPLEX], Orange);
+					NewLine();
+				}
+			}
+			NewLine();
+
+			PrintS(DrawPosX, DrawPosY, TEXT("Debug"), Yellow);
+			NewLine();
+			{
+				PrintS(DrawPosX, DrawPosY, TEXT("  # ADv. Debug"), White);
+				PrintBool(DrawPosX + 150, DrawPosY, SystemInformation.bAdvancedDebugEnabled);
+				NewLine();
+				PrintS(DrawPosX, DrawPosY, TEXT("  # Decal Complex tiles"), White);
+				PrintI(DrawPosX + 150, DrawPosY, SystemInformation.LayerPeelIndex, Orange);
+				NewLine();
+			}
+			NewLine();
+		});
+		////////////////////////////////////////////////////////////////////////////////////////////////////
+
+		AvailableReadback->Unlock();
+	}
+
+
+	////////////////////////////////////////////////////////////////////////////////////////////////////
+	// Copy the debug data on GPU into a readback buffer for debug display when available later.
+	TSharedPtr<FRHIGPUBufferReadback> NewReadBack = MakeShared<FRHIGPUBufferReadback>(TEXT("SystemInfoDebugDataReadback"));
+	AddEnqueueCopyPass(GraphBuilder, NewReadBack.Get(), NewSubstrateSystemInfoDebugBuffer.SystemInfoDebugData, SubstrateViewDebugData.SystemInfoDebugDataSizeBytes);
+	SubstrateViewDebugData.SystemInfoDebugDataReadbackQueries.Enqueue(NewReadBack);
 }
 
 // Draw each material layer independently
@@ -1201,7 +1442,7 @@ FSubstrateViewDebugData::FTransientPixelDebugBuffer FSubstrateViewDebugData::Cre
 
 	const uint32 SizeOfUint = sizeof(uint32);
 	const uint32 NumUints = FMath::DivideAndRoundUp(uint32(sizeof(FSubstratePixelDebugData)), SizeOfUint);
-	PixelMaterialDebugDataSizeBytes = NumUints * SizeOfUint;	// Arbitrary for now
+	PixelMaterialDebugDataSizeBytes = NumUints * SizeOfUint;
 
 	FRDGBufferDesc RDGBufferDesc = FRDGBufferDesc::CreateStructuredDesc(SizeOfUint, NumUints);
 	RDGBufferDesc.Usage |= BUF_SourceCopy;
@@ -1212,10 +1453,31 @@ FSubstrateViewDebugData::FTransientPixelDebugBuffer FSubstrateViewDebugData::Cre
 	return PixelDebugBuffer;
 }
 
+FSubstrateViewDebugData::FTransientSystemInfoDebugBuffer FSubstrateViewDebugData::CreateTransientSystemInfoDebugBuffer(FRDGBuilder& GraphBuilder)
+{
+	FSubstrateViewDebugData::FTransientSystemInfoDebugBuffer SystemInfoBuffer;
+
+	const uint32 SizeOfUint = sizeof(uint32);
+	const uint32 NumUints = SUBSTRATE_TILE_TYPE_COUNT;
+	SystemInfoDebugDataSizeBytes = NumUints * SizeOfUint;
+
+	FRDGBufferDesc RDGBufferDesc = FRDGBufferDesc::CreateStructuredDesc(SizeOfUint, NumUints);
+	RDGBufferDesc.Usage |= BUF_SourceCopy;
+	SystemInfoBuffer.SystemInfoDebugData = GraphBuilder.CreateBuffer(RDGBufferDesc, TEXT("SystemInfoDebugData"));
+
+	SystemInfoBuffer.SystemInfoDebugDataUAV = GraphBuilder.CreateUAV(FRDGBufferUAVDesc(SystemInfoBuffer.SystemInfoDebugData, PF_R32_UINT));
+
+	return SystemInfoBuffer;
+}
+
 void FSubstrateViewDebugData::SafeRelease()
 {
 	TSharedPtr<FRHIGPUBufferReadback> It;
-	while (ReadbackQueries.Dequeue(It))
+	while (PixelMaterialDebugDataReadbackQueries.Dequeue(It))
+	{
+		It.Reset();
+	}
+	while (SystemInfoDebugDataReadbackQueries.Dequeue(It))
 	{
 		It.Reset();
 	}
