Substrate - fixed unassigned variable.

#rnx

diff --git a/Engine/Source/Runtime/Renderer/Private/Substrate/SubstrateVisualize.cpp b/Engine/Source/Runtime/Renderer/Private/Substrate/SubstrateVisualize.cpp
--- a/Engine/Source/Runtime/Renderer/Private/Substrate/SubstrateVisualize.cpp
+++ b/Engine/Source/Runtime/Renderer/Private/Substrate/SubstrateVisualize.cpp
@@ -1049,7 +1049,7 @@ static void AddVisualizeSystemInfoPasses(FRDGBuilder& GraphBuilder, const FViewI
 	SystemInformation.ProjectMaxClosuresPerPixel = GetClosurePerPixel(View.GetShaderPlatform());
 	SystemInformation.MaterialBufferAllocationMode = GetMaterialBufferAllocationMode();
 	SystemInformation.ShadingQuality = Substrate::GetShadingQuality(View.GetShaderPlatform());
-	SystemInformation.bRoughDiffuse; // Substrate.bRoughDiffuse
+	SystemInformation.bRoughDiffuse = View.SubstrateViewData.SceneData->bRoughDiffuse;
 	SystemInformation.bEnergyConservation = ShadingEnergyConservationData.bEnergyConservation;
 	SystemInformation.bEnergyPreservation = ShadingEnergyConservationData.bEnergyPreservation;
 	SystemInformation.bUseClosureCountFromMaterialData = UsesSubstrateClosureCountFromMaterialData() ? 1 : 0;
