[Exposure Fusion] Use FilmTonemap(...) to better match final tonemapping when calculating exposure weights.
- also reduced PPV value mapping from 10 to 6 EV.
- added missing logic for LocalExposure visualization tile when using Exposure Fusion.

diff --git a/Engine/Shaders/Private/PostProcessHistogramCommon.ush b/Engine/Shaders/Private/PostProcessHistogramCommon.ush
--- a/Engine/Shaders/Private/PostProcessHistogramCommon.ush
+++ b/Engine/Shaders/Private/PostProcessHistogramCommon.ush
@@ -4,6 +4,8 @@
 	PostProcessHistogramCommon.ush: PostProcessing histogram shared functions and structures.
 =============================================================================*/
 
+#pragma once
+
 float EyeAdaptation_ExposureLowPercent;
 float EyeAdaptation_ExposureHighPercent;
 float EyeAdaptation_MinAverageLuminance;
@@ -275,12 +277,44 @@ float CalculateLocalExposure(float LogLum, float BaseLogLum, float LogMiddleGrey
 	return exp2(CalculateLogLocalExposure(LogLum, BaseLogLum, LogMiddleGrey, HighlightContrastScale, ShadowContrastScale, DetailStrength));
 }
 
+#define EXPOSURE_FUSION_USE_FILM_TONEMAP 1
+
+#if EXPOSURE_FUSION_USE_FILM_TONEMAP
+#include "TonemapCommon.ush"
+#endif
+
 float ExposureFusionTonemap(float Lum)
 {
-	return Lum / (1 + Lum);
+#if EXPOSURE_FUSION_USE_FILM_TONEMAP
+	float TonemappedLuminance = FilmToneMap(Lum).r;
+#else
+	float TonemappedLuminance = Lum / (1 + Lum);
+#endif
+
+	TonemappedLuminance = sqrt(TonemappedLuminance);
+
+	return TonemappedLuminance;
+}
+
+float ExposureFusionInverseTonemap(float TonemappedLuminance)
+{
+	TonemappedLuminance = TonemappedLuminance * TonemappedLuminance;
+
+#if EXPOSURE_FUSION_USE_FILM_TONEMAP
+	float Lum = FilmToneMapInverse(TonemappedLuminance).r;
+#else
+	float Lum = TonemappedLuminance / (1 - TonemappedLuminance);
+#endif
+
+	return Lum;
 }
 
-float ExposureFusionInverseTonemap(float Lum)
+float CalculateFusionLocalExposure(float Lum, float FusionResult)
 {
-	return Lum / (1 - Lum);
+#if 1
+	return ExposureFusionInverseTonemap(min(FusionResult, 1.0f)) / Lum;
+#else
+	const float LumTonemapped = ExposureFusionTonemap(Lum);
+	return FusionResult / LumTonemapped;
+#endif
 }
\ No newline at end of file
diff --git a/Engine/Shaders/Private/PostProcessTonemap.usf b/Engine/Shaders/Private/PostProcessTonemap.usf
--- a/Engine/Shaders/Private/PostProcessTonemap.usf
+++ b/Engine/Shaders/Private/PostProcessTonemap.usf
@@ -438,12 +438,7 @@ float4 TonemapCommonPS(
 
 		const float LuminanceVal = CalculateEyeAdaptationLuminance(SceneColor.rgb * OneOverPreExposure * GlobalExposure);
 		
-#if 1
-		LocalExposure = ExposureFusionInverseTonemap(min(FusionVal, 1.0f)) / LuminanceVal;
-#else
-		const float LuminanceValTonemapped = ExposureFusionTonemap(LuminanceVal);
-		LocalExposure = FusionVal / LuminanceValTonemapped;
-#endif
+		LocalExposure = CalculateFusionLocalExposure(LuminanceVal, FusionVal);
 	}
 	#else // LOCAL_EXPOSURE_MODE_NONE
 	const float LocalExposure = 1.0;
diff --git a/Engine/Shaders/Private/PostProcessVisualizeLocalExposure.usf b/Engine/Shaders/Private/PostProcessVisualizeLocalExposure.usf
--- a/Engine/Shaders/Private/PostProcessVisualizeLocalExposure.usf
+++ b/Engine/Shaders/Private/PostProcessVisualizeLocalExposure.usf
@@ -16,6 +16,9 @@ Texture2D SceneColorTexture;
 Texture3D LumBilateralGrid;
 Texture2D BlurredLogLum;
 #else
+FScreenTransform ColorToExposureFusion;
+SCREEN_PASS_TEXTURE_VIEWPORT(ExposureFusion)
+Texture2D ExposureFusionTexture;
 Texture2D ExposuresTexture;
 Texture2D WeightsTexture;
 #endif
@@ -63,9 +66,16 @@ float4 MainPS(noperspective float4 UVAndScreenPos : TEXCOORD0, float4 SvPosition
 	float DetailLogLum = (LogLuminance + log2(ExposureScaleMiddleGrey.x)) - BaseLogLum;
 #else
 	float BaseLogLum = LogLuminance;
-	float LocalExposure = 1.0f;
-	float4 Exposures = Texture2DSample(ExposuresTexture, BilinearClampSampler, UV);
-	float4 Weights = Texture2DSample(WeightsTexture, BilinearClampSampler, UV);
+
+	float2 ExposureFusionUV = ApplyScreenTransform(UV.xy, ColorToExposureFusion);
+	ExposureFusionUV = clamp(ExposureFusionUV, ExposureFusion_UVViewportBilinearMin, ExposureFusion_UVViewportBilinearMax);
+
+	float4 Exposures = Texture2DSample(ExposuresTexture, BilinearClampSampler, ExposureFusionUV);
+	float4 Weights = Texture2DSample(WeightsTexture, BilinearClampSampler, ExposureFusionUV);
+
+	const float LuminanceVal2 = CalculateEyeAdaptationLuminance(HDRSceneColor.rgb * ExposureScaleMiddleGrey.x);
+	float FusionVal = Texture2DSample(ExposureFusionTexture, BilinearClampSampler, ExposureFusionUV).x;
+	float LocalExposure = CalculateFusionLocalExposure(LuminanceVal2, FusionVal);
 #endif
 
 	float BaseCentered = (BaseLogLum - MiddleGreyLumValue);
diff --git a/Engine/Shaders/Private/TonemapCommon.ush b/Engine/Shaders/Private/TonemapCommon.ush
--- a/Engine/Shaders/Private/TonemapCommon.ush
+++ b/Engine/Shaders/Private/TonemapCommon.ush
@@ -4,6 +4,8 @@
 	TonemapCommon.usf: PostProcessing tone mapping common
 =============================================================================*/
 
+#pragma once
+
 #include "ACES.ush"
 #include "GammaCorrectionCommon.ush"
 
diff --git a/Engine/Source/Runtime/Renderer/Private/PostProcess/PostProcessLocalExposure.cpp b/Engine/Source/Runtime/Renderer/Private/PostProcess/PostProcessLocalExposure.cpp
--- a/Engine/Source/Runtime/Renderer/Private/PostProcess/PostProcessLocalExposure.cpp
+++ b/Engine/Source/Runtime/Renderer/Private/PostProcess/PostProcessLocalExposure.cpp
@@ -130,6 +130,12 @@ public:
 		SHADER_PARAMETER_RDG_TEXTURE_UAV(RWTexture2D<float4>, OutputFloat4_1)
 
 		SHADER_PARAMETER(float, TargetLuminance)
+
+		SHADER_PARAMETER(float, FilmSlope)
+		SHADER_PARAMETER(float, FilmToe)
+		SHADER_PARAMETER(float, FilmShoulder)
+		SHADER_PARAMETER(float, FilmBlackClip)
+		SHADER_PARAMETER(float, FilmWhiteClip)
 	END_SHADER_PARAMETER_STRUCT()
 
 	static bool ShouldCompilePermutation(const FGlobalShaderPermutationParameters& Parameters)
@@ -238,10 +244,10 @@ FLocalExposureParameters GetLocalExposureParameters(const FViewInfo& View, FIntP
 
 	if (View.FinalPostProcessSettings.LocalExposureMethod == ELocalExposureMethod::Fusion)
 	{
-		const float HighlightEV = FMath::Lerp<float>(10, 0, HighlightContrast);
+		const float HighlightEV = FMath::Lerp<float>(6, 0, HighlightContrast);
 		HighlightContrast = FMath::Pow(2, -HighlightEV);
 
-		const float ShadowEV = FMath::Lerp<float>(10, 0, ShadowContrast);
+		const float ShadowEV = FMath::Lerp<float>(6, 0, ShadowContrast);
 		ShadowContrast = FMath::Pow(2, ShadowEV);
 	}
 
@@ -406,6 +412,14 @@ FExposureFusionData AddLocalExposureFusionPass(
 		PassParameters->OutputFloat4_1 = GraphBuilder.CreateUAV(WeightTexture.Texture);
 		PassParameters->TargetLuminance = CVarExposureFusionTargetLuminance.GetValueOnRenderThread();
 
+		const FPostProcessSettings& Settings = View.FinalPostProcessSettings;
+
+		PassParameters->FilmSlope = Settings.FilmSlope;
+		PassParameters->FilmToe = Settings.FilmToe;
+		PassParameters->FilmShoulder = Settings.FilmShoulder;
+		PassParameters->FilmBlackClip = Settings.FilmBlackClip;
+		PassParameters->FilmWhiteClip = Settings.FilmWhiteClip;
+
 		FComputeShaderUtils::AddPass(
 			GraphBuilder,
 			RDG_EVENT_NAME("FusionSetup %dx%d", Input.ViewRect.Width(), Input.ViewRect.Height()),
diff --git a/Engine/Source/Runtime/Renderer/Private/PostProcess/PostProcessTonemap.cpp b/Engine/Source/Runtime/Renderer/Private/PostProcess/PostProcessTonemap.cpp
--- a/Engine/Source/Runtime/Renderer/Private/PostProcess/PostProcessTonemap.cpp
+++ b/Engine/Source/Runtime/Renderer/Private/PostProcess/PostProcessTonemap.cpp
@@ -320,6 +320,12 @@ BEGIN_SHADER_PARAMETER_STRUCT(FTonemapParameters, )
 	SHADER_PARAMETER_RDG_TEXTURE_SRV(Texture2D, ExposureFusionTexture)
 	SHADER_PARAMETER_SAMPLER(SamplerState, ExposureFusionSampler)
 
+	SHADER_PARAMETER(float, FilmSlope)
+	SHADER_PARAMETER(float, FilmToe)
+	SHADER_PARAMETER(float, FilmShoulder)
+	SHADER_PARAMETER(float, FilmBlackClip)
+	SHADER_PARAMETER(float, FilmWhiteClip)
+
 	SHADER_PARAMETER_RDG_BUFFER_SRV(StructuredBuffer<float4>, EyeAdaptationBuffer)
 	SHADER_PARAMETER_RDG_TEXTURE(Texture2D, ColorGradingLUT)
 	SHADER_PARAMETER_TEXTURE(Texture2D, BloomDirtMaskTexture)
@@ -824,6 +830,13 @@ FScreenPassTexture AddTonemapPass(FRDGBuilder& GraphBuilder, const FViewInfo& Vi
 	CommonParameters.ChromaticAberrationParams = ChromaticAberrationParams;
 	CommonParameters.TonemapperParams = FVector4f(PostProcessSettings.VignetteIntensity, SharpenDiv6, 0.0f, 0.0f);
 	CommonParameters.EditorNITLevel = EditorNITLevel;
+
+	CommonParameters.FilmSlope = PostProcessSettings.FilmSlope;
+	CommonParameters.FilmToe = PostProcessSettings.FilmToe;
+	CommonParameters.FilmShoulder = PostProcessSettings.FilmShoulder;
+	CommonParameters.FilmBlackClip = PostProcessSettings.FilmBlackClip;
+	CommonParameters.FilmWhiteClip = PostProcessSettings.FilmWhiteClip;
+
 	{
 		const bool bSupportsBackbufferQuantizationDithering = EDisplayOutputFormat(CommonParameters.OutputDevice.OutputDevice) != EDisplayOutputFormat::HDR_LinearNoToneCurve && EDisplayOutputFormat(CommonParameters.OutputDevice.OutputDevice) != EDisplayOutputFormat::HDR_LinearWithToneCurve;
 		const int32 BitdepthOverride = CVarBackbufferQuantizationDitheringOverride.GetValueOnRenderThread();
diff --git a/Engine/Source/Runtime/Renderer/Private/PostProcess/PostProcessVisualizeLocalExposure.cpp b/Engine/Source/Runtime/Renderer/Private/PostProcess/PostProcessVisualizeLocalExposure.cpp
--- a/Engine/Source/Runtime/Renderer/Private/PostProcess/PostProcessVisualizeLocalExposure.cpp
+++ b/Engine/Source/Runtime/Renderer/Private/PostProcess/PostProcessVisualizeLocalExposure.cpp
@@ -43,9 +43,18 @@ public:
 		SHADER_PARAMETER_RDG_TEXTURE(Texture3D, LumBilateralGrid)
 		SHADER_PARAMETER_RDG_TEXTURE(Texture2D, BlurredLogLum)
 
+		SHADER_PARAMETER(FScreenTransform, ColorToExposureFusion)
+		SHADER_PARAMETER_STRUCT(FScreenPassTextureViewportParameters, ExposureFusion)
+		SHADER_PARAMETER_RDG_TEXTURE(Texture2D, ExposureFusionTexture)
 		SHADER_PARAMETER_RDG_TEXTURE(Texture2D, ExposuresTexture)
 		SHADER_PARAMETER_RDG_TEXTURE(Texture2D, WeightsTexture)
 
+		SHADER_PARAMETER(float, FilmSlope)
+		SHADER_PARAMETER(float, FilmToe)
+		SHADER_PARAMETER(float, FilmShoulder)
+		SHADER_PARAMETER(float, FilmBlackClip)
+		SHADER_PARAMETER(float, FilmWhiteClip)
+
 		SHADER_PARAMETER_SAMPLER(SamplerState, BilinearClampSampler)
 
 		SHADER_PARAMETER(uint32, DebugMode)
@@ -149,10 +158,20 @@ FScreenPassTexture AddVisualizeLocalExposurePass(FRDGBuilder& GraphBuilder, cons
 
 			if (Inputs.ExposureFusionData != nullptr)
 			{
+				const FScreenPassTextureViewport ExposureFusionViewport(Inputs.ExposureFusionData->Result);
+				PassParameters->ColorToExposureFusion = FScreenTransform::ChangeTextureUVCoordinateFromTo(InputViewport, ExposureFusionViewport);
+				PassParameters->ExposureFusion = GetScreenPassTextureViewportParameters(ExposureFusionViewport);
+				PassParameters->ExposureFusionTexture = Inputs.ExposureFusionData->Result.Texture;
 				PassParameters->ExposuresTexture = Inputs.ExposureFusionData->Exposures.Texture;
 				PassParameters->WeightsTexture = Inputs.ExposureFusionData->Weights.Texture;
 			}
 
+			PassParameters->FilmSlope = Settings.FilmSlope;
+			PassParameters->FilmToe = Settings.FilmToe;
+			PassParameters->FilmShoulder = Settings.FilmShoulder;
+			PassParameters->FilmBlackClip = Settings.FilmBlackClip;
+			PassParameters->FilmWhiteClip = Settings.FilmWhiteClip;
+
 			PassParameters->BilinearClampSampler = TStaticSamplerState<SF_Bilinear, AM_Clamp, AM_Clamp, AM_Clamp>::GetRHI();
 			PassParameters->DebugMode = uint32(VisualizeId);
 
@@ -175,8 +194,8 @@ FScreenPassTexture AddVisualizeLocalExposurePass(FRDGBuilder& GraphBuilder, cons
 		Tiles.SetNum(16);
 		if (Visualization == EVisualizeId::Overview)
 		{
-			Tiles[4 * 0 + 0] = Visualize(EVisualizeId::LocalExposure, TEXT("Local Exposure"));
-			Tiles[4 * 0 + 3] = Visualize(EVisualizeId::Thresholds, TEXT("Thresholds"));
+			Tiles[4 * 0 + 0] = Visualize(EVisualizeId::LocalExposure, FString::Printf(TEXT("Local Exposure (H=%.2f / S=%.2f)"), Inputs.LocalExposureParameters->HighlightContrastScale, Inputs.LocalExposureParameters->ShadowContrastScale));
+			Tiles[4 * 0 + 3] = Visualize(EVisualizeId::Thresholds, FString::Printf(TEXT("Thresholds (H=%.2f / S=%.2f)"), Inputs.LocalExposureParameters->HighlightThreshold, Inputs.LocalExposureParameters->ShadowThreshold));
 
 			if (Inputs.ExposureFusionData != nullptr)
 			{
