Remove unsupported point light shadows editor error message for mobile platforms.
Point light shadows are disabled by default on mobile.

diff --git a/Engine/Source/Runtime/Renderer/Private/SceneRendering.cpp b/Engine/Source/Runtime/Renderer/Private/SceneRendering.cpp
--- a/Engine/Source/Runtime/Renderer/Private/SceneRendering.cpp
+++ b/Engine/Source/Runtime/Renderer/Private/SceneRendering.cpp
@@ -4107,7 +4107,8 @@ void FSceneRenderer::OnRenderFinish(FRDGBuilder& GraphBuilder, FRDGTextureRef Vi
 
 		const bool bRealTimeSkyCaptureButNothingToCapture = Scene->SkyLight && Scene->SkyLight->bRealTimeCaptureEnabled && (!Scene->HasSkyAtmosphere() && !Scene->HasVolumetricCloud() && (Views.Num() > 0 && !Views[0].bSceneHasSkyMaterial));
 
-		const bool bShowPointLightWarning = UsedWholeScenePointLightNames.Num() > 0 && !FReadOnlyCVARCache::EnablePointLightShadows(ShaderPlatform);
+		// Point light shadows are disabled by default on mobile platforms.
+		const bool bShowPointLightWarning = !IsMobilePlatform(ShaderPlatform) ? UsedWholeScenePointLightNames.Num() > 0 && !FReadOnlyCVARCache::EnablePointLightShadows(ShaderPlatform) : false;
 		const bool bShowShadowedLightOverflowWarning = Scene->OverflowingDynamicShadowedLights.Num() > 0;
 
 		const bool bLocalFogVolumeInSceneButProjectDisabled = Scene->HasAnyLocalFogVolume() && !ProjectSupportsLocalFogVolumes();
