Scene Capture unlit rendering upgrades.  Custom Render Pass translucency support.
* Support for "Render in Main Renderer" flag with Scene Color capture source, converting the capture to an unlit Custom Render Pass.  Previously, only Scene View Extensions could create Scene Color CRPs.
* Scene captures set to unlit rendering, either by setting a depth or GBuffer capture source, or disabling the Lighting ShowFlag, now clear a collection of ShowFlags to disable features not required for unlit rendering.
* ShowFlag UnlitRichView added to allow disabling of IsRichView based debug features, which force rendering down a costly dynamic code path when Lighting is disabled, used by debug modes that globally swap shaders.
* ShowFlag UnlitViewmode added to allow disabling or enabling of development shader debug feature that routes the Base Color output to Emissive when Lighting is disabled.
* Experimental base pass optimization removed, as new optimizations achieve similar performance, without requiring a separate code path to maintain.  And Custom Render Passes are faster still.
* Custom Render Pass translucency support.  Disabled by default for existing CRPs, conditionally enabled for Scene Color capture CRPs based on the Translucency ShowFlag.
* Fixed issue where translucency didn't work for scene color Scene Captures sharing main view temporal jitter, due to the renderer assuming post process would composite translucency when temporal AA was enabled.
* Translucent pass render functions converted to static member functions, so they can't access scene renderer members which may be incorrect for custom render passes, similar to base pass render function.

diff --git a/Engine/Source/Runtime/Engine/Classes/Components/SceneCaptureComponent.h b/Engine/Source/Runtime/Engine/Classes/Components/SceneCaptureComponent.h
--- a/Engine/Source/Runtime/Engine/Classes/Components/SceneCaptureComponent.h
+++ b/Engine/Source/Runtime/Engine/Classes/Components/SceneCaptureComponent.h
@@ -245,6 +245,18 @@ public:
 
 	virtual void UpdateSceneCaptureContents(FSceneInterface* Scene) {};
 
+	/**
+	 * Returns whether a scene capture doesn't want or need lighting, and can disable several additional rendering features to save performance
+	 * (see DisableFeaturesForUnlit in ShowFlags.h).  Depth or base pass outputs aren't affected by lighting, while SceneColor outputs may be generated
+	 * with the Lighting flag disabled by the user.  FinalColor requires post processing, and so is assumed to require additional features, and can't
+	 * take advantage of disabling features for unlit.
+	 */
+	bool IsUnlit() const
+	{
+		return CaptureSource == SCS_SceneDepth || CaptureSource == SCS_DeviceDepth || CaptureSource == SCS_Normal || CaptureSource == SCS_BaseColor ||
+			(ShowFlags.Lighting == false && (CaptureSource == SCS_SceneColorHDR || CaptureSource == SCS_SceneColorHDRNoAlpha || CaptureSource == SCS_SceneColorSceneDepth));
+	}
+
 protected:
 	/** Update the show flags from our show flags settings (ideally, you'd be able to set this more directly, but currently unable to make FEngineShowFlags a UStruct to use it as a FProperty...) */
 	ENGINE_API void UpdateShowFlags();
diff --git a/Engine/Source/Runtime/Engine/Classes/Components/SceneCaptureComponent2D.h b/Engine/Source/Runtime/Engine/Classes/Components/SceneCaptureComponent2D.h
--- a/Engine/Source/Runtime/Engine/Classes/Components/SceneCaptureComponent2D.h
+++ b/Engine/Source/Runtime/Engine/Classes/Components/SceneCaptureComponent2D.h
@@ -13,6 +13,19 @@
 class ISceneViewExtension;
 class FSceneInterface;
 
+UENUM()
+enum class ESceneCaptureUnlitViewmode : uint8
+{
+	/** Disabled */
+	Disabled,
+
+	/** Enabled for regular captures */
+	Capture,
+
+	/** Enabled for scene captures and custom render passes (Render In Main Renderer) */
+	CaptureOrCustomRenderPass
+};
+
 /**
  *	Used to capture a 'snapshot' of the scene from a single plane and feed it to a render target.
  */
@@ -117,10 +130,19 @@ public:
 	UPROPERTY(EditAnywhere, BlueprintReadWrite, AdvancedDisplay, Category=SceneCapture)
 	FVector ClipPlaneNormal;
 	
-	/** Render scene capture as additional render passes of the main renderer rather than as an independent renderer. Applies to scene depth, device depth, base color, and normal modes. */
-	UPROPERTY(EditAnywhere, BlueprintReadWrite, Category=PassInfo, meta = (EditCondition = "CaptureSource == ESceneCaptureSource::SCS_SceneDepth || CaptureSource == ESceneCaptureSource::SCS_DeviceDepth || CaptureSource == ESceneCaptureSource::SCS_BaseColor || CaptureSource == ESceneCaptureSource::SCS_Normal"))
+	/** Render scene capture as additional render passes of the main renderer rather than as an independent renderer. Applies to scene depth, device depth, base color, normal, and scene color modes (disables lighting and shadows). */
+	UPROPERTY(EditAnywhere, BlueprintReadWrite, Category=PassInfo, meta = (EditCondition = "CaptureSource == ESceneCaptureSource::SCS_SceneDepth || CaptureSource == ESceneCaptureSource::SCS_DeviceDepth || CaptureSource == ESceneCaptureSource::SCS_BaseColor || CaptureSource == ESceneCaptureSource::SCS_Normal || CaptureSource == ESceneCaptureSource::SCS_SceneColorHDR || CaptureSource == ESceneCaptureSource::SCS_SceneColorHDRNoAlpha || CaptureSource == ESceneCaptureSource::SCS_SceneColorSceneDepth"))
 	bool bRenderInMainRenderer = false;
 
+	/**
+	 * Option to enable a debug feature which outputs base color to the emissive channel when lighting is disabled via ShowFlag
+	 * or via Render In Main Renderer, which renders the capture as a Custom Render Pass.  Note that the debug feature requires
+	 * development shaders to be compiled, generally only true in non-shipping builds on PC!  To work in other cases, materials
+	 * should directly write to the emissive channel (or be unlit materials), rather than counting on the debug feature.
+	 */
+	UPROPERTY(EditAnywhere, BlueprintReadWrite, Category=PassInfo)
+	ESceneCaptureUnlitViewmode UnlitViewmode = ESceneCaptureUnlitViewmode::Capture;
+
 	/** 
 	 * True if we did a camera cut this frame. Automatically reset to false at every capture.
 	 * This flag affects various things in the renderer (such as whether to use the occlusion queries from last frame, and motion blur).
@@ -154,20 +176,20 @@ public:
 	FIntPoint MainViewResolutionDivisor = { 1, 1 };
 
 	/** Expose BaseColor as a UserSceneTexture.  Requires "Render In Main Renderer".  Enables Main View Family and Resolution, disables "Ignore Screen Percentage".  Useful to get multiple outputs from a Custom Render Pass. */
-	UPROPERTY(EditAnywhere, BlueprintReadWrite, Category=PassInfo, meta = (EditCondition = "bRenderInMainRenderer && (CaptureSource == ESceneCaptureSource::SCS_SceneDepth || CaptureSource == ESceneCaptureSource::SCS_DeviceDepth || CaptureSource == ESceneCaptureSource::SCS_BaseColor || CaptureSource == ESceneCaptureSource::SCS_Normal)"))
+	UPROPERTY(EditAnywhere, BlueprintReadWrite, Category=PassInfo, meta = (EditCondition = "bRenderInMainRenderer && (CaptureSource == ESceneCaptureSource::SCS_SceneDepth || CaptureSource == ESceneCaptureSource::SCS_DeviceDepth || CaptureSource == ESceneCaptureSource::SCS_BaseColor || CaptureSource == ESceneCaptureSource::SCS_Normal || CaptureSource == ESceneCaptureSource::SCS_SceneColorHDR || CaptureSource == ESceneCaptureSource::SCS_SceneColorHDRNoAlpha || CaptureSource == ESceneCaptureSource::SCS_SceneColorSceneDepth)"))
 	FName UserSceneTextureBaseColor;
 
 	/** Expose Normal as a UserSceneTexture.  Requires "Render In Main Renderer".  Enables Main View Family and Resolution, disables "Ignore Screen Percentage".  Useful to get multiple outputs from a Custom Render Pass. */
-	UPROPERTY(EditAnywhere, BlueprintReadWrite, Category=PassInfo, meta = (EditCondition = "bRenderInMainRenderer && (CaptureSource == ESceneCaptureSource::SCS_SceneDepth || CaptureSource == ESceneCaptureSource::SCS_DeviceDepth || CaptureSource == ESceneCaptureSource::SCS_BaseColor || CaptureSource == ESceneCaptureSource::SCS_Normal)"))
+	UPROPERTY(EditAnywhere, BlueprintReadWrite, Category=PassInfo, meta = (EditCondition = "bRenderInMainRenderer && (CaptureSource == ESceneCaptureSource::SCS_SceneDepth || CaptureSource == ESceneCaptureSource::SCS_DeviceDepth || CaptureSource == ESceneCaptureSource::SCS_BaseColor || CaptureSource == ESceneCaptureSource::SCS_Normal || CaptureSource == ESceneCaptureSource::SCS_SceneColorHDR || CaptureSource == ESceneCaptureSource::SCS_SceneColorHDRNoAlpha || CaptureSource == ESceneCaptureSource::SCS_SceneColorSceneDepth)"))
 	FName UserSceneTextureNormal;
 
 	/** Expose SceneColor (emissive/unlit) as a UserSceneTexture.  Requires "Render In Main Renderer".  Enables Main View Family and Resolution, disables "Ignore Screen Percentage".  Useful to get multiple outputs from a Custom Render Pass. */
-	UPROPERTY(EditAnywhere, BlueprintReadWrite, Category=PassInfo, meta = (EditCondition = "bRenderInMainRenderer && (CaptureSource == ESceneCaptureSource::SCS_SceneDepth || CaptureSource == ESceneCaptureSource::SCS_DeviceDepth || CaptureSource == ESceneCaptureSource::SCS_BaseColor || CaptureSource == ESceneCaptureSource::SCS_Normal)"))
+	UPROPERTY(EditAnywhere, BlueprintReadWrite, Category=PassInfo, meta = (EditCondition = "bRenderInMainRenderer && (CaptureSource == ESceneCaptureSource::SCS_SceneDepth || CaptureSource == ESceneCaptureSource::SCS_DeviceDepth || CaptureSource == ESceneCaptureSource::SCS_BaseColor || CaptureSource == ESceneCaptureSource::SCS_Normal || CaptureSource == ESceneCaptureSource::SCS_SceneColorHDR || CaptureSource == ESceneCaptureSource::SCS_SceneColorHDRNoAlpha || CaptureSource == ESceneCaptureSource::SCS_SceneColorSceneDepth)"))
 	FName UserSceneTextureSceneColor;
 
 	inline bool ShouldRenderInMainRenderer() const
 	{
-		return bRenderInMainRenderer && (CaptureSource == ESceneCaptureSource::SCS_SceneDepth || CaptureSource == ESceneCaptureSource::SCS_DeviceDepth || CaptureSource == ESceneCaptureSource::SCS_BaseColor || CaptureSource == ESceneCaptureSource::SCS_Normal);
+		return bRenderInMainRenderer && (CaptureSource == ESceneCaptureSource::SCS_SceneDepth || CaptureSource == ESceneCaptureSource::SCS_DeviceDepth || CaptureSource == ESceneCaptureSource::SCS_BaseColor || CaptureSource == ESceneCaptureSource::SCS_Normal || CaptureSource == ESceneCaptureSource::SCS_SceneColorHDR || CaptureSource == ESceneCaptureSource::SCS_SceneColorHDRNoAlpha || CaptureSource == ESceneCaptureSource::SCS_SceneColorSceneDepth);
 	}
 
 	inline bool ShouldRenderWithMainViewResolution() const
diff --git a/Engine/Source/Runtime/Engine/Private/PrimitiveDrawingUtils.cpp b/Engine/Source/Runtime/Engine/Private/PrimitiveDrawingUtils.cpp
--- a/Engine/Source/Runtime/Engine/Private/PrimitiveDrawingUtils.cpp
+++ b/Engine/Source/Runtime/Engine/Private/PrimitiveDrawingUtils.cpp
@@ -1562,7 +1562,7 @@ bool IsRichView(const FSceneViewFamily& ViewFamily)
 		// Force FDrawBasePassDynamicMeshAction to be used since it has access to the view and can implement the show flags
 		!ViewFamily.EngineShowFlags.VolumetricLightmap ||
 		!ViewFamily.EngineShowFlags.IndirectLightingCache ||
-		!ViewFamily.EngineShowFlags.Lighting ||
+		(!ViewFamily.EngineShowFlags.Lighting && ViewFamily.EngineShowFlags.UnlitRichView) ||
 		!ViewFamily.EngineShowFlags.Materials)
 	{
 		return true;
diff --git a/Engine/Source/Runtime/Engine/Private/SceneView.cpp b/Engine/Source/Runtime/Engine/Private/SceneView.cpp
--- a/Engine/Source/Runtime/Engine/Private/SceneView.cpp
+++ b/Engine/Source/Runtime/Engine/Private/SceneView.cpp
@@ -2740,7 +2740,7 @@ void SetupCommonViewUniformBufferParameters(
 		(float)InPrevViewMatrices.GetTemporalAAJitter().X, (float)InPrevViewMatrices.GetTemporalAAJitter().Y );
 
 	ViewUniformShaderParameters.DebugViewModeMask = Inputs.DebugViewShaderMode != DVSM_None ? 1 : 0;
-	ViewUniformShaderParameters.UnlitViewmodeMask = Inputs.EngineShowFlags && (!Inputs.EngineShowFlags->Lighting || Inputs.EngineShowFlags->PathTracing) ? 1 : 0;
+	ViewUniformShaderParameters.UnlitViewmodeMask = Inputs.EngineShowFlags && ((!Inputs.EngineShowFlags->Lighting && Inputs.EngineShowFlags->UnlitViewmode) || Inputs.EngineShowFlags->PathTracing) ? 1 : 0;
 	ViewUniformShaderParameters.OutOfBoundsMask = Inputs.EngineShowFlags && Inputs.EngineShowFlags->VisualizeOutOfBoundsPixels ? 1 : 0;
 
 #if !(UE_BUILD_SHIPPING || UE_BUILD_TEST)
diff --git a/Engine/Source/Runtime/Engine/Public/Rendering/CustomRenderPass.h b/Engine/Source/Runtime/Engine/Public/Rendering/CustomRenderPass.h
--- a/Engine/Source/Runtime/Engine/Public/Rendering/CustomRenderPass.h
+++ b/Engine/Source/Runtime/Engine/Public/Rendering/CustomRenderPass.h
@@ -53,7 +53,7 @@ public:
 	friend class FDeferredShadingSceneRenderer;
 
 	/** Which render passes are needed for the custom render pass. */
-	enum class ERenderMode
+	enum class ERenderMode : uint8
 	{
 		/** Render depth pre-pass only. */
 		DepthPass,
@@ -62,7 +62,7 @@ public:
 	};
 
 	/** The output type into the render target. What type is valid depends on ERenderMode. */
-	enum class ERenderOutput
+	enum class ERenderOutput : uint8
 	{
 		/** Used with ERenderMode::DepthPass. */
 		SceneDepth,
@@ -115,6 +115,16 @@ public:
 	FRDGTextureRef GetRenderTargetTexture() const { return RenderTargetTexture; }
 	const FIntPoint& GetRenderTargetSize() const { return RenderTargetSize; }
 
+	/** Returns true if translucency is enabled AND an output uses SceneColor.  We can skip translucency if SceneColor (where translucency is written) is unused. */
+	bool IsTranslucentIncluded() const
+	{
+		return bSceneColorWithTranslucent &&
+			(RenderOutput == FCustomRenderPassBase::ERenderOutput::SceneColorAndAlpha ||
+			 RenderOutput == FCustomRenderPassBase::ERenderOutput::SceneColorNoAlpha ||
+			 RenderOutput == FCustomRenderPassBase::ERenderOutput::SceneColorAndDepth ||
+			 bSceneColorIsUserSceneTexture);
+	}
+
 	ENGINE_API void SetUserData(TUniquePtr<ICustomRenderPassUserData>&& InUserData);
 	ENGINE_API ICustomRenderPassUserData* GetUserData(const FName& InTypeName) const;
 
@@ -144,6 +154,12 @@ protected:
 	FString DebugName;
 	ERenderMode RenderMode = ERenderMode::DepthPass;
 	ERenderOutput RenderOutput = ERenderOutput::SceneDepth;
+	bool bSceneColorWithTranslucent = false;
+
+	// If the SceneColor generated by this Custom Render Pass is referenced as a UserSceneTexture, translucency should render if enabled,
+	// even if RenderOutput is not SceneColor.  Otherwise, translucency is disabled with the assumption it's wasteful to render translucency
+	// to a SceneColor buffer that's not used.  See IsTranslucentIncluded() function above.
+	bool bSceneColorIsUserSceneTexture = false;
 
 	FRDGTextureRef RenderTargetTexture = nullptr;
 	FIntPoint RenderTargetSize = FIntPoint(ForceInit);
diff --git a/Engine/Source/Runtime/Engine/Public/SceneInterface.h b/Engine/Source/Runtime/Engine/Public/SceneInterface.h
--- a/Engine/Source/Runtime/Engine/Public/SceneInterface.h
+++ b/Engine/Source/Runtime/Engine/Public/SceneInterface.h
@@ -694,6 +694,7 @@ public:
 		FCustomRenderPassRendererInput()
 			: EngineShowFlags(ESFIM_Game)
 		{
+			EngineShowFlags.DisableFeaturesForUnlit();
 		}
 
 		/** Data used to construct scene view for the custom render pass. */
diff --git a/Engine/Source/Runtime/Engine/Public/ShowFlags.h b/Engine/Source/Runtime/Engine/Public/ShowFlags.h
--- a/Engine/Source/Runtime/Engine/Public/ShowFlags.h
+++ b/Engine/Source/Runtime/Engine/Public/ShowFlags.h
@@ -266,6 +266,24 @@ struct FEngineShowFlags
 		SetVirtualShadowMapPersistentData(false);
 	}
 
+	// Features that make sense to disable together for unlit rendering (custom render passes or unlit scene captures)
+	void DisableFeaturesForUnlit(bool bAllowAtmosphere = false)
+	{
+		SetDynamicShadows(false);		// Not supported for CRPs (shadows aren't used if lighting is disabled)
+		SetLighting(false);				// Not supported for CRPs
+		SetPostProcessing(false);		// Not supported for CRPs (post processing is skipped for unlit scene captures, via bResolveScene being set to false)
+		SetRefraction(false);			// Not supported for CRPs (enables separate translucency, which requires post processing to run to composite it back in)
+
+		SetLightFunctions(false);		// Needless performance cost (updates light function atlas, which is unused if lighting is disabled)
+		SetUnlitRichView(false);		// Needless performance cost (expensive dynamic rendering code path, related to main view debug modes active when lighting is disabled)
+		SetUnlitViewmode(false);		// Development shader feature that writes BaseColor to the emissive channel when lighting is disabled (typically only on non-shipping PC builds)
+
+		if (!bAllowAtmosphere)
+		{
+			SetAtmosphere(false);		// Technically works with unlit, but costs performance, and is assumed to be undesirable by default
+		}
+	}
+
 	bool IsVisualizeCalibrationEnabled() const
 	{
 		return (VisualizeCalibrationColor || VisualizeCalibrationGrayscale || VisualizeCalibrationCustom);
diff --git a/Engine/Source/Runtime/Engine/Public/ShowFlagsValues.inl b/Engine/Source/Runtime/Engine/Public/ShowFlagsValues.inl
--- a/Engine/Source/Runtime/Engine/Public/ShowFlagsValues.inl
+++ b/Engine/Source/Runtime/Engine/Public/ShowFlagsValues.inl
@@ -227,6 +227,10 @@ SHOWFLAG_FIXED_IN_SHIPPING(0, BSPSplit, SFG_Advanced, NSLOCTEXT("UnrealEd", "BSP
 SHOWFLAG_FIXED_IN_SHIPPING(0, Brushes, SFG_Hidden, NSLOCTEXT("UnrealEd", "BrushesSF", "Brushes"))
 /** Show the usual material light interaction */
 SHOWFLAG_ALWAYS_ACCESSIBLE(Lighting, SFG_Hidden, NSLOCTEXT("UnrealEd", "LightingSF", "Lighting"))
+/** Enable rich view when lighting is disabled, forcing a slower dynamic render code path.  Used for unlit debug view modes that override shaders. */
+SHOWFLAG_ALWAYS_ACCESSIBLE(UnlitRichView, SFG_Hidden, NSLOCTEXT("UnrealEd", "UnlitRichViewSF", "UnlitRichView"))
+/** Enable view mode that causes shaders to output Base Color as Scene Color (emissive) when lighting is disabled.  Only works when development shaders are enabled, typically only on PC. */
+SHOWFLAG_ALWAYS_ACCESSIBLE(UnlitViewmode, SFG_Hidden, NSLOCTEXT("UnrealEd", "UnlitViewmodeSF", "UnlitViewmode"))
 /** Execute the deferred light passes, can be disabled for debugging, for now SHOWFLAG_ALWAYS_ACCESSIBLE because it's exposed in SceneCapture */
 SHOWFLAG_ALWAYS_ACCESSIBLE(DeferredLighting, SFG_Advanced, NSLOCTEXT("UnrealEd", "DeferredLightingSF", "DeferredLighting"))
 /** Special: Allows to hide objects in the editor, is evaluated per primitive */
diff --git a/Engine/Source/Runtime/Renderer/Private/BasePassRendering.cpp b/Engine/Source/Runtime/Renderer/Private/BasePassRendering.cpp
--- a/Engine/Source/Runtime/Renderer/Private/BasePassRendering.cpp
+++ b/Engine/Source/Runtime/Renderer/Private/BasePassRendering.cpp
@@ -1213,8 +1213,8 @@ void FDeferredShadingSceneRenderer::RenderBasePass(
 		ClearGBufferAtMaxZ(GraphBuilder, InViews, BasePassRenderTargets, SceneColorClearValue);
 	}
 
-	// The anisotropy GBuffer is used for lighting, and not needed when running a scene capture with ERendererOutput::BasePass (BaseColor or Normal output), which doesn't run lighting.
-	if (ShouldRenderAnisotropyPass(InViews) && Renderer.GetRendererOutput() == ERendererOutput::FinalSceneColor)
+	// The anisotropy GBuffer is used for lighting
+	if (ShouldRenderAnisotropyPass(InViews) && ViewFamily.EngineShowFlags.Lighting)
 	{
 		RenderAnisotropyPass(GraphBuilder, InViews, SceneTextures, ViewFamily.Scene->GetRenderScene(), bEnableParallelBasePasses);
 	}
diff --git a/Engine/Source/Runtime/Renderer/Private/DeferredShadingRenderer.cpp b/Engine/Source/Runtime/Renderer/Private/DeferredShadingRenderer.cpp
--- a/Engine/Source/Runtime/Renderer/Private/DeferredShadingRenderer.cpp
+++ b/Engine/Source/Runtime/Renderer/Private/DeferredShadingRenderer.cpp
@@ -1375,7 +1375,7 @@ void FDeferredShadingSceneRenderer::Render(FRDGBuilder& GraphBuilder)
 
 	const bool bUseVirtualTexturing = UseVirtualTexturing(ShaderPlatform);
 
-	// Virtual texturing runs for ERendererOutput::BasePass or ERendererOutput::FinalSceneColor
+	// Virtual texturing isn't needed for depth prepass
 	if (bUseVirtualTexturing && RendererOutput != ERendererOutput::DepthPrepassOnly)
 	{
 		FVirtualTextureUpdateSettings Settings;
@@ -1773,7 +1773,7 @@ void FDeferredShadingSceneRenderer::Render(FRDGBuilder& GraphBuilder)
 	bool bComputeLightGrid = false;
 	bool bAnyLumenEnabled = false;
 
-	// Virtual texturing runs for ERendererOutput::BasePass or ERendererOutput::FinalSceneColor
+	// Virtual texturing isn't needed for depth prepass
 	if (RendererOutput != ERendererOutput::DepthPrepassOnly)
 	{
 		if (bUseVirtualTexturing)
@@ -2077,6 +2077,25 @@ void FDeferredShadingSceneRenderer::Render(FRDGBuilder& GraphBuilder)
 						SceneTextures.SetupMode |= ESceneTextureSetupMode::GBuffers;
 						SceneTextures.UniformBuffer = CreateSceneTextureUniformBuffer(GraphBuilder, &SceneTextures, FeatureLevel, SceneTextures.SetupMode);
 					}
+
+					if (CustomRenderPass->IsTranslucentIncluded())
+					{
+						// Empty defaults
+						FTranslucencyLightingVolumeTextures TranslucencyLightingVolumeTextures;
+						FTranslucencyPassResourcesMap TranslucencyResourceMap(CustomRenderPassViews.Num());
+						const bool bStandardTranslucentCanRenderSeparate = false;
+						FRDGTextureMSAA TranslucencySharedDepthTexture;
+						FSeparateTranslucencyDimensions CustomTranslucencyDimensions = { SceneTexturesConfig.Extent };
+
+						FReflectionCaptureShaderData EmptyData;
+						TUniformBufferRef<FReflectionCaptureShaderData> EmptyReflectionCaptureUniformBuffer = TUniformBufferRef<FReflectionCaptureShaderData>::CreateUniformBufferImmediate(EmptyData, UniformBuffer_SingleFrame);
+						for (FViewInfo& View : CustomRenderPassViews)
+						{
+							View.ReflectionCaptureUniformBuffer = EmptyReflectionCaptureUniformBuffer;
+						}
+
+						RenderTranslucency(*this, GraphBuilder, SceneTextures, TranslucencyLightingVolumeTextures, &TranslucencyResourceMap, CustomRenderPassViews, ETranslucencyView::AboveWater, CustomTranslucencyDimensions, InstanceCullingManager, bStandardTranslucentCanRenderSeparate, TranslucencySharedDepthTexture);
+					}
 				}
 
 				CopySceneCaptureComponentToTarget(GraphBuilder, SceneTextures, CustomRenderPass->GetRenderTargetTexture(), ViewFamily, CustomRenderPassViews);
@@ -2193,7 +2212,7 @@ void FDeferredShadingSceneRenderer::Render(FRDGBuilder& GraphBuilder)
 
 	FRDGTextureRef ViewFamilyTexture = TryCreateViewFamilyTexture(GraphBuilder, ViewFamily);
 	FRDGTextureRef ViewFamilyDepthTexture = TryCreateViewFamilyDepthTexture(GraphBuilder, ViewFamily);
-	if (RendererOutput == ERendererOutput::DepthPrepassOnly || RendererOutput == ERendererOutput::BasePass)
+	if (RendererOutput == ERendererOutput::DepthPrepassOnly)
 	{
 		const FSingleLayerWaterPrePassResult* SingleLayerWaterPrePassResult = nullptr;
 		if (ShouldRenderSingleLayerWaterDepthPrepass(Views))
@@ -2201,52 +2220,7 @@ void FDeferredShadingSceneRenderer::Render(FRDGBuilder& GraphBuilder)
 			SingleLayerWaterPrePassResult = RenderSingleLayerWaterDepthPrepass(GraphBuilder, Views, SceneTextures);
 		}
 
-		bool bOcclusionBeforeBasePass = false;
-		if (RendererOutput == ERendererOutput::BasePass)
-		{
-			// Early occlusion queries
-			bOcclusionBeforeBasePass = ((DepthPass.EarlyZPassMode == EDepthDrawingMode::DDM_AllOccluders) || bIsEarlyDepthComplete);
-			if (bOcclusionBeforeBasePass)
-			{
-				FroxelRenderer = RenderOcclusionLambda();
-			}
-
-			SceneTextures.SetupMode |= ESceneTextureSetupMode::SceneColor;
-			SceneTextures.UniformBuffer = CreateSceneTextureUniformBuffer(GraphBuilder, &SceneTextures, FeatureLevel, SceneTextures.SetupMode);
-
-			for (FSceneViewExtensionRef& ViewExtension : ViewFamily.ViewExtensions)
-			{
-				ViewExtension->PreRenderBasePass_RenderThread(GraphBuilder, ShouldRenderPrePass() /*bDepthBufferIsPopulated*/);
-			}
-
-			RenderBasePass(*this, GraphBuilder, Views, SceneTextures, DBufferTextures, BasePassDepthStencilAccess, /*ForwardScreenSpaceShadowMaskTexture=*/nullptr, InstanceCullingManager, bNaniteEnabled, Scene->NaniteShadingCommands[ENaniteMeshPass::BasePass], NaniteRasterResults);
-
-			if (ShouldRenderSingleLayerWater(Views))
-			{
-				// Virtual shadow map uniforms need to be initialized with dummy data for water.  Their initialization
-				// was skipped above due to (RendererOutput == ERendererOutput::FinalSceneColor) being false.
-				VirtualShadowMapArray.Initialize(GraphBuilder, Scene->GetVirtualShadowMapCache(), /*bEnableVirtualShadowMaps=*/false, ViewFamily.EngineShowFlags);
-
-				FSceneWithoutWaterTextures SceneWithoutWaterTextures;
-				RenderSingleLayerWater(GraphBuilder, Views, SceneTextures, SingleLayerWaterPrePassResult, /*bShouldRenderVolumetricCloud=*/false, SceneWithoutWaterTextures, LumenFrameTemporaries, /*bIsCameraUnderWater=*/false);
-			}
-
-			SceneTextures.SetupMode |= ESceneTextureSetupMode::GBuffers;
-			SceneTextures.UniformBuffer = CreateSceneTextureUniformBuffer(GraphBuilder, &SceneTextures, FeatureLevel, SceneTextures.SetupMode);
-
-			if (bUseVirtualTexturing)
-			{
-				RDG_EVENT_SCOPE_STAT(GraphBuilder, VirtualTextureUpdate, "VirtualTextureUpdate");
-				RDG_GPU_STAT_SCOPE(GraphBuilder, VirtualTextureUpdate);
-
-				VirtualTextureFeedbackEnd(GraphBuilder);
-			}
-		}
-
-		if (!bOcclusionBeforeBasePass)
-		{
-			FroxelRenderer = RenderOcclusionLambda();
-		}
+		FroxelRenderer = RenderOcclusionLambda();
 
 		if (bUpdateNaniteStreaming)
 		{
@@ -3043,7 +3017,7 @@ void FDeferredShadingSceneRenderer::Render(FRDGBuilder& GraphBuilder)
 				SCOPE_CYCLE_COUNTER(STAT_TranslucencyDrawTime);
 				const bool bStandardTranslucentCanRenderSeparate = false;
 				FRDGTextureMSAA SharedDepthTexture;
-				RenderTranslucency(GraphBuilder, SceneTextures, TranslucencyLightingVolumeTextures, &TranslucencyResourceMap, ETranslucencyView::UnderWater, InstanceCullingManager, bStandardTranslucentCanRenderSeparate, SharedDepthTexture);
+				RenderTranslucency(*this, GraphBuilder, SceneTextures, TranslucencyLightingVolumeTextures, &TranslucencyResourceMap, Views, ETranslucencyView::UnderWater, SeparateTranslucencyDimensions, InstanceCullingManager, bStandardTranslucentCanRenderSeparate, SharedDepthTexture);
 				EnumRemoveFlags(TranslucencyViewsToRender, ETranslucencyView::UnderWater);
 			}
 
@@ -3155,7 +3129,7 @@ void FDeferredShadingSceneRenderer::Render(FRDGBuilder& GraphBuilder)
 			{
 				// Render all remaining translucency views.
 				const bool bStandardTranslucentCanRenderSeparate = bShouldRenderDistortion; // It is only needed to render standard translucent as separate when there is distortion (non self distortion of transmittance/specular/etc.)
-				RenderTranslucency(GraphBuilder, SceneTextures, TranslucencyLightingVolumeTextures, &TranslucencyResourceMap, TranslucencyViewsToRender, InstanceCullingManager, bStandardTranslucentCanRenderSeparate, TranslucencySharedDepthTexture);
+				RenderTranslucency(*this, GraphBuilder, SceneTextures, TranslucencyLightingVolumeTextures, &TranslucencyResourceMap, Views, TranslucencyViewsToRender, SeparateTranslucencyDimensions, InstanceCullingManager, bStandardTranslucentCanRenderSeparate, TranslucencySharedDepthTexture);
 			}
 
 			// Compose hair before velocity/distortion pass since these pass write depth value, 
diff --git a/Engine/Source/Runtime/Renderer/Private/DeferredShadingRenderer.h b/Engine/Source/Runtime/Renderer/Private/DeferredShadingRenderer.h
--- a/Engine/Source/Runtime/Renderer/Private/DeferredShadingRenderer.h
+++ b/Engine/Source/Runtime/Renderer/Private/DeferredShadingRenderer.h
@@ -860,24 +860,30 @@ private:
 		FRDGTextureRef LightingChannelsTexture);
 	
 	/** Renders the scene's translucency passes. */
-	void RenderTranslucency(
+	static void RenderTranslucency(
+		FDeferredShadingSceneRenderer& Renderer,
 		FRDGBuilder& GraphBuilder,
 		const FSceneTextures& SceneTextures,
 		const FTranslucencyLightingVolumeTextures& TranslucencyLightingVolumeTextures,
 		FTranslucencyPassResourcesMap* OutTranslucencyResourceMap,
+		TArray<FViewInfo>& InViews,
 		ETranslucencyView ViewsToRender,
+		const FSeparateTranslucencyDimensions& SeparateTranslucencyDimensions,
 		FInstanceCullingManager& InstanceCullingManager,
 		bool bStandardTranslucentCanRenderSeparate,
 		FRDGTextureMSAA& OutSharedDepthTexture);
 
 	/** Renders the scene's translucency given a specific pass. */
-	void RenderTranslucencyInner(
+	static void RenderTranslucencyInner(
+		FDeferredShadingSceneRenderer& Renderer,
 		FRDGBuilder& GraphBuilder,
 		const FMinimalSceneTextures& SceneTextures,
 		const FTranslucencyLightingVolumeTextures& TranslucencyLightingVolumeTextures,
 		FTranslucencyPassResourcesMap* OutTranslucencyResourceMap,
 		FRDGTextureMSAA SharedDepthTexture,
+		TArray<FViewInfo>& InViews,
 		ETranslucencyView ViewsToRender,
+		const FSeparateTranslucencyDimensions& SeparateTranslucencyDimensions,
 		FRDGTextureRef SceneColorCopyTexture,
 		ETranslucencyPass::Type TranslucencyPass,
 		FInstanceCullingManager& InstanceCullingManager,
diff --git a/Engine/Source/Runtime/Renderer/Private/MobileShadingRenderer.cpp b/Engine/Source/Runtime/Renderer/Private/MobileShadingRenderer.cpp
--- a/Engine/Source/Runtime/Renderer/Private/MobileShadingRenderer.cpp
+++ b/Engine/Source/Runtime/Renderer/Private/MobileShadingRenderer.cpp
@@ -867,7 +867,7 @@ void FMobileSceneRenderer::RenderMaskedPrePass(FRHICommandList& RHICmdList, cons
 	}
 }
 
-void FMobileSceneRenderer::RenderCustomRenderPassBasePass(FRDGBuilder& GraphBuilder, TArrayView<FViewInfo> InViews, FRDGTextureRef ViewFamilyTexture, FSceneTextures& SceneTextures)
+void FMobileSceneRenderer::RenderCustomRenderPassBasePass(FRDGBuilder& GraphBuilder, TArrayView<FViewInfo> InViews, FRDGTextureRef ViewFamilyTexture, FSceneTextures& SceneTextures, bool bIncludeTranslucent)
 {
 	FRenderTargetBindingSlots BasePassRenderTargets;
 	// Use the same subpass hints as main render, to avoid generating new PSOs 
@@ -899,26 +899,59 @@ void FMobileSceneRenderer::RenderCustomRenderPassBasePass(FRDGBuilder& GraphBuil
 		PassParameters->MobileBasePass = CreateMobileBasePassUniformBuffer(GraphBuilder, View, EMobileBasePass::Opaque, SetupMode);
 		PassParameters->RenderTargets = BasePassRenderTargets;
 
+		FMobileRenderPassParameters* TranslucencyPassParameters = nullptr;
+		if (bIncludeTranslucent)
+		{
+			TranslucencyPassParameters = GraphBuilder.AllocParameters<FMobileRenderPassParameters>();
+			TranslucencyPassParameters->View = View.GetShaderParameters();
+			TranslucencyPassParameters->MobileBasePass = CreateMobileBasePassUniformBuffer(GraphBuilder, View, EMobileBasePass::Translucent, SetupMode);
+			TranslucencyPassParameters->RenderTargets[0] = BasePassRenderTargets[0];
+			TranslucencyPassParameters->RenderTargets[0].SetLoadAction(ERenderTargetLoadAction::ELoad);
+			TranslucencyPassParameters->RenderTargets.DepthStencil = BasePassRenderTargets.DepthStencil;
+			TranslucencyPassParameters->RenderTargets.DepthStencil.SetDepthLoadAction(ERenderTargetLoadAction::ELoad);
+		}
+
 		if (Scene->GPUScene.IsEnabled())
 		{
 			View.ParallelMeshDrawCommandPasses[EMeshPass::BasePass].BuildRenderingCommands(GraphBuilder, Scene->GPUScene, PassParameters->InstanceCullingDrawParams);
+
+			if (bIncludeTranslucent)
+			{
+				View.ParallelMeshDrawCommandPasses[EMeshPass::TranslucencyAll].BuildRenderingCommands(GraphBuilder, Scene->GPUScene, TranslucencyPassParameters->InstanceCullingDrawParams);
+			}
 		}
 
 		GraphBuilder.AddPass(
 			RDG_EVENT_NAME("RenderMobileBasePass"),
 			PassParameters,
 			ERDGPassFlags::Raster,
-			[this, PassParameters, ViewContext, &SceneTextures, NumAdditionalSubpasses](FRHICommandList& RHICmdList)
+			[this, PassParameters, ViewContext, NumAdditionalSubpasses](FRHICommandList& RHICmdList)
 			{
 				FViewInfo& View = *ViewContext.ViewInfo;
 				RenderMobileBasePass(RHICmdList, View, &PassParameters->InstanceCullingDrawParams);
 
+				// TODO:  Should this render decals?  Deferred shading custom render passes do.
+
 				for (int32 AddSubpass = 0; AddSubpass < NumAdditionalSubpasses; ++AddSubpass)
 				{
 					RHICmdList.NextSubpass();
 				}
 			});
 
+		if (bIncludeTranslucent)
+		{
+			GraphBuilder.AddPass(
+				RDG_EVENT_NAME("RenderMobileTranslucentPass"),
+				TranslucencyPassParameters,
+				ERDGPassFlags::Raster,
+				[this, TranslucencyPassParameters, ViewContext, InViews, &SceneTextures](FRHICommandList& RHICmdList)
+				{
+					// Custom render passes run all translucency in a single pass
+					FViewInfo& View = *ViewContext.ViewInfo;
+					RenderTranslucency(RHICmdList, View, InViews, ETranslucencyPass::TPT_AllTranslucency, EMeshPass::TranslucencyAll, &TranslucencyPassParameters->InstanceCullingDrawParams);
+				});
+		}
+
 		if (!bIsFullDepthPrepassEnabled)
 		{
 			AddResolveSceneDepthPass(GraphBuilder, View, SceneTextures.Depth);
@@ -1249,7 +1282,7 @@ void FMobileSceneRenderer::Render(FRDGBuilder& GraphBuilder)
 				// Render base pass if the custom pass requires it. Otherwise if full depth prepass is not enabled, then depth is generated in the base pass.
 				if (CustomRenderPass->GetRenderMode() == FCustomRenderPassBase::ERenderMode::DepthAndBasePass || (CustomRenderPass->GetRenderMode() == FCustomRenderPassBase::ERenderMode::DepthPass && !bIsFullDepthPrepassEnabled))
 				{
-					RenderCustomRenderPassBasePass(GraphBuilder, CustomRenderPassViews, ViewFamilyTexture, SceneTextures);
+					RenderCustomRenderPassBasePass(GraphBuilder, CustomRenderPassViews, ViewFamilyTexture, SceneTextures, CustomRenderPass->IsTranslucentIncluded());
 				}
 
 				SceneTextures.MobileSetupMode = EMobileSceneTextureSetupMode::SceneColor | EMobileSceneTextureSetupMode::SceneDepth | EMobileSceneTextureSetupMode::SceneDepthAux;
@@ -1667,7 +1700,7 @@ void FMobileSceneRenderer::RenderForwardSinglePass(FRDGBuilder& GraphBuilder, FM
 			RenderFog(RHICmdList, View);
 		}
 		// Draw translucency.
-		RenderTranslucency(RHICmdList, View);
+		RenderTranslucency(RHICmdList, View, Views, StandardTranslucencyPass, StandardTranslucencyMeshPass, &TranslucencyInstanceCullingDrawParams);
 		
 #if UE_ENABLE_DEBUG_DRAWING
 		if ((!IsMobileHDR() || bTonemapSubpass) && FSceneRenderer::ShouldCompositeDebugPrimitivesInPostProcess(View))
@@ -1779,7 +1812,7 @@ void FMobileSceneRenderer::RenderForwardMultiPass(FRDGBuilder& GraphBuilder, FMo
 		RenderModulatedShadowProjections(RHICmdList, ViewContext.ViewIndex, View);
 		RenderFog(RHICmdList, View);
 		// Draw translucency.
-		RenderTranslucency(RHICmdList, View);
+		RenderTranslucency(RHICmdList, View, Views, StandardTranslucencyPass, StandardTranslucencyMeshPass, &TranslucencyInstanceCullingDrawParams);
 
 		if (bDoOcclusionQueries)
 		{
@@ -2010,7 +2043,7 @@ void FMobileSceneRenderer::RenderDeferredSinglePass(FRDGBuilder& GraphBuilder, F
 				}
 				RenderFog(RHICmdList, View);
 				// Draw translucency.
-				RenderTranslucency(RHICmdList, View);
+				RenderTranslucency(RHICmdList, View, Views, StandardTranslucencyPass, StandardTranslucencyMeshPass, &TranslucencyInstanceCullingDrawParams);
 
 				if (bDoOcclusionQueires)
 				{
@@ -2175,7 +2208,7 @@ void FMobileSceneRenderer::RenderDeferredMultiPass(FRDGBuilder& GraphBuilder, FS
 				RenderFog(RHICmdList, View);
 
 				// Draw translucency.
-				RenderTranslucency(RHICmdList, View);
+				RenderTranslucency(RHICmdList, View, Views, StandardTranslucencyPass, StandardTranslucencyMeshPass, &TranslucencyInstanceCullingDrawParams);
 			});
 	}
 }
diff --git a/Engine/Source/Runtime/Renderer/Private/MobileTranslucentRendering.cpp b/Engine/Source/Runtime/Renderer/Private/MobileTranslucentRendering.cpp
--- a/Engine/Source/Runtime/Renderer/Private/MobileTranslucentRendering.cpp
+++ b/Engine/Source/Runtime/Renderer/Private/MobileTranslucentRendering.cpp
@@ -4,9 +4,11 @@
 #include "RenderCore.h"
 #include "TranslucentRendering.h"
 
-void FMobileSceneRenderer::RenderTranslucency(FRHICommandList& RHICmdList, const FViewInfo& View)
-{	
-	const bool bShouldRenderTranslucency = ShouldRenderTranslucency(StandardTranslucencyPass) && ViewFamily.EngineShowFlags.Translucency;
+void FMobileSceneRenderer::RenderTranslucency(FRHICommandList& RHICmdList, const FViewInfo& View, TArrayView<FViewInfo> FamilyViewInfos, ETranslucencyPass::Type InStandardTranslucencyPass, EMeshPass::Type InStandardTranslucencyMeshPass, const FInstanceCullingDrawParams* InTranslucencyInstanceCullingDrawParams)
+{
+	const FSceneViewFamily& ViewFamily = *View.Family;
+
+	const bool bShouldRenderTranslucency = ShouldRenderTranslucency(InStandardTranslucencyPass, FamilyViewInfos) && ViewFamily.EngineShowFlags.Translucency;
 	if (bShouldRenderTranslucency)
 	{
 		CSV_SCOPED_TIMING_STAT_EXCLUSIVE(RenderTranslucency);
@@ -16,6 +18,6 @@ void FMobileSceneRenderer::RenderTranslucency(FRHICommandList& RHICmdList, const
 		SCOPED_GPU_STAT(RHICmdList, Translucency);
 
 		RHICmdList.SetViewport(View.ViewRect.Min.X, View.ViewRect.Min.Y, 0.0f, View.ViewRect.Max.X, View.ViewRect.Max.Y, 1.0f);
-		View.ParallelMeshDrawCommandPasses[StandardTranslucencyMeshPass].Draw(RHICmdList, &TranslucencyInstanceCullingDrawParams);
+		View.ParallelMeshDrawCommandPasses[InStandardTranslucencyMeshPass].Draw(RHICmdList, InTranslucencyInstanceCullingDrawParams);
 	}
 }
diff --git a/Engine/Source/Runtime/Renderer/Private/PostProcess/TemporalSuperResolution.cpp b/Engine/Source/Runtime/Renderer/Private/PostProcess/TemporalSuperResolution.cpp
--- a/Engine/Source/Runtime/Renderer/Private/PostProcess/TemporalSuperResolution.cpp
+++ b/Engine/Source/Runtime/Renderer/Private/PostProcess/TemporalSuperResolution.cpp
@@ -1345,7 +1345,8 @@ bool IsTSRLensDistortionEnabled(EShaderPlatform ShaderPlatform)
 
 bool NeedTSRMoireLuma(const FViewInfo& View)
 {
-	return GetMainTAAPassConfig(View) == EMainTAAPassConfig::TSR;
+	// Need to also check PostProcessing flag, as scene captures may run with temporal AA jitter matching the main view, but post processing disabled.
+	return GetMainTAAPassConfig(View) == EMainTAAPassConfig::TSR && View.Family->EngineShowFlags.PostProcessing;
 }
 
 bool IsVisualizeTSREnabled(const FViewInfo& View)
diff --git a/Engine/Source/Runtime/Renderer/Private/RayTracing/RayTracingTranslucency.cpp b/Engine/Source/Runtime/Renderer/Private/RayTracing/RayTracingTranslucency.cpp
--- a/Engine/Source/Runtime/Renderer/Private/RayTracing/RayTracingTranslucency.cpp
+++ b/Engine/Source/Runtime/Renderer/Private/RayTracing/RayTracingTranslucency.cpp
@@ -144,12 +144,12 @@ bool ShouldRenderRayTracingTranslucency(const FViewInfo& View)
 
 void FDeferredShadingSceneRenderer::RenderRayTracingTranslucency(FRDGBuilder& GraphBuilder, FRDGTextureMSAA SceneColorTexture)
 {
-	if (   !ShouldRenderTranslucency(ETranslucencyPass::TPT_TranslucencyStandard)
-		&& !ShouldRenderTranslucency(ETranslucencyPass::TPT_TranslucencyStandardModulate)
-		&& !ShouldRenderTranslucency(ETranslucencyPass::TPT_TranslucencyAfterDOF)
-		&& !ShouldRenderTranslucency(ETranslucencyPass::TPT_TranslucencyAfterDOFModulate)
-		&& !ShouldRenderTranslucency(ETranslucencyPass::TPT_TranslucencyAfterMotionBlur)
-		&& !ShouldRenderTranslucency(ETranslucencyPass::TPT_AllTranslucency)
+	if (   !ShouldRenderTranslucency(ETranslucencyPass::TPT_TranslucencyStandard, Views)
+		&& !ShouldRenderTranslucency(ETranslucencyPass::TPT_TranslucencyStandardModulate, Views)
+		&& !ShouldRenderTranslucency(ETranslucencyPass::TPT_TranslucencyAfterDOF, Views)
+		&& !ShouldRenderTranslucency(ETranslucencyPass::TPT_TranslucencyAfterDOFModulate, Views)
+		&& !ShouldRenderTranslucency(ETranslucencyPass::TPT_TranslucencyAfterMotionBlur, Views)
+		&& !ShouldRenderTranslucency(ETranslucencyPass::TPT_AllTranslucency, Views)
 		)
 	{
 		return; // Early exit if nothing needs to be done.
diff --git a/Engine/Source/Runtime/Renderer/Private/SceneCaptureRendering.cpp b/Engine/Source/Runtime/Renderer/Private/SceneCaptureRendering.cpp
--- a/Engine/Source/Runtime/Renderer/Private/SceneCaptureRendering.cpp
+++ b/Engine/Source/Runtime/Renderer/Private/SceneCaptureRendering.cpp
@@ -999,6 +999,16 @@ static FSceneRenderer* CreateSceneRendererForSceneCapture(
 			ViewFamily, /* GlobalResolutionFraction = */ 1.0f));
 	}
 
+	if (SceneCaptureComponent->IsUnlit())
+	{
+		bool bAllowAtmosphere =
+			SceneCaptureComponent->CaptureSource == SCS_SceneColorHDR &&
+			SceneCaptureComponent->CaptureSource == SCS_SceneColorHDRNoAlpha &&
+			SceneCaptureComponent->CaptureSource == SCS_SceneColorSceneDepth;
+
+		ViewFamily.EngineShowFlags.DisableFeaturesForUnlit(bAllowAtmosphere);
+	}
+
 	if (IsValid(SceneCaptureComponent2D))
 	{
 		// Scene capture 2D only support a single view
@@ -1053,6 +1063,22 @@ static FSceneRenderer* CreateSceneRendererForSceneCapture(
 				--Index;
 			}
 		}
+
+		// For discoverability and backward compatibility, the unlit viewmode option is its own enum, rather than going through the
+		// UnlitViewmode show flag.  The debug feature only works on non-shipping PC builds, and so going forward, the desired
+		// default behavior for scene color Scene Captures running as Custom Render Passes is to disable the debug feature, for consistent
+		// results across builds.  However, we can't change the default behavior for existing Scene Captures, as there may be licensees
+		// always running PC development builds, using the debug behavior by design.
+		//
+		// The enum allows the creation of three states -- disabled across the board, enabled for captures, and enabled for both captures and
+		// custom render passes, with the default being enabled only for captures.  This accomplishes both goals, and also allows licensees
+		// using the debug feature by design to gain performance by switching to a CRP, and opting in to the debug feature there as well.
+		//
+		// Discoverability comes from the fact that setting the "Render In Main Renderer" flag (switching to a CRP) will toggle the debug
+		// behavior, causing a visual change.  The "Unlit Viewmode" setting is immediately next to the flag the user just toggled,
+		// and gives a clue as to what is happening, and allows them to choose a solution -- also enable the debug feature for the CRP, or
+		// change their content to not assume the presence of the debug feature, and disable it across the board, depending on their goals.
+		ViewFamily.EngineShowFlags.SetUnlitViewmode(SceneCaptureComponent2D->UnlitViewmode != ESceneCaptureUnlitViewmode::Disabled);
 	}
 
 	// Call SetupViewFamily & SetupView on scene view extensions before renderer creation
@@ -1085,6 +1111,8 @@ public:
 #if !UE_BUILD_SHIPPING
 		CaptureComponent->GetOuter()->GetName(UserData->CaptureActorName);
 #endif
+		bSceneColorWithTranslucent = CaptureComponent->ShowFlags.Translucency;
+		bSceneColorIsUserSceneTexture = !CaptureComponent->UserSceneTextureSceneColor.IsNone();
 
 		SetUserData(TUniquePtr<FSceneCaptureCustomRenderPassUserData>(UserData));
 	}
@@ -1265,11 +1293,7 @@ void FScene::UpdateSceneCaptureContents(USceneCaptureComponent2D* CaptureCompone
 		}
 
 		// As optimization for depth capture modes, render scene capture as additional render passes inside the main renderer.
-		if (GSceneCaptureAllowRenderInMainRenderer && 
-			CaptureComponent->bRenderInMainRenderer && 
-			(CaptureComponent->CaptureSource == ESceneCaptureSource::SCS_SceneDepth || CaptureComponent->CaptureSource == ESceneCaptureSource::SCS_DeviceDepth ||
-			 CaptureComponent->CaptureSource == ESceneCaptureSource::SCS_BaseColor || CaptureComponent->CaptureSource == ESceneCaptureSource::SCS_Normal)
-			)
+		if (GSceneCaptureAllowRenderInMainRenderer && CaptureComponent->ShouldRenderInMainRenderer())
 		{
 			FCustomRenderPassRendererInput PassInput;
 			PassInput.ViewLocation = ViewLocation;
@@ -1286,6 +1310,21 @@ void FScene::UpdateSceneCaptureContents(USceneCaptureComponent2D* CaptureCompone
 
 			switch (CaptureComponent->CaptureSource)
 			{
+			case ESceneCaptureSource::SCS_SceneColorHDR:
+				RenderMode = FCustomRenderPassBase::ERenderMode::DepthAndBasePass;
+				RenderOutput = FCustomRenderPassBase::ERenderOutput::SceneColorAndAlpha;
+				DebugName = TEXT("SceneCapturePass_SceneColorAndAlpha");
+				break;
+			case ESceneCaptureSource::SCS_SceneColorHDRNoAlpha:
+				RenderMode = FCustomRenderPassBase::ERenderMode::DepthAndBasePass;
+				RenderOutput = FCustomRenderPassBase::ERenderOutput::SceneColorNoAlpha;
+				DebugName = TEXT("SceneCapturePass_SceneColorNoAlpha");
+				break;
+			case ESceneCaptureSource::SCS_SceneColorSceneDepth:
+				RenderMode = FCustomRenderPassBase::ERenderMode::DepthAndBasePass;
+				RenderOutput = FCustomRenderPassBase::ERenderOutput::SceneColorAndDepth;
+				DebugName = TEXT("SceneCapturePass_SceneColorAndDepth");
+				break;
 			case ESceneCaptureSource::SCS_SceneDepth:
 				if (bHasUserSceneTextureOutput)
 				{
@@ -1330,6 +1369,12 @@ void FScene::UpdateSceneCaptureContents(USceneCaptureComponent2D* CaptureCompone
 			GetShowOnlyAndHiddenComponents(CaptureComponent, PassInput.HiddenPrimitives, PassInput.ShowOnlyPrimitives);
 
 			PassInput.EngineShowFlags = CaptureComponent->ShowFlags;
+			PassInput.EngineShowFlags.DisableFeaturesForUnlit();
+
+			if (CaptureComponent->UnlitViewmode == ESceneCaptureUnlitViewmode::CaptureOrCustomRenderPass)
+			{
+				PassInput.EngineShowFlags.SetUnlitViewmode(true);
+			}
 
 			if (CaptureComponent->PostProcessBlendWeight > 0.0f && CaptureComponent->PostProcessSettings.bOverride_UserFlags)
 			{
diff --git a/Engine/Source/Runtime/Renderer/Private/SceneRendering.cpp b/Engine/Source/Runtime/Renderer/Private/SceneRendering.cpp
--- a/Engine/Source/Runtime/Renderer/Private/SceneRendering.cpp
+++ b/Engine/Source/Runtime/Renderer/Private/SceneRendering.cpp
@@ -123,13 +123,6 @@ static FAutoConsoleVariableRef CVarSceneCaptureDepthPrepassOptimization(
 	TEXT("Whether to apply optimized render path when capturing depth prepass for scene capture 2D. Experimental!\n"),
 	ECVF_RenderThreadSafe | ECVF_Scalability);
 
-int32 GSceneCaptureBasePassOptimization = 0;
-static FAutoConsoleVariableRef CVarSceneCaptureBasePassOptimization(
-	TEXT("r.SceneCapture.BasePassOptimization"),
-	GSceneCaptureBasePassOptimization,
-	TEXT("Whether to apply optimized render path when capturing base pass or normals for scene capture 2D. Experimental!\n"),
-	ECVF_RenderThreadSafe | ECVF_Scalability);
-
 static int32 GAsyncCreateLightPrimitiveInteractions = 1;
 static FAutoConsoleVariableRef CVarAsyncCreateLightPrimitiveInteractions(
 	TEXT("r.AsyncCreateLightPrimitiveInteractions"),
@@ -2992,14 +2985,16 @@ FSceneRenderer::FSceneRenderer(const FSceneViewFamily* InViewFamily, FHitProxyCo
 		// relevant for custom render passes anyway.
 		FSceneViewFamily::ConstructionValues FamilyCVS(ViewFamily.RenderTarget, Scene, PassInput.bUseMainViewFamilyShowFlags ? ViewFamily.EngineShowFlags : PassInput.EngineShowFlags);
 
-		// Disable sky rendering, which is gated by the Atmosphere flag (unnecessary perf cost during base pass rendering)
-		FamilyCVS.EngineShowFlags.Atmosphere = false;
+		if (PassInput.bUseMainViewFamilyShowFlags)
+		{
+			// PassInput.EngineShowFlags will already have had this called at construction, but show flags copied from the ViewFamily will not
+			FamilyCVS.EngineShowFlags.DisableFeaturesForUnlit();
+		}
 
-		// Force enable the Lighting flag.  Custom Render Passes don't run lighting regardless of this flag, but disabling it causes the function
-		// "IsRichView()" to return true, forcing certain render proxy classes to go through a slow dynamic rendering code path used for debug
-		// visualization, which isn't relevant for Custom Render Passes.  Avoids a massive silent perf loss if someone has modified the flag
-		// (which is user facing in the Scene Capture properties).
-		FamilyCVS.EngineShowFlags.Lighting = true;
+		// Conditionally enable translucency.  Custom render passes have their own translucency flag, with the assumption that by default
+		// they shouldn't have translucency.  Also, depending on the output of the CRP, the translucent pass may be writing to the scene
+		// color where it isn't used, and translucency should be disabled as an unnecessary perf cost.
+		FamilyCVS.EngineShowFlags.SetTranslucency(CustomRenderPass->IsTranslucentIncluded());
 
 		FCustomRenderPassInfo& CustomRenderPassInfo = CustomRenderPassInfos.Emplace_GetRef(FamilyCVS, ViewFamily);
 		CustomRenderPassInfo.CustomRenderPass = CustomRenderPass;
@@ -3091,7 +3086,9 @@ FSceneRenderer::FSceneRenderer(const FSceneViewFamily* InViewFamily, FHitProxyCo
 		{
 			for (FViewInfo* View : AllViews)
 			{
-				if (View->AntiAliasingMethod == AAM_TSR)
+				// Need to also check PostProcessing flag, as scene captures may run with temporal AA jitter matching the main view, but post processing disabled.
+				// Without this, translucency doesn't show up, because the renderer assumes post processing will composite in the translucency.
+				if (View->AntiAliasingMethod == AAM_TSR && View->Family->EngineShowFlags.PostProcessing)
 				{
 					ViewFamily.bAllowTranslucencyAfterDOF = true;
 					break;
@@ -3248,18 +3245,6 @@ FSceneRenderer::ERendererOutput FSceneRenderer::GetRendererOutput() const
 			return ERendererOutput::DepthPrepassOnly;
 		}
 	}
-	if (ViewFamily.SceneCaptureSource == ESceneCaptureSource::SCS_BaseColor || ViewFamily.SceneCaptureSource == ESceneCaptureSource::SCS_Normal)
-	{
-		// Forward shading in the deferred renderer requires shadows to run (skipped by BasePass optimization)
-		// HasRayTracedOverlay requires ray tracing to run -- appears to be used by path tracer, and ray tracing debug modes (ditto)
-		//
-		// TODO:  Could there be additional features that should disable the BasePass optimization?  We can add more cases if we run into bugs.
-		//        The above were just obvious from looking at the code.  Also, should some features disable DepthPrepassOnly above?
-		if (GSceneCaptureBasePassOptimization && !IsForwardShadingEnabled(Scene->GetShaderPlatform()) && !HasRayTracedOverlay(ViewFamily))
-		{
-			return ERendererOutput::BasePass;
-		}
-	}
 	return ERendererOutput::FinalSceneColor;
 }
 
@@ -4020,7 +4005,7 @@ IVisibilityTaskData* FSceneRenderer::OnRenderBegin(FRDGBuilder& GraphBuilder)
 			}
 		}
 	
-		// Lighting is skipped when running ERendererOutput::DepthPrepassOnly or ERendererOutput::BasePass
+		// Lighting is skipped when running ERendererOutput::DepthPrepassOnly
 		if (GetRendererOutput() == ERendererOutput::FinalSceneColor)
 		{
 			LightFunctionAtlas::OnRenderBegin(LightFunctionAtlas, *Scene, Views, ViewFamily);
diff --git a/Engine/Source/Runtime/Renderer/Private/SceneRendering.h b/Engine/Source/Runtime/Renderer/Private/SceneRendering.h
--- a/Engine/Source/Runtime/Renderer/Private/SceneRendering.h
+++ b/Engine/Source/Runtime/Renderer/Private/SceneRendering.h
@@ -1901,6 +1901,12 @@ struct FSortedShadowMaps
 
 		return MemorySize;
 	}
+
+	bool IsEmpty() const
+	{
+		return ShadowMapAtlases.IsEmpty() && ShadowMapCubemaps.IsEmpty() && PreshadowCache.Shadows.IsEmpty() && TranslucencyShadowMapAtlases.IsEmpty() &&
+			VirtualShadowMapShadows.IsEmpty() && VirtualShadowMapClipmaps.IsEmpty() && CompleteShadowMapAtlases.IsEmpty();
+	}
 };
 
 struct FOcclusionSubmittedFenceState
@@ -2295,7 +2301,6 @@ protected:
 	enum class ERendererOutput
 	{
 		DepthPrepassOnly,	// Only render depth prepass and its related code paths
-		BasePass,			// Render base pass (deferred renderer only, for mobile renderer equivalent to FinalSceneColor)
 		FinalSceneColor		// Render the whole pipeline
 	};
 
@@ -2442,7 +2447,7 @@ protected:
 
 	/** Returns whether there are translucent primitives to be rendered. */
 	bool ShouldRenderTranslucency() const;
-	bool ShouldRenderTranslucency(ETranslucencyPass::Type TranslucencyPass) const;
+	static bool ShouldRenderTranslucency(ETranslucencyPass::Type TranslucencyPass, TArrayView<FViewInfo> InViews);
 
 	/** Called at the begin / finish of the scene render. */
 	IVisibilityTaskData* OnRenderBegin(FRDGBuilder& GraphBuilder);
@@ -2693,7 +2698,7 @@ protected:
 
 	void RenderMobileLocalLightsBuffer(FRDGBuilder& GraphBuilder, FSceneTextures& SceneTextures, const FSortedLightSetSceneInfo& SortedLights);
 
-	void RenderCustomRenderPassBasePass(FRDGBuilder& GraphBuilder, TArrayView<FViewInfo> InViews, FRDGTextureRef ViewFamilyTexture, FSceneTextures& SceneTextures);
+	void RenderCustomRenderPassBasePass(FRDGBuilder& GraphBuilder, TArrayView<FViewInfo> InViews, FRDGTextureRef ViewFamilyTexture, FSceneTextures& SceneTextures, bool bIncludeTranslucent);
 
 	/** Renders the opaque base pass for mobile. */
 	void RenderMobileBasePass(FRHICommandList& RHICmdList, const FViewInfo& View, const FInstanceCullingDrawParams* InstanceCullingDrawParams);
@@ -2735,7 +2740,7 @@ protected:
 	void RenderFog(FRHICommandList& RHICmdList, const FViewInfo& View);
 
 	/** Renders the base pass for translucency. */
-	void RenderTranslucency(FRHICommandList& RHICmdList, const FViewInfo& View);
+	static void RenderTranslucency(FRHICommandList& RHICmdList, const FViewInfo& View, TArrayView<FViewInfo> FamilyViewInfos, ETranslucencyPass::Type InStandardTranslucencyPass, EMeshPass::Type InStandardTranslucencyMeshPass, const FInstanceCullingDrawParams* InTranslucencyInstanceCullingDrawParams);
 
 	/** On chip pre-tonemap before scene color MSAA resolve (iOS only) */
 	void PreTonemapMSAA(FRHICommandList& RHICmdList, const FMinimalSceneTextures& SceneTextures);
diff --git a/Engine/Source/Runtime/Renderer/Private/ShadowDepthRendering.cpp b/Engine/Source/Runtime/Renderer/Private/ShadowDepthRendering.cpp
--- a/Engine/Source/Runtime/Renderer/Private/ShadowDepthRendering.cpp
+++ b/Engine/Source/Runtime/Renderer/Private/ShadowDepthRendering.cpp
@@ -1678,6 +1678,16 @@ void FSceneRenderer::RenderShadowDepthMaps(FRDGBuilder& GraphBuilder, FDynamicSh
 {
 	ensureMsgf(!bShadowDepthRenderCompleted, TEXT("RenderShadowDepthMaps called twice in the same frame"));
 
+	// Exit early if dynamic shadows are disabled (avoids the unnecessary BeginDeferredCulling call, and prevents RenderShadowDepthMapAtlases from overwriting Scene->PrevAtlasHZBs)
+	if (!ViewFamily.EngineShowFlags.DynamicShadows)
+	{
+		// Earlier code should not have added shadows to be processed, or created dynamic shadow related tasks or queries, if DynamicShadows are disabled
+		check(SortedShadowsForShadowDepthPass.IsEmpty() && DynamicShadowsTaskData == nullptr && (!ShadowSceneRenderer.IsValid() || !ShadowSceneRenderer->HasNaniteVisibilityQuery()));
+
+		bShadowDepthRenderCompleted = true;
+		return;
+	}
+
 	RDG_CSV_STAT_EXCLUSIVE_SCOPE(GraphBuilder, RenderShadows);
 
 	TRACE_CPUPROFILER_EVENT_SCOPE(FSceneRenderer::RenderShadowDepthMaps);
diff --git a/Engine/Source/Runtime/Renderer/Private/Shadows/ShadowSceneRenderer.h b/Engine/Source/Runtime/Renderer/Private/Shadows/ShadowSceneRenderer.h
--- a/Engine/Source/Runtime/Renderer/Private/Shadows/ShadowSceneRenderer.h
+++ b/Engine/Source/Runtime/Renderer/Private/Shadows/ShadowSceneRenderer.h
@@ -111,6 +111,11 @@ public:
 		return LocalLights.Num() > 0;
 	}
 
+	bool HasNaniteVisibilityQuery() const
+	{
+		return NaniteVisibilityQuery != nullptr;
+	}
+
 	UE::Renderer::Private::IShadowInvalidatingInstances *GetInvalidatingInstancesInterface(const FSceneView *SceneView);
 
 private:
diff --git a/Engine/Source/Runtime/Renderer/Private/TranslucentRendering.cpp b/Engine/Source/Runtime/Renderer/Private/TranslucentRendering.cpp
--- a/Engine/Source/Runtime/Renderer/Private/TranslucentRendering.cpp
+++ b/Engine/Source/Runtime/Renderer/Private/TranslucentRendering.cpp
@@ -843,14 +843,14 @@ bool FSceneRenderer::ShouldRenderTranslucency() const
 		&& !ViewFamily.UseDebugViewPS();
 }
 
-bool FSceneRenderer::ShouldRenderTranslucency(ETranslucencyPass::Type TranslucencyPass) const
+bool FSceneRenderer::ShouldRenderTranslucency(ETranslucencyPass::Type TranslucencyPass, TArrayView<FViewInfo> InViews)
 {
 	extern int32 GLightShaftRenderAfterDOF;
 
 	// Change this condition to control where simple elements should be rendered.
 	if (IsMainTranslucencyPass(TranslucencyPass))
 	{
-		for (const FViewInfo& View : Views)
+		for (const FViewInfo& View : InViews)
 		{
 			if (View.bHasTranslucentViewMeshElements || View.SimpleElementCollector.BatchedElements.HasPrimsToDraw())
 			{
@@ -865,7 +865,7 @@ bool FSceneRenderer::ShouldRenderTranslucency(ETranslucencyPass::Type Translucen
 		return true;
 	}
 
-	for (const FViewInfo& View : Views)
+	for (const FViewInfo& View : InViews)
 	{
 		if (View.TranslucentPrimCount.Num(TranslucencyPass) > 0)
 		{
@@ -896,7 +896,7 @@ void SetupPostMotionBlurTranslucencyViewParameters(const FViewInfo& View, FViewU
 
 const FRDGTextureDesc GetPostDOFTranslucentTextureDesc(
 	ETranslucencyPass::Type TranslucencyPass,
-	FSeparateTranslucencyDimensions& SeparateTranslucencyDimensions,
+	const FSeparateTranslucencyDimensions& SeparateTranslucencyDimensions,
 	bool bIsModulate,
 	EShaderPlatform ShaderPlatform)
 {
@@ -920,7 +920,7 @@ const FRDGTextureDesc GetPostDOFTranslucentTextureDesc(
 FRDGTextureMSAA CreatePostDOFTranslucentTexture(
 	FRDGBuilder& GraphBuilder,
 	ETranslucencyPass::Type TranslucencyPass,
-	FSeparateTranslucencyDimensions& SeparateTranslucencyDimensions,
+	const FSeparateTranslucencyDimensions& SeparateTranslucencyDimensions,
 	bool bIsModulate,
 	EShaderPlatform ShaderPlatform)
 {
@@ -1318,7 +1318,12 @@ static void RenderTranslucencyViewInner(
 	PassParameters->View = GetSeparateTranslucencyViewParameters(View, Viewport.Extent, ViewportScale, TranslucencyPass);
 	PassParameters->ReflectionCapture = View.ReflectionCaptureUniformBuffer;
 	PassParameters->BasePass = BasePassParameters;
-	PassParameters->VirtualShadowMapSamplingParameters = SceneRenderer.VirtualShadowMapArray.GetSamplingParameters(GraphBuilder, ViewIndex);
+
+	// Custom render passes run early in the frame before VSMs have been processed.  The VSM uniforms will be dummy at that point in the frame, so it doesn't matter which
+	// ViewIndex in the VirtualShadowMapArray is referenced, but it needs to be an index that's not out of range of the array.  To handle that case, pass in an index of
+	// zero, so it arbitrarily uses the first element.
+	PassParameters->VirtualShadowMapSamplingParameters = SceneRenderer.VirtualShadowMapArray.GetSamplingParameters(GraphBuilder, View.CustomRenderPass == nullptr ? ViewIndex : 0);
+
 	PassParameters->RenderTargets[0] = FRenderTargetBinding(SceneColorTexture.Target, ERenderTargetLoadAction::ELoad);
 	if (TranslucencyPass != ETranslucencyPass::TPT_TranslucencyAfterMotionBlur)
 	{
@@ -1454,22 +1459,28 @@ static void RenderTranslucencyViewInner(
 }
 
 void FDeferredShadingSceneRenderer::RenderTranslucencyInner(
+	FDeferredShadingSceneRenderer& Renderer,
 	FRDGBuilder& GraphBuilder,
 	const FMinimalSceneTextures& SceneTextures,
 	const FTranslucencyLightingVolumeTextures& TranslucentLightingVolumeTextures,
 	FTranslucencyPassResourcesMap* OutTranslucencyResourceMap,
 	FRDGTextureMSAA SharedDepthTexture,
+	TArray<FViewInfo>& InViews,
 	ETranslucencyView ViewsToRender,
+	const FSeparateTranslucencyDimensions& SeparateTranslucencyDimensions,
 	FRDGTextureRef SceneColorCopyTexture,
 	ETranslucencyPass::Type TranslucencyPass,
 	FInstanceCullingManager& InstanceCullingManager,
 	bool bStandardTranslucentCanRenderSeparate)
 {
-	if (!ShouldRenderTranslucency(TranslucencyPass))
+	if (!ShouldRenderTranslucency(TranslucencyPass, InViews))
 	{
 		return;
 	}
 
+	const FScene* Scene = Renderer.Scene;
+	const FSceneViewFamily& ViewFamily = *InViews[0].Family;
+
 	const bool bIsModulate = TranslucencyPass == ETranslucencyPass::TPT_TranslucencyAfterDOFModulate || TranslucencyPass == ETranslucencyPass::TPT_TranslucencyStandardModulate;
 	const bool bDepthTest = TranslucencyPass != ETranslucencyPass::TPT_TranslucencyAfterMotionBlur;
 	const bool bRenderInParallel = IsParallelTranslucencyEnabled();
@@ -1485,11 +1496,11 @@ void FDeferredShadingSceneRenderer::RenderTranslucencyInner(
 	EnumRemoveFlags(SceneTextureSetupMode, ESceneTextureSetupMode::SceneColor);
 
 	// Create resources shared by each view (each view data is tiled into each of the render target resources)
-	FRDGTextureMSAA SharedColorTexture = CreatePostDOFTranslucentTexture(GraphBuilder, TranslucencyPass, SeparateTranslucencyDimensions, bIsModulate, ShaderPlatform);
+	FRDGTextureMSAA SharedColorTexture = CreatePostDOFTranslucentTexture(GraphBuilder, TranslucencyPass, SeparateTranslucencyDimensions, bIsModulate, Scene->GetShaderPlatform());
 
-	for (int32 ViewIndex = 0, NumProcessedViews = 0; ViewIndex < Views.Num(); ++ViewIndex)
+	for (int32 ViewIndex = 0, NumProcessedViews = 0; ViewIndex < InViews.Num(); ++ViewIndex)
 	{
-		FViewInfo& View = Views[ViewIndex];
+		FViewInfo& View = InViews[ViewIndex];
 		const ETranslucencyView TranslucencyView = GetTranslucencyView(View);
 
 		if (!EnumHasAnyFlags(TranslucencyView, ViewsToRender))
@@ -1504,13 +1515,13 @@ void FDeferredShadingSceneRenderer::RenderTranslucencyInner(
 		{
 
 			RDG_GPU_MASK_SCOPE(GraphBuilder, View.GPUMask);
-			RDG_EVENT_SCOPE_CONDITIONAL(GraphBuilder, Views.Num() > 1, "View%d", ViewIndex);
+			RDG_EVENT_SCOPE_CONDITIONAL(GraphBuilder, InViews.Num() > 1, "View%d", ViewIndex);
 
 			FIntRect ScaledViewRect = GetScaledRect(View.ViewRect, SeparateTranslucencyDimensions.Scale);
 
 			const FScreenPassTextureViewport SeparateTranslucencyViewport = SeparateTranslucencyDimensions.GetInstancedStereoViewport(View);
 			const bool bCompositeBackToSceneColor = (IsMainTranslucencyPass(TranslucencyPass) && !bIsStandardSeparatedTranslucency);
-			const bool bLumenGIEnabled = GetViewPipelineState(View).DiffuseIndirectMethod == EDiffuseIndirectMethod::Lumen;
+			const bool bLumenGIEnabled = Renderer.GetViewPipelineState(View).DiffuseIndirectMethod == EDiffuseIndirectMethod::Lumen;
 
 			/** Separate translucency color is either composited immediately or later during post processing. If done immediately, it's because the view doesn't support
 			 *  compositing (e.g. we're rendering an underwater view) or because we're downsampling the main translucency pass. In this case, we use a local set of
@@ -1533,7 +1544,7 @@ void FDeferredShadingSceneRenderer::RenderTranslucencyInner(
 
 			RenderTranslucencyViewInner(
 				GraphBuilder,
-				*this,
+				Renderer,
 				View, ViewIndex,
 				SeparateTranslucencyViewport,
 				SeparateTranslucencyDimensions.Scale,
@@ -1626,11 +1637,11 @@ void FDeferredShadingSceneRenderer::RenderTranslucencyInner(
 			// This three step copy compose: bv = HoldoutOpacity
 
 			RDG_GPU_MASK_SCOPE(GraphBuilder, View.GPUMask);
-			RDG_EVENT_SCOPE_CONDITIONAL(GraphBuilder, Views.Num() > 1, "View%d", ViewIndex);
+			RDG_EVENT_SCOPE_CONDITIONAL(GraphBuilder, InViews.Num() > 1, "View%d", ViewIndex);
 			
 			const FScreenPassTextureViewport SeparateTranslucencyViewport = SeparateTranslucencyDimensions.GetInstancedStereoViewport(View);
 			const bool bCompositeBackToSceneColor = true;
-			const bool bLumenGIEnabled = GetViewPipelineState(View).DiffuseIndirectMethod == EDiffuseIndirectMethod::Lumen;
+			const bool bLumenGIEnabled = Renderer.GetViewPipelineState(View).DiffuseIndirectMethod == EDiffuseIndirectMethod::Lumen;
 
 			/** Separate translucency color is either composited immediately or later during post processing. If done immediately, it's because the view doesn't support
 			 *  compositing (e.g. we're rendering an underwater view) or because we're downsampling the main translucency pass. In this case, we use a local set of
@@ -1664,7 +1675,7 @@ void FDeferredShadingSceneRenderer::RenderTranslucencyInner(
 
 			RenderTranslucencyViewInner(
 				GraphBuilder,
-				*this,
+				Renderer,
 				View, ViewIndex,
 				SeparateTranslucencyViewport,
 				ViewportScale,
@@ -1727,7 +1738,7 @@ void FDeferredShadingSceneRenderer::RenderTranslucencyInner(
 			// When rendering translucent meshes under water, we skip modulate passes which are only required when compositing separate translucency passes from render target.
 			const bool bSkipPass = bViewIsUnderWater && bIsModulate; 
 			RDG_GPU_MASK_SCOPE(GraphBuilder, View.GPUMask);
-			RDG_EVENT_SCOPE_CONDITIONAL(GraphBuilder, Views.Num() > 1 && !bSkipPass, "View%d", ViewIndex);
+			RDG_EVENT_SCOPE_CONDITIONAL(GraphBuilder, InViews.Num() > 1 && !bSkipPass, "View%d", ViewIndex);
 			if (bSkipPass)
 			{
 				return;
@@ -1737,13 +1748,13 @@ void FDeferredShadingSceneRenderer::RenderTranslucencyInner(
 			const FScreenPassTextureViewport Viewport(SceneTextures.Color.Target, View.ViewRect);
 			const float ViewportScale = 1.0f;
 			const bool bResolveColorTexture = false;
-			const bool bLumenGIEnabled = GetViewPipelineState(View).DiffuseIndirectMethod == EDiffuseIndirectMethod::Lumen;
+			const bool bLumenGIEnabled = Renderer.GetViewPipelineState(View).DiffuseIndirectMethod == EDiffuseIndirectMethod::Lumen;
 
 			FOITData OITData = OIT::CreateOITData(GraphBuilder, View, OITPass_RegularTranslucency);
 
 			RenderTranslucencyViewInner(
 				GraphBuilder,
-				*this,
+				Renderer,
 				View, ViewIndex,
 				Viewport,
 				ViewportScale,
@@ -1770,11 +1781,14 @@ static bool ShouldRenderTranslucentView(const FViewInfo& View, ETranslucencyView
 };
 
 void FDeferredShadingSceneRenderer::RenderTranslucency(
+	FDeferredShadingSceneRenderer& Renderer,
 	FRDGBuilder& GraphBuilder,
 	const FSceneTextures& SceneTextures,
 	const FTranslucencyLightingVolumeTextures& TranslucentLightingVolumeTextures,
 	FTranslucencyPassResourcesMap* OutTranslucencyResourceMap,
+	TArray<FViewInfo>& InViews,
 	ETranslucencyView ViewsToRender,
+	const FSeparateTranslucencyDimensions& SeparateTranslucencyDimensions,
 	FInstanceCullingManager& InstanceCullingManager,
 	bool bStandardTranslucentCanRenderSeparate,
 	FRDGTextureMSAA& OutSharedDepthTexture)
@@ -1784,16 +1798,18 @@ void FDeferredShadingSceneRenderer::RenderTranslucency(
 		return;
 	}
 
+	const FSceneViewFamily& ViewFamily = *InViews[0].Family;
+
 	RDG_EVENT_SCOPE_STAT(GraphBuilder, Translucency, "RenderTranslucency");
 	RDG_GPU_STAT_SCOPE(GraphBuilder, Translucency);
 	DynamicRenderScaling::FRDGScope DynamicTranslucencyResolutionScope(GraphBuilder, GDynamicTranslucencyResolution);
 
 	FRDGTextureRef SceneColorCopyTexture = nullptr;
-	const bool bIsTranslucentHoldoutEnabled = IsPrimitiveAlphaHoldoutEnabledForAnyView(Views);
+	const bool bIsTranslucentHoldoutEnabled = IsPrimitiveAlphaHoldoutEnabledForAnyView(InViews);
 
 	if (EnumHasAnyFlags(ViewsToRender, ETranslucencyView::AboveWater))
 	{
-		SceneColorCopyTexture = AddCopySceneColorPass(GraphBuilder, Views, SceneTextures.Color, /*WithAlpha*/ bIsTranslucentHoldoutEnabled);
+		SceneColorCopyTexture = AddCopySceneColorPass(GraphBuilder, InViews, SceneTextures.Color, /*WithAlpha*/ bIsTranslucentHoldoutEnabled);
 	}
 
 	// Create a shared depth texture at the correct resolution.
@@ -1815,9 +1831,9 @@ void FDeferredShadingSceneRenderer::RenderTranslucency(
 			GFastVRamConfig.SeparateTranslucencyModulate); // TODO: this should be SeparateTranslucency, but is what the code was doing
 
 		// Downscale the depth buffer for each individual view, but shared accross all translucencies.
-		for (int32 ViewIndex = 0; ViewIndex < Views.Num(); ++ViewIndex)
+		for (int32 ViewIndex = 0; ViewIndex < InViews.Num(); ++ViewIndex)
 		{
-			FViewInfo& View = Views[ViewIndex];
+			FViewInfo& View = InViews[ViewIndex];
 			const ETranslucencyView TranslucencyView = GetTranslucencyView(View);
 
 			if (!ShouldRenderTranslucentView(View, TranslucencyView, ViewsToRender))
@@ -1841,29 +1857,29 @@ void FDeferredShadingSceneRenderer::RenderTranslucency(
 
 	if (ViewFamily.AllowTranslucencyAfterDOF())
 	{
-		RenderTranslucencyInner(GraphBuilder, SceneTextures, TranslucentLightingVolumeTextures, OutTranslucencyResourceMap, OutSharedDepthTexture, ViewsToRender, SceneColorCopyTexture, ETranslucencyPass::TPT_TranslucencyStandard, InstanceCullingManager, bStandardTranslucentCanRenderSeparate);
+		RenderTranslucencyInner(Renderer, GraphBuilder, SceneTextures, TranslucentLightingVolumeTextures, OutTranslucencyResourceMap, OutSharedDepthTexture, InViews, ViewsToRender, SeparateTranslucencyDimensions, SceneColorCopyTexture, ETranslucencyPass::TPT_TranslucencyStandard, InstanceCullingManager, bStandardTranslucentCanRenderSeparate);
 		if (ViewFamily.AllowStandardTranslucencySeparated() && bStandardTranslucentCanRenderSeparate)
 		{
-			RenderTranslucencyInner(GraphBuilder, SceneTextures, TranslucentLightingVolumeTextures, OutTranslucencyResourceMap, OutSharedDepthTexture, ViewsToRender, SceneColorCopyTexture, ETranslucencyPass::TPT_TranslucencyStandardModulate, InstanceCullingManager, bStandardTranslucentCanRenderSeparate);
+			RenderTranslucencyInner(Renderer, GraphBuilder, SceneTextures, TranslucentLightingVolumeTextures, OutTranslucencyResourceMap, OutSharedDepthTexture, InViews, ViewsToRender, SeparateTranslucencyDimensions, SceneColorCopyTexture, ETranslucencyPass::TPT_TranslucencyStandardModulate, InstanceCullingManager, bStandardTranslucentCanRenderSeparate);
 		}
 
 		if (GetHairStrandsComposition() == EHairStrandsCompositionType::AfterTranslucentBeforeTranslucentAfterDOF)
 		{
-			RenderHairComposition(GraphBuilder, Views, SceneTextures.Color.Target, SceneTextures.Depth.Target, SceneTextures.Velocity, *OutTranslucencyResourceMap);
+			RenderHairComposition(GraphBuilder, InViews, SceneTextures.Color.Target, SceneTextures.Depth.Target, SceneTextures.Velocity, *OutTranslucencyResourceMap);
 		}
-		RenderTranslucencyInner(GraphBuilder, SceneTextures, TranslucentLightingVolumeTextures, OutTranslucencyResourceMap, OutSharedDepthTexture, ViewsToRender, SceneColorCopyTexture, ETranslucencyPass::TPT_TranslucencyAfterDOF, InstanceCullingManager, bStandardTranslucentCanRenderSeparate);
-		RenderTranslucencyInner(GraphBuilder, SceneTextures, TranslucentLightingVolumeTextures, OutTranslucencyResourceMap, OutSharedDepthTexture, ViewsToRender, SceneColorCopyTexture, ETranslucencyPass::TPT_TranslucencyAfterDOFModulate, InstanceCullingManager, bStandardTranslucentCanRenderSeparate);
-		RenderTranslucencyInner(GraphBuilder, SceneTextures, TranslucentLightingVolumeTextures, OutTranslucencyResourceMap, OutSharedDepthTexture, ViewsToRender, SceneColorCopyTexture, ETranslucencyPass::TPT_TranslucencyAfterMotionBlur, InstanceCullingManager, bStandardTranslucentCanRenderSeparate);
+		RenderTranslucencyInner(Renderer, GraphBuilder, SceneTextures, TranslucentLightingVolumeTextures, OutTranslucencyResourceMap, OutSharedDepthTexture, InViews, ViewsToRender, SeparateTranslucencyDimensions, SceneColorCopyTexture, ETranslucencyPass::TPT_TranslucencyAfterDOF, InstanceCullingManager, bStandardTranslucentCanRenderSeparate);
+		RenderTranslucencyInner(Renderer, GraphBuilder, SceneTextures, TranslucentLightingVolumeTextures, OutTranslucencyResourceMap, OutSharedDepthTexture, InViews, ViewsToRender, SeparateTranslucencyDimensions, SceneColorCopyTexture, ETranslucencyPass::TPT_TranslucencyAfterDOFModulate, InstanceCullingManager, bStandardTranslucentCanRenderSeparate);
+		RenderTranslucencyInner(Renderer, GraphBuilder, SceneTextures, TranslucentLightingVolumeTextures, OutTranslucencyResourceMap, OutSharedDepthTexture, InViews, ViewsToRender, SeparateTranslucencyDimensions, SceneColorCopyTexture, ETranslucencyPass::TPT_TranslucencyAfterMotionBlur, InstanceCullingManager, bStandardTranslucentCanRenderSeparate);
 	}
 	else // Otherwise render translucent primitives in a single bucket.
 	{
-		RenderTranslucencyInner(GraphBuilder, SceneTextures, TranslucentLightingVolumeTextures, OutTranslucencyResourceMap, OutSharedDepthTexture, ViewsToRender, SceneColorCopyTexture, ETranslucencyPass::TPT_AllTranslucency, InstanceCullingManager, bStandardTranslucentCanRenderSeparate);
+		RenderTranslucencyInner(Renderer, GraphBuilder, SceneTextures, TranslucentLightingVolumeTextures, OutTranslucencyResourceMap, OutSharedDepthTexture, InViews, ViewsToRender, SeparateTranslucencyDimensions, SceneColorCopyTexture, ETranslucencyPass::TPT_AllTranslucency, InstanceCullingManager, bStandardTranslucentCanRenderSeparate);
 	}
 
 	if (bIsTranslucentHoldoutEnabled && SceneColorCopyTexture)
 	{
 		// Render the translucent holdout background visibility to the alpha channel of SceneColor.
-		RenderTranslucencyInner(GraphBuilder, SceneTextures, TranslucentLightingVolumeTextures, OutTranslucencyResourceMap, OutSharedDepthTexture, ViewsToRender, SceneColorCopyTexture, ETranslucencyPass::TPT_TranslucencyHoldout, InstanceCullingManager, bStandardTranslucentCanRenderSeparate);
+		RenderTranslucencyInner(Renderer, GraphBuilder, SceneTextures, TranslucentLightingVolumeTextures, OutTranslucencyResourceMap, OutSharedDepthTexture, InViews, ViewsToRender, SeparateTranslucencyDimensions, SceneColorCopyTexture, ETranslucencyPass::TPT_TranslucencyHoldout, InstanceCullingManager, bStandardTranslucentCanRenderSeparate);
 	}
 }
 
diff --git a/Engine/Source/Runtime/Renderer/Private/TranslucentRendering.h b/Engine/Source/Runtime/Renderer/Private/TranslucentRendering.h
--- a/Engine/Source/Runtime/Renderer/Private/TranslucentRendering.h
+++ b/Engine/Source/Runtime/Renderer/Private/TranslucentRendering.h
@@ -47,17 +47,14 @@ ETranslucencyView GetTranslucencyViews(TArrayView<const FViewInfo> Views);
 /** Computes the translucency dimensions. */
 FSeparateTranslucencyDimensions UpdateSeparateTranslucencyDimensions(const FSceneRenderer& SceneRenderer);
 
-/** Returns whether the view family is requesting to render translucency. */
-bool ShouldRenderTranslucency(const FSceneViewFamily& ViewFamily);
-
 /** Check if separate translucency pass is needed for given pass and downsample scale */
 bool IsSeparateTranslucencyEnabled(ETranslucencyPass::Type TranslucencyPass, float DownsampleScale);
 
 /** Shared function to get the post DOF texture pixel format and creation flags */
-const FRDGTextureDesc GetPostDOFTranslucentTextureDesc(ETranslucencyPass::Type TranslucencyPass, FSeparateTranslucencyDimensions& SeparateTranslucencyDimensions, bool bIsModulate, EShaderPlatform ShaderPlatform);
+const FRDGTextureDesc GetPostDOFTranslucentTextureDesc(ETranslucencyPass::Type TranslucencyPass, const FSeparateTranslucencyDimensions& SeparateTranslucencyDimensions, bool bIsModulate, EShaderPlatform ShaderPlatform);
 
 /** Shared function used to create Post DOF translucent textures */
-FRDGTextureMSAA CreatePostDOFTranslucentTexture(FRDGBuilder& GraphBuilder, ETranslucencyPass::Type TranslucencyPass, FSeparateTranslucencyDimensions& SeparateTranslucencyDimensions, bool bIsModulate, EShaderPlatform ShaderPlatform);
+FRDGTextureMSAA CreatePostDOFTranslucentTexture(FRDGBuilder& GraphBuilder, ETranslucencyPass::Type TranslucencyPass, const FSeparateTranslucencyDimensions& SeparateTranslucencyDimensions, bool bIsModulate, EShaderPlatform ShaderPlatform);
 
 
 /** Add a pass to compose separate translucency. */
