Scene Capture:  static analysis warning fix (conditional always false).

#rnx

diff --git a/Engine/Source/Runtime/Renderer/Private/SceneCaptureRendering.cpp b/Engine/Source/Runtime/Renderer/Private/SceneCaptureRendering.cpp
--- a/Engine/Source/Runtime/Renderer/Private/SceneCaptureRendering.cpp
+++ b/Engine/Source/Runtime/Renderer/Private/SceneCaptureRendering.cpp
@@ -1002,8 +1002,8 @@ static FSceneRenderer* CreateSceneRendererForSceneCapture(
 	if (SceneCaptureComponent->IsUnlit())
 	{
 		bool bAllowAtmosphere =
-			SceneCaptureComponent->CaptureSource == SCS_SceneColorHDR &&
-			SceneCaptureComponent->CaptureSource == SCS_SceneColorHDRNoAlpha &&
+			SceneCaptureComponent->CaptureSource == SCS_SceneColorHDR ||
+			SceneCaptureComponent->CaptureSource == SCS_SceneColorHDRNoAlpha ||
 			SceneCaptureComponent->CaptureSource == SCS_SceneColorSceneDepth;
 
 		ViewFamily.EngineShowFlags.DisableFeaturesForUnlit(bAllowAtmosphere);
