HeterogeneousVolumes: Disable indirect dispatch for live-shading pipeline to lower material permutation count

diff --git a/Engine/Source/Runtime/Renderer/Private/HeterogeneousVolumes/HeterogeneousVolumesLiveShadingPipeline.cpp b/Engine/Source/Runtime/Renderer/Private/HeterogeneousVolumes/HeterogeneousVolumesLiveShadingPipeline.cpp
--- a/Engine/Source/Runtime/Renderer/Private/HeterogeneousVolumes/HeterogeneousVolumesLiveShadingPipeline.cpp
+++ b/Engine/Source/Runtime/Renderer/Private/HeterogeneousVolumes/HeterogeneousVolumesLiveShadingPipeline.cpp
@@ -529,10 +529,12 @@ class FRenderSingleScatteringWithLiveShadingCS : public FMeshMaterialShader
 };
 
 typedef FRenderSingleScatteringWithLiveShadingCS<HeterogeneousVolumes::DirectDispatch> FRenderSingleScatteringWithLiveShadingDirectCS;
-typedef FRenderSingleScatteringWithLiveShadingCS<HeterogeneousVolumes::IndirectDispatch> FRenderSingleScatteringWithLiveShadingIndirectCS;
+//typedef FRenderSingleScatteringWithLiveShadingCS<HeterogeneousVolumes::IndirectDispatch> FRenderSingleScatteringWithLiveShadingIndirectCS;
+typedef FRenderSingleScatteringWithLiveShadingCS<HeterogeneousVolumes::DirectDispatch> FRenderSingleScatteringWithLiveShadingIndirectCS;
+
 
 IMPLEMENT_MATERIAL_SHADER_TYPE(template<>, FRenderSingleScatteringWithLiveShadingDirectCS, TEXT("/Engine/Private/HeterogeneousVolumes/HeterogeneousVolumesLiveShadingPipeline.usf"), TEXT("RenderSingleScatteringWithLiveShadingCS"), SF_Compute);
-IMPLEMENT_MATERIAL_SHADER_TYPE(template<>, FRenderSingleScatteringWithLiveShadingIndirectCS, TEXT("/Engine/Private/HeterogeneousVolumes/HeterogeneousVolumesLiveShadingPipeline.usf"), TEXT("RenderSingleScatteringWithLiveShadingIndirectCS"), SF_Compute);
+//IMPLEMENT_MATERIAL_SHADER_TYPE(template<>, FRenderSingleScatteringWithLiveShadingIndirectCS, TEXT("/Engine/Private/HeterogeneousVolumes/HeterogeneousVolumesLiveShadingPipeline.usf"), TEXT("RenderSingleScatteringWithLiveShadingIndirectCS"), SF_Compute);
 
 template<bool bWithLumen, HeterogeneousVolumes::EDispatchMode DispatchMode, typename ComputeShaderType>
 void AddComputePass(
@@ -590,62 +592,6 @@ void AddComputePass(
 	);
 }
 
-template<bool bWithLumen, typename ComputeShaderType>
-void AddComputePassDirect(
-	FRDGBuilder& GraphBuilder,
-	TShaderRef<ComputeShaderType>& ComputeShader,
-	typename ComputeShaderType::FParameters* PassParameters,
-	const FScene* Scene,
-	const FMaterialRenderProxy* MaterialRenderProxy,
-	const FMaterial& Material,
-	const FString& PassName,
-	FIntVector GroupCount
-)
-{
-	FRDGBufferRef IndirectArgsBuffer = GSystemTextures.GetDefaultBuffer(GraphBuilder, 4);
-	uint32 IndirectArgOffset = 0;
-	AddComputePass<bWithLumen, HeterogeneousVolumes::EDispatchMode::DirectDispatch>(
-		GraphBuilder,
-		ComputeShader,
-		PassParameters,
-		Scene,
-		MaterialRenderProxy,
-		Material,
-		PassName,
-		GroupCount,
-		IndirectArgsBuffer,
-		IndirectArgOffset
-	);
-}
-
-template<bool bWithLumen, typename ComputeShaderType>
-void AddComputePassIndirect(
-	FRDGBuilder& GraphBuilder,
-	TShaderRef<ComputeShaderType>& ComputeShader,
-	typename ComputeShaderType::FParameters* PassParameters,
-	const FScene* Scene,
-	const FMaterialRenderProxy* MaterialRenderProxy,
-	const FMaterial& Material,
-	const FString& PassName,
-	FRDGBufferRef IndirectArgsBuffer,
-	uint32 IndirectArgOffset
-)
-{
-	FIntVector GroupCount = FIntVector::ZeroValue;
-	AddComputePass<bWithLumen, HeterogeneousVolumes::EDispatchMode::IndirectDispatch>(
-		GraphBuilder,
-		ComputeShader,
-		PassParameters,
-		Scene,
-		MaterialRenderProxy,
-		Material,
-		PassName,
-		GroupCount,
-		IndirectArgsBuffer,
-		IndirectArgOffset
-	);
-}
-
 static void RenderLightingCacheWithLiveShading(
 	FRDGBuilder& GraphBuilder,
 	// Scene data
@@ -872,7 +818,19 @@ static void RenderLightingCacheWithLiveShading(
 	TShaderRef<FRenderLightingCacheWithLiveShadingCS> ComputeShader = Material.GetShader<FRenderLightingCacheWithLiveShadingCS>(&FLocalVertexFactory::StaticType, PermutationVector, false);
 	if (!ComputeShader.IsNull())
 	{
-		AddComputePassDirect<false>(GraphBuilder, ComputeShader, PassParameters, Scene, MaterialRenderProxy, Material, PassName, GroupCount);
+		FRDGBufferRef IndirectArgsBuffer = GSystemTextures.GetDefaultBuffer(GraphBuilder, 4);
+		AddComputePass<false, HeterogeneousVolumes::EDispatchMode::DirectDispatch>(
+			GraphBuilder,
+			ComputeShader,
+			PassParameters,
+			Scene,
+			MaterialRenderProxy,
+			Material,
+			PassName,
+			GroupCount,
+			IndirectArgsBuffer,
+			0
+		);
 	}
 }
 
@@ -1371,7 +1329,7 @@ static void RenderWithTransmittanceVolumePipeline(
 				LightingCacheTexture
 			);
 		}
-
+#if 0
 		if (HeterogeneousVolumes::ShouldUseScreenTileClassification())
 		{
 			RenderSingleScatteringWithLiveShading<HeterogeneousVolumes::IndirectDispatch>(
@@ -1403,6 +1361,7 @@ static void RenderWithTransmittanceVolumePipeline(
 			);
 		}
 		else
+#endif
 		{
 			RenderSingleScatteringWithLiveShading<HeterogeneousVolumes::DirectDispatch>(
 				GraphBuilder,
@@ -1529,7 +1488,7 @@ static void RenderWithInscatteringVolumePipeline(
 		uint32 LightType = 0;
 		FLightSceneInfo* LightSceneInfo = nullptr;
 		const FVisibleLightInfo* VisibleLightInfo = nullptr;
-
+#if 0
 		if (HeterogeneousVolumes::ShouldUseScreenTileClassification())
 		{
 			RenderSingleScatteringWithLiveShading<HeterogeneousVolumes::IndirectDispatch>(
@@ -1561,6 +1520,7 @@ static void RenderWithInscatteringVolumePipeline(
 			);
 		}
 		else
+#endif
 		{
 			RenderSingleScatteringWithLiveShading<HeterogeneousVolumes::DirectDispatch>(
 				GraphBuilder,
