Fixed post-processing dump-to-pipe/file by forcing the post-process material pass to create a new output render target.

More precisely, this is now required due to the removal of a bPrimeOutputColor condition at CL 36478874.

#rnx

diff --git a/Engine/Source/Runtime/Renderer/Private/PostProcess/PostProcessVisualizeBuffer.cpp b/Engine/Source/Runtime/Renderer/Private/PostProcess/PostProcessVisualizeBuffer.cpp
--- a/Engine/Source/Runtime/Renderer/Private/PostProcess/PostProcessVisualizeBuffer.cpp
+++ b/Engine/Source/Runtime/Renderer/Private/PostProcess/PostProcessVisualizeBuffer.cpp
@@ -370,8 +370,6 @@ FScreenPassTexture AddVisualizeGBufferOverviewPass(
 				EPathTracingPostProcessMaterialInput::Radiance, FScreenPassTexture(PathTracingResources.Radiance, ViewRect));
 		}
 
-		PostProcessMaterialInputs.SceneTextures = Inputs.SceneTextures;
-
 		const TSharedPtr<FImagePixelPipe, ESPMode::ThreadSafe>* OutputPipe = PostProcessSettings.BufferVisualizationPipes.Find(MaterialInterface->GetFName());
 		const bool bIsValidOutputPipe = OutputPipe && OutputPipe->IsValid();
 
@@ -383,6 +381,9 @@ FScreenPassTexture AddVisualizeGBufferOverviewPass(
 		{
 			PostProcessMaterialInputs.OutputFormat = OutputFormat;
 		}
+		PostProcessMaterialInputs.SceneTextures = Inputs.SceneTextures;
+		// When dumping to pipe or file, we mandate the allocation of a new transient output texture separate from scene color.
+		PostProcessMaterialInputs.bAllowSceneColorInputAsOutput = !(bIsValidOutputPipe || Inputs.bDumpToFile);
 
 		Output = AddPostProcessMaterialPass(GraphBuilder, View, PostProcessMaterialInputs, MaterialInterface);
 
