Post Process:  Fix for RHI validation error accessing "Tonemap", related to attempting to bind the texture as both an input SRV and output RTV.  The logic that detects whether an optimization can be applied for post process materials that don't read and write the same texture failed for after tonemapping materials, because the input in that case is actually PostProcessInput3, not PostProcessInput0.

#rnx

diff --git a/Engine/Source/Runtime/Renderer/Private/PostProcess/PostProcessMaterial.cpp b/Engine/Source/Runtime/Renderer/Private/PostProcess/PostProcessMaterial.cpp
--- a/Engine/Source/Runtime/Renderer/Private/PostProcess/PostProcessMaterial.cpp
+++ b/Engine/Source/Runtime/Renderer/Private/PostProcess/PostProcessMaterial.cpp
@@ -767,9 +767,11 @@ FScreenPassTexture AddPostProcessMaterialPass(
 
 	FScreenPassRenderTarget Output = Inputs.OverrideOutput;
 
-	// We can re-use the scene color texture as the render target if we're not simultaneously reading from it.
+	// We can re-use the scene color texture as the render target if we're not simultaneously reading from it.  Note that after tonemapping,
+	// the input is Input3 (PostTonemapHDRColor), not Input0 (SceneColor).
 	const bool bValidShaderPlatform = (GMaxRHIShaderPlatform != SP_PCD3D_ES3_1);
-	if (!Output.IsValid() && !MaterialShaderMap->UsesSceneTexture(PPI_PostProcessInput0) && !bForceIntermediateTarget && Inputs.bAllowSceneColorInputAsOutput && bValidShaderPlatform && !Inputs.bUserSceneTextureOutput)
+	const ESceneTextureId InputTextureId = BlendableLocation == BL_SceneColorAfterTonemapping ? PPI_PostProcessInput3 : PPI_PostProcessInput0;
+	if (!Output.IsValid() && !MaterialShaderMap->UsesSceneTexture(InputTextureId) && !bForceIntermediateTarget && Inputs.bAllowSceneColorInputAsOutput && bValidShaderPlatform && !Inputs.bUserSceneTextureOutput)
 	{
 		FScreenPassTexture SceneColor = FScreenPassTexture::CopyFromSlice(GraphBuilder, SceneColorOutput);
 
