Substrate - mac now uses the new unified path to display visualisation data on screen.

#rnx
#test Tested on an old mac.

diff --git a/Engine/Source/Runtime/Renderer/Private/Substrate/SubstrateVisualize.cpp b/Engine/Source/Runtime/Renderer/Private/Substrate/SubstrateVisualize.cpp
--- a/Engine/Source/Runtime/Renderer/Private/Substrate/SubstrateVisualize.cpp
+++ b/Engine/Source/Runtime/Renderer/Private/Substrate/SubstrateVisualize.cpp
@@ -30,7 +30,6 @@ static bool IsSubstrateDebugVisualizationSupported(EShaderPlatform InPlatform, b
 		Substrate::IsSubstrateEnabled() && 
 		Substrate::UsesSubstrateMaterialBuffer(InPlatform) &&
 		GetMaxSupportedFeatureLevel(InPlatform) >= ERHIFeatureLevel::SM5 &&
-        !IsMetalPlatform(InPlatform) &&     // Fails to launch editor on MAC
 		(bIsEditorOnly ? (IsPCPlatform(InPlatform) || EnumHasAllFlags(Flags, EShaderPermutationFlags::HasEditorOnlyData)) : true);
 }
 
@@ -187,7 +186,8 @@ class FMaterialDebugSubstrateTreeCS : public FGlobalShader
 	static bool IsSupported(EShaderPlatform InPlatform, EShaderPermutationFlags InFlags=EShaderPermutationFlags::None)
 	{
 		return IsSubstrateDebugVisualizationSupported(InPlatform, true, InFlags)
-			&& !IsVulkanPlatform(InPlatform);  // SUBSTRATE_TODO Move to CPU debug visualisation and it should then work on all platforms
+			&& !IsVulkanPlatform(InPlatform)  // SUBSTRATE_TODO Move to CPU debug visualisation and it should then work on all platforms
+			&& !IsMetalPlatform(InPlatform);
 	}
 
 	static FPermutationDomain RemapPermutation(FPermutationDomain PermutationVector)
@@ -233,7 +233,8 @@ class FMaterialDebugSubstrateTreePS : public FGlobalShader
 	static bool IsSupported(EShaderPlatform InPlatform, EShaderPermutationFlags InFlags=EShaderPermutationFlags::None)
 	{
 		return IsSubstrateDebugVisualizationSupported(InPlatform, true, InFlags)
-			&& !IsVulkanPlatform(InPlatform);  // SUBSTRATE_TODO Move to CPU debug visualisation and it should then work on all platforms
+			&& !IsVulkanPlatform(InPlatform)  // SUBSTRATE_TODO Move to CPU debug visualisation and it should then work on all platforms
+			&& !IsMetalPlatform(InPlatform);
 	}
 	
 	static FPermutationDomain RemapPermutation(FPermutationDomain PermutationVector)
