Make sure primitive UB gets updated when primitive custom data changes
#rnx

diff --git a/Engine/Source/Runtime/Renderer/Private/RendererScene.cpp b/Engine/Source/Runtime/Renderer/Private/RendererScene.cpp
--- a/Engine/Source/Runtime/Renderer/Private/RendererScene.cpp
+++ b/Engine/Source/Runtime/Renderer/Private/RendererScene.cpp
@@ -6580,6 +6580,7 @@ void FScene::Update(FRDGBuilder& GraphBuilder, const FUpdateParameters& Paramete
 		FPrimitiveSceneProxy* PrimitiveSceneProxy = Item.SceneInfo->Proxy;
 		FScopeCycleCounter Context(PrimitiveSceneProxy->GetStatId());
 		PrimitiveSceneProxy->CustomPrimitiveData = Item.Payload.Value;
+		PrimitivesNeedingUniformBufferUpdate[Item.SceneInfo->PackedIndex] = true;
 	}
 
 	if (auto NaniteMaterialsUpdater = SceneExtensionsUpdaters.GetUpdaterPtr<Nanite::FMaterialsSceneExtension::FUpdater>())
