RHI validation fixes
- Add SceneViewUniform buffer to path tracing shader parameters because SBT records used in dispatch rays calls use GPUScene data and RDG doesn't wait for transitions on these buffers because it doesn't know about them yet
- Make sure InstanceBuffer.Buffer is set to SRV before building TLAS in RT TestBed

diff --git a/Engine/Source/Runtime/Renderer/Private/PathTracing.cpp b/Engine/Source/Runtime/Renderer/Private/PathTracing.cpp
--- a/Engine/Source/Runtime/Renderer/Private/PathTracing.cpp
+++ b/Engine/Source/Runtime/Renderer/Private/PathTracing.cpp
@@ -1205,7 +1205,8 @@ class FPathTracingRG : public FGlobalShader
 		SHADER_PARAMETER_RDG_TEXTURE_UAV(RWTexture2D<float>, DepthTexture)
 		SHADER_PARAMETER_RDG_BUFFER_SRV(RaytracingAccelerationStructure, TLAS)
 		SHADER_PARAMETER_RDG_BUFFER_SRV(RaytracingAccelerationStructure, DecalTLAS)
-
+		
+		SHADER_PARAMETER_RDG_UNIFORM_BUFFER(FSceneUniformParameters, Scene)
 		SHADER_PARAMETER_STRUCT_REF(FViewUniformShaderParameters, ViewUniformBuffer)
 		SHADER_PARAMETER_STRUCT_INCLUDE(FPathTracingData, PathTracingData)
 		SHADER_PARAMETER_STRUCT_INCLUDE(ShaderPrint::FShaderParameters, ShaderPrint)
@@ -1351,6 +1352,7 @@ class FPathTracingInitExtinctionCoefficientRG : public FGlobalShader
 
 	BEGIN_SHADER_PARAMETER_STRUCT(FParameters, )
 		SHADER_PARAMETER_RDG_BUFFER_SRV(RaytracingAccelerationStructure, TLAS)
+		SHADER_PARAMETER_RDG_UNIFORM_BUFFER(FSceneUniformParameters, Scene)
 		SHADER_PARAMETER_STRUCT_REF(FViewUniformShaderParameters, ViewUniformBuffer)
 		SHADER_PARAMETER_RDG_BUFFER_UAV(RWBuffer<float>, RWStartingExtinctionCoefficient)
 	END_SHADER_PARAMETER_STRUCT()
@@ -3312,6 +3314,7 @@ void FDeferredShadingSceneRenderer::RenderPathTracing(
 
 		FPathTracingInitExtinctionCoefficientRG::FParameters* PassParameters = GraphBuilder.AllocParameters<FPathTracingInitExtinctionCoefficientRG::FParameters>();
 		PassParameters->TLAS = Scene->RayTracingScene.GetLayerView(ERayTracingSceneLayer::Base);
+		PassParameters->Scene = View.GetSceneUniforms().GetBuffer(GraphBuilder);
 		PassParameters->ViewUniformBuffer = View.ViewUniformBuffer;
 		PassParameters->RWStartingExtinctionCoefficient = GraphBuilder.CreateUAV(StartingExtinctionCoefficient, PF_R32_FLOAT);
 
@@ -3760,6 +3763,7 @@ void FDeferredShadingSceneRenderer::RenderPathTracing(
 							FPathTracingRG::FParameters* PassParameters = GraphBuilder.AllocParameters<FPathTracingRG::FParameters>();
 							PassParameters->TLAS = Scene->RayTracingScene.GetLayerView(ERayTracingSceneLayer::Base);
 							PassParameters->DecalTLAS = Scene->RayTracingScene.GetLayerView(ERayTracingSceneLayer::Decals);
+							PassParameters->Scene = View.GetSceneUniforms().GetBuffer(GraphBuilder);
 							PassParameters->ViewUniformBuffer = View.ViewUniformBuffer;
 							PassParameters->PathTracingData = Config.PathTracingData;
 							PassParameters->StartingExtinctionCoefficient = GraphBuilder.CreateSRV(StartingExtinctionCoefficient, PF_R32_FLOAT);
diff --git a/Engine/Source/Runtime/Renderer/Private/Tests/RayTracingTestbed.cpp b/Engine/Source/Runtime/Renderer/Private/Tests/RayTracingTestbed.cpp
--- a/Engine/Source/Runtime/Renderer/Private/Tests/RayTracingTestbed.cpp
+++ b/Engine/Source/Runtime/Renderer/Private/Tests/RayTracingTestbed.cpp
@@ -300,6 +300,8 @@ bool RunRayTracingTestbed_RenderThread(const FString& Parameters)
 	BuildParams.PerInstanceGeometries = RayTracingSceneInitializationData.PerInstanceGeometries;
 	BuildParams.NumInstances = RayTracingSceneInitializationData.NumNativeGPUSceneInstances + RayTracingSceneInitializationData.NumNativeCPUInstances;
 
+	RHICmdList.Transition(FRHITransitionInfo(InstanceBuffer.Buffer, ERHIAccess::UAVMask, ERHIAccess::SRVCompute));
+
 	RHICmdList.BuildAccelerationStructure(BuildParams);
 
 	RHICmdList.Transition(FRHITransitionInfo(RayTracingSceneRHI.GetReference(), ERHIAccess::BVHWrite, ERHIAccess::BVHRead));
