VolumetricCloud - fixed cloud disapearing when volumetricrendertarget is disable but high quality aerial perspective is enabled.

#tests Win64 editor sky + clou with different parameters
[RN] Fixed clouds disapearing when VolumetricRenderTarget is disable while using cloud high quality aerial perspective.

diff --git a/Engine/Source/Runtime/Renderer/Private/VolumetricCloudRendering.cpp b/Engine/Source/Runtime/Renderer/Private/VolumetricCloudRendering.cpp
--- a/Engine/Source/Runtime/Renderer/Private/VolumetricCloudRendering.cpp
+++ b/Engine/Source/Runtime/Renderer/Private/VolumetricCloudRendering.cpp
@@ -2137,6 +2137,7 @@ FCloudRenderContext::FCloudRenderContext()
 	TracingCoordToFullResPixelCoordScaleBias = FUintVector4(1, 1, 0, 0);
 	NoiseFrameIndexModPattern = 0;
 
+	bDisableCloudBlending = false;
 	bIsReflectionRendering = false;
 	bIsSkyRealTimeReflectionRendering = false;
 	bSkipAtmosphericLightShadowmap = false;
@@ -2466,7 +2467,7 @@ void FSceneRenderer::RenderVolumetricCloudsInternal(FRDGBuilder& GraphBuilder, F
 
 		FRenderVolumetricCloudRenderViewCS::FParameters* PassParameters = GraphBuilder.AllocParameters<FRenderVolumetricCloudRenderViewCS::FParameters>();
 		PassParameters->OutputViewRect = FVector4f(0.f, 0.f, Desc.Extent.X, Desc.Extent.Y);
-		PassParameters->bBlendCloudColor = !bShouldViewRenderVolumetricRenderTarget && !bCloudDebugViewModeEnabled;
+		PassParameters->bBlendCloudColor = !bShouldViewRenderVolumetricRenderTarget && !bCloudDebugViewModeEnabled && !CloudRC.bDisableCloudBlending;
 		PassParameters->TargetCubeFace = CloudRC.RenderTargets.Output[0].GetArraySlice();
 		PassParameters->VolumetricCloudRenderViewParamsUB = CloudPassUniformBuffer;
 		PassParameters->View = ViewUniformBuffer;
@@ -2697,6 +2698,7 @@ bool FSceneRenderer::RenderVolumetricCloud(
 				const bool bEnableAerialPerspectiveSampling = ViewToProcess.bEnableAerialPerspectiveSampling;
 				const bool bShouldUseHighQualityAerialPerspective = ViewToProcess.bShouldUseHighQualityAerialPerspective;
 
+				CloudRC.bDisableCloudBlending = !bShouldViewRenderVolumetricCloudRenderTarget && bShouldUseHighQualityAerialPerspective;
 				CloudRC.bShouldViewRenderVolumetricRenderTarget = bShouldViewRenderVolumetricCloudRenderTarget;
 				CloudRC.ViewUniformBuffer = bShouldViewRenderVolumetricCloudRenderTarget ? ViewInfo.VolumetricRenderTargetViewUniformBuffer : ViewInfo.ViewUniformBuffer;
 
diff --git a/Engine/Source/Runtime/Renderer/Private/VolumetricCloudRendering.h b/Engine/Source/Runtime/Renderer/Private/VolumetricCloudRendering.h
--- a/Engine/Source/Runtime/Renderer/Private/VolumetricCloudRendering.h
+++ b/Engine/Source/Runtime/Renderer/Private/VolumetricCloudRendering.h
@@ -138,6 +138,7 @@ struct FCloudRenderContext
 
 	FRDGTextureRef SecondaryCloudTracingDataTexture = nullptr;
 
+	bool bDisableCloudBlending;
 	bool bShouldViewRenderVolumetricRenderTarget;
 	bool bSkipAerialPerspective;
 	bool bSkipHeightFog;
