Moved Nanite::PrintStats to after VSM rendering so "nanitestats virtualshadowmaps" show the stats for the current frame (not the previous).

diff --git a/Engine/Source/Runtime/Renderer/Private/DeferredShadingRenderer.cpp b/Engine/Source/Runtime/Renderer/Private/DeferredShadingRenderer.cpp
--- a/Engine/Source/Runtime/Renderer/Private/DeferredShadingRenderer.cpp
+++ b/Engine/Source/Runtime/Renderer/Private/DeferredShadingRenderer.cpp
@@ -2461,18 +2461,6 @@ void FDeferredShadingSceneRenderer::Render(FRDGBuilder& GraphBuilder)
 
 			if (bNaniteEnabled)
 			{
-				if (GNaniteShowStats != 0)
-				{
-					for (int32 ViewIndex = 0; ViewIndex < Views.Num(); ViewIndex++)
-					{
-						const FViewInfo& View = Views[ViewIndex];
-						if (IStereoRendering::IsAPrimaryView(View))
-						{
-							Nanite::PrintStats(GraphBuilder, View);
-						}
-					}
-				}
-
 				if (bVisualizeNanite)
 				{
 					FNanitePickingFeedback PickingFeedback = { 0 };
@@ -2686,6 +2674,18 @@ void FDeferredShadingSceneRenderer::Render(FRDGBuilder& GraphBuilder)
 		{
 			// Needs doing after shadows such that the checks for shadow atlases etc work.
 			Nanite::ListStatFilters(this);
+
+			if (GNaniteShowStats != 0)
+			{
+				for (int32 ViewIndex = 0; ViewIndex < Views.Num(); ViewIndex++)
+				{
+					const FViewInfo& View = Views[ViewIndex];
+					if (IStereoRendering::IsAPrimaryView(View))
+					{
+						Nanite::PrintStats(GraphBuilder, View);
+					}
+				}
+			}
 		}
 
 		if (bUpdateNaniteStreaming)
@@ -2694,8 +2694,7 @@ void FDeferredShadingSceneRenderer::Render(FRDGBuilder& GraphBuilder)
 		}
 
 		{
-			FVirtualShadowMapArrayCacheManager* CacheManager = VirtualShadowMapArray.CacheManager;
-			if (CacheManager)
+			if (FVirtualShadowMapArrayCacheManager* CacheManager = VirtualShadowMapArray.CacheManager)
 			{
 				// Do this even if VSMs are disabled this frame to clean up any previously extracted data
 				CacheManager->ExtractFrameData(
