Post Process Materials:  More general fix for RHI validation errors where the input and output are the same.  Previous fix assumed that After Tonemapping used PPI_PostProcessInput3 as its input, but both PPI_PostProcessInput0 and PPI_PostProcessInput3 are initialized to the same input, and can be (and are) used as such by materials.  Checking ALL used inputs solves the problem across the board, including handling potential future cases where inputs may be configured differently, without needing to revisit the code again.

#rnx

diff --git a/Engine/Source/Runtime/Renderer/Private/PostProcess/PostProcessMaterial.cpp b/Engine/Source/Runtime/Renderer/Private/PostProcess/PostProcessMaterial.cpp
--- a/Engine/Source/Runtime/Renderer/Private/PostProcess/PostProcessMaterial.cpp
+++ b/Engine/Source/Runtime/Renderer/Private/PostProcess/PostProcessMaterial.cpp
@@ -767,11 +767,23 @@ FScreenPassTexture AddPostProcessMaterialPass(
 
 	FScreenPassRenderTarget Output = Inputs.OverrideOutput;
 
-	// We can re-use the scene color texture as the render target if we're not simultaneously reading from it.  Note that after tonemapping,
-	// the input is Input3 (PostTonemapHDRColor), not Input0 (SceneColor).
+	// We can re-use the scene color texture as the render target if we're not simultaneously reading from it.
+	bool bInputReadsFromOutput = false;
+	if (!Output.IsValid())
+	{
+		for (int32 InputIndex = 0; InputIndex < kPostProcessMaterialInputCountMax; InputIndex++)
+		{
+			if (MaterialShaderMap->UsesSceneTexture((ESceneTextureId)(PPI_PostProcessInput0 + InputIndex)) && Inputs.Textures[InputIndex].TextureSRV &&
+				SceneColorOutput.TextureSRV->GetParent() == Inputs.Textures[InputIndex].TextureSRV->GetParent())
+			{
+				bInputReadsFromOutput = true;
+				break;
+			}
+		}
+	}
+
 	const bool bValidShaderPlatform = (GMaxRHIShaderPlatform != SP_PCD3D_ES3_1);
-	const ESceneTextureId InputTextureId = BlendableLocation == BL_SceneColorAfterTonemapping ? PPI_PostProcessInput3 : PPI_PostProcessInput0;
-	if (!Output.IsValid() && !MaterialShaderMap->UsesSceneTexture(InputTextureId) && !bForceIntermediateTarget && Inputs.bAllowSceneColorInputAsOutput && bValidShaderPlatform && !Inputs.bUserSceneTextureOutput)
+	if (!Output.IsValid() && !bInputReadsFromOutput && !bForceIntermediateTarget && Inputs.bAllowSceneColorInputAsOutput && bValidShaderPlatform && !Inputs.bUserSceneTextureOutput)
 	{
 		FScreenPassTexture SceneColor = FScreenPassTexture::CopyFromSlice(GraphBuilder, SceneColorOutput);
 
