[Light Grid] Refactored handling of limits in LightGrid.
- Log warnings when building grid exceeds the max number of entries/links.
- ensure instead of check that NumLocalLights and NumReflectionCaptures are less than the primitive limit imposed by FCellWriter.
- clamp NumCulledLightLinks when using bThreadGroupPerCell (unlikely since the first 256 in each cell are temporarily stored in LDS instead of linked list).

diff --git a/Engine/Shaders/Private/LightGridInjection.usf b/Engine/Shaders/Private/LightGridInjection.usf
--- a/Engine/Shaders/Private/LightGridInjection.usf
+++ b/Engine/Shaders/Private/LightGridInjection.usf
@@ -10,11 +10,14 @@
 #include "ReflectionEnvironmentShared.ush"
 #include "LightGridCommon.ush"
 #include "WaveOpUtil.ush"
+#include "GPUMessaging.ush"
 
 #if USE_HZB_CULL
 	#include "Nanite/NaniteHZBCull.ush"
 #endif
 
+#include "/Engine/Shared/LightGridDefinitions.h"
+
 RWStructuredBuffer<uint> RWNumCulledLightsGrid;
 
 #if LIGHT_GRID_USES_16BIT_BUFFERS
@@ -508,7 +511,7 @@ struct FCellWriter
 FCellWriter InitCellWriter()
 {
 	FCellWriter CellWriter;
-	CellWriter.LinkOffset = 0xFFFFF;
+	CellWriter.LinkOffset = LIGHT_GRID_CELL_WRITER_MAX_LINK_OFFSET;
 	
 	return CellWriter;
 }
@@ -530,15 +533,15 @@ bool AddToCell(inout FCellWriter CellWriter, uint PrimitiveIndex, uint IndexInCe
 
 	if (NextLink < NumAvailableLinks)
 	{
-		checkSlow(PrimitiveIndex <= 0xFFFFFF);
+		checkSlow(PrimitiveIndex <= LIGHT_GRID_CELL_WRITER_MAX_LINK_OFFSET);
 		checkSlow(IndexInCell <= 0xFFFF);
-		checkSlow(CellWriter.LinkOffset <= 0xFFFFFF);
+		checkSlow(CellWriter.LinkOffset <= LIGHT_GRID_CELL_WRITER_MAX_LINK_OFFSET);
 
 		uint IndexInCellL8 = (IndexInCell & 0xFF);
 		uint IndexInCellH8 = ((IndexInCell >> 8) & 0xFF);
 
-		RWCulledLightLinks[NextLink * LIGHT_LINK_STRIDE + 0] = (IndexInCellL8 << 24) | (PrimitiveIndex & 0xFFFFFF);
-		RWCulledLightLinks[NextLink * LIGHT_LINK_STRIDE + 1] = (IndexInCellH8 << 24) | (CellWriter.LinkOffset & 0xFFFFFF);
+		RWCulledLightLinks[NextLink * LIGHT_LINK_STRIDE + 0] = (IndexInCellL8 << 24) | (PrimitiveIndex & LIGHT_GRID_CELL_WRITER_MAX_LINK_OFFSET);
+		RWCulledLightLinks[NextLink * LIGHT_LINK_STRIDE + 1] = (IndexInCellH8 << 24) | (CellWriter.LinkOffset & LIGHT_GRID_CELL_WRITER_MAX_LINK_OFFSET);
 
 		CellWriter.LinkOffset = NextLink;
 
@@ -935,3 +938,30 @@ void DebugLightGridPS(
 
 #endif // SHADER_DEBUG_LIGHT_GRID_PS
 
+#ifdef FeedbackStatusCS
+
+StructuredBuffer<uint> NextCulledLightDataBuffer;
+uint NumCulledLightDataEntries;
+
+StructuredBuffer<uint> NextCulledLightLinkBuffer;
+
+uint StatusMessageId;
+
+[numthreads(1, 1, 1)]
+void FeedbackStatusCS()
+{												
+	uint NumAllocatedEntries = NextCulledLightDataBuffer[0];
+	uint NumAllocatedLinks = NextCulledLightLinkBuffer[0];
+
+	FGPUMessageWriter Mw = GPUMessageBegin(StatusMessageId, 4U);
+
+	// max of allocated entries or links because when link allocation fail entries are not allocated
+	// however every link allocation would result in an entry allocation
+	GPUMessageWriteItem(Mw, max(NumAllocatedEntries, NumAllocatedLinks));
+	GPUMessageWriteItem(Mw, NumCulledLightDataEntries);
+
+	GPUMessageWriteItem(Mw, NumAllocatedLinks);
+	GPUMessageWriteItem(Mw, NumAvailableLinks);
+}
+
+#endif
diff --git a/Engine/Shaders/Shared/LightGridDefinitions.h b/Engine/Shaders/Shared/LightGridDefinitions.h
new file mode 100644
--- /dev/null
+++ b/Engine/Shaders/Shared/LightGridDefinitions.h
@@ -0,0 +1,14 @@
+// Copyright Epic Games, Inc. All Rights Reserved.
+
+/*================================================================================================
+	LightGridDefinitions.h: used in light grid shaders and C++ code to define common constants
+=================================================================================================*/
+
+#pragma once
+
+// Must updating packing in FCellWriter when modifying these
+#define LIGHT_GRID_CELL_WRITER_MAX_NUM_PRIMITIVES (1 << 24) // must be power to two
+#define LIGHT_GRID_CELL_WRITER_MAX_PRIMITIVE_INDEX (LIGHT_GRID_CELL_WRITER_MAX_NUM_PRIMITIVES - 1)
+
+#define LIGHT_GRID_CELL_WRITER_MAX_NUM_LINKS (1 << 24) // must be power to two
+#define LIGHT_GRID_CELL_WRITER_MAX_LINK_OFFSET (LIGHT_GRID_CELL_WRITER_MAX_NUM_LINKS - 1)
diff --git a/Engine/Source/Runtime/Renderer/Private/LightGrid.h b/Engine/Source/Runtime/Renderer/Private/LightGrid.h
new file mode 100644
--- /dev/null
+++ b/Engine/Source/Runtime/Renderer/Private/LightGrid.h
@@ -0,0 +1,32 @@
+// Copyright Epic Games, Inc. All Rights Reserved.
+
+/*=============================================================================
+	LumenGrid.h:
+=============================================================================*/
+
+#pragma once
+
+#include "GPUMessaging.h"
+#include "RenderGraphResources.h"
+
+class FViewInfo;
+
+class FLightGridViewState
+{
+public:
+
+	FLightGridViewState();
+
+	void FeedbackStatus(FRDGBuilder& GraphBuilder, FViewInfo& View, FRDGBufferRef NextCulledLightDataBuffer, uint32 NumCulledLightDataEntries, FRDGBufferRef NextCulledLightLinkBuffer, uint32 NumCulledLightLinks);
+
+#if !UE_BUILD_SHIPPING
+	uint32 GetStatusMessageId() const { return StatusFeedbackSocket.GetMessageId().GetIndex(); }
+#endif
+
+private:
+
+#if !UE_BUILD_SHIPPING
+	GPUMessage::FSocket StatusFeedbackSocket;
+	uint32 MaxEntriesHighWaterMark = 0;
+#endif
+};
\ No newline at end of file
diff --git a/Engine/Source/Runtime/Renderer/Private/LightGridInjection.cpp b/Engine/Source/Runtime/Renderer/Private/LightGridInjection.cpp
--- a/Engine/Source/Runtime/Renderer/Private/LightGridInjection.cpp
+++ b/Engine/Source/Runtime/Renderer/Private/LightGridInjection.cpp
@@ -4,6 +4,8 @@
 LightGridInjection.cpp
 =============================================================================*/
 
+#include "LightGrid.h"
+
 #include "CoreMinimal.h"
 #include "Stats/Stats.h"
 #include "HAL/IConsoleManager.h"
@@ -34,6 +36,7 @@ LightGridInjection.cpp
 #include "ShaderPrintParameters.h"
 #include "RenderUtils.h"
 #include "MegaLights/MegaLights.h"
+#include "LightGridDefinitions.h"
 
 int32 GLightGridPixelSize = 64;
 FAutoConsoleVariableRef CVarLightGridPixelSize(
@@ -559,6 +562,7 @@ FLightGrid LightGridInjection(
 	FRDGBufferSRVRef LightViewSpacePositionAndRadiusSRV,
 	FRDGBufferSRVRef LightViewSpaceDirAndPreprocAngleSRV,
 	FRDGBufferSRVRef LightViewSpaceRectPlanesSRV,
+	FLightGridViewState* LightGridViewState,
 	bool bThreadGroupPerCell,
 	bool bThreadGroupSize32,
 	// parent params
@@ -566,14 +570,16 @@ FLightGrid LightGridInjection(
 	FRDGBufferSRVRef ParentCulledLightDataGridSRV,
 	uint32 ParentGridSizeFactor)
 {
-	const uint32 NumCulledLightLinks = MaxNumCells * GMaxCulledLightsPerCell * NumCulledGridPrimitiveTypes;
+	const uint32 NumCulledLightEntries = MaxNumCells * GMaxCulledLightsPerCell;
+
+	uint32 NumCulledLightLinks = MaxNumCells * GMaxCulledLightsPerCell;
 	
 	if (bThreadGroupPerCell)
 	{
-		checkf(NumLocalLights < 0xFFFFFF, TEXT("NumLocalLights limited to 16M by FCellWriter."));
-		checkf(NumReflectionCaptures < 0xFFFFFF, TEXT("NumReflectionCaptures limited to 16M by FCellWriter."));
+		ensureMsgf(NumLocalLights <= LIGHT_GRID_CELL_WRITER_MAX_NUM_PRIMITIVES, TEXT("NumLocalLights limited to 16M by FCellWriter."));
+		ensureMsgf(NumReflectionCaptures <= LIGHT_GRID_CELL_WRITER_MAX_NUM_PRIMITIVES, TEXT("NumLocalLights limited to 16M by FCellWriter."));
 
-		checkf(NumCulledLightLinks < 0xFFFFFF, TEXT("NumCulledLightLinks limited to 16M by FCellWriter."));
+		NumCulledLightLinks = FMath::Min(NumCulledLightLinks, (uint32)LIGHT_GRID_CELL_WRITER_MAX_NUM_LINKS); // limited to 16M by FCellWriter (will cause warning if exceeded, see FLightGridViewState::Update())
 	}
 
 	const FIntVector ParentGridSize = FIntVector::DivideAndRoundUp(GridSize, ParentGridSizeFactor);
@@ -591,7 +597,7 @@ FLightGrid LightGridInjection(
 	{
 		const SIZE_T LightIndexTypeSize = sizeof(FLightIndexType);
 		const EPixelFormat CulledLightDataGridFormat = PF_R16_UINT;
-		FRDGBufferRef CulledLightDataGrid = GraphBuilder.CreateBuffer(FRDGBufferDesc::CreateBufferDesc(LightIndexTypeSize, MaxNumCells * GMaxCulledLightsPerCell), TEXT("CulledLightDataGrid"));
+		FRDGBufferRef CulledLightDataGrid = GraphBuilder.CreateBuffer(FRDGBufferDesc::CreateBufferDesc(LightIndexTypeSize, NumCulledLightEntries), TEXT("CulledLightDataGrid"));
 		CulledLightDataGridSRV = GraphBuilder.CreateSRV(CulledLightDataGrid, CulledLightDataGridFormat);
 		CulledLightDataGridUAV = GraphBuilder.CreateUAV(CulledLightDataGrid, CulledLightDataGridFormat);
 	}
@@ -599,7 +605,7 @@ FLightGrid LightGridInjection(
 	{
 		const SIZE_T LightIndexTypeSize = sizeof(FLightIndexType32);
 		const EPixelFormat CulledLightDataGridFormat = PF_R32_UINT;
-		FRDGBufferRef CulledLightDataGrid = GraphBuilder.CreateBuffer(FRDGBufferDesc::CreateStructuredDesc(LightIndexTypeSize, MaxNumCells * GMaxCulledLightsPerCell), TEXT("CulledLightDataGrid"));
+		FRDGBufferRef CulledLightDataGrid = GraphBuilder.CreateBuffer(FRDGBufferDesc::CreateStructuredDesc(LightIndexTypeSize, NumCulledLightEntries), TEXT("CulledLightDataGrid"));
 		CulledLightDataGridSRV = GraphBuilder.CreateSRV(CulledLightDataGrid);
 		CulledLightDataGridUAV = GraphBuilder.CreateUAV(CulledLightDataGrid);
 	}
@@ -689,6 +695,11 @@ FLightGrid LightGridInjection(
 	Output.CulledLightDataGridSRV = CulledLightDataGridSRV;
 	Output.NumCulledLightsGridSRV = GraphBuilder.CreateSRV(FRDGBufferSRVDesc(NumCulledLightsGrid));
 
+	if (LightGridViewState != nullptr)
+	{
+		LightGridViewState->FeedbackStatus(GraphBuilder, View, NextCulledLightDataBuffer, NumCulledLightEntries, NextCulledLightLinkBuffer, NumCulledLightLinks);
+	}
+
 	return Output;
 }
 
@@ -1132,8 +1143,9 @@ FComputeLightGridOutput FSceneRenderer::ComputeLightGrid(FRDGBuilder& GraphBuild
 				LightViewSpacePositionAndRadiusSRV,
 				LightViewSpaceDirAndPreprocAngleSRV,
 				LightViewSpaceRectPlanesSRV,
-				true,
-				false,
+				View.ViewState != nullptr ? &View.ViewState->LightGrid : nullptr,
+				/*bThreadGroupPerCell*/ true,
+				/*bThreadGroupSize32*/ false,
 				nullptr,
 				nullptr,
 				1);
@@ -1168,6 +1180,7 @@ FComputeLightGridOutput FSceneRenderer::ComputeLightGrid(FRDGBuilder& GraphBuild
 			LightViewSpacePositionAndRadiusSRV,
 			LightViewSpaceDirAndPreprocAngleSRV,
 			LightViewSpaceRectPlanesSRV,
+			View.ViewState != nullptr ? &View.ViewState->LightGrid : nullptr,
 			NumThreadsPerCell > 1,
 			NumThreadsPerCell == 32,
 			ParentLightGrid.NumCulledLightsGridSRV,
@@ -1419,3 +1432,86 @@ FScreenPassTexture AddVisualizeLightGridPass(FRDGBuilder& GraphBuilder, const FV
 
 	return MoveTemp(ScreenPassSceneColor);
 }
+
+class FLightGridFeedbackStatusCS : public FGlobalShader
+{
+	DECLARE_GLOBAL_SHADER(FLightGridFeedbackStatusCS);
+	SHADER_USE_PARAMETER_STRUCT(FLightGridFeedbackStatusCS, FGlobalShader);
+
+	BEGIN_SHADER_PARAMETER_STRUCT(FParameters, )
+		SHADER_PARAMETER_RDG_BUFFER_SRV(StructuredBuffer<uint>, NextCulledLightDataBuffer)
+		SHADER_PARAMETER(uint32, NumCulledLightDataEntries)
+
+		SHADER_PARAMETER_RDG_BUFFER_SRV(StructuredBuffer<uint>, NextCulledLightLinkBuffer)
+		SHADER_PARAMETER(uint32, NumAvailableLinks)
+
+		SHADER_PARAMETER_STRUCT_INCLUDE(GPUMessage::FParameters, GPUMessageParams)
+		SHADER_PARAMETER(uint32, StatusMessageId)
+	END_SHADER_PARAMETER_STRUCT()
+};
+IMPLEMENT_GLOBAL_SHADER(FLightGridFeedbackStatusCS, "/Engine/Private/LightGridInjection.usf", "FeedbackStatusCS", SF_Compute);
+
+FLightGridViewState::FLightGridViewState()
+{
+#if !UE_BUILD_SHIPPING
+	StatusFeedbackSocket = GPUMessage::RegisterHandler(TEXT("LightGrid.StatusFeedback"),
+		[this](GPUMessage::FReader Message)
+		{
+			const uint32 AllocatedEntries = Message.Read<uint32>(0);
+			const uint32 MaxEntries = Message.Read<uint32>(0);
+
+			const uint32 AllocatedLinks = Message.Read<uint32>(0);
+			const uint32 MaxLinks = Message.Read<uint32>(0);
+
+			if (AllocatedEntries > MaxEntries)
+			{
+				bool bWarn = MaxEntries > MaxEntriesHighWaterMark;
+
+				if (bWarn)
+				{
+					UE_LOG(LogRenderer, Warning, TEXT(	"Building light grid exceeded number of available entries (%u / %u). "
+														"Increase r.Forward.MaxCulledLightsPerCell to prevent potential visual artifacts."), AllocatedEntries, MaxEntries);
+				}
+
+				MaxEntriesHighWaterMark = FMath::Max(MaxEntriesHighWaterMark, MaxEntries);
+			}
+
+			if (AllocatedLinks > MaxLinks)
+			{
+				static bool bWarn = true;
+
+				if (bWarn)
+				{
+					UE_LOG(LogRenderer, Warning, TEXT("Building light grid exceeded number of available links, glitches will be visible (%u / %u)."), AllocatedLinks, MaxLinks);
+					bWarn = false;
+				}
+			}
+		});
+#endif
+}
+
+void FLightGridViewState::FeedbackStatus(FRDGBuilder& GraphBuilder, FViewInfo& View, FRDGBufferRef NextCulledLightDataBuffer, uint32 NumCulledLightDataEntries, FRDGBufferRef NextCulledLightLinkBuffer, uint32 NumCulledLightLinks)
+{
+#if !UE_BUILD_SHIPPING
+	FLightGridFeedbackStatusCS::FParameters* PassParameters = GraphBuilder.AllocParameters<FLightGridFeedbackStatusCS::FParameters>();
+
+	PassParameters->NextCulledLightDataBuffer = GraphBuilder.CreateSRV(NextCulledLightDataBuffer);
+	PassParameters->NumCulledLightDataEntries = NumCulledLightDataEntries;
+
+	PassParameters->NextCulledLightLinkBuffer = GraphBuilder.CreateSRV(NextCulledLightLinkBuffer);
+	PassParameters->NumAvailableLinks = NumCulledLightLinks;
+
+	PassParameters->GPUMessageParams = GPUMessage::GetShaderParameters(GraphBuilder);
+	PassParameters->StatusMessageId = GetStatusMessageId();
+
+	auto ComputeShader = View.ShaderMap->GetShader<FLightGridFeedbackStatusCS>();
+
+	FComputeShaderUtils::AddPass(
+		GraphBuilder,
+		RDG_EVENT_NAME("LightGridFeedbackStatus"),
+		ComputeShader,
+		PassParameters,
+		FIntVector(1, 1, 1)
+	);
+#endif // !UE_BUILD_SHIPPING
+}
\ No newline at end of file
diff --git a/Engine/Source/Runtime/Renderer/Private/ScenePrivate.h b/Engine/Source/Runtime/Renderer/Private/ScenePrivate.h
--- a/Engine/Source/Runtime/Renderer/Private/ScenePrivate.h
+++ b/Engine/Source/Runtime/Renderer/Private/ScenePrivate.h
@@ -43,6 +43,7 @@
 #include "RayTracing/RayTracingShaderBindingTable.h"
 #endif
 #include "Nanite/Nanite.h"
+#include "LightGrid.h"
 #include "Lumen/LumenViewState.h"
 #include "MegaLights/MegaLightsViewState.h"
 #include "StochasticLighting/StochasticLightingViewState.h"
@@ -998,6 +999,8 @@ public:
 	// Burley Subsurface scattering variance texture from the last frame.
 	TRefCountPtr<IPooledRenderTarget> SubsurfaceScatteringQualityHistoryRT;
 
+	FLightGridViewState LightGrid;
+
 	FLumenViewState Lumen;
 	FMegaLightsViewState MegaLights;
 
