MGPU:  For for RDG validation error recently introduced in CL 36942256.  Need pass parameters to include textures used in cross GPU transfers.

#rnx

diff --git a/Engine/Source/Runtime/Renderer/Private/SceneRendering.cpp b/Engine/Source/Runtime/Renderer/Private/SceneRendering.cpp
--- a/Engine/Source/Runtime/Renderer/Private/SceneRendering.cpp
+++ b/Engine/Source/Runtime/Renderer/Private/SceneRendering.cpp
@@ -3616,7 +3616,6 @@ struct FCrossGPUTransfer
 struct FCrossGPUTarget
 {
 	const FRenderTarget* RenderTarget = nullptr;
-	FRDGTextureRef RenderTargetTexture = nullptr;
 	TArray<FCrossGPUTransfer> Transfers;
 };
 
@@ -3726,7 +3725,7 @@ void FSceneRenderer::DoCrossGPUTransfers(FRDGBuilder& GraphBuilder, FRDGTextureR
 		// Accumulate transfers from each scene renderer
 		if (Transfers.Num() > 0)
 		{
-			TransfersDeferred->Targets.Add({ ViewFamily.RenderTarget, nullptr, MoveTemp(Transfers) });
+			TransfersDeferred->Targets.Add({ ViewFamily.RenderTarget, MoveTemp(Transfers) });
 		}
 	}
 	else if (Transfers.Num() > 0)
@@ -3772,6 +3771,12 @@ void FSceneRenderer::DoCrossGPUTransfers(FRDGBuilder& GraphBuilder, FRDGTextureR
 #endif // WITH_MGPU
 }
 
+#if WITH_MGPU
+BEGIN_SHADER_PARAMETER_STRUCT(FFlushCrossGPUTransfersParameters, )
+	RDG_TEXTURE_ACCESS_ARRAY(Textures)
+END_SHADER_PARAMETER_STRUCT()
+#endif
+
 void FSceneRenderer::FlushCrossGPUTransfers(FRDGBuilder& GraphBuilder)
 {
 #if WITH_MGPU
@@ -3780,23 +3785,30 @@ void FSceneRenderer::FlushCrossGPUTransfers(FRDGBuilder& GraphBuilder)
 		// If this is the last scene renderer, flush the transfers
 		if (CrossGPUTransferDeferred->GetRefCount() == 1 && CrossGPUTransferDeferred->Targets.Num())
 		{
+			auto* PassParameters = GraphBuilder.AllocParameters<FFlushCrossGPUTransfersParameters>();
+			PassParameters->Textures.Reserve(CrossGPUTransferDeferred->Targets.Num());
+
 			// Create RDG textures for each render target
 			for (FCrossGPUTarget& Target : CrossGPUTransferDeferred->Targets)
 			{
 				FRHITexture* TextureRHI = Target.RenderTarget->GetRenderTargetTexture();
 				check(TextureRHI);
-				Target.RenderTargetTexture = RegisterExternalTexture(GraphBuilder, TextureRHI, TEXT("CrossGPUTexture"));
+				PassParameters->Textures.Emplace(RegisterExternalTexture(GraphBuilder, TextureRHI, TEXT("CrossGPUTexture")), ERHIAccess::CopySrc);
 			}
 
-			AddPass(GraphBuilder, RDG_EVENT_NAME("CrossGPUTransfers"),
-				[LocalTransfers = CrossGPUTransferDeferred](FRHICommandList& RHICmdList)
+			GraphBuilder.AddPass(
+				RDG_EVENT_NAME("CrossGPUTransfers"),
+				PassParameters,
+				ERDGPassFlags::Copy | ERDGPassFlags::NeverCull,
+				[LocalTransfers = CrossGPUTransferDeferred, PassParameters](FRHICommandList& RHICmdList)
 			{
 				TArray<FTransferResourceParams> TransferParams;
-				for (const FCrossGPUTarget& Target : LocalTransfers->Targets)
+				for (int32 TargetIndex = 0; TargetIndex < LocalTransfers->Targets.Num(); TargetIndex++)
 				{
+					const FCrossGPUTarget& Target = LocalTransfers->Targets[TargetIndex];
 					for (const FCrossGPUTransfer& Transfer : Target.Transfers)
 					{
-						TransferParams.Add(FTransferResourceParams(Target.RenderTargetTexture->GetRHI(), Transfer.SrcGPUIndex, Transfer.DestGPUIndex, true, true));
+						TransferParams.Add(FTransferResourceParams(PassParameters->Textures[TargetIndex].GetTexture()->GetRHI(), Transfer.SrcGPUIndex, Transfer.DestGPUIndex, true, true));
 					}
 				}
 
