Holdout: Adds show flag to allow or disallow primitive alpha holdout rendering.

* This lets scene capture or custom render pass users easily render primitive(s) with holdout in the main view, and without in the separate render.
* The same mechanism is used to globally force-disable primitive alpha holdout if the view is a reflection capture.

diff --git a/Engine/Shaders/Private/BasePassPixelShader.usf b/Engine/Shaders/Private/BasePassPixelShader.usf
--- a/Engine/Shaders/Private/BasePassPixelShader.usf
+++ b/Engine/Shaders/Private/BasePassPixelShader.usf
@@ -910,7 +910,7 @@ void FPixelShaderInOut_MainPS(
 #if (MATERIALBLENDING_ALPHAHOLDOUT || SUBSTRATE_BLENDING_ALPHAHOLDOUT)
 	const bool bIsHoldout = true;
 #else
-	const bool bIsHoldout = (GetPrimitiveData(MaterialParameters).Flags & PRIMITIVE_SCENE_DATA_FLAG_HOLDOUT);
+	const bool bIsHoldout = (GetPrimitiveData(MaterialParameters).Flags & PRIMITIVE_SCENE_DATA_FLAG_HOLDOUT) && ResolvedView.bAllowPrimitiveAlphaHoldout;
 #endif
 #endif
 
@@ -2257,26 +2257,22 @@ void FPixelShaderInOut_MainPS(
 	// Only include the holdout logic if the Renderer is forwarding its alpha channel to post-processing
 	if (bIsHoldout)
 	{
-		// Ignore primitive holdout during reflection captures
-		if(ResolvedView.RenderingReflectionCaptureMask == 0)
-		{
-			// Discard this pixel so that nothing gets written to the G-Buffers and alpha is kept as its default value (background)
-
+		// Discard this pixel so that nothing gets written to the G-Buffers and alpha is kept as its default value (background)
 #if MATERIALBLENDING_SOLID || MATERIALBLENDING_MASKED
-			// This catches the common case as well as handles cases where Opacity has a different meaning, such as for Subsurface materials
-			Out.MRT[0] = half4(0.0, 0.0, 0.0, 1.0);
+		// This catches the common case as well as handles cases where Opacity has a different meaning, such as for Subsurface materials
+		Out.MRT[0] = half4(0.0, 0.0, 0.0, 1.0);
 #elif MATERIALBLENDING_MODULATE
-			Out.MRT[0] = half4(0.0, 0.0, 0.0, Opacity);
-			// Alternative method. Keep MRT[0] as it is to get a closer look to path tracing:
-			//	Out.MRT[0] = half4(FoggedColor, Opacity);
-			// MATERIALBLENDING_MODULATE comments in the holdout pass.
+		Out.MRT[0] = half4(0.0, 0.0, 0.0, Opacity);
+		// Alternative method. Keep MRT[0] as it is to get a closer look to path tracing:
+		//	Out.MRT[0] = half4(FoggedColor, Opacity);
+		// MATERIALBLENDING_MODULATE comments in the holdout pass.
 #elif SUBSTRATE_TRANSLUCENT_MATERIAL
 	#if SUBSTRATE_BLENDING_ALPHAHOLDOUT
-			//This is not the translucent holdout pass, do not reset Out.MRT[0] so holdout can accumulate color contribution.
+		//This is not the translucent holdout pass, do not reset Out.MRT[0] so holdout can accumulate color contribution.
 	#elif SUBSTRATE_BLENDING_TRANSLUCENT_COLOREDTRANSMITTANCE
-			Out.MRT[0] = half4(0.0, 0.0, 0.0, DualBlendSurfaceCoverage);
+		Out.MRT[0] = half4(0.0, 0.0, 0.0, DualBlendSurfaceCoverage);
 	#elif SUBSTRATE_BLENDING_TRANSLUCENT_GREYTRANSMITTANCE
-			Out.MRT[0] = half4(0.0, 0.0, 0.0, DualBlendSurfaceCoverage);
+		Out.MRT[0] = half4(0.0, 0.0, 0.0, DualBlendSurfaceCoverage);
 	#endif
 			
 #else
@@ -2289,11 +2285,10 @@ void FPixelShaderInOut_MainPS(
 #endif
 
 #if USES_GBUFFER
-			Out.MRT[1] = 0.0;
-			Out.MRT[2] = 0.0;
-			Out.MRT[3] = 0.0;
+		Out.MRT[1] = 0.0;
+		Out.MRT[2] = 0.0;
+		Out.MRT[3] = 0.0;
 #endif
-		}
 	}
 
 #if MATERIALBLENDING_ANY_TRANSLUCENT || MATERIALBLENDING_ALPHACOMPOSITE || MATERIALBLENDING_ALPHAHOLDOUT
diff --git a/Engine/Shaders/Private/PathTracing/PathTracingMaterialHitShader.usf b/Engine/Shaders/Private/PathTracing/PathTracingMaterialHitShader.usf
--- a/Engine/Shaders/Private/PathTracing/PathTracingMaterialHitShader.usf
+++ b/Engine/Shaders/Private/PathTracing/PathTracingMaterialHitShader.usf
@@ -1191,7 +1191,7 @@ RAY_TRACING_ENTRY_CLOSEST_HIT(PathTracingMaterialCHS,
 #else
 	if (PackedPayload.IsCameraRay())
 	{
-		if ((PrimitiveFlags & PRIMITIVE_SCENE_DATA_FLAG_HOLDOUT) != 0)
+		if ((PrimitiveFlags & PRIMITIVE_SCENE_DATA_FLAG_HOLDOUT) != 0 && ResolvedView.bAllowPrimitiveAlphaHoldout)
 		{
 			if (Payload.Ior != 0.0)
 			{
diff --git a/Engine/Shaders/Private/RayTracing/RayTracingMaterialHitShaders.usf b/Engine/Shaders/Private/RayTracing/RayTracingMaterialHitShaders.usf
--- a/Engine/Shaders/Private/RayTracing/RayTracingMaterialHitShaders.usf
+++ b/Engine/Shaders/Private/RayTracing/RayTracingMaterialHitShaders.usf
@@ -611,7 +611,7 @@ RAY_TRACING_ENTRY_CLOSEST_HIT(MaterialCHS,
 #if MATERIALBLENDING_ALPHAHOLDOUT
 	const bool bIsHoldout = true;
 #else
-	const bool bIsHoldout = (PrimitiveFlags & PRIMITIVE_SCENE_DATA_FLAG_HOLDOUT) != 0;
+	const bool bIsHoldout = (PrimitiveFlags & PRIMITIVE_SCENE_DATA_FLAG_HOLDOUT) != 0 && ResolvedView.bAllowPrimitiveAlphaHoldout;
 #endif
 	const bool bIsCameraRay = PackedPayload.IsCameraRay();
 	if (bIsHoldout && bIsCameraRay)
diff --git a/Engine/Source/Editor/DetailCustomizations/Private/SceneCaptureDetails.cpp b/Engine/Source/Editor/DetailCustomizations/Private/SceneCaptureDetails.cpp
--- a/Engine/Source/Editor/DetailCustomizations/Private/SceneCaptureDetails.cpp
+++ b/Engine/Source/Editor/DetailCustomizations/Private/SceneCaptureDetails.cpp
@@ -179,6 +179,7 @@ void FSceneCaptureDetails::CustomizeDetails( IDetailLayoutBuilder& DetailLayout
 	ShowFlagsToAllowForCaptures.Add(FEngineShowFlags::EShowFlag::SF_PathTracing);
 	ShowFlagsToAllowForCaptures.Add(FEngineShowFlags::EShowFlag::SF_Tonemapper);
 	ShowFlagsToAllowForCaptures.Add(FEngineShowFlags::EShowFlag::SF_PostProcessMaterial);
+	ShowFlagsToAllowForCaptures.Add(FEngineShowFlags::EShowFlag::SF_AllowPrimitiveAlphaHoldout);
 
 	// Create array of flag name strings for each group
 	TArray< TArray<FString> > ShowFlagsByGroup;
diff --git a/Engine/Source/Runtime/Engine/Public/SceneView.h b/Engine/Source/Runtime/Engine/Public/SceneView.h
--- a/Engine/Source/Runtime/Engine/Public/SceneView.h
+++ b/Engine/Source/Runtime/Engine/Public/SceneView.h
@@ -1031,6 +1031,7 @@ enum ETranslucencyVolumeCascade
 	VIEW_UNIFORM_BUFFER_MEMBER(float, MaterialMaxEmissiveValue) \
 	VIEW_UNIFORM_BUFFER_MEMBER(int32, PostVolumeUserFlags) \
 	VIEW_UNIFORM_BUFFER_MEMBER(FIntVector4, EnvironmentComponentsFlags) \
+	VIEW_UNIFORM_BUFFER_MEMBER(uint32, bAllowPrimitiveAlphaHoldout) \
 
 /** The uniform shader parameters associated with a view. */
 #define VIEW_UNIFORM_BUFFER_MEMBER(type, identifier) \
diff --git a/Engine/Source/Runtime/Engine/Public/ShowFlagsValues.inl b/Engine/Source/Runtime/Engine/Public/ShowFlagsValues.inl
--- a/Engine/Source/Runtime/Engine/Public/ShowFlagsValues.inl
+++ b/Engine/Source/Runtime/Engine/Public/ShowFlagsValues.inl
@@ -466,5 +466,8 @@ SHOWFLAG_ALWAYS_ACCESSIBLE(Cameras, SFG_Normal, NSLOCTEXT("UnrealEd", "CamerasSF
 /** If false, hair geometry won't be rendered. */
 SHOWFLAG_ALWAYS_ACCESSIBLE(Hair, SFG_Normal, NSLOCTEXT("UnrealEd", "HairSF", "Hair"))
 
+/** Allows primitive alpha holdout rendering. Note that the "Support Primitive Alpha Holdout" project setting first needs to be enabled for the deferred renderer. */
+SHOWFLAG_ALWAYS_ACCESSIBLE(AllowPrimitiveAlphaHoldout, SFG_Normal, NSLOCTEXT("UnrealEd", "AllowPrimitiveAlphaHoldoutSF", "AllowPrimitiveAlphaHoldout"))
+
 #undef SHOWFLAG_ALWAYS_ACCESSIBLE
 #undef SHOWFLAG_FIXED_IN_SHIPPING
diff --git a/Engine/Source/Runtime/Renderer/Private/HeterogeneousVolumes/HeterogeneousVolumesLiveShadingPipeline.cpp b/Engine/Source/Runtime/Renderer/Private/HeterogeneousVolumes/HeterogeneousVolumesLiveShadingPipeline.cpp
--- a/Engine/Source/Runtime/Renderer/Private/HeterogeneousVolumes/HeterogeneousVolumesLiveShadingPipeline.cpp
+++ b/Engine/Source/Runtime/Renderer/Private/HeterogeneousVolumes/HeterogeneousVolumesLiveShadingPipeline.cpp
@@ -1202,7 +1202,7 @@ void RenderSingleScatteringWithLiveShading(
 
 		// Output
 		PassParameters->RWLightingTexture = GraphBuilder.CreateUAV(HeterogeneousVolumeRadianceTexture);
-		if (IsPrimitiveAlphaHoldoutEnabled(GetFeatureLevelShadingPath(View.GetFeatureLevel())))
+		if (IsPrimitiveAlphaHoldoutEnabled(View))
 		{
 			PassParameters->RWHoldoutTexture = GraphBuilder.CreateUAV(HeterogeneousVolumeHoldoutTexture);
 		}
diff --git a/Engine/Source/Runtime/Renderer/Private/PostProcess/TemporalSuperResolution.cpp b/Engine/Source/Runtime/Renderer/Private/PostProcess/TemporalSuperResolution.cpp
--- a/Engine/Source/Runtime/Renderer/Private/PostProcess/TemporalSuperResolution.cpp
+++ b/Engine/Source/Runtime/Renderer/Private/PostProcess/TemporalSuperResolution.cpp
@@ -2267,7 +2267,7 @@ FDefaultTemporalUpscaler::FOutputs AddTemporalSuperResolutionPasses(
 		PassParameters->TileOverscan = TileOverscan;
 		PassParameters->bEnableResurrection = bCanResurrectHistory;
 		PassParameters->bEnableFlickeringHeuristic = FlickeringFramePeriod > 0.0f;
-		PassParameters->bPassthroughAlpha = IsPrimitiveAlphaHoldoutEnabled(GetFeatureLevelShadingPath(View.GetFeatureLevel()));
+		PassParameters->bPassthroughAlpha = IsPrimitiveAlphaHoldoutEnabled(View);
 
 		PassParameters->InputTexture = PassInputs.SceneColor.Texture;
 		if (PassInputs.FlickeringInputTexture.IsValid())
diff --git a/Engine/Source/Runtime/Renderer/Private/SceneRendering.cpp b/Engine/Source/Runtime/Renderer/Private/SceneRendering.cpp
--- a/Engine/Source/Runtime/Renderer/Private/SceneRendering.cpp
+++ b/Engine/Source/Runtime/Renderer/Private/SceneRendering.cpp
@@ -1579,7 +1579,7 @@ void FViewInfo::SetupUniformBufferParameters(
 			FSkyAtmosphereRenderSceneInfo* SkyAtmosphere = Scene->SkyAtmosphere;
 			const FSkyAtmosphereSceneProxy& SkyAtmosphereSceneProxy = SkyAtmosphere->GetSkyAtmosphereSceneProxy();
 
-			Flags |= SkyAtmosphereSceneProxy.IsHoldout() ? ENVCOMP_FLAG_SKYATMOSPHERE_HOLDOUT : 0;
+			Flags |= (SkyAtmosphereSceneProxy.IsHoldout() && Family->EngineShowFlags.AllowPrimitiveAlphaHoldout) ? ENVCOMP_FLAG_SKYATMOSPHERE_HOLDOUT : 0;
 			Flags |= SkyAtmosphereSceneProxy.IsRenderedInMainPass() ? ENVCOMP_FLAG_SKYATMOSPHERE_RENDERINMAIN : 0;
 		}
 
@@ -1588,7 +1588,7 @@ void FViewInfo::SetupUniformBufferParameters(
 			FVolumetricCloudRenderSceneInfo* VolumetricCloud = Scene->VolumetricCloud;
 			const FVolumetricCloudSceneProxy& VolumetricCloudSceneProxy = VolumetricCloud->GetVolumetricCloudSceneProxy();
 
-			Flags |= VolumetricCloudSceneProxy.bHoldout ? ENVCOMP_FLAG_VOLUMETRICCLOUD_HOLDOUT : 0;
+			Flags |= (VolumetricCloudSceneProxy.bHoldout && Family->EngineShowFlags.AllowPrimitiveAlphaHoldout) ? ENVCOMP_FLAG_VOLUMETRICCLOUD_HOLDOUT : 0;
 			Flags |= VolumetricCloudSceneProxy.bRenderInMainPass ? ENVCOMP_FLAG_VOLUMETRICCLOUD_RENDERINMAIN : 0;
 		}
 
@@ -1596,7 +1596,7 @@ void FViewInfo::SetupUniformBufferParameters(
 		{
 			FExponentialHeightFogSceneInfo& Fog = Scene->ExponentialFogs[0];
 
-			Flags |= Fog.bHoldout ? ENVCOMP_FLAG_EXPONENTIALFOG_HOLDOUT : 0;
+			Flags |= (Fog.bHoldout && Family->EngineShowFlags.AllowPrimitiveAlphaHoldout) ? ENVCOMP_FLAG_EXPONENTIALFOG_HOLDOUT : 0;
 			Flags |= Fog.bRenderInMainPass ? ENVCOMP_FLAG_EXPONENTIALFOG_RENDERINMAIN : 0;
 		}
 
@@ -1876,6 +1876,8 @@ void FViewInfo::SetupUniformBufferParameters(
 	ViewUniformShaderParameters.RealTimeReflectionCapture = 0.0f;
 	ViewUniformShaderParameters.RealTimeReflectionCapturePreExposure = 1.0f; // This must be 1 for now. If changed, we need to update the SkyLight AverageExposure and take it into account when sampling sky specular and diffuse irradiance.
 
+	ViewUniformShaderParameters.bAllowPrimitiveAlphaHoldout = static_cast<uint32>(Family->EngineShowFlags.AllowPrimitiveAlphaHoldout && !bIsReflectionCapture);
+
 	ViewUniformShaderParameters.AmbientCubemapTint = FinalPostProcessSettings.AmbientCubemapTint;
 	ViewUniformShaderParameters.AmbientCubemapIntensity = FinalPostProcessSettings.AmbientCubemapIntensity;
 
@@ -6704,18 +6706,21 @@ FRDGTextureRef CreateQuarterResolutionDepthMinAndMax(FRDGBuilder& GraphBuilder,
 	return SmallTexture;
 }
 
-bool IsPrimitiveAlphaHoldoutEnabled(EShadingPath ShadingPath)
+bool IsPrimitiveAlphaHoldoutEnabled(const FViewInfo& View)
 {
 	const bool bSupportPrimitiveAlphaHoldout = CVarPrimitiveAlphaHoldoutSupport.GetValueOnRenderThread();
 
-	return bSupportPrimitiveAlphaHoldout && (ShadingPath != EShadingPath::Mobile) && IsPostProcessingWithAlphaChannelSupported();
+	return bSupportPrimitiveAlphaHoldout
+		&& (GetFeatureLevelShadingPath(View.GetFeatureLevel()) != EShadingPath::Mobile) 
+		&& IsPostProcessingWithAlphaChannelSupported()
+		&& View.Family->EngineShowFlags.AllowPrimitiveAlphaHoldout;
 }
 
-bool IsPrimitiveAlphaHoldoutEnabledForAnyView(TArrayView<FViewInfo> Views)
+bool IsPrimitiveAlphaHoldoutEnabledForAnyView(TArrayView<const FViewInfo> Views)
 {
 	for (const FViewInfo& View : Views)
 	{
-		if (IsPrimitiveAlphaHoldoutEnabled(GetFeatureLevelShadingPath(View.GetFeatureLevel())))
+		if (IsPrimitiveAlphaHoldoutEnabled(View))
 		{
 			return true;
 		}
diff --git a/Engine/Source/Runtime/Renderer/Private/SceneRendering.h b/Engine/Source/Runtime/Renderer/Private/SceneRendering.h
--- a/Engine/Source/Runtime/Renderer/Private/SceneRendering.h
+++ b/Engine/Source/Runtime/Renderer/Private/SceneRendering.h
@@ -2973,10 +2973,10 @@ inline const FViewFamilyInfo& GetViewFamilyInfo(const TArrayView<const FViewInfo
 }
 
 /** Checks whether primitive alpha holdout is active.*/
-bool IsPrimitiveAlphaHoldoutEnabled(EShadingPath ShadingPath);
+bool IsPrimitiveAlphaHoldoutEnabled(const FViewInfo& View);
 
 /** Checks whether primitive alpha holdout is active for any view.*/
-bool IsPrimitiveAlphaHoldoutEnabledForAnyView(TArrayView<FViewInfo> Views);
+bool IsPrimitiveAlphaHoldoutEnabledForAnyView(TArrayView<const FViewInfo> Views);
 
 bool SceneCaptureRequiresAlphaChannel(const FSceneView& View);
 
diff --git a/Engine/Source/Runtime/Renderer/Private/SceneVisibility.cpp b/Engine/Source/Runtime/Renderer/Private/SceneVisibility.cpp
--- a/Engine/Source/Runtime/Renderer/Private/SceneVisibility.cpp
+++ b/Engine/Source/Runtime/Renderer/Private/SceneVisibility.cpp
@@ -1361,7 +1361,7 @@ void FRelevancePacket::ComputeRelevance(FDynamicPrimitiveIndexList& DynamicPrimi
 	const bool bMobileMaskedInEarlyPass = (ShadingPath == EShadingPath::Mobile) && Scene.EarlyZPassMode == DDM_MaskedOnly;
 	const bool bMobileBasePassAlwaysUsesCSM = (ShadingPath == EShadingPath::Mobile) && MobileBasePassAlwaysUsesCSM(Scene.GetShaderPlatform());
 	const bool bVelocityPassWritesDepth = Scene.EarlyZPassMode == DDM_AllOpaqueNoVelocity;
-	const bool bIsTranslucentHoldoutEnabled = IsPrimitiveAlphaHoldoutEnabled(ShadingPath);
+	const bool bIsTranslucentHoldoutEnabled = IsPrimitiveAlphaHoldoutEnabled(View);
 	const bool bHLODActive = Scene.SceneLODHierarchy.IsActive();
 	const FHLODVisibilityState* const HLODState = bHLODActive && ViewState ? &ViewState->HLODVisibilityState : nullptr;
 	float MaxDrawDistanceScale = GetCachedScalabilityCVars().ViewDistanceScale;
diff --git a/Engine/Source/Runtime/Renderer/Private/TranslucentRendering.cpp b/Engine/Source/Runtime/Renderer/Private/TranslucentRendering.cpp
--- a/Engine/Source/Runtime/Renderer/Private/TranslucentRendering.cpp
+++ b/Engine/Source/Runtime/Renderer/Private/TranslucentRendering.cpp
@@ -621,7 +621,7 @@ FScreenPassTexture FTranslucencyComposition::AddPass(
 		SeparateTranslucencyTexture = GraphBuilder.RegisterExternalTexture(GSystemTextures.BlackAlphaOneDummy);
 	}
 
-	bool bPassthroughAlpha = IsPrimitiveAlphaHoldoutEnabled(GetFeatureLevelShadingPath(View.GetFeatureLevel()));
+	bool bPassthroughAlpha = IsPrimitiveAlphaHoldoutEnabled(View);
 
 	const TCHAR* OpName = nullptr;
 	FRHIBlendState* BlendState = nullptr;
