Holdout: Missed applications of new allow primitive holdout flag.

diff --git a/Engine/Source/Runtime/Renderer/Private/HairStrands/HairStrandsComposition.cpp b/Engine/Source/Runtime/Renderer/Private/HairStrands/HairStrandsComposition.cpp
--- a/Engine/Source/Runtime/Renderer/Private/HairStrands/HairStrandsComposition.cpp
+++ b/Engine/Source/Runtime/Renderer/Private/HairStrands/HairStrandsComposition.cpp
@@ -641,7 +641,7 @@ static void InternalRenderHairComposition(
 
 		FRDGTextureRef AccumulatedColor = nullptr;
 
-		const bool bHasHoldout = HasHairFlags(View.HairStrandsViewData.Flags, HAIR_FLAGS_HOLDOUT);
+		const bool bHasHoldout = HasHairFlags(View.HairStrandsViewData.Flags, HAIR_FLAGS_HOLDOUT) && View.Family->EngineShowFlags.AllowPrimitiveAlphaHoldout;
 		const bool bTemporalLayeringEnabled = false; // HAIR_TODO: remove
 
 		AddHairVisibilityComposeSamplePass(
diff --git a/Engine/Source/Runtime/Renderer/Private/HeterogeneousVolumes/HeterogeneousVolumes.cpp b/Engine/Source/Runtime/Renderer/Private/HeterogeneousVolumes/HeterogeneousVolumes.cpp
--- a/Engine/Source/Runtime/Renderer/Private/HeterogeneousVolumes/HeterogeneousVolumes.cpp
+++ b/Engine/Source/Runtime/Renderer/Private/HeterogeneousVolumes/HeterogeneousVolumes.cpp
@@ -365,8 +365,10 @@ bool ShouldRenderHeterogeneousVolumesAsHoldoutForView(
 	const FViewInfo& View
 )
 {
+	const bool bAllowPrimitiveAlphaHoldout = View.Family->EngineShowFlags.AllowPrimitiveAlphaHoldout;
+
 	// This query returns true if any volume is marked as a holdout; otherwise, the query returns false
-	if (ShouldRenderHeterogeneousVolumesForView(View))
+	if (ShouldRenderHeterogeneousVolumesForView(View) && bAllowPrimitiveAlphaHoldout)
 	{
 		for (int32 MeshBatchIndex = 0; MeshBatchIndex < View.HeterogeneousVolumesMeshBatches.Num(); ++MeshBatchIndex)
 		{
