Substrate: Fix typo in debug visualization

#rnx

diff --git a/Engine/Source/Runtime/Renderer/Private/Substrate/SubstrateVisualize.cpp b/Engine/Source/Runtime/Renderer/Private/Substrate/SubstrateVisualize.cpp
--- a/Engine/Source/Runtime/Renderer/Private/Substrate/SubstrateVisualize.cpp
+++ b/Engine/Source/Runtime/Renderer/Private/Substrate/SubstrateVisualize.cpp
@@ -523,7 +523,7 @@ void AddProcessAndPrintSubstrateMaterialPropertiesPasses(FRDGBuilder& GraphBuild
 					NewLine();
 
 
-					PrintS(DrawPosX + 0.0f, DrawPosY, TEXT("NomralID"));
+					PrintS(DrawPosX + 0.0f, DrawPosY, TEXT("NormalID"));
 					PrintS(DrawPosX + 100.0f, DrawPosY, TEXT("BasisType"));
 					PrintS(DrawPosX + 200.0f, DrawPosY, TEXT("Anisotropy"));
 					PrintS(DrawPosX + 300.0f, DrawPosY, TEXT("TopLayer"));
