Implemented animation bank mesh debug view

diff --git a/Engine/Source/Runtime/Renderer/Private/Skinning/AnimBankMeshDebugView.cpp b/Engine/Source/Runtime/Renderer/Private/Skinning/AnimBankMeshDebugView.cpp
new file mode 100644
--- /dev/null
+++ b/Engine/Source/Runtime/Renderer/Private/Skinning/AnimBankMeshDebugView.cpp
@@ -0,0 +1,219 @@
+// Copyright Epic Games, Inc. All Rights Reserved.
+
+#include "AnimBankMeshDebugView.h"
+#include "Components/AnimBankMeshComponent.h"
+
+#include "GlobalRenderResources.h"
+#include "GlobalShader.h"
+#include "RenderGraphBuilder.h"
+#include "RenderGraphResources.h"
+#include "RenderGraphUtils.h"
+#include "TextureResource.h"
+
+#include "SceneView.h"
+#include "SceneInterface.h"
+#include "ScenePrivate.h"
+
+#include "PostProcess/PostProcessing.h"
+#include "PostProcess/PostProcessMaterialInputs.h"
+#include "Nanite/NaniteSkinningSceneExtension.h"
+
+#include "AnimationRuntime.h"
+
+TAutoConsoleVariable<bool> CVarAnimBankDebugDraw(
+	TEXT("r.AnimBank.DebugDraw"), false,
+	TEXT("Whether to enable animation bank mesh debug draw."),
+	ECVF_RenderThreadSafe
+);
+
+class FAnimBankMeshDrawLineVS : public FGlobalShader
+{
+	DECLARE_GLOBAL_SHADER(FAnimBankMeshDrawLineVS);
+	SHADER_USE_PARAMETER_STRUCT(FAnimBankMeshDrawLineVS, FGlobalShader);
+
+	static bool ShouldCompilePermutation(const FGlobalShaderPermutationParameters& Parameters)
+	{
+		return true;
+	}
+
+	static void ModifyCompilationEnvironment(const FGlobalShaderPermutationParameters& Parameters, FShaderCompilerEnvironment& OutEnvironment)
+	{
+		FGlobalShader::ModifyCompilationEnvironment(Parameters, OutEnvironment);
+		OutEnvironment.SetDefine(TEXT("VF_SUPPORTS_PRIMITIVE_SCENE_DATA"), 1);
+	}
+
+	BEGIN_SHADER_PARAMETER_STRUCT(FParameters, )
+		SHADER_PARAMETER_STRUCT_REF(FViewUniformShaderParameters, View)
+		SHADER_PARAMETER_RDG_UNIFORM_BUFFER(FSceneUniformParameters, Scene)
+		SHADER_PARAMETER(uint32, PersistentPrimitiveIndex)
+	END_SHADER_PARAMETER_STRUCT()
+};
+
+class FAnimBankMeshDrawLinePS : public FGlobalShader
+{
+	DECLARE_GLOBAL_SHADER(FAnimBankMeshDrawLinePS);
+	SHADER_USE_PARAMETER_STRUCT(FAnimBankMeshDrawLinePS, FGlobalShader);
+
+	static bool ShouldCompilePermutation(const FGlobalShaderPermutationParameters& Parameters)
+	{
+		return true;
+	}
+
+	static void ModifyCompilationEnvironment(const FGlobalShaderPermutationParameters& Parameters, FShaderCompilerEnvironment& OutEnvironment)
+	{
+		FGlobalShader::ModifyCompilationEnvironment(Parameters, OutEnvironment);
+	}
+
+	BEGIN_SHADER_PARAMETER_STRUCT(FParameters, )
+		RENDER_TARGET_BINDING_SLOTS()
+	END_SHADER_PARAMETER_STRUCT()
+};
+
+BEGIN_SHADER_PARAMETER_STRUCT(FAnimBankMeshDrawLineParameters, )
+	SHADER_PARAMETER_STRUCT_INCLUDE(FAnimBankMeshDrawLineVS::FParameters, VSParameters)
+	SHADER_PARAMETER_STRUCT_INCLUDE(FAnimBankMeshDrawLinePS::FParameters, PSParameters)
+END_SHADER_PARAMETER_STRUCT()
+
+IMPLEMENT_GLOBAL_SHADER(FAnimBankMeshDrawLineVS, "/Engine/Private/Skinning/AnimBankMeshDebug.usf", "LineVS", SF_Vertex);
+IMPLEMENT_GLOBAL_SHADER(FAnimBankMeshDrawLinePS, "/Engine/Private/Skinning/AnimBankMeshDebug.usf", "LinePS", SF_Pixel);
+
+#if UE_ENABLE_DEBUG_DRAWING
+
+FAnimBankMeshDebugViewExtension::FAnimBankMeshDebugViewExtension(const FAutoRegister& AutoRegister)
+: FSceneViewExtensionBase(AutoRegister)
+{
+}
+
+void FAnimBankMeshDebugViewExtension::SubscribeToPostProcessingPass(EPostProcessingPass PassId, const FSceneView& View, FAfterPassCallbackDelegateArray& InOutPassCallbacks, bool bIsPassEnabled)
+{
+	if (PassId == EPostProcessingPass::VisualizeDepthOfField)
+	{
+		if (CVarAnimBankDebugDraw.GetValueOnAnyThread())
+		{
+			InOutPassCallbacks.Add(FAfterPassCallbackDelegate::CreateRaw(this, &FAnimBankMeshDebugViewExtension::PostProcessPass_RenderThread));
+		}
+	}
+}
+
+FScreenPassTexture FAnimBankMeshDebugViewExtension::PostProcessPass_RenderThread(FRDGBuilder& GraphBuilder, const FSceneView& View, const FPostProcessMaterialInputs& InOutInputs)
+{
+	FScreenPassTextureSlice SceneColorSlice = InOutInputs.GetInput(EPostProcessMaterialInput::SceneColor);
+	check(SceneColorSlice.IsValid());
+
+	FScreenPassTexture SceneColor(SceneColorSlice);
+
+	FScreenPassRenderTarget Output = InOutInputs.OverrideOutput;
+
+	// If the override output is provided, it means that this is the last pass in post processing.
+	if (!Output.IsValid())
+	{
+		Output = FScreenPassRenderTarget(SceneColor.Texture, SceneColor.ViewRect, View.GetOverwriteLoadAction());
+	}
+
+	RenderSkeletons(GraphBuilder, View, Output);
+
+	return MoveTemp(Output);
+}
+
+void FAnimBankMeshDebugViewExtension::RenderSkeletons(FRDGBuilder& GraphBuilder, const FSceneView& View, const FScreenPassRenderTarget& Output)
+{
+	if (View.Family == nullptr || View.Family->Scene == nullptr)
+	{
+		return;
+	}
+
+	const FScene* Scene = View.Family->Scene->GetRenderScene();
+	if (Scene == nullptr)
+	{
+		return;
+	}
+
+	const Nanite::FSkinningSceneExtension* SkinningExtension = Scene->GetExtensionPtr<Nanite::FSkinningSceneExtension>();
+	if (SkinningExtension == nullptr)
+	{
+		return;
+	}
+
+	TArray<FPrimitiveSceneInfo*> SkinnedPrimitives;
+	SkinningExtension->GetSkinnedPrimitives(SkinnedPrimitives);
+
+	TArray<FAnimBankMeshPrimitive> Primitives;
+	Primitives.Reserve(SkinnedPrimitives.Num());
+
+	for (const FPrimitiveSceneInfo* SceneInfo : SkinnedPrimitives)
+	{
+		if (SceneInfo == nullptr || !SceneInfo->GetPersistentIndex().IsValid() || SceneInfo->Proxy == nullptr)
+		{
+			continue;
+		}
+
+		if (!SceneInfo->Proxy->IsSkinnedMesh() || !SceneInfo->Proxy->IsNaniteMesh())
+		{
+			continue;
+		}
+
+		auto* SkinnedProxy = static_cast<Nanite::FSkinnedSceneProxy*>(SceneInfo->Proxy);
+
+		FAnimBankMeshPrimitive& Primitive = Primitives.AddDefaulted_GetRef();
+		Primitive.Index = SceneInfo->GetPersistentIndex();
+		Primitive.BoneCount = SkinnedProxy->GetMaxBoneTransformCount();
+		Primitive.InstanceCount = SceneInfo->GetNumInstanceSceneDataEntries();
+	}
+
+	if (Primitives.Num() == 0)
+	{
+		return;
+	}
+
+	const FGlobalShaderMap* ShaderMap = GetGlobalShaderMap(View.GetFeatureLevel());
+
+	TShaderMapRef<FAnimBankMeshDrawLineVS> VertexShader(ShaderMap);
+	TShaderMapRef<FAnimBankMeshDrawLinePS> PixelShader(ShaderMap);
+
+	const FIntRect ViewRect = Output.ViewRect;
+
+	TRDGUniformBufferRef<FSceneUniformParameters> SceneUniformParameters = GetSceneUniformBufferRef(GraphBuilder, View);
+
+	FAnimBankMeshDrawLineParameters* PassParameters = GraphBuilder.AllocParameters<FAnimBankMeshDrawLineParameters>();
+	PassParameters->VSParameters.View = View.ViewUniformBuffer;
+	PassParameters->VSParameters.Scene = SceneUniformParameters;
+	PassParameters->PSParameters.RenderTargets[0] = Output.GetRenderTargetBinding();
+
+	GraphBuilder.AddPass(
+		RDG_EVENT_NAME("RenderAnimBankDebug"),
+		PassParameters,
+		ERDGPassFlags::Raster,
+		[VertexShader, PixelShader, PassParameters, ViewRect = ViewRect, Primitives = MoveTemp(Primitives)](FRDGAsyncTask, FRHICommandList& RHICmdList)
+		{
+			FGraphicsPipelineStateInitializer GraphicsPSOInit;
+			RHICmdList.ApplyCachedRenderTargets(GraphicsPSOInit);
+			RHICmdList.SetViewport(ViewRect.Min.X, ViewRect.Min.Y, 0.0f, ViewRect.Max.X, ViewRect.Max.Y, 1.0f);
+
+			GraphicsPSOInit.DepthStencilState	= TStaticDepthStencilState<false, CF_Always>::GetRHI();
+			GraphicsPSOInit.BlendState			= TStaticBlendState<CW_RGBA, BO_Add, BF_SourceAlpha, BF_InverseSourceAlpha, BO_Add, BF_One, BF_Zero>::GetRHI();
+			GraphicsPSOInit.RasterizerState		= TStaticRasterizerState<FM_Solid, CM_None>::GetRHI();
+			GraphicsPSOInit.PrimitiveType		= PT_LineList;
+
+			GraphicsPSOInit.BoundShaderState.VertexDeclarationRHI	= GEmptyVertexDeclaration.VertexDeclarationRHI;
+			GraphicsPSOInit.BoundShaderState.VertexShaderRHI		= VertexShader.GetVertexShader();
+			GraphicsPSOInit.BoundShaderState.PixelShaderRHI			= PixelShader.GetPixelShader();
+
+			SetGraphicsPipelineState(RHICmdList, GraphicsPSOInit, 0);
+
+			for (const FAnimBankMeshPrimitive& Primitive : Primitives)
+			{
+				const uint32 BoneCount = Primitive.BoneCount;
+				const uint32 InstanceCount = Primitive.InstanceCount;
+
+				check(Primitive.Index.Index != INDEX_NONE);
+				PassParameters->VSParameters.PersistentPrimitiveIndex = uint32(Primitive.Index.Index);
+
+				SetShaderParameters(RHICmdList, VertexShader, VertexShader.GetVertexShader(), PassParameters->VSParameters);
+				SetShaderParameters(RHICmdList, PixelShader, PixelShader.GetPixelShader(), PassParameters->PSParameters);
+				RHICmdList.DrawPrimitive(0, BoneCount, InstanceCount);
+			}
+		}
+	);
+}
+
+#endif
diff --git a/Engine/Source/Runtime/Renderer/Private/Skinning/AnimBankMeshDebugView.h b/Engine/Source/Runtime/Renderer/Private/Skinning/AnimBankMeshDebugView.h
new file mode 100644
--- /dev/null
+++ b/Engine/Source/Runtime/Renderer/Private/Skinning/AnimBankMeshDebugView.h
@@ -0,0 +1,73 @@
+// Copyright Epic Games, Inc. All Rights Reserved.
+
+#pragma once
+
+#include "CoreMinimal.h"
+#include "SceneViewExtension.h"
+#include "RenderGraphFwd.h"
+#include "MeshPassProcessor.h"
+
+#if UE_ENABLE_DEBUG_DRAWING
+
+struct FAnimBankMeshPrimitive
+{
+	FPersistentPrimitiveIndex Index;
+	uint32 BoneCount = 0;
+	uint32 InstanceCount = 0;
+};
+
+class FAnimBankMeshDebugViewExtension : public FSceneViewExtensionBase
+{
+public:
+	FAnimBankMeshDebugViewExtension(const FAutoRegister& AutoRegister);
+
+	virtual void SetupViewFamily(FSceneViewFamily& InViewFamily) override
+	{
+	}
+
+	/**
+	 * Called on game thread when creating the view.
+	 */
+	virtual void SetupView(FSceneViewFamily& InViewFamily, FSceneView& InView)  override
+	{
+	}
+
+	/**
+	 * Called on game thread when view family is about to be rendered.
+	 */
+	virtual void BeginRenderViewFamily(FSceneViewFamily& InViewFamily) override
+	{
+	}
+
+	virtual void PostRenderBasePassDeferred_RenderThread(FRDGBuilder& GraphBuilder, FSceneView& InView, const FRenderTargetBindingSlots& RenderTargets, TRDGUniformBufferRef < FSceneTextureUniformParameters > SceneTextures) override
+	{
+	}
+
+	virtual void SubscribeToPostProcessingPass(EPostProcessingPass PassId, const FSceneView& View, FAfterPassCallbackDelegateArray& InOutPassCallbacks, bool bIsPassEnabled) override;
+
+	FScreenPassTexture PostProcessPass_RenderThread(FRDGBuilder& GraphBuilder, const FSceneView& View, const struct FPostProcessMaterialInputs& InOutInputs);
+
+	void RenderSkeletons(class FRDGBuilder& GraphBuilder, const class FSceneView& View, const struct FScreenPassRenderTarget& Output);
+
+private:
+	inline static TSharedPtr<FAnimBankMeshDebugViewExtension, ESPMode::ThreadSafe> Instance;
+
+	struct FStaticConstructor
+	{
+		FStaticConstructor()
+		{
+			FCoreDelegates::OnPostEngineInit.AddLambda([]()
+			{ 
+				Instance = FSceneViewExtensions::NewExtension<FAnimBankMeshDebugViewExtension>();
+			});
+			FCoreDelegates::OnEnginePreExit.AddLambda([]()
+			{ 
+				Instance = nullptr;
+			});
+		}
+	};
+
+	static inline FStaticConstructor StaticConstructor;
+};
+
+#endif
