Force HLSL2021 language for anim bank mesh debug view shader

diff --git a/Engine/Source/Runtime/Renderer/Private/Skinning/AnimBankMeshDebugView.cpp b/Engine/Source/Runtime/Renderer/Private/Skinning/AnimBankMeshDebugView.cpp
--- a/Engine/Source/Runtime/Renderer/Private/Skinning/AnimBankMeshDebugView.cpp
+++ b/Engine/Source/Runtime/Renderer/Private/Skinning/AnimBankMeshDebugView.cpp
@@ -40,6 +40,7 @@ class FAnimBankMeshDrawLineVS : public FGlobalShader
 	{
 		FGlobalShader::ModifyCompilationEnvironment(Parameters, OutEnvironment);
 		OutEnvironment.SetDefine(TEXT("VF_SUPPORTS_PRIMITIVE_SCENE_DATA"), 1);
+		OutEnvironment.CompilerFlags.Add(CFLAG_HLSL2021);
 	}
 
 	BEGIN_SHADER_PARAMETER_STRUCT(FParameters, )
@@ -62,6 +63,7 @@ class FAnimBankMeshDrawLinePS : public FGlobalShader
 	static void ModifyCompilationEnvironment(const FGlobalShaderPermutationParameters& Parameters, FShaderCompilerEnvironment& OutEnvironment)
 	{
 		FGlobalShader::ModifyCompilationEnvironment(Parameters, OutEnvironment);
+		OutEnvironment.CompilerFlags.Add(CFLAG_HLSL2021);
 	}
 
 	BEGIN_SHADER_PARAMETER_STRUCT(FParameters, )
