Added view updates and ID to FScene similar to how other primitives are managed.
 - This applies to view states, as state-less views are not able to retain persistent information.

diff --git a/Engine/Source/Runtime/Renderer/Private/RendererScene.cpp b/Engine/Source/Runtime/Renderer/Private/RendererScene.cpp
--- a/Engine/Source/Runtime/Renderer/Private/RendererScene.cpp
+++ b/Engine/Source/Runtime/Renderer/Private/RendererScene.cpp
@@ -187,6 +187,8 @@ TGlobalResource< FGlobalDistanceCullFadeUniformBuffer > GDistanceCullFadedInUnif
 TGlobalResource< FGlobalDitherUniformBuffer > GDitherFadedInUniformBuffer;
 
 static FThreadSafeCounter FSceneViewState_UniqueID;
+// Maintained on the render thread, must not be accessed on any other thread (except worker threads within the scope of a RT command).
+static TSet<uint32> FSceneViewState_ActiveUniqueIDs;
 
 /**
  * Holds the info to update SpeedTree wind per unique tree object in the scene, instead of per instance
@@ -272,6 +274,13 @@ FSceneViewState::FSceneViewState(ERHIFeatureLevel::Type FeatureLevel, FSceneView
 	SetFeatureLevel(FeatureLevel);
 	
 	UniqueID = FSceneViewState_UniqueID.Increment();
+
+	ENQUEUE_RENDER_COMMAND(SceneViewStateAdddUniqueID)(
+		[UniqueID = UniqueID] (FRHICommandListBase&)
+		{
+			FSceneViewState_ActiveUniqueIDs.Add(UniqueID);
+		});	
+
 	Scene = nullptr;
 	OcclusionFrameCounter = 0;
 	LastRenderTime = -FLT_MAX;
@@ -361,6 +370,9 @@ FSceneViewState::FSceneViewState(ERHIFeatureLevel::Type FeatureLevel, FSceneView
 
 FSceneViewState::~FSceneViewState()
 {
+	check(IsInRenderingThread());
+	FSceneViewState_ActiveUniqueIDs.Remove(UniqueID);
+
 	CachedVisibilityChunk = NULL;
 	ShadowOcclusionQueryMaps.Reset();
 
@@ -5436,8 +5448,9 @@ struct FSceneUpdateChangeSetStorage
 		})
 	{
 	}
-	FSceneUpdateChangeSetStorage(FScenePrimitiveUpdates&& InPrimitiveUpdates)
+	FSceneUpdateChangeSetStorage(FScenePrimitiveUpdates&& InPrimitiveUpdates, const FViewSceneChangeSet* InViewUpdateChangeSet)
 		: PrimitiveUpdates(MoveTemp(InPrimitiveUpdates))
+		, ViewUpdateChangeSet(InViewUpdateChangeSet)
 		, PreUpdateChangeSet(FScenePreUpdateChangeSet {
 			TConstArrayView<FPersistentPrimitiveIndex>(RemovedPrimitiveIds),
 			TConstArrayView<FPrimitiveSceneInfo*>(RemovedPrimitiveSceneInfos),
@@ -5451,13 +5464,13 @@ struct FSceneUpdateChangeSetStorage
 	{
 	}
 
-
 	TArray<FPersistentPrimitiveIndex, SceneRenderingAllocator> RemovedPrimitiveIds;
 	TArray<FPrimitiveSceneInfo*,SceneRenderingAllocator> RemovedPrimitiveSceneInfos;
 
 	TArray<FPersistentPrimitiveIndex, SceneRenderingAllocator> AddedPrimitiveIds;
 	TArray<FPrimitiveSceneInfo*, SceneRenderingAllocator> AddedPrimitiveSceneInfos;
 	FScenePrimitiveUpdates PrimitiveUpdates;
+	const FViewSceneChangeSet* ViewUpdateChangeSet = nullptr;
 
 	FScenePreUpdateChangeSet PreUpdateChangeSet;
 	FScenePostUpdateChangeSet PostUpdateChangeSet;
@@ -5466,6 +5479,7 @@ struct FSceneUpdateChangeSetStorage
 	{
 		PreUpdateChangeSet.RemovedPrimitiveIds = TConstArrayView<FPersistentPrimitiveIndex>(RemovedPrimitiveIds);
 		PreUpdateChangeSet.RemovedPrimitiveSceneInfos =	TConstArrayView<FPrimitiveSceneInfo*>(RemovedPrimitiveSceneInfos);
+		PreUpdateChangeSet.ViewUpdateChangeSet = ViewUpdateChangeSet;
 		return PreUpdateChangeSet;
 	}
 
@@ -5473,6 +5487,7 @@ struct FSceneUpdateChangeSetStorage
 	{
 		PostUpdateChangeSet.AddedPrimitiveIds = TConstArrayView<FPersistentPrimitiveIndex>(AddedPrimitiveIds);
 		PostUpdateChangeSet.AddedPrimitiveSceneInfos = TConstArrayView<FPrimitiveSceneInfo*>(AddedPrimitiveSceneInfos);
+		PostUpdateChangeSet.ViewUpdateChangeSet = ViewUpdateChangeSet;
 		return PostUpdateChangeSet;
 	}
 };
@@ -5503,6 +5518,61 @@ private:
 	FScene* Scene = nullptr;
 };
 
+FViewSceneChangeSet* FScene::ProcessViewChanges(FRDGBuilder& GraphBuilder, TArray<FViewInfo*> &Views)
+{
+	FViewSceneChangeSet* ChangeSet = GraphBuilder.AllocObject<FViewSceneChangeSet>();
+	ChangeSet->AddedViewIds.Reserve(Views.Num());
+	ChangeSet->RemovedViewIds.Reserve(PersistentViewStateUniqueIDs.Num());
+
+	// deferred cleanup of any view states that are no longer in the set of allocated IDs.
+	for (auto It = PersistentViewStateUniqueIDs.CreateIterator(); It; ++It)
+	{
+		int32 ViewStateUniqueId = (*It);
+		if (!FSceneViewState_ActiveUniqueIDs.Contains(ViewStateUniqueId))
+		{
+			ChangeSet->RemovedViewIds.Add(FPersistentViewId { It.GetIndex() });
+			It.RemoveCurrent();
+		}
+	}
+
+	// helper to find existing persistent ID
+	auto FindExistingPersistentId = [&](int32 InViewStateUniqueId)
+	{
+		for (auto It = PersistentViewStateUniqueIDs.CreateConstIterator(); It; ++It)
+		{
+			int32 ViewStateUniqueId = (*It);
+			if (ViewStateUniqueId == InViewStateUniqueId)
+			{
+				return FPersistentViewId { It.GetIndex() };
+			}
+		}
+		return FPersistentViewId{};
+	};
+
+
+	for (FViewInfo* View : Views)
+	{
+		FPersistentViewId PersistentViewId;
+		if (View->ViewState != nullptr)
+		{
+			const int32 ViewStateUniqueId = View->ViewState->UniqueID;
+			check(FSceneViewState_ActiveUniqueIDs.Contains(ViewStateUniqueId));
+			PersistentViewId = FindExistingPersistentId(ViewStateUniqueId);
+
+			if (!PersistentViewId.IsValid())
+			{
+				int32 StartIndex = 0;
+				PersistentViewId = FPersistentViewId { PersistentViewStateUniqueIDs.EmplaceAtLowestFreeIndex(StartIndex, ViewStateUniqueId) };
+				ChangeSet->AddedViewIds.Emplace(PersistentViewId); 
+			}
+		}
+		// Note: In a future where view updates are decoupled from the rendering properly, this association should happen as part of the rendering. 
+		View->PersistentViewId = PersistentViewId;
+	}
+
+	return ChangeSet;
+}
+
 void FScene::Update(FRDGBuilder& GraphBuilder, const FUpdateParameters& Parameters)
 {
 	LLM_SCOPE(ELLMTag::SceneRender);
@@ -5557,7 +5627,7 @@ void FScene::Update(FRDGBuilder& GraphBuilder, const FUpdateParameters& Paramete
 	RDG_EVENT_SCOPE(GraphBuilder, "UpdateAllPrimitiveSceneInfos");
 
 	// Allocated with render graph lifetime, safe to reference from RDG tasks.
-	FSceneUpdateChangeSetStorage& SceneUpdateChangeSetStorage = *GraphBuilder.AllocObject<FSceneUpdateChangeSetStorage>(MoveTemp(PrimitiveUpdates));
+	FSceneUpdateChangeSetStorage& SceneUpdateChangeSetStorage = *GraphBuilder.AllocObject<FSceneUpdateChangeSetStorage>(MoveTemp(PrimitiveUpdates), Parameters.ViewUpdateChangeSet);
 	TArray<FPrimitiveSceneInfo*, SceneRenderingAllocator> DeletedPrimitiveSceneInfos;
 	DeletedPrimitiveSceneInfos.Reserve(SceneUpdateChangeSetStorage.PrimitiveUpdates.NumCommands());
 	SceneUpdateChangeSetStorage.AddedPrimitiveSceneInfos.Reserve(SceneUpdateChangeSetStorage.PrimitiveUpdates.NumCommands());
diff --git a/Engine/Source/Runtime/Renderer/Private/ScenePrimitiveUpdates.h b/Engine/Source/Runtime/Renderer/Private/ScenePrimitiveUpdates.h
--- a/Engine/Source/Runtime/Renderer/Private/ScenePrimitiveUpdates.h
+++ b/Engine/Source/Runtime/Renderer/Private/ScenePrimitiveUpdates.h
@@ -5,6 +5,8 @@
 #include "PrimitiveSceneInfo.h"
 #include "GPUSceneWriter.h"
 
+struct FViewSceneChangeSet;
+
 /**
  * Definitions of primitive scene update commands.
  */
@@ -117,6 +119,7 @@ public:
 	TConstArrayView<FPersistentPrimitiveIndex> RemovedPrimitiveIds;
 	TConstArrayView<FPrimitiveSceneInfo*> RemovedPrimitiveSceneInfos;
 	const FScenePrimitiveUpdates &PrimitiveUpdates;
+	const FViewSceneChangeSet* ViewUpdateChangeSet = nullptr;
 };
 
 /**
@@ -129,4 +132,5 @@ public:
 	TConstArrayView<FPersistentPrimitiveIndex> AddedPrimitiveIds;
 	TConstArrayView<FPrimitiveSceneInfo*> AddedPrimitiveSceneInfos;
 	const FScenePrimitiveUpdates &PrimitiveUpdates;
+	const FViewSceneChangeSet* ViewUpdateChangeSet = nullptr;
 };
diff --git a/Engine/Source/Runtime/Renderer/Private/ScenePrivate.h b/Engine/Source/Runtime/Renderer/Private/ScenePrivate.h
--- a/Engine/Source/Runtime/Renderer/Private/ScenePrivate.h
+++ b/Engine/Source/Runtime/Renderer/Private/ScenePrivate.h
@@ -1216,8 +1216,11 @@ public:
 	/** Default constructor. */
 	FSceneViewState(ERHIFeatureLevel::Type FeatureLevel, FSceneViewState* ShareOriginTarget);
 
+private:
+	// May only be called via Destroy (ensuring it happens on the RT)
 	virtual ~FSceneViewState();
 
+public:
 	// called every frame after the view state was updated
 	void UpdateLastRenderTime(const FSceneViewFamily& Family)
 	{
@@ -2898,6 +2901,12 @@ struct FLightSceneChangeSet
 	FSceneLightInfoUpdates *SceneLightInfoUpdates = nullptr;
 };
 
+struct FViewSceneChangeSet
+{
+	TArray<FPersistentViewId, SceneRenderingAllocator> AddedViewIds;
+	TArray<FPersistentViewId> RemovedViewIds;
+};
+
 /** 
  * Renderer scene which is private to the renderer module.
  * Ordinarily this is the renderer version of a UWorld, but an FScene can be created for previewing in editors which don't have a UWorld as well.
@@ -2916,6 +2925,17 @@ public:
 	/** List of view states associated with the scene. */
 	TArray<FSceneViewState*> ViewStates;
 
+	/**
+	 * Unique IDs of view states that are tracked by the (ever been rendered & not deleted), indexed by their persistent ID. 
+	 * The sparse array defines a semi-compact range that can be used to maintain data for each view state in e.g., scene extensions.
+	 */
+	TSparseArray<int32> PersistentViewStateUniqueIDs;
+
+	/** 
+	 * Get the upper bound on the persistent view index. 
+	 */
+	int32 GetMaxPersistentViewId() const { return PersistentViewStateUniqueIDs.GetMaxIndex(); }
+
 	FPersistentUniformBuffers UniformBuffers;
 
 	/** Instancing state buckets.  These are stored on the scene as they are precomputed at FPrimitiveSceneInfo::AddToScene time. */
@@ -3340,6 +3360,7 @@ public:
 
 	struct FUpdateParameters
 	{
+		FViewSceneChangeSet* ViewUpdateChangeSet = nullptr;
 		EUpdateAllPrimitiveSceneInfosAsyncOps AsyncOps = EUpdateAllPrimitiveSceneInfosAsyncOps::None;
 		UE::Tasks::FTask GPUSceneUpdateTaskPrerequisites;
 		bool bDestruction = false;
@@ -3351,6 +3372,12 @@ public:
 		} Callbacks;
 	};
 
+
+	/**
+	 * Determines added & removed view states based on the rendered views passed in & the RemovedViews.
+	 */
+	FViewSceneChangeSet* ProcessViewChanges(FRDGBuilder& GraphBuilder, TArray<FViewInfo*>& AllViews);
+
 	void Update(FRDGBuilder& GraphBuilder, const FUpdateParameters& Parameters);
 
 	// FSceneInterface interface.
diff --git a/Engine/Source/Runtime/Renderer/Private/SceneRendering.cpp b/Engine/Source/Runtime/Renderer/Private/SceneRendering.cpp
--- a/Engine/Source/Runtime/Renderer/Private/SceneRendering.cpp
+++ b/Engine/Source/Runtime/Renderer/Private/SceneRendering.cpp
@@ -4034,6 +4034,9 @@ IVisibilityTaskData* FSceneRenderer::OnRenderBegin(FRDGBuilder& GraphBuilder)
 		GPUSceneUpdateTaskPrerequisites.Trigger();
 	};
 
+	// Note: in the future persistent views should be added/removed as other scene primitives such that the updates are deferred and so on. 
+	//       right now there's no explicit mechanism for this, so we discover added views here & pass the chage set to the scene update.
+	SceneUpdateParameters.ViewUpdateChangeSet = Scene->ProcessViewChanges(GraphBuilder, AllViews);
 	Scene->Update(GraphBuilder, SceneUpdateParameters);
 
 	return VisibilityTaskData;
diff --git a/Engine/Source/Runtime/Renderer/Private/SceneRendering.h b/Engine/Source/Runtime/Renderer/Private/SceneRendering.h
--- a/Engine/Source/Runtime/Renderer/Private/SceneRendering.h
+++ b/Engine/Source/Runtime/Renderer/Private/SceneRendering.h
@@ -101,6 +101,13 @@ namespace Nanite
 	struct FInstanceDraw;
 }
 
+struct FPersistentViewId
+{
+	bool IsValid() const { return Index != INDEX_NONE; }
+	int32 Index = INDEX_NONE;
+};
+
+
 extern bool ShouldUseStereoLumenOptimizations();
 
 DECLARE_GPU_DRAWCALL_STAT_EXTERN(VirtualTextureUpdate);
@@ -1553,6 +1560,10 @@ public:
 	 */
 	int32 GPUSceneViewId;
 
+	/**
+	 */
+	FPersistentViewId PersistentViewId;
+
 	/* View rect for all instanced views laid out side-by-side. Only primary view will have it populated. 
 	 *
 	 * This may be different than FamilySize if we're using adaptive resolution stereo rendering. In that case, FamilySize represents the maximum size of
@@ -1993,6 +2004,9 @@ public:
 	FSceneExtensionsRenderers& GetSceneExtensionsRenderers() { return SceneExtensionsRenderers; }
 	const FSceneExtensionsRenderers& GetSceneExtensionsRenderers() const { return SceneExtensionsRenderers; }
 
+	/** The views being rendered. */
+	TArray<FViewInfo> Views;
+
 private:
 	FSceneUniformBuffer SceneUniforms;
 	FSceneExtensionsRenderers SceneExtensionsRenderers;
@@ -2012,9 +2026,6 @@ public:
 	/** The view family being rendered.  This references the Views array. */
 	FViewFamilyInfo ViewFamily;
 
-	/** The views being rendered. */
-	TArray<FViewInfo> Views;
-
 	/** Information of a custom render pass that renders as part of the main renderer. */
 	struct FCustomRenderPassInfo
 	{
