Added AddPass convenience function to instance culling load balancer.

#rnx

diff --git a/Engine/Source/Runtime/Renderer/Private/InstanceCulling/InstanceCullingLoadBalancer.cpp b/Engine/Source/Runtime/Renderer/Private/InstanceCulling/InstanceCullingLoadBalancer.cpp
--- a/Engine/Source/Runtime/Renderer/Private/InstanceCulling/InstanceCullingLoadBalancer.cpp
+++ b/Engine/Source/Runtime/Renderer/Private/InstanceCulling/InstanceCullingLoadBalancer.cpp
@@ -13,7 +13,7 @@ void FInstanceCullingLoadBalancerBase::SetShaderDefines(FShaderCompilerEnvironme
 	OutEnvironment.SetDefine(TEXT("PREFIX_BIT_MASK"), PrefixBitMask);
 }
 
-void FInstanceCullingLoadBalancerBase::FGPUData::GetShaderParameters(FRDGBuilder& GraphBuilder, FShaderParameters& ShaderParameters)
+void FInstanceCullingLoadBalancerBase::FGPUData::GetShaderParameters(FRDGBuilder& GraphBuilder, FShaderParameters& ShaderParameters) const
 {
 	ShaderParameters.BatchBuffer = GraphBuilder.CreateSRV(BatchBuffer);
 	ShaderParameters.ItemBuffer = GraphBuilder.CreateSRV(ItemBuffer);
diff --git a/Engine/Source/Runtime/Renderer/Private/InstanceCulling/InstanceCullingLoadBalancer.h b/Engine/Source/Runtime/Renderer/Private/InstanceCulling/InstanceCullingLoadBalancer.h
--- a/Engine/Source/Runtime/Renderer/Private/InstanceCulling/InstanceCullingLoadBalancer.h
+++ b/Engine/Source/Runtime/Renderer/Private/InstanceCulling/InstanceCullingLoadBalancer.h
@@ -6,6 +6,7 @@
 #include "RenderGraphDefinitions.h"
 #include "RendererInterface.h"
 #include "ShaderParameterMacros.h"
+#include "RenderGraphUtils.h"
 
 struct FShaderCompilerEnvironment;
 
@@ -99,7 +100,25 @@ public:
 		FRDGBufferRef BatchBuffer = nullptr;
 		FRDGBufferRef ItemBuffer = nullptr;
 
-		void GetShaderParameters(FRDGBuilder& GraphBuilder, FShaderParameters& ShaderParameters);
+		void GetShaderParameters(FRDGBuilder& GraphBuilder, FShaderParameters& ShaderParameters) const;
+
+		/**
+		 * Helper function to add a pass to run a CS using the loadbalancer, automatically sets the parameters and dispatch dimensions.
+		 * Expects "LoadBalancerParameters" in the Parameters struct.
+		 */
+		template <typename TShaderClass>
+		FRDGPassRef AddPass(
+			FRDGBuilder& GraphBuilder,
+			FRDGEventName&& PassName,
+			const TShaderRef<TShaderClass>& ComputeShader,
+			typename TShaderClass::FParameters* Parameters) const
+		{
+			GetShaderParameters(GraphBuilder, Parameters->LoadBalancerParameters);
+			return FComputeShaderUtils::AddPass(GraphBuilder, Forward<FRDGEventName>(PassName), ComputeShader, Parameters, GetWrappedCsGroupCount());
+		}
+
+		FIntVector GetWrappedCsGroupCount() const { return FComputeShaderUtils::GetGroupCountWrapped(NumBatches); }
+
 	};
 
 	FGPUData Upload(FRDGBuilder& GraphBuilder, TConstArrayView<FPackedBatch> Batches, TConstArrayView<FPackedItem> Items, ERDGInitialDataFlags RDGInitialDataFlags, int32 NumGroupsPerBatch) const;
