Temporarily work around mobile platforms failing to cook due to missing batched primitive uniform buffer, despite not being explicitly needed by the anim bank mesh debug view pass.

diff --git a/Engine/Source/Runtime/Renderer/Private/Skinning/AnimBankMeshDebugView.cpp b/Engine/Source/Runtime/Renderer/Private/Skinning/AnimBankMeshDebugView.cpp
--- a/Engine/Source/Runtime/Renderer/Private/Skinning/AnimBankMeshDebugView.cpp
+++ b/Engine/Source/Runtime/Renderer/Private/Skinning/AnimBankMeshDebugView.cpp
@@ -33,7 +33,8 @@ class FAnimBankMeshDrawLineVS : public FGlobalShader
 
 	static bool ShouldCompilePermutation(const FGlobalShaderPermutationParameters& Parameters)
 	{
-		return true;
+		// TODO: Have an animation bank test
+		return DoesPlatformSupportNanite(Parameters.Platform);
 	}
 
 	static void ModifyCompilationEnvironment(const FGlobalShaderPermutationParameters& Parameters, FShaderCompilerEnvironment& OutEnvironment)
@@ -46,6 +47,7 @@ class FAnimBankMeshDrawLineVS : public FGlobalShader
 	BEGIN_SHADER_PARAMETER_STRUCT(FParameters, )
 		SHADER_PARAMETER_STRUCT_REF(FViewUniformShaderParameters, View)
 		SHADER_PARAMETER_RDG_UNIFORM_BUFFER(FSceneUniformParameters, Scene)
+		//SHADER_PARAMETER_RDG_UNIFORM_BUFFER(FBatchedPrimitiveParameters, BatchedPrimitive)
 		SHADER_PARAMETER(uint32, PersistentPrimitiveIndex)
 	END_SHADER_PARAMETER_STRUCT()
 };
@@ -57,7 +59,8 @@ class FAnimBankMeshDrawLinePS : public FGlobalShader
 
 	static bool ShouldCompilePermutation(const FGlobalShaderPermutationParameters& Parameters)
 	{
-		return true;
+		// TODO: Have an animation bank test
+		return DoesPlatformSupportNanite(Parameters.Platform);
 	}
 
 	static void ModifyCompilationEnvironment(const FGlobalShaderPermutationParameters& Parameters, FShaderCompilerEnvironment& OutEnvironment)
@@ -172,6 +175,11 @@ void FAnimBankMeshDebugViewExtension::RenderSkeletons(FRDGBuilder& GraphBuilder,
 	TShaderMapRef<FAnimBankMeshDrawLineVS> VertexShader(ShaderMap);
 	TShaderMapRef<FAnimBankMeshDrawLinePS> PixelShader(ShaderMap);
 
+	if (VertexShader.IsNull() || PixelShader.IsNull())
+	{
+		return;
+	}
+
 	const FIntRect ViewRect = Output.ViewRect;
 
 	TRDGUniformBufferRef<FSceneUniformParameters> SceneUniformParameters = GetSceneUniformBufferRef(GraphBuilder, View);
@@ -181,6 +189,22 @@ void FAnimBankMeshDebugViewExtension::RenderSkeletons(FRDGBuilder& GraphBuilder,
 	PassParameters->VSParameters.Scene = SceneUniformParameters;
 	PassParameters->PSParameters.RenderTargets[0] = Output.GetRenderTargetBinding();
 
+#if 0
+	// TODO: This feels very unnecessary to stub out - Upload Identity parameters
+	FBatchedPrimitiveParameters* BatchedPrimitiveParameters = GraphBuilder.AllocParameters<FBatchedPrimitiveParameters>();
+	{
+		FRDGBufferDesc PrimitiveDataBufferDesc = FRDGBufferDesc::CreateStructuredUploadDesc(16u, (PLATFORM_MAX_UNIFORM_BUFFER_RANGE / 16u));
+		PrimitiveDataBufferDesc.Usage |= EBufferUsageFlags::UniformBuffer;
+
+		FBatchedPrimitiveShaderData ShaderData{};
+		FRDGBufferRef PrimitiveDataBuffer = GraphBuilder.CreateBuffer(PrimitiveDataBufferDesc, TEXT("SinglePrimitiveUniformView"));
+		GraphBuilder.QueueBufferUpload(PrimitiveDataBuffer, ShaderData.Data.GetData(), ShaderData.Data.Num() * sizeof(FVector4f));
+
+		BatchedPrimitiveParameters->Data = GraphBuilder.CreateSRV(PrimitiveDataBuffer);
+	}
+	PassParameters->VSParameters.BatchedPrimitive = GraphBuilder.CreateUniformBuffer(BatchedPrimitiveParameters);
+#endif
+
 	GraphBuilder.AddPass(
 		RDG_EVENT_NAME("RenderAnimBankDebug"),
 		PassParameters,
