Fix separate translucent with per-pixel OIT.

The transmittance was accounted twice due to the OIT combine pass being rendered twice.

diff --git a/Engine/Shaders/Private/OITCombine.usf b/Engine/Shaders/Private/OITCombine.usf
--- a/Engine/Shaders/Private/OITCombine.usf
+++ b/Engine/Shaders/Private/OITCombine.usf
@@ -26,7 +26,14 @@ Texture2D<uint>	SampleCountTexture;
 // Combine
 
 #if SHADER_OIT_COMBINE
+
+
+#if PERMUTATION_MODULATE
+RWTexture2D<float4> OutLightingTexture;
+RWTexture2D<float4> OutTransmittanceTexture;
+#else
 RWTexture2D<float4> OutColorTexture;
+#endif
 
 [numthreads(8, 8, 1)]
 void MainCS(uint3 DispatchThreadId : SV_DispatchThreadID)
@@ -80,10 +87,23 @@ void MainCS(uint3 DispatchThreadId : SV_DispatchThreadID)
 	float AdjustedAlpha = saturate(dot(Transmittance, float3(1.0f, 1.0f, 1.0f) / 3.0f));
 	
 	// Write final color - typed load
+	#if PERMUTATION_MODULATE
+	if (SampleCount > 0)
+	{
+		OutLightingTexture[PixelCoord] = float4(FinalColor, PassType == OIT_PASS_SEPARATE ? AdjustedAlpha : OutLightingTexture[PixelCoord].a);
+		OutTransmittanceTexture[PixelCoord] = float4(Transmittance, PassType == OIT_PASS_SEPARATE ? AdjustedAlpha : OutTransmittanceTexture[PixelCoord].a);
+	}
+	else
+	{
+		// If there is no translucent sample, clear the transmittance buffer to 1
+		OutTransmittanceTexture[PixelCoord] = 1.f;
+	}
+	#else
 	if (SampleCount > 0)
 	{
 		OutColorTexture[PixelCoord] = float4(FinalColor + OutColorTexture[PixelCoord].rgb * Transmittance, PassType == OIT_PASS_SEPARATE ? AdjustedAlpha : OutColorTexture[PixelCoord].a);
 	}
+	#endif
 }
 #endif // SHADER_OIT_COMBINE
 
diff --git a/Engine/Source/Runtime/Renderer/Private/OIT/OIT.cpp b/Engine/Source/Runtime/Renderer/Private/OIT/OIT.cpp
--- a/Engine/Source/Runtime/Renderer/Private/OIT/OIT.cpp
+++ b/Engine/Source/Runtime/Renderer/Private/OIT/OIT.cpp
@@ -171,6 +171,9 @@ class FOITPixelCombineCS : public FGlobalShader
 	DECLARE_GLOBAL_SHADER(FOITPixelCombineCS);
 	SHADER_USE_PARAMETER_STRUCT(FOITPixelCombineCS, FGlobalShader);
 
+	class FModulate : SHADER_PERMUTATION_BOOL("PERMUTATION_MODULATE");
+	using FPermutationDomain = TShaderPermutationDomain<FModulate>;
+
 	BEGIN_SHADER_PARAMETER_STRUCT(FParameters,)
 		SHADER_PARAMETER(FIntPoint, Resolution)
 		SHADER_PARAMETER(uint32, MaxSideSampleCount)
@@ -180,6 +183,8 @@ class FOITPixelCombineCS : public FGlobalShader
 		SHADER_PARAMETER_RDG_TEXTURE(Texture2DArray<uint>, SampleDataTexture)
 		SHADER_PARAMETER_RDG_TEXTURE(Texture2D<uint>, SampleCountTexture)
 		SHADER_PARAMETER_RDG_TEXTURE_UAV(RWTexture2D<float4>, OutColorTexture)
+		SHADER_PARAMETER_RDG_TEXTURE_UAV(RWTexture2D<float4>, OutLightingTexture)
+		SHADER_PARAMETER_RDG_TEXTURE_UAV(RWTexture2D<float4>, OutTransmittanceTexture)
 	END_SHADER_PARAMETER_STRUCT()
 public:
 	static bool ShouldCompilePermutation(const FGlobalShaderPermutationParameters& Parameters) { return IsOITSortedPixelsSupported(Parameters.Platform); }
@@ -197,18 +202,21 @@ static void AddInternalOITComposePass(
 	FRDGBuilder& GraphBuilder,
 	const FViewInfo& View,
 	FOITData& OITData,
-	FRDGTextureRef SceneColorTexture)
+	FRDGTextureRef OutColorTexture,
+	FRDGTextureRef OutLightingTexture,
+	FRDGTextureRef OutTransmittanceTexture)
 {
-	// TODO: Add tile composition based on the coarse translucent raster AABB buffer
+	const bool bValidSceneColor = OutColorTexture && OutColorTexture->Desc.NumSamples == 1;
+	const bool bValidModulateColor = OutLightingTexture && OutLightingTexture->Desc.NumSamples == 1 && OutTransmittanceTexture && OutTransmittanceTexture->Desc.NumSamples == 1;
+	const bool bModulate = bValidModulateColor;
 
-	if (!SceneColorTexture ||
-		 SceneColorTexture->Desc.NumSamples > 1 ||
+	if (!(bValidSceneColor || bValidModulateColor) ||
 		!OITData.SampleDataTexture ||
 		!OITData.SampleCountTexture)
 		return;
 
 	FIntPoint InResolution = OITData.SampleDataTexture->Desc.Extent;
-	FIntPoint OutResolution = SceneColorTexture->Desc.Extent;
+	FIntPoint OutResolution = bModulate ? OutLightingTexture->Desc.Extent : OutColorTexture->Desc.Extent;
 	FIntPoint Resolution = FIntPoint(FMath::Min(InResolution.X, OutResolution.X), FMath::Min(InResolution.Y, OutResolution.Y));
 
 	FOITPixelCombineCS::FParameters* Parameters = GraphBuilder.AllocParameters<FOITPixelCombineCS::FParameters>();
@@ -219,9 +227,18 @@ static void AddInternalOITComposePass(
 	Parameters->MaxSideSampleCount = OITData.MaxSideSamplePerPixel;
 	Parameters->SampleDataTexture = OITData.SampleDataTexture;
 	Parameters->SampleCountTexture = OITData.SampleCountTexture;
-	Parameters->OutColorTexture = GraphBuilder.CreateUAV(SceneColorTexture);
+	if (bModulate)
+	{
+		Parameters->OutLightingTexture = GraphBuilder.CreateUAV(OutLightingTexture);
+		Parameters->OutTransmittanceTexture = GraphBuilder.CreateUAV(OutTransmittanceTexture);
+	}
+	else
+	{
+		Parameters->OutColorTexture = GraphBuilder.CreateUAV(OutColorTexture);
+	}
 
 	FOITPixelCombineCS::FPermutationDomain PermutationVector;
+	PermutationVector.Set<FOITPixelCombineCS::FModulate>(bModulate);
 	TShaderMapRef<FOITPixelCombineCS > ComputeShader(View.ShaderMap, PermutationVector);
 
 	// Add 64 threads permutation
@@ -995,7 +1012,12 @@ namespace OIT
 
 	void AddOITComposePass(FRDGBuilder& GraphBuilder, const FViewInfo& View, FOITData& OITData, FRDGTextureRef SceneColorTexture)
 	{
-		AddInternalOITComposePass(GraphBuilder, View, OITData, SceneColorTexture);
+		AddInternalOITComposePass(GraphBuilder, View, OITData, SceneColorTexture, nullptr, nullptr);
+	}
+
+	void AddOITComposePass(FRDGBuilder& GraphBuilder, const FViewInfo& View, FOITData& OITData, FRDGTextureRef LightingTexture, FRDGTextureRef ModulateTexture)
+	{
+		AddInternalOITComposePass(GraphBuilder, View, OITData, nullptr, LightingTexture, ModulateTexture);
 	}
 
 	void OnRenderBegin(FOITSceneData& OITSceneData)
diff --git a/Engine/Source/Runtime/Renderer/Private/OIT/OIT.h b/Engine/Source/Runtime/Renderer/Private/OIT/OIT.h
--- a/Engine/Source/Runtime/Renderer/Private/OIT/OIT.h
+++ b/Engine/Source/Runtime/Renderer/Private/OIT/OIT.h
@@ -102,6 +102,7 @@ namespace OIT
 
 	/* Compose all OIT samples into the target color buffer */
 	void AddOITComposePass(FRDGBuilder& GraphBuilder, const FViewInfo& View, FOITData& OITData, FRDGTextureRef SceneColorTexture);
+	void AddOITComposePass(FRDGBuilder& GraphBuilder, const FViewInfo& View, FOITData& OITData, FRDGTextureRef LightingTexture, FRDGTextureRef ModulateTexture);
 
 	/* Call on SceneRenderer OnRenderBegin */
 	void OnRenderBegin(FOITSceneData& OITSceneData);
diff --git a/Engine/Source/Runtime/Renderer/Private/TranslucentRendering.cpp b/Engine/Source/Runtime/Renderer/Private/TranslucentRendering.cpp
--- a/Engine/Source/Runtime/Renderer/Private/TranslucentRendering.cpp
+++ b/Engine/Source/Runtime/Renderer/Private/TranslucentRendering.cpp
@@ -1542,31 +1542,82 @@ void FDeferredShadingSceneRenderer::RenderTranslucencyInner(
 
 			FOITData OITData = OIT::CreateOITData(GraphBuilder, View, OITPass_SeperateTranslucency);
 
-			RenderTranslucencyViewInner(
-				GraphBuilder,
-				Renderer,
-				View, ViewIndex,
-				SeparateTranslucencyViewport,
-				SeparateTranslucencyDimensions.Scale,
-				SeparateTranslucencyColorTexture,
-				SeparateTranslucencyColorLoadAction,
-				SeparateTranslucencyDepthTexture.Target,
-				CreateTranslucentBasePassUniformBuffer(GraphBuilder, Scene, View, ViewIndex, TranslucentLightingVolumeTextures, SceneColorCopyTexture, SceneTextureSetupMode, bLumenGIEnabled, OITData, TranslucencyPass),
-				TranslucencyPass,
-				!bCompositeBackToSceneColor,
-				bRenderInParallel,
-				InstanceCullingManager);
-
+			const bool bRenderTranslucencyViewInner = OIT::IsSortedPixelsEnabled(View) ? !bIsModulate : true;
+			if (bRenderTranslucencyViewInner)
 			{
-				FTranslucencyPassResources& TranslucencyPassResources = OutTranslucencyResourceMap->Get(ViewIndex, TranslucencyPass);
-				TranslucencyPassResources.ViewRect = ScaledViewRect;
-				TranslucencyPassResources.ColorTexture = SharedColorTexture;
-				TranslucencyPassResources.DepthTexture = SharedDepthTexture;
+				RenderTranslucencyViewInner(
+					GraphBuilder,
+					Renderer,
+					View, ViewIndex,
+					SeparateTranslucencyViewport,
+					SeparateTranslucencyDimensions.Scale,
+					SeparateTranslucencyColorTexture,
+					SeparateTranslucencyColorLoadAction,
+					SeparateTranslucencyDepthTexture.Target,
+					CreateTranslucentBasePassUniformBuffer(GraphBuilder, Scene, View, ViewIndex, TranslucentLightingVolumeTextures, SceneColorCopyTexture, SceneTextureSetupMode, bLumenGIEnabled, OITData, TranslucencyPass),
+					TranslucencyPass,
+					!bCompositeBackToSceneColor,
+					bRenderInParallel,
+					InstanceCullingManager);
+
+				{
+					FTranslucencyPassResources& TranslucencyPassResources = OutTranslucencyResourceMap->Get(ViewIndex, TranslucencyPass);
+					TranslucencyPassResources.ViewRect = ScaledViewRect;
+					TranslucencyPassResources.ColorTexture = SharedColorTexture;
+					TranslucencyPassResources.DepthTexture = SharedDepthTexture;
+				}
 			}
 
-			if (OITData.PassType & OITPass_SeperateTranslucency)
+			if (OIT::IsSortedPixelsEnabled(View) && !!(OITData.PassType & OITPass_SeperateTranslucency))
 			{
-				OIT::AddOITComposePass(GraphBuilder, View, OITData, SeparateTranslucencyColorTexture.Target);
+				// * If this pass it not modulate (i.e., TranslucencyPass != ETranslucencyPass::TPT_TranslucencyXXXModulate), then: 
+				//   * Either the result is directly composed to the the color target. 
+				//   * Or the luminance is composed to the color target, and the transmittance composed into a new allocated 'modulate' target
+				// * If this pass is modulate (i.e., TranslucencyPass == ETranslucencyPass::TPT_TranslucencyXXXModulate), then we only redirect SharedColorTexture to the modulate target, allocated during the previous pass
+				if (!bIsModulate)
+				{
+					if (bCompositeBackToSceneColor || TranslucencyPass == ETranslucencyPass::TPT_TranslucencyAfterMotionBlur)
+					{
+						// Compose the final result (Luminance + Transmittance) to the color target
+						OIT::AddOITComposePass(GraphBuilder, View, OITData, SeparateTranslucencyColorTexture.Target);
+					}
+					else
+					{
+						ETranslucencyPass::Type ModulateTranslucencyPass = ETranslucencyPass::TPT_MAX;
+						if (TranslucencyPass == ETranslucencyPass::TPT_TranslucencyStandard)
+						{
+							ModulateTranslucencyPass = ETranslucencyPass::TPT_TranslucencyStandardModulate;
+						}
+						else if (TranslucencyPass == ETranslucencyPass::TPT_TranslucencyAfterDOF)
+						{
+							ModulateTranslucencyPass = ETranslucencyPass::TPT_TranslucencyAfterDOFModulate;
+						}
+
+						// * Compose the luminance to the color target
+						// * Allocate a new target for the transmittance
+						// * Compose the transmittance to the transmittance target
+						if (TranslucencyPass != ETranslucencyPass::TPT_MAX)
+						{
+							FRDGTextureMSAA ModulateTexture = CreatePostDOFTranslucentTexture(GraphBuilder, ModulateTranslucencyPass, SeparateTranslucencyDimensions, true /*bIsModulate*/, Scene->GetShaderPlatform());
+							OIT::AddOITComposePass(GraphBuilder, View, OITData, SeparateTranslucencyColorTexture.Target, ModulateTexture.Target);
+
+							FTranslucencyPassResources& ModulateTranslucencyPassResources = OutTranslucencyResourceMap->Get(ViewIndex, ModulateTranslucencyPass);
+							ModulateTranslucencyPassResources.ViewRect = ScaledViewRect;
+							ModulateTranslucencyPassResources.ColorTexture = nullptr;
+							ModulateTranslucencyPassResources.ColorModulateTexture = ModulateTexture;
+							ModulateTranslucencyPassResources.DepthTexture = SharedDepthTexture;
+						}
+					}
+				}
+				else
+				{
+					// Retrieve the modulate texture, which was allocated and fill-in during the previous (non-modulate) transparentcy pass.
+					check(!bCompositeBackToSceneColor);
+					FTranslucencyPassResources& ModulateTranslucencyPassResources = OutTranslucencyResourceMap->Get(ViewIndex, TranslucencyPass);
+					SharedColorTexture = ModulateTranslucencyPassResources.ColorModulateTexture;
+					ModulateTranslucencyPassResources.ViewRect = ScaledViewRect;
+					ModulateTranslucencyPassResources.DepthTexture = SharedDepthTexture;
+				}
 			}
 
 			if (bCompositeBackToSceneColor)
