Substrate - clarified memory reading when closures = 0. (onlythe header is read)

#tests PC dx12 substarte=1
#rnx

diff --git a/Engine/Shaders/Shared/SubstrateVisualizeDefinitions.h b/Engine/Shaders/Shared/SubstrateVisualizeDefinitions.h
--- a/Engine/Shaders/Shared/SubstrateVisualizeDefinitions.h
+++ b/Engine/Shaders/Shared/SubstrateVisualizeDefinitions.h
@@ -544,7 +544,7 @@ FSubstratePixelDebugData ConvertToSerializableSubstratePixelDebugData(uint2 InCo
 	const uint FootPrint_PostBSDFs = SubstrateAddressing.ReadBytes;
 
 	// Output memory reads
-	if (Header.ClosureCount > 0)
+	Data.MemoryDisplayMode = 0;
 	{
 		const bool bHasSSSData = SubstrateSubSurfaceHeaderGetIsValid(SSSHeader);
 		const uint TopLayerDataBytes = 4;
@@ -558,29 +558,32 @@ FSubstratePixelDebugData ConvertToSerializableSubstratePixelDebugData(uint2 InCo
 		const uint BSDFsSize  = FootPrint_PostBSDFs  - FootPrint_PostHeader;
 		const uint TotalSize  = (FootPrint_PostBSDFs - FootPrint_Start) + (bHasSSSData ? SSSDataBytes : 0);
 
-		if (bIsSimpleMaterial || Header.IsSingleLayerWater())
+		if (Header.ClosureCount > 0)
 		{
-			Data.MemoryDisplayMode = 0;
-			Data.MemorySlotA = HeaderSize + BSDFsSize - TopLayerDataBytes;	// Header+BSDF
-			Data.MemorySlotB = TopLayerDataBytes;							// TopNormalTex
-		}
-		else if (bIsSingleMaterial || Header.IsEye() || Header.IsHair())
-		{
-			Data.MemoryDisplayMode = 1;
-			Data.MemorySlotA = HeaderSize - TopLayerDataBytes;				// Header
-			Data.MemorySlotB = TopLayerDataBytes;							// TopNormalTex
-			Data.MemorySlotC = BSDFsSize;									// BSDF
-		}
-		else
-		{
-			Data.MemoryDisplayMode = 2;
-			Data.MemorySlotA = HeaderSize;									// Header+Norm
-			Data.MemorySlotB = BSDFsSize;									// BSDFs
-		}
-		
-		if (bHasSSSData)
-		{
-			Data.MemorySSSData = SSSDataBytes;								// SSS Data
+			if (bIsSimpleMaterial || Header.IsSingleLayerWater())
+			{
+				Data.MemoryDisplayMode = 1;
+				Data.MemorySlotA = HeaderSize + BSDFsSize - TopLayerDataBytes;	// Header+BSDF
+				Data.MemorySlotB = TopLayerDataBytes;							// TopNormalTex
+			}
+			else if (bIsSingleMaterial || Header.IsEye() || Header.IsHair())
+			{
+				Data.MemoryDisplayMode = 2;
+				Data.MemorySlotA = HeaderSize - TopLayerDataBytes;				// Header
+				Data.MemorySlotB = TopLayerDataBytes;							// TopNormalTex
+				Data.MemorySlotC = BSDFsSize;									// BSDF
+			}
+			else
+			{
+				Data.MemoryDisplayMode = 3;
+				Data.MemorySlotA = HeaderSize;									// Header+Norm
+				Data.MemorySlotB = BSDFsSize;									// BSDFs
+			}
+
+			if (bHasSSSData)
+			{
+				Data.MemorySSSData = SSSDataBytes;								// SSS Data
+			}
 		}
 
 		Data.MemoryTotal = TotalSize;
diff --git a/Engine/Source/Runtime/Renderer/Private/Substrate/SubstrateVisualize.cpp b/Engine/Source/Runtime/Renderer/Private/Substrate/SubstrateVisualize.cpp
--- a/Engine/Source/Runtime/Renderer/Private/Substrate/SubstrateVisualize.cpp
+++ b/Engine/Source/Runtime/Renderer/Private/Substrate/SubstrateVisualize.cpp
@@ -848,46 +848,49 @@ void AddProcessAndPrintSubstrateMaterialPropertiesPasses(FRDGBuilder& GraphBuild
 					PrintS(DrawPosX, DrawPosY, TEXT("Memory Transactions"), FLinearColor(0.2f, 1.0f, 0.2f, 1.0f));
 					NewLine();
 
-					if (Data.MemoryDisplayMode == 0)
+					if (Data.MemoryDisplayMode > 0)
 					{
-						PrintS(DrawPosX + 0.0f, DrawPosY, TEXT(" - Header + BSDF"), Orange);
-						PrintI(DrawPosX + DataOffset, DrawPosY, Data.MemorySlotA);
-						NewLine();
-						PrintS(DrawPosX + 0.0f, DrawPosY, TEXT(" - TopNormalTex"), Orange);
-						PrintI(DrawPosX + DataOffset, DrawPosY, Data.MemorySlotB);
-						NewLine();
-					}
-					else if (Data.MemoryDisplayMode == 1)
-					{
-						PrintS(DrawPosX + 0.0f, DrawPosY, TEXT(" - Header"), Orange);
-						PrintI(DrawPosX + DataOffset, DrawPosY, Data.MemorySlotA);
-						NewLine();
-						PrintS(DrawPosX + 0.0f, DrawPosY, TEXT(" - TopNormalTex"), Orange);
-						PrintI(DrawPosX + DataOffset, DrawPosY, Data.MemorySlotB);
-						NewLine();
-						PrintS(DrawPosX + 0.0f, DrawPosY, TEXT(" - BSDF"), Orange);
-						PrintI(DrawPosX + DataOffset, DrawPosY, Data.MemorySlotC);
-						NewLine();
-					}
-					else if (Data.MemoryDisplayMode == 2)
-					{
-						PrintS(DrawPosX + 0.0f, DrawPosY, TEXT(" - Header + Normals"), Orange);
-						PrintI(DrawPosX + DataOffset, DrawPosY, Data.MemorySlotA);
-						NewLine();
-						PrintS(DrawPosX + 0.0f, DrawPosY, TEXT(" - BSDFs"), Orange);
-						PrintI(DrawPosX + DataOffset, DrawPosY, Data.MemorySlotB);
-						NewLine();
-					}
-					else
-					{
-						PrintS(DrawPosX + 0.0f, DrawPosY, TEXT("Unkown Memory Display Mode"), FLinearColor::Red);
-					}
+						if (Data.MemoryDisplayMode == 1)
+						{
+							PrintS(DrawPosX + 0.0f, DrawPosY, TEXT(" - Header + BSDF"), Orange);
+							PrintI(DrawPosX + DataOffset, DrawPosY, Data.MemorySlotA);
+							NewLine();
+							PrintS(DrawPosX + 0.0f, DrawPosY, TEXT(" - TopNormalTex"), Orange);
+							PrintI(DrawPosX + DataOffset, DrawPosY, Data.MemorySlotB);
+							NewLine();
+						}
+						else if (Data.MemoryDisplayMode == 2)
+						{
+							PrintS(DrawPosX + 0.0f, DrawPosY, TEXT(" - Header"), Orange);
+							PrintI(DrawPosX + DataOffset, DrawPosY, Data.MemorySlotA);
+							NewLine();
+							PrintS(DrawPosX + 0.0f, DrawPosY, TEXT(" - TopNormalTex"), Orange);
+							PrintI(DrawPosX + DataOffset, DrawPosY, Data.MemorySlotB);
+							NewLine();
+							PrintS(DrawPosX + 0.0f, DrawPosY, TEXT(" - BSDF"), Orange);
+							PrintI(DrawPosX + DataOffset, DrawPosY, Data.MemorySlotC);
+							NewLine();
+						}
+						else if (Data.MemoryDisplayMode == 3)
+						{
+							PrintS(DrawPosX + 0.0f, DrawPosY, TEXT(" - Header + Normals"), Orange);
+							PrintI(DrawPosX + DataOffset, DrawPosY, Data.MemorySlotA);
+							NewLine();
+							PrintS(DrawPosX + 0.0f, DrawPosY, TEXT(" - BSDFs"), Orange);
+							PrintI(DrawPosX + DataOffset, DrawPosY, Data.MemorySlotB);
+							NewLine();
+						}
+						else
+						{
+							PrintS(DrawPosX + 0.0f, DrawPosY, TEXT("Unkown Memory Display Mode"), FLinearColor::Red);
+						}
 
-					if (Data.MemorySSSData > 0)
-					{
-						PrintS(DrawPosX + 0.0f, DrawPosY, TEXT(" - SSS Data"), Orange);
-						PrintI(DrawPosX + DataOffset, DrawPosY, Data.MemorySSSData);
-						NewLine();
+						if (Data.MemorySSSData > 0)
+						{
+							PrintS(DrawPosX + 0.0f, DrawPosY, TEXT(" - SSS Data"), Orange);
+							PrintI(DrawPosX + DataOffset, DrawPosY, Data.MemorySSSData);
+							NewLine();
+						}
 					}
 
 					PrintS(DrawPosX + 0.0f, DrawPosY, TEXT(" - Total"), Orange);
