Fix crashes in raytracing geometry update and volumetric fog when instanced stereo is used.

These passes were not binding the InstancedView UB.

#rnx

diff --git a/Engine/Source/Runtime/Renderer/Private/RayTracing/RayTracingDynamicGeometry.cpp b/Engine/Source/Runtime/Renderer/Private/RayTracing/RayTracingDynamicGeometry.cpp
--- a/Engine/Source/Runtime/Renderer/Private/RayTracing/RayTracingDynamicGeometry.cpp
+++ b/Engine/Source/Runtime/Renderer/Private/RayTracing/RayTracingDynamicGeometry.cpp
@@ -418,7 +418,7 @@ void FRayTracingDynamicGeometryCollection::AddDynamicMeshBatchForGeometryUpdate(
 BEGIN_SHADER_PARAMETER_STRUCT(FRayTracingDynamicGeometryUpdatePassParams, )
 	RDG_BUFFER_ACCESS(DynamicGeometryScratchBuffer, ERHIAccess::UAVCompute)
 
-	SHADER_PARAMETER_STRUCT_REF(FViewUniformShaderParameters, View)
+	SHADER_PARAMETER_STRUCT_INCLUDE(FViewShaderParameters, View)
 	SHADER_PARAMETER_RDG_UNIFORM_BUFFER(FSceneUniformParameters, Scene)
 END_SHADER_PARAMETER_STRUCT()
 
@@ -559,7 +559,7 @@ void FRayTracingDynamicGeometryCollection::AddDynamicGeometryUpdatePass(const FV
 	}
 
 	FRayTracingDynamicGeometryUpdatePassParams* PassParams = GraphBuilder.AllocParameters<FRayTracingDynamicGeometryUpdatePassParams>();
-	PassParams->View = View.ViewUniformBuffer;
+	PassParams->View = View.GetShaderParameters();
 	PassParams->Scene = View.GetSceneUniforms().GetBuffer(GraphBuilder);
 	PassParams->DynamicGeometryScratchBuffer = OutDynamicGeometryScratchBuffer;	
 
diff --git a/Engine/Source/Runtime/Renderer/Private/VolumetricFogVoxelization.cpp b/Engine/Source/Runtime/Renderer/Private/VolumetricFogVoxelization.cpp
--- a/Engine/Source/Runtime/Renderer/Private/VolumetricFogVoxelization.cpp
+++ b/Engine/Source/Runtime/Renderer/Private/VolumetricFogVoxelization.cpp
@@ -683,7 +683,7 @@ void VoxelizeVolumePrimitive(FVoxelizeVolumeMeshProcessor& PassMeshProcessor,
 }
 
 BEGIN_SHADER_PARAMETER_STRUCT(FVoxelizeVolumePassParameters, )
-	SHADER_PARAMETER_STRUCT_REF(FViewUniformShaderParameters, View)
+	SHADER_PARAMETER_STRUCT_INCLUDE(FViewShaderParameters, View)
 	SHADER_PARAMETER_RDG_UNIFORM_BUFFER(FSceneUniformParameters, Scene)
 	SHADER_PARAMETER_RDG_UNIFORM_BUFFER(FVoxelizeVolumePassUniformParameters, Pass)
 	SHADER_PARAMETER_RDG_UNIFORM_BUFFER(FInstanceCullingGlobalUniforms, InstanceCulling)
@@ -726,7 +726,8 @@ void FSceneRenderer::VoxelizeFogVolumePrimitives(
 				View.PrevViewInfo.ViewMatrices
 			);
 
-			PassParameters->View = TUniformBufferRef<FViewUniformShaderParameters>::CreateUniformBufferImmediate(ViewVoxelizeParameters, UniformBuffer_SingleFrame);
+			PassParameters->View.View = TUniformBufferRef<FViewUniformShaderParameters>::CreateUniformBufferImmediate(ViewVoxelizeParameters, UniformBuffer_SingleFrame);;
+			PassParameters->View.InstancedView = View.GetInstancedViewUniformBuffer();
 		}
 		PassParameters->Scene = GetSceneUniforms().GetBuffer(GraphBuilder);
 
