Change resource debug name from ShadowMaskTexture to ShadowMaskTextureResolve for better debugging

diff --git a/Engine/Source/Runtime/Renderer/Private/LightGridInjection.cpp b/Engine/Source/Runtime/Renderer/Private/LightGridInjection.cpp
--- a/Engine/Source/Runtime/Renderer/Private/LightGridInjection.cpp
+++ b/Engine/Source/Runtime/Renderer/Private/LightGridInjection.cpp
@@ -1311,7 +1311,7 @@ void FDeferredShadingSceneRenderer::RenderForwardShadowProjections(
 		{
 			FRDGTextureDesc Desc(FRDGTextureDesc::Create2D(SceneTextures.Config.Extent, PF_B8G8R8A8, FClearValueBinding::White, TexCreate_RenderTargetable | TexCreate_ShaderResource));
 			Desc.NumSamples = SceneDepthTexture->Desc.NumSamples;
-			ForwardScreenSpaceShadowMask = CreateTextureMSAA(GraphBuilder, Desc, TEXT("ShadowMaskTextureMS"), TEXT("ShadowMaskTexture"), GFastVRamConfig.ScreenSpaceShadowMask);
+			ForwardScreenSpaceShadowMask = CreateTextureMSAA(GraphBuilder, Desc, TEXT("ShadowMaskTextureMS"), TEXT("ShadowMaskTextureResolve"), GFastVRamConfig.ScreenSpaceShadowMask);
 			if (bIsHairEnable)
 			{
 				Desc.NumSamples = 1;
