[HWRT] Fix non-Nanite static meshes missing from Ray Tracing Scene after toggling r.RayTracing.Enable 0/1

Context:
- FPrimitiveSceneInfo::AllocateRayTracingSBT() only succeds when there's a valid FRHIRayTracingGeometry
- When HWRT is toggled off all FRHIRayTracingGeometry are evicted.
- When HWRT is toggled back on, AllocateRayTracingSBT() is called before FRHIRayTracingGeometry are recreated/streamed.
- Proxies using Nanite static meshes already register with FRayTracingGeometryManager so their cached raytracing data is invalidated when the associated FRHIRayTracingGeometry is recreated/streamed.

Fix:
- Register all proxies with valid RayTracing::GeometryGroupHandle with FRayTracingGeometryManager not just Nanite.

diff --git a/Engine/Source/Runtime/Renderer/Private/RendererScene.cpp b/Engine/Source/Runtime/Renderer/Private/RendererScene.cpp
--- a/Engine/Source/Runtime/Renderer/Private/RendererScene.cpp
+++ b/Engine/Source/Runtime/Renderer/Private/RendererScene.cpp
@@ -5896,11 +5896,18 @@ void FScene::Update(FRDGBuilder& GraphBuilder, const FUpdateParameters& Paramete
 				LumenRemovePrimitive(PrimitiveSceneInfo, PrimitiveIndex);
 
 #if RHI_RAYTRACING
-				if (SceneProxy->IsNaniteMesh() && SceneProxy->HasRayTracingRepresentation())
+				if (SceneProxy->HasRayTracingRepresentation())
 				{
-					((FRayTracingGeometryManager*)GRayTracingGeometryManager)->UnregisterProxyWithCachedRayTracingState(SceneProxy, SceneProxy->GetRayTracingGeometryGroupHandle());
+					RayTracing::GeometryGroupHandle RayTracingGeometryGroupHandle = SceneProxy->GetRayTracingGeometryGroupHandle();
+					if (RayTracingGeometryGroupHandle != INDEX_NONE)
+					{
+						((FRayTracingGeometryManager*)GRayTracingGeometryManager)->UnregisterProxyWithCachedRayTracingState(SceneProxy, RayTracingGeometryGroupHandle);
+					}
 
-					Nanite::GRayTracingManager.Remove(PrimitiveSceneInfo);
+					if (SceneProxy->IsNaniteMesh())
+					{
+						Nanite::GRayTracingManager.Remove(PrimitiveSceneInfo);
+					}
 				}
 #endif
 
@@ -6222,11 +6229,19 @@ void FScene::Update(FRDGBuilder& GraphBuilder, const FUpdateParameters& Paramete
 				LumenAddPrimitive(PrimitiveSceneInfo);
 
 #if RHI_RAYTRACING
-				if (SceneProxy->IsNaniteMesh() && SceneProxy->HasRayTracingRepresentation())
+				if (SceneProxy->HasRayTracingRepresentation())
 				{
-					((FRayTracingGeometryManager*)GRayTracingGeometryManager)->RegisterProxyWithCachedRayTracingState(SceneProxy, SceneProxy->GetRayTracingGeometryGroupHandle());
+					// TODO: investigate registering proxy when data is actually cached (for example, in FPrimitiveSceneInfo::CacheRayTracingPrimitives(...))
+					RayTracing::GeometryGroupHandle RayTracingGeometryGroupHandle = SceneProxy->GetRayTracingGeometryGroupHandle();
+					if (RayTracingGeometryGroupHandle != INDEX_NONE)
+					{
+						((FRayTracingGeometryManager*)GRayTracingGeometryManager)->RegisterProxyWithCachedRayTracingState(SceneProxy, RayTracingGeometryGroupHandle);
+					}
 
-					Nanite::GRayTracingManager.Add(PrimitiveSceneInfo);
+					if (SceneProxy->IsNaniteMesh())
+					{
+						Nanite::GRayTracingManager.Add(PrimitiveSceneInfo);
+					}
 				}
 #endif
 
