Refactor TSR into callable passes. First step to enable other passes to use TSR for denoising.

diff --git a/Engine/Source/Runtime/Renderer/Private/PostProcess/PostProcessing.cpp b/Engine/Source/Runtime/Renderer/Private/PostProcess/PostProcessing.cpp
--- a/Engine/Source/Runtime/Renderer/Private/PostProcess/PostProcessing.cpp
+++ b/Engine/Source/Runtime/Renderer/Private/PostProcess/PostProcessing.cpp
@@ -879,7 +879,7 @@ void AddPostProcessingPasses(
 			FDefaultTemporalUpscaler::FOutputs Outputs;
 			if (TAAConfig == EMainTAAPassConfig::TSR)
 			{
-				Outputs = AddTemporalSuperResolutionPasses(
+				Outputs = AddMainTemporalSuperResolutionPasses(
 					GraphBuilder,
 					View,
 					UpscalerPassInputs);
diff --git a/Engine/Source/Runtime/Renderer/Private/PostProcess/TemporalAA.h b/Engine/Source/Runtime/Renderer/Private/PostProcess/TemporalAA.h
--- a/Engine/Source/Runtime/Renderer/Private/PostProcess/TemporalAA.h
+++ b/Engine/Source/Runtime/Renderer/Private/PostProcess/TemporalAA.h
@@ -219,7 +219,61 @@ FDefaultTemporalUpscaler::FOutputs AddGen4MainTemporalAAPasses(
 	const FViewInfo& View,
 	const FDefaultTemporalUpscaler::FInputs& PassInputs);
 
-FDefaultTemporalUpscaler::FOutputs AddTemporalSuperResolutionPasses(
+/** List of TSR configurations. */
+enum class ETSRPassConfig
+{
+	// Permutations for main scene color TSR.
+	Main,				// Temporal accumulation without upscaling
+	MainUpsampling,		// Temporal accumulation with upsampling
+
+	MAX
+};
+
+struct FTSRPassConfig
+{
+	bool  ResurrectionEnable = true;
+	int32 ResurrectionPersistentFrameCount = 2;
+	int32 ResurrectionPersistentFrameInterval = 31;
+
+	int32 AlphaChannel = -1;
+
+	bool  ShadingRejectionFlickering = true;
+	int32 ShadingRejectionFlickeringAdjustToFrameRate = 1;
+	float ShadingRejectionFlickeringFrameRateCap = 60;
+	float ShadingRejectionFlickeringPeriod = 2.0f;
+	float ShadingRejectionFlickeringMaxParallaxVelocity = 10.0f;
+
+	// Spatial antialiasing
+	int32 RejectionAntiAliasingQuality = 3;
+
+	float HistoryRejectionSampleCount = 2.0f;
+	float HistoryScreenPercentage = 100.0f;
+	float HistorySampleCount = 16.0f;
+	int32 HistoryUpdateQuality = 3;
+	int32 HistoryR11G11B10 = 1;
+
+	int32 ReprojectionField = 0;
+	float ReprojectionFieldAntiAliasPixelSpeed = 0.125f;
+
+	float VelocityWeightClampingSampleCount = 4.0f;
+	float VelocityWeightClampingPixelSpeed = 1.0f;
+
+	int32  Visualize = -1;
+
+	ETSRPassConfig Pass = ETSRPassConfig::MainUpsampling;
+};
+
+FTSRPassConfig GetTSRMainPassConfig(const FViewInfo& View);
+
+FDefaultTemporalUpscaler::FOutputs  AddTemporalSuperResolutionPasses(
+	FRDGBuilder& GraphBuilder,
+	const FViewInfo& View,
+	const FDefaultTemporalUpscaler::FInputs& PassInputs,
+	const FTSRPassConfig& PassConfig,
+	const FTSRHistory& InputHistory,
+	FTSRHistory& OutputHistory);
+
+FDefaultTemporalUpscaler::FOutputs AddMainTemporalSuperResolutionPasses(
 	FRDGBuilder& GraphBuilder,
 	const FViewInfo& View,
 	const FDefaultTemporalUpscaler::FInputs& PassInputs);
diff --git a/Engine/Source/Runtime/Renderer/Private/PostProcess/TemporalSuperResolution.cpp b/Engine/Source/Runtime/Renderer/Private/PostProcess/TemporalSuperResolution.cpp
--- a/Engine/Source/Runtime/Renderer/Private/PostProcess/TemporalSuperResolution.cpp
+++ b/Engine/Source/Runtime/Renderer/Private/PostProcess/TemporalSuperResolution.cpp
@@ -1396,22 +1396,81 @@ FScreenPassTexture AddTSRMeasureFlickeringLuma(FRDGBuilder& GraphBuilder, FGloba
 	return FlickeringLuma;
 }
 
-FDefaultTemporalUpscaler::FOutputs AddTemporalSuperResolutionPasses(
+FDefaultTemporalUpscaler::FOutputs AddMainTemporalSuperResolutionPasses(
 	FRDGBuilder& GraphBuilder,
 	const FViewInfo& View,
 	const FDefaultTemporalUpscaler::FInputs& PassInputs)
 {
+	
+	FTSRPassConfig TSRPassConfig = GetTSRMainPassConfig(View);
 	const FTSRHistory& InputHistory = View.PrevViewInfo.TSRHistory;
+	FTSRHistory& OutputHistory = View.ViewState->PrevFrameViewInfo.TSRHistory;
+
+	return AddTemporalSuperResolutionPasses(
+		GraphBuilder,
+		View,
+		PassInputs,
+		TSRPassConfig,
+		InputHistory,
+		OutputHistory);
+}
+
+FTSRPassConfig GetTSRMainPassConfig(const FViewInfo& View)
+{
+	FTSRPassConfig PassConfig;
+
+	PassConfig.ResurrectionEnable = static_cast<bool>(CVarTSRResurrectionEnable.GetValueOnRenderThread());
+	PassConfig.ResurrectionPersistentFrameCount = CVarTSRResurrectionPersistentFrameCount.GetValueOnRenderThread();
+	PassConfig.ResurrectionPersistentFrameInterval = CVarTSRResurrectionPersistentFrameInterval.GetValueOnRenderThread();
+
+	PassConfig.AlphaChannel = CVarTSRAlphaChannel.GetValueOnRenderThread();
+	PassConfig.ShadingRejectionFlickering = static_cast<bool>(CVarTSRFlickeringEnable.GetValueOnRenderThread());
+	PassConfig.ShadingRejectionFlickeringAdjustToFrameRate = CVarTSRFlickeringAdjustToFrameRate.GetValueOnRenderThread();
+	PassConfig.ShadingRejectionFlickeringFrameRateCap = CVarTSRFlickeringFrameRateCap.GetValueOnRenderThread();
+	PassConfig.ShadingRejectionFlickeringPeriod = CVarTSRFlickeringPeriod.GetValueOnRenderThread();
+	PassConfig.ShadingRejectionFlickeringMaxParallaxVelocity = CVarTSRFlickeringMaxParralaxVelocity.GetValueOnRenderThread();
+
+	PassConfig.RejectionAntiAliasingQuality = CVarTSRRejectionAntiAliasingQuality.GetValueOnRenderThread();
+
+	PassConfig.HistoryRejectionSampleCount = CVarTSRHistoryRejectionSampleCount.GetValueOnRenderThread();
+	PassConfig.HistoryScreenPercentage = CVarTSRHistorySP.GetValueOnRenderThread();
+	PassConfig.HistorySampleCount = CVarTSRHistorySampleCount.GetValueOnRenderThread();
+	PassConfig.HistoryUpdateQuality = CVarTSRHistoryUpdateQuality.GetValueOnRenderThread();
+	PassConfig.HistoryR11G11B10 = CVarTSRR11G11B10History.GetValueOnRenderThread();
 
+	PassConfig.ReprojectionField = CVarTSRReprojectionField.GetValueOnRenderThread();
+	PassConfig.ReprojectionFieldAntiAliasPixelSpeed = CVarTSRReprojectionFieldAntiAliasPixelSpeed.GetValueOnRenderThread();
+
+	PassConfig.VelocityWeightClampingSampleCount = CVarTSRWeightClampingSampleCount.GetValueOnRenderThread();
+	PassConfig.VelocityWeightClampingPixelSpeed = CVarTSRWeightClampingPixelSpeed.GetValueOnRenderThread();
+#if UE_BUILD_OPTIMIZED_SHOWFLAGS
+	PassConfig.Visualize = false;
+#else
+	PassConfig.Visualize = CVarTSRVisualize.GetValueOnRenderThread();
+#endif
+
+	PassConfig.Pass = (View.PrimaryScreenPercentageMethod == EPrimaryScreenPercentageMethod::TemporalUpscale) ? ETSRPassConfig::MainUpsampling : ETSRPassConfig::Main;
+
+	return PassConfig;
+}
+
+FDefaultTemporalUpscaler::FOutputs AddTemporalSuperResolutionPasses(
+	FRDGBuilder& GraphBuilder, 
+	const FViewInfo& View, 
+	const FDefaultTemporalUpscaler::FInputs& PassInputs, 
+	const FTSRPassConfig& PassConfig,
+	const FTSRHistory& InputHistory, 
+	FTSRHistory& OutputHistory)
+{
 	// Number of frames stored in the history.
 	FTSRHistorySliceSequence HistorySliceSequence;
-	if (CVarTSRResurrectionEnable.GetValueOnRenderThread())
+	if (PassConfig.ResurrectionEnable)
 	{
 		HistorySliceSequence.FrameStorageCount = FMath::Clamp(
-			FTSRHistorySliceSequence::kTransientSliceCount + FMath::DivideAndRoundUp(CVarTSRResurrectionPersistentFrameCount.GetValueOnRenderThread(), 2) * 2,
+			FTSRHistorySliceSequence::kTransientSliceCount + FMath::DivideAndRoundUp(PassConfig.ResurrectionPersistentFrameCount, 2) * 2,
 			4,
 			GMaxTextureArrayLayers);
-		HistorySliceSequence.FrameStoragePeriod = FMath::Clamp(CVarTSRResurrectionPersistentFrameInterval.GetValueOnRenderThread() | 0x1, 1, 1024);
+		HistorySliceSequence.FrameStoragePeriod = FMath::Clamp(PassConfig.ResurrectionPersistentFrameInterval | 0x1, 1, 1024);
 	}
 	check(HistorySliceSequence.Check());
 		
@@ -1424,13 +1483,13 @@ FDefaultTemporalUpscaler::FOutputs AddTemporalSuperResolutionPasses(
 	bool bUse16BitVALU = Use16BitVALU(ShaderPlatform);
 
 	// Whether alpha channel is supported.
-	const bool bSupportsAlpha = CVarTSRAlphaChannel.GetValueOnRenderThread() >= 0 ? (CVarTSRAlphaChannel.GetValueOnRenderThread() > 0) : IsPostProcessingWithAlphaChannelSupported();
+	const bool bSupportsAlpha = PassConfig.AlphaChannel >= 0 ? (PassConfig.AlphaChannel > 0) : IsPostProcessingWithAlphaChannelSupported();
 
-	const float RefreshRateToFrameRateCap = (View.Family->Time.GetDeltaRealTimeSeconds() > 0.0f && CVarTSRFlickeringAdjustToFrameRate.GetValueOnRenderThread())
-		? View.Family->Time.GetDeltaRealTimeSeconds() * CVarTSRFlickeringFrameRateCap.GetValueOnRenderThread() : 1.0f;
+	const float RefreshRateToFrameRateCap = (View.Family->Time.GetDeltaRealTimeSeconds() > 0.0f && PassConfig.ShadingRejectionFlickeringAdjustToFrameRate)
+		? View.Family->Time.GetDeltaRealTimeSeconds() * PassConfig.ShadingRejectionFlickeringFrameRateCap : 1.0f;
 
 	// Maximum number sample for each output pixel in the history
-	const float MaxHistorySampleCount = FMath::Clamp(CVarTSRHistorySampleCount.GetValueOnRenderThread(), 8.0f, 32.0f);
+	const float MaxHistorySampleCount = FMath::Clamp(PassConfig.HistorySampleCount, 8.0f, 32.0f);
 
 	// Whether the view is orthographic view
 	const bool bIsOrthoProjection = !View.IsPerspectiveProjection();
@@ -1439,10 +1498,10 @@ FDefaultTemporalUpscaler::FOutputs AddTemporalSuperResolutionPasses(
 	int32 AsyncComputePasses = GSupportsEfficientAsyncCompute ? CVarTSRAsyncCompute.GetValueOnRenderThread() : 0;
 
 	// period at which history changes is considered too distracting.
-	const float FlickeringFramePeriod = CVarTSRFlickeringEnable.GetValueOnRenderThread() ? (CVarTSRFlickeringPeriod.GetValueOnRenderThread() / FMath::Max(RefreshRateToFrameRateCap, 1.0f)) : 0.0f;
+	const float FlickeringFramePeriod = PassConfig.ShadingRejectionFlickering ? (PassConfig.ShadingRejectionFlickeringPeriod / FMath::Max(RefreshRateToFrameRateCap, 1.0f)) : 0.0f;
 
 	// Whether the reprojection field is enabled.
-	const bool bReprojectionField = CVarTSRReprojectionField.GetValueOnRenderThread() != 0;
+	const bool bReprojectionField = PassConfig.ReprojectionField != 0;
 
 	ETSRHistoryFormatBits HistoryFormatBits = ETSRHistoryFormatBits::None;
 	{
@@ -1458,14 +1517,14 @@ FDefaultTemporalUpscaler::FOutputs AddTemporalSuperResolutionPasses(
 	}
 	FTSRHistoryArrayIndices HistoryArrayIndices = TranslateHistoryFormatBitsToArrayIndices(HistoryFormatBits);
 
-	FTSRUpdateHistoryCS::EQuality UpdateHistoryQuality = FTSRUpdateHistoryCS::EQuality(FMath::Clamp(CVarTSRHistoryUpdateQuality.GetValueOnRenderThread(), 0, int32(FTSRUpdateHistoryCS::EQuality::MAX) - 1));
+	FTSRUpdateHistoryCS::EQuality UpdateHistoryQuality = FTSRUpdateHistoryCS::EQuality(FMath::Clamp(PassConfig.HistoryUpdateQuality, 0, int32(FTSRUpdateHistoryCS::EQuality::MAX) - 1));
 
 	bool bIsSeperateTranslucyTexturesValid = PassInputs.PostDOFTranslucencyResources.IsValid();
 
 	EPixelFormat ColorFormat = bSupportsAlpha ? PF_FloatRGBA : PF_FloatR11G11B10;
-	EPixelFormat HistoryColorFormat = (CVarTSRR11G11B10History.GetValueOnRenderThread() != 0 && !bSupportsAlpha) ? PF_FloatR11G11B10 : PF_FloatRGBA;
+	EPixelFormat HistoryColorFormat = (PassConfig.HistoryR11G11B10 != 0 && !bSupportsAlpha) ? PF_FloatR11G11B10 : PF_FloatRGBA;
 
-	int32 RejectionAntiAliasingQuality = FMath::Clamp(CVarTSRRejectionAntiAliasingQuality.GetValueOnRenderThread(), 1, 2);
+	int32 RejectionAntiAliasingQuality = FMath::Clamp(PassConfig.RejectionAntiAliasingQuality, 1, 2);
 	if (UpdateHistoryQuality == FTSRUpdateHistoryCS::EQuality::Low)
 	{
 		RejectionAntiAliasingQuality = 0; 
@@ -1476,7 +1535,7 @@ FDefaultTemporalUpscaler::FOutputs AddTemporalSuperResolutionPasses(
 
 	FIntPoint OutputExtent;
 	FIntRect OutputRect;
-	if (View.PrimaryScreenPercentageMethod == EPrimaryScreenPercentageMethod::TemporalUpscale)
+	if (PassConfig.Pass == ETSRPassConfig::MainUpsampling)
 	{
 		OutputRect.Min = FIntPoint(0, 0);
 		OutputRect.Max = View.GetSecondaryViewRectSize();
@@ -1556,7 +1615,7 @@ FDefaultTemporalUpscaler::FOutputs AddTemporalSuperResolutionPasses(
 	{
 		float MaxHistoryUpscaleFactor = FMath::Max(float(GMaxTextureDimensions) / float(FMath::Max(OutputRect.Width(), OutputRect.Height())), 1.0f);
 
-		float HistoryUpscaleFactor = FMath::Clamp(CVarTSRHistorySP.GetValueOnRenderThread() / 100.0f, 1.0f, 2.0f);
+		float HistoryUpscaleFactor = FMath::Clamp(PassConfig.HistoryScreenPercentage / 100.0f, 1.0f, 2.0f);
 		if (HistoryUpscaleFactor > MaxHistoryUpscaleFactor)
 		{
 			HistoryUpscaleFactor = 1.0f;
@@ -1977,10 +2036,10 @@ FDefaultTemporalUpscaler::FOutputs AddTemporalSuperResolutionPasses(
 		PassParameters->PrevOutputBufferUVMin = CommonParameters.InputInfo.UVViewportBilinearMin - CommonParameters.InputInfo.ExtentInverse;
 		PassParameters->PrevOutputBufferUVMax = CommonParameters.InputInfo.UVViewportBilinearMax + CommonParameters.InputInfo.ExtentInverse;
 		{
-			float FlickeringMaxParralaxVelocity = RefreshRateToFrameRateCap * CVarTSRFlickeringMaxParralaxVelocity.GetValueOnRenderThread() * float(View.ViewRect.Width()) / 1920.0f;
+			float FlickeringMaxParralaxVelocity = RefreshRateToFrameRateCap * PassConfig.ShadingRejectionFlickeringMaxParallaxVelocity * float(View.ViewRect.Width()) / 1920.0f;
 			PassParameters->InvFlickeringMaxParralaxVelocity = 1.0f / FlickeringMaxParralaxVelocity;
 		}
-		PassParameters->ReprojectionFieldAntiAliasVelocityThreshold = FMath::Square(FMath::Max(CVarTSRReprojectionFieldAntiAliasPixelSpeed.GetValueOnRenderThread() / OutputToInputResolutionFraction, 1.0f / 64.0f));
+		PassParameters->ReprojectionFieldAntiAliasVelocityThreshold = FMath::Square(FMath::Max(PassConfig.ReprojectionFieldAntiAliasPixelSpeed / OutputToInputResolutionFraction, 1.0f / 64.0f));
 		PassParameters->bReprojectionField = bReprojectionField;
 		PassParameters->bOutputIsMovingTexture = bOutputIsMovingTexture;
 		
@@ -2473,9 +2532,9 @@ FDefaultTemporalUpscaler::FOutputs AddTemporalSuperResolutionPasses(
 		// All parameters to control the sample count in history.
 		PassParameters->HistorySampleCount = MaxHistorySampleCount / OutputToHistoryResolutionFractionSquare;
 		PassParameters->HistoryHisteresis = 1.0f / PassParameters->HistorySampleCount;
-		PassParameters->WeightClampingRejection = 1.0f - (CVarTSRHistoryRejectionSampleCount.GetValueOnRenderThread() / OutputToHistoryResolutionFractionSquare) * PassParameters->HistoryHisteresis;
-		PassParameters->WeightClampingPixelSpeedAmplitude = FMath::Clamp(1.0f - CVarTSRWeightClampingSampleCount.GetValueOnRenderThread() * PassParameters->HistoryHisteresis, 0.0f, 1.0f);
-		PassParameters->InvWeightClampingPixelSpeed = 1.0f / (CVarTSRWeightClampingPixelSpeed.GetValueOnRenderThread() * OutputToHistoryResolutionFraction);
+		PassParameters->WeightClampingRejection = 1.0f - (PassConfig.HistoryRejectionSampleCount / OutputToHistoryResolutionFractionSquare) * PassParameters->HistoryHisteresis;
+		PassParameters->WeightClampingPixelSpeedAmplitude = FMath::Clamp(1.0f - PassConfig.VelocityWeightClampingSampleCount * PassParameters->HistoryHisteresis, 0.0f, 1.0f);
+		PassParameters->InvWeightClampingPixelSpeed = 1.0f / (PassConfig.VelocityWeightClampingPixelSpeed * OutputToHistoryResolutionFraction);
 		
 		PassParameters->InputToHistoryFactor = float(HistorySize.X) / float(InputRect.Width());
 		PassParameters->InputContributionMultiplier = OutputToHistoryResolutionFractionSquare; 
@@ -2724,7 +2783,6 @@ FDefaultTemporalUpscaler::FOutputs AddTemporalSuperResolutionPasses(
 	// Extract all resources for next frame.
 	if (!View.bStatePrevViewInfoIsReadOnly)
 	{
-		FTSRHistory& OutputHistory = View.ViewState->PrevFrameViewInfo.TSRHistory;
 		OutputHistory.InputViewportRect = InputRect;
 		OutputHistory.OutputViewportRect = FIntRect(FIntPoint(0, 0), HistorySize);
 		OutputHistory.FormatBit = uint32(HistoryFormatBits);
@@ -2826,7 +2884,7 @@ FDefaultTemporalUpscaler::FOutputs AddTemporalSuperResolutionPasses(
 		};
 		static_assert(UE_ARRAY_COUNT(kVisualizationName) == int32(EVisualizeId::MAX), "kVisualizationName doesn't match EVisualizeId");
 
-		const EVisualizeId Visualization = EVisualizeId(FMath::Clamp(CVarTSRVisualize.GetValueOnRenderThread(), int32(EVisualizeId::ReprojectionFieldOverview), int32(EVisualizeId::MAX) - 1));
+		const EVisualizeId Visualization = EVisualizeId(FMath::Clamp(PassConfig.Visualize, int32(EVisualizeId::ReprojectionFieldOverview), int32(EVisualizeId::MAX) - 1));
 		const bool bIsOverviewVisualize =
 			Visualization == EVisualizeId::ShowFlag ||
 			Visualization == EVisualizeId::Overview ||
