Improve SSS quality when checkerboard is enabled and in low TSR resolution.

#tests test on the udn user's project

diff --git a/Engine/Shaders/Private/PostProcessSubsurface.usf b/Engine/Shaders/Private/PostProcessSubsurface.usf
--- a/Engine/Shaders/Private/PostProcessSubsurface.usf
+++ b/Engine/Shaders/Private/PostProcessSubsurface.usf
@@ -591,8 +591,10 @@ float4 SetupSubsurfaceForOnePixel(float2 UVSceneColor)
 	float4 Ret = 0;
 	BRANCH if (bHasSSSDiffusion)
 	{
-		// '1' is lower quality but that is acceptable here
-		SDiffuseAndSpecular DiffuseAndSpecular = ReconstructLighting(UVSceneColor, 1);
+		// '1' is lower quality but that is acceptable here in regular rendering
+		// When rendering resolution is very low, '3' is required to reduce checkerboard pattern.
+		uint CheckerboardReconstructMethod = (SUBSURFACE_RECOMBINE_QUALITY == SUBSURFACE_RECOMBINE_QUALITY_HIGH) ? 3 : 1;
+		SDiffuseAndSpecular DiffuseAndSpecular = ReconstructLighting(UVSceneColor, CheckerboardReconstructMethod);
 
 		Ret.rgb = DiffuseAndSpecular.Diffuse;
 
diff --git a/Engine/Source/Runtime/Renderer/Private/PostProcess/PostProcessSubsurface.cpp b/Engine/Source/Runtime/Renderer/Private/PostProcess/PostProcessSubsurface.cpp
--- a/Engine/Source/Runtime/Renderer/Private/PostProcess/PostProcessSubsurface.cpp
+++ b/Engine/Source/Runtime/Renderer/Private/PostProcess/PostProcessSubsurface.cpp
@@ -542,12 +542,31 @@ public:
 		SHADER_PARAMETER_RDG_TEXTURE_UAV(RWTexture2D, ProfileIdTexture)
 		SHADER_PARAMETER_STRUCT_REF(FSubsurfaceUniformParameters, SubsurfaceUniformParameters)
 	END_SHADER_PARAMETER_STRUCT()
+	
+	// Controls the quality of lighting reconstruction.
+	enum class EQuality : uint32
+	{
+		Low,
+		High,
+		MAX
+	};
+
+	static const TCHAR* GetEventName(EQuality Quality)
+	{
+		static const TCHAR* const kEventNames[] = {
+			TEXT("Low"),
+			TEXT("High"),
+		};
+		static_assert(UE_ARRAY_COUNT(kEventNames) == int32(EQuality::MAX), "Fix me");
+		return kEventNames[int32(Quality)];
+	}
 
 	class FDimensionHalfRes : SHADER_PERMUTATION_BOOL("SUBSURFACE_HALF_RES");
 	class FDimensionCheckerboard : SHADER_PERMUTATION_BOOL("SUBSURFACE_PROFILE_CHECKERBOARD");
 	class FRunningInSeparable : SHADER_PERMUTATION_BOOL("SUBSURFACE_FORCE_SEPARABLE");
 	class FDimensionEnableProfileIDCache : SHADER_PERMUTATION_BOOL("ENABLE_PROFILE_ID_CACHE");
-	using FPermutationDomain = TShaderPermutationDomain<FDimensionHalfRes, FDimensionCheckerboard, FRunningInSeparable, FDimensionEnableProfileIDCache>;
+	class FDimensionQuality : SHADER_PERMUTATION_ENUM_CLASS("SUBSURFACE_RECOMBINE_QUALITY", EQuality);
+	using FPermutationDomain = TShaderPermutationDomain<FDimensionHalfRes, FDimensionCheckerboard, FRunningInSeparable, FDimensionEnableProfileIDCache, FDimensionQuality>;
 };
 
 IMPLEMENT_GLOBAL_SHADER(FSubsurfaceIndirectDispatchSetupCS, "/Engine/Private/PostProcessSubsurface.usf", "SetupIndirectCS", SF_Compute);
@@ -1142,19 +1161,23 @@ void AddSubsurfaceViewPass(
 			PassParameters->SubsurfaceUniformParameters = UniformBuffer;
 			PassParameters->Substrate = Substrate::BindSubstrateGlobalUniformParameters(View);
 
+			SHADER::EQuality Quality = static_cast<SHADER::EQuality>(FSubsurfaceRecombinePS::GetQuality(View));
+
 			SHADER::FPermutationDomain ComputeShaderPermutationVector;
 			ComputeShaderPermutationVector.Set<SHADER::FDimensionHalfRes>(bHalfRes);
 			ComputeShaderPermutationVector.Set<SHADER::FDimensionCheckerboard>(bCheckerboard);
 			ComputeShaderPermutationVector.Set<SHADER::FRunningInSeparable>(bForceRunningInSeparable);
 			ComputeShaderPermutationVector.Set<SHADER::FDimensionEnableProfileIDCache>(bUseProfileIdCache);
+			ComputeShaderPermutationVector.Set<SHADER::FDimensionQuality>(Quality);
 			TShaderMapRef<SHADER> ComputeShader(View.ShaderMap, ComputeShaderPermutationVector);
 
 			FComputeShaderUtils::AddPass(
 				GraphBuilder,
-				RDG_EVENT_NAME("SSS::Setup(%s%s%s) %dx%d",
+				RDG_EVENT_NAME("SSS::Setup(%s%s%s %s) %dx%d",
 					ComputeShaderPermutationVector.Get<SHADER::FDimensionHalfRes>() ? TEXT(" HalfRes") : TEXT(""),
 					ComputeShaderPermutationVector.Get<SHADER::FDimensionCheckerboard>() ? TEXT(" Checkerboard") : TEXT(""),
 					ComputeShaderPermutationVector.Get<SHADER::FRunningInSeparable>() ? TEXT(" RunningInSeparable") : TEXT(""),
+					SHADER::GetEventName(ComputeShaderPermutationVector.Get<SHADER::FDimensionQuality>()),
 					SubsurfaceViewport.Extent.X,
 					SubsurfaceViewport.Extent.Y),
 				ComputeShader,
