[StateStream]
* Added particle system state stream. Feature set is very limited and this is very much WIP
* Added so UParticleSystemComponent can use statestream if enabled

diff --git a/Engine/Source/Runtime/Engine/Classes/Particles/ParticleEmitter.h b/Engine/Source/Runtime/Engine/Classes/Particles/ParticleEmitter.h
--- a/Engine/Source/Runtime/Engine/Classes/Particles/ParticleEmitter.h
+++ b/Engine/Source/Runtime/Engine/Classes/Particles/ParticleEmitter.h
@@ -382,7 +382,7 @@ class UParticleEmitter : public UObject
 	float GetQualityLevelSpawnRateMult();
 
 	/** Returns true if the is emitter has any enabled LODs, false otherwise. */
-	bool HasAnyEnabledLODs()const;
+	ENGINE_API bool HasAnyEnabledLODs() const;
 
 	/** Stat id of this object, 0 if nobody asked for it yet */
 	STAT(mutable TStatId StatID;)
diff --git a/Engine/Source/Runtime/Engine/Classes/Particles/ParticleSystemComponent.h b/Engine/Source/Runtime/Engine/Classes/Particles/ParticleSystemComponent.h
--- a/Engine/Source/Runtime/Engine/Classes/Particles/ParticleSystemComponent.h
+++ b/Engine/Source/Runtime/Engine/Classes/Particles/ParticleSystemComponent.h
@@ -14,6 +14,7 @@
 #include "Runtime/Launch/Resources/Version.h"
 #include "Particles/WorldPSCPool.h"
 #include "Particles/ParticlePerfStats.h"
+#include "StateStream/ParticleSystemStateStreamHandle.h"
 #include "ParticleSystemComponent.generated.h"
 
 class FParticleDynamicData;
@@ -1336,6 +1337,10 @@ private:
 	/** Time in ms since a tick was last performed; used with MinTimeBetweenTicks (on UParticleSystem) to control tick rate */
 	uint32 TimeSinceLastTick;
 
+#if WITH_STATE_STREAM_ACTOR
+	FParticleSystemHandle ParticleSystemHandle;
+#endif
+
 public:
 
 	//~ Begin UActorComponent Interface.
diff --git a/Engine/Source/Runtime/Engine/Private/Particles/ParticleEmitterInstances.cpp b/Engine/Source/Runtime/Engine/Private/Particles/ParticleEmitterInstances.cpp
--- a/Engine/Source/Runtime/Engine/Private/Particles/ParticleEmitterInstances.cpp
+++ b/Engine/Source/Runtime/Engine/Private/Particles/ParticleEmitterInstances.cpp
@@ -40,6 +40,9 @@
 #include "Particles/SubUVAnimation.h"
 #include "Stats/StatsTrace.h"
 
+// Comment this in to debug empty emitter instance templates...
+//#define _PSYSCOMP_DEBUG_INVALID_EMITTER_INSTANCE_TEMPLATES_
+
 /*-----------------------------------------------------------------------------
 FParticlesStatGroup
 -----------------------------------------------------------------------------*/
@@ -1975,7 +1978,7 @@ float FParticleEmitterInstance::Spawn(float DeltaTime)
 
 	bool bProcessSpawnRate = true;
 	bool bProcessBurstList = true;
-	int32 DetailMode = Component.GetCurrentDetailMode();
+	//int32 DetailMode = Component.GetCurrentDetailMode();
 
 	if (SpriteTemplate->QualityLevelSpawnRateScale > 0.0f)
 	{
@@ -3771,7 +3774,10 @@ void FParticleMeshEmitterInstance::GetResourceSizeEx(FResourceSizeEx& Cumulative
  */
 void FParticleMeshEmitterInstance::SetMeshMaterials( const TArray<UMaterialInterface*>& InMaterials )
 {
+	#if !WITH_STATE_STREAM
 	check(IsInGameThread());
+	#endif
+
 	CurrentMaterials = InMaterials;
 }
 
diff --git a/Engine/Source/Runtime/Engine/Private/Particles/ParticleSystemComponent.cpp b/Engine/Source/Runtime/Engine/Private/Particles/ParticleSystemComponent.cpp
--- a/Engine/Source/Runtime/Engine/Private/Particles/ParticleSystemComponent.cpp
+++ b/Engine/Source/Runtime/Engine/Private/Particles/ParticleSystemComponent.cpp
@@ -37,6 +37,7 @@
 #include "Particles/ParticleSystemReplay.h"
 #include "Particles/TypeData/ParticleModuleTypeDataMesh.h"
 #include "SceneInterface.h"
+#include "StateStream/ParticleSystemStateStream.h"
 #include "UObject/FrameworkObjectVersion.h"
 #include "UObject/UObjectIterator.h"
 #include "UnrealEngine.h"
@@ -82,6 +83,10 @@ const FGuid FParticleSystemCustomVersion::GUID(0x4A56EB40, 0x10F511DC, 0x92D3347
 // Register the custom version with core
 FCustomVersionRegistration GRegisterParticleSystemCustomVersion(FParticleSystemCustomVersion::GUID, FParticleSystemCustomVersion::LatestVersion, TEXT("ParticleSystemVer"));
 
+#if WITH_STATE_STREAM_ACTOR
+constexpr bool UseParticleSystemStateStream = true;
+#endif
+
 //////////////////////////////////////////////////////////////////////////
 
 int32 GParticleLODBias = 0;
@@ -1054,6 +1059,14 @@ void UParticleSystemComponent::CreateRenderState_Concurrent(FRegisterComponentCo
 	SCOPE_CYCLE_COUNTER(STAT_ParticleSystemComponent_CreateRenderState_Concurrent);
 	SCOPE_CYCLE_COUNTER(STAT_ParticlesOverview_GT_CNC);
 
+#if WITH_STATE_STREAM_ACTOR
+	if (UseParticleSystemStateStream)
+	{
+		bRenderStateCreated = true;
+		return;
+	}
+#endif
+
 	ForceAsyncWorkCompletion(ENSURE_AND_STALL, false, true);
 	check( GetWorld() );
 	UE_LOG(LogParticles,Verbose,
@@ -3081,6 +3094,18 @@ void UParticleSystemComponent::ActivateSystem(bool bFlagAsJustAttached)
 		Template = nullptr;
 	}
 
+#if WITH_STATE_STREAM_ACTOR
+	if (UseParticleSystemStateStream)
+	{
+		FParticleSystemStaticState Ss;
+		FParticleSystemDynamicState Ds;
+		Ds.SetSystemAsset(Template);
+		Ds.SetTransform(GetOrCreateTransformHandle());
+		ParticleSystemHandle = GetWorld()->GetStateStream<IParticleSystemStateStream>().Game_CreateInstance(Ss, Ds);
+		return;
+	}
+#endif
+
 	bOldPositionValid = false;
 	OldPosition = FVector::ZeroVector;
 	PartSysVelocity = FVector::ZeroVector;
@@ -3330,6 +3355,14 @@ void UParticleSystemComponent::DeactivateSystem()
 	FInGameScopedCycleCounter InGameCycleCounter(World, EInGamePerfTrackers::VFXSignificance, EInGamePerfTrackerThreads::GameThread, bIsManagingSignificance);
 	SCOPE_CYCLE_COUNTER(STAT_ParticlesOverview_GT);
 
+#if WITH_STATE_STREAM_ACTOR
+	if (UseParticleSystemStateStream)
+	{
+		ParticleSystemHandle = {};
+		return;
+	}
+#endif
+
 	if (IsTemplate() == true)
 	{
 		return;
diff --git a/Engine/Source/Runtime/Engine/Public/ParticleHelper.h b/Engine/Source/Runtime/Engine/Public/ParticleHelper.h
--- a/Engine/Source/Runtime/Engine/Public/ParticleHelper.h
+++ b/Engine/Source/Runtime/Engine/Public/ParticleHelper.h
@@ -2267,8 +2267,8 @@ public:
 	}
 
 	/** Custom new/delete with recycling */
-	void* operator new(size_t Size);
-	void operator delete(void *RawMemory, size_t Size);
+	ENGINE_API void* operator new(size_t Size);
+	ENGINE_API void operator delete(void *RawMemory, size_t Size);
 
 	void ClearEmitterDataArray()
 	{
diff --git a/Engine/Source/Runtime/Engine/Public/ParticleSystemSceneProxy.h b/Engine/Source/Runtime/Engine/Public/ParticleSystemSceneProxy.h
--- a/Engine/Source/Runtime/Engine/Public/ParticleSystemSceneProxy.h
+++ b/Engine/Source/Runtime/Engine/Public/ParticleSystemSceneProxy.h
@@ -11,7 +11,6 @@ class FColoredMaterialRenderProxy;
 class FParticleDynamicData;
 class UParticleSystemComponent;
 struct FDynamicEmitterDataBase;
-struct FParticleSystemSceneProxyDesc;
 
 ////////////////////////////////////////////////////////////////////////////////////////////////////
 
@@ -87,7 +86,7 @@ public:
 	 */
 	void ReleaseRenderThreadResources();
 
-	void UpdateData(FParticleDynamicData* NewDynamicData);
+	ENGINE_API void UpdateData(FParticleDynamicData* NewDynamicData);
 	void UpdateData_RenderThread(FParticleDynamicData* NewDynamicData);
 
 	FParticleDynamicData* GetDynamicData()
diff --git a/Engine/Source/Runtime/Engine/Public/StateStream/ParticleSystemStateStream.h b/Engine/Source/Runtime/Engine/Public/StateStream/ParticleSystemStateStream.h
new file mode 100644
--- /dev/null
+++ b/Engine/Source/Runtime/Engine/Public/StateStream/ParticleSystemStateStream.h
@@ -0,0 +1,54 @@
+// Copyright Epic Games, Inc. All Rights Reserved.
+
+#pragma once
+
+#include "ParticleSystemStateStreamHandle.h"
+#include "StateStreamDefinitions.h"
+#include "TransformStateStreamHandle.h"
+#include "ParticleSystemStateStream.generated.h"
+
+class UParticleSystem;
+
+////////////////////////////////////////////////////////////////////////////////////////////////////
+// Static state for mesh instance. Can only be set upon creation
+
+USTRUCT(StateStreamStaticState)
+struct FParticleSystemStaticState
+{
+	GENERATED_USTRUCT_BODY()
+};
+
+
+////////////////////////////////////////////////////////////////////////////////////////////////////
+// Dynamic state for mesh instance. Can be updated inside ticks
+
+USTRUCT(StateStreamDynamicState)
+struct FParticleSystemDynamicState
+{
+	GENERATED_USTRUCT_BODY()
+
+private:
+
+	UPROPERTY()
+	FTransformHandle Transform;
+
+	UPROPERTY()
+	TObjectPtr<UParticleSystem> SystemAsset;
+};
+
+
+////////////////////////////////////////////////////////////////////////////////////////////////////
+// Mesh state stream id used for registering dependencies and find statestream
+
+inline constexpr uint32 ParticleSystemStateStreamId = 4;
+
+
+////////////////////////////////////////////////////////////////////////////////////////////////////
+// Interface for creating mesh instances
+
+class IParticleSystemStateStream
+{
+public:
+	DECLARE_STATESTREAM(ParticleSystem)
+	virtual FParticleSystemHandle Game_CreateInstance(const FParticleSystemStaticState& Ss, const FParticleSystemDynamicState& Ds) = 0;
+};
diff --git a/Engine/Source/Runtime/Engine/Public/StateStream/ParticleSystemStateStreamHandle.h b/Engine/Source/Runtime/Engine/Public/StateStream/ParticleSystemStateStreamHandle.h
new file mode 100644
--- /dev/null
+++ b/Engine/Source/Runtime/Engine/Public/StateStream/ParticleSystemStateStreamHandle.h
@@ -0,0 +1,19 @@
+// Copyright Epic Games, Inc. All Rights Reserved.
+
+#pragma once
+
+#include "StateStreamHandle.h"
+#include "ParticleSystemStateStreamHandle.generated.h"
+
+struct FParticleSystemDynamicState;
+
+////////////////////////////////////////////////////////////////////////////////////////////////////
+// Handle for mesh instance
+
+USTRUCT(StateStreamHandle)
+struct FParticleSystemHandle : public FStateStreamHandle
+{
+	GENERATED_USTRUCT_BODY()
+};
+
+////////////////////////////////////////////////////////////////////////////////////////////////////
diff --git a/Engine/Source/Runtime/Renderer/Private/RendererScene.cpp b/Engine/Source/Runtime/Renderer/Private/RendererScene.cpp
--- a/Engine/Source/Runtime/Renderer/Private/RendererScene.cpp
+++ b/Engine/Source/Runtime/Renderer/Private/RendererScene.cpp
@@ -81,6 +81,7 @@
 #include "ProfilingDebugging/CountersTrace.h"
 #include "StateStream/SkinnedMeshStateStreamImpl.h"
 #include "StateStream/StaticMeshStateStreamImpl.h"
+#include "StateStream/ParticleSystemStateStreamImpl.h"
 #include "StateStreamManagerImpl.h"
 #include "SceneCulling/SceneCulling.h"
 #include "TransformStateStreamImpl.h"
@@ -7162,12 +7163,15 @@ IStateStreamManager* FRendererModule::AllocateStateStream(UWorld* World)
 	FTransformStateStreamImpl& TransformStreamImpl = *new FTransformStateStreamImpl();
 	FStaticMeshStateStreamImpl& StaticMeshStreamImpl = *new FStaticMeshStateStreamImpl(Scene);
 	FSkinnedMeshStateStreamImpl& SkinnedMeshStreamImpl = *new FSkinnedMeshStateStreamImpl(Scene);
+	FParticleSystemStateStreamImpl& ParticleSystemStreamImpl = *new FParticleSystemStateStreamImpl(Scene);
 	Manager->Render_Register(TransformStreamImpl, true);
 	Manager->Render_Register(StaticMeshStreamImpl, true);
 	Manager->Render_Register(SkinnedMeshStreamImpl, true);
+	Manager->Render_Register(ParticleSystemStreamImpl, true);
 
 	Manager->Render_RegisterDependency(StaticMeshStreamImpl, TransformStreamImpl);
 	Manager->Render_RegisterDependency(SkinnedMeshStreamImpl, TransformStreamImpl);
+	Manager->Render_RegisterDependency(ParticleSystemStreamImpl, TransformStreamImpl);
 
 	return Manager;
 }
diff --git a/Engine/Source/Runtime/Renderer/Private/StateStream/ParticleSystemStateStreamImpl.cpp b/Engine/Source/Runtime/Renderer/Private/StateStream/ParticleSystemStateStreamImpl.cpp
new file mode 100644
--- /dev/null
+++ b/Engine/Source/Runtime/Renderer/Private/StateStream/ParticleSystemStateStreamImpl.cpp
@@ -0,0 +1,409 @@
+// Copyright Epic Games, Inc. All Rights Reserved.
+
+#include "ParticleSystemStateStreamImpl.h"
+#include "Engine/World.h"
+#include "ParticleHelper.h"
+#include "ParticleSystemSceneProxy.h"
+#include "Particles/ParticleEmitter.h"
+#include "Particles/ParticleLODLevel.h"
+#include "Particles/ParticleSystem.h"
+#include "ScenePrivate.h"
+
+////////////////////////////////////////////////////////////////////////////////////////////////////
+
+FParticleSystemObject::~FParticleSystemObject()
+{
+	FScene& Scene = (FScene&)PrimitiveSceneData.SceneProxy->GetScene();
+	Scene.RemovePrimitive(&PrimitiveSceneDesc);
+
+	if (TransformObject)
+	{
+		TransformObject->RemoveListener(this);
+		TransformObject = nullptr;
+	}
+}
+
+void FParticleSystemObject::OnTransformObjectDirty()
+{
+	FTransformObject::Info Info = TransformObject->GetInfo();
+	// TODO: Implement this to handle changing visibility and movement
+}
+
+void FParticleSystemObject::InitializeSystem(FParticleSystemSceneProxyDesc& OutDesc, const FParticleSystemStaticState& Ss, const FParticleSystemDynamicState& Ds)
+{
+	if (Template == NULL)
+	{
+		return;
+	}
+
+	// Variables from UParticleSystemComponent and base classes
+	float WarmupTime = Template->WarmupTime;
+	float WarmupTickRate = Template->WarmupTickRate;
+	bool bIsViewRelevanceDirty = true;
+	bool bWasCompleted = false;
+	int32 LODLevel = 0;
+	EDetailMode DetailMode = DM_Medium;
+
+	const int32 GlobalDetailMode = PDM_Medium;//GetCurrentDetailMode();
+	const bool bCanEverRender = true;//CanEverRender();
+
+		
+
+	//simplified version.
+	int32 NumInstances = EmitterInstances.Num();
+	int32 NumEmitters = Template->Emitters.Num();
+	const bool bIsFirstCreate = NumInstances == 0;
+	EmitterInstances.SetNumZeroed(NumEmitters);
+
+	bWasCompleted = bIsFirstCreate ? false : bWasCompleted;
+
+	bool bClearDynamicData = false;
+	int32 PreferredLODLevel = LODLevel;
+	bool bSetLodLevels = LODLevel > 0; //We should set the lod level even when creating all emitters if the requested LOD is not 0. 
+
+	for (int32 Idx = 0; Idx < NumEmitters; Idx++)
+	{
+		UParticleEmitter* Emitter = Template->Emitters[Idx];
+		if (Emitter)
+		{
+			FParticleEmitterInstance* Instance = NumInstances == 0 ? NULL : EmitterInstances[Idx];
+			check(GlobalDetailMode < EParticleDetailMode::PDM_MAX);
+			const bool bDetailModeAllowsRendering = DetailMode <= GlobalDetailMode && (Emitter->DetailModeBitmask & (1 << GlobalDetailMode));
+			const bool bShouldCreateAndOrInit = bDetailModeAllowsRendering && Emitter->HasAnyEnabledLODs() && bCanEverRender;
+
+			if (bShouldCreateAndOrInit)
+			{
+				if (Instance)
+				{
+					Instance->SetHaltSpawning(false);
+					Instance->SetHaltSpawningExternal(false);
+				}
+				else
+				{
+					Instance = Emitter->CreateInstance(*this);
+					EmitterInstances[Idx] = Instance;
+				}
+
+				if (Instance)
+				{
+					Instance->bEnabled = true;
+					Instance->InitParameters(Emitter);
+					Instance->Init();
+
+					UParticleLODLevel* EmitterLODLevel = Emitter->LODLevels[0];//LODIndex];
+					//FMaterialRelevance& LODViewRel = CachedViewRelevanceFlags[LODIndex];
+					Instance->GatherMaterialRelevance(&OutDesc.MaterialRelevance, EmitterLODLevel, GMaxRHIFeatureLevel);
+
+					PreferredLODLevel = FMath::Min(PreferredLODLevel, Emitter->LODLevels.Num());
+					bSetLodLevels |= !bIsFirstCreate;//Only set lod levels if we init any instances and it's not the first creation time.
+				}
+			}
+			else
+			{
+				if (Instance)
+				{
+#if STATS
+					Instance->PreDestructorCall();
+#endif
+					delete Instance;
+					EmitterInstances[Idx] = NULL;
+					bClearDynamicData = true;
+				}
+			}
+		}
+	}
+
+	if (bClearDynamicData)
+	{
+		//ClearDynamicData();
+	}
+
+	if (bSetLodLevels)
+	{
+		if (PreferredLODLevel != LODLevel)
+		{
+			// This should never be higher...
+			check(PreferredLODLevel < LODLevel);
+			LODLevel = PreferredLODLevel;
+		}
+
+		for (int32 Idx = 0; Idx < EmitterInstances.Num(); Idx++)
+		{
+			FParticleEmitterInstance* Instance = EmitterInstances[Idx];
+			// set the LOD levels here
+			if (Instance)
+			{
+				Instance->CurrentLODLevelIndex	= LODLevel;
+
+				// small safety net for OR-11322; can be removed if the ensure never fires after the change in SetTemplate (reset all instances LOD indices to 0)
+				if (Instance->CurrentLODLevelIndex >= Instance->SpriteTemplate->LODLevels.Num())
+				{
+					Instance->CurrentLODLevelIndex = Instance->SpriteTemplate->LODLevels.Num()-1;
+					ensureMsgf(false, TEXT("LOD access out of bounds (OR-11322). Please let olaf.piesche or simon.tovey know."));
+				}
+				Instance->CurrentLODLevel = Instance->SpriteTemplate->LODLevels[Instance->CurrentLODLevelIndex];
+			}
+		}
+	}
+
+	LWCTile = FLargeWorldRenderScalar::GetTileFor(GetComponentTransform().GetLocation());
+}
+
+void FParticleSystemObject::Update()
+{
+	float DeltaTimeTick = 1.0f/30.0f;
+	bool bSuppressSpawning = false;
+
+	// Tick Subemitters.
+	int32 EmitterIndex;
+	//NumSignificantEmitters = 0;
+	for (EmitterIndex = 0; EmitterIndex < EmitterInstances.Num(); EmitterIndex++)
+	{
+		FParticleEmitterInstance* Instance = EmitterInstances[EmitterIndex];
+
+		if (EmitterIndex + 1 < EmitterInstances.Num())
+		{
+			FParticleEmitterInstance* NextInstance = EmitterInstances[EmitterIndex+1];
+			FPlatformMisc::Prefetch(NextInstance);
+		}
+
+		if (Instance && Instance->SpriteTemplate)
+		{
+			check(Instance->SpriteTemplate->LODLevels.Num() > 0);
+
+			UParticleLODLevel* SpriteLODLevel = Instance->SpriteTemplate->GetCurrentLODLevel(Instance);
+			if (SpriteLODLevel && SpriteLODLevel->bEnabled)
+			{
+				Instance->Tick(DeltaTimeTick, bSuppressSpawning);
+
+				Instance->Tick_MaterialOverrides(EmitterIndex);
+				//TotalActiveParticles += Instance->ActiveParticles;
+			}
+		}
+	}
+
+	FParticleDynamicData* ParticleDynamicData = CreateDynamicData(PrimitiveSceneData.SceneProxy->GetScene().GetFeatureLevel());
+	FParticleSystemSceneProxy* Proxy = (FParticleSystemSceneProxy*)PrimitiveSceneData.SceneProxy;
+	Proxy->UpdateData( ParticleDynamicData );
+
+	bJustRegistered = false;
+}
+
+FParticleDynamicData* FParticleSystemObject::CreateDynamicData(ERHIFeatureLevel::Type InFeatureLevel)
+{
+	FParticleDynamicData* ParticleDynamicData = new FParticleDynamicData();
+
+	if (Template)
+	{
+		ParticleDynamicData->SystemPositionForMacroUVs = GetComponentTransform().TransformPosition(Template->MacroUVPosition);
+		ParticleDynamicData->SystemRadiusForMacroUVs = Template->MacroUVRadius;
+	}
+
+	{
+		// Is the particle system allowed to run?
+		if (true)//( bForcedInActive == false )
+		{
+			//SCOPE_CYCLE_COUNTER(STAT_ParticleSystemComponent_CreateDynamicData_Gather);
+			ParticleDynamicData->DynamicEmitterDataArray.Reset();
+			ParticleDynamicData->DynamicEmitterDataArray.Reserve(EmitterInstances.Num());
+
+			//QUICK_SCOPE_CYCLE_COUNTER(STAT_ParticleSystemComponent_GetDynamicData);
+			for (int32 EmitterIndex = 0; EmitterIndex < EmitterInstances.Num(); EmitterIndex++)
+			{
+				FDynamicEmitterDataBase* NewDynamicEmitterData = NULL;
+				FParticleEmitterInstance* EmitterInst = EmitterInstances[EmitterIndex];
+				if (EmitterInst)
+				{
+					//FScopeCycleCounterEmitter AdditionalScope(EmitterInst);
+
+					// Generate the dynamic data for this emitter
+					{
+						//SCOPE_CYCLE_COUNTER(STAT_ParticleSystemComponent_GetDynamicData);
+						bool bIsOwnerSeleted = false;
+						NewDynamicEmitterData = EmitterInst->GetDynamicData(bIsOwnerSeleted, InFeatureLevel);
+					}
+					if( NewDynamicEmitterData != NULL )
+					{
+						NewDynamicEmitterData->bValid = true;
+						ParticleDynamicData->DynamicEmitterDataArray.Add( NewDynamicEmitterData );
+						NewDynamicEmitterData->EmitterIndex = EmitterIndex;
+
+						NewDynamicEmitterData->EmitterIndex = EmitterIndex;
+					}
+				}
+			}
+		}
+	}
+
+	return ParticleDynamicData;
+}
+
+const FTransform& FParticleSystemObject::GetAsyncComponentToWorld() const { return TransformObject->GetInfo().WorldTransform; }
+UObject* FParticleSystemObject::GetDistributionData() const { return nullptr; }
+const FTransform& FParticleSystemObject::GetComponentTransform() const { return TransformObject->GetInfo().WorldTransform; }
+FRotator FParticleSystemObject::GetComponentRotation() const { return {}; }
+const FTransform& FParticleSystemObject::GetComponentToWorld() const { return TransformObject->GetInfo().WorldTransform; }
+const FBoxSphereBounds& FParticleSystemObject::GetBounds() const { return PrimitiveSceneDesc.Bounds; }
+TWeakObjectPtr<UWorld> FParticleSystemObject::GetWeakWorld() const { return {}; }
+bool FParticleSystemObject::HasWorld() const { return true; }
+bool FParticleSystemObject::HasWorldSettings() const { return {}; }
+bool FParticleSystemObject::IsGameWorld() const { return true; }
+float FParticleSystemObject::GetWorldTimeSeconds() const { return {}; }
+float FParticleSystemObject::GetWorldEffectiveTimeDilation() const { return {}; }
+FIntVector FParticleSystemObject::GetWorldOriginLocation() const { return {}; }
+FSceneInterface* FParticleSystemObject::GetScene() const { return {}; }
+bool FParticleSystemObject::GetFloatParameter(const FName InName, float& OutFloat) { return {}; }
+const FVector3f& FParticleSystemObject::GetLWCTile() const { return LWCTile; }
+FString FParticleSystemObject::GetName() const { return {}; }
+FString FParticleSystemObject::GetFullName() const { return {}; }
+FString FParticleSystemObject::GetPathName() const { return {}; }
+bool FParticleSystemObject::IsActive() const { return {}; }
+bool FParticleSystemObject::IsValidLowLevel() const { return {}; }
+TArrayView<const FParticleSysParam> FParticleSystemObject::GetAsyncInstanceParameters() { return {}; }
+int32 FParticleSystemObject::GetCurrentDetailMode() const { return {}; }
+int32 FParticleSystemObject::GetCurrentLODIndex() const { return {}; }
+const FVector& FParticleSystemObject::GetPartSysVelocity() const { return FVector::ZeroVector; }
+const FVector& FParticleSystemObject::GetOldPosition() const { return FVector::ZeroVector; }
+FFXSystem* FParticleSystemObject::GetFXSystem() const { return {}; }
+const UParticleSystem* FParticleSystemObject::GetTemplate() const { return Template; }
+TArrayView<const FParticleSysParam> FParticleSystemObject::GetInstanceParameters() const { return {}; }
+TArrayView<FParticleEmitterInstance*> FParticleSystemObject::GetEmitterInstances() const { return {}; }
+TArrayView<TObjectPtr<UMaterialInterface>> FParticleSystemObject::GetEmitterMaterials() const { return {}; }
+FPrimitiveSceneProxy* FParticleSystemObject::GetSceneProxy() const { return PrimitiveSceneData.SceneProxy; }
+bool FParticleSystemObject::GetIsWarmingUp() const { return {}; }
+bool FParticleSystemObject::GetJustRegistered() const { return bJustRegistered; }
+float FParticleSystemObject::GetWarmupTime() const { return {}; }
+float FParticleSystemObject::GetEmitterDelay() const { return {}; }
+FRandomStream& FParticleSystemObject::GetRandomStream() { return RandomStream; }
+void FParticleSystemObject::SetComponentToWorld(const FTransform& NewComponentToWorld) { }
+void FParticleSystemObject::DeactivateNextTick() { }
+UParticleSystemComponent* FParticleSystemObject::AsComponent() const { return {}; }
+void FParticleSystemObject::ReportEventSpawn(const FName InEventName, const float InEmitterTime, const FVector InLocation, const FVector InVelocity, const TArray<class UParticleModuleEventSendToGame*>& InEventData) { }
+void FParticleSystemObject::ReportEventDeath(const FName InEventName, const float InEmitterTime, const FVector InLocation, const FVector InVelocity, const TArray<class UParticleModuleEventSendToGame*>& InEventData, const float InParticleTime) { }
+void FParticleSystemObject::ReportEventCollision(const FName InEventName, const float InEmitterTime, const FVector InLocation, const FVector InDirection, const FVector InVelocity, const TArray<class UParticleModuleEventSendToGame*>& InEventData, const float InParticleTime, const FVector InNormal, const float InTime, const int32 InItem, const FName InBoneName, UPhysicalMaterial* PhysMat) { }
+void FParticleSystemObject::ReportEventBurst(const FName InEventName, const float InEmitterTime, const int32 ParticleCount, const FVector InLocation, const TArray<class UParticleModuleEventSendToGame*>& InEventData) { }
+TArrayView<FParticleEventSpawnData> FParticleSystemObject::GetSpawnEvents() const { return {}; }
+TArrayView<FParticleEventDeathData> FParticleSystemObject::GetDeathEvents() const { return {}; }
+TArrayView<FParticleEventCollideData> FParticleSystemObject::GetCollisionEvents() const { return {}; }
+TArrayView<FParticleEventBurstData> FParticleSystemObject::GetBurstEvents() const { return {}; }
+TArrayView<FParticleEventKismetData> FParticleSystemObject::GetKismetEvents() const { return {}; }
+
+////////////////////////////////////////////////////////////////////////////////////////////////////
+
+FParticleSystemStateStreamImpl::FParticleSystemStateStreamImpl(FScene* InScene)
+:	Scene(InScene)
+{
+}
+
+void FParticleSystemStateStreamImpl::SetTransformObject(FParticleSystemObject& Object, const FParticleSystemDynamicState& Ds)
+{
+	if (Object.TransformObject)
+	{
+		Object.TransformObject->RemoveListener(&Object);
+		Object.TransformObject = nullptr;
+	}
+
+	const FTransformHandle& TransformHandle = Ds.GetTransform();
+	if (TransformHandle.IsValid())
+	{
+		FTransformObject* TransformObject = static_cast<FTransformObject*>(TransformHandle.Render_GetUserData());
+		check(TransformObject);
+		TransformObject->AddListener(&Object);
+		Object.TransformObject = TransformObject;
+	}
+}
+
+void FParticleSystemStateStreamImpl::Render_OnCreate(const FParticleSystemStaticState& Ss, const FParticleSystemDynamicState& Ds, FParticleSystemObject*& UserData, bool IsDestroyedInSameFrame)
+{
+	check(!UserData);
+
+	FParticleSystemObject& Object = *new FParticleSystemObject();
+	Object.AddRef();
+
+	SetTransformObject(Object, Ds);
+
+	FTransformObject::Info Info = Object.TransformObject->GetInfo();
+	const FTransform& Transform = Info.WorldTransform;
+
+	//FBoxSphereBounds LocalBounds(ForceInitToZero);
+	FBoxSphereBounds LocalBounds(FVector(-1000.0, -1000.0, -1000.0), FVector(1000.0, 1000.0, 1000.0), 1000.0);
+	//if (Mesh)
+	//{
+	//	LocalBounds = Mesh->GetBounds();
+	//}	
+
+	UWorld& World = *Scene->World;
+
+
+	Object.PrimitiveSceneDesc.RenderMatrix = Transform.ToMatrixWithScale();
+	Object.PrimitiveSceneDesc.AttachmentRootPosition = Transform.GetLocation();
+	Object.PrimitiveSceneDesc.PrimitiveSceneData = &Object.PrimitiveSceneData;
+	Object.PrimitiveSceneDesc.LocalBounds = LocalBounds;
+	Object.PrimitiveSceneDesc.Bounds = LocalBounds.TransformBy(Transform);
+
+	FParticleSystemSceneProxyDesc Desc;
+	Desc.SystemAsset = Ds.GetSystemAsset();
+	//Desc.OverrideMaterials = const_cast<TArray<TObjectPtr<UMaterialInterface>>&>(Ds.GetOverrideMaterials());
+	Desc.CustomPrimitiveData = &Object.CustomPrimitiveData;
+	Desc.Scene = Scene;
+	Desc.World = &World;
+	Desc.FeatureLevel = World.GetFeatureLevel();
+	Desc.bIsVisible = Info.bVisible;
+	//Desc.MaterialRelevance.Raw = Ss.GetMaterialRelevance();
+
+	//Desc.DynamicData = CreateDynamicData(Scene->GetFeatureLevel());
+	//Desc.bOnlyOwnerSee = Ds.GetOnlyOwnerSee();
+	//Desc.bOwnerNoSee = Ds.GetOwnerNoSee();
+	//Desc.ActorOwners = Ds.GetOwners();
+
+
+	#if WITH_EDITOR
+	//Desc.TextureStreamingTransformScale = Transform.GetMaximumAxisScale();
+	#endif
+
+	Object.Template = Ds.GetSystemAsset();
+	Object.InitializeSystem(Desc, Ss, Ds);
+
+	Object.PrimitiveSceneData.SceneProxy = new FParticleSystemSceneProxy(Desc);
+	Scene->AddPrimitive(&Object.PrimitiveSceneDesc);
+
+	UserData = &Object;
+	Objects.Add(&Object);
+}
+
+void FParticleSystemStateStreamImpl::Render_OnUpdate(const FParticleSystemStaticState& Ss, const FParticleSystemDynamicState& Ds, FParticleSystemObject*& UserData)
+{
+	if (!UserData)
+	{
+		return;
+	}
+
+	FParticleSystemObject& Object = *UserData;
+
+	if (Ds.TransformModified())
+	{
+		SetTransformObject(Object, Ds);
+	}
+}
+
+void FParticleSystemStateStreamImpl::Render_OnDestroy(const FParticleSystemStaticState& Ss, const FParticleSystemDynamicState& Ds, FParticleSystemObject*& UserData)
+{
+	if (UserData)
+	{
+		Objects.Remove(UserData);
+		UserData->Release();
+		UserData = nullptr;
+	}
+}
+
+void FParticleSystemStateStreamImpl::Render_PostUpdate()
+{
+	TStateStream<FParticleSystemStateStreamSettings>::Render_PostUpdate();
+
+	for (FParticleSystemObject* Object : Objects)
+	{
+		Object->Update();
+	}
+}
+
+////////////////////////////////////////////////////////////////////////////////////////////////////
diff --git a/Engine/Source/Runtime/Renderer/Private/StateStream/ParticleSystemStateStreamImpl.h b/Engine/Source/Runtime/Renderer/Private/StateStream/ParticleSystemStateStreamImpl.h
new file mode 100644
--- /dev/null
+++ b/Engine/Source/Runtime/Renderer/Private/StateStream/ParticleSystemStateStreamImpl.h
@@ -0,0 +1,128 @@
+// Copyright Epic Games, Inc. All Rights Reserved.
+
+#pragma once
+
+#include "ParticleEmitterInstanceOwner.h"
+#include "PrimitiveSceneDesc.h"
+#include "SceneTypes.h"
+#include "StateStream/ParticleSystemStateStream.h"
+#include "TransformStateStreamImpl.h"
+
+#define UE_API RENDERER_API
+
+class FParticleDynamicData;
+class FScene;
+struct FParticleSystemSceneProxyDesc;
+
+////////////////////////////////////////////////////////////////////////////////////////////////////
+
+class FParticleSystemObject : public TRefCountingMixin<FParticleSystemObject>, public FTransformObjectListener, public IParticleEmitterInstanceOwner
+{
+public:
+
+private:
+	virtual ~FParticleSystemObject();
+	virtual void OnTransformObjectDirty() override final;
+
+	TRefCountPtr<FTransformObject> TransformObject;
+
+	FCustomPrimitiveData CustomPrimitiveData;
+	FPrimitiveSceneInfoData	PrimitiveSceneData;
+	FPrimitiveSceneDesc PrimitiveSceneDesc;
+
+	friend class FParticleSystemStateStreamImpl;
+	friend class TRefCountingMixin<FParticleSystemObject>;
+
+	void InitializeSystem(FParticleSystemSceneProxyDesc& OutDesc, const FParticleSystemStaticState& Ss, const FParticleSystemDynamicState& Ds);
+	void Update();
+	FParticleDynamicData* CreateDynamicData(ERHIFeatureLevel::Type InFeatureLevel);
+
+	virtual const FTransform& GetAsyncComponentToWorld() const override;
+
+	virtual UObject* GetDistributionData() const override;
+	virtual const FTransform& GetComponentTransform() const override;
+	virtual FRotator GetComponentRotation() const override;
+	virtual const FTransform& GetComponentToWorld() const override;
+	virtual const FBoxSphereBounds& GetBounds() const override;
+	virtual TWeakObjectPtr<UWorld> GetWeakWorld() const override;
+	virtual bool HasWorld() const override;
+	virtual bool HasWorldSettings() const override;
+	virtual bool IsGameWorld() const override;
+	virtual float GetWorldTimeSeconds() const override;
+	virtual float GetWorldEffectiveTimeDilation() const override;
+	virtual FIntVector GetWorldOriginLocation() const override;
+	virtual FSceneInterface* GetScene() const override;
+	virtual bool GetFloatParameter(const FName InName, float& OutFloat) override;
+	virtual const FVector3f& GetLWCTile() const override;
+	virtual FString GetName() const override;
+	virtual FString GetFullName() const override;
+	virtual FString GetPathName() const override;
+	virtual bool IsActive() const override;
+	virtual bool IsValidLowLevel() const override;
+	virtual TArrayView<const FParticleSysParam> GetAsyncInstanceParameters() override;
+	virtual int32 GetCurrentDetailMode() const override;
+	virtual int32 GetCurrentLODIndex() const override;
+	virtual const FVector& GetPartSysVelocity() const override;
+	virtual const FVector& GetOldPosition() const override;
+	virtual FFXSystem* GetFXSystem() const override;
+	virtual const UParticleSystem* GetTemplate() const override;
+	virtual TArrayView<const FParticleSysParam> GetInstanceParameters() const override;
+	virtual TArrayView<FParticleEmitterInstance*> GetEmitterInstances() const override;
+	virtual TArrayView<TObjectPtr<UMaterialInterface>> GetEmitterMaterials() const override;
+	virtual FPrimitiveSceneProxy* GetSceneProxy() const override;
+	virtual bool GetIsWarmingUp() const override;
+	virtual bool GetJustRegistered() const override;
+	virtual float GetWarmupTime() const override;
+	virtual float GetEmitterDelay() const override;
+	virtual FRandomStream& GetRandomStream() override;
+
+	virtual void SetComponentToWorld(const FTransform& NewComponentToWorld) override;
+	virtual void DeactivateNextTick() override;
+
+	virtual UParticleSystemComponent* AsComponent() const override;
+
+	virtual void ReportEventSpawn(const FName InEventName, const float InEmitterTime, const FVector InLocation, const FVector InVelocity, const TArray<class UParticleModuleEventSendToGame*>& InEventData) override;
+	virtual void ReportEventDeath(const FName InEventName, const float InEmitterTime, const FVector InLocation, const FVector InVelocity, const TArray<class UParticleModuleEventSendToGame*>& InEventData, const float InParticleTime) override;
+	virtual void ReportEventCollision(const FName InEventName, const float InEmitterTime, const FVector InLocation, const FVector InDirection, const FVector InVelocity, const TArray<class UParticleModuleEventSendToGame*>& InEventData, const float InParticleTime, const FVector InNormal, const float InTime, const int32 InItem, const FName InBoneName, UPhysicalMaterial* PhysMat) override;
+	virtual void ReportEventBurst(const FName InEventName, const float InEmitterTime, const int32 ParticleCount, const FVector InLocation, const TArray<class UParticleModuleEventSendToGame*>& InEventData) override;
+
+	virtual TArrayView<FParticleEventSpawnData> GetSpawnEvents() const override;
+	virtual TArrayView<FParticleEventDeathData> GetDeathEvents() const override;
+	virtual TArrayView<FParticleEventCollideData> GetCollisionEvents() const override;
+	virtual TArrayView<FParticleEventBurstData> GetBurstEvents() const override;
+	virtual TArrayView<FParticleEventKismetData> GetKismetEvents() const override;
+
+	const UParticleSystem* Template;
+	TArray<FParticleEmitterInstance*> EmitterInstances;
+	FRandomStream RandomStream;
+	FVector3f LWCTile = FVector3f::ZeroVector;
+	bool bJustRegistered = true;
+};
+
+////////////////////////////////////////////////////////////////////////////////////////////////////
+
+struct FParticleSystemStateStreamSettings : TStateStreamSettings<IParticleSystemStateStream, FParticleSystemObject>
+{
+};
+
+////////////////////////////////////////////////////////////////////////////////////////////////////
+
+class FParticleSystemStateStreamImpl : public TStateStream<FParticleSystemStateStreamSettings>
+{
+public:
+	FParticleSystemStateStreamImpl(FScene* InScene);
+private:
+	UE_API void SetTransformObject(FParticleSystemObject& Object, const FParticleSystemDynamicState& Ds);
+	UE_API virtual void Render_OnCreate(const FParticleSystemStaticState& Ss, const FParticleSystemDynamicState& Ds, FParticleSystemObject*& UserData, bool IsDestroyedInSameFrame) override;
+	UE_API virtual void Render_OnUpdate(const FParticleSystemStaticState& Ss, const FParticleSystemDynamicState& Ds, FParticleSystemObject*& UserData) override;
+	UE_API virtual void Render_OnDestroy(const FParticleSystemStaticState& Ss, const FParticleSystemDynamicState& Ds, FParticleSystemObject*& UserData) override;
+
+	virtual void Render_PostUpdate() override;
+
+	FScene* Scene;
+	TArray<FParticleSystemObject*> Objects;
+};
+
+////////////////////////////////////////////////////////////////////////////////////////////////////
+
+#undef UE_API
