Heterogeneous Volumes: Allow enable/disable toggles for illumination and shadow creation based on light type

These are expected scalability options based on runtime requirements, potentially per platform:
r.HeterogeneousVolumes.LightType.X [default =1]
r.HeterogeneousVolumes.Shadow.LightType.X [default=1]

diff --git a/Engine/Source/Runtime/Renderer/Private/HeterogeneousVolumes/HeterogeneousVolumes.h b/Engine/Source/Runtime/Renderer/Private/HeterogeneousVolumes/HeterogeneousVolumes.h
--- a/Engine/Source/Runtime/Renderer/Private/HeterogeneousVolumes/HeterogeneousVolumes.h
+++ b/Engine/Source/Runtime/Renderer/Private/HeterogeneousVolumes/HeterogeneousVolumes.h
@@ -124,6 +124,8 @@ namespace HeterogeneousVolumes
 	bool EnableAmbientOcclusion();
 	bool UseExistenceMask();
 	bool UseAnalyticDerivatives();
+	bool SupportsLightType(uint32 LightType);
+	bool SupportsShadowForLightType(uint32 LightType);
 
 	enum class EFogMode
 	{
diff --git a/Engine/Source/Runtime/Renderer/Private/HeterogeneousVolumes/HeterogeneousVolumesLiveShadingPipeline.cpp b/Engine/Source/Runtime/Renderer/Private/HeterogeneousVolumes/HeterogeneousVolumesLiveShadingPipeline.cpp
--- a/Engine/Source/Runtime/Renderer/Private/HeterogeneousVolumes/HeterogeneousVolumesLiveShadingPipeline.cpp
+++ b/Engine/Source/Runtime/Renderer/Private/HeterogeneousVolumes/HeterogeneousVolumesLiveShadingPipeline.cpp
@@ -18,6 +18,34 @@
 #include "VolumetricFog.h"
 #include "BlueNoise.h"
 
+static TAutoConsoleVariable<bool> CVarHeterogeneousLightTypeDirectional(
+	TEXT("r.HeterogeneousVolumes.LightType.Directional"),
+	1,
+	TEXT("Enables illumination from the directional light (Default = 1)"),
+	ECVF_RenderThreadSafe
+);
+
+static TAutoConsoleVariable<bool> CVarHeterogeneousLightTypePoint(
+	TEXT("r.HeterogeneousVolumes.LightType.Point"),
+	1,
+	TEXT("Enables illumination from point lights (Default = 1)"),
+	ECVF_RenderThreadSafe
+);
+
+static TAutoConsoleVariable<bool> CVarHeterogeneousLightTypeSpot(
+	TEXT("r.HeterogeneousVolumes.LightType.Spot"),
+	1,
+	TEXT("Enables illumination from spot lights (Default = 1)"),
+	ECVF_RenderThreadSafe
+);
+
+static TAutoConsoleVariable<bool> CVarHeterogeneousLightTypeRect(
+	TEXT("r.HeterogeneousVolumes.LightType.Rect"),
+	1,
+	TEXT("Enables illumination from rect lights (Default = 1)"),
+	ECVF_RenderThreadSafe
+);
+
 static TAutoConsoleVariable<int32> CVarHeterogeneousLightingCacheBoundsCulling(
 	TEXT("r.HeterogeneousVolumes.LightingCache.BoundsCulling"),
 	1,
@@ -78,6 +106,23 @@ static TAutoConsoleVariable<int32> CVarHeterogeneousVolumesAdaptiveMarching(
 
 namespace HeterogeneousVolumes
 {
+	bool SupportsLightType(uint32 LightType)
+	{
+		switch (LightType)
+		{
+		case LightType_Directional:
+			return CVarHeterogeneousLightTypeDirectional.GetValueOnRenderThread();
+		case LightType_Point:
+			return CVarHeterogeneousLightTypePoint.GetValueOnRenderThread();
+		case LightType_Spot:
+			return CVarHeterogeneousLightTypeSpot.GetValueOnRenderThread();
+		case LightType_Rect:
+			return CVarHeterogeneousLightTypeRect.GetValueOnRenderThread();
+		}
+
+		return false;
+	}
+
 	bool ShouldBoundsCull()
 	{
 		return CVarHeterogeneousLightingCacheBoundsCulling.GetValueOnRenderThread() != 0;
@@ -1271,7 +1316,8 @@ static void RenderWithTransmittanceVolumePipeline(
 	TArray<FLightSceneInfoCompact, TInlineAllocator<64>> LightSceneInfoCompact;
 	for (auto LightIt = Scene->Lights.CreateConstIterator(); LightIt; ++LightIt)
 	{
-		if (LightIt->AffectsPrimitive(HeterogeneousVolumeInterface->GetBounds(), HeterogeneousVolumeInterface->GetPrimitiveSceneProxy()))
+		if (HeterogeneousVolumes::SupportsLightType(LightIt->LightType) &&
+			LightIt->AffectsPrimitive(HeterogeneousVolumeInterface->GetBounds(), HeterogeneousVolumeInterface->GetPrimitiveSceneProxy()))
 		{
 			LightSceneInfoCompact.Add(*LightIt);
 		}
@@ -1422,7 +1468,8 @@ static void RenderWithInscatteringVolumePipeline(
 		TArray<FLightSceneInfoCompact, TInlineAllocator<64>> LightSceneInfoCompact;
 		for (auto LightIt = Scene->Lights.CreateConstIterator(); LightIt; ++LightIt)
 		{
-			if (LightIt->AffectsPrimitive(HeterogeneousVolumeInterface->GetBounds(), HeterogeneousVolumeInterface->GetPrimitiveSceneProxy()))
+			if (HeterogeneousVolumes::SupportsLightType(LightIt->LightType) &&
+				LightIt->AffectsPrimitive(HeterogeneousVolumeInterface->GetBounds(), HeterogeneousVolumeInterface->GetPrimitiveSceneProxy()))
 			{
 				LightSceneInfoCompact.Add(*LightIt);
 			}
@@ -2294,6 +2341,7 @@ void RenderAdaptiveVolumetricShadowMapWithLiveShading(
 	{
 		// TODO: Use global bounds information..
 		//if (LightIt->AffectsPrimitive(HeterogeneousVolumeInterface->GetBounds(), HeterogeneousVolumeInterface->GetPrimitiveSceneProxy()))
+		if (HeterogeneousVolumes::SupportsShadowForLightType(LightIt->LightType))
 		{
 			LightSceneInfoCompact.Add(*LightIt);
 		}
diff --git a/Engine/Source/Runtime/Renderer/Private/HeterogeneousVolumes/HeterogeneousVolumesPreshadingPipeline.cpp b/Engine/Source/Runtime/Renderer/Private/HeterogeneousVolumes/HeterogeneousVolumesPreshadingPipeline.cpp
--- a/Engine/Source/Runtime/Renderer/Private/HeterogeneousVolumes/HeterogeneousVolumesPreshadingPipeline.cpp
+++ b/Engine/Source/Runtime/Renderer/Private/HeterogeneousVolumes/HeterogeneousVolumesPreshadingPipeline.cpp
@@ -731,7 +731,8 @@ void RenderWithInscatteringVolumePipelineWithPreshadingCompute(
 	TArray<FLightSceneInfoCompact, TInlineAllocator<64>> LightSceneInfoCompact;
 	for (auto LightIt = Scene->Lights.CreateConstIterator(); LightIt; ++LightIt)
 	{
-		if (LightIt->AffectsPrimitive(HeterogeneousVolumeInterface->GetBounds(), HeterogeneousVolumeInterface->GetPrimitiveSceneProxy()))
+		if (HeterogeneousVolumes::SupportsLightType(LightIt->LightType) &&
+			LightIt->AffectsPrimitive(HeterogeneousVolumeInterface->GetBounds(), HeterogeneousVolumeInterface->GetPrimitiveSceneProxy()))
 		{
 			LightSceneInfoCompact.Add(*LightIt);
 		}
@@ -864,7 +865,8 @@ void RenderWithTransmittanceVolumePipelineWithPreshadingCompute(
 	TArray<FLightSceneInfoCompact, TInlineAllocator<64>> LightSceneInfoCompact;
 	for (auto LightIt = Scene->Lights.CreateConstIterator(); LightIt; ++LightIt)
 	{
-		if (LightIt->AffectsPrimitive(HeterogeneousVolumeInterface->GetBounds(), HeterogeneousVolumeInterface->GetPrimitiveSceneProxy()))
+		if (HeterogeneousVolumes::SupportsLightType(LightIt->LightType) &&
+			LightIt->AffectsPrimitive(HeterogeneousVolumeInterface->GetBounds(), HeterogeneousVolumeInterface->GetPrimitiveSceneProxy()))
 		{
 			LightSceneInfoCompact.Add(*LightIt);
 		}
@@ -1089,7 +1091,8 @@ void RenderWithInscatteringVolumePipelineWithPreshadingHardwareRayTracing(
 	TArray<FLightSceneInfoCompact, TInlineAllocator<64>> LightSceneInfoCompact;
 	for (auto LightIt = Scene->Lights.CreateConstIterator(); LightIt; ++LightIt)
 	{
-		if (LightIt->AffectsPrimitive(HeterogeneousVolumeInterface->GetBounds(), HeterogeneousVolumeInterface->GetPrimitiveSceneProxy()))
+		if (HeterogeneousVolumes::SupportsLightType(LightIt->LightType) &&
+			LightIt->AffectsPrimitive(HeterogeneousVolumeInterface->GetBounds(), HeterogeneousVolumeInterface->GetPrimitiveSceneProxy()))
 		{
 			LightSceneInfoCompact.Add(*LightIt);
 		}
@@ -1218,7 +1221,8 @@ void RenderWithTransmittanceVolumePipelineWithPreshadingHardwareRayTracing(
 	TArray<FLightSceneInfoCompact, TInlineAllocator<64>> LightSceneInfoCompact;
 	for (auto LightIt = Scene->Lights.CreateConstIterator(); LightIt; ++LightIt)
 	{
-		if (LightIt->AffectsPrimitive(HeterogeneousVolumeInterface->GetBounds(), HeterogeneousVolumeInterface->GetPrimitiveSceneProxy()))
+		if (HeterogeneousVolumes::SupportsLightType(LightIt->LightType) &&
+			LightIt->AffectsPrimitive(HeterogeneousVolumeInterface->GetBounds(), HeterogeneousVolumeInterface->GetPrimitiveSceneProxy()))
 		{
 			LightSceneInfoCompact.Add(*LightIt);
 		}
diff --git a/Engine/Source/Runtime/Renderer/Private/HeterogeneousVolumes/HeterogeneousVolumesVoxelGridPipeline.cpp b/Engine/Source/Runtime/Renderer/Private/HeterogeneousVolumes/HeterogeneousVolumesVoxelGridPipeline.cpp
--- a/Engine/Source/Runtime/Renderer/Private/HeterogeneousVolumes/HeterogeneousVolumesVoxelGridPipeline.cpp
+++ b/Engine/Source/Runtime/Renderer/Private/HeterogeneousVolumes/HeterogeneousVolumesVoxelGridPipeline.cpp
@@ -273,6 +273,34 @@ static TAutoConsoleVariable<int32> CVarHeterogeneousVolumesShadowJitter(
 	ECVF_RenderThreadSafe
 );
 
+static TAutoConsoleVariable<bool> CVarHeterogeneousShadowsForLightTypeDirectional(
+	TEXT("r.HeterogeneousVolumes.Shadows.LightType.Directional"),
+	1,
+	TEXT("Enables shadows from the directional light (Default = 1)"),
+	ECVF_RenderThreadSafe
+);
+
+static TAutoConsoleVariable<bool> CVarHeterogeneousShadowsForLightTypePoint(
+	TEXT("r.HeterogeneousVolumes.Shadows.LightType.Point"),
+	1,
+	TEXT("Enables shadows from point lights (Default = 1)"),
+	ECVF_RenderThreadSafe
+);
+
+static TAutoConsoleVariable<bool> CVarHeterogeneousShadowsForLightTypeSpot(
+	TEXT("r.HeterogeneousVolumes.Shadows.LightType.Spot"),
+	1,
+	TEXT("Enables shadows from spot lights (Default = 1)"),
+	ECVF_RenderThreadSafe
+);
+
+static TAutoConsoleVariable<bool> CVarHeterogeneousShadowsForLightTypeRect(
+	TEXT("r.HeterogeneousVolumes.Shadows.LightType.Rect"),
+	1,
+	TEXT("Enables shadows from rect lights (Default = 1)"),
+	ECVF_RenderThreadSafe
+);
+
 IMPLEMENT_UNIFORM_BUFFER_STRUCT(FOrthoVoxelGridUniformBufferParameters, "OrthoGridUniformBuffer");
 IMPLEMENT_UNIFORM_BUFFER_STRUCT(FFrustumVoxelGridUniformBufferParameters, "FrustumGridUniformBuffer");
 IMPLEMENT_UNIFORM_BUFFER_STRUCT(FAdaptiveVolumetricShadowMapUniformBufferParameters, "AVSM");
@@ -461,6 +489,26 @@ namespace HeterogeneousVolumes
 		return CVarHeterogeneousVolumesShadowJitter.GetValueOnRenderThread() != 0;
 	}
 
+	bool SupportsShadowForLightType(uint32 LightType)
+	{
+		if (SupportsLightType(LightType))
+		{
+			switch (LightType)
+			{
+			case LightType_Directional:
+				return CVarHeterogeneousShadowsForLightTypeDirectional.GetValueOnRenderThread();
+			case LightType_Point:
+				return CVarHeterogeneousShadowsForLightTypePoint.GetValueOnRenderThread();
+			case LightType_Spot:
+				return CVarHeterogeneousShadowsForLightTypeSpot.GetValueOnRenderThread();
+			case LightType_Rect:
+				return CVarHeterogeneousShadowsForLightTypeRect.GetValueOnRenderThread();
+			}
+		}
+
+		return false;
+	}
+
 	bool IsDynamicShadow(const FVisibleLightInfo* VisibleLightInfo)
 	{
 		check(VisibleLightInfo != nullptr);
@@ -2688,6 +2736,7 @@ void CollectHeterogeneousVolumeMeshBatches(
 	{
 		// TODO: Use global bounds information..
 		//if (LightIt->AffectsPrimitive(HeterogeneousVolumeInterface->GetBounds(), HeterogeneousVolumeInterface->GetPrimitiveSceneProxy()))
+		if (HeterogeneousVolumes::SupportsShadowForLightType(LightIt->LightType))
 		{
 			LightSceneInfoCompact.Add(*LightIt);
 		}
@@ -3786,6 +3835,7 @@ void RenderAdaptiveVolumetricShadowMapWithVoxelGrid(
 	{
 		// TODO: Use global bounds information..
 		//if (LightIt->AffectsPrimitive(HeterogeneousVolumeInterface->GetBounds(), HeterogeneousVolumeInterface->GetPrimitiveSceneProxy()))
+		if (HeterogeneousVolumes::SupportsShadowForLightType(LightIt->LightType))
 		{
 			LightSceneInfoCompact.Add(*LightIt);
 		}
