Fix volumetric cloud requesting a new render target when dynamic resolution changes
set r.VolumetricRenderTarget.ViewRect to 0 to fallback to the previous behaviour

diff --git a/Engine/Plugins/FX/Niagara/Shaders/Private/NiagaraDataInterfaceOcclusion.ush b/Engine/Plugins/FX/Niagara/Shaders/Private/NiagaraDataInterfaceOcclusion.ush
--- a/Engine/Plugins/FX/Niagara/Shaders/Private/NiagaraDataInterfaceOcclusion.ush
+++ b/Engine/Plugins/FX/Niagara/Shaders/Private/NiagaraDataInterfaceOcclusion.ush
@@ -4,6 +4,7 @@
 #include "/Engine/Private/SkyAtmosphereCommon.ush"
 
 float3				{ParameterName}_SystemLWCTile;
+float2				{ParameterName}_CloudVolumetricTextureUVScale;
 Texture2D<float4>	{ParameterName}_CloudVolumetricTexture;
 SamplerState		{ParameterName}_CloudVolumetricTextureSampler;
 
@@ -212,7 +213,7 @@ void QueryCloudOcclusionWithCircle_{ParameterName}(float3 WorldPosition, float W
 	float ExpectedSamples = 0.0f;
 	if ( bIncludeCenterSample && (all(ScreenUV >= 0.0f) && all(ScreenUV < 1.0f)) )
 	{
-		const float4 SampleColor = {ParameterName}_CloudVolumetricTexture.SampleLevel({ParameterName}_CloudVolumetricTextureSampler, ScreenUV, 0);
+		const float4 SampleColor = {ParameterName}_CloudVolumetricTexture.SampleLevel({ParameterName}_CloudVolumetricTextureSampler, ScreenUV * {ParameterName}_CloudVolumetricTextureUVScale, 0);
 
 		TotalVisibility += SampleColor.a;
 		TotalSamples += 1.0f;
@@ -238,7 +239,7 @@ void QueryCloudOcclusionWithCircle_{ParameterName}(float3 WorldPosition, float W
 				const float2 SampleUV = ScreenUV + (RingRadiusUV * float2(sin(Rad), cos(Rad)));
 				if ( all(SampleUV >= 0.0f) && all(SampleUV < 1.0f) )
 				{
-					const float4 SampleColor = {ParameterName}_CloudVolumetricTexture.SampleLevel({ParameterName}_CloudVolumetricTextureSampler, SampleUV, 0);
+					const float4 SampleColor = {ParameterName}_CloudVolumetricTexture.SampleLevel({ParameterName}_CloudVolumetricTextureSampler, SampleUV * {ParameterName}_CloudVolumetricTextureUVScale, 0);
 
 					TotalVisibility += SampleColor.a;
 					TotalSamples += 1.0f;
diff --git a/Engine/Plugins/FX/Niagara/Source/Niagara/Classes/NiagaraDataInterfaceOcclusion.h b/Engine/Plugins/FX/Niagara/Source/Niagara/Classes/NiagaraDataInterfaceOcclusion.h
--- a/Engine/Plugins/FX/Niagara/Source/Niagara/Classes/NiagaraDataInterfaceOcclusion.h
+++ b/Engine/Plugins/FX/Niagara/Source/Niagara/Classes/NiagaraDataInterfaceOcclusion.h
@@ -13,6 +13,7 @@ class UNiagaraDataInterfaceOcclusion : public UNiagaraDataInterface
 
 	BEGIN_SHADER_PARAMETER_STRUCT(FShaderParameters, )
 		SHADER_PARAMETER(FVector3f,				SystemLWCTile)
+		SHADER_PARAMETER(FVector2f,				CloudVolumetricTextureUVScale)
 		SHADER_PARAMETER_RDG_TEXTURE(Texture2D, CloudVolumetricTexture)
 		SHADER_PARAMETER_SAMPLER(SamplerState,	CloudVolumetricTextureSampler)
 	END_SHADER_PARAMETER_STRUCT();
diff --git a/Engine/Plugins/FX/Niagara/Source/Niagara/Private/NiagaraDataInterfaceOcclusion.cpp b/Engine/Plugins/FX/Niagara/Source/Niagara/Private/NiagaraDataInterfaceOcclusion.cpp
--- a/Engine/Plugins/FX/Niagara/Source/Niagara/Private/NiagaraDataInterfaceOcclusion.cpp
+++ b/Engine/Plugins/FX/Niagara/Source/Niagara/Private/NiagaraDataInterfaceOcclusion.cpp
@@ -239,6 +239,7 @@ void UNiagaraDataInterfaceOcclusion::SetShaderParameters(const FNiagaraDataInter
 			CloudTexture = Context.GetComputeDispatchInterface().GetBlackTexture(Context.GetGraphBuilder(), ETextureDimension::Texture2D);
 		}
 		ShaderParameters->CloudVolumetricTexture = CloudTexture;
+		ShaderParameters->CloudVolumetricTextureUVScale = View && View->State ? View->State->GetVolumetricCloudTextureUVScale() : FVector2f(1.0f, 1.0f);
 	}
 	ShaderParameters->CloudVolumetricTextureSampler = TStaticSamplerState<SF_Bilinear>::GetRHI();
 }
diff --git a/Engine/Shaders/Private/BasePassPixelShader.usf b/Engine/Shaders/Private/BasePassPixelShader.usf
--- a/Engine/Shaders/Private/BasePassPixelShader.usf
+++ b/Engine/Shaders/Private/BasePassPixelShader.usf
@@ -1438,7 +1438,7 @@ void FPixelShaderInOut_MainPS(
 			NDCPosition, GetTranslatedWorldPosition(MaterialParameters), ResolvedView.TranslatedWorldCameraOrigin,
 			TranslucentBasePass.VolumetricCloudColor, TranslucentBasePass.VolumetricCloudColorSampler,
 			TranslucentBasePass.VolumetricCloudDepth, TranslucentBasePass.VolumetricCloudDepthSampler,
-			OneOverPreExposure, Fogging, TranslucentBasePass.SoftBlendingDistanceKm);
+			OneOverPreExposure, Fogging, TranslucentBasePass.SoftBlendingDistanceKm, TranslucentBasePass.VolumetricCloudColorUVScale);
 	}
 
 #endif
diff --git a/Engine/Shaders/Private/BasePassVertexShader.usf b/Engine/Shaders/Private/BasePassVertexShader.usf
--- a/Engine/Shaders/Private/BasePassVertexShader.usf
+++ b/Engine/Shaders/Private/BasePassVertexShader.usf
@@ -134,7 +134,7 @@ void Main(
 			Output.Position, WorldPosition.xyz, ResolvedView.TranslatedWorldCameraOrigin,
 			TranslucentBasePass.VolumetricCloudColor, TranslucentBasePass.VolumetricCloudColorSampler,
 			TranslucentBasePass.VolumetricCloudDepth, TranslucentBasePass.VolumetricCloudDepthSampler,
-			OneOverPreExposure, Output.BasePassInterpolants.VertexFog, TranslucentBasePass.SoftBlendingDistanceKm);
+			OneOverPreExposure, Output.BasePassInterpolants.VertexFog, TranslucentBasePass.SoftBlendingDistanceKm, TranslucentBasePass.VolumetricCloudColorUVScale);
 	}
 #endif
 
diff --git a/Engine/Shaders/Private/VolumetricCloudCommon.ush b/Engine/Shaders/Private/VolumetricCloudCommon.ush
--- a/Engine/Shaders/Private/VolumetricCloudCommon.ush
+++ b/Engine/Shaders/Private/VolumetricCloudCommon.ush
@@ -77,9 +77,10 @@ float4 GetCloudLuminanceTransmittanceOverFog(
 	float4 NDC, float3 SampledWorldPos, float3 CameraWorldPos,
 	Texture2D<float4> VolumetricCloudColor, SamplerState VolumetricCloudColorSampler,
 	Texture2D<float4> VolumetricCloudDepth, SamplerState VolumetricCloudDepthSampler,
-	float OneOverPreExposure, float4 CurrentVertexFog, float SoftBlendingDistanceKm)
+	float OneOverPreExposure, float4 CurrentVertexFog, float SoftBlendingDistanceKm, float2 UVScale)
 {
 	float2 ScreenUv = (NDC.xy / NDC.ww) * float2(0.5f, -0.5f) + 0.5f;
+	ScreenUv *= UVScale;
 
 
 #if FEATURE_LEVEL >= FEATURE_LEVEL_SM5
diff --git a/Engine/Shaders/Private/VolumetricRenderTarget.usf b/Engine/Shaders/Private/VolumetricRenderTarget.usf
--- a/Engine/Shaders/Private/VolumetricRenderTarget.usf
+++ b/Engine/Shaders/Private/VolumetricRenderTarget.usf
@@ -27,6 +27,7 @@ Texture2D<float4> TracingVolumetricDepthTexture;
 uint4 TracingVolumetricTextureValidCoordRect;
 float4 TracingVolumetricTextureValidUvRect;
 float MinimumDistanceKmToEnableReprojection;
+float2 TracingVolumetricTextureUVScale;
 
 struct FDepthData
 {
@@ -49,6 +50,7 @@ float4 SafeLoadTracingVolumetricTexture(uint2 Coord)
 }
 float4 SafeSampleTracingVolumetricTexture(float2 UV)
 {
+	UV *= TracingVolumetricTextureUVScale.xy;
 	return TracingVolumetricTexture.SampleLevel(LinearTextureSampler, clamp(UV, TracingVolumetricTextureValidUvRect.xy, TracingVolumetricTextureValidUvRect.zw), 0);
 }
 float4 SafeLoadSecondaryTracingVolumetricTexture(uint2 Coord)
@@ -57,6 +59,7 @@ float4 SafeLoadSecondaryTracingVolumetricTexture(uint2 Coord)
 }
 float4 SafeSampleSecondaryTracingVolumetricTexture(float2 UV)
 {
+	UV *= TracingVolumetricTextureUVScale.xy;
 	return SecondaryTracingVolumetricTexture.SampleLevel(LinearTextureSampler, clamp(UV, TracingVolumetricTextureValidUvRect.xy, TracingVolumetricTextureValidUvRect.zw), 0);
 }
 FDepthData SafeLoadTracingVolumetricDepthTexture(uint2 Coord)
@@ -65,6 +68,7 @@ FDepthData SafeLoadTracingVolumetricDepthTexture(uint2 Coord)
 }
 FDepthData SafeSampleTracingVolumetricDepthTexture(float2 UV)
 {
+	UV *= TracingVolumetricTextureUVScale.xy;
 	return GetDepthDataFromVector(TracingVolumetricDepthTexture.SampleLevel(LinearTextureSampler, clamp(UV, TracingVolumetricTextureValidUvRect.xy, TracingVolumetricTextureValidUvRect.zw), 0));
 }
 
@@ -75,6 +79,7 @@ float4 PreviousVolumetricTextureSizeAndInvSize;
 
 uint4 PreviousFrameVolumetricTextureValidCoordRect;
 float4 PreviousFrameVolumetricTextureValidUvRect;
+float2 PreviousFrameVolumetricTextureUVScale;
 
 float HistoryPreExposureCorrection;
 
@@ -85,6 +90,7 @@ float4 SafeLoadPreviousFrameVolumetricTexture(uint2 Coord)
 }
 float4 SafeSamplePreviousFrameVolumetricTexture(float2 UV)
 {
+	UV *= PreviousFrameVolumetricTextureUVScale.xy;
 	float4 Result = PreviousFrameVolumetricTexture.SampleLevel(LinearTextureSampler, clamp(UV, PreviousFrameVolumetricTextureValidUvRect.xy, PreviousFrameVolumetricTextureValidUvRect.zw), 0);
 	return float4(Result.rgb * HistoryPreExposureCorrection, Result.a);
 }
@@ -94,6 +100,7 @@ FDepthData SafeLoadPreviousFrameVolumetricDepthTexture(uint2 Coord)
 }
 FDepthData SafeSamplePreviousFrameVolumetricDepthTexture(float2 UV)
 {
+	UV *= PreviousFrameVolumetricTextureUVScale.xy;
 	return GetDepthDataFromVector(float4(PreviousFrameVolumetricDepthTexture.SampleLevel(LinearTextureSampler, clamp(UV, PreviousFrameVolumetricTextureValidUvRect.xy, PreviousFrameVolumetricTextureValidUvRect.zw), 0).rg, 0.0, 0.0));
 }
 #endif // PERMUTATION_HISTORY_AVAILABLE
diff --git a/Engine/Source/Runtime/Engine/Public/SceneManagement.h b/Engine/Source/Runtime/Engine/Public/SceneManagement.h
--- a/Engine/Source/Runtime/Engine/Public/SceneManagement.h
+++ b/Engine/Source/Runtime/Engine/Public/SceneManagement.h
@@ -191,6 +191,7 @@ public:
 
 	/* Return the active volumetric cloud texture, can be null. */
 	virtual FRDGTextureRef GetVolumetricCloudTexture(FRDGBuilder& GraphBuilder) = 0;
+	virtual FVector2f GetVolumetricCloudTextureUVScale() const = 0;
 
 	//
 	virtual uint32 GetCurrentTemporalAASampleIndex() const = 0;
diff --git a/Engine/Source/Runtime/Renderer/Private/BasePassRendering.h b/Engine/Source/Runtime/Renderer/Private/BasePassRendering.h
--- a/Engine/Source/Runtime/Renderer/Private/BasePassRendering.h
+++ b/Engine/Source/Runtime/Renderer/Private/BasePassRendering.h
@@ -128,6 +128,7 @@ BEGIN_GLOBAL_SHADER_PARAMETER_STRUCT(FTranslucentBasePassUniformParameters,)
 	SHADER_PARAMETER_SAMPLER(SamplerState, VolumetricCloudDepthSampler)
 	SHADER_PARAMETER(float, ApplyVolumetricCloudOnTransparent)
 	SHADER_PARAMETER(float, SoftBlendingDistanceKm)
+	SHADER_PARAMETER(FVector2f, VolumetricCloudColorUVScale)
 	// Translucency Lighting Volume
 	SHADER_PARAMETER_STRUCT_INCLUDE(FTranslucencyLightingVolumeParameters, TranslucencyLightingVolume)
 	SHADER_PARAMETER_STRUCT_INCLUDE(FLumenTranslucencyLightingParameters, LumenParameters)
diff --git a/Engine/Source/Runtime/Renderer/Private/DeferredShadingRenderer.cpp b/Engine/Source/Runtime/Renderer/Private/DeferredShadingRenderer.cpp
--- a/Engine/Source/Runtime/Renderer/Private/DeferredShadingRenderer.cpp
+++ b/Engine/Source/Runtime/Renderer/Private/DeferredShadingRenderer.cpp
@@ -2302,7 +2302,7 @@ void FDeferredShadingSceneRenderer::Render(FRDGBuilder& GraphBuilder)
 
 		if (bShouldRenderVolumetricCloudBase)
 		{
-			InitVolumetricRenderTargetForViews(GraphBuilder, Views);
+			InitVolumetricRenderTargetForViews(GraphBuilder, Views, SceneTextures);
 		}
 		else
 		{
diff --git a/Engine/Source/Runtime/Renderer/Private/ScenePrivate.h b/Engine/Source/Runtime/Renderer/Private/ScenePrivate.h
--- a/Engine/Source/Runtime/Renderer/Private/ScenePrivate.h
+++ b/Engine/Source/Runtime/Renderer/Private/ScenePrivate.h
@@ -1143,6 +1143,11 @@ public:
 		return VolumetricCloudRenderTarget.GetDstVolumetricReconstructRT(GraphBuilder);
 	}
 
+	virtual FVector2f GetVolumetricCloudTextureUVScale() const override
+	{
+		return VolumetricCloudRenderTarget.GetDstVolumetricReconstructUVScale();
+	}
+
 	FHairStrandsViewStateData HairStrandsViewStateData;
 
 	FShaderPrintStateData ShaderPrintStateData;
diff --git a/Engine/Source/Runtime/Renderer/Private/TranslucentRendering.cpp b/Engine/Source/Runtime/Renderer/Private/TranslucentRendering.cpp
--- a/Engine/Source/Runtime/Renderer/Private/TranslucentRendering.cpp
+++ b/Engine/Source/Runtime/Renderer/Private/TranslucentRendering.cpp
@@ -1078,6 +1078,7 @@ TRDGUniformBufferRef<FTranslucentBasePassUniformParameters> CreateTranslucentBas
 					BasePassParameters.VolumetricCloudColor = VolumetricReconstructRT;
 					BasePassParameters.VolumetricCloudDepth = View.ViewState->VolumetricCloudRenderTarget.GetOrCreateVolumetricTracingRTDepth(GraphBuilder);
 					BasePassParameters.ApplyVolumetricCloudOnTransparent = 1.0f;
+					BasePassParameters.VolumetricCloudColorUVScale = View.ViewState->VolumetricCloudRenderTarget.GetVolumetricTracingUVScale();
 				}
 			}
 			else
@@ -1089,6 +1090,7 @@ TRDGUniformBufferRef<FTranslucentBasePassUniformParameters> CreateTranslucentBas
 					BasePassParameters.VolumetricCloudColor = GraphBuilder.RegisterExternalTexture(VolumetricReconstructRT);
 					BasePassParameters.VolumetricCloudDepth = GraphBuilder.RegisterExternalTexture(VolumetricReconstructRTDepth);
 					BasePassParameters.ApplyVolumetricCloudOnTransparent = 1.0f;
+					BasePassParameters.VolumetricCloudColorUVScale = View.ViewState->VolumetricCloudRenderTarget.GetDstVolumetricReconstructUVScale();
 				}
 			}
 		}
diff --git a/Engine/Source/Runtime/Renderer/Private/VolumetricCloudRendering.cpp b/Engine/Source/Runtime/Renderer/Private/VolumetricCloudRendering.cpp
--- a/Engine/Source/Runtime/Renderer/Private/VolumetricCloudRendering.cpp
+++ b/Engine/Source/Runtime/Renderer/Private/VolumetricCloudRendering.cpp
@@ -2465,8 +2465,10 @@ void FSceneRenderer::RenderVolumetricCloudsInternal(FRDGBuilder& GraphBuilder, F
 
 		bool bSampleVirtualShadowMap = (!bCloudDebugViewModeEnabled && !bSkipAtmosphericLightShadowmap && CVarVolumetricCloudShadowSampleAtmosphericLightShadowmap.GetValueOnRenderThread() > 0);
 
+		FVolumetricRenderTargetViewStateData& VRT = MainView.ViewState->VolumetricCloudRenderTarget;
 		FRenderVolumetricCloudRenderViewCS::FParameters* PassParameters = GraphBuilder.AllocParameters<FRenderVolumetricCloudRenderViewCS::FParameters>();
-		PassParameters->OutputViewRect = FVector4f(0.f, 0.f, Desc.Extent.X, Desc.Extent.Y);
+		FIntPoint ViewRectSize = VRT.GetCurrentVolumetricTracingViewRect();
+		PassParameters->OutputViewRect = FVector4f(0.f, 0.f, ViewRectSize.X, ViewRectSize.Y);
 		PassParameters->bBlendCloudColor = !bShouldViewRenderVolumetricRenderTarget && !bCloudDebugViewModeEnabled && !CloudRC.bDisableCloudBlending;
 		PassParameters->TargetCubeFace = CloudRC.RenderTargets.Output[0].GetArraySlice();
 		PassParameters->VolumetricCloudRenderViewParamsUB = CloudPassUniformBuffer;
@@ -2496,7 +2498,6 @@ void FSceneRenderer::RenderVolumetricCloudsInternal(FRDGBuilder& GraphBuilder, F
 		const FMaterial* MaterialResource = &CloudVolumeMaterialProxy->GetMaterialWithFallback(Scene->GetFeatureLevel(), MaterialRenderProxy);
 		MaterialRenderProxy = MaterialRenderProxy ? MaterialRenderProxy : CloudVolumeMaterialProxy;
 
-		FVolumetricRenderTargetViewStateData& VRT = MainView.ViewState->VolumetricCloudRenderTarget;
 
 		typename FRenderVolumetricCloudRenderViewCS::FPermutationDomain PermutationVector;
 		PermutationVector.Set<typename FRenderVolumetricCloudRenderViewCS::FCloudPerSampleAtmosphereTransmittance>(!bCloudDebugViewModeEnabled && ShouldUsePerSampleAtmosphereTransmittance(Scene, &MainView));
diff --git a/Engine/Source/Runtime/Renderer/Private/VolumetricRenderTarget.cpp b/Engine/Source/Runtime/Renderer/Private/VolumetricRenderTarget.cpp
--- a/Engine/Source/Runtime/Renderer/Private/VolumetricRenderTarget.cpp
+++ b/Engine/Source/Runtime/Renderer/Private/VolumetricRenderTarget.cpp
@@ -46,6 +46,11 @@ static TAutoConsoleVariable<int32> CVarVolumetricRenderTargetPreferAsyncCompute(
 	TEXT("Whether to prefer using async compute to generate volumetric cloud render targets."),
 	ECVF_RenderThreadSafe | ECVF_Scalability);
 
+static TAutoConsoleVariable<int32> CVarVolumetricRenderTargetViewRect(
+	TEXT("r.VolumetricRenderTarget.ViewRect"), 1,
+	TEXT("Enable ViewRect support: does not reallocate new render targets when dynamic resolution changes"),
+	ECVF_RenderThreadSafe | ECVF_Scalability);
+
 static TAutoConsoleVariable<int32> CVarVolumetricRenderTargetReprojectionBoxConstraint(
 	TEXT("r.VolumetricRenderTarget.ReprojectionBoxConstraint"), 0,
 	TEXT("Whether reprojected data should be constrained to the new incoming cloud data neighborhod value."),
@@ -131,17 +136,20 @@ static uint32 GetTraceDownsampleFactor(int32 Mode)
 	return 2;
 }
 
-static void GetTextureSafeUvCoordBound(FRDGTextureRef Texture, FUintVector4& TextureValidCoordRect, FVector4f& TextureValidUvRect)
+static void GetTextureSafeUvCoordBound(const FIntPoint& TextureViewRect, FRDGTextureRef Texture, FUintVector4& TextureValidCoordRect, FVector4f& TextureValidUvRect, FVector2f& UVScale)
 {
 	FIntVector TexSize = Texture->Desc.GetSize();
+    UVScale.X = (float)TextureViewRect.X / (float)TexSize.X;
+    UVScale.Y = (float)TextureViewRect.Y / (float)TexSize.Y;
+
 	TextureValidCoordRect.X = 0;
 	TextureValidCoordRect.Y = 0;
-	TextureValidCoordRect.Z = TexSize.X - 1;
-	TextureValidCoordRect.W = TexSize.Y - 1;
+	TextureValidCoordRect.Z = TextureViewRect.X - 1;
+	TextureValidCoordRect.W = TextureViewRect.Y - 1;
 	TextureValidUvRect.X = 0.51f / float(TexSize.X);
 	TextureValidUvRect.Y = 0.51f / float(TexSize.Y);
-	TextureValidUvRect.Z = (float(TexSize.X) - 0.51f) / float(TexSize.X);
-	TextureValidUvRect.W = (float(TexSize.Y) - 0.51f) / float(TexSize.Y);
+	TextureValidUvRect.Z = (float(TextureViewRect.X) - 0.51f) / float(TexSize.X);
+	TextureValidUvRect.W = (float(TextureViewRect.Y) - 0.51f) / float(TexSize.Y);
 };
 
 static bool AnyViewRequiresProcessing(TArrayView<FViewInfo> Views)
@@ -172,9 +180,14 @@ FVolumetricRenderTargetViewStateData::FVolumetricRenderTargetViewStateData()
 	, FullResolution(FIntPoint::ZeroValue)
 	, VolumetricReconstructRTResolution(FIntPoint::ZeroValue)
 	, VolumetricTracingRTResolution(FIntPoint::ZeroValue)
+	, VolumetricTracingViewRect(FIntPoint::ZeroValue)
 {
 	VolumetricReconstructRTDownsampleFactor = 0;
 	VolumetricTracingRTDownsampleFactor = 0;
+	for (uint32 i = 0; i < kRenderTargetCount; ++i)
+	{
+		VolumetricReconstructViewRect[i] = FIntPoint::ZeroValue;
+	}
 }
 
 FVolumetricRenderTargetViewStateData::~FVolumetricRenderTargetViewStateData()
@@ -182,11 +195,18 @@ FVolumetricRenderTargetViewStateData::~FVolumetricRenderTargetViewStateData()
 }
 
 void FVolumetricRenderTargetViewStateData::Initialise(
+	FIntPoint& TextureResolutionIn,
 	FIntPoint& ViewRectResolutionIn,
 	int32 InMode,
 	int32 InUpsamplingMode,
 	bool bCameraCut)
 {
+	FIntPoint TextureResolutionInCopy(TextureResolutionIn);
+	if (CVarVolumetricRenderTargetViewRect.GetValueOnAnyThread() == 0)
+	{
+		TextureResolutionInCopy = ViewRectResolutionIn;
+	}
+
 	// Update internal settings
 	Mode = FMath::Clamp(InMode, 0, 3);
 	UpsamplingMode = Mode == 2 || Mode == 3 ? 2 : FMath::Clamp(InUpsamplingMode, 0, 4); // if we are using mode 2 then we cannot intersect with depth and upsampling should be 2 (simple on/off intersection)
@@ -208,12 +228,12 @@ void FVolumetricRenderTargetViewStateData::Initialise(
 
 		// We always reallocate on a resolution change to adapt to dynamic resolution scaling.
 		// TODO allocate once at max resolution and change source and destination coord/uvs/rect.
-		if (FullResolution != ViewRectResolutionIn || GetMainDownsampleFactor(Mode) != VolumetricReconstructRTDownsampleFactor || GetTraceDownsampleFactor(Mode) != VolumetricTracingRTDownsampleFactor)
+		if (FullResolution != TextureResolutionInCopy || GetMainDownsampleFactor(Mode) != VolumetricReconstructRTDownsampleFactor || GetTraceDownsampleFactor(Mode) != VolumetricTracingRTDownsampleFactor)
 		{
 			VolumetricReconstructRTDownsampleFactor = GetMainDownsampleFactor(Mode);
 			VolumetricTracingRTDownsampleFactor = GetTraceDownsampleFactor(Mode);
 
-			FullResolution = ViewRectResolutionIn;
+			FullResolution = TextureResolutionInCopy;
 			VolumetricReconstructRTResolution = FIntPoint::DivideAndRoundUp(FullResolution, VolumetricReconstructRTDownsampleFactor);							// Half resolution
 			VolumetricTracingRTResolution = FIntPoint::DivideAndRoundUp(VolumetricReconstructRTResolution, VolumetricTracingRTDownsampleFactor);	// Half resolution of the volumetric buffer
 
@@ -223,6 +243,9 @@ void FVolumetricRenderTargetViewStateData::Initialise(
 			VolumetricTracingRTDepth.SafeRelease();
 		}
 
+		VolumetricReconstructViewRect[CurrentRT] = FIntPoint::DivideAndRoundUp(ViewRectResolutionIn, VolumetricReconstructRTDownsampleFactor);							// Half resolution
+		VolumetricTracingViewRect = FIntPoint::DivideAndRoundUp(VolumetricReconstructViewRect[CurrentRT], VolumetricTracingRTDownsampleFactor);	// Half resolution of the volumetric buffer
+
 		FIntVector CurrentTargetResVec = VolumetricReconstructRT[CurrentRT].IsValid() ? VolumetricReconstructRT[CurrentRT]->GetDesc().GetSize() : FIntVector::ZeroValue;
 		FIntPoint CurrentTargetRes = FIntPoint::DivideAndRoundUp(FullResolution, VolumetricReconstructRTDownsampleFactor);
 		if (VolumetricReconstructRT[CurrentRT].IsValid() && FIntPoint(CurrentTargetResVec.X, CurrentTargetResVec.Y) != CurrentTargetRes)
@@ -340,6 +363,12 @@ FRDGTextureRef FVolumetricRenderTargetViewStateData::GetOrCreateVolumetricTracin
 	return GraphBuilder.RegisterExternalTexture(VolumetricTracingRTDepth);
 }
 
+FVector2f FVolumetricRenderTargetViewStateData::GetVolumetricTracingUVScale() const
+{
+	return FVector2f(float(VolumetricTracingViewRect.X) / float(VolumetricTracingRTResolution.X), 
+					 float(VolumetricTracingViewRect.Y) / float(VolumetricTracingRTResolution.Y));
+}
+
 FRDGTextureRef FVolumetricRenderTargetViewStateData::GetDstVolumetricReconstructRT(FRDGBuilder& GraphBuilder)
 {
 	return VolumetricReconstructRT[CurrentRT].IsValid() ? GraphBuilder.RegisterExternalTexture(VolumetricReconstructRT[CurrentRT]) : nullptr;
@@ -385,6 +414,17 @@ TRefCountPtr<IPooledRenderTarget> FVolumetricRenderTargetViewStateData::GetDstVo
 	return VolumetricReconstructRTDepth[CurrentRT];
 }
 
+const FIntPoint& FVolumetricRenderTargetViewStateData::GetDstVolumetricReconstructViewRect() const
+{
+	return VolumetricReconstructViewRect[CurrentRT];
+}
+
+FVector2f FVolumetricRenderTargetViewStateData::GetDstVolumetricReconstructUVScale() const
+{
+	return FVector2f(float(VolumetricReconstructViewRect[CurrentRT].X) / float(VolumetricReconstructRTResolution.X), 
+					 float(VolumetricReconstructViewRect[CurrentRT].Y) / float(VolumetricReconstructRTResolution.Y));
+}
+
 FRDGTextureRef FVolumetricRenderTargetViewStateData::GetOrCreateSrcVolumetricReconstructRT(FRDGBuilder& GraphBuilder)
 {
 	check(VolumetricReconstructRTResolution != FIntPoint::ZeroValue); // check that initialization has been done at least once
@@ -401,6 +441,11 @@ FRDGTextureRef FVolumetricRenderTargetViewStateData::GetOrCreateSrcVolumetricRec
 	return GraphBuilder.RegisterExternalTexture(VolumetricReconstructRTDepth[1u - CurrentRT]);
 }
 
+const FIntPoint& FVolumetricRenderTargetViewStateData::GetSrcVolumetricReconstructViewRect() const
+{
+	return VolumetricReconstructViewRect[1u - CurrentRT];
+}
+
 static float GetVolumetricBufferResolutionScale(uint32 VRTMode)
 {
 	if (VRTMode == 2) // Only valid for mode 2
@@ -445,7 +490,16 @@ FUintVector4 FVolumetricRenderTargetViewStateData::GetTracingCoordToFullResPixel
 	FSceneRenderer implementation.
 =============================================================================*/
 
-void InitVolumetricRenderTargetForViews(FRDGBuilder& GraphBuilder, TArrayView<FViewInfo> Views)
+FIntPoint ComputeVolumetricBufferResolution(const FIntPoint& InViewRect, int32 VRTMode)
+{
+	FIntPoint OutViewRect = InViewRect;
+	const float VolumetricRenderTargetScale = GetVolumetricBufferResolutionScale(VRTMode);
+	OutViewRect.X = FMath::RoundToInt((float)OutViewRect.X * VolumetricRenderTargetScale);
+	OutViewRect.Y = FMath::RoundToInt((float)OutViewRect.Y * VolumetricRenderTargetScale);
+	return OutViewRect;
+}
+
+void InitVolumetricRenderTargetForViews(FRDGBuilder& GraphBuilder, TArrayView<FViewInfo> Views, const FMinimalSceneTextures& SceneTextures)
 {
 	for (int32 ViewIndex = 0; ViewIndex < Views.Num(); ViewIndex++)
 	{
@@ -461,11 +515,11 @@ void InitVolumetricRenderTargetForViews(FRDGBuilder& GraphBuilder, TArrayView<FV
 
 		int32 VRTMode = FMath::Clamp(CVarVolumetricRenderTargetMode.GetValueOnRenderThread(), 0, 3);
 
-		FIntPoint ViewRect = ViewInfo.ViewRect.Size();
-		const float VolumetricRenderTargetScale = GetVolumetricBufferResolutionScale(VRTMode);
-		ViewRect.X = FMath::RoundToInt((float)ViewRect.X * VolumetricRenderTargetScale);
-		ViewRect.Y = FMath::RoundToInt((float)ViewRect.Y * VolumetricRenderTargetScale);
+		FIntPoint SceneTexturesExtent = ComputeVolumetricBufferResolution(SceneTextures.Config.Extent, VRTMode);
+		FIntPoint ViewRect = ComputeVolumetricBufferResolution(ViewInfo.ViewRect.Size(), VRTMode);
+
 		VolumetricCloudRT.Initialise(	// TODO this is going to reallocate a buffer each time dynamic resolution scaling is applied 
+			SceneTexturesExtent,
 			ViewRect,
 			VRTMode,
 			CVarVolumetricRenderTargetUpsamplingMode.GetValueOnAnyThread(),
@@ -473,8 +527,9 @@ void InitVolumetricRenderTargetForViews(FRDGBuilder& GraphBuilder, TArrayView<FV
 
 		FViewUniformShaderParameters ViewVolumetricCloudRTParameters = *ViewInfo.CachedViewUniformShaderParameters;
 		{
-			const FIntPoint& VolumetricReconstructResolution = VolumetricCloudRT.GetCurrentVolumetricReconstructRTResolution();
 			const FIntPoint& VolumetricTracingResolution = VolumetricCloudRT.GetCurrentVolumetricTracingRTResolution();
+			const FIntPoint& VolumetricReconstructViewRect = VolumetricCloudRT.GetDstVolumetricReconstructViewRect();
+			const FIntPoint& VolumetricTracingViewRect = VolumetricCloudRT.GetCurrentVolumetricTracingViewRect();
 			const FIntPoint& CurrentPixelOffset = VolumetricCloudRT.GetCurrentTracingPixelOffset();
 			const uint32 VolumetricReconstructRTDownSample = VolumetricCloudRT.GetVolumetricReconstructRTDownsampleFactor();
 			const uint32 VolumetricTracingRTDownSample = VolumetricCloudRT.GetVolumetricTracingRTDownsampleFactor();
@@ -484,18 +539,16 @@ void InitVolumetricRenderTargetForViews(FRDGBuilder& GraphBuilder, TArrayView<FV
 			FViewMatrices ViewMatrices = ViewInfo.ViewMatrices;
 			ViewMatrices.HackRemoveTemporalAAProjectionJitter();
 
-			float DownSampleFactor = float(VolumetricReconstructRTDownSample * VolumetricTracingRTDownSample);
-
 			// Offset to the correct half resolution pixel
 			FVector2D CenterCoord = FVector2D(VolumetricReconstructRTDownSample / 2.0f);
 			FVector2D TargetCoord = FVector2D(CurrentPixelOffset) + FVector2D(0.5f, 0.5f);
-			FVector2D OffsetCoord = (TargetCoord - CenterCoord) * (FVector2D(-2.0f, 2.0f) / FVector2D(VolumetricReconstructResolution));
+			FVector2D OffsetCoord = (TargetCoord - CenterCoord) * (FVector2D(-2.0f, 2.0f) / FVector2D(VolumetricReconstructViewRect));
 			ViewMatrices.HackAddTemporalAAProjectionJitter(OffsetCoord);
 
 			ViewInfo.SetupViewRectUniformBufferParameters(
 				ViewVolumetricCloudRTParameters,
 				VolumetricTracingResolution,
-				FIntRect(0, 0, VolumetricTracingResolution.X, VolumetricTracingResolution.Y),
+				FIntRect(0, 0, VolumetricTracingViewRect.X, VolumetricTracingViewRect.Y),
 				ViewMatrices,
 				ViewInfo.PrevViewInfo.ViewMatrices // This could also be changed if needed
 			);
@@ -545,12 +598,14 @@ class FReconstructVolumetricRenderTargetPS : public FGlobalShader
 		SHADER_PARAMETER(FIntPoint, ViewViewRectMin)
 		SHADER_PARAMETER(int32, DownSampleFactor)
 		SHADER_PARAMETER(int32, VolumetricRenderTargetMode)
+		SHADER_PARAMETER(FVector2f, TracingVolumetricTextureUVScale)
 		SHADER_PARAMETER(FUintVector4, TracingVolumetricTextureValidCoordRect)
 		SHADER_PARAMETER(FVector4f, TracingVolumetricTextureValidUvRect)
 		SHADER_PARAMETER(FUintVector4, PreviousFrameVolumetricTextureValidCoordRect)
 		SHADER_PARAMETER(FVector4f, PreviousFrameVolumetricTextureValidUvRect)
 		SHADER_PARAMETER(float, MinimumDistanceKmToEnableReprojection)
 		SHADER_PARAMETER(float, HistoryPreExposureCorrection)
+		SHADER_PARAMETER(FVector2f, PreviousFrameVolumetricTextureUVScale)
 	END_SHADER_PARAMETER_STRUCT()
 
 	static FPermutationDomain RemapPermutation(FPermutationDomain PermutationVector)
@@ -651,10 +706,13 @@ void ReconstructVolumetricRenderTarget(
 				HalfResolutionDepthCheckerboardMinMaxTexture :
 				SceneDepthTexture);
 
-		GetTextureSafeUvCoordBound(SrcTracingVolumetric, PassParameters->TracingVolumetricTextureValidCoordRect, PassParameters->TracingVolumetricTextureValidUvRect);
-		GetTextureSafeUvCoordBound(PreviousFrameVolumetricTexture, PassParameters->PreviousFrameVolumetricTextureValidCoordRect, PassParameters->PreviousFrameVolumetricTextureValidUvRect);
+		const FIntPoint CurrentVolumetricTracingViewRect = VolumetricCloudRT.GetCurrentVolumetricTracingViewRect();
+		const FIntPoint SrcVolumetricReconstructViewRect = VolumetricCloudRT.GetSrcVolumetricReconstructViewRect();
+
+		GetTextureSafeUvCoordBound(CurrentVolumetricTracingViewRect, SrcTracingVolumetric, PassParameters->TracingVolumetricTextureValidCoordRect, PassParameters->TracingVolumetricTextureValidUvRect, PassParameters->TracingVolumetricTextureUVScale);
+		GetTextureSafeUvCoordBound(SrcVolumetricReconstructViewRect, PreviousFrameVolumetricTexture, PassParameters->PreviousFrameVolumetricTextureValidCoordRect, PassParameters->PreviousFrameVolumetricTextureValidUvRect, PassParameters->PreviousFrameVolumetricTextureUVScale);
 
-		FIntVector DstVolumetricSize = DstVolumetric->Desc.GetSize();
+		FIntPoint DstVolumetricSize = VolumetricCloudRT.GetDstVolumetricReconstructViewRect();
 		FVector2D DstVolumetricTextureSize = FVector2D(float(DstVolumetricSize.X), float(DstVolumetricSize.Y));
 		FVector2D PreviousVolumetricTextureSize = FVector2D(float(PreviousFrameVolumetricTexture->Desc.GetSize().X), float(PreviousFrameVolumetricTexture->Desc.GetSize().Y));
 		PassParameters->DstVolumetricTextureSizeAndInvSize = FVector4f(DstVolumetricTextureSize.X, DstVolumetricTextureSize.Y, 1.0f / DstVolumetricTextureSize.X, 1.0f / DstVolumetricTextureSize.Y);
@@ -745,7 +803,8 @@ static FVector2f GetCompositionFullResolutionToVolumetricBufferResolutionScale(u
 	return FVector2f(1.0f / float(GetMainDownsampleFactor(VRTMode)), float(GetMainDownsampleFactor(VRTMode)));
 }
 
-static void GetCompositionCloudTextures(uint32 VRTMode, FVolumetricRenderTargetViewStateData& VolumetricCloudRT, FRDGBuilder& GraphBuilder, FRDGTextureRef& VolumetricTexture, FRDGTextureRef& VolumetricDepthTexture)
+static void GetCompositionCloudTextures(uint32 VRTMode, FVolumetricRenderTargetViewStateData& VolumetricCloudRT, FRDGBuilder& GraphBuilder, FRDGTextureRef& VolumetricTexture, FRDGTextureRef& VolumetricDepthTexture,
+										FIntPoint& VolumetricViewRect)
 {
 	if (VRTMode == 1 || VRTMode == 3)
 	{
@@ -754,11 +813,13 @@ static void GetCompositionCloudTextures(uint32 VRTMode, FVolumetricRenderTargetV
 		// No reconstruction needed.
 		VolumetricTexture = VolumetricCloudRT.GetOrCreateVolumetricTracingRT(GraphBuilder);
 		VolumetricDepthTexture = VolumetricCloudRT.GetOrCreateVolumetricTracingRTDepth(GraphBuilder);
+		VolumetricViewRect = VolumetricCloudRT.GetCurrentVolumetricTracingViewRect();
 	}
 	else
 	{
 		VolumetricTexture = VolumetricCloudRT.GetOrCreateDstVolumetricReconstructRT(GraphBuilder);
 		VolumetricDepthTexture = VolumetricCloudRT.GetOrCreateDstVolumetricReconstructRTDepth(GraphBuilder);
+		VolumetricViewRect = VolumetricCloudRT.GetDstVolumetricReconstructViewRect();
 	}
 }
 
@@ -818,7 +879,8 @@ void ComposeVolumetricRenderTargetOverScene(
 
 		FRDGTextureRef VolumetricTexture = nullptr;
 		FRDGTextureRef VolumetricDepthTexture = nullptr;
-		GetCompositionCloudTextures(VRTMode, VolumetricCloudRT, GraphBuilder, VolumetricTexture, VolumetricDepthTexture);
+		FIntPoint VolumetricViewRect(FIntPoint::ZeroValue);
+		GetCompositionCloudTextures(VRTMode, VolumetricCloudRT, GraphBuilder, VolumetricTexture, VolumetricDepthTexture, VolumetricViewRect);
 
 		// We only support MSAA up to 8 sample and in forward
 		check(SceneDepthTexture->Desc.NumSamples <= 8);
@@ -849,7 +911,8 @@ void ComposeVolumetricRenderTargetOverScene(
 		PassParameters->SceneTextures = SceneTextures.UniformBuffer;
 		PassParameters->Substrate = Substrate::BindSubstrateGlobalUniformParameters(ViewInfo);
 		PassParameters->ForwardShadingEnable = bForwardShading ? 1 : 0;
-		GetTextureSafeUvCoordBound(PassParameters->VolumetricTexture, PassParameters->VolumetricTextureValidCoordRect, PassParameters->VolumetricTextureValidUvRect);
+		FVector2f DummyUVScale;
+		GetTextureSafeUvCoordBound(VolumetricViewRect, PassParameters->VolumetricTexture, PassParameters->VolumetricTextureValidCoordRect, PassParameters->VolumetricTextureValidUvRect, DummyUVScale);
 
 		PassParameters->WaterLinearDepthTexture = WaterPassData.DepthTexture;
 		PassParameters->WaterLinearDepthSampler = TStaticSamplerState<SF_Point>::GetRHI();
@@ -909,7 +972,8 @@ void ComposeVolumetricRenderTargetOverSceneUnderWater(
 
 		FRDGTextureRef VolumetricTexture = nullptr;
 		FRDGTextureRef VolumetricDepthTexture = nullptr;
-		GetCompositionCloudTextures(VRTMode, VolumetricCloudRT, GraphBuilder, VolumetricTexture, VolumetricDepthTexture);
+		FIntPoint VolumetricViewRect(FIntPoint::ZeroValue);
+		GetCompositionCloudTextures(VRTMode, VolumetricCloudRT, GraphBuilder, VolumetricTexture, VolumetricDepthTexture, VolumetricViewRect);
 
 		FComposeVolumetricRTOverScenePS::FPermutationDomain PermutationVector;
 		PermutationVector.Set<FComposeVolumetricRTOverScenePS::FUpsamplingMode>(UpsamplingMode);
@@ -933,7 +997,8 @@ void ComposeVolumetricRenderTargetOverSceneUnderWater(
 																		WaterPassViewData.ViewRect.Max.X, WaterPassViewData.ViewRect.Max.Y);
 		PassParameters->SceneTextures = SceneTextures.UniformBuffer;
 		PassParameters->Substrate = Substrate::BindSubstrateGlobalUniformParameters(ViewInfo);
-		GetTextureSafeUvCoordBound(PassParameters->VolumetricTexture, PassParameters->VolumetricTextureValidCoordRect, PassParameters->VolumetricTextureValidUvRect);
+		FVector2f DummyUVScale;
+		GetTextureSafeUvCoordBound(VolumetricViewRect, PassParameters->VolumetricTexture, PassParameters->VolumetricTextureValidCoordRect, PassParameters->VolumetricTextureValidUvRect, DummyUVScale);
 
 		FVector2D VolumetricTextureSize = FVector2D(float(VolumetricTexture->Desc.GetSize().X), float(VolumetricTexture->Desc.GetSize().Y));
 		PassParameters->VolumetricTextureSizeAndInvSize = FVector4f(VolumetricTextureSize.X, VolumetricTextureSize.Y, 1.0f / VolumetricTextureSize.X, 1.0f / VolumetricTextureSize.Y);
@@ -976,7 +1041,8 @@ void ComposeVolumetricRenderTargetOverSceneForVisualization(
 
 		FRDGTextureRef VolumetricTexture = nullptr;
 		FRDGTextureRef VolumetricDepthTexture = nullptr;
-		GetCompositionCloudTextures(VRTMode, VolumetricCloudRT, GraphBuilder, VolumetricTexture, VolumetricDepthTexture);
+		FIntPoint VolumetricViewRect(FIntPoint::ZeroValue);
+		GetCompositionCloudTextures(VRTMode, VolumetricCloudRT, GraphBuilder, VolumetricTexture, VolumetricDepthTexture, VolumetricViewRect);
 
 		FComposeVolumetricRTOverScenePS::FPermutationDomain PermutationVector;
 		PermutationVector.Set<FComposeVolumetricRTOverScenePS::FUpsamplingMode>(0);
@@ -995,7 +1061,8 @@ void ComposeVolumetricRenderTargetOverSceneForVisualization(
 		PassParameters->FullResolutionToVolumetricBufferResolutionScale = GetCompositionFullResolutionToVolumetricBufferResolutionScale(VRTMode) * GetVolumetricBufferResolutionScale(VRTMode);
 		PassParameters->SceneTextures = SceneTextures.UniformBuffer;
 		PassParameters->Substrate = Substrate::BindSubstrateGlobalUniformParameters(ViewInfo);
-		GetTextureSafeUvCoordBound(PassParameters->VolumetricTexture, PassParameters->VolumetricTextureValidCoordRect, PassParameters->VolumetricTextureValidUvRect);
+		FVector2f DummyUVScale;
+		GetTextureSafeUvCoordBound(VolumetricViewRect, PassParameters->VolumetricTexture, PassParameters->VolumetricTextureValidCoordRect, PassParameters->VolumetricTextureValidUvRect, DummyUVScale);
 
 		PassParameters->WaterLinearDepthTexture = GSystemTextures.GetBlackDummy(GraphBuilder);
 
diff --git a/Engine/Source/Runtime/Renderer/Private/VolumetricRenderTarget.h b/Engine/Source/Runtime/Renderer/Private/VolumetricRenderTarget.h
--- a/Engine/Source/Runtime/Renderer/Private/VolumetricRenderTarget.h
+++ b/Engine/Source/Runtime/Renderer/Private/VolumetricRenderTarget.h
@@ -17,7 +17,7 @@ class FViewInfo;
 struct FSceneWithoutWaterTextures;
 struct FMinimalSceneTextures;
 
-void InitVolumetricRenderTargetForViews(FRDGBuilder& GraphBuilder, TArrayView<FViewInfo> Views);
+void InitVolumetricRenderTargetForViews(FRDGBuilder& GraphBuilder, TArrayView<FViewInfo> Views, const FMinimalSceneTextures& SceneTextures);
 void ResetVolumetricRenderTargetForViews(FRDGBuilder& GraphBuilder, TArrayView<FViewInfo> Views);
 
 void ReconstructVolumetricRenderTarget(
diff --git a/Engine/Source/Runtime/Renderer/Private/VolumetricRenderTargetViewStateData.h b/Engine/Source/Runtime/Renderer/Private/VolumetricRenderTargetViewStateData.h
--- a/Engine/Source/Runtime/Renderer/Private/VolumetricRenderTargetViewStateData.h
+++ b/Engine/Source/Runtime/Renderer/Private/VolumetricRenderTargetViewStateData.h
@@ -22,6 +22,7 @@ public:
 	~FVolumetricRenderTargetViewStateData();
 
 	void Initialise(
+		FIntPoint& TextureResolutionIn,
 		FIntPoint& ViewRectResolutionIn,
 		int32 Mode,
 		int32 UpsamplingMode,
@@ -42,20 +43,26 @@ public:
 	FRDGTextureRef GetOrCreateVolumetricTracingRT(FRDGBuilder& GraphBuilder);
 	FRDGTextureRef GetOrCreateVolumetricSecondaryTracingRT(FRDGBuilder& GraphBuilder);
 	FRDGTextureRef GetOrCreateVolumetricTracingRTDepth(FRDGBuilder& GraphBuilder);
+	FVector2f GetVolumetricTracingUVScale() const;
 
 	FRDGTextureRef GetDstVolumetricReconstructRT(FRDGBuilder& GraphBuilder);
 	FRDGTextureRef GetOrCreateDstVolumetricReconstructRT(FRDGBuilder& GraphBuilder);
 	FRDGTextureRef GetOrCreateDstVolumetricReconstructRTDepth(FRDGBuilder& GraphBuilder);
+	const FIntPoint& GetDstVolumetricReconstructViewRect() const;
+	FVector2f GetDstVolumetricReconstructUVScale() const;
 
 	TRefCountPtr<IPooledRenderTarget> GetDstVolumetricReconstructRT();
 	TRefCountPtr<IPooledRenderTarget> GetDstVolumetricReconstructRTDepth();
 
 	FRDGTextureRef GetOrCreateSrcVolumetricReconstructRT(FRDGBuilder& GraphBuilder);
 	FRDGTextureRef GetOrCreateSrcVolumetricReconstructRTDepth(FRDGBuilder& GraphBuilder);
+	const FIntPoint& GetSrcVolumetricReconstructViewRect() const;
+
 
 	bool GetHistoryValid() const { return bHistoryValid; }
 	const FIntPoint& GetCurrentVolumetricReconstructRTResolution() const { return VolumetricReconstructRTResolution; }
 	const FIntPoint& GetCurrentVolumetricTracingRTResolution() const { return VolumetricTracingRTResolution; }
+	const FIntPoint& GetCurrentVolumetricTracingViewRect() const { return VolumetricTracingViewRect; }
 	const FIntPoint& GetCurrentTracingPixelOffset() const { return CurrentPixelOffset; }
 	const uint32 GetNoiseFrameIndexModPattern() const { return NoiseFrameIndexModPattern; }
 
@@ -88,10 +95,12 @@ private:
 	FIntPoint FullResolution;
 	FIntPoint VolumetricReconstructRTResolution;
 	FIntPoint VolumetricTracingRTResolution;
+	FIntPoint VolumetricTracingViewRect;
 
 	static constexpr uint32 kRenderTargetCount = 2;
 	TRefCountPtr<IPooledRenderTarget> VolumetricReconstructRT[kRenderTargetCount];
 	TRefCountPtr<IPooledRenderTarget> VolumetricReconstructRTDepth[kRenderTargetCount];
+	FIntPoint VolumetricReconstructViewRect[kRenderTargetCount];
 
 	TRefCountPtr<IPooledRenderTarget> VolumetricTracingRT;
 	TRefCountPtr<IPooledRenderTarget> VolumetricSecondaryTracingRT;
