Fix crash on platform not supporting R16uint format for UAV write.

diff --git a/Engine/Source/Runtime/Renderer/Private/OIT/OIT.cpp b/Engine/Source/Runtime/Renderer/Private/OIT/OIT.cpp
--- a/Engine/Source/Runtime/Renderer/Private/OIT/OIT.cpp
+++ b/Engine/Source/Runtime/Renderer/Private/OIT/OIT.cpp
@@ -490,7 +490,7 @@ public:
 		{
 			return false;
 		}
-		return true;
+		return OIT::IsSortedTrianglesEnabled(Parameters.Platform);
 	}
 
 	static void ModifyCompilationEnvironment(const FGlobalShaderPermutationParameters& Parameters, FShaderCompilerEnvironment& OutEnvironment)
@@ -525,6 +525,10 @@ class FOITSortTriangleIndex_AllocateCS : public FGlobalShader
 		SHADER_PARAMETER_RDG_BUFFER_UAV(RWBuffer<uint>, SliceOffsetsBuffer)
 	END_SHADER_PARAMETER_STRUCT()
 public:
+	static bool ShouldCompilePermutation(const FGlobalShaderPermutationParameters& Parameters) 
+	{ 
+		return OIT::IsSortedTrianglesEnabled(Parameters.Platform);
+	}
 	static void ModifyCompilationEnvironment(const FGlobalShaderPermutationParameters& Parameters, FShaderCompilerEnvironment& OutEnvironment)
 	{
 		FGlobalShader::ModifyCompilationEnvironment(Parameters, OutEnvironment);
@@ -556,6 +560,10 @@ class FOITSortTriangleIndex_WriteOutCS : public FGlobalShader
 		SHADER_PARAMETER_UAV(RWBuffer<uint>, OutIndexBuffer)
 	END_SHADER_PARAMETER_STRUCT()
 public:
+	static bool ShouldCompilePermutation(const FGlobalShaderPermutationParameters& Parameters) 
+	{ 
+		return OIT::IsSortedTrianglesEnabled(Parameters.Platform);
+	}
 	static void ModifyCompilationEnvironment(const FGlobalShaderPermutationParameters& Parameters, FShaderCompilerEnvironment& OutEnvironment)
 	{
 		FGlobalShader::ModifyCompilationEnvironment(Parameters, OutEnvironment);
@@ -849,7 +857,7 @@ namespace OIT
 {
 	bool IsSortedTrianglesEnabled(EShaderPlatform InPlatform)
 	{
-		return CVarOIT_SortedTriangles_Enable_Project.GetValueOnAnyThread() > 0;
+		return CVarOIT_SortedTriangles_Enable_Project.GetValueOnAnyThread() > 0 && !IsOpenGLPlatform(InPlatform);
 	}
 
 	bool IsSortedPixelsEnabledForProject(EShaderPlatform InPlatform)
