Added new ViewData.ush Shader API to produce & get per-view instance data as a step towards nanite skinning LOD distance and VSM invalidation management.
- Added FInstanceViewData / GetInstanceViewData to shader API to get per-view instance data during rendering.
- Moved view management from the instance culling manager to the new FRendererViewDataManager.
- Added GPU-side bit vector to denote whether an instance is deforming (not being used as of yet).
- Added UpdateViewData callback to run GPU view data setup in the scene extensions.

diff --git a/Engine/Shaders/Private/ViewData.ush b/Engine/Shaders/Private/ViewData.ush
new file mode 100644
--- /dev/null
+++ b/Engine/Shaders/Private/ViewData.ush
@@ -0,0 +1,139 @@
+// Copyright Epic Games, Inc. All Rights Reserved.
+
+#pragma once
+
+#include "Common.ush"
+#include "WaveOpUtil.ush"
+
+#include "/Engine/Shared/ViewDefinitions.h"
+
+// define the default access mode to be via the Scene UB
+#ifndef VIEW_DATA_ACCESS_MODE
+#define VIEW_DATA_ACCESS_MODE VIEW_DATA_ACCESS_SCENE_UB
+#endif
+
+// API for shaders that modify the view data.
+#if VIEW_DATA_ACCESS_MODE == VIEW_DATA_ACCESS_RW 
+
+int MaxPersistentViewId;
+uint InstanceMaskWordStride;
+uint NumSceneRendererPrimaryViews;
+RWStructuredBuffer<uint> OutDeformingInstanceViewMask;
+StructuredBuffer<int> PersistentIdToIndexMap;
+
+void MarkInstanceAsDeforming(uint InstanceId, uint ViewIndex)
+{
+	uint Mask = 1u << (InstanceId % 32u);
+	uint MaskWordOffset = InstanceId / 32u;
+	// Mark for the particular view
+	InterlockedOr(OutDeformingInstanceViewMask[(InstanceMaskWordStride * ViewIndex) + MaskWordOffset], Mask);
+}
+
+void MarkInstanceAsDeformingWithViewMask(uint InstanceId, uint ViewMask)
+{
+	uint Mask = 1u << (InstanceId % 32u);
+	uint MaskWordOffset = InstanceId / 32u;
+
+	if (ViewMask != 0u)
+	{
+		// Mark for the view in the mask
+		for (uint ViewIndex = 0u; ViewIndex < NumSceneRendererPrimaryViews; ++ViewIndex)
+		{
+			if (((1u << ViewIndex) & ViewMask) !=  0u)
+			{
+				InterlockedOr(OutDeformingInstanceViewMask[(InstanceMaskWordStride * ViewIndex) + MaskWordOffset], Mask);
+			}
+		}
+	}
+}
+
+#endif 
+
+struct FViewData
+{
+	int MaxPersistentViewId;
+	uint InstanceMaskWordStride;
+	uint NumSceneRendererPrimaryViews;
+};
+
+FViewData GetViewData()
+{
+	FViewData ViewData;
+#if VIEW_DATA_ACCESS_MODE == VIEW_DATA_ACCESS_RW 
+	ViewData.MaxPersistentViewId = MaxPersistentViewId;
+	ViewData.InstanceMaskWordStride = InstanceMaskWordStride;
+	ViewData.NumSceneRendererPrimaryViews = NumSceneRendererPrimaryViews;
+#else
+	ViewData.MaxPersistentViewId = Scene.ViewData.MaxPersistentViewId;
+	ViewData.InstanceMaskWordStride = Scene.ViewData.InstanceMaskWordStride;
+	ViewData.NumSceneRendererPrimaryViews = Scene.ViewData.NumSceneRendererPrimaryViews;
+#endif
+	return ViewData;
+}
+
+/**
+ * Returns a 32-bit word with instance-deforming flags for 32 consecutive instances starting from InstanceId = MaskWordOffset * 32
+ */
+uint LoadInstanceDeformingWord(uint MaskWordOffset, int ViewIndex)
+{
+	uint BaseOffset = ViewIndex * GetViewData().InstanceMaskWordStride;
+#if VIEW_DATA_ACCESS_MODE == VIEW_DATA_ACCESS_RW 
+	return OutDeformingInstanceViewMask[BaseOffset + MaskWordOffset];
+#else // VIEW_DATA_ACCESS_SCENE_UB
+	return Scene.ViewData.DeformingInstanceViewMask[BaseOffset + MaskWordOffset];
+#endif
+}
+
+/**
+ * Returns true if the instance is marked as deforming. 
+ */
+bool IsInstanceDeforming(uint InstanceId, int SceneRendererPrimaryViewId)
+{
+	uint Mask = 1u << (InstanceId % 32u);
+	uint MaskWordOffset = InstanceId / 32u;
+	return (LoadInstanceDeformingWord(MaskWordOffset, SceneRendererPrimaryViewId) & Mask) != 0u;
+}
+
+/**
+ * Map a persistent view ID to a current SceneRendererPrimaryViewId, returns -1 if the view is not current (part of the ongoing scenerenderer)
+ */
+int RemapPersisitentViewId(int PersistentViewId)
+{
+	if (PersistentViewId < GetViewData().MaxPersistentViewId)
+	{
+#if VIEW_DATA_ACCESS_MODE == VIEW_DATA_ACCESS_RW 
+		return int(PersistentIdToIndexMap[PersistentViewId]) - 1;
+#else
+		return int(Scene.ViewData.PersistentIdToIndexMap[PersistentViewId]) - 1;
+#endif
+	}
+	return -1;
+}
+
+/**
+ */
+struct FInstanceViewData
+{
+	bool bIsValid;
+	bool bIsDeforming;
+
+	// Convenience / derived data 
+	uint InstanceId;
+	int SceneRendererPrimaryViewId;
+};
+
+FInstanceViewData GetInstanceViewData(uint InstanceId, int SceneRendererPrimaryViewId)
+{
+	checkSlow(SceneRendererPrimaryViewId < GetViewData().NumSceneRendererPrimaryViews && SceneRendererPrimaryViewId >= 0);
+
+	FInstanceViewData InstanceViewData;
+	
+	// Convenience / derived data 
+	InstanceViewData.SceneRendererPrimaryViewId = SceneRendererPrimaryViewId;
+	InstanceViewData.InstanceId = InstanceId;
+	
+	InstanceViewData.bIsValid = true;
+
+	InstanceViewData.bIsDeforming = IsInstanceDeforming(InstanceId, SceneRendererPrimaryViewId);
+	return InstanceViewData;
+}
diff --git a/Engine/Shaders/Shared/ViewDefinitions.h b/Engine/Shaders/Shared/ViewDefinitions.h
new file mode 100644
--- /dev/null
+++ b/Engine/Shaders/Shared/ViewDefinitions.h
@@ -0,0 +1,6 @@
+// Copyright Epic Games, Inc. All Rights Reserved.
+
+#pragma once
+
+#define VIEW_DATA_ACCESS_SCENE_UB 0
+#define VIEW_DATA_ACCESS_RW 1
diff --git a/Engine/Source/Runtime/Renderer/Private/DeferredShadingRenderer.cpp b/Engine/Source/Runtime/Renderer/Private/DeferredShadingRenderer.cpp
--- a/Engine/Source/Runtime/Renderer/Private/DeferredShadingRenderer.cpp
+++ b/Engine/Source/Runtime/Renderer/Private/DeferredShadingRenderer.cpp
@@ -104,6 +104,7 @@
 #include "EnvironmentComponentsFlags.h"
 #include "GenerateMips.h"
 #include "Froxel/Froxel.h"
+#include "ViewData.h"
 
 #if !UE_BUILD_SHIPPING
 #include "RenderCaptureInterface.h"
@@ -1591,7 +1592,8 @@ void FDeferredShadingSceneRenderer::Render(FRDGBuilder& GraphBuilder)
 		? FExclusiveDepthStencil::DepthRead_StencilWrite
 		: FExclusiveDepthStencil::DepthWrite_StencilWrite;
 
-	FInstanceCullingManager& InstanceCullingManager = *GraphBuilder.AllocObject<FInstanceCullingManager>(GraphBuilder, *Scene, GetSceneUniforms(), AllViews);
+	FRendererViewDataManager& ViewDataManager = *GraphBuilder.AllocObject<FRendererViewDataManager>(GraphBuilder, *Scene, GetSceneUniforms(), AllViews);
+	FInstanceCullingManager& InstanceCullingManager = *GraphBuilder.AllocObject<FInstanceCullingManager>(GraphBuilder, *Scene, GetSceneUniforms(), ViewDataManager);
 
 	::Substrate::PreInitViews(*Scene);
 
@@ -1744,6 +1746,9 @@ void FDeferredShadingSceneRenderer::Render(FRDGBuilder& GraphBuilder)
 
 		SceneCullingRenderer.DebugRender(GraphBuilder, Views);
 
+		// Must be called after all views have flushed the dynamic primitives.
+		ViewDataManager.InitInstanceState(GraphBuilder);
+
 		if (Views.Num() > 0)
 		{
 			FViewInfo& View = Views[0];
@@ -1751,6 +1756,9 @@ void FDeferredShadingSceneRenderer::Render(FRDGBuilder& GraphBuilder)
 		}
 	}
 
+	//SceneCullingInfo.SceneUniformBuffer = GetSceneUniforms().GetBuffer(GraphBuilder);
+	GetSceneExtensionsRenderers().UpdateViewData(GraphBuilder, ViewDataManager);
+
 	// Allow scene extensions to affect the scene uniform buffer after GPU scene has fully updated
 	GetSceneExtensionsRenderers().UpdateSceneUniformBuffer(GraphBuilder, GetSceneUniforms());
 
diff --git a/Engine/Source/Runtime/Renderer/Private/InstanceCulling/InstanceCullingContext.cpp b/Engine/Source/Runtime/Renderer/Private/InstanceCulling/InstanceCullingContext.cpp
--- a/Engine/Source/Runtime/Renderer/Private/InstanceCulling/InstanceCullingContext.cpp
+++ b/Engine/Source/Runtime/Renderer/Private/InstanceCulling/InstanceCullingContext.cpp
@@ -22,6 +22,7 @@
 #include "UnrealEngine.h"
 #include "RHIBreadcrumbs.h"
 #include "Materials/MaterialRenderProxy.h"
+#include "ViewData.h"
 
 static TAutoConsoleVariable<int32> CVarCullInstances(
 	TEXT("r.CullInstances"),
@@ -568,7 +569,7 @@ public:
 		SHADER_PARAMETER(uint32, GPUSceneNumLightmapDataItems)
 
 		SHADER_PARAMETER_STRUCT_INCLUDE(FInstanceProcessingGPULoadBalancer::FShaderParameters, LoadBalancerParameters)
-		SHADER_PARAMETER_STRUCT_INCLUDE(FInstanceCullingManager::FShaderParameters, InstanceCullingManager)
+		SHADER_PARAMETER_STRUCT_INCLUDE(RendererViewData::FCullingShaderParameters, ViewDataCullingParameters)
 
 		SHADER_PARAMETER_RDG_BUFFER_SRV(StructuredBuffer< uint32 >, DrawCommandDescs)
 		SHADER_PARAMETER_RDG_BUFFER_SRV(StructuredBuffer< FInstanceCullingContext::FPayloadData >, InstanceCullingPayloads)
@@ -839,15 +840,15 @@ void FInstanceCullingContext::BuildRenderingCommandsInternal(
 	GraphBuilder.QueueBufferUpload(InstanceIdOffsetBufferRDG, InstanceIdOffsets.GetData(), InstanceIdOffsets.GetTypeSize() * InstanceIdOffsets.Num());
 
 	PassParametersTmp.ViewIds = GraphBuilder.CreateSRV(ViewIdsBuffer);
-	PassParametersTmp.InstanceCullingManager.NumCullingViews = 0;
+	PassParametersTmp.ViewDataCullingParameters.NumCullingViews = 0;
 	if ((bCullInstances || bAllowWPODisable) && InstanceCullingManager)
 	{
-		PassParametersTmp.InstanceCullingManager = InstanceCullingManager->GetShaderParameters(GraphBuilder);
+		PassParametersTmp.ViewDataCullingParameters = InstanceCullingManager->ViewDataManager.GetCullingParameters(GraphBuilder);
 
 #if DO_CHECK
 		for (int32 ViewId : ViewIds)
 		{
-			checkf(ViewId < int32(PassParametersTmp.InstanceCullingManager.NumCullingViews), TEXT("Attempting to process a culling context that references a view that has not been uploaded yet."));
+			checkf(ViewId < int32(PassParametersTmp.ViewDataCullingParameters.NumCullingViews), TEXT("Attempting to process a culling context that references a view that has not been uploaded yet."));
 		}
 #endif 
 	}
@@ -1183,7 +1184,7 @@ FInstanceCullingDeferredContext *FInstanceCullingContext::CreateDeferredContext(
 	PassParametersTmp.InstanceIdOffsetBuffer = GraphBuilder.CreateSRV(InstanceIdOffsetBuffer, PF_R32_UINT);	
 	if (bCullInstances || bAllowWPODisable)
 	{
-		PassParametersTmp.InstanceCullingManager = InstanceCullingManager.GetShaderParameters(GraphBuilder);
+		PassParametersTmp.ViewDataCullingParameters = InstanceCullingManager.ViewDataManager.GetCullingParameters(GraphBuilder);
 	}
 
 	// Compaction parameters
@@ -1204,7 +1205,7 @@ FInstanceCullingDeferredContext *FInstanceCullingContext::CreateDeferredContext(
 	}
 
 	// Record the number of culling views to be able to check that no views referencing out-of bounds views are queued up
-	DeferredContext->NumCullingViews = InstanceCullingManager.NumRegisteredViews.load(std::memory_order_relaxed);
+	DeferredContext->NumCullingViews = InstanceCullingManager.ViewDataManager.GetNumCullingViews();
 
 	for (uint32 BinIndex = 0U; BinIndex < NumBins; ++BinIndex)
 	{
diff --git a/Engine/Source/Runtime/Renderer/Private/InstanceCulling/InstanceCullingManager.cpp b/Engine/Source/Runtime/Renderer/Private/InstanceCulling/InstanceCullingManager.cpp
--- a/Engine/Source/Runtime/Renderer/Private/InstanceCulling/InstanceCullingManager.cpp
+++ b/Engine/Source/Runtime/Renderer/Private/InstanceCulling/InstanceCullingManager.cpp
@@ -10,6 +10,7 @@
 #include "RenderGraphUtils.h"
 #include "SceneRendering.h"
 #include "ScenePrivate.h"
+#include "ViewData.h"
 #include "ProfilingDebugging/CpuProfilerTrace.h"
 
 #include "InstanceCulling/InstanceCullingContext.h"
@@ -21,38 +22,19 @@ static FAutoConsoleVariableRef CVarAllowBatchedBuildRenderingCommands(
 	TEXT("Whether to allow batching BuildRenderingCommands for GPU instance culling"),
 	ECVF_RenderThreadSafe);
 
-FInstanceCullingManager::FInstanceCullingManager(FRDGBuilder& GraphBuilder, const FScene& InScene, FSceneUniformBuffer& InSceneUniforms, TArray<FViewInfo*> &InOutPrimaryViews)
-	: Scene(InScene), GPUScene(InScene.GPUScene), SceneUniforms(InSceneUniforms), bIsEnabled(GPUScene.IsEnabled())
+FInstanceCullingManager::FInstanceCullingManager(FRDGBuilder& GraphBuilder, const FScene& InScene, FSceneUniformBuffer& InSceneUniforms, FRendererViewDataManager& InViewDataManager)
+	: Scene(InScene), GPUScene(InScene.GPUScene), SceneUniforms(InSceneUniforms), ViewDataManager(InViewDataManager), bIsEnabled(GPUScene.IsEnabled())
 {
 	DummyUniformBuffer = FInstanceCullingContext::CreateDummyInstanceCullingUniformBuffer(GraphBuilder);
-
-	if (bIsEnabled)
+	
+	for (FViewInfo *ViewInfo : ViewDataManager.GetRegisteredPrimaryViews())
 	{
-		NumSceneRendererPrimaryViews = InOutPrimaryViews.Num();
-		AllocateViews(NumSceneRendererPrimaryViews);
-		for (int32 ViewIndex = 0; ViewIndex < InOutPrimaryViews.Num(); ++ViewIndex)
-		{
-			FViewInfo &View = *InOutPrimaryViews[ViewIndex];
-			View.GPUSceneViewId = ViewIndex;
-			ensure(View.GPUSceneViewId == RegisterPrimaryView(View));
-		}
-
-		ShaderParameters.InViews = nullptr;
-		ShaderParameters.NumCullingViews = 0;
-
-		// Create a deferred buffer and flush when all views are registered. This doesn't work when immediate mode is used, and in that case we need to re-upload partial buffers when the buffer is requested
-		if (!GraphBuilder.IsImmediateMode())
+		if (ViewInfo->PrevViewInfo.HZB != nullptr)
 		{
-			CullingViewsRDG = CreateStructuredBuffer(GraphBuilder, TEXT("InstanceCullingManager.Views"), [this]() -> TArray<Nanite::FPackedView>& { return CullingViews; });
+			check(IsInRenderingThread());
 
-			ShaderParameters.InViews = GraphBuilder.CreateSRV(CullingViewsRDG);
-			ShaderParameters.NumCullingViews = NumSceneRendererPrimaryViews;
-		}
-		else
-		{
-			FlushRegisteredViews(GraphBuilder);
+			ViewPrevHZBs.AddUnique(ViewInfo->PrevViewInfo.HZB);
 		}
-		ShaderParameters.NumSceneRendererPrimaryViews = NumSceneRendererPrimaryViews;
 	}
 }
 
@@ -62,66 +44,14 @@ FInstanceCullingManager::~FInstanceCullingManager()
 
 void FInstanceCullingManager::AllocateViews(int32 NumViews)
 {
-	if (bIsEnabled)
-	{
-		CullingViews.AddUninitialized(CullingViews.Num() + NumViews);
-	}
-}
-
-int32 FInstanceCullingManager::RegisterPrimaryView(const FViewInfo& ViewInfo)
-{
-	Nanite::FPackedViewParams Params;
-	Params.ViewMatrices = ViewInfo.ViewMatrices;
-	Params.PrevViewMatrices = ViewInfo.PrevViewInfo.ViewMatrices;
-	Params.ViewRect = ViewInfo.ViewRect;
-	// TODO: faking this here (not needed for culling, until we start involving multi-view and HZB)
-	Params.RasterContextSize = ViewInfo.ViewRect.Size();
-	Params.ViewLODDistanceFactor = ViewInfo.LODDistanceFactor;
-	Params.HZBTestViewRect = FIntRect(0, 0, ViewInfo.PrevViewInfo.ViewRect.Width(), ViewInfo.PrevViewInfo.ViewRect.Height());	// needs to be in HZB space, which is 0,0-based for any view, even stereo/splitscreen ones
-	Params.MaxPixelsPerEdgeMultipler = 1.0f;
-	Params.InstanceOcclusionQueryMask = ViewInfo.PrevViewInfo.InstanceOcclusionQueryMask;
-	Params.SceneRendererPrimaryViewId = ViewInfo.GPUSceneViewId;
-
-	if (ViewInfo.PrevViewInfo.HZB != nullptr)
-	{
-		check(IsInRenderingThread());
-
-		ViewPrevHZBs.AddUnique(ViewInfo.PrevViewInfo.HZB);
-	}	
-		
-	return RegisterView(Params);
+	return ViewDataManager.AllocateViews(NumViews);
 }
 
-
-
 int32 FInstanceCullingManager::RegisterView(const Nanite::FPackedViewParams& Params)
 {
-	if (!bIsEnabled)
-	{
-		return 0;
-	}
-
-	const int32 ViewIndex = NumRegisteredViews.fetch_add(1, std::memory_order_relaxed);
-	check(ViewIndex <= CullingViews.Num());
-	CullingViews[ViewIndex] = CreatePackedView(Params);
-	return ViewIndex;
+	return ViewDataManager.RegisterView(Params);
 }
 
-void FInstanceCullingManager::FlushRegisteredViews(FRDGBuilder& GraphBuilder)
-{
-	const int32 LocalNumRegisteredViews = NumRegisteredViews.load(std::memory_order_relaxed);
-
-	if (ShaderParameters.NumCullingViews != LocalNumRegisteredViews)
-	{
-		// No need to recreate in deferred upload mode
-		if (GraphBuilder.IsImmediateMode())
-		{
-			CullingViewsRDG = CreateStructuredBuffer(GraphBuilder, TEXT("InstanceCulling.CullingViews"), MakeArrayView<const Nanite::FPackedView>(CullingViews.GetData(), LocalNumRegisteredViews));
-			ShaderParameters.InViews = GraphBuilder.CreateSRV(CullingViewsRDG);
-		}
-		ShaderParameters.NumCullingViews = LocalNumRegisteredViews;
-	}
-}
 
 int32 FInstanceCullingManager::GetBinIndex(EBatchProcessingMode Mode, const TRefCountPtr<IPooledRenderTarget>& HZB)
 {
@@ -151,13 +81,15 @@ int32 FInstanceCullingManager::GetBinIndex(EBatchProcessingMode Mode, const TRef
 
 bool FInstanceCullingManager::AllowBatchedBuildRenderingCommands(const FGPUScene& GPUScene)
 {
-	return GPUScene.IsEnabled() && !!GAllowBatchedBuildRenderingCommands && !FRDGBuilder::IsImmediateMode();
+	return GPUScene.IsEnabled() && GAllowBatchedBuildRenderingCommands != 0 && !FRDGBuilder::IsImmediateMode();
 }
 
 void FInstanceCullingManager::BeginDeferredCulling(FRDGBuilder& GraphBuilder)
 {
 	TRACE_CPUPROFILER_EVENT_SCOPE(FInstanceCullingManager::BeginDeferredCulling);
-	FlushRegisteredViews(GraphBuilder);
+
+	// TODO: is this needed(?)
+	ViewDataManager.FlushRegisteredViews(GraphBuilder);
 
 	// Cannot defer pass execution in immediate mode.
 	if (!AllowBatchedBuildRenderingCommands(GPUScene))
@@ -166,16 +98,10 @@ void FInstanceCullingManager::BeginDeferredCulling(FRDGBuilder& GraphBuilder)
 	}
 
 	// If there are no instances, there can be no work to perform later.
-	if (GPUScene.GetNumInstances() == 0 || NumRegisteredViews.load(std::memory_order_relaxed) == 0)
+	if (GPUScene.GetNumInstances() == 0 || ViewDataManager.GetNumCullingViews() == 0)
 	{
 		return;
 	}
 
 	DeferredContext = FInstanceCullingContext::CreateDeferredContext(GraphBuilder, GPUScene, *this);
 }
-
-FInstanceCullingManager::FShaderParameters FInstanceCullingManager::GetShaderParameters(FRDGBuilder& GraphBuilder)
-{
-	FlushRegisteredViews(GraphBuilder);
-	return ShaderParameters;
-}
diff --git a/Engine/Source/Runtime/Renderer/Private/InstanceCulling/InstanceCullingManager.h b/Engine/Source/Runtime/Renderer/Private/InstanceCulling/InstanceCullingManager.h
--- a/Engine/Source/Runtime/Renderer/Private/InstanceCulling/InstanceCullingManager.h
+++ b/Engine/Source/Runtime/Renderer/Private/InstanceCulling/InstanceCullingManager.h
@@ -6,9 +6,11 @@
 #include "InstanceCullingLoadBalancer.h"
 #include "RenderGraphResources.h"
 #include "Async/Mutex.h"
+#include "SceneUniformBuffer.h"
+
+class FRendererViewDataManager;
 
 class FGPUScene;
-class FSceneUniformBuffer;
 
 namespace Nanite
 {
@@ -45,7 +47,7 @@ public:
 	/**
 	 * Construct the instance culling manager for a scene renderer with the set of primary views that are used.
 	 */
-	FInstanceCullingManager(FRDGBuilder& GraphBuilder, const FScene& InScene, FSceneUniformBuffer& InSceneUniforms, TArray<FViewInfo*> &InOutSceneRendererPrimaryViews);
+	FInstanceCullingManager(FRDGBuilder& GraphBuilder, const FScene& InScene, FSceneUniformBuffer& InSceneUniforms, FRendererViewDataManager& InViewDataManager);
 
 	~FInstanceCullingManager();
 
@@ -61,15 +63,6 @@ public:
 	
 	static bool AllowBatchedBuildRenderingCommands(const FGPUScene& GPUScene);
 
-
-	BEGIN_SHADER_PARAMETER_STRUCT(FShaderParameters, )
-		SHADER_PARAMETER_RDG_BUFFER_SRV(StructuredBuffer< Nanite::FPackedView >, InViews)
-		SHADER_PARAMETER(uint32, NumSceneRendererPrimaryViews)
-		SHADER_PARAMETER(uint32, NumCullingViews)
-	END_SHADER_PARAMETER_STRUCT()
-
-	FShaderParameters GetShaderParameters(FRDGBuilder& GraphBuilder);
-
 	/**
 	 * Add a deferred, batched, gpu culling pass. Each batch represents a BuildRenderingCommands call from a mesh pass.
 	 * Batches are collected as we walk through the main render setup and call BuildRenderingCommands, and are processed when RDG Execute or Drain is called.
@@ -92,14 +85,9 @@ public:
 
 	// to support merging of multiple contexts with different HZBs, we use separate bins (LoadBalancers) in the InstanceCullingDeferredContext
 	// bin 0 is reserved for UnCulled batches. Bins >= 1 are for the different HZBs.
-
 	int32 GetBinIndex(EBatchProcessingMode Mode, const TRefCountPtr<IPooledRenderTarget>& HZB);
 
 private:
-	void FlushRegisteredViews(FRDGBuilder& GraphBuilder);
-	// Register a primary view
-	int32 RegisterPrimaryView(const FViewInfo& ViewInfo);
-
 	friend class FInstanceCullingContext;
 	
 	// Polulated by FInstanceCullingContext::BuildRenderingCommandsDeferred, used to hold instance culling related data that needs to be passed around
@@ -111,13 +99,10 @@ private:
 	const FScene& Scene;
 	const FGPUScene& GPUScene;
 	FSceneUniformBuffer& SceneUniforms;
-	std::atomic_int32_t NumRegisteredViews = {0};
-	TArray<Nanite::FPackedView> CullingViews;
+	FRendererViewDataManager& ViewDataManager;
+
 	bool bIsEnabled;
 
-	FRDGBuffer* CullingViewsRDG = nullptr;
-	int32 NumSceneRendererPrimaryViews = 0;
-	FShaderParameters ShaderParameters;
-	TRDGUniformBufferRef<FInstanceCullingGlobalUniforms> DummyUniformBuffer = nullptr;
+	TRDGUniformBufferRef<FInstanceCullingGlobalUniforms> DummyUniformBuffer;
 };
 
diff --git a/Engine/Source/Runtime/Renderer/Private/MobileShadingRenderer.cpp b/Engine/Source/Runtime/Renderer/Private/MobileShadingRenderer.cpp
--- a/Engine/Source/Runtime/Renderer/Private/MobileShadingRenderer.cpp
+++ b/Engine/Source/Runtime/Renderer/Private/MobileShadingRenderer.cpp
@@ -76,6 +76,7 @@
 #include "Rendering/CustomRenderPass.h"
 #include "GenerateMips.h"
 #include "MobileSSR.h"
+#include "ViewData.h"
 
 uint32 GetShadowQuality();
 
@@ -1074,7 +1075,8 @@ void FMobileSceneRenderer::Render(FRDGBuilder& GraphBuilder)
 
 	FInitViewTaskDatas InitViewTaskDatas(VisibilityTaskData);
 
-	FInstanceCullingManager& InstanceCullingManager = *GraphBuilder.AllocObject<FInstanceCullingManager>(GraphBuilder, *Scene, GetSceneUniforms(), AllViews);
+	FRendererViewDataManager& ViewDataManager = *GraphBuilder.AllocObject<FRendererViewDataManager>(GraphBuilder, *Scene, GetSceneUniforms(), AllViews);
+	FInstanceCullingManager& InstanceCullingManager = *GraphBuilder.AllocObject<FInstanceCullingManager>(GraphBuilder, *Scene, GetSceneUniforms(), ViewDataManager);
 
 	// Find the visible primitives and prepare targets and buffers for rendering
 	InitViews(GraphBuilder, SceneTexturesConfig, InstanceCullingManager, VirtualTextureUpdater.Get(), InitViewTaskDatas);
diff --git a/Engine/Source/Runtime/Renderer/Private/SceneExtensions.h b/Engine/Source/Runtime/Renderer/Private/SceneExtensions.h
--- a/Engine/Source/Runtime/Renderer/Private/SceneExtensions.h
+++ b/Engine/Source/Runtime/Renderer/Private/SceneExtensions.h
@@ -16,6 +16,7 @@ class FSceneExtensionRegistry;
 class ISceneExtensionUpdater;
 class ISceneExtensionRenderer;
 struct FEngineShowFlags;
+class FRendererViewDataManager;
 
 /** Abstract interface for an extension to the persistent data of a scene */
 class ISceneExtension
@@ -58,6 +59,12 @@ public:
 	virtual void Begin(FSceneRendererBase* InRenderer) { SceneRendererInternal = InRenderer; }
 	virtual void End() { SceneRendererInternal = nullptr; }
 
+	/**
+	 * Perform any view dependent LOD calculations or similar to e.g., update instance state. 
+	 * Called before UpdateSceneUniformBuffer.
+	 */
+	virtual void UpdateViewData(FRDGBuilder& GraphBuilder, const FRendererViewDataManager& ViewDataManager) {}
+
 	// See the note in ISceneExtensionUpdater about the SceneUniforms.
 	virtual void UpdateSceneUniformBuffer(FRDGBuilder& GraphBuilder, FSceneUniformBuffer& SceneUniforms) {}
 
@@ -256,6 +263,14 @@ public:
 	template<typename TRenderer>
 	const TRenderer& GetRenderer() const { return const_cast<FSceneExtensionsRenderers*>(this)->GetRenderer<TRenderer>(); }
 
+	void UpdateViewData(FRDGBuilder& GraphBuilder, const FRendererViewDataManager& ViewDataManager)
+	{
+		for (auto Renderer : Renderers) 
+		{ 
+			Renderer->UpdateViewData(GraphBuilder, ViewDataManager); 
+		}
+	}
+
 	void UpdateSceneUniformBuffer(FRDGBuilder& GraphBuilder, FSceneUniformBuffer& SceneUniforms)
 	{
 		for (auto Renderer : Renderers) { Renderer->UpdateSceneUniformBuffer(GraphBuilder, SceneUniforms); }
diff --git a/Engine/Source/Runtime/Renderer/Private/SceneHitProxyRendering.cpp b/Engine/Source/Runtime/Renderer/Private/SceneHitProxyRendering.cpp
--- a/Engine/Source/Runtime/Renderer/Private/SceneHitProxyRendering.cpp
+++ b/Engine/Source/Runtime/Renderer/Private/SceneHitProxyRendering.cpp
@@ -36,6 +36,7 @@
 #include "SimpleMeshDrawCommandPass.h"
 #include "StaticMeshSceneProxy.h"
 #include "PixelShaderUtils.h"
+#include "ViewData.h"
 
 class FHitProxyShaderElementData : public FMeshMaterialShaderElementData
 {
@@ -542,7 +543,8 @@ void FMobileSceneRenderer::RenderHitProxies(FRDGBuilder& GraphBuilder)
 	FRDGTextureRef HitProxyDepthTexture = nullptr;
 	InitHitProxyRender(GraphBuilder, this, HitProxyTexture, HitProxyDepthTexture);
 
-	FInstanceCullingManager& InstanceCullingManager = *GraphBuilder.AllocObject<FInstanceCullingManager>(GraphBuilder, *Scene, GetSceneUniforms(), AllViews);
+	FRendererViewDataManager& ViewDataManager = *GraphBuilder.AllocObject<FRendererViewDataManager>(GraphBuilder, *Scene, GetSceneUniforms(), AllViews);
+	FInstanceCullingManager& InstanceCullingManager = *GraphBuilder.AllocObject<FInstanceCullingManager>(GraphBuilder, *Scene, GetSceneUniforms(), ViewDataManager);
 
 	FInitViewTaskDatas InitViewTaskDatas(VisibilityTaskData);
 
@@ -587,7 +589,8 @@ void FDeferredShadingSceneRenderer::RenderHitProxies(FRDGBuilder& GraphBuilder)
 
 	const FIntPoint HitProxyTextureSize = HitProxyDepthTexture->Desc.Extent;
 
-	FInstanceCullingManager& InstanceCullingManager = *GraphBuilder.AllocObject<FInstanceCullingManager>(GraphBuilder, *Scene, GetSceneUniforms(), AllViews);
+	FRendererViewDataManager& ViewDataManager = *GraphBuilder.AllocObject<FRendererViewDataManager>(GraphBuilder, *Scene, GetSceneUniforms(), AllViews);
+	FInstanceCullingManager& InstanceCullingManager = *GraphBuilder.AllocObject<FInstanceCullingManager>(GraphBuilder, *Scene, GetSceneUniforms(), ViewDataManager);
 
 	// Find the visible primitives.
 	FLumenSceneFrameTemporaries LumenFrameTemporaries(Views);
diff --git a/Engine/Source/Runtime/Renderer/Private/ViewData.cpp b/Engine/Source/Runtime/Renderer/Private/ViewData.cpp
new file mode 100644
--- /dev/null
+++ b/Engine/Source/Runtime/Renderer/Private/ViewData.cpp
@@ -0,0 +1,146 @@
+// Copyright Epic Games, Inc. All Rights Reserved.
+
+#include "ViewData.h"
+#include "ScenePrivate.h"
+
+FRendererViewDataManager::FRendererViewDataManager(FRDGBuilder& GraphBuilder, const FScene& InScene, FSceneUniformBuffer& InSceneUniforms, TArray<FViewInfo*> &InOutSceneRendererPrimaryViews)
+	: Scene(InScene), GPUScene(InScene.GPUScene), SceneUniforms(InSceneUniforms), SceneRendererPrimaryViews(InOutSceneRendererPrimaryViews), bIsEnabled(GPUScene.IsEnabled())
+{
+	if (bIsEnabled)
+	{
+		{
+			TArray<uint32> PersistentIdToIndexMap;
+			PersistentIdToIndexMap.SetNumZeroed(InScene.GetMaxPersistentViewId());
+
+			NumSceneRendererPrimaryViews = SceneRendererPrimaryViews.Num();
+			AllocateViews(NumSceneRendererPrimaryViews);
+			for (int32 ViewIndex = 0; ViewIndex < SceneRendererPrimaryViews.Num(); ++ViewIndex)
+			{
+				FViewInfo &View = *SceneRendererPrimaryViews[ViewIndex];
+				View.GPUSceneViewId = ViewIndex;
+				ensure(View.GPUSceneViewId == RegisterPrimaryView(View));
+
+				if (View.PersistentViewId.IsValid())
+				{
+					// Allow zero to mean invalid.
+					PersistentIdToIndexMap[View.PersistentViewId.Index] = ViewIndex + 1;
+				}
+			}
+	
+			PersistentIdToIndexMapRDG = CreateStructuredBuffer(GraphBuilder, TEXT("InstanceCullingManager.PersistentIdToIndexMap"), MoveTemp(PersistentIdToIndexMap));
+		}
+		CullingShaderParameters.InViews = nullptr;
+		CullingShaderParameters.NumCullingViews = 0;
+
+		// Create a deferred buffer and flush when all views are registered. This doesn't work when immediate mode is used, and in that case we need to re-upload partial buffers when the buffer is requested
+		if (!GraphBuilder.IsImmediateMode())
+		{
+			CullingViewsRDG = CreateStructuredBuffer(GraphBuilder, TEXT("ViewData.CullingViews"), [this]() -> TArray<Nanite::FPackedView>& { return CullingViews; });
+			CullingShaderParameters.InViews = GraphBuilder.CreateSRV(CullingViewsRDG);
+			CullingShaderParameters.NumCullingViews = NumSceneRendererPrimaryViews;
+		}
+		else
+		{
+			FlushRegisteredViews(GraphBuilder);
+		}
+		PrimaryViewsRDG = CreateStructuredBuffer(GraphBuilder, TEXT("ViewData.PrimaryViews"), CullingViews);
+		CullingShaderParameters.NumSceneRendererPrimaryViews = NumSceneRendererPrimaryViews;
+	}
+}
+
+void FRendererViewDataManager::InitInstanceState(FRDGBuilder& GraphBuilder)
+{
+	if (bIsEnabled)
+	{
+		int32 NumInstances = GPUScene.GetNumInstances();
+		InstanceMaskWordStride = FMath::DivideAndRoundUp(NumInstances, 32);
+		int32 NumWords = InstanceMaskWordStride * NumSceneRendererPrimaryViews;
+	
+		DeformingInstanceViewMaskRDG = GraphBuilder.CreateBuffer(FRDGBufferDesc::CreateStructuredDesc(4u, FMath::Max(1, NumWords)), TEXT("ViewData.DeformingInstanceViewMask"));
+		AddClearUAVPass(GraphBuilder, GraphBuilder.CreateUAV(DeformingInstanceViewMaskRDG), 0u);
+
+		ShaderParameters.DeformingInstanceViewMask = GraphBuilder.CreateSRV(DeformingInstanceViewMaskRDG);
+		ShaderParameters.PersistentIdToIndexMap = GraphBuilder.CreateSRV(PersistentIdToIndexMapRDG);
+		ShaderParameters.Common.InstanceMaskWordStride = InstanceMaskWordStride;
+		ShaderParameters.Common.NumSceneRendererPrimaryViews = NumSceneRendererPrimaryViews;
+		ShaderParameters.Common.MaxPersistentViewId = Scene.GetMaxPersistentViewId();
+		ShaderParameters.Common.InViews = GraphBuilder.CreateSRV(PrimaryViewsRDG);
+		SceneUniforms.Set(SceneUB::ViewData, ShaderParameters);
+
+		WriterShaderParameters.OutDeformingInstanceViewMask = GraphBuilder.CreateUAV(DeformingInstanceViewMaskRDG,ERDGUnorderedAccessViewFlags::SkipBarrier);
+		WriterShaderParameters.Common = ShaderParameters.Common;
+	}
+}
+
+void FRendererViewDataManager::AllocateViews(int32 NumViews)
+{
+	if (bIsEnabled)
+	{
+		CullingViews.AddUninitialized(CullingViews.Num() + NumViews);
+	}
+}
+
+int32 FRendererViewDataManager::RegisterPrimaryView(const FViewInfo& ViewInfo)
+{
+	Nanite::FPackedViewParams Params;
+	Params.ViewMatrices = ViewInfo.ViewMatrices;
+	Params.PrevViewMatrices = ViewInfo.PrevViewInfo.ViewMatrices;
+	Params.ViewRect = ViewInfo.ViewRect;
+	// TODO: faking this here (not needed for culling, until we start involving multi-view and HZB)
+	Params.RasterContextSize = ViewInfo.ViewRect.Size();
+	Params.ViewLODDistanceFactor = ViewInfo.LODDistanceFactor;
+	Params.HZBTestViewRect = FIntRect(0, 0, ViewInfo.PrevViewInfo.ViewRect.Width(), ViewInfo.PrevViewInfo.ViewRect.Height());	// needs to be in HZB space, which is 0,0-based for any view, even stereo/splitscreen ones
+	Params.MaxPixelsPerEdgeMultipler = 1.0f;
+	Params.InstanceOcclusionQueryMask = ViewInfo.PrevViewInfo.InstanceOcclusionQueryMask;
+	Params.SceneRendererPrimaryViewId = ViewInfo.GPUSceneViewId;
+
+	return RegisterView(Params);
+}
+
+int32 FRendererViewDataManager::RegisterView(const Nanite::FPackedViewParams& Params)
+{
+	if (!bIsEnabled)
+	{
+		return 0;
+	}
+
+	const int32 ViewIndex = NumRegisteredViews.fetch_add(1, std::memory_order_relaxed);
+	check(ViewIndex <= CullingViews.Num());
+	CullingViews[ViewIndex] = CreatePackedView(Params);
+	return ViewIndex;
+}
+
+void FRendererViewDataManager::FlushRegisteredViews(FRDGBuilder& GraphBuilder)
+{
+	const int32 LocalNumRegisteredViews = NumRegisteredViews.load(std::memory_order_relaxed);
+
+	if (CullingShaderParameters.NumCullingViews != LocalNumRegisteredViews)
+	{
+		// No need to recreate in deferred upload mode
+		if (GraphBuilder.IsImmediateMode())
+		{
+			CullingViewsRDG = CreateStructuredBuffer(GraphBuilder, TEXT("ViewData.CullingViews"), MakeArrayView<const Nanite::FPackedView>(CullingViews.GetData(), LocalNumRegisteredViews));
+			CullingShaderParameters.InViews = GraphBuilder.CreateSRV(CullingViewsRDG);
+		}
+		CullingShaderParameters.NumCullingViews = LocalNumRegisteredViews;
+	}
+}
+
+RendererViewData::FCullingShaderParameters FRendererViewDataManager::GetCullingParameters(FRDGBuilder& GraphBuilder)
+{
+	FlushRegisteredViews(GraphBuilder);
+	return CullingShaderParameters;
+}
+
+void GetDefaultParameters(RendererViewData::FParameters& OutParameters, FRDGBuilder& GraphBuilder)
+{
+	OutParameters.Common.NumSceneRendererPrimaryViews = 0u;
+	OutParameters.Common.InstanceMaskWordStride = 0u;
+	OutParameters.Common.MaxPersistentViewId = 0;
+	OutParameters.Common.InViews = GraphBuilder.CreateSRV(GSystemTextures.GetDefaultStructuredBuffer(GraphBuilder, sizeof(Nanite::FPackedView)));
+
+	OutParameters.DeformingInstanceViewMask = GraphBuilder.CreateSRV(GSystemTextures.GetDefaultStructuredBuffer(GraphBuilder, 4u, 0xFFFFFFFFu));
+	OutParameters.PersistentIdToIndexMap = GraphBuilder.CreateSRV(GSystemTextures.GetDefaultStructuredBuffer(GraphBuilder, 4u));
+}
+
+IMPLEMENT_SCENE_UB_STRUCT(RendererViewData::FParameters, ViewData, GetDefaultParameters);
diff --git a/Engine/Source/Runtime/Renderer/Private/ViewData.h b/Engine/Source/Runtime/Renderer/Private/ViewData.h
new file mode 100644
--- /dev/null
+++ b/Engine/Source/Runtime/Renderer/Private/ViewData.h
@@ -0,0 +1,136 @@
+// Copyright Epic Games, Inc. All Rights Reserved.
+
+#pragma once
+
+#include "ShaderParameterMacros.h"
+#include "SceneUniformBuffer.h"
+//#include "Nanite/NaniteShared.h"
+
+class FViewInfo;
+class FRDGBuilder;
+class FScene;
+class FGPUScene;
+
+namespace Nanite
+{
+	struct FPackedView;
+	struct FPackedViewParams;
+}
+
+namespace RendererViewData
+{
+
+BEGIN_SHADER_PARAMETER_STRUCT(FCommonParameters, RENDERER_API)
+	// InViews represents the scene renderer primary views (passed to FRendererViewDataManager constructor), and are named "InViews" to be compatible with Nanite conventions (see GetNaniteView),
+	// The buffer may also contain all the non-primary views but these are not generally accessible.
+	SHADER_PARAMETER_RDG_BUFFER_SRV(StructuredBuffer< FPackedNaniteView >, InViews)
+	SHADER_PARAMETER(int32, MaxPersistentViewId)
+	// Number of scene renderer primary views.
+	SHADER_PARAMETER(uint32, NumSceneRendererPrimaryViews)
+	// Stride between each bit vector in the per-view bit masks in dwords.
+	SHADER_PARAMETER(uint32, InstanceMaskWordStride)
+END_SHADER_PARAMETER_STRUCT()
+
+BEGIN_SHADER_PARAMETER_STRUCT(FParameters, RENDERER_API)
+	SHADER_PARAMETER_STRUCT_INCLUDE(FCommonParameters, Common)
+	SHADER_PARAMETER_RDG_BUFFER_SRV(StructuredBuffer<uint>, DeformingInstanceViewMask)
+	SHADER_PARAMETER_RDG_BUFFER_SRV(StructuredBuffer<uint>, PersistentIdToIndexMap)
+END_SHADER_PARAMETER_STRUCT()
+
+/**
+ * Parameters to use in kernels modifying instance visibility state.
+ * Use the API to abstract any access, see: ViewData.ush.
+ */
+BEGIN_SHADER_PARAMETER_STRUCT(FWriterParameters, RENDERER_API)
+	SHADER_PARAMETER_STRUCT_INCLUDE(FCommonParameters, Common)
+	SHADER_PARAMETER_RDG_BUFFER_UAV(RWStructuredBuffer<uint>, OutDeformingInstanceViewMask)
+END_SHADER_PARAMETER_STRUCT()
+
+
+BEGIN_SHADER_PARAMETER_STRUCT(FCullingShaderParameters, )
+	SHADER_PARAMETER_RDG_BUFFER_SRV(StructuredBuffer< FPackedNaniteView >, InViews)
+	SHADER_PARAMETER(uint32, NumSceneRendererPrimaryViews)
+	SHADER_PARAMETER(uint32, NumCullingViews)
+END_SHADER_PARAMETER_STRUCT()
+
+}
+
+DECLARE_SCENE_UB_STRUCT(RendererViewData::FParameters, ViewData, RENDERER_API)
+
+/**
+ * Takes care of all view data management that belongs to a given FSceneRenderer,
+ */
+class FRendererViewDataManager
+{
+public:
+	/**
+	 * Construct the renderer
+	 */
+	FRendererViewDataManager(FRDGBuilder& GraphBuilder, const FScene& InScene, FSceneUniformBuffer& InSceneUniforms, TArray<FViewInfo*> &InOutSceneRendererPrimaryViews);
+
+	bool IsEnabled() const { return bIsEnabled; }
+
+	/**
+	 * Register a non-primary view for culling, returns integer ID of the view.
+	 * Note that these views are not accessible through the SceneUB and general ViewData.ush API as they are added later in the frame. 
+	 * This could be changed in the future.
+	 */ 
+	int32 RegisterView(const Nanite::FPackedViewParams& Params);
+
+	// Allocate space for views ahead of time prior to calling RegisterView.
+	void AllocateViews(int32 NumViews);
+
+	/**
+	 * Must be called _after_ dynamic primitives are allocated & before instance visibility and state processing has begun (i.e., anything that calls GetWriterShaderParameters).
+	 */
+	void InitInstanceState(FRDGBuilder& GraphBuilder);
+
+	/**
+	 * Shader parameters used for culling where all registered views are concerned (not just the primary)
+	 * The secondary views are not accessible throught the scene UB.
+	 */
+	RendererViewData::FCullingShaderParameters GetCullingParameters(FRDGBuilder& GraphBuilder);
+
+	/**
+	 */
+	RendererViewData::FWriterParameters GetWriterShaderParameters(FRDGBuilder& GraphBuilder) const { return WriterShaderParameters; }
+	
+	void FlushRegisteredViews(FRDGBuilder& GraphBuilder);
+
+	const TArray<FViewInfo*> &GetRegisteredPrimaryViews() const { return SceneRendererPrimaryViews; }
+
+	int32 GetNumCullingViews() const { return NumRegisteredViews.load(std::memory_order_relaxed); }
+
+
+private:
+
+	// Register a primary view
+	int32 RegisterPrimaryView(const FViewInfo& ViewInfo);
+
+	FRendererViewDataManager() = delete;
+	FRendererViewDataManager(FRendererViewDataManager&) = delete;
+
+	const FScene& Scene;
+	const FGPUScene& GPUScene;
+	FSceneUniformBuffer& SceneUniforms;
+	TArray<FViewInfo*> &SceneRendererPrimaryViews;
+
+	std::atomic_int32_t NumRegisteredViews = {0};
+	TArray<Nanite::FPackedView> CullingViews;
+	bool bIsEnabled;
+
+	FRDGBuffer* CullingViewsRDG = nullptr;
+	FRDGBuffer* PrimaryViewsRDG = nullptr;
+	int32 NumSceneRendererPrimaryViews = 0;
+	int32 InstanceMaskWordStride = 0;
+	// Non view index strided buffer of bits, one per instance that is deforming (animating or something like that), indexed by InstanceId.
+	// Initialized to zero each frame & updated by interested scene extensions / systems.
+	// Also in the same buffer laid out after the above
+	// 1. counter of number of instances marked & 
+	FRDGBuffer* DeformingInstanceViewMaskRDG = nullptr;
+	FRDGBuffer* PersistentIdToIndexMapRDG = nullptr;
+
+	RendererViewData::FCullingShaderParameters CullingShaderParameters;
+	RendererViewData::FParameters	ShaderParameters;
+	RendererViewData::FWriterParameters	WriterShaderParameters;
+};
