Heterogeneous Volumes: Use FBoxSphereBounds::Builder when building aggregate bounds.

This avoids accidentally including the origin in bounds reporting logic.

diff --git a/Engine/Source/Runtime/Renderer/Private/HeterogeneousVolumes/HeterogeneousVolumesLiveShadingPipeline.cpp b/Engine/Source/Runtime/Renderer/Private/HeterogeneousVolumes/HeterogeneousVolumesLiveShadingPipeline.cpp
--- a/Engine/Source/Runtime/Renderer/Private/HeterogeneousVolumes/HeterogeneousVolumesLiveShadingPipeline.cpp
+++ b/Engine/Source/Runtime/Renderer/Private/HeterogeneousVolumes/HeterogeneousVolumesLiveShadingPipeline.cpp
@@ -1828,7 +1828,7 @@ void CollectHeterogeneousVolumeMeshBatchesForView(
 	const FViewInfo& View,
 	bool bCollectForShadowCasting,
 	TSet<FVolumetricMeshBatch>& HeterogeneousVolumesMeshBatches,
-	FBoxSphereBounds& WorldBounds
+	FBoxSphereBounds::Builder& WorldBoundsBuilder
 )
 {
 	for (int32 MeshBatchIndex = 0; MeshBatchIndex < View.HeterogeneousVolumesMeshBatches.Num(); ++MeshBatchIndex)
@@ -1852,7 +1852,7 @@ void CollectHeterogeneousVolumeMeshBatchesForView(
 		if (bCollectMeshBatch)
 		{
 			HeterogeneousVolumesMeshBatches.FindOrAdd(FVolumetricMeshBatch(MeshBatch.Mesh, MeshBatch.Proxy));
-			WorldBounds = WorldBounds + MeshBatch.Proxy->GetBounds();
+			WorldBoundsBuilder += MeshBatch.Proxy->GetBounds();
 		}
 	}
 }
@@ -1862,7 +1862,7 @@ void CollectHeterogeneousVolumeMeshBatchesForLight(
 	const FVisibleLightInfo* VisibleLightInfo,
 	const FViewInfo& View,
 	TSet<FVolumetricMeshBatch>& HeterogeneousVolumesMeshBatches,
-	FBoxSphereBounds& WorldBounds
+	FBoxSphereBounds::Builder& WorldBoundsBuilder
 )
 {
 	check(LightSceneInfo);
@@ -1874,7 +1874,7 @@ void CollectHeterogeneousVolumeMeshBatchesForLight(
 		{
 			//const FViewInfo& View = Views[ViewIndex];
 			bool bCollectForShadowCasting = true;
-			CollectHeterogeneousVolumeMeshBatchesForView(View, bCollectForShadowCasting, HeterogeneousVolumesMeshBatches, WorldBounds);
+			CollectHeterogeneousVolumeMeshBatchesForView(View, bCollectForShadowCasting, HeterogeneousVolumesMeshBatches, WorldBoundsBuilder);
 		}
 
 		for (int32 ShadowIndex = 0; ShadowIndex < VisibleLightInfo->ShadowsToProject.Num(); ++ShadowIndex)
@@ -1899,7 +1899,7 @@ void CollectHeterogeneousVolumeMeshBatchesForLight(
 					if (bIsShadowCast && bIsVolumeMaterial)
 					{
 						HeterogeneousVolumesMeshBatches.FindOrAdd(FVolumetricMeshBatch(MeshBatch.Mesh, MeshBatch.PrimitiveSceneProxy));
-						WorldBounds = WorldBounds + MeshBatch.PrimitiveSceneProxy->GetBounds();
+						WorldBoundsBuilder += MeshBatch.PrimitiveSceneProxy->GetBounds();
 					}
 				}
 			}
@@ -2104,14 +2104,15 @@ bool RenderVolumetricShadowMapForLightWithLiveShading(
 
 	// Collect shadow-casting volumes
 	TSet<FVolumetricMeshBatch> HeterogeneousVolumesMeshBatches;
-	FBoxSphereBounds WorldVolumeBounds(ForceInit);
-	CollectHeterogeneousVolumeMeshBatchesForLight(LightSceneInfo, VisibleLightInfo, View, HeterogeneousVolumesMeshBatches, WorldVolumeBounds);
-	if (HeterogeneousVolumesMeshBatches.IsEmpty())
+	FBoxSphereBounds::Builder WorldVolumeBoundsBuilder;
+	CollectHeterogeneousVolumeMeshBatchesForLight(LightSceneInfo, VisibleLightInfo, View, HeterogeneousVolumesMeshBatches, WorldVolumeBoundsBuilder);
+	if (HeterogeneousVolumesMeshBatches.IsEmpty() || !WorldVolumeBoundsBuilder.IsValid())
 	{
 		return false;
 	}
 
 	// Build shadow transform
+	FBoxSphereBounds WorldVolumeBounds(WorldVolumeBoundsBuilder);
 	NumShadowMatrices = ProjectedShadowInfo->OnePassShadowViewProjectionMatrices.Num();
 	FMatrix44f ShadowToTranslatedWorld[6];
 
@@ -2494,8 +2495,14 @@ void RenderAdaptiveVolumetricCameraMapWithLiveShading(
 	// Collect all volumes for view
 	bool bCollectForShadowCasting = false;
 	TSet<FVolumetricMeshBatch> HeterogeneousVolumesMeshBatches;
-	FBoxSphereBounds WorldVolumeBounds;
-	CollectHeterogeneousVolumeMeshBatchesForView(View, bCollectForShadowCasting, HeterogeneousVolumesMeshBatches, WorldVolumeBounds);
+	FBoxSphereBounds::Builder WorldBoundsBuilder;
+	CollectHeterogeneousVolumeMeshBatchesForView(View, bCollectForShadowCasting, HeterogeneousVolumesMeshBatches, WorldBoundsBuilder);
+	if (!WorldBoundsBuilder.IsValid())
+	{
+		return;
+	}
+
+	FBoxSphereBounds WorldVolumeBounds(WorldBoundsBuilder);
 
 	bool bShouldRenderCameraMap = !View.ViewRect.IsEmpty() && !HeterogeneousVolumesMeshBatches.IsEmpty();
 	if (bShouldRenderCameraMap)
diff --git a/Engine/Source/Runtime/Renderer/Private/HeterogeneousVolumes/HeterogeneousVolumesVoxelGridPipeline.cpp b/Engine/Source/Runtime/Renderer/Private/HeterogeneousVolumes/HeterogeneousVolumesVoxelGridPipeline.cpp
--- a/Engine/Source/Runtime/Renderer/Private/HeterogeneousVolumes/HeterogeneousVolumesVoxelGridPipeline.cpp
+++ b/Engine/Source/Runtime/Renderer/Private/HeterogeneousVolumes/HeterogeneousVolumesVoxelGridPipeline.cpp
@@ -1414,11 +1414,10 @@ void CalcViewBoundsAndMinimumVoxelSize(
 	FRDGBuilder& GraphBuilder,
 	const FViewInfo& View,
 	const FVoxelGridBuildOptions& BuildOptions,
-	FBoxSphereBounds& TopLevelGridBounds,
+	FBoxSphereBounds::Builder& TopLevelGridBoundsBuilder,
 	float& MinimumVoxelSize
 )
 {
-	TopLevelGridBounds = FBoxSphereBounds(ForceInit);
 	MinimumVoxelSize = HeterogeneousVolumes::GetMinimumVoxelSizeOutsideFrustum();
 
 	// Build view bounds
@@ -1448,22 +1447,19 @@ void CalcViewBoundsAndMinimumVoxelSize(
 				continue;
 			}
 
-			// Only incorporate the primitive if it intersects with the canera bounding sphere where radius=MaxTraceDistance
+			// Only incorporate the primitive if it intersects with the camera bounding sphere where radius=MaxTraceDistance
 			const FBoxSphereBounds& PrimitiveBounds = HeterogeneousVolume->GetBounds();
 			if (View.ViewFrustum.IntersectBox(PrimitiveBounds.Origin, PrimitiveBounds.BoxExtent))
 			{
-				TopLevelGridBounds = Union(TopLevelGridBounds, PrimitiveBounds);
+				TopLevelGridBoundsBuilder += PrimitiveBounds;
 
-				if (View.ViewFrustum.IntersectBox(TopLevelGridBounds.Origin, TopLevelGridBounds.BoxExtent))
-				{
-					// Bandlimit minimum voxel size request with projected voxel size, based on shading rate
-					FVector VoxelCenter = PrimitiveBounds.Origin;
-					float Distance = FMath::Max(FVector(PrimitiveBounds.Origin - WorldCameraOrigin).Length() - PrimitiveBounds.BoxExtent.Length(), 0.0);
-					float VoxelWidth = Distance * PixelWidth * BuildOptions.ShadingRateInFrustum;
+				// Bandlimit minimum voxel size request with projected voxel size, based on shading rate
+				FVector VoxelCenter = PrimitiveBounds.Origin;
+				float Distance = FMath::Max(FVector(PrimitiveBounds.Origin - WorldCameraOrigin).Length() - PrimitiveBounds.BoxExtent.Length(), 0.0);
+				float VoxelWidth = Distance * PixelWidth * BuildOptions.ShadingRateInFrustum;
 
-					float PerVolumeMinimumVoxelSize = FMath::Max(VoxelWidth, HeterogeneousVolume->GetMinimumVoxelSize());
-					MinimumVoxelSize = FMath::Min(PerVolumeMinimumVoxelSize, MinimumVoxelSize);
-				}
+				float PerVolumeMinimumVoxelSize = FMath::Max(VoxelWidth, HeterogeneousVolume->GetMinimumVoxelSize());
+				MinimumVoxelSize = FMath::Min(PerVolumeMinimumVoxelSize, MinimumVoxelSize);
 			}
 		}
 	}
@@ -1474,18 +1470,17 @@ void CalcGlobalBoundsAndMinimumVoxelSize(
 	const TArray<FViewInfo>& Views,
 	const TSet<FVolumetricMeshBatch>& HeterogeneousVolumesMeshBatches,
 	const FVoxelGridBuildOptions& BuildOptions,
-	FBoxSphereBounds& TopLevelGridBounds,
+	FBoxSphereBounds::Builder& TopLevelGridBoundsBuilder,
 	float& GlobalMinimumVoxelSize
 )
 {
-	TopLevelGridBounds = FBoxSphereBounds(ForceInit);
 	GlobalMinimumVoxelSize = HeterogeneousVolumes::GetMinimumVoxelSizeOutsideFrustum();
 
 	// Cycle through all Volume PrimitiveSceneProxies to collect bounds information and minimum voxel-size
 	for (int32 ViewIndex = 0; ViewIndex < Views.Num(); ViewIndex++)
 	{
 		float ViewMinimumVoxelSize = HeterogeneousVolumes::GetMinimumVoxelSizeOutsideFrustum();
-		FBoxSphereBounds AggregatePrimitiveBounds = FBoxSphereBounds(ForceInit);
+		FBoxSphereBounds::Builder AggregatePrimitiveBoundsBuilder;
 
 		// Build view bounds
 		const FViewInfo& View = Views[ViewIndex];
@@ -1517,7 +1512,7 @@ void CalcGlobalBoundsAndMinimumVoxelSize(
 				}
 
 				const FBoxSphereBounds& PrimitiveBounds = HeterogeneousVolume->GetBounds();
-				AggregatePrimitiveBounds = Union(AggregatePrimitiveBounds, PrimitiveBounds);
+				AggregatePrimitiveBoundsBuilder += PrimitiveBounds;
 
 				bool bIntersectsViewFrustum = View.ViewFrustum.IntersectBox(PrimitiveBounds.Origin, PrimitiveBounds.BoxExtent);
 				float VoxelWidth = bIntersectsViewFrustum ? BuildOptions.ShadingRateInFrustum : BuildOptions.ShadingRateOutOfFrustum;
@@ -1534,14 +1529,18 @@ void CalcGlobalBoundsAndMinimumVoxelSize(
 		}
 
 		// When converting to pixel-space units, clamp per-view minimum voxel-size to in-frustum minimum
-		if (BuildOptions.bUseProjectedPixelSizeForOrthoGrid && View.ViewFrustum.IntersectBox(AggregatePrimitiveBounds.Origin, AggregatePrimitiveBounds.BoxExtent))
+		if (AggregatePrimitiveBoundsBuilder.IsValid())
 		{
-			ViewMinimumVoxelSize = FMath::Max(ViewMinimumVoxelSize, HeterogeneousVolumes::GetMinimumVoxelSizeInFrustum());
-		}
+			FBoxSphereBounds AggregatePrimitiveBounds(AggregatePrimitiveBoundsBuilder);
+			if (BuildOptions.bUseProjectedPixelSizeForOrthoGrid && View.ViewFrustum.IntersectBox(AggregatePrimitiveBounds.Origin, AggregatePrimitiveBounds.BoxExtent))
+			{
+				ViewMinimumVoxelSize = FMath::Max(ViewMinimumVoxelSize, HeterogeneousVolumes::GetMinimumVoxelSizeInFrustum());
+			}
 
-		// Update global bounds and minimum voxel-size
-		TopLevelGridBounds = Union(TopLevelGridBounds, AggregatePrimitiveBounds);
-		GlobalMinimumVoxelSize = FMath::Min(GlobalMinimumVoxelSize, ViewMinimumVoxelSize);
+			// Update global bounds and minimum voxel-size
+			TopLevelGridBoundsBuilder += AggregatePrimitiveBounds;
+			GlobalMinimumVoxelSize = FMath::Min(GlobalMinimumVoxelSize, ViewMinimumVoxelSize);
+		}
 	}
 }
 
@@ -2207,16 +2206,16 @@ void BuildFrustumVoxelGrid(
 	RDG_EVENT_SCOPE(GraphBuilder, "Frustum Grid Build");
 
 	// Determine the minimum voxel size for the scene, based on screen projection or user-defined minima
-	FBoxSphereBounds TopLevelGridBounds;
+	FBoxSphereBounds::Builder TopLevelGridBoundsBuilder;
 	float MinimumVoxelSize;
-	CalcViewBoundsAndMinimumVoxelSize(GraphBuilder, View, BuildOptions, TopLevelGridBounds, MinimumVoxelSize);
-
-	if (TopLevelGridBounds.SphereRadius == 0)
+	CalcViewBoundsAndMinimumVoxelSize(GraphBuilder, View, BuildOptions, TopLevelGridBoundsBuilder, MinimumVoxelSize);
+	if (!TopLevelGridBoundsBuilder.IsValid())
 	{
 		FrustumGridUniformBuffer = CreateEmptyFrustumVoxelGridUniformBuffer(GraphBuilder);
 		return;
 	}
 
+	FBoxSphereBounds TopLevelGridBounds(TopLevelGridBoundsBuilder);
 	float NearPlaneDistance = HeterogeneousVolumes::GetNearPlaneDistanceForFrustumGrid();
 	float FarPlaneDistance = HeterogeneousVolumes::GetFarPlaneDistanceForFrustumGrid();
 	ClipNearFarDistances(View, TopLevelGridBounds, NearPlaneDistance, FarPlaneDistance);
@@ -2811,22 +2810,23 @@ void BuildOrthoVoxelGrid(
 	}
 
 	// Collect global bounds
-	FBoxSphereBounds TopLevelGridBounds;
+	FBoxSphereBounds::Builder TopLevelGridBoundsBuilder;
 	float GlobalMinimumVoxelSize;
 	CalcGlobalBoundsAndMinimumVoxelSize(
 		GraphBuilder,
 		Views,
 		HeterogeneousVolumesMeshBatches,
 		BuildOptions,
-		TopLevelGridBounds,
+		TopLevelGridBoundsBuilder,
 		GlobalMinimumVoxelSize
 	);
 
-	if (TopLevelGridBounds.SphereRadius == 0)
+	if (!TopLevelGridBoundsBuilder.IsValid())
 	{
 		OrthoGridUniformBuffer = CreateEmptyOrthoVoxelGridUniformBuffer(GraphBuilder);
 		return;
 	}
+	FBoxSphereBounds TopLevelGridBounds(TopLevelGridBoundsBuilder);
 
 	// Determine top-level grid resolution
 	FIntVector TopLevelGridResolution;
