[Engine]
* Changed particle system state stream to a proxy and added cascade as one backend while making it possible for Niagara to set the other
* Changed so all statestreams are using FStateStreamCreator to create/register statestreams into the manager
* Some smaller cleanups

diff --git a/Engine/Source/Runtime/Engine/Public/StateStream/ParticleSystemStateStream.h b/Engine/Source/Runtime/Engine/Public/StateStream/ParticleSystemStateStream.h
--- a/Engine/Source/Runtime/Engine/Public/StateStream/ParticleSystemStateStream.h
+++ b/Engine/Source/Runtime/Engine/Public/StateStream/ParticleSystemStateStream.h
@@ -7,7 +7,7 @@
 #include "TransformStateStreamHandle.h"
 #include "ParticleSystemStateStream.generated.h"
 
-class UParticleSystem;
+class UFXSystemAsset;
 
 ////////////////////////////////////////////////////////////////////////////////////////////////////
 // Static state for mesh instance. Can only be set upon creation
@@ -33,7 +33,7 @@ private:
 	FTransformHandle Transform;
 
 	UPROPERTY()
-	TObjectPtr<UParticleSystem> SystemAsset;
+	TObjectPtr<UFXSystemAsset> SystemAsset;
 };
 
 
@@ -51,4 +51,6 @@ class IParticleSystemStateStream
 public:
 	DECLARE_STATESTREAM(ParticleSystem)
 	virtual FParticleSystemHandle Game_CreateInstance(const FParticleSystemStaticState& Ss, const FParticleSystemDynamicState& Ds) = 0;
+
+	virtual void SetOtherBackend(IParticleSystemStateStream* Other) {}
 };
diff --git a/Engine/Source/Runtime/Renderer/Private/RendererScene.cpp b/Engine/Source/Runtime/Renderer/Private/RendererScene.cpp
--- a/Engine/Source/Runtime/Renderer/Private/RendererScene.cpp
+++ b/Engine/Source/Runtime/Renderer/Private/RendererScene.cpp
@@ -82,6 +82,7 @@
 #include "StateStream/SkinnedMeshStateStreamImpl.h"
 #include "StateStream/StaticMeshStateStreamImpl.h"
 #include "StateStream/ParticleSystemStateStreamImpl.h"
+#include "StateStreamCreator.h"
 #include "StateStreamManagerImpl.h"
 #include "SceneCulling/SceneCulling.h"
 #include "TransformStateStreamImpl.h"
@@ -7157,28 +7158,17 @@ void FRendererModule::RemoveScene(FSceneInterface* Scene)
 #if WITH_STATE_STREAM
 IStateStreamManager* FRendererModule::AllocateStateStream(UWorld* World)
 {
-	FScene* Scene = (FScene*)World->Scene; // TODO: Check if really FScene?
+	FSceneInterface* Scene = World->Scene;
 	FStateStreamManagerImpl* Manager = new FStateStreamManagerImpl();
-
-	FTransformStateStreamImpl& TransformStreamImpl = *new FTransformStateStreamImpl();
-	FStaticMeshStateStreamImpl& StaticMeshStreamImpl = *new FStaticMeshStateStreamImpl(Scene);
-	FSkinnedMeshStateStreamImpl& SkinnedMeshStreamImpl = *new FSkinnedMeshStateStreamImpl(Scene);
-	FParticleSystemStateStreamImpl& ParticleSystemStreamImpl = *new FParticleSystemStateStreamImpl(Scene);
-	Manager->Render_Register(TransformStreamImpl, true);
-	Manager->Render_Register(StaticMeshStreamImpl, true);
-	Manager->Render_Register(SkinnedMeshStreamImpl, true);
-	Manager->Render_Register(ParticleSystemStreamImpl, true);
-
-	Manager->Render_RegisterDependency(StaticMeshStreamImpl, TransformStreamImpl);
-	Manager->Render_RegisterDependency(SkinnedMeshStreamImpl, TransformStreamImpl);
-	Manager->Render_RegisterDependency(ParticleSystemStreamImpl, TransformStreamImpl);
-
+	FStateStreamCreator::RegisterStateStreams(FStateStreamRegisterContext{*Manager, Scene});
 	return Manager;
 }
 
 void FRendererModule::FreeStateStream(IStateStreamManager* Manager)
 {
-	delete static_cast<FStateStreamManagerImpl*>(Manager);
+	FStateStreamManagerImpl* ManagerImpl = static_cast<FStateStreamManagerImpl*>(Manager);
+	FStateStreamCreator::UnregisterStateStreams(FStateStreamUnregisterContext{*ManagerImpl});
+	delete ManagerImpl;
 }
 #endif
 
diff --git a/Engine/Source/Runtime/Renderer/Private/StateStream/ParticleSystemStateStreamImpl.cpp b/Engine/Source/Runtime/Renderer/Private/StateStream/ParticleSystemStateStreamImpl.cpp
--- a/Engine/Source/Runtime/Renderer/Private/StateStream/ParticleSystemStateStreamImpl.cpp
+++ b/Engine/Source/Runtime/Renderer/Private/StateStream/ParticleSystemStateStreamImpl.cpp
@@ -8,12 +8,15 @@
 #include "Particles/ParticleLODLevel.h"
 #include "Particles/ParticleSystem.h"
 #include "ScenePrivate.h"
+#include "StateStreamCreator.h"
+#include "StateStreamManagerImpl.h"
+#include "TransformStateStream.h"
 
 ////////////////////////////////////////////////////////////////////////////////////////////////////
 
-FParticleSystemObject::~FParticleSystemObject()
+FParticleSystemObjectCascade::~FParticleSystemObjectCascade()
 {
-	FScene& Scene = (FScene&)PrimitiveSceneData.SceneProxy->GetScene();
+	FSceneInterface& Scene = PrimitiveSceneData.SceneProxy->GetScene();
 	Scene.RemovePrimitive(&PrimitiveSceneDesc);
 
 	if (TransformObject)
@@ -23,13 +26,13 @@ FParticleSystemObject::~FParticleSystemObject()
 	}
 }
 
-void FParticleSystemObject::OnTransformObjectDirty()
+void FParticleSystemObjectCascade::OnTransformObjectDirty()
 {
 	FTransformObject::Info Info = TransformObject->GetInfo();
 	// TODO: Implement this to handle changing visibility and movement
 }
 
-void FParticleSystemObject::InitializeSystem(FParticleSystemSceneProxyDesc& OutDesc, const FParticleSystemStaticState& Ss, const FParticleSystemDynamicState& Ds)
+void FParticleSystemObjectCascade::InitializeSystem(FParticleSystemSceneProxyDesc& OutDesc, const FParticleSystemStaticState& Ss, const FParticleSystemDynamicState& Ds)
 {
 	if (Template == NULL)
 	{
@@ -41,7 +44,6 @@ void FParticleSystemObject::InitializeSystem(FParticleSystemSceneProxyDesc& OutD
 	float WarmupTickRate = Template->WarmupTickRate;
 	bool bIsViewRelevanceDirty = true;
 	bool bWasCompleted = false;
-	int32 LODLevel = 0;
 	EDetailMode DetailMode = DM_Medium;
 
 	const int32 GlobalDetailMode = PDM_Medium;//GetCurrentDetailMode();
@@ -149,7 +151,7 @@ void FParticleSystemObject::InitializeSystem(FParticleSystemSceneProxyDesc& OutD
 	LWCTile = FLargeWorldRenderScalar::GetTileFor(GetComponentTransform().GetLocation());
 }
 
-void FParticleSystemObject::Update()
+void FParticleSystemObjectCascade::Update()
 {
 	float DeltaTimeTick = 1.0f/30.0f;
 	bool bSuppressSpawning = false;
@@ -189,7 +191,7 @@ void FParticleSystemObject::Update()
 	bJustRegistered = false;
 }
 
-FParticleDynamicData* FParticleSystemObject::CreateDynamicData(ERHIFeatureLevel::Type InFeatureLevel)
+FParticleDynamicData* FParticleSystemObjectCascade::CreateDynamicData(ERHIFeatureLevel::Type InFeatureLevel)
 {
 	FParticleDynamicData* ParticleDynamicData = new FParticleDynamicData();
 
@@ -238,64 +240,64 @@ FParticleDynamicData* FParticleSystemObject::CreateDynamicData(ERHIFeatureLevel:
 	return ParticleDynamicData;
 }
 
-const FTransform& FParticleSystemObject::GetAsyncComponentToWorld() const { return TransformObject->GetInfo().WorldTransform; }
-UObject* FParticleSystemObject::GetDistributionData() const { return nullptr; }
-const FTransform& FParticleSystemObject::GetComponentTransform() const { return TransformObject->GetInfo().WorldTransform; }
-FRotator FParticleSystemObject::GetComponentRotation() const { return {}; }
-const FTransform& FParticleSystemObject::GetComponentToWorld() const { return TransformObject->GetInfo().WorldTransform; }
-const FBoxSphereBounds& FParticleSystemObject::GetBounds() const { return PrimitiveSceneDesc.Bounds; }
-TWeakObjectPtr<UWorld> FParticleSystemObject::GetWeakWorld() const { return {}; }
-bool FParticleSystemObject::HasWorld() const { return true; }
-bool FParticleSystemObject::HasWorldSettings() const { return {}; }
-bool FParticleSystemObject::IsGameWorld() const { return true; }
-float FParticleSystemObject::GetWorldTimeSeconds() const { return {}; }
-float FParticleSystemObject::GetWorldEffectiveTimeDilation() const { return {}; }
-FIntVector FParticleSystemObject::GetWorldOriginLocation() const { return {}; }
-FSceneInterface* FParticleSystemObject::GetScene() const { return {}; }
-bool FParticleSystemObject::GetFloatParameter(const FName InName, float& OutFloat) { return {}; }
-const FVector3f& FParticleSystemObject::GetLWCTile() const { return LWCTile; }
-FString FParticleSystemObject::GetName() const { return {}; }
-FString FParticleSystemObject::GetFullName() const { return {}; }
-FString FParticleSystemObject::GetPathName() const { return {}; }
-bool FParticleSystemObject::IsActive() const { return {}; }
-bool FParticleSystemObject::IsValidLowLevel() const { return {}; }
-TArrayView<const FParticleSysParam> FParticleSystemObject::GetAsyncInstanceParameters() { return {}; }
-int32 FParticleSystemObject::GetCurrentDetailMode() const { return {}; }
-int32 FParticleSystemObject::GetCurrentLODIndex() const { return {}; }
-const FVector& FParticleSystemObject::GetPartSysVelocity() const { return FVector::ZeroVector; }
-const FVector& FParticleSystemObject::GetOldPosition() const { return FVector::ZeroVector; }
-FFXSystem* FParticleSystemObject::GetFXSystem() const { return {}; }
-const UParticleSystem* FParticleSystemObject::GetTemplate() const { return Template; }
-TArrayView<const FParticleSysParam> FParticleSystemObject::GetInstanceParameters() const { return {}; }
-TArrayView<FParticleEmitterInstance*> FParticleSystemObject::GetEmitterInstances() const { return {}; }
-TArrayView<TObjectPtr<UMaterialInterface>> FParticleSystemObject::GetEmitterMaterials() const { return {}; }
-FPrimitiveSceneProxy* FParticleSystemObject::GetSceneProxy() const { return PrimitiveSceneData.SceneProxy; }
-bool FParticleSystemObject::GetIsWarmingUp() const { return {}; }
-bool FParticleSystemObject::GetJustRegistered() const { return bJustRegistered; }
-float FParticleSystemObject::GetWarmupTime() const { return {}; }
-float FParticleSystemObject::GetEmitterDelay() const { return {}; }
-FRandomStream& FParticleSystemObject::GetRandomStream() { return RandomStream; }
-void FParticleSystemObject::SetComponentToWorld(const FTransform& NewComponentToWorld) { }
-void FParticleSystemObject::DeactivateNextTick() { }
-UParticleSystemComponent* FParticleSystemObject::AsComponent() const { return {}; }
-void FParticleSystemObject::ReportEventSpawn(const FName InEventName, const float InEmitterTime, const FVector InLocation, const FVector InVelocity, const TArray<class UParticleModuleEventSendToGame*>& InEventData) { }
-void FParticleSystemObject::ReportEventDeath(const FName InEventName, const float InEmitterTime, const FVector InLocation, const FVector InVelocity, const TArray<class UParticleModuleEventSendToGame*>& InEventData, const float InParticleTime) { }
-void FParticleSystemObject::ReportEventCollision(const FName InEventName, const float InEmitterTime, const FVector InLocation, const FVector InDirection, const FVector InVelocity, const TArray<class UParticleModuleEventSendToGame*>& InEventData, const float InParticleTime, const FVector InNormal, const float InTime, const int32 InItem, const FName InBoneName, UPhysicalMaterial* PhysMat) { }
-void FParticleSystemObject::ReportEventBurst(const FName InEventName, const float InEmitterTime, const int32 ParticleCount, const FVector InLocation, const TArray<class UParticleModuleEventSendToGame*>& InEventData) { }
-TArrayView<FParticleEventSpawnData> FParticleSystemObject::GetSpawnEvents() const { return {}; }
-TArrayView<FParticleEventDeathData> FParticleSystemObject::GetDeathEvents() const { return {}; }
-TArrayView<FParticleEventCollideData> FParticleSystemObject::GetCollisionEvents() const { return {}; }
-TArrayView<FParticleEventBurstData> FParticleSystemObject::GetBurstEvents() const { return {}; }
-TArrayView<FParticleEventKismetData> FParticleSystemObject::GetKismetEvents() const { return {}; }
+const FTransform& FParticleSystemObjectCascade::GetAsyncComponentToWorld() const { return TransformObject->GetInfo().WorldTransform; }
+UObject* FParticleSystemObjectCascade::GetDistributionData() const { return nullptr; }
+const FTransform& FParticleSystemObjectCascade::GetComponentTransform() const { return TransformObject->GetInfo().WorldTransform; }
+FRotator FParticleSystemObjectCascade::GetComponentRotation() const { return {}; }
+const FTransform& FParticleSystemObjectCascade::GetComponentToWorld() const { return TransformObject->GetInfo().WorldTransform; }
+const FBoxSphereBounds& FParticleSystemObjectCascade::GetBounds() const { return PrimitiveSceneDesc.Bounds; }
+TWeakObjectPtr<UWorld> FParticleSystemObjectCascade::GetWeakWorld() const { return {}; }
+bool FParticleSystemObjectCascade::HasWorld() const { return true; }
+bool FParticleSystemObjectCascade::HasWorldSettings() const { return {}; }
+bool FParticleSystemObjectCascade::IsGameWorld() const { return true; }
+float FParticleSystemObjectCascade::GetWorldTimeSeconds() const { return {}; }
+float FParticleSystemObjectCascade::GetWorldEffectiveTimeDilation() const { return {}; }
+FIntVector FParticleSystemObjectCascade::GetWorldOriginLocation() const { return {}; }
+FSceneInterface* FParticleSystemObjectCascade::GetScene() const { return &PrimitiveSceneData.SceneProxy->GetScene(); }
+bool FParticleSystemObjectCascade::GetFloatParameter(const FName InName, float& OutFloat) { return {}; }
+const FVector3f& FParticleSystemObjectCascade::GetLWCTile() const { return LWCTile; }
+FString FParticleSystemObjectCascade::GetName() const { return {}; }
+FString FParticleSystemObjectCascade::GetFullName() const { return {}; }
+FString FParticleSystemObjectCascade::GetPathName() const { return {}; }
+bool FParticleSystemObjectCascade::IsActive() const { return {}; }
+bool FParticleSystemObjectCascade::IsValidLowLevel() const { return {}; }
+TArrayView<const FParticleSysParam> FParticleSystemObjectCascade::GetAsyncInstanceParameters() { return {}; }
+int32 FParticleSystemObjectCascade::GetCurrentDetailMode() const { return {}; }
+int32 FParticleSystemObjectCascade::GetCurrentLODIndex() const { return {}; }
+const FVector& FParticleSystemObjectCascade::GetPartSysVelocity() const { return FVector::ZeroVector; }
+const FVector& FParticleSystemObjectCascade::GetOldPosition() const { return FVector::ZeroVector; }
+FFXSystem* FParticleSystemObjectCascade::GetFXSystem() const { return {}; }
+const UParticleSystem* FParticleSystemObjectCascade::GetTemplate() const { return Template; }
+TArrayView<const FParticleSysParam> FParticleSystemObjectCascade::GetInstanceParameters() const { return {}; }
+TArrayView<FParticleEmitterInstance*> FParticleSystemObjectCascade::GetEmitterInstances() const { return {}; }
+TArrayView<TObjectPtr<UMaterialInterface>> FParticleSystemObjectCascade::GetEmitterMaterials() const { return {}; }
+FPrimitiveSceneProxy* FParticleSystemObjectCascade::GetSceneProxy() const { return PrimitiveSceneData.SceneProxy; }
+bool FParticleSystemObjectCascade::GetIsWarmingUp() const { return {}; }
+bool FParticleSystemObjectCascade::GetJustRegistered() const { return bJustRegistered; }
+float FParticleSystemObjectCascade::GetWarmupTime() const { return {}; }
+float FParticleSystemObjectCascade::GetEmitterDelay() const { return {}; }
+FRandomStream& FParticleSystemObjectCascade::GetRandomStream() { return RandomStream; }
+void FParticleSystemObjectCascade::SetComponentToWorld(const FTransform& NewComponentToWorld) { }
+void FParticleSystemObjectCascade::DeactivateNextTick() { }
+UParticleSystemComponent* FParticleSystemObjectCascade::AsComponent() const { return {}; }
+void FParticleSystemObjectCascade::ReportEventSpawn(const FName InEventName, const float InEmitterTime, const FVector InLocation, const FVector InVelocity, const TArray<class UParticleModuleEventSendToGame*>& InEventData) { }
+void FParticleSystemObjectCascade::ReportEventDeath(const FName InEventName, const float InEmitterTime, const FVector InLocation, const FVector InVelocity, const TArray<class UParticleModuleEventSendToGame*>& InEventData, const float InParticleTime) { }
+void FParticleSystemObjectCascade::ReportEventCollision(const FName InEventName, const float InEmitterTime, const FVector InLocation, const FVector InDirection, const FVector InVelocity, const TArray<class UParticleModuleEventSendToGame*>& InEventData, const float InParticleTime, const FVector InNormal, const float InTime, const int32 InItem, const FName InBoneName, UPhysicalMaterial* PhysMat) { }
+void FParticleSystemObjectCascade::ReportEventBurst(const FName InEventName, const float InEmitterTime, const int32 ParticleCount, const FVector InLocation, const TArray<class UParticleModuleEventSendToGame*>& InEventData) { }
+TArrayView<FParticleEventSpawnData> FParticleSystemObjectCascade::GetSpawnEvents() const { return {}; }
+TArrayView<FParticleEventDeathData> FParticleSystemObjectCascade::GetDeathEvents() const { return {}; }
+TArrayView<FParticleEventCollideData> FParticleSystemObjectCascade::GetCollisionEvents() const { return {}; }
+TArrayView<FParticleEventBurstData> FParticleSystemObjectCascade::GetBurstEvents() const { return {}; }
+TArrayView<FParticleEventKismetData> FParticleSystemObjectCascade::GetKismetEvents() const { return {}; }
 
 ////////////////////////////////////////////////////////////////////////////////////////////////////
 
-FParticleSystemStateStreamImpl::FParticleSystemStateStreamImpl(FScene* InScene)
+FParticleSystemStateStreamCascade::FParticleSystemStateStreamCascade(FSceneInterface& InScene)
 :	Scene(InScene)
 {
 }
 
-void FParticleSystemStateStreamImpl::SetTransformObject(FParticleSystemObject& Object, const FParticleSystemDynamicState& Ds)
+void FParticleSystemStateStreamCascade::SetTransformObject(FParticleSystemObjectCascade& Object, const FParticleSystemDynamicState& Ds)
 {
 	if (Object.TransformObject)
 	{
@@ -313,11 +315,11 @@ void FParticleSystemStateStreamImpl::SetTransformObject(FParticleSystemObject& O
 	}
 }
 
-void FParticleSystemStateStreamImpl::Render_OnCreate(const FParticleSystemStaticState& Ss, const FParticleSystemDynamicState& Ds, FParticleSystemObject*& UserData, bool IsDestroyedInSameFrame)
+void FParticleSystemStateStreamCascade::Render_OnCreate(const FParticleSystemStaticState& Ss, const FParticleSystemDynamicState& Ds, FParticleSystemObjectCascade*& UserData, bool IsDestroyedInSameFrame)
 {
 	check(!UserData);
 
-	FParticleSystemObject& Object = *new FParticleSystemObject();
+	FParticleSystemObjectCascade& Object = *new FParticleSystemObjectCascade();
 	Object.AddRef();
 
 	SetTransformObject(Object, Ds);
@@ -332,22 +334,19 @@ void FParticleSystemStateStreamImpl::Render_OnCreate(const FParticleSystemStatic
 	//	LocalBounds = Mesh->GetBounds();
 	//}	
 
-	UWorld& World = *Scene->World;
-
-
 	Object.PrimitiveSceneDesc.RenderMatrix = Transform.ToMatrixWithScale();
 	Object.PrimitiveSceneDesc.AttachmentRootPosition = Transform.GetLocation();
 	Object.PrimitiveSceneDesc.PrimitiveSceneData = &Object.PrimitiveSceneData;
 	Object.PrimitiveSceneDesc.LocalBounds = LocalBounds;
 	Object.PrimitiveSceneDesc.Bounds = LocalBounds.TransformBy(Transform);
 
+	UParticleSystem* Asset = static_cast<UParticleSystem*>(Ds.GetSystemAsset());
 	FParticleSystemSceneProxyDesc Desc;
-	Desc.SystemAsset = Ds.GetSystemAsset();
+	Desc.SystemAsset = Asset;
 	//Desc.OverrideMaterials = const_cast<TArray<TObjectPtr<UMaterialInterface>>&>(Ds.GetOverrideMaterials());
 	Desc.CustomPrimitiveData = &Object.CustomPrimitiveData;
-	Desc.Scene = Scene;
-	Desc.World = &World;
-	Desc.FeatureLevel = World.GetFeatureLevel();
+	Desc.Scene = &Scene;
+	Desc.FeatureLevel = Scene.GetFeatureLevel();
 	Desc.bIsVisible = Info.bVisible;
 	//Desc.MaterialRelevance.Raw = Ss.GetMaterialRelevance();
 
@@ -361,24 +360,24 @@ void FParticleSystemStateStreamImpl::Render_OnCreate(const FParticleSystemStatic
 	//Desc.TextureStreamingTransformScale = Transform.GetMaximumAxisScale();
 	#endif
 
-	Object.Template = Ds.GetSystemAsset();
+	Object.Template = Asset;
 	Object.InitializeSystem(Desc, Ss, Ds);
 
 	Object.PrimitiveSceneData.SceneProxy = new FParticleSystemSceneProxy(Desc);
-	Scene->AddPrimitive(&Object.PrimitiveSceneDesc);
+	Scene.AddPrimitive(&Object.PrimitiveSceneDesc);
 
 	UserData = &Object;
 	Objects.Add(&Object);
 }
 
-void FParticleSystemStateStreamImpl::Render_OnUpdate(const FParticleSystemStaticState& Ss, const FParticleSystemDynamicState& Ds, FParticleSystemObject*& UserData)
+void FParticleSystemStateStreamCascade::Render_OnUpdate(const FParticleSystemStaticState& Ss, const FParticleSystemDynamicState& Ds, FParticleSystemObjectCascade*& UserData)
 {
 	if (!UserData)
 	{
 		return;
 	}
 
-	FParticleSystemObject& Object = *UserData;
+	FParticleSystemObjectCascade& Object = *UserData;
 
 	if (Ds.TransformModified())
 	{
@@ -386,7 +385,7 @@ void FParticleSystemStateStreamImpl::Render_OnUpdate(const FParticleSystemStatic
 	}
 }
 
-void FParticleSystemStateStreamImpl::Render_OnDestroy(const FParticleSystemStaticState& Ss, const FParticleSystemDynamicState& Ds, FParticleSystemObject*& UserData)
+void FParticleSystemStateStreamCascade::Render_OnDestroy(const FParticleSystemStaticState& Ss, const FParticleSystemDynamicState& Ds, FParticleSystemObjectCascade*& UserData)
 {
 	if (UserData)
 	{
@@ -396,14 +395,39 @@ void FParticleSystemStateStreamImpl::Render_OnDestroy(const FParticleSystemStati
 	}
 }
 
-void FParticleSystemStateStreamImpl::Render_PostUpdate()
+void FParticleSystemStateStreamCascade::Render_PostUpdate()
 {
-	TStateStream<FParticleSystemStateStreamSettings>::Render_PostUpdate();
+	TStateStream<FParticleSystemStateStreamSettingsCascade>::Render_PostUpdate();
 
-	for (FParticleSystemObject* Object : Objects)
+	for (FParticleSystemObjectCascade* Object : Objects)
 	{
 		Object->Update();
 	}
 }
 
+STATESTREAM_CREATOR_INSTANCE_WITH_FUNC(FParticleSystemStateStreamCascade, [](const FStateStreamRegisterContext& Context, FParticleSystemStateStreamCascade& Impl)
+	{
+		Context.RegisterDependency(ParticleSystemStateStreamCascadeId, TransformStateStreamId);
+		static_cast<FParticleSystemStateStreamImpl*>(Context.Manager.Render_GetStream(ParticleSystemStateStreamId))->CascadeBackend = &Impl;
+	})
+
+////////////////////////////////////////////////////////////////////////////////////////////////////
+
+FParticleSystemHandle FParticleSystemStateStreamImpl::Game_CreateInstance(const FParticleSystemStaticState& Ss, const FParticleSystemDynamicState& Ds)
+{
+	if (Cast<UParticleSystem>(Ds.GetSystemAsset()))
+		return CascadeBackend->Game_CreateInstance(Ss, Ds);
+	check(OtherBackend);
+	if (OtherBackend)
+		return OtherBackend->Game_CreateInstance(Ss, Ds);
+	return {};
+}
+
+void FParticleSystemStateStreamImpl::SetOtherBackend(IParticleSystemStateStream* Other)
+{
+	OtherBackend = Other;
+}
+
+STATESTREAM_CREATOR_INSTANCE(FParticleSystemStateStreamImpl)
+
 ////////////////////////////////////////////////////////////////////////////////////////////////////
diff --git a/Engine/Source/Runtime/Renderer/Private/StateStream/ParticleSystemStateStreamImpl.h b/Engine/Source/Runtime/Renderer/Private/StateStream/ParticleSystemStateStreamImpl.h
--- a/Engine/Source/Runtime/Renderer/Private/StateStream/ParticleSystemStateStreamImpl.h
+++ b/Engine/Source/Runtime/Renderer/Private/StateStream/ParticleSystemStateStreamImpl.h
@@ -11,17 +11,17 @@
 #define UE_API RENDERER_API
 
 class FParticleDynamicData;
-class FScene;
+class FSceneInterface;
 struct FParticleSystemSceneProxyDesc;
 
 ////////////////////////////////////////////////////////////////////////////////////////////////////
 
-class FParticleSystemObject : public TRefCountingMixin<FParticleSystemObject>, public FTransformObjectListener, public IParticleEmitterInstanceOwner
+class FParticleSystemObjectCascade final : public TRefCountingMixin<FParticleSystemObjectCascade>, public FTransformObjectListener, public IParticleEmitterInstanceOwner
 {
 public:
 
 private:
-	virtual ~FParticleSystemObject();
+	virtual ~FParticleSystemObjectCascade();
 	virtual void OnTransformObjectDirty() override final;
 
 	TRefCountPtr<FTransformObject> TransformObject;
@@ -30,8 +30,8 @@ private:
 	FPrimitiveSceneInfoData	PrimitiveSceneData;
 	FPrimitiveSceneDesc PrimitiveSceneDesc;
 
-	friend class FParticleSystemStateStreamImpl;
-	friend class TRefCountingMixin<FParticleSystemObject>;
+	friend class FParticleSystemStateStreamCascade;
+	friend class TRefCountingMixin<FParticleSystemObjectCascade>;
 
 	void InitializeSystem(FParticleSystemSceneProxyDesc& OutDesc, const FParticleSystemStaticState& Ss, const FParticleSystemDynamicState& Ds);
 	void Update();
@@ -97,30 +97,61 @@ private:
 	FRandomStream RandomStream;
 	FVector3f LWCTile = FVector3f::ZeroVector;
 	bool bJustRegistered = true;
+	int32 LODLevel = 0;
 };
 
 ////////////////////////////////////////////////////////////////////////////////////////////////////
 
-struct FParticleSystemStateStreamSettings : TStateStreamSettings<IParticleSystemStateStream, FParticleSystemObject>
+inline constexpr uint32 ParticleSystemStateStreamCascadeId = 5;
+
+////////////////////////////////////////////////////////////////////////////////////////////////////
+
+struct FParticleSystemStateStreamSettingsCascade : TStateStreamSettings<IParticleSystemStateStream, FParticleSystemObjectCascade>
 {
+	enum { Id = ParticleSystemStateStreamCascadeId };
+	static inline constexpr const TCHAR* DebugName = TEXT("ParticleSystem(Cascade)");
 };
 
 ////////////////////////////////////////////////////////////////////////////////////////////////////
 
-class FParticleSystemStateStreamImpl : public TStateStream<FParticleSystemStateStreamSettings>
+class FParticleSystemStateStreamCascade final : public TStateStream<FParticleSystemStateStreamSettingsCascade>
 {
 public:
-	FParticleSystemStateStreamImpl(FScene* InScene);
+	FParticleSystemStateStreamCascade(FSceneInterface& InScene);
 private:
-	UE_API void SetTransformObject(FParticleSystemObject& Object, const FParticleSystemDynamicState& Ds);
-	UE_API virtual void Render_OnCreate(const FParticleSystemStaticState& Ss, const FParticleSystemDynamicState& Ds, FParticleSystemObject*& UserData, bool IsDestroyedInSameFrame) override;
-	UE_API virtual void Render_OnUpdate(const FParticleSystemStaticState& Ss, const FParticleSystemDynamicState& Ds, FParticleSystemObject*& UserData) override;
-	UE_API virtual void Render_OnDestroy(const FParticleSystemStaticState& Ss, const FParticleSystemDynamicState& Ds, FParticleSystemObject*& UserData) override;
+	UE_API void SetTransformObject(FParticleSystemObjectCascade& Object, const FParticleSystemDynamicState& Ds);
+	UE_API virtual void Render_OnCreate(const FParticleSystemStaticState& Ss, const FParticleSystemDynamicState& Ds, FParticleSystemObjectCascade*& UserData, bool IsDestroyedInSameFrame) override;
+	UE_API virtual void Render_OnUpdate(const FParticleSystemStaticState& Ss, const FParticleSystemDynamicState& Ds, FParticleSystemObjectCascade*& UserData) override;
+	UE_API virtual void Render_OnDestroy(const FParticleSystemStaticState& Ss, const FParticleSystemDynamicState& Ds, FParticleSystemObjectCascade*& UserData) override;
 
 	virtual void Render_PostUpdate() override;
 
-	FScene* Scene;
-	TArray<FParticleSystemObject*> Objects;
+	FSceneInterface& Scene;
+	TArray<FParticleSystemObjectCascade*> Objects;
+};
+
+////////////////////////////////////////////////////////////////////////////////////////////////////
+
+class FParticleSystemStateStreamImpl final : public IParticleSystemStateStream, public IStateStream
+{
+public:
+	FParticleSystemStateStreamImpl(FSceneInterface& InScene) {}
+
+	virtual FParticleSystemHandle Game_CreateInstance(const FParticleSystemStaticState& Ss, const FParticleSystemDynamicState& Ds) override;
+	virtual void SetOtherBackend(IParticleSystemStateStream* Other) override;
+
+	virtual void Game_BeginTick() override {}
+	virtual void Game_EndTick(StateStreamTime AbsoluteTime) override {}
+	virtual void Game_Exit() override {}
+	virtual void* Game_GetVoidPointer() override { return this; }
+	virtual void Render_Update(StateStreamTime AbsoluteTime) override {}
+	virtual void Render_PostUpdate() override {}
+	virtual void Render_Exit() override {}
+	virtual void Render_GarbageCollect() override {}
+	virtual uint32 GetId() override { return ParticleSystemStateStreamId; }
+
+	FParticleSystemStateStreamCascade* CascadeBackend = nullptr;
+	IParticleSystemStateStream* OtherBackend = nullptr;
 };
 
 ////////////////////////////////////////////////////////////////////////////////////////////////////
diff --git a/Engine/Source/Runtime/Renderer/Private/StateStream/SkinnedMeshStateStreamImpl.cpp b/Engine/Source/Runtime/Renderer/Private/StateStream/SkinnedMeshStateStreamImpl.cpp
--- a/Engine/Source/Runtime/Renderer/Private/StateStream/SkinnedMeshStateStreamImpl.cpp
+++ b/Engine/Source/Runtime/Renderer/Private/StateStream/SkinnedMeshStateStreamImpl.cpp
@@ -6,13 +6,14 @@
 #include "ScenePrivate.h"
 #include "SkeletalMeshSceneProxy.h"
 #include "SkinnedMeshSceneProxyDesc.h"
+#include "StateStreamCreator.h"
 #include "../../../Engine/Private/SkeletalRenderGPUSkin.h"
 
 ////////////////////////////////////////////////////////////////////////////////////////////////////
 
 FSkinnedMeshObject::~FSkinnedMeshObject()
 {
-	FScene& Scene = static_cast<FScene&>(PrimitiveSceneData.SceneProxy->GetScene());
+	FSceneInterface& Scene = PrimitiveSceneData.SceneProxy->GetScene();
 	Scene.RemovePrimitive(&PrimitiveSceneDesc);
 
 	if (TransformObject)
@@ -74,7 +75,7 @@ void FSkinnedMeshObject::OnTransformObjectDirty()
 
 ////////////////////////////////////////////////////////////////////////////////////////////////////
 
-FSkinnedMeshStateStreamImpl::FSkinnedMeshStateStreamImpl(FScene* InScene)
+FSkinnedMeshStateStreamImpl::FSkinnedMeshStateStreamImpl(FSceneInterface& InScene)
 :	Scene(InScene)
 {
 }
@@ -132,8 +133,6 @@ void FSkinnedMeshStateStreamImpl::Render_OnCreate(const FSkinnedMeshStaticState&
 		LocalBounds = SkinnedAsset->GetBounds();
 	}	
 
-	UWorld& World = *Scene->World;
-
 	Object.PrimitiveSceneDesc.RenderMatrix = Transform.ToMatrixWithScale();
 	Object.PrimitiveSceneDesc.PrimitiveSceneData = &Object.PrimitiveSceneData;
 	Object.PrimitiveSceneDesc.LocalBounds = LocalBounds;
@@ -143,9 +142,8 @@ void FSkinnedMeshStateStreamImpl::Render_OnCreate(const FSkinnedMeshStaticState&
 	Desc.SkinnedAsset = SkinnedAsset;
 	Desc.OverrideMaterials = const_cast<TArray<TObjectPtr<UMaterialInterface>>&>(Ds.GetOverrideMaterials());
 	Desc.CustomPrimitiveData = &Object.CustomPrimitiveData;
-	Desc.Scene = Scene;
-	Desc.World = &World;
-	Desc.FeatureLevel = World.GetFeatureLevel();
+	Desc.Scene = &Scene;
+	Desc.FeatureLevel = Scene.GetFeatureLevel();
 	Desc.bIsVisible = Info.bVisible;
 	Desc.bPerBoneMotionBlur = true;
 	Desc.VisibilityId = -1;
@@ -182,7 +180,7 @@ void FSkinnedMeshStateStreamImpl::Render_OnCreate(const FSkinnedMeshStaticState&
 
 	MeshObject->Update(LODLevel, ProxyData, SceneProxy, SkinnedAsset, FMorphTargetWeightMap(), TArray<float>(), EPreviousBoneTransformUpdateMode::None, FExternalMorphWeightData());
 
-	Scene->AddPrimitive(&Object.PrimitiveSceneDesc);
+	Scene.AddPrimitive(&Object.PrimitiveSceneDesc);
 
 	UserData = &Object;
 }
@@ -208,3 +206,7 @@ void FSkinnedMeshStateStreamImpl::Render_OnDestroy(const FSkinnedMeshStaticState
 }
 
 ////////////////////////////////////////////////////////////////////////////////////////////////////
+
+STATESTREAM_CREATOR_INSTANCE_WITH_DEPENDENCY(FSkinnedMeshStateStreamImpl, TransformStateStreamId)
+
+////////////////////////////////////////////////////////////////////////////////////////////////////
diff --git a/Engine/Source/Runtime/Renderer/Private/StateStream/SkinnedMeshStateStreamImpl.h b/Engine/Source/Runtime/Renderer/Private/StateStream/SkinnedMeshStateStreamImpl.h
--- a/Engine/Source/Runtime/Renderer/Private/StateStream/SkinnedMeshStateStreamImpl.h
+++ b/Engine/Source/Runtime/Renderer/Private/StateStream/SkinnedMeshStateStreamImpl.h
@@ -9,7 +9,7 @@
 
 #define UE_API RENDERER_API
 
-class FScene;
+class FSceneInterface;
 class FSkeletalMeshObject;
 struct FPrimitiveSceneDesc;
 
@@ -57,14 +57,14 @@ struct FSkinnedMeshStateStreamSettings : TStateStreamSettings<ISkinnedMeshStateS
 class FSkinnedMeshStateStreamImpl : public TStateStream<FSkinnedMeshStateStreamSettings>
 {
 public:
-	FSkinnedMeshStateStreamImpl(FScene* InScene);
+	FSkinnedMeshStateStreamImpl(FSceneInterface& InScene);
 private:
 	UE_API void SetTransformObject(FSkinnedMeshObject& Object, const FSkinnedMeshDynamicState& Ds);
 	UE_API virtual void Render_OnCreate(const FSkinnedMeshStaticState& Ss, const FSkinnedMeshDynamicState& Ds, FSkinnedMeshObject*& UserData, bool IsDestroyedInSameFrame) override;
 	UE_API virtual void Render_OnUpdate(const FSkinnedMeshStaticState& Ss, const FSkinnedMeshDynamicState& Ds, FSkinnedMeshObject*& UserData) override;
 	UE_API virtual void Render_OnDestroy(const FSkinnedMeshStaticState& Ss, const FSkinnedMeshDynamicState& Ds, FSkinnedMeshObject*& UserData) override;
 
-	FScene* Scene;
+	FSceneInterface& Scene;
 };
 
 ////////////////////////////////////////////////////////////////////////////////////////////////////
diff --git a/Engine/Source/Runtime/Renderer/Private/StateStream/StaticMeshStateStreamImpl.cpp b/Engine/Source/Runtime/Renderer/Private/StateStream/StaticMeshStateStreamImpl.cpp
--- a/Engine/Source/Runtime/Renderer/Private/StateStream/StaticMeshStateStreamImpl.cpp
+++ b/Engine/Source/Runtime/Renderer/Private/StateStream/StaticMeshStateStreamImpl.cpp
@@ -2,6 +2,7 @@
 
 #include "StaticMeshStateStreamImpl.h"
 #include "ScenePrivate.h"
+#include "StateStreamCreator.h"
 #include "StaticMeshSceneProxy.h"
 #include "StaticMeshSceneProxyDesc.h"
 
@@ -9,7 +10,7 @@
 
 FStaticMeshObject::~FStaticMeshObject()
 {
-	FScene& Scene = (FScene&)PrimitiveSceneData.SceneProxy->GetScene();
+	FSceneInterface& Scene = PrimitiveSceneData.SceneProxy->GetScene();
 	Scene.RemovePrimitive(&PrimitiveSceneDesc);
 
 	if (TransformObject)
@@ -27,7 +28,7 @@ void FStaticMeshObject::OnTransformObjectDirty()
 
 ////////////////////////////////////////////////////////////////////////////////////////////////////
 
-FStaticMeshStateStreamImpl::FStaticMeshStateStreamImpl(FScene* InScene)
+FStaticMeshStateStreamImpl::FStaticMeshStateStreamImpl(FSceneInterface& InScene)
 :	Scene(InScene)
 {
 }
@@ -74,8 +75,6 @@ void FStaticMeshStateStreamImpl::Render_OnCreate(const FStaticMeshStaticState& S
 		LocalBounds = Mesh->GetBounds();
 	}	
 
-	UWorld& World = *Scene->World;
-
 	Object.PrimitiveSceneDesc.RenderMatrix = Transform.ToMatrixWithScale();
 	Object.PrimitiveSceneDesc.AttachmentRootPosition = Transform.GetLocation();
 	Object.PrimitiveSceneDesc.PrimitiveSceneData = &Object.PrimitiveSceneData;
@@ -86,9 +85,8 @@ void FStaticMeshStateStreamImpl::Render_OnCreate(const FStaticMeshStaticState& S
 	Desc.StaticMesh = Mesh;
 	Desc.OverrideMaterials = const_cast<TArray<TObjectPtr<UMaterialInterface>>&>(Ds.GetOverrideMaterials());
 	Desc.CustomPrimitiveData = &Object.CustomPrimitiveData;
-	Desc.Scene = Scene;
-	Desc.World = &World;
-	Desc.FeatureLevel = World.GetFeatureLevel();
+	Desc.Scene = &Scene;
+	Desc.FeatureLevel = Scene.GetFeatureLevel();
 	Desc.bIsVisible = Info.bVisible;
 	Desc.bOnlyOwnerSee = Ds.GetOnlyOwnerSee();
 	Desc.bOwnerNoSee = Ds.GetOwnerNoSee();
@@ -101,7 +99,7 @@ void FStaticMeshStateStreamImpl::Render_OnCreate(const FStaticMeshStaticState& S
 
 	Object.PrimitiveSceneData.SceneProxy = new FStaticMeshSceneProxy(Desc, false);
 
-	Scene->AddPrimitive(&Object.PrimitiveSceneDesc);
+	Scene.AddPrimitive(&Object.PrimitiveSceneDesc);
 
 	UserData = &Object;
 }
@@ -130,3 +128,7 @@ void FStaticMeshStateStreamImpl::Render_OnDestroy(const FStaticMeshStaticState&
 }
 
 ////////////////////////////////////////////////////////////////////////////////////////////////////
+
+STATESTREAM_CREATOR_INSTANCE_WITH_DEPENDENCY(FStaticMeshStateStreamImpl, TransformStateStreamId)
+
+////////////////////////////////////////////////////////////////////////////////////////////////////
diff --git a/Engine/Source/Runtime/Renderer/Private/StateStream/StaticMeshStateStreamImpl.h b/Engine/Source/Runtime/Renderer/Private/StateStream/StaticMeshStateStreamImpl.h
--- a/Engine/Source/Runtime/Renderer/Private/StateStream/StaticMeshStateStreamImpl.h
+++ b/Engine/Source/Runtime/Renderer/Private/StateStream/StaticMeshStateStreamImpl.h
@@ -9,7 +9,7 @@
 
 #define UE_API RENDERER_API
 
-class FScene;
+class FSceneInterface;
 struct FPrimitiveSceneDesc;
 
 ////////////////////////////////////////////////////////////////////////////////////////////////////
@@ -44,14 +44,14 @@ struct FStaticMeshStateStreamSettings : TStateStreamSettings<IStaticMeshStateStr
 class FStaticMeshStateStreamImpl : public TStateStream<FStaticMeshStateStreamSettings>
 {
 public:
-	FStaticMeshStateStreamImpl(FScene* InScene);
+	FStaticMeshStateStreamImpl(FSceneInterface& InScene);
 private:
 	UE_API void SetTransformObject(FStaticMeshObject& Object, const FStaticMeshDynamicState& Ds);
 	UE_API virtual void Render_OnCreate(const FStaticMeshStaticState& Ss, const FStaticMeshDynamicState& Ds, FStaticMeshObject*& UserData, bool IsDestroyedInSameFrame) override;
 	UE_API virtual void Render_OnUpdate(const FStaticMeshStaticState& Ss, const FStaticMeshDynamicState& Ds, FStaticMeshObject*& UserData) override;
 	UE_API virtual void Render_OnDestroy(const FStaticMeshStaticState& Ss, const FStaticMeshDynamicState& Ds, FStaticMeshObject*& UserData) override;
 
-	FScene* Scene;
+	FSceneInterface& Scene;
 };
 
 ////////////////////////////////////////////////////////////////////////////////////////////////////
