Fix Mac and Android CIS errors

#rnx

diff --git a/Engine/Source/Runtime/Renderer/Private/VirtualShadowMaps/VirtualShadowMapCacheManager.cpp b/Engine/Source/Runtime/Renderer/Private/VirtualShadowMaps/VirtualShadowMapCacheManager.cpp
--- a/Engine/Source/Runtime/Renderer/Private/VirtualShadowMaps/VirtualShadowMapCacheManager.cpp
+++ b/Engine/Source/Runtime/Renderer/Private/VirtualShadowMaps/VirtualShadowMapCacheManager.cpp
@@ -1708,6 +1708,10 @@ void FVirtualShadowMapArrayCacheManager::TrimLoggingInfo()
 #endif
 }
 
+FVirtualShadowMapArrayCacheManager::FViewData::FViewData() 
+	: InstanceState(1024, TEXT("Shadow.Virtual.CacheManager.ViewState"))
+{
+}
 
 FVirtualShadowMapInvalidationSceneUpdater::FVirtualShadowMapInvalidationSceneUpdater(FVirtualShadowMapArrayCacheManager& InCacheManager)
 	: CacheManager(InCacheManager)
@@ -1804,7 +1808,7 @@ void FVirtualShadowMapInvalidationSceneUpdater::PostSceneUpdate(FRDGBuilder& Gra
 		FPersistentViewId ViewId { It.GetIndex() };
 		if (!CacheManager.ViewData.IsValidIndex(ViewId.Index))
 		{
-			CacheManager.ViewData.EmplaceAt(ViewId.Index, TPersistentStructuredBuffer<uint32>(1024, TEXT("Shadow.Virtual.CacheManager.ViewState")));
+			CacheManager.ViewData.EmplaceAt(ViewId.Index);
 		}
 	}
 
diff --git a/Engine/Source/Runtime/Renderer/Private/VirtualShadowMaps/VirtualShadowMapCacheManager.h b/Engine/Source/Runtime/Renderer/Private/VirtualShadowMaps/VirtualShadowMapCacheManager.h
--- a/Engine/Source/Runtime/Renderer/Private/VirtualShadowMaps/VirtualShadowMapCacheManager.h
+++ b/Engine/Source/Runtime/Renderer/Private/VirtualShadowMaps/VirtualShadowMapCacheManager.h
@@ -511,6 +511,9 @@ private:
 		// Bit vector 0: CacheAsDynamic
 		// Bit vector 1: IsTracked
 		static constexpr int32 NumBitsPerInstance = 2;
+
+		FViewData();
+
 		// Buffer that stores NumBitsPerInstance bits per instance indicating whether it is dynamic of static.
 		TPersistentStructuredBuffer<uint32> InstanceState;
 	};
