Handle a few cases where there are no instances in the scene or otherwise there is no work for one of the stages in the new GPU cache invalidation path.

#rnx

diff --git a/Engine/Source/Runtime/Renderer/Private/VirtualShadowMaps/VirtualShadowMapCacheManager.cpp b/Engine/Source/Runtime/Renderer/Private/VirtualShadowMaps/VirtualShadowMapCacheManager.cpp
--- a/Engine/Source/Runtime/Renderer/Private/VirtualShadowMaps/VirtualShadowMapCacheManager.cpp
+++ b/Engine/Source/Runtime/Renderer/Private/VirtualShadowMaps/VirtualShadowMapCacheManager.cpp
@@ -1829,6 +1829,13 @@ void FVirtualShadowMapInvalidationSceneUpdater::PostGPUSceneUpdate(FRDGBuilder&
 	const int32 NumInstanceWordsNeeded = CacheManager.InstanceStateMaskWordStride * FVirtualShadowMapArrayCacheManager::FViewData::NumBitsPerInstance;
 	if (NumInstanceWordsNeeded > 0)
 	{
+		TArray<FRDGBuffer*, TInlineAllocator<16,  SceneRenderingAllocator> > InstanceStateRDG;
+		InstanceStateRDG.SetNumZeroed(CacheManager.ViewData.Num());
+		for (int32 Index = 0; Index < CacheManager.ViewData.Num(); ++Index)
+		{
+			InstanceStateRDG[Index] = CacheManager.ViewData[Index].InstanceState.ResizeAndClearBufferIfNeeded(GraphBuilder, NumInstanceWordsNeeded);
+		}
+
 		FLoadBalancer LoadBalancer;
 		for (FPrimitiveSceneInfo* PrimitiveSceneInfo : PostUpdateChangeSet->AddedPrimitiveSceneInfos)
 		{
@@ -1844,12 +1851,10 @@ void FVirtualShadowMapInvalidationSceneUpdater::PostGPUSceneUpdate(FRDGBuilder&
 
 			FGPUSceneResourceParameters GPUSceneParams = CacheManager.Scene->GPUScene.GetShaderParameters(GraphBuilder);
 
-			for (auto& ViewData : CacheManager.ViewData)
+			for (int32 Index = 0; Index < CacheManager.ViewData.Num(); ++Index)
 			{
-				FRDGBuffer* InstanceStateRDG = ViewData.InstanceState.ResizeAndClearBufferIfNeeded(GraphBuilder, NumInstanceWordsNeeded);
-
 				FVSMResetInstanceStateCS::FParameters* PassParameters = GraphBuilder.AllocParameters<FVSMResetInstanceStateCS::FParameters>();
-				PassParameters->InOutViewInstanceState = GraphBuilder.CreateUAV(InstanceStateRDG);
+				PassParameters->InOutViewInstanceState = GraphBuilder.CreateUAV(InstanceStateRDG[Index]);
 				PassParameters->GPUScene = GPUSceneParams;
 				PassParameters->StateWordStride = CacheManager.InstanceStateMaskWordStride;
 				LoadBalancerGpuData.AddPass(GraphBuilder, RDG_EVENT_NAME("ResetInstanceState (%d batches)", LoadBalancerGpuData.NumBatches), ComputeShader, PassParameters);
@@ -1945,19 +1950,30 @@ public:
 		// check(CacheInstanceAsDynamicInstanceMaskWordStride * 32 >= NumInstances);
 		// Allocate space for WordStride words for each view, x2 for storing state change mask
 		int32 NumCacheInstanceAsDynamicWords = 2 * CacheInstanceAsDynamicInstanceMaskWordStride * NumViews;
-		CacheInstanceAsDynamicRDG = GraphBuilder.CreateBuffer(FRDGBufferDesc::CreateStructuredDesc(sizeof(uint32), NumCacheInstanceAsDynamicWords), TEXT("Shadow.Virtual.CacheInstanceAsDynamic"));
-		AddClearUAVPass(GraphBuilder, GraphBuilder.CreateUAV(CacheInstanceAsDynamicRDG), 0u);
 
-		FVirtualShadowMapInvalidationSceneUniforms VirtualShadowMapInvalidationSceneUniforms;
-		VirtualShadowMapInvalidationSceneUniforms.CacheInstanceAsDynamic = GraphBuilder.CreateSRV(CacheInstanceAsDynamicRDG);
-		VirtualShadowMapInvalidationSceneUniforms.InstanceStateViewWordStride = CacheInstanceAsDynamicInstanceMaskWordStride;
-		SceneUniforms.Set(SceneUB::VSMCache, VirtualShadowMapInvalidationSceneUniforms);
+		// If there are no instances or views, leave SceneUB with defaults.
+		if (NumCacheInstanceAsDynamicWords > 0)
+		{
+			CacheInstanceAsDynamicRDG = GraphBuilder.CreateBuffer(FRDGBufferDesc::CreateStructuredDesc(sizeof(uint32), NumCacheInstanceAsDynamicWords), TEXT("Shadow.Virtual.CacheInstanceAsDynamic"));
+			AddClearUAVPass(GraphBuilder, GraphBuilder.CreateUAV(CacheInstanceAsDynamicRDG), 0u);
+
+			FVirtualShadowMapInvalidationSceneUniforms VirtualShadowMapInvalidationSceneUniforms;
+			VirtualShadowMapInvalidationSceneUniforms.CacheInstanceAsDynamic = GraphBuilder.CreateSRV(CacheInstanceAsDynamicRDG);
+			VirtualShadowMapInvalidationSceneUniforms.InstanceStateViewWordStride = CacheInstanceAsDynamicInstanceMaskWordStride;
+			SceneUniforms.Set(SceneUB::VSMCache, VirtualShadowMapInvalidationSceneUniforms);
+		}
 	}
 
 	virtual void PreRender(FRDGBuilder& GraphBuilder) override
 	{
 		// This happens post visibility update so here we know what instances are deemed deforming.
 
+		// Early out if there are no instances (or no views?)
+		if (CacheInstanceAsDynamicRDG == nullptr)
+		{
+			return;
+		}
+
 		FSceneRendererBase& SceneRenderer = GetSceneRenderer();
 		ERHIFeatureLevel::Type FeatureLevel = SceneRenderer.Scene->GetFeatureLevel();
 
@@ -1979,6 +1995,7 @@ public:
 		SharedParameters.InvalidationQueueMax =  InvalidationQueueMax;
 		SharedParameters.StateWordStride = CacheManager.InstanceStateMaskWordStride;
 
+		bool bWasAnyInvalidationsQueued = false;
 		for (int32 ViewIndex = 0; ViewIndex < SceneRenderer.Views.Num(); ++ViewIndex)
 		{
 			FViewInfo& View = SceneRenderer.Views[ViewIndex];
@@ -2001,10 +2018,11 @@ public:
 				PassParameters,
 				FComputeShaderUtils::GetGroupCount(CacheManager.InstanceStateMaskWordStride, FVSMUpdateViewInstanceStateCS::ThreadGroupSize)
 			);
+			bWasAnyInvalidationsQueued = true;
 		}
 
 		// Run the resultant invalidations right away, this could be deferred and flushed later instead if desired.
-		if (CacheManager.IsCacheDataAvailable())
+		if (CacheManager.IsCacheDataAvailable() && bWasAnyInvalidationsQueued)
 		{
 			RDG_EVENT_SCOPE(GraphBuilder, "Shadow.Virtual.ProcessInvalidationsGPU");
 
