Fixed DOF artifacting when r.PostProcessing.PropagateAlpha = true due to missing CoC data output in reduce steps for the last mips.
Fixed crash when DOF quality is set to cinematic due to missing scatter occlusion texture in gather pass.
Fixed crash when bEnableScatterBokehSettings = false and bEnableSlightOutOfFocusBokeh = true due to missing bokeh LUT texture in scatter pass.

diff --git a/Engine/Source/Runtime/Renderer/Private/PostProcess/DiaphragmDOF.cpp b/Engine/Source/Runtime/Renderer/Private/PostProcess/DiaphragmDOF.cpp
--- a/Engine/Source/Runtime/Renderer/Private/PostProcess/DiaphragmDOF.cpp
+++ b/Engine/Source/Runtime/Renderer/Private/PostProcess/DiaphragmDOF.cpp
@@ -2081,6 +2081,7 @@ bool DiaphragmDOF::AddPasses(
 			FIntPoint PassViewSize = FIntPoint::DivideAndRoundUp(PreprocessViewSize, 1 << InputMipLevel);
 
 			FDiaphragmDOFReduceCS::FPermutationDomain PermutationVector;
+			PermutationVector.Set<FDiaphragmDOFShader::FAlphaChannelDim>(bProcessSceneAlpha);
 			PermutationVector.Set<FDiaphragmDOFReduceCS::FReduceMipCount>(ProcessingMipLevelCount + 1);
 			PermutationVector.Set<FDDOFRGBColorBufferDim>(bRGBBufferSeparateCocBuffer);
 			PermutationVector.Set<FDiaphragmDOFReduceCS::FDisableOutputMip0>(true);
@@ -2239,6 +2240,10 @@ bool DiaphragmDOF::AddPasses(
 			FDOFConvolutionTextures* ConvolutionOutputTextures,
 			FRDGTextureRef* ScatterOcclusionTexture)
 		{
+			bool bUsesScatterOcclusionTexture =
+				ConvolutionSettings.QualityConfig == EDiaphragmDOFGatherQuality::HighQualityWithHybridScatterOcclusion ||
+				ConvolutionSettings.QualityConfig == EDiaphragmDOFGatherQuality::Cinematic;
+
 			// Allocate output textures
 			{
 				FRDGTextureDesc Desc = ReducedGatherInputTextures.SceneColor->Desc;
@@ -2276,7 +2281,7 @@ bool DiaphragmDOF::AddPasses(
 					}
 				}
 
-				if (ConvolutionSettings.QualityConfig == EDiaphragmDOFGatherQuality::HighQualityWithHybridScatterOcclusion)
+				if (bUsesScatterOcclusionTexture)
 				{
 					Desc.Format = PF_G16R16F;
 			
@@ -2378,7 +2383,7 @@ bool DiaphragmDOF::AddPasses(
 			PassParameters->BokehLUT = BokehLUT;
 
 			PassParameters->ConvolutionOutput = CreateUAVs(GraphBuilder, *ConvolutionOutputTextures);
-			if (ConvolutionSettings.QualityConfig == EDiaphragmDOFGatherQuality::HighQualityWithHybridScatterOcclusion)
+			if (bUsesScatterOcclusionTexture)
 			{
 				PassParameters->ScatterOcclusionOutput = GraphBuilder.CreateUAV(*ScatterOcclusionTexture);
 			}
@@ -2513,8 +2518,7 @@ bool DiaphragmDOF::AddPasses(
 			PassParameters->CocSqueeze = CocModel.Squeeze;
 			PassParameters->CocInvSqueeze = 1.0 / CocModel.Squeeze;
 			PassParameters->CommonParameters = CommonParameters;
-			if (bEnableScatterBokehSettings)
-				PassParameters->BokehLUT = ScatteringBokehLUT;
+			PassParameters->BokehLUT = ScatteringBokehLUT;
 			PassParameters->ScatterOcclusionSize = FVector4f(RefBufferSize.X, RefBufferSize.Y, 1.0f / RefBufferSize.X, 1.0f / RefBufferSize.Y);
 			PassParameters->ScatterOcclusion = ScatterOcclusionTexture;
 			PassParameters->IndirectDrawParameter = DrawIndirectParametersBuffer;
