[HWRT] Dont access ray tracing shaders in PrepareRayTracingShadows if they are not supported.

#rnx

diff --git a/Engine/Source/Runtime/Renderer/Private/RayTracing/RayTracingShadows.cpp b/Engine/Source/Runtime/Renderer/Private/RayTracing/RayTracingShadows.cpp
--- a/Engine/Source/Runtime/Renderer/Private/RayTracing/RayTracingShadows.cpp
+++ b/Engine/Source/Runtime/Renderer/Private/RayTracing/RayTracingShadows.cpp
@@ -322,6 +322,11 @@ void FDeferredShadingSceneRenderer::PrepareRayTracingShadows(const FViewInfo& Vi
 	// We have to check if ray tracing is enabled on any of the scene lights. The Scene.bHasRayTracedLights is computed using ShouldRenderRayTracingShadowsForLight() helper, 
 	// which handles various override conditions.
 
+	if (!ShouldRenderRayTracingEffect(true, ERayTracingPipelineCompatibilityFlags::FullPipeline, View))
+	{
+		return;
+	}
+
 	if (!View.bHasRayTracingShadows)
 	{
 		return;
