HeterogeneousVolumes: Adding IndirectInscatteringFactor

Allows for global adjustment of indirect inscattering
r.HeterogeneousVolumes.IndirectLighting (default = 0.0)

diff --git a/Engine/Shaders/Private/HeterogeneousVolumes/HeterogeneousVolumesLiveShadingPipeline.usf b/Engine/Shaders/Private/HeterogeneousVolumes/HeterogeneousVolumesLiveShadingPipeline.usf
--- a/Engine/Shaders/Private/HeterogeneousVolumes/HeterogeneousVolumesLiveShadingPipeline.usf
+++ b/Engine/Shaders/Private/HeterogeneousVolumes/HeterogeneousVolumesLiveShadingPipeline.usf
@@ -90,6 +90,9 @@ float ShadowStepFactor;
 int bApplyHeightFog;
 int bApplyVolumetricFog;
 
+// Indirect
+float IndirectInscatteringFactor;
+
 // Ray data
 float MaxTraceDistance;
 float MaxShadowTraceDistance;
@@ -217,6 +220,11 @@ int ShouldApplyVolumetricFog()
 	return bApplyVolumetricFog;
 }
 
+float GetIndirectInscatteringFactor()
+{
+	return IndirectInscatteringFactor;
+}
+
 struct FVolumeSampleContext
 {
 	FMaterialPixelParameters MaterialParameters;
diff --git a/Engine/Shaders/Private/HeterogeneousVolumes/HeterogeneousVolumesLiveShadingShadows.usf b/Engine/Shaders/Private/HeterogeneousVolumes/HeterogeneousVolumesLiveShadingShadows.usf
--- a/Engine/Shaders/Private/HeterogeneousVolumes/HeterogeneousVolumesLiveShadingShadows.usf
+++ b/Engine/Shaders/Private/HeterogeneousVolumes/HeterogeneousVolumesLiveShadingShadows.usf
@@ -142,6 +142,11 @@ int GetStochasticFilteringMode()
 	return StochasticFilteringMode;
 }
 
+float GetIndirectInscatteringFactor()
+{
+	return 0.0;
+}
+
 struct FVolumeSampleContext
 {
 	FMaterialPixelParameters MaterialParameters;
diff --git a/Engine/Shaders/Private/HeterogeneousVolumes/HeterogeneousVolumesRayMarchingUtils.ush b/Engine/Shaders/Private/HeterogeneousVolumes/HeterogeneousVolumesRayMarchingUtils.ush
--- a/Engine/Shaders/Private/HeterogeneousVolumes/HeterogeneousVolumesRayMarchingUtils.ush
+++ b/Engine/Shaders/Private/HeterogeneousVolumes/HeterogeneousVolumesRayMarchingUtils.ush
@@ -403,6 +403,7 @@ void RayMarchSingleScattering(
 			float PhaseG = 0.0;
 			RotatedHGZonalHarmonic.V = float4(1.0f, CameraVector.y, CameraVector.z, CameraVector.x) * float4(1.0f, PhaseG, PhaseG, PhaseG);
 			float3 IndirectInscattering = max(DotSH(TranslucencyGISH, RotatedHGZonalHarmonic), 0) / PI;
+			IndirectInscattering *= GetIndirectInscatteringFactor();
 
 			float3 Extinction = SampleExtinction(SampleContext);
 			float3 Albedo = SampleAlbedo(SampleContext);
diff --git a/Engine/Shaders/Private/HeterogeneousVolumes/HeterogeneousVolumesSparseVoxelUniformBufferUtils.ush b/Engine/Shaders/Private/HeterogeneousVolumes/HeterogeneousVolumesSparseVoxelUniformBufferUtils.ush
--- a/Engine/Shaders/Private/HeterogeneousVolumes/HeterogeneousVolumesSparseVoxelUniformBufferUtils.ush
+++ b/Engine/Shaders/Private/HeterogeneousVolumes/HeterogeneousVolumesSparseVoxelUniformBufferUtils.ush
@@ -67,6 +67,11 @@ int ShouldApplyVolumetricFog()
 	return SparseVoxelUniformBuffer.bApplyVolumetricFog;
 }
 
+float GetIndirectInscatteringFactor()
+{
+	return SparseVoxelUniformBuffer.IndirectInscatteringFactor;
+}
+
 uint3 GetVolumeResolution()
 {
 	return SparseVoxelUniformBuffer.VolumeResolution;
diff --git a/Engine/Source/Runtime/Renderer/Private/HeterogeneousVolumes/HeterogeneousVolumes.cpp b/Engine/Source/Runtime/Renderer/Private/HeterogeneousVolumes/HeterogeneousVolumes.cpp
--- a/Engine/Source/Runtime/Renderer/Private/HeterogeneousVolumes/HeterogeneousVolumes.cpp
+++ b/Engine/Source/Runtime/Renderer/Private/HeterogeneousVolumes/HeterogeneousVolumes.cpp
@@ -72,10 +72,10 @@ static TAutoConsoleVariable<int32> CVarHeterogeneousVolumesHardwareRayTracing(
 	ECVF_RenderThreadSafe
 );
 
-static TAutoConsoleVariable<int32> CVarHeterogeneousVolumesIndirectLighting(
+static TAutoConsoleVariable<float> CVarHeterogeneousVolumesIndirectLighting(
 	TEXT("r.HeterogeneousVolumes.IndirectLighting"),
-	0,
-	TEXT("Enables indirect lighting (Default = 0)"),
+	0.0f,
+	TEXT("Enables indirect lighting (Default = 0.0)"),
 	ECVF_RenderThreadSafe
 );
 
@@ -557,6 +557,11 @@ namespace HeterogeneousVolumes
 		return CVarHeterogeneousVolumesIndirectLighting.GetValueOnRenderThread() != 0;
 	}
 
+	float GetIndirectLightingFactor()
+	{
+		return FMath::Max(CVarHeterogeneousVolumesIndirectLighting.GetValueOnRenderThread(), 0.0f);
+	}
+
 	bool ShouldApplyHeightFog()
 	{
 		return CVarHeterogeneousVolumesApplyHeightFog.GetValueOnRenderThread() != 0;
diff --git a/Engine/Source/Runtime/Renderer/Private/HeterogeneousVolumes/HeterogeneousVolumes.h b/Engine/Source/Runtime/Renderer/Private/HeterogeneousVolumes/HeterogeneousVolumes.h
--- a/Engine/Source/Runtime/Renderer/Private/HeterogeneousVolumes/HeterogeneousVolumes.h
+++ b/Engine/Source/Runtime/Renderer/Private/HeterogeneousVolumes/HeterogeneousVolumes.h
@@ -93,6 +93,7 @@ namespace HeterogeneousVolumes
 	float GetShadowRelativeErrorThreshold();
 	bool UseAVSMCompression();
 	float GetCameraDownsampleFactor();
+	float GetIndirectLightingFactor();
 
 	int32 GetMipLevel();
 	int32 GetDebugMode();
@@ -210,6 +211,7 @@ BEGIN_UNIFORM_BUFFER_STRUCT(FSparseVoxelUniformBufferParameters, )
 	SHADER_PARAMETER(float, StepFactor)
 	SHADER_PARAMETER(float, ShadowStepSize)
 	SHADER_PARAMETER(float, ShadowStepFactor)
+	SHADER_PARAMETER(float, IndirectInscatteringFactor)
 	SHADER_PARAMETER(int, bApplyHeightFog)
 	SHADER_PARAMETER(int, bApplyVolumetricFog)
 END_UNIFORM_BUFFER_STRUCT()
diff --git a/Engine/Source/Runtime/Renderer/Private/HeterogeneousVolumes/HeterogeneousVolumesLiveShadingPipeline.cpp b/Engine/Source/Runtime/Renderer/Private/HeterogeneousVolumes/HeterogeneousVolumesLiveShadingPipeline.cpp
--- a/Engine/Source/Runtime/Renderer/Private/HeterogeneousVolumes/HeterogeneousVolumesLiveShadingPipeline.cpp
+++ b/Engine/Source/Runtime/Renderer/Private/HeterogeneousVolumes/HeterogeneousVolumesLiveShadingPipeline.cpp
@@ -462,6 +462,7 @@ class FRenderSingleScatteringWithLiveShadingCS : public FMeshMaterialShader
 
 		// Indirect Lighting
 		SHADER_PARAMETER_RDG_UNIFORM_BUFFER(FLumenTranslucencyLightingUniforms, LumenGIVolumeStruct)
+		SHADER_PARAMETER(float, IndirectInscatteringFactor)
 
 		// Object data
 		SHADER_PARAMETER(FMatrix44f, LocalToWorld)
@@ -1207,6 +1208,7 @@ void RenderSingleScatteringWithLiveShading(
 		}
 
 		// Indirect lighting data
+		PassParameters->IndirectInscatteringFactor = HeterogeneousVolumes::GetIndirectLightingFactor();
 		auto* LumenUniforms = GraphBuilder.AllocParameters<FLumenTranslucencyLightingUniforms>();
 		LumenUniforms->Parameters = GetLumenTranslucencyLightingParameters(GraphBuilder, View.GetLumenTranslucencyGIVolume(), View.LumenFrontLayerTranslucency);
 		PassParameters->LumenGIVolumeStruct = GraphBuilder.CreateUniformBuffer(LumenUniforms);
diff --git a/Engine/Source/Runtime/Renderer/Private/HeterogeneousVolumes/HeterogeneousVolumesPreshadingPipeline.cpp b/Engine/Source/Runtime/Renderer/Private/HeterogeneousVolumes/HeterogeneousVolumesPreshadingPipeline.cpp
--- a/Engine/Source/Runtime/Renderer/Private/HeterogeneousVolumes/HeterogeneousVolumesPreshadingPipeline.cpp
+++ b/Engine/Source/Runtime/Renderer/Private/HeterogeneousVolumes/HeterogeneousVolumesPreshadingPipeline.cpp
@@ -1676,6 +1676,7 @@ void RenderWithPreshading(
 		SparseVoxelUniformBufferParameters->ShadowStepFactor = HeterogeneousVolumeInterface->GetShadowStepFactor() * LODFactor;
 		SparseVoxelUniformBufferParameters->bApplyHeightFog = HeterogeneousVolumes::ShouldApplyHeightFog();
 		SparseVoxelUniformBufferParameters->bApplyVolumetricFog = HeterogeneousVolumes::ShouldApplyVolumetricFog();
+		SparseVoxelUniformBufferParameters->IndirectInscatteringFactor = HeterogeneousVolumes::GetIndirectLightingFactor();
 	}
 	TRDGUniformBufferRef<FSparseVoxelUniformBufferParameters> SparseVoxelUniformBuffer = GraphBuilder.CreateUniformBuffer(SparseVoxelUniformBufferParameters);
 
