[Local Exposure] Experiment with soft thresholds.

diff --git a/Engine/Shaders/Private/PostProcessHistogramCommon.ush b/Engine/Shaders/Private/PostProcessHistogramCommon.ush
--- a/Engine/Shaders/Private/PostProcessHistogramCommon.ush
+++ b/Engine/Shaders/Private/PostProcessHistogramCommon.ush
@@ -43,6 +43,8 @@ float2 LocalExposure_BilateralGridUVScale;
 
 float LocalExposure_HighlightThreshold;
 float LocalExposure_ShadowThreshold;
+uint LocalExposure_bUseSoftThreshold;
+float LocalExposure_SoftThresholdStrength;
 
 float CalculateEyeAdaptationLuminance(float3 Color)
 {
@@ -252,8 +254,20 @@ float CalculateLogLocalExposure(float LogLum, float BaseLogLum, float LogMiddleG
 	float ContrastScale = BaseCentered > 0 ? HighlightContrastScale : ShadowContrastScale;
 
 	// Apply threshold
+	float ThresholdOffset = 0;
 
-	float ThresholdOffset;
+	if(LocalExposure_bUseSoftThreshold != 0)
+	{
+		if (BaseCentered > 0)
+		{
+			ContrastScale = lerp(1, ContrastScale, saturate((BaseCentered - LocalExposure_HighlightThreshold) * LocalExposure_SoftThresholdStrength));
+		}
+		else
+		{
+			ContrastScale = lerp(1, ContrastScale, saturate((abs(BaseCentered) - LocalExposure_ShadowThreshold) * LocalExposure_SoftThresholdStrength));
+		}
+	}
+	else
 	{
 		if (BaseCentered > 0)
 		{
diff --git a/Engine/Source/Runtime/Renderer/Private/PostProcess/PostProcessLocalExposure.cpp b/Engine/Source/Runtime/Renderer/Private/PostProcess/PostProcessLocalExposure.cpp
--- a/Engine/Source/Runtime/Renderer/Private/PostProcess/PostProcessLocalExposure.cpp
+++ b/Engine/Source/Runtime/Renderer/Private/PostProcess/PostProcessLocalExposure.cpp
@@ -13,6 +13,18 @@
 #include "ShaderCompilerCore.h"
 #include "DataDrivenShaderPlatformInfo.h"
 
+static TAutoConsoleVariable<bool> CVarLocalExposureSoftThreshold(
+	TEXT("r.LocalExposure.SoftThreshold"),
+	false,
+	TEXT("Whether to use soft thresholds."),
+	ECVF_RenderThreadSafe);
+
+static TAutoConsoleVariable<float> CVarLocalExposureSoftThresholdStrength(
+	TEXT("r.LocalExposure.SoftThreshold.Strength"),
+	1.0f,
+	TEXT(""),
+	ECVF_RenderThreadSafe);
+
 static TAutoConsoleVariable<float> CVarExposureFusionTargetLuminance(
 	TEXT("r.LocalExposure.ExposureFusion.TargetLuminance"),
 	0.5f,
@@ -260,6 +272,8 @@ FLocalExposureParameters GetLocalExposureParameters(const FViewInfo& View, FIntP
 	Parameters.BilateralGridUVScale = LocalExposureBilateralGridUVScale;
 	Parameters.HighlightThreshold = Settings.LocalExposureHighlightThreshold;
 	Parameters.ShadowThreshold = Settings.LocalExposureShadowThreshold;
+	Parameters.bUseSoftThreshold = CVarLocalExposureSoftThreshold.GetValueOnRenderThread() ? 1.0f : 0.0f;
+	Parameters.SoftThresholdStrength = CVarLocalExposureSoftThresholdStrength.GetValueOnRenderThread();
 	return Parameters;
 }
 
diff --git a/Engine/Source/Runtime/Renderer/Private/PostProcess/PostProcessLocalExposure.h b/Engine/Source/Runtime/Renderer/Private/PostProcess/PostProcessLocalExposure.h
--- a/Engine/Source/Runtime/Renderer/Private/PostProcess/PostProcessLocalExposure.h
+++ b/Engine/Source/Runtime/Renderer/Private/PostProcess/PostProcessLocalExposure.h
@@ -15,6 +15,8 @@ BEGIN_SHADER_PARAMETER_STRUCT(FLocalExposureParameters, )
 	SHADER_PARAMETER(FVector2f, BilateralGridUVScale)
 	SHADER_PARAMETER(float, HighlightThreshold)
 	SHADER_PARAMETER(float, ShadowThreshold)
+	SHADER_PARAMETER(uint32, bUseSoftThreshold)
+	SHADER_PARAMETER(float, SoftThresholdStrength)
 END_SHADER_PARAMETER_STRUCT()
 
 FLocalExposureParameters GetLocalExposureParameters(const FViewInfo& View, FIntPoint ViewRectSize, const FEyeAdaptationParameters& EyeAdaptationParameters);
