Fix missing exposure correction during OIT combine, when per-pixel OIT is enabled.

diff --git a/Engine/Shaders/Private/OITCombine.usf b/Engine/Shaders/Private/OITCombine.usf
--- a/Engine/Shaders/Private/OITCombine.usf
+++ b/Engine/Shaders/Private/OITCombine.usf
@@ -86,6 +86,8 @@ void MainCS(uint3 DispatchThreadId : SV_DispatchThreadID)
 	// Alpha value used for separate translucency pass
 	float AdjustedAlpha = saturate(dot(Transmittance, float3(1.0f, 1.0f, 1.0f) / 3.0f));
 	
+	FinalColor *= View.PreExposure;
+
 	// Write final color - typed load
 	#if PERMUTATION_MODULATE
 	if (SampleCount > 0)
diff --git a/Engine/Source/Runtime/Renderer/Private/OIT/OIT.cpp b/Engine/Source/Runtime/Renderer/Private/OIT/OIT.cpp
--- a/Engine/Source/Runtime/Renderer/Private/OIT/OIT.cpp
+++ b/Engine/Source/Runtime/Renderer/Private/OIT/OIT.cpp
@@ -185,6 +185,7 @@ class FOITPixelCombineCS : public FGlobalShader
 		SHADER_PARAMETER_RDG_TEXTURE_UAV(RWTexture2D<float4>, OutColorTexture)
 		SHADER_PARAMETER_RDG_TEXTURE_UAV(RWTexture2D<float4>, OutLightingTexture)
 		SHADER_PARAMETER_RDG_TEXTURE_UAV(RWTexture2D<float4>, OutTransmittanceTexture)
+		SHADER_PARAMETER_STRUCT_REF(FViewUniformShaderParameters, ViewUniformBuffer)
 	END_SHADER_PARAMETER_STRUCT()
 public:
 	static bool ShouldCompilePermutation(const FGlobalShaderPermutationParameters& Parameters) { return IsOITSortedPixelsSupported(Parameters.Platform); }
@@ -227,6 +228,7 @@ static void AddInternalOITComposePass(
 	Parameters->MaxSideSampleCount = OITData.MaxSideSamplePerPixel;
 	Parameters->SampleDataTexture = OITData.SampleDataTexture;
 	Parameters->SampleCountTexture = OITData.SampleCountTexture;
+	Parameters->ViewUniformBuffer = View.ViewUniformBuffer;
 	if (bModulate)
 	{
 		Parameters->OutLightingTexture = GraphBuilder.CreateUAV(OutLightingTexture);
