[GPU Scene] Limit number of GPU Lights and show/log warning when it is exceeded.

diff --git a/Engine/Shaders/Shared/SceneDefinitions.h b/Engine/Shaders/Shared/SceneDefinitions.h
--- a/Engine/Shaders/Shared/SceneDefinitions.h
+++ b/Engine/Shaders/Shared/SceneDefinitions.h
@@ -15,13 +15,16 @@
 #define INSTANCE_SCENE_DATA_FLAGS_NUM_BITS	(12u) // Max of 12 flags
 #define INSTANCE_RELATIVE_ID_NUM_BITS		(24u) // Max of 16,777,216 instances per primitive
 #define INSTANCE_CUSTOM_DATA_COUNT_NUM_BITS	 (8u) // Max of 255 custom data floats per instance
+#define LIGHT_ID_NUM_BITS					(24u) // Max of 16,777,216 lights in the buffer
 
 #define MAX_PRIMITIVE_ID					(1U << PRIMITIVE_ID_NUM_BITS)
 #define MAX_INSTANCE_ID						(1U << INSTANCE_ID_NUM_BITS)
+#define MAX_LIGHT_ID						(1U << LIGHT_ID_NUM_BITS)
 
 #define PRIMITIVE_ID_MASK			((1u << PRIMITIVE_ID_NUM_BITS) - 1u)
 #define INSTANCE_ID_MASK			((1u << INSTANCE_ID_NUM_BITS) - 1u)
 #define INSTANCE_RELATIVE_ID_MASK	((1u << INSTANCE_RELATIVE_ID_NUM_BITS) - 1u)
+#define LIGHT_ID_MASK				((1u << LIGHT_ID_NUM_BITS) - 1u)
 
 #define INVALID_PRIMITIVE_ID PRIMITIVE_ID_MASK
 
diff --git a/Engine/Source/Runtime/Renderer/Private/GPUScene.cpp b/Engine/Source/Runtime/Renderer/Private/GPUScene.cpp
--- a/Engine/Source/Runtime/Renderer/Private/GPUScene.cpp
+++ b/Engine/Source/Runtime/Renderer/Private/GPUScene.cpp
@@ -661,6 +661,22 @@ FGPUScene::FGPUScene(FScene &InScene)
 		{
 			bLoggedInstanceOverflow = false;
 		}
+
+		if (NumGPULights < MaxLightsDuringPrevUpdate)
+		{
+			OutMessages.Add(FCoreDelegates::EOnScreenMessageSeverity::Warning, FText::FromString(FString::Printf(TEXT("GPU-Scene Light data overflow detected, reduce light count to avoid rendering artifacts"))));
+			OutMessages.Add(FCoreDelegates::EOnScreenMessageSeverity::Warning, FText::FromString(FString::Printf(TEXT("  Max allocated ID %d, light buffer capacity: %d"), MaxLightsDuringPrevUpdate, NumGPULights)));
+			if (!bLoggedLightOverflow)
+			{
+				UE_LOG(LogRenderer, Warning, TEXT("GPU-Scene Light data overflow detected, reduce light count to avoid rendering artifacts.\n")
+					TEXT(" Max allocated ID %d, light buffer capacity: %d"), MaxLightsDuringPrevUpdate, NumGPULights);
+				bLoggedLightOverflow = true;
+			}
+		}
+		else
+		{
+			bLoggedLightOverflow = false;
+		}
 	});
 #endif
 }
@@ -708,7 +724,14 @@ void FGPUScene::EndRender()
 
 void FGPUScene::UpdateGPULights(FRDGBuilder& GraphBuilder, const UE::Tasks::FTask& PrerequisiteTask)
 {
-	FRDGUploadData<FLightSceneData> LightData(GraphBuilder, FMath::Max(1, Scene.Lights.GetMaxIndex()));
+#if !UE_BUILD_SHIPPING
+	MaxLightsDuringPrevUpdate = uint32(Scene.Lights.GetMaxIndex());
+#endif // UE_BUILD_SHIPPING
+
+	// Clamp buffer to be smaller than the MAX_LIGHT_ID.
+	NumGPULights = FMath::Min(int32(MAX_LIGHT_ID), Scene.Lights.GetMaxIndex());
+
+	FRDGUploadData<FLightSceneData> LightData(GraphBuilder, FMath::Max(1, NumGPULights));
 
 	GraphBuilder.AddSetupTask([this, LightData]
 	{
@@ -716,7 +739,7 @@ void FGPUScene::UpdateGPULights(FRDGBuilder& GraphBuilder, const UE::Tasks::FTas
 		const bool bAllowStaticLighting = IsStaticLightingAllowed();
 		const uint32 Flags = RenderRectLightsAsSpotLights(FeatureLevel) ? ELightShaderParameterFlags::RectAsSpotLight : 0u;
 
-		for (int32 Index = 0; Index < Scene.Lights.GetMaxIndex(); ++Index)
+		for (int32 Index = 0; Index < NumGPULights; ++Index)
 		{
 			if (Scene.Lights.IsAllocated(Index))
 			{
diff --git a/Engine/Source/Runtime/Renderer/Private/GPUScene.h b/Engine/Source/Runtime/Renderer/Private/GPUScene.h
--- a/Engine/Source/Runtime/Renderer/Private/GPUScene.h
+++ b/Engine/Source/Runtime/Renderer/Private/GPUScene.h
@@ -346,7 +346,8 @@ private:
 
 	static constexpr int32 InitialBufferSize = 256;
 
-	TRefCountPtr<FRDGPooledBuffer> LightDataBuffer;
+	TRefCountPtr<FRDGPooledBuffer>	LightDataBuffer;
+	int32							NumGPULights = 0;
 
 	// Buffers used by GPU-Scene, since they can be resized during updates AND the render passes must retain the 
 	// right copy (this is chiefly because the init of shadow views after pre-pass means we need to be able to set 
@@ -456,6 +457,8 @@ private:
 	FDelegateHandle ScreenMessageDelegate;
 	bool bLoggedInstanceOverflow = false;
 	uint32 MaxInstancesDuringPrevUpdate = 0;
+	bool bLoggedLightOverflow = false;
+	int32 MaxLightsDuringPrevUpdate = 0;
 #endif // UE_BUILD_SHIPPING
 
 };
