[MegaLights / Light Grid] Sort lights handled by MegaLights by GPU Scene Light ID

diff --git a/Engine/Source/Runtime/Renderer/Private/LightRendering.cpp b/Engine/Source/Runtime/Renderer/Private/LightRendering.cpp
--- a/Engine/Source/Runtime/Renderer/Private/LightRendering.cpp
+++ b/Engine/Source/Runtime/Renderer/Private/LightRendering.cpp
@@ -1303,7 +1303,13 @@ void FSceneRenderer::GatherAndSortLights(FSortedLightSetSceneInfo& OutSortedLigh
 					if (SortedLightInfo->SortKey.Fields.bLightFunction) 				{ OutSortedLights.bHasLightFunctions = true; }
 
 					SortedLightInfo->SortKey.Fields.bClusteredDeferredNotSupported = !bClusteredDeferredSupported;
-					SortedLightInfo->SortKey.Fields.bHandledByMegaLights = bHandledByMegaLights;
+
+					if (bHandledByMegaLights)
+					{
+						SortedLightInfo->SortKey.Fields.LightSceneId = LightSceneInfo->Id & LIGHT_ID_MASK;
+						SortedLightInfo->SortKey.Fields.bHandledByMegaLights = true;
+					}
+
 					break;
 				}
 			}
diff --git a/Engine/Source/Runtime/Renderer/Private/LightSceneInfo.h b/Engine/Source/Runtime/Renderer/Private/LightSceneInfo.h
--- a/Engine/Source/Runtime/Renderer/Private/LightSceneInfo.h
+++ b/Engine/Source/Runtime/Renderer/Private/LightSceneInfo.h
@@ -83,7 +83,7 @@ struct FSortedLightSceneInfo
 		struct
 		{
 			// Note: the order of these members controls the light sort order!
-			// Currently bHandledByLumen is the MSB and LightType is LSB
+			// Currently bHandledByMegaLights is the MSB and LightType is LSB
 			/** The type of light. */
 			uint32 LightType : LightType_NumBits;
 			/** Whether the light uses a light function. */
@@ -99,11 +99,14 @@ struct FSortedLightSceneInfo
 			 * Super-set of lights supporting tiled, so the tiled lights will end up in the first part of this range.
 			 */
 			uint32 bClusteredDeferredNotSupported : 1;
+			/** Mega Lights must be sorted by GPU Scene Light Index - this field will only be set for MegaLights which don't rely on ordering of the other fields */
+			uint32 LightSceneId : LIGHT_ID_NUM_BITS;
 			/** Whether the light should be handled by Mega Lights, these will be sorted to the end so they can be skipped */
 			uint32 bHandledByMegaLights : 1;
 		} Fields;
 		/** Sort key bits packed into an integer. */
-		int32 Packed;
+		uint32 Packed;
+		static_assert(sizeof(Packed) >= sizeof(Fields), "SortKey Packed representation must include all bits in Fields");
 	} SortKey;
 
 	const FLightSceneInfo* LightSceneInfo;
