Hair Strands: Hairs were getting culled when hidden but still affecting indirect lighting

Hairs did not obey CastShadow flags in ray/path tracing

Refactor hair logic for calculating ray tracing masks. Hide some utility functions inside RayTracingInstanceMask.cpp

Adjust ray traced translucency and primary ray debug view ray mask to not include shadow flags

diff --git a/Engine/Plugins/Runtime/HairStrands/Source/HairStrandsCore/Private/GroomComponent.cpp b/Engine/Plugins/Runtime/HairStrands/Source/HairStrandsCore/Private/GroomComponent.cpp
--- a/Engine/Plugins/Runtime/HairStrands/Source/HairStrandsCore/Private/GroomComponent.cpp
+++ b/Engine/Plugins/Runtime/HairStrands/Source/HairStrandsCore/Private/GroomComponent.cpp
@@ -972,6 +972,9 @@ public:
 
 		const bool bUseCardsOrMeshes = GeometryType == EHairGeometryType::Cards || GeometryType == EHairGeometryType::Meshes;
 		Mesh.CastShadow = bUseCardsOrMeshes;
+#if RHI_RAYTRACING
+		Mesh.CastRayTracedShadow = (MeshBatchType == EHairMeshBatchType::Raytracing || bUseCardsOrMeshes) && bCastDynamicShadow;
+#endif
 		Mesh.bUseForMaterial = MeshBatchType == EHairMeshBatchType::Raytracing || bUseCardsOrMeshes;
 		Mesh.bUseForDepthPass = bUseCardsOrMeshes;
 		Mesh.SegmentIndex = 0;
@@ -1080,7 +1083,8 @@ public:
 		// When path tracing is enabled force DrawRelevance if not visible in main view ('hidden in game'), 
 		// but visible in shadow 'hidden shadow') so that raytracing geometry is created/updated correctly
 		const bool bPathtracing = View->Family->EngineShowFlags.PathTracing;
-		const bool bForceDrawRelevance = bPathtracing && (!IsShown(View) && IsShadowCast(View));
+		const bool bIsShown = IsShown(View);
+		const bool bForceDrawRelevance = bPathtracing && (!bIsShown && (IsShadowCast(View) || bAffectIndirectLightingWhileHidden));
 		const bool bVisible = View->Family->EngineShowFlags.Hair;
 
 		bool bUseCardsOrMesh = false;
@@ -1094,7 +1098,7 @@ public:
 		FPrimitiveViewRelevance Result;
 
 		// Special pass for hair strands geometry (not part of the base pass, and shadowing is handlded in a custom fashion). When cards rendering is enabled we reusethe base pass
-		Result.bDrawRelevance		= bVisible && (IsShown(View) || bForceDrawRelevance);
+		Result.bDrawRelevance		= bVisible && (bIsShown || bForceDrawRelevance);
 		Result.bRenderInMainPass	= bUseCardsOrMesh && ShouldRenderInMainPass();
 		Result.bShadowRelevance		= IsShadowCast(View);
 		Result.bDynamicRelevance	= bUseCardsOrMesh;
@@ -1111,7 +1115,7 @@ public:
 		MaterialRelevance.SetPrimitiveViewRelevance(Result);
 
 		// Override the MaterialRelevance output
-		Result.bHairStrands = IsShown(View) || bForceDrawRelevance;
+		Result.bHairStrands = bIsShown || bForceDrawRelevance;
 		return Result;
 	}
 
@@ -2989,21 +2993,37 @@ void UGroomComponent::InitResources(bool bIsBindingReloading)
 			}
 
 			#if RHI_RAYTRACING
-			if (IsHairRayTracingEnabled() && HairGroupInstance->Strands.Modifier.bUseHairRaytracingGeometry && bVisibleInRayTracing)
+			HairGroupInstance->Strands.ViewRayTracingMask = EHairViewRayTracingMask::None;
+			if (IsHairRayTracingEnabled() && bVisibleInRayTracing)
 			{
-				HairGroupInstance->Strands.ViewRayTracingMask = EHairViewRayTracingMask::RayTracing | EHairViewRayTracingMask::PathTracing;
-				if (bNeedDynamicResources)
+				// If the project has path tracing, we may need raytracing resources
+				static const auto CVarLocalVariablePathTracing = IConsoleManager::Get().FindTConsoleVariableDataInt(TEXT("r.PathTracing"));
+				const int PathTracingEnabled = CVarLocalVariablePathTracing ? CVarLocalVariablePathTracing->GetValueOnAnyThread() : 0;
+				if (PathTracingEnabled)
 				{
-					// Allocate dynamic raytracing resources (owned by the groom component/instance)
-					HairGroupInstance->Strands.RenRaytracingResource = new FHairStrandsRaytracingResource(GroupData.Strands.BulkData, ResourceName, OwnerName);
-					HairGroupInstance->Strands.RenRaytracingResourceOwned = true;
+					// Always visible for path tracing
+					HairGroupInstance->Strands.ViewRayTracingMask |= EHairViewRayTracingMask::PathTracing;
 				}
-				else
+				// Visible to ray tracing if the user has opted in
+				if (HairGroupInstance->Strands.Modifier.bUseHairRaytracingGeometry)
 				{
-					// (Lazy) Allocate static raytracing resources (owned by the grooom asset)
-					HairGroupInstance->Strands.RenRaytracingResourceOwned = false;
-					HairGroupInstance->Strands.RenRaytracingResource = GroomAsset->AllocateStrandsRaytracingResources(GroupIt);
-					check(HairGroupInstance->Strands.RenRaytracingResource);
+					HairGroupInstance->Strands.ViewRayTracingMask |= EHairViewRayTracingMask::RayTracing;
+				}
+				if (HairGroupInstance->Strands.ViewRayTracingMask != EHairViewRayTracingMask::None)
+				{
+					if (bNeedDynamicResources)
+					{
+						// Allocate dynamic raytracing resources (owned by the groom component/instance)
+						HairGroupInstance->Strands.RenRaytracingResource = new FHairStrandsRaytracingResource(GroupData.Strands.BulkData, ResourceName, OwnerName);
+						HairGroupInstance->Strands.RenRaytracingResourceOwned = true;
+					}
+					else
+					{
+						// (Lazy) Allocate static raytracing resources (owned by the grooom asset)
+						HairGroupInstance->Strands.RenRaytracingResourceOwned = false;
+						HairGroupInstance->Strands.RenRaytracingResource = GroomAsset->AllocateStrandsRaytracingResources(GroupIt);
+						check(HairGroupInstance->Strands.RenRaytracingResource);
+					}
 				}
 			}
 			#endif
@@ -3900,9 +3920,8 @@ void UGroomComponent::PostEditChangeProperty(FPropertyChangedEvent& PropertyChan
 		for (const TRefCountPtr<FHairGroupInstance>& Instance : HairGroupInstances)
 		{
 			const FHairGroupDesc GroupDesc = GetGroomGroupsDesc(GroomAsset, this, Instance->Debug.GroupIndex);
-			const bool bRecreate =
-				(Instance->Strands.RenRaytracingResource == nullptr && GroupDesc.bUseHairRaytracingGeometry) ||
-				(Instance->Strands.RenRaytracingResource != nullptr && !GroupDesc.bUseHairRaytracingGeometry);
+			const bool bWasUsingRayTracingGeometry = EnumHasAnyFlags(Instance->Strands.ViewRayTracingMask, EHairViewRayTracingMask::RayTracing);
+			const bool bRecreate = bWasUsingRayTracingGeometry != GroupDesc.bUseHairRaytracingGeometry;
 			if (bRecreate)
 			{
 				bRayTracingGeometryChanged = true;
diff --git a/Engine/Plugins/Runtime/HairStrands/Source/HairStrandsCore/Private/GroomManager.cpp b/Engine/Plugins/Runtime/HairStrands/Source/HairStrandsCore/Private/GroomManager.cpp
--- a/Engine/Plugins/Runtime/HairStrands/Source/HairStrandsCore/Private/GroomManager.cpp
+++ b/Engine/Plugins/Runtime/HairStrands/Source/HairStrandsCore/Private/GroomManager.cpp
@@ -1194,7 +1194,7 @@ static void RunHairStrandsInterpolation_Strands(
 		InstanceData.bNeedRaytracing			= false;
 		InstanceData.bNeedRaytracingUpdate		= false;
 		InstanceData.bNeedRaytracingBuild		= false;
-		InstanceData.bSupportVoxelization		= Instance->Strands.Modifier.bSupportVoxelization && bCastShadow;
+		InstanceData.bSupportVoxelization		= Instance->Strands.Modifier.bSupportVoxelization && (bCastShadow || Instance->Debug.Proxy->AffectsIndirectLightingWhileHidden()); // Lumen also uses hair voxelization
 
 		// Register resources
 		FInstanceRDGResources RDGResources;
@@ -2282,28 +2282,6 @@ static void AddCopyHairStrandsPositionPass(
 	AddCopyBufferPass(GraphBuilder, Deformed1Buffer, RestBuffer);
 }
 
-#if RHI_RAYTRACING
-static void AllocateRaytracingResources(FHairGroupInstance* Instance, int32 InMeshLODIndex)
-{
-	if (IsHairRayTracingEnabled() && !Instance->Strands.RenRaytracingResource)
-	{
-		check(Instance->Strands.IsValid());
-
-#if RHI_ENABLE_RESOURCE_INFO
-		FName OwnerName(FString::Printf(TEXT("%s [LOD%d]"), *Instance->Debug.MeshComponentName, InMeshLODIndex));
-#else
-		FName OwnerName = NAME_None;
-#endif
-
-		// Allocate dynamic raytracing resources (owned by the groom component/instance)
-		FHairResourceName ResourceName(FName(Instance->Debug.GroomAssetName), Instance->Debug.GroupIndex);
-		Instance->Strands.RenRaytracingResource		 = new FHairStrandsRaytracingResource(Instance->Strands.GetData(), ResourceName, OwnerName);
-		Instance->Strands.RenRaytracingResourceOwned = true;
-	}
-	Instance->Strands.ViewRayTracingMask |= EHairViewRayTracingMask::PathTracing;
-}
-#endif
-
 void AddHairStreamingRequest(FHairGroupInstance* Instance, int32 InLODIndex)
 {
 	if (Instance && InLODIndex >= 0)
@@ -2611,8 +2589,7 @@ static bool SelectValidLOD(
 	float PrevHairLODIndex, 
 	float HairLODCount, 
 	float MeshLODIndex, 
-	float PrevMeshLODIndex, 
-	bool bHasPathTracingView)
+	float PrevMeshLODIndex)
 {
 	const int32 IntHairLODIndex = FMath::Clamp(FMath::FloorToInt(HairLODIndex), 0, HairLODCount - 1);
 	const TArray<bool>& LODVisibilities = Instance->HairGroupPublicData->GetLODVisibilities();
@@ -2740,7 +2717,6 @@ static bool SelectValidLOD(
 		if (Instance->Strands.DeformedRootResource)		{ Instance->Strands.DeformedRootResource->Allocate	(GraphBuilder, EHairResourceLoadingType::Async, ResourceStatus, RequestedCurveCount, RequestedPointCount, false/*bAllowDeallocation*/, MeshLODIndex); }
 		if (Instance->Strands.DeformedResource)			{ Instance->Strands.DeformedResource->Allocate		(GraphBuilder, EHairResourceLoadingType::Async, ResourceStatus, RequestedCurveCount, RequestedPointCount); }
 #if RHI_RAYTRACING
-		if (bHasPathTracingView)						{ AllocateRaytracingResources(Instance, MeshLODIndex); }
 		if (Instance->Strands.RenRaytracingResource)	{ Instance->Strands.RenRaytracingResource->Allocate	(GraphBuilder, EHairResourceLoadingType::Async, ResourceStatus, RequestedCurveCount, RequestedPointCount, Instance->Strands.RenRaytracingResourceOwned/*bAllowDeallocation*/); }
 #endif
 		Instance->Strands.VertexFactory->InitResources(GraphBuilder.RHICmdList);
@@ -2824,19 +2800,6 @@ static void RunHairLODSelection(
 		ShaderPlatform = Views[0]->GetShaderPlatform();
 	}
 
-	// Detect if one of the current view has path tracing enabled for allocating raytracing resources
-	bool bHasPathTracingView = false;
-#if RHI_RAYTRACING
-	for (const FSceneView* View : Views)
-	{
-		if (View->Family->EngineShowFlags.PathTracing)
-		{
-			bHasPathTracingView = true;
-			break;
-		}
-	}
-#endif
-
 	struct FInstanceData
 	{
 		FHairGroupInstance* Instance = nullptr;
@@ -3035,8 +2998,7 @@ static void RunHairLODSelection(
 			InstanceData.PrevHairLODIndex,
 			InstanceData.HairLOD.HairLODCount,
 			InstanceData.MeshLODIndex,
-			InstanceData.PrevMeshLODIndex,
-			bHasPathTracingView);
+			InstanceData.PrevMeshLODIndex);
 
 		// 2. If no valid LOD are found, try to find a coarser LOD which is:
 		// * loaded 
@@ -3060,8 +3022,7 @@ static void RunHairLODSelection(
 						InstanceData.PrevHairLODIndex,
 						InstanceData.HairLOD.HairLODCount,
 						InstanceData.MeshLODIndex,
-						InstanceData.PrevMeshLODIndex,
-						bHasPathTracingView);
+						InstanceData.PrevMeshLODIndex);
 					break;
 				}
 			}
@@ -3074,17 +3035,25 @@ static void RunHairLODSelection(
 			InstanceData.Instance->BindingType = EHairBindingType::NoneBinding;
 		}
 
-		// Cull instance which are not visibles
-		for (const FSceneView* View : Views)
+		// Cull instance if it is not visible in any View
 		{
-			if (!InstanceData.Instance->Debug.Proxy->IsShown(View) && !InstanceData.Instance->Debug.Proxy->IsShadowCast(View))
+			const FPrimitiveSceneProxy* Proxy = InstanceData.Instance->Debug.Proxy;
+			check(Proxy != nullptr);
+			bool bVisibleInView = false;
+			for (const FSceneView* View : Views)
+			{
+				if (Proxy->IsShown(View) || Proxy->IsShadowCast(View))
+				{
+					bVisibleInView = true;
+					break;
+				}
+			}
+			if (!bVisibleInView && !Proxy->AffectsIndirectLightingWhileHidden())
 			{
 				InstanceData.Instance->GeometryType = EHairGeometryType::NoneGeometry;
 				InstanceData.Instance->BindingType = EHairBindingType::NoneBinding;
-				break;
 			}
 		}
-
 		// Update the local-to-world transform based on the binding type 
 		if (GetHairSwapBufferType() == EHairBufferSwapType::BeginOfFrame || GetHairSwapBufferType() == EHairBufferSwapType::EndOfFrame)
 		{
diff --git a/Engine/Shaders/Private/RayTracing/RayTracingPrimaryRays.usf b/Engine/Shaders/Private/RayTracing/RayTracingPrimaryRays.usf
--- a/Engine/Shaders/Private/RayTracing/RayTracingPrimaryRays.usf
+++ b/Engine/Shaders/Private/RayTracing/RayTracingPrimaryRays.usf
@@ -135,7 +135,7 @@ RAY_TRACING_ENTRY_RAYGEN(RayTracingPrimaryRaysRGS)
 	{
 		const bool bIsCameraRay = !bHasScattered;
 		const uint RefractionRayFlags = bIsCameraRay ? RAY_FLAG_CULL_BACK_FACING_TRIANGLES : 0;
-		const uint RefractionInstanceInclusionMask = (!bPrimaryView && bIsCameraRay) || !TranslucencyRefraction ? RAY_TRACING_MASK_TRANSLUCENT : RAY_TRACING_MASK_ALL;
+		const uint RefractionInstanceInclusionMask = (!bPrimaryView && bIsCameraRay) || !TranslucencyRefraction ? RAY_TRACING_MASK_TRANSLUCENT : (RAY_TRACING_MASK_OPAQUE | RAY_TRACING_MASK_TRANSLUCENT | RAY_TRACING_MASK_HAIR_STRANDS);
 		const bool bRefractionRayTraceSkyLightContribution = false;
 		const bool bRefractionDecoupleSampleGeneration = true;
 #if PROJECT_SUPPORTS_LUMEN
@@ -298,7 +298,7 @@ RAY_TRACING_ENTRY_RAYGEN(RayTracingPrimaryRaysRGS)
 			ApplyPositionBias(ReflectionRay, Payload.WorldNormal, MaxNormalBias);
 
 			const uint ReflectionRayFlags = RAY_FLAG_CULL_BACK_FACING_TRIANGLES;
-			const uint ReflectionInstanceInclusionMask = RAY_TRACING_MASK_ALL;
+			const uint ReflectionInstanceInclusionMask = (RAY_TRACING_MASK_OPAQUE | RAY_TRACING_MASK_TRANSLUCENT | RAY_TRACING_MASK_HAIR_STRANDS);
 			const bool bReflectionRayTraceSkyLightContribution = false;
 			const bool bReflectionDecoupleSampleGeneration = true;
 			const bool bReflectionEnableSkyLightContribution = bSkyLightAffectReflection;
diff --git a/Engine/Shaders/Shared/PathTracingDefinitions.h b/Engine/Shaders/Shared/PathTracingDefinitions.h
--- a/Engine/Shaders/Shared/PathTracingDefinitions.h
+++ b/Engine/Shaders/Shared/PathTracingDefinitions.h
@@ -32,7 +32,6 @@
 #define PATHTRACER_MASK_INDIRECT_TRANSLUCENT				0x80			
 
 #define PATHTRACER_MASK_IGNORE								0x00			// used when mapping general tracing mask to path tracing mask
-#define PATHTRACER_MASK_UNUSED	(PATHTRACER_MASK_EMPTY_SLOT1|PATHTRACER_MASK_EMPTY_SLOT2)	
 #define PATHTRACER_MASK_ALL									0xFF
 
 // Constants for light contribution types (AOV decomposition of the image)
diff --git a/Engine/Source/Runtime/Renderer/Private/RayTracing/RayTracingInstanceMask.cpp b/Engine/Source/Runtime/Renderer/Private/RayTracing/RayTracingInstanceMask.cpp
--- a/Engine/Source/Runtime/Renderer/Private/RayTracing/RayTracingInstanceMask.cpp
+++ b/Engine/Source/Runtime/Renderer/Private/RayTracing/RayTracingInstanceMask.cpp
@@ -15,85 +15,65 @@
 
 uint8 ComputeRayTracingInstanceMask(ERayTracingInstanceMaskType MaskType, ERayTracingViewMaskMode MaskMode)
 {
-	uint8 Mask = 0;
 	if (MaskMode == ERayTracingViewMaskMode::RayTracing)
 	{
 		switch (MaskType)
 		{
 		case ERayTracingInstanceMaskType::Opaque:
-			Mask = RAY_TRACING_MASK_OPAQUE;
-			break;
+			return RAY_TRACING_MASK_OPAQUE;
 		case ERayTracingInstanceMaskType::Translucent:
-			Mask = RAY_TRACING_MASK_TRANSLUCENT;
-			break;
+			return RAY_TRACING_MASK_TRANSLUCENT;
 		case ERayTracingInstanceMaskType::OpaqueShadow:
-			Mask = RAY_TRACING_MASK_OPAQUE_SHADOW;
-			break;
+			return RAY_TRACING_MASK_OPAQUE_SHADOW;
 		case ERayTracingInstanceMaskType::TranslucentShadow:
-			Mask = RAY_TRACING_MASK_TRANSLUCENT_SHADOW;
-			break;
+			return RAY_TRACING_MASK_TRANSLUCENT_SHADOW;
 		case ERayTracingInstanceMaskType::ThinShadow:
-			Mask = RAY_TRACING_MASK_THIN_SHADOW;
-			break;
+			return RAY_TRACING_MASK_THIN_SHADOW;
 		case ERayTracingInstanceMaskType::FarField:
-			Mask = RAY_TRACING_MASK_FAR_FIELD;
-			break;
+			return RAY_TRACING_MASK_FAR_FIELD;
 		case ERayTracingInstanceMaskType::HairStrands:
-			Mask = RAY_TRACING_MASK_HAIR_STRANDS;
-			break;
+			return RAY_TRACING_MASK_HAIR_STRANDS | RAY_TRACING_MASK_THIN_SHADOW;
 		case ERayTracingInstanceMaskType::FirstPersonWorldSpaceRepresentation:
-			Mask = RAY_TRACING_MASK_OPAQUE_FP_WORLD_SPACE;
-			break;
+			return RAY_TRACING_MASK_OPAQUE_FP_WORLD_SPACE;
 		case ERayTracingInstanceMaskType::VisibleInPrimaryRay:
-			Mask = 0;
-			break;
+			return 0; // There is no distinct notion of primary ray visibility for ray tracing
 		case ERayTracingInstanceMaskType::VisibleInIndirectRay:
-			Mask = 0;
-			break;
+			return RAY_TRACING_MASK_OPAQUE | RAY_TRACING_MASK_TRANSLUCENT | RAY_TRACING_MASK_HAIR_STRANDS;
 		default:
-			checkNoEntry();
 			break;
 		}
 	}
 	else if (MaskMode == ERayTracingViewMaskMode::PathTracing ||
-		MaskMode == ERayTracingViewMaskMode::LightMapTracing)
+		     MaskMode == ERayTracingViewMaskMode::LightMapTracing)
 	{
 		switch (MaskType)
 		{
 		case ERayTracingInstanceMaskType::Opaque:
-			Mask = PATHTRACER_MASK_CAMERA | PATHTRACER_MASK_INDIRECT;
-			break;
+			return PATHTRACER_MASK_CAMERA | PATHTRACER_MASK_INDIRECT;
 		case ERayTracingInstanceMaskType::Translucent:
-			Mask = PATHTRACER_MASK_CAMERA_TRANSLUCENT | PATHTRACER_MASK_INDIRECT_TRANSLUCENT;
-			break;
+			return PATHTRACER_MASK_CAMERA_TRANSLUCENT | PATHTRACER_MASK_INDIRECT_TRANSLUCENT;
 		case ERayTracingInstanceMaskType::OpaqueShadow:
-			Mask = PATHTRACER_MASK_SHADOW;
-			break;
+			return PATHTRACER_MASK_SHADOW;
 		case ERayTracingInstanceMaskType::TranslucentShadow:
-			Mask = PATHTRACER_MASK_SHADOW;
-			break;
+			return PATHTRACER_MASK_SHADOW;
 		case ERayTracingInstanceMaskType::ThinShadow:
-			Mask = PATHTRACER_MASK_HAIR_SHADOW;
-			break;
+			return PATHTRACER_MASK_HAIR_SHADOW;
 		case ERayTracingInstanceMaskType::FarField:
-			Mask = PATHTRACER_MASK_IGNORE;
-			break;
+			return PATHTRACER_MASK_IGNORE;
+		case ERayTracingInstanceMaskType::FirstPersonWorldSpaceRepresentation:
+			return PATHTRACER_MASK_IGNORE;
 		case ERayTracingInstanceMaskType::HairStrands:
-			Mask = PATHTRACER_MASK_HAIR_CAMERA | PATHTRACER_MASK_HAIR_INDIRECT;
-			break;
+			return PATHTRACER_MASK_HAIR_CAMERA | PATHTRACER_MASK_HAIR_SHADOW | PATHTRACER_MASK_HAIR_INDIRECT;
 		case ERayTracingInstanceMaskType::VisibleInPrimaryRay:
-			Mask = PATHTRACER_MASK_CAMERA | PATHTRACER_MASK_HAIR_CAMERA | PATHTRACER_MASK_CAMERA_TRANSLUCENT;
-			break;
+			return PATHTRACER_MASK_CAMERA | PATHTRACER_MASK_HAIR_CAMERA | PATHTRACER_MASK_CAMERA_TRANSLUCENT;
 		case ERayTracingInstanceMaskType::VisibleInIndirectRay:
-			Mask = PATHTRACER_MASK_INDIRECT | PATHTRACER_MASK_HAIR_INDIRECT | PATHTRACER_MASK_INDIRECT_TRANSLUCENT;
-			break;
+			return PATHTRACER_MASK_INDIRECT | PATHTRACER_MASK_HAIR_INDIRECT | PATHTRACER_MASK_INDIRECT_TRANSLUCENT;
 		default:
-			checkNoEntry();
 			break;
 		}
 	}
-
-	return Mask;
+	checkNoEntry();
+	return 0;
 }
 
 uint8 BlendModeToRayTracingInstanceMask(const EBlendMode BlendMode, bool bIsDitherMasked, bool bCastShadow, ERayTracingViewMaskMode MaskMode)
@@ -114,19 +94,40 @@ uint8 BlendModeToRayTracingInstanceMask(const EBlendMode BlendMode, bool bIsDith
 	return InstanceMask;
 }
 
-FSceneProxyRayTracingMaskInfo GetSceneProxyRayTracingMaskInfo(const FPrimitiveSceneProxy& PrimitiveSceneProxy)
+
+/** Util struct and function to derive mask related info from scene proxy*/
+struct FSceneProxyRayTracingMaskInfo
+{
+	bool bVisibleToCamera = false;
+	bool bVisibleToShadow = false;
+	bool bVisibleToIndirect = false;
+	bool bIsFirstPersonWorldSpaceRepresentation = false;
+	ERayTracingViewMaskMode MaskMode = ERayTracingViewMaskMode::RayTracing;
+};
+
+static FSceneProxyRayTracingMaskInfo GetSceneProxyRayTracingMaskInfo(const FPrimitiveSceneProxy& PrimitiveSceneProxy)
 {
 	FSceneProxyRayTracingMaskInfo MaskInfo = {};
 
-	MaskInfo.bAffectsIndirectLightingOnly = PrimitiveSceneProxy.AffectsIndirectLightingWhileHidden() && !PrimitiveSceneProxy.IsDrawnInGame();
-	MaskInfo.bCastHiddenShadow = PrimitiveSceneProxy.CastsHiddenShadow() && !PrimitiveSceneProxy.IsDrawnInGame();
-	MaskInfo.bAffectsDynamicIndirectLighting = PrimitiveSceneProxy.AffectsDynamicIndirectLighting();
-	MaskInfo.bIsFirstPersonWorldSpaceRepresentation = PrimitiveSceneProxy.IsFirstPersonWorldSpaceRepresentation();
-
 	const FScene* RenderScene = PrimitiveSceneProxy.GetScene().GetRenderScene();
-
 	MaskInfo.MaskMode = static_cast<ERayTracingViewMaskMode>(RenderScene->CachedRayTracingMeshCommandsMode);
 
+	if (PrimitiveSceneProxy.IsDrawnInGame())
+	{
+		MaskInfo.bVisibleToCamera = true;
+		MaskInfo.bVisibleToShadow = true;
+		// NOTE: For backwards compatibility, only path tracing obeys the AffectsDynamicIndirectLighting flag
+		MaskInfo.bVisibleToIndirect = MaskInfo.MaskMode == ERayTracingViewMaskMode::RayTracing ? true : PrimitiveSceneProxy.AffectsDynamicIndirectLighting();
+	}
+	else
+	{
+		MaskInfo.bVisibleToCamera = false;
+		MaskInfo.bVisibleToShadow = PrimitiveSceneProxy.CastsHiddenShadow();
+		MaskInfo.bVisibleToIndirect = PrimitiveSceneProxy.AffectsIndirectLightingWhileHidden();
+	}
+
+	MaskInfo.bIsFirstPersonWorldSpaceRepresentation = PrimitiveSceneProxy.IsFirstPersonWorldSpaceRepresentation();
+
 	return MaskInfo;
 }
 
@@ -143,13 +144,15 @@ static uint8 ApplyFirstPersonRayTracingInstanceMaskFlag(uint8 Mask, ERayTracingV
 	return Mask;
 }
 
-
-FRayTracingMaskAndFlags BuildRayTracingInstanceMaskAndFlags(TArrayView<const FMeshBatch> MeshBatches, ERHIFeatureLevel::Type FeatureLevel,
-	const FSceneProxyRayTracingMaskInfo& SceneProxyRayTracingMaskInfo, ERayTracingInstanceLayer InstanceLayer, uint8 ExtraMask)
+FRayTracingMaskAndFlags BuildRayTracingInstanceMaskAndFlags(const FRayTracingInstance& Instance, const FPrimitiveSceneProxy& PrimitiveSceneProxy)
 {
+	FSceneProxyRayTracingMaskInfo MaskInfo = GetSceneProxyRayTracingMaskInfo(PrimitiveSceneProxy);
+	const ERHIFeatureLevel::Type FeatureLevel = PrimitiveSceneProxy.GetScene().GetFeatureLevel();
+	const ERayTracingViewMaskMode MaskMode = MaskInfo.MaskMode;
+
 	FRayTracingMaskAndFlags Result;
 
-	ensureMsgf(MeshBatches.Num() > 0, TEXT("You need to add MeshBatches first for instance mask and flags to build upon."));
+	ensureMsgf(Instance.GetMaterials().Num() > 0, TEXT("You need to add MeshBatches first for instance mask and flags to build upon."));
 
 	bool bAllSegmentsOpaque = true;
 	bool bAllSegmentsTranslucent = true;
@@ -159,19 +162,16 @@ FRayTracingMaskAndFlags BuildRayTracingInstanceMaskAndFlags(TArrayView<const FMe
 	bool bAllSegmentsDecal = true;
 	bool bDoubleSided = false;
 	bool bAllSegmentsReverseCulling = true;
-	ERayTracingViewMaskMode MaskMode = SceneProxyRayTracingMaskInfo.MaskMode;
-	Result.Mask = ExtraMask;
 
-	for (int32 SegmentIndex = 0; SegmentIndex < MeshBatches.Num(); SegmentIndex++)
+	Result.Mask = 0;
+	for (const FMeshBatch& MeshBatch : Instance.GetMaterials())
 	{
-		const FMeshBatch& MeshBatch = MeshBatches[SegmentIndex];
-
 		// Mesh Batches can "null" when they have zero triangles.  Check the MaterialRenderProxy before accessing.
 		if (MeshBatch.bUseForMaterial && MeshBatch.MaterialRenderProxy)
 		{
 			const FMaterial& Material = MeshBatch.MaterialRenderProxy->GetIncompleteMaterialWithFallback(FeatureLevel);
 			const EBlendMode BlendMode = Material.GetBlendMode();
-			const bool bSegmentCastsShadow = MeshBatch.CastRayTracedShadow && Material.CastsRayTracedShadows() && BlendMode != BLEND_Additive;
+			const bool bSegmentCastsShadow = MaskInfo.bVisibleToShadow && MeshBatch.CastRayTracedShadow && Material.CastsRayTracedShadows() && BlendMode != BLEND_Additive;
 
 			Result.Mask |= BlendModeToRayTracingInstanceMask(BlendMode, Material.IsDitherMasked(), bSegmentCastsShadow, MaskMode);
 			bAllSegmentsOpaque &= BlendMode == BLEND_Opaque;
@@ -193,53 +193,39 @@ FRayTracingMaskAndFlags BuildRayTracingInstanceMaskAndFlags(TArrayView<const FMe
 	Result.bAllSegmentsTranslucent = bAllSegmentsTranslucent;
 	Result.bReverseCulling = bAllSegmentsReverseCulling;
 
-	const bool bIsHairStrands = Result.Mask & ComputeRayTracingInstanceMask(ERayTracingInstanceMaskType::HairStrands, MaskMode);
+	const bool bIsHairStrands = Instance.bThinGeometry;
 	if (bIsHairStrands)
 	{
-		// For hair strands, opaque/translucent mask should be cleared to make sure geometry is only in the hair group. 
-		// If any segment receives shadow, it should receive only thin shadow instead of shadow.
-
-		Result.Mask &= ~(
-			ComputeRayTracingInstanceMask(ERayTracingInstanceMaskType::OpaqueShadow, MaskMode) |
-			ComputeRayTracingInstanceMask(ERayTracingInstanceMaskType::TranslucentShadow, MaskMode) |
-			ComputeRayTracingInstanceMask(ERayTracingInstanceMaskType::ThinShadow, MaskMode) |
-			ComputeRayTracingInstanceMask(ERayTracingInstanceMaskType::Translucent, MaskMode) |
-			ComputeRayTracingInstanceMask(ERayTracingInstanceMaskType::Opaque, MaskMode));
-
-		Result.Mask |= ComputeRayTracingInstanceMask(ERayTracingInstanceMaskType::HairStrands, MaskMode);
+		// Reset all hair strands bits "on"
+		Result.Mask = ComputeRayTracingInstanceMask(ERayTracingInstanceMaskType::HairStrands, MaskMode);
+	}
 
-		if (bAnySegmentsCastShadow)
-		{
-			Result.Mask |= ComputeRayTracingInstanceMask(ERayTracingInstanceMaskType::ThinShadow, MaskMode);
-		}
+	if (!MaskInfo.bVisibleToCamera)
+	{
+		// If the object is not visible to camera, remove all direct visibility bits.
+		Result.Mask &= ~ComputeRayTracingInstanceMask(ERayTracingInstanceMaskType::VisibleInPrimaryRay, MaskMode);
 	}
 
-	if (!SceneProxyRayTracingMaskInfo.bAffectsDynamicIndirectLighting)
+	if (!MaskInfo.bVisibleToIndirect)
 	{
 		// If the object does not affect indirect lighting, remove all indirect bits.
 		Result.Mask &= ~ComputeRayTracingInstanceMask(ERayTracingInstanceMaskType::VisibleInIndirectRay, MaskMode);
 	}
 
-	if (SceneProxyRayTracingMaskInfo.bAffectsIndirectLightingOnly)
+	if (!bAnySegmentsCastShadow)
 	{
-		Result.Mask &= ~ComputeRayTracingInstanceMask(ERayTracingInstanceMaskType::VisibleInPrimaryRay, MaskMode);
+		// Not casting shadows, remove any set shadow flags
+		Result.Mask &= ~(
+			ComputeRayTracingInstanceMask(ERayTracingInstanceMaskType::OpaqueShadow, MaskMode) |
+			ComputeRayTracingInstanceMask(ERayTracingInstanceMaskType::TranslucentShadow, MaskMode) |
+			ComputeRayTracingInstanceMask(ERayTracingInstanceMaskType::ThinShadow, MaskMode)
+		);
 	}
 
-	if (SceneProxyRayTracingMaskInfo.bCastHiddenShadow && bAnySegmentsCastShadow)
-	{
-		if (!SceneProxyRayTracingMaskInfo.bAffectsIndirectLightingOnly)
-		{
-			// objects should not be in any visible group if any segments cast shadow and the caster is hidden
-			// and not affecting indirect.
-			Result.Mask &= ~(ComputeRayTracingInstanceMask(ERayTracingInstanceMaskType::Translucent, MaskMode) |
-				ComputeRayTracingInstanceMask(ERayTracingInstanceMaskType::Opaque, MaskMode) |
-				ComputeRayTracingInstanceMask(ERayTracingInstanceMaskType::HairStrands, MaskMode));
-		}
-	}
 
-	Result.Mask = ApplyFirstPersonRayTracingInstanceMaskFlag(Result.Mask, MaskMode, SceneProxyRayTracingMaskInfo.bIsFirstPersonWorldSpaceRepresentation);
+	Result.Mask = ApplyFirstPersonRayTracingInstanceMaskFlag(Result.Mask, MaskMode, MaskInfo.bIsFirstPersonWorldSpaceRepresentation);
 
-	if (InstanceLayer == ERayTracingInstanceLayer::FarField)
+	if (Instance.InstanceLayer == ERayTracingInstanceLayer::FarField)
 	{
 		// if far field, set that flag exclusively
 		Result.Mask = ComputeRayTracingInstanceMask(ERayTracingInstanceMaskType::FarField, MaskMode);
@@ -248,18 +234,6 @@ FRayTracingMaskAndFlags BuildRayTracingInstanceMaskAndFlags(TArrayView<const FMe
 	return Result;
 }
 
-FRayTracingMaskAndFlags BuildRayTracingInstanceMaskAndFlags(const FRayTracingInstance& Instance, const FPrimitiveSceneProxy& PrimitiveSceneProxy)
-{
-	FSceneProxyRayTracingMaskInfo MaskInfo = GetSceneProxyRayTracingMaskInfo(PrimitiveSceneProxy);
-
-	const TArrayView<const FMeshBatch> MeshBatches = Instance.GetMaterials();
-
-	// add extra mask bit for hair.
-	const uint8 ExtraMask = Instance.bThinGeometry ? ComputeRayTracingInstanceMask(ERayTracingInstanceMaskType::HairStrands, MaskInfo.MaskMode) : 0;
-
-	return BuildRayTracingInstanceMaskAndFlags(MeshBatches, PrimitiveSceneProxy.GetScene().GetFeatureLevel(), MaskInfo, Instance.InstanceLayer, ExtraMask);
-}
-
 void SetupRayTracingMeshCommandMaskAndStatus(FRayTracingMeshCommand& MeshCommand, const FMeshBatch& MeshBatch, const FPrimitiveSceneProxy* PrimitiveSceneProxy,
 	const FMaterial& MaterialResource, ERayTracingViewMaskMode MaskMode)
 {
@@ -291,31 +265,31 @@ void SetupRayTracingMeshCommandMaskAndStatus(FRayTracingMeshCommand& MeshCommand
 
 	MeshCommand.InstanceMask = ApplyFirstPersonRayTracingInstanceMaskFlag(MeshCommand.InstanceMask, MaskMode, MaskInfo.bIsFirstPersonWorldSpaceRepresentation);
 
-	// TODO: This should be done once all mesh commands for a mesh are combined (similar to BuildRayTracingInstanceMaskAndFlags(...) above)
-	if (MaskMode == ERayTracingViewMaskMode::PathTracing || MaskMode == ERayTracingViewMaskMode::LightMapTracing)
+	// TODO: This should be done once all mesh commands for a mesh are combined (similar to BuildRayTracingInstanceMaskAndFlags above)
+	// TODO: Why is this logic not applied in the ray tracing case?
+	if (MaskMode == ERayTracingViewMaskMode::PathTracing ||
+		MaskMode == ERayTracingViewMaskMode::LightMapTracing)
 	{
-		if (!MaskInfo.bAffectsDynamicIndirectLighting)
+	
+		if (!MaskInfo.bVisibleToCamera)
 		{
-			// If the object does not affect indirect lighting, remove all indirect bits.
-			MeshCommand.InstanceMask &= ~ComputeRayTracingInstanceMask(ERayTracingInstanceMaskType::VisibleInIndirectRay, MaskMode);
+			// If the object is not visible to camera, remove all direct visibility bits.
+			MeshCommand.InstanceMask &= ~ComputeRayTracingInstanceMask(ERayTracingInstanceMaskType::VisibleInPrimaryRay, MaskMode);
 		}
 
-		if (MaskInfo.bAffectsIndirectLightingOnly)
+		if (!MaskInfo.bVisibleToIndirect)
 		{
-			MeshCommand.InstanceMask &= ~ComputeRayTracingInstanceMask(ERayTracingInstanceMaskType::VisibleInPrimaryRay, MaskMode);
+			// If the object does not affect indirect lighting, remove all indirect bits.
+			MeshCommand.InstanceMask &= ~ComputeRayTracingInstanceMask(ERayTracingInstanceMaskType::VisibleInIndirectRay, MaskMode);
 		}
 
-		if (MaskInfo.bCastHiddenShadow && MeshCommand.bCastRayTracedShadows)
+		if (!MaskInfo.bVisibleToShadow || !MeshCommand.bCastRayTracedShadows)
 		{
-			if (!MaskInfo.bAffectsIndirectLightingOnly)
-			{
-				// objects should not be in any visible group if any segments cast shadow and the caster is hidden
-				// and not affecting indirect.
-				MeshCommand.InstanceMask &= ~(
-					ComputeRayTracingInstanceMask(ERayTracingInstanceMaskType::Translucent, MaskMode) |
-					ComputeRayTracingInstanceMask(ERayTracingInstanceMaskType::Opaque, MaskMode) |
-					ComputeRayTracingInstanceMask(ERayTracingInstanceMaskType::HairStrands, MaskMode));
-			}
+			// Not casting shadows, remove any set shadow flags
+			MeshCommand.InstanceMask &= ~(
+				ComputeRayTracingInstanceMask(ERayTracingInstanceMaskType::OpaqueShadow, MaskMode) |
+				ComputeRayTracingInstanceMask(ERayTracingInstanceMaskType::TranslucentShadow, MaskMode) |
+				ComputeRayTracingInstanceMask(ERayTracingInstanceMaskType::ThinShadow, MaskMode));
 		}
 
 	}
diff --git a/Engine/Source/Runtime/Renderer/Private/RayTracing/RayTracingInstanceMask.h b/Engine/Source/Runtime/Renderer/Private/RayTracing/RayTracingInstanceMask.h
--- a/Engine/Source/Runtime/Renderer/Private/RayTracing/RayTracingInstanceMask.h
+++ b/Engine/Source/Runtime/Renderer/Private/RayTracing/RayTracingInstanceMask.h
@@ -11,15 +11,21 @@ class FSceneViewFamily;
 
 enum class ERayTracingInstanceMaskType : uint8
 {
-	// General mask type
+	// General mask type for primary and secondary rays
 	Opaque,
 	Translucent,
+
+	// Mask types for shadow rays
 	OpaqueShadow,
 	TranslucentShadow,
 	ThinShadow,
-	FarField,
+
+	// Geometry specific ray types
 	HairStrands,
+
+	// Special purpose ray types
 	FirstPersonWorldSpaceRepresentation,
+	FarField,
 
 	// path tracing specific mask type
 	VisibleInPrimaryRay,
@@ -40,26 +46,10 @@ uint8 ComputeRayTracingInstanceMask(ERayTracingInstanceMaskType MaskType, ERayTr
 RENDERER_API uint8 BlendModeToRayTracingInstanceMask(const EBlendMode BlendMode, bool bIsDitherMasked, bool bCastShadow, ERayTracingViewMaskMode MaskMode);
 
 
-/** Util struct and function to derive mask related info from scene proxy*/
-struct FSceneProxyRayTracingMaskInfo
-{
-	bool bAffectsIndirectLightingOnly = false;
-	bool bCastHiddenShadow = false;
-	bool bAffectsDynamicIndirectLighting = true;
-	bool bIsFirstPersonWorldSpaceRepresentation = false;
-	ERayTracingViewMaskMode MaskMode = ERayTracingViewMaskMode::RayTracing;
-};
-
-FSceneProxyRayTracingMaskInfo GetSceneProxyRayTracingMaskInfo(const FPrimitiveSceneProxy& PrimitiveSceneProxy);
-
-
 //-------------------------------------------------------
 //	Build Instance mask and flags (if needed)
 //-------------------------------------------------------
 
-FRayTracingMaskAndFlags BuildRayTracingInstanceMaskAndFlags(TArrayView<const FMeshBatch> MeshBatches, ERHIFeatureLevel::Type FeatureLevel,
-	const FSceneProxyRayTracingMaskInfo& SceneProxyRayTracingMaskInfo, ERayTracingInstanceLayer InstanceLayer, uint8 ExtraMask = 0);
-
 // Build mask and flags without modification of RayTracingInstance
 FRayTracingMaskAndFlags BuildRayTracingInstanceMaskAndFlags(const FRayTracingInstance& Instance, const FPrimitiveSceneProxy& PrimitiveSceneProxy);
 
