[HWRT] Fix potential integer underflow when selecting a buffer in AddDynamicMeshBatchForGeometryUpdate

#rnx

diff --git a/Engine/Source/Runtime/Renderer/Private/RayTracing/RayTracingDynamicGeometry.cpp b/Engine/Source/Runtime/Renderer/Private/RayTracing/RayTracingDynamicGeometry.cpp
--- a/Engine/Source/Runtime/Renderer/Private/RayTracing/RayTracingDynamicGeometry.cpp
+++ b/Engine/Source/Runtime/Renderer/Private/RayTracing/RayTracingDynamicGeometry.cpp
@@ -209,7 +209,7 @@ void FRayTracingDynamicGeometryCollection::AddDynamicMeshBatchForGeometryUpdate(
 		FVertexPositionBuffer* VertexPositionBuffer = nullptr;
 		for (FVertexPositionBuffer* Buffer : VertexPositionBuffers)
 		{
-			if ((Buffer->RWBuffer.NumBytes - Buffer->UsedSize) >= UpdateParams.VertexBufferSize)
+			if (Buffer->RWBuffer.NumBytes >= (UpdateParams.VertexBufferSize + Buffer->UsedSize))
 			{
 				VertexPositionBuffer = Buffer;
 				break;
