Prevent running VSM GPU-driven instance invalidation if there are no cached shadow maps.

#rnx

diff --git a/Engine/Source/Runtime/Renderer/Private/VirtualShadowMaps/VirtualShadowMapCacheManager.cpp b/Engine/Source/Runtime/Renderer/Private/VirtualShadowMaps/VirtualShadowMapCacheManager.cpp
--- a/Engine/Source/Runtime/Renderer/Private/VirtualShadowMaps/VirtualShadowMapCacheManager.cpp
+++ b/Engine/Source/Runtime/Renderer/Private/VirtualShadowMaps/VirtualShadowMapCacheManager.cpp
@@ -2022,14 +2022,12 @@ public:
 		}
 
 		// Run the resultant invalidations right away, this could be deferred and flushed later instead if desired.
-		if (CacheManager.IsCacheDataAvailable() && bWasAnyInvalidationsQueued)
+		if (CacheManager.IsCacheDataAvailable() && bWasAnyInvalidationsQueued && CacheManager.PrevUniformParameters.NumFullShadowMaps > 0)
 		{
 			RDG_EVENT_SCOPE(GraphBuilder, "Shadow.Virtual.ProcessInvalidationsGPU");
 
 			FVirtualShadowMapArrayCacheManager::FInvalidationPassCommon InvalidationPassCommon = CacheManager.GetUniformParametersForInvalidation(GraphBuilder, SceneRenderer.GetSceneUniforms());
 		
-			check(InvalidationPassCommon.UniformParameters->NumFullShadowMaps > 0);
-
 			FProcessInvalidationQueueGPUCS::FParameters* PassParameters = GraphBuilder.AllocParameters<FProcessInvalidationQueueGPUCS::FParameters>();
 
 			PassParameters->InvalidationArgs = GraphBuilder.CreateSRV(GPUInvalidateInstancesArgsRDG);
