[DeferredShadingRenderer] Build light grid unless lighting is disabled or using path tracing.

diff --git a/Engine/Source/Runtime/Renderer/Private/DeferredShadingRenderer.cpp b/Engine/Source/Runtime/Renderer/Private/DeferredShadingRenderer.cpp
--- a/Engine/Source/Runtime/Renderer/Private/DeferredShadingRenderer.cpp
+++ b/Engine/Source/Runtime/Renderer/Private/DeferredShadingRenderer.cpp
@@ -1777,7 +1777,6 @@ void FDeferredShadingSceneRenderer::Render(FRDGBuilder& GraphBuilder)
 		&& bUseGBuffer
 		&& !bHasRayTracedOverlay;
 
-	bool bComputeLightGrid = false;
 	bool bAnyLumenEnabled = false;
 
 	// Virtual texturing isn't needed for depth prepass
@@ -1807,35 +1806,14 @@ void FDeferredShadingSceneRenderer::Render(FRDGBuilder& GraphBuilder)
 		}
 #endif // RHI_RAYTRACING
 
+		if (!bHasRayTracedOverlay)
 		{
-			if (bUseGBuffer)
-			{
-				bComputeLightGrid = bRenderDeferredLighting;
-			}
-			else
-			{
-				bComputeLightGrid = ViewFamily.EngineShowFlags.Lighting;
-			}
-
-			if (!bHasRayTracedOverlay)
+			for (const FViewInfo& View : Views)
 			{
-				for (const FViewInfo& View : Views)
-				{
-					bAnyLumenEnabled = bAnyLumenEnabled
-						|| GetViewPipelineState(View).DiffuseIndirectMethod == EDiffuseIndirectMethod::Lumen
-						|| GetViewPipelineState(View).ReflectionsMethod == EReflectionsMethod::Lumen;
-				}
+				bAnyLumenEnabled = bAnyLumenEnabled
+					|| GetViewPipelineState(View).DiffuseIndirectMethod == EDiffuseIndirectMethod::Lumen
+					|| GetViewPipelineState(View).ReflectionsMethod == EReflectionsMethod::Lumen;
 			}
-
-			bComputeLightGrid |= (
-				bShouldRenderVolumetricFog ||
-				VolumetricCloudWantsToSampleLocalLights(Scene, ViewFamily.EngineShowFlags) ||
-				ViewFamily.ViewMode != VMI_Lit ||
-				bAnyLumenEnabled ||
-				VirtualShadowMapArray.IsEnabled() ||
-				ShouldVisualizeLightGrid() ||
-				ShouldRenderLocalFogVolume(Scene, ViewFamily)); // Needed when accessing forward light data for the directional light
-			bComputeLightGrid &= !ViewFamily.EngineShowFlags.PathTracing;
 		}
 
 		{
@@ -2282,7 +2260,7 @@ void FDeferredShadingSceneRenderer::Render(FRDGBuilder& GraphBuilder)
 			RDG_GPU_STAT_SCOPE(GraphBuilder, SortLights);
 
 			GatherAndSortLightsTask.Wait();
-			ComputeLightGridOutput = GatherLightsAndComputeLightGrid(GraphBuilder, bComputeLightGrid, *SortedLightSet);
+			ComputeLightGridOutput = GatherLightsAndComputeLightGrid(GraphBuilder, /*bComputeLightGrid*/ ViewFamily.EngineShowFlags.Lighting && !ViewFamily.EngineShowFlags.PathTracing, *SortedLightSet);
 
 			CSV_CUSTOM_STAT(LightCount, All,  float(SortedLightSet->SortedLights.Num()), ECsvCustomStatOp::Set);
 			CSV_CUSTOM_STAT(LightCount, Batched, float(SortedLightSet->UnbatchedLightStart), ECsvCustomStatOp::Set);
diff --git a/Engine/Source/Runtime/Renderer/Private/LightGridInjection.cpp b/Engine/Source/Runtime/Renderer/Private/LightGridInjection.cpp
--- a/Engine/Source/Runtime/Renderer/Private/LightGridInjection.cpp
+++ b/Engine/Source/Runtime/Renderer/Private/LightGridInjection.cpp
@@ -1255,26 +1255,10 @@ FComputeLightGridOutput FDeferredShadingSceneRenderer::GatherLightsAndComputeLig
 		return Result;
 	}
 
-	bool bAnyViewUsesForwardLighting = false;
-	bool bAnyViewUsesLumen = false;
-	bool bAnyViewUsesRayTracing = false;
-	for (int32 ViewIndex = 0; ViewIndex < Views.Num(); ViewIndex++)
-	{
-		const FViewInfo& View = Views[ViewIndex];
-		bAnyViewUsesForwardLighting |= View.bTranslucentSurfaceLighting || ShouldRenderVolumetricFog() || View.bHasSingleLayerWaterMaterial 
-			|| VolumetricCloudWantsToSampleLocalLights(Scene, ViewFamily.EngineShowFlags) || ShouldVisualizeLightGrid() || ShouldRenderLocalFogVolume(Scene, ViewFamily);
-		bAnyViewUsesLumen |= GetViewPipelineState(View).DiffuseIndirectMethod == EDiffuseIndirectMethod::Lumen || GetViewPipelineState(View).ReflectionsMethod == EReflectionsMethod::Lumen;
-		bAnyViewUsesRayTracing |= IsRayTracingEnabled() && View.IsRayTracingAllowedForView();
-	}
-	
-	const bool bCullLightsToGrid = GLightCullingQuality 
-		&& (IsForwardShadingEnabled(ShaderPlatform) || bAnyViewUsesForwardLighting || bAnyViewUsesRayTracing || ShouldUseClusteredDeferredShading() ||
-			bAnyViewUsesLumen || ViewFamily.EngineShowFlags.VisualizeMeshDistanceFields || VirtualShadowMapArray.IsEnabled() || MegaLights::IsEnabled(ViewFamily));
+	Result = ComputeLightGrid(GraphBuilder, true, SortedLightSet);
 
 	// Store this flag if lights are injected in the grids, check with 'AreLightsInLightGrid()'
-	bAreLightsInLightGrid = bCullLightsToGrid;
-	
-	Result = ComputeLightGrid(GraphBuilder, bCullLightsToGrid, SortedLightSet);
+	bAreLightsInLightGrid = true;
 
 	return Result;
 }
