Heterogeneous Volumes: Add render and game thread timing stats

diff --git a/Engine/Source/Runtime/Engine/Private/Components/HeterogeneousVolumeComponent.cpp b/Engine/Source/Runtime/Engine/Private/Components/HeterogeneousVolumeComponent.cpp
--- a/Engine/Source/Runtime/Engine/Private/Components/HeterogeneousVolumeComponent.cpp
+++ b/Engine/Source/Runtime/Engine/Private/Components/HeterogeneousVolumeComponent.cpp
@@ -4,6 +4,7 @@
 
 #include "Components/BillboardComponent.h"
 #include "Engine/Texture2D.h"
+#include "Stats/Stats.h"
 #include "HeterogeneousVolumeInterface.h"
 #include "LocalVertexFactory.h"
 #include "MaterialDomain.h"
@@ -21,6 +22,10 @@
 
 #include UE_INLINE_GENERATED_CPP_BY_NAME(HeterogeneousVolumeComponent)
 
+DECLARE_CYCLE_STAT(TEXT("Get Dynamic Mesh Elements"), STAT_HeterogeneousVolumeComponentGetDynamicMeshElements, STATGROUP_HeterogeneousVolumesRT);
+DECLARE_CYCLE_STAT(TEXT("Post Edit Change Property"), STAT_HeterogeneousVolumeComponentPostEditChangeProperty, STATGROUP_HeterogeneousVolumesGT);
+DECLARE_CYCLE_STAT(TEXT("Component Tick"), STAT_HeterogeneousVolumeComponentTick, STATGROUP_HeterogeneousVolumesGT);
+
 #define LOCTEXT_NAMESPACE "HeterogeneousVolumeComponent"
 
 class FHeterogeneousVolumeSceneProxy : public FPrimitiveSceneProxy
@@ -181,6 +186,8 @@ void FHeterogeneousVolumeSceneProxy::GetDynamicMeshElements(
 	uint32 VisibilityMap,
 	FMeshElementCollector& Collector) const
 {
+	CSV_SCOPED_TIMING_STAT_EXCLUSIVE(Effects);
+	SCOPE_CYCLE_COUNTER(STAT_HeterogeneousVolumeComponentGetDynamicMeshElements);
 	if (Views.IsEmpty())
 	{
 		return;
@@ -508,6 +515,8 @@ UMaterialInterface* UHeterogeneousVolumeComponent::GetHeterogeneousVolumeMateria
 #if WITH_EDITOR
 void UHeterogeneousVolumeComponent::PostEditChangeProperty(FPropertyChangedEvent& PropertyChangedEvent)
 {
+	CSV_SCOPED_TIMING_STAT_EXCLUSIVE(Effects);
+	SCOPE_CYCLE_COUNTER(STAT_HeterogeneousVolumeComponentPostEditChangeProperty);
 	Super::PostEditChangeProperty(PropertyChangedEvent);
 
 	const int32 SVTParameterIndex = 0;
@@ -621,6 +630,8 @@ void UHeterogeneousVolumeComponent::SetMaterial(int32 ElementIndex, UMaterialInt
 
 void UHeterogeneousVolumeComponent::TickComponent(float DeltaTime, ELevelTick TickType, FActorComponentTickFunction* ThisTickFunction)
 {
+	CSV_SCOPED_TIMING_STAT_EXCLUSIVE(Effects);
+	SCOPE_CYCLE_COUNTER(STAT_HeterogeneousVolumeComponentTick);
 	Super::TickComponent(DeltaTime, TickType, ThisTickFunction);
 
 	// Forcibly update the MID if it's not initialized
diff --git a/Engine/Source/Runtime/Engine/Public/HeterogeneousVolumeInterface.h b/Engine/Source/Runtime/Engine/Public/HeterogeneousVolumeInterface.h
--- a/Engine/Source/Runtime/Engine/Public/HeterogeneousVolumeInterface.h
+++ b/Engine/Source/Runtime/Engine/Public/HeterogeneousVolumeInterface.h
@@ -3,6 +3,7 @@
 #pragma once
 
 #include "GameFramework/Actor.h"
+#include "Stats/Stats.h"
 #include "CoreMinimal.h"
 #include "SceneManagement.h"
 #include "PrimitiveSceneProxy.h"
@@ -121,3 +122,6 @@ public:
 	virtual FString GetReadableName() const { return PrimitiveSceneProxy->GetResourceName().ToString(); }
 #endif // ACTOR_HAS_LABELS
 };
+
+DECLARE_STATS_GROUP(TEXT("Heterogeneous Volumes RT"), STATGROUP_HeterogeneousVolumesRT, STATCAT_HeterogeneousVolumes);
+DECLARE_STATS_GROUP(TEXT("Heterogeneous Volumes GT"), STATGROUP_HeterogeneousVolumesGT, STATCAT_HeterogeneousVolumes);
diff --git a/Engine/Source/Runtime/Renderer/Private/HeterogeneousVolumes/HeterogeneousVolumes.cpp b/Engine/Source/Runtime/Renderer/Private/HeterogeneousVolumes/HeterogeneousVolumes.cpp
--- a/Engine/Source/Runtime/Renderer/Private/HeterogeneousVolumes/HeterogeneousVolumes.cpp
+++ b/Engine/Source/Runtime/Renderer/Private/HeterogeneousVolumes/HeterogeneousVolumes.cpp
@@ -762,6 +762,8 @@ bool ShouldCacheVoxelGrids(const FScene* Scene, FSceneViewState* ViewState)
 	return false;
 }
 
+DECLARE_CYCLE_STAT(TEXT("Heterogeneous Volumes Render"), STATGROUP_HeterogeneousVolumesRender, STATGROUP_HeterogeneousVolumesRT);
+
 void RenderHeterogeneousVolumeShadows(
 	FRDGBuilder& GraphBuilder,
 	FScene* Scene,
@@ -770,7 +772,7 @@ void RenderHeterogeneousVolumeShadows(
 	TArray<FVisibleLightInfo, SceneRenderingAllocator>& VisibleLightInfos
 )
 {
-	RDG_EVENT_SCOPE(GraphBuilder, "HeterogeneousVolumeShadows");
+	RDG_EVENT_SCOPE_STAT(GraphBuilder, HeterogeneousVolumeShadowsStat, "HeterogeneousVolumeShadows");
 	RDG_GPU_STAT_SCOPE(GraphBuilder, HeterogeneousVolumeShadowsStat);
 	SCOPED_NAMED_EVENT(HeterogeneousVolumes, FColor::Emerald);
 
diff --git a/Engine/Source/Runtime/Renderer/Private/HeterogeneousVolumes/HeterogeneousVolumes.h b/Engine/Source/Runtime/Renderer/Private/HeterogeneousVolumes/HeterogeneousVolumes.h
--- a/Engine/Source/Runtime/Renderer/Private/HeterogeneousVolumes/HeterogeneousVolumes.h
+++ b/Engine/Source/Runtime/Renderer/Private/HeterogeneousVolumes/HeterogeneousVolumes.h
@@ -6,8 +6,8 @@
 #include "ConvexVolume.h"
 #include "RendererInterface.h"
 #include "ShaderParameterMacros.h"
+#include "HeterogeneousVolumeInterface.h"
 
-class IHeterogeneousVolumeInterface;
 class FLightSceneInfo;
 class FPrimitiveSceneProxy;
 class FProjectedShadowInfo;
@@ -855,3 +855,9 @@ void RenderSingleScatteringWithPreshadingHardwareRayTracing(
 );
 
 #endif // RHI_RAYTRACING
+
+DECLARE_CYCLE_STAT(TEXT("Ambient Occlusion"), STATGROUP_HeterogeneousVolumesAmbientOcclusion, STATGROUP_HeterogeneousVolumesRT);
+DECLARE_CYCLE_STAT(TEXT("Light Cache"), STATGROUP_HeterogeneousVolumesLightCache, STATGROUP_HeterogeneousVolumesRT);
+DECLARE_CYCLE_STAT(TEXT("Material Baking"), STATGROUP_HeterogeneousVolumesMaterialBaking, STATGROUP_HeterogeneousVolumesRT);
+DECLARE_CYCLE_STAT(TEXT("Shadows"), STATGROUP_HeterogeneousVolumesShadows, STATGROUP_HeterogeneousVolumesRT);
+DECLARE_CYCLE_STAT(TEXT("Single Scattering"), STATGROUP_HeterogeneousVolumesSingleScattering, STATGROUP_HeterogeneousVolumesRT);
diff --git a/Engine/Source/Runtime/Renderer/Private/HeterogeneousVolumes/HeterogeneousVolumesAmbientOcclusionPipeline.cpp b/Engine/Source/Runtime/Renderer/Private/HeterogeneousVolumes/HeterogeneousVolumesAmbientOcclusionPipeline.cpp
--- a/Engine/Source/Runtime/Renderer/Private/HeterogeneousVolumes/HeterogeneousVolumesAmbientOcclusionPipeline.cpp
+++ b/Engine/Source/Runtime/Renderer/Private/HeterogeneousVolumes/HeterogeneousVolumesAmbientOcclusionPipeline.cpp
@@ -741,6 +741,8 @@ void RenderAmbientOcclusionWithLiveShading(
 	FRDGTextureRef& AmbientOcclusionTexture
 )
 {
+	SCOPE_CYCLE_COUNTER(STATGROUP_HeterogeneousVolumesAmbientOcclusion);
+
 	bool bRenderAmbientOcclusion = HeterogeneousVolumes::EnableAmbientOcclusion() && HeterogeneousVolumes::UseIndirectLighting();
 	if (!bRenderAmbientOcclusion || HeterogeneousVolumeInterface->IsHoldout())
 	{
diff --git a/Engine/Source/Runtime/Renderer/Private/HeterogeneousVolumes/HeterogeneousVolumesLiveShadingPipeline.cpp b/Engine/Source/Runtime/Renderer/Private/HeterogeneousVolumes/HeterogeneousVolumesLiveShadingPipeline.cpp
--- a/Engine/Source/Runtime/Renderer/Private/HeterogeneousVolumes/HeterogeneousVolumesLiveShadingPipeline.cpp
+++ b/Engine/Source/Runtime/Renderer/Private/HeterogeneousVolumes/HeterogeneousVolumesLiveShadingPipeline.cpp
@@ -1463,9 +1463,13 @@ static void RenderWithInscatteringVolumePipeline(
 	FRDGTextureRef& HeterogeneousVolumeHoldout
 )
 {
+	CSV_SCOPED_TIMING_STAT_EXCLUSIVE(Effects);
+
 	bool bRenderLightingCache = !HeterogeneousVolumes::IsHoldout(HeterogeneousVolumeInterface);
 	if (bRenderLightingCache)
 	{
+		SCOPE_CYCLE_COUNTER(STATGROUP_HeterogeneousVolumesLightCache);
+
 		// Light culling
 		TArray<FLightSceneInfoCompact, TInlineAllocator<64>> LightSceneInfoCompact;
 		for (auto LightIt = Scene->Lights.CreateConstIterator(); LightIt; ++LightIt)
@@ -1530,6 +1534,8 @@ static void RenderWithInscatteringVolumePipeline(
 
 	// Direct volume integrator
 	{
+		SCOPE_CYCLE_COUNTER(STATGROUP_HeterogeneousVolumesSingleScattering);
+
 		bool bApplyEmissionAndTransmittance = true;
 		bool bApplyDirectLighting = true;
 		bool bApplyShadowTransmittance = true;
@@ -2090,6 +2096,8 @@ bool RenderVolumetricShadowMapForLightWithLiveShading(
 	FRDGBufferRef& VolumetricShadowLinkedListBuffer
 )
 {
+	CSV_SCOPED_TIMING_STAT_EXCLUSIVE(Effects);
+	SCOPE_CYCLE_COUNTER(STATGROUP_HeterogeneousVolumesShadows);
 	check(LightSceneInfo);
 	check(VisibleLightInfo);
 
diff --git a/Engine/Source/Runtime/Renderer/Private/HeterogeneousVolumes/HeterogeneousVolumesPreshadingPipeline.cpp b/Engine/Source/Runtime/Renderer/Private/HeterogeneousVolumes/HeterogeneousVolumesPreshadingPipeline.cpp
--- a/Engine/Source/Runtime/Renderer/Private/HeterogeneousVolumes/HeterogeneousVolumesPreshadingPipeline.cpp
+++ b/Engine/Source/Runtime/Renderer/Private/HeterogeneousVolumes/HeterogeneousVolumesPreshadingPipeline.cpp
@@ -727,74 +727,82 @@ void RenderWithInscatteringVolumePipelineWithPreshadingCompute(
 {
 	RDG_EVENT_SCOPE(GraphBuilder, "Direct Volume Rendering");
 
-	// Light culling
-	TArray<FLightSceneInfoCompact, TInlineAllocator<64>> LightSceneInfoCompact;
-	for (auto LightIt = Scene->Lights.CreateConstIterator(); LightIt; ++LightIt)
+	bool bRenderLightingCache = !HeterogeneousVolumes::IsHoldout(HeterogeneousVolumeInterface);
+	if (bRenderLightingCache)
 	{
-		if (HeterogeneousVolumes::SupportsLightType(LightIt->LightType) &&
-			LightIt->AffectsPrimitive(HeterogeneousVolumeInterface->GetBounds(), HeterogeneousVolumeInterface->GetPrimitiveSceneProxy()))
+		SCOPE_CYCLE_COUNTER(STATGROUP_HeterogeneousVolumesLightCache);
+
+		// Light culling
+		TArray<FLightSceneInfoCompact, TInlineAllocator<64>> LightSceneInfoCompact;
+		for (auto LightIt = Scene->Lights.CreateConstIterator(); LightIt; ++LightIt)
 		{
-			LightSceneInfoCompact.Add(*LightIt);
+			if (HeterogeneousVolumes::SupportsLightType(LightIt->LightType) &&
+				LightIt->AffectsPrimitive(HeterogeneousVolumeInterface->GetBounds(), HeterogeneousVolumeInterface->GetPrimitiveSceneProxy()))
+			{
+				LightSceneInfoCompact.Add(*LightIt);
+			}
 		}
-	}
-
-	// Light loop:
-	int32 NumPasses = LightSceneInfoCompact.Num();
-	for (int32 PassIndex = 0; PassIndex < NumPasses; ++PassIndex)
-	{
-		bool bApplyEmissionAndTransmittance = PassIndex == 0;
-		bool bApplyDirectLighting = !LightSceneInfoCompact.IsEmpty();
-		bool bApplyShadowTransmittance = false;
 
-		uint32 LightType = 0;
-		FLightSceneInfo* LightSceneInfo = nullptr;
-		const FVisibleLightInfo* VisibleLightInfo = nullptr;
-		if (bApplyDirectLighting)
+		// Light loop:
+		int32 NumPasses = LightSceneInfoCompact.Num();
+		for (int32 PassIndex = 0; PassIndex < NumPasses; ++PassIndex)
 		{
-			LightType = LightSceneInfoCompact[PassIndex].LightType;
-			LightSceneInfo = LightSceneInfoCompact[PassIndex].LightSceneInfo;
-			check(LightSceneInfo != nullptr);
-
-			bApplyDirectLighting = (LightSceneInfo != nullptr);
-			if (LightSceneInfo)
+			bool bApplyEmissionAndTransmittance = PassIndex == 0;
+			bool bApplyDirectLighting = !LightSceneInfoCompact.IsEmpty();
+			bool bApplyShadowTransmittance = false;
+
+			uint32 LightType = 0;
+			FLightSceneInfo* LightSceneInfo = nullptr;
+			const FVisibleLightInfo* VisibleLightInfo = nullptr;
+			if (bApplyDirectLighting)
 			{
-				VisibleLightInfo = &VisibleLightInfos[LightSceneInfo->Id];
-				bApplyShadowTransmittance = LightSceneInfo->Proxy->CastsVolumetricShadow();
+				LightType = LightSceneInfoCompact[PassIndex].LightType;
+				LightSceneInfo = LightSceneInfoCompact[PassIndex].LightSceneInfo;
+				check(LightSceneInfo != nullptr);
+
+				bApplyDirectLighting = (LightSceneInfo != nullptr);
+				if (LightSceneInfo)
+				{
+					VisibleLightInfo = &VisibleLightInfos[LightSceneInfo->Id];
+					bApplyShadowTransmittance = LightSceneInfo->Proxy->CastsVolumetricShadow();
+				}
 			}
-		}
 
-		RenderLightingCacheWithPreshadingCompute(
-			GraphBuilder,
-			// Scene data
-			Scene,
-			View, ViewIndex,
-			SceneTextures,
-			// Light data
-			bApplyEmissionAndTransmittance,
-			bApplyDirectLighting,
-			bApplyShadowTransmittance,
-			LightType,
-			LightSceneInfo,
-			// Shadow data
-			VisibleLightInfo,
-			VirtualShadowMapArray,
-			// Object data
-			HeterogeneousVolumeInterface,
-			// Volume data
-			NumVoxelsBuffer,
-			// Sparse voxel data
-			SparseVoxelUniformBuffer,
-			// Ray marching tile
-			NumRayMarchingTilesBuffer,
-			RayMarchingTilesBuffer,
-			VoxelsPerTileBuffer,
-			// Output
-			LightingCacheTexture
-		);
+			RenderLightingCacheWithPreshadingCompute(
+				GraphBuilder,
+				// Scene data
+				Scene,
+				View, ViewIndex,
+				SceneTextures,
+				// Light data
+				bApplyEmissionAndTransmittance,
+				bApplyDirectLighting,
+				bApplyShadowTransmittance,
+				LightType,
+				LightSceneInfo,
+				// Shadow data
+				VisibleLightInfo,
+				VirtualShadowMapArray,
+				// Object data
+				HeterogeneousVolumeInterface,
+				// Volume data
+				NumVoxelsBuffer,
+				// Sparse voxel data
+				SparseVoxelUniformBuffer,
+				// Ray marching tile
+				NumRayMarchingTilesBuffer,
+				RayMarchingTilesBuffer,
+				VoxelsPerTileBuffer,
+				// Output
+				LightingCacheTexture
+			);
+		}
 	}
 
 	// Direct volume integrator
 	{
+		SCOPE_CYCLE_COUNTER(STATGROUP_HeterogeneousVolumesSingleScattering);
+
 		bool bApplyEmissionAndTransmittance = true;
 		bool bApplyDirectLighting = true;
 		bool bApplyShadowTransmittance = true;
@@ -1547,6 +1555,8 @@ void RenderWithPreshading(
 	FRDGTextureRef& HeterogeneousVolumeRadiance
 )
 {
+	CSV_SCOPED_TIMING_STAT_EXCLUSIVE(Effects);
+
 	// Determine baking voxel resolution
 	float LODFactor = HeterogeneousVolumes::CalcLODFactor(View, HeterogeneousVolumeInterface);
 	FIntVector VolumeResolution = HeterogeneousVolumes::GetVolumeResolution(HeterogeneousVolumeInterface);
@@ -1569,6 +1579,7 @@ void RenderWithPreshading(
 	// Preshading pipeline
 	{
 		RDG_EVENT_SCOPE(GraphBuilder, "Preshading Pipeline");
+		SCOPE_CYCLE_COUNTER(STATGROUP_HeterogeneousVolumesMaterialBaking);
 
 		{
 			RDG_EVENT_SCOPE(GraphBuilder, "Material Baking");
@@ -1614,6 +1625,7 @@ void RenderWithPreshading(
 	FRDGBufferRef VoxelBuffer;
 	{
 		RDG_EVENT_SCOPE(GraphBuilder, "Sparse Voxel Pipeline");
+		SCOPE_CYCLE_COUNTER(STATGROUP_HeterogeneousVolumesMaterialBaking);
 
 		{
 			RDG_EVENT_SCOPE(GraphBuilder, "Min Generation");
