IES bugfix
This may happen when AddTexture / RemoveTexture are called back-to-back before UpdateAtlasTexture has run. When that happens, we call InvalidateSlots with RequestedSliceCount=0, ending up with no FreeSlices
Then, if the atlas has not been created yet, we call CreateAtlasTexture with RequestedSliceCount=0 which actually creates a 1-sliced texture array, with uninitialized data
The next time we call AddTexture / UpdateAtlasTexture, bForceRecreate was false because the array slice is 1. But since we don't have any element in FreeSlices we skip the copy, leaving the data uninitialized.
Finally, even the copy was skipped, we were returning GIESTextureManager.Slots[InSlotId].SliceIndex in GetAtlasSlot. The problem is that InvalidSliceIndex is a uint32, set to UINT_MAX. So instead of doing float(-1), we were returning float(UINT_MAX) and we would try to sample the uninitialized data in ComputeLightProfileMultiplier because InIESAtlasIndex is actually positive

This CL fixes the access to ensure a negative value is returned in case the copy fails (although I don't think it should happen)
It also keeps track if we didn't have any slot when we requested the atlas creation and will force an update when we add a new one

#rnx

diff --git a/Engine/Source/Runtime/Renderer/Private/IESTextureManager.cpp b/Engine/Source/Runtime/Renderer/Private/IESTextureManager.cpp
--- a/Engine/Source/Runtime/Renderer/Private/IESTextureManager.cpp
+++ b/Engine/Source/Runtime/Renderer/Private/IESTextureManager.cpp
@@ -78,6 +78,7 @@ struct FIESTextureManager : public FRenderResource
 	TQueue<uint32> FreeSlots;
 	TQueue<uint32> FreeSlices;
 
+	uint32 RequestedSliceCountWhenTextureAllocated = 0;
 	bool bHasPendingAdds = false;
 	bool bHasPendingRefreshes = false;
 
@@ -452,7 +453,8 @@ float GetAtlasSlot(uint32 InSlotId)
 {
 	if (InSlotId < uint32(GIESTextureManager.Slots.Num()) && GIESTextureManager.AtlasTexture)
 	{
-		return float(GIESTextureManager.Slots[InSlotId].SliceIndex);
+		uint32 SlotSliceIndex = GIESTextureManager.Slots[InSlotId].SliceIndex;
+		return (SlotSliceIndex != InvalidSliceIndex) ? float(SlotSliceIndex) : float(INDEX_NONE);
 	}
 	else
 	{
@@ -503,7 +505,7 @@ void UpdateAtlasTexture(FRDGBuilder& GraphBuilder, const FStaticShaderPlatform S
 	if (GIESTextureManager.bHasPendingAdds)
 	{
 		RequestedSliceCount = FMath::Min(CachedAtlasMaxSlice, GetRequestedSlotCount(GIESTextureManager));
-		bForceRecreate = GIESTextureManager.AtlasTexture == nullptr || (GIESTextureManager.AtlasTexture != nullptr && GIESTextureManager.AtlasTexture->GetDesc().ArraySize < RequestedSliceCount);
+		bForceRecreate = (GIESTextureManager.AtlasTexture == nullptr) || (GIESTextureManager.RequestedSliceCountWhenTextureAllocated < RequestedSliceCount);
 	}
 	bForceRecreate = bForceRecreate || bHasSettingChanged;
 
@@ -568,6 +570,7 @@ void UpdateAtlasTexture(FRDGBuilder& GraphBuilder, const FStaticShaderPlatform S
 			{
 				AtlasTexture = CreateAtlasTexture(GraphBuilder, ShaderPlatform, CachedAtlasResolution, RequestedSliceCount);
 				bNeedExtraction = true;
+				GIESTextureManager.RequestedSliceCountWhenTextureAllocated = RequestedSliceCount;
 			}
 			else
 			{
