First Person: Implemented a specialized first-person clipmap for casting shadow from the (owner no-see) 3rd person mesh onto the environment, but not the (visible) first-person geometry.
The required setup goes like this:
1. There must be two primitive components owned the camera; one for the first person geometry and one for the world space representation seen by shadows and other players.
2. These primitives must be tagged as such using the FirstPersonPrimitiveType property.
3. The first person primitive should probably have bCastShadows=false, but currently this is left to the user to disable and isn't automatically disabled.
4. The third person primitive will automatically set bCastHiddenShadow=false and bOwnerNoSee=true internally. Users don't need to do anything; this happens automatically depending on the value of FirstPersonPrimitiveType. This flag configuration is required for VSM first person handling as this configuration essentially removes the primitive from the regular shadow setup and enables it to go down the newly added path for handling FirstPersonWorldSpace primitives.

Limitations:
 * Does not support Nanite
 * Does not support translucent geometry (except single-layer water with the prepass enabled)
 * Does not work with froxel page marking
 * Unsure about hair

Cvars:
 * r.FirstPerson.Shadow.Virtual.Clipmap - enable/disable the feature
 * r.FirstPerson.Shadow.Virtual.Clipmap.FirstLevel/LastLevel - control what levels are marked in the specialized clipmap
 * r.FirstPerson.Shadow.Virtual.Clipmap.PixelRequestBias - bias the clipmap levels of pages requested from pixels marked as first person (-1 do disable marking)
 * r.FirstPerson.Shadow.Virtual.Clipmap.RequestMinLevelClamp - minimum level to request

Further ideas / notes
 * Page marking clamp & bias might be possible to derive from the FP geometry itself, possibly, but it seems difficult in general
 * Marks the same number of pages for the third person geometry as the general view (minus any pixels as first person) this could probably be trimmed down by only marking those that could possibly be affected by the casting geometry.
    Would save page allocations & clearing but would be somewhat invasive perhaps.
    Might be worth doing depending on where the costs are.
 * It's not clear how to support Nanite because of the very special culling requirements & the GPU-driven rendering. So this requires some work & thinking about.

diff --git a/Engine/Shaders/Private/VirtualShadowMaps/VirtualShadowMapPageMarking.usf b/Engine/Shaders/Private/VirtualShadowMaps/VirtualShadowMapPageMarking.usf
--- a/Engine/Shaders/Private/VirtualShadowMaps/VirtualShadowMapPageMarking.usf
+++ b/Engine/Shaders/Private/VirtualShadowMaps/VirtualShadowMapPageMarking.usf
@@ -39,9 +39,9 @@ uint bCullBackfacingPixels;
 uint NumDirectionalLightSmInds;
 uint MipModeLocal;
 
-int GetBiasedClipmapLevel(FVirtualShadowMapProjectionShaderData BaseProjectionData, float3 TranslatedWorldPosition)
+int GetBiasedClipmapLevel(FVirtualShadowMapProjectionShaderData BaseProjectionData, float3 TranslatedWorldPosition, float ExtraBias)
 {
-	float BiasedLevel = CalcAbsoluteClipmapLevel(BaseProjectionData, TranslatedWorldPosition) + BaseProjectionData.ResolutionLodBias + VirtualShadowMap.GlobalResolutionLodBias;
+	float BiasedLevel = CalcAbsoluteClipmapLevel(BaseProjectionData, TranslatedWorldPosition) + BaseProjectionData.ResolutionLodBias + VirtualShadowMap.GlobalResolutionLodBias + ExtraBias;
 	return int(floor(BiasedLevel));
 }
 
@@ -115,9 +115,11 @@ void MarkPageClipmap(
 	FVirtualShadowMapProjectionShaderData ProjectionData,
 	bool bUsePageDilation, 
 	float2 PageDilationOffset,
-	float3 TranslatedWorldPosition)
+	float3 TranslatedWorldPosition,
+	float ExtraBias = 0.0f,
+	int MinLevelClamp = 0)
 {
-	const int ClipmapLevel = GetBiasedClipmapLevel(ProjectionData, TranslatedWorldPosition);
+	const int ClipmapLevel = max(MinLevelClamp, GetBiasedClipmapLevel(ProjectionData, TranslatedWorldPosition, ExtraBias));
 	int ClipmapIndex = max(0, ClipmapLevel - ProjectionData.ClipmapLevel);
 	if (ClipmapIndex < ProjectionData.ClipmapLevelCountRemaining)
 	{
@@ -224,6 +226,10 @@ FPositionData InitPositionData(uint2 PixelPos, float DeviceZ, bool bIsValid)
 	return Out;
 }
 
+// Bias applied to pixels originating from first-person geometry, if negative marking is skipped for these
+float FirstPersonPixelRequestBias;
+int FirstPersonPixelRequestLevelClamp;
+
 // Note: we use the tile size defined by the water as the group-size - this is needed because the tile mask testing code relies on the size being the same to scalarize efficiently.
 [numthreads(SLW_TILE_SIZE_XY, SLW_TILE_SIZE_XY, 1)]
 void GeneratePageFlagsFromPixels(
@@ -263,6 +269,7 @@ void GeneratePageFlagsFromPixels(
 	bool   bIsValid					= true;
 	bool   bIsWaterDepthValid		= false;
 	bool   bIsTranslucencyDepthValid= false;
+	bool   bIsFirstPersonPixel		= false;
 #if PERMUTATION_INPUT_TYPE == INPUT_TYPE_HAIRSTRANDS
 	{
 		DeviceZ = HairStrands.HairOnlyDepthTexture.Load(uint3(PixelPos, 0)).x;
@@ -279,12 +286,14 @@ void GeneratePageFlagsFromPixels(
 		WorldNormal = TopLayerData.WorldNormal;
 		bIsValid = SubstratePixelHeader.IsSubstrateMaterial();
 		bBackfaceCull = bBackfaceCull && !SubstratePixelHeader.HasSubsurface();
+		bIsFirstPersonPixel = SubstratePixelHeader.IsFirstPerson();
 	#else // !SUBSTRATE_ENABLED
 		FGBufferData GBufferData = GetGBufferDataUint(PixelPos, true);
 		WorldNormal = GBufferData.WorldNormal;
 		// Excluding unlit to avoid including processing sky dome
 		bIsValid = GBufferData.ShadingModelID != SHADINGMODELID_UNLIT;
 		bBackfaceCull = bBackfaceCull && !IsSubsurfaceModel(GBufferData.ShadingModelID);
+		bIsFirstPersonPixel = IsFirstPerson(GBufferData);
 	#endif // SUBSTRATE_ENABLED
 
 	#if PERMUTATION_WATER_DEPTH
@@ -391,6 +400,12 @@ void GeneratePageFlagsFromPixels(
 		(GroupIndex & 1) ? 1.0f : -1.0f,
 		(GroupIndex & 2) ? 1.0f : -1.0f);
 		
+	// Apply extra bias if the pixel is from first-person geometry, as this geometry is not self-shadowing and very close to the camera (requesting high resolution)
+	// We can't outright skip the pixels as we want environment shadow cast onto the FP geometry.
+	float PrimaryExtraBias = bIsFirstPersonPixel ? FirstPersonPixelRequestBias : 0.0f;
+	int PrimaryMinLevelClamp = bIsFirstPersonPixel ? FirstPersonPixelRequestLevelClamp : 0;
+
+	bool bMarkFirstPersonPixels = FirstPersonPixelRequestBias >= 0.0f;
 	// Directional lights
 	{
 		const bool bUsePageDilation = PageDilationBorderSizeDirectional > 0.0f;
@@ -410,9 +425,16 @@ void GeneratePageFlagsFromPixels(
 			}
 			#endif
 
+			// Don't mark first person pixels in a first person (world space) shadow map as it doesn't cast shadow onto the FP mesh, or if the extra bias is negative (to explictly disable).
+			bool bIsFirstPersonShadowMap = (ProjectionData.Flags & VSM_PROJ_FLAG_IS_FIRST_PERSON_SHADOW) != 0u;
+			if (bIsFirstPersonPixel && (bIsFirstPersonShadowMap || !bMarkFirstPersonPixels))
+			{
+				bMarkPrimary = false;
+			}
+
 			if (bMarkPrimary)
 			{
-				MarkPageClipmap(ProjectionData, bUsePageDilation, PageDilationOffset, Primary.TranslatedWorldPosition);
+				MarkPageClipmap(ProjectionData, bUsePageDilation, PageDilationOffset, Primary.TranslatedWorldPosition, PrimaryExtraBias, PrimaryMinLevelClamp);
 			}
 			
 			#if PERMUTATION_WATER_DEPTH
@@ -433,6 +455,8 @@ void GeneratePageFlagsFromPixels(
 	}
 
 	// Local lights
+	// Don't mark pixels originating from first person if disabled.
+	if (bIsFirstPersonPixel && !bMarkFirstPersonPixels)
 	{
 		const bool bUsePageDilation = PageDilationBorderSizeLocal > 0.0f;
 		const float2 PageDilationOffset = PageDilationBorderSizeLocal * PageDilationDither;
diff --git a/Engine/Shaders/Private/VirtualShadowMaps/VirtualShadowMapProjection.usf b/Engine/Shaders/Private/VirtualShadowMaps/VirtualShadowMapProjection.usf
--- a/Engine/Shaders/Private/VirtualShadowMaps/VirtualShadowMapProjection.usf
+++ b/Engine/Shaders/Private/VirtualShadowMaps/VirtualShadowMapProjection.usf
@@ -48,6 +48,7 @@ struct FProjectionShadingInfo
 	bool bIsHair;
 	bool bIsSubsurface;
 	bool bDisableBackFaceCulling;
+	bool bIsFirstPerson;
 
 	float3 WorldNormal;
 };
@@ -82,6 +83,7 @@ FProjectionShadingInfo GetProjectionShadingInfo(uint2 PixelPos)
 	{
 		Out.bDisableBackFaceCulling = true;
 	}
+	Out.bIsFirstPerson = SubstratePixelHeader.IsFirstPerson();
 #else
 	const FGBufferData GBufferData = GetGBufferDataUint(PixelPos, true);
 	Out.bIsValid = GBufferData.ShadingModelID != SHADINGMODELID_UNLIT;
@@ -89,6 +91,7 @@ FProjectionShadingInfo GetProjectionShadingInfo(uint2 PixelPos)
 	Out.bIsSubsurface = IsSubsurfaceModel(GBufferData.ShadingModelID);
 	Out.bDisableBackFaceCulling = GBufferData.ShadingModelID == SHADINGMODELID_CLEAR_COAT;
 	Out.WorldNormal = GBufferData.WorldNormal;
+	Out.bIsFirstPerson = IsFirstPerson(GBufferData);
 #endif
 
 	Out.bDisableBackFaceCulling = Out.bDisableBackFaceCulling || Out.bIsSubsurface || Out.bIsHair;
@@ -163,6 +166,10 @@ int VisualizeModeId;
 int bVisualizeCachedPagesOnly;
 int VisualizeVirtualShadowMapId;
 
+#if FIRST_PERSON_SHADOW
+int FirstPersonVirtualShadowMapId;
+#endif
+
 // Type of input data consume by the page allocation (i.e., data read from the source buffer: Gbuffer, HairStrands data, ...)
 #define INPUT_TYPE_GBUFFER 0
 #define INPUT_TYPE_HAIRSTRANDS 1
@@ -280,6 +287,53 @@ FVirtualShadowMapSampleResult ProjectLight(
 				SMRTRayOffset,
 				WorldNormal);
 		}
+#if FIRST_PERSON_SHADOW
+		// If there is an extra FP shadow map, it represents the shadow cast from the 3rd person mesh of the player onto the world only, thus it should only be sampled if
+		// 1. the pixel is not marked as FP (e.g., the weapon)
+		// 
+		// 2. the pixel is not already in shadow (100%)
+		bool bFullyShadowed = Result.bValid && Result.ShadowFactor == 0.0f;
+
+		if (FirstPersonVirtualShadowMapId >= 0 && !ShadingInfo.bIsFirstPerson && !bFullyShadowed)
+		{
+			FVirtualShadowMapSampleResult FPResult = InitVirtualShadowMapSampleResult();
+			if (GetSMRTRayCount() > 0)
+			{
+				FPResult = TraceDirectional(
+					FirstPersonVirtualShadowMapId,
+					Light,
+					PixelPos,
+					SceneDepth,
+					TranslatedWorldPosition,
+					SMRTRayOffset,
+					Noise,
+					WorldNormal);
+			}
+			else
+			{
+				FPResult = SampleVirtualShadowMapDirectional(
+					FirstPersonVirtualShadowMapId,
+					TranslatedWorldPosition,
+					SMRTRayOffset,
+					WorldNormal);
+			}
+			if (FPResult.bValid)
+			{
+				if (!Result.bValid)
+				{
+					Result = FPResult;
+				}
+				else
+				{
+					// Merge the result, picking the most-shadowed value & nearest occluder distance
+					Result.ShadowFactor = min(Result.ShadowFactor, FPResult.ShadowFactor);
+					Result.OccluderDistance = min(Result.OccluderDistance, FPResult.OccluderDistance);
+				}
+			}
+		}
+#endif // FIRST_PERSON_SHADOW
+
+
 		#else // !DIRECTIONAL_LIGHT
 		if (GetSMRTRayCount() > 0)
 		{
diff --git a/Engine/Shaders/Shared/VirtualShadowMapDefinitions.h b/Engine/Shaders/Shared/VirtualShadowMapDefinitions.h
--- a/Engine/Shaders/Shared/VirtualShadowMapDefinitions.h
+++ b/Engine/Shaders/Shared/VirtualShadowMapDefinitions.h
@@ -40,6 +40,7 @@ VirtualShadowMapDefinitions.h: used in virtual shadow map shaders and C++ code t
 #define VSM_PROJ_FLAG_UNCACHED (1U << 1) // Used to indicate that the light is uncached and should only render to dynamic pages
 #define VSM_PROJ_FLAG_UNREFERENCED (1U << 2) // Used to indicate that the light is not referenced/rendered to this render
 #define VSM_PROJ_FLAG_IS_COARSE_CLIP_LEVEL (1U << 3) // Used to indicate that the clip level is a coarse level
+#define VSM_PROJ_FLAG_IS_FIRST_PERSON_SHADOW (1U << 4) // Used to indicate that this is a "first-person" shadow 
 
 // Hard limit for max distant lights supported 8k for now - we may revise later. We need to keep them in a fixed range for now to make allocation easy and minimize overhead for indexing.
 #define VSM_MAX_SINGLE_PAGE_SHADOW_MAPS (1024U * 8U)
diff --git a/Engine/Source/Runtime/Engine/Classes/Components/PrimitiveComponent.h b/Engine/Source/Runtime/Engine/Classes/Components/PrimitiveComponent.h
--- a/Engine/Source/Runtime/Engine/Classes/Components/PrimitiveComponent.h
+++ b/Engine/Source/Runtime/Engine/Classes/Components/PrimitiveComponent.h
@@ -132,7 +132,8 @@ enum class EFirstPersonPrimitiveType : uint8
 	None UMETA(DisplayName = "None"),
 	/** Primitive is rendered as first person and affected by first person properties on the camera. */
 	FirstPerson UMETA(DisplayName = "First Person"),
-	/** Primitive represents a first person primitive in world space. This is the primitive that other players would see and is used to cast shadow onto the ground among other things. */
+	/** Primitive represents a first person primitive in world space. This is the primitive that other players would see and is used to cast shadow onto the ground among other things.
+		Implicitly sets bCastHiddenShadow=false and bOwnerNoSee=true internally, which is required for first person shadow to work correctly with VSM. */
 	WorldSpaceRepresentation UMETA(DisplayName = "World Space Representation"),
 };
 
@@ -444,7 +445,7 @@ public:
 	UPROPERTY(EditAnywhere, AdvancedDisplay, BlueprintReadOnly, Category = Rendering, Interp)
 	uint8 bHoldout : 1;
 
-	/** If this is True, this component won't be visible when the view actor is the component's owner, directly or indirectly. */
+	/** If this is True, this component won't be visible when the view actor is the component's owner, directly or indirectly. Will be internally overridden to true when FirstPersonPrimitiveType is set to WorldSpaceRepresentation. */
 	UPROPERTY(EditAnywhere, AdvancedDisplay, BlueprintReadOnly, Category = Rendering)
 	uint8 bOwnerNoSee:1;
 
@@ -563,7 +564,8 @@ public:
 	/** 
 	 *	If true, the primitive will cast shadows even if bHidden is true.
 	 *	Controls whether the primitive should cast shadows when hidden.
-	 *	This flag is only used if CastShadow is true.
+	 *	This flag is only used if CastShadow is true and will be internally 
+	 *  overridden to false if FirstPersonPrimitiveType is set to WorldSpaceRepresentation.
 	 */
 	UPROPERTY(EditAnywhere, AdvancedDisplay, BlueprintReadOnly, Category=Lighting, meta=(EditCondition="CastShadow", DisplayName = "Hidden Shadow"), Interp)
 	uint8 bCastHiddenShadow:1;
diff --git a/Engine/Source/Runtime/Engine/Private/Components/PrimitiveComponent.cpp b/Engine/Source/Runtime/Engine/Private/Components/PrimitiveComponent.cpp
--- a/Engine/Source/Runtime/Engine/Private/Components/PrimitiveComponent.cpp
+++ b/Engine/Source/Runtime/Engine/Private/Components/PrimitiveComponent.cpp
@@ -1358,6 +1358,8 @@ bool UPrimitiveComponent::CanEditChange(const FProperty* InProperty) const
 		static FName CastCinematicShadowName = GET_MEMBER_NAME_CHECKED(UPrimitiveComponent, bCastCinematicShadow);
 		static FName CastInsetShadowName = GET_MEMBER_NAME_CHECKED(UPrimitiveComponent, bCastInsetShadow);
 		static FName CastShadowName = GET_MEMBER_NAME_CHECKED(UPrimitiveComponent, CastShadow);
+		static FName CastHiddenShadowName = GET_MEMBER_NAME_CHECKED(UPrimitiveComponent, bCastHiddenShadow);
+		static FName OwnerNoSeeName = GET_MEMBER_NAME_CHECKED(UPrimitiveComponent, bOwnerNoSee);
 
 		if (PropertyName == LightmassSettingsName)
 		{
@@ -1394,6 +1396,13 @@ bool UPrimitiveComponent::CanEditChange(const FProperty* InProperty) const
 		{
 			return !bSelfShadowOnly;
 		}
+
+		// FirstPersonWorldSpaceRepresentation primitives must have bCastHiddenShadows=false and bOwnerNoSee=true to work correctly with VSM.
+		// Other setups are not expected to be needed, so to avoid confusion, we'll gray out these properties and simply override them on the scene proxy.
+		if (PropertyName == CastHiddenShadowName || PropertyName == OwnerNoSeeName)
+		{
+			return FirstPersonPrimitiveType != EFirstPersonPrimitiveType::WorldSpaceRepresentation;
+		}
 	}
 
 	return bIsEditable;
diff --git a/Engine/Source/Runtime/Engine/Private/PrimitiveSceneProxy.cpp b/Engine/Source/Runtime/Engine/Private/PrimitiveSceneProxy.cpp
--- a/Engine/Source/Runtime/Engine/Private/PrimitiveSceneProxy.cpp
+++ b/Engine/Source/Runtime/Engine/Private/PrimitiveSceneProxy.cpp
@@ -702,6 +702,14 @@ FPrimitiveSceneProxy::FPrimitiveSceneProxy(const FPrimitiveSceneProxyDesc& InPro
 	}
 #endif
 
+	// FirstPersonWorldSpaceRepresentation primitives must have bCastHiddenShadow=false and bOwnerNoSee=true.
+	// This way, they are skipped by regular shadow rendering and can be picked up by the special first person
+	// shadow path in VSM. Raytraced shadows ignore these flags anyways, so this works for all cases.
+	if (bIsFirstPersonWorldSpaceRepresentation)
+	{
+		bCastHiddenShadow = false;
+		bOwnerNoSee = true;
+	}
 }
 
 bool FPrimitiveSceneProxy::OnLevelAddedToWorld_RenderThread()
diff --git a/Engine/Source/Runtime/Renderer/Private/SceneRendering.cpp b/Engine/Source/Runtime/Renderer/Private/SceneRendering.cpp
--- a/Engine/Source/Runtime/Renderer/Private/SceneRendering.cpp
+++ b/Engine/Source/Runtime/Renderer/Private/SceneRendering.cpp
@@ -583,7 +583,7 @@ FASTVRAM_CVAR(DistanceFieldAOScreenGridResources, 1);
 FASTVRAM_CVAR(ForwardLightingCullingResources, 1);
 FASTVRAM_CVAR(GlobalDistanceFieldCullGridBuffers, 1);
 
-TSharedPtr<FVirtualShadowMapClipmap> FVisibleLightInfo::FindShadowClipmapForView(const FViewInfo* View) const
+static TSharedPtr<FVirtualShadowMapClipmap> FindShadowClipmapForView(const TArray<TSharedPtr<FVirtualShadowMapClipmap>,SceneRenderingAllocator>& VirtualShadowMapClipmaps, const FViewInfo* View)
 {
 	for (const auto& Clipmap : VirtualShadowMapClipmaps)
 	{
@@ -603,6 +603,16 @@ TSharedPtr<FVirtualShadowMapClipmap> FVisibleLightInfo::FindShadowClipmapForView
 	return TSharedPtr<FVirtualShadowMapClipmap>();
 }
 
+TSharedPtr<FVirtualShadowMapClipmap> FVisibleLightInfo::FindShadowClipmapForView(const FViewInfo* View) const
+{
+	return ::FindShadowClipmapForView(VirtualShadowMapClipmaps, View);
+}
+
+TSharedPtr<FVirtualShadowMapClipmap> FVisibleLightInfo::FindFirstPersonShadowClipmapForView(const FViewInfo* View) const
+{
+	return ::FindShadowClipmapForView(FirstPersonVirtualShadowMapClipmaps, View);
+}
+
 int32 FVisibleLightInfo::GetVirtualShadowMapId(const FViewInfo* View) const
 {
 	if (VirtualShadowMapClipmaps.Num())
diff --git a/Engine/Source/Runtime/Renderer/Private/SceneRendering.h b/Engine/Source/Runtime/Renderer/Private/SceneRendering.h
--- a/Engine/Source/Runtime/Renderer/Private/SceneRendering.h
+++ b/Engine/Source/Runtime/Renderer/Private/SceneRendering.h
@@ -217,13 +217,21 @@ public:
 
 	/** Virtual shadow map clipmap shadows */
 	TArray<TSharedPtr<FVirtualShadowMapClipmap>,SceneRenderingAllocator> VirtualShadowMapClipmaps;
-	
+
+	/** First-person clipmap shadows */
+	TArray<TSharedPtr<FVirtualShadowMapClipmap>,SceneRenderingAllocator> FirstPersonVirtualShadowMapClipmaps;
+
 	/**
 	 * Search VirtualShadowMapClipmaps for the one that was set up with the given view. 
 	 * TODO: should VirtualShadowMapClipmap move to FVisibleLightViewInfo?
 	 */
 	TSharedPtr<FVirtualShadowMapClipmap> FindShadowClipmapForView(const FViewInfo* View) const;
 
+	/**
+	 * Find (if any) first-person special clipmap is present for the given view.
+	 */
+	TSharedPtr<FVirtualShadowMapClipmap> FindFirstPersonShadowClipmapForView(const FViewInfo* View) const;
+
 	/** 
 	* Returns true if the light has only virtual shadow maps (or no shadows at all), i.e. no conventional shadow maps that are allocated
 	*/
diff --git a/Engine/Source/Runtime/Renderer/Private/ShadowRendering.h b/Engine/Source/Runtime/Renderer/Private/ShadowRendering.h
--- a/Engine/Source/Runtime/Renderer/Private/ShadowRendering.h
+++ b/Engine/Source/Runtime/Renderer/Private/ShadowRendering.h
@@ -659,7 +659,7 @@ public:
 	/**
 	 * Adds a primitive to the shadow's subject list.
 	 */
-	bool AddSubjectPrimitive(FPrimitiveSceneInfo* PrimitiveSceneInfo, TArrayView<FViewInfo> ViewArray, bool bRecordShadowSubjectForMobileShading);
+	bool AddSubjectPrimitive(FPrimitiveSceneInfo* PrimitiveSceneInfo, TArrayView<FViewInfo> ViewArray, bool bRecordShadowSubjectForMobileShading, bool bInvertViewShadowRelevance = false);
 
 	uint64 AddSubjectPrimitive_AnyThread(
 		const FPrimitiveSceneInfoCompact& PrimitiveSceneInfoCompact,
diff --git a/Engine/Source/Runtime/Renderer/Private/ShadowSetup.cpp b/Engine/Source/Runtime/Renderer/Private/ShadowSetup.cpp
--- a/Engine/Source/Runtime/Renderer/Private/ShadowSetup.cpp
+++ b/Engine/Source/Runtime/Renderer/Private/ShadowSetup.cpp
@@ -1944,7 +1944,7 @@ bool FProjectedShadowInfo::ShouldDrawStaticMeshes_AnyThread(
 	return bDrawingStaticMeshes;
 }
 
-bool FProjectedShadowInfo::AddSubjectPrimitive(FPrimitiveSceneInfo* PrimitiveSceneInfo, TArrayView<FViewInfo> ViewArray, bool bShouldRecordShadowSubjectsForMobileShading)
+bool FProjectedShadowInfo::AddSubjectPrimitive(FPrimitiveSceneInfo* PrimitiveSceneInfo, TArrayView<FViewInfo> ViewArray, bool bShouldRecordShadowSubjectsForMobileShading, bool bInvertViewShadowRelevance)
 {
 	// Ray traced shadows use the GPU managed distance field object buffers, no CPU culling should be used
 	check(!bRayTracedDistanceField);
@@ -2026,7 +2026,7 @@ bool FProjectedShadowInfo::AddSubjectPrimitive(FPrimitiveSceneInfo* PrimitiveSce
 
 			bOpaque |= ViewRelevance.bOpaque || ViewRelevance.bMasked;
 			bTranslucentRelevance |= ViewRelevance.HasTranslucency();
-			bShadowRelevance |= ViewRelevance.bShadowRelevance;
+			bShadowRelevance |= bInvertViewShadowRelevance ? bool(!ViewRelevance.bShadowRelevance) : bool(ViewRelevance.bShadowRelevance);
 			bDynamicRelevance = bDynamicRelevance || ViewRelevance.bDynamicRelevance;
 			bHeterogeneousVolumeRelevance = ViewRelevance.bHasVolumeMaterialDomain && PrimitiveSceneInfo->Proxy->IsHeterogeneousVolume();
 		}
diff --git a/Engine/Source/Runtime/Renderer/Private/Shadows/ShadowScene.cpp b/Engine/Source/Runtime/Renderer/Private/Shadows/ShadowScene.cpp
--- a/Engine/Source/Runtime/Renderer/Private/Shadows/ShadowScene.cpp
+++ b/Engine/Source/Runtime/Renderer/Private/Shadows/ShadowScene.cpp
@@ -25,6 +25,20 @@ TAutoConsoleVariable<int32> CVarShadowSceneLightActiveFrameCount(
 	ECVF_RenderThreadSafe
 );
 
+static TAutoConsoleVariable<bool> CVarVirtualShadowMapFirstPersonClipmap(
+	TEXT( "r.FirstPerson.Shadow.Virtual.Clipmap" ),
+	true,
+	TEXT( "Enable/Disable support for first-person clipmap for the world-space representation." ),
+	ECVF_Scalability | ECVF_RenderThreadSafe
+);
+
+static bool IsFirstPersonVirtualShadowMapEnabled(EShaderPlatform ShaderPlatform)
+{
+	const bool bCVarEnabled = CVarVirtualShadowMapFirstPersonClipmap.GetValueOnRenderThread() != 0;
+	const bool bGBufferBitSupported = HasFirstPersonGBufferBit(ShaderPlatform);
+	return bCVarEnabled && bGBufferBitSupported;
+}
+
 DECLARE_DWORD_COUNTER_STAT(TEXT("Active Light Count"), STAT_ActiveLightCount, STATGROUP_ShadowRendering);
 
 
@@ -107,6 +121,10 @@ void FShadowScene::PostSceneUpdate(const FScenePreUpdateChangeSet &ScenePreUpdat
 		{
 			AlwaysInvalidatingPrimitives.RemoveSwap(PrimitiveSceneInfo);
 		}
+		if (bEnableVirtualShadowMapFirstPersonClipmap && PrimitiveSceneInfo->Proxy->IsFirstPersonWorldSpaceRepresentation())
+		{
+			FirstPersonWorldSpacePrimitives.RemoveSwap(PrimitiveSceneInfo);
+		}
 	}
 	for (FPrimitiveSceneInfo* PrimitiveSceneInfo : ScenePostUpdateChangeSet.AddedPrimitiveSceneInfos)
 	{
@@ -114,6 +132,29 @@ void FShadowScene::PostSceneUpdate(const FScenePreUpdateChangeSet &ScenePreUpdat
 		{
 			AlwaysInvalidatingPrimitives.Add(PrimitiveSceneInfo);
 		}
+		if (bEnableVirtualShadowMapFirstPersonClipmap && PrimitiveSceneInfo->Proxy->IsFirstPersonWorldSpaceRepresentation())
+		{
+			FirstPersonWorldSpacePrimitives.Add(PrimitiveSceneInfo);
+		}
+	}
+
+	const bool bNewEnabled = IsFirstPersonVirtualShadowMapEnabled(Scene.GetShaderPlatform());
+	// When toggled we must find all primitives marked for this path
+	if (bNewEnabled && !bEnableVirtualShadowMapFirstPersonClipmap)
+	{
+		for (FPrimitiveSceneInfo* PrimitiveSceneInfo : Scene.Primitives)
+		{
+			if (PrimitiveSceneInfo->Proxy->IsFirstPersonWorldSpaceRepresentation())
+			{
+				FirstPersonWorldSpacePrimitives.Add(PrimitiveSceneInfo);
+			}
+		}
+	}
+	bEnableVirtualShadowMapFirstPersonClipmap = bNewEnabled;
+
+	if (!bEnableVirtualShadowMapFirstPersonClipmap)
+	{
+		FirstPersonWorldSpacePrimitives.Empty();
 	}
 }
 
diff --git a/Engine/Source/Runtime/Renderer/Private/Shadows/ShadowScene.h b/Engine/Source/Runtime/Renderer/Private/Shadows/ShadowScene.h
--- a/Engine/Source/Runtime/Renderer/Private/Shadows/ShadowScene.h
+++ b/Engine/Source/Runtime/Renderer/Private/Shadows/ShadowScene.h
@@ -51,7 +51,7 @@ public:
 	TArrayView<FPrimitiveSceneInfo*> GetAlwaysInvalidatingPrimitives() { return AlwaysInvalidatingPrimitives; }
 
 private:
-	//friend class FShadowSceneRenderer;
+	friend class FShadowSceneRenderer;
 
 	inline bool IsActive(int32 LightId) const { return ActiveLights[LightId]; }
 
@@ -82,4 +82,9 @@ private:
 
 	// List of always invalidating primitives, if this gets too popular perhaps a TSet or some such is more appropriate for performance scaling.
 	TArray<FPrimitiveSceneInfo*> AlwaysInvalidatingPrimitives;
+
+	// List of primitives that are marked as "first person world-space" meaning they are to cast shadows only onto the world but not the FP geo itself
+	TArray<FPrimitiveSceneInfo*> FirstPersonWorldSpacePrimitives;
+
+	bool bEnableVirtualShadowMapFirstPersonClipmap = false;
 };
diff --git a/Engine/Source/Runtime/Renderer/Private/Shadows/ShadowSceneRenderer.cpp b/Engine/Source/Runtime/Renderer/Private/Shadows/ShadowSceneRenderer.cpp
--- a/Engine/Source/Runtime/Renderer/Private/Shadows/ShadowSceneRenderer.cpp
+++ b/Engine/Source/Runtime/Renderer/Private/Shadows/ShadowSceneRenderer.cpp
@@ -7,6 +7,7 @@
 #include "ShadowRendering.h"
 #include "VirtualShadowMaps/VirtualShadowMapCacheManager.h"
 #include "VirtualShadowMaps/VirtualShadowMapProjection.h"
+#include "VirtualShadowMaps/VirtualShadowMapClipmap.h"
 #include "SceneCulling/SceneCulling.h"
 #include "Rendering/NaniteStreamingManager.h"
 
@@ -84,6 +85,20 @@ static TAutoConsoleVariable<float> CVarResolutionLodBiasLocalMoving(
 	ECVF_Scalability | ECVF_RenderThreadSafe
 );
 
+static TAutoConsoleVariable<int32> CVarVirtualShadowMapFirstPersonClipmapFirstLevel(
+	TEXT( "r.FirstPerson.Shadow.Virtual.Clipmap.FirstLevel" ),
+	8,
+	TEXT( "First level of the virtual clipmap. Lower values allow higher resolution shadows closer to the camera, but may increase page count." ),
+	ECVF_Scalability | ECVF_RenderThreadSafe
+);
+
+static TAutoConsoleVariable<int32> CVarVirtualShadowMapFirstPersonClipmapLastLevel(
+	TEXT( "r.FirstPerson.Shadow.Virtual.Clipmap.LastLevel" ),
+	18,
+	TEXT( "Last level of the virtual clipmap. Indirectly determines radius the clipmap can cover. Each extra level doubles the maximum range, but may increase page count." ),
+	ECVF_Scalability | ECVF_RenderThreadSafe
+);
+
 extern TAutoConsoleVariable<int32> CVarMarkPixelPagesMipModeLocal;
 
 bool IsVSMOnePassProjectionEnabled(const FEngineShowFlags& ShowFlags)
@@ -259,7 +274,6 @@ FVirtualShadowMapProjectionShaderData FShadowSceneRenderer::GetLocalLightProject
 	return Data;
 }
 
-
 /**
  * Calculate the radius in world-space units of a single pixel at a given depth.
  */
@@ -418,6 +432,53 @@ void FShadowSceneRenderer::AddDirectionalLightShadow(FProjectedShadowInfo* Proje
 {
 	FDirectionalLightShadowFrameSetup& DirectionalLightShadowFrameSetup = DirectionalLights.AddDefaulted_GetRef();
 	DirectionalLightShadowFrameSetup.ProjectedShadowInfo = ProjectedShadowInfo;
+
+	if (!ShadowScene.FirstPersonWorldSpacePrimitives.IsEmpty())
+	{
+		FViewInfo& View = *ProjectedShadowInfo->DependentView;
+		FLightSceneInfo& LightSceneInfo = ProjectedShadowInfo->GetLightSceneInfo();
+		FVisibleLightInfo& VisibleLightInfo = SceneRenderer.VisibleLightInfos[LightSceneInfo.Id];
+
+		// Arbitrary type ID for first person shadow
+		constexpr uint32 FirstPersonShadowTypeId = 123u;
+
+		// Clone the setup from the regular VSM clipmap
+		FVirtualShadowMapClipmapConfig Config = FVirtualShadowMapClipmapConfig::GetGlobal();
+		Config.ShadowTypeId = FirstPersonShadowTypeId;
+		Config.bForceInvalidate = true;
+		Config.FirstCoarseLevel = -1;
+		Config.LastCoarseLevel = -1;
+		Config.FirstLevel = CVarVirtualShadowMapFirstPersonClipmapFirstLevel.GetValueOnRenderThread();
+		Config.LastLevel = CVarVirtualShadowMapFirstPersonClipmapLastLevel.GetValueOnRenderThread();
+		Config.bIsFirstPersonShadow = true;
+
+		TSharedPtr<FVirtualShadowMapClipmap> FPClipmap = TSharedPtr<FVirtualShadowMapClipmap>(new FVirtualShadowMapClipmap(
+			VirtualShadowMapArray,
+			LightSceneInfo,
+			View.ViewMatrices,
+			View.ViewRect.Size(),
+			&View,
+			1.0f,// mobility factor as if moving - this VSM has no persistence, though perhaps it should for HZB?
+			Config
+		));
+
+		FProjectedShadowInfo* FPProjectedShadowInfo = SceneRenderer.Allocator.Create<FProjectedShadowInfo>();
+		FPProjectedShadowInfo->SetupClipmapProjection(&LightSceneInfo, &View, FPClipmap, ProjectedShadowInfo->GetMaxNonFarCascadeDistance());
+		
+		// Need their own array or they'll be found at the wrong times, and we also need them in here to be able to find them during page marking (sigh)
+		VisibleLightInfo.FirstPersonVirtualShadowMapClipmaps.Add(FPClipmap);
+		
+		// This is needed to get it into the line for ending up in the SortedShadowsForShadowDepthPass.VirtualShadowMapShadows which is what drives the shadow rendering.
+		VisibleLightInfo.AllProjectedShadows.Add(FPProjectedShadowInfo);
+		
+		FPProjectedShadowInfo->VirtualShadowMapPerLightCacheEntry = FPClipmap->GetCacheEntry();
+
+		for (FPrimitiveSceneInfo* PrimitiveSceneInfo : ShadowScene.FirstPersonWorldSpacePrimitives)
+		{
+			// _invert_ the shadow relevance because we want to render the thing despite being owner no see, and NOT anyting that is the opposite (owned by other views).
+			FPProjectedShadowInfo->AddSubjectPrimitive(PrimitiveSceneInfo, SceneRenderer.Views, false, true);
+		}
+	}
 }
 
 void FShadowSceneRenderer::PostInitDynamicShadowsSetup()
@@ -709,7 +770,8 @@ void FShadowSceneRenderer::ApplyVirtualShadowMapProjectionForLight(
 						VisibleLightInfo.FindShadowClipmapForView(&View),
 						false, // bModulateRGB
 						nullptr, // TiledVSMProjection
-						OutputScreenShadowMaskTexture);
+						OutputScreenShadowMaskTexture,
+						VisibleLightInfo.FindFirstPersonShadowClipmapForView(&View));				
 				}
 				else if (bShouldUseVirtualShadowMapOnePassProjection)
 				{
@@ -760,7 +822,8 @@ void FShadowSceneRenderer::ApplyVirtualShadowMapProjectionForLight(
 						VisibleLightInfo.FindShadowClipmapForView(&View),
 						false, // bModulateRGB
 						nullptr, // TiledVSMProjection
-						OutputScreenShadowMaskSubPixelTexture);
+						OutputScreenShadowMaskSubPixelTexture,
+						VisibleLightInfo.FindFirstPersonShadowClipmapForView(&View));
 				}
 				else if (bShouldUseVirtualShadowMapOnePassProjection)
 				{
diff --git a/Engine/Source/Runtime/Renderer/Private/SingleLayerWaterRendering.cpp b/Engine/Source/Runtime/Renderer/Private/SingleLayerWaterRendering.cpp
--- a/Engine/Source/Runtime/Renderer/Private/SingleLayerWaterRendering.cpp
+++ b/Engine/Source/Runtime/Renderer/Private/SingleLayerWaterRendering.cpp
@@ -946,7 +946,8 @@ void FDeferredShadingSceneRenderer::RenderSingleLayerWaterReflections(
 						VisibleLightInfo.FindShadowClipmapForView(&View),
 						true, // bModulateRGB
 						bRunTiled  ? &TiledVSMProjection : nullptr,
-						SceneWithoutWaterTextures.SeparatedMainDirLightTexture);
+						SceneWithoutWaterTextures.SeparatedMainDirLightTexture,
+						VisibleLightInfo.FindFirstPersonShadowClipmapForView(&View));
 				}
 			}
 
diff --git a/Engine/Source/Runtime/Renderer/Private/VirtualShadowMaps/VirtualShadowMapArray.cpp b/Engine/Source/Runtime/Renderer/Private/VirtualShadowMaps/VirtualShadowMapArray.cpp
--- a/Engine/Source/Runtime/Renderer/Private/VirtualShadowMaps/VirtualShadowMapArray.cpp
+++ b/Engine/Source/Runtime/Renderer/Private/VirtualShadowMaps/VirtualShadowMapArray.cpp
@@ -129,6 +129,22 @@ static TAutoConsoleVariable<float> CVarPageDilationBorderSizeDirectional(
 	ECVF_RenderThreadSafe
 );
 
+static TAutoConsoleVariable<float> CVarFirstPersonPixelRequestBias(
+	TEXT("r.FirstPerson.Shadow.Virtual.Clipmap.PixelRequestBias"),
+	2.0f,
+	TEXT("Pixels marked with as coming from first person geometry can request a biased resolution as they are not self-shadowing and also very close to the camera (if scaled).\n")
+	TEXT("  Setting to a negative value disables page marking from the FP geometry, which can be used to avoid marking high-res pages for geometry that is scaled to be very small."),
+	ECVF_RenderThreadSafe
+);
+
+static TAutoConsoleVariable<int32> CVarFirstPersonPixelRequestLevelClamp(
+	TEXT("r.FirstPerson.Shadow.Virtual.Clipmap.RequestMinLevelClamp"),
+	8,
+	TEXT("Clamp to avoid high-resolution requests from first-person geometry close to the camera, while still allowing more distant first-person geometry to request full resolution from the environment.\n")
+	TEXT("  Note that this interacts with r.Shadow.Virtual.Clipmap.FirstLevel (as this is what is being clamped) and so may need to be configured in scalability settings."),
+	ECVF_Scalability | ECVF_RenderThreadSafe
+);
+
 static TAutoConsoleVariable<float> CVarPageDilationBorderSizeLocal(
 	TEXT("r.Shadow.Virtual.PageDilationBorderSizeLocal"),
 	0.05f,
@@ -904,6 +920,8 @@ BEGIN_SHADER_PARAMETER_STRUCT(FBaseGeneratePageFlagsParameters, )
 	SHADER_PARAMETER(uint32, bCullBackfacingPixels)
 	SHADER_PARAMETER(uint32, NumDirectionalLightSmInds)
 	SHADER_PARAMETER(uint32, MipModeLocal)
+	SHADER_PARAMETER(float, FirstPersonPixelRequestBias)	
+	SHADER_PARAMETER(uint32, FirstPersonPixelRequestLevelClamp)	
 END_SHADER_PARAMETER_STRUCT()
 
 class FGeneratePageFlagsFromPixelsCS : public FVirtualShadowMapPageManagementShader
@@ -1916,6 +1934,15 @@ void FVirtualShadowMapArray::BuildPageAllocations(
 					DirectionalLightIds.Add(ClipmapID);
 				}
 			}
+
+			for (const TSharedPtr<FVirtualShadowMapClipmap>& Clipmap : VisibleLightInfo.FirstPersonVirtualShadowMapClipmaps)
+			{
+				// mark pages from both views in stereo for the left (primary) view clipmap - assumes GetPrimaryView returns self for non-stereo views.
+				if (Clipmap->GetDependentView() == Views[ViewIndex].GetPrimaryView())
+				{
+					DirectionalLightIds.Add(Clipmap->GetVirtualShadowMapId());
+				}
+			}
 		}	
 		
 		// This view contained no local lights (that were stored in the light grid), and no directional lights, so nothing to do.
@@ -2022,6 +2049,8 @@ void FVirtualShadowMapArray::BuildPageAllocations(
 				BaseParameters.bCullBackfacingPixels = ShouldCullBackfacingPixels() ? 1 : 0;
 				BaseParameters.Substrate = Substrate::BindSubstrateGlobalUniformParameters(View);
 				BaseParameters.MipModeLocal = CVarMarkPixelPagesMipModeLocal.GetValueOnRenderThread();
+				BaseParameters.FirstPersonPixelRequestBias = CVarFirstPersonPixelRequestBias.GetValueOnRenderThread();
+				BaseParameters.FirstPersonPixelRequestLevelClamp = CVarFirstPersonPixelRequestLevelClamp.GetValueOnRenderThread();
 
 				auto GeneratePageFlags = [&](const EVirtualShadowMapProjectionInputType InputType)
 				{
diff --git a/Engine/Source/Runtime/Renderer/Private/VirtualShadowMaps/VirtualShadowMapClipmap.cpp b/Engine/Source/Runtime/Renderer/Private/VirtualShadowMaps/VirtualShadowMapClipmap.cpp
--- a/Engine/Source/Runtime/Renderer/Private/VirtualShadowMaps/VirtualShadowMapClipmap.cpp
+++ b/Engine/Source/Runtime/Renderer/Private/VirtualShadowMaps/VirtualShadowMapClipmap.cpp
@@ -502,6 +502,10 @@ FVirtualShadowMapProjectionShaderData FVirtualShadowMapClipmap::ComputeProjectio
 			Data.Flags |= VSM_PROJ_FLAG_IS_COARSE_CLIP_LEVEL;
 		}
 	}
+	if (Config.bIsFirstPersonShadow)
+	{
+		Data.Flags |= VSM_PROJ_FLAG_IS_FIRST_PERSON_SHADOW;
+	}
 
 	Data.TexelDitherScale = GetLightSceneInfo().Proxy->GetVSMTexelDitherScale();
 	return Data;
diff --git a/Engine/Source/Runtime/Renderer/Private/VirtualShadowMaps/VirtualShadowMapClipmap.h b/Engine/Source/Runtime/Renderer/Private/VirtualShadowMaps/VirtualShadowMapClipmap.h
--- a/Engine/Source/Runtime/Renderer/Private/VirtualShadowMaps/VirtualShadowMapClipmap.h
+++ b/Engine/Source/Runtime/Renderer/Private/VirtualShadowMaps/VirtualShadowMapClipmap.h
@@ -19,8 +19,9 @@ class FVirtualShadowMapArray;
 class FVirtualShadowMapArrayCacheManager;
 class FVirtualShadowMapPerLightCacheEntry;
 
-struct FVirtualShadowMapClipmapConfig
+class FVirtualShadowMapClipmapConfig
 {
+public:
 	static FVirtualShadowMapClipmapConfig GetGlobal();
 
 	int32 FirstLevel = 8;
@@ -29,6 +30,7 @@ struct FVirtualShadowMapClipmapConfig
 	int32 LastCoarseLevel = -1;
 	uint32 ShadowTypeId = 0u;
 	bool bForceInvalidate = false;
+	bool bIsFirstPersonShadow = false;
 };
 
 class FVirtualShadowMapClipmap : FRefCountedObject
diff --git a/Engine/Source/Runtime/Renderer/Private/VirtualShadowMaps/VirtualShadowMapProjection.cpp b/Engine/Source/Runtime/Renderer/Private/VirtualShadowMaps/VirtualShadowMapProjection.cpp
--- a/Engine/Source/Runtime/Renderer/Private/VirtualShadowMaps/VirtualShadowMapProjection.cpp
+++ b/Engine/Source/Runtime/Renderer/Private/VirtualShadowMaps/VirtualShadowMapProjection.cpp
@@ -103,6 +103,7 @@ class FVirtualShadowMapProjectionCS : public FGlobalShader
 	class FVisualizeOutputDim		: SHADER_PERMUTATION_BOOL("VISUALIZE_OUTPUT");
 	class FExtrapolateSlopeDim		: SHADER_PERMUTATION_BOOL("SMRT_EXTRAPOLATE_SLOPE");
 	class FUseTileList				: SHADER_PERMUTATION_BOOL("USE_TILE_LIST");
+	class FFirstPersonShadow		: SHADER_PERMUTATION_BOOL("FIRST_PERSON_SHADOW");
 	// -1 means dynamic count
 	class FSMRTStaticSampleCount	: SHADER_PERMUTATION_RANGE_INT("SMRT_TEMPLATE_STATIC_SAMPLES_PER_RAY", -1, 2);
 	// Used for A/B testing a change that affects reg allocation, etc.
@@ -115,6 +116,7 @@ class FVirtualShadowMapProjectionCS : public FGlobalShader
 		FVisualizeOutputDim,
 		FExtrapolateSlopeDim,
 		FUseTileList,
+		FFirstPersonShadow,
 		FSMRTStaticSampleCount
 #if MAX_TEST_PERMUTATION > 0
 		, FTestDim
@@ -140,6 +142,8 @@ class FVirtualShadowMapProjectionCS : public FGlobalShader
 		SHADER_PARAMETER_STRUCT(FLightShaderParameters, Light)
 		SHADER_PARAMETER(int32, LightUniformVirtualShadowMapId)
 		SHADER_PARAMETER_RDG_TEXTURE_UAV(RWTexture2D, OutShadowFactor)
+		// First Person shadow parameters
+		SHADER_PARAMETER(int32, FirstPersonVirtualShadowMapId)		
 		// Visualization output
 		SHADER_PARAMETER_RDG_BUFFER_SRV(StructuredBuffer< FPhysicalPageMetaData >, PhysicalPageMetaData)
 		SHADER_PARAMETER(int32, VisualizeModeId)
@@ -183,6 +187,12 @@ class FVirtualShadowMapProjectionCS : public FGlobalShader
 			return false;
 		}
 		
+		// Only need the first person permutation for directional lights.
+		if( PermutationVector.Get< FFirstPersonShadow >() && !PermutationVector.Get< FDirectionalLightDim >() )
+		{
+			return false;
+		}
+
 		return DoesPlatformSupportVirtualShadowMaps(Parameters.Platform);
 	}
 
@@ -207,9 +217,10 @@ static void RenderVirtualShadowMapProjectionCommon(
 	const FIntRect ProjectionRect,
 	EVirtualShadowMapProjectionInputType InputType,
 	FRDGTextureRef OutputTexture,
-	const FLightSceneProxy* LightProxy = nullptr,
-	int32 VirtualShadowMapId = INDEX_NONE,
-	FTiledVSMProjection* TiledVSMProjection = nullptr)
+	const FLightSceneProxy* LightProxy,
+	int32 VirtualShadowMapId,
+	FTiledVSMProjection* TiledVSMProjection,
+	int32 FirstPersonVirtualShadowMapId)
 {
 	check(GRHISupportsWaveOperations);
 
@@ -241,6 +252,8 @@ static void RenderVirtualShadowMapProjectionCommon(
 
 	FBlueNoise BlueNoise = GetBlueNoiseGlobalParameters();
 	PassParameters->BlueNoise = CreateUniformBufferImmediate(BlueNoise, EUniformBufferUsage::UniformBuffer_SingleDraw);
+	PassParameters->FirstPersonVirtualShadowMapId = FirstPersonVirtualShadowMapId;
+	bool bHasFirstPersonShadow = FirstPersonVirtualShadowMapId != INDEX_NONE;
 
 	bool bDirectionalLight = false;
 	bool bOnePassProjection = LightProxy == nullptr;
@@ -292,6 +305,12 @@ static void RenderVirtualShadowMapProjectionCommon(
 	PermutationVector.Set< FVirtualShadowMapProjectionCS::FExtrapolateSlopeDim >(ExtrapolateMaxSlope > 0.0f);
 	PermutationVector.Set< FVirtualShadowMapProjectionCS::FUseTileList >(bUseTileList);
 	PermutationVector.Set< FVirtualShadowMapProjectionCS::FSMRTStaticSampleCount >(StaticSamplesPerRay);
+	
+	if (bHasFirstPersonShadow && !ensureMsgf(bDirectionalLight, TEXT("First person shadow can only be used with pass-per light direction shadow projection and will be disabled.")))
+	{
+		bHasFirstPersonShadow = false;
+	}
+	PermutationVector.Set< FVirtualShadowMapProjectionCS::FFirstPersonShadow >(bHasFirstPersonShadow);
 #if MAX_TEST_PERMUTATION > 0
 	{
 		int32 TestPermutation = FMath::Clamp(CVarTestPermutation.GetValueOnRenderThread(), 0, MAX_TEST_PERMUTATION);
@@ -363,7 +382,11 @@ void RenderVirtualShadowMapProjectionOnePass(
 		VirtualShadowMapArray,
 		ProjectionRect,
 		InputType,
-		ShadowMaskBits);
+		ShadowMaskBits,
+		nullptr,
+		INDEX_NONE,
+		nullptr,
+		INDEX_NONE);
 }
 
 static FRDGTextureRef CreateShadowMaskTexture(FRDGBuilder& GraphBuilder, FIntPoint Extent)
@@ -409,7 +432,9 @@ void RenderVirtualShadowMapProjection(
 		InputType,
 		VirtualShadowMaskTexture,
 		LightSceneInfo.Proxy,
-		VirtualShadowMapId);
+		VirtualShadowMapId,
+		nullptr,
+		INDEX_NONE);
 
 	CompositeVirtualShadowMapMask(
 		GraphBuilder,
@@ -432,7 +457,8 @@ void RenderVirtualShadowMapProjection(
 	const TSharedPtr<FVirtualShadowMapClipmap>& Clipmap,
 	bool bModulateRGB,
 	FTiledVSMProjection* TiledVSMProjection,
-	FRDGTextureRef OutputShadowMaskTexture)
+	FRDGTextureRef OutputShadowMaskTexture,
+	const TSharedPtr<FVirtualShadowMapClipmap>& FirstPersonClipmap)
 {
 	FRDGTextureRef VirtualShadowMaskTexture = CreateShadowMaskTexture(GraphBuilder, View.ViewRect.Max);
 
@@ -446,7 +472,8 @@ void RenderVirtualShadowMapProjection(
 		VirtualShadowMaskTexture,
 		Clipmap->GetLightSceneInfo().Proxy,
 		Clipmap->GetVirtualShadowMapId(),
-		TiledVSMProjection);
+		TiledVSMProjection,
+		FirstPersonClipmap.IsValid() ? FirstPersonClipmap->GetVirtualShadowMapId() : INDEX_NONE);
 	
 	CompositeVirtualShadowMapMask(
 		GraphBuilder,
diff --git a/Engine/Source/Runtime/Renderer/Private/VirtualShadowMaps/VirtualShadowMapProjection.h b/Engine/Source/Runtime/Renderer/Private/VirtualShadowMaps/VirtualShadowMapProjection.h
--- a/Engine/Source/Runtime/Renderer/Private/VirtualShadowMaps/VirtualShadowMapProjection.h
+++ b/Engine/Source/Runtime/Renderer/Private/VirtualShadowMaps/VirtualShadowMapProjection.h
@@ -44,7 +44,8 @@ void RenderVirtualShadowMapProjection(
 	const TSharedPtr<FVirtualShadowMapClipmap>& Clipmap,
 	bool bModulateRGB,
 	FTiledVSMProjection* TiledVSMProjection,
-	FRDGTextureRef OutputShadowMaskTexture);
+	FRDGTextureRef OutputShadowMaskTexture,
+	const TSharedPtr<FVirtualShadowMapClipmap>& FirstPersonClipmap);
 
 void RenderVirtualShadowMapProjection(
 	FRDGBuilder& GraphBuilder,
