Fix IES manager force update not working correctly.

diff --git a/Engine/Source/Runtime/Renderer/Private/IESTextureManager.cpp b/Engine/Source/Runtime/Renderer/Private/IESTextureManager.cpp
--- a/Engine/Source/Runtime/Renderer/Private/IESTextureManager.cpp
+++ b/Engine/Source/Runtime/Renderer/Private/IESTextureManager.cpp
@@ -498,11 +498,12 @@ void UpdateAtlasTexture(FRDGBuilder& GraphBuilder, const FStaticShaderPlatform S
 	static uint32 CachedAtlasResolution = CVarIESTextureResolution.GetValueOnRenderThread();
 	static uint32 CachedAtlasMaxSlice = CVarIESTextureMaxProfileCount.GetValueOnRenderThread();
 	const bool bHasSettingChanged = CachedAtlasResolution != CVarIESTextureResolution.GetValueOnRenderThread() || CachedAtlasMaxSlice != CVarIESTextureMaxProfileCount.GetValueOnRenderThread();
+	const bool bForceUpdate = CVarIESForceUpdate.GetValueOnRenderThread() > 0;
 
 	// 1. Determine if a new atlas texture allocation is needed: If there are new entries, ensure it can fit into the atlas
 	bool bForceRecreate = false;
 	uint32 RequestedSliceCount = 0;
-	if (GIESTextureManager.bHasPendingAdds)
+	if (GIESTextureManager.bHasPendingAdds || bForceUpdate)
 	{
 		RequestedSliceCount = FMath::Min(CachedAtlasMaxSlice, GetRequestedSlotCount(GIESTextureManager));
 		bForceRecreate = (GIESTextureManager.AtlasTexture == nullptr) || (GIESTextureManager.RequestedSliceCountWhenTextureAllocated < RequestedSliceCount);
@@ -512,8 +513,7 @@ void UpdateAtlasTexture(FRDGBuilder& GraphBuilder, const FStaticShaderPlatform S
 	// 2. Invalid slots 
 	// * If a global update is requested, invalid all slots
 	// * Otherwise invalid slots which have new texture data (e.g., streamed in)
-	const bool bForceUpdate = CVarIESForceUpdate.GetValueOnRenderThread() > 0 || bForceRecreate;
-	if (bForceUpdate)
+	if (bForceUpdate || bForceRecreate)
 	{
 		CachedAtlasResolution = CVarIESTextureResolution.GetValueOnRenderThread();
 		CachedAtlasMaxSlice = CVarIESTextureMaxProfileCount.GetValueOnRenderThread();
