Small refactor of GatherLightsAndComputeLightGrid(...).
- Renamed to PrepareForwardLightData(...).
- Moved to FSceneRenderer so it can be used by both deferred and mobile renderers.
- Moved management of ForwardLightUniformBuffer from ComputeLightGrid(...) to PrepareForwardLightData(...).

diff --git a/Engine/Source/Runtime/Renderer/Private/DeferredShadingRenderer.cpp b/Engine/Source/Runtime/Renderer/Private/DeferredShadingRenderer.cpp
--- a/Engine/Source/Runtime/Renderer/Private/DeferredShadingRenderer.cpp
+++ b/Engine/Source/Runtime/Renderer/Private/DeferredShadingRenderer.cpp
@@ -2260,7 +2260,18 @@ void FDeferredShadingSceneRenderer::Render(FRDGBuilder& GraphBuilder)
 			RDG_GPU_STAT_SCOPE(GraphBuilder, SortLights);
 
 			GatherAndSortLightsTask.Wait();
-			ComputeLightGridOutput = GatherLightsAndComputeLightGrid(GraphBuilder, /*bComputeLightGrid*/ ViewFamily.EngineShowFlags.Lighting && !ViewFamily.EngineShowFlags.PathTracing, *SortedLightSet);
+
+			if (ViewFamily.EngineShowFlags.Lighting && !ViewFamily.EngineShowFlags.PathTracing)
+			{
+				ComputeLightGridOutput = PrepareForwardLightData(GraphBuilder, true, *SortedLightSet);
+
+				// Store this flag if lights are injected in the grids, check with 'AreLightsInLightGrid()'
+				bAreLightsInLightGrid = true;
+			}
+			else
+			{
+				SetDummyForwardLightUniformBufferOnViews(GraphBuilder, ShaderPlatform, Views);
+			}
 
 			CSV_CUSTOM_STAT(LightCount, All,  float(SortedLightSet->SortedLights.Num()), ECsvCustomStatOp::Set);
 			CSV_CUSTOM_STAT(LightCount, Batched, float(SortedLightSet->UnbatchedLightStart), ECsvCustomStatOp::Set);
diff --git a/Engine/Source/Runtime/Renderer/Private/DeferredShadingRenderer.h b/Engine/Source/Runtime/Renderer/Private/DeferredShadingRenderer.h
--- a/Engine/Source/Runtime/Renderer/Private/DeferredShadingRenderer.h
+++ b/Engine/Source/Runtime/Renderer/Private/DeferredShadingRenderer.h
@@ -336,12 +336,6 @@ public:
 	 */
 	bool RenderPrePassViewParallel(const FViewInfo& View, FRHICommandListImmediate& ParentCmdList,TFunctionRef<void()> AfterTasksAreStarted, bool bDoPrePre);
 
-	/** 
-	 * Culls local lights and reflection probes to a grid in frustum space, builds one light list and grid per view in the current Views.  
-	 * Needed for forward shading or translucency using the Surface lighting mode, and clustered deferred shading. 
-	 */
-	FComputeLightGridOutput GatherLightsAndComputeLightGrid(FRDGBuilder& GraphBuilder, bool bNeedLightGrid, const FSortedLightSetSceneInfo &SortedLightSet);
-
 	/** 
 	 * Debug light grid content on screen.
 	 */
diff --git a/Engine/Source/Runtime/Renderer/Private/LightGridInjection.cpp b/Engine/Source/Runtime/Renderer/Private/LightGridInjection.cpp
--- a/Engine/Source/Runtime/Renderer/Private/LightGridInjection.cpp
+++ b/Engine/Source/Runtime/Renderer/Private/LightGridInjection.cpp
@@ -716,7 +716,7 @@ FLightGrid LightGridInjection(
 	return Output;
 }
 
-FComputeLightGridOutput FSceneRenderer::ComputeLightGrid(FRDGBuilder& GraphBuilder, bool bCullLightsToGrid, const FSortedLightSetSceneInfo& SortedLightSet)
+FComputeLightGridOutput FSceneRenderer::ComputeLightGrid(FRDGBuilder& GraphBuilder, bool bCullLightsToGrid, const FSortedLightSetSceneInfo& SortedLightSet, TArray<FForwardLightData*, TInlineAllocator<2>>& PerViewForwardLightData)
 {
 	FComputeLightGridOutput Result = {};
 
@@ -740,7 +740,7 @@ FComputeLightGridOutput FSceneRenderer::ComputeLightGrid(FRDGBuilder& GraphBuild
 
 		View.ForwardLightingResources.SelectedForwardDirectionalLightProxy = nullptr;
 
-		FForwardLightData* ForwardLightData = GraphBuilder.AllocParameters<FForwardLightData>();
+		FForwardLightData* ForwardLightData = PerViewForwardLightData[ViewIndex];
 		ForwardLightData->DirectionalLightShadowmapAtlas = SystemTextures.Black;
 		ForwardLightData->DirectionalLightStaticShadowmap = GBlackTexture->TextureRHI;
 
@@ -1227,7 +1227,6 @@ FComputeLightGridOutput FSceneRenderer::ComputeLightGrid(FRDGBuilder& GraphBuild
 		ForwardLightData->CulledLightDataGrid32Bit = LightGrid.CulledLightDataGridSRV;
 		ForwardLightData->CulledLightDataGrid16Bit = LightGrid.CulledLightDataGridSRV;
 		ForwardLightData->NumCulledLightsGrid = LightGrid.NumCulledLightsGridSRV;
-		View.ForwardLightingResources.SetUniformBuffer(GraphBuilder.CreateUniformBuffer(ForwardLightData));
 	}
 
 #if WITH_EDITOR
@@ -1244,21 +1243,25 @@ FComputeLightGridOutput FSceneRenderer::ComputeLightGrid(FRDGBuilder& GraphBuild
 	return Result;
 }
 
-FComputeLightGridOutput FDeferredShadingSceneRenderer::GatherLightsAndComputeLightGrid(FRDGBuilder& GraphBuilder, bool bNeedLightGrid, const FSortedLightSetSceneInfo& SortedLightSet)
+FComputeLightGridOutput FSceneRenderer::PrepareForwardLightData(FRDGBuilder& GraphBuilder, bool bCullLightsToGrid, const FSortedLightSetSceneInfo& SortedLightSet)
 {
-	SCOPED_NAMED_EVENT(GatherLightsAndComputeLightGrid, FColor::Emerald);
-	FComputeLightGridOutput Result = {};
+	SCOPED_NAMED_EVENT(PrepareForwardLightData, FColor::Emerald);
+
+	TArray<FForwardLightData*, TInlineAllocator<2>> PerViewForwardLightData;
+	PerViewForwardLightData.Reserve(Views.Num());
 
-	if (!bNeedLightGrid)
+	for (int32 ViewIndex = 0; ViewIndex < Views.Num(); ++ViewIndex)
 	{
-		SetDummyForwardLightUniformBufferOnViews(GraphBuilder, ShaderPlatform, Views);
-		return Result;
+		PerViewForwardLightData.Add(GraphBuilder.AllocParameters<FForwardLightData>());
 	}
 
-	Result = ComputeLightGrid(GraphBuilder, true, SortedLightSet);
+	FComputeLightGridOutput Result = ComputeLightGrid(GraphBuilder, bCullLightsToGrid, SortedLightSet, PerViewForwardLightData);
 
-	// Store this flag if lights are injected in the grids, check with 'AreLightsInLightGrid()'
-	bAreLightsInLightGrid = true;
+	for (int32 ViewIndex = 0; ViewIndex < Views.Num(); ++ViewIndex)
+	{
+		FViewInfo& View = Views[ViewIndex];
+		View.ForwardLightingResources.SetUniformBuffer(GraphBuilder.CreateUniformBuffer(PerViewForwardLightData[ViewIndex]));
+	}
 
 	return Result;
 }
diff --git a/Engine/Source/Runtime/Renderer/Private/MobileShadingRenderer.cpp b/Engine/Source/Runtime/Renderer/Private/MobileShadingRenderer.cpp
--- a/Engine/Source/Runtime/Renderer/Private/MobileShadingRenderer.cpp
+++ b/Engine/Source/Runtime/Renderer/Private/MobileShadingRenderer.cpp
@@ -1136,12 +1136,10 @@ void FMobileSceneRenderer::Render(FRDGBuilder& GraphBuilder)
 			UpdateLightFunctionAtlasTaskFunction();
 
 			GatherAndSortLights(SortedLightSet, bShadowedLightsInClustered);
-			int32 NumReflectionCaptures = Views[0].NumBoxReflectionCaptures + Views[0].NumSphereReflectionCaptures;
-			bool bCullLightsToGrid = (((bEnableClusteredReflections || bDeferredShading) && NumReflectionCaptures > 0) || bEnableClusteredLocalLights);
-			if (bCullLightsToGrid)
-			{
-				ComputeLightGrid(GraphBuilder, bEnableClusteredLocalLights, SortedLightSet);
-			}
+
+			const int32 NumReflectionCaptures = Views[0].NumBoxReflectionCaptures + Views[0].NumSphereReflectionCaptures;
+			const bool bCullLightsToGrid = (((bEnableClusteredReflections || bDeferredShading) && NumReflectionCaptures > 0) || bEnableClusteredLocalLights);
+			PrepareForwardLightData(GraphBuilder, bCullLightsToGrid, SortedLightSet);
 
 			LightFunctionAtlas.RenderLightFunctionAtlas(GraphBuilder, Views);
 		}
diff --git a/Engine/Source/Runtime/Renderer/Private/SceneRendering.h b/Engine/Source/Runtime/Renderer/Private/SceneRendering.h
--- a/Engine/Source/Runtime/Renderer/Private/SceneRendering.h
+++ b/Engine/Source/Runtime/Renderer/Private/SceneRendering.h
@@ -2528,7 +2528,12 @@ protected:
 	 * Culls local lights and reflection probes to a grid in frustum space, builds one light list and grid per view in the current Views.  
 	 * Needed for forward shading or translucency using the Surface lighting mode, and clustered deferred shading. 
 	 */
-	FComputeLightGridOutput ComputeLightGrid(FRDGBuilder& GraphBuilder, bool bCullLightsToGrid, const FSortedLightSetSceneInfo& SortedLightSet);
+	FComputeLightGridOutput ComputeLightGrid(FRDGBuilder& GraphBuilder, bool bCullLightsToGrid, const FSortedLightSetSceneInfo& SortedLightSet, TArray<FForwardLightData*, TInlineAllocator<2>>& PerViewForwardLightData);
+
+	/**
+	 * Prepare ForwardLightingResources for the Views (including building light grids)
+	 */
+	FComputeLightGridOutput PrepareForwardLightData(FRDGBuilder& GraphBuilder, bool bCullLightsToGrid, const FSortedLightSetSceneInfo& SortedLightSet);
 
 	/**
 	* Used by RenderLights to figure out if light functions need to be rendered to the attenuation buffer.
