PR #12448: Fix GTAO console command minor issues

diff --git a/Engine/Source/Runtime/Renderer/Private/CompositionLighting/CompositionLighting.cpp b/Engine/Source/Runtime/Renderer/Private/CompositionLighting/CompositionLighting.cpp
--- a/Engine/Source/Runtime/Renderer/Private/CompositionLighting/CompositionLighting.cpp
+++ b/Engine/Source/Runtime/Renderer/Private/CompositionLighting/CompositionLighting.cpp
@@ -50,14 +50,6 @@ static TAutoConsoleVariable<int32> CVarGTAOSpatialFilter(
 	TEXT("1: On (default)\n "),
 	ECVF_RenderThreadSafe | ECVF_Scalability);
 
-static TAutoConsoleVariable<int32> CVarGTAOCombined(
-	TEXT("r.GTAO.Combined"),
-	1,
-	TEXT("Enable Spatial Filter for GTAO \n ")
-	TEXT("0: Off \n ")
-	TEXT("1: On (default)\n "),
-	ECVF_RenderThreadSafe | ECVF_Scalability);
-
 bool IsAmbientCubemapPassRequired(const FSceneView& View)
 {
 	return View.FinalPostProcessSettings.ContributingCubemaps.Num() != 0 && IsUsingGBuffers(View.GetShaderPlatform());
diff --git a/Engine/Source/Runtime/Renderer/Private/CompositionLighting/PostProcessAmbientOcclusion.cpp b/Engine/Source/Runtime/Renderer/Private/CompositionLighting/PostProcessAmbientOcclusion.cpp
--- a/Engine/Source/Runtime/Renderer/Private/CompositionLighting/PostProcessAmbientOcclusion.cpp
+++ b/Engine/Source/Runtime/Renderer/Private/CompositionLighting/PostProcessAmbientOcclusion.cpp
@@ -99,8 +99,8 @@ static TAutoConsoleVariable<int32> CVarGTAOUseNormals(
 	TEXT("r.GTAO.UseNormals"),
 	0,
 	TEXT("Whether to use GBuffer Normals or Depth Derived normals \n ")
-	TEXT("0: Off \n ")
-	TEXT("1: On (default)\n "),
+	TEXT("0: Off (default) \n ")
+	TEXT("1: On\n "),
 	ECVF_RenderThreadSafe | ECVF_Scalability);
 
 static TAutoConsoleVariable<int32> CVarGTAOFilterWidth(
