Fixed: Broken DOF on mobile platforms
SceneDepth target was not stored to a system memory in a deferred shading case

diff --git a/Engine/Source/Runtime/Renderer/Private/MobileShadingRenderer.cpp b/Engine/Source/Runtime/Renderer/Private/MobileShadingRenderer.cpp
--- a/Engine/Source/Runtime/Renderer/Private/MobileShadingRenderer.cpp
+++ b/Engine/Source/Runtime/Renderer/Private/MobileShadingRenderer.cpp
@@ -206,6 +206,12 @@ END_SHADER_PARAMETER_STRUCT()
 
 static bool PostProcessUsesSceneDepth(const FViewInfo& View)
 {
+	if ((View.FinalPostProcessSettings.DepthOfFieldScale > 0.0f && View.Family->EngineShowFlags.DepthOfField)
+		|| View.MobileLightShaft.IsSet())
+	{
+		return true;
+	}
+	
 	// Find out whether post-process materials use CustomDepth/Stencil lookups
 	const FBlendableManager& BlendableManager = View.FinalPostProcessSettings.BlendableManager;
 	FBlendableEntry* BlendableIt = nullptr;
