[GPU Scene Lights] Store view-depedent light properties in buffers indexed by LightSceneId matching GPUSceneLightData order.

diff --git a/Engine/Shaders/Private/GeneratedUniformBufferTypes.ush b/Engine/Shaders/Private/GeneratedUniformBufferTypes.ush
--- a/Engine/Shaders/Private/GeneratedUniformBufferTypes.ush
+++ b/Engine/Shaders/Private/GeneratedUniformBufferTypes.ush
@@ -13,4 +13,5 @@
 #include "HairStrands/HairStrandsDeepShadowCommonStruct.ush"
 
 #include "/Engine/Shared/LightSceneData.h"
+#include "/Engine/Shared/LightViewData.h"
 #include "/Engine/Shared/RayTracingTypes.h"
\ No newline at end of file
diff --git a/Engine/Shaders/Private/LightData.ush b/Engine/Shaders/Private/LightData.ush
--- a/Engine/Shaders/Private/LightData.ush
+++ b/Engine/Shaders/Private/LightData.ush
@@ -4,6 +4,7 @@
 
 #include "Common.ush"
 #include "HairShadingCommon.ush" // Used for FHairTransmittanceData
+#include "/Engine/Shared/LightViewData.h"
 #include "/Engine/Shared/LightDefinitions.h"
 
 ///////////////////////////////////////////////////////////////////////////////////////////////////
@@ -504,15 +505,15 @@ FDeferredLightData ConvertToDeferredLight(
 }
 
 // Helper function for converting FLightSceneData into FDeferredLightData
-FDeferredLightData ConvertToDeferredLight(FLightSceneData In)
+FDeferredLightData ConvertToDeferredLight(FLightSceneData In, FLightViewData InViewData)
 {
 	FDeferredLightData Out = (FDeferredLightData)0;
 
 	const uint LightType = UnpackLightType(In.LightTypeAndShadowMapChannelMaskPacked);
 
-	Out.TranslatedWorldPosition = DFFastAddDemote(In.WorldPosition, PrimaryView.PreViewTranslation);
+	Out.TranslatedWorldPosition = InViewData.TranslatedWorldPosition;
 	Out.InvRadius				= In.InvRadius;
-	Out.Color					= In.Color.xyz;
+	Out.Color					= In.Color.xyz * InViewData.ColorScale;
 	Out.FalloffExponent			= In.FalloffExponent;
 	Out.Direction				= In.Direction;
 	Out.SpotAngles				= In.SpotAngles;
diff --git a/Engine/Shaders/Private/LightGridCommon.ush b/Engine/Shaders/Private/LightGridCommon.ush
--- a/Engine/Shaders/Private/LightGridCommon.ush
+++ b/Engine/Shaders/Private/LightGridCommon.ush
@@ -365,6 +365,28 @@ FLocalLightData GetLocalLightData(uint GridIndex, uint EyeIndex)
 	return Out;
 }
 
+FLightViewData GetLightViewDataNonStereo(int LightId)
+{
+	return ForwardLightData.LightViewData[LightId];
+}
+
+FLightViewData GetLightViewData(int LightId, uint EyeIndex)
+{
+#if (INSTANCED_STEREO || MOBILE_MULTI_VIEW)
+	BRANCH
+	if (EyeIndex == 0)
+	{
+#endif
+		return ForwardLightData.LightViewData[LightId];
+#if (INSTANCED_STEREO || MOBILE_MULTI_VIEW)
+	}
+	else
+	{
+		return ForwardLightDataISR.LightViewData[LightId];
+	}
+#endif
+}
+
 /**
  * Helpers to pack/unpack the shadow mask for cluster shading and virtual shadow map
  * Currently hard-coded for 4 bits per light, up to 32 lights per pixel in a uint4
diff --git a/Engine/Shaders/Shared/LightViewData.h b/Engine/Shaders/Shared/LightViewData.h
new file mode 100644
--- /dev/null
+++ b/Engine/Shaders/Shared/LightViewData.h
@@ -0,0 +1,27 @@
+// Copyright Epic Games, Inc. All Rights Reserved.
+
+#pragma once
+
+#ifdef __cplusplus
+#include "HLSLTypeAliases.h"
+
+namespace UE::HLSL
+{
+#endif
+
+/**
+ * Light view-dependent data
+ */
+struct FLightViewData
+{
+	// TODO: pack members to reduce memory overhead/bandwidth
+	float3 TranslatedWorldPosition;
+	float ColorScale;
+	float VolumetricScatteringIntensity;
+	uint VirtualShadowMapId;
+};
+
+#ifdef __cplusplus
+} // namespace
+using FLightViewData = UE::HLSL::FLightViewData;
+#endif
diff --git a/Engine/Source/Runtime/Renderer/Private/GPUScene.h b/Engine/Source/Runtime/Renderer/Private/GPUScene.h
--- a/Engine/Source/Runtime/Renderer/Private/GPUScene.h
+++ b/Engine/Source/Runtime/Renderer/Private/GPUScene.h
@@ -264,6 +264,7 @@ public:
 
 	int32 GetNumInstances() const { return InstanceSceneDataAllocator.GetMaxSize(); }
 	int32 GetNumPrimitives() const { return DynamicPrimitivesOffset; }
+	int32 GetMaxLightId() const { return NumGPULights; }
 	int32 GetNumLightmapDataItems() const { return LightmapDataAllocator.GetMaxSize(); }
 	/**
 	 * Returns the highest instance ID that is represented in the GPU scene (which may be lower than the host allocated IDs due to various limits)
diff --git a/Engine/Source/Runtime/Renderer/Private/LightGridInjection.cpp b/Engine/Source/Runtime/Renderer/Private/LightGridInjection.cpp
--- a/Engine/Source/Runtime/Renderer/Private/LightGridInjection.cpp
+++ b/Engine/Source/Runtime/Renderer/Private/LightGridInjection.cpp
@@ -38,6 +38,7 @@ LightGridInjection.cpp
 #include "MegaLights/MegaLights.h"
 #include "LightGridDefinitions.h"
 #include "VolumetricFog.h"
+#include "LightViewData.h"
 
 int32 GLightGridPixelSize = 64;
 FAutoConsoleVariableRef CVarLightGridPixelSize(
@@ -1255,6 +1256,83 @@ FComputeLightGridOutput FSceneRenderer::PrepareForwardLightData(FRDGBuilder& Gra
 		PerViewForwardLightData.Add(GraphBuilder.AllocParameters<FForwardLightData>());
 	}
 
+	// TODO: Add simple lights to GPU Scene Lights
+
+	// Build light view data buffers
+	const bool bRenderRectLightsAsSpotLights = RenderRectLightsAsSpotLights(FeatureLevel);
+	const uint32 LightShaderParameterFlags = bRenderRectLightsAsSpotLights ? ELightShaderParameterFlags::RectAsSpotLight : 0u;
+
+	UE::Tasks::FTask PrerequisiteTask; // TODO: should match prerequisite of FGPUScene::UpdateGPULights(...), currently is null
+
+	const int32 MaxLightId = Scene->GPUScene.GetMaxLightId();
+
+	for (int32 ViewIndex = 0; ViewIndex < Views.Num(); ++ViewIndex)
+	{
+		const FViewInfo& View = Views[ViewIndex];
+
+		const float Exposure = View.GetLastEyeAdaptationExposure();
+
+		FRDGUploadData<FLightViewData> LightViewData(GraphBuilder, FMath::Max(1, MaxLightId));
+
+		GraphBuilder.AddSetupTask([this, &View, &SortedLightSet, LightViewData, LightShaderParameterFlags, Exposure]
+			{
+				SCOPED_NAMED_EVENT(PrepareLightViewData, FColor::Green);
+
+				for (int32 SortedIndex = SortedLightSet.SimpleLightsEnd; SortedIndex < SortedLightSet.SortedLights.Num(); ++SortedIndex)
+				{
+					const FSortedLightSceneInfo& SortedLightInfo = SortedLightSet.SortedLights[SortedIndex];
+					const FLightSceneInfo* const LightSceneInfo = SortedLightInfo.LightSceneInfo;
+					const FLightSceneProxy* LightProxy = LightSceneInfo->Proxy;
+
+					const FLightSceneInfo::FPersistentId LightSceneId = LightSceneInfo->Id;
+
+					if (!ensureMsgf(Scene->Lights.IsAllocated(LightSceneId), TEXT("Visible light is missing from GPU Scene.")))
+					{
+						continue;
+					}
+
+					if (!ensureMsgf(VisibleLightInfos.IsValidIndex(LightSceneId), TEXT("Visible light doesn't have valid info.")))
+					{
+						continue;
+					}
+
+					if (LightSceneInfo->ShouldRenderLight(View))
+					{
+						FLightRenderParameters LightParameters;
+						LightProxy->GetLightShaderParameters(LightParameters, LightShaderParameterFlags);
+
+						const float LightFade = GetLightFadeFactor(View, LightProxy);
+
+						float VolumetricScatteringIntensity = LightProxy->GetVolumetricScatteringIntensity();
+						if (LightNeedsSeparateInjectionIntoVolumetricFogForOpaqueShadow(View, LightSceneInfo, VisibleLightInfos[LightSceneId], *Scene))
+						{
+							// Disable this lights forward shading volumetric scattering contribution
+							VolumetricScatteringIntensity = 0;
+						}
+
+						const int32 VirtualShadowMapId = ViewFamily.EngineShowFlags.DynamicShadows ? VisibleLightInfos[LightSceneId].GetVirtualShadowMapId(&View) : INDEX_NONE;
+
+						LightViewData[LightSceneId].TranslatedWorldPosition = FVector3f(LightParameters.WorldPosition + View.ViewMatrices.GetPreViewTranslation());
+						LightViewData[LightSceneId].ColorScale = LightFade * LightParameters.GetLightExposureScale(Exposure);
+						LightViewData[LightSceneId].VolumetricScatteringIntensity = VolumetricScatteringIntensity;
+						LightViewData[LightSceneId].VirtualShadowMapId = VirtualShadowMapId;
+					}
+					else
+					{
+						LightViewData[LightSceneId].TranslatedWorldPosition = FVector3f::ZeroVector;
+						LightViewData[LightSceneId].ColorScale = 0.0f;
+						LightViewData[LightSceneId].VolumetricScatteringIntensity = 0.0f;
+						LightViewData[LightSceneId].VirtualShadowMapId = INDEX_NONE;
+					}
+				}
+			}, PrerequisiteTask);
+
+		FRDGBufferRef LightViewDataBuffer = CreateStructuredBuffer(GraphBuilder, TEXT("LightViewData"), LightViewData);
+
+		FForwardLightData* ForwardLightData = PerViewForwardLightData[ViewIndex];
+		ForwardLightData->LightViewData = GraphBuilder.CreateSRV(LightViewDataBuffer);
+	}
+
 	FComputeLightGridOutput Result = ComputeLightGrid(GraphBuilder, bCullLightsToGrid, SortedLightSet, PerViewForwardLightData);
 
 	for (int32 ViewIndex = 0; ViewIndex < Views.Num(); ++ViewIndex)
diff --git a/Engine/Source/Runtime/Renderer/Private/SceneRendering.h b/Engine/Source/Runtime/Renderer/Private/SceneRendering.h
--- a/Engine/Source/Runtime/Renderer/Private/SceneRendering.h
+++ b/Engine/Source/Runtime/Renderer/Private/SceneRendering.h
@@ -697,6 +697,7 @@ BEGIN_GLOBAL_SHADER_PARAMETER_STRUCT_WITH_CONSTRUCTOR(FForwardLightData, )
 	SHADER_PARAMETER_RDG_BUFFER_SRV(StructuredBuffer<uint>, NumCulledLightsGrid)
 	SHADER_PARAMETER_RDG_BUFFER_SRV(StructuredBuffer<uint>, CulledLightDataGrid32Bit)
 	SHADER_PARAMETER_RDG_BUFFER_SRV(Buffer<uint>, CulledLightDataGrid16Bit)
+	SHADER_PARAMETER_RDG_BUFFER_SRV(StructuredBuffer<FLightViewData>, LightViewData)
 END_GLOBAL_SHADER_PARAMETER_STRUCT()
 
 extern TRDGUniformBufferRef<FForwardLightData> CreateDummyForwardLightUniformBuffer(FRDGBuilder& GraphBuilder, EShaderPlatform ShaderPlatform);
