Refactored decal visibility into an async task instead of a parallel for. Removed decal scene proxy access from the RDG execution lambda and split into multiple async passes to improve parallelism.

diff --git a/Engine/Source/Runtime/Renderer/Private/CompositionLighting/CompositionLighting.cpp b/Engine/Source/Runtime/Renderer/Private/CompositionLighting/CompositionLighting.cpp
--- a/Engine/Source/Runtime/Renderer/Private/CompositionLighting/CompositionLighting.cpp
+++ b/Engine/Source/Runtime/Renderer/Private/CompositionLighting/CompositionLighting.cpp
@@ -485,32 +485,25 @@ static FScreenPassTexture AddPostProcessingAmbientOcclusion(
 	return FinalOutput;
 }
 
-FCompositionLighting::FCompositionLighting(TArrayView<FViewInfo> InViews, const FSceneTextures& InSceneTextures, TUniqueFunction<bool(int32)> RequestSSAOFunction)
+FCompositionLighting::FCompositionLighting(FDecalVisibilityTaskData* InDecalVisibility, TArrayView<FViewInfo> InViews, const FSceneTextures& InSceneTextures, TUniqueFunction<bool(int32)> RequestSSAOFunction)
 	: Views(InViews)
 	, ViewFamily(*InViews[0].Family)
 	, SceneTextures(InSceneTextures)
-	, bEnableDBuffer(IsDBufferEnabled(ViewFamily, SceneTextures.Config.ShaderPlatform))
-	, bEnableDecals(AreDecalsEnabled(ViewFamily))
+	, DecalVisibility(InDecalVisibility)
 {
-	checkf(bEnableDecals || !bEnableDBuffer, TEXT("DBuffer should only be enabled when Decals are enabled."));
-
-	const FScene& Scene = *(FScene*)ViewFamily.Scene;
-
 	ViewAOConfigs.SetNum(Views.Num());
 
-	if (bEnableDecals)
-	{
-		VisibleDecals.SetNum(Views.Num());
-	}
-
 	for (int32 Index = 0; Index < Views.Num(); ++Index)
 	{
 		ViewAOConfigs[Index].bRequested = RequestSSAOFunction(Index);
+	}
+}
 
-		if (bEnableDecals)
-		{
-			VisibleDecals[Index] = DecalRendering::BuildVisibleDecalList(Scene.Decals, Views[Index]);
-		}
+FCompositionLighting::~FCompositionLighting()
+{
+	if (DecalVisibility)
+	{
+		DecalVisibility->Finish();
 	}
 }
 
@@ -572,8 +565,9 @@ void FCompositionLighting::ProcessBeforeBasePass(FRDGBuilder& GraphBuilder, FDBu
 		const FAOConfig& ViewConfig = ViewAOConfigs[ViewIndex];
 
 		const bool bEnableSSAO = ViewConfig.SSAOLocation == ESSAOLocation::BeforeBasePass;
+		const bool bEnableDecals = DecalVisibility && DecalVisibility->HasStage(ViewIndex, EDecalRenderStage::BeforeBasePass);
 
-		if (!bEnableDBuffer && !bEnableSSAO)
+		if (!bEnableDecals && !bEnableSSAO)
 		{
 			continue;
 		}
@@ -584,10 +578,11 @@ void FCompositionLighting::ProcessBeforeBasePass(FRDGBuilder& GraphBuilder, FDBu
 		View.BeginRenderView();
 
 		// decals are before AmbientOcclusion so the decal can output a normal that AO is affected by
-		if (bEnableDBuffer)
+		if (bEnableDecals)
 		{
 			FDeferredDecalPassTextures DecalPassTextures = GetDeferredDecalPassTextures(GraphBuilder, View, SubstrateSceneData, SceneTextures, &DBufferTextures, EDecalRenderStage::BeforeBasePass);
-			AddDeferredDecalPass(GraphBuilder, View, VisibleDecals[ViewIndex], DecalPassTextures, InstanceCullingManager, EDecalRenderStage::BeforeBasePass);
+			TConstArrayView<const FVisibleDecal*> SortedDecals = DecalVisibility->FinishRelevantDecals(ViewIndex, EDecalRenderStage::BeforeBasePass);
+			AddDeferredDecalPass(GraphBuilder, View, SortedDecals, DecalPassTextures, InstanceCullingManager, EDecalRenderStage::BeforeBasePass);
 		}
 
 		if (bEnableSSAO)
@@ -626,20 +621,27 @@ void FCompositionLighting::ProcessAfterBasePass(FRDGBuilder& GraphBuilder, FInst
 
 		View.BeginRenderView();
 
-
-		if (bEnableDecals && !bEnableDBuffer && IsUsingGBuffers(SceneTextures.Config.ShaderPlatform) && Mode != EProcessAfterBasePassMode::SkipBeforeLightingDecals)
+		const auto AddDeferredDecalPassLambda = [&] (const FDeferredDecalPassTextures& DecalPassTextures, EDecalRenderStage Stage)
 		{
-			// We can disable this pass if using DBuffer decals
-			// Decals are before AmbientOcclusion so the decal can output a normal that AO is affected by
-			FDeferredDecalPassTextures DecalPassTextures = GetDeferredDecalPassTextures(GraphBuilder, View, SubstrateSceneData, SceneTextures, nullptr, EDecalRenderStage::BeforeLighting);
-			AddDeferredDecalPass(GraphBuilder, View, VisibleDecals[ViewIndex], DecalPassTextures, InstanceCullingManager, EDecalRenderStage::BeforeLighting);
-		}
+			AddDeferredDecalPass(GraphBuilder, View, DecalVisibility->FinishRelevantDecals(ViewIndex, Stage), DecalPassTextures, InstanceCullingManager, Stage);
+		};
 
-		if (bEnableDecals && Mode != EProcessAfterBasePassMode::OnlyBeforeLightingDecals)
+		if (DecalVisibility)
 		{
-			// DBuffer decals with emissive component
-			FDeferredDecalPassTextures DecalPassTextures = GetDeferredDecalPassTextures(GraphBuilder, View, SubstrateSceneData, SceneTextures, nullptr, EDecalRenderStage::Emissive);
-			AddDeferredDecalPass(GraphBuilder, View, VisibleDecals[ViewIndex], DecalPassTextures, InstanceCullingManager, EDecalRenderStage::Emissive);
+			if (DecalVisibility->HasStage(ViewIndex, EDecalRenderStage::BeforeLighting) && Mode != EProcessAfterBasePassMode::SkipBeforeLightingDecals)
+			{
+				// We can disable this pass if using DBuffer decals
+				// Decals are before AmbientOcclusion so the decal can output a normal that AO is affected by
+				FDeferredDecalPassTextures DecalPassTextures = GetDeferredDecalPassTextures(GraphBuilder, View, SubstrateSceneData, SceneTextures, nullptr, EDecalRenderStage::BeforeLighting);
+				AddDeferredDecalPassLambda(DecalPassTextures, EDecalRenderStage::BeforeLighting);
+			}
+
+			if (DecalVisibility->HasStage(ViewIndex, EDecalRenderStage::Emissive) && Mode != EProcessAfterBasePassMode::OnlyBeforeLightingDecals)
+			{
+				// DBuffer decals with emissive component
+				FDeferredDecalPassTextures DecalPassTextures = GetDeferredDecalPassTextures(GraphBuilder, View, SubstrateSceneData, SceneTextures, nullptr, EDecalRenderStage::Emissive);
+				AddDeferredDecalPassLambda(DecalPassTextures, EDecalRenderStage::Emissive);
+			}
 		}
 
 		// Forward shading SSAO is applied before the base pass using only the depth buffer.
@@ -663,7 +665,7 @@ void FCompositionLighting::ProcessAfterBasePass(FRDGBuilder& GraphBuilder, FInst
 					AmbientOcclusion = AddPostProcessingGTAOPostAsync(GraphBuilder, View, Parameters, GTAOHorizons, FinalTarget);
 
 					ensureMsgf(
-						DecalRendering::BuildRelevantDecalList(VisibleDecals[ViewIndex], EDecalRenderStage::AmbientOcclusion, nullptr) == false,
+						!DecalVisibility || !DecalVisibility->HasStage(ViewIndex, EDecalRenderStage::AmbientOcclusion),
 						TEXT("Ambient occlusion decals are not supported with Async compute SSAO."));
 				}
 				else
@@ -679,11 +681,11 @@ void FCompositionLighting::ProcessAfterBasePass(FRDGBuilder& GraphBuilder, FInst
 						AmbientOcclusion = AddPostProcessingAmbientOcclusion(GraphBuilder, View, Parameters, FinalTarget);
 					}
 
-					if (bEnableDecals)
+					if (DecalVisibility)
 					{
 						FDeferredDecalPassTextures DecalPassTextures = GetDeferredDecalPassTextures(GraphBuilder, View, SubstrateSceneData, SceneTextures, nullptr, EDecalRenderStage::AmbientOcclusion);
 						DecalPassTextures.ScreenSpaceAO = AmbientOcclusion.Texture;
-						AddDeferredDecalPass(GraphBuilder, View, VisibleDecals[ViewIndex], DecalPassTextures, InstanceCullingManager, EDecalRenderStage::AmbientOcclusion);
+						AddDeferredDecalPassLambda(DecalPassTextures, EDecalRenderStage::AmbientOcclusion);
 					}
 				}
 			}
diff --git a/Engine/Source/Runtime/Renderer/Private/CompositionLighting/CompositionLighting.h b/Engine/Source/Runtime/Renderer/Private/CompositionLighting/CompositionLighting.h
--- a/Engine/Source/Runtime/Renderer/Private/CompositionLighting/CompositionLighting.h
+++ b/Engine/Source/Runtime/Renderer/Private/CompositionLighting/CompositionLighting.h
@@ -18,7 +18,8 @@
 class FCompositionLighting
 {
 public:
-	FCompositionLighting(TArrayView<FViewInfo> InViews, const FSceneTextures& InSceneTextures, TUniqueFunction<bool(int32)> RequestSSAOFunction);
+	FCompositionLighting(FDecalVisibilityTaskData* InDecalVisibility, TArrayView<FViewInfo> InViews, const FSceneTextures& InSceneTextures, TUniqueFunction<bool(int32)> RequestSSAOFunction);
+	~FCompositionLighting();
 
 	void ProcessAfterOcclusion(FRDGBuilder& GraphBuilder);
 
@@ -40,9 +41,6 @@ private:
 	const FSceneViewFamily& ViewFamily;
 	const FSceneTextures& SceneTextures;
 
-	const bool bEnableDBuffer;
-	const bool bEnableDecals;
-
 	enum class ESSAOLocation
 	{
 		None,
@@ -59,8 +57,8 @@ private:
 		bool bRequested = false;
 	};
 
+	FDecalVisibilityTaskData* DecalVisibility = nullptr;
 	TArray<FAOConfig, TInlineAllocator<8>> ViewAOConfigs;
-	TArray<FTransientDecalRenderDataList, TInlineAllocator<8>> VisibleDecals;
 	FRDGTextureRef HorizonsTexture = nullptr;
 	bool bInitialized = false;
 };
diff --git a/Engine/Source/Runtime/Renderer/Private/CompositionLighting/PostProcessDeferredDecals.cpp b/Engine/Source/Runtime/Renderer/Private/CompositionLighting/PostProcessDeferredDecals.cpp
--- a/Engine/Source/Runtime/Renderer/Private/CompositionLighting/PostProcessDeferredDecals.cpp
+++ b/Engine/Source/Runtime/Renderer/Private/CompositionLighting/PostProcessDeferredDecals.cpp
@@ -637,7 +637,7 @@ void CollectDeferredDecalPassPSOInitializers(
 void AddDeferredDecalPass(
 	FRDGBuilder& GraphBuilder,
 	FViewInfo& View,
-	TConstArrayView<FTransientDecalRenderData> VisibleDecals,
+	TConstArrayView<const FVisibleDecal*> SortedDecals,
 	const FDeferredDecalPassTextures& PassTextures,
 	FInstanceCullingManager& InstanceCullingManager,
 	EDecalRenderStage DecalRenderStage)
@@ -658,21 +658,12 @@ void AddDeferredDecalPass(
 	const EShaderPlatform ShaderPlatform = View.GetShaderPlatform();
 	const ERHIFeatureLevel::Type FeatureLevel = View.GetFeatureLevel();
 	const uint32 DecalCount = Scene.Decals.Num();
-	uint32 SortedDecalCount = 0;
-	FTransientDecalRenderDataList* SortedDecals = nullptr;
+	uint32 SortedDecalCount = SortedDecals.Num();
+	INC_DWORD_STAT_BY(STAT_Decals, SortedDecalCount);
 
 	checkf(DecalRenderStage != EDecalRenderStage::AmbientOcclusion || PassTextures.ScreenSpaceAO, TEXT("Attepting to render AO decals without SSAO having emitted a valid render target."));
 	checkf(DecalRenderStage != EDecalRenderStage::BeforeBasePass || IsUsingDBuffers(ShaderPlatform), TEXT("Only DBuffer decals are supported before the base pass."));
 
-	if (!VisibleDecals.IsEmpty())
-	{
-		SortedDecals = GraphBuilder.AllocObject<FTransientDecalRenderDataList>();
-		DecalRendering::BuildRelevantDecalList(VisibleDecals, DecalRenderStage, SortedDecals);
-		SortedDecalCount = SortedDecals->Num();
-
-		INC_DWORD_STAT_BY(STAT_Decals, SortedDecalCount);
-	}
-
 	const bool bHasAnyDrawCommandDecalCount = HasAnyDrawCommandDecalCount(DecalRenderStage, View);
 	const bool bVisibleDecalsInView = SortedDecalCount > 0 || bHasAnyDrawCommandDecalCount;
 	const bool bShaderComplexity = View.Family->EngineShowFlags.ShaderComplexity;
@@ -696,11 +687,11 @@ void AddDeferredDecalPass(
 		auto* PassParameters = GraphBuilder.AllocParameters<FDeferredDecalPassParameters>();
 		GetDeferredDecalPassParameters(GraphBuilder, View, PassTextures, DecalRenderStage, RenderTargetMode, *PassParameters);
 
-		GraphBuilder.AddPass(
+		FRDGPass* Pass = GraphBuilder.AddPass(
 			RDG_EVENT_NAME("Batch [%d, %d]", DecalIndexBegin, DecalIndexEnd - 1),
 			PassParameters,
 			ERDGPassFlags::Raster,
-			[&View, FeatureLevel, ShaderPlatform, DecalIndexBegin, DecalIndexEnd, SortedDecals, DecalRenderStage, RenderTargetMode, bStencilSizeThreshold, bShaderComplexity](FRHICommandList& RHICmdList)
+			[&View, FeatureLevel, ShaderPlatform, DecalIndexBegin, DecalIndexEnd, SortedDecals, DecalRenderStage, RenderTargetMode, bStencilSizeThreshold, bShaderComplexity](FRDGAsyncTask, FRHICommandList& RHICmdList)
 		{
 			RHICmdList.SetViewport(View.ViewRect.Min.X, View.ViewRect.Min.Y, 0.0f, View.ViewRect.Max.X, View.ViewRect.Max.Y, 1.0f);
 
@@ -710,9 +701,8 @@ void AddDeferredDecalPass(
 
 			for (uint32 DecalIndex = DecalIndexBegin; DecalIndex < DecalIndexEnd; ++DecalIndex)
 			{
-				const FTransientDecalRenderData& DecalData = (*SortedDecals)[DecalIndex];
-				const FDeferredDecalProxy& DecalProxy = *DecalData.Proxy;
-				const FMatrix ComponentToWorldMatrix = DecalProxy.ComponentTrans.ToMatrixWithScale();
+				const FVisibleDecal& VisibleDecal = *SortedDecals[DecalIndex];
+				const FMatrix ComponentToWorldMatrix = VisibleDecal.ComponentTrans.ToMatrixWithScale();
 				const FMatrix FrustumComponentToClip = DecalRendering::ComputeComponentToClipMatrix(View, ComponentToWorldMatrix);
 				const bool bStencilThisDecal = IsStencilOptimizationAvailable(DecalRenderStage);
 
@@ -723,14 +713,14 @@ void AddDeferredDecalPass(
 					bThisDecalUsesStencil = RenderPreStencil(RHICmdList, View, ComponentToWorldMatrix, FrustumComponentToClip);
 				}
 
-				const bool bInsideDecal = ((FVector)View.ViewMatrices.GetViewOrigin() - ComponentToWorldMatrix.GetOrigin()).SizeSquared() < FMath::Square(DecalData.ConservativeRadius * 1.05f + View.NearClippingDistance * 2.0f);
+				const bool bInsideDecal = ((FVector)View.ViewMatrices.GetViewOrigin() - ComponentToWorldMatrix.GetOrigin()).SizeSquared() < FMath::Square(VisibleDecal.ConservativeRadius * 1.05f + View.NearClippingDistance * 2.0f);
 
 				FGraphicsPipelineStateInitializer GraphicsPSOInit;
 				RHICmdList.ApplyCachedRenderTargets(GraphicsPSOInit);
 
 				{
 					// Account for the reversal of handedness caused by negative scale on the decal
-					const FVector Scale = DecalProxy.ComponentTrans.GetScale3D();
+					const FVector Scale = VisibleDecal.ComponentTrans.GetScale3D();
 					const bool bReverseHanded = Scale.X * Scale.Y * Scale.Z < 0.0f;
 					const EDecalRasterizerState DecalRasterizerState = DecalRendering::GetDecalRasterizerState(bInsideDecal, bReverseHanded, View.bReverseCulling);
 					GraphicsPSOInit.RasterizerState = DecalRendering::GetDecalRasterizerState(DecalRasterizerState);
@@ -743,15 +733,15 @@ void AddDeferredDecalPass(
 					GraphicsPSOInit.DepthStencilState = GetDecalDepthState(StencilRef, DecalDepthState);
 				}
 
-				GraphicsPSOInit.BlendState = DecalRendering::GetDecalBlendState(DecalData.BlendDesc, DecalRenderStage, RenderTargetMode);
+				GraphicsPSOInit.BlendState = DecalRendering::GetDecalBlendState(VisibleDecal.BlendDesc, DecalRenderStage, RenderTargetMode);
 				GraphicsPSOInit.PrimitiveType = PT_TriangleList;
 
-				DecalRendering::SetShader(RHICmdList, GraphicsPSOInit, StencilRef, View, DecalData, DecalRenderStage, FrustumComponentToClip);
+				DecalRendering::SetShader(RHICmdList, GraphicsPSOInit, StencilRef, View, VisibleDecal, DecalRenderStage, FrustumComponentToClip);
 
 #if PSO_PRECACHING_VALIDATE
 				if (PSOCollectorStats::IsFullPrecachingValidationEnabled())
-				{					
-					PSOCollectorStats::CheckFullPipelineStateInCache(GraphicsPSOInit, EPSOPrecacheResult::Unknown, DecalData.MaterialProxy, &FLocalVertexFactory::StaticType, nullptr, PSOCollectorIndex);
+				{
+					PSOCollectorStats::CheckFullPipelineStateInCache(GraphicsPSOInit, EPSOPrecacheResult::Unknown, VisibleDecal.MaterialProxy, &FLocalVertexFactory::StaticType, nullptr, PSOCollectorIndex);
 				}
 #endif // PSO_PRECACHING_VALIDATE
 
@@ -759,6 +749,8 @@ void AddDeferredDecalPass(
 				RHICmdList.DrawIndexedPrimitive(GetUnitCubeIndexBuffer(), 0, 0, 8, 0, UE_ARRAY_COUNT(GCubeIndices) / 3, InstanceCount);
 			}
 		});
+
+		GraphBuilder.SetPassWorkload(Pass, DecalIndexEnd - DecalIndexBegin);
 	};
 
 	if (bVisibleDecalsInView)
@@ -772,21 +764,25 @@ void AddDeferredDecalPass(
 
 		if (SortedDecalCount > 0)
 		{
-			RDG_EVENT_SCOPE(GraphBuilder, "Decals (Relevant: %d, Visible: %d, Total: %d)", SortedDecalCount, VisibleDecals.Num(), DecalCount);
+			RDG_EVENT_SCOPE(GraphBuilder, "Decals (Relevant: %d, Total: %d)", SortedDecalCount, DecalCount);
+
+			const int32 MaxNumDecals = 128;
 
+			uint32 NumDecals = 0;
 			uint32 SortedDecalIndex = 1;
 			uint32 LastSortedDecalIndex = 0;
-			EDecalRenderTargetMode LastRenderTargetMode = DecalRendering::GetRenderTargetMode((*SortedDecals)[0].BlendDesc, DecalRenderStage);
+			EDecalRenderTargetMode LastRenderTargetMode = DecalRendering::GetRenderTargetMode(SortedDecals[0]->BlendDesc, DecalRenderStage);
 
-			for (; SortedDecalIndex < SortedDecalCount; ++SortedDecalIndex)
+			for (; SortedDecalIndex < SortedDecalCount; ++SortedDecalIndex, ++NumDecals)
 			{
-				const EDecalRenderTargetMode RenderTargetMode = DecalRendering::GetRenderTargetMode((*SortedDecals)[SortedDecalIndex].BlendDesc, DecalRenderStage);
+				const EDecalRenderTargetMode RenderTargetMode = DecalRendering::GetRenderTargetMode(SortedDecals[SortedDecalIndex]->BlendDesc, DecalRenderStage);
 
-				if (LastRenderTargetMode != RenderTargetMode)
+				if (LastRenderTargetMode != RenderTargetMode || NumDecals > MaxNumDecals)
 				{
 					RenderDecals(LastSortedDecalIndex, SortedDecalIndex, LastRenderTargetMode);
 					LastRenderTargetMode = RenderTargetMode;
 					LastSortedDecalIndex = SortedDecalIndex;
+					NumDecals = 0;
 				}
 			}
 
diff --git a/Engine/Source/Runtime/Renderer/Private/CompositionLighting/PostProcessDeferredDecals.h b/Engine/Source/Runtime/Renderer/Private/CompositionLighting/PostProcessDeferredDecals.h
--- a/Engine/Source/Runtime/Renderer/Private/CompositionLighting/PostProcessDeferredDecals.h
+++ b/Engine/Source/Runtime/Renderer/Private/CompositionLighting/PostProcessDeferredDecals.h
@@ -10,7 +10,7 @@
 
 enum class EDecalRenderStage : uint8;
 enum class EDecalRenderTargetMode : uint8;
-struct FTransientDecalRenderData;
+struct FVisibleDecal;
 struct FDBufferTextures;
 struct FSceneTextures;
 class FViewInfo;
@@ -43,17 +43,17 @@ struct FDeferredDecalPassTextures
 };
 
 FDeferredDecalPassTextures GetDeferredDecalPassTextures(
-	FRDGBuilder& GraphBuilder, 
+	FRDGBuilder& GraphBuilder,
 	const FViewInfo& View,
 	const FSubstrateSceneData& SubstrateSceneData,
-	const FSceneTextures& SceneTextures, 
+	const FSceneTextures& SceneTextures,
 	FDBufferTextures* DBufferTextures,
 	EDecalRenderStage DecalRenderStage);
 
 void AddDeferredDecalPass(
 	FRDGBuilder& GraphBuilder,
 	FViewInfo& ViewInfo,
-	TConstArrayView<FTransientDecalRenderData> VisibleDecals,
+	TConstArrayView<const FVisibleDecal*> SortedDecals,
 	const FDeferredDecalPassTextures& Textures,
 	FInstanceCullingManager& InstanceCullingManager,
 	EDecalRenderStage RenderStage);
diff --git a/Engine/Source/Runtime/Renderer/Private/DecalRenderingShared.cpp b/Engine/Source/Runtime/Renderer/Private/DecalRenderingShared.cpp
--- a/Engine/Source/Runtime/Renderer/Private/DecalRenderingShared.cpp
+++ b/Engine/Source/Runtime/Renderer/Private/DecalRenderingShared.cpp
@@ -24,20 +24,19 @@ static TAutoConsoleVariable<float> CVarDecalFadeScreenSizeMultiplier(
 	TEXT("  Smaller means decals fade less aggressively.")
 	);
 
-static TAutoConsoleVariable<bool> CVarDecalVisibilityMultithreaded(
-	TEXT("r.Decal.Visibility.Multithreaded"),
-	true,
-	TEXT("Whether to build visible decal list using multithreading. 0=disabled, 1=enabled (default)"),
-	ECVF_RenderThreadSafe);
-
-template<> struct TIsPODType<FTransientDecalRenderData> { enum { Value = true }; };
-
-FTransientDecalRenderData::FTransientDecalRenderData(const FDeferredDecalProxy& InDecalProxy, float InConservativeRadius, float InFadeAlpha, EShaderPlatform ShaderPlatform, ERHIFeatureLevel::Type FeatureLevel)
-	: Proxy(&InDecalProxy)
-	, MaterialProxy(InDecalProxy.DecalMaterial->GetRenderProxy())
+FVisibleDecal::FVisibleDecal(const FDeferredDecalProxy& InDecalProxy, float InConservativeRadius, float InFadeAlpha, EShaderPlatform ShaderPlatform, ERHIFeatureLevel::Type FeatureLevel)
+	: MaterialProxy(InDecalProxy.DecalMaterial->GetRenderProxy())
+	, Component((uintptr_t)InDecalProxy.Component)
+	, SortOrder(InDecalProxy.SortOrder)
 	, ConservativeRadius(InConservativeRadius)
 	, FadeAlpha(InFadeAlpha)
+	, InvFadeDuration(InDecalProxy.InvFadeDuration)
+	, InvFadeInDuration(InDecalProxy.InvFadeInDuration)
+	, FadeStartDelayNormalized(InDecalProxy.FadeStartDelayNormalized)
+	, FadeInStartDelayNormalized(InDecalProxy.FadeInStartDelayNormalized)
 	, DecalColor(InDecalProxy.DecalColor)
+	, ComponentTrans(InDecalProxy.ComponentTrans)
+	, BoxBounds(InDecalProxy.GetBounds().GetBox())
 {
 	// Build BlendDesc from a potentially incomplete material.
 	// If our shader isn't compiled yet then we will potentially render later with a different fallback material.
@@ -45,6 +44,102 @@ FTransientDecalRenderData::FTransientDecalRenderData(const FDeferredDecalProxy&
 	BlendDesc = DecalRendering::ComputeDecalBlendDesc(ShaderPlatform, MaterialResource);
 }
 
+FDecalVisibilityTaskData* FDecalVisibilityTaskData::Launch(FRDGBuilder& GraphBuilder, const FScene& Scene, TConstArrayView<FViewInfo> Views)
+{
+	const FSceneViewFamily& ViewFamily = *Views[0].Family;
+	if (AreDecalsEnabled(ViewFamily) && !HasRayTracedOverlay(ViewFamily))
+	{
+		const bool bDBufferEnabled = ::IsDBufferEnabled(ViewFamily, ViewFamily.GetShaderPlatform());
+		const bool bGBufferEnabled = IsUsingGBuffers(ViewFamily.GetShaderPlatform());
+
+		return GraphBuilder.AllocObject<FDecalVisibilityTaskData>(Scene, Views, bDBufferEnabled, bGBufferEnabled);
+	}
+	return nullptr;
+}
+
+FDecalVisibilityTaskData::FDecalVisibilityTaskData(const FScene& Scene, TConstArrayView<FViewInfo> Views, bool bInDBufferEnabled, bool bInGBufferEnabled)
+	: bDBufferEnabled(bInDBufferEnabled)
+	, bGBufferEnabled(bInGBufferEnabled)
+{
+	ViewPackets.Reserve(Views.Num());
+
+	for (const FViewInfo& View : Views)
+	{
+		ViewPackets.Emplace(*this, Scene, View);
+	}
+}
+
+FDecalVisibilityViewPacket::FDecalVisibilityViewPacket(const FDecalVisibilityTaskData& InTaskData, const FScene& Scene, const FViewInfo& InView)
+	: TaskData(InTaskData)
+	, View(InView)
+{
+	VisibleDecals.Task = LaunchSceneRenderTask(TEXT("BuildVisibleDecalList"), [&OutputList = VisibleDecals.List, &InputList = Scene.Decals, &View = View]
+	{
+		OutputList = DecalRendering::BuildVisibleDecalList(InputList, View);
+	});
+
+	AllTasksEvent.AddPrerequisites(VisibleDecals.Task);
+
+	const auto LaunchRelevantDecalTask = [this] (EDecalRenderStage Stage)
+	{
+		FRelevantDecals& RelevantDecals = RelevantDecalsMap.Emplace(Stage);
+
+		RelevantDecals.Task = LaunchSceneRenderTask(TEXT("BuildRelevantDecalList"), [&OutputList = RelevantDecals.List, &InputList = VisibleDecals.List, Stage]
+		{
+			OutputList = DecalRendering::BuildRelevantDecalList(InputList, Stage);
+
+		}, VisibleDecals.Task);
+
+		AllTasksEvent.AddPrerequisites(RelevantDecals.Task);
+	};
+
+	// DBuffer Passes
+	if (TaskData.bDBufferEnabled)
+	{
+		LaunchRelevantDecalTask(EDecalRenderStage::BeforeBasePass);
+		LaunchRelevantDecalTask(EDecalRenderStage::Emissive);
+	}
+	// GBuffer Passes
+	else if (TaskData.bGBufferEnabled)
+	{
+		LaunchRelevantDecalTask(EDecalRenderStage::BeforeLighting);
+	}
+
+	// AmbientOcclusion Pass gets built on-demand, since we don't know if it will be enabled until later in the pipeline.
+}
+
+TConstArrayView<FVisibleDecal> FDecalVisibilityViewPacket::FinishVisibleDecals()
+{
+	check(IsInRenderingThread());
+	if (VisibleDecals.Task.IsValid())
+	{
+		VisibleDecals.Task.Wait();
+		VisibleDecals.Task = {};
+	}
+	return VisibleDecals.List;
+}
+
+TConstArrayView<const FVisibleDecal*> FDecalVisibilityViewPacket::FinishRelevantDecals(EDecalRenderStage Stage)
+{
+	check(IsInRenderingThread());
+	FRelevantDecals* RelevantDecals = RelevantDecalsMap.Find(Stage);
+	if (RelevantDecals)
+	{
+		if (RelevantDecals->Task.IsValid())
+		{
+			RelevantDecals->Task.Wait();
+			RelevantDecals->Task = {};
+		}
+	}
+	else
+	{
+		// Wasn't requested as a task. Build on-demand.
+		RelevantDecals = &RelevantDecalsMap.Emplace(Stage);
+		RelevantDecals->List = DecalRendering::BuildRelevantDecalList(FinishVisibleDecals(), Stage);
+	}
+	return RelevantDecals->List;
+}
+
 /**
  * A vertex shader for projecting a deferred decal onto the scene.
  */
@@ -99,19 +194,19 @@ public:
 		MobileReflectionCaptureParam.Bind(Initializer.ParameterMap, FMobileReflectionCaptureShaderParameters::FTypeInfo::GetStructMetadata()->GetShaderVariableName());
 	}
 
-	void SetParameters(FRHIBatchedShaderParameters& BatchedParameters, const FViewInfo& View, const FDeferredDecalProxy& DecalProxy, const FMaterialRenderProxy* MaterialProxy, const FMaterial* MaterialResource, const float FadeAlphaValue = 1.0f, const FScene* Scene = nullptr)
+	void SetParameters(FRHIBatchedShaderParameters& BatchedParameters, const FViewInfo& View, const FVisibleDecal& VisibleDecal, const FMaterialRenderProxy* MaterialProxy, const FMaterial* MaterialResource, const float FadeAlphaValue = 1.0f, const FScene* Scene = nullptr)
 	{
 		auto& PrimitivePS = GetUniformBufferParameter<FPrimitiveUniformShaderParameters>();
 		SetUniformBufferParameter(BatchedParameters, PrimitivePS, GIdentityPrimitiveUniformBuffer);
 
 		FMaterialShader::SetParameters(BatchedParameters, MaterialProxy, *MaterialResource, View);
 
-		const FMatrix DecalToWorldMatrix = DecalProxy.ComponentTrans.ToMatrixWithScale();
-		const FMatrix WorldToDecalMatrix = DecalProxy.ComponentTrans.ToInverseMatrixWithScale();
+		const FMatrix DecalToWorldMatrix = VisibleDecal.ComponentTrans.ToMatrixWithScale();
+		const FMatrix WorldToDecalMatrix = VisibleDecal.ComponentTrans.ToInverseMatrixWithScale();
 		const FDFVector3 AbsoluteOrigin(DecalToWorldMatrix.GetOrigin());
 		const FVector3f PositionHigh = AbsoluteOrigin.High;
 		const FMatrix44f RelativeDecalToWorldMatrix = FDFMatrix::MakeToRelativeWorldMatrix(PositionHigh, DecalToWorldMatrix).M;
-		const FVector3f OrientationVector = (FVector3f)DecalProxy.ComponentTrans.GetUnitAxis(EAxis::X);
+		const FVector3f OrientationVector = (FVector3f)VisibleDecal.ComponentTrans.GetUnitAxis(EAxis::X);
 
 		if (DecalPositionHigh.IsBound())
 		{
@@ -172,11 +267,11 @@ public:
 		// Certain engine captures (e.g. environment reflection) don't have a tick. Default to fully opaque.
 		if (View.Family->Time.GetWorldTimeSeconds())
 		{
-			LifetimeAlpha = FMath::Clamp(FMath::Min(View.Family->Time.GetWorldTimeSeconds() * -DecalProxy.InvFadeDuration + DecalProxy.FadeStartDelayNormalized, View.Family->Time.GetWorldTimeSeconds() * DecalProxy.InvFadeInDuration + DecalProxy.FadeInStartDelayNormalized), 0.0f, 1.0f);
+			LifetimeAlpha = FMath::Clamp(FMath::Min(View.Family->Time.GetWorldTimeSeconds() * -VisibleDecal.InvFadeDuration + VisibleDecal.FadeStartDelayNormalized, View.Family->Time.GetWorldTimeSeconds() * VisibleDecal.InvFadeInDuration + VisibleDecal.FadeInStartDelayNormalized), 0.0f, 1.0f);
 		}
  
 		SetShaderValue(BatchedParameters, DecalParams, FVector2f(FadeAlphaValue, LifetimeAlpha));
-		SetShaderValue(BatchedParameters, DecalColorParam, DecalProxy.DecalColor);
+		SetShaderValue(BatchedParameters, DecalColorParam, VisibleDecal.DecalColor);
 
 		if (MobileDirectionLightBufferParam.IsBound() && Scene)
 		{
@@ -313,212 +408,136 @@ namespace DecalRendering
 		return FMath::Clamp(Alpha, 0.0f, 1.0f);
 	}
 
-	void SortDecalList(FTransientDecalRenderDataList& Decals)
+	void SortDecalList(FRelevantDecalList& Decals)
 	{
-		// Sort by sort order to allow control over composited result
-		// Then sort decals by state to reduce render target switches
-		// Also sort by component since Sort() is not stable
-		struct FCompareFTransientDecalRenderData
+		Decals.Sort([](const FVisibleDecal& A, const FVisibleDecal& B)
 		{
-			FORCEINLINE bool operator()(const FTransientDecalRenderData& A, const FTransientDecalRenderData& B) const
+			// Sort by sort order to allow control over composited result
+			if (B.SortOrder != A.SortOrder)
+			{ 
+				return A.SortOrder < B.SortOrder;
+			}
+			// Then sort decals by state to reduce render target switches
+			if (B.BlendDesc.bWriteNormal != A.BlendDesc.bWriteNormal)
 			{
-				if (B.Proxy->SortOrder != A.Proxy->SortOrder)
-				{ 
-					return A.Proxy->SortOrder < B.Proxy->SortOrder;
-				}
-				if (B.BlendDesc.bWriteNormal != A.BlendDesc.bWriteNormal)
-				{
-					// bWriteNormal here has priority because we want to render decals that output normals before those could read normals.
-					// Also this is the only flag that can trigger a change of EDecalRenderTargetMode inside a single EDecalRenderStage, and we batch according to this.
-					return B.BlendDesc.bWriteNormal < A.BlendDesc.bWriteNormal; // < so that those outputting normal are first.
-				}
-				if (B.BlendDesc.Packed != A.BlendDesc.Packed)
-				{
-					// Sorting by the FDecalBlendDesc contents will reduce blend state changes.
-					return (int32)B.BlendDesc.Packed < (int32)A.BlendDesc.Packed;
-				}
-				if (B.MaterialProxy != A.MaterialProxy)
-				{
-					// Batch decals with the same material together
-					return B.MaterialProxy < A.MaterialProxy;
-				}
-				return (PTRINT)B.Proxy->Component < (PTRINT)A.Proxy->Component;
+				// bWriteNormal here has priority because we want to render decals that output normals before those could read normals.
+				// Also this is the only flag that can trigger a change of EDecalRenderTargetMode inside a single EDecalRenderStage, and we batch according to this.
+				return B.BlendDesc.bWriteNormal < A.BlendDesc.bWriteNormal; // < so that those outputting normal are first.
 			}
-		};
-
-		// Sort decals by blend mode to reduce render target switches
-		Decals.Sort(FCompareFTransientDecalRenderData());
+			// Sort decals by blend mode to reduce render target switches
+			if (B.BlendDesc.Packed != A.BlendDesc.Packed)
+			{
+				// Sorting by the FDecalBlendDesc contents will reduce blend state changes.
+				return (int32)B.BlendDesc.Packed < (int32)A.BlendDesc.Packed;
+			}
+			if (B.MaterialProxy != A.MaterialProxy)
+			{
+				// Batch decals with the same material together
+				return B.MaterialProxy < A.MaterialProxy;
+			}
+			// Also sort by component since Sort() is not stable
+			return B.Component < A.Component;
+		});
 	}
 
-	static bool ProcessDecal(const FDeferredDecalProxy* DecalProxy, const FViewInfo& View, float FadeMultiplier, EShaderPlatform ShaderPlatform, bool bIsPerspectiveProjection, FTransientDecalRenderData& OutData)
-	{
-		if (!DecalProxy->DecalMaterial || !DecalProxy->DecalMaterial->IsValidLowLevelFast())
-		{
-			return false;
-		}
-
-		if (!DecalProxy->IsShown(&View))
-		{
-			return false;
-		}
-
-		const FMatrix ComponentToWorldMatrix = DecalProxy->ComponentTrans.ToMatrixWithScale();
-
-		// can be optimized as we test against a sphere around the box instead of the box itself
-		const float ConservativeRadius = FMath::Sqrt(
-			ComponentToWorldMatrix.GetScaledAxis(EAxis::X).SizeSquared() +
-			ComponentToWorldMatrix.GetScaledAxis(EAxis::Y).SizeSquared() +
-			ComponentToWorldMatrix.GetScaledAxis(EAxis::Z).SizeSquared());
-
-		const bool bIsVisibleInFirstView = View.ViewFrustum.IntersectSphere(ComponentToWorldMatrix.GetOrigin(), ConservativeRadius);
-		const FViewInfo* InstancedView = View.GetInstancedView();
-		const bool bIsVisibleInSecondView = InstancedView ? InstancedView->ViewFrustum.IntersectSphere(ComponentToWorldMatrix.GetOrigin(), ConservativeRadius) : false;
-
-		// can be optimized as the test is too conservative (sphere instead of OBB)
-		if (ConservativeRadius < SMALL_NUMBER || !(bIsVisibleInFirstView || bIsVisibleInSecondView))
-		{
-			return false;
-		}
-
-		float FadeAlpha = 1.0f;
-
-		if (bIsPerspectiveProjection && DecalProxy->FadeScreenSize != 0.0f)
-		{
-			FadeAlpha = CalculateDecalFadeAlpha(DecalProxy->FadeScreenSize, ComponentToWorldMatrix, View, FadeMultiplier);
-		}
-
-		const bool bShouldRender = FadeAlpha > 0.0f;
-
-		if (!bShouldRender)
-		{
-			return false;
-		}
-
-		OutData = FTransientDecalRenderData(*DecalProxy, ConservativeRadius, FadeAlpha, ShaderPlatform, View.GetFeatureLevel());
-
-		return true;
-	};
-
-	FTransientDecalRenderDataList BuildVisibleDecalList(TConstArrayView<FDeferredDecalProxy*> Decals, const FViewInfo& View)
+	FVisibleDecalList BuildVisibleDecalList(TConstArrayView<FDeferredDecalProxy*> Decals, const FViewInfo& View)
 	{
-		QUICK_SCOPE_CYCLE_COUNTER(BuildVisibleDecalList);
+		TRACE_CPUPROFILER_EVENT_SCOPE(BuildVisibleDecalList);
 
 		// Don't draw for shader complexity mode.
 		// todo: Handle shader complexity mode for deferred decal.
-		if (View.Family->EngineShowFlags.ShaderComplexity)
+		if (Decals.IsEmpty() || View.Family->EngineShowFlags.ShaderComplexity)
 		{
 			return {};
 		}
 
-		FTransientDecalRenderDataList OutVisibleDecals;
+		FVisibleDecalList VisibleDecals;
+		VisibleDecals.Reserve(Decals.Num());
 
 		const float FadeMultiplier = GetDecalFadeScreenSizeMultiplier();
 		const EShaderPlatform ShaderPlatform = View.GetShaderPlatform();
 
 		const bool bIsPerspectiveProjection = View.IsPerspectiveProjection();
 
-		// Build a list of decals that need to be rendered for this view
-
-		if (CVarDecalVisibilityMultithreaded.GetValueOnRenderThread() && FApp::ShouldUseThreadingForPerformance())
+		for (const FDeferredDecalProxy* DecalProxy : Decals)
 		{
-			struct FVisibleDecalListContext
+			if (!DecalProxy->DecalMaterial || !DecalProxy->DecalMaterial->IsValidLowLevelFast())
 			{
-				TChunkedArray<FTransientDecalRenderData> VisibleDecals;
-			};
-
-			TArray<FVisibleDecalListContext> Contexts;
-			const int32 MinBatchSize = 128;
-			ParallelForWithTaskContext(
-				TEXT("BuildVisibleDecalList_Parallel"),
-				Contexts,
-				Decals.Num(),
-				MinBatchSize,
-				[Decals, &View, ShaderPlatform, bIsPerspectiveProjection, FadeMultiplier](FVisibleDecalListContext& Context, int32 ItemIndex)
-				{
-					FTaskTagScope TaskTagScope(ETaskTag::EParallelRenderingThread);
+				continue;
+			}
+
+			if (!DecalProxy->IsShown(&View))
+			{
+				continue;
+			}
 
-					const FDeferredDecalProxy* DecalProxy = Decals[ItemIndex];
+			const FMatrix ComponentToWorldMatrix = DecalProxy->ComponentTrans.ToMatrixWithScale();
 
-					FTransientDecalRenderData Data;
+			// can be optimized as we test against a sphere around the box instead of the box itself
+			const float ConservativeRadius = FMath::Sqrt(
+				ComponentToWorldMatrix.GetScaledAxis(EAxis::X).SizeSquared() +
+				ComponentToWorldMatrix.GetScaledAxis(EAxis::Y).SizeSquared() +
+				ComponentToWorldMatrix.GetScaledAxis(EAxis::Z).SizeSquared());
 
-					if (ProcessDecal(DecalProxy, View, FadeMultiplier, ShaderPlatform, bIsPerspectiveProjection, Data))
-					{
-						Context.VisibleDecals.AddElement(MoveTemp(Data));
-					}
-				});
+			const bool bIsVisibleInFirstView = View.ViewFrustum.IntersectSphere(ComponentToWorldMatrix.GetOrigin(), ConservativeRadius);
+			const FViewInfo* InstancedView = View.GetInstancedView();
+			const bool bIsVisibleInSecondView = InstancedView ? InstancedView->ViewFrustum.IntersectSphere(ComponentToWorldMatrix.GetOrigin(), ConservativeRadius) : false;
 
-			if (Contexts.Num() > 0)
+			// can be optimized as the test is too conservative (sphere instead of OBB)
+			if (ConservativeRadius < SMALL_NUMBER || !(bIsVisibleInFirstView || bIsVisibleInSecondView))
 			{
-				SCOPED_NAMED_EVENT(BuildVisibleDecalList_Parallel_Merge, FColor::Emerald);
+				continue;
+			}
 
-				int32 NumVisibleDecals = 0;
+			float FadeAlpha = 1.0f;
 
-				for (auto& Context : Contexts)
-				{
-					NumVisibleDecals += Context.VisibleDecals.Num();
-				}
+			if (bIsPerspectiveProjection && DecalProxy->FadeScreenSize != 0.0f)
+			{
+				FadeAlpha = CalculateDecalFadeAlpha(DecalProxy->FadeScreenSize, ComponentToWorldMatrix, View, FadeMultiplier);
+			}
 
-				OutVisibleDecals.Empty(NumVisibleDecals);
+			const bool bShouldRender = FadeAlpha > 0.0f;
 
-				for (auto& Context : Contexts)
-				{
-					Context.VisibleDecals.CopyToLinearArray(OutVisibleDecals);
-				}
-			}
-		}
-		else
-		{
-			for (const FDeferredDecalProxy* DecalProxy : Decals)
+			if (!bShouldRender)
 			{
-				FTransientDecalRenderData Data;
-
-				if (ProcessDecal(DecalProxy, View, FadeMultiplier, ShaderPlatform, bIsPerspectiveProjection, Data))
-				{
-					OutVisibleDecals.Add(MoveTemp(Data));
-				}
+				continue;
 			}
+
+			VisibleDecals.Emplace(*DecalProxy, ConservativeRadius, FadeAlpha, ShaderPlatform, View.GetFeatureLevel());
 		}
 
-		return OutVisibleDecals;
+		return VisibleDecals;
 	}
 
-	bool BuildRelevantDecalList(TConstArrayView<FTransientDecalRenderData> Decals, EDecalRenderStage DecalRenderStage, FTransientDecalRenderDataList* OutVisibleDecals)
+	FRelevantDecalList BuildRelevantDecalList(TConstArrayView<FVisibleDecal> Decals, EDecalRenderStage DecalRenderStage)
 	{
-		QUICK_SCOPE_CYCLE_COUNTER(BuildRelevantDecalList);
+		TRACE_CPUPROFILER_EVENT_SCOPE(BuildRelevantDecalList);
 
-		if (OutVisibleDecals)
-		{
-			OutVisibleDecals->Empty(Decals.Num());
-		}
+		FRelevantDecalList RelevantDecals;
+		RelevantDecals.Reserve(Decals.Num());
 
-		// Build a list of decals that need to be rendered for this stage in SortedDecals
-		for (const FTransientDecalRenderData& DecalRenderData : Decals)
+		for (const FVisibleDecal& VisibleDecal : Decals)
 		{
-			checkf(DecalRenderData.Proxy->DecalMaterial && DecalRenderData.Proxy->DecalMaterial->IsValidLowLevelFast(),
-				TEXT("Decals should've been filtered earlier in BuildVisibleDecalList(...)"));
-
-			if (IsCompatibleWithRenderStage(DecalRenderData.BlendDesc, DecalRenderStage))
+			if (IsCompatibleWithRenderStage(VisibleDecal.BlendDesc, DecalRenderStage))
 			{
-				if (!OutVisibleDecals)
-				{
-					return true;
-				}
-
-				OutVisibleDecals->Add(DecalRenderData);
+				RelevantDecals.Emplace(&VisibleDecal);
 			}
 		}
 
-		if (!OutVisibleDecals)
-		{
-			return false;
-		}
+		SortDecalList(RelevantDecals);
+		return RelevantDecals;
+	}
 
-		if (OutVisibleDecals->Num() > 0)
-		{
-			// Could potentially run sort of auxiliary thread (kicked-off in BuildVisibleDecalList(...))
-			SortDecalList(*OutVisibleDecals);
+	bool HasRelevantDecals(TConstArrayView<FVisibleDecal> Decals, EDecalRenderStage DecalRenderStage)
+	{
+		TRACE_CPUPROFILER_EVENT_SCOPE(HasRelevantDecals);
 
-			return true;
+		for (const FVisibleDecal& VisibleDecal : Decals)
+		{
+			if (IsCompatibleWithRenderStage(VisibleDecal.BlendDesc, DecalRenderStage))
+			{
+				return true;
+			}
 		}
 
 		return false;
@@ -641,11 +660,11 @@ namespace DecalRendering
 	}
 
 	void SetShader(FRHICommandList& RHICmdList, FGraphicsPipelineStateInitializer& GraphicsPSOInit, uint32 StencilRef, const FViewInfo& View,
-		const FTransientDecalRenderData& DecalData, EDecalRenderStage DecalRenderStage, const FMatrix& FrustumComponentToClip, const FScene* Scene)
+		const FVisibleDecal& VisibleDecal, EDecalRenderStage DecalRenderStage, const FMatrix& FrustumComponentToClip, const FScene* Scene)
 	{
 		FMaterial const* MaterialResource = nullptr;
 		TShaderRef<FDeferredDecalPS> PixelShader;
-		FMaterialRenderProxy const* MaterialProxy = TryGetDeferredDecalMaterial(DecalData.MaterialProxy, View.GetFeatureLevel(), DecalRenderStage, MaterialResource, PixelShader);
+		FMaterialRenderProxy const* MaterialProxy = TryGetDeferredDecalMaterial(VisibleDecal.MaterialProxy, View.GetFeatureLevel(), DecalRenderStage, MaterialResource, PixelShader);
 
 		TShaderMapRef<FDeferredDecalVS> VertexShader(View.ShaderMap);
 
@@ -666,7 +685,7 @@ namespace DecalRendering
 
 		// Set pixel shader parameters.
 		{
-			SetShaderParametersLegacyPS(RHICmdList, PixelShader, View, *DecalData.Proxy, MaterialProxy, MaterialResource, DecalData.FadeAlpha, Scene);
+			SetShaderParametersLegacyPS(RHICmdList, PixelShader, View, VisibleDecal, MaterialProxy, MaterialResource, VisibleDecal.FadeAlpha, Scene);
 		}
 
 		// Set stream source after updating cached strides
diff --git a/Engine/Source/Runtime/Renderer/Private/DecalRenderingShared.h b/Engine/Source/Runtime/Renderer/Private/DecalRenderingShared.h
--- a/Engine/Source/Runtime/Renderer/Private/DecalRenderingShared.h
+++ b/Engine/Source/Runtime/Renderer/Private/DecalRenderingShared.h
@@ -11,6 +11,7 @@ class FMaterial;
 class FMaterialRenderProxy;
 class FScene;
 class FViewInfo;
+class FDecalVisibilityTaskData;
 
 class FShader;
 class FShaderMapPointerTable;
@@ -20,20 +21,125 @@ template<typename ShaderType> using TShaderRef = TShaderRefBase<ShaderType, FSha
 /**
  * Compact deferred decal data for rendering.
  */
-struct FTransientDecalRenderData
+struct FVisibleDecal
 {
-	const FDeferredDecalProxy* Proxy;
+	FVisibleDecal(const FDeferredDecalProxy& InDecalProxy, float InConservativeRadius, float InFadeAlpha, EShaderPlatform ShaderPlatform, ERHIFeatureLevel::Type FeatureLevel);
+
 	const FMaterialRenderProxy* MaterialProxy;
+	uintptr_t Component;
+	uint32 SortOrder;
 	FDecalBlendDesc BlendDesc;
 	float ConservativeRadius;
 	float FadeAlpha;
+	float InvFadeDuration;
+	float InvFadeInDuration;
+	float FadeStartDelayNormalized;
+	float FadeInStartDelayNormalized;
 	FLinearColor DecalColor;
+	FTransform ComponentTrans;
+	FBox BoxBounds;
+};
+
+using FVisibleDecalList = TArray<FVisibleDecal, SceneRenderingAllocator>;
+using FRelevantDecalList = TArray<const FVisibleDecal*, SceneRenderingAllocator>;
+
+class FDecalVisibilityViewPacket
+{
+public:
+	FDecalVisibilityViewPacket(const FDecalVisibilityTaskData& InTaskData, const FScene& Scene, const FViewInfo& InView);
+
+	~FDecalVisibilityViewPacket()
+	{
+		check(bFinishCalled);
+		check(AllTasksEvent.IsCompleted());
+	}
+
+	TConstArrayView<FVisibleDecal> FinishVisibleDecals();
+	TConstArrayView<const FVisibleDecal*> FinishRelevantDecals(EDecalRenderStage Stage);
 
-	FTransientDecalRenderData() = default; // required to support TChunkedArray
-	FTransientDecalRenderData(const FDeferredDecalProxy& InDecalProxy, float InConservativeRadius, float InFadeAlpha, EShaderPlatform ShaderPlatform, ERHIFeatureLevel::Type FeatureLevel);
+	void Finish()
+	{
+		bFinishCalled = true;
+		AllTasksEvent.Trigger();
+		AllTasksEvent.Wait();
+	}
+
+	bool HasStage(EDecalRenderStage Stage) const
+	{
+		return RelevantDecalsMap.Contains(Stage);
+	}
+
+private:
+	const FDecalVisibilityTaskData& TaskData;
+	const FViewInfo& View;
+
+	struct FVisibleDecals
+	{
+		FVisibleDecalList List;
+		UE::Tasks::FTask Task;
+
+	} VisibleDecals;
+
+	struct FRelevantDecals
+	{
+		FRelevantDecalList List;
+		UE::Tasks::FTask Task;
+	};
+
+	TMap<EDecalRenderStage, FRelevantDecals> RelevantDecalsMap;
+	UE::Tasks::FTaskEvent AllTasksEvent{ UE_SOURCE_LOCATION };
+	bool bFinishCalled = false;
+};
+
+class FDecalVisibilityTaskData
+{
+public:
+	static FDecalVisibilityTaskData* Launch(FRDGBuilder& GraphBuilder, const FScene& Scene, TConstArrayView<FViewInfo> Views);
+
+	TConstArrayView<FVisibleDecal> FinishVisibleDecals(int32 ViewIndex)
+	{
+		return ViewPackets[ViewIndex].FinishVisibleDecals();
+	}
+
+	TConstArrayView<const FVisibleDecal*> FinishRelevantDecals(int32 ViewIndex, EDecalRenderStage Stage)
+	{
+		return ViewPackets[ViewIndex].FinishRelevantDecals(Stage);
+	}
+
+	bool HasStage(int32 ViewIndex, EDecalRenderStage Stage) const
+	{
+		return ViewPackets[ViewIndex].HasStage(Stage);
+	}
+
+	void Finish()
+	{
+		for (FDecalVisibilityViewPacket& ViewPacket : ViewPackets)
+		{
+			ViewPacket.Finish();
+		}
+	}
+
+	bool IsDBufferEnabled() const
+	{
+		return bDBufferEnabled;
+	}
+
+	bool IsGBufferEnabled() const
+	{
+		return bGBufferEnabled;
+	}
+
+private:
+	FDecalVisibilityTaskData(const FScene& Scene, TConstArrayView<FViewInfo> Views, bool bInDBufferEnabled, bool bInGBufferEnabled);
+
+	const bool bDBufferEnabled;
+	const bool bGBufferEnabled;
+
+	TArray<FDecalVisibilityViewPacket, FRDGArrayAllocator> ViewPackets;
+
+	friend class FDecalVisibilityViewPacket;
+	RDG_FRIEND_ALLOCATOR_FRIEND(FDecalVisibilityTaskData);
 };
-	
-typedef TArray<FTransientDecalRenderData, SceneRenderingAllocator> FTransientDecalRenderDataList;
 
 /**
  * Shared deferred decal functionality.
@@ -44,10 +150,11 @@ namespace DecalRendering
 	float CalculateDecalFadeAlpha(float DecalFadeScreenSize, const FMatrix& ComponentToWorldMatrix, const FViewInfo& View, float FadeMultiplier);
 	FMatrix ComputeComponentToClipMatrix(const FViewInfo& View, const FMatrix& DecalComponentToWorld);
 	void SetVertexShaderOnly(FRHICommandList& RHICmdList, FGraphicsPipelineStateInitializer& GraphicsPSOInit, const FViewInfo& View, const FMatrix& FrustumComponentToClip);
-	void SortDecalList(FTransientDecalRenderDataList& Decals);
-	FTransientDecalRenderDataList BuildVisibleDecalList(TConstArrayView<FDeferredDecalProxy*> Decals, const FViewInfo& View);
-	bool BuildRelevantDecalList(TConstArrayView<FTransientDecalRenderData> Decals, EDecalRenderStage DecalRenderStage, FTransientDecalRenderDataList* OutVisibleDecals);
+	void SortDecalList(FRelevantDecalList& Decals);
+	FVisibleDecalList BuildVisibleDecalList(TConstArrayView<FDeferredDecalProxy*> Decals, const FViewInfo& View);
+	FRelevantDecalList BuildRelevantDecalList(TConstArrayView<FVisibleDecal> Decals, EDecalRenderStage DecalRenderStage);
+	bool HasRelevantDecals(TConstArrayView<FVisibleDecal> Decals, EDecalRenderStage DecalRenderStage);
 	bool GetShaders(ERHIFeatureLevel::Type FeatureLevel, const FMaterial& Material, EDecalRenderStage DecalRenderStage, TShaderRef<FShader>& OutVertexShader, TShaderRef<FShader>& OutPixelShader);
 	bool SetupShaderState(ERHIFeatureLevel::Type FeatureLevel, const FMaterial& Material, EDecalRenderStage DecalRenderStage, FBoundShaderStateInput& OutBoundShaderState);
-	void SetShader(FRHICommandList& RHICmdList, FGraphicsPipelineStateInitializer& GraphicsPSOInit, uint32 StencilRef, const FViewInfo& View, const FTransientDecalRenderData& DecalData, EDecalRenderStage DecalRenderStage, const FMatrix& FrustumComponentToClip, const FScene* Scene = nullptr);
-};
+	void SetShader(FRHICommandList& RHICmdList, FGraphicsPipelineStateInitializer& GraphicsPSOInit, uint32 StencilRef, const FViewInfo& View, const FVisibleDecal& DecalData, EDecalRenderStage DecalRenderStage, const FMatrix& FrustumComponentToClip, const FScene* Scene = nullptr);
+};
\ No newline at end of file
diff --git a/Engine/Source/Runtime/Renderer/Private/DeferredShadingRenderer.cpp b/Engine/Source/Runtime/Renderer/Private/DeferredShadingRenderer.cpp
--- a/Engine/Source/Runtime/Renderer/Private/DeferredShadingRenderer.cpp
+++ b/Engine/Source/Runtime/Renderer/Private/DeferredShadingRenderer.cpp
@@ -2159,7 +2159,7 @@ void FDeferredShadingSceneRenderer::Render(FRDGBuilder& GraphBuilder)
 
 	FComputeLightGridOutput ComputeLightGridOutput = {};
 
-	FCompositionLighting CompositionLighting(Views, SceneTextures, [this](int32 ViewIndex)
+	FCompositionLighting CompositionLighting(InitViewTaskDatas.Decals, Views, SceneTextures, [this](int32 ViewIndex)
 	{
 		return GetViewPipelineState(Views[ViewIndex]).AmbientOcclusionMethod == EAmbientOcclusionMethod::SSAO;
 	});
diff --git a/Engine/Source/Runtime/Renderer/Private/DeferredShadingRenderer.h b/Engine/Source/Runtime/Renderer/Private/DeferredShadingRenderer.h
--- a/Engine/Source/Runtime/Renderer/Private/DeferredShadingRenderer.h
+++ b/Engine/Source/Runtime/Renderer/Private/DeferredShadingRenderer.h
@@ -48,6 +48,8 @@ namespace LumenRadiosity
 class FRenderLightParameters;
 class FRayTracingScene;
 class FNaniteVisibility;
+class FDecalVisibilityTaskData;
+
 struct FNaniteVisibilityQuery;
 
 struct FSceneWithoutWaterTextures;
@@ -563,6 +565,7 @@ private:
 		FILCUpdatePrimTaskData* ILCUpdatePrim = nullptr;
 		RayTracing::FGatherInstancesTaskData* RayTracingGatherInstances = nullptr;
 		FDynamicShadowsTaskData* DynamicShadows = nullptr;
+		FDecalVisibilityTaskData* Decals = nullptr;
 		FLumenDirectLightingTaskData* LumenDirectLighting = nullptr;
 		FLumenSceneFrameTemporaries* LumenFrameTemporaries = nullptr;
 	};
diff --git a/Engine/Source/Runtime/Renderer/Private/MobileDecalRendering.cpp b/Engine/Source/Runtime/Renderer/Private/MobileDecalRendering.cpp
--- a/Engine/Source/Runtime/Renderer/Private/MobileDecalRendering.cpp
+++ b/Engine/Source/Runtime/Renderer/Private/MobileDecalRendering.cpp
@@ -87,20 +87,19 @@ void FMobileSceneRenderer::RenderDecals(FRHICommandList& RHICmdList, FViewInfo&
 
 void RenderDeferredDecalsMobile(FRHICommandList& RHICmdList, const FScene& Scene, const FViewInfo& View, EDecalRenderStage DecalRenderStage, EDecalRenderTargetMode RenderTargetMode)
 {
-	int32 SortedDecalCount = 0;
-	FTransientDecalRenderDataList SortedDecals;
+	FVisibleDecalList VisibleDecals;
+	FRelevantDecalList SortedDecals;
 
 	if (!Scene.Decals.IsEmpty())
 	{
-		FTransientDecalRenderDataList VisibleDecals = DecalRendering::BuildVisibleDecalList(Scene.Decals, View);
+		VisibleDecals = DecalRendering::BuildVisibleDecalList(Scene.Decals, View);
 
 		// Build a list of decals that need to be rendered for this view
-		DecalRendering::BuildRelevantDecalList(VisibleDecals, DecalRenderStage, &SortedDecals);
-		SortedDecalCount = SortedDecals.Num();
-		INC_DWORD_STAT_BY(STAT_Decals, SortedDecalCount);
+		SortedDecals = DecalRendering::BuildRelevantDecalList(VisibleDecals, DecalRenderStage);
+		INC_DWORD_STAT_BY(STAT_Decals, SortedDecals.Num());
 	}
 
-	if (SortedDecalCount > 0)
+	if (SortedDecals.Num() > 0)
 	{
 		FGraphicsPipelineStateInitializer GraphicsPSOInit;
 		RHICmdList.ApplyCachedRenderTargets(GraphicsPSOInit);
@@ -108,19 +107,18 @@ void RenderDeferredDecalsMobile(FRHICommandList& RHICmdList, const FScene& Scene
 		RHICmdList.SetViewport(View.ViewRect.Min.X, View.ViewRect.Min.Y, 0, View.ViewRect.Max.X, View.ViewRect.Max.Y, 1);
 		RHICmdList.SetStreamSource(0, GetUnitCubeVertexBuffer(), 0);
 
-		for (int32 DecalIndex = 0; DecalIndex < SortedDecalCount; DecalIndex++)
+		for (int32 DecalIndex = 0; DecalIndex < SortedDecals.Num(); DecalIndex++)
 		{
-			const FTransientDecalRenderData& DecalData = SortedDecals[DecalIndex];
-			const FDeferredDecalProxy& DecalProxy = *DecalData.Proxy;
-			const FMatrix ComponentToWorldMatrix = DecalProxy.ComponentTrans.ToMatrixWithScale();
+			const FVisibleDecal& VisibleDecal = *SortedDecals[DecalIndex];
+			const FMatrix ComponentToWorldMatrix = VisibleDecal.ComponentTrans.ToMatrixWithScale();
 			const FMatrix FrustumComponentToClip = DecalRendering::ComputeComponentToClipMatrix(View, ComponentToWorldMatrix);
 
-			const float ConservativeRadius = DecalData.ConservativeRadius;
+			const float ConservativeRadius = VisibleDecal.ConservativeRadius;
 			const bool bInsideDecal = ((FVector)View.ViewMatrices.GetViewOrigin() - ComponentToWorldMatrix.GetOrigin()).SizeSquared() < FMath::Square(ConservativeRadius * 1.05f + View.NearClippingDistance * 2.0f);
 			bool bReverseHanded = false;
 			{
 				// Account for the reversal of handedness caused by negative scale on the decal
-				const auto& Scale3d = DecalProxy.ComponentTrans.GetScale3D();
+				const auto& Scale3d = VisibleDecal.ComponentTrans.GetScale3D();
 				bReverseHanded = Scale3d[0] * Scale3d[1] * Scale3d[2] < 0.f;
 			}
 			EDecalRasterizerState DecalRasterizerState = DecalRendering::GetDecalRasterizerState(bInsideDecal, bReverseHanded, View.bReverseCulling);
@@ -143,10 +141,10 @@ void RenderDeferredDecalsMobile(FRHICommandList& RHICmdList, const FScene& Scene
 					GET_STENCIL_BIT_MASK(RECEIVE_DECAL, 1), 0x00>::GetRHI();
 			}
 
-			GraphicsPSOInit.BlendState = DecalRendering::GetDecalBlendState(DecalData.BlendDesc, DecalRenderStage, RenderTargetMode);
+			GraphicsPSOInit.BlendState = DecalRendering::GetDecalBlendState(VisibleDecal.BlendDesc, DecalRenderStage, RenderTargetMode);
 
 			// Set shader params
-			DecalRendering::SetShader(RHICmdList, GraphicsPSOInit, 0, View, DecalData, DecalRenderStage, FrustumComponentToClip, &Scene);
+			DecalRendering::SetShader(RHICmdList, GraphicsPSOInit, 0, View, VisibleDecal, DecalRenderStage, FrustumComponentToClip, &Scene);
 
 			RHICmdList.DrawIndexedPrimitive(GetUnitCubeIndexBuffer(), 0, 0, 8, 0, UE_ARRAY_COUNT(GCubeIndices) / 3, View.GetStereoPassInstanceFactor());
 		}
@@ -169,9 +167,9 @@ void FMobileSceneRenderer::RenderDBuffer(FRDGBuilder& GraphBuilder, FSceneTextur
 			continue;
 		}
 
-		FTransientDecalRenderDataList VisibleDecals = DecalRendering::BuildVisibleDecalList(Scene->Decals, View);
-
+		FVisibleDecalList VisibleDecals = DecalRendering::BuildVisibleDecalList(Scene->Decals, View);
+		FRelevantDecalList SortedDecals = DecalRendering::BuildRelevantDecalList(VisibleDecals, EDecalRenderStage::BeforeBasePass);
 		FDeferredDecalPassTextures DecalPassTextures = GetDeferredDecalPassTextures(GraphBuilder, View, Scene->SubstrateSceneData, SceneTextures, &DBufferTextures, EDecalRenderStage::BeforeBasePass);
-		AddDeferredDecalPass(GraphBuilder, View, VisibleDecals, DecalPassTextures, InstanceCullingManager, EDecalRenderStage::BeforeBasePass);
+		AddDeferredDecalPass(GraphBuilder, View, SortedDecals, DecalPassTextures, InstanceCullingManager, EDecalRenderStage::BeforeBasePass);
 	}
 }
diff --git a/Engine/Source/Runtime/Renderer/Private/RayTracing/RayTracingDecals.cpp b/Engine/Source/Runtime/Renderer/Private/RayTracing/RayTracingDecals.cpp
--- a/Engine/Source/Runtime/Renderer/Private/RayTracing/RayTracingDecals.cpp
+++ b/Engine/Source/Runtime/Renderer/Private/RayTracing/RayTracingDecals.cpp
@@ -77,39 +77,37 @@ IMPLEMENT_GLOBAL_SHADER_PARAMETER_STRUCT(FDecalParametersRayTracing, "RayTracing
 static TUniformBufferRef<FDecalParametersRayTracing> CreateDecalParametersBuffer(
 	const FViewInfo& View,
 	EShaderPlatform Platform,
-	const FTransientDecalRenderData& DecalData,
+	const FVisibleDecal& VisibleDecal,
 	EUniformBufferUsage Usage)
 {
-	const FDeferredDecalProxy& DecalProxy = *DecalData.Proxy;
-
 	FDecalParametersRayTracing Parameters;
 
-	const FMatrix DecalToWorldMatrix = DecalProxy.ComponentTrans.ToMatrixWithScale();
-	const FMatrix WorldToDecalMatrix = DecalProxy.ComponentTrans.ToInverseMatrixWithScale();
+	const FMatrix DecalToWorldMatrix = VisibleDecal.ComponentTrans.ToMatrixWithScale();
+	const FMatrix WorldToDecalMatrix = VisibleDecal.ComponentTrans.ToInverseMatrixWithScale();
 	const FDFVector3 AbsoluteOrigin(DecalToWorldMatrix.GetOrigin());
 	const FVector3f PositionHigh = AbsoluteOrigin.High;
 	const FMatrix44f RelativeDecalToWorldMatrix = FDFMatrix::MakeToRelativeWorldMatrix(PositionHigh, DecalToWorldMatrix).M;
-	const FVector3f OrientationVector = (FVector3f)DecalProxy.ComponentTrans.GetUnitAxis(EAxis::X);
+	const FVector3f OrientationVector = (FVector3f)VisibleDecal.ComponentTrans.GetUnitAxis(EAxis::X);
 
 	Parameters.DecalPositionHigh = PositionHigh;
 	Parameters.WorldToDecal = FMatrix44f(FTranslationMatrix(-View.ViewMatrices.GetPreViewTranslation()) * WorldToDecalMatrix);
 	Parameters.DecalToWorld = RelativeDecalToWorldMatrix;
 	Parameters.DecalToWorldInvScale = static_cast<FVector3f>(DecalToWorldMatrix.GetScaleVector().Reciprocal());
 	Parameters.DecalOrientation = OrientationVector;
-	Parameters.DecalColorParam = DecalData.DecalColor;
+	Parameters.DecalColorParam = VisibleDecal.DecalColor;
 
 	float LifetimeAlpha = 1.0f;
 
 	// Certain engine captures (e.g. environment reflection) don't have a tick. Default to fully opaque.
 	if (View.Family->Time.GetWorldTimeSeconds())
 	{
-		LifetimeAlpha = FMath::Clamp(FMath::Min(View.Family->Time.GetWorldTimeSeconds() * -DecalProxy.InvFadeDuration + DecalProxy.FadeStartDelayNormalized, View.Family->Time.GetWorldTimeSeconds() * DecalProxy.InvFadeInDuration + DecalProxy.FadeInStartDelayNormalized), 0.0f, 1.0f);
+		LifetimeAlpha = FMath::Clamp(FMath::Min(View.Family->Time.GetWorldTimeSeconds() * -VisibleDecal.InvFadeDuration + VisibleDecal.FadeStartDelayNormalized, View.Family->Time.GetWorldTimeSeconds() * VisibleDecal.InvFadeInDuration + VisibleDecal.FadeInStartDelayNormalized), 0.0f, 1.0f);
 	}
 
-	Parameters.DecalParams = FVector2f(DecalData.FadeAlpha, LifetimeAlpha);
-	Parameters.DecalColorParam = DecalData.DecalColor;
+	Parameters.DecalParams = FVector2f(VisibleDecal.FadeAlpha, LifetimeAlpha);
+	Parameters.DecalColorParam = VisibleDecal.DecalColor;
 
-	Parameters.DecalWriteFlags = CalculateDecalWriteFlags(DecalData.BlendDesc);
+	Parameters.DecalWriteFlags = CalculateDecalWriteFlags(VisibleDecal.BlendDesc);
 
 	return CreateUniformBufferImmediate(Parameters, Usage);
 }
@@ -467,13 +465,16 @@ void BuildDecalGrid(FRDGBuilder& GraphBuilder, uint32 NumDecals, FRDGBufferSRVRe
 	}
 }
 
-FTransientDecalRenderDataList GetSortedDecals(TConstArrayView<FDeferredDecalProxy*> Decals, FScene& Scene, const FViewInfo& View)
+FRelevantDecalList GetSortedDecals(FVisibleDecalList& VisibleDecals, TConstArrayView<FDeferredDecalProxy*> Decals, FScene& Scene, const FViewInfo& View)
 {
 	const float FadeMultiplier = DecalRendering::GetDecalFadeScreenSizeMultiplier();
 	const bool bIsPerspectiveProjection = View.IsPerspectiveProjection();
 
-	FTransientDecalRenderDataList SortedDecals;
+	FRelevantDecalList SortedDecals;
+
+	// NOTE: Reserve is critical here to avoid reallocation of the underlying array, since the sorted decals array holds a pointer!
 	SortedDecals.Reserve(Decals.Num());
+	VisibleDecals.Reserve(Decals.Num());
 
 	for (auto DecalProxy : Decals)
 	{
@@ -503,14 +504,13 @@ FTransientDecalRenderDataList GetSortedDecals(TConstArrayView<FDeferredDecalProx
 
 			if (bShouldRender)
 			{
-				FTransientDecalRenderData Data(*DecalProxy, 0.0f, FadeAlpha, Scene.GetShaderPlatform(), Scene.GetFeatureLevel());
-				SortedDecals.Add(MoveTemp(Data));
+				VisibleDecals.Emplace(*DecalProxy, 0.0f, FadeAlpha, Scene.GetShaderPlatform(), Scene.GetFeatureLevel());
+				SortedDecals.Emplace(&VisibleDecals.Last());
 			}
 		}
 	}
 
 	DecalRendering::SortDecalList(SortedDecals);
-
 	return SortedDecals;
 }
 
@@ -530,7 +530,8 @@ TRDGUniformBufferRef<FRayTracingDecals> CreateRayTracingDecalData(FRDGBuilder& G
 {
 	const EShaderPlatform ShaderPlatform = View.Family->GetShaderPlatform();
 
-	FTransientDecalRenderDataList SortedDecals = GetSortedDecals(Scene.Decals, Scene, View);
+	FVisibleDecalList VisibleDecals;
+	FRelevantDecalList SortedDecals = GetSortedDecals(VisibleDecals, Scene.Decals, Scene, View);
 
 	// Allocate from the graph builder so that we don't need to copy the data again when queuing the upload
 	auto& RayTracingDecals = GraphBuilder.AllocArray<FRayTracingDecal>();
@@ -542,9 +543,9 @@ TRDGUniformBufferRef<FRayTracingDecals> CreateRayTracingDecalData(FRDGBuilder& G
 
 	for (int32 Index = 0; Index < SortedDecals.Num(); ++Index)
 	{
-		const auto& DecalData = SortedDecals[Index];
+		const FVisibleDecal& VisibleDecal = *SortedDecals[Index];
 
-		const FMaterialRenderProxy* MaterialProxy = DecalData.MaterialProxy;
+		const FMaterialRenderProxy* MaterialProxy = VisibleDecal.MaterialProxy;
 		const FMaterial* MaterialResource = &MaterialProxy->GetMaterialWithFallback(View.GetFeatureLevel(), MaterialProxy);
 		const FMaterialShaderMap* MaterialShaderMap = MaterialResource->GetRenderingThreadShaderMap();
 
@@ -555,7 +556,7 @@ TRDGUniformBufferRef<FRayTracingDecals> CreateRayTracingDecalData(FRDGBuilder& G
 			continue;
 		}
 
-		TUniformBufferRef<FDecalParametersRayTracing> DecalParameters = CreateDecalParametersBuffer(View, ShaderPlatform, DecalData, EUniformBufferUsage::UniformBuffer_SingleFrame);
+		TUniformBufferRef<FDecalParametersRayTracing> DecalParameters = CreateDecalParametersBuffer(View, ShaderPlatform, VisibleDecal, EUniformBufferUsage::UniformBuffer_SingleFrame);
 
 		Scene.RayTracingScene.UniformBuffers.Add(DecalParameters); // Hold uniform buffer ref in RayTracingScene since FMeshDrawSingleShaderBindings doesn't
 
@@ -566,14 +567,14 @@ TRDGUniformBufferRef<FRayTracingDecals> CreateRayTracingDecalData(FRDGBuilder& G
 		FMeshDrawSingleShaderBindings SingleShaderBindings = Command.ShaderBindings.GetSingleShaderBindings(SF_RayCallable);
 		CallableShader->GetShaderBindings(&Scene, Scene.GetFeatureLevel(), *MaterialProxy, *MaterialResource, View, DecalParameters, SingleShaderBindings);
 
-		const FBox BoxBounds = DecalData.Proxy->GetBounds().GetBox();
+		const FBox BoxBounds = VisibleDecal.BoxBounds;
 
 		FRayTracingDecal& DestDecal = RayTracingDecals.AddDefaulted_GetRef();
 		DestDecal.TranslatedBoundMin = FVector3f(BoxBounds.Min + View.ViewMatrices.GetPreViewTranslation());
 		DestDecal.TranslatedBoundMax = FVector3f(BoxBounds.Max + View.ViewMatrices.GetPreViewTranslation());
 		DestDecal.CallableSlotIndex = BaseCallableSlotIndex + Index;
 		{
-			const FMatrix WorldToDecalMatrix = DecalData.Proxy->ComponentTrans.ToInverseMatrixWithScale();
+			const FMatrix WorldToDecalMatrix = VisibleDecal.ComponentTrans.ToInverseMatrixWithScale();
 			DestDecal.TranslatedWorldToDecal = FMatrix44f(FTranslationMatrix(-View.ViewMatrices.GetPreViewTranslation()) * WorldToDecalMatrix);
 		}
 		TranslatedBoundMin = FVector3f::Min(TranslatedBoundMin, DestDecal.TranslatedBoundMin);
diff --git a/Engine/Source/Runtime/Renderer/Private/SceneVisibility.cpp b/Engine/Source/Runtime/Renderer/Private/SceneVisibility.cpp
--- a/Engine/Source/Runtime/Renderer/Private/SceneVisibility.cpp
+++ b/Engine/Source/Runtime/Renderer/Private/SceneVisibility.cpp
@@ -55,6 +55,7 @@
 #include "ViewDebug.h"
 #include "DecalRenderingCommon.h"
 #include "CompositionLighting/PostProcessDeferredDecals.h"
+#include "DecalRenderingShared.h"
 
 bool GDistanceCullToSphereEdge = true;
 static FAutoConsoleVariableRef CVarDistanceCullToSphereEdge(
@@ -5684,6 +5685,11 @@ void FDeferredShadingSceneRenderer::BeginInitViews(
 	// Start processing dynamic mesh elements tasks early enough to overlap with shadows and GPU scene update.
 	TaskDatas.VisibilityTaskData->StartGatherDynamicMeshElements();
 
+	if (!ViewFamily.EngineShowFlags.HitProxies)
+	{
+		TaskDatas.Decals = FDecalVisibilityTaskData::Launch(GraphBuilder, *Scene, Views);
+	}
+
 	// Attempt to launch dynamic shadow tasks early before finalizing visibility.
 	if (bRendererOutputFinalSceneColor)
 	{
