[HWRT] Fix crash enabling Lumen visualization when project is using Hardware Ray Tracing Mode = Inline.

diff --git a/Engine/Source/Runtime/Renderer/Private/DeferredShadingRenderer.cpp b/Engine/Source/Runtime/Renderer/Private/DeferredShadingRenderer.cpp
--- a/Engine/Source/Runtime/Renderer/Private/DeferredShadingRenderer.cpp
+++ b/Engine/Source/Runtime/Renderer/Private/DeferredShadingRenderer.cpp
@@ -619,7 +619,8 @@ bool FDeferredShadingSceneRenderer::SetupRayTracingPipelineStatesAndSBT(FRDGBuil
 	}
 
 	const bool bIsPathTracing = ViewFamily.EngineShowFlags.PathTracing;
-
+	
+	if (GRHISupportsRayTracingShaders)
 	{
 		// #dxr_todo: UE-72565: refactor ray tracing effects to not be member functions of DeferredShadingRenderer. 
 		// Should register each effect at startup and just loop over them automatically to gather all required shaders.
@@ -685,31 +686,35 @@ bool FDeferredShadingSceneRenderer::SetupRayTracingPipelineStatesAndSBT(FRDGBuil
 	{
 		TArray<FRHIRayTracingShader*> LumenHardwareRayTracingRayGenShaders;
 
-		if (DoesPlatformSupportLumenGI(ShaderPlatform))
+		if (GRHISupportsRayTracingShaders)
 		{
+			if (DoesPlatformSupportLumenGI(ShaderPlatform))
+			{
+				for (const FViewInfo& View : Views)
+				{
+					PrepareLumenHardwareRayTracingVisualizeLumenMaterial(View, LumenHardwareRayTracingRayGenShaders);
+					PrepareLumenHardwareRayTracingRadianceCacheLumenMaterial(View, LumenHardwareRayTracingRayGenShaders);
+					PrepareLumenHardwareRayTracingTranslucencyVolumeLumenMaterial(View, LumenHardwareRayTracingRayGenShaders);
+					PrepareLumenHardwareRayTracingRadiosityLumenMaterial(View, LumenHardwareRayTracingRayGenShaders);
+					PrepareLumenHardwareRayTracingReflectionsLumenMaterial(View, LumenHardwareRayTracingRayGenShaders);
+					PrepareLumenHardwareRayTracingReSTIRLumenMaterial(View, LumenHardwareRayTracingRayGenShaders);
+					PrepareLumenHardwareRayTracingScreenProbeGatherLumenMaterial(View, LumenHardwareRayTracingRayGenShaders);
+					PrepareLumenHardwareRayTracingDirectLightingLumenMaterial(View, LumenHardwareRayTracingRayGenShaders);
+				}
+			}
+
 			for (const FViewInfo& View : Views)
 			{
-				PrepareLumenHardwareRayTracingVisualizeLumenMaterial(View, LumenHardwareRayTracingRayGenShaders);
-				PrepareLumenHardwareRayTracingRadianceCacheLumenMaterial(View, LumenHardwareRayTracingRayGenShaders);
-				PrepareLumenHardwareRayTracingTranslucencyVolumeLumenMaterial(View, LumenHardwareRayTracingRayGenShaders);
-				PrepareLumenHardwareRayTracingRadiosityLumenMaterial(View, LumenHardwareRayTracingRayGenShaders);
-				PrepareLumenHardwareRayTracingReflectionsLumenMaterial(View, LumenHardwareRayTracingRayGenShaders);
-				PrepareLumenHardwareRayTracingReSTIRLumenMaterial(View, LumenHardwareRayTracingRayGenShaders);
-				PrepareLumenHardwareRayTracingScreenProbeGatherLumenMaterial(View, LumenHardwareRayTracingRayGenShaders);
-				PrepareLumenHardwareRayTracingDirectLightingLumenMaterial(View, LumenHardwareRayTracingRayGenShaders);
+				PrepareMegaLightsHardwareRayTracingLumenMaterial(View, LumenHardwareRayTracingRayGenShaders);
 			}
-		}
 
-		for (const FViewInfo& View : Views)
-		{
-			PrepareMegaLightsHardwareRayTracingLumenMaterial(View, LumenHardwareRayTracingRayGenShaders);
+			DeduplicateRayGenerationShaders(LumenHardwareRayTracingRayGenShaders);
 		}
-
-		DeduplicateRayGenerationShaders(LumenHardwareRayTracingRayGenShaders);
 		
 		uint32 MaxLocalBindingDataSize = 0;
 		ERayTracingShaderBindingMode ShaderBindingMode = (bAnyLumenHardwareInlineRayTracingPassEnabled && GRHIGlobals.RayTracing.RequiresInlineRayTracingSBT) ? 
 			ERayTracingShaderBindingMode::Inline : ERayTracingShaderBindingMode::Disabled;
+
 		if (LumenHardwareRayTracingRayGenShaders.Num())
 		{
 			CreateLumenHardwareRayTracingMaterialPipeline(GraphBuilder, ReferenceView, LumenHardwareRayTracingRayGenShaders, MaxLocalBindingDataSize);
