Fix crash due to SetupDummyForwardLightUniformParameters(...) not setting a dummy buffer for LightViewData.

diff --git a/Engine/Source/Runtime/Renderer/Private/LightGridInjection.cpp b/Engine/Source/Runtime/Renderer/Private/LightGridInjection.cpp
--- a/Engine/Source/Runtime/Renderer/Private/LightGridInjection.cpp
+++ b/Engine/Source/Runtime/Renderer/Private/LightGridInjection.cpp
@@ -202,6 +202,9 @@ void SetupDummyForwardLightUniformParameters(FRDGBuilder& GraphBuilder, FForward
 	ForwardLightData.LightFunctionAtlasLightIndex = 0;
 
 	ForwardLightData.bAffectsTranslucentLighting = 0;
+
+	FRDGBufferRef LightViewDataBuffer = GSystemTextures.GetDefaultStructuredBuffer(GraphBuilder, sizeof(FLightViewData));
+	ForwardLightData.LightViewData = GraphBuilder.CreateSRV(LightViewDataBuffer);
 }
 
 TRDGUniformBufferRef<FForwardLightData> CreateDummyForwardLightUniformBuffer(FRDGBuilder& GraphBuilder, EShaderPlatform ShaderPlatform)
