Fix crash in FSceneRenderer::PrepareForwardLightData(...) when toggling platform preview.

diff --git a/Engine/Source/Runtime/Renderer/Private/LightGridInjection.cpp b/Engine/Source/Runtime/Renderer/Private/LightGridInjection.cpp
--- a/Engine/Source/Runtime/Renderer/Private/LightGridInjection.cpp
+++ b/Engine/Source/Runtime/Renderer/Private/LightGridInjection.cpp
@@ -1289,6 +1289,11 @@ FComputeLightGridOutput FSceneRenderer::PrepareForwardLightData(FRDGBuilder& Gra
 
 					const FLightSceneInfo::FPersistentId LightSceneId = LightSceneInfo->Id;
 
+					if (!ensureMsgf(LightViewData.IsValidIndex(LightSceneId), TEXT("Visible light is missing from GPU Scene.")))
+					{
+						continue;
+					}
+
 					if (!ensureMsgf(Scene->Lights.IsAllocated(LightSceneId), TEXT("Visible light is missing from GPU Scene.")))
 					{
 						continue;
