Heterogeneous Volumes: Correct lighting cache bounds culling at scale

Using truncated matrices is incorrectly culling lights at scale.

diff --git a/Engine/Source/Runtime/Renderer/Private/HeterogeneousVolumes/HeterogeneousVolumesLiveShadingPipeline.cpp b/Engine/Source/Runtime/Renderer/Private/HeterogeneousVolumes/HeterogeneousVolumesLiveShadingPipeline.cpp
--- a/Engine/Source/Runtime/Renderer/Private/HeterogeneousVolumes/HeterogeneousVolumesLiveShadingPipeline.cpp
+++ b/Engine/Source/Runtime/Renderer/Private/HeterogeneousVolumes/HeterogeneousVolumesLiveShadingPipeline.cpp
@@ -705,6 +705,13 @@ static void RenderLightingCacheWithLiveShading(
 		DilatedExistenceTexture
 	);
 
+	FMatrix InstanceToLocal = HeterogeneousVolumeInterface->GetInstanceToLocal();
+	FMatrix LocalToWorld = HeterogeneousVolumeInterface->GetLocalToWorld();
+	FMatrix InstanceToWorld = InstanceToLocal * LocalToWorld;
+	FMatrix WorldToInstance = InstanceToWorld.Inverse();
+	FMatrix LocalToInstance = InstanceToLocal.Inverse();
+	FBoxSphereBounds InstanceBoxSphereBounds = LocalBoxSphereBounds.TransformBy(LocalToInstance);
+
 	FRenderLightingCacheWithLiveShadingCS::FParameters* PassParameters = GraphBuilder.AllocParameters<FRenderLightingCacheWithLiveShadingCS::FParameters>();
 	{
 		// Scene data
@@ -725,13 +732,9 @@ static void RenderLightingCacheWithLiveShading(
 		// TODO: Convert to relative-local space
 		//FVector3f ViewOriginHigh = FDFVector3(View.ViewMatrices.GetViewOrigin()).High;
 		//FMatrix44f RelativeLocalToWorld = FDFMatrix::MakeToRelativeWorldMatrix(ViewOriginHigh, HeterogeneousVolumeInterface->GetLocalToWorld()).M;
-		FMatrix InstanceToLocal = HeterogeneousVolumeInterface->GetInstanceToLocal();
-		FMatrix LocalToWorld = HeterogeneousVolumeInterface->GetLocalToWorld();
-		PassParameters->LocalToWorld = FMatrix44f(InstanceToLocal * LocalToWorld);
-		PassParameters->WorldToLocal = PassParameters->LocalToWorld.Inverse();
+		PassParameters->LocalToWorld = FMatrix44f(InstanceToWorld);
+		PassParameters->WorldToLocal = FMatrix44f(WorldToInstance);
 
-		FMatrix LocalToInstance = InstanceToLocal.Inverse();
-		FBoxSphereBounds InstanceBoxSphereBounds = LocalBoxSphereBounds.TransformBy(LocalToInstance);
 		PassParameters->LocalBoundsOrigin = FVector3f(InstanceBoxSphereBounds.Origin);
 		PassParameters->LocalBoundsExtent = FVector3f(InstanceBoxSphereBounds.BoxExtent);
 		PassParameters->PrimitiveId = PersistentPrimitiveIndex.Index;
@@ -803,12 +806,8 @@ static void RenderLightingCacheWithLiveShading(
 	bool bShouldBoundsCull = HeterogeneousVolumes::ShouldBoundsCull();
 	if (LightType != LightType_Directional && bShouldBoundsCull)
 	{
-		auto ToFVector3f = [](FVector4 V4f) {
-			return FVector3f(V4f.X, V4f.Y, V4f.Z);
-		};
-
 		auto FloorVector = [](const FVector& V) {
-			return FVector(
+			return FIntVector(
 				FMath::FloorToFloat(V.X),
 				FMath::FloorToFloat(V.Y),
 				FMath::FloorToFloat(V.Z)
@@ -816,14 +815,14 @@ static void RenderLightingCacheWithLiveShading(
 		};
 
 		auto CeilVector = [](const FVector& V) {
-			return FVector(
+			return FIntVector(
 				FMath::CeilToFloat(V.X),
 				FMath::CeilToFloat(V.Y),
 				FMath::CeilToFloat(V.Z)
 			);
 		};
 
-		auto ClampVector = [](const FVector& V, const FIntVector& Min, const FIntVector& Max) {
+		auto ClampVector = [](const FIntVector& V, const FIntVector& Min, const FIntVector& Max) {
 			FIntVector IntV;
 			IntV.X = FMath::Clamp(V.X, Min.X, Max.X);
 			IntV.Y = FMath::Clamp(V.Y, Min.Y, Max.Y);
@@ -832,14 +831,14 @@ static void RenderLightingCacheWithLiveShading(
 		};
 
 		FSphere WorldLightBoundingSphere = LightSceneInfo->Proxy->GetBoundingSphere();
-		FVector LocalLightCenter = FVector(PassParameters->WorldToLocal.TransformPosition(ToFVector3f(WorldLightBoundingSphere.Center)));
-		FVector3f ScalingTerm = PassParameters->WorldToLocal.GetScaleVector();
-		FVector LocalLightExtent = FVector(ScalingTerm) * WorldLightBoundingSphere.W;
+		FVector LocalLightCenter = WorldToInstance.TransformPosition(WorldLightBoundingSphere.Center);
+		FVector LocalLightExtent = WorldToInstance.GetScaleVector() * WorldLightBoundingSphere.W;
+
 		FVector LocalLightMin = LocalLightCenter - LocalLightExtent;
 		FVector LocalLightMax = LocalLightCenter + LocalLightExtent;
 
-		FVector LightingCacheMin = LocalBoxSphereBounds.Origin - LocalBoxSphereBounds.BoxExtent;
-		FVector LightingCacheMax = LocalBoxSphereBounds.Origin + LocalBoxSphereBounds.BoxExtent;
+		FVector LightingCacheMin = InstanceBoxSphereBounds.Origin - InstanceBoxSphereBounds.BoxExtent;
+		FVector LightingCacheMax = InstanceBoxSphereBounds.Origin + InstanceBoxSphereBounds.BoxExtent;
 
 		FVector LocalLightMinUV = (LocalLightMin - LightingCacheMin) / (LightingCacheMax - LightingCacheMin);
 		FVector LocalLightMaxUV = (LocalLightMax - LightingCacheMin) / (LightingCacheMax - LightingCacheMin);
@@ -848,35 +847,39 @@ static void RenderLightingCacheWithLiveShading(
 		PassParameters->VoxelMax = ClampVector(CeilVector(LocalLightMaxUV * LightingCacheResolutionVector), FIntVector::ZeroValue, PassParameters->VoxelMax);
 	}
 
-	FIntVector GroupCount = PassParameters->VoxelMax - PassParameters->VoxelMin + FIntVector(1);
-	check(GroupCount.X > 0 && GroupCount.Y > 0 && GroupCount.Z > 0);
-	GroupCount.X = FMath::DivideAndRoundUp(GroupCount.X, FRenderLightingCacheWithLiveShadingCS::GetThreadGroupSize3D());
-	GroupCount.Y = FMath::DivideAndRoundUp(GroupCount.Y, FRenderLightingCacheWithLiveShadingCS::GetThreadGroupSize3D());
-	GroupCount.Z = FMath::DivideAndRoundUp(GroupCount.Z, FRenderLightingCacheWithLiveShadingCS::GetThreadGroupSize3D());
+	FIntVector VoxelDimensions = PassParameters->VoxelMax - PassParameters->VoxelMin;
+	if (VoxelDimensions.GetMin() > 0)
+	{
+		FIntVector GroupCount = PassParameters->VoxelMax - PassParameters->VoxelMin + FIntVector(1);
+		check(GroupCount.X > 0 && GroupCount.Y > 0 && GroupCount.Z > 0);
+		GroupCount.X = FMath::DivideAndRoundUp(GroupCount.X, FRenderLightingCacheWithLiveShadingCS::GetThreadGroupSize3D());
+		GroupCount.Y = FMath::DivideAndRoundUp(GroupCount.Y, FRenderLightingCacheWithLiveShadingCS::GetThreadGroupSize3D());
+		GroupCount.Z = FMath::DivideAndRoundUp(GroupCount.Z, FRenderLightingCacheWithLiveShadingCS::GetThreadGroupSize3D());
 
-	bool bUseAVSM = HeterogeneousVolumes::UseAdaptiveVolumetricShadowMapForSelfShadowing(HeterogeneousVolumeInterface->GetPrimitiveSceneProxy());
+		bool bUseAVSM = HeterogeneousVolumes::UseAdaptiveVolumetricShadowMapForSelfShadowing(HeterogeneousVolumeInterface->GetPrimitiveSceneProxy());
 
-	FRenderLightingCacheWithLiveShadingCS::FPermutationDomain PermutationVector;
-	PermutationVector.Set<FRenderLightingCacheWithLiveShadingCS::FLightingCacheMode>(HeterogeneousVolumes::GetLightingCacheMode() - 1);
-	PermutationVector.Set<FRenderLightingCacheWithLiveShadingCS::FAVSMSampleMode>(static_cast<int32>(HeterogeneousVolumes::GetAVSMSampleMode(bUseAVSM)));
-	PermutationVector.Set<FRenderLightingCacheWithLiveShadingCS::FUseExistenceMask>(HeterogeneousVolumes::UseExistenceMask());
-	PermutationVector.Set<FRenderLightingCacheWithLiveShadingCS::FIsOfflineRender>(View.bIsOfflineRender);
-	TShaderRef<FRenderLightingCacheWithLiveShadingCS> ComputeShader = Material.GetShader<FRenderLightingCacheWithLiveShadingCS>(&FLocalVertexFactory::StaticType, PermutationVector, false);
-	if (!ComputeShader.IsNull())
-	{
-		FRDGBufferRef IndirectArgsBuffer = GSystemTextures.GetDefaultBuffer(GraphBuilder, 4);
-		AddComputePass<false, HeterogeneousVolumes::EDispatchMode::DirectDispatch>(
-			GraphBuilder,
-			ComputeShader,
-			PassParameters,
-			Scene,
-			MaterialRenderProxy,
-			Material,
-			PassName,
-			GroupCount,
-			IndirectArgsBuffer,
-			0
-		);
+		FRenderLightingCacheWithLiveShadingCS::FPermutationDomain PermutationVector;
+		PermutationVector.Set<FRenderLightingCacheWithLiveShadingCS::FLightingCacheMode>(HeterogeneousVolumes::GetLightingCacheMode() - 1);
+		PermutationVector.Set<FRenderLightingCacheWithLiveShadingCS::FAVSMSampleMode>(static_cast<int32>(HeterogeneousVolumes::GetAVSMSampleMode(bUseAVSM)));
+		PermutationVector.Set<FRenderLightingCacheWithLiveShadingCS::FUseExistenceMask>(HeterogeneousVolumes::UseExistenceMask());
+		PermutationVector.Set<FRenderLightingCacheWithLiveShadingCS::FIsOfflineRender>(View.bIsOfflineRender);
+		TShaderRef<FRenderLightingCacheWithLiveShadingCS> ComputeShader = Material.GetShader<FRenderLightingCacheWithLiveShadingCS>(&FLocalVertexFactory::StaticType, PermutationVector, false);
+		if (!ComputeShader.IsNull())
+		{
+			FRDGBufferRef IndirectArgsBuffer = GSystemTextures.GetDefaultBuffer(GraphBuilder, 4);
+			AddComputePass<false, HeterogeneousVolumes::EDispatchMode::DirectDispatch>(
+				GraphBuilder,
+				ComputeShader,
+				PassParameters,
+				Scene,
+				MaterialRenderProxy,
+				Material,
+				PassName,
+				GroupCount,
+				IndirectArgsBuffer,
+				0
+			);
+		}
 	}
 }
 
