1-Add another bookmark for deformer graph + fix enqueue trigger graph

diff --git a/Engine/Plugins/Animation/DeformerGraph/Source/OptimusCore/Private/OptimusDeformerInstance.cpp b/Engine/Plugins/Animation/DeformerGraph/Source/OptimusCore/Private/OptimusDeformerInstance.cpp
--- a/Engine/Plugins/Animation/DeformerGraph/Source/OptimusCore/Private/OptimusDeformerInstance.cpp
+++ b/Engine/Plugins/Animation/DeformerGraph/Source/OptimusCore/Private/OptimusDeformerInstance.cpp
@@ -515,9 +515,12 @@ void UOptimusDeformerInstance::EnqueueWork(UMeshDeformerInstance::FEnqueueWorkDe
 		ExecutionGroupName = ComputeTaskExecutionGroup::Immediate;
 		break;
 	case UMeshDeformerInstance::ExecutionGroup_Default:
-	case UMeshDeformerInstance::ExecutionGroup_EndOfFrameUodate:
+	case UMeshDeformerInstance::ExecutionGroup_EndOfFrameUpdate:
 		ExecutionGroupName = ComputeTaskExecutionGroup::EndOfFrameUpdate;
 		break;
+	case UMeshDeformerInstance::ExecutionGroup_BeginInitViews:
+		ExecutionGroupName = ComputeTaskExecutionGroup::BeginInitViews;
+		break;
 	default:
 		ensure(0);
 		return;
diff --git a/Engine/Source/Runtime/Engine/Classes/Animation/MeshDeformerInstance.h b/Engine/Source/Runtime/Engine/Classes/Animation/MeshDeformerInstance.h
--- a/Engine/Source/Runtime/Engine/Classes/Animation/MeshDeformerInstance.h
+++ b/Engine/Source/Runtime/Engine/Classes/Animation/MeshDeformerInstance.h
@@ -59,7 +59,8 @@ public:
 	{
 		ExecutionGroup_Default,
 		ExecutionGroup_Immediate,
-		ExecutionGroup_EndOfFrameUodate,
+		ExecutionGroup_EndOfFrameUpdate,
+		ExecutionGroup_BeginInitViews,
 	};
 
 	/** Structure of inputs to EnqueueWork. */
diff --git a/Engine/Source/Runtime/Engine/Private/ComputeWorkerInterface.cpp b/Engine/Source/Runtime/Engine/Private/ComputeWorkerInterface.cpp
--- a/Engine/Source/Runtime/Engine/Private/ComputeWorkerInterface.cpp
+++ b/Engine/Source/Runtime/Engine/Private/ComputeWorkerInterface.cpp
@@ -5,3 +5,4 @@
 
 FName ComputeTaskExecutionGroup::Immediate("Immediate");
 FName ComputeTaskExecutionGroup::EndOfFrameUpdate("EndOfFrameUpdate");
+FName ComputeTaskExecutionGroup::BeginInitViews("BeginInitViews");
diff --git a/Engine/Source/Runtime/Engine/Public/ComputeWorkerInterface.h b/Engine/Source/Runtime/Engine/Public/ComputeWorkerInterface.h
--- a/Engine/Source/Runtime/Engine/Public/ComputeWorkerInterface.h
+++ b/Engine/Source/Runtime/Engine/Public/ComputeWorkerInterface.h
@@ -28,4 +28,5 @@ struct ComputeTaskExecutionGroup
 {
 	static ENGINE_API FName Immediate;
 	static ENGINE_API FName EndOfFrameUpdate;
+	static ENGINE_API FName BeginInitViews;
 };
diff --git a/Engine/Source/Runtime/Renderer/Private/SceneVisibility.cpp b/Engine/Source/Runtime/Renderer/Private/SceneVisibility.cpp
--- a/Engine/Source/Runtime/Renderer/Private/SceneVisibility.cpp
+++ b/Engine/Source/Runtime/Renderer/Private/SceneVisibility.cpp
@@ -3,6 +3,7 @@
 #include "SceneVisibility.h"
 #include "SceneVisibilityPrivate.h"
 #include "CoreMinimal.h"
+#include "ComputeWorkerInterface.h"
 #include "HAL/ThreadSafeCounter.h"
 #include "Stats/Stats.h"
 #include "Misc/MemStack.h"
@@ -4842,6 +4843,14 @@ void FSceneRenderer::PreVisibilityFrameSetup(FRDGBuilder& GraphBuilder)
 			}
 		}
 
+		for (IComputeTaskWorker* ComputeTaskWorker : Scene->ComputeTaskWorkers)
+		{
+			if (ComputeTaskWorker->HasWork(ComputeTaskExecutionGroup::BeginInitViews))
+			{
+				ComputeTaskWorker->SubmitWork(GraphBuilder, ComputeTaskExecutionGroup::BeginInitViews, FeatureLevel);
+			}
+		}
+
 		// Notify the FX system that the scene is about to perform visibility checks.
 
 		if (FXSystem && Views.IsValidIndex(0))
