Fix CIS in SceneVisibility

diff --git a/Engine/Source/Runtime/Renderer/Private/SceneVisibility.cpp b/Engine/Source/Runtime/Renderer/Private/SceneVisibility.cpp
--- a/Engine/Source/Runtime/Renderer/Private/SceneVisibility.cpp
+++ b/Engine/Source/Runtime/Renderer/Private/SceneVisibility.cpp
@@ -4843,11 +4843,14 @@ void FSceneRenderer::PreVisibilityFrameSetup(FRDGBuilder& GraphBuilder)
 			}
 		}
 
-		for (IComputeTaskWorker* ComputeTaskWorker : Scene->ComputeTaskWorkers)
+		if(Scene)
 		{
-			if (ComputeTaskWorker->HasWork(ComputeTaskExecutionGroup::BeginInitViews))
+			for (IComputeTaskWorker* ComputeTaskWorker : Scene->ComputeTaskWorkers)
 			{
-				ComputeTaskWorker->SubmitWork(GraphBuilder, ComputeTaskExecutionGroup::BeginInitViews, FeatureLevel);
+				if (ComputeTaskWorker->HasWork(ComputeTaskExecutionGroup::BeginInitViews))
+				{
+					ComputeTaskWorker->SubmitWork(GraphBuilder, ComputeTaskExecutionGroup::BeginInitViews, FeatureLevel);
+				}
 			}
 		}
 
