Fixed: Post Process Blend mode "Scene Color Before Bloom" doesn't work on mobile platforms

diff --git a/Engine/Source/Runtime/Renderer/Private/PostProcess/PostProcessing.cpp b/Engine/Source/Runtime/Renderer/Private/PostProcess/PostProcessing.cpp
--- a/Engine/Source/Runtime/Renderer/Private/PostProcess/PostProcessing.cpp
+++ b/Engine/Source/Runtime/Renderer/Private/PostProcess/PostProcessing.cpp
@@ -2858,6 +2858,8 @@ void AddMobilePostProcessingPasses(FRDGBuilder& GraphBuilder, FScene* Scene, con
 		PassSequence.SetEnabled(EPass::FXAA, false);
 		PassSequence.Finalize();
 	}
+
+	AddPostProcessMaterialPass(BL_SceneColorBeforeBloom, false);
 	
 	if (PassSequence.IsEnabled(EPass::Tonemap))
 	{
