Prepare ForwardLightData even when !ViewFamily.EngineShowFlags.Lighting since VSM checks it.

#tests Level without any lights but r.Shadow.Virtual.Enable 1

diff --git a/Engine/Source/Runtime/Renderer/Private/DeferredShadingRenderer.cpp b/Engine/Source/Runtime/Renderer/Private/DeferredShadingRenderer.cpp
--- a/Engine/Source/Runtime/Renderer/Private/DeferredShadingRenderer.cpp
+++ b/Engine/Source/Runtime/Renderer/Private/DeferredShadingRenderer.cpp
@@ -2267,7 +2267,7 @@ void FDeferredShadingSceneRenderer::Render(FRDGBuilder& GraphBuilder)
 
 			GatherAndSortLightsTask.Wait();
 
-			if (ViewFamily.EngineShowFlags.Lighting && !ViewFamily.EngineShowFlags.PathTracing)
+			if (!ViewFamily.EngineShowFlags.PathTracing)
 			{
 				ComputeLightGridOutput = PrepareForwardLightData(GraphBuilder, true, *SortedLightSet);
 
