Update full cached RayTracing State (Cached RT MDC and Cached instance) when processing scene update requests instead of only the cache instances

FStaticMeshStreamOut updates CurrentFirstLODIdx on game thread, and only invalidates cached ray tracing data later on render thread, so updating only cached ray tracing instances is not safe.

diff --git a/Engine/Source/Runtime/Renderer/Private/PrimitiveSceneInfo.cpp b/Engine/Source/Runtime/Renderer/Private/PrimitiveSceneInfo.cpp
--- a/Engine/Source/Runtime/Renderer/Private/PrimitiveSceneInfo.cpp
+++ b/Engine/Source/Runtime/Renderer/Private/PrimitiveSceneInfo.cpp
@@ -957,19 +957,13 @@ void FScene::RefreshRayTracingMeshCommandCache()
 }
 
 void FScene::RefreshRayTracingInstances()
-{
-	// Re-cache all current primitives
-	FPrimitiveSceneInfo::UpdateCachedRayTracingInstances(this, Primitives);
-}
-
-void FPrimitiveSceneInfo::UpdateCachedRayTracingInstances(FScene* Scene, const TArrayView<FPrimitiveSceneInfo*>& SceneInfos)
 {
 	if (IsRayTracingEnabled())
 	{
 		SCOPED_NAMED_EVENT(FPrimitiveSceneInfo_UpdateCachedRayTracingInstances, FColor::Turquoise);
 		checkf(GRHISupportsMultithreadedShaderCreation, TEXT("Raytracing code needs the ability to create shaders from task threads."));
 
-		for (FPrimitiveSceneInfo* SceneInfo : SceneInfos)
+		for (FPrimitiveSceneInfo* SceneInfo : Primitives)
 		{
 			const int32 PrimitiveIndex = SceneInfo->GetIndex();
 
@@ -979,20 +973,20 @@ void FPrimitiveSceneInfo::UpdateCachedRayTracingInstances(FScene* Scene, const T
 				const int32 RayTracingGroupId = SceneInfo->Proxy->GetRayTracingGroupId();
 				if (RayTracingGroupId != -1)
 				{
-					SceneRayTracingGroupId = Scene->PrimitiveRayTracingGroups.FindId(RayTracingGroupId);
+					SceneRayTracingGroupId = PrimitiveRayTracingGroups.FindId(RayTracingGroupId);
 				}
 
-				check(Scene->PrimitiveRayTracingGroupIds[PrimitiveIndex] == SceneRayTracingGroupId);
+				check(PrimitiveRayTracingGroupIds[PrimitiveIndex] == SceneRayTracingGroupId);
 				check(SceneInfo->CoarseMeshStreamingHandle == SceneInfo->Proxy->GetCoarseMeshStreamingHandle());
 			}
 #endif
 
 			FRayTracingInstance CachedRayTracingInstance;
-			ERayTracingPrimitiveFlags& Flags = Scene->PrimitiveRayTracingFlags[PrimitiveIndex];
+			ERayTracingPrimitiveFlags& Flags = PrimitiveRayTracingFlags[PrimitiveIndex];
 
 			// Write flags
 			Flags = SceneInfo->Proxy->GetCachedRayTracingInstance(CachedRayTracingInstance);
-			UpdateCachedRayTracingInstance(SceneInfo, CachedRayTracingInstance, Flags);
+			FPrimitiveSceneInfo::UpdateCachedRayTracingInstance(SceneInfo, CachedRayTracingInstance, Flags);
 		}
 	}
 }
diff --git a/Engine/Source/Runtime/Renderer/Private/RendererScene.cpp b/Engine/Source/Runtime/Renderer/Private/RendererScene.cpp
--- a/Engine/Source/Runtime/Renderer/Private/RendererScene.cpp
+++ b/Engine/Source/Runtime/Renderer/Private/RendererScene.cpp
@@ -6336,12 +6336,6 @@ void FScene::Update(FRDGBuilder& GraphBuilder, const FUpdateParameters& Paramete
 	}
 
 
-#if RHI_RAYTRACING
-	TArray<FPrimitiveSceneInfo*, SceneRenderingAllocator> RayTracingPrimitivesToUpdate;
-	RayTracingPrimitivesToUpdate.Reserve(UpdatedInstances.Num());
-	bool bUpdateCachedRayTracingInstances = false;
-#endif
-
 	// Updating step of instance update. Note that instance IDs are not allocated at this point.
 	{
 		CSV_SCOPED_TIMING_STAT_EXCLUSIVE(UpdatePrimitiveInstances);
@@ -6381,8 +6375,7 @@ void FScene::Update(FRDGBuilder& GraphBuilder, const FUpdateParameters& Paramete
 #if RHI_RAYTRACING
 			if (!PrimitiveSceneInfo->bPendingAddStaticMeshes)
 			{
-				RayTracingPrimitivesToUpdate.Add(PrimitiveSceneInfo);
-				bUpdateCachedRayTracingInstances = true;
+				UpdateCachedRayTracingState(PrimitiveSceneProxy);
 			}
 #endif
 
@@ -6493,13 +6486,6 @@ void FScene::Update(FRDGBuilder& GraphBuilder, const FUpdateParameters& Paramete
 		FPrimitiveSceneInfo::AddToScene(this, SceneInfosWithAddToScene);
 	}
 
-#if RHI_RAYTRACING
-	if (bUpdateCachedRayTracingInstances)
-	{
-		FPrimitiveSceneInfo::UpdateCachedRayTracingInstances(this, RayTracingPrimitivesToUpdate);
-	}
-#endif
-
 	for (const FLevelCommand& Command : LevelCommands)
 	{
 		const auto InvalidateNaniteGPUScene = [](FPrimitiveSceneInfo* Primitive)
diff --git a/Engine/Source/Runtime/Renderer/Public/PrimitiveSceneInfo.h b/Engine/Source/Runtime/Renderer/Public/PrimitiveSceneInfo.h
--- a/Engine/Source/Runtime/Renderer/Public/PrimitiveSceneInfo.h
+++ b/Engine/Source/Runtime/Renderer/Public/PrimitiveSceneInfo.h
@@ -806,7 +806,6 @@ private:
 	void UpdateCachedRayTracingInstanceWorldBounds(const FMatrix& NewPrimitiveLocalToWorld) {};
 
 	/** Updates cached ray tracing instances. Utility closely mirrors CacheRayTracingPrimitives(..) */
-	static void UpdateCachedRayTracingInstances(FScene* Scene, const TArrayView<FPrimitiveSceneInfo*>& SceneInfos);
 	static void UpdateCachedRayTracingInstance(FPrimitiveSceneInfo* SceneInfo, const FRayTracingInstance& RayTracingInstance, const ERayTracingPrimitiveFlags Flags);
 #endif
 
