Added check to catch when a primitive is queued for static mesh update while also being queued for deletion.

diff --git a/Engine/Source/Runtime/Renderer/Private/RendererScene.cpp b/Engine/Source/Runtime/Renderer/Private/RendererScene.cpp
--- a/Engine/Source/Runtime/Renderer/Private/RendererScene.cpp
+++ b/Engine/Source/Runtime/Renderer/Private/RendererScene.cpp
@@ -5654,7 +5654,10 @@ void FScene::Update(FRDGBuilder& GraphBuilder, const FUpdateParameters& Paramete
 		}
 		if (Cmd.IsDelete())
 		{
-			DeletedPrimitiveSceneInfos.Add(Cmd.GetSceneInfo());
+			FPrimitiveSceneInfo* SceneInfo = Cmd.GetSceneInfo();
+			SceneInfo->bPendingDelete = true;
+
+			DeletedPrimitiveSceneInfos.Add(SceneInfo);
 			// Skip those that were added & removed in the same frame
 			if (!Cmd.IsAdd())
 			{
@@ -5963,6 +5966,11 @@ void FScene::Update(FRDGBuilder& GraphBuilder, const FUpdateParameters& Paramete
 
 	const auto QueueAddStaticMeshes = [&](FPrimitiveSceneInfo* SceneInfo)
 	{
+		if (SceneInfo->bPendingDelete)
+		{
+			UE_LOG(LogRenderer, Fatal, TEXT("SceneInfo %s is marked for deletion but is being queued for static mesh updates. This is a race condition and will crash."), *SceneInfo->Proxy->GetOwnerName().ToString());
+		}
+
 		if (!SceneInfo->bPendingAddStaticMeshes)
 		{
 			SceneInfo->bPendingAddStaticMeshes = 1;
@@ -6348,7 +6356,7 @@ void FScene::Update(FRDGBuilder& GraphBuilder, const FUpdateParameters& Paramete
 			FPrimitiveSceneInfo* PrimitiveSceneInfo = Item.SceneInfo;
 			const FUpdateInstanceCommand &UpdateInstance = Item.Payload;
 			FScopeCycleCounter Context(PrimitiveSceneProxy->GetStatId());
-			
+
 			QueueFlushVirtualTexture(PrimitiveSceneInfo);
 
 			// TODO: no need to do this if only the payload size changed, we only need it because the MDC stores the instance count!
diff --git a/Engine/Source/Runtime/Renderer/Public/PrimitiveSceneInfo.h b/Engine/Source/Runtime/Renderer/Public/PrimitiveSceneInfo.h
--- a/Engine/Source/Runtime/Renderer/Public/PrimitiveSceneInfo.h
+++ b/Engine/Source/Runtime/Renderer/Public/PrimitiveSceneInfo.h
@@ -669,6 +669,9 @@ private:
 	/** True if the primitive is queued to have static meshes built. */
 	bool bPendingAddStaticMeshes : 1;
 
+	/** True if the primitive is queued for deletion. */
+	bool bPendingDelete : 1 = false;
+
 	/** True if the primitive is queued to have its virtual texture flushed. */
 	bool bPendingFlushVirtualTexture : 1;
 
