Heterogeneous Volumes: Clamp CLOD value before applying bias controls to heighten effect on foreground elements

diff --git a/Engine/Source/Runtime/Renderer/Private/HeterogeneousVolumes/HeterogeneousVolumes.cpp b/Engine/Source/Runtime/Renderer/Private/HeterogeneousVolumes/HeterogeneousVolumes.cpp
--- a/Engine/Source/Runtime/Renderer/Private/HeterogeneousVolumes/HeterogeneousVolumes.cpp
+++ b/Engine/Source/Runtime/Renderer/Private/HeterogeneousVolumes/HeterogeneousVolumes.cpp
@@ -653,8 +653,10 @@ namespace HeterogeneousVolumes
 
 			// MIP is defined as the log of the ratio of native voxel resolution to pixel-coverage of volume bounds
 			float PixelWidthCoverage = (2.0 * WorldBounds.BoxExtent.GetMax()) / ProjectedPixelWidth;
-			float ViewLODValue = FMath::Log2(VoxelResolutionMin / PixelWidthCoverage) + HeterogeneousVolume->GetMipBias() + HeterogeneousVolumes::GetCLODBias();
-			ViewLODValue = FMath::Max(ViewLODValue, 0);
+
+			// Clamp LOD to heighten the effect on foreground elements before applying bias controls
+			float ViewLODValue = FMath::Max(FMath::Log2(VoxelResolutionMin / PixelWidthCoverage), 0.0f);
+			ViewLODValue += HeterogeneousVolume->GetMipBias() + HeterogeneousVolumes::GetCLODBias();
 
 			LODValue = FMath::Min(ViewLODValue, LODValue);
 		}
