Make LastRenderTime a double to match the incoming time so it doesn't cause negative delta time due to rounding error.

diff --git a/Engine/Source/Runtime/Renderer/Private/RendererScene.cpp b/Engine/Source/Runtime/Renderer/Private/RendererScene.cpp
--- a/Engine/Source/Runtime/Renderer/Private/RendererScene.cpp
+++ b/Engine/Source/Runtime/Renderer/Private/RendererScene.cpp
@@ -291,7 +291,6 @@ FSceneViewState::FSceneViewState(ERHIFeatureLevel::Type FeatureLevel, FSceneView
 	Scene = nullptr;
 	OcclusionFrameCounter = 0;
 	LastRenderTime = -FLT_MAX;
-	LastRenderTimeDelta = 0.0f;
 	MotionBlurTimeScale = 1.0f;
 	MotionBlurTargetDeltaTime = 1.0f / 60.0f; // Start with a reasonable default of 60hz.
 	PrevViewMatrixForOcclusionQuery.SetIdentity();
diff --git a/Engine/Source/Runtime/Renderer/Private/ScenePrivate.h b/Engine/Source/Runtime/Renderer/Private/ScenePrivate.h
--- a/Engine/Source/Runtime/Renderer/Private/ScenePrivate.h
+++ b/Engine/Source/Runtime/Renderer/Private/ScenePrivate.h
@@ -809,8 +809,7 @@ public:
 
 	uint32		PendingPrevFrameNumber;
 	uint32		PrevFrameNumber;
-	float		LastRenderTime;
-	float		LastRenderTimeDelta;
+	double		LastRenderTime;
 	float		MotionBlurTimeScale;
 	float		MotionBlurTargetDeltaTime;
 	FMatrix44f	PrevViewMatrixForOcclusionQuery;
@@ -1226,12 +1225,7 @@ public:
 	// called every frame after the view state was updated
 	void UpdateLastRenderTime(const FSceneViewFamily& Family)
 	{
-		// The editor can trigger multiple update calls within a frame
-		if(Family.Time.GetRealTimeSeconds() != LastRenderTime)
-		{
-			LastRenderTimeDelta = Family.Time.GetRealTimeSeconds() - LastRenderTime;
-			LastRenderTime = Family.Time.GetRealTimeSeconds();
-		}
+		LastRenderTime = Family.Time.GetRealTimeSeconds();
 	}
 
 	// InScene is passed in, as the Scene pointer in the class itself may be null, if it was allocated without a scene.
