Path Tracer: Fix adaptive sampling

NumPathStates should only be cleared once, and the 1D dispatch optimization should not be used on the first iteration (where adaptive sampling has not disabled any paths yet).

diff --git a/Engine/Source/Runtime/Renderer/Private/PathTracing.cpp b/Engine/Source/Runtime/Renderer/Private/PathTracing.cpp
--- a/Engine/Source/Runtime/Renderer/Private/PathTracing.cpp
+++ b/Engine/Source/Runtime/Renderer/Private/PathTracing.cpp
@@ -3728,6 +3728,7 @@ void FDeferredShadingSceneRenderer::RenderPathTracing(
 
 						RDG_EVENT_SCOPE_CONDITIONAL(GraphBuilder, DispatchResX > DispatchSize || DispatchResY > DispatchSize, "Tile=(%d,%d - %dx%d)", TileX, TileY, DispatchSizeX, DispatchSizeYLocal);
 
+						AddClearUAVPass(GraphBuilder, GraphBuilder.CreateUAV(NumActivePaths, PF_R32_UINT), 0);
 						if (Config.UseAdaptiveSampling && Config.PathTracingData.Iteration > 0)
 						{
 							// If we are using adaptive sampling, build a smaller list of active paths after the first iteration
@@ -3743,7 +3744,6 @@ void FDeferredShadingSceneRenderer::RenderPathTracing(
 
 							PassParameters->NextActivePaths = GraphBuilder.CreateUAV(ActivePaths[0], PF_R32_UINT);
 							PassParameters->NumPathStates = GraphBuilder.CreateUAV(NumActivePaths, PF_R32_UINT);
-							AddClearUAVPass(GraphBuilder, PassParameters->NumPathStates, 0);
 
 							PassParameters->TileTextureOffset.X = TileX;
 							PassParameters->TileTextureOffset.Y = TileY + CurrentGPU * DispatchSizeYSplit;
@@ -3756,7 +3756,6 @@ void FDeferredShadingSceneRenderer::RenderPathTracing(
 								PassParameters,
 								FComputeShaderUtils::GetGroupCount(PassParameters->DispatchDim, FComputeShaderUtils::kGolden2DGroupSize));
 						}
-						AddClearUAVPass(GraphBuilder, GraphBuilder.CreateUAV(NumActivePaths, PF_R32_UINT), 0);
 
 						// Run a pass per bounce, up until the compaction depth. Beyond that point, the path tracer will handle any remaining bounces.
 						// Generally, since there is some launch overhead - it can be worthwhile to handle the "bounce tail" in a single loop because not too many paths survive
@@ -3896,7 +3895,7 @@ void FDeferredShadingSceneRenderer::RenderPathTracing(
 							TShaderMapRef<FPathTracingRG> RayGenShader(View.ShaderMap, GetPathTracingRGPermutation(*Scene, bUseCompaction));
 							ClearUnusedGraphResources(RayGenShader, PassParameters);
 							const bool bFlushRenderingCommands = FlushRenderingCommands == 1 || (FlushRenderingCommands == 2 && Bounce == CompactionDepth);
-							const bool bUse1DDispatch = Config.UseAdaptiveSampling || Bounce > 0;
+							const bool bUse1DDispatch = (Config.UseAdaptiveSampling && Config.PathTracingData.Iteration > 0) || Bounce > 0;
 							GraphBuilder.AddPass(
 								RDG_EVENT_NAME("Path Tracer Sample=%d/%d NumLights=%d (Bounce=%d%s%s)", PathTracingState->SampleIndex, MaxSPP, PassParameters->SceneLightCount, PassParameters->FirstBounce, bUseCompaction ? "" : "+", bUseIndirectDispatch && Bounce > 0 ? TEXT(" indirect") : TEXT("")),
 								PassParameters,
