Path Tracer: Add HitKind to debug visualization

This visualizes front facing triangles in green and backfacing ones in red

diff --git a/Engine/Shaders/Private/PathTracing/PathTracingDebug.usf b/Engine/Shaders/Private/PathTracing/PathTracingDebug.usf
--- a/Engine/Shaders/Private/PathTracing/PathTracingDebug.usf
+++ b/Engine/Shaders/Private/PathTracing/PathTracingDebug.usf
@@ -233,6 +233,10 @@ PathTracingResult PathTracingDebugCore(FRayDesc Ray, uint2 LaunchIndex)
 				Color *= abs(dot(WorldNormal, Ray.Direction));
 				return MakeResult(Color, 0.0, Z);
 			}
+			case PATH_TRACER_DEBUG_VIZ_HITKIND:
+			{
+				return MakeResult((PackedPayload.IsFrontFace() ? float3(0, 1, 0) : float3(1, 0, 0)) * abs(dot(WorldNormal, Ray.Direction)), 0.0, Z);
+			}
 		}
 
 		FPathTracingPayload Payload = UnpackPathTracingPayload(PackedPayload, Ray);
diff --git a/Engine/Shaders/Shared/PathTracingDefinitions.h b/Engine/Shaders/Shared/PathTracingDefinitions.h
--- a/Engine/Shaders/Shared/PathTracingDefinitions.h
+++ b/Engine/Shaders/Shared/PathTracingDefinitions.h
@@ -95,4 +95,5 @@
 #define PATH_TRACER_DEBUG_VIZ_LIGHT_GRID_AXIS           21
 #define PATH_TRACER_DEBUG_VIZ_DECAL_GRID_COUNT 			22
 #define PATH_TRACER_DEBUG_VIZ_DECAL_GRID_AXIS			23
-#define PATH_TRACER_DEBUG_VIZ_VOLUME_LIGHT_COUNT        24
\ No newline at end of file
+#define PATH_TRACER_DEBUG_VIZ_VOLUME_LIGHT_COUNT        24
+#define PATH_TRACER_DEBUG_VIZ_HITKIND                   25
\ No newline at end of file
diff --git a/Engine/Source/Runtime/Renderer/Private/PathTracing.cpp b/Engine/Source/Runtime/Renderer/Private/PathTracing.cpp
--- a/Engine/Source/Runtime/Renderer/Private/PathTracing.cpp
+++ b/Engine/Source/Runtime/Renderer/Private/PathTracing.cpp
@@ -755,7 +755,7 @@ int GetPathTracingDebugVisualizationMode() {
 		PathTracingDebugModes.Emplace(FName(TEXT("DecalGridCount"))	   	, PATH_TRACER_DEBUG_VIZ_DECAL_GRID_COUNT		);
 		PathTracingDebugModes.Emplace(FName(TEXT("DecalGridAxis"))	   	, PATH_TRACER_DEBUG_VIZ_DECAL_GRID_AXIS			);
 		PathTracingDebugModes.Emplace(FName(TEXT("VolumeLightCount"))	, PATH_TRACER_DEBUG_VIZ_VOLUME_LIGHT_COUNT		);
-		
+		PathTracingDebugModes.Emplace(FName(TEXT("HitKind"))            , PATH_TRACER_DEBUG_VIZ_HITKIND                 );
 	}
 	FString DebugMode = CVarPathTracingDebugVisualize.GetValueOnRenderThread();
 	if (!DebugMode.IsEmpty())
