Merge velocity rendering and full depth pre pass into a single render pass

diff --git a/Engine/Source/Runtime/Renderer/Private/MobileShadingRenderer.cpp b/Engine/Source/Runtime/Renderer/Private/MobileShadingRenderer.cpp
--- a/Engine/Source/Runtime/Renderer/Private/MobileShadingRenderer.cpp
+++ b/Engine/Source/Runtime/Renderer/Private/MobileShadingRenderer.cpp
@@ -795,7 +795,13 @@ void FMobileSceneRenderer::RenderFullDepthPrepass(FRDGBuilder& GraphBuilder, TAr
 
 		View.BeginRenderView();
 
+		const bool bRenderVelocityInDepthPrePass = (Scene->EarlyZPassMode == DDM_AllOpaqueNoVelocity);
+
 		auto* PassParameters = GraphBuilder.AllocParameters<FMobileRenderPassParameters>();
+		if (bRenderVelocityInDepthPrePass)
+		{
+			PassParameters->RenderTargets[0] = FRenderTargetBinding(SceneTextures.Velocity, ERenderTargetLoadAction::EClear);
+		}
 		PassParameters->RenderTargets.DepthStencil = FDepthStencilBinding(SceneTextures.Depth.Target, ERenderTargetLoadAction::EClear, ERenderTargetLoadAction::EClear, FExclusiveDepthStencil::DepthWrite_StencilWrite);
 		PassParameters->View = View.GetShaderParameters();
 		PassParameters->MobileBasePass = CreateMobileBasePassUniformBuffer(GraphBuilder, View, EMobileBasePass::DepthPrePass, EMobileSceneTextureSetupMode::None);
@@ -811,27 +817,33 @@ void FMobileSceneRenderer::RenderFullDepthPrepass(FRDGBuilder& GraphBuilder, TAr
 			PassParameters->RenderTargets.DepthStencil.SetDepthLoadAction(ERenderTargetLoadAction::ELoad);
 			PassParameters->RenderTargets.DepthStencil.SetStencilLoadAction(ERenderTargetLoadAction::ELoad);
 			PassParameters->RenderTargets.DepthStencil.SetDepthStencilAccess(FExclusiveDepthStencil::DepthWrite_StencilWrite);
+			if (bRenderVelocityInDepthPrePass)
+			{
+				PassParameters->RenderTargets[0].SetLoadAction(ERenderTargetLoadAction::ELoad);
+			}
 		}
 
 		View.ParallelMeshDrawCommandPasses[EMeshPass::DepthPass].BuildRenderingCommands(GraphBuilder, Scene->GPUScene, PassParameters->InstanceCullingDrawParams);
 
+		if (bRenderVelocityInDepthPrePass)
+		{
+			View.ParallelMeshDrawCommandPasses[EMeshPass::Velocity].BuildRenderingCommands(GraphBuilder, Scene->GPUScene, VelocityInstanceCullingDrawParams);
+		}
+
 		GraphBuilder.AddPass(
 			RDG_EVENT_NAME("FullDepthPrepass"),
 			PassParameters,
 			ERDGPassFlags::Raster,
-			[this, PassParameters, &View](FRHICommandList& RHICmdList)
+			[this, PassParameters, &View, bRenderVelocityInDepthPrePass](FRHICommandList& RHICmdList)
 			{
 				RenderPrePass(RHICmdList, View, &PassParameters->InstanceCullingDrawParams);
+				if (bRenderVelocityInDepthPrePass)
+				{
+					RenderVelocityPass(RHICmdList, View, &VelocityInstanceCullingDrawParams);
+				}
 			});
 	}
 
-	// Do not check for bShouldRenderVelocities here to avoid dropping velocity primitives from depth pre-pass
-	if (Scene->EarlyZPassMode == DDM_AllOpaqueNoVelocity) 
-	{
-		// Render the velocities and depth of movable objects
-		RenderVelocities(GraphBuilder, InViews, SceneTextures, EVelocityPass::Opaque, false);
-	}
-
 	for (FRenderViewContext& ViewContext : RenderViews)
 	{
 		FViewInfo& View = *ViewContext.ViewInfo;
diff --git a/Engine/Source/Runtime/Renderer/Private/RendererScene.cpp b/Engine/Source/Runtime/Renderer/Private/RendererScene.cpp
--- a/Engine/Source/Runtime/Renderer/Private/RendererScene.cpp
+++ b/Engine/Source/Runtime/Renderer/Private/RendererScene.cpp
@@ -4700,7 +4700,7 @@ void FScene::GetEarlyZPassMode(ERHIFeatureLevel::Type InFeatureLevel, EDepthDraw
 
 		if (MobileUsesFullDepthPrepass(ShaderPlatform))
 		{
-			const bool bDepthPassCanOutputVelocity = FVelocityRendering::DepthPassCanOutputVelocity(InFeatureLevel);
+			const bool bDepthPassCanOutputVelocity = PlatformSupportsVelocityRendering(ShaderPlatform) && FVelocityRendering::DepthPassCanOutputVelocity(InFeatureLevel);
 			OutZPassMode = bDepthPassCanOutputVelocity ? DDM_AllOpaqueNoVelocity : DDM_AllOpaque;
 		}
 	}
diff --git a/Engine/Source/Runtime/Renderer/Private/SceneRendering.h b/Engine/Source/Runtime/Renderer/Private/SceneRendering.h
--- a/Engine/Source/Runtime/Renderer/Private/SceneRendering.h
+++ b/Engine/Source/Runtime/Renderer/Private/SceneRendering.h
@@ -2721,6 +2721,8 @@ protected:
 	void RenderMaskedPrePass(FRHICommandList& RHICmdList, const FViewInfo& View);
 	void RenderFullDepthPrepass(FRDGBuilder& GraphBuilder, TArrayView<FViewInfo> InViews, FSceneTextures& SceneTextures, bool bIsSceneCaptureRenderPass=false);
 
+	void RenderVelocityPass(FRHICommandList& RHICmdList, const FViewInfo& View, const FInstanceCullingDrawParams* InstanceCullingDrawParams);
+
 	void RenderMobileLocalLightsBuffer(FRDGBuilder& GraphBuilder, FSceneTextures& SceneTextures, const FSortedLightSetSceneInfo& SortedLights);
 
 	void RenderCustomRenderPassBasePass(FRDGBuilder& GraphBuilder, TArrayView<FViewInfo> InViews, FRDGTextureRef ViewFamilyTexture, FSceneTextures& SceneTextures, bool bIncludeTranslucent);
@@ -2831,6 +2833,7 @@ private:
 	FInstanceCullingDrawParams TranslucencyInstanceCullingDrawParams;
 	FInstanceCullingDrawParams MeshDecalSceneColorInstanceCullingDrawParams;
 	FInstanceCullingDrawParams MeshDecalSceneColorAndGBufferInstanceCullingDrawParams;
+	FInstanceCullingDrawParams VelocityInstanceCullingDrawParams;
 
 	const FViewInfo* CachedView = nullptr;
 };
diff --git a/Engine/Source/Runtime/Renderer/Private/VelocityRendering.cpp b/Engine/Source/Runtime/Renderer/Private/VelocityRendering.cpp
--- a/Engine/Source/Runtime/Renderer/Private/VelocityRendering.cpp
+++ b/Engine/Source/Runtime/Renderer/Private/VelocityRendering.cpp
@@ -355,6 +355,21 @@ void FSceneRenderer::RenderVelocities(
 
 }
 
+void FMobileSceneRenderer::RenderVelocityPass(FRHICommandList& RHICmdList, const FViewInfo& View, const FInstanceCullingDrawParams* InstanceCullingDrawParams)
+{
+	checkSlow(RHICmdList.IsInsideRenderPass());
+
+	SCOPED_NAMED_EVENT(FMobileSceneRenderer_RenderVelocityPass, FColor::Emerald);
+	RHI_BREADCRUMB_EVENT_STAT(RHICmdList, RenderVelocities, "MobileRenderVelocityPass");
+	SCOPED_GPU_STAT(RHICmdList, RenderVelocities);
+
+	SCOPE_CYCLE_COUNTER(STAT_RenderVelocities);
+	CSV_SCOPED_TIMING_STAT_EXCLUSIVE(RenderVelocityPass);
+
+	SetStereoViewport(RHICmdList, View);
+	View.ParallelMeshDrawCommandPasses[EMeshPass::Velocity].Draw(RHICmdList, InstanceCullingDrawParams);
+}
+
 EPixelFormat FVelocityRendering::GetFormat(EShaderPlatform ShaderPlatform)
 {
 	// Lumen needs velocity depth
