[HWRT] Add name of the mesh that triggers the ensure so it will be easier to investigate.

#rnx

diff --git a/Engine/Source/Runtime/Renderer/Private/RayTracing/RayTracingDynamicGeometry.cpp b/Engine/Source/Runtime/Renderer/Private/RayTracing/RayTracingDynamicGeometry.cpp
--- a/Engine/Source/Runtime/Renderer/Private/RayTracing/RayTracingDynamicGeometry.cpp
+++ b/Engine/Source/Runtime/Renderer/Private/RayTracing/RayTracingDynamicGeometry.cpp
@@ -308,8 +308,8 @@ void FRayTracingDynamicGeometryCollection::AddDynamicMeshBatchForGeometryUpdate(
 
 					const uint32 VertexBufferNumElements = UpdateParams.VertexBufferSize / sizeof(FVector3f) - MinVertexIndex;
 					if (!ensureMsgf(NumCPUVertices <= VertexBufferNumElements,
-						TEXT("Vertex buffer contains %d vertices, but RayTracingDynamicGeometryConverterCS dispatch command expects at least %d."),
-						VertexBufferNumElements, NumCPUVertices))
+						TEXT("%s: Vertex buffer contains %d vertices, but RayTracingDynamicGeometryConverterCS dispatch command expects at least %d."),
+						*PrimitiveSceneProxy->GetOwnerName().ToString(), VertexBufferNumElements, NumCPUVertices))
 					{
 						NumCPUVertices = VertexBufferNumElements;
 					}
