Disable velocity rendering during custom render passes, which prevents polluting the main render velocities.

diff --git a/Engine/Source/Runtime/Renderer/Private/DeferredShadingRenderer.cpp b/Engine/Source/Runtime/Renderer/Private/DeferredShadingRenderer.cpp
--- a/Engine/Source/Runtime/Renderer/Private/DeferredShadingRenderer.cpp
+++ b/Engine/Source/Runtime/Renderer/Private/DeferredShadingRenderer.cpp
@@ -1963,7 +1963,7 @@ void FDeferredShadingSceneRenderer::Render(FRDGBuilder& GraphBuilder)
 			}
 
 			// special pass for DDM_AllOpaqueNoVelocity, which uses the velocity pass to finish the early depth pass write
-			if (bShouldRenderVelocities && Scene->EarlyZPassMode == DDM_AllOpaqueNoVelocity && RendererOutput == ERendererOutput::FinalSceneColor)
+			if (bShouldRenderVelocities && Scene->EarlyZPassMode == DDM_AllOpaqueNoVelocity && RendererOutput == ERendererOutput::FinalSceneColor && !InViews[0].CustomRenderPass)
 			{
 				// Render the velocities of movable objects
 				RenderVelocities(GraphBuilder, InViews, SceneTextures, EVelocityPass::Opaque, bForceVelocityOutput);
