Decouple application of LOD Value and Bias depending on use

diff --git a/Engine/Source/Runtime/Renderer/Private/HeterogeneousVolumes/HeterogeneousVolumes.cpp b/Engine/Source/Runtime/Renderer/Private/HeterogeneousVolumes/HeterogeneousVolumes.cpp
--- a/Engine/Source/Runtime/Renderer/Private/HeterogeneousVolumes/HeterogeneousVolumes.cpp
+++ b/Engine/Source/Runtime/Renderer/Private/HeterogeneousVolumes/HeterogeneousVolumes.cpp
@@ -617,28 +617,35 @@ namespace HeterogeneousVolumes
 		);
 	}
 	
-	float CalcLOD(
+	float CalcLODBias(const IHeterogeneousVolumeInterface* HeterogeneousVolume)
+	{
+		return HeterogeneousVolume->GetMipBias() + HeterogeneousVolumes::GetCLODBias();
+	}
+
+	FLODValue CalcLOD(
 		const HeterogeneousVolumes::FLODInfo& LODInfo,
 		const IHeterogeneousVolumeInterface* HeterogeneousVolume
 	)
 	{
 		if (!HeterogeneousVolumes::UseContinuousLOD())
 		{
-			return 0.0f;
+			return FLODValue();
 		}
 
 		FBoxSphereBounds WorldBounds = HeterogeneousVolume->GetBounds();
 		FIntVector VoxelResolution = HeterogeneousVolume->GetVoxelResolution();
 		float VoxelResolutionMin = VoxelResolution.GetMin();
 
-		float LODValue = FMath::Floor(FMath::Log2(VoxelResolutionMin));
+		const float MaxLOD = FMath::Floor(FMath::Log2(VoxelResolutionMin));
+		FLODValue LODValue;
+		LODValue.LOD = MaxLOD;
 		if (!LODInfo.bIsPerspective)
 		{
 			float VolumeRatio = FVector(LODInfo.WorldSceneBounds.BoxExtent / WorldBounds.BoxExtent).Length();
 			float ViewLODValue = FMath::Log2(VolumeRatio) + HeterogeneousVolume->GetMipBias() + HeterogeneousVolumes::GetCLODBias();
 			ViewLODValue = FMath::Max(ViewLODValue, 0);
 
-			LODValue = FMath::Min(ViewLODValue, LODValue);
+			LODValue.LOD = FMath::Min(ViewLODValue, LODValue.LOD);
 		}
 		else if (LODInfo.WorldShadowFrustum.IntersectBox(WorldBounds.Origin, WorldBounds.BoxExtent))
 		{
@@ -656,15 +663,17 @@ namespace HeterogeneousVolumes
 
 			// Clamp LOD to heighten the effect on foreground elements before applying bias controls
 			float ViewLODValue = FMath::Max(FMath::Log2(VoxelResolutionMin / PixelWidthCoverage), 0.0f);
-			ViewLODValue += HeterogeneousVolume->GetMipBias() + HeterogeneousVolumes::GetCLODBias();
 
-			LODValue = FMath::Min(ViewLODValue, LODValue);
+			LODValue.LOD = FMath::Min(ViewLODValue, LODValue.LOD);
 		}
 
+		float TotalBias = HeterogeneousVolume->GetMipBias() + HeterogeneousVolumes::GetCLODBias();
+		LODValue.Bias = FMath::Min(TotalBias, MaxLOD - LODValue.LOD);
+
 		return LODValue;
 	}
 	
-	float CalcLOD(const FSceneView& View, const IHeterogeneousVolumeInterface* HeterogeneousVolume)
+	FLODValue CalcLOD(const FSceneView& View, const IHeterogeneousVolumeInterface* HeterogeneousVolume)
 	{
 		FLODInfo LODInfo;
 		// TODO: Not supporting orthographic projection for now
@@ -681,9 +690,9 @@ namespace HeterogeneousVolumes
 		return CalcLOD(LODInfo, HeterogeneousVolume);
 	}
 
-	float CalcLODFactor(float LODValue)
+	float CalcLODFactor(float LODValue, float LODBias)
 	{
-		return FMath::Pow(2, LODValue);
+		return FMath::Pow(2, LODValue + LODBias);
 	}
 
 	float CalcLODFactor(
@@ -691,16 +700,19 @@ namespace HeterogeneousVolumes
 		const IHeterogeneousVolumeInterface* HeterogeneousVolume
 	)
 	{
-		return CalcLODFactor(CalcLOD(LODInfo, HeterogeneousVolume));
+		FLODValue LODValue = CalcLOD(LODInfo, HeterogeneousVolume);
+		return CalcLODFactor(LODValue.LOD, LODValue.Bias);
 	}
 
 	float CalcLODFactor(const FSceneView& View, const IHeterogeneousVolumeInterface* HeterogeneousVolume)
 	{
-		return CalcLODFactor(CalcLOD(View, HeterogeneousVolume));
+		FLODValue LODValue = CalcLOD(View, HeterogeneousVolume);
+		return CalcLODFactor(LODValue.LOD, LODValue.Bias);
 	}
-	
-	FIntVector GetLightingCacheResolution(const IHeterogeneousVolumeInterface* RenderInterface, float LODFactor)
+
+	FIntVector GetLightingCacheResolution(const IHeterogeneousVolumeInterface* RenderInterface, FLODValue LODValue)
 	{
+		float LODFactor = CalcLODFactor(LODValue.LOD, 0.0f);
 		float OverrideDownsampleFactor = CVarHeterogeneousVolumesLightingCacheDownsampleFactor.GetValueOnRenderThread();
 		float DownsampleFactor = OverrideDownsampleFactor > 0.0 ? OverrideDownsampleFactor : RenderInterface->GetLightingDownsampleFactor() * LODFactor;
 		DownsampleFactor = FMath::Max(DownsampleFactor, 0.125);
@@ -1021,8 +1033,8 @@ void FDeferredShadingSceneRenderer::RenderHeterogeneousVolumes(
 					if (HeterogeneousVolumes::GetLightingCacheMode() != 0)
 					{
 						// TODO: Allow option for scalar transmittance to conserve bandwidth
-						float LODFactor = HeterogeneousVolumes::CalcLODFactor(View, HeterogeneousVolume);
-						FIntVector LightingCacheResolution = HeterogeneousVolumes::GetLightingCacheResolution(HeterogeneousVolume, LODFactor);
+						HeterogeneousVolumes::FLODValue LODValue = HeterogeneousVolumes::CalcLOD(View, HeterogeneousVolume);
+						FIntVector LightingCacheResolution = HeterogeneousVolumes::GetLightingCacheResolution(HeterogeneousVolume, LODValue);
 						uint32 NumMips = FMath::Log2(float(FMath::Min(FMath::Min(LightingCacheResolution.X, LightingCacheResolution.Y), LightingCacheResolution.Z))) + 1;
 						FRDGTextureDesc LightingCacheDesc = FRDGTextureDesc::Create3D(
 							LightingCacheResolution,
diff --git a/Engine/Source/Runtime/Renderer/Private/HeterogeneousVolumes/HeterogeneousVolumes.h b/Engine/Source/Runtime/Renderer/Private/HeterogeneousVolumes/HeterogeneousVolumes.h
--- a/Engine/Source/Runtime/Renderer/Private/HeterogeneousVolumes/HeterogeneousVolumes.h
+++ b/Engine/Source/Runtime/Renderer/Private/HeterogeneousVolumes/HeterogeneousVolumes.h
@@ -61,8 +61,6 @@ namespace HeterogeneousVolumes
 {
 	// CVars
 	FIntVector GetVolumeResolution(const IHeterogeneousVolumeInterface*);
-	FIntVector GetLightingCacheResolution(const IHeterogeneousVolumeInterface*, float LODFactor);
-	FIntVector GetAmbientOcclusionResolution(const IHeterogeneousVolumeInterface*, float LODFactor);
 
 	int32 GetDownsampleFactor();
 	FIntPoint GetScaledViewRect(FIntRect ViewRect);
@@ -163,11 +161,22 @@ namespace HeterogeneousVolumes
 		bool bIsPerspective = false;
 	};
 
-	float CalcLODFactor(float LOD);
+	struct FLODValue
+	{
+		float LOD = 0.0f;
+		float Bias = 0.0f;
+	};
+
+	FLODValue CalcLOD(const HeterogeneousVolumes::FLODInfo& LODInfo, const IHeterogeneousVolumeInterface* HeterogeneousVolume);
+	FLODValue CalcLOD(const FSceneView& View, const IHeterogeneousVolumeInterface* HeterogeneousVolume);
+	FIntVector GetLightingCacheResolution(const IHeterogeneousVolumeInterface*, FLODValue LODValue);
+	FIntVector GetAmbientOcclusionResolution(const IHeterogeneousVolumeInterface*, FLODValue LODValue);
+
+	float CalcLODBias(const IHeterogeneousVolumeInterface* HeterogeneousVolume);
+	float CalcLODFactor(float LODValue, float LODBias);
 	float CalcLODFactor(const HeterogeneousVolumes::FLODInfo& LODInfo, const IHeterogeneousVolumeInterface* HeterogeneousVolume);
 	float CalcLODFactor(const FSceneView& View, const IHeterogeneousVolumeInterface* HeterogeneousVolume);
-	float CalcLOD(const HeterogeneousVolumes::FLODInfo& LODInfo, const IHeterogeneousVolumeInterface* HeterogeneousVolume);
-	float CalcLOD(const FSceneView& View, const IHeterogeneousVolumeInterface* HeterogeneousVolume);
+
 	bool IsHoldout(const IHeterogeneousVolumeInterface* HeterogeneousVolumeInterface);
 
 	const FProjectedShadowInfo* GetProjectedShadowInfo(const FVisibleLightInfo* VisibleLightInfo, int32 ShadowIndex);
diff --git a/Engine/Source/Runtime/Renderer/Private/HeterogeneousVolumes/HeterogeneousVolumesAmbientOcclusionPipeline.cpp b/Engine/Source/Runtime/Renderer/Private/HeterogeneousVolumes/HeterogeneousVolumesAmbientOcclusionPipeline.cpp
--- a/Engine/Source/Runtime/Renderer/Private/HeterogeneousVolumes/HeterogeneousVolumesAmbientOcclusionPipeline.cpp
+++ b/Engine/Source/Runtime/Renderer/Private/HeterogeneousVolumes/HeterogeneousVolumesAmbientOcclusionPipeline.cpp
@@ -89,8 +89,9 @@ namespace HeterogeneousVolumes
 		return FMath::Clamp(CVarHeterogeneousVolumesAmbientOcclusionMaxStepCount.GetValueOnRenderThread(), 1, GetMaxStepCount());
 	}
 
-	FIntVector GetAmbientOcclusionResolution(const IHeterogeneousVolumeInterface* RenderInterface, float LODFactor)
+	FIntVector GetAmbientOcclusionResolution(const IHeterogeneousVolumeInterface* RenderInterface, FLODValue LODValue)
 	{
+		float LODFactor = CalcLODFactor(LODValue.LOD, 0.0f);
 		float DownsampleFactor = GetAmbientOcclusionDownsampleFactor() * LODFactor;
 		DownsampleFactor = FMath::Max(DownsampleFactor, 0.125f);
 
@@ -758,8 +759,8 @@ void RenderAmbientOcclusionWithLiveShading(
 		return;
 	}
 
-	float LODFactor = HeterogeneousVolumes::CalcLODFactor(View, HeterogeneousVolumeInterface);
-	FIntVector AmbientOcclusionTextureResolution = HeterogeneousVolumes::GetAmbientOcclusionResolution(HeterogeneousVolumeInterface, LODFactor);
+	HeterogeneousVolumes::FLODValue LODValue = HeterogeneousVolumes::CalcLOD(View, HeterogeneousVolumeInterface);
+	FIntVector AmbientOcclusionTextureResolution = HeterogeneousVolumes::GetAmbientOcclusionResolution(HeterogeneousVolumeInterface, LODValue);
 
 	// Build existence mask
 	FRDGTextureRef ExistenceMaskTexture;
diff --git a/Engine/Source/Runtime/Renderer/Private/HeterogeneousVolumes/HeterogeneousVolumesHardwareRayTracing.cpp b/Engine/Source/Runtime/Renderer/Private/HeterogeneousVolumes/HeterogeneousVolumesHardwareRayTracing.cpp
--- a/Engine/Source/Runtime/Renderer/Private/HeterogeneousVolumes/HeterogeneousVolumesHardwareRayTracing.cpp
+++ b/Engine/Source/Runtime/Renderer/Private/HeterogeneousVolumes/HeterogeneousVolumesHardwareRayTracing.cpp
@@ -601,8 +601,8 @@ void RenderLightingCacheWithPreshadingHardwareRayTracing(
 		PassParameters->SparseVoxelUniformBuffer = SparseVoxelUniformBuffer;
 
 		// Transmittance volume
-		float LODFactor = HeterogeneousVolumes::CalcLODFactor(View, HeterogeneousVolumeInterface);
-		PassParameters->LightingCache.LightingCacheResolution = HeterogeneousVolumes::GetLightingCacheResolution(HeterogeneousVolumeInterface, LODFactor);
+		HeterogeneousVolumes::FLODValue LODValue = HeterogeneousVolumes::CalcLOD(View, HeterogeneousVolumeInterface);
+		PassParameters->LightingCache.LightingCacheResolution = HeterogeneousVolumes::GetLightingCacheResolution(HeterogeneousVolumeInterface, LODValue);
 		PassParameters->LightingCache.LightingCacheVoxelBias = HeterogeneousVolumeInterface->GetShadowBiasFactor();
 		//PassParameters->LightingCache.LightingCacheTexture = GraphBuilder.CreateSRV(LightingCacheTexture);
 
@@ -656,8 +656,8 @@ void RenderLightingCacheWithPreshadingHardwareRayTracing(
 	FRenderLightingCacheWithPreshadingRGS::FPermutationDomain PermutationVector;
 	PermutationVector.Set<FRenderLightingCacheWithPreshadingRGS::FLightingCacheMode>(HeterogeneousVolumes::GetLightingCacheMode() - 1);
 	TShaderRef<FRenderLightingCacheWithPreshadingRGS> RayGenerationShader = View.ShaderMap->GetShader<FRenderLightingCacheWithPreshadingRGS>(PermutationVector);
-	float LODFactor = HeterogeneousVolumes::CalcLODFactor(View, HeterogeneousVolumeInterface);
-	FIntVector VolumeResolution = HeterogeneousVolumes::GetLightingCacheResolution(HeterogeneousVolumeInterface, LODFactor);
+	HeterogeneousVolumes::FLODValue LODValue = HeterogeneousVolumes::CalcLOD(View, HeterogeneousVolumeInterface);
+	FIntVector VolumeResolution = HeterogeneousVolumes::GetLightingCacheResolution(HeterogeneousVolumeInterface, LODValue);
 	FIntPoint DispatchResolution = FIntPoint(VolumeResolution.X, VolumeResolution.Y * VolumeResolution.Z);
 
 	GraphBuilder.AddPass(
@@ -801,8 +801,8 @@ void RenderSingleScatteringWithPreshadingHardwareRayTracing(
 		// Transmittance volume
 		if ((HeterogeneousVolumes::UseLightingCacheForTransmittance() && bApplyShadowTransmittance) || HeterogeneousVolumes::UseLightingCacheForInscattering())
 		{
-			float LODFactor = HeterogeneousVolumes::CalcLODFactor(View, HeterogeneousVolumeInterface);
-			PassParameters->LightingCache.LightingCacheResolution = HeterogeneousVolumes::GetLightingCacheResolution(HeterogeneousVolumeInterface, LODFactor);
+			HeterogeneousVolumes::FLODValue LODValue = HeterogeneousVolumes::CalcLOD(View, HeterogeneousVolumeInterface);
+			PassParameters->LightingCache.LightingCacheResolution = HeterogeneousVolumes::GetLightingCacheResolution(HeterogeneousVolumeInterface, LODValue);
 			PassParameters->LightingCache.LightingCacheVoxelBias = HeterogeneousVolumeInterface->GetShadowBiasFactor();
 			PassParameters->LightingCache.LightingCacheTexture = LightingCacheTexture;
 		}
diff --git a/Engine/Source/Runtime/Renderer/Private/HeterogeneousVolumes/HeterogeneousVolumesLiveShadingPipeline.cpp b/Engine/Source/Runtime/Renderer/Private/HeterogeneousVolumes/HeterogeneousVolumesLiveShadingPipeline.cpp
--- a/Engine/Source/Runtime/Renderer/Private/HeterogeneousVolumes/HeterogeneousVolumesLiveShadingPipeline.cpp
+++ b/Engine/Source/Runtime/Renderer/Private/HeterogeneousVolumes/HeterogeneousVolumesLiveShadingPipeline.cpp
@@ -671,9 +671,9 @@ static void RenderLightingCacheWithLiveShading(
 	// Note must be done in the same scope as we add the pass otherwise the UB lifetime will not be guaranteed
 	FDeferredLightUniformStruct DeferredLightUniform = GetDeferredLightParameters(View, *LightSceneInfo);
 	TUniformBufferRef<FDeferredLightUniformStruct> DeferredLightUB = CreateUniformBufferImmediate(DeferredLightUniform, UniformBuffer_SingleDraw);
-	float LODFactor = HeterogeneousVolumes::CalcLODFactor(View, HeterogeneousVolumeInterface);
 
-	FIntVector LightingCacheResolution = HeterogeneousVolumes::GetLightingCacheResolution(HeterogeneousVolumeInterface, LODFactor);
+	HeterogeneousVolumes::FLODValue LODValue = HeterogeneousVolumes::CalcLOD(View, HeterogeneousVolumeInterface);
+	FIntVector LightingCacheResolution = HeterogeneousVolumes::GetLightingCacheResolution(HeterogeneousVolumeInterface, LODValue);
 
 	FRDGTextureRef ExistenceMaskTexture = nullptr;
 	RenderExistenceMaskWithLiveShading(
@@ -748,6 +748,7 @@ static void RenderLightingCacheWithLiveShading(
 		PassParameters->ExistenceMaskTexture = GraphBuilder.CreateSRV(DilatedExistenceTexture);
 
 		// Ray data
+		float LODFactor = HeterogeneousVolumes::CalcLODFactor(LODValue.LOD, LODValue.Bias);
 		PassParameters->MaxTraceDistance = HeterogeneousVolumes::GetMaxTraceDistance();
 		PassParameters->MaxShadowTraceDistance = HeterogeneousVolumes::GetMaxShadowTraceDistance();
 		PassParameters->StepSize = HeterogeneousVolumes::GetStepSize();
@@ -801,7 +802,7 @@ static void RenderLightingCacheWithLiveShading(
 #endif // WANTS_DRAW_MESH_EVENTS
 
 	PassParameters->VoxelMin = FIntVector::ZeroValue;
-	PassParameters->VoxelMax = HeterogeneousVolumes::GetLightingCacheResolution(HeterogeneousVolumeInterface, LODFactor) - FIntVector(1);
+	PassParameters->VoxelMax = HeterogeneousVolumes::GetLightingCacheResolution(HeterogeneousVolumeInterface, LODValue) - FIntVector(1);
 	
 	bool bShouldBoundsCull = HeterogeneousVolumes::ShouldBoundsCull();
 	if (LightType != LightType_Directional && bShouldBoundsCull)
@@ -1125,7 +1126,8 @@ void RenderSingleScatteringWithLiveShading(
 		PassParameters->BlueNoise = CreateUniformBufferImmediate(BlueNoise, EUniformBufferUsage::UniformBuffer_SingleDraw);
 
 		// Light data
-		float LODFactor = HeterogeneousVolumes::CalcLODFactor(View, HeterogeneousVolumeInterface);
+		HeterogeneousVolumes::FLODValue LODValue = HeterogeneousVolumes::CalcLOD(View, HeterogeneousVolumeInterface);
+		float LODFactor = HeterogeneousVolumes::CalcLODFactor(LODValue.LOD, LODValue.Bias);
 		PassParameters->bHoldout = HeterogeneousVolumes::IsHoldout(HeterogeneousVolumeInterface);
 		PassParameters->bApplyEmissionAndTransmittance = bApplyEmissionAndTransmittance;
 		PassParameters->bApplyDirectLighting = bApplyDirectLighting;
@@ -1219,7 +1221,7 @@ void RenderSingleScatteringWithLiveShading(
 		// Volume data
 		if ((HeterogeneousVolumes::UseLightingCacheForTransmittance() && bApplyShadowTransmittance) || HeterogeneousVolumes::UseLightingCacheForInscattering())
 		{
-			PassParameters->LightingCache.LightingCacheResolution = HeterogeneousVolumes::GetLightingCacheResolution(HeterogeneousVolumeInterface, LODFactor);
+			PassParameters->LightingCache.LightingCacheResolution = HeterogeneousVolumes::GetLightingCacheResolution(HeterogeneousVolumeInterface, LODValue);
 			PassParameters->LightingCache.LightingCacheVoxelBias = HeterogeneousVolumeInterface->GetShadowBiasFactor();
 			PassParameters->LightingCache.LightingCacheTexture = LightingCacheTexture;
 		}
@@ -1227,7 +1229,7 @@ void RenderSingleScatteringWithLiveShading(
 		{
 			if (bUseAVSM)
 			{
-				PassParameters->LightingCache.LightingCacheResolution = HeterogeneousVolumes::GetLightingCacheResolution(HeterogeneousVolumeInterface, LODFactor);
+				PassParameters->LightingCache.LightingCacheResolution = HeterogeneousVolumes::GetLightingCacheResolution(HeterogeneousVolumeInterface, LODValue);
 				PassParameters->LightingCache.LightingCacheVoxelBias = HeterogeneousVolumeInterface->GetShadowBiasFactor();
 			}
 			else
@@ -1238,7 +1240,7 @@ void RenderSingleScatteringWithLiveShading(
 			PassParameters->LightingCache.LightingCacheTexture = FRDGSystemTextures::Get(GraphBuilder).VolumetricBlack;
 		}
 
-		PassParameters->AmbientOcclusionResolution = HeterogeneousVolumes::GetAmbientOcclusionResolution(HeterogeneousVolumeInterface, LODFactor);
+		PassParameters->AmbientOcclusionResolution = HeterogeneousVolumes::GetAmbientOcclusionResolution(HeterogeneousVolumeInterface, LODValue);
 		PassParameters->AmbientOcclusionTexture = AmbientOcclusionTexture;
 
 		// Dispatch data
diff --git a/Engine/Source/Runtime/Renderer/Private/HeterogeneousVolumes/HeterogeneousVolumesPreshadingPipeline.cpp b/Engine/Source/Runtime/Renderer/Private/HeterogeneousVolumes/HeterogeneousVolumesPreshadingPipeline.cpp
--- a/Engine/Source/Runtime/Renderer/Private/HeterogeneousVolumes/HeterogeneousVolumesPreshadingPipeline.cpp
+++ b/Engine/Source/Runtime/Renderer/Private/HeterogeneousVolumes/HeterogeneousVolumesPreshadingPipeline.cpp
@@ -447,7 +447,7 @@ void RenderLightingCacheWithPreshadingCompute(
 	FDeferredLightUniformStruct DeferredLightUniform = GetDeferredLightParameters(View, *LightSceneInfo);
 	TUniformBufferRef<FDeferredLightUniformStruct> DeferredLightUB = CreateUniformBufferImmediate(DeferredLightUniform, UniformBuffer_SingleDraw);
 
-	float LODFactor = HeterogeneousVolumes::CalcLODFactor(View, HeterogeneousVolumeInterface);
+	HeterogeneousVolumes::FLODValue LODValue = HeterogeneousVolumes::CalcLOD(View, HeterogeneousVolumeInterface);
 
 	FRenderLightingCacheWithPreshadingCS::FParameters* PassParameters = GraphBuilder.AllocParameters<FRenderLightingCacheWithPreshadingCS::FParameters>();
 	{
@@ -468,7 +468,7 @@ void RenderLightingCacheWithPreshadingCompute(
 		PassParameters->SparseVoxelUniformBuffer = SparseVoxelUniformBuffer;
 
 		// Transmittance volume
-		PassParameters->LightingCache.LightingCacheResolution = HeterogeneousVolumes::GetLightingCacheResolution(HeterogeneousVolumeInterface, LODFactor);
+		PassParameters->LightingCache.LightingCacheResolution = HeterogeneousVolumes::GetLightingCacheResolution(HeterogeneousVolumeInterface, LODValue);
 		PassParameters->LightingCache.LightingCacheVoxelBias = HeterogeneousVolumeInterface->GetShadowBiasFactor();
 		PassParameters->LightingCache.LightingCacheTexture = LightingCacheTexture;
 
@@ -526,7 +526,7 @@ void RenderLightingCacheWithPreshadingCompute(
 	PermutationVector.Set<FRenderLightingCacheWithPreshadingCS::FUseAdaptiveVolumetricShadowMap>(bUseAVSM);
 	TShaderRef<FRenderLightingCacheWithPreshadingCS> ComputeShader = View.ShaderMap->GetShader<FRenderLightingCacheWithPreshadingCS>(PermutationVector);
 
-	FIntVector GroupCount = HeterogeneousVolumes::GetLightingCacheResolution(HeterogeneousVolumeInterface, LODFactor);
+	FIntVector GroupCount = HeterogeneousVolumes::GetLightingCacheResolution(HeterogeneousVolumeInterface, LODValue);
 	GroupCount.X = FMath::DivideAndRoundUp(GroupCount.X, FRenderLightingCacheWithPreshadingCS::GetThreadGroupSize3D());
 	GroupCount.Y = FMath::DivideAndRoundUp(GroupCount.Y, FRenderLightingCacheWithPreshadingCS::GetThreadGroupSize3D());
 	GroupCount.Z = FMath::DivideAndRoundUp(GroupCount.Z, FRenderLightingCacheWithPreshadingCS::GetThreadGroupSize3D());
@@ -639,8 +639,8 @@ void RenderSingleScatteringWithPreshadingCompute(
 		// Transmittance volume
 		if ((HeterogeneousVolumes::UseLightingCacheForTransmittance() && bApplyShadowTransmittance) || HeterogeneousVolumes::UseLightingCacheForInscattering())
 		{
-			float LODFactor = HeterogeneousVolumes::CalcLODFactor(View, HeterogeneousVolumeInterface);
-			PassParameters->LightingCache.LightingCacheResolution = HeterogeneousVolumes::GetLightingCacheResolution(HeterogeneousVolumeInterface, LODFactor);
+			HeterogeneousVolumes::FLODValue LODValue = HeterogeneousVolumes::CalcLOD(View, HeterogeneousVolumeInterface);
+			PassParameters->LightingCache.LightingCacheResolution = HeterogeneousVolumes::GetLightingCacheResolution(HeterogeneousVolumeInterface, LODValue);
 			PassParameters->LightingCache.LightingCacheVoxelBias = HeterogeneousVolumeInterface->GetShadowBiasFactor();
 			PassParameters->LightingCache.LightingCacheTexture = LightingCacheTexture;
 		}
