Heterogeneous Volumes: Remove ancillary creation of ExistenceMask

diff --git a/Engine/Source/Runtime/Renderer/Private/HeterogeneousVolumes/HeterogeneousVolumesAmbientOcclusionPipeline.cpp b/Engine/Source/Runtime/Renderer/Private/HeterogeneousVolumes/HeterogeneousVolumesAmbientOcclusionPipeline.cpp
--- a/Engine/Source/Runtime/Renderer/Private/HeterogeneousVolumes/HeterogeneousVolumesAmbientOcclusionPipeline.cpp
+++ b/Engine/Source/Runtime/Renderer/Private/HeterogeneousVolumes/HeterogeneousVolumesAmbientOcclusionPipeline.cpp
@@ -242,20 +242,6 @@ void RenderExistenceMaskWithLiveShading(
 	FRDGTextureRef& ExistenceMaskTexture
 )
 {
-	if (!HeterogeneousVolumes::UseExistenceMask())
-	{
-		FRDGTextureDesc ExistenceTextureDesc = FRDGTextureDesc::Create3D(
-			FIntVector(1),
-			PF_R8,
-			FClearValueBinding::Black,
-			TexCreate_ShaderResource | TexCreate_UAV | TexCreate_3DTiling
-		);
-
-		ExistenceMaskTexture = GraphBuilder.CreateTexture(ExistenceTextureDesc, TEXT("HeterogeneousVolumes.ExistenceMaskTexture"));
-		AddClearUAVPass(GraphBuilder, GraphBuilder.CreateUAV(ExistenceMaskTexture), 1.0f);
-		return;
-	}
-
 	FRDGTextureDesc ExistenceTextureDesc = FRDGTextureDesc::Create3D(
 		ExistenceTextureResolution,
 		PF_R8,
@@ -398,12 +384,6 @@ void DilateExistenceMask(
 	FRDGTextureRef& DilatedExistenceTexture
 )
 {
-	if (!HeterogeneousVolumes::UseExistenceMask())
-	{
-		DilatedExistenceTexture = ExistenceTexture;
-		return;
-	}
-
 	FRDGTextureDesc ExistenceTextureDesc = FRDGTextureDesc::Create3D(
 		ExistenceTextureResolution,
 		PF_R8,
@@ -763,33 +743,35 @@ void RenderAmbientOcclusionWithLiveShading(
 	FIntVector AmbientOcclusionTextureResolution = HeterogeneousVolumes::GetAmbientOcclusionResolution(HeterogeneousVolumeInterface, LODValue);
 
 	// Build existence mask
-	FRDGTextureRef ExistenceMaskTexture;
-	RenderExistenceMaskWithLiveShading(
-		GraphBuilder,
-		// Scene
-		Scene,
-		View,
-		SceneTextures,
-		// Object
-		HeterogeneousVolumeInterface,
-		DefaultMaterialRenderProxy,
-		PersistentPrimitiveIndex,
-		LocalBoxSphereBounds,
-		// Mask
-		AmbientOcclusionTextureResolution,
-		ExistenceMaskTexture
-	);
+	FRDGTextureRef DilatedExistenceMaskTexture = FRDGSystemTextures::Get(GraphBuilder).VolumetricBlack;
+	if (HeterogeneousVolumes::UseExistenceMask())
+	{
+		FRDGTextureRef ExistenceMaskTexture;
+		RenderExistenceMaskWithLiveShading(
+			GraphBuilder,
+			// Scene
+			Scene,
+			View,
+			SceneTextures,
+			// Object
+			HeterogeneousVolumeInterface,
+			DefaultMaterialRenderProxy,
+			PersistentPrimitiveIndex,
+			LocalBoxSphereBounds,
+			// Mask
+			AmbientOcclusionTextureResolution,
+			ExistenceMaskTexture
+		);
 
-	// Dilate by a pixel
-	FRDGTextureRef DilatedExistenceMaskTexture;
-	DilateExistenceMask(
-		GraphBuilder,
-		Scene,
-		View,
-		ExistenceMaskTexture,
-		AmbientOcclusionTextureResolution,
-		DilatedExistenceMaskTexture
-	);
+		DilateExistenceMask(
+			GraphBuilder,
+			Scene,
+			View,
+			ExistenceMaskTexture,
+			AmbientOcclusionTextureResolution,
+			DilatedExistenceMaskTexture
+		);
+	}
 
 	// Calculate ambient occlusion
 	FRDGTextureRef AmbientOcclusionTextureAsFixedPoint;
diff --git a/Engine/Source/Runtime/Renderer/Private/HeterogeneousVolumes/HeterogeneousVolumesLiveShadingPipeline.cpp b/Engine/Source/Runtime/Renderer/Private/HeterogeneousVolumes/HeterogeneousVolumesLiveShadingPipeline.cpp
--- a/Engine/Source/Runtime/Renderer/Private/HeterogeneousVolumes/HeterogeneousVolumesLiveShadingPipeline.cpp
+++ b/Engine/Source/Runtime/Renderer/Private/HeterogeneousVolumes/HeterogeneousVolumesLiveShadingPipeline.cpp
@@ -675,35 +675,38 @@ static void RenderLightingCacheWithLiveShading(
 	HeterogeneousVolumes::FLODValue LODValue = HeterogeneousVolumes::CalcLOD(View, HeterogeneousVolumeInterface);
 	FIntVector LightingCacheResolution = HeterogeneousVolumes::GetLightingCacheResolution(HeterogeneousVolumeInterface, LODValue);
 
-	FRDGTextureRef ExistenceMaskTexture = nullptr;
-	RenderExistenceMaskWithLiveShading(
-		GraphBuilder,
-		// Scene data
-		Scene,
-		View,
-		SceneTextures,
-		// Object data
-		HeterogeneousVolumeInterface,
-		DefaultMaterialRenderProxy,
-		PersistentPrimitiveIndex,
-		LocalBoxSphereBounds,
-		LightingCacheResolution,
-		// Output
-		ExistenceMaskTexture
-	);
+	FRDGTextureRef DilatedExistenceTexture = FRDGSystemTextures::Get(GraphBuilder).VolumetricBlack;
+	if (HeterogeneousVolumes::UseExistenceMask())
+	{
+		FRDGTextureRef ExistenceMaskTexture;
+		RenderExistenceMaskWithLiveShading(
+			GraphBuilder,
+			// Scene data
+			Scene,
+			View,
+			SceneTextures,
+			// Object data
+			HeterogeneousVolumeInterface,
+			DefaultMaterialRenderProxy,
+			PersistentPrimitiveIndex,
+			LocalBoxSphereBounds,
+			LightingCacheResolution,
+			// Output
+			ExistenceMaskTexture
+		);
 
-	FRDGTextureRef DilatedExistenceTexture;
-	DilateExistenceMask(
-		GraphBuilder,
-		// Scene data
-		Scene,
-		View,
-		// Existence texture data
-		ExistenceMaskTexture,
-		LightingCacheResolution,
-		// Output
-		DilatedExistenceTexture
-	);
+		DilateExistenceMask(
+			GraphBuilder,
+			// Scene data
+			Scene,
+			View,
+			// Existence texture data
+			ExistenceMaskTexture,
+			LightingCacheResolution,
+			// Output
+			DilatedExistenceTexture
+		);
+	}
 
 	FMatrix InstanceToLocal = HeterogeneousVolumeInterface->GetInstanceToLocal();
 	FMatrix LocalToWorld = HeterogeneousVolumeInterface->GetLocalToWorld();
