[GPU Scene Lights] Use TPersistentByteAddressBuffer instead of FRDGPooledBuffer/StructuredBuffer.
- In this change all GPU Scene Lights are still uploaded every frame.
- Upload using TByteAddressBufferScatterUploader.
- Extended RDG and UnifiedBuffer helpers to support ByteAddressUploadBuffer created using async setup tasks (using callbacks to get number of elements and data).

diff --git a/Engine/Shaders/Private/ByteBuffer.ush b/Engine/Shaders/Private/ByteBuffer.ush
new file mode 100644
--- /dev/null
+++ b/Engine/Shaders/Private/ByteBuffer.ush
@@ -0,0 +1,53 @@
+// Copyright Epic Games, Inc. All Rights Reserved.
+
+void LoadAndIncrementOffset(inout uint Value, ByteAddressBuffer Buffer, inout uint Offset)
+{
+	Value = Buffer.Load(Offset);
+	Offset += 4;
+}
+
+void LoadAndIncrementOffset(inout uint2 Value, ByteAddressBuffer Buffer, inout uint Offset)
+{
+	Value = Buffer.Load2(Offset);
+	Offset += 8;
+}
+
+void LoadAndIncrementOffset(inout uint3 Value, ByteAddressBuffer Buffer, inout uint Offset)
+{
+	Value = Buffer.Load3(Offset);
+	Offset += 12;
+}
+
+void LoadAndIncrementOffset(inout uint4 Value, ByteAddressBuffer Buffer, inout uint Offset)
+{
+	Value = Buffer.Load4(Offset);
+	Offset += 16;
+}
+
+void LoadAndIncrementOffset(inout float Value, ByteAddressBuffer Buffer, inout uint Offset)
+{
+	uint UintValue;
+	LoadAndIncrementOffset(UintValue, Buffer, Offset);
+	Value = asfloat(UintValue);
+}
+
+void LoadAndIncrementOffset(inout float2 Value, ByteAddressBuffer Buffer, inout uint Offset)
+{
+	uint2 UintValue;
+	LoadAndIncrementOffset(UintValue, Buffer, Offset);
+	Value = asfloat(UintValue);
+}
+
+void LoadAndIncrementOffset(inout float3 Value, ByteAddressBuffer Buffer, inout uint Offset)
+{
+	uint3 UintValue;
+	LoadAndIncrementOffset(UintValue, Buffer, Offset);
+	Value = asfloat(UintValue);
+}
+
+void LoadAndIncrementOffset(inout float4 Value, ByteAddressBuffer Buffer, inout uint Offset)
+{
+	uint4 UintValue;
+	LoadAndIncrementOffset(UintValue, Buffer, Offset);
+	Value = asfloat(UintValue);
+}
diff --git a/Engine/Shaders/Private/SceneData.ush b/Engine/Shaders/Private/SceneData.ush
--- a/Engine/Shaders/Private/SceneData.ush
+++ b/Engine/Shaders/Private/SceneData.ush
@@ -5,6 +5,7 @@
 #include "LargeWorldCoordinates.ush"
 #include "OctahedralCommon.ush"
 #include "WaveOpUtil.ush"
+#include "ByteBuffer.ush"
 #include "/Engine/Shared/LightSceneData.h"
 #include "/Engine/Shared/SceneDefinitions.h"
 
@@ -1390,18 +1391,50 @@ FSceneDataIntermediates GetSceneDataIntermediates()
 //
 
 #if USE_GLOBAL_GPU_SCENE_DATA
-StructuredBuffer<FLightSceneData> GPUSceneLights;
+ByteAddressBuffer GPUSceneLights;
 #elif USE_GPU_SCENE_DATA_RW
-RWStructuredBuffer<FLightSceneData> GPUSceneLightsRW;
+RWByteAddressBuffer GPUSceneLightsRW;
 #endif
 
-FLightSceneData GetLightSceneData(int LightId)
-{
 #if USE_GLOBAL_GPU_SCENE_DATA
-	return GPUSceneLights[LightId];
+#define GPUSceneLightsBuffer GPUSceneLights
 #elif USE_GPU_SCENE_DATA_RW
-	return GPUSceneLightsRW[LightId];
+#define GPUSceneLightsBuffer GPUSceneLightsRW
+#else
+#define GPUSceneLightsBuffer Scene.GPUScene.GPUSceneLightData
+#endif
+
+FLightSceneData GetLightSceneData(int LightId)
+{
+#if 0
+	return GPUSceneLightsBuffer.Load<FLightSceneData>(LightId * sizeof(FLightSceneData));
 #else
-	return Scene.GPUScene.GPUSceneLightData[LightId];
+	FLightSceneData Output;
+
+	uint Offset = LightId * sizeof(FLightSceneData);
+
+	LoadAndIncrementOffset(Output.WorldPosition.High, GPUSceneLightsBuffer, Offset);
+	LoadAndIncrementOffset(Output.WorldPosition.Low, GPUSceneLightsBuffer, Offset);
+	LoadAndIncrementOffset(Output.InvRadius, GPUSceneLightsBuffer, Offset);
+	LoadAndIncrementOffset(Output.FalloffExponent, GPUSceneLightsBuffer, Offset);
+	LoadAndIncrementOffset(Output.Color, GPUSceneLightsBuffer, Offset);
+	LoadAndIncrementOffset(Output.Direction, GPUSceneLightsBuffer, Offset);
+	LoadAndIncrementOffset(Output.LightTypeAndShadowMapChannelMaskPacked, GPUSceneLightsBuffer, Offset);
+	LoadAndIncrementOffset(Output.Tangent, GPUSceneLightsBuffer, Offset);
+	LoadAndIncrementOffset(Output.SourceRadius, GPUSceneLightsBuffer, Offset);
+	LoadAndIncrementOffset(Output.SpotAngles, GPUSceneLightsBuffer, Offset);
+	LoadAndIncrementOffset(Output.SoftSourceRadius, GPUSceneLightsBuffer, Offset);
+	LoadAndIncrementOffset(Output.SourceLength, GPUSceneLightsBuffer, Offset);
+	LoadAndIncrementOffset(Output.RectLightBarnCosAngle, GPUSceneLightsBuffer, Offset);
+	LoadAndIncrementOffset(Output.RectLightBarnLength, GPUSceneLightsBuffer, Offset);
+	LoadAndIncrementOffset(Output.RectLightAtlasUVOffset, GPUSceneLightsBuffer, Offset);
+	LoadAndIncrementOffset(Output.RectLightAtlasUVScale, GPUSceneLightsBuffer, Offset);
+	LoadAndIncrementOffset(Output.RectLightAtlasMaxLevel, GPUSceneLightsBuffer, Offset);
+	LoadAndIncrementOffset(Output.InverseExposureBlend, GPUSceneLightsBuffer, Offset);
+	LoadAndIncrementOffset(Output.IESAtlasIndex, GPUSceneLightsBuffer, Offset);
+	LoadAndIncrementOffset(Output.SpecularScale, GPUSceneLightsBuffer, Offset);
+	LoadAndIncrementOffset(Output.DiffuseScale, GPUSceneLightsBuffer, Offset);
 #endif
+
+	return Output;
 }
diff --git a/Engine/Source/Runtime/RenderCore/Private/UnifiedBuffer.cpp b/Engine/Source/Runtime/RenderCore/Private/UnifiedBuffer.cpp
--- a/Engine/Source/Runtime/RenderCore/Private/UnifiedBuffer.cpp
+++ b/Engine/Source/Runtime/RenderCore/Private/UnifiedBuffer.cpp
@@ -764,6 +764,36 @@ void ScatterCopyResource(FRDGBuilder& GraphBuilder, FRDGViewableResource* DstRes
 		DispatchConfig.GroupCountWrapped);
 }
 
+void ScatterCopyResource(FRDGBuilder& GraphBuilder, FRDGBuffer* DstResource, FRDGBufferSRV* ScatterBufferSRV, FRDGBufferSRV* UploadBufferSRV, const FAsyncScatterCopyParams& Params)
+{
+	const FScatterUploadConfig Config = GetScatterUploadConfig(Params.NumBytesPerElement, Params.NumElementsPerScatter);
+
+	FRDGScatterCopyCS::FParameters Parameters;
+	FRDGScatterCopyCS::FPermutationDomain PermutationVector;
+	GetScatterCopyParamsAndPermutation(GraphBuilder, DstResource, ScatterBufferSRV, UploadBufferSRV, Config, 0, Parameters, PermutationVector);
+
+	TShaderMapRef<FRDGScatterCopyCS> ComputeShader(GetGlobalShaderMap(GMaxRHIFeatureLevel), PermutationVector);
+
+	FRDGScatterCopyCS::FParameters* PassParameters = GraphBuilder.AllocParameters(&Parameters);
+
+	GraphBuilder.AddPass(
+		RDG_EVENT_NAME("ScatterUpload"),
+		PassParameters,
+		ERDGPassFlags::Compute,
+		[PassParameters, ComputeShader, Config, GetNumScatters = Params.GetNumScatters](FRDGAsyncTask, FRHIComputeCommandList& RHICmdList)
+		{
+			PassParameters->NumScatters = GetNumScatters();
+			const FScatterUploadDispatchConfig DispatchConfig = GetScatterUploadDispatchConfig(Config, PassParameters->NumScatters);
+
+			FComputeShaderUtils::Dispatch(
+				RHICmdList,
+				ComputeShader,
+				*PassParameters,
+				DispatchConfig.GroupCountWrapped);
+		}
+	);
+}
+
 void ScatterCopyResource(FRDGBuilder& GraphBuilder, FRDGBuffer* DstResource, FRDGBufferSRV* ScatterBufferSRV, FRDGBufferSRV* UploadBufferSRV, const FScatterCopyParams &Params)
 {
 	PRAGMA_DISABLE_DEPRECATION_WARNINGS
diff --git a/Engine/Source/Runtime/RenderCore/Public/RenderGraphUtils.h b/Engine/Source/Runtime/RenderCore/Public/RenderGraphUtils.h
--- a/Engine/Source/Runtime/RenderCore/Public/RenderGraphUtils.h
+++ b/Engine/Source/Runtime/RenderCore/Public/RenderGraphUtils.h
@@ -1169,6 +1169,51 @@ FORCEINLINE FRDGBufferRef CreateByteAddressUploadBuffer(
 	return Buffer;
 }
 
+/** A variant where NumElements, InitialData, and InitialDataSize are supplied through callbacks. This allows creating a buffer with
+ *  information unknown at creation time. Though, data must be ready before the most recent RDG pass that references the buffer
+ *  is executed. For byte address buffers, NumElements must be Size / 4.
+ */
+FORCEINLINE FRDGBufferRef CreateByteAddressUploadBuffer(
+	FRDGBuilder& GraphBuilder,
+	const TCHAR* Name,
+	FRDGBufferNumElementsCallback&& NumElementsCallback,
+	FRDGBufferInitialDataCallback&& InitialDataCallback,
+	FRDGBufferInitialDataSizeCallback&& InitialDataSizeCallback)
+{
+	FRDGBufferRef Buffer = GraphBuilder.CreateBuffer(FRDGBufferDesc::CreateByteAddressUploadDesc(4), Name, MoveTemp(NumElementsCallback));
+	GraphBuilder.QueueBufferUpload(Buffer, MoveTemp(InitialDataCallback), MoveTemp(InitialDataSizeCallback));
+	return Buffer;
+}
+
+/** A variant where the TArray is supplied through callbacks. This allows creating a buffer with
+ *  information unknown at creation time. Though, data must be ready before the most recent RDG pass that references the buffer
+ *  is executed.
+ */
+template <typename ArrayType>
+FORCEINLINE FRDGBufferRef CreateByteAddressUploadBuffer_Impl(
+	FRDGBuilder& GraphBuilder,
+	const TCHAR* Name,
+	TRDGBufferArrayCallback<ArrayType>&& ArrayCallback)
+{
+	return CreateByteAddressUploadBuffer(GraphBuilder, Name,
+		/*NumElementsCallback = */[ArrayCallback]() { const auto& Array = ArrayCallback(); return Array.Num() * Array.GetTypeSize() / 4; },
+		/*InitialDataCallback = */[ArrayCallback]() { return ArrayCallback().GetData(); },
+		/*InitialDataSizeCallback = */[ArrayCallback]() { const auto& Array = ArrayCallback(); return Array.Num() * Array.GetTypeSize(); });
+}
+
+/** Same as the previous function but where the type of the array is automatically inferred, so we can do :
+ *  TArray<FSomeType> Array = GraphBuilder.AllocArray<...>();
+ *  CreateByteAddressBuffer(..., [&]() -> auto&{ return Array; });
+ */
+template <typename GetArrayRefCallback, typename Type = TInvokeResult_T<GetArrayRefCallback>>
+FORCEINLINE FRDGBufferRef CreateByteAddressUploadBuffer(
+	FRDGBuilder& GraphBuilder,
+	const TCHAR* Name,
+	GetArrayRefCallback&& ArrayCallback)
+{
+	return CreateByteAddressUploadBuffer_Impl<Type>(GraphBuilder, Name, MoveTemp(ArrayCallback));
+}
+
 /** Creates a vertex buffer with initial data by creating an upload pass. */
 RENDERCORE_API FRDGBufferRef CreateVertexBuffer(
 	FRDGBuilder& GraphBuilder,
diff --git a/Engine/Source/Runtime/RenderCore/Public/UnifiedBuffer.h b/Engine/Source/Runtime/RenderCore/Public/UnifiedBuffer.h
--- a/Engine/Source/Runtime/RenderCore/Public/UnifiedBuffer.h
+++ b/Engine/Source/Runtime/RenderCore/Public/UnifiedBuffer.h
@@ -204,6 +204,15 @@ struct FScatterCopyParams
 
 void RENDERCORE_API ScatterCopyResource(FRDGBuilder& GraphBuilder, FRDGBuffer* DstResource, FRDGBufferSRV* ScatterBufferSRV, FRDGBufferSRV* UploadBufferSRV, const FScatterCopyParams &Params);
 
+struct FAsyncScatterCopyParams
+{
+	TFunction<uint64()> GetNumScatters;
+	uint32 NumBytesPerElement = 0u;
+	int32 NumElementsPerScatter = INDEX_NONE; // INDEX_NONE lets the setup figure it out, otherwise it will run NumScatters * NumElementsPerScatter threads to copy the source data.
+};
+
+void RENDERCORE_API ScatterCopyResource(FRDGBuilder& GraphBuilder, FRDGBuffer* DstResource, FRDGBufferSRV* ScatterBufferSRV, FRDGBufferSRV* UploadBufferSRV, const FAsyncScatterCopyParams& Params);
+
 extern RENDERCORE_API FRDGBuffer* ResizeBufferIfNeeded(FRDGBuilder& GraphBuilder, TRefCountPtr<FRDGPooledBuffer>& ExternalBuffer, const FRDGBufferDesc& BufferDesc, const TCHAR* Name);
 extern RENDERCORE_API FRDGBuffer* ResizeBufferIfNeeded(FRDGBuilder& GraphBuilder, TRefCountPtr<FRDGPooledBuffer>& ExternalBuffer, EPixelFormat Format, uint32 NumElements, const TCHAR* Name);
 extern RENDERCORE_API FRDGBuffer* ResizeStructuredBufferIfNeeded(FRDGBuilder& GraphBuilder, TRefCountPtr<FRDGPooledBuffer>& ExternalBuffer, uint32 NumBytes, const TCHAR* Name);
diff --git a/Engine/Source/Runtime/Renderer/Private/GPUScene.cpp b/Engine/Source/Runtime/Renderer/Private/GPUScene.cpp
--- a/Engine/Source/Runtime/Renderer/Private/GPUScene.cpp
+++ b/Engine/Source/Runtime/Renderer/Private/GPUScene.cpp
@@ -118,6 +118,13 @@ static TAutoConsoleVariable<int32> CVarGPUSceneUseGrowOnlyAllocationPolicy(
 	ECVF_RenderThreadSafe | ECVF_ReadOnly
 );
 
+static TAutoConsoleVariable<int32> CVarGPUSceneLightsAsyncSetup(
+	TEXT("r.GPUScene.Lights.AsyncSetup"),
+	1,
+	TEXT("Whether to run GPU Scene Lights setup in an async task."),
+	ECVF_RenderThreadSafe
+);
+
 LLM_DECLARE_TAG_API(GPUScene, RENDERER_API);
 DECLARE_LLM_MEMORY_STAT(TEXT("GPUScene"), STAT_GPUSceneLLM, STATGROUP_LLMFULL);
 DECLARE_LLM_MEMORY_STAT(TEXT("GPUScene"), STAT_GPUSceneSummaryLLM, STATGROUP_LLM);
@@ -638,6 +645,7 @@ FGPUScene::FGPUScene(FScene &InScene)
 	, InstancePayloadDataAllocator(CVarGPUSceneUseGrowOnlyAllocationPolicy.GetValueOnAnyThread() != 0)
 	, LightmapDataAllocator(CVarGPUSceneUseGrowOnlyAllocationPolicy.GetValueOnAnyThread() != 0)
 	, InstanceSceneDataSOAStride(0)
+	, LightDataBuffer(InitialBufferSize, TEXT("GPUScene.LightData"))
 	, Scene(InScene)
 	, InstanceSceneDataAllocator(CVarGPUSceneUseGrowOnlyAllocationPolicy.GetValueOnAnyThread() != 0)
 {
@@ -727,9 +735,11 @@ void FGPUScene::UpdateGPULights(FRDGBuilder& GraphBuilder, const UE::Tasks::FTas
 	// Clamp buffer to be smaller than the MAX_LIGHT_ID.
 	NumGPULights = FMath::Min(int32(MAX_LIGHT_ID), Scene.Lights.GetMaxIndex());
 
-	FRDGUploadData<FLightSceneData> LightData(GraphBuilder, FMath::Max(1, NumGPULights));
+	auto* LightDataUploader = TAsyncByteAddressBufferScatterUploader<FLightSceneData>::Allocate(GraphBuilder, NumGPULights);
 
-	GraphBuilder.AddSetupTask([this, LightData]
+	const bool bAsyncSetup = CVarGPUSceneLightsAsyncSetup.GetValueOnRenderThread() != 0;
+
+	GraphBuilder.AddSetupTask([this, LightDataUploader]
 	{
 		SCOPED_NAMED_EVENT(UpdateGPUScene_Lights, FColor::Green);
 		const bool bAllowStaticLighting = IsStaticLightingAllowed();
@@ -737,49 +747,34 @@ void FGPUScene::UpdateGPULights(FRDGBuilder& GraphBuilder, const UE::Tasks::FTas
 
 		for (int32 Index = 0; Index < NumGPULights; ++Index)
 		{
+			FLightSceneData LightData;
+
 			if (Scene.Lights.IsAllocated(Index))
 			{
-				InitLightData(Scene.Lights[Index], bAllowStaticLighting, Flags, LightData[Index]);
+				InitLightData(Scene.Lights[Index], bAllowStaticLighting, Flags, LightData);
 			}
 			else
 			{
-				LightData[Index].WorldPosition = FDFVector3{};
-				LightData[Index].Color = FVector3f::ZeroVector;
-				LightData[Index].InvRadius = 0.0f;
+				LightData.WorldPosition = FDFVector3{};
+				LightData.Color = FVector3f::ZeroVector;
+				LightData.InvRadius = 0.0f;
 			}
-		}
-
-	}, PrerequisiteTask);
 
-	const uint32 LightDataBufferNumElements = FMath::RoundUpToPowerOfTwo(FMath::Max(Scene.Lights.GetMaxIndex(), InitialBufferSize));
-	const FRDGBufferDesc LightDataBufferDesc = FRDGBufferDesc::CreateStructuredDesc(sizeof(FLightSceneData), LightDataBufferNumElements);
+			LightDataUploader->Add(LightData, Index);
+		}
+	},
+	PrerequisiteTask,
+	UE::Tasks::ETaskPriority::Normal,
+	bAsyncSetup);
 
-	FRDGBuffer* LightDataBufferRDG;
-	if (LightDataBuffer == nullptr)
+	if (bAsyncSetup)
 	{
-		LightDataBufferRDG = GraphBuilder.CreateBuffer(LightDataBufferDesc, TEXT("GPUScene.LightData"));
-		LightDataBuffer = GraphBuilder.ConvertToExternalBuffer(LightDataBufferRDG);
+		LightDataUploader->ResizeAndUploadToAsync(GraphBuilder, LightDataBuffer, NumGPULights);
 	}
 	else
 	{
-		const uint32 BufferSizeNew = LightDataBufferDesc.GetSize();
-		const uint32 BufferSizeOld = LightDataBuffer->GetCommittedSize();
-
-		if (BufferSizeNew == BufferSizeOld)
-		{
-			LightDataBufferRDG = GraphBuilder.RegisterExternalBuffer(LightDataBuffer);
-		}
-		else
-		{
-			LightDataBufferRDG = GraphBuilder.CreateBuffer(LightDataBufferDesc, TEXT("GPUScene.LightData"));
-			LightDataBuffer = GraphBuilder.ConvertToExternalBuffer(LightDataBufferRDG);
-
-			// Currently whole contents of LightDataBuffer are uploaded so don't need to copy previous contents when resizing
-		}
-
+		LightDataUploader->ResizeAndUploadTo(GraphBuilder, LightDataBuffer, NumGPULights);
 	}
-
-	GraphBuilder.QueueBufferUpload<FLightSceneData>(LightDataBufferRDG, LightData, ERDGInitialDataFlags::NoCopy);
 }
 
 void FGPUScene::InitLightData(const FLightSceneInfoCompact& LightInfoCompact, bool bAllowStaticLighting, uint32 LightShaderParameterFlags, FLightSceneData& DataOut)
@@ -967,7 +962,7 @@ FGPUScene::FRegisteredBuffers FGPUScene::UpdateBufferAllocations(FRDGBuilder& Gr
 	const uint32 LightMapDataBufferSize = FMath::RoundUpToPowerOfTwo(FMath::Max(LightmapDataAllocator.GetMaxSize(), InitialBufferSize));
 	BufferState.LightmapDataBuffer = ResizeStructuredBufferIfNeeded(GraphBuilder, LightmapDataBuffer, LightMapDataBufferSize * sizeof(FLightmapSceneShaderData::Data), TEXT("GPUScene.LightmapData"));
 
-	BufferState.LightDataBuffer = GraphBuilder.RegisterExternalBuffer(LightDataBuffer);
+	BufferState.LightDataBuffer = GraphBuilder.RegisterExternalBuffer(LightDataBuffer.GetPooledBuffer());
 
 	CommonParameters.GPUSceneInstanceDataSOAStride = InstanceSceneDataSOAStride;
 	CommonParameters.GPUSceneFrameNumber = GetSceneFrameNumber();
@@ -1006,7 +1001,7 @@ FGPUScene::FRegisteredBuffers FGPUScene::RegisterBuffers(FRDGBuilder& GraphBuild
 	Result.InstanceSceneDataBuffer = GraphBuilder.RegisterExternalBuffer(InstanceSceneDataBuffer);
 	Result.InstancePayloadDataBuffer = GraphBuilder.RegisterExternalBuffer(InstancePayloadDataBuffer);
 	Result.LightmapDataBuffer = GraphBuilder.RegisterExternalBuffer(LightmapDataBuffer);
-	Result.LightDataBuffer = GraphBuilder.RegisterExternalBuffer(LightDataBuffer);
+	Result.LightDataBuffer = GraphBuilder.RegisterExternalBuffer(LightDataBuffer.GetPooledBuffer());
 
 	return Result;
 }
@@ -2182,7 +2177,7 @@ void FGPUScene::OnPostLightSceneInfoUpdate(FRDGBuilder& GraphBuilder, const FLig
 	// SceneFrameNumber is updated in UpdateInternal so if it is the same, we have already uploaded the lights this "frame" - this is not 100% robust since 
 	// if UpdateAllPrimitiveSceneInfos is called multiple times without updating the FScene frame number it will skip light uploads. The real solution is 
 	// to implement functional change tracking for lights so we can actually know when there are relevant changes.
-	if (SceneFrameNumber != Scene.GetFrameNumberRenderThread() || !LightDataBuffer.IsValid())
+	if (SceneFrameNumber != Scene.GetFrameNumberRenderThread())
 	{
 		UpdateGPULights(GraphBuilder, UE::Tasks::FTask{});
 	}
@@ -2205,7 +2200,7 @@ FGPUSceneResourceParametersRHI FGPUScene::GetShaderParametersRHI() const
 	Result.GPUSceneInstancePayloadData = InstancePayloadDataBuffer->GetSRV();
 	Result.GPUScenePrimitiveSceneData = PrimitiveBuffer->GetSRV();
 	Result.GPUSceneLightmapData = LightmapDataBuffer->GetSRV();
-	Result.GPUSceneLightData = LightDataBuffer->GetSRV();
+	Result.GPUSceneLightData = LightDataBuffer.GetPooledBuffer()->GetSRV();
 
 	return Result;
 }
diff --git a/Engine/Source/Runtime/Renderer/Private/GPUScene.h b/Engine/Source/Runtime/Renderer/Private/GPUScene.h
--- a/Engine/Source/Runtime/Renderer/Private/GPUScene.h
+++ b/Engine/Source/Runtime/Renderer/Private/GPUScene.h
@@ -15,6 +15,7 @@
 #include "UnifiedBuffer.h"
 #include "ScenePrimitiveUpdates.h"
 #include "GPUSceneWriter.h"
+#include "RendererPrivateUtils.h"
 
 class FRDGExternalAccessQueue;
 class FRHICommandList;
@@ -56,7 +57,7 @@ BEGIN_SHADER_PARAMETER_STRUCT(FGPUSceneResourceParameters, RENDERER_API)
 	SHADER_PARAMETER_RDG_BUFFER_SRV(StructuredBuffer<float4>, GPUSceneInstancePayloadData)
 	SHADER_PARAMETER_RDG_BUFFER_SRV(StructuredBuffer<float4>, GPUScenePrimitiveSceneData)
 	SHADER_PARAMETER_RDG_BUFFER_SRV(StructuredBuffer<float4>, GPUSceneLightmapData)
-	SHADER_PARAMETER_RDG_BUFFER_SRV(StructuredBuffer<FLightSceneData>, GPUSceneLightData)
+	SHADER_PARAMETER_RDG_BUFFER_SRV(ByteAddressBuffer, GPUSceneLightData)
 END_SHADER_PARAMETER_STRUCT()
 
 //Non-RDG paramters, untracked, use with care (or preferably, don't)
@@ -66,7 +67,7 @@ BEGIN_SHADER_PARAMETER_STRUCT(FGPUSceneResourceParametersRHI, RENDERER_API)
 	SHADER_PARAMETER_SRV(StructuredBuffer<float4>, GPUSceneInstancePayloadData)
 	SHADER_PARAMETER_SRV(StructuredBuffer<float4>, GPUScenePrimitiveSceneData)
 	SHADER_PARAMETER_SRV(StructuredBuffer<float4>, GPUSceneLightmapData)
-	SHADER_PARAMETER_SRV(StructuredBuffer<FLightSceneData>, GPUSceneLightData)
+	SHADER_PARAMETER_SRV(ByteAddressBuffer, GPUSceneLightData)
 END_SHADER_PARAMETER_STRUCT()
 
 DECLARE_SCENE_UB_STRUCT(FGPUSceneResourceParameters, GPUScene, RENDERER_API)
@@ -369,8 +370,9 @@ private:
 
 	static constexpr int32 InitialBufferSize = 256;
 
-	TRefCountPtr<FRDGPooledBuffer>	LightDataBuffer;
-	int32							NumGPULights = 0;
+	TPersistentByteAddressBuffer<FLightSceneData> LightDataBuffer;
+	int32 NumGPULights = 0;
+
 	FGPUSceneCommonParameters CommonParameters;
 
 	FGPUSceneCommonParameters GetCommonParameters() const;
diff --git a/Engine/Source/Runtime/Renderer/Private/Renderer.cpp b/Engine/Source/Runtime/Renderer/Private/Renderer.cpp
--- a/Engine/Source/Runtime/Renderer/Private/Renderer.cpp
+++ b/Engine/Source/Runtime/Renderer/Private/Renderer.cpp
@@ -225,7 +225,7 @@ FSceneUniformBuffer* FRendererModule::CreateSinglePrimitiveSceneUniformBuffer(FR
 			FRDGBufferRef LightmapSceneDataBuffer = CreateStructuredBuffer(GraphBuilder, TEXT("LightmapSceneDataBuffer"), TConstArrayView<FVector4f>(LightmapSceneData.Data));
 			FRDGBufferRef InstanceSceneDataBuffer = CreateStructuredBuffer(GraphBuilder, TEXT("InstanceSceneDataBuffer"), TConstArrayView<FVector4f>(InstanceSceneData.Data));
 			FRDGBufferRef InstancePayloadDataBuffer = GSystemTextures.GetDefaultStructuredBuffer(GraphBuilder, sizeof(FVector4f));
-			FRDGBufferRef DummyBufferLight = GSystemTextures.GetDefaultStructuredBuffer(GraphBuilder, sizeof(FLightSceneData));
+			FRDGBufferRef DummyBufferLight = GSystemTextures.GetDefaultByteAddressBuffer(GraphBuilder, sizeof(FLightSceneData));
 
 			FGPUSceneResourceParameters ShaderParameters;
 			ShaderParameters.GPUScenePrimitiveSceneData = GraphBuilder.CreateSRV(PrimitiveSceneDataBuffer);
diff --git a/Engine/Source/Runtime/Renderer/Private/RendererPrivateUtils.h b/Engine/Source/Runtime/Renderer/Private/RendererPrivateUtils.h
--- a/Engine/Source/Runtime/Renderer/Private/RendererPrivateUtils.h
+++ b/Engine/Source/Runtime/Renderer/Private/RendererPrivateUtils.h
@@ -102,6 +102,7 @@ public:
 	};
 
 	void UploadTo(FRDGBuilder& GraphBuilder, FRDGBuffer *DestBuffer, FRDGBuffer *ScatterOffsets, FRDGBuffer *Values, uint32 NumScatters, uint32 NumBytesPerElement, int32 NumValuesPerScatter);
+	void UploadTo(FRDGBuilder& GraphBuilder, FRDGBuffer* DestBuffer, FRDGBuffer* ScatterOffsets, FRDGBuffer* Values, TFunction<uint64()>&& GetNumScatters, uint32 NumBytesPerElement, int32 NumValuesPerScatter);
 };
 
 namespace UE::RendererPrivateUtils::Implementation
@@ -155,6 +156,12 @@ struct FByteAddressBufferTraits
 		return CreateByteAddressUploadBuffer(GraphBuilder, Name, InitialData, ERDGInitialDataFlags::NoCopy);
 	}
 
+	template <typename GetArrayRefCallback>
+	static FRDGBuffer* CreateUploadBuffer(FRDGBuilder& GraphBuilder, const TCHAR* Name, GetArrayRefCallback&& ArrayCallback)
+	{
+		return CreateByteAddressUploadBuffer(GraphBuilder, Name, MoveTemp(ArrayCallback));
+	}
+
 	static FRDGBufferDesc CreateDesc(uint32 BytesPerElement, uint32 NumElements)
 	{
 		uint32 NumBytes = BytesPerElement * NumElements;
@@ -218,6 +225,7 @@ public:
 	}
 
 	TRefCountPtr<FRDGPooledBuffer> &GetPooledBuffer() { return PooledBuffer; }
+	const TRefCountPtr<FRDGPooledBuffer> &GetPooledBuffer() const { return PooledBuffer; }
 };
 
 template <typename InValueType, typename InBufferTraits,  int32 InNumValuesPerScatter = 1>
@@ -239,6 +247,8 @@ public:
 		Reserve(NumScatters);
 	}
 
+	virtual ~TBufferScatterUploader() = default;
+
 	/**
 	 * Pre-size the value and scatter arrays, allocates space for InNumValuesPerScatter * NumScatters values & NumScatters scatter offsets.
 	 */
@@ -353,7 +363,7 @@ public:
 		return ResizeAndUploadTo(GraphBuilder, DestDataBuffer, DestDataMinimumSize, ScatterInfo);
 	}
 
-private:
+protected:
 	struct FUploadData
 	{
 		TArray<ValueType> Values;
@@ -363,6 +373,64 @@ private:
 	FUploadData UploadData;
 };
 
+template <typename InValueType, typename InBufferTraits, int32 InNumValuesPerScatter = 1>
+class TAsyncBufferScatterUploader : public TBufferScatterUploader<InValueType, InBufferTraits, InNumValuesPerScatter>
+{
+	using Base = TBufferScatterUploader<InValueType, InBufferTraits, InNumValuesPerScatter>;
+
+public:
+
+	/**
+	* Allocate an instance of TAsyncBufferScatterUploader using an allocator tied to the lifetime of the graph
+	*/
+	static TAsyncBufferScatterUploader* Allocate(FRDGBuilder& GraphBuilder, int32 NumScatters = 0)
+	{
+		return GraphBuilder.AllocObject<TAsyncBufferScatterUploader>(NumScatters);
+	}
+
+	/**
+	 * Resize the destination persistent buffer (if needed) and upload & scatter the collected data to it.
+	 * Callbacks are used to create upload buffers and scheduling scatter dispatch since data might not be ready yet.
+	 */
+	FRDGBuffer* ResizeAndUploadToAsync(FRDGBuilder& GraphBuilder, TPersistentBuffer<InValueType, InBufferTraits>& DestDataBuffer, int32 DestDataMinimumSize)
+	{
+		check(Base::UploadDataProxy == nullptr);
+
+		FRDGBuffer* DestBufferRDG = DestDataBuffer.ResizeBufferIfNeeded(GraphBuilder, DestDataMinimumSize);
+
+		// upload the values & offsets
+		FRDGBufferRef ScatterOffsetsRDG = Base::BufferTraits::CreateUploadBuffer(GraphBuilder, TEXT("ScatterUploader.Offsets"), [&]() -> auto& { return Base::UploadData.ScatterOffsets; });
+		FRDGBufferRef ValuesRDG = Base::BufferTraits::CreateUploadBuffer(GraphBuilder, TEXT("ScatterUploader.Values"), [&]() -> auto& { return Base::UploadData.Values; });
+
+		uint32 ElementSize = sizeof(typename Base::ValueType);
+		uint32 ElementsPerScatter = Base::NumValuesPerScatter;
+		if (Base::BufferTraits::bAutoValuesPerScatter)
+		{
+			// Let the implementation determine the optimal way to divide up the scatter
+			ElementSize *= ElementsPerScatter;
+			ElementsPerScatter = INDEX_NONE;
+		}
+		FBufferScatterUploader::UploadTo(GraphBuilder, DestBufferRDG, ScatterOffsetsRDG, ValuesRDG, [&]() -> int64 { return Base::UploadData.ScatterOffsets.Num(); }, ElementSize, ElementsPerScatter);
+
+		return DestBufferRDG;
+	}
+
+private:
+
+	/**
+	* Optionally reserve space for NumScatters items.
+	* 
+	* Constructor is private since instances should be allocated directly from GraphBuilder or using the static helper function Allocate(...)
+	* because ResizeAndUploadToAsync(...) requires that UploadData is tied to the lifetime of the graph
+	*/
+	TAsyncBufferScatterUploader(int32 NumScatters = 0) : Base(NumScatters)
+	{
+	}
+
+	friend class FRDGAllocator;
+	friend class FRDGAllocator::TObject<TAsyncBufferScatterUploader>;
+};
+
 }
 
 template <typename InValueType>
@@ -377,6 +445,9 @@ using TPersistentByteAddressBuffer = UE::RendererPrivateUtils::Implementation::T
 template <typename InValueType, int32 InNumValuesPerScatter = 1>
 using TByteAddressBufferScatterUploader = UE::RendererPrivateUtils::Implementation::TBufferScatterUploader<InValueType, UE::RendererPrivateUtils::Implementation::FByteAddressBufferTraits, InNumValuesPerScatter>;
 
+template <typename InValueType, int32 InNumValuesPerScatter = 1>
+using TAsyncByteAddressBufferScatterUploader = UE::RendererPrivateUtils::Implementation::TAsyncBufferScatterUploader<InValueType, UE::RendererPrivateUtils::Implementation::FByteAddressBufferTraits, InNumValuesPerScatter>;
+
 /**
  * Clear indirect args for GPU-side building, initializes the XYZ dimensions of the dispatch arg to DimClearValue and everything else in the strided range to zero. 
  * Can clear one or more args (NumIndirectArgs). 
diff --git a/Engine/Source/Runtime/Renderer/Private/RendererUtils.cpp b/Engine/Source/Runtime/Renderer/Private/RendererUtils.cpp
--- a/Engine/Source/Runtime/Renderer/Private/RendererUtils.cpp
+++ b/Engine/Source/Runtime/Renderer/Private/RendererUtils.cpp
@@ -257,6 +257,12 @@ void FBufferScatterUploader::UploadTo(FRDGBuilder& GraphBuilder, FRDGBuffer *Des
 	ScatterCopyResource(GraphBuilder, DestBuffer, GraphBuilder.CreateSRV(ScatterOffsets), GraphBuilder.CreateSRV(Values), ScatterCopyParams);
 }
 
+void FBufferScatterUploader::UploadTo(FRDGBuilder& GraphBuilder, FRDGBuffer* DestBuffer, FRDGBuffer* ScatterOffsets, FRDGBuffer* Values, TFunction<uint64()>&& GetNumScatters, uint32 NumBytesPerElement, int32 NumValuesPerScatter)
+{
+	FAsyncScatterCopyParams ScatterCopyParams{ MoveTemp(GetNumScatters), NumBytesPerElement, NumValuesPerScatter };
+	ScatterCopyResource(GraphBuilder, DestBuffer, GraphBuilder.CreateSRV(ScatterOffsets), GraphBuilder.CreateSRV(Values), ScatterCopyParams);
+}
+
 namespace UE::RendererPrivateUtils::Implementation
 {
 
