Fix Mobile Renderer Pass Parameters

The mobile renderer was caching parameters that needed to be unique per pass.  When there is only a single pass of each type, this is fine, but when there are multiple passes, such as when there are multiple views, it's not.  Each view stomps the previous view's instance data and then all the views use the same instance data.  This results in one view looking correct and the rest going nuts.  Remove the state caching and instead allocate the parameters per pass via the graph builder.

diff --git a/Engine/Source/Runtime/Renderer/Private/MobileBasePassRendering.cpp b/Engine/Source/Runtime/Renderer/Private/MobileBasePassRendering.cpp
--- a/Engine/Source/Runtime/Renderer/Private/MobileBasePassRendering.cpp
+++ b/Engine/Source/Runtime/Renderer/Private/MobileBasePassRendering.cpp
@@ -510,7 +510,7 @@ void SetupMobileSkyReflectionUniformParameters(FSkyLightSceneProxy* SkyLight, FM
 	Parameters.TextureBlendSampler = CaptureTextureBlend->SamplerStateRHI;
 }
 
-void FMobileSceneRenderer::RenderMobileBasePass(FRHICommandList& RHICmdList, const FViewInfo& View, const FInstanceCullingDrawParams* InstanceCullingDrawParams)
+void FMobileSceneRenderer::RenderMobileBasePass(FRHICommandList& RHICmdList, const FViewInfo& View, const FInstanceCullingDrawParams* InstanceCullingDrawParams, const FInstanceCullingDrawParams* SkyPassInstanceCullingDrawParams)
 {
 	CSV_SCOPED_TIMING_STAT_EXCLUSIVE(RenderBasePass);
 	SCOPE_CYCLE_COUNTER(STAT_BasePassDrawTime);
@@ -523,7 +523,7 @@ void FMobileSceneRenderer::RenderMobileBasePass(FRHICommandList& RHICmdList, con
 		
 	if (View.Family->EngineShowFlags.Atmosphere)
 	{
-		View.ParallelMeshDrawCommandPasses[EMeshPass::SkyPass].Draw(RHICmdList, &SkyPassInstanceCullingDrawParams);
+		View.ParallelMeshDrawCommandPasses[EMeshPass::SkyPass].Draw(RHICmdList, SkyPassInstanceCullingDrawParams);
 	}
 
 	// editor primitives
diff --git a/Engine/Source/Runtime/Renderer/Private/MobileDecalRendering.cpp b/Engine/Source/Runtime/Renderer/Private/MobileDecalRendering.cpp
--- a/Engine/Source/Runtime/Renderer/Private/MobileDecalRendering.cpp
+++ b/Engine/Source/Runtime/Renderer/Private/MobileDecalRendering.cpp
@@ -44,7 +44,7 @@ static bool DoesPlatformSupportDecals(EShaderPlatform ShaderPlatform)
 	return true;
 }
 
-void FMobileSceneRenderer::RenderDecals(FRHICommandList& RHICmdList, FViewInfo& View)
+void FMobileSceneRenderer::RenderDecals(FRHICommandList& RHICmdList, FViewInfo& View, const FInstanceCullingDrawParams* InstanceCullingDrawParams)
 {
 	if (!DoesPlatformSupportDecals(View.GetShaderPlatform()) || !ViewFamily.EngineShowFlags.Decals || View.bIsPlanarReflection)
 	{
@@ -68,19 +68,6 @@ void FMobileSceneRenderer::RenderDecals(FRHICommandList& RHICmdList, FViewInfo&
 	EMeshPass::Type DecalMeshPassType = DecalRendering::GetMeshPassType(RenderTargetMode);
 	if (View.ParallelMeshDrawCommandPasses[DecalMeshPassType].HasAnyDraw())
 	{
-		const FInstanceCullingDrawParams* InstanceCullingDrawParams = nullptr;
-		switch (DecalMeshPassType)
-		{
-		case EMeshPass::MeshDecal_SceneColor:
-			InstanceCullingDrawParams = &MeshDecalSceneColorInstanceCullingDrawParams;
-			break;
-		case EMeshPass::MeshDecal_SceneColorAndGBuffer:
-			InstanceCullingDrawParams = &MeshDecalSceneColorAndGBufferInstanceCullingDrawParams;
-			break;
-		default:
-			checkf(false, TEXT("Unexpected MeshDecal pass, please add corresponding InstanceCullingDrawParams!"))
-		};
-				
 		RHICmdList.SetViewport(View.ViewRect.Min.X, View.ViewRect.Min.Y, 0, View.ViewRect.Max.X, View.ViewRect.Max.Y, 1);
 		View.ParallelMeshDrawCommandPasses[DecalMeshPassType].Draw(RHICmdList, InstanceCullingDrawParams);
 	}
diff --git a/Engine/Source/Runtime/Renderer/Private/MobileShadingRenderer.cpp b/Engine/Source/Runtime/Renderer/Private/MobileShadingRenderer.cpp
--- a/Engine/Source/Runtime/Renderer/Private/MobileShadingRenderer.cpp
+++ b/Engine/Source/Runtime/Renderer/Private/MobileShadingRenderer.cpp
@@ -797,7 +797,14 @@ void FMobileSceneRenderer::RenderFullDepthPrepass(FRDGBuilder& GraphBuilder, TAr
 
 		const bool bRenderVelocityInDepthPrePass = (Scene->EarlyZPassMode == DDM_AllOpaqueNoVelocity);
 
-		auto* PassParameters = GraphBuilder.AllocParameters<FMobileRenderPassParameters>();
+		struct FFullDepthPrepassParameterCollection
+		{
+			FMobileRenderPassParameters PassParameters;
+			FInstanceCullingDrawParams VelocityInstanceCullingDrawParams;
+		};
+		auto ParameterCollection = GraphBuilder.AllocParameters<FFullDepthPrepassParameterCollection>();
+
+		auto PassParameters = &ParameterCollection->PassParameters;
 		if (bRenderVelocityInDepthPrePass)
 		{
 			PassParameters->RenderTargets[0] = FRenderTargetBinding(SceneTextures.Velocity, ERenderTargetLoadAction::EClear);
@@ -827,19 +834,20 @@ void FMobileSceneRenderer::RenderFullDepthPrepass(FRDGBuilder& GraphBuilder, TAr
 
 		if (bRenderVelocityInDepthPrePass)
 		{
-			View.ParallelMeshDrawCommandPasses[EMeshPass::Velocity].BuildRenderingCommands(GraphBuilder, Scene->GPUScene, VelocityInstanceCullingDrawParams);
+			View.ParallelMeshDrawCommandPasses[EMeshPass::Velocity].BuildRenderingCommands(GraphBuilder, Scene->GPUScene, ParameterCollection->VelocityInstanceCullingDrawParams);
 		}
 
 		GraphBuilder.AddPass(
 			RDG_EVENT_NAME("FullDepthPrepass"),
 			PassParameters,
 			ERDGPassFlags::Raster,
-			[this, PassParameters, &View, bRenderVelocityInDepthPrePass](FRHICommandList& RHICmdList)
+			[this, ParameterCollection, &View, bRenderVelocityInDepthPrePass](FRHICommandList& RHICmdList)
 			{
+				auto PassParameters = &ParameterCollection->PassParameters;
 				RenderPrePass(RHICmdList, View, &PassParameters->InstanceCullingDrawParams);
 				if (bRenderVelocityInDepthPrePass)
 				{
-					RenderVelocityPass(RHICmdList, View, &VelocityInstanceCullingDrawParams);
+					RenderVelocityPass(RHICmdList, View, &ParameterCollection->VelocityInstanceCullingDrawParams);
 				}
 			});
 	}
@@ -873,11 +881,11 @@ void FMobileSceneRenderer::RenderFullDepthPrepass(FRDGBuilder& GraphBuilder, TAr
 	FenceOcclusionTests(GraphBuilder);
 }
 
-void FMobileSceneRenderer::RenderMaskedPrePass(FRHICommandList& RHICmdList, const FViewInfo& View)
+void FMobileSceneRenderer::RenderMaskedPrePass(FRHICommandList& RHICmdList, const FViewInfo& View, const FInstanceCullingDrawParams* DepthPassInstanceCullingDrawParams)
 {
 	if (bIsMaskedOnlyDepthPrepassEnabled)
 	{
-		RenderPrePass(RHICmdList, View, &DepthPassInstanceCullingDrawParams);
+		RenderPrePass(RHICmdList, View, DepthPassInstanceCullingDrawParams);
 	}
 }
 
@@ -907,8 +915,15 @@ void FMobileSceneRenderer::RenderCustomRenderPassBasePass(FRDGBuilder& GraphBuil
 	{
 		FViewInfo& View = *ViewContext.ViewInfo;
 
+		struct FCustomPassParameterCollection
+		{
+			FMobileRenderPassParameters PassParameters;
+			FInstanceCullingDrawParams SkyPassInstanceCullingDrawParams;
+		};
+		auto ParameterCollection = GraphBuilder.AllocParameters<FCustomPassParameterCollection>();
+
 		EMobileSceneTextureSetupMode SetupMode = bIsFullDepthPrepassEnabled ? EMobileSceneTextureSetupMode::SceneDepth : EMobileSceneTextureSetupMode::None;
-		FMobileRenderPassParameters* PassParameters = GraphBuilder.AllocParameters<FMobileRenderPassParameters>();
+		auto PassParameters = &ParameterCollection->PassParameters;
 		PassParameters->View = View.GetShaderParameters();
 		PassParameters->MobileBasePass = CreateMobileBasePassUniformBuffer(GraphBuilder, View, EMobileBasePass::Opaque, SetupMode);
 		PassParameters->RenderTargets = BasePassRenderTargets;
@@ -928,6 +943,7 @@ void FMobileSceneRenderer::RenderCustomRenderPassBasePass(FRDGBuilder& GraphBuil
 		if (Scene->GPUScene.IsEnabled())
 		{
 			View.ParallelMeshDrawCommandPasses[EMeshPass::BasePass].BuildRenderingCommands(GraphBuilder, Scene->GPUScene, PassParameters->InstanceCullingDrawParams);
+			View.ParallelMeshDrawCommandPasses[EMeshPass::SkyPass].BuildRenderingCommands(GraphBuilder, Scene->GPUScene, ParameterCollection->SkyPassInstanceCullingDrawParams);
 
 			if (bIncludeTranslucent)
 			{
@@ -939,10 +955,11 @@ void FMobileSceneRenderer::RenderCustomRenderPassBasePass(FRDGBuilder& GraphBuil
 			RDG_EVENT_NAME("RenderMobileBasePass"),
 			PassParameters,
 			ERDGPassFlags::Raster,
-			[this, PassParameters, ViewContext, NumAdditionalSubpasses](FRHICommandList& RHICmdList)
+			[this, ParameterCollection, ViewContext, NumAdditionalSubpasses](FRHICommandList& RHICmdList)
 			{
 				FViewInfo& View = *ViewContext.ViewInfo;
-				RenderMobileBasePass(RHICmdList, View, &PassParameters->InstanceCullingDrawParams);
+				auto PassParameters = &ParameterCollection->PassParameters;
+				RenderMobileBasePass(RHICmdList, View, &PassParameters->InstanceCullingDrawParams, &ParameterCollection->SkyPassInstanceCullingDrawParams);
 
 				// TODO:  Should this render decals?  Deferred shading custom render passes do.
 
@@ -1523,23 +1540,6 @@ void FMobileSceneRenderer::Render(FRDGBuilder& GraphBuilder)
 	}
 }
 
-void FMobileSceneRenderer::BuildInstanceCullingDrawParams(FRDGBuilder& GraphBuilder, FViewInfo& View, FMobileRenderPassParameters* PassParameters)
-{
-	if (Scene->GPUScene.IsEnabled())
-	{
-		if (!bIsFullDepthPrepassEnabled)
-		{
-			View.ParallelMeshDrawCommandPasses[EMeshPass::DepthPass].BuildRenderingCommands(GraphBuilder, Scene->GPUScene, DepthPassInstanceCullingDrawParams);
-		}
-		View.ParallelMeshDrawCommandPasses[EMeshPass::BasePass].BuildRenderingCommands(GraphBuilder, Scene->GPUScene, PassParameters->InstanceCullingDrawParams);
-		View.ParallelMeshDrawCommandPasses[EMeshPass::SkyPass].BuildRenderingCommands(GraphBuilder, Scene->GPUScene, SkyPassInstanceCullingDrawParams);
-		View.ParallelMeshDrawCommandPasses[StandardTranslucencyMeshPass].BuildRenderingCommands(GraphBuilder, Scene->GPUScene, TranslucencyInstanceCullingDrawParams);
-		View.ParallelMeshDrawCommandPasses[EMeshPass::DebugViewMode].BuildRenderingCommands(GraphBuilder, Scene->GPUScene, DebugViewModeInstanceCullingDrawParams);
-		View.ParallelMeshDrawCommandPasses[EMeshPass::MeshDecal_SceneColor].BuildRenderingCommands(GraphBuilder, Scene->GPUScene, MeshDecalSceneColorInstanceCullingDrawParams);
-		View.ParallelMeshDrawCommandPasses[EMeshPass::MeshDecal_SceneColorAndGBuffer].BuildRenderingCommands(GraphBuilder, Scene->GPUScene, MeshDecalSceneColorAndGBufferInstanceCullingDrawParams);
-	}
-}
-
 FRenderTargetBindingSlots FMobileSceneRenderer::InitRenderTargetBindings_Forward(FRDGTextureRef ViewFamilyTexture, FSceneTextures& SceneTextures)
 {
 	FRDGTextureRef SceneColor = nullptr;
@@ -1653,8 +1653,6 @@ void FMobileSceneRenderer::RenderForward(FRDGBuilder& GraphBuilder, FRDGTextureR
 		PassParameters->LocalFogVolumeTileDataBuffer = View.LocalFogVolumeViewData.GPUTileDataBufferSRV;
 		PassParameters->HalfResLocalFogVolumeViewSRV = View.LocalFogVolumeViewData.HalfResLocalFogVolumeViewSRV;
 		PassParameters->HalfResLocalFogVolumeDepthSRV = View.LocalFogVolumeViewData.HalfResLocalFogVolumeDepthSRV;
-	
-		BuildInstanceCullingDrawParams(GraphBuilder, View, PassParameters);
 
 		// Split if we need to render translucency in a separate render pass
 		if (bRequiresMultiPass)
@@ -1676,6 +1674,31 @@ void FMobileSceneRenderer::RenderForward(FRDGBuilder& GraphBuilder, FRDGTextureR
 
 void FMobileSceneRenderer::RenderForwardSinglePass(FRDGBuilder& GraphBuilder, FMobileRenderPassParameters* PassParameters, FRenderViewContext& ViewContext, FSceneTextures& SceneTextures)
 {
+	struct FForwardSinglePassParameterCollection
+	{
+		FInstanceCullingDrawParams DepthPassInstanceCullingDrawParams;
+		FInstanceCullingDrawParams SkyPassInstanceCullingDrawParams;
+		FInstanceCullingDrawParams DebugViewModeInstanceCullingDrawParams;
+		FInstanceCullingDrawParams MeshDecalSceneColorInstanceCullingDrawParams;
+		FInstanceCullingDrawParams TranslucencyInstanceCullingDrawParams;
+	};
+	auto ParameterCollection = GraphBuilder.AllocParameters<FForwardSinglePassParameterCollection>();
+
+	FViewInfo& View = *ViewContext.ViewInfo;
+
+	if (Scene->GPUScene.IsEnabled())
+	{
+		if (!bIsFullDepthPrepassEnabled)
+		{
+			View.ParallelMeshDrawCommandPasses[EMeshPass::DepthPass].BuildRenderingCommands(GraphBuilder, Scene->GPUScene, ParameterCollection->DepthPassInstanceCullingDrawParams);
+		}
+		View.ParallelMeshDrawCommandPasses[EMeshPass::BasePass].BuildRenderingCommands(GraphBuilder, Scene->GPUScene, PassParameters->InstanceCullingDrawParams);
+		View.ParallelMeshDrawCommandPasses[EMeshPass::SkyPass].BuildRenderingCommands(GraphBuilder, Scene->GPUScene, ParameterCollection->SkyPassInstanceCullingDrawParams);
+		View.ParallelMeshDrawCommandPasses[EMeshPass::DebugViewMode].BuildRenderingCommands(GraphBuilder, Scene->GPUScene, ParameterCollection->DebugViewModeInstanceCullingDrawParams);
+		View.ParallelMeshDrawCommandPasses[EMeshPass::MeshDecal_SceneColor].BuildRenderingCommands(GraphBuilder, Scene->GPUScene, ParameterCollection->MeshDecalSceneColorInstanceCullingDrawParams);
+		View.ParallelMeshDrawCommandPasses[StandardTranslucencyMeshPass].BuildRenderingCommands(GraphBuilder, Scene->GPUScene, ParameterCollection->TranslucencyInstanceCullingDrawParams);
+	}
+
 	if (bTonemapSubpassInline)
 	{
 		// tonemapping LUT pass before we start main render pass. The texture is needed by the custom resolve pass which does tonemapping
@@ -1691,7 +1714,7 @@ void FMobileSceneRenderer::RenderForwardSinglePass(FRDGBuilder& GraphBuilder, FM
 		PassParameters,
 		// the second view pass should not be merged with the first view pass on mobile since the subpass would not work properly.
 		ERDGPassFlags::Raster | ERDGPassFlags::NeverMerge,
-		[this, PassParameters, ViewContext, bDoOcclusionQueries, &SceneTextures](FRHICommandList& RHICmdList)
+		[this, PassParameters, ParameterCollection, ViewContext, bDoOcclusionQueries, &SceneTextures](FRHICommandList& RHICmdList)
 	{
 		FViewInfo& View = *ViewContext.ViewInfo;
 			
@@ -1701,22 +1724,22 @@ void FMobileSceneRenderer::RenderForwardSinglePass(FRDGBuilder& GraphBuilder, FM
 		}
 
 		// Depth pre-pass
-		RenderMaskedPrePass(RHICmdList, View);
+		RenderMaskedPrePass(RHICmdList, View, &ParameterCollection->DepthPassInstanceCullingDrawParams);
 		// Opaque and masked
-		RenderMobileBasePass(RHICmdList, View, &PassParameters->InstanceCullingDrawParams);
-		RenderMobileDebugView(RHICmdList, View);
+		RenderMobileBasePass(RHICmdList, View, &PassParameters->InstanceCullingDrawParams, &ParameterCollection->SkyPassInstanceCullingDrawParams);
+		RenderMobileDebugView(RHICmdList, View, &ParameterCollection->DebugViewModeInstanceCullingDrawParams);
 
 		PostRenderBasePass(RHICmdList, View);
 		// scene depth is read only and can be fetched
 		RHICmdList.NextSubpass();
-		RenderDecals(RHICmdList, View);
+		RenderDecals(RHICmdList, View, &ParameterCollection->MeshDecalSceneColorInstanceCullingDrawParams);
 		RenderModulatedShadowProjections(RHICmdList, ViewContext.ViewIndex, View);
 		if (GMaxRHIShaderPlatform != SP_METAL_SIM)
 		{
 			RenderFog(RHICmdList, View);
 		}
 		// Draw translucency.
-		RenderTranslucency(RHICmdList, View, Views, StandardTranslucencyPass, StandardTranslucencyMeshPass, &TranslucencyInstanceCullingDrawParams);
+		RenderTranslucency(RHICmdList, View, Views, StandardTranslucencyPass, StandardTranslucencyMeshPass, &ParameterCollection->TranslucencyInstanceCullingDrawParams);
 		
 #if UE_ENABLE_DEBUG_DRAWING
 		if ((!IsMobileHDR() || bTonemapSubpass) && FSceneRenderer::ShouldCompositeDebugPrimitivesInPostProcess(View))
@@ -1756,11 +1779,32 @@ void FMobileSceneRenderer::RenderForwardSinglePass(FRDGBuilder& GraphBuilder, FM
 
 void FMobileSceneRenderer::RenderForwardMultiPass(FRDGBuilder& GraphBuilder, FMobileRenderPassParameters* PassParameters, FRenderViewContext& ViewContext, FSceneTextures& SceneTextures)
 {
+	struct FForwardFirstPassParameterCollection
+	{
+		FInstanceCullingDrawParams DepthPassInstanceCullingDrawParams;
+		FInstanceCullingDrawParams SkyPassInstanceCullingDrawParams;
+		FInstanceCullingDrawParams DebugViewModeInstanceCullingDrawParams;
+	};
+	auto ParameterCollection = GraphBuilder.AllocParameters<FForwardFirstPassParameterCollection>();
+
+	FViewInfo& View = *ViewContext.ViewInfo;
+
+	if (Scene->GPUScene.IsEnabled())
+	{
+		if (!bIsFullDepthPrepassEnabled)
+		{
+			View.ParallelMeshDrawCommandPasses[EMeshPass::DepthPass].BuildRenderingCommands(GraphBuilder, Scene->GPUScene, ParameterCollection->DepthPassInstanceCullingDrawParams);
+		}
+		View.ParallelMeshDrawCommandPasses[EMeshPass::BasePass].BuildRenderingCommands(GraphBuilder, Scene->GPUScene, PassParameters->InstanceCullingDrawParams);
+		View.ParallelMeshDrawCommandPasses[EMeshPass::SkyPass].BuildRenderingCommands(GraphBuilder, Scene->GPUScene, ParameterCollection->SkyPassInstanceCullingDrawParams);
+		View.ParallelMeshDrawCommandPasses[EMeshPass::DebugViewMode].BuildRenderingCommands(GraphBuilder, Scene->GPUScene, ParameterCollection->DebugViewModeInstanceCullingDrawParams);
+	}
+
 	GraphBuilder.AddPass(
 		RDG_EVENT_NAME("SceneColorRendering"),
 		PassParameters,
 		ERDGPassFlags::Raster,
-		[this, PassParameters, ViewContext, &SceneTextures](FRHICommandList& RHICmdList)
+		[this, PassParameters, ParameterCollection, ViewContext, &SceneTextures](FRHICommandList& RHICmdList)
 	{
 		FViewInfo& View = *ViewContext.ViewInfo;
 			
@@ -1770,16 +1814,14 @@ void FMobileSceneRenderer::RenderForwardMultiPass(FRDGBuilder& GraphBuilder, FMo
 		}
 
 		// Depth pre-pass
-		RenderMaskedPrePass(RHICmdList, View);
+		RenderMaskedPrePass(RHICmdList, View, &ParameterCollection->DepthPassInstanceCullingDrawParams);
 		// Opaque and masked
-		RenderMobileBasePass(RHICmdList, View, &PassParameters->InstanceCullingDrawParams);
-		RenderMobileDebugView(RHICmdList, View);
+		RenderMobileBasePass(RHICmdList, View, &PassParameters->InstanceCullingDrawParams, &ParameterCollection->SkyPassInstanceCullingDrawParams);
+		RenderMobileDebugView(RHICmdList, View, &ParameterCollection->DebugViewModeInstanceCullingDrawParams);
 
 		PostRenderBasePass(RHICmdList, View);
 	});
 
-	FViewInfo& View = *ViewContext.ViewInfo;
-
 	// resolve MSAA depth
 	if (!bIsFullDepthPrepassEnabled)
 	{
@@ -1802,7 +1844,16 @@ void FMobileSceneRenderer::RenderForwardMultiPass(FRDGBuilder& GraphBuilder, FMo
 	}
 
 	EMobileSceneTextureSetupMode SetupMode = EMobileSceneTextureSetupMode::SceneDepth | EMobileSceneTextureSetupMode::SceneDepthAux | EMobileSceneTextureSetupMode::CustomDepth;
-	FMobileRenderPassParameters* SecondPassParameters = GraphBuilder.AllocParameters<FMobileRenderPassParameters>();
+
+	struct FForwardSecondPassParameterCollection
+	{
+		FMobileRenderPassParameters PassParameters;
+		FInstanceCullingDrawParams MeshDecalSceneColorInstanceCullingDrawParams;
+		FInstanceCullingDrawParams TranslucencyInstanceCullingDrawParams;
+	};
+	auto SecondParameterCollection = GraphBuilder.AllocParameters<FForwardSecondPassParameterCollection>();
+
+	auto SecondPassParameters = &SecondParameterCollection->PassParameters;
 	*SecondPassParameters = *PassParameters;
 	SecondPassParameters->MobileBasePass = CreateMobileBasePassUniformBuffer(GraphBuilder, View, EMobileBasePass::Translucent, SetupMode);
 	SecondPassParameters->ReflectionCapture = View.MobileReflectionCaptureUniformBuffer;
@@ -1819,16 +1870,17 @@ void FMobileSceneRenderer::RenderForwardMultiPass(FRDGBuilder& GraphBuilder, FMo
 		RDG_EVENT_NAME("DecalsAndTranslucency"),
 		SecondPassParameters,
 		ERDGPassFlags::Raster,
-		[this, SecondPassParameters, ViewContext, bDoOcclusionQueries, &SceneTextures](FRHICommandList& RHICmdList)
+		[this, SecondParameterCollection, ViewContext, bDoOcclusionQueries, &SceneTextures](FRHICommandList& RHICmdList)
 	{
 		FViewInfo& View = *ViewContext.ViewInfo;
+		auto SecondPassParameters = &SecondParameterCollection->PassParameters;
 			
 		// scene depth is read only and can be fetched
-		RenderDecals(RHICmdList, View);
+		RenderDecals(RHICmdList, View, &SecondParameterCollection->MeshDecalSceneColorInstanceCullingDrawParams);
 		RenderModulatedShadowProjections(RHICmdList, ViewContext.ViewIndex, View);
 		RenderFog(RHICmdList, View);
 		// Draw translucency.
-		RenderTranslucency(RHICmdList, View, Views, StandardTranslucencyPass, StandardTranslucencyMeshPass, &TranslucencyInstanceCullingDrawParams);
+		RenderTranslucency(RHICmdList, View, Views, StandardTranslucencyPass, StandardTranslucencyMeshPass, &SecondParameterCollection->TranslucencyInstanceCullingDrawParams);
 
 		if (bDoOcclusionQueries)
 		{
@@ -2012,7 +2064,19 @@ void FMobileSceneRenderer::RenderDeferredSinglePass(FRDGBuilder& GraphBuilder, F
 		UpdateDirectionalLightUniformBuffers(GraphBuilder, View);
 
 		EMobileSceneTextureSetupMode SetupMode = (bIsFullDepthPrepassEnabled ? EMobileSceneTextureSetupMode::SceneDepth : EMobileSceneTextureSetupMode::None) | EMobileSceneTextureSetupMode::CustomDepth;
-		auto* PassParameters = GraphBuilder.AllocParameters<FMobileRenderPassParameters>();
+
+		struct FDeferredSinglePassParameterCollection
+		{
+			FMobileRenderPassParameters PassParameters;
+			FInstanceCullingDrawParams DepthPassInstanceCullingDrawParams;
+			FInstanceCullingDrawParams SkyPassInstanceCullingDrawParams;
+			FInstanceCullingDrawParams DebugViewModeInstanceCullingDrawParams;
+			FInstanceCullingDrawParams MeshDecalSceneColorAndGBufferInstanceCullingDrawParams;
+			FInstanceCullingDrawParams TranslucencyInstanceCullingDrawParams;
+		};
+		auto ParameterCollection = GraphBuilder.AllocParameters<FDeferredSinglePassParameterCollection>();
+
+		auto PassParameters = &ParameterCollection->PassParameters;
 		PassParameters->View = View.GetShaderParameters();
 		PassParameters->MobileBasePass = CreateMobileBasePassUniformBuffer(GraphBuilder, View, EMobileBasePass::Opaque, SetupMode);
 		PassParameters->ReflectionCapture = View.MobileReflectionCaptureUniformBuffer;
@@ -2028,27 +2092,39 @@ void FMobileSceneRenderer::RenderDeferredSinglePass(FRDGBuilder& GraphBuilder, F
 		const bool bDoOcclusionQueires = (!bIsFullDepthPrepassEnabled && ViewContext.bIsLastView && DoOcclusionQueries());
 		PassParameters->RenderTargets.NumOcclusionQueries = bDoOcclusionQueires ? ComputeNumOcclusionQueriesToBatch() : 0u;
 
-		BuildInstanceCullingDrawParams(GraphBuilder, View, PassParameters);
+		if (Scene->GPUScene.IsEnabled())
+		{
+			if (!bIsFullDepthPrepassEnabled)
+			{
+				View.ParallelMeshDrawCommandPasses[EMeshPass::DepthPass].BuildRenderingCommands(GraphBuilder, Scene->GPUScene, ParameterCollection->DepthPassInstanceCullingDrawParams);
+			}
+			View.ParallelMeshDrawCommandPasses[EMeshPass::BasePass].BuildRenderingCommands(GraphBuilder, Scene->GPUScene, PassParameters->InstanceCullingDrawParams);
+			View.ParallelMeshDrawCommandPasses[EMeshPass::SkyPass].BuildRenderingCommands(GraphBuilder, Scene->GPUScene, ParameterCollection->SkyPassInstanceCullingDrawParams);
+			View.ParallelMeshDrawCommandPasses[EMeshPass::DebugViewMode].BuildRenderingCommands(GraphBuilder, Scene->GPUScene, ParameterCollection->DebugViewModeInstanceCullingDrawParams);
+			View.ParallelMeshDrawCommandPasses[EMeshPass::MeshDecal_SceneColorAndGBuffer].BuildRenderingCommands(GraphBuilder, Scene->GPUScene, ParameterCollection->MeshDecalSceneColorAndGBufferInstanceCullingDrawParams);
+			View.ParallelMeshDrawCommandPasses[StandardTranslucencyMeshPass].BuildRenderingCommands(GraphBuilder, Scene->GPUScene, ParameterCollection->TranslucencyInstanceCullingDrawParams);
+		}
 
 		GraphBuilder.AddPass(
 			RDG_EVENT_NAME("SceneColorRendering"),
 			PassParameters,
 			// the second view pass should not be merged with the first view pass on mobile since the subpass would not work properly.
 			ERDGPassFlags::Raster | ERDGPassFlags::NeverMerge,
-			[this, PassParameters, ViewContext, bDoOcclusionQueires, MobileSSRQuality, &SortedLightSet, bUsingPixelLocalStorage](FRHICommandList& RHICmdList)
+			[this, ParameterCollection, ViewContext, bDoOcclusionQueires, MobileSSRQuality, &SortedLightSet, bUsingPixelLocalStorage](FRHICommandList& RHICmdList)
 			{
 				FViewInfo& View = *ViewContext.ViewInfo;
+				auto PassParameters = &ParameterCollection->PassParameters;
 
 				// Depth pre-pass
-				RenderMaskedPrePass(RHICmdList, View);
+				RenderMaskedPrePass(RHICmdList, View, &ParameterCollection->DepthPassInstanceCullingDrawParams);
 				// Opaque and masked
-				RenderMobileBasePass(RHICmdList, View, &PassParameters->InstanceCullingDrawParams);
-				RenderMobileDebugView(RHICmdList, View);
+				RenderMobileBasePass(RHICmdList, View, &PassParameters->InstanceCullingDrawParams, &ParameterCollection->SkyPassInstanceCullingDrawParams);
+				RenderMobileDebugView(RHICmdList, View, &ParameterCollection->DebugViewModeInstanceCullingDrawParams);
 
 				PostRenderBasePass(RHICmdList, View);
 				// SceneColor + GBuffer write, SceneDepth is read only
 				RHICmdList.NextSubpass();
-				RenderDecals(RHICmdList, View);
+				RenderDecals(RHICmdList, View, &ParameterCollection->MeshDecalSceneColorAndGBufferInstanceCullingDrawParams);
 				// SceneColor write, SceneDepth is read only
 				RHICmdList.NextSubpass();
 				MobileDeferredShadingPass(RHICmdList, ViewContext.ViewIndex, Views.Num(), View, *Scene, SortedLightSet, VisibleLightInfos, MobileSSRQuality);
@@ -2059,7 +2135,7 @@ void FMobileSceneRenderer::RenderDeferredSinglePass(FRDGBuilder& GraphBuilder, F
 				}
 				RenderFog(RHICmdList, View);
 				// Draw translucency.
-				RenderTranslucency(RHICmdList, View, Views, StandardTranslucencyPass, StandardTranslucencyMeshPass, &TranslucencyInstanceCullingDrawParams);
+				RenderTranslucency(RHICmdList, View, Views, StandardTranslucencyPass, StandardTranslucencyMeshPass, &ParameterCollection->TranslucencyInstanceCullingDrawParams);
 
 				if (bDoOcclusionQueires)
 				{
@@ -2088,7 +2164,17 @@ void FMobileSceneRenderer::RenderDeferredMultiPass(FRDGBuilder& GraphBuilder, FS
 		
 		View.BeginRenderView();
 
-		auto* PassParameters = GraphBuilder.AllocParameters<FMobileRenderPassParameters>();
+		struct FDeferredMultiPassParameterCollection
+		{
+			FMobileRenderPassParameters PassParameters;
+			FInstanceCullingDrawParams DepthPassInstanceCullingDrawParams;
+			FInstanceCullingDrawParams SkyPassInstanceCullingDrawParams;
+			FInstanceCullingDrawParams DebugViewModeInstanceCullingDrawParams;
+			FInstanceCullingDrawParams MeshDecalSceneColorAndGBufferInstanceCullingDrawParams;
+		};
+		auto ParameterCollection = GraphBuilder.AllocParameters<FDeferredMultiPassParameterCollection>();
+
+		auto PassParameters = &ParameterCollection->PassParameters;
 		PassParameters->View = View.GetShaderParameters();
 		EMobileSceneTextureSetupMode SetupMode = bIsFullDepthPrepassEnabled ? EMobileSceneTextureSetupMode::SceneDepth : EMobileSceneTextureSetupMode::None;
 		PassParameters->MobileBasePass = CreateMobileBasePassUniformBuffer(GraphBuilder, View, EMobileBasePass::Opaque, SetupMode);
@@ -2108,21 +2194,35 @@ void FMobileSceneRenderer::RenderDeferredMultiPass(FRDGBuilder& GraphBuilder, FS
 		const bool bDoOcclusionQueires = (!bIsFullDepthPrepassEnabled && ViewContext.bIsLastView && DoOcclusionQueries());
 		PassParameters->RenderTargets.NumOcclusionQueries = bDoOcclusionQueires ? ComputeNumOcclusionQueriesToBatch() : 0u;
 
-		BuildInstanceCullingDrawParams(GraphBuilder, View, PassParameters);
+		if (Scene->GPUScene.IsEnabled())
+		{
+			if (!bIsFullDepthPrepassEnabled)
+			{
+				View.ParallelMeshDrawCommandPasses[EMeshPass::DepthPass].BuildRenderingCommands(GraphBuilder, Scene->GPUScene, ParameterCollection->DepthPassInstanceCullingDrawParams);
+			}
+			View.ParallelMeshDrawCommandPasses[EMeshPass::BasePass].BuildRenderingCommands(GraphBuilder, Scene->GPUScene, PassParameters->InstanceCullingDrawParams);
+			View.ParallelMeshDrawCommandPasses[EMeshPass::SkyPass].BuildRenderingCommands(GraphBuilder, Scene->GPUScene, ParameterCollection->SkyPassInstanceCullingDrawParams);
+			View.ParallelMeshDrawCommandPasses[EMeshPass::DebugViewMode].BuildRenderingCommands(GraphBuilder, Scene->GPUScene, ParameterCollection->DebugViewModeInstanceCullingDrawParams);
+			if (bIsFullDepthPrepassEnabled)
+			{
+				View.ParallelMeshDrawCommandPasses[EMeshPass::MeshDecal_SceneColorAndGBuffer].BuildRenderingCommands(GraphBuilder, Scene->GPUScene, ParameterCollection->MeshDecalSceneColorAndGBufferInstanceCullingDrawParams);
+			}
+		}
 
 		GraphBuilder.AddPass(
 			RDG_EVENT_NAME("BasePass"),
 			PassParameters,
 			ERDGPassFlags::Raster,
-			[this, PassParameters, ViewContext, &SceneTextures, bDoOcclusionQueires](FRHICommandList& RHICmdList)
+			[this, ParameterCollection, ViewContext, &SceneTextures, bDoOcclusionQueires](FRHICommandList& RHICmdList)
 			{
+				auto PassParameters = &ParameterCollection->PassParameters;
 				FViewInfo& View = *ViewContext.ViewInfo;
 
 				// Depth pre-pass
-				RenderMaskedPrePass(RHICmdList, View);
+				RenderMaskedPrePass(RHICmdList, View, &ParameterCollection->DepthPassInstanceCullingDrawParams);
 				// Opaque and masked
-				RenderMobileBasePass(RHICmdList, View, &PassParameters->InstanceCullingDrawParams);
-				RenderMobileDebugView(RHICmdList, View);
+				RenderMobileBasePass(RHICmdList, View, &PassParameters->InstanceCullingDrawParams, &ParameterCollection->SkyPassInstanceCullingDrawParams);
+				RenderMobileDebugView(RHICmdList, View, &ParameterCollection->DebugViewModeInstanceCullingDrawParams);
 
 				PostRenderBasePass(RHICmdList, View);
 
@@ -2134,7 +2234,7 @@ void FMobileSceneRenderer::RenderDeferredMultiPass(FRDGBuilder& GraphBuilder, FS
 
 				if (bIsFullDepthPrepassEnabled)
 				{
-					RenderDecals(RHICmdList, View);
+					RenderDecals(RHICmdList, View, &ParameterCollection->MeshDecalSceneColorAndGBufferInstanceCullingDrawParams);
 				}
 			});
 	}
@@ -2163,19 +2263,31 @@ void FMobileSceneRenderer::RenderDeferredMultiPass(FRDGBuilder& GraphBuilder, FS
 
 			View.BeginRenderView();
 
-			auto* PassParameters = GraphBuilder.AllocParameters<FMobileRenderPassParameters>();
+			struct FDeferredMultiPassDecalParameterCollection
+			{
+				FMobileRenderPassParameters PassParameters;
+				FInstanceCullingDrawParams MeshDecalSceneColorAndGBufferInstanceCullingDrawParams;
+			};
+			auto ParameterCollection = GraphBuilder.AllocParameters<FDeferredMultiPassDecalParameterCollection>();
+
+			auto PassParameters = &ParameterCollection->PassParameters;
 			PassParameters->View = View.GetShaderParameters();
 			PassParameters->MobileBasePass = CreateMobileBasePassUniformBuffer(GraphBuilder, View, EMobileBasePass::Opaque, EMobileSceneTextureSetupMode::SceneDepth);
 			PassParameters->RenderTargets = BasePassRenderTargets;
-			PassParameters->InstanceCullingDrawParams = MeshDecalSceneColorInstanceCullingDrawParams;
+
+			if (Scene->GPUScene.IsEnabled())
+			{
+				View.ParallelMeshDrawCommandPasses[EMeshPass::MeshDecal_SceneColor].BuildRenderingCommands(GraphBuilder, Scene->GPUScene, PassParameters->InstanceCullingDrawParams);
+				View.ParallelMeshDrawCommandPasses[EMeshPass::MeshDecal_SceneColorAndGBuffer].BuildRenderingCommands(GraphBuilder, Scene->GPUScene, ParameterCollection->MeshDecalSceneColorAndGBufferInstanceCullingDrawParams);
+			}
 
 			GraphBuilder.AddPass(
 				RDG_EVENT_NAME("Decals"),
 				PassParameters,
 				ERDGPassFlags::Raster,
-				[this, PassParameters, &View](FRHICommandList& RHICmdList)
+				[this, ParameterCollection, &View](FRHICommandList& RHICmdList)
 				{
-					RenderDecals(RHICmdList, View);
+					RenderDecals(RHICmdList, View, &ParameterCollection->MeshDecalSceneColorAndGBufferInstanceCullingDrawParams);
 				});
 		}
 	}
@@ -2209,7 +2321,12 @@ void FMobileSceneRenderer::RenderDeferredMultiPass(FRDGBuilder& GraphBuilder, FS
 		// Only SceneColor and Depth
 		PassParameters->RenderTargets[0] = BasePassRenderTargets[0];
 		PassParameters->RenderTargets.DepthStencil = BasePassRenderTargets.DepthStencil;
-		PassParameters->InstanceCullingDrawParams = TranslucencyInstanceCullingDrawParams;
+
+		if (Scene->GPUScene.IsEnabled())
+		{
+			View.ParallelMeshDrawCommandPasses[StandardTranslucencyMeshPass].BuildRenderingCommands(GraphBuilder, Scene->GPUScene, PassParameters->InstanceCullingDrawParams);
+		}
+
 		const EMobileSSRQuality MobileSSRQuality = ActiveMobileSSRQuality(View, bShouldRenderVelocities);
 
 		GraphBuilder.AddPass(
@@ -2224,7 +2341,7 @@ void FMobileSceneRenderer::RenderDeferredMultiPass(FRDGBuilder& GraphBuilder, FS
 				RenderFog(RHICmdList, View);
 
 				// Draw translucency.
-				RenderTranslucency(RHICmdList, View, Views, StandardTranslucencyPass, StandardTranslucencyMeshPass, &TranslucencyInstanceCullingDrawParams);
+				RenderTranslucency(RHICmdList, View, Views, StandardTranslucencyPass, StandardTranslucencyMeshPass, &PassParameters->InstanceCullingDrawParams);
 			});
 	}
 }
@@ -2242,7 +2359,7 @@ void FMobileSceneRenderer::PostRenderBasePass(FRHICommandList& RHICmdList, FView
 	}
 }
 
-void FMobileSceneRenderer::RenderMobileDebugView(FRHICommandList& RHICmdList, const FViewInfo& View)
+void FMobileSceneRenderer::RenderMobileDebugView(FRHICommandList& RHICmdList, const FViewInfo& View, const FInstanceCullingDrawParams* DebugViewModeInstanceCullingDrawParams)
 {
 #if WITH_DEBUG_VIEW_MODES
 	if (ViewFamily.UseDebugViewPS())
@@ -2256,7 +2373,7 @@ void FMobileSceneRenderer::RenderMobileDebugView(FRHICommandList& RHICmdList, co
 		DrawClearQuad(RHICmdList, FLinearColor::Black);
 
 		RHICmdList.SetViewport(View.ViewRect.Min.X, View.ViewRect.Min.Y, 0, View.ViewRect.Max.X, View.ViewRect.Max.Y, 1);
-		View.ParallelMeshDrawCommandPasses[EMeshPass::DebugViewMode].Draw(RHICmdList, &DebugViewModeInstanceCullingDrawParams);
+		View.ParallelMeshDrawCommandPasses[EMeshPass::DebugViewMode].Draw(RHICmdList, DebugViewModeInstanceCullingDrawParams);
 	}
 #endif // WITH_DEBUG_VIEW_MODES
 }
diff --git a/Engine/Source/Runtime/Renderer/Private/SceneRendering.h b/Engine/Source/Runtime/Renderer/Private/SceneRendering.h
--- a/Engine/Source/Runtime/Renderer/Private/SceneRendering.h
+++ b/Engine/Source/Runtime/Renderer/Private/SceneRendering.h
@@ -2718,7 +2718,7 @@ protected:
 		FInitViewTaskDatas& TaskDatas);
 
 	void RenderPrePass(FRHICommandList& RHICmdList, const FViewInfo& View, const FInstanceCullingDrawParams* InstanceCullingDrawParams);
-	void RenderMaskedPrePass(FRHICommandList& RHICmdList, const FViewInfo& View);
+	void RenderMaskedPrePass(FRHICommandList& RHICmdList, const FViewInfo& View, const FInstanceCullingDrawParams* DepthPassInstanceCullingDrawParams);
 	void RenderFullDepthPrepass(FRDGBuilder& GraphBuilder, TArrayView<FViewInfo> InViews, FSceneTextures& SceneTextures, bool bIsSceneCaptureRenderPass=false);
 
 	void RenderVelocityPass(FRHICommandList& RHICmdList, const FViewInfo& View, const FInstanceCullingDrawParams* InstanceCullingDrawParams);
@@ -2728,7 +2728,7 @@ protected:
 	void RenderCustomRenderPassBasePass(FRDGBuilder& GraphBuilder, TArrayView<FViewInfo> InViews, FRDGTextureRef ViewFamilyTexture, FSceneTextures& SceneTextures, bool bIncludeTranslucent);
 
 	/** Renders the opaque base pass for mobile. */
-	void RenderMobileBasePass(FRHICommandList& RHICmdList, const FViewInfo& View, const FInstanceCullingDrawParams* InstanceCullingDrawParams);
+	void RenderMobileBasePass(FRHICommandList& RHICmdList, const FViewInfo& View, const FInstanceCullingDrawParams* InstanceCullingDrawParams, const FInstanceCullingDrawParams* SkyPassInstanceCullingDrawParams);
 
 	void PostRenderBasePass(FRHICommandList& RHICmdList, FViewInfo& View);
 
@@ -2740,7 +2740,7 @@ protected:
 #endif
 
 	/** Renders the debug view pass for mobile. */
-	void RenderMobileDebugView(FRHICommandList& RHICmdList, const FViewInfo& View);
+	void RenderMobileDebugView(FRHICommandList& RHICmdList, const FViewInfo& View, const FInstanceCullingDrawParams* DebugViewModeInstanceCullingDrawParams);
 
 	/** Render modulated shadow projections in to the scene, loops over any unrendered shadows until all are processed.*/
 	void RenderModulatedShadowProjections(FRHICommandList& RHICmdList, int32 ViewIndex, const FViewInfo& View);
@@ -2760,7 +2760,7 @@ protected:
 
 
 	/** Renders decals. */
-	void RenderDecals(FRHICommandList& RHICmdList, FViewInfo& View);
+	void RenderDecals(FRHICommandList& RHICmdList, FViewInfo& View, const FInstanceCullingDrawParams* InstanceCullingDrawParams);
 	void RenderDBuffer(FRDGBuilder& GraphBuilder, FSceneTextures& SceneTextures, FDBufferTextures& DBufferTextures, FInstanceCullingManager& InstanceCullingManager);
 
 	/** Renders the atmospheric and height fog */
@@ -2777,8 +2777,6 @@ protected:
 	void UpdateDirectionalLightUniformBuffers(FRDGBuilder& GraphBuilder, const FViewInfo& View);
 	void UpdateSkyReflectionUniformBuffer(FRHICommandListBase& RHICmdList);
 	
-	void BuildInstanceCullingDrawParams(FRDGBuilder& GraphBuilder, FViewInfo& View, class FMobileRenderPassParameters* PassParameters);
-
 	void RenderForward(FRDGBuilder& GraphBuilder, FRDGTextureRef ViewFamilyTexture, FSceneTextures& SceneTextures, FDBufferTextures& DBufferTextures);
 	void RenderForwardSinglePass(FRDGBuilder& GraphBuilder, class FMobileRenderPassParameters* PassParameters, struct FRenderViewContext& ViewContext, FSceneTextures& SceneTextures);
 	void RenderForwardMultiPass(FRDGBuilder& GraphBuilder, class FMobileRenderPassParameters* PassParameters, struct FRenderViewContext& ViewContext, FSceneTextures& SceneTextures);
@@ -2825,16 +2823,6 @@ private:
 	ETranslucencyPass::Type StandardTranslucencyPass;
 	EMeshPass::Type StandardTranslucencyMeshPass;
 
-	// All mesh passes that can be fused into single render-pass
-	// Base mesh pass gets its culling parameters from a render-pass struct
-	FInstanceCullingDrawParams DepthPassInstanceCullingDrawParams;
-	FInstanceCullingDrawParams SkyPassInstanceCullingDrawParams;
-	FInstanceCullingDrawParams DebugViewModeInstanceCullingDrawParams;
-	FInstanceCullingDrawParams TranslucencyInstanceCullingDrawParams;
-	FInstanceCullingDrawParams MeshDecalSceneColorInstanceCullingDrawParams;
-	FInstanceCullingDrawParams MeshDecalSceneColorAndGBufferInstanceCullingDrawParams;
-	FInstanceCullingDrawParams VelocityInstanceCullingDrawParams;
-
 	const FViewInfo* CachedView = nullptr;
 };
 
@@ -3009,4 +2997,4 @@ bool IsPrimitiveAlphaHoldoutEnabledForAnyView(TArrayView<const FViewInfo> Views)
 bool SceneCaptureRequiresAlphaChannel(const FSceneView& View);
 
 /** Checks whether the material and scene proxy combination will modify the mesh position. Useful for determining whether the material can be substituted with a default material for depth rendering etc. */
-bool DoMaterialAndPrimitiveModifyMeshPosition(const FMaterial& Material, const FPrimitiveSceneProxy* PrimitiveSceneProxy);
\ No newline at end of file
+bool DoMaterialAndPrimitiveModifyMeshPosition(const FMaterial& Material, const FPrimitiveSceneProxy* PrimitiveSceneProxy);
