Added explicit chunk bounds to instance culling hierarchy and used those to make it possible to store and update only the bounds for the dynamic contents of a cell. This improves performance in CitySample by some 100us (which had regressed when we swtiched dynamic geometry to be "cullable").
 - Packed the bounds to 8-bit per component.
 - Added computation of max instance radius (for use in early out in size based culling).
 - Made it so that r.SceneCulling.ExplicitCellBounds can be set to 2 to indicate full rebuild every frame (for testing).
 - Added  r.SceneCulling.ExplicitChunkBounds to be able to turn off explicit chunk bounds (also disables the optimization they bring).
 - Changed the scene culling debug rendering to use cell coloring instead of block (being more of an implementation detail, really).

diff --git a/Engine/Shaders/Private/BitPacking.ush b/Engine/Shaders/Private/BitPacking.ush
--- a/Engine/Shaders/Private/BitPacking.ush
+++ b/Engine/Shaders/Private/BitPacking.ush
@@ -2,14 +2,14 @@
 
 #pragma once
 
-uint3 UnpackToUint3(uint Value, int3 NumComponentBits)
+uint3 UnpackToUint3(uint Value, uint3 NumComponentBits)
 {
 	return uint3(BitFieldExtractU32(Value, NumComponentBits.x, 0),
 				 BitFieldExtractU32(Value, NumComponentBits.y, NumComponentBits.x),
 				 BitFieldExtractU32(Value, NumComponentBits.z, NumComponentBits.x + NumComponentBits.y));
 }
 
-uint4 UnpackToUint4(uint Value, int4 NumComponentBits)
+uint4 UnpackToUint4(uint Value, uint4 NumComponentBits)
 {
 	return uint4(BitFieldExtractU32(Value, NumComponentBits.x, 0),
 				 BitFieldExtractU32(Value, NumComponentBits.y, NumComponentBits.x),
@@ -17,6 +17,19 @@ uint4 UnpackToUint4(uint Value, int4 NumComponentBits)
 				 BitFieldExtractU32(Value, NumComponentBits.w, NumComponentBits.x + NumComponentBits.y + NumComponentBits.z));
 }
 
+float3 UnpackToFloat3(uint Value, uint3 NumComponentBits)
+{
+	checkSlow(NumComponentBits.x + NumComponentBits.y + NumComponentBits.z <= 32u);
+	return float3(BitFieldExtractFloat(Value, NumComponentBits.x, 0),
+				 BitFieldExtractFloat(Value, NumComponentBits.y, NumComponentBits.x),
+				 BitFieldExtractFloat(Value, NumComponentBits.z, NumComponentBits.x + NumComponentBits.y));
+}
+
+float3 UnpackToFloat3(uint Value, uint BitsPerElement)
+{
+	return UnpackToFloat3(Value, uint3(BitsPerElement, BitsPerElement, BitsPerElement));
+}
+
 uint FloatToUIntScaled(float Value, float Scale)
 {
 	return (uint)floor(Value * Scale + 0.5f);
@@ -42,6 +55,69 @@ float4 Unpack_R10G10B10A2_UNORM_To_Float4(uint Packed)
 	return Unpacked;
 }
 
+/**
+ * Pack normalized, i.e., range [0..1], float value to uint at a certain number of bits, rounding down.
+ */
+uint PackNormToUintFloor(float FloatNorm, uint BitsPerElement)
+{
+	return uint(floor(FloatNorm * float((1u << BitsPerElement) - 1u)));
+}
+
+/**
+ * Pack normalized, i.e., range [0..1], float value to uint at a certain number of bits, rounding up.
+ */
+uint PackNormToUintCeil(float FloatNorm, uint BitsPerElement)
+{
+	return uint(ceil(FloatNorm * float((1u << BitsPerElement) - 1u)));
+}
+
+/**
+ * Pack uint3 into uint at a certain number of bits.
+ */
+uint PackToUint(uint3 Packed3, uint3 NumComponentBits)
+{
+	checkSlow(NumComponentBits.x + NumComponentBits.y + NumComponentBits.z <= 32u);
+	return Packed3.x | (Packed3.y << NumComponentBits.x) | (Packed3.z << (NumComponentBits.x + NumComponentBits.y));
+}
+
+/**
+ * Pack normalized, i.e., range [0..1], float3 value to uint at a certain number of bits per component, rounding down.
+ */
+uint PackNormToUintFloor(float3 FloatNorm, uint3 NumComponentBits)
+{
+	return PackToUint(
+		uint3(
+			PackNormToUintFloor(FloatNorm.x, NumComponentBits.x),
+			PackNormToUintFloor(FloatNorm.y, NumComponentBits.y),
+			PackNormToUintFloor(FloatNorm.z, NumComponentBits.z)
+		),
+		NumComponentBits);
+}
+
+uint PackNormToUintFloor(float3 FloatNorm, uint BitsPerElement)
+{
+	return PackNormToUintFloor(FloatNorm, uint3(BitsPerElement, BitsPerElement, BitsPerElement));
+}
+
+/**
+ * Pack normalized, i.e., range [0..1], float3 value to uint at a certain number of bits per component, rounding up.
+ */
+uint PackNormToUintCeil(float3 FloatNorm, uint3 NumComponentBits)
+{
+	return PackToUint(
+		uint3(
+			PackNormToUintCeil(FloatNorm.x, NumComponentBits.x),
+			PackNormToUintCeil(FloatNorm.y, NumComponentBits.y),
+			PackNormToUintCeil(FloatNorm.z, NumComponentBits.z)
+		), 
+		NumComponentBits);
+}
+
+uint PackNormToUintCeil(float3 FloatNorm, uint BitsPerElement)
+{
+	return PackNormToUintCeil(FloatNorm, uint3(BitsPerElement,BitsPerElement,BitsPerElement));
+}
+
 // Implement BitStreamReader for ByteAddressBuffer (RO), RWByteAddressBuffer (RW) and dynamic choice (RORW).
 struct FBitStreamReaderState
 {
diff --git a/Engine/Shaders/Private/SceneCulling/SceneCulling.ush b/Engine/Shaders/Private/SceneCulling/SceneCulling.ush
--- a/Engine/Shaders/Private/SceneCulling/SceneCulling.ush
+++ b/Engine/Shaders/Private/SceneCulling/SceneCulling.ush
@@ -5,8 +5,11 @@
 #include "../Common.ush"
 
 #include "/Engine/Shared/SceneCullingDefinitions.h"
+#include "../BitPacking.ush"
 
 
+#define SCENE_CULL_EXPLICIT_BOUNDS_BITS_PER_ELEMENT (8u) // 2x3 components a K bits each
+
 /**
  * Unpacked cell data.
  */
@@ -16,6 +19,7 @@ struct FSceneHiearchyCellData
 	float3 LocalBoundsCenter; // Bounds center is local to the cell block
 	float3 LocalBoundsExtent;
 	uint BlockId;
+	float MaxInstanceRadius;
 
 	FCellBlockData BlockData;
 };
@@ -49,7 +53,34 @@ inline uint CellIndexToBlockId(uint CellIndex)
 }
 
 uint bUseExplicitCellBounds;
-StructuredBuffer<float4> ExplicitCellBounds;
+StructuredBuffer<uint4> ExplicitCellBounds;
+
+struct FAabb
+{
+	float3 Min;
+	float3 Max;
+};
+
+FAabb UnpackAabb(uint2 Packed, uint BitsPerElement, float SizeScale = 1.0f)
+{
+	FAabb Aabb;
+	Aabb.Min = UnpackToFloat3(Packed.x, BitsPerElement) * SizeScale;
+	Aabb.Max = UnpackToFloat3(Packed.y, BitsPerElement) * SizeScale;
+	return Aabb;
+}
+
+struct FInstanceCellAttributes
+{
+	FAabb Aabb;
+};
+
+FInstanceCellAttributes UnpackInstanceCellAttributes(uint2 Packed, float AabbSizeScale = 1.0f, float RadiusScale = 1.0f)
+{
+	FInstanceCellAttributes Result;
+	Result.Aabb = UnpackAabb(Packed, SCENE_CULL_EXPLICIT_BOUNDS_BITS_PER_ELEMENT, AabbSizeScale);	
+
+	return Result;
+}
 
 FSceneHiearchyCellData GetSceneHiearchyCellData(uint CellId, bool bAllowExplicitBounds = true)
 {
@@ -64,15 +95,22 @@ FSceneHiearchyCellData GetSceneHiearchyCellData(uint CellId, bool bAllowExplicit
 	CellData.BlockData = GetCellBlockData(CellData.BlockId);
 	const float LevelCellSize = CellData.BlockData.LevelCellSize;
 
-	// TODO: roll half-offset into the block world position?
-	CellData.LocalBoundsCenter = float3(CellData.LocalCellCoord) * LevelCellSize + (LevelCellSize * 0.5f).xxx;
+	CellData.MaxInstanceRadius = LevelCellSize * 0.5f;
 	if (bUseExplicitCellBounds && bAllowExplicitBounds)
 	{
-		CellData.LocalBoundsCenter += ExplicitCellBounds[CellId * 2 + 0].xyz;
-		CellData.LocalBoundsExtent = ExplicitCellBounds[CellId * 2 + 1].xyz;
+		float3 ImplicitBoundsMin = float3(CellData.LocalCellCoord) * LevelCellSize - (LevelCellSize * 0.5f).xxx;
+
+		FInstanceCellAttributes InstanceCellAttributes = UnpackInstanceCellAttributes(ExplicitCellBounds[CellId].xy, LevelCellSize * 2.0f, LevelCellSize * 0.5f);
+		//FInstanceCellAttributes InstanceCellAttributes = UnpackInstanceCellAttributes(ExplicitCellBounds[CellId].zw, LevelCellSize * 2.0f, LevelCellSize * 0.5f);
+
+		FAabb CellRelativeBounds = InstanceCellAttributes.Aabb;
+		CellData.LocalBoundsCenter = (CellRelativeBounds.Min + CellRelativeBounds.Max) * 0.5f + ImplicitBoundsMin;
+		CellData.LocalBoundsExtent = (CellRelativeBounds.Max - CellRelativeBounds.Min) * 0.5f;
 	}
 	else
 	{
+		// TODO: roll half-offset into the block world position?
+		CellData.LocalBoundsCenter = float3(CellData.LocalCellCoord) * LevelCellSize + (LevelCellSize * 0.5f).xxx;
 		// Note: extent is _not_ half the cell size as we have loose bounds.
 		CellData.LocalBoundsExtent = LevelCellSize;
 	}
diff --git a/Engine/Shaders/Private/SceneCulling/SceneCullingBuildExplicitBounds.usf b/Engine/Shaders/Private/SceneCulling/SceneCullingBuildExplicitBounds.usf
--- a/Engine/Shaders/Private/SceneCulling/SceneCullingBuildExplicitBounds.usf
+++ b/Engine/Shaders/Private/SceneCulling/SceneCullingBuildExplicitBounds.usf
@@ -12,10 +12,7 @@
 
 #include "SceneCulling.ush"
 
-RWStructuredBuffer<float4> OutExplicitCellBoundsBuffer;
-
-
-void TransformAABB(float4x4 Transform, float3 Center, float3 Extent, inout float3 OutTransformedBoundsMin, inout float3 OutTransformedBoundsMax)
+FAabb TransformAABB(float4x4 Transform, float3 Center, float3 Extent)
 {
 	float3 NewCenter = mul(float4(Center, 1.0f), Transform).xyz;
 
@@ -23,35 +20,187 @@ void TransformAABB(float4x4 Transform, float3 Center, float3 Extent, inout float
 		+ abs(Extent.yyy * Transform[1].xyz)
 		+ abs(Extent.zzz * Transform[2].xyz);
 
-	OutTransformedBoundsMin = NewCenter - NewExtent;
-	OutTransformedBoundsMax = NewCenter + NewExtent;
+	FAabb Result;
+	Result.Min = NewCenter - NewExtent;
+	Result.Max = NewCenter + NewExtent;
+	return Result;
+}
+
+// pack a normalized AABB to an uint2, using BitsPerElement for each dimension
+// using floor on both min and max to avoid expansion when repacking already conservative bounds to the same precision.
+uint2 PackNormToUint2Floor(FAabb NormAabb, uint BitsPerElement)
+{
+	checkSlow(all(NormAabb.Min >= 0.0f) && all(NormAabb.Min <= 1.0f));
+	checkSlow(all(NormAabb.Max >= 0.0f) && all(NormAabb.Max <= 1.0f));
+
+	uint2 Result = {
+		PackNormToUintFloor(NormAabb.Min, BitsPerElement),
+		PackNormToUintFloor(NormAabb.Max, BitsPerElement)
+	};
+	return Result;
+}
+
+// pack a normalized AABB to an uint2, using BitsPerElement for each dimension
+// using floor on both min and max to avoid expansion when repacking already conservative bounds to the same precision.
+uint2 PackNormToUint2Floor(FInstanceCellAttributes InstanceCellAttributes)
+{
+	uint2 Result = PackNormToUint2Floor(InstanceCellAttributes.Aabb, SCENE_CULL_EXPLICIT_BOUNDS_BITS_PER_ELEMENT);
+	return Result;
 }
 
+// pack a normalized FInstanceCellAttributes to an uint2, using BitsPerElement for each dimension
+// rounds the upper bounds up, and the lower down to ensure conservative bounds.
+uint2 PackNormToUint2(FAabb NormAabb, uint BitsPerElement)
+{
+	checkSlow(all(NormAabb.Min >= 0.0f) && all(NormAabb.Min <= 1.0f));
+	checkSlow(all(NormAabb.Max >= 0.0f) && all(NormAabb.Max <= 1.0f));
+
+	uint2 Result = {
+		PackNormToUintFloor(NormAabb.Min, BitsPerElement),
+		PackNormToUintCeil(NormAabb.Max, BitsPerElement)
+	};
+	return Result;
+}
+
+// pack a normalized FInstanceCellAttributes to an uint2, using BitsPerElement for each dimension
+// rounds the upper bounds up, and the lower down to ensure conservative bounds.
+uint2 PackNormToUint2(FInstanceCellAttributes InstanceCellAttributes)
+{
+	uint2 Result = PackNormToUint2(InstanceCellAttributes.Aabb, SCENE_CULL_EXPLICIT_BOUNDS_BITS_PER_ELEMENT);
+	return Result;
+}
+
+FInstanceCellAttributes ThreadGroupReduce(FInstanceCellAttributes InstanceCellAttributes, uint GroupThreadIndex)
+{
+	FInstanceCellAttributes Result;
+	GroupMemoryBarrierWithGroupSync();
+	Result.Aabb.Min.x = ThreadGroupReduceMin(InstanceCellAttributes.Aabb.Min.x, GroupThreadIndex); 
+	GroupMemoryBarrierWithGroupSync();
+	Result.Aabb.Max.x = ThreadGroupReduceMax(InstanceCellAttributes.Aabb.Max.x, GroupThreadIndex); 
+
+	GroupMemoryBarrierWithGroupSync();
+	Result.Aabb.Min.y = ThreadGroupReduceMin(InstanceCellAttributes.Aabb.Min.y, GroupThreadIndex); 
+	GroupMemoryBarrierWithGroupSync();
+	Result.Aabb.Max.y = ThreadGroupReduceMax(InstanceCellAttributes.Aabb.Max.y, GroupThreadIndex); 
+
+	GroupMemoryBarrierWithGroupSync();
+	Result.Aabb.Min.z = ThreadGroupReduceMin(InstanceCellAttributes.Aabb.Min.z, GroupThreadIndex); 
+	GroupMemoryBarrierWithGroupSync();
+	Result.Aabb.Max.z = ThreadGroupReduceMax(InstanceCellAttributes.Aabb.Max.z, GroupThreadIndex); 
+
+	return Result;
+}
+
+void ProcessInstance(uint InstanceId, FDFVector3 WorldToBlockRelativeTranslation, inout FInstanceCellAttributes InOutAttributes)
+{
+	uint PrimitiveId = 0u;
+	uint InstanceFlags = 0u;
+	LoadInstancePrimitiveIdAndFlags(InstanceId, PrimitiveId, InstanceFlags);
+
+	if (PrimitiveId != INVALID_PRIMITIVE_ID)
+	{
+		FInstanceSceneData InstanceData = GetInstanceSceneDataUnchecked(InstanceId);
+		float4x4 LocalToBlockRelativeWorld = DFFastMultiplyTranslationDemote(InstanceData.LocalToWorld, WorldToBlockRelativeTranslation);
+		FAabb RelBounds = TransformAABB(LocalToBlockRelativeWorld, InstanceData.LocalBoundsCenter, InstanceData.LocalBoundsExtent);
+		InOutAttributes.Aabb.Min = min(InOutAttributes.Aabb.Min, RelBounds.Min);
+		InOutAttributes.Aabb.Max = max(InOutAttributes.Aabb.Max, RelBounds.Max);
+	}
+}
+
+#ifdef ComputeExplicitCellBounds
+
 StructuredBuffer<uint> UpdatedCellIds;
+RWStructuredBuffer<uint4> OutExplicitCellBoundsBuffer;
+
+#if USE_CHUNK_BOUNDS
+
+StructuredBuffer<uint2> ExplicitChunkBounds;
+
+#endif
+
+StructuredBuffer<uint> DirtyCellBoundsMaskStatic;
+
+#if USE_CHUNK_BOUNDS
+
+void BuildChunksAttribs(uint ItemChunksOffset, uint StartIndex, uint EndIndex, uint GroupThreadIndex, inout FInstanceCellAttributes InOutInstanceCellAttributesNorm)
+{
+	FInstanceCellAttributes InstanceCellAttributesNorm = InOutInstanceCellAttributesNorm;
+	// Group-stride loop over the chunks
+	for (uint Index = StartIndex + GroupThreadIndex; Index < EndIndex; Index += NUM_THREADS_PER_GROUP)
+	{
+		uint ChunkId = ItemChunksOffset + Index;
+		FInstanceCellAttributes ChunkAttributesNorm = UnpackInstanceCellAttributes(ExplicitChunkBounds[ChunkId]);
+		InstanceCellAttributesNorm.Aabb.Min = min(InstanceCellAttributesNorm.Aabb.Min, ChunkAttributesNorm.Aabb.Min);
+		InstanceCellAttributesNorm.Aabb.Max = max(InstanceCellAttributesNorm.Aabb.Max, ChunkAttributesNorm.Aabb.Max);
+	}
+
+	InOutInstanceCellAttributesNorm = ThreadGroupReduce(InstanceCellAttributesNorm, GroupThreadIndex);
+}
+#endif
 
 [numthreads(NUM_THREADS_PER_GROUP, 1, 1)]
 void ComputeExplicitCellBounds(uint3 GroupId : SV_GroupID, uint GroupThreadIndex : SV_GroupIndex)
 {
-#if DO_FULL_REBUILD
-	const uint CellId = GetUnWrappedDispatchGroupId(GroupId);
-
-	if (CellId < MaxCells)
-	{
-#else
 	const uint CellIdIndex = GetUnWrappedDispatchGroupId(GroupId);
 
 	if (CellIdIndex < MaxCells)
 	{
-		uint CellId = UpdatedCellIds[CellIdIndex];
-#endif
-		FSceneHiearchyCellData CellData = GetSceneHiearchyCellData(CellId, /*bAllowExplicitBounds*/ false);
 
-		float3 BoundsMin = float3(POSITIVE_INFINITY, POSITIVE_INFINITY, POSITIVE_INFINITY);
-		float3 BoundsMax = float3(NEGATIVE_INFINITY, NEGATIVE_INFINITY, NEGATIVE_INFINITY);
+#if USE_CHUNK_BOUNDS
+		uint PackedCellId = UpdatedCellIds[CellIdIndex];
+		uint CellId = PackedCellId >> 1u;
+		FCellHeader CellHeader = GetCellHeader(CellId);
 
-		FDFVector3 WorldToBlockRelativeTranslation = DFNegate(CellData.BlockData.WorldPos);
+		// If static was changed we buld both that and the combined bounds
+		bool bWasStaticChanged = (PackedCellId & 1u) != 0u;
+		
+		FInstanceCellAttributes StaticAttribsNorm;
+		StaticAttribsNorm.Aabb.Min = float3(1.0f, 1.0f, 1.0f);
+		StaticAttribsNorm.Aabb.Max = float3(0.0f, 0.0f, 0.0f);
 
+		// 1. build the static bounds if any
+		if (CellHeader.NumStaticChunks > 0u)
+		{
+			if (bWasStaticChanged)
+			{
+				BuildChunksAttribs(CellHeader.ItemChunksOffset, 0u, CellHeader.NumStaticChunks, GroupThreadIndex, StaticAttribsNorm);
+			}
+			else // unchanged, pull the previous result
+			{
+				// Note: redundant load on every thread, but should be a scalar load anyway, hope the compiler figures that out
+				StaticAttribsNorm = UnpackInstanceCellAttributes(OutExplicitCellBoundsBuffer[CellId].zw);
+			}
+		} 
+
+		// Start from static bounds and merge in dynamic ones
+		FInstanceCellAttributes CombinedAttrbisNorm = StaticAttribsNorm;
+		// 1. loop over the dynamic bounds (if any) 
+		if (CellHeader.NumDynamicChunks > 0u)
+		{
+			BuildChunksAttribs(CellHeader.ItemChunksOffset, CellHeader.NumStaticChunks, CellHeader.NumItemChunks, GroupThreadIndex, CombinedAttrbisNorm);
+		}
+
+		if (GroupThreadIndex == 0)
+		{
+			OutExplicitCellBoundsBuffer[CellId] = 
+				// combined, 
+				uint4(
+					// combined, 
+					PackNormToUint2Floor(CombinedAttrbisNorm), 
+					// static, 
+					PackNormToUint2Floor(StaticAttribsNorm));
+		}
+
+#else
+		uint CellId = UpdatedCellIds[CellIdIndex];
 		FCellHeader CellHeader = GetCellHeader(CellId);
+
+		FSceneHiearchyCellData CellData = GetSceneHiearchyCellData(CellId, /*bAllowExplicitBounds*/ false);
+		FInstanceCellAttributes Attribs;
+		Attribs.Aabb.Min = float3(POSITIVE_INFINITY, POSITIVE_INFINITY, POSITIVE_INFINITY);
+		Attribs.Aabb.Max = float3(NEGATIVE_INFINITY, NEGATIVE_INFINITY, NEGATIVE_INFINITY);
+
+		FDFVector3 WorldToBlockRelativeTranslation = DFNegate(CellData.BlockData.WorldPos);
 		for (uint Index = 0u; Index < CellHeader.NumItemChunks; ++Index)
 		{
 			uint PackedItemChunkDesc = InstanceHierarchyItemChunks[CellHeader.ItemChunksOffset + Index];
@@ -65,44 +214,86 @@ void ComputeExplicitCellBounds(uint3 GroupId : SV_GroupID, uint GroupThreadIndex
 			}
 
 			uint InstanceId = UnpackChunkInstanceId(bIsRLEPackedChunk, PackedItemChunkDesc, GroupThreadIndex);
-
-			uint PrimitiveId;
-			uint InstanceFlags;
-			LoadInstancePrimitiveIdAndFlags(InstanceId, PrimitiveId, InstanceFlags);
-
-			if (PrimitiveId == INVALID_PRIMITIVE_ID)
-			{
-				continue;
-			}
-
-			FInstanceSceneData InstanceData = GetInstanceSceneDataUnchecked(InstanceId);
-			
-			float3 RelBoundsMin;
-			float3 RelBoundsMax;
-
-			float4x4 LocalToBlockRelativeWorld = DFFastMultiplyTranslationDemote(InstanceData.LocalToWorld, WorldToBlockRelativeTranslation);
-			TransformAABB(LocalToBlockRelativeWorld, InstanceData.LocalBoundsCenter, InstanceData.LocalBoundsExtent, RelBoundsMin, RelBoundsMax);
-			BoundsMin = min(BoundsMin, RelBoundsMin);
-			BoundsMax = max(BoundsMax, RelBoundsMax);
-
-			// TODO: Maybe do a group reduce here & store per chunk bounds as well? TBD Per chunk is a lot more data, and also more logic in the unpack may lead to very unbalanced workloads
-			//       Also, if implementing this it would make more sense to go wide over the chunks & then reduce the final values (atomically, probably).
+			ProcessInstance(InstanceId, WorldToBlockRelativeTranslation, Attribs);
 		}
 
-		float3 GroupBoundsMin = BoundsMin;
-		float3 GroupBoundsMax = BoundsMax;
+		FInstanceCellAttributes GroupAttribs = Attribs;
 
 		if (CellHeader.NumItemChunks > 0)
 		{
-			GroupBoundsMin = float3(ThreadGroupReduceMin(BoundsMin.x, GroupThreadIndex), ThreadGroupReduceMin(BoundsMin.y, GroupThreadIndex), ThreadGroupReduceMin(BoundsMin.z, GroupThreadIndex)); 
-			GroupBoundsMax = float3(ThreadGroupReduceMax(BoundsMax.x, GroupThreadIndex), ThreadGroupReduceMax(BoundsMax.y, GroupThreadIndex), ThreadGroupReduceMax(BoundsMax.z, GroupThreadIndex)); 
+			GroupAttribs = ThreadGroupReduce(Attribs, GroupThreadIndex);
 		}
 
 		if (GroupThreadIndex == 0)
 		{
-			// TODO: Pack into far fewer bits! The bounds are _cell_ relative and thus need very few bits
-			OutExplicitCellBoundsBuffer[CellId * 2 + 0] = float4((GroupBoundsMin + GroupBoundsMax) * 0.5f - CellData.LocalBoundsCenter, CellId);
-			OutExplicitCellBoundsBuffer[CellId * 2 + 1] = float4((GroupBoundsMax - GroupBoundsMin) * 0.5f, CellId);
+			// bounds are block-relative, store offset from min corner instead of center to get positive range [0, CellData.LocalBoundsExtent * 2]
+			float3 CellMin = CellData.LocalBoundsCenter - CellData.LocalBoundsExtent;
+
+			FInstanceCellAttributes CellRelativeNorm;
+			// min/max in positive range [0, CellData.LocalBoundsExtent * 2] -> normalized [0,1.0f)
+			CellRelativeNorm.Aabb.Min = saturate((GroupAttribs.Aabb.Min - CellMin) / (CellData.LocalBoundsExtent * 2.0f));
+			CellRelativeNorm.Aabb.Max = saturate((GroupAttribs.Aabb.Max - CellMin) / (CellData.LocalBoundsExtent * 2.0f));
+			
+			OutExplicitCellBoundsBuffer[CellId] = uint4(PackNormToUint2(CellRelativeNorm), 0u, 0u);
 		}
+
+#endif
+	}
+}
+#endif // ComputeExplicitCellBounds
+
+#ifdef ComputeExplicitChunkBounds
+
+int DirtyChunkCount;
+StructuredBuffer<uint2> DirtyChunkBoundsData;
+RWStructuredBuffer<uint2> OutExplicitChunkBoundsBuffer;
+
+[numthreads(NUM_THREADS_PER_GROUP, 1, 1)]
+void ComputeExplicitChunkBounds(uint3 GroupId : SV_GroupID, uint GroupThreadIndex : SV_GroupIndex)
+{
+	const uint ChunkIdIndex = GetUnWrappedDispatchGroupId(GroupId);
+
+	if (ChunkIdIndex >= DirtyChunkCount)
+	{
+		return;
+	}
+	uint CellId  = DirtyChunkBoundsData[ChunkIdIndex].x;
+	uint ChunkId = DirtyChunkBoundsData[ChunkIdIndex].y;
+
+	// Load cell data
+	FSceneHiearchyCellData CellData = GetSceneHiearchyCellData(CellId, /*bAllowExplicitBounds*/ false);
+	FDFVector3 WorldToBlockRelativeTranslation = DFNegate(CellData.BlockData.WorldPos);
+
+	// Load chunk bounds (max 64)
+	uint PackedItemChunkDesc = InstanceHierarchyItemChunks[ChunkId];
+
+	bool bIsRLEPackedChunk = false;
+	uint ChunkNumInstances =  UnpackChunkInstanceCount(PackedItemChunkDesc, bIsRLEPackedChunk);
+
+	FInstanceCellAttributes Attribs;
+	Attribs.Aabb.Min = float3(POSITIVE_INFINITY, POSITIVE_INFINITY, POSITIVE_INFINITY);
+	Attribs.Aabb.Max = float3(NEGATIVE_INFINITY, NEGATIVE_INFINITY, NEGATIVE_INFINITY);
+
+	if (GroupThreadIndex < ChunkNumInstances)
+	{
+		uint InstanceId = UnpackChunkInstanceId(bIsRLEPackedChunk, PackedItemChunkDesc, GroupThreadIndex);
+		ProcessInstance(InstanceId, WorldToBlockRelativeTranslation, Attribs);
+	}
+
+	FInstanceCellAttributes GroupAttribs = ThreadGroupReduce(Attribs, GroupThreadIndex);
+
+	if (GroupThreadIndex == 0)
+	{
+		// store offset from min corner instead of center to get positive range [0, CellData.LocalBoundsExtent * 2]
+		float3 CellMin = CellData.LocalBoundsCenter - CellData.LocalBoundsExtent;
+
+		// min/max in positive range [0, CellData.LocalBoundsExtent * 2] -> normalized [0,1.0f)
+		FInstanceCellAttributes CellRelativeNorm;
+		CellRelativeNorm.Aabb.Min = saturate((GroupAttribs.Aabb.Min - CellMin) / (CellData.LocalBoundsExtent * 2.0f));
+		CellRelativeNorm.Aabb.Max = saturate((GroupAttribs.Aabb.Max - CellMin) / (CellData.LocalBoundsExtent * 2.0f));
+			
+		OutExplicitChunkBoundsBuffer[ChunkId] = PackNormToUint2(CellRelativeNorm);
 	}
 }
+
+#endif // ComputeExplicitChunkBounds
diff --git a/Engine/Shaders/Private/SceneCulling/SceneCullingDebugRender.usf b/Engine/Shaders/Private/SceneCulling/SceneCullingDebugRender.usf
--- a/Engine/Shaders/Private/SceneCulling/SceneCullingDebugRender.usf
+++ b/Engine/Shaders/Private/SceneCulling/SceneCullingDebugRender.usf
@@ -20,7 +20,7 @@ void DebugRender(uint3 GroupId : SV_GroupID, uint GroupThreadIndex : SV_GroupInd
 		FSceneHiearchyCellData CellData = GetSceneHiearchyCellData(CellId);
 
 		float3 Offset = DFDemote(CellData.BlockData.WorldPos);
-		float4 CellColor = float4(IntToColor(CellId >> NumCellsPerBlockLog2), 1.0f);
+		float4 CellColor = float4(IntToColor(CellId), 1.0f);
 		AddAABBWS(CellData.LocalBoundsCenter.xyz + Offset - CellData.LocalBoundsExtent.xyz, CellData.LocalBoundsCenter.xyz + Offset + CellData.LocalBoundsExtent.xyz, CellColor);
 	}
 }
diff --git a/Engine/Shaders/Shared/SceneCullingDefinitions.h b/Engine/Shaders/Shared/SceneCullingDefinitions.h
--- a/Engine/Shaders/Shared/SceneCullingDefinitions.h
+++ b/Engine/Shaders/Shared/SceneCullingDefinitions.h
@@ -2,7 +2,6 @@
 
 #pragma once
 
-
 // 1 + 7 + 24 bits_
 #define  INSTANCE_HIERARCHY_ITEM_CHUNK_COUNT_SHIFT  (24u)
 #define  INSTANCE_HIERARCHY_ITEM_CHUNK_ID_MASK  ((1u << INSTANCE_HIERARCHY_ITEM_CHUNK_COUNT_SHIFT) - 1u)
@@ -33,6 +32,7 @@ struct FCellHeader
 	uint NumItemChunks;
 	uint ItemChunksOffset;
 	uint NumStaticChunks;
+	uint NumDynamicChunks;
 };
 
 inline FCellHeader UnpackCellHeader(FPackedCellHeader Packed)
@@ -40,7 +40,8 @@ inline FCellHeader UnpackCellHeader(FPackedCellHeader Packed)
 	FCellHeader Result;
 	Result.ItemChunksOffset = Packed.ItemChunksOffset;
 	Result.NumStaticChunks = Packed.NumStaticDynamicItemChunks >> 16;
-	Result.NumItemChunks = (Packed.NumStaticDynamicItemChunks & 0xFFFFu) + Result.NumStaticChunks;
+	Result.NumDynamicChunks = Packed.NumStaticDynamicItemChunks & 0xFFFFu;
+	Result.NumItemChunks = Result.NumDynamicChunks + Result.NumStaticChunks;
 	return Result;
 }
 
diff --git a/Engine/Source/Runtime/Renderer/Private/SceneCulling/SceneCulling.cpp b/Engine/Source/Runtime/Renderer/Private/SceneCulling/SceneCulling.cpp
--- a/Engine/Source/Runtime/Renderer/Private/SceneCulling/SceneCulling.cpp
+++ b/Engine/Source/Runtime/Renderer/Private/SceneCulling/SceneCulling.cpp
@@ -171,7 +171,17 @@ static TAutoConsoleVariable<int32> CVarValidateAllInstanceAllocations(
 static TAutoConsoleVariable<int32> CVarSceneCullingUseExplicitCellBounds(
 	TEXT("r.SceneCulling.ExplicitCellBounds"), 
 	1, 
-	TEXT("Enable to to construct explicit cell bounds by processing the instance bounds as the scene is updated. Adds some GPU cost to the update but this is typically more than paid for by improved culling."),
+	TEXT("Enable to to construct explicit cell bounds by processing the instance bounds as the scene is updated.\n")
+	TEXT("   Adds some GPU cost to the update but this is typically more than paid for by improved culling.\n")
+	TEXT("   For debugging purposes it can be set to 2 to force updating all bounds every frame."),
+	ECVF_RenderThreadSafe);
+
+static TAutoConsoleVariable<int32> CVarSceneCullingUseExplicitChunkBounds(
+	TEXT("r.SceneCulling.ExplicitChunkBounds"), 
+	1, 
+	TEXT("Enable to to construct explicit chunk bounds by processing the instance bounds as the scene is updated.\n")
+	TEXT("   Building the chunk bounds on demand means we can often update only the dynamic ones, which saves a lot of cost update cost.\n")
+	TEXT("   In addition, this path has better load balancing in in scenes where a lot of items are in a few cells."),
 	ECVF_RenderThreadSafe);
 
 #if !UE_BUILD_SHIPPING
@@ -744,6 +754,7 @@ void FSceneCulling::Empty()
 	InstanceIdsBuffer.Empty();
 	CellBlockDataBuffer.Empty();
 	ExplicitCellBoundsBuffer.Empty();
+	ExplicitChunkBoundsBuffer.Empty();
 }
 
 class FComputeExplicitCellBounds_CS : public FGlobalShader
@@ -751,8 +762,8 @@ class FComputeExplicitCellBounds_CS : public FGlobalShader
 	DECLARE_GLOBAL_SHADER(FComputeExplicitCellBounds_CS);
 	SHADER_USE_PARAMETER_STRUCT(FComputeExplicitCellBounds_CS, FGlobalShader);
 
-	class FFullBuildDim : SHADER_PERMUTATION_BOOL("DO_FULL_REBUILD");
-	using FPermutationDomain = TShaderPermutationDomain<FFullBuildDim>;
+	class FUseChunkBoundsDim : SHADER_PERMUTATION_BOOL("USE_CHUNK_BOUNDS");
+	using FPermutationDomain = TShaderPermutationDomain<FUseChunkBoundsDim>;
 
 	static constexpr int32 NumThreadsPerGroup = 64;
 
@@ -784,10 +795,51 @@ class FComputeExplicitCellBounds_CS : public FGlobalShader
 		SHADER_PARAMETER_RDG_BUFFER_SRV( StructuredBuffer< uint >, InstanceIds)
 		SHADER_PARAMETER_RDG_BUFFER_SRV( StructuredBuffer< uint >, UpdatedCellIds)
 		SHADER_PARAMETER_RDG_BUFFER_UAV( RWStructuredBuffer<FVector4f>, OutExplicitCellBoundsBuffer)
+		SHADER_PARAMETER_RDG_BUFFER_SRV( StructuredBuffer< uint2 >, ExplicitChunkBounds)
 	END_SHADER_PARAMETER_STRUCT()
 };
 IMPLEMENT_GLOBAL_SHADER(FComputeExplicitCellBounds_CS, "/Engine/Private/SceneCulling/SceneCullingBuildExplicitBounds.usf", "ComputeExplicitCellBounds", SF_Compute);
 
+
+class FComputeExplicitChunkBounds_CS : public FGlobalShader
+{
+	DECLARE_GLOBAL_SHADER(FComputeExplicitChunkBounds_CS);
+	SHADER_USE_PARAMETER_STRUCT(FComputeExplicitChunkBounds_CS, FGlobalShader);
+
+	static constexpr int32 NumThreadsPerGroup = 64;
+
+	static bool ShouldCompilePermutation(const FGlobalShaderPermutationParameters& Parameters) 
+	{ 
+		return DoesPlatformSupportNanite(Parameters.Platform);
+	}
+
+	static void ModifyCompilationEnvironment(const FGlobalShaderPermutationParameters& Parameters, FShaderCompilerEnvironment& OutEnvironment)
+	{
+		FGlobalShader::ModifyCompilationEnvironment(Parameters, OutEnvironment);
+
+		OutEnvironment.SetDefine(TEXT("VF_SUPPORTS_PRIMITIVE_SCENE_DATA"), 1);
+		OutEnvironment.SetDefine(TEXT("NUM_THREADS_PER_GROUP"), NumThreadsPerGroup);
+	}
+
+	BEGIN_SHADER_PARAMETER_STRUCT(FParameters, )
+		SHADER_PARAMETER_RDG_UNIFORM_BUFFER( FSceneUniformParameters, Scene )
+
+		SHADER_PARAMETER(uint32, NumCellsPerBlockLog2)
+		SHADER_PARAMETER(uint32, CellBlockDimLog2)
+		SHADER_PARAMETER(uint32, LocalCellCoordMask) // (1 << NumCellsPerBlockLog2) - 1
+		SHADER_PARAMETER(int32, FirstLevel)
+		SHADER_PARAMETER(int32, DirtyChunkCount)
+
+		SHADER_PARAMETER_RDG_BUFFER_SRV( StructuredBuffer< FCellBlockData >, InstanceHierarchyCellBlockData)
+		SHADER_PARAMETER_RDG_BUFFER_SRV( StructuredBuffer< FPackedCellHeader >, InstanceHierarchyCellHeaders)
+		SHADER_PARAMETER_RDG_BUFFER_SRV( StructuredBuffer< uint >, InstanceHierarchyItemChunks)
+		SHADER_PARAMETER_RDG_BUFFER_SRV( StructuredBuffer< uint >, InstanceIds)
+		SHADER_PARAMETER_RDG_BUFFER_SRV( StructuredBuffer< int2 >, DirtyChunkBoundsData)
+		SHADER_PARAMETER_RDG_BUFFER_UAV( RWStructuredBuffer<uint2>, OutExplicitChunkBoundsBuffer)
+	END_SHADER_PARAMETER_STRUCT()
+};
+IMPLEMENT_GLOBAL_SHADER(FComputeExplicitChunkBounds_CS, "/Engine/Private/SceneCulling/SceneCullingBuildExplicitBounds.usf", "ComputeExplicitChunkBounds", SF_Compute);
+
 /**
  * Produce a world-space bounding sphere for an instance given local bounds and transforms.
  */
@@ -1267,6 +1319,8 @@ public:
 		int32 CellOffset = INDEX_NONE;
 		int32 ItemChunksOffset = INDEX_NONE;
 		int32 RemovedInstanceCount[EUpdateFrequencyCategory::Num] = { 0, 0 };
+		bool  bAnyChange[EUpdateFrequencyCategory::Num] = { false, false };
+
 		FCellHeader PrevCellHeader = FCellHeader { 0u, InvalidCellFlag };
 
 		SC_FORCEINLINE int32 GetRemovedInstanceCount() const
@@ -1391,7 +1445,8 @@ public:
 			SceneCulling.CellOccupancyMask.SetNum(NewMinSize, false);
 			if (ExplicitBoundsMode == EExplicitBoundsUpdateMode::Incremental)
 			{
-				DirtyCellBoundsMask.SetNum(NewMinSize, false);
+				DirtyCellBoundsMaskStatic.SetNum(NewMinSize, false);
+				DirtyCellBoundsMaskDynamic.SetNum(NewMinSize, false);
 			}
 			BUILDER_LOG("Alloc Cells [%d, %d], Block %d", StartIndex, StartIndex + CellBlockSize, BlockId.GetIndex());
 		}
@@ -1439,6 +1494,7 @@ public:
 		FTempCell& TempCell = GetOrAddTempCell(InstanceCellLoc);
 
 		TempCell.Builders[UpdateFrequencyCategory].AddRange(*this, InInstanceIdOffset, InInstanceIdCount);
+		TempCell.bAnyChange[UpdateFrequencyCategory] = true;
 
 #if SC_ENABLE_DETAILED_BUILDER_STATS
 		if (UpdateFrequencyCategory == EUpdateFrequencyCategory::Dynamic )
@@ -1467,6 +1523,7 @@ public:
 #endif
 		FTempCell& TempCell = GetOrAddTempCell(InstanceCellLoc);
 		TempCell.Builders[UpdateFrequencyCategory].Add(*this, uint32(InstanceId));
+		TempCell.bAnyChange[UpdateFrequencyCategory] = true;
 		return TempCell.CellOffset;
 	}
 
@@ -1646,7 +1703,7 @@ public:
 				// 3. Copy the list of chunk IDs
 				TempCell.OutputChunkIds(*this);
 
-				MarkCellBoundsDirty(TempCell.CellOffset);
+				MarkCellBoundsDirty(TempCell.CellOffset, TempCell.bAnyChange[EUpdateFrequencyCategory::Static], TempCell.bAnyChange[EUpdateFrequencyCategory::Dynamic]);
 			}
 			LogCell(UE::HLSL::UnpackCellHeader(CellHeader));
 			CellHeaderUploader.Add(CellHeader, TempCell.CellOffset);
@@ -1907,6 +1964,7 @@ public:
 		check(IsValidCell(TempCell.PrevCellHeader));
 		// Track total to be removed.
 		TempCell.RemovedInstanceCount[UpdateFrequencyCategory] += NumInstances;
+		TempCell.bAnyChange[UpdateFrequencyCategory] = true;
 
 #if SC_ENABLE_DETAILED_BUILDER_STATS
 		if (UpdateFrequencyCategory == EUpdateFrequencyCategory::Dynamic )
@@ -1926,18 +1984,24 @@ public:
 		MarkCellForRemove(CellIndex, NumInstances, UpdateFrequencyCategory);
 	}
 
-	SC_FORCEINLINE void MarkCellBoundsDirty(int32 CellIndex)
+	SC_FORCEINLINE void MarkCellBoundsDirty(int32 CellIndex, bool bStaticChanged, bool bDynamicChanged)
 	{
 		if (ExplicitBoundsMode == EExplicitBoundsUpdateMode::Incremental)
 		{
-			if (!DirtyCellBoundsMask[CellIndex])
+			if (!(DirtyCellBoundsMaskStatic[CellIndex] || DirtyCellBoundsMaskDynamic[CellIndex]))
 			{
 				DirtyCellBoundsIndices.Add(CellIndex);
 			}
-			DirtyCellBoundsMask[CellIndex] = true;
+			DirtyCellBoundsMaskStatic[CellIndex] = bStaticChanged;
+			DirtyCellBoundsMaskDynamic[CellIndex] = bDynamicChanged;
 		}
 	}
 
+	SC_FORCEINLINE void MarkCellBoundsDirty(int32 CellIndex, EUpdateFrequencyCategory::EType UpdateFrequencyCategory)
+	{
+		MarkCellBoundsDirty(CellIndex, UpdateFrequencyCategory == EUpdateFrequencyCategory::Static, UpdateFrequencyCategory == EUpdateFrequencyCategory::Dynamic);
+	}
+
 	inline void RemovePrecomputed(FSceneCulling::FPrimitiveState &PrimitiveState)
 	{
 		BUILDER_LOG("FPrimitiveState::Precomputed");
@@ -2032,7 +2096,7 @@ public:
 				}
 				else
 				{
-					MarkCellBoundsDirty(PrevCellIndex);
+					MarkCellBoundsDirty(PrevCellIndex, EUpdateFrequencyCategory::Dynamic);
 				}
 			}
 			if (bNeedAdd)
@@ -2114,7 +2178,7 @@ public:
 				{
 					// retain previous state.
 					NewPrimitiveState.Payload = PrevCellIndex;
-					MarkCellBoundsDirty(PrevCellIndex);
+					MarkCellBoundsDirty(PrevCellIndex, PrevPrimitiveState.bDynamic ? EUpdateFrequencyCategory::Dynamic : EUpdateFrequencyCategory::Static);
 				}
 			}
 		}
@@ -2293,21 +2357,61 @@ public:
 		INC_DWORD_STAT_BY(STAT_SceneCulling_UploadedBlocks, BlockDataUploader.GetNumScatters());
 
 		//TODO: capture and return the (returned) registered buffers, probably need to do that elsewhere anyway?
-		FRDGBuffer *CellBlockDataRDG = BlockDataUploader.ResizeAndUploadTo(GraphBuilder, SceneCulling.CellBlockDataBuffer, SceneCulling.CellBlockData.Num());
+		FRDGBuffer* CellBlockDataRDG = BlockDataUploader.ResizeAndUploadTo(GraphBuilder, SceneCulling.CellBlockDataBuffer, SceneCulling.CellBlockData.Num());
 		FBufferScatterUploader::FScatterInfo UpdatedCellScatterInfo;
-		FRDGBuffer *CellHeadersRDG = CellHeaderUploader.ResizeAndUploadTo(GraphBuilder, SceneCulling.CellHeadersBuffer, SceneCulling.CellHeaders.Num(), UpdatedCellScatterInfo);
-		FRDGBuffer *InstanceIdsRDG = ItemChunkDataUploader.ResizeAndUploadTo(GraphBuilder, SceneCulling.InstanceIdsBuffer, SceneCulling.PackedCellData.Num());
-		FRDGBuffer *ItemChunksRDG = ItemChunkUploader.ResizeAndUploadTo(GraphBuilder, SceneCulling.ItemChunksBuffer, SceneCulling.PackedCellChunkData.Num());
+		FRDGBuffer* CellHeadersRDG = CellHeaderUploader.ResizeAndUploadTo(GraphBuilder, SceneCulling.CellHeadersBuffer, SceneCulling.CellHeaders.Num(), UpdatedCellScatterInfo);
+		FRDGBuffer* InstanceIdsRDG = ItemChunkDataUploader.ResizeAndUploadTo(GraphBuilder, SceneCulling.InstanceIdsBuffer, SceneCulling.PackedCellData.Num());
+		FRDGBuffer* ItemChunksRDG = ItemChunkUploader.ResizeAndUploadTo(GraphBuilder, SceneCulling.ItemChunksBuffer, SceneCulling.PackedCellChunkData.Num());
 		
 		if (SceneCulling.bUseExplictBounds)
 		{
+			RDG_EVENT_SCOPE(GraphBuilder, "SceneCulling_ComputeExplicitCellBounds");
 			check(ExplicitBoundsMode != EExplicitBoundsUpdateMode::Disabled);
 
-			SceneCulling.ExplicitCellBoundsBuffer.ResizeBufferIfNeeded(GraphBuilder, SceneCulling.CellHeaders.Num() * 2);
+			SceneCulling.ExplicitCellBoundsBuffer.ResizeBufferIfNeeded(GraphBuilder, SceneCulling.CellHeaders.Num());
+
+			FRDGBuffer* ExplicitChunkBoundsBufferRDG = nullptr;
+			if (SceneCulling.bUseExplicitChunkBounds)
+			{
+				ExplicitChunkBoundsBufferRDG = SceneCulling.ExplicitChunkBoundsBuffer.ResizeBufferIfNeeded(GraphBuilder, SceneCulling.CellChunkIdAllocator.GetMaxSize());
+
+				if (!DirtyChunkBoundsData.IsEmpty())
+				{
+					check(SceneCulling.bUseExplicitChunkBounds);
+					FComputeExplicitChunkBounds_CS::FParameters* PassParameters = GraphBuilder.AllocParameters<FComputeExplicitChunkBounds_CS::FParameters>();
+					PassParameters->Scene = SceneUniformBuffer.GetBuffer(GraphBuilder);
+					PassParameters->NumCellsPerBlockLog2 = FSpatialHash::NumCellsPerBlockLog2;
+					PassParameters->CellBlockDimLog2 = FSpatialHash::CellBlockDimLog2;
+					PassParameters->LocalCellCoordMask = (1U << FSpatialHash::CellBlockDimLog2) - 1U;
+					PassParameters->FirstLevel = SceneCulling.SpatialHash.GetFirstLevel();
+					PassParameters->InstanceHierarchyCellBlockData = GraphBuilder.CreateSRV(CellBlockDataRDG);
+					PassParameters->InstanceHierarchyCellHeaders = GraphBuilder.CreateSRV(CellHeadersRDG);
+					PassParameters->InstanceIds = GraphBuilder.CreateSRV(InstanceIdsRDG);
+					PassParameters->InstanceHierarchyItemChunks = GraphBuilder.CreateSRV(ItemChunksRDG);
+					int32 DirtyChunkCount = DirtyChunkBoundsData.Num();
+					// Note: this Moves the DirtyCellBoundsIndices & thus implicitly clears it!
+					PassParameters->DirtyChunkBoundsData = GraphBuilder.CreateSRV(CreateStructuredBuffer(GraphBuilder, TEXT("SceneCulling.DirtyChunkBoundsData"), MoveTemp(DirtyChunkBoundsData)));
+					PassParameters->DirtyChunkCount = DirtyChunkCount;
+					PassParameters->OutExplicitChunkBoundsBuffer = GraphBuilder.CreateUAV(ExplicitChunkBoundsBufferRDG);
+
+					auto ComputeShader = GetGlobalShaderMap(SceneCulling.Scene.GetFeatureLevel())->GetShader<FComputeExplicitChunkBounds_CS>();
+
+					FComputeShaderUtils::AddPass(
+						GraphBuilder,
+						RDG_EVENT_NAME("ComputeExplicitChunkBounds"),
+						ComputeShader,
+						PassParameters,
+						FComputeShaderUtils::GetGroupCountWrapped(DirtyChunkCount)
+					);
+				}
+			}
+			else
+			{
+				SceneCulling.ExplicitChunkBoundsBuffer.Empty();
+			}
 
 			if (!DirtyCellBoundsIndices.IsEmpty())
 			{
-				RDG_EVENT_SCOPE(GraphBuilder, "SceneCulling_ComputeExplicitCellBounds");
 				FComputeExplicitCellBounds_CS::FParameters* PassParameters = GraphBuilder.AllocParameters<FComputeExplicitCellBounds_CS::FParameters>();
 				PassParameters->Scene = SceneUniformBuffer.GetBuffer(GraphBuilder);
 				PassParameters->NumCellsPerBlockLog2 = FSpatialHash::NumCellsPerBlockLog2;
@@ -2318,6 +2422,7 @@ public:
 				PassParameters->InstanceHierarchyCellHeaders = GraphBuilder.CreateSRV(CellHeadersRDG);
 				PassParameters->InstanceIds = GraphBuilder.CreateSRV(InstanceIdsRDG);
 				PassParameters->InstanceHierarchyItemChunks = GraphBuilder.CreateSRV(ItemChunksRDG);
+				PassParameters->ExplicitChunkBounds = ExplicitChunkBoundsBufferRDG ? GraphBuilder.CreateSRV(ExplicitChunkBoundsBufferRDG) : nullptr;
 				int32 MaxCellCount = DirtyCellBoundsIndices.Num();
 				// Note: this Moves the DirtyCellBoundsIndices & thus implicitly clears it!
 				PassParameters->UpdatedCellIds = GraphBuilder.CreateSRV(CreateStructuredBuffer(GraphBuilder, TEXT("SceneCulling.DirtyCellBoundsIndices"), MoveTemp(DirtyCellBoundsIndices)));
@@ -2325,7 +2430,7 @@ public:
 				PassParameters->OutExplicitCellBoundsBuffer = GraphBuilder.CreateUAV(SceneCulling.ExplicitCellBoundsBuffer.Register(GraphBuilder));
 
 				FComputeExplicitCellBounds_CS::FPermutationDomain PermutationVector;
-				PermutationVector.Set<FComputeExplicitCellBounds_CS::FFullBuildDim>(false /*TODO: actually remove the permutation after 5.4.1 */);
+				PermutationVector.Set<FComputeExplicitCellBounds_CS::FUseChunkBoundsDim>(ExplicitChunkBoundsBufferRDG != nullptr);
 
 				auto ComputeShader = GetGlobalShaderMap(SceneCulling.Scene.GetFeatureLevel())->GetShader<FComputeExplicitCellBounds_CS>(PermutationVector);
 
@@ -2341,6 +2446,7 @@ public:
 		else
 		{
 			SceneCulling.ExplicitCellBoundsBuffer.Empty();
+			SceneCulling.ExplicitChunkBoundsBuffer.Empty();
 		}
 
 
@@ -2384,16 +2490,26 @@ public:
 	{
 		if (SceneCulling.bUseExplictBounds)
 		{
-			bool bFullUpload = !SceneCulling.ExplicitCellBoundsBuffer.GetPooledBuffer().IsValid();
+			// trigger full upload if either feature has been turned on
+			bool bFullUpload = SceneCulling.bForceFullExplictBoundsBuild
+				|| !SceneCulling.ExplicitCellBoundsBuffer.GetPooledBuffer().IsValid()
+				|| (SceneCulling.bUseExplicitChunkBounds && !SceneCulling.ExplicitChunkBoundsBuffer.GetPooledBuffer().IsValid());
+
 			ExplicitBoundsMode = bFullUpload ? EExplicitBoundsUpdateMode::Full : EExplicitBoundsUpdateMode::Incremental;
 			
 			// No need for dirty tracking if we are going to upload all of them
 			if (ExplicitBoundsMode == EExplicitBoundsUpdateMode::Incremental)
 			{
-				DirtyCellBoundsMask.Init(false, SceneCulling.CellHeaders.Num());
+				DirtyCellBoundsMaskStatic.Init(false, SceneCulling.CellHeaders.Num());
+				DirtyCellBoundsMaskDynamic.Init(false, SceneCulling.CellHeaders.Num());
 			}
 			// But we're still going to populate the index list
 			DirtyCellBoundsIndices.Reset(SceneCulling.CellHeaders.Num());
+			
+			if (SceneCulling.bUseExplicitChunkBounds)
+			{
+				DirtyChunkBoundsData.Reset(SceneCulling.CellChunkIdAllocator.GetMaxSize());
+			}
 		}
 		else
 		{
@@ -2515,16 +2631,46 @@ public:
 			}
 		}
 
+		// expand into chunk indexes for the prospective cells to load balance better,
+		// we could do this expansion on the GPU instead to save some CPU time (but waste some GPU time).
+		if (!DirtyCellBoundsIndices.IsEmpty() && SceneCulling.bUseExplicitChunkBounds)
+		{
+			for (int32& CellIndex : DirtyCellBoundsIndices)
+			{
+				bool bStaticDirty = ExplicitBoundsMode == EExplicitBoundsUpdateMode::Full || DirtyCellBoundsMaskStatic[CellIndex];
+
+				FCellHeader CellHeader = UnpackCellHeader(SceneCulling.CellHeaders[CellIndex]);
+				check(IsValidCell(CellHeader));
+				uint32 StartIndex = 0u;
+
+				// If no static data was touched we don't update those chunks and instead reuse the already stored data in the cell.
+				if (!bStaticDirty)
+				{
+					StartIndex += CellHeader.NumStaticChunks;
+				}
+
+				for (uint32 ChunkIndex = StartIndex; ChunkIndex < CellHeader.NumItemChunks; ++ChunkIndex)
+				{
+					DirtyChunkBoundsData.Add(FIntVector2(CellIndex, CellHeader.ItemChunksOffset + ChunkIndex));
+				}
+
+				// tag index with the flag & overwrite the index in DirtyCellBoundsIndices for upload to the GPU
+				CellIndex = (CellIndex << 1u) | (bStaticDirty ? 1 : 0);
+			}
+
+		}
+
+
 		// Note: the below could be put in a separate task, to cut down wait time for subsequent tasks
 		// Releasing memory we're done with in the task saves some RT time
 		DirtyChunks.Empty();
 		DirtyBlocks.Empty();
 		RemovedInstanceFlags.Empty();
-		DirtyCellBoundsMask.Empty();
+		DirtyCellBoundsMaskDynamic.Empty();
+		DirtyCellBoundsMaskStatic.Empty();
 	}
 
-	using FRemovedIntanceFlags = TBitArray<SceneRenderingAllocator>;
-	FRemovedIntanceFlags RemovedInstanceFlags;
+	TBitArray<SceneRenderingAllocator> RemovedInstanceFlags;
 
 	// Used for allocating new cache slots, important to reset/maintain whenever a cache slot is freed
 	int32 LowestFreeCacheIndex = 0;
@@ -2544,8 +2690,10 @@ public:
 	TBitArray<SceneRenderingAllocator> DirtyBlocks;
 	TBitArray<SceneRenderingAllocator> DirtyChunks;
 	// Dirty tracking for explicit cell bounds
-	TBitArray<SceneRenderingAllocator> DirtyCellBoundsMask;
+	TBitArray<SceneRenderingAllocator> DirtyCellBoundsMaskStatic;
+	TBitArray<SceneRenderingAllocator> DirtyCellBoundsMaskDynamic;
 	TArray<int32, SceneRenderingAllocator> DirtyCellBoundsIndices;
+	TArray<FIntVector2, SceneRenderingAllocator> DirtyChunkBoundsData;
 
 	TStructuredBufferScatterUploader<FCellBlockData> BlockDataUploader;
 	TStructuredBufferScatterUploader<uint32, INSTANCE_HIERARCHY_MAX_CHUNK_SIZE> ItemChunkDataUploader;
@@ -2637,6 +2785,7 @@ FSceneCulling::FSceneCulling(FScene& InScene)
 	, InstanceIdsBuffer(16, TEXT("SceneCulling.Items"))
 	, CellBlockDataBuffer(16, TEXT("SceneCulling.CellBlockData"))
 	, ExplicitCellBoundsBuffer(16, TEXT("SceneCulling.ExplicitCellBounds"))
+	, ExplicitChunkBoundsBuffer(16, TEXT("SceneCulling.ExplicitChunkBounds"))
 {
 #if (!(UE_BUILD_SHIPPING || UE_BUILD_TEST))
 	if (CVarSceneCulling.GetValueOnAnyThread() != 0 && !UseNanite(InScene.GetShaderPlatform()))
@@ -2673,7 +2822,9 @@ FSceneCulling::FUpdater &FSceneCulling::BeginUpdate(FRDGBuilder& GraphBuilder, F
 	// Note: this only works in concert with the FGlobalComponentRecreateRenderStateContext on the CVarSceneCulling callback ensuring all geometry is re-registered
 	bIsEnabled = UseSceneCulling(Scene.GetShaderPlatform());
 
+	bForceFullExplictBoundsBuild = CVarSceneCullingUseExplicitCellBounds.GetValueOnRenderThread() == 2;
 	bUseExplictBounds = CVarSceneCullingUseExplicitCellBounds.GetValueOnRenderThread() != 0;
+	bUseExplicitChunkBounds = bUseExplictBounds && CVarSceneCullingUseExplicitChunkBounds.GetValueOnRenderThread() != 0;
 
 	SmallFootprintCellSideThreshold = CVarSmallFootprintSideThreshold.GetValueOnRenderThread();
 	bUseAsyncUpdate = CVarSceneCullingAsyncUpdate.GetValueOnRenderThread() != 0;
diff --git a/Engine/Source/Runtime/Renderer/Private/SceneCulling/SceneCulling.h b/Engine/Source/Runtime/Renderer/Private/SceneCulling/SceneCulling.h
--- a/Engine/Source/Runtime/Renderer/Private/SceneCulling/SceneCulling.h
+++ b/Engine/Source/Runtime/Renderer/Private/SceneCulling/SceneCulling.h
@@ -280,7 +280,10 @@ private:
 
 	FScene& Scene;
 	bool bIsEnabled = false;
+	bool bForceFullExplictBoundsBuild = false;
 	bool bUseExplictBounds = false;
+	bool bUseExplicitChunkBounds = false;
+
 	FSpatialHash SpatialHash;
 
 	// Kept in the class for now, since we only want one active at a time anyway.
@@ -320,5 +323,8 @@ private:
 	TPersistentStructuredBuffer<uint32> ItemChunksBuffer;
 	TPersistentStructuredBuffer<uint32> InstanceIdsBuffer;
 	TPersistentStructuredBuffer<FCellBlockData> CellBlockDataBuffer;
-	TPersistentStructuredBuffer<FVector4f> ExplicitCellBoundsBuffer;
+	// 2x bounds per cell, packed into an uint2 each, one for all items, and one for just the static.
+	TPersistentStructuredBuffer<FIntVector4> ExplicitCellBoundsBuffer;
+	// explicit chunk bounds, use the same packing.
+	TPersistentStructuredBuffer<FIntVector2> ExplicitChunkBoundsBuffer;
 };
diff --git a/Engine/Source/Runtime/Renderer/Private/SceneCulling/SceneCullingRenderer.cpp b/Engine/Source/Runtime/Renderer/Private/SceneCulling/SceneCullingRenderer.cpp
--- a/Engine/Source/Runtime/Renderer/Private/SceneCulling/SceneCullingRenderer.cpp
+++ b/Engine/Source/Runtime/Renderer/Private/SceneCulling/SceneCullingRenderer.cpp
@@ -30,7 +30,7 @@ FInstanceHierarchyParameters& FSceneCullingRenderer::GetShaderParameters(FRDGBui
 		ItemChunksRDG = SceneCulling.ItemChunksBuffer.Register(GraphBuilder);
 		InstanceIdsRDG = SceneCulling.InstanceIdsBuffer.Register(GraphBuilder);
 
-		ExplicitCellBoundsRDG = SceneCulling.bUseExplictBounds ? SceneCulling.ExplicitCellBoundsBuffer.Register(GraphBuilder) :  GSystemTextures.GetDefaultStructuredBuffer<FVector4f>(GraphBuilder);
+		ExplicitCellBoundsRDG = SceneCulling.bUseExplictBounds ? SceneCulling.ExplicitCellBoundsBuffer.Register(GraphBuilder) :  GSystemTextures.GetDefaultStructuredBuffer<FIntVector4>(GraphBuilder);
 
 #if 0
 		// Fully upload the buffers for debugging. 
diff --git a/Engine/Source/Runtime/Renderer/Private/SceneCulling/SceneCullingRenderer.h b/Engine/Source/Runtime/Renderer/Private/SceneCulling/SceneCullingRenderer.h
--- a/Engine/Source/Runtime/Renderer/Private/SceneCulling/SceneCullingRenderer.h
+++ b/Engine/Source/Runtime/Renderer/Private/SceneCulling/SceneCullingRenderer.h
@@ -19,7 +19,7 @@ BEGIN_SHADER_PARAMETER_STRUCT( FInstanceHierarchyParameters, )
 	SHADER_PARAMETER_RDG_BUFFER_SRV( StructuredBuffer< FPackedCellHeader >, InstanceHierarchyCellHeaders)
 	SHADER_PARAMETER_RDG_BUFFER_SRV( StructuredBuffer< uint >, InstanceHierarchyItemChunks)
 	SHADER_PARAMETER_RDG_BUFFER_SRV( StructuredBuffer< uint >, InstanceIds)
-	SHADER_PARAMETER_RDG_BUFFER_SRV( StructuredBuffer< float4 >, ExplicitCellBounds)
+	SHADER_PARAMETER_RDG_BUFFER_SRV( StructuredBuffer< uint4 >, ExplicitCellBounds)
 END_SHADER_PARAMETER_STRUCT()
 
 /**
