Disable ADAPTIVE_VOLUMETRIC_SHADOW_MAP in distortion shaders when using mobile renderer

diff --git a/Engine/Source/Runtime/Renderer/Private/DistortionRendering.cpp b/Engine/Source/Runtime/Renderer/Private/DistortionRendering.cpp
--- a/Engine/Source/Runtime/Renderer/Private/DistortionRendering.cpp
+++ b/Engine/Source/Runtime/Renderer/Private/DistortionRendering.cpp
@@ -249,7 +249,7 @@ public:
 
 		// Skip the material clip if depth test should not be done
 		OutEnvironment.SetDefine(TEXT("MATERIAL_SHOULD_DISABLE_DEPTH_TEST"), Parameters.MaterialParameters.bShouldDisableDepthTest ? 1 : 0);
-		OutEnvironment.SetDefine(TEXT("ADAPTIVE_VOLUMETRIC_SHADOW_MAP"), ShouldCompositeHeterogeneousVolumesWithTranslucency());
+		OutEnvironment.SetDefine(TEXT("ADAPTIVE_VOLUMETRIC_SHADOW_MAP"), !IsMobilePlatform(Parameters.Platform) && ShouldCompositeHeterogeneousVolumesWithTranslucency());
 	}
 };
 
