Don't apply AerialPerspective to Translucent Materials on Mobile when r.Mobile.PixelFogQuality is 0

diff --git a/Engine/Source/Runtime/Renderer/Private/SkyAtmosphereRendering.cpp b/Engine/Source/Runtime/Renderer/Private/SkyAtmosphereRendering.cpp
--- a/Engine/Source/Runtime/Renderer/Private/SkyAtmosphereRendering.cpp
+++ b/Engine/Source/Runtime/Renderer/Private/SkyAtmosphereRendering.cpp
@@ -393,6 +393,12 @@ static auto GetSizeAndInvSize = [](int32 Width, int32 Height)
 	return FVector4f(FWidth, FHeight, 1.0f / FWidth, 1.0f / FHeight);
 };
 
+bool ShouldApplyAerialPerspective(EShaderPlatform ShaderPlatform)
+{
+	static auto* MobilePixelFogQuality = IConsoleManager::Get().FindTConsoleVariableDataInt(TEXT("r.Mobile.PixelFogQuality"));
+	return !IsMobilePlatform(ShaderPlatform) || (MobilePixelFogQuality && (MobilePixelFogQuality->GetValueOnAnyThread() == 1));
+}
+
 void SetupSkyAtmosphereViewSharedUniformShaderParameters(const FViewInfo& View, const FSkyAtmosphereSceneProxy& SkyAtmosphereProxy, FSkyAtmosphereViewSharedUniformShaderParameters& OutParameters)
 {
 	GET_VALID_DATA_FROM_CVAR;
@@ -404,7 +410,7 @@ void SetupSkyAtmosphereViewSharedUniformShaderParameters(const FViewInfo& View,
 	}
 
 	OutParameters.CameraAerialPerspectiveVolumeSizeAndInvSize = GetSizeAndInvSize(CameraAerialPerspectiveVolumeScreenResolution, CameraAerialPerspectiveVolumeScreenResolution);
-	OutParameters.ApplyCameraAerialPerspectiveVolume = View.SkyAtmosphereCameraAerialPerspectiveVolume == nullptr ? 0.0f : 1.0f;
+	OutParameters.ApplyCameraAerialPerspectiveVolume = (View.SkyAtmosphereCameraAerialPerspectiveVolume == nullptr) || !ShouldApplyAerialPerspective(View.GetShaderPlatform()) ? 0.0f : 1.0f;
 	OutParameters.CameraAerialPerspectiveVolumeDepthResolution = float(CameraAerialPerspectiveVolumeDepthResolution);
 	OutParameters.CameraAerialPerspectiveVolumeDepthResolutionInv = 1.0f / OutParameters.CameraAerialPerspectiveVolumeDepthResolution;
 	OutParameters.CameraAerialPerspectiveVolumeDepthSliceLengthKm = CameraAerialPerspectiveVolumeDepthSliceLengthKm;
@@ -1465,6 +1471,7 @@ void FSceneRenderer::RenderSkyAtmosphereLookUpTables(FRDGBuilder& GraphBuilder,
 	const bool bRealTimeReflectionCaptureSkyAtmosphereViewLutTexture = Views.Num() > 0 && Scene->SkyLight && Scene->SkyLight->bRealTimeCaptureEnabled;
 	// CameraAP volume is required if there is a skydome or a volumetric cloud component rendered in a cubemap real time capture.
 	const bool bRealTimeReflectionCapture360APLutTexture = Views.Num() > 0 && Scene->SkyLight && Scene->SkyLight->bRealTimeCaptureEnabled && (Views[0].bSceneHasSkyMaterial || Scene->HasVolumetricCloud());
+	const bool bShouldApplyAerialPerspective = ShouldApplyAerialPerspective(Scene->GetShaderPlatform());
 	if (bRealTimeReflectionCaptureSkyAtmosphereViewLutTexture || bRealTimeReflectionCapture360APLutTexture)
 	{
 		FViewInfo& View = Views[0];
@@ -1552,7 +1559,7 @@ void FSceneRenderer::RenderSkyAtmosphereLookUpTables(FRDGBuilder& GraphBuilder,
 			PendingRDGResources.RealTimeReflectionCaptureSkyAtmosphereViewLutTexture = RealTimeReflectionCaptureSkyAtmosphereViewLutTexture;
 		}
 
-		if (bRealTimeReflectionCapture360APLutTexture)
+		if (bRealTimeReflectionCapture360APLutTexture && bShouldApplyAerialPerspective)
 		{
 			FIntVector CameraAPLutSize = View.SkyAtmosphereCameraAerialPerspectiveVolume->GetDesc().GetSize();
 			EPixelFormat TextureAerialLUTFormat = (CVarSkyAtmosphereLUT32.GetValueOnAnyThread() != 0) ? PF_A32B32G32R32F : PF_FloatRGBA;
@@ -1611,21 +1618,6 @@ void FSceneRenderer::RenderSkyAtmosphereLookUpTables(FRDGBuilder& GraphBuilder,
 
 		FRDGTextureRef SkyAtmosphereViewLutTexture = GraphBuilder.RegisterExternalTexture(View.SkyAtmosphereViewLutTexture);
 		FRDGTextureUAVRef SkyAtmosphereViewLutTextureUAV = GraphBuilder.CreateUAV(FRDGTextureUAVDesc(SkyAtmosphereViewLutTexture, 0));
-		FRDGTextureRef SkyAtmosphereCameraAerialPerspectiveVolume = GraphBuilder.RegisterExternalTexture(View.SkyAtmosphereCameraAerialPerspectiveVolume);
-		FRDGTextureUAVRef SkyAtmosphereCameraAerialPerspectiveVolumeUAV = GraphBuilder.CreateUAV(FRDGTextureUAVDesc(SkyAtmosphereCameraAerialPerspectiveVolume, 0));
-
-		FRDGTextureRef SkyAtmosphereCameraAerialPerspectiveVolumeMieOnly = nullptr;
-		FRDGTextureUAVRef SkyAtmosphereCameraAerialPerspectiveVolumeMieOnlyUAV = nullptr;
-		FRDGTextureRef SkyAtmosphereCameraAerialPerspectiveVolumeRayOnly = nullptr;
-		FRDGTextureUAVRef SkyAtmosphereCameraAerialPerspectiveVolumeRayOnlyUAV = nullptr;
-		if (bSeparatedAtmosphereMieRayLeigh)
-		{
-			SkyAtmosphereCameraAerialPerspectiveVolumeMieOnly = GraphBuilder.RegisterExternalTexture(View.SkyAtmosphereCameraAerialPerspectiveVolumeMieOnly);
-			SkyAtmosphereCameraAerialPerspectiveVolumeMieOnlyUAV = GraphBuilder.CreateUAV(FRDGTextureUAVDesc(SkyAtmosphereCameraAerialPerspectiveVolumeMieOnly, 0));
-
-			SkyAtmosphereCameraAerialPerspectiveVolumeRayOnly = GraphBuilder.RegisterExternalTexture(View.SkyAtmosphereCameraAerialPerspectiveVolumeRayOnly);
-			SkyAtmosphereCameraAerialPerspectiveVolumeRayOnlyUAV = GraphBuilder.CreateUAV(FRDGTextureUAVDesc(SkyAtmosphereCameraAerialPerspectiveVolumeRayOnly, 0));
-		}
 
 		TRDGUniformBufferRef<FVolumeShadowingShaderParametersGlobal0> LightShadowShaderParams0UniformBuffer{};
 		TRDGUniformBufferRef<FVolumeShadowingShaderParametersGlobal1> LightShadowShaderParams1UniformBuffer{};
@@ -1681,7 +1673,24 @@ void FSceneRenderer::RenderSkyAtmosphereLookUpTables(FRDGBuilder& GraphBuilder,
 		}
 
 		// Camera Atmosphere Volume
+		if(bShouldApplyAerialPerspective)
 		{
+			FRDGTextureRef SkyAtmosphereCameraAerialPerspectiveVolume = GraphBuilder.RegisterExternalTexture(View.SkyAtmosphereCameraAerialPerspectiveVolume);
+			FRDGTextureUAVRef SkyAtmosphereCameraAerialPerspectiveVolumeUAV = GraphBuilder.CreateUAV(FRDGTextureUAVDesc(SkyAtmosphereCameraAerialPerspectiveVolume, 0));
+
+			FRDGTextureRef SkyAtmosphereCameraAerialPerspectiveVolumeMieOnly = nullptr;
+			FRDGTextureUAVRef SkyAtmosphereCameraAerialPerspectiveVolumeMieOnlyUAV = nullptr;
+			FRDGTextureRef SkyAtmosphereCameraAerialPerspectiveVolumeRayOnly = nullptr;
+			FRDGTextureUAVRef SkyAtmosphereCameraAerialPerspectiveVolumeRayOnlyUAV = nullptr;
+			if (bSeparatedAtmosphereMieRayLeigh)
+			{
+				SkyAtmosphereCameraAerialPerspectiveVolumeMieOnly = GraphBuilder.RegisterExternalTexture(View.SkyAtmosphereCameraAerialPerspectiveVolumeMieOnly);
+				SkyAtmosphereCameraAerialPerspectiveVolumeMieOnlyUAV = GraphBuilder.CreateUAV(FRDGTextureUAVDesc(SkyAtmosphereCameraAerialPerspectiveVolumeMieOnly, 0));
+
+				SkyAtmosphereCameraAerialPerspectiveVolumeRayOnly = GraphBuilder.RegisterExternalTexture(View.SkyAtmosphereCameraAerialPerspectiveVolumeRayOnly);
+				SkyAtmosphereCameraAerialPerspectiveVolumeRayOnlyUAV = GraphBuilder.CreateUAV(FRDGTextureUAVDesc(SkyAtmosphereCameraAerialPerspectiveVolumeRayOnly, 0));
+			}
+
 			FRenderCameraAerialPerspectiveVolumeCS::FPermutationDomain PermutationVector;
 			PermutationVector.Set<FSampleCloudSkyAO>(CloudShadowAOData.bShouldSampleCloudSkyAO);
 			PermutationVector.Set<FSecondAtmosphereLight>(bSecondAtmosphereLightEnabled);
@@ -1725,10 +1734,11 @@ void FSceneRenderer::RenderSkyAtmosphereLookUpTables(FRDGBuilder& GraphBuilder,
 			FIntVector TextureSize = SkyAtmosphereCameraAerialPerspectiveVolume->Desc.GetSize();
 			const FIntVector NumGroups = FIntVector::DivideAndRoundUp(TextureSize, FRenderCameraAerialPerspectiveVolumeCS::GroupSize);
 			FComputeShaderUtils::AddPass(GraphBuilder, RDG_EVENT_NAME("CameraVolumeLut"), PassFlag, ComputeShader, PassParameters, NumGroups);
+
+			PendingRDGResources.ViewResources[ViewIndex].SkyAtmosphereCameraAerialPerspectiveVolume = SkyAtmosphereCameraAerialPerspectiveVolume;
 		}
 
 		PendingRDGResources.ViewResources[ViewIndex].SkyAtmosphereViewLutTexture = SkyAtmosphereViewLutTexture;
-		PendingRDGResources.ViewResources[ViewIndex].SkyAtmosphereCameraAerialPerspectiveVolume = SkyAtmosphereCameraAerialPerspectiveVolume;
 	}
 
 	PendingRDGResources.TransmittanceLut = TransmittanceLut;
@@ -1771,7 +1781,15 @@ void FSkyAtmospherePendingRDGResources::CommitToSceneAndViewUniformBuffers(FRDGB
 	{
 		FViewInfo& View = SceneRenderer->Views[ViewIndex];
 		View.SkyAtmosphereViewLutTexture = ConvertToExternalAccessTexture(GraphBuilder, ExternalAccessQueue, ViewResources[ViewIndex].SkyAtmosphereViewLutTexture, ERHIAccess::SRVMask, ERHIPipeline::All);
-		View.SkyAtmosphereCameraAerialPerspectiveVolume = ConvertToExternalAccessTexture(GraphBuilder, ExternalAccessQueue, ViewResources[ViewIndex].SkyAtmosphereCameraAerialPerspectiveVolume, ERHIAccess::SRVMask, ERHIPipeline::All);
+		
+		if (ViewResources[ViewIndex].SkyAtmosphereCameraAerialPerspectiveVolume)
+		{
+			View.SkyAtmosphereCameraAerialPerspectiveVolume = ConvertToExternalAccessTexture(GraphBuilder, ExternalAccessQueue, ViewResources[ViewIndex].SkyAtmosphereCameraAerialPerspectiveVolume, ERHIAccess::SRVMask, ERHIPipeline::All);
+		}
+		else
+		{
+			View.SkyAtmosphereCameraAerialPerspectiveVolume = nullptr;
+		}
 	}
 
 	check(TransmittanceLut);
