[GPU Scene Lights] Only update modified lights.
- Also added validation.

diff --git a/Engine/Source/Runtime/Renderer/Private/GPUScene.cpp b/Engine/Source/Runtime/Renderer/Private/GPUScene.cpp
--- a/Engine/Source/Runtime/Renderer/Private/GPUScene.cpp
+++ b/Engine/Source/Runtime/Renderer/Private/GPUScene.cpp
@@ -726,7 +726,7 @@ void FGPUScene::EndRender()
 	CachedRegisteredBuffers = {};
 }
 
-void FGPUScene::UpdateGPULights(FRDGBuilder& GraphBuilder, const UE::Tasks::FTask& PrerequisiteTask)
+void FGPUScene::UpdateGPULights(FRDGBuilder& GraphBuilder, const UE::Tasks::FTask& PrerequisiteTask, const FLightSceneChangeSet& LightsPostUpdateData)
 {
 #if !UE_BUILD_SHIPPING
 	MaxLightsDuringPrevUpdate = uint32(Scene.Lights.GetMaxIndex());
@@ -735,44 +735,82 @@ void FGPUScene::UpdateGPULights(FRDGBuilder& GraphBuilder, const UE::Tasks::FTas
 	// Clamp buffer to be smaller than the MAX_LIGHT_ID.
 	NumGPULights = FMath::Min(int32(MAX_LIGHT_ID), Scene.Lights.GetMaxIndex());
 
-	auto* LightDataUploader = TAsyncByteAddressBufferScatterUploader<FLightSceneData>::Allocate(GraphBuilder, NumGPULights);
+#if DO_CHECK
+	LightDataArray.SetNum(NumGPULights);
+#endif
 
-	const bool bAsyncSetup = CVarGPUSceneLightsAsyncSetup.GetValueOnRenderThread() != 0;
+	auto* LightDataUploader = TAsyncByteAddressBufferScatterUploader<FLightSceneData>::Allocate(GraphBuilder, LightsPostUpdateData.SceneLightInfoUpdates->NumCommands());
 
-	GraphBuilder.AddSetupTask([this, LightDataUploader]
-	{
-		SCOPED_NAMED_EVENT(UpdateGPUScene_Lights, FColor::Green);
-		const bool bAllowStaticLighting = IsStaticLightingAllowed();
-		const uint32 Flags = RenderRectLightsAsSpotLights(FeatureLevel) ? ELightShaderParameterFlags::RectAsSpotLight : 0u;
+	const bool bAsyncSetup = CVarGPUSceneLightsAsyncSetup.GetValueOnRenderThread() != 0;
 
-		for (int32 Index = 0; Index < NumGPULights; ++Index)
+	LightDataSetupTask = GraphBuilder.AddSetupTask([this, LightDataUploader, ChangeSet = LightsPostUpdateData]
 		{
-			FLightSceneData LightData;
+			SCOPED_NAMED_EVENT(UpdateGPUScene_Lights, FColor::Green);
+			const bool bAllowStaticLighting = IsStaticLightingAllowed();
+			const uint32 Flags = RenderRectLightsAsSpotLights(FeatureLevel) ? ELightShaderParameterFlags::RectAsSpotLight : 0u;
 
-			if (Scene.Lights.IsAllocated(Index))
-			{
-				InitLightData(Scene.Lights[Index], bAllowStaticLighting, Flags, LightData);
-			}
-			else
+			ChangeSet.SceneLightInfoUpdates->ForEachCommand(ESceneUpdateCommandFilter::AddedUpdated, [&](const FUpdateLightCommand& UpdateLightCommand)
+				{
+					int32 Index = UpdateLightCommand.GetSceneInfo()->Id;
+
+					if (!ensureMsgf(Index < NumGPULights, TEXT("Index is out-of-range. Either Light Id = %d is incorrect or MAX_LIGHT_ID was exceeded."), Index))
+					{
+						return;
+					}
+
+					const FLightSceneInfoCompact& LightInfoCompact = Scene.Lights[Index];
+
+					FLightSceneData LightData;
+					InitLightData(LightInfoCompact, bAllowStaticLighting, Flags, LightData);
+
+					LightDataUploader->Add(LightData, Index);
+
+#if DO_CHECK
+					LightDataArray[Index] = LightData;
+#endif
+				});
+
+#if DO_CHECK
+			for (int32 Index = 0; Index < NumGPULights; ++Index)
 			{
-				LightData.WorldPosition = FDFVector3{};
-				LightData.InvRadius = 0.0f;
-			}
+				FLightSceneData LightData;
 
-			LightDataUploader->Add(LightData, Index);
-		}
-	},
-	PrerequisiteTask,
-	UE::Tasks::ETaskPriority::Normal,
-	bAsyncSetup);
+				if (Scene.Lights.IsAllocated(Index))
+				{
+					const FLightSceneInfoCompact& LightInfoCompact = Scene.Lights[Index];
+
+					InitLightData(LightInfoCompact, bAllowStaticLighting, Flags, LightData);
+
+					ensureMsgf(FMemory::Memcmp(&LightData, &LightDataArray[Index], sizeof(LightData)) == 0, TEXT("Data for Id = %d in GPU Scene Lights is stale."), Index);
+				}
+				else
+				{
+					LightData.WorldPosition = FDFVector3{};
+					LightData.InvRadius = 0.0f;
+				}
 
-	if (bAsyncSetup)
+				LightDataArray[Index] = LightData;
+			}
+#endif
+		},
+		MakeArrayView({PrerequisiteTask, LightDataSetupTask}), // chain with previous setup task to prevent race condition accessing LightDataArray
+		UE::Tasks::ETaskPriority::Normal,
+		bAsyncSetup);
+
+	if (!LightsPostUpdateData.SceneLightInfoUpdates->IsEmpty())
 	{
-		LightDataUploader->ResizeAndUploadToAsync(GraphBuilder, LightDataBuffer, NumGPULights);
+		if (bAsyncSetup)
+		{
+			LightDataUploader->ResizeAndUploadToAsync(GraphBuilder, LightDataBuffer, NumGPULights);
+		}
+		else
+		{
+			LightDataUploader->ResizeAndUploadTo(GraphBuilder, LightDataBuffer, NumGPULights);
+		}
 	}
 	else
 	{
-		LightDataUploader->ResizeAndUploadTo(GraphBuilder, LightDataBuffer, NumGPULights);
+		LightDataBuffer.ResizeBufferIfNeeded(GraphBuilder, NumGPULights);
 	}
 }
 
@@ -2164,16 +2202,7 @@ void FGPUScene::OnPostSceneUpdate(FRDGBuilder& GraphBuilder, const FScenePostUpd
 
 void FGPUScene::OnPostLightSceneInfoUpdate(FRDGBuilder& GraphBuilder, const FLightSceneChangeSet& LightsPostUpdateData)
 {
-	// TODO: implement proper dirty tracking such that we can actually do incremental updates without loosing information
-	// const bool bAnythingChanged = !(LightsPostUpdateData.AddedLightIds.IsEmpty() && LightsPostUpdateData.RemovedLightIds.IsEmpty() && LightsPostUpdateData.TransformUpdatedLightIds.IsEmpty() && LightsPostUpdateData.ColorUpdatedLightIds.IsEmpty());
-	
-	// SceneFrameNumber is updated in UpdateInternal so if it is the same, we have already uploaded the lights this "frame" - this is not 100% robust since 
-	// if UpdateAllPrimitiveSceneInfos is called multiple times without updating the FScene frame number it will skip light uploads. The real solution is 
-	// to implement functional change tracking for lights so we can actually know when there are relevant changes.
-	if (SceneFrameNumber != Scene.GetFrameNumberRenderThread())
-	{
-		UpdateGPULights(GraphBuilder, UE::Tasks::FTask{});
-	}
+	UpdateGPULights(GraphBuilder, UE::Tasks::FTask{}, LightsPostUpdateData);
 }
 
 SIZE_T FGPUScene::GetAllocatedSize() const
diff --git a/Engine/Source/Runtime/Renderer/Private/GPUScene.h b/Engine/Source/Runtime/Renderer/Private/GPUScene.h
--- a/Engine/Source/Runtime/Renderer/Private/GPUScene.h
+++ b/Engine/Source/Runtime/Renderer/Private/GPUScene.h
@@ -372,6 +372,11 @@ private:
 
 	TPersistentByteAddressBuffer<FLightSceneData> LightDataBuffer;
 	int32 NumGPULights = 0;
+	UE::Tasks::FTask LightDataSetupTask;
+
+#if DO_CHECK
+	TArray<FLightSceneData> LightDataArray;
+#endif
 
 	FGPUSceneCommonParameters CommonParameters;
 
@@ -468,7 +473,7 @@ private:
 	/**
 	 * Upload scene light data to gpu
 	 */
-	void UpdateGPULights(FRDGBuilder& GraphBuilder, const UE::Tasks::FTask& PrerequisiteTask);
+	void UpdateGPULights(FRDGBuilder& GraphBuilder, const UE::Tasks::FTask& PrerequisiteTask, const FLightSceneChangeSet& LightsPostUpdateData);
 
 	static void InitLightData(const FLightSceneInfoCompact& LightInfoCompact, bool bAllowStaticLighting, uint32 LightShaderParameterFlags, FLightSceneData& DataOut);
 
