Ray Tracing: Various cleanups to ray tracing API

Hide the FRayTracingMaskAndFlags struct internally as it no longer is used outside the rendering module and has been deprecated since 5.4

Unify the ERayTracingViewMaskMode and ERayTracingMeshCommandsMode enums to a single ERayTracingType enum as they describe the same thing.

Hide ERayTracingInstanceMaskType inside RayTracingInstanceMask.cpp as it does not need to be public. The one external use has been changed to a small utility function so we can keep it private.

Make the internals of FRayTracingMeshProcessor private and devirtualize the class since it no longer contains any virtual methods. Update the FeatureLevel when no FScene is present to match what GPULightmass needs (this is the only case that uses FRayTracingMeshProcessor without a valid FScene)

diff --git a/Engine/Plugins/Experimental/GPULightmass/Source/GPULightmass/Private/LightmapRayTracing.cpp b/Engine/Plugins/Experimental/GPULightmass/Source/GPULightmass/Private/LightmapRayTracing.cpp
--- a/Engine/Plugins/Experimental/GPULightmass/Source/GPULightmass/Private/LightmapRayTracing.cpp
+++ b/Engine/Plugins/Experimental/GPULightmass/Source/GPULightmass/Private/LightmapRayTracing.cpp
@@ -11,9 +11,8 @@ bool IsSupportedVertexFactoryType(const FVertexFactoryType* VertexFactoryType)
 #if RHI_RAYTRACING
 
 FLightmapRayTracingMeshProcessor::FLightmapRayTracingMeshProcessor(FRayTracingMeshCommandContext* InCommandContext)
-	: FRayTracingMeshProcessor(InCommandContext, nullptr, nullptr, ERayTracingMeshCommandsMode::LIGHTMAP_TRACING)
+	: FRayTracingMeshProcessor(InCommandContext, nullptr, nullptr, ERayTracingType::LightMapTracing)
 {
-	FeatureLevel = GMaxRHIFeatureLevel;
 }
 
 IMPLEMENT_GLOBAL_SHADER(FLightmapPathTracingRGS, "/Plugin/GPULightmass/Private/LightmapPathTracing.usf", "LightmapPathTracingMainRG", SF_RayGen);
diff --git a/Engine/Plugins/Experimental/GPULightmass/Source/GPULightmass/Private/LightmapRenderer.cpp b/Engine/Plugins/Experimental/GPULightmass/Source/GPULightmass/Private/LightmapRenderer.cpp
--- a/Engine/Plugins/Experimental/GPULightmass/Source/GPULightmass/Private/LightmapRenderer.cpp
+++ b/Engine/Plugins/Experimental/GPULightmass/Source/GPULightmass/Private/LightmapRenderer.cpp
@@ -37,7 +37,7 @@
 #include "PrimitiveUniformShaderParametersBuilder.h"
 #include "PrimitiveSceneShaderData.h"
 
-RENDERER_API uint8 BlendModeToRayTracingInstanceMask(const EBlendMode BlendMode, bool bIsDitherMasked, bool bCastShadow, ERayTracingViewMaskMode MaskMode);
+RENDERER_API uint8 BlendModeToRayTracingInstanceMask(const EBlendMode BlendMode, bool bIsDitherMasked, bool bCastShadow, ERayTracingType RayTracingType);
 
 class FCopyConvergedLightmapTilesCS : public FGlobalShader
 {
@@ -629,7 +629,7 @@ struct FRayTracingMaskAndStatus
 	bool bAllSegmentsOpaque = true;
 	bool bAnySegmentsCastShadow = false;
 
-	void UpdateInstanceMaskAndStatus(ERHIFeatureLevel::Type FeatureLevel, ERayTracingViewMaskMode MaskMode, TArray<FMeshBatch>& MeshBatches)
+	void UpdateInstanceMaskAndStatus(ERHIFeatureLevel::Type FeatureLevel, TArray<FMeshBatch>& MeshBatches)
 	{
 		for (int32 SegmentIndex = 0; SegmentIndex < MeshBatches.Num(); SegmentIndex++)
 		{
@@ -643,7 +643,7 @@ struct FRayTracingMaskAndStatus
 			bAllSegmentsUnlit &= Material.GetShadingModels().HasOnlyShadingModel(MSM_Unlit) || !MeshBatch.CastShadow;
 			bAllSegmentsOpaque &= Material.GetBlendMode() == EBlendMode::BLEND_Opaque;
 			bAnySegmentsCastShadow |= bSegmentCastsShadow;
-			InstanceMask |= BlendModeToRayTracingInstanceMask(Material.GetBlendMode(), Material.IsDitherMasked(), bSegmentCastsShadow, MaskMode);
+			InstanceMask |= BlendModeToRayTracingInstanceMask(Material.GetBlendMode(), Material.IsDitherMasked(), bSegmentCastsShadow, ERayTracingType::LightMapTracing);
 		}
 	}
 };
@@ -699,7 +699,7 @@ void FCachedRayTracingSceneData::SetupFromSceneRenderState(FSceneRenderState& Sc
 			TArray<FMeshBatch> MeshBatches = Instance.GetMeshBatchesForGBufferRendering(LODIndexToUse);
 
 			FRayTracingMaskAndStatus RayTracingMaskAndStatus;
-			RayTracingMaskAndStatus.UpdateInstanceMaskAndStatus(Scene.FeatureLevel, ERayTracingViewMaskMode::LightMapTracing, MeshBatches);
+			RayTracingMaskAndStatus.UpdateInstanceMaskAndStatus(Scene.FeatureLevel, MeshBatches);
 
 			if (!RayTracingMaskAndStatus.bAllSegmentsUnlit)
 			{
@@ -744,7 +744,7 @@ void FCachedRayTracingSceneData::SetupFromSceneRenderState(FSceneRenderState& Sc
 			TArray<FMeshBatch> MeshBatches = InstanceGroup.GetMeshBatchesForGBufferRendering(LODIndexToUse, FTileVirtualCoordinates{});
 
 			FRayTracingMaskAndStatus RayTracingMaskAndStatus;
-			RayTracingMaskAndStatus.UpdateInstanceMaskAndStatus(Scene.FeatureLevel, ERayTracingViewMaskMode::LightMapTracing, MeshBatches);
+			RayTracingMaskAndStatus.UpdateInstanceMaskAndStatus(Scene.FeatureLevel, MeshBatches);
 
 			if (!RayTracingMaskAndStatus.bAllSegmentsUnlit)
 			{
@@ -1129,7 +1129,7 @@ bool FSceneRenderState::SetupRayTracingScene(FRDGBuilder& GraphBuilder, FSceneUn
 						}
 
 						FRayTracingMaskAndStatus RayTracingMaskAndStatus;
-						RayTracingMaskAndStatus.UpdateInstanceMaskAndStatus(FeatureLevel, ERayTracingViewMaskMode::LightMapTracing, MeshBatches);
+						RayTracingMaskAndStatus.UpdateInstanceMaskAndStatus(FeatureLevel, MeshBatches);
 
 						if (RayTracingMaskAndStatus.bAllSegmentsUnlit)
 						{
diff --git a/Engine/Source/Runtime/Renderer/Private/PathTracing.cpp b/Engine/Source/Runtime/Renderer/Private/PathTracing.cpp
--- a/Engine/Source/Runtime/Renderer/Private/PathTracing.cpp
+++ b/Engine/Source/Runtime/Renderer/Private/PathTracing.cpp
@@ -2256,9 +2256,9 @@ bool FRayTracingMeshProcessor::ProcessPathTracing(
 	{
 		const bool bUseProceduralPrimitive = MeshBatch.VertexFactory->GetType()->SupportsRayTracingProceduralPrimitive() &&
 			FDataDrivenShaderPlatformInfo::GetSupportsRayTracingProceduralPrimitive(GMaxRHIShaderPlatform);
-		switch (RayTracingMeshCommandsMode)
+		switch (RayTracingType)
 		{
-			case ERayTracingMeshCommandsMode::PATH_TRACING:
+			case ERayTracingType::PathTracing:
 			{
 				// In order to use Substrate simplified materials, Substrate has to be enabled, we have to have _compiled_ the extra permutations _and_ the runtime toggle must be true
 				const bool bUseSimplifiedMaterial = Substrate::IsSubstrateEnabled() &&
@@ -2300,7 +2300,7 @@ bool FRayTracingMeshProcessor::ProcessPathTracing(
 				}
 				break;
 			}
-			case ERayTracingMeshCommandsMode::LIGHTMAP_TRACING:
+			case ERayTracingType::LightMapTracing:
 			{
 				if (NeedsAnyHitShader(MaterialResource))
 				{
@@ -2341,8 +2341,7 @@ bool FRayTracingMeshProcessor::ProcessPathTracing(
 		MaterialRenderProxy,
 		MaterialResource,
 		RayTracingShader,
-		ShaderElementData,
-		ERayTracingViewMaskMode::PathTracing);
+		ShaderElementData);
 
 	return true;
 }
diff --git a/Engine/Source/Runtime/Renderer/Private/PrimitiveSceneInfo.cpp b/Engine/Source/Runtime/Renderer/Private/PrimitiveSceneInfo.cpp
--- a/Engine/Source/Runtime/Renderer/Private/PrimitiveSceneInfo.cpp
+++ b/Engine/Source/Runtime/Renderer/Private/PrimitiveSceneInfo.cpp
@@ -1005,7 +1005,7 @@ class FCacheRayTracingPrimitivesContext
 public:
 	FCacheRayTracingPrimitivesContext(FScene* Scene)
 		: CommandContext(Commands)
-		, RayTracingMeshProcessor(&CommandContext, Scene, nullptr, Scene->CachedRayTracingMeshCommandsMode)
+		, RayTracingMeshProcessor(&CommandContext, Scene, nullptr, Scene->CachedRayTracingMeshCommandsType)
 	{ }
 
 	FTempRayTracingMeshCommandStorage Commands;
@@ -1151,13 +1151,14 @@ void CacheRayTracingPrimitive(
 		}
 		
 		// Setup the instance mask on the LOD data
-		const ERayTracingViewMaskMode MaskMode = static_cast<ERayTracingViewMaskMode>(Scene->CachedRayTracingMeshCommandsMode);
+		// TODO: do we really need to set the InstanceMask to far field here? Can it not be set through the other codepaths that setup InstanceMask?
 		if (EnumHasAllFlags(OutFlags, ERayTracingPrimitiveFlags::FarField))
 		{
+			const ERayTracingType RayTracingType = Scene->CachedRayTracingMeshCommandsType;
 			for (int32 LODIndex = 0; LODIndex < LODCount; ++LODIndex)
 			{
 				FPrimitiveSceneInfo::FRayTracingLODData& LODData = RayTracingLODData[LODIndex];
-				LODData.CachedMeshCommandFlags.InstanceMask = ComputeRayTracingInstanceMask(ERayTracingInstanceMaskType::FarField, MaskMode);
+				LODData.CachedMeshCommandFlags.InstanceMask = GetFarFieldRayTracingInstanceMask(RayTracingType);
 			}
 		}
 
@@ -1228,7 +1229,7 @@ void FPrimitiveSceneInfo::CacheRayTracingPrimitives(FScene* Scene, const TArrayV
 		else
 		{
 			FCachedRayTracingMeshCommandContext CommandContext(CachedRayTracingMeshCommands);
-			FRayTracingMeshProcessor RayTracingMeshProcessor(&CommandContext, Scene, nullptr, Scene->CachedRayTracingMeshCommandsMode);
+			FRayTracingMeshProcessor RayTracingMeshProcessor(&CommandContext, Scene, nullptr, Scene->CachedRayTracingMeshCommandsType);
 
 			for (FPrimitiveSceneInfo* SceneInfo : SceneInfos)
 			{
diff --git a/Engine/Source/Runtime/Renderer/Private/RayTracing/RayTracing.cpp b/Engine/Source/Runtime/Renderer/Private/RayTracing/RayTracing.cpp
--- a/Engine/Source/Runtime/Renderer/Private/RayTracing/RayTracing.cpp
+++ b/Engine/Source/Runtime/Renderer/Private/RayTracing/RayTracing.cpp
@@ -282,14 +282,14 @@ namespace RayTracing
 
 	void OnRenderBegin(FScene& Scene, TArray<FViewInfo>& Views, const FViewFamilyInfo& ViewFamily)
 	{
-		const ERayTracingMeshCommandsMode CurrentMode = ViewFamily.EngineShowFlags.PathTracing ? ERayTracingMeshCommandsMode::PATH_TRACING : ERayTracingMeshCommandsMode::RAY_TRACING;
+		const ERayTracingType CurrentType = ViewFamily.EngineShowFlags.PathTracing ? ERayTracingType::PathTracing : ERayTracingType::RayTracing;
 		bool bNaniteCoarseMeshStreamingModeChanged = false;
 #if WITH_EDITOR
 		bNaniteCoarseMeshStreamingModeChanged = Nanite::FCoarseMeshStreamingManager::CheckStreamingMode();
 #endif // WITH_EDITOR
 		const bool bNaniteRayTracingModeChanged = Nanite::GRayTracingManager.CheckModeChanged();
 
-		if (CurrentMode != Scene.CachedRayTracingMeshCommandsMode
+		if (CurrentType != Scene.CachedRayTracingMeshCommandsType
 			|| bNaniteCoarseMeshStreamingModeChanged
 			|| bNaniteRayTracingModeChanged
 			|| bUpdateCachedRayTracingState)
@@ -297,8 +297,8 @@ namespace RayTracing
 			Scene.WaitForCacheRayTracingPrimitivesTask();
 
 			// In some situations, we need to refresh the cached ray tracing mesh commands because they contain data about the currently bound shader. 
-			// This operation is a bit expensive but only happens once as we transition between modes which should be rare.
-			Scene.CachedRayTracingMeshCommandsMode = CurrentMode;
+			// This operation is a bit expensive but only happens once as we transition between RT types which should be rare.
+			Scene.CachedRayTracingMeshCommandsType = CurrentType;
 			Scene.RefreshRayTracingMeshCommandCache();
 			bUpdateCachedRayTracingState = false;
 		}
@@ -310,7 +310,7 @@ namespace RayTracing
 			// In some situations, we need to refresh the cached ray tracing instance.
 			// This assumes that cached instances will keep using the same LOD since CachedRayTracingMeshCommands is not recalculated
 			// eg: Need to update PrimitiveRayTracingFlags
-			// This operation is a bit expensive but only happens once as we transition between modes which should be rare.
+			// This operation is a bit expensive but only happens once as we transition between RT types which should be rare.
 			Scene.RefreshRayTracingInstances();
 			bRefreshRayTracingInstances = false;
 		}
@@ -452,7 +452,7 @@ namespace RayTracing
 							FDynamicRayTracingMeshCommandContext CommandContext(
 								*TaskDynamicCommandStorage, *TaskVisibleShaderBindings,
 								WorkItem.RayTracingGeometry, SegmentIndex, WorkItem.SBTAllocation);
-							FRayTracingMeshProcessor RayTracingMeshProcessor(&CommandContext, &Scene, &View, Scene.CachedRayTracingMeshCommandsMode);
+							FRayTracingMeshProcessor RayTracingMeshProcessor(&CommandContext, &Scene, &View, Scene.CachedRayTracingMeshCommandsType);
 							RayTracingMeshProcessor.AddMeshBatch(MeshBatch, 1, WorkItem.SceneProxy);
 						}
 					}
@@ -869,7 +869,7 @@ namespace RayTracing
 						{
 							const FMeshBatch& MeshBatch = InstanceMaterials[SegmentIndex];
 							FDynamicRayTracingMeshCommandContext CommandContext(View.DynamicRayTracingMeshCommandStorage, View.VisibleRayTracingShaderBindings, Geometry->GetRHI(), SegmentIndex, SBTAllocation);
-							FRayTracingMeshProcessor RayTracingMeshProcessor(&CommandContext, &Scene, &View, Scene.CachedRayTracingMeshCommandsMode);
+							FRayTracingMeshProcessor RayTracingMeshProcessor(&CommandContext, &Scene, &View, Scene.CachedRayTracingMeshCommandsType);
 							RayTracingMeshProcessor.AddMeshBatch(MeshBatch, 1, SceneProxy);
 						}
 					}
diff --git a/Engine/Source/Runtime/Renderer/Private/RayTracing/RayTracingInstanceMask.cpp b/Engine/Source/Runtime/Renderer/Private/RayTracing/RayTracingInstanceMask.cpp
--- a/Engine/Source/Runtime/Renderer/Private/RayTracing/RayTracingInstanceMask.cpp
+++ b/Engine/Source/Runtime/Renderer/Private/RayTracing/RayTracingInstanceMask.cpp
@@ -12,10 +12,37 @@
 #include "MeshPassProcessor.h"
 #include "ScenePrivate.h"
 
+namespace {
 
-uint8 ComputeRayTracingInstanceMask(ERayTracingInstanceMaskType MaskType, ERayTracingViewMaskMode MaskMode)
+// Helper class so we can refer to both RAY_TRACING_MASK_* and PATH_TRACING_MASK_* in a unified way within this file
+enum class ERayTracingInstanceMaskType : uint8
 {
-	if (MaskMode == ERayTracingViewMaskMode::RayTracing)
+	// General mask type for primary and secondary rays
+	Opaque,
+	Translucent,
+
+	// Mask types for shadow rays
+	OpaqueShadow,
+	TranslucentShadow,
+	ThinShadow,
+
+	// Geometry specific ray types
+	HairStrands,
+
+	// Special purpose ray types
+	FirstPersonWorldSpaceRepresentation,
+	FarField,
+
+	// path tracing specific mask type
+	VisibleInPrimaryRay,
+	VisibleInIndirectRay
+};
+
+} // anonymous namespace
+
+static uint8 ComputeRayTracingInstanceMask(ERayTracingInstanceMaskType MaskType, ERayTracingType RayTracingType)
+{
+	if (RayTracingType == ERayTracingType::RayTracing)
 	{
 		switch (MaskType)
 		{
@@ -43,8 +70,8 @@ uint8 ComputeRayTracingInstanceMask(ERayTracingInstanceMaskType MaskType, ERayTr
 			break;
 		}
 	}
-	else if (MaskMode == ERayTracingViewMaskMode::PathTracing ||
-		     MaskMode == ERayTracingViewMaskMode::LightMapTracing)
+	else if (RayTracingType == ERayTracingType::PathTracing ||
+			 RayTracingType == ERayTracingType::LightMapTracing)
 	{
 		switch (MaskType)
 		{
@@ -76,24 +103,30 @@ uint8 ComputeRayTracingInstanceMask(ERayTracingInstanceMaskType MaskType, ERayTr
 	return 0;
 }
 
-uint8 BlendModeToRayTracingInstanceMask(const EBlendMode BlendMode, bool bIsDitherMasked, bool bCastShadow, ERayTracingViewMaskMode MaskMode)
+uint8 GetFarFieldRayTracingInstanceMask(ERayTracingType RayTracingType)
+{
+	return ComputeRayTracingInstanceMask(ERayTracingInstanceMaskType::FarField, RayTracingType);
+}
+
+uint8 BlendModeToRayTracingInstanceMask(const EBlendMode BlendMode, bool bIsDitherMasked, bool bCastShadow, ERayTracingType RayTracingType)
 {
 	uint8 InstanceMask = 0;
 
 	if (IsOpaqueOrMaskedBlendMode(BlendMode) && !bIsDitherMasked)
 	{
-		InstanceMask |= ComputeRayTracingInstanceMask(ERayTracingInstanceMaskType::Opaque, MaskMode);
-		InstanceMask |= bCastShadow ? ComputeRayTracingInstanceMask(ERayTracingInstanceMaskType::OpaqueShadow, MaskMode) : 0;
+		InstanceMask |= ComputeRayTracingInstanceMask(ERayTracingInstanceMaskType::Opaque, RayTracingType);
+		InstanceMask |= bCastShadow ? ComputeRayTracingInstanceMask(ERayTracingInstanceMaskType::OpaqueShadow, RayTracingType) : 0;
 	}
 	else
 	{
-		InstanceMask |= ComputeRayTracingInstanceMask(ERayTracingInstanceMaskType::Translucent, MaskMode);
-		InstanceMask |= bCastShadow ? ComputeRayTracingInstanceMask(ERayTracingInstanceMaskType::TranslucentShadow, MaskMode) : 0;
+		InstanceMask |= ComputeRayTracingInstanceMask(ERayTracingInstanceMaskType::Translucent, RayTracingType);
+		InstanceMask |= bCastShadow ? ComputeRayTracingInstanceMask(ERayTracingInstanceMaskType::TranslucentShadow, RayTracingType) : 0;
 	}
 
 	return InstanceMask;
 }
 
+namespace {
 
 /** Util struct and function to derive mask related info from scene proxy*/
 struct FSceneProxyRayTracingMaskInfo
@@ -102,22 +135,24 @@ struct FSceneProxyRayTracingMaskInfo
 	bool bVisibleToShadow = false;
 	bool bVisibleToIndirect = false;
 	bool bIsFirstPersonWorldSpaceRepresentation = false;
-	ERayTracingViewMaskMode MaskMode = ERayTracingViewMaskMode::RayTracing;
+	ERayTracingType RayTracingType = ERayTracingType::RayTracing;
 };
 
+} // anonymous namespace
+
 static FSceneProxyRayTracingMaskInfo GetSceneProxyRayTracingMaskInfo(const FPrimitiveSceneProxy& PrimitiveSceneProxy)
 {
 	FSceneProxyRayTracingMaskInfo MaskInfo = {};
 
 	const FScene* RenderScene = PrimitiveSceneProxy.GetScene().GetRenderScene();
-	MaskInfo.MaskMode = static_cast<ERayTracingViewMaskMode>(RenderScene->CachedRayTracingMeshCommandsMode);
+	MaskInfo.RayTracingType = RenderScene->CachedRayTracingMeshCommandsType;
 
 	if (PrimitiveSceneProxy.IsDrawnInGame())
 	{
 		MaskInfo.bVisibleToCamera = true;
 		MaskInfo.bVisibleToShadow = true;
 		// NOTE: For backwards compatibility, only path tracing obeys the AffectsDynamicIndirectLighting flag
-		MaskInfo.bVisibleToIndirect = MaskInfo.MaskMode == ERayTracingViewMaskMode::RayTracing ? true : PrimitiveSceneProxy.AffectsDynamicIndirectLighting();
+		MaskInfo.bVisibleToIndirect = MaskInfo.RayTracingType == ERayTracingType::RayTracing ? true : PrimitiveSceneProxy.AffectsDynamicIndirectLighting();
 	}
 	else
 	{
@@ -131,15 +166,15 @@ static FSceneProxyRayTracingMaskInfo GetSceneProxyRayTracingMaskInfo(const FPrim
 	return MaskInfo;
 }
 
-static uint8 ApplyFirstPersonRayTracingInstanceMaskFlag(uint8 Mask, ERayTracingViewMaskMode MaskMode, bool bIsFirstPersonWorldSpaceRepresentation)
+static uint8 ApplyFirstPersonRayTracingInstanceMaskFlag(uint8 Mask, ERayTracingType RayTracingType, bool bIsFirstPersonWorldSpaceRepresentation)
 {
 	if (bIsFirstPersonWorldSpaceRepresentation)
 	{
-		const bool bIsOpaque = (Mask & (ComputeRayTracingInstanceMask(ERayTracingInstanceMaskType::Opaque, MaskMode) | ComputeRayTracingInstanceMask(ERayTracingInstanceMaskType::OpaqueShadow, MaskMode))) != 0;
+		const bool bIsOpaque = (Mask & (ComputeRayTracingInstanceMask(ERayTracingInstanceMaskType::Opaque, RayTracingType) | ComputeRayTracingInstanceMask(ERayTracingInstanceMaskType::OpaqueShadow, RayTracingType))) != 0;
 
 		// Tag world space representations of first person meshes so rays originating from first person meshes can skip them.
 		// We currently only support opaque world space representations of first person objects, so set the mask to 0 otherwise.
-		Mask = bIsOpaque ? ComputeRayTracingInstanceMask(ERayTracingInstanceMaskType::FirstPersonWorldSpaceRepresentation, MaskMode) : 0;
+		Mask = bIsOpaque ? ComputeRayTracingInstanceMask(ERayTracingInstanceMaskType::FirstPersonWorldSpaceRepresentation, RayTracingType) : 0;
 	}
 	return Mask;
 }
@@ -148,7 +183,7 @@ FRayTracingMaskAndFlags BuildRayTracingInstanceMaskAndFlags(const FRayTracingIns
 {
 	FSceneProxyRayTracingMaskInfo MaskInfo = GetSceneProxyRayTracingMaskInfo(PrimitiveSceneProxy);
 	const ERHIFeatureLevel::Type FeatureLevel = PrimitiveSceneProxy.GetScene().GetFeatureLevel();
-	const ERayTracingViewMaskMode MaskMode = MaskInfo.MaskMode;
+	const ERayTracingType RayTracingType = MaskInfo.RayTracingType;
 
 	FRayTracingMaskAndFlags Result;
 
@@ -173,7 +208,7 @@ FRayTracingMaskAndFlags BuildRayTracingInstanceMaskAndFlags(const FRayTracingIns
 			const EBlendMode BlendMode = Material.GetBlendMode();
 			const bool bSegmentCastsShadow = MaskInfo.bVisibleToShadow && MeshBatch.CastRayTracedShadow && Material.CastsRayTracedShadows() && BlendMode != BLEND_Additive;
 
-			Result.Mask |= BlendModeToRayTracingInstanceMask(BlendMode, Material.IsDitherMasked(), bSegmentCastsShadow, MaskMode);
+			Result.Mask |= BlendModeToRayTracingInstanceMask(BlendMode, Material.IsDitherMasked(), bSegmentCastsShadow, RayTracingType);
 			bAllSegmentsOpaque &= BlendMode == BLEND_Opaque;
 			bAllSegmentsTranslucent &= BlendMode == BLEND_Translucent;
 			bAnySegmentsCastShadow |= bSegmentCastsShadow;
@@ -197,45 +232,45 @@ FRayTracingMaskAndFlags BuildRayTracingInstanceMaskAndFlags(const FRayTracingIns
 	if (bIsHairStrands)
 	{
 		// Reset all hair strands bits "on"
-		Result.Mask = ComputeRayTracingInstanceMask(ERayTracingInstanceMaskType::HairStrands, MaskMode);
+		Result.Mask = ComputeRayTracingInstanceMask(ERayTracingInstanceMaskType::HairStrands, RayTracingType);
 	}
 
 	if (!MaskInfo.bVisibleToCamera)
 	{
 		// If the object is not visible to camera, remove all direct visibility bits.
-		Result.Mask &= ~ComputeRayTracingInstanceMask(ERayTracingInstanceMaskType::VisibleInPrimaryRay, MaskMode);
+		Result.Mask &= ~ComputeRayTracingInstanceMask(ERayTracingInstanceMaskType::VisibleInPrimaryRay, RayTracingType);
 	}
 
 	if (!MaskInfo.bVisibleToIndirect)
 	{
 		// If the object does not affect indirect lighting, remove all indirect bits.
-		Result.Mask &= ~ComputeRayTracingInstanceMask(ERayTracingInstanceMaskType::VisibleInIndirectRay, MaskMode);
+		Result.Mask &= ~ComputeRayTracingInstanceMask(ERayTracingInstanceMaskType::VisibleInIndirectRay, RayTracingType);
 	}
 
 	if (!bAnySegmentsCastShadow)
 	{
 		// Not casting shadows, remove any set shadow flags
 		Result.Mask &= ~(
-			ComputeRayTracingInstanceMask(ERayTracingInstanceMaskType::OpaqueShadow, MaskMode) |
-			ComputeRayTracingInstanceMask(ERayTracingInstanceMaskType::TranslucentShadow, MaskMode) |
-			ComputeRayTracingInstanceMask(ERayTracingInstanceMaskType::ThinShadow, MaskMode)
+			ComputeRayTracingInstanceMask(ERayTracingInstanceMaskType::OpaqueShadow, RayTracingType) |
+			ComputeRayTracingInstanceMask(ERayTracingInstanceMaskType::TranslucentShadow, RayTracingType) |
+			ComputeRayTracingInstanceMask(ERayTracingInstanceMaskType::ThinShadow, RayTracingType)
 		);
 	}
 
 
-	Result.Mask = ApplyFirstPersonRayTracingInstanceMaskFlag(Result.Mask, MaskMode, MaskInfo.bIsFirstPersonWorldSpaceRepresentation);
+	Result.Mask = ApplyFirstPersonRayTracingInstanceMaskFlag(Result.Mask, RayTracingType, MaskInfo.bIsFirstPersonWorldSpaceRepresentation);
 
 	if (Instance.InstanceLayer == ERayTracingInstanceLayer::FarField)
 	{
 		// if far field, set that flag exclusively
-		Result.Mask = ComputeRayTracingInstanceMask(ERayTracingInstanceMaskType::FarField, MaskMode);
+		Result.Mask = ComputeRayTracingInstanceMask(ERayTracingInstanceMaskType::FarField, RayTracingType);
 	}
 
 	return Result;
 }
 
 void SetupRayTracingMeshCommandMaskAndStatus(FRayTracingMeshCommand& MeshCommand, const FMeshBatch& MeshBatch, const FPrimitiveSceneProxy* PrimitiveSceneProxy,
-	const FMaterial& MaterialResource, ERayTracingViewMaskMode MaskMode)
+	const FMaterial& MaterialResource, ERayTracingType RayTracingType)
 {
 	const FVertexFactory* VertexFactory = MeshBatch.VertexFactory;
 
@@ -248,7 +283,7 @@ void SetupRayTracingMeshCommandMaskAndStatus(FRayTracingMeshCommand& MeshCommand
 	MeshCommand.bIsTranslucent = MaterialResource.GetBlendMode() == EBlendMode::BLEND_Translucent;
 	MeshCommand.bReverseCulling = MeshBatch.ReverseCulling;
 
-	MeshCommand.InstanceMask = BlendModeToRayTracingInstanceMask(MaterialResource.GetBlendMode(), MaterialResource.IsDitherMasked(), MeshCommand.bCastRayTracedShadows, MaskMode);
+	MeshCommand.InstanceMask = BlendModeToRayTracingInstanceMask(MaterialResource.GetBlendMode(), MaterialResource.IsDitherMasked(), MeshCommand.bCastRayTracedShadows, RayTracingType);
 
 	if (!PrimitiveSceneProxy)
 	{
@@ -263,36 +298,36 @@ void SetupRayTracingMeshCommandMaskAndStatus(FRayTracingMeshCommand& MeshCommand
 
 	FSceneProxyRayTracingMaskInfo MaskInfo = GetSceneProxyRayTracingMaskInfo(*PrimitiveSceneProxy);
 
-	MeshCommand.InstanceMask = ApplyFirstPersonRayTracingInstanceMaskFlag(MeshCommand.InstanceMask, MaskMode, MaskInfo.bIsFirstPersonWorldSpaceRepresentation);
+	MeshCommand.InstanceMask = ApplyFirstPersonRayTracingInstanceMaskFlag(MeshCommand.InstanceMask, RayTracingType, MaskInfo.bIsFirstPersonWorldSpaceRepresentation);
 
 	// TODO: This should be done once all mesh commands for a mesh are combined (similar to BuildRayTracingInstanceMaskAndFlags above)
 	// TODO: Why is this logic not applied in the ray tracing case?
-	if (MaskMode == ERayTracingViewMaskMode::PathTracing ||
-		MaskMode == ERayTracingViewMaskMode::LightMapTracing)
+	if (RayTracingType == ERayTracingType::PathTracing ||
+		RayTracingType == ERayTracingType::LightMapTracing)
 	{
 	
 		if (!MaskInfo.bVisibleToCamera)
 		{
 			// If the object is not visible to camera, remove all direct visibility bits.
-			MeshCommand.InstanceMask &= ~ComputeRayTracingInstanceMask(ERayTracingInstanceMaskType::VisibleInPrimaryRay, MaskMode);
+			MeshCommand.InstanceMask &= ~ComputeRayTracingInstanceMask(ERayTracingInstanceMaskType::VisibleInPrimaryRay, RayTracingType);
 		}
 
 		if (!MaskInfo.bVisibleToIndirect)
 		{
 			// If the object does not affect indirect lighting, remove all indirect bits.
-			MeshCommand.InstanceMask &= ~ComputeRayTracingInstanceMask(ERayTracingInstanceMaskType::VisibleInIndirectRay, MaskMode);
+			MeshCommand.InstanceMask &= ~ComputeRayTracingInstanceMask(ERayTracingInstanceMaskType::VisibleInIndirectRay, RayTracingType);
 		}
 
 		if (!MaskInfo.bVisibleToShadow || !MeshCommand.bCastRayTracedShadows)
 		{
 			// Not casting shadows, remove any set shadow flags
 			MeshCommand.InstanceMask &= ~(
-				ComputeRayTracingInstanceMask(ERayTracingInstanceMaskType::OpaqueShadow, MaskMode) |
-				ComputeRayTracingInstanceMask(ERayTracingInstanceMaskType::TranslucentShadow, MaskMode) |
-				ComputeRayTracingInstanceMask(ERayTracingInstanceMaskType::ThinShadow, MaskMode));
+				ComputeRayTracingInstanceMask(ERayTracingInstanceMaskType::OpaqueShadow, RayTracingType) |
+				ComputeRayTracingInstanceMask(ERayTracingInstanceMaskType::TranslucentShadow, RayTracingType) |
+				ComputeRayTracingInstanceMask(ERayTracingInstanceMaskType::ThinShadow, RayTracingType));
 		}
 
 	}
 }
 
-#endif
\ No newline at end of file
+#endif
diff --git a/Engine/Source/Runtime/Renderer/Private/RayTracing/RayTracingInstanceMask.h b/Engine/Source/Runtime/Renderer/Private/RayTracing/RayTracingInstanceMask.h
--- a/Engine/Source/Runtime/Renderer/Private/RayTracing/RayTracingInstanceMask.h
+++ b/Engine/Source/Runtime/Renderer/Private/RayTracing/RayTracingInstanceMask.h
@@ -1,49 +1,56 @@
 // Copyright Epic Games, Inc. All Rights Reserved.
 #pragma once
 
-#include "RayTracingInstance.h"
+#include "CoreMinimal.h"
+#include "RHIDefinitions.h"
 
 #if RHI_RAYTRACING
 
+// Forward declarations used in this header
 struct FMeshBatch;
 class FPrimitiveSceneProxy;
 class FSceneViewFamily;
+struct FRayTracingInstance;
+class FRayTracingMeshCommand;
+class FMaterial;
+enum EBlendMode : int;
 
-enum class ERayTracingInstanceMaskType : uint8
+struct FRayTracingMaskAndFlags
 {
-	// General mask type for primary and secondary rays
-	Opaque,
-	Translucent,
+	/** Instance mask that can be used to exclude the instance from specific effects (eg. ray traced shadows). */
+	uint8 Mask = 0xFF;
 
-	// Mask types for shadow rays
-	OpaqueShadow,
-	TranslucentShadow,
-	ThinShadow,
+	/** Whether the instance is forced opaque, i.e. anyhit shaders are disabled on this instance. */
+	uint8 bForceOpaque : 1 = false;
 
-	// Geometry specific ray types
-	HairStrands,
+	/** Whether ray hits should be registered for front and back faces. */
+	uint8 bDoubleSided : 1 = false;
+	
+	/** Whether front and back facings should be reversed. */
+	uint8 bReverseCulling : 1 = false;
 
-	// Special purpose ray types
-	FirstPersonWorldSpaceRepresentation,
-	FarField,
+	/** Whether any or all of the segments in the instance are decals. */
+	uint8 bAnySegmentsDecal : 1 = false;
+	uint8 bAllSegmentsDecal : 1 = false;
 
-	// path tracing specific mask type
-	VisibleInPrimaryRay,
-	VisibleInIndirectRay
+	/** Whether all of the segments in the instance are translucent. */
+	uint8 bAllSegmentsTranslucent : 1 = false;
 };
 
-/** MeshCommands mode shares the same status for ray tracing view mask mode*/
-enum class ERayTracingViewMaskMode : uint8
+
+/** Describes what type of ray tracing we are doing. This is used to know which set of ray flags
+ * to use, and is also stored in the mesh command so we know when to invalidate them. */
+enum class ERayTracingType : uint8
 {
 	RayTracing,
 	PathTracing,
 	LightMapTracing,
 };
 
-uint8 ComputeRayTracingInstanceMask(ERayTracingInstanceMaskType MaskType, ERayTracingViewMaskMode MaskMode);
+uint8 GetFarFieldRayTracingInstanceMask(ERayTracingType RayTracingType);
 
 /** Compute the mask based on blend mode for different ray tracing mode*/
-RENDERER_API uint8 BlendModeToRayTracingInstanceMask(const EBlendMode BlendMode, bool bIsDitherMasked, bool bCastShadow, ERayTracingViewMaskMode MaskMode);
+RENDERER_API uint8 BlendModeToRayTracingInstanceMask(const EBlendMode BlendMode, bool bIsDitherMasked, bool bCastShadow, ERayTracingType RayTracingType);
 
 
 //-------------------------------------------------------
@@ -56,11 +63,7 @@ FRayTracingMaskAndFlags BuildRayTracingInstanceMaskAndFlags(const FRayTracingIns
 //-------------------------------------------------------
 //	FRayTracingMeshCommand related mask setup and update
 //-------------------------------------------------------
-class FRayTracingMeshCommand;
-enum class ERayTracingPrimitiveFlags : uint8;
-class FMaterial;
-
 void SetupRayTracingMeshCommandMaskAndStatus(FRayTracingMeshCommand& MeshCommand, const FMeshBatch& MeshBatch, const FPrimitiveSceneProxy* PrimitiveSceneProxy,
-	const FMaterial& MaterialResource, ERayTracingViewMaskMode MaskMode);
+	const FMaterial& MaterialResource, ERayTracingType RayTracingType);
 
 #endif
\ No newline at end of file
diff --git a/Engine/Source/Runtime/Renderer/Private/RayTracing/RayTracingMaterialHitShaders.cpp b/Engine/Source/Runtime/Renderer/Private/RayTracing/RayTracingMaterialHitShaders.cpp
--- a/Engine/Source/Runtime/Renderer/Private/RayTracing/RayTracingMaterialHitShaders.cpp
+++ b/Engine/Source/Runtime/Renderer/Private/RayTracing/RayTracingMaterialHitShaders.cpp
@@ -412,13 +412,13 @@ static bool GetRayTracingMeshProcessorShaders(
 	return true;
 }
 
-FRayTracingMeshProcessor::FRayTracingMeshProcessor(FRayTracingMeshCommandContext* InCommandContext, const FScene* InScene, const FSceneView* InViewIfDynamicMeshCommand, ERayTracingMeshCommandsMode InRayTracingMeshCommandsMode)
+FRayTracingMeshProcessor::FRayTracingMeshProcessor(FRayTracingMeshCommandContext* InCommandContext, const FScene* InScene, const FSceneView* InViewIfDynamicMeshCommand, ERayTracingType InRayTracingType)
 	:
 	CommandContext(InCommandContext),
 	Scene(InScene),
 	ViewIfDynamicMeshCommand(InViewIfDynamicMeshCommand),
-	FeatureLevel(InScene ? InScene->GetFeatureLevel() : ERHIFeatureLevel::SM5),
-	RayTracingMeshCommandsMode(InRayTracingMeshCommandsMode)
+	FeatureLevel(InScene ? InScene->GetFeatureLevel() : GMaxRHIFeatureLevel),
+	RayTracingType(InRayTracingType)
 {
 }
 
@@ -451,8 +451,7 @@ bool FRayTracingMeshProcessor::Process(
 		MaterialRenderProxy,
 		MaterialResource,
 		RayTracingShader,
-		ShaderElementData,
-		ERayTracingViewMaskMode::RayTracing);
+		ShaderElementData);
 
 	return true;
 }
@@ -492,8 +491,8 @@ bool FRayTracingMeshProcessor::TryAddMeshBatch(
 	if ((!PrimitiveSceneProxy || PrimitiveSceneProxy->ShouldRenderInMainPass())
 		&& ShouldIncludeDomainInMeshPass(Material.GetMaterialDomain()))
 	{
-		if (RayTracingMeshCommandsMode == ERayTracingMeshCommandsMode::PATH_TRACING ||
-			RayTracingMeshCommandsMode == ERayTracingMeshCommandsMode::LIGHTMAP_TRACING)
+		if (RayTracingType == ERayTracingType::PathTracing ||
+			RayTracingType == ERayTracingType::LightMapTracing)
 		{
 			// Path Tracer has its own process call so that it can attach its own material permutation
 			return ProcessPathTracing(MeshBatch, BatchElementMask, PrimitiveSceneProxy, MaterialRenderProxy, Material);
diff --git a/Engine/Source/Runtime/Renderer/Private/RayTracing/RayTracingMaterialHitShaders.h b/Engine/Source/Runtime/Renderer/Private/RayTracing/RayTracingMaterialHitShaders.h
--- a/Engine/Source/Runtime/Renderer/Private/RayTracing/RayTracingMaterialHitShaders.h
+++ b/Engine/Source/Runtime/Renderer/Private/RayTracing/RayTracingMaterialHitShaders.h
@@ -18,8 +18,6 @@
 #include "ShaderParameterStruct.h"
 #include <type_traits>
 
-enum class ERayTracingMeshCommandsMode : uint8;
-
 FRHIRayTracingShader* GetRayTracingDefaultMissShader(const FGlobalShaderMap* ShaderMap);
 FRHIRayTracingShader* GetRayTracingDefaultOpaqueShader(const FGlobalShaderMap* ShaderMap);
 FRHIRayTracingShader* GetRayTracingDefaultHiddenShader(const FGlobalShaderMap* ShaderMap);
@@ -27,19 +25,19 @@ FRHIRayTracingShader* GetRayTracingDefaultHiddenShader(const FGlobalShaderMap* S
 class FRayTracingMeshProcessor
 {
 public:
-	RENDERER_API FRayTracingMeshProcessor(FRayTracingMeshCommandContext* InCommandContext, const FScene* InScene, const FSceneView* InViewIfDynamicMeshCommand, ERayTracingMeshCommandsMode InRayTracingMeshCommandsMode);
-	RENDERER_API virtual ~FRayTracingMeshProcessor();
+	RENDERER_API FRayTracingMeshProcessor(FRayTracingMeshCommandContext* InCommandContext, const FScene* InScene, const FSceneView* InViewIfDynamicMeshCommand, ERayTracingType InRayTracingType);
+	RENDERER_API ~FRayTracingMeshProcessor();
 
 	RENDERER_API void AddMeshBatch(const FMeshBatch& RESTRICT MeshBatch, uint64 BatchElementMask, const FPrimitiveSceneProxy* RESTRICT PrimitiveSceneProxy);
 
-protected:
+private:
 	FRayTracingMeshCommandContext* CommandContext;
 	const FScene* Scene;
 	const FSceneView* ViewIfDynamicMeshCommand;
 	ERHIFeatureLevel::Type FeatureLevel;
-	ERayTracingMeshCommandsMode RayTracingMeshCommandsMode;
+	ERayTracingType RayTracingType;
 
-	RENDERER_API bool Process(
+	bool Process(
 		const FMeshBatch& RESTRICT MeshBatch,
 		uint64 BatchElementMask,
 		const FPrimitiveSceneProxy* RESTRICT PrimitiveSceneProxy,
@@ -55,8 +53,7 @@ protected:
 		const FMaterialRenderProxy& RESTRICT MaterialRenderProxy,
 		const FMaterial& RESTRICT MaterialResource,
 		const TShaderRef<RayTracingShaderType>& RayTracingShader,
-		const ShaderElementDataType& ShaderElementData,
-		ERayTracingViewMaskMode MaskMode)
+		const ShaderElementDataType& ShaderElementData)
 	{
 		const FVertexFactory* RESTRICT VertexFactory = MeshBatch.VertexFactory;
 
@@ -64,7 +61,7 @@ protected:
 
 		FRayTracingMeshCommand SharedCommand;
 
-		SetupRayTracingMeshCommandMaskAndStatus(SharedCommand, MeshBatch, PrimitiveSceneProxy, MaterialResource, MaskMode);
+		SetupRayTracingMeshCommandMaskAndStatus(SharedCommand, MeshBatch, PrimitiveSceneProxy, MaterialResource, RayTracingType);
 
 		if (GRHISupportsRayTracingShaders)
 		{
@@ -102,15 +99,14 @@ protected:
 		}
 	}
 
-private:
-	RENDERER_API bool ProcessPathTracing(
+	bool ProcessPathTracing(
 		const FMeshBatch& RESTRICT MeshBatch,
 		uint64 BatchElementMask,
 		const FPrimitiveSceneProxy* RESTRICT PrimitiveSceneProxy,
 		const FMaterialRenderProxy& RESTRICT MaterialRenderProxy,
 		const FMaterial& RESTRICT MaterialResource);
 
-	RENDERER_API bool TryAddMeshBatch(
+	bool TryAddMeshBatch(
 		const FMeshBatch& RESTRICT MeshBatch,
 		uint64 BatchElementMask,
 		const FPrimitiveSceneProxy* RESTRICT PrimitiveSceneProxy,
diff --git a/Engine/Source/Runtime/Renderer/Private/RendererScene.cpp b/Engine/Source/Runtime/Renderer/Private/RendererScene.cpp
--- a/Engine/Source/Runtime/Renderer/Private/RendererScene.cpp
+++ b/Engine/Source/Runtime/Renderer/Private/RendererScene.cpp
@@ -106,6 +106,7 @@
 #include "RayTracingDynamicGeometryCollection.h"
 #include "RayTracingSkinnedGeometry.h"
 #include "RayTracing/RayTracingScene.h"
+#include "RayTracing/RayTracingInstanceMask.h"
 #endif
 #include "RHIGPUReadback.h"
 #include "ShaderPrint.h"
@@ -1604,7 +1605,7 @@ FScene::FScene(UWorld* InWorld, bool bInRequiresHitProxies, bool bInIsEditorScen
 ,	bScenesPrimitivesNeedStaticMeshElementUpdate(false)
 ,   PathTracingInvalidationCounter(0)
 #if RHI_RAYTRACING
-,   CachedRayTracingMeshCommandsMode(ERayTracingMeshCommandsMode::RAY_TRACING)
+,   CachedRayTracingMeshCommandsType(ERayTracingType::RayTracing)
 #endif
 ,	SkyLight(NULL)
 ,	ConvolvedSkyRenderTargetReadyIndex(-1)
diff --git a/Engine/Source/Runtime/Renderer/Private/ScenePrivate.h b/Engine/Source/Runtime/Renderer/Private/ScenePrivate.h
--- a/Engine/Source/Runtime/Renderer/Private/ScenePrivate.h
+++ b/Engine/Source/Runtime/Renderer/Private/ScenePrivate.h
@@ -2807,11 +2807,7 @@ struct FMeshComputeDispatchCommand
 	FRWBuffer* TargetBuffer;
 };
 
-enum class ERayTracingMeshCommandsMode : uint8 {
-	RAY_TRACING,
-	PATH_TRACING,
-	LIGHTMAP_TRACING,
-};
+enum class ERayTracingType : uint8;
 
 #endif
 
@@ -3090,7 +3086,7 @@ public:
 
 #if RHI_RAYTRACING
 	/** What mode was the cached RT commands prepared for last? */
-	ERayTracingMeshCommandsMode CachedRayTracingMeshCommandsMode;
+	ERayTracingType CachedRayTracingMeshCommandsType;
 #endif
 
 	/** The scene's sky light, if any. */
diff --git a/Engine/Source/Runtime/Renderer/Public/RayTracingInstance.h b/Engine/Source/Runtime/Renderer/Public/RayTracingInstance.h
--- a/Engine/Source/Runtime/Renderer/Public/RayTracingInstance.h
+++ b/Engine/Source/Runtime/Renderer/Public/RayTracingInstance.h
@@ -7,48 +7,12 @@
 
 class FRayTracingGeometry;
 
-struct FRayTracingMaskAndFlags
-{
-	FRayTracingMaskAndFlags()
-		: Mask(0xFF)
-		, bForceOpaque(false)
-		, bDoubleSided(false)
-		, bReverseCulling(false)
-		, bAnySegmentsDecal(false)
-		, bAllSegmentsDecal(false)
-		, bAllSegmentsTranslucent(false)
-	{}
-
-	/** Instance mask that can be used to exclude the instance from specific effects (eg. ray traced shadows). */
-	uint8 Mask;
-
-	/** Whether the instance is forced opaque, i.e. anyhit shaders are disabled on this instance. */
-	uint8 bForceOpaque : 1;
-
-	/** Whether ray hits should be registered for front and back faces. */
-	uint8 bDoubleSided : 1;
-	
-	/** Whether front and back facings should be reversed. */
-	uint8 bReverseCulling : 1;
-
-	/** Whether any or all of the segments in the instance are decals. */
-	uint8 bAnySegmentsDecal : 1;
-	uint8 bAllSegmentsDecal : 1;
-
-	/** Whether all of the segments in the instance are translucent. */
-	uint8 bAllSegmentsTranslucent : 1;
-};
-
 enum class ERayTracingInstanceLayer : uint8
 {
 	NearField,
 	FarField,
 };
 
-/** MeshCommands mode shares the same status for ray tracing view mask mode*/
-//@TODO refactor ERayTracingViewMaskMode and ERayTracingMeshCommandsMode to a single header.
-enum class ERayTracingViewMaskMode : uint8;
-
 struct FRayTracingInstance
 {
 	/** The underlying geometry of this instance specification. */
@@ -81,9 +45,6 @@ struct FRayTracingInstance
 			return MaterialsView;
 		}
 	}
-	
-	UE_DEPRECATED(5.4, "MaskAndFlags is automatically built and cached in RayTracing.cpp")
-	FRayTracingMaskAndFlags MaskAndFlags;
 
 	/** Whether local bounds scale and center translation should be applied to the instance transform. */
 	bool bApplyLocalBoundsTransform = false;
