Refactored planar reflection rendering into one simplified code path in preparation for merging scene renderers.

diff --git a/Engine/Source/Runtime/Renderer/Private/PlanarReflectionRendering.cpp b/Engine/Source/Runtime/Renderer/Private/PlanarReflectionRendering.cpp
--- a/Engine/Source/Runtime/Renderer/Private/PlanarReflectionRendering.cpp
+++ b/Engine/Source/Runtime/Renderer/Private/PlanarReflectionRendering.cpp
@@ -267,26 +267,31 @@ void PrefilterPlanarReflection(
 	}
 }
 
+struct FUpdatePlanarReflectionContentsInputs
+{
+	FSceneRenderer* ParentSceneRenderer;
+	FPlanarReflectionSceneProxy* SceneProxy;
+	FPlanarReflectionRenderTarget* RenderTarget;
+	FPlane MirrorPlane;
+	FName OwnerName;
+	bool bIsMobilePixelProjectedReflectionEnabled;
+};
+
 static void UpdatePlanarReflectionContents_RenderThread(
-	FRHICommandListImmediate& RHICmdList, 
-	FSceneRenderer* MainSceneRenderer, 
-	FSceneRenderer* SceneRenderer, 
-	FPlanarReflectionSceneProxy* SceneProxy,
-	FPlanarReflectionRenderTarget* RenderTarget,
-	const FPlane& MirrorPlane,
-	const FName OwnerName, 
-	bool bUseSceneColorTexture)
+	FRDGBuilder& GraphBuilder,
+	FSceneRenderer* SceneRenderer,
+	const FUpdatePlanarReflectionContentsInputs& Inputs)
 {
 	QUICK_SCOPE_CYCLE_COUNTER(STAT_RenderPlanarReflection);
 
 	{
-		FBox PlanarReflectionBounds = SceneProxy->WorldBounds;
+		FBox PlanarReflectionBounds = Inputs.SceneProxy->WorldBounds;
 		bool bIsInAnyFrustum = false;
 
-		for (int32 ViewIndex = 0; ViewIndex < MainSceneRenderer->Views.Num(); ++ViewIndex)
+		for (int32 ViewIndex = 0; ViewIndex < Inputs.ParentSceneRenderer->Views.Num(); ++ViewIndex)
 		{
-			FViewInfo& View = MainSceneRenderer->Views[ViewIndex];
-			if (MirrorPlane.PlaneDot(View.ViewMatrices.GetViewOrigin()) > 0)
+			FViewInfo& View = Inputs.ParentSceneRenderer->Views[ViewIndex];
+			if (Inputs.MirrorPlane.PlaneDot(View.ViewMatrices.GetViewOrigin()) > 0)
 			{
 				if (View.ViewFrustum.IntersectBox(PlanarReflectionBounds.GetCenter(), PlanarReflectionBounds.GetExtent()))
 				{
@@ -298,20 +303,18 @@ static void UpdatePlanarReflectionContents_RenderThread(
 
 		if (!bIsInAnyFrustum)
 		{
-			delete SceneRenderer;
 			return;
 		}
 
 		bool bIsVisibleInAnyView = true;
 
-		for (int32 ViewIndex = 0; ViewIndex < MainSceneRenderer->Views.Num(); ++ViewIndex)
+		for (const FViewInfo& View : Inputs.ParentSceneRenderer->Views)
 		{
-			FViewInfo& View = MainSceneRenderer->Views[ViewIndex];
 			FSceneViewState* ViewState = View.ViewState;
 
 			if (ViewState)
 			{
-				FIndividualOcclusionHistory& OcclusionHistory = ViewState->PlanarReflectionOcclusionHistories.FindOrAdd(SceneProxy->PlanarReflectionId);
+				FIndividualOcclusionHistory& OcclusionHistory = ViewState->PlanarReflectionOcclusionHistories.FindOrAdd(Inputs.SceneProxy->PlanarReflectionId);
 
 				// +1 to buffered frames because the query is submitted late into the main frame, but read at the beginning of a reflection capture frame
 				const int32 NumBufferedFrames = FOcclusionQueryHelpers::GetNumBufferedFrames(SceneRenderer->FeatureLevel) + 1;
@@ -338,140 +341,100 @@ static void UpdatePlanarReflectionContents_RenderThread(
 
 		if (!bIsVisibleInAnyView)
 		{
-			delete SceneRenderer;
 			return;
 		}
 	}
 
-	SceneRenderer->RenderThreadBegin(RHICmdList);
+	// Clear the SRV prior to rendering so we don't attempt to read from the target we are writing to.
+	Inputs.SceneProxy->RenderTarget = nullptr;
+
+	// Assign the view rect to the resolved scene renderer view rect for the main scene render pass.
+	// Do this at the end of the scope since we can early out for mobile pixel reflections. This needs
+	// to be called after PrepareViewRectsForRendering (which is done within Render for the main path).
+	ON_SCOPE_EXIT
+	{
+		// Update the view rects into the planar reflection proxy.
+		for (int32 ViewIndex = 0; ViewIndex < SceneRenderer->Views.Num(); ++ViewIndex)
+		{
+			// Make sure screen percentage has correctly been set on render thread.
+			check(SceneRenderer->Views[ViewIndex].ViewRect.Area() > 0);
+			Inputs.SceneProxy->ViewRect[ViewIndex] = SceneRenderer->Views[ViewIndex].ViewRect;
+		}
+
+		// Re-assign the SRV after rendering.
+		Inputs.SceneProxy->RenderTarget = Inputs.RenderTarget;
+	};
+
+	// Generate valid data to update planar reflection uniform buffer but don't actually render the reflection scene when we are using mobile pixel projected reflection.
+	if (Inputs.bIsMobilePixelProjectedReflectionEnabled)
+	{
+		SceneRenderer->PrepareViewRectsForRendering(GraphBuilder.RHICmdList);
+		return;
+	}
 
 	FUniformExpressionCacheAsyncUpdateScope AsyncUpdateScope;
 
 	const ERHIFeatureLevel::Type FeatureLevel = SceneRenderer->FeatureLevel;
-	FRDGBuilder GraphBuilder(RHICmdList, RDG_EVENT_NAME("PlanarReflection"), ERDGBuilderFlags::Parallel, SceneRenderer->ShaderPlatform);
 
 	// Make sure we render to the same set of GPUs as the main scene renderer.
-	if (MainSceneRenderer->ViewFamily.RenderTarget != nullptr)
+	if (Inputs.ParentSceneRenderer->ViewFamily.RenderTarget != nullptr)
 	{
-		RenderTarget->SetActiveGPUMask(MainSceneRenderer->ViewFamily.RenderTarget->GetGPUMask(RHICmdList));
+		Inputs.RenderTarget->SetActiveGPUMask(Inputs.ParentSceneRenderer->ViewFamily.RenderTarget->GetGPUMask(GraphBuilder.RHICmdList));
 	}
 	else
 	{
-		RenderTarget->SetActiveGPUMask(FRHIGPUMask::GPU0());
+		Inputs.RenderTarget->SetActiveGPUMask(FRHIGPUMask::GPU0());
 	}
 
-	{
 #if WANTS_DRAW_MESH_EVENTS
-		FString EventName;
-		OwnerName.ToString(EventName);
-		RDG_EVENT_SCOPE(GraphBuilder, "PlanarReflection %s", *EventName);
+	FString EventName;
+	Inputs.OwnerName.ToString(EventName);
+	RDG_EVENT_SCOPE(GraphBuilder, "PlanarReflection %s", *EventName);
 #else
-		RDG_EVENT_SCOPE(GraphBuilder, "UpdatePlanarReflectionContent_RenderThread");
+	RDG_EVENT_SCOPE(GraphBuilder, "UpdatePlanarReflectionContent_RenderThread");
 #endif
 
-		// Render the scene normally
-		{
-			RDG_RHI_EVENT_SCOPE(GraphBuilder, RenderScene);
-			SceneRenderer->Render(GraphBuilder);
-		}
-
-		SceneProxy->RenderTarget = RenderTarget;
-
-		// Update the view rects into the planar reflection proxy.
+	// Applies late update (if any) to view matrices and re-reflects
+	if (SceneRenderer->Views.Num() > 1)
+	{
+		const FMirrorMatrix MirrorMatrix(Inputs.MirrorPlane);
 		for (int32 ViewIndex = 0; ViewIndex < SceneRenderer->Views.Num(); ++ViewIndex)
 		{
-			// Make sure screen percentage has correctly been set on render thread.
-			check(SceneRenderer->Views[ViewIndex].ViewRect.Area() > 0);
-			SceneProxy->ViewRect[ViewIndex] = SceneRenderer->Views[ViewIndex].ViewRect;
-		}
-
-		FRDGTextureRef ReflectionOutputTexture = GraphBuilder.RegisterExternalTexture(CreateRenderTarget(RenderTarget->TextureRHI, TEXT("ReflectionOutputTexture")));
-		GraphBuilder.SetTextureAccessFinal(ReflectionOutputTexture, ERHIAccess::SRVGraphics);
+			FViewInfo& ReflectionViewToUpdate = SceneRenderer->Views[ViewIndex];
 
-		FSceneTextureShaderParameters SceneTextureParameters = CreateSceneTextureShaderParameters(GraphBuilder, &SceneRenderer->GetActiveSceneTextures(), SceneRenderer->FeatureLevel, ESceneTextureSetupMode::SceneDepth);
-		const FMinimalSceneTextures& SceneTextures = SceneRenderer->GetActiveSceneTextures();
+			// Updates view matrices to match new ViewLocation/ViewRotation, un-reflects
+			// Normally performed in late update itself, delayed to here to ensure we don't ever re-reflect without first un-reflecting
+			ReflectionViewToUpdate.UpdateViewMatrix(); 
 
-		for (int32 ViewIndex = 0; ViewIndex < SceneRenderer->Views.Num(); ++ViewIndex)
-		{
-			FViewInfo& View = SceneRenderer->Views[ViewIndex];
-			RDG_GPU_MASK_SCOPE(GraphBuilder, View.GPUMask);
-			if (MainSceneRenderer->Scene->GetShadingPath() == EShadingPath::Deferred)
-			{
-				PrefilterPlanarReflection<true>(GraphBuilder, View, SceneTextureParameters, SceneProxy, SceneTextures.Color.Resolve, ReflectionOutputTexture);
-			}
-			else
-			{
-				PrefilterPlanarReflection<false>(GraphBuilder, View, SceneTextureParameters, SceneProxy, SceneTextures.Color.Resolve, ReflectionOutputTexture);
-			}
+			// Re-reflects view matrices
+			ReflectionViewToUpdate.UpdatePlanarReflectionViewMatrix(ReflectionViewToUpdate, MirrorMatrix);
 		}
 	}
 
-	GraphBuilder.Execute();
-
-	SceneRenderer->RenderThreadEnd(RHICmdList);
-}
-
-// Used for generate valid data to update planar reflection uniform buffer but don't actually render the reflection scene when we are using mobile pixel projected reflection.
-static void UpdatePlanarReflectionContentsWithoutRendering_RenderThread(
-	FRHICommandListImmediate& RHICmdList, 
-	FSceneRenderer* MainSceneRenderer, 
-	FSceneRenderer* SceneRenderer, 
-	FPlanarReflectionSceneProxy* SceneProxy,
-	FPlanarReflectionRenderTarget* RenderTarget,  
-	const FPlane& MirrorPlane,
-	const FName OwnerName)
-{
-	QUICK_SCOPE_CYCLE_COUNTER(STAT_RenderPlanarReflection);
-
-	SceneRenderer->RenderThreadBegin(RHICmdList);
-
-	FBox PlanarReflectionBounds = SceneProxy->WorldBounds;
-
-	bool bIsInAnyFrustum = false;
-	for (int32 ViewIndex = 0; ViewIndex < MainSceneRenderer->Views.Num(); ++ViewIndex)
+	// Render the scene normally
 	{
-		FViewInfo& View = MainSceneRenderer->Views[ViewIndex];
-		if (MirrorPlane.PlaneDot(View.ViewMatrices.GetViewOrigin()) > 0)
-		{
-			if (View.ViewFrustum.IntersectBox(PlanarReflectionBounds.GetCenter(), PlanarReflectionBounds.GetExtent()))
-			{
-				bIsInAnyFrustum = true;
-				break;
-			}
-		}
+		RDG_RHI_EVENT_SCOPE(GraphBuilder, RenderScene);
+		SceneRenderer->Render(GraphBuilder);
 	}
 
-	if (bIsInAnyFrustum)
-	{
-		SCOPED_DRAW_EVENTF(RHICmdList, SceneCapture, TEXT("PlanarReflection %s"), OwnerName);
+	FRDGTextureRef ReflectionOutputTexture = GraphBuilder.RegisterExternalTexture(CreateRenderTarget(Inputs.RenderTarget->TextureRHI, TEXT("ReflectionOutputTexture")));
+	GraphBuilder.SetTextureAccessFinal(ReflectionOutputTexture, ERHIAccess::SRVGraphics);
 
-		// Reflection view late update
-		if (SceneRenderer->Views.Num() > 1)
-		{
-			const FMirrorMatrix MirrorMatrix(MirrorPlane);
-			for (int32 ViewIndex = 0; ViewIndex < SceneRenderer->Views.Num(); ++ViewIndex)
-			{
-				FViewInfo& ReflectionViewToUpdate = SceneRenderer->Views[ViewIndex];
-				const FViewInfo& UpdatedParentView = MainSceneRenderer->Views[ViewIndex];
+	FSceneTextureShaderParameters SceneTextureParameters = CreateSceneTextureShaderParameters(GraphBuilder, &SceneRenderer->GetActiveSceneTextures(), SceneRenderer->FeatureLevel, ESceneTextureSetupMode::SceneDepth);
+	const FMinimalSceneTextures& SceneTextures = SceneRenderer->GetActiveSceneTextures();
 
-				ReflectionViewToUpdate.UpdatePlanarReflectionViewMatrix(UpdatedParentView, MirrorMatrix);
-			}
+	for (FViewInfo& View : SceneRenderer->Views)
+	{
+		RDG_GPU_MASK_SCOPE(GraphBuilder, View.GPUMask);
+		if (Inputs.ParentSceneRenderer->Scene->GetShadingPath() == EShadingPath::Deferred)
+		{
+			PrefilterPlanarReflection<true>(GraphBuilder, View, SceneTextureParameters, Inputs.SceneProxy, SceneTextures.Color.Resolve, ReflectionOutputTexture);
 		}
-
-		SceneRenderer->PrepareViewRectsForRendering(RHICmdList);
-
-		SceneProxy->RenderTarget = RenderTarget;
-
-		// Update the view rects into the planar reflection proxy.
-		for (int32 ViewIndex = 0; ViewIndex < SceneRenderer->Views.Num(); ++ViewIndex)
+		else
 		{
-			// Make sure screen percentage has correctly been set on render thread.
-			check(SceneRenderer->Views[ViewIndex].ViewRect.Area() > 0);
-			SceneProxy->ViewRect[ViewIndex] = SceneRenderer->Views[ViewIndex].ViewRect;
+			PrefilterPlanarReflection<false>(GraphBuilder, View, SceneTextureParameters, Inputs.SceneProxy, SceneTextures.Color.Resolve, ReflectionOutputTexture);
 		}
 	}
-
-	SceneRenderer->RenderThreadEnd(RHICmdList);
 }
 
 extern void BuildProjectionMatrix(FIntPoint RenderTargetSize, float FOV, float InNearClippingPlane, FMatrix& ProjectionMatrix);
@@ -492,222 +455,194 @@ extern void SetupSceneViewExtensionsForSceneCapture(
 	FSceneViewFamily& ViewFamily,
 	TConstArrayView<FSceneView*> Views);
 
-void FScene::UpdatePlanarReflectionContents(UPlanarReflectionComponent* CaptureComponent, FSceneRenderer& MainSceneRenderer)
+void FScene::UpdatePlanarReflectionContents(UPlanarReflectionComponent* CaptureComponent, FSceneRenderer& ParentSceneRenderer)
 {
 	check(CaptureComponent);
 
-	{
-		FIntPoint DesiredBufferSize = FSceneRenderer::GetDesiredInternalBufferSize(MainSceneRenderer.ViewFamily);
-		FVector2D DesiredPlanarReflectionTextureSizeFloat = FVector2D(DesiredBufferSize.X, DesiredBufferSize.Y) * FMath::Clamp(CaptureComponent->ScreenPercentage / 100.f, 0.25f, 1.f);
-		FIntPoint DesiredPlanarReflectionTextureSize;
-		DesiredPlanarReflectionTextureSize.X = FMath::Clamp(FMath::CeilToInt32(DesiredPlanarReflectionTextureSizeFloat.X), 1, static_cast<int32>(DesiredBufferSize.X));
-		DesiredPlanarReflectionTextureSize.Y = FMath::Clamp(FMath::CeilToInt32(DesiredPlanarReflectionTextureSizeFloat.Y), 1, static_cast<int32>(DesiredBufferSize.Y));
-
-		const bool bIsMobilePixelProjectedReflectionEnabled = IsMobilePixelProjectedReflectionEnabled(GetShaderPlatform());
-
-		const bool bIsRenderTargetValid = CaptureComponent->RenderTarget != nullptr
-									&& CaptureComponent->RenderTarget->GetSizeXY() == DesiredPlanarReflectionTextureSize
-									// The RenderTarget's TextureRHI could be nullptr if it is used for mobile pixel projected reflection.
-									&& (bIsMobilePixelProjectedReflectionEnabled || CaptureComponent->RenderTarget->TextureRHI.IsValid());
-		
-
-		if (CaptureComponent->RenderTarget != nullptr && !bIsRenderTargetValid)
-		{
-			FPlanarReflectionRenderTarget* RenderTarget = CaptureComponent->RenderTarget;
-			ENQUEUE_RENDER_COMMAND(ReleaseRenderTargetCommand)(
-				[RenderTarget](FRHICommandListImmediate& RHICmdList)
-				{
-					RenderTarget->ReleaseResource();
-					delete RenderTarget;
-				});
+	FIntPoint DesiredBufferSize = FSceneRenderer::GetDesiredInternalBufferSize(ParentSceneRenderer.ViewFamily);
+	FVector2D DesiredPlanarReflectionTextureSizeFloat = FVector2D(DesiredBufferSize.X, DesiredBufferSize.Y) * FMath::Clamp(CaptureComponent->ScreenPercentage / 100.f, 0.25f, 1.f);
+	FIntPoint DesiredPlanarReflectionTextureSize;
+	DesiredPlanarReflectionTextureSize.X = FMath::Clamp(FMath::CeilToInt32(DesiredPlanarReflectionTextureSizeFloat.X), 1, static_cast<int32>(DesiredBufferSize.X));
+	DesiredPlanarReflectionTextureSize.Y = FMath::Clamp(FMath::CeilToInt32(DesiredPlanarReflectionTextureSizeFloat.Y), 1, static_cast<int32>(DesiredBufferSize.Y));
 
-			CaptureComponent->RenderTarget = nullptr;
-		}
+	const bool bIsMobilePixelProjectedReflectionEnabled = IsMobilePixelProjectedReflectionEnabled(GetShaderPlatform());
 
-		if (CaptureComponent->RenderTarget == nullptr)
-		{
-			CaptureComponent->RenderTarget = new FPlanarReflectionRenderTarget(DesiredPlanarReflectionTextureSize);
-
-			FPlanarReflectionRenderTarget* RenderTarget = CaptureComponent->RenderTarget;
-			FPlanarReflectionSceneProxy* SceneProxy = CaptureComponent->SceneProxy;
-			ENQUEUE_RENDER_COMMAND(InitRenderTargetCommand)(
-				[RenderTarget, SceneProxy, bIsMobilePixelProjectedReflectionEnabled](FRHICommandListImmediate& RHICmdList)
-				{
-					// Don't create the RenderTarget's RHI if it is used for mobile pixel projected reflection
-					if (!bIsMobilePixelProjectedReflectionEnabled)
-					{
-						RenderTarget->InitResource(RHICmdList);
-					}
-					SceneProxy->RenderTarget = nullptr;
-				});
-		}
-		else
-		{
-			// Remove the render target on the planar reflection proxy so that this planar reflection is not getting drawn in its own FSceneRenderer.
-			FPlanarReflectionSceneProxy* SceneProxy = CaptureComponent->SceneProxy;
-			ENQUEUE_RENDER_COMMAND(InitRenderTargetCommand)(
-				[SceneProxy](FRHICommandListImmediate& RHICmdList)
-				{
-					SceneProxy->RenderTarget = nullptr;
-				});
-		}
+	const bool bIsRenderTargetValid = CaptureComponent->RenderTarget != nullptr
+								&& CaptureComponent->RenderTarget->GetSizeXY() == DesiredPlanarReflectionTextureSize
+								// The RenderTarget's TextureRHI could be nullptr if it is used for mobile pixel projected reflection.
+								&& (bIsMobilePixelProjectedReflectionEnabled || CaptureComponent->RenderTarget->TextureRHI.IsValid());
 
-		const FMatrix ComponentTransform = CaptureComponent->GetComponentTransform().ToMatrixWithScale();
-		FPlane MirrorPlane = FPlane(ComponentTransform.TransformPosition(FVector::ZeroVector), ComponentTransform.TransformVector(FVector(0, 0, 1)));
+	FPlanarReflectionRenderTarget* RenderTargetToRelease = nullptr;
+	bool bInitResource = false;
 
-		// Normalize the plane to remove component scaling
-		bool bNormalized = MirrorPlane.Normalize();
+	if (CaptureComponent->RenderTarget != nullptr && !bIsRenderTargetValid)
+	{
+		RenderTargetToRelease = CaptureComponent->RenderTarget;
+		CaptureComponent->RenderTarget = nullptr;
+	}
 
-		if (!bNormalized)
-		{
-			MirrorPlane = FPlane(FVector(0, 0, 1), 0);
-		}
+	if (CaptureComponent->RenderTarget == nullptr)
+	{
+		CaptureComponent->RenderTarget = new FPlanarReflectionRenderTarget(DesiredPlanarReflectionTextureSize);
 
-		TArray<FSceneCaptureViewInfo> SceneCaptureViewInfo;
+		// Don't create the RenderTarget's RHI if it is used for mobile pixel projected reflection
+		bInitResource = !bIsMobilePixelProjectedReflectionEnabled;
+	}
 
-		for (int32 ViewIndex = 0; ViewIndex < MainSceneRenderer.Views.Num() && ViewIndex < GMaxPlanarReflectionViews; ++ViewIndex)
-		{
-			const FViewInfo& View = MainSceneRenderer.Views[ViewIndex];
-			FSceneCaptureViewInfo NewView;
-
-			FVector2D ViewRectMin = FVector2D(View.UnscaledViewRect.Min.X, View.UnscaledViewRect.Min.Y);
-			FVector2D ViewRectMax = FVector2D(View.UnscaledViewRect.Max.X, View.UnscaledViewRect.Max.Y);
-			ViewRectMin *= FMath::Clamp(CaptureComponent->ScreenPercentage / 100.f, 0.25f, 1.f);
-			ViewRectMax *= FMath::Clamp(CaptureComponent->ScreenPercentage / 100.f, 0.25f, 1.f);
-
-			NewView.ViewRect.Min.X = FMath::TruncToInt(ViewRectMin.X);
-			NewView.ViewRect.Min.Y = FMath::TruncToInt(ViewRectMin.Y);
-			NewView.ViewRect.Max.X = FMath::CeilToInt(ViewRectMax.X);
-			NewView.ViewRect.Max.Y = FMath::CeilToInt(ViewRectMax.Y);
-
-			// Create a mirror matrix and premultiply the view transform by it
-			const FMirrorMatrix MirrorMatrix(MirrorPlane);
-			const FMatrix ViewMatrix(MirrorMatrix * View.ViewMatrices.GetViewMatrix());
-			const FVector ViewOrigin = ViewMatrix.InverseTransformPosition(FVector::ZeroVector);
-			const FMatrix ViewRotationMatrix = ViewMatrix.RemoveTranslation();
-
-			FMatrix ProjectionMatrix;
-			if (!View.IsPerspectiveProjection() || (CaptureComponent->ExtraFOV == 0.f && MainSceneRenderer.Views.Num() > 1))
-			{
-				// Prefer exact (potentially uneven) stereo projection matrices when no extra FOV is requested
-				ProjectionMatrix = View.ViewMatrices.GetProjectionMatrix();
-			}
-			else
-			{
-				const float HalfFOV = FMath::Atan(1.0f / View.ViewMatrices.GetProjectionMatrix().M[0][0]);
-				BuildProjectionMatrix(View.UnscaledViewRect.Size(), HalfFOV + FMath::DegreesToRadians(CaptureComponent->ExtraFOV), GNearClippingPlane, ProjectionMatrix);
-			}
+	const FMatrix ComponentTransform = CaptureComponent->GetComponentTransform().ToMatrixWithScale();
+	FPlane MirrorPlane = FPlane(ComponentTransform.TransformPosition(FVector::ZeroVector), ComponentTransform.TransformVector(FVector(0, 0, 1)));
 
-			NewView.ViewLocation = View.ViewLocation;
-			NewView.ViewRotation = View.ViewRotation;
-			NewView.ViewOrigin = ViewOrigin;
-			NewView.ViewRotationMatrix = ViewRotationMatrix;
-			NewView.ProjectionMatrix = ProjectionMatrix;
-			NewView.StereoPass = View.StereoPass;
-			NewView.StereoViewIndex = View.StereoViewIndex;
-			NewView.FOV = View.FOV;
+	if (!MirrorPlane.Normalize())
+	{
+		MirrorPlane = FPlane(FVector(0, 0, 1), 0);
+	}
 
-			SceneCaptureViewInfo.Add(NewView);
-		}
-		
-		FPostProcessSettings PostProcessSettings;
+	TArray<FSceneCaptureViewInfo, FConcurrentLinearArrayAllocator> SceneCaptureViewInfo;
 
-		bool bIsMobileMultiViewEnabled = false;
-		if (MainSceneRenderer.ViewFamily.Views.Num() > 0)
+	for (int32 ViewIndex = 0; ViewIndex < ParentSceneRenderer.Views.Num() && ViewIndex < GMaxPlanarReflectionViews; ++ViewIndex)
+	{
+		const FViewInfo& View = ParentSceneRenderer.Views[ViewIndex];
+		FSceneCaptureViewInfo NewView;
+
+		FVector2D ViewRectMin = FVector2D(View.UnscaledViewRect.Min.X, View.UnscaledViewRect.Min.Y);
+		FVector2D ViewRectMax = FVector2D(View.UnscaledViewRect.Max.X, View.UnscaledViewRect.Max.Y);
+		ViewRectMin *= FMath::Clamp(CaptureComponent->ScreenPercentage / 100.f, 0.25f, 1.f);
+		ViewRectMax *= FMath::Clamp(CaptureComponent->ScreenPercentage / 100.f, 0.25f, 1.f);
+
+		NewView.ViewRect.Min.X = FMath::TruncToInt(ViewRectMin.X);
+		NewView.ViewRect.Min.Y = FMath::TruncToInt(ViewRectMin.Y);
+		NewView.ViewRect.Max.X = FMath::CeilToInt(ViewRectMax.X);
+		NewView.ViewRect.Max.Y = FMath::CeilToInt(ViewRectMax.Y);
+
+		// Create a mirror matrix and premultiply the view transform by it
+		const FMirrorMatrix MirrorMatrix(MirrorPlane);
+		const FMatrix ViewMatrix(MirrorMatrix * View.ViewMatrices.GetViewMatrix());
+		const FVector ViewOrigin = ViewMatrix.InverseTransformPosition(FVector::ZeroVector);
+		const FMatrix ViewRotationMatrix = ViewMatrix.RemoveTranslation();
+
+		FMatrix ProjectionMatrix;
+		if (!View.IsPerspectiveProjection() || (CaptureComponent->ExtraFOV == 0.f && ParentSceneRenderer.Views.Num() > 1))
 		{
-			bIsMobileMultiViewEnabled = MainSceneRenderer.ViewFamily.Views[0]->Aspects.IsMobileMultiViewEnabled();
+			// Prefer exact (potentially uneven) stereo projection matrices when no extra FOV is requested
+			ProjectionMatrix = View.ViewMatrices.GetProjectionMatrix();
 		}
-
-		FSceneViewFamilyContext ViewFamily(FSceneViewFamily::ConstructionValues(
-			CaptureComponent->RenderTarget,
-			this,
-			CaptureComponent->ShowFlags)
-			.SetResolveScene(false)
-			.SetRealtimeUpdate(true)
-			.SetRequireMobileMultiView(bIsMobileMultiViewEnabled));
-
-		// Uses the exact same secondary view fraction on the planar reflection as the main viewport.
-		ViewFamily.SecondaryViewFraction = MainSceneRenderer.ViewFamily.SecondaryViewFraction;
-
-		FSceneViewExtensionContext ViewExtensionContext(this);
-		ViewExtensionContext.bStereoEnabled = true;
-		ViewFamily.ViewExtensions = GEngine->ViewExtensions->GatherActiveExtensions(ViewExtensionContext);
-
-		TArray<FSceneView*> Views = SetupViewFamilyForSceneCapture(
-			ViewFamily,
-			CaptureComponent,
-			SceneCaptureViewInfo, CaptureComponent->MaxViewDistanceOverride,
-			/* bCaptureSceneColor = */ true,
-			/* bIsPlanarReflection = */ true,
-			&PostProcessSettings, 1.0f,
-			/*ViewActor =*/ nullptr,
-			/*CubemapFaceIndex =*/ INDEX_NONE);
-
-		// Fork main renderer's screen percentage interface to have exactly same settings.
-		ViewFamily.EngineShowFlags.ScreenPercentage = MainSceneRenderer.ViewFamily.EngineShowFlags.ScreenPercentage;
-		ViewFamily.SetScreenPercentageInterface(FSceneRenderer::ForkScreenPercentageInterface(
-			MainSceneRenderer.ViewFamily.GetScreenPercentageInterface(), ViewFamily));
-
-		for (FSceneView* View : Views)
+		else
 		{
-			View->GlobalClippingPlane = MirrorPlane;
-			// Jitter can't be removed completely due to the clipping plane
-			// Also, this prevents the prefilter pass, which reads from jittered depth, from having to do special handling of it's depth-dependent input
-			View->bAllowTemporalJitter = false;
-			View->bRenderSceneTwoSided = CaptureComponent->bRenderSceneTwoSided;
+			const float HalfFOV = FMath::Atan(1.0f / View.ViewMatrices.GetProjectionMatrix().M[0][0]);
+			BuildProjectionMatrix(View.UnscaledViewRect.Size(), HalfFOV + FMath::DegreesToRadians(CaptureComponent->ExtraFOV), GNearClippingPlane, ProjectionMatrix);
 		}
 
-		// Call SetupViewFamily & SetupView on scene view extensions before renderer creation
-		SetupSceneViewExtensionsForSceneCapture(ViewFamily, Views);
+		NewView.ViewLocation = View.ViewLocation;
+		NewView.ViewRotation = View.ViewRotation;
+		NewView.ViewOrigin = ViewOrigin;
+		NewView.ViewRotationMatrix = ViewRotationMatrix;
+		NewView.ProjectionMatrix = ProjectionMatrix;
+		NewView.StereoPass = View.StereoPass;
+		NewView.StereoViewIndex = View.StereoViewIndex;
+		NewView.FOV = View.FOV;
 
-		FSceneRenderer* SceneRenderer = FSceneRenderer::CreateSceneRenderer(&ViewFamily, nullptr);
+		SceneCaptureViewInfo.Add(NewView);
+	}
 
-		// Disable screen percentage on planar reflection renderer if main one has screen percentage disabled.
-		SceneRenderer->ViewFamily.EngineShowFlags.ScreenPercentage = MainSceneRenderer.ViewFamily.EngineShowFlags.ScreenPercentage;
+	FPostProcessSettings PostProcessSettings;
 
-		for (int32 ViewIndex = 0; ViewIndex < SceneCaptureViewInfo.Num(); ++ViewIndex)
-		{
-			CaptureComponent->ProjectionWithExtraFOV[ViewIndex] = SceneCaptureViewInfo[ViewIndex].ProjectionMatrix;
+	bool bIsMobileMultiViewEnabled = false;
+	if (ParentSceneRenderer.ViewFamily.Views.Num() > 0)
+	{
+		bIsMobileMultiViewEnabled = ParentSceneRenderer.ViewFamily.Views[0]->Aspects.IsMobileMultiViewEnabled();
+	}
 
-			const bool bIsStereo = IStereoRendering::IsStereoEyeView(MainSceneRenderer.Views[0]);
+	FSceneViewFamilyContext ViewFamily(FSceneViewFamily::ConstructionValues(
+		CaptureComponent->RenderTarget,
+		this,
+		CaptureComponent->ShowFlags)
+		.SetResolveScene(false)
+		.SetRealtimeUpdate(true)
+		.SetRequireMobileMultiView(bIsMobileMultiViewEnabled));
+
+	FSceneViewExtensionContext ViewExtensionContext(this);
+	ViewExtensionContext.bStereoEnabled = true;
+	ViewFamily.ViewExtensions = GEngine->ViewExtensions->GatherActiveExtensions(ViewExtensionContext);
+
+	TArray<FSceneView*> Views = SetupViewFamilyForSceneCapture(
+		ViewFamily,
+		CaptureComponent,
+		SceneCaptureViewInfo, CaptureComponent->MaxViewDistanceOverride,
+		/* bCaptureSceneColor = */ true,
+		/* bIsPlanarReflection = */ true,
+		&PostProcessSettings, 1.0f,
+		/*ViewActor =*/ nullptr,
+		/*CubemapFaceIndex =*/ INDEX_NONE);
+
+	// Fork main renderer's screen percentage interface to have exactly same settings.
+	ViewFamily.EngineShowFlags.ScreenPercentage = ParentSceneRenderer.ViewFamily.EngineShowFlags.ScreenPercentage;
+	ViewFamily.SetScreenPercentageInterface(FSceneRenderer::ForkScreenPercentageInterface(
+		ParentSceneRenderer.ViewFamily.GetScreenPercentageInterface(), ViewFamily));
+
+	for (FSceneView* View : Views)
+	{
+		View->GlobalClippingPlane = MirrorPlane;
+		// Jitter can't be removed completely due to the clipping plane
+		// Also, this prevents the prefilter pass, which reads from jittered depth, from having to do special handling of it's depth-dependent input
+		View->bAllowTemporalJitter = false;
+		View->bRenderSceneTwoSided = CaptureComponent->bRenderSceneTwoSided;
+	}
 
-			const FMatrix ProjectionMatrix = SceneCaptureViewInfo[ViewIndex].ProjectionMatrix;
-			FPlanarReflectionSceneProxy* SceneProxy = CaptureComponent->SceneProxy;
+	const bool bIsStereo = IStereoRendering::IsStereoEyeView(ParentSceneRenderer.Views[0]);
 
-			ENQUEUE_RENDER_COMMAND(UpdateProxyCommand)(
-				[ProjectionMatrix, ViewIndex, bIsStereo, SceneProxy](FRHICommandList& RHICmdList)
-				{
-					SceneProxy->ProjectionWithExtraFOV[ViewIndex] = ProjectionMatrix;
-					SceneProxy->bIsStereo = bIsStereo;
-				});
-		}
+	// Call SetupViewFamily & SetupView on scene view extensions before renderer creation
+	SetupSceneViewExtensionsForSceneCapture(ViewFamily, Views);
 
-		{
-			const FName OwnerName = CaptureComponent->GetOwner() ? CaptureComponent->GetOwner()->GetFName() : NAME_None;
-			FSceneRenderer* MainSceneRendererPtr = &MainSceneRenderer;
-			FPlanarReflectionSceneProxy* SceneProxyPtr = CaptureComponent->SceneProxy;
-			FPlanarReflectionRenderTarget* RenderTargetPtr = CaptureComponent->RenderTarget;
-
-			UE::RenderCommandPipe::FSyncScope SyncScope;
+	for (int32 ViewIndex = 0; ViewIndex < SceneCaptureViewInfo.Num(); ++ViewIndex)
+	{
+		CaptureComponent->ProjectionWithExtraFOV[ViewIndex] = SceneCaptureViewInfo[ViewIndex].ProjectionMatrix;
+	}
 
-			if (bIsMobilePixelProjectedReflectionEnabled)
+	const FUpdatePlanarReflectionContentsInputs Inputs
+	{
+		  .ParentSceneRenderer = &ParentSceneRenderer
+		, .SceneProxy = CaptureComponent->SceneProxy
+		, .RenderTarget = CaptureComponent->RenderTarget
+		, .MirrorPlane = MirrorPlane
+		, .OwnerName = CaptureComponent->GetOwner() ? CaptureComponent->GetOwner()->GetFName() : NAME_None
+		, .bIsMobilePixelProjectedReflectionEnabled = bIsMobilePixelProjectedReflectionEnabled
+	};
+
+	FSceneRenderer* SceneRenderer = FSceneRenderer::CreateSceneRenderer(&ViewFamily, nullptr);
+
+	ENQUEUE_RENDER_COMMAND(RenderReflectionCapture)(
+		[SceneRenderer, Inputs, SceneCaptureViewInfo = MoveTemp(SceneCaptureViewInfo), RenderTargetToRelease, bInitResource, bIsStereo] (FRHICommandListImmediate& RHICmdList)
+	{
+		SceneRenderer->RenderThreadBegin(RHICmdList);
+		{
+			FRDGBuilder GraphBuilder(RHICmdList, RDG_EVENT_NAME("PlanarReflection"), ERDGBuilderFlags::Parallel, SceneRenderer->ShaderPlatform);
+	
+			if (RenderTargetToRelease)
 			{
-				ENQUEUE_RENDER_COMMAND(CaptureCommand)(
-					[SceneRenderer, MirrorPlane, OwnerName, MainSceneRendererPtr, SceneProxyPtr, RenderTargetPtr](FRHICommandListImmediate& RHICmdList)
-				{
-					UpdatePlanarReflectionContentsWithoutRendering_RenderThread(RHICmdList, MainSceneRendererPtr, SceneRenderer, SceneProxyPtr, RenderTargetPtr, MirrorPlane, OwnerName);
-				});
+				RenderTargetToRelease->ReleaseResource();
+				delete RenderTargetToRelease;
 			}
-			else
+	
+			if (bInitResource)
 			{
-				ENQUEUE_RENDER_COMMAND(CaptureCommand)(
-					[SceneRenderer, MirrorPlane, OwnerName, MainSceneRendererPtr, SceneProxyPtr, RenderTargetPtr](FRHICommandListImmediate& RHICmdList)
-				{
-					UpdatePlanarReflectionContents_RenderThread(RHICmdList, MainSceneRendererPtr, SceneRenderer, SceneProxyPtr, RenderTargetPtr, MirrorPlane, OwnerName, true);
-				});
+				Inputs.RenderTarget->InitResource(GraphBuilder.RHICmdList);
+			}
+	
+			// Remove the render target on the planar reflection proxy so that this planar reflection is not getting drawn in its own FSceneRenderer.
+			Inputs.SceneProxy->RenderTarget = nullptr;
+	
+			for (int32 ViewIndex = 0; ViewIndex < SceneCaptureViewInfo.Num(); ++ViewIndex)
+			{
+				Inputs.SceneProxy->ProjectionWithExtraFOV[ViewIndex] = SceneCaptureViewInfo[ViewIndex].ProjectionMatrix;
+				Inputs.SceneProxy->bIsStereo = bIsStereo;
 			}
+	
+			UpdatePlanarReflectionContents_RenderThread(GraphBuilder, SceneRenderer, Inputs);
+
+			GraphBuilder.Execute();
 		}
-	}
+		SceneRenderer->RenderThreadEnd(RHICmdList);
+	});
 }
 
 void FScene::AddPlanarReflection(UPlanarReflectionComponent* Component)
