HeterogeneousVolumes: Updating scalability flag

diff --git a/Engine/Source/Runtime/Renderer/Private/HeterogeneousVolumes/HeterogeneousVolumesVoxelGridPipeline.cpp b/Engine/Source/Runtime/Renderer/Private/HeterogeneousVolumes/HeterogeneousVolumesVoxelGridPipeline.cpp
--- a/Engine/Source/Runtime/Renderer/Private/HeterogeneousVolumes/HeterogeneousVolumesVoxelGridPipeline.cpp
+++ b/Engine/Source/Runtime/Renderer/Private/HeterogeneousVolumes/HeterogeneousVolumesVoxelGridPipeline.cpp
@@ -221,7 +221,7 @@ static TAutoConsoleVariable<int32> CVarHeterogeneousVolumesShadowMaxSampleCount(
 	TEXT("r.HeterogeneousVolumes.Shadows.MaxSampleCount"),
 	8,
 	TEXT("Maximum sample count when building volumetric shadow map (Default = 8)\n"),
-	ECVF_RenderThreadSafe
+	ECVF_Scalability | ECVF_RenderThreadSafe
 );
 
 static TAutoConsoleVariable<float> CVarHeterogeneousVolumesShadowAbsoluteErrorThreshold(
