[Renderer]
* Fixed non unity/pch compile errors

diff --git a/Engine/Source/Runtime/Renderer/Private/HairStrands/HairStrandsRasterCommon.h b/Engine/Source/Runtime/Renderer/Private/HairStrands/HairStrandsRasterCommon.h
--- a/Engine/Source/Runtime/Renderer/Private/HairStrands/HairStrandsRasterCommon.h
+++ b/Engine/Source/Runtime/Renderer/Private/HairStrands/HairStrandsRasterCommon.h
@@ -10,6 +10,7 @@
 #include "RendererInterface.h"
 #include "InstanceCulling/InstanceCullingContext.h"
 #include "HairStrands/HairStrandsData.h"
+#include "SceneView.h"
 
 class FViewInfo;
 class FInstanceCullingManager;
diff --git a/Engine/Source/Runtime/Renderer/Private/ScreenSpaceRayTracing.h b/Engine/Source/Runtime/Renderer/Private/ScreenSpaceRayTracing.h
--- a/Engine/Source/Runtime/Renderer/Private/ScreenSpaceRayTracing.h
+++ b/Engine/Source/Runtime/Renderer/Private/ScreenSpaceRayTracing.h
@@ -3,6 +3,7 @@
 #pragma once
 
 #include "RenderGraph.h"
+#include "SceneView.h"
 #include "ScreenSpaceDenoise.h"
 #include "IndirectLightRendering.h"
 
