Refactored scene capture path in preparation for merging scene renderers.

diff --git a/Engine/Source/Runtime/Renderer/Private/SceneCaptureRendering.cpp b/Engine/Source/Runtime/Renderer/Private/SceneCaptureRendering.cpp
--- a/Engine/Source/Runtime/Renderer/Private/SceneCaptureRendering.cpp
+++ b/Engine/Source/Runtime/Renderer/Private/SceneCaptureRendering.cpp
@@ -211,7 +211,10 @@ static bool CaptureNeedsSceneColor(ESceneCaptureSource CaptureSource)
 	return CaptureSource != SCS_FinalColorLDR && CaptureSource != SCS_FinalColorHDR && CaptureSource != SCS_FinalToneCurveHDR;
 }
 
-static TFunction<void(FRHICommandList& RHICmdList, int32 ViewIndex)> CopyCaptureToTargetSetViewportFn = [](FRHICommandList& RHICmdList, int32 ViewIndex) {};
+using FSceneCaptureViewportSetterFunction = TFunction<void(FRHICommandList& RHICmdList, int32 ViewIndex)>;
+class FSceneCaptureViewportSetterMap : public TMap<FRDGTexture*, FSceneCaptureViewportSetterFunction> {};
+
+RDG_REGISTER_BLACKBOARD_STRUCT(FSceneCaptureViewportSetterMap);
 
 void CopySceneCaptureComponentToTarget(
 	FRDGBuilder& GraphBuilder,
@@ -313,18 +316,31 @@ void CopySceneCaptureComponentToTarget(
 			TargetSize = (View.bIsSceneCaptureCube && NumViews == 6) ? ViewFamilyTexture->Desc.Extent : View.UnconstrainedViewRect.Size();
 		}
 
+		TFunction<void(FRHICommandList& RHICmdList, int32 ViewIndex)> SetViewportLambda;
+
+		if (FSceneCaptureViewportSetterMap* Map = GraphBuilder.Blackboard.GetMutable<FSceneCaptureViewportSetterMap>())
+		{
+			if (auto* Function = Map->Find(ViewFamilyTexture))
+			{
+				SetViewportLambda = MoveTemp(*Function);
+			}
+		}
+
 		GraphBuilder.AddPass(
 			RDG_EVENT_NAME("View(%d)", ViewIndex),
 			PassParameters,
 			ERDGPassFlags::Raster,
-			[PassParameters, GraphicsPSOInit, VertexShader, PixelShader, &View, ViewIndex, TargetSize] (FRDGAsyncTask, FRHICommandList& RHICmdList)
+			[PassParameters, GraphicsPSOInit, VertexShader, PixelShader, &View, ViewIndex, TargetSize, SetViewportLambda = MoveTemp(SetViewportLambda)] (FRDGAsyncTask, FRHICommandList& RHICmdList)
 		{
 			FGraphicsPipelineStateInitializer LocalGraphicsPSOInit = GraphicsPSOInit;
 			RHICmdList.ApplyCachedRenderTargets(LocalGraphicsPSOInit);
 			SetGraphicsPipelineState(RHICmdList, LocalGraphicsPSOInit, 0);
 			SetShaderParameters(RHICmdList, PixelShader, PixelShader.GetPixelShader(), *PassParameters);
-			
-			CopyCaptureToTargetSetViewportFn(RHICmdList, ViewIndex);
+
+			if (SetViewportLambda)
+			{
+				SetViewportLambda(RHICmdList, ViewIndex);
+			}
 
 			DrawRectangle(
 				RHICmdList,
@@ -392,255 +408,101 @@ void CopySceneCaptureComponentToTarget(
 	CopySceneCaptureComponentToTarget(GraphBuilder, ViewFamilyTexture, nullptr, ViewFamily, Views);
 }
 
-static void UpdateSceneCaptureContentDeferred_RenderThread(
-	FRHICommandListImmediate& RHICmdList, 
-	FSceneRenderer* SceneRenderer, 
-	FRenderTarget* RenderTarget, 
-	FTexture* RenderTargetTexture, 
-	const FString& EventName, 
+static void UpdateSceneCaptureContent_RenderThread(
+	FRDGBuilder& GraphBuilder,
+	FSceneRenderer* SceneRenderer,
+	FRenderTarget* RenderTarget,
+	FTexture* RenderTargetTexture,
 	TConstArrayView<FRHICopyTextureInfo> CopyInfos,
 	bool bGenerateMips,
 	const FGenerateMipsParams& GenerateMipsParams,
 	bool bClearRenderTarget,
-	bool bOrthographicCamera
-	)
+	bool bOrthographicCamera)
 {
-	SceneRenderer->RenderThreadBegin(RHICmdList);
+	const EShadingPath ShadingPath = GetFeatureLevelShadingPath(SceneRenderer->Scene->GetFeatureLevel());
 
-	const ERHIFeatureLevel::Type FeatureLevel = SceneRenderer->FeatureLevel;
+	FUniformExpressionCacheAsyncUpdateScope AsyncUpdateScope;
 
-#if WANTS_DRAW_MESH_EVENTS
-	SCOPED_DRAW_EVENTF(RHICmdList, SceneCapture, TEXT("SceneCapture %s"), EventName);
-	FRDGBuilder GraphBuilder(RHICmdList, RDG_EVENT_NAME("SceneCapture %s", *EventName), ERDGBuilderFlags::Parallel, SceneRenderer->ShaderPlatform);
-#else
-	SCOPED_DRAW_EVENT(RHICmdList, UpdateSceneCaptureContent_RenderThread);
-	FRDGBuilder GraphBuilder(RHICmdList, RDG_EVENT_NAME("SceneCapture"), ERDGBuilderFlags::Parallel, SceneRenderer->ShaderPlatform);
-#endif
+	// The target texture is what gets rendered to, while OutputTexture is the final output.  For 2D scene captures, these textures
+	// are the same.  For cube captures, OutputTexture will be a cube map, while TargetTexture will be a 2D render target containing either
+	// one face of the cube map (when GSceneCaptureCubeSinglePass=0) or the six faces of the cube map tiled in a split screen configuration.
+	FRDGTextureRef TargetTexture = RegisterExternalTexture(GraphBuilder, RenderTarget->GetRenderTargetTexture(), TEXT("SceneCaptureTarget"));
+	FRDGTextureRef OutputTexture = RegisterExternalTexture(GraphBuilder, RenderTargetTexture->TextureRHI, TEXT("SceneCaptureTexture"));
 
+	if (bClearRenderTarget)
 	{
-		// The target texture is what gets rendered to, while OutputTexture is the final output.  For 2D scene captures, these textures
-		// are the same.  For cube captures, OutputTexture will be a cube map, while TargetTexture will be a 2D render target containing either
-		// one face of the cube map (when GSceneCaptureCubeSinglePass=0) or the six faces of the cube map tiled in a split screen configuration.
-		FRDGTextureRef TargetTexture = RegisterExternalTexture(GraphBuilder, RenderTarget->GetRenderTargetTexture(), TEXT("SceneCaptureTarget"));
-		FRDGTextureRef OutputTexture = RegisterExternalTexture(GraphBuilder, RenderTargetTexture->TextureRHI, TEXT("SceneCaptureTexture"));
-
-		if (bClearRenderTarget)
-		{
-			AddClearRenderTargetPass(GraphBuilder, TargetTexture, FLinearColor::Black, SceneRenderer->Views[0].UnscaledViewRect);
-		}
-
-		// The lambda below applies to tiled orthographic rendering, where the captured result is blitted from the origin in a scene texture
-		// to a viewport on a larger output texture.  It specifically doesn't apply to cube maps, where the output texture has the same tiling
-		// as the scene textures, and no viewport remapping is required.
-		if (!CopyInfos[0].Size.IsZero() && !OutputTexture->Desc.IsTextureCube())
-		{
-			// Fix for static analysis warning -- lambda lifetime exceeds lifetime of CopyInfos.  Technically, the lambda is consumed in the
-			// scene render call below and not used afterwards, but static analysis doesn't know that, so we make a copy.
-			TArray<FRHICopyTextureInfo, TInlineAllocator<1>> CopyInfosLocal(CopyInfos);
-
-			CopyCaptureToTargetSetViewportFn = [CopyInfos = MoveTemp(CopyInfosLocal)](FRHICommandList& RHICmdList, int32 ViewIndex)
-			{
-				const FIntRect CopyDestRect = CopyInfos[ViewIndex].GetDestRect();
-
-				RHICmdList.SetScissorRect(false, 0, 0, 0, 0);
-				RHICmdList.SetViewport
-				(
-					float(CopyDestRect.Min.X),
-					float(CopyDestRect.Min.Y),
-					0.0f,
-					float(CopyDestRect.Max.X),
-					float(CopyDestRect.Max.Y),
-					1.0f
-				);
-			};
-		}
-		else
-		{
-			CopyCaptureToTargetSetViewportFn = [](FRHICommandList& RHICmdList, int32 ViewIndex) {};
-		}
-
-
-		// Disable occlusion queries when in orthographic mode
-		if (bOrthographicCamera)
-		{
-			FViewInfo& View = SceneRenderer->Views[0];
-			View.bDisableQuerySubmissions = true;
-			View.bIgnoreExistingQueries = true;
-		}
-
-		// Render the scene normally
-		{
-			RDG_RHI_EVENT_SCOPE(GraphBuilder, RenderScene);
-			SceneRenderer->Render(GraphBuilder);
-		}
-
-		// These copies become a no-op (function returns immediately) if TargetTexture and OutputTexture are the same, which
-		// is true for 2D scene captures.  Actual copies only occur for cube captures, where copying is necessary to get
-		// result data to specific slices.
-		for (const FRHICopyTextureInfo& CopyInfo : CopyInfos)
-		{
-			AddCopyTexturePass(GraphBuilder, TargetTexture, OutputTexture, CopyInfo);
-		}
-
-		if (bGenerateMips)
-		{
-			FGenerateMips::Execute(GraphBuilder, SceneRenderer->FeatureLevel, OutputTexture, GenerateMipsParams);
-		}
-
-		GraphBuilder.Execute();
+		AddClearRenderTargetPass(GraphBuilder, TargetTexture, FLinearColor::Black, SceneRenderer->Views[0].UnscaledViewRect);
 	}
 
-	SceneRenderer->RenderThreadEnd(RHICmdList);
-}
-
-static void UpdateSceneCaptureContentMobile_RenderThread(
-	FRHICommandListImmediate& RHICmdList,
-	FSceneRenderer* SceneRenderer,
-	FRenderTarget* RenderTarget,
-	FTexture* RenderTargetTexture,
-	const FString& EventName,
-	TConstArrayView<FRHICopyTextureInfo> CopyInfos,
-	bool bGenerateMips,
-	const FGenerateMipsParams& GenerateMipsParams)
-{
-	SceneRenderer->RenderThreadBegin(RHICmdList);
-
-#if WANTS_DRAW_MESH_EVENTS
-	SCOPED_DRAW_EVENTF(RHICmdList, SceneCaptureMobile, TEXT("SceneCaptureMobile %s"), EventName);
-	FRDGBuilder GraphBuilder(RHICmdList, RDG_EVENT_NAME("SceneCaptureMobile %s", *EventName));
-#else
-	SCOPED_DRAW_EVENT(RHICmdList, UpdateSceneCaptureContentMobile_RenderThread);
-	FRDGBuilder GraphBuilder(RHICmdList, RDG_EVENT_NAME("SceneCaptureMobile"));
-#endif
-
+	// The lambda below applies to tiled orthographic rendering, where the captured result is blitted from the origin in a scene texture
+	// to a viewport on a larger output texture.  It specifically doesn't apply to cube maps, where the output texture has the same tiling
+	// as the scene textures, and no viewport remapping is required.
+	if (!CopyInfos[0].Size.IsZero() && !OutputTexture->Desc.IsTextureCube())
 	{
-		// The target texture is what gets rendered to, while OutputTexture is the final output.  For 2D scene captures, these textures
-		// are the same.  For cube captures, OutputTexture will be a cube map, while TargetTexture will be a 2D render target containing either
-		// one face of the cube map (when GSceneCaptureCubeSinglePass=0) or the six faces of the cube map tiled in a split screen configuration.
-		FRDGTextureRef TargetTexture = RegisterExternalTexture(GraphBuilder, RenderTarget->GetRenderTargetTexture(), TEXT("SceneCaptureTarget"));
-		FRDGTextureRef OutputTexture = RegisterExternalTexture(GraphBuilder, RenderTargetTexture->TextureRHI, TEXT("SceneCaptureTexture"));
-
-		// The lambda below applies to tiled orthographic rendering, where the captured result is blitted from the origin in a scene texture
-		// to a viewport on a larger output texture.  It specifically doesn't apply to cube maps, where the output texture has the same tiling
-		// as the scene textures, and no viewport remapping is required.
-		if (!CopyInfos[0].Size.IsZero() && !OutputTexture->Desc.IsTextureCube())
-		{
-			// Fix for static analysis warning -- lambda lifetime exceeds lifetime of CopyInfos.  Technically, the lambda is consumed in the
-			// scene render call below and not used afterwards, but static analysis doesn't know that, so we make a copy.
-			TArray<FRHICopyTextureInfo, TInlineAllocator<1>> CopyInfosLocal(CopyInfos);
-
-			CopyCaptureToTargetSetViewportFn = [CopyInfos = MoveTemp(CopyInfosLocal)](FRHICommandList& RHICmdList, int32 ViewIndex)
-			{
-				const FIntRect CopyDestRect = CopyInfos[ViewIndex].GetDestRect();
-
-				RHICmdList.SetScissorRect(false, 0, 0, 0, 0);
-				RHICmdList.SetViewport
-				(
-					float(CopyDestRect.Min.X),
-					float(CopyDestRect.Min.Y),
-					0.0f,
-					float(CopyDestRect.Max.X),
-					float(CopyDestRect.Max.Y),
-					1.0f
-				);
-			};
-		}
-		else
-		{
-			CopyCaptureToTargetSetViewportFn = [](FRHICommandList& RHICmdList, int32 ViewIndex) {};
-		}
-
-		// Render the scene normally
-		{
-			RDG_RHI_EVENT_SCOPE(GraphBuilder, RenderScene);
-			SceneRenderer->Render(GraphBuilder);
-		}
-
-		{
-			// Handles copying the SceneColor render target to the output if necessary (this happens inside the renderer for the deferred path).
-			// Other scene captures are automatically written directly to the output, in which case this function returns and does nothing.
-			const FRenderTarget* FamilyTarget = SceneRenderer->ViewFamily.RenderTarget;
-			FRDGTextureRef FamilyTexture = RegisterExternalTexture(GraphBuilder, FamilyTarget->GetRenderTargetTexture(), TEXT("OutputTexture"));
-			const FMinimalSceneTextures& SceneTextures = SceneRenderer->GetActiveSceneTextures();
-
-			RDG_EVENT_SCOPE(GraphBuilder, "CaptureSceneColor");
-			CopySceneCaptureComponentToTarget(
-				GraphBuilder,
-				SceneTextures,
-				FamilyTexture,
-				SceneRenderer->ViewFamily,
-				SceneRenderer->Views);
-		}
-
-		// These copies become a no-op (function returns immediately) if TargetTexture and OutputTexture are the same, which
-		// is true for 2D scene captures.  Actual copies only occur for cube captures, where copying is necessary to get
-		// result data to specific slices.
-		for (const FRHICopyTextureInfo& CopyInfo : CopyInfos)
-		{
-			AddCopyTexturePass(GraphBuilder, TargetTexture, OutputTexture, CopyInfo);
-		}
-
-		if (bGenerateMips)
-		{
-			FGenerateMips::Execute(GraphBuilder, SceneRenderer->FeatureLevel, OutputTexture, GenerateMipsParams);
-		}
+		// Fix for static analysis warning -- lambda lifetime exceeds lifetime of CopyInfos.  Technically, the lambda is consumed in the
+		// scene render call below and not used afterwards, but static analysis doesn't know that, so we make a copy.
+		TArray<FRHICopyTextureInfo, FConcurrentLinearArrayAllocator> CopyInfosLocal(CopyInfos);
+
+		GraphBuilder.Blackboard.GetOrCreate<FSceneCaptureViewportSetterMap>().Emplace(TargetTexture, [CopyInfos = MoveTemp(CopyInfosLocal)](FRHICommandList& RHICmdList, int32 ViewIndex)
+		{
+			const FIntRect CopyDestRect = CopyInfos[ViewIndex].GetDestRect();
+
+			RHICmdList.SetScissorRect(false, 0, 0, 0, 0);
+			RHICmdList.SetViewport
+			(
+				float(CopyDestRect.Min.X),
+				float(CopyDestRect.Min.Y),
+				0.0f,
+				float(CopyDestRect.Max.X),
+				float(CopyDestRect.Max.Y),
+				1.0f
+			);
+		});
+	}
 
-		GraphBuilder.Execute();
+	// Disable occlusion queries when in orthographic mode
+	if (bOrthographicCamera)
+	{
+		FViewInfo& View = SceneRenderer->Views[0];
+		View.bDisableQuerySubmissions = true;
+		View.bIgnoreExistingQueries = true;
 	}
 
-	SceneRenderer->RenderThreadEnd(RHICmdList);
-}
+	// Render the scene normally
+	{
+		RDG_RHI_EVENT_SCOPE(GraphBuilder, RenderScene);
+		SceneRenderer->Render(GraphBuilder);
+	}
 
-static void UpdateSceneCaptureContent_RenderThread(
-	FRHICommandListImmediate& RHICmdList,
-	FSceneRenderer* SceneRenderer,
-	FRenderTarget* RenderTarget,
-	FTexture* RenderTargetTexture,
-	const FString& EventName,
-	TConstArrayView<FRHICopyTextureInfo> CopyInfos,
-	bool bGenerateMips,
-	const FGenerateMipsParams& GenerateMipsParams,
-	bool bClearRenderTarget,
-	bool bOrthographicCamera)
-{
-	FUniformExpressionCacheAsyncUpdateScope AsyncUpdateScope;
+	if (ShadingPath == EShadingPath::Mobile)
+	{
+		// Handles copying the SceneColor render target to the output if necessary (this happens inside the renderer for the deferred path).
+		// Other scene captures are automatically written directly to the output, in which case this function returns and does nothing.
+		const FRenderTarget* FamilyTarget = SceneRenderer->ViewFamily.RenderTarget;
+		FRDGTextureRef FamilyTexture = RegisterExternalTexture(GraphBuilder, FamilyTarget->GetRenderTargetTexture(), TEXT("OutputTexture"));
+		const FMinimalSceneTextures& SceneTextures = SceneRenderer->GetActiveSceneTextures();
+
+		RDG_EVENT_SCOPE(GraphBuilder, "CaptureSceneColor");
+		CopySceneCaptureComponentToTarget(
+			GraphBuilder,
+			SceneTextures,
+			FamilyTexture,
+			SceneRenderer->ViewFamily,
+			SceneRenderer->Views);
+	}
 
-	switch (GetFeatureLevelShadingPath(SceneRenderer->Scene->GetFeatureLevel()))
+	// These copies become a no-op (function returns immediately) if TargetTexture and OutputTexture are the same, which
+	// is true for 2D scene captures.  Actual copies only occur for cube captures, where copying is necessary to get
+	// result data to specific slices.
+	for (const FRHICopyTextureInfo& CopyInfo : CopyInfos)
 	{
-		case EShadingPath::Mobile:
-		{
-			UpdateSceneCaptureContentMobile_RenderThread(
-				RHICmdList,
-				SceneRenderer,
-				RenderTarget,
-				RenderTargetTexture,
-				EventName,
-				CopyInfos,
-				bGenerateMips,
-				GenerateMipsParams);
-			break;
-		}
-		case EShadingPath::Deferred:
-		{
-			UpdateSceneCaptureContentDeferred_RenderThread(
-				RHICmdList,
-				SceneRenderer,
-				RenderTarget,
-				RenderTargetTexture,
-				EventName,
-				CopyInfos,
-				bGenerateMips,
-				GenerateMipsParams,
-				bClearRenderTarget,
-				bOrthographicCamera);
-			break;
-		}
-		default:
-			checkNoEntry();
-			break;
+		AddCopyTexturePass(GraphBuilder, TargetTexture, OutputTexture, CopyInfo);
 	}
 
-	RHICmdList.Transition(FRHITransitionInfo(RenderTargetTexture->TextureRHI, ERHIAccess::Unknown, ERHIAccess::SRVMask));
+	if (bGenerateMips)
+	{
+		FGenerateMips::Execute(GraphBuilder, SceneRenderer->FeatureLevel, OutputTexture, GenerateMipsParams);
+	}
 }
 
 static void BuildOrthoMatrix(FIntPoint InRenderTargetSize, float InOrthoWidth, int32 InTileID, int32 InNumXTiles, int32 InNumYTiles, FMatrix& OutProjectionMatrix)
@@ -938,7 +800,7 @@ static FSceneRenderer* CreateSceneRendererForSceneCapture(
 	bool bCopyMainViewTemporalSettings2D,
 	FPostProcessSettings* PostProcessSettings,
 	float PostProcessBlendWeight,
-	const AActor* ViewActor, 
+	const AActor* ViewActor,
 	int32 CubemapFaceIndex = INDEX_NONE)
 {
 	FSceneCaptureViewInfo SceneCaptureViewInfo;
@@ -1517,11 +1379,22 @@ void FScene::UpdateSceneCaptureContents(USceneCaptureComponent2D* CaptureCompone
 
 		ENQUEUE_RENDER_COMMAND(CaptureCommand)(
 			[SceneRenderer, TextureRenderTargetResource, TexturePtrNotDeferenced, EventName, TargetName, bGenerateMips, GenerateMipsParams, GameViewportRT, bEnableOrthographicTiling, bIsCompositing, bOrthographicCamera, NumXTiles, NumYTiles, TileID, CaptureMemorySize, bCapturingGPU, CaptureSize](FRHICommandListImmediate& RHICmdList)
-			{
-				// Resize the render resource if necessary, either to the main viewport size overridden above (see ShouldRenderWithMainViewResolution()),
-				// or the original size if we are changing back to that (the resize call does nothing if the size already matches).
-				TextureRenderTargetResource->GetTextureRenderTarget2DResource()->Resize(RHICmdList, CaptureSize.X, CaptureSize.Y, bGenerateMips);
+		{
+			SceneRenderer->RenderThreadBegin(RHICmdList);
+
+		#if WANTS_DRAW_MESH_EVENTS
+			SCOPED_DRAW_EVENTF(RHICmdList, SceneCaptureMobile, TEXT("SceneCaptureMobile %s"), EventName);
+			FRDGBuilder GraphBuilder(RHICmdList, RDG_EVENT_NAME("SceneCaptureMobile %s", *EventName));
+		#else
+			SCOPED_DRAW_EVENT(RHICmdList, UpdateSceneCaptureContentMobile_RenderThread);
+			FRDGBuilder GraphBuilder(RHICmdList, RDG_EVENT_NAME("SceneCaptureMobile"));
+		#endif
 
+			// Resize the render resource if necessary, either to the main viewport size overridden above (see ShouldRenderWithMainViewResolution()),
+			// or the original size if we are changing back to that (the resize call does nothing if the size already matches).
+			TextureRenderTargetResource->GetTextureRenderTarget2DResource()->Resize(RHICmdList, CaptureSize.X, CaptureSize.Y, bGenerateMips);
+
+			{
 				RenderCaptureInterface::FScopedCapture RenderCapture(bCapturingGPU, &RHICmdList, *FString::Format(TEXT("Scene Capture : {0}"), { EventName }));
 
 				if (GameViewportRT != nullptr)
@@ -1550,22 +1423,22 @@ void FScene::UpdateSceneCaptureContents(USceneCaptureComponent2D* CaptureCompone
 
 				RectLightAtlas::FAtlasTextureInvalidationScope Invalidation(TexturePtrNotDeferenced);
 
-#if WITH_EDITOR
+			#if WITH_EDITOR
 				// Scene renderer may be deleted in UpdateSceneCaptureContent_RenderThread, grab view state pointer first
 				const FSceneViewState* ViewState = SceneRenderer->Views[0].ViewState;
-#endif  // WITH_EDITOR
+			#endif  // WITH_EDITOR
 
 				// Don't clear the render target when compositing, or in a tiling mode that fills in the render target in multiple passes.
 				bool bClearRenderTarget = !bIsCompositing && !bEnableOrthographicTiling;
 
 				VISUALIZE_TEXTURE_BEGIN_VIEW(SceneRenderer->FeatureLevel, SceneRenderer->Views[0].GetViewKey(), *EventName, true);
 
-				UpdateSceneCaptureContent_RenderThread(RHICmdList, SceneRenderer, TextureRenderTargetResource, TextureRenderTargetResource, EventName,
-					TConstArrayView<FRHICopyTextureInfo>(&CopyInfo, 1), bGenerateMips, GenerateMipsParams, bClearRenderTarget, bOrthographicCamera);
+				UpdateSceneCaptureContent_RenderThread(GraphBuilder, SceneRenderer, TextureRenderTargetResource, TextureRenderTargetResource,
+					MakeArrayView(&CopyInfo, 1), bGenerateMips, GenerateMipsParams, bClearRenderTarget, bOrthographicCamera);
 
 				VISUALIZE_TEXTURE_END_VIEW();
 
-#if WITH_EDITOR
+			#if WITH_EDITOR
 				if (ViewState)
 				{
 					const bool bLogSizes = GDumpSceneCaptureMemoryFrame == GFrameNumberRenderThread;
@@ -1579,9 +1452,11 @@ void FScene::UpdateSceneCaptureContents(USceneCaptureComponent2D* CaptureCompone
 				{
 					CaptureMemorySize->Size = 0;
 				}
-#endif  // WITH_EDITOR
+			#endif  // WITH_EDITOR
 			}
-		);
+			GraphBuilder.Execute();
+			SceneRenderer->RenderThreadEnd(RHICmdList);
+		});
 	}
 
 	EndGpuCaptureOrDump(bDumpingGPU);
@@ -1721,9 +1596,9 @@ void FScene::UpdateSceneCaptureContents(USceneCaptureComponentCube* CaptureCompo
 				ScopedCapturePtr = nullptr;
 			}
 
-			for (int32 faceidx = 0; faceidx < (int32)ECubeFace::CubeFace_MAX; faceidx++)
+			for (int32 FaceIndex = 0; FaceIndex < (int32)ECubeFace::CubeFace_MAX; FaceIndex++)
 			{
-				const ECubeFace TargetFace = (ECubeFace)faceidx;
+				const ECubeFace TargetFace = (ECubeFace)FaceIndex;
 
 				FMatrix ViewRotationMatrix;
 				FMatrix ProjectionMatrix;
@@ -1734,11 +1609,11 @@ void FScene::UpdateSceneCaptureContents(USceneCaptureComponentCube* CaptureCompo
 				FSceneRenderer* SceneRenderer = CreateSceneRendererForSceneCapture(this, CaptureComponent,
 					TextureTarget->GameThread_GetRenderTargetResource(), CaptureSize, ViewRotationMatrix,
 					Location, ProjectionMatrix, CaptureComponent->MaxViewDistanceOverride, FOVInDegrees,
-					bCaptureSceneColor, bCameraCut2D, bCopyMainViewTemporalSettings2D, &CaptureComponent->PostProcessSettings, CaptureComponent->PostProcessBlendWeight, CaptureComponent->GetViewOwner(), faceidx);
+					bCaptureSceneColor, bCameraCut2D, bCopyMainViewTemporalSettings2D, &CaptureComponent->PostProcessSettings, CaptureComponent->PostProcessBlendWeight, CaptureComponent->GetViewOwner(), FaceIndex);
 
 				// When bIsMultipleSceneCapture is true, set bIsFirstSceneRenderer to false, which tells the scene renderer it can skip RHI resource flush, saving performance.
 				// We can also skip RHI resource flush on faces after the first.
-				SceneRenderer->bIsFirstSceneRenderer = (faceidx == 0) && !bIsMultipleSceneCapture;
+				SceneRenderer->bIsFirstSceneRenderer = (FaceIndex == 0) && !bIsMultipleSceneCapture;
 
 				for (const FSceneViewExtensionRef& Extension : SceneRenderer->ViewFamily.ViewExtensions)
 				{
@@ -1746,55 +1621,67 @@ void FScene::UpdateSceneCaptureContents(USceneCaptureComponentCube* CaptureCompo
 				}
 
 				// Include the cube face index in the event name
-				EventName.Appendf(TEXT(" [%d]"), faceidx);
+				EventName.Appendf(TEXT(" [%d]"), FaceIndex);
 
 				UE::RenderCommandPipe::FSyncScope SyncScope;
 
 				ENQUEUE_RENDER_COMMAND(CaptureCommand)(
 					[SceneRenderer, TextureRenderTarget, EventName, TargetFace, bGenerateMips, GenerateMipsParams, ScopedCapturePtr](FRHICommandListImmediate& RHICmdList)
+				{
+					if (ScopedCapturePtr && (int32)TargetFace == 0)
 					{
-						if (ScopedCapturePtr && (int32)TargetFace == 0)
-						{
-							*ScopedCapturePtr = new RenderCaptureInterface::FScopedCapture(true, &RHICmdList, *FString::Format(TEXT("Scene Capture : {0}"), { EventName }));
-						}
+						*ScopedCapturePtr = new RenderCaptureInterface::FScopedCapture(true, &RHICmdList, *FString::Format(TEXT("Scene Capture : {0}"), { EventName }));
+					}
+
+					SceneRenderer->RenderThreadBegin(RHICmdList);
+
+				#if WANTS_DRAW_MESH_EVENTS
+					SCOPED_DRAW_EVENTF(RHICmdList, SceneCaptureMobile, TEXT("SceneCaptureMobile %s"), EventName);
+					FRDGBuilder GraphBuilder(RHICmdList, RDG_EVENT_NAME("SceneCaptureMobile %s", *EventName));
+				#else
+					SCOPED_DRAW_EVENT(RHICmdList, UpdateSceneCaptureContentMobile_RenderThread);
+					FRDGBuilder GraphBuilder(RHICmdList, RDG_EVENT_NAME("SceneCaptureMobile"));
+				#endif
 
 #if WITH_EDITOR
-						// Scene renderer may be deleted in UpdateSceneCaptureContent_RenderThread, grab view state pointer first
-						const FSceneViewState* ViewState = SceneRenderer->Views[0].ViewState;
+					// Scene renderer may be deleted in UpdateSceneCaptureContent_RenderThread, grab view state pointer first
+					const FSceneViewState* ViewState = SceneRenderer->Views[0].ViewState;
 #endif  // WITH_EDITOR
 
-						VISUALIZE_TEXTURE_BEGIN_VIEW(SceneRenderer->FeatureLevel, SceneRenderer->Views[0].GetViewKey(), *EventName, true);
+					VISUALIZE_TEXTURE_BEGIN_VIEW(SceneRenderer->FeatureLevel, SceneRenderer->Views[0].GetViewKey(), *EventName, true);
 
-						// We need to generate mips on last cube face
-						const bool bLastCubeFace = ((int32)TargetFace == (int32)ECubeFace::CubeFace_MAX - 1);
+					// We need to generate mips on last cube face
+					const bool bLastCubeFace = ((int32)TargetFace == (int32)ECubeFace::CubeFace_MAX - 1);
 
-						FRHICopyTextureInfo CopyInfo;
-						CopyInfo.DestSliceIndex = TargetFace;
-						UpdateSceneCaptureContent_RenderThread(RHICmdList, SceneRenderer, TextureRenderTarget, TextureRenderTarget, EventName,
-							TConstArrayView<FRHICopyTextureInfo>(&CopyInfo, 1), bGenerateMips && bLastCubeFace, GenerateMipsParams, true, false);
+					FRHICopyTextureInfo CopyInfo;
+					CopyInfo.DestSliceIndex = TargetFace;
+					UpdateSceneCaptureContent_RenderThread(GraphBuilder, SceneRenderer, TextureRenderTarget, TextureRenderTarget,
+						MakeArrayView(&CopyInfo, 1), bGenerateMips && bLastCubeFace, GenerateMipsParams, true, false);
 
-						VISUALIZE_TEXTURE_END_VIEW();
+					VISUALIZE_TEXTURE_END_VIEW();
 
 #if WITH_EDITOR
-						if (ViewState)
+					if (ViewState)
+					{
+						const bool bLogSizes = GDumpSceneCaptureMemoryFrame == GFrameNumberRenderThread;
+						if (bLogSizes)
 						{
-							const bool bLogSizes = GDumpSceneCaptureMemoryFrame == GFrameNumberRenderThread;
-							if (bLogSizes)
-							{
-								UE_LOG(LogRenderer, Log, TEXT("LogSizes\tSceneCaptureCube[%d]\t%s\t%dx%d"), TargetFace, *EventName, TextureRenderTarget->GetSizeX(), TextureRenderTarget->GetSizeY());
-								ViewState->GetGPUSizeBytes(bLogSizes);
-							}
+							UE_LOG(LogRenderer, Log, TEXT("LogSizes\tSceneCaptureCube[%d]\t%s\t%dx%d"), TargetFace, *EventName, TextureRenderTarget->GetSizeX(), TextureRenderTarget->GetSizeY());
+							ViewState->GetGPUSizeBytes(bLogSizes);
 						}
+					}
 #endif  // WITH_EDITOR
 
-						if (ScopedCapturePtr && (int32)TargetFace == (int32)ECubeFace::CubeFace_MAX - 1)
-						{
-							// Delete the scope, and the persistent pointer we allocated on the heap to share across the render command lambdas
-							delete *ScopedCapturePtr;
-							delete ScopedCapturePtr;
-						}
+					GraphBuilder.Execute();
+					SceneRenderer->RenderThreadEnd(RHICmdList);
+
+					if (ScopedCapturePtr && (int32)TargetFace == (int32)ECubeFace::CubeFace_MAX - 1)
+					{
+						// Delete the scope, and the persistent pointer we allocated on the heap to share across the render command lambdas
+						delete *ScopedCapturePtr;
+						delete ScopedCapturePtr;
 					}
-				);
+				});
 
 				// Trim the cube face index from the event name
 				EventName.LeftChopInline(4, EAllowShrinking::No);
@@ -1939,8 +1826,21 @@ void FScene::UpdateSceneCaptureContents(USceneCaptureComponentCube* CaptureCompo
 
 			ENQUEUE_RENDER_COMMAND(CaptureAllCubeFaces)(
 				[SceneRenderer, CubeFaceTarget, TextureRenderTarget, EventName, CaptureSize, bGenerateMips, GenerateMipsParams, bCapturingGPU](FRHICommandListImmediate& RHICmdList)
+			{
+				SceneRenderer->RenderThreadBegin(RHICmdList);
+
+			#if WANTS_DRAW_MESH_EVENTS
+				SCOPED_DRAW_EVENTF(RHICmdList, SceneCaptureMobile, TEXT("SceneCaptureMobile %s"), EventName);
+				FRDGBuilder GraphBuilder(RHICmdList, RDG_EVENT_NAME("SceneCaptureMobile %s", *EventName));
+			#else
+				SCOPED_DRAW_EVENT(RHICmdList, UpdateSceneCaptureContentMobile_RenderThread);
+				FRDGBuilder GraphBuilder(RHICmdList, RDG_EVENT_NAME("SceneCaptureMobile"));
+			#endif
+
+				CubeFaceTarget->InitRHI(RHICmdList);
+
 				{
-					RenderCaptureInterface::FScopedCapture RenderCapture(bCapturingGPU, &RHICmdList, *FString::Format(TEXT("Scene Capture : {0}"), { EventName }));
+					RenderCaptureInterface::FScopedCapture RenderCapture(bCapturingGPU, GraphBuilder, *FString::Format(TEXT("Scene Capture : {0}"), { EventName }));
 
 					TStaticArray<FRHICopyTextureInfo, (int32)ECubeFace::CubeFace_MAX> CopyInfos;
 					for (int32 faceidx = 0; faceidx < (int32)ECubeFace::CubeFace_MAX; faceidx++)
@@ -1952,25 +1852,23 @@ void FScene::UpdateSceneCaptureContents(USceneCaptureComponentCube* CaptureCompo
 						CopyInfos[faceidx].DestSliceIndex = faceidx;
 					}
 
-					CubeFaceTarget->InitRHI(RHICmdList);
-
-#if WITH_EDITOR
+				#if WITH_EDITOR
 					// Scene renderer may be deleted in UpdateSceneCaptureContent_RenderThread, grab view state pointer first
 					TStaticArray<const FSceneViewState*, (int32)ECubeFace::CubeFace_MAX> SceneViewStates;
 					for (int32 FaceIdx = 0; FaceIdx < (int32)ECubeFace::CubeFace_MAX; FaceIdx++)
 					{
 						SceneViewStates[FaceIdx] = SceneRenderer->Views[FaceIdx].ViewState;
 					}
-#endif  // WITH_EDITOR
+				#endif  // WITH_EDITOR
 
 					VISUALIZE_TEXTURE_BEGIN_VIEW(SceneRenderer->FeatureLevel, SceneRenderer->Views[0].GetViewKey(), *EventName, true);
 
-					UpdateSceneCaptureContent_RenderThread(RHICmdList, SceneRenderer, CubeFaceTarget, TextureRenderTarget, EventName,
+					UpdateSceneCaptureContent_RenderThread(GraphBuilder, SceneRenderer, CubeFaceTarget, TextureRenderTarget,
 						CopyInfos, bGenerateMips, GenerateMipsParams, true, false);
 
 					VISUALIZE_TEXTURE_END_VIEW();
 
-#if WITH_EDITOR
+				#if WITH_EDITOR
 					if (SceneViewStates[0])
 					{
 						const bool bLogSizes = GDumpSceneCaptureMemoryFrame == GFrameNumberRenderThread;
@@ -1983,11 +1881,14 @@ void FScene::UpdateSceneCaptureContents(USceneCaptureComponentCube* CaptureCompo
 							}
 						}
 					}
-#endif  // WITH_EDITOR
-
-					delete CubeFaceTarget;
+				#endif  // WITH_EDITOR
 				}
-			);
+
+				delete CubeFaceTarget;
+
+				GraphBuilder.Execute();
+				SceneRenderer->RenderThreadEnd(RHICmdList);
+			});
 		}
 	}
 
