MGPU:  Added functionality to split screen debug mode (r.SplitScreenDebug.Enable) to render the two panes on different GPUs, when MGPU is enabled.  Designed to support running MGPU unit tests or debugging in any project, without needing to write code or set up nDisplay content and configuration files.  Also fixed issues with MGPU split screen, where transfers weren't using the view rect, and lower level code wasn't propagating the transfer rect from the intermediate to final struct.  DumpGPU changed to display the actual mask rather than hardware GPU mask (always 1, indicating the first GPU) when -vmgpu is active, to make such dumps readable in terms of what's running on a given virtual GPU.

#rnx

diff --git a/Engine/Source/Runtime/RHI/Public/MultiGPU.h b/Engine/Source/Runtime/RHI/Public/MultiGPU.h
--- a/Engine/Source/Runtime/RHI/Public/MultiGPU.h
+++ b/Engine/Source/Runtime/RHI/Public/MultiGPU.h
@@ -162,6 +162,12 @@ public:
 		return GVirtualMGPU ? 1 : GetMask();
 	}
 
+	// Direct use of the internal mask is discouraged, but it can be useful for debugging to display
+	FORCEINLINE SGPU_CONSTEXPR uint32 GetForDisplay() const
+	{
+		return GetMask();
+	}
+
 	FORCEINLINE SGPU_CONSTEXPR FRHIGPUMask operator &(FRHIGPUMask Rhs) const
 	{
 		return FRHIGPUMask(GetMask() & Rhs.GetMask());
diff --git a/Engine/Source/Runtime/RHI/Public/RHIContext.h b/Engine/Source/Runtime/RHI/Public/RHIContext.h
--- a/Engine/Source/Runtime/RHI/Public/RHIContext.h
+++ b/Engine/Source/Runtime/RHI/Public/RHIContext.h
@@ -169,6 +169,27 @@ struct FTransferResourceParams
 		check(InBuffer);
 	}
 
+	void SetRect(const FIntRect& Rect)
+	{
+		// Do nothing if this isn't a texture
+		if (!Texture)
+		{
+			return;
+		}
+
+		// Reset the extents to default (zero), if the rect is empty or covers the full extent of the texture.  Forces a more efficient full resource copy.
+		if (Rect.IsEmpty() || (Rect.Min.X == 0 && Rect.Min.Y == 0 && Rect.Max.X == Texture->GetDesc().Extent.X && Rect.Max.Y == Texture->GetDesc().Extent.Y))
+		{
+			Min = FIntVector3(0);
+			Max = FIntVector3(0);
+			return;
+		}
+
+		// Initialize the extent
+		Min = FIntVector3(Rect.Min.X, Rect.Min.Y, 0);
+		Max = FIntVector3(Rect.Max.X, Rect.Max.Y, 1);
+	}
+
 	// The texture which must be must be allocated on both GPUs 
 	FTextureRHIRef Texture;
 	// Or alternately, a buffer that's allocated on both GPUs
diff --git a/Engine/Source/Runtime/RenderCore/Private/DumpGPU.cpp b/Engine/Source/Runtime/RenderCore/Private/DumpGPU.cpp
--- a/Engine/Source/Runtime/RenderCore/Private/DumpGPU.cpp
+++ b/Engine/Source/Runtime/RenderCore/Private/DumpGPU.cpp
@@ -2352,7 +2352,7 @@ static const TCHAR* GetPassEventNameWithGPUMask(const FRDGPass* Pass, FString& O
 		// Prepend GPU mask on the event name of each pass, so you can see which GPUs the pass ran on.  Putting the mask at the
 		// front rather than the back makes all the masks line up, and easier to read (or ignore if you don't care about them).
 		// Also, it's easy to globally search for passes with a particular GPU mask using name search in the dump browser.
-		OutNameStorage = FString::Printf(TEXT("[%x] %s"), Pass->GetGPUMask().GetNative(), Pass->GetEventName().GetTCHAR());
+		OutNameStorage = FString::Printf(TEXT("[%x] %s"), Pass->GetGPUMask().GetForDisplay(), Pass->GetEventName().GetTCHAR());
 		return *OutNameStorage;
 	}
 	else
diff --git a/Engine/Source/Runtime/Renderer/Private/SceneRendering.cpp b/Engine/Source/Runtime/Renderer/Private/SceneRendering.cpp
--- a/Engine/Source/Runtime/Renderer/Private/SceneRendering.cpp
+++ b/Engine/Source/Runtime/Renderer/Private/SceneRendering.cpp
@@ -203,6 +203,12 @@ static TAutoConsoleVariable<int32> CVarSplitScreenDebugLumenScene(
 	TEXT("For split screen debugging, allocate a separate Lumen scene for the second view."),
 	ECVF_Default);
 
+static TAutoConsoleVariable<int32> CVarSplitScreenDebugMGPU(
+	TEXT("r.SplitScreenDebug.MGPU"),
+	1,
+	TEXT("For split screen debugging, if MGPU is active, run the second view on the other GPU.  Uses two scene renderers."),
+	ECVF_Default);
+
 #endif  // !UE_BUILD_SHIPPING
 
 bool UseCachedMeshDrawCommands()
@@ -3661,6 +3667,23 @@ static void GetCrossGPUTransfers(FSceneRenderer* SceneRenderer, TArray<FCrossGPU
 					if (!ViewInfo.GPUMask.Contains(RenderTargetGPUIndex))
 					{
 						OutTransfers.Add(FCrossGPUTransfer(TransferRect, ViewInfo.GPUMask.GetFirstIndex(), RenderTargetGPUIndex));
+
+						// If multiple families write to the same render target via MGPU, mask cross GPU copies after the first view family
+						// to the view rect.
+						for (const FSceneViewFamily* ViewFamily : SceneRenderer->AllFamilies)
+						{
+							if (ViewFamily == &SceneRenderer->ViewFamily)
+							{
+								// Exit if we found the current view family
+								break;
+							}
+							else if (ViewFamily->RenderTarget == SceneRenderer->ViewFamily.RenderTarget)
+							{
+								// We found another view family writing to the same target, set the TransferRect
+								OutTransfers.Last().TransferRect = ViewInfo.UnscaledViewRect;
+								break;
+							}
+						}
 					}
 				}
 			}
@@ -3759,6 +3782,7 @@ void FSceneRenderer::DoCrossGPUTransfers(FRDGBuilder& GraphBuilder, FRDGTextureR
 				for (const FCrossGPUTransfer& Transfer : LocalTransfers)
 				{
 					TransferParams.Add(FTransferResourceParams(RenderTargetTexture->GetRHI(), Transfer.SrcGPUIndex, Transfer.DestGPUIndex, false, false));
+					TransferParams.Last().SetRect(Transfer.TransferRect);
 				}
 
 				// Transition resources on destination GPU and signal when transition has finished
@@ -3780,6 +3804,7 @@ void FSceneRenderer::DoCrossGPUTransfers(FRDGBuilder& GraphBuilder, FRDGTextureR
 				for (const FCrossGPUTransfer& Transfer : LocalTransfers)
 				{
 					TransferParams.Add(FTransferResourceParams(RenderTargetTexture->GetRHI(), Transfer.SrcGPUIndex, Transfer.DestGPUIndex, true, false));
+					TransferParams.Last().SetRect(Transfer.TransferRect);
 				}
 
 				RHICmdList.TransferResources(TransferParams);
@@ -3827,6 +3852,7 @@ void FSceneRenderer::FlushCrossGPUTransfers(FRDGBuilder& GraphBuilder)
 					for (const FCrossGPUTransfer& Transfer : Target.Transfers)
 					{
 						TransferParams.Add(FTransferResourceParams(PassParameters->Textures[TargetIndex].GetTexture()->GetRHI(), Transfer.SrcGPUIndex, Transfer.DestGPUIndex, true, true));
+						TransferParams.Last().SetRect(Transfer.TransferRect);
 					}
 				}
 
@@ -4692,11 +4718,40 @@ void FSceneRenderer::SetupMeshPass(FViewInfo& View, FExclusiveDepthStencil::Type
 }
 
 #if !UE_BUILD_SHIPPING
-FSceneViewFamily* FSceneRenderer::CreateSplitScreenDebugViewFamily(const FSceneViewFamily& InFamily)
+void FSceneRenderer::CreateSplitScreenDebugViewFamilies(const FSceneViewFamily& InFamily, TArrayView<const FSceneViewFamily*>& OutFamilies, TArray<const FSceneViewFamily*, TFixedAllocator<2>>& OutFamiliesStorage)
 {
-	FSceneViewFamily* Family = new FSceneViewFamily(InFamily);
-	Family->ScreenPercentageInterface = InFamily.ScreenPercentageInterface->Fork_GameThread(InFamily);
-	Family->Views.SetNumZeroed(2);
+	// We either generate a single family with 2 views, or two families with 1 view each for MGPU
+	int32 NumFamilies;
+	int32 NumViewsPerFamily;
+	if (GNumExplicitGPUsForRendering > 1 && CVarSplitScreenDebugMGPU.GetValueOnGameThread())
+	{
+		NumFamilies = 2;
+		NumViewsPerFamily = 1;
+	}
+	else
+	{
+		NumFamilies = 1;
+		NumViewsPerFamily = 2;
+	}
+
+	TArray<FSceneViewFamily*, TFixedAllocator<2>> Families;
+	TArray<const FSceneView**, TFixedAllocator<2>> ViewPointers;
+	TArray<FSceneViewFamily*, TFixedAllocator<2>> ViewParents;
+
+	for (int32 FamilyIndex = 0; FamilyIndex < NumFamilies; FamilyIndex++)
+	{
+		FSceneViewFamily* Family = new FSceneViewFamily(InFamily);
+
+		Families.Add(Family);
+		Family->ScreenPercentageInterface = InFamily.ScreenPercentageInterface->Fork_GameThread(InFamily);
+		Family->Views.SetNumZeroed(NumViewsPerFamily);
+
+		for (int32 ViewIndex = 0; ViewIndex < NumViewsPerFamily; ViewIndex++)
+		{
+			ViewPointers.Add(&Family->Views[ViewIndex]);
+			ViewParents.Add(Family);
+		}
+	}
 
 	FIntRect OriginalViewRect = InFamily.Views[0]->SceneViewInitOptions.ViewRect;
 
@@ -4802,7 +4857,7 @@ FSceneViewFamily* FSceneRenderer::CreateSplitScreenDebugViewFamily(const FSceneV
 		InitOptions.SetViewRectangle(ViewRect);
 
 		// Set view family to dynamically allocated copy
-		InitOptions.ViewFamily = Family;
+		InitOptions.ViewFamily = ViewParents[ViewIndex];
 
 		// Use new static view state for second view
 		if (ViewIndex == 1)
@@ -4820,37 +4875,59 @@ FSceneViewFamily* FSceneRenderer::CreateSplitScreenDebugViewFamily(const FSceneV
 			// Add or remove optional Lumen scene for second view state
 			if (CVarSplitScreenDebugLumenScene.GetValueOnGameThread())
 			{
-				GSecondViewState->AddLumenSceneData(Family->Scene, 1.0f);
+				GSecondViewState->AddLumenSceneData(ViewParents[ViewIndex]->Scene, 1.0f);
 			}
 			else
 			{
-				GSecondViewState->RemoveLumenSceneData(Family->Scene);
+				GSecondViewState->RemoveLumenSceneData(ViewParents[ViewIndex]->Scene);
 			}
 
 			InitOptions.SceneViewStateInterface = GSecondViewState;
 		}
 
 		FSceneView* View = new FSceneView(InitOptions);
-		View->PrimaryViewIndex = ViewIndex;
+		View->PrimaryViewIndex = NumViewsPerFamily == 1 ? 0 : ViewIndex;
 		View->FinalPostProcessSettings = InFamily.Views[0]->FinalPostProcessSettings;
 
-		Family->Views[ViewIndex] = View;
+		*ViewPointers[ViewIndex] = View;
+
+		// Set up second view for multi-GPU
+		if (NumFamilies > 1 && ViewIndex == 1)
+		{
+			FSceneViewFamily* Family = ViewParents[ViewIndex];
+
+			// Prevent the render target from being cleared, and enable cross GPU transfers
+			Family->bAdditionalViewFamily = true;
+			Family->bMultiGPUForkAndJoin = true;
+
+			// Set the view to run on the second GPU
+			View->bOverrideGPUMask = true;
+			View->GPUMask = FRHIGPUMask::FromIndex(ViewIndex);
+		}
 	}
 
-	return Family;
+	// Copy the view families to the const output array
+	for (FSceneViewFamily* Family : Families)
+	{
+		OutFamiliesStorage.Add(Family);
+	}
+	OutFamilies = OutFamiliesStorage;
 }
 
-void FSceneRenderer::DestroySplitScreenDebugViewFamily(FSceneViewFamily* Family)
+void FSceneRenderer::DestroySplitScreenDebugViewFamilies(TArrayView<const FSceneViewFamily*> Families)
 {
-	for (int32 ViewIndex = 0; ViewIndex < Family->Views.Num(); ViewIndex++)
+	for (const FSceneViewFamily* Family : Families)
 	{
-		delete Family->Views[ViewIndex];
+		for (int32 ViewIndex = 0; ViewIndex < Family->Views.Num(); ViewIndex++)
+		{
+			delete Family->Views[ViewIndex];
+		}
+		delete Family;
 	}
-	delete Family;
 }
 #endif  // !UE_BUILD_SHIPPING
 
-void FSceneRenderer::CreateSceneRenderers(TArrayView<const FSceneViewFamily*> InViewFamilies, FHitProxyConsumer* HitProxyConsumer, TArray<FSceneRenderer*>& OutSceneRenderers)
+void FSceneRenderer::CreateSceneRenderers(TArrayView<const FSceneViewFamily*> InViewFamilies, FHitProxyConsumer* HitProxyConsumer, bool bAllowSplitScreenDebug, TArray<FSceneRenderer*>& OutSceneRenderers)
 {
 	OutSceneRenderers.Empty(InViewFamilies.Num());
 
@@ -4861,13 +4938,11 @@ void FSceneRenderer::CreateSceneRenderers(TArrayView<const FSceneViewFamily*> In
 
 #if !UE_BUILD_SHIPPING
 	bool bSplitScreenDebug = false;
-	if (CVarSplitScreenDebugEnable.GetValueOnGameThread() > 0)
+	TArray<const FSceneViewFamily*, TFixedAllocator<2>> DebugViewFamilyStorage;
+	if (CVarSplitScreenDebugEnable.GetValueOnGameThread() > 0 && bAllowSplitScreenDebug && InViewFamilies.Num() == 1 && InViewFamilies[0]->bSplitScreenDebugAllowed && InViewFamilies[0]->Views.Num() == 1)
 	{
-		if (InViewFamilies.Num() == 1 && InViewFamilies[0]->bSplitScreenDebugAllowed && InViewFamilies[0]->Views.Num() == 1)
-		{
-			InViewFamilies[0] = CreateSplitScreenDebugViewFamily(*InViewFamilies[0]);
-			bSplitScreenDebug = true;
-		}
+		CreateSplitScreenDebugViewFamilies(*InViewFamilies[0], InViewFamilies, DebugViewFamilyStorage);
+		bSplitScreenDebug = true;
 	}
 #endif
 
@@ -4906,7 +4981,7 @@ void FSceneRenderer::CreateSceneRenderers(TArrayView<const FSceneViewFamily*> In
 #if !UE_BUILD_SHIPPING
 	if (bSplitScreenDebug)
 	{
-		DestroySplitScreenDebugViewFamily(const_cast<FSceneViewFamily*>(InViewFamilies[0]));
+		DestroySplitScreenDebugViewFamilies(InViewFamilies);
 	}
 #endif
 }
@@ -4920,7 +4995,7 @@ FSceneRenderer* FSceneRenderer::CreateSceneRenderer(const FSceneViewFamily* InVi
 	TArray<const FSceneViewFamily*> InViewFamilies;
 	InViewFamilies.Add(InViewFamily);
 
-	CreateSceneRenderers(InViewFamilies, HitProxyConsumer, SceneRenderers);
+	CreateSceneRenderers(InViewFamilies, HitProxyConsumer, /*bAllowSplitScreenDebug=*/ false, SceneRenderers);
 
 	return SceneRenderers[0];
 }
@@ -5667,7 +5742,7 @@ void FRendererModule::BeginRenderingViewFamilies(FCanvas* Canvas, TConstArrayVie
 			SceneCaptureUpdateDeferredCapturesInternal(Scene, ViewFamiliesConst);
 		}
 
-		FSceneRenderer::CreateSceneRenderers(ViewFamiliesConst, Canvas->GetHitProxyConsumer(), SceneRenderers);
+		FSceneRenderer::CreateSceneRenderers(ViewFamiliesConst, Canvas->GetHitProxyConsumer(), /*bAllowSplitScreenDebug=*/ true, SceneRenderers);
 
 		if (!bShowHitProxies)
 		{
diff --git a/Engine/Source/Runtime/Renderer/Private/SceneRendering.h b/Engine/Source/Runtime/Renderer/Private/SceneRendering.h
--- a/Engine/Source/Runtime/Renderer/Private/SceneRendering.h
+++ b/Engine/Source/Runtime/Renderer/Private/SceneRendering.h
@@ -2125,7 +2125,7 @@ public:
 	RENDERER_API static FSceneRenderer* CreateSceneRenderer(const FSceneViewFamily* InViewFamily, FHitProxyConsumer* HitProxyConsumer);
 
 	/** Creates multiple scene renderers based on the current feature level.  All view families must point to the same Scene. */
-	RENDERER_API static void CreateSceneRenderers(TArrayView<const FSceneViewFamily*> InViewFamilies, FHitProxyConsumer* HitProxyConsumer, TArray<FSceneRenderer*>& OutSceneRenderers);
+	RENDERER_API static void CreateSceneRenderers(TArrayView<const FSceneViewFamily*> InViewFamilies, FHitProxyConsumer* HitProxyConsumer, bool bAllowSplitScreenDebug, TArray<FSceneRenderer*>& OutSceneRenderers);
 
 	/** Setups FViewInfo::ViewRect according to ViewFamilly's ScreenPercentageInterface. */
 	void PrepareViewRectsForRendering(FRHICommandListImmediate& RHICmdList);
@@ -2556,8 +2556,8 @@ protected:
 	void RenderMeshDistanceFieldVisualization(FRDGBuilder& GraphBuilder, const FMinimalSceneTextures& SceneTextures);
 
 #if !UE_BUILD_SHIPPING
-	static FSceneViewFamily* CreateSplitScreenDebugViewFamily(const FSceneViewFamily& InFamily);
-	static void DestroySplitScreenDebugViewFamily(FSceneViewFamily* Family);
+	static void CreateSplitScreenDebugViewFamilies(const FSceneViewFamily& InFamily, TArrayView<const FSceneViewFamily*>& OutFamilies, TArray<const FSceneViewFamily*, TFixedAllocator<2>>& OutFamilyStorage);
+	static void DestroySplitScreenDebugViewFamilies(TArrayView<const FSceneViewFamily*> Families);
 #endif
 
 protected:
