Scene proxy instance management cleanup.
* removed dependency on UInstancedStaticMeshComponent/FInstancedStaticMeshSceneDesc from scene interface.
* removed StaticMeshBounds from UpdateInstances_RenderThread as it is unused and tied instancing to the static mesh for no reason and deprecated the old version.
* removed some previously deprecated interface from FPrimitiveSceneProxy.

diff --git a/Engine/Source/Runtime/Engine/Classes/Engine/InstancedStaticMesh.h b/Engine/Source/Runtime/Engine/Classes/Engine/InstancedStaticMesh.h
--- a/Engine/Source/Runtime/Engine/Classes/Engine/InstancedStaticMesh.h
+++ b/Engine/Source/Runtime/Engine/Classes/Engine/InstancedStaticMesh.h
@@ -442,7 +442,7 @@ public:
 
 	ENGINE_API virtual void DestroyRenderThreadResources() override;
 
-	ENGINE_API virtual void UpdateInstances_RenderThread(FRHICommandListBase& RHICmdList, const FBoxSphereBounds& InBounds, const FBoxSphereBounds& InLocalBounds, const FBoxSphereBounds& InStaticMeshBounds) override;
+	ENGINE_API virtual void UpdateInstances_RenderThread(FRHICommandListBase& RHICmdList, const FBoxSphereBounds& InBounds, const FBoxSphereBounds& InLocalBounds) override;
 
 	virtual FPrimitiveViewRelevance GetViewRelevance(const FSceneView* View) const override
 	{
@@ -497,7 +497,6 @@ public:
 	ENGINE_API virtual bool GetWireframeMeshElement(int32 LODIndex, int32 BatchIndex, const FMaterialRenderProxy* WireframeRenderProxy, uint8 InDepthPriorityGroup, bool bAllowPreCulledIndices, FMeshBatch& OutMeshBatch) const override;
 
 	ENGINE_API virtual void GetDistanceFieldAtlasData(const FDistanceFieldVolumeData*& OutDistanceFieldData, float& SelfShadowBias) const override;
-	ENGINE_API virtual void GetDistanceFieldInstanceData(TArray<FRenderTransform>& InstanceLocalToPrimitiveTransforms) const override;
 
 	/**
 	 * Creates the hit proxies are used when DrawDynamicElements is called.
diff --git a/Engine/Source/Runtime/Engine/Private/InstancedStaticMesh.cpp b/Engine/Source/Runtime/Engine/Private/InstancedStaticMesh.cpp
--- a/Engine/Source/Runtime/Engine/Private/InstancedStaticMesh.cpp
+++ b/Engine/Source/Runtime/Engine/Private/InstancedStaticMesh.cpp
@@ -1460,14 +1460,14 @@ void FInstancedStaticMeshSceneProxy::DestroyRenderThreadResources()
 #endif
 }
 
-void FInstancedStaticMeshSceneProxy::UpdateInstances_RenderThread(FRHICommandListBase& RHICmdList, const FBoxSphereBounds& InBounds, const FBoxSphereBounds& InLocalBounds, const FBoxSphereBounds& InStaticMeshBounds)
+void FInstancedStaticMeshSceneProxy::UpdateInstances_RenderThread(FRHICommandListBase& RHICmdList, const FBoxSphereBounds& InBounds, const FBoxSphereBounds& InLocalBounds)
 {
 	if (FStaticMeshInstanceBuffer *LegacyInstanceBuffer = InstanceDataSceneProxy->GetLegacyInstanceBuffer())
 	{
 		InstancedRenderData.BindBuffersToVertexFactories(RHICmdList, LegacyInstanceBuffer);
 	}
 
-	return FPrimitiveSceneProxy::UpdateInstances_RenderThread(RHICmdList, InBounds, InLocalBounds, InStaticMeshBounds);
+	return FPrimitiveSceneProxy::UpdateInstances_RenderThread(RHICmdList, InBounds, InLocalBounds);
 }
 
 void FInstancedStaticMeshSceneProxy::SetupInstancedMeshBatch(int32 LODIndex, int32 BatchIndex, FMeshBatch& OutMeshBatch) const
@@ -1547,11 +1547,6 @@ void FInstancedStaticMeshSceneProxy::GetDistanceFieldAtlasData(const FDistanceFi
 	FStaticMeshSceneProxy::GetDistanceFieldAtlasData(OutDistanceFieldData, SelfShadowBias);
 }
 
-void FInstancedStaticMeshSceneProxy::GetDistanceFieldInstanceData(TArray<FRenderTransform>& InstanceLocalToPrimitiveTransforms) const
-{
-	check(false);
-}
-
 HHitProxy* FInstancedStaticMeshSceneProxy::CreateHitProxies(UPrimitiveComponent* Component,TArray<TRefCountPtr<HHitProxy> >& OutHitProxies)
 {
 	return FInstancedStaticMeshSceneProxy::CreateHitProxies(Component->GetPrimitiveComponentInterface(), OutHitProxies);
diff --git a/Engine/Source/Runtime/Engine/Private/PrimitiveSceneProxy.cpp b/Engine/Source/Runtime/Engine/Private/PrimitiveSceneProxy.cpp
--- a/Engine/Source/Runtime/Engine/Private/PrimitiveSceneProxy.cpp
+++ b/Engine/Source/Runtime/Engine/Private/PrimitiveSceneProxy.cpp
@@ -972,7 +972,17 @@ void FPrimitiveSceneProxy::SetTransform(FRHICommandListBase& RHICmdList, const F
 void FPrimitiveSceneProxy::UpdateInstances_RenderThread(FRHICommandListBase& RHICmdList, const FBoxSphereBounds& InBounds, const FBoxSphereBounds& InLocalBounds, const FBoxSphereBounds& InStaticMeshBounds)
 {
 	TRACE_CPUPROFILER_EVENT_SCOPE_STR("FPrimitiveSceneProxy::UpdateInstances_RenderThread");
+	ensureMsgf(false, TEXT("Incorrect call to UpdateInstances_RenderThread, please update the code to call the correct version."));
+	// Update the cached bounds.
+	Bounds = InBounds;
+	LocalBounds = InLocalBounds;
+
+	bAlwaysHasVelocity = CVarVelocityForceOutput.GetValueOnAnyThread() ||  GetInstanceDataHeader().Flags.bHasPerInstanceDynamicData;
+}
 
+void FPrimitiveSceneProxy::UpdateInstances_RenderThread(FRHICommandListBase& RHICmdList, const FBoxSphereBounds& InBounds, const FBoxSphereBounds& InLocalBounds)
+{
+	TRACE_CPUPROFILER_EVENT_SCOPE_STR("FPrimitiveSceneProxy::UpdateInstances_RenderThread");
 	// Update the cached bounds.
 	Bounds = InBounds;
 	LocalBounds = InLocalBounds;
diff --git a/Engine/Source/Runtime/Engine/Public/PrimitiveSceneProxy.h b/Engine/Source/Runtime/Engine/Public/PrimitiveSceneProxy.h
--- a/Engine/Source/Runtime/Engine/Public/PrimitiveSceneProxy.h
+++ b/Engine/Source/Runtime/Engine/Public/PrimitiveSceneProxy.h
@@ -502,7 +502,7 @@ public:
 		SelfShadowBias = 0;
 	}
 	UE_DEPRECATED(5.4, "Use generic instance data through GetInstanceSceneDataBuffers().")
-	virtual void GetDistanceFieldInstanceData(TArray<FRenderTransform>& InstanceLocalToPrimitiveTransforms) const
+	virtual void GetDistanceFieldInstanceData(TArray<FRenderTransform>& InstanceLocalToPrimitiveTransforms) const final
 	{
 	}
 
@@ -523,9 +523,6 @@ public:
 	 */
 	virtual void CreateRenderThreadResources(FRHICommandListBase& RHICmdList) {}
 
-	UE_DEPRECATED(5.4, "CreateRenderThreadResources now requires a command list.")
-	virtual void CreateRenderThreadResources() final {}
-
 	/**
 	 *	Called when the rendering thread removes the proxy from the scene.
 	 *	This function allows for removing renderer-side resources.
@@ -548,9 +545,6 @@ public:
 	 */
 	virtual void OnTransformChanged(FRHICommandListBase& RHICmdList) { }
 
-	UE_DEPRECATED(5.4, "OnTransformChanged now takes a command list.")
-	void OnTransformChanged() { OnTransformChanged(FRHICommandListImmediate::Get()); }
-
 	/**
 	 * Called to notify the proxy that the level has been fully added to
 	 * the world and the primitive will now be rendered.
@@ -1012,29 +1006,17 @@ public:
 	 */
 	ENGINE_API virtual void ApplyWorldOffset(FRHICommandListBase& RHICmdList, FVector InOffset);
 
-	UE_DEPRECATED(5.4, "ApplyWorldOffset now takes a command list.")
-	void ApplyWorldOffset(FVector InOffset) { ApplyWorldOffset(FRHICommandListImmediate::Get(), InOffset); }
-
 	/**
 	 * Applies a "late in the frame" adjustment to the proxy's existing transform
 	 * @param LateUpdateTransform - The post-transform to be applied to the LocalToWorld matrix
 	 */
 	ENGINE_API virtual void ApplyLateUpdateTransform(FRHICommandListBase& RHICmdList, const FMatrix& LateUpdateTransform);
 
-	UE_DEPRECATED(5.4, "ApplyWorldOffset now takes a command list.")
-	void ApplyLateUpdateTransform(const FMatrix& LateUpdateTransform) { ApplyLateUpdateTransform(FRHICommandListImmediate::Get(), LateUpdateTransform); }
-
 	/**
 	 * Updates the primitive proxy's uniform buffer.
 	 */
 	ENGINE_API void UpdateUniformBuffer(FRHICommandList& RHICmdList);
 
-	UE_DEPRECATED(5.3, "UpdateUniformBuffer now takes a command list.")
-	inline void UpdateUniformBuffer()
-	{
-		UpdateUniformBuffer(FRHICommandListImmediate::Get());
-	}
-
 	/**
 	 * Apply the unform shader parameter settings for the proxy to the builder.
 	 */
@@ -1671,7 +1653,9 @@ protected:
 	 * @param InLocalBounds - Primitive local space bounds.
 	 * @param InStaticMeshBounds - Bounds of the primitive mesh instance.
 	 */
-	ENGINE_API virtual void UpdateInstances_RenderThread(FRHICommandListBase& RHICmdList, const FBoxSphereBounds& InBounds, const FBoxSphereBounds& InLocalBounds, const FBoxSphereBounds& InStaticMeshBounds);
+	UE_DEPRECATED(5.6, "StaticMeshBounds is no longer a parameter of this interface, please update the implementation.")
+	ENGINE_API virtual void UpdateInstances_RenderThread(FRHICommandListBase& RHICmdList, const FBoxSphereBounds& InBounds, const FBoxSphereBounds& InLocalBounds, const FBoxSphereBounds& InStaticMeshBounds) final;
+	ENGINE_API virtual void UpdateInstances_RenderThread(FRHICommandListBase& RHICmdList, const FBoxSphereBounds& InBounds, const FBoxSphereBounds& InLocalBounds);
 
 	/** Updates selection for the primitive proxy. This is called in the rendering thread by SetSelection_GameThread. */
 	void SetSelection_RenderThread(const bool bInParentSelected, const bool bInIndividuallySelected);
@@ -1685,13 +1669,6 @@ protected:
 	/** Allows child implementations to do render-thread work when bEvaluateWorldPositionOffset changes */
 	virtual void OnEvaluateWorldPositionOffsetChanged_RenderThread() {}
 
-	/**
-	 * Sets the instance local bounds for the specified instance index, and optionally will pad the bounds extents to
-	 * accomodate Max World Position Offset Distance.
-	 */
-	UE_DEPRECATED(5.4, "This does not do anything as this has been refactored.")
-	inline void SetInstanceLocalBounds(uint32 InstanceIndex, const FRenderBounds& InBounds, bool bPadForWPO = true) { }
-
 	ENGINE_API FRenderBounds PadInstanceLocalBounds(const FRenderBounds& InBounds);
 };
 
diff --git a/Engine/Source/Runtime/Engine/Public/SceneInterface.h b/Engine/Source/Runtime/Engine/Public/SceneInterface.h
--- a/Engine/Source/Runtime/Engine/Public/SceneInterface.h
+++ b/Engine/Source/Runtime/Engine/Public/SceneInterface.h
@@ -144,11 +144,11 @@ public:
 	virtual void UpdatePrimitiveAttachment(UPrimitiveComponent* Primitive) = 0;
 
 	/**
-	 * Updates all the instances that have been updated through the InstanceUpdateCmdBuffer on the UPrimitiveComponent.
+	 * Call to notify that an instace proxy update has been queued for the primitive.
 	 *
 	 * @param Primitive - primitive component to update
 	 */
-	virtual void UpdatePrimitiveInstances(UInstancedStaticMeshComponent* Primitive) = 0;
+	virtual void UpdatePrimitiveInstances(UPrimitiveComponent* Primitive) = 0;	
 
 	/** 
 	 * Updates the custom primitive data of a primitive component which has already been added to the scene. 
@@ -537,8 +537,10 @@ public:
 	
 	virtual void UpdateCustomPrimitiveData(FPrimitiveSceneDesc* Primitive, const FCustomPrimitiveData& CustomPrimitiveData) = 0;
 
-	virtual void UpdatePrimitiveInstances(FInstancedStaticMeshSceneDesc* Primitive) = 0;	
+	virtual void UpdatePrimitiveInstances(FPrimitiveSceneDesc* Primitive) = 0;	
 
+	UE_DEPRECATED(5.6, "StaticMeshBounds is no longer a parameter of this interface, please update FPrimitiveSceneDesc version.")
+	virtual void UpdatePrimitiveInstances(FInstancedStaticMeshSceneDesc* Primitive) = 0;	
 
 	/**
 	 * Release this scene and remove it from the rendering thread
diff --git a/Engine/Source/Runtime/GeometryFramework/Private/Components/BaseDynamicMeshSceneProxy.cpp b/Engine/Source/Runtime/GeometryFramework/Private/Components/BaseDynamicMeshSceneProxy.cpp
--- a/Engine/Source/Runtime/GeometryFramework/Private/Components/BaseDynamicMeshSceneProxy.cpp
+++ b/Engine/Source/Runtime/GeometryFramework/Private/Components/BaseDynamicMeshSceneProxy.cpp
@@ -871,16 +871,6 @@ void FBaseDynamicMeshSceneProxy::GetDistanceFieldAtlasData(const class FDistance
 	}
 }
 
-
-void FBaseDynamicMeshSceneProxy::GetDistanceFieldInstanceData(TArray<FRenderTransform>& InstanceLocalToPrimitiveTransforms) const
-{
-	check(InstanceLocalToPrimitiveTransforms.IsEmpty());
-	if (DistanceField.IsValid() && bDistanceFieldValid)
-	{
-		InstanceLocalToPrimitiveTransforms.Add(FRenderTransform::Identity);
-	}
-}
-
 bool FBaseDynamicMeshSceneProxy::HasDistanceFieldRepresentation() const 
 {
 	return CastsDynamicShadow() && AffectsDistanceFieldLighting() && bDistanceFieldValid && DistanceField.IsValid();
diff --git a/Engine/Source/Runtime/GeometryFramework/Public/Components/BaseDynamicMeshSceneProxy.h b/Engine/Source/Runtime/GeometryFramework/Public/Components/BaseDynamicMeshSceneProxy.h
--- a/Engine/Source/Runtime/GeometryFramework/Public/Components/BaseDynamicMeshSceneProxy.h
+++ b/Engine/Source/Runtime/GeometryFramework/Public/Components/BaseDynamicMeshSceneProxy.h
@@ -326,7 +326,6 @@ protected:
 
 public:
 	GEOMETRYFRAMEWORK_API virtual void GetDistanceFieldAtlasData(const class FDistanceFieldVolumeData*& OutDistanceFieldData, float& SelfShadowBias) const;
-	GEOMETRYFRAMEWORK_API virtual void GetDistanceFieldInstanceData(TArray<FRenderTransform>& InstanceLocalToPrimitiveTransforms) const override;
 	GEOMETRYFRAMEWORK_API virtual bool HasDistanceFieldRepresentation() const override;
 	GEOMETRYFRAMEWORK_API virtual bool HasDynamicIndirectShadowCasterRepresentation() const override;
 
diff --git a/Engine/Source/Runtime/Renderer/Private/RendererScene.cpp b/Engine/Source/Runtime/Renderer/Private/RendererScene.cpp
--- a/Engine/Source/Runtime/Renderer/Private/RendererScene.cpp
+++ b/Engine/Source/Runtime/Renderer/Private/RendererScene.cpp
@@ -2223,7 +2223,7 @@ void FScene::UpdateInstanceCullDistance(UPrimitiveComponent* Primitive, float St
 	UpdatePrimitiveInternal(Primitive->GetSceneProxy(), FUpdateInstanceCullDistanceData(FVector2f(StartCullDistance, EndCullDistance)));
 }
 
-void FScene::UpdatePrimitiveInstances(UInstancedStaticMeshComponent* Primitive)
+void FScene::UpdatePrimitiveInstances(UPrimitiveComponent* Primitive)
 {
 	SCOPE_CYCLE_COUNTER(STAT_UpdatePrimitiveInstanceGT);
 	SCOPED_NAMED_EVENT(FScene_UpdatePrimitiveInstance, FColor::Yellow);	
@@ -2240,9 +2240,6 @@ void FScene::UpdatePrimitiveInstances(UInstancedStaticMeshComponent* Primitive)
 	UpdateParams.WorldBounds = Primitive->Bounds;
 	UpdateParams.LocalBounds = static_cast<UPrimitiveComponent*>(Primitive)->GetLocalBounds();
 
-	// #todo (jnadro) This code should not be dependent on static mesh bounds.		
-	UpdateParams.StaticMeshBounds = Primitive->GetStaticMesh()->GetBounds();
-
 	UpdatePrimitiveInstances(UpdateParams);
 }
 
@@ -2272,6 +2269,11 @@ void FScene::UpdatePrimitiveInstancesFromCompute(FPrimitiveSceneInfo* PrimitiveS
 }
 
 void FScene::UpdatePrimitiveInstances(FInstancedStaticMeshSceneDesc* Primitive)
+{
+	UpdatePrimitiveInstances(*Primitive);
+}
+
+void FScene::UpdatePrimitiveInstances(FPrimitiveSceneDesc* Primitive)
 {
 	SCOPE_CYCLE_COUNTER(STAT_UpdatePrimitiveInstanceGT);
 	SCOPED_NAMED_EVENT(FScene_UpdatePrimitiveInstance, FColor::Yellow);
@@ -2279,7 +2281,7 @@ void FScene::UpdatePrimitiveInstances(FInstancedStaticMeshSceneDesc* Primitive)
 	// If the primitive doesn't have a scene info object yet, it must be added from scratch.
 	if (!Primitive->GetSceneProxy())
 	{
-		AddPrimitive(*Primitive);
+		AddPrimitive(Primitive);
 		return;
 	}
 	
@@ -2288,9 +2290,6 @@ void FScene::UpdatePrimitiveInstances(FInstancedStaticMeshSceneDesc* Primitive)
 	UpdateParams.WorldBounds = Primitive->GetBounds();
 	UpdateParams.LocalBounds = Primitive->GetLocalBounds();
 
-	// #todo (jnadro) This code should not be dependent on static mesh bounds.		
-	UpdateParams.StaticMeshBounds = Primitive->GetStaticMesh()->GetBounds();
-
 	UpdatePrimitiveInstances(UpdateParams);
 }
 
@@ -6386,7 +6385,7 @@ void FScene::Update(FRDGBuilder& GraphBuilder, const FUpdateParameters& Paramete
 #endif
 
 			// Update the Proxy's data.
-			PrimitiveSceneProxy->UpdateInstances_RenderThread(GraphBuilder.RHICmdList, UpdateInstance.WorldBounds, UpdateInstance.LocalBounds, UpdateInstance.StaticMeshBounds);
+			PrimitiveSceneProxy->UpdateInstances_RenderThread(GraphBuilder.RHICmdList, UpdateInstance.WorldBounds, UpdateInstance.LocalBounds);
 
 			if (!RHISupportsVolumeTextures(GetFeatureLevel())
 				&& (PrimitiveSceneProxy->IsMovable() || PrimitiveSceneProxy->NeedsUnbuiltPreviewLighting() || PrimitiveSceneProxy->GetLightmapType() == ELightmapType::ForceVolumetric))
@@ -6970,7 +6969,7 @@ public:
 
 	/** Updates the transform of a primitive which has already been added to the scene. */
 	virtual void UpdatePrimitiveTransform(UPrimitiveComponent* Primitive) override {}
-	virtual void UpdatePrimitiveInstances(UInstancedStaticMeshComponent* Primitive) override {}
+	virtual void UpdatePrimitiveInstances(UPrimitiveComponent* Primitive) override {}
 	virtual void UpdatePrimitiveOcclusionBoundsSlack(UPrimitiveComponent* Primitive, float NewSlack) override {}
 	virtual void UpdatePrimitiveDrawDistance(UPrimitiveComponent* Primitive, float MinDrawDistance, float MaxDrawDistance, float VirtualTextureMaxDrawDistance) override {}
 	virtual void UpdateInstanceCullDistance(UPrimitiveComponent* Primitive, float StartCullDistance, float EndCullDistance) {}
@@ -7118,7 +7117,8 @@ public:
 	virtual void BatchReleasePrimitives(TArrayView<FPrimitiveSceneDesc*> InPrimitives) override {};
 
 	virtual void UpdateCustomPrimitiveData(FPrimitiveSceneDesc* Primitive, const FCustomPrimitiveData&) override {}
-	virtual void UpdatePrimitiveInstances(FInstancedStaticMeshSceneDesc* Primitive) override {};
+	virtual void UpdatePrimitiveInstances(FPrimitiveSceneDesc* Primitive) override {};
+	virtual void UpdatePrimitiveInstances(FInstancedStaticMeshSceneDesc* Primitive) override {};	
 
 	virtual void StartUpdatePrimitiveTransform(int32 NumPrimitives) override {}
 	virtual void FinishUpdatePrimitiveTransform() override {}
diff --git a/Engine/Source/Runtime/Renderer/Private/ScenePrimitiveUpdates.h b/Engine/Source/Runtime/Renderer/Private/ScenePrimitiveUpdates.h
--- a/Engine/Source/Runtime/Renderer/Private/ScenePrimitiveUpdates.h
+++ b/Engine/Source/Runtime/Renderer/Private/ScenePrimitiveUpdates.h
@@ -73,7 +73,6 @@ struct FUpdateInstanceCommand : public TPrimitiveUpdatePayloadBase<EPrimitiveUpd
 	FPrimitiveSceneProxy* PrimitiveSceneProxy{ nullptr };
 	FBoxSphereBounds WorldBounds;
 	FBoxSphereBounds LocalBounds;
-	FBoxSphereBounds StaticMeshBounds;
 };
 
 struct FUpdateInstanceFromComputeCommand : public TPrimitiveUpdatePayloadBase<EPrimitiveUpdateId::UpdateInstanceFromCompute, EPrimitiveUpdateDirtyFlags::CullingBounds | EPrimitiveUpdateDirtyFlags::InstanceData>
diff --git a/Engine/Source/Runtime/Renderer/Private/ScenePrivate.h b/Engine/Source/Runtime/Renderer/Private/ScenePrivate.h
--- a/Engine/Source/Runtime/Renderer/Private/ScenePrivate.h
+++ b/Engine/Source/Runtime/Renderer/Private/ScenePrivate.h
@@ -3381,7 +3381,7 @@ public:
 	virtual void BatchReleasePrimitives(TArrayView<UPrimitiveComponent*> InPrimitives) override;
 	virtual void UpdateAllPrimitiveSceneInfos(FRDGBuilder& GraphBuilder, EUpdateAllPrimitiveSceneInfosAsyncOps AsyncOps = EUpdateAllPrimitiveSceneInfosAsyncOps::None) override;
 	virtual void UpdatePrimitiveTransform(UPrimitiveComponent* Primitive) override;
-	virtual void UpdatePrimitiveInstances(UInstancedStaticMeshComponent* Primitive) override;
+	virtual void UpdatePrimitiveInstances(UPrimitiveComponent* Primitive) override;
 	virtual void UpdatePrimitiveOcclusionBoundsSlack(UPrimitiveComponent* Primitive, float NewSlack) override;
 	virtual void UpdatePrimitiveDrawDistance(UPrimitiveComponent* Primitive, float MinDrawDistance, float MaxDrawDistance, float VirtualTextureMaxDrawDistance) override;
 	virtual void UpdateInstanceCullDistance(UPrimitiveComponent* Primitive, float StartCullDistance, float EndCullDistance) override;
@@ -3529,6 +3529,7 @@ public:
 	virtual void UpdateCustomPrimitiveData(FPrimitiveSceneDesc* Primitive, const FCustomPrimitiveData& CustomPrimitiveData) override;
 
 	virtual void UpdatePrimitiveInstances(FInstancedStaticMeshSceneDesc* Primitive) override;
+	virtual void UpdatePrimitiveInstances(FPrimitiveSceneDesc* Primitive) override;
 
 	bool HasSkyAtmosphere() const
 	{
