LFV - LocalFogVolume.GPUTileDataBuffer is safely allocated against 0x0 viewport.

#rnx
#tests PC - dx12

diff --git a/Engine/Source/Runtime/Renderer/Private/LocalFogVolumeRendering.cpp b/Engine/Source/Runtime/Renderer/Private/LocalFogVolumeRendering.cpp
--- a/Engine/Source/Runtime/Renderer/Private/LocalFogVolumeRendering.cpp
+++ b/Engine/Source/Runtime/Renderer/Private/LocalFogVolumeRendering.cpp
@@ -522,7 +522,8 @@ void CreateViewLocalFogVolumeBufferSRV(const FScene* Scene, FViewInfo& View, FRD
 	// Create the texture that will contain the tiled culled result: count in the first slice and indices in the remaining slices
 	const uint32 LocalFogVolumeTilePixelSize				= GetLocalFogVolumeTilePixelSize();
 
-	const FIntPoint TileDataTextureResolution				= FIntPoint::DivideAndRoundUp(View.ViewRect.Size(), FIntPoint(LocalFogVolumeTilePixelSize, LocalFogVolumeTilePixelSize));
+	const FIntPoint TileDataTextureResolutionUnsafe			= FIntPoint::DivideAndRoundUp(View.ViewRect.Size(), FIntPoint(LocalFogVolumeTilePixelSize, LocalFogVolumeTilePixelSize));
+	const FIntPoint TileDataTextureResolution				= FIntPoint(FMath::Max(1, TileDataTextureResolutionUnsafe.X), FMath::Max(1, TileDataTextureResolutionUnsafe.Y));
 	const uint32 TileDataTextureSliceCount					= LocalFogVolumeTileMaxInstanceCount + 1; // +1 because the first slice is the culled instance count
 
 	EPixelFormat TileDataFormat = PF_R8_UINT;	
