SVE: Improve alignment between scene view extensions delegates and post-process material blendable locations.

* This adds 4 new locations that already exist for post-process materials: BeforeDOF, AfterDOF, TranslucencyAfterDOF & ReplacingTonemapper.
* For consistency, we always call the scene view extension before PPMs. As a result SSRInput delegates are now run before their PPM equivalents.
* Soft deprecation of the "after pass" scene view extension names.
* For ReplacingTonemapper, only the highest priority delegate is executed.

#rnx

diff --git a/Engine/Source/Runtime/Engine/Private/SceneViewExtension.cpp b/Engine/Source/Runtime/Engine/Private/SceneViewExtension.cpp
--- a/Engine/Source/Runtime/Engine/Private/SceneViewExtension.cpp
+++ b/Engine/Source/Runtime/Engine/Private/SceneViewExtension.cpp
@@ -98,12 +98,6 @@ void FSceneViewExtensions::ForEachActiveViewExtension(
 }
 
 // @todo viewext : We should cache all the active extensions in OnStartFrame somewhere
-const TArray<FSceneViewExtensionRef> FSceneViewExtensions::GatherActiveExtensions(FViewport* InViewport /* = nullptr */) const
-{
-	FSceneViewExtensionContext Context(InViewport);
-	return GatherActiveExtensions(Context);
-}
-
 const TArray<FSceneViewExtensionRef> FSceneViewExtensions::GatherActiveExtensions(const FSceneViewExtensionContext& InContext) const
 {
 	TArray<FSceneViewExtensionRef> ActiveExtensions;
diff --git a/Engine/Source/Runtime/Engine/Public/SceneViewExtension.h b/Engine/Source/Runtime/Engine/Public/SceneViewExtension.h
--- a/Engine/Source/Runtime/Engine/Public/SceneViewExtension.h
+++ b/Engine/Source/Runtime/Engine/Public/SceneViewExtension.h
@@ -92,21 +92,32 @@ class FViewport;
 
 
 /** This is used to add more flexibility to Post Processing, so that users can subscribe to any after Post Processing Pass events. */
-FUNC_DECLARE_DELEGATE(FAfterPassCallbackDelegate, FScreenPassTexture /*ReturnSceneColor*/, FRDGBuilder& /*GraphBuilder*/, const FSceneView& /*View*/, const FPostProcessMaterialInputs& /*Inputs*/)
-using FAfterPassCallbackDelegateArray = TArray<FAfterPassCallbackDelegate, SceneRenderingAllocator>;
+FUNC_DECLARE_DELEGATE(FPostProcessingPassDelegate, FScreenPassTexture /*ReturnSceneColor*/, FRDGBuilder& /*GraphBuilder*/, const FSceneView& /*View*/, const FPostProcessMaterialInputs& /*Inputs*/)
+using FPostProcessingPassDelegateArray = TArray<FPostProcessingPassDelegate, SceneRenderingAllocator>;
 
+// Soft deprecation of the previous after pass delegate names
+using FAfterPassCallbackDelegate = FPostProcessingPassDelegate;
+using FAfterPassCallbackDelegateArray = FPostProcessingPassDelegateArray;
 
 class ISceneViewExtension
 {
 public:
 
+	// Each post-processing pass immediately precedes a PPM blend location, if it exists. See comments below.
 	enum class EPostProcessingPass : uint32
 	{
-		SSRInput,
-		MotionBlur,
-		Tonemap,
+		BeforeDOF,				// BL_SceneColorBeforeDOF
+		AfterDOF,				// BL_SceneColorAfterDOF
+		TranslucencyAfterDOF,	// BL_TranslucencyAfterDOF
+		SSRInput,				// BL_SSRInput
+
+		// The following post-processing passes may be last, and therefore receive a valid OverrideOutput render target.
+		ReplacingTonemapper,	// BL_ReplacingTonemapper
+		MotionBlur,				// BL_SceneColorBeforeBloom
+		Tonemap,				// BL_SceneColorAfterTonemapping
 		FXAA,
 		VisualizeDepthOfField,
+
 		MAX
 	};
 
@@ -181,11 +192,11 @@ public:
 	virtual void PrePostProcessPass_RenderThread(FRDGBuilder& GraphBuilder, const FSceneView& InView, const FPostProcessingInputs& Inputs) {};
 
 	/**
-	* This will be called at the beginning of post processing to make sure that each view extension gets a chance to subscribe to an after pass event.
+	* This will be called at the beginning of post processing to make sure that each view extension gets a chance to subscribe to a post-processing pass event.
 	*  - The pass MUST write to the override output texture if it is active (this occurs when the pass is the last in the post processing chain writing to the back buffer).
-	*    For performance reasons it is recommended to only subscribe to a pass when the pass will produce a GPU resource. Calling 
+	*    For performance reasons it is recommended to only subscribe to a pass when the pass will produce a GPU resource.
 	*/
-	virtual void SubscribeToPostProcessingPass(EPostProcessingPass Pass, const FSceneView& InView, FAfterPassCallbackDelegateArray& InOutPassCallbacks, bool bIsPassEnabled) {}
+	virtual void SubscribeToPostProcessingPass(EPostProcessingPass Pass, const FSceneView& InView, FPostProcessingPassDelegateArray& InOutPassCallbacks, bool bIsPassEnabled) {}
 
 	/**
 	 * Allows to render content after the 3D content scene, useful for debugging
@@ -226,7 +237,7 @@ public:
 	virtual void PostRenderBasePass_RenderThread(FRHICommandListImmediate& RHICmdList, FSceneView& InView) {};
 
 	UE_DEPRECATED(5.5, "SubscribeToPostProcessingPass now takes a SceneView")
-	virtual void SubscribeToPostProcessingPass(EPostProcessingPass Pass, FAfterPassCallbackDelegateArray& InOutPassCallbacks, bool bIsPassEnabled) {}
+	virtual void SubscribeToPostProcessingPass(EPostProcessingPass Pass, FPostProcessingPassDelegateArray& InOutPassCallbacks, bool bIsPassEnabled) {}
 
 	///////////////////////////////////////////////////////////////////////////////////////
 
@@ -317,13 +328,6 @@ public:
 		const FSceneViewExtensionContext& InContext,
 		const TFunctionRef<void(const FSceneViewExtensionRef&)>& Func);
 
-	/**
-	 * Gathers all ViewExtensions that want to be active for a given viewport (@see ISceneViewExtension::IsActiveThisFrame()).
-	 * The list is sorted by priority (@see ISceneViewExtension::GetPriority())
-	 */
-	UE_DEPRECATED(4.27, "Deprecated. Please use GatherActiveExtensions by passing an FSceneViewExtensionContext parameter")
-	ENGINE_API const TArray<FSceneViewExtensionRef> GatherActiveExtensions(class FViewport* InViewport = nullptr) const;
-
 	/**
      * Gathers all ViewExtensions that want to be active in a given context (@see ISceneViewExtension::IsActiveThisFrame()).
      * The list is sorted by priority (@see ISceneViewExtension::GetPriority())
diff --git a/Engine/Source/Runtime/Renderer/Private/PostProcess/PostProcessing.cpp b/Engine/Source/Runtime/Renderer/Private/PostProcess/PostProcessing.cpp
--- a/Engine/Source/Runtime/Renderer/Private/PostProcess/PostProcessing.cpp
+++ b/Engine/Source/Runtime/Renderer/Private/PostProcess/PostProcessing.cpp
@@ -302,6 +302,41 @@ FDefaultTemporalUpscaler::FOutputs AddThirdPartyTemporalUpscalerPasses(
 	return Outputs;
 }
 
+namespace UE::Renderer::PostProcess
+{
+	/*
+	* Issue scene view extension pass callbacks, mimicking "AddPostProcessMaterialChain".
+	* The "AddAfterPass" lambdas are used instead for later extension points in the override pass sequence.
+	*/
+	FScreenPassTexture AddSceneViewExtensionPassChain(
+		FRDGBuilder& GraphBuilder,
+		const FViewInfo& View,
+		const FPostProcessMaterialInputs& InputsTemplate,
+		const FPostProcessingPassDelegateArray& Delegates,
+		EPostProcessMaterialInput MaterialInput = EPostProcessMaterialInput::SceneColor)
+	{
+		FScreenPassTextureSlice CurrentInput = InputsTemplate.GetInput(MaterialInput);
+		FScreenPassTexture Outputs;
+
+		for (int32 DelegateIndex = 0; DelegateIndex < Delegates.Num(); ++DelegateIndex)
+		{
+			FPostProcessMaterialInputs Inputs = InputsTemplate;
+			Inputs.SetInput(MaterialInput, CurrentInput);
+			
+			Outputs = Delegates[DelegateIndex].Execute(GraphBuilder, View, Inputs);
+
+			CurrentInput = FScreenPassTextureSlice::CreateFromScreenPassTexture(GraphBuilder, Outputs);
+		}
+
+		if (!Outputs.IsValid())
+		{
+			Outputs = FScreenPassTexture::CopyFromSlice(GraphBuilder, CurrentInput);
+		}
+
+		return Outputs;
+	};
+}
+
 bool ComposeSeparateTranslucencyInTSR(const FViewInfo& View);
 
 void AddPostProcessingPasses(
@@ -323,6 +358,7 @@ void AddPostProcessingPasses(
 {
 	RDG_CSV_STAT_EXCLUSIVE_SCOPE(GraphBuilder, RenderPostProcessing);
 	QUICK_SCOPE_CYCLE_COUNTER(STAT_PostProcessing_Process);
+	using namespace UE::Renderer::PostProcess;
 
 	check(IsInRenderingThread());
 	check(View.VerifyMembersChecks());
@@ -590,7 +626,7 @@ void AddPostProcessingPasses(
 	{
 		// In some cases (e.g. OCIO color conversion) we want View Extensions to be able to add extra custom post processing after the pass.
 
-		FAfterPassCallbackDelegateArray& PassCallbacks = PassSequence.GetAfterPassCallbacks(InPass);
+		FPostProcessingPassDelegateArray& PassCallbacks = PassSequence.GetAfterPassCallbacks(InPass);
 
 		if (PassCallbacks.Num())
 		{
@@ -611,7 +647,7 @@ void AddPostProcessingPasses(
 
 	const auto AddAfterPassForSceneColorSlice = [&](EPass InPass, const FScreenPassTextureSlice& InSceneColor) -> FScreenPassTextureSlice
 	{
-		FAfterPassCallbackDelegateArray& PassCallbacks = PassSequence.GetAfterPassCallbacks(InPass);
+		FPostProcessingPassDelegateArray& PassCallbacks = PassSequence.GetAfterPassCallbacks(InPass);
 
 		if (PassCallbacks.Num())
 		{
@@ -689,18 +725,28 @@ void AddPostProcessingPasses(
 		PassSequence.SetEnabled(EPass::VisualizeDepthOfField, bVisualizeDepthOfField);
 		PassSequence.SetEnabled(EPass::VisualizeLocalExposure, EngineShowFlags.VisualizeLocalExposure);
 
-		FAfterPassCallbackDelegateArray SSRInputDelegates;
+		static_assert(static_cast<int32>(EPass::MotionBlur) == 0);
+		constexpr int32 FirstAfterPass = static_cast<int32>(ISceneViewExtension::EPostProcessingPass::MotionBlur);
+		
+		// Scene view extension delegates that precede the override pass sequence are to be called directly.
+		TStaticArray<FPostProcessingPassDelegateArray, static_cast<uint32>(FirstAfterPass)> SceneViewExtensionDelegates;
 
-		for (int32 ViewExt = 0; ViewExt < View.Family->ViewExtensions.Num(); ++ViewExt)
+		for (const TSharedRef<ISceneViewExtension>& ViewExtension : View.Family->ViewExtensions)
 		{
-			View.Family->ViewExtensions[ViewExt]->SubscribeToPostProcessingPass(ISceneViewExtension::EPostProcessingPass::SSRInput, View, SSRInputDelegates, true);
+			for (int32 SceneViewPassId = 0; SceneViewPassId < FirstAfterPass; SceneViewPassId++)
+			{
+				const ISceneViewExtension::EPostProcessingPass SceneViewPass = static_cast<ISceneViewExtension::EPostProcessingPass>(SceneViewPassId);
+				const bool bIsEnabled = (SceneViewPass == ISceneViewExtension::EPostProcessingPass::ReplacingTonemapper) ? PassSequence.IsEnabled(EPass::Tonemap) : true;
+
+				ViewExtension->SubscribeToPostProcessingPass(SceneViewPass, View, SceneViewExtensionDelegates[SceneViewPassId], bIsEnabled);
+			}
 
-			for (int32 SceneViewPassId = static_cast<int32>(ISceneViewExtension::EPostProcessingPass::MotionBlur); SceneViewPassId != static_cast<int32>(ISceneViewExtension::EPostProcessingPass::MAX); SceneViewPassId++)
+			for (int32 SceneViewPassId = FirstAfterPass; SceneViewPassId < static_cast<int32>(ISceneViewExtension::EPostProcessingPass::MAX); SceneViewPassId++)
 			{
-				ISceneViewExtension::EPostProcessingPass SceneViewPass = static_cast<ISceneViewExtension::EPostProcessingPass>(SceneViewPassId);
-				EPass PostProcessingPass = TranslatePass(SceneViewPass);
+				const ISceneViewExtension::EPostProcessingPass SceneViewPass = static_cast<ISceneViewExtension::EPostProcessingPass>(SceneViewPassId);
+				const EPass PostProcessingPass = TranslatePass(SceneViewPass);
 
-				View.Family->ViewExtensions[ViewExt]->SubscribeToPostProcessingPass(
+				ViewExtension->SubscribeToPostProcessingPass(
 					SceneViewPass,
 					View,
 					PassSequence.GetAfterPassCallbacks(PostProcessingPass),
@@ -737,6 +783,13 @@ void AddPostProcessingPasses(
 		const FPostProcessMaterialChain MaterialChainSceneColorAfterDOF = GetPostProcessMaterialChain(View, BL_SceneColorAfterDOF);
 		const FPostProcessMaterialChain MaterialChainTranslucencyAfterDOF = GetPostProcessMaterialChain(View, BL_TranslucencyAfterDOF);
 
+		// Scene view extension delegates - BeforeDOF
+		if (SceneViewExtensionDelegates[static_cast<uint32>(ISceneViewExtension::EPostProcessingPass::BeforeDOF)].Num())
+		{
+			SceneColor = AddSceneViewExtensionPassChain(GraphBuilder, View, GetPostProcessMaterialInputs(SceneColor),
+				SceneViewExtensionDelegates[static_cast<uint32>(ISceneViewExtension::EPostProcessingPass::BeforeDOF)]);
+		}
+
 		// Post Process Material Chain - BL_SceneColorBeforeDOF
 		if (MaterialChainSceneColorBeforeDOF.Num())
 		{
@@ -772,6 +825,13 @@ void AddPostProcessingPasses(
 			}
 		}
 
+		// Scene view extension delegates - AfterDOF
+		if (SceneViewExtensionDelegates[static_cast<uint32>(ISceneViewExtension::EPostProcessingPass::AfterDOF)].Num())
+		{
+			SceneColor = AddSceneViewExtensionPassChain(GraphBuilder, View, GetPostProcessMaterialInputs(SceneColor),
+				SceneViewExtensionDelegates[static_cast<uint32>(ISceneViewExtension::EPostProcessingPass::AfterDOF)]);
+		}
+
 		// Post Process Material Chain - BL_SceneColorAfterDOF
 		if (MaterialChainSceneColorAfterDOF.Num())
 		{
@@ -786,6 +846,18 @@ void AddPostProcessingPasses(
 		}
 		else if (PostDOFTranslucencyResources.IsValid())
 		{
+			if (SceneViewExtensionDelegates[static_cast<uint32>(ISceneViewExtension::EPostProcessingPass::TranslucencyAfterDOF)].Num())
+			{
+				FScreenPassTexture PostDOFTranslucency = AddSceneViewExtensionPassChain(
+					GraphBuilder, View,
+					GetPostProcessMaterialInputs(SceneColor),
+					SceneViewExtensionDelegates[static_cast<uint32>(ISceneViewExtension::EPostProcessingPass::TranslucencyAfterDOF)],
+					EPostProcessMaterialInput::SeparateTranslucency);
+
+				PostDOFTranslucencyResources.ColorTexture = PostDOFTranslucency.Texture;
+				ensure(PostDOFTranslucencyResources.ViewRect == PostDOFTranslucency.ViewRect);
+			}
+
 			if (MaterialChainTranslucencyAfterDOF.Num())
 			{
 				FScreenPassTexture PostDOFTranslucency = AddPostProcessMaterialChain(
@@ -942,26 +1014,20 @@ void AddPostProcessingPasses(
 		if (View.ViewState && !View.bStatePrevViewInfoIsReadOnly)
 		{
 			FScreenPassTexture PassOutput;
-			FPostProcessMaterialInputs InOutPassInputs;
+			FPostProcessMaterialInputs PassInputs;
 			const FPostProcessMaterialChain MaterialChain = GetPostProcessMaterialChain(View, BL_SSRInput);
 
-			if (MaterialChain.Num())
+			if (SceneViewExtensionDelegates[static_cast<uint32>(ISceneViewExtension::EPostProcessingPass::SSRInput)].Num())
 			{
-				InOutPassInputs = GetPostProcessMaterialInputs(FScreenPassTexture::CopyFromSlice(GraphBuilder, SceneColorSlice));
-				PassOutput = AddPostProcessMaterialChain(GraphBuilder, View, ViewIndex, InOutPassInputs, MaterialChain);
+				PassInputs = GetPostProcessMaterialInputs(FScreenPassTexture::CopyFromSlice(GraphBuilder, SceneColorSlice));
+				PassOutput = AddSceneViewExtensionPassChain(GraphBuilder, View, PassInputs,
+					SceneViewExtensionDelegates[static_cast<uint32>(ISceneViewExtension::EPostProcessingPass::SSRInput)]);
 			}
 
-			for (FAfterPassCallbackDelegate& PassCallback : SSRInputDelegates)
+			if (MaterialChain.Num())
 			{
-				if (PassOutput.IsValid())
-				{
-					InOutPassInputs.SetInput(GraphBuilder, EPostProcessMaterialInput::SceneColor, PassOutput);
-				}
-				else
-				{
-					InOutPassInputs.SetInput(GraphBuilder, EPostProcessMaterialInput::SceneColor, FScreenPassTexture::CopyFromSlice(GraphBuilder, SceneColorSlice));
-				}
-				PassOutput = PassCallback.Execute(GraphBuilder, View, InOutPassInputs);
+				PassInputs = GetPostProcessMaterialInputs(PassOutput.IsValid() ? PassOutput : FScreenPassTexture::CopyFromSlice(GraphBuilder, SceneColorSlice));
+				PassOutput = AddPostProcessMaterialChain(GraphBuilder, View, ViewIndex, PassInputs, MaterialChain);
 			}
 
 			if (PassOutput.IsValid())
@@ -1288,22 +1354,33 @@ void AddPostProcessingPasses(
 
 		if (PassSequence.IsEnabled(EPass::Tonemap))
 		{
+			const FPostProcessingPassDelegateArray& ReplacingTonemapperDelegates = SceneViewExtensionDelegates[static_cast<uint32>(ISceneViewExtension::EPostProcessingPass::ReplacingTonemapper)];
 			const FPostProcessMaterialChain MaterialChain = GetPostProcessMaterialChain(View, BL_ReplacingTonemapper);
 
-			if (MaterialChain.Num())
+			auto GetReplaceTonemapperInputs = [&]() -> FPostProcessMaterialInputs
+				{
+					FPostProcessMaterialInputs PassInputs;
+					PassSequence.AcceptOverrideIfLastPass(EPass::Tonemap, PassInputs.OverrideOutput);
+					PassInputs.SetInput(GraphBuilder, EPostProcessMaterialInput::SceneColor, FScreenPassTexture::CopyFromSlice(GraphBuilder, SceneColorSlice));
+					PassInputs.SetInput(GraphBuilder, EPostProcessMaterialInput::CombinedBloom, Bloom);
+					PassInputs.SceneTextures = GetSceneTextureShaderParameters(Inputs.SceneTextures);
+					PassInputs.CustomDepthTexture = CustomDepth.Texture;
+					PassInputs.bManualStencilTest = Inputs.bSeparateCustomStencil;
+
+					return PassInputs;
+				};
+
+			if(ReplacingTonemapperDelegates.Num())
+			{
+				const FAfterPassCallbackDelegate& HighestPriorityDelegate = ReplacingTonemapperDelegates[0];
+
+				SceneColor = HighestPriorityDelegate.Execute(GraphBuilder, View, GetReplaceTonemapperInputs());
+			}
+			else if (MaterialChain.Num())
 			{
 				const UMaterialInterface* HighestPriorityMaterial = MaterialChain[0];
 
-				FPostProcessMaterialInputs PassInputs;
-				PassSequence.AcceptOverrideIfLastPass(EPass::Tonemap, PassInputs.OverrideOutput);
-				PassInputs.SetInput(GraphBuilder, EPostProcessMaterialInput::SceneColor, FScreenPassTexture::CopyFromSlice(GraphBuilder, SceneColorSlice));
-				//PassInputs.SetInput(EPostProcessMaterialInput::SeparateTranslucency, SeparateTranslucency);
-				PassInputs.SetInput(GraphBuilder, EPostProcessMaterialInput::CombinedBloom, Bloom);
-				PassInputs.SceneTextures = GetSceneTextureShaderParameters(Inputs.SceneTextures);
-				PassInputs.CustomDepthTexture = CustomDepth.Texture;
-				PassInputs.bManualStencilTest = Inputs.bSeparateCustomStencil;
-
-				SceneColor = AddPostProcessMaterialPass(GraphBuilder, View, PassInputs, HighestPriorityMaterial);
+				SceneColor = AddPostProcessMaterialPass(GraphBuilder, View, GetReplaceTonemapperInputs(), HighestPriorityMaterial);
 			}
 			else
 			{
