Heterogeneous Volumes: Upsampling improvements

r.HeterogeneousVolumes.Upsample:
Upsampling iterative smoothing (Default = 2)
0: Off
1: Nearest Neighbor
2: Bilinear

r.HeterogeneousVolumes.Filter:
Controls iterative smoothing filter applied during upsampling (Default = 2)
0: Off
1: Bilateral
2: Gaussian 3x3
3: Gaussian 5x5

diff --git a/Engine/Shaders/Private/HeterogeneousVolumes/HeterogeneousVolumesComposite.usf b/Engine/Shaders/Private/HeterogeneousVolumes/HeterogeneousVolumesComposite.usf
--- a/Engine/Shaders/Private/HeterogeneousVolumes/HeterogeneousVolumesComposite.usf
+++ b/Engine/Shaders/Private/HeterogeneousVolumes/HeterogeneousVolumesComposite.usf
@@ -1,11 +1,185 @@
 // Copyright Epic Games, Inc. All Rights Reserved.
 
 #include "../Common.ush"
+#include "HeterogeneousVolumesAdaptiveVolumetricShadowMapSampling.ush"
 
+#ifndef UPSAMPLE_MODE
+#define UPSAMPLE_MODE 0
+#endif // UPSAMPLE_MODE
+
+#define UPSAMPLE_MODE_NEAREST 1
+#define UPSAMPLE_MODE_BILINEAR 2
+
+#ifndef FILTER_MODE
+#define FILTER_MODE 0
+#endif // FILTER_MODE
+
+#define FILTER_MODE_BILATERAL 1
+#define FILTER_MODE_GAUSSIAN_3X3 2
+#define FILTER_MODE_GAUSSIAN_5x5 3
+
+Texture2D SceneDepthTexture;
 Texture2D<float4> HeterogeneousVolumeRadiance;
 Texture2D<float> HeterogeneousVolumeHoldout;
 RWTexture2D<float4> RWColorTexture;
 int DownsampleFactor;
+int bUseAVSM;
+
+float SampleTransmittance(uint2 ViewPixelCoord)
+{
+	float DeviceZ = SceneDepthTexture.Load(int3(ViewPixelCoord, 0)).r;
+#if HAS_INVERTED_Z_BUFFER
+	DeviceZ = max(0.000000000001, DeviceZ);
+#endif // HAS_INVERTED_Z_BUFFER
+
+	float3 TranslatedWorldPosition = SvPositionToResolvedTranslatedWorld(float4(ViewPixelCoord + View.TemporalAAJitter.xy, DeviceZ, 1));
+	return AVSM_SampleTransmittance(TranslatedWorldPosition, ResolvedView.TranslatedWorldCameraOrigin);
+}
+
+RWTexture2D<float4> RWUpsampledTexture;
+SamplerState TextureSampler;
+uint2 UpsampledResolution;
+uint2 UpsampledViewRect;
+
+[numthreads(THREADGROUP_SIZE_2D, THREADGROUP_SIZE_2D, 1)]
+void HeterogeneousVolumesUpsampleCS(
+	uint2 GroupThreadId : SV_GroupThreadID,
+	uint2 DispatchThreadId : SV_DispatchThreadID
+)
+{
+	if (any(DispatchThreadId.xy >= UpsampledViewRect))
+	{
+		return;
+	}
+	int2 UpsampledPixelCoord = DispatchThreadId.xy; // TODO: UpsampledViewRectMin
+
+#if UPSAMPLE_MODE == UPSAMPLE_MODE_BILINEAR
+	float2 UV = (UpsampledPixelCoord + 0.5 * DownsampleFactor) / float2(UpsampledResolution);
+	UV = clamp(UV, 0, (UpsampledViewRect - 1) / float2(UpsampledResolution));
+	float4 Result = Texture2DSample(HeterogeneousVolumeRadiance, TextureSampler, UV);
+#else
+	int2 DownsampledPixelCoord = UpsampledPixelCoord / DownsampleFactor;
+	float4 Result = HeterogeneousVolumeRadiance[DownsampledPixelCoord];
+#endif
+
+	RWUpsampledTexture[UpsampledPixelCoord] = Result;
+}
+
+int FilterWidth;
+RWTexture2D<float4> RWFilteredTexture;
+
+[numthreads(THREADGROUP_SIZE_2D, THREADGROUP_SIZE_2D, 1)]
+void HeterogeneousVolumesFilterCS(
+	uint2 GroupThreadId : SV_GroupThreadID,
+	uint2 DispatchThreadId : SV_DispatchThreadID
+)
+{
+	if (any(DispatchThreadId.xy >= UpsampledViewRect))
+	{
+		return;
+	}
+
+	int2 PixelCoord = DispatchThreadId.xy; // TODO: UpsampledViewRectMin
+	float4 Radiance = HeterogeneousVolumeRadiance[PixelCoord];
+	float HeterogeneousVolumeTransmittance = Radiance.a;
+
+	// Bilaterial Filter
+#if FILTER_MODE == FILTER_MODE_BILATERAL
+	{
+		Radiance = 0.0;
+		float WeightSum = 0.0;
+
+		uint FilterHalfWidth = FilterWidth / 2;
+		float FilterRadius2 = max(2.0 * FilterWidth * FilterWidth, 1.0);
+
+		for (int X = 0; X < FilterWidth; ++X)
+		{
+			for (int Y = 0; Y < FilterWidth; ++Y)
+			{
+				int2 Offset = int2(X, Y) - FilterHalfWidth;
+				int2 SampleCoord = PixelCoord + Offset;
+				SampleCoord = clamp(SampleCoord, 0, int2(UpsampledViewRect) - 1);
+				float FilterWeight = max(FilterRadius2 - dot(Offset, Offset), 0.0);
+
+				float4 Value = HeterogeneousVolumeRadiance[SampleCoord];
+				float3 RadianceDelta = HeterogeneousVolumeRadiance[PixelCoord].rgb - Value.rgb;
+				float RadianceWeight = saturate(1.0 - dot(RadianceDelta, RadianceDelta));
+
+				float TauDelta = HeterogeneousVolumeTransmittance - Value.a;
+				float TauWeight = 1.0 - TauDelta * TauDelta;
+
+				Radiance += Value * FilterWeight * RadianceWeight * TauWeight;
+				WeightSum += FilterWeight * RadianceWeight * TauWeight;
+
+			}
+		}
+
+		Radiance /= WeightSum;
+	}
+	// Gaussian 3x3 Filter
+#elif FILTER_MODE == FILTER_MODE_GAUSSIAN_3X3
+	{
+		float Gaussian3x3Weights[] = {
+			1.0, 2.0, 1.0,
+			2.0, 4.0, 2.0,
+			1.0, 2.0, 1.0
+		};
+
+		Radiance = 0.0;
+		float WeightSum = 0.0;
+
+		uint FilterHalfWidth = 1;
+		for (int X = 0; X < 3; ++X)
+		{
+			for (int Y = 0; Y < 3; ++Y)
+			{
+				int2 OffsetCoord = int2(X, Y);
+				int2 SampleCoord = PixelCoord + OffsetCoord - FilterHalfWidth;
+				SampleCoord = clamp(SampleCoord, 0, int2(UpsampledViewRect) - 1);
+				float FilterWeight = Gaussian3x3Weights[OffsetCoord.y * 3 + OffsetCoord.x];
+
+				Radiance += HeterogeneousVolumeRadiance[SampleCoord] * FilterWeight;
+				WeightSum += FilterWeight;
+			}
+		}
+
+		Radiance /= WeightSum;
+	}
+	// Gaussian 5x5 Filter
+#elif FILTER_MODE == FILTER_MODE_GAUSSIAN_5x5
+	{
+		float Gaussian5x5Weights[] = {
+			1.0, 4.0, 7.0, 4.0, 1.0,
+			4.0, 16.0, 26.0, 16.0, 4.0,
+			7.0, 26.0, 41.0, 26.0, 7.0,
+			4.0, 16.0, 26.0, 16.0, 4.0,
+			1.0, 4.0, 7.0, 4.0, 1.0,
+		};
+
+		Radiance = 0.0;
+		float WeightSum = 0.0;
+
+		uint FilterHalfWidth = 2;
+		for (int X = 0; X < 5; ++X)
+		{
+			for (int Y = 0; Y < 5; ++Y)
+			{
+				int2 OffsetCoord = int2(X, Y);
+				int2 SampleCoord = PixelCoord + OffsetCoord - FilterHalfWidth;
+				SampleCoord = clamp(SampleCoord, 0, int2(UpsampledViewRect) - 1);
+				float FilterWeight = Gaussian5x5Weights[OffsetCoord.y * 5 + OffsetCoord.x];
+
+				Radiance += HeterogeneousVolumeRadiance[SampleCoord] * FilterWeight;
+				WeightSum += FilterWeight;
+			}
+		}
+
+		Radiance /= WeightSum;
+	}
+#endif
+
+	RWFilteredTexture[PixelCoord] = Radiance;
+}
 
 [numthreads(THREADGROUP_SIZE_2D, THREADGROUP_SIZE_2D, 1)]
 void HeterogeneousVolumesCompositeCS(
@@ -13,16 +187,20 @@ void HeterogeneousVolumesCompositeCS(
 	uint2 DispatchThreadId : SV_DispatchThreadID
 )
 {
-	if (any(DispatchThreadId.xy >= View.ViewSizeAndInvSize.xy))
+	ResolvedView = ResolveView();
+	if (any(DispatchThreadId.xy >= ResolvedView.ViewSizeAndInvSize.xy))
 	{
 		return;
 	}
-	uint2 ViewPixelCoord = DispatchThreadId.xy + View.ViewRectMin.xy;
+	uint2 ViewPixelCoord = DispatchThreadId.xy; // TODO: ViewRectMin
 	uint2 PixelCoord = ViewPixelCoord / DownsampleFactor;
 
+	// Cache Sample values
+	float4 Radiance = HeterogeneousVolumeRadiance[PixelCoord];
+	float HeterogeneousVolumeTransmittance = bUseAVSM ? SampleTransmittance(ViewPixelCoord) : Radiance.a;
+
 	float4 Result;
-	float HeterogeneousVolumeTransmittance = HeterogeneousVolumeRadiance[PixelCoord].a;
-	Result.rgb = HeterogeneousVolumeRadiance[PixelCoord].rgb + RWColorTexture[ViewPixelCoord].rgb * HeterogeneousVolumeTransmittance;
+	Result.rgb = Radiance.rgb + RWColorTexture[ViewPixelCoord].rgb * HeterogeneousVolumeTransmittance;
 #if SUPPORT_PRIMITIVE_ALPHA_HOLDOUT
 	if (View.RenderingReflectionCaptureMask == 0.0f)
 	{
diff --git a/Engine/Source/Runtime/Renderer/Private/HeterogeneousVolumes/HeterogeneousVolumes.cpp b/Engine/Source/Runtime/Renderer/Private/HeterogeneousVolumes/HeterogeneousVolumes.cpp
--- a/Engine/Source/Runtime/Renderer/Private/HeterogeneousVolumes/HeterogeneousVolumes.cpp
+++ b/Engine/Source/Runtime/Renderer/Private/HeterogeneousVolumes/HeterogeneousVolumes.cpp
@@ -30,7 +30,7 @@ static TAutoConsoleVariable<int32> CVarTranslucencyHeterogeneousVolumes(
 	TEXT("r.Translucency.HeterogeneousVolumes"),
 	0,
 	TEXT("Enables composting with heterogeneous volumes when rendering translucency (Default = 0)\n"),
-	ECVF_RenderThreadSafe | ECVF_ReadOnly
+	ECVF_RenderThreadSafe
 );
 
 static TAutoConsoleVariable<float> CVarHeterogeneousVolumesDownsampleFactor(
@@ -50,6 +50,34 @@ static TAutoConsoleVariable<int32> CVarHeterogeneousVolumesComposition(
 	ECVF_RenderThreadSafe
 );
 
+static TAutoConsoleVariable<int32> CVarHeterogeneousVolumesUpsample(
+	TEXT("r.HeterogeneousVolumes.Upsample"),
+	2,
+	TEXT("Upsampling iterative smoothing (Default = 2)\n")
+	TEXT("0: Off\n")
+	TEXT("1: Nearest Neighbor\n")
+	TEXT("2: Bilinear"),
+	ECVF_RenderThreadSafe
+);
+
+static TAutoConsoleVariable<int32> CVarHeterogeneousVolumesFilter(
+	TEXT("r.HeterogeneousVolumes.Filter"),
+	2,
+	TEXT("Controls iterative smoothing filter applied during upsampling (Default = 2)\n")
+	TEXT("0: Off\n")
+	TEXT("1: Bilateral\n")
+	TEXT("2: Gaussian 3x3\n")
+	TEXT("3: Gaussian 5x5"),
+	ECVF_RenderThreadSafe
+);
+
+static TAutoConsoleVariable<int32> CVarHeterogeneousVolumesFilterWidth(
+	TEXT("r.HeterogeneousVolumes.Filter.Width"),
+	3,
+	TEXT("Adjusts filter width of bilateral kernel (Default = 3)"),
+	ECVF_RenderThreadSafe
+);
+
 static TAutoConsoleVariable<int32> CVarHeterogeneousVolumesShadowMode(
 	TEXT("r.HeterogeneousVolumes.Shadows.Mode"),
 	0,
@@ -428,6 +456,21 @@ namespace HeterogeneousVolumes
 		return FMath::Clamp(CVarHeterogeneousVolumesDownsampleFactor.GetValueOnRenderThread(), 1, 8);
 	}
 
+	int32 GetUpsampleMode()
+	{
+		return FMath::Clamp(CVarHeterogeneousVolumesUpsample.GetValueOnRenderThread(), 0, 2);
+	}
+
+	int32 GetFilterMode()
+	{
+		return FMath::Clamp(CVarHeterogeneousVolumesFilter.GetValueOnRenderThread(), 0, 3);
+	}
+
+	int32 GetFilterWidth()
+	{
+		return FMath::Clamp(CVarHeterogeneousVolumesFilterWidth.GetValueOnRenderThread(), 0, 11);
+	}
+
 	FIntPoint GetScaledViewRect(FIntRect ViewRect)
 	{
 		return FIntPoint::DivideAndRoundUp(ViewRect.Size(), GetDownsampleFactor());
@@ -924,6 +967,7 @@ void FDeferredShadingSceneRenderer::RenderHeterogeneousVolumes(
 	if (ShouldRenderHeterogeneousVolumesForAnyView(Views))
 	{
 		FRDGTextureDesc Desc = SceneTextures.Color.Target->Desc;
+		Desc.Extent /= HeterogeneousVolumes::GetDownsampleFactor();
 		Desc.Format = PF_FloatRGBA;
 		Desc.Flags &= ~(TexCreate_FastVRAM);
 		HeterogeneousVolumeRadiance = GraphBuilder.CreateTexture(Desc, TEXT("HeterogeneousVolumes"));
@@ -1101,6 +1145,128 @@ void FDeferredShadingSceneRenderer::RenderHeterogeneousVolumes(
 	}
 }
 
+class FHeterogeneousVolumesUpsampleCS : public FGlobalShader
+{
+	DECLARE_GLOBAL_SHADER(FHeterogeneousVolumesUpsampleCS);
+	SHADER_USE_PARAMETER_STRUCT(FHeterogeneousVolumesUpsampleCS, FGlobalShader);
+
+	class FUpsampleMode : SHADER_PERMUTATION_INT("UPSAMPLE_MODE", 3);
+	using FPermutationDomain = TShaderPermutationDomain<FUpsampleMode>;
+
+	BEGIN_SHADER_PARAMETER_STRUCT(FParameters, )
+		SHADER_PARAMETER_STRUCT_REF(FViewUniformShaderParameters, View)
+		SHADER_PARAMETER_STRUCT_INCLUDE(FSceneTextureParameters, SceneTextures)
+		// Texture data
+		SHADER_PARAMETER_RDG_TEXTURE(Texture2D<float4>, HeterogeneousVolumeRadiance)
+		SHADER_PARAMETER_SAMPLER(SamplerState, TextureSampler)
+		SHADER_PARAMETER(FIntPoint, UpsampledResolution)
+		SHADER_PARAMETER(FIntPoint, UpsampledViewRect)
+		SHADER_PARAMETER(int, DownsampleFactor)
+
+		// Output
+		SHADER_PARAMETER_RDG_TEXTURE_UAV(RWTexture2D<float4>, RWUpsampledTexture)
+	END_SHADER_PARAMETER_STRUCT()
+
+public:
+	static bool ShouldCompilePermutation(
+		const FGlobalShaderPermutationParameters& Parameters
+	)
+	{
+		FPermutationDomain PermutationVector(Parameters.PermutationId);
+		if (PermutationVector.template Get<FUpsampleMode>() == 0)
+		{
+			return false;
+		}
+
+		// Apply conditional project settings for Heterogeneous volumes?
+		return DoesPlatformSupportHeterogeneousVolumes(Parameters.Platform);
+	}
+
+	static void ModifyCompilationEnvironment(
+		const FGlobalShaderPermutationParameters& Parameters,
+		FShaderCompilerEnvironment& OutEnvironment
+	)
+	{
+		FGlobalShader::ModifyCompilationEnvironment(Parameters, OutEnvironment);
+		OutEnvironment.SetDefine(TEXT("THREADGROUP_SIZE_1D"), GetThreadGroupSize1D());
+		OutEnvironment.SetDefine(TEXT("THREADGROUP_SIZE_2D"), GetThreadGroupSize2D());
+
+		// This shader takes a very long time to compile with FXC, so we pre-compile it with DXC first and then forward the optimized HLSL to FXC.
+		//OutEnvironment.CompilerFlags.Add(CFLAG_PrecompileWithDXC);
+		OutEnvironment.CompilerFlags.Add(CFLAG_AllowTypedUAVLoads);
+
+		//OutEnvironment.SetDefine(TEXT("GET_PRIMITIVE_DATA_OVERRIDE"), 1);
+	}
+
+	static int32 GetThreadGroupSize1D() { return GetThreadGroupSize2D() * GetThreadGroupSize2D(); }
+	static int32 GetThreadGroupSize2D() { return 8; }
+};
+
+IMPLEMENT_GLOBAL_SHADER(FHeterogeneousVolumesUpsampleCS, "/Engine/Private/HeterogeneousVolumes/HeterogeneousVolumesComposite.usf", "HeterogeneousVolumesUpsampleCS", SF_Compute);
+
+class FHeterogeneousVolumesFilterCS : public FGlobalShader
+{
+	DECLARE_GLOBAL_SHADER(FHeterogeneousVolumesFilterCS);
+	SHADER_USE_PARAMETER_STRUCT(FHeterogeneousVolumesFilterCS, FGlobalShader);
+
+	class FFilterMode : SHADER_PERMUTATION_INT("FILTER_MODE", 4);
+	using FPermutationDomain = TShaderPermutationDomain<FFilterMode>;
+
+	BEGIN_SHADER_PARAMETER_STRUCT(FParameters, )
+		// Scene data
+		SHADER_PARAMETER_STRUCT_REF(FViewUniformShaderParameters, View)
+		SHADER_PARAMETER_STRUCT_INCLUDE(FSceneTextureParameters, SceneTextures)
+
+		// Volume data
+		SHADER_PARAMETER_RDG_TEXTURE(Texture2D<float4>, HeterogeneousVolumeRadiance)
+		SHADER_PARAMETER_SAMPLER(SamplerState, TextureSampler)
+		SHADER_PARAMETER(FIntPoint, UpsampledResolution)
+		SHADER_PARAMETER(FIntPoint, UpsampledViewRect)
+		SHADER_PARAMETER(int, DownsampleFactor)
+		SHADER_PARAMETER(int, FilterWidth)
+
+		// Output
+		SHADER_PARAMETER_RDG_TEXTURE_UAV(RWTexture2D<float4>, RWFilteredTexture)
+		END_SHADER_PARAMETER_STRUCT()
+
+public:
+	static bool ShouldCompilePermutation(
+		const FGlobalShaderPermutationParameters& Parameters
+	)
+	{
+		FPermutationDomain PermutationVector(Parameters.PermutationId);
+		if (PermutationVector.template Get<FFilterMode>() == 0)
+		{
+			return false;
+		}
+
+		// Apply conditional project settings for Heterogeneous volumes?
+		return DoesPlatformSupportHeterogeneousVolumes(Parameters.Platform);
+	}
+
+	static void ModifyCompilationEnvironment(
+		const FGlobalShaderPermutationParameters& Parameters,
+		FShaderCompilerEnvironment& OutEnvironment
+	)
+	{
+		FGlobalShader::ModifyCompilationEnvironment(Parameters, OutEnvironment);
+		OutEnvironment.SetDefine(TEXT("THREADGROUP_SIZE_1D"), GetThreadGroupSize1D());
+		OutEnvironment.SetDefine(TEXT("THREADGROUP_SIZE_2D"), GetThreadGroupSize2D());
+
+		// This shader takes a very long time to compile with FXC, so we pre-compile it with DXC first and then forward the optimized HLSL to FXC.
+		//OutEnvironment.CompilerFlags.Add(CFLAG_PrecompileWithDXC);
+		OutEnvironment.CompilerFlags.Add(CFLAG_AllowTypedUAVLoads);
+
+		OutEnvironment.SetDefine(TEXT("GET_PRIMITIVE_DATA_OVERRIDE"), 1);
+	}
+
+	static int32 GetThreadGroupSize1D() { return GetThreadGroupSize2D() * GetThreadGroupSize2D(); }
+	static int32 GetThreadGroupSize2D() { return 8; }
+};
+
+IMPLEMENT_GLOBAL_SHADER(FHeterogeneousVolumesFilterCS, "/Engine/Private/HeterogeneousVolumes/HeterogeneousVolumesComposite.usf", "HeterogeneousVolumesFilterCS", SF_Compute);
+
+
 class FHeterogeneousVolumesCompositeCS : public FGlobalShader
 {
 	DECLARE_GLOBAL_SHADER(FHeterogeneousVolumesCompositeCS);
@@ -1109,14 +1275,19 @@ class FHeterogeneousVolumesCompositeCS : public FGlobalShader
 	BEGIN_SHADER_PARAMETER_STRUCT(FParameters, )
 		// Scene data
 		SHADER_PARAMETER_STRUCT_REF(FViewUniformShaderParameters, View)
+		SHADER_PARAMETER_STRUCT_INCLUDE(FSceneTextureParameters, SceneTextures)
 
 		// Volume data
 		SHADER_PARAMETER_RDG_TEXTURE(Texture2D<float4>, HeterogeneousVolumeRadiance)
 		SHADER_PARAMETER_RDG_TEXTURE(Texture2D<float>, HeterogeneousVolumeHoldout)
 
+		// Transmittance structure
+		SHADER_PARAMETER_RDG_UNIFORM_BUFFER(FAdaptiveVolumetricShadowMapUniformBufferParameters, AVSM)
+		SHADER_PARAMETER(int32, bUseAVSM)
+
 		// Dispatch data
 		SHADER_PARAMETER(FIntVector, GroupCount)
-		SHADER_PARAMETER(int32, DownsampleFactor)
+		SHADER_PARAMETER(int, DownsampleFactor)
 
 		// Output
 		SHADER_PARAMETER_RDG_TEXTURE_UAV(RWTexture2D<float4>, RWColorTexture)
@@ -1159,35 +1330,132 @@ void FDeferredShadingSceneRenderer::CompositeHeterogeneousVolumes(
 	const FSceneTextures& SceneTextures
 )
 {
+	RDG_EVENT_SCOPE_STAT(GraphBuilder, HeterogeneousVolumesStat, "HeterogeneousVolumes");
+	RDG_GPU_STAT_SCOPE(GraphBuilder, HeterogeneousVolumesStat);
+	SCOPED_NAMED_EVENT(HeterogeneousVolumes, FColor::Emerald);
+
 	for (int32 ViewIndex = 0; ViewIndex < Views.Num(); ViewIndex++)
 	{
 		FViewInfo& View = Views[ViewIndex];
 
 		if (ShouldRenderHeterogeneousVolumesForView(View))
 		{
-			FIntVector GroupCount = FComputeShaderUtils::GetGroupCount(View.ViewRect.Size(), FHeterogeneousVolumesCompositeCS::GetThreadGroupSize2D());
+			FRDGTextureDesc UpsampleDesc = View.HeterogeneousVolumeRadiance->Desc;
+			int32 DownsampleFactor = HeterogeneousVolumes::GetDownsampleFactor();
 
-			FHeterogeneousVolumesCompositeCS::FParameters* PassParameters = GraphBuilder.AllocParameters<FHeterogeneousVolumesCompositeCS::FParameters>();
+			bool bIterativeUpsampling = HeterogeneousVolumes::GetUpsampleMode() != 0;
+			while (bIterativeUpsampling && (DownsampleFactor > 1))
 			{
-				// Scene data
-				PassParameters->View = View.ViewUniformBuffer;
-				// Volume data
-				PassParameters->HeterogeneousVolumeRadiance = View.HeterogeneousVolumeRadiance;
-				PassParameters->HeterogeneousVolumeHoldout = View.HeterogeneousVolumeHoldout;
-				// Dispatch data
-				PassParameters->GroupCount = GroupCount;
-				PassParameters->DownsampleFactor = HeterogeneousVolumes::GetDownsampleFactor();
-				// Output
-				PassParameters->RWColorTexture = GraphBuilder.CreateUAV(SceneTextures.Color.Target);
+				int32 DownsampleIterationFactor = (DownsampleFactor % 3 == 0) ? 3 : 2;
+
+				FIntPoint DownsampledViewRect = FIntPoint::DivideAndRoundUp(View.ViewRect.Size(), DownsampleFactor);
+				FIntPoint UpsampledViewRect = DownsampledViewRect * DownsampleIterationFactor;
+				FIntVector GroupCount = FComputeShaderUtils::GetGroupCount(UpsampledViewRect, FHeterogeneousVolumesUpsampleCS::GetThreadGroupSize2D());
+				UpsampleDesc.Extent *= DownsampleIterationFactor;
+
+				// Upsample
+				{
+					FRDGTextureRef HeterogeneousVolumeUpsample = GraphBuilder.CreateTexture(UpsampleDesc, TEXT("HeterogeneousVolumeUpsample"));
+
+					FHeterogeneousVolumesUpsampleCS::FParameters* PassParameters = GraphBuilder.AllocParameters<FHeterogeneousVolumesUpsampleCS::FParameters>();
+					{
+						// Scene data
+						PassParameters->View = View.ViewUniformBuffer;
+						PassParameters->SceneTextures = GetSceneTextureParameters(GraphBuilder, SceneTextures);
+
+						// Texture data
+						PassParameters->HeterogeneousVolumeRadiance = View.HeterogeneousVolumeRadiance;
+						PassParameters->TextureSampler = TStaticSamplerState<SF_Bilinear, AM_Clamp, AM_Clamp, AM_Clamp>::GetRHI();
+						PassParameters->UpsampledResolution = UpsampleDesc.Extent;
+						PassParameters->UpsampledViewRect = UpsampledViewRect;
+						PassParameters->DownsampleFactor = DownsampleIterationFactor;
+
+						// Output
+						PassParameters->RWUpsampledTexture = GraphBuilder.CreateUAV(HeterogeneousVolumeUpsample);
+					}
+
+					FHeterogeneousVolumesUpsampleCS::FPermutationDomain PermutationVector;
+					PermutationVector.Set<FHeterogeneousVolumesUpsampleCS::FUpsampleMode>(HeterogeneousVolumes::GetUpsampleMode());
+					TShaderRef<FHeterogeneousVolumesUpsampleCS> ComputeShader = View.ShaderMap->GetShader<FHeterogeneousVolumesUpsampleCS>(PermutationVector);
+					FComputeShaderUtils::AddPass(
+						GraphBuilder,
+						RDG_EVENT_NAME("FHeterogeneousVolumesUpsampleCS"),
+						ComputeShader,
+						PassParameters,
+						GroupCount);
+
+					View.HeterogeneousVolumeRadiance = HeterogeneousVolumeUpsample;
+				}
+
+				// Filter neighborhood
+				if (HeterogeneousVolumes::GetFilterMode() != 0)
+				{
+					FRDGTextureRef HeterogeneousVolumeFilter = GraphBuilder.CreateTexture(UpsampleDesc, TEXT("HeterogeneousVolumeFilter"));
+
+					FHeterogeneousVolumesFilterCS::FParameters* PassParameters = GraphBuilder.AllocParameters<FHeterogeneousVolumesFilterCS::FParameters>();
+					{
+						// Scene data
+						PassParameters->View = View.ViewUniformBuffer;
+						PassParameters->SceneTextures = GetSceneTextureParameters(GraphBuilder, SceneTextures);
+						// Volume data
+						PassParameters->HeterogeneousVolumeRadiance = View.HeterogeneousVolumeRadiance;
+						PassParameters->TextureSampler = TStaticSamplerState<SF_Bilinear, AM_Clamp, AM_Clamp, AM_Clamp>::GetRHI();
+						PassParameters->UpsampledResolution = UpsampleDesc.Extent;
+						PassParameters->UpsampledViewRect = UpsampledViewRect;
+						PassParameters->DownsampleFactor = 1;
+						PassParameters->FilterWidth = HeterogeneousVolumes::GetFilterWidth();
+
+						// Output
+						PassParameters->RWFilteredTexture = GraphBuilder.CreateUAV(HeterogeneousVolumeFilter);
+					}
+
+					FHeterogeneousVolumesFilterCS::FPermutationDomain PermutationVector;
+					PermutationVector.Set<FHeterogeneousVolumesFilterCS::FFilterMode>(HeterogeneousVolumes::GetFilterMode());
+					TShaderRef<FHeterogeneousVolumesFilterCS> ComputeShader = View.ShaderMap->GetShader<FHeterogeneousVolumesFilterCS>(PermutationVector);
+					FComputeShaderUtils::AddPass(
+						GraphBuilder,
+						RDG_EVENT_NAME("FHeterogeneousVolumesFilterCS"),
+						ComputeShader,
+						PassParameters,
+						GroupCount);
+
+					View.HeterogeneousVolumeRadiance = HeterogeneousVolumeFilter;
+				}
+
+				DownsampleFactor /= DownsampleIterationFactor;
 			}
 
-			TShaderRef<FHeterogeneousVolumesCompositeCS> ComputeShader = View.ShaderMap->GetShader<FHeterogeneousVolumesCompositeCS>();
-			FComputeShaderUtils::AddPass(
-				GraphBuilder,
-				RDG_EVENT_NAME("FHeterogeneousVolumesCompositeCS"),
-				ComputeShader,
-				PassParameters,
-				GroupCount);
+			// Composite
+			{
+				FIntVector GroupCount = FComputeShaderUtils::GetGroupCount(View.ViewRect.Size(), FHeterogeneousVolumesUpsampleCS::GetThreadGroupSize2D());
+				FHeterogeneousVolumesCompositeCS::FParameters* PassParameters = GraphBuilder.AllocParameters<FHeterogeneousVolumesCompositeCS::FParameters>();
+				{
+					// Scene data
+					PassParameters->View = View.ViewUniformBuffer;
+					PassParameters->SceneTextures = GetSceneTextureParameters(GraphBuilder, SceneTextures);
+					// Volume data
+					PassParameters->HeterogeneousVolumeRadiance = View.HeterogeneousVolumeRadiance;
+					PassParameters->HeterogeneousVolumeHoldout = View.HeterogeneousVolumeHoldout;
+					// Transmittance structure
+					PassParameters->bUseAVSM = ShouldCompositeHeterogeneousVolumesWithTranslucency();
+					PassParameters->AVSM = HeterogeneousVolumes::GetAdaptiveVolumetricCameraMapUniformBuffer(GraphBuilder, View.ViewState);
+					// Dispatch data
+					PassParameters->GroupCount = GroupCount;
+					//PassParameters->DownsampleFactor = HeterogeneousVolumes::GetDownsampleFactor();
+					PassParameters->DownsampleFactor = DownsampleFactor;
+
+					// Output
+					PassParameters->RWColorTexture = GraphBuilder.CreateUAV(SceneTextures.Color.Target);
+				}
+
+				TShaderRef<FHeterogeneousVolumesCompositeCS> ComputeShader = View.ShaderMap->GetShader<FHeterogeneousVolumesCompositeCS>();
+				FComputeShaderUtils::AddPass(
+					GraphBuilder,
+					RDG_EVENT_NAME("FHeterogeneousVolumesCompositeCS"),
+					ComputeShader,
+					PassParameters,
+					GroupCount);
+			}
 		}
 	}
 }
diff --git a/Engine/Source/Runtime/Renderer/Private/HeterogeneousVolumes/HeterogeneousVolumesVoxelGridPipeline.cpp b/Engine/Source/Runtime/Renderer/Private/HeterogeneousVolumes/HeterogeneousVolumesVoxelGridPipeline.cpp
--- a/Engine/Source/Runtime/Renderer/Private/HeterogeneousVolumes/HeterogeneousVolumesVoxelGridPipeline.cpp
+++ b/Engine/Source/Runtime/Renderer/Private/HeterogeneousVolumes/HeterogeneousVolumesVoxelGridPipeline.cpp
@@ -3280,7 +3280,7 @@ namespace HeterogeneousVolumes {
 		FSceneViewState* ViewState
 	)
 	{
-		if (ViewState)
+		if (ViewState && ShouldCompositeHeterogeneousVolumesWithTranslucency())
 		{
 			return ViewState->AdaptiveVolumetricCameraMapUniformBuffer;
 		}
