[HWRT] Rename cvars related to debug visualizations.

diff --git a/Engine/Source/Runtime/Core/Public/HAL/IConsoleManager.h b/Engine/Source/Runtime/Core/Public/HAL/IConsoleManager.h
--- a/Engine/Source/Runtime/Core/Public/HAL/IConsoleManager.h
+++ b/Engine/Source/Runtime/Core/Public/HAL/IConsoleManager.h
@@ -2143,7 +2143,7 @@ enum class EShadowCVarBehavior : uint8
 /**
  * A class that can be used to shadow a CVar, where both are valid to use. To deprecate a CVar, it's suggested to use the FAutoConsoleVariableDeprecated class below
  * 
- * FAutoConsoleVariableShadow ShadowCVar(TEXT("r.newname"), TEXT("r.oldname"));
+ * FAutoConsoleVariableShadow ShadowCVar(TEXT("r.oldname"), TEXT("r.newname"));
  * 
  */
 class FAutoConsoleVariableShadow
@@ -2162,8 +2162,8 @@ public:
  * A class that can be used to deprecate a CVar - rename your existing CVar to the new name, then
  * create a shadow of it with the old name with something like this:
  *
- * FAutoConsoleVariableDeprecated ShadowCVar(TEXT("r.newname"), TEXT("r.oldname"), TEXT("5.5"));
- * FAutoConsoleVariableDeprecated ShadowCVar(TEXT("r.newname"), TEXT("r.oldname"), TEXT("5.5"), EShadowCVarBehavior::Assert);
+ * FAutoConsoleVariableDeprecated ShadowCVar(TEXT("r.oldname"), TEXT("r.newname"), TEXT("5.5"));
+ * FAutoConsoleVariableDeprecated ShadowCVar(TEXT("r.oldname"), TEXT("r.newname"), TEXT("5.5"), EShadowCVarBehavior::Assert);
  *
  */
 class FAutoConsoleVariableDeprecated
diff --git a/Engine/Source/Runtime/Renderer/Private/RayTracing/RayTracingDebug.cpp b/Engine/Source/Runtime/Renderer/Private/RayTracing/RayTracingDebug.cpp
--- a/Engine/Source/Runtime/Renderer/Private/RayTracing/RayTracingDebug.cpp
+++ b/Engine/Source/Runtime/Renderer/Private/RayTracing/RayTracingDebug.cpp
@@ -28,14 +28,16 @@
 
 DECLARE_GPU_STAT(RayTracingDebug);
 
-static TAutoConsoleVariable<FString> CVarRayTracingDebugMode(
-	TEXT("r.RayTracing.DebugVisualizationMode"),
+static TAutoConsoleVariable<FString> CVarRayTracingVisualize(
+	TEXT("r.RayTracing.Visualize"),
 	TEXT(""),
 	TEXT("Sets the ray tracing debug visualization mode (default = None - Driven by viewport menu) .\n")
 	);
 
-TAutoConsoleVariable<int32> CVarRayTracingDebugPickerDomain(
-	TEXT("r.RayTracing.Debug.PickerDomain"),
+static FAutoConsoleVariableDeprecated CVarRayTracingDebugMode_Deprecated(TEXT("r.RayTracing.DebugVisualizationMode"), TEXT("r.RayTracing.Visualize"), TEXT("5.6"));
+
+TAutoConsoleVariable<int32> CVarRayTracingVisualizePickerDomain(
+	TEXT("r.RayTracing.Visualize.PickerDomain"),
 	0,
 	TEXT("Changes the picker domain to highlight:\n")
 	TEXT("0 - Triangles (default)\n")
@@ -43,82 +45,108 @@ TAutoConsoleVariable<int32> CVarRayTracingDebugPickerDomain(
 	ECVF_RenderThreadSafe
 );
 
-TAutoConsoleVariable<int32> CVarRayTracingDebugModeOpaqueOnly(
-	TEXT("r.RayTracing.DebugVisualizationMode.OpaqueOnly"),
+static FAutoConsoleVariableDeprecated CVarRayTracingDebugPickerDomain_Deprecated(TEXT("r.RayTracing.Debug.PickerDomain"), TEXT("r.RayTracing.Visualize.PickerDomain"), TEXT("5.6"));
+
+static TAutoConsoleVariable<int32> CVarRayTracingVisualizeOpaqueOnly(
+	TEXT("r.RayTracing.Visualize.OpaqueOnly"),
 	1,
 	TEXT("Sets whether the view mode rendes opaque objects only (default = 1, render only opaque objects, 0 = render all objects)"),
 	ECVF_RenderThreadSafe
 );
 
-static TAutoConsoleVariable<float> CVarRayTracingDebugTimingScale(
-	TEXT("r.RayTracing.DebugTimingScale"),
+static FAutoConsoleVariableDeprecated CVarRayTracingDebugModeOpaqueOnly_Deprecated(TEXT("r.RayTracing.DebugVisualizationMode.OpaqueOnly"), TEXT("r.RayTracing.Visualize.OpaqueOnly"), TEXT("5.6"));
+
+static TAutoConsoleVariable<float> CVarRayTracingVisualizeTimingScale(
+	TEXT("r.RayTracing.Visualize.TimingScale"),
 	1.0f,
 	TEXT("Scaling factor for ray timing heat map visualization. (default = 1)\n")
 );
 
-static TAutoConsoleVariable<float> CVarRayTracingDebugTraversalBoxScale(
-	TEXT("r.RayTracing.DebugTraversalScale.Box"),
+static FAutoConsoleVariableDeprecated CVarRayTracingDebugTimingScale_Deprecated(TEXT("r.RayTracing.DebugTimingScale"), TEXT("r.RayTracing.Visualize.TimingScale"), TEXT("5.6"));
+
+static TAutoConsoleVariable<float> CVarRayTracingVisualizeTraversalBoxScale(
+	TEXT("r.RayTracing.Visualize.Traversal.BoxScale"),
 	150.0f,
 	TEXT("Scaling factor for box traversal heat map visualization. (default = 150)\n")
 );
 
-static TAutoConsoleVariable<float> CVarRayTracingDebugTraversalClusterScale(
-	TEXT("r.RayTracing.DebugTraversalScale.Cluster"),
+static FAutoConsoleVariableDeprecated CVarRayTracingDebugTraversalBoxScale_Deprecated(TEXT("r.RayTracing.DebugTraversalScale.Box"), TEXT("r.RayTracing.Visualize.Traversal.BoxScale"), TEXT("5.6"));
+
+static TAutoConsoleVariable<float> CVarRayTracingVisualizeTraversalClusterScale(
+	TEXT("r.RayTracing.Visualize.Traversal.ClusterScale"),
 	2500.0f,
 	TEXT("Scaling factor for cluster traversal heat map visualization. (default = 2500)\n")
 );
 
-static TAutoConsoleVariable<float> CVarRayTracingDebugInstanceOverlapScale(
-	TEXT("r.RayTracing.Debug.InstanceOverlap.Scale"),
+static FAutoConsoleVariableDeprecated CVarRayTracingDebugTraversalClusterScale_Deprecated(TEXT("r.RayTracing.DebugTraversalScale.Cluster"), TEXT("r.RayTracing.Visualize.Traversal.ClusterScale"), TEXT("5.6"));
+
+static TAutoConsoleVariable<float> CVarRayTracingVisualizeInstanceOverlapScale(
+	TEXT("r.RayTracing.Visualize.InstanceOverlap.Scale"),
 	16.0f,
 	TEXT("Scaling factor for instance traversal heat map visualization. (default = 16)\n")
 );
 
-static TAutoConsoleVariable<float> CVarRayTracingDebugInstanceOverlapBoundingBoxScale(
-	TEXT("r.RayTracing.Debug.InstanceOverlap.BoundingBoxScale"),
+static FAutoConsoleVariableDeprecated CVarRayTracingDebugInstanceOverlapScale_Deprecated(TEXT("r.RayTracing.Debug.InstanceOverlap.Scale"), TEXT("r.RayTracing.Visualize.InstanceOverlap.Scale"), TEXT("5.6"));
+
+static TAutoConsoleVariable<float> CVarRayTracingVisualizeInstanceOverlapBoundingBoxScale(
+	TEXT("r.RayTracing.Visualize.InstanceOverlap.BoundingBoxScale"),
 	1.001f,
 	TEXT("Scaling factor for instance bounding box extent for avoiding z-fighting. (default = 1.001)\n")
 );
 
-static TAutoConsoleVariable<int32> CVarRayTracingDebugInstanceOverlapShowWireframe(
-	TEXT("r.RayTracing.Debug.InstanceOverlap.ShowWireframe"),
+static FAutoConsoleVariableDeprecated CVarRayTracingDebugInstanceOverlapBoundingBoxScale_Deprecated(TEXT("r.RayTracing.Debug.InstanceOverlap.BoundingBoxScale"), TEXT("r.RayTracing.Visualize.InstanceOverlap.BoundingBoxScale"), TEXT("5.6"));
+
+static TAutoConsoleVariable<int32> CVarRayTracingVisualizeInstanceOverlapShowWireframe(
+	TEXT("r.RayTracing.Visualize.InstanceOverlap.ShowWireframe"),
 	1,
 	TEXT("Show instance bounding boxes in wireframe in Instances Overlap mode. (default = 1)\n")
 );
 
-static TAutoConsoleVariable<float> CVarRayTracingDebugTraversalTriangleScale(
-	TEXT("r.RayTracing.DebugTraversalScale.Triangle"),
+static FAutoConsoleVariableDeprecated CVarRayTracingDebugInstanceOverlapShowWireframe_Deprecated(TEXT("r.RayTracing.Debug.InstanceOverlap.ShowWireframe"), TEXT("r.RayTracing.Visualize.InstanceOverlap.ShowWireframe"), TEXT("5.6"));
+
+static TAutoConsoleVariable<float> CVarRayTracingVisualizeTraversalTriangleScale(
+	TEXT("r.RayTracing.Visualize.Traversal.TriangleScale"),
 	30.0f,
 	TEXT("Scaling factor for triangle traversal heat map visualization. (default = 30)\n")
 );
 
-static TAutoConsoleVariable<int32> CVarRayTracingDebugHitCountMaxThreshold(
-	TEXT("r.RayTracing.DebugTriangleHitCount.MaxThreshold"),
+static FAutoConsoleVariableDeprecated CVarRayTracingDebugTraversalTriangleScale_Deprecated(TEXT("r.RayTracing.DebugTraversalScale.Triangle"), TEXT("r.RayTracing.Visualize.Traversal.TriangleScale"), TEXT("5.6"));
+
+static TAutoConsoleVariable<int32> CVarRayTracingVisualizeHitCountMaxThreshold(
+	TEXT("r.RayTracing.Visualize.TriangleHitCount.MaxThreshold"),
 	6,
 	TEXT("Maximum hit count threshold for debug ray tracing triangle hit count heat map visualization. (default = 6)\n")
 );
 
-static TAutoConsoleVariable<int32> CVarRayTracingDebugHitCountPerInstanceMaxThreshold(
-	TEXT("r.RayTracing.DebugTriangleHitCountPerInstance.MaxThreshold"),
+static FAutoConsoleVariableDeprecated CVarRayTracingDebugHitCountMaxThreshold_Deprecated(TEXT("r.RayTracing.DebugTriangleHitCount.MaxThreshold"), TEXT("r.RayTracing.Visualize.TriangleHitCount.MaxThreshold"), TEXT("5.6"));
+
+static TAutoConsoleVariable<int32> CVarRayTracingVisualizeHitCountPerInstanceMaxThreshold(
+	TEXT("r.RayTracing.Visualize.HitCountPerInstance.MaxThreshold"),
 	100000,
 	TEXT("Maximum hit count threshold for debug ray tracing hit count per instance heat map visualization. (default = 100000)\n")
 );
 
-static TAutoConsoleVariable<int32> CVarRayTracingDebugHitCountTopKHits(
-	TEXT("r.RayTracing.DebugTriangleHitCount.TopKMostHits"),
+static FAutoConsoleVariableDeprecated CVarRayTracingDebugHitCountPerInstanceMaxThreshold_Deprecated(TEXT("r.RayTracing.DebugTriangleHitCountPerInstance.MaxThreshold"), TEXT("r.RayTracing.Visualize.HitCountPerInstance.MaxThreshold"), TEXT("5.6"));
+
+static TAutoConsoleVariable<int32> CVarRayTracingVisualizeHitCountTopKHits(
+	TEXT("r.RayTracing.Visualize.TriangleHitCount.TopKMostHits"),
 	10,
 	TEXT("Highlight top k most hit instances in the view. (default = 10)\n")
 );
 
+static FAutoConsoleVariableDeprecated CVarRayTracingDebugHitCountTopKHits_Deprecated(TEXT("r.RayTracing.DebugTriangleHitCount.TopKMostHits"), TEXT("r.RayTracing.Visualize.TriangleHitCount.TopKMostHits"), TEXT("5.6"));
+
 static int32 GVisualizeProceduralPrimitives = 0;
 static FAutoConsoleVariableRef CVarVisualizeProceduralPrimitives(
-	TEXT("r.RayTracing.DebugVisualizationMode.ProceduralPrimitives"),
+	TEXT("r.RayTracing.Visualize.ProceduralPrimitives"),
 	GVisualizeProceduralPrimitives,
 	TEXT("Whether to include procedural primitives in visualization modes.\n")
 	TEXT("Currently only supports Nanite primitives in inline barycentrics mode."),
 	ECVF_RenderThreadSafe
 );
 
+static FAutoConsoleVariableDeprecated CVarVisualizeProceduralPrimitives_Deprecated(TEXT("r.RayTracing.DebugVisualizationMode.ProceduralPrimitives"), TEXT("r.RayTracing.Visualize.ProceduralPrimitives"), TEXT("5.6"));
+
 IMPLEMENT_RT_PAYLOAD_TYPE(ERayTracingPayloadType::RayTracingDebug, 44);
 
 BEGIN_UNIFORM_BUFFER_STRUCT(FRayTracingDebugHitStatsUniformBufferParameters, )
@@ -858,7 +886,7 @@ static FRDGBufferRef RayTracingPerformPicking(FRDGBuilder& GraphBuilder, const F
 	FRayTracingPickingRGS::FParameters* RayGenParameters = GraphBuilder.AllocParameters<FRayTracingPickingRGS::FParameters>();
 	RayGenParameters->InstancesDebugData = GraphBuilder.CreateSRV(RayTracingScene.GetInstanceDebugBuffer(ERayTracingSceneLayer::Base));
 	RayGenParameters->TLAS = RayTracingScene.GetLayerView(ERayTracingSceneLayer::Base);
-	RayGenParameters->OpaqueOnly = CVarRayTracingDebugModeOpaqueOnly.GetValueOnRenderThread();
+	RayGenParameters->OpaqueOnly = CVarRayTracingVisualizeOpaqueOnly.GetValueOnRenderThread();
 	RayGenParameters->InstanceBuffer = GraphBuilder.CreateSRV(RayTracingScene.GetInstanceBuffer(ERayTracingSceneLayer::Base));
 	RayGenParameters->ViewUniformBuffer = View.ViewUniformBuffer;
 	RayGenParameters->SceneUniformBuffer = GetSceneUniformBufferRef(GraphBuilder, View); // TODO: use a separate params structure
@@ -984,7 +1012,7 @@ static void PrintTopKMostHitMessage(FRDGBuilder& GraphBuilder, const FScene* Sce
 	// Force ShaderPrint on.
 	ShaderPrint::SetEnabled(true);
 
-	int32 NumPrimitives = CVarRayTracingDebugHitCountTopKHits.GetValueOnRenderThread();
+	int32 NumPrimitives = CVarRayTracingVisualizeHitCountTopKHits.GetValueOnRenderThread();
 	if (ShaderPrint::IsEnabled(View.ShaderPrintData) && NumPrimitives > 0)
 	{
 		// This lags by one frame, so may miss some in one frame, also overallocates since we will cull a lot.
@@ -1103,7 +1131,7 @@ static FRDGBufferRef RayTracingPerformHitStatsPerPrimitive(FRDGBuilder& GraphBui
 	FShaderBindingTableRHIRef HitStatsSBT = Scene->RayTracingSBT.AllocateTransientRHI(GraphBuilder.RHICmdList, ERayTracingShaderBindingMode::RTPSO, ERayTracingHitGroupIndexingMode::Allow, RayTracingScene.NumMissShaderSlots, RayTracingScene.NumCallableShaderSlots, Initializer.GetMaxLocalBindingDataSize());
 
 	// TODO: Should check RayTracingScene for actual number of instances instead of max number in FRHIRayTracingScene initializer
-	const uint32 NumInstancesInTLAS = FMath::Max(RayTracingScene.GetRHIRayTracingSceneChecked(ERayTracingSceneLayer::Base)->GetInitializer().MaxNumInstances, (uint32)CVarRayTracingDebugHitCountTopKHits.GetValueOnRenderThread());
+	const uint32 NumInstancesInTLAS = FMath::Max(RayTracingScene.GetRHIRayTracingSceneChecked(ERayTracingSceneLayer::Base)->GetInitializer().MaxNumInstances, (uint32)CVarRayTracingVisualizeHitCountTopKHits.GetValueOnRenderThread());
 	FRDGBufferDesc HitStatsPerPrimitiveBufferDesc = FRDGBufferDesc::CreateStructuredDesc(sizeof(FRayTracingHitStatsEntry), NumInstancesInTLAS);
 	HitStatsPerPrimitiveBufferDesc.Usage = EBufferUsageFlags(HitStatsPerPrimitiveBufferDesc.Usage | BUF_SourceCopy);
 	FRDGBufferRef HitStatsBuffer = GraphBuilder.CreateBuffer(HitStatsPerPrimitiveBufferDesc, TEXT("RayTracingDebug.HitStatsBuffer"));
@@ -1114,7 +1142,7 @@ static FRDGBufferRef RayTracingPerformHitStatsPerPrimitive(FRDGBuilder& GraphBui
 
 	FRayTracingDebugHitStatsRGS::FParameters* RayGenParameters = GraphBuilder.AllocParameters<FRayTracingDebugHitStatsRGS::FParameters>();
 	RayGenParameters->TLAS = RayTracingScene.GetLayerView(ERayTracingSceneLayer::Base);
-	RayGenParameters->OpaqueOnly = CVarRayTracingDebugModeOpaqueOnly.GetValueOnRenderThread();
+	RayGenParameters->OpaqueOnly = CVarRayTracingVisualizeOpaqueOnly.GetValueOnRenderThread();
 	RayGenParameters->ViewUniformBuffer = View.ViewUniformBuffer;
 	RayGenParameters->SceneUniformBuffer = GetSceneUniformBufferRef(GraphBuilder, View); // TODO: use a separate params structure
 	RayGenParameters->RayTracingDebugHitStatsUniformBuffer = DebugHitStatsUniformBuffer;
@@ -1221,7 +1249,7 @@ static void RayTracingDrawInstances(FRDGBuilder& GraphBuilder, const FViewInfo&
 	PassParameters->VS.View = View.ViewUniformBuffer;
 	PassParameters->VS.Scene = View.GetSceneUniforms().GetBuffer(GraphBuilder);
 	PassParameters->VS.InstanceGPUSceneIndices = GraphBuilder.CreateSRV(InstanceGPUSceneIndexBuffer);
-	PassParameters->VS.BoundingBoxExtentScale = CVarRayTracingDebugInstanceOverlapBoundingBoxScale.GetValueOnRenderThread();
+	PassParameters->VS.BoundingBoxExtentScale = CVarRayTracingVisualizeInstanceOverlapBoundingBoxScale.GetValueOnRenderThread();
 
 	PassParameters->PS.View = View.ViewUniformBuffer;
 
@@ -1309,7 +1337,7 @@ static void DrawInstanceOverlap(FRDGBuilder& GraphBuilder, const FScene* Scene,
 		PassParameters->RenderTargets[0] = FRenderTargetBinding(SceneColorTexture, ERenderTargetLoadAction::ELoad, 0);
 
 		PassParameters->InstanceOverlap = GraphBuilder.CreateSRV(InstanceOverlapTexture);
-		PassParameters->HeatmapScale = CVarRayTracingDebugInstanceOverlapScale.GetValueOnRenderThread();
+		PassParameters->HeatmapScale = CVarRayTracingVisualizeInstanceOverlapScale.GetValueOnRenderThread();
 
 		TShaderMapRef<FRayTracingDebugBlendInstanceOverlapPS> PixelShader(View.ShaderMap);
 
@@ -1326,7 +1354,7 @@ static void DrawInstanceOverlap(FRDGBuilder& GraphBuilder, const FScene* Scene,
 	}
 
 	// Draw instance AABB with lines
-	if (CVarRayTracingDebugInstanceOverlapShowWireframe.GetValueOnRenderThread() != 0)
+	if (CVarRayTracingVisualizeInstanceOverlapShowWireframe.GetValueOnRenderThread() != 0)
 	{
 		RayTracingDrawInstances(GraphBuilder, View, SceneColorTexture, SceneDepthTexture, Scene->RayTracingScene.GetDebugInstanceGPUSceneIndexBuffer(ERayTracingSceneLayer::Base), true);
 	}
@@ -1338,7 +1366,7 @@ static uint32 GetRaytracingDebugViewMode(const FSceneView& View)
 
 	uint32 OutMode;
 
-	FString ConsoleViewMode = CVarRayTracingDebugMode.GetValueOnAnyThread();
+	FString ConsoleViewMode = CVarRayTracingVisualize.GetValueOnAnyThread();
 	if (!ConsoleViewMode.IsEmpty())
 	{
 		OutMode = VisualizationData.GetModeID(FName(*ConsoleViewMode));
@@ -1400,9 +1428,9 @@ void FDeferredShadingSceneRenderer::RenderRayTracingDebug(FRDGBuilder& GraphBuil
 		PassParameters->NaniteShadingUniformBuffer = CreateDebugNaniteShadingUniformBuffer(GraphBuilder);
 
 		PassParameters->VisualizationMode = DebugVisualizationMode;
-		PassParameters->TraversalBoxScale = CVarRayTracingDebugTraversalBoxScale.GetValueOnAnyThread();
-		PassParameters->TraversalClusterScale = CVarRayTracingDebugTraversalClusterScale.GetValueOnAnyThread();
-		PassParameters->TraversalTriangleScale = CVarRayTracingDebugTraversalTriangleScale.GetValueOnAnyThread();
+		PassParameters->TraversalBoxScale = CVarRayTracingVisualizeTraversalBoxScale.GetValueOnAnyThread();
+		PassParameters->TraversalClusterScale = CVarRayTracingVisualizeTraversalClusterScale.GetValueOnAnyThread();
+		PassParameters->TraversalTriangleScale = CVarRayTracingVisualizeTraversalTriangleScale.GetValueOnAnyThread();
 
 		PassParameters->RTDebugVisualizationNaniteCutError = 0.0f;
 
@@ -1544,15 +1572,15 @@ void FDeferredShadingSceneRenderer::RenderRayTracingDebug(FRDGBuilder& GraphBuil
 
 	FRayTracingDebugRGS::FParameters* RayGenParameters = GraphBuilder.AllocParameters<FRayTracingDebugRGS::FParameters>();
 	RayGenParameters->VisualizationMode = DebugVisualizationMode;
-	RayGenParameters->PickerDomain = CVarRayTracingDebugPickerDomain.GetValueOnRenderThread();
+	RayGenParameters->PickerDomain = CVarRayTracingVisualizePickerDomain.GetValueOnRenderThread();
 	RayGenParameters->ShouldUsePreExposure = View.Family->EngineShowFlags.Tonemapper;
-	RayGenParameters->TimingScale = CVarRayTracingDebugTimingScale.GetValueOnAnyThread() / 25000.0f;
-	RayGenParameters->OpaqueOnly = CVarRayTracingDebugModeOpaqueOnly.GetValueOnRenderThread();
-	RayGenParameters->TriangleHitCountMaxThreshold = FMath::Clamp((float)CVarRayTracingDebugHitCountMaxThreshold.GetValueOnRenderThread(), 1, 100000);
-	RayGenParameters->TriangleHitCountPerInstanceMaxThreshold = FMath::Max(1, CVarRayTracingDebugHitCountPerInstanceMaxThreshold.GetValueOnRenderThread());
+	RayGenParameters->TimingScale = CVarRayTracingVisualizeTimingScale.GetValueOnAnyThread() / 25000.0f;
+	RayGenParameters->OpaqueOnly = CVarRayTracingVisualizeOpaqueOnly.GetValueOnRenderThread();
+	RayGenParameters->TriangleHitCountMaxThreshold = FMath::Clamp((float)CVarRayTracingVisualizeHitCountMaxThreshold.GetValueOnRenderThread(), 1, 100000);
+	RayGenParameters->TriangleHitCountPerInstanceMaxThreshold = FMath::Max(1, CVarRayTracingVisualizeHitCountPerInstanceMaxThreshold.GetValueOnRenderThread());
 	RayGenParameters->RayTracingDebugHitStatsUniformBuffer = DebugHitStatsUniformBuffer;
 	RayGenParameters->LightGridPacked = View.RayTracingLightGridUniformBuffer;
-	RayGenParameters->TopKMostHitInstances = CVarRayTracingDebugHitCountTopKHits.GetValueOnRenderThread();
+	RayGenParameters->TopKMostHitInstances = CVarRayTracingVisualizeHitCountTopKHits.GetValueOnRenderThread();
 	RayGenParameters->NumTotalInstances = NumInstances;
 	ShaderPrint::SetParameters(GraphBuilder, View.ShaderPrintData, RayGenParameters->ShaderPrintUniformBuffer);
 
@@ -1638,7 +1666,7 @@ void FDeferredShadingSceneRenderer::RayTracingDisplayPicking(const FRayTracingPi
 		return;
 	}
 	
-	int32 PickerDomain = CVarRayTracingDebugPickerDomain.GetValueOnRenderThread();
+	int32 PickerDomain = CVarRayTracingVisualizePickerDomain.GetValueOnRenderThread();
 	switch (PickerDomain)
 	{
 	case RAY_TRACING_DEBUG_PICKER_DOMAIN_TRIANGLE:
@@ -1761,7 +1789,7 @@ bool IsRayTracingInstanceDebugDataEnabled(const FViewInfo& View)
 		RayTracingDebugVisualizationModes.Add(*LOCTEXT("Picker", "Picker").ToString());
 	}
 
-	FString ConsoleViewMode = CVarRayTracingDebugMode.GetValueOnRenderThread();
+	FString ConsoleViewMode = CVarRayTracingVisualize.GetValueOnRenderThread();
 	if (!ConsoleViewMode.IsEmpty())
 	{
 		return RayTracingDebugVisualizationModes.Contains(FName(*ConsoleViewMode));
