Fixed crash due to indexing uninitialized data in TSparseArray in the VSM cache manager.

diff --git a/Engine/Source/Runtime/Renderer/Private/VirtualShadowMaps/VirtualShadowMapCacheManager.cpp b/Engine/Source/Runtime/Renderer/Private/VirtualShadowMaps/VirtualShadowMapCacheManager.cpp
--- a/Engine/Source/Runtime/Renderer/Private/VirtualShadowMaps/VirtualShadowMapCacheManager.cpp
+++ b/Engine/Source/Runtime/Renderer/Private/VirtualShadowMaps/VirtualShadowMapCacheManager.cpp
@@ -1827,13 +1827,13 @@ void FVirtualShadowMapInvalidationSceneUpdater::PostGPUSceneUpdate(FRDGBuilder&
 	// This depends on GPU-Scene having correct instance ranges, so must happen post GPU scene update.
 	CacheManager.InstanceStateMaskWordStride = FMath::DivideAndRoundUp(CacheManager.Scene->GPUScene.GetNumInstances(), 32);
 	const int32 NumInstanceWordsNeeded = CacheManager.InstanceStateMaskWordStride * FVirtualShadowMapArrayCacheManager::FViewData::NumBitsPerInstance;
-	if (NumInstanceWordsNeeded > 0)
+	if (NumInstanceWordsNeeded > 0 && !CacheManager.ViewData.IsEmpty())
 	{
-		TArray<FRDGBuffer*, TInlineAllocator<16,  SceneRenderingAllocator> > InstanceStateRDG;
-		InstanceStateRDG.SetNumZeroed(CacheManager.ViewData.Num());
-		for (int32 Index = 0; Index < CacheManager.ViewData.Num(); ++Index)
+		TArray<FRDGBuffer*, TInlineAllocator<16,  SceneRenderingAllocator> > InstanceStatesRDG;
+		InstanceStatesRDG.Reserve(CacheManager.ViewData.Num());
+		for (FVirtualShadowMapArrayCacheManager::FViewData& ViewData : CacheManager.ViewData)
 		{
-			InstanceStateRDG[Index] = CacheManager.ViewData[Index].InstanceState.ResizeAndClearBufferIfNeeded(GraphBuilder, NumInstanceWordsNeeded);
+			InstanceStatesRDG.Emplace(ViewData.InstanceState.ResizeAndClearBufferIfNeeded(GraphBuilder, NumInstanceWordsNeeded));
 		}
 
 		FLoadBalancer LoadBalancer;
@@ -1851,10 +1851,10 @@ void FVirtualShadowMapInvalidationSceneUpdater::PostGPUSceneUpdate(FRDGBuilder&
 
 			FGPUSceneResourceParameters GPUSceneParams = CacheManager.Scene->GPUScene.GetShaderParameters(GraphBuilder);
 
-			for (int32 Index = 0; Index < CacheManager.ViewData.Num(); ++Index)
+			for (FRDGBuffer* InstanceStateRDG : InstanceStatesRDG)
 			{
 				FVSMResetInstanceStateCS::FParameters* PassParameters = GraphBuilder.AllocParameters<FVSMResetInstanceStateCS::FParameters>();
-				PassParameters->InOutViewInstanceState = GraphBuilder.CreateUAV(InstanceStateRDG[Index]);
+				PassParameters->InOutViewInstanceState = GraphBuilder.CreateUAV(InstanceStateRDG);
 				PassParameters->GPUScene = GPUSceneParams;
 				PassParameters->StateWordStride = CacheManager.InstanceStateMaskWordStride;
 				LoadBalancerGpuData.AddPass(GraphBuilder, RDG_EVENT_NAME("ResetInstanceState (%d batches)", LoadBalancerGpuData.NumBatches), ComputeShader, PassParameters);
