Removing unnecessary use of SceneRenderingAllocator. The concurrent linear allocator wasn't handling the allocation sizes correctly in large Editor scenes.

diff --git a/Engine/Source/Runtime/Renderer/Private/PrimitiveSceneInfo.cpp b/Engine/Source/Runtime/Renderer/Private/PrimitiveSceneInfo.cpp
--- a/Engine/Source/Runtime/Renderer/Private/PrimitiveSceneInfo.cpp
+++ b/Engine/Source/Runtime/Renderer/Private/PrimitiveSceneInfo.cpp
@@ -641,7 +641,7 @@ void FPrimitiveSceneInfo::CacheMeshDrawCommands(FScene* Scene, TArrayView<FPrimi
 	bool bAnyLooseParameterBuffers = false;
 	if (GMeshDrawCommandsCacheMultithreaded && FApp::ShouldUseThreadingForPerformance())
 	{
-		TArray<FCachedPassMeshDrawListContextDeferred, SceneRenderingAllocator> DrawListContexts;
+		TArray<FCachedPassMeshDrawListContextDeferred> DrawListContexts;
 		DrawListContexts.Reserve(NumBatches);
 		for(int32 ContextIndex = 0; ContextIndex < NumBatches; ++ContextIndex)
 		{
