Fix VT on Android/Vulkan

diff --git a/Engine/Source/Runtime/Core/Public/PixelFormat.h b/Engine/Source/Runtime/Core/Public/PixelFormat.h
--- a/Engine/Source/Runtime/Core/Public/PixelFormat.h
+++ b/Engine/Source/Runtime/Core/Public/PixelFormat.h
@@ -327,6 +327,11 @@ static inline bool IsBlockCompressedFormat(EPixelFormat Format)
 	return false;
 }
 
+FORCEINLINE bool IsAnyBlockCompressedPixelFormat(EPixelFormat PixelFormat)
+{
+	return IsBlockCompressedFormat(PixelFormat) || IsASTCBlockCompressedTextureFormat(PixelFormat) || IsETCBlockCompressedPixelFormat(PixelFormat);
+}
+
 // unclear what IsHDR is supposed to mean
 //	see also IsFloatFormat
 FORCEINLINE bool IsHDR(EPixelFormat PixelFormat)
diff --git a/Engine/Source/Runtime/RHI/Public/RHIUtilities.h b/Engine/Source/Runtime/RHI/Public/RHIUtilities.h
--- a/Engine/Source/Runtime/RHI/Public/RHIUtilities.h
+++ b/Engine/Source/Runtime/RHI/Public/RHIUtilities.h
@@ -20,6 +20,9 @@ static inline EPixelFormat GetBlockCompressedFormatUAVAliasFormat(EPixelFormat F
 	{
 	case PF_DXT1:
 	case PF_BC4:
+	case PF_ETC1:
+	case PF_ETC2_RGB:
+	case PF_ETC2_R11_EAC:
 		return PF_R32G32_UINT;
 
 	case PF_DXT3:
@@ -27,6 +30,8 @@ static inline EPixelFormat GetBlockCompressedFormatUAVAliasFormat(EPixelFormat F
 	case PF_BC5:
 	case PF_BC6H:
 	case PF_BC7:
+	case PF_ETC2_RGBA:
+	case PF_ETC2_RG11_EAC:
 		return PF_R32G32B32A32_UINT;
 	}
 
diff --git a/Engine/Source/Runtime/Renderer/Private/VT/VirtualTexturePhysicalSpace.cpp b/Engine/Source/Runtime/Renderer/Private/VT/VirtualTexturePhysicalSpace.cpp
--- a/Engine/Source/Runtime/Renderer/Private/VT/VirtualTexturePhysicalSpace.cpp
+++ b/Engine/Source/Runtime/Renderer/Private/VT/VirtualTexturePhysicalSpace.cpp
@@ -116,7 +116,7 @@ static EPixelFormat GetUnorderedAccessViewFormat(EPixelFormat InFormat)
 {
 	// Use alias formats for compressed textures on APIs where that is possible
 	// This allows us to compress runtime data directly to the physical texture
-	if (IsBlockCompressedFormat(InFormat))
+	if (IsAnyBlockCompressedPixelFormat(InFormat))
 	{
 		return GRHISupportsUAVFormatAliasing ? GetBlockCompressedFormatUAVAliasFormat(InFormat) : PF_Unknown;
 	}
diff --git a/Engine/Source/Runtime/VulkanRHI/Private/VulkanTexture.cpp b/Engine/Source/Runtime/VulkanRHI/Private/VulkanTexture.cpp
--- a/Engine/Source/Runtime/VulkanRHI/Private/VulkanTexture.cpp
+++ b/Engine/Source/Runtime/VulkanRHI/Private/VulkanTexture.cpp
@@ -309,7 +309,7 @@ void FVulkanTexture::GenerateImageCreateInfo(
 		}
 
 		ImageCreateInfo.flags |= VK_IMAGE_CREATE_MUTABLE_FORMAT_BIT;
-		if (bHasUAVFormat && IsBlockCompressedFormat(InDesc.Format) && !IsBlockCompressedFormat(InDesc.UAVFormat))
+		if (bHasUAVFormat && IsAnyBlockCompressedPixelFormat(InDesc.Format) && !IsAnyBlockCompressedPixelFormat(InDesc.UAVFormat))
 		{
 			ImageCreateInfo.flags |= VK_IMAGE_CREATE_BLOCK_TEXEL_VIEW_COMPATIBLE_BIT;
 		}
