Avoid making AO render targets a UAV when it's not needed.

Enable lossy framebuffer compression for AO and distortion render targets

diff --git a/Engine/Source/Runtime/Renderer/Private/MobileDistortionPass.cpp b/Engine/Source/Runtime/Renderer/Private/MobileDistortionPass.cpp
--- a/Engine/Source/Runtime/Renderer/Private/MobileDistortionPass.cpp
+++ b/Engine/Source/Runtime/Renderer/Private/MobileDistortionPass.cpp
@@ -68,7 +68,14 @@ TRDGUniformBufferRef<FMobileDistortionPassUniformParameters> CreateMobileDistort
 
 FMobileDistortionAccumulateOutputs AddMobileDistortionAccumulatePass(FRDGBuilder& GraphBuilder, FScene* Scene, const FViewInfo& View, const FMobileDistortionAccumulateInputs& Inputs)
 {
-	FRDGTextureDesc DistortionAccumulateDesc = FRDGTextureDesc::Create2D(Inputs.SceneColor.Texture->Desc.Extent, PF_B8G8R8A8, FClearValueBinding::Transparent, TexCreate_ShaderResource | TexCreate_RenderTargetable);
+	const EPixelFormat Format = PF_B8G8R8A8;
+	ETextureCreateFlags TextureCreateFlags = TexCreate_ShaderResource | TexCreate_RenderTargetable;
+	if (UE::PixelFormat::HasCapabilities(Format, EPixelFormatCapabilities::LossyCompressible))
+	{
+		TextureCreateFlags |= TexCreate_LossyCompressionLowBitrate;
+	}
+
+	FRDGTextureDesc DistortionAccumulateDesc = FRDGTextureDesc::Create2D(Inputs.SceneColor.Texture->Desc.Extent, Format, FClearValueBinding::Transparent, TextureCreateFlags);
 
 	FScreenPassRenderTarget DistortionAccumulateOutput = FScreenPassRenderTarget(GraphBuilder.CreateTexture(DistortionAccumulateDesc, TEXT("DistortionAccumulatePass")), Inputs.SceneColor.ViewRect, ERenderTargetLoadAction::EClear);
 
diff --git a/Engine/Source/Runtime/Renderer/Private/PostProcess/PostProcessAmbientOcclusionMobile.cpp b/Engine/Source/Runtime/Renderer/Private/PostProcess/PostProcessAmbientOcclusionMobile.cpp
--- a/Engine/Source/Runtime/Renderer/Private/PostProcess/PostProcessAmbientOcclusionMobile.cpp
+++ b/Engine/Source/Runtime/Renderer/Private/PostProcess/PostProcessAmbientOcclusionMobile.cpp
@@ -60,11 +60,46 @@ static TAutoConsoleVariable<int32> CVarMobileAmbientOcclusionShaderType(
 	TEXT("r.Mobile.AmbientOcclusionShaderType"),
 	2,
 	TEXT("0: ComputeShader.\n")
-	TEXT("1: Seperate ComputeShader.\n")
+	TEXT("1: Separate ComputeShader.\n")
 	TEXT("2: PixelShader.\n"),
 	ECVF_RenderThreadSafe
 );
 
+static TAutoConsoleVariable<int32> CVarMobileAmbientOcclusionCompression(
+	TEXT("r.Mobile.AmbientOcclusionCompression"),
+	1,
+	TEXT("0: Disable framebuffer compression for AO texture.\n")
+	TEXT("1: Enable framebuffer compression for AO texture if supported (Default).\n"),
+	ECVF_RenderThreadSafe
+);
+
+enum class EMobileAmbientOcclusionShaderType
+{
+	ComputeShader,
+	SeparateComputeShader,
+	PixelShader
+};
+
+static EMobileAmbientOcclusionShaderType GetCVarMobileAmbientOcclusionShaderTypeOnRenderThread()
+{
+	int ShaderType = FMath::Clamp(CVarMobileAmbientOcclusionShaderType.GetValueOnRenderThread(), 0, 2);
+	if (ShaderType == 0 && GetMaxWorkGroupInvocations() < 1024)
+	{
+		ShaderType = 1;
+	}
+	return static_cast<EMobileAmbientOcclusionShaderType>(ShaderType);
+}
+
+static bool IsComputeShaderType(EMobileAmbientOcclusionShaderType ShaderType)
+{
+	return ShaderType != EMobileAmbientOcclusionShaderType::PixelShader;
+}
+
+static bool IsTwoPassShaderType(EMobileAmbientOcclusionShaderType ShaderType)
+{
+	return ShaderType != EMobileAmbientOcclusionShaderType::ComputeShader;
+}
+
 static TAutoConsoleVariable<int32> CVarMobileAmbientOcclusionDepthBoundsTest(
 	TEXT("r.Mobile.AmbientOcclusionDepthBoundsTest"),
 	1,
@@ -445,18 +480,29 @@ FRDGTextureRef CreateMobileScreenSpaceAOTexture(FRDGBuilder& GraphBuilder, const
 
 	EPixelFormat Format = PF_G8;
 
+	const EMobileAmbientOcclusionShaderType ShaderType = GetCVarMobileAmbientOcclusionShaderTypeOnRenderThread();
+
 	// G8 isn't supported as UAV on Android OpenGLES, fall back to RGBA8 for compute shader usage.
 	if (bGTAO
 		&& IsOpenGLPlatform(Config.ShaderPlatform)
-		&& ((GetMaxWorkGroupInvocations() >= 1024 && CVarMobileAmbientOcclusionShaderType.GetValueOnRenderThread() == 0) || CVarMobileAmbientOcclusionShaderType.GetValueOnRenderThread() == 1)
+		&& IsComputeShaderType(ShaderType)
 		)
 	{
 		Format = PF_R8G8B8A8;
 	}
 
-	return GraphBuilder.CreateTexture(
-		FRDGTextureDesc::Create2D(Extent, Format, FClearValueBinding::Black, TexCreate_ShaderResource | TexCreate_RenderTargetable | TexCreate_UAV),
-		TEXT("ScreenSpaceAO"));
+	ETextureCreateFlags TextureCreateFlags = TexCreate_ShaderResource | TexCreate_RenderTargetable;
+	if (IsComputeShaderType(ShaderType))
+	{
+		TextureCreateFlags |= TexCreate_UAV;
+	}
+
+	if (UE::PixelFormat::HasCapabilities(Format, EPixelFormatCapabilities::LossyCompressible) && CVarMobileAmbientOcclusionCompression.GetValueOnRenderThread() != 0)
+	{
+		TextureCreateFlags |= TexCreate_LossyCompressionLowBitrate;
+	}
+
+	return GraphBuilder.CreateTexture(FRDGTextureDesc::Create2D(Extent, Format, FClearValueBinding::Black, TextureCreateFlags), TEXT("ScreenSpaceAO"));
 }
 
 static void RenderGTAO(FRDGBuilder& GraphBuilder, FRDGTextureRef SceneDepthTexture, FRDGTextureRef AmbientOcclusionTexture, const TArray<FViewInfo>& Views)
@@ -470,8 +516,6 @@ static void RenderGTAO(FRDGBuilder& GraphBuilder, FRDGTextureRef SceneDepthTextu
 	const int32 MobileGTAOPreIntegratedTextureType = FMath::Min(CVarMobileGTAOPreIntegratedTextureType.GetValueOnRenderThread(), 2);
 	const int32 MobileAmbientOcclusionQuality = FMath::Min(CVarMobileAmbientOcclusionQuality.GetValueOnRenderThread(), 3);
 
-	FRDGTextureUAVRef AmbientOcclusionTextureUAV = GraphBuilder.CreateUAV(AmbientOcclusionTexture);
-
 	const FIntPoint& DepthBufferSize = SceneDepthTexture->Desc.Extent;
 	const FIntPoint& BufferSize = AmbientOcclusionTexture->Desc.Extent;
 
@@ -493,9 +537,31 @@ static void RenderGTAO(FRDGBuilder& GraphBuilder, FRDGTextureRef SceneDepthTextu
 	float SinDeltaAngle, CosDeltaAngle;
 	FMath::SinCos(&SinDeltaAngle, &CosDeltaAngle, PI / NumAngles);
 
-	ETextureCreateFlags TextureCreateFlags = TexCreate_ShaderResource | TexCreate_RenderTargetable | TexCreate_UAV;
-	FRDGTextureRef HorizonSearchIntegralTexture = GraphBuilder.CreateTexture(FRDGTextureDesc::Create2D(BufferSize, PF_R8G8B8A8, FClearValueBinding::Black, TextureCreateFlags), TEXT("HorizonSearchIntegralTexture"));
-	FRDGTextureUAVRef HorizonSearchIntegralTextureUAV = GraphBuilder.CreateUAV(HorizonSearchIntegralTexture);
+	const EMobileAmbientOcclusionShaderType ShaderType = GetCVarMobileAmbientOcclusionShaderTypeOnRenderThread();
+
+	FRDGTextureRef HorizonSearchIntegralTexture{};
+	FRDGTextureUAVRef HorizonSearchIntegralTextureUAV{};
+	FRDGTextureUAVRef AmbientOcclusionTextureUAV{};
+
+	if (IsComputeShaderType(ShaderType))
+	{
+		AmbientOcclusionTextureUAV = GraphBuilder.CreateUAV(AmbientOcclusionTexture);
+	}
+	if (IsTwoPassShaderType(ShaderType))
+	{
+		ETextureCreateFlags TextureCreateFlags = TexCreate_ShaderResource | TexCreate_RenderTargetable;
+		const EPixelFormat PixelFormat = PF_R8G8B8A8;
+		if (UE::PixelFormat::HasCapabilities(PixelFormat, EPixelFormatCapabilities::LossyCompressible) && CVarMobileAmbientOcclusionCompression.GetValueOnRenderThread() != 0)
+		{
+			TextureCreateFlags |= TexCreate_LossyCompressionLowBitrate;
+		}
+
+		HorizonSearchIntegralTexture = GraphBuilder.CreateTexture(FRDGTextureDesc::Create2D(BufferSize, PixelFormat, FClearValueBinding::Black, TextureCreateFlags | (IsComputeShaderType(ShaderType) ? TexCreate_UAV : TexCreate_None)), TEXT("HorizonSearchIntegralTexture"));
+		if (IsComputeShaderType(ShaderType))
+		{
+			HorizonSearchIntegralTextureUAV = GraphBuilder.CreateUAV(HorizonSearchIntegralTexture);
+		}
+	}
 
 	for (int32 ViewIndex = 0; ViewIndex < Views.Num(); ViewIndex++)
 	{
@@ -507,7 +573,7 @@ static void RenderGTAO(FRDGBuilder& GraphBuilder, FRDGTextureRef SceneDepthTextu
 
 		FVector4f WorldRadiusAdjSinCosDeltaAngleThickness(FallOffStartEndScaleBias.Y * DepthBufferSize.Y * View.ViewMatrices.GetProjectionMatrix().M[0][0], SinDeltaAngle, CosDeltaAngle, ThicknessBlend);
 
-		if (GetMaxWorkGroupInvocations() >= 1024 && CVarMobileAmbientOcclusionShaderType.GetValueOnRenderThread() == 0)
+		if (ShaderType == EMobileAmbientOcclusionShaderType::ComputeShader)
 		{
 			FGTAOMobile_HorizonSearchIntegralSpatialFilterCS::FParameters* HorizonSearchIntegralSpatialFilterParameters = GraphBuilder.AllocParameters<FGTAOMobile_HorizonSearchIntegralSpatialFilterCS::FParameters>();
 			FGTAOMobile_HorizonSearchIntegral::SetupShaderParameters(HorizonSearchIntegralSpatialFilterParameters->Common, GraphBuilder, View, ViewRect, DepthBufferSize, BufferSize, FallOffStartEndScaleBias, WorldRadiusAdjSinCosDeltaAngleThickness, SceneDepthTexture);
@@ -527,7 +593,7 @@ static void RenderGTAO(FRDGBuilder& GraphBuilder, FRDGTextureRef SceneDepthTextu
 				HorizonSearchIntegralSpatialFilterParameters,
 				FComputeShaderUtils::GetGroupCount(ViewRect.Size(), FGTAOMobile_HorizonSearchIntegralSpatialFilterCS::TexelsPerThreadGroup));
 		}
-		else if (CVarMobileAmbientOcclusionShaderType.GetValueOnRenderThread() != 1)
+		else if (ShaderType == EMobileAmbientOcclusionShaderType::PixelShader)
 		{
 			TShaderMapRef<FScreenPassVS> VertexShader(View.ShaderMap);
 
