Merging CL#38156889

Scalability changes for Heterogeneous Volumes

diff --git a/Engine/Source/Runtime/Renderer/Private/HeterogeneousVolumes/HeterogeneousVolumesLiveShadingPipeline.cpp b/Engine/Source/Runtime/Renderer/Private/HeterogeneousVolumes/HeterogeneousVolumesLiveShadingPipeline.cpp
--- a/Engine/Source/Runtime/Renderer/Private/HeterogeneousVolumes/HeterogeneousVolumesLiveShadingPipeline.cpp
+++ b/Engine/Source/Runtime/Renderer/Private/HeterogeneousVolumes/HeterogeneousVolumesLiveShadingPipeline.cpp
@@ -22,28 +22,28 @@ static TAutoConsoleVariable<bool> CVarHeterogeneousLightTypeDirectional(
 	TEXT("r.HeterogeneousVolumes.LightType.Directional"),
 	1,
 	TEXT("Enables illumination from the directional light (Default = 1)"),
-	ECVF_RenderThreadSafe
+	ECVF_Scalability | ECVF_RenderThreadSafe
 );
 
 static TAutoConsoleVariable<bool> CVarHeterogeneousLightTypePoint(
 	TEXT("r.HeterogeneousVolumes.LightType.Point"),
 	1,
 	TEXT("Enables illumination from point lights (Default = 1)"),
-	ECVF_RenderThreadSafe
+	ECVF_Scalability | ECVF_RenderThreadSafe
 );
 
 static TAutoConsoleVariable<bool> CVarHeterogeneousLightTypeSpot(
 	TEXT("r.HeterogeneousVolumes.LightType.Spot"),
 	1,
 	TEXT("Enables illumination from spot lights (Default = 1)"),
-	ECVF_RenderThreadSafe
+	ECVF_Scalability | ECVF_RenderThreadSafe
 );
 
 static TAutoConsoleVariable<bool> CVarHeterogeneousLightTypeRect(
 	TEXT("r.HeterogeneousVolumes.LightType.Rect"),
 	1,
 	TEXT("Enables illumination from rect lights (Default = 1)"),
-	ECVF_RenderThreadSafe
+	ECVF_Scalability | ECVF_RenderThreadSafe
 );
 
 static TAutoConsoleVariable<int32> CVarHeterogeneousLightingCacheBoundsCulling(
