Add RDG tracking for RDGBuffers used by raytracing dynamic instances which are used in the SBT scene:
- FRayTracingInstanceCollector collects all the RDG pooled buffers from the dynamic RT instances
- RDG_BUFFER_ACCESS_ARRAY can be used in BEGIN_SHADER_PARAMETER_STRUCT of the pass performing the RayTraceDispatch using these buffers (currently only fixed up for path tracer init function - only known problematic use case)
- Groom RDP pooled buffers are added during dynamic RT instance collection

diff --git a/Engine/Plugins/Runtime/HairStrands/Source/HairStrandsCore/Private/GroomComponent.cpp b/Engine/Plugins/Runtime/HairStrands/Source/HairStrandsCore/Private/GroomComponent.cpp
--- a/Engine/Plugins/Runtime/HairStrands/Source/HairStrandsCore/Private/GroomComponent.cpp
+++ b/Engine/Plugins/Runtime/HairStrands/Source/HairStrandsCore/Private/GroomComponent.cpp
@@ -722,6 +722,14 @@ public:
 					RayTracingInstance.bThinGeometry = bIsHairStrands;
 
 					Collector.AddRayTracingInstance(MoveTemp(RayTracingInstance));
+					if (RTGeometry->PositionBuffer.Buffer)
+					{
+						Collector.AddRDGPooledBuffer(RTGeometry->PositionBuffer.Buffer);
+					}
+					if (RTGeometry->IndexBuffer.Buffer)
+					{
+						Collector.AddRDGPooledBuffer(RTGeometry->IndexBuffer.Buffer);
+					}
 				}
 			}
 		}
diff --git a/Engine/Source/Runtime/Engine/Public/SceneManagement.h b/Engine/Source/Runtime/Engine/Public/SceneManagement.h
--- a/Engine/Source/Runtime/Engine/Public/SceneManagement.h
+++ b/Engine/Source/Runtime/Engine/Public/SceneManagement.h
@@ -1581,6 +1581,11 @@ public:
 
 	ENGINE_API void AddRayTracingGeometryUpdate(FRayTracingDynamicGeometryUpdateParams Params);
 
+	void AddRDGPooledBuffer(FRDGPooledBuffer* PooledBuffer)
+	{
+		RDGPooledBuffers.Add(PooledBuffer);
+	}
+
 private:
 
 	ENGINE_API FRayTracingInstanceCollector(
@@ -1592,6 +1597,7 @@ private:
 	const FSceneView* ReferenceView;
 	TArray<FRayTracingInstance> RayTracingInstances;
 	TArray<FRayTracingDynamicGeometryUpdateParams> RayTracingGeometriesToUpdate;
+	TSet<FRDGPooledBuffer*> RDGPooledBuffers;
 	TSet<RayTracing::GeometryGroupHandle> ReferencedGeometryGroups;
 
 	const bool bTrackReferencedGeometryGroups;
diff --git a/Engine/Source/Runtime/Renderer/Private/PathTracing.cpp b/Engine/Source/Runtime/Renderer/Private/PathTracing.cpp
--- a/Engine/Source/Runtime/Renderer/Private/PathTracing.cpp
+++ b/Engine/Source/Runtime/Renderer/Private/PathTracing.cpp
@@ -1356,6 +1356,7 @@ class FPathTracingInitExtinctionCoefficientRG : public FGlobalShader
 		SHADER_PARAMETER_RDG_BUFFER_SRV(RaytracingAccelerationStructure, TLAS)
 		SHADER_PARAMETER_RDG_UNIFORM_BUFFER(FSceneUniformParameters, Scene)
 		SHADER_PARAMETER_STRUCT_REF(FViewUniformShaderParameters, ViewUniformBuffer)
+		RDG_BUFFER_ACCESS_ARRAY(SBTBuffers)
 		SHADER_PARAMETER_RDG_BUFFER_UAV(RWBuffer<float>, RWStartingExtinctionCoefficient)
 	END_SHADER_PARAMETER_STRUCT()
 };
@@ -3319,6 +3320,11 @@ void FDeferredShadingSceneRenderer::RenderPathTracing(
 		PassParameters->ViewUniformBuffer = View.ViewUniformBuffer;
 		PassParameters->RWStartingExtinctionCoefficient = GraphBuilder.CreateUAV(StartingExtinctionCoefficient, PF_R32_FLOAT);
 
+		for (FRDGBuffer* RDGBuffer : View.DynamicRayTracingRDGBuffers)
+		{
+			PassParameters->SBTBuffers.Emplace(RDGBuffer, ERHIAccess::SRVCompute);
+		}
+
 		GraphBuilder.AddPass(
 			RDG_EVENT_NAME("Path Tracer Init Sigma"),
 			PassParameters,
diff --git a/Engine/Source/Runtime/Renderer/Private/RayTracing/RayTracing.cpp b/Engine/Source/Runtime/Renderer/Private/RayTracing/RayTracing.cpp
--- a/Engine/Source/Runtime/Renderer/Private/RayTracing/RayTracing.cpp
+++ b/Engine/Source/Runtime/Renderer/Private/RayTracing/RayTracing.cpp
@@ -481,6 +481,8 @@ namespace RayTracing
 
 		void AddInstancesToScene(FRayTracingScene& RayTracingScene, FRayTracingShaderBindingTable& RayTracingSBT, int64 SharedBufferGenerationID);
 
+		void CollectRDGResources(FRDGBuilder& RDGBuilder);
+
 	private:
 
 		FScene& Scene;
@@ -613,6 +615,15 @@ namespace RayTracing
 		InRHICmdList.QueueAsyncCommandListSubmit(RHICmdList);
 	}
 
+	void FDynamicRayTracingInstancesContext::CollectRDGResources(FRDGBuilder& RDGBuilder)
+	{
+		for (FRDGPooledBuffer* PooledBuffer : RayTracingInstanceCollector.RDGPooledBuffers)
+		{
+			FRDGBufferRef RDGBuffer = RDGBuilder.RegisterExternalBuffer(PooledBuffer);
+			View.DynamicRayTracingRDGBuffers.Add(RDGBuffer);
+		}
+	}
+
 	void FDynamicRayTracingInstancesContext::AddInstancesToScene(FRayTracingScene& RayTracingScene, FRayTracingShaderBindingTable& RayTracingSBT, int64 SharedBufferGenerationID)
 	{
 		TRACE_CPUPROFILER_EVENT_SCOPE(RayTracing_AddInstancesToScene);
@@ -1649,6 +1660,7 @@ namespace RayTracing
 		Scene.GPUScene.UploadDynamicPrimitiveShaderDataForView(GraphBuilder, View, /*bRayTracing*/ true);
 
 		DynamicRayTracingInstancesContext.AddInstancesToScene(RayTracingScene, RayTracingSBT, SharedBufferGenerationID);
+		DynamicRayTracingInstancesContext.CollectRDGResources(GraphBuilder);
 	}
 
 	void BeginGatherInstances(FGatherInstancesTaskData& TaskData, UE::Tasks::FTask FrustumCullTask)
diff --git a/Engine/Source/Runtime/Renderer/Private/SceneRendering.h b/Engine/Source/Runtime/Renderer/Private/SceneRendering.h
--- a/Engine/Source/Runtime/Renderer/Private/SceneRendering.h
+++ b/Engine/Source/Runtime/Renderer/Private/SceneRendering.h
@@ -1289,6 +1289,7 @@ public:
 #if RHI_RAYTRACING 
 	FRayTracingShaderBindingDataOneFrameArray VisibleRayTracingShaderBindings;
 	FDynamicRayTracingMeshCommandStorage DynamicRayTracingMeshCommandStorage;
+	TSet<FRDGBuffer*> DynamicRayTracingRDGBuffers;
 
 	FRayTracingCullingParameters RayTracingCullingParameters;
 
