[Light Grid] Replaced "SortLights" stats with a more specific "PrepareForwardLightData"/"LightGrid"

diff --git a/Engine/Source/Runtime/Renderer/Private/DeferredShadingRenderer.cpp b/Engine/Source/Runtime/Renderer/Private/DeferredShadingRenderer.cpp
--- a/Engine/Source/Runtime/Renderer/Private/DeferredShadingRenderer.cpp
+++ b/Engine/Source/Runtime/Renderer/Private/DeferredShadingRenderer.cpp
@@ -273,6 +273,7 @@ DECLARE_CYCLE_STAT(TEXT("DeferredShadingSceneRenderer UpdateDownsampledDepthSurf
 DECLARE_CYCLE_STAT(TEXT("DeferredShadingSceneRenderer Render Init"), STAT_FDeferredShadingSceneRenderer_Render_Init, STATGROUP_SceneRendering);
 DECLARE_CYCLE_STAT(TEXT("DeferredShadingSceneRenderer FXSystem PreRender"), STAT_FDeferredShadingSceneRenderer_FXSystem_PreRender, STATGROUP_SceneRendering);
 DECLARE_CYCLE_STAT(TEXT("DeferredShadingSceneRenderer AllocGBufferTargets"), STAT_FDeferredShadingSceneRenderer_AllocGBufferTargets, STATGROUP_SceneRendering);
+DECLARE_CYCLE_STAT(TEXT("DeferredShadingSceneRenderer PrepareForwardLightData"), STAT_FDeferredShadingSceneRenderer_PrepareForwardLightData, STATGROUP_SceneRendering);
 DECLARE_CYCLE_STAT(TEXT("DeferredShadingSceneRenderer DBuffer"), STAT_FDeferredShadingSceneRenderer_DBuffer, STATGROUP_SceneRendering);
 DECLARE_CYCLE_STAT(TEXT("DeferredShadingSceneRenderer ResolveDepth After Basepass"), STAT_FDeferredShadingSceneRenderer_ResolveDepth_After_Basepass, STATGROUP_SceneRendering);
 DECLARE_CYCLE_STAT(TEXT("DeferredShadingSceneRenderer Resolve After Basepass"), STAT_FDeferredShadingSceneRenderer_Resolve_After_Basepass, STATGROUP_SceneRendering);
@@ -295,7 +296,6 @@ DECLARE_GPU_STAT(VisibilityCommands);
 DECLARE_GPU_STAT(RenderDeferredLighting);
 DECLARE_GPU_STAT(AllocateRendertargets);
 DECLARE_GPU_STAT(FrameRenderFinish);
-DECLARE_GPU_STAT(SortLights);
 DECLARE_GPU_STAT(PostRenderOpsFX);
 DECLARE_GPU_STAT_NAMED(Unaccounted, TEXT("[unaccounted]"));
 DECLARE_GPU_STAT(WaterRendering);
@@ -397,6 +397,8 @@ FDeferredShadingSceneRenderer::FDeferredShadingSceneRenderer(const FSceneViewFam
 DECLARE_CYCLE_STAT(TEXT("Wait RayTracing Add Mesh Batch"), STAT_WaitRayTracingAddMesh, STATGROUP_SceneRendering);
 DECLARE_CYCLE_STAT(TEXT("Wait Ray Tracing Scene Initialization"), STAT_WaitRayTracingSceneInitTask, STATGROUP_SceneRendering);
 
+DECLARE_CYCLE_STAT(TEXT("Wait Gather And Sort Lights"), STAT_WaitGatherAndSortLightsTask, STATGROUP_SceneRendering);
+
 /**
  * Returns true if the depth Prepass needs to run
  */
@@ -2268,13 +2270,15 @@ void FDeferredShadingSceneRenderer::Render(FRDGBuilder& GraphBuilder)
 		}
 
 		{
-			const FSortedLightSetSceneInfo* SortedLightSet = GatherAndSortLightsTask.GetResult();
+			SCOPE_CYCLE_COUNTER(STAT_WaitGatherAndSortLightsTask);
+			GatherAndSortLightsTask.Wait();
+		}
 
-			RDG_CSV_STAT_EXCLUSIVE_SCOPE(GraphBuilder, SortLights);
-			RDG_EVENT_SCOPE_STAT(GraphBuilder, SortLights, "SortLights");
-			RDG_GPU_STAT_SCOPE(GraphBuilder, SortLights);
+		{
+			RDG_CSV_STAT_EXCLUSIVE_SCOPE(GraphBuilder, PrepareForwardLightData);
+			SCOPE_CYCLE_COUNTER(STAT_FDeferredShadingSceneRenderer_PrepareForwardLightData);
 
-			GatherAndSortLightsTask.Wait();
+			const FSortedLightSetSceneInfo* SortedLightSet = GatherAndSortLightsTask.GetResult();
 
 			if (!ViewFamily.EngineShowFlags.PathTracing)
 			{
diff --git a/Engine/Source/Runtime/Renderer/Private/LightGridInjection.cpp b/Engine/Source/Runtime/Renderer/Private/LightGridInjection.cpp
--- a/Engine/Source/Runtime/Renderer/Private/LightGridInjection.cpp
+++ b/Engine/Source/Runtime/Renderer/Private/LightGridInjection.cpp
@@ -159,6 +159,8 @@ FAutoConsoleVariableRef CVarLightCullingMaxDistanceOverride(
 	ECVF_RenderThreadSafe
 );
 
+DECLARE_GPU_STAT(LightGrid);
+
 bool ShouldVisualizeLightGrid()
 {
 	return GForwardLightGridDebug > 0;
@@ -724,9 +726,10 @@ FComputeLightGridOutput FSceneRenderer::ComputeLightGrid(FRDGBuilder& GraphBuild
 {
 	FComputeLightGridOutput Result = {};
 
-	RDG_CSV_STAT_EXCLUSIVE_SCOPE(GraphBuilder, ComputeLightGrid);
 	QUICK_SCOPE_CYCLE_COUNTER(STAT_ComputeLightGrid);
-	RDG_EVENT_SCOPE(GraphBuilder, "ComputeLightGrid");
+	RDG_CSV_STAT_EXCLUSIVE_SCOPE(GraphBuilder, ComputeLightGrid);
+	RDG_EVENT_SCOPE_STAT(GraphBuilder, LightGrid, "ComputeLightGrid");
+	RDG_GPU_STAT_SCOPE(GraphBuilder, LightGrid);
 
 	const bool bAllowStaticLighting = IsStaticLightingAllowed();
 	const bool bLightGridUses16BitBuffers = LightGridUses16BitBuffers(ShaderPlatform);
