Fix editor crash when using r.Bloom.ApplyLocalExposure=0

diff --git a/Engine/Source/Runtime/Renderer/Private/PostProcess/PostProcessBloomSetup.cpp b/Engine/Source/Runtime/Renderer/Private/PostProcess/PostProcessBloomSetup.cpp
--- a/Engine/Source/Runtime/Renderer/Private/PostProcess/PostProcessBloomSetup.cpp
+++ b/Engine/Source/Runtime/Renderer/Private/PostProcess/PostProcessBloomSetup.cpp
@@ -55,7 +55,7 @@ FBloomSetupParameters GetBloomSetupParameters(
 	Parameters.LumBilateralGridSampler = TStaticSamplerState<SF_Bilinear, AM_Clamp, AM_Clamp, AM_Clamp>::GetRHI();
 	Parameters.BlurredLogLum = Inputs.BlurredLogLuminanceTexture;
 	Parameters.BlurredLogLumSampler = TStaticSamplerState<SF_Bilinear, AM_Clamp, AM_Clamp, AM_Clamp>::GetRHI();
-	Parameters.LocalExposure = *Inputs.LocalExposureParameters;
+	Parameters.LocalExposure = Inputs.LocalExposureParameters != nullptr ? *Inputs.LocalExposureParameters : FLocalExposureParameters();
 	Parameters.EyeAdaptationBuffer = GraphBuilder.CreateSRV(Inputs.EyeAdaptationBuffer);
 	Parameters.EyeAdaptation = *Inputs.EyeAdaptationParameters;
 	Parameters.BloomThreshold = Inputs.Threshold;
diff --git a/Engine/Source/Runtime/Renderer/Private/PostProcess/PostProcessing.cpp b/Engine/Source/Runtime/Renderer/Private/PostProcess/PostProcessing.cpp
--- a/Engine/Source/Runtime/Renderer/Private/PostProcess/PostProcessing.cpp
+++ b/Engine/Source/Runtime/Renderer/Private/PostProcess/PostProcessing.cpp
@@ -1304,7 +1304,12 @@ void AddPostProcessingPasses(
 			}
 			else
 			{
-				const bool bBloomSetupRequiredEnabled = View.FinalPostProcessSettings.BloomThreshold > -1.0f || (LocalExposureBilateralGridTexture != nullptr && View.FinalPostProcessSettings.LocalExposureMethod == ELocalExposureMethod::Bilateral);
+				const bool bApplyLocalExposureToBloom = 
+					CVarBloomApplyLocalExposure.GetValueOnRenderThread()
+					&& View.FinalPostProcessSettings.LocalExposureMethod == ELocalExposureMethod::Bilateral
+					&& LocalExposureBilateralGridTexture != nullptr;
+
+				const bool bBloomSetupRequiredEnabled = View.FinalPostProcessSettings.BloomThreshold > -1.0f || bApplyLocalExposureToBloom;
 
 				// Reuse the main scene downsample chain if setup isn't required for gaussian bloom.
 				if (SceneDownsampleChain.IsInitialized() && !bBloomSetupRequiredEnabled)
@@ -1323,11 +1328,15 @@ void AddPostProcessingPasses(
 						SetupPassInputs.SceneColor = DownsampleInput;
 						SetupPassInputs.EyeAdaptationBuffer = EyeAdaptationBuffer;
 						SetupPassInputs.EyeAdaptationParameters = &EyeAdaptationParameters;
-						SetupPassInputs.LocalExposureParameters = &LocalExposureParameters;
-						SetupPassInputs.LocalExposureTexture = CVarBloomApplyLocalExposure.GetValueOnRenderThread() && View.FinalPostProcessSettings.LocalExposureMethod == ELocalExposureMethod::Bilateral ? LocalExposureBilateralGridTexture : nullptr;
-						SetupPassInputs.BlurredLogLuminanceTexture = LocalExposureBlurredLogLumTexture;
 						SetupPassInputs.Threshold = BloomThreshold;
 
+						if (bApplyLocalExposureToBloom)
+						{
+							SetupPassInputs.LocalExposureParameters = &LocalExposureParameters;
+							SetupPassInputs.LocalExposureTexture = LocalExposureBilateralGridTexture;
+							SetupPassInputs.BlurredLogLuminanceTexture = LocalExposureBlurredLogLumTexture;
+						}
+
 						DownsampleInput = FScreenPassTextureSlice::CreateFromScreenPassTexture(GraphBuilder, AddBloomSetupPass(GraphBuilder, View, SetupPassInputs));
 					}
 
