Added option to disable CSM into translucent lighting volume.

#tests PC build with PGO
#rnx

diff --git a/Engine/Source/Runtime/Renderer/Private/TranslucentLighting.cpp b/Engine/Source/Runtime/Renderer/Private/TranslucentLighting.cpp
--- a/Engine/Source/Runtime/Renderer/Private/TranslucentLighting.cpp
+++ b/Engine/Source/Runtime/Renderer/Private/TranslucentLighting.cpp
@@ -129,6 +129,12 @@ static TAutoConsoleVariable<int32> CVarTranslucencyLightingVolumeAccurateRectLig
 	TEXT("Only accurate rect lights are included in batching, so it is recommended to enable this when batching is used."),
 	ECVF_RenderThreadSafe | ECVF_Scalability);
 
+static TAutoConsoleVariable<int32> CVarTranslucencyLightingVolumeInjectDirectionalLightCSM(
+	TEXT("r.TranslucencyLightingVolume.InjectDirectionalLightCSM"),
+	1,
+	TEXT("Enable sampling of the directional light CSM.\n"),
+	ECVF_RenderThreadSafe | ECVF_Scalability);
+
 /** Function returning current translucency lighting volume dimensions. */
 int32 GetTranslucencyLightingVolumeDim()
 {
@@ -1484,11 +1490,19 @@ void InjectTranslucencyLightingVolume(
 					const FMaterial& Material = MaterialProxy->GetMaterialWithFallback( View.GetFeatureLevel(), MaterialProxy );
 					const FMaterialShaderMap* MaterialShaderMap = Material.GetRenderingThreadShaderMap();
 
+					const bool bUseVSM = InjectionData.VirtualShadowMapId != INDEX_NONE;
+					bool bDynamicShadow = InjectionData.ProjectedShadowInfo != nullptr;
+
+					if (CVarTranslucencyLightingVolumeInjectDirectionalLightCSM.GetValueOnRenderThread() <= 0 && bDirectionalLight)
+					{
+						bDynamicShadow = false;
+					}
+
 					FTranslucentLightingInjectPS::FPermutationDomain PermutationVector;
 					PermutationVector.Set< FTranslucentLightingInjectPS::FRadialAttenuation >(!bDirectionalLight);
-					PermutationVector.Set< FTranslucentLightingInjectPS::FDynamicallyShadowed >(InjectionData.ProjectedShadowInfo != nullptr);
+					PermutationVector.Set< FTranslucentLightingInjectPS::FDynamicallyShadowed >(bDynamicShadow);
 					PermutationVector.Set< FTranslucentLightingInjectPS::FLightFunction >(InjectionData.bApplyLightFunction);
-					PermutationVector.Set< FTranslucentLightingInjectPS::FVirtualShadowMap >(InjectionData.VirtualShadowMapId != INDEX_NONE);
+					PermutationVector.Set< FTranslucentLightingInjectPS::FVirtualShadowMap >(bUseVSM);
 					PermutationVector.Set< FTranslucentLightingInjectPS::FAdaptiveVolumetricShadowMap >(InjectionData.bUseAdaptiveVolumetricShadowMap);
 
 					auto PixelShader = MaterialShaderMap->GetShader< FTranslucentLightingInjectPS >( PermutationVector );
