[HWRT] Fix issues when HasRaytracingDebugViewModeRaytracedOverlay(...) returns false.
- AnyRayTracingPassEnabled(...) was checking HasRayTracedOverlay(...) instead of EngineShowFlags.PathTracing and EngineShowFlags.RayTracingDebug.
- RenderDeferredReflectionsAndSkyLighting(...) was checking EngineShowFlags.RayTracingDebug instead of HasRayTracedOverlay(...).
- UpdatePreExposure() was checking EngineShowFlags.RayTracingDebug instead of HasRayTracedOverlay(...)

diff --git a/Engine/Source/Runtime/Renderer/Private/DeferredShadingRenderer.cpp b/Engine/Source/Runtime/Renderer/Private/DeferredShadingRenderer.cpp
--- a/Engine/Source/Runtime/Renderer/Private/DeferredShadingRenderer.cpp
+++ b/Engine/Source/Runtime/Renderer/Private/DeferredShadingRenderer.cpp
@@ -3518,7 +3518,8 @@ static bool AnyRayTracingPassEnabled(const FScene* Scene, const FViewInfo& View,
 	}
 
 	// Path tracer, ray tracing visualization debug modes, and sky light ray tracing force ray tracing on, regardless of what the view says
-	if (HasRayTracedOverlay(*View.Family)
+	if (View.Family->EngineShowFlags.PathTracing
+		|| View.Family->EngineShowFlags.RayTracingDebug
 		|| ShouldRenderRayTracingSkyLight(Scene->SkyLight, View.GetShaderPlatform()))
 	{
 		return true;
diff --git a/Engine/Source/Runtime/Renderer/Private/IndirectLightRendering.cpp b/Engine/Source/Runtime/Renderer/Private/IndirectLightRendering.cpp
--- a/Engine/Source/Runtime/Renderer/Private/IndirectLightRendering.cpp
+++ b/Engine/Source/Runtime/Renderer/Private/IndirectLightRendering.cpp
@@ -1948,7 +1948,7 @@ void FDeferredShadingSceneRenderer::RenderDeferredReflectionsAndSkyLighting(
 {
 	extern int32 GLumenVisualizeIndirectDiffuse;
 	if (ViewFamily.EngineShowFlags.VisualizeLightCulling 
-		|| ViewFamily.EngineShowFlags.RayTracingDebug
+		|| HasRayTracedOverlay(ViewFamily)
 		|| ViewFamily.EngineShowFlags.PathTracing
 		|| !ViewFamily.EngineShowFlags.Lighting
 		|| GLumenVisualizeIndirectDiffuse != 0)
diff --git a/Engine/Source/Runtime/Renderer/Private/PostProcess/PostProcessEyeAdaptation.cpp b/Engine/Source/Runtime/Renderer/Private/PostProcess/PostProcessEyeAdaptation.cpp
--- a/Engine/Source/Runtime/Renderer/Private/PostProcess/PostProcessEyeAdaptation.cpp
+++ b/Engine/Source/Runtime/Renderer/Private/PostProcess/PostProcessEyeAdaptation.cpp
@@ -1227,7 +1227,7 @@ void FViewInfo::UpdatePreExposure()
 		&& !ViewFamily.EngineShowFlags.HLODColoration
 		&& !ViewFamily.EngineShowFlags.ActorColoration
 		&& ((!ViewFamily.EngineShowFlags.VisualizeBuffer) || CurrentBufferVisualizationMode != NAME_None) // disable pre-exposure for the buffer visualization modes
-		&& !ViewFamily.EngineShowFlags.RayTracingDebug;
+		&& !HasRayTracedOverlay(ViewFamily);
 
 	// Compute the PreExposure to use.
 	bool bUpdateLastExposure = false;
