[HWRT] Refactor IsRayTracingInstanceDebugDataEnabled(...)

diff --git a/Engine/Source/Runtime/Renderer/Private/RayTracing/RayTracingDebug.cpp b/Engine/Source/Runtime/Renderer/Private/RayTracing/RayTracingDebug.cpp
--- a/Engine/Source/Runtime/Renderer/Private/RayTracing/RayTracingDebug.cpp
+++ b/Engine/Source/Runtime/Renderer/Private/RayTracing/RayTracingDebug.cpp
@@ -1782,28 +1782,15 @@ void FDeferredShadingSceneRenderer::RayTracingDisplayPicking(const FRayTracingPi
 
 bool IsRayTracingInstanceDebugDataEnabled(const FViewInfo& View)
 {
-#if !UE_BUILD_SHIPPING	
-	static TArray<FName> RayTracingDebugVisualizationModes;
-	if (RayTracingDebugVisualizationModes.Num() == 0)
-	{
-		RayTracingDebugVisualizationModes.Add(*LOCTEXT("Dynamic Instances", "Dynamic Instances").ToString());
-		RayTracingDebugVisualizationModes.Add(*LOCTEXT("Proxy Type", "Proxy Type").ToString());
-		RayTracingDebugVisualizationModes.Add(*LOCTEXT("Picker", "Picker").ToString());
-	}
+#if !UE_BUILD_SHIPPING
+	const uint32 Mode = GetRaytracingDebugViewMode(View);
 
-	FString ConsoleViewMode = CVarRayTracingVisualize.GetValueOnRenderThread();
-	if (!ConsoleViewMode.IsEmpty())
-	{
-		return RayTracingDebugVisualizationModes.Contains(FName(*ConsoleViewMode));
-	}
-	else if (View.CurrentRayTracingDebugVisualizationMode != NAME_None)
-	{
-		return RayTracingDebugVisualizationModes.Contains(View.CurrentRayTracingDebugVisualizationMode);
-	}
+	return Mode == RAY_TRACING_DEBUG_VIZ_DYNAMIC_INSTANCES
+		|| Mode == RAY_TRACING_DEBUG_VIZ_PROXY_TYPE
+		|| Mode == RAY_TRACING_DEBUG_VIZ_PICKER;
+#else
+	return false;
 #endif
-	{
-		return false;
-	}
 }
 
 bool IsRayTracingInstanceOverlapEnabled(const FViewInfo& View)
