Substrate - fix for crash on mac when clearing MRTs with integer format.

[RN] Fixed Substrate crash on MAC when opening material editor.
#tests Tested PC editor, dx11, dx12 and ps4, with and without Substrate. And on MAC by [REDACTED_PERSON].

diff --git a/Engine/Shaders/Private/OneColorShader.usf b/Engine/Shaders/Private/OneColorShader.usf
--- a/Engine/Shaders/Private/OneColorShader.usf
+++ b/Engine/Shaders/Private/OneColorShader.usf
@@ -53,18 +53,53 @@ void MainVertexShader(
 #ifndef NUM_OUTPUTS
 #define NUM_OUTPUTS 1
 #endif
+#ifndef NUM_UINT_OUTPUTS
+#define NUM_UINT_OUTPUTS 0
+#endif
+
+#define NUM_FLOAT_OUTPUTS (NUM_OUTPUTS - NUM_UINT_OUTPUTS) 
 
 float4 DrawColorMRT[8];
 
 //These shaders used for RHIClear.  Don't use any global uniform variables here as it greatly complicates
 //the state restoration for an inline shader clear.
 void MainPixelShaderMRT( 
-	out float4 OutColor[NUM_OUTPUTS] : SV_Target0)
+#if NUM_UINT_OUTPUTS == 0
+	out float4	OutColor[NUM_OUTPUTS]					: SV_Target0
+#elif NUM_UINT_OUTPUTS == NUM_OUTPUTS
+	out uint	OutUints[NUM_UINT_OUTPUTS]				: SV_Target0
+#else
+	out float4	OutColor[NUM_OUTPUTS - NUM_UINT_OUTPUTS]: SV_Target0
+#if NUM_FLOAT_OUTPUTS == 1
+	, out uint	OutUints[NUM_UINT_OUTPUTS]				: SV_Target1
+#elif(NUM_FLOAT_OUTPUTS == 2)
+	, out uint	OutUints[NUM_UINT_OUTPUTS]				: SV_Target2
+#elif(NUM_FLOAT_OUTPUTS == 3)
+	, out uint	OutUints[NUM_UINT_OUTPUTS]				: SV_Target3
+#elif(NUM_FLOAT_OUTPUTS == 4)
+	, out uint	OutUints[NUM_UINT_OUTPUTS]				: SV_Target4
+#elif(NUM_FLOAT_OUTPUTS == 5)
+	, out uint	OutUints[NUM_UINT_OUTPUTS]				: SV_Target5
+#elif(NUM_FLOAT_OUTPUTS == 6)
+	, out uint	OutUints[NUM_UINT_OUTPUTS]				: SV_Target6
+#elif(NUM_FLOAT_OUTPUTS == 7)
+	, out uint	OutUints[NUM_UINT_OUTPUTS]				: SV_Target7
+#endif
+#endif
+)
 {
-	for (int i = 0; i < NUM_OUTPUTS; i++)
+#if NUM_UINT_OUTPUTS < NUM_OUTPUTS
+	for (int i = 0; i < NUM_FLOAT_OUTPUTS; i++)
 	{
 		OutColor[i] = DrawColorMRT[i];
 	}
+#endif
+#if NUM_UINT_OUTPUTS > 0
+	for (int i = 0; i < NUM_UINT_OUTPUTS; i++)
+	{
+		OutUints[i] = DrawColorMRT[NUM_FLOAT_OUTPUTS + i].x;
+	}
+#endif
 }
 
 void MainPixelShader( 
diff --git a/Engine/Source/Runtime/RenderCore/Private/ClearQuad.cpp b/Engine/Source/Runtime/RenderCore/Private/ClearQuad.cpp
--- a/Engine/Source/Runtime/RenderCore/Private/ClearQuad.cpp
+++ b/Engine/Source/Runtime/RenderCore/Private/ClearQuad.cpp
@@ -29,7 +29,9 @@ TGlobalResource<FClearVertexBuffer> GClearVertexBuffer;
 DEFINE_LOG_CATEGORY_STATIC(LogClearQuad, Log, Log)
 
 template<EColorWriteMask ColorWriteMask>
-static void ClearQuadSetup(FRHICommandList& RHICmdList, int32 NumClearColors, const FLinearColor* ClearColorArray, bool bClearDepth, float Depth, bool bClearStencil, uint32 Stencil, TFunction<void(FGraphicsPipelineStateInitializer&)> PSOModifier = nullptr)
+static void ClearQuadSetup(FRHICommandList& RHICmdList, int32 NumClearColors, const FLinearColor* ClearColorArray, bool bClearDepth, float Depth, bool bClearStencil
+	, uint32 Stencil, TFunction<void(FGraphicsPipelineStateInitializer&)> PSOModifier = nullptr
+	, uint8 NumUintOutput = 0)
 {
 	if (UNLIKELY(!FApp::CanEverRender()))
 	{
@@ -77,6 +79,7 @@ static void ClearQuadSetup(FRHICommandList& RHICmdList, int32 NumClearColors, co
 	TOneColorPixelShaderMRT::FPermutationDomain PermutationVector;
 	PermutationVector.Set<TOneColorPixelShaderMRT::TOneColorPixelShaderNumOutputs>(NumClearColors ? NumClearColors : 1);
 	PermutationVector.Set<TOneColorPixelShaderMRT::TOneColorPixelShader128bitRT>(PlatformRequires128bitRT((EPixelFormat)GraphicsPSOInit.RenderTargetFormats[0]));
+	PermutationVector.Set<TOneColorPixelShaderMRT::TOneColorPixelNumUintOutputs>(NumUintOutput);
 	TShaderMapRef<TOneColorPixelShaderMRT > PixelShader(ShaderMap, PermutationVector);
 
 	GraphicsPSOInit.BoundShaderState.VertexDeclarationRHI = GetVertexDeclarationFVector4();
@@ -122,6 +125,22 @@ void DrawClearQuadMRT(FRHICommandList& RHICmdList, bool bClearColor, int32 NumCl
 	RHICmdList.DrawPrimitive(0, 2, 1);
 }
 
+void DrawClearQuadMRTWithUints(FRHICommandList& RHICmdList, bool bClearColor, int32 NumClearColors, const FLinearColor* ClearColorArray, bool bClearDepth, float Depth, bool bClearStencil, uint32 Stencil, uint8 NumUintOutput)
+{
+	check(NumUintOutput <= NumClearColors);
+	if (bClearColor)
+	{
+		ClearQuadSetup<CW_RGBA>(RHICmdList, NumClearColors, ClearColorArray, bClearDepth, Depth, bClearStencil, Stencil, nullptr , NumUintOutput);
+	}
+	else
+	{
+		ClearQuadSetup<CW_NONE>(RHICmdList, NumClearColors, ClearColorArray, bClearDepth, Depth, bClearStencil, Stencil, nullptr, NumUintOutput);
+	}
+
+	RHICmdList.SetStreamSource(0, GClearVertexBuffer.VertexBufferRHI, 0);
+	RHICmdList.DrawPrimitive(0, 2, 1);
+}
+
 void DrawClearQuadMRT(FRHICommandList& RHICmdList, bool bClearColor, int32 NumClearColors, const FLinearColor* ClearColorArray, bool bClearDepth, float Depth, bool bClearStencil, uint32 Stencil, FClearQuadCallbacks ClearQuadCallbacks)
 {
 	if (bClearColor)
diff --git a/Engine/Source/Runtime/RenderCore/Private/OneColorShader.cpp b/Engine/Source/Runtime/RenderCore/Private/OneColorShader.cpp
--- a/Engine/Source/Runtime/RenderCore/Private/OneColorShader.cpp
+++ b/Engine/Source/Runtime/RenderCore/Private/OneColorShader.cpp
@@ -49,6 +49,13 @@ bool TOneColorPixelShaderMRT::ShouldCompilePermutation(const FGlobalShaderPermut
 {
 	FPermutationDomain PermutationVector(Parameters.PermutationId);
 
+	const int32 NumOutputs = PermutationVector.Get<TOneColorPixelShaderMRT::TOneColorPixelShaderNumOutputs>();
+	const int32 NumUintOutputs = PermutationVector.Get<TOneColorPixelShaderMRT::TOneColorPixelNumUintOutputs>();
+	if (NumOutputs < NumUintOutputs)
+	{
+		return false;
+	}
+
 	if (PermutationVector.Get<TOneColorPixelShaderMRT::TOneColorPixelShaderNumOutputs>())
 	{
 		return (PermutationVector.Get<TOneColorPixelShaderMRT::TOneColorPixelShader128bitRT>() ? FDataDrivenShaderPlatformInfo::GetRequiresExplicit128bitRT(Parameters.Platform) : true);
@@ -66,6 +73,14 @@ void TOneColorPixelShaderMRT::ModifyCompilationEnvironment(const FGlobalShaderPe
 	{
 		OutEnvironment.SetRenderTargetOutputFormat(0, PF_A32B32G32R32F);
 	}
+
+	const int32 NumOutputs = PermutationVector.Get<TOneColorPixelShaderMRT::TOneColorPixelShaderNumOutputs>();
+	const int32 NumUintOutputs = PermutationVector.Get<TOneColorPixelShaderMRT::TOneColorPixelNumUintOutputs>();
+	const int32 UintOutputStartIndex = NumOutputs - NumUintOutputs;
+	for (int32 UintTargetIndex = 0; UintTargetIndex < NumUintOutputs; UintTargetIndex++)
+	{
+		OutEnvironment.SetRenderTargetOutputFormat(UintOutputStartIndex + UintTargetIndex, PF_R32_UINT);
+	}
 }
 
 void TOneColorPixelShaderMRT::SetColors(FRHICommandList& RHICmdList, const TShaderMapRef<TOneColorPixelShaderMRT>& Shader, const FLinearColor* Colors, int32 NumColors)
diff --git a/Engine/Source/Runtime/RenderCore/Public/ClearQuad.h b/Engine/Source/Runtime/RenderCore/Public/ClearQuad.h
--- a/Engine/Source/Runtime/RenderCore/Public/ClearQuad.h
+++ b/Engine/Source/Runtime/RenderCore/Public/ClearQuad.h
@@ -38,6 +38,7 @@ extern RENDERCORE_API void DrawClearQuadAlpha(FRHICommandList& RHICmdList, float
 extern RENDERCORE_API void DrawClearQuadMRT(FRHICommandList& RHICmdList, bool bClearColor, int32 NumClearColors, const FLinearColor* ClearColorArray, bool bClearDepth, float Depth, bool bClearStencil, uint32 Stencil);
 extern RENDERCORE_API void DrawClearQuadMRT(FRHICommandList& RHICmdList, bool bClearColor, int32 NumClearColors, const FLinearColor* ClearColorArray, bool bClearDepth, float Depth, bool bClearStencil, uint32 Stencil, FClearQuadCallbacks ClearQuadCallbacks);
 extern RENDERCORE_API void DrawClearQuadMRT(FRHICommandList& RHICmdList, bool bClearColor, int32 NumClearColors, const FLinearColor* ClearColorArray, bool bClearDepth, float Depth, bool bClearStencil, uint32 Stencil, FIntPoint ViewSize, FIntRect ExcludeRect);
+extern RENDERCORE_API void DrawClearQuadMRTWithUints(FRHICommandList& RHICmdList, bool bClearColor, int32 NumClearColors, const FLinearColor* ClearColorArray, bool bClearDepth, float Depth, bool bClearStencil, uint32 Stencil, uint8 NumUintOutput);
 
 inline void DrawClearQuad(FRHICommandList& RHICmdList, bool bClearColor, const FLinearColor& Color, bool bClearDepth, float Depth, bool bClearStencil, uint32 Stencil)
 {
diff --git a/Engine/Source/Runtime/RenderCore/Public/OneColorShader.h b/Engine/Source/Runtime/RenderCore/Public/OneColorShader.h
--- a/Engine/Source/Runtime/RenderCore/Public/OneColorShader.h
+++ b/Engine/Source/Runtime/RenderCore/Public/OneColorShader.h
@@ -101,7 +101,8 @@ class TOneColorPixelShaderMRT : public FOneColorPS
 public:
 	class TOneColorPixelShader128bitRT : SHADER_PERMUTATION_BOOL("b128BITRENDERTARGET");
 	class TOneColorPixelShaderNumOutputs : SHADER_PERMUTATION_RANGE_INT("NUM_OUTPUTS", 1, 8);
-	using FPermutationDomain = TShaderPermutationDomain<TOneColorPixelShaderNumOutputs, TOneColorPixelShader128bitRT>;
+	class TOneColorPixelNumUintOutputs : SHADER_PERMUTATION_RANGE_INT("NUM_UINT_OUTPUTS", 0, 8);
+	using FPermutationDomain = TShaderPermutationDomain<TOneColorPixelShaderNumOutputs, TOneColorPixelShader128bitRT, TOneColorPixelNumUintOutputs>;
 
 	TOneColorPixelShaderMRT();
 	TOneColorPixelShaderMRT(const ShaderMetaType::CompiledShaderInitializerType& Initializer);
diff --git a/Engine/Source/Runtime/Renderer/Private/BasePassRendering.cpp b/Engine/Source/Runtime/Renderer/Private/BasePassRendering.cpp
--- a/Engine/Source/Runtime/Renderer/Private/BasePassRendering.cpp
+++ b/Engine/Source/Runtime/Renderer/Private/BasePassRendering.cpp
@@ -1134,21 +1134,56 @@ void FDeferredShadingSceneRenderer::RenderBasePass(
 			// If no fast-clear action was used, we need to do an MRT shader clear.
 			if (ColorLoadAction == ERenderTargetLoadAction::ENoAction)
 			{
-				const FRenderTargetBindingSlots& RenderTargets = PassParameters->RenderTargets;
-				FLinearColor ClearColors[MaxSimultaneousRenderTargets];
-				FRHITexture* Textures[MaxSimultaneousRenderTargets];
-				int32 TextureIndex = 0;
+				if (Substrate::IsSubstrateEnabled())
+				{
+					const FRenderTargetBindingSlots& RenderTargets = PassParameters->RenderTargets;
+					FLinearColor ClearColors[MaxSimultaneousRenderTargets];
+					int32 MRTCount = 0;
+					uint8 NumUintOutputs = 0;	// Substrate textures are always integers and we assume they are all at the end of the MRTs.
 
-				RenderTargets.Enumerate([&](const FRenderTargetBinding& RenderTarget)
+					// Clear Only non-UINT targets.
+					RenderTargets.Enumerate([&](const FRenderTargetBinding& RenderTarget)
+						{
+							FRHITexture* TextureRHI = RenderTarget.GetTexture()->GetRHI();
+
+							EPixelFormat Format = TextureRHI->GetFormat();
+							if (!IsInteger(Format))
+							{
+								ClearColors[MRTCount] = MRTCount == 0 ? SceneColorClearValue : TextureRHI->GetClearColor();
+								++MRTCount;
+
+								// We do not support non integer textures after we have identified integer textures starting point.
+								check(NumUintOutputs == 0);
+							}
+							else
+							{
+								ClearColors[MRTCount] = TextureRHI->GetClearColor();
+								++NumUintOutputs;
+								++MRTCount;
+							}
+						});
+
+					// Clear color only; depth-stencil is fast cleared.
+					DrawClearQuadMRTWithUints(RHICmdList, true, MRTCount, ClearColors, false, 0, false, 0, NumUintOutputs);
+				}
+				else
 				{
-					FRHITexture* TextureRHI = RenderTarget.GetTexture()->GetRHI();
-					ClearColors[TextureIndex] = TextureIndex == 0 ? SceneColorClearValue : TextureRHI->GetClearColor();
-					Textures[TextureIndex] = TextureRHI;
-					++TextureIndex;
-				});
-
-				// Clear color only; depth-stencil is fast cleared.
-				DrawClearQuadMRT(RHICmdList, true, TextureIndex, ClearColors, false, 0, false, 0);
+					const FRenderTargetBindingSlots& RenderTargets = PassParameters->RenderTargets;
+					FLinearColor ClearColors[MaxSimultaneousRenderTargets];
+					FRHITexture* Textures[MaxSimultaneousRenderTargets];
+					int32 TextureIndex = 0;
+
+					RenderTargets.Enumerate([&](const FRenderTargetBinding& RenderTarget)
+						{
+							FRHITexture* TextureRHI = RenderTarget.GetTexture()->GetRHI();
+							ClearColors[TextureIndex] = TextureIndex == 0 ? SceneColorClearValue : TextureRHI->GetClearColor();
+							Textures[TextureIndex] = TextureRHI;
+							++TextureIndex;
+						});
+
+					// Clear color only; depth-stencil is fast cleared.
+					DrawClearQuadMRT(RHICmdList, true, TextureIndex, ClearColors, false, 0, false, 0);
+				}
 			}
 		});
 
