HeterogeneousVolumes: Conditionally store indirect lighting term in lighting cache rather than evaluate live in single-scattering pass.

Folding computation into the lighting cache couples the evaluation frequency to that of the lighting cache, but removes an additional texture-fetch during the single-scattering pass.

r.HeterogeneousVolumes.IndirectLighting.Mode
Changes where indirect is accumulated in the pipeline (Default = 0)
0: Off
1. Lighting cache
2: Single-scattering

diff --git a/Engine/Shaders/Private/HeterogeneousVolumes/HeterogeneousVolumesRayMarchingUtils.ush b/Engine/Shaders/Private/HeterogeneousVolumes/HeterogeneousVolumesRayMarchingUtils.ush
--- a/Engine/Shaders/Private/HeterogeneousVolumes/HeterogeneousVolumesRayMarchingUtils.ush
+++ b/Engine/Shaders/Private/HeterogeneousVolumes/HeterogeneousVolumesRayMarchingUtils.ush
@@ -42,6 +42,14 @@
 #define IS_OFFLINE_RENDER 0
 #endif // IS_OFFLINE_RENDER
 
+#define INDIRECT_LIGHTING_MODE_OFF 0
+#define INDIRECT_LIGHTING_MODE_LIGHTING_CACHE 1
+#define INDIRECT_LIGHTING_MODE_SINGLE_SCATTERING 2
+
+#ifndef INDIRECT_LIGHTING_MODE
+#define INDIRECT_LIGHTING_MODE 0
+#endif // INDIRECT_LIGHTING_MODE
+
 #if DIM_USE_ADAPTIVE_VOLUMETRIC_SHADOW_MAP
 #include "HeterogeneousVolumesAdaptiveVolumetricShadowMapSampling.ush"
 #include "HeterogeneousVolumesAdaptiveVolumetricShadowMapSampler.ush"
@@ -244,6 +252,24 @@ float3 ComputeTransmittance(
 	return Transmittance;
 }
 
+float3 ComputeIndirectInscattering(
+	float3 WorldPosition,
+	float3 WorldRayDirection
+)
+{
+	FTwoBandSHVectorRGB TranslucencyGISH = GetTranslucencyGIVolumeLighting(DFPromote(WorldPosition), PrimaryView.WorldToClip, true);
+
+	FTwoBandSHVector RotatedHGZonalHarmonic;
+	float3 CameraVector = -WorldRayDirection;
+	float PhaseG = 0.0;
+	RotatedHGZonalHarmonic.V = float4(1.0f, CameraVector.y, CameraVector.z, CameraVector.x) * float4(1.0f, PhaseG, PhaseG, PhaseG);
+	float3 IndirectInscattering = max(DotSH(TranslucencyGISH, RotatedHGZonalHarmonic), 0) / PI;
+	IndirectInscattering *= GetIndirectInscatteringFactor();
+
+	// Note: Phase is accounted for in spherical hamonic calculation
+	return IndirectInscattering * EvalAmbientOcclusion(WorldPosition);
+}
+
 float3 ComputeInscattering(
 	float3 WorldPosition,
 	FDeferredLightData LightData,
@@ -293,6 +319,15 @@ float3 ComputeInscattering(
 		float3 BiasedWorldPosition = WorldPosition + L * WorldShadowBias;
 		Inscattering *= ComputeTransmittance(BiasedWorldPosition, ToLight, MaxStepCount);
 	}
+
+#if INDIRECT_LIGHTING_MODE == INDIRECT_LIGHTING_MODE_LIGHTING_CACHE
+	float3 WorldRayDirection = normalize(WorldPosition - PrimaryView.TranslatedWorldCameraOrigin);
+	float3 IndirectInscattering = ComputeIndirectInscattering(WorldPosition, WorldRayDirection);
+	// Phase is already accounted for, so add a correction factor
+	float IsotropicPhaseRcp = 4.0 * PI;
+	Inscattering += IndirectInscattering * IsotropicPhaseRcp;
+#endif
+
 #endif // DIM_USE_INSCATTERING_VOLUME
 
 	return Inscattering;
@@ -395,23 +430,13 @@ void RayMarchSingleScattering(
 			float3 Emission = SampleEmission(SampleContext);
 			Radiance += Emission * StepSize * Transmittance;
 
-#if DIM_USE_LUMEN_GI
-			FTwoBandSHVectorRGB TranslucencyGISH = GetTranslucencyGIVolumeLighting(DFPromote(WorldPosition), PrimaryView.WorldToClip, true);
-
-			FTwoBandSHVector RotatedHGZonalHarmonic;
-			float3 CameraVector = normalize(RayMarchingContext.WorldRayDirection);
-			float PhaseG = 0.0;
-			RotatedHGZonalHarmonic.V = float4(1.0f, CameraVector.y, CameraVector.z, CameraVector.x) * float4(1.0f, PhaseG, PhaseG, PhaseG);
-			float3 IndirectInscattering = max(DotSH(TranslucencyGISH, RotatedHGZonalHarmonic), 0) / PI;
-			IndirectInscattering *= GetIndirectInscatteringFactor();
-
+#if INDIRECT_LIGHTING_MODE == INDIRECT_LIGHTING_MODE_SINGLE_SCATTERING
+			float3 IndirectInscattering = ComputeIndirectInscattering(WorldPosition, RayMarchingContext.WorldRayDirection);
 			float3 Extinction = SampleExtinction(SampleContext);
 			float3 Albedo = SampleAlbedo(SampleContext);
 			float3 ScatteringCoefficient = Albedo * Extinction;
-			float AmbientOcclusion = EvalAmbientOcclusion(WorldPosition);
 
-			// Note: Phase is accounted for in spherical hamonic calculation
-			Radiance += IndirectInscattering * AmbientOcclusion * ScatteringCoefficient * Transmittance * StepSize;
+			Radiance += IndirectInscattering * ScatteringCoefficient * Transmittance * StepSize;
 #endif // DIM_USE_LUMEN_GI
 		}
 
diff --git a/Engine/Source/Runtime/Renderer/Private/HeterogeneousVolumes/HeterogeneousVolumes.cpp b/Engine/Source/Runtime/Renderer/Private/HeterogeneousVolumes/HeterogeneousVolumes.cpp
--- a/Engine/Source/Runtime/Renderer/Private/HeterogeneousVolumes/HeterogeneousVolumes.cpp
+++ b/Engine/Source/Runtime/Renderer/Private/HeterogeneousVolumes/HeterogeneousVolumes.cpp
@@ -107,6 +107,16 @@ static TAutoConsoleVariable<float> CVarHeterogeneousVolumesIndirectLighting(
 	ECVF_Scalability | ECVF_RenderThreadSafe
 );
 
+static TAutoConsoleVariable<int32> CVarHeterogeneousVolumesIndirectLightingMode(
+	TEXT("r.HeterogeneousVolumes.IndirectLighting.Mode"),
+	0,
+	TEXT("Changes where indirect is accumulated in the pipeline (Default = 0)\n")
+	TEXT("0: Off\n")
+	TEXT("1: Lighting cache\n")
+	TEXT("2: Single-scattering\n"),
+	ECVF_Scalability | ECVF_RenderThreadSafe
+);
+
 static TAutoConsoleVariable<int32> CVarHeterogeneousVolumesStochasticFiltering(
 	TEXT("r.HeterogeneousVolumes.StochasticFiltering"),
 	3,
@@ -605,6 +615,11 @@ namespace HeterogeneousVolumes
 		return FMath::Max(CVarHeterogeneousVolumesIndirectLighting.GetValueOnRenderThread(), 0.0f);
 	}
 
+	EIndirectLightingMode GetIndirectLightingMode()
+	{
+		return static_cast<EIndirectLightingMode>(FMath::Clamp(CVarHeterogeneousVolumesIndirectLightingMode.GetValueOnRenderThread(), 0, 2));
+	}
+
 	bool ShouldApplyHeightFog()
 	{
 		return CVarHeterogeneousVolumesApplyHeightFog.GetValueOnRenderThread() != 0;
diff --git a/Engine/Source/Runtime/Renderer/Private/HeterogeneousVolumes/HeterogeneousVolumes.h b/Engine/Source/Runtime/Renderer/Private/HeterogeneousVolumes/HeterogeneousVolumes.h
--- a/Engine/Source/Runtime/Renderer/Private/HeterogeneousVolumes/HeterogeneousVolumes.h
+++ b/Engine/Source/Runtime/Renderer/Private/HeterogeneousVolumes/HeterogeneousVolumes.h
@@ -92,6 +92,13 @@ namespace HeterogeneousVolumes
 	bool UseAVSMCompression();
 	float GetCameraDownsampleFactor();
 	float GetIndirectLightingFactor();
+	enum class EIndirectLightingMode
+	{
+		Disabled,
+		LightingCachePass,
+		SingleScatteringPass
+	};
+	EIndirectLightingMode GetIndirectLightingMode();
 
 	int32 GetMipLevel();
 	int32 GetDebugMode();
diff --git a/Engine/Source/Runtime/Renderer/Private/HeterogeneousVolumes/HeterogeneousVolumesHardwareRayTracing.cpp b/Engine/Source/Runtime/Renderer/Private/HeterogeneousVolumes/HeterogeneousVolumesHardwareRayTracing.cpp
--- a/Engine/Source/Runtime/Renderer/Private/HeterogeneousVolumes/HeterogeneousVolumesHardwareRayTracing.cpp
+++ b/Engine/Source/Runtime/Renderer/Private/HeterogeneousVolumes/HeterogeneousVolumesHardwareRayTracing.cpp
@@ -324,7 +324,8 @@ class FRenderLightingCacheWithPreshadingRGS : public FGlobalShader
 	SHADER_USE_ROOT_PARAMETER_STRUCT(FRenderLightingCacheWithPreshadingRGS, FGlobalShader)
 
 	class FLightingCacheMode : SHADER_PERMUTATION_INT("DIM_LIGHTING_CACHE_MODE", 2);
-	using FPermutationDomain = TShaderPermutationDomain<FLightingCacheMode>;
+	class FIndirectLightingMode : SHADER_PERMUTATION_INT("INDIRECT_LIGHTING_MODE", 3);
+	using FPermutationDomain = TShaderPermutationDomain<FLightingCacheMode, FIndirectLightingMode>;
 
 	BEGIN_SHADER_PARAMETER_STRUCT(FParameters, )
 		// Scene 
@@ -346,6 +347,13 @@ class FRenderLightingCacheWithPreshadingRGS : public FGlobalShader
 		SHADER_PARAMETER_STRUCT_INCLUDE(FVirtualShadowMapSamplingParameters, VirtualShadowMapSamplingParameters)
 		SHADER_PARAMETER(int32, VirtualShadowMapId)
 
+		// Global illumination data
+		SHADER_PARAMETER_RDG_UNIFORM_BUFFER(FLumenTranslucencyLightingUniforms, LumenGIVolumeStruct)
+		// TODO: Ambient occlusion pipeline
+		SHADER_PARAMETER_RDG_TEXTURE(Texture3D, AmbientOcclusionTexture)
+		SHADER_PARAMETER(FIntVector, AmbientOcclusionResolution)
+		SHADER_PARAMETER(float, IndirectInscatteringFactor)
+
 		// Sparse Volume
 		SHADER_PARAMETER_RDG_UNIFORM_BUFFER(FSparseVoxelUniformBufferParameters, SparseVoxelUniformBuffer)
 
@@ -366,6 +374,12 @@ class FRenderLightingCacheWithPreshadingRGS : public FGlobalShader
 
 	static bool ShouldCompilePermutation(const FGlobalShaderPermutationParameters& Parameters)
 	{
+		FPermutationDomain PermutationVector(Parameters.PermutationId);
+		if (PermutationVector.template Get<FIndirectLightingMode>() != static_cast<int32>(HeterogeneousVolumes::EIndirectLightingMode::LightingCachePass))
+		{
+			return false;
+		}
+
 		return ShouldCompileRayTracingShadersForProject(Parameters.Platform) && DoesPlatformSupportHeterogeneousVolumes(Parameters.Platform) &&
 			FDataDrivenShaderPlatformInfo::GetSupportsRayTracingProceduralPrimitive(Parameters.Platform);
 	}
@@ -635,6 +649,15 @@ void RenderLightingCacheWithPreshadingHardwareRayTracing(
 		}
 		PassParameters->VirtualShadowMapSamplingParameters = VirtualShadowMapArray.GetSamplingParameters(GraphBuilder, ViewIndex);
 
+		// Global illumination data
+		auto* LumenUniforms = GraphBuilder.AllocParameters<FLumenTranslucencyLightingUniforms>();
+		LumenUniforms->Parameters = GetLumenTranslucencyLightingParameters(GraphBuilder, View.GetLumenTranslucencyGIVolume(), View.LumenFrontLayerTranslucency);
+		PassParameters->LumenGIVolumeStruct = GraphBuilder.CreateUniformBuffer(LumenUniforms);
+		// TODO: Ambient occlusion pipeline
+		//PassParameters->AmbientOcclusionTexture = AmbientOcclusionTexture;
+		//PassParameters->AmbientOcclusionResolution = HeterogeneousVolumes::GetAmbientOcclusionResolution(HeterogeneousVolumeInterface, LODValue);
+		PassParameters->IndirectInscatteringFactor = HeterogeneousVolumes::GetIndirectLightingFactor();
+
 		// Output
 		PassParameters->RWLightingCacheTexture = GraphBuilder.CreateUAV(LightingCacheTexture);
 	}
@@ -653,8 +676,12 @@ void RenderLightingCacheWithPreshadingHardwareRayTracing(
 	}
 #endif // WANTS_DRAW_MESH_EVENTS
 
+	// Indirect lighting accumulation is coupled with directional light, because it doesn't voxel cull. It is assumed to exist and shadow.
+	int32 IndirectLightingMode = LightType == LightType_Directional ? static_cast<int32>(HeterogeneousVolumes::GetIndirectLightingMode()) : 0;
+
 	FRenderLightingCacheWithPreshadingRGS::FPermutationDomain PermutationVector;
 	PermutationVector.Set<FRenderLightingCacheWithPreshadingRGS::FLightingCacheMode>(HeterogeneousVolumes::GetLightingCacheMode() - 1);
+	PermutationVector.Set<FRenderLightingCacheWithPreshadingRGS::FIndirectLightingMode>(IndirectLightingMode);
 	TShaderRef<FRenderLightingCacheWithPreshadingRGS> RayGenerationShader = View.ShaderMap->GetShader<FRenderLightingCacheWithPreshadingRGS>(PermutationVector);
 	HeterogeneousVolumes::FLODValue LODValue = HeterogeneousVolumes::CalcLOD(View, HeterogeneousVolumeInterface);
 	FIntVector VolumeResolution = HeterogeneousVolumes::GetLightingCacheResolution(HeterogeneousVolumeInterface, LODValue);
diff --git a/Engine/Source/Runtime/Renderer/Private/HeterogeneousVolumes/HeterogeneousVolumesLiveShadingPipeline.cpp b/Engine/Source/Runtime/Renderer/Private/HeterogeneousVolumes/HeterogeneousVolumesLiveShadingPipeline.cpp
--- a/Engine/Source/Runtime/Renderer/Private/HeterogeneousVolumes/HeterogeneousVolumesLiveShadingPipeline.cpp
+++ b/Engine/Source/Runtime/Renderer/Private/HeterogeneousVolumes/HeterogeneousVolumesLiveShadingPipeline.cpp
@@ -265,7 +265,8 @@ class FRenderLightingCacheWithLiveShadingCS : public FMeshMaterialShader
 	class FAVSMSampleMode : SHADER_PERMUTATION_INT("AVSM_SAMPLE_MODE", 3);
 	class FUseExistenceMask : SHADER_PERMUTATION_INT("USE_EXISTENCE_MASK", 2);
 	class FIsOfflineRender : SHADER_PERMUTATION_INT("IS_OFFLINE_RENDER", 2);
-	using FPermutationDomain = TShaderPermutationDomain<FLightingCacheMode, FAVSMSampleMode, FUseExistenceMask, FIsOfflineRender>;
+	class FIndirectLightingMode : SHADER_PERMUTATION_INT("INDIRECT_LIGHTING_MODE", 3);
+	using FPermutationDomain = TShaderPermutationDomain<FLightingCacheMode, FAVSMSampleMode, FUseExistenceMask, FIsOfflineRender, FIndirectLightingMode>;
 
 	BEGIN_SHADER_PARAMETER_STRUCT(FParameters, )
 		// Scene data
@@ -290,6 +291,12 @@ class FRenderLightingCacheWithLiveShadingCS : public FMeshMaterialShader
 		SHADER_PARAMETER(int32, VirtualShadowMapId)
 		SHADER_PARAMETER_RDG_UNIFORM_BUFFER(FAdaptiveVolumetricShadowMapUniformBufferParameters, AVSM)
 
+		// Global illumination data
+		SHADER_PARAMETER_RDG_UNIFORM_BUFFER(FLumenTranslucencyLightingUniforms, LumenGIVolumeStruct)
+		SHADER_PARAMETER_RDG_TEXTURE(Texture3D, AmbientOcclusionTexture)
+		SHADER_PARAMETER(FIntVector, AmbientOcclusionResolution)
+		SHADER_PARAMETER(float, IndirectInscatteringFactor)
+
 		// Object data
 		SHADER_PARAMETER(FMatrix44f, LocalToWorld)
 		SHADER_PARAMETER(FMatrix44f, WorldToLocal)
@@ -339,12 +346,23 @@ class FRenderLightingCacheWithLiveShadingCS : public FMeshMaterialShader
 		const FMaterialShaderPermutationParameters & Parameters
 	)
 	{
+		FPermutationDomain PermutationVector(Parameters.PermutationId);
+		if (RemapPermutation(PermutationVector) != PermutationVector)
+		{
+			return false;
+		}
+
 		return DoesPlatformSupportHeterogeneousVolumes(Parameters.Platform)
 			&& DoesMaterialShaderSupportHeterogeneousVolumes(Parameters.MaterialParameters);
 	}
 
 	static FPermutationDomain RemapPermutation(FPermutationDomain PermutationVector)
 	{
+		if (PermutationVector.Get<FIndirectLightingMode>() != static_cast<int32>(HeterogeneousVolumes::EIndirectLightingMode::LightingCachePass))
+		{
+			PermutationVector.Set<FIndirectLightingMode>(static_cast<int32>(HeterogeneousVolumes::EIndirectLightingMode::Disabled));
+		}
+
 		return PermutationVector;
 	}
 
@@ -404,7 +422,7 @@ class FRenderSingleScatteringWithLiveShadingCS : public FMeshMaterialShader
 
 	//class FUseTransmittanceVolume : SHADER_PERMUTATION_BOOL("DIM_USE_TRANSMITTANCE_VOLUME");
 	class FUseInscatteringVolume : SHADER_PERMUTATION_BOOL("DIM_USE_INSCATTERING_VOLUME");
-	class FUseLumenGI : SHADER_PERMUTATION_BOOL("DIM_USE_LUMEN_GI");
+	class FIndirectLightingMode : SHADER_PERMUTATION_INT("INDIRECT_LIGHTING_MODE", 3);
 	class FWriteVelocity : SHADER_PERMUTATION_BOOL("DIM_WRITE_VELOCITY");
 	class FAVSMSampleMode : SHADER_PERMUTATION_INT("AVSM_SAMPLE_MODE", 3);
 	class FSupportOverlappingVolumes : SHADER_PERMUTATION_BOOL("SUPPORT_OVERLAPPING_VOLUMES");
@@ -416,7 +434,7 @@ class FRenderSingleScatteringWithLiveShadingCS : public FMeshMaterialShader
 	using FPermutationDomain = TShaderPermutationDomain<
 		//FUseTransmittanceVolume, 
 		FUseInscatteringVolume, 
-		FUseLumenGI, 
+		FIndirectLightingMode,
 		FWriteVelocity, 
 		FApplyFogInscattering, 
 		FAVSMSampleMode, 
@@ -524,7 +542,7 @@ class FRenderSingleScatteringWithLiveShadingCS : public FMeshMaterialShader
 	)
 	{
 		FPermutationDomain PermutationVector(Parameters.PermutationId);
-		if (PermutationVector.template Get<FApplyFogInscattering>() == 0)
+		if (RemapPermutation(PermutationVector) != PermutationVector)
 		{
 			return false;
 		}
@@ -541,6 +559,11 @@ class FRenderSingleScatteringWithLiveShadingCS : public FMeshMaterialShader
 			PermutationVector.template Set<FApplyFogInscattering>(2);
 		}
 
+		if (PermutationVector.template Get<FIndirectLightingMode>() != static_cast<int32>(HeterogeneousVolumes::EIndirectLightingMode::SingleScatteringPass))
+		{
+			PermutationVector.template Set<FIndirectLightingMode>(static_cast<int32>(HeterogeneousVolumes::EIndirectLightingMode::Disabled));
+		}
+
 		return PermutationVector;
 	}
 
@@ -653,6 +676,8 @@ static void RenderLightingCacheWithLiveShading(
 	// Shadow data
 	const FVisibleLightInfo* VisibleLightInfo,
 	const FVirtualShadowMapArray& VirtualShadowMapArray,
+	// Global illumination data
+	FRDGTextureRef AmbientOcclusionTexture,
 	// Object data
 	const IHeterogeneousVolumeInterface* HeterogeneousVolumeInterface,
 	const FMaterialRenderProxy* DefaultMaterialRenderProxy,
@@ -786,6 +811,14 @@ static void RenderLightingCacheWithLiveShading(
 		PassParameters->VirtualShadowMapSamplingParameters = VirtualShadowMapArray.GetSamplingParameters(GraphBuilder, ViewIndex);
 		PassParameters->AVSM = HeterogeneousVolumes::GetAdaptiveVolumetricShadowMapUniformBuffer(GraphBuilder, View.ViewState, LightSceneInfo);
 
+		// Global illumination data
+		auto* LumenUniforms = GraphBuilder.AllocParameters<FLumenTranslucencyLightingUniforms>();
+		LumenUniforms->Parameters = GetLumenTranslucencyLightingParameters(GraphBuilder, View.GetLumenTranslucencyGIVolume(), View.LumenFrontLayerTranslucency);
+		PassParameters->LumenGIVolumeStruct = GraphBuilder.CreateUniformBuffer(LumenUniforms);
+		PassParameters->AmbientOcclusionTexture = AmbientOcclusionTexture;
+		PassParameters->AmbientOcclusionResolution = HeterogeneousVolumes::GetAmbientOcclusionResolution(HeterogeneousVolumeInterface, LODValue);
+		PassParameters->IndirectInscatteringFactor = HeterogeneousVolumes::GetIndirectLightingFactor();
+
 		// Output
 		PassParameters->RWLightingCacheTexture = GraphBuilder.CreateUAV(LightingCacheTexture);
 	}
@@ -861,12 +894,16 @@ static void RenderLightingCacheWithLiveShading(
 		GroupCount.Z = FMath::DivideAndRoundUp(GroupCount.Z, FRenderLightingCacheWithLiveShadingCS::GetThreadGroupSize3D());
 
 		bool bUseAVSM = HeterogeneousVolumes::UseAdaptiveVolumetricShadowMapForSelfShadowing(HeterogeneousVolumeInterface->GetPrimitiveSceneProxy());
+		// Indirect lighting accumulation is coupled with directional light, because it doesn't voxel cull. It is assumed to exist and shadow.
+		int32 IndirectLightingMode = LightType == LightType_Directional ? static_cast<int32>(HeterogeneousVolumes::GetIndirectLightingMode()) : 0;
 
 		FRenderLightingCacheWithLiveShadingCS::FPermutationDomain PermutationVector;
 		PermutationVector.Set<FRenderLightingCacheWithLiveShadingCS::FLightingCacheMode>(HeterogeneousVolumes::GetLightingCacheMode() - 1);
 		PermutationVector.Set<FRenderLightingCacheWithLiveShadingCS::FAVSMSampleMode>(static_cast<int32>(HeterogeneousVolumes::GetAVSMSampleMode(bUseAVSM)));
 		PermutationVector.Set<FRenderLightingCacheWithLiveShadingCS::FUseExistenceMask>(HeterogeneousVolumes::UseExistenceMask());
 		PermutationVector.Set<FRenderLightingCacheWithLiveShadingCS::FIsOfflineRender>(View.bIsOfflineRender);
+		PermutationVector.Set<FRenderLightingCacheWithLiveShadingCS::FIndirectLightingMode>(IndirectLightingMode);
+		PermutationVector = FRenderLightingCacheWithLiveShadingCS::RemapPermutation(PermutationVector);
 		TShaderRef<FRenderLightingCacheWithLiveShadingCS> ComputeShader = Material.GetShader<FRenderLightingCacheWithLiveShadingCS>(&FLocalVertexFactory::StaticType, PermutationVector, false);
 		if (!ComputeShader.IsNull())
 		{
@@ -1280,11 +1317,13 @@ void RenderSingleScatteringWithLiveShading(
 		PassName = FString::Printf(TEXT("RenderSingleScatteringWithLiveShadingCS (Light = %s)"), *LightName);
 	}
 #endif // WANTS_DRAW_MESH_EVENTS
+	// Indirect lighting accumulation is coupled with directional light, because it doesn't cull voxels. It is assumed to exist and shadow.
+	int32 IndirectLightingMode = View.GetLumenTranslucencyGIVolume().Texture0 != nullptr ? static_cast<int32>(HeterogeneousVolumes::GetIndirectLightingMode()) : 0;
 
 	typename FRenderSingleScatteringWithLiveShadingDispatchTypeCS::FPermutationDomain PermutationVector;
 	//PermutationVector.template Set<typename FRenderSingleScatteringWithLiveShadingDispatchTypeCS::FUseTransmittanceVolume>(HeterogeneousVolumes::UseLightingCacheForTransmittance() && PassParameters->bApplyShadowTransmittance);
 	PermutationVector.template Set<typename FRenderSingleScatteringWithLiveShadingDispatchTypeCS::FUseInscatteringVolume>(HeterogeneousVolumes::UseLightingCacheForInscattering());
-	PermutationVector.template Set<typename FRenderSingleScatteringWithLiveShadingDispatchTypeCS::FUseLumenGI>(HeterogeneousVolumes::UseIndirectLighting() && View.GetLumenTranslucencyGIVolume().Texture0 != nullptr);
+	PermutationVector.template Set<typename FRenderSingleScatteringWithLiveShadingDispatchTypeCS::FIndirectLightingMode>(IndirectLightingMode);
 	PermutationVector.template Set<typename FRenderSingleScatteringWithLiveShadingDispatchTypeCS::FWriteVelocity>(bWriteVelocity);
 	PermutationVector.template Set<typename FRenderSingleScatteringWithLiveShadingDispatchTypeCS::FAVSMSampleMode>(static_cast<int32>(HeterogeneousVolumes::GetAVSMSampleMode(bUseAVSM)));
 	PermutationVector.template Set<typename FRenderSingleScatteringWithLiveShadingDispatchTypeCS::FSupportOverlappingVolumes>(HeterogeneousVolumes::SupportsOverlappingVolumes());
@@ -1376,6 +1415,8 @@ static void RenderWithTransmittanceVolumePipeline(
 				// Shadow data
 				VisibleLightInfo,
 				VirtualShadowMapArray,
+				// Global illumination data
+				AmbientOcclusionTexture,
 				// Object data
 				HeterogeneousVolumeInterface,
 				MaterialRenderProxy,
@@ -1529,6 +1570,8 @@ static void RenderWithInscatteringVolumePipeline(
 				// Shadow data
 				VisibleLightInfo,
 				VirtualShadowMapArray,
+				// Global illumination data
+				AmbientOcclusionTexture,
 				// Object data
 				HeterogeneousVolumeInterface,
 				MaterialRenderProxy,
diff --git a/Engine/Source/Runtime/Renderer/Private/HeterogeneousVolumes/HeterogeneousVolumesPreshadingPipeline.cpp b/Engine/Source/Runtime/Renderer/Private/HeterogeneousVolumes/HeterogeneousVolumesPreshadingPipeline.cpp
--- a/Engine/Source/Runtime/Renderer/Private/HeterogeneousVolumes/HeterogeneousVolumesPreshadingPipeline.cpp
+++ b/Engine/Source/Runtime/Renderer/Private/HeterogeneousVolumes/HeterogeneousVolumesPreshadingPipeline.cpp
@@ -210,7 +210,8 @@ class FRenderLightingCacheWithPreshadingCS : public FGlobalShader
 
 	class FLightingCacheMode : SHADER_PERMUTATION_INT("DIM_LIGHTING_CACHE_MODE", 2);
 	class FUseAdaptiveVolumetricShadowMap : SHADER_PERMUTATION_BOOL("DIM_USE_ADAPTIVE_VOLUMETRIC_SHADOW_MAP");
-	using FPermutationDomain = TShaderPermutationDomain<FLightingCacheMode, FUseAdaptiveVolumetricShadowMap>;
+	class FIndirectLightingMode : SHADER_PERMUTATION_INT("INDIRECT_LIGHTING_MODE", 3);
+	using FPermutationDomain = TShaderPermutationDomain<FLightingCacheMode, FUseAdaptiveVolumetricShadowMap, FIndirectLightingMode>;
 
 	BEGIN_SHADER_PARAMETER_STRUCT(FParameters, )
 		// Scene data
@@ -232,6 +233,13 @@ class FRenderLightingCacheWithPreshadingCS : public FGlobalShader
 		SHADER_PARAMETER(int32, VirtualShadowMapId)
 		SHADER_PARAMETER_RDG_UNIFORM_BUFFER(FAdaptiveVolumetricShadowMapUniformBufferParameters, AVSM)
 
+		// Global illumination data
+		SHADER_PARAMETER_RDG_UNIFORM_BUFFER(FLumenTranslucencyLightingUniforms, LumenGIVolumeStruct)
+		// TODO: Ambient occlusion pipeline
+		//SHADER_PARAMETER_RDG_TEXTURE(Texture3D, AmbientOcclusionTexture)
+		//SHADER_PARAMETER(FIntVector, AmbientOcclusionResolution)
+		SHADER_PARAMETER(float, IndirectInscatteringFactor)
+
 		// Volume structures
 		SHADER_PARAMETER_RDG_UNIFORM_BUFFER(FSparseVoxelUniformBufferParameters, SparseVoxelUniformBuffer)
 		SHADER_PARAMETER_STRUCT_INCLUDE(FLightingCacheParameters, LightingCache)
@@ -251,9 +259,25 @@ class FRenderLightingCacheWithPreshadingCS : public FGlobalShader
 		const FGlobalShaderPermutationParameters& Parameters
 	)
 	{
+		FPermutationDomain PermutationVector(Parameters.PermutationId);
+		if (PermutationVector.template Get<FIndirectLightingMode>() == static_cast<int32>(HeterogeneousVolumes::EIndirectLightingMode::SingleScatteringPass))
+		{
+			return false;
+		}
+
 		return DoesPlatformSupportHeterogeneousVolumes(Parameters.Platform);
 	}
 
+	static FPermutationDomain RemapPermutation(FPermutationDomain PermutationVector)
+	{
+		if (PermutationVector.Get<FIndirectLightingMode>() != static_cast<int32>(HeterogeneousVolumes::EIndirectLightingMode::LightingCachePass))
+		{
+			PermutationVector.Set<FIndirectLightingMode>(static_cast<int32>(HeterogeneousVolumes::EIndirectLightingMode::Disabled));
+		}
+
+		return PermutationVector;
+	}
+
 	static void ModifyCompilationEnvironment(
 		const FGlobalShaderPermutationParameters& Parameters,
 		FShaderCompilerEnvironment& OutEnvironment
@@ -295,11 +319,11 @@ class FRenderSingleScatteringWithPreshadingCS : public FGlobalShader
 	class FSparseVoxelTracingDim : SHADER_PERMUTATION_BOOL("DIM_SPARSE_VOXEL_TRACING");
 	//class FUseTransmittanceVolume : SHADER_PERMUTATION_BOOL("DIM_USE_TRANSMITTANCE_VOLUME");
 	class FUseInscatteringVolume : SHADER_PERMUTATION_BOOL("DIM_USE_INSCATTERING_VOLUME");
-	class FUseLumenGI : SHADER_PERMUTATION_BOOL("DIM_USE_LUMEN_GI");
+	class FIndirectLightingMode : SHADER_PERMUTATION_INT("INDIRECT_LIGHTING_MODE", 3);
 	class FWriteVelocity : SHADER_PERMUTATION_BOOL("DIM_WRITE_VELOCITY");
 	class FUseAdaptiveVolumetricShadowMap : SHADER_PERMUTATION_BOOL("DIM_USE_ADAPTIVE_VOLUMETRIC_SHADOW_MAP");
 	class FDebugDim : SHADER_PERMUTATION_BOOL("DIM_DEBUG");
-	using FPermutationDomain = TShaderPermutationDomain<FApplyShadowTransmittanceDim, FVoxelCullingDim, FSparseVoxelTracingDim, FUseInscatteringVolume, FUseLumenGI, FWriteVelocity, FUseAdaptiveVolumetricShadowMap, FDebugDim>;
+	using FPermutationDomain = TShaderPermutationDomain<FApplyShadowTransmittanceDim, FVoxelCullingDim, FSparseVoxelTracingDim, FUseInscatteringVolume, FIndirectLightingMode, FWriteVelocity, FUseAdaptiveVolumetricShadowMap, FDebugDim>;
 
 	BEGIN_SHADER_PARAMETER_STRUCT(FParameters, )
 		// Scene data
@@ -357,9 +381,25 @@ class FRenderSingleScatteringWithPreshadingCS : public FGlobalShader
 		const FGlobalShaderPermutationParameters& Parameters
 	)
 	{
+		FPermutationDomain PermutationVector(Parameters.PermutationId);
+		if (PermutationVector.Get<FIndirectLightingMode>() == static_cast<int32>(HeterogeneousVolumes::EIndirectLightingMode::LightingCachePass))
+		{
+			return false;
+		}
+
 		return DoesPlatformSupportHeterogeneousVolumes(Parameters.Platform);
 	}
 
+	static FPermutationDomain RemapPermutation(FPermutationDomain PermutationVector)
+	{
+		if (PermutationVector.Get<FIndirectLightingMode>() != static_cast<int32>(HeterogeneousVolumes::EIndirectLightingMode::SingleScatteringPass))
+		{
+			PermutationVector.Set<FIndirectLightingMode>(static_cast<int32>(HeterogeneousVolumes::EIndirectLightingMode::Disabled));
+		}
+
+		return PermutationVector;
+	}
+
 	static EShaderPermutationPrecacheRequest ShouldPrecachePermutation(const FGlobalShaderPermutationParameters& Parameters)
 	{
 		FPermutationDomain PermutationVector(Parameters.PermutationId);
@@ -502,6 +542,15 @@ void RenderLightingCacheWithPreshadingCompute(
 		PassParameters->VirtualShadowMapSamplingParameters = VirtualShadowMapArray.GetSamplingParameters(GraphBuilder, ViewIndex);
 		PassParameters->AVSM = HeterogeneousVolumes::GetAdaptiveVolumetricShadowMapUniformBuffer(GraphBuilder, View.ViewState, LightSceneInfo);
 
+		// Global illumination data
+		auto* LumenUniforms = GraphBuilder.AllocParameters<FLumenTranslucencyLightingUniforms>();
+		LumenUniforms->Parameters = GetLumenTranslucencyLightingParameters(GraphBuilder, View.GetLumenTranslucencyGIVolume(), View.LumenFrontLayerTranslucency);
+		PassParameters->LumenGIVolumeStruct = GraphBuilder.CreateUniformBuffer(LumenUniforms);
+		// TODO: Ambient occlusion pipeline
+		//PassParameters->AmbientOcclusionTexture = AmbientOcclusionTexture;
+		//PassParameters->AmbientOcclusionResolution = HeterogeneousVolumes::GetAmbientOcclusionResolution(HeterogeneousVolumeInterface, LODValue);
+		PassParameters->IndirectInscatteringFactor = HeterogeneousVolumes::GetIndirectLightingFactor();
+
 		// Output
 		PassParameters->RWLightingCacheTexture = GraphBuilder.CreateUAV(LightingCacheTexture);
 	}
@@ -520,10 +569,14 @@ void RenderLightingCacheWithPreshadingCompute(
 	}
 #endif // WANTS_DRAW_MESH_EVENTS
 	bool bUseAVSM = HeterogeneousVolumes::UseAdaptiveVolumetricShadowMapForSelfShadowing(HeterogeneousVolumeInterface->GetPrimitiveSceneProxy());
+	// Indirect lighting accumulation is coupled with directional light, because it doesn't voxel cull. It is assumed to exist and shadow.
+	int32 IndirectLightingMode = LightType == LightType_Directional ? static_cast<int32>(HeterogeneousVolumes::GetIndirectLightingMode()) : 0;
 
 	FRenderLightingCacheWithPreshadingCS::FPermutationDomain PermutationVector;
 	PermutationVector.Set<FRenderLightingCacheWithPreshadingCS::FLightingCacheMode>(HeterogeneousVolumes::GetLightingCacheMode() - 1);
 	PermutationVector.Set<FRenderLightingCacheWithPreshadingCS::FUseAdaptiveVolumetricShadowMap>(bUseAVSM);
+	PermutationVector.Set<FRenderLightingCacheWithPreshadingCS::FIndirectLightingMode>(IndirectLightingMode);
+	PermutationVector = FRenderLightingCacheWithPreshadingCS::RemapPermutation(PermutationVector);
 	TShaderRef<FRenderLightingCacheWithPreshadingCS> ComputeShader = View.ShaderMap->GetShader<FRenderLightingCacheWithPreshadingCS>(PermutationVector);
 
 	FIntVector GroupCount = HeterogeneousVolumes::GetLightingCacheResolution(HeterogeneousVolumeInterface, LODValue);
@@ -680,6 +733,8 @@ void RenderSingleScatteringWithPreshadingCompute(
 		FSceneRenderer::GetLightNameForDrawEvent(LightSceneInfo->Proxy, LightName);
 	}
 	bool bUseAVSM = HeterogeneousVolumes::UseAdaptiveVolumetricShadowMapForSelfShadowing(HeterogeneousVolumeInterface->GetPrimitiveSceneProxy());
+	// Indirect lighting accumulation is coupled with directional light, because it doesn't cull voxels. It is assumed to exist and shadow.
+	int32 IndirectLightingMode = View.GetLumenTranslucencyGIVolume().Texture0 != nullptr ? static_cast<int32>(HeterogeneousVolumes::GetIndirectLightingMode()) : 0;
 
 	FRenderSingleScatteringWithPreshadingCS::FPermutationDomain PermutationVector;
 	PermutationVector.Set<FRenderSingleScatteringWithPreshadingCS::FApplyShadowTransmittanceDim>(bApplyShadowTransmittance);
@@ -687,10 +742,12 @@ void RenderSingleScatteringWithPreshadingCompute(
 	PermutationVector.Set<FRenderSingleScatteringWithPreshadingCS::FSparseVoxelTracingDim>(HeterogeneousVolumes::UseSparseVoxelPipeline());
 	//PermutationVector.Set<FRenderSingleScatteringWithPreshadingCS::FUseTransmittanceVolume>(HeterogeneousVolumes::UseLightingCacheForTransmittance());
 	PermutationVector.Set<FRenderSingleScatteringWithPreshadingCS::FUseInscatteringVolume>(HeterogeneousVolumes::UseLightingCacheForInscattering());
-	PermutationVector.Set<FRenderSingleScatteringWithPreshadingCS::FUseLumenGI>(HeterogeneousVolumes::UseIndirectLighting() && View.GetLumenTranslucencyGIVolume().Texture0 != nullptr);
+	PermutationVector.Set<FRenderSingleScatteringWithPreshadingCS::FIndirectLightingMode>(IndirectLightingMode);
 	PermutationVector.Set<FRenderSingleScatteringWithPreshadingCS::FWriteVelocity>(bWriteVelocity);
 	PermutationVector.Set<FRenderSingleScatteringWithPreshadingCS::FUseAdaptiveVolumetricShadowMap>(bUseAVSM);
 	PermutationVector.Set<FRenderSingleScatteringWithPreshadingCS::FDebugDim>(HeterogeneousVolumes::GetDebugMode() != 0);
+	PermutationVector = FRenderSingleScatteringWithPreshadingCS::RemapPermutation(PermutationVector);
+
 	TShaderRef<FRenderSingleScatteringWithPreshadingCS> ComputeShader = View.ShaderMap->GetShader<FRenderSingleScatteringWithPreshadingCS>(PermutationVector);
 	FComputeShaderUtils::AddPass(
 		GraphBuilder,
