Potential fix for crash with new blue nosie texture.

#tests PC editor dx12 with PIE
#rnx

diff --git a/Engine/Source/Runtime/Renderer/Private/SceneRendering.cpp b/Engine/Source/Runtime/Renderer/Private/SceneRendering.cpp
--- a/Engine/Source/Runtime/Renderer/Private/SceneRendering.cpp
+++ b/Engine/Source/Runtime/Renderer/Private/SceneRendering.cpp
@@ -2208,6 +2208,7 @@ void FViewInfo::SetupUniformBufferParameters(
 	{
 		FBlueNoise BlueNoiseParam = GetBlueNoiseGlobalParameters();
 		ViewUniformShaderParameters.BlueNoiseScalarTexture = BlueNoiseParam.BlueNoise.ScalarTexture;
+		ViewUniformShaderParameters.BlueNoiseScalarTexture = OrBlack2DIfNull(ViewUniformShaderParameters.BlueNoiseScalarTexture);
 		ViewUniformShaderParameters.BlueNoiseDimensions = BlueNoiseParam.BlueNoise.Dimensions;
 		ViewUniformShaderParameters.BlueNoiseModuloMasks = BlueNoiseParam.BlueNoise.ModuloMasks;
 	}
