Adding scalability flag

diff --git a/Engine/Source/Runtime/Renderer/Private/HeterogeneousVolumes/HeterogeneousVolumesLiveShadingPipeline.cpp b/Engine/Source/Runtime/Renderer/Private/HeterogeneousVolumes/HeterogeneousVolumesLiveShadingPipeline.cpp
--- a/Engine/Source/Runtime/Renderer/Private/HeterogeneousVolumes/HeterogeneousVolumesLiveShadingPipeline.cpp
+++ b/Engine/Source/Runtime/Renderer/Private/HeterogeneousVolumes/HeterogeneousVolumesLiveShadingPipeline.cpp
@@ -71,7 +71,7 @@ static TAutoConsoleVariable<int32> CVarHeterogeneousVolumesUseExistenceMask(
 	TEXT("r.HeterogeneousVolumes.UseExistenceMask"),
 	1,
 	TEXT("Creates an evaluation mask which culls operations to the areas with non-zero extinction (Default = 1)"),
-	ECVF_RenderThreadSafe
+	ECVF_Scalability | ECVF_RenderThreadSafe
 );
 
 static TAutoConsoleVariable<int32> CVarHeterogeneousVolumesShadowsUseCameraSceneDepth(
