- Fix for skipping out FX calls in the mobile renderers hit proxy rendering
- Add a check to the count manager to ensure we don't leave dangling references in the indirect arg map

diff --git a/Engine/Plugins/FX/Niagara/Source/Niagara/Private/NiagaraGPUInstanceCountManager.cpp b/Engine/Plugins/FX/Niagara/Source/Niagara/Private/NiagaraGPUInstanceCountManager.cpp
--- a/Engine/Plugins/FX/Niagara/Source/Niagara/Private/NiagaraGPUInstanceCountManager.cpp
+++ b/Engine/Plugins/FX/Niagara/Source/Niagara/Private/NiagaraGPUInstanceCountManager.cpp
@@ -319,6 +319,7 @@ void FNiagaraGPUInstanceCountManager::ResizeBuffers(FRHICommandListImmediate& RH
 void FNiagaraGPUInstanceCountManager::FlushIndirectArgsPool(FRHICommandListBase& RHICmdList)
 {
 	check(IsInRenderingThread());
+	checkf(DrawIndirectArgMap.IsEmpty(), TEXT("DrawIndirectArgMap is not empty in FlushIndirectArgsPool.  This means that OnPostRenderOpaque was not called on the GPU sort manager."));
 
 	// Cull indirect draw pool entries so that we only keep the last pool
 	while (DrawIndirectPool.Num() > 1)
diff --git a/Engine/Source/Runtime/Renderer/Private/SceneHitProxyRendering.cpp b/Engine/Source/Runtime/Renderer/Private/SceneHitProxyRendering.cpp
--- a/Engine/Source/Runtime/Renderer/Private/SceneHitProxyRendering.cpp
+++ b/Engine/Source/Runtime/Renderer/Private/SceneHitProxyRendering.cpp
@@ -556,6 +556,25 @@ void FMobileSceneRenderer::RenderHitProxies(FRDGBuilder& GraphBuilder)
 	GetSceneExtensionsRenderers().PreRender(GraphBuilder);
 	GEngine->GetPreRenderDelegateEx().Broadcast(GraphBuilder);
 
+	// Notify the FX system that the scene is about to be rendered.
+	// TODO: These should probably be moved to scene extensions
+	if (FXSystem && Views.IsValidIndex(0))
+	{
+		FGPUSortManager* GPUSortManager = FXSystem->GetGPUSortManager();
+		FXSystem->PreRender(GraphBuilder, GetSceneViews(), GetSceneUniforms(), false);
+		if (GPUSortManager)
+		{
+			GPUSortManager->OnPreRender(GraphBuilder);
+		}
+		// Call PostRenderOpaque now as this is irrelevant for when rendering hit proxies.
+		// because we don't tick the particles in the render loop (see last param being "false").
+		FXSystem->PostRenderOpaque(GraphBuilder, GetSceneViews(), GetSceneUniforms(), false /*bAllowGPUParticleUpdate*/);
+		if (GPUSortManager)
+		{
+			GPUSortManager->OnPostRenderOpaque(GraphBuilder);
+		}
+	}
+
 	TArray<Nanite::FRasterResults, TInlineAllocator<2>> NaniteRasterResults;
 	::DoRenderHitProxies(GraphBuilder, this, HitProxyTexture, HitProxyDepthTexture, NaniteRasterResults, InstanceCullingManager);
 
