Implemented depth-of-field Petzval lens effect for swirly bokeh. In the gathering pass, sample positions are warped relative to the screen center. In the scattering phase, the bokeh pixel shader warps the sprite in a similar way. The warping shrinks the bokeh in one direction, rather than enlarging it in the other direction. In other words, the footprint of the warped bokeh fits within the non-warped size. This is to avoid performance pitfalls of overdraw and large sampling offsets.

diff --git a/Engine/Shaders/Private/DiaphragmDOF/DOFGatherCommon.ush b/Engine/Shaders/Private/DiaphragmDOF/DOFGatherCommon.ush
--- a/Engine/Shaders/Private/DiaphragmDOF/DOFGatherCommon.ush
+++ b/Engine/Shaders/Private/DiaphragmDOF/DOFGatherCommon.ush
@@ -135,17 +135,20 @@
 #elif DIM_GATHER_QUALITY == GATHER_QUALITY_HIGH_QUALITY
 	#define CONFIG_WEIGHT_SAMPLE_FROM_COC 1
 	#define CONFIG_MIRRORED_SAMPLER 1
+	#define CONFIG_PETZVAL 1
 
 #elif DIM_GATHER_QUALITY == GATHER_QUALITY_SCATTER_OCCLUSION_OUTPUT
 	#define CONFIG_WEIGHT_SAMPLE_FROM_COC 1
 	#define CONFIG_GENERATE_SCATTER_OCCLUSION_BUFFER 1
 	#define CONFIG_MIRRORED_SAMPLER 1
+	#define CONFIG_PETZVAL 1
 
 #elif DIM_GATHER_QUALITY == GATHER_QUALITY_CINEMATIC
 	#define CONFIG_WEIGHT_SAMPLE_FROM_COC 1
 	#define CONFIG_GENERATE_SCATTER_OCCLUSION_BUFFER 1
 	#define CONFIG_MIRRORED_SAMPLER 1
 	#define CONFIG_LAYER_GATHERING 1
+	#define CONFIG_PETZVAL 1
 
 #else
 	#error Unknon gathering quality.
@@ -305,6 +308,8 @@ struct FGatherInputParameters
 
 	// buffer UV coordinate of the center of the kernel to gather.
 	float2 InputBufferCenterUV;
+	
+	float2 ViewportUV;
 
 	// Multiplication factors to map DispatchThreadID -> InputBufferUV.
 	float2 DispatchThreadIdToInputBufferUV;
diff --git a/Engine/Shaders/Private/DiaphragmDOF/DOFGatherKernel.ush b/Engine/Shaders/Private/DiaphragmDOF/DOFGatherKernel.ush
--- a/Engine/Shaders/Private/DiaphragmDOF/DOFGatherKernel.ush
+++ b/Engine/Shaders/Private/DiaphragmDOF/DOFGatherKernel.ush
@@ -20,6 +20,9 @@ Texture2D GatherInput_SeparateCoc;
 
 Texture2D BokehLUT;
 
+float Petzval;
+float PetzvalFalloffPower;
+
 
 //------------------------------------------------------- GATHERING KERNEL
 
@@ -105,6 +108,9 @@ struct FPrivateKernelSharedScope
 
 	// Multiplier to apply when computing sample intersection. Higher value sharpen, lower value feather.
 	float IntersectionMultiplier;
+
+	// Used to stretch sample positions for petzval lens effect
+	float2x2 PetzvalTransform;
 };
 
 
@@ -295,6 +301,29 @@ void UpdateRuntimeKernelScope(
 
 	KernelScope.KernelCenterUV = GatherParameters.InputBufferCenterUV + (RandomCircle * KernelRandomRange) * GatherInputSize.zw;
 
+	KernelScope.PetzvalTransform = 0;
+	if (Petzval != 0)
+	{
+		const float2 BokehCenter = GatherParameters.ViewportUV * 2.0 - 1.0;
+		const float BokehCenterLength = length(BokehCenter);
+		const float InvBokehCenterLength = rcp(BokehCenterLength);
+		const float2 KernelCenterNormal = BokehCenter * InvBokehCenterLength;
+
+		const float PetzvalIntensity = pow(BokehCenterLength, PetzvalFalloffPower);
+		const float PetzvalStretchFactor = rcp(1 + abs(Petzval) * PetzvalIntensity);
+
+		const float2 Normal = KernelCenterNormal;
+		const float2 Tangent = float2(Normal.y, -Normal.x);
+		const float2x2 ToSagittalSpace = {Normal.x, Normal.y, Tangent.x, Tangent.y};
+		bool bMask = Petzval > 0;
+		const float2x2 Scale = { 
+			bMask ? PetzvalStretchFactor : 1.0, // Sagittal
+			0.0, 0.0, 
+			!bMask ? PetzvalStretchFactor : 1.0 // Tangential
+		};
+		KernelScope.PetzvalTransform = mul(mul(ToSagittalSpace, Scale), transpose(ToSagittalSpace));
+	}
+
 	KernelScope.SampleCountToSamplePosition = GatherParameters.KernelRadius * UnitSampleRadius;
 
 	float MipTexelSize = pow(2.0, uint(MipLevel));
@@ -319,6 +348,24 @@ void UpdateRuntimeKernelScope(
 	#endif
 }
 
+// Apply approximate Petzval bokeh lens effect.
+// A more precise version would calculate the transform per sample, rather than per kernel.
+// This version assumes the bokeh that are gathered are all centered on the kernel center, rather than on the sample center.
+// This is incorrect and can cause bokeh to stretch into a 'banana shape', rather than oval.
+// But it's significantly faster and usually not noticeable (more prominent bokeh are scattered anyway).
+float2 ApplyApproxPetzval(in FPrivateKernelSharedScope KernelScope, float2 SampleUV)
+{
+#if CONFIG_PETZVAL
+	if (Petzval != 0)
+	{
+		return mul(SampleUV, KernelScope.PetzvalTransform);
+	}
+	else
+#endif //CONFIG_PETZVAL
+	{
+		return SampleUV;
+	}
+}
 
 // Transform at compile time a 2 dimensional batch's constant into sample pair constant, by using rotation invariance.
 float2 SampleConstFromBatchConst(const uint BatchSampleId, float2 BatchConst)
@@ -507,6 +554,8 @@ void GatherRingSamples(
 			{
 				const float SampleSign = (k == 0) ? 1.0 : -1.0;
 
+				SamplePosOffset[k] = ApplyApproxPetzval(KernelScope, SamplePosOffset[k]);
+
 				// Compute sample's input buffer UV.
 				float2 SampleUV = KernelScope.KernelCenterUV + SamplePosOffset[k] * GatherInputSize.zw;
 
diff --git a/Engine/Shaders/Private/DiaphragmDOF/DOFGatherPass.usf b/Engine/Shaders/Private/DiaphragmDOF/DOFGatherPass.usf
--- a/Engine/Shaders/Private/DiaphragmDOF/DOFGatherPass.usf
+++ b/Engine/Shaders/Private/DiaphragmDOF/DOFGatherPass.usf
@@ -212,6 +212,7 @@ void GatherMainCS(
 	// Sets the per wave line gather parameters
 	{
 		GatherParameters.InputBufferCenterUV = InputBufferCenterUV;
+		GatherParameters.ViewportUV = ViewportUV;
 
 		// Generate random signals
 		GatherParameters.Random[0] = InterleavedGradientNoise(DispatchThreadId, 0);
diff --git a/Engine/Shaders/Private/DiaphragmDOF/DOFHybridScatterCommon.ush b/Engine/Shaders/Private/DiaphragmDOF/DOFHybridScatterCommon.ush
--- a/Engine/Shaders/Private/DiaphragmDOF/DOFHybridScatterCommon.ush
+++ b/Engine/Shaders/Private/DiaphragmDOF/DOFHybridScatterCommon.ush
@@ -85,3 +85,8 @@
 	#error Unknown bokeh drawing method.
 
 #endif
+
+#if DIM_SCATTER_QUALITY == SCATTER_QUALITY_MEDIUM
+#elif DIM_SCATTER_QUALITY == SCATTER_QUALITY_HIGH
+	#define CONFIG_PETZVAL 1
+#endif //DIM_SCATTER_QUALITY
\ No newline at end of file
diff --git a/Engine/Shaders/Private/DiaphragmDOF/DOFHybridScatterPixelShader.usf b/Engine/Shaders/Private/DiaphragmDOF/DOFHybridScatterPixelShader.usf
--- a/Engine/Shaders/Private/DiaphragmDOF/DOFHybridScatterPixelShader.usf
+++ b/Engine/Shaders/Private/DiaphragmDOF/DOFHybridScatterPixelShader.usf
@@ -27,11 +27,43 @@ Texture2D BokehLUT;
 float4 ScatterOcclusionSize;
 Texture2D ScatterOcclusion;
 
+float Petzval;
+float PetzvalFalloffPower;
+
+
+float2 ApplyPetzval(float2 BokehCenter, float2 SamplePos)
+{
+#if CONFIG_PETZVAL
+	if (Petzval != 0)
+	{
+		const float BokehCenterLength = length(BokehCenter);
+
+		const float2 Normal = BokehCenter / BokehCenterLength;
+		const float2 Tangent = float2(Normal.y, -Normal.x);
+
+		const float Intensity = pow(BokehCenterLength, PetzvalFalloffPower);
+		const float Stretch = 1 + abs(Petzval) * Intensity;
+
+		float2x2 ToSaggitalSpace = {Normal.x, Normal.y, Tangent.x, Tangent.y};
+			
+		float2 SagittalPos = mul(SamplePos, ToSaggitalSpace);
+		bool bMask = Petzval > 0;
+		SagittalPos.x *=  bMask ? Stretch : 1.0; // Sagittal
+		SagittalPos.y *= !bMask ? Stretch : 1.0; // Tangential
+		float2 Result = mul(SagittalPos, transpose(ToSaggitalSpace));
+		return Result;
+	}
+	else
+#endif //CONFIG_PETZVAL
+	{
+		return SamplePos;
+	}
+}
 
 //------------------------------------------------------- ENTRY POINT
 
 void ScatterMainPS(
-	in nointerpolation float2 BokehCenter0 : TEXCOORD0,
+	in nointerpolation float4 BokehCenterScreenAndClipPos : TEXCOORD0,
 	in nointerpolation float4 ColorAndAbsCocRadius0 : TEXCOORD1,
 	in nointerpolation float4 ColorAndAbsCocRadius1 : TEXCOORD2,
 	in nointerpolation float4 ColorAndAbsCocRadius2 : TEXCOORD3,
@@ -45,6 +77,9 @@ void ScatterMainPS(
 {
 	OutColor = 0;
 
+	const float2 BokehCenter0 = BokehCenterScreenAndClipPos.xy;
+	const float2 BokehCenterClipPos = BokehCenterScreenAndClipPos.zw;
+
 	float4 Color[PIXEL_COUNT_PER_SCATTER_GROUP] = {
 		float4(ColorAndAbsCocRadius0.rgb, 0),
 		float4(ColorAndAbsCocRadius1.rgb, 0),
@@ -83,6 +118,9 @@ void ScatterMainPS(
 	{
 		// Pixel coordinate of the center of the bokeh.
 		float2 PixelOffset = PixelOffset0 - ScatteringScaling * kSquare2x2[i];
+		
+		PixelOffset = ApplyPetzval(BokehCenterClipPos, PixelOffset);
+
 		PixelOffset.x *= rcp(CocInvSqueeze);
 
 		float Intersection;
diff --git a/Engine/Shaders/Private/DiaphragmDOF/DOFHybridScatterVertexShader.usf b/Engine/Shaders/Private/DiaphragmDOF/DOFHybridScatterVertexShader.usf
--- a/Engine/Shaders/Private/DiaphragmDOF/DOFHybridScatterVertexShader.usf
+++ b/Engine/Shaders/Private/DiaphragmDOF/DOFHybridScatterVertexShader.usf
@@ -33,7 +33,7 @@ StructuredBuffer<float4> ScatterDrawList;
 void ScatterMainVS(
 	in uint VertexId : SV_VertexID,
 	in uint InstanceId : SV_InstanceID,
-	out nointerpolation float2 BokehCenter0 : TEXCOORD0,
+	out nointerpolation float4 BokehCenterScreenAndClipPos : TEXCOORD0,
 	out nointerpolation float4 ColorAndAbsCocRadius0 : TEXCOORD1,
 	out nointerpolation float4 ColorAndAbsCocRadius1 : TEXCOORD2,
 	out nointerpolation float4 ColorAndAbsCocRadius2 : TEXCOORD3,
@@ -100,13 +100,16 @@ void ScatterMainVS(
 	}
 
 	// Output interpolators.
+	const float2 BokehCenterScreenPos = ScatteringScaling * ScatterDrawList[5 * GlobalSpriteId].xy;
+	const float2 BokehCenterClipPos = (BokehCenterScreenPos * ViewportSize.zw) * 2.0 - 1.0;
 	{
 		ColorAndAbsCocRadius0 = float4(Color[0].rgb, AbsCocRadius[0]);
 		ColorAndAbsCocRadius1 = float4(Color[1].rgb, AbsCocRadius[1]);
 		ColorAndAbsCocRadius2 = float4(Color[2].rgb, AbsCocRadius[2]);
 		ColorAndAbsCocRadius3 = float4(Color[3].rgb, AbsCocRadius[3]);
 	
-		BokehCenter0 = ScatteringScaling * ScatterDrawList[5 * GlobalSpriteId].xy;
+		BokehCenterScreenAndClipPos.xy = BokehCenterScreenPos;
+		BokehCenterScreenAndClipPos.zw = BokehCenterClipPos;
 	
 		#if TAYLOR_EXPENDED_COC_INTERSECTION
 		{
@@ -132,10 +135,10 @@ void ScatterMainVS(
 
 	#endif
 
-	// The center of the pixel group is 0.5 south east of BokehCenter0.
+	// The center of the pixel group is 0.5 south east of BokehCenterScreen.
 	const float HalfPixelShift = 0.5;
 
-	float2 ViewportUV = ((MaxAbsCocRadius * CocRadiusToCircumscribedRadius * float2(CocInvSqueeze, 1.0) + (HalfPixelShift + ANTI_ALIASING_FEATHER_OFFSET)) * VertexPos + BokehCenter0 + HalfPixelShift) * ViewportSize.zw; 
+	float2 ViewportUV = ((MaxAbsCocRadius * CocRadiusToCircumscribedRadius * float2(CocInvSqueeze, 1.0) + (HalfPixelShift + ANTI_ALIASING_FEATHER_OFFSET)) * VertexPos + BokehCenterScreenPos + HalfPixelShift) * ViewportSize.zw; 
 
 	OutPosition = float4(ViewportUV.x * 2.0 - 1.0, 1.0 - ViewportUV.y * 2.0, 0.0, 1.0);
-}
+}
\ No newline at end of file
diff --git a/Engine/Shaders/Private/PathTracing/PathTracingCore.ush b/Engine/Shaders/Private/PathTracing/PathTracingCore.ush
--- a/Engine/Shaders/Private/PathTracing/PathTracingCore.ush
+++ b/Engine/Shaders/Private/PathTracing/PathTracingCore.ush
@@ -85,6 +85,8 @@ float MaxNormalBias;
 float FilterWidth;
 float CameraFocusDistance;
 float2 CameraLensRadius;
+float Petzval;
+float PetzvalFalloffPower;
 
 RWTexture2D<float4> RadianceTexture;
 #if PATH_TRACER_USE_ADAPTIVE_SAMPLING
@@ -1520,6 +1522,37 @@ bool ProcessSubsurfaceRandomWalk(inout FPathTracingPayload Payload, inout float3
 	return false;
 }
 
+// Apply sagittal/tangential Petzval lens effect
+//  BokehCenterUV: Viewport UV of the center of the bokeh
+//  SampleUV: UV of the sample relative to the bokeh center, with 1 = bokeh boundary
+float2 ApplyPetzval(float2 BokehCenterUV, float2 SampleUV)
+{
+	if (Petzval != 0)
+	{
+		const float2 BokehCenter = (BokehCenterUV * 2.0-1.0);
+		const float BokehCenterLength = length(BokehCenter);
+
+		float2 Normal = BokehCenter / BokehCenterLength;
+		Normal.y *= -1.0;
+		const float2 Tangent = float2(Normal.y, -Normal.x);
+
+		const float Intensity = pow(BokehCenterLength, PetzvalFalloffPower);
+		const float Stretch = rcp(1 + abs(Petzval) * Intensity);
+
+		const float2x2 ToSagittalSpace = {Normal.x, Normal.y, Tangent.x, Tangent.y};
+			
+		float2 SagittalPos = mul(SampleUV, ToSagittalSpace);
+		bool bMask = Petzval > 0;
+		SagittalPos.x *=  bMask ? Stretch : 1.0; // Sagittal
+		SagittalPos.y *= !bMask ? Stretch : 1.0; // Tangential
+		return mul(SagittalPos, transpose(ToSagittalSpace));
+	}
+	else
+	{
+		return SampleUV;
+	}
+}
+
 FPathState CreatePathState(int2 PixelIndex, int2 TextureIndex)
 {
 	FPathState Output = (FPathState)0;
@@ -1546,6 +1579,7 @@ FPathState CreatePathState(int2 PixelIndex, int2 TextureIndex)
 #endif
 
 	// Initialize ray and payload
+	float2 ViewportUV;
 	{
 		float2 AAJitter = RandomSequence_GenerateSample2D(Output.RandSequence);
 #if 0
@@ -1558,8 +1592,9 @@ FPathState CreatePathState(int2 PixelIndex, int2 TextureIndex)
 		float Sigma = FilterWidth / 6.0; // user-provided width covers +/-3*Sigma
 		AAJitter = 0.5 + Sigma * Disk.xy * sqrt(-2.0 * log(1.0 - Disk.z * Disk.z));
 #endif
-		float2 ViewportUV = (LaunchIndex + AAJitter) * View.BufferSizeAndInvSize.zw;
-		Output.Ray = CreatePrimaryRay(ViewportUV);
+		ViewportUV = (PixelIndex + AAJitter) * View.ViewSizeAndInvSize.zw;
+		float2 JitteredViewportUV = (LaunchIndex + AAJitter) * View.BufferSizeAndInvSize.zw;
+		Output.Ray = CreatePrimaryRay(JitteredViewportUV);
 	}
 
 	if (CameraLensRadius.y > 0)
@@ -1576,7 +1611,9 @@ FPathState CreatePathState(int2 PixelIndex, int2 TextureIndex)
 		// compute focus plane
 		float3 FocusP = Origin + (CameraFocusDistance * ZFactor) * Output.Ray.Direction;
 		// nudge ray origin
-		LensSample = CameraLensRadius * UniformSampleDiskConcentric(LensSample);
+		LensSample = UniformSampleDiskConcentric(LensSample);
+		LensSample = ApplyPetzval(ViewportUV, LensSample);
+		LensSample = CameraLensRadius * LensSample;
 		Origin += LensSample.x * ViewX + LensSample.y * ViewY;
 		// recompute direction
 		Output.Ray.Direction = normalize(FocusP - Origin);
diff --git a/Engine/Source/Runtime/Engine/Classes/Engine/Scene.h b/Engine/Source/Runtime/Engine/Classes/Engine/Scene.h
--- a/Engine/Source/Runtime/Engine/Classes/Engine/Scene.h
+++ b/Engine/Source/Runtime/Engine/Classes/Engine/Scene.h
@@ -1133,6 +1133,12 @@ struct FPostProcessSettings
 	UPROPERTY(EditAnywhere, BlueprintReadWrite, Category=Overrides, meta=(PinHiddenByDefault, InlineEditConditionToggle))
 	uint8 bOverride_DepthOfFieldUseHairDepth:1;
 
+	UPROPERTY(EditAnywhere, BlueprintReadWrite, Category=Overrides, meta=(PinHiddenByDefault, InlineEditConditionToggle))
+	uint8 bOverride_DepthOfFieldPetzvalBokeh:1;
+
+	UPROPERTY(EditAnywhere, BlueprintReadWrite, Category=Overrides, meta=(PinHiddenByDefault, InlineEditConditionToggle))
+	uint8 bOverride_DepthOfFieldPetzvalBokehFalloff:1;
+
 	UPROPERTY(EditAnywhere, BlueprintReadWrite, Category=Overrides, meta=(PinHiddenByDefault, InlineEditConditionToggle))
 	uint8 bOverride_DepthOfFieldDepthBlurAmount:1;
 
@@ -2150,6 +2156,14 @@ struct FPostProcessSettings
 	UPROPERTY(interp, EditAnywhere, BlueprintReadWrite, Category="Lens|Depth of Field", AdvancedDisplay, meta=(editcondition = "bOverride_DepthOfFieldUseHairDepth", DisplayName = "Use Hair Depth"))
 	uint32 DepthOfFieldUseHairDepth:1;
 
+	/** Simulate stretching in blur and bokeh. Positive values for sagittal (swirly bokeh), negative values for tangential. */
+	UPROPERTY(interp, EditAnywhere, BlueprintReadWrite, Category = "Lens|Depth of Field", meta = (ClampMin = "-10.0", ClampMax = "10.0", editcondition = "bOverride_DepthOfFieldPetzvalBokeh", DisplayName = "Petzval Bokeh"))
+	float DepthOfFieldPetzvalBokeh;
+
+	/** How quickly does the Petzval bokeh effect increase towards the edge of the image */
+	UPROPERTY(interp, EditAnywhere, BlueprintReadWrite, Category="Lens|Depth of Field", meta=(ClampMin = "0.0", UIMax = "10.0", editcondition = "bOverride_DepthOfFieldPetzvalBokehFalloff", DisplayName = "Petzval Bokeh Falloff"))
+	float DepthOfFieldPetzvalBokehFalloff;
+
 	/** Artificial region where all content is in focus, starting after DepthOfFieldFocalDistance, in unreal units  (cm) */
 	UPROPERTY(interp, EditAnywhere, BlueprintReadWrite, Category="Lens|Mobile Depth of Field", meta=(UIMin = "0.0", UIMax = "10000.0", editcondition = "bOverride_DepthOfFieldFocalRegion", DisplayName = "Focal Region"))
 	float DepthOfFieldFocalRegion;
diff --git a/Engine/Source/Runtime/Engine/Private/Scene.cpp b/Engine/Source/Runtime/Engine/Private/Scene.cpp
--- a/Engine/Source/Runtime/Engine/Private/Scene.cpp
+++ b/Engine/Source/Runtime/Engine/Private/Scene.cpp
@@ -589,6 +589,8 @@ FPostProcessSettings::FPostProcessSettings()
 	DepthOfFieldDepthBlurAmount = 1.0f;
 	DepthOfFieldDepthBlurRadius = 0.0f;
 	DepthOfFieldUseHairDepth = 0;
+	DepthOfFieldPetzvalBokeh = 0;
+	DepthOfFieldPetzvalBokehFalloff = 1.0f;
 	DepthOfFieldFocalRegion = 0.0f;
 	DepthOfFieldNearTransitionRegion = 300.0f;
 	DepthOfFieldFarTransitionRegion = 500.0f;
diff --git a/Engine/Source/Runtime/Engine/Private/SceneView.cpp b/Engine/Source/Runtime/Engine/Private/SceneView.cpp
--- a/Engine/Source/Runtime/Engine/Private/SceneView.cpp
+++ b/Engine/Source/Runtime/Engine/Private/SceneView.cpp
@@ -1711,6 +1711,8 @@ void FSceneView::OverridePostProcessSettings(const FPostProcessSettings& Src, fl
 		LERP_PP(DepthOfFieldSqueezeFactor);
 		LERP_PP(DepthOfFieldDepthBlurRadius);
 		SET_PP(DepthOfFieldUseHairDepth)
+		LERP_PP(DepthOfFieldPetzvalBokeh);
+		LERP_PP(DepthOfFieldPetzvalBokehFalloff);
 		LERP_PP(DepthOfFieldDepthBlurAmount);
 		LERP_PP(DepthOfFieldFocalRegion);
 		LERP_PP(DepthOfFieldNearTransitionRegion);
diff --git a/Engine/Source/Runtime/Renderer/Private/PathTracing.cpp b/Engine/Source/Runtime/Renderer/Private/PathTracing.cpp
--- a/Engine/Source/Runtime/Renderer/Private/PathTracing.cpp
+++ b/Engine/Source/Runtime/Renderer/Private/PathTracing.cpp
@@ -563,6 +563,8 @@ BEGIN_SHADER_PARAMETER_STRUCT(FPathTracingData, )
 	SHADER_PARAMETER(float, MeshDecalBias)
 	SHADER_PARAMETER(float, CameraFocusDistance)
 	SHADER_PARAMETER(FVector2f, CameraLensRadius)
+	SHADER_PARAMETER(float, Petzval)
+	SHADER_PARAMETER(float, PetzvalFalloffPower)
 END_SHADER_PARAMETER_STRUCT()
 
 
@@ -640,8 +642,9 @@ struct FPathTracingConfig
 	bool IsDOFDifferent(const FPathTracingConfig& Other) const
 	{
 		return PathTracingData.CameraFocusDistance != Other.PathTracingData.CameraFocusDistance ||
-			   PathTracingData.CameraLensRadius != Other.PathTracingData.CameraLensRadius;
-
+			   PathTracingData.CameraLensRadius != Other.PathTracingData.CameraLensRadius ||
+			   PathTracingData.Petzval != Other.PathTracingData.Petzval ||
+			   PathTracingData.PetzvalFalloffPower != Other.PathTracingData.PetzvalFalloffPower;
 	}
 };
 
@@ -825,14 +828,18 @@ static void PreparePathTracingData(const FScene* Scene, const FViewInfo& View, F
 	PathTracingData.FilterWidth = CVarPathTracingFilterWidth.GetValueOnRenderThread();
 	PathTracingData.CameraFocusDistance = 0;
 	PathTracingData.CameraLensRadius = FVector2f::ZeroVector;
+	PathTracingData.Petzval = 0;
+	PathTracingData.PetzvalFalloffPower = 0;
 	if (ShowFlags.DepthOfField &&
 		PathTracing::UsesReferenceDOF(View) &&
 		PPV.DepthOfFieldFocalDistance > 0)
 	{
 		DiaphragmDOF::FPhysicalCocModel CocModel;
 		CocModel.Compile(View);
-		PathTracingData.CameraFocusDistance = CocModel.FocusDistance;
-		PathTracingData.CameraLensRadius    = CocModel.GetLensRadius();
+		PathTracingData.CameraFocusDistance 			= CocModel.FocusDistance;
+		PathTracingData.CameraLensRadius    			= CocModel.GetLensRadius();
+		PathTracingData.Petzval	 			= CocModel.Petzval;
+		PathTracingData.PetzvalFalloffPower = CocModel.PetzvalFalloffPower;
 	}
 
 	const bool bUseReferenceAtmosphere = ShouldRenderSkyAtmosphere(Scene, ShowFlags) && 
diff --git a/Engine/Source/Runtime/Renderer/Private/PostProcess/DiaphragmDOF.cpp b/Engine/Source/Runtime/Renderer/Private/PostProcess/DiaphragmDOF.cpp
--- a/Engine/Source/Runtime/Renderer/Private/PostProcess/DiaphragmDOF.cpp
+++ b/Engine/Source/Runtime/Renderer/Private/PostProcess/DiaphragmDOF.cpp
@@ -245,6 +245,15 @@ enum class EDiaphragmDOFGatherQuality
 	MAX,
 };
 
+/** Quality configurations for scattering passes. */
+enum class EDiaphragmDOFScatterQuality
+{
+	Medium,
+	High,
+
+	MAX,
+};
+
 /** Format of the LUT to generate. */
 enum class EDiaphragmDOFBokehLUTFormat
 {
@@ -1123,6 +1132,11 @@ class FDiaphragmDOFGatherCS : public FDiaphragmDOFShader
 
 		// This shader takes a very long time to compile with FXC, so we pre-compile it with DXC first and then forward the optimized HLSL to FXC.
 		OutEnvironment.CompilerFlags.Add(CFLAG_PrecompileWithDXC);
+
+		if (FDataDrivenShaderPlatformInfo::GetSupportsRealTypes(Parameters.Platform) == ERHIFeatureSupport::RuntimeGuaranteed)
+		{
+			OutEnvironment.CompilerFlags.Add(CFLAG_AllowRealTypes);
+		}
 	}
 
 	
@@ -1139,6 +1153,9 @@ class FDiaphragmDOFGatherCS : public FDiaphragmDOFShader
 		SHADER_PARAMETER(float, MaxRecombineAbsCocRadius)
 		SHADER_PARAMETER(float, CocSqueeze)
 		SHADER_PARAMETER(float, CocInvSqueeze)
+		
+		SHADER_PARAMETER(float, Petzval)
+		SHADER_PARAMETER(float, PetzvalFalloffPower)
 
 		SHADER_PARAMETER_STRUCT_INCLUDE(FDOFTileDecisionParameters, TileDecisionParameters)
 		SHADER_PARAMETER_STRUCT_INCLUDE(FDOFCommonShaderParameters, CommonParameters)
@@ -1211,6 +1228,9 @@ BEGIN_SHADER_PARAMETER_STRUCT(FDOFHybridScatterParameters, )
 	SHADER_PARAMETER(float, CocSqueeze)
 	SHADER_PARAMETER(float, CocInvSqueeze)
 	
+	SHADER_PARAMETER(float, Petzval)
+	SHADER_PARAMETER(float, PetzvalFalloffPower)
+	
 	SHADER_PARAMETER_STRUCT_INCLUDE(FDOFCommonShaderParameters, CommonParameters)
 	
 	SHADER_PARAMETER(FVector4f, ScatterOcclusionSize)
@@ -1247,9 +1267,10 @@ class FDiaphragmDOFHybridScatterPS : public FDiaphragmDOFShader
 	DECLARE_GLOBAL_SHADER(FDiaphragmDOFHybridScatterPS);
 	SHADER_USE_PARAMETER_STRUCT(FDiaphragmDOFHybridScatterPS, FDiaphragmDOFShader);
 
-	class FBokehSimulationDim : SHADER_PERMUTATION_BOOL("DIM_BOKEH_SIMULATION");
+	class FBokehSimulationDim		: SHADER_PERMUTATION_BOOL("DIM_BOKEH_SIMULATION");
+	class FDDOFScatterQualityDim	: SHADER_PERMUTATION_ENUM_CLASS("DIM_SCATTER_QUALITY", EDiaphragmDOFScatterQuality);
 
-	using FPermutationDomain = TShaderPermutationDomain<FDDOFLayerProcessingDim, FBokehSimulationDim, FDDOFScatterOcclusionDim>;
+	using FPermutationDomain = TShaderPermutationDomain<FDDOFLayerProcessingDim, FBokehSimulationDim, FDDOFScatterOcclusionDim, FDDOFScatterQualityDim>;
 	using FParameters = FDOFHybridScatterParameters;
 
 	static FPermutationDomain RemapPermutation(FPermutationDomain PermutationVector)
@@ -1282,8 +1303,19 @@ class FDiaphragmDOFHybridScatterPS : public FDiaphragmDOFShader
 		if (PermutationVector.Get<FDDOFLayerProcessingDim>() != EDiaphragmDOFLayerProcessing::ForegroundOnly &&
 			PermutationVector.Get<FDDOFLayerProcessingDim>() != EDiaphragmDOFLayerProcessing::BackgroundOnly) return false;
 
+		if (PermutationVector.Get<FDDOFScatterQualityDim>() != EDiaphragmDOFScatterQuality::Medium &&
+			PermutationVector.Get<FDDOFScatterQualityDim>() != EDiaphragmDOFScatterQuality::High) return false;
+
 		return FDiaphragmDOFShader::ShouldCompilePermutation(Parameters);
 	}
+
+	static void ModifyCompilationEnvironment(const FGlobalShaderPermutationParameters& Parameters, FShaderCompilerEnvironment& OutEnvironment)
+	{
+		FDiaphragmDOFShader::ModifyCompilationEnvironment(Parameters, OutEnvironment);
+
+		OutEnvironment.SetDefine(TEXT("SCATTER_QUALITY_MEDIUM"), (int)EDiaphragmDOFScatterQuality::Medium);
+		OutEnvironment.SetDefine(TEXT("SCATTER_QUALITY_HIGH"), (int)EDiaphragmDOFScatterQuality::High);
+	}
 };
 
 class FDiaphragmDOFRecombineCS : public FDiaphragmDOFShader
@@ -2370,6 +2402,9 @@ bool DiaphragmDOF::AddPasses(
 			PassParameters->MaxRecombineAbsCocRadius = float(kMaxSlightOutOfFocusCocRadius) / (GatheringCocRadiusBasis / EncodedCocRadiusBasis);
 			PassParameters->CocSqueeze = CocModel.Squeeze;
 			PassParameters->CocInvSqueeze = 1.0f / CocModel.Squeeze;
+			
+			PassParameters->Petzval = CocModel.Petzval;
+			PassParameters->PetzvalFalloffPower = CocModel.PetzvalFalloffPower;
 
 			PassParameters->TileDecisionParameters.MinGatherRadius = PassParameters->MaxRecombineAbsCocRadius - 1;
 			PassParameters->TileDecisionParameters.SlightOutOfFocusRadiusBoundary = float(kMaxSlightOutOfFocusCocRadius) / (GatheringCocRadiusBasis / EncodedCocRadiusBasis);
@@ -2502,10 +2537,21 @@ bool DiaphragmDOF::AddPasses(
 			bool bIsForeground = ConvolutionSettings.LayerProcessing == EDiaphragmDOFLayerProcessing::ForegroundOnly;
 			int32 DrawIndirectParametersOffset = bIsForeground ? 0 : 1;
 
+			EDiaphragmDOFScatterQuality Quality {};
+			if (ConvolutionSettings.QualityConfig >= EDiaphragmDOFGatherQuality::HighQuality)
+			{
+				Quality = EDiaphragmDOFScatterQuality::High;
+			}
+			else
+			{
+				Quality = EDiaphragmDOFScatterQuality::Medium;
+			}
+
 			FDiaphragmDOFHybridScatterPS::FPermutationDomain PermutationVector;
 			PermutationVector.Set<FDDOFLayerProcessingDim>(ConvolutionSettings.LayerProcessing);
 			PermutationVector.Set<FDiaphragmDOFHybridScatterPS::FBokehSimulationDim>(ScatteringBokehLUT ? true : false);
 			PermutationVector.Set<FDDOFScatterOcclusionDim>(ScatterOcclusionTexture ? true : false);
+			PermutationVector.Set<FDiaphragmDOFHybridScatterPS::FDDOFScatterQualityDim>(Quality);
 			PermutationVector = FDiaphragmDOFHybridScatterPS::RemapPermutation(PermutationVector);
 
 			TShaderMapRef<FDiaphragmDOFHybridScatterVS> VertexShader(View.ShaderMap);
@@ -2517,6 +2563,8 @@ bool DiaphragmDOF::AddPasses(
 			PassParameters->ScatteringScaling = float(GatheringViewSize.X) / float(PreprocessViewSize.X);
 			PassParameters->CocSqueeze = CocModel.Squeeze;
 			PassParameters->CocInvSqueeze = 1.0 / CocModel.Squeeze;
+			PassParameters->Petzval = CocModel.Petzval;
+			PassParameters->PetzvalFalloffPower = CocModel.PetzvalFalloffPower;
 			PassParameters->CommonParameters = CommonParameters;
 			PassParameters->BokehLUT = ScatteringBokehLUT;
 			PassParameters->ScatterOcclusionSize = FVector4f(RefBufferSize.X, RefBufferSize.Y, 1.0f / RefBufferSize.X, 1.0f / RefBufferSize.Y);
diff --git a/Engine/Source/Runtime/Renderer/Private/PostProcess/DiaphragmDOF.h b/Engine/Source/Runtime/Renderer/Private/PostProcess/DiaphragmDOF.h
--- a/Engine/Source/Runtime/Renderer/Private/PostProcess/DiaphragmDOF.h
+++ b/Engine/Source/Runtime/Renderer/Private/PostProcess/DiaphragmDOF.h
@@ -67,6 +67,10 @@ struct FPhysicalCocModel
 	// SqueezeFactor = VerticalFocalDistance / HorizontalFocalDistance
 	float Squeeze;
 
+	// Swirly bokeh Petzval lens effect
+	float Petzval;
+	float PetzvalFalloffPower;
+
 	// The maximum radius of depth blur.
 	float MaxDepthBlurRadius;
 	float DepthBlurExponent;
diff --git a/Engine/Source/Runtime/Renderer/Private/PostProcess/DiaphragmDOFUtils.cpp b/Engine/Source/Runtime/Renderer/Private/PostProcess/DiaphragmDOFUtils.cpp
--- a/Engine/Source/Runtime/Renderer/Private/PostProcess/DiaphragmDOFUtils.cpp
+++ b/Engine/Source/Runtime/Renderer/Private/PostProcess/DiaphragmDOFUtils.cpp
@@ -108,6 +108,8 @@ void DiaphragmDOF::FPhysicalCocModel::Compile(const FViewInfo& View)
 		FocusDistance = View.FinalPostProcessSettings.DepthOfFieldFocalDistance;
 		FStops = View.FinalPostProcessSettings.DepthOfFieldFstop;
 		Squeeze = FMath::Clamp(View.FinalPostProcessSettings.DepthOfFieldSqueezeFactor, 0.1f, 10.0f);
+		Petzval = View.FinalPostProcessSettings.DepthOfFieldPetzvalBokeh;
+		PetzvalFalloffPower = FMath::Clamp(View.FinalPostProcessSettings.DepthOfFieldPetzvalBokehFalloff, 0, 100.0);
 
 		RenderingAspectRatio = float(View.UnscaledViewRect.Width()) / float(View.UnscaledViewRect.Height());
 		const float HorizontalHalfFOV = FMath::Atan(1.0f / View.ViewMatrices.GetProjectionMatrix().M[0][0]);
