[HWRT] Support secondary ray traversal visualization

diff --git a/Engine/Shaders/Private/RayTracing/RayTracingDebugTraversal.usf b/Engine/Shaders/Private/RayTracing/RayTracingDebugTraversal.usf
--- a/Engine/Shaders/Private/RayTracing/RayTracingDebugTraversal.usf
+++ b/Engine/Shaders/Private/RayTracing/RayTracingDebugTraversal.usf
@@ -3,10 +3,15 @@
 #include "../Common.ush"
 #include "../ShadingCommon.ush"
 #include "../ColorMap.ush"
+#include "../BlueNoise.ush"
 
 #include "RayTracingCommon.ush"
 #include "RayTracingDeferredMaterials.ush"
 
+#include "RayTracingDeferredShadingCommon.ush"
+#include "../Substrate/Substrate.ush"
+#include "../Lumen/LumenPosition.ush"
+
 #include "/Engine/Shared/RayTracingDebugDefinitions.h"
 
 #include "TraceRayInline.ush"
@@ -19,6 +24,9 @@ float TraversalBoxScale;
 float TraversalClusterScale;
 float TraversalTriangleScale;
 
+float NormalBias;
+float TraceDistance;
+
 struct FNaniteTraceRayInlineCallback : FDefaultTraceRayInlineCallback
 {
 	uint ClusterIntersectionCount;
@@ -38,6 +46,43 @@ struct FNaniteTraceRayInlineCallback : FDefaultTraceRayInlineCallback
 	}
 };
 
+FRayDesc CreateSecondaryRay(uint2 PixelCoord, float2 ScreenUV, out bool bIsValid)
+{
+#if SUBSTRATE_ENABLED
+	const FSubstrateTopLayerData TopLayerData = SubstrateUnpackTopLayerData(Substrate.TopLayerTexture.Load(uint3(PixelCoord, 0)));
+	bIsValid = IsSubstrateMaterial(TopLayerData);
+	const float3 WorldNormal = TopLayerData.WorldNormal;
+	const float Depth = ConvertFromDeviceZ(SceneDepthTexture.Load(int3(PixelCoord, 0)).r);
+#else
+	const FGBufferData GBufferData = GetGBufferDataFromSceneTexturesLoad(PixelCoord);
+	bIsValid = GBufferData.ShadingModelID != SHADINGMODELID_UNLIT;
+	const float3 WorldNormal = GBufferData.WorldNormal;
+	const float Depth = GBufferData.Depth;
+#endif
+
+	FRayDesc Ray = (FRayDesc)0;
+
+	if (bIsValid)
+	{
+		float3 TranslatedWorldPosition = GetTranslatedWorldPositionFromScreenUV(ScreenUV, Depth);
+		float3x3 TangentBasis = GetTangentBasis(WorldNormal);
+
+		float2 UniformRandom = BlueNoiseVec2(PixelCoord, View.StateFrameIndex);
+
+		float4 HemisphereSample = CosineSampleHemisphere(UniformRandom);
+		float3 RayDirection = mul(HemisphereSample.xyz, TangentBasis);
+		
+		Ray.Origin = TranslatedWorldPosition;
+		Ray.TMin = 0;
+		Ray.Direction = RayDirection;
+		Ray.TMax = TraceDistance;
+
+		ApplyCameraRelativeDepthBias(Ray, PixelCoord, ConvertToDeviceZ(Depth), WorldNormal, NormalBias);
+	}
+
+	return Ray;
+}
+
 [numthreads(INLINE_RAY_TRACING_THREAD_GROUP_SIZE_X, INLINE_RAY_TRACING_THREAD_GROUP_SIZE_Y, 1)]
 void RayTracingDebugTraversalCS(uint3 DispatchThread : SV_DispatchThreadID)
 {
@@ -46,6 +91,21 @@ void RayTracingDebugTraversalCS(uint3 DispatchThread : SV_DispatchThreadID)
 	float2 RenderTargetUV = (float2(PixelCoord) + .5f) * View.BufferSizeAndInvSize.zw;
 
 	FRayDesc Ray = CreatePrimaryRay(RenderTargetUV);
+	bool bIsValid = true;
+
+	if (VisualizationMode == RAY_TRACING_DEBUG_VIZ_TRAVERSAL_SECONDARY_ALL
+		|| VisualizationMode == RAY_TRACING_DEBUG_VIZ_TRAVERSAL_SECONDARY_CLUSTER
+		|| VisualizationMode == RAY_TRACING_DEBUG_VIZ_TRAVERSAL_SECONDARY_NODE
+		|| VisualizationMode == RAY_TRACING_DEBUG_VIZ_TRAVERSAL_SECONDARY_TRIANGLE)
+	{
+		Ray = CreateSecondaryRay(PixelCoord, RenderTargetUV, bIsValid);
+	}
+
+	if (!bIsValid)
+	{
+		Output[PixelCoord] = 0.0f;
+		return;
+	}
 
 	uint RayFlags = RAY_FLAG_CULL_BACK_FACING_TRIANGLES | RAY_FLAG_FORCE_OPAQUE;
 
@@ -62,16 +122,20 @@ void RayTracingDebugTraversalCS(uint3 DispatchThread : SV_DispatchThreadID)
 	switch (VisualizationMode)
 	{
 	default:
-	case RAY_TRACING_DEBUG_VIZ_TRAVERSAL_ALL:
+	case RAY_TRACING_DEBUG_VIZ_TRAVERSAL_PRIMARY_ALL:
+	case RAY_TRACING_DEBUG_VIZ_TRAVERSAL_SECONDARY_ALL:
 		Result.rgb = ColorMapTurbo(Intersections / (TraversalBoxScale + TraversalTriangleScale));
 		break;
-	case RAY_TRACING_DEBUG_VIZ_TRAVERSAL_CLUSTER:
+	case RAY_TRACING_DEBUG_VIZ_TRAVERSAL_PRIMARY_CLUSTER:
+	case RAY_TRACING_DEBUG_VIZ_TRAVERSAL_SECONDARY_CLUSTER:
 		Result.rgb = ColorMapTurbo(NaniteCallback.ClusterIntersectionCount / TraversalClusterScale);
 		break;
-	case RAY_TRACING_DEBUG_VIZ_TRAVERSAL_NODE:
+	case RAY_TRACING_DEBUG_VIZ_TRAVERSAL_PRIMARY_NODE:
+	case RAY_TRACING_DEBUG_VIZ_TRAVERSAL_SECONDARY_NODE:
 		Result.rgb = ColorMapTurbo(TraceResult.TraversalStatistics.NodeIntersectionCount / TraversalBoxScale);
 		break;
-	case RAY_TRACING_DEBUG_VIZ_TRAVERSAL_TRIANGLE:
+	case RAY_TRACING_DEBUG_VIZ_TRAVERSAL_PRIMARY_TRIANGLE:
+	case RAY_TRACING_DEBUG_VIZ_TRAVERSAL_SECONDARY_TRIANGLE:
 		Result.rgb = ColorMapTurbo(TraceResult.TraversalStatistics.TriangleIntersectionCount / TraversalTriangleScale);
 		break;
 	}
diff --git a/Engine/Shaders/Shared/RaytracingDebugDefinitions.h b/Engine/Shaders/Shared/RaytracingDebugDefinitions.h
--- a/Engine/Shaders/Shared/RaytracingDebugDefinitions.h
+++ b/Engine/Shaders/Shared/RaytracingDebugDefinitions.h
@@ -7,47 +7,52 @@
 
 #pragma once
 
-#define RAY_TRACING_DEBUG_VIZ_RADIANCE					0 
-#define RAY_TRACING_DEBUG_VIZ_WORLD_NORMAL				1 
-#define RAY_TRACING_DEBUG_VIZ_BASE_COLOR				2 
-#define RAY_TRACING_DEBUG_VIZ_DIFFUSE_COLOR				3 
-#define RAY_TRACING_DEBUG_VIZ_SPECULAR_COLOR			4 
-#define RAY_TRACING_DEBUG_VIZ_OPACITY					5 
-#define RAY_TRACING_DEBUG_VIZ_METALLIC					6 
-#define RAY_TRACING_DEBUG_VIZ_SPECULAR					7 
-#define RAY_TRACING_DEBUG_VIZ_ROUGHNESS					8 
-#define RAY_TRACING_DEBUG_VIZ_IOR						9 
-#define RAY_TRACING_DEBUG_VIZ_SHADING_MODEL				10 
-#define RAY_TRACING_DEBUG_VIZ_BLENDING_MODE				11  
-#define RAY_TRACING_DEBUG_VIZ_LIGHTING_CHANNEL_MASK		12  
-#define RAY_TRACING_DEBUG_VIZ_CUSTOM_DATA				13  
-#define RAY_TRACING_DEBUG_VIZ_GBUFFER_AO				14  
-#define RAY_TRACING_DEBUG_VIZ_INDIRECT_IRRADIANCE		15  
-#define RAY_TRACING_DEBUG_VIZ_WORLD_POSITION			16  
-#define RAY_TRACING_DEBUG_VIZ_HITKIND					17  
-#define RAY_TRACING_DEBUG_VIZ_BARYCENTRICS				18  
-#define RAY_TRACING_DEBUG_VIZ_PRIMARY_RAYS				19  
-#define RAY_TRACING_DEBUG_VIZ_WORLD_TANGENT				20 
-#define RAY_TRACING_DEBUG_VIZ_ANISOTROPY				21 
-#define RAY_TRACING_DEBUG_VIZ_INSTANCES					22
-#define RAY_TRACING_DEBUG_VIZ_PERFORMANCE				23 
-#define RAY_TRACING_DEBUG_VIZ_TRIANGLES					24
-#define RAY_TRACING_DEBUG_VIZ_FAR_FIELD					25
-#define RAY_TRACING_DEBUG_VIZ_TRAVERSAL_NODE			26
-#define RAY_TRACING_DEBUG_VIZ_TRAVERSAL_CLUSTER			27
-#define RAY_TRACING_DEBUG_VIZ_TRAVERSAL_TRIANGLE		28
-#define RAY_TRACING_DEBUG_VIZ_TRAVERSAL_ALL				29
-#define RAY_TRACING_DEBUG_VIZ_TRAVERSAL_STATISTICS		30
-#define RAY_TRACING_DEBUG_VIZ_DYNAMIC_INSTANCES			31
-#define RAY_TRACING_DEBUG_VIZ_PROXY_TYPE				32
-#define RAY_TRACING_DEBUG_VIZ_PICKER					33
-#define RAY_TRACING_DEBUG_VIZ_INSTANCE_OVERLAP			34
-#define RAY_TRACING_DEBUG_VIZ_TRIANGLE_HITCOUNT			35
-#define RAY_TRACING_DEBUG_VIZ_HITCOUNT_PER_INSTANCE		36
-#define RAY_TRACING_DEBUG_VIZ_LIGHT_GRID_COUNT			37
-#define RAY_TRACING_DEBUG_VIZ_SUBSTRATE_DATA			38
+#define RAY_TRACING_DEBUG_VIZ_RADIANCE							0 
+#define RAY_TRACING_DEBUG_VIZ_WORLD_NORMAL						1 
+#define RAY_TRACING_DEBUG_VIZ_BASE_COLOR						2 
+#define RAY_TRACING_DEBUG_VIZ_DIFFUSE_COLOR						3 
+#define RAY_TRACING_DEBUG_VIZ_SPECULAR_COLOR					4 
+#define RAY_TRACING_DEBUG_VIZ_OPACITY							5 
+#define RAY_TRACING_DEBUG_VIZ_METALLIC							6 
+#define RAY_TRACING_DEBUG_VIZ_SPECULAR							7 
+#define RAY_TRACING_DEBUG_VIZ_ROUGHNESS							8 
+#define RAY_TRACING_DEBUG_VIZ_IOR								9 
+#define RAY_TRACING_DEBUG_VIZ_SHADING_MODEL						10 
+#define RAY_TRACING_DEBUG_VIZ_BLENDING_MODE						11  
+#define RAY_TRACING_DEBUG_VIZ_LIGHTING_CHANNEL_MASK				12  
+#define RAY_TRACING_DEBUG_VIZ_CUSTOM_DATA						13  
+#define RAY_TRACING_DEBUG_VIZ_GBUFFER_AO						14  
+#define RAY_TRACING_DEBUG_VIZ_INDIRECT_IRRADIANCE				15  
+#define RAY_TRACING_DEBUG_VIZ_WORLD_POSITION					16  
+#define RAY_TRACING_DEBUG_VIZ_HITKIND							17  
+#define RAY_TRACING_DEBUG_VIZ_BARYCENTRICS						18  
+#define RAY_TRACING_DEBUG_VIZ_PRIMARY_RAYS						19  
+#define RAY_TRACING_DEBUG_VIZ_WORLD_TANGENT						20 
+#define RAY_TRACING_DEBUG_VIZ_ANISOTROPY						21 
+#define RAY_TRACING_DEBUG_VIZ_INSTANCES							22
+#define RAY_TRACING_DEBUG_VIZ_PERFORMANCE						23 
+#define RAY_TRACING_DEBUG_VIZ_TRIANGLES							24
+#define RAY_TRACING_DEBUG_VIZ_FAR_FIELD							25
+#define RAY_TRACING_DEBUG_VIZ_TRAVERSAL_PRIMARY_NODE			26
+#define RAY_TRACING_DEBUG_VIZ_TRAVERSAL_PRIMARY_CLUSTER			27
+#define RAY_TRACING_DEBUG_VIZ_TRAVERSAL_PRIMARY_TRIANGLE		28
+#define RAY_TRACING_DEBUG_VIZ_TRAVERSAL_PRIMARY_ALL				29
+#define RAY_TRACING_DEBUG_VIZ_TRAVERSAL_PRIMARY_STATISTICS		30
+#define RAY_TRACING_DEBUG_VIZ_TRAVERSAL_SECONDARY_NODE			31
+#define RAY_TRACING_DEBUG_VIZ_TRAVERSAL_SECONDARY_CLUSTER		32
+#define RAY_TRACING_DEBUG_VIZ_TRAVERSAL_SECONDARY_TRIANGLE		33
+#define RAY_TRACING_DEBUG_VIZ_TRAVERSAL_SECONDARY_ALL			34
+#define RAY_TRACING_DEBUG_VIZ_TRAVERSAL_SECONDARY_STATISTICS	35
+#define RAY_TRACING_DEBUG_VIZ_DYNAMIC_INSTANCES					36
+#define RAY_TRACING_DEBUG_VIZ_PROXY_TYPE						37
+#define RAY_TRACING_DEBUG_VIZ_PICKER							38
+#define RAY_TRACING_DEBUG_VIZ_INSTANCE_OVERLAP					39
+#define RAY_TRACING_DEBUG_VIZ_TRIANGLE_HITCOUNT					40
+#define RAY_TRACING_DEBUG_VIZ_HITCOUNT_PER_INSTANCE				41
+#define RAY_TRACING_DEBUG_VIZ_LIGHT_GRID_COUNT					42
+#define RAY_TRACING_DEBUG_VIZ_SUBSTRATE_DATA					43
 
-#define RAY_TRACING_DEBUG_PICKER_DOMAIN_TRIANGLE		0
-#define RAY_TRACING_DEBUG_PICKER_DOMAIN_INSTANCE		1
-#define RAY_TRACING_DEBUG_PICKER_DOMAIN_SEGMENT			2
+#define RAY_TRACING_DEBUG_PICKER_DOMAIN_TRIANGLE				0
+#define RAY_TRACING_DEBUG_PICKER_DOMAIN_INSTANCE				1
+#define RAY_TRACING_DEBUG_PICKER_DOMAIN_SEGMENT					2
 
diff --git a/Engine/Source/Runtime/Engine/Private/RayTracingVisualizationData.cpp b/Engine/Source/Runtime/Engine/Private/RayTracingVisualizationData.cpp
--- a/Engine/Source/Runtime/Engine/Private/RayTracingVisualizationData.cpp
+++ b/Engine/Source/Runtime/Engine/Private/RayTracingVisualizationData.cpp
@@ -67,11 +67,17 @@ void FRayTracingVisualizationData::Initialize()
 				AddVisualizationMode(TEXT("Performance"), LOCTEXT("Performance", "Performance"), FModeType::Performance, RAY_TRACING_DEBUG_VIZ_PERFORMANCE, false);
 			}
 
-			AddVisualizationMode(TEXT("Traversal Node"), LOCTEXT("TraversalNode", "Traversal Node"), FModeType::Performance, RAY_TRACING_DEBUG_VIZ_TRAVERSAL_NODE, false, /*bHiddenInEditor*/ true);
-			AddVisualizationMode(TEXT("Traversal Cluster"), LOCTEXT("TraversalCluster", "Traversal Cluster"), FModeType::Performance, RAY_TRACING_DEBUG_VIZ_TRAVERSAL_CLUSTER, false, /*bHiddenInEditor*/ true);
-			AddVisualizationMode(TEXT("Traversal Triangle"), LOCTEXT("TraversalTriangle", "Traversal Triangle"), FModeType::Performance, RAY_TRACING_DEBUG_VIZ_TRAVERSAL_TRIANGLE, false, /*bHiddenInEditor*/ true);
-			AddVisualizationMode(TEXT("Traversal All"), LOCTEXT("TraversalAll", "Traversal All"), FModeType::Performance, RAY_TRACING_DEBUG_VIZ_TRAVERSAL_ALL, false, /*bHiddenInEditor*/ true);
-			AddVisualizationMode(TEXT("Traversal Statistics"), LOCTEXT("TraversalStatistics", "Traversal Statistics"), FModeType::Performance, RAY_TRACING_DEBUG_VIZ_TRAVERSAL_STATISTICS, false, /*bHiddenInEditor*/ true);
+			AddVisualizationMode(TEXT("Traversal Primary Node"), LOCTEXT("TraversalPrimaryNode", "Traversal Primary Node"), FModeType::Performance, RAY_TRACING_DEBUG_VIZ_TRAVERSAL_PRIMARY_NODE, false, /*bHiddenInEditor*/ true);
+			AddVisualizationMode(TEXT("Traversal Primary Cluster"), LOCTEXT("TraversalPrimaryCluster", "Traversal Primary Cluster"), FModeType::Performance, RAY_TRACING_DEBUG_VIZ_TRAVERSAL_PRIMARY_CLUSTER, false, /*bHiddenInEditor*/ true);
+			AddVisualizationMode(TEXT("Traversal Primary Triangle"), LOCTEXT("TraversalPrimaryTriangle", "Traversal Primary Triangle"), FModeType::Performance, RAY_TRACING_DEBUG_VIZ_TRAVERSAL_PRIMARY_TRIANGLE, false, /*bHiddenInEditor*/ true);
+			AddVisualizationMode(TEXT("Traversal Primary All"), LOCTEXT("TraversalPrimaryAll", "Traversal Primary All"), FModeType::Performance, RAY_TRACING_DEBUG_VIZ_TRAVERSAL_PRIMARY_ALL, false, /*bHiddenInEditor*/ true);
+			AddVisualizationMode(TEXT("Traversal Primary Statistics"), LOCTEXT("TraversalPrimaryStatistics", "Traversal Primary Statistics"), FModeType::Performance, RAY_TRACING_DEBUG_VIZ_TRAVERSAL_PRIMARY_STATISTICS, false, /*bHiddenInEditor*/ true);
+
+			AddVisualizationMode(TEXT("Traversal Secondary Node"), LOCTEXT("TraversalSecondaryNode", "Traversal Secondary Node"), FModeType::Performance, RAY_TRACING_DEBUG_VIZ_TRAVERSAL_SECONDARY_NODE, false, /*bHiddenInEditor*/ true);
+			AddVisualizationMode(TEXT("Traversal Secondary Cluster"), LOCTEXT("TraversalSecondaryCluster", "Traversal Secondary Cluster"), FModeType::Performance, RAY_TRACING_DEBUG_VIZ_TRAVERSAL_SECONDARY_CLUSTER, false, /*bHiddenInEditor*/ true);
+			AddVisualizationMode(TEXT("Traversal Secondary Triangle"), LOCTEXT("TraversalSecondaryTriangle", "Traversal Secondary Triangle"), FModeType::Performance, RAY_TRACING_DEBUG_VIZ_TRAVERSAL_SECONDARY_TRIANGLE, false, /*bHiddenInEditor*/ true);
+			AddVisualizationMode(TEXT("Traversal Secondary All"), LOCTEXT("TraversalSecondaryAll", "Traversal Secondary All"), FModeType::Performance, RAY_TRACING_DEBUG_VIZ_TRAVERSAL_SECONDARY_ALL, false, /*bHiddenInEditor*/ true);
+			AddVisualizationMode(TEXT("Traversal Secondary Statistics"), LOCTEXT("TraversalSecondaryStatistics", "Traversal Secondary Statistics"), FModeType::Performance, RAY_TRACING_DEBUG_VIZ_TRAVERSAL_SECONDARY_STATISTICS, false, /*bHiddenInEditor*/ true);
 
 			if (Substrate::IsSubstrateEnabled())
 			{
diff --git a/Engine/Source/Runtime/Renderer/Private/DeferredShadingRenderer.cpp b/Engine/Source/Runtime/Renderer/Private/DeferredShadingRenderer.cpp
--- a/Engine/Source/Runtime/Renderer/Private/DeferredShadingRenderer.cpp
+++ b/Engine/Source/Runtime/Renderer/Private/DeferredShadingRenderer.cpp
@@ -3260,7 +3260,7 @@ void FDeferredShadingSceneRenderer::Render(FRDGBuilder& GraphBuilder)
 				for (const FViewInfo& View : Views)
 				{
 					FRayTracingPickingFeedback PickingFeedback = {};
-					RenderRayTracingDebug(GraphBuilder, View, SceneTextures.Color.Target, PickingFeedback);
+					RenderRayTracingDebug(GraphBuilder, View, SceneTextures, PickingFeedback);
 
 					OnGetOnScreenMessages.AddLambda([this, PickingFeedback](FScreenMessageWriter& ScreenMessageWriter)->void
 						{
diff --git a/Engine/Source/Runtime/Renderer/Private/DeferredShadingRenderer.h b/Engine/Source/Runtime/Renderer/Private/DeferredShadingRenderer.h
--- a/Engine/Source/Runtime/Renderer/Private/DeferredShadingRenderer.h
+++ b/Engine/Source/Runtime/Renderer/Private/DeferredShadingRenderer.h
@@ -1125,7 +1125,7 @@ private:
 	void SetupRayTracingRenderingData(FRDGBuilder& GraphBuilder);
 
 	/** Debug ray tracing functions. */
-	void RenderRayTracingDebug(FRDGBuilder& GraphBuilder, const FViewInfo& View, FRDGTextureRef SceneColorOutputTexture, FRayTracingPickingFeedback& PickingFeedback);
+	void RenderRayTracingDebug(FRDGBuilder& GraphBuilder, const FViewInfo& View, FSceneTextures& SceneTextures, FRayTracingPickingFeedback& PickingFeedback);
 	void RenderRayTracingBarycentrics(FRDGBuilder& GraphBuilder, const FViewInfo& View, FRDGTextureRef SceneColorOutputTexture, bool bVisualizeProceduralPrimitives);
 	void RayTracingDisplayPicking(const FRayTracingPickingFeedback& PickingFeedback, FScreenMessageWriter& Writer);
 
diff --git a/Engine/Source/Runtime/Renderer/Private/RayTracing/RayTracingDebug.cpp b/Engine/Source/Runtime/Renderer/Private/RayTracing/RayTracingDebug.cpp
--- a/Engine/Source/Runtime/Renderer/Private/RayTracing/RayTracingDebug.cpp
+++ b/Engine/Source/Runtime/Renderer/Private/RayTracing/RayTracingDebug.cpp
@@ -23,6 +23,7 @@
 #include "Nanite/NaniteRayTracing.h"
 #include "PixelShaderUtils.h"
 #include "SystemTextures.h"
+#include "BlueNoise.h"
 
 #define LOCTEXT_NAMESPACE "RayTracingDebugVisualizationMenuCommands"
 
@@ -373,6 +374,14 @@ class FRayTracingDebugTraversalCS : public FGlobalShader
 		SHADER_PARAMETER(float, TraversalTriangleScale)
 
 		SHADER_PARAMETER(float, RTDebugVisualizationNaniteCutError)
+
+		SHADER_PARAMETER_STRUCT_REF(FBlueNoise, BlueNoise)
+
+		SHADER_PARAMETER_STRUCT_INCLUDE(FSceneTextureParameters, SceneTextures)
+		SHADER_PARAMETER_RDG_UNIFORM_BUFFER(FSubstrateGlobalUniformParameters, Substrate)
+
+		SHADER_PARAMETER(float, NormalBias)
+		SHADER_PARAMETER(float, TraceDistance)
 	END_SHADER_PARAMETER_STRUCT()
 
 	class FSupportProceduralPrimitive : SHADER_PERMUTATION_BOOL("ENABLE_TRACE_RAY_INLINE_PROCEDURAL_PRIMITIVE");
@@ -802,23 +811,28 @@ void BindRayTracingDebugCHSMaterialBindings(FRHICommandList& RHICmdList, FRHISha
 
 static bool RequiresRayTracingDebugCHS(uint32 DebugVisualizationMode)
 {
-	return DebugVisualizationMode == RAY_TRACING_DEBUG_VIZ_INSTANCES ||
-		DebugVisualizationMode == RAY_TRACING_DEBUG_VIZ_TRIANGLES ||
-		DebugVisualizationMode == RAY_TRACING_DEBUG_VIZ_DYNAMIC_INSTANCES ||
-		DebugVisualizationMode == RAY_TRACING_DEBUG_VIZ_PROXY_TYPE ||
-		DebugVisualizationMode == RAY_TRACING_DEBUG_VIZ_PICKER ||
-		DebugVisualizationMode == RAY_TRACING_DEBUG_VIZ_INSTANCE_OVERLAP ||
-		DebugVisualizationMode == RAY_TRACING_DEBUG_VIZ_TRIANGLE_HITCOUNT ||
-		DebugVisualizationMode == RAY_TRACING_DEBUG_VIZ_HITCOUNT_PER_INSTANCE;
+	return DebugVisualizationMode == RAY_TRACING_DEBUG_VIZ_INSTANCES
+		|| DebugVisualizationMode == RAY_TRACING_DEBUG_VIZ_TRIANGLES
+		|| DebugVisualizationMode == RAY_TRACING_DEBUG_VIZ_DYNAMIC_INSTANCES
+		|| DebugVisualizationMode == RAY_TRACING_DEBUG_VIZ_PROXY_TYPE
+		|| DebugVisualizationMode == RAY_TRACING_DEBUG_VIZ_PICKER
+		|| DebugVisualizationMode == RAY_TRACING_DEBUG_VIZ_INSTANCE_OVERLAP
+		|| DebugVisualizationMode == RAY_TRACING_DEBUG_VIZ_TRIANGLE_HITCOUNT
+		|| DebugVisualizationMode == RAY_TRACING_DEBUG_VIZ_HITCOUNT_PER_INSTANCE;
 }
 
 static bool IsRayTracingDebugTraversalMode(uint32 DebugVisualizationMode)
 {
-	return DebugVisualizationMode == RAY_TRACING_DEBUG_VIZ_TRAVERSAL_NODE ||
-		DebugVisualizationMode == RAY_TRACING_DEBUG_VIZ_TRAVERSAL_CLUSTER ||
-		DebugVisualizationMode == RAY_TRACING_DEBUG_VIZ_TRAVERSAL_TRIANGLE ||
-		DebugVisualizationMode == RAY_TRACING_DEBUG_VIZ_TRAVERSAL_ALL ||
-		DebugVisualizationMode == RAY_TRACING_DEBUG_VIZ_TRAVERSAL_STATISTICS;
+	return DebugVisualizationMode == RAY_TRACING_DEBUG_VIZ_TRAVERSAL_PRIMARY_NODE 
+		|| DebugVisualizationMode == RAY_TRACING_DEBUG_VIZ_TRAVERSAL_PRIMARY_CLUSTER
+		|| DebugVisualizationMode == RAY_TRACING_DEBUG_VIZ_TRAVERSAL_PRIMARY_TRIANGLE
+		|| DebugVisualizationMode == RAY_TRACING_DEBUG_VIZ_TRAVERSAL_PRIMARY_ALL
+		|| DebugVisualizationMode == RAY_TRACING_DEBUG_VIZ_TRAVERSAL_PRIMARY_STATISTICS
+		|| DebugVisualizationMode == RAY_TRACING_DEBUG_VIZ_TRAVERSAL_SECONDARY_NODE
+		|| DebugVisualizationMode == RAY_TRACING_DEBUG_VIZ_TRAVERSAL_SECONDARY_CLUSTER
+		|| DebugVisualizationMode == RAY_TRACING_DEBUG_VIZ_TRAVERSAL_SECONDARY_TRIANGLE
+		|| DebugVisualizationMode == RAY_TRACING_DEBUG_VIZ_TRAVERSAL_SECONDARY_ALL
+		|| DebugVisualizationMode == RAY_TRACING_DEBUG_VIZ_TRAVERSAL_SECONDARY_STATISTICS;
 }
 
 static bool IsRayTracingPickingEnabled(uint32 DebugVisualizationMode)
@@ -1387,19 +1401,30 @@ static uint32 GetRaytracingDebugViewMode(const FSceneView& View)
 bool HasRaytracingDebugViewModeRaytracedOverlay(const FSceneViewFamily& ViewFamily)
 {
 	bool bAnySubstrate = false;
+	bool bAnyTraversalSecondary = false;
 
 	for (const FSceneView* View : ViewFamily.Views)
 	{
 		const uint32 Mode = View != nullptr ? GetRaytracingDebugViewMode(*View) : UINT32_MAX;
 
 		bAnySubstrate |= Mode == RAY_TRACING_DEBUG_VIZ_SUBSTRATE_DATA;
+
+		// can't get WorldNormal in inline ray tracing so need GBuffer Depth/WorldNormal rendered by raster passes to generate secondary rays
+		bAnyTraversalSecondary |= Mode == RAY_TRACING_DEBUG_VIZ_TRAVERSAL_SECONDARY_NODE;
+		bAnyTraversalSecondary |= Mode == RAY_TRACING_DEBUG_VIZ_TRAVERSAL_SECONDARY_TRIANGLE;
+		bAnyTraversalSecondary |= Mode == RAY_TRACING_DEBUG_VIZ_TRAVERSAL_SECONDARY_ALL;
+		bAnyTraversalSecondary |= Mode == RAY_TRACING_DEBUG_VIZ_TRAVERSAL_SECONDARY_CLUSTER;
+		bAnyTraversalSecondary |= Mode == RAY_TRACING_DEBUG_VIZ_TRAVERSAL_SECONDARY_STATISTICS;
 	}
 
-	return !bAnySubstrate;
+	return !bAnySubstrate && !bAnyTraversalSecondary;
 }
 
-void FDeferredShadingSceneRenderer::RenderRayTracingDebug(FRDGBuilder& GraphBuilder, const FViewInfo& View, FRDGTextureRef SceneColorTexture, FRayTracingPickingFeedback& PickingFeedback)
+void FDeferredShadingSceneRenderer::RenderRayTracingDebug(FRDGBuilder& GraphBuilder, const FViewInfo& View, FSceneTextures& SceneTextures, FRayTracingPickingFeedback& PickingFeedback)
 {
+	FSceneTextureParameters SceneTextureParameters = GetSceneTextureParameters(GraphBuilder, SceneTextures.UniformBuffer);
+	FRDGTextureRef SceneColorTexture = SceneTextures.Color.Target;
+
 	const uint32 DebugVisualizationMode = GetRaytracingDebugViewMode(View);
 
 	if (DebugVisualizationMode == RAY_TRACING_DEBUG_VIZ_SUBSTRATE_DATA)
@@ -1419,7 +1444,7 @@ void FDeferredShadingSceneRenderer::RenderRayTracingDebug(FRDGBuilder& GraphBuil
 	if (IsRayTracingDebugTraversalMode(DebugVisualizationMode) && ShouldRenderRayTracingEffect(true, ERayTracingPipelineCompatibilityFlags::Inline, View))
 	{
 		const bool bPrintTraversalStats = FDataDrivenShaderPlatformInfo::GetSupportsRayTracingTraversalStatistics(GMaxRHIShaderPlatform)
-			&& (DebugVisualizationMode == RAY_TRACING_DEBUG_VIZ_TRAVERSAL_STATISTICS);
+			&& (DebugVisualizationMode == RAY_TRACING_DEBUG_VIZ_TRAVERSAL_PRIMARY_STATISTICS || DebugVisualizationMode == RAY_TRACING_DEBUG_VIZ_TRAVERSAL_SECONDARY_STATISTICS);
 
 		FRayTracingDebugTraversalCS::FParameters* PassParameters = GraphBuilder.AllocParameters<FRayTracingDebugTraversalCS::FParameters>();
 		PassParameters->Output = GraphBuilder.CreateUAV(SceneColorTexture);
@@ -1436,6 +1461,14 @@ void FDeferredShadingSceneRenderer::RenderRayTracingDebug(FRDGBuilder& GraphBuil
 
 		PassParameters->RTDebugVisualizationNaniteCutError = 0.0f;
 
+		FBlueNoise BlueNoise = GetBlueNoiseGlobalParameters();
+		PassParameters->BlueNoise = CreateUniformBufferImmediate(BlueNoise, EUniformBufferUsage::UniformBuffer_SingleDraw);
+
+		PassParameters->SceneTextures = SceneTextureParameters;
+		PassParameters->Substrate = Substrate::BindSubstrateGlobalUniformParameters(View);
+		PassParameters->NormalBias = GetRaytracingMaxNormalBias();
+		PassParameters->TraceDistance = 20000.0f;
+
 		RaytracingTraversalStatistics::FTraceRayInlineStatisticsData TraversalData;
 		if (bPrintTraversalStats)
 		{
