Fix RHI validation with RayTracing dispatches on async compute

diff --git a/Engine/Plugins/Runtime/HairStrands/Source/HairStrandsCore/Private/GroomManager.cpp b/Engine/Plugins/Runtime/HairStrands/Source/HairStrandsCore/Private/GroomManager.cpp
--- a/Engine/Plugins/Runtime/HairStrands/Source/HairStrandsCore/Private/GroomManager.cpp
+++ b/Engine/Plugins/Runtime/HairStrands/Source/HairStrandsCore/Private/GroomManager.cpp
@@ -2034,6 +2034,14 @@ static void RunHairStrandsInterpolation_Cards(
 		}
 	}
 
+	ERHIPipeline CardPositionExternalAccessPipeline = ERHIPipeline::Graphics;
+#if RHI_RAYTRACING
+	// Also add async compute as queue which can access these external hair resources
+	if (GSupportsEfficientAsyncCompute && GRHIGlobals.RayTracing.SupportsAsyncRayTraceDispatch)
+	{
+		EnumAddFlags(CardPositionExternalAccessPipeline, ERHIPipeline::AsyncCompute);
+	}
+#endif
 	FRDGExternalAccessQueue ExternalAccessQueue;
 
 	// Cards only - Deform guide with skinning and/or simulation
@@ -2106,7 +2114,8 @@ static void RunHairStrandsInterpolation_Cards(
 					ShaderPrintData,
 					InstanceData.Instance,
 					InstanceData.HairLODIndex,
-					InstanceData.MeshLODIndex);
+					InstanceData.MeshLODIndex,
+					CardPositionExternalAccessPipeline);
 			}
 		}
 	}
@@ -2156,7 +2165,7 @@ static void RunHairStrandsInterpolation_Cards(
 
 			if (LOD.DeformedResource)
 			{
-				ExternalAccessQueue.Add(Register(GraphBuilder, LOD.DeformedResource->GetBuffer(FHairCardsDeformedResource::Current), ERDGImportedBufferFlags::None).Buffer, ERHIAccess::SRVMask);
+				ExternalAccessQueue.Add(Register(GraphBuilder, LOD.DeformedResource->GetBuffer(FHairCardsDeformedResource::Current), ERDGImportedBufferFlags::None).Buffer, ERHIAccess::SRVMask, CardPositionExternalAccessPipeline);
 				ExternalAccessQueue.Add(Register(GraphBuilder, LOD.DeformedResource->GetBuffer(FHairCardsDeformedResource::Previous), ERDGImportedBufferFlags::None).Buffer, ERHIAccess::SRVMask);
 			}
 		}
diff --git a/Engine/Plugins/Runtime/HairStrands/Source/HairStrandsCore/Private/HairStrandsInterpolation.cpp b/Engine/Plugins/Runtime/HairStrands/Source/HairStrandsCore/Private/HairStrandsInterpolation.cpp
--- a/Engine/Plugins/Runtime/HairStrands/Source/HairStrandsCore/Private/HairStrandsInterpolation.cpp
+++ b/Engine/Plugins/Runtime/HairStrands/Source/HairStrandsCore/Private/HairStrandsInterpolation.cpp
@@ -990,7 +990,8 @@ void AddHairCardsDeformationPass(
 	const FShaderPrintData* ShaderPrintData,
 	FHairGroupInstance* Instance,
 	const int32 HairLODIndex,
-	const int32 MeshLODIndex)
+	const int32 MeshLODIndex,
+	ERHIPipeline CardPositionExternalAccessPipeline)
 {
 	if (!Instance->Cards.IsValid(HairLODIndex))
 		return;
@@ -1096,8 +1097,9 @@ void AddHairCardsDeformationPass(
 		AddCopyBufferPass(GraphBuilder, CardsDeformedPositionBuffer_Prev.Buffer, CardsDeformedPositionBuffer_Curr.Buffer);
 	}
 
+	GraphBuilder.UseExternalAccessMode(CardsDeformedPositionBuffer_Curr.Buffer, ERHIAccess::SRVMask, CardPositionExternalAccessPipeline);
+
 	GraphBuilder.UseExternalAccessMode({
-		CardsDeformedPositionBuffer_Curr.Buffer,
 		CardsDeformedPositionBuffer_Prev.Buffer,
 		CardsDeformedNormalBuffer.Buffer
 	}, ERHIAccess::SRVMask);
diff --git a/Engine/Plugins/Runtime/HairStrands/Source/HairStrandsCore/Public/HairStrandsInterpolation.h b/Engine/Plugins/Runtime/HairStrands/Source/HairStrandsCore/Public/HairStrandsInterpolation.h
--- a/Engine/Plugins/Runtime/HairStrands/Source/HairStrandsCore/Public/HairStrandsInterpolation.h
+++ b/Engine/Plugins/Runtime/HairStrands/Source/HairStrandsCore/Public/HairStrandsInterpolation.h
@@ -240,4 +240,5 @@ void AddHairCardsDeformationPass(
 	const FShaderPrintData* ShaderPrintData,
 	FHairGroupInstance* Instance,
 	const int32 HairLODIndex,
-	const int32 MeshLODIndex);
\ No newline at end of file
+	const int32 MeshLODIndex,
+	ERHIPipeline CardPositionExternalAccessPipeline);
\ No newline at end of file
diff --git a/Engine/Source/Runtime/Renderer/Private/DeferredShadingRenderer.cpp b/Engine/Source/Runtime/Renderer/Private/DeferredShadingRenderer.cpp
--- a/Engine/Source/Runtime/Renderer/Private/DeferredShadingRenderer.cpp
+++ b/Engine/Source/Runtime/Renderer/Private/DeferredShadingRenderer.cpp
@@ -785,7 +785,7 @@ void FDeferredShadingSceneRenderer::SetupRayTracingLightDataForViews(FRDGBuilder
 	}
 }
 
-bool FDeferredShadingSceneRenderer::DispatchRayTracingWorldUpdates(FRDGBuilder& GraphBuilder, FRDGBufferRef& OutDynamicGeometryScratchBuffer)
+bool FDeferredShadingSceneRenderer::DispatchRayTracingWorldUpdates(FRDGBuilder& GraphBuilder, FRDGBufferRef& OutDynamicGeometryScratchBuffer, ERHIPipeline ResourceAccessPipelines)
 {
 	OutDynamicGeometryScratchBuffer = nullptr;
 
@@ -846,7 +846,7 @@ bool FDeferredShadingSceneRenderer::DispatchRayTracingWorldUpdates(FRDGBuilder&
 	// Keep mask the same as what's already set (which will be the view mask) if TLAS updates should be masked to the view
 	RDG_GPU_MASK_SCOPE(GraphBuilder, GRayTracingMultiGpuTLASMask ? GraphBuilder.RHICmdList.GetGPUMask() : FRHIGPUMask::All());
 
-	Scene->GetRayTracingDynamicGeometryCollection()->AddDynamicGeometryUpdatePass(ReferenceView, GraphBuilder, ComputePassFlags, OutDynamicGeometryScratchBuffer);
+	Scene->GetRayTracingDynamicGeometryCollection()->AddDynamicGeometryUpdatePass(ReferenceView, GraphBuilder, ComputePassFlags, ResourceAccessPipelines, OutDynamicGeometryScratchBuffer);
 
 
 	{
@@ -2440,8 +2440,11 @@ void FDeferredShadingSceneRenderer::Render(FRDGBuilder& GraphBuilder)
 
 		FRDGBufferRef DynamicGeometryScratchBuffer = nullptr;
 #if RHI_RAYTRACING
+		
+		ERHIPipeline DynamicRTResourceAccessPipelines = Lumen::UseAsyncCompute(ViewFamily) ? ERHIPipeline::All : ERHIPipeline::Graphics;
+
 		// Async AS builds can potentially overlap with BasePass.
-		bool bNeedToSetupRayTracingRenderingData = DispatchRayTracingWorldUpdates(GraphBuilder, DynamicGeometryScratchBuffer);
+		bool bNeedToSetupRayTracingRenderingData = DispatchRayTracingWorldUpdates(GraphBuilder, DynamicGeometryScratchBuffer, DynamicRTResourceAccessPipelines);
 
 		/** Should be called somewhere before "SetupRayTracingRenderingData" */
 		SetupRayTracingLightDataForViews(GraphBuilder);
@@ -2613,7 +2616,7 @@ void FDeferredShadingSceneRenderer::Render(FRDGBuilder& GraphBuilder)
 		{
 #if RHI_RAYTRACING
 			// When Lumen HWRT is running async we need to wait for ray tracing scene before dispatching the work
-			if (bNeedToSetupRayTracingRenderingData && Lumen::UseAsyncCompute(ViewFamily) && Lumen::UseHardwareInlineRayTracing(ViewFamily))
+			if (bNeedToSetupRayTracingRenderingData && Lumen::UseAsyncCompute(ViewFamily))
 			{
 				SetupRayTracingRenderingData(GraphBuilder);
 				bNeedToSetupRayTracingRenderingData = false;
diff --git a/Engine/Source/Runtime/Renderer/Private/DeferredShadingRenderer.h b/Engine/Source/Runtime/Renderer/Private/DeferredShadingRenderer.h
--- a/Engine/Source/Runtime/Renderer/Private/DeferredShadingRenderer.h
+++ b/Engine/Source/Runtime/Renderer/Private/DeferredShadingRenderer.h
@@ -1131,7 +1131,7 @@ private:
 
 	bool SetupRayTracingPipelineStatesAndSBT(FRDGBuilder& GraphBuilder, bool bAnyLumenHardwareInlineRayTracingPassEnabled, bool& bOutIsUsingFallbackRTPSO);
 	void SetupRayTracingLightDataForViews(FRDGBuilder& GraphBuilder);
-	bool DispatchRayTracingWorldUpdates(FRDGBuilder& GraphBuilder, FRDGBufferRef& OutDynamicGeometryScratchBuffer);
+	bool DispatchRayTracingWorldUpdates(FRDGBuilder& GraphBuilder, FRDGBufferRef& OutDynamicGeometryScratchBuffer, ERHIPipeline ResourceAccessPipelines);
 
 	/** Functions to create ray tracing pipeline state objects for various effects */
 	void CreateMaterialRayTracingMaterialPipeline(FRDGBuilder& GraphBuilder, FViewInfo& View, const TArrayView<FRHIRayTracingShader*>& RayGenShaderTable, uint32& OutMaxLocalBindingDataSize, bool& bOutIsUsingFallbackRTPSO);
diff --git a/Engine/Source/Runtime/Renderer/Private/RayTracing/RayTracingDynamicGeometry.cpp b/Engine/Source/Runtime/Renderer/Private/RayTracing/RayTracingDynamicGeometry.cpp
--- a/Engine/Source/Runtime/Renderer/Private/RayTracing/RayTracingDynamicGeometry.cpp
+++ b/Engine/Source/Runtime/Renderer/Private/RayTracing/RayTracingDynamicGeometry.cpp
@@ -539,7 +539,7 @@ uint32 FRayTracingDynamicGeometryCollection::Update()
 	return BLASScratchSize;
 }
 
-void FRayTracingDynamicGeometryCollection::AddDynamicGeometryUpdatePass(const FViewInfo& View, FRDGBuilder& GraphBuilder, ERDGPassFlags ComputePassFlags, FRDGBufferRef& OutDynamicGeometryScratchBuffer)
+void FRayTracingDynamicGeometryCollection::AddDynamicGeometryUpdatePass(const FViewInfo& View, FRDGBuilder& GraphBuilder, ERDGPassFlags ComputePassFlags, ERHIPipeline ResourceAccessPipelines, FRDGBufferRef& OutDynamicGeometryScratchBuffer)
 {
 	RDG_GPU_MASK_SCOPE(GraphBuilder, FRHIGPUMask::All());
 	RDG_EVENT_SCOPE_STAT(GraphBuilder, RayTracingDynamicGeometry, "RayTracingDynamicGeometry");
@@ -563,19 +563,21 @@ void FRayTracingDynamicGeometryCollection::AddDynamicGeometryUpdatePass(const FV
 	PassParams->Scene = View.GetSceneUniforms().GetBuffer(GraphBuilder);
 	PassParams->DynamicGeometryScratchBuffer = OutDynamicGeometryScratchBuffer;	
 
+	ERHIPipeline SrcResourceAccessPipelines = ComputePassFlags == ERDGPassFlags::AsyncCompute ? ERHIPipeline::AsyncCompute : ERHIPipeline::Graphics;
+
 	GraphBuilder.AddPass(RDG_EVENT_NAME("RayTracingDynamicUpdate"), PassParams, ComputePassFlags | ERDGPassFlags::NeverCull,
-		[this, PassParams](FRDGAsyncTask, FRHICommandList& RHICmdList)
+		[this, SrcResourceAccessPipelines, ResourceAccessPipelines, PassParams](FRDGAsyncTask, FRHICommandList& RHICmdList)
 		{
 			FRHIBuffer* DynamicGeometryScratchBuffer = PassParams->DynamicGeometryScratchBuffer ? PassParams->DynamicGeometryScratchBuffer->GetRHI() : nullptr;
 
 			PRAGMA_DISABLE_DEPRECATION_WARNINGS
-			DispatchUpdates(RHICmdList, DynamicGeometryScratchBuffer);
+			DispatchUpdates(RHICmdList, DynamicGeometryScratchBuffer, SrcResourceAccessPipelines, ResourceAccessPipelines);
 			EndUpdate();
 			PRAGMA_ENABLE_DEPRECATION_WARNINGS
 		});
 }
 
-void FRayTracingDynamicGeometryCollection::DispatchUpdates(FRHICommandList& RHICmdList, FRHIBuffer* ScratchBuffer)
+void FRayTracingDynamicGeometryCollection::DispatchUpdates(FRHICommandList& RHICmdList, FRHIBuffer* ScratchBuffer, ERHIPipeline SrcResourceAccessPipelines, ERHIPipeline DstResourceAccessPipelines)
 {
 	if (DispatchCommands.Num() > 0)
 	{
@@ -682,7 +684,9 @@ void FRayTracingDynamicGeometryCollection::DispatchUpdates(FRHICommandList& RHIC
 
 			// Make sure buffers are readable again and disable UAV overlap.
 			RHICmdList.EndUAVOverlap(OverlapUAVs);
-			RHICmdList.Transition(TransitionsAfter);
+
+			// Transition to SRV state and mark readable on requested pipelines
+			RHICmdList.Transition(TransitionsAfter, SrcResourceAccessPipelines, DstResourceAccessPipelines);
 		}
 
 		if (BuildParams.Num() > 0)
diff --git a/Engine/Source/Runtime/Renderer/Public/RayTracingDynamicGeometryCollection.h b/Engine/Source/Runtime/Renderer/Public/RayTracingDynamicGeometryCollection.h
--- a/Engine/Source/Runtime/Renderer/Public/RayTracingDynamicGeometryCollection.h
+++ b/Engine/Source/Runtime/Renderer/Public/RayTracingDynamicGeometryCollection.h
@@ -47,12 +47,12 @@ public:
 	RENDERER_API int64 BeginUpdate();
 
 	UE_DEPRECATED(5.5, "Use AddDynamicGeometryUpdatePass instead.")
-	RENDERER_API void DispatchUpdates(FRHICommandList& RHICmdList, FRHIBuffer* ScratchBuffer);
+	RENDERER_API void DispatchUpdates(FRHICommandList& RHICmdList, FRHIBuffer* ScratchBuffer, ERHIPipeline SrcResourceAccessPipelines = ERHIPipeline::Graphics, ERHIPipeline DstResourceAccessPipelines = ERHIPipeline::Graphics);
 
 	UE_DEPRECATED(5.5, "Use AddDynamicGeometryUpdatePass instead.")
 	RENDERER_API void EndUpdate();
 
-	RENDERER_API void AddDynamicGeometryUpdatePass(const FViewInfo& View, FRDGBuilder& GraphBuilder, ERDGPassFlags ComputePassFlags, FRDGBufferRef& OutDynamicGeometryScratchBuffer);
+	RENDERER_API void AddDynamicGeometryUpdatePass(const FViewInfo& View, FRDGBuilder& GraphBuilder, ERDGPassFlags ComputePassFlags, ERHIPipeline ResourceAccessPipelines, FRDGBufferRef& OutDynamicGeometryScratchBuffer);
 
 	// Clears the working arrays to not hold any references.
 	RENDERER_API void Clear();
