Added instance hierarchy support for updating primitives with precomputed data. This supports rebuilding the hierarhcy on the fly, which is useful for debugging, among other things.

diff --git a/Engine/Source/Runtime/Renderer/Private/SceneCulling/SceneCulling.cpp b/Engine/Source/Runtime/Renderer/Private/SceneCulling/SceneCulling.cpp
--- a/Engine/Source/Runtime/Renderer/Private/SceneCulling/SceneCulling.cpp
+++ b/Engine/Source/Runtime/Renderer/Private/SceneCulling/SceneCulling.cpp
@@ -1868,6 +1868,18 @@ public:
 		return PrimitiveSceneInfo->GetInstanceSceneDataBuffers();
 	}
 
+	SC_FORCEINLINE void AddPrecomputed(int32 InstanceDataOffset, const TSharedPtr<FInstanceSceneDataImmutable, ESPMode::ThreadSafe>& InstanceSceneDataImmutable)
+	{
+		BUILDER_LOG_LIST("Add AddPrecomputed");
+		int32 InstanceDataOffsetCurrent = InstanceDataOffset;
+		for (FInstanceSceneDataBuffers::FCompressedSpatialHashItem Item : InstanceSceneDataImmutable->GetCompressedInstanceSpatialHashes())
+		{
+			int32 CellIndex = AddRange<EUpdateFrequencyCategory::Static>(Item.Location, InstanceDataOffsetCurrent, Item.NumInstances);
+			BUILDER_LOG_LIST_APPEND("(%d, %d, %d)", CellIndex, InstanceDataOffsetCurrent, Item.NumInstances);
+			InstanceDataOffsetCurrent += Item.NumInstances;
+		}
+	}
+
 	SC_FORCEINLINE void AddInstances(FPersistentPrimitiveIndex PersistentId, FPrimitiveSceneInfo* PrimitiveSceneInfo)
 	{
 		FScene &Scene = SceneCulling.Scene;
@@ -1912,21 +1924,10 @@ public:
 			{
 				check(InstanceSceneDataBuffers);
 				SC_SCOPED_NAMED_EVENT_DETAIL(SceneCulling_Post_AddInstances_Precomputed, FColor::Emerald);
-
-				BUILDER_LOG_LIST("Add CompressedInstanceSpatialHashes");
-
-				int32 InstanceDataOffsetCurrent = InstanceDataOffset;
-
 				TSharedPtr<FInstanceSceneDataImmutable, ESPMode::ThreadSafe> InstanceSceneDataImmutable = InstanceSceneDataBuffers->GetImmutable();
-
-				for (FInstanceSceneDataBuffers::FCompressedSpatialHashItem Item : InstanceSceneDataImmutable->GetCompressedInstanceSpatialHashes())
-				{
-					int32 CellIndex = AddRange<EUpdateFrequencyCategory::Static>(Item.Location, InstanceDataOffsetCurrent, Item.NumInstances);
-					BUILDER_LOG_LIST_APPEND("(%d, %d, %d)", CellIndex, InstanceDataOffsetCurrent, Item.NumInstances);
-					InstanceDataOffsetCurrent += Item.NumInstances;
-				}
+				AddPrecomputed(InstanceDataOffset, InstanceSceneDataImmutable);
 				// Hang onto this so we can safely remove the thing if it changes in the future.
-				NewPrimitiveState.InstanceSceneDataImmutable = InstanceSceneDataImmutable;
+				NewPrimitiveState.InstanceSceneDataImmutable = MoveTemp(InstanceSceneDataImmutable);
 			}
 			break;
 			case FPrimitiveState::UnCullable:
@@ -2222,9 +2223,6 @@ public:
 		{
 			SC_SCOPED_NAMED_EVENT_DETAIL(SceneCulling_Post_UpdateInstances_Precomputed, FColor::Emerald);
 			BUILDER_LOG("FPrimitiveState::Precomputed");
-			// This _should_ not happen (but it does!! I.e., we have static ISMs with updating instances), these must be remove/readded to the scene			
-			// Check that they are indeed not being re-added as precomputed.
-			check(bStateChanged);
 
 			RemovedInstanceFlags.SetRange(PrevPrimitiveState.InstanceDataOffset, PrevPrimitiveState.NumInstances, true);
 
@@ -2282,8 +2280,12 @@ public:
 				break;
 				case FPrimitiveState::Precomputed:
 				{
-					// this is wrong, post-update a precomputed item should be transitioned to dynamic.
-					check(false);
+					check(InstanceSceneDataBuffers);
+					SC_SCOPED_NAMED_EVENT_DETAIL(SceneCulling_Post_AddInstances_Precomputed, FColor::Emerald);
+					TSharedPtr<FInstanceSceneDataImmutable, ESPMode::ThreadSafe> InstanceSceneDataImmutable = InstanceSceneDataBuffers->GetImmutable();
+					AddPrecomputed(InstanceDataOffset, InstanceSceneDataImmutable);
+					// Hang onto this so we can safely remove the thing if it changes in the future.
+					NewPrimitiveState.InstanceSceneDataImmutable = MoveTemp(InstanceSceneDataImmutable);
 				}
 				break;
 				case FPrimitiveState::UnCullable:
