Removed LocalLightVisibleLightInfosIndex since we only need to know whether there are local lights.

diff --git a/Engine/Source/Runtime/Renderer/Private/LightGridInjection.cpp b/Engine/Source/Runtime/Renderer/Private/LightGridInjection.cpp
--- a/Engine/Source/Runtime/Renderer/Private/LightGridInjection.cpp
+++ b/Engine/Source/Runtime/Renderer/Private/LightGridInjection.cpp
@@ -752,7 +752,6 @@ FComputeLightGridOutput FSceneRenderer::ComputeLightGrid(FRDGBuilder& GraphBuild
 		ForwardLightUniformParameters->DirectionalLightStaticShadowmap = GBlackTexture->TextureRHI;
 
 		TArray<FForwardLocalLightData, SceneRenderingAllocator> ForwardLocalLightData;
-		TArray<int32, SceneRenderingAllocator>  LocalLightVisibleLightInfosIndex;
 
 		TArray<FVector4f, SceneRenderingAllocator> ViewSpacePosAndRadiusData;
 		TArray<FVector4f, SceneRenderingAllocator> ViewSpaceDirAndPreprocAngleData;
@@ -780,7 +779,6 @@ FComputeLightGridOutput FSceneRenderer::ComputeLightGrid(FRDGBuilder& GraphBuild
 			if (SimpleLightsEnd > 0)
 			{
 				ForwardLocalLightData.Reserve(SimpleLightsEnd);
-				LocalLightVisibleLightInfosIndex.Reserve(SimpleLightsEnd);
 
 				ViewSpacePosAndRadiusData.Reserve(SimpleLightsEnd);
 				ViewSpaceDirAndPreprocAngleData.Reserve(SimpleLightsEnd);
@@ -808,9 +806,6 @@ FComputeLightGridOutput FSceneRenderer::ComputeLightGrid(FRDGBuilder& GraphBuild
 					ForwardLocalLightData.AddUninitialized(1);
 					FForwardLocalLightData& LightData = ForwardLocalLightData.Last();
 
-					// Simple lights have no 'VisibleLight' info
-					LocalLightVisibleLightInfosIndex.Add(INDEX_NONE);
-
 					const FSimpleLightEntry& SimpleLight = SimpleLights.InstanceData[SimpleLightIndex];
 					const FSimpleLightPerViewEntry& SimpleLightPerViewData = SimpleLights.GetViewDependentData(SimpleLightIndex, ViewIndex, Views.Num());
 					PackLocalLightData(LightData, View, SimpleLight, SimpleLightPerViewData);
@@ -869,7 +864,6 @@ FComputeLightGridOutput FSceneRenderer::ComputeLightGrid(FRDGBuilder& GraphBuild
 
 						ForwardLocalLightData.AddUninitialized(1);
 						FForwardLocalLightData& LightData = ForwardLocalLightData.Last();
-						LocalLightVisibleLightInfosIndex.Add(LightSceneInfo->Id);
 
 						// Track the last one to support clustered deferred
 						if (!SortedLightInfo.SortKey.Fields.bClusteredDeferredNotSupported)
@@ -1078,8 +1072,6 @@ FComputeLightGridOutput FSceneRenderer::ComputeLightGrid(FRDGBuilder& GraphBuild
 		FRDGBufferRef ForwardLocalLightBuffer = CreateStructuredBuffer(GraphBuilder, TEXT("ForwardLocalLightBuffer"),
 			TConstArrayView<FVector4f>(reinterpret_cast<const FVector4f*>(ForwardLocalLightData.GetData()), ForwardLocalLightDataSizeNumFloat4));
 
-		View.ForwardLightingResources.LocalLightVisibleLightInfosIndex = LocalLightVisibleLightInfosIndex;
-
 		View.bLightGridHasRectLights = bHasRectLights;
 		View.bLightGridHasTexturedLights = bHasTexturedLights;
 
diff --git a/Engine/Source/Runtime/Renderer/Private/SceneRendering.h b/Engine/Source/Runtime/Renderer/Private/SceneRendering.h
--- a/Engine/Source/Runtime/Renderer/Private/SceneRendering.h
+++ b/Engine/Source/Runtime/Renderer/Private/SceneRendering.h
@@ -716,9 +716,6 @@ public:
 	TRDGUniformBufferRef<FForwardLightUniformParameters> ForwardLightUniformBuffer = nullptr;
 
 	const FLightSceneProxy* SelectedForwardDirectionalLightProxy = nullptr;
-
-	// Index into FSceneRenderer::VisibleLightInfos for each light in the ForwardLocalLightBuffer (these are copied when the light grid is built)
-	TArray<int32> LocalLightVisibleLightInfosIndex;
 };
 
 BEGIN_GLOBAL_SHADER_PARAMETER_STRUCT_WITH_CONSTRUCTOR(FVolumetricFogGlobalData,) 
diff --git a/Engine/Source/Runtime/Renderer/Private/VirtualShadowMaps/VirtualShadowMapArray.cpp b/Engine/Source/Runtime/Renderer/Private/VirtualShadowMaps/VirtualShadowMapArray.cpp
--- a/Engine/Source/Runtime/Renderer/Private/VirtualShadowMaps/VirtualShadowMapArray.cpp
+++ b/Engine/Source/Runtime/Renderer/Private/VirtualShadowMaps/VirtualShadowMapArray.cpp
@@ -1946,7 +1946,7 @@ void FVirtualShadowMapArray::BuildPageAllocations(
 		}	
 		
 		// This view contained no local lights (that were stored in the light grid), and no directional lights, so nothing to do.
-		if (View.ForwardLightingResources.LocalLightVisibleLightInfosIndex.Num() + DirectionalLightIds.Num() == 0)
+		if (View.ForwardLightingResources.ForwardLightUniformParameters->NumLocalLights + DirectionalLightIds.Num() == 0)
 		{
 			PerViewData = MakeEmptyVirtualShadowMapPerViewParameters(GraphBuilder);
 			continue;
