[GPU Scene Lights] Support directional lights in ForwardLightBuffer.
- Modified LightGridInjectionCS to support applying an indirection on the indices written to the final grid, since the shader only processes visible local lights but needs to write the index of the relevant lights in ForwardLightBuffer.
- Renamed ForwardLocalLightBuffer to ForwardLightBuffer and other related names.

diff --git a/Engine/Shaders/Private/LightGridCommon.ush b/Engine/Shaders/Private/LightGridCommon.ush
--- a/Engine/Shaders/Private/LightGridCommon.ush
+++ b/Engine/Shaders/Private/LightGridCommon.ush
@@ -104,8 +104,8 @@ uint ComputeLightGridCellIndex(uint2 PixelPos, float SceneDepth)
 #define NUM_CULLED_LIGHTS_GRID_STRIDE 0
 #endif
 
-#ifndef LOCAL_LIGHT_DATA_STRIDE
-#define LOCAL_LIGHT_DATA_STRIDE 0
+#ifndef FORWARD_LIGHT_DATA_STRIDE
+#define FORWARD_LIGHT_DATA_STRIDE 0
 #endif
 
 uint GetNumLocalLights(uint EyeIndex)
@@ -301,18 +301,18 @@ FLocalLightData GetLocalLightData_Internal(
 FLocalLightData GetLocalLightDataNonStereo(uint LocalLightIndex)
 {
 	FLocalLightData Out;
-	uint LocalLightBaseIndex = LocalLightIndex * LOCAL_LIGHT_DATA_STRIDE;
+	uint LocalLightBaseIndex = LocalLightIndex * FORWARD_LIGHT_DATA_STRIDE;
 
 	Out = GetLocalLightData_Internal(
 		LocalLightIndex,
 		ForwardLightStruct.ClusteredDeferredSupportedEndIndex,
 		ForwardLightStruct.SimpleLightsEndIndex,
-		ForwardLightStruct.ForwardLocalLightBuffer[LocalLightBaseIndex + 0],
-		ForwardLightStruct.ForwardLocalLightBuffer[LocalLightBaseIndex + 1],
-		ForwardLightStruct.ForwardLocalLightBuffer[LocalLightBaseIndex + 2],
-		ForwardLightStruct.ForwardLocalLightBuffer[LocalLightBaseIndex + 3],
-		ForwardLightStruct.ForwardLocalLightBuffer[LocalLightBaseIndex + 4],
-		ForwardLightStruct.ForwardLocalLightBuffer[LocalLightBaseIndex + 5]);
+		ForwardLightStruct.ForwardLightBuffer[LocalLightBaseIndex + 0],
+		ForwardLightStruct.ForwardLightBuffer[LocalLightBaseIndex + 1],
+		ForwardLightStruct.ForwardLightBuffer[LocalLightBaseIndex + 2],
+		ForwardLightStruct.ForwardLightBuffer[LocalLightBaseIndex + 3],
+		ForwardLightStruct.ForwardLightBuffer[LocalLightBaseIndex + 4],
+		ForwardLightStruct.ForwardLightBuffer[LocalLightBaseIndex + 5]);
 	return Out;
 }
 
@@ -326,35 +326,35 @@ FLocalLightData GetLocalLightData(uint GridIndex, uint EyeIndex)
 	{
 #endif
 		uint LocalLightIndex = GetCulledLightDataGrid(GridIndex);
-		uint LocalLightBaseIndex = LocalLightIndex * LOCAL_LIGHT_DATA_STRIDE;
+		uint LocalLightBaseIndex = LocalLightIndex * FORWARD_LIGHT_DATA_STRIDE;
 
 		Out = GetLocalLightData_Internal(
 			LocalLightIndex, 
 			ForwardLightStruct.ClusteredDeferredSupportedEndIndex,
 			ForwardLightStruct.SimpleLightsEndIndex,
-			ForwardLightStruct.ForwardLocalLightBuffer[LocalLightBaseIndex + 0],
-			ForwardLightStruct.ForwardLocalLightBuffer[LocalLightBaseIndex + 1],
-			ForwardLightStruct.ForwardLocalLightBuffer[LocalLightBaseIndex + 2],
-			ForwardLightStruct.ForwardLocalLightBuffer[LocalLightBaseIndex + 3],
-			ForwardLightStruct.ForwardLocalLightBuffer[LocalLightBaseIndex + 4],
-			ForwardLightStruct.ForwardLocalLightBuffer[LocalLightBaseIndex + 5]);
+			ForwardLightStruct.ForwardLightBuffer[LocalLightBaseIndex + 0],
+			ForwardLightStruct.ForwardLightBuffer[LocalLightBaseIndex + 1],
+			ForwardLightStruct.ForwardLightBuffer[LocalLightBaseIndex + 2],
+			ForwardLightStruct.ForwardLightBuffer[LocalLightBaseIndex + 3],
+			ForwardLightStruct.ForwardLightBuffer[LocalLightBaseIndex + 4],
+			ForwardLightStruct.ForwardLightBuffer[LocalLightBaseIndex + 5]);
 #if (INSTANCED_STEREO || MOBILE_MULTI_VIEW)
 	}
 	else
 	{
 		uint LocalLightIndex = GetCulledLightDataGridISR(GridIndex);
-		uint LocalLightBaseIndex = LocalLightIndex * LOCAL_LIGHT_DATA_STRIDE;
+		uint LocalLightBaseIndex = LocalLightIndex * FORWARD_LIGHT_DATA_STRIDE;
 
 		Out = GetLocalLightData_Internal(
 			LocalLightIndex,
 			ForwardLightStructISR.ClusteredDeferredSupportedEndIndex,
 			ForwardLightStructISR.SimpleLightsEndIndex,
-			ForwardLightStructISR.ForwardLocalLightBuffer[LocalLightBaseIndex + 0],
-			ForwardLightStructISR.ForwardLocalLightBuffer[LocalLightBaseIndex + 1],
-			ForwardLightStructISR.ForwardLocalLightBuffer[LocalLightBaseIndex + 2],
-			ForwardLightStructISR.ForwardLocalLightBuffer[LocalLightBaseIndex + 3],
-			ForwardLightStructISR.ForwardLocalLightBuffer[LocalLightBaseIndex + 4],
-			ForwardLightStructISR.ForwardLocalLightBuffer[LocalLightBaseIndex + 5]);
+			ForwardLightStructISR.ForwardLightBuffer[LocalLightBaseIndex + 0],
+			ForwardLightStructISR.ForwardLightBuffer[LocalLightBaseIndex + 1],
+			ForwardLightStructISR.ForwardLightBuffer[LocalLightBaseIndex + 2],
+			ForwardLightStructISR.ForwardLightBuffer[LocalLightBaseIndex + 3],
+			ForwardLightStructISR.ForwardLightBuffer[LocalLightBaseIndex + 4],
+			ForwardLightStructISR.ForwardLightBuffer[LocalLightBaseIndex + 5]);
 	}
 #endif
 	return Out;
diff --git a/Engine/Shaders/Private/LightGridInjection.usf b/Engine/Shaders/Private/LightGridInjection.usf
--- a/Engine/Shaders/Private/LightGridInjection.usf
+++ b/Engine/Shaders/Private/LightGridInjection.usf
@@ -53,6 +53,8 @@ StructuredBuffer<float4> LightViewSpacePositionAndRadius;
 StructuredBuffer<float4> LightViewSpaceDirAndPreprocAngle;
 StructuredBuffer<float4> LightViewSpaceRectPlanes;
 
+StructuredBuffer<int> IndirectionIndices;
+
 #if USE_PARENT_LIGHT_GRID
 	StructuredBuffer<uint> ParentNumCulledLightsGrid;
 #	if LIGHT_GRID_USES_16BIT_BUFFERS
@@ -313,14 +315,20 @@ bool OverlapsReflectionCapture(uint ReflectionCaptureIndex, FCellBounds CellBoun
 	return (BoxDistanceSq < CaptureRadius * CaptureRadius);
 }
 
-void OutputReversedLinkedList(uint LinkOffset, uint NumPrimitives, uint CulledLightDataStart)
+void OutputReversedLinkedList(uint LinkOffset, uint NumPrimitives, uint CulledLightDataStart, bool bApplyIndirection)
 {
 	uint CurrentIndex = 0;
 
 	while (LinkOffset != 0xFFFFFFFF && CurrentIndex < NumPrimitives)
 	{
 		// Reverse the order as we write them out, which restores the original order before the reverse linked list was built
-		const uint PrimitiveIndex = RWCulledLightLinks[LinkOffset * LIGHT_LINK_STRIDE + 0];
+		uint PrimitiveIndex = RWCulledLightLinks[LinkOffset * LIGHT_LINK_STRIDE + 0];
+
+		if (bApplyIndirection)
+		{
+			PrimitiveIndex = IndirectionIndices[PrimitiveIndex];
+		}
+
 		RWCulledLightDataGrid[CulledLightDataStart + NumPrimitives - CurrentIndex - 1] = PrimitiveIndex;
 		CurrentIndex++;
 		LinkOffset = RWCulledLightLinks[LinkOffset * LIGHT_LINK_STRIDE + 1];
@@ -422,7 +430,11 @@ void CullLights_SingleThread(uint GridIndex, FCellBounds CellBounds, uint Parent
 #else // !USE_LINKED_CULL_LIST
 			if (NumVisibleLights < MaxCulledLightsPerCell)
 			{
-				RWCulledLightDataGrid[GridIndex * MaxCulledLightsPerCell + NumVisibleLights] = LocalLightIndex;
+				uint PrimitiveIndex = LocalLightIndex;
+#if APPLY_INDIRECTION
+				PrimitiveIndex = IndirectionIndices[PrimitiveIndex];
+#endif
+				RWCulledLightDataGrid[GridIndex * MaxCulledLightsPerCell + NumVisibleLights] = PrimitiveIndex;
 
 				++NumVisibleLights;
 				if (LocalLightIndex >= MegaLightsSupportedStartIndex)
@@ -445,7 +457,7 @@ void CullLights_SingleThread(uint GridIndex, FCellBounds CellBounds, uint Parent
 	#endif
 
 	// Mega Lights are order dependent
-	OutputReversedLinkedList(LinkOffset, NumVisibleLights, CulledLightDataStart);
+	OutputReversedLinkedList(LinkOffset, NumVisibleLights, CulledLightDataStart, APPLY_INDIRECTION);
 #else
 	uint CulledLightDataStart = GridIndex * MaxCulledLightsPerCell;
 #endif
@@ -499,7 +511,7 @@ void CullReflectionCaptures_SingleThread(uint GridIndex, FCellBounds CellBounds,
 	WaveInterlockedAdd_(RWNextCulledLightData[0], NumVisibleCaptures, CulledReflectionCaptureDataStart);
 
 	// Reflection captures are order dependent
-	OutputReversedLinkedList(LinkOffset, NumVisibleCaptures, CulledReflectionCaptureDataStart);
+	OutputReversedLinkedList(LinkOffset, NumVisibleCaptures, CulledReflectionCaptureDataStart, false);
 #else
 	uint CulledReflectionCaptureDataStart = (NumGridCells + GridIndex) * MaxCulledLightsPerCell;
 #endif
@@ -565,14 +577,21 @@ bool AddToCell(inout FCellWriter CellWriter, uint PrimitiveIndex, uint IndexInCe
 	return false;
 }
 
-void OutputCell(FCellWriter CellWriter, uint NumPrimitives, uint CulledLightDataStart, uint ThreadIndex)
+void OutputCell(FCellWriter CellWriter, uint NumPrimitives, uint CulledLightDataStart, uint ThreadIndex, bool bApplyIndirection)
 {
 #if USE_LDS
 	// output entries in LDS to RWCulledLightDataGrid
 	LOOP
 	for (uint OutIndex = ThreadIndex; OutIndex < NumPrimitives; OutIndex += THREADGROUP_SIZE)
 	{
-		RWCulledLightDataGrid[CulledLightDataStart + OutIndex] = CellLDSItems[OutIndex];
+		uint PrimitiveIndex = CellLDSItems[OutIndex];
+
+		if (bApplyIndirection)
+		{
+			PrimitiveIndex = IndirectionIndices[PrimitiveIndex];
+		}
+
+		RWCulledLightDataGrid[CulledLightDataStart + OutIndex] = PrimitiveIndex;
 	}
 #endif
 
@@ -582,12 +601,16 @@ void OutputCell(FCellWriter CellWriter, uint NumPrimitives, uint CulledLightData
 	{
 		const uint PackedData0 = RWCulledLightLinks[CellWriter.LinkOffset * LIGHT_LINK_STRIDE + 0];
 		const uint PackedData1 = RWCulledLightLinks[CellWriter.LinkOffset * LIGHT_LINK_STRIDE + 1];
-		const uint PrimitiveIndex = (PackedData0 & 0xFFFFFF);
+		uint PrimitiveIndex = (PackedData0 & 0xFFFFFF);
 		const uint LinkOffset = (PackedData1 & 0xFFFFFF);
 		const uint IndexInCellL8 = (PackedData0 >> 24) & 0xFF;
 		const uint IndexInCellH8 = (PackedData1 >> 24) & 0xFF;
 		const uint IndexInCell = (IndexInCellH8 << 8) | IndexInCellL8;
 
+#if APPLY_INDIRECTION
+		PrimitiveIndex = IndirectionIndices[PrimitiveIndex];
+#endif
+
 		RWCulledLightDataGrid[CulledLightDataStart + IndexInCell] = PrimitiveIndex;
 		CellWriter.LinkOffset = LinkOffset;
 	}
@@ -682,7 +705,7 @@ void CullLights_ThreadGroup(uint ThreadIndex, uint GridIndex, FCellBounds CellBo
 
 	GroupMemoryBarrierWithGroupSync();
 
-	OutputCell(CellWriter, CellNumVisibleLights, CulledLightDataStart, ThreadIndex);
+	OutputCell(CellWriter, CellNumVisibleLights, CulledLightDataStart, ThreadIndex, APPLY_INDIRECTION);
 
 	if (ThreadIndex == 0)
 	{
@@ -753,7 +776,7 @@ void CullReflectionCaptures_ThreadGroup(uint ThreadIndex, uint GridIndex, FCellB
 
 	GroupMemoryBarrierWithGroupSync();
 
-	OutputCell(CellWriter, CellNumVisibleReflectionCaptures, CulledLightDataStart, ThreadIndex);
+	OutputCell(CellWriter, CellNumVisibleReflectionCaptures, CulledLightDataStart, ThreadIndex, false);
 
 	if (ThreadIndex == 0)
 	{
diff --git a/Engine/Source/Runtime/Renderer/Private/BasePassRendering.h b/Engine/Source/Runtime/Renderer/Private/BasePassRendering.h
--- a/Engine/Source/Runtime/Renderer/Private/BasePassRendering.h
+++ b/Engine/Source/Runtime/Renderer/Private/BasePassRendering.h
@@ -56,7 +56,7 @@ namespace OIT
 /** Whether to allow the indirect lighting cache to be applied to dynamic objects. */
 extern int32 GIndirectLightingCache;
 
-class FForwardLocalLightData
+class FForwardLightData
 {
 public:
 	FVector4f LightPositionAndInvRadius;
@@ -181,7 +181,7 @@ public:
 
 	static void ModifyCompilationEnvironment(EShaderPlatform Platform, FShaderCompilerEnvironment& OutEnvironment)
 	{
-		OutEnvironment.SetDefine(TEXT("LOCAL_LIGHT_DATA_STRIDE"), FMath::DivideAndRoundUp<int32>(sizeof(FForwardLocalLightData), sizeof(FVector4f)));
+		OutEnvironment.SetDefine(TEXT("FORWARD_LIGHT_DATA_STRIDE"), FMath::DivideAndRoundUp<int32>(sizeof(FForwardLightData), sizeof(FVector4f)));
 		extern int32 NumCulledLightsGridStride;
 		OutEnvironment.SetDefine(TEXT("NUM_CULLED_LIGHTS_GRID_STRIDE"), NumCulledLightsGridStride);
 		extern int32 NumCulledGridPrimitiveTypes;
diff --git a/Engine/Source/Runtime/Renderer/Private/LightGridInjection.cpp b/Engine/Source/Runtime/Renderer/Private/LightGridInjection.cpp
--- a/Engine/Source/Runtime/Renderer/Private/LightGridInjection.cpp
+++ b/Engine/Source/Runtime/Renderer/Private/LightGridInjection.cpp
@@ -40,6 +40,21 @@ LightGridInjection.cpp
 #include "VolumetricFog.h"
 #include "LightViewData.h"
 
+enum class ELightBufferMode
+{
+	VisibleLocalLights	= 0,
+	VisibleLights		= 1
+};
+
+static int32 GLightBufferMode = (int32)ELightBufferMode::VisibleLocalLights;
+static FAutoConsoleVariableRef CVarLightBufferMode(
+	TEXT("r.Forward.LightBuffer.Mode"),
+	GLightBufferMode,
+	TEXT("0 - Visible local lights.\n")
+	TEXT("1 - Visible local + directional lights."),
+	ECVF_RenderThreadSafe
+);
+
 int32 GLightGridPixelSize = 64;
 FAutoConsoleVariableRef CVarLightGridPixelSize(
 	TEXT("r.Forward.LightGridPixelSize"),
@@ -180,8 +195,8 @@ void SetupDummyForwardLightUniformParameters(FRDGBuilder& GraphBuilder, FForward
 	ForwardLightUniformParameters.DirectionalLightShadowmapAtlas = SystemTextures.Black;
 	ForwardLightUniformParameters.DirectionalLightStaticShadowmap = GBlackTexture->TextureRHI;
 
-	FRDGBufferRef ForwardLocalLightBuffer = GSystemTextures.GetDefaultStructuredBuffer(GraphBuilder, sizeof(FVector4f));
-	ForwardLightUniformParameters.ForwardLocalLightBuffer = GraphBuilder.CreateSRV(ForwardLocalLightBuffer);
+	FRDGBufferRef ForwardLightBuffer = GSystemTextures.GetDefaultStructuredBuffer(GraphBuilder, sizeof(FVector4f));
+	ForwardLightUniformParameters.ForwardLightBuffer = GraphBuilder.CreateSRV(ForwardLightBuffer);
 
 	FRDGBufferRef NumCulledLightsGrid = GSystemTextures.GetDefaultStructuredBuffer(GraphBuilder, sizeof(uint32));
 	ForwardLightUniformParameters.NumCulledLightsGrid = GraphBuilder.CreateSRV(NumCulledLightsGrid);
@@ -256,7 +271,8 @@ public:
 	class FUseParentLightGrid : SHADER_PERMUTATION_BOOL("USE_PARENT_LIGHT_GRID");
 	class FUseThreadGroupPerCell : SHADER_PERMUTATION_BOOL("USE_THREAD_GROUP_PER_CELL");
 	class FUseThreadGroupSize32 : SHADER_PERMUTATION_BOOL("USE_THREAD_GROUP_SIZE_32");
-	using FPermutationDomain = TShaderPermutationDomain<FUseLinkedList, FRefineRectLightBounds, FUseHZBCull, FUseParentLightGrid, FUseThreadGroupPerCell, FUseThreadGroupSize32>;
+	class FApplyIndirection : SHADER_PERMUTATION_BOOL("APPLY_INDIRECTION");
+	using FPermutationDomain = TShaderPermutationDomain<FUseLinkedList, FRefineRectLightBounds, FUseHZBCull, FUseParentLightGrid, FUseThreadGroupPerCell, FUseThreadGroupSize32, FApplyIndirection>;
 
 	BEGIN_SHADER_PARAMETER_STRUCT(FParameters, )
 		SHADER_PARAMETER_STRUCT_REF(FReflectionCaptureShaderData, ReflectionCapture)
@@ -271,6 +287,7 @@ public:
 		SHADER_PARAMETER_RDG_BUFFER_SRV(StructuredBuffer<float4>, LightViewSpacePositionAndRadius)
 		SHADER_PARAMETER_RDG_BUFFER_SRV(StructuredBuffer<float4>, LightViewSpaceDirAndPreprocAngle)
 		SHADER_PARAMETER_RDG_BUFFER_SRV(StructuredBuffer<float4>, LightViewSpaceRectPlanes)
+		SHADER_PARAMETER_RDG_BUFFER_SRV(StructuredBuffer<int32>, IndirectionIndices)
 
 		SHADER_PARAMETER(FIntVector, CulledGridSize)
 		SHADER_PARAMETER(uint32, NumReflectionCaptures)
@@ -418,8 +435,8 @@ static uint32 PackVirtualShadowMapIdAndPrevLocalLightIndex(int32 VirtualShadowMa
 	return (VSMPacked << 16) | (PrevPacked & 0xFFFF);
 }
 
-static void PackLocalLightData(
-	FForwardLocalLightData& Out,
+static void PackLightData(
+	FForwardLightData& Out,
 	const FViewInfo& View,
 	const FSimpleLightEntry& SimpleLight,
 	const FSimpleLightPerViewEntry& SimpleLightPerViewData)
@@ -470,8 +487,8 @@ static void PackLocalLightData(
 	Out.RectDataAndVirtualShadowMapIdOrPrevLocalLightIndex	= FVector4f(FMath::AsFloat(RectPackedX), FMath::AsFloat(RectPackedY), FMath::AsFloat(RectPackedZ), FMath::AsFloat(VirtualShadowMapIdAndPrevLocalLightIndex));
 }
 
-static void PackLocalLightData(
-	FForwardLocalLightData& Out, 
+static void PackLightData(
+	FForwardLightData& Out,
 	const FViewInfo& View,
 	const FLightRenderParameters& LightParameters,
 	const uint32 LightSceneInfoExtraDataPacked,
@@ -577,6 +594,7 @@ FLightGrid LightGridInjection(
 	FRDGBufferSRVRef LightViewSpacePositionAndRadiusSRV,
 	FRDGBufferSRVRef LightViewSpaceDirAndPreprocAngleSRV,
 	FRDGBufferSRVRef LightViewSpaceRectPlanesSRV,
+	FRDGBufferSRVRef IndirectionIndicesSRV,
 	FLightGridViewState* LightGridViewState,
 	bool bThreadGroupPerCell,
 	bool bThreadGroupSize32,
@@ -670,6 +688,8 @@ FLightGrid LightGridInjection(
 	PassParameters->LightViewSpaceDirAndPreprocAngle = LightViewSpaceDirAndPreprocAngleSRV;
 	PassParameters->LightViewSpaceRectPlanes = LightViewSpaceRectPlanesSRV;
 
+	PassParameters->IndirectionIndices = IndirectionIndicesSRV;
+
 	{
 		PassParameters->HZBTexture = View.HZB;
 		PassParameters->HZBSampler = TStaticSamplerState< SF_Point, AM_Clamp, AM_Clamp, AM_Clamp >::GetRHI();
@@ -685,6 +705,7 @@ FLightGrid LightGridInjection(
 	PermutationVector.Set<FLightGridInjectionCS::FUseParentLightGrid>(ParentNumCulledLightsGridSRV != nullptr && ParentCulledLightDataGridSRV != nullptr);
 	PermutationVector.Set<FLightGridInjectionCS::FUseThreadGroupPerCell>(bThreadGroupPerCell);
 	PermutationVector.Set<FLightGridInjectionCS::FUseThreadGroupSize32>(bThreadGroupSize32);
+	PermutationVector.Set<FLightGridInjectionCS::FApplyIndirection>(IndirectionIndicesSRV != nullptr);
 	auto ComputeShader = View.ShaderMap->GetShader<FLightGridInjectionCS>(PermutationVector);
 
 	AddClearUAVPass(GraphBuilder, PassParameters->RWNextCulledLightLink, 0);
@@ -751,12 +772,15 @@ FComputeLightGridOutput FSceneRenderer::ComputeLightGrid(FRDGBuilder& GraphBuild
 		ForwardLightUniformParameters->DirectionalLightShadowmapAtlas = SystemTextures.Black;
 		ForwardLightUniformParameters->DirectionalLightStaticShadowmap = GBlackTexture->TextureRHI;
 
-		TArray<FForwardLocalLightData, SceneRenderingAllocator> ForwardLocalLightData;
+		TArray<FForwardLightData, SceneRenderingAllocator> ForwardLightData;
+		TArray<int32, SceneRenderingAllocator>  DirectionalLightIndices;
 
 		TArray<FVector4f, SceneRenderingAllocator> ViewSpacePosAndRadiusData;
 		TArray<FVector4f, SceneRenderingAllocator> ViewSpaceDirAndPreprocAngleData;
 		TArray<FVector4f, SceneRenderingAllocator> ViewSpaceRectPlanesData;
 
+		TArray<int32, SceneRenderingAllocator> IndirectionIndices;
+
 		float FurthestLight = 1000;
 
 		int32 ConflictingLightCountForForwardShading = 0;
@@ -778,7 +802,7 @@ FComputeLightGridOutput FSceneRenderer::ComputeLightGrid(FRDGBuilder& GraphBuild
 			// 1. insert simple lights
 			if (SimpleLightsEnd > 0)
 			{
-				ForwardLocalLightData.Reserve(SimpleLightsEnd);
+				ForwardLightData.Reserve(SimpleLightsEnd);
 
 				ViewSpacePosAndRadiusData.Reserve(SimpleLightsEnd);
 				ViewSpaceDirAndPreprocAngleData.Reserve(SimpleLightsEnd);
@@ -803,17 +827,19 @@ FComputeLightGridOutput FSceneRenderer::ComputeLightGrid(FRDGBuilder& GraphBuild
 
 					int32 SimpleLightIndex = SortedLightSet.SortedLights[SortedIndex].SimpleLightIndex;
 
-					ForwardLocalLightData.AddUninitialized(1);
-					FForwardLocalLightData& LightData = ForwardLocalLightData.Last();
+					const int32 IndexInBuffer = ForwardLightData.AddUninitialized(1);
+					FForwardLightData& LightData = ForwardLightData[IndexInBuffer];
 
 					const FSimpleLightEntry& SimpleLight = SimpleLights.InstanceData[SimpleLightIndex];
 					const FSimpleLightPerViewEntry& SimpleLightPerViewData = SimpleLights.GetViewDependentData(SimpleLightIndex, ViewIndex, Views.Num());
-					PackLocalLightData(LightData, View, SimpleLight, SimpleLightPerViewData);
+					PackLightData(LightData, View, SimpleLight, SimpleLightPerViewData);
 
 					FVector4f ViewSpacePosAndRadius(FVector4f(View.ViewMatrices.GetViewMatrix().TransformPosition(SimpleLightPerViewData.Position)), SimpleLight.Radius);
 					ViewSpacePosAndRadiusData.Add(ViewSpacePosAndRadius);
 					ViewSpaceDirAndPreprocAngleData.AddZeroed();
 					ViewSpaceRectPlanesData.AddZeroed(NUM_PLANES_PER_RECT_LIGHT);
+
+					IndirectionIndices.Add(IndexInBuffer);
 				}
 			}
 
@@ -851,30 +877,27 @@ FComputeLightGridOutput FSceneRenderer::ComputeLightGrid(FRDGBuilder& GraphBuild
 					const bool bDynamicShadows = ViewFamily.EngineShowFlags.DynamicShadows && VisibleLightInfos.IsValidIndex(LightSceneInfo->Id);
 					const int32 VirtualShadowMapId = bDynamicShadows ? VisibleLightInfos[LightSceneInfo->Id].GetVirtualShadowMapId( &View ) : INDEX_NONE;
 
-					if ((SortedLightInfo.SortKey.Fields.LightType == LightType_Point && ViewFamily.EngineShowFlags.PointLights) ||
-						(SortedLightInfo.SortKey.Fields.LightType == LightType_Spot && ViewFamily.EngineShowFlags.SpotLights) ||
-						(SortedLightInfo.SortKey.Fields.LightType == LightType_Rect && ViewFamily.EngineShowFlags.RectLights))
+					// Add light to ForwardLightData
+					if ((SortedLightInfo.SortKey.Fields.LightType == LightType_Point && ViewFamily.EngineShowFlags.PointLights)
+						|| (SortedLightInfo.SortKey.Fields.LightType == LightType_Spot && ViewFamily.EngineShowFlags.SpotLights)
+						|| (SortedLightInfo.SortKey.Fields.LightType == LightType_Rect && ViewFamily.EngineShowFlags.RectLights)
+						|| (GLightBufferMode != (int32)ELightBufferMode::VisibleLocalLights && SortedLightInfo.SortKey.Fields.LightType == LightType_Directional && ViewFamily.EngineShowFlags.DirectionalLights))
 					{
-						int32 PrevLocalLightIndex = INDEX_NONE;
+						const int32 IndexInBuffer = ForwardLightData.AddUninitialized(1);
+						FForwardLightData& LightData = ForwardLightData[IndexInBuffer];
+
+						int32 PrevForwardLightIndex = INDEX_NONE;
 						if (View.ViewState)
 						{
-							PrevLocalLightIndex = View.ViewState->LightSceneIdToLocalLightIndex.FindOrAdd(LightSceneInfo->Id, INDEX_NONE);
-							View.ViewState->LightSceneIdToLocalLightIndex[LightSceneInfo->Id] = ForwardLocalLightData.Num();
+							PrevForwardLightIndex = View.ViewState->LightSceneIdToForwardLightIndex.FindOrAdd(LightSceneInfo->Id, INDEX_NONE);
+							View.ViewState->LightSceneIdToForwardLightIndex[LightSceneInfo->Id] = IndexInBuffer;
 						}
 
-						ForwardLocalLightData.AddUninitialized(1);
-						FForwardLocalLightData& LightData = ForwardLocalLightData.Last();
-
-						// Track the last one to support clustered deferred
-						if (!SortedLightInfo.SortKey.Fields.bClusteredDeferredNotSupported)
+						if (SortedLightInfo.SortKey.Fields.LightType != LightType_Directional) // only local lights go into the grid
 						{
-							ClusteredSupportedEnd = FMath::Max(ClusteredSupportedEnd, ForwardLocalLightData.Num());
+							IndirectionIndices.Add(IndexInBuffer);
 						}
 
-						if (SortedLightInfo.SortKey.Fields.bHandledByMegaLights && MegaLightsSupportedStart == MAX_int32)
-						{
-							MegaLightsSupportedStart = ForwardLocalLightData.Num() - 1;
-						}
 						const float LightFade = GetLightFadeFactor(View, LightProxy);
 						LightParameters.Color *= LightFade;
 						LightParameters.Color *= LightParameters.GetLightExposureScale(Exposure);
@@ -886,15 +909,33 @@ FComputeLightGridOutput FSceneRenderer::ComputeLightGrid(FRDGBuilder& GraphBuild
 							VolumetricScatteringIntensity = 0;
 						}
 
-						PackLocalLightData(
+						PackLightData(
 							LightData,
 							View,
 							LightParameters,
 							LightSceneInfoExtraDataPacked,
 							LightSceneInfo->Id,
 							VirtualShadowMapId,
-							PrevLocalLightIndex,
+							PrevForwardLightIndex,
 							VolumetricScatteringIntensity);
+					}
+
+					if ((SortedLightInfo.SortKey.Fields.LightType == LightType_Point && ViewFamily.EngineShowFlags.PointLights)
+						|| (SortedLightInfo.SortKey.Fields.LightType == LightType_Spot && ViewFamily.EngineShowFlags.SpotLights)
+						|| (SortedLightInfo.SortKey.Fields.LightType == LightType_Rect && ViewFamily.EngineShowFlags.RectLights))
+					{
+						const int32 LocalLightIndex = ViewSpacePosAndRadiusData.Num();
+
+						// Track the last one to support clustered deferred
+						if (!SortedLightInfo.SortKey.Fields.bClusteredDeferredNotSupported)
+						{
+							ClusteredSupportedEnd = FMath::Max(ClusteredSupportedEnd, LocalLightIndex + 1);
+						}
+
+						if (SortedLightInfo.SortKey.Fields.bHandledByMegaLights && MegaLightsSupportedStart == MAX_int32)
+						{
+							MegaLightsSupportedStart = LocalLightIndex;
+						}
 
 						const FSphere BoundingSphere = LightProxy->GetBoundingSphere();
 						const float Distance = View.ViewMatrices.GetViewMatrix().TransformPosition(BoundingSphere.Center).Z + BoundingSphere.W;
@@ -1059,29 +1100,33 @@ FComputeLightGridOutput FSceneRenderer::ComputeLightGrid(FRDGBuilder& GraphBuild
 		bMultipleDirLightsConflictForForwardShading |= ConflictingLightCountForForwardShading >= 2;
 #endif
 
-		// Store off the number of lights before we add a fake entry
-		const int32 NumLocalLightsFinal = ForwardLocalLightData.Num();
+		const int32 NumLightsFinal = ForwardLightData.Num();
+		const int32 NumVisibleLocalLights = ViewSpacePosAndRadiusData.Num();
 
-		MegaLightsSupportedStart = FMath::Min<int32>(MegaLightsSupportedStart, NumLocalLightsFinal);
+		MegaLightsSupportedStart = FMath::Min<int32>(MegaLightsSupportedStart, NumVisibleLocalLights);
 
 		// Some platforms index the StructuredBuffer in the shader based on the stride specified at buffer creation time, not from the stride specified in the shader.
-		// ForwardLocalLightBuffer is a StructuredBuffer<float4> in the shader, so create the buffer with a stride of sizeof(float4)
-		static_assert(sizeof(FForwardLocalLightData) % sizeof(FVector4f) == 0, "ForwardLocalLightBuffer is used as a StructuredBuffer<float4> in the shader");
-		const uint32 ForwardLocalLightDataSizeNumFloat4 = (NumLocalLightsFinal * sizeof(FForwardLocalLightData)) / sizeof(FVector4f);
+		// ForwardLightBuffer is a StructuredBuffer<float4> in the shader, so create the buffer with a stride of sizeof(float4)
+		static_assert(sizeof(FForwardLightData) % sizeof(FVector4f) == 0, "ForwardLightBuffer is used as a StructuredBuffer<float4> in the shader");
+		const uint32 ForwardLightDataSizeNumFloat4 = (NumLightsFinal * sizeof(FForwardLightData)) / sizeof(FVector4f);
 
-		FRDGBufferRef ForwardLocalLightBuffer = CreateStructuredBuffer(GraphBuilder, TEXT("ForwardLocalLightBuffer"),
-			TConstArrayView<FVector4f>(reinterpret_cast<const FVector4f*>(ForwardLocalLightData.GetData()), ForwardLocalLightDataSizeNumFloat4));
+		FRDGBufferRef ForwardLightBuffer = CreateStructuredBuffer(GraphBuilder, TEXT("ForwardLightBuffer"),
+			TConstArrayView<FVector4f>(reinterpret_cast<const FVector4f*>(ForwardLightData.GetData()), ForwardLightDataSizeNumFloat4));
+
+		FRDGBufferRef DirectionalLightIndicesBuffer = CreateStructuredBuffer(GraphBuilder, TEXT("DirectionalLightIndicesBuffer"), DirectionalLightIndices);
 
 		View.bLightGridHasRectLights = bHasRectLights;
 		View.bLightGridHasTexturedLights = bHasTexturedLights;
 
 		const FIntPoint LightGridSizeXY = FIntPoint::DivideAndRoundUp(View.ViewRect.Size(), GLightGridPixelSize);
-		ForwardLightUniformParameters->ForwardLocalLightBuffer = GraphBuilder.CreateSRV(FRDGBufferSRVDesc(ForwardLocalLightBuffer));
-		ForwardLightUniformParameters->NumLocalLights = NumLocalLightsFinal;
+		ForwardLightUniformParameters->ForwardLightBuffer = GraphBuilder.CreateSRV(FRDGBufferSRVDesc(ForwardLightBuffer));
+		ForwardLightUniformParameters->DirectionalLightIndices = GraphBuilder.CreateSRV(FRDGBufferSRVDesc(DirectionalLightIndicesBuffer));
+		ForwardLightUniformParameters->NumLocalLights = NumVisibleLocalLights;
+		ForwardLightUniformParameters->NumDirectionalLights = DirectionalLightIndices.Num();
 		ForwardLightUniformParameters->NumReflectionCaptures = View.NumBoxReflectionCaptures + View.NumSphereReflectionCaptures;
 		ForwardLightUniformParameters->NumGridCells = LightGridSizeXY.X * LightGridSizeXY.Y * GLightGridSizeZ;
 		ForwardLightUniformParameters->CulledGridSize = FIntVector(LightGridSizeXY.X, LightGridSizeXY.Y, GLightGridSizeZ);
-		ForwardLightUniformParameters->MaxCulledLightsPerCell = GLightLinkedListCulling ? NumLocalLightsFinal: GMaxCulledLightsPerCell;
+		ForwardLightUniformParameters->MaxCulledLightsPerCell = GLightLinkedListCulling ? NumVisibleLocalLights : GMaxCulledLightsPerCell;
 		ForwardLightUniformParameters->LightGridPixelSizeShift = FMath::FloorLog2(GLightGridPixelSize);
 		ForwardLightUniformParameters->SimpleLightsEndIndex = SimpleLightsEnd;
 		ForwardLightUniformParameters->ClusteredDeferredSupportedEndIndex = ClusteredSupportedEnd;
@@ -1096,20 +1141,21 @@ FComputeLightGridOutput FSceneRenderer::ComputeLightGrid(FRDGBuilder& GraphBuild
 
 		const uint64 NumIndexableLights = !bLightGridUses16BitBuffers ? (1llu << (sizeof(FLightIndexType32) * 8llu)) : (1llu << (sizeof(FLightIndexType) * 8llu));
 
-		if ((uint64)ForwardLocalLightData.Num() > NumIndexableLights)
+		if ((uint64)ForwardLightData.Num() > NumIndexableLights)
 		{
 			static bool bWarned = false;
 
 			if (!bWarned)
 			{
-				UE_LOG(LogRenderer, Warning, TEXT("Exceeded indexable light count, glitches will be visible (%u / %llu)"), ForwardLocalLightData.Num(), NumIndexableLights);
+				UE_LOG(LogRenderer, Warning, TEXT("Exceeded indexable light count, glitches will be visible (%u / %llu)"), ForwardLightData.Num(), NumIndexableLights);
 				bWarned = true;
 			}
 		}
 
-		check(ViewSpacePosAndRadiusData.Num() == ForwardLocalLightData.Num());
-		check(ViewSpaceDirAndPreprocAngleData.Num() == ForwardLocalLightData.Num());
-		check(ViewSpaceRectPlanesData.Num() == ForwardLocalLightData.Num() * NUM_PLANES_PER_RECT_LIGHT);
+		check(ViewSpacePosAndRadiusData.Num() == NumVisibleLocalLights);
+		check(ViewSpaceDirAndPreprocAngleData.Num() == NumVisibleLocalLights);
+		check(ViewSpaceRectPlanesData.Num() == NumVisibleLocalLights * NUM_PLANES_PER_RECT_LIGHT);
+		check(IndirectionIndices.Num() == NumVisibleLocalLights);
 
 		FRDGBufferRef LightViewSpacePositionAndRadius = CreateStructuredBuffer(GraphBuilder, TEXT("ViewSpacePosAndRadiusData"), TConstArrayView<FVector4f>(ViewSpacePosAndRadiusData));
 		FRDGBufferRef LightViewSpaceDirAndPreprocAngle = CreateStructuredBuffer(GraphBuilder, TEXT("ViewSpaceDirAndPreprocAngleData"), TConstArrayView<FVector4f>(ViewSpaceDirAndPreprocAngleData));
@@ -1119,6 +1165,9 @@ FComputeLightGridOutput FSceneRenderer::ComputeLightGrid(FRDGBuilder& GraphBuild
 		FRDGBufferSRVRef LightViewSpaceDirAndPreprocAngleSRV = GraphBuilder.CreateSRV(FRDGBufferSRVDesc(LightViewSpaceDirAndPreprocAngle));
 		FRDGBufferSRVRef LightViewSpaceRectPlanesSRV = GraphBuilder.CreateSRV(FRDGBufferSRVDesc(LightViewSpaceRectPlanes));
 
+		FRDGBufferRef IndirectionIndicesBuffer = CreateStructuredBuffer(GraphBuilder, TEXT("IndirectionIndices"), TConstArrayView<int32>(IndirectionIndices));
+		FRDGBufferSRVRef IndirectionIndicesSRV = GraphBuilder.CreateSRV(FRDGBufferSRVDesc(IndirectionIndicesBuffer));
+
 		// Allocate buffers using the scene render targets size so we won't reallocate every frame with dynamic resolution
 		const FIntPoint MaxLightGridSizeXY = FIntPoint::DivideAndRoundUp(View.GetSceneTexturesConfig().Extent, GLightGridPixelSize);
 
@@ -1176,6 +1225,7 @@ FComputeLightGridOutput FSceneRenderer::ComputeLightGrid(FRDGBuilder& GraphBuild
 				LightViewSpacePositionAndRadiusSRV,
 				LightViewSpaceDirAndPreprocAngleSRV,
 				LightViewSpaceRectPlanesSRV,
+				nullptr,
 				View.ViewState != nullptr ? &View.ViewState->LightGrid : nullptr,
 				/*bThreadGroupPerCell*/ true,
 				/*bThreadGroupSize32*/ false,
@@ -1217,6 +1267,7 @@ FComputeLightGridOutput FSceneRenderer::ComputeLightGrid(FRDGBuilder& GraphBuild
 			LightViewSpacePositionAndRadiusSRV,
 			LightViewSpaceDirAndPreprocAngleSRV,
 			LightViewSpaceRectPlanesSRV,
+			IndirectionIndicesSRV,
 			View.ViewState != nullptr ? &View.ViewState->LightGrid : nullptr,
 			NumThreadsPerCell > 1,
 			NumThreadsPerCell == 32,
diff --git a/Engine/Source/Runtime/Renderer/Private/ScenePrivate.h b/Engine/Source/Runtime/Renderer/Private/ScenePrivate.h
--- a/Engine/Source/Runtime/Renderer/Private/ScenePrivate.h
+++ b/Engine/Source/Runtime/Renderer/Private/ScenePrivate.h
@@ -1016,7 +1016,7 @@ public:
 
 	// Map from Light ID in GPU Scene to index in the View's ForwardLightData array
 	// This is stored in ViewState so we can access previous frame mapping
-	TMap<int32, int32> LightSceneIdToLocalLightIndex;
+	TMap<int32, int32> LightSceneIdToForwardLightIndex;
 
 	// Pre-computed filter in spectral (i.e. FFT) domain along with data to determine if we need to up date it
 	struct {
diff --git a/Engine/Source/Runtime/Renderer/Private/SceneRendering.h b/Engine/Source/Runtime/Renderer/Private/SceneRendering.h
--- a/Engine/Source/Runtime/Renderer/Private/SceneRendering.h
+++ b/Engine/Source/Runtime/Renderer/Private/SceneRendering.h
@@ -656,7 +656,8 @@ public:
 const int32 GMaxForwardShadowCascades = 4;
 
 BEGIN_GLOBAL_SHADER_PARAMETER_STRUCT_WITH_CONSTRUCTOR(FForwardLightUniformParameters, )
-	SHADER_PARAMETER(uint32,NumLocalLights)
+	SHADER_PARAMETER(uint32, NumLocalLights)
+	SHADER_PARAMETER(uint32, NumDirectionalLights)
 	SHADER_PARAMETER(uint32, NumReflectionCaptures)
 	SHADER_PARAMETER(uint32, HasDirectionalLight)
 	SHADER_PARAMETER(uint32, NumGridCells)
@@ -692,10 +693,11 @@ BEGIN_GLOBAL_SHADER_PARAMETER_STRUCT_WITH_CONSTRUCTOR(FForwardLightUniformParame
 	SHADER_PARAMETER_SAMPLER(SamplerState, ShadowmapSampler)
 	SHADER_PARAMETER_TEXTURE(Texture2D, DirectionalLightStaticShadowmap)
 	SHADER_PARAMETER_SAMPLER(SamplerState, StaticShadowmapSampler)
-	SHADER_PARAMETER_RDG_BUFFER_SRV(StructuredBuffer<float4>, ForwardLocalLightBuffer)
+	SHADER_PARAMETER_RDG_BUFFER_SRV(StructuredBuffer<float4>, ForwardLightBuffer)
 	SHADER_PARAMETER_RDG_BUFFER_SRV(StructuredBuffer<uint>, NumCulledLightsGrid)
 	SHADER_PARAMETER_RDG_BUFFER_SRV(StructuredBuffer<uint>, CulledLightDataGrid32Bit)
 	SHADER_PARAMETER_RDG_BUFFER_SRV(Buffer<uint>, CulledLightDataGrid16Bit)
+	SHADER_PARAMETER_RDG_BUFFER_SRV(StructuredBuffer<uint>, DirectionalLightIndices)
 	SHADER_PARAMETER_RDG_BUFFER_SRV(StructuredBuffer<FLightViewData>, LightViewData)
 END_GLOBAL_SHADER_PARAMETER_STRUCT()
 
diff --git a/Engine/Source/Runtime/Renderer/Private/TranslucentLighting.cpp b/Engine/Source/Runtime/Renderer/Private/TranslucentLighting.cpp
--- a/Engine/Source/Runtime/Renderer/Private/TranslucentLighting.cpp
+++ b/Engine/Source/Runtime/Renderer/Private/TranslucentLighting.cpp
@@ -1168,10 +1168,10 @@ void FTranslucentLightInjectionCollector::AddLightForInjection(
 		// If it would otherwise be supported, see if we can find the local light index in the light grid
 		if (bSupportsBatching && View.ForwardLightingResources.ForwardLightUniformBuffer && View.ViewState)
 		{
-			const int32* Value = View.ViewState->LightSceneIdToLocalLightIndex.Find(LightSceneInfo.Id);
+			const int32* Value = View.ViewState->LightSceneIdToForwardLightIndex.Find(LightSceneInfo.Id);
 			if (Value)
 			{
-				const int32 LocalLightIndex = *Value;
+				const int32 ForwardLightIndex = *Value;
 
 				FPerViewData& InjectionData = InjectionDataPerView[ViewIndex];
 				// There shouldn't be any duplication here
@@ -1182,8 +1182,8 @@ void FTranslucentLightInjectionCollector::AddLightForInjection(
 					if (VolumeBounds[VolumeCascadeIndex].IsValid())
 					{
 						TBitArray<SceneRenderingAllocator>& Bits = InjectionData.BatchedLocalLights[VolumeCascadeIndex];
-						Bits.PadToNum(LocalLightIndex + 1, false);
-						Bits[LocalLightIndex] = true;
+						Bits.PadToNum(ForwardLightIndex + 1, false);
+						Bits[ForwardLightIndex] = true;
 					}
 				}
 
