Fix crash due to DirectionalLightIndices not being initialized in SetupDummyForwardLightUniformParameters(...).

diff --git a/Engine/Source/Runtime/Renderer/Private/LightGridInjection.cpp b/Engine/Source/Runtime/Renderer/Private/LightGridInjection.cpp
--- a/Engine/Source/Runtime/Renderer/Private/LightGridInjection.cpp
+++ b/Engine/Source/Runtime/Renderer/Private/LightGridInjection.cpp
@@ -220,6 +220,9 @@ void SetupDummyForwardLightUniformParameters(FRDGBuilder& GraphBuilder, FForward
 
 	ForwardLightUniformParameters.bAffectsTranslucentLighting = 0;
 
+	FRDGBufferRef DirectionalLightIndicesBuffer = GSystemTextures.GetDefaultStructuredBuffer(GraphBuilder, sizeof(uint32));
+	ForwardLightUniformParameters.DirectionalLightIndices = GraphBuilder.CreateSRV(DirectionalLightIndicesBuffer);
+
 	FRDGBufferRef LightViewDataBuffer = GSystemTextures.GetDefaultStructuredBuffer(GraphBuilder, sizeof(FLightViewData));
 	ForwardLightUniformParameters.LightViewData = GraphBuilder.CreateSRV(LightViewDataBuffer);
 }
