Fix RHI valiation errors for virtual texture page table and physical textures.
The validation error occurs only when an RVT is rendering a page using a material that reads the same physical texture.
In this case the RDG transition scheduling doesn't know that the virtual texture physical and page table textures are in use in the RVT render pass (since that is bound using raw RHI objects in the material uniforms).
So we need to encourage RDG to transition at the last moment before the page table writes or the physical texture compress.

diff --git a/Engine/Source/Runtime/Renderer/Private/VT/RuntimeVirtualTextureRender.cpp b/Engine/Source/Runtime/Renderer/Private/VT/RuntimeVirtualTextureRender.cpp
--- a/Engine/Source/Runtime/Renderer/Private/VT/RuntimeVirtualTextureRender.cpp
+++ b/Engine/Source/Runtime/Renderer/Private/VT/RuntimeVirtualTextureRender.cpp
@@ -138,6 +138,9 @@ namespace RuntimeVirtualTexture
 			SHADER_PARAMETER_STRUCT_REF(FViewUniformShaderParameters, View)
 			SHADER_PARAMETER_RDG_UNIFORM_BUFFER(FSceneUniformParameters, Scene)
 			SHADER_PARAMETER_STRUCT_INCLUDE(FInstanceCullingDrawParams, InstanceCullingDrawParams)
+			RDG_TEXTURE_ACCESS(PhysicalTexture0, ERHIAccess::SRVGraphics)
+			RDG_TEXTURE_ACCESS(PhysicalTexture1, ERHIAccess::SRVGraphics)
+			RDG_TEXTURE_ACCESS(PhysicalTexture2, ERHIAccess::SRVGraphics)
 			RENDER_TARGET_BINDING_SLOTS()
 		END_SHADER_PARAMETER_STRUCT()
 
@@ -1579,10 +1582,20 @@ namespace RuntimeVirtualTexture
 		// Draw Pass
 		if (GraphSetup.bRenderPass)
 		{
+			RDG_EVENT_SCOPE(GraphBuilder, "RenderPass");
+
 			ERenderTargetLoadAction LoadAction = bClearTextures ? ERenderTargetLoadAction::EClear : ERenderTargetLoadAction::ENoAction;
 			FShader_VirtualTextureMaterialDraw::FParameters* PassParameters = GraphBuilder.AllocParameters<FShader_VirtualTextureMaterialDraw::FParameters>();
 			PassParameters->View = View->ViewUniformBuffer;
 			PassParameters->Scene = SceneRenderer->GetSceneUniformBufferRef(GraphBuilder);
+			// When direct aliasing the output physical texture for compression, the texture gets managed by RDG.
+			// The texture needs to be transitioned to UAV during the compression pass.
+			// But we may also be accessing the physical texture in a material in this render pass (using the raw RHI SRV).
+			// So we need to declare the physical textures as SRV here to avoid RDG hoisting their UAV transition to before this pass.
+			const bool bFixEarlyUAVTransition = GraphSetup.bCompressPass && GraphSetup.bDirectAliasing;
+			PassParameters->PhysicalTexture0 = bFixEarlyUAVTransition ? GraphSetup.CompressTexture0 : nullptr;
+			PassParameters->PhysicalTexture1 = bFixEarlyUAVTransition ? GraphSetup.CompressTexture1 : nullptr;
+			PassParameters->PhysicalTexture2 = bFixEarlyUAVTransition ? GraphSetup.CompressTexture2 : nullptr;
 			PassParameters->RenderTargets[0] = GraphSetup.RenderTexture0 ? FRenderTargetBinding(GraphSetup.RenderTexture0, LoadAction) : FRenderTargetBinding();
 			PassParameters->RenderTargets[1] = GraphSetup.RenderTexture1 ? FRenderTargetBinding(GraphSetup.RenderTexture1, LoadAction) : FRenderTargetBinding();
 			PassParameters->RenderTargets[2] = GraphSetup.RenderTexture2 ? FRenderTargetBinding(GraphSetup.RenderTexture2, LoadAction) : FRenderTargetBinding();
diff --git a/Engine/Source/Runtime/Renderer/Private/VT/VirtualTextureSpace.cpp b/Engine/Source/Runtime/Renderer/Private/VT/VirtualTextureSpace.cpp
--- a/Engine/Source/Runtime/Renderer/Private/VT/VirtualTextureSpace.cpp
+++ b/Engine/Source/Runtime/Renderer/Private/VT/VirtualTextureSpace.cpp
@@ -337,6 +337,8 @@ void FVirtualTextureSpace::AllocateTextures(FRDGBuilder& GraphBuilder)
 			}
 
 			TextureEntry.RenderTarget = GraphBuilder.ConvertToExternalTexture(DstTexture);
+			GraphBuilder.UseExternalAccessMode(DstTexture, ERHIAccess::SRVMask, ERHIPipeline::All);
+
 			RHIUpdateTextureReference(TextureEntry.TextureReferenceRHI, TextureEntry.RenderTarget->GetRHI());
 		}
 
@@ -486,9 +488,12 @@ void FVirtualTextureSpace::ApplyUpdates(FVirtualTextureSystem* System, FRDGBuild
 	{
 		const uint32 TextureIndex = LayerIndex / LayersPerPageTableTexture;
 		const uint32 LayerInTexture = LayerIndex % LayersPerPageTableTexture;
-
 		FTextureEntry& PageTableEntry = PageTable[TextureIndex];
-		FRDGTextureRef PageTableTexture = GraphBuilder.RegisterExternalTexture(PageTableEntry.RenderTarget);
+
+		// Use ForceImmediateFirstBarrier to prevent PageTableTexture RVT transition from being hoisted to earlier in the render graph.
+		// This can be a problem if the raw RHI object for the page table is being used earlier in the graph.
+		// An example where that can happen is if we sample virtual textures while rendering runtime virtual texture pages.
+		FRDGTextureRef PageTableTexture = GraphBuilder.RegisterExternalTexture(PageTableEntry.RenderTarget, ERDGTextureFlags::ForceImmediateFirstBarrier);
 		GraphBuilder.UseInternalAccessMode(PageTableTexture);
 
 		// Use color write mask to update the proper page table entry for this layer
