Fix crash encountered in commandlet rendering

diff --git a/Engine/Source/Runtime/Renderer/Private/BlueNoise.cpp b/Engine/Source/Runtime/Renderer/Private/BlueNoise.cpp
--- a/Engine/Source/Runtime/Renderer/Private/BlueNoise.cpp
+++ b/Engine/Source/Runtime/Renderer/Private/BlueNoise.cpp
@@ -26,27 +26,34 @@ FBlueNoiseParameters GetBlueNoiseParameters()
 {
 	FBlueNoiseParameters Out;
 	check(GEngine);
-	check(GEngine->BlueNoiseScalarTexture && GEngine->BlueNoiseVec2Texture);
 
-	Out.Dimensions = FIntVector(
-		GEngine->BlueNoiseScalarTexture->GetSizeX(),
-		GEngine->BlueNoiseScalarTexture->GetSizeX(),
-		GEngine->BlueNoiseScalarTexture->GetSizeY() / FMath::Max<int32>(1, GEngine->BlueNoiseScalarTexture->GetSizeX()));
+	if (GEngine->BlueNoiseScalarTexture)
+	{
+		Out.Dimensions = FIntVector(
+			GEngine->BlueNoiseScalarTexture->GetSizeX(),
+			GEngine->BlueNoiseScalarTexture->GetSizeX(),
+			GEngine->BlueNoiseScalarTexture->GetSizeY() / FMath::Max<int32>(1, GEngine->BlueNoiseScalarTexture->GetSizeX()));
 
-	Out.ModuloMasks = FIntVector(
-		(1u << FMath::FloorLog2(Out.Dimensions.X)) - 1,
-		(1u << FMath::FloorLog2(Out.Dimensions.Y)) - 1,
-		(1u << FMath::FloorLog2(Out.Dimensions.Z)) - 1);
+		Out.ModuloMasks = FIntVector(
+			(1u << FMath::FloorLog2(Out.Dimensions.X)) - 1,
+			(1u << FMath::FloorLog2(Out.Dimensions.Y)) - 1,
+			(1u << FMath::FloorLog2(Out.Dimensions.Z)) - 1);
 
-	check((Out.ModuloMasks.X + 1)  == Out.Dimensions.X
-		&&(Out.ModuloMasks.Y + 1) == Out.Dimensions.Y
-		&&(Out.ModuloMasks.Z + 1) == Out.Dimensions.Z);
+		check((Out.ModuloMasks.X + 1)  == Out.Dimensions.X
+			&&(Out.ModuloMasks.Y + 1) == Out.Dimensions.Y
+			&&(Out.ModuloMasks.Z + 1) == Out.Dimensions.Z);
 
-	Out.ScalarTexture = GEngine->BlueNoiseScalarTexture->GetResource()->TextureRHI;
-	Out.Vec2Texture = GSystemTextures.BlackDummy->GetRHI();
-	if (GEngine->BlueNoiseVec2Texture != nullptr)
+		Out.ScalarTexture = GEngine->BlueNoiseScalarTexture->GetResource()->TextureRHI;
+
+		Out.Vec2Texture = GSystemTextures.BlackDummy->GetRHI();
+		if (GEngine->BlueNoiseVec2Texture != nullptr)
+		{
+			Out.Vec2Texture = GEngine->BlueNoiseVec2Texture->GetResource()->TextureRHI;
+		}
+	}
+	else
 	{
-		Out.Vec2Texture = GEngine->BlueNoiseVec2Texture->GetResource()->TextureRHI;
+		Out = GetBlueNoiseDummyParameters();
 	}
 	return Out;
 }
