Reinstate Heterogeneous Volumes.

Disable extended upsampling chain.

diff --git a/Engine/Source/Runtime/Renderer/Private/HeterogeneousVolumes/HeterogeneousVolumes.cpp b/Engine/Source/Runtime/Renderer/Private/HeterogeneousVolumes/HeterogeneousVolumes.cpp
--- a/Engine/Source/Runtime/Renderer/Private/HeterogeneousVolumes/HeterogeneousVolumes.cpp
+++ b/Engine/Source/Runtime/Renderer/Private/HeterogeneousVolumes/HeterogeneousVolumes.cpp
@@ -57,7 +57,7 @@ static TAutoConsoleVariable<int32> CVarHeterogeneousVolumesUpsample(
 	TEXT("0: Off\n")
 	TEXT("1: Nearest Neighbor\n")
 	TEXT("2: Bilinear"),
-	ECVF_RenderThreadSafe
+	ECVF_Scalability | ECVF_RenderThreadSafe
 );
 
 static TAutoConsoleVariable<int32> CVarHeterogeneousVolumesFilter(
