Disable TwoLevel lightgrid Ingestion

diff --git a/Engine/Source/Runtime/Renderer/Private/LightGridInjection.cpp b/Engine/Source/Runtime/Renderer/Private/LightGridInjection.cpp
--- a/Engine/Source/Runtime/Renderer/Private/LightGridInjection.cpp
+++ b/Engine/Source/Runtime/Renderer/Private/LightGridInjection.cpp
@@ -145,7 +145,7 @@ static TAutoConsoleVariable<int32> CVarLightCullingWorkloadDistributionMode(
 
 static TAutoConsoleVariable<bool> CVarLightCullingTwoLevel(
 	TEXT("r.LightCulling.TwoLevel"),
-	true,
+	false,
 	TEXT("Whether to build light grid in two passes."),
 	ECVF_RenderThreadSafe
 );
