Hacky workaround for BlueNoise params crash during HLOD rebuild.

diff --git a/Engine/Source/Runtime/Renderer/Private/BlueNoise.cpp b/Engine/Source/Runtime/Renderer/Private/BlueNoise.cpp
--- a/Engine/Source/Runtime/Renderer/Private/BlueNoise.cpp
+++ b/Engine/Source/Runtime/Renderer/Private/BlueNoise.cpp
@@ -27,6 +27,15 @@ FBlueNoiseParameters GetBlueNoiseParameters()
 	FBlueNoiseParameters Out;
 	check(GEngine);
 
+	// Hacky fix for BlueNoise not being loaded before use on some code paths. Erring on the side of caution, e.g. GEngine shouldn't realistically be null
+	if (!GEngine->BlueNoiseScalarTexture || !GEngine->BlueNoiseVec2Texture
+		|| GEngine->BlueNoiseScalarTexture->GetResource() == nullptr
+		|| GEngine->BlueNoiseVec2Texture->GetResource() == nullptr
+		)
+	{
+		return Out;
+	}
+	
 	if (GEngine->BlueNoiseScalarTexture)
 	{
 		Out.Dimensions = FIntVector(
