Mobile: Submit occlusion feedback copy operation as early as possible when depth pre-pass is enabled

diff --git a/Engine/Source/Runtime/Renderer/Private/MobileShadingRenderer.cpp b/Engine/Source/Runtime/Renderer/Private/MobileShadingRenderer.cpp
--- a/Engine/Source/Runtime/Renderer/Private/MobileShadingRenderer.cpp
+++ b/Engine/Source/Runtime/Renderer/Private/MobileShadingRenderer.cpp
@@ -878,6 +878,7 @@ void FMobileSceneRenderer::RenderFullDepthPrepass(FRDGBuilder& GraphBuilder, TAr
 		}
 	}
 
+	EndOcclusionScope(GraphBuilder, Views);
 	FenceOcclusionTests(GraphBuilder);
 }
 
@@ -1424,10 +1425,9 @@ void FMobileSceneRenderer::Render(FRDGBuilder& GraphBuilder)
 			RenderForward(GraphBuilder, ViewFamilyTexture, SceneTextures, DBufferTextures);
 		}
 
-		EndOcclusionScope(GraphBuilder, Views);
-
 		if (!bIsFullDepthPrepassEnabled)
 		{
+			EndOcclusionScope(GraphBuilder, Views);
 			FenceOcclusionTests(GraphBuilder);
 		}
 
@@ -2400,6 +2400,11 @@ int32 FMobileSceneRenderer::ComputeNumOcclusionQueriesToBatch() const
 // Whether we need a separate render-passes for translucency, decals etc
 bool FMobileSceneRenderer::RequiresMultiPass(int32 NumMSAASamples, EShaderPlatform ShaderPlatform)
 {
+	if (!MobileAllowFramebufferFetch(ShaderPlatform))
+	{
+		return true;
+	}
+	
 	// Vulkan uses subpasses
 	if (IsVulkanPlatform(ShaderPlatform))
 	{
