Report bloom convolution memory inefficiency once

diff --git a/Engine/Source/Runtime/Renderer/Private/PostProcess/PostProcessFFTBloom.cpp b/Engine/Source/Runtime/Renderer/Private/PostProcess/PostProcessFFTBloom.cpp
--- a/Engine/Source/Runtime/Renderer/Private/PostProcess/PostProcessFFTBloom.cpp
+++ b/Engine/Source/Runtime/Renderer/Private/PostProcess/PostProcessFFTBloom.cpp
@@ -651,7 +651,10 @@ void InitDomainAndGetKernel(
 				{
 					WarningTexturePathName = "(" + View.FinalPostProcessSettings.BloomConvolutionTexture.GetPathName() + ")";
 				}
-				UE_LOG(LogRenderer, Warning, TEXT("The FPostProcessSettings::BloomConvolutionTexture %s could have it's resolution lowered by a factor of %d to save memory."), *WarningTexturePathName, RecommendedKernelDownsscale);
+
+				static bool bLogged = false;
+				UE_CLOG(!bLogged, LogRenderer, Warning, TEXT("The FPostProcessSettings::BloomConvolutionTexture %s could have it's resolution lowered by a factor of %d to save memory."), *WarningTexturePathName, RecommendedKernelDownsscale);
+				bLogged = true;
 			}
 		}
 
