Added r.DOF.Gather.DistinctionLimit.
Pixels that are bright and have a large CoC should be scattered instead of gathered, to avoid gather noise.
If they aren't, e.g. due to performance reasons, we may prefer to cull them instead of gathering them anyway.
This will bias the image, as bokeh may appear slightly darker, but helps clean up noise in large bokeh.

diff --git a/Engine/Shaders/Private/DiaphragmDOF/DOFReduce.usf b/Engine/Shaders/Private/DiaphragmDOF/DOFReduce.usf
--- a/Engine/Shaders/Private/DiaphragmDOF/DOFReduce.usf
+++ b/Engine/Shaders/Private/DiaphragmDOF/DOFReduce.usf
@@ -67,6 +67,7 @@ uint MaxScatteringGroupCount;
 float2 MaxInputBufferUV;
 float MinScatteringCocRadius;
 float PreProcessingToProcessingCocRadiusFactor;
+float BokehGatherDistinctionLimit;
 
 float4 GatherInputSize;
 Texture2D GatherInput_SceneColor;
@@ -244,6 +245,24 @@ void OutputMipLevel(const uint MipLevel, uint2 MipPixelPos, float4 Color, float
 
 //------------------------------------------------------- ENTRY POINT
 
+// Pixels that are bright and have a large CoC should be scattered instead of gathered, to avoid gather noise.
+// If they aren't, e.g. due to performance reasons, we may prefer to cull them instead of gathering them anyway.
+// This will bias the image, as bokeh may appear slightly darker, but helps clean up noise in large bokeh.
+float ApplyBokehGatherDistinctionFilter(float CocRadius, float4 Color, float FrameExposureScale)
+{
+	if (BokehGatherDistinctionLimit != 0)
+	{
+		// see also ComputeLuminanceOnlyBasedScatterFactor
+		float PerceivedLuma = Luma4(Color.rgb) * FrameExposureScale;
+		float BokehDistinction = abs(CocRadius) * PerceivedLuma;
+		return smoothstep(BokehGatherDistinctionLimit + 200, BokehGatherDistinctionLimit - 200, BokehDistinction);
+	}
+	else
+	{
+		return 1;
+	}
+}
+
 [numthreads(THREADGROUP_TILE_SIZE, THREADGROUP_TILE_SIZE, 1)]
 void ReduceCS(
 	uint2 GroupId : SV_GroupID,
@@ -383,6 +402,11 @@ void ReduceCS(
 		const float ScatterFactorThreshold = 0.01;
 		bool bScatterThisColor = ScatterFactor > ScatterFactorThreshold;
 		bool bIsForeground = IsForeground(CocRadius);
+
+		if (!bScatterThisColor)
+		{
+			GatheredColor *= ApplyBokehGatherDistinctionFilter(CocRadius, GatheredColor, View.OneOverPreExposure);
+		}
 			
 		// Whether scattering group actually scatter or not.
 		uint ForegroundScatterGroupMask = 0;
diff --git a/Engine/Source/Runtime/Renderer/Private/PostProcess/DiaphragmDOF.cpp b/Engine/Source/Runtime/Renderer/Private/PostProcess/DiaphragmDOF.cpp
--- a/Engine/Source/Runtime/Renderer/Private/PostProcess/DiaphragmDOF.cpp
+++ b/Engine/Source/Runtime/Renderer/Private/PostProcess/DiaphragmDOF.cpp
@@ -44,6 +44,18 @@ TAutoConsoleVariable<int32> CVarAccumulatorQuality(
 	TEXT("Controles the quality of the gathering accumulator.\n"),
 	ECVF_Scalability | ECVF_RenderThreadSafe);
 
+TAutoConsoleVariable<float> CVarGatherDistinctionLimit(
+	TEXT("r.DOF.Gather.DistinctionLimit"),
+	1000.0,
+	TEXT("Bokeh that are not scattered but exceed this distinction limit are discarded instead of gathered.\n")
+	TEXT("Distinction is calculcated based on perceived brightness and radius.\n")
+	TEXT("This is used to reduce noise in gathered bokeh, without drawing additional scattered bokeh.\n")
+	TEXT("This can cause bokeh to become darker than they should be, which is a compromise for noise reduction.\n")
+	TEXT("See also: r.DOF.Scatter.NeighborCompareMaxColor \n")
+	TEXT("Set to 0 to disable.\n")
+	TEXT("(default: 1000)"),
+	ECVF_Scalability | ECVF_RenderThreadSafe);
+
 TAutoConsoleVariable<int32> CVarEnableGatherBokehSettings(
 	TEXT("r.DOF.Gather.EnableBokehSettings"),
 	1,
@@ -138,8 +150,10 @@ TAutoConsoleVariable<int32> CVarDOFTemporalAAQuality(
 TAutoConsoleVariable<float> CVarScatterNeighborCompareMaxColor(
 	TEXT("r.DOF.Scatter.NeighborCompareMaxColor"),
 	10,
-	TEXT("Controles the linear color clamping upperbound applied before color of pixel and neighbors are compared.")
-	TEXT(" To low, and you may not scatter enough; to high you may scatter unnecessarily too much in highlights")
+	TEXT("Threshold used when deciding if a pixel should get scattered if it has a neighbor that is already getting scattered. \n")
+	TEXT("The linear color of the pixel and its neighbors are compared. If the difference is lower than this threshold, it is not scattered. \n")
+	TEXT("Too low, and you may not scatter enough, causing quality issues. Too high you may scatter unnecessarily too much in highlights, causing overdraw. \n")
+	TEXT("See also: r.DOF.Gather.DistinctionLimit \n")
 	TEXT(" (Default: 10)."),
 	ECVF_RenderThreadSafe);
 
@@ -919,6 +933,7 @@ class FDiaphragmDOFReduceCS : public FDiaphragmDOFShader
 		SHADER_PARAMETER(FVector2f, MaxInputBufferUV)
 		SHADER_PARAMETER(int32, MaxScatteringGroupCount)
 		SHADER_PARAMETER(float, PreProcessingToProcessingCocRadiusFactor)
+		SHADER_PARAMETER(float, BokehGatherDistinctionLimit)
 		SHADER_PARAMETER(float, MinScatteringCocRadius)
 		SHADER_PARAMETER(float, NeighborCompareMaxColor)
 		SHADER_PARAMETER(float, CocSqueeze)
@@ -2058,6 +2073,7 @@ bool DiaphragmDOF::AddPasses(
 				(PreprocessViewSize.Y - 0.5f) / SrcSize.Y);
 			PassParameters->MaxScatteringGroupCount = MaxScatteringGroupCount;
 			PassParameters->PreProcessingToProcessingCocRadiusFactor = GatheringCocRadiusBasis / EncodedCocRadiusBasis;
+			PassParameters->BokehGatherDistinctionLimit = CVarGatherDistinctionLimit.GetValueOnRenderThread();
 			PassParameters->MinScatteringCocRadius = MinScatteringCocRadius;
 			PassParameters->NeighborCompareMaxColor = CVarScatterNeighborCompareMaxColor.GetValueOnRenderThread();
 			PassParameters->CocSqueeze = CocModel.Squeeze;
