Fixed typo in scene occlusion CPU profile scope.

diff --git a/Engine/Source/Runtime/Renderer/Private/SceneOcclusion.cpp b/Engine/Source/Runtime/Renderer/Private/SceneOcclusion.cpp
--- a/Engine/Source/Runtime/Renderer/Private/SceneOcclusion.cpp
+++ b/Engine/Source/Runtime/Renderer/Private/SceneOcclusion.cpp
@@ -1072,7 +1072,7 @@ static void TrimAllOcclusionHistory(TArrayView<FViewInfo> Views)
 
 static void AllocateOcclusionTests(FViewOcclusionQueriesPerView& QueriesPerView, const FScene* Scene, TArrayView<const FVisibleLightInfo> VisibleLightInfos, TArrayView<FViewInfo> Views)
 {
-	SCOPED_NAMED_EVENT(FSceneRenderer_AllocateOcclusionTestsOcclusionTests, FColor::Emerald);
+	SCOPED_NAMED_EVENT(FSceneRenderer_AllocateOcclusionTests, FColor::Emerald);
 
 	const ERHIFeatureLevel::Type FeatureLevel = Scene->GetFeatureLevel();
 	const int32 NumBufferedFrames = FOcclusionQueryHelpers::GetNumBufferedFrames(FeatureLevel);
