Implemented new ISceneRenderBuilder interface with a basic implementation that merges scene renders together under one FSceneRenderer::RenderThread{Begin, End} instance.
- Supports interleaving game thread commands with scene render commands.
- Scene render commands are issued in sequence. Currently a unique FRDGBuilder instance is used for each one.
- Game thread commands are used to handle edge cases where game state is saved / restored after scene captures.
- A scene render group scope can be used to wrap render commands. Multiple renderers can be within one group (for example, 6 faces of a cube map).
- Groups support Gpu Dump / Capture via a group flag.

diff --git a/Engine/Source/Runtime/Engine/Private/SceneRenderBuilderInterface.cpp b/Engine/Source/Runtime/Engine/Private/SceneRenderBuilderInterface.cpp
new file mode 100644
--- /dev/null
+++ b/Engine/Source/Runtime/Engine/Private/SceneRenderBuilderInterface.cpp
@@ -0,0 +1,10 @@
+// Copyright Epic Games, Inc. All Rights Reserved.
+
+#include "SceneRenderBuilderInterface.h"
+#include "SceneRendererInterface.h"
+#include "EngineModule.h"
+
+TUniquePtr<ISceneRenderBuilder> ISceneRenderBuilder::Create(FSceneInterface* SceneInterface)
+{
+	return GetRendererModule().CreateSceneRenderBuilder(SceneInterface);
+}
\ No newline at end of file
diff --git a/Engine/Source/Runtime/Engine/Public/SceneRenderBuilderInterface.h b/Engine/Source/Runtime/Engine/Public/SceneRenderBuilderInterface.h
new file mode 100644
--- /dev/null
+++ b/Engine/Source/Runtime/Engine/Public/SceneRenderBuilderInterface.h
@@ -0,0 +1,113 @@
+// Copyright Epic Games, Inc. All Rights Reserved.
+
+#pragma once
+
+#include "CoreMinimal.h"
+#include "Experimental/ConcurrentLinearAllocator.h"
+#include "RenderingThread.h"
+
+class FRHICommandListImmediate;
+class FRDGBuilder;
+class FSceneRenderer;
+class FSceneInterface;
+
+enum class ESceneRenderGroupFlags : uint8
+{
+	None       = 0,
+
+	// This group of renderers will perform a GPU capture.
+	GpuCapture = 1 << 1,
+
+	// This group of renderers will perform a GPU dump.
+	GpuDump    = 1 << 2
+};
+ENUM_CLASS_FLAGS(ESceneRenderGroupFlags);
+
+// A setup of inputs passed into a scene render callback function.
+struct FSceneRenderFunctionInputs
+{
+	FSceneRenderFunctionInputs(FSceneRenderer* InRenderer, const TCHAR* InName, const TCHAR* InFullPath)
+		: Renderer(InRenderer)
+		, Name(InName)
+		, FullPath(InFullPath)
+	{}
+
+	FSceneRenderer* Renderer;
+	const TCHAR* Name;
+	const TCHAR* FullPath;
+};
+
+using FSceneRenderFunction = TUniqueFunction<void(FRDGBuilder&, const FSceneRenderFunctionInputs&)>;
+
+// A game side builder interface that collects scene renderers to process as one workload.
+class ISceneRenderBuilder
+{
+public:
+	virtual ~ISceneRenderBuilder() = default;
+
+	// Creates an instance of a scene render builder using a specified scene. The scene must be
+	// valid and all scene renderers must reference the same scene instance.
+	static ENGINE_API TUniquePtr<ISceneRenderBuilder> Create(FSceneInterface* SceneInterface);
+
+	/** Adds a command that executes on the game thread and is pipelined with scene render commands on the render thread.
+	 *  This is useful to process game thread side logic in between scene renders which may insert render commands. Though
+	 *  be mindful of how commands may be deferred until RDG execution in certain configurations.
+	 */
+	virtual void AddCommand(TUniqueFunction<void()>&& Function) = 0;
+
+	// Utility method that enqueues a render command that is interleaved with scene render commands.
+	void AddRenderCommand(TUniqueFunction<void(FRHICommandListImmediate&)>&& Function)
+	{
+		AddCommand([Function = MoveTemp(Function)]() mutable
+		{
+			ENQUEUE_RENDER_COMMAND(SceneRenderBuilder_AddRenderCommand)(MoveTemp(Function));
+		});
+	}
+
+	// Adds a command to render a scene renderer on the render thread using RDG.
+	virtual void AddRenderer(FSceneRenderer* Renderer, FString&& Name, FSceneRenderFunction&& Function) = 0;
+
+	void AddRenderer(FSceneRenderer* Renderer, FSceneRenderFunction&& Function)
+	{
+		AddRenderer(Renderer, {}, MoveTemp(Function));
+	}
+
+	//////////////////////////////////////////////////////////////////////////
+
+	// Call to execute all enqueued commands on the render thread. Resets the builder back to empty.
+	virtual void Execute() = 0;
+
+	// Returns a linear allocator for allocating objects that can be used safely in command functions.
+	virtual FConcurrentLinearBulkObjectAllocator& GetAllocator() = 0;
+
+private:
+	friend class FSceneRenderGroupScope;
+	virtual void BeginGroup(FString&& Name, ESceneRenderGroupFlags Flags) = 0;
+	virtual void EndGroup() = 0;
+};
+
+// A scope class used to construct a scene render group.
+class FSceneRenderGroupScope
+{
+	ISceneRenderBuilder& SceneRenderBuilder;
+
+public:
+	FSceneRenderGroupScope(ISceneRenderBuilder& InSceneRenderBuilder, FString&& Name, ESceneRenderGroupFlags Flags)
+		: SceneRenderBuilder(InSceneRenderBuilder)
+	{
+		SceneRenderBuilder.BeginGroup(Forward<FString&&>(Name), Flags);
+	}
+
+	~FSceneRenderGroupScope()
+	{
+		SceneRenderBuilder.EndGroup();
+	}
+};
+
+// A scope macro used to construct a scene render group scope.
+#define SCENE_RENDER_GROUP_SCOPE(SceneRenderBuilder, Name, Flags) \
+	FSceneRenderGroupScope PREPROCESSOR_JOIN(__SceneRenderGroup_ScopeRef_,__LINE__)( \
+		  (SceneRenderBuilder)                                                       \
+		, Forward<FString&&>(Name)                                                   \
+		, Flags                                                                      \
+	)
\ No newline at end of file
diff --git a/Engine/Source/Runtime/RenderCore/Public/RendererInterface.h b/Engine/Source/Runtime/RenderCore/Public/RendererInterface.h
--- a/Engine/Source/Runtime/RenderCore/Public/RendererInterface.h
+++ b/Engine/Source/Runtime/RenderCore/Public/RendererInterface.h
@@ -42,6 +42,7 @@ class FSceneUniformBuffer;
 class FBatchedPrimitiveParameters;
 class ISceneRenderer;
 class IStateStreamManager;
+class ISceneRenderBuilder;
 
 namespace Nanite
 {
@@ -880,5 +881,7 @@ public:
 	virtual const FGlobalDistanceFieldParameterData* GetGlobalDistanceFieldParameterData(const FSceneView& View) = 0;
 	virtual void RequestStaticMeshUpdate(FPrimitiveSceneInfo* Info) = 0;
 	virtual void AddMeshBatchToGPUScene(FGPUScenePrimitiveCollector* Collector, FMeshBatch& MeshBatch) = 0;
+
+	virtual TUniquePtr<ISceneRenderBuilder> CreateSceneRenderBuilder(FSceneInterface* Scene) = 0;
 };
 
diff --git a/Engine/Source/Runtime/Renderer/Private/RendererModule.h b/Engine/Source/Runtime/Renderer/Private/RendererModule.h
--- a/Engine/Source/Runtime/Renderer/Private/RendererModule.h
+++ b/Engine/Source/Runtime/Renderer/Private/RendererModule.h
@@ -158,6 +158,8 @@ public:
 	virtual void FreeStateStream(IStateStreamManager* Manager) override;
 #endif
 
+	virtual TUniquePtr<ISceneRenderBuilder> CreateSceneRenderBuilder(FSceneInterface* Scene) override;
+
 private:
 	TSet<FSceneInterface*> AllocatedScenes;
 	FOnPostOpaqueRender PostOpaqueRenderDelegate;
diff --git a/Engine/Source/Runtime/Renderer/Private/RendererScene.cpp b/Engine/Source/Runtime/Renderer/Private/RendererScene.cpp
--- a/Engine/Source/Runtime/Renderer/Private/RendererScene.cpp
+++ b/Engine/Source/Runtime/Renderer/Private/RendererScene.cpp
@@ -80,6 +80,7 @@
 #include "RectLightSceneProxy.h"
 #include "RectLightTextureManager.h"
 #include "RenderCore.h"
+#include "SceneRenderBuilder.h"
 #include "IESTextureManager.h"
 #include "Materials/MaterialRenderProxy.h"
 #include "ProfilingDebugging/AssetMetadataTrace.h"
@@ -7267,6 +7268,15 @@ FSceneViewStateInterface* FRendererModule::AllocateViewState(ERHIFeatureLevel::T
 	return new FSceneViewState(FeatureLevel, (FSceneViewState*)ShareOriginTarget);
 }
 
+TUniquePtr<ISceneRenderBuilder> FRendererModule::CreateSceneRenderBuilder(FSceneInterface* SceneInterface)
+{
+	check(SceneInterface);
+	FScene* Scene = SceneInterface->GetRenderScene();
+	check(Scene);
+
+	return MakeUnique<FSceneRenderBuilder>(Scene);
+}
+
 void FRendererModule::InvalidatePathTracedOutput()
 {
 	// AllocatedScenes is managed by the game thread
diff --git a/Engine/Source/Runtime/Renderer/Private/SceneRenderBuilder.cpp b/Engine/Source/Runtime/Renderer/Private/SceneRenderBuilder.cpp
new file mode 100644
--- /dev/null
+++ b/Engine/Source/Runtime/Renderer/Private/SceneRenderBuilder.cpp
@@ -0,0 +1,447 @@
+// Copyright Epic Games, Inc. All Rights Reserved.
+
+#include "SceneRenderBuilder.h"
+#include "RenderCaptureInterface.h"
+#include "VisualizeTexture.h"
+#include "SceneRendering.h"
+#include "GPUDebugCrashUtils.h"
+#include "ScenePrivate.h"
+#include "DumpGPU.h"
+#include "Materials/MaterialRenderProxy.h"
+#include "SceneViewExtension.h"
+#include "Containers/IntrusiveDoubleLinkedList.h"
+#include "TextureResource.h"
+
+/** This class is responsible for processing a batch of scene renderers. */
+class FSceneRenderProcessor
+{
+	struct FRenderNode;
+
+	struct FGroupNode : public TIntrusiveDoubleLinkedListNode<FGroupNode>
+	{
+		FGroupNode(FString&& InName, ESceneRenderGroupFlags InFlags)
+			: Name(Forward<FString&&>(InName))
+			, Flags(InFlags)
+		{}
+
+		TIntrusiveDoubleLinkedList<FRenderNode> RenderNodes;
+		FString Name;
+		ESceneRenderGroupFlags Flags;
+	};
+
+	struct FRenderNode : public TIntrusiveDoubleLinkedListNode<FRenderNode>
+	{
+		FRenderNode(FSceneRenderer* InRenderer, FString&& InName, FSceneRenderFunction&& InFunction, FGroupNode* InGroup)
+			: Renderer(InRenderer)
+			, Name(Forward<FString&&>(InName))
+			, Function(Forward<FSceneRenderFunction&&>(InFunction))
+			, Group(InGroup)
+		{}
+
+		FSceneRenderer* Renderer;
+		FString Name;
+		FSceneRenderFunction Function;
+		FGroupNode* Group;
+	};
+
+	struct FOp
+	{
+		enum class EType
+		{
+			Render,
+			FunctionCall,
+			BeginGroup,
+			EndGroup
+		};
+
+		static FOp Render(FRenderNode* Data)
+		{
+			check(Data);
+			FOp Op;
+			Op.Type = EType::Render;
+			Op.Data_Render = Data;
+			return Op;
+		}
+
+		static FOp FunctionCall(TUniqueFunction<void()>* Data)
+		{
+			check(Data);
+			FOp Op;
+			Op.Type = EType::FunctionCall;
+			Op.Data_FunctionCall = Data;
+			return Op;
+		}
+
+		static FOp BeginGroup(FGroupNode* Data)
+		{
+			check(Data);
+			FOp Op;
+			Op.Type = EType::BeginGroup;
+			Op.Data_Group = Data;
+			return Op;
+		}
+		
+		static FOp EndGroup(FGroupNode* Data)
+		{
+			check(Data);
+			FOp Op;
+			Op.Type = EType::EndGroup;
+			Op.Data_Group = Data;
+			return Op;
+		}
+
+		EType Type;
+
+		union
+		{
+			FRenderNode* Data_Render = nullptr;
+			TUniqueFunction<void()>* Data_FunctionCall;
+			FGroupNode* Data_Group;
+		};
+	};
+
+	enum class FGroupEventLocation
+	{
+		GroupCommand,
+		SceneRenderCommand
+	};
+
+	class FGroupEvent
+	{
+	public:
+		void Begin(FGroupNode* Group, FRHICommandList& RHICmdList, FGroupEventLocation Location)
+		{
+#if WITH_RHI_BREADCRUMBS
+			if (Group && IsLocationActive(Location))
+			{
+				FRHIBreadcrumbData BreadcrumbData(__FILE__, __LINE__, TStatId(), NAME_None);
+
+				// User specified an explicit group name.
+				if (!Group->Name.IsEmpty())
+				{
+					Event.Emplace(RHICmdList, MoveTemp(BreadcrumbData), TEXT("%s"), Group->Name);
+				}
+				// User didn't specify a name, but we have more than one renderer, so a group event is useful.
+				else if (!Group->RenderNodes.IsEmpty() && Group->RenderNodes.GetHead() != Group->RenderNodes.GetTail())
+				{
+					Event.Emplace(RHICmdList, MoveTemp(BreadcrumbData), TEXT("SceneRenderGroup"));
+				}
+			}
+#endif
+		};
+
+		void End(FRHICommandList& RHICmdList, FGroupEventLocation Location)
+		{
+#if WITH_RHI_BREADCRUMBS
+			if (IsLocationActive(Location) && Event)
+			{
+				Event->End(RHICmdList);
+				Event.Reset();
+			}
+#endif
+		};
+
+	private:
+#if WITH_RHI_BREADCRUMBS
+		bool IsLocationActive(FGroupEventLocation Location) const
+		{
+			// When the render command channel is enabled, we push a unique group scope for each scene render command.
+			// Otherwise we push a scope inside of the group begin / end commands. This makes scopes behave properly.
+			const bool bRenderCommandsChannelEnabled = UE_TRACE_CHANNELEXPR_IS_ENABLED(RenderCommandsChannel);
+			const bool bGroupCommandLocation = Location == FGroupEventLocation::SceneRenderCommand;
+			return (bRenderCommandsChannelEnabled == bGroupCommandLocation);
+		};
+
+		TOptional<FRHIBreadcrumbEventManual> Event;
+#endif
+	};
+
+public:
+	FSceneRenderProcessor(FScene* InScene)
+		: Scene(InScene)
+	{}
+
+	void AddCommand(TUniqueFunction<void()>&& Function)
+	{
+		Ops.Emplace(FOp::FunctionCall(Allocator.Create<TUniqueFunction<void()>>(MoveTemp(Function))));
+	}
+
+	void AddRenderer(FSceneRenderer* Renderer, FString&& Name, FSceneRenderFunction&& Function)
+	{
+		check(Renderer);
+		check(Renderer->Scene == Scene);
+		check(Renderer->ViewFamily.Views.Num() > 0);
+		check(Renderer->ViewFamily.Views[0]);
+
+		FGroupNode* TailGroup = GroupNodes.GetTail();
+
+		FRenderNode* RenderNode = Allocator.Create<FRenderNode>(
+			Renderer,
+			Forward<FString&&>(Name),
+			Forward<FSceneRenderFunction&&>(Function),
+			TailGroup);
+
+		if (TailGroup)
+		{
+			TailGroup->RenderNodes.AddTail(RenderNode);
+		}
+		RenderNodes.AddTail(RenderNode);
+
+		Ops.Emplace(FOp::Render(RenderNode));
+		Renderers.Emplace(Renderer);
+	}
+
+	void BeginGroup(FString&& Name, ESceneRenderGroupFlags Flags)
+	{
+		checkf(!bInsideGroup, TEXT("SceneRenderBuilderGroup scope %s is being nested with the group %s. Groups do not currently support nesting."), *GroupNodes.GetTail()->Name, *Name);
+		FGroupNode* Group = Allocator.Create<FGroupNode>(Forward<FString&&>(Name), Flags);
+		GroupNodes.AddTail(Group);
+		Ops.Emplace(FOp::BeginGroup(Group));
+		bInsideGroup = true;
+	}
+
+	void EndGroup()
+	{
+		checkf(bInsideGroup, TEXT("EndGroup called without a matching BeginGroup"));
+		Ops.Emplace(FOp::EndGroup(GroupNodes.GetTail()));
+		bInsideGroup = false;
+	}
+
+	void Execute();
+
+	FConcurrentLinearBulkObjectAllocator& GetAllocator()
+	{
+		return Allocator;
+	}
+
+private:
+	FScene* Scene;
+	TArray<FSceneRenderer*, FConcurrentLinearArrayAllocator> Renderers;
+	TIntrusiveDoubleLinkedList<FRenderNode> RenderNodes;
+	TIntrusiveDoubleLinkedList<FGroupNode> GroupNodes;
+	TArray<FOp, FConcurrentLinearArrayAllocator> Ops;
+	FConcurrentLinearBulkObjectAllocator Allocator;
+
+	struct
+	{
+		FGroupNode* Group = nullptr;
+		FGroupEvent GroupEvent;
+		FString FullPath;
+
+	} RenderState;
+	
+	bool bInsideGroup : 1 = false;
+};
+
+void FSceneRenderProcessor::Execute()
+{
+	checkf(!bInsideGroup, TEXT("FSceneRenderBuilder::Execute called within scene render group scope %s. You must end the scope first."), *GroupNodes.GetTail()->Name);
+
+#if WITH_GPUDEBUGCRASH
+	if (GRHIGlobals.TriggerGPUCrash != ERequestedGPUCrash::None)
+	{
+		ENQUEUE_RENDER_COMMAND(ScheduleGPUDebugCrash)([] (FRHICommandListImmediate& RHICmdList)
+		{
+			TRACE_CPUPROFILER_EVENT_SCOPE(TriggerGPUCrash);
+			FRDGBuilder GraphBuilder(RHICmdList);
+			ScheduleGPUDebugCrash(GraphBuilder);
+			GraphBuilder.Execute();
+		});
+	}
+#endif
+
+	UE::RenderCommandPipe::FSyncScope SyncScope;
+	FUniformExpressionCacheAsyncUpdateScope AsyncUpdateScope;
+
+	ENQUEUE_RENDER_COMMAND(SceneRenderBuilder_Begin)([this](FRHICommandListImmediate& RHICmdList)
+	{
+		if (!Renderers.IsEmpty())
+		{
+			FSceneRenderer::RenderThreadBegin(RHICmdList, Renderers);
+		}
+	});
+
+	TOptional<UE::RenderCore::DumpGPU::FDumpScope> GpuDumpScope;
+	TOptional<RenderCaptureInterface::FScopedCapture> GpuCaptureScope;
+
+	for (FOp Op : Ops)
+	{
+		switch (Op.Type)
+		{
+		case FOp::EType::BeginGroup:
+		{
+			FGroupNode* Group = Op.Data_Group;
+
+			if (EnumHasAnyFlags(Group->Flags, ESceneRenderGroupFlags::GpuCapture))
+			{
+				GpuCaptureScope.Emplace(true, *Group->Name);
+			}
+
+			if (EnumHasAnyFlags(Group->Flags, ESceneRenderGroupFlags::GpuDump))
+			{
+				GpuDumpScope.Emplace();
+			}
+
+			ENQUEUE_RENDER_COMMAND(SceneRenderBuilder_BeginGroup)([this, Group] (FRHICommandListImmediate& RHICmdList)
+			{
+				RenderState.Group = Group;
+				RenderState.GroupEvent.Begin(Group, RHICmdList, FGroupEventLocation::GroupCommand);
+				RenderState.FullPath += Group->Name;
+			});
+		}
+		break;
+		case FOp::EType::EndGroup:
+		{
+			ENQUEUE_RENDER_COMMAND(SceneRenderBuilder_EndGroup)([this](FRHICommandListImmediate& RHICmdList) mutable
+			{
+				RenderState.GroupEvent.End(RHICmdList, FGroupEventLocation::GroupCommand);
+				RenderState.Group = nullptr;
+				RenderState.FullPath.Reset();
+			});
+
+			GpuCaptureScope.Reset();
+			GpuDumpScope.Reset();
+		}
+		break;
+		case FOp::EType::FunctionCall:
+		{
+			(*Op.Data_FunctionCall)();
+		}
+		break;
+		case FOp::EType::Render:
+		{
+			FRenderNode& RenderNode = *Op.Data_Render;
+
+			for (const FSceneViewExtensionRef& Extension : RenderNode.Renderer->ViewFamily.ViewExtensions)
+			{
+				Extension->BeginRenderViewFamily(RenderNode.Renderer->ViewFamily);
+			}
+
+			ENQUEUE_RENDER_COMMAND(SceneRenderBuilder_Render)([this, &RenderNode] (FRHICommandListImmediate& RHICmdList)
+			{
+				RenderState.GroupEvent.Begin(RenderState.Group, RHICmdList, FGroupEventLocation::SceneRenderCommand);
+
+				ON_SCOPE_EXIT
+				{
+					RenderState.GroupEvent.End(RHICmdList, FGroupEventLocation::SceneRenderCommand);
+				};
+
+				RHI_BREADCRUMB_EVENT_CONDITIONAL(RHICmdList, !RenderNode.Name.IsEmpty(), "%s", RenderNode.Name);
+				RHI_BREADCRUMB_EVENT_CONDITIONAL(RHICmdList,  RenderNode.Name.IsEmpty(), "SceneRender");
+
+				FSceneRenderer& Renderer = *RenderNode.Renderer;
+
+				if (!RenderNode.Name.IsEmpty())
+				{
+					RenderState.FullPath += TEXT("/");
+					RenderState.FullPath += RenderNode.Name;
+				}
+
+				const FSceneRenderFunctionInputs FunctionInputs(&Renderer, *RenderNode.Name, *RenderState.FullPath);
+
+				FRDGBuilder GraphBuilder(RHICmdList, RDG_EVENT_NAME("%s", FunctionInputs.FullPath), ERDGBuilderFlags::Parallel, Scene->GetShaderPlatform());
+
+				if (!Renderer.ViewFamily.EngineShowFlags.HitProxies)
+				{
+					VISUALIZE_TEXTURE_BEGIN_VIEW(Scene->GetFeatureLevel(), Renderer.Views[0].GetViewKey(), FunctionInputs.FullPath, Renderer.Views[0].bIsSceneCapture);
+				}
+
+				RenderNode.Function(GraphBuilder, FunctionInputs);
+
+				if (!Renderer.ViewFamily.EngineShowFlags.HitProxies)
+				{
+					VISUALIZE_TEXTURE_END_VIEW();
+				}
+
+				if (!RenderNode.Name.IsEmpty())
+				{
+					RenderState.FullPath.LeftChopInline(RenderNode.Name.Len() + 1, EAllowShrinking::No);
+				}
+
+				// The final graph builder is responsible for flushing resources.
+				if (!RenderNode.GetNext())
+				{
+					GraphBuilder.SetFlushResourcesRHI();
+				}
+
+				GraphBuilder.Execute();
+			});
+		}
+		break;
+		}
+	}
+
+	ENQUEUE_RENDER_COMMAND(SceneRenderBuilder_End)([this](FRHICommandListImmediate& RHICmdList) mutable
+	{
+		if (!Renderers.IsEmpty())
+		{
+			FSceneRenderer::RenderThreadEnd(RHICmdList, Renderers);
+		}
+
+		delete this;
+	});
+
+	// NOTE: 'this' is queued for deletion and is no longer valid!
+}
+
+//////////////////////////////////////////////////////////////////////////////
+
+FSceneRenderBuilder::~FSceneRenderBuilder()
+{
+	if (Processor)
+	{
+		checkf(false, TEXT("SceneRenderBuilder is being destructed without having called Execute."));
+		delete Processor;
+		Processor = nullptr;
+	}
+}
+
+void FSceneRenderBuilder::LazyInit()
+{
+	if (!Processor)
+	{
+		Processor = new FSceneRenderProcessor(Scene);
+	}
+}
+
+void FSceneRenderBuilder::AddCommand(TUniqueFunction<void()>&& Function)
+{
+	LazyInit();
+	Processor->AddCommand(MoveTemp(Function));
+}
+
+void FSceneRenderBuilder::AddRenderer(FSceneRenderer* Renderer, FString&& Name, FSceneRenderFunction&& Function)
+{
+	LazyInit();
+	Processor->AddRenderer(Renderer, Forward<FString&&>(Name), MoveTemp(Function));
+}
+
+void FSceneRenderBuilder::BeginGroup(FString&& Name, ESceneRenderGroupFlags Flags)
+{
+	// If user sets both capture and dump flags, prefer capturing over dumping (or clear flag if dumping is not available or we are currently dumping already).
+	if (EnumHasAnyFlags(Flags, ESceneRenderGroupFlags::GpuCapture) || !(WITH_ENGINE && WITH_DUMPGPU) || FRDGBuilder::IsDumpingFrame())
+	{
+		EnumRemoveFlags(Flags, ESceneRenderGroupFlags::GpuDump);
+	}
+
+	LazyInit();
+	Processor->BeginGroup(Forward<FString&&>(Name), Flags);
+}
+
+void FSceneRenderBuilder::EndGroup()
+{
+	checkf(Processor, TEXT("EndGroup called on an empty scene render builder."));
+	Processor->EndGroup();
+}
+
+void FSceneRenderBuilder::Execute()
+{
+	Processor->Execute();
+	Processor = nullptr;
+}
+
+FConcurrentLinearBulkObjectAllocator& FSceneRenderBuilder::GetAllocator()
+{
+	LazyInit();
+	return Processor->GetAllocator();
+}
\ No newline at end of file
diff --git a/Engine/Source/Runtime/Renderer/Private/SceneRenderBuilder.h b/Engine/Source/Runtime/Renderer/Private/SceneRenderBuilder.h
new file mode 100644
--- /dev/null
+++ b/Engine/Source/Runtime/Renderer/Private/SceneRenderBuilder.h
@@ -0,0 +1,56 @@
+// Copyright Epic Games, Inc. All Rights Reserved.
+
+#pragma once
+
+#include "SceneRenderBuilderInterface.h"
+#include "Templates/PimplPtr.h"
+
+class FSceneRenderProcessor;
+class FSceneRenderer;
+class FScene;
+
+class FSceneRenderBuilder final : public ISceneRenderBuilder
+{
+public:
+	FSceneRenderBuilder(FScene* InScene)
+		: Scene(InScene)
+	{}
+
+	FSceneRenderBuilder(const FSceneRenderBuilder&) = delete;
+	FSceneRenderBuilder& operator=(const FSceneRenderBuilder&) = delete;
+
+	FSceneRenderBuilder(FSceneRenderBuilder&& RHS)
+	{
+		*this = MoveTemp(RHS);
+	}
+
+	FSceneRenderBuilder& operator=(FSceneRenderBuilder&& RHS)
+	{
+		// Scene stays persistent while the processor instance moves.
+		Scene         = RHS.Scene;
+		Processor     = RHS.Processor;
+		RHS.Processor = nullptr;
+		return *this;
+	}
+
+	~FSceneRenderBuilder() override;
+
+	//////////////////////////////////////////////////////////////////////////
+	// ISceneRenderBuilder Overrides
+
+	void AddCommand(TUniqueFunction<void()>&& Function) override;
+	void AddRenderer(FSceneRenderer* Renderer, FString&& Name, FSceneRenderFunction&& Function) override;
+	using ISceneRenderBuilder::AddRenderer;
+	void Execute() override;
+	FConcurrentLinearBulkObjectAllocator& GetAllocator() override;
+
+	//////////////////////////////////////////////////////////////////////////
+
+private:
+	void LazyInit();
+	void BeginGroup(FString&& Name, ESceneRenderGroupFlags Flags) override;
+	void EndGroup() override;
+
+	FScene* Scene;
+	FSceneRenderProcessor* Processor = nullptr;
+};
\ No newline at end of file
