Back out use of ForceImmediateFirstBarrier in Virtual Texture Page Table update.
This is to avoid issue with pipeline transition with async compute.
It will reintroduce potential transition error when reading virtual textures to write a runtime virtual texture.

diff --git a/Engine/Source/Runtime/Renderer/Private/VT/VirtualTextureSpace.cpp b/Engine/Source/Runtime/Renderer/Private/VT/VirtualTextureSpace.cpp
--- a/Engine/Source/Runtime/Renderer/Private/VT/VirtualTextureSpace.cpp
+++ b/Engine/Source/Runtime/Renderer/Private/VT/VirtualTextureSpace.cpp
@@ -490,10 +490,7 @@ void FVirtualTextureSpace::ApplyUpdates(FVirtualTextureSystem* System, FRDGBuild
 		const uint32 LayerInTexture = LayerIndex % LayersPerPageTableTexture;
 		FTextureEntry& PageTableEntry = PageTable[TextureIndex];
 
-		// Use ForceImmediateFirstBarrier to prevent PageTableTexture RVT transition from being hoisted to earlier in the render graph.
-		// This can be a problem if the raw RHI object for the page table is being used earlier in the graph.
-		// An example where that can happen is if we sample virtual textures while rendering runtime virtual texture pages.
-		FRDGTextureRef PageTableTexture = GraphBuilder.RegisterExternalTexture(PageTableEntry.RenderTarget, ERDGTextureFlags::ForceImmediateFirstBarrier);
+		FRDGTextureRef PageTableTexture = GraphBuilder.RegisterExternalTexture(PageTableEntry.RenderTarget, ERDGTextureFlags::None);
 		GraphBuilder.UseInternalAccessMode(PageTableTexture);
 
 		// Use color write mask to update the proper page table entry for this layer
