Converted renderer to use new SceneRenderBuilder interface.

diff --git a/Engine/Plugins/Experimental/CineCameraSceneCapture/Source/CineCameraSceneCapture/Private/CineCameraSceneCaptureComponent.cpp b/Engine/Plugins/Experimental/CineCameraSceneCapture/Source/CineCameraSceneCapture/Private/CineCameraSceneCaptureComponent.cpp
--- a/Engine/Plugins/Experimental/CineCameraSceneCapture/Source/CineCameraSceneCapture/Private/CineCameraSceneCaptureComponent.cpp
+++ b/Engine/Plugins/Experimental/CineCameraSceneCapture/Source/CineCameraSceneCapture/Private/CineCameraSceneCaptureComponent.cpp
@@ -165,7 +165,7 @@ UCineCaptureComponent2D::UCineCaptureComponent2D(const FObjectInitializer& Objec
 	bUseRayTracingIfEnabled = true;
 }
 
-void UCineCaptureComponent2D::UpdateSceneCaptureContents(FSceneInterface* Scene)
+void UCineCaptureComponent2D::UpdateSceneCaptureContents(FSceneInterface* Scene, ISceneRenderBuilder& SceneRenderBuilder)
 {
 	if (!CineCameraComponent.IsValid() || !CineCaptureSVE.IsValid())
 	{
@@ -179,7 +179,7 @@ void UCineCaptureComponent2D::UpdateSceneCaptureContents(FSceneInterface* Scene)
 	bOverride_CustomNearClippingPlane = CineCameraComponent->bOverride_CustomNearClippingPlane;
 	CustomNearClippingPlane = CineCameraComponent->CustomNearClippingPlane;
 
-	Scene->UpdateSceneCaptureContents(this);
+	Scene->UpdateSceneCaptureContents(this, SceneRenderBuilder);
 }
 
 void UCineCaptureComponent2D::CheckResizeRenderTarget()
diff --git a/Engine/Plugins/Experimental/CineCameraSceneCapture/Source/CineCameraSceneCapture/Public/CineCameraSceneCaptureComponent.h b/Engine/Plugins/Experimental/CineCameraSceneCapture/Source/CineCameraSceneCapture/Public/CineCameraSceneCaptureComponent.h
--- a/Engine/Plugins/Experimental/CineCameraSceneCapture/Source/CineCameraSceneCapture/Public/CineCameraSceneCaptureComponent.h
+++ b/Engine/Plugins/Experimental/CineCameraSceneCapture/Source/CineCameraSceneCapture/Public/CineCameraSceneCaptureComponent.h
@@ -56,7 +56,7 @@ public:
 	TSoftObjectPtr<UCineCameraComponent> CineCameraComponent;
 
 public:
-	virtual void UpdateSceneCaptureContents(FSceneInterface* Scene) override;
+	virtual void UpdateSceneCaptureContents(FSceneInterface* Scene, ISceneRenderBuilder& SceneRenderBuilder) override;
 	virtual void OnRegister() override;
 	virtual void OnUnregister() override;
 
diff --git a/Engine/Plugins/Runtime/MixedRealityCaptureFramework/Source/MixedRealityCaptureFramework/Private/MixedRealityCaptureComponent.cpp b/Engine/Plugins/Runtime/MixedRealityCaptureFramework/Source/MixedRealityCaptureFramework/Private/MixedRealityCaptureComponent.cpp
--- a/Engine/Plugins/Runtime/MixedRealityCaptureFramework/Source/MixedRealityCaptureFramework/Private/MixedRealityCaptureComponent.cpp
+++ b/Engine/Plugins/Runtime/MixedRealityCaptureFramework/Source/MixedRealityCaptureFramework/Private/MixedRealityCaptureComponent.cpp
@@ -24,6 +24,7 @@
 #include "PlatformFeatures.h"
 #include "SaveGameSystem.h"
 #include "Serialization/ObjectAndNameAsStringProxyArchive.h"
+#include "SceneRenderBuilderInterface.h"
 
 #if WITH_EDITORONLY_DATA
 #include "Components/StaticMeshComponent.h"
@@ -578,7 +579,7 @@ const AActor* UMixedRealityCaptureComponent::GetViewOwner() const
 }
 
 //------------------------------------------------------------------------------
-void UMixedRealityCaptureComponent::UpdateSceneCaptureContents(FSceneInterface* Scene)
+void UMixedRealityCaptureComponent::UpdateSceneCaptureContents(FSceneInterface* Scene, ISceneRenderBuilder& SceneRenderBuilder)
 { 
 	if (TextureTarget)
 	{
@@ -590,14 +591,23 @@ void UMixedRealityCaptureComponent::UpdateSceneCaptureContents(FSceneInterface*
 		ViewExtension = FSceneViewExtensions::NewExtension<FMrcLatencyViewExtension>(this);
 		FMotionDelayService::RegisterDelayClient(ViewExtension.ToSharedRef());
 	}
-	const bool bPreCommandQueued = ViewExtension->SetupPreCapture(Scene);
 
-	Super::UpdateSceneCaptureContents(Scene);
+	bool* bPreCommandQueuePtr = SceneRenderBuilder.GetAllocator().Malloc<bool>();
 
-	if (bPreCommandQueued)
+	SceneRenderBuilder.AddCommand([bPreCommandQueuePtr, ViewExtension = ViewExtension, Scene]
 	{
-		ViewExtension->SetupPostCapture(Scene);
-	}
+		*bPreCommandQueuePtr = ViewExtension->SetupPreCapture(Scene);
+	});
+
+	Super::UpdateSceneCaptureContents(Scene, SceneRenderBuilder);
+
+	SceneRenderBuilder.AddCommand([bPreCommandQueuePtr, ViewExtension = ViewExtension, Scene]
+	{
+		if (*bPreCommandQueuePtr)
+		{
+			ViewExtension->SetupPostCapture(Scene);
+		}
+	});
 }
 
 //------------------------------------------------------------------------------
diff --git a/Engine/Plugins/Runtime/MixedRealityCaptureFramework/Source/MixedRealityCaptureFramework/Public/MixedRealityCaptureComponent.h b/Engine/Plugins/Runtime/MixedRealityCaptureFramework/Source/MixedRealityCaptureFramework/Public/MixedRealityCaptureComponent.h
--- a/Engine/Plugins/Runtime/MixedRealityCaptureFramework/Source/MixedRealityCaptureFramework/Public/MixedRealityCaptureComponent.h
+++ b/Engine/Plugins/Runtime/MixedRealityCaptureFramework/Source/MixedRealityCaptureFramework/Public/MixedRealityCaptureComponent.h
@@ -100,7 +100,7 @@ public:
 	//~ USceneCaptureComponent interface
 
 	virtual const AActor* GetViewOwner() const override;
-	virtual void UpdateSceneCaptureContents(FSceneInterface* Scene) override;
+	virtual void UpdateSceneCaptureContents(FSceneInterface* Scene, ISceneRenderBuilder& SceneRenderBuilder) override;
 
 public:
 	//~ Blueprint API	
diff --git a/Engine/Plugins/Runtime/nDisplay/Source/DisplayCluster/Private/Game/EngineClasses/Basics/DisplayClusterViewportClient.cpp b/Engine/Plugins/Runtime/nDisplay/Source/DisplayCluster/Private/Game/EngineClasses/Basics/DisplayClusterViewportClient.cpp
--- a/Engine/Plugins/Runtime/nDisplay/Source/DisplayCluster/Private/Game/EngineClasses/Basics/DisplayClusterViewportClient.cpp
+++ b/Engine/Plugins/Runtime/nDisplay/Source/DisplayCluster/Private/Game/EngineClasses/Basics/DisplayClusterViewportClient.cpp
@@ -723,10 +723,6 @@ void UDisplayClusterViewportClient::Draw(FViewport* InViewport, FCanvas* SceneCa
 #endif
 
 				ViewFamily.ProfileDescription = DCViewFamily.Views[0].Viewport->GetId();
-				if (!ViewFamily.ProfileDescription.IsEmpty())
-				{
-					ViewFamily.ProfileSceneRenderTime = ((FDisplayClusterSceneViewport*)InViewport)->GetNextHistoryWriteAddress(ViewFamily.ProfileDescription);
-				}
 
 				// Draw the player views.
 				if (!bDisableWorldRendering && PlayerViewMap.Num() > 0 && FSlateApplication::Get().GetPlatformApplication()->IsAllowedToRender()) //-V560
diff --git a/Engine/Source/Runtime/Engine/Classes/Components/SceneCaptureComponent.h b/Engine/Source/Runtime/Engine/Classes/Components/SceneCaptureComponent.h
--- a/Engine/Source/Runtime/Engine/Classes/Components/SceneCaptureComponent.h
+++ b/Engine/Source/Runtime/Engine/Classes/Components/SceneCaptureComponent.h
@@ -240,11 +240,14 @@ public:
 	/** Whether this component is a USceneCaptureComponentCube */
 	virtual bool IsCube() const { return false; }
 
-	/** Updates "FrameUpdated" field, returns "true" if this is one of multiple scene captures this frame.  Called from FScene::UpdateSceneCaptureContents. */
-	ENGINE_API bool SetFrameUpdated();
+	UE_DEPRECATED(5.6, "SetFrameUpdated is no longer used")
+	bool SetFrameUpdated() { return false; }
 
+	UE_DEPRECATED(5.6, "UpdateSceneCaptureContents now requires a SceneRenderBuilder")
 	virtual void UpdateSceneCaptureContents(FSceneInterface* Scene) {};
 
+	virtual void UpdateSceneCaptureContents(FSceneInterface* Scene, class ISceneRenderBuilder& SceneRenderBuilder) {};
+
 	/**
 	 * Returns whether a scene capture doesn't want or need lighting, and can disable several additional rendering features to save performance
 	 * (see DisableFeaturesForUnlit in ShowFlags.h).  Depth or base pass outputs aren't affected by lighting, while SceneColor outputs may be generated
@@ -274,9 +277,6 @@ protected:
 	 */
 	TIndirectArray<FSceneViewStateReference> ViewStates;
 
-	/** Frame this scene capture was last updated */
-	uint64 FrameUpdated;
-
 #if WITH_EDITORONLY_DATA
 	/** The mesh used by ProxyMeshComponent */
 	UPROPERTY(transient)
diff --git a/Engine/Source/Runtime/Engine/Classes/Components/SceneCaptureComponent2D.h b/Engine/Source/Runtime/Engine/Classes/Components/SceneCaptureComponent2D.h
--- a/Engine/Source/Runtime/Engine/Classes/Components/SceneCaptureComponent2D.h
+++ b/Engine/Source/Runtime/Engine/Classes/Components/SceneCaptureComponent2D.h
@@ -278,7 +278,7 @@ public:
 	UFUNCTION(BlueprintCallable,Category = "Rendering|SceneCapture")
 	ENGINE_API void CaptureScene();
 
-	ENGINE_API void UpdateSceneCaptureContents(FSceneInterface* Scene) override;
+	ENGINE_API void UpdateSceneCaptureContents(FSceneInterface* Scene, ISceneRenderBuilder& SceneRenderBuilder) override;
 
 	/* Return if orthographic tiling rendering is enabled or not */
 	ENGINE_API bool GetEnableOrthographicTiling() const;
diff --git a/Engine/Source/Runtime/Engine/Classes/Components/SceneCaptureComponentCube.h b/Engine/Source/Runtime/Engine/Classes/Components/SceneCaptureComponentCube.h
--- a/Engine/Source/Runtime/Engine/Classes/Components/SceneCaptureComponentCube.h
+++ b/Engine/Source/Runtime/Engine/Classes/Components/SceneCaptureComponentCube.h
@@ -65,7 +65,7 @@ public:
 	// For backwards compatibility
 	void UpdateContent() { CaptureSceneDeferred(); }
 
-	ENGINE_API void UpdateSceneCaptureContents(FSceneInterface* Scene) override;
+	ENGINE_API void UpdateSceneCaptureContents(FSceneInterface* Scene, ISceneRenderBuilder& SceneRenderBuilder) override;
 
 	/** Whether this component is a USceneCaptureComponentCube */
 	virtual bool IsCube() const override { return true; }
diff --git a/Engine/Source/Runtime/Engine/Internal/SceneCapture/SceneCaptureInternal.h b/Engine/Source/Runtime/Engine/Internal/SceneCapture/SceneCaptureInternal.h
--- a/Engine/Source/Runtime/Engine/Internal/SceneCapture/SceneCaptureInternal.h
+++ b/Engine/Source/Runtime/Engine/Internal/SceneCapture/SceneCaptureInternal.h
@@ -6,9 +6,10 @@
 
 class FSceneInterface;
 class FSceneViewFamily;
+class ISceneRenderBuilder;
 
 /**
  * Internal function called from FRendererModule::BeginRenderingViewFamilies.  List of view families is required to handle captures where
  * USceneCaptureComponent2D::ShouldRenderWithMainViewFamily() is true.
  */
-ENGINE_API void SceneCaptureUpdateDeferredCapturesInternal(FSceneInterface* Scene, TArray<const FSceneViewFamily*>& InOutViewFamilies);
+ENGINE_API void SceneCaptureUpdateDeferredCapturesInternal(FSceneInterface* Scene, TConstArrayView<const FSceneViewFamily*> ViewFamilies, ISceneRenderBuilder& SceneRenderBuilder);
diff --git a/Engine/Source/Runtime/Engine/Private/Components/SceneCaptureComponent.cpp b/Engine/Source/Runtime/Engine/Private/Components/SceneCaptureComponent.cpp
--- a/Engine/Source/Runtime/Engine/Private/Components/SceneCaptureComponent.cpp
+++ b/Engine/Source/Runtime/Engine/Private/Components/SceneCaptureComponent.cpp
@@ -38,6 +38,7 @@
 #include "Engine/SCS_Node.h"
 #include "Engine/TextureRenderTarget2D.h"
 #include "UnrealEngine.h"
+#include "SceneRenderBuilderInterface.h"
 
 #include UE_INLINE_GENERATED_CPP_BY_NAME(SceneCaptureComponent)
 
@@ -46,9 +47,6 @@
 static TMultiMap<TWeakObjectPtr<UWorld>, TWeakObjectPtr<USceneCaptureComponent> > SceneCapturesToUpdateMap;
 static FCriticalSection SceneCapturesToUpdateMapCS;
 
-// Tracks the latest frame a scene capture was rendered
-static uint64 GSceneCaptureLatestRenderedFrame = INDEX_NONE;
-
 static TAutoConsoleVariable<bool> CVarSCOverrideOrthographicTilingValues(
 	TEXT("r.SceneCapture.OverrideOrthographicTilingValues"),
 	false,
@@ -448,17 +446,6 @@ void USceneCaptureComponent::UpdateShowFlags()
 	}
 }
 
-bool USceneCaptureComponent::SetFrameUpdated()
-{
-	bool bFirstCaptureThisFrame = GFrameCounter != GSceneCaptureLatestRenderedFrame;
-	bool bRepeatOfThisCapture = GFrameCounter == FrameUpdated;
-	GSceneCaptureLatestRenderedFrame = GFrameCounter;
-	FrameUpdated = GFrameCounter;
-
-	// We consider this a "multiple" capture if it's not the first capture this frame, and not a repeat of the same capture.
-	return !bFirstCaptureThisFrame && !bRepeatOfThisCapture;
-}
-
 #if WITH_EDITOR
 
 bool USceneCaptureComponent::CanEditChange(const FProperty* InProperty) const
@@ -521,11 +508,12 @@ void USceneCaptureComponent::Serialize(FArchive& Ar)
 
 void USceneCaptureComponent::UpdateDeferredCaptures(FSceneInterface* Scene)
 {
-	TArray<const FSceneViewFamily*> EmptyViewFamilies;
-	SceneCaptureUpdateDeferredCapturesInternal(Scene, EmptyViewFamilies);
+	TUniquePtr<ISceneRenderBuilder> SceneRenderBuilder = ISceneRenderBuilder::Create(Scene);
+	SceneCaptureUpdateDeferredCapturesInternal(Scene, {}, *SceneRenderBuilder);
+	SceneRenderBuilder->Execute();
 }
 
-void SceneCaptureUpdateDeferredCapturesInternal(FSceneInterface* Scene, TArray<const FSceneViewFamily*>& InOutViewFamilies)
+void SceneCaptureUpdateDeferredCapturesInternal(FSceneInterface* Scene, TConstArrayView<const FSceneViewFamily*> ViewFamilies, ISceneRenderBuilder& SceneRenderBuilder)
 {
 	FScopeLock ScopeLock(&SceneCapturesToUpdateMapCS);
 	UWorld* World = Scene->GetWorld();
@@ -540,7 +528,7 @@ void SceneCaptureUpdateDeferredCapturesInternal(FSceneInterface* Scene, TArray<c
 	SceneCapturesToUpdateMap.MultiFind(World, SceneCapturesToUpdate);
 
 	const FSceneViewFamily* MainViewFamily = nullptr;
-	for (const FSceneViewFamily* Family : InOutViewFamilies)
+	for (const FSceneViewFamily* Family : ViewFamilies)
 	{
 		if (Family->bIsMainViewFamily)
 		{
@@ -592,7 +580,7 @@ void SceneCaptureUpdateDeferredCapturesInternal(FSceneInterface* Scene, TArray<c
 				((USceneCaptureComponent2D*)Component.Get())->MainViewFamily = MainViewFamily;
 			}
 
-			Component->UpdateSceneCaptureContents(Scene);
+			Component->UpdateSceneCaptureContents(Scene, SceneRenderBuilder);
 
 			if (bIsCaptureComponent2D)
 			{
@@ -822,7 +810,10 @@ void USceneCaptureComponent2D::CaptureScene()
 
 		// We must push any deferred render state recreations before causing any rendering to happen, to make sure that deleted resource references are updated
 		World->SendAllEndOfFrameUpdates();
-		UpdateSceneCaptureContents(World->Scene);
+
+		TUniquePtr<ISceneRenderBuilder> SceneRenderBuilder = ISceneRenderBuilder::Create(World->Scene);
+		UpdateSceneCaptureContents(World->Scene, *SceneRenderBuilder);
+		SceneRenderBuilder->Execute();
 	}
 
 	if (bCaptureEveryFrame)
@@ -979,9 +970,9 @@ void USceneCaptureComponent2D::Serialize(FArchive& Ar)
 	}
 }
 
-void USceneCaptureComponent2D::UpdateSceneCaptureContents(FSceneInterface* Scene)
+void USceneCaptureComponent2D::UpdateSceneCaptureContents(FSceneInterface* Scene, ISceneRenderBuilder& SceneRenderBuilder)
 {
-	Scene->UpdateSceneCaptureContents(this);
+	Scene->UpdateSceneCaptureContents(this, SceneRenderBuilder);
 }
 
 bool USceneCaptureComponent2D::GetEnableOrthographicTiling() const
@@ -1425,8 +1416,11 @@ void USceneCaptureComponentCube::CaptureScene()
 
 		// We must push any deferred render state recreations before causing any rendering to happen, to make sure that deleted resource references are updated
 		World->SendAllEndOfFrameUpdates();
-		UpdateSceneCaptureContents(World->Scene);
-	}	
+
+		TUniquePtr<ISceneRenderBuilder> SceneRenderBuilder = ISceneRenderBuilder::Create(World->Scene);
+		UpdateSceneCaptureContents(World->Scene, *SceneRenderBuilder);
+		SceneRenderBuilder->Execute();
+	}
 
 	if (bCaptureEveryFrame)
 	{
@@ -1434,9 +1428,9 @@ void USceneCaptureComponentCube::CaptureScene()
 	}
 }
 
-void USceneCaptureComponentCube::UpdateSceneCaptureContents(FSceneInterface* Scene)
+void USceneCaptureComponentCube::UpdateSceneCaptureContents(FSceneInterface* Scene, ISceneRenderBuilder& SceneRenderBuilder)
 {
-	Scene->UpdateSceneCaptureContents(this);
+	Scene->UpdateSceneCaptureContents(this, SceneRenderBuilder);
 }
 
 #if WITH_EDITOR
diff --git a/Engine/Source/Runtime/Engine/Private/SceneView.cpp b/Engine/Source/Runtime/Engine/Private/SceneView.cpp
--- a/Engine/Source/Runtime/Engine/Private/SceneView.cpp
+++ b/Engine/Source/Runtime/Engine/Private/SceneView.cpp
@@ -2960,7 +2960,6 @@ PRAGMA_DISABLE_DEPRECATION_WARNINGS
 PRAGMA_ENABLE_DEPRECATION_WARNINGS
 	SecondaryViewFraction(1.0f),
 	SecondaryScreenPercentageMethod(ESecondaryScreenPercentageMethod::LowerPixelDensitySimulation),
-	ProfileSceneRenderTime(nullptr),
 	SceneRenderer(nullptr),
 	ScreenPercentageInterface(nullptr),
 	TemporalUpscalerInterface(nullptr),
diff --git a/Engine/Source/Runtime/Engine/Public/SceneInterface.h b/Engine/Source/Runtime/Engine/Public/SceneInterface.h
--- a/Engine/Source/Runtime/Engine/Public/SceneInterface.h
+++ b/Engine/Source/Runtime/Engine/Public/SceneInterface.h
@@ -273,10 +273,10 @@ public:
 	* Updates the contents of the given scene capture by rendering the scene. 
 	* This must be called on the game thread.
 	*/
-	virtual void UpdateSceneCaptureContents(class USceneCaptureComponent2D* CaptureComponent) {}
-	virtual void UpdateSceneCaptureContents(class USceneCaptureComponentCube* CaptureComponent) {}
-	virtual void UpdatePlanarReflectionContents(class UPlanarReflectionComponent* CaptureComponent, class FSceneRenderer& MainSceneRenderer) {}
-	
+	virtual void UpdateSceneCaptureContents(class USceneCaptureComponent2D* CaptureComponent, class ISceneRenderBuilder& SceneRenderBuilder) {}
+	virtual void UpdateSceneCaptureContents(class USceneCaptureComponentCube* CaptureComponent, class ISceneRenderBuilder& SceneRenderBuilder) {}
+	virtual void UpdatePlanarReflectionContents(class UPlanarReflectionComponent* CaptureComponent, class FSceneRenderer& MainSceneRenderer, class ISceneRenderBuilder& SceneRenderBuilder) {}
+
 	virtual void AddPrecomputedLightVolume(const class FPrecomputedLightVolume* Volume) {}
 	virtual void RemovePrecomputedLightVolume(const class FPrecomputedLightVolume* Volume) {}
 
diff --git a/Engine/Source/Runtime/Engine/Public/SceneView.h b/Engine/Source/Runtime/Engine/Public/SceneView.h
--- a/Engine/Source/Runtime/Engine/Public/SceneView.h
+++ b/Engine/Source/Runtime/Engine/Public/SceneView.h
@@ -2330,12 +2330,8 @@ public:
 	// Optional description of view family for Unreal Insights profiling
 	FString ProfileDescription;
 
-	/**
-	 * Optional tracking of scene render time in seconds, useful when rendering multiple scenes (such as via UDisplayClusterViewportClient)
-	 * where the user may want a live breakdown of where render time is going per scene, which the stat system otherwise doesn't provide.
-	 * The destination memory location needs to be persistent and safe to access in any thread (for example, a standalone allocated buffer). 
-	 */
-	float* ProfileSceneRenderTime;
+	UE_DEPRECATED(5.6, "This member is no longer used")
+	float* ProfileSceneRenderTime = nullptr;
 
 	/** Views' origin of the (optional) streaming views. Used for prefetching rendering data ahead of time. */
 	TArray<FVector, TInlineAllocator<2>> StreamingViewOrigins;
diff --git a/Engine/Source/Runtime/Renderer/Private/DeferredShadingRenderer.cpp b/Engine/Source/Runtime/Renderer/Private/DeferredShadingRenderer.cpp
--- a/Engine/Source/Runtime/Renderer/Private/DeferredShadingRenderer.cpp
+++ b/Engine/Source/Runtime/Renderer/Private/DeferredShadingRenderer.cpp
@@ -1735,7 +1735,8 @@ void FDeferredShadingSceneRenderer::Render(FRDGBuilder& GraphBuilder)
 	if (FXSystem && Views.IsValidIndex(0))
 	{
 		SCOPE_CYCLE_COUNTER(STAT_FDeferredShadingSceneRenderer_FXSystem_PreRender);
-		FXSystem->PreRender(GraphBuilder, GetSceneViews(), GetSceneUniforms(), bIsFirstSceneRenderer /*bAllowGPUParticleUpdate*/);
+		const bool bAllowGPUParticleUpdate = true;
+		FXSystem->PreRender(GraphBuilder, GetSceneViews(), GetSceneUniforms(), bAllowGPUParticleUpdate);
 		if (FGPUSortManager* GPUSortManager = FXSystem->GetGPUSortManager())
 		{
 			GPUSortManager->OnPreRender(GraphBuilder);
@@ -1747,11 +1748,6 @@ void FDeferredShadingSceneRenderer::Render(FRDGBuilder& GraphBuilder)
 		RDG_EVENT_SCOPE_STAT(GraphBuilder, GPUSceneUpdate, "GPUSceneUpdate");
 		RDG_GPU_STAT_SCOPE(GraphBuilder, GPUSceneUpdate);
 
-		if (bIsFirstSceneRenderer)
-		{
-			GraphBuilder.SetFlushResourcesRHI();
-		}
-
 		for (int32 ViewIndex = 0; ViewIndex < AllViews.Num(); ViewIndex++)
 		{
 			FViewInfo& View = *AllViews[ViewIndex];
diff --git a/Engine/Source/Runtime/Renderer/Private/MeshEdgesRendering.cpp b/Engine/Source/Runtime/Renderer/Private/MeshEdgesRendering.cpp
--- a/Engine/Source/Runtime/Renderer/Private/MeshEdgesRendering.cpp
+++ b/Engine/Source/Runtime/Renderer/Private/MeshEdgesRendering.cpp
@@ -34,6 +34,7 @@
 #include "LegacyScreenPercentageDriver.h"
 #include "ScreenPass.h"
 #include "PostProcess/TemporalAA.h"
+#include "SceneRenderBuilder.h"
 
 class FComposeMeshEdgesPS : public FGlobalShader
 {
@@ -265,39 +266,6 @@ void CopyViewFamily(const FSceneViewFamily& SrcViewFamily, FSceneViewFamily& Vie
 	}
 }
 
-void RenderMeshEdges_RenderThread(
-	FRHICommandListImmediate& RHICmdList,
-	FSceneRenderer* SceneRenderer,
-	FMeshEdgesViewFamilyData& MeshEdgesData,
-	bool bUseSceneColorTexture)
-{
-	MeshEdgesData.WireframeColor->InitResource(RHICmdList);
-	MeshEdgesData.WireframeDepth->InitResource(RHICmdList);
-
-	SceneRenderer->RenderThreadBegin(RHICmdList);
-
-	FUniformExpressionCacheAsyncUpdateScope AsyncUpdateScope;
-
-	FRDGBuilder GraphBuilder(RHICmdList, RDG_EVENT_NAME("MeshEdges"), ERDGBuilderFlags::Parallel, SceneRenderer->ShaderPlatform);
-	{
-		RDG_EVENT_SCOPE(GraphBuilder, "RenderMeshEdges_RenderThread");
-
-		// Render the scene normally
-		{
-			RDG_RHI_EVENT_SCOPE(GraphBuilder, RenderScene);
-			SceneRenderer->Render(GraphBuilder);
-		}
-	}
-	GraphBuilder.Execute();
-	
-	for (const FViewInfo& ViewInfo : SceneRenderer->Views)
-	{
-		MeshEdgesData.ViewRects.Emplace(ViewInfo.ViewRect);
-	}
-
-	SceneRenderer->RenderThreadEnd(RHICmdList);
-}
-
 void RenderMeshEdges(FSceneViewFamily& ViewFamily)
 {
 	if (!ViewFamily.EngineShowFlags.MeshEdges || ViewFamily.EngineShowFlags.HitProxies)
@@ -385,10 +353,25 @@ void RenderMeshEdges(FSceneViewFamily& ViewFamily)
 		}
 	}
 
-	ENQUEUE_RENDER_COMMAND(CaptureCommand)(
-		[SceneRenderer, ViewFamilyData](FRHICommandListImmediate& RHICmdList)
+	FSceneRenderBuilder SceneRenderBuilder(SceneRenderer->Scene);
+
+	ON_SCOPE_EXIT
 	{
-		RenderMeshEdges_RenderThread(RHICmdList, SceneRenderer, *ViewFamilyData, true);
+		SceneRenderBuilder.Execute();
+	};
+
+	SceneRenderBuilder.AddRenderer(SceneRenderer, TEXT("RenderMeshEdges"),
+		[ViewFamilyData] (FRDGBuilder& GraphBuilder, const FSceneRenderFunctionInputs& Inputs)
+	{
+		ViewFamilyData->WireframeColor->InitResource(GraphBuilder.RHICmdList);
+		ViewFamilyData->WireframeDepth->InitResource(GraphBuilder.RHICmdList);
+
+		Inputs.Renderer->Render(GraphBuilder);
+
+		for (const FViewInfo& ViewInfo : Inputs.Renderer->Views)
+		{
+			ViewFamilyData->ViewRects.Emplace(ViewInfo.ViewRect);
+		}
 	});
 }
 
diff --git a/Engine/Source/Runtime/Renderer/Private/MobileShadingRenderer.cpp b/Engine/Source/Runtime/Renderer/Private/MobileShadingRenderer.cpp
--- a/Engine/Source/Runtime/Renderer/Private/MobileShadingRenderer.cpp
+++ b/Engine/Source/Runtime/Renderer/Private/MobileShadingRenderer.cpp
@@ -1123,11 +1123,6 @@ void FMobileSceneRenderer::Render(FRDGBuilder& GraphBuilder)
 		BeginOcclusionScope(GraphBuilder, Views);
 	}
 	
-	if (bIsFirstSceneRenderer)
-	{
-		GraphBuilder.SetFlushResourcesRHI();
-	}
-
 	// Allow scene extensions to affect the scene uniform buffer
 	GetSceneExtensionsRenderers().UpdateSceneUniformBuffer(GraphBuilder, GetSceneUniforms());
 
diff --git a/Engine/Source/Runtime/Renderer/Private/PlanarReflectionRendering.cpp b/Engine/Source/Runtime/Renderer/Private/PlanarReflectionRendering.cpp
--- a/Engine/Source/Runtime/Renderer/Private/PlanarReflectionRendering.cpp
+++ b/Engine/Source/Runtime/Renderer/Private/PlanarReflectionRendering.cpp
@@ -36,6 +36,7 @@
 #include "SceneTextureParameters.h"
 #include "SceneViewExtension.h"
 #include "Substrate/Substrate.h"
+#include "SceneRenderBuilder.h"
 
 void SetupPlanarReflectionUniformParameters(const class FSceneView& View, const FPlanarReflectionSceneProxy* ReflectionSceneProxy, FPlanarReflectionUniformParameters& OutParameters)
 {
@@ -273,7 +274,6 @@ struct FUpdatePlanarReflectionContentsInputs
 	FPlanarReflectionSceneProxy* SceneProxy;
 	FPlanarReflectionRenderTarget* RenderTarget;
 	FPlane MirrorPlane;
-	FName OwnerName;
 	bool bIsMobilePixelProjectedReflectionEnabled;
 };
 
@@ -372,8 +372,6 @@ static void UpdatePlanarReflectionContents_RenderThread(
 		return;
 	}
 
-	FUniformExpressionCacheAsyncUpdateScope AsyncUpdateScope;
-
 	const ERHIFeatureLevel::Type FeatureLevel = SceneRenderer->FeatureLevel;
 
 	// Make sure we render to the same set of GPUs as the main scene renderer.
@@ -386,14 +384,6 @@ static void UpdatePlanarReflectionContents_RenderThread(
 		Inputs.RenderTarget->SetActiveGPUMask(FRHIGPUMask::GPU0());
 	}
 
-#if WANTS_DRAW_MESH_EVENTS
-	FString EventName;
-	Inputs.OwnerName.ToString(EventName);
-	RDG_EVENT_SCOPE(GraphBuilder, "PlanarReflection %s", *EventName);
-#else
-	RDG_EVENT_SCOPE(GraphBuilder, "UpdatePlanarReflectionContent_RenderThread");
-#endif
-
 	// Applies late update (if any) to view matrices and re-reflects
 	if (SceneRenderer->Views.Num() > 1)
 	{
@@ -411,11 +401,7 @@ static void UpdatePlanarReflectionContents_RenderThread(
 		}
 	}
 
-	// Render the scene normally
-	{
-		RDG_RHI_EVENT_SCOPE(GraphBuilder, RenderScene);
-		SceneRenderer->Render(GraphBuilder);
-	}
+	SceneRenderer->Render(GraphBuilder);
 
 	FRDGTextureRef ReflectionOutputTexture = GraphBuilder.RegisterExternalTexture(CreateRenderTarget(Inputs.RenderTarget->TextureRHI, TEXT("ReflectionOutputTexture")));
 	GraphBuilder.SetTextureAccessFinal(ReflectionOutputTexture, ERHIAccess::SRVGraphics);
@@ -455,7 +441,7 @@ extern void SetupSceneViewExtensionsForSceneCapture(
 	FSceneViewFamily& ViewFamily,
 	TConstArrayView<FSceneView*> Views);
 
-void FScene::UpdatePlanarReflectionContents(UPlanarReflectionComponent* CaptureComponent, FSceneRenderer& ParentSceneRenderer)
+void FScene::UpdatePlanarReflectionContents(UPlanarReflectionComponent* CaptureComponent, FSceneRenderer& ParentSceneRenderer, ISceneRenderBuilder& SceneRenderBuilder)
 {
 	check(CaptureComponent);
 
@@ -473,7 +459,7 @@ void FScene::UpdatePlanarReflectionContents(UPlanarReflectionComponent* CaptureC
 								&& (bIsMobilePixelProjectedReflectionEnabled || CaptureComponent->RenderTarget->TextureRHI.IsValid());
 
 	FPlanarReflectionRenderTarget* RenderTargetToRelease = nullptr;
-	bool bInitResource = false;
+	FPlanarReflectionRenderTarget* RenderTargetToInit = nullptr;
 
 	if (CaptureComponent->RenderTarget != nullptr && !bIsRenderTargetValid)
 	{
@@ -486,7 +472,10 @@ void FScene::UpdatePlanarReflectionContents(UPlanarReflectionComponent* CaptureC
 		CaptureComponent->RenderTarget = new FPlanarReflectionRenderTarget(DesiredPlanarReflectionTextureSize);
 
 		// Don't create the RenderTarget's RHI if it is used for mobile pixel projected reflection
-		bInitResource = !bIsMobilePixelProjectedReflectionEnabled;
+		if (!bIsMobilePixelProjectedReflectionEnabled)
+		{
+			RenderTargetToInit = CaptureComponent->RenderTarget;
+		}
 	}
 
 	const FMatrix ComponentTransform = CaptureComponent->GetComponentTransform().ToMatrixWithScale();
@@ -598,50 +587,63 @@ void FScene::UpdatePlanarReflectionContents(UPlanarReflectionComponent* CaptureC
 		CaptureComponent->ProjectionWithExtraFOV[ViewIndex] = SceneCaptureViewInfo[ViewIndex].ProjectionMatrix;
 	}
 
-	const FUpdatePlanarReflectionContentsInputs Inputs
+	FString EventName = TEXT("PlanarReflection");
+
+#if WANTS_DRAW_MESH_EVENTS
+	if (CaptureComponent->GetOwner())
+	{
+		EventName = CaptureComponent->GetOwner()->GetFName().ToString();
+	}
+#endif
+
+	SCENE_RENDER_GROUP_SCOPE(SceneRenderBuilder, MoveTemp(EventName), ESceneRenderGroupFlags::None);
+
+	if (RenderTargetToRelease)
+	{
+		SceneRenderBuilder.AddRenderCommand([RenderTargetToRelease](FRHICommandListImmediate& RHICmdList)
+		{
+			RenderTargetToRelease->ReleaseResource();
+			delete RenderTargetToRelease;
+		});
+		RenderTargetToRelease = nullptr;
+	}
+
+	if (RenderTargetToInit)
+	{
+		SceneRenderBuilder.AddRenderCommand([RenderTargetToInit] (FRHICommandListImmediate& RHICmdList)
+		{
+			RenderTargetToInit->InitResource(RHICmdList);
+		});
+	}
+
+	const FUpdatePlanarReflectionContentsInputs UpdateInputs
 	{
 		  .ParentSceneRenderer = &ParentSceneRenderer
 		, .SceneProxy = CaptureComponent->SceneProxy
 		, .RenderTarget = CaptureComponent->RenderTarget
 		, .MirrorPlane = MirrorPlane
-		, .OwnerName = CaptureComponent->GetOwner() ? CaptureComponent->GetOwner()->GetFName() : NAME_None
 		, .bIsMobilePixelProjectedReflectionEnabled = bIsMobilePixelProjectedReflectionEnabled
 	};
 
 	FSceneRenderer* SceneRenderer = FSceneRenderer::CreateSceneRenderer(&ViewFamily, nullptr);
 
-	ENQUEUE_RENDER_COMMAND(RenderReflectionCapture)(
-		[SceneRenderer, Inputs, SceneCaptureViewInfo = MoveTemp(SceneCaptureViewInfo), RenderTargetToRelease, bInitResource, bIsStereo] (FRHICommandListImmediate& RHICmdList)
+	SceneRenderBuilder.AddRenderer(SceneRenderer,
+		[
+			  UpdateInputs
+			, SceneCaptureViewInfo = MoveTemp(SceneCaptureViewInfo)
+			, bIsStereo
+		] (FRDGBuilder& GraphBuilder, const FSceneRenderFunctionInputs& Inputs)
 	{
-		SceneRenderer->RenderThreadBegin(RHICmdList);
-		{
-			FRDGBuilder GraphBuilder(RHICmdList, RDG_EVENT_NAME("PlanarReflection"), ERDGBuilderFlags::Parallel, SceneRenderer->ShaderPlatform);
+		// Remove the render target on the planar reflection proxy so that this planar reflection is not getting drawn in its own FSceneRenderer.
+		UpdateInputs.SceneProxy->RenderTarget = nullptr;
 	
-			if (RenderTargetToRelease)
-			{
-				RenderTargetToRelease->ReleaseResource();
-				delete RenderTargetToRelease;
-			}
-	
-			if (bInitResource)
-			{
-				Inputs.RenderTarget->InitResource(GraphBuilder.RHICmdList);
-			}
-	
-			// Remove the render target on the planar reflection proxy so that this planar reflection is not getting drawn in its own FSceneRenderer.
-			Inputs.SceneProxy->RenderTarget = nullptr;
-	
-			for (int32 ViewIndex = 0; ViewIndex < SceneCaptureViewInfo.Num(); ++ViewIndex)
-			{
-				Inputs.SceneProxy->ProjectionWithExtraFOV[ViewIndex] = SceneCaptureViewInfo[ViewIndex].ProjectionMatrix;
-				Inputs.SceneProxy->bIsStereo = bIsStereo;
-			}
-	
-			UpdatePlanarReflectionContents_RenderThread(GraphBuilder, SceneRenderer, Inputs);
-
-			GraphBuilder.Execute();
+		for (int32 ViewIndex = 0; ViewIndex < SceneCaptureViewInfo.Num(); ++ViewIndex)
+		{
+			UpdateInputs.SceneProxy->ProjectionWithExtraFOV[ViewIndex] = SceneCaptureViewInfo[ViewIndex].ProjectionMatrix;
+			UpdateInputs.SceneProxy->bIsStereo = bIsStereo;
 		}
-		SceneRenderer->RenderThreadEnd(RHICmdList);
+	
+		UpdatePlanarReflectionContents_RenderThread(GraphBuilder, Inputs.Renderer, UpdateInputs);
 	});
 }
 
diff --git a/Engine/Source/Runtime/Renderer/Private/ReflectionEnvironmentCapture.cpp b/Engine/Source/Runtime/Renderer/Private/ReflectionEnvironmentCapture.cpp
--- a/Engine/Source/Runtime/Renderer/Private/ReflectionEnvironmentCapture.cpp
+++ b/Engine/Source/Runtime/Renderer/Private/ReflectionEnvironmentCapture.cpp
@@ -51,6 +51,7 @@
 #include "PixelShaderUtils.h"
 #include "Materials/MaterialRenderProxy.h"
 #include "UnrealEngine.h"
+#include "SceneRenderBuilder.h"
 
 /** Near plane to use when capturing the scene. */
 float GReflectionCaptureNearPlane = 5;
@@ -502,7 +503,7 @@ int32 FindOrAllocateCubemapIndex(FScene* Scene, const UReflectionCaptureComponen
 
 /** Captures the scene for a reflection capture by rendering the scene multiple times and copying into a cubemap texture. */
 void CaptureSceneToScratchCubemap(
-	FRHICommandListImmediate& RHICmdList,
+	FRDGBuilder& GraphBuilder,
 	FSceneRenderer* SceneRenderer,
 	const FReflectionCubemapTexture& ReflectionCubemapTexture,
 	ECubeFace CubeFace,
@@ -512,122 +513,108 @@ void CaptureSceneToScratchCubemap(
 	const FLinearColor& LowerHemisphereColor,
 	bool bCapturingForMobile)
 {
-	SceneRenderer->RenderThreadBegin(RHICmdList);
-
 	const ERHIFeatureLevel::Type FeatureLevel = SceneRenderer->FeatureLevel;
 
-	FUniformExpressionCacheAsyncUpdateScope AsyncUpdateScope;
-
-	FRDGBuilder GraphBuilder(RHICmdList, RDG_EVENT_NAME("CubeMapCapture"), ERDGBuilderFlags::Parallel, SceneRenderer->ShaderPlatform);
+	// Render the scene normally for one face of the cubemap
+	SceneRenderer->Render(GraphBuilder);
 
+	AddPass(GraphBuilder, RDG_EVENT_NAME("FlushGPU"), [](FRHICommandListImmediate& InRHICmdList)
 	{
-		RDG_EVENT_SCOPE(GraphBuilder, "CubeMapCapture");
+		QUICK_SCOPE_CYCLE_COUNTER(STAT_CaptureSceneToScratchCubemap_Flush);
+		InRHICmdList.ImmediateFlush(EImmediateFlushType::FlushRHIThread);
 
-		// Render the scene normally for one face of the cubemap
-		SceneRenderer->Render(GraphBuilder);
-
-		AddPass(GraphBuilder, RDG_EVENT_NAME("FlushGPU"), [](FRHICommandListImmediate& InRHICmdList)
-		{
-			QUICK_SCOPE_CYCLE_COUNTER(STAT_CaptureSceneToScratchCubemap_Flush);
-			InRHICmdList.ImmediateFlush(EImmediateFlushType::FlushRHIThread);
+		// some platforms may not be able to keep enqueueing commands like crazy, this will
+		// allow them to restart their command buffers
+		InRHICmdList.SubmitCommandsAndFlushGPU();
+	});
 
-			// some platforms may not be able to keep enqueueing commands like crazy, this will
-			// allow them to restart their command buffers
-			InRHICmdList.SubmitCommandsAndFlushGPU();
-		});
+	const FViewInfo& View = SceneRenderer->Views[0];
 
-		const FViewInfo& View = SceneRenderer->Views[0];
+	FRDGTextureRef OutputTexture = ReflectionCubemapTexture.GetRDG(GraphBuilder);
 
-		FRDGTextureRef OutputTexture = ReflectionCubemapTexture.GetRDG(GraphBuilder);
+	auto* PassParameters = GraphBuilder.AllocParameters<FCopySceneColorToCubeFacePS::FParameters>();
+	PassParameters->RenderTargets[0] = FRenderTargetBinding(OutputTexture, ERenderTargetLoadAction::ENoAction, 0, CubeFace);
+	PassParameters->LowerHemisphereColor = LowerHemisphereColor;
 
-		auto* PassParameters = GraphBuilder.AllocParameters<FCopySceneColorToCubeFacePS::FParameters>();
-		PassParameters->RenderTargets[0] = FRenderTargetBinding(OutputTexture, ERenderTargetLoadAction::ENoAction, 0, CubeFace);
-		PassParameters->LowerHemisphereColor = LowerHemisphereColor;
+	{
+		FVector4f SkyLightParametersValue(ForceInitToZero);
+		FScene* Scene = SceneRenderer->Scene;
 
+		if (bCapturingForSkyLight)
 		{
-			FVector4f SkyLightParametersValue(ForceInitToZero);
-			FScene* Scene = SceneRenderer->Scene;
-
-			if (bCapturingForSkyLight)
-			{
-				// When capturing reflection captures, support forcing all low hemisphere lighting to be black
-				SkyLightParametersValue = FVector4f(0, 0, bLowerHemisphereIsBlack ? 1.0f : 0.0f, 0);
-			}
-			else if (!bCapturingForMobile && Scene->SkyLight && !Scene->SkyLight->bHasStaticLighting)	
-			{
-				// Mobile renderer can't blend reflections with a sky at runtime, so we dont use this path when capturing for a mobile renderer
+			// When capturing reflection captures, support forcing all low hemisphere lighting to be black
+			SkyLightParametersValue = FVector4f(0, 0, bLowerHemisphereIsBlack ? 1.0f : 0.0f, 0);
+		}
+		else if (!bCapturingForMobile && Scene->SkyLight && !Scene->SkyLight->bHasStaticLighting)	
+		{
+			// Mobile renderer can't blend reflections with a sky at runtime, so we dont use this path when capturing for a mobile renderer
 				
-				// When capturing reflection captures and there's a stationary sky light, mask out any pixels whose depth classify it as part of the sky
-				// This will allow changing the stationary sky light at runtime
-				SkyLightParametersValue = FVector4f(1, Scene->SkyLight->SkyDistanceThreshold, 0, 0);
-			}
-			else
-			{
-				// When capturing reflection captures and there's no sky light, or only a static sky light, capture all depth ranges
-				SkyLightParametersValue = FVector4f(2, 0, 0, 0);
-			}
-
-			PassParameters->SkyLightCaptureParameters = SkyLightParametersValue;
+			// When capturing reflection captures and there's a stationary sky light, mask out any pixels whose depth classify it as part of the sky
+			// This will allow changing the stationary sky light at runtime
+			SkyLightParametersValue = FVector4f(1, Scene->SkyLight->SkyDistanceThreshold, 0, 0);
+		}
+		else
+		{
+			// When capturing reflection captures and there's no sky light, or only a static sky light, capture all depth ranges
+			SkyLightParametersValue = FVector4f(2, 0, 0, 0);
 		}
 
-		const FMinimalSceneTextures& SceneTextures = View.GetSceneTextures();
-
-		PassParameters->View = View.ViewUniformBuffer;
-		PassParameters->SceneColorSampler = TStaticSamplerState<SF_Bilinear, AM_Clamp, AM_Clamp, AM_Clamp>::GetRHI();
-		PassParameters->SceneColorTexture = SceneTextures.Color.Resolve;
-        
-        if (!IsMobilePlatform(SceneRenderer->ShaderPlatform))
-        {
-            PassParameters->SceneDepthSampler = TStaticSamplerState<SF_Point, AM_Clamp, AM_Clamp, AM_Clamp>::GetRHI();
-            PassParameters->SceneDepthTexture = SceneTextures.Depth.Resolve;
-        }
+		PassParameters->SkyLightCaptureParameters = SkyLightParametersValue;
+	}
 
-		const int32 EffectiveSize = CubemapSize;
-		const FIntPoint SceneTextureExtent = SceneTextures.Config.Extent;
+	const FMinimalSceneTextures& SceneTextures = View.GetSceneTextures();
 
-		GraphBuilder.AddPass(
-			RDG_EVENT_NAME("CopySceneToCubeFace"),
-			PassParameters,
-			ERDGPassFlags::Raster,
-			[EffectiveSize, SceneTextureExtent, FeatureLevel, PassParameters](FRDGAsyncTask, FRHICommandList& InRHICmdList)
-		{
-			const FIntRect ViewRect(0, 0, EffectiveSize, EffectiveSize);
-			InRHICmdList.SetViewport(0.0f, 0.0f, 0.0f, (float)EffectiveSize, (float)EffectiveSize, 1.0f);
-
-			FGraphicsPipelineStateInitializer GraphicsPSOInit;
-			InRHICmdList.ApplyCachedRenderTargets(GraphicsPSOInit);
-			GraphicsPSOInit.RasterizerState = TStaticRasterizerState<FM_Solid, CM_None>::GetRHI();
-			GraphicsPSOInit.DepthStencilState = TStaticDepthStencilState<false, CF_Always>::GetRHI();
-			GraphicsPSOInit.BlendState = TStaticBlendState<>::GetRHI();
-
-			TShaderMapRef<FCopyToCubeFaceVS> VertexShader(GetGlobalShaderMap(FeatureLevel));
-			TShaderMapRef<FCopySceneColorToCubeFacePS> PixelShader(GetGlobalShaderMap(FeatureLevel));
-
-			GraphicsPSOInit.BoundShaderState.VertexDeclarationRHI = GFilterVertexDeclaration.VertexDeclarationRHI;
-			GraphicsPSOInit.BoundShaderState.VertexShaderRHI = VertexShader.GetVertexShader();
-			GraphicsPSOInit.BoundShaderState.PixelShaderRHI = PixelShader.GetPixelShader();
-			GraphicsPSOInit.PrimitiveType = PT_TriangleList;
-
-			SetGraphicsPipelineState(InRHICmdList, GraphicsPSOInit, 0);
-			SetShaderParameters(InRHICmdList, PixelShader, PixelShader.GetPixelShader(), *PassParameters);
-
-			const int32 SupersampleCaptureFactor = FMath::Clamp(GSupersampleCaptureFactor, MinSupersampleCaptureFactor, MaxSupersampleCaptureFactor);
-
-			DrawRectangle( 
-				InRHICmdList,
-				ViewRect.Min.X, ViewRect.Min.Y, 
-				ViewRect.Width(), ViewRect.Height(),
-				ViewRect.Min.X, ViewRect.Min.Y, 
-				ViewRect.Width() * SupersampleCaptureFactor, ViewRect.Height() * SupersampleCaptureFactor,
-				FIntPoint(ViewRect.Width(), ViewRect.Height()),
-				SceneTextureExtent,
-				VertexShader);
-		});
-	}
+	PassParameters->View = View.ViewUniformBuffer;
+	PassParameters->SceneColorSampler = TStaticSamplerState<SF_Bilinear, AM_Clamp, AM_Clamp, AM_Clamp>::GetRHI();
+	PassParameters->SceneColorTexture = SceneTextures.Color.Resolve;
+        
+    if (!IsMobilePlatform(SceneRenderer->ShaderPlatform))
+    {
+        PassParameters->SceneDepthSampler = TStaticSamplerState<SF_Point, AM_Clamp, AM_Clamp, AM_Clamp>::GetRHI();
+        PassParameters->SceneDepthTexture = SceneTextures.Depth.Resolve;
+    }
 
-	GraphBuilder.Execute();
+	const int32 EffectiveSize = CubemapSize;
+	const FIntPoint SceneTextureExtent = SceneTextures.Config.Extent;
 
-	SceneRenderer->RenderThreadEnd(RHICmdList);
+	GraphBuilder.AddPass(
+		RDG_EVENT_NAME("CopySceneToCubeFace"),
+		PassParameters,
+		ERDGPassFlags::Raster,
+		[EffectiveSize, SceneTextureExtent, FeatureLevel, PassParameters](FRDGAsyncTask, FRHICommandList& InRHICmdList)
+	{
+		const FIntRect ViewRect(0, 0, EffectiveSize, EffectiveSize);
+		InRHICmdList.SetViewport(0.0f, 0.0f, 0.0f, (float)EffectiveSize, (float)EffectiveSize, 1.0f);
+
+		FGraphicsPipelineStateInitializer GraphicsPSOInit;
+		InRHICmdList.ApplyCachedRenderTargets(GraphicsPSOInit);
+		GraphicsPSOInit.RasterizerState = TStaticRasterizerState<FM_Solid, CM_None>::GetRHI();
+		GraphicsPSOInit.DepthStencilState = TStaticDepthStencilState<false, CF_Always>::GetRHI();
+		GraphicsPSOInit.BlendState = TStaticBlendState<>::GetRHI();
+
+		TShaderMapRef<FCopyToCubeFaceVS> VertexShader(GetGlobalShaderMap(FeatureLevel));
+		TShaderMapRef<FCopySceneColorToCubeFacePS> PixelShader(GetGlobalShaderMap(FeatureLevel));
+
+		GraphicsPSOInit.BoundShaderState.VertexDeclarationRHI = GFilterVertexDeclaration.VertexDeclarationRHI;
+		GraphicsPSOInit.BoundShaderState.VertexShaderRHI = VertexShader.GetVertexShader();
+		GraphicsPSOInit.BoundShaderState.PixelShaderRHI = PixelShader.GetPixelShader();
+		GraphicsPSOInit.PrimitiveType = PT_TriangleList;
+
+		SetGraphicsPipelineState(InRHICmdList, GraphicsPSOInit, 0);
+		SetShaderParameters(InRHICmdList, PixelShader, PixelShader.GetPixelShader(), *PassParameters);
+
+		const int32 SupersampleCaptureFactor = FMath::Clamp(GSupersampleCaptureFactor, MinSupersampleCaptureFactor, MaxSupersampleCaptureFactor);
+
+		DrawRectangle( 
+			InRHICmdList,
+			ViewRect.Min.X, ViewRect.Min.Y, 
+			ViewRect.Width(), ViewRect.Height(),
+			ViewRect.Min.X, ViewRect.Min.Y, 
+			ViewRect.Width() * SupersampleCaptureFactor, ViewRect.Height() * SupersampleCaptureFactor,
+			FIntPoint(ViewRect.Width(), ViewRect.Height()),
+			SceneTextureExtent,
+			VertexShader);
+	});
 }
 
 
@@ -1256,19 +1243,33 @@ void CaptureSceneIntoScratchCubemap(
 
 	TRenderThreadStruct<FDummyRenderTarget> DummyRenderTarget;
 
-	for (int32 CubeFace = 0; CubeFace < CubeFace_MAX; CubeFace++)
+	if (!bCapturingForSkyLight && !bInsideTick)
+	{
+		ENQUEUE_RENDER_COMMAND(BeginCubemapCapture)([] (FRHICommandListImmediate&)
+		{
+			GFrameNumberRenderThread++;
+		});
+	}
+
+	FSceneRenderBuilder SceneRenderBuilder(Scene);
+
+	ON_SCOPE_EXIT
 	{
-		if( !bCapturingForSkyLight && !bInsideTick )
+		SceneRenderBuilder.Execute();
+
+		if (!bCapturingForSkyLight && !bInsideTick)
 		{
-			// Alert the RHI that we're rendering a new frame
-			// Not really a new frame, but it will allow pooling mechanisms to update, like the uniform buffer pool
-			ENQUEUE_RENDER_COMMAND(BeginFrame)(
-				[](FRHICommandListImmediate& RHICmdList)
+			ENQUEUE_RENDER_COMMAND(EndCubemapCapture)([] (FRHICommandListImmediate& RHICmdList)
 			{
-				GFrameNumberRenderThread++;
+				RHICmdList.EndFrame();
 			});
 		}
+	};
+
+	SCENE_RENDER_GROUP_SCOPE(SceneRenderBuilder, TEXT("CubemapCapture"), ESceneRenderGroupFlags::None);
 
+	for (int32 CubeFace = 0; CubeFace < CubeFace_MAX; CubeFace++)
+	{
 		DummyRenderTarget->SetSize(CubemapSize);
 
 		auto ViewFamilyInit = FSceneViewFamily::ConstructionValues(
@@ -1354,15 +1355,19 @@ void CaptureSceneIntoScratchCubemap(
 
 		FSceneRenderer* SceneRenderer = FSceneRenderer::CreateSceneRenderer(&ViewFamily, nullptr);
 
-		ENQUEUE_RENDER_COMMAND(CaptureCommand)(
-			[SceneRenderer, &ReflectionCubemapTexture, CubeFace, CubemapSize, bCapturingForSkyLight, bLowerHemisphereIsBlack, LowerHemisphereColor, bCapturingForMobile, bInsideTick](FRHICommandListImmediate& RHICmdList)
+		SceneRenderBuilder.AddRenderer(SceneRenderer,
+			[
+				  &ReflectionCubemapTexture
+				, CubeFace
+				, CubemapSize
+				, bCapturingForSkyLight
+				, bLowerHemisphereIsBlack
+				, LowerHemisphereColor
+				, bCapturingForMobile
+				, bInsideTick
+			] (FRDGBuilder& GraphBuilder, const FSceneRenderFunctionInputs& Inputs)
 		{
-			CaptureSceneToScratchCubemap(RHICmdList, SceneRenderer, ReflectionCubemapTexture, (ECubeFace)CubeFace, CubemapSize, bCapturingForSkyLight, bLowerHemisphereIsBlack, LowerHemisphereColor, bCapturingForMobile);
-
-			if (!bCapturingForSkyLight && !bInsideTick)
-			{
-				RHICmdList.EndFrame();
-			}
+			CaptureSceneToScratchCubemap(GraphBuilder, Inputs.Renderer, ReflectionCubemapTexture, (ECubeFace)CubeFace, CubemapSize, bCapturingForSkyLight, bLowerHemisphereIsBlack, LowerHemisphereColor, bCapturingForMobile);
 		});
 	}
 }
diff --git a/Engine/Source/Runtime/Renderer/Private/SceneCaptureRendering.cpp b/Engine/Source/Runtime/Renderer/Private/SceneCaptureRendering.cpp
--- a/Engine/Source/Runtime/Renderer/Private/SceneCaptureRendering.cpp
+++ b/Engine/Source/Runtime/Renderer/Private/SceneCaptureRendering.cpp
@@ -47,6 +47,7 @@
 #include "IRenderCaptureProvider.h"
 #include "RenderCaptureInterface.h"
 #include "CustomRenderPassSceneCapture.h"
+#include "SceneRenderBuilder.h"
 
 bool GSceneCaptureAllowRenderInMainRenderer = true;
 static FAutoConsoleVariableRef CVarSceneCaptureAllowRenderInMainRenderer(
@@ -421,8 +422,6 @@ static void UpdateSceneCaptureContent_RenderThread(
 {
 	const EShadingPath ShadingPath = GetFeatureLevelShadingPath(SceneRenderer->Scene->GetFeatureLevel());
 
-	FUniformExpressionCacheAsyncUpdateScope AsyncUpdateScope;
-
 	// The target texture is what gets rendered to, while OutputTexture is the final output.  For 2D scene captures, these textures
 	// are the same.  For cube captures, OutputTexture will be a cube map, while TargetTexture will be a 2D render target containing either
 	// one face of the cube map (when GSceneCaptureCubeSinglePass=0) or the six faces of the cube map tiled in a split screen configuration.
@@ -468,11 +467,7 @@ static void UpdateSceneCaptureContent_RenderThread(
 		View.bIgnoreExistingQueries = true;
 	}
 
-	// Render the scene normally
-	{
-		RDG_RHI_EVENT_SCOPE(GraphBuilder, RenderScene);
-		SceneRenderer->Render(GraphBuilder);
-	}
+	SceneRenderer->Render(GraphBuilder);
 
 	if (ShadingPath == EShadingPath::Mobile)
 	{
@@ -1005,73 +1000,34 @@ public:
 	bool bAutoGenerateMips = false;
 };
 
-static void BeginGpuCaptureOrDump(USceneCaptureComponent* CaptureComponent, bool& bCapturingGPU, bool& bDumpingGPU)
+ESceneRenderGroupFlags GetSceneCaptureGroupFlags(USceneCaptureComponent* CaptureComponent)
 {
-	bCapturingGPU = false;
-	bDumpingGPU = false;
-
-	if (CaptureComponent->bSuppressGpuCaptureOrDump)
-	{
-		CaptureComponent->bSuppressGpuCaptureOrDump = false;
-		return;
-	}
-	
-	bCapturingGPU = CaptureComponent->bCaptureGpuNextRender;
-	bDumpingGPU = CaptureComponent->bDumpGpuNextRender;
-
-	CaptureComponent->bCaptureGpuNextRender = false;
-	CaptureComponent->bDumpGpuNextRender = false;
+	ESceneRenderGroupFlags Flags = ESceneRenderGroupFlags::None;
 
-	// Clear capturing flag if it's not available
-	if (!IRenderCaptureProvider::IsAvailable())
+	if (!CaptureComponent->bSuppressGpuCaptureOrDump)
 	{
-		bCapturingGPU = false;
-	}
-
-	// If user sets both capture and dump flags, prefer capturing over dumping (or clear flag if dumping is not available)
-	if (bCapturingGPU || !(WITH_ENGINE && WITH_DUMPGPU))
-	{
-		bDumpingGPU = false;
-	}
-
-#if WITH_ENGINE && WITH_DUMPGPU
-	if (bDumpingGPU)
-	{
-		// Don't try to start a dump if we are already dumping for some reason
-		if (FRDGBuilder::IsDumpingFrame())
+		if (CaptureComponent->bCaptureGpuNextRender)
 		{
-			bDumpingGPU = false;
+			CaptureComponent->bCaptureGpuNextRender = false;
+			Flags |= ESceneRenderGroupFlags::GpuCapture;
 		}
-		else
-		{
-			// Pass "-oneframe" to override CVar that could enable multiple frames of capture
-			FRDGBuilder::BeginResourceDump(TEXT("-oneframe"));
 
-			// Tick the DumpGPU system, which will start the dump
-			UE::RenderCore::DumpGPU::TickEndFrame();
+		if (CaptureComponent->bDumpGpuNextRender)
+		{
+			CaptureComponent->bDumpGpuNextRender = false;
+			Flags |= ESceneRenderGroupFlags::GpuDump;
 		}
 	}
-#endif
-}
 
-static void EndGpuCaptureOrDump(bool bDumpingGPU)
-{
-#if WITH_ENGINE && WITH_DUMPGPU
-	if (bDumpingGPU)
-	{
-		// Tick the dump GPU system again, which will end the active dump, so it just includes the scene capture
-		UE::RenderCore::DumpGPU::TickEndFrame();
-	}
-#endif
+	CaptureComponent->bSuppressGpuCaptureOrDump = false;
+	return Flags;
 }
 
-void FScene::UpdateSceneCaptureContents(USceneCaptureComponent2D* CaptureComponent)
+void FScene::UpdateSceneCaptureContents(USceneCaptureComponent2D* CaptureComponent, ISceneRenderBuilder& SceneRenderBuilder)
 {
 	check(CaptureComponent);
 
-	bool bCapturingGPU;
-	bool bDumpingGPU;
-	BeginGpuCaptureOrDump(CaptureComponent, bCapturingGPU, bDumpingGPU);
+	const ESceneRenderGroupFlags SceneRenderGroupFlags = GetSceneCaptureGroupFlags(CaptureComponent);
 
 	if (UTextureRenderTarget2D* TextureRenderTarget = CaptureComponent->TextureTarget)
 	{
@@ -1288,10 +1244,6 @@ void FScene::UpdateSceneCaptureContents(USceneCaptureComponent2D* CaptureCompone
 
 		check(SceneRenderer != nullptr);
 
-		// When bIsMultipleSceneCapture is true, set bIsFirstSceneRenderer to false, which tells the scene renderer it can skip RHI resource flush, saving performance
-		bool bIsMultipleSceneCapture = CaptureComponent->SetFrameUpdated();
-		SceneRenderer->bIsFirstSceneRenderer = !bIsMultipleSceneCapture;
-
 		SceneRenderer->Views[0].bSceneCaptureMainViewJitter = bCopyMainViewTemporalSettings2D;
 		SceneRenderer->Views[0].bFogOnlyOnRenderedOpaque = CaptureComponent->bConsiderUnrenderedOpaquePixelAsFullyTranslucent;
 
@@ -1356,8 +1308,6 @@ void FScene::UpdateSceneCaptureContents(USceneCaptureComponent2D* CaptureCompone
 			}
 		}
 
-		UTexture* TexturePtrNotDeferenced = TextureRenderTarget;
-
 		// Compositing feature is only active when using SceneColor as the source
 		bool bIsCompositing = (CaptureComponent->CompositeMode != SCCM_Overwrite) && (CaptureComponent->CaptureSource == SCS_SceneColorHDR);
 #if WITH_EDITOR
@@ -1367,84 +1317,85 @@ void FScene::UpdateSceneCaptureContents(USceneCaptureComponent2D* CaptureCompone
 		}
 		TRefCountPtr<FSceneCaptureMemorySize> CaptureMemorySize = CaptureComponent->CaptureMemorySize;
 #else
-		void* CaptureMemorySize = nullptr;		// Dummy value for lambda capture argument list
+		void* CaptureMemorySize = nullptr;
 #endif
 
-		for (const FSceneViewExtensionRef& Extension : SceneRenderer->ViewFamily.ViewExtensions)
-		{
-			Extension->BeginRenderViewFamily(SceneRenderer->ViewFamily);
-		}
-
-		UE::RenderCommandPipe::FSyncScope SyncScope;
-
-		ENQUEUE_RENDER_COMMAND(CaptureCommand)(
-			[SceneRenderer, TextureRenderTargetResource, TexturePtrNotDeferenced, EventName, TargetName, bGenerateMips, GenerateMipsParams, GameViewportRT, bEnableOrthographicTiling, bIsCompositing, bOrthographicCamera, NumXTiles, NumYTiles, TileID, CaptureMemorySize, bCapturingGPU, CaptureSize](FRHICommandListImmediate& RHICmdList)
+		// This needs to happen now as a render command since the resize operation will queue up a deferred resource update.
+		ENQUEUE_RENDER_COMMAND(ResizeSceneCapture)([TextureRenderTargetResource, CaptureSize, bGenerateMips, GameViewportRT] (FRHICommandListImmediate& RHICmdList)
 		{
 			// Resize the render resource if necessary, either to the main viewport size overridden above (see ShouldRenderWithMainViewResolution()),
 			// or the original size if we are changing back to that (the resize call does nothing if the size already matches).
 			TextureRenderTargetResource->GetTextureRenderTarget2DResource()->Resize(RHICmdList, CaptureSize.X, CaptureSize.Y, bGenerateMips);
 
-			SceneRenderer->RenderThreadBegin(RHICmdList);
-
-		#if WANTS_DRAW_MESH_EVENTS
-			SCOPED_DRAW_EVENTF(RHICmdList, SceneCaptureMobile, TEXT("SceneCaptureMobile %s"), EventName);
-			FRDGBuilder GraphBuilder(RHICmdList, RDG_EVENT_NAME("SceneCaptureMobile %s", *EventName), ERDGBuilderFlags::Parallel, SceneRenderer->ShaderPlatform);
-		#else
-			SCOPED_DRAW_EVENT(RHICmdList, UpdateSceneCaptureContentMobile_RenderThread);
-			FRDGBuilder GraphBuilder(RHICmdList, RDG_EVENT_NAME("SceneCaptureMobile"), ERDGBuilderFlags::Parallel, SceneRenderer->ShaderPlatform);
-		#endif
-
+			if (GameViewportRT != nullptr)
 			{
-				RenderCaptureInterface::FScopedCapture RenderCapture(bCapturingGPU, &RHICmdList, *FString::Format(TEXT("Scene Capture : {0}"), { EventName }));
-
-				if (GameViewportRT != nullptr)
-				{
-					const FRHIGPUMask GPUMask = GameViewportRT->GetGPUMask(RHICmdList);
-					TextureRenderTargetResource->SetActiveGPUMask(GPUMask);
-				}
-				else
-				{
-					TextureRenderTargetResource->SetActiveGPUMask(FRHIGPUMask::All());
-				}
-
-				FRHICopyTextureInfo CopyInfo;
-
-				if (bEnableOrthographicTiling)
-				{
-					const uint32 RTSizeX = TextureRenderTargetResource->GetSizeX() / NumXTiles;
-					const uint32 RTSizeY = TextureRenderTargetResource->GetSizeY() / NumYTiles;
-					const uint32 TileX = TileID % NumXTiles;
-					const uint32 TileY = TileID / NumXTiles;
-					CopyInfo.DestPosition.X = TileX * RTSizeX;
-					CopyInfo.DestPosition.Y = TileY * RTSizeY;
-					CopyInfo.Size.X = RTSizeX;
-					CopyInfo.Size.Y = RTSizeY;
-				}
-
-				RectLightAtlas::FAtlasTextureInvalidationScope Invalidation(TexturePtrNotDeferenced);
+				TextureRenderTargetResource->SetActiveGPUMask(GameViewportRT->GetGPUMask(RHICmdList));
+			}
+			else
+			{
+				TextureRenderTargetResource->SetActiveGPUMask(FRHIGPUMask::All());
+			}
+		});
 
-			#if WITH_EDITOR
-				// Scene renderer may be deleted in UpdateSceneCaptureContent_RenderThread, grab view state pointer first
-				const FSceneViewState* ViewState = SceneRenderer->Views[0].ViewState;
-			#endif  // WITH_EDITOR
+		SCENE_RENDER_GROUP_SCOPE(SceneRenderBuilder, MoveTemp(EventName), SceneRenderGroupFlags);
+
+		SceneRenderBuilder.AddRenderer(
+			SceneRenderer,
+			[
+				  TextureRenderTargetResource
+				, TextureRenderTarget
+				, TargetName
+				, bGenerateMips
+				, GenerateMipsParams
+				, bEnableOrthographicTiling
+				, bIsCompositing
+				, bOrthographicCamera
+				, NumXTiles
+				, NumYTiles
+				, TileID
+				, CaptureMemorySize
+				, CaptureSize
+			] (FRDGBuilder& GraphBuilder, const FSceneRenderFunctionInputs& Inputs)
+		{
+			FRHICopyTextureInfo CopyInfo;
 
-				// Don't clear the render target when compositing, or in a tiling mode that fills in the render target in multiple passes.
-				bool bClearRenderTarget = !bIsCompositing && !bEnableOrthographicTiling;
+			if (bEnableOrthographicTiling)
+			{
+				const uint32 RTSizeX = TextureRenderTargetResource->GetSizeX() / NumXTiles;
+				const uint32 RTSizeY = TextureRenderTargetResource->GetSizeY() / NumYTiles;
+				const uint32 TileX = TileID % NumXTiles;
+				const uint32 TileY = TileID / NumXTiles;
+				CopyInfo.DestPosition.X = TileX * RTSizeX;
+				CopyInfo.DestPosition.Y = TileY * RTSizeY;
+				CopyInfo.Size.X = RTSizeX;
+				CopyInfo.Size.Y = RTSizeY;
+			}
 
-				VISUALIZE_TEXTURE_BEGIN_VIEW(SceneRenderer->FeatureLevel, SceneRenderer->Views[0].GetViewKey(), *EventName, true);
+			RectLightAtlas::FAtlasTextureInvalidationScope Invalidation(TextureRenderTarget);
 
-				UpdateSceneCaptureContent_RenderThread(GraphBuilder, SceneRenderer, TextureRenderTargetResource, TextureRenderTargetResource,
-					MakeArrayView(&CopyInfo, 1), bGenerateMips, GenerateMipsParams, bClearRenderTarget, bOrthographicCamera);
+			// Don't clear the render target when compositing, or in a tiling mode that fills in the render target in multiple passes.
+			const bool bClearRenderTarget = !bIsCompositing && !bEnableOrthographicTiling;
 
-				VISUALIZE_TEXTURE_END_VIEW();
+			UpdateSceneCaptureContent_RenderThread(
+				GraphBuilder,
+				Inputs.Renderer,
+				TextureRenderTargetResource,
+				TextureRenderTargetResource,
+				MakeArrayView({ CopyInfo }),
+				bGenerateMips,
+				GenerateMipsParams,
+				bClearRenderTarget,
+				bOrthographicCamera);
 
-			#if WITH_EDITOR
+		#if WITH_EDITOR
+			if (const FSceneViewState* ViewState = Inputs.Renderer->Views[0].ViewState)
+			{
 				if (ViewState)
 				{
 					const bool bLogSizes = GDumpSceneCaptureMemoryFrame == GFrameNumberRenderThread;
 					if (bLogSizes)
 					{
-						UE_LOG(LogRenderer, Log, TEXT("LogSizes\tSceneCapture\t%s\t%s\t%dx%d"), *EventName, *TargetName.ToString(), TextureRenderTargetResource->GetSizeX(), TextureRenderTargetResource->GetSizeY());
+						UE_LOG(LogRenderer, Log, TEXT("LogSizes\tSceneCapture\t%s\t%s\t%dx%d"), Inputs.FullPath, *TargetName.ToString(), TextureRenderTargetResource->GetSizeX(), TextureRenderTargetResource->GetSizeY());
 					}
 					CaptureMemorySize->Size = ViewState->GetGPUSizeBytes(bLogSizes);
 				}
@@ -1452,14 +1403,10 @@ void FScene::UpdateSceneCaptureContents(USceneCaptureComponent2D* CaptureCompone
 				{
 					CaptureMemorySize->Size = 0;
 				}
-			#endif  // WITH_EDITOR
 			}
-			GraphBuilder.Execute();
-			SceneRenderer->RenderThreadEnd(RHICmdList);
+		#endif  // WITH_EDITOR
 		});
 	}
-
-	EndGpuCaptureOrDump(bDumpingGPU);
 }
 
 // Split screen cube map faces are rendered as 3x2 tiles.
@@ -1473,11 +1420,9 @@ static const int32 GCubeFaceViewportOffsets[6][2] =
 	{ 2,1 },
 };
 
-void FScene::UpdateSceneCaptureContents(USceneCaptureComponentCube* CaptureComponent)
+void FScene::UpdateSceneCaptureContents(USceneCaptureComponentCube* CaptureComponent, ISceneRenderBuilder& SceneRenderBuilder)
 {
-	bool bCapturingGPU;
-	bool bDumpingGPU;
-	BeginGpuCaptureOrDump(CaptureComponent, bCapturingGPU, bDumpingGPU);
+	ESceneRenderGroupFlags SceneRenderGroupFlags = GetSceneCaptureGroupFlags(CaptureComponent);
 
 	struct FLocal
 	{
@@ -1578,24 +1523,12 @@ void FScene::UpdateSceneCaptureContents(USceneCaptureComponentCube* CaptureCompo
 
 		const FVector Location = CaptureComponent->GetComponentToWorld().GetTranslation();
 
-		bool bIsMultipleSceneCapture = CaptureComponent->SetFrameUpdated();
 		bool bCaptureSceneColor = CaptureNeedsSceneColor(CaptureComponent->CaptureSource);
 
+		SCENE_RENDER_GROUP_SCOPE(SceneRenderBuilder, MoveTemp(EventName), SceneRenderGroupFlags);
+
 		if (GSceneCaptureCubeSinglePass == false)
 		{
-			// For GPU capture to work for multi-pass rendering, we need the capture scope to persist across all the scene render command lambdas,
-			// so we need to allocate a pointer on the heap, and let the last render command clean it up.
-			RenderCaptureInterface::FScopedCapture** ScopedCapturePtr;
-			if (bCapturingGPU)
-			{
-				ScopedCapturePtr = new RenderCaptureInterface::FScopedCapture*;
-				*ScopedCapturePtr = nullptr;
-			}
-			else
-			{
-				ScopedCapturePtr = nullptr;
-			}
-
 			for (int32 FaceIndex = 0; FaceIndex < (int32)ECubeFace::CubeFace_MAX; FaceIndex++)
 			{
 				const ECubeFace TargetFace = (ECubeFace)FaceIndex;
@@ -1611,80 +1544,44 @@ void FScene::UpdateSceneCaptureContents(USceneCaptureComponentCube* CaptureCompo
 					Location, ProjectionMatrix, CaptureComponent->MaxViewDistanceOverride, FOVInDegrees,
 					bCaptureSceneColor, bCameraCut2D, bCopyMainViewTemporalSettings2D, &CaptureComponent->PostProcessSettings, CaptureComponent->PostProcessBlendWeight, CaptureComponent->GetViewOwner(), FaceIndex);
 
-				// When bIsMultipleSceneCapture is true, set bIsFirstSceneRenderer to false, which tells the scene renderer it can skip RHI resource flush, saving performance.
-				// We can also skip RHI resource flush on faces after the first.
-				SceneRenderer->bIsFirstSceneRenderer = (FaceIndex == 0) && !bIsMultipleSceneCapture;
-
-				for (const FSceneViewExtensionRef& Extension : SceneRenderer->ViewFamily.ViewExtensions)
+				SceneRenderBuilder.AddRenderer(SceneRenderer, FString::Printf(TEXT("CubeFace[%d]"), FaceIndex),
+					[
+						  TextureRenderTarget
+						, TargetFace
+						, bGenerateMips
+						, GenerateMipsParams
+					] (FRDGBuilder& GraphBuilder, const FSceneRenderFunctionInputs& Inputs)
 				{
-					Extension->BeginRenderViewFamily(SceneRenderer->ViewFamily);
-				}
-
-				// Include the cube face index in the event name
-				EventName.Appendf(TEXT(" [%d]"), FaceIndex);
-
-				UE::RenderCommandPipe::FSyncScope SyncScope;
-
-				ENQUEUE_RENDER_COMMAND(CaptureCommand)(
-					[SceneRenderer, TextureRenderTarget, EventName, TargetFace, bGenerateMips, GenerateMipsParams, ScopedCapturePtr](FRHICommandListImmediate& RHICmdList)
-				{
-					if (ScopedCapturePtr && (int32)TargetFace == 0)
-					{
-						*ScopedCapturePtr = new RenderCaptureInterface::FScopedCapture(true, &RHICmdList, *FString::Format(TEXT("Scene Capture : {0}"), { EventName }));
-					}
-
-					SceneRenderer->RenderThreadBegin(RHICmdList);
-
-				#if WANTS_DRAW_MESH_EVENTS
-					SCOPED_DRAW_EVENTF(RHICmdList, SceneCaptureMobile, TEXT("SceneCaptureMobile %s"), EventName);
-					FRDGBuilder GraphBuilder(RHICmdList, RDG_EVENT_NAME("SceneCaptureMobile %s", *EventName), ERDGBuilderFlags::Parallel, SceneRenderer->ShaderPlatform);
-				#else
-					SCOPED_DRAW_EVENT(RHICmdList, UpdateSceneCaptureContentMobile_RenderThread);
-					FRDGBuilder GraphBuilder(RHICmdList, RDG_EVENT_NAME("SceneCaptureMobile"), ERDGBuilderFlags::Parallel, SceneRenderer->ShaderPlatform);
-				#endif
-
-#if WITH_EDITOR
-					// Scene renderer may be deleted in UpdateSceneCaptureContent_RenderThread, grab view state pointer first
-					const FSceneViewState* ViewState = SceneRenderer->Views[0].ViewState;
-#endif  // WITH_EDITOR
-
-					VISUALIZE_TEXTURE_BEGIN_VIEW(SceneRenderer->FeatureLevel, SceneRenderer->Views[0].GetViewKey(), *EventName, true);
-
 					// We need to generate mips on last cube face
 					const bool bLastCubeFace = ((int32)TargetFace == (int32)ECubeFace::CubeFace_MAX - 1);
+					const bool bClearRenderTarget = true;
+					const bool bOrthographicCamera = false;
 
 					FRHICopyTextureInfo CopyInfo;
 					CopyInfo.DestSliceIndex = TargetFace;
-					UpdateSceneCaptureContent_RenderThread(GraphBuilder, SceneRenderer, TextureRenderTarget, TextureRenderTarget,
-						MakeArrayView(&CopyInfo, 1), bGenerateMips && bLastCubeFace, GenerateMipsParams, true, false);
-
-					VISUALIZE_TEXTURE_END_VIEW();
+					UpdateSceneCaptureContent_RenderThread(
+						GraphBuilder,
+						Inputs.Renderer,
+						TextureRenderTarget,
+						TextureRenderTarget,
+						MakeArrayView({ CopyInfo }),
+						bGenerateMips && bLastCubeFace,
+						GenerateMipsParams,
+						bClearRenderTarget,
+						bOrthographicCamera);
 
 #if WITH_EDITOR
-					if (ViewState)
+					if (const FSceneViewState* ViewState = Inputs.Renderer->Views[0].ViewState)
 					{
 						const bool bLogSizes = GDumpSceneCaptureMemoryFrame == GFrameNumberRenderThread;
 						if (bLogSizes)
 						{
-							UE_LOG(LogRenderer, Log, TEXT("LogSizes\tSceneCaptureCube[%d]\t%s\t%dx%d"), TargetFace, *EventName, TextureRenderTarget->GetSizeX(), TextureRenderTarget->GetSizeY());
+							UE_LOG(LogRenderer, Log, TEXT("LogSizes\tSceneCaptureCube[%d]\t%s\t%dx%d"), TargetFace, Inputs.FullPath, TextureRenderTarget->GetSizeX(), TextureRenderTarget->GetSizeY());
 							ViewState->GetGPUSizeBytes(bLogSizes);
 						}
 					}
 #endif  // WITH_EDITOR
-
-					GraphBuilder.Execute();
-					SceneRenderer->RenderThreadEnd(RHICmdList);
-
-					if (ScopedCapturePtr && (int32)TargetFace == (int32)ECubeFace::CubeFace_MAX - 1)
-					{
-						// Delete the scope, and the persistent pointer we allocated on the heap to share across the render command lambdas
-						delete *ScopedCapturePtr;
-						delete ScopedCapturePtr;
-					}
 				});
-
-				// Trim the cube face index from the event name
-				EventName.LeftChopInline(4, EAllowShrinking::No);
 			}
 		}
 		else
@@ -1760,11 +1657,11 @@ void FScene::UpdateSceneCaptureContents(USceneCaptureComponentCube* CaptureCompo
 				FTextureRHIRef RenderTargetTexture;
 			};
 
-			FCubeFaceRenderTarget* CubeFaceTarget = new FCubeFaceRenderTarget(TextureRenderTarget);
+			TUniquePtr<FCubeFaceRenderTarget> CubeFaceTarget = MakeUnique<FCubeFaceRenderTarget>(TextureRenderTarget);
 
 			// Copied from CreateSceneRendererForSceneCapture
 			FSceneViewFamilyContext ViewFamily(FSceneViewFamily::ConstructionValues(
-				CubeFaceTarget,
+				CubeFaceTarget.Get(),
 				this,
 				CaptureComponent->ShowFlags)
 				.SetResolveScene(!bCaptureSceneColor)
@@ -1814,83 +1711,57 @@ void FScene::UpdateSceneCaptureContents(USceneCaptureComponentCube* CaptureCompo
 				ViewStateInterface->RemoveLumenSceneData(this);
 			}
 
-			// When bIsMultipleSceneCapture is true, set bIsFirstSceneRenderer to false, which tells the scene renderer it can skip RHI resource flush, saving performance
-			SceneRenderer->bIsFirstSceneRenderer = !bIsMultipleSceneCapture;
-
-			for (const FSceneViewExtensionRef& Extension : SceneRenderer->ViewFamily.ViewExtensions)
+			SceneRenderBuilder.AddRenderer(SceneRenderer,
+				[
+					  CubeFaceTarget = MoveTemp(CubeFaceTarget)
+					, TextureRenderTarget
+					, CaptureSize
+					, bGenerateMips
+					, GenerateMipsParams
+				] (FRDGBuilder& GraphBuilder, const FSceneRenderFunctionInputs& Inputs)
 			{
-				Extension->BeginRenderViewFamily(SceneRenderer->ViewFamily);
-			}
-
-			UE::RenderCommandPipe::FSyncScope SyncScope;
-
-			ENQUEUE_RENDER_COMMAND(CaptureAllCubeFaces)(
-				[SceneRenderer, CubeFaceTarget, TextureRenderTarget, EventName, CaptureSize, bGenerateMips, GenerateMipsParams, bCapturingGPU](FRHICommandListImmediate& RHICmdList)
-			{
-				CubeFaceTarget->InitRHI(RHICmdList);
-
-				SceneRenderer->RenderThreadBegin(RHICmdList);
+				CubeFaceTarget->InitRHI(GraphBuilder.RHICmdList);
 
-			#if WANTS_DRAW_MESH_EVENTS
-				SCOPED_DRAW_EVENTF(RHICmdList, SceneCaptureMobile, TEXT("SceneCaptureMobile %s"), EventName);
-				FRDGBuilder GraphBuilder(RHICmdList, RDG_EVENT_NAME("SceneCaptureMobile %s", *EventName), ERDGBuilderFlags::Parallel, SceneRenderer->ShaderPlatform);
-			#else
-				SCOPED_DRAW_EVENT(RHICmdList, UpdateSceneCaptureContentMobile_RenderThread);
-				FRDGBuilder GraphBuilder(RHICmdList, RDG_EVENT_NAME("SceneCaptureMobile"), ERDGBuilderFlags::Parallel, SceneRenderer->ShaderPlatform);
-			#endif
+				TStaticArray<FRHICopyTextureInfo, (int32)ECubeFace::CubeFace_MAX> CopyInfos;
 
+				for (int32 FaceIndex = 0; FaceIndex < (int32)ECubeFace::CubeFace_MAX; FaceIndex++)
 				{
-					RenderCaptureInterface::FScopedCapture RenderCapture(bCapturingGPU, GraphBuilder, *FString::Format(TEXT("Scene Capture : {0}"), { EventName }));
-
-					TStaticArray<FRHICopyTextureInfo, (int32)ECubeFace::CubeFace_MAX> CopyInfos;
-					for (int32 faceidx = 0; faceidx < (int32)ECubeFace::CubeFace_MAX; faceidx++)
-					{
-						CopyInfos[faceidx].Size.X = CaptureSize.X;
-						CopyInfos[faceidx].Size.Y = CaptureSize.Y;
-						CopyInfos[faceidx].SourcePosition.X = GCubeFaceViewportOffsets[faceidx][0] * CaptureSize.X;
-						CopyInfos[faceidx].SourcePosition.Y = GCubeFaceViewportOffsets[faceidx][1] * CaptureSize.Y;
-						CopyInfos[faceidx].DestSliceIndex = faceidx;
-					}
-
-				#if WITH_EDITOR
-					// Scene renderer may be deleted in UpdateSceneCaptureContent_RenderThread, grab view state pointer first
-					TStaticArray<const FSceneViewState*, (int32)ECubeFace::CubeFace_MAX> SceneViewStates;
-					for (int32 FaceIdx = 0; FaceIdx < (int32)ECubeFace::CubeFace_MAX; FaceIdx++)
-					{
-						SceneViewStates[FaceIdx] = SceneRenderer->Views[FaceIdx].ViewState;
-					}
-				#endif  // WITH_EDITOR
-
-					VISUALIZE_TEXTURE_BEGIN_VIEW(SceneRenderer->FeatureLevel, SceneRenderer->Views[0].GetViewKey(), *EventName, true);
-
-					UpdateSceneCaptureContent_RenderThread(GraphBuilder, SceneRenderer, CubeFaceTarget, TextureRenderTarget,
-						CopyInfos, bGenerateMips, GenerateMipsParams, true, false);
-
-					VISUALIZE_TEXTURE_END_VIEW();
+					CopyInfos[FaceIndex].Size.X = CaptureSize.X;
+					CopyInfos[FaceIndex].Size.Y = CaptureSize.Y;
+					CopyInfos[FaceIndex].SourcePosition.X = GCubeFaceViewportOffsets[FaceIndex][0] * CaptureSize.X;
+					CopyInfos[FaceIndex].SourcePosition.Y = GCubeFaceViewportOffsets[FaceIndex][1] * CaptureSize.Y;
+					CopyInfos[FaceIndex].DestSliceIndex = FaceIndex;
+				}
 
-				#if WITH_EDITOR
-					if (SceneViewStates[0])
+				const bool bClearRenderTarget = true;
+				const bool bOrthographicCamera = false;
+
+				UpdateSceneCaptureContent_RenderThread(
+					GraphBuilder,
+					Inputs.Renderer,
+					CubeFaceTarget.Get(),
+					TextureRenderTarget,
+					CopyInfos,
+					bGenerateMips,
+					GenerateMipsParams,
+					bClearRenderTarget,
+					bOrthographicCamera);
+
+		#if WITH_EDITOR
+				if (Inputs.Renderer->Views[0].ViewState)
+				{
+					const bool bLogSizes = GDumpSceneCaptureMemoryFrame == GFrameNumberRenderThread;
+					if (bLogSizes)
 					{
-						const bool bLogSizes = GDumpSceneCaptureMemoryFrame == GFrameNumberRenderThread;
-						if (bLogSizes)
+						UE_LOG(LogRenderer, Log, TEXT("LogSizes\tSceneCaptureCube\t%s\t%dx%d"), Inputs.FullPath, CubeFaceTarget->GetSizeXY().X, CubeFaceTarget->GetSizeXY().Y);
+						for (int32 FaceIdx = 0; FaceIdx < (int32)ECubeFace::CubeFace_MAX; FaceIdx++)
 						{
-							UE_LOG(LogRenderer, Log, TEXT("LogSizes\tSceneCaptureCube\t%s\t%dx%d"), *EventName, CubeFaceTarget->GetSizeXY().X, CubeFaceTarget->GetSizeXY().Y);
-							for (int32 FaceIdx = 0; FaceIdx < (int32)ECubeFace::CubeFace_MAX; FaceIdx++)
-							{
-								SceneViewStates[FaceIdx]->GetGPUSizeBytes(bLogSizes);
-							}
+							Inputs.Renderer->Views[FaceIdx].ViewState->GetGPUSizeBytes(bLogSizes);
 						}
 					}
-				#endif  // WITH_EDITOR
 				}
-
-				delete CubeFaceTarget;
-
-				GraphBuilder.Execute();
-				SceneRenderer->RenderThreadEnd(RHICmdList);
+		#endif  // WITH_EDITOR
 			});
 		}
 	}
-
-	EndGpuCaptureOrDump(bDumpingGPU);
-}
+}
\ No newline at end of file
diff --git a/Engine/Source/Runtime/Renderer/Private/ScenePrivate.h b/Engine/Source/Runtime/Renderer/Private/ScenePrivate.h
--- a/Engine/Source/Runtime/Renderer/Private/ScenePrivate.h
+++ b/Engine/Source/Runtime/Renderer/Private/ScenePrivate.h
@@ -3413,9 +3413,9 @@ public:
 	virtual void AddPlanarReflection(class UPlanarReflectionComponent* Component) override;
 	virtual void RemovePlanarReflection(UPlanarReflectionComponent* Component) override;
 	virtual void UpdatePlanarReflectionTransform(UPlanarReflectionComponent* Component) override;
-	virtual void UpdateSceneCaptureContents(class USceneCaptureComponent2D* CaptureComponent) override;
-	virtual void UpdateSceneCaptureContents(class USceneCaptureComponentCube* CaptureComponent) override;
-	virtual void UpdatePlanarReflectionContents(UPlanarReflectionComponent* CaptureComponent, FSceneRenderer& MainSceneRenderer) override;
+	virtual void UpdateSceneCaptureContents(class USceneCaptureComponent2D* CaptureComponent, ISceneRenderBuilder& SceneRenderBuilder) override;
+	virtual void UpdateSceneCaptureContents(class USceneCaptureComponentCube* CaptureComponent, ISceneRenderBuilder& SceneRenderBuilder) override;
+	virtual void UpdatePlanarReflectionContents(UPlanarReflectionComponent* CaptureComponent, FSceneRenderer& MainSceneRenderer, ISceneRenderBuilder& SceneRenderBuilder) override;
 	virtual void AllocateReflectionCaptures(const TArray<UReflectionCaptureComponent*>& NewCaptures, const TCHAR* CaptureReason, bool bVerifyOnlyCapturing, bool bCapturingForMobile, bool bInsideTick) override;
 	virtual void ResetReflectionCaptures(bool bOnlyIfOOM) override;
 	virtual void UpdateSkyCaptureContents(const USkyLightComponent* CaptureComponent, bool bCaptureEmissiveOnly, UTextureCube* SourceCubemap, FTexture* OutProcessedTexture, float& OutAverageBrightness, FSHVectorRGB3& OutIrradianceEnvironmentMap, TArray<FFloat16Color>* OutRadianceMap, FLinearColor* SpecifiedCubemapColorScale) override;
diff --git a/Engine/Source/Runtime/Renderer/Private/SceneRendering.cpp b/Engine/Source/Runtime/Renderer/Private/SceneRendering.cpp
--- a/Engine/Source/Runtime/Renderer/Private/SceneRendering.cpp
+++ b/Engine/Source/Runtime/Renderer/Private/SceneRendering.cpp
@@ -110,6 +110,7 @@
 #include "VolumetricCloudProxy.h"
 #include "VT/VirtualTextureScalability.h"
 #include "VT/VirtualTextureSystem.h"
+#include "SceneRenderBuilder.h"
 #include <type_traits>
 #include "BlueNoise.h"
 
@@ -210,6 +211,11 @@ static TAutoConsoleVariable<int32> CVarSplitScreenDebugMGPU(
 	TEXT("For split screen debugging, if MGPU is active, run the second view on the other GPU.  Uses two scene renderers."),
 	ECVF_Default);
 
+static bool IsSplitScreenDebugEnabled(TConstArrayView<const FSceneViewFamily*> ViewFamilies)
+{
+	return CVarSplitScreenDebugEnable.GetValueOnGameThread() > 0 && ViewFamilies.Num() == 1 && ViewFamilies[0]->bSplitScreenDebugAllowed && ViewFamilies[0]->Views.Num() == 1;
+}
+
 #endif  // !UE_BUILD_SHIPPING
 
 bool UseCachedMeshDrawCommands()
@@ -2722,8 +2728,6 @@ FSceneRenderer::FSceneRenderer(const FSceneViewFamily* InViewFamily, FHitProxyCo
 ,	bHasRequestedToggleFreeze(false)
 ,	bUsedPrecomputedVisibility(false)
 ,	bGPUMasksComputed(false)
-,	bIsFirstSceneRenderer(true)		// By default, assume there is just one scene renderer, so it will be the first and last
-,	bIsLastSceneRenderer(true)
 ,	FamilySize(0, 0)
 ,	GPUSceneDynamicContext(CheckPointer(Scene)->GPUScene)
 ,	bShadowDepthRenderCompleted(false)
@@ -4726,7 +4730,7 @@ void FSceneRenderer::SetupMeshPass(FViewInfo& View, FExclusiveDepthStencil::Type
 }
 
 #if !UE_BUILD_SHIPPING
-void FSceneRenderer::CreateSplitScreenDebugViewFamilies(const FSceneViewFamily& InFamily, TArrayView<const FSceneViewFamily*>& OutFamilies, TArray<const FSceneViewFamily*, TFixedAllocator<2>>& OutFamiliesStorage)
+void FSceneRenderer::CreateSplitScreenDebugViewFamilies(const FSceneViewFamily& InFamily, TConstArrayView<const FSceneViewFamily*>& OutFamilies, TArray<const FSceneViewFamily*, TFixedAllocator<2>>& OutFamiliesStorage)
 {
 	// We either generate a single family with 2 views, or two families with 1 view each for MGPU
 	int32 NumFamilies;
@@ -4922,9 +4926,9 @@ void FSceneRenderer::CreateSplitScreenDebugViewFamilies(const FSceneViewFamily&
 	OutFamilies = OutFamiliesStorage;
 }
 
-void FSceneRenderer::DestroySplitScreenDebugViewFamilies(TArrayView<const FSceneViewFamily*> Families)
+void FSceneRenderer::DestroySplitScreenDebugViewFamilies(TConstArrayView<const FSceneViewFamily*> ViewFamilies)
 {
-	for (const FSceneViewFamily* Family : Families)
+	for (const FSceneViewFamily* Family : ViewFamilies)
 	{
 		for (int32 ViewIndex = 0; ViewIndex < Family->Views.Num(); ViewIndex++)
 		{
@@ -4935,7 +4939,7 @@ void FSceneRenderer::DestroySplitScreenDebugViewFamilies(TArrayView<const FScene
 }
 #endif  // !UE_BUILD_SHIPPING
 
-void FSceneRenderer::CreateSceneRenderers(TArrayView<const FSceneViewFamily*> InViewFamilies, FHitProxyConsumer* HitProxyConsumer, bool bAllowSplitScreenDebug, TArray<FSceneRenderer*>& OutSceneRenderers)
+void FSceneRenderer::CreateSceneRenderers(TConstArrayView<const FSceneViewFamily*> InViewFamilies, FHitProxyConsumer* HitProxyConsumer, bool bAllowSplitScreenDebug, TArray<FSceneRenderer*>& OutSceneRenderers)
 {
 	OutSceneRenderers.Empty(InViewFamilies.Num());
 
@@ -4947,7 +4951,7 @@ void FSceneRenderer::CreateSceneRenderers(TArrayView<const FSceneViewFamily*> In
 #if !UE_BUILD_SHIPPING
 	bool bSplitScreenDebug = false;
 	TArray<const FSceneViewFamily*, TFixedAllocator<2>> DebugViewFamilyStorage;
-	if (CVarSplitScreenDebugEnable.GetValueOnGameThread() > 0 && bAllowSplitScreenDebug && InViewFamilies.Num() == 1 && InViewFamilies[0]->bSplitScreenDebugAllowed && InViewFamilies[0]->Views.Num() == 1)
+	if (bAllowSplitScreenDebug && IsSplitScreenDebugEnabled(InViewFamilies))
 	{
 		CreateSplitScreenDebugViewFamilies(*InViewFamilies[0], InViewFamilies, DebugViewFamilyStorage);
 		bSplitScreenDebug = true;
@@ -4977,9 +4981,6 @@ void FSceneRenderer::CreateSceneRenderers(TArrayView<const FSceneViewFamily*> In
 			OutSceneRenderers.Add(new FMobileSceneRenderer(InViewFamily, HitProxyConsumer));
 		}
 
-		OutSceneRenderers.Last()->bIsFirstSceneRenderer = (FamilyIndex == 0);
-		OutSceneRenderers.Last()->bIsLastSceneRenderer = (FamilyIndex == InViewFamilies.Num() - 1);
-
 		for (int ViewExt = 0; ViewExt < InViewFamily->ViewExtensions.Num(); ViewExt++)
 		{
 			InViewFamily->ViewExtensions[ViewExt]->PostCreateSceneRenderer(*InViewFamily, OutSceneRenderers.Last());
@@ -5378,144 +5379,87 @@ void ResetAndShrinkModifiedBounds(TArray<FBox>& Bounds)
 	Bounds.Reset();
 }
 
-/**
- * Helper function performing actual work in render thread.
- *
- * @param SceneRenderers	List of scene renderers to use for rendering.
- */
-static void RenderViewFamilies_RenderThread(FRHICommandListImmediate& RHICmdList, TConstArrayView<FSceneRenderer*> SceneRenderers)
+static void RenderViewFamily_RenderThread(FRDGBuilder& GraphBuilder, FSceneRenderer* Renderer, int32 ViewFamilyIndex)
 {
-	LLM_SCOPE(ELLMTag::SceneRender);
-
-	// All renderers point to the same Scene (calling code asserts this)
-	FScene* const Scene = SceneRenderers[0]->Scene;
-
-	FSceneRenderer::RenderThreadBegin(RHICmdList, SceneRenderers);
+	FSceneViewFamily& ViewFamily = Renderer->ViewFamily;
 
-	bool bAnyShowHitProxies = false;
+	LLM_SCOPE(ELLMTag::SceneRender);
+	SCOPE_CYCLE_COUNTER(STAT_TotalSceneRenderingTime);
+	SCOPED_NAMED_EVENT_TCHAR_CONDITIONAL(*ViewFamily.ProfileDescription, FColor::Red, !ViewFamily.ProfileDescription.IsEmpty());
 
 #if WITH_DEBUG_VIEW_MODES
-	// Flag so we only call FGPUSkinCacheVisualizationData::Update and draw the visualization text once
-	bool bUpdatedGPUSkinCacheVisualization = false;
-#endif
-
-#if WITH_STATE_STREAM
-	FStateStreamManagerImpl* StateStreamManager = static_cast<FStateStreamManagerImpl*>(Scene->GetWorld()->GetStateStreamManager());
-	StateStreamManager->Render_Update(Scene->GetWorld()->RealTimeSeconds);
-	StateStreamManager->Render_GarbageCollect(true);
-#endif
-
-	for (FSceneRenderer* SceneRenderer : SceneRenderers)
+	if (AllowDebugViewShaderMode(DVSM_VisualizeGPUSkinCache, ViewFamily.GetShaderPlatform(), ViewFamily.GetFeatureLevel()))
 	{
-		const ERHIFeatureLevel::Type FeatureLevel = SceneRenderer->FeatureLevel;
-		
-		FSceneViewFamily& ViewFamily = SceneRenderer->ViewFamily;
-
-		FRDGBuilder GraphBuilder(
-			RHICmdList,
-			RDG_EVENT_NAME("SceneRenderer_%s(ViewFamily=%s)",
-				ViewFamily.EngineShowFlags.HitProxies ? TEXT("RenderHitProxies") : TEXT("Render"),
-				ViewFamily.bResolveScene ? TEXT("Primary") : TEXT("Auxiliary")
-			),
-			ERDGBuilderFlags::Parallel
-		);
+		FViewInfo& View = Renderer->Views[0];
+		FGPUSkinCacheVisualizationData& VisualizationData = GetGPUSkinCacheVisualizationData();
 
-#if WITH_GPUDEBUGCRASH
-		if (GRHIGlobals.TriggerGPUCrash != ERequestedGPUCrash::None) 
+		// Only run visualization update once, but set debug flags for all view families if the mode is active
+		// Note VisualizationData.Update needs to be called per frame, as || lazy evaluation is used, so need to do it before evaluating VisualizeGPUSkinCache flag
+		if (ViewFamilyIndex > 0 || VisualizationData.Update(View.CurrentGPUSkinCacheVisualizationMode) || ViewFamily.EngineShowFlags.VisualizeGPUSkinCache)
 		{
-			ScheduleGPUDebugCrash(GraphBuilder);
-		}
-#endif
-		SCOPE_CYCLE_COUNTER(STAT_TotalSceneRenderingTime);
-		SCOPED_NAMED_EVENT_TCHAR_CONDITIONAL(*ViewFamily.ProfileDescription, FColor::Red, !ViewFamily.ProfileDescription.IsEmpty());
-		const uint64 FamilyRenderStart = FPlatformTime::Cycles64();
-
-#if WITH_DEBUG_VIEW_MODES
-		const bool bAllowGPUSkinCacheVisualization = AllowDebugViewShaderMode(DVSM_VisualizeGPUSkinCache, ViewFamily.GetShaderPlatform(), ViewFamily.GetFeatureLevel());
-		if (bAllowGPUSkinCacheVisualization && SceneRenderer->Views.Num() > 0)
-		{
-			FViewInfo& View = SceneRenderer->Views[0];
-
-			FGPUSkinCacheVisualizationData& VisualizationData = GetGPUSkinCacheVisualizationData();
+			// When activating visualization from the command line, enable VisualizeGPUSkinCache.
+			ViewFamily.EngineShowFlags.SetVisualizeGPUSkinCache(true);
+			ViewFamily.DebugViewShaderMode = DVSM_VisualizeGPUSkinCache;
 
-			// Only run visualization update once, but set debug flags for all view families if the mode is active
-			// Note VisualizationData.Update needs to be called per frame, as || lazy evaluation is used, so need to do it before evaluating VisualizeGPUSkinCache flag
-			if (bUpdatedGPUSkinCacheVisualization || VisualizationData.Update(View.CurrentGPUSkinCacheVisualizationMode) || ViewFamily.EngineShowFlags.VisualizeGPUSkinCache)
+			// Only draw the visualization info text once.
+			if (ViewFamilyIndex == 0)
 			{
-				// When activating visualization from the command line, enable VisualizeGPUSkinCache.
-				ViewFamily.EngineShowFlags.SetVisualizeGPUSkinCache(true);
-				ViewFamily.DebugViewShaderMode = DVSM_VisualizeGPUSkinCache;
-
-				// Only draw the visualization info text once.
-				if (!bUpdatedGPUSkinCacheVisualization)
-				{
-					SceneRenderer->DrawGPUSkinCacheVisualizationInfoText();
-					bUpdatedGPUSkinCacheVisualization = true;
-				}
+				Renderer->DrawGPUSkinCacheVisualizationInfoText();
 			}
 		}
+	}
 #endif  // WITH_DEBUG_VIEW_MODES
 
 #if WITH_MGPU
-		if (ViewFamily.bForceCopyCrossGPU)
-		{
-			GraphBuilder.EnableForceCopyCrossGPU();
-		}
+	if (ViewFamily.bForceCopyCrossGPU)
+	{
+		GraphBuilder.EnableForceCopyCrossGPU();
+	}
 #endif
 
-		if (ViewFamily.EngineShowFlags.HitProxies)
-		{
-			// Render the scene's hit proxies.
-			SceneRenderer->RenderHitProxies(GraphBuilder);
-			bAnyShowHitProxies = true;
-		}
-		else
-		{
-			// We don't want to run visualize texture for hit proxies, so do this here.  Note that the reason we run this here
-			// rather than inside the scene renderer is to customize the user facing description string for the visualization in
-			// other code paths (scene captures), in addition to skipping it for views where it's not useful (hit proxies and
-			// utility scene renderers like water).
-			VISUALIZE_TEXTURE_BEGIN_VIEW(SceneRenderer->FeatureLevel, SceneRenderer->Views[0].GetViewKey(),
-				ViewFamily.ProfileDescription.IsEmpty() ? (ViewFamily.bResolveScene ? TEXT("ScenePrimary") : TEXT("SceneAuxiliary")) : *ViewFamily.ProfileDescription, false);
-
-			// Render the scene.
-			SceneRenderer->Render(GraphBuilder);
-
-			VISUALIZE_TEXTURE_END_VIEW();
-		}
-
-		SceneRenderer->FlushCrossGPUFences(GraphBuilder);
+	if (ViewFamily.EngineShowFlags.HitProxies)
+	{
+		Renderer->RenderHitProxies(GraphBuilder);
+	}
+	else
+	{
+		Renderer->Render(GraphBuilder);
+	}
 
-		GraphBuilder.Execute();
+	Renderer->FlushCrossGPUFences(GraphBuilder);
 
 #if MESH_DRAW_COMMAND_STATS
-		// Needs to happen after graph builder execute to make sure all FMeshDrawCommandPassSetupTask are finished.
-		// Only during iteration of the visible MDCs we know the custom indirect args used in the frame
-		if (FMeshDrawCommandStatsManager* Instance = FMeshDrawCommandStatsManager::Get())
+	// Needs to happen after graph builder execute to make sure all FMeshDrawCommandPassSetupTask are finished.
+	// Only during iteration of the visible MDCs we know the custom indirect args used in the frame
+	if (FMeshDrawCommandStatsManager* Instance = FMeshDrawCommandStatsManager::Get())
+	{
+		GraphBuilder.AddPostExecuteCallback([Instance, &RHICmdList = GraphBuilder.RHICmdList]
 		{
 			Instance->QueueCustomDrawIndirectArgsReadback(RHICmdList);
-		}
+		});
+	}
 #endif
 
-		// Copy relevant data from ViewInfo to ViewState->PrevFrameViewInfo
-		for (FViewInfo& View : SceneRenderer->Views)
-		{
-			if (View.ViewState)
-			{
-				View.ViewState->PrevFrameViewInfo.bUsesGlobalDistanceField = View.bUsesGlobalDistanceField;
-			}
-		}
-
-		if (SceneRenderer->ViewFamily.ProfileSceneRenderTime)
+	// Copy relevant data from ViewInfo to ViewState->PrevFrameViewInfo
+	for (FViewInfo& View : Renderer->Views)
+	{
+		if (View.ViewState)
 		{
-			*SceneRenderer->ViewFamily.ProfileSceneRenderTime = (float)FPlatformTime::ToSeconds64(FPlatformTime::Cycles64() - FamilyRenderStart);
+			View.ViewState->PrevFrameViewInfo.bUsesGlobalDistanceField = View.bUsesGlobalDistanceField;
 		}
 	}
+}
+
+static void CleanupViewFamilies_RenderThread(TConstArrayView<FSceneRenderer*> SceneRenderers)
+{
+	LLM_SCOPE(ELLMTag::SceneRender);
+
+	FScene* Scene = SceneRenderers[0]->Scene;
 
 	{
 		CSV_SCOPED_TIMING_STAT_EXCLUSIVE(PostRenderCleanUp);
 
-		if (IsHairStrandsEnabled(EHairStrandsShaderType::All, Scene->GetShaderPlatform()) && (SceneRenderers[0]->AllFamilyViews.Num() > 0) && !bAnyShowHitProxies)
+		if (IsHairStrandsEnabled(EHairStrandsShaderType::All, Scene->GetShaderPlatform()) && (SceneRenderers[0]->AllFamilyViews.Num() > 0) )
 		{
 			FHairStrandsBookmarkParameters Parameters;
 			CreateHairStrandsBookmarkParameters(Scene, SceneRenderers[0]->Views, SceneRenderers[0]->AllFamilyViews, Parameters, false /*bComputeVisibleInstances*/);
@@ -5580,8 +5524,6 @@ static void RenderViewFamilies_RenderThread(FRHICommandListImmediate& RHICmdList
 		SET_CYCLE_COUNTER(STAT_TotalGPUFrameTime, RHIGetGPUFrameCycles());
 	}
 #endif
-
-	FSceneRenderer::RenderThreadEnd(RHICmdList, SceneRenderers);
 }
 
 void OnChangeCVarRequiringRecreateRenderState(IConsoleVariable* Var)
@@ -5644,6 +5586,14 @@ void FRendererModule::BeginRenderingViewFamilies(FCanvas* Canvas, TConstArrayVie
 			UE::Stats::FThreadIdleStats::EndCriticalPath();
 
 			GetNaniteVisualizationData().Pick(World);
+
+		#if WITH_STATE_STREAM
+			ENQUEUE_RENDER_COMMAND(UpdateStateStream)([StateStreamManager = static_cast<FStateStreamManagerImpl*>(World->GetStateStreamManager()), RealTimeSeconds = World->RealTimeSeconds] (FRHICommandListImmediate&)
+			{
+				StateStreamManager->Render_Update(RealTimeSeconds);
+				StateStreamManager->Render_GarbageCollect(true);
+			});
+		#endif
 		}
 	}
 
@@ -5715,43 +5665,42 @@ void FRendererModule::BeginRenderingViewFamilies(FCanvas* Canvas, TConstArrayVie
 		check(ViewFamily->GetTemporalUpscalerInterface() == nullptr);
 		check(ViewFamily->GetPrimarySpatialUpscalerInterface() == nullptr);
 		check(ViewFamily->GetSecondarySpatialUpscalerInterface() == nullptr);
-
-		for (int ViewExt = 0; ViewExt < ViewFamily->ViewExtensions.Num(); ViewExt++)
-		{
-			ViewFamily->ViewExtensions[ViewExt]->BeginRenderViewFamily(*ViewFamily);
-		}
-
 		checkf(!(ViewFamily->GetTemporalUpscalerInterface() != nullptr && ViewFamily->GetPrimarySpatialUpscalerInterface() != nullptr),
 			TEXT("Conflict setting up a third party primary spatial upscaler or temporal upscaler."));
 	}
 
 	if (Scene)
-	{		
+	{
 		// Set the world's "needs full lighting rebuild" flag if the scene has any uncached static lighting interactions.
-		if(World)
+		if (World)
 		{
 			// Note: reading NumUncachedStaticLightingInteractions on the game thread here which is written to by the rendering thread
 			// This is reliable because the RT uses interlocked mechanisms to update it
 			World->SetMapNeedsLightingFullyRebuilt(Scene->NumUncachedStaticLightingInteractions, Scene->NumUnbuiltReflectionCaptures);
 		}
 
+	#if CSV_PROFILER && !CSV_PROFILER_MINIMAL
+		ENQUEUE_RENDER_COMMAND(SetDrawSceneCommand_StartDelay)([DrawSceneEnqueue = FPlatformTime::Cycles64()] (FRHICommandListImmediate&)
+		{
+			const uint64 SceneRenderStart = FPlatformTime::Cycles64();
+			const float StartDelayMillisec = FPlatformTime::ToMilliseconds64(SceneRenderStart - DrawSceneEnqueue);
+			CSV_CUSTOM_STAT_GLOBAL(DrawSceneCommand_StartDelay, StartDelayMillisec, ECsvCustomStatOp::Set);
+		});
+	#endif
+
 		// Construct the scene renderers.  This copies the view family attributes into its own structures.
+		const bool bAllowSplitScreenDebug = true;
 		TArray<FSceneRenderer*> SceneRenderers;
+		FSceneRenderer::CreateSceneRenderers(ViewFamilies, Canvas->GetHitProxyConsumer(), bAllowSplitScreenDebug, SceneRenderers);
 
-		TArray<const FSceneViewFamily*> ViewFamiliesConst;
-		for (FSceneViewFamily* ViewFamily : ViewFamilies)
-		{
-			ViewFamiliesConst.Add(ViewFamily);
-		}
+		FSceneRenderBuilder SceneRenderBuilder(Scene);
 
 		bool bShowHitProxies = (Canvas->GetHitProxyConsumer() != nullptr);
 		if (!bShowHitProxies)
 		{
-			SceneCaptureUpdateDeferredCapturesInternal(Scene, ViewFamiliesConst);
+			SceneCaptureUpdateDeferredCapturesInternal(Scene, ViewFamilies, SceneRenderBuilder);
 		}
 
-		FSceneRenderer::CreateSceneRenderers(ViewFamiliesConst, Canvas->GetHitProxyConsumer(), /*bAllowSplitScreenDebug=*/ true, SceneRenderers);
-
 		if (!bShowHitProxies)
 		{
 			for (int32 ReflectionIndex = 0; ReflectionIndex < Scene->PlanarReflections_GameThread.Num(); ReflectionIndex++)
@@ -5763,28 +5712,31 @@ void FRendererModule::BeginRenderingViewFamilies(FCanvas* Canvas, TConstArrayVie
 					{
 						continue;
 					}
-					Scene->UpdatePlanarReflectionContents(ReflectionComponent, *SceneRenderer);
+					Scene->UpdatePlanarReflectionContents(ReflectionComponent, *SceneRenderer, SceneRenderBuilder);
 				}
 			}
 		}
 
-		for (FSceneRenderer* SceneRenderer : SceneRenderers)
+		FSceneRenderer::PreallocateCrossGPUFences(SceneRenderers);
+
+		for (int32 ViewFamilyIndex = 0; ViewFamilyIndex < ViewFamilies.Num(); ++ViewFamilyIndex)
 		{
+			FSceneRenderer* SceneRenderer = SceneRenderers[ViewFamilyIndex];
 			SceneRenderer->ViewFamily.DisplayInternalsData.Setup(World);
-		}
 
-		FSceneRenderer::PreallocateCrossGPUFences(SceneRenderers);
-
-		const uint64 DrawSceneEnqueue = FPlatformTime::Cycles64();
-		ENQUEUE_RENDER_COMMAND(FDrawSceneCommand)(
-			[LocalSceneRenderers = CopyTemp(SceneRenderers), DrawSceneEnqueue](FRHICommandListImmediate& RHICmdList)
+			SceneRenderBuilder.AddRenderer(SceneRenderer, bShowHitProxies ? TEXT("HitProxies") : TEXT("ViewFamilies"),
+				[ViewFamilyIndex] (FRDGBuilder& GraphBuilder, const FSceneRenderFunctionInputs& Inputs)
 			{
-				uint64 SceneRenderStart = FPlatformTime::Cycles64();
-				const float StartDelayMillisec = FPlatformTime::ToMilliseconds64(SceneRenderStart - DrawSceneEnqueue);
-				CSV_CUSTOM_STAT_GLOBAL(DrawSceneCommand_StartDelay, StartDelayMillisec, ECsvCustomStatOp::Set);
-				RenderViewFamilies_RenderThread(RHICmdList, LocalSceneRenderers);
-				FlushPendingDeleteRHIResources_RenderThread();
+				RenderViewFamily_RenderThread(GraphBuilder, Inputs.Renderer, ViewFamilyIndex);
 			});
+		}
+
+		SceneRenderBuilder.AddRenderCommand([SceneRenderers = MoveTemp(SceneRenderers)](FRHICommandListImmediate&)
+		{
+			CleanupViewFamilies_RenderThread(SceneRenderers);
+		});
+
+		SceneRenderBuilder.Execute();
 
 		// Force kick the RT if we've got RT polling on.
 		// This saves us having to wait until the polling period before the scene draw starts executing.
@@ -6876,4 +6828,4 @@ bool DoMaterialAndPrimitiveModifyMeshPosition(const FMaterial& Material, const F
 	// but even without that custom logic, optimized position-only shaders for depth-only rendering do not support first person, so by checking the bIsFirstPerson flag here, we prevent first person primitives
 	// from being rendered with such incompatible shaders.
 	return (bMaterialModifiesMeshPosition && bPrimitiveAllowsWPOEvaluation) || bIsFirstPerson;
-}
+}
\ No newline at end of file
diff --git a/Engine/Source/Runtime/Renderer/Private/SceneRendering.h b/Engine/Source/Runtime/Renderer/Private/SceneRendering.h
--- a/Engine/Source/Runtime/Renderer/Private/SceneRendering.h
+++ b/Engine/Source/Runtime/Renderer/Private/SceneRendering.h
@@ -2093,10 +2093,6 @@ public:
 	bool bGPUMasksComputed;
 	FRHIGPUMask RenderTargetGPUMask;
 
-	/** Whether the given scene renderer is the first or last in a group being rendered. */
-	bool bIsFirstSceneRenderer;
-	bool bIsLastSceneRenderer;
-
 public:
 
 	FSceneRenderer(const FSceneViewFamily* InViewFamily, FHitProxyConsumer* HitProxyConsumer);
@@ -2124,7 +2120,7 @@ public:
 	RENDERER_API static FSceneRenderer* CreateSceneRenderer(const FSceneViewFamily* InViewFamily, FHitProxyConsumer* HitProxyConsumer);
 
 	/** Creates multiple scene renderers based on the current feature level.  All view families must point to the same Scene. */
-	RENDERER_API static void CreateSceneRenderers(TArrayView<const FSceneViewFamily*> InViewFamilies, FHitProxyConsumer* HitProxyConsumer, bool bAllowSplitScreenDebug, TArray<FSceneRenderer*>& OutSceneRenderers);
+	RENDERER_API static void CreateSceneRenderers(TConstArrayView<const FSceneViewFamily*> InViewFamilies, FHitProxyConsumer* HitProxyConsumer, bool bAllowSplitScreenDebug, TArray<FSceneRenderer*>& OutSceneRenderers);
 
 	/** Setups FViewInfo::ViewRect according to ViewFamilly's ScreenPercentageInterface. */
 	void PrepareViewRectsForRendering(FRHICommandListImmediate& RHICmdList);
@@ -2555,8 +2551,8 @@ protected:
 	void RenderMeshDistanceFieldVisualization(FRDGBuilder& GraphBuilder, const FMinimalSceneTextures& SceneTextures);
 
 #if !UE_BUILD_SHIPPING
-	static void CreateSplitScreenDebugViewFamilies(const FSceneViewFamily& InFamily, TArrayView<const FSceneViewFamily*>& OutFamilies, TArray<const FSceneViewFamily*, TFixedAllocator<2>>& OutFamilyStorage);
-	static void DestroySplitScreenDebugViewFamilies(TArrayView<const FSceneViewFamily*> Families);
+	static void CreateSplitScreenDebugViewFamilies(const FSceneViewFamily& InFamily, TConstArrayView<const FSceneViewFamily*>& OutFamilies, TArray<const FSceneViewFamily*, TFixedAllocator<2>>& OutFamiliesStorage);
+	static void DestroySplitScreenDebugViewFamilies(TConstArrayView<const FSceneViewFamily*> ViewFamilies);
 #endif
 
 protected:
