Disable occlusion queries while on the bus on all platforms that use mobile renderer to workaround occlusion latency artifacts

diff --git a/Engine/Source/Runtime/Renderer/Private/MobileShadingRenderer.cpp b/Engine/Source/Runtime/Renderer/Private/MobileShadingRenderer.cpp
--- a/Engine/Source/Runtime/Renderer/Private/MobileShadingRenderer.cpp
+++ b/Engine/Source/Runtime/Renderer/Private/MobileShadingRenderer.cpp
@@ -877,9 +877,12 @@ void FMobileSceneRenderer::RenderFullDepthPrepass(FRDGBuilder& GraphBuilder, TAr
 				});
 		}
 	}
-
-	EndOcclusionScope(GraphBuilder, Views);
-	FenceOcclusionTests(GraphBuilder);
+	
+	if (DoOcclusionQueries())
+	{
+		EndOcclusionScope(GraphBuilder, Views);
+		FenceOcclusionTests(GraphBuilder);
+	}
 }
 
 void FMobileSceneRenderer::RenderMaskedPrePass(FRHICommandList& RHICmdList, const FViewInfo& View, const FInstanceCullingDrawParams* DepthPassInstanceCullingDrawParams)
@@ -1118,7 +1121,7 @@ void FMobileSceneRenderer::Render(FRDGBuilder& GraphBuilder)
 	// Find the visible primitives and prepare targets and buffers for rendering
 	InitViews(GraphBuilder, SceneTexturesConfig, InstanceCullingManager, VirtualTextureUpdater.Get(), InitViewTaskDatas);
 
-	if (bRendererOutputFinalSceneColor)
+	if (bRendererOutputFinalSceneColor && DoOcclusionQueries())
 	{
 		BeginOcclusionScope(GraphBuilder, Views);
 	}
@@ -1420,7 +1423,7 @@ void FMobileSceneRenderer::Render(FRDGBuilder& GraphBuilder)
 			RenderForward(GraphBuilder, ViewFamilyTexture, SceneTextures, DBufferTextures);
 		}
 
-		if (!bIsFullDepthPrepassEnabled)
+		if (!bIsFullDepthPrepassEnabled && DoOcclusionQueries())
 		{
 			EndOcclusionScope(GraphBuilder, Views);
 			FenceOcclusionTests(GraphBuilder);
