Removed empty scene extension implementations.

diff --git a/Engine/Plugins/Experimental/CineCameraSceneCapture/Source/CineCameraSceneCapture/Private/CineCameraSceneCaptureComponent.cpp b/Engine/Plugins/Experimental/CineCameraSceneCapture/Source/CineCameraSceneCapture/Private/CineCameraSceneCaptureComponent.cpp
--- a/Engine/Plugins/Experimental/CineCameraSceneCapture/Source/CineCameraSceneCapture/Private/CineCameraSceneCaptureComponent.cpp
+++ b/Engine/Plugins/Experimental/CineCameraSceneCapture/Source/CineCameraSceneCapture/Private/CineCameraSceneCaptureComponent.cpp
@@ -88,7 +88,6 @@ public:
 		}
 	}
 
-	virtual void BeginRenderViewFamily(FSceneViewFamily& InViewFamily) override {};
 	//~ End FSceneViewExtensionBase Interface
 
 	/** 
diff --git a/Engine/Source/Runtime/Renderer/Private/MeshEdgesRendering.cpp b/Engine/Source/Runtime/Renderer/Private/MeshEdgesRendering.cpp
--- a/Engine/Source/Runtime/Renderer/Private/MeshEdgesRendering.cpp
+++ b/Engine/Source/Runtime/Renderer/Private/MeshEdgesRendering.cpp
@@ -213,9 +213,6 @@ public:
 	: FSceneViewExtensionBase( AutoRegister )
 	{}
 
-	virtual void SetupViewFamily(FSceneViewFamily& InViewFamily) override {}
-	virtual void SetupView(FSceneViewFamily& InViewFamily, FSceneView& InView) override {}
-	virtual void BeginRenderViewFamily(FSceneViewFamily& InViewFamily) override {}
 	virtual void PostCreateSceneRenderer(const FSceneViewFamily& InViewFamily, ISceneRenderer* Renderer) override
 	{
 		RenderMeshEdges(static_cast<FSceneRenderer*>(Renderer)->ViewFamily);
