Fixed split screen multi-GPU mode, and also made it work for single-GPU as well. Also refactored GPU skin cache visualization out of the main render path.

diff --git a/Engine/Source/Runtime/Renderer/Private/SceneRendering.cpp b/Engine/Source/Runtime/Renderer/Private/SceneRendering.cpp
--- a/Engine/Source/Runtime/Renderer/Private/SceneRendering.cpp
+++ b/Engine/Source/Runtime/Renderer/Private/SceneRendering.cpp
@@ -205,10 +205,10 @@ static TAutoConsoleVariable<int32> CVarSplitScreenDebugLumenScene(
 	TEXT("For split screen debugging, allocate a separate Lumen scene for the second view."),
 	ECVF_Default);
 
-static TAutoConsoleVariable<int32> CVarSplitScreenDebugMGPU(
-	TEXT("r.SplitScreenDebug.MGPU"),
-	1,
-	TEXT("For split screen debugging, if MGPU is active, run the second view on the other GPU.  Uses two scene renderers."),
+static TAutoConsoleVariable<int32> CVarSplitScreenDebugMultiViewFamily(
+	TEXT("r.SplitScreenDebug.MultiViewFamily"),
+	0,
+	TEXT("Uses two renderers with one view each rather than two views."),
 	ECVF_Default);
 
 static bool IsSplitScreenDebugEnabled(TConstArrayView<const FSceneViewFamily*> ViewFamilies)
@@ -4096,6 +4096,35 @@ IVisibilityTaskData* FSceneRenderer::OnRenderBegin(FRDGBuilder& GraphBuilder)
 	return VisibilityTaskData;
 }
 
+static void SetupDebugViewModes(TConstArrayView<FSceneRenderer*> Renderers)
+{
+#if WITH_DEBUG_VIEW_MODES
+	check(!Renderers.IsEmpty());
+	FScene* Scene = Renderers[0]->Scene;
+
+	if (AllowDebugViewShaderMode(DVSM_VisualizeGPUSkinCache, Scene->GetShaderPlatform(), Scene->GetFeatureLevel()))
+	{
+		bool UpdatedGPUSkinCacheVisualization = false;
+
+		for (FSceneRenderer* Renderer : Renderers)
+		{
+			FViewInfo& View = Renderer->Views[0];
+			FGPUSkinCacheVisualizationData& VisualizationData = GetGPUSkinCacheVisualizationData();
+
+			// Only run visualization update once, but set debug flags for all view families if the mode is active
+			// Note VisualizationData.Update needs to be called per frame, as || lazy evaluation is used, so need to do it before evaluating VisualizeGPUSkinCache flag
+			if (UpdatedGPUSkinCacheVisualization || VisualizationData.Update(View.CurrentGPUSkinCacheVisualizationMode) || Renderer->ViewFamily.EngineShowFlags.VisualizeGPUSkinCache)
+			{
+				// When activating visualization from the command line, enable VisualizeGPUSkinCache.
+				Renderer->ViewFamily.EngineShowFlags.SetVisualizeGPUSkinCache(true);
+				Renderer->ViewFamily.DebugViewShaderMode = DVSM_VisualizeGPUSkinCache;
+				UpdatedGPUSkinCacheVisualization = true;
+			}
+		}
+	}
+#endif  // WITH_DEBUG_VIEW_MODES
+}
+
 /** 
 * Finishes the view family rendering.
 */
@@ -4138,8 +4167,9 @@ void FSceneRenderer::OnRenderFinish(FRDGBuilder& GraphBuilder, FRDGTextureRef Vi
 		
 		static auto* CVarSkinCacheOOM = IConsoleManager::Get().FindTConsoleVariableDataFloat(TEXT("r.SkinCache.SceneMemoryLimitInMB"));
 
+		FGPUSkinCache* SkinCache = Scene->GetGPUSkinCache();
 		uint64 GPUSkinCacheExtraRequiredMemory = 0;
-		if (FGPUSkinCache* SkinCache = Scene->GetGPUSkinCache())
+		if (SkinCache)
 		{
 			GPUSkinCacheExtraRequiredMemory = SkinCache->GetExtraRequiredMemoryAndReset();
 		}
@@ -4274,6 +4304,7 @@ void FSceneRenderer::OnRenderFinish(FRDGBuilder& GraphBuilder, FRDGTextureRef Vi
 				bool bIsFrozen = ViewState && (ViewState->bIsFrozen);
 				bool bLocked = View.bIsLocked;
 				const bool bStereoView = IStereoRendering::IsStereoEyeView(View);
+				const FName GPUSkinCacheVisualizationMode = View.CurrentGPUSkinCacheVisualizationMode;
 
 				// display a warning if an ambient cubemap uses non-angular mipmap filtering
 				bool bShowAmbientCubemapMipGenSettingsWarning = false;
@@ -4291,7 +4322,7 @@ void FSceneRenderer::OnRenderFinish(FRDGBuilder& GraphBuilder, FRDGTextureRef Vi
 					}
 				}
 #endif
-				if ((GAreScreenMessagesEnabled && !GEngine->bSuppressMapWarnings) && (bIsFrozen || bLocked || bStereoView || bShowAmbientCubemapMipGenSettingsWarning || bAnyWarning))
+				if ((GAreScreenMessagesEnabled && !GEngine->bSuppressMapWarnings) && (bIsFrozen || bLocked || bStereoView || bShowAmbientCubemapMipGenSettingsWarning || bAnyWarning || GPUSkinCacheVisualizationMode != NAME_None))
 				{
 					RDG_EVENT_SCOPE_CONDITIONAL(GraphBuilder, Views.Num() > 1, "View%d", ViewIndex);
 
@@ -4301,7 +4332,7 @@ void FSceneRenderer::OnRenderFinish(FRDGBuilder& GraphBuilder, FRDGTextureRef Vi
 					const bool bIsMobileMultiViewEnabled = View.bIsMobileMultiViewEnabled;
 
 					AddDrawCanvasPass(GraphBuilder, {}, View, Output,
-						[this, ViewState, GPUSkinCacheExtraRequiredMemory,
+						[this, SkinCache, GPUSkinCacheVisualizationMode, ViewState, GPUSkinCacheExtraRequiredMemory,
 						bLocked, bShowPrecomputedVisibilityWarning, bShowDemotedLocalMemoryWarning, bShowGlobalClipPlaneWarning, bShowDFAODisabledWarning, bShowDFDisabledWarning,
 						bIsFrozen, bShowSkylightWarning, bShowPointLightWarning, bShowShadowedLightOverflowWarning,
 						bShowMobileLowQualityLightmapWarning, bShowMobileDynamicCSMWarning, bMobileMissingSkyMaterial, 
@@ -4398,6 +4429,12 @@ void FSceneRenderer::OnRenderFinish(FRDGBuilder& GraphBuilder, FRDGTextureRef Vi
 							FString String = FString::Printf(TEXT("OUT OF MEMORY FOR SKIN CACHE, REQUIRES %.3f extra MB (currently at %.3f)"), (float)GPUSkinCacheExtraRequiredMemory / 1048576.0f, CVarSkinCacheOOM->GetValueOnAnyThread());
 							Writer.DrawLine(FText::FromString(String));
 						}
+
+						if (GPUSkinCacheVisualizationMode.IsValid() && SkinCache)
+						{
+							SkinCache->DrawVisualizationInfoText(GPUSkinCacheVisualizationMode, Writer);
+						}
+
 						if (bShowLocalExposureDisabledWarning)
 						{
 							static const FText Message = NSLOCTEXT("Renderer", "LocalExposureDisabled", "Local Exposure is disabled.");
@@ -4634,18 +4671,6 @@ void FSceneRenderer::OnRenderFinish(FRDGBuilder& GraphBuilder, FRDGTextureRef Vi
 	});
 }
 
-void FSceneRenderer::DrawGPUSkinCacheVisualizationInfoText()
-{
-	FGPUSkinCache* SkinCache = Scene ? Scene->GetGPUSkinCache() : nullptr;
-	if (SkinCache && Views.Num() > 0)
-	{
-		const FName& VisualizationMode = Views[0].CurrentGPUSkinCacheVisualizationMode;
-		OnGetOnScreenMessages.AddLambda([SkinCache, &VisualizationMode](FScreenMessageWriter& ScreenMessageWriter)->void
-		{
-			SkinCache->DrawVisualizationInfoText(VisualizationMode, ScreenMessageWriter);
-		});
-	}
-}
 
 void FSceneRenderer::SetupMeshPass(FViewInfo& View, FExclusiveDepthStencil::Type BasePassDepthStencilAccess, FViewCommands& ViewCommands, FInstanceCullingManager& InstanceCullingManager)
 {
@@ -4730,12 +4755,12 @@ void FSceneRenderer::SetupMeshPass(FViewInfo& View, FExclusiveDepthStencil::Type
 }
 
 #if !UE_BUILD_SHIPPING
-void FSceneRenderer::CreateSplitScreenDebugViewFamilies(const FSceneViewFamily& InFamily, TConstArrayView<const FSceneViewFamily*>& OutFamilies, TArray<const FSceneViewFamily*, TFixedAllocator<2>>& OutFamiliesStorage)
+void FSceneRenderer::CreateSplitScreenDebugViewFamilies(const FSceneViewFamily& InFamily, TArray<const FSceneViewFamily*>& OutFamilies, TArray<const FSceneViewFamily*, TFixedAllocator<2>>& OutFamiliesStorage)
 {
 	// We either generate a single family with 2 views, or two families with 1 view each for MGPU
 	int32 NumFamilies;
 	int32 NumViewsPerFamily;
-	if (GNumExplicitGPUsForRendering > 1 && CVarSplitScreenDebugMGPU.GetValueOnGameThread())
+	if (CVarSplitScreenDebugMultiViewFamily.GetValueOnGameThread())
 	{
 		NumFamilies = 2;
 		NumViewsPerFamily = 1;
@@ -4908,13 +4933,18 @@ void FSceneRenderer::CreateSplitScreenDebugViewFamilies(const FSceneViewFamily&
 		{
 			FSceneViewFamily* Family = ViewParents[ViewIndex];
 
-			// Prevent the render target from being cleared, and enable cross GPU transfers
+			// Prevent the render target from being cleared
 			Family->bAdditionalViewFamily = true;
-			Family->bMultiGPUForkAndJoin = true;
 
-			// Set the view to run on the second GPU
-			View->bOverrideGPUMask = true;
-			View->GPUMask = FRHIGPUMask::FromIndex(ViewIndex);
+			if (GNumExplicitGPUsForRendering > 1)
+			{
+				// Enable cross GPU transfers
+				Family->bMultiGPUForkAndJoin = true;
+
+				// Set the view to run on the second GPU
+				View->bOverrideGPUMask = true;
+				View->GPUMask = FRHIGPUMask::FromIndex(ViewIndex);
+			}
 		}
 	}
 
@@ -4939,11 +4969,11 @@ void FSceneRenderer::DestroySplitScreenDebugViewFamilies(TConstArrayView<const F
 }
 #endif  // !UE_BUILD_SHIPPING
 
-void FSceneRenderer::CreateSceneRenderers(TConstArrayView<const FSceneViewFamily*> InViewFamilies, FHitProxyConsumer* HitProxyConsumer, bool bAllowSplitScreenDebug, TArray<FSceneRenderer*>& OutSceneRenderers)
+void FSceneRenderer::CreateSceneRenderers(TArray<const FSceneViewFamily*>& InOutViewFamilies, FHitProxyConsumer* HitProxyConsumer, bool bAllowSplitScreenDebug, TArray<FSceneRenderer*>& OutSceneRenderers)
 {
-	OutSceneRenderers.Empty(InViewFamilies.Num());
+	OutSceneRenderers.Empty(InOutViewFamilies.Num());
 
-	if (!InViewFamilies.Num())
+	if (!InOutViewFamilies.Num())
 	{
 		return;
 	}
@@ -4951,21 +4981,21 @@ void FSceneRenderer::CreateSceneRenderers(TConstArrayView<const FSceneViewFamily
 #if !UE_BUILD_SHIPPING
 	bool bSplitScreenDebug = false;
 	TArray<const FSceneViewFamily*, TFixedAllocator<2>> DebugViewFamilyStorage;
-	if (bAllowSplitScreenDebug && IsSplitScreenDebugEnabled(InViewFamilies))
+	if (bAllowSplitScreenDebug && IsSplitScreenDebugEnabled(InOutViewFamilies))
 	{
-		CreateSplitScreenDebugViewFamilies(*InViewFamilies[0], InViewFamilies, DebugViewFamilyStorage);
+		CreateSplitScreenDebugViewFamilies(*InOutViewFamilies[0], InOutViewFamilies, DebugViewFamilyStorage);
 		bSplitScreenDebug = true;
 	}
 #endif
 
-	const FSceneInterface* Scene = InViewFamilies[0]->Scene;
+	const FSceneInterface* Scene = InOutViewFamilies[0]->Scene;
 	check(Scene);
 
 	EShadingPath ShadingPath = GetFeatureLevelShadingPath(Scene->GetFeatureLevel());
 
-	for (int32 FamilyIndex = 0; FamilyIndex < InViewFamilies.Num(); FamilyIndex++)
+	for (int32 FamilyIndex = 0; FamilyIndex < InOutViewFamilies.Num(); FamilyIndex++)
 	{
-		const FSceneViewFamily* InViewFamily = InViewFamilies[FamilyIndex];
+		const FSceneViewFamily* InViewFamily = InOutViewFamilies[FamilyIndex];
 
 		check(InViewFamily);
 		check(InViewFamily->Scene == Scene);
@@ -4990,7 +5020,7 @@ void FSceneRenderer::CreateSceneRenderers(TConstArrayView<const FSceneViewFamily
 #if !UE_BUILD_SHIPPING
 	if (bSplitScreenDebug)
 	{
-		DestroySplitScreenDebugViewFamilies(InViewFamilies);
+		DestroySplitScreenDebugViewFamilies(InOutViewFamilies);
 	}
 #endif
 }
@@ -5379,7 +5409,7 @@ void ResetAndShrinkModifiedBounds(TArray<FBox>& Bounds)
 	Bounds.Reset();
 }
 
-static void RenderViewFamily_RenderThread(FRDGBuilder& GraphBuilder, FSceneRenderer* Renderer, int32 ViewFamilyIndex)
+static void RenderViewFamily_RenderThread(FRDGBuilder& GraphBuilder, FSceneRenderer* Renderer)
 {
 	FSceneViewFamily& ViewFamily = Renderer->ViewFamily;
 
@@ -5387,29 +5417,6 @@ static void RenderViewFamily_RenderThread(FRDGBuilder& GraphBuilder, FSceneRende
 	SCOPE_CYCLE_COUNTER(STAT_TotalSceneRenderingTime);
 	SCOPED_NAMED_EVENT_TCHAR_CONDITIONAL(*ViewFamily.ProfileDescription, FColor::Red, !ViewFamily.ProfileDescription.IsEmpty());
 
-#if WITH_DEBUG_VIEW_MODES
-	if (AllowDebugViewShaderMode(DVSM_VisualizeGPUSkinCache, ViewFamily.GetShaderPlatform(), ViewFamily.GetFeatureLevel()))
-	{
-		FViewInfo& View = Renderer->Views[0];
-		FGPUSkinCacheVisualizationData& VisualizationData = GetGPUSkinCacheVisualizationData();
-
-		// Only run visualization update once, but set debug flags for all view families if the mode is active
-		// Note VisualizationData.Update needs to be called per frame, as || lazy evaluation is used, so need to do it before evaluating VisualizeGPUSkinCache flag
-		if (ViewFamilyIndex > 0 || VisualizationData.Update(View.CurrentGPUSkinCacheVisualizationMode) || ViewFamily.EngineShowFlags.VisualizeGPUSkinCache)
-		{
-			// When activating visualization from the command line, enable VisualizeGPUSkinCache.
-			ViewFamily.EngineShowFlags.SetVisualizeGPUSkinCache(true);
-			ViewFamily.DebugViewShaderMode = DVSM_VisualizeGPUSkinCache;
-
-			// Only draw the visualization info text once.
-			if (ViewFamilyIndex == 0)
-			{
-				Renderer->DrawGPUSkinCacheVisualizationInfoText();
-			}
-		}
-	}
-#endif  // WITH_DEBUG_VIEW_MODES
-
 #if WITH_MGPU
 	if (ViewFamily.bForceCopyCrossGPU)
 	{
@@ -5427,35 +5434,21 @@ static void RenderViewFamily_RenderThread(FRDGBuilder& GraphBuilder, FSceneRende
 	}
 
 	Renderer->FlushCrossGPUFences(GraphBuilder);
-
-#if MESH_DRAW_COMMAND_STATS
-	// Needs to happen after graph builder execute to make sure all FMeshDrawCommandPassSetupTask are finished.
-	// Only during iteration of the visible MDCs we know the custom indirect args used in the frame
-	if (FMeshDrawCommandStatsManager* Instance = FMeshDrawCommandStatsManager::Get())
-	{
-		GraphBuilder.AddPostExecuteCallback([Instance, &RHICmdList = GraphBuilder.RHICmdList]
-		{
-			Instance->QueueCustomDrawIndirectArgsReadback(RHICmdList);
-		});
-	}
-#endif
-
-	// Copy relevant data from ViewInfo to ViewState->PrevFrameViewInfo
-	for (FViewInfo& View : Renderer->Views)
-	{
-		if (View.ViewState)
-		{
-			View.ViewState->PrevFrameViewInfo.bUsesGlobalDistanceField = View.bUsesGlobalDistanceField;
-		}
-	}
 }
 
-static void CleanupViewFamilies_RenderThread(TConstArrayView<FSceneRenderer*> SceneRenderers)
+static void CleanupViewFamilies_RenderThread(FRHICommandListImmediate& RHICmdList, TConstArrayView<FSceneRenderer*> SceneRenderers)
 {
 	LLM_SCOPE(ELLMTag::SceneRender);
 
 	FScene* Scene = SceneRenderers[0]->Scene;
 
+#if MESH_DRAW_COMMAND_STATS
+	if (FMeshDrawCommandStatsManager* Instance = FMeshDrawCommandStatsManager::Get())
+	{
+		Instance->QueueCustomDrawIndirectArgsReadback(RHICmdList);
+	}
+#endif
+
 	{
 		CSV_SCOPED_TIMING_STAT_EXCLUSIVE(PostRenderCleanUp);
 
@@ -5486,44 +5479,42 @@ static void CleanupViewFamilies_RenderThread(TConstArrayView<FSceneRenderer*> Sc
 
 #if RHI_RAYTRACING
 	Scene->RayTracingScene.EndFrame();
-#endif  // RHI_RAYTRACING
-
-#if STATS
-	{
-		QUICK_SCOPE_CYCLE_COUNTER(STAT_RenderViewFamily_RenderThread_MemStats);
+#endif
 
-		// Update scene memory stats that couldn't be tracked continuously
-		SET_MEMORY_STAT(STAT_RenderingSceneMemory, Scene->GetSizeBytes());
+	// Update scene memory stats that couldn't be tracked continuously
+	SET_MEMORY_STAT(STAT_RenderingSceneMemory, Scene->GetSizeBytes());
 
-		SIZE_T ViewStateMemory = 0;
-		for (FSceneRenderer* SceneRenderer : SceneRenderers)
+	SIZE_T ViewStateMemory = 0;
+	for (FSceneRenderer* SceneRenderer : SceneRenderers)
+	{
+		for (FViewInfo& View : SceneRenderer->Views)
 		{
-			for (int32 ViewIndex = 0; ViewIndex < SceneRenderer->Views.Num(); ViewIndex++)
+			// Copy relevant data from ViewInfo to ViewState->PrevFrameViewInfo
+			if (View.ViewState)
 			{
-				if (SceneRenderer->Views[ViewIndex].State)
-				{
-					ViewStateMemory += SceneRenderer->Views[ViewIndex].State->GetSizeBytes();
-				}
+				View.ViewState->PrevFrameViewInfo.bUsesGlobalDistanceField = View.bUsesGlobalDistanceField;
+
+			#if STATS
+				ViewStateMemory += View.ViewState->GetSizeBytes();
+			#endif
 			}
 		}
-		SET_MEMORY_STAT(STAT_ViewStateMemory, ViewStateMemory);
-		SET_MEMORY_STAT(STAT_LightInteractionMemory, FLightPrimitiveInteraction::GetMemoryPoolSize());
 	}
-#endif
-
-#if !UE_BUILD_SHIPPING
-	// Update on screen notifications.
-	FRendererOnScreenNotification::Get().Broadcast();
-#endif
+	SET_MEMORY_STAT(STAT_ViewStateMemory, ViewStateMemory);
+	SET_MEMORY_STAT(STAT_LightInteractionMemory, FLightPrimitiveInteraction::GetMemoryPoolSize());
 
 #if STATS
-	QUICK_SCOPE_CYCLE_COUNTER(STAT_RenderViewFamily_RenderThread_RHIGetGPUFrameCycles);
 	if (FPlatformProperties::SupportsWindowedMode() == false)
 	{
 		/** Update STATS with the total GPU time taken to render the last frame. */
 		SET_CYCLE_COUNTER(STAT_TotalGPUFrameTime, RHIGetGPUFrameCycles());
 	}
 #endif
+
+#if !UE_BUILD_SHIPPING
+	// Update on screen notifications.
+	FRendererOnScreenNotification::Get().Broadcast();
+#endif
 }
 
 void OnChangeCVarRequiringRecreateRenderState(IConsoleVariable* Var)
@@ -5688,17 +5679,26 @@ void FRendererModule::BeginRenderingViewFamilies(FCanvas* Canvas, TConstArrayVie
 		});
 	#endif
 
+		// This array can be mutated by CreateSceneRenderers when split screen debug is enabled.
+		TArray<const FSceneViewFamily*> ViewFamiliesFinal;
+		ViewFamiliesFinal.Reserve(ViewFamilies.Num());
+		for (const FSceneViewFamily* Family : ViewFamilies)
+		{
+			ViewFamiliesFinal.Emplace(Family);
+		}
+
 		// Construct the scene renderers.  This copies the view family attributes into its own structures.
 		const bool bAllowSplitScreenDebug = true;
 		TArray<FSceneRenderer*> SceneRenderers;
-		FSceneRenderer::CreateSceneRenderers(ViewFamilies, Canvas->GetHitProxyConsumer(), bAllowSplitScreenDebug, SceneRenderers);
+		FSceneRenderer::CreateSceneRenderers(ViewFamiliesFinal, Canvas->GetHitProxyConsumer(), bAllowSplitScreenDebug, SceneRenderers);
+		SetupDebugViewModes(SceneRenderers);
 
 		FSceneRenderBuilder SceneRenderBuilder(Scene);
 
 		bool bShowHitProxies = (Canvas->GetHitProxyConsumer() != nullptr);
 		if (!bShowHitProxies)
 		{
-			SceneCaptureUpdateDeferredCapturesInternal(Scene, ViewFamilies, SceneRenderBuilder);
+			SceneCaptureUpdateDeferredCapturesInternal(Scene, ViewFamiliesFinal, SceneRenderBuilder);
 		}
 
 		if (!bShowHitProxies)
@@ -5719,21 +5719,21 @@ void FRendererModule::BeginRenderingViewFamilies(FCanvas* Canvas, TConstArrayVie
 
 		FSceneRenderer::PreallocateCrossGPUFences(SceneRenderers);
 
-		for (int32 ViewFamilyIndex = 0; ViewFamilyIndex < ViewFamilies.Num(); ++ViewFamilyIndex)
+		for (int32 ViewFamilyIndex = 0; ViewFamilyIndex < ViewFamiliesFinal.Num(); ++ViewFamilyIndex)
 		{
 			FSceneRenderer* SceneRenderer = SceneRenderers[ViewFamilyIndex];
 			SceneRenderer->ViewFamily.DisplayInternalsData.Setup(World);
 
 			SceneRenderBuilder.AddRenderer(SceneRenderer, bShowHitProxies ? TEXT("HitProxies") : TEXT("ViewFamilies"),
-				[ViewFamilyIndex] (FRDGBuilder& GraphBuilder, const FSceneRenderFunctionInputs& Inputs)
+				[] (FRDGBuilder& GraphBuilder, const FSceneRenderFunctionInputs& Inputs)
 			{
-				RenderViewFamily_RenderThread(GraphBuilder, Inputs.Renderer, ViewFamilyIndex);
+				RenderViewFamily_RenderThread(GraphBuilder, Inputs.Renderer);
 			});
 		}
 
-		SceneRenderBuilder.AddRenderCommand([SceneRenderers = MoveTemp(SceneRenderers)](FRHICommandListImmediate&)
+		SceneRenderBuilder.AddRenderCommand([SceneRenderers = MoveTemp(SceneRenderers)](FRHICommandListImmediate& RHICmdList)
 		{
-			CleanupViewFamilies_RenderThread(SceneRenderers);
+			CleanupViewFamilies_RenderThread(RHICmdList, SceneRenderers);
 		});
 
 		SceneRenderBuilder.Execute();
@@ -6307,12 +6307,6 @@ TSharedRef<ISceneViewExtension, ESPMode::ThreadSafe> GetRendererViewExtension()
 	class FRendererViewExtension : public ISceneViewExtension
 	{
 	public:
-		virtual void SetupViewFamily(FSceneViewFamily& InViewFamily) {}
-		virtual void SetupView(FSceneViewFamily& InViewFamily, FSceneView& InView) {}
-		virtual void BeginRenderViewFamily(FSceneViewFamily& InViewFamily) {}
-		virtual void PreRenderViewFamily_RenderThread(FRDGBuilder& GraphBuilder, FSceneViewFamily& InViewFamily) {}
-		virtual void PreRenderView_RenderThread(FRDGBuilder& GraphBuilder, FSceneView& InView) {}
-		virtual int32 GetPriority() const { return 0; }
 		virtual void PostRenderView_RenderThread(FRDGBuilder& GraphBuilder, FSceneView& InView)
 		{
 			FViewInfo& View = static_cast<FViewInfo&>(InView);
diff --git a/Engine/Source/Runtime/Renderer/Private/SceneRendering.h b/Engine/Source/Runtime/Renderer/Private/SceneRendering.h
--- a/Engine/Source/Runtime/Renderer/Private/SceneRendering.h
+++ b/Engine/Source/Runtime/Renderer/Private/SceneRendering.h
@@ -2120,7 +2120,7 @@ public:
 	RENDERER_API static FSceneRenderer* CreateSceneRenderer(const FSceneViewFamily* InViewFamily, FHitProxyConsumer* HitProxyConsumer);
 
 	/** Creates multiple scene renderers based on the current feature level.  All view families must point to the same Scene. */
-	RENDERER_API static void CreateSceneRenderers(TConstArrayView<const FSceneViewFamily*> InViewFamilies, FHitProxyConsumer* HitProxyConsumer, bool bAllowSplitScreenDebug, TArray<FSceneRenderer*>& OutSceneRenderers);
+	RENDERER_API static void CreateSceneRenderers(TArray<const FSceneViewFamily*>& InOutViewFamilies, FHitProxyConsumer* HitProxyConsumer, bool bAllowSplitScreenDebug, TArray<FSceneRenderer*>& OutSceneRenderers);
 
 	/** Setups FViewInfo::ViewRect according to ViewFamilly's ScreenPercentageInterface. */
 	void PrepareViewRectsForRendering(FRHICommandListImmediate& RHICmdList);
@@ -2258,8 +2258,6 @@ public:
 	bool ShouldRenderVolumetricFog() const;
 	void ComputeVolumetricFog(FRDGBuilder& GraphBuilder, const FSceneTextures& SceneTextures);
 
-	void DrawGPUSkinCacheVisualizationInfoText();
-
 	virtual bool IsLumenEnabled(const FViewInfo& View) const { return false; }
 	virtual bool IsLumenGIEnabled(const FViewInfo& View) const { return false; }
 	virtual bool AnyViewHasGIMethodSupportingDFAO() const { return false; }
@@ -2551,7 +2549,7 @@ protected:
 	void RenderMeshDistanceFieldVisualization(FRDGBuilder& GraphBuilder, const FMinimalSceneTextures& SceneTextures);
 
 #if !UE_BUILD_SHIPPING
-	static void CreateSplitScreenDebugViewFamilies(const FSceneViewFamily& InFamily, TConstArrayView<const FSceneViewFamily*>& OutFamilies, TArray<const FSceneViewFamily*, TFixedAllocator<2>>& OutFamiliesStorage);
+	static void CreateSplitScreenDebugViewFamilies(const FSceneViewFamily& InFamily, TArray<const FSceneViewFamily*>& OutViewFamilies, TArray<const FSceneViewFamily*, TFixedAllocator<2>>& OutFamiliesStorage);
 	static void DestroySplitScreenDebugViewFamilies(TConstArrayView<const FSceneViewFamily*> ViewFamilies);
 #endif
 
