Fixed GPU skin cache text being overlayed all the time.

diff --git a/Engine/Source/Runtime/Renderer/Private/SceneRendering.cpp b/Engine/Source/Runtime/Renderer/Private/SceneRendering.cpp
--- a/Engine/Source/Runtime/Renderer/Private/SceneRendering.cpp
+++ b/Engine/Source/Runtime/Renderer/Private/SceneRendering.cpp
@@ -4304,7 +4304,7 @@ void FSceneRenderer::OnRenderFinish(FRDGBuilder& GraphBuilder, FRDGTextureRef Vi
 				bool bIsFrozen = ViewState && (ViewState->bIsFrozen);
 				bool bLocked = View.bIsLocked;
 				const bool bStereoView = IStereoRendering::IsStereoEyeView(View);
-				const FName GPUSkinCacheVisualizationMode = View.CurrentGPUSkinCacheVisualizationMode;
+				const bool bGPUSkinCacheVisualzationMode = SkinCache && ViewFamily.EngineShowFlags.VisualizeGPUSkinCache && View.CurrentGPUSkinCacheVisualizationMode != NAME_None;
 
 				// display a warning if an ambient cubemap uses non-angular mipmap filtering
 				bool bShowAmbientCubemapMipGenSettingsWarning = false;
@@ -4322,7 +4322,7 @@ void FSceneRenderer::OnRenderFinish(FRDGBuilder& GraphBuilder, FRDGTextureRef Vi
 					}
 				}
 #endif
-				if ((GAreScreenMessagesEnabled && !GEngine->bSuppressMapWarnings) && (bIsFrozen || bLocked || bStereoView || bShowAmbientCubemapMipGenSettingsWarning || bAnyWarning || GPUSkinCacheVisualizationMode != NAME_None))
+				if ((GAreScreenMessagesEnabled && !GEngine->bSuppressMapWarnings) && (bIsFrozen || bLocked || bStereoView || bShowAmbientCubemapMipGenSettingsWarning || bAnyWarning || bGPUSkinCacheVisualzationMode))
 				{
 					RDG_EVENT_SCOPE_CONDITIONAL(GraphBuilder, Views.Num() > 1, "View%d", ViewIndex);
 
@@ -4332,7 +4332,7 @@ void FSceneRenderer::OnRenderFinish(FRDGBuilder& GraphBuilder, FRDGTextureRef Vi
 					const bool bIsMobileMultiViewEnabled = View.bIsMobileMultiViewEnabled;
 
 					AddDrawCanvasPass(GraphBuilder, {}, View, Output,
-						[this, SkinCache, GPUSkinCacheVisualizationMode, ViewState, GPUSkinCacheExtraRequiredMemory,
+						[this, &View, SkinCache, bGPUSkinCacheVisualzationMode, ViewState, GPUSkinCacheExtraRequiredMemory,
 						bLocked, bShowPrecomputedVisibilityWarning, bShowDemotedLocalMemoryWarning, bShowGlobalClipPlaneWarning, bShowDFAODisabledWarning, bShowDFDisabledWarning,
 						bIsFrozen, bShowSkylightWarning, bShowPointLightWarning, bShowShadowedLightOverflowWarning,
 						bShowMobileLowQualityLightmapWarning, bShowMobileDynamicCSMWarning, bMobileMissingSkyMaterial, 
@@ -4430,9 +4430,9 @@ void FSceneRenderer::OnRenderFinish(FRDGBuilder& GraphBuilder, FRDGTextureRef Vi
 							Writer.DrawLine(FText::FromString(String));
 						}
 
-						if (GPUSkinCacheVisualizationMode.IsValid() && SkinCache)
+						if (bGPUSkinCacheVisualzationMode)
 						{
-							SkinCache->DrawVisualizationInfoText(GPUSkinCacheVisualizationMode, Writer);
+							SkinCache->DrawVisualizationInfoText(View.CurrentGPUSkinCacheVisualizationMode, Writer);
 						}
 
 						if (bShowLocalExposureDisabledWarning)
