Add missing BuildRenderingCommands for Mobile Forward Multipass Translucency and MeshDecal

diff --git a/Engine/Source/Runtime/Renderer/Private/MobileShadingRenderer.cpp b/Engine/Source/Runtime/Renderer/Private/MobileShadingRenderer.cpp
--- a/Engine/Source/Runtime/Renderer/Private/MobileShadingRenderer.cpp
+++ b/Engine/Source/Runtime/Renderer/Private/MobileShadingRenderer.cpp
@@ -1864,6 +1864,12 @@ void FMobileSceneRenderer::RenderForwardMultiPass(FRDGBuilder& GraphBuilder, FMo
 	const bool bDoOcclusionQueries = (!bIsFullDepthPrepassEnabled && ViewContext.bIsLastView && DoOcclusionQueries());
 	SecondPassParameters->RenderTargets.NumOcclusionQueries = bDoOcclusionQueries ? ComputeNumOcclusionQueriesToBatch() : 0u;
 
+	if (Scene->GPUScene.IsEnabled())
+	{
+		View.ParallelMeshDrawCommandPasses[EMeshPass::MeshDecal_SceneColor].BuildRenderingCommands(GraphBuilder, Scene->GPUScene, SecondParameterCollection->MeshDecalSceneColorInstanceCullingDrawParams);
+		View.ParallelMeshDrawCommandPasses[StandardTranslucencyMeshPass].BuildRenderingCommands(GraphBuilder, Scene->GPUScene, SecondParameterCollection->TranslucencyInstanceCullingDrawParams);
+	}
+
 	GraphBuilder.AddPass(
 		RDG_EVENT_NAME("DecalsAndTranslucency"),
 		SecondPassParameters,
