Fix for crash when rendering with simple views and blue noise resources are not available, nor GSystemTextures.

#tests PC dx12, cooking HLOD in Game
#rnx

diff --git a/Engine/Source/Runtime/Renderer/Private/BlueNoise.cpp b/Engine/Source/Runtime/Renderer/Private/BlueNoise.cpp
--- a/Engine/Source/Runtime/Renderer/Private/BlueNoise.cpp
+++ b/Engine/Source/Runtime/Renderer/Private/BlueNoise.cpp
@@ -9,6 +9,7 @@
 #include "Engine/Texture2D.h"
 #include "TextureResource.h"
 #include "SystemTextures.h"
+#include "GlobalRenderResources.h"
 
 IMPLEMENT_GLOBAL_SHADER_PARAMETER_STRUCT(FBlueNoise, "BlueNoise");
 
@@ -22,54 +23,60 @@ FBlueNoiseParameters GetBlueNoiseDummyParameters()
 	return Out;
 }
 
+static void FillUpBlueNoiseParametersFromTexture(FBlueNoiseParameters& Out)
+{
+	const FIntVector BlueNoiseSize = Out.ScalarTexture->GetSizeXYZ();
+
+	Out.Dimensions = FIntVector(
+		BlueNoiseSize.X,
+		BlueNoiseSize.X,
+		BlueNoiseSize.Y / FMath::Max<int32>(1, BlueNoiseSize.X));
+
+	Out.ModuloMasks = FIntVector(
+		(1u << FMath::FloorLog2(Out.Dimensions.X)) - 1,
+		(1u << FMath::FloorLog2(Out.Dimensions.Y)) - 1,
+		(1u << FMath::FloorLog2(Out.Dimensions.Z)) - 1);
+
+	check((Out.ModuloMasks.X + 1) == Out.Dimensions.X
+		&& (Out.ModuloMasks.Y + 1) == Out.Dimensions.Y
+		&& (Out.ModuloMasks.Z + 1) == Out.Dimensions.Z);
+}
+
 FBlueNoiseParameters GetBlueNoiseParameters()
 {
 	FBlueNoiseParameters Out;
 	check(GEngine);
+	check(GEngine->BlueNoiseScalarTexture && GEngine->BlueNoiseVec2Texture);
 
-	// Hacky fix for BlueNoise not being loaded before use on some code paths. Erring on the side of caution, e.g. GEngine shouldn't realistically be null
-	if (!GEngine->BlueNoiseScalarTexture || !GEngine->BlueNoiseVec2Texture
-		|| GEngine->BlueNoiseScalarTexture->GetResource() == nullptr
-		|| GEngine->BlueNoiseVec2Texture->GetResource() == nullptr
-		)
-	{
-		return Out;
-	}
-	
-	if (GEngine->BlueNoiseScalarTexture)
-	{
-		Out.Dimensions = FIntVector(
-			GEngine->BlueNoiseScalarTexture->GetSizeX(),
-			GEngine->BlueNoiseScalarTexture->GetSizeX(),
-			GEngine->BlueNoiseScalarTexture->GetSizeY() / FMath::Max<int32>(1, GEngine->BlueNoiseScalarTexture->GetSizeX()));
+	Out.ScalarTexture = GEngine->BlueNoiseScalarTexture->GetResource()->TextureRHI;
+	Out.Vec2Texture = GEngine->BlueNoiseVec2Texture->GetResource()->TextureRHI;
 
-		Out.ModuloMasks = FIntVector(
-			(1u << FMath::FloorLog2(Out.Dimensions.X)) - 1,
-			(1u << FMath::FloorLog2(Out.Dimensions.Y)) - 1,
-			(1u << FMath::FloorLog2(Out.Dimensions.Z)) - 1);
+	FillUpBlueNoiseParametersFromTexture(Out);
+	return Out;
+}
 
-		check((Out.ModuloMasks.X + 1)  == Out.Dimensions.X
-			&&(Out.ModuloMasks.Y + 1) == Out.Dimensions.Y
-			&&(Out.ModuloMasks.Z + 1) == Out.Dimensions.Z);
+FBlueNoise GetBlueNoiseGlobalParameters()
+{
+	FBlueNoise Out;
+	Out.BlueNoise = GetBlueNoiseParameters();
+	return Out;
+}
 
-		Out.ScalarTexture = GEngine->BlueNoiseScalarTexture->GetResource()->TextureRHI;
+FBlueNoiseParameters GetBlueNoiseParametersForView()
+{
+	FBlueNoiseParameters Out;
+	check(GEngine);
 
-		Out.Vec2Texture = GSystemTextures.BlackDummy->GetRHI();
-		if (GEngine->BlueNoiseVec2Texture != nullptr)
-		{
-			Out.Vec2Texture = GEngine->BlueNoiseVec2Texture->GetResource()->TextureRHI;
-		}
+	if (GEngine->BlueNoiseScalarTexture != nullptr && GEngine->BlueNoiseScalarTexture->GetResource() != nullptr)
+	{
+		Out.ScalarTexture = GEngine->BlueNoiseScalarTexture->GetResource()->TextureRHI;
 	}
 	else
 	{
-		Out = GetBlueNoiseDummyParameters();
+		Out.ScalarTexture = GBlackVolumeTexture->TextureRHI.GetReference();
 	}
-	return Out;
-}
+	Out.Vec2Texture = GBlackVolumeTexture->TextureRHI.GetReference();
 
-FBlueNoise GetBlueNoiseGlobalParameters()
-{
-	FBlueNoise Out;
-	Out.BlueNoise = GetBlueNoiseParameters();
+	FillUpBlueNoiseParametersFromTexture(Out);
 	return Out;
 }
\ No newline at end of file
diff --git a/Engine/Source/Runtime/Renderer/Private/BlueNoise.h b/Engine/Source/Runtime/Renderer/Private/BlueNoise.h
--- a/Engine/Source/Runtime/Renderer/Private/BlueNoise.h
+++ b/Engine/Source/Runtime/Renderer/Private/BlueNoise.h
@@ -20,6 +20,13 @@ BEGIN_GLOBAL_SHADER_PARAMETER_STRUCT(FBlueNoise, RENDERER_API)
 	SHADER_PARAMETER_STRUCT_INCLUDE(FBlueNoiseParameters, BlueNoise)
 END_GLOBAL_SHADER_PARAMETER_STRUCT()
 
+// These GetBlueNoiseXXX functions should only be called from a rendering system, when this system make it possible 
+// to load the blue noise textures when calling GEngine->LoadBlueNoiseTexture from SceneRendering.cpp.
 extern RENDERER_API FBlueNoise GetBlueNoiseGlobalParameters();
 extern RENDERER_API FBlueNoiseParameters GetBlueNoiseParameters();
 extern RENDERER_API FBlueNoiseParameters GetBlueNoiseDummyParameters();
+
+// This function is for filling up the View blue noise parameters used for materials.
+// Sometimes views are created for rendering without GSystemTextures initialised, e.g. HLOD baking, canvas DrawTile.
+// So we have to rely on global default GPU resources in this case.
+extern RENDERER_API FBlueNoiseParameters GetBlueNoiseParametersForView();
diff --git a/Engine/Source/Runtime/Renderer/Private/SceneRendering.cpp b/Engine/Source/Runtime/Renderer/Private/SceneRendering.cpp
--- a/Engine/Source/Runtime/Renderer/Private/SceneRendering.cpp
+++ b/Engine/Source/Runtime/Renderer/Private/SceneRendering.cpp
@@ -2212,11 +2212,10 @@ void FViewInfo::SetupUniformBufferParameters(
 	ViewUniformShaderParameters.GPUSceneViewId = GPUSceneViewId;
 
 	{
-		FBlueNoise BlueNoiseParam = GetBlueNoiseGlobalParameters();
-		ViewUniformShaderParameters.BlueNoiseScalarTexture = BlueNoiseParam.BlueNoise.ScalarTexture;
-		ViewUniformShaderParameters.BlueNoiseScalarTexture = OrBlack2DIfNull(ViewUniformShaderParameters.BlueNoiseScalarTexture);
-		ViewUniformShaderParameters.BlueNoiseDimensions = BlueNoiseParam.BlueNoise.Dimensions;
-		ViewUniformShaderParameters.BlueNoiseModuloMasks = BlueNoiseParam.BlueNoise.ModuloMasks;
+		FBlueNoiseParameters BlueNoiseParam = GetBlueNoiseParametersForView();
+		ViewUniformShaderParameters.BlueNoiseScalarTexture	= BlueNoiseParam.ScalarTexture;
+		ViewUniformShaderParameters.BlueNoiseDimensions		= BlueNoiseParam.Dimensions;
+		ViewUniformShaderParameters.BlueNoiseModuloMasks	= BlueNoiseParam.ModuloMasks;
 	}
 }
 
