Translucency lighting volume improved stability w.r.t. camera FoV angle.
- By default, now we disable adaptation to FoV as the volume is around the camera anyway
- If enabled, I have made it safe for the volume size to not explode dueto the use of tan
- Added debug view information to see the volume size in world, as well as size in float

Previous default behavior is used for EngineTest to not have to update many screenshots unnecessarily.

#tests Tested pc dx12 editor.
#rnx

diff --git a/Engine/Shaders/Private/TranslucentLightingShaders.usf b/Engine/Shaders/Private/TranslucentLightingShaders.usf
--- a/Engine/Shaders/Private/TranslucentLightingShaders.usf
+++ b/Engine/Shaders/Private/TranslucentLightingShaders.usf
@@ -153,3 +153,74 @@ void InjectAmbientCubemapMainPS(
 
 
 #endif // FEATURE_LEVEL
+
+
+#ifdef DEBUG_TRANSLUCENCY_LIGHTING_VOLUME_CS
+
+#include "/Engine/Private/ShaderPrint.ush"
+
+void PrintTranslucencyVolume(inout FShaderPrintContext Context, float3 Min, float3 Max, float4 Color)
+{
+	float3 C000 = float3(Min.x, Min.y, Min.z);
+	float3 C001 = float3(Min.x, Min.y, Max.z);
+	float3 C010 = float3(Min.x, Max.y, Min.z);
+	float3 C011 = float3(Min.x, Max.y, Max.z);
+	float3 C100 = float3(Max.x, Min.y, Min.z);
+	float3 C101 = float3(Max.x, Min.y, Max.z);
+	float3 C110 = float3(Max.x, Max.y, Min.z);
+	float3 C111 = float3(Max.x, Max.y, Max.z);
+
+	AddLineTWS(Context, C000, C001, Color);
+	AddLineTWS(Context, C001, C011, Color);
+	AddLineTWS(Context, C011, C010, Color);
+	AddLineTWS(Context, C010, C000, Color);
+
+	AddLineTWS(Context, C100, C101, Color);
+	AddLineTWS(Context, C101, C111, Color);
+	AddLineTWS(Context, C111, C110, Color);
+	AddLineTWS(Context, C110, C100, Color);
+
+	AddLineTWS(Context, C001, C101, Color);
+	AddLineTWS(Context, C100, C000, Color);
+	AddLineTWS(Context, C011, C111, Color);
+	AddLineTWS(Context, C110, C010, Color);
+}
+
+[numthreads(THREADGROUP_SIZE, THREADGROUP_SIZE, THREADGROUP_SIZE)]
+void DebugTranslucencyLightingVolumeCS(
+	uint3 GroupId : SV_GroupID,
+	uint3 DispatchThreadId : SV_DispatchThreadID,
+	uint3 GroupThreadId : SV_GroupThreadID)
+{
+	if (all(DispatchThreadId == 0))
+	{
+		FShaderPrintContext Context = InitShaderPrintContext(true, uint2(50, 50));
+
+		const float3 VolMin0 = View.TranslucencyLightingVolumeMin[0].xyz;
+		const float3 VolMin1 = View.TranslucencyLightingVolumeMin[1].xyz;
+		const float3 VolSize0 = 1.0f / View.TranslucencyLightingVolumeInvSize[0].xyz;
+		const float3 VolSize1 = 1.0f / View.TranslucencyLightingVolumeInvSize[1].xyz;
+		const float3 VolMax0 = VolMin0 + VolSize0;
+		const float3 VolMax1 = VolMin1 + VolSize1;
+
+		PrintTranslucencyVolume(Context, VolMin0, VolMax0, float4(0, 1, 0, 1));
+		PrintTranslucencyVolume(Context, VolMin1, VolMax1, float4(1, 1, 0, 1));
+
+		Print(Context, TEXT("Translucency Light Volume0 Size = "), FontWhite);
+		Print(Context, VolSize0.x, FontYellow, 10, 1);
+		Print(Context, TEXT(", "), FontWhite);
+		Print(Context, VolSize0.y, FontYellow, 10, 1);
+		Print(Context, TEXT(", "), FontWhite);
+		Print(Context, VolSize0.y, FontYellow, 10, 1);
+		Newline(Context);
+		Print(Context, TEXT("Translucency Light Volume1 Size = "), FontWhite);
+		Print(Context, VolSize1.x, FontYellow, 10, 1);
+		Print(Context, TEXT(", "), FontWhite);
+		Print(Context, VolSize1.y, FontYellow, 10, 1);
+		Print(Context, TEXT(", "), FontWhite);
+		Print(Context, VolSize1.y, FontYellow, 10, 1);
+
+	}
+}
+#endif // DEBUG_TRANSLUCENCY_LIGHTING_VOLUME_CS
+
diff --git a/Engine/Source/Runtime/Renderer/Private/TranslucentLighting.cpp b/Engine/Source/Runtime/Renderer/Private/TranslucentLighting.cpp
--- a/Engine/Source/Runtime/Renderer/Private/TranslucentLighting.cpp
+++ b/Engine/Source/Runtime/Renderer/Private/TranslucentLighting.cpp
@@ -76,7 +76,7 @@ float GTranslucentVolumeMinFOV = 45;
 static FAutoConsoleVariableRef CVarTranslucentVolumeMinFOV(
 	TEXT("r.TranslucentVolumeMinFOV"),
 	GTranslucentVolumeMinFOV,
-	TEXT("Minimum FOV for translucent lighting volume.  Prevents popping in lighting when zooming in."),
+	TEXT("Minimum FOV for translucent lighting volume. Prevents popping in lighting when zooming in."),
 	ECVF_RenderThreadSafe | ECVF_Scalability
 	);
 
@@ -136,6 +136,24 @@ static TAutoConsoleVariable<int32> CVarTranslucencyLightingVolumeInjectDirection
 	TEXT("Enable sampling of the directional light CSM.\n"),
 	ECVF_RenderThreadSafe | ECVF_Scalability);
 
+// Adaptation to camera angle is now disabled by default to avoid change on lighting when zoomin in or out.
+// The volume remains around the camera anyway, so the camera angle should not matter when the setup mostly depends on start and end distance for each cascade.
+int32 GTranslucencyLightingVolumeAdaptToPerspective = 0;
+static FAutoConsoleVariableRef CVarTranslucencyLightingVolumeAdaptToPerspective(
+	TEXT("r.TranslucencyLightingVolume.AdaptToPerspective"),
+	GTranslucencyLightingVolumeAdaptToPerspective,
+	TEXT("The translucent volume will adapt to the camera perspective when zooming. This can can result in pops for extreme zoom-in so it can be disabled if needed."),
+	ECVF_RenderThreadSafe
+);
+
+int32 GTranslucencyLightingVolumeDebug = 0;
+static FAutoConsoleVariableRef CVarTranslucencyLightingVolumeDebug(
+	TEXT("r.TranslucencyLightingVolume.Debug"),
+	GTranslucencyLightingVolumeDebug,
+	TEXT("Debug information for the translucency lighting volume."),
+	ECVF_RenderThreadSafe
+);
+
 /** Function returning current translucency lighting volume dimensions. */
 int32 GetTranslucencyLightingVolumeDim()
 {
@@ -156,7 +174,7 @@ void FViewInfo::CalcTranslucencyLightingVolumeBounds(FBox* InOutCascadeBoundsArr
 		double FieldOfView = DOUBLE_PI / 4.0;
 		double AspectRatio = 1.0;
 
-		if (IsPerspectiveProjection())
+		if (IsPerspectiveProjection() && GTranslucencyLightingVolumeAdaptToPerspective > 0)
 		{
 			// Derive FOV and aspect ratio from the perspective projection matrix
 			FieldOfView = FMath::Atan(1.0 / ShadowViewMatrices.GetProjectionMatrix().M[0][0]);
@@ -169,12 +187,18 @@ void FViewInfo::CalcTranslucencyLightingVolumeBounds(FBox* InOutCascadeBoundsArr
 			AspectRatio = ShadowViewMatrices.GetProjectionMatrix().M[1][1] / ShadowViewMatrices.GetProjectionMatrix().M[0][0];
 		}
 
-		const double StartHorizontalLength = FrustumStartDistance * FMath::Tan(FieldOfView);
+		// Tan of field of view can explode when FieldOfView is close to 180 degree, when using the adaptation to Fov above.
+		// Close to this edge case, the volume fitted on the frustum vertices can explode to infinity.
+		// In order to fix that, we clamp the value to the length of the diagonal of a cube of size OuterDistance.
+		const double MaxTanFieldOfViewLength = FMath::Sqrt(FMath::Max(1.0f, OuterDistance * OuterDistance + OuterDistance * OuterDistance + OuterDistance * OuterDistance));
+		const double TanFieldOfView = FMath::Tan(FieldOfView);
+
+		const double StartHorizontalLength = FMath::Min(MaxTanFieldOfViewLength, FrustumStartDistance * TanFieldOfView);
 		const FVector StartCameraRightOffset = ShadowViewMatrices.GetViewMatrix().GetColumn(0) * StartHorizontalLength;
 		const double StartVerticalLength = StartHorizontalLength / AspectRatio;
 		const FVector StartCameraUpOffset = ShadowViewMatrices.GetViewMatrix().GetColumn(1) * StartVerticalLength;
 
-		const double EndHorizontalLength = FrustumEndDistance * FMath::Tan(FieldOfView);
+		const double EndHorizontalLength = FMath::Min(MaxTanFieldOfViewLength, FrustumEndDistance * TanFieldOfView);
 		const FVector EndCameraRightOffset = ShadowViewMatrices.GetViewMatrix().GetColumn(0) * EndHorizontalLength;
 		const double EndVerticalLength = EndHorizontalLength / AspectRatio;
 		const FVector EndCameraUpOffset = ShadowViewMatrices.GetViewMatrix().GetColumn(1) * EndVerticalLength;
@@ -1631,11 +1655,89 @@ void InjectSimpleTranslucencyLightingVolumeArray(
 	}
 }
 
+class FDebugTranslucencyLightingVolumeCS : public FGlobalShader
+{
+	DECLARE_GLOBAL_SHADER(FDebugTranslucencyLightingVolumeCS)
+	SHADER_USE_PARAMETER_STRUCT(FDebugTranslucencyLightingVolumeCS, FGlobalShader)
+
+	BEGIN_SHADER_PARAMETER_STRUCT(FParameters, )
+		SHADER_PARAMETER_STRUCT_REF(FViewUniformShaderParameters, ViewUniformBuffer)
+		SHADER_PARAMETER_STRUCT_INCLUDE(ShaderPrint::FShaderParameters, ShaderPrintUniformBuffer)
+	END_SHADER_PARAMETER_STRUCT()
+
+	using FPermutationDomain = TShaderPermutationDomain<>;
+
+	static int32 GetGroupSize()
+	{
+		return 1;
+	}
+
+	static bool ShouldCompilePermutation(const FGlobalShaderPermutationParameters& Parameters)
+	{
+		return IsFeatureLevelSupported(Parameters.Platform, ERHIFeatureLevel::SM5);
+	}
+
+	static void ModifyCompilationEnvironment(const FGlobalShaderPermutationParameters& Parameters, FShaderCompilerEnvironment& OutEnvironment)
+	{
+		FGlobalShader::ModifyCompilationEnvironment(Parameters, OutEnvironment);
+		OutEnvironment.SetDefine(TEXT("THREADGROUP_SIZE"), GetGroupSize());
+		OutEnvironment.SetDefine(TEXT("DEBUG_TRANSLUCENCY_LIGHTING_VOLUME_CS"), 1);
+	}
+
+	static EShaderPermutationPrecacheRequest ShouldPrecachePermutation(const FGlobalShaderPermutationParameters& Parameters)
+	{
+		return EShaderPermutationPrecacheRequest::NotPrecached;
+	}
+};
+
+IMPLEMENT_GLOBAL_SHADER(FDebugTranslucencyLightingVolumeCS, "/Engine/Private/TranslucentLightingShaders.usf", "DebugTranslucencyLightingVolumeCS", SF_Compute);
+
+static void DrawDebugTranslucencyLightingVolume(
+	FRDGBuilder& GraphBuilder,
+	const TArrayView<const FViewInfo> Views)
+{
+	// Draw debug the translucency volume lighting buffer
+	if (GTranslucencyLightingVolumeDebug <= 0)
+	{
+		return;
+	}
+
+	for (int32 ViewIndex = 0; ViewIndex < Views.Num(); ++ViewIndex)
+	{
+		const FViewInfo& View = Views[ViewIndex];
+		RDG_GPU_MASK_SCOPE(GraphBuilder, View.GPUMask);
+		RDG_EVENT_SCOPE(GraphBuilder, "View%d", ViewIndex);
+
+		FDebugTranslucencyLightingVolumeCS::FParameters* PassParameters = GraphBuilder.AllocParameters<FDebugTranslucencyLightingVolumeCS::FParameters>();
+		PassParameters->ViewUniformBuffer = View.ViewUniformBuffer;
+
+		ShaderPrint::SetEnabled(true);
+		ShaderPrint::RequestSpaceForLines(128u);
+		ShaderPrint::RequestSpaceForCharacters(128u);
+		ShaderPrint::RequestSpaceForTriangles(64u);
+		ShaderPrint::SetParameters(GraphBuilder, View.ShaderPrintData, PassParameters->ShaderPrintUniformBuffer);
+
+		FDebugTranslucencyLightingVolumeCS::FPermutationDomain PermutationVector;
+		auto ComputeShader = View.ShaderMap->GetShader<FDebugTranslucencyLightingVolumeCS>(PermutationVector);
+		const FIntVector GroupCount = FComputeShaderUtils::GetGroupCount(View.ViewRect.Size(), FDebugTranslucencyLightingVolumeCS::GetGroupSize());
+
+		FComputeShaderUtils::AddPass(
+			GraphBuilder,
+			RDG_EVENT_NAME("DebugTranslucencyLightingVolumeCS"),
+			ComputeShader,
+			PassParameters,
+			GroupCount);
+	}
+}
+
 void FilterTranslucencyLightingVolume(
 	FRDGBuilder& GraphBuilder,
 	const TArrayView<const FViewInfo> Views,
 	FTranslucencyLightingVolumeTextures& Textures)
 {
+	// Draw debug information if needed.
+	DrawDebugTranslucencyLightingVolume(GraphBuilder, Views);
+
 	if (!GUseTranslucentLightingVolumes || Views.Num() == 0 || !RHISupportsVolumeTextureRendering(Views[0].GetShaderPlatform()) || !GUseTranslucencyVolumeBlur)
 	{
 		return;
