Change bokeh gather distinction filter to be opt-in for now

diff --git a/Engine/Source/Runtime/Renderer/Private/PostProcess/DiaphragmDOF.cpp b/Engine/Source/Runtime/Renderer/Private/PostProcess/DiaphragmDOF.cpp
--- a/Engine/Source/Runtime/Renderer/Private/PostProcess/DiaphragmDOF.cpp
+++ b/Engine/Source/Runtime/Renderer/Private/PostProcess/DiaphragmDOF.cpp
@@ -46,14 +46,14 @@ TAutoConsoleVariable<int32> CVarAccumulatorQuality(
 
 TAutoConsoleVariable<float> CVarGatherDistinctionLimit(
 	TEXT("r.DOF.Gather.DistinctionLimit"),
-	1000.0,
+	0.0,
 	TEXT("Bokeh that are not scattered but exceed this distinction limit are discarded instead of gathered.\n")
 	TEXT("Distinction is calculcated based on perceived brightness and radius.\n")
 	TEXT("This is used to reduce noise in gathered bokeh, without drawing additional scattered bokeh.\n")
 	TEXT("This can cause bokeh to become darker than they should be, which is a compromise for noise reduction.\n")
 	TEXT("See also: r.DOF.Scatter.NeighborCompareMaxColor \n")
 	TEXT("Set to 0 to disable.\n")
-	TEXT("(default: 1000)"),
+	TEXT("(default: 0)"),
 	ECVF_Scalability | ECVF_RenderThreadSafe);
 
 TAutoConsoleVariable<int32> CVarEnableGatherBokehSettings(
