Removed cached bloom and film grain TObjectPtr<UTexture2D> instances, not needed in rendering code

#rnx

diff --git a/Engine/Source/Runtime/Engine/Private/LocalPlayer.cpp b/Engine/Source/Runtime/Engine/Private/LocalPlayer.cpp
--- a/Engine/Source/Runtime/Engine/Private/LocalPlayer.cpp
+++ b/Engine/Source/Runtime/Engine/Private/LocalPlayer.cpp
@@ -61,12 +61,6 @@ static FAutoConsoleVariableRef CVarCalcLocalPlayerCachedLODDistanceFactor(
 	TEXT("Should we calculate a LOD Distance Factor based on the current FOV.  Should not be necessary since LOD is already based on screen size.\n")
 	);
 
-static bool GLocalPlayerCleanUpSceneViewStateTransientReferences = true;
-static FAutoConsoleVariableRef CVarLocalPlayerCleanUpSceneViewStateTransientReferences(
-	TEXT("r.CleanUpSceneViewStateTransientReferences"),
-	GLocalPlayerCleanUpSceneViewStateTransientReferences,
-	TEXT("clean up scene view state transient references"));
-
 DECLARE_CYCLE_STAT(TEXT("CalcSceneView"), STAT_CalcSceneView, STATGROUP_Engine);
 
 //////////////////////////////////////////////////////////////////////////
@@ -1852,10 +1846,6 @@ void ULocalPlayer::CleanupViewState(FStringView MidParentRootPath /*= {}*/)
 		if (Ref)
 		{
 			Ref->ClearMIDPool(MidParentRootPath);
-			if (GLocalPlayerCleanUpSceneViewStateTransientReferences)
-			{
-				Ref->CleanUpTransientReferences();
-			}
 		}
 	}
 }
diff --git a/Engine/Source/Runtime/Engine/Public/SceneManagement.h b/Engine/Source/Runtime/Engine/Public/SceneManagement.h
--- a/Engine/Source/Runtime/Engine/Public/SceneManagement.h
+++ b/Engine/Source/Runtime/Engine/Public/SceneManagement.h
@@ -138,7 +138,6 @@ public:
 	virtual FSceneViewState* GetConcreteViewState () = 0;
 
 	virtual void AddReferencedObjects(FReferenceCollector& Collector) = 0;
-	virtual void CleanUpTransientReferences() = 0;
 
 	virtual SIZE_T GetSizeBytes() const { return 0; }
 
diff --git a/Engine/Source/Runtime/Renderer/Private/PostProcess/PostProcessFFTBloom.cpp b/Engine/Source/Runtime/Renderer/Private/PostProcess/PostProcessFFTBloom.cpp
--- a/Engine/Source/Runtime/Renderer/Private/PostProcess/PostProcessFFTBloom.cpp
+++ b/Engine/Source/Runtime/Renderer/Private/PostProcess/PostProcessFFTBloom.cpp
@@ -316,12 +316,12 @@ void InitDomainAndGetKernel(
 
 	const auto& PPSettings = View.FinalPostProcessSettings;
 
+	// This should exist if we called IsFFTBloomPhysicalKernelReady.
+	check(View.FFTBloomKernelTexture && View.FFTBloomKernelTexture->TextureRHI);
+
 	const FTexture2DResource* BloomConvolutionTextureResource = View.FFTBloomKernelTexture;
 	const FTextureRHIRef& PhysicalSpaceKernelTextureRef = BloomConvolutionTextureResource->TextureRHI;
 
-	// This should exist if we called IsFFTBloomPhysicalKernelReady.
-	check(BloomConvolutionTextureResource && PhysicalSpaceKernelTextureRef);
-
 	const float BloomConvolutionSize = PPSettings.BloomConvolutionSize;
 	const FVector2f CenterUV = FVector2f(PPSettings.BloomConvolutionCenterUV); // TODO: remove
 
@@ -337,7 +337,7 @@ void InitDomainAndGetKernel(
 		FClearValueBinding::None,
 		TexCreate_ShaderResource | TexCreate_UAV);
 
-	if (ViewState && ViewState->BloomFFTKernel.Physical != nullptr && ViewState->BloomFFTKernel.Spectral && CVarBloomCacheKernel.GetValueOnRenderThread() != 0)
+	if (ViewState && ViewState->BloomFFTKernel.Spectral && CVarBloomCacheKernel.GetValueOnRenderThread() != 0)
 	{
 		auto& FFTKernel = ViewState->BloomFFTKernel;
 
@@ -699,7 +699,6 @@ void InitDomainAndGetKernel(
 	{
 		ViewState->BloomFFTKernel.Scale = BloomConvolutionSize;
 		ViewState->BloomFFTKernel.ImageSize = Intermediates.ImageSize;
-		ViewState->BloomFFTKernel.Physical = View.FinalPostProcessSettings.BloomConvolutionTexture;
 		ViewState->BloomFFTKernel.PhysicalRHI = PhysicalSpaceKernelTextureRef;
 		ViewState->BloomFFTKernel.PhysicalMipLevel = BloomConvolutionTextureResource->GetCurrentMipCount();
 		{
diff --git a/Engine/Source/Runtime/Renderer/Private/PostProcess/PostProcessTonemap.cpp b/Engine/Source/Runtime/Renderer/Private/PostProcess/PostProcessTonemap.cpp
--- a/Engine/Source/Runtime/Renderer/Private/PostProcess/PostProcessTonemap.cpp
+++ b/Engine/Source/Runtime/Renderer/Private/PostProcess/PostProcessTonemap.cpp
@@ -725,9 +725,8 @@ FScreenPassTexture AddTonemapPass(FRDGBuilder& GraphBuilder, const FViewInfo& Vi
 		FRHITexture* FilmGrainTextureRHI = View.FilmGrainTexture->GetTexture2DRHI();
 		FRDGTextureRef FilmGrainTexture = RegisterExternalTexture(GraphBuilder, FilmGrainTextureRHI, TEXT("FilmGrain.OriginalTexture"));
 
-
 		FRDGBufferRef FilmGrainConstantsBuffer;
-		if (View.ViewState && View.ViewState->FilmGrainCache.Texture != nullptr && View.ViewState->FilmGrainCache.TextureRHI == FilmGrainTextureRHI && CVarFilmGrainCacheTextureConstants.GetValueOnRenderThread() != 0)
+		if (View.ViewState && View.ViewState->FilmGrainCache.TextureRHI == FilmGrainTextureRHI && CVarFilmGrainCacheTextureConstants.GetValueOnRenderThread() != 0)
 		{
 			FilmGrainConstantsBuffer = GraphBuilder.RegisterExternalBuffer(View.ViewState->FilmGrainCache.ConstantsBuffer, TEXT("FilmGrain.TextureConstants"));
 		}
@@ -737,7 +736,6 @@ FScreenPassTexture AddTonemapPass(FRDGBuilder& GraphBuilder, const FViewInfo& Vi
 
 			if (View.ViewState)
 			{
-				View.ViewState->FilmGrainCache.Texture = View.FinalPostProcessSettings.FilmGrainTexture;
 				View.ViewState->FilmGrainCache.TextureRHI = FilmGrainTextureRHI;
 				GraphBuilder.QueueBufferExtraction(FilmGrainConstantsBuffer, &View.ViewState->FilmGrainCache.ConstantsBuffer);
 			}
diff --git a/Engine/Source/Runtime/Renderer/Private/ScenePrivate.h b/Engine/Source/Runtime/Renderer/Private/ScenePrivate.h
--- a/Engine/Source/Runtime/Renderer/Private/ScenePrivate.h
+++ b/Engine/Source/Runtime/Renderer/Private/ScenePrivate.h
@@ -1023,7 +1023,6 @@ public:
 		/// @cond DOXYGEN_WARNINGS
 		void SafeRelease()
 		{
-			Physical = nullptr;
 			PhysicalRHI = nullptr;
 			Spectral.SafeRelease();
 			ConstantsBuffer.SafeRelease();
@@ -1035,7 +1034,6 @@ public:
 		TRefCountPtr<FRDGPooledBuffer> ConstantsBuffer;
 
 		// The physical space source texture
-		TObjectPtr<UTexture2D> Physical = nullptr;
 		FRHITexture* PhysicalRHI = nullptr;
 
 		// The Scale * 100 = percentage of the image space that the physical kernel represents.
@@ -1056,13 +1054,11 @@ public:
 		/// @cond DOXYGEN_WARNINGS
 		void SafeRelease()
 		{
-			Texture = nullptr;
 			TextureRHI = nullptr;
 			ConstantsBuffer.SafeRelease();
 		}
 		/// @endcond
-
-		TObjectPtr<UTexture2D> Texture = nullptr;
+		// Cache View.FilmGrainTexture in TextureRHI and compare it every frame in case a new one is set (and rebuild constant buffers if this is the case)
 		FRHITexture* TextureRHI = nullptr;
 		TRefCountPtr<FRDGPooledBuffer> ConstantsBuffer;
 	} FilmGrainCache;
@@ -1448,28 +1444,7 @@ PRAGMA_ENABLE_DEPRECATION_WARNINGS
 
 	virtual void AddReferencedObjects(FReferenceCollector& Collector) override
 	{
-
 		Collector.AddReferencedObjects(MIDPool);
-
-		if (BloomFFTKernel.Physical)
-		{
-			Collector.AddReferencedObject(BloomFFTKernel.Physical);
-		}
-
-		if (FilmGrainCache.Texture)
-		{
-			Collector.AddReferencedObject(FilmGrainCache.Texture);
-		}
-	}
-
-	virtual void CleanUpTransientReferences()
-	{
-		ENQUEUE_RENDER_COMMAND(Cleanup)(
-			[this](FRHICommandList& RHICmdList)
-			{
-				BloomFFTKernel.Physical = nullptr;
-				FilmGrainCache.Texture = nullptr;
-			});
 	}
 
 	// needed for GetReusableMID()
