Back out  38375014 due to crashes
#tests tested on console. Doesn't crash anymore

diff --git a/Engine/Source/Runtime/Renderer/Private/MobileShadingRenderer.cpp b/Engine/Source/Runtime/Renderer/Private/MobileShadingRenderer.cpp
--- a/Engine/Source/Runtime/Renderer/Private/MobileShadingRenderer.cpp
+++ b/Engine/Source/Runtime/Renderer/Private/MobileShadingRenderer.cpp
@@ -880,7 +880,6 @@ void FMobileSceneRenderer::RenderFullDepthPrepass(FRDGBuilder& GraphBuilder, TAr
 	
 	if (DoOcclusionQueries())
 	{
-		EndOcclusionScope(GraphBuilder, Views);
 		FenceOcclusionTests(GraphBuilder);
 	}
 }
@@ -1423,10 +1422,13 @@ void FMobileSceneRenderer::Render(FRDGBuilder& GraphBuilder)
 			RenderForward(GraphBuilder, ViewFamilyTexture, SceneTextures, DBufferTextures);
 		}
 
-		if (!bIsFullDepthPrepassEnabled && DoOcclusionQueries())
+		if (DoOcclusionQueries())
 		{
 			EndOcclusionScope(GraphBuilder, Views);
-			FenceOcclusionTests(GraphBuilder);
+		    if (!bIsFullDepthPrepassEnabled)
+		    {
+			    FenceOcclusionTests(GraphBuilder);
+		    }
 		}
 
 		SceneTextures.MobileSetupMode = EMobileSceneTextureSetupMode::All;
@@ -2404,11 +2406,6 @@ int32 FMobileSceneRenderer::ComputeNumOcclusionQueriesToBatch() const
 // Whether we need a separate render-passes for translucency, decals etc
 bool FMobileSceneRenderer::RequiresMultiPass(int32 NumMSAASamples, EShaderPlatform ShaderPlatform)
 {
-	if (!MobileAllowFramebufferFetch(ShaderPlatform))
-	{
-		return true;
-	}
-	
 	// Vulkan uses subpasses
 	if (IsVulkanPlatform(ShaderPlatform))
 	{
