Heterogeneous Volumes: Prevent crash when Lumen is off and Indirect Mode is set to LightingCache.

diff --git a/Engine/Shaders/Private/HeterogeneousVolumes/HeterogeneousVolumesRayMarchingUtils.ush b/Engine/Shaders/Private/HeterogeneousVolumes/HeterogeneousVolumesRayMarchingUtils.ush
--- a/Engine/Shaders/Private/HeterogeneousVolumes/HeterogeneousVolumesRayMarchingUtils.ush
+++ b/Engine/Shaders/Private/HeterogeneousVolumes/HeterogeneousVolumesRayMarchingUtils.ush
@@ -1,5 +1,8 @@
 // Copyright Epic Games, Inc. All Rights Reserved.
 
+// Change this to force shader recompilation
+#pragma message("UESHADERMETADATA_VERSION 498251B5-BC1F-4250-A0A2-7BB218050A9A")
+
 #include "HeterogeneousVolumesRayMarchingTypes.ush"
 #include "HeterogeneousVolumesTracingUtils.ush"
 #include "HeterogeneousVolumesTransmittanceVolumeUtils.ush"
@@ -55,7 +58,6 @@
 #include "HeterogeneousVolumesAdaptiveVolumetricShadowMapSampler.ush"
 #endif // DIM_USE_ADAPTIVE_VOLUMETRIC_SHADOW_MAP
 
-
 #if FOG_INSCATTERING_MODE
 #include "../HeightFogCommon.ush"
 #endif // FOG_INSCATTERING_MODE
diff --git a/Engine/Source/Runtime/Renderer/Private/HeterogeneousVolumes/HeterogeneousVolumesLiveShadingPipeline.cpp b/Engine/Source/Runtime/Renderer/Private/HeterogeneousVolumes/HeterogeneousVolumesLiveShadingPipeline.cpp
--- a/Engine/Source/Runtime/Renderer/Private/HeterogeneousVolumes/HeterogeneousVolumesLiveShadingPipeline.cpp
+++ b/Engine/Source/Runtime/Renderer/Private/HeterogeneousVolumes/HeterogeneousVolumesLiveShadingPipeline.cpp
@@ -888,8 +888,13 @@ static void RenderLightingCacheWithLiveShading(
 		GroupCount.Z = FMath::DivideAndRoundUp(GroupCount.Z, FRenderLightingCacheWithLiveShadingCS::GetThreadGroupSize3D());
 
 		bool bUseAVSM = HeterogeneousVolumes::UseAdaptiveVolumetricShadowMapForSelfShadowing(HeterogeneousVolumeInterface->GetPrimitiveSceneProxy());
-		// Indirect lighting accumulation is coupled with directional light, because it doesn't voxel cull. It is assumed to exist and shadow.
-		int32 IndirectLightingMode = LightType == LightType_Directional ? static_cast<int32>(HeterogeneousVolumes::GetIndirectLightingMode()) : 0;
+
+		int32 IndirectLightingMode = 0;
+		if ((View.GetLumenTranslucencyGIVolume().Texture0 != nullptr) &&
+			(LightType == LightType_Directional))
+		{
+			IndirectLightingMode = static_cast<int32>(HeterogeneousVolumes::GetIndirectLightingMode());
+		}
 
 		FRenderLightingCacheWithLiveShadingCS::FPermutationDomain PermutationVector;
 		PermutationVector.Set<FRenderLightingCacheWithLiveShadingCS::FLightingCacheMode>(HeterogeneousVolumes::GetLightingCacheMode() - 1);
