Add support for PSO precaching of dynamic raytracing geometry compute shaders

diff --git a/Engine/Source/Runtime/Renderer/Private/RayTracing/RayTracingDynamicGeometry.cpp b/Engine/Source/Runtime/Renderer/Private/RayTracing/RayTracingDynamicGeometry.cpp
--- a/Engine/Source/Runtime/Renderer/Private/RayTracing/RayTracingDynamicGeometry.cpp
+++ b/Engine/Source/Runtime/Renderer/Private/RayTracing/RayTracingDynamicGeometry.cpp
@@ -7,6 +7,8 @@
 #include "RayTracingInstance.h"
 #include "RayTracingGeometry.h"
 #include "RenderGraphBuilder.h"
+#include "PSOPrecacheMaterial.h"
+#include "PSOPrecacheValidation.h"
 
 #if RHI_RAYTRACING
 
@@ -115,6 +117,75 @@ public:
 
 IMPLEMENT_MATERIAL_SHADER_TYPE(, FRayTracingDynamicGeometryConverterCS, TEXT("/Engine/Private/RayTracing/RayTracingDynamicMesh.usf"), TEXT("RayTracingDynamicGeometryConverterCS"), SF_Compute);
 
+static const TCHAR* RayTracingDynamicGeometryPSOCollectorName = TEXT("RayTracingDynamicGeometry");
+
+class FRayTracingDynamicGeometryPSOCollector : public IPSOCollector
+{
+public:
+	FRayTracingDynamicGeometryPSOCollector(ERHIFeatureLevel::Type InFeatureLevel) : 
+		IPSOCollector(FPSOCollectorCreateManager::GetIndex(GetFeatureLevelShadingPath(InFeatureLevel), RayTracingDynamicGeometryPSOCollectorName)),
+		FeatureLevel(InFeatureLevel)
+	{
+	}
+
+	virtual void CollectPSOInitializers(
+		const FSceneTexturesConfig& SceneTexturesConfig,
+		const FMaterial& Material,
+		const FPSOPrecacheVertexFactoryData& VertexFactoryData,
+		const FPSOPrecacheParams& PreCacheParams,
+		TArray<FPSOPrecacheData>& PSOInitializers
+	) override final;
+
+private:
+
+	ERHIFeatureLevel::Type FeatureLevel;
+};
+
+
+void FRayTracingDynamicGeometryPSOCollector::CollectPSOInitializers(
+	const FSceneTexturesConfig& SceneTexturesConfig,
+	const FMaterial& Material,
+	const FPSOPrecacheVertexFactoryData& VertexFactoryData,
+	const FPSOPrecacheParams& PreCacheParams,
+	TArray<FPSOPrecacheData>& PSOInitializers
+)
+{
+	if (!VertexFactoryData.VertexFactoryType->SupportsRayTracingDynamicGeometry())
+	{
+		return;
+	}
+
+	FMaterialShaderTypes ShaderTypes;
+	ShaderTypes.AddShaderType<FRayTracingDynamicGeometryConverterCS>();
+
+	FMaterialShaders MaterialShaders;
+	if (!Material.TryGetShaders(ShaderTypes, VertexFactoryData.VertexFactoryType, MaterialShaders))
+	{
+		return;
+	}
+
+	TShaderRef<FRayTracingDynamicGeometryConverterCS> Shader;
+	if (!MaterialShaders.TryGetShader(SF_Compute, Shader))
+	{
+		return;
+	}
+
+	FPSOPrecacheData RTPrecacheData;
+	RTPrecacheData.Type = FPSOPrecacheData::EType::Compute;
+	RTPrecacheData.SetComputeShader(Shader);
+#if PSO_PRECACHING_VALIDATE
+	RTPrecacheData.PSOCollectorIndex = PSOCollectorIndex;
+	RTPrecacheData.VertexFactoryType = VertexFactoryData.VertexFactoryType;
+#endif // PSO_PRECACHING_VALIDATE
+	PSOInitializers.Add(MoveTemp(RTPrecacheData));
+}
+
+IPSOCollector* CreateRayTracingDynamicGeometryPSOCollector(ERHIFeatureLevel::Type FeatureLevel)
+{
+	return new FRayTracingDynamicGeometryPSOCollector(FeatureLevel);
+}
+FRegisterPSOCollectorCreateFunction RegisterRayTracingDynamicGeometryPSOCollector(&CreateRayTracingDynamicGeometryPSOCollector, EShadingPath::Deferred, RayTracingDynamicGeometryPSOCollectorName);
+
 FRayTracingDynamicGeometryCollection::FRayTracingDynamicGeometryCollection() 
 {
 }
@@ -247,6 +318,8 @@ void FRayTracingDynamicGeometryCollection::AddDynamicMeshBatchForGeometryUpdate(
 	FRayTracingDynamicGeometryBuildParams GeometryBuildParams;
 	GeometryBuildParams.DispatchCommands.Reserve(UpdateParams.MeshBatches.Num());
 
+	const int32 PSOCollectorIndex = FPSOCollectorCreateManager::GetIndex(EShadingPath::Deferred, RayTracingDynamicGeometryPSOCollectorName);
+
 	for (const FMeshBatch& MeshBatch : UpdateParams.MeshBatches)
 	{
 		if (!ensureMsgf(MeshBatch.VertexFactory->GetType()->SupportsRayTracingDynamicGeometry(),
@@ -277,7 +350,7 @@ void FRayTracingDynamicGeometryCollection::AddDynamicMeshBatchForGeometryUpdate(
 				{
 					TShaderRef<FRayTracingDynamicGeometryConverterCS> Shader;
 					MaterialShaders.TryGetShader(SF_Compute, Shader);
-
+							
 					FMeshProcessorShaders MeshProcessorShaders;
 					MeshProcessorShaders.ComputeShader = Shader;
 
@@ -326,6 +399,15 @@ void FRayTracingDynamicGeometryCollection::AddDynamicMeshBatchForGeometryUpdate(
 #if MESH_DRAW_COMMAND_DEBUG_DATA
 					ShaderBindings.Finalize(&MeshProcessorShaders);
 #endif
+								
+					#if PSO_PRECACHING_VALIDATE
+					FRHIComputeShader* ComputeShader = DispatchCmd.MaterialShader.GetComputeShader();
+					if (ComputeShader != nullptr)
+					{
+						EPSOPrecacheResult PSOPrecacheResult = PipelineStateCache::CheckPipelineStateInCache(ComputeShader);
+						PSOCollectorStats::CheckComputePipelineStateInCache(*ComputeShader, PSOPrecacheResult, &MaterialRenderProxy, PSOCollectorIndex);
+					}
+					#endif
 
 					GeometryBuildParams.DispatchCommands.Add(DispatchCmd);
 
