Volumetric Cloud - Fixed mistake where shadow sample scale count was rounded to an integervalue.

#tests PC dx11, engine test volumetriccloud functional tests passed.
#rnx

diff --git a/Engine/Source/Runtime/Renderer/Private/VolumetricCloudRendering.cpp b/Engine/Source/Runtime/Renderer/Private/VolumetricCloudRendering.cpp
--- a/Engine/Source/Runtime/Renderer/Private/VolumetricCloudRendering.cpp
+++ b/Engine/Source/Runtime/Renderer/Private/VolumetricCloudRendering.cpp
@@ -389,13 +389,13 @@ static int32 GetVolumetricCloudShadowMapResolution(const FLightSceneProxy* Atmos
 	return 32;
 }
 
-static int32 GetVolumetricCloudShadowRaySampleCountScale(const FLightSceneProxy* AtmosphericLight)
+static float GetVolumetricCloudShadowRaySampleCountScale(const FLightSceneProxy* AtmosphericLight)
 {
 	if (AtmosphericLight)
 	{
 		return AtmosphericLight->GetCloudShadowRaySampleCountScale();
 	}
-	return 4;
+	return 4.0f;
 }
 
 static float GetVolumetricCloudShadowMapExtentKm(const FLightSceneProxy* AtmosphericLight)
