Reduce occlusuion latency on mobile platforms with depth pre-pass enabled
Resubmitting changes from CL# 38407897

diff --git a/Engine/Source/Runtime/Renderer/Private/MobileShadingRenderer.cpp b/Engine/Source/Runtime/Renderer/Private/MobileShadingRenderer.cpp
--- a/Engine/Source/Runtime/Renderer/Private/MobileShadingRenderer.cpp
+++ b/Engine/Source/Runtime/Renderer/Private/MobileShadingRenderer.cpp
@@ -877,9 +877,10 @@ void FMobileSceneRenderer::RenderFullDepthPrepass(FRDGBuilder& GraphBuilder, TAr
 				});
 		}
 	}
-	
-	if (DoOcclusionQueries())
+
+	if (DoOcclusionQueries() && !bIsSceneCaptureRenderPass)
 	{
+		EndOcclusionScope(GraphBuilder, Views);
 		FenceOcclusionTests(GraphBuilder);
 	}
 }
@@ -1422,13 +1423,10 @@ void FMobileSceneRenderer::Render(FRDGBuilder& GraphBuilder)
 			RenderForward(GraphBuilder, ViewFamilyTexture, SceneTextures, DBufferTextures);
 		}
 
-		if (DoOcclusionQueries())
+		if (DoOcclusionQueries() && !bIsFullDepthPrepassEnabled)
 		{
 			EndOcclusionScope(GraphBuilder, Views);
-		    if (!bIsFullDepthPrepassEnabled)
-		    {
-			    FenceOcclusionTests(GraphBuilder);
-		    }
+		    FenceOcclusionTests(GraphBuilder);
 		}
 
 		SceneTextures.MobileSetupMode = EMobileSceneTextureSetupMode::All;
