Changed the shadow fade out to depend on the shadow resolution scale to give greater control over how individual lights fade out their shadow, the old behavior can be enabled by turning off "r.Shadow.DoesFadeUseResolutionScale".

diff --git a/Engine/Source/Runtime/Renderer/Private/ShadowSetup.cpp b/Engine/Source/Runtime/Renderer/Private/ShadowSetup.cpp
--- a/Engine/Source/Runtime/Renderer/Private/ShadowSetup.cpp
+++ b/Engine/Source/Runtime/Renderer/Private/ShadowSetup.cpp
@@ -414,6 +414,15 @@ static TAutoConsoleVariable<bool> CVarCSMScissorOptim(
 	ECVF_RenderThreadSafe
 );
 
+static bool GDoesShadowFadeUseResolutionScale = 1;
+FAutoConsoleVariableRef CVarComputeShadowFadeUsingResolutionScale(
+	TEXT("r.Shadow.DoesFadeUseResolutionScale"),
+	GDoesShadowFadeUseResolutionScale,
+	TEXT("When enabled (default), the shadow resolution scale for the individual light feeds into the fade out calculation, increasing the control of shadow fade out.\n"
+		 "  Setting the cvar to 0 restores the behavior to pre UE 5.6."),
+	ECVF_RenderThreadSafe
+	);
+
 /**
  * Helper function to maintain the common login of caching the whole scene shadow and be able to handle the specific-purpose through the CacheLambda and UnCacheLambda
  */
@@ -3243,6 +3252,28 @@ void FSceneRenderer::SetupInteractionShadows(
 	}
 }
 
+static float ClampShadowRolution(float UnclampedResolution, float ShadowResolutionScale, float MaxShadowResolution)
+{
+	float ClampedResolution = UnclampedResolution;
+
+	if (ShadowResolutionScale > 1.0f)
+	{
+		// Apply ShadowResolutionScale before the MaxShadowResolution clamp if raising the resolution
+		ClampedResolution *= ShadowResolutionScale;
+	}
+
+	ClampedResolution = FMath::Min<float>(ClampedResolution, MaxShadowResolution);
+
+	if (ShadowResolutionScale <= 1.0f)
+	{
+		// Apply ShadowResolutionScale after the MaxShadowResolution clamp if lowering the resolution
+		// Artists want to modify the softness of the shadow with ShadowResolutionScale
+		ClampedResolution *= ShadowResolutionScale;
+	}
+	return ClampedResolution;
+}
+
+
 void FSceneRenderer::CreatePerObjectProjectedShadow(
 	FDynamicShadowsTaskData& TaskData,
 	FLightPrimitiveInteraction* Interaction, 
@@ -3404,34 +3435,16 @@ void FSceneRenderer::CreatePerObjectProjectedShadow(
 
 		// Determine the amount of shadow buffer resolution needed for this view.
 		const float UnclampedResolution = ScreenRadius * CVarShadowTexelsPerPixel.GetValueOnRenderThread();
+		const float ClampedResolution = ClampShadowRolution(UnclampedResolution, LightSceneInfo->Proxy->GetShadowResolutionScale(), MaxShadowResolution);
 
-		// Calculate fading based on resolution
-		// Compute FadeAlpha before ShadowResolutionScale contribution (artists want to modify the softness of the shadow, not change the fade ranges)
-		const float ViewSpecificAlpha = CalculateShadowFadeAlpha(UnclampedResolution, ShadowFadeResolution, MinShadowResolution) * LightSceneInfo->Proxy->GetShadowAmount();
-		MaxResolutionFadeAlpha = FMath::Max(MaxResolutionFadeAlpha, ViewSpecificAlpha);
-		ResolutionFadeAlphas.Add(ViewSpecificAlpha);
-
-		const float ViewSpecificPreShadowAlpha = CalculateShadowFadeAlpha(UnclampedResolution * CVarPreShadowResolutionFactor.GetValueOnRenderThread(), PreShadowFadeResolution, MinPreShadowResolution) * LightSceneInfo->Proxy->GetShadowAmount();
-		MaxResolutionPreShadowFadeAlpha = FMath::Max(MaxResolutionPreShadowFadeAlpha, ViewSpecificPreShadowAlpha);
-		ResolutionPreShadowFadeAlphas.Add(ViewSpecificPreShadowAlpha);
-
-		const float ShadowResolutionScale = LightSceneInfo->Proxy->GetShadowResolutionScale();
-
-		float ClampedResolution = UnclampedResolution;
-
-		if (ShadowResolutionScale > 1.0f)
-		{
-			// Apply ShadowResolutionScale before the MaxShadowResolution clamp if raising the resolution
-			ClampedResolution *= ShadowResolutionScale;
-		}
-
-		ClampedResolution = FMath::Min<float>(ClampedResolution, MaxShadowResolution);
-
-		if (ShadowResolutionScale <= 1.0f)
 		{
-			// Apply ShadowResolutionScale after the MaxShadowResolution clamp if lowering the resolution
-			// Artists want to modify the softness of the shadow with ShadowResolutionScale
-			ClampedResolution *= ShadowResolutionScale;
+			const float FadeCalcResolution = GDoesShadowFadeUseResolutionScale ? ClampedResolution : UnclampedResolution;
+			const float ViewSpecificAlpha = CalculateShadowFadeAlpha(FadeCalcResolution, ShadowFadeResolution, MinShadowResolution) * LightSceneInfo->Proxy->GetShadowAmount();
+			MaxResolutionFadeAlpha = FMath::Max(MaxResolutionFadeAlpha, ViewSpecificAlpha);
+			ResolutionFadeAlphas.Add(ViewSpecificAlpha);
+			const float ViewSpecificPreShadowAlpha = CalculateShadowFadeAlpha(FadeCalcResolution * CVarPreShadowResolutionFactor.GetValueOnRenderThread(), PreShadowFadeResolution, MinPreShadowResolution) * LightSceneInfo->Proxy->GetShadowAmount();
+			MaxResolutionPreShadowFadeAlpha = FMath::Max(MaxResolutionPreShadowFadeAlpha, ViewSpecificPreShadowAlpha);
+			ResolutionPreShadowFadeAlphas.Add(ViewSpecificPreShadowAlpha);
 		}
 
 		MaxDesiredResolution = FMath::Max(
@@ -4156,29 +4169,17 @@ void FSceneRenderer::CreateWholeSceneProjectedShadow(
 					*LightSceneInfo->Proxy->GetOwnerNameOrLabel());
 			}
 
-			// Compute FadeAlpha before ShadowResolutionScale contribution (artists want to modify the softness of the shadow, not change the fade ranges)
-			const float FadeAlpha = CalculateShadowFadeAlpha( UnclampedResolution, ShadowFadeResolution, MinShadowResolution ) * LightSceneInfo->Proxy->GetShadowAmount();
-			MaxFadeAlpha = FMath::Max(MaxFadeAlpha, FadeAlpha);
-			FadeAlphas.Add(FadeAlpha);
+			const float ClampedResolution = ClampShadowRolution(UnclampedResolution, LightSceneInfo->Proxy->GetShadowResolutionScale(), MaxShadowResolution);
 
-			const float ShadowResolutionScale = LightSceneInfo->Proxy->GetShadowResolutionScale();
-
-			float ClampedResolution = UnclampedResolution;
-
-			if (ShadowResolutionScale > 1.0f)
 			{
-				// Apply ShadowResolutionScale before the MaxShadowResolution clamp if raising the resolution
-				ClampedResolution *= ShadowResolutionScale;
-			}
-
-			ClampedResolution = FMath::Min<float>(ClampedResolution, MaxShadowResolution);
-
-			if (ShadowResolutionScale <= 1.0f)
-			{
-				// Apply ShadowResolutionScale after the MaxShadowResolution clamp if lowering the resolution
-				// Artists want to modify the softness of the shadow with ShadowResolutionScale
-				ClampedResolution *= ShadowResolutionScale;
+				// pick resolution to use for fade calcs.
+				// Compute FadeAlpha before or after ShadowResolutionScale contribution (some artists may want to modify the softness of the shadow, not change the fade ranges, others do)		
+				const float FadeCalcResolution = GDoesShadowFadeUseResolutionScale ? ClampedResolution : UnclampedResolution;
+				const float FadeAlpha = CalculateShadowFadeAlpha( FadeCalcResolution, ShadowFadeResolution, MinShadowResolution ) * LightSceneInfo->Proxy->GetShadowAmount();
+				MaxFadeAlpha = FMath::Max(MaxFadeAlpha, FadeAlpha);
+				FadeAlphas.Add(FadeAlpha);
 			}
+		
 
 			MaxDesiredResolution = FMath::Max(
 				MaxDesiredResolution,
