Thin geometry temporal stability improvement (majorly when the background is volumetric sky).

Add a dedicated pass TSRThinGeometry in TSR to handle subpixel geometry. Detect edge by depth and/or luma. Use r.TSR.ThinGeometryDetection to turn on/off.

diff --git a/Engine/Shaders/Private/TemporalSuperResolution/TSRDetectThinGeometry.usf b/Engine/Shaders/Private/TemporalSuperResolution/TSRDetectThinGeometry.usf
new file mode 100644
--- /dev/null
+++ b/Engine/Shaders/Private/TemporalSuperResolution/TSRDetectThinGeometry.usf
@@ -0,0 +1,299 @@
+// Copyright Epic Games, Inc. All Rights Reserved.
+
+
+//------------------------------------------------------- INCLUDES
+#include "TSRColorSpace.ush"
+#include "TSRDepthVelocityAnalysis.ush"
+#include "TSRConvolutionNetworkPass.ush"
+
+//------------------------------------------------------- CONFIG
+
+#define DEBUG_ARRAY_SIZE 8
+
+#define DEBUG_DISABLE_PADDING 0
+
+#define DEBUG_DIGESTABLE 0
+
+#define EDGE_HORIZONTAL_VERTICAL 1
+#define EDGE_HORIZONTAL_VERTICAL_DIAGNAL 2
+
+#define EDGE_TYPE EDGE_HORIZONTAL_VERTICAL_DIAGNAL
+
+//------------------------------------------------------- CONSTANTS
+
+#if CONFIG_LDS_DESPILL
+	static const uint kDespillInput = LDS_DESPILL_THREAD_COUNT * LDS_DESPILL_OFFSET * 0;
+	static const uint kDespillHistory = LDS_DESPILL_THREAD_COUNT * LDS_DESPILL_OFFSET * 1;
+#endif
+
+//------------------------------------------------------- PARAMETERS
+
+int TileOverscan;
+int ThinGeometryTextureIndex;
+float ErrorMultiplier;
+
+Texture2D<float> SceneDepthTexture;
+RWTexture2DArray<uint> R8Output;
+
+//------------------------------------------------------- GLOBALS
+static uint2 GGroupId = 0;
+
+//------------------------------------------------------- DEBUG
+
+static tsr_tensor_half4 Debug[DEBUG_ARRAY_SIZE];
+
+//------------------------------------------------------- INPUT FETCHES
+
+tsr_short2 ComputeElementPixelPos(const uint ElementIndex)
+{
+	tsr_short2 ClampedFetchPixelPos = ComputeElementInputPixelPos(
+		/* LaneStride = */ uint2(LANE_STRIDE_X, LANE_STRIDE_Y),
+		TileOverscan,
+		InputInfo_ViewportMin,
+		InputInfo_ViewportMax,
+		GGroupId,
+		GGroupThreadIndex,
+		ElementIndex);
+
+	return ClampedFetchPixelPos;
+}
+
+tsr_tensor_short2 ComputePixelPos()
+{
+	tsr_tensor_short2 PixelPos;
+
+	UNROLL_N(SIMD_SIZE)
+	for (uint ElementIndex = 0; ElementIndex < SIMD_SIZE; ElementIndex++)
+	{
+		PixelPos.SetElement(ElementIndex, ComputeElementPixelPos(ElementIndex));
+	}
+
+	return PixelPos;
+}
+
+tsr_short2 ComputeElementOutputPixelPos(const uint ElementIndex)
+{
+	return ComputeElementOutputPixelPos(
+		/* LaneStride = */ uint2(LANE_STRIDE_X, LANE_STRIDE_Y),
+		TileOverscan,
+		InputInfo_ViewportMin,
+		InputInfo_ViewportMax,
+		GGroupId,
+		GGroupThreadIndex,
+		ElementIndex);
+}
+
+void FetchInput(
+	tsr_tensor_short2 PixelPos,
+	out tsr_tensor_half SceneDepth,
+	out tsr_tensor_half DeviceZThinGeometryDetectionError)
+{
+	tsr_tensor_half LumaAbsDelta;
+
+	UNROLL_N(SIMD_SIZE)
+	for (uint ElementIndex = 0; ElementIndex < SIMD_SIZE; ElementIndex++)
+	{
+		tsr_short2 ClampedFetchPixelPos = PixelPos.GetElement(ElementIndex);
+
+		float Depth = SceneDepthTexture[ClampedFetchPixelPos];
+		float DeviceZError = ComputePixelDeviceZError(Depth);
+
+		SceneDepth.SetElement(ElementIndex, tsr_half(Depth));
+		DeviceZThinGeometryDetectionError.SetElement(ElementIndex, tsr_half(ErrorMultiplier * DeviceZError));
+	}
+
+	//Debug[0].SetComponent(1, LumaAbsDelta);
+
+} // FetchInput
+
+
+//------------------------------------------------------- FUNCTIONS
+
+// Horizontal thin edge detection.
+// 0 0 0
+// x x x
+// 0 0 0
+
+CALL_SITE_DEBUGLOC
+template<
+typename ScalarType, uint VectorSize, uint LaneStrideX, uint LaneStrideY>
+TLaneVector2D<ScalarType, VectorSize, LaneStrideX, LaneStrideY> HorizontalThinEdge3x3(
+	TLaneVector2D<ScalarType, VectorSize, LaneStrideX, LaneStrideY> Center,
+	TLaneVector2D<ScalarType, VectorSize, LaneStrideX, LaneStrideY> EdgeThreshold)
+{
+
+	TLaneVector2D<ScalarType, VectorSize, LaneStrideX, LaneStrideY> Top;
+	TLaneVector2D<ScalarType, VectorSize, LaneStrideX, LaneStrideY> Down;
+	
+	const TLaneVector2D<ScalarType, VectorSize, LaneStrideX, LaneStrideY> One = TLaneVector2D<ScalarType, VectorSize, LaneStrideX, LaneStrideY>::Const(1.0);
+	const TLaneVector2D<ScalarType, VectorSize, LaneStrideX, LaneStrideY> Zero = TLaneVector2D<ScalarType, VectorSize, LaneStrideX, LaneStrideY>::Const(0);
+	
+
+	AccessNeighborTexels1x3(Center, Top, Down);
+	TLaneVector2D<ScalarType, VectorSize, LaneStrideX, LaneStrideY> TopThinEdge = select(Center - Top > EdgeThreshold, One, Zero);
+	TLaneVector2D<ScalarType, VectorSize, LaneStrideX, LaneStrideY> DownThinEdge = select(Center - Down > EdgeThreshold, One,Zero);
+	
+	TLaneVector2D<ScalarType, VectorSize, LaneStrideX, LaneStrideY> PotentialThinGeometryEdge = TopThinEdge + DownThinEdge;
+	
+
+	TLaneVector2D<ScalarType, VectorSize, LaneStrideX, LaneStrideY> Left;
+	TLaneVector2D<ScalarType, VectorSize, LaneStrideX, LaneStrideY> Right;
+	AccessNeighborTexels3x1(PotentialThinGeometryEdge, Left, Right);
+#if EDGE_TYPE == EDGE_HORIZONTAL_VERTICAL
+	const TLaneVector2D<ScalarType, VectorSize, LaneStrideX, LaneStrideY> Five = TLaneVector2D<ScalarType, VectorSize, LaneStrideX, LaneStrideY>::Const(5.0);
+	TLaneVector2D<ScalarType, VectorSize, LaneStrideX, LaneStrideY>  ThinGeometryEdge = saturate(Left + Right + PotentialThinGeometryEdge - Five);
+#elif EDGE_TYPE == EDGE_HORIZONTAL_VERTICAL_DIAGNAL
+	TLaneVector2D<ScalarType, VectorSize, LaneStrideX, LaneStrideY> LeftRightSum =  Left + Right;
+	LeftRightSum = LeftRightSum * saturate(PotentialThinGeometryEdge - One); // Center pixel should be edge
+
+	const TLaneVector2D<ScalarType, VectorSize, LaneStrideX, LaneStrideY> Two = TLaneVector2D<ScalarType, VectorSize, LaneStrideX, LaneStrideY>::Const(2.0);
+	TLaneVector2D<ScalarType, VectorSize, LaneStrideX, LaneStrideY>  ThinGeometryEdge = saturate((LeftRightSum - Two) * rcp(Two)); // At most one side being edge
+#elif
+#error EDGE_TYPE not implemented
+#endif
+
+	return ThinGeometryEdge;
+}
+
+// Vertical thin edge detection.
+// 0 x 0
+// 0 x 0
+// 0 x 0
+
+CALL_SITE_DEBUGLOC
+template<
+	typename ScalarType, uint VectorSize, uint LaneStrideX, uint LaneStrideY>
+TLaneVector2D<ScalarType, VectorSize, LaneStrideX, LaneStrideY> VerticalThinEdge3x3(
+	TLaneVector2D<ScalarType, VectorSize, LaneStrideX, LaneStrideY> Center,
+	TLaneVector2D<ScalarType, VectorSize, LaneStrideX, LaneStrideY> EdgeThreshold)
+{
+
+	const TLaneVector2D<ScalarType, VectorSize, LaneStrideX, LaneStrideY> One = TLaneVector2D<ScalarType, VectorSize, LaneStrideX, LaneStrideY>::Const(1.0);
+	const TLaneVector2D<ScalarType, VectorSize, LaneStrideX, LaneStrideY> Zero = TLaneVector2D<ScalarType, VectorSize, LaneStrideX, LaneStrideY>::Const(0);
+
+	TLaneVector2D<ScalarType, VectorSize, LaneStrideX, LaneStrideY> Left;
+	TLaneVector2D<ScalarType, VectorSize, LaneStrideX, LaneStrideY> Right;
+	AccessNeighborTexels3x1(Center, Left, Right);
+	TLaneVector2D<ScalarType, VectorSize, LaneStrideX, LaneStrideY> LeftThinEdge  = select(Center -  Left > EdgeThreshold, One, Zero);
+	TLaneVector2D<ScalarType, VectorSize, LaneStrideX, LaneStrideY> RightThinEdge = select(Center - Right > EdgeThreshold, One, Zero);
+	TLaneVector2D<ScalarType, VectorSize, LaneStrideX, LaneStrideY> PotentialThinGeometryEdge = LeftThinEdge + RightThinEdge;
+
+
+	TLaneVector2D<ScalarType, VectorSize, LaneStrideX, LaneStrideY> Top;
+	TLaneVector2D<ScalarType, VectorSize, LaneStrideX, LaneStrideY> Bottom;
+	AccessNeighborTexels1x3(PotentialThinGeometryEdge, Top, Bottom);
+#if EDGE_TYPE == EDGE_HORIZONTAL_VERTICAL
+	const TLaneVector2D<ScalarType, VectorSize, LaneStrideX, LaneStrideY> Five = TLaneVector2D<ScalarType, VectorSize, LaneStrideX, LaneStrideY>::Const(5.0);
+	TLaneVector2D<ScalarType, VectorSize, LaneStrideX, LaneStrideY>  ThinGeometryEdge = saturate(Top + Bottom + PotentialThinGeometryEdge - Five);
+#elif EDGE_TYPE == EDGE_HORIZONTAL_VERTICAL_DIAGNAL
+	TLaneVector2D<ScalarType, VectorSize, LaneStrideX, LaneStrideY> TopBottomSum =  Top + Bottom;
+	TopBottomSum = TopBottomSum * saturate(PotentialThinGeometryEdge - One); // Center pixel should be edge
+
+	const TLaneVector2D<ScalarType, VectorSize, LaneStrideX, LaneStrideY> Two = TLaneVector2D<ScalarType, VectorSize, LaneStrideX, LaneStrideY>::Const(2.0);
+	TLaneVector2D<ScalarType, VectorSize, LaneStrideX, LaneStrideY>  ThinGeometryEdge = saturate((TopBottomSum - Two) * rcp(Two)); // At most one side not being edge
+#elif
+#error EDGE_TYPE not implemented
+#endif
+
+	return ThinGeometryEdge;
+}
+
+//------------------------------------------------------- ENTRY POINT
+
+#if COMPILER_SUPPORTS_WAVE_SIZE && DIM_WAVE_SIZE > 0
+WAVESIZE(DIM_WAVE_SIZE)
+#endif
+[numthreads(LANE_COUNT * WAVE_COUNT, 1, 1)]
+void MainCS(
+	uint2 GroupId : SV_GroupID,
+	uint GroupThreadIndex : SV_GroupIndex)
+{
+	GGroupId = GroupId;
+	GGroupThreadIndex = GroupThreadIndex;
+
+	#if DEBUG_OUTPUT
+	{
+		UNROLL_N(DEBUG_ARRAY_SIZE)
+		for (uint DebugId = 0; DebugId < DEBUG_ARRAY_SIZE; DebugId++)
+		{
+			Debug[DebugId].SetAllElements(0.0);
+		}
+	}
+	#endif
+
+
+	tsr_tensor_short2 PixelPos = ComputePixelPos();
+	tsr_tensor_half DeviceZ;
+	tsr_tensor_half DeviceZThinGeometryDetectionError;
+	FetchInput(PixelPos,
+		DeviceZ, DeviceZThinGeometryDetectionError);
+	
+	//Debug[0].SetComponent(0, DeviceZ);
+	//Debug[0].SetComponent(1, DeviceZThinGeometryDetectionError);
+	
+	tsr_tensor_half ThinGeometryEdge = tsr_tensor_half::Const(0.0);
+
+	{
+		ThinGeometryEdge = HorizontalThinEdge3x3(DeviceZ, DeviceZThinGeometryDetectionError);
+		ThinGeometryEdge = ThinGeometryEdge + VerticalThinEdge3x3(DeviceZ, DeviceZThinGeometryDetectionError);
+		//Debug[0].SetComponent(2, tsr_tensor_half(ThinGeometryEdge));
+	}
+
+	{
+		//dilate the edge with (0.8,1,0.8), such that, dilated edges has 0.8 for direct neighbor, 0.64 for diagnal, 1 for consecutive neighbors.
+		// This value are used to boost the weight with
+		// lerp(blend, 1.0, thinedgeblend)
+		ThinGeometryEdge = saturate(WeightedAvg3x3(ThinGeometryEdge, float3(0.8, 1.0, 0.8)));
+		//Debug[0].SetComponent(3, tsr_tensor_half(ThinGeometryEdge));
+	}
+
+	// Output the thin geometry mask
+	UNROLL_N(SIMD_SIZE)
+	for(uint ElementIndex = 0; ElementIndex < SIMD_SIZE; ElementIndex++)
+	{
+		const tsr_short2 LaneSimdPixelOffset = GetElementPixelOffsetInTile(LANE_STRIDE_X, LANE_STRIDE_Y, GGroupThreadIndex, ElementIndex);
+		const tsr_short2 OutputPixelPos = ComputeElementOutputPixelPos(ElementIndex);
+		
+		tsr_half ThinGeometryFlag = ThinGeometryEdge.GetElement(ElementIndex);
+		R8Output[uint3(OutputPixelPos,ThinGeometryTextureIndex)] = uint(ThinGeometryFlag * tsr_half(255.0));
+	}
+	
+	#if DEBUG_OUTPUT
+	{
+		// Mark the top left corner of the tile for debugging purposes.
+		#if DEBUG_DISABLE_PADDING == 2
+		if (GGroupThreadIndex == 0)
+		{
+			UNROLL_N(DEBUG_ARRAY_SIZE)
+			for (uint DebugId = 0; DebugId < DEBUG_ARRAY_SIZE; DebugId++)
+			{
+				Debug[DebugId].SetElement(/* ElementIndex = */ 0, 1.0);
+			}
+		}
+		#endif // DEBUG_DISABLE_PADDING == 2
+	
+		UNROLL_N(SIMD_SIZE)
+		for (uint ElementIndex = 0; ElementIndex < SIMD_SIZE; ElementIndex++)
+		{
+			const tsr_short2 OutputPixelPos = ComputeElementOutputPixelPos(
+				/* LaneStride = */ uint2(LANE_STRIDE_X, LANE_STRIDE_Y),
+				TileOverscan,
+				InputInfo_ViewportMin,
+				InputInfo_ViewportMax,
+				GGroupId,
+				GGroupThreadIndex,
+				ElementIndex,
+				/* OutputDataOffset = */ int2(0, 0),
+				/* OutputResDivisor = */ 1,
+				/* bDebugDisablePadding = */ DEBUG_DISABLE_PADDING != 0);
+
+			UNROLL_N(DEBUG_ARRAY_SIZE)
+			for (uint DebugId = 0; DebugId < DEBUG_ARRAY_SIZE; DebugId++)
+			{
+				DebugOutput[tsr_short3(OutputPixelPos, DebugId)] = float4(Debug[DebugId].GetElement(ElementIndex));
+			}
+		}
+	}
+	#endif // DEBUG_OUTPUT
+	
+}
diff --git a/Engine/Shaders/Private/TemporalSuperResolution/TSRRejectShading.usf b/Engine/Shaders/Private/TemporalSuperResolution/TSRRejectShading.usf
--- a/Engine/Shaders/Private/TemporalSuperResolution/TSRRejectShading.usf
+++ b/Engine/Shaders/Private/TemporalSuperResolution/TSRRejectShading.usf
@@ -57,6 +57,7 @@ Texture2D<tsr_half4> ReprojectedHistoryMoireTexture;
 #endif
 Texture2D<tsr_half2> DecimateMaskTexture;
 Texture2D<tsr_ushort> IsMovingMaskTexture;
+Texture2D<tsr_ushort> ThinGeometryTexture;
 Texture2D<float2> ClosestDepthTexture;
 
 RWTexture2DArray<tsr_half4> HistoryGuideOutput;
@@ -116,6 +117,19 @@ tsr_short2 ComputeElementOutputPixelPos(const uint ElementIndex)
 		ElementIndex);
 }
 
+void FetchThinGeometryEdge(
+	tsr_tensor_short2 PixelPos,
+	out tsr_tensor_half IsThinGeometryEdge)
+{
+	UNROLL_N(SIMD_SIZE)
+	for (uint ElementIndex = 0; ElementIndex < SIMD_SIZE; ElementIndex++)
+	{
+		tsr_short2 ClampedFetchPixelPos = PixelPos.GetElement(ElementIndex);
+		
+		IsThinGeometryEdge.SetElement(ElementIndex, tsr_half(ThinGeometryTexture[ClampedFetchPixelPos]) * rcp(tsr_half(255.0)));
+	}
+} // FetchMoireInput()
+
 void FetchDecimateMask(
 	tsr_tensor_short2 PixelPos,
 	out tsr_tensor_half VelocityEdge,
@@ -705,6 +719,15 @@ void MainCS(
 	} // if (bComputeResurrection)
 	#endif // CONFIG_HISTORY_RESURRECTION
 	
+	PLATFORM_SPECIFIC_ISOLATE
+	{
+		tsr_tensor_short2 PixelPos = ComputePixelPos();
+		tsr_tensor_half IsThinGeometryEdge;
+		FetchThinGeometryEdge(PixelPos, IsThinGeometryEdge);
+		RejectionBlendFinal = lerp(RejectionBlendFinal, tsr_tensor_half::Const(tsr_half(1.0)), IsThinGeometryEdge);
+		RejectionClampBlend = lerp(RejectionClampBlend, tsr_tensor_half::Const(tsr_half(1.0)), IsThinGeometryEdge);
+	}
+	
 	// Update the guide for next frame.
 	tsr_tensor_halfC FinalGuide;
 	tsr_tensor_half  InputUncertainty;
diff --git a/Engine/Shaders/Private/TemporalSuperResolution/TSRVisualize.usf b/Engine/Shaders/Private/TemporalSuperResolution/TSRVisualize.usf
--- a/Engine/Shaders/Private/TemporalSuperResolution/TSRVisualize.usf
+++ b/Engine/Shaders/Private/TemporalSuperResolution/TSRVisualize.usf
@@ -56,6 +56,7 @@ Texture2D<uint>      ReprojectionBoundaryTexture;
 Texture2D<uint>      ReprojectionJacobianTexture;
 Texture2D<uint>      ReprojectionVectorTexture;
 Texture2D<tsr_ushort> IsMovingMaskTexture;
+Texture2D<tsr_ushort> ThinGeometryTexture;
 Texture2D<tsr_half2> DecimateMaskTexture;
 Texture2D<tsr_half4> HistoryRejectionTexture;
 Texture2D<tsr_half4> MoireHistoryTexture;
@@ -200,6 +201,17 @@ tsr_half3 DisplayHistorySampleCount(FContext Ctx)
 #endif
 }
 
+tsr_half3 DisplayThinGeometry(FContext Ctx)
+{
+	tsr_half ThinGeometryType = saturate(tsr_half(ThinGeometryTexture[Ctx.InputPixelPos] * rcp(255.0f)));
+	ThinGeometryType = (View.FrameNumber % 240) > 120 ? ThinGeometryType : tsr_half(0);
+	tsr_half3 Color = kRed;
+	Color = lerp(kWhite,    Color, saturate((ThinGeometryType - 0.0) * 2.0));
+	Color = select(bCameraCut == 0, Color, kPink);
+
+	return SampleTSROutputColor(Ctx) * Color * tsr_half(View.PreExposure);   
+}
+
 tsr_half3 DisplayHistoryRejection(FContext Ctx)
 {
 	tsr_half LowFrequencyRejection = HistoryRejectionTexture.SampleLevel(GlobalPointClampedSampler, Ctx.InputUV, 0).r;
@@ -555,6 +567,10 @@ void MainCS(
 	{
 		OutputColor = DisplayReprojectionFieldDilatedJacobian(Ctx);
 	}
+	else if (VisualizeId == 15)
+	{
+		OutputColor = DisplayThinGeometry(Ctx);
+	}
 
 	{
 		tsr_short2 OutputPixelPos = InvalidateOutputPixelPos(Ctx.OutputPixelPos, OutputViewRectMax);
diff --git a/Engine/Source/Runtime/Renderer/Private/PostProcess/TemporalAA.h b/Engine/Source/Runtime/Renderer/Private/PostProcess/TemporalAA.h
--- a/Engine/Source/Runtime/Renderer/Private/PostProcess/TemporalAA.h
+++ b/Engine/Source/Runtime/Renderer/Private/PostProcess/TemporalAA.h
@@ -243,6 +243,9 @@ struct FTSRPassConfig
 	float ShadingRejectionFlickeringPeriod = 2.0f;
 	float ShadingRejectionFlickeringMaxParallaxVelocity = 10.0f;
 
+	bool  ThinGeometryDetectionEnable = false;
+	float ThinGeometryErrorMultiplier = 200.0f;
+
 	// Spatial antialiasing
 	int32 RejectionAntiAliasingQuality = 3;
 
diff --git a/Engine/Source/Runtime/Renderer/Private/PostProcess/TemporalSuperResolution.cpp b/Engine/Source/Runtime/Renderer/Private/PostProcess/TemporalSuperResolution.cpp
--- a/Engine/Source/Runtime/Renderer/Private/PostProcess/TemporalSuperResolution.cpp
+++ b/Engine/Source/Runtime/Renderer/Private/PostProcess/TemporalSuperResolution.cpp
@@ -302,6 +302,19 @@ TAutoConsoleVariable<float> CVarTSRWeightClampingPixelSpeed(
 	TEXT("(Default = 1.0f)."),
 	ECVF_RenderThreadSafe);
 
+TAutoConsoleVariable<int32> CVarTSRThinGeometryDetection(
+	TEXT("r.TSR.ThinGeometryDetection"), 0,
+	TEXT("Define if we should perform another pass to detect thin geometies (sub-pixel in frame buffer) either due to sampling algorithm before TSR or the geometry being too thin. ")
+	TEXT("When enabled thin geometry pixels will under-reject if current pixel misses.")
+	TEXT("(Default = 1)."),
+	ECVF_RenderThreadSafe);
+
+TAutoConsoleVariable<float> CVarTSRThinGeometryDetectionErrorMultiplier(
+	TEXT("r.TSR.ThinGeometryDetection.ErrorMultiplier"), 200.0,
+	TEXT("Define the depth difference multiplier between neighbors will be used to detect thin geometry. The larger the upscaler, the smaller the multiplier should be for temporal stability. ")
+	TEXT("(Default = 200)."),
+	ECVF_RenderThreadSafe);
+
 #if !UE_BUILD_OPTIMIZED_SHOWFLAGS
 
 TAutoConsoleVariable<int32> CVarTSRVisualize(
@@ -324,7 +337,8 @@ TAutoConsoleVariable<int32> CVarTSRVisualize(
 	TEXT(" 11: Mask where the reprojection field is anti-aliased from the depth buffer in green (handy to tune r.TSR.ReprojectionField.AntiAliasPixelSpeed);\n")
 	TEXT(" 12: Mask where the pixel's jacobian is null in the reprojection field in orange;\n")
 	TEXT(" 13: Mask where the pixel's jacobian has reached its encoding limit in the reprojection field in red;\n")
-	TEXT(" 14: Mask where the reprojected history is upscaled (in red) or downscaled (in green) by the reprojection field's jacobian (for instance due to getting closer or further away from camera respectively, or an object is getting scaled dynamicaly);\n"),
+	TEXT(" 14: Mask where the reprojected history is upscaled (in red) or downscaled (in green) by the reprojection field's jacobian (for instance due to getting closer or further away from camera respectively, or an object is getting scaled dynamicaly);\n")
+	TEXT(" 15: Mask where thin geometry is detected by depth and luma difference (in red);\n"),
 	ECVF_RenderThreadSafe);
 
 #endif
@@ -814,6 +828,7 @@ class FTSRRejectShadingCS : public FTSRConvolutionNetworkShader
 		SHADER_PARAMETER_RDG_TEXTURE_SRV(Texture2D, ResurrectedHistoryGuideMetadataTexture)
 		SHADER_PARAMETER_RDG_TEXTURE(Texture2D, DecimateMaskTexture)
 		SHADER_PARAMETER_RDG_TEXTURE_SRV(Texture2D, IsMovingMaskTexture)
+		SHADER_PARAMETER_RDG_TEXTURE_SRV(Texture2D, ThinGeometryTexture)
 		SHADER_PARAMETER_RDG_TEXTURE(Texture2D, ClosestDepthTexture)
 
 		SHADER_PARAMETER_RDG_TEXTURE_UAV(RWTexture2DArray, HistoryGuideOutput)
@@ -887,6 +902,73 @@ class FTSRRejectShadingCS : public FTSRConvolutionNetworkShader
 	}
 }; // class FTSRRejectShadingCS
 
+class FTSRDetectThinGeometryCS : public FTSRConvolutionNetworkShader
+{
+	DECLARE_GLOBAL_SHADER(FTSRDetectThinGeometryCS);
+	SHADER_USE_PARAMETER_STRUCT(FTSRDetectThinGeometryCS, FTSRConvolutionNetworkShader);
+
+	using FPermutationDomain = TShaderPermutationDomain<
+		FTSRConvolutionNetworkShader::FPermutationDomain>;
+
+	BEGIN_SHADER_PARAMETER_STRUCT(FParameters, )
+		SHADER_PARAMETER_STRUCT_INCLUDE(FTSRCommonParameters, CommonParameters)
+
+		SHADER_PARAMETER(int32, TileOverscan)
+		SHADER_PARAMETER(int32, ThinGeometryTextureIndex)
+		SHADER_PARAMETER(float, ErrorMultiplier)
+
+		SHADER_PARAMETER_RDG_TEXTURE(Texture2D, SceneDepthTexture)
+		SHADER_PARAMETER_RDG_TEXTURE_UAV(RWTexture2DArray, R8Output)
+		SHADER_PARAMETER_RDG_TEXTURE_UAV(RWTexture2D, DebugOutput)
+	END_SHADER_PARAMETER_STRUCT()
+
+	static FPermutationDomain RemapPermutation(FPermutationDomain PermutationVector)
+	{
+		// Remap redondant convolution permutations.
+		PermutationVector.Set<FTSRConvolutionNetworkShader::FPermutationDomain>(
+			FTSRConvolutionNetworkShader::RemapPermutation(PermutationVector.Get<FTSRConvolutionNetworkShader::FPermutationDomain>()));
+
+		return PermutationVector;
+	}
+
+	static bool ShouldCompilePermutation(const FGlobalShaderPermutationParameters& Parameters)
+	{
+		FPermutationDomain PermutationVector(Parameters.PermutationId);
+		if (PermutationVector != RemapPermutation(PermutationVector))
+		{
+			return false;
+		}
+
+		if (!FTSRConvolutionNetworkShader::ShouldCompilePermutation(Parameters, PermutationVector.Get<FTSRConvolutionNetworkShader::FPermutationDomain>()))
+		{
+			return false;
+		}
+
+		return true;
+	}
+
+	static EShaderPermutationPrecacheRequest ShouldPrecachePermutation(const FGlobalShaderPermutationParameters& Parameters)
+	{
+		FPermutationDomain PermutationVector(Parameters.PermutationId);
+		return FTSRConvolutionNetworkShader::ShouldPrecachePermutation(Parameters, PermutationVector.Get<FTSRConvolutionNetworkShader::FPermutationDomain>());
+	}
+
+	static void ModifyCompilationEnvironment(const FGlobalShaderPermutationParameters& Parameters, FShaderCompilerEnvironment& OutEnvironment)
+	{
+		FPermutationDomain PermutationVector(Parameters.PermutationId);
+
+		FTSRConvolutionNetworkShader::ModifyCompilationEnvironment(
+			Parameters,
+			PermutationVector.Get<FTSRConvolutionNetworkShader::FPermutationDomain>(),
+			OutEnvironment);
+	}
+
+	static EShaderCompileJobPriority GetOverrideJobPriority()
+	{
+		return EShaderCompileJobPriority::High;
+	}
+}; // class FTSRRejectShadingCS
+
 class FTSRSpatialAntiAliasingCS : public FTSRShader
 {
 	DECLARE_GLOBAL_SHADER(FTSRSpatialAntiAliasingCS);
@@ -1145,6 +1227,7 @@ class FTSRVisualizeCS : public FTSRShader
 		SHADER_PARAMETER_RDG_TEXTURE_SRV(Texture2D, ReprojectionJacobianTexture)
 		SHADER_PARAMETER_RDG_TEXTURE_SRV(Texture2D, ReprojectionVectorTexture)
 		SHADER_PARAMETER_RDG_TEXTURE_SRV(Texture2D, IsMovingMaskTexture)
+		SHADER_PARAMETER_RDG_TEXTURE_SRV(Texture2D, ThinGeometryTexture)
 		SHADER_PARAMETER_RDG_TEXTURE(Texture2D, DecimateMaskTexture)
 		SHADER_PARAMETER_RDG_TEXTURE(Texture2D, HistoryRejectionTexture)
 		SHADER_PARAMETER_RDG_TEXTURE_SRV(Texture2D, MoireHistoryTexture)
@@ -1160,6 +1243,7 @@ class FTSRVisualizeCS : public FTSRShader
 IMPLEMENT_GLOBAL_SHADER(FTSRMeasureFlickeringLumaCS, "/Engine/Private/TemporalSuperResolution/TSRMeasureFlickeringLuma.usf", "MainCS", SF_Compute);
 IMPLEMENT_GLOBAL_SHADER(FTSRClearPrevTexturesCS,     "/Engine/Private/TemporalSuperResolution/TSRClearPrevTextures.usf",     "MainCS", SF_Compute);
 IMPLEMENT_GLOBAL_SHADER(FTSRDilateVelocityCS,        "/Engine/Private/TemporalSuperResolution/TSRDilateVelocity.usf",        "MainCS", SF_Compute);
+IMPLEMENT_GLOBAL_SHADER(FTSRDetectThinGeometryCS,	 "/Engine/Private/TemporalSuperResolution/TSRDetectThinGeometry.usf",	 "MainCS", SF_Compute);
 IMPLEMENT_GLOBAL_SHADER(FTSRDecimateHistoryCS,       "/Engine/Private/TemporalSuperResolution/TSRDecimateHistory.usf",       "MainCS", SF_Compute);
 IMPLEMENT_GLOBAL_SHADER(FTSRRejectShadingCS,         "/Engine/Private/TemporalSuperResolution/TSRRejectShading.usf",         "MainCS", SF_Compute);
 IMPLEMENT_GLOBAL_SHADER(FTSRSpatialAntiAliasingCS,   "/Engine/Private/TemporalSuperResolution/TSRSpatialAntiAliasing.usf",   "MainCS", SF_Compute);
@@ -1430,6 +1514,9 @@ FTSRPassConfig GetTSRMainPassConfig(const FViewInfo& View)
 	PassConfig.ShadingRejectionFlickeringPeriod = CVarTSRFlickeringPeriod.GetValueOnRenderThread();
 	PassConfig.ShadingRejectionFlickeringMaxParallaxVelocity = CVarTSRFlickeringMaxParralaxVelocity.GetValueOnRenderThread();
 
+	PassConfig.ThinGeometryDetectionEnable =  static_cast<bool>(CVarTSRThinGeometryDetection.GetValueOnRenderThread());
+	PassConfig.ThinGeometryErrorMultiplier = CVarTSRThinGeometryDetectionErrorMultiplier.GetValueOnRenderThread();
+
 	PassConfig.RejectionAntiAliasingQuality = CVarTSRRejectionAntiAliasingQuality.GetValueOnRenderThread();
 
 	PassConfig.HistoryRejectionSampleCount = CVarTSRHistoryRejectionSampleCount.GetValueOnRenderThread();
@@ -1955,8 +2042,11 @@ FDefaultTemporalUpscaler::FOutputs AddTemporalSuperResolutionPasses(
 	FRDGTextureRef ClosestDepthTexture;
 	FRDGTextureSRVRef DilateMaskTexture;
 	FRDGTextureSRVRef DepthErrorTexture;
+	FRDGTextureSRVRef ThinGeometryTexture = nullptr;
 	FRDGTextureSRVRef IsMovingMaskTexture = nullptr;
+	FRDGTextureRef R8OutputTexture = nullptr;
 	FVelocityFlattenTextures VelocityFlattenTextures;
+	int32 ThinGeometryTextureIndex = 2;
 	{
 		const bool bOutputIsMovingTexture = FlickeringFramePeriod > 0.0f;
 
@@ -1979,14 +2069,13 @@ FDefaultTemporalUpscaler::FOutputs AddTemporalSuperResolutionPasses(
 			ClosestDepthTexture = GraphBuilder.CreateTexture(Desc, TEXT("TSR.ClosestDepthTexture"));
 		}
 
-		FRDGTextureRef R8OutputTexture;
 		{
 			FRDGTextureDesc Desc = FRDGTextureDesc::Create2DArray(
 				InputExtent,
 				PF_R8_UINT,
 				FClearValueBinding::None,
 				/* InFlags = */ TexCreate_ShaderResource | TexCreate_UAV,
-				/* ArraySize = */ bOutputIsMovingTexture ? 3 : 2);
+				/* ArraySize = */ bOutputIsMovingTexture ? 4 : 3);
 
 			R8OutputTexture = GraphBuilder.CreateTexture(Desc, TEXT("TSR.DilateR8"));
 			DilateMaskTexture = GraphBuilder.CreateSRV(FRDGTextureSRVDesc::CreateForSlice(R8OutputTexture, 0));
@@ -1994,6 +2083,12 @@ FDefaultTemporalUpscaler::FOutputs AddTemporalSuperResolutionPasses(
 			if (bOutputIsMovingTexture)
 			{
 				IsMovingMaskTexture = GraphBuilder.CreateSRV(FRDGTextureSRVDesc::CreateForSlice(R8OutputTexture, 2));
+				ThinGeometryTexture = GraphBuilder.CreateSRV(FRDGTextureSRVDesc::CreateForSlice(R8OutputTexture, 3));
+				ThinGeometryTextureIndex = 3;
+			}
+			else
+			{
+				ThinGeometryTexture = GraphBuilder.CreateSRV(FRDGTextureSRVDesc::CreateForSlice(R8OutputTexture, 2));
 			}
 		}
 
@@ -2105,6 +2200,49 @@ FDefaultTemporalUpscaler::FOutputs AddTemporalSuperResolutionPasses(
 			FComputeShaderUtils::GetGroupCount(InputRect.Size(), TileSize));
 	}
 
+	//Create the thin geometry mask to avoid over rejection
+	if(PassConfig.ThinGeometryDetectionEnable)
+	{
+		FTSRConvolutionNetworkShader::FPermutationDomain ConvolutionNetworkPermutationVector;
+		ConvolutionNetworkPermutationVector.Set<FTSRConvolutionNetworkShader::FWaveSizeOps>(SelectWaveSize(View.GetShaderPlatform(), { 16, 32, 64 }));
+		ConvolutionNetworkPermutationVector.Set<FTSRShader::F16BitVALUDim>(bUse16BitVALU);
+		ConvolutionNetworkPermutationVector.Set<FTSRShader::FAlphaChannelDim>(bSupportsAlpha);
+
+		FTSRDetectThinGeometryCS::FPermutationDomain PermutationVector;
+		PermutationVector.Set<FTSRConvolutionNetworkShader::FPermutationDomain>(ConvolutionNetworkPermutationVector);
+		PermutationVector = FTSRDetectThinGeometryCS::RemapPermutation(PermutationVector);
+
+		const int32 GroupTileSize = 32;
+		const int32 TileOverscan = FMath::Clamp(CVarTSRShadingTileOverscan.GetValueOnRenderThread(), 3, GroupTileSize / 2 - 1);
+		const int32 TileSize = GroupTileSize - 2 * TileOverscan;
+
+		FTSRDetectThinGeometryCS::FParameters* PassParameters = GraphBuilder.AllocParameters<FTSRDetectThinGeometryCS::FParameters>();
+		PassParameters->CommonParameters = CommonParameters;
+		PassParameters->TileOverscan = TileOverscan;
+		PassParameters->ThinGeometryTextureIndex = ThinGeometryTextureIndex;
+		PassParameters->ErrorMultiplier = PassConfig.ThinGeometryErrorMultiplier;
+		PassParameters->SceneDepthTexture = PassInputs.SceneDepth.Texture;
+		PassParameters->R8Output = GraphBuilder.CreateUAV(R8OutputTexture);
+		PassParameters->DebugOutput = CreateDebugUAV(InputExtent, TEXT("Debug.TSR.DetectThinGeometry"));
+
+		TShaderMapRef<FTSRDetectThinGeometryCS> ComputeShader(View.ShaderMap, PermutationVector);
+		FComputeShaderUtils::AddPass(
+			GraphBuilder,
+			RDG_EVENT_NAME("TSR DetectThinGeometry(#%d TileSize=%d PaddingCostMultiplier=%1.1f WaveSize=%d VALU=%s%s) %dx%d",
+				PermutationVector.ToDimensionValueId(),
+				TileSize,
+				FMath::Pow(float(GroupTileSize) / float(TileSize), 2),
+				int32(PermutationVector.Get<FTSRConvolutionNetworkShader::FPermutationDomain>().Get<FTSRDetectThinGeometryCS::FWaveSizeOps>()),
+				PermutationVector.Get<FTSRConvolutionNetworkShader::FPermutationDomain>().Get<FTSRShader::F16BitVALUDim>() ? TEXT("16bit") : TEXT("32bit"),
+				PermutationVector.Get<FTSRConvolutionNetworkShader::FPermutationDomain>().Get<FTSRShader::FAlphaChannelDim>() ? TEXT(" AlphaChannel") : TEXT(""),
+				InputRect.Width(), InputRect.Height()),
+			AsyncComputePasses >= 2 ? ERDGPassFlags::AsyncCompute : ERDGPassFlags::Compute,
+			ComputeShader,
+			PassParameters,
+			FComputeShaderUtils::GetGroupCount(InputRect.Size(), TileSize));
+
+	}
+
 	// Decimate input to flicker at same frequency as input.
 	FRDGTextureRef ReprojectedHistoryGuideTexture = nullptr;
 	FRDGTextureRef ReprojectedHistoryMoireTexture = nullptr;
@@ -2366,6 +2504,7 @@ FDefaultTemporalUpscaler::FOutputs AddTemporalSuperResolutionPasses(
 		}
 		PassParameters->DecimateMaskTexture = DecimateMaskTexture;
 		PassParameters->IsMovingMaskTexture = IsMovingMaskTexture ? IsMovingMaskTexture : GraphBuilder.CreateSRV(FRDGTextureSRVDesc(BlackUintDummy));
+		PassParameters->ThinGeometryTexture = PassConfig.ThinGeometryDetectionEnable ? ThinGeometryTexture : GraphBuilder.CreateSRV(FRDGTextureSRVDesc(BlackUintDummy));
 		PassParameters->ClosestDepthTexture = ClosestDepthTexture;
 
 		// Outputs
@@ -2862,6 +3001,7 @@ FDefaultTemporalUpscaler::FOutputs AddTemporalSuperResolutionPasses(
 			ReprojectionFieldNullJacobian = 12,
 			ReprojectionFieldClampedJacobian = 13,
 			ReprojectionFieldDilatedJacobian = 14,
+			ThinGeometry = 15,
 			MAX,
 		};
 
@@ -2881,6 +3021,7 @@ FDefaultTemporalUpscaler::FOutputs AddTemporalSuperResolutionPasses(
 			TEXT("ReprojectionFieldNullJacobian"),
 			TEXT("ReprojectionFieldClampedJacobian"),
 			TEXT("ReprojectionFieldDilatedJacobian"),
+			TEXT("ThinGeometry")
 		};
 		static_assert(UE_ARRAY_COUNT(kVisualizationName) == int32(EVisualizeId::MAX), "kVisualizationName doesn't match EVisualizeId");
 
@@ -2954,6 +3095,7 @@ FDefaultTemporalUpscaler::FOutputs AddTemporalSuperResolutionPasses(
 			PassParameters->ReprojectionJacobianTexture = ReprojectionJacobianTexture ? ReprojectionJacobianTexture : GraphBuilder.CreateSRV(FRDGTextureSRVDesc(BlackUintDummy));
 			PassParameters->ReprojectionVectorTexture = ReprojectionVectorTexture;
 			PassParameters->IsMovingMaskTexture = IsMovingMaskTexture ? IsMovingMaskTexture : GraphBuilder.CreateSRV(FRDGTextureSRVDesc(BlackUintDummy));
+			PassParameters->ThinGeometryTexture = PassConfig.ThinGeometryDetectionEnable ? ThinGeometryTexture : GraphBuilder.CreateSRV(FRDGTextureSRVDesc(BlackUintDummy));
 			PassParameters->DecimateMaskTexture = DecimateMaskTexture;
 			PassParameters->HistoryRejectionTexture = HistoryRejectionTexture;
 			PassParameters->MoireHistoryTexture = MoireHistoryTexture ? MoireHistoryTexture : GraphBuilder.CreateSRV(FRDGTextureSRVDesc(BlackDummy));
