Fixed RHI validation error due to scene capture not being transitioned back to SRV.

diff --git a/Engine/Source/Runtime/Renderer/Private/SceneCaptureRendering.cpp b/Engine/Source/Runtime/Renderer/Private/SceneCaptureRendering.cpp
--- a/Engine/Source/Runtime/Renderer/Private/SceneCaptureRendering.cpp
+++ b/Engine/Source/Runtime/Renderer/Private/SceneCaptureRendering.cpp
@@ -498,6 +498,8 @@ static void UpdateSceneCaptureContent_RenderThread(
 	{
 		FGenerateMips::Execute(GraphBuilder, SceneRenderer->FeatureLevel, OutputTexture, GenerateMipsParams);
 	}
+
+	GraphBuilder.SetTextureAccessFinal(OutputTexture, ERHIAccess::SRVMask);
 }
 
 static void BuildOrthoMatrix(FIntPoint InRenderTargetSize, float InOrthoWidth, int32 InTileID, int32 InNumXTiles, int32 InNumYTiles, FMatrix& OutProjectionMatrix)
