Fixed scene render builder bugs.
- Added scene render LLM tags to RDG builder execute and scene render begin / end calls.
- Added null check for Execute().
- Moved BeginRenderViewFamily scene extension calls to the beginning due to certain extensions calling CaptureScene, which would recursively render scenes.
- Fixed assert when null RHI is used.

diff --git a/Engine/Source/Runtime/Renderer/Private/RendererScene.cpp b/Engine/Source/Runtime/Renderer/Private/RendererScene.cpp
--- a/Engine/Source/Runtime/Renderer/Private/RendererScene.cpp
+++ b/Engine/Source/Runtime/Renderer/Private/RendererScene.cpp
@@ -7270,11 +7270,7 @@ FSceneViewStateInterface* FRendererModule::AllocateViewState(ERHIFeatureLevel::T
 
 TUniquePtr<ISceneRenderBuilder> FRendererModule::CreateSceneRenderBuilder(FSceneInterface* SceneInterface)
 {
-	check(SceneInterface);
-	FScene* Scene = SceneInterface->GetRenderScene();
-	check(Scene);
-
-	return MakeUnique<FSceneRenderBuilder>(Scene);
+	return MakeUnique<FSceneRenderBuilder>(SceneInterface);
 }
 
 void FRendererModule::InvalidatePathTracedOutput()
diff --git a/Engine/Source/Runtime/Renderer/Private/SceneRenderBuilder.cpp b/Engine/Source/Runtime/Renderer/Private/SceneRenderBuilder.cpp
--- a/Engine/Source/Runtime/Renderer/Private/SceneRenderBuilder.cpp
+++ b/Engine/Source/Runtime/Renderer/Private/SceneRenderBuilder.cpp
@@ -252,6 +252,17 @@ void FSceneRenderProcessor::Execute()
 
 	UE::RenderCommandPipe::FSyncScope SyncScope;
 	FUniformExpressionCacheAsyncUpdateScope AsyncUpdateScope;
+	TOptional<UE::RenderCore::DumpGPU::FDumpScope> GpuDumpScope;
+	TOptional<RenderCaptureInterface::FScopedCapture> GpuCaptureScope;
+
+	// This has to be called first because it can actually recursively invoke scene renders.
+	for (FSceneRenderer* Renderer : Renderers)
+	{
+		for (const FSceneViewExtensionRef& Extension : Renderer->ViewFamily.ViewExtensions)
+		{
+			Extension->BeginRenderViewFamily(Renderer->ViewFamily);
+		}
+	}
 
 	ENQUEUE_RENDER_COMMAND(SceneRenderBuilder_Begin)([this](FRHICommandListImmediate& RHICmdList)
 	{
@@ -261,9 +272,6 @@ void FSceneRenderProcessor::Execute()
 		}
 	});
 
-	TOptional<UE::RenderCore::DumpGPU::FDumpScope> GpuDumpScope;
-	TOptional<RenderCaptureInterface::FScopedCapture> GpuCaptureScope;
-
 	for (FOp Op : Ops)
 	{
 		switch (Op.Type)
@@ -312,13 +320,9 @@ void FSceneRenderProcessor::Execute()
 		{
 			FRenderNode& RenderNode = *Op.Data_Render;
 
-			for (const FSceneViewExtensionRef& Extension : RenderNode.Renderer->ViewFamily.ViewExtensions)
-			{
-				Extension->BeginRenderViewFamily(RenderNode.Renderer->ViewFamily);
-			}
-
 			ENQUEUE_RENDER_COMMAND(SceneRenderBuilder_Render)([this, &RenderNode] (FRHICommandListImmediate& RHICmdList)
 			{
+				LLM_SCOPE(ELLMTag::SceneRender);
 				RenderState.GroupEvent.Begin(RenderState.Group, RHICmdList, FGroupEventLocation::SceneRenderCommand);
 
 				ON_SCOPE_EXIT
@@ -386,11 +390,18 @@ void FSceneRenderProcessor::Execute()
 
 //////////////////////////////////////////////////////////////////////////////
 
+FSceneRenderBuilder::FSceneRenderBuilder(FSceneInterface* InScene)
+	: Scene(InScene->GetRenderScene())
+{}
+
 FSceneRenderBuilder::~FSceneRenderBuilder()
 {
 	if (Processor)
 	{
+	#if !USE_NULL_RHI
 		checkf(false, TEXT("SceneRenderBuilder is being destructed without having called Execute."));
+	#endif
+
 		delete Processor;
 		Processor = nullptr;
 	}
@@ -406,18 +417,23 @@ void FSceneRenderBuilder::LazyInit()
 
 void FSceneRenderBuilder::AddCommand(TUniqueFunction<void()>&& Function)
 {
+#if !USE_NULL_RHI
 	LazyInit();
 	Processor->AddCommand(MoveTemp(Function));
+#endif
 }
 
 void FSceneRenderBuilder::AddRenderer(FSceneRenderer* Renderer, FString&& Name, FSceneRenderFunction&& Function)
 {
+#if !USE_NULL_RHI
 	LazyInit();
 	Processor->AddRenderer(Renderer, Forward<FString&&>(Name), MoveTemp(Function));
+#endif
 }
 
 void FSceneRenderBuilder::BeginGroup(FString&& Name, ESceneRenderGroupFlags Flags)
 {
+#if !USE_NULL_RHI
 	// If user sets both capture and dump flags, prefer capturing over dumping (or clear flag if dumping is not available or we are currently dumping already).
 	if (EnumHasAnyFlags(Flags, ESceneRenderGroupFlags::GpuCapture) || !(WITH_ENGINE && WITH_DUMPGPU) || FRDGBuilder::IsDumpingFrame())
 	{
@@ -426,18 +442,26 @@ void FSceneRenderBuilder::BeginGroup(FString&& Name, ESceneRenderGroupFlags Flag
 
 	LazyInit();
 	Processor->BeginGroup(Forward<FString&&>(Name), Flags);
+#endif
 }
 
 void FSceneRenderBuilder::EndGroup()
 {
+#if !USE_NULL_RHI
 	checkf(Processor, TEXT("EndGroup called on an empty scene render builder."));
 	Processor->EndGroup();
+#endif
 }
 
 void FSceneRenderBuilder::Execute()
 {
-	Processor->Execute();
-	Processor = nullptr;
+#if !USE_NULL_RHI
+	if (Processor)
+	{
+		Processor->Execute();
+		Processor = nullptr;
+	}
+#endif
 }
 
 FConcurrentLinearBulkObjectAllocator& FSceneRenderBuilder::GetAllocator()
diff --git a/Engine/Source/Runtime/Renderer/Private/SceneRenderBuilder.h b/Engine/Source/Runtime/Renderer/Private/SceneRenderBuilder.h
--- a/Engine/Source/Runtime/Renderer/Private/SceneRenderBuilder.h
+++ b/Engine/Source/Runtime/Renderer/Private/SceneRenderBuilder.h
@@ -12,10 +12,7 @@ class FScene;
 class FSceneRenderBuilder final : public ISceneRenderBuilder
 {
 public:
-	FSceneRenderBuilder(FScene* InScene)
-		: Scene(InScene)
-	{}
-
+	FSceneRenderBuilder(FSceneInterface* InScene);
 	FSceneRenderBuilder(const FSceneRenderBuilder&) = delete;
 	FSceneRenderBuilder& operator=(const FSceneRenderBuilder&) = delete;
 
diff --git a/Engine/Source/Runtime/Renderer/Private/SceneRendering.cpp b/Engine/Source/Runtime/Renderer/Private/SceneRendering.cpp
--- a/Engine/Source/Runtime/Renderer/Private/SceneRendering.cpp
+++ b/Engine/Source/Runtime/Renderer/Private/SceneRendering.cpp
@@ -5223,6 +5223,8 @@ void FSceneRenderer::RenderThreadEnd(FRHICommandListImmediate& RHICmdList)
 
 void FSceneRenderer::RenderThreadBegin(FRHICommandListImmediate& RHICmdList, TConstArrayView<FSceneRenderer*> SceneRenderers)
 {
+	LLM_SCOPE(ELLMTag::SceneRender);
+
 	CleanUp(RHICmdList);
 
 	// Initialize "AllFamilyViews" and "AllFamilies" array in each scene renderer.  Initialize array in first SceneRenderer, then copy to remaining.
@@ -5271,6 +5273,8 @@ static FSceneRenderCleanUpState GSceneRenderCleanUpState;
 
 void FSceneRenderer::RenderThreadEnd(FRHICommandListImmediate& RHICmdList, TConstArrayView<FSceneRenderer*> SceneRenderers)
 {
+	LLM_SCOPE(ELLMTag::SceneRender);
+
 	// We need to sync async uniform expression cache updates since we're about to start deleting material proxies.
 	FUniformExpressionCacheAsyncUpdateScope::WaitForTask();
 
