Initial update for Orthographic rendering use with Substrate enabled.

diff --git a/Engine/Shaders/Private/Common.ush b/Engine/Shaders/Private/Common.ush
--- a/Engine/Shaders/Private/Common.ush
+++ b/Engine/Shaders/Private/Common.ush
@@ -1533,6 +1533,14 @@ float3 GetTranslatedWorldCameraPosFromView(float2 SvPosition, bool bForceAddOrth
 	return GetTranslatedWorldCameraPosFromView(PrimaryView, SvPosition, bForceAddOrthoHeight);
 }
 
+float3 GetScreenWorldDir(in float4 SVPos)
+{
+	float2 ScreenPosition = SvPositionToScreenPosition(SVPos).xy;
+	const float Depth = 1000000.0f;
+	float4 TranslatedWorldPos = mul(float4(GetScreenPositionForProjectionType(ScreenPosition, Depth), Depth, 1), PrimaryView.ScreenToTranslatedWorld);
+	return GetCameraVectorFromTranslatedWorldPosition(TranslatedWorldPos.xyz);
+}
+
 float2 SvPositionToViewportUV(float4 SvPosition)
 {
 	// can be optimized from 2SUB+2MUL to 2MAD
diff --git a/Engine/Shaders/Private/LightFunctionAtlas/LightFunctionAtlasCommon.usf b/Engine/Shaders/Private/LightFunctionAtlas/LightFunctionAtlasCommon.usf
--- a/Engine/Shaders/Private/LightFunctionAtlas/LightFunctionAtlasCommon.usf
+++ b/Engine/Shaders/Private/LightFunctionAtlas/LightFunctionAtlasCommon.usf
@@ -110,7 +110,7 @@ FLightFunctionColor GetLocalLightFunctionCommon(float3 TranslatedWorldPosition,
 	LightFunctionColor = LightFunctionColor * LightFunctionColor; // Simple Gamma decompression
 
 	// Calculate radial view distance for stable fading
-	float ViewDistance = length(PrimaryView.TranslatedWorldCameraOrigin - TranslatedWorldPosition);
+	float ViewDistance = GetDistanceToCameraFromViewVector(PrimaryView.TranslatedWorldCameraOrigin - TranslatedWorldPosition);
 	float DistanceFadeAlpha = saturate((FadeDistance - ViewDistance) / (FadeDistance * .2f));
 	// Fade to disabled based on FadeDistance
 	LightFunctionColor = lerp(DisabledBrightness, LightFunctionColor, DistanceFadeAlpha);
diff --git a/Engine/Shaders/Private/SkyAtmosphere.usf b/Engine/Shaders/Private/SkyAtmosphere.usf
--- a/Engine/Shaders/Private/SkyAtmosphere.usf
+++ b/Engine/Shaders/Private/SkyAtmosphere.usf
@@ -162,14 +162,6 @@ float4 GetScreenTranslatedWorldPos(float4 SVPos, float DeviceZ)
 	return float4(SvPositionToTranslatedWorld(float4(SVPos.xy, DeviceZ, 1.0)), 1.0);
 }
 
-float3 GetScreenWorldDir(in float4 SVPos)
-{
-	float2 ScreenPosition = SvPositionToScreenPosition(SVPos).xy;
-	const float Depth = 1000000.0f;
-	float4 TranslatedWorldPos = mul(float4(GetScreenPositionForProjectionType(ScreenPosition, Depth), Depth, 1), View.ScreenToTranslatedWorld);
-	return GetCameraVectorFromTranslatedWorldPosition(TranslatedWorldPos.xyz);
-}
-
 // This is the world position of the camera. It is also force to be at the top of the virutal planet surface. 
 // This is to always see the sky even when the camera is buried into the virtual planet.
 float3 GetCameraTranslatedWorldPos()
diff --git a/Engine/Shaders/Private/Substrate/SubstrateRoughRefraction.usf b/Engine/Shaders/Private/Substrate/SubstrateRoughRefraction.usf
--- a/Engine/Shaders/Private/Substrate/SubstrateRoughRefraction.usf
+++ b/Engine/Shaders/Private/Substrate/SubstrateRoughRefraction.usf
@@ -167,14 +167,6 @@ bool Slabs(float3 P0, float3 P1, float3 RayOrigin, float3 InvRaydir, out float O
 	return MaxTMin <= MinTMax;
 }
 
-float3 GetScreenWorldDir(in float4 SVPos)
-{
-	float2 ScreenPosition = SvPositionToScreenPosition(SVPos).xy;
-	const float Depth = 1000000.0f;
-	float4 WorldPos = mul(float4(ScreenPosition * Depth, Depth, 1), View.ScreenToTranslatedWorld);
-	return normalize(WorldPos.xyz - View.TranslatedWorldCameraOrigin);
-}
-
 float GetEta12()
 {
 	const float AirIOR = 1.0f;
diff --git a/Engine/Shaders/Private/Substrate/SubstrateVisualize.usf b/Engine/Shaders/Private/Substrate/SubstrateVisualize.usf
--- a/Engine/Shaders/Private/Substrate/SubstrateVisualize.usf
+++ b/Engine/Shaders/Private/Substrate/SubstrateVisualize.usf
@@ -42,7 +42,7 @@ float3 GetWorldPositionFromPixelPos(uint2 PixelPos)
 	const float2 ScreenPosition = SvPositionToScreenPosition(float4(PixelPos, 0.5f, 1)).xy;
 	const float DeviceZ = SampleDeviceZFromSceneTextures(BufferUV);
 	const float SceneDepth = ConvertFromDeviceZ(DeviceZ);
-	const float3 WorldPosition = mul(float4(ScreenPosition * SceneDepth, SceneDepth, 1), PrimaryView.ScreenToTranslatedWorld).xyz;
+	const float3 WorldPosition = mul(float4(GetScreenPositionForProjectionType(ScreenPosition, SceneDepth), SceneDepth, 1), PrimaryView.ScreenToTranslatedWorld).xyz;
 	return WorldPosition;
 }
 
@@ -51,7 +51,7 @@ void MaterialPrintInfoCS(uint3 DispatchThreadId : SV_DispatchThreadID)
 {
 	const uint2 PixelPos = (float2(GetCursorPosition()) * View.ViewResolutionFraction);
 	const float3 WorldPosition = GetWorldPositionFromPixelPos(PixelPos);
-	const float3 V = -normalize(WorldPosition - PrimaryView.TranslatedWorldCameraOrigin);
+	const float3 V = -GetCameraVectorFromTranslatedWorldPosition(WorldPosition);
 
 	FShaderPrintContext Context = InitShaderPrintContext(true, uint2(50, 50));
 
@@ -133,7 +133,7 @@ void VisualizeMaterialPS(
 	float DeviceZ = SampleDeviceZFromSceneTextures(BufferUV);
 	float SceneDepth = ConvertFromDeviceZ(DeviceZ);
 	const float3 WorldPosition = mul(float4(ScreenMaterialPosition * SceneDepth, SceneDepth, 1), PrimaryView.ScreenToTranslatedWorld).xyz;
-	const float3 V = -normalize(WorldPosition - PrimaryView.TranslatedWorldCameraOrigin);
+	const float3 V = -GetCameraVectorFromTranslatedWorldPosition(WorldPosition);
 
 	FSubstrateAddressing SubstrateAddressing = GetSubstratePixelDataByteOffset(PixelPosDynRes, uint2(View.BufferSizeAndInvSize.xy), Substrate.MaxBytesPerPixel);
 	const uint FootPrint_Start = SubstrateAddressing.ReadBytes;
diff --git a/Engine/Source/Runtime/Renderer/Private/Substrate/Substrate.cpp b/Engine/Source/Runtime/Renderer/Private/Substrate/Substrate.cpp
--- a/Engine/Source/Runtime/Renderer/Private/Substrate/Substrate.cpp
+++ b/Engine/Source/Runtime/Renderer/Private/Substrate/Substrate.cpp
@@ -1096,19 +1096,36 @@ static FSubstrateTileParameter InternalSetTileParameters(FRDGBuilder* GraphBuild
 	return Out;
 }
 
-FSubstrateTilePassVS::FParameters SetTileParameters(
-	const FViewInfo& View,
-	const ESubstrateTileType TileType,
-	EPrimitiveType& PrimitiveType)
+static FSubstrateTilePassVS::FParameters SetTileParametersCommon(FRDGBuilder* GraphBuilder,
+		const FViewInfo& View, 
+		const ESubstrateTileType TileType, 
+		EPrimitiveType& PrimitiveType)
 {
-	FSubstrateTileParameter Temp = InternalSetTileParameters(nullptr, View, TileType);
+	FSubstrateTileParameter Temp = InternalSetTileParameters(GraphBuilder, View, TileType);
 	PrimitiveType = GRHISupportsRectTopology ? PT_RectList : PT_TriangleList;
 
 	FSubstrateTilePassVS::FParameters Out;
 	Out.OutputViewMinRect = FVector2f(View.CachedViewUniformShaderParameters->ViewRectMin.X, View.CachedViewUniformShaderParameters->ViewRectMin.Y);
 	Out.OutputViewSizeAndInvSize = View.CachedViewUniformShaderParameters->ViewSizeAndInvSize;
 	Out.OutputBufferSizeAndInvSize = View.CachedViewUniformShaderParameters->BufferSizeAndInvSize;
-	Out.ViewScreenToTranslatedWorld = View.CachedViewUniformShaderParameters->ScreenToTranslatedWorld;
+
+	if (View.IsPerspectiveProjection())
+	{
+		Out.ViewScreenToTranslatedWorld = View.CachedViewUniformShaderParameters->ScreenToTranslatedWorld;
+	}
+	else
+	{
+		/**
+		* This forces the calculation for orthographic cameras to use the hard coded screen depth == 1 to resolve the view 
+		* without having to introduce a new shader parameter.
+		* For now this is fine but this should be changed to use ScreenVectorFromScreenRect, matching other passes, once the whole view buffer 
+		* is being passed to shader params. Note: will not be compatible with reverse Z if the hard coded behaviour is altered.
+		*/
+		FMatrix44f OrthoMatrix(EForceInit::ForceInitToZero);
+		OrthoMatrix.SetAxis(2, View.CachedViewUniformShaderParameters->ViewForward);
+		Out.ViewScreenToTranslatedWorld = OrthoMatrix;
+	}
+
 	Out.TileListBuffer = Temp.TileListBuffer;
 	Out.TileListBufferOffset = Temp.TileListBufferOffset;
 	Out.TileEncoding = Temp.TileEncoding;
@@ -1116,25 +1133,22 @@ FSubstrateTilePassVS::FParameters SetTileParameters(
 	return Out;
 }
 
+
+FSubstrateTilePassVS::FParameters SetTileParameters(
+	const FViewInfo& View,
+	const ESubstrateTileType TileType,
+	EPrimitiveType& PrimitiveType)
+{
+	return SetTileParametersCommon(nullptr, View, TileType, PrimitiveType);
+}
+
 FSubstrateTilePassVS::FParameters SetTileParameters(
 	FRDGBuilder& GraphBuilder, 
 	const FViewInfo& View, 
 	const ESubstrateTileType TileType,
 	EPrimitiveType& PrimitiveType)
 {
-	FSubstrateTileParameter Temp = InternalSetTileParameters(&GraphBuilder, View, TileType);
-	PrimitiveType = GRHISupportsRectTopology ? PT_RectList : PT_TriangleList;
-
-	FSubstrateTilePassVS::FParameters Out;
-	Out.OutputViewMinRect = FVector2f(View.CachedViewUniformShaderParameters->ViewRectMin.X, View.CachedViewUniformShaderParameters->ViewRectMin.Y);
-	Out.OutputViewSizeAndInvSize = View.CachedViewUniformShaderParameters->ViewSizeAndInvSize;
-	Out.OutputBufferSizeAndInvSize = View.CachedViewUniformShaderParameters->BufferSizeAndInvSize;
-	Out.ViewScreenToTranslatedWorld = View.CachedViewUniformShaderParameters->ScreenToTranslatedWorld;
-	Out.TileListBuffer = Temp.TileListBuffer;
-	Out.TileListBufferOffset = Temp.TileListBufferOffset;
-	Out.TileEncoding = Temp.TileEncoding;
-	Out.TileIndirectBuffer = Temp.TileIndirectBuffer;
-	return Out;
+	return SetTileParametersCommon(&GraphBuilder, View, TileType, PrimitiveType);
 }
 
 FSubstrateTileParameter SetTileParameters(FRDGBuilder& GraphBuilder, const FViewInfo& View, const ESubstrateTileType TileType)
