Substrate - tile vertex shader now uses the view uniform buffer instead of duplicating data.
Follow up to 38535975 fixing ortho rendering, now without hack, using the common ScreenVectorFromScreenRect function.

#tests PC dx12 with/without Substrate, Lumen, clustered shading
#rnx

diff --git a/Engine/Shaders/Private/Substrate/SubstrateTile.usf b/Engine/Shaders/Private/Substrate/SubstrateTile.usf
--- a/Engine/Shaders/Private/Substrate/SubstrateTile.usf
+++ b/Engine/Shaders/Private/Substrate/SubstrateTile.usf
@@ -9,11 +9,6 @@
 
 #if SHADER_TILE_VS
 
-float2 OutputViewMinRect;
-float4 OutputViewSizeAndInvSize;
-float4 OutputBufferSizeAndInvSize;
-float4x4 ViewScreenToTranslatedWorld;
-
 Buffer<uint> TileListBuffer;
 uint		 TileListBufferOffset;
 uint		 TileEncoding;
@@ -41,11 +36,11 @@ void SubstrateTilePassVS(
 	TileVertex.x += InVertexId == 1 || InVertexId == 2 || InVertexId == 4 ? SUBSTRATE_TILE_SIZE : 0;
 	TileVertex.y += InVertexId == 2 || InVertexId == 4 || InVertexId == 5 ? SUBSTRATE_TILE_SIZE : 0;
 
-	OutPosition = float4(float2(TileVertex) * OutputViewSizeAndInvSize.zw * float2(2.0f, -2.0f) + float2(-1.0, 1.0f), 0.5f, 1.0f);
+	OutPosition = float4(float2(TileVertex) * View.ViewSizeAndInvSize.zw * float2(2.0f, -2.0f) + float2(-1.0, 1.0f), 0.5f, 1.0f);
 
 #if PERMUTATION_ENABLE_TEXCOORD_SCREENVECTOR
-	OutTexCoord = (TileVertex + OutputViewMinRect.xy) * OutputBufferSizeAndInvSize.zw;
-	OutScreenVector = mul(float4(OutPosition.xy, 1, 0), ViewScreenToTranslatedWorld).xyz;
+	OutTexCoord = (TileVertex + View.ViewRectMin.xy) * View.BufferSizeAndInvSize.zw;
+	OutScreenVector = ScreenVectorFromScreenRect(float4(OutPosition.xy, 1, 0));
 #endif
 
 #if PERMUTATION_ENABLE_DEBUG
diff --git a/Engine/Source/Runtime/Renderer/Private/ClusteredDeferredShadingPass.cpp b/Engine/Source/Runtime/Renderer/Private/ClusteredDeferredShadingPass.cpp
--- a/Engine/Source/Runtime/Renderer/Private/ClusteredDeferredShadingPass.cpp
+++ b/Engine/Source/Runtime/Renderer/Private/ClusteredDeferredShadingPass.cpp
@@ -115,7 +115,6 @@ class FClusteredShadingPS : public FGlobalShader
 		SHADER_PARAMETER_RDG_TEXTURE(Texture2D, ShadowMaskBits)
 		SHADER_PARAMETER_RDG_TEXTURE(Texture2D<uint>, LightingChannelsTexture)
 		SHADER_PARAMETER_STRUCT_INCLUDE(FVirtualShadowMapSamplingParameters, VirtualShadowMapSamplingParameters)
-		SHADER_PARAMETER_STRUCT_INCLUDE(Substrate::FSubstrateTilePassVS::FParameters, SubstrateTile)
 		SHADER_PARAMETER_RDG_BUFFER_SRV(Buffer<uint>, HairTransmittanceBuffer)
 		RENDER_TARGET_BINDING_SLOTS()
 	END_SHADER_PARAMETER_STRUCT()
@@ -155,6 +154,11 @@ class FClusteredShadingPS : public FGlobalShader
 
 IMPLEMENT_GLOBAL_SHADER(FClusteredShadingPS, "/Engine/Private/ClusteredDeferredShadingPixelShader.usf", "ClusteredShadingPixelShader", SF_Pixel);
 
+BEGIN_SHADER_PARAMETER_STRUCT(FClusteredShadingParameters, )
+	SHADER_PARAMETER_STRUCT(FClusteredShadingPS::FParameters, PS)
+	SHADER_PARAMETER_STRUCT(Substrate::FSubstrateTilePassVS::FParameters, VS)
+END_SHADER_PARAMETER_STRUCT()
+
 enum class EClusterPassInputType : uint8
 {
 	GBuffer,
@@ -184,37 +188,38 @@ static void InternalAddClusteredDeferredShadingPass(
 	const bool bLightFunctionAtlas = LightFunctionAtlas::IsEnabled(View, ELightFunctionAtlasSystem::DeferredLighting);	
 	const bool bSupportAnisotropyPermutation = ShouldRenderAnisotropyPass(View) && !Substrate::IsSubstrateEnabled(); // Strata managed anisotropy differently than legacy path. No need for special permutation.
 
-	FClusteredShadingPS::FParameters *PassParameters = GraphBuilder.AllocParameters<FClusteredShadingPS::FParameters>();
-	PassParameters->bHasLightChannels = SortedLightsSet.bHasLightChannels;
-	PassParameters->View = View.ViewUniformBuffer;
-	PassParameters->HairStrands = HairStrands::BindHairStrandsViewUniformParameters(View);
-	PassParameters->ForwardLightStruct = View.ForwardLightingResources.ForwardLightUniformBuffer;
-	PassParameters->SceneTextures = SceneTextures.UniformBuffer;
-	PassParameters->ShadowMaskBits = ShadowMaskBits ? ShadowMaskBits : GSystemTextures.GetZeroUIntDummy(GraphBuilder);
-	PassParameters->LightingChannelsTexture = LightingChannelsTexture ? LightingChannelsTexture : GSystemTextures.GetZeroUIntDummy(GraphBuilder);
-	PassParameters->VirtualShadowMapSamplingParameters = VirtualShadowMapArray.GetSamplingParameters(GraphBuilder, ViewIndex);
-	PassParameters->HairTransmittanceBuffer = HairTransmittanceBuffer;
-	PassParameters->Substrate = Substrate::BindSubstrateGlobalUniformParameters(View);
-	PassParameters->LightFunctionAtlas = LightFunctionAtlas::BindGlobalParameters(GraphBuilder, View);
-	ShaderPrint::SetParameters(GraphBuilder, View.ShaderPrintData, PassParameters->ShaderPrintUniformBuffer);
-
-	PassParameters->RenderTargets[0] = FRenderTargetBinding(SceneTextures.Color.Target, ERenderTargetLoadAction::ELoad);
+	FClusteredShadingParameters* PassParameters = GraphBuilder.AllocParameters<FClusteredShadingParameters>();
+	FClusteredShadingPS::FParameters* PSParameters = &PassParameters->PS;
+	PSParameters->bHasLightChannels = SortedLightsSet.bHasLightChannels;
+	PSParameters->View = View.ViewUniformBuffer;
+	PSParameters->HairStrands = HairStrands::BindHairStrandsViewUniformParameters(View);
+	PSParameters->ForwardLightStruct = View.ForwardLightingResources.ForwardLightUniformBuffer;
+	PSParameters->SceneTextures = SceneTextures.UniformBuffer;
+	PSParameters->ShadowMaskBits = ShadowMaskBits ? ShadowMaskBits : GSystemTextures.GetZeroUIntDummy(GraphBuilder);
+	PSParameters->LightingChannelsTexture = LightingChannelsTexture ? LightingChannelsTexture : GSystemTextures.GetZeroUIntDummy(GraphBuilder);
+	PSParameters->VirtualShadowMapSamplingParameters = VirtualShadowMapArray.GetSamplingParameters(GraphBuilder, ViewIndex);
+	PSParameters->HairTransmittanceBuffer = HairTransmittanceBuffer;
+	PSParameters->Substrate = Substrate::BindSubstrateGlobalUniformParameters(View);
+	PSParameters->LightFunctionAtlas = LightFunctionAtlas::BindGlobalParameters(GraphBuilder, View);
+	ShaderPrint::SetParameters(GraphBuilder, View.ShaderPrintData, PSParameters->ShaderPrintUniformBuffer);
+
+	PSParameters->RenderTargets[0] = FRenderTargetBinding(SceneTextures.Color.Target, ERenderTargetLoadAction::ELoad);
 	if (bHairStrands)
 	{
-		PassParameters->RenderTargets[0] = FRenderTargetBinding(View.HairStrandsViewData.VisibilityData.SampleLightingTexture, ERenderTargetLoadAction::ELoad);
+		PSParameters->RenderTargets[0] = FRenderTargetBinding(View.HairStrandsViewData.VisibilityData.SampleLightingTexture, ERenderTargetLoadAction::ELoad);
 	}
 	if (Substrate::IsOpaqueRoughRefractionEnabled())
 	{
 		check(SubstrateSceneData);
-		PassParameters->RenderTargets[1] = FRenderTargetBinding(SubstrateSceneData->SeparatedOpaqueRoughRefractionSceneColor, ERenderTargetLoadAction::ELoad);
-		PassParameters->RenderTargets[2] = FRenderTargetBinding(SubstrateSceneData->SeparatedSubSurfaceSceneColor, ERenderTargetLoadAction::ELoad);
+		PSParameters->RenderTargets[1] = FRenderTargetBinding(SubstrateSceneData->SeparatedOpaqueRoughRefractionSceneColor, ERenderTargetLoadAction::ELoad);
+		PSParameters->RenderTargets[2] = FRenderTargetBinding(SubstrateSceneData->SeparatedSubSurfaceSceneColor, ERenderTargetLoadAction::ELoad);
 	}
 
 	// VS - Substrate tile parameters
 	EPrimitiveType PrimitiveType = PT_TriangleList;
 	if (Substrate::IsSubstrateEnabled())
 	{
-		PassParameters->SubstrateTile = Substrate::SetTileParameters(GraphBuilder, View, TileType, PrimitiveType);
+		PassParameters->VS = Substrate::SetTileParameters(GraphBuilder, View, TileType, PrimitiveType);
 	}
 	
 	const TCHAR* TileTypeName = ToString(TileType);
@@ -267,16 +272,16 @@ static void InternalAddClusteredDeferredShadingPass(
 			SetGraphicsPipelineState(InRHICmdList, GraphicsPSOInit, 0);
 		}
 
-		SetShaderParameters(InRHICmdList, PixelShader, PixelShader.GetPixelShader(), *PassParameters);
+		SetShaderParameters(InRHICmdList, PixelShader, PixelShader.GetPixelShader(), PassParameters->PS);
 
 		if (bHairStrands)
 		{
 			InRHICmdList.SetViewport(0, 0, 0.0f, View.HairStrandsViewData.VisibilityData.SampleLightingViewportResolution.X, View.HairStrandsViewData.VisibilityData.SampleLightingViewportResolution.Y, 1.0f);
 
 			FClusteredShadingVS::FParameters VertexParameters;
-			VertexParameters.View = PassParameters->View;
-			VertexParameters.HairStrands = PassParameters->HairStrands;
-			VertexParameters.SceneTextures = PassParameters->SceneTextures;
+			VertexParameters.View = PassParameters->PS.View;
+			VertexParameters.HairStrands = PassParameters->PS.HairStrands;
+			VertexParameters.SceneTextures = PassParameters->PS.SceneTextures;
 			SetShaderParameters(InRHICmdList, HairVertexShader, HairVertexShader.GetVertexShader(), VertexParameters);
 			InRHICmdList.SetStreamSource(0, nullptr, 0);
 			InRHICmdList.DrawPrimitive(0, 1, 1);
@@ -284,8 +289,8 @@ static void InternalAddClusteredDeferredShadingPass(
 		else if (bSubstrate)
 		{
 			InRHICmdList.SetViewport(View.ViewRect.Min.X, View.ViewRect.Min.Y, 0.0f, View.ViewRect.Max.X, View.ViewRect.Max.Y, 1.0f);
-			SetShaderParameters(InRHICmdList, TileVertexShader, TileVertexShader.GetVertexShader(), PassParameters->SubstrateTile);
-			InRHICmdList.DrawPrimitiveIndirect(PassParameters->SubstrateTile.TileIndirectBuffer->GetIndirectRHICallBuffer(), Substrate::TileTypeDrawIndirectArgOffset(TileType));
+			SetShaderParameters(InRHICmdList, TileVertexShader, TileVertexShader.GetVertexShader(), PassParameters->VS);
+			InRHICmdList.DrawPrimitiveIndirect(PassParameters->VS.TileIndirectBuffer->GetIndirectRHICallBuffer(), Substrate::TileTypeDrawIndirectArgOffset(TileType));
 		}
 		else
 		{
diff --git a/Engine/Source/Runtime/Renderer/Private/IndirectLightRendering.cpp b/Engine/Source/Runtime/Renderer/Private/IndirectLightRendering.cpp
--- a/Engine/Source/Runtime/Renderer/Private/IndirectLightRendering.cpp
+++ b/Engine/Source/Runtime/Renderer/Private/IndirectLightRendering.cpp
@@ -191,7 +191,6 @@ class FDiffuseIndirectCompositePS : public FGlobalShader
 
 		SHADER_PARAMETER_RDG_UNIFORM_BUFFER(FSceneTextureUniformParameters, SceneTexturesStruct)
 		SHADER_PARAMETER_RDG_UNIFORM_BUFFER(FSubstrateGlobalUniformParameters, Substrate)
-		SHADER_PARAMETER_STRUCT_INCLUDE(Substrate::FSubstrateTilePassVS::FParameters, SubstrateTile)
 		SHADER_PARAMETER_STRUCT_INCLUDE(Denoiser::FCommonShaderParameters, DenoiserCommonParameters)
 		SHADER_PARAMETER_STRUCT_INCLUDE(FSceneTextureParameters, SceneTextures)
 		SHADER_PARAMETER_STRUCT_REF(FViewUniformShaderParameters, ViewUniformBuffer)
@@ -222,6 +221,11 @@ class FDiffuseIndirectCompositePS : public FGlobalShader
 
 };
 
+BEGIN_SHADER_PARAMETER_STRUCT(FDiffuseIndirectCompositeParameters, )
+	SHADER_PARAMETER_STRUCT(FDiffuseIndirectCompositePS::FParameters, PS)
+	SHADER_PARAMETER_STRUCT(Substrate::FSubstrateTilePassVS::FParameters, VS)
+END_SHADER_PARAMETER_STRUCT()
+
 class FAmbientCubemapCompositePS : public FGlobalShader
 {
 	DECLARE_GLOBAL_SHADER(FAmbientCubemapCompositePS)
@@ -240,7 +244,6 @@ class FAmbientCubemapCompositePS : public FGlobalShader
 		SHADER_PARAMETER_SAMPLER(SamplerState,  AmbientOcclusionSampler)
 
 		SHADER_PARAMETER_RDG_UNIFORM_BUFFER(FSubstrateGlobalUniformParameters, Substrate)
-		SHADER_PARAMETER_STRUCT_INCLUDE(Substrate::FSubstrateTilePassVS::FParameters, SubstrateTile)
 		SHADER_PARAMETER_STRUCT_INCLUDE(FAmbientCubemapParameters, AmbientCubemap)
 		SHADER_PARAMETER_STRUCT_INCLUDE(FSceneTextureParameters, SceneTextures)
 		SHADER_PARAMETER_STRUCT_REF(FViewUniformShaderParameters, ViewUniformBuffer)
@@ -249,6 +252,11 @@ class FAmbientCubemapCompositePS : public FGlobalShader
 	END_SHADER_PARAMETER_STRUCT()
 };
 
+BEGIN_SHADER_PARAMETER_STRUCT(FAmbientCubemapCompositeParameters, )
+	SHADER_PARAMETER_STRUCT(FAmbientCubemapCompositePS::FParameters, PS)
+	SHADER_PARAMETER_STRUCT(Substrate::FSubstrateTilePassVS::FParameters, VS)
+END_SHADER_PARAMETER_STRUCT()
+
 // Setups all shader parameters related to skylight.
 FSkyDiffuseLightingParameters GetSkyDiffuseLightingParameters(const FSkyLightSceneProxy* SkyLight, float DynamicBentNormalAO)
 {
@@ -1309,11 +1317,11 @@ void FDeferredShadingSceneRenderer::RenderDiffuseIndirectAndAmbientOcclusion(
 				CopyInfo.Size = SceneColorTexture->Desc.GetSize();
 				AddCopyTexturePass(GraphBuilder, SceneColorTexture, SceneColorCopyTexture, CopyInfo);
 			}
-
-			FDiffuseIndirectCompositePS::FParameters* PassParameters = GraphBuilder.AllocParameters<FDiffuseIndirectCompositePS::FParameters>();
-			PassParameters->Substrate = Substrate::BindSubstrateGlobalUniformParameters(View);
-			PassParameters->SceneTexturesStruct = SceneTextures.UniformBuffer;
-			PassParameters->AmbientOcclusionStaticFraction = FMath::Clamp(View.FinalPostProcessSettings.AmbientOcclusionStaticFraction, 0.0f, 1.0f);
+			FDiffuseIndirectCompositeParameters* PassParameters = GraphBuilder.AllocParameters<FDiffuseIndirectCompositeParameters>();
+			FDiffuseIndirectCompositePS::FParameters* PSParameters = &PassParameters->PS;
+			PSParameters->Substrate = Substrate::BindSubstrateGlobalUniformParameters(View);
+			PSParameters->SceneTexturesStruct = SceneTextures.UniformBuffer;
+			PSParameters->AmbientOcclusionStaticFraction = FMath::Clamp(View.FinalPostProcessSettings.AmbientOcclusionStaticFraction, 0.0f, 1.0f);
 
 			const FIntPoint BufferExtent = SceneTextureParameters.SceneDepthTexture->Desc.Extent;
 
@@ -1326,70 +1334,70 @@ void FDeferredShadingSceneRenderer::RenderDiffuseIndirectAndAmbientOcclusion(
 					TexCreate_ShaderResource);
 				FRDGTextureRef CompressedMetadataPlaceholder = GraphBuilder.CreateTexture(Desc, TEXT("CompressedMetadataPlaceholder"));
 
-				PassParameters->CompressedMetadata[0] = CompressedMetadataPlaceholder;
-				PassParameters->CompressedMetadata[1] = CompressedMetadataPlaceholder;
+				PSParameters->CompressedMetadata[0] = CompressedMetadataPlaceholder;
+				PSParameters->CompressedMetadata[1] = CompressedMetadataPlaceholder;
 			}
 
-			PassParameters->BufferUVToOutputPixelPosition = BufferExtent;
-			PassParameters->EyeAdaptation = GraphBuilder.CreateSRV(GetEyeAdaptationBuffer(GraphBuilder, View));
-			LumenReflections::SetupCompositeParameters(View, PassParameters->ReflectionsCompositeParameters);
-			PassParameters->ScreenBentNormalParameters = ScreenBentNormalParameters;
-			PassParameters->bLumenSupportBackfaceDiffuse = ViewPipelineState.DiffuseIndirectMethod == EDiffuseIndirectMethod::Lumen && DenoiserOutputs.Textures[1] != SystemTextures.Black;
-			PassParameters->bLumenReflectionInputIsSSR = ViewPipelineState.DiffuseIndirectMethod == EDiffuseIndirectMethod::Lumen && ViewPipelineState.ReflectionsMethod == EReflectionsMethod::SSR;
+			PSParameters->BufferUVToOutputPixelPosition = BufferExtent;
+			PSParameters->EyeAdaptation = GraphBuilder.CreateSRV(GetEyeAdaptationBuffer(GraphBuilder, View));
+			LumenReflections::SetupCompositeParameters(View, PSParameters->ReflectionsCompositeParameters);
+			PSParameters->ScreenBentNormalParameters = ScreenBentNormalParameters;
+			PSParameters->bLumenSupportBackfaceDiffuse = ViewPipelineState.DiffuseIndirectMethod == EDiffuseIndirectMethod::Lumen && DenoiserOutputs.Textures[1] != SystemTextures.Black;
+			PSParameters->bLumenReflectionInputIsSSR = ViewPipelineState.DiffuseIndirectMethod == EDiffuseIndirectMethod::Lumen && ViewPipelineState.ReflectionsMethod == EReflectionsMethod::SSR;
 			extern float GLumenShortRangeAOFoliageOcclusionStrength;
-			PassParameters->LumenFoliageOcclusionStrength = GLumenShortRangeAOFoliageOcclusionStrength;
+			PSParameters->LumenFoliageOcclusionStrength = GLumenShortRangeAOFoliageOcclusionStrength;
 			extern float GLumenMaxShortRangeAOMultibounceAlbedo;
-			PassParameters->LumenMaxAOMultibounceAlbedo = GLumenMaxShortRangeAOMultibounceAlbedo;
-			PassParameters->LumenReflectionSpecularScale = GetLumenReflectionSpecularScale();
-			PassParameters->LumenReflectionContrast = GetLumenReflectionContrast();
+			PSParameters->LumenMaxAOMultibounceAlbedo = GLumenMaxShortRangeAOMultibounceAlbedo;
+			PSParameters->LumenReflectionSpecularScale = GetLumenReflectionSpecularScale();
+			PSParameters->LumenReflectionContrast = GetLumenReflectionContrast();
 
-			PassParameters->bVisualizeDiffuseIndirect = bIsVisualizePass;
+			PSParameters->bVisualizeDiffuseIndirect = bIsVisualizePass;
 
-			PassParameters->DiffuseIndirect_Lumen_0 = SetDefaultBlackIfNotTexture2DArray(DenoiserOutputs.Textures[0]);
-			PassParameters->DiffuseIndirect_Lumen_1 = SetDefaultBlackIfNotTexture2DArray(DenoiserOutputs.Textures[1]);
-			PassParameters->DiffuseIndirect_Lumen_2 = SetDefaultBlackIfNotTexture2DArray(DenoiserOutputs.Textures[2]);
-			PassParameters->DiffuseIndirect_Lumen_3 = SetDefaultBlackIfNotTexture2DArray(DenoiserOutputs.Textures[3]);
+			PSParameters->DiffuseIndirect_Lumen_0 = SetDefaultBlackIfNotTexture2DArray(DenoiserOutputs.Textures[0]);
+			PSParameters->DiffuseIndirect_Lumen_1 = SetDefaultBlackIfNotTexture2DArray(DenoiserOutputs.Textures[1]);
+			PSParameters->DiffuseIndirect_Lumen_2 = SetDefaultBlackIfNotTexture2DArray(DenoiserOutputs.Textures[2]);
+			PSParameters->DiffuseIndirect_Lumen_3 = SetDefaultBlackIfNotTexture2DArray(DenoiserOutputs.Textures[3]);
 
-			PassParameters->DiffuseIndirect = DenoiserOutputs;
+			PSParameters->DiffuseIndirect = DenoiserOutputs;
 			for (uint32 EntryIt = 0; EntryIt < kMaxDenoiserBufferProcessingCount; ++EntryIt)
 			{
-				PassParameters->DiffuseIndirect.Textures[EntryIt] = SetDefaultBlackIfNotTexture2D(PassParameters->DiffuseIndirect.Textures[EntryIt]);
+				PSParameters->DiffuseIndirect.Textures[EntryIt] = SetDefaultBlackIfNotTexture2D(PSParameters->DiffuseIndirect.Textures[EntryIt]);
 			}
 
-			PassParameters->DiffuseIndirectSampler = TStaticSamplerState<SF_Point>::GetRHI();
+			PSParameters->DiffuseIndirectSampler = TStaticSamplerState<SF_Point>::GetRHI();
 
-			PassParameters->PreIntegratedGF = GSystemTextures.PreintegratedGF->GetRHI();
-			PassParameters->PreIntegratedGFSampler = TStaticSamplerState<SF_Bilinear, AM_Clamp, AM_Clamp, AM_Clamp>::GetRHI();
+			PSParameters->PreIntegratedGF = GSystemTextures.PreintegratedGF->GetRHI();
+			PSParameters->PreIntegratedGFSampler = TStaticSamplerState<SF_Bilinear, AM_Clamp, AM_Clamp, AM_Clamp>::GetRHI();
 
-			PassParameters->AmbientOcclusionTexture = AmbientOcclusionMask;
-			PassParameters->AmbientOcclusionSampler = TStaticSamplerState<SF_Point>::GetRHI();
+			PSParameters->AmbientOcclusionTexture = AmbientOcclusionMask;
+			PSParameters->AmbientOcclusionSampler = TStaticSamplerState<SF_Point>::GetRHI();
 
 			if (bRunSSRTiledComposite)
 			{
-				PassParameters->SSRTiledViewRes = ScreenSpaceReflectionTileClassification.TiledViewRes;
-				PassParameters->SSRTileMaskBuffer = GraphBuilder.CreateSRV(ScreenSpaceReflectionTileClassification.TileMaskBuffer);
+				PSParameters->SSRTiledViewRes = ScreenSpaceReflectionTileClassification.TiledViewRes;
+				PSParameters->SSRTileMaskBuffer = GraphBuilder.CreateSRV(ScreenSpaceReflectionTileClassification.TileMaskBuffer);
 			}
 
 			if (bEnableCopyPass)
 			{
-				PassParameters->SceneColorTexture = SceneColorCopyTexture;
-				PassParameters->SceneColorSampler = TStaticSamplerState<SF_Point>::GetRHI();
+				PSParameters->SceneColorTexture = SceneColorCopyTexture;
+				PSParameters->SceneColorSampler = TStaticSamplerState<SF_Point>::GetRHI();
 			}
 			
-			if (!PassParameters->AmbientOcclusionTexture || bIsVisualizePass)
+			if (!PSParameters->AmbientOcclusionTexture || bIsVisualizePass)
 			{
-				PassParameters->AmbientOcclusionTexture = SystemTextures.White;
+				PSParameters->AmbientOcclusionTexture = SystemTextures.White;
 			}
 
 			Denoiser::SetupCommonShaderParameters(
 				View, SceneTextureParameters,
 				View.ViewRect,
 				1.0f / CommonDiffuseParameters.DownscaleFactor,
-				/* out */ &PassParameters->DenoiserCommonParameters);
-			PassParameters->SceneTextures = SceneTextureParameters;
-			PassParameters->ViewUniformBuffer = View.ViewUniformBuffer;
+				/* out */ &PSParameters->DenoiserCommonParameters);
+			PSParameters->SceneTextures = SceneTextureParameters;
+			PSParameters->ViewUniformBuffer = View.ViewUniformBuffer;
 
-			PassParameters->RenderTargets[0] = FRenderTargetBinding(
+			PSParameters->RenderTargets[0] = FRenderTargetBinding(
 				SceneColorTexture, ERenderTargetLoadAction::ELoad);
 
 			{
@@ -1399,7 +1407,7 @@ void FDeferredShadingSceneRenderer::RenderDiffuseIndirectAndAmbientOcclusion(
 					FClearValueBinding::None,
 					TexCreate_ShaderResource | TexCreate_UAV);
 
-				PassParameters->PassDebugOutput = GraphBuilder.CreateUAV(
+				PSParameters->PassDebugOutput = GraphBuilder.CreateUAV(
 					GraphBuilder.CreateTexture(Desc, TEXT("DebugDiffuseIndirectComposite")));
 			}
 
@@ -1474,7 +1482,7 @@ void FDeferredShadingSceneRenderer::RenderDiffuseIndirectAndAmbientOcclusion(
 
 			if (TileType == ESubstrateTileType::ECount)
 			{
-				ClearUnusedGraphResources(PixelShader, PassParameters);
+				ClearUnusedGraphResources(PixelShader, PSParameters);
 
 				// only use depth bound optimization when diffuse indirect is disable
 				bool bUseDepthBounds = CVarDiffuseIndirectOffUseDepthBoundsAO.GetValueOnRenderThread() &&
@@ -1485,7 +1493,7 @@ void FDeferredShadingSceneRenderer::RenderDiffuseIndirectAndAmbientOcclusion(
 
 				if (bUseDepthBounds)
 				{
-					PassParameters->RenderTargets.DepthStencil = FDepthStencilBinding(SceneTextures.Depth.Target, 
+					PSParameters->RenderTargets.DepthStencil = FDepthStencilBinding(SceneTextures.Depth.Target,
 						ERenderTargetLoadAction::ELoad, ERenderTargetLoadAction::ELoad, FExclusiveDepthStencil::DepthRead_StencilRead);
 				}
 
@@ -1496,9 +1504,9 @@ void FDeferredShadingSceneRenderer::RenderDiffuseIndirectAndAmbientOcclusion(
 						PermutationVector.Get<FDiffuseIndirectCompositePS::FUpscaleDiffuseIndirectDim>() ? TEXT(" UpscaleDiffuseIndirect") : TEXT(""),
 						AmbientOcclusionMask ? TEXT(" ApplyAOToSceneColor") : TEXT(""),
 						View.ViewRect.Width(), View.ViewRect.Height()), 
-					PassParameters,
+					PSParameters,
 					ERDGPassFlags::Raster,
-					[&View, PassParameters,PixelShader, BlendState, bUseDepthBounds](FRDGAsyncTask, FRHICommandList& RHICmdList)
+					[&View, PSParameters,PixelShader, BlendState, bUseDepthBounds](FRDGAsyncTask, FRHICommandList& RHICmdList)
 					{
 						check(PixelShader.IsValid());
 						RHICmdList.SetViewport((float)View.ViewRect.Min.X, (float)View.ViewRect.Min.Y, 0.0f, (float)View.ViewRect.Max.X, (float)View.ViewRect.Max.Y, 1.0f);
@@ -1519,7 +1527,7 @@ void FDeferredShadingSceneRenderer::RenderDiffuseIndirectAndAmbientOcclusion(
 						}
 
 						SetGraphicsPipelineState(RHICmdList, GraphicsPSOInit, 0);
-						SetShaderParameters(RHICmdList, PixelShader, PixelShader.GetPixelShader(), *PassParameters);
+						SetShaderParameters(RHICmdList, PixelShader, PixelShader.GetPixelShader(), *PSParameters);
 
 						FPixelShaderUtils::DrawFullscreenTriangle(RHICmdList);
 
@@ -1536,8 +1544,8 @@ void FDeferredShadingSceneRenderer::RenderDiffuseIndirectAndAmbientOcclusion(
 				// Rough refraction targets
 				if (Substrate::IsOpaqueRoughRefractionEnabled())
 				{
-					PassParameters->OutOpaqueRoughRefractionSceneColor = GraphBuilder.CreateUAV(Scene->SubstrateSceneData.SeparatedOpaqueRoughRefractionSceneColor);
-					PassParameters->OutSubSurfaceSceneColor = GraphBuilder.CreateUAV(Scene->SubstrateSceneData.SeparatedSubSurfaceSceneColor);
+					PSParameters->OutOpaqueRoughRefractionSceneColor = GraphBuilder.CreateUAV(Scene->SubstrateSceneData.SeparatedOpaqueRoughRefractionSceneColor);
+					PSParameters->OutSubSurfaceSceneColor = GraphBuilder.CreateUAV(Scene->SubstrateSceneData.SeparatedSubSurfaceSceneColor);
 				}
 
 				Substrate::FSubstrateTilePassVS::FPermutationDomain VSPermutationVector;
@@ -1545,10 +1553,11 @@ void FDeferredShadingSceneRenderer::RenderDiffuseIndirectAndAmbientOcclusion(
 				VSPermutationVector.Set< Substrate::FSubstrateTilePassVS::FEnableTexCoordScreenVector >(false);
 				TShaderMapRef<Substrate::FSubstrateTilePassVS> TileVertexShader(View.ShaderMap, VSPermutationVector);
 
-				ClearUnusedGraphResources(PixelShader, PassParameters);
+				ClearUnusedGraphResources(PixelShader, PSParameters);
 
 				EPrimitiveType PrimitiveType = PT_TriangleList;
-				PassParameters->SubstrateTile = Substrate::SetTileParameters(GraphBuilder, View, TileType, PrimitiveType);
+				PassParameters->VS = Substrate::SetTileParameters(GraphBuilder, View, TileType, PrimitiveType);
+				ClearUnusedGraphResources(TileVertexShader, &PassParameters->VS);
 
 				GraphBuilder.AddPass(
 					RDG_EVENT_NAME(
@@ -1576,9 +1585,9 @@ void FDeferredShadingSceneRenderer::RenderDiffuseIndirectAndAmbientOcclusion(
 					GraphicsPSOInit.BoundShaderState.PixelShaderRHI = PixelShader.GetPixelShader();		
 					GraphicsPSOInit.DepthStencilState = TStaticDepthStencilState<false, CF_Always>::GetRHI();
 					SetGraphicsPipelineState(RHICmdList, GraphicsPSOInit, 0x0);
-					SetShaderParameters(RHICmdList, PixelShader, PixelShader.GetPixelShader(), *PassParameters);
-					SetShaderParameters(RHICmdList, TileVertexShader, TileVertexShader.GetVertexShader(), PassParameters->SubstrateTile);
-					RHICmdList.DrawPrimitiveIndirect(PassParameters->SubstrateTile.TileIndirectBuffer->GetIndirectRHICallBuffer(), Substrate::TileTypeDrawIndirectArgOffset(TileType));
+					SetShaderParameters(RHICmdList, PixelShader, PixelShader.GetPixelShader(), PassParameters->PS);
+					SetShaderParameters(RHICmdList, TileVertexShader, TileVertexShader.GetVertexShader(), PassParameters->VS);
+					RHICmdList.DrawPrimitiveIndirect(PassParameters->VS.TileIndirectBuffer->GetIndirectRHICallBuffer(), Substrate::TileTypeDrawIndirectArgOffset(TileType));
 				});
 			}
 
@@ -1607,27 +1616,28 @@ void FDeferredShadingSceneRenderer::RenderDiffuseIndirectAndAmbientOcclusion(
 		{
 			auto ApplyAmbientCubemapComposite = [&](ESubstrateTileType TileType)
 			{			
-				FAmbientCubemapCompositePS::FParameters* PassParameters = GraphBuilder.AllocParameters<FAmbientCubemapCompositePS::FParameters>();
-				PassParameters->PreIntegratedGF = GSystemTextures.PreintegratedGF->GetRHI();
-				PassParameters->PreIntegratedGFSampler = TStaticSamplerState<SF_Bilinear, AM_Clamp, AM_Clamp, AM_Clamp>::GetRHI();
+				FAmbientCubemapCompositeParameters* PassParameters = GraphBuilder.AllocParameters<FAmbientCubemapCompositeParameters>();
+				FAmbientCubemapCompositePS::FParameters* PSParameters = &PassParameters->PS;
+				PSParameters->PreIntegratedGF = GSystemTextures.PreintegratedGF->GetRHI();
+				PSParameters->PreIntegratedGFSampler = TStaticSamplerState<SF_Bilinear, AM_Clamp, AM_Clamp, AM_Clamp>::GetRHI();
 			
-				PassParameters->AmbientOcclusionTexture = AmbientOcclusionMask;
-				PassParameters->AmbientOcclusionSampler = TStaticSamplerState<SF_Point>::GetRHI();
+				PSParameters->AmbientOcclusionTexture = AmbientOcclusionMask;
+				PSParameters->AmbientOcclusionSampler = TStaticSamplerState<SF_Point>::GetRHI();
 			
-				if (!PassParameters->AmbientOcclusionTexture)
+				if (!PSParameters->AmbientOcclusionTexture)
 				{
-					PassParameters->AmbientOcclusionTexture = SystemTextures.White;
+					PSParameters->AmbientOcclusionTexture = SystemTextures.White;
 				}
 
-				PassParameters->Substrate = Substrate::BindSubstrateGlobalUniformParameters(View);
-				PassParameters->SceneTextures = SceneTextureParameters;
-				PassParameters->ViewUniformBuffer = View.ViewUniformBuffer;
+				PSParameters->Substrate = Substrate::BindSubstrateGlobalUniformParameters(View);
+				PSParameters->SceneTextures = SceneTextureParameters;
+				PSParameters->ViewUniformBuffer = View.ViewUniformBuffer;
 
-				PassParameters->RenderTargets[0] = FRenderTargetBinding(SceneColorTexture, ERenderTargetLoadAction::ELoad);
+				PSParameters->RenderTargets[0] = FRenderTargetBinding(SceneColorTexture, ERenderTargetLoadAction::ELoad);
 				if (Substrate::IsOpaqueRoughRefractionEnabled())
 				{
-					PassParameters->RenderTargets[1] = FRenderTargetBinding(Scene->SubstrateSceneData.SeparatedOpaqueRoughRefractionSceneColor, ERenderTargetLoadAction::ELoad);
-					PassParameters->RenderTargets[2] = FRenderTargetBinding(Scene->SubstrateSceneData.SeparatedSubSurfaceSceneColor, ERenderTargetLoadAction::ELoad);
+					PSParameters->RenderTargets[1] = FRenderTargetBinding(Scene->SubstrateSceneData.SeparatedOpaqueRoughRefractionSceneColor, ERenderTargetLoadAction::ELoad);
+					PSParameters->RenderTargets[2] = FRenderTargetBinding(Scene->SubstrateSceneData.SeparatedSubSurfaceSceneColor, ERenderTargetLoadAction::ELoad);
 				}
 
 				TShaderMapRef<FPostProcessVS> FullScreenVertexShader(View.ShaderMap);
@@ -1641,7 +1651,7 @@ void FDeferredShadingSceneRenderer::RenderDiffuseIndirectAndAmbientOcclusion(
 				EPrimitiveType PrimitiveType = PT_TriangleList;
 				if (Substrate::IsSubstrateEnabled())
 				{
-					PassParameters->SubstrateTile = Substrate::SetTileParameters(GraphBuilder, View, TileType, PrimitiveType);
+					PassParameters->VS = Substrate::SetTileParameters(GraphBuilder, View, TileType, PrimitiveType);
 				}
 
 				TShaderMapRef<FAmbientCubemapCompositePS> PixelShader(View.ShaderMap);
@@ -1671,14 +1681,14 @@ void FDeferredShadingSceneRenderer::RenderDiffuseIndirectAndAmbientOcclusion(
 					uint32 Count = View.FinalPostProcessSettings.ContributingCubemaps.Num();
 					for (const FFinalPostProcessSettings::FCubemapEntry& CubemapEntry : View.FinalPostProcessSettings.ContributingCubemaps)
 					{
-						FAmbientCubemapCompositePS::FParameters ShaderParameters = *PassParameters;
+						FAmbientCubemapCompositePS::FParameters ShaderParameters = PassParameters->PS;
 						SetupAmbientCubemapParameters(CubemapEntry, &ShaderParameters.AmbientCubemap);
 						SetShaderParameters(RHICmdList, PixelShader, PixelShader.GetPixelShader(), ShaderParameters);
 					
 						if (bSubstrateEnabled)
 						{
-							SetShaderParameters(RHICmdList, TileVertexShader, TileVertexShader.GetVertexShader(), PassParameters->SubstrateTile);
-							RHICmdList.DrawPrimitiveIndirect(PassParameters->SubstrateTile.TileIndirectBuffer->GetIndirectRHICallBuffer(), Substrate::TileTypeDrawIndirectArgOffset(TileType));
+							SetShaderParameters(RHICmdList, TileVertexShader, TileVertexShader.GetVertexShader(), PassParameters->VS);
+							RHICmdList.DrawPrimitiveIndirect(PassParameters->VS.TileIndirectBuffer->GetIndirectRHICallBuffer(), Substrate::TileTypeDrawIndirectArgOffset(TileType));
 						}
 						else
 						{
diff --git a/Engine/Source/Runtime/Renderer/Private/Substrate/Substrate.cpp b/Engine/Source/Runtime/Renderer/Private/Substrate/Substrate.cpp
--- a/Engine/Source/Runtime/Renderer/Private/Substrate/Substrate.cpp
+++ b/Engine/Source/Runtime/Renderer/Private/Substrate/Substrate.cpp
@@ -1105,27 +1105,7 @@ static FSubstrateTilePassVS::FParameters SetTileParametersCommon(FRDGBuilder* Gr
 	PrimitiveType = GRHISupportsRectTopology ? PT_RectList : PT_TriangleList;
 
 	FSubstrateTilePassVS::FParameters Out;
-	Out.OutputViewMinRect = FVector2f(View.CachedViewUniformShaderParameters->ViewRectMin.X, View.CachedViewUniformShaderParameters->ViewRectMin.Y);
-	Out.OutputViewSizeAndInvSize = View.CachedViewUniformShaderParameters->ViewSizeAndInvSize;
-	Out.OutputBufferSizeAndInvSize = View.CachedViewUniformShaderParameters->BufferSizeAndInvSize;
-
-	if (View.IsPerspectiveProjection())
-	{
-		Out.ViewScreenToTranslatedWorld = View.CachedViewUniformShaderParameters->ScreenToTranslatedWorld;
-	}
-	else
-	{
-		/**
-		* This forces the calculation for orthographic cameras to use the hard coded screen depth == 1 to resolve the view 
-		* without having to introduce a new shader parameter.
-		* For now this is fine but this should be changed to use ScreenVectorFromScreenRect, matching other passes, once the whole view buffer 
-		* is being passed to shader params. Note: will not be compatible with reverse Z if the hard coded behaviour is altered.
-		*/
-		FMatrix44f OrthoMatrix(EForceInit::ForceInitToZero);
-		OrthoMatrix.SetAxis(2, View.CachedViewUniformShaderParameters->ViewForward);
-		Out.ViewScreenToTranslatedWorld = OrthoMatrix;
-	}
-
+	Out.View = View.ViewUniformBuffer;
 	Out.TileListBuffer = Temp.TileListBuffer;
 	Out.TileListBufferOffset = Temp.TileListBufferOffset;
 	Out.TileEncoding = Temp.TileEncoding;
diff --git a/Engine/Source/Runtime/Renderer/Private/Substrate/Substrate.h b/Engine/Source/Runtime/Renderer/Private/Substrate/Substrate.h
--- a/Engine/Source/Runtime/Renderer/Private/Substrate/Substrate.h
+++ b/Engine/Source/Runtime/Renderer/Private/Substrate/Substrate.h
@@ -11,6 +11,7 @@
 #include "GBufferInfo.h"
 #include "RendererUtils.h"
 #include "ShaderParameterStruct.h"
+#include "SceneView.h"
 
 // Forward declarations.
 class FScene;
@@ -269,13 +270,7 @@ class FSubstrateTilePassVS : public FGlobalShader
 	using FPermutationDomain = TShaderPermutationDomain<FEnableDebug, FEnableTexCoordScreenVector>;
 
 	BEGIN_SHADER_PARAMETER_STRUCT(FParameters, )
-		// It would be possible to use the ViewUniformBuffer instead of copying the data here, 
-		// but we would have to make sure the view UB is added to all passes using this parameter structure.
-		// We should not add it here to now have duplicated input UB.
-		SHADER_PARAMETER(FVector2f, OutputViewMinRect)
-		SHADER_PARAMETER(FVector4f, OutputViewSizeAndInvSize)
-		SHADER_PARAMETER(FVector4f, OutputBufferSizeAndInvSize)
-		SHADER_PARAMETER(FMatrix44f, ViewScreenToTranslatedWorld)
+		SHADER_PARAMETER_STRUCT_REF(FViewUniformShaderParameters, View)
 		SHADER_PARAMETER_RDG_BUFFER_SRV(Buffer<uint>, TileListBuffer)
 		SHADER_PARAMETER(uint32, TileListBufferOffset)
 		SHADER_PARAMETER(uint32, TileEncoding)
