Fixed regression in PSO memory by flushing the final renderer in a chain of scene renderers.

diff --git a/Engine/Source/Runtime/Renderer/Private/SceneRenderBuilder.cpp b/Engine/Source/Runtime/Renderer/Private/SceneRenderBuilder.cpp
--- a/Engine/Source/Runtime/Renderer/Private/SceneRenderBuilder.cpp
+++ b/Engine/Source/Runtime/Renderer/Private/SceneRenderBuilder.cpp
@@ -363,7 +363,7 @@ void FSceneRenderProcessor::Execute()
 				}
 
 				// The final graph builder is responsible for flushing resources.
-				if (!RenderNode.GetNext())
+				if (&Renderer == Renderers.Last())
 				{
 					GraphBuilder.SetFlushResourcesRHI();
 				}
