Cleanup spammy render trace scopes.

+ Remove duplicate scopes where were were using SCOPE_CYCLE_COUNTER and SCOPED_NAMED_EVENT.
+ Mark STAT_D3D12LockTextureTime, STAT_D3D12UnlockTextureTime, STAT_D3D12LockBufferTime and STAT_D3D12UnlockBufferTime as verbose as we can have thousands in a frame and they're very short.
+ Mark STAT_FRDGPass_Execute as verbose as it's also very spammy.

#rnx

diff --git a/Engine/Source/Runtime/D3D12RHI/Private/D3D12Stats.h b/Engine/Source/Runtime/D3D12RHI/Private/D3D12Stats.h
--- a/Engine/Source/Runtime/D3D12RHI/Private/D3D12Stats.h
+++ b/Engine/Source/Runtime/D3D12RHI/Private/D3D12Stats.h
@@ -28,12 +28,12 @@ DECLARE_DWORD_ACCUMULATOR_STAT_EXTERN(TEXT("Num pipeline state objects (PSOs)"),
 DECLARE_DWORD_COUNTER_STAT_EXTERN(TEXT("Textures Allocated"), STAT_D3D12TexturesAllocated, STATGROUP_D3D12RHI, );
 DECLARE_DWORD_COUNTER_STAT_EXTERN(TEXT("Textures Released"), STAT_D3D12TexturesReleased, STATGROUP_D3D12RHI, );
 DECLARE_CYCLE_STAT_EXTERN(TEXT("CreateTexture time"), STAT_D3D12CreateTextureTime, STATGROUP_D3D12RHI, );
-DECLARE_CYCLE_STAT_EXTERN(TEXT("LockTexture time"), STAT_D3D12LockTextureTime, STATGROUP_D3D12RHI, );
-DECLARE_CYCLE_STAT_EXTERN(TEXT("UnlockTexture time"), STAT_D3D12UnlockTextureTime, STATGROUP_D3D12RHI, );
+DECLARE_CYCLE_STAT_WITH_FLAGS_EXTERN(TEXT("LockTexture time"), STAT_D3D12LockTextureTime, STATGROUP_D3D12RHI, EStatFlags::Verbose, );
+DECLARE_CYCLE_STAT_WITH_FLAGS_EXTERN(TEXT("UnlockTexture time"), STAT_D3D12UnlockTextureTime, STATGROUP_D3D12RHI, EStatFlags::Verbose, );
 DECLARE_CYCLE_STAT_EXTERN(TEXT("CreateBuffer time"), STAT_D3D12CreateBufferTime, STATGROUP_D3D12RHI, );
 DECLARE_CYCLE_STAT_EXTERN(TEXT("CopyToStagingBuffer time"), STAT_D3D12CopyToStagingBufferTime, STATGROUP_D3D12RHI, );
-DECLARE_CYCLE_STAT_EXTERN(TEXT("LockBuffer time"), STAT_D3D12LockBufferTime, STATGROUP_D3D12RHI, );
-DECLARE_CYCLE_STAT_EXTERN(TEXT("UnlockBuffer time"), STAT_D3D12UnlockBufferTime, STATGROUP_D3D12RHI, );
+DECLARE_CYCLE_STAT_WITH_FLAGS_EXTERN(TEXT("LockBuffer time"), STAT_D3D12LockBufferTime, STATGROUP_D3D12RHI, EStatFlags::Verbose, );
+DECLARE_CYCLE_STAT_WITH_FLAGS_EXTERN(TEXT("UnlockBuffer time"), STAT_D3D12UnlockBufferTime, STATGROUP_D3D12RHI, EStatFlags::Verbose, );
 DECLARE_CYCLE_STAT_EXTERN(TEXT("Commit transient resource time"), STAT_D3D12CommitTransientResourceTime, STATGROUP_D3D12RHI, );
 DECLARE_CYCLE_STAT_EXTERN(TEXT("Decommit transient resource time"), STAT_D3D12DecommitTransientResourceTime, STATGROUP_D3D12RHI, );
 
diff --git a/Engine/Source/Runtime/RenderCore/Public/RenderGraphPass.h b/Engine/Source/Runtime/RenderCore/Public/RenderGraphPass.h
--- a/Engine/Source/Runtime/RenderCore/Public/RenderGraphPass.h
+++ b/Engine/Source/Runtime/RenderCore/Public/RenderGraphPass.h
@@ -705,7 +705,9 @@ private:
 	void Execute(FRHIComputeCommandList& RHICmdList) override
 	{
 #if !USE_NULL_RHI
-		QUICK_SCOPE_CYCLE_COUNTER(STAT_FRDGPass_Execute);
+		DECLARE_CYCLE_STAT_WITH_FLAGS(TEXT("FRDGPass Execute"), STAT_FRDGPass_Execute, STATGROUP_RHI, EStatFlags::Verbose);
+
+		SCOPE_CYCLE_COUNTER(STAT_FRDGPass_Execute);
 		RHICmdList.SetStaticUniformBuffers(ParameterStruct.GetStaticUniformBuffers());
 		ExecuteLambdaFunc(static_cast<typename ExecuteLambdaTraits::TRHICommandList&>(RHICmdList));
 #else
diff --git a/Engine/Source/Runtime/Renderer/Private/RendererScene.cpp b/Engine/Source/Runtime/Renderer/Private/RendererScene.cpp
--- a/Engine/Source/Runtime/Renderer/Private/RendererScene.cpp
+++ b/Engine/Source/Runtime/Renderer/Private/RendererScene.cpp
@@ -1863,7 +1863,6 @@ void FScene::BatchAddPrimitivesInternal(TArrayView<T*> InPrimitives)
 	UE_TRACE_METADATA_SCOPE_ASSET_FNAME(NAME_None, NAME_None, InPrimitives[0]->GetOutermost()->GetFName());
 
 	SCOPE_CYCLE_COUNTER(STAT_AddScenePrimitiveGT);
-	SCOPED_NAMED_EVENT(FScene_AddPrimitive, FColor::Green);
 
 	struct FCreateCommand
 	{
@@ -2067,7 +2066,6 @@ template<class T>
 void FScene::UpdatePrimitiveTransformInternal(T* Primitive)
 {
 	SCOPE_CYCLE_COUNTER(STAT_UpdatePrimitiveTransformGT);
-	SCOPED_NAMED_EVENT(FScene_UpdatePrimitiveTransform, FColor::Yellow);
 
 	FPrimitiveSceneInfoData& SceneData = Primitive->GetSceneData();
 
@@ -2227,7 +2225,6 @@ void FScene::UpdateInstanceCullDistance(UPrimitiveComponent* Primitive, float St
 void FScene::UpdatePrimitiveInstances(UPrimitiveComponent* Primitive)
 {
 	SCOPE_CYCLE_COUNTER(STAT_UpdatePrimitiveInstanceGT);
-	SCOPED_NAMED_EVENT(FScene_UpdatePrimitiveInstance, FColor::Yellow);	
 
 	// If the primitive doesn't have a scene info object yet, it must be added from scratch.
 	if (!Primitive->GetSceneProxy())
@@ -2277,7 +2274,6 @@ void FScene::UpdatePrimitiveInstances(FInstancedStaticMeshSceneDesc* Primitive)
 void FScene::UpdatePrimitiveInstances(FPrimitiveSceneDesc* Primitive)
 {
 	SCOPE_CYCLE_COUNTER(STAT_UpdatePrimitiveInstanceGT);
-	SCOPED_NAMED_EVENT(FScene_UpdatePrimitiveInstance, FColor::Yellow);
 
 	// If the primitive doesn't have a scene info object yet, it must be added from scratch.
 	if (!Primitive->GetSceneProxy())
@@ -2427,7 +2423,6 @@ template <class T>
 void FScene::BatchRemovePrimitivesInternal(TArrayView<T*> InPrimitives)
 {
 	SCOPE_CYCLE_COUNTER(STAT_RemoveScenePrimitiveGT);
-	SCOPED_NAMED_EVENT(FScene_RemovePrimitive, FColor::Yellow);
 
 	struct FDetachCommand
 	{
@@ -2589,7 +2584,6 @@ void FScene::AssignAvailableShadowMapChannelForLight(FLightSceneInfo* LightScene
 void FScene::AddLightSceneInfo_RenderThread(FLightSceneInfo* LightSceneInfo)
 {
 	SCOPE_CYCLE_COUNTER(STAT_AddSceneLightTime);
-	SCOPED_NAMED_EVENT(FScene_AddLightSceneInfo_RenderThread, FColor::Green);
 
 	check(LightSceneInfo->bVisible);
 
@@ -3920,7 +3914,6 @@ void FScene::UpdateLightColorAndBrightness(ULightComponent* Light)
 void FScene::RemoveLightSceneInfo_RenderThread(FLightSceneInfo* LightSceneInfo)
 {
 	SCOPE_CYCLE_COUNTER(STAT_RemoveSceneLightTime);
-	SCOPED_NAMED_EVENT(FScene_RemoveLightSceneInfo_RenderThread, FColor::Red);
 
 	check(LightSceneInfo->bVisible);
 
@@ -4511,7 +4504,6 @@ void FScene::SetPrecomputedVisibility(const FPrecomputedVisibilityHandler* NewPr
 void FScene::UpdateStaticDrawLists_RenderThread(FRHICommandListImmediate& RHICmdList)
 {
 	SCOPE_CYCLE_COUNTER(STAT_Scene_UpdateStaticDrawLists_RT);
-	SCOPED_NAMED_EVENT(FScene_UpdateStaticDrawLists_RenderThread, FColor::Yellow);
 
 	const int32 NumPrimitives = Primitives.Num();
 
@@ -4791,7 +4783,6 @@ void FScene::ApplyWorldOffset(const FVector& InOffset)
 void FScene::ApplyWorldOffset_RenderThread(FRHICommandListBase& RHICmdList, const FVector& InOffset)
 {
 	QUICK_SCOPE_CYCLE_COUNTER(STAT_SceneApplyWorldOffset);
-	SCOPED_NAMED_EVENT(FScene_ApplyWorldOffset_RenderThread, FColor::Yellow);
 	
 	GPUScene.bUpdateAllPrimitives = true;
 
@@ -5296,7 +5287,6 @@ static void CreateReflectionCaptureUniformBuffer(const TArray<FReflectionCapture
 void UpdateReflectionSceneData(FScene* Scene)
 {
 	QUICK_SCOPE_CYCLE_COUNTER(STAT_UpdateReflectionSceneData)
-	SCOPED_NAMED_EVENT(UpdateReflectionScene, FColor::Red);
 
 	FReflectionEnvironmentSceneData& ReflectionSceneData = Scene->ReflectionSceneData;
 
@@ -5581,8 +5571,6 @@ FViewSceneChangeSet* FScene::ProcessViewChanges(FRDGBuilder& GraphBuilder, TArra
 void FScene::Update(FRDGBuilder& GraphBuilder, const FUpdateParameters& Parameters)
 {
 	LLM_SCOPE(ELLMTag::SceneRender);
-	TRACE_CPUPROFILER_EVENT_SCOPE(Scene::UpdateAllPrimitiveSceneInfos);
-	SCOPED_NAMED_EVENT(FScene_UpdateAllPrimitiveSceneInfos, FColor::Orange);
 	SCOPE_CYCLE_COUNTER(STAT_UpdateScenePrimitiveRenderThreadTime);
 
 	check(IsInRenderingThread());
@@ -5715,7 +5703,6 @@ void FScene::Update(FRDGBuilder& GraphBuilder, const FUpdateParameters& Paramete
 	bool bNeedPathTracedInvalidation = false;
 	{
 		CSV_SCOPED_TIMING_STAT_EXCLUSIVE(RemovePrimitiveSceneInfos);
-		SCOPED_NAMED_EVENT(FScene_RemovePrimitiveSceneInfos, FColor::Red);
 		SCOPE_CYCLE_COUNTER(STAT_RemoveScenePrimitiveTime);
 
 		for (FPrimitiveSceneInfo* PrimitiveSceneInfo : RemovedLocalPrimitiveSceneInfos)
@@ -5988,7 +5975,6 @@ void FScene::Update(FRDGBuilder& GraphBuilder, const FUpdateParameters& Paramete
 	// Perform instance releases before adds, such that allocator consolidation can happen only once.
 	{
 		CSV_SCOPED_TIMING_STAT_EXCLUSIVE(UpdatePrimitiveInstances);
-		SCOPED_NAMED_EVENT(FScene_UpdatePrimitiveInstances, FColor::Emerald);
 		SCOPE_CYCLE_COUNTER(STAT_UpdatePrimitiveInstanceRenderThreadTime);
 
 		for (const auto& UpdateInstance : UpdatedInstances)
@@ -6013,7 +5999,6 @@ void FScene::Update(FRDGBuilder& GraphBuilder, const FUpdateParameters& Paramete
 
 	{
 		CSV_SCOPED_TIMING_STAT_EXCLUSIVE(AddPrimitiveSceneInfos);
-		SCOPED_NAMED_EVENT(FScene_AddPrimitiveSceneInfos, FColor::Green);
 		SCOPE_CYCLE_COUNTER(STAT_AddScenePrimitiveRenderThreadTime);
 
 		PersistentPrimitiveIdAllocator.Consolidate();
@@ -6269,7 +6254,6 @@ void FScene::Update(FRDGBuilder& GraphBuilder, const FUpdateParameters& Paramete
 	}
 	{
 		CSV_SCOPED_TIMING_STAT_EXCLUSIVE(UpdatePrimitiveTransform);
-		SCOPED_NAMED_EVENT(FScene_AddPrimitiveSceneInfos, FColor::Yellow);
 		SCOPE_CYCLE_COUNTER(STAT_UpdatePrimitiveTransformRenderThreadTime);
 
 
@@ -6345,7 +6329,6 @@ void FScene::Update(FRDGBuilder& GraphBuilder, const FUpdateParameters& Paramete
 	// Updating step of instance update. Note that instance IDs are not allocated at this point.
 	{
 		CSV_SCOPED_TIMING_STAT_EXCLUSIVE(UpdatePrimitiveInstances);
-		SCOPED_NAMED_EVENT(FScene_UpdatePrimitiveInstances, FColor::Emerald);
 		SCOPE_CYCLE_COUNTER(STAT_UpdatePrimitiveInstanceRenderThreadTime);
 
 		for (const auto& Item : UpdatedInstances)
