XR Passthrough alpha blending incorrect on some platforms.
-The common issue is that Unreal uses 0=opaque 1=transparent for alpha values rather than the reverse which is much more common.
-Here I extend and refactor the AlphaInvert pass I originally added only to the Mobile renderer to the desktop renderer.
-A new AlphaInvert Engine ShowFlag will cause the pass to run at the end of postprocessing in the desktop renderer.  It also runs at the end of the mobile renderer.  OpenXRHMD will set this showflag if the CVAR r.AlphaInvert is true.
-In principle setting this showflag could also be useful for other scenarios where we are outputting with alpha.

diff --git a/Engine/Plugins/Runtime/OpenXR/Source/OpenXRHMD/Private/OpenXRHMD.cpp b/Engine/Plugins/Runtime/OpenXR/Source/OpenXRHMD/Private/OpenXRHMD.cpp
--- a/Engine/Plugins/Runtime/OpenXR/Source/OpenXRHMD/Private/OpenXRHMD.cpp
+++ b/Engine/Plugins/Runtime/OpenXR/Source/OpenXRHMD/Private/OpenXRHMD.cpp
@@ -1430,6 +1430,21 @@ void FOpenXRHMD::SetupViewFamily(FSceneViewFamily& InViewFamily)
 	InViewFamily.EngineShowFlags.HMDDistortion = false;
 	InViewFamily.EngineShowFlags.StereoRendering = IsStereoEnabled();
 
+	// For now we are enabling the invert alpha pass with a cvar.
+	// However it seems likely that we might know we are providing alpha to openxr, and that
+	// the current runtime does not support unreal's inverted alpha and set this based on that.
+	if (InViewFamily.Scene)
+	{
+		EShaderPlatform Platform = InViewFamily.Scene->GetShaderPlatform();
+		static FShaderPlatformCachedIniValue<bool> AlphaInvertPassIniValue(TEXT("r.AlphaInvertPass"));
+		InViewFamily.EngineShowFlags.AlphaInvert = AlphaInvertPassIniValue.Get(Platform);
+	}
+	else
+	{
+		static auto CVarAlphaInvertPassLocal = IConsoleManager::Get().FindConsoleVariable(TEXT("r.AlphaInvertPass"));
+		InViewFamily.EngineShowFlags.AlphaInvert = (CVarAlphaInvertPassLocal != nullptr) && CVarAlphaInvertPassLocal->GetBool();
+	}
+
 	FPipelinedFrameStateAccessorReadOnly LockedPipelineState = (const_cast<const FOpenXRHMD*>(this))->GetPipelinedFrameStateForThread();
 	const FPipelinedFrameState& FrameState = LockedPipelineState.GetFrameState();
 	if (FrameState.Views.Num() > 2)
diff --git a/Engine/Shaders/Private/PostProcessAlphaInvert.usf b/Engine/Shaders/Private/PostProcessAlphaInvert.usf
new file mode 100644
--- /dev/null
+++ b/Engine/Shaders/Private/PostProcessAlphaInvert.usf
@@ -0,0 +1,41 @@
+// Copyright Epic Games, Inc. All Rights Reserved.
+
+/*=============================================================================
+	PostProcessAlphaInvert.usf: Invert the alpha values.
+	This can be used to convert from Unreal standard 0=opaque 1=transparent to the much more common 
+	standard of 1=opaque 0=transparent.
+=============================================================================*/
+
+#include "Common.ush"
+
+#if MOBILE_MULTI_VIEW
+Texture2DArray InTexture;
+#else
+Texture2D InTexture;
+#endif
+
+SamplerState InTextureSampler;
+
+void AlphaInvert_MainPS(
+	FScreenVertexOutput								Input,
+	out				 float4 OutColor				: SV_Target0
+#if  MOBILE_MULTI_VIEW
+	, in nointerpolation uint MultiViewId			: SV_ViewID
+#endif
+)
+{
+	uint ArrayIndex = 0;
+#if MOBILE_MULTI_VIEW
+	ArrayIndex = MultiViewId;
+#endif
+	
+	float4 SceneColor;
+	#if MOBILE_MULTI_VIEW
+		SceneColor = Texture2DArraySample(InTexture, InTextureSampler, float3(Input.UV.xy,ArrayIndex));
+	#else
+		SceneColor = Texture2DSample(InTexture, InTextureSampler, Input.UV);
+	#endif
+	
+	OutColor.rgb = SceneColor.rgb;
+	OutColor.a = 1.0 - SceneColor.a;
+}
diff --git a/Engine/Shaders/Private/PostProcessTonemap.usf b/Engine/Shaders/Private/PostProcessTonemap.usf
--- a/Engine/Shaders/Private/PostProcessTonemap.usf
+++ b/Engine/Shaders/Private/PostProcessTonemap.usf
@@ -746,33 +746,3 @@ void MobileCustomResolve_MainPS(
 	OutColor.rgb = DeviceColor;
 	OutColor.a = 1.0;
 }
-
-void AlphaInvert_MainPS(
-	in float4 SvPosition							: SV_POSITION, 
-	in noperspective float2 UV						: TEXCOORD0,
-	out				 float4 OutColor				: SV_Target0
-#if MOBILE_MULTI_VIEW_FALLBACK
-	, in nointerpolation uint EyeIndex				: VIEW_ID
-#elif MOBILE_MULTI_VIEW
-	, in nointerpolation uint MultiViewId : SV_ViewID
-#endif
-)
-{
-	uint ArrayIndex = 0;
-#if MOBILE_MULTI_VIEW_FALLBACK
-	ArrayIndex = EyeIndex;
-#elif MOBILE_MULTI_VIEW
-	ArrayIndex = MultiViewId;
-#else
-#endif
-	
-	half4 SceneColor;
-	#if MOBILE_MULTI_VIEW || MOBILE_MULTI_VIEW_FALLBACK
-		SceneColor = Texture2DArraySample(ColorTexture, ColorSampler, float3(UV.xy,ArrayIndex));
-	#else
-		SceneColor = Texture2DSample(ColorTexture, ColorSampler, UV);
-	#endif
-	
-	OutColor.rgb = SceneColor.rgb;
-	OutColor.a = 1.0 - SceneColor.a;
-}
diff --git a/Engine/Source/Runtime/Engine/Public/ShowFlagsValues.inl b/Engine/Source/Runtime/Engine/Public/ShowFlagsValues.inl
--- a/Engine/Source/Runtime/Engine/Public/ShowFlagsValues.inl
+++ b/Engine/Source/Runtime/Engine/Public/ShowFlagsValues.inl
@@ -469,5 +469,8 @@ SHOWFLAG_ALWAYS_ACCESSIBLE(Hair, SFG_Normal, NSLOCTEXT("UnrealEd", "HairSF", "Ha
 /** Allows primitive alpha holdout rendering. Note that the "Support Primitive Alpha Holdout" project setting first needs to be enabled for the deferred renderer. */
 SHOWFLAG_ALWAYS_ACCESSIBLE(AllowPrimitiveAlphaHoldout, SFG_Advanced, NSLOCTEXT("UnrealEd", "AllowPrimitiveAlphaHoldoutSF", "Allow Primitive Alpha Holdout"))
 
+/** Invert the alpha channel at tne end of post processing. */
+SHOWFLAG_ALWAYS_ACCESSIBLE(AlphaInvert, SFG_PostProcess, NSLOCTEXT("UnrealEd", "AlphaInvertSF", "Invert Alpha Channel"))
+
 #undef SHOWFLAG_ALWAYS_ACCESSIBLE
 #undef SHOWFLAG_FIXED_IN_SHIPPING
diff --git a/Engine/Source/Runtime/Renderer/Private/MobileShadingRenderer.cpp b/Engine/Source/Runtime/Renderer/Private/MobileShadingRenderer.cpp
--- a/Engine/Source/Runtime/Renderer/Private/MobileShadingRenderer.cpp
+++ b/Engine/Source/Runtime/Renderer/Private/MobileShadingRenderer.cpp
@@ -36,6 +36,7 @@
 #include "PostProcess/PostProcessAmbientOcclusionMobile.h"
 #include "PostProcess/PostProcessCombineLUTs.h"
 #include "PostProcess/PostProcessTonemap.h"
+#include "PostProcess/AlphaInvert.h"
 #include "IHeadMountedDisplay.h"
 #include "IXRTrackingSystem.h"
 #include "SceneViewExtension.h"
@@ -1663,10 +1664,9 @@ void FMobileSceneRenderer::RenderForward(FRDGBuilder& GraphBuilder, FRDGTextureR
 		}
 	}
 	
-	static auto CVarAlphaInvertPassLocal = IConsoleManager::Get().FindConsoleVariable(TEXT("r.AlphaInvertPass"));
-	if(CVarAlphaInvertPassLocal != nullptr && CVarAlphaInvertPassLocal->GetBool())
+	if (ViewFamily.EngineShowFlags.AlphaInvert)
 	{
-		AddAlphaInvertPass(GraphBuilder, MainView, SceneTextures);
+		AlphaInvert::AddAlphaInvertPass(GraphBuilder, MainView, SceneTextures);
 	}
 }
 
diff --git a/Engine/Source/Runtime/Renderer/Private/PostProcess/AlphaInvert.cpp b/Engine/Source/Runtime/Renderer/Private/PostProcess/AlphaInvert.cpp
new file mode 100644
--- /dev/null
+++ b/Engine/Source/Runtime/Renderer/Private/PostProcess/AlphaInvert.cpp
@@ -0,0 +1,173 @@
+// Copyright Epic Games, Inc. All Rights Reserved.
+
+/*=============================================================================
+	AlphaInvert.cpp: Inverts alpha channel in a render pass.
+=============================================================================*/
+
+#include "AlphaInvert.h"
+#include "ScenePrivate.h"
+#include "PostProcess/SceneFilterRendering.h"
+#include "ScreenRendering.h"
+
+namespace AlphaInvert
+{
+	static TAutoConsoleVariable<int32> CVarAlphaInvertPass(
+		TEXT("r.AlphaInvertPass"),
+		0,
+		TEXT("Whether to run a render pass to un-invert the alpha value from unreal standard to the much more common standard where alpha 0 is fully transparent and alpha 1 is fully opaque."),
+		ECVF_RenderThreadSafe);
+
+	class FAlphaInvertPS : public FScreenPS
+	{
+		DECLARE_SHADER_TYPE(FAlphaInvertPS, Global);
+	public:
+		FAlphaInvertPS(const ShaderMetaType::CompiledShaderInitializerType& Initializer) :
+			FScreenPS(Initializer)
+		{
+		}
+		FAlphaInvertPS() {}
+	};
+	IMPLEMENT_SHADER_TYPE(, FAlphaInvertPS, TEXT("/Engine/Private/PostProcessAlphaInvert.usf"), TEXT("AlphaInvert_MainPS"), SF_Pixel);
+
+	void RenderAlphaInvertPass(FRHICommandList& RHICmdList, const FViewInfo& View, const FScreenPassTexture& Color)
+	{
+		// Part of scene rendering pass
+		check(RHICmdList.IsInsideRenderPass());
+		SCOPED_DRAW_EVENT(RHICmdList, AlphaInvert);
+
+		const FIntPoint OutputTextureRect = Color.Texture->Desc.Extent;
+		TShaderMapRef<FScreenVS> VertexShader(View.ShaderMap);
+		TShaderMapRef<FAlphaInvertPS> PixelShader(View.ShaderMap);
+
+		FGraphicsPipelineStateInitializer GraphicsPSOInit;
+		RHICmdList.ApplyCachedRenderTargets(GraphicsPSOInit);
+		GraphicsPSOInit.BlendState = TStaticBlendState<>::GetRHI();
+		GraphicsPSOInit.RasterizerState = TStaticRasterizerState<>::GetRHI();
+		GraphicsPSOInit.DepthStencilState = TStaticDepthStencilState<false, CF_Always>::GetRHI();
+		GraphicsPSOInit.BoundShaderState.VertexDeclarationRHI = GFilterVertexDeclaration.VertexDeclarationRHI;
+		GraphicsPSOInit.BoundShaderState.VertexShaderRHI = VertexShader.GetVertexShader();
+		GraphicsPSOInit.BoundShaderState.PixelShaderRHI = PixelShader.GetPixelShader();
+		GraphicsPSOInit.PrimitiveType = PT_TriangleList;
+
+		const FIntRect OutputViewRect(Color.ViewRect);
+
+		SetGraphicsPipelineState(RHICmdList, GraphicsPSOInit, 0);
+		SetShaderParametersLegacyPS(RHICmdList, PixelShader, TStaticSamplerState<SF_Point, AM_Clamp, AM_Clamp, AM_Clamp>::GetRHI(), Color.Texture->GetRHI());
+		RHICmdList.SetViewport(OutputViewRect.Min.X, OutputViewRect.Min.Y, 0.0f, OutputViewRect.Max.X, OutputViewRect.Max.Y, 1.0f);
+
+		DrawRectangle(
+			RHICmdList,
+			// Output rect, relative to rhi viewport.
+			0, 0, OutputViewRect.Width(), OutputViewRect.Height(),
+			// Input rect, relative to the input texture
+			OutputViewRect.Min.X, OutputViewRect.Min.Y, OutputViewRect.Width(), OutputViewRect.Height(),
+			OutputViewRect.Size(),
+			OutputTextureRect,
+			VertexShader,
+			EDRF_UseTriangleOptimization,
+			1);
+	}
+
+	BEGIN_SHADER_PARAMETER_STRUCT(FAlphaInvertParameters, )
+		SHADER_PARAMETER_STRUCT_INCLUDE(FViewShaderParameters, View)
+		SHADER_PARAMETER_RDG_TEXTURE(Texture2D, ColorTexture)
+		RENDER_TARGET_BINDING_SLOTS()
+	END_SHADER_PARAMETER_STRUCT()
+}
+
+FScreenPassTexture AlphaInvert::AddAlphaInvertPass(FRDGBuilder& GraphBuilder, const FViewInfo& View, const struct FAlphaInvertInputs& Inputs)
+{
+	FAlphaInvertParameters* PassParameters = GraphBuilder.AllocParameters<FAlphaInvertParameters>();
+	PassParameters->View = View.GetShaderParameters();
+	PassParameters->ColorTexture = Inputs.SceneColor.Texture;
+
+	FScreenPassRenderTarget Output = Inputs.OverrideOutput;
+	if (!Output.IsValid())
+	{
+		Output = FScreenPassRenderTarget::CreateFromInput(GraphBuilder, Inputs.SceneColor, ERenderTargetLoadAction::ELoad, TEXT("AlphaInvert"));
+	}
+
+	PassParameters->RenderTargets[0] = Output.GetRenderTargetBinding();
+
+	FScreenPassTexture Color = Inputs.SceneColor;
+
+	// We need to make sure that both input and output have alpha channels if not the pass is useless.
+	const EPixelFormat InputPixelFormat = Color.Texture->Desc.Format;
+	const bool bInputHasAlphaChannel = (GetPixelFormatValidChannels(InputPixelFormat) & EPixelFormatChannelFlags::A) != EPixelFormatChannelFlags::None;
+	const EPixelFormat OutputPixelFormat = Output.Texture->Desc.Format;
+	const bool bOutputHasAlphaChannel = (GetPixelFormatValidChannels(OutputPixelFormat) & EPixelFormatChannelFlags::A) != EPixelFormatChannelFlags::None;
+
+	if (!(bInputHasAlphaChannel && bOutputHasAlphaChannel))
+	{
+		return Inputs.SceneColor;
+	}
+
+	GraphBuilder.AddPass(
+		RDG_EVENT_NAME("AlphaInvertPass"),
+		PassParameters,
+		ERDGPassFlags::Raster,
+		[&View, Color](FRHICommandListImmediate& RHICmdList)
+		{
+			AlphaInvert::RenderAlphaInvertPass(RHICmdList, View, Color);
+		});
+
+	return Inputs.SceneColor;
+}
+
+namespace AlphaInvert
+{
+	// In this version we will copy the entire input to the entire output, ignoring the view rect
+	// This would do one pass instead of two when stereo rendering.
+	void RenderAlphaInvertPassFullTexture(FRHICommandList& RHICmdList, const FViewInfo& View, FRDGTextureRef Color)
+	{
+		// Part of scene rendering pass
+		check(RHICmdList.IsInsideRenderPass());
+		SCOPED_DRAW_EVENT(RHICmdList, AlphaInvert);
+
+		const FIntPoint TargetSize = Color->Desc.Extent;
+
+		TShaderMapRef<FScreenVS> VertexShader(View.ShaderMap);
+		TShaderMapRef<FAlphaInvertPS> PixelShader(View.ShaderMap);
+
+		FGraphicsPipelineStateInitializer GraphicsPSOInit;
+		RHICmdList.ApplyCachedRenderTargets(GraphicsPSOInit);
+		GraphicsPSOInit.BlendState = TStaticBlendState<>::GetRHI();
+		GraphicsPSOInit.RasterizerState = TStaticRasterizerState<>::GetRHI();
+		GraphicsPSOInit.DepthStencilState = TStaticDepthStencilState<false, CF_Always>::GetRHI();
+		GraphicsPSOInit.BoundShaderState.VertexDeclarationRHI = GFilterVertexDeclaration.VertexDeclarationRHI;
+		GraphicsPSOInit.BoundShaderState.VertexShaderRHI = VertexShader.GetVertexShader();
+		GraphicsPSOInit.BoundShaderState.PixelShaderRHI = PixelShader.GetPixelShader();
+		GraphicsPSOInit.PrimitiveType = PT_TriangleList;
+
+		SetGraphicsPipelineState(RHICmdList, GraphicsPSOInit, 0);
+		SetShaderParametersLegacyPS(RHICmdList, PixelShader, TStaticSamplerState<SF_Point, AM_Clamp, AM_Clamp, AM_Clamp>::GetRHI(), Color->GetRHI());
+		RHICmdList.SetViewport(0, 0, 0.0f, TargetSize.X, TargetSize.Y, 1.0f);
+
+		DrawRectangle(
+			RHICmdList,
+			0, 0, TargetSize.X, TargetSize.Y,
+			0, 0, TargetSize.X, TargetSize.Y,
+			TargetSize,
+			TargetSize,
+			VertexShader,
+			EDRF_UseTriangleOptimization,
+			View.GetStereoPassInstanceFactor());
+	}
+}
+
+void AlphaInvert::AddAlphaInvertPass(FRDGBuilder& GraphBuilder, const FViewInfo& View, FSceneTextures& SceneTextures)
+{
+	FAlphaInvertParameters* PassParameters = GraphBuilder.AllocParameters<FAlphaInvertParameters>();
+	PassParameters->View = View.GetShaderParameters();
+	PassParameters->ColorTexture = SceneTextures.Color.Resolve;
+	PassParameters->RenderTargets[0] = FRenderTargetBinding(SceneTextures.Color.Resolve, ERenderTargetLoadAction::ELoad);
+
+	GraphBuilder.AddPass(
+		RDG_EVENT_NAME("AlphaInvertPass"),
+		PassParameters,
+		ERDGPassFlags::Raster,
+		[&View, &SceneTextures](FRHICommandListImmediate& RHICmdList)
+		{
+			AlphaInvert::RenderAlphaInvertPassFullTexture(RHICmdList, View, SceneTextures.Color.Resolve);
+		});
+}
\ No newline at end of file
diff --git a/Engine/Source/Runtime/Renderer/Private/PostProcess/AlphaInvert.h b/Engine/Source/Runtime/Renderer/Private/PostProcess/AlphaInvert.h
new file mode 100644
--- /dev/null
+++ b/Engine/Source/Runtime/Renderer/Private/PostProcess/AlphaInvert.h
@@ -0,0 +1,23 @@
+// Copyright Epic Games, Inc. All Rights Reserved.
+
+#pragma once
+
+#include "ScreenPass.h"
+
+namespace AlphaInvert
+{
+
+	// Pass per view version, used in the desktop renderer.
+	struct FAlphaInvertInputs
+	{
+		// [Optional] Render to the specified output. If invalid, a new texture is created and returned.
+		FScreenPassRenderTarget OverrideOutput;
+
+		// [Required] The input scene color and view rect.
+		FScreenPassTexture SceneColor;
+	};
+	FScreenPassTexture AddAlphaInvertPass(class FRDGBuilder& GraphBuilder, const class FViewInfo& View, const struct FAlphaInvertInputs& Inputs);
+	
+	// Single pass version, used in the mobile renderer.
+	void AddAlphaInvertPass(FRDGBuilder& GraphBuilder, const FViewInfo& View, FSceneTextures& SceneTextures);
+}
diff --git a/Engine/Source/Runtime/Renderer/Private/PostProcess/PostProcessTonemap.cpp b/Engine/Source/Runtime/Renderer/Private/PostProcess/PostProcessTonemap.cpp
--- a/Engine/Source/Runtime/Renderer/Private/PostProcess/PostProcessTonemap.cpp
+++ b/Engine/Source/Runtime/Renderer/Private/PostProcess/PostProcessTonemap.cpp
@@ -72,12 +72,6 @@ static TAutoConsoleVariable<int32> CVarBackbufferQuantizationDitheringOverride(
 	TEXT("Disabled by default. Instead is automatically found out by FSceneViewFamily::RenderTarget's pixel format of the backbuffer."),
 	ECVF_RenderThreadSafe);
 
-static TAutoConsoleVariable<int32> CVarAlphaInvertPass(
-	TEXT("r.AlphaInvertPass"),
-	0,
-	TEXT("Whether to run a render pass to un-invert the alpha value from unreal standard to the much more common standard where alpha 0 is fully transparent and alpha 1 is fully opaque."),
-	ECVF_RenderThreadSafe);
-
 const int32 GTonemapComputeTileSizeX = 8;
 const int32 GTonemapComputeTileSizeY = 8;
 
@@ -1194,84 +1188,3 @@ void AddMobileCustomResolvePass(FRDGBuilder& GraphBuilder, const FViewInfo& View
 		});
 }
 
-// Pixel shader to un-invert the alpha channel for output.
-class FAlphaInvertPS : public FGlobalShader
-{
-	DECLARE_GLOBAL_SHADER(FAlphaInvertPS);
-
-	SHADER_USE_PARAMETER_STRUCT(FAlphaInvertPS, FGlobalShader);
-	BEGIN_SHADER_PARAMETER_STRUCT(FParameters, )
-		SHADER_PARAMETER_STRUCT_INCLUDE(FViewShaderParameters, View)
-		SHADER_PARAMETER_RDG_TEXTURE(Texture2D, ColorTexture)
-		SHADER_PARAMETER_SAMPLER(SamplerState, ColorSampler)
-		SHADER_PARAMETER(FVector4f, ColorScale0)
-	END_SHADER_PARAMETER_STRUCT()
-};
-
-IMPLEMENT_GLOBAL_SHADER(FAlphaInvertPS, "/Engine/Private/PostProcessTonemap.usf", "AlphaInvert_MainPS", SF_Pixel);
-
-void RenderAlphaInvertPass(FRHICommandList& RHICmdList, const FViewInfo& View, FSceneTextures& SceneTextures)
-{
-	// Part of scene rendering pass
-	check(RHICmdList.IsInsideRenderPass());
-	SCOPED_DRAW_EVENT(RHICmdList, AlphaInvert);
-
-	const FIntPoint TargetSize = SceneTextures.Color.Resolve->Desc.Extent;
-	
-	TShaderMapRef<FScreenVS> VertexShader(View.ShaderMap);
-	TShaderMapRef<FAlphaInvertPS> PixelShader(View.ShaderMap);
-
-	FAlphaInvertPS::FParameters PSShaderParameters;
-	PSShaderParameters.View = View.GetShaderParameters();
-	PSShaderParameters.ColorTexture = SceneTextures.Color.Resolve;
-	PSShaderParameters.ColorSampler = TStaticSamplerState<SF_Point, AM_Clamp, AM_Clamp, AM_Clamp>::GetRHI();
-
-	FGraphicsPipelineStateInitializer GraphicsPSOInit;
-	RHICmdList.ApplyCachedRenderTargets(GraphicsPSOInit);
-	GraphicsPSOInit.BlendState = TStaticBlendState<>::GetRHI();
-	GraphicsPSOInit.RasterizerState = TStaticRasterizerState<>::GetRHI();
-	GraphicsPSOInit.DepthStencilState = TStaticDepthStencilState<false, CF_Always>::GetRHI();
-	GraphicsPSOInit.BoundShaderState.VertexDeclarationRHI = GFilterVertexDeclaration.VertexDeclarationRHI;
-	GraphicsPSOInit.BoundShaderState.VertexShaderRHI = VertexShader.GetVertexShader();
-	GraphicsPSOInit.BoundShaderState.PixelShaderRHI = PixelShader.GetPixelShader();
-	GraphicsPSOInit.PrimitiveType = PT_TriangleList;
-
-	SetGraphicsPipelineState(RHICmdList, GraphicsPSOInit, 0);
-	SetShaderParameters(RHICmdList, PixelShader, PixelShader.GetPixelShader(), PSShaderParameters);
-	RHICmdList.SetViewport(0, 0, 0.0f, TargetSize.X, TargetSize.Y, 1.0f);
-
-	DrawRectangle(
-		RHICmdList,
-		0, 0,
-		TargetSize.X, TargetSize.Y,
-		0, 0,
-		TargetSize.X, TargetSize.Y,
-		TargetSize,
-		TargetSize,
-		VertexShader,
-		EDRF_UseTriangleOptimization,
-		View.GetStereoPassInstanceFactor());
-}
-
-BEGIN_SHADER_PARAMETER_STRUCT(FAlphaInvertParameters, )
-	SHADER_PARAMETER_STRUCT_INCLUDE(FViewShaderParameters, View)
-	SHADER_PARAMETER_RDG_TEXTURE(Texture2D, ColorTexture)
-	RENDER_TARGET_BINDING_SLOTS()
-END_SHADER_PARAMETER_STRUCT()
-
-void AddAlphaInvertPass(FRDGBuilder& GraphBuilder, const FViewInfo& View, FSceneTextures& SceneTextures)
-{
-	FAlphaInvertParameters* PassParameters = GraphBuilder.AllocParameters<FAlphaInvertParameters>();
-	PassParameters->View = View.GetShaderParameters();
-	PassParameters->ColorTexture = SceneTextures.Color.Resolve;
-	PassParameters->RenderTargets[0] = FRenderTargetBinding(SceneTextures.Color.Resolve, ERenderTargetLoadAction::ELoad);
-
-	GraphBuilder.AddPass(
-		RDG_EVENT_NAME("AlphaInvertPass"),
-		PassParameters,
-		ERDGPassFlags::Raster,
-		[&View, &SceneTextures](FRHICommandListImmediate& RHICmdList)
-		{
-			RenderAlphaInvertPass(RHICmdList, View, SceneTextures);
-		});
-}
diff --git a/Engine/Source/Runtime/Renderer/Private/PostProcess/PostProcessTonemap.h b/Engine/Source/Runtime/Renderer/Private/PostProcess/PostProcessTonemap.h
--- a/Engine/Source/Runtime/Renderer/Private/PostProcess/PostProcessTonemap.h
+++ b/Engine/Source/Runtime/Renderer/Private/PostProcess/PostProcessTonemap.h
@@ -80,5 +80,3 @@ struct FTonemapInputs
 FScreenPassTexture AddTonemapPass(FRDGBuilder& GraphBuilder, const FViewInfo& View, const FTonemapInputs& Inputs);
 void RenderMobileCustomResolve(FRHICommandList& RHICmdList, const FViewInfo& View, const int32 SubpassMSAASamples, FSceneTextures& SceneTextures);
 void AddMobileCustomResolvePass(FRDGBuilder& GraphBuilder, const FViewInfo& View, FSceneTextures& SceneTextures, FRDGTextureRef ViewFamilyTexture);
-
-void AddAlphaInvertPass(FRDGBuilder& GraphBuilder, const FViewInfo& View, FSceneTextures& SceneTextures);
diff --git a/Engine/Source/Runtime/Renderer/Private/PostProcess/PostProcessing.cpp b/Engine/Source/Runtime/Renderer/Private/PostProcess/PostProcessing.cpp
--- a/Engine/Source/Runtime/Renderer/Private/PostProcess/PostProcessing.cpp
+++ b/Engine/Source/Runtime/Renderer/Private/PostProcess/PostProcessing.cpp
@@ -32,6 +32,7 @@
 #include "PostProcess/PostProcessDOF.h"
 #include "PostProcess/PostProcessUpscale.h"
 #include "PostProcess/PostProcessHMD.h"
+#include "PostProcess/AlphaInvert.h"
 #include "PostProcess/PostProcessVisualizeComplexity.h"
 #if UE_ENABLE_DEBUG_DRAWING
 	#include "PostProcess/PostProcessCompositeDebugPrimitives.h"
@@ -468,6 +469,7 @@ void AddPostProcessingPasses(
 #endif
 		PrimaryUpscale,
 		SecondaryUpscale,
+		AlphaInvert,
 		MAX
 	};
 
@@ -524,7 +526,8 @@ void AddPostProcessingPasses(
 		TEXT("DebugPrimitive"),
 #endif
 		TEXT("PrimaryUpscale"),
-		TEXT("SecondaryUpscale")
+		TEXT("SecondaryUpscale"),
+		TEXT("AlphaInvert")
 	};
 
 	static_assert(static_cast<uint32>(EPass::MAX) == UE_ARRAY_COUNT(PassNames), "EPass does not match PassNames."); 
@@ -586,6 +589,7 @@ void AddPostProcessingPasses(
 #endif
 	PassSequence.SetEnabled(EPass::PrimaryUpscale, (bApplyLensDistortion && !bApplyLensDistortionInTSR) || (View.PrimaryScreenPercentageMethod == EPrimaryScreenPercentageMethod::SpatialUpscale && PrimaryViewRect.Size() != View.GetSecondaryViewRectSize()));
 	PassSequence.SetEnabled(EPass::SecondaryUpscale, View.RequiresSecondaryUpscale() || View.Family->GetSecondarySpatialUpscalerInterface() != nullptr);
+	PassSequence.SetEnabled(EPass::AlphaInvert, EngineShowFlags.AlphaInvert);
 
 	const auto GetPostProcessMaterialInputs = [&](FScreenPassTexture InSceneColor)
 	{ 
@@ -2038,6 +2042,17 @@ void AddPostProcessingPasses(
 			AddUserSceneTextureDebugPass(GraphBuilder, View, ViewIndex, SceneColor);
 		}
 	}
+
+	if (PassSequence.IsEnabled(EPass::AlphaInvert))
+	{
+		{
+			AlphaInvert::FAlphaInvertInputs PassInputs;
+			PassSequence.AcceptOverrideIfLastPass(EPass::AlphaInvert, PassInputs.OverrideOutput);
+			PassInputs.SceneColor = SceneColor;
+			SceneColor = AlphaInvert::AddAlphaInvertPass(GraphBuilder, View, PassInputs);
+		}
+	}
+
 	#endif
 }
 
