Add global graphics PSO precaching for light materials

diff --git a/Engine/Source/Runtime/Renderer/Private/LightRendering.cpp b/Engine/Source/Runtime/Renderer/Private/LightRendering.cpp
--- a/Engine/Source/Runtime/Renderer/Private/LightRendering.cpp
+++ b/Engine/Source/Runtime/Renderer/Private/LightRendering.cpp
@@ -37,6 +37,8 @@
 #include "RHIResourceUtils.h"
 #include "RectLightSceneProxy.h"
 
+#include "PSOPrecacheValidation.h"
+
 using namespace LightFunctionAtlas;
 
 // ENABLE_DEBUG_DISCARD_PROP is used to test the lighting code by allowing to discard lights to see how performance scales
@@ -58,6 +60,15 @@ IMPLEMENT_GLOBAL_SHADER_PARAMETER_STRUCT(FDeferredLightUniformStruct, "DeferredL
 
 DECLARE_DWORD_COUNTER_STAT(TEXT("VSM Light Projections (Local One Pass Fast)"), STAT_VSMLocalProjectionOnePassFast, STATGROUP_ShadowRendering);
 
+static int32 GDeferredLightGlobalPSOPrecache = 1;
+static FAutoConsoleVariableRef CVarGDeferredLightGlobalPSOPrecache(
+	TEXT("r.PSOPrecache.DeferredLight"), 
+	GDeferredLightGlobalPSOPrecache, 
+	TEXT("Precache all possible global deferred light PSOs."),
+	ECVF_ReadOnly);
+
+static const TCHAR* DeferredLightGlobalPSOCollectorName = TEXT("DeferredLightGlobalPSOCollector");
+
 extern int32 GUseTranslucentLightingVolumes;
 
 static int32 GAllowDepthBoundsTest = 1;
@@ -1014,6 +1025,17 @@ class FDeferredLightPS : public FGlobalShader
 		return PermutationVector;
 	}
 
+	static EShaderPermutationPrecacheRequest ShouldPrecachePermutation(const FGlobalShaderPermutationParameters& Parameters)
+	{
+		FPermutationDomain PermutationVector(Parameters.PermutationId);
+		if (PermutationVector.Get<FVisualizeCullingDim>())
+		{
+			return EShaderPermutationPrecacheRequest::NotPrecached;
+		}
+
+		return FGlobalShader::ShouldPrecachePermutation(Parameters);
+	}
+
 	static void ModifyCompilationEnvironment(const FGlobalShaderPermutationParameters& Parameters, FShaderCompilerEnvironment& OutEnvironment)
 	{
 		FPermutationDomain PermutationVector(Parameters.PermutationId);
@@ -2316,17 +2338,17 @@ void FDeferredShadingSceneRenderer::RenderStationaryLightOverlap(
 	}
 }
 
-static void InternalSetBoundingGeometryRasterizerState(FGraphicsPipelineStateInitializer& GraphicsPSOInit, const FViewInfo& View, bool bCameraInsideLightGeometry)
+static void InternalSetBoundingGeometryRasterizerState(FGraphicsPipelineStateInitializer& GraphicsPSOInit, bool bReverseCulling, bool bCameraInsideLightGeometry)
 {
 	if (bCameraInsideLightGeometry)
 	{
 		// Render backfaces with depth tests disabled since the camera is inside (or close to inside) the light geometry
-		GraphicsPSOInit.RasterizerState = View.bReverseCulling ? TStaticRasterizerState<FM_Solid, CM_CW>::GetRHI() : TStaticRasterizerState<FM_Solid, CM_CCW>::GetRHI();
+		GraphicsPSOInit.RasterizerState = bReverseCulling ? TStaticRasterizerState<FM_Solid, CM_CW>::GetRHI() : TStaticRasterizerState<FM_Solid, CM_CCW>::GetRHI();
 	}
 	else
 	{
 		// Render frontfaces with depth tests on to get the speedup from HiZ since the camera is outside the light geometry
-		GraphicsPSOInit.RasterizerState = View.bReverseCulling ? TStaticRasterizerState<FM_Solid, CM_CCW>::GetRHI() : TStaticRasterizerState<FM_Solid, CM_CW>::GetRHI();
+		GraphicsPSOInit.RasterizerState = bReverseCulling ? TStaticRasterizerState<FM_Solid, CM_CCW>::GetRHI() : TStaticRasterizerState<FM_Solid, CM_CW>::GetRHI();
 	}
 }
 
@@ -2356,10 +2378,10 @@ static uint32 InternalSetBoundingGeometryDepthState(FGraphicsPipelineStateInitia
 }
 
 /** Sets up rasterizer and depth state for rendering bounding geometry in a deferred pass. */
-static uint32 SetBoundingGeometryRasterizerAndDepthState(FGraphicsPipelineStateInitializer& GraphicsPSOInit, const FViewInfo& View, bool bCameraInsideLightGeometry, ESubstrateTileType TileType)
+static uint32 SetBoundingGeometryRasterizerAndDepthState(FGraphicsPipelineStateInitializer& GraphicsPSOInit, bool bReverseCulling, bool bCameraInsideLightGeometry, ESubstrateTileType TileType)
 {
 	uint32 StencilRef = 0u;
-	InternalSetBoundingGeometryRasterizerState(GraphicsPSOInit, View, bCameraInsideLightGeometry);
+	InternalSetBoundingGeometryRasterizerState(GraphicsPSOInit, bReverseCulling, bCameraInsideLightGeometry);
 	if (bCameraInsideLightGeometry)
 	{
 		StencilRef = InternalSetBoundingGeometryDepthState<CF_Always>(GraphicsPSOInit, TileType);
@@ -2485,6 +2507,48 @@ static FDeferredLightPS::FParameters GetDeferredLightPSParameters(
 	return Out;
 }
 
+uint32 SetupLightGraphicsPSOState( 
+	bool bDirectional,
+	bool bCameraInsideLightGeometry, 
+	bool bReverseCulling, 
+	ESubstrateTileType SubstrateTileMaterialType,
+	FGraphicsPipelineStateInitializer& GraphicsPSOInit)
+{
+	GraphicsPSOInit.PrimitiveType = PT_TriangleList;
+	if (Substrate::IsOpaqueRoughRefractionEnabled())
+	{
+		GraphicsPSOInit.BlendState = TStaticBlendState<
+			CW_RGBA, BO_Add, BF_One, BF_One, BO_Add, BF_One, BF_One,
+			CW_RGBA, BO_Add, BF_One, BF_One, BO_Add, BF_One, BF_One,
+			CW_RGBA, BO_Add, BF_One, BF_One, BO_Add, BF_One, BF_One>::GetRHI();
+	}
+	else
+	{
+		GraphicsPSOInit.BlendState = TStaticBlendState<CW_RGBA, BO_Add, BF_One, BF_One, BO_Add, BF_One, BF_One>::GetRHI();
+	}
+
+	uint32 StencilRef = 0;
+	if (bDirectional)
+	{
+		// Turn DBT back off
+		GraphicsPSOInit.bDepthBounds = false;
+		GraphicsPSOInit.RasterizerState = TStaticRasterizerState<FM_Solid, CM_None>::GetRHI();
+		GraphicsPSOInit.BoundShaderState.VertexDeclarationRHI = GFilterVertexDeclaration.VertexDeclarationRHI;
+		GraphicsPSOInit.DepthStencilState = TStaticDepthStencilState<false, CF_Always>::GetRHI();
+	}
+	else
+	{
+		// Use DBT to allow work culling on shadow lights
+		// Disable depth bound when hair rendering is enabled as this rejects partially covered pixel write (with opaque background)
+		GraphicsPSOInit.bDepthBounds = GSupportsDepthBoundsTest && GAllowDepthBoundsTest != 0;
+
+		StencilRef = SetBoundingGeometryRasterizerAndDepthState(GraphicsPSOInit, bReverseCulling, bCameraInsideLightGeometry, SubstrateTileMaterialType);
+		GraphicsPSOInit.BoundShaderState.VertexDeclarationRHI = GetVertexDeclarationFVector4();
+	}
+
+	return StencilRef;
+}
+
 // Used by RenderLights to render a light to the scene color buffer.
 template<typename TShaderType, typename TParametersType>
 static void InternalRenderLight(
@@ -2515,22 +2579,18 @@ static void InternalRenderLight(
 
 		FGraphicsPipelineStateInitializer GraphicsPSOInit;
 		RHICmdList.ApplyCachedRenderTargets(GraphicsPSOInit);
+
+		const bool bCameraInsideLightGeometry = ((FVector)View.ViewMatrices.GetViewOrigin() - LightBounds.Center).SizeSquared() < FMath::Square(LightBounds.W * 1.05f + View.NearClippingDistance * 2.0f)
+			//const bool bCameraInsideLightGeometry = LightProxy->AffectsBounds( FSphere( View.ViewMatrices.GetViewOrigin(), View.NearClippingDistance * 2.0f ) )
+				// Always draw backfaces in ortho
+				//@todo - accurate ortho camera / light intersection
+			|| !View.IsPerspectiveProjection();
+
+		const uint32 StencilRef = SetupLightGraphicsPSOState(LightType == LightType_Directional, bCameraInsideLightGeometry, View.bReverseCulling, SubstrateTileMaterialType, GraphicsPSOInit);
+
 		// Set the device viewport for the view.
 		RHICmdList.SetViewport(View.ViewRect.Min.X, View.ViewRect.Min.Y, 0.0f, View.ViewRect.Max.X, View.ViewRect.Max.Y, 1.0f);
 
-		if (Substrate::IsOpaqueRoughRefractionEnabled())
-		{
-			GraphicsPSOInit.BlendState = TStaticBlendState<
-				CW_RGBA, BO_Add, BF_One, BF_One, BO_Add, BF_One, BF_One,
-				CW_RGBA, BO_Add, BF_One, BF_One, BO_Add, BF_One, BF_One,
-				CW_RGBA, BO_Add, BF_One, BF_One, BO_Add, BF_One, BF_One>::GetRHI();
-		}
-		else
-		{
-			GraphicsPSOInit.BlendState = TStaticBlendState<CW_RGBA, BO_Add, BF_One, BF_One, BO_Add, BF_One, BF_One>::GetRHI();
-		}
-		GraphicsPSOInit.PrimitiveType = PT_TriangleList;
-
 		if (LightType == LightType_Directional)
 		{
 			FDeferredLightVS::FPermutationDomain PermutationVectorVS;
@@ -2542,21 +2602,25 @@ static void InternalRenderLight(
 			VSPermutationVector.Set< Substrate::FSubstrateTilePassVS::FEnableTexCoordScreenVector >(true);
 			TShaderMapRef<Substrate::FSubstrateTilePassVS> TileVertexShader(View.ShaderMap, VSPermutationVector);
 
+			GraphicsPSOInit.BoundShaderState.VertexShaderRHI = bEnableSubstrateTiledPass ? TileVertexShader.GetVertexShader() : VertexShader.GetVertexShader();
+			GraphicsPSOInit.BoundShaderState.PixelShaderRHI = PixelShader.GetPixelShader();
+
+#if PSO_PRECACHING_VALIDATE
+			if (PSOCollectorStats::IsFullPrecachingValidationEnabled())
+			{
+				static const int32 GlobalPSOCollectorIndex = FGlobalPSOCollectorManager::GetIndex(DeferredLightGlobalPSOCollectorName);
+				PSOCollectorStats::CheckGlobalGraphicsPipelineStateInCache(GraphicsPSOInit, GlobalPSOCollectorIndex);
+			}
+#endif // PSO_PRECACHING_VALIDATE
+
+			SetGraphicsPipelineState(RHICmdList, GraphicsPSOInit, StencilRef);
+
 			Substrate::FSubstrateTilePassVS::FParameters VSParameters;
 			if (Substrate::IsSubstrateEnabled())
 			{
 				VSParameters = Substrate::SetTileParameters(View, SubstrateTileMaterialType, GraphicsPSOInit.PrimitiveType);
 			}
 
-			// Turn DBT back off
-			GraphicsPSOInit.bDepthBounds = false;
-			GraphicsPSOInit.RasterizerState = TStaticRasterizerState<FM_Solid, CM_None>::GetRHI();
-			GraphicsPSOInit.BoundShaderState.VertexDeclarationRHI = GFilterVertexDeclaration.VertexDeclarationRHI;
-			GraphicsPSOInit.BoundShaderState.VertexShaderRHI = bEnableSubstrateTiledPass ? TileVertexShader.GetVertexShader() : VertexShader.GetVertexShader();
-			GraphicsPSOInit.BoundShaderState.PixelShaderRHI = PixelShader.GetPixelShader();		
-			GraphicsPSOInit.DepthStencilState = TStaticDepthStencilState<false, CF_Always>::GetRHI();
-			SetGraphicsPipelineState(RHICmdList, GraphicsPSOInit, 0x0);
-
 			SetShaderParameters(RHICmdList, PixelShader, PixelShader.GetPixelShader(), PassParameters->PS);
 
 			if (SubstrateTileMaterialType != ECount)
@@ -2584,25 +2648,22 @@ static void InternalRenderLight(
 			}
 		}
 		else // Radial light (LightType_Point, LightType_Spot, LightType_Rect)
-		{
-			// Use DBT to allow work culling on shadow lights
-			// Disable depth bound when hair rendering is enabled as this rejects partially covered pixel write (with opaque background)
-			GraphicsPSOInit.bDepthBounds = GSupportsDepthBoundsTest && GAllowDepthBoundsTest != 0;
-
+		{	
 			FDeferredLightVS::FPermutationDomain PermutationVectorVS;
 			PermutationVectorVS.Set<FDeferredLightVS::FRadialLight>(true);
 			TShaderMapRef<FDeferredLightVS> VertexShader(View.ShaderMap, PermutationVectorVS);
 
-			const bool bCameraInsideLightGeometry = ((FVector)View.ViewMatrices.GetViewOrigin() - LightBounds.Center).SizeSquared() < FMath::Square(LightBounds.W * 1.05f + View.NearClippingDistance * 2.0f)
-			//const bool bCameraInsideLightGeometry = LightProxy->AffectsBounds( FSphere( View.ViewMatrices.GetViewOrigin(), View.NearClippingDistance * 2.0f ) )
-				// Always draw backfaces in ortho
-				//@todo - accurate ortho camera / light intersection
-				|| !View.IsPerspectiveProjection();
-
-			const uint32 StencilRef = SetBoundingGeometryRasterizerAndDepthState(GraphicsPSOInit, View, bCameraInsideLightGeometry, SubstrateTileMaterialType);
-			GraphicsPSOInit.BoundShaderState.VertexDeclarationRHI = GetVertexDeclarationFVector4();
 			GraphicsPSOInit.BoundShaderState.VertexShaderRHI = VertexShader.GetVertexShader();
 			GraphicsPSOInit.BoundShaderState.PixelShaderRHI = PixelShader.GetPixelShader();
+
+#if PSO_PRECACHING_VALIDATE
+			if (PSOCollectorStats::IsFullPrecachingValidationEnabled())
+			{
+				static const int32 GlobalPSOCollectorIndex = FGlobalPSOCollectorManager::GetIndex(DeferredLightGlobalPSOCollectorName);
+				PSOCollectorStats::CheckGlobalGraphicsPipelineStateInCache(GraphicsPSOInit, GlobalPSOCollectorIndex);
+			}
+#endif // PSO_PRECACHING_VALIDATE
+
 			SetGraphicsPipelineState(RHICmdList, GraphicsPSOInit, StencilRef);
 			SetShaderParameters(RHICmdList, PixelShader, PixelShader.GetPixelShader(), PassParameters->PS);
 
@@ -2842,6 +2903,16 @@ BEGIN_SHADER_PARAMETER_STRUCT(FRenderLightForHairParameters, )
 	SHADER_PARAMETER_STRUCT_INCLUDE(FDeferredLightPS::FParameters, PS)
 END_SHADER_PARAMETER_STRUCT()
 
+void SetupLightForHairGraphicsPSOState(FGraphicsPipelineStateInitializer& GraphicsPSOInit)
+{
+	GraphicsPSOInit.BlendState = TStaticBlendState<CW_RGBA, BO_Add, BF_One, BF_One, BO_Max, BF_One, BF_One>::GetRHI();
+	GraphicsPSOInit.PrimitiveType = PT_TriangleList;
+	GraphicsPSOInit.BoundShaderState.VertexDeclarationRHI = GetVertexDeclarationFVector4();
+	GraphicsPSOInit.bDepthBounds = false;
+	GraphicsPSOInit.RasterizerState = TStaticRasterizerState<FM_Solid, CM_None>::GetRHI();
+	GraphicsPSOInit.DepthStencilState = TStaticDepthStencilState<false, CF_Always>::GetRHI();
+}
+
 void FDeferredShadingSceneRenderer::RenderLightForHair(
 	FRDGBuilder& GraphBuilder,
 	FViewInfo& View,
@@ -2957,16 +3028,20 @@ void FDeferredShadingSceneRenderer::RenderLightForHair(
 
 		FGraphicsPipelineStateInitializer GraphicsPSOInit;
 		RHICmdList.ApplyCachedRenderTargets(GraphicsPSOInit);
-		GraphicsPSOInit.BlendState = TStaticBlendState<CW_RGBA, BO_Add, BF_One, BF_One, BO_Max, BF_One, BF_One>::GetRHI();
-		GraphicsPSOInit.PrimitiveType = PT_TriangleList;
+		
+		SetupLightForHairGraphicsPSOState(GraphicsPSOInit);
 
-		GraphicsPSOInit.BoundShaderState.VertexDeclarationRHI = GetVertexDeclarationFVector4();
-		GraphicsPSOInit.bDepthBounds = false;
-		GraphicsPSOInit.RasterizerState = TStaticRasterizerState<FM_Solid, CM_None>::GetRHI();
-		GraphicsPSOInit.DepthStencilState = TStaticDepthStencilState<false, CF_Always>::GetRHI();
 		GraphicsPSOInit.BoundShaderState.VertexShaderRHI = VertexShader.GetVertexShader();
 		GraphicsPSOInit.BoundShaderState.PixelShaderRHI = PixelShader.GetPixelShader();
-		GraphicsPSOInit.PrimitiveType = PT_TriangleList;
+				
+#if PSO_PRECACHING_VALIDATE
+		if (PSOCollectorStats::IsFullPrecachingValidationEnabled())
+		{
+			static const int32 GlobalPSOCollectorIndex = FGlobalPSOCollectorManager::GetIndex(DeferredLightGlobalPSOCollectorName);
+			PSOCollectorStats::CheckGlobalGraphicsPipelineStateInCache(GraphicsPSOInit, GlobalPSOCollectorIndex);
+		}
+#endif // PSO_PRECACHING_VALIDATE
+
 		SetGraphicsPipelineState(RHICmdList, GraphicsPSOInit, 0);
 
 		SetShaderParameters(RHICmdList, VertexShader, VertexShader.GetVertexShader(), PassParameters->VS);
@@ -3140,6 +3215,8 @@ static void InternalRenderSimpleLightsStandardDeferred(
 		GraphicsPSOInit.BlendState = TStaticBlendState<CW_RGBA, BO_Add, BF_One, BF_One, BO_Add, BF_One, BF_One>::GetRHI();
 		GraphicsPSOInit.PrimitiveType = PT_TriangleList;
 
+		const int32 GlobalPSOCollectorIndex = FGlobalPSOCollectorManager::GetIndex(DeferredLightGlobalPSOCollectorName);
+
 		for (int32 LightIndex = 0; LightIndex < SimpleLights.InstanceData.Num(); LightIndex++)
 		{
 			const FSimpleLightEntry& SimpleLight = SimpleLights.InstanceData[LightIndex];
@@ -3156,10 +3233,18 @@ static void InternalRenderSimpleLightsStandardDeferred(
 							//@todo - accurate ortho camera / light intersection
 							|| !View.IsPerspectiveProjection();
 
-			const uint32 StencilRef = SetBoundingGeometryRasterizerAndDepthState(GraphicsPSOInit, View, bCameraInsideLightGeometry, TileType);
+			const uint32 StencilRef = SetBoundingGeometryRasterizerAndDepthState(GraphicsPSOInit, View.bReverseCulling, bCameraInsideLightGeometry, TileType);
 			GraphicsPSOInit.BoundShaderState.VertexDeclarationRHI = GetVertexDeclarationFVector4();
 			GraphicsPSOInit.BoundShaderState.VertexShaderRHI = VertexShader.GetVertexShader();
 			GraphicsPSOInit.BoundShaderState.PixelShaderRHI = PixelShader.GetPixelShader();
+						
+#if PSO_PRECACHING_VALIDATE
+			if (PSOCollectorStats::IsFullPrecachingValidationEnabled())
+			{
+				PSOCollectorStats::CheckGlobalGraphicsPipelineStateInCache(GraphicsPSOInit, GlobalPSOCollectorIndex);
+			}
+#endif // PSO_PRECACHING_VALIDATE
+
 			SetGraphicsPipelineState(RHICmdList, GraphicsPSOInit, StencilRef);
 
 			// Update the light parameters with a custom uniform buffer
@@ -3308,3 +3393,120 @@ FRDGTextureRef FDeferredShadingSceneRenderer::CopyStencilToLightingChannelTextur
 
 	return LightingChannelsTexture;
 }
+
+void DeferredLightGlobalPSOCollector(const FSceneTexturesConfig& SceneTexturesConfig, int32 GlobalPSOCollectorIndex, TArray<FPSOPrecacheData>& PSOInitializers)
+{
+	if (GDeferredLightGlobalPSOPrecache == 0)
+	{
+		return;
+	}
+
+	EShaderPlatform ShaderPlatform = SceneTexturesConfig.ShaderPlatform;
+	FGlobalShaderMap* GlobalShaderMap = GetGlobalShaderMap(ShaderPlatform);
+
+	auto AddPSOInitializer = [&](bool bIsHairShader, bool bCameraInsideLightGeometry, bool bReverseCulling, bool bIsDirectional, ESubstrateTileType SubstrateTileMaterialType, FRHIPixelShader* PixelShaderRHI)
+	{
+		FGraphicsPipelineStateInitializer GraphicsPSOInit;	
+
+		FGraphicsPipelineRenderTargetsInfo RenderTargetsInfo;
+		RenderTargetsInfo.NumSamples = 1;
+		if (bIsHairShader)
+		{
+			SetupLightForHairGraphicsPSOState(GraphicsPSOInit);
+
+			// TODO: find out the render target info for hair first
+			return;
+		}
+		else
+		{
+			SetupLightGraphicsPSOState(bIsDirectional, bCameraInsideLightGeometry, bReverseCulling, SubstrateTileMaterialType, GraphicsPSOInit);
+
+			if (bIsDirectional)
+			{
+				FDeferredLightVS::FPermutationDomain PermutationVectorVS;
+				PermutationVectorVS.Set<FDeferredLightVS::FRadialLight>(false);
+				TShaderMapRef<FDeferredLightVS> VertexShader(GlobalShaderMap, PermutationVectorVS);
+
+				Substrate::FSubstrateTilePassVS::FPermutationDomain VSPermutationVector;
+				VSPermutationVector.Set< Substrate::FSubstrateTilePassVS::FEnableDebug >(false);
+				VSPermutationVector.Set< Substrate::FSubstrateTilePassVS::FEnableTexCoordScreenVector >(true);
+				TShaderMapRef<Substrate::FSubstrateTilePassVS> TileVertexShader(GlobalShaderMap, VSPermutationVector);
+
+				const bool bEnableSubstrateTiledPass = SubstrateTileMaterialType != ESubstrateTileType::ECount;
+				GraphicsPSOInit.BoundShaderState.VertexShaderRHI = bEnableSubstrateTiledPass ? TileVertexShader.GetVertexShader() : VertexShader.GetVertexShader();
+				GraphicsPSOInit.BoundShaderState.PixelShaderRHI = PixelShaderRHI;
+			}
+			else // Radial light (LightType_Point, LightType_Spot, LightType_Rect)
+			{
+				FDeferredLightVS::FPermutationDomain PermutationVectorVS;
+				PermutationVectorVS.Set<FDeferredLightVS::FRadialLight>(true);
+				TShaderMapRef<FDeferredLightVS> VertexShader(GlobalShaderMap, PermutationVectorVS);
+
+				GraphicsPSOInit.BoundShaderState.VertexShaderRHI = VertexShader.GetVertexShader();
+				GraphicsPSOInit.BoundShaderState.PixelShaderRHI = PixelShaderRHI;
+			}
+
+			AddRenderTargetInfo(SceneTexturesConfig.ColorFormat, SceneTexturesConfig.ColorCreateFlags, RenderTargetsInfo);
+			if (Substrate::IsOpaqueRoughRefractionEnabled() && Substrate::UsesSubstrateMaterialBuffer(ShaderPlatform))
+			{
+				// TODO: find out the render target info for subtrate here
+				//Out.RenderTargets[1] = FRenderTargetBinding(Scene->SubstrateSceneData.SeparatedOpaqueRoughRefractionSceneColor, ERenderTargetLoadAction::ELoad);
+				//Out.RenderTargets[2] = FRenderTargetBinding(Scene->SubstrateSceneData.SeparatedSubSurfaceSceneColor, ERenderTargetLoadAction::ELoad);
+				return;
+			}
+			ETextureCreateFlags DepthStencilCreateFlags = SceneTexturesConfig.DepthCreateFlags;
+			SetupDepthStencilInfo(PF_DepthStencil, DepthStencilCreateFlags, ERenderTargetLoadAction::ELoad,
+				ERenderTargetLoadAction::ELoad, FExclusiveDepthStencil::DepthRead_StencilWrite, RenderTargetsInfo);			
+		}
+
+		GraphicsPSOInit.StatePrecachePSOHash = RHIComputeStatePrecachePSOHash(GraphicsPSOInit);		
+		ApplyTargetsInfo(GraphicsPSOInit, RenderTargetsInfo);
+
+		FPSOPrecacheData PSOPrecacheData;
+		PSOPrecacheData.bRequired = true;
+		PSOPrecacheData.Type = FPSOPrecacheData::EType::Graphics;
+		PSOPrecacheData.GraphicsPSOInitializer = GraphicsPSOInit;
+#if PSO_PRECACHING_VALIDATE
+		PSOPrecacheData.PSOCollectorIndex = GlobalPSOCollectorIndex;
+		PSOPrecacheData.VertexFactoryType = nullptr;
+#endif // PSO_PRECACHING_VALIDATE
+
+		PSOInitializers.Add(MoveTemp(PSOPrecacheData));
+	};
+
+	// Create variations for given render & depth stencil compare states - influence raster state culling for non directional only
+	const bool bCameraInsideLightGeometry = true;
+	const bool bReverseCulling = false;
+
+	// Precache PSOs are never required
+	const bool bRequired = false;
+
+	EShaderPermutationFlags PermutationFlags = EShaderPermutationFlags::None;
+
+	// Is there a different way to get the global shader type from given name?
+	FShaderType* ShaderType = FShaderType::GetShaderTypeByName(TEXT("FDeferredLightPS"));
+	FGlobalShaderType* GlobalShaderType = ShaderType->GetGlobalShaderType();
+	for (int32 PermutationId = 0; PermutationId < GlobalShaderType->GetPermutationCount(); PermutationId++)
+	{
+		if (GlobalShaderType->ShouldCompilePermutation(ShaderPlatform, PermutationId, PermutationFlags) &&
+			GlobalShaderType->ShouldPrecachePermutation(ShaderPlatform, PermutationId, PermutationFlags) == EShaderPermutationPrecacheRequest::Precached)
+		{
+			TShaderRef<FShader> GlobalShader = GlobalShaderMap->GetShader(GlobalShaderType, PermutationId);
+			
+			FRHIPixelShader* RHIPixelShader = static_cast<FRHIPixelShader*>(GlobalShader.GetRHIShaderBase(SF_Pixel, bRequired));
+			if (RHIPixelShader)
+			{
+				FDeferredLightPS::FPermutationDomain PermutationVector(PermutationId);
+
+				// Extract useful information from the permutation vector
+				bool bIsDirectional = (PermutationVector.Get< FDeferredLightPS::FSourceShapeDim >() == ELightSourceShape::Directional);
+				ESubstrateTileType SubstrateTileMaterialType = Substrate::IsSubstrateEnabled() ? (ESubstrateTileType) PermutationVector.Get< FDeferredLightPS::FSubstrateTileType >() : ESubstrateTileType::ECount;
+				bool bIsHairShader = PermutationVector.Get<FDeferredLightPS::FHairLighting>() > 0;
+
+				AddPSOInitializer(bIsHairShader, bCameraInsideLightGeometry, bReverseCulling, bIsDirectional, SubstrateTileMaterialType, RHIPixelShader);
+			}
+		}
+	}
+}
+
+FRegisterGlobalPSOCollectorFunction RegisterDeferredLightGlobalPSOCollector(&DeferredLightGlobalPSOCollector, DeferredLightGlobalPSOCollectorName);
