Remove the mobile PixelProjectReflection because we have the mobile SSR.

diff --git a/Engine/Config/Android/AndroidScalability.ini b/Engine/Config/Android/AndroidScalability.ini
--- a/Engine/Config/Android/AndroidScalability.ini
+++ b/Engine/Config/Android/AndroidScalability.ini
@@ -92,7 +92,6 @@ r.DepthOfFieldQuality=0
 r.LightShaftQuality=0
 r.EyeAdaptationQuality=1
 r.Mobile.AmbientOcclusionQuality=0
-r.Mobile.PixelProjectedReflectionQuality=0
 r.LensFlareQuality=0
 
 [PostProcessQuality@1]
@@ -102,7 +101,6 @@ r.DepthOfFieldQuality=1
 r.LightShaftQuality=1
 r.EyeAdaptationQuality=1
 r.Mobile.AmbientOcclusionQuality=0
-r.Mobile.PixelProjectedReflectionQuality=0
 r.LensFlareQuality=1
 
 [PostProcessQuality@2]
@@ -112,7 +110,6 @@ r.DepthOfFieldQuality=1
 r.LightShaftQuality=1
 r.EyeAdaptationQuality=1
 r.Mobile.AmbientOcclusionQuality=1
-r.Mobile.PixelProjectedReflectionQuality=1
 r.LensFlareQuality=2
 
 [PostProcessQuality@3]
@@ -122,7 +119,6 @@ r.DepthOfFieldQuality=1
 r.LightShaftQuality=1
 r.EyeAdaptationQuality=1
 r.Mobile.AmbientOcclusionQuality=1
-r.Mobile.PixelProjectedReflectionQuality=1
 r.LensFlareQuality=2
 
 [AntiAliasingQuality@0]
diff --git a/Engine/Config/BaseDeviceProfiles.ini b/Engine/Config/BaseDeviceProfiles.ini
--- a/Engine/Config/BaseDeviceProfiles.ini
+++ b/Engine/Config/BaseDeviceProfiles.ini
@@ -356,7 +356,6 @@ BaseProfileName=IOS
 +CVars=sg.ShadowQuality=2
 +CVars=sg.GlobalIlluminationQuality=2
 +CVars=r.Mobile.AmbientOcclusionQuality=0
-+CVars=r.Mobile.PixelProjectedReflectionQuality=0
 
 [iPhone6SPlus DeviceProfile]
 DeviceType=IOS
@@ -367,7 +366,6 @@ BaseProfileName=IOS
 +CVars=sg.ShadowQuality=2
 +CVars=sg.GlobalIlluminationQuality=2
 +CVars=r.Mobile.AmbientOcclusionQuality=0
-+CVars=r.Mobile.PixelProjectedReflectionQuality=0
 
 [iPhoneSE DeviceProfile]
 DeviceType=IOS
@@ -376,7 +374,6 @@ BaseProfileName=IOS
 +CVars=r.MobileContentScaleFactor=2
 +CVars=sg.PostProcessQuality=3
 +CVars=r.Mobile.AmbientOcclusionQuality=0
-+CVars=r.Mobile.PixelProjectedReflectionQuality=0
 
 [iPhone7 DeviceProfile]
 DeviceType=IOS
@@ -384,7 +381,6 @@ BaseProfileName=iPhone6S
 +CVars=ios.PhysicalScreenDensity=326
 +CVars=r.MobileContentScaleFactor=2
 +CVars=r.Mobile.AmbientOcclusionQuality=0
-+CVars=r.Mobile.PixelProjectedReflectionQuality=0
 
 [iPhone7Plus DeviceProfile]
 DeviceType=IOS
@@ -392,7 +388,6 @@ BaseProfileName=iPhone6SPlus
 +CVars=ios.PhysicalScreenDensity=401
 +CVars=r.MobileContentScaleFactor=2
 +CVars=r.Mobile.AmbientOcclusionQuality=0
-+CVars=r.Mobile.PixelProjectedReflectionQuality=0
 
 [iPhone8 DeviceProfile]
 DeviceType=IOS
@@ -401,7 +396,6 @@ BaseProfileName=iPhone7
 +CVars=sg.ShadowQuality=3
 +CVars=sg.GlobalIlluminationQuality=3
 +CVars=r.Mobile.AmbientOcclusionQuality=1
-+CVars=r.Mobile.PixelProjectedReflectionQuality=1
 
 [iPhone8Plus DeviceProfile]
 DeviceType=IOS
@@ -410,7 +404,6 @@ BaseProfileName=iPhone7Plus
 +CVars=sg.ShadowQuality=3
 +CVars=sg.GlobalIlluminationQuality=3
 +CVars=r.Mobile.AmbientOcclusionQuality=1
-+CVars=r.Mobile.PixelProjectedReflectionQuality=1
 
 [iPhoneX DeviceProfile]
 DeviceType=IOS
@@ -563,7 +556,6 @@ BaseProfileName=iPhone7
 +CVars=ios.PhysicalScreenDensity=326
 +CVars=r.MobileContentScaleFactor=2
 +CVars=r.Mobile.AmbientOcclusionQuality=0
-+CVars=r.Mobile.PixelProjectedReflectionQuality=0
 
 [iPadAir2 DeviceProfile]
 DeviceType=IOS
@@ -573,7 +565,6 @@ BaseProfileName=IOS
 +CVars=sg.ShadowQuality=2
 +CVars=sg.GlobalIlluminationQuality=2
 +CVars=r.Mobile.AmbientOcclusionQuality=0
-+CVars=r.Mobile.PixelProjectedReflectionQuality=0
 
 [iPadMini4 DeviceProfile]
 DeviceType=IOS
@@ -581,7 +572,6 @@ BaseProfileName=IOS
 +CVars=ios.PhysicalScreenDensity=401
 +CVars=sg.PostProcessQuality=2
 +CVars=r.Mobile.AmbientOcclusionQuality=0
-+CVars=r.Mobile.PixelProjectedReflectionQuality=0
 
 [iPadPro DeviceProfile]
 DeviceType=IOS
@@ -732,7 +722,6 @@ BaseProfileName=IOS
 +CVars=sg.ShadowQuality=3
 +CVars=sg.GlobalIlluminationQuality=3
 +CVars=r.Mobile.AmbientOcclusionQuality=0
-+CVars=r.Mobile.PixelProjectedReflectionQuality=0
 +CVars=UI.SlateSDFText.RasterizationMode=SdfApproximation
 +CVars=UI.SlateSDFText.ResolutionLevel=1
 
diff --git a/Engine/Config/IOS/IOSScalability.ini b/Engine/Config/IOS/IOSScalability.ini
--- a/Engine/Config/IOS/IOSScalability.ini
+++ b/Engine/Config/IOS/IOSScalability.ini
@@ -92,7 +92,6 @@ r.DepthOfFieldQuality=0
 r.LightShaftQuality=0
 r.EyeAdaptationQuality=1
 r.Mobile.AmbientOcclusionQuality=0
-r.Mobile.PixelProjectedReflectionQuality=0
 r.LensFlareQuality=0
 
 [PostProcessQuality@1]
@@ -102,7 +101,6 @@ r.DepthOfFieldQuality=0
 r.LightShaftQuality=0
 r.EyeAdaptationQuality=1
 r.Mobile.AmbientOcclusionQuality=0
-r.Mobile.PixelProjectedReflectionQuality=0
 r.LensFlareQuality=1
 
 [PostProcessQuality@2]
@@ -112,7 +110,6 @@ r.DepthOfFieldQuality=0
 r.LightShaftQuality=0
 r.EyeAdaptationQuality=1
 r.Mobile.AmbientOcclusionQuality=0
-r.Mobile.PixelProjectedReflectionQuality=0
 r.LensFlareQuality=2
 
 [PostProcessQuality@3]
@@ -122,7 +119,6 @@ r.DepthOfFieldQuality=1
 r.LightShaftQuality=1
 r.EyeAdaptationQuality=1
 r.Mobile.AmbientOcclusionQuality=1
-r.Mobile.PixelProjectedReflectionQuality=1
 r.LensFlareQuality=2
 
 [AntiAliasingQuality@0]
diff --git a/Engine/Shaders/Private/PostProcessPixelProjectedReflectionMobile.usf b/Engine/Shaders/Private/PostProcessPixelProjectedReflectionMobile.usf
deleted file mode 100644
diff --git a/Engine/Source/Runtime/Engine/Classes/Engine/RendererSettings.h b/Engine/Source/Runtime/Engine/Classes/Engine/RendererSettings.h
--- a/Engine/Source/Runtime/Engine/Classes/Engine/RendererSettings.h
+++ b/Engine/Source/Runtime/Engine/Classes/Engine/RendererSettings.h
@@ -1418,11 +1418,9 @@ class URendererSettings : public UDeveloperSettings
 		ClampMin = "12", UIMin = "12"))
 		FPerPlatformInt MaxSkinBones;
 	
-	UPROPERTY(config, EditAnywhere, Category = Mobile, meta = (
-		ConsoleVariable = "r.Mobile.PlanarReflectionMode", DisplayName = "Planar Reflection Mode",
-		ToolTip = "The PlanarReflection will work differently on different mode on mobile platform, choose the proper mode as expect. Changing this setting requires restarting the editor.",
-		ConfigRestartRequired = true))
-		TEnumAsByte<EMobilePlanarReflectionMode::Type> MobilePlanarReflectionMode;
+	UE_DEPRECATED(5.6, "MobilePlanarReflectionMode is deprecated. The screen space reflections on mobile is supported and the pixel projected reflection is removed.")
+	UPROPERTY()
+	uint8 MobilePlanarReflectionMode;
 
 	UPROPERTY(config, EditAnywhere, Category = Mobile, meta = (
 		ConsoleVariable = "r.Mobile.ScreenSpaceReflections", DisplayName = "Support Screen Space Reflections on mobile",
diff --git a/Engine/Source/Runtime/Engine/Classes/Engine/Scene.h b/Engine/Source/Runtime/Engine/Classes/Engine/Scene.h
--- a/Engine/Source/Runtime/Engine/Classes/Engine/Scene.h
+++ b/Engine/Source/Runtime/Engine/Classes/Engine/Scene.h
@@ -115,34 +115,6 @@ enum class EReflectedAndRefractedRayTracedShadows : uint8
 	Area_shadows	UMETA(DisplayName = "Area Shadows"),
 };
 
-UENUM()
-namespace EMobilePlanarReflectionMode
-{
-	enum Type : int
-	{
-		Usual = 0 UMETA(DisplayName = "Usual", ToolTip = "The PlanarReflection actor works as usual on all platforms."),
-		MobilePPRExclusive = 1 UMETA(DisplayName = "MobilePPR Exclusive", ToolTip = "The PlanarReflection actor is only used for mobile pixel projection reflection, it will not affect PC/Console. MobileMSAA will be disabled as a side effect."),
-		MobilePPR = 2 UMETA(DisplayName = "MobilePPR", ToolTip = "The PlanarReflection actor still works as usual on PC/Console platform and is used for mobile pixel projected reflection on mobile platform. MobileMSAA will be disabled as a side effect."),
-	};
-}
-
-UENUM()
-namespace EMobilePixelProjectedReflectionQuality
-{
-	enum Type : int
-	{
-		Disabled = 0 UMETA(DisplayName = "Disabled", ToolTip = "Disabled."),
-		BestPerformance = 1 UMETA(DisplayName = "Best Performance", ToolTip = "Best performance but may have some artifacts in some view angles."),
-		BetterQuality = 2 UMETA(DisplayName = "Better Quality", ToolTip = "Better quality and reasonable performance and could fix some artifacts, but the PlanarReflection mesh has to render twice."),
-		BestQuality = 3 UMETA(DisplayName = "Best Quality", ToolTip = "Best quality but will be much heavier."),
-	};
-}
-
-ENGINE_API int32 GetMobilePlanarReflectionMode();
-ENGINE_API int32 GetMobilePixelProjectedReflectionQuality();
-ENGINE_API bool IsMobilePixelProjectedReflectionEnabled(EShaderPlatform ShaderPlatform);
-ENGINE_API bool IsUsingMobilePixelProjectedReflection(EShaderPlatform ShaderPlatform);
-
 USTRUCT(BlueprintType)
 struct FColorGradePerRangeSettings
 {
diff --git a/Engine/Source/Runtime/Engine/Private/Components/SceneCaptureComponent.cpp b/Engine/Source/Runtime/Engine/Private/Components/SceneCaptureComponent.cpp
--- a/Engine/Source/Runtime/Engine/Private/Components/SceneCaptureComponent.cpp
+++ b/Engine/Source/Runtime/Engine/Private/Components/SceneCaptureComponent.cpp
@@ -1210,10 +1210,7 @@ bool UPlanarReflectionComponent::ShouldComponentAddToScene() const
 {
 	bool bSceneAdd = USceneComponent::ShouldComponentAddToScene();
 
-	ERHIFeatureLevel::Type FeatureLevel = GetScene() ? GetScene()->GetFeatureLevel() : ERHIFeatureLevel::SM5;
-
-	// The PlanarReflectionComponent should not be added to scene if it is used only for mobile pixel projected reflection
-	return bSceneAdd && (FeatureLevel <= ERHIFeatureLevel::ES3_1 || GetMobilePlanarReflectionMode() != EMobilePlanarReflectionMode::MobilePPRExclusive);
+	return bSceneAdd;
 }
 
 #if WITH_EDITOR
diff --git a/Engine/Source/Runtime/Engine/Private/Materials/MaterialInterface.cpp b/Engine/Source/Runtime/Engine/Private/Materials/MaterialInterface.cpp
--- a/Engine/Source/Runtime/Engine/Private/Materials/MaterialInterface.cpp
+++ b/Engine/Source/Runtime/Engine/Private/Materials/MaterialInterface.cpp
@@ -472,8 +472,6 @@ FMaterialRelevance UMaterialInterface::GetRelevance_Internal(const UMaterial* Ma
 	const bool bIsMobile = InFeatureLevel <= ERHIFeatureLevel::ES3_1;
 	const bool bUsesSingleLayerWaterMaterial = MaterialResource->GetShadingModels().HasShadingModel(MSM_SingleLayerWater);
 	const bool bIsSinglePassWaterTranslucent = bIsMobile && bUsesSingleLayerWaterMaterial;
-	const bool bIsMobilePixelProjectedTranslucent = MaterialResource->IsUsingPlanarForwardReflections() 
-													&& IsUsingMobilePixelProjectedReflection(GetFeatureLevelShaderPlatform(InFeatureLevel));
 
 	// Note that even though XX_GameThread() api is called, this function can be called on non game thread via 
 	// GetRelevance_Concurrent()
@@ -481,7 +479,7 @@ FMaterialRelevance UMaterialInterface::GetRelevance_Internal(const UMaterial* Ma
 		MaterialResource->MaterialUsesAnisotropy_GameThread();
 
 	const EBlendMode BlendMode = (EBlendMode)GetBlendMode();
-	const bool bIsTranslucent = IsTranslucentBlendMode(BlendMode) || bIsSinglePassWaterTranslucent || bIsMobilePixelProjectedTranslucent; // We want meshes with water materials to be scheduled for translucent pass on mobile. And we also have to render the meshes used for mobile pixel projection reflection in translucent pass.
+	const bool bIsTranslucent = IsTranslucentBlendMode(BlendMode) || bIsSinglePassWaterTranslucent; // We want meshes with water materials to be scheduled for translucent pass on mobile.
 
 	EMaterialDomain Domain = (EMaterialDomain)MaterialResource->GetMaterialDomain();
 	bool bDecal = (Domain == MD_DeferredDecal);
diff --git a/Engine/Source/Runtime/Engine/Private/Materials/MaterialShared.cpp b/Engine/Source/Runtime/Engine/Private/Materials/MaterialShared.cpp
--- a/Engine/Source/Runtime/Engine/Private/Materials/MaterialShared.cpp
+++ b/Engine/Source/Runtime/Engine/Private/Materials/MaterialShared.cpp
@@ -2000,9 +2000,7 @@ bool FMaterialResource::GetForwardBlendsSkyLightCubemaps() const
 
 bool FMaterialResource::IsUsingPlanarForwardReflections() const
 {
-	return Material->bUsePlanarForwardReflections
-		// Don't use planar reflection if it is used only for mobile pixel projected reflection.
-		&& (GetFeatureLevel() <= ERHIFeatureLevel::ES3_1 || GetMobilePlanarReflectionMode() != EMobilePlanarReflectionMode::MobilePPRExclusive);
+	return Material->bUsePlanarForwardReflections;
 }
 
 bool FMaterialResource::IsNonmetal() const
diff --git a/Engine/Source/Runtime/Engine/Private/Scene.cpp b/Engine/Source/Runtime/Engine/Private/Scene.cpp
--- a/Engine/Source/Runtime/Engine/Private/Scene.cpp
+++ b/Engine/Source/Runtime/Engine/Private/Scene.cpp
@@ -11,30 +11,6 @@
 #include "DataDrivenShaderPlatformInfo.h"
 #include "UObject/UnrealType.h"
 
-int32 GetMobilePlanarReflectionMode()
-{
-	static const auto MobilePlanarReflectionModeCVar = IConsoleManager::Get().FindTConsoleVariableDataInt(TEXT("r.Mobile.PlanarReflectionMode"));
-
-	return MobilePlanarReflectionModeCVar->GetValueOnAnyThread();
-}
-
-int32 GetMobilePixelProjectedReflectionQuality()
-{
-	static const auto MobilePixelProjectedReflectionQualityCVar = IConsoleManager::Get().FindTConsoleVariableDataInt(TEXT("r.Mobile.PixelProjectedReflectionQuality"));
-
-	return MobilePixelProjectedReflectionQualityCVar->GetValueOnAnyThread();
-}
-
-bool IsMobilePixelProjectedReflectionEnabled(EShaderPlatform ShaderPlatform)
-{
-	return IsMobilePlatform(ShaderPlatform) && IsMobileHDR() && (GetMobilePlanarReflectionMode() == EMobilePlanarReflectionMode::MobilePPRExclusive || GetMobilePlanarReflectionMode() == EMobilePlanarReflectionMode::MobilePPR);
-}
-
-bool IsUsingMobilePixelProjectedReflection(EShaderPlatform ShaderPlatform)
-{
-	return IsMobilePixelProjectedReflectionEnabled(ShaderPlatform) && GetMobilePixelProjectedReflectionQuality() > EMobilePixelProjectedReflectionQuality::Disabled;
-}
-
 void FColorGradingSettings::ExportToPostProcessSettings(FPostProcessSettings* OutPostProcessSettings) const
 {
 	OutPostProcessSettings->bOverride_ColorSaturation = true;
diff --git a/Engine/Source/Runtime/Engine/Private/SceneUtils.cpp b/Engine/Source/Runtime/Engine/Private/SceneUtils.cpp
--- a/Engine/Source/Runtime/Engine/Private/SceneUtils.cpp
+++ b/Engine/Source/Runtime/Engine/Private/SceneUtils.cpp
@@ -205,18 +205,6 @@ ENGINE_API uint32 GetDefaultMSAACount(const FStaticFeatureLevel InFeatureLevel,
 
 			bool bRHISupportsMSAA = RHISupportsMSAA(ShaderPlatform);
 
-			if (InFeatureLevel == ERHIFeatureLevel::ES3_1)
-			{
-				bool bMobilePixelProjectedReflection = IsUsingMobilePixelProjectedReflection(ShaderPlatform);
-
-				bRendererSupportMSAA = bRHISupportsMSAA && !bMobilePixelProjectedReflection;
-
-				if (!bRendererSupportMSAA)
-				{
-					FailedReason = FString::Printf(TEXT("RHISupportsMSAA %d, MobilePixelProjectedReflection %d"), bRHISupportsMSAA ? 1 : 0, bMobilePixelProjectedReflection ? 1 : 0);
-				}
-			}
-			else
 			{
 				bRendererSupportMSAA = bRHISupportsMSAA;
 
diff --git a/Engine/Source/Runtime/Engine/Public/SceneTexturesConfig.h b/Engine/Source/Runtime/Engine/Public/SceneTexturesConfig.h
--- a/Engine/Source/Runtime/Engine/Public/SceneTexturesConfig.h
+++ b/Engine/Source/Runtime/Engine/Public/SceneTexturesConfig.h
@@ -159,9 +159,6 @@ struct FSceneTexturesConfig
 	// Extent of all full-resolution textures.
 	FIntPoint Extent = FIntPoint::ZeroValue;
 
-	// Extend of the mobile Pixel Projected Reflection texture
-	FIntPoint MobilePixelProjectedReflectionExtent = FIntPoint::ZeroValue;
-
 	// Downsample factors to divide against the full resolution texture extent.
 	uint32 SmallDepthDownsampleFactor = 2;
 
diff --git a/Engine/Source/Runtime/Renderer/Internal/SceneTextures.h b/Engine/Source/Runtime/Renderer/Internal/SceneTextures.h
--- a/Engine/Source/Runtime/Renderer/Internal/SceneTextures.h
+++ b/Engine/Source/Runtime/Renderer/Internal/SceneTextures.h
@@ -151,9 +151,6 @@ struct FSceneTextures : public FMinimalSceneTextures
 	// Texture used by the quad overdraw debug view mode when enabled.
 	FRDGTextureRef QuadOverdraw{};
 
-	// (Mobile) Texture used by mobile PPR in the next frame.
-	FRDGTextureRef PixelProjectedReflection{};
-
 	// Textures used to composite editor primitives. Also used by the base pass when in wireframe mode.
 #if WITH_EDITOR
 	FRDGTextureRef EditorPrimitiveColor{};
diff --git a/Engine/Source/Runtime/Renderer/Private/MobileBasePass.cpp b/Engine/Source/Runtime/Renderer/Private/MobileBasePass.cpp
--- a/Engine/Source/Runtime/Renderer/Private/MobileBasePass.cpp
+++ b/Engine/Source/Runtime/Renderer/Private/MobileBasePass.cpp
@@ -521,8 +521,6 @@ void MobileBasePass::SetTranslucentRenderState(FMeshPassProcessorRenderState& Dr
 	constexpr ERHIFeatureLevel::Type FeatureLevel = ERHIFeatureLevel::ES3_1;
 	EShaderPlatform ShaderPlatform = GetFeatureLevelShaderPlatform(FeatureLevel);
 
-	const bool bIsUsingMobilePixelProjectedReflection = Material.IsUsingPlanarForwardReflections() && IsUsingMobilePixelProjectedReflection(ShaderPlatform);
-	
 	if (Substrate::IsSubstrateEnabled())
 	{
 		if (Material.IsDualBlendingEnabled(ShaderPlatform) || Material.GetBlendMode() == BLEND_TranslucentColoredTransmittance)
@@ -861,13 +859,10 @@ bool FMobileBasePassMeshProcessor::Process(
 	FMeshDrawCommandSortKey SortKey; 
 	if (bTranslucentBasePass)
 	{
-		const bool bIsUsingMobilePixelProjectedReflection = MaterialResource.IsUsingPlanarForwardReflections() 
-															&& IsUsingMobilePixelProjectedReflection(GetFeatureLevelShaderPlatform(FeatureLevel));
-
 		SortKey = CalculateTranslucentMeshStaticSortKey(PrimitiveSceneProxy, MeshBatch.MeshIdInPrimitive);
 		// We always want water to be rendered first on mobile in order to mimic other renderers where it is opaque. We shift the other priorities by 1.
 		// And we also want to render the meshes used for mobile pixel projected reflection first if it is opaque.
-		SortKey.Translucent.Priority = ShadingModels.HasShadingModel(MSM_SingleLayerWater) || (!bIsTranslucent && bIsUsingMobilePixelProjectedReflection) ? uint16(0) : uint16(FMath::Clamp(uint32(SortKey.Translucent.Priority) + 1, 0u, uint32(USHRT_MAX)));
+		SortKey.Translucent.Priority = ShadingModels.HasShadingModel(MSM_SingleLayerWater) ? uint16(0) : uint16(FMath::Clamp(uint32(SortKey.Translucent.Priority) + 1, 0u, uint32(USHRT_MAX)));
 	}
 	else
 	{
diff --git a/Engine/Source/Runtime/Renderer/Private/MobileBasePassRendering.cpp b/Engine/Source/Runtime/Renderer/Private/MobileBasePassRendering.cpp
--- a/Engine/Source/Runtime/Renderer/Private/MobileBasePassRendering.cpp
+++ b/Engine/Source/Runtime/Renderer/Private/MobileBasePassRendering.cpp
@@ -17,7 +17,6 @@
 #include "EditorPrimitivesRendering.h"
 
 #include "FramePro/FrameProProfiler.h"
-#include "PostProcess/PostProcessPixelProjectedReflectionMobile.h"
 #include "Engine/SubsurfaceProfile.h"
 #include "Engine/SpecularProfile.h"
 #include "LocalLightSceneProxy.h"
@@ -301,12 +300,7 @@ void SetupMobileBasePassUniformParameters(
 	const FScene* Scene = View.Family->Scene ? View.Family->Scene->GetRenderScene() : nullptr;
 	const FPlanarReflectionSceneProxy* ReflectionSceneProxy = Scene ? Scene->GetForwardPassGlobalPlanarReflection() : nullptr;
 	SetupPlanarReflectionUniformParameters(View, ReflectionSceneProxy, BasePassParameters.PlanarReflection);
-	if (View.PrevViewInfo.MobilePixelProjectedReflection.IsValid())
-	{
-		BasePassParameters.PlanarReflection.PlanarReflectionTexture = View.PrevViewInfo.MobilePixelProjectedReflection->GetRHI();
-		BasePassParameters.PlanarReflection.PlanarReflectionSampler = TStaticSamplerState<SF_Bilinear, AM_Clamp, AM_Clamp, AM_Clamp>::GetRHI();
-	}
-	else if (BasePassParameters.PlanarReflection.PlanarReflectionTexture == nullptr)
+	if (BasePassParameters.PlanarReflection.PlanarReflectionTexture == nullptr)
 	{
 		BasePassParameters.PlanarReflection.PlanarReflectionTexture = GBlackTexture->TextureRHI;
 		BasePassParameters.PlanarReflection.PlanarReflectionSampler = GBlackTexture->SamplerStateRHI;
diff --git a/Engine/Source/Runtime/Renderer/Private/MobileShadingRenderer.cpp b/Engine/Source/Runtime/Renderer/Private/MobileShadingRenderer.cpp
--- a/Engine/Source/Runtime/Renderer/Private/MobileShadingRenderer.cpp
+++ b/Engine/Source/Runtime/Renderer/Private/MobileShadingRenderer.cpp
@@ -32,7 +32,6 @@
 #include "PostProcess/PostProcessUpscale.h"
 #include "PostProcess/PostProcessCompositeEditorPrimitives.h"
 #include "PostProcess/PostProcessHMD.h"
-#include "PostProcess/PostProcessPixelProjectedReflectionMobile.h"
 #include "PostProcess/PostProcessAmbientOcclusionMobile.h"
 #include "PostProcess/PostProcessCombineLUTs.h"
 #include "PostProcess/PostProcessTonemap.h"
@@ -281,7 +280,6 @@ FMobileSceneRenderer::FMobileSceneRenderer(const FSceneViewFamily* InViewFamily,
 	bKeepDepthContent = false;
 	bModulatedShadowsInUse = false;
 	bShouldRenderCustomDepth = false;
-	bRequiresPixelProjectedPlanarRelfectionPass = false;
 	bRequiresAmbientOcclusionPass = false;
 	bRequiresShadowProjections = false;
 	bIsFullDepthPrepassEnabled = (Scene->EarlyZPassMode == DDM_AllOpaque || Scene->EarlyZPassMode == DDM_AllOpaqueNoVelocity);
@@ -530,16 +528,6 @@ void FMobileSceneRenderer::InitViews(
 
 	const FPlanarReflectionSceneProxy* PlanarReflectionSceneProxy = Scene ? Scene->GetForwardPassGlobalPlanarReflection() : nullptr;
 
-	bRequiresPixelProjectedPlanarRelfectionPass = IsUsingMobilePixelProjectedReflection(ShaderPlatform)
-		&& PlanarReflectionSceneProxy != nullptr
-		&& PlanarReflectionSceneProxy->RenderTarget != nullptr
-		&& !Views[0].bIsReflectionCapture
-		&& !ViewFamily.EngineShowFlags.HitProxies
-		&& ViewFamily.EngineShowFlags.Lighting
-		&& !ViewFamily.EngineShowFlags.VisualizeLightCulling
-		&& !ViewFamily.UseDebugViewPS()
-		&& bRendererOutputFinalSceneColor;
-
 	bRequiresAmbientOcclusionPass = IsUsingMobileAmbientOcclusion(ShaderPlatform)
 		&& Views[0].FinalPostProcessSettings.AmbientOcclusionIntensity > 0
 		&& (Views[0].FinalPostProcessSettings.AmbientOcclusionStaticFraction >= 1 / 100.0f || (Scene && Scene->SkyLight && Scene->SkyLight->ProcessedTexture && Views[0].Family->EngineShowFlags.SkyLighting))
@@ -575,7 +563,6 @@ void FMobileSceneRenderer::InitViews(
 	bKeepDepthContent =
 		bRequiresMultiPass ||
 		bForceDepthResolve ||
-		bRequiresPixelProjectedPlanarRelfectionPass ||
 		bSeparateTranslucencyActive ||
 		Views[0].bIsReflectionCapture ||
 		(bDeferredShading && bPostProcessUsesSceneDepth) ||
@@ -611,16 +598,6 @@ void FMobileSceneRenderer::InitViews(
 		SceneTexturesConfig.SetupMobileGBufferFlags(bRequiresMultiPass || GraphBuilder.IsDumpingFrame() || bRequireSeparateViewPass);
 	}
 
-	// Update the pixel projected reflection extent according to the settings in the PlanarReflectionComponent.
-	if (bRequiresPixelProjectedPlanarRelfectionPass)
-	{
-		SceneTexturesConfig.MobilePixelProjectedReflectionExtent = PlanarReflectionSceneProxy->RenderTarget->GetSizeXY();
-	}
-	else
-	{
-		SceneTexturesConfig.MobilePixelProjectedReflectionExtent = FIntPoint::ZeroValue;
-	}
-
 	SceneTexturesConfig.bRequiresDepthAux = bRequiresSceneDepthAux;
 	// When we capturing scene depth, use a more precise format for SceneDepthAux as it will be used as a source DepthTexture
 	if (bSceneDepthCapture)
@@ -1453,13 +1430,6 @@ void FMobileSceneRenderer::Render(FRDGBuilder& GraphBuilder)
 
 		RenderOpaqueFX(GraphBuilder, GetSceneViews(), GetSceneUniforms(), FXSystem, SceneTextures.MobileUniformBuffer);
 
-		if (bRequiresPixelProjectedPlanarRelfectionPass)
-		{
-			const FPlanarReflectionSceneProxy* PlanarReflectionSceneProxy = Scene ? Scene->GetForwardPassGlobalPlanarReflection() : nullptr;
-
-			RenderPixelProjectedReflection(GraphBuilder, SceneTextures.Color.Resolve, SceneTextures.Depth.Resolve, SceneTextures.PixelProjectedReflection, PlanarReflectionSceneProxy);
-		}
-
 		if (ViewFamily.EngineShowFlags.VisualizeMeshDistanceFields || ViewFamily.EngineShowFlags.VisualizeGlobalDistanceField)
 		{
 			RenderMeshDistanceFieldVisualization(GraphBuilder, SceneTextures);
diff --git a/Engine/Source/Runtime/Renderer/Private/PlanarReflectionRendering.cpp b/Engine/Source/Runtime/Renderer/Private/PlanarReflectionRendering.cpp
--- a/Engine/Source/Runtime/Renderer/Private/PlanarReflectionRendering.cpp
+++ b/Engine/Source/Runtime/Renderer/Private/PlanarReflectionRendering.cpp
@@ -274,7 +274,6 @@ struct FUpdatePlanarReflectionContentsInputs
 	FPlanarReflectionSceneProxy* SceneProxy;
 	FPlanarReflectionRenderTarget* RenderTarget;
 	FPlane MirrorPlane;
-	bool bIsMobilePixelProjectedReflectionEnabled;
 };
 
 static void UpdatePlanarReflectionContents_RenderThread(
@@ -365,13 +364,6 @@ static void UpdatePlanarReflectionContents_RenderThread(
 		Inputs.SceneProxy->RenderTarget = Inputs.RenderTarget;
 	};
 
-	// Generate valid data to update planar reflection uniform buffer but don't actually render the reflection scene when we are using mobile pixel projected reflection.
-	if (Inputs.bIsMobilePixelProjectedReflectionEnabled)
-	{
-		SceneRenderer->PrepareViewRectsForRendering(GraphBuilder.RHICmdList);
-		return;
-	}
-
 	const ERHIFeatureLevel::Type FeatureLevel = SceneRenderer->FeatureLevel;
 
 	// Make sure we render to the same set of GPUs as the main scene renderer.
@@ -451,12 +443,9 @@ void FScene::UpdatePlanarReflectionContents(UPlanarReflectionComponent* CaptureC
 	DesiredPlanarReflectionTextureSize.X = FMath::Clamp(FMath::CeilToInt32(DesiredPlanarReflectionTextureSizeFloat.X), 1, static_cast<int32>(DesiredBufferSize.X));
 	DesiredPlanarReflectionTextureSize.Y = FMath::Clamp(FMath::CeilToInt32(DesiredPlanarReflectionTextureSizeFloat.Y), 1, static_cast<int32>(DesiredBufferSize.Y));
 
-	const bool bIsMobilePixelProjectedReflectionEnabled = IsMobilePixelProjectedReflectionEnabled(GetShaderPlatform());
-
 	const bool bIsRenderTargetValid = CaptureComponent->RenderTarget != nullptr
 								&& CaptureComponent->RenderTarget->GetSizeXY() == DesiredPlanarReflectionTextureSize
-								// The RenderTarget's TextureRHI could be nullptr if it is used for mobile pixel projected reflection.
-								&& (bIsMobilePixelProjectedReflectionEnabled || CaptureComponent->RenderTarget->TextureRHI.IsValid());
+								&& (CaptureComponent->RenderTarget->TextureRHI.IsValid());
 
 	FPlanarReflectionRenderTarget* RenderTargetToRelease = nullptr;
 	FPlanarReflectionRenderTarget* RenderTargetToInit = nullptr;
@@ -471,11 +460,7 @@ void FScene::UpdatePlanarReflectionContents(UPlanarReflectionComponent* CaptureC
 	{
 		CaptureComponent->RenderTarget = new FPlanarReflectionRenderTarget(DesiredPlanarReflectionTextureSize);
 
-		// Don't create the RenderTarget's RHI if it is used for mobile pixel projected reflection
-		if (!bIsMobilePixelProjectedReflectionEnabled)
-		{
-			RenderTargetToInit = CaptureComponent->RenderTarget;
-		}
+		RenderTargetToInit = CaptureComponent->RenderTarget;
 	}
 
 	const FMatrix ComponentTransform = CaptureComponent->GetComponentTransform().ToMatrixWithScale();
@@ -622,7 +607,6 @@ void FScene::UpdatePlanarReflectionContents(UPlanarReflectionComponent* CaptureC
 		, .SceneProxy = CaptureComponent->SceneProxy
 		, .RenderTarget = CaptureComponent->RenderTarget
 		, .MirrorPlane = MirrorPlane
-		, .bIsMobilePixelProjectedReflectionEnabled = bIsMobilePixelProjectedReflectionEnabled
 	};
 
 	FSceneRenderer* SceneRenderer = FSceneRenderer::CreateSceneRenderer(&ViewFamily, nullptr);
diff --git a/Engine/Source/Runtime/Renderer/Private/PostProcess/PostProcessPixelProjectedReflectionMobile.cpp b/Engine/Source/Runtime/Renderer/Private/PostProcess/PostProcessPixelProjectedReflectionMobile.cpp
deleted file mode 100644
diff --git a/Engine/Source/Runtime/Renderer/Private/PostProcess/PostProcessPixelProjectedReflectionMobile.h b/Engine/Source/Runtime/Renderer/Private/PostProcess/PostProcessPixelProjectedReflectionMobile.h
deleted file mode 100644
diff --git a/Engine/Source/Runtime/Renderer/Private/RendererScene.cpp b/Engine/Source/Runtime/Renderer/Private/RendererScene.cpp
--- a/Engine/Source/Runtime/Renderer/Private/RendererScene.cpp
+++ b/Engine/Source/Runtime/Renderer/Private/RendererScene.cpp
@@ -860,7 +860,6 @@ uint64 FPreviousViewInfo::GetGPUSizeBytes(bool bLogSizes) const
 		ReflectedSkyLightHistory.GetGPUSizeBytes(bLogSizes) +
 		PolychromaticPenumbraHarmonicsHistory.GetGPUSizeBytes(bLogSizes) +
 		GetRenderTargetGPUSizeBytes(MobileBloomSetup_EyeAdaptation, bLogSizes) +
-		GetRenderTargetGPUSizeBytes(MobilePixelProjectedReflection, bLogSizes) +
 		GetRenderTargetGPUSizeBytes(MobileAmbientOcclusion, bLogSizes) +
 		GetRenderTargetGPUSizeBytes(VisualizeMotionVectors, bLogSizes);
 
diff --git a/Engine/Source/Runtime/Renderer/Private/SceneRendering.cpp b/Engine/Source/Runtime/Renderer/Private/SceneRendering.cpp
--- a/Engine/Source/Runtime/Renderer/Private/SceneRendering.cpp
+++ b/Engine/Source/Runtime/Renderer/Private/SceneRendering.cpp
@@ -4154,11 +4154,7 @@ void FSceneRenderer::OnRenderFinish(FRDGBuilder& GraphBuilder, FRDGTextureRef Vi
 		{
 			static const auto* CVarClipPlane = IConsoleManager::Get().FindTConsoleVariableDataInt(TEXT("r.AllowGlobalClipPlane"));
 			
-			const bool bShouldUseClipPlaneForPlanarReflection = (FeatureLevel > ERHIFeatureLevel::ES3_1 && GetMobilePlanarReflectionMode() != EMobilePlanarReflectionMode::MobilePPRExclusive)
-															|| GetMobilePlanarReflectionMode() == EMobilePlanarReflectionMode::Usual;
-			
-			if (CVarClipPlane && CVarClipPlane->GetValueOnRenderThread() == 0
-				&& bShouldUseClipPlaneForPlanarReflection)
+			if (CVarClipPlane && CVarClipPlane->GetValueOnRenderThread() == 0)
 			{
 				bShowGlobalClipPlaneWarning = true;
 			}
diff --git a/Engine/Source/Runtime/Renderer/Private/SceneRendering.h b/Engine/Source/Runtime/Renderer/Private/SceneRendering.h
--- a/Engine/Source/Runtime/Renderer/Private/SceneRendering.h
+++ b/Engine/Source/Runtime/Renderer/Private/SceneRendering.h
@@ -1080,9 +1080,6 @@ struct FPreviousViewInfo
 	// Mobile bloom setup eye adaptation surface.
 	TRefCountPtr<IPooledRenderTarget> MobileBloomSetup_EyeAdaptation;
 
-	// Mobile pixel projected reflection textureused for next frame.
-	TRefCountPtr<IPooledRenderTarget> MobilePixelProjectedReflection = nullptr;
-
 	// Mobile ambient occlusion texture used for next frame.
 	TRefCountPtr<IPooledRenderTarget> MobileAmbientOcclusion = nullptr;
 
@@ -2780,8 +2777,6 @@ protected:
 
 	void RenderAmbientOcclusion(FRDGBuilder& GraphBuilder, FRDGTextureRef SceneDepthTexture, FRDGTextureRef AmbientOcclusionTexture);
 
-	void RenderPixelProjectedReflection(FRDGBuilder& GraphBuilder, FRDGTextureRef SceneColorTexture, FRDGTextureRef SceneDepthTexture, FRDGTextureRef PixelProjectedReflectionTexture, const FPlanarReflectionSceneProxy* PlanarReflectionSceneProxy);
-
 	void RenderMobileShadowProjections(FRDGBuilder& GraphBuilder);
 
 	FRenderTargetBindingSlots InitRenderTargetBindings_Deferred(FSceneTextures& SceneTextures, FColorTargets& ColorTargets);
@@ -2802,7 +2797,6 @@ private:
 	bool bKeepDepthContent;
 	bool bModulatedShadowsInUse;
 	bool bShouldRenderCustomDepth;
-	bool bRequiresPixelProjectedPlanarRelfectionPass;
 	bool bRequiresAmbientOcclusionPass;
 	bool bShouldRenderVelocities;
 	bool bShouldRenderHZB;
diff --git a/Engine/Source/Runtime/Renderer/Private/SceneTextures.cpp b/Engine/Source/Runtime/Renderer/Private/SceneTextures.cpp
--- a/Engine/Source/Runtime/Renderer/Private/SceneTextures.cpp
+++ b/Engine/Source/Runtime/Renderer/Private/SceneTextures.cpp
@@ -19,7 +19,6 @@
 #include "SystemTextures.h"
 #include "PostProcess/PostProcessing.h"
 #include "PostProcess/PostProcessAmbientOcclusionMobile.h"
-#include "PostProcess/PostProcessPixelProjectedReflectionMobile.h"
 #include "IHeadMountedDisplayModule.h"
 #include "Substrate/Substrate.h"
 #include "VisualizeTexture.h"
@@ -681,11 +680,6 @@ void FSceneTextures::InitializeViewFamily(FRDGBuilder& GraphBuilder, FViewFamily
 	{
 		// Mobile Screen Space Ambient Occlusion
 		SceneTextures.ScreenSpaceAO = CreateMobileScreenSpaceAOTexture(GraphBuilder, Config);
-
-		if (Config.MobilePixelProjectedReflectionExtent != FIntPoint::ZeroValue)
-		{
-			SceneTextures.PixelProjectedReflection = CreateMobilePixelProjectedReflectionTexture(GraphBuilder, Config.MobilePixelProjectedReflectionExtent);
-		}
 	}
 
 	// Velocity
@@ -921,7 +915,6 @@ TArray<FRDGTextureRef> FSceneTextures::EnumerateSceneTextures() const
 	AddTextureIfNonNull(MobileLocalLightTextureB, Results);
 	AddTextureIfNonNull(ScreenSpaceAO, Results);
 	AddTextureIfNonNull(QuadOverdraw, Results);
-	AddTextureIfNonNull(PixelProjectedReflection, Results);
 #if WITH_EDITOR
 	AddTextureIfNonNull(EditorPrimitiveColor, Results);
 	AddTextureIfNonNull(EditorPrimitiveDepth, Results);
