[HWRT] Fix Lumen SBT not being released when using path tracing.
- caused a crash because ended up with null View.LumenHardwareRayTracingMaterialPipeline.

diff --git a/Engine/Source/Runtime/Renderer/Private/DeferredShadingRenderer.cpp b/Engine/Source/Runtime/Renderer/Private/DeferredShadingRenderer.cpp
--- a/Engine/Source/Runtime/Renderer/Private/DeferredShadingRenderer.cpp
+++ b/Engine/Source/Runtime/Renderer/Private/DeferredShadingRenderer.cpp
@@ -699,43 +699,49 @@ bool FDeferredShadingSceneRenderer::SetupRayTracingPipelineStatesAndSBT(FRDGBuil
 		}
 	}
 
-	// Add Lumen hardware ray tracing materials
-	if (!bIsPathTracing)
+	// Create Lumen hardware ray tracing SBT and material pipeline
 	{
-		TArray<FRHIRayTracingShader*> LumenHardwareRayTracingRayGenShaders;
+		ERayTracingShaderBindingMode LumenShaderBindingMode = ERayTracingShaderBindingMode::Disabled;
 
-		if (GRHISupportsRayTracingShaders)
+		if (!bIsPathTracing)
 		{
-			if (DoesPlatformSupportLumenGI(ShaderPlatform))
+			if (bAnyLumenHardwareInlineRayTracingPassEnabled && GRHIGlobals.RayTracing.RequiresInlineRayTracingSBT)
+			{
+				EnumAddFlags(LumenShaderBindingMode, ERayTracingShaderBindingMode::Inline);
+			}
+
+			TArray<FRHIRayTracingShader*> LumenHardwareRayTracingRayGenShaders;
+
+			if (GRHISupportsRayTracingShaders)
 			{
+				if (DoesPlatformSupportLumenGI(ShaderPlatform))
+				{
+					for (const FViewInfo& View : Views)
+					{
+						PrepareLumenHardwareRayTracingVisualizeLumenMaterial(View, LumenHardwareRayTracingRayGenShaders);
+						PrepareLumenHardwareRayTracingRadianceCacheLumenMaterial(View, LumenHardwareRayTracingRayGenShaders);
+						PrepareLumenHardwareRayTracingTranslucencyVolumeLumenMaterial(View, LumenHardwareRayTracingRayGenShaders);
+						PrepareLumenHardwareRayTracingRadiosityLumenMaterial(View, LumenHardwareRayTracingRayGenShaders);
+						PrepareLumenHardwareRayTracingReflectionsLumenMaterial(View, LumenHardwareRayTracingRayGenShaders);
+						PrepareLumenHardwareRayTracingReSTIRLumenMaterial(View, LumenHardwareRayTracingRayGenShaders);
+						PrepareLumenHardwareRayTracingScreenProbeGatherLumenMaterial(View, LumenHardwareRayTracingRayGenShaders);
+						PrepareLumenHardwareRayTracingDirectLightingLumenMaterial(View, LumenHardwareRayTracingRayGenShaders);
+					}
+				}
+
 				for (const FViewInfo& View : Views)
 				{
-					PrepareLumenHardwareRayTracingVisualizeLumenMaterial(View, LumenHardwareRayTracingRayGenShaders);
-					PrepareLumenHardwareRayTracingRadianceCacheLumenMaterial(View, LumenHardwareRayTracingRayGenShaders);
-					PrepareLumenHardwareRayTracingTranslucencyVolumeLumenMaterial(View, LumenHardwareRayTracingRayGenShaders);
-					PrepareLumenHardwareRayTracingRadiosityLumenMaterial(View, LumenHardwareRayTracingRayGenShaders);
-					PrepareLumenHardwareRayTracingReflectionsLumenMaterial(View, LumenHardwareRayTracingRayGenShaders);
-					PrepareLumenHardwareRayTracingReSTIRLumenMaterial(View, LumenHardwareRayTracingRayGenShaders);
-					PrepareLumenHardwareRayTracingScreenProbeGatherLumenMaterial(View, LumenHardwareRayTracingRayGenShaders);
-					PrepareLumenHardwareRayTracingDirectLightingLumenMaterial(View, LumenHardwareRayTracingRayGenShaders);
+					PrepareMegaLightsHardwareRayTracingLumenMaterial(View, LumenHardwareRayTracingRayGenShaders);
 				}
+
+				DeduplicateRayGenerationShaders(LumenHardwareRayTracingRayGenShaders);
 			}
 
-			for (const FViewInfo& View : Views)
+			if (LumenHardwareRayTracingRayGenShaders.Num())
 			{
-				PrepareMegaLightsHardwareRayTracingLumenMaterial(View, LumenHardwareRayTracingRayGenShaders);
+				CreateLumenHardwareRayTracingMaterialPipeline(GraphBuilder, ReferenceView, LumenHardwareRayTracingRayGenShaders, MaxLocalBindingDataSize);
+				EnumAddFlags(LumenShaderBindingMode, ERayTracingShaderBindingMode::RTPSO);
 			}
-
-			DeduplicateRayGenerationShaders(LumenHardwareRayTracingRayGenShaders);
-		}
-
-		ERayTracingShaderBindingMode LumenShaderBindingMode = (bAnyLumenHardwareInlineRayTracingPassEnabled && GRHIGlobals.RayTracing.RequiresInlineRayTracingSBT) ? 
-			ERayTracingShaderBindingMode::Inline : ERayTracingShaderBindingMode::Disabled;
-
-		if (LumenHardwareRayTracingRayGenShaders.Num())
-		{
-			CreateLumenHardwareRayTracingMaterialPipeline(GraphBuilder, ReferenceView, LumenHardwareRayTracingRayGenShaders, MaxLocalBindingDataSize);
-			EnumAddFlags(LumenShaderBindingMode, ERayTracingShaderBindingMode::RTPSO);
 		}
 
 		if (LumenShaderBindingMode != ERayTracingShaderBindingMode::Disabled)
@@ -748,11 +754,12 @@ bool FDeferredShadingSceneRenderer::SetupRayTracingPipelineStatesAndSBT(FRDGBuil
 		}
 		else if (Scene->LumenRayTracingSBTID != INDEX_NONE)
 		{
+			// release Lumen persistent SBT when not needed
 			Scene->RayTracingSBT.ReleasePersistentSBT(Scene->LumenRayTracingSBTID);
 			Scene->LumenRayTracingSBTID = INDEX_NONE;
 		}
 
-		// Merge with shared binding mode
+		// Merge Lumen shader binding mode with the shared shader binding mode
 		EnumAddFlags(ShaderBindingMode, LumenShaderBindingMode);
 	}
 
