Optimize the maximum culling views using the computed MinMipLevel

#rnx

diff --git a/Engine/Source/Runtime/Renderer/Private/Shadows/ShadowSceneRenderer.cpp b/Engine/Source/Runtime/Renderer/Private/Shadows/ShadowSceneRenderer.cpp
--- a/Engine/Source/Runtime/Renderer/Private/Shadows/ShadowSceneRenderer.cpp
+++ b/Engine/Source/Runtime/Renderer/Private/Shadows/ShadowSceneRenderer.cpp
@@ -589,24 +589,12 @@ void FShadowSceneRenderer::RenderVirtualShadowMaps(FRDGBuilder& GraphBuilder, bo
 
 	if (bNaniteEnabled && NaniteRenderPasses.Num() > 0)
 	{
-		// Little awkward but we don't want the interface in the rendering to depend on our NaniteRenderPasses structure
-		TArray<Nanite::FPackedViewArray*, SceneRenderingAllocator> VirtualShadowMapViews;
-		TArray<FSceneInstanceCullingQuery*, SceneRenderingAllocator> SceneInstanceCullingQueries;
-		VirtualShadowMapViews.SetNumUninitialized(NaniteRenderPasses.Num());
-		SceneInstanceCullingQueries.SetNumUninitialized(NaniteRenderPasses.Num());
-		for (int32 Index = 0; Index < NaniteRenderPasses.Num(); ++Index)
-		{
-			VirtualShadowMapViews[Index] = NaniteRenderPasses[Index].VirtualShadowMapViews;
-			SceneInstanceCullingQueries[Index] = NaniteRenderPasses[Index].SceneInstanceCullingQuery;
-		}
-
 		VirtualShadowMapArray.RenderVirtualShadowMapsNanite(
 			GraphBuilder,
 			SceneRenderer,
 			bUpdateNaniteStreaming,
 			NaniteVisibilityQuery,
-			VirtualShadowMapViews,
-			SceneInstanceCullingQueries);
+			NaniteRenderPasses);
 	}
 
 	if (UseNonNaniteVirtualShadowMaps(SceneRenderer.ShaderPlatform, SceneRenderer.FeatureLevel))
@@ -627,7 +615,7 @@ void FShadowSceneRenderer::DispatchVirtualShadowMapViewAndCullingSetup(FRDGBuild
 		CreateNaniteRenderPasses(GraphBuilder, SceneRenderer.Views, VirtualShadowMapShadows);
 
 		// Dispatch collected queries
-		for (FNaniteRenderPass& RenderPass : NaniteRenderPasses)
+		for (FNaniteVirtualShadowMapRenderPass& RenderPass : NaniteRenderPasses)
 		{
 			RenderPass.SceneInstanceCullingQuery->Dispatch(GraphBuilder);
 		}
@@ -879,9 +867,9 @@ void FShadowSceneRenderer::ApplyVirtualShadowMapProjectionForLight(
 
 struct FVSMRenderViewInfo
 {
-	uint32 NumPrimaryViews = 0u;
-	uint32 MaxNumViews = 0u;	// Including worst case mips
 	FCullingVolume CullingVolume;
+	uint32 NumPrimaryViews = 0u;
+	uint32 MaxCullingViews = 0u;
 };
 
 static FVSMRenderViewInfo GetRenderViewInfo(const FProjectedShadowInfo* ProjectedShadowInfo)
@@ -894,7 +882,7 @@ static FVSMRenderViewInfo GetRenderViewInfo(const FProjectedShadowInfo* Projecte
 	if (ProjectedShadowInfo->VirtualShadowMapClipmap)
 	{
 		Info.NumPrimaryViews = ProjectedShadowInfo->VirtualShadowMapClipmap->GetLevelCount();
-		Info.MaxNumViews = Info.NumPrimaryViews;
+		Info.MaxCullingViews = Info.NumPrimaryViews;
 
 		const bool bIsCached = ProjectedShadowInfo->VirtualShadowMapClipmap->GetCacheEntry() && !ProjectedShadowInfo->VirtualShadowMapClipmap->GetCacheEntry()->IsUncached();
 
@@ -914,10 +902,12 @@ static FVSMRenderViewInfo GetRenderViewInfo(const FProjectedShadowInfo* Projecte
 	else
 	{
 		Info.NumPrimaryViews = ProjectedShadowInfo->bOnePassPointLightShadow ? 6u : 1u;
-		Info.MaxNumViews = Info.NumPrimaryViews * FVirtualShadowMap::MaxMipLevels;
 
 		Info.CullingVolume.Sphere = ProjectedShadowInfo->GetLightSceneInfo().Proxy->GetBoundingSphere();
 		Info.CullingVolume.ConvexVolume = ProjectedShadowInfo->CasterOuterFrustum;
+
+		uint32 MinMipLevel = ProjectedShadowInfo->VirtualShadowMapPerLightCacheEntry->ShadowMapEntries[0].ProjectionData.MinMipLevel;
+		Info.MaxCullingViews = Info.NumPrimaryViews * (FVirtualShadowMap::MaxMipLevels - MinMipLevel);
 	}
 
 	return Info;
@@ -928,7 +918,7 @@ void FShadowSceneRenderer::CreateNaniteViewsForPass(
 	const FVirtualShadowMapArray& VirtualShadowMapArray,
 	TConstArrayView<FViewInfo> Views,
 	float ShadowsLODScaleFactor,
-	FNaniteRenderPass& InOutRenderPass)
+	FNaniteVirtualShadowMapRenderPass& InOutRenderPass)
 {
 	InOutRenderPass.VirtualShadowMapViews = Nanite::FPackedViewArray::CreateWithSetupTask(
 		GraphBuilder,
@@ -968,7 +958,18 @@ void FShadowSceneRenderer::CreateNaniteRenderPasses(
 
 	{
 		const bool bAllowMultipass = CVarVirtualShadowMapNaniteAllowMultipassViews.GetValueOnRenderThread() != 0;		
-		FNaniteRenderPass RenderPass;
+		FNaniteVirtualShadowMapRenderPass RenderPass;
+		auto Flush = [this](FNaniteVirtualShadowMapRenderPass& NaniteRenderPass)
+		{
+			if (NaniteRenderPass.MaxCullingViews > 0)
+			{
+				// Flush any previous render pass
+				check(NaniteRenderPass.TotalPrimaryViews > 0);
+				check(NaniteRenderPass.SceneInstanceCullingQuery);
+				NaniteRenderPasses.Add(NaniteRenderPass);
+				NaniteRenderPass = FNaniteVirtualShadowMapRenderPass();
+			}
+		};
 
 		for (FProjectedShadowInfo* ProjectedShadowInfo : Shadows)
 		{
@@ -977,15 +978,10 @@ void FShadowSceneRenderer::CreateNaniteRenderPasses(
 				FVSMRenderViewInfo Info = GetRenderViewInfo(ProjectedShadowInfo);
 
 				// Space for the new views in the current pass?
-				if (bAllowMultipass && (RenderPass.TotalPrimaryViews + Info.NumPrimaryViews) > MaxPrimaryViews)
+				if (bAllowMultipass && (RenderPass.MaxCullingViews + Info.MaxCullingViews) > MaxViews)
 				{
-					// Flush any previous render pass
-					check(RenderPass.TotalPrimaryViews > 0);
-					check(RenderPass.SceneInstanceCullingQuery);
-					NaniteRenderPasses.Add(RenderPass);
-					RenderPass = FNaniteRenderPass();
+					Flush(RenderPass);
 				}
-
 				RenderPass.Shadows.Add(ProjectedShadowInfo);
 
 				// Add a shadow thingo to be culled, need to know the primary view ranges.
@@ -993,21 +989,17 @@ void FShadowSceneRenderer::CreateNaniteRenderPasses(
 				{
 					RenderPass.SceneInstanceCullingQuery = SceneRenderer.SceneCullingRenderer.CreateInstanceQuery(GraphBuilder);
 				}
-				RenderPass.SceneInstanceCullingQuery->Add(RenderPass.TotalPrimaryViews, Info.NumPrimaryViews, Info.MaxNumViews, Info.CullingVolume);
+				RenderPass.SceneInstanceCullingQuery->Add(RenderPass.TotalPrimaryViews, Info.NumPrimaryViews, Info.MaxCullingViews, Info.CullingVolume);
 
+				RenderPass.MaxCullingViews += Info.MaxCullingViews;
 				RenderPass.TotalPrimaryViews += Info.NumPrimaryViews;
 			}
 		}
-
-		// Add last one if needed
-		if (RenderPass.TotalPrimaryViews > 0)
-		{
-			NaniteRenderPasses.Add(RenderPass);
-		}
+		Flush(RenderPass);
 	}
 
 	const float ShadowsLODScaleFactor = ComputeNaniteShadowsLODScaleFactor();
-	for (FNaniteRenderPass& RenderPass : NaniteRenderPasses)
+	for (FNaniteVirtualShadowMapRenderPass& RenderPass : NaniteRenderPasses)
 	{
 		CreateNaniteViewsForPass(GraphBuilder, VirtualShadowMapArray, Views, ShadowsLODScaleFactor, RenderPass);
 	}
diff --git a/Engine/Source/Runtime/Renderer/Private/Shadows/ShadowSceneRenderer.h b/Engine/Source/Runtime/Renderer/Private/Shadows/ShadowSceneRenderer.h
--- a/Engine/Source/Runtime/Renderer/Private/Shadows/ShadowSceneRenderer.h
+++ b/Engine/Source/Runtime/Renderer/Private/Shadows/ShadowSceneRenderer.h
@@ -149,7 +149,7 @@ private:
 		const FVirtualShadowMapArray& VirtualShadowMapArray,
 		TConstArrayView<FViewInfo> Views,
 		float ShadowsLODScaleFactor,
-		FNaniteRenderPass& InOutRenderPass);
+		FNaniteVirtualShadowMapRenderPass& InOutRenderPass);
 
 	struct FLocalLightShadowFrameSetup
 	{
@@ -175,7 +175,7 @@ private:
 	FShadowScene& ShadowScene;
 	FVirtualShadowMapArray& VirtualShadowMapArray;
 
-	TArray<FNaniteRenderPass, SceneRenderingAllocator> NaniteRenderPasses;
+	TArray<FNaniteVirtualShadowMapRenderPass, SceneRenderingAllocator> NaniteRenderPasses;
 
 	FNaniteVisibilityQuery* NaniteVisibilityQuery = nullptr;
 	TArray<FViewData, SceneRenderingAllocator> ViewDatas;
diff --git a/Engine/Source/Runtime/Renderer/Private/VirtualShadowMaps/VirtualShadowMapArray.cpp b/Engine/Source/Runtime/Renderer/Private/VirtualShadowMaps/VirtualShadowMapArray.cpp
--- a/Engine/Source/Runtime/Renderer/Private/VirtualShadowMaps/VirtualShadowMapArray.cpp
+++ b/Engine/Source/Runtime/Renderer/Private/VirtualShadowMaps/VirtualShadowMapArray.cpp
@@ -3112,8 +3112,7 @@ void FVirtualShadowMapArray::RenderVirtualShadowMapsNanite(
 	FSceneRenderer& SceneRenderer,
 	bool bUpdateNaniteStreaming,
 	const FNaniteVisibilityQuery* VisibilityQuery,
-	TConstArrayView<Nanite::FPackedViewArray*> VirtualShadowMapViewsBatches,
-	TConstArrayView<FSceneInstanceCullingQuery*> SceneInstanceCullingQueries)
+	TConstArrayView<FNaniteVirtualShadowMapRenderPass> VirtualShadowMapPasses)
 {
 	bool bCsvLogEnabled = IsCsvLogEnabled();
 
@@ -3151,21 +3150,18 @@ void FVirtualShadowMapArray::RenderVirtualShadowMapsNanite(
 	// TODO: Stats probably doesn't work correctly with multiple passes
 	static FString VirtualFilterName = TEXT("VirtualShadowMaps");
 
-	// Parallel arrays
-	check(VirtualShadowMapViewsBatches.Num() == SceneInstanceCullingQueries.Num());
-
-	for (int32 Index = 0; Index < VirtualShadowMapViewsBatches.Num(); ++Index)
+	for (int32 Index = 0; Index < VirtualShadowMapPasses.Num(); ++Index)
 	{
-		const Nanite::FPackedViewArray* RenderViews = VirtualShadowMapViewsBatches[Index];
-		FSceneInstanceCullingQuery* SceneInstanceCullingQuery = SceneInstanceCullingQueries[Index];
+		const FNaniteVirtualShadowMapRenderPass& NaniteRenderPass = VirtualShadowMapPasses[Index];
+		const Nanite::FPackedViewArray* RenderViews = NaniteRenderPass.VirtualShadowMapViews;
+		FSceneInstanceCullingQuery* SceneInstanceCullingQuery = NaniteRenderPass.SceneInstanceCullingQuery;
+
 		INC_DWORD_STAT_BY(STAT_VSMNaniteViewsPrimary, RenderViews->NumViews);
 
-		// Compact the views to remove needless (empty) mip views - need to do on GPU as that is where we know what mips have pages
-		// It also generates views for any mips that need them, thus we must do a fairly conservative allocation for the worst case
+		// It generates views for any mips that need them and compact away primary views where no views are used.
 		// TODO: Nanite however can only ever access up to NANITE_MAX_VIEWS_PER_CULL_RASTERIZE_PASS in a given pass, so we could clamp this
 		// to that value, just need to detect overflow in the compaction shader and stop writing.
-		const uint32 MaxOutputViews = FMath::RoundUpToPowerOfTwo(RenderViews->NumViews * FVirtualShadowMap::MaxMipLevels);
-		FRDGBufferRef CompactedViews = GraphBuilder.CreateBuffer(FRDGBufferDesc::CreateStructuredDesc(sizeof(Nanite::FPackedView), MaxOutputViews), TEXT("Shadow.Virtual.CompactedViews"));
+		FRDGBufferRef CompactedViews = GraphBuilder.CreateBuffer(FRDGBufferDesc::CreateStructuredDesc(sizeof(Nanite::FPackedView), NaniteRenderPass.MaxCullingViews), TEXT("Shadow.Virtual.CompactedViews"));
 
 		const int32 NumViewDrawRanges = SceneInstanceCullingQuery->GetViewDrawGroups().Num();
 		FRDGBufferRef ViewDrawRanges = CreateStructuredBuffer(GraphBuilder, TEXT("Shadow.ViewDrawRanges"), SceneInstanceCullingQuery->GetViewDrawGroups());
@@ -3240,7 +3236,7 @@ void FVirtualShadowMapArray::RenderVirtualShadowMapsNanite(
 			CompactedViews,
 			ViewDrawRanges,
 			0,			// GPU provided view counts (multiview)
-			SceneInstanceCullingQueries[Index],
+			SceneInstanceCullingQuery,
 			nullptr		// OptionalInstanceDraws
 		);
 
diff --git a/Engine/Source/Runtime/Renderer/Private/VirtualShadowMaps/VirtualShadowMapArray.h b/Engine/Source/Runtime/Renderer/Private/VirtualShadowMaps/VirtualShadowMapArray.h
--- a/Engine/Source/Runtime/Renderer/Private/VirtualShadowMaps/VirtualShadowMapArray.h
+++ b/Engine/Source/Runtime/Renderer/Private/VirtualShadowMaps/VirtualShadowMapArray.h
@@ -40,6 +40,18 @@ namespace Nanite
 	struct FRasterResults;
 }
 
+
+// Generally only one pass, but we collect this to handle exceptional cases
+struct FNaniteVirtualShadowMapRenderPass
+{
+	FSceneInstanceCullingQuery *SceneInstanceCullingQuery = nullptr;
+	TArray<FProjectedShadowInfo*, SceneRenderingAllocator> Shadows;
+	uint32 TotalPrimaryViews = 0u;
+	uint32 MaxNumMips = 0u;
+	uint32 MaxCullingViews = 0u;
+	Nanite::FPackedViewArray* VirtualShadowMapViews = nullptr;
+};
+
 inline bool IsSinglePageVirtualShadowMap(int32 VirtualShadowMapId)
 {
 	return VirtualShadowMapId < int32(VSM_MAX_SINGLE_PAGE_SHADOW_MAPS);
@@ -361,8 +373,7 @@ public:
 		FSceneRenderer& SceneRenderer,
 		bool bUpdateNaniteStreaming,
 		const FNaniteVisibilityQuery* VisibilityQuery,
-		TConstArrayView<Nanite::FPackedViewArray*> VirtualShadowMapViewsBatches,
-		TConstArrayView<FSceneInstanceCullingQuery*> SceneInstanceCullingQueries);
+		TConstArrayView<FNaniteVirtualShadowMapRenderPass> VirtualShadowMapPasses);
 
 	/**
 	 * Draw Non-Nanite geometry into the VSMs.
