Substrate - fixed compilation issue for opqaue rough refraction shaders. (separating VS/PS parameters)

#tests PC dx12
#rnx

diff --git a/Engine/Source/Runtime/Renderer/Private/Substrate/SubstrateRoughRefraction.cpp b/Engine/Source/Runtime/Renderer/Private/Substrate/SubstrateRoughRefraction.cpp
--- a/Engine/Source/Runtime/Renderer/Private/Substrate/SubstrateRoughRefraction.cpp
+++ b/Engine/Source/Runtime/Renderer/Private/Substrate/SubstrateRoughRefraction.cpp
@@ -35,7 +35,6 @@ class FOpaqueRoughRefractionPS : public FGlobalShader
 		SHADER_PARAMETER_STRUCT_REF(FViewUniformShaderParameters, ViewUniformBuffer)
 		SHADER_PARAMETER_STRUCT_INCLUDE(FSceneTextureShaderParameters, SceneTextures)
 		SHADER_PARAMETER_RDG_UNIFORM_BUFFER(FSubstrateGlobalUniformParameters, Substrate)
-		SHADER_PARAMETER_STRUCT_INCLUDE(Substrate::FSubstrateTilePassVS::FParameters, SubstrateTile)
 		SHADER_PARAMETER_RDG_TEXTURE(Texture2D, SeparatedOpaqueRoughRefractionSceneColor)
 		SHADER_PARAMETER(FVector2f, BlurDirection)
 		SHADER_PARAMETER(float, BlurScale)
@@ -58,6 +57,11 @@ class FOpaqueRoughRefractionPS : public FGlobalShader
 
 IMPLEMENT_GLOBAL_SHADER(FOpaqueRoughRefractionPS, "/Engine/Private/Substrate/SubstrateRoughRefraction.usf", "OpaqueRoughRefractionPS", SF_Pixel);
 
+BEGIN_SHADER_PARAMETER_STRUCT(FOpaqueRoughRefractionParameters, )
+	SHADER_PARAMETER_STRUCT(FOpaqueRoughRefractionPS::FParameters, PS)
+	SHADER_PARAMETER_STRUCT(Substrate::FSubstrateTilePassVS::FParameters, VS)
+END_SHADER_PARAMETER_STRUCT()
+
 
 void AddSubstrateOpaqueRoughRefractionPasses(
 	FRDGBuilder& GraphBuilder,
@@ -92,14 +96,15 @@ void AddSubstrateOpaqueRoughRefractionPasses(
 		FRDGTextureRef SeparatedOpaqueRoughRefractionSceneColor = View.SubstrateViewData.SceneData->SeparatedOpaqueRoughRefractionSceneColor;
 
 		FSceneTextureParameters SceneTextureParameters = GetSceneTextureParameters(GraphBuilder, View);
-		FOpaqueRoughRefractionPS::FParameters* PassParameters = GraphBuilder.AllocParameters<FOpaqueRoughRefractionPS::FParameters>();
-		PassParameters->ViewUniformBuffer = View.ViewUniformBuffer;
-		PassParameters->SceneTextures = GetSceneTextureShaderParameters(SceneTextures.UniformBuffer);
-		PassParameters->Substrate = Substrate::BindSubstrateGlobalUniformParameters(View);
-		PassParameters->SeparatedOpaqueRoughRefractionSceneColor = SeparatedOpaqueRoughRefractionSceneColor;
-		PassParameters->RenderTargets[0] = FRenderTargetBinding(TempTexture, LoadAction);
-		PassParameters->BlurDirection = FVector2f(1.0f, 0.0f);
-		PassParameters->BlurScale = FMath::Max(0.0f, CVarSubstrateOpaqueMaterialRoughRefractionBlurScale.GetValueOnAnyThread());
+		FOpaqueRoughRefractionParameters* PassParameters = GraphBuilder.AllocParameters<FOpaqueRoughRefractionParameters>();
+		FOpaqueRoughRefractionPS::FParameters* PSParameters = &PassParameters->PS;
+		PSParameters->ViewUniformBuffer = View.ViewUniformBuffer;
+		PSParameters->SceneTextures = GetSceneTextureShaderParameters(SceneTextures.UniformBuffer);
+		PSParameters->Substrate = Substrate::BindSubstrateGlobalUniformParameters(View);
+		PSParameters->SeparatedOpaqueRoughRefractionSceneColor = SeparatedOpaqueRoughRefractionSceneColor;
+		PSParameters->RenderTargets[0] = FRenderTargetBinding(TempTexture, LoadAction);
+		PSParameters->BlurDirection = FVector2f(1.0f, 0.0f);
+		PSParameters->BlurScale = FMath::Max(0.0f, CVarSubstrateOpaqueMaterialRoughRefractionBlurScale.GetValueOnAnyThread());
 
 		FOpaqueRoughRefractionPS::FPermutationDomain PermutationVector;
 		PermutationVector.Set<FOpaqueRoughRefractionPS::FEnableBlur>(true);
@@ -111,7 +116,7 @@ void AddSubstrateOpaqueRoughRefractionPasses(
 		TShaderMapRef<Substrate::FSubstrateTilePassVS> TileVertexShader(View.ShaderMap, VSPermutationVector);
 
 		EPrimitiveType PrimitiveType = PT_TriangleList;
-		PassParameters->SubstrateTile = Substrate::SetTileParameters(GraphBuilder, View, SubstrateTileType, PrimitiveType);
+		PassParameters->VS = Substrate::SetTileParameters(GraphBuilder, View, SubstrateTileType, PrimitiveType);
 
 		GraphBuilder.AddPass(
 			RDG_EVENT_NAME("Substrate::OpaqueRoughRefraction(Blur,Horizontal)"),
@@ -133,9 +138,9 @@ void AddSubstrateOpaqueRoughRefractionPasses(
 				GraphicsPSOInit.BoundShaderState.PixelShaderRHI = PixelShader.GetPixelShader();
 				GraphicsPSOInit.DepthStencilState = TStaticDepthStencilState<false, CF_Always>::GetRHI();
 				SetGraphicsPipelineState(RHICmdList, GraphicsPSOInit, 0x0);
-				SetShaderParameters(RHICmdList, PixelShader, PixelShader.GetPixelShader(), *PassParameters);
-				SetShaderParameters(RHICmdList, TileVertexShader, TileVertexShader.GetVertexShader(), PassParameters->SubstrateTile);
-				RHICmdList.DrawPrimitiveIndirect(PassParameters->SubstrateTile.TileIndirectBuffer->GetIndirectRHICallBuffer(), Substrate::TileTypeDrawIndirectArgOffset(SubstrateTileType));
+				SetShaderParameters(RHICmdList, PixelShader, PixelShader.GetPixelShader(), PassParameters->PS);
+				SetShaderParameters(RHICmdList, TileVertexShader, TileVertexShader.GetVertexShader(), PassParameters->VS);
+				RHICmdList.DrawPrimitiveIndirect(PassParameters->VS.TileIndirectBuffer->GetIndirectRHICallBuffer(), Substrate::TileTypeDrawIndirectArgOffset(SubstrateTileType));
 			});
 
 		LoadAction = ERenderTargetLoadAction::ELoad;
@@ -149,14 +154,15 @@ void AddSubstrateOpaqueRoughRefractionPasses(
 		const FViewInfo& View = Views[ViewIndex];
 
 		FSceneTextureParameters SceneTextureParameters = GetSceneTextureParameters(GraphBuilder, View);
-		FOpaqueRoughRefractionPS::FParameters* PassParameters = GraphBuilder.AllocParameters<FOpaqueRoughRefractionPS::FParameters>();
-		PassParameters->ViewUniformBuffer = View.ViewUniformBuffer;
-		PassParameters->SceneTextures = GetSceneTextureShaderParameters(SceneTextures.UniformBuffer);
-		PassParameters->Substrate = Substrate::BindSubstrateGlobalUniformParameters(View);
-		PassParameters->SeparatedOpaqueRoughRefractionSceneColor = TempTexture;
-		PassParameters->RenderTargets[0] = FRenderTargetBinding(SceneColorTexture, ERenderTargetLoadAction::ELoad);
-		PassParameters->BlurDirection = FVector2f(0.0f, 1.0f);
-		PassParameters->BlurScale = FMath::Max(0.0f, CVarSubstrateOpaqueMaterialRoughRefractionBlurScale.GetValueOnAnyThread());
+		FOpaqueRoughRefractionParameters* PassParameters = GraphBuilder.AllocParameters<FOpaqueRoughRefractionParameters>();
+		FOpaqueRoughRefractionPS::FParameters* PSParameters = &PassParameters->PS;
+		PSParameters->ViewUniformBuffer = View.ViewUniformBuffer;
+		PSParameters->SceneTextures = GetSceneTextureShaderParameters(SceneTextures.UniformBuffer);
+		PSParameters->Substrate = Substrate::BindSubstrateGlobalUniformParameters(View);
+		PSParameters->SeparatedOpaqueRoughRefractionSceneColor = TempTexture;
+		PSParameters->RenderTargets[0] = FRenderTargetBinding(SceneColorTexture, ERenderTargetLoadAction::ELoad);
+		PSParameters->BlurDirection = FVector2f(0.0f, 1.0f);
+		PSParameters->BlurScale = FMath::Max(0.0f, CVarSubstrateOpaqueMaterialRoughRefractionBlurScale.GetValueOnAnyThread());
 
 		FOpaqueRoughRefractionPS::FPermutationDomain PermutationVector;
 		PermutationVector.Set<FOpaqueRoughRefractionPS::FEnableBlur>(true);
@@ -168,7 +174,7 @@ void AddSubstrateOpaqueRoughRefractionPasses(
 		TShaderMapRef<Substrate::FSubstrateTilePassVS> TileVertexShader(View.ShaderMap, VSPermutationVector);
 
 		EPrimitiveType PrimitiveType = PT_TriangleList;
-		PassParameters->SubstrateTile = Substrate::SetTileParameters(GraphBuilder, View, SubstrateTileType, PrimitiveType);
+		PassParameters->VS = Substrate::SetTileParameters(GraphBuilder, View, SubstrateTileType, PrimitiveType);
 
 		GraphBuilder.AddPass(
 			RDG_EVENT_NAME("Substrate::OpaqueRoughRefraction(Blur,Vertical)"),
@@ -190,9 +196,9 @@ void AddSubstrateOpaqueRoughRefractionPasses(
 				GraphicsPSOInit.BoundShaderState.PixelShaderRHI = PixelShader.GetPixelShader();
 				GraphicsPSOInit.DepthStencilState = TStaticDepthStencilState<false, CF_Always>::GetRHI();
 				SetGraphicsPipelineState(RHICmdList, GraphicsPSOInit, 0x0);
-				SetShaderParameters(RHICmdList, PixelShader, PixelShader.GetPixelShader(), *PassParameters);
-				SetShaderParameters(RHICmdList, TileVertexShader, TileVertexShader.GetVertexShader(), PassParameters->SubstrateTile);
-				RHICmdList.DrawPrimitiveIndirect(PassParameters->SubstrateTile.TileIndirectBuffer->GetIndirectRHICallBuffer(), Substrate::TileTypeDrawIndirectArgOffset(SubstrateTileType));
+				SetShaderParameters(RHICmdList, PixelShader, PixelShader.GetPixelShader(), PassParameters->PS);
+				SetShaderParameters(RHICmdList, TileVertexShader, TileVertexShader.GetVertexShader(), PassParameters->VS);
+				RHICmdList.DrawPrimitiveIndirect(PassParameters->VS.TileIndirectBuffer->GetIndirectRHICallBuffer(), Substrate::TileTypeDrawIndirectArgOffset(SubstrateTileType));
 			});
 	}
 
@@ -207,14 +213,15 @@ void AddSubstrateOpaqueRoughRefractionPasses(
 		FRDGTextureRef SeparatedOpaqueRoughRefractionSceneColor = View.SubstrateViewData.SceneData->SeparatedOpaqueRoughRefractionSceneColor;
 
 		FSceneTextureParameters SceneTextureParameters = GetSceneTextureParameters(GraphBuilder, View);
-		FOpaqueRoughRefractionPS::FParameters* PassParameters = GraphBuilder.AllocParameters<FOpaqueRoughRefractionPS::FParameters>();
-		PassParameters->ViewUniformBuffer = View.ViewUniformBuffer;
-		PassParameters->SceneTextures = GetSceneTextureShaderParameters(SceneTextures.UniformBuffer);
-		PassParameters->Substrate = Substrate::BindSubstrateGlobalUniformParameters(View);
-		PassParameters->SeparatedOpaqueRoughRefractionSceneColor = SeparatedOpaqueRoughRefractionSceneColor;
-		PassParameters->RenderTargets[0] = FRenderTargetBinding(SceneColorTexture, ERenderTargetLoadAction::ELoad);
-		PassParameters->BlurDirection = FVector2f(0.0f, 0.0f);
-		PassParameters->BlurScale = FMath::Max(0.0f, CVarSubstrateOpaqueMaterialRoughRefractionBlurScale.GetValueOnAnyThread());
+		FOpaqueRoughRefractionParameters* PassParameters = GraphBuilder.AllocParameters<FOpaqueRoughRefractionParameters>();
+		FOpaqueRoughRefractionPS::FParameters* PSParameters = &PassParameters->PS;
+		PSParameters->ViewUniformBuffer = View.ViewUniformBuffer;
+		PSParameters->SceneTextures = GetSceneTextureShaderParameters(SceneTextures.UniformBuffer);
+		PSParameters->Substrate = Substrate::BindSubstrateGlobalUniformParameters(View);
+		PSParameters->SeparatedOpaqueRoughRefractionSceneColor = SeparatedOpaqueRoughRefractionSceneColor;
+		PSParameters->RenderTargets[0] = FRenderTargetBinding(SceneColorTexture, ERenderTargetLoadAction::ELoad);
+		PSParameters->BlurDirection = FVector2f(0.0f, 0.0f);
+		PSParameters->BlurScale = FMath::Max(0.0f, CVarSubstrateOpaqueMaterialRoughRefractionBlurScale.GetValueOnAnyThread());
 
 		FOpaqueRoughRefractionPS::FPermutationDomain PermutationVector;
 		PermutationVector.Set<FOpaqueRoughRefractionPS::FEnableBlur>(false);
@@ -226,7 +233,7 @@ void AddSubstrateOpaqueRoughRefractionPasses(
 		TShaderMapRef<Substrate::FSubstrateTilePassVS> TileVertexShader(View.ShaderMap, VSPermutationVector);
 
 		EPrimitiveType PrimitiveType = PT_TriangleList;
-		PassParameters->SubstrateTile = Substrate::SetTileParameters(GraphBuilder, View, SubstrateTileType, PrimitiveType);
+		PassParameters->VS = Substrate::SetTileParameters(GraphBuilder, View, SubstrateTileType, PrimitiveType);
 
 		GraphBuilder.AddPass(
 			RDG_EVENT_NAME("Substrate::OpaqueRoughRefraction(SSS)"),
@@ -248,9 +255,9 @@ void AddSubstrateOpaqueRoughRefractionPasses(
 				GraphicsPSOInit.BoundShaderState.PixelShaderRHI = PixelShader.GetPixelShader();
 				GraphicsPSOInit.DepthStencilState = TStaticDepthStencilState<false, CF_Always>::GetRHI();
 				SetGraphicsPipelineState(RHICmdList, GraphicsPSOInit, 0x0);
-				SetShaderParameters(RHICmdList, PixelShader, PixelShader.GetPixelShader(), *PassParameters);
-				SetShaderParameters(RHICmdList, TileVertexShader, TileVertexShader.GetVertexShader(), PassParameters->SubstrateTile);
-				RHICmdList.DrawPrimitiveIndirect(PassParameters->SubstrateTile.TileIndirectBuffer->GetIndirectRHICallBuffer(), Substrate::TileTypeDrawIndirectArgOffset(SubstrateTileType));
+				SetShaderParameters(RHICmdList, PixelShader, PixelShader.GetPixelShader(), PassParameters->PS);
+				SetShaderParameters(RHICmdList, TileVertexShader, TileVertexShader.GetVertexShader(), PassParameters->VS);
+				RHICmdList.DrawPrimitiveIndirect(PassParameters->VS.TileIndirectBuffer->GetIndirectRHICallBuffer(), Substrate::TileTypeDrawIndirectArgOffset(SubstrateTileType));
 			});
 	}
 }
