Integrated scene renderer creation and destruction into the scene render builder system.
- Added CreateSceneRenderer and CreateLinkedSceneRenderers to ISceneRenderBuilder.
- Linked renderers are a set of renderers created together that are allowed to reference each other. This is typically mGPU or display cluster rendering.
- The scene render builder validations that renderers are added (and not more than once).
- Refactored how renderers are released so that clean up can pipeline between renderers are the scene update no longer blocks on scene render deletion tasks.
- Replaced AllViews / AllFamilyViews arrays with a scene renderer linked list and some utility functions to enumerate views / families.

diff --git a/Engine/Source/Runtime/Engine/Public/SceneRenderBuilderInterface.h b/Engine/Source/Runtime/Engine/Public/SceneRenderBuilderInterface.h
--- a/Engine/Source/Runtime/Engine/Public/SceneRenderBuilderInterface.h
+++ b/Engine/Source/Runtime/Engine/Public/SceneRenderBuilderInterface.h
@@ -10,6 +10,8 @@ class FRHICommandListImmediate;
 class FRDGBuilder;
 class FSceneRenderer;
 class FSceneInterface;
+class FSceneViewFamily;
+class FHitProxyConsumer;
 
 enum class ESceneRenderGroupFlags : uint8
 {
@@ -49,6 +51,14 @@ public:
 	// valid and all scene renderers must reference the same scene instance.
 	static ENGINE_API TUniquePtr<ISceneRenderBuilder> Create(FSceneInterface* SceneInterface);
 
+	/** Creates a scene renderer based on the current feature level. */
+	virtual FSceneRenderer* CreateSceneRenderer(const FSceneViewFamily* ViewFamily) = 0;
+
+	/** Creates a linked set of scene renderers with related view families based on the current feature level. */
+	virtual TArray<FSceneRenderer*, FConcurrentLinearArrayAllocator> CreateLinkedSceneRenderers(
+		TConstArrayView<const FSceneViewFamily*> ViewFamilies,
+		FHitProxyConsumer* HitProxyConsumer) = 0;
+
 	/** Adds a command that executes on the game thread and is pipelined with scene render commands on the render thread.
 	 *  This is useful to process game thread side logic in between scene renders which may insert render commands. Though
 	 *  be mindful of how commands may be deferred until RDG execution in certain configurations.
diff --git a/Engine/Source/Runtime/Engine/Public/SceneView.h b/Engine/Source/Runtime/Engine/Public/SceneView.h
--- a/Engine/Source/Runtime/Engine/Public/SceneView.h
+++ b/Engine/Source/Runtime/Engine/Public/SceneView.h
@@ -2385,6 +2385,7 @@ public:
 
 	/** Initialization constructor. */
 	ENGINE_API FSceneViewFamily( const ConstructionValues& CVS );
+	ENGINE_API FSceneViewFamily(const FSceneViewFamily&);
 	ENGINE_API virtual ~FSceneViewFamily();
 
 	ENGINE_API ERHIFeatureLevel::Type GetFeatureLevel() const;
@@ -2439,6 +2440,11 @@ public:
 		ScreenPercentageInterface = InScreenPercentageInterface;
 	}
 
+	FORCEINLINE void SetScreenPercentageInterface_Unchecked(ISceneViewFamilyScreenPercentage* InScreenPercentageInterface)
+	{
+		ScreenPercentageInterface = InScreenPercentageInterface;
+	}
+
 	// View family assignment operator is not allowed because of ScreenPercentageInterface lifetime.
 	void operator = (const FSceneViewFamily&) = delete;
 
@@ -2542,9 +2548,6 @@ private:
 	/** True if this view is the current editing view or the active game view */
 	bool bIsInFocus = true;
 
-	// Only FSceneRenderer can copy a view family.
-	ENGINE_API FSceneViewFamily(const FSceneViewFamily&);
-
 	friend class FSceneRenderer;
 	friend class FViewFamilyInfo;
 
diff --git a/Engine/Source/Runtime/Renderer/Private/DeferredShadingRenderer.cpp b/Engine/Source/Runtime/Renderer/Private/DeferredShadingRenderer.cpp
--- a/Engine/Source/Runtime/Renderer/Private/DeferredShadingRenderer.cpp
+++ b/Engine/Source/Runtime/Renderer/Private/DeferredShadingRenderer.cpp
@@ -782,23 +782,20 @@ bool FDeferredShadingSceneRenderer::SetupRayTracingPipelineStatesAndSBT(FRDGBuil
 	}
 
 	// Initialize common resources used for lighting in ray tracing effects
-	for (int32 ViewIndex = 0; ViewIndex < AllFamilyViews.Num(); ++ViewIndex)
+	EnumerateLinkedViews([&ReferenceView] (FViewInfo& View)
 	{
-		// TODO:  It would make more sense for common ray tracing resources to be in a shared structure, rather than copied into each FViewInfo.
-		//        A goal is to have the FViewInfo structure only be visible to the scene renderer that owns it, to avoid dependencies being created
-		//        that could lead to maintenance issues or interfere with parallelism goals.  For now, this works though...
-		FViewInfo* View = const_cast<FViewInfo*>(static_cast<const FViewInfo*>(AllFamilyViews[ViewIndex]));
-
 		// Send common ray tracing resources from reference view to all others.
-		if (View->bHasAnyRayTracingPass && View != &ReferenceView)
+		if (View.bHasAnyRayTracingPass)
 		{
-			View->RayTracingMaterialPipeline = ReferenceView.RayTracingMaterialPipeline;
-			View->RayTracingSBT = ReferenceView.RayTracingSBT;
+			View.RayTracingMaterialPipeline = ReferenceView.RayTracingMaterialPipeline;
+			View.RayTracingSBT = ReferenceView.RayTracingSBT;
 
-			View->LumenHardwareRayTracingMaterialPipeline = ReferenceView.LumenHardwareRayTracingMaterialPipeline;
-			View->LumenHardwareRayTracingSBT = ReferenceView.LumenHardwareRayTracingSBT;
+			View.LumenHardwareRayTracingMaterialPipeline = ReferenceView.LumenHardwareRayTracingMaterialPipeline;
+			View.LumenHardwareRayTracingSBT = ReferenceView.LumenHardwareRayTracingSBT;
 		}
-	}
+		return true;
+
+	}, &ReferenceView);
 
 	return true;
 }
@@ -925,16 +922,17 @@ void FDeferredShadingSceneRenderer::SetupRayTracingRenderingData(FRDGBuilder& Gr
 	SetupLumenHardwareRayTracingUniformBuffer(ReferenceView);
 
 	// Send ray tracing resources from reference view to all others.
-	for (int32 ViewIndex = 0; ViewIndex < AllFamilyViews.Num(); ++ViewIndex)
+	EnumerateLinkedViews([&ReferenceView] (FViewInfo& View)
 	{
 		// See comment above where we copy "RayTracingSubSurfaceProfileTexture" to each view...
-		FViewInfo* View = const_cast<FViewInfo*>(static_cast<const FViewInfo*>(AllFamilyViews[ViewIndex]));
-		if (View->bHasAnyRayTracingPass && View != &ReferenceView)
+		if (View.bHasAnyRayTracingPass)
 		{
-			View->LumenHardwareRayTracingMaterialPipeline = ReferenceView.LumenHardwareRayTracingMaterialPipeline;
-			View->LumenHardwareRayTracingUniformBuffer = ReferenceView.LumenHardwareRayTracingUniformBuffer;
+			View.LumenHardwareRayTracingMaterialPipeline = ReferenceView.LumenHardwareRayTracingMaterialPipeline;
+			View.LumenHardwareRayTracingUniformBuffer = ReferenceView.LumenHardwareRayTracingUniformBuffer;
 		}
-	}
+		return true;
+
+	}, &ReferenceView);
 	
 	bool bAnyLumenHardwareInlineRayTracingPassEnabled = false;
 	for (const FViewInfo& View : Views)
@@ -1796,7 +1794,7 @@ void FDeferredShadingSceneRenderer::Render(FRDGBuilder& GraphBuilder)
 	if (FXSystem && Views.IsValidIndex(0))
 	{
 		SCOPE_CYCLE_COUNTER(STAT_FDeferredShadingSceneRenderer_FXSystem_PreRender);
-		const bool bAllowGPUParticleUpdate = true;
+		const bool bAllowGPUParticleUpdate = IsHeadLink();
 		FXSystem->PreRender(GraphBuilder, GetSceneViews(), GetSceneUniforms(), bAllowGPUParticleUpdate);
 		if (FGPUSortManager* GPUSortManager = FXSystem->GetGPUSortManager())
 		{
@@ -1967,7 +1965,7 @@ void FDeferredShadingSceneRenderer::Render(FRDGBuilder& GraphBuilder)
 	FHairStrandsBookmarkParameters& HairStrandsBookmarkParameters = *GraphBuilder.AllocObject<FHairStrandsBookmarkParameters>();
 	if (IsHairStrandsEnabled(EHairStrandsShaderType::All, Scene->GetShaderPlatform()) && RendererOutput == ERendererOutput::FinalSceneColor)
 	{
-		CreateHairStrandsBookmarkParameters(Scene, Views, AllFamilyViews, HairStrandsBookmarkParameters);
+		CreateHairStrandsBookmarkParameters(Scene, Views, AllViews, HairStrandsBookmarkParameters);
 		check(Scene->HairStrandsSceneData.TransientResources);
 		HairStrandsBookmarkParameters.TransientResources = Scene->HairStrandsSceneData.TransientResources;
 		RunHairStrandsBookmark(GraphBuilder, EHairStrandsBookmark::ProcessTasks, HairStrandsBookmarkParameters);
diff --git a/Engine/Source/Runtime/Renderer/Private/HairStrands/HairStrandsInterface.cpp b/Engine/Source/Runtime/Renderer/Private/HairStrands/HairStrandsInterface.cpp
--- a/Engine/Source/Runtime/Renderer/Private/HairStrands/HairStrandsInterface.cpp
+++ b/Engine/Source/Runtime/Renderer/Private/HairStrands/HairStrandsInterface.cpp
@@ -468,10 +468,14 @@ void UpdateHairStrandsBookmarkParameters(FScene* Scene, TArray<FViewInfo>& Views
 	}
 }
 
-void CreateHairStrandsBookmarkParameters(FScene* Scene, TArray<FViewInfo>& Views, TArray<const FSceneView*>& AllFamilyViews, FHairStrandsBookmarkParameters& Out, bool bComputeVisibleInstances)
+void CreateHairStrandsBookmarkParameters(FScene* Scene, TArray<FViewInfo>& Views, TConstArrayView<FViewInfo*> AllViews, FHairStrandsBookmarkParameters& Out, bool bComputeVisibleInstances)
 {
 	CreateHairStrandsBookmarkParameters(Scene, Views[0], Out, bComputeVisibleInstances);
-	Out.AllViews = AllFamilyViews;
+	Out.AllViews.Reserve(AllViews.Num());
+	for (FViewInfo* View : AllViews)
+	{
+		Out.AllViews.Emplace(View);
+	}
 }
 
 namespace HairStrands
diff --git a/Engine/Source/Runtime/Renderer/Private/HairStrands/HairStrandsRendering.h b/Engine/Source/Runtime/Renderer/Private/HairStrands/HairStrandsRendering.h
--- a/Engine/Source/Runtime/Renderer/Private/HairStrands/HairStrandsRendering.h
+++ b/Engine/Source/Runtime/Renderer/Private/HairStrands/HairStrandsRendering.h
@@ -47,5 +47,5 @@ void RunHairStrandsBookmark(
 	FHairStrandsBookmarkParameters& Parameters);
 
 void CreateHairStrandsBookmarkParameters(FScene* Scene, FViewInfo& View, FHairStrandsBookmarkParameters& Out, bool bComputeVisibleInstances=true);
-void CreateHairStrandsBookmarkParameters(FScene* Scene, TArray<FViewInfo>& Views, TArray<const FSceneView*>& AllFamilyViews, FHairStrandsBookmarkParameters& Out, bool bComputeVisibleInstances=true);
+void CreateHairStrandsBookmarkParameters(FScene* Scene, TArray<FViewInfo>& Views, TConstArrayView<FViewInfo*> AllViews, FHairStrandsBookmarkParameters& Out, bool bComputeVisibleInstances=true);
 void UpdateHairStrandsBookmarkParameters(FScene* Scene, TArray<FViewInfo>& Views, FHairStrandsBookmarkParameters& Out);
diff --git a/Engine/Source/Runtime/Renderer/Private/MeshEdgesRendering.cpp b/Engine/Source/Runtime/Renderer/Private/MeshEdgesRendering.cpp
--- a/Engine/Source/Runtime/Renderer/Private/MeshEdgesRendering.cpp
+++ b/Engine/Source/Runtime/Renderer/Private/MeshEdgesRendering.cpp
@@ -331,7 +331,9 @@ void RenderMeshEdges(FSceneViewFamily& ViewFamily)
 
 	Settings.OnBeforeWireframeRender(CaptureViewFamily);
 
-	FSceneRenderer* SceneRenderer = FSceneRenderer::CreateSceneRenderer(&CaptureViewFamily, nullptr);
+	FSceneRenderBuilder SceneRenderBuilder(ViewFamily.Scene);
+
+	FSceneRenderer* SceneRenderer = SceneRenderBuilder.CreateSceneRenderer(&CaptureViewFamily);
 
 	for (const FSceneViewExtensionRef& Extension : CaptureViewFamily.ViewExtensions)
 	{
@@ -350,8 +352,6 @@ void RenderMeshEdges(FSceneViewFamily& ViewFamily)
 		}
 	}
 
-	FSceneRenderBuilder SceneRenderBuilder(SceneRenderer->Scene);
-
 	ON_SCOPE_EXIT
 	{
 		SceneRenderBuilder.Execute();
diff --git a/Engine/Source/Runtime/Renderer/Private/MobileShadingRenderer.cpp b/Engine/Source/Runtime/Renderer/Private/MobileShadingRenderer.cpp
--- a/Engine/Source/Runtime/Renderer/Private/MobileShadingRenderer.cpp
+++ b/Engine/Source/Runtime/Renderer/Private/MobileShadingRenderer.cpp
@@ -1186,7 +1186,7 @@ void FMobileSceneRenderer::Render(FRDGBuilder& GraphBuilder)
 		if (IsHairStrandsEnabled(EHairStrandsShaderType::All, Scene->GetShaderPlatform()) && RendererOutput != ERendererOutput::DepthPrepassOnly)
 		{
 			FHairStrandsBookmarkParameters& HairStrandsBookmarkParameters = *GraphBuilder.AllocObject<FHairStrandsBookmarkParameters>();
-			CreateHairStrandsBookmarkParameters(Scene, Views, AllFamilyViews, HairStrandsBookmarkParameters);
+			CreateHairStrandsBookmarkParameters(Scene, Views, AllViews, HairStrandsBookmarkParameters);
 			check(Scene->HairStrandsSceneData.TransientResources);
 			HairStrandsBookmarkParameters.TransientResources = Scene->HairStrandsSceneData.TransientResources;
 
diff --git a/Engine/Source/Runtime/Renderer/Private/PlanarReflectionRendering.cpp b/Engine/Source/Runtime/Renderer/Private/PlanarReflectionRendering.cpp
--- a/Engine/Source/Runtime/Renderer/Private/PlanarReflectionRendering.cpp
+++ b/Engine/Source/Runtime/Renderer/Private/PlanarReflectionRendering.cpp
@@ -609,7 +609,7 @@ void FScene::UpdatePlanarReflectionContents(UPlanarReflectionComponent* CaptureC
 		, .MirrorPlane = MirrorPlane
 	};
 
-	FSceneRenderer* SceneRenderer = FSceneRenderer::CreateSceneRenderer(&ViewFamily, nullptr);
+	FSceneRenderer* SceneRenderer = SceneRenderBuilder.CreateSceneRenderer(&ViewFamily);
 
 	SceneRenderBuilder.AddRenderer(SceneRenderer,
 		[
diff --git a/Engine/Source/Runtime/Renderer/Private/ReflectionEnvironmentCapture.cpp b/Engine/Source/Runtime/Renderer/Private/ReflectionEnvironmentCapture.cpp
--- a/Engine/Source/Runtime/Renderer/Private/ReflectionEnvironmentCapture.cpp
+++ b/Engine/Source/Runtime/Renderer/Private/ReflectionEnvironmentCapture.cpp
@@ -1353,7 +1353,7 @@ void CaptureSceneIntoScratchCubemap(
 			Extension->SetupView(ViewFamily, *View);
 		}
 
-		FSceneRenderer* SceneRenderer = FSceneRenderer::CreateSceneRenderer(&ViewFamily, nullptr);
+		FSceneRenderer* SceneRenderer = SceneRenderBuilder.CreateSceneRenderer(&ViewFamily);
 
 		SceneRenderBuilder.AddRenderer(SceneRenderer,
 			[
diff --git a/Engine/Source/Runtime/Renderer/Private/Renderer.cpp b/Engine/Source/Runtime/Renderer/Private/Renderer.cpp
--- a/Engine/Source/Runtime/Renderer/Private/Renderer.cpp
+++ b/Engine/Source/Runtime/Renderer/Private/Renderer.cpp
@@ -48,6 +48,7 @@
 #include "Rendering/RayTracingGeometryManager.h"
 #include "PathTracing.h"
 #include "LightFunctionAtlas.h"
+#include "SceneRenderBuilder.h"
 
 DEFINE_LOG_CATEGORY(LogRenderer);
 
@@ -91,7 +92,7 @@ void FRendererModule::StartupModule()
 			[](FRHICommandListImmediate& RHICmdList)
 		{
 			FRDGBuilder::WaitForAsyncDeleteTask();
-			FSceneRenderer::CleanUp(RHICmdList);
+			FSceneRenderBuilder::WaitForAsyncDeleteTask();
 		});
 	}));
 
diff --git a/Engine/Source/Runtime/Renderer/Private/RendererScene.cpp b/Engine/Source/Runtime/Renderer/Private/RendererScene.cpp
--- a/Engine/Source/Runtime/Renderer/Private/RendererScene.cpp
+++ b/Engine/Source/Runtime/Renderer/Private/RendererScene.cpp
@@ -5609,9 +5609,8 @@ void FScene::Update(FRDGBuilder& GraphBuilder, const FUpdateParameters& Paramete
 		BlackboardEntry.SceneComputeUpdates->SetScene(nullptr);
 	}
 
-	// Avoid overlapping prior scene render and async RDG execution tasks to simplify things as there are a lot of moving pieces.
-	FSceneRenderer::WaitForCleanUpTasks(GraphBuilder.RHICmdList);
-	FRDGBuilder::WaitForAsyncExecuteTask();
+	// Wait for all outstanding RHI command lists so we can safely modify proxies.
+	GraphBuilder.RHICmdList.ImmediateFlush(EImmediateFlushType::WaitForOutstandingTasksOnly);
 
 	UE::Tasks::FTask UpdateUniformExpressionsTask;
 	FMaterialRenderProxy::UpdateDeferredCachedUniformExpressions(GraphBuilder.RHICmdList, EnumHasAnyFlags(Parameters.AsyncOps, EUpdateAllPrimitiveSceneInfosAsyncOps::CacheMaterialUniformExpressions) ? &UpdateUniformExpressionsTask : nullptr);
diff --git a/Engine/Source/Runtime/Renderer/Private/SceneCaptureRendering.cpp b/Engine/Source/Runtime/Renderer/Private/SceneCaptureRendering.cpp
--- a/Engine/Source/Runtime/Renderer/Private/SceneCaptureRendering.cpp
+++ b/Engine/Source/Runtime/Renderer/Private/SceneCaptureRendering.cpp
@@ -783,6 +783,7 @@ void SetupSceneViewExtensionsForSceneCapture(
 }
 
 static FSceneRenderer* CreateSceneRendererForSceneCapture(
+	ISceneRenderBuilder& SceneRenderBuilder,
 	FScene* Scene,
 	USceneCaptureComponent* SceneCaptureComponent,
 	FRenderTarget* RenderTarget,
@@ -949,7 +950,7 @@ static FSceneRenderer* CreateSceneRendererForSceneCapture(
 	// Call SetupViewFamily & SetupView on scene view extensions before renderer creation
 	SetupSceneViewExtensionsForSceneCapture(ViewFamily, Views);
 
-	return FSceneRenderer::CreateSceneRenderer(&ViewFamily, nullptr);
+	return SceneRenderBuilder.CreateSceneRenderer(&ViewFamily);
 }
 
 FSceneCaptureCustomRenderPassUserData FSceneCaptureCustomRenderPassUserData::GDefaultData;
@@ -1228,6 +1229,7 @@ void FScene::UpdateSceneCaptureContents(USceneCaptureComponent2D* CaptureCompone
 		const bool bCameraCut2D = bCopyMainViewTemporalSettings2D ? CaptureComponent->MainViewFamily->Views[0]->bCameraCut : CaptureComponent->bCameraCutThisFrame;
 
 		FSceneRenderer* SceneRenderer = CreateSceneRendererForSceneCapture(
+			SceneRenderBuilder,
 			this, 
 			CaptureComponent, 
 			TextureRenderTarget->GameThread_GetRenderTargetResource(), 
@@ -1322,7 +1324,6 @@ void FScene::UpdateSceneCaptureContents(USceneCaptureComponent2D* CaptureCompone
 		void* CaptureMemorySize = nullptr;
 #endif
 
-		// This needs to happen now as a render command since the resize operation will queue up a deferred resource update.
 		ENQUEUE_RENDER_COMMAND(ResizeSceneCapture)([TextureRenderTargetResource, CaptureSize, bGenerateMips, GameViewportRT] (FRHICommandListImmediate& RHICmdList)
 		{
 			// Resize the render resource if necessary, either to the main viewport size overridden above (see ShouldRenderWithMainViewResolution()),
@@ -1541,10 +1542,24 @@ void FScene::UpdateSceneCaptureContents(USceneCaptureComponentCube* CaptureCompo
 
 				constexpr bool bCameraCut2D = false;
 				constexpr bool bCopyMainViewTemporalSettings2D = false;
-				FSceneRenderer* SceneRenderer = CreateSceneRendererForSceneCapture(this, CaptureComponent,
-					TextureTarget->GameThread_GetRenderTargetResource(), CaptureSize, ViewRotationMatrix,
-					Location, ProjectionMatrix, CaptureComponent->MaxViewDistanceOverride, FOVInDegrees,
-					bCaptureSceneColor, bCameraCut2D, bCopyMainViewTemporalSettings2D, &CaptureComponent->PostProcessSettings, CaptureComponent->PostProcessBlendWeight, CaptureComponent->GetViewOwner(), FaceIndex);
+				FSceneRenderer* SceneRenderer = CreateSceneRendererForSceneCapture(
+					SceneRenderBuilder,
+					this,
+					CaptureComponent,
+					TextureTarget->GameThread_GetRenderTargetResource(),
+					CaptureSize,
+					ViewRotationMatrix,
+					Location,
+					ProjectionMatrix,
+					CaptureComponent->MaxViewDistanceOverride,
+					FOVInDegrees,
+					bCaptureSceneColor,
+					bCameraCut2D,
+					bCopyMainViewTemporalSettings2D,
+					&CaptureComponent->PostProcessSettings,
+					CaptureComponent->PostProcessBlendWeight,
+					CaptureComponent->GetViewOwner(),
+					FaceIndex);
 
 				SceneRenderBuilder.AddRenderer(SceneRenderer, FString::Printf(TEXT("CubeFace[%d]"), FaceIndex),
 					[
@@ -1695,7 +1710,7 @@ void FScene::UpdateSceneCaptureContents(USceneCaptureComponentCube* CaptureCompo
 			// Call SetupViewFamily & SetupView on scene view extensions before renderer creation
 			SetupSceneViewExtensionsForSceneCapture(ViewFamily, Views);
 
-			FSceneRenderer* SceneRenderer = FSceneRenderer::CreateSceneRenderer(&ViewFamily, nullptr);
+			FSceneRenderer* SceneRenderer = SceneRenderBuilder.CreateSceneRenderer(&ViewFamily);
 
 			// Need view state interface to be allocated for Lumen, as it requires persistent data.  This means
 			// "bCaptureEveryFrame" or "bAlwaysPersistRenderingState" must be enabled.
diff --git a/Engine/Source/Runtime/Renderer/Private/SceneRenderBuilder.cpp b/Engine/Source/Runtime/Renderer/Private/SceneRenderBuilder.cpp
--- a/Engine/Source/Runtime/Renderer/Private/SceneRenderBuilder.cpp
+++ b/Engine/Source/Runtime/Renderer/Private/SceneRenderBuilder.cpp
@@ -11,6 +11,315 @@
 #include "SceneViewExtension.h"
 #include "Containers/IntrusiveDoubleLinkedList.h"
 #include "TextureResource.h"
+#include "DeferredShadingRenderer.h"
+#include "FXSystem.h"
+
+//////////////////////////////////////////////////////////////////////////
+
+#if !UE_BUILD_SHIPPING
+
+static TAutoConsoleVariable<int32> CVarSplitScreenDebugEnable(
+	TEXT("r.SplitScreenDebug.Enable"),
+	0,
+	TEXT("Debug feature to replace the main view with a pair of split screen views for testing purposes."),
+	ECVF_Default);
+
+static TAutoConsoleVariable<int32> CVarSplitScreenDebugVertical(
+	TEXT("r.SplitScreenDebug.Vertical"),
+	0,
+	TEXT("Split screen debug use vertical split (two panes vertically stacked).  If false, uses horizontal split (two panes side by side)."),
+	ECVF_Default);
+
+static TAutoConsoleVariable<float> CVarSplitScreenDebugFOVZoom(
+	TEXT("r.SplitScreenDebug.FOVZoom"),
+	1.0f,
+	TEXT("Amount to zoom FOV.  Split screen expands the FOV for the new aspect.  This setting can counteract that expansion."),
+	ECVF_Default);
+
+static TAutoConsoleVariable<float> CVarSplitScreenDebugRotate0(
+	TEXT("r.SplitScreenDebug.Rotate0"),
+	0,
+	TEXT("Rotate first split screen view by this amount.  Values [-1..1] are rotations in view space by fraction of horizontal FOV, outside that range are yaw rotation in degrees."),
+	ECVF_Default);
+
+static TAutoConsoleVariable<float> CVarSplitScreenDebugRotate1(
+	TEXT("r.SplitScreenDebug.Rotate1"),
+	0,
+	TEXT("Rotate second split screen view by this amount.  Values [-1..1] are rotations in view space by fraction of horizontal FOV, outside that range are yaw rotation in degrees."),
+	ECVF_Default);
+
+static TAutoConsoleVariable<int32> CVarSplitScreenDebugOrbit(
+	TEXT("r.SplitScreenDebug.Orbit"),
+	1,
+	TEXT("When rotating by yaw, orbit around camera target actor, to keep third person character visible."),
+	ECVF_Default);
+
+static TAutoConsoleVariable<float> CVarSplitScreenDebugLetterbox(
+	TEXT("r.SplitScreenDebug.Letterbox"),
+	0,
+	TEXT("When non-zero, letterboxes away this percent of screen (rounds up to nearest multiple of 8 pixels, max 50%)."),
+	ECVF_Default);
+
+static TAutoConsoleVariable<int32> CVarSplitScreenDebugLumenScene(
+	TEXT("r.SplitScreenDebug.LumenScene"),
+	1,
+	TEXT("For split screen debugging, allocate a separate Lumen scene for the second view."),
+	ECVF_Default);
+
+static TAutoConsoleVariable<int32> CVarSplitScreenDebugMultiViewFamily(
+	TEXT("r.SplitScreenDebug.MultiViewFamily"),
+	0,
+	TEXT("Uses two renderers with one view each rather than two views."),
+	ECVF_Default);
+
+static bool IsSplitScreenDebugEnabled(TConstArrayView<const FSceneViewFamily*> ViewFamilies)
+{
+	return CVarSplitScreenDebugEnable.GetValueOnGameThread() > 0 && ViewFamilies.Num() == 1 && ViewFamilies[0]->bSplitScreenDebugAllowed && ViewFamilies[0]->Views.Num() == 1;
+}
+
+static void CreateSplitScreenDebugViewFamilies(const FSceneViewFamily& InFamily, TConstArrayView<const FSceneViewFamily*>& OutFamilies, TArray<const FSceneViewFamily*, TFixedAllocator<2>>& OutFamiliesStorage)
+{
+	// We either generate a single family with 2 views, or two families with 1 view each for MGPU
+	int32 NumFamilies;
+	int32 NumViewsPerFamily;
+	if (CVarSplitScreenDebugMultiViewFamily.GetValueOnGameThread())
+	{
+		NumFamilies = 2;
+		NumViewsPerFamily = 1;
+	}
+	else
+	{
+		NumFamilies = 1;
+		NumViewsPerFamily = 2;
+	}
+
+	TArray<FSceneViewFamily*, TFixedAllocator<2>> Families;
+	TArray<const FSceneView**, TFixedAllocator<2>> ViewPointers;
+	TArray<FSceneViewFamily*, TFixedAllocator<2>> ViewParents;
+
+	for (int32 FamilyIndex = 0; FamilyIndex < NumFamilies; FamilyIndex++)
+	{
+		FSceneViewFamily* Family = new FSceneViewFamily(InFamily);
+
+		Families.Add(Family);
+		Family->SetScreenPercentageInterface_Unchecked(InFamily.GetScreenPercentageInterface()->Fork_GameThread(InFamily));
+		Family->Views.SetNumZeroed(NumViewsPerFamily);
+
+		for (int32 ViewIndex = 0; ViewIndex < NumViewsPerFamily; ViewIndex++)
+		{
+			ViewPointers.Add(&Family->Views[ViewIndex]);
+			ViewParents.Add(Family);
+		}
+	}
+
+	FIntRect OriginalViewRect = InFamily.Views[0]->SceneViewInitOptions.ViewRect;
+
+	int32 SplitVertical = CVarSplitScreenDebugVertical.GetValueOnGameThread();
+	float FOVZoom = CVarSplitScreenDebugFOVZoom.GetValueOnGameThread();
+	float FOVScaleX = FOVZoom;
+	float FOVScaleY = FOVZoom;
+
+	float Letterbox = FMath::Clamp(CVarSplitScreenDebugLetterbox.GetValueOnGameThread(), 0.0f, 50.0f);
+	int32 LetterboxPixels;
+
+	if (SplitVertical)
+	{
+		// Double FOV X
+		FOVScaleX *= 0.5f;
+
+		// Convert letterbox from percentage to a multiple of 8 pixels, then reduce FOV by the relative pixel size
+		LetterboxPixels = FMath::CeilToInt((Letterbox / 100.0f) * OriginalViewRect.Size().X * 0.125f) * 8;
+		FOVScaleX = FOVScaleX * OriginalViewRect.Size().X / (OriginalViewRect.Size().X - LetterboxPixels);
+	}
+	else
+	{
+		// Double FOV Y
+		FOVScaleY *= 0.5f;
+
+		// Convert letterbox from percentage to a multiple of 8 pixels, then reduce FOV by the relative pixel size
+		LetterboxPixels = FMath::CeilToInt((Letterbox / 100.0f) * OriginalViewRect.Size().Y * 0.125f) * 8;
+		FOVScaleY = FOVScaleY * OriginalViewRect.Size().Y / (OriginalViewRect.Size().Y - LetterboxPixels);
+	}
+
+	for (int32 ViewIndex = 0; ViewIndex < 2; ViewIndex++)
+	{
+		FSceneViewInitOptions InitOptions = InFamily.Views[0]->SceneViewInitOptions;
+
+		// Adjust projection
+		InitOptions.ProjectionMatrix *= FMatrix(FVector(FOVScaleX, 0.0, 0.0), FVector(0.0, FOVScaleY, 0.0), FVector(0.0, 0.0, 1.0), FVector(0.0, 0.0, 0.0));
+
+		// Adjust view matrix rotation
+		double Rotate = ViewIndex == 0 ? CVarSplitScreenDebugRotate0.GetValueOnGameThread() : CVarSplitScreenDebugRotate1.GetValueOnGameThread();
+		if (Rotate)
+		{
+			if (FMath::Abs(Rotate) <= 1.0)
+			{
+				// Rotation in view space (post multiply) as a fraction of horizontal FOV.  This mode is useful for creating views
+				// that line up exactly along an edge with each other, without needing to do complex FOV calculations.  For example,
+				// setting the left pane to -0.5 and right pane to 0.5 rotates the views away from each other by half the FOV,
+				// producing a matching frustum edge at the middle of the screen (setting the right pane to 1.0 is another example).
+				double FOV = FMath::RadiansToDegrees(FMath::Atan(1.0 / InitOptions.ProjectionMatrix.M[0][0]) * 2.0);
+				Rotate *= FOV;
+
+				InitOptions.ViewRotationMatrix = InitOptions.ViewRotationMatrix * UE::Math::TRotationMatrix<double>::Make(FRotator(Rotate, 0.0, 0.0));
+			}
+			else
+			{
+				// Rotate by degrees in Yaw
+				FMatrix YawRotation = UE::Math::TRotationMatrix<double>::Make(FRotator(0.0, Rotate, 0.0));
+				InitOptions.ViewRotationMatrix = YawRotation * InitOptions.ViewRotationMatrix;
+
+#if !WITH_STATE_STREAM // TODO: Fix. ViewActor does not exist on render side
+				// And optionally orbit the position around the player
+				if (CVarSplitScreenDebugOrbit.GetValueOnGameThread() && InitOptions.ViewActor)
+				{
+					FVector TargetTranslation = InitOptions.ViewActor->GetTransform().GetTranslation();
+					InitOptions.ViewOrigin = YawRotation.GetTransposed().TransformVector(InitOptions.ViewOrigin - TargetTranslation) + TargetTranslation;
+					InitOptions.ViewLocation = InitOptions.ViewOrigin;
+				}
+#endif
+			}
+
+			// Convert adjusted matrix back to a rotation
+			InitOptions.ViewRotation = InitOptions.ViewRotationMatrix.Rotator();
+		}
+		
+		// Make view rectangles half the width / height and adjust opposite dimension for letterbox
+		FIntRect ViewRect = OriginalViewRect;
+		if (SplitVertical)
+		{
+			if (ViewIndex == 0)
+			{
+				ViewRect.Max.Y = (ViewRect.Min.Y + ViewRect.Max.Y) / 2;
+			}
+			else
+			{
+				ViewRect.Min.Y = (ViewRect.Min.Y + ViewRect.Max.Y) / 2;
+			}
+			ViewRect.Min.X += LetterboxPixels / 2;
+			ViewRect.Max.X -= LetterboxPixels / 2;
+		}
+		else
+		{
+			if (ViewIndex == 0)
+			{
+				ViewRect.Max.X = (ViewRect.Min.X + ViewRect.Max.X) / 2;
+			}
+			else
+			{
+				ViewRect.Min.X = (ViewRect.Min.X + ViewRect.Max.X) / 2;
+			}
+			ViewRect.Min.Y += LetterboxPixels / 2;
+			ViewRect.Max.Y -= LetterboxPixels / 2;
+		}
+
+		InitOptions.SetViewRectangle(ViewRect);
+
+		// Set view family to dynamically allocated copy
+		InitOptions.ViewFamily = ViewParents[ViewIndex];
+
+		// Use new static view state for second view
+		if (ViewIndex == 1)
+		{
+			static FSceneViewState* GSecondViewState = nullptr;
+			if (!GSecondViewState)
+			{
+				GSecondViewState = new FSceneViewState(InFamily.GetFeatureLevel(), nullptr);
+			}
+
+			// Propagate this user writable field between FSceneViewState
+			const FSceneViewState* SourceViewState = InFamily.Views[0]->State->GetConcreteViewState();
+			GSecondViewState->SequencerState = SourceViewState->SequencerState;
+
+			// Add or remove optional Lumen scene for second view state
+			if (CVarSplitScreenDebugLumenScene.GetValueOnGameThread())
+			{
+				GSecondViewState->AddLumenSceneData(ViewParents[ViewIndex]->Scene, 1.0f);
+			}
+			else
+			{
+				GSecondViewState->RemoveLumenSceneData(ViewParents[ViewIndex]->Scene);
+			}
+
+			InitOptions.SceneViewStateInterface = GSecondViewState;
+		}
+
+		FSceneView* View = new FSceneView(InitOptions);
+		View->PrimaryViewIndex = NumViewsPerFamily == 1 ? 0 : ViewIndex;
+		View->FinalPostProcessSettings = InFamily.Views[0]->FinalPostProcessSettings;
+
+		*ViewPointers[ViewIndex] = View;
+
+		// Set up second view for multi-GPU
+		if (NumFamilies > 1 && ViewIndex == 1)
+		{
+			FSceneViewFamily* Family = ViewParents[ViewIndex];
+
+			// Prevent the render target from being cleared
+			Family->bAdditionalViewFamily = true;
+
+			if (GNumExplicitGPUsForRendering > 1)
+			{
+				// Enable cross GPU transfers
+				Family->bMultiGPUForkAndJoin = true;
+
+				// Set the view to run on the second GPU
+				View->bOverrideGPUMask = true;
+				View->GPUMask = FRHIGPUMask::FromIndex(ViewIndex);
+			}
+		}
+	}
+
+	// Copy the view families to the const output array
+	for (FSceneViewFamily* Family : Families)
+	{
+		OutFamiliesStorage.Add(Family);
+	}
+	OutFamilies = OutFamiliesStorage;
+}
+
+static void DestroySplitScreenDebugViewFamilies(TConstArrayView<const FSceneViewFamily*> ViewFamilies)
+{
+	for (const FSceneViewFamily* Family : ViewFamilies)
+	{
+		for (int32 ViewIndex = 0; ViewIndex < Family->Views.Num(); ViewIndex++)
+		{
+			delete Family->Views[ViewIndex];
+		}
+		delete Family;
+	}
+}
+
+#endif  // !UE_BUILD_SHIPPING
+
+//////////////////////////////////////////////////////////////////////////
+
+static int32 GSceneRenderCleanUpMode = 1;
+static FAutoConsoleVariableRef CVarSceneRenderCleanUpMode(
+	TEXT("r.SceneRender.CleanUpMode"),
+	GSceneRenderCleanUpMode,
+	TEXT("Controls when to perform clean up of the scene renderer.\n")
+	TEXT(" 0: clean up is performed immediately after render on the render thread.\n")
+	TEXT(" 1: clean up is performed asynchronously in a task. (default)\n"),
+	ECVF_RenderThreadSafe
+);
+
+enum class ESceneRenderCleanUpMode : uint8
+{
+	Immediate,
+	Async
+};
+
+inline ESceneRenderCleanUpMode GetSceneRenderCleanUpMode()
+{
+	if (IsRunningRHIInSeparateThread() && GSceneRenderCleanUpMode == 1)
+	{
+		return ESceneRenderCleanUpMode::Async;
+	}
+	return ESceneRenderCleanUpMode::Immediate;
+}
+
+//////////////////////////////////////////////////////////////////////////
 
 /** This class is responsible for processing a batch of scene renderers. */
 class FSceneRenderProcessor
@@ -161,6 +470,11 @@ public:
 		: Scene(InScene)
 	{}
 
+	TArray<FSceneRenderer*, FConcurrentLinearArrayAllocator> CreateSceneRenderers(
+		TConstArrayView<const FSceneViewFamily*> ViewFamilies,
+		FHitProxyConsumer* HitProxyConsumer,
+		bool bAllowSplitScreenDebug);
+
 	void AddCommand(TUniqueFunction<void()>&& Function)
 	{
 		Ops.Emplace(FOp::FunctionCall(Allocator.Create<TUniqueFunction<void()>>(MoveTemp(Function))));
@@ -173,6 +487,20 @@ public:
 		check(Renderer->ViewFamily.Views.Num() > 0);
 		check(Renderer->ViewFamily.Views[0]);
 
+	#if DO_CHECK
+		bool bFoundRenderer = false;
+		for (int32 Index = 0; Index < RenderersToAdd.Num(); ++Index)
+		{
+			if (RenderersToAdd[Index] == Renderer)
+			{
+				RenderersToAdd.RemoveAtSwap(Index, EAllowShrinking::No);
+				bFoundRenderer = true;
+				break;
+			}
+		}
+		checkf(bFoundRenderer, TEXT("Renderer being added was not created with this scene render builder or is being added twice."));
+	#endif
+
 		FGroupNode* TailGroup = GroupNodes.GetTail();
 
 		FRenderNode* RenderNode = Allocator.Create<FRenderNode>(
@@ -214,7 +542,21 @@ public:
 		return Allocator;
 	}
 
+	static void WaitForAsyncDeleteTask()
+	{
+		check(IsInRenderingThread());
+		if (AsyncDeleteTask)
+		{
+			AsyncDeleteTask->Wait();
+			AsyncDeleteTask = {};
+		}
+	}
+
 private:
+	static FGraphEventRef AsyncDeleteTask;
+	void Cleanup(FRHICommandListImmediate& RHICmdList, FSceneRenderer* SceneRenderer);
+	FGraphEventRef CleanupTaskEvent;
+
 	FScene* Scene;
 	TArray<FSceneRenderer*, FConcurrentLinearArrayAllocator> Renderers;
 	TIntrusiveDoubleLinkedList<FRenderNode> RenderNodes;
@@ -222,6 +564,10 @@ private:
 	TArray<FOp, FConcurrentLinearArrayAllocator> Ops;
 	FConcurrentLinearBulkObjectAllocator Allocator;
 
+#if DO_CHECK
+	TArray<FSceneRenderer*, FConcurrentLinearArrayAllocator> RenderersToAdd;
+#endif
+
 	struct
 	{
 		FGroupNode* Group = nullptr;
@@ -229,13 +575,173 @@ private:
 		FString FullPath;
 
 	} RenderState;
-	
+
 	bool bInsideGroup : 1 = false;
 };
 
+FGraphEventRef FSceneRenderProcessor::AsyncDeleteTask;
+
+static void CleanupSceneRenderer(FSceneRenderer* Renderer)
+{
+	TRACE_CPUPROFILER_EVENT_SCOPE(CleanupSceneRenderer);
+
+	for (auto* Pass : Renderer->DispatchedShadowDepthPasses)
+	{
+		Pass->WaitForTasksAndEmpty();
+	}
+
+	for (FViewInfo* View : Renderer->AllViews)
+	{
+		View->WaitForTasks();
+	}
+}
+
+void FSceneRenderProcessor::Cleanup(FRHICommandListImmediate& RHICmdList, FSceneRenderer* Renderer)
+{
+	TRACE_CPUPROFILER_EVENT_SCOPE(Cleanup);
+
+	// We need to sync async uniform expression cache updates since we're about to start deleting material proxies.
+	FUniformExpressionCacheAsyncUpdateScope::WaitForTask();
+
+	const ESceneRenderCleanUpMode SceneRenderCleanUpMode = GetSceneRenderCleanUpMode();
+
+	if (SceneRenderCleanUpMode == ESceneRenderCleanUpMode::Immediate)
+	{
+		WaitForAsyncDeleteTask(); // This is to handle cases where as switch from async to immediate.
+		RHICmdList.ImmediateFlush(EImmediateFlushType::WaitForOutstandingTasksOnly);
+		CleanupSceneRenderer(Renderer);
+	}
+	else
+	{
+		FGraphEventArray Prereqs;
+		Renderer->GPUSceneDynamicContext.Release();
+
+		// Wait on all setup tasks now to ensure that no additional render commands are enqueued which
+		// might mess with render state, since setup tasks are working with high-level render objects.
+		for (FParallelMeshDrawCommandPass* DispatchedShadowDepthPass : Renderer->DispatchedShadowDepthPasses)
+		{
+			if (DispatchedShadowDepthPass->GetTaskEvent())
+			{
+				Prereqs.Add(DispatchedShadowDepthPass->GetTaskEvent());
+			}
+		}
+
+		for (const FViewInfo& View : Renderer->Views)
+		{
+			for (const FParallelMeshDrawCommandPass& Pass : View.ParallelMeshDrawCommandPasses)
+			{
+				if (Pass.GetTaskEvent())
+				{
+					Prereqs.Add(Pass.GetTaskEvent());
+				}
+			}
+		}
+
+		if (!Prereqs.IsEmpty())
+		{
+			FTaskGraphInterface::Get().WaitUntilTasksComplete(Prereqs, ENamedThreads::GetRenderThread_Local());
+		}
+
+		CleanupTaskEvent->AddPrerequisites(FFunctionGraphTask::CreateAndDispatchWhenReady([Renderer]
+		{
+			CleanupSceneRenderer(Renderer);
+
+		}, TStatId(), &GRHICommandList.WaitOutstandingTasks));
+	}
+
+	GlobalDynamicBuffer::GarbageCollect();
+	GPrimitiveIdVertexBufferPool.DiscardAll();
+	FGraphicsMinimalPipelineStateId::ResetLocalPipelineIdTableSize();
+}
+
+TArray<FSceneRenderer*, FConcurrentLinearArrayAllocator> FSceneRenderProcessor::CreateSceneRenderers(
+	TConstArrayView<const FSceneViewFamily*> ViewFamilies,
+	FHitProxyConsumer* HitProxyConsumer,
+	bool bAllowSplitScreenDebug)
+{
+	if (!ViewFamilies.Num())
+	{
+		return {};
+	}
+
+#if !UE_BUILD_SHIPPING
+	bool bSplitScreenDebug = false;
+	TArray<const FSceneViewFamily*, TFixedAllocator<2>> DebugViewFamilyStorage;
+	if (bAllowSplitScreenDebug && IsSplitScreenDebugEnabled(ViewFamilies))
+	{
+		CreateSplitScreenDebugViewFamilies(*ViewFamilies[0], ViewFamilies, DebugViewFamilyStorage);
+		bSplitScreenDebug = true;
+	}
+#endif
+
+	const FSceneInterface* SceneInterface = ViewFamilies[0]->Scene;
+	check(SceneInterface);
+	check(SceneInterface->GetRenderScene() == Scene);
+
+	const EShadingPath ShadingPath = GetFeatureLevelShadingPath(SceneInterface->GetFeatureLevel());
+
+	TArray<FSceneRenderer*, FConcurrentLinearArrayAllocator> OutRenderers;
+	OutRenderers.Reserve(ViewFamilies.Num());
+
+	for (const FSceneViewFamily* ViewFamily : ViewFamilies)
+	{
+		check(ViewFamily);
+		check(ViewFamily->Scene == SceneInterface);
+
+		if (ShadingPath == EShadingPath::Deferred)
+		{
+			OutRenderers.Add(new FDeferredShadingSceneRenderer(ViewFamily, HitProxyConsumer));
+		}
+		else
+		{
+			check(ShadingPath == EShadingPath::Mobile);
+			OutRenderers.Add(new FMobileSceneRenderer(ViewFamily, HitProxyConsumer));
+		}
+
+		for (auto& ViewExtension : ViewFamily->ViewExtensions)
+		{
+			ViewExtension->PostCreateSceneRenderer(*ViewFamily, OutRenderers.Last());
+		}
+	}
+
+	// Cache the FXSystem for the duration of the scene render
+	// UWorld::CleanupWorldInternal() will mark the system as pending kill on the GameThread and then enqueue a delete command
+	//-TODO: The call to IsPendingKill should no longer be required as we are caching & using within a single render command
+	FFXSystemInterface* FXSystem = nullptr;
+	if (Scene && Scene->FXSystem && !Scene->FXSystem->IsPendingKill())
+	{
+		FXSystem = Scene->FXSystem;
+	}
+
+	for (FSceneRenderer* Renderer : OutRenderers)
+	{
+		Renderer->Link.Head = OutRenderers[0];
+		Renderer->FXSystem  = FXSystem;
+	}
+
+	for (int32 Index = 1; Index < OutRenderers.Num(); ++Index)
+	{
+		OutRenderers[Index - 1]->Link.Next = OutRenderers[Index];
+	}
+
+#if !UE_BUILD_SHIPPING
+	if (bSplitScreenDebug)
+	{
+		DestroySplitScreenDebugViewFamilies(ViewFamilies);
+	}
+#endif
+
+#if DO_CHECK
+	RenderersToAdd.Append(OutRenderers);
+#endif
+
+	return OutRenderers;
+}
+
 void FSceneRenderProcessor::Execute()
 {
 	checkf(!bInsideGroup, TEXT("FSceneRenderBuilder::Execute called within scene render group scope %s. You must end the scope first."), *GroupNodes.GetTail()->Name);
+	checkf(RenderersToAdd.IsEmpty(), TEXT("FSceneRenderBuilder::Execute called but %d renderers were not added."), RenderersToAdd.Num());
 
 #if WITH_GPUDEBUGCRASH
 	if (GRHIGlobals.TriggerGPUCrash != ERequestedGPUCrash::None)
@@ -250,6 +756,8 @@ void FSceneRenderProcessor::Execute()
 	}
 #endif
 
+	CleanupTaskEvent = FGraphEvent::CreateGraphEvent();
+
 	UE::RenderCommandPipe::FSyncScope SyncScope;
 	FUniformExpressionCacheAsyncUpdateScope AsyncUpdateScope;
 	TOptional<UE::RenderCore::DumpGPU::FDumpScope> GpuDumpScope;
@@ -264,14 +772,6 @@ void FSceneRenderProcessor::Execute()
 		}
 	}
 
-	ENQUEUE_RENDER_COMMAND(SceneRenderBuilder_Begin)([this](FRHICommandListImmediate& RHICmdList)
-	{
-		if (!Renderers.IsEmpty())
-		{
-			FSceneRenderer::RenderThreadBegin(RHICmdList, Renderers);
-		}
-	});
-
 	for (FOp Op : Ops)
 	{
 		switch (Op.Type)
@@ -334,6 +834,7 @@ void FSceneRenderProcessor::Execute()
 				RHI_BREADCRUMB_EVENT_CONDITIONAL(RHICmdList,  RenderNode.Name.IsEmpty(), "SceneRender");
 
 				FSceneRenderer& Renderer = *RenderNode.Renderer;
+				FDeferredUpdateResource::UpdateResources(RHICmdList);
 
 				if (!RenderNode.Name.IsEmpty())
 				{
@@ -369,6 +870,8 @@ void FSceneRenderProcessor::Execute()
 				}
 
 				GraphBuilder.Execute();
+
+				Cleanup(RHICmdList, &Renderer);
 			});
 		}
 		break;
@@ -377,12 +880,19 @@ void FSceneRenderProcessor::Execute()
 
 	ENQUEUE_RENDER_COMMAND(SceneRenderBuilder_End)([this](FRHICommandListImmediate& RHICmdList) mutable
 	{
-		if (!Renderers.IsEmpty())
+		CleanupTaskEvent->AddPrerequisites(AsyncDeleteTask);
+		CleanupTaskEvent->DispatchSubsequents();
+
+		AsyncDeleteTask = FFunctionGraphTask::CreateAndDispatchWhenReady([this]
 		{
-			FSceneRenderer::RenderThreadEnd(RHICmdList, Renderers);
-		}
+			TRACE_CPUPROFILER_EVENT_SCOPE(FSceneRenderProcessor::DeleteSceneRenderers);
+			for (FSceneRenderer* Renderer : Renderers)
+			{
+				delete Renderer;
+			}
+			delete this;
 
-		delete this;
+		}, TStatId(), CleanupTaskEvent);
 	});
 
 	// NOTE: 'this' is queued for deletion and is no longer valid!
@@ -390,6 +900,22 @@ void FSceneRenderProcessor::Execute()
 
 //////////////////////////////////////////////////////////////////////////////
 
+FSceneRenderer* FSceneRenderBuilder::CreateSceneRenderer(const FSceneViewFamily* ViewFamily)
+{
+	LazyInit();
+	const bool bAllowSplitScreenDebug = false;
+	return Processor->CreateSceneRenderers(MakeArrayView({ ViewFamily }), nullptr, bAllowSplitScreenDebug)[0];
+}
+
+TArray<FSceneRenderer*, FConcurrentLinearArrayAllocator> FSceneRenderBuilder::CreateLinkedSceneRenderers(
+	TConstArrayView<const FSceneViewFamily*> ViewFamilies,
+	FHitProxyConsumer* HitProxyConsumer)
+{
+	LazyInit();
+	const bool bAllowSplitScreenDebug = true;
+	return Processor->CreateSceneRenderers(ViewFamilies, HitProxyConsumer, bAllowSplitScreenDebug);
+}
+
 FSceneRenderBuilder::FSceneRenderBuilder(FSceneInterface* InScene)
 	: Scene(InScene->GetRenderScene())
 {}
@@ -468,4 +994,9 @@ FConcurrentLinearBulkObjectAllocator& FSceneRenderBuilder::GetAllocator()
 {
 	LazyInit();
 	return Processor->GetAllocator();
+}
+
+void FSceneRenderBuilder::WaitForAsyncDeleteTask()
+{
+	FSceneRenderProcessor::WaitForAsyncDeleteTask();
 }
\ No newline at end of file
diff --git a/Engine/Source/Runtime/Renderer/Private/SceneRenderBuilder.h b/Engine/Source/Runtime/Renderer/Private/SceneRenderBuilder.h
--- a/Engine/Source/Runtime/Renderer/Private/SceneRenderBuilder.h
+++ b/Engine/Source/Runtime/Renderer/Private/SceneRenderBuilder.h
@@ -12,7 +12,7 @@ class FScene;
 class FSceneRenderBuilder final : public ISceneRenderBuilder
 {
 public:
-	FSceneRenderBuilder(FSceneInterface* InScene);
+	RENDERER_API FSceneRenderBuilder(FSceneInterface* InScene);
 	FSceneRenderBuilder(const FSceneRenderBuilder&) = delete;
 	FSceneRenderBuilder& operator=(const FSceneRenderBuilder&) = delete;
 
@@ -30,19 +30,25 @@ public:
 		return *this;
 	}
 
-	~FSceneRenderBuilder() override;
+	RENDERER_API ~FSceneRenderBuilder() override;
 
 	//////////////////////////////////////////////////////////////////////////
 	// ISceneRenderBuilder Overrides
 
-	void AddCommand(TUniqueFunction<void()>&& Function) override;
-	void AddRenderer(FSceneRenderer* Renderer, FString&& Name, FSceneRenderFunction&& Function) override;
+	RENDERER_API FSceneRenderer* CreateSceneRenderer(const FSceneViewFamily* ViewFamily) override;
+	RENDERER_API TArray<FSceneRenderer*, FConcurrentLinearArrayAllocator> CreateLinkedSceneRenderers(
+		TConstArrayView<const FSceneViewFamily*> ViewFamilies,
+		FHitProxyConsumer* HitProxyConsumer) override;
+	RENDERER_API void AddCommand(TUniqueFunction<void()>&& Function) override;
+	RENDERER_API void AddRenderer(FSceneRenderer* Renderer, FString&& Name, FSceneRenderFunction&& Function) override;
 	using ISceneRenderBuilder::AddRenderer;
-	void Execute() override;
-	FConcurrentLinearBulkObjectAllocator& GetAllocator() override;
+	RENDERER_API void Execute() override;
+	RENDERER_API FConcurrentLinearBulkObjectAllocator& GetAllocator() override;
 
 	//////////////////////////////////////////////////////////////////////////
 
+	static void WaitForAsyncDeleteTask();
+
 private:
 	void LazyInit();
 	void BeginGroup(FString&& Name, ESceneRenderGroupFlags Flags) override;
diff --git a/Engine/Source/Runtime/Renderer/Private/SceneRendering.cpp b/Engine/Source/Runtime/Renderer/Private/SceneRendering.cpp
--- a/Engine/Source/Runtime/Renderer/Private/SceneRendering.cpp
+++ b/Engine/Source/Runtime/Renderer/Private/SceneRendering.cpp
@@ -155,69 +155,6 @@ static TAutoConsoleVariable<int32> CVarCachedMeshDrawCommands(
 	TEXT("Whether to render from cached mesh draw commands (on vertex factories that support it), or to generate draw commands every frame."),
 	ECVF_RenderThreadSafe);
 
-#if !UE_BUILD_SHIPPING
-
-static TAutoConsoleVariable<int32> CVarSplitScreenDebugEnable(
-	TEXT("r.SplitScreenDebug.Enable"),
-	0,
-	TEXT("Debug feature to replace the main view with a pair of split screen views for testing purposes."),
-	ECVF_Default);
-
-static TAutoConsoleVariable<int32> CVarSplitScreenDebugVertical(
-	TEXT("r.SplitScreenDebug.Vertical"),
-	0,
-	TEXT("Split screen debug use vertical split (two panes vertically stacked).  If false, uses horizontal split (two panes side by side)."),
-	ECVF_Default);
-
-static TAutoConsoleVariable<float> CVarSplitScreenDebugFOVZoom(
-	TEXT("r.SplitScreenDebug.FOVZoom"),
-	1.0f,
-	TEXT("Amount to zoom FOV.  Split screen expands the FOV for the new aspect.  This setting can counteract that expansion."),
-	ECVF_Default);
-
-static TAutoConsoleVariable<float> CVarSplitScreenDebugRotate0(
-	TEXT("r.SplitScreenDebug.Rotate0"),
-	0,
-	TEXT("Rotate first split screen view by this amount.  Values [-1..1] are rotations in view space by fraction of horizontal FOV, outside that range are yaw rotation in degrees."),
-	ECVF_Default);
-
-static TAutoConsoleVariable<float> CVarSplitScreenDebugRotate1(
-	TEXT("r.SplitScreenDebug.Rotate1"),
-	0,
-	TEXT("Rotate second split screen view by this amount.  Values [-1..1] are rotations in view space by fraction of horizontal FOV, outside that range are yaw rotation in degrees."),
-	ECVF_Default);
-
-static TAutoConsoleVariable<int32> CVarSplitScreenDebugOrbit(
-	TEXT("r.SplitScreenDebug.Orbit"),
-	1,
-	TEXT("When rotating by yaw, orbit around camera target actor, to keep third person character visible."),
-	ECVF_Default);
-
-static TAutoConsoleVariable<float> CVarSplitScreenDebugLetterbox(
-	TEXT("r.SplitScreenDebug.Letterbox"),
-	0,
-	TEXT("When non-zero, letterboxes away this percent of screen (rounds up to nearest multiple of 8 pixels, max 50%)."),
-	ECVF_Default);
-
-static TAutoConsoleVariable<int32> CVarSplitScreenDebugLumenScene(
-	TEXT("r.SplitScreenDebug.LumenScene"),
-	1,
-	TEXT("For split screen debugging, allocate a separate Lumen scene for the second view."),
-	ECVF_Default);
-
-static TAutoConsoleVariable<int32> CVarSplitScreenDebugMultiViewFamily(
-	TEXT("r.SplitScreenDebug.MultiViewFamily"),
-	0,
-	TEXT("Uses two renderers with one view each rather than two views."),
-	ECVF_Default);
-
-static bool IsSplitScreenDebugEnabled(TConstArrayView<const FSceneViewFamily*> ViewFamilies)
-{
-	return CVarSplitScreenDebugEnable.GetValueOnGameThread() > 0 && ViewFamilies.Num() == 1 && ViewFamilies[0]->bSplitScreenDebugAllowed && ViewFamilies[0]->Views.Num() == 1;
-}
-
-#endif  // !UE_BUILD_SHIPPING
-
 bool UseCachedMeshDrawCommands()
 {
 	return CVarCachedMeshDrawCommands.GetValueOnRenderThread() > 0;
@@ -2409,20 +2346,41 @@ FViewInfo* FViewInfo::CreateSnapshot() const
 
 void FViewInfo::DestroyAllSnapshots()
 {
-	QUICK_SCOPE_CYCLE_COUNTER(STAT_FViewInfo_DestroyAllSnapshots);
+	TRACE_CPUPROFILER_EVENT_SCOPE(FViewInfo::DestroyAllSnapshots);
+	static FCriticalSection Mutex;
+
+	// Someone else is destroying snapshots, so just return.
+	if (!Mutex.TryLock())
+	{
+		return;
+	}
+
+	ON_SCOPE_EXIT
+	{
+		Mutex.Unlock();
+	};
 
 	// we will only keep double the number actually used, plus a few
 	const int32 NumToRemove = FMath::Max(ViewInfoSnapshotCache.NumFreeSnapshots - (ViewInfoSnapshotCache.NumSnapshots + 2), 0);
 
-	for (int32 Index = 0; Index < NumToRemove; Index++)
+	int32 NumRemoved = 0;
+	while (FViewInfo* Snapshot = ViewInfoSnapshotCache.FreeSnapshots.Pop())
 	{
-		FViewInfo* ViewInfo = ViewInfoSnapshotCache.FreeSnapshots.Pop();
-		check(ViewInfo);
-		FMemory::Free(ViewInfo);
+		ViewInfoSnapshotCache.NumFreeSnapshots.fetch_sub(1, std::memory_order_relaxed);
+
+		FMemory::Free(Snapshot);
+
+
+		if (++NumRemoved == NumToRemove)
+		{
+			break;
+		}
 	}
 
 	while (FViewInfo* Snapshot = ViewInfoSnapshotCache.Snapshots.Pop())
 	{
+		ViewInfoSnapshotCache.NumSnapshots.fetch_sub(1, std::memory_order_relaxed);
+
 		Snapshot->ViewUniformBuffer.SafeRelease();
 		Snapshot->InstancedViewUniformBuffer.SafeRelease();
 		Snapshot->CachedViewUniformShaderParameters.Reset();
@@ -2440,10 +2398,8 @@ void FViewInfo::DestroyAllSnapshots()
 		Snapshot->SnapshotOriginView = nullptr;
 
 		ViewInfoSnapshotCache.FreeSnapshots.Push(Snapshot);
+		ViewInfoSnapshotCache.NumFreeSnapshots.fetch_add(1, std::memory_order_relaxed);
 	}
-
-	ViewInfoSnapshotCache.NumFreeSnapshots = ViewInfoSnapshotCache.NumFreeSnapshots + ViewInfoSnapshotCache.NumSnapshots - NumToRemove;
-	ViewInfoSnapshotCache.NumSnapshots = 0;
 }
 
 FInt32Range FViewInfo::GetDynamicMeshElementRange(uint32 PrimitiveIndex) const
@@ -3681,20 +3637,21 @@ static void GetCrossGPUTransfers(FSceneRenderer* SceneRenderer, TArray<FCrossGPU
 
 						// If multiple families write to the same render target via MGPU, mask cross GPU copies after the first view family
 						// to the view rect.
-						for (const FSceneViewFamily* ViewFamily : SceneRenderer->AllFamilies)
+						SceneRenderer->EnumerateLinkedViewFamilies([&] (const FSceneViewFamily& ViewFamily)
 						{
-							if (ViewFamily == &SceneRenderer->ViewFamily)
+							if (&ViewFamily == &SceneRenderer->ViewFamily)
 							{
 								// Exit if we found the current view family
-								break;
+								return false;
 							}
-							else if (ViewFamily->RenderTarget == SceneRenderer->ViewFamily.RenderTarget)
+							else if (ViewFamily.RenderTarget == SceneRenderer->ViewFamily.RenderTarget)
 							{
 								// We found another view family writing to the same target, set the TransferRect
 								OutTransfers.Last().TransferRect = ViewInfo.UnscaledViewRect;
-								break;
+								return false;
 							}
-						}
+							return true;
+						});
 					}
 				}
 			}
@@ -4061,7 +4018,7 @@ IVisibilityTaskData* FSceneRenderer::OnRenderBegin(FRDGBuilder& GraphBuilder)
 			if (Views.Num() > 0 && !ViewFamily.EngineShowFlags.HitProxies)
 			{
 				FHairStrandsBookmarkParameters Parameters;
-				CreateHairStrandsBookmarkParameters(Scene, Views, AllFamilyViews, Parameters, false/*bComputeVisibleInstances*/);
+				CreateHairStrandsBookmarkParameters(Scene, Views, AllViews, Parameters, false/*bComputeVisibleInstances*/);
 				if (Parameters.HasInstances())
 				{
 					Scene->WaitForGPUSkinCacheTask();
@@ -4749,291 +4706,6 @@ void FSceneRenderer::SetupMeshPass(FViewInfo& View, FExclusiveDepthStencil::Type
 	}
 }
 
-#if !UE_BUILD_SHIPPING
-void FSceneRenderer::CreateSplitScreenDebugViewFamilies(const FSceneViewFamily& InFamily, TArray<const FSceneViewFamily*>& OutFamilies, TArray<const FSceneViewFamily*, TFixedAllocator<2>>& OutFamiliesStorage)
-{
-	// We either generate a single family with 2 views, or two families with 1 view each for MGPU
-	int32 NumFamilies;
-	int32 NumViewsPerFamily;
-	if (CVarSplitScreenDebugMultiViewFamily.GetValueOnGameThread())
-	{
-		NumFamilies = 2;
-		NumViewsPerFamily = 1;
-	}
-	else
-	{
-		NumFamilies = 1;
-		NumViewsPerFamily = 2;
-	}
-
-	TArray<FSceneViewFamily*, TFixedAllocator<2>> Families;
-	TArray<const FSceneView**, TFixedAllocator<2>> ViewPointers;
-	TArray<FSceneViewFamily*, TFixedAllocator<2>> ViewParents;
-
-	for (int32 FamilyIndex = 0; FamilyIndex < NumFamilies; FamilyIndex++)
-	{
-		FSceneViewFamily* Family = new FSceneViewFamily(InFamily);
-
-		Families.Add(Family);
-		Family->ScreenPercentageInterface = InFamily.ScreenPercentageInterface->Fork_GameThread(InFamily);
-		Family->Views.SetNumZeroed(NumViewsPerFamily);
-
-		for (int32 ViewIndex = 0; ViewIndex < NumViewsPerFamily; ViewIndex++)
-		{
-			ViewPointers.Add(&Family->Views[ViewIndex]);
-			ViewParents.Add(Family);
-		}
-	}
-
-	FIntRect OriginalViewRect = InFamily.Views[0]->SceneViewInitOptions.ViewRect;
-
-	int32 SplitVertical = CVarSplitScreenDebugVertical.GetValueOnGameThread();
-	float FOVZoom = CVarSplitScreenDebugFOVZoom.GetValueOnGameThread();
-	float FOVScaleX = FOVZoom;
-	float FOVScaleY = FOVZoom;
-
-	float Letterbox = FMath::Clamp(CVarSplitScreenDebugLetterbox.GetValueOnGameThread(), 0.0f, 50.0f);
-	int32 LetterboxPixels;
-
-	if (SplitVertical)
-	{
-		// Double FOV X
-		FOVScaleX *= 0.5f;
-
-		// Convert letterbox from percentage to a multiple of 8 pixels, then reduce FOV by the relative pixel size
-		LetterboxPixels = FMath::CeilToInt((Letterbox / 100.0f) * OriginalViewRect.Size().X * 0.125f) * 8;
-		FOVScaleX = FOVScaleX * OriginalViewRect.Size().X / (OriginalViewRect.Size().X - LetterboxPixels);
-	}
-	else
-	{
-		// Double FOV Y
-		FOVScaleY *= 0.5f;
-
-		// Convert letterbox from percentage to a multiple of 8 pixels, then reduce FOV by the relative pixel size
-		LetterboxPixels = FMath::CeilToInt((Letterbox / 100.0f) * OriginalViewRect.Size().Y * 0.125f) * 8;
-		FOVScaleY = FOVScaleY * OriginalViewRect.Size().Y / (OriginalViewRect.Size().Y - LetterboxPixels);
-	}
-
-	for (int32 ViewIndex = 0; ViewIndex < 2; ViewIndex++)
-	{
-		FSceneViewInitOptions InitOptions = InFamily.Views[0]->SceneViewInitOptions;
-
-		// Adjust projection
-		InitOptions.ProjectionMatrix *= FMatrix(FVector(FOVScaleX, 0.0, 0.0), FVector(0.0, FOVScaleY, 0.0), FVector(0.0, 0.0, 1.0), FVector(0.0, 0.0, 0.0));
-
-		// Adjust view matrix rotation
-		double Rotate = ViewIndex == 0 ? CVarSplitScreenDebugRotate0.GetValueOnGameThread() : CVarSplitScreenDebugRotate1.GetValueOnGameThread();
-		if (Rotate)
-		{
-			if (FMath::Abs(Rotate) <= 1.0)
-			{
-				// Rotation in view space (post multiply) as a fraction of horizontal FOV.  This mode is useful for creating views
-				// that line up exactly along an edge with each other, without needing to do complex FOV calculations.  For example,
-				// setting the left pane to -0.5 and right pane to 0.5 rotates the views away from each other by half the FOV,
-				// producing a matching frustum edge at the middle of the screen (setting the right pane to 1.0 is another example).
-				double FOV = FMath::RadiansToDegrees(FMath::Atan(1.0 / InitOptions.ProjectionMatrix.M[0][0]) * 2.0);
-				Rotate *= FOV;
-
-				InitOptions.ViewRotationMatrix = InitOptions.ViewRotationMatrix * UE::Math::TRotationMatrix<double>::Make(FRotator(Rotate, 0.0, 0.0));
-			}
-			else
-			{
-				// Rotate by degrees in Yaw
-				FMatrix YawRotation = UE::Math::TRotationMatrix<double>::Make(FRotator(0.0, Rotate, 0.0));
-				InitOptions.ViewRotationMatrix = YawRotation * InitOptions.ViewRotationMatrix;
-
-#if !WITH_STATE_STREAM // TODO: Fix. ViewActor does not exist on render side
-				// And optionally orbit the position around the player
-				if (CVarSplitScreenDebugOrbit.GetValueOnGameThread() && InitOptions.ViewActor)
-				{
-					FVector TargetTranslation = InitOptions.ViewActor->GetTransform().GetTranslation();
-					InitOptions.ViewOrigin = YawRotation.GetTransposed().TransformVector(InitOptions.ViewOrigin - TargetTranslation) + TargetTranslation;
-					InitOptions.ViewLocation = InitOptions.ViewOrigin;
-				}
-#endif
-			}
-
-			// Convert adjusted matrix back to a rotation
-			InitOptions.ViewRotation = InitOptions.ViewRotationMatrix.Rotator();
-		}
-		
-		// Make view rectangles half the width / height and adjust opposite dimension for letterbox
-		FIntRect ViewRect = OriginalViewRect;
-		if (SplitVertical)
-		{
-			if (ViewIndex == 0)
-			{
-				ViewRect.Max.Y = (ViewRect.Min.Y + ViewRect.Max.Y) / 2;
-			}
-			else
-			{
-				ViewRect.Min.Y = (ViewRect.Min.Y + ViewRect.Max.Y) / 2;
-			}
-			ViewRect.Min.X += LetterboxPixels / 2;
-			ViewRect.Max.X -= LetterboxPixels / 2;
-		}
-		else
-		{
-			if (ViewIndex == 0)
-			{
-				ViewRect.Max.X = (ViewRect.Min.X + ViewRect.Max.X) / 2;
-			}
-			else
-			{
-				ViewRect.Min.X = (ViewRect.Min.X + ViewRect.Max.X) / 2;
-			}
-			ViewRect.Min.Y += LetterboxPixels / 2;
-			ViewRect.Max.Y -= LetterboxPixels / 2;
-		}
-
-		InitOptions.SetViewRectangle(ViewRect);
-
-		// Set view family to dynamically allocated copy
-		InitOptions.ViewFamily = ViewParents[ViewIndex];
-
-		// Use new static view state for second view
-		if (ViewIndex == 1)
-		{
-			static FSceneViewState* GSecondViewState = nullptr;
-			if (!GSecondViewState)
-			{
-				GSecondViewState = new FSceneViewState(InFamily.GetFeatureLevel(), nullptr);
-			}
-
-			// Propagate this user writable field between FSceneViewState
-			const FSceneViewState* SourceViewState = InFamily.Views[0]->State->GetConcreteViewState();
-			GSecondViewState->SequencerState = SourceViewState->SequencerState;
-
-			// Add or remove optional Lumen scene for second view state
-			if (CVarSplitScreenDebugLumenScene.GetValueOnGameThread())
-			{
-				GSecondViewState->AddLumenSceneData(ViewParents[ViewIndex]->Scene, 1.0f);
-			}
-			else
-			{
-				GSecondViewState->RemoveLumenSceneData(ViewParents[ViewIndex]->Scene);
-			}
-
-			InitOptions.SceneViewStateInterface = GSecondViewState;
-		}
-
-		FSceneView* View = new FSceneView(InitOptions);
-		View->PrimaryViewIndex = NumViewsPerFamily == 1 ? 0 : ViewIndex;
-		View->FinalPostProcessSettings = InFamily.Views[0]->FinalPostProcessSettings;
-
-		*ViewPointers[ViewIndex] = View;
-
-		// Set up second view for multi-GPU
-		if (NumFamilies > 1 && ViewIndex == 1)
-		{
-			FSceneViewFamily* Family = ViewParents[ViewIndex];
-
-			// Prevent the render target from being cleared
-			Family->bAdditionalViewFamily = true;
-
-			if (GNumExplicitGPUsForRendering > 1)
-			{
-				// Enable cross GPU transfers
-				Family->bMultiGPUForkAndJoin = true;
-
-				// Set the view to run on the second GPU
-				View->bOverrideGPUMask = true;
-				View->GPUMask = FRHIGPUMask::FromIndex(ViewIndex);
-			}
-		}
-	}
-
-	// Copy the view families to the const output array
-	for (FSceneViewFamily* Family : Families)
-	{
-		OutFamiliesStorage.Add(Family);
-	}
-	OutFamilies = OutFamiliesStorage;
-}
-
-void FSceneRenderer::DestroySplitScreenDebugViewFamilies(TConstArrayView<const FSceneViewFamily*> ViewFamilies)
-{
-	for (const FSceneViewFamily* Family : ViewFamilies)
-	{
-		for (int32 ViewIndex = 0; ViewIndex < Family->Views.Num(); ViewIndex++)
-		{
-			delete Family->Views[ViewIndex];
-		}
-		delete Family;
-	}
-}
-#endif  // !UE_BUILD_SHIPPING
-
-void FSceneRenderer::CreateSceneRenderers(TArray<const FSceneViewFamily*>& InOutViewFamilies, FHitProxyConsumer* HitProxyConsumer, bool bAllowSplitScreenDebug, TArray<FSceneRenderer*>& OutSceneRenderers)
-{
-	OutSceneRenderers.Empty(InOutViewFamilies.Num());
-
-	if (!InOutViewFamilies.Num())
-	{
-		return;
-	}
-
-#if !UE_BUILD_SHIPPING
-	bool bSplitScreenDebug = false;
-	TArray<const FSceneViewFamily*, TFixedAllocator<2>> DebugViewFamilyStorage;
-	if (bAllowSplitScreenDebug && IsSplitScreenDebugEnabled(InOutViewFamilies))
-	{
-		CreateSplitScreenDebugViewFamilies(*InOutViewFamilies[0], InOutViewFamilies, DebugViewFamilyStorage);
-		bSplitScreenDebug = true;
-	}
-#endif
-
-	const FSceneInterface* Scene = InOutViewFamilies[0]->Scene;
-	check(Scene);
-
-	EShadingPath ShadingPath = GetFeatureLevelShadingPath(Scene->GetFeatureLevel());
-
-	for (int32 FamilyIndex = 0; FamilyIndex < InOutViewFamilies.Num(); FamilyIndex++)
-	{
-		const FSceneViewFamily* InViewFamily = InOutViewFamilies[FamilyIndex];
-
-		check(InViewFamily);
-		check(InViewFamily->Scene == Scene);
-
-		if (ShadingPath == EShadingPath::Deferred)
-		{
-			FDeferredShadingSceneRenderer* SceneRenderer = new FDeferredShadingSceneRenderer(InViewFamily, HitProxyConsumer);
-			OutSceneRenderers.Add(SceneRenderer);
-		}
-		else
-		{
-			check(ShadingPath == EShadingPath::Mobile);
-			OutSceneRenderers.Add(new FMobileSceneRenderer(InViewFamily, HitProxyConsumer));
-		}
-
-		for (int ViewExt = 0; ViewExt < InViewFamily->ViewExtensions.Num(); ViewExt++)
-		{
-			InViewFamily->ViewExtensions[ViewExt]->PostCreateSceneRenderer(*InViewFamily, OutSceneRenderers.Last());
-		}
-	}
-
-#if !UE_BUILD_SHIPPING
-	if (bSplitScreenDebug)
-	{
-		DestroySplitScreenDebugViewFamilies(InOutViewFamilies);
-	}
-#endif
-}
-
-FSceneRenderer* FSceneRenderer::CreateSceneRenderer(const FSceneViewFamily* InViewFamily, FHitProxyConsumer* HitProxyConsumer)
-{
-	check(InViewFamily);
-
-	TArray<FSceneRenderer*> SceneRenderers;
-
-	TArray<const FSceneViewFamily*> InViewFamilies;
-	InViewFamilies.Add(InViewFamily);
-
-	CreateSceneRenderers(InViewFamilies, HitProxyConsumer, /*bAllowSplitScreenDebug=*/ false, SceneRenderers);
-
-	return SceneRenderers[0];
-}
-
 bool FSceneRenderer::ShouldCompositeEditorPrimitives(const FViewInfo& View)
 {
 	const FEngineShowFlags& ShowFlags = View.Family->EngineShowFlags;
@@ -5128,219 +4800,6 @@ void FSceneRenderer::UpdatePrimitiveIndirectLightingCacheBuffers(FRHICommandList
 	FRendererModule
 -----------------------------------------------------------------------------*/
 
-static int32 GSceneRenderCleanUpMode = 2;
-static FAutoConsoleVariableRef CVarSceneRenderCleanUpMode(
-	TEXT("r.SceneRenderCleanUpMode"),
-	GSceneRenderCleanUpMode,
-	TEXT("Controls when to perform clean up of the scene renderer.\n")
-	TEXT(" 0: clean up is performed immediately after render on the render thread.\n")
-	TEXT(" 1: clean up deferred until the start of the next scene render on the render thread.\n")
-	TEXT(" 2: clean up deferred until the start of the next scene render on the render thread, with some work distributed to an async task. (default)\n"),
-	ECVF_RenderThreadSafe
-);
-
-enum class ESceneRenderCleanUpMode
-{
-	Immediate,
-	Deferred,
-	DeferredAndAsync
-};
-
-inline ESceneRenderCleanUpMode GetSceneRenderCleanUpMode()
-{
-	if (GSceneRenderCleanUpMode != 1 && GSceneRenderCleanUpMode != 2)
-	{
-		return ESceneRenderCleanUpMode::Immediate;
-	}
-
-	if (!IsRunningRHIInSeparateThread())
-	{
-		return ESceneRenderCleanUpMode::Immediate;
-	}
-
-	return (ESceneRenderCleanUpMode)GSceneRenderCleanUpMode;
-}
-
-static void FinishCleanUp()
-{
-	TRACE_CPUPROFILER_EVENT_SCOPE(FinishCleanUp);
-
-	// Can release only after all mesh pass tasks are finished.
-	GlobalDynamicBuffer::GarbageCollect();
-	GPrimitiveIdVertexBufferPool.DiscardAll();
-	FGraphicsMinimalPipelineStateId::ResetLocalPipelineIdTableSize();
-}
-
-static void DeleteSceneRenderers(TConstArrayView<FSceneRenderer*> SceneRenderers)
-{
-	SCOPED_NAMED_EVENT_TEXT("DeleteSceneRenderer", FColor::Red);
-
-	for (FSceneRenderer* SceneRenderer : SceneRenderers)
-	{
-		// Wait for all dispatched shadow mesh draw tasks.
-		for (int32 PassIndex = 0; PassIndex < SceneRenderer->DispatchedShadowDepthPasses.Num(); ++PassIndex)
-		{
-			SceneRenderer->DispatchedShadowDepthPasses[PassIndex]->WaitForTasksAndEmpty();
-		}
-
-		for (FViewInfo* View : SceneRenderer->AllViews)
-		{
-			View->WaitForTasks();
-		}
-	}
-
-	FViewInfo::DestroyAllSnapshots();
-
-	for (FSceneRenderer* SceneRenderer : SceneRenderers)
-	{
-		delete SceneRenderer;
-	}
-}
-
-static void WaitForTasksAndDeleteSceneRenderers(FRHICommandListImmediate& RHICmdList, TConstArrayView<FSceneRenderer*> SceneRenderers)
-{
-	{
-		QUICK_SCOPE_CYCLE_COUNTER(STAT_DeleteSceneRenderer_WaitForTasks);
-		RHICmdList.ImmediateFlush(EImmediateFlushType::WaitForOutstandingTasksOnly);
-	}
-
-	DeleteSceneRenderers(SceneRenderers);
-}
-
-void FSceneRenderer::RenderThreadBegin(FRHICommandListImmediate& RHICmdList)
-{
-	RenderThreadBegin(RHICmdList, MakeArrayView({ this }));
-}
-
-void FSceneRenderer::RenderThreadEnd(FRHICommandListImmediate& RHICmdList)
-{
-	RenderThreadEnd(RHICmdList, MakeArrayView({ this }));
-}
-
-void FSceneRenderer::RenderThreadBegin(FRHICommandListImmediate& RHICmdList, TConstArrayView<FSceneRenderer*> SceneRenderers)
-{
-	LLM_SCOPE(ELLMTag::SceneRender);
-
-	CleanUp(RHICmdList);
-
-	// Initialize "AllFamilyViews" and "AllFamilies" array in each scene renderer.  Initialize array in first SceneRenderer, then copy to remaining.
-	for (FSceneRenderer* SceneRenderer : SceneRenderers)
-	{
-		for (const FViewInfo& ViewInfo : SceneRenderer->Views)
-		{
-			SceneRenderers[0]->AllFamilyViews.Add(&ViewInfo);
-		}
-		SceneRenderers[0]->AllFamilies.Add(&SceneRenderer->ViewFamily);
-	}
-	for (int32 RendererIndex = 1; RendererIndex < SceneRenderers.Num(); RendererIndex++)
-	{
-		SceneRenderers[RendererIndex]->AllFamilyViews = SceneRenderers[0]->AllFamilyViews;
-		SceneRenderers[RendererIndex]->AllFamilies = SceneRenderers[0]->AllFamilies;
-	}
-
-	FScene* Scene = SceneRenderers[0]->Scene;
-
-	// Cache the FXSystem for the duration of the scene render
-	// UWorld::CleanupWorldInternal() will mark the system as pending kill on the GameThread and then enqueue a delete command
-	//-TODO: The call to IsPendingKill should no longer be required as we are caching & using within a single render command
-	class FFXSystemInterface* FXSystem = Scene ? Scene->FXSystem : nullptr;
-	if (FXSystem && FXSystem->IsPendingKill())
-	{
-		FXSystem = nullptr;
-	}
-
-	for (FSceneRenderer* SceneRenderer : SceneRenderers)
-	{
-		SceneRenderer->FXSystem = FXSystem;
-	}
-
-	FDeferredUpdateResource::UpdateResources(RHICmdList);
-}
-
-struct FSceneRenderCleanUpState
-{
-	TArray<FSceneRenderer*> Renderers;
-	FGraphEventRef Task;
-	ESceneRenderCleanUpMode CompletionMode = ESceneRenderCleanUpMode::Immediate;
-	bool bWaitForTasksComplete = false;
-};
-
-static FSceneRenderCleanUpState GSceneRenderCleanUpState;
-
-void FSceneRenderer::RenderThreadEnd(FRHICommandListImmediate& RHICmdList, TConstArrayView<FSceneRenderer*> SceneRenderers)
-{
-	LLM_SCOPE(ELLMTag::SceneRender);
-
-	// We need to sync async uniform expression cache updates since we're about to start deleting material proxies.
-	FUniformExpressionCacheAsyncUpdateScope::WaitForTask();
-
-	check(GSceneRenderCleanUpState.Renderers.IsEmpty());
-
-	const ESceneRenderCleanUpMode SceneRenderCleanUpMode = GetSceneRenderCleanUpMode();
-
-	GSceneRenderCleanUpState.CompletionMode = SceneRenderCleanUpMode;
-
-	if (GSceneRenderCleanUpState.CompletionMode == ESceneRenderCleanUpMode::Immediate)
-	{
-		WaitForTasksAndDeleteSceneRenderers(RHICmdList, SceneRenderers);
-		FinishCleanUp();
-	}
-	else
-	{
-		for (FSceneRenderer* SceneRenderer : SceneRenderers)
-		{
-			SceneRenderer->GPUSceneDynamicContext.Release();
-		}
-
-		// Mem stack mark is stored on first scene renderer
-		GSceneRenderCleanUpState.Renderers = SceneRenderers;
-
-		if (SceneRenderCleanUpMode == ESceneRenderCleanUpMode::DeferredAndAsync)
-		{
-			// Wait on all setup tasks now to ensure that no additional render commands are enqueued which
-			// might mess with render state, since setup tasks are working with high-level render objects.
-			{
-				FGraphEventArray SetupTasks;
-
-				for (FSceneRenderer* SceneRenderer : SceneRenderers)
-				{
-					for (FParallelMeshDrawCommandPass* DispatchedShadowDepthPass : SceneRenderer->DispatchedShadowDepthPasses)
-					{
-						if (DispatchedShadowDepthPass->GetTaskEvent())
-						{
-							SetupTasks.Add(DispatchedShadowDepthPass->GetTaskEvent());
-						}
-					}
-
-					for (const FViewInfo& View : SceneRenderer->Views)
-					{
-						for (const FParallelMeshDrawCommandPass& Pass : View.ParallelMeshDrawCommandPasses)
-						{
-							if (Pass.GetTaskEvent())
-							{
-								SetupTasks.Add(Pass.GetTaskEvent());
-							}
-						}
-					}
-				}
-
-				if (!SetupTasks.IsEmpty())
-				{
-					FTaskGraphInterface::Get().WaitUntilTasksComplete(SetupTasks, ENamedThreads::GetRenderThread_Local());
-				}
-			}
-
-			check(IsInRenderingThread());
-			FGraphEventArray CommandListTasks = MoveTemp(GRHICommandList.WaitOutstandingTasks);
-
-			GSceneRenderCleanUpState.Task = FFunctionGraphTask::CreateAndDispatchWhenReady([SceneRenderers = TArray<FSceneRenderer*, TInlineAllocator<1>>(SceneRenderers)]
-			{
-				DeleteSceneRenderers(SceneRenderers);
-			}, TStatId(), &CommandListTasks);
-		}
-	}
-}
-
 UE::Renderer::Private::IShadowInvalidatingInstances *FSceneRenderer::GetShadowInvalidatingInstancesInterface(const FSceneView *SceneView)
 {
 	checkf(IsInRenderingThread(), TEXT("Accessing the ShadowInvalidatingInstancesInterface should only be allowed from the rendering thread!"));
@@ -5351,51 +4810,6 @@ UE::Renderer::Private::IShadowInvalidatingInstances *FSceneRenderer::GetShadowIn
 	return nullptr;
 }
 
-void FSceneRenderer::CleanUp(FRHICommandListImmediate& RHICmdList)
-{
-	if (GSceneRenderCleanUpState.CompletionMode == ESceneRenderCleanUpMode::Immediate || GSceneRenderCleanUpState.Renderers.IsEmpty())
-	{
-		return;
-	}
-
-	if (!GSceneRenderCleanUpState.bWaitForTasksComplete)
-	{
-		switch (GSceneRenderCleanUpState.CompletionMode)
-		{
-		case ESceneRenderCleanUpMode::Deferred:
-			WaitForTasksAndDeleteSceneRenderers(RHICmdList, GSceneRenderCleanUpState.Renderers);
-			break;
-		case ESceneRenderCleanUpMode::DeferredAndAsync:
-			GSceneRenderCleanUpState.Task->Wait(ENamedThreads::GetRenderThread_Local());
-			break;
-		}
-	}
-
-	FinishCleanUp();
-	GSceneRenderCleanUpState = {};
-}
-
-void FSceneRenderer::WaitForCleanUpTasks(FRHICommandListImmediate& RHICmdList)
-{
-	if (GSceneRenderCleanUpState.CompletionMode == ESceneRenderCleanUpMode::Immediate || GSceneRenderCleanUpState.Renderers.IsEmpty() || GSceneRenderCleanUpState.bWaitForTasksComplete)
-	{
-		return;
-	}
-
-	switch (GSceneRenderCleanUpState.CompletionMode)
-	{
-	case ESceneRenderCleanUpMode::Deferred:
-		WaitForTasksAndDeleteSceneRenderers(RHICmdList, GSceneRenderCleanUpState.Renderers);
-		break;
-	case ESceneRenderCleanUpMode::DeferredAndAsync:
-		GSceneRenderCleanUpState.Task->Wait(ENamedThreads::GetRenderThread_Local());
-		GSceneRenderCleanUpState.Task = nullptr;
-		break;
-	}
-
-	GSceneRenderCleanUpState.bWaitForTasksComplete = true;
-}
-
 void ResetAndShrinkModifiedBounds(TArray<FBox>& Bounds)
 {
 	const int32 MaxAllocatedSize = FMath::RoundUpToPowerOfTwo(FMath::Max<uint32>(DistanceField::MinPrimitiveModifiedBoundsAllocation, Bounds.Num()));
@@ -5451,10 +4865,10 @@ static void CleanupViewFamilies_RenderThread(FRHICommandListImmediate& RHICmdLis
 	{
 		CSV_SCOPED_TIMING_STAT_EXCLUSIVE(PostRenderCleanUp);
 
-		if (IsHairStrandsEnabled(EHairStrandsShaderType::All, Scene->GetShaderPlatform()) && (SceneRenderers[0]->AllFamilyViews.Num() > 0) )
+		if (IsHairStrandsEnabled(EHairStrandsShaderType::All, Scene->GetShaderPlatform()) && (SceneRenderers[0]->AllViews.Num() > 0))
 		{
 			FHairStrandsBookmarkParameters Parameters;
-			CreateHairStrandsBookmarkParameters(Scene, SceneRenderers[0]->Views, SceneRenderers[0]->AllFamilyViews, Parameters, false /*bComputeVisibleInstances*/);
+			CreateHairStrandsBookmarkParameters(Scene, SceneRenderers[0]->Views, SceneRenderers[0]->AllViews, Parameters, false /*bComputeVisibleInstances*/);
 			if (Parameters.HasInstances())
 			{
 				RunHairStrandsBookmark(EHairStrandsBookmark::ProcessEndOfFrame, Parameters);
@@ -5462,7 +4876,7 @@ static void CleanupViewFamilies_RenderThread(FRHICommandListImmediate& RHICmdLis
 		}
 
 		// Only reset per-frame scene state once all views have processed their frame, including those in planar reflections
-		for (int32 CacheType = 0; CacheType < UE_ARRAY_COUNT(SceneRenderers[0]->Scene->DistanceFieldSceneData.PrimitiveModifiedBounds); CacheType++)
+		for (int32 CacheType = 0; CacheType < UE_ARRAY_COUNT(Scene->DistanceFieldSceneData.PrimitiveModifiedBounds); CacheType++)
 		{
 			ResetAndShrinkModifiedBounds(Scene->DistanceFieldSceneData.PrimitiveModifiedBounds[CacheType]);
 		}
@@ -5678,26 +5092,15 @@ void FRendererModule::BeginRenderingViewFamilies(FCanvas* Canvas, TConstArrayVie
 		});
 	#endif
 
-		// This array can be mutated by CreateSceneRenderers when split screen debug is enabled.
-		TArray<const FSceneViewFamily*> ViewFamiliesFinal;
-		ViewFamiliesFinal.Reserve(ViewFamilies.Num());
-		for (const FSceneViewFamily* Family : ViewFamilies)
-		{
-			ViewFamiliesFinal.Emplace(Family);
-		}
+		FSceneRenderBuilder SceneRenderBuilder(Scene);
 
-		// Construct the scene renderers.  This copies the view family attributes into its own structures.
-		const bool bAllowSplitScreenDebug = true;
-		TArray<FSceneRenderer*> SceneRenderers;
-		FSceneRenderer::CreateSceneRenderers(ViewFamiliesFinal, Canvas->GetHitProxyConsumer(), bAllowSplitScreenDebug, SceneRenderers);
+		TArray<FSceneRenderer*, FConcurrentLinearArrayAllocator> SceneRenderers = SceneRenderBuilder.CreateLinkedSceneRenderers(ViewFamilies, Canvas->GetHitProxyConsumer());
 		SetupDebugViewModes(SceneRenderers);
 
-		FSceneRenderBuilder SceneRenderBuilder(Scene);
-
 		bool bShowHitProxies = (Canvas->GetHitProxyConsumer() != nullptr);
 		if (!bShowHitProxies)
 		{
-			SceneCaptureUpdateDeferredCapturesInternal(Scene, ViewFamiliesFinal, SceneRenderBuilder);
+			SceneCaptureUpdateDeferredCapturesInternal(Scene, ViewFamilies, SceneRenderBuilder);
 		}
 
 		if (!bShowHitProxies)
@@ -5718,9 +5121,8 @@ void FRendererModule::BeginRenderingViewFamilies(FCanvas* Canvas, TConstArrayVie
 
 		FSceneRenderer::PreallocateCrossGPUFences(SceneRenderers);
 
-		for (int32 ViewFamilyIndex = 0; ViewFamilyIndex < ViewFamiliesFinal.Num(); ++ViewFamilyIndex)
+		for (FSceneRenderer* SceneRenderer : SceneRenderers)
 		{
-			FSceneRenderer* SceneRenderer = SceneRenderers[ViewFamilyIndex];
 			SceneRenderer->ViewFamily.DisplayInternalsData.Setup(World);
 
 			SceneRenderBuilder.AddRenderer(SceneRenderer, bShowHitProxies ? TEXT("HitProxies") : TEXT("ViewFamilies"),
@@ -5779,7 +5181,7 @@ void FRendererModule::PerFrameCleanupIfSkipRenderer()
 	{
 		RHICmdList.ImmediateFlush(EImmediateFlushType::FlushRHIThreadFlushResources);
 		PipelineStateCache::FlushResources();
-		FSceneRenderer::CleanUp(RHICmdList);
+		FSceneRenderBuilder::WaitForAsyncDeleteTask();
 		GPrimitiveIdVertexBufferPool.DiscardAll();
 	});
 }
diff --git a/Engine/Source/Runtime/Renderer/Private/SceneRendering.h b/Engine/Source/Runtime/Renderer/Private/SceneRendering.h
--- a/Engine/Source/Runtime/Renderer/Private/SceneRendering.h
+++ b/Engine/Source/Runtime/Renderer/Private/SceneRendering.h
@@ -2041,7 +2041,7 @@ public:
 
 		/** Custom render pass that render as part of the main renderer. */
 		FCustomRenderPassBase* CustomRenderPass;
-		/** View family used by custom render pass.  NOT added to AllFamilies.  Required to allow different EngineShowFlags from main renderer ViewFamily. */
+		/** View family used by custom render pass. NOT treated as a linked view. Required to allow different EngineShowFlags from main renderer ViewFamily. */
 		FViewFamilyInfo ViewFamily;
 		/** Views used to render the custom render pass. */
 		TArray<FViewInfo> Views;
@@ -2052,9 +2052,52 @@ public:
 	/** All views include main camera views and custom render pass views. */
 	TArray<FViewInfo*> AllViews;
 
-	/** Views across all view families (may contain additional views if rendering multiple families together). */
-	TArray<const FSceneView*> AllFamilyViews;
-	TArray<const FSceneViewFamily*> AllFamilies;
+	//////////////////////////////////////////////////////////////////////////
+	// Provides access to properties of linked scene renderers.
+
+	struct
+	{
+		FSceneRenderer* Head = nullptr;
+		FSceneRenderer* Next = nullptr;
+	
+	} Link;
+
+	bool IsHeadLink() const { return Link.Head == this; }
+
+	template <typename LambdaType>
+	void EnumerateLinkedViews(LambdaType&& Lambda, const FSceneView* ViewToSkip = nullptr)
+	{
+		for (FSceneRenderer* Renderer = Link.Head; Renderer; Renderer = Renderer->Link.Next)
+		{
+			for (FViewInfo& View : Renderer->Views)
+			{
+				if (ViewToSkip != &View)
+				{
+					if (!Lambda(View))
+					{
+						return;
+					}
+				}
+			}
+		}
+	}
+
+	template <typename LambdaType>
+	void EnumerateLinkedViewFamilies(LambdaType&& Lambda, const FSceneViewFamily* ViewFamilyToSkip = nullptr)
+	{
+		for (FSceneRenderer* Renderer = Link.Head; Renderer; Renderer = Renderer->Link.Next)
+		{
+			if (ViewFamilyToSkip != &Renderer->ViewFamily)
+			{
+				if (!Lambda(Renderer->ViewFamily))
+				{
+					return;
+				}
+			}
+		}
+	}
+
+	//////////////////////////////////////////////////////////////////////////
 
 	/** All the dynamic scaling informations */
 	DynamicRenderScaling::TMap<float> DynamicResolutionFractions;
@@ -2098,13 +2141,6 @@ public:
 	FSceneRenderer(const FSceneViewFamily* InViewFamily, FHitProxyConsumer* HitProxyConsumer);
 	virtual ~FSceneRenderer();
 
-	// Initializes the scene renderer on the render thread.
-	RENDERER_API void RenderThreadBegin(FRHICommandListImmediate& RHICmdList);
-	RENDERER_API void RenderThreadEnd(FRHICommandListImmediate& RHICmdList);
-
-	RENDERER_API static void RenderThreadBegin(FRHICommandListImmediate& RHICmdList, TConstArrayView<FSceneRenderer*> SceneRenderers);
-	RENDERER_API static void RenderThreadEnd(FRHICommandListImmediate& RHICmdList, TConstArrayView<FSceneRenderer*> SceneRenderers);
-
 	// ISceneRenderer interface
 	virtual UE::Renderer::Private::IShadowInvalidatingInstances *GetShadowInvalidatingInstancesInterface(const FSceneView *SceneView) override;
 
@@ -2116,12 +2152,6 @@ public:
 	virtual bool ShouldRenderNanite() const { return false; }
 	virtual bool AllowSimpleLights() const;
 
-	/** Creates a scene renderer based on the current feature level. */
-	RENDERER_API static FSceneRenderer* CreateSceneRenderer(const FSceneViewFamily* InViewFamily, FHitProxyConsumer* HitProxyConsumer);
-
-	/** Creates multiple scene renderers based on the current feature level.  All view families must point to the same Scene. */
-	RENDERER_API static void CreateSceneRenderers(TArray<const FSceneViewFamily*>& InOutViewFamilies, FHitProxyConsumer* HitProxyConsumer, bool bAllowSplitScreenDebug, TArray<FSceneRenderer*>& OutSceneRenderers);
-
 	/** Setups FViewInfo::ViewRect according to ViewFamilly's ScreenPercentageInterface. */
 	void PrepareViewRectsForRendering(FRHICommandListImmediate& RHICmdList);
 
@@ -2165,12 +2195,6 @@ public:
 	static bool ShouldCompositeDebugPrimitivesInPostProcess(const FViewInfo& View);
 #endif
 
-	/** Called to release any deallocations that were deferred until the next render. */
-	static void CleanUp(FRHICommandListImmediate& RHICmdList);
-
-	/** Called to wait for dependent clean up tasks, but CleanUp must still be called later. */
-	static void WaitForCleanUpTasks(FRHICommandListImmediate& RHICmdList);
-
 	/** Apply the ResolutionFraction on ViewSize, taking into account renderer's requirements. */
 	static FIntPoint ApplyResolutionFraction(
 		const FSceneViewFamily& ViewFamily, const FIntPoint& UnscaledViewSize, float ResolutionFraction);
@@ -2579,6 +2603,7 @@ private:
 #endif
 
 	friend class FRendererModule;
+	friend class FSceneRenderProcessor;
 };
 
 template <typename LambdaType, typename PrerequisiteTaskCollectionType>
diff --git a/Engine/Source/Runtime/Renderer/Private/SceneVisibility.cpp b/Engine/Source/Runtime/Renderer/Private/SceneVisibility.cpp
--- a/Engine/Source/Runtime/Renderer/Private/SceneVisibility.cpp
+++ b/Engine/Source/Runtime/Renderer/Private/SceneVisibility.cpp
@@ -4825,7 +4825,7 @@ void FSceneRenderer::PreVisibilityFrameSetup(FRDGBuilder& GraphBuilder)
 		
 		if (Views.Num() > 0 && !ViewFamily.EngineShowFlags.HitProxies)
 		{
-			CreateHairStrandsBookmarkParameters(Scene, Views, AllFamilyViews, Parameters, false /*bComputeVisibleInstances*/);
+			CreateHairStrandsBookmarkParameters(Scene, Views, AllViews, Parameters, false /*bComputeVisibleInstances*/);
 			Parameters.TransientResources = AllocateHairTransientResources(GraphBuilder, Scene);
 			if (Parameters.HasInstances())
 			{
@@ -4861,7 +4861,13 @@ void FSceneRenderer::PreVisibilityFrameSetup(FRDGBuilder& GraphBuilder)
 
 		if (FXSystem && Views.IsValidIndex(0))
 		{
-			FXSystem->PreInitViews(GraphBuilder, Views[0].AllowGPUParticleUpdate() && !ViewFamily.EngineShowFlags.HitProxies, AllFamilies, &ViewFamily);
+			TArray<const FSceneViewFamily*, TInlineAllocator<2>> AllLinkedFamilies;
+			EnumerateLinkedViewFamilies([&] (const FSceneViewFamily& Family)
+			{
+				AllLinkedFamilies.Emplace(&Family);
+				return true;
+			});
+			FXSystem->PreInitViews(GraphBuilder, Views[0].AllowGPUParticleUpdate() && !ViewFamily.EngineShowFlags.HitProxies, AllLinkedFamilies, &ViewFamily);
 		}
 
 		if(bHasHairStrandsBookmarks)
