Workaround for when the clipmap bit array is empty and overrides a valid cached bit array.

diff --git a/Engine/Source/Runtime/Renderer/Private/VirtualShadowMaps/VirtualShadowMapClipmap.cpp b/Engine/Source/Runtime/Renderer/Private/VirtualShadowMaps/VirtualShadowMapClipmap.cpp
--- a/Engine/Source/Runtime/Renderer/Private/VirtualShadowMaps/VirtualShadowMapClipmap.cpp
+++ b/Engine/Source/Runtime/Renderer/Private/VirtualShadowMaps/VirtualShadowMapClipmap.cpp
@@ -543,7 +543,7 @@ void FVirtualShadowMapClipmap::OnPrimitiveRendered(const FPrimitiveSceneInfo* Pr
 
 void FVirtualShadowMapClipmap::UpdateCachedFrameData()
 {
-	if (PerLightCacheEntry.IsValid())
+	if (PerLightCacheEntry.IsValid() && !RenderedPrimitives.IsEmpty())
 	{
 		PerLightCacheEntry->RenderedPrimitives = MoveTemp(RenderedPrimitives);
 	}
