Add thin geometry detection permutation in reject shading pass to not regress runtime performance.

diff --git a/Engine/Shaders/Private/TemporalSuperResolution/TSRRejectShading.usf b/Engine/Shaders/Private/TemporalSuperResolution/TSRRejectShading.usf
--- a/Engine/Shaders/Private/TemporalSuperResolution/TSRRejectShading.usf
+++ b/Engine/Shaders/Private/TemporalSuperResolution/TSRRejectShading.usf
@@ -19,6 +19,7 @@
 
 #define CONFIG_HISTORY_RESURRECTION (DIM_HISTORY_RESURRECTION)
 
+#define CONFIG_THIN_GEOMETRY_DETECTION (DIM_THIN_GEOMETRY_DETECTION)
 
 //------------------------------------------------------- CONSTANTS
 
@@ -719,6 +720,7 @@ void MainCS(
 	} // if (bComputeResurrection)
 	#endif // CONFIG_HISTORY_RESURRECTION
 	
+    #if CONFIG_THIN_GEOMETRY_DETECTION
 	PLATFORM_SPECIFIC_ISOLATE
 	{
 		tsr_tensor_short2 PixelPos = ComputePixelPos();
@@ -727,6 +729,7 @@ void MainCS(
 		RejectionBlendFinal = lerp(RejectionBlendFinal, tsr_tensor_half::Const(tsr_half(1.0)), IsThinGeometryEdge);
 		RejectionClampBlend = lerp(RejectionClampBlend, tsr_tensor_half::Const(tsr_half(1.0)), IsThinGeometryEdge);
 	}
+    #endif
 	
 	// Update the guide for next frame.
 	tsr_tensor_halfC FinalGuide;
diff --git a/Engine/Source/Runtime/Renderer/Private/PostProcess/TemporalSuperResolution.cpp b/Engine/Source/Runtime/Renderer/Private/PostProcess/TemporalSuperResolution.cpp
--- a/Engine/Source/Runtime/Renderer/Private/PostProcess/TemporalSuperResolution.cpp
+++ b/Engine/Source/Runtime/Renderer/Private/PostProcess/TemporalSuperResolution.cpp
@@ -794,11 +794,13 @@ class FTSRRejectShadingCS : public FTSRConvolutionNetworkShader
 
 	class FFlickeringDetectionDim : SHADER_PERMUTATION_BOOL("DIM_FLICKERING_DETECTION");
 	class FHistoryResurrectionDim : SHADER_PERMUTATION_BOOL("DIM_HISTORY_RESURRECTION");
+	class FThinGeometryDetectionDim : SHADER_PERMUTATION_BOOL("DIM_THIN_GEOMETRY_DETECTION");
 
 	using FPermutationDomain = TShaderPermutationDomain<
 		FTSRConvolutionNetworkShader::FPermutationDomain,
 		FFlickeringDetectionDim,
-		FHistoryResurrectionDim>;
+		FHistoryResurrectionDim,
+		FThinGeometryDetectionDim>;
 
 	BEGIN_SHADER_PARAMETER_STRUCT(FParameters, )
 		SHADER_PARAMETER_STRUCT_INCLUDE(FTSRCommonParameters, CommonParameters)
@@ -2434,6 +2436,7 @@ FDefaultTemporalUpscaler::FOutputs AddTemporalSuperResolutionPasses(
 		PermutationVector.Set<FTSRConvolutionNetworkShader::FPermutationDomain>(ConvolutionNetworkPermutationVector);
 		PermutationVector.Set<FTSRRejectShadingCS::FFlickeringDetectionDim>(FlickeringFramePeriod > 0.0f);
 		PermutationVector.Set<FTSRRejectShadingCS::FHistoryResurrectionDim>(bCanResurrectHistory);
+		PermutationVector.Set<FTSRRejectShadingCS::FThinGeometryDetectionDim>(PassConfig.ThinGeometryDetectionEnable);
 		PermutationVector = FTSRRejectShadingCS::RemapPermutation(PermutationVector);
 
 		const int32 GroupTileSize = 32;
