Default r.VolumetricRenderTarget.ViewRect to 0 which fallbacks to the old behaviour when dynamic res changes.
Doing this temporarily while we investigate why we're seeing GPU hangs with it set to 1.

diff --git a/Engine/Source/Runtime/Renderer/Private/VolumetricRenderTarget.cpp b/Engine/Source/Runtime/Renderer/Private/VolumetricRenderTarget.cpp
--- a/Engine/Source/Runtime/Renderer/Private/VolumetricRenderTarget.cpp
+++ b/Engine/Source/Runtime/Renderer/Private/VolumetricRenderTarget.cpp
@@ -47,7 +47,7 @@ static TAutoConsoleVariable<int32> CVarVolumetricRenderTargetPreferAsyncCompute(
 	ECVF_RenderThreadSafe | ECVF_Scalability);
 
 static TAutoConsoleVariable<int32> CVarVolumetricRenderTargetViewRect(
-	TEXT("r.VolumetricRenderTarget.ViewRect"), 1,
+	TEXT("r.VolumetricRenderTarget.ViewRect"), 0,
 	TEXT("Enable ViewRect support: does not reallocate new render targets when dynamic resolution changes"),
 	ECVF_RenderThreadSafe | ECVF_Scalability);
 
