Changed the HZB build to avoid using more than 8 UAVs when froxel build is enabled & both nearest and furthest are used by reducing the max out mips to 3 for the first pass (if froxels are enabled).

diff --git a/Engine/Source/Runtime/Renderer/Private/SceneTextureReductions.cpp b/Engine/Source/Runtime/Renderer/Private/SceneTextureReductions.cpp
--- a/Engine/Source/Runtime/Renderer/Private/SceneTextureReductions.cpp
+++ b/Engine/Source/Runtime/Renderer/Private/SceneTextureReductions.cpp
@@ -160,10 +160,24 @@ void BuildHZB(
 
 	int32 MaxMipBatchSize = bUseCompute ? FHZBBuildCS::kMaxMipBatchSize : 1;
 
+
+	int32 MaxMipBatchSizeFirstPass = MaxMipBatchSize;
+
+	// Avoid overflowing the 8 UAV limit by reducing the number of levels in the first pass (the only one that writes both near/far textures at once)
+	if (bUseCompute && bReduceClosestDepth && OutFroxelViewData != nullptr && MaxSimultaneousUAVs <= 8)
+	{
+		// 6 UAVs for the ouput mip levels + 2 from the froxel builder
+		MaxMipBatchSizeFirstPass = 3;
+	}
+
 	auto ReduceMips = [&](
 		FRDGTextureSRVRef ParentTextureMip, FIntPoint SrcSize,
-		int32 StartDestMip, FVector4f DispatchThreadIdToBufferUV, FVector2D InputViewportMaxBound,
-		FIntVector4 PixelViewPortMinMax, bool bOutputClosest, bool bOutputFurthest, bool bProduceFroxelData)
+		int32 MipBatchSize,
+		int32 StartDestMip, 
+		FVector4f DispatchThreadIdToBufferUV, 
+		FVector2D InputViewportMaxBound,
+		FIntVector4 PixelViewPortMinMax, 
+		bool bOutputClosest, bool bOutputFurthest, bool bProduceFroxelData)
 	{
 
 		FSharedHZBParameters ShaderParameters;
@@ -179,7 +193,7 @@ void BuildHZB(
 
 		if (bUseCompute)
 		{
-			int32 EndDestMip = FMath::Min(StartDestMip + FHZBBuildCS::kMaxMipBatchSize, NumMips);
+			int32 EndDestMip = FMath::Min(StartDestMip + MipBatchSize, NumMips);
 
 			FHZBBuildCS::FParameters* PassParameters = GraphBuilder.AllocParameters<FHZBBuildCS::FParameters>();
 			PassParameters->Shared = ShaderParameters;
@@ -300,13 +314,14 @@ void BuildHZB(
 
 		ReduceMips(
 			ParentTextureMip, SrcSize,
-			/* StartDestMip = */ 0, DispatchThreadIdToBufferUV, InputViewportMaxBound, PixelViewPortMinMax,
+			MaxMipBatchSizeFirstPass, /* StartDestMip = */ 0, 
+			DispatchThreadIdToBufferUV, InputViewportMaxBound, PixelViewPortMinMax,
 			/* bOutputClosest = */ bReduceClosestDepth, /* bOutputFurthest = */ true,
 			OutFroxelViewData != nullptr);
 	}
 
 	// Reduce the next mips
-	for (int32 StartDestMip = MaxMipBatchSize; StartDestMip < NumMips; StartDestMip += MaxMipBatchSize)
+	for (int32 StartDestMip = MaxMipBatchSizeFirstPass; StartDestMip < NumMips; StartDestMip += MaxMipBatchSize)
 	{
 		FIntPoint SrcSize = FIntPoint::DivideAndRoundUp(HZBSize, 1 << int32(StartDestMip - 1));
 
@@ -322,7 +337,8 @@ void BuildHZB(
 		{
 			FRDGTextureSRVRef ParentTextureMip = GraphBuilder.CreateSRV(GRHISupportsTextureViews ? FRDGTextureSRVDesc::CreateForMipLevel(FurthestHZBTexture, StartDestMip - 1) : FRDGTextureSRVDesc::Create(FurthestHZBTexture));
 			ReduceMips(ParentTextureMip, SrcSize,
-				StartDestMip, DispatchThreadIdToBufferUV, InputViewportMaxBound, PixelViewPortMinMax,
+				MaxMipBatchSize, StartDestMip, 
+				DispatchThreadIdToBufferUV, InputViewportMaxBound, PixelViewPortMinMax,
 				/* bOutputClosest = */ false, /* bOutputFurthest = */ true, false);
 		}
 
@@ -331,7 +347,8 @@ void BuildHZB(
 			check(ClosestHZBTexture)
 			FRDGTextureSRVRef ParentTextureMip = GraphBuilder.CreateSRV(GRHISupportsTextureViews ? FRDGTextureSRVDesc::CreateForMipLevel(ClosestHZBTexture, StartDestMip - 1) : FRDGTextureSRVDesc::Create(ClosestHZBTexture));
 			ReduceMips(ParentTextureMip, SrcSize,
-				StartDestMip, DispatchThreadIdToBufferUV, InputViewportMaxBound, PixelViewPortMinMax,
+				MaxMipBatchSize, StartDestMip, 
+				DispatchThreadIdToBufferUV, InputViewportMaxBound, PixelViewPortMinMax,
 				/* bOutputClosest = */ true, /* bOutputFurthest = */ false, false);
 		}
 	}
