Tiny followup change to make the logic that detects MaxSimultaneousUAVs overflow for the HZB build cleaner & more robust.

#rnx

diff --git a/Engine/Source/Runtime/Renderer/Private/SceneTextureReductions.cpp b/Engine/Source/Runtime/Renderer/Private/SceneTextureReductions.cpp
--- a/Engine/Source/Runtime/Renderer/Private/SceneTextureReductions.cpp
+++ b/Engine/Source/Runtime/Renderer/Private/SceneTextureReductions.cpp
@@ -163,9 +163,10 @@ void BuildHZB(
 
 	int32 MaxMipBatchSizeFirstPass = MaxMipBatchSize;
 
-	// Avoid overflowing the 8 UAV limit by reducing the number of levels in the first pass (the only one that writes both near/far textures at once)
-	if (bUseCompute && bReduceClosestDepth && OutFroxelViewData != nullptr && MaxSimultaneousUAVs <= 8)
+	// Avoid overflowing MaxSimultaneousUAVs by reducing the number of levels in the first pass (the only one that writes both near/far textures at once)
+	if (bUseCompute && bReduceClosestDepth && OutFroxelViewData != nullptr && MaxSimultaneousUAVs < 10)
 	{
+		ensure(MaxSimultaneousUAVs >= 8);
 		// 6 UAVs for the ouput mip levels + 2 from the froxel builder
 		MaxMipBatchSizeFirstPass = 3;
 	}
