Re-added missing FViewInfo:::DestroyAllSnapshots() call, which was causing snapshots to accumulate and cause OOMs.

diff --git a/Engine/Source/Runtime/Renderer/Private/SceneRenderBuilder.cpp b/Engine/Source/Runtime/Renderer/Private/SceneRenderBuilder.cpp
--- a/Engine/Source/Runtime/Renderer/Private/SceneRenderBuilder.cpp
+++ b/Engine/Source/Runtime/Renderer/Private/SceneRenderBuilder.cpp
@@ -594,6 +594,8 @@ static void CleanupSceneRenderer(FSceneRenderer* Renderer)
 	{
 		View->WaitForTasks();
 	}
+
+	FViewInfo::DestroyAllSnapshots();
 }
 
 void FSceneRenderProcessor::Cleanup(FRHICommandListImmediate& RHICmdList, FSceneRenderer* Renderer)
