Fix [Holdout] Doesn't work on Mac anymore.

#tests test in jira on mac run as expected with this CL. Windows also run as expected.

diff --git a/Engine/Source/Runtime/Renderer/Private/TranslucentRendering.cpp b/Engine/Source/Runtime/Renderer/Private/TranslucentRendering.cpp
--- a/Engine/Source/Runtime/Renderer/Private/TranslucentRendering.cpp
+++ b/Engine/Source/Runtime/Renderer/Private/TranslucentRendering.cpp
@@ -1708,8 +1708,9 @@ void FDeferredShadingSceneRenderer::RenderTranslucencyInner(
 			if (NumProcessedViews == 0 || View.Family->bMultiGPUForkAndJoin)
 			{
 				FTranslucencyPassResources& TranslucencyPassResources = OutTranslucencyResourceMap->Get(ViewIndex, TranslucencyPass);
+				FIntRect CopyRect = TranslucencyPassResources.ViewRect == FIntRect(0,0,0,0) ?  View.ViewRect : TranslucencyPassResources.ViewRect;
 				AddCopyBackgroundVisibilityPass(GraphBuilder, View, SceneColorCopyTexture/*SceneTextures.Color.Resolve*/, SeparateTranslucencyColorTexture.Target,
-					TranslucencyPassResources.ViewRect, FCopyBackgroundVisibilityPS::EVisibilityCopyType::FromSceneColor);
+					CopyRect, FCopyBackgroundVisibilityPS::EVisibilityCopyType::FromSceneColor);
 			}
 
 			// NOTE: We need to read the depth texture for final fog accumulation, but
