Refactored view snapshot cache into its own file and fixed race condition where snapshots could be destroyed that are still being used in the next frame.

diff --git a/Engine/Source/Runtime/Engine/Public/SceneView.h b/Engine/Source/Runtime/Engine/Public/SceneView.h
--- a/Engine/Source/Runtime/Engine/Public/SceneView.h
+++ b/Engine/Source/Runtime/Engine/Public/SceneView.h
@@ -1992,6 +1992,8 @@ public:
 	}
 
 protected:
+	friend class FViewSnapshotCache;
+
 	FSceneViewStateInterface* EyeAdaptationViewState = nullptr;
 
 	/** Instanced view uniform buffer held on the primary view. */
diff --git a/Engine/Source/Runtime/Renderer/Private/Renderer/ViewSnapshotCache.cpp b/Engine/Source/Runtime/Renderer/Private/Renderer/ViewSnapshotCache.cpp
new file mode 100644
--- /dev/null
+++ b/Engine/Source/Runtime/Renderer/Private/Renderer/ViewSnapshotCache.cpp
@@ -0,0 +1,147 @@
+// Copyright Epic Games, Inc. All Rights Reserved.
+
+#include "ViewSnapshotCache.h"
+#include "SceneRendering.h"
+
+class FViewSnapshotCache
+{
+public:
+	FViewSnapshotCache() = default;
+
+	FViewInfo* Create(const FViewInfo* InView)
+	{
+		TRACE_CPUPROFILER_EVENT_SCOPE(ViewSnapshotCache::Create);
+		check(IsInParallelRenderingThread()); // we do not want this popped before the end of the scene and it better be the scene allocator
+
+		FViewInfo* Result = FreeSnapshots.Pop();
+
+		if (Result)
+		{
+			NumFreeSnapshots--;
+		}
+		else
+		{
+			Result = (FViewInfo*)FMemory::Malloc(sizeof(FViewInfo), alignof(FViewInfo));
+		}
+
+		Snapshots.Push(Result);
+		NumSnapshots++;
+
+		FMemory::Memcpy(*Result, *InView);
+
+		// we want these to start null without a reference count, since we clear a ref later
+		TUniformBufferRef<FViewUniformShaderParameters> NullViewUniformBuffer;
+		FMemory::Memcpy(Result->ViewUniformBuffer, NullViewUniformBuffer);
+		TUniformBufferRef<FInstancedViewUniformShaderParameters> NullInstancedViewUniformBuffer;
+		FMemory::Memcpy(Result->InstancedViewUniformBuffer, NullInstancedViewUniformBuffer);
+
+		TUniquePtr<FViewUniformShaderParameters> NullViewParameters;
+		FMemory::Memcpy(Result->CachedViewUniformShaderParameters, NullViewParameters); 
+
+		TArray<FPrimitiveUniformShaderParameters> NullDynamicPrimitiveShaderData;
+		TStaticArray<FParallelMeshDrawCommandPass, EMeshPass::Num> NullParallelMeshDrawCommandPasses;
+		FMemory::Memcpy(Result->ParallelMeshDrawCommandPasses, NullParallelMeshDrawCommandPasses);
+
+		for (int i = 0; i < EMeshPass::Num; i++)
+		{
+			Result->ParallelMeshDrawCommandPasses[i].InitCreateSnapshot();
+		}
+
+		// Ensure the internal state is maintained, needed because we've just Memcpy'd the member data.
+		static_assert(std::is_trivially_destructible_v<FGPUScenePrimitiveCollector> != 0, "The destructor is not invoked properly because of FMemory::Memcpy(*Result, *this) above");
+		Result->DynamicPrimitiveCollector = FGPUScenePrimitiveCollector(InView->DynamicPrimitiveCollector);
+
+		Result->SnapshotOriginView = InView;
+		return Result;
+	}
+
+	void Deallocate()
+	{
+		check(IsInRenderingThread());
+
+		// Only keep double the number actually used, plus a few
+		NumSnapshotsToRemove = FMath::Max(NumFreeSnapshots - (NumSnapshots + 2), 0);
+		NumSnapshots = 0;
+
+		while (FViewInfo* Snapshot = Snapshots.Pop())
+		{
+			DeallocatedSnapshots.Add(Snapshot);
+		}
+	}
+
+	void Destroy()
+	{
+		TRACE_CPUPROFILER_EVENT_SCOPE(FViewInfo::DestroyAllSnapshots);
+
+		if (NumSnapshotsToRemove > 0)
+		{
+			while (FViewInfo* Snapshot = FreeSnapshots.Pop())
+			{
+				NumFreeSnapshots.fetch_sub(1, std::memory_order_relaxed);
+				FMemory::Free(Snapshot);
+
+				if (--NumSnapshotsToRemove == 0)
+				{
+					break;
+				}
+			}
+
+			// It's possible that this will not be 0 if allocation happened concurrently, which is safe.
+			NumSnapshotsToRemove = 0;
+		}
+
+		for (FViewInfo* Snapshot : DeallocatedSnapshots)
+		{
+			Snapshot->ViewUniformBuffer.SafeRelease();
+			Snapshot->InstancedViewUniformBuffer.SafeRelease();
+			Snapshot->CachedViewUniformShaderParameters.Reset();
+
+			for (auto& Pass : Snapshot->ParallelMeshDrawCommandPasses)
+			{
+				Pass.WaitForTasksAndEmpty();
+			}
+
+			for (int32 MeshPassIndex = 0; MeshPassIndex < EMeshPass::Num; ++MeshPassIndex)
+			{
+				Snapshot->ParallelMeshDrawCommandPasses[MeshPassIndex].FreeCreateSnapshot();
+			}
+
+			Snapshot->SnapshotOriginView = nullptr;
+
+			FreeSnapshots.Push(Snapshot);
+			NumFreeSnapshots.fetch_add(1, std::memory_order_relaxed);
+		}
+		DeallocatedSnapshots.Empty();
+	}
+
+private:
+	// These are not real view infos, just dumb memory blocks
+	TLockFreePointerListLIFOPad<FViewInfo, PLATFORM_CACHE_LINE_SIZE> Snapshots;
+	// these are never freed, even at program shutdown
+	TLockFreePointerListLIFOPad<FViewInfo, PLATFORM_CACHE_LINE_SIZE> FreeSnapshots;
+
+	TArray<FViewInfo*> DeallocatedSnapshots;
+	std::atomic_int32_t NumSnapshots{ 0 };
+	std::atomic_int32_t NumFreeSnapshots{ 0 };
+	int32 NumSnapshotsToRemove = 0;
+};
+
+namespace ViewSnapshotCache
+{
+	static FViewSnapshotCache Cache;
+
+	FViewInfo* Create(const FViewInfo* View)
+	{
+		return Cache.Create(View);
+	}
+
+	void Deallocate()
+	{
+		Cache.Deallocate();
+	}
+
+	void Destroy()
+	{
+		Cache.Destroy();
+	}
+}
\ No newline at end of file
diff --git a/Engine/Source/Runtime/Renderer/Private/Renderer/ViewSnapshotCache.h b/Engine/Source/Runtime/Renderer/Private/Renderer/ViewSnapshotCache.h
new file mode 100644
--- /dev/null
+++ b/Engine/Source/Runtime/Renderer/Private/Renderer/ViewSnapshotCache.h
@@ -0,0 +1,17 @@
+// Copyright Epic Games, Inc. All Rights Reserved.
+
+#pragma once
+
+class FViewInfo;
+
+namespace ViewSnapshotCache
+{
+	// Creates and returns a snapshot from the input view.
+	extern FViewInfo* Create(const FViewInfo* View);
+
+	// Called on the render thread after rendering to prepare snapshots for destruction.
+	extern void Deallocate();
+
+	// Called on a worker thread to destroy snapshots.
+	extern void Destroy();
+}
\ No newline at end of file
diff --git a/Engine/Source/Runtime/Renderer/Private/SceneRenderBuilder.cpp b/Engine/Source/Runtime/Renderer/Private/SceneRenderBuilder.cpp
--- a/Engine/Source/Runtime/Renderer/Private/SceneRenderBuilder.cpp
+++ b/Engine/Source/Runtime/Renderer/Private/SceneRenderBuilder.cpp
@@ -13,6 +13,7 @@
 #include "TextureResource.h"
 #include "DeferredShadingRenderer.h"
 #include "FXSystem.h"
+#include "Renderer/ViewSnapshotCache.h"
 
 //////////////////////////////////////////////////////////////////////////
 
@@ -545,17 +546,23 @@ public:
 	static void WaitForAsyncDeleteTask()
 	{
 		check(IsInRenderingThread());
-		if (AsyncDeleteTask)
+		if (AsyncTasks.Delete)
 		{
-			AsyncDeleteTask->Wait();
-			AsyncDeleteTask = {};
+			AsyncTasks.Delete->Wait();
+			AsyncTasks = {};
 		}
 	}
 
 private:
-	static FGraphEventRef AsyncDeleteTask;
+	struct FAsyncTasks
+	{
+		FGraphEventRef Cleanup;
+		FGraphEventRef Delete;
+	};
+
+	static FAsyncTasks AsyncTasks;
+
 	void Cleanup(FRHICommandListImmediate& RHICmdList, FSceneRenderer* SceneRenderer);
-	FGraphEventRef CleanupTaskEvent;
 
 	FScene* Scene;
 	TArray<FSceneRenderer*, FConcurrentLinearArrayAllocator> Renderers;
@@ -579,7 +586,7 @@ private:
 	bool bInsideGroup : 1 = false;
 };
 
-FGraphEventRef FSceneRenderProcessor::AsyncDeleteTask;
+FSceneRenderProcessor::FAsyncTasks FSceneRenderProcessor::AsyncTasks;
 
 static void CleanupSceneRenderer(FSceneRenderer* Renderer)
 {
@@ -595,7 +602,7 @@ static void CleanupSceneRenderer(FSceneRenderer* Renderer)
 		View->WaitForTasks();
 	}
 
-	FViewInfo::DestroyAllSnapshots();
+	ViewSnapshotCache::Destroy();
 }
 
 void FSceneRenderProcessor::Cleanup(FRHICommandListImmediate& RHICmdList, FSceneRenderer* Renderer)
@@ -639,16 +646,24 @@ void FSceneRenderProcessor::Cleanup(FRHICommandListImmediate& RHICmdList, FScene
 			}
 		}
 
+		// Wait for the last renderer's cleanup tasks so that snapshot deallocation and destruction don't overlap.
+		if (AsyncTasks.Cleanup)
+		{
+			Prereqs.Add(AsyncTasks.Cleanup);
+		}
+
 		if (!Prereqs.IsEmpty())
 		{
 			FTaskGraphInterface::Get().WaitUntilTasksComplete(Prereqs, ENamedThreads::GetRenderThread_Local());
 		}
 
-		CleanupTaskEvent->AddPrerequisites(FFunctionGraphTask::CreateAndDispatchWhenReady([Renderer]
+		ViewSnapshotCache::Deallocate();
+
+		AsyncTasks.Cleanup = FFunctionGraphTask::CreateAndDispatchWhenReady([Renderer]
 		{
 			CleanupSceneRenderer(Renderer);
 
-		}, TStatId(), &GRHICommandList.WaitOutstandingTasks));
+		}, TStatId(), &GRHICommandList.WaitOutstandingTasks);
 	}
 
 	GlobalDynamicBuffer::GarbageCollect();
@@ -758,8 +773,6 @@ void FSceneRenderProcessor::Execute()
 	}
 #endif
 
-	CleanupTaskEvent = FGraphEvent::CreateGraphEvent();
-
 	UE::RenderCommandPipe::FSyncScope SyncScope;
 	FUniformExpressionCacheAsyncUpdateScope AsyncUpdateScope;
 	TOptional<UE::RenderCore::DumpGPU::FDumpScope> GpuDumpScope;
@@ -882,10 +895,11 @@ void FSceneRenderProcessor::Execute()
 
 	ENQUEUE_RENDER_COMMAND(SceneRenderBuilder_End)([this](FRHICommandListImmediate& RHICmdList) mutable
 	{
-		CleanupTaskEvent->AddPrerequisites(AsyncDeleteTask);
-		CleanupTaskEvent->DispatchSubsequents();
+		FGraphEventArray Prereqs;
+		Prereqs.Add(AsyncTasks.Cleanup);
+		Prereqs.Add(AsyncTasks.Delete);
 
-		AsyncDeleteTask = FFunctionGraphTask::CreateAndDispatchWhenReady([this]
+		AsyncTasks.Delete = FFunctionGraphTask::CreateAndDispatchWhenReady([this]
 		{
 			TRACE_CPUPROFILER_EVENT_SCOPE(FSceneRenderProcessor::DeleteSceneRenderers);
 			for (FSceneRenderer* Renderer : Renderers)
@@ -894,7 +908,7 @@ void FSceneRenderProcessor::Execute()
 			}
 			delete this;
 
-		}, TStatId(), CleanupTaskEvent);
+		}, TStatId(), &Prereqs);
 	});
 
 	// NOTE: 'this' is queued for deletion and is no longer valid!
diff --git a/Engine/Source/Runtime/Renderer/Private/SceneRendering.cpp b/Engine/Source/Runtime/Renderer/Private/SceneRendering.cpp
--- a/Engine/Source/Runtime/Renderer/Private/SceneRendering.cpp
+++ b/Engine/Source/Runtime/Renderer/Private/SceneRendering.cpp
@@ -113,6 +113,7 @@
 #include "SceneRenderBuilder.h"
 #include <type_traits>
 #include "BlueNoise.h"
+#include "Renderer/ViewSnapshotCache.h"
 
 /*-----------------------------------------------------------------------------
 	Globals
@@ -2283,123 +2284,11 @@ const FViewInfo* FViewInfo::GetPrimaryView() const
 	return this;
 }
 
-struct FViewInfoSnapshotCache
-{
-	// These are not real view infos, just dumb memory blocks
-	TLockFreePointerListLIFOPad<FViewInfo, PLATFORM_CACHE_LINE_SIZE> Snapshots;
-	// these are never freed, even at program shutdown
-	TLockFreePointerListLIFOPad<FViewInfo, PLATFORM_CACHE_LINE_SIZE> FreeSnapshots;
-
-	std::atomic_int32_t NumSnapshots{ 0 };
-	std::atomic_int32_t NumFreeSnapshots{ 0 };
-};
-
-static FViewInfoSnapshotCache ViewInfoSnapshotCache;
-
 extern TUniformBufferRef<FMobileDirectionalLightShaderParameters>& GetNullMobileDirectionalLightShaderParameters();
 
 FViewInfo* FViewInfo::CreateSnapshot() const
 {
-	QUICK_SCOPE_CYCLE_COUNTER(STAT_FViewInfo_CreateSnapshot);
-
-	check(IsInParallelRenderingThread()); // we do not want this popped before the end of the scene and it better be the scene allocator
-	FViewInfo* Result = ViewInfoSnapshotCache.FreeSnapshots.Pop();
-	if (Result)
-	{
-		ViewInfoSnapshotCache.NumFreeSnapshots--;
-	}
-	else
-	{
-		Result = (FViewInfo*)FMemory::Malloc(sizeof(FViewInfo), alignof(FViewInfo));
-	}
-
-	ViewInfoSnapshotCache.Snapshots.Push(Result);
-	ViewInfoSnapshotCache.NumSnapshots++;
-
-	FMemory::Memcpy(*Result, *this);
-
-	// we want these to start null without a reference count, since we clear a ref later
-	TUniformBufferRef<FViewUniformShaderParameters> NullViewUniformBuffer;
-	FMemory::Memcpy(Result->ViewUniformBuffer, NullViewUniformBuffer);
-	TUniformBufferRef<FInstancedViewUniformShaderParameters> NullInstancedViewUniformBuffer;
-	FMemory::Memcpy(Result->InstancedViewUniformBuffer, NullInstancedViewUniformBuffer);
-
-	TUniquePtr<FViewUniformShaderParameters> NullViewParameters;
-	FMemory::Memcpy(Result->CachedViewUniformShaderParameters, NullViewParameters); 
-
-	TArray<FPrimitiveUniformShaderParameters> NullDynamicPrimitiveShaderData;
-	TStaticArray<FParallelMeshDrawCommandPass, EMeshPass::Num> NullParallelMeshDrawCommandPasses;
-	FMemory::Memcpy(Result->ParallelMeshDrawCommandPasses, NullParallelMeshDrawCommandPasses);
-
-	for (int i = 0; i < EMeshPass::Num; i++)
-	{
-		Result->ParallelMeshDrawCommandPasses[i].InitCreateSnapshot();
-	}
-	
-	// Ensure the internal state is maintained, needed because we've just Memcpy'd the member data.
-	static_assert(std::is_trivially_destructible_v<FGPUScenePrimitiveCollector> != 0, "The destructor is not invoked properly because of FMemory::Memcpy(*Result, *this) above");
-	Result->DynamicPrimitiveCollector = FGPUScenePrimitiveCollector(DynamicPrimitiveCollector);
-
-	Result->SnapshotOriginView = this;
-	return Result;
-}
-
-void FViewInfo::DestroyAllSnapshots()
-{
-	TRACE_CPUPROFILER_EVENT_SCOPE(FViewInfo::DestroyAllSnapshots);
-	static FCriticalSection Mutex;
-
-	// Someone else is destroying snapshots, so just return.
-	if (!Mutex.TryLock())
-	{
-		return;
-	}
-
-	ON_SCOPE_EXIT
-	{
-		Mutex.Unlock();
-	};
-
-	// we will only keep double the number actually used, plus a few
-	const int32 NumToRemove = FMath::Max(ViewInfoSnapshotCache.NumFreeSnapshots - (ViewInfoSnapshotCache.NumSnapshots + 2), 0);
-
-	int32 NumRemoved = 0;
-	while (FViewInfo* Snapshot = ViewInfoSnapshotCache.FreeSnapshots.Pop())
-	{
-		ViewInfoSnapshotCache.NumFreeSnapshots.fetch_sub(1, std::memory_order_relaxed);
-
-		FMemory::Free(Snapshot);
-
-
-		if (++NumRemoved == NumToRemove)
-		{
-			break;
-		}
-	}
-
-	while (FViewInfo* Snapshot = ViewInfoSnapshotCache.Snapshots.Pop())
-	{
-		ViewInfoSnapshotCache.NumSnapshots.fetch_sub(1, std::memory_order_relaxed);
-
-		Snapshot->ViewUniformBuffer.SafeRelease();
-		Snapshot->InstancedViewUniformBuffer.SafeRelease();
-		Snapshot->CachedViewUniformShaderParameters.Reset();
-
-		for (int32 Index = 0; Index < Snapshot->ParallelMeshDrawCommandPasses.Num(); ++Index)
-		{
-			Snapshot->ParallelMeshDrawCommandPasses[Index].WaitForTasksAndEmpty();
-		}
-
-		for (int i = 0; i < EMeshPass::Num; i++)
-		{
-			Snapshot->ParallelMeshDrawCommandPasses[i].FreeCreateSnapshot();
-		}
-
-		Snapshot->SnapshotOriginView = nullptr;
-
-		ViewInfoSnapshotCache.FreeSnapshots.Push(Snapshot);
-		ViewInfoSnapshotCache.NumFreeSnapshots.fetch_add(1, std::memory_order_relaxed);
-	}
+	return ViewSnapshotCache::Create(this);
 }
 
 FInt32Range FViewInfo::GetDynamicMeshElementRange(uint32 PrimitiveIndex) const
diff --git a/Engine/Source/Runtime/Renderer/Private/SceneRendering.h b/Engine/Source/Runtime/Renderer/Private/SceneRendering.h
--- a/Engine/Source/Runtime/Renderer/Private/SceneRendering.h
+++ b/Engine/Source/Runtime/Renderer/Private/SceneRendering.h
@@ -1758,9 +1758,6 @@ public:
 	/** Create a snapshot of this view info on the scene allocator. */
 	FViewInfo* CreateSnapshot() const;
 
-	/** Destroy all snapshots before we wipe the scene allocator. */
-	static void DestroyAllSnapshots();
-
 	void WaitForTasks();
 
 	// Get the range in DynamicMeshElements[] for a given PrimitiveIndex
