Scene render builder will now provide inputs for the first scene renderer in a series allowing for one-time updates for all renderers.
- Modified Nanite streaming and Scene only update once.
- Modified occlusion to use the 'multi-family' fencing code path all the time, and deprecated the flag that controls it.

diff --git a/Engine/Plugins/MovieScene/MovieRenderPipeline/Source/MovieRenderPipelineRenderPasses/Private/Graph/Renderers/MovieGraphImagePassBase.cpp b/Engine/Plugins/MovieScene/MovieRenderPipeline/Source/MovieRenderPipelineRenderPasses/Private/Graph/Renderers/MovieGraphImagePassBase.cpp
--- a/Engine/Plugins/MovieScene/MovieRenderPipeline/Source/MovieRenderPipelineRenderPasses/Private/Graph/Renderers/MovieGraphImagePassBase.cpp
+++ b/Engine/Plugins/MovieScene/MovieRenderPipeline/Source/MovieRenderPipelineRenderPasses/Private/Graph/Renderers/MovieGraphImagePassBase.cpp
@@ -206,8 +206,6 @@ void FMovieGraphImagePassBase::ApplyMovieGraphOverridesToViewFamily(TSharedRef<F
 			InOutFamily->bIsFirstViewInMultipleViewFamily = false;
 			InOutFamily->bAdditionalViewFamily = true;
 		}
-
-		InOutFamily->bIsMultipleViewFamily = true;
 	}
 	
 	// Skip the whole pass if they don't want motion blur.
diff --git a/Engine/Plugins/MovieScene/MovieRenderPipeline/Source/MovieRenderPipelineRenderPasses/Private/MoviePipelineImagePassBase.cpp b/Engine/Plugins/MovieScene/MovieRenderPipeline/Source/MovieRenderPipelineRenderPasses/Private/MoviePipelineImagePassBase.cpp
--- a/Engine/Plugins/MovieScene/MovieRenderPipeline/Source/MovieRenderPipelineRenderPasses/Private/MoviePipelineImagePassBase.cpp
+++ b/Engine/Plugins/MovieScene/MovieRenderPipeline/Source/MovieRenderPipelineRenderPasses/Private/MoviePipelineImagePassBase.cpp
@@ -203,14 +203,12 @@ TSharedPtr<FSceneViewFamilyContext> UMoviePipelineImagePassBase::CalculateViewFa
 		GetPipeline()->bHasRenderedFirstViewThisFrame = true;
 		
 		OutViewFamily->bIsFirstViewInMultipleViewFamily = true;
-		OutViewFamily->bIsMultipleViewFamily = true;
 		OutViewFamily->bAdditionalViewFamily = false;
 	}
 	else
 	{
 		OutViewFamily->bIsFirstViewInMultipleViewFamily = false;
 		OutViewFamily->bAdditionalViewFamily = true;
-		OutViewFamily->bIsMultipleViewFamily = true;
 	}
 
 	const bool bIsPerspective = true;
diff --git a/Engine/Plugins/Runtime/nDisplay/Source/DisplayCluster/Private/Game/EngineClasses/Basics/DisplayClusterViewportClient.cpp b/Engine/Plugins/Runtime/nDisplay/Source/DisplayCluster/Private/Game/EngineClasses/Basics/DisplayClusterViewportClient.cpp
--- a/Engine/Plugins/Runtime/nDisplay/Source/DisplayCluster/Private/Game/EngineClasses/Basics/DisplayClusterViewportClient.cpp
+++ b/Engine/Plugins/Runtime/nDisplay/Source/DisplayCluster/Private/Game/EngineClasses/Basics/DisplayClusterViewportClient.cpp
@@ -758,14 +758,12 @@ void UDisplayClusterViewportClient::Draw(FViewport* InViewport, FCanvas* SceneCa
 			// Initialize some flags for which view family is which, now that any view family reordering has been handled.
 			ViewFamilies[0]->bAdditionalViewFamily = false;
 			ViewFamilies[0]->bIsFirstViewInMultipleViewFamily = true;
-			ViewFamilies[0]->bIsMultipleViewFamily = true;
 
 			for (int32 FamilyIndex = 1; FamilyIndex < ViewFamilies.Num(); FamilyIndex++)
 			{
 				FSceneViewFamily& ViewFamily = *ViewFamilies[FamilyIndex];
 				ViewFamily.bAdditionalViewFamily = true;
 				ViewFamily.bIsFirstViewInMultipleViewFamily = false;
-				ViewFamily.bIsMultipleViewFamily = true;
 			}
 		}
 
diff --git a/Engine/Source/Runtime/Engine/Public/SceneRenderBuilderInterface.h b/Engine/Source/Runtime/Engine/Public/SceneRenderBuilderInterface.h
--- a/Engine/Source/Runtime/Engine/Public/SceneRenderBuilderInterface.h
+++ b/Engine/Source/Runtime/Engine/Public/SceneRenderBuilderInterface.h
@@ -13,6 +13,8 @@ class FSceneInterface;
 class FSceneViewFamily;
 class FHitProxyConsumer;
 
+struct FSceneRenderUpdateInputs;
+
 enum class ESceneRenderGroupFlags : uint8
 {
 	None       = 0,
@@ -28,18 +30,21 @@ ENUM_CLASS_FLAGS(ESceneRenderGroupFlags);
 // A setup of inputs passed into a scene render callback function.
 struct FSceneRenderFunctionInputs
 {
-	FSceneRenderFunctionInputs(FSceneRenderer* InRenderer, const TCHAR* InName, const TCHAR* InFullPath)
+	FSceneRenderFunctionInputs(FSceneRenderer* InRenderer, FSceneRenderUpdateInputs* InSceneUpdateInputs, const TCHAR* InName, const TCHAR* InFullPath)
 		: Renderer(InRenderer)
+		, SceneUpdateInputs(InSceneUpdateInputs)
 		, Name(InName)
 		, FullPath(InFullPath)
 	{}
 
 	FSceneRenderer* Renderer;
+	FSceneRenderUpdateInputs* SceneUpdateInputs;
 	const TCHAR* Name;
 	const TCHAR* FullPath;
 };
 
-using FSceneRenderFunction = TUniqueFunction<void(FRDGBuilder&, const FSceneRenderFunctionInputs&)>;
+// The user must return whether FSceneRenderer::Render was called inside of the function.
+using FSceneRenderFunction = TUniqueFunction<bool(FRDGBuilder&, const FSceneRenderFunctionInputs&)>;
 
 // A game side builder interface that collects scene renderers to process as one workload.
 class ISceneRenderBuilder
diff --git a/Engine/Source/Runtime/Engine/Public/SceneView.h b/Engine/Source/Runtime/Engine/Public/SceneView.h
--- a/Engine/Source/Runtime/Engine/Public/SceneView.h
+++ b/Engine/Source/Runtime/Engine/Public/SceneView.h
@@ -2282,7 +2282,7 @@ public:
 	bool bForceCopyCrossGPU = false;
 #endif
 
-	/** Whether this view is one of multiple view families rendered in a single frame.  Affects occlusion query synchronization logic. */
+	UE_DEPRECATED(5.6, "bIsMultipleViewFamily is no longer used")
 	bool bIsMultipleViewFamily = false;
 	
 	/** 
diff --git a/Engine/Source/Runtime/Renderer/Private/DeferredShadingRenderer.cpp b/Engine/Source/Runtime/Renderer/Private/DeferredShadingRenderer.cpp
--- a/Engine/Source/Runtime/Renderer/Private/DeferredShadingRenderer.cpp
+++ b/Engine/Source/Runtime/Renderer/Private/DeferredShadingRenderer.cpp
@@ -1408,7 +1408,7 @@ void FDeferredShadingSceneRenderer::RenderNanite(FRDGBuilder& GraphBuilder, cons
 	}
 }
 
-void FDeferredShadingSceneRenderer::Render(FRDGBuilder& GraphBuilder)
+void FDeferredShadingSceneRenderer::Render(FRDGBuilder& GraphBuilder, const FSceneRenderUpdateInputs* SceneUpdateInputs)
 {
 	if (!ViewFamily.EngineShowFlags.Rendering)
 	{
@@ -1430,7 +1430,7 @@ void FDeferredShadingSceneRenderer::Render(FRDGBuilder& GraphBuilder)
 	GPU_MESSAGE_SCOPE(GraphBuilder);
 
 #if RHI_RAYTRACING
-	if (RendererOutput == FSceneRenderer::ERendererOutput::FinalSceneColor)
+	if (SceneUpdateInputs && RendererOutput == FSceneRenderer::ERendererOutput::FinalSceneColor)
 	{
 		GRayTracingGeometryManager->PreRender();
 
@@ -1442,7 +1442,7 @@ void FDeferredShadingSceneRenderer::Render(FRDGBuilder& GraphBuilder)
 	}
 #endif
 
-	FInitViewTaskDatas InitViewTaskDatas = OnRenderBegin(GraphBuilder);
+	FInitViewTaskDatas InitViewTaskDatas = OnRenderBegin(GraphBuilder, SceneUpdateInputs);
 
 	FRDGExternalAccessQueue ExternalAccessQueue;
 	TUniquePtr<FVirtualTextureUpdater> VirtualTextureUpdater;
@@ -1465,8 +1465,11 @@ void FDeferredShadingSceneRenderer::Render(FRDGBuilder& GraphBuilder)
 	// Compute & commit the final state of the entire dependency topology of the renderer.
 	CommitFinalPipelineState();
 
-	// Initialize global system textures (pass-through if already initialized).
-	GSystemTextures.InitializeTextures(GraphBuilder.RHICmdList, FeatureLevel);
+	if (SceneUpdateInputs)
+	{
+		// Initialize global system textures (pass-through if already initialized).
+		GSystemTextures.InitializeTextures(GraphBuilder.RHICmdList, FeatureLevel);
+	}
 
 	UE::Tasks::TTask<void> UpdateLightFunctionAtlasTask;
 	if (LightFunctionAtlas.IsLightFunctionAtlasEnabled())
@@ -1596,14 +1599,11 @@ void FDeferredShadingSceneRenderer::Render(FRDGBuilder& GraphBuilder)
 	bool bVisualizeNanite = false;
 	if (bNaniteEnabled)
 	{
-		Nanite::GGlobalResources.Update(GraphBuilder);
-
-		// Only update Nanite streaming residency for the first view when multiple view rendering (nDisplay) is enabled.
-		// Streaming requests are still accumulated from the remaining views.
-		bUpdateNaniteStreaming =  !ViewFamily.bIsMultipleViewFamily || ViewFamily.bIsFirstViewInMultipleViewFamily;
-		if (bUpdateNaniteStreaming)
+		if (SceneUpdateInputs)
 		{
+			Nanite::GGlobalResources.Update(GraphBuilder);
 			Nanite::GStreamingManager.BeginAsyncUpdate(GraphBuilder);
+			bUpdateNaniteStreaming = true;
 		}
 
 		FNaniteVisualizationData& NaniteVisualization = GetNaniteVisualizationData();
diff --git a/Engine/Source/Runtime/Renderer/Private/DeferredShadingRenderer.h b/Engine/Source/Runtime/Renderer/Private/DeferredShadingRenderer.h
--- a/Engine/Source/Runtime/Renderer/Private/DeferredShadingRenderer.h
+++ b/Engine/Source/Runtime/Renderer/Private/DeferredShadingRenderer.h
@@ -426,10 +426,10 @@ public:
 	bool RenderHzb(FRDGBuilder& GraphBuilder, FRDGTextureRef SceneDepthTexture, const FBuildHZBAsyncComputeParams* AsyncComputeParams, Froxel::FRenderer& FroxelRenderer);
 
 	/** Renders the view family. */
-	virtual void Render(FRDGBuilder& GraphBuilder) override;
+	virtual void Render(FRDGBuilder& GraphBuilder, const FSceneRenderUpdateInputs* SceneUpdateInputs) override;
 
 	/** Render the view family's hit proxies. */
-	virtual void RenderHitProxies(FRDGBuilder& GraphBuilder) override;
+	virtual void RenderHitProxies(FRDGBuilder& GraphBuilder, const FSceneRenderUpdateInputs* SceneUpdateInputs) override;
 
 	virtual bool ShouldRenderVelocities() const override;
 
diff --git a/Engine/Source/Runtime/Renderer/Private/MeshEdgesRendering.cpp b/Engine/Source/Runtime/Renderer/Private/MeshEdgesRendering.cpp
--- a/Engine/Source/Runtime/Renderer/Private/MeshEdgesRendering.cpp
+++ b/Engine/Source/Runtime/Renderer/Private/MeshEdgesRendering.cpp
@@ -363,12 +363,14 @@ void RenderMeshEdges(FSceneViewFamily& ViewFamily)
 		ViewFamilyData->WireframeColor->InitResource(GraphBuilder.RHICmdList);
 		ViewFamilyData->WireframeDepth->InitResource(GraphBuilder.RHICmdList);
 
-		Inputs.Renderer->Render(GraphBuilder);
+		Inputs.Renderer->Render(GraphBuilder, Inputs.SceneUpdateInputs);
 
 		for (const FViewInfo& ViewInfo : Inputs.Renderer->Views)
 		{
 			ViewFamilyData->ViewRects.Emplace(ViewInfo.ViewRect);
 		}
+
+		return true;
 	});
 }
 
diff --git a/Engine/Source/Runtime/Renderer/Private/MobileShadingRenderer.cpp b/Engine/Source/Runtime/Renderer/Private/MobileShadingRenderer.cpp
--- a/Engine/Source/Runtime/Renderer/Private/MobileShadingRenderer.cpp
+++ b/Engine/Source/Runtime/Renderer/Private/MobileShadingRenderer.cpp
@@ -975,7 +975,7 @@ void FMobileSceneRenderer::RenderCustomRenderPassBasePass(FRDGBuilder& GraphBuil
 	}
 }
 
-void FMobileSceneRenderer::Render(FRDGBuilder& GraphBuilder)
+void FMobileSceneRenderer::Render(FRDGBuilder& GraphBuilder, const FSceneRenderUpdateInputs* SceneUpdateInputs)
 {
 	if (!ViewFamily.EngineShowFlags.Rendering)
 	{
@@ -988,7 +988,7 @@ void FMobileSceneRenderer::Render(FRDGBuilder& GraphBuilder)
 	RDG_RHI_EVENT_SCOPE(GraphBuilder, MobileSceneRender);
 	RDG_RHI_GPU_STAT_SCOPE(GraphBuilder, MobileSceneRender);
 
-	IVisibilityTaskData* VisibilityTaskData = OnRenderBegin(GraphBuilder);
+	IVisibilityTaskData* VisibilityTaskData = OnRenderBegin(GraphBuilder, SceneUpdateInputs);
 
 	FRDGExternalAccessQueue ExternalAccessQueue;
 
diff --git a/Engine/Source/Runtime/Renderer/Private/PlanarReflectionRendering.cpp b/Engine/Source/Runtime/Renderer/Private/PlanarReflectionRendering.cpp
--- a/Engine/Source/Runtime/Renderer/Private/PlanarReflectionRendering.cpp
+++ b/Engine/Source/Runtime/Renderer/Private/PlanarReflectionRendering.cpp
@@ -276,9 +276,10 @@ struct FUpdatePlanarReflectionContentsInputs
 	FPlane MirrorPlane;
 };
 
-static void UpdatePlanarReflectionContents_RenderThread(
+static bool UpdatePlanarReflectionContents_RenderThread(
 	FRDGBuilder& GraphBuilder,
 	FSceneRenderer* SceneRenderer,
+	const FSceneRenderUpdateInputs* SceneUpdateInputs,
 	const FUpdatePlanarReflectionContentsInputs& Inputs)
 {
 	QUICK_SCOPE_CYCLE_COUNTER(STAT_RenderPlanarReflection);
@@ -302,7 +303,7 @@ static void UpdatePlanarReflectionContents_RenderThread(
 
 		if (!bIsInAnyFrustum)
 		{
-			return;
+			return false;
 		}
 
 		bool bIsVisibleInAnyView = true;
@@ -340,7 +341,7 @@ static void UpdatePlanarReflectionContents_RenderThread(
 
 		if (!bIsVisibleInAnyView)
 		{
-			return;
+			return false;
 		}
 	}
 
@@ -393,7 +394,7 @@ static void UpdatePlanarReflectionContents_RenderThread(
 		}
 	}
 
-	SceneRenderer->Render(GraphBuilder);
+	SceneRenderer->Render(GraphBuilder, SceneUpdateInputs);
 
 	FRDGTextureRef ReflectionOutputTexture = GraphBuilder.RegisterExternalTexture(CreateRenderTarget(Inputs.RenderTarget->TextureRHI, TEXT("ReflectionOutputTexture")));
 	GraphBuilder.SetTextureAccessFinal(ReflectionOutputTexture, ERHIAccess::SRVGraphics);
@@ -413,6 +414,8 @@ static void UpdatePlanarReflectionContents_RenderThread(
 			PrefilterPlanarReflection<false>(GraphBuilder, View, SceneTextureParameters, Inputs.SceneProxy, SceneTextures.Color.Resolve, ReflectionOutputTexture);
 		}
 	}
+
+	return true;
 }
 
 extern void BuildProjectionMatrix(FIntPoint RenderTargetSize, float FOV, float InNearClippingPlane, FMatrix& ProjectionMatrix);
@@ -620,14 +623,14 @@ void FScene::UpdatePlanarReflectionContents(UPlanarReflectionComponent* CaptureC
 	{
 		// Remove the render target on the planar reflection proxy so that this planar reflection is not getting drawn in its own FSceneRenderer.
 		UpdateInputs.SceneProxy->RenderTarget = nullptr;
-	
+
 		for (int32 ViewIndex = 0; ViewIndex < SceneCaptureViewInfo.Num(); ++ViewIndex)
 		{
 			UpdateInputs.SceneProxy->ProjectionWithExtraFOV[ViewIndex] = SceneCaptureViewInfo[ViewIndex].ProjectionMatrix;
 			UpdateInputs.SceneProxy->bIsStereo = bIsStereo;
 		}
-	
-		UpdatePlanarReflectionContents_RenderThread(GraphBuilder, Inputs.Renderer, UpdateInputs);
+
+		return UpdatePlanarReflectionContents_RenderThread(GraphBuilder, Inputs.Renderer, Inputs.SceneUpdateInputs, UpdateInputs);
 	});
 }
 
diff --git a/Engine/Source/Runtime/Renderer/Private/ReflectionEnvironmentCapture.cpp b/Engine/Source/Runtime/Renderer/Private/ReflectionEnvironmentCapture.cpp
--- a/Engine/Source/Runtime/Renderer/Private/ReflectionEnvironmentCapture.cpp
+++ b/Engine/Source/Runtime/Renderer/Private/ReflectionEnvironmentCapture.cpp
@@ -505,6 +505,7 @@ int32 FindOrAllocateCubemapIndex(FScene* Scene, const UReflectionCaptureComponen
 void CaptureSceneToScratchCubemap(
 	FRDGBuilder& GraphBuilder,
 	FSceneRenderer* SceneRenderer,
+	FSceneRenderUpdateInputs* SceneUpdateInputs,
 	const FReflectionCubemapTexture& ReflectionCubemapTexture,
 	ECubeFace CubeFace,
 	int32 CubemapSize,
@@ -516,7 +517,7 @@ void CaptureSceneToScratchCubemap(
 	const ERHIFeatureLevel::Type FeatureLevel = SceneRenderer->FeatureLevel;
 
 	// Render the scene normally for one face of the cubemap
-	SceneRenderer->Render(GraphBuilder);
+	SceneRenderer->Render(GraphBuilder, SceneUpdateInputs);
 
 	AddPass(GraphBuilder, RDG_EVENT_NAME("FlushGPU"), [](FRHICommandListImmediate& InRHICmdList)
 	{
@@ -1367,7 +1368,8 @@ void CaptureSceneIntoScratchCubemap(
 				, bInsideTick
 			] (FRDGBuilder& GraphBuilder, const FSceneRenderFunctionInputs& Inputs)
 		{
-			CaptureSceneToScratchCubemap(GraphBuilder, Inputs.Renderer, ReflectionCubemapTexture, (ECubeFace)CubeFace, CubemapSize, bCapturingForSkyLight, bLowerHemisphereIsBlack, LowerHemisphereColor, bCapturingForMobile);
+			CaptureSceneToScratchCubemap(GraphBuilder, Inputs.Renderer, Inputs.SceneUpdateInputs, ReflectionCubemapTexture, (ECubeFace)CubeFace, CubemapSize, bCapturingForSkyLight, bLowerHemisphereIsBlack, LowerHemisphereColor, bCapturingForMobile);
+			return true;
 		});
 	}
 }
diff --git a/Engine/Source/Runtime/Renderer/Private/RendererScene.cpp b/Engine/Source/Runtime/Renderer/Private/RendererScene.cpp
--- a/Engine/Source/Runtime/Renderer/Private/RendererScene.cpp
+++ b/Engine/Source/Runtime/Renderer/Private/RendererScene.cpp
@@ -5512,7 +5512,7 @@ private:
 	FScene* Scene = nullptr;
 };
 
-FViewSceneChangeSet* FScene::ProcessViewChanges(FRDGBuilder& GraphBuilder, TArray<FViewInfo*> &Views)
+FViewSceneChangeSet* FScene::ProcessViewChanges(FRDGBuilder& GraphBuilder, TConstArrayView<FViewInfo*> Views)
 {
 	FViewSceneChangeSet* ChangeSet = GraphBuilder.AllocObject<FViewSceneChangeSet>();
 	ChangeSet->AddedViewIds.Reserve(Views.Num());
diff --git a/Engine/Source/Runtime/Renderer/Private/SceneCaptureRendering.cpp b/Engine/Source/Runtime/Renderer/Private/SceneCaptureRendering.cpp
--- a/Engine/Source/Runtime/Renderer/Private/SceneCaptureRendering.cpp
+++ b/Engine/Source/Runtime/Renderer/Private/SceneCaptureRendering.cpp
@@ -412,6 +412,7 @@ void CopySceneCaptureComponentToTarget(
 static void UpdateSceneCaptureContent_RenderThread(
 	FRDGBuilder& GraphBuilder,
 	FSceneRenderer* SceneRenderer,
+	const FSceneRenderUpdateInputs* SceneUpdateInputs,
 	FRenderTarget* RenderTarget,
 	FTexture* RenderTargetTexture,
 	TConstArrayView<FRHICopyTextureInfo> CopyInfos,
@@ -467,7 +468,7 @@ static void UpdateSceneCaptureContent_RenderThread(
 		View.bIgnoreExistingQueries = true;
 	}
 
-	SceneRenderer->Render(GraphBuilder);
+	SceneRenderer->Render(GraphBuilder, SceneUpdateInputs);
 
 	if (ShadingPath == EShadingPath::Mobile)
 	{
@@ -1382,6 +1383,7 @@ void FScene::UpdateSceneCaptureContents(USceneCaptureComponent2D* CaptureCompone
 			UpdateSceneCaptureContent_RenderThread(
 				GraphBuilder,
 				Inputs.Renderer,
+				Inputs.SceneUpdateInputs,
 				TextureRenderTargetResource,
 				TextureRenderTargetResource,
 				MakeArrayView({ CopyInfo }),
@@ -1408,6 +1410,8 @@ void FScene::UpdateSceneCaptureContents(USceneCaptureComponent2D* CaptureCompone
 				}
 			}
 		#endif  // WITH_EDITOR
+
+			return true;
 		});
 	}
 }
@@ -1579,6 +1583,7 @@ void FScene::UpdateSceneCaptureContents(USceneCaptureComponentCube* CaptureCompo
 					UpdateSceneCaptureContent_RenderThread(
 						GraphBuilder,
 						Inputs.Renderer,
+						Inputs.SceneUpdateInputs,
 						TextureRenderTarget,
 						TextureRenderTarget,
 						MakeArrayView({ CopyInfo }),
@@ -1598,6 +1603,8 @@ void FScene::UpdateSceneCaptureContents(USceneCaptureComponentCube* CaptureCompo
 						}
 					}
 #endif  // WITH_EDITOR
+
+					return true;
 				});
 			}
 		}
@@ -1756,6 +1763,7 @@ void FScene::UpdateSceneCaptureContents(USceneCaptureComponentCube* CaptureCompo
 				UpdateSceneCaptureContent_RenderThread(
 					GraphBuilder,
 					Inputs.Renderer,
+					Inputs.SceneUpdateInputs,
 					CubeFaceTarget.Get(),
 					TextureRenderTarget,
 					CopyInfos,
@@ -1778,6 +1786,8 @@ void FScene::UpdateSceneCaptureContents(USceneCaptureComponentCube* CaptureCompo
 					}
 				}
 		#endif  // WITH_EDITOR
+
+				return true;
 			});
 		}
 	}
diff --git a/Engine/Source/Runtime/Renderer/Private/SceneHitProxyRendering.cpp b/Engine/Source/Runtime/Renderer/Private/SceneHitProxyRendering.cpp
--- a/Engine/Source/Runtime/Renderer/Private/SceneHitProxyRendering.cpp
+++ b/Engine/Source/Runtime/Renderer/Private/SceneHitProxyRendering.cpp
@@ -529,9 +529,9 @@ static void DoRenderHitProxies(
 }
 #endif
 
-void FMobileSceneRenderer::RenderHitProxies(FRDGBuilder& GraphBuilder)
+void FMobileSceneRenderer::RenderHitProxies(FRDGBuilder& GraphBuilder, const FSceneRenderUpdateInputs* SceneUpdateInputs)
 {
-	IVisibilityTaskData* VisibilityTaskData = OnRenderBegin(GraphBuilder);
+	IVisibilityTaskData* VisibilityTaskData = OnRenderBegin(GraphBuilder, SceneUpdateInputs);
 
 	GPU_MESSAGE_SCOPE(GraphBuilder);
 
@@ -586,11 +586,11 @@ void FMobileSceneRenderer::RenderHitProxies(FRDGBuilder& GraphBuilder)
 	OnRenderFinish(GraphBuilder, nullptr);
 }
 
-void FDeferredShadingSceneRenderer::RenderHitProxies(FRDGBuilder& GraphBuilder)
+void FDeferredShadingSceneRenderer::RenderHitProxies(FRDGBuilder& GraphBuilder, const FSceneRenderUpdateInputs* SceneUpdateInputs)
 {
 	const bool bNaniteEnabled = UseNanite(ShaderPlatform);
 
-	IVisibilityTaskData* VisibilityTaskData = OnRenderBegin(GraphBuilder);
+	IVisibilityTaskData* VisibilityTaskData = OnRenderBegin(GraphBuilder, SceneUpdateInputs);
 
 	CommitFinalPipelineState();
 
diff --git a/Engine/Source/Runtime/Renderer/Private/SceneOcclusion.cpp b/Engine/Source/Runtime/Renderer/Private/SceneOcclusion.cpp
--- a/Engine/Source/Runtime/Renderer/Private/SceneOcclusion.cpp
+++ b/Engine/Source/Runtime/Renderer/Private/SceneOcclusion.cpp
@@ -1521,29 +1521,15 @@ void FSceneRenderer::FenceOcclusionTests(FRDGBuilder& GraphBuilder)
 		{
 			SCOPE_CYCLE_COUNTER(STAT_OcclusionSubmittedFence_Dispatch);
 
-			if (ViewFamily.bIsMultipleViewFamily)
-			{
-				// If there are multiple view families, we implement a queue of buffered fences, so we can avoid waiting on queries for
-				// another view family that's rendering in the same frame.  Here we push a new fence into the queue of buffered fences.
-				// We assume that the queue isn't full, since WaitOcclusionTests (called earlier in the frame) will always pop at least
-				// one fence if the queue is full.
-				check(OcclusionSubmittedFence[FOcclusionQueryHelpers::MaxBufferedOcclusionFrames - 1].Fence == nullptr);
+			// We implement a queue of buffered fences, so we can avoid waiting on queries for another view family that's rendering
+			// in the same frame.  Here we push a new fence into the queue of buffered fences. We assume that the queue isn't full,
+			// since WaitOcclusionTests (called earlier in the frame) will always pop at least one fence if the queue is full.
+			check(OcclusionSubmittedFence[FOcclusionQueryHelpers::MaxBufferedOcclusionFrames - 1].Fence == nullptr);
 
-				for (int32 Dest = FOcclusionQueryHelpers::MaxBufferedOcclusionFrames - 1; Dest >= 1; Dest--)
-				{
-					CA_SUPPRESS(6385);
-					OcclusionSubmittedFence[Dest] = OcclusionSubmittedFence[Dest - 1];
-				}
-			}
-			else
+			for (int32 Dest = FOcclusionQueryHelpers::MaxBufferedOcclusionFrames - 1; Dest >= 1; Dest--)
 			{
-				// Single view family implementation, fixed number of buffered frames.
-				int32 NumFrames = FOcclusionQueryHelpers::GetNumBufferedFrames(FeatureLevel);
-				for (int32 Dest = NumFrames - 1; Dest >= 1; Dest--)
-				{
-					CA_SUPPRESS(6385);
-					OcclusionSubmittedFence[Dest] = OcclusionSubmittedFence[Dest - 1];
-				}
+				CA_SUPPRESS(6385);
+				OcclusionSubmittedFence[Dest] = OcclusionSubmittedFence[Dest - 1];
 			}
 
 			OcclusionSubmittedFence[0].Fence = RHICmdList.RHIThreadFence();
@@ -1560,60 +1546,49 @@ void FSceneRenderer::WaitOcclusionTests(FRHICommandListImmediate& RHICmdList)
 	{
 		SCOPE_CYCLE_COUNTER(STAT_OcclusionSubmittedFence_Wait);
 
-		if (ViewFamily.bIsMultipleViewFamily)
-		{
-			// Count the number of fences in the queue for the view in question
-			const uint32 ViewStateUniqueID = GetViewStateUniqueID(this);
+		// Count the number of fences in the queue for the view in question
+		const uint32 ViewStateUniqueID = GetViewStateUniqueID(this);
 
-			int32 ViewStateFenceCount = 0;
-			for (int32 FenceIndex = 0; FenceIndex < FOcclusionQueryHelpers::MaxBufferedOcclusionFrames && OcclusionSubmittedFence[FenceIndex].Fence; FenceIndex++)
+		int32 ViewStateFenceCount = 0;
+		for (int32 FenceIndex = 0; FenceIndex < FOcclusionQueryHelpers::MaxBufferedOcclusionFrames && OcclusionSubmittedFence[FenceIndex].Fence; FenceIndex++)
+		{
+			if (OcclusionSubmittedFence[FenceIndex].Fence && (OcclusionSubmittedFence[FenceIndex].ViewStateUniqueID == ViewStateUniqueID))
 			{
-				if (OcclusionSubmittedFence[FenceIndex].Fence && (OcclusionSubmittedFence[FenceIndex].ViewStateUniqueID == ViewStateUniqueID))
-				{
-					ViewStateFenceCount++;
-				}
+				ViewStateFenceCount++;
 			}
+		}
 
-			// Figure out how many fences are allowed to remain in the queue after the wait, which is one less than the number of buffered frames
-			const int32 FencesAllowedInQueue = FOcclusionQueryHelpers::GetNumBufferedFrames(FeatureLevel) - 1;
-			check(FencesAllowedInQueue >= 0);
+		// Figure out how many fences are allowed to remain in the queue after the wait, which is one less than the number of buffered frames
+		const int32 FencesAllowedInQueue = FOcclusionQueryHelpers::GetNumBufferedFrames(FeatureLevel) - 1;
+		check(FencesAllowedInQueue >= 0);
 
-			// Scan the fences in the queue in order (oldest will be at the end), and wait on them as appropriate.  If the queue is full,
-			// we always want to wait on at least one fence, so we have room to push a new fence onto the queue if necessary.  Thus the
-			// loop always runs at least once before breaking out.
-			int32 FenceToWaitOn = FOcclusionQueryHelpers::MaxBufferedOcclusionFrames - 1;
-			do
-			{
-				// If our fence counting logic is correct, we should already have broken out of the loop
-				check(FenceToWaitOn >= 0);
+		// Scan the fences in the queue in order (oldest will be at the end), and wait on them as appropriate.  If the queue is full,
+		// we always want to wait on at least one fence, so we have room to push a new fence onto the queue if necessary.  Thus the
+		// loop always runs at least once before breaking out.
+		int32 FenceToWaitOn = FOcclusionQueryHelpers::MaxBufferedOcclusionFrames - 1;
+		do
+		{
+			// If our fence counting logic is correct, we should already have broken out of the loop
+			check(FenceToWaitOn >= 0);
 
-				// Is there a fence in this queue entry?
-				if (OcclusionSubmittedFence[FenceToWaitOn].Fence)
+			// Is there a fence in this queue entry?
+			if (OcclusionSubmittedFence[FenceToWaitOn].Fence)
+			{
+				// If this is one of the fences for the current scene's view state, mark that we're waiting on it
+				if (OcclusionSubmittedFence[FenceToWaitOn].ViewStateUniqueID == ViewStateUniqueID)
 				{
-					// If this is one of the fences for the current scene's view state, mark that we're waiting on it
-					if (OcclusionSubmittedFence[FenceToWaitOn].ViewStateUniqueID == ViewStateUniqueID)
-					{
-						ViewStateFenceCount--;
-					}
-
-					// Wait on the current fence and clear it out
-					FRHICommandListExecutor::WaitOnRHIThreadFence(OcclusionSubmittedFence[FenceToWaitOn].Fence);
-					OcclusionSubmittedFence[FenceToWaitOn].Fence = nullptr;
-					OcclusionSubmittedFence[FenceToWaitOn].ViewStateUniqueID = 0;
+					ViewStateFenceCount--;
 				}
 
-				FenceToWaitOn--;
+				// Wait on the current fence and clear it out
+				FRHICommandListExecutor::WaitOnRHIThreadFence(OcclusionSubmittedFence[FenceToWaitOn].Fence);
+				OcclusionSubmittedFence[FenceToWaitOn].Fence = nullptr;
+				OcclusionSubmittedFence[FenceToWaitOn].ViewStateUniqueID = 0;
+			}
 
-			} while (ViewStateFenceCount > FencesAllowedInQueue);
-		}
-		else
-		{
-			// Single view family implementation, fixed number of buffered frames.
-			int32 BlockFrame = FOcclusionQueryHelpers::GetNumBufferedFrames(FeatureLevel) - 1;
-			FRHICommandListExecutor::WaitOnRHIThreadFence(OcclusionSubmittedFence[BlockFrame].Fence);
-			OcclusionSubmittedFence[BlockFrame].Fence = nullptr;
-			OcclusionSubmittedFence[BlockFrame].ViewStateUniqueID = 0;
-		}
+			FenceToWaitOn--;
+
+		} while (ViewStateFenceCount > FencesAllowedInQueue);
 	}
 }
 
diff --git a/Engine/Source/Runtime/Renderer/Private/ScenePrivate.h b/Engine/Source/Runtime/Renderer/Private/ScenePrivate.h
--- a/Engine/Source/Runtime/Renderer/Private/ScenePrivate.h
+++ b/Engine/Source/Runtime/Renderer/Private/ScenePrivate.h
@@ -3365,7 +3365,7 @@ public:
 	/**
 	 * Determines added & removed view states based on the rendered views passed in & the RemovedViews.
 	 */
-	FViewSceneChangeSet* ProcessViewChanges(FRDGBuilder& GraphBuilder, TArray<FViewInfo*>& AllViews);
+	FViewSceneChangeSet* ProcessViewChanges(FRDGBuilder& GraphBuilder, TConstArrayView<FViewInfo*> AllViews);
 
 	void Update(FRDGBuilder& GraphBuilder, const FUpdateParameters& Parameters);
 
diff --git a/Engine/Source/Runtime/Renderer/Private/SceneRenderBuilder.cpp b/Engine/Source/Runtime/Renderer/Private/SceneRenderBuilder.cpp
--- a/Engine/Source/Runtime/Renderer/Private/SceneRenderBuilder.cpp
+++ b/Engine/Source/Runtime/Renderer/Private/SceneRenderBuilder.cpp
@@ -487,6 +487,9 @@ public:
 		check(Renderer->Scene == Scene);
 		check(Renderer->ViewFamily.Views.Num() > 0);
 		check(Renderer->ViewFamily.Views[0]);
+		checkf(Renderers.IsEmpty() || Renderer->ViewFamily.EngineShowFlags.HitProxies == bHitProxyRendering, TEXT("You cannot mix hit proxy renderers with non-hit proxy renderers."));
+
+		bHitProxyRendering = Renderer->ViewFamily.EngineShowFlags.HitProxies;
 
 	#if DO_CHECK
 		bool bFoundRenderer = false;
@@ -518,6 +521,17 @@ public:
 
 		Ops.Emplace(FOp::Render(RenderNode));
 		Renderers.Emplace(Renderer);
+
+		if (Renderer->ViewFamily.EngineShowFlags.Rendering)
+		{
+			ActiveRenderers.Emplace(Renderer);
+			ActiveViewFamilies.Emplace(&Renderer->ViewFamily);
+			ActiveViews.Reserve(ActiveViews.Num() + Renderer->Views.Num());
+			for (FViewInfo& View : Renderer->Views)
+			{
+				ActiveViews.Emplace(&View);
+			}
+		}
 	}
 
 	void BeginGroup(FString&& Name, ESceneRenderGroupFlags Flags)
@@ -566,6 +580,9 @@ private:
 
 	FScene* Scene;
 	TArray<FSceneRenderer*, FConcurrentLinearArrayAllocator> Renderers;
+	TArray<FSceneRenderer*, FConcurrentLinearArrayAllocator> ActiveRenderers;
+	TArray<FViewFamilyInfo*, FConcurrentLinearArrayAllocator> ActiveViewFamilies;
+	TArray<FViewInfo*, FConcurrentLinearArrayAllocator> ActiveViews;
 	TIntrusiveDoubleLinkedList<FRenderNode> RenderNodes;
 	TIntrusiveDoubleLinkedList<FGroupNode> GroupNodes;
 	TArray<FOp, FConcurrentLinearArrayAllocator> Ops;
@@ -580,10 +597,12 @@ private:
 		FGroupNode* Group = nullptr;
 		FGroupEvent GroupEvent;
 		FString FullPath;
+		bool bSceneUpdateConsumed = false;
 
 	} RenderState;
 
 	bool bInsideGroup : 1 = false;
+	bool bHitProxyRendering : 1 = false;
 };
 
 FSceneRenderProcessor::FAsyncTasks FSceneRenderProcessor::AsyncTasks;
@@ -779,7 +798,7 @@ void FSceneRenderProcessor::Execute()
 	TOptional<RenderCaptureInterface::FScopedCapture> GpuCaptureScope;
 
 	// This has to be called first because it can actually recursively invoke scene renders.
-	for (FSceneRenderer* Renderer : Renderers)
+	for (FSceneRenderer* Renderer : ActiveRenderers)
 	{
 		for (const FSceneViewExtensionRef& Extension : Renderer->ViewFamily.ViewExtensions)
 		{
@@ -857,16 +876,41 @@ void FSceneRenderProcessor::Execute()
 					RenderState.FullPath += RenderNode.Name;
 				}
 
-				const FSceneRenderFunctionInputs FunctionInputs(&Renderer, *RenderNode.Name, *RenderState.FullPath);
+				const bool bFirstRenderer = RenderNode.Renderer == ActiveRenderers[0];
+				const bool bLastRenderer  = RenderNode.Renderer == ActiveRenderers.Last();
+
+				TOptional<FSceneRenderUpdateInputs> SceneUpdateInputs;
+
+				if (Renderer.ViewFamily.EngineShowFlags.Rendering && !RenderState.bSceneUpdateConsumed)
+				{
+					SceneUpdateInputs.Emplace();
+					SceneUpdateInputs->Renderers = ActiveRenderers;
+					SceneUpdateInputs->ViewFamilies = ActiveViewFamilies;
+					SceneUpdateInputs->Views = ActiveViews;
+				}
+
+				const FSceneRenderFunctionInputs FunctionInputs(&Renderer, SceneUpdateInputs ? &SceneUpdateInputs.GetValue() : nullptr, *RenderNode.Name, *RenderState.FullPath);
 
 				FRDGBuilder GraphBuilder(RHICmdList, RDG_EVENT_NAME("%s", FunctionInputs.FullPath), ERDGBuilderFlags::Parallel, Scene->GetShaderPlatform());
 
+			#if WITH_MGPU
+				if (Renderer.ViewFamily.bForceCopyCrossGPU)
+				{
+					GraphBuilder.EnableForceCopyCrossGPU();
+				}
+			#endif
+
 				if (!Renderer.ViewFamily.EngineShowFlags.HitProxies)
 				{
 					VISUALIZE_TEXTURE_BEGIN_VIEW(Scene->GetFeatureLevel(), Renderer.Views[0].GetViewKey(), FunctionInputs.FullPath, Renderer.Views[0].bIsSceneCapture);
 				}
+	
+				const bool bRenderCalled = RenderNode.Function(GraphBuilder, FunctionInputs);
 
-				RenderNode.Function(GraphBuilder, FunctionInputs);
+				if (SceneUpdateInputs)
+				{
+					RenderState.bSceneUpdateConsumed |= bRenderCalled;
+				}
 
 				if (!Renderer.ViewFamily.EngineShowFlags.HitProxies)
 				{
@@ -879,7 +923,7 @@ void FSceneRenderProcessor::Execute()
 				}
 
 				// The final graph builder is responsible for flushing resources.
-				if (&Renderer == Renderers.Last())
+				if (RenderNode.Renderer == Renderers.Last())
 				{
 					GraphBuilder.SetFlushResourcesRHI();
 				}
diff --git a/Engine/Source/Runtime/Renderer/Private/SceneRendering.cpp b/Engine/Source/Runtime/Renderer/Private/SceneRendering.cpp
--- a/Engine/Source/Runtime/Renderer/Private/SceneRendering.cpp
+++ b/Engine/Source/Runtime/Renderer/Private/SceneRendering.cpp
@@ -3814,7 +3814,7 @@ FSceneRenderer::~FSceneRenderer()
 	}
 }
 
-IVisibilityTaskData* FSceneRenderer::OnRenderBegin(FRDGBuilder& GraphBuilder)
+IVisibilityTaskData* FSceneRenderer::OnRenderBegin(FRDGBuilder& GraphBuilder, const FSceneRenderUpdateInputs* SceneUpdateInputs)
 {
 	check(!FDeferredUpdateResource::IsUpdateNeeded());
 
@@ -3858,17 +3858,16 @@ IVisibilityTaskData* FSceneRenderer::OnRenderBegin(FRDGBuilder& GraphBuilder)
 
 		if (!ViewFamily.ViewExtensions.IsEmpty())
 		{
+			RDG_CSV_STAT_EXCLUSIVE_SCOPE(GraphBuilder, PreRender);
+			SCOPE_CYCLE_COUNTER(STAT_FDeferredShadingSceneRenderer_ViewExtensionPreRenderView);
+
+			for (auto& ViewExtension : ViewFamily.ViewExtensions)
 			{
-				RDG_CSV_STAT_EXCLUSIVE_SCOPE(GraphBuilder, PreRender);
-				SCOPE_CYCLE_COUNTER(STAT_FDeferredShadingSceneRenderer_ViewExtensionPreRenderView);
-	
-				for (int32 ViewExt = 0; ViewExt < ViewFamily.ViewExtensions.Num(); ViewExt++)
+				ViewExtension->PreRenderViewFamily_RenderThread(GraphBuilder, ViewFamily);
+
+				for (FViewInfo* View : AllViews)
 				{
-					ViewFamily.ViewExtensions[ViewExt]->PreRenderViewFamily_RenderThread(GraphBuilder, ViewFamily);
-					for (int32 ViewIndex = 0; ViewIndex < AllViews.Num(); ViewIndex++)
-					{
-						ViewFamily.ViewExtensions[ViewExt]->PreRenderView_RenderThread(GraphBuilder, *AllViews[ViewIndex]);
-					}
+					ViewExtension->PreRenderView_RenderThread(GraphBuilder, *View);
 				}
 			}
 		}
@@ -3933,10 +3932,18 @@ IVisibilityTaskData* FSceneRenderer::OnRenderBegin(FRDGBuilder& GraphBuilder)
 		GPUSceneUpdateTaskPrerequisites.Trigger();
 	};
 
-	// Note: in the future persistent views should be added/removed as other scene primitives such that the updates are deferred and so on. 
-	//       right now there's no explicit mechanism for this, so we discover added views here & pass the chage set to the scene update.
-	SceneUpdateParameters.ViewUpdateChangeSet = Scene->ProcessViewChanges(GraphBuilder, AllViews);
-	Scene->Update(GraphBuilder, SceneUpdateParameters);
+	if (SceneUpdateInputs)
+	{
+		// Note: in the future persistent views should be added/removed as other scene primitives such that the updates are deferred and so on. 
+		//       right now there's no explicit mechanism for this, so we discover added views here & pass the change set to the scene update.
+
+		SceneUpdateParameters.ViewUpdateChangeSet = Scene->ProcessViewChanges(GraphBuilder, SceneUpdateInputs->Views);
+		Scene->Update(GraphBuilder, SceneUpdateParameters);
+	}
+	else
+	{
+		SceneUpdateParameters.Callbacks.PostStaticMeshUpdate(UE::Tasks::FTask{});
+	}
 
 	return VisibilityTaskData;
 }
@@ -4711,7 +4718,7 @@ void ResetAndShrinkModifiedBounds(TArray<FBox>& Bounds)
 	Bounds.Reset();
 }
 
-static void RenderViewFamily_RenderThread(FRDGBuilder& GraphBuilder, FSceneRenderer* Renderer)
+static void RenderViewFamily_RenderThread(FRDGBuilder& GraphBuilder, FSceneRenderer* Renderer, const FSceneRenderUpdateInputs* SceneUpdateInputs)
 {
 	FSceneViewFamily& ViewFamily = Renderer->ViewFamily;
 
@@ -4719,20 +4726,13 @@ static void RenderViewFamily_RenderThread(FRDGBuilder& GraphBuilder, FSceneRende
 	SCOPE_CYCLE_COUNTER(STAT_TotalSceneRenderingTime);
 	SCOPED_NAMED_EVENT_TCHAR_CONDITIONAL(*ViewFamily.ProfileDescription, FColor::Red, !ViewFamily.ProfileDescription.IsEmpty());
 
-#if WITH_MGPU
-	if (ViewFamily.bForceCopyCrossGPU)
-	{
-		GraphBuilder.EnableForceCopyCrossGPU();
-	}
-#endif
-
 	if (ViewFamily.EngineShowFlags.HitProxies)
 	{
-		Renderer->RenderHitProxies(GraphBuilder);
+		Renderer->RenderHitProxies(GraphBuilder, SceneUpdateInputs);
 	}
 	else
 	{
-		Renderer->Render(GraphBuilder);
+		Renderer->Render(GraphBuilder, SceneUpdateInputs);
 	}
 
 	Renderer->FlushCrossGPUFences(GraphBuilder);
@@ -5017,7 +5017,8 @@ void FRendererModule::BeginRenderingViewFamilies(FCanvas* Canvas, TConstArrayVie
 			SceneRenderBuilder.AddRenderer(SceneRenderer, bShowHitProxies ? TEXT("HitProxies") : TEXT("ViewFamilies"),
 				[] (FRDGBuilder& GraphBuilder, const FSceneRenderFunctionInputs& Inputs)
 			{
-				RenderViewFamily_RenderThread(GraphBuilder, Inputs.Renderer);
+				RenderViewFamily_RenderThread(GraphBuilder, Inputs.Renderer, Inputs.SceneUpdateInputs);
+				return true;
 			});
 		}
 
diff --git a/Engine/Source/Runtime/Renderer/Private/SceneRendering.h b/Engine/Source/Runtime/Renderer/Private/SceneRendering.h
--- a/Engine/Source/Runtime/Renderer/Private/SceneRendering.h
+++ b/Engine/Source/Runtime/Renderer/Private/SceneRendering.h
@@ -1984,6 +1984,26 @@ struct FComputeLightGridOutput
 	FRDGPassRef CompactLinksPass = {};
 };
 
+/** An optional set of all renderers, view families, and views for operations that are performed once for a
+ *  batch of scene renderers. Only the first scene renderer will have this applied.
+ */
+struct FSceneRenderUpdateInputs
+{
+	TConstArrayView<const FSceneView*> GetAsSceneViews() const
+	{
+		return MakeArrayView(reinterpret_cast<const FSceneView* const*>(Views.GetData()), Views.Num());
+	}
+
+	TConstArrayView<const FSceneViewFamily*> GetAsSceneViewFamilies() const
+	{
+		return MakeArrayView(reinterpret_cast<const FSceneViewFamily* const*>(ViewFamilies.GetData()), ViewFamilies.Num());
+	}
+
+	TConstArrayView<FSceneRenderer*> Renderers;
+	TConstArrayView<FViewFamilyInfo*> ViewFamilies;
+	TConstArrayView<FViewInfo*> Views;
+};
+
 class FSceneRendererBase : public ISceneRenderer
 {
 public:
@@ -2142,8 +2162,8 @@ public:
 	virtual UE::Renderer::Private::IShadowInvalidatingInstances *GetShadowInvalidatingInstancesInterface(const FSceneView *SceneView) override;
 
 	// FSceneRenderer interface
-	virtual void Render(FRDGBuilder& GraphBuilder) = 0;
-	virtual void RenderHitProxies(FRDGBuilder& GraphBuilder) {}
+	virtual void Render(FRDGBuilder& GraphBuilder, const FSceneRenderUpdateInputs* SceneUpdateInputs) = 0;
+	virtual void RenderHitProxies(FRDGBuilder& GraphBuilder, const FSceneRenderUpdateInputs* SceneUpdateInputs) {}
 	virtual bool ShouldRenderVelocities() const { return false; }
 	virtual bool ShouldRenderPrePass() const { return false; }
 	virtual bool ShouldRenderNanite() const { return false; }
@@ -2484,7 +2504,7 @@ protected:
 	static bool ShouldRenderTranslucency(ETranslucencyPass::Type TranslucencyPass, TArrayView<FViewInfo> InViews);
 
 	/** Called at the begin / finish of the scene render. */
-	IVisibilityTaskData* OnRenderBegin(FRDGBuilder& GraphBuilder);
+	IVisibilityTaskData* OnRenderBegin(FRDGBuilder& GraphBuilder, const FSceneRenderUpdateInputs* SceneUpdateInputs);
 	void OnRenderFinish(FRDGBuilder& GraphBuilder, FRDGTextureRef ViewFamilyTexture);
 
 	bool RenderCustomDepthPass(
@@ -2693,9 +2713,9 @@ public:
 
 	// FSceneRenderer interface
 
-	virtual void Render(FRDGBuilder& GraphBuilder) override;
+	virtual void Render(FRDGBuilder& GraphBuilder, const FSceneRenderUpdateInputs* SceneUpdateInputs) override;
 
-	virtual void RenderHitProxies(FRDGBuilder& GraphBuilder) override;
+	virtual void RenderHitProxies(FRDGBuilder& GraphBuilder, const FSceneRenderUpdateInputs* SceneUpdateInputs) override;
 
 	virtual bool ShouldRenderVelocities() const override;
 
