Saved 2k of memory for each FSceneView / FViewInfo. Deprecated ShadowViewMatrices in favor of an accessor that is compiled out in shipping / test builds, as it is only used by r.Shadow.FreezeCamera.

diff --git a/Engine/Source/Runtime/Engine/Private/Components/DirectionalLightComponent.cpp b/Engine/Source/Runtime/Engine/Private/Components/DirectionalLightComponent.cpp
--- a/Engine/Source/Runtime/Engine/Private/Components/DirectionalLightComponent.cpp
+++ b/Engine/Source/Runtime/Engine/Private/Components/DirectionalLightComponent.cpp
@@ -815,8 +815,10 @@ private:
 
 	virtual FSphere GetShadowSplitBoundsDepthRange(const FSceneView& View, FVector ViewOrigin, float SplitNear, float SplitFar, FShadowCascadeSettings* OutCascadeSettings) const override
 	{
-		const FMatrix& ViewMatrix = View.ShadowViewMatrices.GetViewMatrix();
-		const FMatrix& ProjectionMatrix = View.ShadowViewMatrices.GetProjectionMatrix();
+		const FViewMatrices& ViewMatrices = View.GetShadowViewMatrices();
+
+		const FMatrix& ViewMatrix = ViewMatrices.GetViewMatrix();
+		const FMatrix& ProjectionMatrix = ViewMatrices.GetProjectionMatrix();
 
 		const FVector CameraDirection = ViewMatrix.GetColumn(2);
 		const FVector LightDirection = -GetDirection();
@@ -824,7 +826,7 @@ private:
 		// Support asymmetric projection
 		// Get FOV and AspectRatio from the view's projection matrix.
 		float AspectRatio = ProjectionMatrix.M[1][1] / ProjectionMatrix.M[0][0];
-		bool bIsPerspectiveProjection = View.ShadowViewMatrices.IsPerspectiveProjection();
+		bool bIsPerspectiveProjection = ViewMatrices.IsPerspectiveProjection();
 
 		// Build the camera frustum for this cascade
 		float HalfHorizontalFOV = bIsPerspectiveProjection ? FMath::Atan(1.0f / ProjectionMatrix.M[0][0]) : UE_PI / 4.0f;
diff --git a/Engine/Source/Runtime/Engine/Private/SceneView.cpp b/Engine/Source/Runtime/Engine/Private/SceneView.cpp
--- a/Engine/Source/Runtime/Engine/Private/SceneView.cpp
+++ b/Engine/Source/Runtime/Engine/Private/SceneView.cpp
@@ -835,6 +835,30 @@ static void SetupViewFrustum(FSceneView& View)
 	View.NearClippingDistance = View.ViewMatrices.ComputeNearPlane();
 }
 
+#if !(UE_BUILD_SHIPPING || UE_BUILD_TEST)
+namespace ShadowFreezeCamera
+{
+	static TOptional<FViewMatrices> Backup;
+
+	void Init(const FViewMatrices& Matrices)
+	{
+		const bool bShadowFreezeCamera = CVarShadowFreezeCamera.GetValueOnAnyThread() != 0;
+
+		if (bShadowFreezeCamera)
+		{
+			if (!Backup)
+			{
+				Backup.Emplace(Matrices);
+			}
+		}
+		else
+		{
+			Backup.Reset();
+		}
+	}
+}
+#endif
+
 FSceneView::FSceneView(const FSceneViewInitOptions& InitOptions)
 	: Family(InitOptions.ViewFamily)
 	, State(InitOptions.SceneViewStateInterface)
@@ -852,7 +876,6 @@ FSceneView::FSceneView(const FSceneViewInitOptions& InitOptions)
 	, BaseHmdOrientation(EForceInit::ForceInit)
 	, BaseHmdLocation(ForceInitToZero)
 	, WorldToMetersScale(InitOptions.WorldToMetersScale)
-	, ShadowViewMatrices(InitOptions)
 	, ProjectionMatrixUnadjustedForRHI(InitOptions.ProjectionMatrix)
 	, BackgroundColor(InitOptions.BackgroundColor)
 	, OverlayColor(InitOptions.OverlayColor)
@@ -909,26 +932,10 @@ PRAGMA_ENABLE_DEPRECATION_WARNINGS
 	, PrimaryScreenPercentageMethod(EPrimaryScreenPercentageMethod::SpatialUpscale)
 	, FeatureLevel(InitOptions.ViewFamily ? InitOptions.ViewFamily->GetFeatureLevel() : GMaxRHIFeatureLevel)
 {
-	ShadowViewMatrices = ViewMatrices;
-
 	SceneViewInitOptions = FSceneViewInitOptions(InitOptions);
 
 #if !(UE_BUILD_SHIPPING || UE_BUILD_TEST)
-	{
-		// console variable override
-		int32 Value = CVarShadowFreezeCamera.GetValueOnAnyThread();
-
-		static FViewMatrices Backup = ShadowViewMatrices;
-
-		if(Value)
-		{
-			ShadowViewMatrices = Backup;
-		}
-		else
-		{
-			Backup = ShadowViewMatrices;
-		}
-	}
+	ShadowFreezeCamera::Init(ViewMatrices);
 #endif // !(UE_BUILD_SHIPPING || UE_BUILD_TEST)
 
 	SetupViewFrustum(*this);
@@ -1062,6 +1069,12 @@ PRAGMA_ENABLE_DEPRECATION_WARNINGS
 	check(VerifyMembersChecks());
 }
 
+#if !(UE_BUILD_SHIPPING || UE_BUILD_TEST)
+const FViewMatrices& FSceneView::GetShadowViewMatrices() const
+{
+	return ShadowFreezeCamera::Backup ? *ShadowFreezeCamera::Backup : ViewMatrices;
+}
+#endif
 
 #if DO_CHECK || USING_CODE_ANALYSIS
 bool FSceneView::VerifyMembersChecks() const
@@ -1270,9 +1283,6 @@ void FSceneView::UpdateViewMatrix()
 
 	ViewMatrices.UpdateViewMatrix(StereoViewLocation, StereoViewRotation);
 	GetViewFrustumBounds(ViewFrustum, ViewMatrices.GetViewProjectionMatrix(), false);
-
-	// We need to keep ShadowViewMatrices in sync.
-	ShadowViewMatrices = ViewMatrices;
 }
 
 void FViewMatrices::UpdatePlanarReflectionViewMatrix(const FSceneView& SourceView, const FMirrorMatrix& MirrorMatrix)
@@ -1305,9 +1315,6 @@ void FSceneView::UpdatePlanarReflectionViewMatrix(const FSceneView& SourceView,
 
 	// Update bounds
 	GetViewFrustumBounds(ViewFrustum, ViewMatrices.GetViewProjectionMatrix(), false);
-
-	// We need to keep ShadowViewMatrices in sync.
-	ShadowViewMatrices = ViewMatrices;
 }
 
 FVector4 FSceneView::WorldToScreen(const FVector& WorldPoint) const
diff --git a/Engine/Source/Runtime/Engine/Public/SceneView.h b/Engine/Source/Runtime/Engine/Public/SceneView.h
--- a/Engine/Source/Runtime/Engine/Public/SceneView.h
+++ b/Engine/Source/Runtime/Engine/Public/SceneView.h
@@ -1401,8 +1401,14 @@ public:
 	float		WorldToMetersScale;
 	TOptional<FTransform> PreviousViewTransform;
 
-	// normally the same as ViewMatrices unless "r.Shadow.FreezeCamera" is activated
-	FViewMatrices ShadowViewMatrices;
+	UE_DEPRECATED(5.6, "Use GetShadowViewMatrices instead")
+	uint8 ShadowViewMatrices;
+
+#if !(UE_BUILD_SHIPPING || UE_BUILD_TEST)
+	ENGINE_API const FViewMatrices& GetShadowViewMatrices() const;
+#else
+	const FViewMatrices& GetShadowViewMatrices() const { return ViewMatrices; }
+#endif
 
 	FMatrix ProjectionMatrixUnadjustedForRHI;
 
diff --git a/Engine/Source/Runtime/Renderer/Private/SceneOcclusion.cpp b/Engine/Source/Runtime/Renderer/Private/SceneOcclusion.cpp
--- a/Engine/Source/Runtime/Renderer/Private/SceneOcclusion.cpp
+++ b/Engine/Source/Runtime/Renderer/Private/SceneOcclusion.cpp
@@ -613,13 +613,15 @@ static void ExecutePointLightShadowOcclusionQuery(FRHICommandList& RHICmdList, F
 
 static void PrepareDirectionalLightShadowOcclusionQuery(uint32& BaseVertexIndex, FVector3f* DestinationBuffer, const FViewInfo& View, const FProjectedShadowInfo& ProjectedShadowInfo)
 {
-	const FMatrix& ViewMatrix = View.ShadowViewMatrices.GetViewMatrix();
-	const FMatrix& ProjectionMatrix = View.ShadowViewMatrices.GetProjectionMatrix();
+	const FViewMatrices& ViewMatrices = View.GetShadowViewMatrices();
+
+	const FMatrix& ViewMatrix = ViewMatrices.GetViewMatrix();
+	const FMatrix& ProjectionMatrix = ViewMatrices.GetProjectionMatrix();
 	const FVector CameraDirection = ViewMatrix.GetColumn(2);
 	const float SplitNear = ProjectedShadowInfo.CascadeSettings.SplitNear;
 
 	float AspectRatio = ProjectionMatrix.M[1][1] / ProjectionMatrix.M[0][0];
-	float HalfFOV = View.ShadowViewMatrices.IsPerspectiveProjection() ? FMath::Atan(1.0f / ProjectionMatrix.M[0][0]) : PI / 4.0f;
+	float HalfFOV = ViewMatrices.IsPerspectiveProjection() ? FMath::Atan(1.0f / ProjectionMatrix.M[0][0]) : PI / 4.0f;
 
 	// Build the camera frustum for this cascade
 	const float StartHorizontalLength = SplitNear * FMath::Tan(HalfFOV);
diff --git a/Engine/Source/Runtime/Renderer/Private/ShadowSetup.cpp b/Engine/Source/Runtime/Renderer/Private/ShadowSetup.cpp
--- a/Engine/Source/Runtime/Renderer/Private/ShadowSetup.cpp
+++ b/Engine/Source/Runtime/Renderer/Private/ShadowSetup.cpp
@@ -1999,7 +1999,7 @@ bool FProjectedShadowInfo::AddSubjectPrimitive(FPrimitiveSceneInfo* PrimitiveSce
 				if( CurrentView.IsPerspectiveProjection() )
 				{
 					// Compute the distance between the view and the primitive.
-					float DistanceSquared = (Proxy->GetBounds().Origin - CurrentView.ShadowViewMatrices.GetViewOrigin()).SizeSquared();
+					float DistanceSquared = (Proxy->GetBounds().Origin - CurrentView.GetShadowViewMatrices().GetViewOrigin()).SizeSquared();
 
 					bool bIsDistanceCulled = CurrentView.IsDistanceCulled(
 						DistanceSquared,
@@ -2059,7 +2059,7 @@ bool FProjectedShadowInfo::AddSubjectPrimitive(FPrimitiveSceneInfo* PrimitiveSce
 					// Note: skip small-mesh culling for VSM since it needs it drawn for GPU-side caching.
 					if (bWholeSceneShadow && CacheMode != SDCM_StaticPrimitivesOnly && !bVSM)
 					{
-						const float DistanceSquared = ( Bounds.Origin - CurrentView.ShadowViewMatrices.GetViewOrigin() ).SizeSquared();
+						const float DistanceSquared = ( Bounds.Origin - CurrentView.GetShadowViewMatrices().GetViewOrigin() ).SizeSquared();
 						const float LODScaleSquared = FMath::Square(CurrentView.LODDistanceFactor);
 						const bool bDrawShadowDepth = FMath::Square(Bounds.SphereRadius) > FMath::Square(GMinScreenRadiusForShadowCaster) * DistanceSquared * LODScaleSquared;
 						if( !bDrawShadowDepth )
@@ -2130,7 +2130,7 @@ bool FProjectedShadowInfo::TestPrimitiveFarCascadeConditions(bool bPrimitiveCast
 		{
 			check(DependentView);
 
-			FVector ViewOrigin = DependentView->ShadowViewMatrices.GetViewOrigin();
+			FVector ViewOrigin = DependentView->GetShadowViewMatrices().GetViewOrigin();
 			float ViewDistance = (Bounds.Origin - ViewOrigin) | DependentView->GetViewDirection();
 			if (ViewDistance - Bounds.SphereRadius > MaxNonFarCascadeDistance)
 			{
@@ -2215,7 +2215,7 @@ uint64 FProjectedShadowInfo::AddSubjectPrimitive_AnyThread(
 			{
 				bool bFadingIn;
 				// Compute the distance between the view and the primitive.
-				const float DistanceSquared = (PrimitiveSceneInfoCompact.Bounds.Origin - CurrentView->ShadowViewMatrices.GetViewOrigin()).SizeSquared();
+				const float DistanceSquared = (PrimitiveSceneInfoCompact.Bounds.Origin - CurrentView->GetShadowViewMatrices().GetViewOrigin()).SizeSquared();
 
 				if (CurrentView->IsDistanceCulled_AnyThread(
 					DistanceSquared,
@@ -2280,7 +2280,7 @@ uint64 FProjectedShadowInfo::AddSubjectPrimitive_AnyThread(
 				if (bWholeSceneShadow)
 				{
 					const FCompactBoxSphereBounds& Bounds = PrimitiveSceneInfoCompact.Bounds;
-					const float DistanceSquared = (Bounds.Origin - CurrentView->ShadowViewMatrices.GetViewOrigin()).SizeSquared();
+					const float DistanceSquared = (Bounds.Origin - CurrentView->GetShadowViewMatrices().GetViewOrigin()).SizeSquared();
 					const float LODScaleSquared = FMath::Square(CurrentView->LODDistanceFactor);
 					const bool bDrawShadowDepth = FMath::Square(Bounds.SphereRadius) > FMath::Square(GMinScreenRadiusForShadowCaster) * DistanceSquared * LODScaleSquared;
 					if (!bDrawShadowDepth)
@@ -3414,19 +3414,20 @@ void FSceneRenderer::CreatePerObjectProjectedShadow(
 	for (int32 ViewIndex = 0; ViewIndex < Views.Num(); ViewIndex++)
 	{
 		const FViewInfo& View = Views[ViewIndex];
+		const FViewMatrices& ShadowViewMatrices = View.GetShadowViewMatrices();
 
 		// Determine the size of the subject's bounding sphere in this view.
 		const FVector ShadowViewOrigin = View.ViewMatrices.GetViewOrigin();
 		float ShadowViewDistFromBounds = (OriginalBounds.Origin - ShadowViewOrigin).Size();
-		const float ScreenRadius = View.ShadowViewMatrices.GetScreenScale() *
+		const float ScreenRadius = ShadowViewMatrices.GetScreenScale() *
 			OriginalBounds.SphereRadius /
 			FMath::Max(ShadowViewDistFromBounds, 1.0f);
 		// Early catch for invalid CalculateShadowFadeAlpha()
-		ensureMsgf(ScreenRadius >= 0.0f, TEXT("View.ShadowViewMatrices.ScreenScale %f, OriginalBounds.SphereRadius %f, ShadowViewDistFromBounds %f"), View.ShadowViewMatrices.GetScreenScale(), OriginalBounds.SphereRadius, ShadowViewDistFromBounds);
+		ensureMsgf(ScreenRadius >= 0.0f, TEXT("ShadowViewMatrices.ScreenScale %f, OriginalBounds.SphereRadius %f, ShadowViewDistFromBounds %f"), ShadowViewMatrices.GetScreenScale(), OriginalBounds.SphereRadius, ShadowViewDistFromBounds);
 
 		const float ScreenPercent = FMath::Max(
-			1.0f / 2.0f * View.ShadowViewMatrices.GetProjectionScale().X,
-			1.0f / 2.0f * View.ShadowViewMatrices.GetProjectionScale().Y
+			1.0f / 2.0f * ShadowViewMatrices.GetProjectionScale().X,
+			1.0f / 2.0f * ShadowViewMatrices.GetProjectionScale().Y
 			) *
 			OriginalBounds.SphereRadius /
 			FMath::Max(ShadowViewDistFromBounds, 1.0f);
@@ -3601,7 +3602,7 @@ void FSceneRenderer::CreatePerObjectProjectedShadow(
 				const float DistanceFromShadowCenterSquared = (WholeSceneShadow->ShadowBounds.Center - Bounds.Origin).SizeSquared();
 				//@todo - if view dependent whole scene shadows are ever supported in splitscreen, 
 				// We can only disable the preshadow at this point if it is inside a whole scene shadow for all views
-				const float DistanceFromViewSquared = ((FVector)WholeSceneShadow->DependentView->ShadowViewMatrices.GetViewOrigin() - Bounds.Origin).SizeSquared();
+				const float DistanceFromViewSquared = ((FVector)WholeSceneShadow->DependentView->GetShadowViewMatrices().GetViewOrigin() - Bounds.Origin).SizeSquared();
 				// Mark the preshadow as inside the whole scene shadow if its bounding sphere is inside the near fade distance
 				if (DistanceFromShadowCenterSquared < FMath::Square(FMath::Max(WholeSceneShadow->ShadowBounds.W - Bounds.SphereRadius, 0.0f))
 					//@todo - why is this extra threshold required?
@@ -4145,7 +4146,7 @@ void FSceneRenderer::CreateWholeSceneProjectedShadow(
 		{
 			const FViewInfo& View = Views[ViewIndex];
 
-			const float ScreenRadius = LightSceneInfo->Proxy->GetEffectiveScreenRadius(View.ShadowViewMatrices, View.ViewRect.Size());
+			const float ScreenRadius = LightSceneInfo->Proxy->GetEffectiveScreenRadius(View.GetShadowViewMatrices(), View.ViewRect.Size());
 			MaxScreenRadius = FMath::Max(MaxScreenRadius, ScreenRadius);
 
 			// Determine the amount of shadow buffer resolution needed for this view.
@@ -4969,7 +4970,7 @@ struct FGatherShadowPrimitivesPacket
 				check(ProjectedShadowInfo->DependentView);
 
 				{
-					const float DistanceSquared = (PrimitiveBounds.Origin - ProjectedShadowInfo->DependentView->ShadowViewMatrices.GetViewOrigin()).SizeSquared();
+					const float DistanceSquared = (PrimitiveBounds.Origin - ProjectedShadowInfo->DependentView->GetShadowViewMatrices().GetViewOrigin()).SizeSquared();
 					const float LODScaleSquared = FMath::Square(ProjectedShadowInfo->DependentView->LODDistanceFactor);
 					bScreenSpaceSizeCulled = FMath::Square(PrimitiveBounds.SphereRadius) < FMath::Square(MinScreenRadiusForShadowCaster) * DistanceSquared * LODScaleSquared;
 				}
diff --git a/Engine/Source/Runtime/Renderer/Private/Shadows/ShadowSceneRenderer.cpp b/Engine/Source/Runtime/Renderer/Private/Shadows/ShadowSceneRenderer.cpp
--- a/Engine/Source/Runtime/Renderer/Private/Shadows/ShadowSceneRenderer.cpp
+++ b/Engine/Source/Runtime/Renderer/Private/Shadows/ShadowSceneRenderer.cpp
@@ -297,7 +297,7 @@ static float GetWorldSpacePixelFootprint(float ViewSpaceDepth, float ClipToViewS
 static uint32 GetConservativeMipLevelLocal(const FViewInfo& View, float ClipToViewSizeScale, float ClipToViewSizeBias, const FVector& LightOrigin, float LightRadius, float WorldToShadowFootprintScale, float ResolutionLodBias, float GlobalResolutionLodBias, uint32 MipModeLocal)
 {
 	// Note: not just a rotation, full world-space DP.
-	FVector ViewSpaceOrigin = View.ShadowViewMatrices.GetViewMatrix().TransformPosition(LightOrigin);
+	FVector ViewSpaceOrigin = View.GetShadowViewMatrices().GetViewMatrix().TransformPosition(LightOrigin);
 
 	// Remove radius to arrive at minimum possible z-distance in view space, from primary view.
 	float RadiusWorld = GetWorldSpacePixelFootprint(FMath::Max(0.0f, float(ViewSpaceOrigin.Z) - LightRadius), ClipToViewSizeScale, ClipToViewSizeBias);
@@ -337,7 +337,7 @@ TSharedPtr<FVirtualShadowMapPerLightCacheEntry> FShadowSceneRenderer::AddLocalLi
 	{
 		const FViewInfo& View = SceneRenderer.Views[ViewIndex];
 		const FViewData& ViewData = ViewDatas[ViewIndex];
-		FVector TestOrigin = View.ShadowViewMatrices.GetViewOrigin();
+		FVector TestOrigin = View.GetShadowViewMatrices().GetViewOrigin();
 		double TestDistanceSq = (TestOrigin + ProjectedShadowInfo->PreShadowTranslation).SquaredLength();
 		if (TestDistanceSq < MinDistanceSq)
 		{
diff --git a/Engine/Source/Runtime/Renderer/Private/TranslucentLighting.cpp b/Engine/Source/Runtime/Renderer/Private/TranslucentLighting.cpp
--- a/Engine/Source/Runtime/Renderer/Private/TranslucentLighting.cpp
+++ b/Engine/Source/Runtime/Renderer/Private/TranslucentLighting.cpp
@@ -183,17 +183,19 @@ void FViewInfo::CalcTranslucencyLightingVolumeBounds(FBox* InOutCascadeBoundsArr
 		double FieldOfView = DOUBLE_PI / 4.0;
 		double AspectRatio = 1.0;
 
+		const FViewMatrices& LocalShadowViewMatrices = GetShadowViewMatrices();
+
 		if (IsPerspectiveProjection() && GTranslucencyLightingVolumeAdaptToPerspective > 0)
 		{
 			// Derive FOV and aspect ratio from the perspective projection matrix
-			FieldOfView = FMath::Atan(1.0 / ShadowViewMatrices.GetProjectionMatrix().M[0][0]);
+			FieldOfView = FMath::Atan(1.0 / LocalShadowViewMatrices.GetProjectionMatrix().M[0][0]);
 			// Clamp to prevent shimmering when zooming in
 			FieldOfView = FMath::Max(FieldOfView, GTranslucentVolumeMinFOV * DOUBLE_PI / 180.0);
 			const double RoundFactorRadians = GTranslucentVolumeFOVSnapFactor * DOUBLE_PI / 180.0;
 			// Round up to a fixed factor
 			// This causes the volume lighting to make discreet jumps as the FOV animates, instead of slowly crawling over a long period
 			FieldOfView = FieldOfView + RoundFactorRadians - FMath::Fmod(FieldOfView, RoundFactorRadians);
-			AspectRatio = ShadowViewMatrices.GetProjectionMatrix().M[1][1] / ShadowViewMatrices.GetProjectionMatrix().M[0][0];
+			AspectRatio = LocalShadowViewMatrices.GetProjectionMatrix().M[1][1] / LocalShadowViewMatrices.GetProjectionMatrix().M[0][0];
 		}
 
 		// Tan of field of view can explode when FieldOfView is close to 180 degree, when using the adaptation to Fov above.
@@ -203,17 +205,17 @@ void FViewInfo::CalcTranslucencyLightingVolumeBounds(FBox* InOutCascadeBoundsArr
 		const double TanFieldOfView = FMath::Tan(FieldOfView);
 
 		const double StartHorizontalLength = FMath::Min(MaxTanFieldOfViewLength, FrustumStartDistance * TanFieldOfView);
-		const FVector StartCameraRightOffset = ShadowViewMatrices.GetViewMatrix().GetColumn(0) * StartHorizontalLength;
+		const FVector StartCameraRightOffset = LocalShadowViewMatrices.GetViewMatrix().GetColumn(0) * StartHorizontalLength;
 		const double StartVerticalLength = StartHorizontalLength / AspectRatio;
-		const FVector StartCameraUpOffset = ShadowViewMatrices.GetViewMatrix().GetColumn(1) * StartVerticalLength;
+		const FVector StartCameraUpOffset = LocalShadowViewMatrices.GetViewMatrix().GetColumn(1) * StartVerticalLength;
 
 		const double EndHorizontalLength = FMath::Min(MaxTanFieldOfViewLength, FrustumEndDistance * TanFieldOfView);
-		const FVector EndCameraRightOffset = ShadowViewMatrices.GetViewMatrix().GetColumn(0) * EndHorizontalLength;
+		const FVector EndCameraRightOffset = LocalShadowViewMatrices.GetViewMatrix().GetColumn(0) * EndHorizontalLength;
 		const double EndVerticalLength = EndHorizontalLength / AspectRatio;
-		const FVector EndCameraUpOffset = ShadowViewMatrices.GetViewMatrix().GetColumn(1) * EndVerticalLength;
+		const FVector EndCameraUpOffset = LocalShadowViewMatrices.GetViewMatrix().GetColumn(1) * EndVerticalLength;
 
 		FVector SplitVertices[8];
-		const FVector ShadowViewOrigin = ShadowViewMatrices.GetViewOrigin();
+		const FVector ShadowViewOrigin = LocalShadowViewMatrices.GetViewOrigin();
 
 		SplitVertices[0] = ShadowViewOrigin + GetViewDirection() * FrustumStartDistance + StartCameraRightOffset + StartCameraUpOffset;
 		SplitVertices[1] = ShadowViewOrigin + GetViewDirection() * FrustumStartDistance + StartCameraRightOffset - StartCameraUpOffset;
diff --git a/Engine/Source/Runtime/Renderer/Private/VirtualShadowMaps/VirtualShadowMapArray.cpp b/Engine/Source/Runtime/Renderer/Private/VirtualShadowMaps/VirtualShadowMapArray.cpp
--- a/Engine/Source/Runtime/Renderer/Private/VirtualShadowMaps/VirtualShadowMapArray.cpp
+++ b/Engine/Source/Runtime/Renderer/Private/VirtualShadowMaps/VirtualShadowMapArray.cpp
@@ -3925,10 +3925,10 @@ void FVirtualShadowMapArray::AddRenderViewsLocal(
 	// Local lights, select the view closest to the local light to get some kind of reasonable behavior for split screen.
 	int32 ClosestCullingViewIndex = 0;
 	{
-		double MinDistanceSq = (Views[0].ShadowViewMatrices.GetViewOrigin() + ProjectedShadowInfo->PreShadowTranslation).SquaredLength();
+		double MinDistanceSq = (Views[0].GetShadowViewMatrices().GetViewOrigin() + ProjectedShadowInfo->PreShadowTranslation).SquaredLength();
 		for (int Index = 1; Index < Views.Num(); ++Index)
 		{
-			FVector TestOrigin = Views[Index].ShadowViewMatrices.GetViewOrigin();
+			FVector TestOrigin = Views[Index].GetShadowViewMatrices().GetViewOrigin();
 			double TestDistanceSq = (TestOrigin + ProjectedShadowInfo->PreShadowTranslation).SquaredLength();
 			if (TestDistanceSq < MinDistanceSq)
 			{
