Fixed: Incorrect mesh rendering on mobile devices with -rdgimmediate commandline parameter

diff --git a/Engine/Source/Runtime/Renderer/Private/InstanceCulling/InstanceCullingContext.cpp b/Engine/Source/Runtime/Renderer/Private/InstanceCulling/InstanceCullingContext.cpp
--- a/Engine/Source/Runtime/Renderer/Private/InstanceCulling/InstanceCullingContext.cpp
+++ b/Engine/Source/Runtime/Renderer/Private/InstanceCulling/InstanceCullingContext.cpp
@@ -95,6 +95,7 @@ FMeshDrawCommandOverrideArgs GetMeshDrawCommandOverrideArgs(const FInstanceCulli
 	FMeshDrawCommandOverrideArgs Result;
 	Result.InstanceBuffer = InstanceCullingDrawParams.InstanceIdOffsetBuffer.GetBuffer() != nullptr ? InstanceCullingDrawParams.InstanceIdOffsetBuffer.GetBuffer()->GetRHI() : nullptr;
 	Result.IndirectArgsBuffer = InstanceCullingDrawParams.DrawIndirectArgsBuffer.GetBuffer() != nullptr ? InstanceCullingDrawParams.DrawIndirectArgsBuffer.GetBuffer()->GetRHI() : nullptr;
+	Result.BatchedPrimitiveUB = InstanceCullingDrawParams.BatchedPrimitive.GetUniformBuffer() != nullptr ? InstanceCullingDrawParams.BatchedPrimitive.GetUniformBuffer()->GetRHI() : nullptr;
 	Result.InstanceDataByteOffset = InstanceCullingDrawParams.InstanceDataByteOffset;
 	Result.IndirectArgsByteOffset = InstanceCullingDrawParams.IndirectArgsByteOffset;
 	return Result;
@@ -1672,6 +1673,8 @@ void FInstanceCullingContext::SubmitDrawCommands(
 			// UniformBufferView suplies instance data through global UB binding
 			SceneArgs.PrimitiveIdsBuffer = nullptr;
 			SceneArgs.BatchedPrimitiveSlot = BatchedPrimitiveSlot;
+			// Only does anything when batching is disabled and each mesh pass has an unique BatchedPrimitive UB
+			RHICmdList.SetStaticUniformBuffer(BatchedPrimitiveSlot, OverrideArgs.BatchedPrimitiveUB);
 		}
 							
 		FMeshDrawCommandStateCache StateCache;
diff --git a/Engine/Source/Runtime/Renderer/Private/MobileShadingRenderer.cpp b/Engine/Source/Runtime/Renderer/Private/MobileShadingRenderer.cpp
--- a/Engine/Source/Runtime/Renderer/Private/MobileShadingRenderer.cpp
+++ b/Engine/Source/Runtime/Renderer/Private/MobileShadingRenderer.cpp
@@ -189,6 +189,22 @@ static void RenderOpaqueFX(
 	}
 }
 
+static void BuildMeshRenderingCommands(FRDGBuilder& GraphBuilder, EMeshPass::Type MeshPass, FViewInfo& View, const FGPUScene& GPUScene, FInstanceCullingDrawParams& OutParams)
+{
+	View.ParallelMeshDrawCommandPasses[MeshPass].BuildRenderingCommands(GraphBuilder, GPUScene, OutParams);
+	
+	// When batching is disabled instead of a single UniformBuffer we get a separate buffers for each mesh pass
+	// Because mobile renderer manually merges several mesh passes into a single RDG pass we can't specify InstanceCullingDrawParams for each mesh pass through RDG pass parameters (only one)
+	// We do these dummy RDG passes to make sure InstanceCullingDrawParams are initialized for each mesh pass
+	if (!FInstanceCullingManager::AllowBatchedBuildRenderingCommands(GPUScene))
+	{
+		GraphBuilder.AddPass(
+			RDG_EVENT_NAME("SetupInstanceCullingDrawParams"),
+			&OutParams,
+			ERDGPassFlags::Raster | ERDGPassFlags::SkipRenderPass | ERDGPassFlags::NeverCull, [](FRHICommandList&){});
+	}
+}
+
 BEGIN_SHADER_PARAMETER_STRUCT(FMobileRenderPassParameters,)
 	SHADER_PARAMETER_STRUCT_INCLUDE(FViewShaderParameters, View)
 	SHADER_PARAMETER_STRUCT_INCLUDE(FInstanceCullingDrawParams, InstanceCullingDrawParams)
@@ -807,11 +823,10 @@ void FMobileSceneRenderer::RenderFullDepthPrepass(FRDGBuilder& GraphBuilder, TAr
 			}
 		}
 
-		View.ParallelMeshDrawCommandPasses[EMeshPass::DepthPass].BuildRenderingCommands(GraphBuilder, Scene->GPUScene, PassParameters->InstanceCullingDrawParams);
-
+		BuildMeshRenderingCommands(GraphBuilder, EMeshPass::DepthPass, View, Scene->GPUScene, PassParameters->InstanceCullingDrawParams);
 		if (bRenderVelocityInDepthPrePass)
 		{
-			View.ParallelMeshDrawCommandPasses[EMeshPass::Velocity].BuildRenderingCommands(GraphBuilder, Scene->GPUScene, ParameterCollection->VelocityInstanceCullingDrawParams);
+			BuildMeshRenderingCommands(GraphBuilder, EMeshPass::Velocity, View, Scene->GPUScene, ParameterCollection->VelocityInstanceCullingDrawParams);
 		}
 
 		GraphBuilder.AddPass(
@@ -923,12 +938,11 @@ void FMobileSceneRenderer::RenderCustomRenderPassBasePass(FRDGBuilder& GraphBuil
 
 		if (Scene->GPUScene.IsEnabled())
 		{
-			View.ParallelMeshDrawCommandPasses[EMeshPass::BasePass].BuildRenderingCommands(GraphBuilder, Scene->GPUScene, PassParameters->InstanceCullingDrawParams);
-			View.ParallelMeshDrawCommandPasses[EMeshPass::SkyPass].BuildRenderingCommands(GraphBuilder, Scene->GPUScene, ParameterCollection->SkyPassInstanceCullingDrawParams);
-
+			BuildMeshRenderingCommands(GraphBuilder, EMeshPass::BasePass, View, Scene->GPUScene, PassParameters->InstanceCullingDrawParams);
+			BuildMeshRenderingCommands(GraphBuilder, EMeshPass::SkyPass, View, Scene->GPUScene, ParameterCollection->SkyPassInstanceCullingDrawParams);
 			if (bIncludeTranslucent)
 			{
-				View.ParallelMeshDrawCommandPasses[EMeshPass::TranslucencyAll].BuildRenderingCommands(GraphBuilder, Scene->GPUScene, TranslucencyPassParameters->InstanceCullingDrawParams);
+				BuildMeshRenderingCommands(GraphBuilder, EMeshPass::TranslucencyAll, View, Scene->GPUScene, TranslucencyPassParameters->InstanceCullingDrawParams);
 			}
 		}
 
@@ -1658,13 +1672,13 @@ void FMobileSceneRenderer::RenderForwardSinglePass(FRDGBuilder& GraphBuilder, FM
 	{
 		if (!bIsFullDepthPrepassEnabled)
 		{
-			View.ParallelMeshDrawCommandPasses[EMeshPass::DepthPass].BuildRenderingCommands(GraphBuilder, Scene->GPUScene, ParameterCollection->DepthPassInstanceCullingDrawParams);
+			BuildMeshRenderingCommands(GraphBuilder, EMeshPass::DepthPass, View, Scene->GPUScene, ParameterCollection->DepthPassInstanceCullingDrawParams);
 		}
-		View.ParallelMeshDrawCommandPasses[EMeshPass::BasePass].BuildRenderingCommands(GraphBuilder, Scene->GPUScene, PassParameters->InstanceCullingDrawParams);
-		View.ParallelMeshDrawCommandPasses[EMeshPass::SkyPass].BuildRenderingCommands(GraphBuilder, Scene->GPUScene, ParameterCollection->SkyPassInstanceCullingDrawParams);
-		View.ParallelMeshDrawCommandPasses[EMeshPass::DebugViewMode].BuildRenderingCommands(GraphBuilder, Scene->GPUScene, ParameterCollection->DebugViewModeInstanceCullingDrawParams);
-		View.ParallelMeshDrawCommandPasses[EMeshPass::MeshDecal_SceneColor].BuildRenderingCommands(GraphBuilder, Scene->GPUScene, ParameterCollection->MeshDecalSceneColorInstanceCullingDrawParams);
-		View.ParallelMeshDrawCommandPasses[StandardTranslucencyMeshPass].BuildRenderingCommands(GraphBuilder, Scene->GPUScene, ParameterCollection->TranslucencyInstanceCullingDrawParams);
+		BuildMeshRenderingCommands(GraphBuilder, EMeshPass::BasePass, View, Scene->GPUScene, PassParameters->InstanceCullingDrawParams);
+		BuildMeshRenderingCommands(GraphBuilder, EMeshPass::SkyPass, View, Scene->GPUScene, ParameterCollection->SkyPassInstanceCullingDrawParams);
+		BuildMeshRenderingCommands(GraphBuilder, EMeshPass::DebugViewMode, View, Scene->GPUScene, ParameterCollection->DebugViewModeInstanceCullingDrawParams);
+		BuildMeshRenderingCommands(GraphBuilder, EMeshPass::MeshDecal_SceneColor, View, Scene->GPUScene, ParameterCollection->MeshDecalSceneColorInstanceCullingDrawParams);
+		BuildMeshRenderingCommands(GraphBuilder, StandardTranslucencyMeshPass, View, Scene->GPUScene, ParameterCollection->TranslucencyInstanceCullingDrawParams);
 	}
 
 	if (bTonemapSubpassInline)
@@ -1761,11 +1775,11 @@ void FMobileSceneRenderer::RenderForwardMultiPass(FRDGBuilder& GraphBuilder, FMo
 	{
 		if (!bIsFullDepthPrepassEnabled)
 		{
-			View.ParallelMeshDrawCommandPasses[EMeshPass::DepthPass].BuildRenderingCommands(GraphBuilder, Scene->GPUScene, ParameterCollection->DepthPassInstanceCullingDrawParams);
+			BuildMeshRenderingCommands(GraphBuilder, EMeshPass::DepthPass, View, Scene->GPUScene, ParameterCollection->DepthPassInstanceCullingDrawParams);
 		}
-		View.ParallelMeshDrawCommandPasses[EMeshPass::BasePass].BuildRenderingCommands(GraphBuilder, Scene->GPUScene, PassParameters->InstanceCullingDrawParams);
-		View.ParallelMeshDrawCommandPasses[EMeshPass::SkyPass].BuildRenderingCommands(GraphBuilder, Scene->GPUScene, ParameterCollection->SkyPassInstanceCullingDrawParams);
-		View.ParallelMeshDrawCommandPasses[EMeshPass::DebugViewMode].BuildRenderingCommands(GraphBuilder, Scene->GPUScene, ParameterCollection->DebugViewModeInstanceCullingDrawParams);
+		BuildMeshRenderingCommands(GraphBuilder, EMeshPass::BasePass, View, Scene->GPUScene, PassParameters->InstanceCullingDrawParams);
+		BuildMeshRenderingCommands(GraphBuilder, EMeshPass::SkyPass, View, Scene->GPUScene, ParameterCollection->SkyPassInstanceCullingDrawParams);
+		BuildMeshRenderingCommands(GraphBuilder, EMeshPass::DebugViewMode, View, Scene->GPUScene, ParameterCollection->DebugViewModeInstanceCullingDrawParams);
 	}
 
 	GraphBuilder.AddPass(
@@ -1836,8 +1850,8 @@ void FMobileSceneRenderer::RenderForwardMultiPass(FRDGBuilder& GraphBuilder, FMo
 
 	if (Scene->GPUScene.IsEnabled())
 	{
-		View.ParallelMeshDrawCommandPasses[EMeshPass::MeshDecal_SceneColor].BuildRenderingCommands(GraphBuilder, Scene->GPUScene, SecondParameterCollection->MeshDecalSceneColorInstanceCullingDrawParams);
-		View.ParallelMeshDrawCommandPasses[StandardTranslucencyMeshPass].BuildRenderingCommands(GraphBuilder, Scene->GPUScene, SecondParameterCollection->TranslucencyInstanceCullingDrawParams);
+		BuildMeshRenderingCommands(GraphBuilder, EMeshPass::MeshDecal_SceneColor, View, Scene->GPUScene, SecondParameterCollection->MeshDecalSceneColorInstanceCullingDrawParams);
+		BuildMeshRenderingCommands(GraphBuilder, StandardTranslucencyMeshPass, View, Scene->GPUScene, SecondParameterCollection->TranslucencyInstanceCullingDrawParams);
 	}
 
 	GraphBuilder.AddPass(
@@ -2070,13 +2084,13 @@ void FMobileSceneRenderer::RenderDeferredSinglePass(FRDGBuilder& GraphBuilder, F
 		{
 			if (!bIsFullDepthPrepassEnabled)
 			{
-				View.ParallelMeshDrawCommandPasses[EMeshPass::DepthPass].BuildRenderingCommands(GraphBuilder, Scene->GPUScene, ParameterCollection->DepthPassInstanceCullingDrawParams);
+				BuildMeshRenderingCommands(GraphBuilder, EMeshPass::DepthPass, View, Scene->GPUScene, ParameterCollection->DepthPassInstanceCullingDrawParams);
 			}
-			View.ParallelMeshDrawCommandPasses[EMeshPass::BasePass].BuildRenderingCommands(GraphBuilder, Scene->GPUScene, PassParameters->InstanceCullingDrawParams);
-			View.ParallelMeshDrawCommandPasses[EMeshPass::SkyPass].BuildRenderingCommands(GraphBuilder, Scene->GPUScene, ParameterCollection->SkyPassInstanceCullingDrawParams);
-			View.ParallelMeshDrawCommandPasses[EMeshPass::DebugViewMode].BuildRenderingCommands(GraphBuilder, Scene->GPUScene, ParameterCollection->DebugViewModeInstanceCullingDrawParams);
-			View.ParallelMeshDrawCommandPasses[EMeshPass::MeshDecal_SceneColorAndGBuffer].BuildRenderingCommands(GraphBuilder, Scene->GPUScene, ParameterCollection->MeshDecalSceneColorAndGBufferInstanceCullingDrawParams);
-			View.ParallelMeshDrawCommandPasses[StandardTranslucencyMeshPass].BuildRenderingCommands(GraphBuilder, Scene->GPUScene, ParameterCollection->TranslucencyInstanceCullingDrawParams);
+			BuildMeshRenderingCommands(GraphBuilder, EMeshPass::BasePass, View, Scene->GPUScene, PassParameters->InstanceCullingDrawParams);
+			BuildMeshRenderingCommands(GraphBuilder, EMeshPass::SkyPass, View, Scene->GPUScene, ParameterCollection->SkyPassInstanceCullingDrawParams);
+			BuildMeshRenderingCommands(GraphBuilder, EMeshPass::DebugViewMode, View, Scene->GPUScene, ParameterCollection->DebugViewModeInstanceCullingDrawParams);
+			BuildMeshRenderingCommands(GraphBuilder, EMeshPass::MeshDecal_SceneColorAndGBuffer, View, Scene->GPUScene, ParameterCollection->MeshDecalSceneColorAndGBufferInstanceCullingDrawParams);
+			BuildMeshRenderingCommands(GraphBuilder, StandardTranslucencyMeshPass, View, Scene->GPUScene, ParameterCollection->TranslucencyInstanceCullingDrawParams);
 		}
 
 		GraphBuilder.AddPass(
@@ -2172,14 +2186,14 @@ void FMobileSceneRenderer::RenderDeferredMultiPass(FRDGBuilder& GraphBuilder, FS
 		{
 			if (!bIsFullDepthPrepassEnabled)
 			{
-				View.ParallelMeshDrawCommandPasses[EMeshPass::DepthPass].BuildRenderingCommands(GraphBuilder, Scene->GPUScene, ParameterCollection->DepthPassInstanceCullingDrawParams);
+				BuildMeshRenderingCommands(GraphBuilder, EMeshPass::DepthPass, View, Scene->GPUScene, ParameterCollection->DepthPassInstanceCullingDrawParams);
 			}
-			View.ParallelMeshDrawCommandPasses[EMeshPass::BasePass].BuildRenderingCommands(GraphBuilder, Scene->GPUScene, PassParameters->InstanceCullingDrawParams);
-			View.ParallelMeshDrawCommandPasses[EMeshPass::SkyPass].BuildRenderingCommands(GraphBuilder, Scene->GPUScene, ParameterCollection->SkyPassInstanceCullingDrawParams);
-			View.ParallelMeshDrawCommandPasses[EMeshPass::DebugViewMode].BuildRenderingCommands(GraphBuilder, Scene->GPUScene, ParameterCollection->DebugViewModeInstanceCullingDrawParams);
+			BuildMeshRenderingCommands(GraphBuilder, EMeshPass::BasePass, View, Scene->GPUScene, PassParameters->InstanceCullingDrawParams);
+			BuildMeshRenderingCommands(GraphBuilder, EMeshPass::SkyPass, View, Scene->GPUScene, ParameterCollection->SkyPassInstanceCullingDrawParams);
+			BuildMeshRenderingCommands(GraphBuilder, EMeshPass::DebugViewMode, View, Scene->GPUScene, ParameterCollection->DebugViewModeInstanceCullingDrawParams);
 			if (bIsFullDepthPrepassEnabled)
 			{
-				View.ParallelMeshDrawCommandPasses[EMeshPass::MeshDecal_SceneColorAndGBuffer].BuildRenderingCommands(GraphBuilder, Scene->GPUScene, ParameterCollection->MeshDecalSceneColorAndGBufferInstanceCullingDrawParams);
+				BuildMeshRenderingCommands(GraphBuilder, EMeshPass::MeshDecal_SceneColorAndGBuffer, View, Scene->GPUScene, ParameterCollection->MeshDecalSceneColorAndGBufferInstanceCullingDrawParams);
 			}
 		}
 
@@ -2251,8 +2265,8 @@ void FMobileSceneRenderer::RenderDeferredMultiPass(FRDGBuilder& GraphBuilder, FS
 
 			if (Scene->GPUScene.IsEnabled())
 			{
-				View.ParallelMeshDrawCommandPasses[EMeshPass::MeshDecal_SceneColor].BuildRenderingCommands(GraphBuilder, Scene->GPUScene, PassParameters->InstanceCullingDrawParams);
-				View.ParallelMeshDrawCommandPasses[EMeshPass::MeshDecal_SceneColorAndGBuffer].BuildRenderingCommands(GraphBuilder, Scene->GPUScene, ParameterCollection->MeshDecalSceneColorAndGBufferInstanceCullingDrawParams);
+				BuildMeshRenderingCommands(GraphBuilder, EMeshPass::MeshDecal_SceneColor, View, Scene->GPUScene, PassParameters->InstanceCullingDrawParams);
+				BuildMeshRenderingCommands(GraphBuilder, EMeshPass::MeshDecal_SceneColorAndGBuffer, View, Scene->GPUScene, ParameterCollection->MeshDecalSceneColorAndGBufferInstanceCullingDrawParams);
 			}
 
 			GraphBuilder.AddPass(
@@ -2298,7 +2312,7 @@ void FMobileSceneRenderer::RenderDeferredMultiPass(FRDGBuilder& GraphBuilder, FS
 
 		if (Scene->GPUScene.IsEnabled())
 		{
-			View.ParallelMeshDrawCommandPasses[StandardTranslucencyMeshPass].BuildRenderingCommands(GraphBuilder, Scene->GPUScene, PassParameters->InstanceCullingDrawParams);
+			BuildMeshRenderingCommands(GraphBuilder, StandardTranslucencyMeshPass, View, Scene->GPUScene, PassParameters->InstanceCullingDrawParams);
 		}
 
 		const EMobileSSRQuality MobileSSRQuality = ActiveMobileSSRQuality(View, bShouldRenderVelocities);
diff --git a/Engine/Source/Runtime/Renderer/Public/MeshPassProcessor.h b/Engine/Source/Runtime/Renderer/Public/MeshPassProcessor.h
--- a/Engine/Source/Runtime/Renderer/Public/MeshPassProcessor.h
+++ b/Engine/Source/Runtime/Renderer/Public/MeshPassProcessor.h
@@ -1175,6 +1175,7 @@ struct FMeshDrawCommandOverrideArgs
 {
 	FRHIBuffer* InstanceBuffer;
 	FRHIBuffer* IndirectArgsBuffer;
+	FRHIUniformBuffer* BatchedPrimitiveUB;
 	uint32 InstanceDataByteOffset;
 	uint32 IndirectArgsByteOffset;
 	
@@ -1182,6 +1183,7 @@ struct FMeshDrawCommandOverrideArgs
 	{
 		InstanceBuffer = nullptr;
 		IndirectArgsBuffer = nullptr;
+		BatchedPrimitiveUB = nullptr;
 		InstanceDataByteOffset = 0u;
 		IndirectArgsByteOffset = 0u;
 	}
