Fixed uniform buffer parameters getting mutated after submission to RDG which breaks resource tracking.

diff --git a/Engine/Source/Runtime/Renderer/Private/LightFunctionAtlas.cpp b/Engine/Source/Runtime/Renderer/Private/LightFunctionAtlas.cpp
--- a/Engine/Source/Runtime/Renderer/Private/LightFunctionAtlas.cpp
+++ b/Engine/Source/Runtime/Renderer/Private/LightFunctionAtlas.cpp
@@ -559,12 +559,12 @@ TRDGUniformBufferRef<FLightFunctionAtlasGlobalParameters> FLightFunctionAtlas::G
 
 FLightFunctionAtlasGlobalParameters* FLightFunctionAtlas::GetDefaultLightFunctionAtlasGlobalParametersStruct(FRDGBuilder& GraphBuilder)
 {
-	static FLightFunctionAtlasGlobalParameters DefaultLightFunctionAtlasGlobalParameters;
-	DefaultLightFunctionAtlasGlobalParameters.LightFunctionAtlasTexture = GSystemTextures.GetWhiteDummy(GraphBuilder);
-	DefaultLightFunctionAtlasGlobalParameters.LightFunctionAtlasSampler = TStaticSamplerState<SF_Bilinear>::GetRHI();
-	DefaultLightFunctionAtlasGlobalParameters.LightInfoDataBuffer = GraphBuilder.CreateSRV(GSystemTextures.GetDefaultStructuredBuffer(GraphBuilder, sizeof(FVector4f) * 1, 0.0f), PF_A32B32G32R32F);
-	DefaultLightFunctionAtlasGlobalParameters.Slot_UVSize = 1.0f;
-	return &DefaultLightFunctionAtlasGlobalParameters;
+	FLightFunctionAtlasGlobalParameters* DefaultLightFunctionAtlasGlobalParameters = GraphBuilder.AllocParameters<FLightFunctionAtlasGlobalParameters>();
+	DefaultLightFunctionAtlasGlobalParameters->LightFunctionAtlasTexture = GSystemTextures.GetWhiteDummy(GraphBuilder);
+	DefaultLightFunctionAtlasGlobalParameters->LightFunctionAtlasSampler = TStaticSamplerState<SF_Bilinear>::GetRHI();
+	DefaultLightFunctionAtlasGlobalParameters->LightInfoDataBuffer = GraphBuilder.CreateSRV(GSystemTextures.GetDefaultStructuredBuffer(GraphBuilder, sizeof(FVector4f) * 1, 0.0f), PF_A32B32G32R32F);
+	DefaultLightFunctionAtlasGlobalParameters->Slot_UVSize = 1.0f;
+	return DefaultLightFunctionAtlasGlobalParameters;
 }
 
 TRDGUniformBufferRef<FLightFunctionAtlasGlobalParameters> FLightFunctionAtlas::GetDefaultLightFunctionAtlasGlobalParameters(FRDGBuilder& GraphBuilder)
