Custom Render Pass:  Fixed bug where objects rendered to depth prepass via the velocity pass were being skipped, leading to missing objects in the depth pass, and incorrect sorting in the base pass.  The velocity pass was skipped to avoid polluting the velocity buffer with Custom Render Pass results, which weren't used downstream, but would later would need to be cleared.  The substitute fix is to skip binding the velocity buffer render target.

#rnx

diff --git a/Engine/Source/Runtime/Renderer/Private/DeferredShadingRenderer.cpp b/Engine/Source/Runtime/Renderer/Private/DeferredShadingRenderer.cpp
--- a/Engine/Source/Runtime/Renderer/Private/DeferredShadingRenderer.cpp
+++ b/Engine/Source/Runtime/Renderer/Private/DeferredShadingRenderer.cpp
@@ -2029,10 +2029,10 @@ void FDeferredShadingSceneRenderer::Render(FRDGBuilder& GraphBuilder, const FSce
 			}
 
 			// special pass for DDM_AllOpaqueNoVelocity, which uses the velocity pass to finish the early depth pass write
-			if (bShouldRenderVelocities && Scene->EarlyZPassMode == DDM_AllOpaqueNoVelocity && RendererOutput == ERendererOutput::FinalSceneColor && !InViews[0].CustomRenderPass)
+			if (bShouldRenderVelocities && Scene->EarlyZPassMode == DDM_AllOpaqueNoVelocity && RendererOutput == ERendererOutput::FinalSceneColor)
 			{
-				// Render the velocities of movable objects
-				RenderVelocities(GraphBuilder, InViews, SceneTextures, EVelocityPass::Opaque, bForceVelocityOutput);
+				// Render the velocities of movable objects.  Don't bind the velocity render target for custom render passes (it's not used downstream), to avoid needing to clear it again.
+				RenderVelocities(GraphBuilder, InViews, SceneTextures, EVelocityPass::Opaque, bForceVelocityOutput, /*bBindRenderTarget=*/ InViews[0].CustomRenderPass == nullptr);
 			}
 		}
 
diff --git a/Engine/Source/Runtime/Renderer/Private/SceneRendering.h b/Engine/Source/Runtime/Renderer/Private/SceneRendering.h
--- a/Engine/Source/Runtime/Renderer/Private/SceneRendering.h
+++ b/Engine/Source/Runtime/Renderer/Private/SceneRendering.h
@@ -2585,7 +2585,8 @@ protected:
 		TArrayView<FViewInfo> InViews,
 		const FSceneTextures& SceneTextures,
 		EVelocityPass VelocityPass,
-		bool bForceVelocity);
+		bool bForceVelocity,
+		bool bBindRenderTarget = true);
 
 	void RenderMeshDistanceFieldVisualization(FRDGBuilder& GraphBuilder, const FMinimalSceneTextures& SceneTextures);
 
diff --git a/Engine/Source/Runtime/Renderer/Private/VelocityRendering.cpp b/Engine/Source/Runtime/Renderer/Private/VelocityRendering.cpp
--- a/Engine/Source/Runtime/Renderer/Private/VelocityRendering.cpp
+++ b/Engine/Source/Runtime/Renderer/Private/VelocityRendering.cpp
@@ -244,7 +244,8 @@ void FSceneRenderer::RenderVelocities(
 	TArrayView<FViewInfo> InViews,
 	const FSceneTextures& SceneTextures,
 	EVelocityPass VelocityPass,
-	bool bForceVelocity)
+	bool bForceVelocity,
+	bool bBindRenderTarget)
 {
 	RDG_CSV_STAT_EXCLUSIVE_SCOPE(GraphBuilder, RenderVelocities);
 	SCOPED_NAMED_EVENT(FSceneRenderer_RenderVelocities, FColor::Emerald);
@@ -306,9 +307,12 @@ void FSceneRenderer::RenderVelocities(
 				ERenderTargetLoadAction::ELoad,
 				ExclusiveDepthStencil);
 
-			PassParameters->RenderTargets[0] = FRenderTargetBinding(SceneTextures.Velocity, (bNeedsClearMask & View.GPUMask.GetNative()) ? ERenderTargetLoadAction::EClear : ERenderTargetLoadAction::ELoad);
+			if (bBindRenderTarget)
+			{
+				PassParameters->RenderTargets[0] = FRenderTargetBinding(SceneTextures.Velocity, (bNeedsClearMask & View.GPUMask.GetNative()) ? ERenderTargetLoadAction::EClear : ERenderTargetLoadAction::ELoad);
+				bNeedsClearMask &= ~View.GPUMask.GetNative();
+			}
 			PassParameters->RenderTargets.MultiViewCount = (View.bIsMobileMultiViewEnabled) ? 2 : (View.Aspects.IsMobileMultiViewEnabled() ? 1 : 0);
-			bNeedsClearMask &= ~View.GPUMask.GetNative();
 
 			if (bIsParallelVelocity)
 			{
