Height Fog - added visible in reflection and in realtime capture control on component.

#tests pc dx12 editor
#rnx

diff --git a/Engine/Source/Runtime/Engine/Classes/Components/ExponentialHeightFogComponent.h b/Engine/Source/Runtime/Engine/Classes/Components/ExponentialHeightFogComponent.h
--- a/Engine/Source/Runtime/Engine/Classes/Components/ExponentialHeightFogComponent.h
+++ b/Engine/Source/Runtime/Engine/Classes/Components/ExponentialHeightFogComponent.h
@@ -191,6 +191,14 @@ class UExponentialHeightFogComponent : public USceneComponent
 	UPROPERTY(EditAnywhere, AdvancedDisplay, BlueprintReadOnly, Category = Rendering)
 	uint8 bRenderInMainPass : 1;
 
+	/** If true, this component will be visible in reflection captures. */
+	UPROPERTY(EditAnywhere, AdvancedDisplay, BlueprintReadOnly, Category = Rendering)
+	uint8 bVisibleInReflectionCaptures : 1;
+
+	/** If true, this component will be visible in real-time sky light reflection captures. */
+	UPROPERTY(EditAnywhere, AdvancedDisplay, BlueprintReadOnly, Category = Rendering)
+	uint8 bVisibleInRealTimeSkyCaptures : 1;
+
 
 public:
 	UFUNCTION(BlueprintCallable, Category="Rendering|Components|ExponentialHeightFog")
diff --git a/Engine/Source/Runtime/Engine/Private/Components/ExponentialHeightFogComponent.cpp b/Engine/Source/Runtime/Engine/Private/Components/ExponentialHeightFogComponent.cpp
--- a/Engine/Source/Runtime/Engine/Private/Components/ExponentialHeightFogComponent.cpp
+++ b/Engine/Source/Runtime/Engine/Private/Components/ExponentialHeightFogComponent.cpp
@@ -35,6 +35,8 @@ constexpr bool UseExponentialHeightFogStateStream = true;
 	UEXPONENTIALHEIGHTFOG_VARIABLE(bool, b, OverrideLightColorsWithFogInscatteringColors) \
 	UEXPONENTIALHEIGHTFOG_VARIABLE(bool, b, Holdout) \
 	UEXPONENTIALHEIGHTFOG_VARIABLE(bool, b, RenderInMainPass) \
+	UEXPONENTIALHEIGHTFOG_VARIABLE(bool, b, VisibleInReflectionCaptures) \
+	UEXPONENTIALHEIGHTFOG_VARIABLE(bool, b, VisibleInRealTimeSkyCaptures) \
 	UEXPONENTIALHEIGHTFOG_VARIABLES_WITH_SETFUNCTION
 
 #define UEXPONENTIALHEIGHTFOG_VARIABLES_WITH_SETFUNCTION \
@@ -90,6 +92,8 @@ UExponentialHeightFogComponent::UExponentialHeightFogComponent(const FObjectInit
 
 	bHoldout = false;
 	bRenderInMainPass = true;
+	bVisibleInReflectionCaptures = true;
+	bVisibleInRealTimeSkyCaptures = true;
 
 	VolumetricFogScatteringDistribution = .2f;
 	VolumetricFogAlbedo = FColor::White;
diff --git a/Engine/Source/Runtime/Engine/Public/StateStream/ExponentialHeightFogStateStream.h b/Engine/Source/Runtime/Engine/Public/StateStream/ExponentialHeightFogStateStream.h
--- a/Engine/Source/Runtime/Engine/Public/StateStream/ExponentialHeightFogStateStream.h
+++ b/Engine/Source/Runtime/Engine/Public/StateStream/ExponentialHeightFogStateStream.h
@@ -90,7 +90,10 @@ struct FExponentialHeightFogDynamicState
 	uint8 bHoldout : 1;
 	UPROPERTY()
 	uint8 bRenderInMainPass : 1;
-
+	UPROPERTY()
+	uint8 bVisibleInReflectionCaptures : 1;
+	UPROPERTY()
+	uint8 bVisibleInRealTimeSkyCaptures : 1;
 	UPROPERTY()
 	float Height;
 };
diff --git a/Engine/Source/Runtime/Renderer/Private/FogRendering.cpp b/Engine/Source/Runtime/Renderer/Private/FogRendering.cpp
--- a/Engine/Source/Runtime/Renderer/Private/FogRendering.cpp
+++ b/Engine/Source/Runtime/Renderer/Private/FogRendering.cpp
@@ -329,6 +329,19 @@ static FFogPassParameters* CreateDefaultFogPassParameters(
 	return PassParameters;
 }
 
+static bool ShouldRenderFogForView(
+	const FScene* Scene,
+	const FViewInfo& View)
+{
+	const bool bViewIsReflectionCapture = View.bIsReflectionCapture;
+	bool bFogRendersInReflectionCapture = true;
+	if (Scene && Scene->HasAnyExponentialHeightFog())
+	{
+		bFogRendersInReflectionCapture = Scene->ExponentialFogs[0].bVisibleInReflectionCaptures;
+	}
+	return !bViewIsReflectionCapture || bFogRendersInReflectionCapture;
+}
+
 static void RenderViewFog(
 	FRHICommandList& RHICmdList, 
 	const FViewInfo& View, 
@@ -446,7 +459,7 @@ void RenderFogOnClouds(
 	const bool bShouldRenderVolumetricFog,
 	const bool bUseVolumetricRenderTarget)
 {
-	if (Scene->ExponentialFogs.Num() > 0)
+	if (Scene->ExponentialFogs.Num() > 0 && ShouldRenderFogForView(Scene, View))
 	{
 		RDG_EVENT_SCOPE_STAT(GraphBuilder, Fog, "ExponentialHeightFog on Clouds");
 		RDG_GPU_STAT_SCOPE(GraphBuilder, Fog);
@@ -506,7 +519,14 @@ void FDeferredShadingSceneRenderer::RenderFog(
 		for(int32 ViewIndex = 0;ViewIndex < Views.Num();ViewIndex++)
 		{
 			const FViewInfo& View = Views[ViewIndex];
-			RDG_EVENT_SCOPE_CONDITIONAL(GraphBuilder, Views.Num() > 1, "View%d", ViewIndex);
+
+			const bool bShouldRenderFogForView = ShouldRenderFogForView(Scene, View);
+			if (!bShouldRenderFogForView)
+			{
+				continue;
+			}
+
+			RDG_EVENT_SCOPE_CONDITIONAL(GraphBuilder, bShouldRenderFogForView && Views.Num() > 1, "View%d", ViewIndex);
 			RDG_GPU_MASK_SCOPE(GraphBuilder, View.GPUMask);
 
 			TRDGUniformBufferRef<FFogUniformParameters> FogUniformBuffer = CreateFogUniformBuffer(GraphBuilder, View);
@@ -548,6 +568,13 @@ void FDeferredShadingSceneRenderer::RenderUnderWaterFog(
 		for (int32 ViewIndex = 0; ViewIndex < Views.Num(); ViewIndex++)
 		{
 			const FViewInfo& View = Views[ViewIndex];
+
+			const bool bShouldRenderFogForView = ShouldRenderFogForView(Scene, View);
+			if (!bShouldRenderFogForView)
+			{
+				continue;
+			}
+
 			if (View.IsUnderwater() || CVarUnderwaterFogWhenCameraIsAboveWater.GetValueOnRenderThread())
 			{
 				RDG_EVENT_SCOPE_CONDITIONAL(GraphBuilder, Views.Num() > 1, "View%d", ViewIndex);
diff --git a/Engine/Source/Runtime/Renderer/Private/ReflectionEnvironmentRealTimeCapture.cpp b/Engine/Source/Runtime/Renderer/Private/ReflectionEnvironmentRealTimeCapture.cpp
--- a/Engine/Source/Runtime/Renderer/Private/ReflectionEnvironmentRealTimeCapture.cpp
+++ b/Engine/Source/Runtime/Renderer/Private/ReflectionEnvironmentRealTimeCapture.cpp
@@ -678,7 +678,8 @@ void FScene::AllocateAndCaptureFrameSkyEnvMap(
 						}
 
 						// Also render the height fog as part of the sky render pass when time slicing is enabled.
-						if (Scene && Scene->ExponentialFogs.Num() > 0)
+						
+						if (Scene && Scene->HasAnyExponentialHeightFog() && Scene->ExponentialFogs[0].bVisibleInRealTimeSkyCaptures)
 						{
 							FRenderRealTimeReflectionHeightFogVS::FPermutationDomain VsPermutationVector;
 							TShaderMapRef<FRenderRealTimeReflectionHeightFogVS> VertexShader(GetGlobalShaderMap(SkyRC.FeatureLevel), VsPermutationVector);
diff --git a/Engine/Source/Runtime/Renderer/Private/SceneCore.cpp b/Engine/Source/Runtime/Renderer/Private/SceneCore.cpp
--- a/Engine/Source/Runtime/Renderer/Private/SceneCore.cpp
+++ b/Engine/Source/Runtime/Renderer/Private/SceneCore.cpp
@@ -429,6 +429,8 @@ FExponentialHeightFogSceneInfo::FExponentialHeightFogSceneInfo(uint64 InId, cons
 	bOverrideLightColorsWithFogInscatteringColors = State.bOverrideLightColorsWithFogInscatteringColors;
 	bHoldout = State.bHoldout;
 	bRenderInMainPass = State.bRenderInMainPass;
+	bVisibleInReflectionCaptures = State.bVisibleInReflectionCaptures;
+	bVisibleInRealTimeSkyCaptures = State.bVisibleInRealTimeSkyCaptures;
 
 	VolumetricFogStartDistance = State.VolumetricFogStartDistance;
 	VolumetricFogNearFadeInDistance = State.VolumetricFogNearFadeInDistance;
diff --git a/Engine/Source/Runtime/Renderer/Private/SceneCore.h b/Engine/Source/Runtime/Renderer/Private/SceneCore.h
--- a/Engine/Source/Runtime/Renderer/Private/SceneCore.h
+++ b/Engine/Source/Runtime/Renderer/Private/SceneCore.h
@@ -182,6 +182,8 @@ public:
 	bool bOverrideLightColorsWithFogInscatteringColors;
 	bool bHoldout;
 	bool bRenderInMainPass;
+	bool bVisibleInReflectionCaptures;
+	bool bVisibleInRealTimeSkyCaptures;
 	float VolumetricFogStartDistance;
 	float VolumetricFogNearFadeInDistance;
 
