Fixed race condition between shadow scene renderer and virtual shadow map cache manager that caused (rare) crashes.

diff --git a/Engine/Source/Runtime/Renderer/Private/Shadows/ShadowScene.cpp b/Engine/Source/Runtime/Renderer/Private/Shadows/ShadowScene.cpp
--- a/Engine/Source/Runtime/Renderer/Private/Shadows/ShadowScene.cpp
+++ b/Engine/Source/Runtime/Renderer/Private/Shadows/ShadowScene.cpp
@@ -115,6 +115,7 @@ void FShadowScene::PostLightsUpdate(FRDGBuilder& GraphBuilder, const FLightScene
 
 void FShadowScene::PostSceneUpdate(const FScenePreUpdateChangeSet &ScenePreUpdateChangeSet, const FScenePostUpdateChangeSet &ScenePostUpdateChangeSet)
 {
+	// Note: if moving this to an async task, also make sure to properly sync/depend any use of `AlwaysInvalidatingPrimitives` (which is exposed in GetAlwaysInvalidatingPrimitives)
 	for (FPrimitiveSceneInfo* PrimitiveSceneInfo : ScenePreUpdateChangeSet.RemovedPrimitiveSceneInfos)
 	{
 		if (PrimitiveSceneInfo->Proxy->GetShadowCacheInvalidationBehavior() == EShadowCacheInvalidationBehavior::Always)
@@ -252,4 +253,9 @@ void FShadowScene::DebugRender(TArrayView<FViewInfo> Views)
 	}
 
 #endif
-}
\ No newline at end of file
+}
+
+void FShadowScene::WaitForSceneLightsUpdateTask()
+{
+	SceneChangeUpdateTask.Wait();
+}
diff --git a/Engine/Source/Runtime/Renderer/Private/Shadows/ShadowScene.h b/Engine/Source/Runtime/Renderer/Private/Shadows/ShadowScene.h
--- a/Engine/Source/Runtime/Renderer/Private/Shadows/ShadowScene.h
+++ b/Engine/Source/Runtime/Renderer/Private/Shadows/ShadowScene.h
@@ -50,6 +50,11 @@ public:
 	// List of always invalidating primitives, if this gets too popular perhaps a TSet or some such is more appropriate for performance scaling.
 	TArrayView<FPrimitiveSceneInfo*> GetAlwaysInvalidatingPrimitives() { return AlwaysInvalidatingPrimitives; }
 
+	/**
+	 * Wait for any scene update task started in the PostLightsUpdate.
+	 */
+	void WaitForSceneLightsUpdateTask();
+
 private:
 	friend class FShadowSceneRenderer;
 
diff --git a/Engine/Source/Runtime/Renderer/Private/VirtualShadowMaps/VirtualShadowMapCacheManager.cpp b/Engine/Source/Runtime/Renderer/Private/VirtualShadowMaps/VirtualShadowMapCacheManager.cpp
--- a/Engine/Source/Runtime/Renderer/Private/VirtualShadowMaps/VirtualShadowMapCacheManager.cpp
+++ b/Engine/Source/Runtime/Renderer/Private/VirtualShadowMaps/VirtualShadowMapCacheManager.cpp
@@ -1732,6 +1732,9 @@ void FVirtualShadowMapInvalidationSceneUpdater::PreSceneUpdate(FRDGBuilder& Grap
 		}
 	}
 
+	// Note: waits for the light update task that the shadow scene uses because it calls ProcessRemovedLights from that thread.
+	CacheManager.Scene->ShadowScene->WaitForSceneLightsUpdateTask();
+
 	// Needs to be called before the first time we start adding invalidations.
 	// There may be a way to avoid doing this both in pre and post, but it is pretty light if there is nothing to do anyways.
 	CacheManager.ReallocatePersistentPrimitiveIndices();
