Fixed CIS errrors caused by a bad merge.

#rnx

diff --git a/Engine/Source/Runtime/Renderer/Private/VirtualShadowMaps/VirtualShadowMapCacheManager.cpp b/Engine/Source/Runtime/Renderer/Private/VirtualShadowMaps/VirtualShadowMapCacheManager.cpp
--- a/Engine/Source/Runtime/Renderer/Private/VirtualShadowMaps/VirtualShadowMapCacheManager.cpp
+++ b/Engine/Source/Runtime/Renderer/Private/VirtualShadowMaps/VirtualShadowMapCacheManager.cpp
@@ -1733,7 +1733,7 @@ void FVirtualShadowMapInvalidationSceneUpdater::PreSceneUpdate(FRDGBuilder& Grap
 	}
 
 	// Note: waits for the light update task that the shadow scene uses because it calls ProcessRemovedLights from that thread.
-	CacheManager.Scene->ShadowScene->WaitForSceneLightsUpdateTask();
+	CacheManager.Scene.ShadowScene->WaitForSceneLightsUpdateTask();
 
 	// Needs to be called before the first time we start adding invalidations.
 	// There may be a way to avoid doing this both in pre and post, but it is pretty light if there is nothing to do anyways.
