Path Tracing: Fix missing function on platforms without raytracing support

#rnx

diff --git a/Engine/Source/Runtime/Renderer/Private/PathTracing.cpp b/Engine/Source/Runtime/Renderer/Private/PathTracing.cpp
--- a/Engine/Source/Runtime/Renderer/Private/PathTracing.cpp
+++ b/Engine/Source/Runtime/Renderer/Private/PathTracing.cpp
@@ -4389,6 +4389,11 @@ namespace PathTracing
 	{
 		return false;
 	}
+
+	bool NeedsTonemapping()
+	{
+		return false;
+	}
 }
 
 #endif
