Move configuration of view state / view rects into a task that operates on all renderers and overlaps with the scene update.

diff --git a/Engine/Source/Runtime/Engine/Private/DynamicResolution.cpp b/Engine/Source/Runtime/Engine/Private/DynamicResolution.cpp
--- a/Engine/Source/Runtime/Engine/Private/DynamicResolution.cpp
+++ b/Engine/Source/Runtime/Engine/Private/DynamicResolution.cpp
@@ -576,7 +576,7 @@ void FDynamicResolutionHeuristicProxy::RefreshCurrentFrameResolutionFraction_Ren
 
 void FDynamicResolutionHeuristicProxy::SetTemporalUpscaler(const UE::Renderer::Private::ITemporalUpscaler* InTemporalUpscaler)
 {
-	check(IsInRenderingThread());
+	check(IsInParallelRenderingThread());
 
 	float NewMinResolutionFraction = InTemporalUpscaler ? InTemporalUpscaler->GetMinUpsampleResolutionFraction() : ISceneViewFamilyScreenPercentage::kMinResolutionFraction;
 	float NewMaxResolutionFraction = InTemporalUpscaler ? InTemporalUpscaler->GetMaxUpsampleResolutionFraction() : ISceneViewFamilyScreenPercentage::kMaxResolutionFraction;
@@ -626,7 +626,7 @@ void FDynamicResolutionHeuristicProxy::RefreshCurrentFrameResolutionFractionUppe
 void FDynamicResolutionHeuristicProxy::RefreshHeuristicStats_RenderThread()
 {
 #if STATS
-	check(IsInRenderingThread());
+	check(IsInParallelRenderingThread());
 	for (TLinkedList<DynamicRenderScaling::FBudget*>::TIterator BudgetIt(DynamicRenderScaling::FBudget::GetGlobalList()); BudgetIt; BudgetIt.Next())
 	{
 		const DynamicRenderScaling::FBudget& Budget = **BudgetIt;
@@ -655,7 +655,7 @@ void FDynamicResolutionHeuristicProxy::RefreshHeuristicStats_RenderThread()
 // static
 DynamicRenderScaling::TMap<float> FDynamicResolutionHeuristicProxy::GetResolutionFractionUpperBounds() const
 {
-	check(IsInGameThread() || IsInRenderingThread());
+	check(IsInGameThread() || IsInParallelRenderingThread());
 	return CurrentFrameMaxResolutionFractions;
 }
 
diff --git a/Engine/Source/Runtime/Engine/Private/LegacyScreenPercentageDriver.cpp b/Engine/Source/Runtime/Engine/Private/LegacyScreenPercentageDriver.cpp
--- a/Engine/Source/Runtime/Engine/Private/LegacyScreenPercentageDriver.cpp
+++ b/Engine/Source/Runtime/Engine/Private/LegacyScreenPercentageDriver.cpp
@@ -146,7 +146,7 @@ DynamicRenderScaling::TMap<float> FLegacyScreenPercentageDriver::GetResolutionFr
 
 DynamicRenderScaling::TMap<float> FLegacyScreenPercentageDriver::GetResolutionFractions_RenderThread() const
 {
-	check(IsInRenderingThread());
+	check(IsInParallelRenderingThread());
 
 	DynamicRenderScaling::TMap<float> ResolutionFractions;
 	ResolutionFractions.SetAll(1.0f);
diff --git a/Engine/Source/Runtime/RenderCore/Public/RenderCounters.h b/Engine/Source/Runtime/RenderCore/Public/RenderCounters.h
--- a/Engine/Source/Runtime/RenderCore/Public/RenderCounters.h
+++ b/Engine/Source/Runtime/RenderCore/Public/RenderCounters.h
@@ -21,7 +21,6 @@ public:
 
 	void AddViewStatistics(uint32 PixelRenderCount, uint32 PixelDisplayCount)
 	{
-		check(IsInRenderingThread());
 		CurrentPixelRenderCount += PixelRenderCount;
 		CurrentPixelDisplayCount += PixelDisplayCount;
 	}
diff --git a/Engine/Source/Runtime/Renderer/Private/SceneRendering.cpp b/Engine/Source/Runtime/Renderer/Private/SceneRendering.cpp
--- a/Engine/Source/Runtime/Renderer/Private/SceneRendering.cpp
+++ b/Engine/Source/Runtime/Renderer/Private/SceneRendering.cpp
@@ -3138,9 +3138,8 @@ FSceneRenderer::ERendererOutput FSceneRenderer::GetRendererOutput() const
 	return ERendererOutput::FinalSceneColor;
 }
 
-void FSceneRenderer::PrepareViewRectsForRendering(FRHICommandListImmediate& RHICmdList)
+void FSceneRenderer::PrepareViewRectsForRendering()
 {
-	check(IsInRenderingThread());
 	TRACE_CPUPROFILER_EVENT_SCOPE(PrepareViewRectsForRendering);
 
 	// If we support screen percentage, update the dynamic resolution state with our current temporal upscaler, which clamps the screen percentage to its supported range.
@@ -3172,20 +3171,7 @@ void FSceneRenderer::PrepareViewRectsForRendering(FRHICommandListImmediate& RHIC
 		}
 
 		ComputeFamilySize();
-		
-		// Notify StereoRenderingDevice about new ViewRects
-		if (GEngine->StereoRenderingDevice.IsValid() && ViewFamily.EngineShowFlags.StereoRendering)
-		{
-			for (int32 i = 0; i < Views.Num(); i++)
-			{
-				FViewInfo& View = Views[i];
-				
-				if (IStereoRendering::IsStereoEyePass(View.StereoPass))
-				{
-					GEngine->StereoRenderingDevice->SetFinalViewRect(RHICmdList, View.StereoViewIndex, View.ViewRect);
-				}
-			}
-		}
+
 		return;
 	}
 
@@ -3358,22 +3344,6 @@ void FSceneRenderer::PrepareViewRectsForRendering(FRHICommandListImmediate& RHIC
 
 	ComputeFamilySize();
 
-	// Notify StereoRenderingDevice about new ViewRects
-	if (GEngine->StereoRenderingDevice.IsValid() && ViewFamily.EngineShowFlags.StereoRendering)
-	{
-		for (const FViewInfo& View : Views)
-		{
-			// if we have an upscale pass, the final rect is _unscaled_ for the compositor
-			const FIntRect OutputViewRect =
-				(View.PrimaryScreenPercentageMethod == EPrimaryScreenPercentageMethod::RawOutput) ? View.ViewRect : View.UnscaledViewRect;
-
-			if (IStereoRendering::IsStereoEyePass(View.StereoPass))
-			{
-				GEngine->StereoRenderingDevice->SetFinalViewRect(RHICmdList, View.StereoViewIndex, OutputViewRect);
-			}
-		}
-	}
-
 	for (FCustomRenderPassInfo& PassInfo : CustomRenderPassInfos)
 	{
 		for (FViewInfo& View : PassInfo.Views)
@@ -3754,7 +3724,7 @@ void FSceneRenderer::FlushCrossGPUFences(FRDGBuilder& GraphBuilder)
 void FSceneRenderer::ComputeFamilySize()
 {
 	check(FamilySize.X == 0);
-	check(IsInRenderingThread());
+	check(IsInParallelRenderingThread());
 
 	// Calculate the screen extents of the view family.
 	bool bInitializedExtents = false;
@@ -3859,8 +3829,38 @@ IVisibilityTaskData* FSceneRenderer::OnRenderBegin(FRDGBuilder& GraphBuilder, co
 	UE::Tasks::FTaskEvent GPUSceneUpdateTaskPrerequisites{ UE_SOURCE_LOCATION };
 	SceneUpdateParameters.GPUSceneUpdateTaskPrerequisites = GPUSceneUpdateTaskPrerequisites;
 
+	UE::Tasks::FTask PrepareViewStateTask;
+
+	if (SceneUpdateInputs)
+	{
+		PrepareViewStateTask = UE::Tasks::Launch(UE_SOURCE_LOCATION, [SceneUpdateInputs]
+		{
+			TRACE_CPUPROFILER_EVENT_SCOPE(PrepareViewRects);
+			FOptionalTaskTagScope TagScope(ETaskTag::EParallelRenderingThread);
+		
+			for (FSceneRenderer* Renderer : SceneUpdateInputs->Renderers)
+			{
+				Renderer->PrepareViewRectsForRendering();
+	
+				InitializeSceneTexturesConfig(Renderer->ViewFamily.SceneTexturesConfig, Renderer->ViewFamily);
+				const FSceneTexturesConfig& SceneTexturesConfig = Renderer->GetActiveSceneTexturesConfig();
+
+				// Custom render passes have their own view family structure, so they can have separate EngineShowFlags, so the SceneTexturesConfig
+				// needs to be copied.  The FSceneTextures structure itself is pointer shared, and doesn't need to be copied.
+				for (FCustomRenderPassInfo& CustomRenderPass : Renderer->CustomRenderPassInfos)
+				{
+					CustomRenderPass.ViewFamily.SceneTexturesConfig = Renderer->ViewFamily.SceneTexturesConfig;
+				}
+	
+				Renderer->PrepareViewStateForVisibility(SceneTexturesConfig);
+			}
+		});
+	}
+
 	SceneUpdateParameters.Callbacks.PostStaticMeshUpdate = [&] (const UE::Tasks::FTask& StaticMeshUpdateTask)
 	{
+		PrepareViewStateTask.Wait();
+
 #if RHI_RAYTRACING
 		RayTracing::OnRenderBegin(*Scene, Views, ViewFamily);
 		InitializeRayTracingFlags_RenderThread();
@@ -3894,22 +3894,7 @@ IVisibilityTaskData* FSceneRenderer::OnRenderBegin(FRDGBuilder& GraphBuilder, co
 				}
 			}
 		}
-	
-		PrepareViewRectsForRendering(GraphBuilder.RHICmdList);
-	
-		InitializeSceneTexturesConfig(ViewFamily.SceneTexturesConfig, ViewFamily);
-		FSceneTexturesConfig& SceneTexturesConfig = GetActiveSceneTexturesConfig();
-		FSceneTexturesConfig::Set(SceneTexturesConfig);
 
-		// Custom render passes have their own view family structure, so they can have separate EngineShowFlags, so the SceneTexturesConfig
-		// needs to be copied.  The FSceneTextures structure itself is pointer shared, and doesn't need to be copied.
-		for (FCustomRenderPassInfo& CustomRenderPass : CustomRenderPassInfos)
-		{
-			CustomRenderPass.ViewFamily.SceneTexturesConfig = ViewFamily.SceneTexturesConfig;
-		}
-	
-		PrepareViewStateForVisibility(SceneTexturesConfig);
-	
 		// Run Groom LOD selection prior to visibility for selecting appropriate LOD & geometry type
 		if (IsGroomEnabled())
 		{
@@ -3955,6 +3940,24 @@ IVisibilityTaskData* FSceneRenderer::OnRenderBegin(FRDGBuilder& GraphBuilder, co
 		SceneUpdateParameters.Callbacks.PostStaticMeshUpdate(UE::Tasks::FTask{});
 	}
 
+	FSceneTexturesConfig::Set(GetActiveSceneTexturesConfig());
+
+	// Notify StereoRenderingDevice about new ViewRects
+	if (GEngine->StereoRenderingDevice.IsValid() && ViewFamily.EngineShowFlags.StereoRendering)
+	{
+		for (const FViewInfo& View : Views)
+		{
+			// if we have an upscale pass, the final rect is _unscaled_ for the compositor
+			const FIntRect OutputViewRect =
+				(View.PrimaryScreenPercentageMethod == EPrimaryScreenPercentageMethod::RawOutput) ? View.ViewRect : View.UnscaledViewRect;
+
+			if (IStereoRendering::IsStereoEyePass(View.StereoPass))
+			{
+				GEngine->StereoRenderingDevice->SetFinalViewRect(GraphBuilder.RHICmdList, View.StereoViewIndex, OutputViewRect);
+			}
+		}
+	}
+
 	return VisibilityTaskData;
 }
 
diff --git a/Engine/Source/Runtime/Renderer/Private/SceneRendering.h b/Engine/Source/Runtime/Renderer/Private/SceneRendering.h
--- a/Engine/Source/Runtime/Renderer/Private/SceneRendering.h
+++ b/Engine/Source/Runtime/Renderer/Private/SceneRendering.h
@@ -2175,7 +2175,7 @@ public:
 	virtual bool AllowSimpleLights() const;
 
 	/** Setups FViewInfo::ViewRect according to ViewFamilly's ScreenPercentageInterface. */
-	void PrepareViewRectsForRendering(FRHICommandListImmediate& RHICmdList);
+	void PrepareViewRectsForRendering();
 
 #if WITH_MGPU
 	/**
