A few CPU shadow/VSM optimizations:
Avoid storing/copying large matrix sets around if they haven't changed

#rnx

diff --git a/Engine/Source/Runtime/Renderer/Private/Shadows/ShadowSceneRenderer.cpp b/Engine/Source/Runtime/Renderer/Private/Shadows/ShadowSceneRenderer.cpp
--- a/Engine/Source/Runtime/Renderer/Private/Shadows/ShadowSceneRenderer.cpp
+++ b/Engine/Source/Runtime/Renderer/Private/Shadows/ShadowSceneRenderer.cpp
@@ -250,38 +250,22 @@ static float GetResolutionLODBiasLocal(float LightMobilityFactor, float LightLOD
 		LightMobilityFactor) + LightLODBias;
 }
 
-FVirtualShadowMapProjectionShaderData FShadowSceneRenderer::GetLocalLightProjectionShaderData(
-	float ResolutionLODBiasLocal,
+void FShadowSceneRenderer::UpdateLocalLightProjectionShaderDataMatrices(
 	const FProjectedShadowInfo* ProjectedShadowInfo,
-	int32 MapIndex) const
+	int32 MapIndex,
+	FVirtualShadowMapProjectionShaderData* OutProjectionShaderData) const
 {
-	TSharedPtr<FVirtualShadowMapPerLightCacheEntry> CacheEntry = ProjectedShadowInfo->VirtualShadowMapPerLightCacheEntry;
-	check(CacheEntry.IsValid());
-
-	int32 VirtualShadowMapId = ProjectedShadowInfo->VirtualShadowMapId + MapIndex;
-	bool bIsSinglePageSM = FVirtualShadowMapArray::IsSinglePage(VirtualShadowMapId);
-	check(VirtualShadowMapId != INDEX_NONE && CacheEntry->bIsDistantLight == bIsSinglePageSM);
-
-	uint32 Flags = CacheEntry->IsUncached() ? VSM_PROJ_FLAG_UNCACHED : 0U;
-
 	const FViewMatrices ViewMatrices = ProjectedShadowInfo->GetShadowDepthRenderingViewMatrices(MapIndex, true);
 	const FDFVector3 PreViewTranslation(ProjectedShadowInfo->PreShadowTranslation);
 
-	FVirtualShadowMapProjectionShaderData Data; 
-	Data.ShadowViewToClipMatrix					= FMatrix44f(ViewMatrices.GetProjectionMatrix());
-	Data.TranslatedWorldToShadowUVMatrix		= FMatrix44f(CalcTranslatedWorldToShadowUVMatrix( ViewMatrices.GetTranslatedViewMatrix(), ViewMatrices.GetProjectionMatrix() ));
-	Data.TranslatedWorldToShadowUVNormalMatrix	= FMatrix44f(CalcTranslatedWorldToShadowUVNormalMatrix( ViewMatrices.GetTranslatedViewMatrix(), ViewMatrices.GetProjectionMatrix() ));
-	Data.PreViewTranslationHigh					= PreViewTranslation.High;
-	Data.PreViewTranslationLow					= PreViewTranslation.Low;
-	Data.LightType								= ProjectedShadowInfo->GetLightSceneInfo().Proxy->GetLightType();
-	Data.LightSourceRadius						= ProjectedShadowInfo->GetLightSceneInfo().Proxy->GetSourceRadius();
-	Data.ResolutionLodBias						= ResolutionLODBiasLocal;
-	Data.LightRadius							= ProjectedShadowInfo->GetLightSceneInfo().Proxy->GetRadius();
-	Data.LightDirection                         = FVector3f(0, 0, 0);	// Unused for local lights
-	Data.Flags									= Flags;
-	Data.TexelDitherScale						= ProjectedShadowInfo->GetLightSceneInfo().Proxy->GetVSMTexelDitherScale();
-
-	return Data;
+	OutProjectionShaderData->ShadowViewToClipMatrix						= FMatrix44f(ViewMatrices.GetProjectionMatrix());
+	OutProjectionShaderData->TranslatedWorldToShadowUVMatrix			= FMatrix44f(CalcTranslatedWorldToShadowUVMatrix( ViewMatrices.GetTranslatedViewMatrix(), ViewMatrices.GetProjectionMatrix() ));
+	OutProjectionShaderData->TranslatedWorldToShadowUVNormalMatrix		= FMatrix44f(CalcTranslatedWorldToShadowUVNormalMatrix( ViewMatrices.GetTranslatedViewMatrix(), ViewMatrices.GetProjectionMatrix() ));
+	OutProjectionShaderData->PreViewTranslationHigh						= PreViewTranslation.High;
+	OutProjectionShaderData->PreViewTranslationLow						= PreViewTranslation.Low;
+
+	OutProjectionShaderData->LightDirection								= FVector3f(0, 0, 0);	// Unused for local lights
+	OutProjectionShaderData->ClipmapLevel_ClipmapLevelCountRemaining	= -1;					// Not a clipmap
 }
 
 /**
@@ -385,9 +369,9 @@ TSharedPtr<FVirtualShadowMapPerLightCacheEntry> FShadowSceneRenderer::AddLocalLi
 
 	const int32 NumMaps = ProjectedShadowInitializer.bOnePassPointLightShadow ? 6 : 1;
 	TSharedPtr<FVirtualShadowMapPerLightCacheEntry> PerLightCacheEntry = CacheManager->FindCreateLightCacheEntry(LightSceneInfo->Id, 0, NumMaps);
-			
 	LocalLightShadowFrameSetup.PerLightCacheEntry = PerLightCacheEntry;
-	bool bIsCached = PerLightCacheEntry->UpdateLocal(
+
+	PerLightCacheEntry->UpdateLocal(
 		ProjectedShadowInitializer,
 		LightSceneProxy->GetOrigin(),
 		LightSceneProxy->GetRadius(),
@@ -395,7 +379,7 @@ TSharedPtr<FVirtualShadowMapPerLightCacheEntry> FShadowSceneRenderer::AddLocalLi
 		CacheManager->IsCacheEnabled(),
 		!bShouldForceTimeSliceDistantUpdate);
 
-	if (bIsCached && bIsDistantLight && PerLightCacheEntry->Prev.ScheduledFrameNumber == Scene.GetFrameNumber())
+	if (bIsDistantLight && PerLightCacheEntry->Prev.ScheduledFrameNumber == Scene.GetFrameNumber())
 	{
 		PerLightCacheEntry->Invalidate();
 	}
@@ -406,17 +390,38 @@ TSharedPtr<FVirtualShadowMapPerLightCacheEntry> FShadowSceneRenderer::AddLocalLi
 	ProjectedShadowInfo->VirtualShadowMapPerLightCacheEntry = PerLightCacheEntry;
 	ProjectedShadowInfo->bShouldRenderVSM = !PerLightCacheEntry->IsFullyCached();
 
-	uint32 PackedCullingViewId = FVirtualShadowMapProjectionShaderData::PackCullingViewId(SceneRenderer.Views[ClosestCullingViewIndex].GPUSceneViewId, SceneRenderer.Views[ClosestCullingViewIndex].PersistentViewId);
-
-	for (int32 Index = 0; Index < NumMaps; ++Index)
 	{
-		const int32 FaceVirtualShadowMapId = VirtualShadowMapId + Index;
-		FVirtualShadowMapCacheEntry& VirtualSmCacheEntry = PerLightCacheEntry->ShadowMapEntries[Index];
-		VirtualSmCacheEntry.Update(VirtualShadowMapArray, *PerLightCacheEntry, FaceVirtualShadowMapId);
-		// Update projection data
-		VirtualSmCacheEntry.ProjectionData = GetLocalLightProjectionShaderData(ResolutionLODBiasLocal, ProjectedShadowInfo, Index);
-		VirtualSmCacheEntry.ProjectionData.MinMipLevel = MinMipLevel;
-		VirtualSmCacheEntry.ProjectionData.PackedCullingViewId = PackedCullingViewId;
+		uint32 PackedCullingViewId = FVirtualShadowMapProjectionShaderData::PackCullingViewId(SceneRenderer.Views[ClosestCullingViewIndex].GPUSceneViewId, SceneRenderer.Views[ClosestCullingViewIndex].PersistentViewId);
+		uint32 Flags = PerLightCacheEntry->IsUncached() ? VSM_PROJ_FLAG_UNCACHED : 0U;
+		const FLightSceneProxy* Proxy = ProjectedShadowInfo->GetLightSceneInfo().Proxy;
+
+		// For now just tie this to whether anything has invalidated the light (including movement)
+		// This is slightly over-conservative but catches the important cases
+		const bool bUpdateMatrices = PerLightCacheEntry->IsInvalidated();
+
+		for (int32 Index = 0; Index < NumMaps; ++Index)
+		{
+			const int32 FaceVirtualShadowMapId = VirtualShadowMapId + Index;
+			FVirtualShadowMapCacheEntry& VirtualSmCacheEntry = PerLightCacheEntry->ShadowMapEntries[Index];
+			VirtualSmCacheEntry.Update(VirtualShadowMapArray, *PerLightCacheEntry, FaceVirtualShadowMapId);
+
+			FVirtualShadowMapProjectionShaderData& ProjectionData = VirtualSmCacheEntry.ProjectionData;
+
+			if (bUpdateMatrices)
+			{
+				UpdateLocalLightProjectionShaderDataMatrices(ProjectedShadowInfo, Index, &ProjectionData);
+			}
+
+			// TODO: All of this is per-light data; splitting this out to a separate structure could help
+			ProjectionData.LightType			= Proxy->GetLightType();
+			ProjectionData.LightSourceRadius	= Proxy->GetSourceRadius();	
+			ProjectionData.LightRadius			= Proxy->GetRadius();
+			ProjectionData.TexelDitherScale		= Proxy->GetVSMTexelDitherScale();
+			ProjectionData.ResolutionLodBias	= ResolutionLODBiasLocal;
+			ProjectionData.Flags				= Flags;
+			ProjectionData.MinMipLevel			= MinMipLevel;
+			ProjectionData.PackedCullingViewId	= PackedCullingViewId;
+		}
 	}
 
 	FLightOcclusionType OcclusionType = GetLightOcclusionType(*LightSceneInfo->Proxy, SceneRenderer.ViewFamily);
diff --git a/Engine/Source/Runtime/Renderer/Private/Shadows/ShadowSceneRenderer.h b/Engine/Source/Runtime/Renderer/Private/Shadows/ShadowSceneRenderer.h
--- a/Engine/Source/Runtime/Renderer/Private/Shadows/ShadowSceneRenderer.h
+++ b/Engine/Source/Runtime/Renderer/Private/Shadows/ShadowSceneRenderer.h
@@ -137,7 +137,15 @@ private:
 
 	UE::Tasks::FTask RendererSetupTask;
 
-	FVirtualShadowMapProjectionShaderData GetLocalLightProjectionShaderData(float ResolutionLODBiasLocal, const FProjectedShadowInfo* ProjectedShadowInfo, int32 MapIndex) const;
+	FVirtualShadowMapProjectionShaderData GetLocalLightProjectionShaderData(
+		float ResolutionLODBiasLocal,
+		const FProjectedShadowInfo* ProjectedShadowInfo,
+		int32 MapIndex) const;
+
+	 void UpdateLocalLightProjectionShaderDataMatrices(
+		const FProjectedShadowInfo* ProjectedShadowInfo,
+		int32 MapIndex,
+		FVirtualShadowMapProjectionShaderData* OutProjectionShaderData) const;
 
 	void CreateNaniteRenderPasses(
 		FRDGBuilder& GraphBuilder,
diff --git a/Engine/Source/Runtime/Renderer/Private/VirtualShadowMaps/VirtualShadowMapArray.cpp b/Engine/Source/Runtime/Renderer/Private/VirtualShadowMaps/VirtualShadowMapArray.cpp
--- a/Engine/Source/Runtime/Renderer/Private/VirtualShadowMaps/VirtualShadowMapArray.cpp
+++ b/Engine/Source/Runtime/Renderer/Private/VirtualShadowMaps/VirtualShadowMapArray.cpp
@@ -3886,11 +3886,7 @@ void FVirtualShadowMapArray::AddRenderViewsClipmap(
 			// If we're going to generate a new HZB this frame, save the associated metadata
 			if (bUpdateHZBMetaData)
 			{
-				FVirtualShadowMapHZBMetadata HZBMeta;
-				HZBMeta.TargetLayerIndex = Params.TargetLayerIndex;
-				HZBMeta.ViewMatrices = Params.ViewMatrices;
-				HZBMeta.ViewRect = Params.ViewRect;
-				LevelEntry.CurrentHZBMetadata = HZBMeta;
+				LevelEntry.UpdateHZBMetadata(Params.ViewMatrices, Params.ViewRect, Params.TargetLayerIndex);
 			}
 		}
 
@@ -3971,11 +3967,7 @@ void FVirtualShadowMapArray::AddRenderViewsLocal(
 			// If we're going to generate a new HZB this frame, save the associated metadata
 			if (bUpdateHZBMetaData)
 			{
-				FVirtualShadowMapHZBMetadata HZBMeta;
-				HZBMeta.TargetLayerIndex = Params.TargetLayerIndex;
-				HZBMeta.ViewMatrices = Params.ViewMatrices;
-				HZBMeta.ViewRect = Params.ViewRect;
-				LevelEntry.CurrentHZBMetadata = HZBMeta;
+				LevelEntry.UpdateHZBMetadata(Params.ViewMatrices, Params.ViewRect, Params.TargetLayerIndex);
 			}
 		}
 
diff --git a/Engine/Source/Runtime/Renderer/Private/VirtualShadowMaps/VirtualShadowMapArray.h b/Engine/Source/Runtime/Renderer/Private/VirtualShadowMaps/VirtualShadowMapArray.h
--- a/Engine/Source/Runtime/Renderer/Private/VirtualShadowMaps/VirtualShadowMapArray.h
+++ b/Engine/Source/Runtime/Renderer/Private/VirtualShadowMaps/VirtualShadowMapArray.h
@@ -134,13 +134,6 @@ struct FVirtualShadowMapProjectionShaderData
 };
 static_assert(sizeof(FVirtualShadowMapProjectionShaderData) == (16*18), "FVirtualShadowMapProjectionShaderData does not match size in shader. See VirtualShadowMapProjectionStructs.ush.");
 
-struct FVirtualShadowMapHZBMetadata
-{
-	FViewMatrices ViewMatrices;
-	FIntRect	  ViewRect;
-	uint32		  TargetLayerIndex = INDEX_NONE;
-};
-
 BEGIN_SHADER_PARAMETER_STRUCT(FVirtualShadowMapPerViewParameters, )
 	// Light grid with only the lights that have VSMs present
 	// Still references the original indices from the global light grid
diff --git a/Engine/Source/Runtime/Renderer/Private/VirtualShadowMaps/VirtualShadowMapCacheManager.cpp b/Engine/Source/Runtime/Renderer/Private/VirtualShadowMaps/VirtualShadowMapCacheManager.cpp
--- a/Engine/Source/Runtime/Renderer/Private/VirtualShadowMaps/VirtualShadowMapCacheManager.cpp
+++ b/Engine/Source/Runtime/Renderer/Private/VirtualShadowMaps/VirtualShadowMapCacheManager.cpp
@@ -183,7 +183,7 @@ void FVirtualShadowMapCacheEntry::UpdateClipmapLevel(
 {
 	int32 PrevVirtualShadowMapId = CurrentVirtualShadowMapId;
 	FInt64Point PrevPageSpaceLocation = Clipmap.PageSpaceLocation;
-	PrevHZBMetadata = CurrentHZBMetadata;
+	UpdatePrevHZBMetadata();
 
 	bool bCacheValid = (PrevVirtualShadowMapId != INDEX_NONE);
 	
@@ -250,10 +250,7 @@ void FVirtualShadowMapCacheEntry::Update(
 {
 	// Swap previous frame data over.
 	int32 PrevVirtualShadowMapId = CurrentVirtualShadowMapId;
-
-	// TODO: This is pretty wrong specifically for unreferenced lights, as the VSM IDs will have changed and not been updated since
-	// this gets updated by rendering! Need to figure out a better way to track this data, and probably not here...
-	PrevHZBMetadata = CurrentHZBMetadata;
+	UpdatePrevHZBMetadata();
 	
 	bool bCacheValid = (PrevVirtualShadowMapId != INDEX_NONE);
 
@@ -288,6 +285,24 @@ void FVirtualShadowMapCacheEntry::SetHZBViewParams(Nanite::FPackedViewParams& Ou
 	OutParams.Flags |= NANITE_VIEW_FLAG_HZBTEST;
 }
 
+void FVirtualShadowMapCacheEntry::UpdateHZBMetadata(const FViewMatrices& ViewMatrices, const FIntRect& ViewRect, uint32 TargetLayerIndex)
+{
+	CurrentHZBMetadata.ViewRect = ViewRect;
+	CurrentHZBMetadata.TargetLayerIndex = TargetLayerIndex;
+
+	// As this structure has grown a lot, we try to avoid updating it if we don't have to
+	// For shadows, this should be a reasonable test of the relevant parameters that could change
+	// TODO: This test is very similar to the tests we do for whether caching is valid; could probably share the two better with some refactoring
+	CurrentHZBMetadata.bMatricesDirty =
+		ViewMatrices.GetPreViewTranslation() != CurrentHZBMetadata.ViewMatrices.GetPreViewTranslation() ||
+		ViewMatrices.GetTranslatedViewProjectionMatrix() != CurrentHZBMetadata.ViewMatrices.GetTranslatedViewProjectionMatrix();
+
+	if (CurrentHZBMetadata.bMatricesDirty)
+	{
+		CurrentHZBMetadata.ViewMatrices = ViewMatrices;
+	}
+}
+
 void FVirtualShadowMapPerLightCacheEntry::UpdateClipmap(const FVector& LightDirection, int FirstLevel, bool bForceInvalidate)
 {
 	Prev.RenderedFrameNumber = FMath::Max(Prev.RenderedFrameNumber, Current.RenderedFrameNumber);
@@ -315,7 +330,7 @@ void FVirtualShadowMapPerLightCacheEntry::UpdateClipmap(const FVector& LightDire
 	LightRadius = -1.0f;
 }
 
-bool FVirtualShadowMapPerLightCacheEntry::UpdateLocal(
+void FVirtualShadowMapPerLightCacheEntry::UpdateLocal(
 	const FProjectedShadowInitializer& InCacheKey,
 	const FVector& NewLightOrigin,
 	const float NewLightRadius,
@@ -328,16 +343,20 @@ bool FVirtualShadowMapPerLightCacheEntry::UpdateLocal(
 	Prev.RenderedFrameNumber = FMath::Max(Prev.RenderedFrameNumber, Current.RenderedFrameNumber);
 	Prev.ScheduledFrameNumber = FMath::Max(Prev.ScheduledFrameNumber, Current.ScheduledFrameNumber);
 
+	const bool bLightMoved = 
+		LocalCacheKey.PreShadowTranslation != InCacheKey.PreShadowTranslation ||
+		LocalCacheKey.WorldToLight != InCacheKey.WorldToLight;
+
+	LocalCacheKey.PreShadowTranslation = InCacheKey.PreShadowTranslation;
+	LocalCacheKey.WorldToLight = InCacheKey.WorldToLight;
+
 	// Check cache validity based of shadow setup
 	// If it is a distant light, we want to let the time-share perform the invalidation.
-	if (!bCacheEnabled
-		|| (bAllowInvalidation && !LocalCacheKey.IsCachedShadowValid(InCacheKey)))
+	if (!bCacheEnabled || (bAllowInvalidation && bLightMoved))
 	{
-		// TODO: track invalidation state somehow for later.
 		Prev.RenderedFrameNumber = -1;
 		//UE_LOG(LogRenderer, Display, TEXT("Invalidated!"));
 	}
-	LocalCacheKey = InCacheKey;
 
 	// On transition between uncached <-> cached we must invalidate since the static pages may not be initialized
 	bool bNewIsUncached = Prev.RenderedFrameNumber < 0;
@@ -358,8 +377,6 @@ bool FVirtualShadowMapPerLightCacheEntry::UpdateLocal(
 	bIsUncached = bNewIsUncached;
 	LightOrigin = NewLightOrigin;
 	LightRadius = NewLightRadius;
-
-	return Prev.RenderedFrameNumber >= 0;
 }
 
 void FVirtualShadowMapArrayCacheManager::FShadowInvalidatingInstancesImplementation::AddPrimitive(const FPrimitiveSceneInfo* PrimitiveSceneInfo)
@@ -376,11 +393,6 @@ void FVirtualShadowMapArrayCacheManager::FShadowInvalidatingInstancesImplementat
 		true});
 }
 
-void FVirtualShadowMapPerLightCacheEntry::Invalidate()
-{
-	Prev.RenderedFrameNumber = -1;
-}
-
 static uint32 EncodeInstanceInvalidationPayload(int32 VirtualShadowMapId, uint32 Flags = VSM_INVALIDATION_PAYLOAD_FLAG_NONE)
 {
 	check(VirtualShadowMapId >= 0);		// Should not be INDEX_NONE by this point
diff --git a/Engine/Source/Runtime/Renderer/Private/VirtualShadowMaps/VirtualShadowMapCacheManager.h b/Engine/Source/Runtime/Renderer/Private/VirtualShadowMaps/VirtualShadowMapCacheManager.h
--- a/Engine/Source/Runtime/Renderer/Private/VirtualShadowMaps/VirtualShadowMapCacheManager.h
+++ b/Engine/Source/Runtime/Renderer/Private/VirtualShadowMaps/VirtualShadowMapCacheManager.h
@@ -28,6 +28,15 @@ struct FVirtualShadowMapInstanceRange
 	bool bMarkAsDynamic;					// If true, swaps the primitive/instance to dynamic caching
 };
 
+struct FVirtualShadowMapHZBMetadata
+{
+	// See UpdatePrevHZBMetadata if you modify fields here
+	FViewMatrices ViewMatrices;
+	FIntRect	  ViewRect;
+	uint32		  TargetLayerIndex = INDEX_NONE;
+	bool          bMatricesDirty = true;
+};
+
 #define VSM_LOG_INVALIDATIONS 0
 
 class FVirtualShadowMapCacheEntry
@@ -53,6 +62,18 @@ public:
 
 	void SetHZBViewParams(Nanite::FPackedViewParams& OutParams);
 
+	void UpdateHZBMetadata(const FViewMatrices& ViewMatrices, const FIntRect& ViewRect, uint32 TargetLayerIndex);
+
+	void UpdatePrevHZBMetadata()
+	{
+		PrevHZBMetadata.TargetLayerIndex = CurrentHZBMetadata.TargetLayerIndex;
+		PrevHZBMetadata.ViewRect = CurrentHZBMetadata.ViewRect;
+		if (CurrentHZBMetadata.bMatricesDirty)
+		{
+			PrevHZBMetadata.ViewMatrices = CurrentHZBMetadata.ViewMatrices;
+		}
+	}
+
 	// Previous frame data
 	FVirtualShadowMapHZBMetadata PrevHZBMetadata;
 
@@ -97,12 +118,12 @@ public:
 
 	void MarkRendered(int32 FrameIndex) { Current.RenderedFrameNumber = FrameIndex; }
 	int32 GetLastScheduledFrameNumber() const { return Prev.ScheduledFrameNumber; }
-	void UpdateClipmap(const FVector& LightDirection, int FirstLevelx, bool bForceInvalidate);
+	void UpdateClipmap(const FVector& LightDirection, int FirstLevel, bool bForceInvalidate);
 
 	/**
 	 * Returns true if the cache entry is valid (has previous state).
 	 */
-	bool UpdateLocal(
+	void UpdateLocal(
 		const FProjectedShadowInitializer &InCacheKey,
 		const FVector& NewLightOrigin,
 		const float NewLightRadius,
@@ -110,10 +131,9 @@ public:
 		bool bCacheEnabled,
 		bool bAllowInvalidation);
 
-	/**
-	 * Mark as invalid, i.e., needing rendering.
-	 */
-	void Invalidate();
+	void Invalidate() { Prev.RenderedFrameNumber = -1; }
+
+	bool IsInvalidated() const { return Prev.RenderedFrameNumber < 0; }
 
 	bool AffectsBounds(const FBoxSphereBounds& Bounds) const
 	{
@@ -155,7 +175,13 @@ public:
 	float LightRadius = -1.0f;		// Negative means infinite
 
 private:
-	FProjectedShadowInitializer LocalCacheKey;
+	struct FLocalLightCacheKey
+	{
+		FMatrix WorldToLight;
+		FVector PreShadowTranslation;
+	};
+	FLocalLightCacheKey LocalCacheKey;
+
 	struct FClipmapCacheKey
 	{
 		FVector LightDirection;
