Moved several FViewInfo members to heap allocations on the view's object allocator. This reduces the size of FViewInfo considerably for snapshots.

diff --git a/Engine/Source/Runtime/Renderer/Private/SceneRendering.h b/Engine/Source/Runtime/Renderer/Private/SceneRendering.h
--- a/Engine/Source/Runtime/Renderer/Private/SceneRendering.h
+++ b/Engine/Source/Runtime/Renderer/Private/SceneRendering.h
@@ -1207,7 +1207,7 @@ public:
 	FTranslucentSelfShadowUniformBufferMap TranslucentSelfShadowUniformBufferMap;
 
 	/** View dependent global distance field clipmap info. */
-	FGlobalDistanceFieldInfo GlobalDistanceFieldInfo;
+	FGlobalDistanceFieldInfo& GlobalDistanceFieldInfo = *Allocator.Create<FGlobalDistanceFieldInfo>();
 
 	/** Count of translucent prims for this view. */
 	FTranslucenyPrimCount TranslucentPrimCount;
@@ -1240,10 +1240,10 @@ public:
 	TArray<const FLightSceneProxy*, SceneRenderingAllocator> VisibleReflectionCaptureLights;
 
 	/** The view's batched elements. */
-	FBatchedElements BatchedViewElements;
+	FBatchedElements& BatchedViewElements = *Allocator.Create<FBatchedElements>();
 
 	/** The view's batched elements, above all other elements, for gizmos that should never be occluded. */
-	FBatchedElements TopBatchedViewElements;
+	FBatchedElements& TopBatchedViewElements = *Allocator.Create<FBatchedElements>();
 
 	/** The view's mesh elements. */
 	TIndirectArray<FMeshBatch, SceneRenderingAllocator> ViewMeshElements;
@@ -1272,13 +1272,13 @@ public:
 
 	TArray<FMeshBatchAndRelevance,SceneRenderingAllocator> DynamicEditorMeshElements;
 
-	FSimpleElementCollector SimpleElementCollector;
+	FSimpleElementCollector& SimpleElementCollector = *Allocator.Create<FSimpleElementCollector>();
 
-	FSimpleElementCollector EditorSimpleElementCollector;
+	FSimpleElementCollector& EditorSimpleElementCollector = *Allocator.Create<FSimpleElementCollector>();;
 
 #if UE_ENABLE_DEBUG_DRAWING
 	//Separate DebugSimpleElementCollector to not conflate any other simple element collector draws which may have been added from non-debug draw passes (e.g. non-opaque draws)
-	FSimpleElementCollector DebugSimpleElementCollector;
+	FSimpleElementCollector& DebugSimpleElementCollector = *Allocator.Create<FSimpleElementCollector>();;
 #endif
 
 	FParallelMeshDrawCommandPass* CreateMeshPass(EMeshPass::Type MeshPass);
@@ -1303,9 +1303,9 @@ public:
 	// Used by mobile renderer to determine whether static meshes will be rendered with CSM shaders or not.
 	FMobileCSMVisibilityInfo MobileCSMVisibilityInfo;
 
-	FSubstrateViewData SubstrateViewData;
+	FSubstrateViewData& SubstrateViewData = *Allocator.Create<FSubstrateViewData>();
 
-	FLocalFogVolumeViewData LocalFogVolumeViewData;
+	FLocalFogVolumeViewData& LocalFogVolumeViewData = *Allocator.Create<FLocalFogVolumeViewData>();
 
 	FHairStrandsViewData* HairStrandsViewData = nullptr;
 
@@ -1410,7 +1410,7 @@ public:
 	uint16 ShadingModelMaskInView;
 
 	/** Informations from the previous frame to use for this view. */
-	FPreviousViewInfo PrevViewInfo;
+	FPreviousViewInfo& PrevViewInfo = *Allocator.Create<FPreviousViewInfo>();
 
 	/** An intermediate number of visible static meshes.  Doesn't account for occlusion until after FinishOcclusionQueries is called. */
 	int32 NumVisibleStaticMeshElements;
@@ -1494,7 +1494,7 @@ public:
 	FShaderPrintData ShaderPrintData;
 
 private:
-	FLumenTranslucencyGIVolume LumenTranslucencyGIVolume;
+	FLumenTranslucencyGIVolume& LumenTranslucencyGIVolume = *Allocator.Create<FLumenTranslucencyGIVolume>();
 	FMegaLightsVolume MegaLightsVolume;
 
 public:
