Add Bloom to ReplaceTonemapper Mobile

diff --git a/Engine/Shaders/Private/MaterialTemplate.ush b/Engine/Shaders/Private/MaterialTemplate.ush
--- a/Engine/Shaders/Private/MaterialTemplate.ush
+++ b/Engine/Shaders/Private/MaterialTemplate.ush
@@ -3073,6 +3073,15 @@ MaterialFloat4 MobileSceneTextureLookup(inout FMaterialPixelParameters Parameter
 		#endif
 		return Input0;
 	}
+
+	#define RETURN_MOBILE_POST_PROCESS_INPUT_SAMPLE(Index) \
+		return Texture2DSample(PostProcessInput_##Index##_Texture, PostProcessInput_##Index##_SharedSampler, UV)
+
+	case PPI_PostProcessInput1:		RETURN_MOBILE_POST_PROCESS_INPUT_SAMPLE(1);
+	case PPI_PostProcessInput2:		RETURN_MOBILE_POST_PROCESS_INPUT_SAMPLE(2);
+	case PPI_PostProcessInput3:		RETURN_MOBILE_POST_PROCESS_INPUT_SAMPLE(3);
+	case PPI_PostProcessInput4:		RETURN_MOBILE_POST_PROCESS_INPUT_SAMPLE(4);
+
 #endif // POST_PROCESS_MATERIAL
 	case PPI_DecalMask:
 		return 0;  // material compiler will return an error
diff --git a/Engine/Source/Runtime/Renderer/Private/PostProcess/PostProcessing.cpp b/Engine/Source/Runtime/Renderer/Private/PostProcess/PostProcessing.cpp
--- a/Engine/Source/Runtime/Renderer/Private/PostProcess/PostProcessing.cpp
+++ b/Engine/Source/Runtime/Renderer/Private/PostProcess/PostProcessing.cpp
@@ -2486,7 +2486,7 @@ void AddMobilePostProcessingPasses(FRDGBuilder& GraphBuilder, FScene* Scene, con
 	// The scene color will be decoded at the first post-process material and output linear color space for the following passes
 	// bMetalMSAAHDRDecode will be set to false if there is any post-process material exist
 
-	auto AddPostProcessMaterialPass = [&GraphBuilder, &View, ViewIndex, &Inputs, &SceneColor, &CustomDepth, &bMetalMSAAHDRDecode, &PassSequence](EBlendableLocation BlendableLocation, bool bLastPass)
+	auto AddPostProcessMaterialPass = [&GraphBuilder, &View, ViewIndex, &Inputs, &SceneColor, &CustomDepth, &bMetalMSAAHDRDecode, &PassSequence, &BloomOutput, &BlackAlphaOneDummy](EBlendableLocation BlendableLocation, bool bLastPass)
 	{
 		FPostProcessMaterialInputs PostProcessMaterialInputs;
 
@@ -2502,6 +2502,15 @@ void AddMobilePostProcessingPasses(FRDGBuilder& GraphBuilder, FScene* Scene, con
 
 		PostProcessMaterialInputs.SetInput(GraphBuilder, EPostProcessMaterialInput::SceneColor, SceneColor);
 
+		if (BlendableLocation == BL_ReplacingTonemapper && PassSequence.IsEnabled(EPass::Tonemap))
+		{
+			if (!BloomOutput.IsValid())
+			{
+				BloomOutput = BlackAlphaOneDummy;
+			}
+			PostProcessMaterialInputs.SetInput(GraphBuilder, EPostProcessMaterialInput::CombinedBloom, BloomOutput);
+		}
+
 		PostProcessMaterialInputs.CustomDepthTexture = CustomDepth.Texture;
 
 		PostProcessMaterialInputs.bMetalMSAAHDRDecode = bMetalMSAAHDRDecode;
