Fix to directional light not casting any light when cast shadows is disabled.
The shadow mask channel concurs at calculating the light attenuation due to dynamic shadows for a given light. When dynamic shadows are not enabled, the shadow mask channel was being set to 0 which ended up in a multiplier that was zeroing out the contribution of the light.

diff --git a/Engine/Source/Runtime/Renderer/Private/MobileBasePassRendering.cpp b/Engine/Source/Runtime/Renderer/Private/MobileBasePassRendering.cpp
--- a/Engine/Source/Runtime/Renderer/Private/MobileBasePassRendering.cpp
+++ b/Engine/Source/Runtime/Renderer/Private/MobileBasePassRendering.cpp
@@ -440,6 +440,13 @@ void SetupMobileDirectionalLightUniformParameters(
 			(DynamicShadowMapChannel == 2 ? 64 : 0) |
 			(DynamicShadowMapChannel == 3 ? 128 : 0);
 
+		// When dynamic shadow casting is disabled for the light we still need to set its 
+		// dynamic shadow map channel to 1 or the light's contribution will be zeroed out.
+		if (!Light->Proxy->CastsDynamicShadow())
+		{
+			Params.DirectionalLightShadowMapChannelMask |= (DynamicShadowMapChannel == -1 ? 16 : 0);
+		}
+
 		if (bDynamicShadows && VisibleLightInfos.IsValidIndex(Light->Id) && VisibleLightInfos[Light->Id].AllProjectedShadows.Num() > 0)
 		{
 			const TArray<FProjectedShadowInfo*, SceneRenderingAllocator>& DirectionalLightShadowInfos = VisibleLightInfos[Light->Id].AllProjectedShadows;
