Fix assert due to CVarRayTracingVisualize being accessed on game thread.

diff --git a/Engine/Source/Runtime/Renderer/Private/RayTracing/RayTracingDebug.cpp b/Engine/Source/Runtime/Renderer/Private/RayTracing/RayTracingDebug.cpp
--- a/Engine/Source/Runtime/Renderer/Private/RayTracing/RayTracingDebug.cpp
+++ b/Engine/Source/Runtime/Renderer/Private/RayTracing/RayTracingDebug.cpp
@@ -1378,7 +1378,7 @@ static uint32 GetRaytracingDebugViewMode(const FSceneView& View)
 {
 	const FRayTracingVisualizationData& VisualizationData = GetRayTracingVisualizationData();
 
-	FString ConsoleViewMode = CVarRayTracingVisualize.GetValueOnRenderThread();
+	FString ConsoleViewMode = CVarRayTracingVisualize.GetValueOnAnyThread();
 	if (!ConsoleViewMode.IsEmpty())
 	{
 		return  VisualizationData.GetModeID(FName(*ConsoleViewMode));
