Use dummy culling parameters if no meshes are in a pass.

diff --git a/Engine/Source/Runtime/Renderer/Private/TranslucentRendering.cpp b/Engine/Source/Runtime/Renderer/Private/TranslucentRendering.cpp
--- a/Engine/Source/Runtime/Renderer/Private/TranslucentRendering.cpp
+++ b/Engine/Source/Runtime/Renderer/Private/TranslucentRendering.cpp
@@ -1382,6 +1382,10 @@ static void RenderTranslucencyViewInner(
 			}
 		}
 	}
+	else
+	{
+		InstanceCullingManager.SetDummyCullingParams(GraphBuilder, PassParameters->InstanceCullingDrawParams);
+	}
 
 	if (IsMainTranslucencyPass(TranslucencyPass) && (View.SimpleElementCollector.HasAnyPrimitives() || View.bHasTranslucentViewMeshElements))
 	{
