[HWRT] Support gathering dynamic ray tracing instances on parallel task.

diff --git a/Engine/Plugins/FX/Niagara/Source/Niagara/Private/NiagaraComponent.cpp b/Engine/Plugins/FX/Niagara/Source/Niagara/Private/NiagaraComponent.cpp
--- a/Engine/Plugins/FX/Niagara/Source/Niagara/Private/NiagaraComponent.cpp
+++ b/Engine/Plugins/FX/Niagara/Source/Niagara/Private/NiagaraComponent.cpp
@@ -235,6 +235,11 @@ FNiagaraSceneProxy::FNiagaraSceneProxy(const FNiagaraSceneProxyDesc& Desc)
 	// Optionally enable GDME for Niagara, any caveats are covered below
 	bSupportsParallelGDME = GNiagaraParallelGDME;
 
+	// Both GDME and GDRTI call methods such as InitializeVertexBuffersResources(...)
+	// which means GDME and GDRTI can't be called in parallel.
+	// Disable bSupportsParallelGDRTI so that bSupportsParallelGDME can run in parallel
+	bSupportsParallelGDRTI = false;
+
 	FNiagaraSystemInstanceControllerConstPtr SystemInstanceController = Desc.SystemInstanceController;
 	UNiagaraSystem* NiagaraSystem = Desc.SystemAsset;
 	if (SystemInstanceController && NiagaraSystem)
diff --git a/Engine/Source/Runtime/Engine/Private/PrimitiveSceneProxy.cpp b/Engine/Source/Runtime/Engine/Private/PrimitiveSceneProxy.cpp
--- a/Engine/Source/Runtime/Engine/Private/PrimitiveSceneProxy.cpp
+++ b/Engine/Source/Runtime/Engine/Private/PrimitiveSceneProxy.cpp
@@ -112,6 +112,21 @@ bool IsParallelGatherDynamicMeshElementsEnabled()
 		GSupportsParallelRenderingTasksWithSeparateRHIThread;
 }
 
+bool GParallelGatherDynamicRayTracingInstances = true;
+static FAutoConsoleVariableRef CVarParallelGatherDynamicRayTracingInstances(
+	TEXT("r.RayTracing.DynamicInstances.Parallel"),
+	GParallelGatherDynamicRayTracingInstances,
+	TEXT("Enables parallel processing of the gather dynamic ray tracing instances."),
+	ECVF_RenderThreadSafe
+);
+
+bool IsParallelGatherDynamicRayTracingInstancesEnabled()
+{
+	return GParallelGatherDynamicRayTracingInstances &&
+		// parallel GDRTI creates RHI resources on a task threads which is not supported by some RHIs (eg. OpenGL)
+		GSupportsParallelRenderingTasksWithSeparateRHIThread;
+}
+
 bool FPrimitiveSceneProxy::ShouldRenderCustomDepth() const
 {
 	return IsCustomDepthPassEnabled() && bRenderCustomDepth;
@@ -443,6 +458,7 @@ FPrimitiveSceneProxy::FPrimitiveSceneProxy(const FPrimitiveSceneProxyDesc& InPro
 ,	bTreatAsBackgroundForOcclusion(InProxyDesc.bTreatAsBackgroundForOcclusion)
 ,	bSupportsParallelGDME(true)
 ,	bSinglePassGDME(false)
+,	bSupportsParallelGDRTI(true)
 ,	bVisibleInLumenScene(false)
 ,	bOpaqueOrMasked(true)
 ,	bCanSkipRedundantTransformUpdates(true)
diff --git a/Engine/Source/Runtime/Engine/Public/PrimitiveSceneProxy.h b/Engine/Source/Runtime/Engine/Public/PrimitiveSceneProxy.h
--- a/Engine/Source/Runtime/Engine/Public/PrimitiveSceneProxy.h
+++ b/Engine/Source/Runtime/Engine/Public/PrimitiveSceneProxy.h
@@ -744,6 +744,7 @@ public:
 	inline bool ShouldRenderInDepthPass() const { return bRenderInMainPass || bRenderInDepthPass; }
 	inline bool SupportsParallelGDME() const { return bSupportsParallelGDME; }
 	inline bool SinglePassGDME() const { return bSinglePassGDME; }
+	inline bool SupportsParallelGDRTI() const { return bSupportsParallelGDRTI; }
 	inline bool IsCollisionEnabled() const { return bCollisionEnabled; }
 	inline bool IsHovered() const { return bHovered; }
 	inline bool IsOwnedBy(const FSceneViewOwner& Owner) const { return Owner.IsPartOf(Owners); }
@@ -1269,6 +1270,9 @@ protected:
 	 */
 	uint8 bSinglePassGDME : 1;
 
+	/** Whether the proxy supports asynchronously calling GetDynamicRayTracingInstances. If disabled, all calls for various proxies are serialized with respect to each other. */
+	uint8 bSupportsParallelGDRTI : 1;
+
 	/** Whether this component should be tracked by Lumen Scene. Turning this off will remove it from Lumen Scene and Lumen won't generate surface cache for it. */
 	uint8 bVisibleInLumenScene : 1;
 
@@ -1696,6 +1700,9 @@ ENGINE_API extern bool SupportsNaniteRendering(const class FVertexFactoryType* R
 // Whether scene proxies will have GetDynamicMeshElements called in parallel.
 ENGINE_API bool IsParallelGatherDynamicMeshElementsEnabled();
 
+// Whether scene proxies will have GetDynamicRayTracingInstances called in parallel.
+ENGINE_API bool IsParallelGatherDynamicRayTracingInstancesEnabled();
+
 #if UE_ENABLE_INCLUDE_ORDER_DEPRECATED_IN_5_6
 #include "SceneView.h"
 #endif
diff --git a/Engine/Source/Runtime/Renderer/Private/RayTracing/RayTracing.cpp b/Engine/Source/Runtime/Renderer/Private/RayTracing/RayTracing.cpp
--- a/Engine/Source/Runtime/Renderer/Private/RayTracing/RayTracing.cpp
+++ b/Engine/Source/Runtime/Renderer/Private/RayTracing/RayTracing.cpp
@@ -223,14 +223,59 @@ namespace RayTracing
 		int32 VisibleShaderBindingsOffset = -1;
 	};
 
+	struct FDynamicRayTracingPrimitive
+	{
+		int32 PrimitiveIndex;
+		TRange<int32> InstancesRange;
+		TRange<int32> GeometriesToUpdateRange;
+	};
+
+	class FDynamicRayTracingInstancesContext
+	{
+	public:
+		FDynamicRayTracingInstancesContext(FScene& Scene, FViewInfo& InView, FSceneRenderingBulkObjectAllocator& InBulkAllocator, const FSceneOptions& InSceneOptions);
+
+		void GatherDynamicRayTracingInstances(TConstArrayView<int32> InDynamicPrimitives);
+
+		void GatherDynamicRayTracingInstances_RenderThread();
+
+		void Finish(FRHICommandListImmediate& InRHICmdList);
+
+		void AddInstancesToScene(FRayTracingScene& RayTracingScene, FRayTracingShaderBindingTable& RayTracingSBT, int64 SharedBufferGenerationID);
+
+		void CollectRDGResources(FRDGBuilder& RDGBuilder);
+
+	private:
+
+		void GatherDynamicRayTracingInstances_Internal(int32 PrimitiveIndex);
+
+		FScene& Scene;
+		FViewInfo& View;
+		FSceneRenderingBulkObjectAllocator& BulkAllocator;
+		const FSceneOptions& SceneOptions;
+
+		bool bTrackReferencedGeometryGroups;
+
+		FRHICommandList* RHICmdList;
+		FGlobalDynamicVertexBuffer DynamicVertexBuffer;
+		FGlobalDynamicIndexBuffer DynamicIndexBuffer;
+
+		FRayTracingInstanceCollector RayTracingInstanceCollector;
+
+		TArray<FDynamicRayTracingPrimitive> DynamicRayTracingPrimitives;
+
+		TArray<int32> RenderThreadDynamicPrimitives;
+	};
+
 	struct FGatherInstancesTaskData
 	{
 		UE_NONCOPYABLE(FGatherInstancesTaskData)
 
-		FGatherInstancesTaskData(FScene& InScene, FViewInfo& InView, FSceneOptions InSceneOptions)
+		FGatherInstancesTaskData(FScene& InScene, FViewInfo& InView, FSceneRenderingBulkObjectAllocator& InBulkAllocator, FSceneOptions InSceneOptions)
 			: Scene(InScene)
 			, View(InView)
 			, SceneOptions(MoveTemp(InSceneOptions))
+			, DynamicRayTracingInstancesContext(Scene, View, InBulkAllocator, SceneOptions)
 		{
 
 		}
@@ -248,12 +293,17 @@ namespace RayTracing
 
 		TArray<FRelevantPrimitiveGatherContext> GatherContexts;
 
+		FDynamicRayTracingInstancesContext DynamicRayTracingInstancesContext;
+
 		// This task must complete before accessing StaticPrimitivesIndices/DynamicPrimitivesIndices.
 		UE::Tasks::FTask GatherRelevantPrimitivesTask;
 
 		// This task must complete before accessing StaticPrimitives/CachedStaticPrimitives.
 		UE::Tasks::FTask GatherRelevantStaticPrimitivesTask;
 
+		// This task must complete before accessing DynamicRayTracingInstancesContext.
+		UE::Tasks::FTask GatherDynamicRayTracingInstancesTask;
+
 		// This task must complete before PostRenderAllViewports().
 		UE::Tasks::FTask AddUsedStreamingHandlesTask;
 
@@ -280,7 +330,7 @@ namespace RayTracing
 		EDiffuseIndirectMethod DiffuseIndirectMethod,
 		EReflectionsMethod ReflectionsMethod)
 	{
-		return InAllocator.Create<FGatherInstancesTaskData>(Scene, View, FSceneOptions(Scene, ViewFamily, View, DiffuseIndirectMethod, ReflectionsMethod));
+		return InAllocator.Create<FGatherInstancesTaskData>(Scene, View, InAllocator, FSceneOptions(Scene, ViewFamily, View, DiffuseIndirectMethod, ReflectionsMethod));
 	}
 
 	void OnRenderBegin(FScene& Scene, TArray<FViewInfo>& Views, const FViewFamilyInfo& ViewFamily)
@@ -466,44 +516,6 @@ namespace RayTracing
 			}));
 	};
 
-	struct FDynamicRayTracingPrimitive
-	{
-		int32 PrimitiveIndex;
-		TRange<int32> InstancesRange;
-		TRange<int32> GeometriesToUpdateRange;
-	};
-
-	class FDynamicRayTracingInstancesContext
-	{
-	public:
-		FDynamicRayTracingInstancesContext(FScene& Scene, FViewInfo& InView, FSceneRenderingBulkObjectAllocator& InBulkAllocator, const FSceneOptions& InSceneOptions);
-
-		void GatherDynamicRayTracingInstances(TConstArrayView<int32> InDynamicPrimitives);
-
-		void Finish(FRHICommandListImmediate& InRHICmdList);
-
-		void AddInstancesToScene(FRayTracingScene& RayTracingScene, FRayTracingShaderBindingTable& RayTracingSBT, int64 SharedBufferGenerationID);
-
-		void CollectRDGResources(FRDGBuilder& RDGBuilder);
-
-	private:
-
-		FScene& Scene;
-		FViewInfo& View;
-		FSceneRenderingBulkObjectAllocator& BulkAllocator;
-		const FSceneOptions& SceneOptions;
-
-		bool bTrackReferencedGeometryGroups;
-
-		FRHICommandList* RHICmdList;
-		FGlobalDynamicVertexBuffer DynamicVertexBuffer;
-		FGlobalDynamicIndexBuffer DynamicIndexBuffer;
-
-		FRayTracingInstanceCollector RayTracingInstanceCollector;
-
-		TArray<FDynamicRayTracingPrimitive> DynamicRayTracingPrimitives;
-	};
-
 	FDynamicRayTracingInstancesContext::FDynamicRayTracingInstancesContext(FScene& Scene, FViewInfo& InView, FSceneRenderingBulkObjectAllocator& InBulkAllocator, const FSceneOptions& InSceneOptions)
 		: Scene(Scene)
 		, View(InView)
@@ -535,35 +547,71 @@ namespace RayTracing
 		);
 	}
 
+	void FDynamicRayTracingInstancesContext::GatherDynamicRayTracingInstances_Internal(int32 PrimitiveIndex)
+	{
+		FPrimitiveSceneProxy* SceneProxy = Scene.PrimitiveSceneProxies[PrimitiveIndex];
+		FPrimitiveSceneInfo* SceneInfo = Scene.Primitives[PrimitiveIndex];
+		const FPersistentPrimitiveIndex PersistentPrimitiveIndex = SceneInfo->GetPersistentIndex();
+
+		RayTracingInstanceCollector.SetPrimitive(SceneProxy, FHitProxyId::InvisibleHitProxyId);
+
+		int32 BaseRayTracingInstance = RayTracingInstanceCollector.RayTracingInstances.Num();
+		int32 BaseGeometryToUpdate = RayTracingInstanceCollector.RayTracingGeometriesToUpdate.Num();
+
+		if (SceneOptions.bTranslucentGeometry || SceneProxy->IsOpaqueOrMasked())
+		{
+			SceneProxy->GetDynamicRayTracingInstances(RayTracingInstanceCollector);
+		}
+
+		FDynamicRayTracingPrimitive Tmp;
+		Tmp.PrimitiveIndex = PrimitiveIndex;
+		Tmp.InstancesRange = TRange<int32>(BaseRayTracingInstance, RayTracingInstanceCollector.RayTracingInstances.Num());
+		Tmp.GeometriesToUpdateRange = TRange<int32>(BaseGeometryToUpdate, RayTracingInstanceCollector.RayTracingGeometriesToUpdate.Num());
+
+		DynamicRayTracingPrimitives.Add(MoveTemp(Tmp));
+	}
+
 	void FDynamicRayTracingInstancesContext::GatherDynamicRayTracingInstances(TConstArrayView<int32> InDynamicPrimitives)
 	{
 		TRACE_CPUPROFILER_EVENT_SCOPE(RayTracing_GatherDynamicRayTracingInstances);
 
 		DynamicRayTracingPrimitives.Reserve(DynamicRayTracingPrimitives.Num() + InDynamicPrimitives.Num());
 
-		for (int32 PrimitiveIndex : InDynamicPrimitives)
+		if (!IsParallelGatherDynamicRayTracingInstancesEnabled())
 		{
-			FPrimitiveSceneProxy* SceneProxy = Scene.PrimitiveSceneProxies[PrimitiveIndex];
-			FPrimitiveSceneInfo* SceneInfo = Scene.Primitives[PrimitiveIndex];
-			const FPersistentPrimitiveIndex PersistentPrimitiveIndex = SceneInfo->GetPersistentIndex();
+			RenderThreadDynamicPrimitives = InDynamicPrimitives;
+			return;
+		}
 
-			RayTracingInstanceCollector.SetPrimitive(SceneProxy, FHitProxyId::InvisibleHitProxyId);
+		// TODO: Could filter primitives whose proxy supports ParallelGDRTI during GatherRayTracingRelevantPrimitives_Parallel
 
-			int32 BaseRayTracingInstance = RayTracingInstanceCollector.RayTracingInstances.Num();
-			int32 BaseGeometryToUpdate = RayTracingInstanceCollector.RayTracingGeometriesToUpdate.Num();
+		for (int32 PrimitiveIndex : InDynamicPrimitives)
+		{
+			FPrimitiveSceneProxy* SceneProxy = Scene.PrimitiveSceneProxies[PrimitiveIndex];
 
-			if (SceneOptions.bTranslucentGeometry || SceneProxy->IsOpaqueOrMasked())
+			if (SceneProxy->SupportsParallelGDRTI())
 			{
-				SceneProxy->GetDynamicRayTracingInstances(RayTracingInstanceCollector);
+				GatherDynamicRayTracingInstances_Internal(PrimitiveIndex);
 			}
+			else
+			{
+				RenderThreadDynamicPrimitives.Add(PrimitiveIndex);
+			}
+		}
+	}
+
+	void FDynamicRayTracingInstancesContext::GatherDynamicRayTracingInstances_RenderThread()
+	{
+		TRACE_CPUPROFILER_EVENT_SCOPE(RayTracing_GatherDynamicRayTracingInstances_RenderThread);
 
-			FDynamicRayTracingPrimitive Tmp;
-			Tmp.PrimitiveIndex = PrimitiveIndex;
-			Tmp.InstancesRange = TRange<int32>(BaseRayTracingInstance, RayTracingInstanceCollector.RayTracingInstances.Num());
-			Tmp.GeometriesToUpdateRange = TRange<int32>(BaseGeometryToUpdate, RayTracingInstanceCollector.RayTracingGeometriesToUpdate.Num());
+		check(IsInRenderingThread());
 
-			DynamicRayTracingPrimitives.Add(MoveTemp(Tmp));
+		for (int32 PrimitiveIndex : RenderThreadDynamicPrimitives)
+		{
+			GatherDynamicRayTracingInstances_Internal(PrimitiveIndex);
 		}
+
+		RenderThreadDynamicPrimitives.Empty();
 	}
 
 	void FDynamicRayTracingInstancesContext::Finish(FRHICommandListImmediate& InRHICmdList)
@@ -629,7 +677,7 @@ namespace RayTracing
 
 	void FDynamicRayTracingInstancesContext::AddInstancesToScene(FRayTracingScene& RayTracingScene, FRayTracingShaderBindingTable& RayTracingSBT, int64 SharedBufferGenerationID)
 	{
-		TRACE_CPUPROFILER_EVENT_SCOPE(RayTracing_AddInstancesToScene);
+		TRACE_CPUPROFILER_EVENT_SCOPE(RayTracing_AddDynamicInstancesToScene);
 
 		const float CurrentWorldTime = View.Family->Time.GetWorldTimeSeconds();
 
@@ -1049,7 +1097,7 @@ namespace RayTracing
 			}
 		}
 
-		// TODO: check whether it's ok to do this on a parallel task
+		// TODO: could be moved to GatherRelevantStaticPrimitivesTask so that launching GatherDynamicRayTracingInstancesTask isn't delayed
 		FPrimitiveSceneInfo::UpdateCachedRaytracingData(&TaskData.Scene, TaskData.DirtyCachedRayTracingPrimitives);
 	}
 
@@ -1637,33 +1685,14 @@ namespace RayTracing
 	}
 
 	void AddDynamicInstancesToRayTracingScene(
-		FRDGBuilder& GraphBuilder,
-		FScene& Scene,
-		FViewInfo& View,
-		const RayTracing::FSceneOptions& SceneOptions,
-		TConstArrayView<int32> DynamicPrimitives,
+		FGatherInstancesTaskData& TaskData,
 		FRayTracingScene& RayTracingScene,
 		FRayTracingShaderBindingTable& RayTracingSBT,
-		FGlobalDynamicReadBuffer& InDynamicReadBuffer,
-		FSceneRenderingBulkObjectAllocator& InBulkAllocator)
+		int64 SharedBufferGenerationID)
 	{
 		TRACE_CPUPROFILER_EVENT_SCOPE(RayTracingScene_AddDynamicInstances);
 
-		const float CurrentWorldTime = View.Family->Time.GetWorldTimeSeconds();
-
-		const bool bParallelMeshBatchSetup = GRayTracingParallelMeshBatchSetup && FApp::ShouldUseThreadingForPerformance();
-		const bool bUsingReferenceBasedResidency = IsRayTracingUsingReferenceBasedResidency();
-
-		const int64 SharedBufferGenerationID = Scene.GetRayTracingDynamicGeometryCollection()->BeginUpdate();
-
-		FDynamicRayTracingInstancesContext DynamicRayTracingInstancesContext(Scene, View, InBulkAllocator, SceneOptions);
-		DynamicRayTracingInstancesContext.GatherDynamicRayTracingInstances(DynamicPrimitives);
-		DynamicRayTracingInstancesContext.Finish(GraphBuilder.RHICmdList);
-
-		Scene.GPUScene.UploadDynamicPrimitiveShaderDataForView(GraphBuilder, View, /*bRayTracing*/ true);
-
-		DynamicRayTracingInstancesContext.AddInstancesToScene(RayTracingScene, RayTracingSBT, SharedBufferGenerationID);
-		DynamicRayTracingInstancesContext.CollectRDGResources(GraphBuilder);
+		TaskData.DynamicRayTracingInstancesContext.AddInstancesToScene(RayTracingScene, RayTracingSBT, SharedBufferGenerationID);
 	}
 
 	void BeginGatherInstances(FGatherInstancesTaskData& TaskData, UE::Tasks::FTask FrustumCullTask)
@@ -1688,6 +1717,13 @@ namespace RayTracing
 				GatherRelevantStaticPrimitives(TaskData, LODScaleCVarValue, ForcedLODLevel, bUsingReferenceBasedResidency);
 			}, TaskData.GatherRelevantPrimitivesTask);
 
+		// TODO: Could gather dynamic ray tracing instances using multiple tasks / FDynamicRayTracingInstancesContext
+		TaskData.GatherDynamicRayTracingInstancesTask = UE::Tasks::Launch(UE_SOURCE_LOCATION, [&TaskData]
+			{
+				FTaskTagScope Scope(ETaskTag::EParallelRenderingThread);
+				TaskData.DynamicRayTracingInstancesContext.GatherDynamicRayTracingInstances(TaskData.DynamicPrimitivesIndices);
+			}, TaskData.GatherRelevantPrimitivesTask);
+
 		TaskData.AddUsedStreamingHandlesTask = UE::Tasks::Launch(UE_SOURCE_LOCATION, [&TaskData, LODScaleCVarValue, ForcedLODLevel, bUsingReferenceBasedResidency]
 			{
 				FTaskTagScope Scope(ETaskTag::EParallelRenderingThread);
@@ -1729,30 +1765,28 @@ namespace RayTracing
 		// is done on task threads, therefore all work must be done here up-front as UpdateUniformExpressionCacheIfNeeded is not free-threaded.
 		check(!FMaterialRenderProxy::HasDeferredUniformExpressionCacheRequests());
 
-		View.DynamicRayTracingMeshCommandStorage.Reserve(Scene.Primitives.Num());
-
 		RayTracingSBT.ResetDynamicAllocationData();
 
 		INC_DWORD_STAT_BY(STAT_VisibleRayTracingPrimitives, TaskData.StaticPrimitives.Num() + TaskData.CachedStaticPrimitives.Num() + TaskData.DynamicPrimitivesIndices.Num());
 
-		AddDynamicInstancesToRayTracingScene(
-			GraphBuilder,
-			Scene,
-			View,
-			TaskData.SceneOptions,
-			TaskData.DynamicPrimitivesIndices,
-			RayTracingScene,
-			RayTracingSBT,
-			InDynamicReadBuffer,
-			InBulkAllocator);
-
-		// Task to iterate over static ray tracing instances (also perform auto-instancing on non-cached instances).
+		const int64 SharedBufferGenerationID = Scene.GetRayTracingDynamicGeometryCollection()->BeginUpdate();
+
+		TaskData.GatherDynamicRayTracingInstancesTask.Wait();
+		TaskData.DynamicRayTracingInstancesContext.GatherDynamicRayTracingInstances_RenderThread();
+		TaskData.DynamicRayTracingInstancesContext.Finish(GraphBuilder.RHICmdList);
+
+		Scene.GPUScene.UploadDynamicPrimitiveShaderDataForView(GraphBuilder, View, /*bRayTracing*/ true);
+
+		TaskData.DynamicRayTracingInstancesContext.CollectRDGResources(GraphBuilder);
+
 		// This adds final instances to the ray tracing scene and must be done before FRayTracingScene::BuildInitializationData().
-		UE::Tasks::FTask AddStaticInstancesTask = UE::Tasks::Launch(UE_SOURCE_LOCATION,
-			[&TaskData, &Scene, &View, &RayTracingScene, &RayTracingSBT]()
+		UE::Tasks::FTask AddInstancesTask = UE::Tasks::Launch(UE_SOURCE_LOCATION,
+			[&TaskData, &Scene, &View, &RayTracingScene, &RayTracingSBT, SharedBufferGenerationID]()
 			{
 				FTaskTagScope TaskTagScope(ETaskTag::EParallelRenderingThread);
 
+				AddDynamicInstancesToRayTracingScene(TaskData, RayTracingScene, RayTracingSBT, SharedBufferGenerationID);
+
 				View.VisibleRayTracingShaderBindings.Reserve(TaskData.StaticPrimitives.Num() + TaskData.NumCachedStaticVisibleShaderBindings);
 
 				AddStaticInstancesToRayTracingScene(
@@ -1806,7 +1840,7 @@ namespace RayTracing
 
 				TRACE_CPUPROFILER_EVENT_SCOPE(RayTracingSceneInitTask);
 				RayTracingScene.BuildInitializationData();
-			}, AddStaticInstancesTask);
+			}, AddInstancesTask);
 
 		// wait for this task here, although it could be done later in the frame
 		// since it's only consumed by FCoarseMeshStreamingManager::UpdateResourceStates() during PostRenderAllViewports_RenderThread
