Force FullPrecision for Mobile DOF shaders

diff --git a/Engine/Source/Runtime/Renderer/Private/PostProcess/PostProcessMobile.cpp b/Engine/Source/Runtime/Renderer/Private/PostProcess/PostProcessMobile.cpp
--- a/Engine/Source/Runtime/Renderer/Private/PostProcess/PostProcessMobile.cpp
+++ b/Engine/Source/Runtime/Renderer/Private/PostProcess/PostProcessMobile.cpp
@@ -1281,6 +1281,11 @@ public:
 		PermutationVector.Set<FUseSunDim>(bInUseSun);
 		return PermutationVector;
 	}
+
+	static void ModifyCompilationEnvironment(const FGlobalShaderPermutationParameters& Parameters, FShaderCompilerEnvironment& OutEnvironment)
+	{
+		OutEnvironment.FullPrecisionInPS = 1;
+	}
 };
 
 IMPLEMENT_GLOBAL_SHADER(FMobileDofDownPS, "/Engine/Private/PostProcessMobile.usf", "DofDownPS_Mobile", SF_Pixel);
@@ -1415,6 +1420,11 @@ public:
 		PermutationVector.Set<FUseSunDim>(bInUseSun);
 		return PermutationVector;
 	}
+
+	static void ModifyCompilationEnvironment(const FGlobalShaderPermutationParameters& Parameters, FShaderCompilerEnvironment& OutEnvironment)
+	{
+		OutEnvironment.FullPrecisionInPS = 1;
+	}
 };
 
 IMPLEMENT_GLOBAL_SHADER(FMobileDofNearPS, "/Engine/Private/PostProcessMobile.usf", "DofNearPS_Mobile", SF_Pixel);
@@ -1515,6 +1525,11 @@ public:
 	{
 		return IsMobilePlatform(Parameters.Platform);
 	}
+
+	static void ModifyCompilationEnvironment(const FGlobalShaderPermutationParameters& Parameters, FShaderCompilerEnvironment& OutEnvironment)
+	{
+		OutEnvironment.FullPrecisionInPS = 1;
+	}
 };
 
 IMPLEMENT_GLOBAL_SHADER(FMobileDofBlurPS, "/Engine/Private/PostProcessMobile.usf", "DofBlurPS_Mobile", SF_Pixel);
@@ -1632,6 +1647,11 @@ public:
 	{
 		return IsMobilePlatform(Parameters.Platform);
 	}
+
+	static void ModifyCompilationEnvironment(const FGlobalShaderPermutationParameters& Parameters, FShaderCompilerEnvironment& OutEnvironment)
+	{
+		OutEnvironment.FullPrecisionInPS = 1;
+	}
 };
 
 IMPLEMENT_GLOBAL_SHADER(FMobileIntegrateDofPS, "/Engine/Private/PostProcessMobile.usf", "IntegrateDOFPS_Mobile", SF_Pixel);
