Turned off scene extension callback validation as it triggers for planar reflections (as they may skip calling render).

diff --git a/Engine/Source/Runtime/Renderer/Private/SceneExtensions.cpp b/Engine/Source/Runtime/Renderer/Private/SceneExtensions.cpp
--- a/Engine/Source/Runtime/Renderer/Private/SceneExtensions.cpp
+++ b/Engine/Source/Runtime/Renderer/Private/SceneExtensions.cpp
@@ -115,7 +115,7 @@ void FSceneExtensionsUpdaters::End()
 
 void FSceneExtensionsRenderers::Begin(FSceneRendererBase& InSceneRenderer, const FEngineShowFlags& EngineShowFlags, bool bInValidateCallbacks)
 {
-	bValidateCallbacks = bInValidateCallbacks;
+	bValidateCallbacks = false;//bInValidateCallbacks;
 	ValidateAdvanceCallbackStage(ECallbackStage::Begin);
 	checkf(!IsRendering(), TEXT("Detected FSceneExtensionsRenderer Begin() without matching End()"));
 
diff --git a/Engine/Source/Runtime/Renderer/Private/SceneExtensions.h b/Engine/Source/Runtime/Renderer/Private/SceneExtensions.h
--- a/Engine/Source/Runtime/Renderer/Private/SceneExtensions.h
+++ b/Engine/Source/Runtime/Renderer/Private/SceneExtensions.h
@@ -333,7 +333,7 @@ public:
 private:
 
 	// true if the callback order should be validated.
-	bool bValidateCallbacks = true;
+	bool bValidateCallbacks = false;
 	// Just used to validate that the stages happen in the expected order and none are skipped.
 	enum class ECallbackStage : int32
 	{
