Fixed crash in mobile depth prepass when depth pass is not valid.

diff --git a/Engine/Source/Runtime/Renderer/Private/DepthRendering.cpp b/Engine/Source/Runtime/Renderer/Private/DepthRendering.cpp
--- a/Engine/Source/Runtime/Renderer/Private/DepthRendering.cpp
+++ b/Engine/Source/Runtime/Renderer/Private/DepthRendering.cpp
@@ -654,18 +654,20 @@ bool FMobileSceneRenderer::ShouldRenderPrePass() const
 
 void FMobileSceneRenderer::RenderPrePass(FRHICommandList& RHICmdList, const FViewInfo& View, const FInstanceCullingDrawParams* InstanceCullingDrawParams)
 {
-	checkSlow(RHICmdList.IsInsideRenderPass());
+	if (auto* Pass = View.ParallelMeshDrawCommandPasses[EMeshPass::DepthPass])
+	{
+		checkSlow(RHICmdList.IsInsideRenderPass());
 
-	SCOPED_NAMED_EVENT(FMobileSceneRenderer_RenderPrePass, FColor::Emerald);
-	RHI_BREADCRUMB_EVENT_STAT(RHICmdList, Prepass, "MobileRenderPrePass");
-	SCOPED_GPU_STAT(RHICmdList, Prepass);
+		SCOPED_NAMED_EVENT(FMobileSceneRenderer_RenderPrePass, FColor::Emerald);
+		RHI_BREADCRUMB_EVENT_STAT(RHICmdList, Prepass, "MobileRenderPrePass");
+		SCOPED_GPU_STAT(RHICmdList, Prepass);
 
-	SCOPE_CYCLE_COUNTER(STAT_DepthDrawTime);
-	CSV_SCOPED_TIMING_STAT_EXCLUSIVE(RenderPrePass);
+		SCOPE_CYCLE_COUNTER(STAT_DepthDrawTime);
+		CSV_SCOPED_TIMING_STAT_EXCLUSIVE(RenderPrePass);
 
-	SetStereoViewport(RHICmdList, View);
-	check(View.ParallelMeshDrawCommandPasses[EMeshPass::DepthPass]);
-	View.ParallelMeshDrawCommandPasses[EMeshPass::DepthPass]->Draw(RHICmdList, InstanceCullingDrawParams);
+		SetStereoViewport(RHICmdList, View);
+		Pass->Draw(RHICmdList, InstanceCullingDrawParams);
+	}
 }
 
 void FDeferredShadingSceneRenderer::RenderPrePassHMD(FRDGBuilder& GraphBuilder, TArrayView<FViewInfo> InViews, FRDGTextureRef DepthTexture)
