[HWRT] Add GPUSkinCacheTask has a prerequisite for GatherDynamicRayTracingInstancesTask since FSkeletalMeshSceneProxy::GetDynamicRayTracingInstances(...) depends on it.

diff --git a/Engine/Source/Runtime/Renderer/Private/RayTracing/RayTracing.cpp b/Engine/Source/Runtime/Renderer/Private/RayTracing/RayTracing.cpp
--- a/Engine/Source/Runtime/Renderer/Private/RayTracing/RayTracing.cpp
+++ b/Engine/Source/Runtime/Renderer/Private/RayTracing/RayTracing.cpp
@@ -1717,12 +1717,15 @@ namespace RayTracing
 				GatherRelevantStaticPrimitives(TaskData, LODScaleCVarValue, ForcedLODLevel, bUsingReferenceBasedResidency);
 			}, TaskData.GatherRelevantPrimitivesTask);
 
+		UE::Tasks::FTask GPUSkinCacheTask = TaskData.Scene.GetGPUSkinCacheTask();
+
 		// TODO: Could gather dynamic ray tracing instances using multiple tasks / FDynamicRayTracingInstancesContext
 		TaskData.GatherDynamicRayTracingInstancesTask = UE::Tasks::Launch(UE_SOURCE_LOCATION, [&TaskData]
 			{
 				FTaskTagScope Scope(ETaskTag::EParallelRenderingThread);
 				TaskData.DynamicRayTracingInstancesContext.GatherDynamicRayTracingInstances(TaskData.DynamicPrimitivesIndices);
-			}, TaskData.GatherRelevantPrimitivesTask);
+			}, UE::Tasks::Prerequisites(TaskData.GatherRelevantPrimitivesTask, GPUSkinCacheTask));
+
 
 		TaskData.AddUsedStreamingHandlesTask = UE::Tasks::Launch(UE_SOURCE_LOCATION, [&TaskData, LODScaleCVarValue, ForcedLODLevel, bUsingReferenceBasedResidency]
 			{
