[HWRT] Add geometry debug name to check in AddDynamicMeshBatchForGeometryUpdate(...).

diff --git a/Engine/Source/Runtime/Renderer/Private/RayTracing/RayTracingDynamicGeometry.cpp b/Engine/Source/Runtime/Renderer/Private/RayTracing/RayTracingDynamicGeometry.cpp
--- a/Engine/Source/Runtime/Renderer/Private/RayTracing/RayTracingDynamicGeometry.cpp
+++ b/Engine/Source/Runtime/Renderer/Private/RayTracing/RayTracingDynamicGeometry.cpp
@@ -442,7 +442,7 @@ void FRayTracingDynamicGeometryCollection::AddDynamicMeshBatchForGeometryUpdate(
 
 	if (Geometry.Initializer.TotalPrimitiveCount != UpdateParams.NumTriangles)
 	{
-		check(Geometry.Initializer.Segments.Num() <= 1);
+		checkf(Geometry.Initializer.Segments.Num() <= 1, TEXT("Dynamic ray tracing geometry '%s' has an unexpected number of segments."), *Geometry.Initializer.DebugName.ToString());
 		Geometry.Initializer.TotalPrimitiveCount = UpdateParams.NumTriangles;
 		Geometry.Initializer.Segments.Empty();
 		FRayTracingGeometrySegment Segment;
