[GPU Scene Lights] Store bClusteredDeferredSupported flag explicitly in FForwardLightData instead of deriving it from LocalLighIndex < ClusteredDeferredSupportedEndIndex.
- we won't be able to rely on order of lights in FForwardLightData when using persistent / stable light indices.
- this doesn't affect order of light indices in light grid cells (clusterred lights are still sorted so they're always first / together).

diff --git a/Engine/Shaders/Private/ClusteredDeferredShadingPixelShader.usf b/Engine/Shaders/Private/ClusteredDeferredShadingPixelShader.usf
--- a/Engine/Shaders/Private/ClusteredDeferredShadingPixelShader.usf
+++ b/Engine/Shaders/Private/ClusteredDeferredShadingPixelShader.usf
@@ -161,7 +161,7 @@ float4 GetLightGridLocalLighting(
 
 		// The lights are sorted such that all that support clustered deferred are at the beginning, there might be others
 		// (e.g., lights with dynamic shadows) so we break out when the condition fails.
-		if (!LocalLight.Internal.bClusteredDeferredSupported)
+		if (!UnpackIsClusteredDeferredSupported(LocalLight))
 		{
 			break;
 		}
diff --git a/Engine/Shaders/Private/LightData.ush b/Engine/Shaders/Private/LightData.ush
--- a/Engine/Shaders/Private/LightData.ush
+++ b/Engine/Shaders/Private/LightData.ush
@@ -81,8 +81,6 @@ struct FForwardLightData
 	float3 RectData;
 
 	// Derived terms after unpacking
-	/** Flag is true if clustered deferred is supported for this light. They are always first / together.*/
-	bool bClusteredDeferredSupported;
 	/** Virtual shadow map ID or INDEX_NONE if not present.*/
 	int  VirtualShadowMapId;
 	/** Previous frame LocalLightIndex or INDEX_NONE if not present.*/
@@ -377,6 +375,22 @@ uint UnpackIsHandledByMegaLights(FLocalLightData In)
 	return UnpackIsHandledByMegaLights(In.Internal);
 }
 
+uint UnpackIsClusteredDeferredSupported(uint InLightSceneInfoExtraDataPacked)
+{
+	return ((InLightSceneInfoExtraDataPacked >> 30) & 0x1);
+}
+
+uint UnpackIsClusteredDeferredSupported(FForwardLightData In)
+{
+	const uint LightSceneInfoExtraDataPacked = asuint(In.LightDirectionAndSceneInfoExtraDataPacked.w);
+	return UnpackIsClusteredDeferredSupported(LightSceneInfoExtraDataPacked);
+}
+
+uint UnpackIsClusteredDeferredSupported(FLocalLightData In)
+{
+	return UnpackIsClusteredDeferredSupported(In.Internal);
+}
+
 float UnpackVolumetricScatteringIntensity(FForwardLightData In)
 {
 	return f16tof32(asuint(In.SpotAnglesAndSourceRadiusPacked.w) >> 16);
diff --git a/Engine/Shaders/Private/LightGridCommon.ush b/Engine/Shaders/Private/LightGridCommon.ush
--- a/Engine/Shaders/Private/LightGridCommon.ush
+++ b/Engine/Shaders/Private/LightGridCommon.ush
@@ -276,12 +276,10 @@ FDirectionalLightData GetDirectionalLightData(uint EyeIndex)
 
 FForwardLightData GetForwardLightData_Internal(
 	uint LightIndex,
-	uint ClusteredDeferredSupportedEndIndex,
 	float4 InData0, float4 InData1, float4 InData2, float4 InData3, float4 InData4, float4 InData5)
 {
 	FForwardLightData Out = (FForwardLightData)0;
 
-	Out.bClusteredDeferredSupported					= LightIndex < ClusteredDeferredSupportedEndIndex;
 	Out.LightPositionAndInvRadius					= InData0;
 	Out.LightColorAndIdAndFalloffExponent			= InData1;
 	Out.LightDirectionAndSceneInfoExtraDataPacked	= InData2;
@@ -302,7 +300,6 @@ FForwardLightData GetForwardLightDataNonStereo(uint LightIndex)
 
 	return GetForwardLightData_Internal(
 		LightIndex,
-		ForwardLightStruct.ClusteredDeferredSupportedEndIndex,
 		ForwardLightStruct.ForwardLightBuffer[LightBaseIndex + 0],
 		ForwardLightStruct.ForwardLightBuffer[LightBaseIndex + 1],
 		ForwardLightStruct.ForwardLightBuffer[LightBaseIndex + 2],
@@ -331,8 +328,7 @@ FLocalLightData GetLocalLightData(uint GridIndex, uint EyeIndex)
 		uint LocalLightBaseIndex = LocalLightIndex * FORWARD_LIGHT_DATA_STRIDE;
 
 		Out.Internal = GetForwardLightData_Internal(
-			LocalLightIndex, 
-			ForwardLightStruct.ClusteredDeferredSupportedEndIndex,
+			LocalLightIndex,
 			ForwardLightStruct.ForwardLightBuffer[LocalLightBaseIndex + 0],
 			ForwardLightStruct.ForwardLightBuffer[LocalLightBaseIndex + 1],
 			ForwardLightStruct.ForwardLightBuffer[LocalLightBaseIndex + 2],
@@ -348,7 +344,6 @@ FLocalLightData GetLocalLightData(uint GridIndex, uint EyeIndex)
 
 		Out.Internal = GetForwardLightData_Internal(
 			LocalLightIndex,
-			ForwardLightStructISR.ClusteredDeferredSupportedEndIndex,
 			ForwardLightStructISR.ForwardLightBuffer[LocalLightBaseIndex + 0],
 			ForwardLightStructISR.ForwardLightBuffer[LocalLightBaseIndex + 1],
 			ForwardLightStructISR.ForwardLightBuffer[LocalLightBaseIndex + 2],
diff --git a/Engine/Shaders/Private/MobileLightingCommon.ush b/Engine/Shaders/Private/MobileLightingCommon.ush
--- a/Engine/Shaders/Private/MobileLightingCommon.ush
+++ b/Engine/Shaders/Private/MobileLightingCommon.ush
@@ -589,7 +589,7 @@ void AccumulateLightGridLocalLighting(const FCulledLightsGridHeader InLightGridH
 #if MOBILE_DEFERRED_LIGHTING
 		// The lights are sorted such that all that support clustered deferred are at the beginning, there might be others
 		// (e.g., lights with dynamic shadows) so we break out when the condition fails.
-		if (!LocalLight.Internal.bClusteredDeferredSupported)
+		if (!UnpackIsClusteredDeferredSupported(LocalLight))
 		{
 			break;
 		}
diff --git a/Engine/Source/Runtime/Renderer/Private/LightGridInjection.cpp b/Engine/Source/Runtime/Renderer/Private/LightGridInjection.cpp
--- a/Engine/Source/Runtime/Renderer/Private/LightGridInjection.cpp
+++ b/Engine/Source/Runtime/Renderer/Private/LightGridInjection.cpp
@@ -454,6 +454,7 @@ static void PackLightData(
 	// No shadowmap channels for simple lights
 	uint32 LightSceneInfoExtraDataPacked = 0;
 	LightSceneInfoExtraDataPacked |= SimpleLightLightingChannelMask << 8;
+	LightSceneInfoExtraDataPacked |= 1 << 30; // bClusteredDeferredSupported
 
 	// Pack both values into a single float to keep float4 alignment
 	const float SimpleLightSourceLength = 0;
@@ -790,7 +791,6 @@ FComputeLightGridOutput FSceneRenderer::ComputeLightGrid(FRDGBuilder& GraphBuild
 
 		// Track the end markers for different types
 		int32 SimpleLightsEnd = 0;
-		int32 ClusteredSupportedEnd = 0;
 		int32 MegaLightsSupportedStart = 0;
 		int32 DirectionalMegaLightsSupportedStart = 0;
 
@@ -851,7 +851,6 @@ FComputeLightGridOutput FSceneRenderer::ComputeLightGrid(FRDGBuilder& GraphBuild
 			float SelectedForwardDirectionalLightIntensitySq = 0.0f;
 			int32 SelectedForwardDirectionalLightPriority = -1;
 			const TArray<FSortedLightSceneInfo, SceneRenderingAllocator>& SortedLights = SortedLightSet.SortedLights;
-			ClusteredSupportedEnd = SimpleLightsEnd;
 			MegaLightsSupportedStart = MAX_int32;
 			DirectionalMegaLightsSupportedStart = MAX_int32;
 			// Next add all the other lights, track the end index for clustered supporting lights
@@ -877,7 +876,11 @@ FComputeLightGridOutput FSceneRenderer::ComputeLightGrid(FRDGBuilder& GraphBuild
 						LightParameters.Color *= LightProxy->GetIndirectLightingScale();
 					}
 
-					uint32 LightSceneInfoExtraDataPacked = LightSceneInfo->PackExtraData(bAllowStaticLighting, SortedLightInfo.SortKey.Fields.bLightFunction, SortedLightInfo.SortKey.Fields.bHandledByMegaLights);
+					uint32 LightSceneInfoExtraDataPacked = LightSceneInfo->PackExtraData(
+						bAllowStaticLighting,
+						SortedLightInfo.SortKey.Fields.bLightFunction,
+						SortedLightInfo.SortKey.Fields.bHandledByMegaLights,
+						!SortedLightInfo.SortKey.Fields.bClusteredDeferredNotSupported);
 
 					const bool bDynamicShadows = ViewFamily.EngineShowFlags.DynamicShadows && VisibleLightInfos.IsValidIndex(LightSceneInfo->Id);
 					const int32 VirtualShadowMapId = bDynamicShadows ? VisibleLightInfos[LightSceneInfo->Id].GetVirtualShadowMapId( &View ) : INDEX_NONE;
@@ -940,12 +943,6 @@ FComputeLightGridOutput FSceneRenderer::ComputeLightGrid(FRDGBuilder& GraphBuild
 					{
 						const int32 LocalLightIndex = ViewSpacePosAndRadiusData.Num();
 
-						// Track the last one to support clustered deferred
-						if (!SortedLightInfo.SortKey.Fields.bClusteredDeferredNotSupported)
-						{
-							ClusteredSupportedEnd = FMath::Max(ClusteredSupportedEnd, LocalLightIndex + 1);
-						}
-
 						if (SortedLightInfo.SortKey.Fields.bHandledByMegaLights && MegaLightsSupportedStart == MAX_int32)
 						{
 							MegaLightsSupportedStart = LocalLightIndex;
@@ -1143,7 +1140,6 @@ FComputeLightGridOutput FSceneRenderer::ComputeLightGrid(FRDGBuilder& GraphBuild
 		ForwardLightUniformParameters->CulledGridSize = FIntVector(LightGridSizeXY.X, LightGridSizeXY.Y, GLightGridSizeZ);
 		ForwardLightUniformParameters->MaxCulledLightsPerCell = GLightLinkedListCulling ? NumVisibleLocalLights : GMaxCulledLightsPerCell;
 		ForwardLightUniformParameters->LightGridPixelSizeShift = FMath::FloorLog2(GLightGridPixelSize);
-		ForwardLightUniformParameters->ClusteredDeferredSupportedEndIndex = ClusteredSupportedEnd;
 		ForwardLightUniformParameters->DirectionalMegaLightsSupportedStartIndex = DirectionalMegaLightsSupportedStart;
 		ForwardLightUniformParameters->DirectLightingShowFlag = ViewFamily.EngineShowFlags.DirectLighting ? 1 : 0;
 
@@ -1390,7 +1386,7 @@ FComputeLightGridOutput FSceneRenderer::PrepareForwardLightData(FRDGBuilder& Gra
 						CurrentLightViewData.Color = FVector3f(LightParameters.Color) * LightFade * LightParameters.GetLightExposureScale(Exposure);
 						CurrentLightViewData.VolumetricScatteringIntensity = VolumetricScatteringIntensity;
 						CurrentLightViewData.VirtualShadowMapId = VirtualShadowMapId;
-						CurrentLightViewData.LightSceneInfoExtraDataPacked = LightSceneInfo->PackExtraData(bAllowStaticLighting);
+						CurrentLightViewData.LightSceneInfoExtraDataPacked = LightSceneInfo->PackExtraData(bAllowStaticLighting, false, false, false); // TODO: bLightFunction, bMegaLight, bClusteredDeferredSupported
 						CurrentLightViewData.RectLightAtlasUVOffset = LightParameters.RectLightAtlasUVOffset;
 						CurrentLightViewData.RectLightAtlasUVScale = LightParameters.RectLightAtlasUVScale;
 						CurrentLightViewData.RectLightAtlasMaxLevel = LightParameters.RectLightAtlasMaxLevel;
diff --git a/Engine/Source/Runtime/Renderer/Private/LightSceneInfo.cpp b/Engine/Source/Runtime/Renderer/Private/LightSceneInfo.cpp
--- a/Engine/Source/Runtime/Renderer/Private/LightSceneInfo.cpp
+++ b/Engine/Source/Runtime/Renderer/Private/LightSceneInfo.cpp
@@ -380,7 +380,7 @@ bool FLightSceneInfo::ShouldRecordShadowSubjectsForMobile() const
 	return bShouldRecordShadowSubjectsForMobile;
 }
 
-uint32 FLightSceneInfo::PackExtraData(bool bAllowStaticLighting, bool bLightFunction, bool bMegaLight) const
+uint32 FLightSceneInfo::PackExtraData(bool bAllowStaticLighting, bool bLightFunction, bool bMegaLight, bool bClusteredDeferredSupported) const
 {
 	uint32 Result = 0;
 	uint32 BitOffset = 0;
@@ -427,7 +427,10 @@ uint32 FLightSceneInfo::PackExtraData(bool bAllowStaticLighting, bool bLightFunc
 	uint32 MegaLight = bMegaLight ? 1 : 0;
 	Result |= MegaLight << BitOffset;									BitOffset += 1;
 
-	// 30 bits used, 2 bits free
+	uint32 ClusteredDeferredSupported = bClusteredDeferredSupported ? 1 : 0;
+	Result |= ClusteredDeferredSupported << BitOffset;					BitOffset += 1;
+
+	// 31 bits used, 1 bit free
 
 	return Result;
 }
diff --git a/Engine/Source/Runtime/Renderer/Private/LightSceneInfo.h b/Engine/Source/Runtime/Renderer/Private/LightSceneInfo.h
--- a/Engine/Source/Runtime/Renderer/Private/LightSceneInfo.h
+++ b/Engine/Source/Runtime/Renderer/Private/LightSceneInfo.h
@@ -337,7 +337,7 @@ public:
 
 	// Extra data such as shadow channel masks, light type and flags packed into a uint32.
 	// Unpacked in shader using helper functions defined in LightData.ush
-	uint32 PackExtraData(bool bAllowStaticLighting, bool bLightFunction = false, bool bMegaLight = false) const;
+	uint32 PackExtraData(bool bAllowStaticLighting, bool bLightFunction, bool bMegaLight, bool bClusteredDeferredSupported) const;
 };
 
 /** Defines how the light is stored in the scene's light octree. */
diff --git a/Engine/Source/Runtime/Renderer/Private/SceneRendering.h b/Engine/Source/Runtime/Renderer/Private/SceneRendering.h
--- a/Engine/Source/Runtime/Renderer/Private/SceneRendering.h
+++ b/Engine/Source/Runtime/Renderer/Private/SceneRendering.h
@@ -683,7 +683,6 @@ BEGIN_GLOBAL_SHADER_PARAMETER_STRUCT_WITH_CONSTRUCTOR(FForwardLightUniformParame
 	SHADER_PARAMETER(float, DirectionalLightDepthBias)
 	SHADER_PARAMETER(uint32, DirectionalLightUseStaticShadowing)
 	SHADER_PARAMETER(uint32, DirectionalLightHandledByMegaLights)
-	SHADER_PARAMETER(uint32, ClusteredDeferredSupportedEndIndex)
 	SHADER_PARAMETER(uint32, DirectionalMegaLightsSupportedStartIndex)
 	SHADER_PARAMETER(FVector4f, DirectionalLightStaticShadowBufferSize)
 	SHADER_PARAMETER(FMatrix44f, DirectionalLightTranslatedWorldToStaticShadow)
