[GPU Scene Lights] Store simple lights after regular lights in ForwardLightBuffer.
- this doesn't affect order of light indices in light grid cells (simple lights are still sorted so they're always first / together).

diff --git a/Engine/Source/Runtime/Renderer/Private/LightGridInjection.cpp b/Engine/Source/Runtime/Renderer/Private/LightGridInjection.cpp
--- a/Engine/Source/Runtime/Renderer/Private/LightGridInjection.cpp
+++ b/Engine/Source/Runtime/Renderer/Private/LightGridInjection.cpp
@@ -803,48 +803,15 @@ FComputeLightGridOutput FSceneRenderer::ComputeLightGrid(FRDGBuilder& GraphBuild
 		{
 			// Simple lights are copied without view dependent checks, so same in and out
 			SimpleLightsEnd = SortedLightSet.SimpleLightsEnd;
-			// 1. insert simple lights
+
+			// 1. reserve entries for simple lights
 			if (SimpleLightsEnd > 0)
 			{
-				ForwardLightData.Reserve(SimpleLightsEnd);
-
-				ViewSpacePosAndRadiusData.Reserve(SimpleLightsEnd);
-				ViewSpaceDirAndPreprocAngleData.Reserve(SimpleLightsEnd);
-				ViewSpaceRectPlanesData.Reserve(SimpleLightsEnd * NUM_PLANES_PER_RECT_LIGHT);
-
-				const FSimpleLightArray& SimpleLights = SortedLightSet.SimpleLights;
-
-				// Pack both values into a single float to keep float4 alignment
-				const FFloat16 SimpleLightSourceLength16f = FFloat16(0);
-				FLightingChannels SimpleLightLightingChannels;
-
-				// Put simple lights in all lighting channels
-				SimpleLightLightingChannels.bChannel0 = SimpleLightLightingChannels.bChannel1 = SimpleLightLightingChannels.bChannel2 = true;
-				const uint32 SimpleLightLightingChannelMask = GetLightingChannelMaskForStruct(SimpleLightLightingChannels);
-
-				// Now using the sorted lights, and keep track of ranges as we go.
-				for (int32 SortedIndex = 0; SortedIndex < SimpleLightsEnd; ++SortedIndex)
-				{
-					check(SortedLightSet.SortedLights[SortedIndex].LightSceneInfo == nullptr);
-					check(!SortedLightSet.SortedLights[SortedIndex].SortKey.Fields.bIsNotSimpleLight);
-
+				IndirectionIndices.AddUninitialized(SimpleLightsEnd);
 
-					int32 SimpleLightIndex = SortedLightSet.SortedLights[SortedIndex].SimpleLightIndex;
-
-					const int32 IndexInBuffer = ForwardLightData.AddUninitialized(1);
-					FForwardLightData& LightData = ForwardLightData[IndexInBuffer];
-
-					const FSimpleLightEntry& SimpleLight = SimpleLights.InstanceData[SimpleLightIndex];
-					const FSimpleLightPerViewEntry& SimpleLightPerViewData = SimpleLights.GetViewDependentData(SimpleLightIndex, ViewIndex, Views.Num());
-					PackLightData(LightData, View, SimpleLight, SimpleLightPerViewData);
-
-					FVector4f ViewSpacePosAndRadius(FVector4f(View.ViewMatrices.GetViewMatrix().TransformPosition(SimpleLightPerViewData.Position)), SimpleLight.Radius);
-					ViewSpacePosAndRadiusData.Add(ViewSpacePosAndRadius);
-					ViewSpaceDirAndPreprocAngleData.AddZeroed();
-					ViewSpaceRectPlanesData.AddZeroed(NUM_PLANES_PER_RECT_LIGHT);
-
-					IndirectionIndices.Add(IndexInBuffer);
-				}
+				ViewSpacePosAndRadiusData.AddUninitialized(SimpleLightsEnd);
+				ViewSpaceDirAndPreprocAngleData.AddZeroed(SimpleLightsEnd);
+				ViewSpaceRectPlanesData.AddZeroed(SimpleLightsEnd * NUM_PLANES_PER_RECT_LIGHT);
 			}
 
 			const uint32 LightShaderParameterFlags = bRenderRectLightsAsSpotLights ? ELightShaderParameterFlags::RectAsSpotLight : 0u;
@@ -1105,6 +1072,34 @@ FComputeLightGridOutput FSceneRenderer::ComputeLightGrid(FRDGBuilder& GraphBuild
 					}
 				}
 			}
+
+			// 3. add simple lights into ForwardLightData and fill uninitialized ViewSpacePosAndRadiusData/IndirectionIndices
+			if (SimpleLightsEnd > 0)
+			{
+				ForwardLightData.Reserve(ForwardLightData.Num() + SimpleLightsEnd);
+
+				const FSimpleLightArray& SimpleLights = SortedLightSet.SimpleLights;
+
+				for (int32 SortedIndex = 0; SortedIndex < SimpleLightsEnd; ++SortedIndex)
+				{
+					check(SortedLightSet.SortedLights[SortedIndex].LightSceneInfo == nullptr);
+					check(!SortedLightSet.SortedLights[SortedIndex].SortKey.Fields.bIsNotSimpleLight);
+
+					int32 SimpleLightIndex = SortedLightSet.SortedLights[SortedIndex].SimpleLightIndex;
+
+					const int32 IndexInBuffer = ForwardLightData.AddUninitialized(1);
+					FForwardLightData& LightData = ForwardLightData[IndexInBuffer];
+
+					const FSimpleLightEntry& SimpleLight = SimpleLights.InstanceData[SimpleLightIndex];
+					const FSimpleLightPerViewEntry& SimpleLightPerViewData = SimpleLights.GetViewDependentData(SimpleLightIndex, ViewIndex, Views.Num());
+					PackLightData(LightData, View, SimpleLight, SimpleLightPerViewData); 
+					
+					FVector4f ViewSpacePosAndRadius(FVector4f(View.ViewMatrices.GetViewMatrix().TransformPosition(SimpleLightPerViewData.Position)), SimpleLight.Radius);
+					ViewSpacePosAndRadiusData[SortedIndex] = ViewSpacePosAndRadius;
+
+					IndirectionIndices[SortedIndex] = IndexInBuffer;
+				}
+			}
 		}
 
 #if WITH_EDITOR
