Fix for occlusion queries/feedback in stereo rendering.

* Occlusion Queries and Feedback are now issued only once for both eyes. Only exception is when using instanced stereo with Occlusion Feedback where we issue them twice from the left eye (first view).
* Use the frustum that covers both eyes in single pass stereo (either multi viewport or mobile multi view).
* Adreno fast occlusion mode is added for devices that support it.

diff --git a/Engine/Config/BaseDeviceProfiles.ini b/Engine/Config/BaseDeviceProfiles.ini
--- a/Engine/Config/BaseDeviceProfiles.ini
+++ b/Engine/Config/BaseDeviceProfiles.ini
@@ -1217,6 +1217,7 @@ BaseProfileName=Oculus_Quest
 +CVars=r.Mobile.Oculus.ForceSymmetric=1
 +CVars=fx.NiagaraAllowGPUParticles=1
 +CVars=FX.AllowGPUSorting=1
++CVars=r.Mobile.AdrenoOcclusionMode=1
 
 [Oculus_Quest DeviceProfile]
 DeviceType=Android
diff --git a/Engine/Shaders/Private/OcclusionFeedbackShaders.usf b/Engine/Shaders/Private/OcclusionFeedbackShaders.usf
--- a/Engine/Shaders/Private/OcclusionFeedbackShaders.usf
+++ b/Engine/Shaders/Private/OcclusionFeedbackShaders.usf
@@ -8,17 +8,37 @@
 #endif
 
 uint StartIndex;
+#if INSTANCED_STEREO
+uint InstancedViewIndex;
+#endif
 
 void MainVS(
 	float4 InPosition	: ATTRIBUTE0,
 	float4 InOrigin		: ATTRIBUTE1,
 	float4 InExtent		: ATTRIBUTE2,
 	uint InstanceId		: SV_InstanceID,
+#if MOBILE_MULTI_VIEW && !INSTANCED_STEREO
+	uint ViewId			: SV_ViewID,
+#endif
 	out nointerpolation uint PrimitiveIdx : PRIMITIVE_ID,
 	out float4 OutPosition : SV_POSITION
+#if MOBILE_MULTI_VIEW && INSTANCED_STEREO
+	, out uint StereoLayer : SV_RenderTargetArrayIndex
+#elif INSTANCED_STEREO
+	, out uint StereoLayer : SV_ViewPortArrayIndex
+#endif
 )
 {
+#if MOBILE_MULTI_VIEW && !INSTANCED_STEREO
+	ResolvedView = ResolveView(ViewId);
+#elif INSTANCED_STEREO
+	uint ViewId = InstancedViewIndex;
+	StereoLayer = ViewId;
+	ResolvedView = ResolveView(ViewId);
+#else
 	ResolvedView = ResolveView();
+#endif
+
 	float3 TransformedPosition = (InPosition.xyz * InExtent.xyz) + InOrigin.xyz;
 	OutPosition = mul(float4(TransformedPosition, 1), ResolvedView.TranslatedWorldToClip);
 	PrimitiveIdx = StartIndex + InstanceId;
diff --git a/Engine/Shaders/Private/OcclusionQueryVertexShader.usf b/Engine/Shaders/Private/OcclusionQueryVertexShader.usf
--- a/Engine/Shaders/Private/OcclusionQueryVertexShader.usf
+++ b/Engine/Shaders/Private/OcclusionQueryVertexShader.usf
@@ -6,35 +6,43 @@
 
 #include "Common.ush"
 
+#if INSTANCED_STEREO || MOBILE_MULTI_VIEW
+float4 StencilingGeometryPosAndScale[2];
+#else
 float4 StencilingGeometryPosAndScale;
-
-#if MOBILE_MULTI_VIEW
-	uint ViewId;
 #endif
 
 void Main(
 	in float4 InPosition : ATTRIBUTE0,
 	out float4 OutPosition : SV_POSITION
-#if MOBILE_MULTI_VIEW && !INSTANCED_STEREO
+#if INSTANCED_STEREO
+	, uint InstanceId : SV_InstanceID
+#if MOBILE_MULTI_VIEW
+	, out nointerpolation uint LayerIndex : SV_RenderTargetArrayIndex
+#else
+	, out nointerpolation uint ViewportIndex : SV_ViewPortArrayIndex
+#endif // MOBILE_MULTI_VIEW
+#elif MOBILE_MULTI_VIEW
 	, in uint InViewId : SV_ViewID
-	, out uint OutViewId : VIEW_ID
-#endif
+#endif // INSTANCED_STEREO
 	)
 {
-#if MOBILE_MULTI_VIEW && !INSTANCED_STEREO
-	OutViewId = InViewId;
-
-	// Make the triangles degenerate if we're not targeting this view.
-	// We don't need to resolve the view info, since we already set the
-	// primary view with the data for the view we're interested in.
-	// TODO: Support targetting a single view properly.
-	if (InViewId != ViewId)
-	{
-		OutPosition = float4(0, 0, 0, 1);
-		return;
-	}
+#if INSTANCED_STEREO
+	uint InViewId = GetEyeIndex(InstanceId);
+#if MOBILE_MULTI_VIEW
+	LayerIndex = InViewId;
+#else
+	ViewportIndex = InViewId;
+#endif
+#endif
+#if INSTANCED_STEREO || MOBILE_MULTI_VIEW
+	ResolvedView = ResolveView(InViewId);
+	float4 PosAndScale = StencilingGeometryPosAndScale[InViewId];
+#else
+	ResolvedView = ResolveView();
+	float4 PosAndScale = StencilingGeometryPosAndScale;
 #endif
 
-	float3 TransformedPosition = InPosition.xyz * StencilingGeometryPosAndScale.w + StencilingGeometryPosAndScale.xyz;
-	OutPosition = mul(float4(TransformedPosition, 1), View.TranslatedWorldToClip);
+	float3 TransformedPosition = InPosition.xyz * PosAndScale.w + PosAndScale.xyz;
+	OutPosition = mul(float4(TransformedPosition, 1), ResolvedView.TranslatedWorldToClip);
 }
diff --git a/Engine/Source/Runtime/Engine/Private/HierarchicalInstancedStaticMesh.cpp b/Engine/Source/Runtime/Engine/Private/HierarchicalInstancedStaticMesh.cpp
--- a/Engine/Source/Runtime/Engine/Private/HierarchicalInstancedStaticMesh.cpp
+++ b/Engine/Source/Runtime/Engine/Private/HierarchicalInstancedStaticMesh.cpp
@@ -1569,7 +1569,7 @@ void FHierarchicalStaticMeshSceneProxy::GetDynamicMeshElements(const TArray<cons
 
 					// Instanced stereo needs to use the right plane from the right eye when constructing the frustum bounds to cull against.
 					// Otherwise we'll cull objects visible in the right eye, but not the left.
-					if ((View->IsInstancedStereoPass() || View->bIsMobileMultiViewEnabled) && IStereoRendering::IsStereoEyeView(*View) && GEngine->StereoRenderingDevice.IsValid())
+					if (View->bIsSinglePassStereo && IStereoRendering::IsStereoEyeView(*View) && GEngine->StereoRenderingDevice.IsValid())
 					{
 						// TODO: Stereo culling frustum needs to use the culling origin instead of the view origin.
 						InstanceParams.ViewFrustumLocal = View->GetCullingFrustum();
diff --git a/Engine/Source/Runtime/Engine/Private/SceneView.cpp b/Engine/Source/Runtime/Engine/Private/SceneView.cpp
--- a/Engine/Source/Runtime/Engine/Private/SceneView.cpp
+++ b/Engine/Source/Runtime/Engine/Private/SceneView.cpp
@@ -793,8 +793,9 @@ void FSceneView::SetupViewFrustum()
 		GetViewFrustumBounds(ViewFrustum, ViewMatrices.GetViewProjectionMatrix(), false);
 	}
 
-	// Use a unified frustum for culling an instanced stereo pass.
-	if (IStereoRendering::IsStereoEyeView(*this) && GEngine->StereoRenderingDevice.IsValid())
+	// Use a unified frustum for culling in single pass stereo (i.e. when we're drawing both eyes in one pass rather than
+	// rendering them separately as separate views).
+	if (bIsSinglePassStereo && IStereoRendering::IsStereoEyeView(*this) && GEngine->StereoRenderingDevice.IsValid())
 	{
 		StereoCullingFrustum = MakeShared<FConvexVolume>();
 
@@ -917,6 +918,7 @@ PRAGMA_ENABLE_DEPRECATION_WARNINGS
 	, bIsMultiViewportEnabled(false)
 	, bIsMobileMultiViewEnabled(false)
 	, bShouldBindInstancedViewUB(false)
+	, bIsSinglePassStereo(false)
 	, bForceCameraVisibilityReset(false)
 	, bDisableDistanceBasedFadeTransitions(false)
 	, bForcePathTracerReset(false)
@@ -938,15 +940,41 @@ PRAGMA_ENABLE_DEPRECATION_WARNINGS
 	ShadowFreezeCamera::Init(ViewMatrices);
 #endif // !(UE_BUILD_SHIPPING || UE_BUILD_TEST)
 
+	auto ShaderPlatform = GShaderPlatformForFeatureLevel[FeatureLevel];
+
+	// Query instanced stereo and multi-view state
+	{
+		// The shader variants that are compiled have ISR or MMV _enabled_ in the shader, even if the current ViewFamily doesn't
+		// require multiple views functionality.
+		Aspects = UE::StereoRenderUtils::FStereoShaderAspects(ShaderPlatform);
+		bShouldBindInstancedViewUB = Aspects.IsInstancedStereoEnabled() || Aspects.IsMobileMultiViewEnabled();
+
+		if (Family && Family->bRequireMultiView && !Aspects.IsMobileMultiViewEnabled())
+		{
+			UE_LOG(LogMultiView, Fatal, TEXT("Family requires Mobile Multi-View, but it is not supported by the RHI and no fallback is available."));
+		}
+
+		const bool bIsSingleViewCapture = (bIsSceneCapture || bIsReflectionCapture) && !bIsPlanarReflection; // Planar reflections have bIsSceneCapture == true, but can have instanced stereo views
+		bIsInstancedStereoEnabled = !bIsSingleViewCapture && Aspects.IsInstancedStereoEnabled();
+		bIsMultiViewportEnabled = !bIsSingleViewCapture && Aspects.IsInstancedMultiViewportEnabled();
+		bIsMobileMultiViewEnabled = !bIsSingleViewCapture && Family && Family->bRequireMultiView && Aspects.IsMobileMultiViewEnabled();
+		bIsSinglePassStereo = bIsInstancedStereoEnabled || bIsMobileMultiViewEnabled;
+
+		PRAGMA_DISABLE_DEPRECATION_WARNINGS
+		bIsMultiViewEnabled = bIsMultiViewportEnabled;	// temporary, as a graceful way to support plugins/licensee mods
+		PRAGMA_ENABLE_DEPRECATION_WARNINGS
+
+		// If instanced stereo is enabled, we should also have either multiviewport enabled, or mmv fallback enabled.
+		// Assert this since a more graceful handling is done earlier (see PostInitRHI)
+		checkf(!bIsInstancedStereoEnabled || (bIsMultiViewportEnabled || Aspects.IsMobileMultiViewEnabled()),
+			TEXT("If instanced stereo is enabled, either multi-viewport or mobile multi-view needs to be enabled as well."));
+	}
+	// Stereo setup needs to be done before this as that influences how the culling frustrum is created.
 	SetupViewFrustum();
 
 	// Determine whether the view should reverse the cull mode due to a negative determinant.  Only do this for a valid scene
 	bReverseCulling = (Family && Family->Scene) ? ViewMatrices.GetViewMatrix().Determinant() < 0.0f : false;
 
-	// OpenGL Gamma space output in GLSL flips Y when rendering directly to the back buffer (so not needed on PC, as we never render directly into the back buffer)
-	auto ShaderPlatform = GShaderPlatformForFeatureLevel[FeatureLevel];
-	bool bUsingMobileRenderer = GetFeatureLevelShadingPath(FeatureLevel) == EShadingPath::Mobile;
-
 	// Setup transformation constants to be used by the graphics hardware to transform device normalized depth samples
 	// into world oriented z.
 	InvDeviceZToWorldZTransform = CreateInvDeviceZToWorldZTransform(ProjectionMatrixUnadjustedForRHI);
@@ -988,33 +1016,6 @@ PRAGMA_ENABLE_DEPRECATION_WARNINGS
 	SubduedSelectionOutlineColor = GEngine->GetSubduedSelectionOutlineColor();
 #endif
 
-	// Query instanced stereo and multi-view state
-	{
-		// The shader variants that are compiled have ISR or MMV _enabled_ in the shader, even if the current ViewFamily doesn't
-		// require multiple views functionality.
-		Aspects = UE::StereoRenderUtils::FStereoShaderAspects(ShaderPlatform);
-		bShouldBindInstancedViewUB = Aspects.IsInstancedStereoEnabled() || Aspects.IsMobileMultiViewEnabled();
-
-		if (Family && Family->bRequireMultiView && !Aspects.IsMobileMultiViewEnabled())
-		{
-			UE_LOG(LogMultiView, Fatal, TEXT("Family requires Mobile Multi-View, but it is not supported by the RHI and no fallback is available."));
-		}
-
-		const bool bIsSingleViewCapture = (bIsSceneCapture || bIsReflectionCapture) && !bIsPlanarReflection; // Planar reflections have bIsSceneCapture == true, but can have instanced stereo views
-		bIsInstancedStereoEnabled = !bIsSingleViewCapture && Aspects.IsInstancedStereoEnabled();
-		bIsMultiViewportEnabled = !bIsSingleViewCapture && Aspects.IsInstancedMultiViewportEnabled();
-		bIsMobileMultiViewEnabled = !bIsSingleViewCapture && Family && Family->bRequireMultiView && Aspects.IsMobileMultiViewEnabled();
-
-		PRAGMA_DISABLE_DEPRECATION_WARNINGS
-		bIsMultiViewEnabled = bIsMultiViewportEnabled;	// temporary, as a graceful way to support plugins/licensee mods
-		PRAGMA_ENABLE_DEPRECATION_WARNINGS
-
-		// If instanced stereo is enabled, we should also have either multiviewport enabled, or mmv fallback enabled.
-		// Assert this since a more graceful handling is done earlier (see PostInitRHI)
-		checkf(!bIsInstancedStereoEnabled || (bIsMultiViewportEnabled || Aspects.IsMobileMultiViewEnabled()),
-			TEXT("If instanced stereo is enabled, either multi-viewport or mobile multi-view needs to be enabled as well."));
-	}
-
 	SetupAntiAliasingMethod();
 
 	if (AntiAliasingMethod == AAM_TSR)
@@ -2893,7 +2894,7 @@ const FSceneView* FSceneView::GetPrimarySceneView() const
 const FSceneView* FSceneView::GetInstancedSceneView() const
 {
 	// if we don't have ISR or MMV enabled, or we're rendering a full pass, we don't have instanced views
-	if ((bIsMultiViewportEnabled || bIsMobileMultiViewEnabled) && StereoPass != EStereoscopicPass::eSSP_FULL)
+	if (bIsSinglePassStereo && StereoPass != EStereoscopicPass::eSSP_FULL)
 	{
 		// If called on the first secondary view it'll return itself.
 		if (Family && Family->Views.IsValidIndex(PrimaryViewIndex + 1))
diff --git a/Engine/Source/Runtime/Engine/Public/SceneView.h b/Engine/Source/Runtime/Engine/Public/SceneView.h
--- a/Engine/Source/Runtime/Engine/Public/SceneView.h
+++ b/Engine/Source/Runtime/Engine/Public/SceneView.h
@@ -1609,6 +1609,10 @@ public:
 	/** True if we need to bind the instanced view uniform buffer parameters. */
 	bool bShouldBindInstancedViewUB;
 
+	/** True if any single-pass stereo mechanism is enabled (mobile multi view, instanced stereo multi-layer, or instanced stereo multi-viewport).
+	 * Not related to the deprecated NVidia feature known as Single-Pass Stereo. */
+	bool bIsSinglePassStereo;
+
 	/** True if we need to force the camera to discard previous frames occlusion. Necessary for overlapped tile rendering
 	 * where we discard previous frame occlusion because the projection matrix changes.
 	 */
diff --git a/Engine/Source/Runtime/Renderer/Private/MobileShadingRenderer.cpp b/Engine/Source/Runtime/Renderer/Private/MobileShadingRenderer.cpp
--- a/Engine/Source/Runtime/Renderer/Private/MobileShadingRenderer.cpp
+++ b/Engine/Source/Runtime/Renderer/Private/MobileShadingRenderer.cpp
@@ -111,6 +111,12 @@ static TAutoConsoleVariable<int32> CVarMobileXRMSAAMode(
 	TEXT(" 2 = Make the depth swap chain be MSAA and use it directly as scene depth"),
 	ECVF_ReadOnly | ECVF_RenderThreadSafe);
 
+TAutoConsoleVariable<int32> GAdrenoOcclusionUseFDM(
+	TEXT("r.AdrenoOcclusionUseFDM"),
+	0,
+	TEXT("Use FDM with adreno occlusion mode"),
+	ECVF_RenderThreadSafe);
+
 DECLARE_GPU_STAT_NAMED(MobileSceneRender, TEXT("Mobile Scene Render"));
 
 extern bool IsMobileEyeAdaptationEnabled(const FViewInfo& View);
@@ -573,6 +579,9 @@ void FMobileSceneRenderer::InitViews(
 
 	bShouldRenderHZB = ShouldRenderHZB(Views) && bRendererOutputFinalSceneColor;
 
+	bool bUsingOcclusionFeedback = Views[0].ViewState && Views[0].ViewState->OcclusionFeedback.IsInitialized();
+	bAdrenoOcclusionMode = (DoOcclusionQueries() && !bUsingOcclusionFeedback && !Views[0].bDisableQuerySubmissions && CVarMobileAdrenoOcclusionMode.GetValueOnAnyThread() != 0);
+
 	// Whether we need to store depth for post-processing
 	// On PowerVR we see flickering of shadows and depths not updating correctly if targets are discarded.
 	const bool bForceDepthResolve = (CVarMobileForceDepthResolve.GetValueOnRenderThread() == 1);
@@ -592,6 +601,7 @@ void FMobileSceneRenderer::InitViews(
 		bRequireSeparateViewPass ||
 		bIsFullDepthPrepassEnabled ||
 		bShouldRenderHZB ||
+		(bAdrenoOcclusionMode && IsVulkanPlatform(ShaderPlatform)) ||
 		GraphBuilder.IsDumpingFrame();
 	// never keep MSAA depth if SceneDepthAux is enabled
 	bKeepDepthContent = ((NumMSAASamples > 1) && bRequiresSceneDepthAux) ? false : bKeepDepthContent;
@@ -1699,14 +1709,17 @@ void FMobileSceneRenderer::RenderForwardSinglePass(FRDGBuilder& GraphBuilder, FM
 		
 	PassParameters->RenderTargets.SubpassHint = bTonemapSubpassInline ? ESubpassHint::CustomResolveSubpass : ESubpassHint::DepthReadSubpass;
 	const bool bDoOcclusionQueries = (!bIsFullDepthPrepassEnabled && ViewContext.bIsLastView && DoOcclusionQueries());
-	PassParameters->RenderTargets.NumOcclusionQueries = bDoOcclusionQueries ? ComputeNumOcclusionQueriesToBatch() : 0u;
+	const int32 NumOcclusionQueries = bDoOcclusionQueries ? ComputeNumOcclusionQueriesToBatch() : 0u;
+	const bool bVulkanAdrenoOcclusionMode = bAdrenoOcclusionMode && IsVulkanPlatform(ShaderPlatform);
+	const bool bDoOcclusionInMainPass = bDoOcclusionQueries && !bVulkanAdrenoOcclusionMode;
+	PassParameters->RenderTargets.NumOcclusionQueries = bVulkanAdrenoOcclusionMode ? 0u : NumOcclusionQueries;
 	
 	GraphBuilder.AddPass(
 		RDG_EVENT_NAME("SceneColorRendering"),
 		PassParameters,
 		// the second view pass should not be merged with the first view pass on mobile since the subpass would not work properly.
 		ERDGPassFlags::Raster | ERDGPassFlags::NeverMerge,
-		[this, PassParameters, ParameterCollection, ViewContext, bDoOcclusionQueries, &SceneTextures](FRHICommandList& RHICmdList)
+		[this, PassParameters, ParameterCollection, ViewContext, bDoOcclusionInMainPass, &SceneTextures](FRHICommandList& RHICmdList)
 	{
 		FViewInfo& View = *ViewContext.ViewInfo;
 			
@@ -1741,11 +1754,10 @@ void FMobileSceneRenderer::RenderForwardSinglePass(FRDGBuilder& GraphBuilder, FM
 		}
 #endif
 
-		if (bDoOcclusionQueries)
+		if (bDoOcclusionInMainPass)
 		{
 			// Issue occlusion queries
-			const bool bAdrenoOcclusionMode = (CVarMobileAdrenoOcclusionMode.GetValueOnRenderThread() != 0 && IsOpenGLPlatform(ShaderPlatform));
-			if (bAdrenoOcclusionMode)
+			if (bAdrenoOcclusionMode && IsOpenGLPlatform(ShaderPlatform))
 			{
 				// flush
 				RHICmdList.SubmitCommandsHint();
@@ -1767,6 +1779,30 @@ void FMobileSceneRenderer::RenderForwardSinglePass(FRDGBuilder& GraphBuilder, FM
 	{
 		AddResolveSceneDepthPass(GraphBuilder, *ViewContext.ViewInfo, SceneTextures.Depth);
 	}
+
+	if (bDoOcclusionQueries && bVulkanAdrenoOcclusionMode)
+	{
+		FMobileRenderPassParameters* OcclusionPassParams = GraphBuilder.AllocParameters<FMobileRenderPassParameters>();
+		OcclusionPassParams->View = PassParameters->View;
+		OcclusionPassParams->RenderTargets.DepthStencil = FDepthStencilBinding(SceneTextures.Depth.Resolve, ERenderTargetLoadAction::ELoad, FExclusiveDepthStencil::DepthRead_StencilNop);
+		OcclusionPassParams->RenderTargets.MultiViewCount = PassParameters->RenderTargets.MultiViewCount;
+		OcclusionPassParams->RenderTargets.NumOcclusionQueries = NumOcclusionQueries;
+		if (GAdrenoOcclusionUseFDM.GetValueOnRenderThread())
+		{
+			OcclusionPassParams->RenderTargets.ShadingRateTexture = PassParameters->RenderTargets.ShadingRateTexture;
+		}
+
+		GraphBuilder.AddPass(
+			RDG_EVENT_NAME("VulkanAdrenoOcclusion"),
+			OcclusionPassParams,
+			// The occlusion pass needs to be unique to be optimized properly by the driver, don't merge it.
+			// This pass has no observable outputs on the RenderGraph, so it needs to be marked as NeverCull.
+			ERDGPassFlags::Raster | ERDGPassFlags::NeverMerge | ERDGPassFlags::NeverCull,
+			[this](FRHICommandListImmediate& RHICmdList)
+		{
+			RenderOcclusion(RHICmdList);
+		});
+	}
 }
 
 void FMobileSceneRenderer::RenderForwardMultiPass(FRDGBuilder& GraphBuilder, FMobileRenderPassParameters* PassParameters, FRenderViewContext& ViewContext, FSceneTextures& SceneTextures)
diff --git a/Engine/Source/Runtime/Renderer/Private/PlanarReflectionRendering.cpp b/Engine/Source/Runtime/Renderer/Private/PlanarReflectionRendering.cpp
--- a/Engine/Source/Runtime/Renderer/Private/PlanarReflectionRendering.cpp
+++ b/Engine/Source/Runtime/Renderer/Private/PlanarReflectionRendering.cpp
@@ -728,7 +728,7 @@ bool FDeferredShadingSceneRenderer::HasDeferredPlanarReflections(const FViewInfo
 	{
 		FPlanarReflectionSceneProxy* ReflectionSceneProxy = Scene->PlanarReflections[PlanarReflectionIndex];
 
-		if (View.ViewFrustum.IntersectBox(ReflectionSceneProxy->WorldBounds.GetCenter(), ReflectionSceneProxy->WorldBounds.GetExtent()))
+		if (View.GetCullingFrustum().IntersectBox(ReflectionSceneProxy->WorldBounds.GetCenter(), ReflectionSceneProxy->WorldBounds.GetExtent()))
 		{
 			bAnyVisiblePlanarReflections = true;
 			break;
@@ -793,7 +793,7 @@ void FDeferredShadingSceneRenderer::RenderDeferredPlanarReflections(FRDGBuilder&
 
 		for (FPlanarReflectionSceneProxy* ReflectionSceneProxy : Scene->PlanarReflections)
 		{
-			if (!View.ViewFrustum.IntersectBox(ReflectionSceneProxy->WorldBounds.GetCenter(), ReflectionSceneProxy->WorldBounds.GetExtent()))
+			if (!View.GetCullingFrustum().IntersectBox(ReflectionSceneProxy->WorldBounds.GetCenter(), ReflectionSceneProxy->WorldBounds.GetExtent()))
 			{
 				continue;
 			}
diff --git a/Engine/Source/Runtime/Renderer/Private/SceneOcclusion.cpp b/Engine/Source/Runtime/Renderer/Private/SceneOcclusion.cpp
--- a/Engine/Source/Runtime/Renderer/Private/SceneOcclusion.cpp
+++ b/Engine/Source/Runtime/Renderer/Private/SceneOcclusion.cpp
@@ -163,7 +163,7 @@ const uint8* FSceneViewState::ResolvePrecomputedVisibilityData(FViewInfo& View,
 					const FPrecomputedVisibilityCell& CurrentCell = Handler.PrecomputedVisibilityCellBuckets[BucketIndex].Cells[CellIndex];
 					// Construct the cell's bounds
 					const FBox CellBounds(CurrentCell.Min, CurrentCell.Min + FVector(Handler.PrecomputedVisibilityCellSizeXY, Handler.PrecomputedVisibilityCellSizeXY, Handler.PrecomputedVisibilityCellSizeZ));
-					if (View.ViewFrustum.IntersectBox(CellBounds.GetCenter(), CellBounds.GetExtent()))
+					if (View.GetCullingFrustum().IntersectBox(CellBounds.GetCenter(), CellBounds.GetExtent()))
 					{
 						DrawWireBox(&VisibilityCellsPDI, CellBounds, FColor(50, 50, 255), SDPG_World);
 					}
@@ -172,7 +172,7 @@ const uint8* FSceneViewState::ResolvePrecomputedVisibilityData(FViewInfo& View,
 		}
 
 		//Determine view origin
-		FVector ViewOrigin = View.ViewMatrices.GetViewOrigin();
+		FVector ViewOrigin = View.CullingOrigin;
 #if !(UE_BUILD_SHIPPING || UE_BUILD_TEST)
 		if (const FViewMatrices* FrozenViewMatrices = GetFrozenViewMatrices())
 		{
@@ -444,11 +444,12 @@ void FFrameBasedOcclusionQueryPool::AdvanceFrame(uint32 InOcclusionFrameCounter,
 	FrameQueries[CurrentFrameIndex].OcclusionFrameCounter = OcclusionFrameCounter;
 }
 
-FOcclusionQueryBatcher::FOcclusionQueryBatcher(class FSceneViewState* ViewState,uint32 InMaxBatchedPrimitives)
+FOcclusionQueryBatcher::FOcclusionQueryBatcher(class FSceneViewState* ViewState,uint32 InMaxBatchedPrimitives, uint32 InNumInstances)
 :	CurrentBatchOcclusionQuery(NULL)
 ,	MaxBatchedPrimitives(InMaxBatchedPrimitives)
 ,	NumBatchedPrimitives(0)
 ,	OcclusionQueryPool(ViewState ? &ViewState->PrimitiveOcclusionQueryPool : NULL)
+,	NumInstances(InNumInstances)
 {}
 
 FOcclusionQueryBatcher::~FOcclusionQueryBatcher()
@@ -481,7 +482,7 @@ void FOcclusionQueryBatcher::Flush(FRHICommandList& RHICmdList)
 				/*NumVertices=*/ 8 * NumPrimitivesThisBatch,
 				/*StartIndex=*/ 0,
 				/*NumPrimitives=*/ 12 * NumPrimitivesThisBatch,
-				/*NumInstances=*/ 1
+				/*NumInstances=*/ NumInstances
 				);
 			RHICmdList.EndRenderQuery(BatchOcclusionQuery);
 		}
@@ -703,7 +704,7 @@ static bool AllocatePlanarReflectionOcclusionQuery(const FViewInfo& View, const
 	
 	bool bAllowBoundsTest = false;
 	
-	if (View.ViewFrustum.IntersectBox(SceneProxy->WorldBounds.GetCenter(), SceneProxy->WorldBounds.GetExtent()))
+	if (View.GetCullingFrustum().IntersectBox(SceneProxy->WorldBounds.GetCenter(), SceneProxy->WorldBounds.GetExtent()))
 	{
 		const FBoxSphereBounds OcclusionBounds(SceneProxy->WorldBounds);
 		
@@ -1086,7 +1087,13 @@ static void AllocateOcclusionTests(FViewOcclusionQueriesPerView& QueriesPerView,
 	// Perform occlusion queries for each view
 	for (int32 ViewIndex = 0; ViewIndex < Views.Num(); ViewIndex++)
 	{
+		// Don't draw occlusion queries for both eyes in multiview
 		FViewInfo& View = Views[ViewIndex];
+		if (View.bIsSinglePassStereo && IStereoRendering::IsASecondaryPass(View.StereoPass))
+		{
+			continue;
+		}
+
 		FViewOcclusionQueries& ViewQuery = QueriesPerView[ViewIndex];
 		FSceneViewState* ViewState = (FSceneViewState*)View.State;
 		const FSceneViewFamily& ViewFamily = *View.Family;
@@ -1244,9 +1251,15 @@ static void BeginOcclusionTests(
 
 	for (int32 ViewIndex = 0; ViewIndex < Views.Num(); ViewIndex++)
 	{
+		// Don't draw occlusion queries for both eyes in multiview
+		FViewInfo& View = Views[ViewIndex];
+		if (View.bIsSinglePassStereo && IStereoRendering::IsASecondaryPass(View.StereoPass))
+		{
+			continue;
+		}
+
 		SCOPED_DRAW_EVENTF(RHICmdList, ViewOcclusionTests, TEXT("ViewOcclusionTests %d"), ViewIndex);
 
-		FViewInfo& View = Views[ViewIndex];
 		const FViewOcclusionQueries& ViewQuery = QueriesPerView[ViewIndex];
 		FSceneViewState* ViewState = (FSceneViewState*)View.State;
 		SCOPED_GPU_MASK(RHICmdList, View.GPUMask);
@@ -1789,8 +1802,9 @@ public:
 	}
 
 	BEGIN_SHADER_PARAMETER_STRUCT(FParameters, )
-		SHADER_PARAMETER_STRUCT_REF(FViewUniformShaderParameters, View)
+		SHADER_PARAMETER_STRUCT_INCLUDE(FViewShaderParameters, View)
 		SHADER_PARAMETER(uint32, StartIndex)
+		SHADER_PARAMETER(uint32, InstancedViewIndex)
 	END_SHADER_PARAMETER_STRUCT()
 };
 
@@ -1884,23 +1898,31 @@ void FOcclusionFeedback::SubmitOcclusionDraws(FRHICommandList& RHICmdList, FView
 		FOcclusionBatch& Batch = BatchOcclusionQueries[BatchIndex];
 		
 		FOcclusionFeedbackVS::FParameters VSParams;
-		VSParams.View = View.ViewUniformBuffer;
+		VSParams.View = View.GetShaderParameters();
 		VSParams.StartIndex = PrimitiveStartIndex;
-		SetShaderParameters(RHICmdList, VertexShader, VertexShader.GetVertexShader(), VSParams);
 				
 		FRHIBuffer* PrimitiveBufferRHI = Batch.VertexAllocation.VertexBuffer->VertexBufferRHI;
 		uint32 VertexBufferOffset = Batch.VertexAllocation.VertexOffset;
 		RHICmdList.SetStreamSource(1, PrimitiveBufferRHI, VertexBufferOffset);
 
-		RHICmdList.DrawIndexedPrimitive(
-			CubeIndexBuffer,
-			/*BaseVertexIndex=*/ 0,
-			/*MinIndex=*/ 0,
-			/*NumVertices=*/ 8,
-			/*StartIndex=*/ 0,
-			/*NumPrimitives=*/ 12,
-			/*NumInstances=*/ Batch.NumBatchedPrimitives
-		);
+		// Because the RHI does not yet support instance divisors, Occlusion Feedback supports instanced stereo 
+		// by issuing a draw call for each eye. But both draw calls mark the same primitives, so it does not need to duplicate that state. 
+		const uint32 InstanceCount = View.bIsInstancedStereoEnabled ? 2 : 1;
+		for (uint32 InstancedViewIndex = 0; InstancedViewIndex < InstanceCount; ++InstancedViewIndex)
+		{
+			VSParams.InstancedViewIndex = InstancedViewIndex;
+			SetShaderParameters(RHICmdList, VertexShader, VertexShader.GetVertexShader(), VSParams);
+
+			RHICmdList.DrawIndexedPrimitive(
+				CubeIndexBuffer,
+				/*BaseVertexIndex=*/ 0,
+				/*MinIndex=*/ 0,
+				/*NumVertices=*/ 8,
+				/*StartIndex=*/ 0,
+				/*NumPrimitives=*/ 12,
+				/*NumInstances=*/ Batch.NumBatchedPrimitives
+			);
+		}
 
 		PrimitiveStartIndex += Batch.NumBatchedPrimitives;
 	}
diff --git a/Engine/Source/Runtime/Renderer/Private/SceneOcclusion.h b/Engine/Source/Runtime/Renderer/Private/SceneOcclusion.h
--- a/Engine/Source/Runtime/Renderer/Private/SceneOcclusion.h
+++ b/Engine/Source/Runtime/Renderer/Private/SceneOcclusion.h
@@ -59,41 +59,64 @@ public:
 	FOcclusionQueryVS(const ShaderMetaType::CompiledShaderInitializerType& Initializer):
 		FGlobalShader(Initializer)
 	{
-			StencilingGeometryParameters.Bind(Initializer.ParameterMap);
-			ViewId.Bind(Initializer.ParameterMap, TEXT("ViewId"));
+		StencilingGeometryParameters.Bind(Initializer.ParameterMap, TEXT("StencilingGeometryPosAndScale"));
 	}
 
 	FOcclusionQueryVS() {}
 
 	void SetParameters(FRHIBatchedShaderParameters& BatchedParameters, const FViewInfo& View, const FSphere& BoundingSphere)
 	{
-		FVector4f StencilingSpherePosAndScale;
-		StencilingGeometry::GStencilSphereVertexBuffer.CalcTransform(StencilingSpherePosAndScale, BoundingSphere, View.ViewMatrices.GetPreViewTranslation());
+		FVector4f StencilingSpherePosAndScale[2];
+		StencilingGeometry::GStencilSphereVertexBuffer.CalcTransform(StencilingSpherePosAndScale[0], BoundingSphere, View.ViewMatrices.GetPreViewTranslation());
+		if (View.bShouldBindInstancedViewUB)
+		{
+			if (const FViewInfo* InstancedView = View.GetInstancedView())
+			{
+				StencilingGeometry::GStencilSphereVertexBuffer.CalcTransform(StencilingSpherePosAndScale[1], BoundingSphere, InstancedView->ViewMatrices.GetPreViewTranslation());
+			}
+			else
+			{
+				StencilingSpherePosAndScale[1] = StencilingSpherePosAndScale[0];
+			}
+		}
+
 		SetParametersInternal(BatchedParameters, View, StencilingSpherePosAndScale);
 	}
 
 	void SetParameters(FRHIBatchedShaderParameters& BatchedParameters, const FViewInfo& View)
 	{
-		SetParametersInternal(BatchedParameters, View, FVector4f(0, 0, 0, 1));
+		// Don't transform if rendering frustum
+		FVector4f InstancedOffsetScale[2] = { FVector4f(0, 0, 0, 1), FVector4f(0, 0, 0, 1) };
+		if (View.bShouldBindInstancedViewUB)
+		{
+			// Occlusion queries are generated to be offset by the PreViewTranslation.
+			// However for multiview, we still need to apply that offset to the right eye.
+			// In that case we also need to undo the baked left eye translation.
+			if (const FViewInfo* InstancedView = View.GetInstancedView())
+			{
+				InstancedOffsetScale[1] = FVector4f(FVector3f(InstancedView->ViewMatrices.GetPreViewTranslation() - View.ViewMatrices.GetPreViewTranslation()), 1.0f);
+			}
+		}
+		SetParametersInternal(BatchedParameters, View, InstancedOffsetScale);
 	}
 
 private:
-	void SetParametersInternal(FRHIBatchedShaderParameters& BatchedParameters, const FViewInfo& View, const FVector4f& StencilingSpherePosAndScale)
+	void SetParametersInternal(FRHIBatchedShaderParameters& BatchedParameters, const FViewInfo& View, const FVector4f StencilingSpherePosAndScale[2])
 	{
 		FGlobalShader::SetParameters<FViewUniformShaderParameters>(BatchedParameters, View.ViewUniformBuffer);
-
-		// Don't transform if rendering frustum
-		StencilingGeometryParameters.Set(BatchedParameters, StencilingSpherePosAndScale);
-
-		if (ViewId.IsBound())
+		if (!View.bShouldBindInstancedViewUB)
+		{
+			SetShaderValue(BatchedParameters, StencilingGeometryParameters, StencilingSpherePosAndScale[0]);
+		}
+		else
 		{
-			SetShaderValue(BatchedParameters, ViewId, (View.StereoPass == EStereoscopicPass::eSSP_FULL) ? 0 : View.StereoViewIndex);
+			FGlobalShader::SetParameters<FInstancedViewUniformShaderParameters>(BatchedParameters, View.GetInstancedViewUniformBuffer());
+			SetShaderValueArray(BatchedParameters, StencilingGeometryParameters, StencilingSpherePosAndScale, 2);
 		}
 	}
 
 private:
-	LAYOUT_FIELD(FStencilingGeometryShaderParameters, StencilingGeometryParameters)
-	LAYOUT_FIELD(FShaderParameter, ViewId)
+	LAYOUT_FIELD(FShaderParameter, StencilingGeometryParameters)
 };
 
 /**
diff --git a/Engine/Source/Runtime/Renderer/Private/SceneRendering.cpp b/Engine/Source/Runtime/Renderer/Private/SceneRendering.cpp
--- a/Engine/Source/Runtime/Renderer/Private/SceneRendering.cpp
+++ b/Engine/Source/Runtime/Renderer/Private/SceneRendering.cpp
@@ -860,8 +860,8 @@ bool IsHMDHiddenAreaMaskActive()
  */
 FViewInfo::FViewInfo(const FSceneViewInitOptions& InitOptions)
 	:	FSceneView(InitOptions)
-	,	IndividualOcclusionQueries((FSceneViewState*)InitOptions.SceneViewStateInterface, 1)	
-	,	GroupedOcclusionQueries((FSceneViewState*)InitOptions.SceneViewStateInterface, FOcclusionQueryBatcher::OccludedPrimitiveQueryBatchSize)
+	,	IndividualOcclusionQueries((FSceneViewState*)InitOptions.SceneViewStateInterface, 1, bIsInstancedStereoEnabled ? 2 : 1)
+	,	GroupedOcclusionQueries((FSceneViewState*)InitOptions.SceneViewStateInterface, FOcclusionQueryBatcher::OccludedPrimitiveQueryBatchSize, bIsInstancedStereoEnabled ? 2 : 1)
 	,	CustomVisibilityQuery(nullptr)
 {
 	Init();
@@ -873,8 +873,8 @@ FViewInfo::FViewInfo(const FSceneViewInitOptions& InitOptions)
  */
 FViewInfo::FViewInfo(const FSceneView* InView)
 	:	FSceneView(*InView)
-	,	IndividualOcclusionQueries((FSceneViewState*)InView->State,1)	
-	,	GroupedOcclusionQueries((FSceneViewState*)InView->State,FOcclusionQueryBatcher::OccludedPrimitiveQueryBatchSize)
+	,	IndividualOcclusionQueries((FSceneViewState*)InView->State,1, bIsInstancedStereoEnabled ? 2 : 1)
+	,	GroupedOcclusionQueries((FSceneViewState*)InView->State,FOcclusionQueryBatcher::OccludedPrimitiveQueryBatchSize, bIsInstancedStereoEnabled ? 2 : 1)
 	,	CustomVisibilityQuery(nullptr)
 {
 	Init();
diff --git a/Engine/Source/Runtime/Renderer/Private/SceneRendering.h b/Engine/Source/Runtime/Renderer/Private/SceneRendering.h
--- a/Engine/Source/Runtime/Renderer/Private/SceneRendering.h
+++ b/Engine/Source/Runtime/Renderer/Private/SceneRendering.h
@@ -355,7 +355,7 @@ public:
 	enum { OccludedPrimitiveQueryBatchSize = 16 };
 
 	/** Initialization constructor. */
-	FOcclusionQueryBatcher(class FSceneViewState* ViewState, uint32 InMaxBatchedPrimitives);
+	FOcclusionQueryBatcher(class FSceneViewState* ViewState, uint32 InMaxBatchedPrimitives, uint32 InNumInstances);
 
 	/** Destructor. */
 	~FOcclusionQueryBatcher();
@@ -398,6 +398,9 @@ private:
 
 	/** The pool to allocate occlusion queries from. */
 	FFrameBasedOcclusionQueryPool* OcclusionQueryPool;
+
+	/** The number of instances for instanced stereo rendering */
+	uint32 NumInstances;
 };
 
 class FHZBOcclusionTester : public FRenderResource
@@ -1723,11 +1726,11 @@ public:
 		{
 			return false;
 		}
-		else if (!bIsInstancedStereoEnabled && !bIsMobileMultiViewEnabled)
+		else if (!bIsSinglePassStereo)
 		{
 			return true;
 		}
-		else if ((bIsInstancedStereoEnabled || bIsMobileMultiViewEnabled) && !IStereoRendering::IsASecondaryPass(StereoPass))
+		else if (bIsSinglePassStereo && !IStereoRendering::IsASecondaryPass(StereoPass))
 		{
 			return true;
 		}
@@ -2868,6 +2871,7 @@ private:
 	bool bEnableClusteredLocalLights;
 	bool bEnableClusteredReflections;
 	bool bRequiresShadowProjections;
+	bool bAdrenoOcclusionMode;
 
 	ETranslucencyPass::Type StandardTranslucencyPass;
 	EMeshPass::Type StandardTranslucencyMeshPass;
diff --git a/Engine/Source/Runtime/Renderer/Private/SceneVisibility.cpp b/Engine/Source/Runtime/Renderer/Private/SceneVisibility.cpp
--- a/Engine/Source/Runtime/Renderer/Private/SceneVisibility.cpp
+++ b/Engine/Source/Runtime/Renderer/Private/SceneVisibility.cpp
@@ -536,7 +536,7 @@ struct FFrustumCullingFlags
 };
 
 // Returns true if the frustum and bounds intersect
-inline bool IsPrimitiveVisible(FViewInfo& View, const FPlane* PermutedPlanePtr, const FPrimitiveBounds& Bounds, int32 VisibilityId, FFrustumCullingFlags Flags)
+inline bool IsPrimitiveVisible(FViewInfo& View, const FPlane* PermutedPlanePtr, const FConvexVolume& ViewCullingFrustum, const FPrimitiveBounds& Bounds, int32 VisibilityId, FFrustumCullingFlags Flags)
 {
 	// The custom culling and sphere culling are additional tests, meaning that if they pass, the
 	// remaining culling tests will still be performed.  If any of the tests fail, then the primitive
@@ -547,7 +547,7 @@ inline bool IsPrimitiveVisible(FViewInfo& View, const FPlane* PermutedPlanePtr,
 		return false;
 	}
 
-	if (Flags.bUseSphereTestFirst && !View.ViewFrustum.IntersectSphere(Bounds.BoxSphereBounds.Origin, Bounds.BoxSphereBounds.SphereRadius))
+	if (Flags.bUseSphereTestFirst && !ViewCullingFrustum.IntersectSphere(Bounds.BoxSphereBounds.Origin, Bounds.BoxSphereBounds.SphereRadius))
 	{
 		return false;
 	}
@@ -558,7 +558,7 @@ inline bool IsPrimitiveVisible(FViewInfo& View, const FPlane* PermutedPlanePtr,
 	}
 	else
 	{
-		return View.ViewFrustum.IntersectBox(Bounds.BoxSphereBounds.Origin, Bounds.BoxSphereBounds.BoxExtent);
+		return ViewCullingFrustum.IntersectBox(Bounds.BoxSphereBounds.Origin, Bounds.BoxSphereBounds.BoxExtent);
 	}
 }
 
@@ -617,7 +617,7 @@ inline void ComputeDistances(const FPrimitiveBounds& Bounds, const FVector& View
 	OutFurthestDistSq = FMath::Square(FurthestDist);
 }
 
-static void CullOctree(const FScene& Scene, FViewInfo& View, const FFrustumCullingFlags& Flags, FSceneBitArray& OutVisibleNodes)
+static void CullOctree(const FScene& Scene, FViewInfo& View, const FFrustumCullingFlags& Flags, FSceneBitArray& OutVisibleNodes, const FConvexVolume& ViewCullingFrustum)
 {
 	TRACE_CPUPROFILER_EVENT_SCOPE(SceneVisibility_CullOctree);
 
@@ -625,7 +625,7 @@ static void CullOctree(const FScene& Scene, FViewInfo& View, const FFrustumCulli
 	OutVisibleNodes.Init(false, Scene.PrimitiveOctree.GetNumNodes() * 2);
 
 	Scene.PrimitiveOctree.FindNodesWithPredicate(
-		[&View, &OutVisibleNodes, &Flags](FScenePrimitiveOctree::FNodeIndex ParentNodeIndex, FScenePrimitiveOctree::FNodeIndex NodeIndex, const FBoxCenterAndExtent& NodeBounds)
+		[&View, &OutVisibleNodes, &Flags, &ViewCullingFrustum](FScenePrimitiveOctree::FNodeIndex ParentNodeIndex, FScenePrimitiveOctree::FNodeIndex NodeIndex, const FBoxCenterAndExtent& NodeBounds)
 		{
 			// If the parent node is completely contained there is no need to test containment
 			if (ParentNodeIndex != INDEX_NONE && !OutVisibleNodes[(ParentNodeIndex * 2) + 1])
@@ -635,7 +635,7 @@ static void CullOctree(const FScene& Scene, FViewInfo& View, const FFrustumCulli
 				return true;
 			}
 
-			const FPlane* PermutedPlanePtr = View.ViewFrustum.PermutedPlanes.GetData();
+			const FPlane* PermutedPlanePtr = ViewCullingFrustum.PermutedPlanes.GetData();
 			bool bIntersects = false;
 
 			if (Flags.bUseFastIntersect)
@@ -644,13 +644,13 @@ static void CullOctree(const FScene& Scene, FViewInfo& View, const FFrustumCulli
 			}
 			else
 			{
-				bIntersects = View.ViewFrustum.IntersectBox(NodeBounds.Center, NodeBounds.Extent);
+				bIntersects = ViewCullingFrustum.IntersectBox(NodeBounds.Center, NodeBounds.Extent);
 			}
 
 			if (bIntersects)
 			{
 				OutVisibleNodes[NodeIndex * 2] = true;
-				OutVisibleNodes[NodeIndex * 2 + 1] = View.ViewFrustum.GetBoxIntersectionOutcode(NodeBounds.Center, NodeBounds.Extent).GetOutside();
+				OutVisibleNodes[NodeIndex * 2 + 1] = ViewCullingFrustum.GetBoxIntersectionOutcode(NodeBounds.Center, NodeBounds.Extent).GetOutside();
 			}
 
 			return bIntersects;
@@ -715,8 +715,9 @@ static int32 FrustumCull(const FScene& Scene, FViewInfo& View, FFrustumCullingFl
 	QUICK_SCOPE_CYCLE_COUNTER(STAT_FrustumCull_Loop);
 
 	bool bDisableLODFade = GDisableLODFade || View.bDisableDistanceBasedFadeTransitions;
-	const FPlane* PermutedPlanePtr = View.ViewFrustum.PermutedPlanes.GetData();
-	FVector ViewOriginForDistanceCulling = View.ViewMatrices.GetViewOrigin();
+	const FPlane* PermutedPlanePtr = View.GetCullingFrustum().PermutedPlanes.GetData();
+	const FConvexVolume& ViewCullingFrustum = View.GetCullingFrustum();
+	FVector ViewOriginForDistanceCulling = View.CullingOrigin;
 	float FadeRadius = bDisableLODFade ? 0.0f : GDistanceFadeMaxTravel;
 	uint8 CustomVisibilityFlags = EOcclusionFlags::CanBeOccluded | EOcclusionFlags::HasPrecomputedVisibility;
 
@@ -727,8 +728,6 @@ static int32 FrustumCull(const FScene& Scene, FViewInfo& View, FFrustumCullingFl
 	// Primitives may be explicitly removed from stereo views when using mono
 	const int32 TaskWordOffset = TaskIndex * TaskConfig.FrustumCull.NumWordsPerTask;
 
-	FVector ViewOrigin = View.ViewMatrices.GetViewOrigin();
-
 	uint32* RESTRICT VisWords = View.PrimitiveVisibilityMap.GetData();
 	uint32* RESTRICT FadeWords = View.PotentiallyFadingPrimitiveMap.GetData();
 #if RHI_RAYTRACING
@@ -817,7 +816,7 @@ static int32 FrustumCull(const FScene& Scene, FViewInfo& View, FFrustumCullingFl
 							VisibilityId = Scene.PrimitiveSceneProxies[Index]->GetVisibilityId();
 						}
 
-						bIsVisible = !bPartiallyOutside || IsPrimitiveVisible(View, PermutedPlanePtr, Bounds, VisibilityId, Flags);
+						bIsVisible = !bPartiallyOutside || IsPrimitiveVisible(View, PermutedPlanePtr, ViewCullingFrustum, Bounds, VisibilityId, Flags);
 					}
 				}
 
@@ -2056,26 +2055,11 @@ FComputeAndMarkRelevance::FComputeAndMarkRelevance(FVisibilityTaskData& InTaskDa
 	FMemory::Memzero(MarkMasks, NumMeshes + 31);
 	Packets.Reserve(InTaskData.TaskConfig.Relevance.NumEstimatedPackets);
 	CreateRelevancePacket();
-	InstancedPrimitiveAddedMap.Init(false, InScene.Primitives.Num());
 }
 
 void FComputeAndMarkRelevance::AddPrimitives(FPrimitiveIndexList&& PrimitiveIndexList)
 {
-	// In ISR only, the primary view will also receive all primitives visible in secondary views, since all rendering is handled from the primary
-	if (View.bIsMultiViewportEnabled && View.StereoPass == EStereoscopicPass::eSSP_PRIMARY)
-	{
-		for (int32 Index : PrimitiveIndexList)
-		{
-			// Because we've queued primitives from both the primary and secondary views, we need to filter duplicates
-			if (!InstancedPrimitiveAddedMap[Index])
-			{
-				AddPrimitive(Index);
-				InstancedPrimitiveAddedMap[Index] = true;
-			}
-		}
-	}
-
-	else if (PrimitiveIndexList.Num() == NumPrimitivesPerPacket)
+	if (PrimitiveIndexList.Num() == NumPrimitivesPerPacket)
 	{
 		// Create a one-off packet that will take all the primitives.
 		FRelevancePacket* Packet = CreateRelevancePacket();
@@ -2262,6 +2246,40 @@ void FComputeAndMarkRelevance::Finalize()
 		StaticMeshFadeOutDitheredLODMap_Words++;
 		StaticMeshFadeInDitheredLODMap_Words++;
 	}
+
+	if (View.bIsSinglePassStereo)
+	{
+		check(View.StereoPass != EStereoscopicPass::eSSP_SECONDARY); // Secondary views should not calculate relevance
+		if (const FViewInfo* SecondaryViewConst = View.GetInstancedView())
+		{
+			FViewInfo* SecondaryView = const_cast<FViewInfo*>(SecondaryViewConst);
+			SecondaryView->PrimitiveVisibilityMap = View.PrimitiveVisibilityMap;
+			SecondaryView->PrimitiveRayTracingVisibilityMap = View.PrimitiveRayTracingVisibilityMap;
+			SecondaryView->PrimitiveDefinitelyUnoccludedMap = View.PrimitiveDefinitelyUnoccludedMap;
+			SecondaryView->PrimitiveViewRelevanceMap = View.PrimitiveViewRelevanceMap;
+			SecondaryView->StaticMeshVisibilityMap = View.StaticMeshVisibilityMap;
+			SecondaryView->PrimitivesLODMask = View.PrimitivesLODMask;
+			SecondaryView->NumVisibleDynamicPrimitives = View.NumVisibleDynamicPrimitives;
+			SecondaryView->NumVisibleDynamicEditorPrimitives = View.NumVisibleDynamicEditorPrimitives;
+
+			SecondaryView->ShadingModelMaskInView = View.ShadingModelMaskInView;
+			SecondaryView->bUsesGlobalDistanceField = View.bUsesGlobalDistanceField;
+			SecondaryView->bUsesLightingChannels = View.bUsesLightingChannels;
+			SecondaryView->bTranslucentSurfaceLighting = View.bTranslucentSurfaceLighting;
+			SecondaryView->bSceneHasSkyMaterial = View.bSceneHasSkyMaterial;
+			SecondaryView->bHasSingleLayerWaterMaterial = View.bHasSingleLayerWaterMaterial;
+			SecondaryView->bUsesSecondStageDepthPass = View.bUsesSecondStageDepthPass;
+			SecondaryView->TranslucentPrimCount = View.TranslucentPrimCount;
+			SecondaryView->bHasDistortionPrimitives = View.bHasDistortionPrimitives;
+			SecondaryView->bHasCustomDepthPrimitives = View.bHasCustomDepthPrimitives;
+			SecondaryView->CustomDepthStencilValues = View.CustomDepthStencilValues;
+			SecondaryView->bUsesCustomDepth = View.bUsesCustomDepth;
+			SecondaryView->bUsesCustomStencil = View.bUsesCustomStencil;
+			SecondaryView->SubstrateViewData.MaxClosurePerPixel = View.SubstrateViewData.MaxClosurePerPixel;
+			SecondaryView->SubstrateViewData.MaxBytesPerPixel = View.SubstrateViewData.MaxClosurePerPixel;
+			SecondaryView->SubstrateViewData.bUsesComplexSpecialRenderPath = View.SubstrateViewData.bUsesComplexSpecialRenderPath;
+		}
+	}
 }
 
 ///////////////////////////////////////////////////////////////////////////////
@@ -2973,17 +2991,19 @@ FGPUOcclusion::FGPUOcclusion(FVisibilityViewPacket& InViewPacket)
 	, View(ViewPacket.View)
 	, ViewState(*ViewPacket.ViewState)
 {
-	State.bSubmitQueries = !View.bDisableQuerySubmissions;
+	// In single pass stereo we only submit queries from the primary view.
+	State.bSubmitQueries = !View.bDisableQuerySubmissions && !(View.bIsSinglePassStereo && View.StereoPass == EStereoscopicPass::eSSP_SECONDARY);
 
 #if !(UE_BUILD_SHIPPING || UE_BUILD_TEST)
 	State.bSubmitQueries &= !ViewState.bIsFrozen;
 #endif
 
 	// Perform round-robin occlusion queries
-	if (ViewState.IsRoundRobinEnabled() &&
+	const bool bUseRoundRobinOcclusion = ViewState.IsRoundRobinEnabled() &&
+		!View.bIsSinglePassStereo && // No round robin on single-pass stereo
 		!View.bIsSceneCapture && // We only round-robin on the main renderer (not scene captures)
-		!View.bIgnoreExistingQueries && // We do not alternate occlusion queries when we want to refresh the occlusion history
-		IStereoRendering::IsStereoEyeView(View)) // Only relevant to stereo views
+		IStereoRendering::IsStereoEyeView(View); // Only relevant to stereo views
+	if (bUseRoundRobinOcclusion && !View.bIgnoreExistingQueries) // We do not alternate occlusion queries when we want to refresh the occlusion history
 	{
 		// For even frames, prevent left eye from occlusion querying
 		// For odd frames, prevent right eye from occlusion querying
@@ -2998,7 +3018,7 @@ FGPUOcclusion::FGPUOcclusion(FVisibilityViewPacket& InViewPacket)
 	State.bHZBOcclusion &= GHZBOcclusion != 0;
 	State.bHZBOcclusion |= GHZBOcclusion == 2;
 
-	State.bUseRoundRobinOcclusion = IStereoRendering::IsStereoEyeView(View) && !View.bIsSceneCapture && ViewState.IsRoundRobinEnabled();
+	State.bUseRoundRobinOcclusion = bUseRoundRobinOcclusion;
 
 	State.NumBufferedFrames = FOcclusionQueryHelpers::GetNumBufferedFrames(Scene.GetFeatureLevel());
 	State.ReadBackLagTolerance = State.NumBufferedFrames;
@@ -3032,7 +3052,7 @@ void FGPUOcclusion::Map(FRHICommandListImmediate& RHICmdListImmediate)
 	ViewState.PrimitiveOcclusionQueryPool.AdvanceFrame(
 		ViewState.OcclusionFrameCounter,
 		FOcclusionQueryHelpers::GetNumBufferedFrames(Scene.GetFeatureLevel()),
-		ViewState.IsRoundRobinEnabled() && !View.bIsSceneCapture && IStereoRendering::IsStereoEyeView(View));
+		State.bUseRoundRobinOcclusion);
 
 	ViewState.Occlusion.NumRequestedQueries = 0;
 }
@@ -3497,23 +3517,6 @@ FVisibilityViewPacket::FVisibilityViewPacket(FVisibilityTaskData& InTaskData, FS
 			{
 				OcclusionCull.ContextIfParallel->AddPrimitives(PrimitiveRange);
 			}
-
-			else if (View.bIsMultiViewportEnabled)
-			{
-				// In ISR, we still need to use the relevance pipe even when occlusion is disabled to ensure all commands are forwarded to the primary view.
-				FPrimitiveIndexList PrimitiveIndexList;
-				for (FSceneSetBitIterator BitIt(View.PrimitiveVisibilityMap, PrimitiveRange.StartIndex); BitIt.GetIndex() < PrimitiveRange.EndIndex; ++BitIt)
-				{
-					PrimitiveIndexList.Emplace(BitIt.GetIndex());
-				}
-
-				if (!PrimitiveIndexList.IsEmpty())
-				{
-					Relevance.CommandPipe.AddNumCommands(1);
-					Relevance.CommandPipe.EnqueueCommand(PrimitiveIndexList);
-				}
-			}
-
 			else
 			{
 				// When occlusion is disabled primitives are queued directly to the relevance context rather than as a command on the relevance pipe.
@@ -3553,49 +3556,22 @@ FVisibilityViewPacket::FVisibilityViewPacket(FVisibilityTaskData& InTaskData, FS
 		OcclusionCull.CommandPipe.AddNumCommands(1);
 
 		// Callback for when a relevance command is queued from occlusion.
-		if (View.bIsMultiViewportEnabled && View.StereoPass == EStereoscopicPass::eSSP_SECONDARY)
-		{
-			Relevance.CommandPipe.SetCommandFunction([this](FPrimitiveIndexList&& PrimitiveIndexList)
-			{
-				// For instanced stereo secondary views, also send this command to the primary view
-				Relevance.PrimaryViewCommandPipe->AddNumCommands(1);
-				Relevance.PrimaryViewCommandPipe->EnqueueCommand(PrimitiveIndexList);
-
-				Relevance.Context->AddPrimitives(CopyTemp(PrimitiveIndexList));
-			});
-		}
-		else
+		Relevance.CommandPipe.SetCommandFunction([this](FPrimitiveIndexList&& PrimitiveIndexList)
 		{
-			Relevance.CommandPipe.SetCommandFunction([this](FPrimitiveIndexList&& PrimitiveIndexList)
-			{
-				Relevance.Context->AddPrimitives(MoveTemp(PrimitiveIndexList));
-			});
-		}
+			Relevance.Context->AddPrimitives(MoveTemp(PrimitiveIndexList));
+		});
 
 		// Callback for when the relevance pipe is done accepting commands.
-		if (View.bIsMultiViewportEnabled && View.StereoPass == EStereoscopicPass::eSSP_SECONDARY)
+		Relevance.CommandPipe.SetEmptyFunction([this]
 		{
-			Relevance.CommandPipe.SetEmptyFunction([this]
-			{
-				Relevance.Context->Finish(Tasks.ComputeRelevance);
+			Relevance.Context->Finish(Tasks.ComputeRelevance);
 
-				// Release reference on the instanced primary pipe
-				Relevance.PrimaryViewCommandPipe->ReleaseNumCommands(1);
-			});
-		}
-		else
-		{
-			Relevance.CommandPipe.SetEmptyFunction([this]
+			// Only used in the ViewPackets.Num() == 1 case
+			if (TaskData.DynamicMeshElements.CommandPipe)
 			{
-				Relevance.Context->Finish(Tasks.ComputeRelevance);
-
-				// Only used in the Views.Num() == 1 case
-				if (TaskData.DynamicMeshElements.CommandPipe)
-				{
-					TaskData.DynamicMeshElements.CommandPipe->ReleaseNumCommands(1);
-				}
-			});
-		}
+				TaskData.DynamicMeshElements.CommandPipe->ReleaseNumCommands(1);
+			}
+		});
 
 		// Take a reference on the relevance command pipe that is released by the occlusion pipe when all occlusion commands are complete.
 		Relevance.CommandPipe.AddNumCommands(1);
@@ -3696,7 +3672,7 @@ void FVisibilityViewPacket::BeginInitVisibility()
 	Flags.bShouldVisibilityCull  = bShouldVisibilityCull;
 	Flags.bUseCustomCulling      = View.CustomVisibilityQuery && View.CustomVisibilityQuery->Prepare();
 	Flags.bUseSphereTestFirst    = GFrustumCullUseSphereTestFirst;
-	Flags.bUseFastIntersect      = (View.ViewFrustum.PermutedPlanes.Num() == 8) && GFrustumCullUseFastIntersect;
+	Flags.bUseFastIntersect      = (View.GetCullingFrustum().PermutedPlanes.Num() == 8) && GFrustumCullUseFastIntersect;
 	Flags.bUseVisibilityOctree   = GFrustumCullUseOctree;
 	Flags.bHasHiddenPrimitives   = View.HiddenPrimitives.Num() > 0;
 	Flags.bHasShowOnlyPrimitives = View.ShowOnlyPrimitives.IsSet();
@@ -3715,7 +3691,8 @@ void FVisibilityViewPacket::BeginInitVisibility()
 	if (bShouldVisibilityCull && Flags.bUseVisibilityOctree)
 	{
 		VisibleNodes = TaskData.Allocator.Create<FSceneBitArray>();
-		CullOctree(Scene, View, Flags, *VisibleNodes);
+		const FConvexVolume& ViewCullingFrustum = View.GetCullingFrustum();
+		CullOctree(Scene, View, Flags, *VisibleNodes, ViewCullingFrustum);
 	}
 
 	const bool bCullingIsThreadsafe = (!Flags.bUseCustomCulling || View.CustomVisibilityQuery->IsThreadsafe());
@@ -3794,9 +3771,7 @@ void FVisibilityViewPacket::BeginInitVisibility()
 
 					const uint32 NumVisiblePrimitives = PrimitiveRange.EndIndex - PrimitiveRange.StartIndex - NumCulledPrimitives;
 
-					// Primary views can have additional visible primitives derived from secondary views. In that case forward
-					// the command down the pipe even though there weren't any visible primitives from frustum culling.
-					if (NumVisiblePrimitives == 0 && View.StereoPass != EStereoscopicPass::eSSP_PRIMARY)
+					if (NumVisiblePrimitives == 0)
 					{
 						OcclusionCull.CommandPipe.ReleaseNumCommands(1);
 					}
@@ -4257,7 +4232,12 @@ void FVisibilityTaskData::LaunchVisibilityTasks(const UE::Tasks::FTask& BeginIni
 
 	for (int32 ViewIndex = 0; ViewIndex < Views.Num(); ++ViewIndex)
 	{
-		// Each view gets its own visibility task packet which contains all the state to manage the task graph for a view.
+		if (Views[ViewIndex]->bIsSinglePassStereo && Views[ViewIndex]->StereoPass == EStereoscopicPass::eSSP_SECONDARY)
+		{
+			continue;
+		}
+
+		// Each non-vestigial view gets its own visibility task packet which contains all the state to manage the task graph for a view.
 		FVisibilityViewPacket& ViewPacket = ViewPackets.Emplace_GetRef(*this, Scene, *Views[ViewIndex], ViewIndex);
 
 		// Each respective stage of visibility for each view is connected to its corresponding task event in order to track when all views complete each stage.
@@ -4288,39 +4268,7 @@ void FVisibilityTaskData::LaunchVisibilityTasks(const UE::Tasks::FTask& BeginIni
 
 	if (TaskConfig.Schedule == EVisibilityTaskSchedule::Parallel)
 	{
-		if (Views.Num() > 1 && Views[0]->bIsMultiViewportEnabled)
-		{
-			// Instanced multi-view scenarios have secondary views which feed visibility data into primary views. In this
-			// case we make all views finish culling first before launching a task that merges secondary visibility
-			// maps into the primary view, which then becomes a prerequisite task for processing relevance pipe commands.
-			// All views will share a single relevance pipe and therefore a single relevance context.
-
-			TArray<UE::Tasks::FTask, SceneRenderingAllocator> PreSyncCullingTasks;
-
-			for (FVisibilityViewPacket& ViewPacket : ViewPackets)
-			{
-				// Ensure we merge after all culling is complete.
-				PreSyncCullingTasks.Emplace(ViewPacket.Tasks.FrustumCull);
-				PreSyncCullingTasks.Emplace(ViewPacket.Tasks.OcclusionCull);
-			}
-
-			UE::Tasks::FTask MergeSecondaryViewsTask = UE::Tasks::Launch(UE_SOURCE_LOCATION, [this]
-			{
-				MergeSecondaryViewVisibility();
-
-			}, PreSyncCullingTasks, TaskConfig.OcclusionCull.TaskPriority);
-
-			for (FVisibilityViewPacket& ViewPacket : ViewPackets)
-			{
-				// Force the shared primary view relevance pipe to wait until the merge task completes.
-				if (ViewPacket.View.StereoPass == EStereoscopicPass::eSSP_PRIMARY)
-				{
-					ViewPacket.Relevance.CommandPipe.SetPrerequisiteTask(MergeSecondaryViewsTask);
-				}
-			}
-		}
-
-		else if (Views.Num() == 1)
+		if (ViewPackets.Num() == 1)
 		{
 			// When using a single view, dynamic mesh elements are pushed into a pipe that is executed on the render thread which allows for some overlap with compute relevance work.
 			DynamicMeshElements.CommandPipe = Allocator.Create<TCommandPipe<FDynamicPrimitiveIndexList>>(TEXT("GatherDynamicMeshElements"));
@@ -4379,20 +4327,6 @@ void FVisibilityTaskData::LaunchVisibilityTasks(const UE::Tasks::FTask& BeginIni
 		check(!Views.IsEmpty())
 		checkf(!Views[0]->bIsMobileMultiViewEnabled, TEXT("This culling path was not tested with MMV"));
 
-		// In instanced stereo, we'll also redirect all primitives to the primary view's relevance command pipe, so secondary viewports need a reference
-		if (Views[0]->bIsMultiViewportEnabled)
-		{
-			int32 PrimaryViewIndex = Views[0]->PrimaryViewIndex;
-			for (FVisibilityViewPacket& ViewPacket : ViewPackets)
-			{
-				if (ViewPacket.View.StereoPass == EStereoscopicPass::eSSP_SECONDARY)
-				{
-					ViewPacket.Relevance.PrimaryViewCommandPipe = &ViewPackets[PrimaryViewIndex].Relevance.CommandPipe;
-					ViewPacket.Relevance.PrimaryViewCommandPipe->AddNumCommands(1); // Ensure the primary pipe doesn't close until the secondary pipe has forwarded all commands
-				}
-			}
-		}
-
 		for (FVisibilityViewPacket& ViewPacket : ViewPackets)
 		{
 			if (ViewPacket.OcclusionCull.ContextIfParallel)
@@ -4429,31 +4363,6 @@ void FVisibilityTaskData::LaunchVisibilityTasks(const UE::Tasks::FTask& BeginIni
 	Tasks.ComputeRelevance.Trigger();
 }
 
-void FVisibilityTaskData::MergeSecondaryViewVisibility()
-{
-	SCOPED_NAMED_EVENT(MergeSecondaryViewVisibility, FColor::Magenta);
-
-	// When using instanced multi-view, any primitive that is marked visible in ANY view must be marked visible in its primary view.
-	for (int32 ViewIndex = 0; ViewIndex < Views.Num(); ++ViewIndex)
-	{
-		FViewInfo& View = *Views[ViewIndex];
-
-		if (View.StereoPass == EStereoscopicPass::eSSP_SECONDARY)
-		{
-			const FSceneBitArray& SecondaryVisibilityMap = View.PrimitiveVisibilityMap;
-			FSceneBitArray& PrimaryVisibilityMap = Views[View.PrimaryViewIndex]->PrimitiveVisibilityMap;
-			check(PrimaryVisibilityMap.Num() == SecondaryVisibilityMap.Num());
-
-			const uint32 NumWords = FMath::DivideAndRoundUp(PrimaryVisibilityMap.Num(), NumBitsPerDWORD);
-
-			for (uint32 Index = 0; Index < NumWords; ++Index)
-			{
-				PrimaryVisibilityMap.GetData()[Index] |= SecondaryVisibilityMap.GetData()[Index];
-			}
-		}
-	}
-}
-
 void FVisibilityTaskData::GatherDynamicMeshElements(FDynamicPrimitiveIndexList&& DynamicPrimitiveIndexList)
 {
 	FDynamicPrimitiveIndexList RenderThreadDynamicPrimitiveIndexList;
@@ -4693,8 +4602,6 @@ void FVisibilityTaskData::ProcessRenderThreadTasks()
 			ViewPacket.Tasks.OcclusionCull.Trigger();
 		}
 
-		MergeSecondaryViewVisibility();
-
 		// Relevance requires that cached mesh commands be available first.
 		Scene.WaitForCacheMeshDrawCommandsTask();
 
@@ -5478,11 +5385,11 @@ void FSceneRenderer::ComputeLightVisibility()
 				Proxy->GetLightType() == LightType_Rect )
 			{
 				const FSphere& BoundingSphere = Proxy->GetBoundingSphere();
-				const bool bInViewFrustum = View.ViewFrustum.IntersectSphere(BoundingSphere.Center, BoundingSphere.W);
+				const bool bInViewFrustum = View.GetCullingFrustum().IntersectSphere(BoundingSphere.Center, BoundingSphere.W);
 
 				if (View.IsPerspectiveProjection())
 				{
-					const float DistanceSquared = (BoundingSphere.Center - View.ViewMatrices.GetViewOrigin()).SizeSquared();
+					const float DistanceSquared = (BoundingSphere.Center - View.CullingOrigin).SizeSquared();
 					const float ProxyMaxDistance = Proxy->GetMaxDrawDistance();
 					const float ScaledMaxDistance = ProxyMaxDistance * GLightMaxDrawDistanceScale;
 					const bool bDrawLight = (FMath::Square(FMath::Min(0.0002f, GMinScreenRadiusForLights / BoundingSphere.W) * View.LODDistanceFactor) * DistanceSquared < 1.0f)
diff --git a/Engine/Source/Runtime/Renderer/Private/SceneVisibilityPrivate.h b/Engine/Source/Runtime/Renderer/Private/SceneVisibilityPrivate.h
--- a/Engine/Source/Runtime/Renderer/Private/SceneVisibilityPrivate.h
+++ b/Engine/Source/Runtime/Renderer/Private/SceneVisibilityPrivate.h
@@ -451,7 +451,6 @@ private:
 	{
 		FComputeAndMarkRelevance* Context = nullptr;
 		TCommandPipe<FPrimitiveIndexList> CommandPipe{ TEXT("RelevancePipe") };
-		TCommandPipe<FPrimitiveIndexList>* PrimaryViewCommandPipe; // When using instanced stereo, secondary views will also send their commands to the primary view's command pipe to merge data
 
 	} Relevance;
 
@@ -556,8 +555,6 @@ public:
 	}
 
 private:
-	void MergeSecondaryViewVisibility();
-
 	void GatherDynamicMeshElements(FDynamicPrimitiveIndexList&& Primitives);
 	void GatherDynamicMeshElements(const FDynamicPrimitiveViewMasks& Primitives);
 
@@ -825,7 +822,6 @@ private:
 	FScene& Scene;
 	FViewInfo& View;
 	FViewCommands& ViewCommands;
-	FSceneBitArray InstancedPrimitiveAddedMap;
 	int32 ViewIndex;
 	const FFilterStaticMeshesForViewData ViewData;
 	const uint32 NumMeshes;
