[GPU Scene Lights] Support ForwardLightBuffer mode with stable light indices.
- this doesn't affect order of light indices in light grid cells.

diff --git a/Engine/Source/Runtime/Renderer/Private/LightGridInjection.cpp b/Engine/Source/Runtime/Renderer/Private/LightGridInjection.cpp
--- a/Engine/Source/Runtime/Renderer/Private/LightGridInjection.cpp
+++ b/Engine/Source/Runtime/Renderer/Private/LightGridInjection.cpp
@@ -42,8 +42,9 @@ LightGridInjection.cpp
 
 enum class ELightBufferMode
 {
-	VisibleLocalLights	= 0,
-	VisibleLights		= 1
+	VisibleLocalLights			= 0,
+	VisibleLights				= 1,
+	VisibleLightsStableIndices	= 2
 };
 
 static int32 GLightBufferMode = (int32)ELightBufferMode::VisibleLocalLights;
@@ -51,7 +52,8 @@ static FAutoConsoleVariableRef CVarLightBufferMode(
 	TEXT("r.Forward.LightBuffer.Mode"),
 	GLightBufferMode,
 	TEXT("0 - Visible local lights.\n")
-	TEXT("1 - Visible local + directional lights."),
+	TEXT("1 - Visible local + directional lights.\n")
+	TEXT("2 - Visible local + directional lights (with stable indices)."),
 	ECVF_RenderThreadSafe
 );
 
@@ -801,6 +803,14 @@ FComputeLightGridOutput FSceneRenderer::ComputeLightGrid(FRDGBuilder& GraphBuild
 
 		if (bCullLightsToGrid)
 		{
+			if (GLightBufferMode == (int32)ELightBufferMode::VisibleLightsStableIndices)
+			{
+				// when using stable light indices, indexing in ForwardLightBuffer is done using LightSceneInfo->Id
+				// so need to allocate max light ID entries
+				const int32 MaxLightId = Scene->GPUScene.GetMaxLightId();
+				ForwardLightData.AddUninitialized(FMath::Max(1, MaxLightId));
+			}
+
 			// Simple lights are copied without view dependent checks, so same in and out
 			SimpleLightsEnd = SortedLightSet.SimpleLightsEnd;
 
@@ -858,7 +868,7 @@ FComputeLightGridOutput FSceneRenderer::ComputeLightGrid(FRDGBuilder& GraphBuild
 						|| (SortedLightInfo.SortKey.Fields.LightType == LightType_Rect && ViewFamily.EngineShowFlags.RectLights)
 						|| (GLightBufferMode != (int32)ELightBufferMode::VisibleLocalLights && SortedLightInfo.SortKey.Fields.LightType == LightType_Directional && ViewFamily.EngineShowFlags.DirectionalLights))
 					{
-						const int32 IndexInBuffer = ForwardLightData.AddUninitialized(1);
+						const int32 IndexInBuffer = GLightBufferMode == (int32)ELightBufferMode::VisibleLightsStableIndices ? LightSceneInfo->Id : ForwardLightData.AddUninitialized(1);
 						FForwardLightData& LightData = ForwardLightData[IndexInBuffer];
 
 						int32 PrevForwardLightIndex = INDEX_NONE;
