Fixed "Affects Indirect While Hidden" or "Hidden Shadow" not working with in scene capture renders with primitives tagged as "Hidden In Scene Capture".

diff --git a/Engine/Source/Runtime/Engine/Private/PrimitiveSceneProxy.cpp b/Engine/Source/Runtime/Engine/Private/PrimitiveSceneProxy.cpp
--- a/Engine/Source/Runtime/Engine/Private/PrimitiveSceneProxy.cpp
+++ b/Engine/Source/Runtime/Engine/Private/PrimitiveSceneProxy.cpp
@@ -1716,11 +1716,11 @@ bool FPrimitiveSceneProxy::IsShadowCast(const FSceneView* View) const
 		{
 			return false;
 		}
-	}
 
-	if (View->bIsSceneCapture && IsHiddenInSceneCapture())
-	{
-		return false;
+		if (View->bIsSceneCapture && IsHiddenInSceneCapture())
+		{
+			return false;
+		}
 	}
 
 	if (!View->bIsSceneCapture && IsVisibleInSceneCaptureOnly())
diff --git a/Engine/Source/Runtime/Renderer/Private/RayTracing/RayTracing.cpp b/Engine/Source/Runtime/Renderer/Private/RayTracing/RayTracing.cpp
--- a/Engine/Source/Runtime/Renderer/Private/RayTracing/RayTracing.cpp
+++ b/Engine/Source/Runtime/Renderer/Private/RayTracing/RayTracing.cpp
@@ -998,8 +998,11 @@ namespace RayTracing
 
 				const FPrimitiveSceneInfo* SceneInfo = Scene.Primitives[PrimitiveIndex];
 
+				checkf(SceneInfo->Proxy != nullptr, TEXT("SceneInfo does not have a valid Proxy object. If this occurs, this object should probably have been filtered out before being added to Scene.Primitives"));
+				const bool bRetainWhileHidden = SceneInfo->Proxy->CastsHiddenShadow() || SceneInfo->Proxy->AffectsIndirectLightingWhileHidden();
+
 				// #dxr_todo: ray tracing in scene captures should re-use the persistent RT scene. (UE-112448)
-				if (View.bIsSceneCapture && !SceneInfo->bIsVisibleInSceneCaptures)
+				if (View.bIsSceneCapture && !SceneInfo->bIsVisibleInSceneCaptures && !bRetainWhileHidden)
 				{
 					return;
 				}
@@ -1010,15 +1013,11 @@ namespace RayTracing
 				}
 
 				// Some primitives should only be visible editor mode, however far field geometry 
-				// and hidden shadow casters must still always be added to the RT scene.
-				if (bGameView && !SceneInfo->bDrawInGame && !SceneInfo->bRayTracingFarField)
+				// and geometry that retains visibility while hidden (affect indirect while hidden
+				// or hidden shadow casters) must still always be added to the RT scene.
+				if (bGameView && !SceneInfo->bDrawInGame && !SceneInfo->bRayTracingFarField && !bRetainWhileHidden)
 				{
-					// Make sure this isn't an object that wants to be hidden to camera but still wants to cast shadows or be visible to indirect
-					checkf(SceneInfo->Proxy != nullptr, TEXT("SceneInfo does not have a valid Proxy object. If this occurs, this object should probably have been filtered out before being added to Scene.Primitives"));
-					if (!SceneInfo->Proxy->CastsHiddenShadow() && !SceneInfo->Proxy->AffectsIndirectLightingWhileHidden())
-					{
-						return;
-					}
+					return;
 				}
 
 				// Marked visible and used after point, check if streaming then mark as used in the TLAS (so it can be streamed in)
