Fix AddVisualizeBufferPass(...) not taking OutputViewport.Rect.Min into account.

diff --git a/Engine/Source/Runtime/Renderer/Private/PostProcess/PostProcessVisualizeBuffer.cpp b/Engine/Source/Runtime/Renderer/Private/PostProcess/PostProcessVisualizeBuffer.cpp
--- a/Engine/Source/Runtime/Renderer/Private/PostProcess/PostProcessVisualizeBuffer.cpp
+++ b/Engine/Source/Runtime/Renderer/Private/PostProcess/PostProcessVisualizeBuffer.cpp
@@ -83,7 +83,7 @@ FScreenPassTexture AddVisualizeBufferPass(FRDGBuilder& GraphBuilder, const FView
 		const int32 TileY = TileIndex / MaxTilesX;
 
 		FScreenPassTextureViewport OutputViewport(Output);
-		OutputViewport.Rect.Min = FIntPoint(TileX * TileWidth, TileY * TileHeight);
+		OutputViewport.Rect.Min = OutputViewport.Rect.Min + FIntPoint(TileX * TileWidth, TileY * TileHeight);
 		OutputViewport.Rect.Max = OutputViewport.Rect.Min + FIntPoint(TileWidth, TileHeight);
 
 		const FLinearColor SelectionColor = Tile.bSelected ? FLinearColor::Yellow : FLinearColor::Transparent;
