Defer removal of virtual texture callbacks until safe points in the frame.
Without this we can hit an assert when RemoveAllProducerDestroyedCallbacks() is called from  ~FMaterialRenderProxy() during deferred async scene renderer deletion.

diff --git a/Engine/Source/Runtime/Renderer/Private/VT/VirtualTextureSystem.cpp b/Engine/Source/Runtime/Renderer/Private/VT/VirtualTextureSystem.cpp
--- a/Engine/Source/Runtime/Renderer/Private/VT/VirtualTextureSystem.cpp
+++ b/Engine/Source/Runtime/Renderer/Private/VT/VirtualTextureSystem.cpp
@@ -373,6 +373,7 @@ FVirtualTextureSystem::~FVirtualTextureSystem()
 	UDebugDrawService::Unregister(DrawResidencyHudDelegateHandle);
 #endif
 
+	RemovePendingProducerDestroyedCallbacks();
 	DestroyPendingVirtualTextures(true);
 
 	check(AllocatedVTs.Num() == 0);
@@ -835,9 +836,22 @@ void FVirtualTextureSystem::AddProducerDestroyedCallback(const FVirtualTexturePr
 
 uint32 FVirtualTextureSystem::RemoveAllProducerDestroyedCallbacks(const void* Baton)
 {
-	check(!bUpdating);
 	UE::TScopeLock Lock(Mutex);
-	return Producers.RemoveAllCallbacks(Baton);
+	PendingRemoveProducerDestroyedCallbacks.Add(Baton);
+	return 0;
+}
+
+void FVirtualTextureSystem::RemovePendingProducerDestroyedCallbacks()
+{
+	UE::TScopeLock Lock(Mutex);
+	if (PendingRemoveProducerDestroyedCallbacks.Num())
+	{
+		for (const void* Baton : PendingRemoveProducerDestroyedCallbacks)
+		{
+			Producers.RemoveAllCallbacks(Baton);
+		}
+		PendingRemoveProducerDestroyedCallbacks.Reset();
+	}
 }
 
 FVirtualTextureProducer* FVirtualTextureSystem::FindProducer(const FVirtualTextureProducerHandle& Handle)
@@ -2784,6 +2798,7 @@ TUniquePtr<FVirtualTextureUpdater> FVirtualTextureSystem::BeginUpdate(FRDGBuilde
 		bFlushCaches = false;
 	}
 
+	RemovePendingProducerDestroyedCallbacks();
 	DestroyPendingVirtualTextures(false);
 
 	// Early out when no allocated VTs
@@ -2909,6 +2924,7 @@ void FVirtualTextureSystem::ReleasePendingResources()
 {
 	check(!bUpdating);
 	UE::TScopeLock Lock(Mutex);
+	RemovePendingProducerDestroyedCallbacks();
 	DestroyPendingVirtualTextures(true);
 	ReleasePendingSpaces();
 }
diff --git a/Engine/Source/Runtime/Renderer/Private/VT/VirtualTextureSystem.h b/Engine/Source/Runtime/Renderer/Private/VT/VirtualTextureSystem.h
--- a/Engine/Source/Runtime/Renderer/Private/VT/VirtualTextureSystem.h
+++ b/Engine/Source/Runtime/Renderer/Private/VT/VirtualTextureSystem.h
@@ -201,6 +201,7 @@ private:
 	void AllocateResources(FRDGBuilder& GraphBuilder);
 	void DestroyPendingVirtualTextures(bool bForceDestroyAll);
 	void ReleasePendingSpaces();
+	void RemovePendingProducerDestroyedCallbacks();
 
 	void RequestTilesForRegionInternal(const IAllocatedVirtualTexture* AllocatedVT, const FVector2D& InScreenSpaceSize, const FVector2D& InViewportPosition, const FVector2D& InViewportSize, const FVector2D& InUV0, const FVector2D& InUV1, int32 InMipLevel);
 	void RequestTilesInternal(const IAllocatedVirtualTexture* AllocatedVT, int32 InMipLevel);
@@ -251,6 +252,7 @@ private:
 	FVirtualTextureProducerCollection Producers;
 
 	TArray<IAllocatedVirtualTexture*> PendingDeleteAllocatedVTs;
+	TArray<const void*> PendingRemoveProducerDestroyedCallbacks;
 
 	TMap<FAllocatedVTDescription, FAllocatedVirtualTexture*> AllocatedVTs;
 	TArray<IAllocatedVirtualTexture*> AllocatedVTsToMap;
