Backout changelist 38841707
Deferred removal of virtual texture callback can cause callback being called on released object which can cause crash or deadlock.

diff --git a/Engine/Source/Runtime/Renderer/Private/VT/VirtualTextureSystem.cpp b/Engine/Source/Runtime/Renderer/Private/VT/VirtualTextureSystem.cpp
--- a/Engine/Source/Runtime/Renderer/Private/VT/VirtualTextureSystem.cpp
+++ b/Engine/Source/Runtime/Renderer/Private/VT/VirtualTextureSystem.cpp
@@ -373,7 +373,6 @@ FVirtualTextureSystem::~FVirtualTextureSystem()
 	UDebugDrawService::Unregister(DrawResidencyHudDelegateHandle);
 #endif
 
-	RemovePendingProducerDestroyedCallbacks();
 	DestroyPendingVirtualTextures(true);
 
 	check(AllocatedVTs.Num() == 0);
@@ -836,22 +835,9 @@ void FVirtualTextureSystem::AddProducerDestroyedCallback(const FVirtualTexturePr
 
 uint32 FVirtualTextureSystem::RemoveAllProducerDestroyedCallbacks(const void* Baton)
 {
+	check(!bUpdating);
 	UE::TScopeLock Lock(Mutex);
-	PendingRemoveProducerDestroyedCallbacks.Add(Baton);
-	return 0;
-}
-
-void FVirtualTextureSystem::RemovePendingProducerDestroyedCallbacks()
-{
-	UE::TScopeLock Lock(Mutex);
-	if (PendingRemoveProducerDestroyedCallbacks.Num())
-	{
-		for (const void* Baton : PendingRemoveProducerDestroyedCallbacks)
-		{
-			Producers.RemoveAllCallbacks(Baton);
-		}
-		PendingRemoveProducerDestroyedCallbacks.Reset();
-	}
+	return Producers.RemoveAllCallbacks(Baton);
 }
 
 FVirtualTextureProducer* FVirtualTextureSystem::FindProducer(const FVirtualTextureProducerHandle& Handle)
@@ -2798,7 +2784,6 @@ TUniquePtr<FVirtualTextureUpdater> FVirtualTextureSystem::BeginUpdate(FRDGBuilde
 		bFlushCaches = false;
 	}
 
-	RemovePendingProducerDestroyedCallbacks();
 	DestroyPendingVirtualTextures(false);
 
 	// Early out when no allocated VTs
@@ -2924,7 +2909,6 @@ void FVirtualTextureSystem::ReleasePendingResources()
 {
 	check(!bUpdating);
 	UE::TScopeLock Lock(Mutex);
-	RemovePendingProducerDestroyedCallbacks();
 	DestroyPendingVirtualTextures(true);
 	ReleasePendingSpaces();
 }
diff --git a/Engine/Source/Runtime/Renderer/Private/VT/VirtualTextureSystem.h b/Engine/Source/Runtime/Renderer/Private/VT/VirtualTextureSystem.h
--- a/Engine/Source/Runtime/Renderer/Private/VT/VirtualTextureSystem.h
+++ b/Engine/Source/Runtime/Renderer/Private/VT/VirtualTextureSystem.h
@@ -201,7 +201,6 @@ private:
 	void AllocateResources(FRDGBuilder& GraphBuilder);
 	void DestroyPendingVirtualTextures(bool bForceDestroyAll);
 	void ReleasePendingSpaces();
-	void RemovePendingProducerDestroyedCallbacks();
 
 	void RequestTilesForRegionInternal(const IAllocatedVirtualTexture* AllocatedVT, const FVector2D& InScreenSpaceSize, const FVector2D& InViewportPosition, const FVector2D& InViewportSize, const FVector2D& InUV0, const FVector2D& InUV1, int32 InMipLevel);
 	void RequestTilesInternal(const IAllocatedVirtualTexture* AllocatedVT, int32 InMipLevel);
@@ -252,7 +251,6 @@ private:
 	FVirtualTextureProducerCollection Producers;
 
 	TArray<IAllocatedVirtualTexture*> PendingDeleteAllocatedVTs;
-	TArray<const void*> PendingRemoveProducerDestroyedCallbacks;
 
 	TMap<FAllocatedVTDescription, FAllocatedVirtualTexture*> AllocatedVTs;
 	TArray<IAllocatedVirtualTexture*> AllocatedVTsToMap;
