Wait on the scene renderer builder before kicking off async VT update task.
Fixes race condition where deferred destruction of scene renderer has race condition that can hit assert in RemoveAllProducerDestroyedCallbacks().
This is a quick fix to unblock people.
Use r.VT.WaitBeforeUpdate to revert to old behavior.

diff --git a/Engine/Source/Runtime/Renderer/Private/VT/VirtualTextureSystem.cpp b/Engine/Source/Runtime/Renderer/Private/VT/VirtualTextureSystem.cpp
--- a/Engine/Source/Runtime/Renderer/Private/VT/VirtualTextureSystem.cpp
+++ b/Engine/Source/Runtime/Renderer/Private/VT/VirtualTextureSystem.cpp
@@ -14,6 +14,7 @@
 #include "ProfilingDebugging/CpuProfilerTrace.h"
 #include "RendererOnScreenNotification.h"
 #include "ScenePrivate.h"
+#include "SceneRenderBuilder.h"
 #include "SceneUtils.h"
 #include "Stats/Stats.h"
 #include "VirtualTexturing.h"
@@ -164,6 +165,12 @@ static TAutoConsoleVariable<int32> CVarVTAsyncPageRequestTask(
 	TEXT("Performs VT page requests on an async task."),
 	ECVF_RenderThreadSafe
 );
+static TAutoConsoleVariable<int32> CVarVTWaitBeforeUpdate(
+	TEXT("r.VT.WaitBeforeUpdate"),
+	1,
+	TEXT("Wait on the SceneRenderBuilder before kicking VT async update."),
+	ECVF_RenderThreadSafe
+);
 
 FVirtualTextureUpdateSettings::FVirtualTextureUpdateSettings()
 {
@@ -2681,6 +2688,15 @@ void FVirtualTextureSystem::GatherPackedTileRequests(FConcurrentLinearBulkObject
 
 void FVirtualTextureSystem::BeginUpdate(FRDGBuilder& GraphBuilder, FVirtualTextureUpdater* Updater)
 {
+	if (CVarVTWaitBeforeUpdate.GetValueOnRenderThread())
+	{
+		// Wait for SceneRenderBuilder deferred deletion.
+		// Otherwise scene renderer deletion may happen during virtual texture update.
+		// This can trigger material proxy deletion and calls to RemoveAllProducerDestroyedCallbacks().
+		// RemoveAllProducerDestroyedCallbacks() can't be called during a virtual texture update.
+		FSceneRenderBuilder::WaitForAsyncDeleteTask();
+	}
+
 	// Mark updating to true now that we are potentially launching async tasks.
 	bUpdating = true;
 
