Make raytracing ShaderBindingLayout functions public so they can be used from plugins as well

diff --git a/Engine/Source/Runtime/Renderer/Private/RayTracing/RayTracing.h b/Engine/Source/Runtime/Renderer/Private/RayTracing/RayTracing.h
--- a/Engine/Source/Runtime/Renderer/Private/RayTracing/RayTracing.h
+++ b/Engine/Source/Runtime/Renderer/Private/RayTracing/RayTracing.h
@@ -5,6 +5,7 @@
 #include "RendererInterface.h"
 #include "RenderGraphDefinitions.h"
 #include "RayTracingDefinitions.h"
+#include "RayTracingShaderBindingLayout.h"
 #include "RHIDefinitions.h"
 #include "ShaderCore.h"
 
@@ -41,12 +42,6 @@ namespace RayTracing
 
 	void OnRenderBegin(FScene& Scene, TArray<FViewInfo>& Views, const FViewFamilyInfo& ViewFamily);
 
-	// Get shader resource table desc used for all raytracing shaders which is shared between all shaders in the RTPSO
-	const FShaderBindingLayout* GetShaderBindingLayout(EShaderPlatform ShaderPlatform);
-
-	// Setup the runtime static uniform buffer bindings on the command list if enabled
-	TOptional<FScopedUniformBufferStaticBindings> BindStaticUniformBufferBindings(const FViewInfo& View, FRHIUniformBuffer* SceneUniformBuffer, FRHICommandList& RHICmdList);
-
 	FGatherInstancesTaskData* CreateGatherInstancesTaskData(
 		FSceneRenderingBulkObjectAllocator& InAllocator,
 		FScene& Scene,
diff --git a/Engine/Source/Runtime/Renderer/Public/RayTracingShaderBindingLayout.h b/Engine/Source/Runtime/Renderer/Public/RayTracingShaderBindingLayout.h
new file mode 100644
--- /dev/null
+++ b/Engine/Source/Runtime/Renderer/Public/RayTracingShaderBindingLayout.h
@@ -0,0 +1,24 @@
+// Copyright Epic Games, Inc. All Rights Reserved.
+
+#pragma once
+
+#include "RHIDefinitions.h"
+
+#if RHI_RAYTRACING
+
+#include "RHICommandList.h"
+#include "RHIResources.h"
+
+class FShaderBindingLayout;
+class FViewInfo;
+
+namespace RayTracing
+{
+	// Get shader resource table desc used for all raytracing shaders which is shared between all shaders in the RTPSO
+	RENDERER_API const FShaderBindingLayout* GetShaderBindingLayout(EShaderPlatform ShaderPlatform);
+
+	// Setup the runtime static uniform buffer bindings on the command list if enabled
+	RENDERER_API TOptional<FScopedUniformBufferStaticBindings> BindStaticUniformBufferBindings(const FViewInfo& View, FRHIUniformBuffer* SceneUniformBuffer, FRHICommandList& RHICmdList);
+}
+
+#endif // RHI_RAYTRACING
\ No newline at end of file
