- Add RHI global flag for RayTracing persistent SBTs
- Enable high level cvar by default and rely on per RHI side enabling

diff --git a/Engine/Source/Runtime/RHI/Public/RHIGlobals.h b/Engine/Source/Runtime/RHI/Public/RHIGlobals.h
--- a/Engine/Source/Runtime/RHI/Public/RHIGlobals.h
+++ b/Engine/Source/Runtime/RHI/Public/RHIGlobals.h
@@ -449,6 +449,9 @@ struct FRHIGlobals
 		/** Whether or not the RHI supports inlined callbacks */
 		bool SupportsInlinedCallbacks = false;
 
+		/** Whether or not the RHI supports persistent SBTs */
+		bool SupportsPersistentSBTs = false;
+
 		/** Wether an extra uniform buffer parameter is required when loose parameters are present, or if they are stored directly in the shader record. */
 		bool SupportsLooseParamsInShaderRecord = false;
 
diff --git a/Engine/Source/Runtime/Renderer/Private/RayTracing/RayTracingShaderBindingTable.cpp b/Engine/Source/Runtime/Renderer/Private/RayTracing/RayTracingShaderBindingTable.cpp
--- a/Engine/Source/Runtime/Renderer/Private/RayTracing/RayTracingShaderBindingTable.cpp
+++ b/Engine/Source/Runtime/Renderer/Private/RayTracing/RayTracingShaderBindingTable.cpp
@@ -9,7 +9,7 @@
 
 #if RHI_RAYTRACING
 
-static int32 GPersistentSBTEnabled = 0;
+static int32 GPersistentSBTEnabled = 1;
 static FAutoConsoleVariableRef CVarRayTracingPersistentSBT(
 	TEXT("r.RayTracing.PersistentSBT"),
 	GPersistentSBTEnabled,
@@ -69,6 +69,11 @@ static FAutoConsoleVariableRef CVarRayTracingPersistentSBTValidateBindings(
 	ECVF_RenderThreadSafe
 );
 
+static bool UseRayTracingPersistentSBTs()
+{
+	return GRHIGlobals.RayTracing.SupportsPersistentSBTs && GPersistentSBTEnabled;
+}
+
 uint32 FRayTracingSBTAllocation::GetRecordIndex(ERayTracingSceneLayer Layer, uint32 SegmentIndex) const
 {
 	check(HasLayer(Layer));
@@ -91,7 +96,7 @@ int32 FRayTracingSBTAllocation::GetSegmentCount() const
 FRayTracingShaderBindingTable::FRayTracingShaderBindingTable() : NumShaderSlotsPerGeometrySegment(RAY_TRACING_NUM_SHADER_SLOTS), StaticRangeAllocator(true /*bInGrowOnly*/)
 {
 	PersistentSBTInitializer.ShaderBindingMode = ERayTracingShaderBindingMode::RTPSO;
-	PersistentSBTInitializer.Lifetime = GPersistentSBTEnabled ? ERayTracingShaderBindingTableLifetime::Persistent : ERayTracingShaderBindingTableLifetime::Transient;
+	PersistentSBTInitializer.Lifetime = UseRayTracingPersistentSBTs() ? ERayTracingShaderBindingTableLifetime::Persistent : ERayTracingShaderBindingTableLifetime::Transient;
 	PersistentSBTInitializer.LocalBindingDataSize = GMinLocalBindingDataSize;
 	PersistentSBTInitializer.NumMissShaderSlots = GMinMissShaderSlots;
 }
@@ -141,7 +146,7 @@ FRayTracingShaderBindingDataOneFrameArray FRayTracingShaderBindingTable::GetDirt
 	FRayTracingShaderBindingDataOneFrameArray DirtyBindings;
 
 	// Persistent SBTs disabled (none of the cached SBTs have any persistent bindings so safe to just copy visible bindings and update type inside binding will be ignored)
-	if (!GPersistentSBTEnabled)
+	if (!UseRayTracingPersistentSBTs())
 	{
 		DirtyBindings = VisibleBindings;
 	}
@@ -229,7 +234,7 @@ void FRayTracingShaderBindingTable::CheckPersistentRHI(
 
 	// Build the new SBT initializer
 	FRayTracingShaderBindingTableInitializer NewSBTInitializer;
-	NewSBTInitializer.Lifetime = GPersistentSBTEnabled ? Lifetime : ERayTracingShaderBindingTableLifetime::Transient;
+	NewSBTInitializer.Lifetime = UseRayTracingPersistentSBTs() ? Lifetime : ERayTracingShaderBindingTableLifetime::Transient;
 	NewSBTInitializer.ShaderBindingMode = ShaderBindingMode;
 	NewSBTInitializer.HitGroupIndexingMode = HitGroupIndexingMode;
 	NewSBTInitializer.NumShaderSlotsPerGeometrySegment = NumShaderSlotsPerGeometrySegment;
@@ -239,8 +244,8 @@ void FRayTracingShaderBindingTable::CheckPersistentRHI(
 	NewSBTInitializer.LocalBindingDataSize = FMath::Max(LocalBindingDataSize, PersistentSBTInitializer.LocalBindingDataSize);
 
 	// Always force recreate when persistent SBT is not enabled
-	bool bRecreate = !GPersistentSBTEnabled;
-	if (GPersistentSBTEnabled)
+	bool bRecreate = !UseRayTracingPersistentSBTs();
+	if (UseRayTracingPersistentSBTs())
 	{
 		FStringBuilderBase Reason;
 		auto NewDifferent = [&Reason](const TCHAR* MemberName, uint32 CurrentValue, uint32 NewValue)
