Move over path tracing to use raytracing shader binding layout so it can use the persistent SBT as well

diff --git a/Engine/Source/Runtime/Renderer/Private/DeferredShadingRenderer.cpp b/Engine/Source/Runtime/Renderer/Private/DeferredShadingRenderer.cpp
--- a/Engine/Source/Runtime/Renderer/Private/DeferredShadingRenderer.cpp
+++ b/Engine/Source/Runtime/Renderer/Private/DeferredShadingRenderer.cpp
@@ -766,7 +766,7 @@ bool FDeferredShadingSceneRenderer::SetupRayTracingPipelineStatesAndSBT(FRDGBuil
 	if (ShaderBindingMode != ERayTracingShaderBindingMode::Disabled)
 	{
 		// Path tracer doesn't support persistent SBT yet because it doesn't use the shader binding layout yet
-		ERayTracingShaderBindingTableLifetime SBTLifetime = bIsPathTracing ? ERayTracingShaderBindingTableLifetime::Transient : ERayTracingShaderBindingTableLifetime::Persistent;
+		ERayTracingShaderBindingTableLifetime SBTLifetime = ERayTracingShaderBindingTableLifetime::Persistent;
 		const FRayTracingScene& RayTracingScene = Scene->RayTracingScene;
 		Scene->RayTracingSBT.CheckPersistentRHI(
 			GraphBuilder.RHICmdList,
diff --git a/Engine/Source/Runtime/Renderer/Private/PathTracing.cpp b/Engine/Source/Runtime/Renderer/Private/PathTracing.cpp
--- a/Engine/Source/Runtime/Renderer/Private/PathTracing.cpp
+++ b/Engine/Source/Runtime/Renderer/Private/PathTracing.cpp
@@ -1214,6 +1214,11 @@ class FPathTracingRG : public FGlobalShader
 	{
 		return ERayTracingPayloadType::PathTracingMaterial | ERayTracingPayloadType::Decals;
 	}
+			
+	static const FShaderBindingLayout* GetShaderBindingLayout(const FShaderPermutationParameters& Parameters)
+	{
+		return RayTracing::GetShaderBindingLayout(Parameters.Platform);
+	}
 
 	BEGIN_SHADER_PARAMETER_STRUCT(FParameters, )
 		SHADER_PARAMETER_RDG_TEXTURE_UAV(RWTexture2D<float4>, RadianceTexture)
@@ -1304,6 +1309,11 @@ class FPathTracingDebugRG : public FGlobalShader
 	{
 		return ERayTracingPayloadType::PathTracingMaterial | ERayTracingPayloadType::Decals;
 	}
+		
+	static const FShaderBindingLayout* GetShaderBindingLayout(const FShaderPermutationParameters& Parameters)
+	{
+		return RayTracing::GetShaderBindingLayout(Parameters.Platform);
+	}
 
 	BEGIN_SHADER_PARAMETER_STRUCT(FParameters, )
 		SHADER_PARAMETER_RDG_TEXTURE_UAV(RWTexture2D<float4>, RWSceneColor)
@@ -1322,6 +1332,9 @@ class FPathTracingDebugRG : public FGlobalShader
 		// scene decals
 		SHADER_PARAMETER_RDG_UNIFORM_BUFFER(FRayTracingDecals, DecalParameters)
 
+		// scene uniform buffer
+		SHADER_PARAMETER_RDG_UNIFORM_BUFFER(FSceneUniformParameters, Scene)
+
 		// Skylight
 		SHADER_PARAMETER_STRUCT_INCLUDE(FPathTracingSkylight, SkylightParameters)
 	END_SHADER_PARAMETER_STRUCT()
@@ -1367,6 +1380,11 @@ class FPathTracingInitExtinctionCoefficientRG : public FGlobalShader
 	{
 		return ERayTracingPayloadType::PathTracingMaterial;
 	}
+			
+	static const FShaderBindingLayout* GetShaderBindingLayout(const FShaderPermutationParameters& Parameters)
+	{
+		return RayTracing::GetShaderBindingLayout(Parameters.Platform);
+	}
 
 	BEGIN_SHADER_PARAMETER_STRUCT(FParameters, )
 		SHADER_PARAMETER_RDG_BUFFER_SRV(RaytracingAccelerationStructure, TLAS)
@@ -1627,6 +1645,11 @@ public:
 		// The simplest solution is probably to remove FDecalShaderPayload and use the PT payload everywhere
 		return ERayTracingPayloadType::Decals;
 	}
+			
+	static const FShaderBindingLayout* GetShaderBindingLayout(const FShaderPermutationParameters& Parameters)
+	{
+		return RayTracing::GetShaderBindingLayout(Parameters.Platform);
+	}
 
 	void GetShaderBindings(
 		const FScene* Scene,
@@ -1828,6 +1851,11 @@ public:
 			return ERayTracingPayloadType::PathTracingMaterial;
 		}
 	}
+			
+	static const FShaderBindingLayout* GetShaderBindingLayout(const FShaderPermutationParameters& Parameters)
+	{
+		return RayTracing::GetShaderBindingLayout(Parameters.Platform);
+	}
 };
 
 using FPathTracingDefaultMS  = TPathTracingDefaultMS<false>;
@@ -1949,6 +1977,11 @@ public:
 	{
 		return ERayTracingPayloadType::PathTracingMaterial;
 	}
+			
+	static const FShaderBindingLayout* GetShaderBindingLayout(const FShaderPermutationParameters& Parameters)
+	{
+		return RayTracing::GetShaderBindingLayout(Parameters.Platform);
+	}
 };
 
 IMPLEMENT_MATERIAL_SHADER_TYPE(, FPathTracingLightingMS, TEXT("/Engine/Private/PathTracing/PathTracingLightingMissShader.usf"), TEXT("PathTracingLightingMS"), SF_RayMiss);
@@ -2165,6 +2198,11 @@ public:
 			return ERayTracingPayloadType::PathTracingMaterial;
 		}
 	}
+			
+	static const FShaderBindingLayout* GetShaderBindingLayout(const FShaderPermutationParameters& Parameters)
+	{
+		return RayTracing::GetShaderBindingLayout(Parameters.Platform);
+	}
 };
 
 
@@ -2225,6 +2263,11 @@ class TPathTracingDefaultHitGroup : public FGlobalShader
 			return ERayTracingPayloadType::PathTracingMaterial;
 		}
 	}
+			
+	static const FShaderBindingLayout* GetShaderBindingLayout(const FShaderPermutationParameters& Parameters)
+	{
+		return RayTracing::GetShaderBindingLayout(Parameters.Platform);
+	}
 
 	using FParameters = FEmptyShaderParameters;
 };
@@ -3040,6 +3083,13 @@ void PreparePathTracingRTPSO()
 				Initializer.SetMissShaderTable(MissShaderTable);
 				Initializer.SetHitGroupTable(HitGroupTable);
 				Initializer.MaxPayloadSizeInBytes = RayGenShaderTable[0]->RayTracingPayloadSize;
+
+				const FShaderBindingLayout* ShaderBindingLayout = RayTracing::GetShaderBindingLayout(GMaxRHIShaderPlatform);
+				if (ShaderBindingLayout)
+				{
+					Initializer.ShaderBindingLayout = &ShaderBindingLayout->RHILayout;
+				}
+
 				FRayTracingPipelineState* PipelineState = PipelineStateCache::GetAndOrCreateRayTracingPipelineState(RHICmdList, Initializer);
 				NumValidPermutations++;
 			}
@@ -3180,6 +3230,7 @@ void FDeferredShadingSceneRenderer::RenderPathTracing(
 		);
 
 		PassParameters->DecalParameters = View.RayTracingDecalUniformBuffer;
+		PassParameters->Scene = View.GetSceneUniforms().GetBuffer(GraphBuilder);
 
 		TShaderMapRef<FPathTracingDebugRG> RayGenShader(View.ShaderMap);
 		ClearUnusedGraphResources(RayGenShader, PassParameters);
@@ -3191,6 +3242,9 @@ void FDeferredShadingSceneRenderer::RenderPathTracing(
 			{
 				FRHIBatchedShaderParameters& GlobalResources = RHICmdList.GetScratchShaderParameters();
 				SetShaderParameters(GlobalResources, RayGenShader, *PassParameters);
+				
+				TOptional<FScopedUniformBufferStaticBindings> StaticUniformBufferScope = RayTracing::BindStaticUniformBufferBindings(View, PassParameters->Scene->GetRHI(), RHICmdList);
+
 				RHICmdList.RayTraceDispatch(
 					View.RayTracingMaterialPipeline,
 					RayGenShader.GetRayTracingShader(),
@@ -3355,6 +3409,8 @@ void FDeferredShadingSceneRenderer::RenderPathTracing(
 				FRHIBatchedShaderParameters& GlobalResources = RHICmdList.GetScratchShaderParameters();
 				SetShaderParameters(GlobalResources, RayGenShader, *PassParameters);
 
+				TOptional<FScopedUniformBufferStaticBindings> StaticUniformBufferScope = RayTracing::BindStaticUniformBufferBindings(View, PassParameters->Scene->GetRHI(), RHICmdList);
+
 				RHICmdList.RayTraceDispatch(
 					View.RayTracingMaterialPipeline,
 					RayGenShader.GetRayTracingShader(),
@@ -3932,6 +3988,9 @@ void FDeferredShadingSceneRenderer::RenderPathTracing(
 								{
 									FRHIBatchedShaderParameters& GlobalResources = RHICmdList.GetScratchShaderParameters();
 									SetShaderParameters(GlobalResources, RayGenShader, *PassParameters);
+									
+									TOptional<FScopedUniformBufferStaticBindings> StaticUniformBufferScope = RayTracing::BindStaticUniformBufferBindings(View, PassParameters->Scene->GetRHI(), RHICmdList);
+
 									if (bUseIndirectDispatch && PassParameters->FirstBounce > 0)
 									{
 										PassParameters->PathTracingIndirectArgs->MarkResourceAsUsed();
diff --git a/Engine/Source/Runtime/Renderer/Private/RayTracing/RayTracingDecals.cpp b/Engine/Source/Runtime/Renderer/Private/RayTracing/RayTracingDecals.cpp
--- a/Engine/Source/Runtime/Renderer/Private/RayTracing/RayTracingDecals.cpp
+++ b/Engine/Source/Runtime/Renderer/Private/RayTracing/RayTracingDecals.cpp
@@ -213,6 +213,11 @@ public:
 		return ERayTracingPayloadType::Decals;
 	}
 
+	static const FShaderBindingLayout* GetShaderBindingLayout(const FShaderPermutationParameters& Parameters)
+	{
+		return RayTracing::GetShaderBindingLayout(Parameters.Platform);
+	}
+
 	void GetShaderBindings(
 		const FScene* Scene,
 		ERHIFeatureLevel::Type FeatureLevel,
@@ -323,6 +328,11 @@ public:
 	{
 		return ERayTracingPayloadType::Decals;
 	}
+
+	static const FShaderBindingLayout* GetShaderBindingLayout(const FShaderPermutationParameters& Parameters)
+	{
+		return RayTracing::GetShaderBindingLayout(Parameters.Platform);
+	}
 };
 
 using FRayTracingDecalMaterialCHS = TRayTracingDecalMaterial<false>;
@@ -358,6 +368,11 @@ class FDefaultOpaqueMeshDecalHitGroup : public FGlobalShader
 		return ERayTracingPayloadType::Decals;
 	}
 
+	static const FShaderBindingLayout* GetShaderBindingLayout(const FShaderPermutationParameters& Parameters)
+	{
+		return RayTracing::GetShaderBindingLayout(Parameters.Platform);
+	}
+
 	using FParameters = FEmptyShaderParameters;
 };
 
@@ -378,6 +393,11 @@ class FDefaultHiddenMeshDecalHitGroup : public FGlobalShader
 		return ERayTracingPayloadType::Decals;
 	}
 
+	static const FShaderBindingLayout* GetShaderBindingLayout(const FShaderPermutationParameters& Parameters)
+	{
+		return RayTracing::GetShaderBindingLayout(Parameters.Platform);
+	}
+
 	using FParameters = FEmptyShaderParameters;
 };
 
diff --git a/Engine/Source/Runtime/Renderer/Private/RayTracing/RayTracingMaterialHitShaders.cpp b/Engine/Source/Runtime/Renderer/Private/RayTracing/RayTracingMaterialHitShaders.cpp
--- a/Engine/Source/Runtime/Renderer/Private/RayTracing/RayTracingMaterialHitShaders.cpp
+++ b/Engine/Source/Runtime/Renderer/Private/RayTracing/RayTracingMaterialHitShaders.cpp
@@ -664,7 +664,7 @@ void FDeferredShadingSceneRenderer::CreateMaterialRayTracingMaterialPipeline(
 	Initializer.MaxPayloadSizeInBytes = GetRayTracingPayloadTypeMaxSize(PayloadType);
 
 	const FShaderBindingLayout* ShaderBindingLayout = RayTracing::GetShaderBindingLayout(View.GetShaderPlatform());
-	if (!bIsPathTracing && ShaderBindingLayout)
+	if (ShaderBindingLayout)
 	{
 		Initializer.ShaderBindingLayout = &ShaderBindingLayout->RHILayout;
 	}
