Real time sky light capture for the mobile renderer.
The SkyEnvIrradiance is readback on CPU to be copied into the view UB in order to not add any buffer sampling on the mobile path that was reported out of SRV before already for some materials.

#tests PC dx12/dx11, preview mobile pc d3d, console, mobile OGL and Vulkan.
[RN] Sky light real time capture is now supported by the mobile renderer.

diff --git a/Engine/Shaders/Private/MobileBasePassPixelShader.usf b/Engine/Shaders/Private/MobileBasePassPixelShader.usf
--- a/Engine/Shaders/Private/MobileBasePassPixelShader.usf
+++ b/Engine/Shaders/Private/MobileBasePassPixelShader.usf
@@ -1132,7 +1132,13 @@ void Main(
 #endif
 
 #if !MATERIALBLENDING_MODULATE
-	OutColor.rgb *= ResolvedView.PreExposure;
+#if MATERIAL_IS_SKY
+	// Dynamic capture exposure is 1 as of today.
+	const float ViewPreExposure = View.RealTimeReflectionCapture > 0.0f ? View.RealTimeReflectionCapturePreExposure : View.PreExposure;
+#else
+	const float ViewPreExposure = View.PreExposure;
+#endif
+	OutColor.rgb *= ViewPreExposure;
 #endif
 
 #if MATERIAL_IS_SKY
diff --git a/Engine/Source/Runtime/Engine/Private/Components/SkyLightComponent.cpp b/Engine/Source/Runtime/Engine/Private/Components/SkyLightComponent.cpp
--- a/Engine/Source/Runtime/Engine/Private/Components/SkyLightComponent.cpp
+++ b/Engine/Source/Runtime/Engine/Private/Components/SkyLightComponent.cpp
@@ -89,7 +89,7 @@ FAutoConsoleVariableRef CVarSkylightIntensityMultiplier(
 
 static bool FeatureLevelSupportsRealTimeReflectionCapture(FStaticFeatureLevel FeatureLevel)
 {
-	return FeatureLevel > ERHIFeatureLevel::ES3_1;
+	return FeatureLevel >= ERHIFeatureLevel::ES3_1;
 }
 
 void OnChangeSkylightRealTimeReflectionCapture(IConsoleVariable* Var)
diff --git a/Engine/Source/Runtime/Renderer/Private/DecalRenderingShared.cpp b/Engine/Source/Runtime/Renderer/Private/DecalRenderingShared.cpp
--- a/Engine/Source/Runtime/Renderer/Private/DecalRenderingShared.cpp
+++ b/Engine/Source/Runtime/Renderer/Private/DecalRenderingShared.cpp
@@ -283,7 +283,11 @@ public:
 
 		if (MobileReflectionCaptureParam.IsBound())
 		{
-			if (Scene && Scene->SkyLight && Scene->SkyLight->ProcessedTexture && Scene->SkyLight->ProcessedTexture->TextureRHI)
+			if (Scene &&
+				(
+					(Scene->SkyLight && Scene->SkyLight->ProcessedTexture && Scene->SkyLight->ProcessedTexture->TextureRHI)
+					|| Scene->CanSampleSkyLightRealTimeCaptureData()
+				))
 			{
 				SetUniformBufferParameter(BatchedParameters, MobileReflectionCaptureParam, Scene->UniformBuffers.MobileSkyReflectionUniformBuffer);
 			}
diff --git a/Engine/Source/Runtime/Renderer/Private/IlluminanceMeter.cpp b/Engine/Source/Runtime/Renderer/Private/IlluminanceMeter.cpp
--- a/Engine/Source/Runtime/Renderer/Private/IlluminanceMeter.cpp
+++ b/Engine/Source/Runtime/Renderer/Private/IlluminanceMeter.cpp
@@ -120,7 +120,7 @@ FScreenPassTexture ProcessAndRenderIlluminanceMeter(FRDGBuilder& GraphBuilder, c
 	// 0- Get source cubemap and allocate cube map transient resources
 
 	FTextureRHIRef SkyLightTextureRHI = GBlackTextureCube->TextureRHI;
-	if (SkyLight->bRealTimeCaptureEnabled && Scene->ConvolvedSkyRenderTargetReadyIndex >= 0)
+	if (Scene && Scene->CanSampleSkyLightRealTimeCaptureData())
 	{
 		// Cannot blend with this capture mode as of today.
 		SkyLightTextureRHI = Scene->ConvolvedSkyRenderTarget[Scene->ConvolvedSkyRenderTargetReadyIndex]->GetRHI();
diff --git a/Engine/Source/Runtime/Renderer/Private/IndirectLightRendering.cpp b/Engine/Source/Runtime/Renderer/Private/IndirectLightRendering.cpp
--- a/Engine/Source/Runtime/Renderer/Private/IndirectLightRendering.cpp
+++ b/Engine/Source/Runtime/Renderer/Private/IndirectLightRendering.cpp
@@ -573,12 +573,12 @@ void SetupReflectionUniformParameters(FRDGBuilder& GraphBuilder, const FViewInfo
 
 	if (Scene
 		&& Scene->SkyLight
-		&& (Scene->SkyLight->ProcessedTexture || (Scene->SkyLight->bRealTimeCaptureEnabled && Scene->ConvolvedSkyRenderTargetReadyIndex >= 0))
+		&& (Scene->SkyLight->ProcessedTexture || Scene->CanSampleSkyLightRealTimeCaptureData())
 		&& bApplySkyLight)
 	{
 		const FSkyLightSceneProxy& SkyLight = *Scene->SkyLight;
 
-		if (Scene->SkyLight->bRealTimeCaptureEnabled && Scene->ConvolvedSkyRenderTargetReadyIndex >= 0)
+		if (Scene->CanSampleSkyLightRealTimeCaptureData())
 		{
 			// Cannot blend with this capture mode as of today.
 			SkyLightTextureResource = Scene->ConvolvedSkyRenderTarget[Scene->ConvolvedSkyRenderTargetReadyIndex]->GetRHI();
diff --git a/Engine/Source/Runtime/Renderer/Private/MobileBasePass.cpp b/Engine/Source/Runtime/Renderer/Private/MobileBasePass.cpp
--- a/Engine/Source/Runtime/Renderer/Private/MobileBasePass.cpp
+++ b/Engine/Source/Runtime/Renderer/Private/MobileBasePass.cpp
@@ -251,8 +251,15 @@ static bool UseSkyReflectionCapture(const FScene* RenderScene)
 {
 	return RenderScene
 		&& RenderScene->SkyLight
-		&& RenderScene->SkyLight->ProcessedTexture
-		&& RenderScene->SkyLight->ProcessedTexture->TextureRHI;
+		&& 
+		(
+			(
+				   RenderScene->SkyLight->ProcessedTexture
+				&& RenderScene->SkyLight->ProcessedTexture->TextureRHI
+			)
+			||
+			RenderScene->CanSampleSkyLightRealTimeCaptureData()
+		);
 }
 
 const FLightSceneInfo* MobileBasePass::GetDirectionalLightInfo(const FScene* Scene, const FPrimitiveSceneProxy* PrimitiveSceneProxy)
diff --git a/Engine/Source/Runtime/Renderer/Private/MobileBasePassRendering.cpp b/Engine/Source/Runtime/Renderer/Private/MobileBasePassRendering.cpp
--- a/Engine/Source/Runtime/Renderer/Private/MobileBasePassRendering.cpp
+++ b/Engine/Source/Runtime/Renderer/Private/MobileBasePassRendering.cpp
@@ -479,36 +479,47 @@ void SetupMobileDirectionalLightUniformParameters(
 	}
 }
 
-void SetupMobileSkyReflectionUniformParameters(FSkyLightSceneProxy* SkyLight, FMobileReflectionCaptureShaderParameters& Parameters)
+void SetupMobileSkyReflectionUniformParameters(const FScene* Scene, FSkyLightSceneProxy* SkyLight, FMobileReflectionCaptureShaderParameters& Parameters)
 {
+	Parameters.Texture				= GBlackTextureCube->TextureRHI;
+	Parameters.TextureSampler		= GBlackTextureCube->SamplerStateRHI;
+	Parameters.TextureBlend			= GBlackTextureCube->TextureRHI;
+	Parameters.TextureBlendSampler	= GBlackTextureCube->SamplerStateRHI;
+ 
+	bool bSkyLightIsDynamic = false;
 	float Brightness = 0.f;
-	float SkyMaxMipIndex = 0.f;
-	FTexture* CaptureTexture = GBlackTextureCube;
-	FTexture* CaptureTextureBlend = GBlackTextureCube;
 	float BlendFraction = 0.f;
+	bool bIsRealTimeCapture = Scene && Scene->CanSampleSkyLightRealTimeCaptureData();
 
-	bool bSkyLightIsDynamic = false;
-	if (SkyLight && SkyLight->ProcessedTexture)
+	if (bIsRealTimeCapture)
+	{
+		bSkyLightIsDynamic			= SkyLight && !SkyLight->bHasStaticLighting && !SkyLight->bWantsStaticShadowing;
+		Parameters.Texture			= Scene->ConvolvedSkyRenderTarget[Scene->ConvolvedSkyRenderTargetReadyIndex]->GetRHI();
+		Parameters.TextureSampler	= TStaticSamplerState<SF_Trilinear, AM_Clamp, AM_Clamp, AM_Clamp>::GetRHI();
+	}
+	else
 	{
-		check(SkyLight->ProcessedTexture->IsInitialized());
-		CaptureTexture = SkyLight->ProcessedTexture;
-		SkyMaxMipIndex = FMath::Log2(static_cast<float>(CaptureTexture->GetSizeX()));
-		Brightness = SkyLight->AverageBrightness;
-		bSkyLightIsDynamic = !SkyLight->bHasStaticLighting && !SkyLight->bWantsStaticShadowing;
-
-		BlendFraction = SkyLight->BlendFraction;
-		if (BlendFraction > 0.0 && SkyLight->BlendDestinationProcessedTexture)
+		if (SkyLight && SkyLight->ProcessedTexture)
 		{
-			CaptureTextureBlend = SkyLight->BlendDestinationProcessedTexture;
+			check(SkyLight->ProcessedTexture->IsInitialized());
+			Parameters.Texture			= SkyLight->ProcessedTexture->TextureRHI;
+			Parameters.TextureSampler	= SkyLight->ProcessedTexture->SamplerStateRHI;
+			Brightness					= SkyLight->AverageBrightness;
+			bSkyLightIsDynamic			= SkyLight && !SkyLight->bHasStaticLighting && !SkyLight->bWantsStaticShadowing;
+
+			BlendFraction = SkyLight->BlendFraction;
+			if (BlendFraction > 0.0 && SkyLight->BlendDestinationProcessedTexture)
+			{
+				Parameters.TextureBlend			= SkyLight->BlendDestinationProcessedTexture->TextureRHI;
+				Parameters.TextureBlendSampler	= SkyLight->BlendDestinationProcessedTexture->SamplerStateRHI;
+			}
 		}
 	}
-	
+
+	float SkyMaxMipIndex = FMath::Log2(static_cast<float>(Parameters.Texture->GetDesc().Extent.X));
+
 	//To keep ImageBasedReflectionLighting coherence with PC, use AverageBrightness instead of InvAverageBrightness to calculate the IBL contribution
-	Parameters.Params = FVector4f(Brightness, SkyMaxMipIndex, bSkyLightIsDynamic ? 1.0f : 0.0f, BlendFraction);
-	Parameters.Texture = CaptureTexture->TextureRHI;
-	Parameters.TextureSampler = CaptureTexture->SamplerStateRHI;
-	Parameters.TextureBlend = CaptureTextureBlend->TextureRHI;
-	Parameters.TextureBlendSampler = CaptureTextureBlend->SamplerStateRHI;
+	Parameters.Params = FVector4f(Brightness, SkyMaxMipIndex, bIsRealTimeCapture ? 2.0 : (bSkyLightIsDynamic ? 1.0f : 0.0f), BlendFraction);
 }
 
 void FMobileSceneRenderer::RenderMobileBasePass(FRHICommandList& RHICmdList, const FViewInfo& View, const FInstanceCullingDrawParams* InstanceCullingDrawParams, const FInstanceCullingDrawParams* SkyPassInstanceCullingDrawParams)
diff --git a/Engine/Source/Runtime/Renderer/Private/MobileBasePassRendering.h b/Engine/Source/Runtime/Renderer/Private/MobileBasePassRendering.h
--- a/Engine/Source/Runtime/Renderer/Private/MobileBasePassRendering.h
+++ b/Engine/Source/Runtime/Renderer/Private/MobileBasePassRendering.h
@@ -98,6 +98,7 @@ extern void SetupMobileDirectionalLightUniformParameters(
 	FMobileDirectionalLightShaderParameters& Parameters);
 
 extern void SetupMobileSkyReflectionUniformParameters(
+	const FScene* Scene,
 	class FSkyLightSceneProxy* SkyLight,
 	FMobileReflectionCaptureShaderParameters& Parameters);
 
diff --git a/Engine/Source/Runtime/Renderer/Private/MobileShadingRenderer.cpp b/Engine/Source/Runtime/Renderer/Private/MobileShadingRenderer.cpp
--- a/Engine/Source/Runtime/Renderer/Private/MobileShadingRenderer.cpp
+++ b/Engine/Source/Runtime/Renderer/Private/MobileShadingRenderer.cpp
@@ -653,10 +653,15 @@ void FMobileSceneRenderer::InitViews(
 	// Finalize and set the scene textures config.
 	FSceneTexturesConfig::Set(SceneTexturesConfig);
 
+	bool bShouldRenderSkyAtmosphere = false;
 	if (bRendererOutputFinalSceneColor)
 	{
+		// This must happen before we start initialising and using views (allocating Scene->SkyIrradianceEnvironmentMap).
+		UpdateSkyIrradianceGpuBuffer(GraphBuilder, ViewFamily.EngineShowFlags, Scene->SkyLight, Scene->SkyIrradianceEnvironmentMap);
+
 		// Initialise Sky/View resources before the view global uniform buffer is built.
-		if (ShouldRenderSkyAtmosphere(Scene, ViewFamily.EngineShowFlags))
+		bShouldRenderSkyAtmosphere = ShouldRenderSkyAtmosphere(Scene, ViewFamily.EngineShowFlags);
+		if (bShouldRenderSkyAtmosphere)
 		{
 			InitSkyAtmosphereForViews(RHICmdList, GraphBuilder);
 		}
@@ -1217,6 +1222,25 @@ void FMobileSceneRenderer::Render(FRDGBuilder& GraphBuilder, const FSceneRenderU
 			PendingRDGResources.CommitToSceneAndViewUniformBuffers(GraphBuilder, /* out */ ExternalAccessQueue);
 		}
 
+		// Run local fog volume initialization before base pass and volumetric fog for all the culled instance instance data to be ready.
+		InitLocalFogVolumesForViews(Scene, Views, ViewFamily, GraphBuilder, ShouldRenderVolumetricFog(), bUseHalfResLocalFogVolume);
+
+		if (bRendererOutputFinalSceneColor && Views.Num() > 0)
+		{
+			FViewInfo& MainView = Views[0];
+			const bool bRealTimeSkyCaptureEnabled = (bShouldRenderSkyAtmosphere || MainView.SkyMeshBatches.Num() > 0) && Scene->SkyLight && Scene->SkyLight->bRealTimeCaptureEnabled && ViewFamily.EngineShowFlags.SkyLighting;
+			if (bRealTimeSkyCaptureEnabled)
+			{
+				// We must execute the submit for transition of SkyAtmosphere resources to happen (see CommitToSceneAndViewUniformBuffers) and avoid validation error.
+				ExternalAccessQueue.Submit(GraphBuilder);
+
+				const bool bShouldRenderVolumetricCloud = false; // Not supported on the mobile renderer.
+				Scene->AllocateAndCaptureFrameSkyEnvMap(GraphBuilder, *this, MainView, bShouldRenderSkyAtmosphere, bShouldRenderVolumetricCloud, InstanceCullingManager, ExternalAccessQueue);
+
+				UpdateSkyReflectionUniformBuffer(GraphBuilder.RHICmdList);
+			}
+		}
+
 		// Hair update
 		if (IsHairStrandsEnabled(EHairStrandsShaderType::All, Scene->GetShaderPlatform()) && RendererOutput != ERendererOutput::DepthPrepassOnly)
 		{
@@ -1250,15 +1274,26 @@ void FMobileSceneRenderer::Render(FRDGBuilder& GraphBuilder, const FSceneRenderU
 		RenderShadowDepthMaps(GraphBuilder, nullptr, InstanceCullingManager, ExternalAccessQueue);
 		GraphBuilder.AddDispatchHint();
 
-		// Run local fog volume initialization before base pass and volumetric fog for all the culled instance instance data to be ready.
-		InitLocalFogVolumesForViews(Scene, Views, ViewFamily, GraphBuilder, ShouldRenderVolumetricFog(), bUseHalfResLocalFogVolume);
-
 		if (ShouldRenderVolumetricFog())
 		{
 			ComputeVolumetricFog(GraphBuilder, SceneTextures);
 		}
 		ExternalAccessQueue.Submit(GraphBuilder);
 
+//		if (bRendererOutputFinalSceneColor && Views.Num() > 0)
+//		{
+//			FViewInfo& MainView = Views[0];
+//			const bool bRealTimeSkyCaptureEnabled = (bShouldRenderSkyAtmosphere || MainView.SkyMeshBatches.Num() > 0) && Scene->SkyLight && Scene->SkyLight->bRealTimeCaptureEnabled && ViewFamily.EngineShowFlags.SkyLighting;
+//			if (bRealTimeSkyCaptureEnabled)
+//			{
+//				const bool bShouldRenderVolumetricCloud = false; // Not supported on the mobile renderer.
+//				Scene->AllocateAndCaptureFrameSkyEnvMap(GraphBuilder, *this, MainView, bShouldRenderSkyAtmosphere, bShouldRenderVolumetricCloud, InstanceCullingManager, ExternalAccessQueue);
+//				ExternalAccessQueue.Submit(GraphBuilder);
+//
+//				UpdateSkyReflectionUniformBuffer(GraphBuilder.RHICmdList);
+//			}
+//		}
+
 		// Custom depth
 		// bShouldRenderCustomDepth has been initialized in InitViews on mobile platform
 		if (bShouldRenderCustomDepth)
@@ -2502,10 +2537,16 @@ void FMobileSceneRenderer::UpdateSkyReflectionUniformBuffer(FRHICommandListBase&
 {
 	FSkyLightSceneProxy* SkyLight = nullptr;
 	if (Scene->SkyLight
-		&& Scene->SkyLight->ProcessedTexture
-		&& Scene->SkyLight->ProcessedTexture->TextureRHI
-		// Don't use skylight reflection if it is a static sky light for keeping coherence with PC.
-		&& !Scene->SkyLight->bHasStaticLighting)
+		&& 
+		(
+			(
+				Scene->SkyLight->ProcessedTexture && Scene->SkyLight->ProcessedTexture->TextureRHI
+				// Don't use skylight reflection if it is a static sky light for keeping coherence with PC.
+				&& !Scene->SkyLight->bHasStaticLighting
+			)
+			||
+			Scene->CanSampleSkyLightRealTimeCaptureData()
+		))
 	{
 		SkyLight = Scene->SkyLight;
 	}
@@ -2517,7 +2558,7 @@ void FMobileSceneRenderer::UpdateSkyReflectionUniformBuffer(FRHICommandListBase&
 	}
 
 	FMobileReflectionCaptureShaderParameters Parameters;
-	SetupMobileSkyReflectionUniformParameters(SkyLight, Parameters);
+	SetupMobileSkyReflectionUniformParameters(Scene, SkyLight, Parameters);
 	Scene->UniformBuffers.MobileSkyReflectionUniformBuffer.UpdateUniformBufferImmediate(RHICmdList, Parameters);
 }
 
diff --git a/Engine/Source/Runtime/Renderer/Private/ReflectionEnvironment.cpp b/Engine/Source/Runtime/Renderer/Private/ReflectionEnvironment.cpp
--- a/Engine/Source/Runtime/Renderer/Private/ReflectionEnvironment.cpp
+++ b/Engine/Source/Runtime/Renderer/Private/ReflectionEnvironment.cpp
@@ -679,8 +679,10 @@ void UpdateSkyIrradianceGpuBuffer(FRDGBuilder& GraphBuilder, const FEngineShowFl
 		&& EngineShowFlags.SkyLighting
 		&& !SkyLight->bRealTimeCaptureEnabled; // When bRealTimeCaptureEnabled is true, the buffer will be setup on GPU directly in this case
 
+	FRDGBufferDesc BufferDesc = FRDGBufferDesc::CreateStructuredDesc(sizeof(FVector4f), SKY_IRRADIANCE_ENVIRONMENT_MAP_VEC4_COUNT);
+	BufferDesc.Usage |= EBufferUsageFlags::SourceCopy;
 	const bool bNewAlloc = AllocatePooledBuffer(
-		FRDGBufferDesc::CreateStructuredDesc(sizeof(FVector4f), SKY_IRRADIANCE_ENVIRONMENT_MAP_VEC4_COUNT),
+		BufferDesc,
 		Buffer,
 		TEXT("SkyIrradianceEnvironmentMap"),
 		ERDGPooledBufferAlignment::None);
diff --git a/Engine/Source/Runtime/Renderer/Private/ReflectionEnvironmentRealTimeCapture.cpp b/Engine/Source/Runtime/Renderer/Private/ReflectionEnvironmentRealTimeCapture.cpp
--- a/Engine/Source/Runtime/Renderer/Private/ReflectionEnvironmentRealTimeCapture.cpp
+++ b/Engine/Source/Runtime/Renderer/Private/ReflectionEnvironmentRealTimeCapture.cpp
@@ -21,6 +21,7 @@
 #include "GPUScene.h"
 #include "ScreenPass.h"
 #include "SkyAtmosphereRendering.h"
+#include "MobileBasePassRendering.h"
 
 #if WITH_EDITOR
 #include "CanvasTypes.h"
@@ -74,7 +75,7 @@ class FDownsampleCubeFaceCS : public FGlobalShader
 	END_SHADER_PARAMETER_STRUCT()
 
 public:
-	static bool ShouldCompilePermutation(const FGlobalShaderPermutationParameters& Parameters) { return GetMaxSupportedFeatureLevel(Parameters.Platform) >= ERHIFeatureLevel::SM5; }
+	static bool ShouldCompilePermutation(const FGlobalShaderPermutationParameters& Parameters) { return GetMaxSupportedFeatureLevel(Parameters.Platform) >= ERHIFeatureLevel::ES3_1; }
 	static void ModifyCompilationEnvironment(const FGlobalShaderPermutationParameters& Parameters, FShaderCompilerEnvironment& OutEnvironment)
 	{
 		FGlobalShader::ModifyCompilationEnvironment(Parameters, OutEnvironment);
@@ -107,7 +108,7 @@ class FConvolveSpecularFaceCS : public FGlobalShader
 	END_SHADER_PARAMETER_STRUCT()
 
 public:
-	static bool ShouldCompilePermutation(const FGlobalShaderPermutationParameters& Parameters) { return GetMaxSupportedFeatureLevel(Parameters.Platform) >= ERHIFeatureLevel::SM5; }
+	static bool ShouldCompilePermutation(const FGlobalShaderPermutationParameters& Parameters) { return GetMaxSupportedFeatureLevel(Parameters.Platform) >= ERHIFeatureLevel::ES3_1; }
 	static void ModifyCompilationEnvironment(const FGlobalShaderPermutationParameters& Parameters, FShaderCompilerEnvironment& OutEnvironment)
 	{
 		FGlobalShader::ModifyCompilationEnvironment(Parameters, OutEnvironment);
@@ -141,7 +142,7 @@ class FComputeSkyEnvMapDiffuseIrradianceCS : public FGlobalShader
 	END_SHADER_PARAMETER_STRUCT()
 
 public:
-	static bool ShouldCompilePermutation(const FGlobalShaderPermutationParameters& Parameters) { return GetMaxSupportedFeatureLevel(Parameters.Platform) >= ERHIFeatureLevel::SM5; }
+	static bool ShouldCompilePermutation(const FGlobalShaderPermutationParameters& Parameters) { return GetMaxSupportedFeatureLevel(Parameters.Platform) >= ERHIFeatureLevel::ES3_1; }
 	static void ModifyCompilationEnvironment(const FGlobalShaderPermutationParameters& Parameters, FShaderCompilerEnvironment& OutEnvironment)
 	{
 		FGlobalShader::ModifyCompilationEnvironment(Parameters, OutEnvironment);
@@ -299,6 +300,13 @@ BEGIN_SHADER_PARAMETER_STRUCT(FCaptureSkyMeshReflectionPassParameters, )
 	RENDER_TARGET_BINDING_SLOTS()
 END_SHADER_PARAMETER_STRUCT()
 
+BEGIN_SHADER_PARAMETER_STRUCT(FMobileCaptureSkyMeshReflectionPassParameters, )
+	SHADER_PARAMETER_STRUCT_INCLUDE(FViewShaderParameters, View)
+	SHADER_PARAMETER_STRUCT_INCLUDE(FInstanceCullingDrawParams, InstanceCullingDrawParams)
+	SHADER_PARAMETER_RDG_UNIFORM_BUFFER(FMobileBasePassUniformParameters, BasePass)
+	RENDER_TARGET_BINDING_SLOTS()
+END_SHADER_PARAMETER_STRUCT()
+
 
 void FScene::AllocateAndCaptureFrameSkyEnvMap(
 	FRDGBuilder& GraphBuilder, FSceneRenderer& SceneRenderer, FViewInfo& MainView,
@@ -602,76 +610,153 @@ void FScene::AllocateAndCaptureFrameSkyEnvMap(
 				const bool bUseDepthBuffer = CVarRealTimeReflectionCaptureDepthBuffer.GetValueOnRenderThread() > 0;
 				FRDGTextureRef CubeDepthTexture = nullptr;
 
+				const bool bIsMobilePlatform = IsMobilePlatform(MainView.GetShaderPlatform());
 				if (bExecuteSky)
 				{
 					if(MainView.bSceneHasSkyMaterial || bShouldRenderSkyAtmosphere)
 					{
 						// If there are any mesh tagged as IsSky then we render them only, otherwise we simply render the sky atmosphere itself.
-						if (MainView.bSceneHasSkyMaterial)
+						if (MainView.bSceneHasSkyMaterial)	// TODO skypass using FMobileBasePassUniformParameters
 						{
-							RDG_EVENT_SCOPE(GraphBuilder, "Capture Sky Materials", CubeFace);
-							auto* PassParameters = GraphBuilder.AllocParameters<FCaptureSkyMeshReflectionPassParameters>();
-							PassParameters->View = CubeView.GetShaderParameters();
-							PassParameters->RenderTargets = SkyRC.RenderTargets;
-							PassParameters->BasePass = CreateOpaqueBasePassUniformBuffer(GraphBuilder, MainView, 0);
-
-							// Setup the depth buffer
-							if (bUseDepthBuffer)
-							{
-								FRDGTextureDesc CubeDepthTextureDesc = FRDGTextureDesc::Create2D(FIntPoint(CubeWidth, CubeWidth), PF_DepthStencil,
-									MainView.GetSceneTexturesConfig().DepthClearValue,
-									TexCreate_DepthStencilTargetable | TexCreate_ShaderResource);
-								CubeDepthTexture = GraphBuilder.CreateTexture(CubeDepthTextureDesc, TEXT("SkyLight.CubeDepthTexture"));
-								PassParameters->RenderTargets.DepthStencil = FDepthStencilBinding(CubeDepthTexture, ERenderTargetLoadAction::EClear, FExclusiveDepthStencil::DepthWrite_StencilNop);
-							}
-
-							AddSimpleMeshPass(GraphBuilder, PassParameters, Scene, MainView, &InstanceCullingManager, RDG_EVENT_NAME("CaptureSkyMeshReflection"), SkyRC.Viewport,
-							[&MainView, CubeViewUniformBuffer = CubeView.ViewUniformBuffer, bUseDepthBuffer, Scene](FDynamicPassMeshDrawListContext* DynamicMeshPassContext)
+							if (bIsMobilePlatform)
 							{
-								FMeshPassProcessorRenderState DrawRenderState;
-
-								FExclusiveDepthStencil::Type BasePassDepthStencilAccess_Sky = bUseDepthBuffer ? FExclusiveDepthStencil::Type(Scene->DefaultBasePassDepthStencilAccess | FExclusiveDepthStencil::DepthWrite)
-									: FExclusiveDepthStencil::Type(Scene->DefaultBasePassDepthStencilAccess & ~FExclusiveDepthStencil::DepthWrite);
-								SetupBasePassState(BasePassDepthStencilAccess_Sky, false, DrawRenderState);
-
-								FSkyPassMeshProcessor PassMeshProcessor(Scene, Scene->GetFeatureLevel(), nullptr, DrawRenderState, DynamicMeshPassContext);
-								const int32 SkyRealTimeReflectionOnlyMeshBatcheCount = MainView.SkyMeshBatches.Num();
-								for (int32 MeshBatchIndex = 0; MeshBatchIndex < SkyRealTimeReflectionOnlyMeshBatcheCount; ++MeshBatchIndex)
+								RDG_EVENT_SCOPE(GraphBuilder, "Capture Sky Materials", CubeFace);
+								auto* PassParameters = GraphBuilder.AllocParameters<FMobileCaptureSkyMeshReflectionPassParameters>();
+								PassParameters->View = CubeView.GetShaderParameters();
+								PassParameters->RenderTargets = SkyRC.RenderTargets;
+								const EMobileBasePass BasePass = EMobileBasePass::DepthPrePass;
+								PassParameters->BasePass = CreateMobileBasePassUniformBuffer(GraphBuilder, MainView, BasePass, EMobileSceneTextureSetupMode::SceneDepth);
+					
+								// Setup the depth buffer
+								if (bUseDepthBuffer)
+								{
+									FRDGTextureDesc CubeDepthTextureDesc = FRDGTextureDesc::Create2D(FIntPoint(CubeWidth, CubeWidth), PF_DepthStencil,
+										MainView.GetSceneTexturesConfig().DepthClearValue,
+										TexCreate_DepthStencilTargetable | TexCreate_ShaderResource);
+									CubeDepthTexture = GraphBuilder.CreateTexture(CubeDepthTextureDesc, TEXT("SkyLight.CubeDepthTexture"));
+									PassParameters->RenderTargets.DepthStencil = FDepthStencilBinding(CubeDepthTexture, ERenderTargetLoadAction::EClear, FExclusiveDepthStencil::DepthWrite_StencilWrite);
+								}
+					
+								AddSimpleMeshPass(GraphBuilder, PassParameters, Scene, MainView, &InstanceCullingManager, RDG_EVENT_NAME("CaptureSkyMeshReflection"), SkyRC.Viewport,
+								[&MainView, CubeViewUniformBuffer = CubeView.ViewUniformBuffer, bUseDepthBuffer, Scene](FDynamicPassMeshDrawListContext* DynamicMeshPassContext)
 								{
-									FSkyMeshBatch& SkyMeshBatch = MainView.SkyMeshBatches[MeshBatchIndex];
-									if (!SkyMeshBatch.bVisibleInRealTimeSkyCapture)
+									FMeshPassProcessorRenderState DrawRenderState;
+					
+									DrawRenderState.SetBlendState(TStaticBlendStateWriteMask<CW_RGBA>::GetRHI());
+									if (bUseDepthBuffer)
 									{
-										continue;
+										DrawRenderState.SetDepthStencilAccess(FExclusiveDepthStencil::DepthWrite_StencilWrite);
+										DrawRenderState.SetDepthStencilState(TStaticDepthStencilState<true, CF_DepthNearOrEqual>::GetRHI());
 									}
-
-									const FMeshBatch* MeshBatch = SkyMeshBatch.Mesh;
-									const FPrimitiveSceneProxy* PrimitiveSceneProxy = SkyMeshBatch.Proxy;
-									const FPrimitiveSceneInfo* PrimitiveSceneInfo = PrimitiveSceneProxy->GetPrimitiveSceneInfo();
-
-									// Real time sky light capture cannot render dynamic meshes for now.
-									// For those to be rendered we would need to specify a view to the PassMeshProcessor creation above.
-									// Dynamic draws uses temporary per frame & per view data (appended at the end of the GPUScene buffer).
-									// But the view is transient and data on it can morph, and correct data would need to be added to FGPUScenePrimitiveCollector (see UploadDynamicPrimitiveShaderDataForViewInternal)
-									bool bSkipDynamicMesh = false;
-									for (auto& Element : MeshBatch->Elements)
+									else
 									{
-										if (Element.PrimitiveIdMode == PrimID_DynamicPrimitiveShaderData)
+										DrawRenderState.SetDepthStencilAccess(FExclusiveDepthStencil::DepthNop_StencilWrite);
+										DrawRenderState.SetDepthStencilState(TStaticDepthStencilState<false, CF_Always>::GetRHI());
+									}
+					
+									FSkyPassMeshProcessor PassMeshProcessor(Scene, Scene->GetFeatureLevel(), nullptr, DrawRenderState, DynamicMeshPassContext);
+									const int32 SkyRealTimeReflectionOnlyMeshBatcheCount = MainView.SkyMeshBatches.Num();
+									for (int32 MeshBatchIndex = 0; MeshBatchIndex < SkyRealTimeReflectionOnlyMeshBatcheCount; ++MeshBatchIndex)
+									{
+										FSkyMeshBatch& SkyMeshBatch = MainView.SkyMeshBatches[MeshBatchIndex];
+										if (!SkyMeshBatch.bVisibleInRealTimeSkyCapture)
+										{
+											continue;
+										}
+					
+										const FMeshBatch* MeshBatch = SkyMeshBatch.Mesh;
+										const FPrimitiveSceneProxy* PrimitiveSceneProxy = SkyMeshBatch.Proxy;
+										const FPrimitiveSceneInfo* PrimitiveSceneInfo = PrimitiveSceneProxy->GetPrimitiveSceneInfo();
+					
+										// Real time sky light capture cannot render dynamic meshes for now.
+										// For those to be rendered we would need to specify a view to the PassMeshProcessor creation above.
+										// Dynamic draws uses temporary per frame & per view data (appended at the end of the GPUScene buffer).
+										// But the view is transient and data on it can morph, and correct data would need to be added to FGPUScenePrimitiveCollector (see UploadDynamicPrimitiveShaderDataForViewInternal)
+										bool bSkipDynamicMesh = false;
+										for (auto& Element : MeshBatch->Elements)
 										{
-											bSkipDynamicMesh = true;
+											if (Element.PrimitiveIdMode == PrimID_DynamicPrimitiveShaderData)
+											{
+												bSkipDynamicMesh = true;
+											}
 										}
+										if (bSkipDynamicMesh)
+										{
+											continue;
+										}
+					
+										const uint64 DefaultBatchElementMask = ~0ull;
+										PassMeshProcessor.AddMeshBatch(*MeshBatch, DefaultBatchElementMask, PrimitiveSceneProxy);
 									}
-									if (bSkipDynamicMesh)
+								});
+							}
+							else
+							{
+								RDG_EVENT_SCOPE(GraphBuilder, "Capture Sky Materials", CubeFace);
+								auto* PassParameters = GraphBuilder.AllocParameters<FCaptureSkyMeshReflectionPassParameters>();
+								PassParameters->View = CubeView.GetShaderParameters();
+								PassParameters->RenderTargets = SkyRC.RenderTargets;
+								PassParameters->BasePass = CreateOpaqueBasePassUniformBuffer(GraphBuilder, MainView, 0);
+					
+								// Setup the depth buffer
+								if (bUseDepthBuffer)
+								{
+									FRDGTextureDesc CubeDepthTextureDesc = FRDGTextureDesc::Create2D(FIntPoint(CubeWidth, CubeWidth), PF_DepthStencil,
+										MainView.GetSceneTexturesConfig().DepthClearValue,
+										TexCreate_DepthStencilTargetable | TexCreate_ShaderResource);
+									CubeDepthTexture = GraphBuilder.CreateTexture(CubeDepthTextureDesc, TEXT("SkyLight.CubeDepthTexture"));
+									PassParameters->RenderTargets.DepthStencil = FDepthStencilBinding(CubeDepthTexture, ERenderTargetLoadAction::EClear, FExclusiveDepthStencil::DepthWrite_StencilNop);
+								}
+					
+								AddSimpleMeshPass(GraphBuilder, PassParameters, Scene, MainView, &InstanceCullingManager, RDG_EVENT_NAME("CaptureSkyMeshReflection"), SkyRC.Viewport,
+								[&MainView, CubeViewUniformBuffer = CubeView.ViewUniformBuffer, bUseDepthBuffer, Scene](FDynamicPassMeshDrawListContext* DynamicMeshPassContext)
+								{
+									FMeshPassProcessorRenderState DrawRenderState;
+
+									FExclusiveDepthStencil::Type BasePassDepthStencilAccess_Sky = bUseDepthBuffer ? FExclusiveDepthStencil::Type(Scene->DefaultBasePassDepthStencilAccess | FExclusiveDepthStencil::DepthWrite)
+										: FExclusiveDepthStencil::Type(Scene->DefaultBasePassDepthStencilAccess & ~FExclusiveDepthStencil::DepthWrite);
+									SetupBasePassState(BasePassDepthStencilAccess_Sky, false, DrawRenderState);
+					
+									FSkyPassMeshProcessor PassMeshProcessor(Scene, Scene->GetFeatureLevel(), nullptr, DrawRenderState, DynamicMeshPassContext);
+									const int32 SkyRealTimeReflectionOnlyMeshBatcheCount = MainView.SkyMeshBatches.Num();
+									for (int32 MeshBatchIndex = 0; MeshBatchIndex < SkyRealTimeReflectionOnlyMeshBatcheCount; ++MeshBatchIndex)
 									{
-										continue;
+										FSkyMeshBatch& SkyMeshBatch = MainView.SkyMeshBatches[MeshBatchIndex];
+										if (!SkyMeshBatch.bVisibleInRealTimeSkyCapture)
+										{
+											continue;
+										}
+					
+										const FMeshBatch* MeshBatch = SkyMeshBatch.Mesh;
+										const FPrimitiveSceneProxy* PrimitiveSceneProxy = SkyMeshBatch.Proxy;
+										const FPrimitiveSceneInfo* PrimitiveSceneInfo = PrimitiveSceneProxy->GetPrimitiveSceneInfo();
+					
+										// Real time sky light capture cannot render dynamic meshes for now.
+										// For those to be rendered we would need to specify a view to the PassMeshProcessor creation above.
+										// Dynamic draws uses temporary per frame & per view data (appended at the end of the GPUScene buffer).
+										// But the view is transient and data on it can morph, and correct data would need to be added to FGPUScenePrimitiveCollector (see UploadDynamicPrimitiveShaderDataForViewInternal)
+										bool bSkipDynamicMesh = false;
+										for (auto& Element : MeshBatch->Elements)
+										{
+											if (Element.PrimitiveIdMode == PrimID_DynamicPrimitiveShaderData)
+											{
+												bSkipDynamicMesh = true;
+											}
+										}
+										if (bSkipDynamicMesh)
+										{
+											continue;
+										}
+					
+										const uint64 DefaultBatchElementMask = ~0ull;
+										PassMeshProcessor.AddMeshBatch(*MeshBatch, DefaultBatchElementMask, PrimitiveSceneProxy);
 									}
-
-									const uint64 DefaultBatchElementMask = ~0ull;
-									PassMeshProcessor.AddMeshBatch(*MeshBatch, DefaultBatchElementMask, PrimitiveSceneProxy);
+								});
 								}
-							});
 						}
-						else
+						else if(!bIsMobilePlatform)	// The SkyAtmosphere full screen pass shader is not compiled on mobile, a sky dome is required.
 						{
+							// TODO: mobile should never get there
 							RDG_EVENT_SCOPE(GraphBuilder, "Capture Sky Raw", CubeFace);
 							FSceneTextureShaderParameters SceneTextures = CreateSceneTextureShaderParameters(GraphBuilder, &SceneRenderer.GetActiveSceneTextures(), SceneRenderer.FeatureLevel, ESceneTextureSetupMode::SceneDepth);
 							SceneRenderer.RenderSkyAtmosphereInternal(GraphBuilder, SceneTextures, SkyRC);
@@ -906,7 +991,7 @@ void FScene::AllocateAndCaptureFrameSkyEnvMap(
 		// ForceImmediateFirstBarrier is required because the RHI resource is used as an SRV outside of RDG prior to this UAV pass. Without
 		// the flag, RDG will split the transition to UAV to the start of the graph, which results in a validation error. With the flag, RDG
 		// will transition to UAV at the start of the pass instead.
-		FRDGBuffer* SkyIrradianceEnvironmentMapRDG = GraphBuilder.RegisterExternalBuffer(SkyIrradianceEnvironmentMap, ERDGBufferFlags::ForceImmediateFirstBarrier);
+		FRDGBuffer* SkyIrradianceEnvironmentMapRDG = GraphBuilder.RegisterExternalBuffer(SkyIrradianceEnvironmentMap, ERDGBufferFlags::ForceImmediateFirstBarrier); // TODO SkyIrradianceEnvironmentMap is nullptr
 		GraphBuilder.UseInternalAccessMode(SkyIrradianceEnvironmentMapRDG);
 
 		TShaderMapRef<FComputeSkyEnvMapDiffuseIrradianceCS> ComputeShader(GetGlobalShaderMap(FeatureLevel));
@@ -929,9 +1014,57 @@ void FScene::AllocateAndCaptureFrameSkyEnvMap(
 		const FIntVector NumGroups = FIntVector(1, 1, 1);
 		FComputeShaderUtils::AddPass(GraphBuilder, RDG_EVENT_NAME("ComputeSkyEnvMapDiffuseIrradianceCS"), ComputeShader, PassParameters, NumGroups);
 
+		const bool bIsMobilePlatform = IsMobilePlatform(MainView.GetShaderPlatform());
+		if (bIsMobilePlatform)
+		{
+			TSharedPtr<FRHIGPUBufferReadback> NewReadBack = MakeShared<FRHIGPUBufferReadback>(TEXT("MobileSkyLightRealTimeCaptureIrradianceReadBack"));
+			AddEnqueueCopyPass(GraphBuilder, NewReadBack.Get(), SkyIrradianceEnvironmentMapRDG, SkyIrradianceEnvironmentMapRDG->GetSize());
+
+			FScene* Scene = MainView.Family->Scene->GetRenderScene();
+			Scene->MobileSkyLightRealTimeCaptureIrradianceReadBackQueries.Enqueue(NewReadBack);
+		}
+
 		ExternalAccessQueue.Add(SkyIrradianceEnvironmentMapRDG, ERHIAccess::SRVMask, ERHIPipeline::All);
 	};
 
+
+	auto Mobile_ReadBackSkyIrradianceEnvironmentMap = [&]()
+	{
+		const bool bIsMobilePlatform = IsMobilePlatform(MainView.GetShaderPlatform());
+		if (bIsMobilePlatform)
+		{
+			FScene* Scene = MainView.Family->Scene->GetRenderScene();
+
+			// Add a readback real tiem capture irradiance buffer
+
+			TSharedPtr<FRHIGPUBufferReadback> AvailableReadback;
+			Scene->MobileSkyLightRealTimeCaptureIrradianceReadBackQueries.Peek(AvailableReadback);
+			if (AvailableReadback.IsValid() && AvailableReadback->IsReady())
+			{
+				// Update the irradiance as soon as possible
+				Scene->MobileSkyLightRealTimeCaptureIrradianceReadBackQueries.Dequeue(AvailableReadback);
+
+				// Access the data and copy to a frame transient buffer for rendering pass.
+				uint64 SizeBytes = AvailableReadback->GetGPUSizeBytes();
+				void* SrcSkyIrradianceEnvironmentMap = AvailableReadback->Lock(SizeBytes);
+				if (SrcSkyIrradianceEnvironmentMap == nullptr)
+				{
+					// Reset to 0 in case this happens...
+					//SkyLight->IrradianceEnvironmentMap.R = TSHVector<3>();
+					//SkyLight->IrradianceEnvironmentMap.G = TSHVector<3>();
+					//SkyLight->IrradianceEnvironmentMap.B = TSHVector<3>();
+					// Keep the last value
+				}
+				else
+				{
+					check(sizeof(Scene->MobileSkyLightRealTimeCaptureIrradianceEnvironmentMap) <= SizeBytes);
+					memcpy(Scene->MobileSkyLightRealTimeCaptureIrradianceEnvironmentMap, SrcSkyIrradianceEnvironmentMap, sizeof(Scene->MobileSkyLightRealTimeCaptureIrradianceEnvironmentMap));
+				}
+			}
+			// else, we keep the last value
+		}
+	};
+
 	const uint32 LastMipLevel = CubeMipCount - 1;
 
 	// Ensure the main view got the full cubemap by running all the capture operations for the first frame.
@@ -991,6 +1124,8 @@ void FScene::AllocateAndCaptureFrameSkyEnvMap(
 		// 0.015ms
 		RenderCubeFaces_DiffuseIrradiance(ConvolvedSkyRenderTarget[ConvolvedSkyRenderTargetReadyIndex]);
 
+		Mobile_ReadBackSkyIrradianceEnvironmentMap();
+
 		// Reset Scene time slicing state so that it starts from the beginning if/when we get out of non-time-sliced.
 		Capture.State = -1; // Value of -1 indicates this is the first time-sliced iteration.
 		Capture.StateSubStep = 0;
@@ -1143,6 +1278,8 @@ void FScene::AllocateAndCaptureFrameSkyEnvMap(
 			RenderCubeFaces_DiffuseIrradiance(ConvolvedSkyRenderTarget[ConvolvedSkyRenderTargetWorkIndex]);
 		}
 
+		Mobile_ReadBackSkyIrradianceEnvironmentMap();
+
 #if DEBUG_TIME_SLICE
 		}
 		ConvolvedSkyRenderTargetReadyIndex = 1 - ConvolvedSkyRenderTargetReadyIndex;
diff --git a/Engine/Source/Runtime/Renderer/Private/RendererScene.cpp b/Engine/Source/Runtime/Renderer/Private/RendererScene.cpp
--- a/Engine/Source/Runtime/Renderer/Private/RendererScene.cpp
+++ b/Engine/Source/Runtime/Renderer/Private/RendererScene.cpp
@@ -1602,6 +1602,7 @@ FScene::FScene(UWorld* InWorld, bool bInRequiresHitProxies, bool bInIsEditorScen
 ,	World(InWorld)
 ,	FXSystem(nullptr)
 ,	bCachedShouldRenderSkylightInBasePass(false)
+,	bCachedSkyLightRealTimeCapture(false)
 ,	bScenesPrimitivesNeedStaticMeshElementUpdate(false)
 ,   PathTracingInvalidationCounter(0)
 #if RHI_RAYTRACING
@@ -1648,6 +1649,7 @@ FScene::FScene(UWorld* InWorld, bool bInRequiresHitProxies, bool bInIsEditorScen
 {
 	FMemory::Memzero(MobileDirectionalLights);
 	FMemory::Memzero(AtmosphereLights);
+	FMemory::Memzero(MobileSkyLightRealTimeCaptureIrradianceEnvironmentMap);
 
 	FullWorldName = World->GetFullName();
 
@@ -1771,6 +1773,12 @@ FScene::~FScene()
 #endif // RHI_RAYTRACING
 
 	delete SceneLightInfoUpdates;
+
+	TSharedPtr<FRHIGPUBufferReadback> It;
+	while (MobileSkyLightRealTimeCaptureIrradianceReadBackQueries.Dequeue(It))
+	{
+		It.Reset();
+	}
 }
 
 // Helpers for internal templates
@@ -5478,6 +5486,12 @@ void FScene::UpdateAllPrimitiveSceneInfos(FRDGBuilder& GraphBuilder, EUpdateAllP
 	Update(GraphBuilder, Parameters);
 }
 
+bool FScene::CanSampleSkyLightRealTimeCaptureData() const
+{
+	// We need a sky light, with bRealTimeCaptureEnabled (only true if supported by the platform settings) and if the captured data is ready.
+	return SkyLight && SkyLight->bRealTimeCaptureEnabled && ConvolvedSkyRenderTargetReadyIndex >= 0;
+}
+
 class FSceneComputeUpdates : public ISceneComputeUpdates
 {
 public:
@@ -6549,6 +6563,14 @@ void FScene::Update(FRDGBuilder& GraphBuilder, const FUpdateParameters& Paramete
 			bScenesPrimitivesNeedStaticMeshElementUpdate = true;
 		}
 
+		const bool bLastFrameSkyLightRealTimeCapture = bCachedSkyLightRealTimeCapture;
+		bCachedSkyLightRealTimeCapture = CanSampleSkyLightRealTimeCaptureData();
+		if (bCachedSkyLightRealTimeCapture != bLastFrameSkyLightRealTimeCapture)
+		{
+			// When real time capture on the sky light is toggled, we also need to recreate draw list for parameters.
+			UE_CLOG(!GIsEditor, LogRenderer, Log, TEXT("Forcing update for all mesh draw commands: SkyLight real-time capture change"));
+			bScenesPrimitivesNeedStaticMeshElementUpdate = true;
+		}
 
 		bool bVariableRateShadingToggled = GVRSImageManager.GetNeedStaticMeshUpdate();
 
diff --git a/Engine/Source/Runtime/Renderer/Private/ScenePrivate.h b/Engine/Source/Runtime/Renderer/Private/ScenePrivate.h
--- a/Engine/Source/Runtime/Renderer/Private/ScenePrivate.h
+++ b/Engine/Source/Runtime/Renderer/Private/ScenePrivate.h
@@ -3071,6 +3071,9 @@ public:
 	/** Previous frame SkyLight state. */
 	bool bCachedShouldRenderSkylightInBasePass;
 
+	/** Previous frame was using skylight realtime capture or not */
+	bool bCachedSkyLightRealTimeCapture;
+
 	/** True if a change to SkyLight / Lighting has occurred that requires static draw lists to be updated. */
 	bool bScenesPrimitivesNeedStaticMeshElementUpdate;
 
@@ -3103,6 +3106,9 @@ public:
 	/** The index of the ConvolvedSkyRenderTarget to use when rendering meshes. -1 when not initialised. */
 	int32 ConvolvedSkyRenderTargetReadyIndex;
 
+	/** Return true if the real time sky light captured data is ready to be sampled */
+	bool CanSampleSkyLightRealTimeCaptureData() const;
+
 	struct FRealTimeSlicedReflectionCapture
 	{
 		/** We always enforce a complete one the first frame even with time slicing for correct start up lighting.*/
@@ -3143,6 +3149,11 @@ public:
 	/** Used to track the order that skylights were enabled in. */
 	TArray<FSkyLightSceneProxy*> SkyLightStack;
 
+	/** Mobile renderer: used to readback the sky irradiance in order to set it onto the view constant UB */
+	TQueue<TSharedPtr<FRHIGPUBufferReadback>> MobileSkyLightRealTimeCaptureIrradianceReadBackQueries;
+	/** Mobile renderer: store a ready-to-copy to view UB sky irradiance environment map */
+	FVector4f MobileSkyLightRealTimeCaptureIrradianceEnvironmentMap[SKY_IRRADIANCE_ENVIRONMENT_MAP_VEC4_COUNT];
+
 	/** The directional light to use for simple dynamic lighting, if any. */
 	FLightSceneInfo* SimpleDirectionalLight;
 
diff --git a/Engine/Source/Runtime/Renderer/Private/SceneRendering.cpp b/Engine/Source/Runtime/Renderer/Private/SceneRendering.cpp
--- a/Engine/Source/Runtime/Renderer/Private/SceneRendering.cpp
+++ b/Engine/Source/Runtime/Renderer/Private/SceneRendering.cpp
@@ -1897,7 +1897,13 @@ void FViewInfo::SetupUniformBufferParameters(
 			&& (!Scene->SkyLight->bHasStaticLighting || !IsStaticLightingAllowed())
 			&& Family->EngineShowFlags.SkyLighting;
 
-		if (bSetupSkyIrradiance)
+		const bool bMobileRealTimeSkyLightCapture = Scene && Scene->CanSampleSkyLightRealTimeCaptureData() && Family->EngineShowFlags.SkyLighting;
+
+		if (bMobileRealTimeSkyLightCapture)
+		{
+			FMemory::Memcpy((void*)&ViewUniformShaderParameters.MobileSkyIrradianceEnvironmentMap, (void*)&Scene->MobileSkyLightRealTimeCaptureIrradianceEnvironmentMap, sizeof(Scene->MobileSkyLightRealTimeCaptureIrradianceEnvironmentMap));
+		}
+		else if (bSetupSkyIrradiance)
 		{
 			const FSHVectorRGB3& SkyIrradiance = Scene->SkyLight->IrradianceEnvironmentMap;
 			SetupSkyIrradianceEnvironmentMapConstantsFromSkyIrradiance((FVector4f*)&ViewUniformShaderParameters.MobileSkyIrradianceEnvironmentMap, SkyIrradiance);
