[HWRT] Decouple RenderRayTracingDebug/RenderRayTracingBarycentrics/RenderRayTracingPrimaryRaysView from FDeferredShadingSceneRenderer.
- makes it easier to implement Overview visualization mode in follow up CL.

diff --git a/Engine/Source/Runtime/Renderer/Private/DeferredShadingRenderer.cpp b/Engine/Source/Runtime/Renderer/Private/DeferredShadingRenderer.cpp
--- a/Engine/Source/Runtime/Renderer/Private/DeferredShadingRenderer.cpp
+++ b/Engine/Source/Runtime/Renderer/Private/DeferredShadingRenderer.cpp
@@ -1408,6 +1408,10 @@ void FDeferredShadingSceneRenderer::RenderNanite(FRDGBuilder& GraphBuilder, cons
 	}
 }
 
+#if RHI_RAYTRACING
+extern void RenderRayTracingDebug(FRDGBuilder& GraphBuilder, const FScene& Scene, const FViewInfo& View, FSceneTextures& SceneTextures, FRayTracingPickingFeedback& PickingFeedback);
+#endif // RHI_RAYTRACING
+
 void FDeferredShadingSceneRenderer::Render(FRDGBuilder& GraphBuilder, const FSceneRenderUpdateInputs* SceneUpdateInputs)
 {
 	if (!ViewFamily.EngineShowFlags.Rendering)
@@ -3327,7 +3331,7 @@ void FDeferredShadingSceneRenderer::Render(FRDGBuilder& GraphBuilder, const FSce
 				for (const FViewInfo& View : Views)
 				{
 					FRayTracingPickingFeedback PickingFeedback = {};
-					RenderRayTracingDebug(GraphBuilder, View, SceneTextures, PickingFeedback);
+					RenderRayTracingDebug(GraphBuilder, *Scene, View, SceneTextures, PickingFeedback);
 
 					OnGetOnScreenMessages.AddLambda([this, PickingFeedback](FScreenMessageWriter& ScreenMessageWriter)->void
 						{
diff --git a/Engine/Source/Runtime/Renderer/Private/DeferredShadingRenderer.h b/Engine/Source/Runtime/Renderer/Private/DeferredShadingRenderer.h
--- a/Engine/Source/Runtime/Renderer/Private/DeferredShadingRenderer.h
+++ b/Engine/Source/Runtime/Renderer/Private/DeferredShadingRenderer.h
@@ -1095,16 +1095,6 @@ private:
 		FRDGTextureRef& OutSkyLightTexture,
 		FRDGTextureRef& OutHitDistanceTexture);
 
-	void RenderRayTracingPrimaryRaysView(
-		FRDGBuilder& GraphBuilder,
-		const FViewInfo& View,
-		FRDGTextureRef* InOutColorTexture,
-		FRDGTextureRef* InOutRayHitDistanceTexture,
-		int32 SamplePerPixel,
-		int32 HeightFog,
-		float ResolutionFraction,
-		ERayTracingPrimaryRaysFlag Flags);
-
 	void RenderRayTracingTranslucency(FRDGBuilder& GraphBuilder, FRDGTextureMSAA SceneColorTexture);
 
 	void RenderRayTracingTranslucencyView(
@@ -1138,8 +1128,6 @@ private:
 	void SetupRayTracingRenderingData(FRDGBuilder& GraphBuilder);
 
 	/** Debug ray tracing functions. */
-	void RenderRayTracingDebug(FRDGBuilder& GraphBuilder, const FViewInfo& View, FSceneTextures& SceneTextures, FRayTracingPickingFeedback& PickingFeedback);
-	void RenderRayTracingBarycentrics(FRDGBuilder& GraphBuilder, const FViewInfo& View, FRDGTextureRef SceneColorOutputTexture, bool bVisualizeProceduralPrimitives);
 	void RayTracingDisplayPicking(const FRayTracingPickingFeedback& PickingFeedback, FScreenMessageWriter& Writer);
 
 	bool SetupRayTracingPipelineStatesAndSBT(FRDGBuilder& GraphBuilder, bool bAnyLumenHardwareInlineRayTracingPassEnabled, bool& bOutIsUsingFallbackRTPSO);
diff --git a/Engine/Source/Runtime/Renderer/Private/RayTracing/RayTracingBarycentrics.cpp b/Engine/Source/Runtime/Renderer/Private/RayTracing/RayTracingBarycentrics.cpp
--- a/Engine/Source/Runtime/Renderer/Private/RayTracing/RayTracingBarycentrics.cpp
+++ b/Engine/Source/Runtime/Renderer/Private/RayTracing/RayTracingBarycentrics.cpp
@@ -121,7 +121,7 @@ class FRayTracingBarycentricsCS : public FGlobalShader
 };
 IMPLEMENT_GLOBAL_SHADER(FRayTracingBarycentricsCS, "/Engine/Private/RayTracing/RayTracingBarycentrics.usf", "RayTracingBarycentricsMainCS", SF_Compute);
 
-void RenderRayTracingBarycentricsCS(FRDGBuilder& GraphBuilder, const FScene& Scene, const FViewInfo& View, FRDGTextureRef SceneColor, bool bVisualizeProceduralPrimitives)
+void RenderRayTracingBarycentricsCS(FRDGBuilder& GraphBuilder, const FViewInfo& View, FRDGTextureRef SceneColor, bool bVisualizeProceduralPrimitives)
 {
 	FRayTracingBarycentricsCS::FParameters* PassParameters = GraphBuilder.AllocParameters<FRayTracingBarycentricsCS::FParameters>();
 
@@ -204,18 +204,18 @@ void RenderRayTracingBarycentricsRGS(FRDGBuilder& GraphBuilder, const FScene& Sc
 	});
 }
 
-void FDeferredShadingSceneRenderer::RenderRayTracingBarycentrics(FRDGBuilder& GraphBuilder, const FViewInfo& View, FRDGTextureRef SceneColor, bool bVisualizeProceduralPrimitives)
+void RenderRayTracingBarycentrics(FRDGBuilder& GraphBuilder, const FScene& Scene, const FViewInfo& View, FRDGTextureRef SceneColor, bool bVisualizeProceduralPrimitives)
 {
 	const bool bRayTracingInline = ShouldRenderRayTracingEffect(true, ERayTracingPipelineCompatibilityFlags::Inline, View);
 	const bool bRayTracingPipeline = ShouldRenderRayTracingEffect(true, ERayTracingPipelineCompatibilityFlags::FullPipeline, View);
 
 	if (bRayTracingInline)
 	{
-		RenderRayTracingBarycentricsCS(GraphBuilder, *Scene, View, SceneColor, bVisualizeProceduralPrimitives);
+		RenderRayTracingBarycentricsCS(GraphBuilder, View, SceneColor, bVisualizeProceduralPrimitives);
 	}
 	else if (bRayTracingPipeline)
 	{
-		RenderRayTracingBarycentricsRGS(GraphBuilder, *Scene, View, SceneColor);
+		RenderRayTracingBarycentricsRGS(GraphBuilder, Scene, View, SceneColor);
 	}
 }
 #endif
diff --git a/Engine/Source/Runtime/Renderer/Private/RayTracing/RayTracingDebug.cpp b/Engine/Source/Runtime/Renderer/Private/RayTracing/RayTracingDebug.cpp
--- a/Engine/Source/Runtime/Renderer/Private/RayTracing/RayTracingDebug.cpp
+++ b/Engine/Source/Runtime/Renderer/Private/RayTracing/RayTracingDebug.cpp
@@ -857,9 +857,9 @@ void FDeferredShadingSceneRenderer::PrepareRayTracingDebug(const FSceneViewFamil
 	}
 }
 
-static FRDGBufferRef RayTracingPerformPicking(FRDGBuilder& GraphBuilder, const FScene* Scene, const FViewInfo& View, FRayTracingPickingFeedback& PickingFeedback)
+static FRDGBufferRef RayTracingPerformPicking(FRDGBuilder& GraphBuilder, const FScene& Scene, const FViewInfo& View, FRayTracingPickingFeedback& PickingFeedback)
 {
-	const FRayTracingScene& RayTracingScene = Scene->RayTracingScene;
+	const FRayTracingScene& RayTracingScene = Scene.RayTracingScene;
 
 	FGlobalShaderMap* ShaderMap = GetGlobalShaderMap(GMaxRHIFeatureLevel);
 	auto RayGenShader = ShaderMap->GetShader<FRayTracingPickingRGS>();
@@ -867,7 +867,7 @@ static FRDGBufferRef RayTracingPerformPicking(FRDGBuilder& GraphBuilder, const F
 	FRayTracingPipelineStateInitializer Initializer;
 	Initializer.MaxPayloadSizeInBytes = GetRayTracingPayloadTypeMaxSize(ERayTracingPayloadType::RayTracingDebug);
 
-	const FShaderBindingLayout* ShaderBindingLayout = RayTracing::GetShaderBindingLayout(Scene->GetShaderPlatform());
+	const FShaderBindingLayout* ShaderBindingLayout = RayTracing::GetShaderBindingLayout(Scene.GetShaderPlatform());
 	if (ShaderBindingLayout)
 	{
 		Initializer.ShaderBindingLayout = &ShaderBindingLayout->RHILayout;
@@ -893,7 +893,7 @@ static FRDGBufferRef RayTracingPerformPicking(FRDGBuilder& GraphBuilder, const F
 
 	FRayTracingPipelineState* PickingPipeline = PipelineStateCache::GetAndOrCreateRayTracingPipelineState(GraphBuilder.RHICmdList, Initializer);
 		
-	FShaderBindingTableRHIRef PickingSBT = Scene->RayTracingSBT.AllocateTransientRHI(GraphBuilder.RHICmdList, ERayTracingShaderBindingMode::RTPSO, ERayTracingHitGroupIndexingMode::Allow, RayTracingScene.NumMissShaderSlots, RayTracingScene.NumCallableShaderSlots, Initializer.GetMaxLocalBindingDataSize());
+	FShaderBindingTableRHIRef PickingSBT = Scene.RayTracingSBT.AllocateTransientRHI(GraphBuilder.RHICmdList, ERayTracingShaderBindingMode::RTPSO, ERayTracingHitGroupIndexingMode::Allow, RayTracingScene.NumMissShaderSlots, RayTracingScene.NumCallableShaderSlots, Initializer.GetMaxLocalBindingDataSize());
 	   
 	FRDGBufferDesc PickingBufferDesc = FRDGBufferDesc::CreateStructuredDesc(sizeof(FRayTracingPickingFeedback), 1);
 	PickingBufferDesc.Usage = EBufferUsageFlags(PickingBufferDesc.Usage | BUF_SourceCopy);
@@ -1023,7 +1023,7 @@ public:
 };
 IMPLEMENT_GLOBAL_SHADER(FRayTracingSceneHitStatsDebugRenderCS, "/Engine/Private/RayTracing/RayTracingDebugHitStatsUtils.usf", "RayTracingSceneDebugHitStatsRenderCS", SF_Compute);
 
-static void PrintTopKMostHitMessage(FRDGBuilder& GraphBuilder, const FScene* Scene, const FViewInfo& View, const TArray<FRayTracingHitStatsEntry>& HitStatsArray)
+static void PrintTopKMostHitMessage(FRDGBuilder& GraphBuilder, const FScene& Scene, const FViewInfo& View, const TArray<FRayTracingHitStatsEntry>& HitStatsArray)
 {
 	// Force ShaderPrint on.
 	ShaderPrint::SetEnabled(true);
@@ -1049,7 +1049,7 @@ static void PrintTopKMostHitMessage(FRDGBuilder& GraphBuilder, const FScene* Sce
 			FPersistentPrimitiveIndex PersistentPrimitiveIndex;
 			PersistentPrimitiveIndex.Index = PrimitiveID;
 
-			FPrimitiveSceneInfo* SceneInfo = Scene->GetPrimitiveSceneInfo(PersistentPrimitiveIndex);
+			FPrimitiveSceneInfo* SceneInfo = Scene.GetPrimitiveSceneInfo(PersistentPrimitiveIndex);
 
 			if (SceneInfo == nullptr)
 			{
@@ -1112,9 +1112,9 @@ static void PrintTopKMostHitMessage(FRDGBuilder& GraphBuilder, const FScene* Sce
 	}
 }
 
-static FRDGBufferRef RayTracingPerformHitStatsPerPrimitive(FRDGBuilder& GraphBuilder, const FScene* Scene, const FViewInfo& View)
+static FRDGBufferRef RayTracingPerformHitStatsPerPrimitive(FRDGBuilder& GraphBuilder, const FScene& Scene, const FViewInfo& View)
 {
-	const FRayTracingScene& RayTracingScene = Scene->RayTracingScene;
+	const FRayTracingScene& RayTracingScene = Scene.RayTracingScene;
 
 	FGlobalShaderMap* ShaderMap = GetGlobalShaderMap(GMaxRHIFeatureLevel);
 	auto RayGenShader = ShaderMap->GetShader<FRayTracingDebugHitStatsRGS>();
@@ -1122,7 +1122,7 @@ static FRDGBufferRef RayTracingPerformHitStatsPerPrimitive(FRDGBuilder& GraphBui
 	FRayTracingPipelineStateInitializer Initializer;
 	Initializer.MaxPayloadSizeInBytes = GetRayTracingPayloadTypeMaxSize(ERayTracingPayloadType::RayTracingDebug);
 	
-	const FShaderBindingLayout* ShaderBindingLayout = RayTracing::GetShaderBindingLayout(Scene->GetShaderPlatform());
+	const FShaderBindingLayout* ShaderBindingLayout = RayTracing::GetShaderBindingLayout(Scene.GetShaderPlatform());
 	if (ShaderBindingLayout)
 	{
 		Initializer.ShaderBindingLayout = &ShaderBindingLayout->RHILayout;
@@ -1144,7 +1144,7 @@ static FRDGBufferRef RayTracingPerformHitStatsPerPrimitive(FRDGBuilder& GraphBui
 
 	FRayTracingPipelineState* HitStatsPerPrimitivePipeline = PipelineStateCache::GetAndOrCreateRayTracingPipelineState(GraphBuilder.RHICmdList, Initializer);
 
-	FShaderBindingTableRHIRef HitStatsSBT = Scene->RayTracingSBT.AllocateTransientRHI(GraphBuilder.RHICmdList, ERayTracingShaderBindingMode::RTPSO, ERayTracingHitGroupIndexingMode::Allow, RayTracingScene.NumMissShaderSlots, RayTracingScene.NumCallableShaderSlots, Initializer.GetMaxLocalBindingDataSize());
+	FShaderBindingTableRHIRef HitStatsSBT = Scene.RayTracingSBT.AllocateTransientRHI(GraphBuilder.RHICmdList, ERayTracingShaderBindingMode::RTPSO, ERayTracingHitGroupIndexingMode::Allow, RayTracingScene.NumMissShaderSlots, RayTracingScene.NumCallableShaderSlots, Initializer.GetMaxLocalBindingDataSize());
 
 	// TODO: Should check RayTracingScene for actual number of instances instead of max number in FRHIRayTracingScene initializer
 	const uint32 NumInstancesInTLAS = FMath::Max(RayTracingScene.GetRHIRayTracingSceneChecked(ERayTracingSceneLayer::Base)->GetInitializer().MaxNumInstances, (uint32)CVarRayTracingVisualizeHitCountTopKHits.GetValueOnRenderThread());
@@ -1309,7 +1309,7 @@ static void RayTracingDrawInstances(FRDGBuilder& GraphBuilder, const FViewInfo&
 		});
 }
 
-static void DrawInstanceOverlap(FRDGBuilder& GraphBuilder, const FScene* Scene, const FViewInfo& View, FRDGTextureRef SceneColorTexture, FRDGTextureRef InputDepthTexture)
+static void DrawInstanceOverlap(FRDGBuilder& GraphBuilder, const FScene& Scene, const FViewInfo& View, FRDGTextureRef SceneColorTexture, FRDGTextureRef InputDepthTexture)
 {
 	FRDGTextureRef SceneDepthTexture = GraphBuilder.CreateTexture(
 		FRDGTextureDesc::Create2D(
@@ -1344,7 +1344,7 @@ static void DrawInstanceOverlap(FRDGBuilder& GraphBuilder, const FScene* Scene,
 	FRDGTextureDesc InstanceOverlapTextureDesc = FRDGTextureDesc::Create2D(SceneColorTexture->Desc.Extent, PF_R32_FLOAT, FClearValueBinding::Black, TexCreate_ShaderResource | TexCreate_RenderTargetable);
 	FRDGTextureRef InstanceOverlapTexture = GraphBuilder.CreateTexture(InstanceOverlapTextureDesc, TEXT("RayTracingDebug::InstanceOverlap"));
 	
-	RayTracingDrawInstances(GraphBuilder, View, InstanceOverlapTexture, SceneDepthTexture, Scene->RayTracingScene.GetDebugInstanceGPUSceneIndexBuffer(ERayTracingSceneLayer::Base), false);
+	RayTracingDrawInstances(GraphBuilder, View, InstanceOverlapTexture, SceneDepthTexture, Scene.RayTracingScene.GetDebugInstanceGPUSceneIndexBuffer(ERayTracingSceneLayer::Base), false);
 
 	// Calculate heatmap of instance overlap and blend it on top of ray tracing debug output
 	{
@@ -1372,7 +1372,7 @@ static void DrawInstanceOverlap(FRDGBuilder& GraphBuilder, const FScene* Scene,
 	// Draw instance AABB with lines
 	if (CVarRayTracingVisualizeInstanceOverlapShowWireframe.GetValueOnRenderThread() != 0)
 	{
-		RayTracingDrawInstances(GraphBuilder, View, SceneColorTexture, SceneDepthTexture, Scene->RayTracingScene.GetDebugInstanceGPUSceneIndexBuffer(ERayTracingSceneLayer::Base), true);
+		RayTracingDrawInstances(GraphBuilder, View, SceneColorTexture, SceneDepthTexture, Scene.RayTracingScene.GetDebugInstanceGPUSceneIndexBuffer(ERayTracingSceneLayer::Base), true);
 	}
 }
 
@@ -1423,8 +1423,17 @@ bool HasRaytracingDebugViewModeRaytracedOverlay(const FSceneViewFamily& ViewFami
 	return !bAnySubstrate && !bAnyTraversalSecondary;
 }
 
-void FDeferredShadingSceneRenderer::RenderRayTracingDebug(FRDGBuilder& GraphBuilder, const FViewInfo& View, FSceneTextures& SceneTextures, FRayTracingPickingFeedback& PickingFeedback)
+extern void RenderRayTracingBarycentrics(FRDGBuilder& GraphBuilder, const FScene& Scene, const FViewInfo& View, FRDGTextureRef SceneColor, bool bVisualizeProceduralPrimitives);
+
+extern void RenderRayTracingPrimaryRaysView(
+	FRDGBuilder& GraphBuilder, const FViewInfo& View, const FSceneTextures& SceneTextures,
+	FRDGTextureRef* InOutColorTexture, FRDGTextureRef* InOutRayHitDistanceTexture,
+	int32 SamplePerPixel, int32 HeightFog, float ResolutionFraction, ERayTracingPrimaryRaysFlag Flags);
+
+void RenderRayTracingDebug(FRDGBuilder& GraphBuilder, const FScene& Scene, const FViewInfo& View, FSceneTextures& SceneTextures, FRayTracingPickingFeedback& PickingFeedback)
 {
+	const EShaderPlatform ShaderPlatform = Scene.GetShaderPlatform();
+
 	FSceneTextureParameters SceneTextureParameters = GetSceneTextureParameters(GraphBuilder, SceneTextures.UniformBuffer);
 	FRDGTextureRef SceneColorTexture = SceneTextures.Color.Target;
 
@@ -1439,10 +1448,10 @@ void FDeferredShadingSceneRenderer::RenderRayTracingDebug(FRDGBuilder& GraphBuil
 
 	if (DebugVisualizationMode == RAY_TRACING_DEBUG_VIZ_BARYCENTRICS)
 	{
-		return RenderRayTracingBarycentrics(GraphBuilder, View, SceneColorTexture, (bool)GVisualizeProceduralPrimitives);
+		return RenderRayTracingBarycentrics(GraphBuilder, Scene, View, SceneColorTexture, (bool)GVisualizeProceduralPrimitives);
 	}
 
-	FRayTracingScene& RayTracingScene = Scene->RayTracingScene;
+	const FRayTracingScene& RayTracingScene = Scene.RayTracingScene;
 
 	if (IsRayTracingDebugTraversalMode(DebugVisualizationMode) && ShouldRenderRayTracingEffect(true, ERayTracingPipelineCompatibilityFlags::Inline, View))
 	{
@@ -1453,7 +1462,7 @@ void FDeferredShadingSceneRenderer::RenderRayTracingDebug(FRDGBuilder& GraphBuil
 		PassParameters->Output = GraphBuilder.CreateUAV(SceneColorTexture);
 		PassParameters->TLAS = RayTracingScene.GetLayerView(ERayTracingSceneLayer::Base);
 		PassParameters->ViewUniformBuffer = View.ViewUniformBuffer;
-		PassParameters->Scene = GetSceneUniforms().GetBuffer(GraphBuilder);
+		PassParameters->Scene = View.GetSceneUniforms().GetBuffer(GraphBuilder);
 		PassParameters->NaniteRasterUniformBuffer = CreateDebugNaniteRasterUniformBuffer(GraphBuilder);
 		PassParameters->NaniteShadingUniformBuffer = CreateDebugNaniteShadingUniformBuffer(GraphBuilder);
 
@@ -1517,8 +1526,8 @@ void FDeferredShadingSceneRenderer::RenderRayTracingDebug(FRDGBuilder& GraphBuil
 		FRDGTextureRef HitDistanceTexture = nullptr;
 
 		RenderRayTracingPrimaryRaysView(
-				GraphBuilder, View, &OutputColor, &HitDistanceTexture, 1, 1, 1,
-			ERayTracingPrimaryRaysFlag::PrimaryView);
+			GraphBuilder, View, SceneTextures, 
+			&OutputColor, &HitDistanceTexture, 1, 1, 1, ERayTracingPrimaryRaysFlag::PrimaryView);
 
 		AddDrawTexturePass(GraphBuilder, View, OutputColor, SceneColorTexture, View.ViewRect.Min, View.ViewRect.Min, View.ViewRect.Size());
 		return;
@@ -1561,8 +1570,7 @@ void FDeferredShadingSceneRenderer::RenderRayTracingDebug(FRDGBuilder& GraphBuil
 	FRayTracingDebugRGS::FPermutationDomain PermutationVector;
 	PermutationVector.Set<FRayTracingDebugRGS::FUseDebugCHSType>(bRequiresDebugCHS);
 
-	FGlobalShaderMap* ShaderMap = GetGlobalShaderMap(FeatureLevel);
-	auto RayGenShader = ShaderMap->GetShader<FRayTracingDebugRGS>(PermutationVector);
+	auto RayGenShader = View.ShaderMap->GetShader<FRayTracingDebugRGS>(PermutationVector);
 
 	FRayTracingPipelineState* Pipeline = View.RayTracingMaterialPipeline;
 	FShaderBindingTableRHIRef SBT = View.RayTracingSBT;
@@ -1595,13 +1603,13 @@ void FDeferredShadingSceneRenderer::RenderRayTracingDebug(FRDGBuilder& GraphBuil
 		FRHIRayTracingShader* HitGroupTable[] = { HitGroupShader.GetRayTracingShader(), HitGroupShaderNaniteRT.GetRayTracingShader()};
 		Initializer.SetHitGroupTable(HitGroupTable);
 
-		auto MissShader = ShaderMap->GetShader<FRayTracingDebugMS>();
+		auto MissShader = View.ShaderMap->GetShader<FRayTracingDebugMS>();
 		FRHIRayTracingShader* MissTable[] = { MissShader.GetRayTracingShader() };
 		Initializer.SetMissShaderTable(MissTable);
 
 		Pipeline = PipelineStateCache::GetAndOrCreateRayTracingPipelineState(GraphBuilder.RHICmdList, Initializer);
 
-		SBT = Scene->RayTracingSBT.AllocateTransientRHI(GraphBuilder.RHICmdList, ERayTracingShaderBindingMode::RTPSO, ERayTracingHitGroupIndexingMode::Allow, RayTracingScene.NumMissShaderSlots, RayTracingScene.NumCallableShaderSlots, Initializer.GetMaxLocalBindingDataSize());
+		SBT = Scene.RayTracingSBT.AllocateTransientRHI(GraphBuilder.RHICmdList, ERayTracingShaderBindingMode::RTPSO, ERayTracingHitGroupIndexingMode::Allow, RayTracingScene.NumMissShaderSlots, RayTracingScene.NumCallableShaderSlots, Initializer.GetMaxLocalBindingDataSize());
 
 		bRequiresBindings = true;
 	}
@@ -1637,7 +1645,7 @@ void FDeferredShadingSceneRenderer::RenderRayTracingDebug(FRDGBuilder& GraphBuil
 	FRDGTextureRef OutputDepthTexture = GraphBuilder.CreateTexture(OutputDepthTextureDesc, TEXT("RayTracingDebug::Depth"));
 	RayGenParameters->OutputDepth = GraphBuilder.CreateUAV(OutputDepthTexture);
 
-	if (Lumen::UseFarField(ViewFamily) || MegaLights::UseFarField(ViewFamily))
+	if (Lumen::UseFarField(*View.Family) || MegaLights::UseFarField(*View.Family))
 	{
 		RayGenParameters->MaxTraceDistance = Lumen::GetMaxTraceDistance(View);
 		RayGenParameters->FarFieldMaxTraceDistance = Lumen::GetFarFieldMaxTraceDistance();
@@ -1657,7 +1665,7 @@ void FDeferredShadingSceneRenderer::RenderRayTracingDebug(FRDGBuilder& GraphBuil
 	RayGenParameters->Output = GraphBuilder.CreateUAV(SceneColorTexture);
 	RayGenParameters->TopKHitStats = GraphBuilder.CreateSRV(StatsBuffer);
 
-	RayGenParameters->SceneUniformBuffer = GetSceneUniformBufferRef(GraphBuilder); // TODO: use a separate params structure
+	RayGenParameters->SceneUniformBuffer = GetSceneUniformBufferRef(GraphBuilder, View); // TODO: use a separate params structure
 
 	FIntRect ViewRect = View.ViewRect;
 	
@@ -1668,7 +1676,7 @@ void FDeferredShadingSceneRenderer::RenderRayTracingDebug(FRDGBuilder& GraphBuil
 		RDG_EVENT_NAME("RayTracingDebug"),
 		RayGenParameters,
 		ERDGPassFlags::Compute,
-		[this, RayGenParameters, RayGenShader, &View, Pipeline, SBT, ViewRect, bRequiresBindings](FRDGAsyncTask, FRHICommandList& RHICmdList)
+		[RayGenParameters, RayGenShader, &View, Pipeline, SBT, ViewRect, bRequiresBindings](FRDGAsyncTask, FRHICommandList& RHICmdList)
 	{
 		FRHIBatchedShaderParameters& GlobalResources = RHICmdList.GetScratchShaderParameters();
 		SetShaderParameters(GlobalResources, RayGenShader, *RayGenParameters);
diff --git a/Engine/Source/Runtime/Renderer/Private/RayTracing/RayTracingPrimaryRays.cpp b/Engine/Source/Runtime/Renderer/Private/RayTracing/RayTracingPrimaryRays.cpp
--- a/Engine/Source/Runtime/Renderer/Private/RayTracing/RayTracingPrimaryRays.cpp
+++ b/Engine/Source/Runtime/Renderer/Private/RayTracing/RayTracingPrimaryRays.cpp
@@ -101,9 +101,10 @@ void FDeferredShadingSceneRenderer::PrepareRayTracingTranslucency(const FViewInf
 	}
 }
 
-void FDeferredShadingSceneRenderer::RenderRayTracingPrimaryRaysView(
+void RenderRayTracingPrimaryRaysView(
 	FRDGBuilder& GraphBuilder,
 	const FViewInfo& View,
+	const FSceneTextures& SceneTextures,
 	FRDGTextureRef* InOutColorTexture,
 	FRDGTextureRef* InOutRayHitDistanceTexture,
 	int32 SamplePerPixel,
@@ -111,7 +112,6 @@ void FDeferredShadingSceneRenderer::RenderRayTracingPrimaryRaysView(
 	float ResolutionFraction,
 	ERayTracingPrimaryRaysFlag Flags)
 {
-	const FSceneTextures& SceneTextures = GetActiveSceneTextures();
 	const FSceneTextureParameters SceneTextureParameters = GetSceneTextureParameters(GraphBuilder, SceneTextures);
 
 	int32 UpscaleFactor = int32(1.0f / ResolutionFraction);
@@ -129,7 +129,6 @@ void FDeferredShadingSceneRenderer::RenderRayTracingPrimaryRaysView(
 		if(*InOutColorTexture == nullptr) 
 		{
 			*InOutColorTexture = GraphBuilder.CreateTexture(Desc, TEXT("RayTracingPrimaryRays"));
-			
 		}
 
 		Desc.Format = PF_R16F;
@@ -166,7 +165,7 @@ void FDeferredShadingSceneRenderer::RenderRayTracingPrimaryRaysView(
 	PassParameters->LumenGIVolumeStruct = GraphBuilder.CreateUniformBuffer(LumenUniforms);
 
 	PassParameters->SceneTextures = SceneTextureParameters;
-	PassParameters->Scene = GetSceneUniformBufferRef(GraphBuilder);
+	PassParameters->Scene = GetSceneUniformBufferRef(GraphBuilder, View);
 
 	PassParameters->SceneColorTexture = GetIfProduced(SceneTextures.Color.Resolve, FRDGSystemTextures::Get(GraphBuilder).Black);
 
@@ -190,7 +189,7 @@ void FDeferredShadingSceneRenderer::RenderRayTracingPrimaryRaysView(
 		RDG_EVENT_NAME("RayTracingPrimaryRays %dx%d", RayTracingResolution.X, RayTracingResolution.Y),
 		PassParameters,
 		ERDGPassFlags::Compute,
-		[PassParameters, this, &View, RayGenShader, RayTracingResolution](FRDGAsyncTask, FRHICommandList& RHICmdList)
+		[PassParameters, &View, RayGenShader, RayTracingResolution](FRDGAsyncTask, FRHICommandList& RHICmdList)
 	{
 		FRHIShaderBindingTable* SBT = View.RayTracingSBT;
 		FRayTracingPipelineState* Pipeline = View.RayTracingMaterialPipeline;
diff --git a/Engine/Source/Runtime/Renderer/Private/RayTracing/RayTracingTranslucency.cpp b/Engine/Source/Runtime/Renderer/Private/RayTracing/RayTracingTranslucency.cpp
--- a/Engine/Source/Runtime/Renderer/Private/RayTracing/RayTracingTranslucency.cpp
+++ b/Engine/Source/Runtime/Renderer/Private/RayTracing/RayTracingTranslucency.cpp
@@ -142,6 +142,11 @@ bool ShouldRenderRayTracingTranslucency(const FViewInfo& View)
 }
 #endif // RHI_RAYTRACING
 
+extern void RenderRayTracingPrimaryRaysView(
+	FRDGBuilder& GraphBuilder, const FViewInfo& View, const FSceneTextures& SceneTextures,
+	FRDGTextureRef* InOutColorTexture, FRDGTextureRef* InOutRayHitDistanceTexture,
+	int32 SamplePerPixel, int32 HeightFog, float ResolutionFraction, ERayTracingPrimaryRaysFlag Flags);
+
 void FDeferredShadingSceneRenderer::RenderRayTracingTranslucency(FRDGBuilder& GraphBuilder, FRDGTextureMSAA SceneColorTexture)
 {
 	if (   !ShouldRenderTranslucency(ETranslucencyPass::TPT_TranslucencyStandard, Views)
@@ -178,8 +183,8 @@ void FDeferredShadingSceneRenderer::RenderRayTracingTranslucency(FRDGBuilder& Gr
 			int32 TranslucencySPP = GRayTracingTranslucencySamplesPerPixel > -1 ? GRayTracingTranslucencySamplesPerPixel : View.FinalPostProcessSettings.RayTracingTranslucencySamplesPerPixel;
 		
 			RenderRayTracingPrimaryRaysView(
-				GraphBuilder,
-				View, &DenoiserInputs.Color, &DenoiserInputs.RayHitDistance,
+				GraphBuilder, View, GetActiveSceneTextures(),
+				&DenoiserInputs.Color, &DenoiserInputs.RayHitDistance,
 				TranslucencySPP, GRayTracingTranslucencyHeightFog, ResolutionFraction, 
 				ERayTracingPrimaryRaysFlag::AllowSkipSkySample | ERayTracingPrimaryRaysFlag::UseGBufferForMaxDistance);
 
