Fixed crash when attempting to create uniform buffer off the render thread in OpenGL.

diff --git a/Engine/Source/Runtime/Renderer/Private/SceneRendering.cpp b/Engine/Source/Runtime/Renderer/Private/SceneRendering.cpp
--- a/Engine/Source/Runtime/Renderer/Private/SceneRendering.cpp
+++ b/Engine/Source/Runtime/Renderer/Private/SceneRendering.cpp
@@ -3805,6 +3805,8 @@ IVisibilityTaskData* FSceneRenderer::OnRenderBegin(FRDGBuilder& GraphBuilder, co
 	// This is called prior to scene update
 	OIT::OnRenderBegin(Scene->OITSceneData);
 
+	const bool bIsMobilePlatform = IsMobilePlatform(ShaderPlatform);
+
 	EUpdateAllPrimitiveSceneInfosAsyncOps AsyncOps = EUpdateAllPrimitiveSceneInfosAsyncOps::None;
 
 	if (GAsyncCreateLightPrimitiveInteractions > 0)
@@ -3817,7 +3819,7 @@ IVisibilityTaskData* FSceneRenderer::OnRenderBegin(FRDGBuilder& GraphBuilder, co
 		AsyncOps |= EUpdateAllPrimitiveSceneInfosAsyncOps::CacheMeshDrawCommands;
 	}
 
-	if (GAsyncCacheMaterialUniformExpressions > 0 && !IsMobilePlatform(ShaderPlatform))
+	if (GAsyncCacheMaterialUniformExpressions > 0 && !bIsMobilePlatform)
 	{
 		AsyncOps |= EUpdateAllPrimitiveSceneInfosAsyncOps::CacheMaterialUniformExpressions;
 	}
@@ -3855,7 +3857,8 @@ IVisibilityTaskData* FSceneRenderer::OnRenderBegin(FRDGBuilder& GraphBuilder, co
 	
 				Renderer->PrepareViewStateForVisibility(SceneTexturesConfig);
 			}
-		});
+
+		}, UE::Tasks::ETaskPriority::Normal, bIsMobilePlatform ? UE::Tasks::EExtendedTaskPriority::Inline : UE::Tasks::EExtendedTaskPriority::None);
 	}
 
 	SceneUpdateParameters.Callbacks.PostStaticMeshUpdate = [&] (const UE::Tasks::FTask& StaticMeshUpdateTask)
