Reinstated moving FPrevViewInfo to a reference while also fixing nulling out the first 8 bytes of it.

diff --git a/Engine/Source/Runtime/Renderer/Private/Renderer/ViewSnapshotCache.cpp b/Engine/Source/Runtime/Renderer/Private/Renderer/ViewSnapshotCache.cpp
--- a/Engine/Source/Runtime/Renderer/Private/Renderer/ViewSnapshotCache.cpp
+++ b/Engine/Source/Runtime/Renderer/Private/Renderer/ViewSnapshotCache.cpp
@@ -41,9 +41,6 @@ public:
 		TArray<FPrimitiveUniformShaderParameters> NullDynamicPrimitiveShaderData;
 		FMemory::Memset(&Result->ParallelMeshDrawCommandPasses, 0, sizeof(Result->ParallelMeshDrawCommandPasses));
 
-		// Null out the previous view info reference.
-		FMemory::Memset(&Result->PrevViewInfo, 0, sizeof(uintptr_t));
-
 		// Ensure the internal state is maintained, needed because we've just Memcpy'd the member data.
 		static_assert(std::is_trivially_destructible_v<FGPUScenePrimitiveCollector> != 0, "The destructor is not invoked properly because of FMemory::Memcpy(*Result, *this) above");
 		Result->DynamicPrimitiveCollector = FGPUScenePrimitiveCollector(InView->DynamicPrimitiveCollector);
diff --git a/Engine/Source/Runtime/Renderer/Private/SceneRendering.h b/Engine/Source/Runtime/Renderer/Private/SceneRendering.h
--- a/Engine/Source/Runtime/Renderer/Private/SceneRendering.h
+++ b/Engine/Source/Runtime/Renderer/Private/SceneRendering.h
@@ -1408,7 +1408,7 @@ public:
 	uint16 ShadingModelMaskInView;
 
 	/** Informations from the previous frame to use for this view. */
-	FPreviousViewInfo PrevViewInfo;
+	FPreviousViewInfo& PrevViewInfo = *Allocator.Create<FPreviousViewInfo>();
 
 	/** An intermediate number of visible static meshes.  Doesn't account for occlusion until after FinishOcclusionQueries is called. */
 	int32 NumVisibleStaticMeshElements;
