Real time sky capture - addeda cvar to enforce sky light capture resolution on some platforms.

#rnx
#tests pc dx12 editor and mobile console

diff --git a/Engine/Source/Runtime/Renderer/Private/ReflectionEnvironmentRealTimeCapture.cpp b/Engine/Source/Runtime/Renderer/Private/ReflectionEnvironmentRealTimeCapture.cpp
--- a/Engine/Source/Runtime/Renderer/Private/ReflectionEnvironmentRealTimeCapture.cpp
+++ b/Engine/Source/Runtime/Renderer/Private/ReflectionEnvironmentRealTimeCapture.cpp
@@ -53,6 +53,11 @@ static TAutoConsoleVariable<int32> CVarRealTimeReflectionCaptureDepthBuffer(
 	TEXT("When enabled, the real-time sky light capture will have a depth buffer, this is for multiple meshes to be cover each other correctly. The height fog will also be applied according to the depth buffer."),
 	ECVF_RenderThreadSafe);
 
+static TAutoConsoleVariable<int32> CVarRealTimeReflectionCaptureResolutionOverride(
+	TEXT("r.SkyLight.RealTimeReflectionCapture.ResolutionOverride"), 0,
+	TEXT("Override the real time sky light capture resolution if >=16."),
+	ECVF_ReadOnly | ECVF_RenderThreadSafe); // Read only because it is not really toggable mid-update when time slicing is used
+
 
 class FDownsampleCubeFaceCS : public FGlobalShader
 {
@@ -347,7 +352,8 @@ void FScene::AllocateAndCaptureFrameSkyEnvMap(
 	RDG_GPU_STAT_SCOPE(GraphBuilder, CaptureConvolveSkyEnvMap);
 	RDG_CSV_STAT_EXCLUSIVE_SCOPE(GraphBuilder, SkyAtmosphere);
 
-	const uint32 CubeWidth = SkyLight->CaptureCubeMapResolution;
+	const int32 CaptureResolutionOverride = CVarRealTimeReflectionCaptureResolutionOverride.GetValueOnRenderThread();
+	const uint32 CubeWidth = CaptureResolutionOverride >= 16 ? CaptureResolutionOverride : SkyLight->CaptureCubeMapResolution;
 	const uint32 CubeMipCount = FMath::CeilLogTwo(CubeWidth) + 1;
 
 	// Make a snapshot we are going to use for the 6 cubemap faces and set it up.
