Don't add TexCreate_Foveation to the shading rate textures on platforms which don't support hardware VRS.

RDG transitions foveation textures to ERHIAccess::ShadingRateSource after graph execution, which is invalid on platforms which do not support this, and some RHIs complain (e.g. DX12 says "Resource state is D3D12_RESOURCE_STATE_SHADING_RATE_SOURCE but RHI does not support VRS"). Shading rate images are used by [REDACTED_PERSON], but it samples them explicitly, so it doesn't need the foveation flag.

#rnx

diff --git a/Engine/Source/Runtime/Renderer/Private/VariableRateShading/VariableRateShadingImageManager.cpp b/Engine/Source/Runtime/Renderer/Private/VariableRateShading/VariableRateShadingImageManager.cpp
--- a/Engine/Source/Runtime/Renderer/Private/VariableRateShading/VariableRateShadingImageManager.cpp
+++ b/Engine/Source/Runtime/Renderer/Private/VariableRateShading/VariableRateShadingImageManager.cpp
@@ -441,11 +441,17 @@ FRDGTextureDesc FVariableRateShadingImageManager::GetSRIDesc(const FSceneViewFam
 	const FIntRect FamilyViewRect = ViewInfo->GetFamilyViewRect(); // May vary from the size of the scene textures if using constrained aspect ratios
 	const FIntPoint SRISize = FMath::DivideAndRoundUp(FamilyViewRect.Size(), GetSRITileSize(bSoftwareVRS));
 
+	ETextureCreateFlags Flags = TexCreate_UAV | TexCreate_ShaderResource | TexCreate_DisableDCC;
+	if (GRHIGlobals.VariableRateShading.SupportsAttachment)
+	{
+		Flags |= TexCreate_Foveation;
+	}
+
 	return FRDGTextureDesc::Create2D(
 		SRISize,
 		bSoftwareVRS ? PF_R8_UINT : GRHIVariableRateShadingImageFormat,
 		FClearValueBinding::None,
-		TexCreate_Foveation | TexCreate_UAV | TexCreate_ShaderResource | TexCreate_DisableDCC);
+		Flags);
 }
 
 int32 FVariableRateShadingImageManager::GetNumberOfSupportedRates()
