Mobile/Desktop PostFX parity. PostProcessMaterialAfterTonemapping after FXAA. Drive-by cleanup of unused bLastPass.

diff --git a/Engine/Source/Runtime/Renderer/Private/PostProcess/PostProcessing.cpp b/Engine/Source/Runtime/Renderer/Private/PostProcess/PostProcessing.cpp
--- a/Engine/Source/Runtime/Renderer/Private/PostProcess/PostProcessing.cpp
+++ b/Engine/Source/Runtime/Renderer/Private/PostProcess/PostProcessing.cpp
@@ -1656,10 +1656,10 @@ void AddPostProcessingPasses(
 	}
 	
 #if WITH_EDITOR || !UE_BUILD_SHIPPING
- 	if (EngineShowFlags.VisualizeNanite && NaniteRasterResults != nullptr)
- 	{
- 		AddVisualizeNanitePass(GraphBuilder, View, SceneColor, *NaniteRasterResults);
- 	}
+	if (EngineShowFlags.VisualizeNanite && NaniteRasterResults != nullptr)
+	{
+		AddVisualizeNanitePass(GraphBuilder, View, SceneColor, *NaniteRasterResults);
+	}
 #endif
 
 #if WITH_EDITOR
@@ -2391,8 +2391,8 @@ void AddMobilePostProcessingPasses(FRDGBuilder& GraphBuilder, FScene* Scene, con
 		SeparateTranslucency,
 		TAA,
 		Tonemap,
-		PostProcessMaterialAfterTonemapping,
 		FXAA,
+		PostProcessMaterialAfterTonemapping,
 		HighResolutionScreenshotMask,
 		SelectionOutline,
 		EditorPrimitive,
@@ -2498,7 +2498,7 @@ void AddMobilePostProcessingPasses(FRDGBuilder& GraphBuilder, FScene* Scene, con
 	// The scene color will be decoded at the first post-process material and output linear color space for the following passes
 	// bMetalMSAAHDRDecode will be set to false if there is any post-process material exist
 
-	auto AddPostProcessMaterialPass = [&GraphBuilder, &View, ViewIndex, &Inputs, &SceneColor, &CustomDepth, &bMetalMSAAHDRDecode, &PassSequence, &BloomOutput, &BlackAlphaOneDummy](EBlendableLocation BlendableLocation, bool bLastPass)
+	auto AddPostProcessMaterialPass = [&GraphBuilder, &View, ViewIndex, &Inputs, &SceneColor, &CustomDepth, &bMetalMSAAHDRDecode, &PassSequence, &BloomOutput, &BlackAlphaOneDummy](EBlendableLocation BlendableLocation)
 	{
 		FPostProcessMaterialInputs PostProcessMaterialInputs;
 
@@ -2570,8 +2570,8 @@ void AddMobilePostProcessingPasses(FRDGBuilder& GraphBuilder, FScene* Scene, con
 		PassSequence.SetEnabled(EPass::SunMerge, bUseBloom || bUseSun);
 		PassSequence.SetEnabled(EPass::SeparateTranslucency, bUseSeparateTranslucency);
 		PassSequence.SetEnabled(EPass::TAA, bUseTAA);
-		PassSequence.SetEnabled(EPass::PostProcessMaterialAfterTonemapping, PostProcessMaterialAfterTonemappingChain.Num() != 0);
 		PassSequence.SetEnabled(EPass::FXAA, View.AntiAliasingMethod == AAM_FXAA);
+		PassSequence.SetEnabled(EPass::PostProcessMaterialAfterTonemapping, PostProcessMaterialAfterTonemappingChain.Num() != 0);
 		PassSequence.Finalize();
 			
 		if (PassSequence.IsEnabled(EPass::Distortion))
@@ -2589,7 +2589,7 @@ void AddMobilePostProcessingPasses(FRDGBuilder& GraphBuilder, FScene* Scene, con
 			SceneColor = AddMobileDistortionMergePass(GraphBuilder, View, DistortionMergeInputs);
 		}
 
-		AddPostProcessMaterialPass(BL_SceneColorBeforeDOF, false);
+		AddPostProcessMaterialPass(BL_SceneColorBeforeDOF);
 
 		// Optional fixed pass processes
 		if (PassSequence.IsEnabled(EPass::SunMask))
@@ -2895,7 +2895,7 @@ void AddMobilePostProcessingPasses(FRDGBuilder& GraphBuilder, FScene* Scene, con
 			AddMobileSeparateTranslucencyPass(GraphBuilder, Scene, View, SeparateTranslucencyInputs);
 		}
 
-		AddPostProcessMaterialPass(BL_SceneColorAfterDOF, false);
+		AddPostProcessMaterialPass(BL_SceneColorAfterDOF);
 
 		// Temporal Anti-aliasing. Also may perform a temporal upsample from primary to secondary view rect.
 		if (PassSequence.IsEnabled(EPass::TAA))
@@ -2961,12 +2961,12 @@ void AddMobilePostProcessingPasses(FRDGBuilder& GraphBuilder, FScene* Scene, con
 		PassSequence.SetEnabled(EPass::SunMerge, false);
 		PassSequence.SetEnabled(EPass::SeparateTranslucency, false);
 		PassSequence.SetEnabled(EPass::TAA, false);
-		PassSequence.SetEnabled(EPass::PostProcessMaterialAfterTonemapping, false);
 		PassSequence.SetEnabled(EPass::FXAA, false);
+		PassSequence.SetEnabled(EPass::PostProcessMaterialAfterTonemapping, false);
 		PassSequence.Finalize();
 	}
 
-	AddPostProcessMaterialPass(BL_SceneColorBeforeBloom, false);
+	AddPostProcessMaterialPass(BL_SceneColorBeforeBloom);
 	
 	if (PassSequence.IsEnabled(EPass::Tonemap))
 	{
@@ -2974,7 +2974,7 @@ void AddMobilePostProcessingPasses(FRDGBuilder& GraphBuilder, FScene* Scene, con
 
 		if (MaterialChain.Num())
 		{
-			AddPostProcessMaterialPass(BL_ReplacingTonemapper, false);
+			AddPostProcessMaterialPass(BL_ReplacingTonemapper);
 		}
 		else
 		{
@@ -3045,11 +3045,6 @@ void AddMobilePostProcessingPasses(FRDGBuilder& GraphBuilder, FScene* Scene, con
 
 	if (IsPostProcessingEnabled(View))
 	{
-		if (PassSequence.IsEnabled(EPass::PostProcessMaterialAfterTonemapping))
-		{
-			AddPostProcessMaterialPass(BL_SceneColorAfterTonemapping, true);
-		}
-
 		if (PassSequence.IsEnabled(EPass::FXAA))
 		{
 			FFXAAInputs PassInputs;
@@ -3059,6 +3054,11 @@ void AddMobilePostProcessingPasses(FRDGBuilder& GraphBuilder, FScene* Scene, con
 
 			SceneColor = AddFXAAPass(GraphBuilder, View, PassInputs);
 		}
+
+		if (PassSequence.IsEnabled(EPass::PostProcessMaterialAfterTonemapping))
+		{
+			AddPostProcessMaterialPass(BL_SceneColorAfterTonemapping);
+		}
 	}
 
 	if (PassSequence.IsEnabled(EPass::HighResolutionScreenshotMask))
