* Add option to clear the rect atlas texture on change
* Add option to recreate atlas layout on change.
* Fix force update logic which would lead to incorrect layout for one frame.

diff --git a/Engine/Shaders/Private/RectLightAtlas.usf b/Engine/Shaders/Private/RectLightAtlas.usf
--- a/Engine/Shaders/Private/RectLightAtlas.usf
+++ b/Engine/Shaders/Private/RectLightAtlas.usf
@@ -231,6 +231,10 @@ uint  AtlasMIPIndex;
 uint  AtlasSourceTextureMIPBias;
 uint  AtlasFormat;
 
+uint  ForceUpdate;
+uint  ResetOnChange;
+uint  ClearAtlas;
+
 SamplerState AtlasSampler;
 Texture2D<float4> AtlasTexture;
 RWTexture2D<float4> OutputTexture;
@@ -297,6 +301,19 @@ void MainCS(uint3 DispatchThreadId : SV_DispatchThreadID)
 		else  
 		{ Print(Context, TEXT("Unknown"), FontOrange); }
 		Newline(Context);
+
+		Print(Context, TEXT("Force Update     : "), FontSilver);
+		PrintBool(Context, ForceUpdate);
+		Newline(Context);
+
+		Print(Context, TEXT("Reset On Change  : "), FontSilver);
+		PrintBool(Context, ResetOnChange);
+		Newline(Context);
+
+		Print(Context, TEXT("Clear Atlas      : "), FontSilver);
+		PrintBool(Context, ClearAtlas);
+		Newline(Context);
+
 	}
 
 	// Draw atlas
@@ -306,7 +323,7 @@ void MainCS(uint3 DispatchThreadId : SV_DispatchThreadID)
 	const float OutputRatio = OutputResolution.y / float(OutputResolution.x);
 
 	const float ScaleFactor = 0.25f;
-	const uint2 DisplayAtlasOffset     = uint2(50, 220);
+	const uint2 DisplayAtlasOffset     = uint2(50, 250);
 	const uint2 DisplayAtlasResolution = uint2(OutputResolution.x * ScaleFactor, OutputResolution.x * ScaleFactor * AtlasResolution.y / AtlasResolution.x);
 	if (all(Coord > DisplayAtlasOffset) && all(Coord < DisplayAtlasOffset + DisplayAtlasResolution))
 	{
diff --git a/Engine/Source/Runtime/Renderer/Private/RectLightTextureManager.cpp b/Engine/Source/Runtime/Renderer/Private/RectLightTextureManager.cpp
--- a/Engine/Source/Runtime/Renderer/Private/RectLightTextureManager.cpp
+++ b/Engine/Source/Runtime/Renderer/Private/RectLightTextureManager.cpp
@@ -52,6 +52,18 @@ static TAutoConsoleVariable<int32> CVarRectLighForceUpdate(
 	TEXT("Force rect. light atlas update very frame."),
 	ECVF_RenderThreadSafe);
 
+static TAutoConsoleVariable<int32> CVarRectLighResetOnChange(
+	TEXT("r.RectLightAtlas.ResetOnChange"),
+	0,
+	TEXT("Recreate the atlas layout from scratch each time an update is made to the atlas."),
+	ECVF_RenderThreadSafe);
+
+static TAutoConsoleVariable<int32> CVarRectLighClearAtlas(
+	TEXT("r.RectLightAtlas.ClearAtlas"),
+	0,
+	TEXT("Clear rect light texture atlas each time the atlas is recreated/updated."),
+	ECVF_RenderThreadSafe);
+
 static TAutoConsoleVariable<int32> CVarRectLighFilterQuality(
 	TEXT("r.RectLightAtlas.FilterQuality"),
 	1,
@@ -340,6 +352,9 @@ class FRectLightAtlasDebugInfoCS : public FGlobalShader
 		SHADER_PARAMETER(uint32, AtlasMIPIndex)
 		SHADER_PARAMETER(uint32, AtlasSourceTextureMIPBias)
 		SHADER_PARAMETER(uint32, AtlasFormat)
+		SHADER_PARAMETER(uint32, ForceUpdate)
+		SHADER_PARAMETER(uint32, ResetOnChange)
+		SHADER_PARAMETER(uint32, ClearAtlas)
 		SHADER_PARAMETER_SAMPLER(SamplerState, AtlasSampler)
 		SHADER_PARAMETER_RDG_TEXTURE(Texture2D, AtlasTexture)
 		SHADER_PARAMETER_RDG_TEXTURE_UAV(RWTexture2D, OutputTexture)
@@ -426,6 +441,9 @@ static FRDGTextureRef AddRectLightDebugInfoPass(
 	Parameters->HorizonCount = Horizons.Num();
 	Parameters->FreeCount = FreeRects.Num();
 	Parameters->OutputResolution = OutputResolution;
+	Parameters->ForceUpdate = CVarRectLighForceUpdate.GetValueOnRenderThread() > 0;
+	Parameters->ResetOnChange = CVarRectLighResetOnChange.GetValueOnRenderThread() > 0;
+	Parameters->ClearAtlas = CVarRectLighClearAtlas.GetValueOnRenderThread() > 0;
 	Parameters->AtlasSampler = TStaticSamplerState<SF_Bilinear, AM_Clamp, AM_Clamp, AM_Clamp>::GetRHI();
 	Parameters->AtlasTexture = AtlasTexture;
 	Parameters->SlotBuffer = GraphBuilder.CreateSRV(SlotBuffer, PF_R16G16B16A16_UINT);
@@ -1624,17 +1642,26 @@ void UpdateAtlasTexture(FRDGBuilder& GraphBuilder, const ERHIFeatureLevel::Type
 	static int32 CachedMaxAtlasResolution = CVarRectLightTextureResolution.GetValueOnRenderThread();
 	static EPixelFormat CachedAtlasFormat = GetRectLightAtlasFormat();
 	const EPixelFormat AtlasFormat = GetRectLightAtlasFormat();
-	const bool bForceUpdate = CVarRectLighForceUpdate.GetValueOnRenderThread() > 0 || CachedMaxAtlasResolution != CVarRectLightTextureResolution.GetValueOnRenderThread() || CachedAtlasFormat != AtlasFormat;
-	if (bForceUpdate)
+	const bool bForceUpdate = CVarRectLighForceUpdate.GetValueOnRenderThread() > 0 || CachedMaxAtlasResolution != CVarRectLightTextureResolution.GetValueOnRenderThread() || CachedAtlasFormat != AtlasFormat;	
+	const bool bResetOnChange = CVarRectLighResetOnChange.GetValueOnRenderThread() > 0;
+	const bool bClearAtlas = CVarRectLighClearAtlas.GetValueOnRenderThread() > 0;
+	auto ResetAtlas = [&]()
 	{
 		CachedMaxAtlasResolution = CVarRectLightTextureResolution.GetValueOnRenderThread();
 		CachedAtlasFormat = AtlasFormat;
 		GRectLightTextureManager.bHasPendingAdds = true;
 		GRectLightTextureManager.AtlasLayout = FAtlasLayout(FIntPoint(CachedMaxAtlasResolution, CachedMaxAtlasResolution));
+		GRectLightTextureManager.DeletedSlots.Empty();
+		GRectLightTextureManager.FreeSlots.Empty();
 		for (FAtlasSlot& Slot : GRectLightTextureManager.AtlasSlots)
 		{
 			Slot.Rect.Origin = InvalidOrigin;
 		}
+	};
+
+	if (bForceUpdate)
+	{
+		ResetAtlas();
 	}
 
 	// Force update if among the existing valid slot a streamed a higher resolution than the existing one
@@ -1659,6 +1686,12 @@ void UpdateAtlasTexture(FRDGBuilder& GraphBuilder, const ERHIFeatureLevel::Type
 		}
 	}
 
+	// Force the atlas to be reset when some change are detected (optional)
+	if (GRectLightTextureManager.bHasPendingAdds && !bForceUpdate && bResetOnChange)
+	{
+		ResetAtlas();
+	}
+
 	// Update the atlas layout with the delete slots
 	if (GRectLightTextureManager.DeletedSlots.Num() > 0)
 	{
@@ -1699,6 +1732,12 @@ void UpdateAtlasTexture(FRDGBuilder& GraphBuilder, const ERHIFeatureLevel::Type
 			else
 			{
 				AtlasTexture = GraphBuilder.RegisterExternalTexture(GRectLightTextureManager.AtlasTexture);
+
+				// Clear the atlas (optional)
+				if (bClearAtlas && bResetOnChange && CopySlots.Num() == 0)
+				{
+					AddClearRenderTargetPass(GraphBuilder, AtlasTexture, FLinearColor::Transparent);
+				}
 			}
 
 			// 2.1 Copy slots from previous to new atlas texture
