When rendering shadow depths on platforms without shadow stencil bits, prevent shadow stretching by using scissor rect to restrict the rendering to inside borders.

diff --git a/Engine/Source/Runtime/Renderer/Private/ShadowDepthRendering.cpp b/Engine/Source/Runtime/Renderer/Private/ShadowDepthRendering.cpp
--- a/Engine/Source/Runtime/Renderer/Private/ShadowDepthRendering.cpp
+++ b/Engine/Source/Runtime/Renderer/Private/ShadowDepthRendering.cpp
@@ -701,7 +701,7 @@ void FProjectedShadowInfo::ClearDepth(FRHICommandList& RHICmdList) const
 	DrawClearQuadMRT(RHICmdList, bClearColor, NumClearColors, Colors, true, 1.0f, false, 0);
 }
 
-void FProjectedShadowInfo::SetStateForView(FRHICommandList& RHICmdList) const
+void FProjectedShadowInfo::SetStateForView(FRHICommandList& RHICmdList, bool bRequiresBorderClamp) const
 {
 	check(bAllocated);
 
@@ -724,6 +724,18 @@ void FProjectedShadowInfo::SetStateForView(FRHICommandList& RHICmdList) const
 			ScissorRectOptim.Max.Y + Y
 		);
 	}
+	else if (bRequiresBorderClamp)
+	{
+		// Platforms without the required stencil bits need to prevent shadow stretching by
+		// restricting the rendering to inside the texture borders
+		RHICmdList.SetScissorRect(
+			true,
+			X + BorderSize,
+			Y + BorderSize,
+			X + ResolutionX,
+			Y + ResolutionY
+		);
+	}
 }
 
 void SetStateForShadowDepth(EShadowDepthType ShadowDepthType, FMeshPassProcessorRenderState& DrawRenderState)
@@ -1135,29 +1147,31 @@ void FProjectedShadowInfo::RenderDepth(
 
 	ShadowDepthPass.BuildRenderingCommands(GraphBuilder, Scene->GPUScene, PassParameters->InstanceCullingDrawParams);
 
+	bool bRequiresBorderClamp = !FProjectedShadowInfo::HasShadowStencilCulling(SceneRenderer->ShaderPlatform);
+
 	if (bDoParallelDispatch)
 	{
 		GraphBuilder.AddDispatchPass(
 			RDG_EVENT_NAME("ShadowDepthPassParallel"),
 			PassParameters,
 			ERDGPassFlags::Raster,
-			[this, PassParameters](FRDGDispatchPassBuilder& DispatchPassBuilder)
+			[this, PassParameters, bRequiresBorderClamp](FRDGDispatchPassBuilder& DispatchPassBuilder)
 		{
-			ShadowDepthPass.Dispatch(DispatchPassBuilder, &PassParameters->InstanceCullingDrawParams, [this] (FRHICommandList& RHICmdList)
+			ShadowDepthPass.Dispatch(DispatchPassBuilder, &PassParameters->InstanceCullingDrawParams, [this, bRequiresBorderClamp] (FRHICommandList& RHICmdList)
 			{
-				SetStateForView(RHICmdList);
+				SetStateForView(RHICmdList, bRequiresBorderClamp);
 			});
 		});
 	}
 	else
 	{
-		GraphBuilder.AddPass(
+			GraphBuilder.AddPass(
 			RDG_EVENT_NAME("ShadowDepthPass"),
 			PassParameters,
 			ERDGPassFlags::Raster,
-			[this, PassParameters](FRDGAsyncTask, FRHICommandList& RHICmdList)
+			[this, PassParameters, bRequiresBorderClamp](FRDGAsyncTask, FRHICommandList& RHICmdList)
 		{
-			SetStateForView(RHICmdList);
+			SetStateForView(RHICmdList, bRequiresBorderClamp);
 			ShadowDepthPass.Draw(RHICmdList, &PassParameters->InstanceCullingDrawParams);
 		});
 	}
diff --git a/Engine/Source/Runtime/Renderer/Private/ShadowRendering.h b/Engine/Source/Runtime/Renderer/Private/ShadowRendering.h
--- a/Engine/Source/Runtime/Renderer/Private/ShadowRendering.h
+++ b/Engine/Source/Runtime/Renderer/Private/ShadowRendering.h
@@ -521,7 +521,7 @@ public:
 	float GetShaderMaxSlopeDepthBias() const { return ShaderMaxSlopeDepthBias; }
 	float GetShaderReceiverDepthBias() const;
 
-	void SetStateForView(FRHICommandList& RHICmdList) const;
+	void SetStateForView(FRHICommandList& RHICmdList, bool bRequiresBorderClamp = false) const;
 
 	/** Get view rect excluding border area */
 	FIntRect GetInnerViewRect() const
