Fix shadows crashing when trying to build rendering commands for second eye in stereo.

diff --git a/Engine/Source/Runtime/Renderer/Private/ShadowSetup.cpp b/Engine/Source/Runtime/Renderer/Private/ShadowSetup.cpp
--- a/Engine/Source/Runtime/Renderer/Private/ShadowSetup.cpp
+++ b/Engine/Source/Runtime/Renderer/Private/ShadowSetup.cpp
@@ -2056,6 +2056,10 @@ bool FProjectedShadowInfo::AddSubjectPrimitive(FPrimitiveSceneInfo* PrimitiveSce
 			{
 				for (FViewInfo& CurrentView : Views)
 				{
+					if (!CurrentView.ShouldRenderView())
+					{
+						continue;
+					}
 					// Note: skip small-mesh culling for VSM since it needs it drawn for GPU-side caching.
 					if (bWholeSceneShadow && CacheMode != SDCM_StaticPrimitivesOnly && !bVSM)
 					{
