[LightGrid] Renamed NextCulledLightLink/NextCulledLightData resource names/references.

diff --git a/Engine/Shaders/Private/LightGridInjection.usf b/Engine/Shaders/Private/LightGridInjection.usf
--- a/Engine/Shaders/Private/LightGridInjection.usf
+++ b/Engine/Shaders/Private/LightGridInjection.usf
@@ -28,8 +28,8 @@ RWStructuredBuffer<uint> RWNumCulledLightsGrid;
 	#define RWCulledLightDataGrid	RWCulledLightDataGrid32Bit
 #endif
 
-RWStructuredBuffer<uint> RWNextCulledLightData;
-RWStructuredBuffer<uint> RWNextCulledLightLink;
+RWStructuredBuffer<uint> RWCulledLightDataAllocator;
+RWStructuredBuffer<uint> RWCulledLightLinkAllocator;
 RWStructuredBuffer<uint> RWCulledLightLinks;
 
 uint NumReflectionCaptures;
@@ -412,7 +412,7 @@ void CullLights_SingleThread(uint GridIndex, FCellBounds CellBounds, uint Parent
 
 #if USE_LINKED_CULL_LIST
 			uint NextLink;
-			WaveInterlockedAddScalar_(RWNextCulledLightLink[0], 1U, NextLink);
+			WaveInterlockedAddScalar_(RWCulledLightLinkAllocator[0], 1U, NextLink);
 
 			if (NextLink < NumAvailableLinks)
 			{
@@ -451,9 +451,9 @@ void CullLights_SingleThread(uint GridIndex, FCellBounds CellBounds, uint Parent
 	#if FEATURE_LEVEL == FEATURE_LEVEL_ES3_1
 	// Adreno compiler fails to reg-alloc on 6xx for this shader,
 	// and doesn't support UGPR spilling, so just fails to compile.
-	InterlockedAdd(RWNextCulledLightData[0], NumVisibleLights, CulledLightDataStart);
+	InterlockedAdd(RWCulledLightDataAllocator[0], NumVisibleLights, CulledLightDataStart);
 	#else
-	WaveInterlockedAdd_(RWNextCulledLightData[0], NumVisibleLights, CulledLightDataStart);
+	WaveInterlockedAdd_(RWCulledLightDataAllocator[0], NumVisibleLights, CulledLightDataStart);
 	#endif
 
 	// Mega Lights are order dependent
@@ -484,7 +484,7 @@ void CullReflectionCaptures_SingleThread(uint GridIndex, FCellBounds CellBounds,
 		{
 #if USE_LINKED_CULL_LIST
 			uint NextLink;
-			WaveInterlockedAddScalar_(RWNextCulledLightLink[0], 1U, NextLink);
+			WaveInterlockedAddScalar_(RWCulledLightLinkAllocator[0], 1U, NextLink);
 
 			if (NextLink < NumAvailableLinks)
 			{
@@ -508,7 +508,7 @@ void CullReflectionCaptures_SingleThread(uint GridIndex, FCellBounds CellBounds,
 
 #if USE_LINKED_CULL_LIST
 	uint CulledReflectionCaptureDataStart;
-	WaveInterlockedAdd_(RWNextCulledLightData[0], NumVisibleCaptures, CulledReflectionCaptureDataStart);
+	WaveInterlockedAdd_(RWCulledLightDataAllocator[0], NumVisibleCaptures, CulledReflectionCaptureDataStart);
 
 	// Reflection captures are order dependent
 	OutputReversedLinkedList(LinkOffset, NumVisibleCaptures, CulledReflectionCaptureDataStart, false);
@@ -555,7 +555,7 @@ bool AddToCell(inout FCellWriter CellWriter, uint PrimitiveIndex, uint IndexInCe
 
 	// otherwise store in linked list
 	uint NextLink;
-	WaveInterlockedAddScalar_(RWNextCulledLightLink[0], 1U, NextLink);
+	WaveInterlockedAddScalar_(RWCulledLightLinkAllocator[0], 1U, NextLink);
 
 	if (NextLink < NumAvailableLinks)
 	{
@@ -700,7 +700,7 @@ void CullLights_ThreadGroup(uint ThreadIndex, uint GridIndex, FCellBounds CellBo
 
 	if (ThreadIndex == 0)
 	{
-		InterlockedAdd(RWNextCulledLightData[0], CellNumVisibleLights, CulledLightDataStart);
+		InterlockedAdd(RWCulledLightDataAllocator[0], CellNumVisibleLights, CulledLightDataStart);
 	}
 
 	GroupMemoryBarrierWithGroupSync();
@@ -771,7 +771,7 @@ void CullReflectionCaptures_ThreadGroup(uint ThreadIndex, uint GridIndex, FCellB
 
 	if (ThreadIndex == 0)
 	{
-		InterlockedAdd(RWNextCulledLightData[0], CellNumVisibleReflectionCaptures, CulledLightDataStart);
+		InterlockedAdd(RWCulledLightDataAllocator[0], CellNumVisibleReflectionCaptures, CulledLightDataStart);
 	}
 
 	GroupMemoryBarrierWithGroupSync();
@@ -1000,18 +1000,18 @@ void DebugLightGridPS(
 
 #ifdef FeedbackStatusCS
 
-StructuredBuffer<uint> NextCulledLightDataBuffer;
+StructuredBuffer<uint> CulledLightDataAllocatorBuffer;
 uint NumCulledLightDataEntries;
 
-StructuredBuffer<uint> NextCulledLightLinkBuffer;
+StructuredBuffer<uint> CulledLightLinkAllocatorBuffer;
 
 uint StatusMessageId;
 
 [numthreads(1, 1, 1)]
 void FeedbackStatusCS()
 {												
-	uint NumAllocatedEntries = NextCulledLightDataBuffer[0];
-	uint NumAllocatedLinks = NextCulledLightLinkBuffer[0];
+	uint NumAllocatedEntries = CulledLightDataAllocatorBuffer[0];
+	uint NumAllocatedLinks = CulledLightLinkAllocatorBuffer[0];
 
 	FGPUMessageWriter Mw = GPUMessageBegin(StatusMessageId, 4U);
 
diff --git a/Engine/Source/Runtime/Renderer/Private/LightGrid.h b/Engine/Source/Runtime/Renderer/Private/LightGrid.h
--- a/Engine/Source/Runtime/Renderer/Private/LightGrid.h
+++ b/Engine/Source/Runtime/Renderer/Private/LightGrid.h
@@ -17,7 +17,7 @@ public:
 
 	FLightGridViewState();
 
-	void FeedbackStatus(FRDGBuilder& GraphBuilder, FViewInfo& View, FRDGBufferRef NextCulledLightDataBuffer, uint32 NumCulledLightDataEntries, FRDGBufferRef NextCulledLightLinkBuffer, uint32 NumCulledLightLinks);
+	void FeedbackStatus(FRDGBuilder& GraphBuilder, FViewInfo& View, FRDGBufferRef CulledLightDataAllocatorBuffer, uint32 NumCulledLightDataEntries, FRDGBufferRef CulledLightLinkAllocatorBuffer, uint32 NumCulledLightLinks);
 
 #if !UE_BUILD_SHIPPING
 	uint32 GetStatusMessageId() const { return StatusFeedbackSocket.GetMessageId().GetIndex(); }
diff --git a/Engine/Source/Runtime/Renderer/Private/LightGridInjection.cpp b/Engine/Source/Runtime/Renderer/Private/LightGridInjection.cpp
--- a/Engine/Source/Runtime/Renderer/Private/LightGridInjection.cpp
+++ b/Engine/Source/Runtime/Renderer/Private/LightGridInjection.cpp
@@ -155,7 +155,7 @@ static TAutoConsoleVariable<bool> CVarLightCullingTwoLevel(
 static TAutoConsoleVariable<int32> CVarLightCullingTwoLevelThreshold(
 	TEXT("r.LightCulling.TwoLevel.Threshold"),
 	128,
-	TEXT("Threshold used to determine whether to use two level culling basedon the number of lights in view."),
+	TEXT("Threshold used to determine whether to use two level culling based on the number of lights in view."),
 	ECVF_RenderThreadSafe
 );
 
@@ -286,8 +286,8 @@ public:
 		SHADER_PARAMETER_RDG_BUFFER_UAV(RWStructuredBuffer<uint>, RWNumCulledLightsGrid)
 		SHADER_PARAMETER_RDG_BUFFER_UAV(RWStructuredBuffer<uint>, RWCulledLightDataGrid32Bit)
 		SHADER_PARAMETER_RDG_BUFFER_UAV(RWBuffer<uint>, RWCulledLightDataGrid16Bit)
-		SHADER_PARAMETER_RDG_BUFFER_UAV(RWStructuredBuffer<uint>, RWNextCulledLightLink)
-		SHADER_PARAMETER_RDG_BUFFER_UAV(RWStructuredBuffer<uint>, RWNextCulledLightData)
+		SHADER_PARAMETER_RDG_BUFFER_UAV(RWStructuredBuffer<uint>, RWCulledLightLinkAllocator)
+		SHADER_PARAMETER_RDG_BUFFER_UAV(RWStructuredBuffer<uint>, RWCulledLightDataAllocator)
 		SHADER_PARAMETER_RDG_BUFFER_UAV(RWStructuredBuffer<uint>, RWCulledLightLinks)
 		SHADER_PARAMETER_RDG_BUFFER_SRV(StructuredBuffer<float4>, LightViewSpacePositionAndRadius)
 		SHADER_PARAMETER_RDG_BUFFER_SRV(StructuredBuffer<float4>, LightViewSpaceDirAndPreprocAngle)
@@ -624,9 +624,14 @@ FLightGrid LightGridInjection(
 	const FIntVector ParentGridSize = FIntVector::DivideAndRoundUp(GridSize, ParentGridSizeFactor);
 
 	FRDGBufferRef CulledLightLinksBuffer = GraphBuilder.CreateBuffer(FRDGBufferDesc::CreateStructuredDesc(sizeof(uint32), NumCulledLightLinks * LightLinkStride), TEXT("CulledLightLinks"));
-	FRDGBufferRef NextCulledLightLinkBuffer = GraphBuilder.CreateBuffer(FRDGBufferDesc::CreateStructuredDesc(sizeof(uint32), 1), TEXT("NextCulledLightLink"));
-	FRDGBufferRef NextCulledLightDataBuffer = GraphBuilder.CreateBuffer(FRDGBufferDesc::CreateStructuredDesc(sizeof(uint32), 1), TEXT("NextCulledLightData"));
-	FRDGBufferUAVRef NextCulledLightDataUAV = GraphBuilder.CreateUAV(NextCulledLightDataBuffer);
+	FRDGBufferUAVRef CulledLightLinksUAV = GraphBuilder.CreateUAV(CulledLightLinksBuffer);
+
+	FRDGBufferRef CulledLightLinkAllocatorBuffer = GraphBuilder.CreateBuffer(FRDGBufferDesc::CreateStructuredDesc(sizeof(uint32), 1), TEXT("CulledLightLinkAllocator"));
+	FRDGBufferUAVRef CulledLightLinkAllocatorUAV = GraphBuilder.CreateUAV(CulledLightLinkAllocatorBuffer);
+
+	FRDGBufferRef CulledLightDataAllocatorBuffer = GraphBuilder.CreateBuffer(FRDGBufferDesc::CreateStructuredDesc(sizeof(uint32), 1), TEXT("CulledLightDataAllocator"));
+	FRDGBufferUAVRef CulledLightDataAllocatorUAV = GraphBuilder.CreateUAV(CulledLightDataAllocatorBuffer);
+
 	FRDGBufferRef NumCulledLightsGrid = GraphBuilder.CreateBuffer(FRDGBufferDesc::CreateStructuredDesc(sizeof(uint32), MaxNumCells * NumCulledLightsGridStride), TEXT("NumCulledLightsGrid"));
 	FRDGBufferUAVRef NumCulledLightsGridUAV = GraphBuilder.CreateUAV(NumCulledLightsGrid);
 
@@ -649,6 +654,10 @@ FLightGrid LightGridInjection(
 		CulledLightDataGridUAV = GraphBuilder.CreateUAV(CulledLightDataGrid);
 	}
 
+	AddClearUAVPass(GraphBuilder, CulledLightLinkAllocatorUAV, 0);
+	AddClearUAVPass(GraphBuilder, CulledLightDataAllocatorUAV, 0);
+	AddClearUAVPass(GraphBuilder, NumCulledLightsGridUAV, 0);
+
 	FLightGridInjectionCS::FParameters* PassParameters = GraphBuilder.AllocParameters<FLightGridInjectionCS::FParameters>();
 
 	PassParameters->View = View.ViewUniformBuffer;
@@ -665,9 +674,9 @@ FLightGrid LightGridInjection(
 	PassParameters->RWNumCulledLightsGrid = NumCulledLightsGridUAV;
 	PassParameters->RWCulledLightDataGrid32Bit = CulledLightDataGridUAV;
 	PassParameters->RWCulledLightDataGrid16Bit = CulledLightDataGridUAV;
-	PassParameters->RWNextCulledLightLink = GraphBuilder.CreateUAV(NextCulledLightLinkBuffer);
-	PassParameters->RWNextCulledLightData = NextCulledLightDataUAV;
-	PassParameters->RWCulledLightLinks = GraphBuilder.CreateUAV(CulledLightLinksBuffer);
+	PassParameters->RWCulledLightLinkAllocator = CulledLightLinkAllocatorUAV;
+	PassParameters->RWCulledLightDataAllocator = CulledLightDataAllocatorUAV;
+	PassParameters->RWCulledLightLinks = CulledLightLinksUAV;
 	PassParameters->CulledGridSize = GridSize;
 	PassParameters->LightGridZParams = ZParams;
 	PassParameters->NumReflectionCaptures = NumReflectionCaptures;
@@ -714,10 +723,6 @@ FLightGrid LightGridInjection(
 	PermutationVector.Set<FLightGridInjectionCS::FApplyIndirection>(IndirectionIndicesSRV != nullptr);
 	auto ComputeShader = View.ShaderMap->GetShader<FLightGridInjectionCS>(PermutationVector);
 
-	AddClearUAVPass(GraphBuilder, PassParameters->RWNextCulledLightLink, 0);
-	AddClearUAVPass(GraphBuilder, NextCulledLightDataUAV, 0);
-	AddClearUAVPass(GraphBuilder, NumCulledLightsGridUAV, 0);
-
 	FIntVector NumGroups;
 	if (bThreadGroupPerCell)
 	{
@@ -743,7 +748,7 @@ FLightGrid LightGridInjection(
 
 	if (LightGridViewState != nullptr)
 	{
-		LightGridViewState->FeedbackStatus(GraphBuilder, View, NextCulledLightDataBuffer, NumCulledLightEntries, NextCulledLightLinkBuffer, NumCulledLightLinks);
+		LightGridViewState->FeedbackStatus(GraphBuilder, View, CulledLightDataAllocatorBuffer, NumCulledLightEntries, CulledLightLinkAllocatorBuffer, NumCulledLightLinks);
 	}
 
 	return Output;
@@ -1626,10 +1631,10 @@ class FLightGridFeedbackStatusCS : public FGlobalShader
 	SHADER_USE_PARAMETER_STRUCT(FLightGridFeedbackStatusCS, FGlobalShader);
 
 	BEGIN_SHADER_PARAMETER_STRUCT(FParameters, )
-		SHADER_PARAMETER_RDG_BUFFER_SRV(StructuredBuffer<uint>, NextCulledLightDataBuffer)
+		SHADER_PARAMETER_RDG_BUFFER_SRV(StructuredBuffer<uint>, CulledLightDataAllocatorBuffer)
 		SHADER_PARAMETER(uint32, NumCulledLightDataEntries)
 
-		SHADER_PARAMETER_RDG_BUFFER_SRV(StructuredBuffer<uint>, NextCulledLightLinkBuffer)
+		SHADER_PARAMETER_RDG_BUFFER_SRV(StructuredBuffer<uint>, CulledLightLinkAllocatorBuffer)
 		SHADER_PARAMETER(uint32, NumAvailableLinks)
 
 		SHADER_PARAMETER_STRUCT_INCLUDE(GPUMessage::FParameters, GPUMessageParams)
@@ -1677,15 +1682,15 @@ FLightGridViewState::FLightGridViewState()
 #endif
 }
 
-void FLightGridViewState::FeedbackStatus(FRDGBuilder& GraphBuilder, FViewInfo& View, FRDGBufferRef NextCulledLightDataBuffer, uint32 NumCulledLightDataEntries, FRDGBufferRef NextCulledLightLinkBuffer, uint32 NumCulledLightLinks)
+void FLightGridViewState::FeedbackStatus(FRDGBuilder& GraphBuilder, FViewInfo& View, FRDGBufferRef CulledLightDataAllocatorBuffer, uint32 NumCulledLightDataEntries, FRDGBufferRef CulledLightLinkAllocatorBuffer, uint32 NumCulledLightLinks)
 {
 #if !UE_BUILD_SHIPPING
 	FLightGridFeedbackStatusCS::FParameters* PassParameters = GraphBuilder.AllocParameters<FLightGridFeedbackStatusCS::FParameters>();
 
-	PassParameters->NextCulledLightDataBuffer = GraphBuilder.CreateSRV(NextCulledLightDataBuffer);
+	PassParameters->CulledLightDataAllocatorBuffer = GraphBuilder.CreateSRV(CulledLightDataAllocatorBuffer);
 	PassParameters->NumCulledLightDataEntries = NumCulledLightDataEntries;
 
-	PassParameters->NextCulledLightLinkBuffer = GraphBuilder.CreateSRV(NextCulledLightLinkBuffer);
+	PassParameters->CulledLightLinkAllocatorBuffer = GraphBuilder.CreateSRV(CulledLightLinkAllocatorBuffer);
 	PassParameters->NumAvailableLinks = NumCulledLightLinks;
 
 	PassParameters->GPUMessageParams = GPUMessage::GetShaderParameters(GraphBuilder);
