Remove some unused members from FSceneViewState.
- also an unused function declaration from TranslucentLighting.h

diff --git a/Engine/Source/Runtime/Renderer/Private/RendererScene.cpp b/Engine/Source/Runtime/Renderer/Private/RendererScene.cpp
--- a/Engine/Source/Runtime/Renderer/Private/RendererScene.cpp
+++ b/Engine/Source/Runtime/Renderer/Private/RendererScene.cpp
@@ -320,8 +320,6 @@ FSceneViewState::FSceneViewState(ERHIFeatureLevel::Type FeatureLevel, FSceneView
 	FrameIndex = 0;
 	OutputFrameIndex = 0;
 	DistanceFieldTemporalSampleIndex = 0;
-	bDOFHistory = true;
-	bDOFHistory2 = true;
 	
 	// Sets the mipbias to invalid large number.
 	MaterialTextureCachedMipBias = BIG_NUMBER;
@@ -333,11 +331,6 @@ FSceneViewState::FSceneViewState(ERHIFeatureLevel::Type FeatureLevel, FSceneView
 
 	bRoundRobinOcclusionEnabled = false;
 
-	for (int32 CascadeIndex = 0; CascadeIndex < UE_ARRAY_COUNT(TranslucencyLightingCacheAllocations); CascadeIndex++)
-	{
-		TranslucencyLightingCacheAllocations[CascadeIndex] = NULL;
-	}
-
 	if (ShareOriginTarget)
 	{
 		GlobalDistanceFieldData = ShareOriginTarget->GlobalDistanceFieldData;
@@ -352,11 +345,6 @@ FSceneViewState::FSceneViewState(ERHIFeatureLevel::Type FeatureLevel, FSceneView
 	ShadowOcclusionQueryMaps.Empty(FOcclusionQueryHelpers::MaxBufferedOcclusionFrames);
 	ShadowOcclusionQueryMaps.AddZeroed(FOcclusionQueryHelpers::MaxBufferedOcclusionFrames);	
 
-	LastAutoDownsampleChangeTime = 0;
-	SmoothedHalfResTranslucencyGPUDuration = 0;
-	SmoothedFullResTranslucencyGPUDuration = 0;
-	bShouldAutoDownsampleTranslucency = false;
-
 	PreExposure = 1.f;
 	bUpdateLastExposure = false;
 
@@ -384,11 +372,6 @@ FSceneViewState::~FSceneViewState()
 	CachedVisibilityChunk = NULL;
 	ShadowOcclusionQueryMaps.Reset();
 
-	for (int32 CascadeIndex = 0; CascadeIndex < UE_ARRAY_COUNT(TranslucencyLightingCacheAllocations); CascadeIndex++)
-	{
-		delete TranslucencyLightingCacheAllocations[CascadeIndex];
-	}
-
 	HairStrandsViewStateData.Release();
 	ShaderPrintStateData.Release();
 
diff --git a/Engine/Source/Runtime/Renderer/Private/ScenePrivate.h b/Engine/Source/Runtime/Renderer/Private/ScenePrivate.h
--- a/Engine/Source/Runtime/Renderer/Private/ScenePrivate.h
+++ b/Engine/Source/Runtime/Renderer/Private/ScenePrivate.h
@@ -819,8 +819,6 @@ public:
 	/** For this view, the set of primitives that are currently fading, either in or out. */
 	FPrimitiveFadingStateMap PrimitiveFadingStates;
 
-	FIndirectLightingCacheAllocation* TranslucencyLightingCacheAllocations[TVC_MAX];
-
 	TMap<int32, FIndividualOcclusionHistory> PlanarReflectionOcclusionHistories;
 
 #if !(UE_BUILD_SHIPPING || UE_BUILD_TEST)
@@ -1107,19 +1105,6 @@ public:
 	/** Potentially shared to save memory in cases where multiple view states share a common origin, such as cube map capture faces. */
 	TRefCountPtr<FPersistentGlobalDistanceFieldData> GlobalDistanceFieldData;
 
-	/** This is float since it is derived off of UWorld::RealTimeSeconds, which is relative to BeginPlay time. */
-	float LastAutoDownsampleChangeTime;
-	float SmoothedHalfResTranslucencyGPUDuration;
-	float SmoothedFullResTranslucencyGPUDuration;
-
-	/** Current desired state of auto-downsampled separate translucency for this view. */
-	bool bShouldAutoDownsampleTranslucency;
-
-	// Is DOFHistoryRT set from DepthOfField?
-	bool bDOFHistory;
-	// Is DOFHistoryRT2 set from DepthOfField?
-	bool bDOFHistory2;
-
 	// Sequencer state for view management
 	ESequencerState SequencerState;
 
diff --git a/Engine/Source/Runtime/Renderer/Private/TranslucentLighting.h b/Engine/Source/Runtime/Renderer/Private/TranslucentLighting.h
--- a/Engine/Source/Runtime/Renderer/Private/TranslucentLighting.h
+++ b/Engine/Source/Runtime/Renderer/Private/TranslucentLighting.h
@@ -130,11 +130,6 @@ void InjectSimpleTranslucencyLightingVolumeArray(
 	const FTranslucencyLightingVolumeTextures& Textures,
 	const FSimpleLightArray& SimpleLights);
 
-void FinalizeTranslucencyLightingVolume(
-	FRDGBuilder& GraphBuilder,
-	const TArrayView<const FViewInfo> Views,
-	FTranslucencyLightingVolumeTextures& Textures);
-
 void InjectTranslucencyLightingVolumeAmbientCubemap(
 	FRDGBuilder& GraphBuilder,
 	const TArrayView<const FViewInfo> Views,
