Flush the RDG setup queue after light passes are added to reduce bottlenecks further down the pipeline.

diff --git a/Engine/Source/Runtime/Renderer/Private/LightRendering.cpp b/Engine/Source/Runtime/Renderer/Private/LightRendering.cpp
--- a/Engine/Source/Runtime/Renderer/Private/LightRendering.cpp
+++ b/Engine/Source/Runtime/Renderer/Private/LightRendering.cpp
@@ -2278,6 +2278,8 @@ void FDeferredShadingSceneRenderer::RenderLights(
 			}
 		}
 	}
+
+	GraphBuilder.FlushSetupQueue();
 }
 
 static void RenderLightArrayForOverlapViewmode(
diff --git a/Engine/Source/Runtime/Renderer/Private/TranslucentLighting.cpp b/Engine/Source/Runtime/Renderer/Private/TranslucentLighting.cpp
--- a/Engine/Source/Runtime/Renderer/Private/TranslucentLighting.cpp
+++ b/Engine/Source/Runtime/Renderer/Private/TranslucentLighting.cpp
@@ -1594,6 +1594,8 @@ void InjectTranslucencyLightingVolume(
 			}
 		}
 	}
+
+	GraphBuilder.FlushSetupQueue();
 }
 
 class FSimpleLightTranslucentLightingInjectPS : public FGlobalShader
