Added trace scope to InitializeTextures.

diff --git a/Engine/Source/Runtime/Renderer/Private/SystemTextures.cpp b/Engine/Source/Runtime/Renderer/Private/SystemTextures.cpp
--- a/Engine/Source/Runtime/Renderer/Private/SystemTextures.cpp
+++ b/Engine/Source/Runtime/Renderer/Private/SystemTextures.cpp
@@ -146,6 +146,8 @@ TGlobalResource<FSystemTextures> GSystemTextures;
 
 void FSystemTextures::InitializeTextures(FRHICommandListImmediate& RHICmdList, const ERHIFeatureLevel::Type InFeatureLevel)
 {
+	TRACE_CPUPROFILER_EVENT_SCOPE(FSystemTextures::InitializeTextures);
+
 	// When we render to system textures it should occur on all GPUs since this only
 	// happens once on startup (or when the feature level changes).
 	SCOPED_GPU_MASK(RHICmdList, FRHIGPUMask::All());
