[HWRT] Delete deprecated code in RayTracingInstanceBufferUtil.cpp.

diff --git a/Engine/Source/Runtime/Renderer/Private/RayTracing/RayTracingInstanceBufferUtil.cpp b/Engine/Source/Runtime/Renderer/Private/RayTracing/RayTracingInstanceBufferUtil.cpp
--- a/Engine/Source/Runtime/Renderer/Private/RayTracing/RayTracingInstanceBufferUtil.cpp
+++ b/Engine/Source/Runtime/Renderer/Private/RayTracing/RayTracingInstanceBufferUtil.cpp
@@ -95,119 +95,6 @@ FRayTracingSceneInitializationData BuildRayTracingSceneInitializationData(TConst
 	return MoveTemp(Output);
 }
 
-FRayTracingSceneWithGeometryInstances CreateRayTracingSceneWithGeometryInstances(
-	TConstArrayView<FRayTracingGeometryInstance> Instances,
-	uint8 NumLayers,
-	uint32 NumShaderSlotsPerGeometrySegment,
-	uint32 NumMissShaderSlots,
-	uint32 NumCallableShaderSlots,
-	ERayTracingAccelerationStructureFlags BuildFlags)
-{
-	const uint32 NumSceneInstances = Instances.Num();
-
-	FRayTracingSceneWithGeometryInstances Output;
-	Output.NumNativeGPUSceneInstances = 0;
-	Output.NumNativeCPUInstances = 0;
-	Output.TotalNumSegments = 0;
-	Output.InstanceGeometryIndices.SetNumUninitialized(NumSceneInstances);
-	Output.BaseUploadBufferOffsets.SetNumUninitialized(NumSceneInstances);
-	Output.BaseInstancePrefixSum.SetNumUninitialized(NumSceneInstances);
-	Output.PerInstanceGeometries.SetNumUninitialized(NumSceneInstances);
-
-	FRayTracingSceneInitializer Initializer;
-	Initializer.DebugName = FName(TEXT("FRayTracingScene"));
-	Initializer.BuildFlags = BuildFlags;
-
-	PRAGMA_DISABLE_DEPRECATION_WARNINGS
-	Initializer.BaseInstancePrefixSum.SetNumUninitialized(NumSceneInstances);
-	Initializer.SegmentPrefixSum.SetNumUninitialized(NumSceneInstances);
-	Initializer.NumNativeInstancesPerLayer.SetNumZeroed(NumLayers);
-	Initializer.ShaderSlotsPerGeometrySegment = NumShaderSlotsPerGeometrySegment;
-	Initializer.NumMissShaderSlots = NumMissShaderSlots;
-	Initializer.NumCallableShaderSlots = NumCallableShaderSlots;
-	PRAGMA_ENABLE_DEPRECATION_WARNINGS
-
-	Experimental::TSherwoodMap<FRHIRayTracingGeometry*, uint32> UniqueGeometries;
-
-	// Compute geometry segment prefix sum used by GetHitRecordBaseIndex() during resource binding
-
-	for (uint32 InstanceIndex = 0; InstanceIndex < NumSceneInstances; ++InstanceIndex)
-	{
-		const FRayTracingGeometryInstance& InstanceDesc = Instances[InstanceIndex];
-
-		const bool bGpuSceneInstance = InstanceDesc.BaseInstanceSceneDataOffset != -1 || !InstanceDesc.InstanceSceneDataOffsets.IsEmpty();
-		const bool bCpuInstance = !bGpuSceneInstance;
-
-		checkf(!bGpuSceneInstance || InstanceDesc.BaseInstanceSceneDataOffset != -1 || InstanceDesc.NumTransforms <= uint32(InstanceDesc.InstanceSceneDataOffsets.Num()),
-			TEXT("Expected at least %d ray tracing geometry instance scene data offsets, but got %d."),
-			InstanceDesc.NumTransforms, InstanceDesc.InstanceSceneDataOffsets.Num());
-		checkf(!bCpuInstance || InstanceDesc.NumTransforms <= uint32(InstanceDesc.Transforms.Num()),
-			TEXT("Expected at least %d ray tracing geometry instance transforms, but got %d."),
-			InstanceDesc.NumTransforms, InstanceDesc.Transforms.Num());
-
-		checkf(InstanceDesc.GeometryRHI, TEXT("Ray tracing instance must have a valid geometry."));
-
-		Output.PerInstanceGeometries[InstanceIndex] = InstanceDesc.GeometryRHI;
-
-		PRAGMA_DISABLE_DEPRECATION_WARNINGS
-		// Compute geometry segment count prefix sum to be later used in GetHitRecordBaseIndex()
-		Initializer.SegmentPrefixSum[InstanceIndex] = Output.TotalNumSegments;
-		PRAGMA_ENABLE_DEPRECATION_WARNINGS
-
-		Output.TotalNumSegments += InstanceDesc.GeometryRHI->GetNumSegments();
-
-		uint32 GeometryIndex = UniqueGeometries.FindOrAdd(InstanceDesc.GeometryRHI, Output.ReferencedGeometries.Num());
-		Output.InstanceGeometryIndices[InstanceIndex] = GeometryIndex;
-		if (GeometryIndex == Output.ReferencedGeometries.Num())
-		{
-			Output.ReferencedGeometries.Add(InstanceDesc.GeometryRHI);
-		}
-
-		if (bGpuSceneInstance)
-		{
-			check(InstanceDesc.Transforms.IsEmpty());
-			Output.BaseUploadBufferOffsets[InstanceIndex] = Output.NumNativeGPUSceneInstances;
-			Output.NumNativeGPUSceneInstances += InstanceDesc.NumTransforms;
-		}
-		else if (bCpuInstance)
-		{
-			Output.BaseUploadBufferOffsets[InstanceIndex] = Output.NumNativeCPUInstances;
-			Output.NumNativeCPUInstances += InstanceDesc.NumTransforms;
-		}
-		else
-		{
-			checkNoEntry();
-		}
-
-		PRAGMA_DISABLE_DEPRECATION_WARNINGS
-		checkf(InstanceDesc.LayerIndex < NumLayers, 
-			TEXT("FRayTracingGeometryInstance is assigned to layer %d but raytracing scene being created only has %d layers."),
-			InstanceDesc.LayerIndex, NumLayers);
-
-		Initializer.BaseInstancePrefixSum[InstanceIndex] = Initializer.NumNativeInstancesPerLayer[InstanceDesc.LayerIndex];
-
-		// Can't support same instance in multiple layers because BaseInstancePrefixSum would be different per layer
-		Output.BaseInstancePrefixSum[InstanceIndex] = Initializer.NumNativeInstancesPerLayer[InstanceDesc.LayerIndex];
-
-		Initializer.NumNativeInstancesPerLayer[InstanceDesc.LayerIndex] += InstanceDesc.NumTransforms;
-		PRAGMA_ENABLE_DEPRECATION_WARNINGS
-	}
-
-	PRAGMA_DISABLE_DEPRECATION_WARNINGS
-	Initializer.NumTotalSegments = Output.TotalNumSegments;
-
-	if(NumLayers == 1)
-	{
-		Initializer.MaxNumInstances = Initializer.NumNativeInstancesPerLayer[0];
-		Initializer.NumNativeInstancesPerLayer.Empty();
-	}
-	PRAGMA_ENABLE_DEPRECATION_WARNINGS
-
-	Output.Scene = RHICreateRayTracingScene(MoveTemp(Initializer));
-
-	return MoveTemp(Output);
-}
-
 void FillRayTracingInstanceUploadBuffer(
 	FRayTracingSceneRHIRef RayTracingSceneRHI,
 	FVector PreViewTranslation,
@@ -327,34 +214,6 @@ void FillRayTracingInstanceUploadBuffer(
 		});
 }
 
-void FillRayTracingInstanceUploadBuffer(
-	FRayTracingSceneRHIRef RayTracingSceneRHI,
-	FVector PreViewTranslation,
-	TConstArrayView<FRayTracingGeometryInstance> Instances,
-	TConstArrayView<uint32> InstanceGeometryIndices,
-	TConstArrayView<uint32> BaseUploadBufferOffsets,
-	uint32 NumNativeGPUSceneInstances,
-	uint32 NumNativeCPUInstances,
-	TArrayView<FRayTracingInstanceDescriptorInput> OutInstanceUploadData,
-	TArrayView<FVector4f> OutTransformData)
-{
-	const FRayTracingSceneInitializer& SceneInitializer = RayTracingSceneRHI->GetInitializer();
-
-	FillRayTracingInstanceUploadBuffer(
-		RayTracingSceneRHI,
-		PreViewTranslation,
-		Instances,
-		InstanceGeometryIndices,
-		BaseUploadBufferOffsets,
-	PRAGMA_DISABLE_DEPRECATION_WARNINGS
-		SceneInitializer.BaseInstancePrefixSum,
-	PRAGMA_ENABLE_DEPRECATION_WARNINGS
-		NumNativeGPUSceneInstances,
-		NumNativeCPUInstances,
-		OutInstanceUploadData,
-		OutTransformData);
-}
-
 struct FRayTracingBuildInstanceBufferCS : public FGlobalShader
 {
 	DECLARE_GLOBAL_SHADER(FRayTracingBuildInstanceBufferCS);
diff --git a/Engine/Source/Runtime/Renderer/Public/RayTracingInstanceBufferUtil.h b/Engine/Source/Runtime/Renderer/Public/RayTracingInstanceBufferUtil.h
--- a/Engine/Source/Runtime/Renderer/Public/RayTracingInstanceBufferUtil.h
+++ b/Engine/Source/Runtime/Renderer/Public/RayTracingInstanceBufferUtil.h
@@ -65,48 +65,6 @@ struct FRayTracingSceneInitializationData
 
 RENDERER_API FRayTracingSceneInitializationData BuildRayTracingSceneInitializationData(TConstArrayView<FRayTracingGeometryInstance> Instances);
 
-struct FRayTracingSceneWithGeometryInstances
-{
-	FRayTracingSceneRHIRef Scene;
-	uint32 NumNativeGPUSceneInstances;
-	uint32 NumNativeCPUInstances;
-	uint32 TotalNumSegments;
-	// index of each instance geometry in FRayTracingSceneRHIRef ReferencedGeometries
-	TArray<uint32> InstanceGeometryIndices;
-	// base offset of each instance entries in the instance upload buffer
-	TArray<uint32> BaseUploadBufferOffsets;
-	// prefix sum of `Instance.NumTransforms` for all instances in this scene
-	TArray<uint32> BaseInstancePrefixSum;
-	UE_DEPRECATED(5.5, "GPUInstances no longer supported. Use GPUSceneInstances instead.")
-	TArray<FRayTracingGPUInstance> GPUInstances;
-
-	// Unique list of geometries referenced by all instances in this scene.
-	// Any referenced geometry is kept alive while the scene is alive.
-	TArray<FRHIRayTracingGeometry*> ReferencedGeometries;
-	// One entry per instance
-	TArray<FRHIRayTracingGeometry*> PerInstanceGeometries;
-
-	PRAGMA_DISABLE_DEPRECATION_WARNINGS
-	FRayTracingSceneWithGeometryInstances() = default;
-	FRayTracingSceneWithGeometryInstances(const FRayTracingSceneWithGeometryInstances&) = default;
-	FRayTracingSceneWithGeometryInstances& operator=(const FRayTracingSceneWithGeometryInstances&) = default;
-	FRayTracingSceneWithGeometryInstances(FRayTracingSceneWithGeometryInstances&&) = default;
-	FRayTracingSceneWithGeometryInstances& operator=(FRayTracingSceneWithGeometryInstances&&) = default;
-	~FRayTracingSceneWithGeometryInstances() = default;
-	PRAGMA_ENABLE_DEPRECATION_WARNINGS
-};
-
-// Helper function to create FRayTracingSceneRHI using array of high level instances
-// Also outputs data required to build the instance buffer
-UE_DEPRECATED(5.5, "Use FRHIShaderBindingTable instead.")
-RENDERER_API FRayTracingSceneWithGeometryInstances CreateRayTracingSceneWithGeometryInstances(
-	TConstArrayView<FRayTracingGeometryInstance> Instances,
-	uint8 NumLayers,
-	uint32 NumShaderSlotsPerGeometrySegment,
-	uint32 NumMissShaderSlots,
-	uint32 NumCallableShaderSlots = 0,
-	ERayTracingAccelerationStructureFlags BuildFlags = ERayTracingAccelerationStructureFlags::FastTrace);
-
 // Helper function to fill upload buffers required by BuildRayTracingInstanceBuffer with instance descriptors
 // Transforms of CPU instances are copied to OutTransformData
 RENDERER_API void FillRayTracingInstanceUploadBuffer(
@@ -121,18 +79,6 @@ RENDERER_API void FillRayTracingInstanceUploadBuffer(
 	TArrayView<FRayTracingInstanceDescriptorInput> OutInstanceUploadData,
 	TArrayView<FVector4f> OutTransformData);
 
-UE_DEPRECATED(5.5, "Must specify BaseInstancePrefixSum.")
-RENDERER_API void FillRayTracingInstanceUploadBuffer(
-	FRayTracingSceneRHIRef RayTracingSceneRHI,
-	FVector PreViewTranslation,
-	TConstArrayView<FRayTracingGeometryInstance> Instances,
-	TConstArrayView<uint32> InstanceGeometryIndices,
-	TConstArrayView<uint32> BaseUploadBufferOffsets,
-	uint32 NumNativeGPUSceneInstances,
-	uint32 NumNativeCPUInstances,
-	TArrayView<FRayTracingInstanceDescriptorInput> OutInstanceUploadData,
-	TArrayView<FVector4f> OutTransformData);
-
 RENDERER_API void BuildRayTracingInstanceBuffer(
 	FRHICommandList& RHICmdList,
 	const FGPUScene* GPUScene,
@@ -147,27 +93,4 @@ RENDERER_API void BuildRayTracingInstanceBuffer(
 	FUnorderedAccessViewRHIRef OutputStatsUAV,
 	FUnorderedAccessViewRHIRef DebugInstanceGPUSceneIndexUAV);
 
-UE_DEPRECATED(5.5, "GPUInstances no longer supported. Use GPUSceneInstances instead.")
-inline void BuildRayTracingInstanceBuffer(
-	FRHICommandList& RHICmdList,
-	const FGPUScene* GPUScene,
-	const FDFVector3& PreViewTranslation,
-	FUnorderedAccessViewRHIRef InstancesUAV,
-	FShaderResourceViewRHIRef InstanceUploadSRV,
-	FShaderResourceViewRHIRef AccelerationStructureAddressesSRV,
-	FShaderResourceViewRHIRef CPUInstanceTransformSRV,
-	uint32 NumNativeGPUSceneInstances,
-	uint32 NumNativeCPUInstances,
-	TConstArrayView<FRayTracingGPUInstance> GPUInstances,
-	const FRayTracingCullingParameters* CullingParameters,
-	FUnorderedAccessViewRHIRef DebugInstanceGPUSceneIndexUAV)
-{
-	BuildRayTracingInstanceBuffer(RHICmdList, GPUScene, PreViewTranslation,
-		InstancesUAV, InstanceUploadSRV, AccelerationStructureAddressesSRV, CPUInstanceTransformSRV,
-		NumNativeGPUSceneInstances, NumNativeCPUInstances,
-		CullingParameters,
-		nullptr,
-		DebugInstanceGPUSceneIndexUAV);
-}
-
 #endif // RHI_RAYTRACING
\ No newline at end of file
