Fix directional light contribution when no dynamic shadow map channel is assigned.

diff --git a/Engine/Source/Runtime/Renderer/Private/MobileBasePassRendering.cpp b/Engine/Source/Runtime/Renderer/Private/MobileBasePassRendering.cpp
--- a/Engine/Source/Runtime/Renderer/Private/MobileBasePassRendering.cpp
+++ b/Engine/Source/Runtime/Renderer/Private/MobileBasePassRendering.cpp
@@ -435,18 +435,13 @@ void SetupMobileDirectionalLightUniformParameters(
 			(ShadowMapChannel == 1 ? 2 : 0) |
 			(ShadowMapChannel == 2 ? 4 : 0) |
 			(ShadowMapChannel == 3 ? 8 : 0) |
-			(DynamicShadowMapChannel == 0 ? 16 : 0) |
+			// The dynamic shadow map channel needs to be set to 16 even when not present
+			// or the light's contribution will be zeroed out in the shader.
+			(DynamicShadowMapChannel <= 0 ? 16 : 0) |
 			(DynamicShadowMapChannel == 1 ? 32 : 0) |
 			(DynamicShadowMapChannel == 2 ? 64 : 0) |
 			(DynamicShadowMapChannel == 3 ? 128 : 0);
 
-		// When dynamic shadow casting is disabled for the light we still need to set its 
-		// dynamic shadow map channel to 1 or the light's contribution will be zeroed out.
-		if (!Light->Proxy->CastsDynamicShadow())
-		{
-			Params.DirectionalLightShadowMapChannelMask |= (DynamicShadowMapChannel == -1 ? 16 : 0);
-		}
-
 		if (bDynamicShadows && VisibleLightInfos.IsValidIndex(Light->Id) && VisibleLightInfos[Light->Id].AllProjectedShadows.Num() > 0)
 		{
 			const TArray<FProjectedShadowInfo*, SceneRenderingAllocator>& DirectionalLightShadowInfos = VisibleLightInfos[Light->Id].AllProjectedShadows;
