Change GetBlueNoiseGlobalParameters to fallback to dummy safe parameters when BlueNoiseScalarTexture isn't present, following the behaviour of GetBlueNoiseParametersForView
#rnx

diff --git a/Engine/Source/Runtime/Renderer/Private/BlueNoise.cpp b/Engine/Source/Runtime/Renderer/Private/BlueNoise.cpp
--- a/Engine/Source/Runtime/Renderer/Private/BlueNoise.cpp
+++ b/Engine/Source/Runtime/Renderer/Private/BlueNoise.cpp
@@ -58,7 +58,15 @@ FBlueNoiseParameters GetBlueNoiseParameters()
 FBlueNoise GetBlueNoiseGlobalParameters()
 {
 	FBlueNoise Out;
-	Out.BlueNoise = GetBlueNoiseParameters();
+	if (GEngine->BlueNoiseScalarTexture != nullptr && GEngine->BlueNoiseScalarTexture->GetResource() != nullptr)		
+	{
+		Out.BlueNoise = GetBlueNoiseParameters();
+	}
+	else
+	{
+		ensure(IsRunningCommandlet() && IsAllowCommandletRendering());	//	If running a commandlet, the load path won't be visited so the blue noise textures not present. Allow the fallback only in that case.
+		Out.BlueNoise = GetBlueNoiseDummyParameters();
+	}
 	return Out;
 }
 
