Unshelved from pending changelist '38962439':

PR #12681: Fix a spelling error in DeferredShadingRenderer.cpp.

#rnx

diff --git a/Engine/Source/Runtime/Renderer/Private/DeferredShadingRenderer.cpp b/Engine/Source/Runtime/Renderer/Private/DeferredShadingRenderer.cpp
--- a/Engine/Source/Runtime/Renderer/Private/DeferredShadingRenderer.cpp
+++ b/Engine/Source/Runtime/Renderer/Private/DeferredShadingRenderer.cpp
@@ -3021,7 +3021,7 @@ void FDeferredShadingSceneRenderer::Render(FRDGBuilder& GraphBuilder, const FSce
 		FRDGTextureRef LightShaftOcclusionTexture = nullptr;
 		const bool bShouldRenderSingleLayerWater = !bHasRayTracedOverlay && ShouldRenderSingleLayerWater(Views);
 		FSceneWithoutWaterTextures SceneWithoutWaterTextures;
-		auto RenderLigthShaftSkyFogAndCloud = [&]()
+		auto RenderLightShaftSkyFogAndCloud = [&]()
 		{
 			// Draw Lightshafts
 			if (!bHasRayTracedOverlay && ViewFamily.EngineShowFlags.LightShafts)
@@ -3096,7 +3096,7 @@ void FDeferredShadingSceneRenderer::Render(FRDGBuilder& GraphBuilder, const FSce
 		{
 			if (bIsCameraUnderWater)
 			{
-				RenderLigthShaftSkyFogAndCloud();
+				RenderLightShaftSkyFogAndCloud();
 
 				RDG_CSV_STAT_EXCLUSIVE_SCOPE(GraphBuilder, RenderTranslucency);
 				SCOPED_NAMED_EVENT(RenderTranslucency, FColor::Emerald);
@@ -3123,7 +3123,7 @@ void FDeferredShadingSceneRenderer::Render(FRDGBuilder& GraphBuilder, const FSce
 
 		if (!bIsCameraUnderWater)
 		{
-			RenderLigthShaftSkyFogAndCloud();
+			RenderLightShaftSkyFogAndCloud();
 		}
 
 		FRDGTextureRef ExposureIlluminance = nullptr;
