RayTracing: Fix unity builds

#rnx

diff --git a/Engine/Source/Runtime/Renderer/Private/RayTracing/RayTracingMaterialHitShaders.cpp b/Engine/Source/Runtime/Renderer/Private/RayTracing/RayTracingMaterialHitShaders.cpp
--- a/Engine/Source/Runtime/Renderer/Private/RayTracing/RayTracingMaterialHitShaders.cpp
+++ b/Engine/Source/Runtime/Renderer/Private/RayTracing/RayTracingMaterialHitShaders.cpp
@@ -125,7 +125,7 @@ public:
 	}
 };
 
-static bool NeedsAnyHitShader(EBlendMode BlendMode)
+static bool RTNeedsAnyHitShader(EBlendMode BlendMode)
 {
 	switch (BlendMode)
 	{
@@ -166,7 +166,7 @@ public:
 		}
 
 		static FShaderPlatformCachedIniValue<int32> CVarCompileMaterialAHS(TEXT("r.RayTracing.CompileMaterialAHS"));
-		const bool bWantAnyHitShader = CVarCompileMaterialAHS.Get(Parameters.Platform) != 0 && NeedsAnyHitShader(Parameters.MaterialParameters.BlendMode);
+		const bool bWantAnyHitShader = CVarCompileMaterialAHS.Get(Parameters.Platform) != 0 && RTNeedsAnyHitShader(Parameters.MaterialParameters.BlendMode);
 		const bool bSupportProceduralPrimitive = Parameters.VertexFactoryType->SupportsRayTracingProceduralPrimitive() && FDataDrivenShaderPlatformInfo::GetSupportsRayTracingProceduralPrimitive(Parameters.Platform);
 
 		return IsSupportedVertexFactoryType(Parameters.VertexFactoryType)
@@ -350,7 +350,7 @@ static bool GetMaterialHitShader(const FMaterial& RESTRICT MaterialResource, con
 	FMaterialShaderTypes ShaderTypes;
 	const FVertexFactoryType* VFType = VertexFactory->GetType();
 	const bool bUseIntersectionShader = VFType->HasFlags(EVertexFactoryFlags::SupportsRayTracingProceduralPrimitive) && FDataDrivenShaderPlatformInfo::GetSupportsRayTracingProceduralPrimitive(GMaxRHIShaderPlatform);
-	const bool UseAnyHitShader = (MaterialResource.IsMasked() || NeedsAnyHitShader(MaterialResource.GetBlendMode())) && GCompileRayTracingMaterialAHS;
+	const bool UseAnyHitShader = (MaterialResource.IsMasked() || RTNeedsAnyHitShader(MaterialResource.GetBlendMode())) && GCompileRayTracingMaterialAHS;
 
 	GetMaterialHitShader_AnyHit_Intersection_TextureLOD<LightMapPolicyType>(ShaderTypes, UseAnyHitShader, bUseIntersectionShader, UseTextureLod);
 
