Path Tracing: Add support for indirect lighting scale

diff --git a/Engine/Shaders/Private/PathTracing/Light/PathTracingLightCommon.ush b/Engine/Shaders/Private/PathTracing/Light/PathTracingLightCommon.ush
--- a/Engine/Shaders/Private/PathTracing/Light/PathTracingLightCommon.ush
+++ b/Engine/Shaders/Private/PathTracing/Light/PathTracingLightCommon.ush
@@ -367,6 +367,12 @@ float GetLightDiffuseScale(int LightId)
 	return UnpackFloat2FromUInt(SceneLights[LightId].DiffuseSpecularScale).x;
 }
 
+float GetLightIndirectScale(int LightId, float PathRoughness)
+{
+	// Directly visible hits and sharp reflections should not be scaled, but as the roughness increases, we want to apply the indirect lighting scale
+	return lerp(1.0, SceneLights[LightId].IndirectLightingScale, saturate(4.0 * PathRoughness));
+}
+
 uint GetLightMissShaderIndex(int LightId)
 {
 	return SceneLights[LightId].MissShaderIndex;
diff --git a/Engine/Shaders/Private/PathTracing/PathTracingCore.ush b/Engine/Shaders/Private/PathTracing/PathTracingCore.ush
--- a/Engine/Shaders/Private/PathTracing/PathTracingCore.ush
+++ b/Engine/Shaders/Private/PathTracing/PathTracingCore.ush
@@ -1866,6 +1866,7 @@ bool PathTracingKernel(inout FPathState PathState, int Bounce)
 			{
 				LightEstimate *= GetVolumetricScatteringIntensity(LightIndex);
 			}
+			LightEstimate *= GetLightIndirectScale(LightIndex, PathState.GetPathRoughness());
 			LightPickingCdfSum += LightEstimate;
 			LightPickingCdf[Index] = LightPickingCdfSum;
 		}
@@ -1902,6 +1903,8 @@ bool PathTracingKernel(inout FPathState PathState, int Bounce)
 					LightSample.RadianceOverPdf *= GetVolumetricScatteringIntensity(LightId);
 				}
 
+				LightSample.RadianceOverPdf *= GetLightIndirectScale(LightId, PathState.GetPathRoughness());
+
 				// #dxr_todo: Is it cheaper to fire the ray first? Or eval the material first?
 				if (any(LightSample.RadianceOverPdf > 0))
 				{
@@ -2057,6 +2060,8 @@ bool PathTracingKernel(inout FPathState PathState, int Bounce)
 				LightContrib *= MISWeightPower(MaterialSample.Pdf, LightResult.Pdf * LightPickPdf);
 			}
 
+			LightContrib *= GetLightIndirectScale(LightId, PathState.GetPathRoughness());
+
 			if (any(LightContrib > 0))
 			{
 				FRayDesc LightRay = PathState.Ray;
diff --git a/Engine/Shaders/Shared/RayTracingTypes.h b/Engine/Shaders/Shared/RayTracingTypes.h
--- a/Engine/Shaders/Shared/RayTracingTypes.h
+++ b/Engine/Shaders/Shared/RayTracingTypes.h
@@ -52,6 +52,7 @@ struct FPathTracingLight
 	float2  Dimensions; // Radius,Length or RectWidth,RectHeight or Sin(Angle/2),0 depending on light type
 	float2  Shaping;    // Barndoor controls for RectLights, Cone angles for spot lights
 	uint    DiffuseSpecularScale;
+	float   IndirectLightingScale;
 	float   Attenuation;
 	float   FalloffExponent; // for non-inverse square decay lights only
 	float   VolumetricScatteringIntensity;  // scale for volume contributions
@@ -61,7 +62,7 @@ struct FPathTracingLight
 	float2  RectLightAtlasUVScale;  // Rect. light atlas UV transformation
 	float2  RectLightAtlasUVOffset; // Rect. light atlas UV transformation
 };
-HLSL_STATIC_ASSERT(sizeof(FPathTracingLight) == 108, "Path tracing light structure should be kept as small as possible");
+HLSL_STATIC_ASSERT(sizeof(FPathTracingLight) == 112, "Path tracing light structure should be kept as small as possible");
 
 struct FPathTracingPackedPathState
 {
diff --git a/Engine/Source/Runtime/Renderer/Private/PathTracing.cpp b/Engine/Source/Runtime/Renderer/Private/PathTracing.cpp
--- a/Engine/Source/Runtime/Renderer/Private/PathTracing.cpp
+++ b/Engine/Source/Runtime/Renderer/Private/PathTracing.cpp
@@ -2737,6 +2737,7 @@ void SetLightParameters(
 		DestLight.Flags |= Scene->SkyLight->bCastShadows ? PATHTRACER_FLAG_CAST_SHADOW_MASK : 0;
 		DestLight.Flags |= Scene->SkyLight->bCastVolumetricShadow ? PATHTRACER_FLAG_CAST_VOL_SHADOW_MASK : 0;
 		DestLight.DiffuseSpecularScale = PackRG16(1.f, 1.f);
+		DestLight.IndirectLightingScale = Scene->SkyLight->IndirectLightingIntensity;
 		DestLight.VolumetricScatteringIntensity = Scene->SkyLight->VolumetricScatteringIntensity;
 		DestLight.IESAtlasIndex = INDEX_NONE;
 		DestLight.MissShaderIndex = 0;
@@ -2808,6 +2809,7 @@ void SetLightParameters(
 			DestLight.Tangent = LightParameters.Tangent;
 			DestLight.Shaping = FVector2f(0.0f, 0.0f);
 			DestLight.DiffuseSpecularScale = PackRG16(LightParameters.DiffuseScale, LightParameters.SpecularScale);
+			DestLight.IndirectLightingScale = Light.LightSceneInfo->Proxy->GetIndirectLightingScale();
 			DestLight.Attenuation = LightParameters.InvRadius;
 			DestLight.FalloffExponent = 0;
 			DestLight.VolumetricScatteringIntensity = Light.LightSceneInfo->Proxy->GetVolumetricScatteringIntensity();
@@ -2872,6 +2874,7 @@ void SetLightParameters(
 		DestLight.Tangent = LightParameters.Tangent;
 		DestLight.Shaping = FVector2f(0.0f, 0.0f);
 		DestLight.DiffuseSpecularScale = PackRG16(LightParameters.DiffuseScale, LightParameters.SpecularScale);
+		DestLight.IndirectLightingScale = Light.LightSceneInfo->Proxy->GetIndirectLightingScale();
 		DestLight.Attenuation = LightParameters.InvRadius;
 		DestLight.FalloffExponent = 0;
 		DestLight.VolumetricScatteringIntensity = Light.LightSceneInfo->Proxy->GetVolumetricScatteringIntensity();
