[HWRT] Use wrapped group count in RayTracingBuildInstanceBufferCS to support 4M+ instances.

diff --git a/Engine/Shaders/Private/RayTracing/RayTracingInstanceBufferUtil.usf b/Engine/Shaders/Private/RayTracing/RayTracingInstanceBufferUtil.usf
--- a/Engine/Shaders/Private/RayTracing/RayTracingInstanceBufferUtil.usf
+++ b/Engine/Shaders/Private/RayTracing/RayTracingInstanceBufferUtil.usf
@@ -6,6 +6,7 @@
 
 #include "RayTracingCommon.ush"
 #include "../SceneData.ush"
+#include "../ComputeShaderUtils.ush"
 #include "../Common.ush"
 
 #ifndef THREADGROUP_SIZE
@@ -101,11 +102,11 @@ bool IsTransformElementUnsafe(float Element)
 }
 
 [numthreads(THREADGROUP_SIZE, 1, 1)]
-void RayTracingBuildInstanceBufferCS(uint3 DispatchThreadId : SV_DispatchThreadID)
+void RayTracingBuildInstanceBufferCS(uint GroupIndex : SV_GroupIndex, uint3 GroupId : SV_GroupID)
 {
 	const FDFVector3 PreViewTranslation = MakeDFVector3(PreViewTranslationHigh, PreViewTranslationLow);
 
-	const uint Index = DispatchThreadId.x;
+	const uint Index = GetUnWrappedDispatchThreadId(GroupId, GroupIndex, THREADGROUP_SIZE);
 
 	bool bCulled = false;
 
diff --git a/Engine/Source/Runtime/Renderer/Private/RayTracing/RayTracingInstanceBufferUtil.cpp b/Engine/Source/Runtime/Renderer/Private/RayTracing/RayTracingInstanceBufferUtil.cpp
--- a/Engine/Source/Runtime/Renderer/Private/RayTracing/RayTracingInstanceBufferUtil.cpp
+++ b/Engine/Source/Runtime/Renderer/Private/RayTracing/RayTracingInstanceBufferUtil.cpp
@@ -337,13 +337,15 @@ void BuildRayTracingInstanceBuffer(
 	auto ComputeShader = GetGlobalShaderMap(GMaxRHIFeatureLevel)->GetShader<FRayTracingBuildInstanceBufferCS>(PermutationVector);
 	const int32 GroupSize = FMath::DivideAndRoundUp(PassParams.NumInstances, FRayTracingBuildInstanceBufferCS::ThreadGroupSize);
 
+	const FIntVector GroupCount = FComputeShaderUtils::GetGroupCountWrapped(PassParams.NumInstances, FRayTracingBuildInstanceBufferCS::ThreadGroupSize);
+
 	//ClearUnusedGraphResources(ComputeShader, &PassParams);
 
 	SetComputePipelineState(RHICmdList, ComputeShader.GetComputeShader());
 
 	SetShaderParameters(RHICmdList, ComputeShader, ComputeShader.GetComputeShader(), PassParams);
 
-	DispatchComputeShader(RHICmdList, ComputeShader.GetShader(), GroupSize, 1, 1);
+	DispatchComputeShader(RHICmdList, ComputeShader.GetShader(), GroupCount.X, GroupCount.Y, GroupCount.Z);
 
 	UnsetShaderUAVs(RHICmdList, ComputeShader, ComputeShader.GetComputeShader());
 }
