[HWRT] FillRayTracingInstanceUploadBuffer(...) doesn't need RayTracingSceneRHI parameter.

diff --git a/Engine/Plugins/Experimental/GPULightmass/Source/GPULightmass/Private/LightmapRenderer.cpp b/Engine/Plugins/Experimental/GPULightmass/Source/GPULightmass/Private/LightmapRenderer.cpp
--- a/Engine/Plugins/Experimental/GPULightmass/Source/GPULightmass/Private/LightmapRenderer.cpp
+++ b/Engine/Plugins/Experimental/GPULightmass/Source/GPULightmass/Private/LightmapRenderer.cpp
@@ -1212,7 +1212,6 @@ bool FSceneRenderState::SetupRayTracingScene(FRDGBuilder& GraphBuilder, FSceneUn
 				FRayTracingInstanceDescriptorInput* InstanceUploadData = (FRayTracingInstanceDescriptorInput*)RHICmdList.LockBuffer(InstanceUploadBuffer, 0, InstanceUploadBufferSize, RLM_WriteOnly);
 				FVector4f* TransformUploadData = (FVector4f*)RHICmdList.LockBuffer(TransformUploadBuffer, 0, TransformUploadBufferSize, RLM_WriteOnly);
 				FillRayTracingInstanceUploadBuffer(
-					RayTracingScene,
 					View.ViewMatrices.GetPreViewTranslation(),
 					RayTracingGeometryInstances,
 					RayTracingSceneInitializationData.InstanceGeometryIndices,
diff --git a/Engine/Source/Runtime/Renderer/Private/RayTracing/RayTracingInstanceBufferUtil.cpp b/Engine/Source/Runtime/Renderer/Private/RayTracing/RayTracingInstanceBufferUtil.cpp
--- a/Engine/Source/Runtime/Renderer/Private/RayTracing/RayTracingInstanceBufferUtil.cpp
+++ b/Engine/Source/Runtime/Renderer/Private/RayTracing/RayTracingInstanceBufferUtil.cpp
@@ -96,7 +96,6 @@ FRayTracingSceneInitializationData BuildRayTracingSceneInitializationData(TConst
 }
 
 void FillRayTracingInstanceUploadBuffer(
-	FRayTracingSceneRHIRef RayTracingSceneRHI,
 	FVector PreViewTranslation,
 	TConstArrayView<FRayTracingGeometryInstance> Instances,
 	TConstArrayView<uint32> InstanceGeometryIndices,
diff --git a/Engine/Source/Runtime/Renderer/Private/RayTracing/RayTracingScene.cpp b/Engine/Source/Runtime/Renderer/Private/RayTracing/RayTracingScene.cpp
--- a/Engine/Source/Runtime/Renderer/Private/RayTracing/RayTracingScene.cpp
+++ b/Engine/Source/Runtime/Renderer/Private/RayTracing/RayTracingScene.cpp
@@ -282,7 +282,6 @@ void FRayTracingScene::Create(FRDGBuilder& GraphBuilder, const FViewInfo& View,
 					FVector4f* TransformUploadData = (NumNativeCPUInstances > 0) ? (FVector4f*)RHICmdList.LockBuffer(TransformUploadBuffer, 0, TransformUploadBytes, RLM_WriteOnly) : nullptr;
 
 					FillRayTracingInstanceUploadBuffer(
-						RayTracingSceneRHI,
 						PreViewTranslation,
 						Instances,
 						InstanceGeometryIndices,
diff --git a/Engine/Source/Runtime/Renderer/Private/Tests/RayTracingTestbed.cpp b/Engine/Source/Runtime/Renderer/Private/Tests/RayTracingTestbed.cpp
--- a/Engine/Source/Runtime/Renderer/Private/Tests/RayTracingTestbed.cpp
+++ b/Engine/Source/Runtime/Renderer/Private/Tests/RayTracingTestbed.cpp
@@ -234,7 +234,6 @@ bool RunRayTracingTestbed_RenderThread(const FString& Parameters)
 		FRayTracingInstanceDescriptorInput* InstanceUploadData = (FRayTracingInstanceDescriptorInput*)RHICmdList.LockBuffer(InstanceUploadBuffer, 0, InstanceUploadBufferSize, RLM_WriteOnly);
 		FVector4f* TransformUploadData = (FVector4f*)RHICmdList.LockBuffer(TransformUploadBuffer, 0, TransformUploadBufferSize, RLM_WriteOnly);
 		FillRayTracingInstanceUploadBuffer(
-			RayTracingSceneRHI,
 			FVector::ZeroVector,
 			Instances,
 			RayTracingSceneInitializationData.InstanceGeometryIndices,
diff --git a/Engine/Source/Runtime/Renderer/Public/RayTracingInstanceBufferUtil.h b/Engine/Source/Runtime/Renderer/Public/RayTracingInstanceBufferUtil.h
--- a/Engine/Source/Runtime/Renderer/Public/RayTracingInstanceBufferUtil.h
+++ b/Engine/Source/Runtime/Renderer/Public/RayTracingInstanceBufferUtil.h
@@ -68,7 +68,6 @@ RENDERER_API FRayTracingSceneInitializationData BuildRayTracingSceneInitializati
 // Helper function to fill upload buffers required by BuildRayTracingInstanceBuffer with instance descriptors
 // Transforms of CPU instances are copied to OutTransformData
 RENDERER_API void FillRayTracingInstanceUploadBuffer(
-	FRayTracingSceneRHIRef RayTracingSceneRHI,
 	FVector PreViewTranslation,
 	TConstArrayView<FRayTracingGeometryInstance> Instances,
 	TConstArrayView<uint32> InstanceGeometryIndices,
@@ -79,6 +78,31 @@ RENDERER_API void FillRayTracingInstanceUploadBuffer(
 	TArrayView<FRayTracingInstanceDescriptorInput> OutInstanceUploadData,
 	TArrayView<FVector4f> OutTransformData);
 
+UE_DEPRECATED(5.6, "No longer need to provide RayTracingSceneRHI.")
+inline void FillRayTracingInstanceUploadBuffer(
+	FRayTracingSceneRHIRef RayTracingSceneRHI,
+	FVector PreViewTranslation,
+	TConstArrayView<FRayTracingGeometryInstance> Instances,
+	TConstArrayView<uint32> InstanceGeometryIndices,
+	TConstArrayView<uint32> BaseUploadBufferOffsets,
+	TConstArrayView<uint32> BaseInstancePrefixSum,
+	uint32 NumNativeGPUSceneInstances,
+	uint32 NumNativeCPUInstances,
+	TArrayView<FRayTracingInstanceDescriptorInput> OutInstanceUploadData,
+	TArrayView<FVector4f> OutTransformData)
+{
+	FillRayTracingInstanceUploadBuffer(
+		PreViewTranslation,
+		Instances,
+		InstanceGeometryIndices,
+		BaseUploadBufferOffsets,
+		BaseInstancePrefixSum,
+		NumNativeGPUSceneInstances,
+		NumNativeCPUInstances,
+		OutInstanceUploadData,
+		OutTransformData);
+}
+
 RENDERER_API void BuildRayTracingInstanceBuffer(
 	FRHICommandList& RHICmdList,
 	const FGPUScene* GPUScene,
