Memzero the FRayTracingLocalShaderBindings loose parameter data because it's used for validation check with persistent SBTs

diff --git a/Engine/Source/Runtime/Renderer/Private/MeshPassProcessor.cpp b/Engine/Source/Runtime/Renderer/Private/MeshPassProcessor.cpp
--- a/Engine/Source/Runtime/Renderer/Private/MeshPassProcessor.cpp
+++ b/Engine/Source/Runtime/Renderer/Private/MeshPassProcessor.cpp
@@ -353,6 +353,10 @@ FRayTracingLocalShaderBindings* FMeshDrawShaderBindings::SetRayTracingShaderBind
 	{
 		const FShaderLooseParameterBufferInfo& LooseParameterBuffer = SingleShaderBindings.ParameterMapInfo.LooseParameterBuffers[0];
 		const uint8* LooseDataOffset = SingleShaderBindings.GetLooseDataStart();
+#if DO_CHECK
+		// Memzero the loose parameter data because it's used for validation on persistent SBT entries to make sure binding data is always the same
+		FMemory::Memzero(Bindings.LooseParameterData, LooseParameterDataSize);
+#endif // DO_CHECK
 		for (int32 LooseParameterIndex = 0; LooseParameterIndex < LooseParameterBuffer.Parameters.Num(); LooseParameterIndex++)
 		{
 			FShaderLooseParameterInfo LooseParameter = LooseParameterBuffer.Parameters[LooseParameterIndex];
