Fixed cloud alpha holdout black edge

#rnx
#tests pc dx12 editor

diff --git a/Engine/Shaders/Private/VolumetricCloud.usf b/Engine/Shaders/Private/VolumetricCloud.usf
--- a/Engine/Shaders/Private/VolumetricCloud.usf
+++ b/Engine/Shaders/Private/VolumetricCloud.usf
@@ -1752,6 +1752,9 @@ RWTexture2D<float4> OutCloudColor1;
 #endif
 RWTexture2D<float4> OutCloudDepth;
 RWTextureCube<float4> OutCloudColorCube;
+#if SUPPORT_PRIMITIVE_ALPHA_HOLDOUT
+RWTexture2D<float> OutCloudAlphaHoldout;
+#endif
 
 [numthreads(THREADGROUP_SIZEX, THREADGROUP_SIZEY, 1)]
 void MainCS(uint2 DispatchThreadId : SV_DispatchThreadID)
@@ -1796,31 +1799,26 @@ void MainCS(uint2 DispatchThreadId : SV_DispatchThreadID)
 			if (bBlendCloudColor)
 			{
 				float4 CurColor0 = OutCloudColor0[DispatchThreadId];
+				OutCloudColor0[DispatchThreadId] = float4(CurColor0.rgb * OutColor0.a + OutColor0.rgb, CurColor0.a);
+				// No alpha holdout supported with blening.
+			}
+			else
+			{
+				OutCloudColor0[DispatchThreadId] = OutColor0;
 			#if SUPPORT_PRIMITIVE_ALPHA_HOLDOUT
 				if (View.RenderingReflectionCaptureMask == 0.0f)
 				{
-					if(IsVolumetricCloudHoldout(View.EnvironmentComponentsFlags) && bAccumulateAlphaHoldOut > 0)
+					if (IsVolumetricCloudHoldout(View.EnvironmentComponentsFlags) && bAccumulateAlphaHoldOut > 0)
 					{
-						// Operations for Holdout for when we blend directly with scene color Buffer read to CurColor0.
-						// We maintain the color of the scene since it has alreayd been composed
-						// Then we simply add the holdout value on top. Since the alpha has already been impacted by cloud GetCloudLuminanceTransmittanceOverFog when the clouds ahve been composed.
-						OutCloudColor0[DispatchThreadId] = float4(CurColor0.rgb, CurColor0.a + OutHoldOutCoverage + OutAboveCloudHoldOut);
+						OutCloudAlphaHoldout[DispatchThreadId] = OutHoldOutCoverage + OutAboveCloudHoldOut;
 					}
 					else
 					{
-						OutCloudColor0[DispatchThreadId] = float4(CurColor0.rgb * OutColor0.a + OutColor0.rgb, CurColor0.a + OutAboveCloudHoldOut);
+						OutCloudAlphaHoldout[DispatchThreadId] = OutAboveCloudHoldOut;
 					}
 				}
-			#else
-				{
-					OutCloudColor0[DispatchThreadId] = float4(CurColor0.rgb * OutColor0.a + OutColor0.rgb, CurColor0.a);
-				}
 			#endif
 			}
-			else
-			{
-				OutCloudColor0[DispatchThreadId] = OutColor0;
-			}
 		}
 #if CLOUD_MIN_AND_MAX_DEPTH
 		OutCloudColor1[DispatchThreadId] = OutColor1;
diff --git a/Engine/Shaders/Private/VolumetricRenderTarget.usf b/Engine/Shaders/Private/VolumetricRenderTarget.usf
--- a/Engine/Shaders/Private/VolumetricRenderTarget.usf
+++ b/Engine/Shaders/Private/VolumetricRenderTarget.usf
@@ -479,6 +479,7 @@ Texture2D<float2> VolumetricDepthTexture;
 uint4 VolumetricTextureValidCoordRect;
 float4 VolumetricTextureValidUvRect;
 uint ForwardShadingEnable;
+uint OutputAlphaHoldout;
 
 #if PERMUTATION_RENDER_UNDERWATER_BUFFER || PERMUTATION_COMPOSE_WITH_WATER
 Texture2D WaterLinearDepthTexture;
@@ -1002,6 +1003,12 @@ void ComposeVolumetricRTOverScenePS(
 
 #endif // PERMUTATION_UPSAMPLINGMODE==4 || PERMUTATION_UPSAMPLINGMODE==3 || PERMUTATION_UPSAMPLINGMODE==2
 
+#if SUPPORT_PRIMITIVE_ALPHA_HOLDOUT
+	if (OutputAlphaHoldout)
+	{
+		OutputRt0 = float4(0.0, 0.0, 0.0, OutputRt0.r);
+	}
+#endif
 }
 
 #endif // SHADER_COMPOSE_VOLUMETRICRT
diff --git a/Engine/Source/Runtime/Renderer/Private/DeferredShadingRenderer.cpp b/Engine/Source/Runtime/Renderer/Private/DeferredShadingRenderer.cpp
--- a/Engine/Source/Runtime/Renderer/Private/DeferredShadingRenderer.cpp
+++ b/Engine/Source/Runtime/Renderer/Private/DeferredShadingRenderer.cpp
@@ -2707,8 +2707,7 @@ void FDeferredShadingSceneRenderer::Render(FRDGBuilder& GraphBuilder, const FSce
 
 				bool bSkipVolumetricRenderTarget = false;
 				bool bSkipPerPixelTracing = true;
-				bool bAccumulateAlphaHoldOut = false;
-				RenderVolumetricCloud(GraphBuilder, SceneTextures, bSkipVolumetricRenderTarget, bSkipPerPixelTracing, bAccumulateAlphaHoldOut,
+				RenderVolumetricCloud(GraphBuilder, SceneTextures, bSkipVolumetricRenderTarget, bSkipPerPixelTracing,
 					HalfResolutionDepthCheckerboardMinMaxTexture, QuarterResolutionDepthMinMaxTexture, true, InstanceCullingManager);
 			}
 
@@ -2983,8 +2982,7 @@ void FDeferredShadingSceneRenderer::Render(FRDGBuilder& GraphBuilder, const FSce
 				// Generate the volumetric cloud render target
 				bool bSkipVolumetricRenderTarget = false;
 				bool bSkipPerPixelTracing = true;
-				bool bAccumulateAlphaHoldOut = false;
-				RenderVolumetricCloud(GraphBuilder, SceneTextures, bSkipVolumetricRenderTarget, bSkipPerPixelTracing, bAccumulateAlphaHoldOut,
+				RenderVolumetricCloud(GraphBuilder, SceneTextures, bSkipVolumetricRenderTarget, bSkipPerPixelTracing,
 					HalfResolutionDepthCheckerboardMinMaxTexture, QuarterResolutionDepthMinMaxTexture, false, InstanceCullingManager);
 			}
 			// Reconstruct the volumetric cloud render target to be ready to compose it over the scene
@@ -3065,8 +3063,7 @@ void FDeferredShadingSceneRenderer::Render(FRDGBuilder& GraphBuilder, const FSce
 			{
 				bool bSkipVolumetricRenderTarget = true;
 				bool bSkipPerPixelTracing = false;
-				bool bAccumulateAlphaHoldOut = false;
-				RenderVolumetricCloud(GraphBuilder, SceneTextures, bSkipVolumetricRenderTarget, bSkipPerPixelTracing, bAccumulateAlphaHoldOut,
+				RenderVolumetricCloud(GraphBuilder, SceneTextures, bSkipVolumetricRenderTarget, bSkipPerPixelTracing,
 					HalfResolutionDepthCheckerboardMinMaxTexture, QuarterResolutionDepthMinMaxTexture, false, InstanceCullingManager);
 			}
 
@@ -3078,17 +3075,6 @@ void FDeferredShadingSceneRenderer::Render(FRDGBuilder& GraphBuilder, const FSce
 					GraphBuilder, Views, SceneTextures.Color.Target, SceneTextures.Depth.Target,
 					bComposeWithWater,
 					SceneWithoutWaterTextures, SceneTextures);
-
-				if (IsPrimitiveAlphaHoldoutEnabledForAnyView(Views))
-				{
-					// When alpha is enabled to work with holdout. We need another full screen tracing pass to update the alpha channel containing the "holdout alpha throughput".
-					// Alpha hold out only works when using r.volumetricrendertarget.mode 3 which is the mode use by MRQ.
-					bool bSkipVolumetricRenderTarget = true;
-					bool bSkipPerPixelTracing = false;
-					bool bAccumulateAlphaHoldOut = true;
-					RenderVolumetricCloud(GraphBuilder, SceneTextures, bSkipVolumetricRenderTarget, bSkipPerPixelTracing, bAccumulateAlphaHoldOut,
-						HalfResolutionDepthCheckerboardMinMaxTexture, QuarterResolutionDepthMinMaxTexture, false, InstanceCullingManager);
-				}
 			}
 		};
 
@@ -3398,10 +3384,10 @@ void FDeferredShadingSceneRenderer::Render(FRDGBuilder& GraphBuilder, const FSce
 
 		if (bShouldVisualizeVolumetricCloud && !bHasRayTracedOverlay)
 		{
-			RenderVolumetricCloud(GraphBuilder, SceneTextures, false, true, false, HalfResolutionDepthCheckerboardMinMaxTexture, QuarterResolutionDepthMinMaxTexture, false, InstanceCullingManager);
+			RenderVolumetricCloud(GraphBuilder, SceneTextures, false, true, HalfResolutionDepthCheckerboardMinMaxTexture, QuarterResolutionDepthMinMaxTexture, false, InstanceCullingManager);
 			ReconstructVolumetricRenderTarget(GraphBuilder, Views, SceneTextures.Depth.Resolve, HalfResolutionDepthCheckerboardMinMaxTexture, false);
 			ComposeVolumetricRenderTargetOverSceneForVisualization(GraphBuilder, Views, SceneTextures.Color.Target, SceneTextures);
-			RenderVolumetricCloud(GraphBuilder, SceneTextures, true, false, false, HalfResolutionDepthCheckerboardMinMaxTexture, QuarterResolutionDepthMinMaxTexture, false, InstanceCullingManager);
+			RenderVolumetricCloud(GraphBuilder, SceneTextures, true, false, HalfResolutionDepthCheckerboardMinMaxTexture, QuarterResolutionDepthMinMaxTexture, false, InstanceCullingManager);
 		}
 
 		if (!bHasRayTracedOverlay)
diff --git a/Engine/Source/Runtime/Renderer/Private/SceneRendering.h b/Engine/Source/Runtime/Renderer/Private/SceneRendering.h
--- a/Engine/Source/Runtime/Renderer/Private/SceneRendering.h
+++ b/Engine/Source/Runtime/Renderer/Private/SceneRendering.h
@@ -2555,7 +2555,6 @@ protected:
 		const FMinimalSceneTextures& SceneTextures,
 		bool bSkipVolumetricRenderTarget,
 		bool bSkipPerPixelTracing,
-		bool bAccumulateAlphaHoldOut,
 		FRDGTextureRef HalfResolutionDepthCheckerboardMinMaxTexture,
 		FRDGTextureRef QuarterResolutionDepthMinMaxTexture,
 		bool bAsyncCompute,
diff --git a/Engine/Source/Runtime/Renderer/Private/VolumetricCloudRendering.cpp b/Engine/Source/Runtime/Renderer/Private/VolumetricCloudRendering.cpp
--- a/Engine/Source/Runtime/Renderer/Private/VolumetricCloudRendering.cpp
+++ b/Engine/Source/Runtime/Renderer/Private/VolumetricCloudRendering.cpp
@@ -29,6 +29,7 @@
 #include "BasePassRendering.h"
 #include "EnvironmentComponentsFlags.h"
 #include "SceneRendererInterface.h"
+#include "PostProcess/PostProcessing.h" // IsPostProcessingWithAlphaChannelSupported
 
 //  If this is enabled, you also need to touch VolumetricCloud.usf for shaders to be recompiled.
 #define CLOUD_DEBUG_SAMPLES 0 /*!Never check in enabled!*/
@@ -877,6 +878,7 @@ class FRenderVolumetricCloudRenderViewCS : public FMeshMaterialShader
 		SHADER_PARAMETER_RDG_TEXTURE_UAV(RWTexture2D, OutCloudColor1)
 		SHADER_PARAMETER_RDG_TEXTURE_UAV(RWTexture2D, OutCloudDepth)
 		SHADER_PARAMETER_RDG_TEXTURE_UAV(RWTexture2D, OutCloudColorCube)
+		SHADER_PARAMETER_RDG_TEXTURE_UAV(RWTexture2D, OutCloudAlphaHoldout)
 		SHADER_PARAMETER_RDG_TEXTURE(Texture2D, StartTracingDistanceTexture)
 #if CLOUD_DEBUG_SAMPLES
 		SHADER_PARAMETER_STRUCT_INCLUDE(ShaderPrint::FShaderParameters, ShaderPrintParameters)
@@ -2179,11 +2181,15 @@ void FCloudRenderContext::CreateDefaultTexturesIfNeeded(FRDGBuilder& GraphBuilde
 			FRDGTextureDesc::Create2D(FIntPoint(1, 1), PF_G16R16F, FClearValueBinding::None, TexCreate_ShaderResource | TexCreate_RenderTargetable | TexCreate_UAV | NoFastClearFlags),
 			TEXT("Cloud.DepthDummy"));
 
+		DefaultCloudAlphaHoldout = GraphBuilder.CreateTexture(
+			FRDGTextureDesc::Create2D(FIntPoint(1, 1), PF_R16F, FClearValueBinding::None, TexCreate_ShaderResource | TexCreate_RenderTargetable | TexCreate_UAV | NoFastClearFlags),
+			TEXT("Cloud.AlphaHoldoutDummy"));
 
 		DefaultCloudColorCubeTextureUAV = GraphBuilder.CreateUAV(DefaultCloudColorCubeTexture, ERDGUnorderedAccessViewFlags::SkipBarrier);
 		DefaultCloudColor02DTextureUAV = GraphBuilder.CreateUAV(DefaultCloudColor12DTexture, ERDGUnorderedAccessViewFlags::SkipBarrier);
 		DefaultCloudColor12DTextureUAV = GraphBuilder.CreateUAV(DefaultCloudColor02DTexture, ERDGUnorderedAccessViewFlags::SkipBarrier);
 		DefaultCloudDepthTextureUAV = GraphBuilder.CreateUAV(DefaultCloudDepthTexture, ERDGUnorderedAccessViewFlags::SkipBarrier);
+		DefaultCloudAlphaHoldoutUAV = GraphBuilder.CreateUAV(DefaultCloudAlphaHoldout, ERDGUnorderedAccessViewFlags::SkipBarrier);
 	}
 }
 
@@ -2284,7 +2290,8 @@ static void GetOutputTexturesWithFallback(
 	FRDGTextureUAVRef& CloudColorCubeTextureUAV, 
 	FRDGTextureUAVRef& CloudColor0TextureUAV,
 	FRDGTextureUAVRef& CloudColor1TextureUAV,
-	FRDGTextureUAVRef& CloudDepthTextureUAV)
+	FRDGTextureUAVRef& CloudDepthTextureUAV,
+	FRDGTextureUAVRef& CloudAlphaHoldoutUAV)
 {
 	CloudRC.CreateDefaultTexturesIfNeeded(GraphBuilder);
 
@@ -2335,6 +2342,16 @@ static void GetOutputTexturesWithFallback(
 	{
 		CloudDepthTextureUAV = GraphBuilder.CreateUAV(CloudDepthTexture, ERDGUnorderedAccessViewFlags::SkipBarrier);
 	}
+
+	FRDGTextureRef CloudAlphaHoldoutTexture = CloudRC.CloudAlphaHoldoutTexture;
+	if (!CloudAlphaHoldoutTexture)
+	{
+		CloudAlphaHoldoutUAV = CloudRC.DefaultCloudAlphaHoldoutUAV;
+	}
+	else
+	{
+		CloudAlphaHoldoutUAV = GraphBuilder.CreateUAV(CloudAlphaHoldoutTexture, ERDGUnorderedAccessViewFlags::SkipBarrier);
+	}
 }
 
 void FSceneRenderer::RenderVolumetricCloudsInternal(FRDGBuilder& GraphBuilder, FCloudRenderContext& CloudRC, FInstanceCullingManager& InstanceCullingManager, const FIntPoint& CloudViewportSize)
@@ -2461,7 +2478,8 @@ void FSceneRenderer::RenderVolumetricCloudsInternal(FRDGBuilder& GraphBuilder, F
 		FRDGTextureUAVRef CloudColor0TextureUAV;
 		FRDGTextureUAVRef CloudColor1TextureUAV;
 		FRDGTextureUAVRef CloudDepthTextureUAV;
-		GetOutputTexturesWithFallback(GraphBuilder, CloudRC, CloudColorCubeTextureUAV, CloudColor0TextureUAV, CloudColor1TextureUAV, CloudDepthTextureUAV);
+		FRDGTextureUAVRef CloudAlphaHoldoutUAV;
+		GetOutputTexturesWithFallback(GraphBuilder, CloudRC, CloudColorCubeTextureUAV, CloudColor0TextureUAV, CloudColor1TextureUAV, CloudDepthTextureUAV, CloudAlphaHoldoutUAV);
 
 		bool bSampleVirtualShadowMap = (!bCloudDebugViewModeEnabled && !bSkipAtmosphericLightShadowmap && CVarVolumetricCloudShadowSampleAtmosphericLightShadowmap.GetValueOnRenderThread() > 0);
 
@@ -2480,6 +2498,7 @@ void FSceneRenderer::RenderVolumetricCloudsInternal(FRDGBuilder& GraphBuilder, F
 		PassParameters->OutCloudColor1 = CloudColor1TextureUAV;
 		PassParameters->OutCloudDepth = CloudDepthTextureUAV;
 		PassParameters->OutCloudColorCube = CloudColorCubeTextureUAV;
+		PassParameters->OutCloudAlphaHoldout = CloudAlphaHoldoutUAV;
 
 		PassParameters->StartTracingDistanceTextureResolution = FVector2f(EmptySpaceSkippingTexture->Desc.Extent.X, EmptySpaceSkippingTexture->Desc.Extent.Y);
 		PassParameters->StartTracingDistanceTexture = EmptySpaceSkippingTexture;
@@ -2579,7 +2598,6 @@ bool FSceneRenderer::RenderVolumetricCloud(
 	const FMinimalSceneTextures& SceneTextures,
 	bool bSkipVolumetricRenderTarget,
 	bool bSkipPerPixelTracing,
-	bool bAccumulateAlphaHoldOut,
 	FRDGTextureRef HalfResolutionDepthCheckerboardMinMaxTexture,
 	FRDGTextureRef QuarterResolutionDepthMinMaxTexture,
 	bool bAsyncCompute,
@@ -2600,7 +2618,7 @@ bool FSceneRenderer::RenderVolumetricCloud(
 		FMaterialRenderProxy* CloudVolumeMaterialProxy = CloudSceneProxy.GetCloudVolumeMaterial()->GetRenderProxy();
 		if (CloudVolumeMaterialProxy->GetIncompleteMaterialWithFallback(ViewFamily.GetFeatureLevel()).GetMaterialDomain() == MD_Volume)
 		{
-			RDG_EVENT_SCOPE_STAT(GraphBuilder, VolumetricCloud, "%s", bAccumulateAlphaHoldOut ? TEXT("VolumetricCloudAlphaHoldout") : TEXT("VolumetricCloud"));
+			RDG_EVENT_SCOPE_STAT(GraphBuilder, VolumetricCloud, "VolumetricCloud"); 
 			RDG_GPU_STAT_SCOPE(GraphBuilder, VolumetricCloud);
 			SCOPED_NAMED_EVENT(VolumetricCloud, FColor::Emerald);
 
@@ -2629,22 +2647,8 @@ bool FSceneRenderer::RenderVolumetricCloud(
 			{
 				FViewInfo& ViewInfo = Views[ViewIndex];
 
-				if (bAccumulateAlphaHoldOut && ViewInfo.CachedViewUniformShaderParameters->RenderingReflectionCaptureMask > 0.0f)
-				{
-					// We do not render the cloud holdout pass if we are rendering a reflection captures
-					continue;
-				}
-
-				const FIntVector4& EnvironmentComponentFlags = ViewInfo.CachedViewUniformShaderParameters->EnvironmentComponentsFlags;
-				if (bAccumulateAlphaHoldOut && !(IsSkyAtmosphereHoldout(EnvironmentComponentFlags) || IsVolumetricCloudHoldout(EnvironmentComponentFlags) || IsExponentialFogHoldout(EnvironmentComponentFlags)))
-				{
-					// We do not render the cloud holdout pass alpha contribution if not required (when none of the volumetric effects can contribute holdout alpha value).
-					// The cloud transmittance is still applied on the SkyAtmosphere alpha holdout contribution in any case during the regular pass.
-					continue;
-				}
-
 				// RenderTarget are not used by reflection captures for instance.
-				bool bShouldViewRenderVolumetricCloudRenderTarget = ShouldViewRenderVolumetricCloudRenderTarget(ViewInfo) && !bAccumulateAlphaHoldOut; // When rendering alpha holdout per pixel, we never skip even if a volumetric render target is used.; 
+				bool bShouldViewRenderVolumetricCloudRenderTarget = ShouldViewRenderVolumetricCloudRenderTarget(ViewInfo);
 
 				if ((bShouldViewRenderVolumetricCloudRenderTarget && bSkipVolumetricRenderTarget) || (!bShouldViewRenderVolumetricCloudRenderTarget && bSkipPerPixelTracing))
 				{
@@ -2719,6 +2723,7 @@ bool FSceneRenderer::RenderVolumetricCloud(
 				CloudRC.NoiseFrameIndexModPattern = ViewInfo.CachedViewUniformShaderParameters->StateFrameIndexMod8;
 
 				FIntPoint CloudDefaultRtSize(0, 0);
+				bool bAccumulateAlphaHoldOut = false;
 
 				if (bShouldViewRenderVolumetricCloudRenderTarget)
 				{
@@ -2726,6 +2731,13 @@ bool FSceneRenderer::RenderVolumetricCloud(
 					DestinationRT = VRT.GetOrCreateVolumetricTracingRT(GraphBuilder);
 					DestinationRTDepth = VRT.GetOrCreateVolumetricTracingRTDepth(GraphBuilder);
 
+					// Only account for alphaholdout if any important primitive is holdout, and we are in mode 3
+					if (ViewFamily.EngineShowFlags.AllowPrimitiveAlphaHoldout && IsPostProcessingWithAlphaChannelSupported() && ViewInfo.CachedViewUniformShaderParameters->RenderingReflectionCaptureMask == 0.0f)
+					{
+						const FIntVector4& EnvironmentComponentFlags = ViewInfo.CachedViewUniformShaderParameters->EnvironmentComponentsFlags;
+						bAccumulateAlphaHoldOut = VRT.GetMode() == 3 &&  (IsSkyAtmosphereHoldout(EnvironmentComponentFlags) || IsVolumetricCloudHoldout(EnvironmentComponentFlags) || IsExponentialFogHoldout(EnvironmentComponentFlags));
+					}
+
 					CloudDefaultRtSize = DestinationRT->Desc.Extent;
 
 					if (bShouldUseHighQualityAerialPerspective)
@@ -2759,6 +2771,12 @@ bool FSceneRenderer::RenderVolumetricCloud(
 					CloudRC.TracingCoordToZbufferCoordScaleBias.W += ViewInfo.CachedViewUniformShaderParameters->ViewRectMin.Y / ((VRT.GetMode() == 0 || VRT.GetMode() == 3) ? 2 : 1);
 					CloudRC.TracingCoordToFullResPixelCoordScaleBias = VRT.GetTracingCoordToFullResPixelCoordScaleBias();
 
+					if (bAccumulateAlphaHoldOut)
+					{
+						// Requesting the holdout texture will automatically trigger the composition of holdout into alpha during the composition stage.
+						CloudRC.CloudAlphaHoldoutTexture = VRT.GetOrCreateVolumetricTracingRTHoldout(GraphBuilder);
+					}
+
 					const bool bShouldVolumetricCloudTraceWithMinMaxDepth = ShouldVolumetricCloudTraceWithMinMaxDepth(ViewInfo);
 					if (bShouldVolumetricCloudTraceWithMinMaxDepth)
 					{
diff --git a/Engine/Source/Runtime/Renderer/Private/VolumetricCloudRendering.h b/Engine/Source/Runtime/Renderer/Private/VolumetricCloudRendering.h
--- a/Engine/Source/Runtime/Renderer/Private/VolumetricCloudRendering.h
+++ b/Engine/Source/Runtime/Renderer/Private/VolumetricCloudRendering.h
@@ -137,6 +137,7 @@ struct FCloudRenderContext
 	FRenderTargetBindingSlots RenderTargets;
 
 	FRDGTextureRef SecondaryCloudTracingDataTexture = nullptr;
+	FRDGTextureRef CloudAlphaHoldoutTexture = nullptr;
 
 	bool bDisableCloudBlending;
 	bool bShouldViewRenderVolumetricRenderTarget;
@@ -164,10 +165,12 @@ struct FCloudRenderContext
 	FRDGTextureRef DefaultCloudColor02DTexture = nullptr;
 	FRDGTextureRef DefaultCloudColor12DTexture = nullptr;
 	FRDGTextureRef DefaultCloudDepthTexture = nullptr;
+	FRDGTextureRef DefaultCloudAlphaHoldout = nullptr;
 	FRDGTextureUAVRef DefaultCloudColorCubeTextureUAV = nullptr;
 	FRDGTextureUAVRef DefaultCloudColor02DTextureUAV = nullptr;
 	FRDGTextureUAVRef DefaultCloudColor12DTextureUAV = nullptr;
 	FRDGTextureUAVRef DefaultCloudDepthTextureUAV = nullptr;
+	FRDGTextureUAVRef DefaultCloudAlphaHoldoutUAV = nullptr;
 
 	FRDGTextureUAVRef ComputeOverlapCloudColorCubeTextureUAVWithoutBarrier = nullptr;
 
diff --git a/Engine/Source/Runtime/Renderer/Private/VolumetricRenderTarget.cpp b/Engine/Source/Runtime/Renderer/Private/VolumetricRenderTarget.cpp
--- a/Engine/Source/Runtime/Renderer/Private/VolumetricRenderTarget.cpp
+++ b/Engine/Source/Runtime/Renderer/Private/VolumetricRenderTarget.cpp
@@ -176,6 +176,7 @@ FVolumetricRenderTargetViewStateData::FVolumetricRenderTargetViewStateData()
 	: CurrentRT(1)
 	, bFirstTimeUsed(true)
 	, bHistoryValid(false)
+	, bHoldoutValid(false)
 	, PreViewExposure(1.0f)
 	, FullResolution(FIntPoint::ZeroValue)
 	, VolumetricReconstructRTResolution(FIntPoint::ZeroValue)
@@ -211,6 +212,8 @@ void FVolumetricRenderTargetViewStateData::Initialise(
 	Mode = FMath::Clamp(InMode, 0, 3);
 	UpsamplingMode = Mode == 2 || Mode == 3 ? 2 : FMath::Clamp(InUpsamplingMode, 0, 4); // if we are using mode 2 then we cannot intersect with depth and upsampling should be 2 (simple on/off intersection)
 
+	bHoldoutValid = false;
+
 	if (bFirstTimeUsed || bCameraCut)
 	{
 		bFirstTimeUsed = false;
@@ -298,6 +301,7 @@ void FVolumetricRenderTargetViewStateData::Reset()
 {
 	bFirstTimeUsed = false;
 	bHistoryValid = false;
+	bHoldoutValid = false;
 	PreViewExposure = 1.0f;
 	FrameId = 0;
 	NoiseFrameIndex = 0;
@@ -363,6 +367,21 @@ FRDGTextureRef FVolumetricRenderTargetViewStateData::GetOrCreateVolumetricTracin
 	return GraphBuilder.RegisterExternalTexture(VolumetricTracingRTDepth);
 }
 
+FRDGTextureRef FVolumetricRenderTargetViewStateData::GetOrCreateVolumetricTracingRTHoldout(FRDGBuilder& GraphBuilder)
+{
+	check(FullResolution != FIntPoint::ZeroValue); // check that initialization has been done at least once
+
+	if (!VolumetricTracingRTHoldout.IsValid())
+	{
+		FPooledRenderTargetDesc Desc = FPooledRenderTargetDesc::Create2DDesc(
+			VolumetricTracingRTResolution, PF_R16F, FClearValueBinding(FLinearColor(0.0f, 0.0f, 0.0f, 0.0f)),
+			TexCreate_None, TexCreate_ShaderResource | TexCreate_UAV, false);
+		GRenderTargetPool.FindFreeElement(GraphBuilder.RHICmdList, Desc, VolumetricTracingRTHoldout, TEXT("VolumetricRenderTarget.TracingHoldOut"));
+	}
+	bHoldoutValid = true;
+	return GraphBuilder.RegisterExternalTexture(VolumetricTracingRTHoldout);
+}
+
 FVector2f FVolumetricRenderTargetViewStateData::GetVolumetricTracingUVScale() const
 {
 	return FVector2f(float(VolumetricTracingViewRect.X) / float(VolumetricTracingRTResolution.X), 
@@ -780,6 +799,7 @@ class FComposeVolumetricRTOverScenePS : public FGlobalShader
 		SHADER_PARAMETER_SAMPLER(SamplerState, WaterLinearDepthSampler)
 		RENDER_TARGET_BINDING_SLOTS()
 		SHADER_PARAMETER(uint32, ForwardShadingEnable)
+		SHADER_PARAMETER(uint32, OutputAlphaHoldout)
 		SHADER_PARAMETER(float, UvOffsetSampleAcceptanceWeight)
 		SHADER_PARAMETER(FVector4f, VolumetricTextureSizeAndInvSize)
 		SHADER_PARAMETER(FVector2f, FullResolutionToVolumetricBufferResolutionScale)
@@ -936,6 +956,7 @@ void ComposeVolumetricRenderTargetOverScene(
 		PassParameters->FullResolutionToVolumetricBufferResolutionScale = GetCompositionFullResolutionToVolumetricBufferResolutionScale(VRTMode) * GetVolumetricBufferResolutionScale(VRTMode);
 		PassParameters->SceneTextures = SceneTextures.UniformBuffer;
 		PassParameters->Substrate = Substrate::BindSubstrateGlobalUniformParameters(ViewInfo);
+		PassParameters->OutputAlphaHoldout = 0;
 		PassParameters->ForwardShadingEnable = bForwardShading ? 1 : 0;
 		FVector2f DummyUVScale;
 		GetTextureSafeUvCoordBound(VolumetricViewRect, PassParameters->VolumetricTexture, PassParameters->VolumetricTextureValidCoordRect, PassParameters->VolumetricTextureValidUvRect, DummyUVScale);
@@ -957,8 +978,29 @@ void ComposeVolumetricRenderTargetOverScene(
 		PassParameters->LFV = ViewInfo.LocalFogVolumeViewData.UniformParametersStruct;
 
 		FPixelShaderUtils::AddFullscreenPass<FComposeVolumetricRTOverScenePS>(
-			GraphBuilder, ViewInfo.ShaderMap, RDG_EVENT_NAME("VolumetricComposeOverScene"), PixelShader, PassParameters, ViewInfo.ViewRect,
+			GraphBuilder, ViewInfo.ShaderMap, RDG_EVENT_NAME("VolumetricComposeCloudOverScene"), PixelShader, PassParameters, ViewInfo.ViewRect,
 			PreMultipliedColorTransmittanceBlend);
+
+		if (bSupportsAlpha && VolumetricCloudRT.GetHoldoutValid())
+		{
+			check(VRTMode == 3);	// This is the only supported way today.
+
+			// Also compose the alpha value the same way
+			PermutationVector.Set<FComposeVolumetricRTOverScenePS::FApplyFog>(0);
+			PermutationVector.Set<FComposeVolumetricRTOverScenePS::FApplyLocalFogVolume>(0);
+			TShaderMapRef<FComposeVolumetricRTOverScenePS> HoldoutPixelShader(ViewInfo.ShaderMap, PermutationVector);
+
+			FComposeVolumetricRTOverScenePS::FParameters* HoldoutPassParameters = GraphBuilder.AllocParameters<FComposeVolumetricRTOverScenePS::FParameters>();
+			*HoldoutPassParameters = *PassParameters;
+			HoldoutPassParameters->VolumetricTexture = VolumetricCloudRT.GetOrCreateVolumetricTracingRTHoldout(GraphBuilder);
+			HoldoutPassParameters->OutputAlphaHoldout = 1;
+
+			FRHIBlendState* AddAlphaBlendMode = TStaticBlendState<CW_ALPHA, BO_Add, BF_One, BF_One, BO_Add, BF_One, BF_One>::GetRHI();
+
+			FPixelShaderUtils::AddFullscreenPass<FComposeVolumetricRTOverScenePS>(
+				GraphBuilder, ViewInfo.ShaderMap, RDG_EVENT_NAME("VolumetricComposeHoldoutOverScene"), HoldoutPixelShader, HoldoutPassParameters, ViewInfo.ViewRect,
+				AddAlphaBlendMode);
+		}
 	}
 }
 
@@ -1023,6 +1065,7 @@ void ComposeVolumetricRenderTargetOverSceneUnderWater(
 																		WaterPassViewData.ViewRect.Max.X, WaterPassViewData.ViewRect.Max.Y);
 		PassParameters->SceneTextures = SceneTextures.UniformBuffer;
 		PassParameters->Substrate = Substrate::BindSubstrateGlobalUniformParameters(ViewInfo);
+		PassParameters->OutputAlphaHoldout = 0;
 		FVector2f DummyUVScale;
 		GetTextureSafeUvCoordBound(VolumetricViewRect, PassParameters->VolumetricTexture, PassParameters->VolumetricTextureValidCoordRect, PassParameters->VolumetricTextureValidUvRect, DummyUVScale);
 
@@ -1087,6 +1130,7 @@ void ComposeVolumetricRenderTargetOverSceneForVisualization(
 		PassParameters->FullResolutionToVolumetricBufferResolutionScale = GetCompositionFullResolutionToVolumetricBufferResolutionScale(VRTMode) * GetVolumetricBufferResolutionScale(VRTMode);
 		PassParameters->SceneTextures = SceneTextures.UniformBuffer;
 		PassParameters->Substrate = Substrate::BindSubstrateGlobalUniformParameters(ViewInfo);
+		PassParameters->OutputAlphaHoldout = 0;
 		FVector2f DummyUVScale;
 		GetTextureSafeUvCoordBound(VolumetricViewRect, PassParameters->VolumetricTexture, PassParameters->VolumetricTextureValidCoordRect, PassParameters->VolumetricTextureValidUvRect, DummyUVScale);
 
diff --git a/Engine/Source/Runtime/Renderer/Private/VolumetricRenderTargetViewStateData.h b/Engine/Source/Runtime/Renderer/Private/VolumetricRenderTargetViewStateData.h
--- a/Engine/Source/Runtime/Renderer/Private/VolumetricRenderTargetViewStateData.h
+++ b/Engine/Source/Runtime/Renderer/Private/VolumetricRenderTargetViewStateData.h
@@ -43,6 +43,7 @@ public:
 	FRDGTextureRef GetOrCreateVolumetricTracingRT(FRDGBuilder& GraphBuilder);
 	FRDGTextureRef GetOrCreateVolumetricSecondaryTracingRT(FRDGBuilder& GraphBuilder);
 	FRDGTextureRef GetOrCreateVolumetricTracingRTDepth(FRDGBuilder& GraphBuilder);
+	FRDGTextureRef GetOrCreateVolumetricTracingRTHoldout(FRDGBuilder& GraphBuilder);
 	FVector2f GetVolumetricTracingUVScale() const;
 	FVector2f GetVolumetricTracingUVMax() const;
 
@@ -62,6 +63,7 @@ public:
 
 
 	bool GetHistoryValid() const { return bHistoryValid; }
+	bool GetHoldoutValid() const { return bHoldoutValid; }
 	const FIntPoint& GetCurrentVolumetricReconstructRTResolution() const { return VolumetricReconstructRTResolution; }
 	const FIntPoint& GetCurrentVolumetricTracingRTResolution() const { return VolumetricTracingRTResolution; }
 	const FIntPoint& GetCurrentVolumetricTracingViewRect() const { return VolumetricTracingViewRect; }
@@ -87,6 +89,7 @@ private:
 	uint32 CurrentRT;
 	bool bFirstTimeUsed;
 	bool bHistoryValid;
+	bool bHoldoutValid;
 	float PreViewExposure;
 
 	int32 FrameId;
@@ -107,6 +110,7 @@ private:
 	TRefCountPtr<IPooledRenderTarget> VolumetricTracingRT;
 	TRefCountPtr<IPooledRenderTarget> VolumetricSecondaryTracingRT;
 	TRefCountPtr<IPooledRenderTarget> VolumetricTracingRTDepth;
+	TRefCountPtr<IPooledRenderTarget> VolumetricTracingRTHoldout;
 
 	int32 Mode;
 	int32 UpsamplingMode;
