Fixed issue where scene render builder would merge path traced renderers together with non-path traced renderers. This does not work since we are no longer updating the scene with each scene render. Instead, we check if the main scene renderers are compatible with any existing auxiliary renderers and flush if they aren't. This can be extended to support other show flags that can't be merged together.

diff --git a/Engine/Source/Runtime/Engine/Public/SceneRenderBuilderInterface.h b/Engine/Source/Runtime/Engine/Public/SceneRenderBuilderInterface.h
--- a/Engine/Source/Runtime/Engine/Public/SceneRenderBuilderInterface.h
+++ b/Engine/Source/Runtime/Engine/Public/SceneRenderBuilderInterface.h
@@ -13,6 +13,7 @@ class FSceneInterface;
 class FSceneViewFamily;
 class FHitProxyConsumer;
 
+struct FEngineShowFlags;
 struct FSceneRenderUpdateInputs;
 
 enum class ESceneRenderGroupFlags : uint8
@@ -87,6 +88,11 @@ public:
 		AddRenderer(Renderer, {}, MoveTemp(Function));
 	}
 
+	/** Whether a renderer with the provided show flags can be added to the builder. If not compatible, you must call Execute()
+	 *  first to flush existing renderers prior to adding.
+	 */
+	virtual bool IsCompatible(const FEngineShowFlags& EngineShowFlags) const = 0;
+
 	//////////////////////////////////////////////////////////////////////////
 
 	// Call to execute all enqueued commands on the render thread. Resets the builder back to empty.
diff --git a/Engine/Source/Runtime/Renderer/Private/RayTracing/RayTracing.cpp b/Engine/Source/Runtime/Renderer/Private/RayTracing/RayTracing.cpp
--- a/Engine/Source/Runtime/Renderer/Private/RayTracing/RayTracing.cpp
+++ b/Engine/Source/Runtime/Renderer/Private/RayTracing/RayTracing.cpp
@@ -333,15 +333,18 @@ namespace RayTracing
 		return InAllocator.Create<FGatherInstancesTaskData>(Scene, View, InAllocator, FSceneOptions(Scene, ViewFamily, View, DiffuseIndirectMethod, ReflectionsMethod));
 	}
 
-	void OnRenderBegin(FScene& Scene, TArray<FViewInfo>& Views, const FViewFamilyInfo& ViewFamily)
+	void OnRenderBegin(const FSceneRenderUpdateInputs& SceneUpdateInputs)
 	{
-		const ERayTracingType CurrentType = ViewFamily.EngineShowFlags.PathTracing ? ERayTracingType::PathTracing : ERayTracingType::RayTracing;
+		const ERayTracingType CurrentType = EnumHasAnyFlags(SceneUpdateInputs.CommonShowFlags, ESceneRenderCommonShowFlags::PathTracing) ? ERayTracingType::PathTracing : ERayTracingType::RayTracing;
+
 		bool bNaniteCoarseMeshStreamingModeChanged = false;
 #if WITH_EDITOR
 		bNaniteCoarseMeshStreamingModeChanged = Nanite::FCoarseMeshStreamingManager::CheckStreamingMode();
 #endif // WITH_EDITOR
 		const bool bNaniteRayTracingModeChanged = Nanite::GRayTracingManager.CheckModeChanged();
 
+		FScene& Scene = *SceneUpdateInputs.Scene;
+
 		if (CurrentType != Scene.CachedRayTracingMeshCommandsType
 			|| bNaniteCoarseMeshStreamingModeChanged
 			|| bNaniteRayTracingModeChanged
@@ -370,12 +373,12 @@ namespace RayTracing
 
 		if (bNaniteRayTracingModeChanged)
 		{
-			for (FViewInfo& View : Views)
+			for (FViewInfo* View : SceneUpdateInputs.Views)
 			{
-				if (View.ViewState != nullptr && !View.bIsOfflineRender)
+				if (View->ViewState != nullptr && !View->bIsOfflineRender)
 				{
 					// don't invalidate in the offline case because we only get one attempt at rendering each sample
-					View.ViewState->PathTracingInvalidate();
+					View->ViewState->PathTracingInvalidate();
 				}
 			}
 		}
diff --git a/Engine/Source/Runtime/Renderer/Private/RayTracing/RayTracing.h b/Engine/Source/Runtime/Renderer/Private/RayTracing/RayTracing.h
--- a/Engine/Source/Runtime/Renderer/Private/RayTracing/RayTracing.h
+++ b/Engine/Source/Runtime/Renderer/Private/RayTracing/RayTracing.h
@@ -17,6 +17,7 @@ class FScene;
 class FViewInfo;
 class FViewFamilyInfo;
 class FGlobalDynamicReadBuffer;
+struct FSceneRenderUpdateInputs;
 
 // Settings controlling ray tracing instance caching
 namespace RayTracing
@@ -40,7 +41,7 @@ namespace RayTracing
 {
 	struct FGatherInstancesTaskData;
 
-	void OnRenderBegin(FScene& Scene, TArray<FViewInfo>& Views, const FViewFamilyInfo& ViewFamily);
+	void OnRenderBegin(const FSceneRenderUpdateInputs& SceneUpdateInputs);
 
 	FGatherInstancesTaskData* CreateGatherInstancesTaskData(
 		FSceneRenderingBulkObjectAllocator& InAllocator,
diff --git a/Engine/Source/Runtime/Renderer/Private/SceneRenderBuilder.cpp b/Engine/Source/Runtime/Renderer/Private/SceneRenderBuilder.cpp
--- a/Engine/Source/Runtime/Renderer/Private/SceneRenderBuilder.cpp
+++ b/Engine/Source/Runtime/Renderer/Private/SceneRenderBuilder.cpp
@@ -474,7 +474,82 @@ public:
 	TArray<FSceneRenderer*, FConcurrentLinearArrayAllocator> CreateSceneRenderers(
 		TConstArrayView<const FSceneViewFamily*> ViewFamilies,
 		FHitProxyConsumer* HitProxyConsumer,
-		bool bAllowSplitScreenDebug);
+		bool bAllowSplitScreenDebug) const
+	{
+		if (!ViewFamilies.Num())
+		{
+			return {};
+		}
+
+	#if !UE_BUILD_SHIPPING
+		bool bSplitScreenDebug = false;
+		TArray<const FSceneViewFamily*, TFixedAllocator<2>> DebugViewFamilyStorage;
+		if (bAllowSplitScreenDebug && IsSplitScreenDebugEnabled(ViewFamilies))
+		{
+			CreateSplitScreenDebugViewFamilies(*ViewFamilies[0], ViewFamilies, DebugViewFamilyStorage);
+			bSplitScreenDebug = true;
+		}
+	#endif
+
+		const FSceneInterface* SceneInterface = ViewFamilies[0]->Scene;
+		check(SceneInterface);
+		check(SceneInterface->GetRenderScene() == Scene);
+
+		const EShadingPath ShadingPath = GetFeatureLevelShadingPath(SceneInterface->GetFeatureLevel());
+
+		TArray<FSceneRenderer*, FConcurrentLinearArrayAllocator> OutRenderers;
+		OutRenderers.Reserve(ViewFamilies.Num());
+
+		for (const FSceneViewFamily* ViewFamily : ViewFamilies)
+		{
+			check(ViewFamily);
+			check(ViewFamily->Scene == SceneInterface);
+
+			if (ShadingPath == EShadingPath::Deferred)
+			{
+				OutRenderers.Add(new FDeferredShadingSceneRenderer(ViewFamily, HitProxyConsumer));
+			}
+			else
+			{
+				check(ShadingPath == EShadingPath::Mobile);
+				OutRenderers.Add(new FMobileSceneRenderer(ViewFamily, HitProxyConsumer));
+			}
+
+			for (auto& ViewExtension : ViewFamily->ViewExtensions)
+			{
+				ViewExtension->PostCreateSceneRenderer(*ViewFamily, OutRenderers.Last());
+			}
+		}
+
+		// Cache the FXSystem for the duration of the scene render
+		// UWorld::CleanupWorldInternal() will mark the system as pending kill on the GameThread and then enqueue a delete command
+		//-TODO: The call to IsPendingKill should no longer be required as we are caching & using within a single render command
+		FFXSystemInterface* FXSystem = nullptr;
+		if (Scene && Scene->FXSystem && !Scene->FXSystem->IsPendingKill())
+		{
+			FXSystem = Scene->FXSystem;
+		}
+
+		for (FSceneRenderer* Renderer : OutRenderers)
+		{
+			Renderer->Link.Head = OutRenderers[0];
+			Renderer->FXSystem  = FXSystem;
+		}
+
+		for (int32 Index = 1; Index < OutRenderers.Num(); ++Index)
+		{
+			OutRenderers[Index - 1]->Link.Next = OutRenderers[Index];
+		}
+
+	#if !UE_BUILD_SHIPPING
+		if (bSplitScreenDebug)
+		{
+			DestroySplitScreenDebugViewFamilies(ViewFamilies);
+		}
+	#endif
+
+		return OutRenderers;
+	}
 
 	void AddCommand(TUniqueFunction<void()>&& Function)
 	{
@@ -487,23 +562,15 @@ public:
 		check(Renderer->Scene == Scene);
 		check(Renderer->ViewFamily.Views.Num() > 0);
 		check(Renderer->ViewFamily.Views[0]);
-		checkf(Renderers.IsEmpty() || Renderer->ViewFamily.EngineShowFlags.HitProxies == bHitProxyRendering, TEXT("You cannot mix hit proxy renderers with non-hit proxy renderers."));
+		checkf(IsCompatible(Renderer->ViewFamily.EngineShowFlags),
+			TEXT("Renderer contains show flags that are not compatible with other renderers that were previously added. ")
+			TEXT("Use IsCompatible(...) to check if the show flags are compatible"));
 
-		bHitProxyRendering = Renderer->ViewFamily.EngineShowFlags.HitProxies;
-
-	#if DO_CHECK
-		bool bFoundRenderer = false;
-		for (int32 Index = 0; Index < RenderersToAdd.Num(); ++Index)
+		if (Renderers.IsEmpty())
 		{
-			if (RenderersToAdd[Index] == Renderer)
-			{
-				RenderersToAdd.RemoveAtSwap(Index, EAllowShrinking::No);
-				bFoundRenderer = true;
-				break;
-			}
+			CommonShowFlags |= Renderer->ViewFamily.EngineShowFlags.HitProxies ? ESceneRenderCommonShowFlags::HitProxies : ESceneRenderCommonShowFlags::None;
+			CommonShowFlags |= Renderer->ViewFamily.EngineShowFlags.PathTracing ? ESceneRenderCommonShowFlags::PathTracing : ESceneRenderCommonShowFlags::None;
 		}
-		checkf(bFoundRenderer, TEXT("Renderer being added was not created with this scene render builder or is being added twice."));
-	#endif
 
 		FGroupNode* TailGroup = GroupNodes.GetTail();
 
@@ -557,6 +624,17 @@ public:
 		return Allocator;
 	}
 
+	bool IsCompatible(const FEngineShowFlags& EngineShowFlags)
+	{
+		if (Renderers.IsEmpty())
+		{
+			return true;
+		}
+
+		return EnumHasAnyFlags(CommonShowFlags, ESceneRenderCommonShowFlags::HitProxies)  == EngineShowFlags.HitProxies
+			&& EnumHasAnyFlags(CommonShowFlags, ESceneRenderCommonShowFlags::PathTracing) == EngineShowFlags.PathTracing;
+	}
+
 	static void WaitForAsyncDeleteTask()
 	{
 		check(IsInRenderingThread());
@@ -587,10 +665,7 @@ private:
 	TIntrusiveDoubleLinkedList<FGroupNode> GroupNodes;
 	TArray<FOp, FConcurrentLinearArrayAllocator> Ops;
 	FConcurrentLinearBulkObjectAllocator Allocator;
-
-#if DO_CHECK
-	TArray<FSceneRenderer*, FConcurrentLinearArrayAllocator> RenderersToAdd;
-#endif
+	ESceneRenderCommonShowFlags CommonShowFlags = ESceneRenderCommonShowFlags::None;
 
 	struct
 	{
@@ -602,7 +677,6 @@ private:
 	} RenderState;
 
 	bool bInsideGroup : 1 = false;
-	bool bHitProxyRendering : 1 = false;
 };
 
 FSceneRenderProcessor::FAsyncTasks FSceneRenderProcessor::AsyncTasks;
@@ -690,94 +764,9 @@ void FSceneRenderProcessor::Cleanup(FRHICommandListImmediate& RHICmdList, FScene
 	FGraphicsMinimalPipelineStateId::ResetLocalPipelineIdTableSize();
 }
 
-TArray<FSceneRenderer*, FConcurrentLinearArrayAllocator> FSceneRenderProcessor::CreateSceneRenderers(
-	TConstArrayView<const FSceneViewFamily*> ViewFamilies,
-	FHitProxyConsumer* HitProxyConsumer,
-	bool bAllowSplitScreenDebug)
-{
-	if (!ViewFamilies.Num())
-	{
-		return {};
-	}
-
-#if !UE_BUILD_SHIPPING
-	bool bSplitScreenDebug = false;
-	TArray<const FSceneViewFamily*, TFixedAllocator<2>> DebugViewFamilyStorage;
-	if (bAllowSplitScreenDebug && IsSplitScreenDebugEnabled(ViewFamilies))
-	{
-		CreateSplitScreenDebugViewFamilies(*ViewFamilies[0], ViewFamilies, DebugViewFamilyStorage);
-		bSplitScreenDebug = true;
-	}
-#endif
-
-	const FSceneInterface* SceneInterface = ViewFamilies[0]->Scene;
-	check(SceneInterface);
-	check(SceneInterface->GetRenderScene() == Scene);
-
-	const EShadingPath ShadingPath = GetFeatureLevelShadingPath(SceneInterface->GetFeatureLevel());
-
-	TArray<FSceneRenderer*, FConcurrentLinearArrayAllocator> OutRenderers;
-	OutRenderers.Reserve(ViewFamilies.Num());
-
-	for (const FSceneViewFamily* ViewFamily : ViewFamilies)
-	{
-		check(ViewFamily);
-		check(ViewFamily->Scene == SceneInterface);
-
-		if (ShadingPath == EShadingPath::Deferred)
-		{
-			OutRenderers.Add(new FDeferredShadingSceneRenderer(ViewFamily, HitProxyConsumer));
-		}
-		else
-		{
-			check(ShadingPath == EShadingPath::Mobile);
-			OutRenderers.Add(new FMobileSceneRenderer(ViewFamily, HitProxyConsumer));
-		}
-
-		for (auto& ViewExtension : ViewFamily->ViewExtensions)
-		{
-			ViewExtension->PostCreateSceneRenderer(*ViewFamily, OutRenderers.Last());
-		}
-	}
-
-	// Cache the FXSystem for the duration of the scene render
-	// UWorld::CleanupWorldInternal() will mark the system as pending kill on the GameThread and then enqueue a delete command
-	//-TODO: The call to IsPendingKill should no longer be required as we are caching & using within a single render command
-	FFXSystemInterface* FXSystem = nullptr;
-	if (Scene && Scene->FXSystem && !Scene->FXSystem->IsPendingKill())
-	{
-		FXSystem = Scene->FXSystem;
-	}
-
-	for (FSceneRenderer* Renderer : OutRenderers)
-	{
-		Renderer->Link.Head = OutRenderers[0];
-		Renderer->FXSystem  = FXSystem;
-	}
-
-	for (int32 Index = 1; Index < OutRenderers.Num(); ++Index)
-	{
-		OutRenderers[Index - 1]->Link.Next = OutRenderers[Index];
-	}
-
-#if !UE_BUILD_SHIPPING
-	if (bSplitScreenDebug)
-	{
-		DestroySplitScreenDebugViewFamilies(ViewFamilies);
-	}
-#endif
-
-#if DO_CHECK
-	RenderersToAdd.Append(OutRenderers);
-#endif
-
-	return OutRenderers;
-}
-
 void FSceneRenderProcessor::Execute()
 {
 	checkf(!bInsideGroup, TEXT("FSceneRenderBuilder::Execute called within scene render group scope %s. You must end the scope first."), *GroupNodes.GetTail()->Name);
-	checkf(RenderersToAdd.IsEmpty(), TEXT("FSceneRenderBuilder::Execute called but %d renderers were not added."), RenderersToAdd.Num());
 
 #if WITH_GPUDEBUGCRASH
 	if (GRHIGlobals.TriggerGPUCrash != ERequestedGPUCrash::None)
@@ -884,9 +873,15 @@ void FSceneRenderProcessor::Execute()
 				if (Renderer.ViewFamily.EngineShowFlags.Rendering && !RenderState.bSceneUpdateConsumed)
 				{
 					SceneUpdateInputs.Emplace();
+					SceneUpdateInputs->Scene = Scene;
+					SceneUpdateInputs->FXSystem = Scene->FXSystem;
+					SceneUpdateInputs->FeatureLevel = Scene->GetFeatureLevel();
+					SceneUpdateInputs->ShaderPlatform = Scene->GetShaderPlatform();
+					SceneUpdateInputs->GlobalShaderMap = GetGlobalShaderMap(SceneUpdateInputs->ShaderPlatform);
 					SceneUpdateInputs->Renderers = ActiveRenderers;
 					SceneUpdateInputs->ViewFamilies = ActiveViewFamilies;
 					SceneUpdateInputs->Views = ActiveViews;
+					SceneUpdateInputs->CommonShowFlags = CommonShowFlags;
 				}
 
 				const FSceneRenderFunctionInputs FunctionInputs(&Renderer, SceneUpdateInputs ? &SceneUpdateInputs.GetValue() : nullptr, *RenderNode.Name, *RenderState.FullPath);
@@ -960,28 +955,56 @@ void FSceneRenderProcessor::Execute()
 
 //////////////////////////////////////////////////////////////////////////////
 
-FSceneRenderer* FSceneRenderBuilder::CreateSceneRenderer(const FSceneViewFamily* ViewFamily)
+struct FSceneRenderBuilder::FPersistentState
+{
+#if DO_CHECK
+	TArray<FSceneRenderer*, FConcurrentLinearArrayAllocator> RenderersToAdd;
+#endif
+};
+
+TArray<FSceneRenderer*, FConcurrentLinearArrayAllocator> FSceneRenderBuilder::CreateSceneRenderers(
+	TConstArrayView<const FSceneViewFamily*> ViewFamilies,
+	FHitProxyConsumer* HitProxyConsumer,
+	bool bAllowSplitScreenDebug)
 {
 	LazyInit();
+	TArray<FSceneRenderer*, FConcurrentLinearArrayAllocator> Renderers = Processor->CreateSceneRenderers(ViewFamilies, HitProxyConsumer, bAllowSplitScreenDebug);
+
+#if DO_CHECK
+	PersistentState->RenderersToAdd.Append(Renderers);
+#endif
+
+	return Renderers;
+}
+
+FSceneRenderer* FSceneRenderBuilder::CreateSceneRenderer(const FSceneViewFamily* ViewFamily)
+{
 	const bool bAllowSplitScreenDebug = false;
-	return Processor->CreateSceneRenderers(MakeArrayView({ ViewFamily }), nullptr, bAllowSplitScreenDebug)[0];
+	return CreateSceneRenderers(MakeArrayView({ ViewFamily }), nullptr, bAllowSplitScreenDebug)[0];
 }
 
 TArray<FSceneRenderer*, FConcurrentLinearArrayAllocator> FSceneRenderBuilder::CreateLinkedSceneRenderers(
 	TConstArrayView<const FSceneViewFamily*> ViewFamilies,
 	FHitProxyConsumer* HitProxyConsumer)
 {
-	LazyInit();
 	const bool bAllowSplitScreenDebug = true;
-	return Processor->CreateSceneRenderers(ViewFamilies, HitProxyConsumer, bAllowSplitScreenDebug);
+	return CreateSceneRenderers(ViewFamilies, HitProxyConsumer, bAllowSplitScreenDebug);
 }
 
 FSceneRenderBuilder::FSceneRenderBuilder(FSceneInterface* InScene)
 	: Scene(InScene->GetRenderScene())
+	, PersistentState(new FPersistentState)
 {}
 
 FSceneRenderBuilder::~FSceneRenderBuilder()
 {
+	if (PersistentState)
+	{
+		checkf(PersistentState->RenderersToAdd.IsEmpty(), TEXT("FSceneRenderBuilder::Execute called but %d renderers were not added."), PersistentState->RenderersToAdd.Num());
+		delete PersistentState;
+		PersistentState = nullptr;
+	}
+
 	if (Processor)
 	{
 	#if !USE_NULL_RHI
@@ -1013,6 +1036,21 @@ void FSceneRenderBuilder::AddRenderer(FSceneRenderer* Renderer, FString&& Name,
 {
 #if !USE_NULL_RHI
 	LazyInit();
+
+#if DO_CHECK
+	bool bFoundRenderer = false;
+	for (int32 Index = 0; Index < PersistentState->RenderersToAdd.Num(); ++Index)
+	{
+		if (PersistentState->RenderersToAdd[Index] == Renderer)
+		{
+			PersistentState->RenderersToAdd.RemoveAtSwap(Index, EAllowShrinking::No);
+			bFoundRenderer = true;
+			break;
+		}
+	}
+	checkf(bFoundRenderer, TEXT("Renderer being added was not created with this scene render builder or is being added twice."));
+#endif
+
 	Processor->AddRenderer(Renderer, Forward<FString&&>(Name), MoveTemp(Function));
 #endif
 }
@@ -1056,6 +1094,15 @@ FConcurrentLinearBulkObjectAllocator& FSceneRenderBuilder::GetAllocator()
 	return Processor->GetAllocator();
 }
 
+bool FSceneRenderBuilder::IsCompatible(const FEngineShowFlags& EngineShowFlags) const
+{
+	if (Processor)
+	{
+		return Processor->IsCompatible(EngineShowFlags);
+	}
+	return true;
+}
+
 void FSceneRenderBuilder::WaitForAsyncDeleteTask()
 {
 	FSceneRenderProcessor::WaitForAsyncDeleteTask();
diff --git a/Engine/Source/Runtime/Renderer/Private/SceneRenderBuilder.h b/Engine/Source/Runtime/Renderer/Private/SceneRenderBuilder.h
--- a/Engine/Source/Runtime/Renderer/Private/SceneRenderBuilder.h
+++ b/Engine/Source/Runtime/Renderer/Private/SceneRenderBuilder.h
@@ -24,9 +24,11 @@ public:
 	FSceneRenderBuilder& operator=(FSceneRenderBuilder&& RHS)
 	{
 		// Scene stays persistent while the processor instance moves.
-		Scene         = RHS.Scene;
-		Processor     = RHS.Processor;
-		RHS.Processor = nullptr;
+		Scene               = RHS.Scene;
+		Processor           = RHS.Processor;
+		PersistentState     = RHS.PersistentState;
+		RHS.Processor       = nullptr;
+		RHS.PersistentState = nullptr;
 		return *this;
 	}
 
@@ -42,11 +44,20 @@ public:
 	RENDERER_API void AddCommand(TUniqueFunction<void()>&& Function) override;
 	RENDERER_API void AddRenderer(FSceneRenderer* Renderer, FString&& Name, FSceneRenderFunction&& Function) override;
 	using ISceneRenderBuilder::AddRenderer;
+	RENDERER_API bool IsCompatible(const FEngineShowFlags& EngineShowFlags) const override;
 	RENDERER_API void Execute() override;
 	RENDERER_API FConcurrentLinearBulkObjectAllocator& GetAllocator() override;
 
 	//////////////////////////////////////////////////////////////////////////
 
+	void FlushIfIncompatible(const FEngineShowFlags& EngineShowFlags)
+	{
+		if (!IsCompatible(EngineShowFlags))
+		{
+			Execute();
+		}
+	}
+
 	static void WaitForAsyncDeleteTask();
 
 private:
@@ -54,6 +65,14 @@ private:
 	void BeginGroup(FString&& Name, ESceneRenderGroupFlags Flags) override;
 	void EndGroup() override;
 
+	TArray<FSceneRenderer*, FConcurrentLinearArrayAllocator> CreateSceneRenderers(
+		TConstArrayView<const FSceneViewFamily*> ViewFamilies,
+		FHitProxyConsumer* HitProxyConsumer,
+		bool bAllowSplitScreenDebug);
+
+	struct FPersistentState;
+
 	FScene* Scene;
 	FSceneRenderProcessor* Processor = nullptr;
+	FPersistentState* PersistentState = nullptr;
 };
\ No newline at end of file
diff --git a/Engine/Source/Runtime/Renderer/Private/SceneRendering.cpp b/Engine/Source/Runtime/Renderer/Private/SceneRendering.cpp
--- a/Engine/Source/Runtime/Renderer/Private/SceneRendering.cpp
+++ b/Engine/Source/Runtime/Renderer/Private/SceneRendering.cpp
@@ -3866,8 +3866,15 @@ IVisibilityTaskData* FSceneRenderer::OnRenderBegin(FRDGBuilder& GraphBuilder, co
 		PrepareViewStateTask.Wait();
 
 #if RHI_RAYTRACING
-		RayTracing::OnRenderBegin(*Scene, Views, ViewFamily);
-		InitializeRayTracingFlags_RenderThread();
+		if (SceneUpdateInputs)
+		{
+			RayTracing::OnRenderBegin(*SceneUpdateInputs);
+
+			for (FSceneRenderer* Renderer : SceneUpdateInputs->Renderers)
+			{
+				Renderer->InitializeRayTracingFlags_RenderThread();
+			}
+		}
 #endif
 
 		if (!ViewFamily.ViewExtensions.IsEmpty())
@@ -5037,6 +5044,9 @@ void FRendererModule::BeginRenderingViewFamilies(FCanvas* Canvas, TConstArrayVie
 
 		FSceneRenderer::PreallocateCrossGPUFences(SceneRenderers);
 
+		// Flush if the current show flags can't be merged with the current set renderers already added.
+		SceneRenderBuilder.FlushIfIncompatible(ViewFamilies[0]->EngineShowFlags);
+
 		for (FSceneRenderer* SceneRenderer : SceneRenderers)
 		{
 			SceneRenderer->ViewFamily.DisplayInternalsData.Setup(World);
diff --git a/Engine/Source/Runtime/Renderer/Private/SceneRendering.h b/Engine/Source/Runtime/Renderer/Private/SceneRendering.h
--- a/Engine/Source/Runtime/Renderer/Private/SceneRendering.h
+++ b/Engine/Source/Runtime/Renderer/Private/SceneRendering.h
@@ -1987,26 +1987,7 @@ struct FComputeLightGridOutput
 	FRDGPassRef CompactLinksPass = {};
 };
 
-/** An optional set of all renderers, view families, and views for operations that are performed once for a
- *  batch of scene renderers. Only the first scene renderer will have this applied.
- */
-struct FSceneRenderUpdateInputs
-{
-	TConstArrayView<const FSceneView*> GetAsSceneViews() const
-	{
-		return MakeArrayView(reinterpret_cast<const FSceneView* const*>(Views.GetData()), Views.Num());
-	}
-
-	TConstArrayView<const FSceneViewFamily*> GetAsSceneViewFamilies() const
-	{
-		return MakeArrayView(reinterpret_cast<const FSceneViewFamily* const*>(ViewFamilies.GetData()), ViewFamilies.Num());
-	}
-
-	TConstArrayView<FSceneRenderer*> Renderers;
-	TConstArrayView<FViewFamilyInfo*> ViewFamilies;
-	TConstArrayView<FViewInfo*> Views;
-};
-
+struct FSceneRenderUpdateInputs;
 class FDeferredShadingSceneRenderer;
 
 class FSceneRendererBase : public ISceneRenderer
@@ -2642,6 +2623,69 @@ private:
 	friend class FSceneRenderProcessor;
 };
 
+/** A set of show flags thare are guaranteed to be common among all renderers in the set. */
+enum class ESceneRenderCommonShowFlags
+{
+	None             = 0,
+	HitProxies       = 1 << 0,
+	PathTracing      = 1 << 1
+};
+ENUM_CLASS_FLAGS(ESceneRenderCommonShowFlags);
+
+/** An optional set of all renderers, view families, and views for operations that are performed once for a
+ *  batch of scene renderers. Only the first scene renderer will have this applied.
+ */
+struct FSceneRenderUpdateInputs
+{
+	TConstArrayView<const FSceneView*> GetAsSceneViews() const
+	{
+		return MakeArrayView(reinterpret_cast<const FSceneView* const*>(Views.GetData()), Views.Num());
+	}
+
+	TConstArrayView<const FSceneViewFamily*> GetAsSceneViewFamilies() const
+	{
+		return MakeArrayView(reinterpret_cast<const FSceneViewFamily* const*>(ViewFamilies.GetData()), ViewFamilies.Num());
+	}
+
+	template <typename LambdaType>
+	bool HasAnyShowFlags(LambdaType&& Lambda) const
+	{
+		for (FViewFamilyInfo* Family : ViewFamilies)
+		{
+			if (Lambda(Family->EngineShowFlags))
+			{
+				return true;
+			}
+		}
+		return false;
+	}
+
+	template <typename LambdaType>
+	void ForEachView(LambdaType&& Lambda) const
+	{
+		for (FSceneRenderer* Renderer : Renderers)
+		{
+			for (FViewInfo& View : Renderer->Views)
+			{
+				if (!Lambda(Renderer, View))
+				{
+					return;
+				}
+			}
+		}
+	}
+
+	EShaderPlatform ShaderPlatform;
+	ERHIFeatureLevel::Type FeatureLevel;
+	FScene* Scene;
+	FFXSystemInterface* FXSystem;
+	FGlobalShaderMap* GlobalShaderMap;
+	TConstArrayView<FSceneRenderer*> Renderers;
+	TConstArrayView<FViewFamilyInfo*> ViewFamilies;
+	TConstArrayView<FViewInfo*> Views;
+	ESceneRenderCommonShowFlags CommonShowFlags;
+};
+
 template <typename LambdaType, typename PrerequisiteTaskCollectionType>
 UE::Tasks::FTask LaunchSceneRenderTask(const TCHAR* DebugName, LambdaType&& Lambda, PrerequisiteTaskCollectionType&& Prerequisites, bool bExecuteInParallelCondition = true, UE::Tasks::ETaskPriority TaskPriority = UE::Tasks::ETaskPriority::High)
 {
