Fixed post process upscale pass name format.

diff --git a/Engine/Source/Runtime/Renderer/Private/PostProcess/PostProcessUpscale.cpp b/Engine/Source/Runtime/Renderer/Private/PostProcess/PostProcessUpscale.cpp
--- a/Engine/Source/Runtime/Renderer/Private/PostProcess/PostProcessUpscale.cpp
+++ b/Engine/Source/Runtime/Renderer/Private/PostProcess/PostProcessUpscale.cpp
@@ -168,7 +168,7 @@ FScreenPassTexture ISpatialUpscaler::AddDefaultUpscalePass(
 	const TCHAR* StageName = StageNames[static_cast<uint32>(Inputs.Stage)];
 
 	GraphBuilder.AddPass(
-		RDG_EVENT_NAME("Upscale(%s Method=%d%s) %dx%d -> %dx%d",
+		RDG_EVENT_NAME("Upscale(%s Method=%d%s%s) %dx%d -> %dx%d",
 			StageName,
 			int32(Method),
 			PixelPermutationVector.Get<FUpscalePS::FAlphaChannelDim>() ? TEXT(" Alpha") : TEXT(""),
