Merged distance field streaming operations into one invocation for all scene renderers and removed distance field methods from FSceneRenderer.

diff --git a/Engine/Source/Runtime/Renderer/Private/CapsuleShadowRendering.cpp b/Engine/Source/Runtime/Renderer/Private/CapsuleShadowRendering.cpp
--- a/Engine/Source/Runtime/Renderer/Private/CapsuleShadowRendering.cpp
+++ b/Engine/Source/Runtime/Renderer/Private/CapsuleShadowRendering.cpp
@@ -1197,9 +1197,9 @@ void FDeferredShadingSceneRenderer::RenderIndirectCapsuleShadows(FRDGBuilder& Gr
 	}
 }
 
-bool FSceneRenderer::ShouldPrepareForDFInsetIndirectShadow() const
+bool ShouldPrepareForDFInsetIndirectShadow(EShaderPlatform ShaderPlatform, const FEngineShowFlags& EngineShowFlags)
 {
-	return SupportsCapsuleIndirectShadows(ShaderPlatform) && ViewFamily.EngineShowFlags.CapsuleShadows;
+	return SupportsCapsuleIndirectShadows(ShaderPlatform) && EngineShowFlags.CapsuleShadows;
 }
 
 void FDeferredShadingSceneRenderer::RenderCapsuleShadowsForMovableSkylight(
diff --git a/Engine/Source/Runtime/Renderer/Private/DeferredShadingRenderer.cpp b/Engine/Source/Runtime/Renderer/Private/DeferredShadingRenderer.cpp
--- a/Engine/Source/Runtime/Renderer/Private/DeferredShadingRenderer.cpp
+++ b/Engine/Source/Runtime/Renderer/Private/DeferredShadingRenderer.cpp
@@ -1466,11 +1466,17 @@ void FDeferredShadingSceneRenderer::Render(FRDGBuilder& GraphBuilder, const FSce
 		VirtualTextureFeedbackBegin(GraphBuilder, Views, GetActiveSceneTexturesConfig().Extent);
 	}
 
-	// Compute & commit the final state of the entire dependency topology of the renderer.
-	CommitFinalPipelineState();
-
 	if (SceneUpdateInputs)
 	{
+		{
+			TRACE_CPUPROFILER_EVENT_SCOPE(CommitFinalPipelineState);
+			for (FSceneRenderer* Renderer : SceneUpdateInputs->Renderers)
+			{
+				// Compute & commit the final state of the entire dependency topology of the renderer.
+				static_cast<FDeferredShadingSceneRenderer*>(Renderer)->CommitFinalPipelineState();
+			}
+		}
+
 		// Initialize global system textures (pass-through if already initialized).
 		GSystemTextures.InitializeTextures(GraphBuilder.RHICmdList, FeatureLevel);
 	}
@@ -1566,7 +1572,10 @@ void FDeferredShadingSceneRenderer::Render(FRDGBuilder& GraphBuilder, const FSce
 
 	if (RendererOutput == ERendererOutput::FinalSceneColor)
 	{
-		PrepareDistanceFieldScene(GraphBuilder, ExternalAccessQueue);
+		if (SceneUpdateInputs)
+		{
+			PrepareDistanceFieldScene(GraphBuilder, ExternalAccessQueue, *SceneUpdateInputs);
+		}
 
 		for (int32 ViewIndex = 0; ViewIndex < Views.Num(); ViewIndex++)
 		{
@@ -3342,7 +3351,7 @@ void FDeferredShadingSceneRenderer::Render(FRDGBuilder& GraphBuilder, const FSce
 			RenderPhysicsField(GraphBuilder, Views, Scene->PhysicsField, SceneTextures.Color.Target);
 		}
 
-		if (ViewFamily.EngineShowFlags.VisualizeDistanceFieldAO && ShouldRenderDistanceFieldLighting())
+		if (ViewFamily.EngineShowFlags.VisualizeDistanceFieldAO && ShouldRenderDistanceFieldLighting(Scene->DistanceFieldSceneData, Views))
 		{
 			// Use the skylight's max distance if there is one, to be consistent with DFAO shadowing on the skylight
 			const float OcclusionMaxDistance = Scene->SkyLight && !Scene->SkyLight->bWantsStaticShadowing ? Scene->SkyLight->OcclusionMaxDistance : Scene->DefaultMaxDistanceFieldOcclusionDistance;
diff --git a/Engine/Source/Runtime/Renderer/Private/DeferredShadingRenderer.h b/Engine/Source/Runtime/Renderer/Private/DeferredShadingRenderer.h
--- a/Engine/Source/Runtime/Renderer/Private/DeferredShadingRenderer.h
+++ b/Engine/Source/Runtime/Renderer/Private/DeferredShadingRenderer.h
@@ -707,8 +707,6 @@ private:
 		const FSceneTextures& SceneTextures,
 		TArray<FRDGTextureRef>& DynamicBentNormalAOTextures);
 
-	bool ShouldRenderDistanceFieldLighting() const;
-
 	/** Render Ambient Occlusion using mesh distance fields and the surface cache, which supports dynamic rigid meshes. */
 	void RenderDistanceFieldLighting(
 		FRDGBuilder& GraphBuilder,
diff --git a/Engine/Source/Runtime/Renderer/Private/DistanceFieldAmbientOcclusion.cpp b/Engine/Source/Runtime/Renderer/Private/DistanceFieldAmbientOcclusion.cpp
--- a/Engine/Source/Runtime/Renderer/Private/DistanceFieldAmbientOcclusion.cpp
+++ b/Engine/Source/Runtime/Renderer/Private/DistanceFieldAmbientOcclusion.cpp
@@ -707,101 +707,139 @@ bool ShouldDoReflectionEnvironment(const FScene* Scene, const FSceneViewFamily&
 		&& ViewFamily.EngineShowFlags.ReflectionEnvironment;
 }
 
-bool FSceneRenderer::ShouldPrepareForDistanceFieldAO() const
+bool ShouldPrepareForDistanceFieldAO(const FScene* Scene, const FSceneViewFamily& ViewFamily, bool bAnyViewHasGIMethodSupportingDFAO)
 {
-	bool bAnyViewHasGIMethodSupportingDFAO = AnyViewHasGIMethodSupportingDFAO();
+	if (SupportsDistanceFieldAO(Scene->GetFeatureLevel(), Scene->GetShaderPlatform()))
+	{
+		if (bAnyViewHasGIMethodSupportingDFAO && ViewFamily.EngineShowFlags.DistanceFieldAO)
+		{
+			if (GDistanceFieldAOApplyToStaticIndirect)
+			{
+				return true;
+			}
+
+			if (Scene->SkyLight->bCastShadows && ShouldRenderDeferredDynamicSkyLight(Scene, ViewFamily))
+			{
+				return true;
+			}
+		}
+
+		if (ViewFamily.EngineShowFlags.VisualizeDistanceFieldAO)
+		{
+			return true;
+		}
+	}
 
-	return SupportsDistanceFieldAO(Scene->GetFeatureLevel(), Scene->GetShaderPlatform())
-		&& ((ShouldRenderDeferredDynamicSkyLight(Scene, ViewFamily) && bAnyViewHasGIMethodSupportingDFAO && Scene->SkyLight->bCastShadows && ViewFamily.EngineShowFlags.DistanceFieldAO)
-			|| ViewFamily.EngineShowFlags.VisualizeDistanceFieldAO
-			|| (GDistanceFieldAOApplyToStaticIndirect && bAnyViewHasGIMethodSupportingDFAO && ViewFamily.EngineShowFlags.DistanceFieldAO));
+	return false;
 }
 
-bool FSceneRenderer::ShouldPrepareDistanceFieldScene() const
+bool ShouldPrepareForDistanceFieldAO(const FSceneRenderUpdateInputs& SceneUpdateInputs)
 {
-	if (!ensure(Scene != nullptr))
+	for (const FSceneRenderer* Renderer : SceneUpdateInputs.Renderers)
 	{
-		return false;
+		if (ShouldPrepareForDistanceFieldAO(SceneUpdateInputs.Scene, Renderer->ViewFamily, Renderer->AnyViewHasGIMethodSupportingDFAO()))
+		{
+			return true;
+		}
 	}
 
-	if (!DoesProjectSupportDistanceFields())
+	return false;
+}
+
+bool ShouldPrepareGlobalDistanceField(const FSceneRenderer* Renderer)
+{
+	if (!DoesProjectSupportDistanceFields() || !UseGlobalDistanceField())
 	{
 		return false;
 	}
 
-	if (GRHIDeviceIsIntegrated && GMaxRHIFeatureLevel < ERHIFeatureLevel::SM6)
+	const FScene* Scene = Renderer->Scene;
+	const FFXSystemInterface* FXSystem = Scene->FXSystem;
+	const FViewInfo& View = Renderer->Views[0];
+	const FSceneViewFamily& ViewFamily = Renderer->ViewFamily;
+
+	const bool bShouldPrepareForMaterialsOrNiagara = !GAOGlobalDistanceFieldDetailedNecessityCheck
+		|| View.PrevViewInfo.bUsesGlobalDistanceField // use previous frame relevance because this frame relevance flags haven't been calculate yet
+		|| ((FXSystem != nullptr) && FXSystem->UsesGlobalDistanceField());
+
+	if (ViewFamily.EngineShowFlags.PathTracing)
 	{
-		// In case of iGPU we use the maximum feature level to differentiate between older and newer more capable iGPU
-		return false;
+		return bShouldPrepareForMaterialsOrNiagara;
+	}
+
+	if (SupportsDistanceFieldAO(Renderer->FeatureLevel, Renderer->ShaderPlatform) && (bShouldPrepareForMaterialsOrNiagara || ShouldPrepareForDistanceFieldAO(Scene, ViewFamily, Renderer->AnyViewHasGIMethodSupportingDFAO())))
+	{
+		return true;
 	}
 
-	const bool bShouldPrepareForAO = ShouldPrepareForDistanceFieldAO();
-	const bool bShouldPrepareGlobalDistanceField = ShouldPrepareGlobalDistanceField();
-	const bool bShouldPrepareForDFInsetIndirectShadow = ShouldPrepareForDFInsetIndirectShadow();
+	if (Renderer->IsLumenEnabled(View) && Lumen::UseGlobalSDFObjectGrid(ViewFamily))
+	{
+		return true;
+	}
 
-	if (ViewFamily.EngineShowFlags.PathTracing)
+	if (MegaLights::IsUsingGlobalSDF(ViewFamily))
+	{
+		return true;
+	}
+
+	if (ViewFamily.EngineShowFlags.VisualizeGlobalDistanceField)
 	{
-		return bShouldPrepareGlobalDistanceField;
+		return true;
 	}
 
-	// Prepare the distance field scene (object buffers and distance field atlas) if any feature needs it
-	return bShouldPrepareGlobalDistanceField 
-		|| bShouldPrepareForAO 
-		|| ShouldPrepareForDistanceFieldShadows() 
-		|| bShouldPrepareForDFInsetIndirectShadow 
-		|| ViewFamily.EngineShowFlags.VisualizeGlobalDistanceField
-		|| ViewFamily.EngineShowFlags.VisualizeMeshDistanceFields;
+	return false;
 }
 
-bool FSceneRenderer::ShouldPrepareGlobalDistanceField() const
+bool ShouldPrepareDistanceFieldScene(const FSceneRenderUpdateInputs& SceneUpdateInputs)
 {
-	if (!ensure(Scene != nullptr))
+	extern bool ShouldPrepareForDFInsetIndirectShadow(EShaderPlatform, const FEngineShowFlags&);
+	extern bool ShouldPrepareForDistanceFieldShadows(const FSceneRenderUpdateInputs& SceneUpdateInputs);
+
+	if (!DoesProjectSupportDistanceFields())
 	{
 		return false;
 	}
 
-	if (!DoesProjectSupportDistanceFields() || !UseGlobalDistanceField())
+	if (GRHIDeviceIsIntegrated && GMaxRHIFeatureLevel < ERHIFeatureLevel::SM6)
 	{
+		// In case of iGPU we use the maximum feature level to differentiate between older and newer more capable iGPU
 		return false;
 	}
 
-	const bool bShouldPrepareForMaterialsOrNiagara = 
-		!GAOGlobalDistanceFieldDetailedNecessityCheck
-		|| Views[0].PrevViewInfo.bUsesGlobalDistanceField // use previous frame relevance because this frame relevance flags haven't been calculate yet
-		|| ((FXSystem != nullptr) && FXSystem->UsesGlobalDistanceField());
-
-	if (ViewFamily.EngineShowFlags.PathTracing)
+	if (ShouldPrepareForDistanceFieldAO(SceneUpdateInputs))
 	{
-		return bShouldPrepareForMaterialsOrNiagara;
+		return true;
 	}
 
-	const bool bShouldPrepareForAO = SupportsDistanceFieldAO(Scene->GetFeatureLevel(), Scene->GetShaderPlatform())
-		&& (ShouldPrepareForDistanceFieldAO() || bShouldPrepareForMaterialsOrNiagara);
-
-	const bool bShouldPrepareForLumen = IsLumenEnabled(Views[0]) && Lumen::UseGlobalSDFObjectGrid(ViewFamily);
+	if (ShouldPrepareForDistanceFieldShadows(SceneUpdateInputs))
+	{
+		return true;
+	}
 
-	// TODO: Should check any light actually uses Mega Lights
-	const bool bShouldPrepareForMegaLights = MegaLights::IsUsingGlobalSDF(ViewFamily);
+	for (const FSceneRenderer* Renderer : SceneUpdateInputs.Renderers)
+	{
+		const FViewFamilyInfo& ViewFamily = Renderer->ViewFamily;
 
-	const bool bShouldPrepareForVisualization = ViewFamily.EngineShowFlags.VisualizeGlobalDistanceField;
+		if (ShouldPrepareGlobalDistanceField(Renderer))
+		{
+			return true;
+		}
 
-	return (bShouldPrepareForAO || bShouldPrepareForLumen || bShouldPrepareForMegaLights || bShouldPrepareForVisualization);
-}
+		if (ShouldPrepareForDFInsetIndirectShadow(SceneUpdateInputs.ShaderPlatform, ViewFamily.EngineShowFlags))
+		{
+			return true;
+		}
 
-void FDeferredShadingSceneRenderer::RenderDFAOAsIndirectShadowing(
-	FRDGBuilder& GraphBuilder,
-	const FSceneTextures& SceneTextures,
-	TArray<FRDGTextureRef>& DynamicBentNormalAOTextures)
-{
-	if (GDistanceFieldAOApplyToStaticIndirect && ShouldRenderDistanceFieldAO() && ShouldRenderDistanceFieldLighting())
-	{
-		// Use the skylight's max distance if there is one, to be consistent with DFAO shadowing on the skylight
-		const float OcclusionMaxDistance = Scene->SkyLight && !Scene->SkyLight->bWantsStaticShadowing ? Scene->SkyLight->OcclusionMaxDistance : Scene->DefaultMaxDistanceFieldOcclusionDistance;
-		RenderDistanceFieldLighting(GraphBuilder, SceneTextures, FDistanceFieldAOParameters(OcclusionMaxDistance), DynamicBentNormalAOTextures, true, false);
+		if (ViewFamily.EngineShowFlags.VisualizeGlobalDistanceField || ViewFamily.EngineShowFlags.VisualizeMeshDistanceFields)
+		{
+			return true;
+		}
 	}
+
+	return false;
 }
 
-bool FDeferredShadingSceneRenderer::ShouldRenderDistanceFieldLighting() const
+bool ShouldRenderDistanceFieldLighting(const FDistanceFieldSceneData& SceneData, TConstArrayView<FViewInfo> Views)
 {
 	if (!GDistanceFieldAOMultiView && Views.Num() > 1)
 	{
@@ -821,7 +859,20 @@ bool FDeferredShadingSceneRenderer::ShouldRenderDistanceFieldLighting() const
 		}
 	}
 
-	return bSupportsDistanceFieldAO && Scene->DistanceFieldSceneData.NumObjectsInBuffer;
+	return bSupportsDistanceFieldAO && SceneData.NumObjectsInBuffer;
+}
+
+void FDeferredShadingSceneRenderer::RenderDFAOAsIndirectShadowing(
+	FRDGBuilder& GraphBuilder,
+	const FSceneTextures& SceneTextures,
+	TArray<FRDGTextureRef>& DynamicBentNormalAOTextures)
+{
+	if (GDistanceFieldAOApplyToStaticIndirect && ShouldRenderDistanceFieldAO(Views, ViewFamily.EngineShowFlags) && ShouldRenderDistanceFieldLighting(Scene->DistanceFieldSceneData, Views))
+	{
+		// Use the skylight's max distance if there is one, to be consistent with DFAO shadowing on the skylight
+		const float OcclusionMaxDistance = Scene->SkyLight && !Scene->SkyLight->bWantsStaticShadowing ? Scene->SkyLight->OcclusionMaxDistance : Scene->DefaultMaxDistanceFieldOcclusionDistance;
+		RenderDistanceFieldLighting(GraphBuilder, SceneTextures, FDistanceFieldAOParameters(OcclusionMaxDistance), DynamicBentNormalAOTextures, true, false);
+	}
 }
 
 void FDeferredShadingSceneRenderer::RenderDistanceFieldLighting(
@@ -832,7 +883,7 @@ void FDeferredShadingSceneRenderer::RenderDistanceFieldLighting(
 	bool bModulateToSceneColor,
 	bool bVisualizeAmbientOcclusion)
 {
-	check(ShouldRenderDistanceFieldLighting());
+	check(ShouldRenderDistanceFieldLighting(Scene->DistanceFieldSceneData, Views));
 	check(!Scene->DistanceFieldSceneData.HasPendingOperations());
 
 	RDG_EVENT_SCOPE(GraphBuilder, "DistanceFieldLighting");
@@ -923,7 +974,7 @@ void FDeferredShadingSceneRenderer::RenderDistanceFieldLighting(
 	}
 }
 
-bool FSceneRenderer::ShouldRenderDistanceFieldAO() const
+bool ShouldRenderDistanceFieldAO(TConstArrayView<FViewInfo> Views, const FEngineShowFlags& EngineShowFlags)
 {
 	bool bShouldRenderRTAO = false;
 	for (int ViewIndex = 0; ViewIndex < Views.Num(); ++ViewIndex)
@@ -931,10 +982,10 @@ bool FSceneRenderer::ShouldRenderDistanceFieldAO() const
 		bShouldRenderRTAO = bShouldRenderRTAO || ShouldRenderRayTracingAmbientOcclusion(Views[ViewIndex]);
 	}
 
-	return ViewFamily.EngineShowFlags.DistanceFieldAO
+	return EngineShowFlags.DistanceFieldAO
 		&& !bShouldRenderRTAO
-		&& !ViewFamily.EngineShowFlags.VisualizeDistanceFieldAO
-		&& !ViewFamily.EngineShowFlags.VisualizeMeshDistanceFields
-		&& !ViewFamily.EngineShowFlags.VisualizeGlobalDistanceField
-		&& !ViewFamily.EngineShowFlags.PathTracing;
+		&& !EngineShowFlags.VisualizeDistanceFieldAO
+		&& !EngineShowFlags.VisualizeMeshDistanceFields
+		&& !EngineShowFlags.VisualizeGlobalDistanceField
+		&& !EngineShowFlags.PathTracing;
 }
diff --git a/Engine/Source/Runtime/Renderer/Private/DistanceFieldLightingShared.h b/Engine/Source/Runtime/Renderer/Private/DistanceFieldLightingShared.h
--- a/Engine/Source/Runtime/Renderer/Private/DistanceFieldLightingShared.h
+++ b/Engine/Source/Runtime/Renderer/Private/DistanceFieldLightingShared.h
@@ -16,9 +16,14 @@ class FPrimitiveSceneInfo;
 class FRDGBuffer;
 class FRDGBuilder;
 class FRDGPooledBuffer;
+class FScene;
 class FSceneRenderer;
+class FSceneViewFamily;
 class FShaderParameterMap;
 class FViewInfo;
+
+struct FEngineShowFlags;
+
 using FRDGBufferRef = FRDGBuffer*;
 
 DECLARE_LOG_CATEGORY_EXTERN(LogDistanceField, Log, All);
@@ -150,4 +155,9 @@ extern void CullDistanceFieldObjectsForLight(
 	FDistanceFieldCulledObjectBufferParameters& CulledObjectBufferParameters,
 	FLightTileIntersectionParameters& LightTileIntersectionParameters);
 
+extern bool ShouldPrepareForDistanceFieldAO(const FScene* Scene, const FSceneViewFamily& ViewFamily, bool bAnyViewHasGIMethodSupportingDFAO);
+extern bool ShouldRenderDistanceFieldAO(TConstArrayView<FViewInfo> Views, const FEngineShowFlags& EngineShowFlags);
+
 extern bool SupportsDistanceFieldAO(ERHIFeatureLevel::Type FeatureLevel, EShaderPlatform ShaderPlatform);
+
+extern bool ShouldRenderDistanceFieldLighting(const FDistanceFieldSceneData& SceneData, TConstArrayView<FViewInfo> Views);
\ No newline at end of file
diff --git a/Engine/Source/Runtime/Renderer/Private/DistanceFieldObjectManagement.cpp b/Engine/Source/Runtime/Renderer/Private/DistanceFieldObjectManagement.cpp
--- a/Engine/Source/Runtime/Renderer/Private/DistanceFieldObjectManagement.cpp
+++ b/Engine/Source/Runtime/Renderer/Private/DistanceFieldObjectManagement.cpp
@@ -747,13 +747,11 @@
 	}
 }
 
-void FSceneRenderer::UpdateGlobalHeightFieldObjectBuffers(FRDGBuilder& GraphBuilder, const TArray<uint32>& IndicesToUpdateInHeightFieldObjectBuffers)
+void UpdateGlobalHeightFieldObjectBuffers(FRDGBuilder& GraphBuilder, const TArray<uint32>& IndicesToUpdateInHeightFieldObjectBuffers, FDistanceFieldSceneData& DistanceFieldSceneData)
 {
 	// Mask should be set in FSceneRenderer::PrepareDistanceFieldScene before calling this
 	check(GraphBuilder.RHICmdList.GetGPUMask() == FRHIGPUMask::All());
 
-	FDistanceFieldSceneData& DistanceFieldSceneData = Scene->DistanceFieldSceneData;
-
 	check(DistanceFieldSceneData.PendingHeightFieldAddOps.IsEmpty());
 	check(DistanceFieldSceneData.PendingHeightFieldRemoveOps.IsEmpty());
 
@@ -893,82 +891,7 @@
 	}
 }
 
-void FSceneRenderer::PrepareDistanceFieldScene(FRDGBuilder& GraphBuilder, FRDGExternalAccessQueue& ExternalAccessQueue)
+void ProcessPendingHeightFieldPrimitiveAddAndRemoveOps(FScene* Scene, TArray<uint32>& IndicesToUpdateInHeightFieldObjectBuffers)
 {
-	RDG_CSV_STAT_EXCLUSIVE_SCOPE(GraphBuilder, PrepareDistanceFieldScene);
-	TRACE_CPUPROFILER_EVENT_SCOPE(FSceneRenderer::PrepareDistanceFieldScene);
-	QUICK_SCOPE_CYCLE_COUNTER(STAT_PrepareDistanceFieldScene);
-	LLM_SCOPE_BYTAG(DistanceFields);
-	RDG_GPU_MASK_SCOPE(GraphBuilder, FRHIGPUMask::All());
-	RDG_RHI_GPU_STAT_SCOPE(GraphBuilder, DistanceFields);
-
-	const bool bShouldPrepareHeightFieldScene = ShouldPrepareHeightFieldScene();
-	const bool bShouldPrepareDistanceFieldScene = ShouldPrepareDistanceFieldScene();
-
-	if (!bShouldPrepareDistanceFieldScene && !bShouldPrepareHeightFieldScene)
-	{
-		return;
-	}
-
-	FDistanceFieldSceneData& DistanceFieldSceneData = Scene->DistanceFieldSceneData;
-
-	TArray<uint32> IndicesToUpdateInHeightFieldObjectBuffers;
-
-	ProcessPendingHeightFieldPrimitiveAddAndRemoveOps(IndicesToUpdateInHeightFieldObjectBuffers);
-
-	if (bShouldPrepareHeightFieldScene)
-	{
-		extern int32 GHFShadowQuality;
-		if (GHFShadowQuality > 2)
-		{
-			GHFVisibilityTextureAtlas.UpdateAllocations(GraphBuilder, FeatureLevel);
-		}
-		GHeightFieldTextureAtlas.UpdateAllocations(GraphBuilder, FeatureLevel);
-
-		UpdateGlobalHeightFieldObjectBuffers(GraphBuilder, IndicesToUpdateInHeightFieldObjectBuffers);
-	}
-	else if (DistanceFieldSceneData.HeightFieldObjectBuffers)
-	{
-		// if we don't need HeightFieldScene release the buffers
-		delete DistanceFieldSceneData.HeightFieldObjectBuffers;
-		DistanceFieldSceneData.HeightFieldObjectBuffers = nullptr;
-		DistanceFieldSceneData.HeightFieldAtlasGeneration = 0;
-		DistanceFieldSceneData.HFVisibilityAtlasGenerattion = 0;
-	}
-
-	if (bShouldPrepareDistanceFieldScene)
-	{
-		TArray<FDistanceFieldAssetMipId> DistanceFieldAssetAdds;
-		TArray<FSetElementId> DistanceFieldAssetRemoves;
-		DistanceFieldSceneData.UpdateDistanceFieldObjectBuffers(GraphBuilder, ExternalAccessQueue, Scene, DistanceFieldAssetAdds, DistanceFieldAssetRemoves);
-
-		DistanceFieldSceneData.UpdateDistanceFieldAtlas(GraphBuilder, ExternalAccessQueue, Views[0], Scene, IsLumenEnabled(Views[0]), Views[0].ShaderMap, DistanceFieldAssetAdds, DistanceFieldAssetRemoves);
-
-		if (ShouldPrepareGlobalDistanceField())
-		{
-			for (int32 ViewIndex = 0; ViewIndex < Views.Num(); ViewIndex++)
-			{
-				FViewInfo& View = Views[ViewIndex];
-
-				RDG_GPU_MASK_SCOPE(GraphBuilder, View.GPUMask);
-
-				float OcclusionMaxDistance = Scene->DefaultMaxDistanceFieldOcclusionDistance;
-
-				// Use the skylight's max distance if there is one
-				if (Scene->SkyLight && Scene->SkyLight->bCastShadows && !Scene->SkyLight->bWantsStaticShadowing)
-				{
-					OcclusionMaxDistance = Scene->SkyLight->OcclusionMaxDistance;
-				}
-
-				const bool bUseLumenGlobalDistanceFieldSettings = IsLumenEnabled(View) && Lumen::IsUsingGlobalSDF(*View.Family);
-
-				UpdateGlobalDistanceFieldVolume(GraphBuilder, ExternalAccessQueue, View, Scene, OcclusionMaxDistance, bUseLumenGlobalDistanceFieldSettings, View.GlobalDistanceFieldInfo);
-			}
-		}
-	}
-}
-
-void FSceneRenderer::ProcessPendingHeightFieldPrimitiveAddAndRemoveOps(TArray<uint32>& IndicesToUpdateInHeightFieldObjectBuffers)
-{
 	FDistanceFieldSceneData& SceneData = Scene->DistanceFieldSceneData;
 
@@ -1048,3 +971,82 @@
 
 	SceneData.PendingHeightFieldAddOps.Reset();
 }
+
+void PrepareDistanceFieldScene(FRDGBuilder& GraphBuilder, FRDGExternalAccessQueue& ExternalAccessQueue, const FSceneRenderUpdateInputs& SceneUpdateInputs)
+{
+	RDG_CSV_STAT_EXCLUSIVE_SCOPE(GraphBuilder, PrepareDistanceFieldScene);
+	TRACE_CPUPROFILER_EVENT_SCOPE(FSceneRenderer::PrepareDistanceFieldScene);
+	QUICK_SCOPE_CYCLE_COUNTER(STAT_PrepareDistanceFieldScene);
+	LLM_SCOPE_BYTAG(DistanceFields);
+	RDG_GPU_MASK_SCOPE(GraphBuilder, FRHIGPUMask::All());
+	RDG_RHI_GPU_STAT_SCOPE(GraphBuilder, DistanceFields);
+
+	extern bool ShouldPrepareDistanceFieldScene(const FSceneRenderUpdateInputs & SceneUpdateInputs);
+	extern bool ShouldPrepareGlobalDistanceField(const FSceneRenderer*);
+	extern bool ShouldPrepareHeightFieldScene(const FSceneRenderUpdateInputs& SceneUpdateInputs);
+
+	const bool bShouldPrepareHeightFieldScene = ShouldPrepareHeightFieldScene(SceneUpdateInputs);
+	const bool bShouldPrepareDistanceFieldScene = ShouldPrepareDistanceFieldScene(SceneUpdateInputs);
+
+	if (!bShouldPrepareDistanceFieldScene && !bShouldPrepareHeightFieldScene)
+	{
+		return;
+	}
+
+	FScene* Scene = SceneUpdateInputs.Scene;
+	FDistanceFieldSceneData& DistanceFieldSceneData = Scene->DistanceFieldSceneData;
+
+	TArray<uint32> IndicesToUpdateInHeightFieldObjectBuffers;
+
+	ProcessPendingHeightFieldPrimitiveAddAndRemoveOps(Scene, IndicesToUpdateInHeightFieldObjectBuffers);
+
+	if (bShouldPrepareHeightFieldScene)
+	{
+		extern int32 GHFShadowQuality;
+		if (GHFShadowQuality > 2)
+		{
+			GHFVisibilityTextureAtlas.UpdateAllocations(GraphBuilder, SceneUpdateInputs.FeatureLevel);
+		}
+		GHeightFieldTextureAtlas.UpdateAllocations(GraphBuilder, SceneUpdateInputs.FeatureLevel);
+
+		UpdateGlobalHeightFieldObjectBuffers(GraphBuilder, IndicesToUpdateInHeightFieldObjectBuffers, DistanceFieldSceneData);
+	}
+	else if (DistanceFieldSceneData.HeightFieldObjectBuffers)
+	{
+		// if we don't need HeightFieldScene release the buffers
+		delete DistanceFieldSceneData.HeightFieldObjectBuffers;
+		DistanceFieldSceneData.HeightFieldObjectBuffers = nullptr;
+		DistanceFieldSceneData.HeightFieldAtlasGeneration = 0;
+		DistanceFieldSceneData.HFVisibilityAtlasGenerattion = 0;
+	}
+
+	if (bShouldPrepareDistanceFieldScene)
+	{
+		TArray<FDistanceFieldAssetMipId> DistanceFieldAssetAdds;
+		TArray<FSetElementId> DistanceFieldAssetRemoves;
+		DistanceFieldSceneData.UpdateDistanceFieldObjectBuffers(GraphBuilder, ExternalAccessQueue, Scene, DistanceFieldAssetAdds, DistanceFieldAssetRemoves);
+
+		DistanceFieldSceneData.UpdateDistanceFieldAtlas(GraphBuilder, ExternalAccessQueue, SceneUpdateInputs, DistanceFieldAssetAdds, DistanceFieldAssetRemoves);
+
+		SceneUpdateInputs.ForEachView([&] (const FSceneRenderer* Renderer, FViewInfo& View)
+		{
+			if (ShouldPrepareGlobalDistanceField(Renderer))
+			{
+				RDG_GPU_MASK_SCOPE(GraphBuilder, View.GPUMask);
+
+				float OcclusionMaxDistance = Scene->DefaultMaxDistanceFieldOcclusionDistance;
+
+				// Use the skylight's max distance if there is one
+				if (Scene->SkyLight && Scene->SkyLight->bCastShadows && !Scene->SkyLight->bWantsStaticShadowing)
+				{
+					OcclusionMaxDistance = Scene->SkyLight->OcclusionMaxDistance;
+				}
+
+				const bool bUseLumenGlobalDistanceFieldSettings = Renderer->IsLumenEnabled(View) && Lumen::IsUsingGlobalSDF(*View.Family);
+
+				UpdateGlobalDistanceFieldVolume(GraphBuilder, ExternalAccessQueue, View, Scene, OcclusionMaxDistance, bUseLumenGlobalDistanceFieldSettings, View.GlobalDistanceFieldInfo);
+			}
+			return true;
+		});
+	}
+}
diff --git a/Engine/Source/Runtime/Renderer/Private/DistanceFieldShadowing.cpp b/Engine/Source/Runtime/Renderer/Private/DistanceFieldShadowing.cpp
--- a/Engine/Source/Runtime/Renderer/Private/DistanceFieldShadowing.cpp
+++ b/Engine/Source/Runtime/Renderer/Private/DistanceFieldShadowing.cpp
@@ -758,9 +758,10 @@ bool SupportsHeightFieldShadows(ERHIFeatureLevel::Type FeatureLevel, EShaderPlat
 		&& DoesPlatformSupportDistanceFieldShadowing(ShaderPlatform);
 }
 
-bool FSceneRenderer::ShouldPrepareForDistanceFieldShadows() const
+bool ShouldPrepareForDistanceFieldShadows(const FSceneRenderUpdateInputs& SceneUpdateInputs)
 {
-	if (!ViewFamily.EngineShowFlags.DynamicShadows || !SupportsDistanceFieldShadows(Scene->GetFeatureLevel(), Scene->GetShaderPlatform()))
+	if (!SceneUpdateInputs.HasAnyShowFlags([] (const FEngineShowFlags& EngineShowFlags) { return EngineShowFlags.DynamicShadows; }) ||
+		!SupportsDistanceFieldShadows(SceneUpdateInputs.FeatureLevel, SceneUpdateInputs.ShaderPlatform))
 	{
 		return false;
 	}
@@ -768,12 +769,10 @@ bool FSceneRenderer::ShouldPrepareForDistanceFieldShadows() const
 	return true;
 }
 
-bool FSceneRenderer::ShouldPrepareHeightFieldScene() const
+bool ShouldPrepareHeightFieldScene(const FSceneRenderUpdateInputs& SceneUpdateInputs)
 {
-	return Scene
-		&& ViewFamily.EngineShowFlags.DynamicShadows
-		&& !ViewFamily.EngineShowFlags.PathTracing
-		&& SupportsHeightFieldShadows(Scene->GetFeatureLevel(), Scene->GetShaderPlatform());
+	return SceneUpdateInputs.HasAnyShowFlags([](const FEngineShowFlags& ShowFlags) { return ShowFlags.DynamicShadows && !ShowFlags.PathTracing; })
+		&& SupportsHeightFieldShadows(SceneUpdateInputs.FeatureLevel, SceneUpdateInputs.ShaderPlatform);
 }
 
 void GetDistanceFieldShadowRange(const FProjectedShadowInfo* ProjectedShadowInfo, EDistanceFieldPrimitiveType PrimitiveType, float& OutMinDepth, float& OutMaxDepth)
diff --git a/Engine/Source/Runtime/Renderer/Private/DistanceFieldStreaming.cpp b/Engine/Source/Runtime/Renderer/Private/DistanceFieldStreaming.cpp
--- a/Engine/Source/Runtime/Renderer/Private/DistanceFieldStreaming.cpp
+++ b/Engine/Source/Runtime/Renderer/Private/DistanceFieldStreaming.cpp
@@ -244,7 +244,6 @@ class FComputeDistanceFieldAssetWantedMipsCS : public FGlobalShader
 	SHADER_USE_PARAMETER_STRUCT(FComputeDistanceFieldAssetWantedMipsCS, FGlobalShader)
 
 	BEGIN_SHADER_PARAMETER_STRUCT(FParameters, )
-		SHADER_PARAMETER_STRUCT_REF(FViewUniformShaderParameters, View)
 		SHADER_PARAMETER_RDG_BUFFER_UAV(RWStructuredBuffer<uint>, RWDistanceFieldAssetWantedNumMips)
 		SHADER_PARAMETER_RDG_BUFFER_UAV(RWStructuredBuffer<uint>, RWDistanceFieldAssetStreamingRequests)
 		SHADER_PARAMETER_STRUCT_INCLUDE(FDistanceFieldObjectBufferParameters, DistanceFieldObjectBuffers)
@@ -1203,17 +1202,12 @@ void FDistanceFieldSceneData::DefragmentIndirectionAtlas(FIntVector MinSize, TAr
 	}
 }
 
-void FDistanceFieldSceneData::GenerateStreamingRequests(
-	FRDGBuilder& GraphBuilder, 
-	const FViewInfo& View,
-	FScene* Scene,
-	bool bLumenEnabled,
-	FGlobalShaderMap* GlobalShaderMap)
+void FDistanceFieldSceneData::GenerateStreamingRequests(FRDGBuilder& GraphBuilder, const FSceneRenderUpdateInputs& SceneUpdateInputs)
 {
 	// It is not safe to EnqueueCopy on a buffer that already has a pending copy.
 	if (ReadbackBuffersNumPending < MaxStreamingReadbackBuffers && NumObjectsInBuffer > 0)
 	{
-		RDG_GPU_MASK_SCOPE(GraphBuilder, View.GPUMask);
+		RDG_GPU_MASK_SCOPE(GraphBuilder, FRHIGPUMask::All());
 
 		if (!StreamingRequestReadbackBuffers[ReadbackBuffersWriteIndex])
 		{
@@ -1232,11 +1226,13 @@ void FDistanceFieldSceneData::GenerateStreamingRequests(
 		StreamingRequestsDesc.Usage = EBufferUsageFlags(StreamingRequestsDesc.Usage | BUF_SourceCopy);
 		FRDGBufferRef StreamingRequestsBuffer = GraphBuilder.CreateBuffer(StreamingRequestsDesc, TEXT("DistanceFields.DistanceFieldStreamingRequests"));
 
+		SceneUpdateInputs.ForEachView([&] (const FSceneRenderer* Renderer, const FViewInfo& View)
 		{
+			const bool bLumenEnabled = Renderer->IsLumenEnabled(View);
+
 			FDFVector3 PreViewTranslation(View.ViewMatrices.GetPreViewTranslation());
 
 			FComputeDistanceFieldAssetWantedMipsCS::FParameters* PassParameters = GraphBuilder.AllocParameters<FComputeDistanceFieldAssetWantedMipsCS::FParameters>();
-			PassParameters->View = View.ViewUniformBuffer;
 			checkf(DistanceField::NumMips == 3, TEXT("Shader needs to be updated"));
 			PassParameters->RWDistanceFieldAssetWantedNumMips = GraphBuilder.CreateUAV(WantedNumMips);
 			PassParameters->RWDistanceFieldAssetStreamingRequests = GraphBuilder.CreateUAV(StreamingRequestsBuffer);
@@ -1245,13 +1241,13 @@ void FDistanceFieldSceneData::GenerateStreamingRequests(
 			extern int32 GAOGlobalDistanceFieldNumClipmaps;
 			// Request Mesh SDF mips based off of the Global SDF clipmaps
 			PassParameters->Mip1WorldTranslatedCenter = FVector3f(View.ViewMatrices.GetViewOrigin() + View.ViewMatrices.GetPreViewTranslation());
-			PassParameters->Mip1WorldExtent = FVector3f(GlobalDistanceField::GetClipmapExtent(GAOGlobalDistanceFieldNumClipmaps - 1, Scene, bLumenEnabled));
+			PassParameters->Mip1WorldExtent = FVector3f(GlobalDistanceField::GetClipmapExtent(GAOGlobalDistanceFieldNumClipmaps - 1, SceneUpdateInputs.Scene, bLumenEnabled));
 			PassParameters->Mip2WorldTranslatedCenter = FVector3f(View.ViewMatrices.GetViewOrigin() + View.ViewMatrices.GetPreViewTranslation());
-			PassParameters->Mip2WorldExtent = FVector3f(GlobalDistanceField::GetClipmapExtent(FMath::Max<int32>(GAOGlobalDistanceFieldNumClipmaps / 2 - 1, 0), Scene, bLumenEnabled));
+			PassParameters->Mip2WorldExtent = FVector3f(GlobalDistanceField::GetClipmapExtent(FMath::Max<int32>(GAOGlobalDistanceFieldNumClipmaps / 2 - 1, 0), SceneUpdateInputs.Scene, bLumenEnabled));
 			PassParameters->PreViewTranslationHigh = PreViewTranslation.High;
 			PassParameters->PreViewTranslationLow = PreViewTranslation.Low;
 
-			auto ComputeShader = GlobalShaderMap->GetShader<FComputeDistanceFieldAssetWantedMipsCS>();
+			auto ComputeShader = SceneUpdateInputs.GlobalShaderMap->GetShader<FComputeDistanceFieldAssetWantedMipsCS>();
 
 			FComputeShaderUtils::AddPass(
 				GraphBuilder,
@@ -1259,7 +1255,9 @@ void FDistanceFieldSceneData::GenerateStreamingRequests(
 				ComputeShader,
 				PassParameters,
 				FComputeShaderUtils::GetGroupCountWrapped(NumObjectsInBuffer, FComputeDistanceFieldAssetWantedMipsCS::GetGroupSize()));
-		}
+
+			return true;
+		});
 
 		{
 			FGenerateDistanceFieldAssetStreamingRequestsCS::FParameters* PassParameters = GraphBuilder.AllocParameters<FGenerateDistanceFieldAssetStreamingRequestsCS::FParameters>();
@@ -1270,7 +1268,7 @@ void FDistanceFieldSceneData::GenerateStreamingRequests(
 			PassParameters->NumDistanceFieldAssets = NumAssets;
 			PassParameters->MaxNumStreamingRequests = MaxStreamingRequests;
 
-			auto ComputeShader = GlobalShaderMap->GetShader<FGenerateDistanceFieldAssetStreamingRequestsCS>();
+			auto ComputeShader = SceneUpdateInputs.GlobalShaderMap->GetShader<FGenerateDistanceFieldAssetStreamingRequestsCS>();
 
 			FComputeShaderUtils::AddPass(
 				GraphBuilder,
@@ -1421,10 +1419,7 @@ void FDistanceFieldSceneData::UploadAllAssetData(FRDGBuilder& GraphBuilder, FRDG
 void FDistanceFieldSceneData::UpdateDistanceFieldAtlas(
 	FRDGBuilder& GraphBuilder,
 	FRDGExternalAccessQueue& ExternalAccessQueue,
-	const FViewInfo& View,
-	FScene* Scene,
-	bool bLumenEnabled,
-	FGlobalShaderMap* GlobalShaderMap,
+	const FSceneRenderUpdateInputs& SceneUpdateInputs,
 	TArray<FDistanceFieldAssetMipId>& DistanceFieldAssetMipAdds,
 	TArray<FSetElementId>& DistanceFieldAssetRemoves)
 {
@@ -1516,7 +1511,7 @@ void FDistanceFieldSceneData::UpdateDistanceFieldAtlas(
 	}
 
 	// Now that DistanceFieldAtlasBlockAllocator has been modified, potentially resize the atlas
-	FRDGTextureRef DistanceFieldBrickVolumeTextureRDG = ResizeBrickAtlasIfNeeded(GraphBuilder, GlobalShaderMap);
+	FRDGTextureRef DistanceFieldBrickVolumeTextureRDG = ResizeBrickAtlasIfNeeded(GraphBuilder, SceneUpdateInputs.GlobalShaderMap);
 
 	const uint32 NumAssets = AssetStateArray.GetMaxIndex();
 	const int32 AssetDataStrideFloat4s = DistanceField::NumMips * AssetDataMipStrideFloat4s;
@@ -1541,7 +1536,7 @@ void FDistanceFieldSceneData::UpdateDistanceFieldAtlas(
 	}
 	else
 	{
-		bIndirectionAtlasResized = ResizeIndirectionAtlasIfNeeded(GraphBuilder, GlobalShaderMap, IndirectionAtlasRDG);
+		bIndirectionAtlasResized = ResizeIndirectionAtlasIfNeeded(GraphBuilder, SceneUpdateInputs.GlobalShaderMap, IndirectionAtlasRDG);
 	}
 
 	{
@@ -1619,7 +1614,7 @@ void FDistanceFieldSceneData::UpdateDistanceFieldAtlas(
 				}
 				else
 				{
-					TShaderMapRef<FScatterUploadDistanceFieldIndirectionAtlasCS> ComputeShader(GlobalShaderMap);
+					TShaderMapRef<FScatterUploadDistanceFieldIndirectionAtlasCS> ComputeShader(SceneUpdateInputs.GlobalShaderMap);
 
 					auto* PassParameters = GraphBuilder.AllocParameters<FScatterUploadDistanceFieldIndirectionAtlasCS::FParameters>();
 					PassParameters->RWIndirectionAtlas = GraphBuilder.CreateUAV(IndirectionAtlasRDG);
@@ -1657,7 +1652,7 @@ void FDistanceFieldSceneData::UpdateDistanceFieldAtlas(
 				PassParameters->NumBrickUploads = NumBrickUploadsThisPass;
 				PassParameters->BrickSize = DistanceField::BrickSize;
 
-				auto ComputeShader = GlobalShaderMap->GetShader<FScatterUploadDistanceFieldAtlasCS>();
+				auto ComputeShader = SceneUpdateInputs.GlobalShaderMap->GetShader<FScatterUploadDistanceFieldAtlasCS>();
 
 				FComputeShaderUtils::AddPass(
 					GraphBuilder,
@@ -1678,7 +1673,7 @@ void FDistanceFieldSceneData::UpdateDistanceFieldAtlas(
 		UploadAssetData(GraphBuilder, AssetDataUploads, AssetDataBufferRDG);
 	}
 
-	GenerateStreamingRequests(GraphBuilder, View, Scene, bLumenEnabled, GlobalShaderMap);
+	GenerateStreamingRequests(GraphBuilder, SceneUpdateInputs);
 
 	if (GDistanceFieldAtlasLogStats)
 	{
diff --git a/Engine/Source/Runtime/Renderer/Private/GlobalDistanceField.h b/Engine/Source/Runtime/Renderer/Private/GlobalDistanceField.h
--- a/Engine/Source/Runtime/Renderer/Private/GlobalDistanceField.h
+++ b/Engine/Source/Runtime/Renderer/Private/GlobalDistanceField.h
@@ -9,8 +9,11 @@ class FGlobalDistanceFieldInfo;
 class FRDGBuilder;
 class FRDGExternalAccessQueue;
 class FScene;
+class FSceneViewFamily;
 class FViewInfo;
 
+struct FSceneRenderUpdateInputs;
+
 enum FGlobalDFCacheType
 {
 	GDF_MostlyStatic,
diff --git a/Engine/Source/Runtime/Renderer/Private/IndirectLightRendering.cpp b/Engine/Source/Runtime/Renderer/Private/IndirectLightRendering.cpp
--- a/Engine/Source/Runtime/Renderer/Private/IndirectLightRendering.cpp
+++ b/Engine/Source/Runtime/Renderer/Private/IndirectLightRendering.cpp
@@ -1994,8 +1994,8 @@ void FDeferredShadingSceneRenderer::RenderDeferredReflectionsAndSkyLighting(
 		extern int32 GDistanceFieldAOApplyToStaticIndirect;
 		if (Scene->SkyLight->bCastShadows
 			&& !GDistanceFieldAOApplyToStaticIndirect
-			&& ShouldRenderDistanceFieldAO()
-			&& ShouldRenderDistanceFieldLighting()
+			&& ShouldRenderDistanceFieldAO(Views, ViewFamily.EngineShowFlags)
+			&& ShouldRenderDistanceFieldLighting(Scene->DistanceFieldSceneData, Views)
 			&& ViewFamily.EngineShowFlags.AmbientOcclusion
 			&& !bReflectionCapture)
 		{
diff --git a/Engine/Source/Runtime/Renderer/Private/MobileShadingRenderer.cpp b/Engine/Source/Runtime/Renderer/Private/MobileShadingRenderer.cpp
--- a/Engine/Source/Runtime/Renderer/Private/MobileShadingRenderer.cpp
+++ b/Engine/Source/Runtime/Renderer/Private/MobileShadingRenderer.cpp
@@ -1028,19 +1028,26 @@ void FMobileSceneRenderer::Render(FRDGBuilder& GraphBuilder, const FSceneRenderU
 
 	FRDGExternalAccessQueue ExternalAccessQueue;
 
+	if (SceneUpdateInputs)
 	{
 		static auto CVarDistanceFieldShadowQuality = IConsoleManager::Get().FindConsoleVariable(TEXT("r.DFShadowQuality"));
 
 		if (IsMobileDistanceFieldEnabled(ShaderPlatform)
-			&& ViewFamily.EngineShowFlags.Lighting
-			&& !Views[0].bIsReflectionCapture
-			&& !Views[0].bIsPlanarReflection
-			&& !ViewFamily.EngineShowFlags.VisualizeLightCulling
-			&& !ViewFamily.UseDebugViewPS()
-			&& (CVarDistanceFieldShadowQuality != nullptr && CVarDistanceFieldShadowQuality->GetInt() > 0)
+			&& CVarDistanceFieldShadowQuality != nullptr
+			&& CVarDistanceFieldShadowQuality->GetInt() > 0
 			&& bRendererOutputFinalSceneColor)
 		{
-			PrepareDistanceFieldScene(GraphBuilder, ExternalAccessQueue);
+			for (FViewFamilyInfo* Family : SceneUpdateInputs->ViewFamilies)
+			{
+				const FEngineShowFlags& EngineShowFlags = Family->EngineShowFlags;
+				const FSceneView& View = *Family->Views[0];
+	
+				if (EngineShowFlags.Lighting && !EngineShowFlags.VisualizeLightCulling && !Family->UseDebugViewPS() && !View.bIsReflectionCapture && !View.bIsPlanarReflection)
+				{
+					PrepareDistanceFieldScene(GraphBuilder, ExternalAccessQueue, *SceneUpdateInputs);
+					break;
+				}
+			}
 		}
 	}
 
diff --git a/Engine/Source/Runtime/Renderer/Private/ScenePrivate.h b/Engine/Source/Runtime/Renderer/Private/ScenePrivate.h
--- a/Engine/Source/Runtime/Renderer/Private/ScenePrivate.h
+++ b/Engine/Source/Runtime/Renderer/Private/ScenePrivate.h
@@ -2099,10 +2099,7 @@ public:
 	void UpdateDistanceFieldAtlas(
 		FRDGBuilder& GraphBuilder,
 		FRDGExternalAccessQueue& ExternalAccessQueue,
-		const FViewInfo& View,
-		FScene* Scene,
-		bool bLumenEnabled,
-		FGlobalShaderMap* GlobalShaderMap,
+		const FSceneRenderUpdateInputs& SceneUpdateInputs,
 		TArray<FDistanceFieldAssetMipId>& DistanceFieldAssetAdds,
 		TArray<FSetElementId>& DistanceFieldAssetRemoves);
 
@@ -2237,16 +2234,16 @@ private:
 
 	void AsyncUpdate(FRHICommandListBase& RHICmdList, FDistanceFieldAsyncUpdateParameters& UpdateParameters);
 
-	void GenerateStreamingRequests(
-		FRDGBuilder& GraphBuilder, 
-		const FViewInfo& View,
-		FScene* Scene,
-		bool bLumenEnabled,
-		FGlobalShaderMap* GlobalShaderMap);
+	void GenerateStreamingRequests(FRDGBuilder& GraphBuilder, const FSceneRenderUpdateInputs& SceneUpdateInputs);
 
 	friend class FDistanceFieldStreamingUpdateTask;
 };
 
+/**
+ * Prepares the distance field scene for all scene renderers.
+ */
+extern void PrepareDistanceFieldScene(FRDGBuilder& GraphBuilder, FRDGExternalAccessQueue& ExternalAccessQueue, const FSceneRenderUpdateInputs& SceneUpdateInputs);
+
 /** Stores data for an allocation in the FIndirectLightingCache. */
 class FIndirectLightingCacheBlock
 {
diff --git a/Engine/Source/Runtime/Renderer/Private/SceneRendering.h b/Engine/Source/Runtime/Renderer/Private/SceneRendering.h
--- a/Engine/Source/Runtime/Renderer/Private/SceneRendering.h
+++ b/Engine/Source/Runtime/Renderer/Private/SceneRendering.h
@@ -2245,22 +2245,8 @@ public:
 	/** Sets the stereo-compatible RHI viewport. If the view doesn't requires stereo rendering, the standard viewport is set. */
 	static void SetStereoViewport(FRHICommandList& RHICmdList, const FViewInfo& View, float ViewportScale = 1.0f);
 
-	/** Whether distance field global data structures should be prepared for features that use it. */
-	bool ShouldPrepareForDistanceFieldShadows() const;
-	bool ShouldPrepareForDistanceFieldAO() const;
-	bool ShouldPrepareForDFInsetIndirectShadow() const;
 	FGPUSceneDynamicContext& GetGPUSceneDynamicContext() { return GPUSceneDynamicContext; }
 
-	bool ShouldPrepareDistanceFieldScene() const;
-	bool ShouldPrepareGlobalDistanceField() const;
-	bool ShouldPrepareHeightFieldScene() const;
-
-	bool ShouldRenderDistanceFieldAO() const;
-	void UpdateGlobalDistanceFieldObjectBuffers(FRDGBuilder& GraphBuilder);
-	void UpdateGlobalHeightFieldObjectBuffers(FRDGBuilder& GraphBuilder, const TArray<uint32>& IndicesToUpdateInHeightFieldObjectBuffers);
-	void ProcessPendingHeightFieldPrimitiveAddAndRemoveOps(TArray<uint32>& IndicesToUpdateInHeightFieldObjectBuffers);
-	void PrepareDistanceFieldScene(FRDGBuilder& GraphBuilder, FRDGExternalAccessQueue& ExternalAccessQueue);
-
 	void VoxelizeFogVolumePrimitives(
 		FRDGBuilder& GraphBuilder,
 		const FViewInfo& View,
diff --git a/Engine/Source/Runtime/Renderer/Private/VelocityRendering.cpp b/Engine/Source/Runtime/Renderer/Private/VelocityRendering.cpp
--- a/Engine/Source/Runtime/Renderer/Private/VelocityRendering.cpp
+++ b/Engine/Source/Runtime/Renderer/Private/VelocityRendering.cpp
@@ -192,7 +192,7 @@ bool FDeferredShadingSceneRenderer::ShouldRenderVelocities() const
 		bool bTemporalAA = IsTemporalAccumulationBasedMethod(View.AntiAliasingMethod) && !View.bCameraCut;
 		bool bMotionBlur = IsMotionBlurEnabled(View);
 		bool bVisualizeMotionblur = View.Family->EngineShowFlags.VisualizeMotionBlur || View.Family->EngineShowFlags.VisualizeTemporalUpscaler;
-		bool bDistanceFieldAO = ShouldPrepareForDistanceFieldAO();
+		bool bDistanceFieldAO = ShouldPrepareForDistanceFieldAO(Scene, ViewFamily, AnyViewHasGIMethodSupportingDFAO());
 
 		bool bSceneSSREnabled = ViewPipelineState.ReflectionsMethod == EReflectionsMethod::SSR && ScreenSpaceRayTracing::ShouldRenderScreenSpaceReflections(View);
 		bool bWaterSSREnabled = ViewPipelineState.ReflectionsMethodWater == EReflectionsMethod::SSR && ScreenSpaceRayTracing::ShouldRenderScreenSpaceReflectionsWater(View);
diff --git a/Engine/Source/Runtime/Renderer/Private/VolumetricFog.cpp b/Engine/Source/Runtime/Renderer/Private/VolumetricFog.cpp
--- a/Engine/Source/Runtime/Renderer/Private/VolumetricFog.cpp
+++ b/Engine/Source/Runtime/Renderer/Private/VolumetricFog.cpp
@@ -1830,7 +1830,7 @@ void FSceneRenderer::ComputeVolumetricFog(FRDGBuilder& GraphBuilder,
 				&& Scene->SkyLight
 				&& Scene->SkyLight->bCastShadows
 				&& Scene->SkyLight->bCastVolumetricShadow
-				&& ShouldRenderDistanceFieldAO()
+				&& ShouldRenderDistanceFieldAO(Views, ViewFamily.EngineShowFlags)
 				&& SupportsDistanceFieldAO(View.GetFeatureLevel(), View.GetShaderPlatform())
 				&& bUseGlobalDistanceField
 				&& Views.Num() == 1
