Fixed crash due to incorrect ordering of skylight check in distance field AO.

diff --git a/Engine/Source/Runtime/Renderer/Private/DistanceFieldAmbientOcclusion.cpp b/Engine/Source/Runtime/Renderer/Private/DistanceFieldAmbientOcclusion.cpp
--- a/Engine/Source/Runtime/Renderer/Private/DistanceFieldAmbientOcclusion.cpp
+++ b/Engine/Source/Runtime/Renderer/Private/DistanceFieldAmbientOcclusion.cpp
@@ -718,7 +718,7 @@ bool ShouldPrepareForDistanceFieldAO(const FScene* Scene, const FSceneViewFamily
 				return true;
 			}
 
-			if (Scene->SkyLight->bCastShadows && ShouldRenderDeferredDynamicSkyLight(Scene, ViewFamily))
+			if (ShouldRenderDeferredDynamicSkyLight(Scene, ViewFamily) && Scene->SkyLight->bCastShadows)
 			{
 				return true;
 			}
