Fixed mobile shading to use dummy instance culling parameters if no mesh pass exists.

diff --git a/Engine/Source/Runtime/Renderer/Private/MobileShadingRenderer.cpp b/Engine/Source/Runtime/Renderer/Private/MobileShadingRenderer.cpp
--- a/Engine/Source/Runtime/Renderer/Private/MobileShadingRenderer.cpp
+++ b/Engine/Source/Runtime/Renderer/Private/MobileShadingRenderer.cpp
@@ -196,7 +196,7 @@ static void RenderOpaqueFX(
 	}
 }
 
-static void BuildMeshRenderingCommands(FRDGBuilder& GraphBuilder, EMeshPass::Type MeshPass, FViewInfo& View, const FGPUScene& GPUScene, FInstanceCullingDrawParams& OutParams)
+static void BuildMeshRenderingCommands(FRDGBuilder& GraphBuilder, EMeshPass::Type MeshPass, FViewInfo& View, const FGPUScene& GPUScene, FInstanceCullingManager& InstanceCullingManager, FInstanceCullingDrawParams& OutParams)
 {
 	if (auto* Pass = View.ParallelMeshDrawCommandPasses[MeshPass])
 	{
@@ -213,6 +213,10 @@ static void BuildMeshRenderingCommands(FRDGBuilder& GraphBuilder, EMeshPass::Typ
 				ERDGPassFlags::Raster | ERDGPassFlags::SkipRenderPass | ERDGPassFlags::NeverCull, [](FRHICommandList&){});
 		}
 	}
+	else
+	{
+		InstanceCullingManager.SetDummyCullingParams(GraphBuilder, OutParams);
+	}
 }
 
 BEGIN_SHADER_PARAMETER_STRUCT(FMobileRenderPassParameters,)
@@ -797,7 +801,7 @@ static void EndOcclusionScope(FRDGBuilder& GraphBuilder, TArray<FViewInfo>& View
 /*
 * Renders the Full Depth Prepass
 */
-void FMobileSceneRenderer::RenderFullDepthPrepass(FRDGBuilder& GraphBuilder, TArrayView<FViewInfo> InViews, FSceneTextures& SceneTextures, bool bIsSceneCaptureRenderPass)
+void FMobileSceneRenderer::RenderFullDepthPrepass(FRDGBuilder& GraphBuilder, TArrayView<FViewInfo> InViews, FSceneTextures& SceneTextures, FInstanceCullingManager& InstanceCullingManager, bool bIsSceneCaptureRenderPass)
 {
 	FRenderViewContextArray RenderViews;
 	GetRenderViews(InViews, RenderViews);
@@ -845,10 +849,10 @@ void FMobileSceneRenderer::RenderFullDepthPrepass(FRDGBuilder& GraphBuilder, TAr
 			}
 		}
 
-		BuildMeshRenderingCommands(GraphBuilder, EMeshPass::DepthPass, View, Scene->GPUScene, PassParameters->InstanceCullingDrawParams);
+		BuildMeshRenderingCommands(GraphBuilder, EMeshPass::DepthPass, View, Scene->GPUScene, InstanceCullingManager, PassParameters->InstanceCullingDrawParams);
 		if (bRenderVelocityInDepthPrePass)
 		{
-			BuildMeshRenderingCommands(GraphBuilder, EMeshPass::Velocity, View, Scene->GPUScene, ParameterCollection->VelocityInstanceCullingDrawParams);
+			BuildMeshRenderingCommands(GraphBuilder, EMeshPass::Velocity, View, Scene->GPUScene, InstanceCullingManager, ParameterCollection->VelocityInstanceCullingDrawParams);
 		}
 
 		GraphBuilder.AddPass(
@@ -907,7 +911,7 @@ void FMobileSceneRenderer::RenderMaskedPrePass(FRHICommandList& RHICmdList, cons
 	}
 }
 
-void FMobileSceneRenderer::RenderCustomRenderPassBasePass(FRDGBuilder& GraphBuilder, TArrayView<FViewInfo> InViews, FRDGTextureRef ViewFamilyTexture, FSceneTextures& SceneTextures, bool bIncludeTranslucent)
+void FMobileSceneRenderer::RenderCustomRenderPassBasePass(FRDGBuilder& GraphBuilder, TArrayView<FViewInfo> InViews, FRDGTextureRef ViewFamilyTexture, FSceneTextures& SceneTextures, FInstanceCullingManager& InstanceCullingManager, bool bIncludeTranslucent)
 {
 	FRenderTargetBindingSlots BasePassRenderTargets;
 	// Use the same subpass hints as main render, to avoid generating new PSOs 
@@ -960,11 +964,11 @@ void FMobileSceneRenderer::RenderCustomRenderPassBasePass(FRDGBuilder& GraphBuil
 
 		if (Scene->GPUScene.IsEnabled())
 		{
-			BuildMeshRenderingCommands(GraphBuilder, EMeshPass::BasePass, View, Scene->GPUScene, PassParameters->InstanceCullingDrawParams);
-			BuildMeshRenderingCommands(GraphBuilder, EMeshPass::SkyPass, View, Scene->GPUScene, ParameterCollection->SkyPassInstanceCullingDrawParams);
+			BuildMeshRenderingCommands(GraphBuilder, EMeshPass::BasePass, View, Scene->GPUScene, InstanceCullingManager, PassParameters->InstanceCullingDrawParams);
+			BuildMeshRenderingCommands(GraphBuilder, EMeshPass::SkyPass, View, Scene->GPUScene, InstanceCullingManager, ParameterCollection->SkyPassInstanceCullingDrawParams);
 			if (bIncludeTranslucent)
 			{
-				BuildMeshRenderingCommands(GraphBuilder, EMeshPass::TranslucencyAll, View, Scene->GPUScene, TranslucencyPassParameters->InstanceCullingDrawParams);
+				BuildMeshRenderingCommands(GraphBuilder, EMeshPass::TranslucencyAll, View, Scene->GPUScene, InstanceCullingManager, TranslucencyPassParameters->InstanceCullingDrawParams);
 			}
 		}
 
@@ -1358,7 +1362,7 @@ void FMobileSceneRenderer::Render(FRDGBuilder& GraphBuilder, const FSceneRenderU
 
 				if (bIsFullDepthPrepassEnabled)
 				{
-					RenderFullDepthPrepass(GraphBuilder, CustomRenderPassViews, SceneTextures, true);
+					RenderFullDepthPrepass(GraphBuilder, CustomRenderPassViews, SceneTextures, InstanceCullingManager, true);
 					if (!bRequiresSceneDepthAux)
 					{
 						AddResolveSceneDepthPass(GraphBuilder, CustomRenderPassViews, SceneTextures.Depth);
@@ -1368,7 +1372,7 @@ void FMobileSceneRenderer::Render(FRDGBuilder& GraphBuilder, const FSceneRenderU
 				// Render base pass if the custom pass requires it. Otherwise if full depth prepass is not enabled, then depth is generated in the base pass.
 				if (CustomRenderPass->GetRenderMode() == FCustomRenderPassBase::ERenderMode::DepthAndBasePass || (CustomRenderPass->GetRenderMode() == FCustomRenderPassBase::ERenderMode::DepthPass && !bIsFullDepthPrepassEnabled))
 				{
-					RenderCustomRenderPassBasePass(GraphBuilder, CustomRenderPassViews, ViewFamilyTexture, SceneTextures, CustomRenderPass->IsTranslucentIncluded());
+					RenderCustomRenderPassBasePass(GraphBuilder, CustomRenderPassViews, ViewFamilyTexture, SceneTextures, InstanceCullingManager, CustomRenderPass->IsTranslucentIncluded());
 				}
 
 				SceneTextures.MobileSetupMode = EMobileSceneTextureSetupMode::SceneColor | EMobileSceneTextureSetupMode::SceneDepth | EMobileSceneTextureSetupMode::SceneDepthAux;
@@ -1396,7 +1400,7 @@ void FMobileSceneRenderer::Render(FRDGBuilder& GraphBuilder, const FSceneRenderU
 	FDBufferTextures DBufferTextures{};
 	if (bIsFullDepthPrepassEnabled)
 	{
-		RenderFullDepthPrepass(GraphBuilder, Views, SceneTextures);
+		RenderFullDepthPrepass(GraphBuilder, Views, SceneTextures, InstanceCullingManager);
 
 		if (!bRequiresSceneDepthAux)
 		{
@@ -1465,16 +1469,16 @@ void FMobileSceneRenderer::Render(FRDGBuilder& GraphBuilder, const FSceneRenderU
 		{
 			if (bRequiresMultiPass)
 			{
-				RenderDeferredMultiPass(GraphBuilder, SceneTextures, SortedLightSet);
+				RenderDeferredMultiPass(GraphBuilder, SceneTextures, SortedLightSet, InstanceCullingManager);
 			}
 			else
 			{
-				RenderDeferredSinglePass(GraphBuilder, SceneTextures, SortedLightSet);
+				RenderDeferredSinglePass(GraphBuilder, SceneTextures, SortedLightSet, InstanceCullingManager);
 			}
 		}
 		else
 		{
-			RenderForward(GraphBuilder, ViewFamilyTexture, SceneTextures, DBufferTextures);
+			RenderForward(GraphBuilder, ViewFamilyTexture, SceneTextures, DBufferTextures, InstanceCullingManager);
 		}
 
 		if (DoOcclusionQueries() && !bIsFullDepthPrepassEnabled)
@@ -1641,7 +1645,7 @@ FRenderTargetBindingSlots FMobileSceneRenderer::InitRenderTargetBindings_Forward
 	return BasePassRenderTargets;
 }
 
-void FMobileSceneRenderer::RenderForward(FRDGBuilder& GraphBuilder, FRDGTextureRef ViewFamilyTexture, FSceneTextures& SceneTextures, FDBufferTextures& DBufferTextures)
+void FMobileSceneRenderer::RenderForward(FRDGBuilder& GraphBuilder, FRDGTextureRef ViewFamilyTexture, FSceneTextures& SceneTextures, FDBufferTextures& DBufferTextures, FInstanceCullingManager& InstanceCullingManager)
 {
 	const FViewInfo& MainView = Views[0];
 
@@ -1702,11 +1706,11 @@ void FMobileSceneRenderer::RenderForward(FRDGBuilder& GraphBuilder, FRDGTextureR
 		// Split if we need to render translucency in a separate render pass
 		if (bRequiresMultiPass)
 		{
-			RenderForwardMultiPass(GraphBuilder, PassParameters, ViewContext, SceneTextures);
+			RenderForwardMultiPass(GraphBuilder, PassParameters, ViewContext, SceneTextures, InstanceCullingManager);
 		}
 		else
 		{
-			RenderForwardSinglePass(GraphBuilder, PassParameters, ViewContext, SceneTextures);
+			RenderForwardSinglePass(GraphBuilder, PassParameters, ViewContext, SceneTextures, InstanceCullingManager);
 		}
 	}
 	
@@ -1716,7 +1720,7 @@ void FMobileSceneRenderer::RenderForward(FRDGBuilder& GraphBuilder, FRDGTextureR
 	}
 }
 
-void FMobileSceneRenderer::RenderForwardSinglePass(FRDGBuilder& GraphBuilder, FMobileRenderPassParameters* PassParameters, FRenderViewContext& ViewContext, FSceneTextures& SceneTextures)
+void FMobileSceneRenderer::RenderForwardSinglePass(FRDGBuilder& GraphBuilder, FMobileRenderPassParameters* PassParameters, FRenderViewContext& ViewContext, FSceneTextures& SceneTextures, FInstanceCullingManager& InstanceCullingManager)
 {
 	struct FForwardSinglePassParameterCollection
 	{
@@ -1734,13 +1738,13 @@ void FMobileSceneRenderer::RenderForwardSinglePass(FRDGBuilder& GraphBuilder, FM
 	{
 		if (!bIsFullDepthPrepassEnabled)
 		{
-			BuildMeshRenderingCommands(GraphBuilder, EMeshPass::DepthPass, View, Scene->GPUScene, ParameterCollection->DepthPassInstanceCullingDrawParams);
+			BuildMeshRenderingCommands(GraphBuilder, EMeshPass::DepthPass, View, Scene->GPUScene, InstanceCullingManager, ParameterCollection->DepthPassInstanceCullingDrawParams);
 		}
-		BuildMeshRenderingCommands(GraphBuilder, EMeshPass::BasePass, View, Scene->GPUScene, PassParameters->InstanceCullingDrawParams);
-		BuildMeshRenderingCommands(GraphBuilder, EMeshPass::SkyPass, View, Scene->GPUScene, ParameterCollection->SkyPassInstanceCullingDrawParams);
-		BuildMeshRenderingCommands(GraphBuilder, EMeshPass::DebugViewMode, View, Scene->GPUScene, ParameterCollection->DebugViewModeInstanceCullingDrawParams);
-		BuildMeshRenderingCommands(GraphBuilder, EMeshPass::MeshDecal_SceneColor, View, Scene->GPUScene, ParameterCollection->MeshDecalSceneColorInstanceCullingDrawParams);
-		BuildMeshRenderingCommands(GraphBuilder, StandardTranslucencyMeshPass, View, Scene->GPUScene, ParameterCollection->TranslucencyInstanceCullingDrawParams);
+		BuildMeshRenderingCommands(GraphBuilder, EMeshPass::BasePass, View, Scene->GPUScene, InstanceCullingManager, PassParameters->InstanceCullingDrawParams);
+		BuildMeshRenderingCommands(GraphBuilder, EMeshPass::SkyPass, View, Scene->GPUScene, InstanceCullingManager, ParameterCollection->SkyPassInstanceCullingDrawParams);
+		BuildMeshRenderingCommands(GraphBuilder, EMeshPass::DebugViewMode, View, Scene->GPUScene, InstanceCullingManager, ParameterCollection->DebugViewModeInstanceCullingDrawParams);
+		BuildMeshRenderingCommands(GraphBuilder, EMeshPass::MeshDecal_SceneColor, View, Scene->GPUScene, InstanceCullingManager, ParameterCollection->MeshDecalSceneColorInstanceCullingDrawParams);
+		BuildMeshRenderingCommands(GraphBuilder, StandardTranslucencyMeshPass, View, Scene->GPUScene, InstanceCullingManager, ParameterCollection->TranslucencyInstanceCullingDrawParams);
 	}
 
 	if (bTonemapSubpassInline)
@@ -1847,7 +1851,7 @@ void FMobileSceneRenderer::RenderForwardSinglePass(FRDGBuilder& GraphBuilder, FM
 	}
 }
 
-void FMobileSceneRenderer::RenderForwardMultiPass(FRDGBuilder& GraphBuilder, FMobileRenderPassParameters* PassParameters, FRenderViewContext& ViewContext, FSceneTextures& SceneTextures)
+void FMobileSceneRenderer::RenderForwardMultiPass(FRDGBuilder& GraphBuilder, FMobileRenderPassParameters* PassParameters, FRenderViewContext& ViewContext, FSceneTextures& SceneTextures, FInstanceCullingManager& InstanceCullingManager)
 {
 	struct FForwardFirstPassParameterCollection
 	{
@@ -1863,11 +1867,11 @@ void FMobileSceneRenderer::RenderForwardMultiPass(FRDGBuilder& GraphBuilder, FMo
 	{
 		if (!bIsFullDepthPrepassEnabled)
 		{
-			BuildMeshRenderingCommands(GraphBuilder, EMeshPass::DepthPass, View, Scene->GPUScene, ParameterCollection->DepthPassInstanceCullingDrawParams);
+			BuildMeshRenderingCommands(GraphBuilder, EMeshPass::DepthPass, View, Scene->GPUScene, InstanceCullingManager, ParameterCollection->DepthPassInstanceCullingDrawParams);
 		}
-		BuildMeshRenderingCommands(GraphBuilder, EMeshPass::BasePass, View, Scene->GPUScene, PassParameters->InstanceCullingDrawParams);
-		BuildMeshRenderingCommands(GraphBuilder, EMeshPass::SkyPass, View, Scene->GPUScene, ParameterCollection->SkyPassInstanceCullingDrawParams);
-		BuildMeshRenderingCommands(GraphBuilder, EMeshPass::DebugViewMode, View, Scene->GPUScene, ParameterCollection->DebugViewModeInstanceCullingDrawParams);
+		BuildMeshRenderingCommands(GraphBuilder, EMeshPass::BasePass, View, Scene->GPUScene, InstanceCullingManager, PassParameters->InstanceCullingDrawParams);
+		BuildMeshRenderingCommands(GraphBuilder, EMeshPass::SkyPass, View, Scene->GPUScene, InstanceCullingManager, ParameterCollection->SkyPassInstanceCullingDrawParams);
+		BuildMeshRenderingCommands(GraphBuilder, EMeshPass::DebugViewMode, View, Scene->GPUScene, InstanceCullingManager, ParameterCollection->DebugViewModeInstanceCullingDrawParams);
 	}
 
 	GraphBuilder.AddPass(
@@ -1938,8 +1942,8 @@ void FMobileSceneRenderer::RenderForwardMultiPass(FRDGBuilder& GraphBuilder, FMo
 
 	if (Scene->GPUScene.IsEnabled())
 	{
-		BuildMeshRenderingCommands(GraphBuilder, EMeshPass::MeshDecal_SceneColor, View, Scene->GPUScene, SecondParameterCollection->MeshDecalSceneColorInstanceCullingDrawParams);
-		BuildMeshRenderingCommands(GraphBuilder, StandardTranslucencyMeshPass, View, Scene->GPUScene, SecondParameterCollection->TranslucencyInstanceCullingDrawParams);
+		BuildMeshRenderingCommands(GraphBuilder, EMeshPass::MeshDecal_SceneColor, View, Scene->GPUScene, InstanceCullingManager, SecondParameterCollection->MeshDecalSceneColorInstanceCullingDrawParams);
+		BuildMeshRenderingCommands(GraphBuilder, StandardTranslucencyMeshPass, View, Scene->GPUScene, InstanceCullingManager, SecondParameterCollection->TranslucencyInstanceCullingDrawParams);
 	}
 
 	GraphBuilder.AddPass(
@@ -2106,7 +2110,7 @@ FRenderTargetBindingSlots FMobileSceneRenderer::InitRenderTargetBindings_Deferre
 	return BasePassRenderTargets;
 }
 
-void FMobileSceneRenderer::RenderDeferredSinglePass(FRDGBuilder& GraphBuilder, FSceneTextures& SceneTextures, const FSortedLightSetSceneInfo& SortedLightSet)
+void FMobileSceneRenderer::RenderDeferredSinglePass(FRDGBuilder& GraphBuilder, FSceneTextures& SceneTextures, const FSortedLightSetSceneInfo& SortedLightSet, FInstanceCullingManager& InstanceCullingManager)
 {
 	bool bUsingPixelLocalStorage = UsingPixelLocalStorage(ShaderPlatform);
 	FColorTargets ColorTargets = GetColorTargets_Deferred(SceneTextures);
@@ -2172,13 +2176,13 @@ void FMobileSceneRenderer::RenderDeferredSinglePass(FRDGBuilder& GraphBuilder, F
 		{
 			if (!bIsFullDepthPrepassEnabled)
 			{
-				BuildMeshRenderingCommands(GraphBuilder, EMeshPass::DepthPass, View, Scene->GPUScene, ParameterCollection->DepthPassInstanceCullingDrawParams);
+				BuildMeshRenderingCommands(GraphBuilder, EMeshPass::DepthPass, View, Scene->GPUScene, InstanceCullingManager, ParameterCollection->DepthPassInstanceCullingDrawParams);
 			}
-			BuildMeshRenderingCommands(GraphBuilder, EMeshPass::BasePass, View, Scene->GPUScene, PassParameters->InstanceCullingDrawParams);
-			BuildMeshRenderingCommands(GraphBuilder, EMeshPass::SkyPass, View, Scene->GPUScene, ParameterCollection->SkyPassInstanceCullingDrawParams);
-			BuildMeshRenderingCommands(GraphBuilder, EMeshPass::DebugViewMode, View, Scene->GPUScene, ParameterCollection->DebugViewModeInstanceCullingDrawParams);
-			BuildMeshRenderingCommands(GraphBuilder, EMeshPass::MeshDecal_SceneColorAndGBuffer, View, Scene->GPUScene, ParameterCollection->MeshDecalSceneColorAndGBufferInstanceCullingDrawParams);
-			BuildMeshRenderingCommands(GraphBuilder, StandardTranslucencyMeshPass, View, Scene->GPUScene, ParameterCollection->TranslucencyInstanceCullingDrawParams);
+			BuildMeshRenderingCommands(GraphBuilder, EMeshPass::BasePass, View, Scene->GPUScene, InstanceCullingManager, PassParameters->InstanceCullingDrawParams);
+			BuildMeshRenderingCommands(GraphBuilder, EMeshPass::SkyPass, View, Scene->GPUScene, InstanceCullingManager, ParameterCollection->SkyPassInstanceCullingDrawParams);
+			BuildMeshRenderingCommands(GraphBuilder, EMeshPass::DebugViewMode, View, Scene->GPUScene, InstanceCullingManager, ParameterCollection->DebugViewModeInstanceCullingDrawParams);
+			BuildMeshRenderingCommands(GraphBuilder, EMeshPass::MeshDecal_SceneColorAndGBuffer, View, Scene->GPUScene, InstanceCullingManager, ParameterCollection->MeshDecalSceneColorAndGBufferInstanceCullingDrawParams);
+			BuildMeshRenderingCommands(GraphBuilder, StandardTranslucencyMeshPass, View, Scene->GPUScene, InstanceCullingManager, ParameterCollection->TranslucencyInstanceCullingDrawParams);
 		}
 
 		GraphBuilder.AddPass(
@@ -2222,7 +2226,7 @@ void FMobileSceneRenderer::RenderDeferredSinglePass(FRDGBuilder& GraphBuilder, F
 	}
 }
 
-void FMobileSceneRenderer::RenderDeferredMultiPass(FRDGBuilder& GraphBuilder, FSceneTextures& SceneTextures, const FSortedLightSetSceneInfo& SortedLightSet)
+void FMobileSceneRenderer::RenderDeferredMultiPass(FRDGBuilder& GraphBuilder, FSceneTextures& SceneTextures, const FSortedLightSetSceneInfo& SortedLightSet, FInstanceCullingManager& InstanceCullingManager)
 {
 	FColorTargets ColorTargets = GetColorTargets_Deferred(SceneTextures);
 	FRenderTargetBindingSlots BasePassRenderTargets = InitRenderTargetBindings_Deferred(SceneTextures, ColorTargets);
@@ -2274,14 +2278,14 @@ void FMobileSceneRenderer::RenderDeferredMultiPass(FRDGBuilder& GraphBuilder, FS
 		{
 			if (!bIsFullDepthPrepassEnabled)
 			{
-				BuildMeshRenderingCommands(GraphBuilder, EMeshPass::DepthPass, View, Scene->GPUScene, ParameterCollection->DepthPassInstanceCullingDrawParams);
+				BuildMeshRenderingCommands(GraphBuilder, EMeshPass::DepthPass, View, Scene->GPUScene, InstanceCullingManager, ParameterCollection->DepthPassInstanceCullingDrawParams);
 			}
-			BuildMeshRenderingCommands(GraphBuilder, EMeshPass::BasePass, View, Scene->GPUScene, PassParameters->InstanceCullingDrawParams);
-			BuildMeshRenderingCommands(GraphBuilder, EMeshPass::SkyPass, View, Scene->GPUScene, ParameterCollection->SkyPassInstanceCullingDrawParams);
-			BuildMeshRenderingCommands(GraphBuilder, EMeshPass::DebugViewMode, View, Scene->GPUScene, ParameterCollection->DebugViewModeInstanceCullingDrawParams);
+			BuildMeshRenderingCommands(GraphBuilder, EMeshPass::BasePass, View, Scene->GPUScene, InstanceCullingManager, PassParameters->InstanceCullingDrawParams);
+			BuildMeshRenderingCommands(GraphBuilder, EMeshPass::SkyPass, View, Scene->GPUScene, InstanceCullingManager, ParameterCollection->SkyPassInstanceCullingDrawParams);
+			BuildMeshRenderingCommands(GraphBuilder, EMeshPass::DebugViewMode, View, Scene->GPUScene, InstanceCullingManager, ParameterCollection->DebugViewModeInstanceCullingDrawParams);
 			if (bIsFullDepthPrepassEnabled)
 			{
-				BuildMeshRenderingCommands(GraphBuilder, EMeshPass::MeshDecal_SceneColorAndGBuffer, View, Scene->GPUScene, ParameterCollection->MeshDecalSceneColorAndGBufferInstanceCullingDrawParams);
+				BuildMeshRenderingCommands(GraphBuilder, EMeshPass::MeshDecal_SceneColorAndGBuffer, View, Scene->GPUScene, InstanceCullingManager, ParameterCollection->MeshDecalSceneColorAndGBufferInstanceCullingDrawParams);
 			}
 		}
 
@@ -2353,8 +2357,8 @@ void FMobileSceneRenderer::RenderDeferredMultiPass(FRDGBuilder& GraphBuilder, FS
 
 			if (Scene->GPUScene.IsEnabled())
 			{
-				BuildMeshRenderingCommands(GraphBuilder, EMeshPass::MeshDecal_SceneColor, View, Scene->GPUScene, PassParameters->InstanceCullingDrawParams);
-				BuildMeshRenderingCommands(GraphBuilder, EMeshPass::MeshDecal_SceneColorAndGBuffer, View, Scene->GPUScene, ParameterCollection->MeshDecalSceneColorAndGBufferInstanceCullingDrawParams);
+				BuildMeshRenderingCommands(GraphBuilder, EMeshPass::MeshDecal_SceneColor, View, Scene->GPUScene, InstanceCullingManager, PassParameters->InstanceCullingDrawParams);
+				BuildMeshRenderingCommands(GraphBuilder, EMeshPass::MeshDecal_SceneColorAndGBuffer, View, Scene->GPUScene, InstanceCullingManager, ParameterCollection->MeshDecalSceneColorAndGBufferInstanceCullingDrawParams);
 			}
 
 			GraphBuilder.AddPass(
@@ -2400,7 +2404,7 @@ void FMobileSceneRenderer::RenderDeferredMultiPass(FRDGBuilder& GraphBuilder, FS
 
 		if (Scene->GPUScene.IsEnabled())
 		{
-			BuildMeshRenderingCommands(GraphBuilder, StandardTranslucencyMeshPass, View, Scene->GPUScene, PassParameters->InstanceCullingDrawParams);
+			BuildMeshRenderingCommands(GraphBuilder, StandardTranslucencyMeshPass, View, Scene->GPUScene, InstanceCullingManager, PassParameters->InstanceCullingDrawParams);
 		}
 
 		const EMobileSSRQuality MobileSSRQuality = ActiveMobileSSRQuality(View, bShouldRenderVelocities);
diff --git a/Engine/Source/Runtime/Renderer/Private/SceneRendering.h b/Engine/Source/Runtime/Renderer/Private/SceneRendering.h
--- a/Engine/Source/Runtime/Renderer/Private/SceneRendering.h
+++ b/Engine/Source/Runtime/Renderer/Private/SceneRendering.h
@@ -2803,13 +2803,13 @@ protected:
 
 	void RenderPrePass(FRHICommandList& RHICmdList, const FViewInfo& View, const FInstanceCullingDrawParams* InstanceCullingDrawParams);
 	void RenderMaskedPrePass(FRHICommandList& RHICmdList, const FViewInfo& View, const FInstanceCullingDrawParams* DepthPassInstanceCullingDrawParams);
-	void RenderFullDepthPrepass(FRDGBuilder& GraphBuilder, TArrayView<FViewInfo> InViews, FSceneTextures& SceneTextures, bool bIsSceneCaptureRenderPass=false);
+	void RenderFullDepthPrepass(FRDGBuilder& GraphBuilder, TArrayView<FViewInfo> InViews, FSceneTextures& SceneTextures, FInstanceCullingManager& InstanceCullingManager, bool bIsSceneCaptureRenderPass=false);
 
 	void RenderVelocityPass(FRHICommandList& RHICmdList, const FViewInfo& View, const FInstanceCullingDrawParams* InstanceCullingDrawParams);
 
 	void RenderMobileLocalLightsBuffer(FRDGBuilder& GraphBuilder, FSceneTextures& SceneTextures, const FSortedLightSetSceneInfo& SortedLights);
 
-	void RenderCustomRenderPassBasePass(FRDGBuilder& GraphBuilder, TArrayView<FViewInfo> InViews, FRDGTextureRef ViewFamilyTexture, FSceneTextures& SceneTextures, bool bIncludeTranslucent);
+	void RenderCustomRenderPassBasePass(FRDGBuilder& GraphBuilder, TArrayView<FViewInfo> InViews, FRDGTextureRef ViewFamilyTexture, FSceneTextures& SceneTextures, FInstanceCullingManager& InstanceCullingManager, bool bIncludeTranslucent);
 
 	/** Renders the opaque base pass for mobile. */
 	void RenderMobileBasePass(FRHICommandList& RHICmdList, const FViewInfo& View, const FInstanceCullingDrawParams* InstanceCullingDrawParams, const FInstanceCullingDrawParams* SkyPassInstanceCullingDrawParams);
@@ -2861,12 +2861,12 @@ protected:
 	void UpdateDirectionalLightUniformBuffers(FRDGBuilder& GraphBuilder, const FViewInfo& View);
 	void UpdateSkyReflectionUniformBuffer(FRHICommandListBase& RHICmdList);
 	
-	void RenderForward(FRDGBuilder& GraphBuilder, FRDGTextureRef ViewFamilyTexture, FSceneTextures& SceneTextures, FDBufferTextures& DBufferTextures);
-	void RenderForwardSinglePass(FRDGBuilder& GraphBuilder, class FMobileRenderPassParameters* PassParameters, struct FRenderViewContext& ViewContext, FSceneTextures& SceneTextures);
-	void RenderForwardMultiPass(FRDGBuilder& GraphBuilder, class FMobileRenderPassParameters* PassParameters, struct FRenderViewContext& ViewContext, FSceneTextures& SceneTextures);
+	void RenderForward(FRDGBuilder& GraphBuilder, FRDGTextureRef ViewFamilyTexture, FSceneTextures& SceneTextures, FDBufferTextures& DBufferTextures, FInstanceCullingManager& InstanceCullingManager);
+	void RenderForwardSinglePass(FRDGBuilder& GraphBuilder, class FMobileRenderPassParameters* PassParameters, struct FRenderViewContext& ViewContext, FSceneTextures& SceneTextures, FInstanceCullingManager& InstanceCullingManager);
+	void RenderForwardMultiPass(FRDGBuilder& GraphBuilder, class FMobileRenderPassParameters* PassParameters, struct FRenderViewContext& ViewContext, FSceneTextures& SceneTextures, FInstanceCullingManager& InstanceCullingManager);
 	
-	void RenderDeferredSinglePass(FRDGBuilder& GraphBuilder, FSceneTextures& SceneTextures, const FSortedLightSetSceneInfo& SortedLightSet);
-	void RenderDeferredMultiPass(FRDGBuilder& GraphBuilder, FSceneTextures& SceneTextures, const FSortedLightSetSceneInfo& SortedLightSet);
+	void RenderDeferredSinglePass(FRDGBuilder& GraphBuilder, FSceneTextures& SceneTextures, const FSortedLightSetSceneInfo& SortedLightSet, FInstanceCullingManager& InstanceCullingManager);
+	void RenderDeferredMultiPass(FRDGBuilder& GraphBuilder, FSceneTextures& SceneTextures, const FSortedLightSetSceneInfo& SortedLightSet, FInstanceCullingManager& InstanceCullingManager);
 
 	void RenderAmbientOcclusion(FRDGBuilder& GraphBuilder, FRDGTextureRef SceneDepthTexture, FRDGTextureRef AmbientOcclusionTexture);
 
