* Fix hit proxy rendering with Groom when using PPLL rendering
* Fix outline selection when using PPLL rendering

diff --git a/Engine/Source/Runtime/Renderer/Private/HairStrands/HairStrandsVisibility.cpp b/Engine/Source/Runtime/Renderer/Private/HairStrands/HairStrandsVisibility.cpp
--- a/Engine/Source/Runtime/Renderer/Private/HairStrands/HairStrandsVisibility.cpp
+++ b/Engine/Source/Runtime/Renderer/Private/HairStrands/HairStrandsVisibility.cpp
@@ -4399,7 +4399,10 @@ void DrawHitProxies(
 	// Proxy rendering is only supported/compatible with MSAA-visibility rendering. 
 	// PPLL is not supported, but it is supposed to be used only for final render.
 	const EHairVisibilityRenderMode RenderMode = GetHairVisibilityRenderMode(View.GetShaderPlatform());
-	if (RenderMode != HairVisibilityRenderMode_MSAA_Visibility && RenderMode != HairVisibilityRenderMode_ComputeRaster && RenderMode != HairVisibilityRenderMode_ComputeRasterForward)
+	if (RenderMode != HairVisibilityRenderMode_MSAA_Visibility && 
+		RenderMode != HairVisibilityRenderMode_PPLL && 
+		RenderMode != HairVisibilityRenderMode_ComputeRaster && 
+		RenderMode != HairVisibilityRenderMode_ComputeRasterForward)
 	{
 		return;
 	}
@@ -4519,6 +4522,43 @@ void DrawHitProxies(
 			IndirectArgsBuffer,
 			OutMaxNodeCount);
 	}
+	else if (RenderMode == HairVisibilityRenderMode_PPLL)
+	{
+		// Linked list generation pass
+		FRDGTextureRef PPLLNodeCounterTexture;
+		FRDGTextureRef PPLLNodeIndexTexture;
+		FRDGBufferRef PPLLNodeDataBuffer;
+		AddHairVisibilityPPLLPass(GraphBuilder, &Scene, &View, MacroGroupDatas, Resolution, InstanceCullingManager, SceneDepthTexture, PPLLNodeCounterTexture, PPLLNodeIndexTexture, PPLLNodeDataBuffer);
+
+		// Generate Tile data
+		check(PPLLNodeIndexTexture)
+		TileData = AddHairStrandsGenerateTilesPass(GraphBuilder, View, PPLLNodeIndexTexture);
+
+		// Linked list sorting pass and compaction into common representation
+		FHairVisibilityControlPointIdCompactionCS::FParameters* PassParameters = GraphBuilder.AllocParameters<FHairVisibilityControlPointIdCompactionCS::FParameters>();
+		PassParameters->PPLLCounter  = PPLLNodeCounterTexture;
+		PassParameters->PPLLNodeIndex= PPLLNodeIndexTexture;
+		PassParameters->PPLLNodeData = GraphBuilder.CreateSRV(PPLLNodeDataBuffer);
+		PassParameters->ViewTransmittanceTexture = nullptr;
+
+		FRDGTextureRef ResolveMaskTexture = nullptr;
+		AddHairVisibilityControlPointIdCompactionPass(
+			true, // bUsePPLL
+			GraphBuilder,
+			View,
+			SceneDepthTexture,
+			MacroGroupDatas,
+			NodeGroupSize,
+			TileData,
+			PassParameters,
+			NodeCounter,
+			VisNodeIndex,
+			VisNodeData,
+			VisNodeCoord,
+			CoverageTexture,
+			IndirectArgsBuffer,
+			OutMaxNodeCount);
+	}
 
 	// Make this tile dispatch
 	AddHairStrandsHitProxyIdPass(GraphBuilder, Scene, View, MacroGroupDatas, TileData, VisNodeIndex, VisNodeData, CoverageTexture, HitProxyTexture, HitProxyDepthTexture);
@@ -5185,6 +5225,7 @@ void RenderHairStrandsVisibilityBuffer(
 					VisibilityData.NodeIndex = CompactNodeIndex;
 					VisibilityData.CoverageTexture = CoverageTexture;
 					VisibilityData.HairOnlyDepthTexture = HairOnlyDepthTexture;
+					VisibilityData.NodeVisData = CompactNodeVis;
 					VisibilityData.NodeData = CompactNodeData;
 					VisibilityData.NodeCoord = CompactNodeCoord;
 					VisibilityData.NodeIndirectArg = IndirectArgsBuffer;
