Support higher quality levels for mobile bloom postprocessing effect (r.BloomQualit)

r.BloomQuality:
1: same as before
2: one additional downsample pass
3, 4, 5: BloomSetup output at half res instead of quarter res and one additional downsample pass

diff --git a/Engine/Config/Android/AndroidScalability.ini b/Engine/Config/Android/AndroidScalability.ini
--- a/Engine/Config/Android/AndroidScalability.ini
+++ b/Engine/Config/Android/AndroidScalability.ini
@@ -105,7 +105,7 @@ r.LensFlareQuality=1
 
 [PostProcessQuality@2]
 r.Tonemapper.Quality=3
-r.BloomQuality=1
+r.BloomQuality=2
 r.DepthOfFieldQuality=1
 r.LightShaftQuality=1
 r.EyeAdaptationQuality=1
@@ -114,7 +114,7 @@ r.LensFlareQuality=2
 
 [PostProcessQuality@3]
 r.Tonemapper.Quality=3
-r.BloomQuality=1
+r.BloomQuality=3
 r.DepthOfFieldQuality=1
 r.LightShaftQuality=1
 r.EyeAdaptationQuality=1
diff --git a/Engine/Source/Runtime/Renderer/Private/PostProcess/PostProcessMobile.cpp b/Engine/Source/Runtime/Renderer/Private/PostProcess/PostProcessMobile.cpp
--- a/Engine/Source/Runtime/Renderer/Private/PostProcess/PostProcessMobile.cpp
+++ b/Engine/Source/Runtime/Renderer/Private/PostProcess/PostProcessMobile.cpp
@@ -14,6 +14,7 @@
 #include "ScenePrivate.h"
 #include "PostProcess/SceneFilterRendering.h"
 #include "PostProcess/PostProcessEyeAdaptation.h"
+#include "PostProcess/PostProcessBloomSetup.h"
 #include "PipelineStateCache.h"
 #include "ClearQuad.h"
 #include "PostProcess/PostProcessing.h"
@@ -275,7 +276,7 @@ IMPLEMENT_GLOBAL_SHADER(FMobileBloomSetupPS, "/Engine/Private/PostProcessMobile.
 
 FMobileBloomSetupOutputs AddMobileBloomSetupPass(FRDGBuilder& GraphBuilder, const FViewInfo& View, const FEyeAdaptationParameters& EyeAdaptationParameters, const FMobileBloomSetupInputs& Inputs)
 {
-	FIntPoint OutputSize = FIntPoint::DivideAndRoundUp(Inputs.SceneColor.ViewRect.Size(), 4);
+	FIntPoint OutputSize = FIntPoint::DivideAndRoundUp(Inputs.SceneColor.ViewRect.Size(), (GetBloomQuality() <= EBloomQuality::Q2) ? 4 : 2);
 
 	const FIntPoint& BufferSize = Inputs.SceneColor.Texture->Desc.Extent;
 
@@ -1115,7 +1116,7 @@ IMPLEMENT_GLOBAL_SHADER(FMobileSunMergeVS, "/Engine/Private/PostProcessMobile.us
 
 FScreenPassTexture AddMobileSunMergePass(FRDGBuilder& GraphBuilder, const FViewInfo& View, const FMobileSunMergeInputs& Inputs)
 {
-	FIntPoint OutputSize = FIntPoint::DivideAndRoundUp(View.ViewRect.Size(), 4);
+	FIntPoint OutputSize = Inputs.BloomSetup_Bloom.Texture->Desc.Extent;
 	
 	const bool bUseTextureArrays = View.bIsMobileMultiViewEnabled;
 	FRDGTextureDesc SunMergeDesc = bUseTextureArrays ? 
diff --git a/Engine/Source/Runtime/Renderer/Private/PostProcess/PostProcessing.cpp b/Engine/Source/Runtime/Renderer/Private/PostProcess/PostProcessing.cpp
--- a/Engine/Source/Runtime/Renderer/Private/PostProcess/PostProcessing.cpp
+++ b/Engine/Source/Runtime/Renderer/Private/PostProcess/PostProcessing.cpp
@@ -2750,7 +2750,7 @@ void AddMobilePostProcessingPasses(FRDGBuilder& GraphBuilder, FScene* Scene, con
 		if (PassSequence.IsEnabled(EPass::Bloom))
 		{
 			PassSequence.AcceptPass(EPass::Bloom);
-			auto AddBloomDownPass = [&GraphBuilder, &View](FScreenPassTexture& BloomDownSource, float BloomDownScale)
+			auto AddBloomDownPass = [&GraphBuilder, &View](const FScreenPassTexture& BloomDownSource, float BloomDownScale)
 			{
 				FMobileBloomDownInputs BloomDownInputs;
 				BloomDownInputs.BloomDownScale = BloomDownScale;
@@ -2759,11 +2759,13 @@ void AddMobilePostProcessingPasses(FRDGBuilder& GraphBuilder, FScene* Scene, con
 				return AddMobileBloomDownPass(GraphBuilder, View, BloomDownInputs);
 			};
 
-			float BloomDownScale = 0.66f * 4.0f;
+			const float BloomDownScale = 0.66f * 4.0f;
+			const int MaxPasses = 6;
+			const EBloomQuality BloomQuality = GetBloomQuality();
+			const int NumDownsamplePasses = BloomQuality == EBloomQuality::Q1 ? 4 : BloomQuality == EBloomQuality::Q2 ? 5 : 6;
+			FScreenPassTexture PostProcessDownsample_Bloom[MaxPasses];
 
-			FScreenPassTexture PostProcessDownsample_Bloom[4];
-
-			for (int32 i = 0; i < 4; ++i)
+			for (int32 i = 0; i < NumDownsamplePasses; ++i)
 			{
 				PostProcessDownsample_Bloom[i] = AddBloomDownPass(i == 0 ? BloomSetupOutputs.Bloom : PostProcessDownsample_Bloom[i - 1], BloomDownScale);
 			}
@@ -2782,24 +2784,30 @@ void AddMobilePostProcessingPasses(FRDGBuilder& GraphBuilder, FScene* Scene, con
 				return AddMobileBloomUpPass(GraphBuilder, View, BloomUpInputs);
 			};
 
-			float BloomUpScale = 0.66f * 2.0f;
+			const float BloomUpScale = 0.66f * 2.0f;
+
+			const FLinearColor Tint[]{ Settings.Bloom1Tint, Settings.Bloom2Tint, Settings.Bloom3Tint, Settings.Bloom4Tint, Settings.Bloom5Tint, Settings.Bloom6Tint };
+
 			// Upsample by 2
 			{
-				FVector4f TintA = FVector4f(Settings.Bloom4Tint.R, Settings.Bloom4Tint.G, Settings.Bloom4Tint.B, 0.0f);
-				FVector4f TintB = FVector4f(Settings.Bloom5Tint.R, Settings.Bloom5Tint.G, Settings.Bloom5Tint.B, 0.0f);
+				const int IndexA = NumDownsamplePasses - 2;
+				const int IndexB = NumDownsamplePasses - 1;
+				FVector4f TintA{ Tint[IndexA], 0.0f };
+				FVector4f TintB{ Tint[IndexB], 0.0f };
 				TintA *= Settings.BloomIntensity;
 				TintB *= Settings.BloomIntensity;
 
-				BloomUpOutputs = AddBloomUpPass(PostProcessDownsample_Bloom[2], PostProcessDownsample_Bloom[3], BloomUpScale, TintA, TintB);
+				BloomUpOutputs = AddBloomUpPass(PostProcessDownsample_Bloom[IndexA], PostProcessDownsample_Bloom[IndexB], BloomUpScale, TintA, TintB);
 			}
 
-			// Upsample by 2
+			for (int Index = NumDownsamplePasses - 3; Index > 0; --Index)
 			{
-				FVector4f TintA = FVector4f(Settings.Bloom3Tint.R, Settings.Bloom3Tint.G, Settings.Bloom3Tint.B, 0.0f);
+				// Upsample by 2
+				FVector4f TintA{ Tint[Index], 0.0f };
 				TintA *= Settings.BloomIntensity;
 				FVector4f TintB = FVector4f(1.0f, 1.0f, 1.0f, 0.0f);
 
-				BloomUpOutputs = AddBloomUpPass(PostProcessDownsample_Bloom[1], BloomUpOutputs, BloomUpScale, TintA, TintB);
+				BloomUpOutputs = AddBloomUpPass(PostProcessDownsample_Bloom[Index], BloomUpOutputs, BloomUpScale, TintA, TintB);
 			}
 
 			// Upsample by 2
@@ -2879,7 +2887,6 @@ void AddMobilePostProcessingPasses(FRDGBuilder& GraphBuilder, FScene* Scene, con
 			SunMergeInputs.bUseSun = bUseSun;
 
 			BloomOutput = AddMobileSunMergePass(GraphBuilder, View, SunMergeInputs);
-
 		}
 
 		// mobile separate translucency 
