[GPULM] GPULM is not using shader binding layout so remove it from the GPULM CHS shaders otherwise it will crash.

#rnx

diff --git a/Engine/Source/Runtime/Renderer/Private/PathTracing.cpp b/Engine/Source/Runtime/Renderer/Private/PathTracing.cpp
--- a/Engine/Source/Runtime/Renderer/Private/PathTracing.cpp
+++ b/Engine/Source/Runtime/Renderer/Private/PathTracing.cpp
@@ -2201,7 +2201,8 @@ public:
 			
 	static const FShaderBindingLayout* GetShaderBindingLayout(const FShaderPermutationParameters& Parameters)
 	{
-		return RayTracing::GetShaderBindingLayout(Parameters.Platform);
+		// GPULM does not use shader binding layout
+		return IsGPULightmass ? nullptr : RayTracing::GetShaderBindingLayout(Parameters.Platform);
 	}
 };
 
