[HWRT] Don't need ref counted resource pointers in BuildRayTracingInstanceBuffer(...)

diff --git a/Engine/Source/Runtime/Renderer/Private/RayTracing/RayTracingInstanceBufferUtil.cpp b/Engine/Source/Runtime/Renderer/Private/RayTracing/RayTracingInstanceBufferUtil.cpp
--- a/Engine/Source/Runtime/Renderer/Private/RayTracing/RayTracingInstanceBufferUtil.cpp
+++ b/Engine/Source/Runtime/Renderer/Private/RayTracing/RayTracingInstanceBufferUtil.cpp
@@ -289,13 +289,13 @@ void BuildRayTracingInstanceBuffer(
 	const FDFVector3& PreViewTranslation,
 	uint32 NumInstances,
 	uint32 InputDescOffset,
-	FUnorderedAccessViewRHIRef InstancesUAV,
-	FShaderResourceViewRHIRef InstanceUploadSRV,
-	FShaderResourceViewRHIRef AccelerationStructureAddressesSRV,
-	FShaderResourceViewRHIRef InstanceTransformSRV,
+	FRHIUnorderedAccessView* InstancesUAV,
+	FRHIShaderResourceView* InstanceUploadSRV,
+	FRHIShaderResourceView* AccelerationStructureAddressesSRV,
+	FRHIShaderResourceView* InstanceTransformSRV,
 	const FRayTracingCullingParameters* CullingParameters,
-	FUnorderedAccessViewRHIRef OutputStatsUAV,
-	FUnorderedAccessViewRHIRef DebugInstanceGPUSceneIndexUAV)
+	FRHIUnorderedAccessView* OutputStatsUAV,
+	FRHIUnorderedAccessView* DebugInstanceGPUSceneIndexUAV)
 {
 	FRayTracingBuildInstanceBufferCS::FParameters PassParams;
 	PassParams.InstancesDescriptors = InstancesUAV;
@@ -353,15 +353,15 @@ void BuildRayTracingInstanceBuffer(
 	FRHICommandList& RHICmdList,
 	const FGPUScene* GPUScene,
 	const FDFVector3& PreViewTranslation,
-	FUnorderedAccessViewRHIRef InstancesUAV,
-	FShaderResourceViewRHIRef InstanceUploadSRV,
-	FShaderResourceViewRHIRef AccelerationStructureAddressesSRV,
-	FShaderResourceViewRHIRef CPUInstanceTransformSRV,
+	FRHIUnorderedAccessView* InstancesUAV,
+	FRHIShaderResourceView* InstanceUploadSRV,
+	FRHIShaderResourceView* AccelerationStructureAddressesSRV,
+	FRHIShaderResourceView* CPUInstanceTransformSRV,
 	uint32 NumNativeGPUSceneInstances,
 	uint32 NumNativeCPUInstances,
 	const FRayTracingCullingParameters* CullingParameters,
-	FUnorderedAccessViewRHIRef OutputStatsUAV,
-	FUnorderedAccessViewRHIRef DebugInstanceGPUSceneIndexUAV)
+	FRHIUnorderedAccessView* OutputStatsUAV,
+	FRHIUnorderedAccessView* DebugInstanceGPUSceneIndexUAV)
 {
 	if (NumNativeGPUSceneInstances > 0)
 	{
diff --git a/Engine/Source/Runtime/Renderer/Public/RayTracingInstanceBufferUtil.h b/Engine/Source/Runtime/Renderer/Public/RayTracingInstanceBufferUtil.h
--- a/Engine/Source/Runtime/Renderer/Public/RayTracingInstanceBufferUtil.h
+++ b/Engine/Source/Runtime/Renderer/Public/RayTracingInstanceBufferUtil.h
@@ -36,20 +36,13 @@ struct FRayTracingInstanceDescriptorInput
 	uint32 bApplyLocalBoundsTransform;
 };
 
-struct FRayTracingGPUInstance
-{
-	FShaderResourceViewRHIRef TransformSRV;
-	uint32 NumInstances;
-	uint32 DescBufferOffset;
-};
-
 struct FRayTracingSceneInitializationData
 {
 	uint32 NumNativeGPUSceneInstances;
 	uint32 NumNativeCPUInstances;
 	uint32 TotalNumSegments;
 
-	// index of each instance geometry in FRayTracingSceneRHIRef ReferencedGeometries
+	// index of each instance geometry in ReferencedGeometries
 	TArray<uint32> InstanceGeometryIndices;
 	// base offset of each instance entries in the instance upload buffer
 	TArray<uint32> BaseUploadBufferOffsets;
@@ -107,14 +100,14 @@ RENDERER_API void BuildRayTracingInstanceBuffer(
 	FRHICommandList& RHICmdList,
 	const FGPUScene* GPUScene,
 	const FDFVector3& PreViewTranslation,
-	FUnorderedAccessViewRHIRef InstancesUAV,
-	FShaderResourceViewRHIRef InstanceUploadSRV,
-	FShaderResourceViewRHIRef AccelerationStructureAddressesSRV,
-	FShaderResourceViewRHIRef CPUInstanceTransformSRV,
+	FRHIUnorderedAccessView* InstancesUAV,
+	FRHIShaderResourceView* InstanceUploadSRV,
+	FRHIShaderResourceView* AccelerationStructureAddressesSRV,
+	FRHIShaderResourceView* CPUInstanceTransformSRV,
 	uint32 NumNativeGPUSceneInstances,
 	uint32 NumNativeCPUInstances,
 	const FRayTracingCullingParameters* CullingParameters,
-	FUnorderedAccessViewRHIRef OutputStatsUAV,
-	FUnorderedAccessViewRHIRef DebugInstanceGPUSceneIndexUAV);
+	FRHIUnorderedAccessView* OutputStatsUAV,
+	FRHIUnorderedAccessView* DebugInstanceGPUSceneIndexUAV);
 
 #endif // RHI_RAYTRACING
\ No newline at end of file
