Mobile: Wait for main visibility tasks to be finished before starting shadow tasks which depend on their results. Failing to do so results in a race condition where the shadow system references these results while they are still being computed resulting in a variety of potential crashes.

diff --git a/Engine/Source/Runtime/Renderer/Private/MobileShadingRenderer.cpp b/Engine/Source/Runtime/Renderer/Private/MobileShadingRenderer.cpp
--- a/Engine/Source/Runtime/Renderer/Private/MobileShadingRenderer.cpp
+++ b/Engine/Source/Runtime/Renderer/Private/MobileShadingRenderer.cpp
@@ -714,6 +714,9 @@ void FMobileSceneRenderer::InitViews(
 		}
 	}
 
+	// Trigger shadow tasks after the main visibility tasks are fully executed. Projection stencil shadows reference the main view dynamic elements result.
+	TaskDatas.VisibilityTaskData->Finish();
+
 	if (bRendererOutputFinalSceneColor)
 	{
 		const bool bDynamicShadows = ViewFamily.EngineShowFlags.DynamicShadows;
@@ -729,8 +732,6 @@ void FMobileSceneRenderer::InitViews(
 		}
 	}
 
-	TaskDatas.VisibilityTaskData->Finish();
-
 	if (bRendererOutputFinalSceneColor)
 	{
 		SetupMobileBasePassAfterShadowInit(BasePassDepthStencilAccess, TaskDatas.VisibilityTaskData->GetViewCommandsPerView(), InstanceCullingManager);
