Fix FTemporalUpsampleDepthPS validation error on console where the output format from the PS is not implicitly converted to PF_R32_FLOAT as requested by the texture description.

#tests Test in a default scene
#rnx

diff --git a/Engine/Source/Runtime/Renderer/Private/PostProcess/PostProcessCompositePrimitivesCommon.cpp b/Engine/Source/Runtime/Renderer/Private/PostProcess/PostProcessCompositePrimitivesCommon.cpp
--- a/Engine/Source/Runtime/Renderer/Private/PostProcess/PostProcessCompositePrimitivesCommon.cpp
+++ b/Engine/Source/Runtime/Renderer/Private/PostProcess/PostProcessCompositePrimitivesCommon.cpp
@@ -82,6 +82,16 @@ public:
 		*/
 		return SupportsGen4TAA(Parameters.Platform);
 	}
+
+	static const EPixelFormat OutputFormat = PF_R32_FLOAT;
+
+	static void ModifyCompilationEnvironment(const FGlobalShaderPermutationParameters& InParameters, FShaderCompilerEnvironment& OutEnvironment)
+	{
+		FGlobalShader::ModifyCompilationEnvironment(InParameters, OutEnvironment);
+
+		/*Explicitly match PS render target output format to the desc of the texture.*/
+		OutEnvironment.SetRenderTargetOutputFormat(0, OutputFormat);
+	}
 };
 
 class FPopulateCompositeDepthPS : public FGlobalShader
@@ -181,7 +191,7 @@ void TemporalUpscaleDepthPass(
 			FScreenPassTexture History;
 			const FRDGTextureDesc Desc = FRDGTextureDesc::Create2D(
 				Extent * DepthUpsampleFactor,
-				PF_R32_FLOAT,
+				FTemporalUpsampleDepthPS::OutputFormat /* PF_R32_FLOAT */,
 				FClearValueBinding::None,
 				TexCreate_ShaderResource | TexCreate_RenderTargetable);
 
