Convert a check() to an if( ensure() ) to see if we can catch a possible out-of-range index causing a live crash (most likely only in the mobile renderer) I can't reproduce internally.

#rnx
#tests local playtesting, local RR testing

diff --git a/Engine/Source/Runtime/Renderer/Private/GPUScene.cpp b/Engine/Source/Runtime/Renderer/Private/GPUScene.cpp
--- a/Engine/Source/Runtime/Renderer/Private/GPUScene.cpp
+++ b/Engine/Source/Runtime/Renderer/Private/GPUScene.cpp
@@ -509,10 +509,11 @@ struct FUploadDataSourceAdapterScenePrimitives
 	{
 		const FPersistentPrimitiveIndex PersistentPrimitiveIndex = PrimitivesToUpdate[ItemIndex];
 		const int32 PrimitiveID = Scene.GetPrimitiveIndex(PersistentPrimitiveIndex);
-		check(Scene.PrimitiveSceneProxies.IsValidIndex(PrimitiveID));
-
-		const FPrimitiveSceneProxy* RESTRICT PrimitiveSceneProxy = Scene.PrimitiveSceneProxies[PrimitiveID];
-		FPrimitiveSceneShaderData::BuildDataFromProxy(PrimitiveSceneProxy, OutData);
+		if (ensure(Scene.PrimitiveSceneProxies.IsValidIndex(PrimitiveID)))
+		{
+			const FPrimitiveSceneProxy* RESTRICT PrimitiveSceneProxy = Scene.PrimitiveSceneProxies[PrimitiveID];
+			FPrimitiveSceneShaderData::BuildDataFromProxy(PrimitiveSceneProxy, OutData);
+		}
 	}
 
 	FORCEINLINE_GPUSCENE void GetInstanceInfo(int32 ItemIndex, FInstanceUploadInfo& InstanceUploadInfo) const
