Initialize the fallback color for runtime virtual texture.
The fallback color is used when the locked pages aren't available if we are not stalling on locked pages (r.VT.SyncProduceLockedTiles=0).
For runtime virtual texture with streaming mips we can take the fallback texture from the one calculated in the streaming mips asset.

diff --git a/Engine/Source/Runtime/Engine/Private/VT/RuntimeVirtualTexture.cpp b/Engine/Source/Runtime/Engine/Private/VT/RuntimeVirtualTexture.cpp
--- a/Engine/Source/Runtime/Engine/Private/VT/RuntimeVirtualTexture.cpp
+++ b/Engine/Source/Runtime/Engine/Private/VT/RuntimeVirtualTexture.cpp
@@ -698,6 +698,11 @@ namespace RuntimeVirtualTexture
 					OutStreamingProducerDesc.BlockWidthInTiles = BlockWidthInTiles >> FirstMipToUse;
 					OutStreamingProducerDesc.BlockHeightInTiles = BlockHeightInTiles >> FirstMipToUse;
 					OutStreamingProducerDesc.MaxLevel = NumLevels - FirstMipToUse;
+					
+					for (uint32 LayerIndex = 0u; LayerIndex < VTData->GetNumLayers(); ++LayerIndex)
+					{
+						OutStreamingProducerDesc.LayerFallbackColor[LayerIndex] = VTData->LayerFallbackColors[LayerIndex];
+					}
 
 					return new FUploadingVirtualTexture(InStreamingTexture->GetFName(), VTData, FirstMipToUse);
 				}
diff --git a/Engine/Source/Runtime/Renderer/Private/VT/RuntimeVirtualTextureSceneProxy.cpp b/Engine/Source/Runtime/Renderer/Private/VT/RuntimeVirtualTextureSceneProxy.cpp
--- a/Engine/Source/Runtime/Renderer/Private/VT/RuntimeVirtualTextureSceneProxy.cpp
+++ b/Engine/Source/Runtime/Renderer/Private/VT/RuntimeVirtualTextureSceneProxy.cpp
@@ -95,6 +95,12 @@ FRuntimeVirtualTextureSceneProxy::FRuntimeVirtualTextureSceneProxy(URuntimeVirtu
 				FVTProducerDescription StreamingProducerDesc;
 				IVirtualTexture* StreamingProducer = RuntimeVirtualTexture::CreateStreamingTextureProducer(StreamingTexture, ProducerDesc, StreamingProducerDesc);
 
+				// Copy the layer fallback colors from the streaming virtual texture.
+				for (uint32 LayerIndex = 0u; LayerIndex < ProducerDesc.NumTextureLayers; ++LayerIndex)
+				{
+					ProducerDesc.LayerFallbackColor[LayerIndex] = StreamingProducerDesc.LayerFallbackColor[LayerIndex];
+				}
+
 				if (InComponent->IsStreamingLowMipsOnly())
 				{
 					// Clamp the the runtime virtual texture producer dimensions to the streaming virtual texture dimensions.
