Reverted condition that broke VerifyMembersChecks.

diff --git a/Engine/Source/Runtime/Engine/Private/SceneView.cpp b/Engine/Source/Runtime/Engine/Private/SceneView.cpp
--- a/Engine/Source/Runtime/Engine/Private/SceneView.cpp
+++ b/Engine/Source/Runtime/Engine/Private/SceneView.cpp
@@ -1079,7 +1079,7 @@ const FViewMatrices& FSceneView::GetShadowViewMatrices() const
 }
 #endif
 
-#if DO_CHECK || USING_CODE_ANALYSIS
+#if DO_CHECK
 bool FSceneView::VerifyMembersChecks() const
 {
 	bool bIsTemporalAccumulation = IsTemporalAccumulationBasedMethod(AntiAliasingMethod);
diff --git a/Engine/Source/Runtime/Engine/Public/SceneView.h b/Engine/Source/Runtime/Engine/Public/SceneView.h
--- a/Engine/Source/Runtime/Engine/Public/SceneView.h
+++ b/Engine/Source/Runtime/Engine/Public/SceneView.h
@@ -1775,7 +1775,7 @@ public:
 	/** Default destructor */
 	virtual ~FSceneView() = default;
 
-#if DO_CHECK || USING_CODE_ANALYSIS
+#if DO_CHECK
 	/** Verifies all the assertions made on members. */
 	ENGINE_API bool VerifyMembersChecks() const;
 #endif
diff --git a/Engine/Source/Runtime/Renderer/Private/SceneRendering.cpp b/Engine/Source/Runtime/Renderer/Private/SceneRendering.cpp
--- a/Engine/Source/Runtime/Renderer/Private/SceneRendering.cpp
+++ b/Engine/Source/Runtime/Renderer/Private/SceneRendering.cpp
@@ -1090,7 +1090,7 @@ FRDGBufferSRVRef FViewInfo::GetRayTracingSceneLayerViewChecked(ERayTracingSceneL
 }
 #endif // RHI_RAYTRACING
 
-#if DO_CHECK || USING_CODE_ANALYSIS
+#if DO_CHECK
 bool FViewInfo::VerifyMembersChecks() const
 {
 	FSceneView::VerifyMembersChecks();
diff --git a/Engine/Source/Runtime/Renderer/Private/SceneRendering.h b/Engine/Source/Runtime/Renderer/Private/SceneRendering.h
--- a/Engine/Source/Runtime/Renderer/Private/SceneRendering.h
+++ b/Engine/Source/Runtime/Renderer/Private/SceneRendering.h
@@ -1604,7 +1604,7 @@ public:
 	*/
 	RENDERER_API virtual ~FViewInfo();
 
-#if DO_CHECK || USING_CODE_ANALYSIS
+#if DO_CHECK
 	/** Verifies all the assertions made on members. */
 	bool VerifyMembersChecks() const;
 #endif
