Volumetric Cloud Component - Improved upsampling quality at a distance

Added the option to use a high precision depth buffer (fp32) for the volumetric cloud component's min/max depth texture, which can resolve flickering artifacts due to low depth precision. The option can be enabled via the r.VolumetricRenderTarget.ImprovedLargeScaleDistanceReconstruction cvar. Additionally, the r.VolumetricRenderTarget.ThresholdToNearFarUpsamplingCm cvar allows for control of when to trigger upsampling of the min/max depth texture. The changes have also been placed behind a shader permutation path, in order to the reduce any unforeseen performance changes.

diff --git a/Engine/Shaders/Private/VolumetricRenderTarget.usf b/Engine/Shaders/Private/VolumetricRenderTarget.usf
--- a/Engine/Shaders/Private/VolumetricRenderTarget.usf
+++ b/Engine/Shaders/Private/VolumetricRenderTarget.usf
@@ -20,6 +20,7 @@
 SamplerState LinearTextureSampler;
 
 Texture2D HalfResDepthTexture;
+Texture2D QuarterResMinMaxDepthTexture;
 
 Texture2D<float4> TracingVolumetricTexture;
 Texture2D<float4> SecondaryTracingVolumetricTexture;
@@ -28,6 +29,7 @@ uint4 TracingVolumetricTextureValidCoordRect;
 float4 TracingVolumetricTextureValidUvRect;
 float MinimumDistanceKmToEnableReprojection;
 float2 TracingVolumetricTextureUVScale;
+float ThresholdToNearFarUpsamplingCm;
 
 struct FDepthData
 {
@@ -254,31 +256,33 @@ void ReconstructVolumetricRenderTargetPS(
 
 
 		#if PERMUTATION_CLOUD_MIN_AND_MAX_DEPTH
+				// Smalll floating point value to prevent division by zero
+				const float DistanceEpsilon = 0.000000000001; 
 
-
+				#if PERMUTATION_LARGE_DISTANCE_RECONSTRUCTION
+					// Fill up min/max quarter resolution depth from original texture.
+					NewDepthData.MinMaxDeviceZ = max(DistanceEpsilon, QuarterResMinMaxDepthTexture.Load(int3(CenterSample, 0)).rg);
+				#endif
 
 				//
 				// NOTE: The best temporal upsampling path reserver for mode 0 when compute is available.
 				//
-
 				const float SceneHalfResDeviceZ = HalfResDepthTexture.Load(int3(SVPos.xy + ViewViewRectMin, 0)).r;
-				const float SceneHalfResDepth = ConvertFromDeviceZ(max(0.000000000001, SceneHalfResDeviceZ));
+				const float SceneHalfResDepth = ConvertFromDeviceZ(max(DistanceEpsilon, SceneHalfResDeviceZ));
 
-				const float ThresholdToUpsamplingCm = 100000.0; // super arbitrary :/
+				const float NearCloudTracingDepth = ConvertFromDeviceZ(max(DistanceEpsilon, NewDepthData.MinMaxDeviceZ.x));
+				const float FarCloudTracingDepth  = ConvertFromDeviceZ(max(DistanceEpsilon, NewDepthData.MinMaxDeviceZ.y));
 
-				const float NearCloudTracingDepth = ConvertFromDeviceZ(max(0.000000000001, NewDepthData.MinMaxDeviceZ.x));
-				const float FarCloudTracingDepth  = ConvertFromDeviceZ(max(0.000000000001, NewDepthData.MinMaxDeviceZ.y));
-				if (abs(NearCloudTracingDepth - FarCloudTracingDepth) > ThresholdToUpsamplingCm)
+				if (abs(NearCloudTracingDepth - FarCloudTracingDepth) > ThresholdToNearFarUpsamplingCm)
 				{
 					// if the min/max depth has a large difference, then always upsample with the min and max traced data
 					float4 NewRGBANear = SafeLoadSecondaryTracingVolumetricTexture(CenterSample);
-					
+
 					const float LerpFactor = saturate((SceneHalfResDepth - NearCloudTracingDepth) / max(0.00001, FarCloudTracingDepth - NearCloudTracingDepth));
 					RGBA = lerp(NewRGBANear, NewRGBA, LerpFactor);
 
-				//	DepthData = NewDepthData;
-					DepthData.SceneDepthFromViewKm = SceneHalfResDepth * CENTIMETER_TO_KILOMETER;	// TODO trasform from scene depth to view depth.
-					DepthData.CloudFrontDepthFromViewKm = DepthData.SceneDepthFromViewKm;			// Keer far cloud tracing front depth. That should be good.
+					DepthData.SceneDepthFromViewKm = SceneHalfResDepth * CENTIMETER_TO_KILOMETER;	// TODO transform from scene depth to view depth.
+					DepthData.CloudFrontDepthFromViewKm = DepthData.SceneDepthFromViewKm;			// Keep far cloud tracing front depth. That should be good.
 						
 			#if DEBUGCOLORS
 					RGBA = float4(1, 1, 0, 0.5);
@@ -290,9 +294,7 @@ void ReconstructVolumetricRenderTargetPS(
 					FDepthData HistoryDepthData = SafeSamplePreviousFrameVolumetricDepthTexture(PrevScreenUVs);
 
 					const float HistorySceneDepthFromViewCm = HistoryDepthData.SceneDepthFromViewKm * KILOMETER_TO_CENTIMETER;
-
-					//if (HistorySceneDepthFromViewCm < (NearCloudTracingDepth - ThresholdToUpsamplingCm))
-					if (HistorySceneDepthFromViewCm < (NewDepthData.SceneDepthFromViewKm * KILOMETER_TO_CENTIMETER - ThresholdToUpsamplingCm))
+					if (HistorySceneDepthFromViewCm < (NewDepthData.SceneDepthFromViewKm * KILOMETER_TO_CENTIMETER - ThresholdToNearFarUpsamplingCm))
 					{
 						// History is closer than the near cloud tracing this frame. This means an object had move and an new disocluded area is discovered.
 						// So we simply use the new data from this frame according to the new depth
@@ -306,7 +308,7 @@ void ReconstructVolumetricRenderTargetPS(
 						RGBA = float4(1, 0, 0, 0.5);
 			#endif
 					}
-					else if (HistorySceneDepthFromViewCm > (NewDepthData.SceneDepthFromViewKm * KILOMETER_TO_CENTIMETER + ThresholdToUpsamplingCm))
+					else if (HistorySceneDepthFromViewCm > (NewDepthData.SceneDepthFromViewKm * KILOMETER_TO_CENTIMETER + ThresholdToNearFarUpsamplingCm))
 					{
 						// An area that just go covered (history is invalid because occluded)
 						RGBA = NewRGBA;
diff --git a/Engine/Source/Runtime/Renderer/Private/DeferredShadingRenderer.cpp b/Engine/Source/Runtime/Renderer/Private/DeferredShadingRenderer.cpp
--- a/Engine/Source/Runtime/Renderer/Private/DeferredShadingRenderer.cpp
+++ b/Engine/Source/Runtime/Renderer/Private/DeferredShadingRenderer.cpp
@@ -2449,7 +2449,7 @@ void FDeferredShadingSceneRenderer::Render(FRDGBuilder& GraphBuilder, const FSce
 			{
 				check(InputTexture != nullptr);	// Must receive a valid texture
 				check(QuarterResolutionDepthMinMaxTexture == nullptr);	// Only generate it once
-				return CreateQuarterResolutionDepthMinAndMax(GraphBuilder, Views, InputTexture);
+				return CreateQuarterResolutionDepthMinMax(GraphBuilder, Views, InputTexture);
 			}
 			return FRDGTextureRef(nullptr);
 		};
@@ -2995,7 +2995,7 @@ void FDeferredShadingSceneRenderer::Render(FRDGBuilder& GraphBuilder, const FSce
 					HalfResolutionDepthCheckerboardMinMaxTexture, QuarterResolutionDepthMinMaxTexture, false, InstanceCullingManager);
 			}
 			// Reconstruct the volumetric cloud render target to be ready to compose it over the scene
-			ReconstructVolumetricRenderTarget(GraphBuilder, Views, SceneTextures.Depth.Resolve, HalfResolutionDepthCheckerboardMinMaxTexture, bAsyncComputeVolumetricCloud);
+			ReconstructVolumetricRenderTarget(GraphBuilder, Views, SceneTextures.Depth.Resolve, HalfResolutionDepthCheckerboardMinMaxTexture, QuarterResolutionDepthMinMaxTexture, bAsyncComputeVolumetricCloud);
 		}
 
 		TArray<FScreenPassTexture, TInlineAllocator<4>> TSRFlickeringInputTextures;
@@ -3394,7 +3394,7 @@ void FDeferredShadingSceneRenderer::Render(FRDGBuilder& GraphBuilder, const FSce
 		if (bShouldVisualizeVolumetricCloud && !bHasRayTracedOverlay)
 		{
 			RenderVolumetricCloud(GraphBuilder, SceneTextures, false, true, HalfResolutionDepthCheckerboardMinMaxTexture, QuarterResolutionDepthMinMaxTexture, false, InstanceCullingManager);
-			ReconstructVolumetricRenderTarget(GraphBuilder, Views, SceneTextures.Depth.Resolve, HalfResolutionDepthCheckerboardMinMaxTexture, false);
+			ReconstructVolumetricRenderTarget(GraphBuilder, Views, SceneTextures.Depth.Resolve, HalfResolutionDepthCheckerboardMinMaxTexture, QuarterResolutionDepthMinMaxTexture, false);
 			ComposeVolumetricRenderTargetOverSceneForVisualization(GraphBuilder, Views, SceneTextures.Color.Target, SceneTextures);
 			RenderVolumetricCloud(GraphBuilder, SceneTextures, true, false, HalfResolutionDepthCheckerboardMinMaxTexture, QuarterResolutionDepthMinMaxTexture, false, InstanceCullingManager);
 		}
diff --git a/Engine/Source/Runtime/Renderer/Private/SceneRendering.cpp b/Engine/Source/Runtime/Renderer/Private/SceneRendering.cpp
--- a/Engine/Source/Runtime/Renderer/Private/SceneRendering.cpp
+++ b/Engine/Source/Runtime/Renderer/Private/SceneRendering.cpp
@@ -94,6 +94,7 @@
 #include "SkyAtmosphereRendering.h"
 #include "VolumetricCloudRendering.h"
 #include "VolumetricFog.h"
+#include "VolumetricRenderTarget.h"
 #include "PrimitiveSceneShaderData.h"
 #include "Engine/SpecularProfile.h"
 #include "Engine/VolumeTexture.h"
@@ -6080,23 +6081,26 @@ FRDGTextureRef CreateHalfResolutionDepthCheckerboardMinMax(FRDGBuilder& GraphBui
 	return SmallDepthTexture;
 }
 
-FRDGTextureRef CreateQuarterResolutionDepthMinAndMax(FRDGBuilder& GraphBuilder, TArrayView<const FViewInfo> Views, FRDGTextureRef InputDepthTexture)
+FRDGTextureRef CreateQuarterResolutionDepthMinMax(FRDGBuilder& GraphBuilder, TArrayView<const FViewInfo> Views, FRDGTextureRef InputDepthTexture)
 {
-	const FIntPoint SmallDepthExtent = GetDownscaledExtent(InputDepthTexture->Desc.Extent, 2);
+	const uint32 DownscaleFactor = 2;
+	const FIntPoint SmallDepthExtent = GetDownscaledExtent(InputDepthTexture->Desc.Extent, DownscaleFactor);
+	const EPixelFormat TexFormat = IsImprovedLargeScaleDistanceReconstructionEnabled() ? PF_G32R32F : PF_G16R16F;
 	const ETextureCreateFlags NoFastClearFlags = (CVarHalfResDepthNoFastClear.GetValueOnAnyThread() != 0) ? TexCreate_NoFastClear : TexCreate_None;
-	const FRDGTextureDesc SmallTextureDesc = FRDGTextureDesc::Create2D(SmallDepthExtent, PF_G16R16F, FClearValueBinding::None, TexCreate_RenderTargetable | TexCreate_ShaderResource | NoFastClearFlags);
-	FRDGTextureRef SmallTexture = GraphBuilder.CreateTexture(SmallTextureDesc, TEXT("HalfResolutionDepthMinAndMax"));
+
+	const FRDGTextureDesc SmallDepthDesc = FRDGTextureDesc::Create2D(SmallDepthExtent, TexFormat, FClearValueBinding::None, TexCreate_RenderTargetable | TexCreate_ShaderResource | NoFastClearFlags);
+	FRDGTextureRef SmallDepthTexture = GraphBuilder.CreateTexture(SmallDepthDesc, TEXT("HalfResolutionDepthMinAndMax"));
 
 	for (const FViewInfo& View : Views)
 	{
 		RDG_GPU_MASK_SCOPE(GraphBuilder, View.GPUMask);
 
-		const FScreenPassTexture InputDepth(InputDepthTexture, GetDownscaledRect(View.ViewRect, 2));
-		const FScreenPassRenderTarget SmallTextureRT(SmallTexture, GetDownscaledRect(InputDepth.ViewRect, 2), View.DecayLoadAction(ERenderTargetLoadAction::ENoAction));
-		AddDownsampleDepthPass(GraphBuilder, View, InputDepth, SmallTextureRT, EDownsampleDepthFilter::MinAndMaxDepth);
+		const FScreenPassTexture InputDepth(InputDepthTexture, GetDownscaledRect(View.ViewRect, DownscaleFactor));
+		const FScreenPassRenderTarget QuarterResTextureRT(SmallDepthTexture, GetDownscaledRect(InputDepth.ViewRect, DownscaleFactor), View.DecayLoadAction(ERenderTargetLoadAction::ENoAction));
+		AddDownsampleDepthPass(GraphBuilder, View, InputDepth, QuarterResTextureRT, EDownsampleDepthFilter::MinAndMaxDepth);
 	}
 
-	return SmallTexture;
+	return SmallDepthTexture;
 }
 
 bool IsPrimitiveAlphaHoldoutEnabled(const FViewInfo& View)
diff --git a/Engine/Source/Runtime/Renderer/Private/SceneRendering.h b/Engine/Source/Runtime/Renderer/Private/SceneRendering.h
--- a/Engine/Source/Runtime/Renderer/Private/SceneRendering.h
+++ b/Engine/Source/Runtime/Renderer/Private/SceneRendering.h
@@ -3017,7 +3017,7 @@ void VirtualTextureFeedbackEnd(FRDGBuilder& GraphBuilder);
 FRDGTextureRef CreateHalfResolutionDepthCheckerboardMinMax(FRDGBuilder& GraphBuilder, TArrayView<const FViewInfo> Views, FRDGTextureRef SceneDepth);
 
 /** Creates a half resolution depth buffer storing the min and max depth for each 2x2 pixel quad of a half resolution buffer. */
-FRDGTextureRef CreateQuarterResolutionDepthMinAndMax(FRDGBuilder& GraphBuilder, TArrayView<const FViewInfo> Views, FRDGTextureRef DepthTexture);
+FRDGTextureRef CreateQuarterResolutionDepthMinMax(FRDGBuilder& GraphBuilder, TArrayView<const FViewInfo> Views, FRDGTextureRef DepthTexture);
 
 inline const FSceneTexturesConfig& FViewInfo::GetSceneTexturesConfig() const
 {
diff --git a/Engine/Source/Runtime/Renderer/Private/VolumetricRenderTarget.cpp b/Engine/Source/Runtime/Renderer/Private/VolumetricRenderTarget.cpp
--- a/Engine/Source/Runtime/Renderer/Private/VolumetricRenderTarget.cpp
+++ b/Engine/Source/Runtime/Renderer/Private/VolumetricRenderTarget.cpp
@@ -28,12 +28,12 @@ static TAutoConsoleVariable<float> CVarVolumetricRenderTargetUvNoiseSampleAccept
 
 static TAutoConsoleVariable<int32> CVarVolumetricRenderTargetMode(
 	TEXT("r.VolumetricRenderTarget.Mode"), 0,
-	TEXT("[0] trace quarter resolution + reconstruct at half resolution + upsample [1] trace half res + upsample [2] trace at quarter resolution + reconstruct full resolution (cannot intersect with opaque meshes and forces UpsamplingMode=2 [3] Cinematic mode with tracing done at full reoslution in render target so that clouds can also be applied on translucent.)"),
+	TEXT("Sets the algorithm for rendering volumetric clouds. [0] trace quarter resolution + reconstruct at half resolution + upsample [1] trace half res + upsample [2] trace at quarter resolution + reconstruct full resolution (cannot intersect with opaque meshes and forces UpsamplingMode=2 [3] Cinematic mode with tracing done at full reoslution in render target so that clouds can also be applied on translucent.)"),
 	ECVF_RenderThreadSafe | ECVF_Scalability);
 
 static TAutoConsoleVariable<int32> CVarVolumetricRenderTargetUpsamplingMode(
 	TEXT("r.VolumetricRenderTarget.UpsamplingMode"), 4,
-	TEXT("Used in compositing volumetric RT over the scene. [0] bilinear [1] bilinear + jitter [2] nearest + depth test [3] bilinear + jitter + keep closest [4] bilaterial upsampling"),
+	TEXT("Used in compositing the volumetric render target over the scene. [0] bilinear [1] bilinear + jitter [2] nearest + depth test [3] bilinear + jitter + keep closest [4] bilateral upsampling"),
 	ECVF_RenderThreadSafe | ECVF_Scalability);
 
 static TAutoConsoleVariable<float> CVarVolumetricRenderTargetScale(
@@ -61,6 +61,16 @@ static TAutoConsoleVariable<float> CVarVolumetricRenderTargetMinimumDistanceKmTo
 	TEXT("This is the distance in kilometer at which the `cloud surface` must be before we enable reprojection of the previous frame data. One could start with a value of 4km. This helps hide reprojection issues due to imperfect approximation of cloud depth as a single front surface, especially visible when flying through the cloud layer. It is not perfect but will help in lots of cases. The problem when using this method: clouds will look noisier when closer to that distance."),
 	ECVF_RenderThreadSafe | ECVF_Scalability);
 
+static TAutoConsoleVariable<float> CVarVolumetricRenderTargetThresholdToNearFarUpsamplingCm(
+	TEXT("r.VolumetricRenderTarget.ThresholdToNearFarUpsamplingCm"), 100000.0,
+	TEXT("The distance threshold in centimeters that will trigger upsampling. Enabled when the near to far tracing distance is larger than the threshold. Use in conjunction with r.VolumetricRenderTarget.ImprovedLargeScaleDistanceReconstruction for higher quality results."),
+	ECVF_RenderThreadSafe | ECVF_Scalability);
+
+static TAutoConsoleVariable<int32> CVarImprovedLargeScaleDistanceReconstruction(
+	TEXT("r.VolumetricRenderTarget.ImprovedLargeScaleDistanceReconstruction"), 0,
+	TEXT("Enable extra floating point precision (32-bit) for the depth buffer that is used for reconstruction. Disabled by default."),
+	ECVF_RenderThreadSafe | ECVF_Scalability);
+
 static float GetUvNoiseSampleAcceptanceWeight()
 {
 	return FMath::Max(0.0f, CVarVolumetricRenderTargetUvNoiseSampleAcceptanceWeight.GetValueOnRenderThread());
@@ -91,6 +101,11 @@ bool IsVolumetricRenderTargetAsyncCompute()
 	return GSupportsEfficientAsyncCompute && CVarVolumetricRenderTargetPreferAsyncCompute.GetValueOnRenderThread() > 0 && !bCloudComputePathDisabled;
 }
 
+bool IsImprovedLargeScaleDistanceReconstructionEnabled()
+{
+	return CVarImprovedLargeScaleDistanceReconstruction.GetValueOnRenderThread() > 0;
+}
+
 static bool ShouldViewComposeVolumetricRenderTarget(const FViewInfo& ViewInfo)
 {
 	return ShouldViewRenderVolumetricCloudRenderTarget(ViewInfo);
@@ -356,10 +371,8 @@ FRDGTextureRef FVolumetricRenderTargetViewStateData::GetOrCreateVolumetricTracin
 
 	if (!VolumetricTracingRTDepth.IsValid())
 	{
-		EPixelFormat DepthDataFormat = Mode == 0 ? PF_FloatRGBA : PF_G16R16F; // Mode 0 supports MinAndMax depth tracing when the compute path is used so always allocate a 4-components texture in this case.
-
 		FPooledRenderTargetDesc Desc = FPooledRenderTargetDesc::Create2DDesc(
-			VolumetricTracingRTResolution, DepthDataFormat, FClearValueBinding(FLinearColor(63000.0f, 63000.0f, 63000.0f, 63000.0f)),
+			VolumetricTracingRTResolution, PF_G16R16F, FClearValueBinding(FLinearColor(63000.0f, 63000.0f, 63000.0f, 63000.0f)),
 			TexCreate_None, TexCreate_ShaderResource | TexCreate_RenderTargetable | TexCreate_UAV, false);
 		GRenderTargetPool.FindFreeElement(GraphBuilder.RHICmdList, Desc, VolumetricTracingRTDepth, TEXT("VolumetricRenderTarget.TracingDepth"));
 	}
@@ -625,7 +638,8 @@ class FReconstructVolumetricRenderTargetPS : public FGlobalShader
 	class FHistoryAvailable : SHADER_PERMUTATION_BOOL("PERMUTATION_HISTORY_AVAILABLE");
 	class FReprojectionBoxConstraint : SHADER_PERMUTATION_BOOL("PERMUTATION_REPROJECTION_BOX_CONSTRAINT");
 	class FCloudMinAndMaxDepth : SHADER_PERMUTATION_BOOL("PERMUTATION_CLOUD_MIN_AND_MAX_DEPTH");
-	using FPermutationDomain = TShaderPermutationDomain<FHistoryAvailable, FReprojectionBoxConstraint, FCloudMinAndMaxDepth>;
+	class FImprovedLargeDistanceReconstruction : SHADER_PERMUTATION_BOOL("PERMUTATION_LARGE_DISTANCE_RECONSTRUCTION");
+	using FPermutationDomain = TShaderPermutationDomain<FHistoryAvailable, FReprojectionBoxConstraint, FCloudMinAndMaxDepth, FImprovedLargeDistanceReconstruction>;
 
 	BEGIN_SHADER_PARAMETER_STRUCT(FParameters, )
 		SHADER_PARAMETER_STRUCT_REF(FViewUniformShaderParameters, ViewUniformBuffer)
@@ -635,6 +649,7 @@ class FReconstructVolumetricRenderTargetPS : public FGlobalShader
 		SHADER_PARAMETER_RDG_TEXTURE(Texture2D, PreviousFrameVolumetricTexture)
 		SHADER_PARAMETER_RDG_TEXTURE(Texture2D, PreviousFrameVolumetricDepthTexture)
 		SHADER_PARAMETER_RDG_TEXTURE(Texture2D, HalfResDepthTexture)
+		SHADER_PARAMETER_RDG_TEXTURE(Texture2D, QuarterResMinMaxDepthTexture)
 		SHADER_PARAMETER_SAMPLER(SamplerState, LinearTextureSampler)
 		RENDER_TARGET_BINDING_SLOTS()
 		SHADER_PARAMETER(FVector4f, DstVolumetricTextureSizeAndInvSize)
@@ -651,6 +666,7 @@ class FReconstructVolumetricRenderTargetPS : public FGlobalShader
 		SHADER_PARAMETER(float, MinimumDistanceKmToEnableReprojection)
 		SHADER_PARAMETER(float, HistoryPreExposureCorrection)
 		SHADER_PARAMETER(FVector2f, PreviousFrameVolumetricTextureUVScale)
+		SHADER_PARAMETER(float, ThresholdToNearFarUpsamplingCm)
 	END_SHADER_PARAMETER_STRUCT()
 
 	static FPermutationDomain RemapPermutation(FPermutationDomain PermutationVector)
@@ -679,6 +695,7 @@ void ReconstructVolumetricRenderTarget(
 	TArrayView<FViewInfo> Views,
 	FRDGTextureRef SceneDepthTexture,
 	FRDGTextureRef HalfResolutionDepthCheckerboardMinMaxTexture,
+	FRDGTextureRef QuarterResolutionDepthMinMaxTexture,
 	bool bWaitFinishFence)
 {
 	if (!AnyViewRequiresProcessing(Views))
@@ -724,6 +741,7 @@ void ReconstructVolumetricRenderTarget(
 		PermutationVector.Set<FReconstructVolumetricRenderTargetPS::FHistoryAvailable>(VolumetricCloudRT.GetHistoryValid());
 		PermutationVector.Set<FReconstructVolumetricRenderTargetPS::FReprojectionBoxConstraint>(CVarVolumetricRenderTargetReprojectionBoxConstraint.GetValueOnAnyThread() > 0);
 		PermutationVector.Set<FReconstructVolumetricRenderTargetPS::FCloudMinAndMaxDepth>(ShouldVolumetricCloudTraceWithMinMaxDepth(ViewInfo));
+		PermutationVector.Set<FReconstructVolumetricRenderTargetPS::FImprovedLargeDistanceReconstruction>(IsImprovedLargeScaleDistanceReconstructionEnabled());
 		TShaderMapRef<FReconstructVolumetricRenderTargetPS> PixelShader(ViewInfo.ShaderMap, PermutationVector);
 
 		FReconstructVolumetricRenderTargetPS::FParameters* PassParameters = GraphBuilder.AllocParameters<FReconstructVolumetricRenderTargetPS::FParameters>();
@@ -741,8 +759,10 @@ void ReconstructVolumetricRenderTarget(
 		PassParameters->DownSampleFactor = TracingVolumetricCloudRTDownSample;
 		PassParameters->VolumetricRenderTargetMode = VolumetricCloudRT.GetMode();
 		PassParameters->HalfResDepthTexture = (VolumetricCloudRT.GetMode() == 0 || VolumetricCloudRT.GetMode() == 3) ? HalfResolutionDepthCheckerboardMinMaxTexture : SceneDepthTexture;
+		PassParameters->QuarterResMinMaxDepthTexture = QuarterResolutionDepthMinMaxTexture != nullptr && VolumetricCloudRT.GetMode() == 0 ? QuarterResolutionDepthMinMaxTexture : SystemTextures.Black;
 		PassParameters->MinimumDistanceKmToEnableReprojection = FMath::Max(0.0f, CVarVolumetricRenderTargetMinimumDistanceKmToEnableReprojection.GetValueOnRenderThread());
 		PassParameters->HistoryPreExposureCorrection = ViewInfo.PreExposure / VolumetricCloudRT.GetPrevViewExposure();
+		PassParameters->ThresholdToNearFarUpsamplingCm = FMath::Max(0.0f, CVarVolumetricRenderTargetThresholdToNearFarUpsamplingCm.GetValueOnRenderThread());
 
 		const bool bVisualizeConservativeDensity = ShouldViewVisualizeVolumetricCloudConservativeDensity(ViewInfo, ViewInfo.Family->EngineShowFlags);
 		PassParameters->HalfResDepthTexture = bVisualizeConservativeDensity ?
diff --git a/Engine/Source/Runtime/Renderer/Private/VolumetricRenderTarget.h b/Engine/Source/Runtime/Renderer/Private/VolumetricRenderTarget.h
--- a/Engine/Source/Runtime/Renderer/Private/VolumetricRenderTarget.h
+++ b/Engine/Source/Runtime/Renderer/Private/VolumetricRenderTarget.h
@@ -25,6 +25,7 @@ void ReconstructVolumetricRenderTarget(
 	TArrayView<FViewInfo> Views,
 	FRDGTextureRef SceneDepthTexture,
 	FRDGTextureRef HalfResolutionDepthCheckerboardMinMaxTexture,
+	FRDGTextureRef QuarterResolutionDepthMinMaxTexture,
 	bool bWaitFinishFence);
 
 void ComposeVolumetricRenderTargetOverScene(
@@ -50,4 +51,6 @@ void ComposeVolumetricRenderTargetOverSceneForVisualization(
 
 bool ShouldViewRenderVolumetricCloudRenderTarget(const FViewInfo& ViewInfo);
 bool IsVolumetricRenderTargetEnabled();
-bool IsVolumetricRenderTargetAsyncCompute();
\ No newline at end of file
+bool IsVolumetricRenderTargetAsyncCompute();
+
+bool IsImprovedLargeScaleDistanceReconstructionEnabled();
\ No newline at end of file
