[HWRT] Rename some structs related to RayTracingBuildInstanceBufferCS
- FRayTracingInstanceDescriptor -> FPlatformRayTracingInstanceDescriptor
- FRayTracingInstanceDescriptorInput -> FRayTracingInstanceDescriptor

diff --git a/Engine/Plugins/Experimental/GPULightmass/Source/GPULightmass/Private/LightmapRenderer.cpp b/Engine/Plugins/Experimental/GPULightmass/Source/GPULightmass/Private/LightmapRenderer.cpp
--- a/Engine/Plugins/Experimental/GPULightmass/Source/GPULightmass/Private/LightmapRenderer.cpp
+++ b/Engine/Plugins/Experimental/GPULightmass/Source/GPULightmass/Private/LightmapRenderer.cpp
@@ -1190,12 +1190,12 @@ bool FSceneRenderState::SetupRayTracingScene(FRDGBuilder& GraphBuilder, FSceneUn
 				RayTracingSceneInitializationData.ReferencedGeometries.Num() * sizeof(FRayTracingAccelerationStructureAddress),
 				BUF_Volatile | BUF_MultiGPUAllocate);
 
-			const uint32 InstanceUploadBufferSize = SceneInitializer.MaxNumInstances * sizeof(FRayTracingInstanceDescriptorInput);
+			const uint32 InstanceUploadBufferSize = SceneInitializer.MaxNumInstances * sizeof(FRayTracingInstanceDescriptor);
 			FBufferRHIRef InstanceUploadBuffer;
 			FShaderResourceViewRHIRef InstanceUploadSRV;
 			{
 				FRHIResourceCreateInfo CreateInfo(TEXT("LightmassRayTracingInstanceUploadBuffer"));
-				InstanceUploadBuffer = RHICmdList.CreateStructuredBuffer(sizeof(FRayTracingInstanceDescriptorInput), InstanceUploadBufferSize, BUF_ShaderResource | BUF_Volatile, CreateInfo);
+				InstanceUploadBuffer = RHICmdList.CreateStructuredBuffer(sizeof(FRayTracingInstanceDescriptor), InstanceUploadBufferSize, BUF_ShaderResource | BUF_Volatile, CreateInfo);
 				InstanceUploadSRV = RHICmdList.CreateShaderResourceView(InstanceUploadBuffer, FRHIViewDesc::CreateBufferSRV().SetTypeFromBuffer(InstanceUploadBuffer));
 			}
 
@@ -1209,7 +1209,7 @@ bool FSceneRenderState::SetupRayTracingScene(FRDGBuilder& GraphBuilder, FSceneUn
 			}
 
 			{
-				FRayTracingInstanceDescriptorInput* InstanceUploadData = (FRayTracingInstanceDescriptorInput*)RHICmdList.LockBuffer(InstanceUploadBuffer, 0, InstanceUploadBufferSize, RLM_WriteOnly);
+				FRayTracingInstanceDescriptor* InstanceUploadData = (FRayTracingInstanceDescriptor*)RHICmdList.LockBuffer(InstanceUploadBuffer, 0, InstanceUploadBufferSize, RLM_WriteOnly);
 				FVector4f* TransformUploadData = (FVector4f*)RHICmdList.LockBuffer(TransformUploadBuffer, 0, TransformUploadBufferSize, RLM_WriteOnly);
 				FillRayTracingInstanceUploadBuffer(
 					View.ViewMatrices.GetPreViewTranslation(),
diff --git a/Engine/Shaders/Private/RayTracing/RayTracingCommon.ush b/Engine/Shaders/Private/RayTracing/RayTracingCommon.ush
--- a/Engine/Shaders/Private/RayTracing/RayTracingCommon.ush
+++ b/Engine/Shaders/Private/RayTracing/RayTracingCommon.ush
@@ -195,7 +195,7 @@ void name##_INTERNAL(uint3 DispatchThreadIndex)
 #define OVERRIDE_PLATFORM_RAYTRACING_INSTANCE_DESCRITOR
 
 // Must match D3D12_RAYTRACING_INSTANCE_DESC / VkAccelerationStructureInstanceKHR
-struct FRayTracingInstanceDescriptor
+struct FPlatformRayTracingInstanceDescriptor
 {
 	float3x4 Transform;
 	uint InstanceIdAndMask;
@@ -223,34 +223,34 @@ struct FRayTracingInstanceDescriptor
 	}
 };
 
-void SetRayTracingInstanceTransform(inout FRayTracingInstanceDescriptor Desc, float3x4 LocalToWorld)
+void SetRayTracingInstanceTransform(inout FPlatformRayTracingInstanceDescriptor Desc, float3x4 LocalToWorld)
 {
 	Desc.Transform = LocalToWorld;
 }
 
-uint GetRayTracingInstanceDescriptorInstanceId(FRayTracingInstanceDescriptor Desc)
+uint GetRayTracingInstanceDescriptorInstanceId(FPlatformRayTracingInstanceDescriptor Desc)
 {
 	return Desc.GetInstanceId();
 }
 
-uint GetRayTracingInstanceDescriptorMask(FRayTracingInstanceDescriptor Desc)
+uint GetRayTracingInstanceDescriptorMask(FPlatformRayTracingInstanceDescriptor Desc)
 {
 	return Desc.GetMask();
 }
 
-uint GetRayTracingInstanceDescriptorInstanceContributionToHitGroupIndex(FRayTracingInstanceDescriptor Desc)
+uint GetRayTracingInstanceDescriptorInstanceContributionToHitGroupIndex(FPlatformRayTracingInstanceDescriptor Desc)
 {
 	return Desc.GetInstanceContributionToHitGroupIndex();
 }
 
-uint GetRayTracingInstanceDescriptorFlags(FRayTracingInstanceDescriptor Desc)
+uint GetRayTracingInstanceDescriptorFlags(FPlatformRayTracingInstanceDescriptor Desc)
 {
 	return Desc.GetFlags();
 }
 
-FRayTracingInstanceDescriptor BuildPlatformRayTracingInstanceDesc(uint InstanceMask, uint InstanceId, uint InstanceFlags, uint InstanceContributionToHitGroupIndex, float3x4 LocalToWorld, uint2 AccelerationStructureAddress)
+FPlatformRayTracingInstanceDescriptor BuildPlatformRayTracingInstanceDesc(uint InstanceMask, uint InstanceId, uint InstanceFlags, uint InstanceContributionToHitGroupIndex, float3x4 LocalToWorld, uint2 AccelerationStructureAddress)
 {
-	FRayTracingInstanceDescriptor Desc;
+	FPlatformRayTracingInstanceDescriptor Desc;
 	Desc.InstanceIdAndMask = (InstanceMask << 24) | (InstanceId & 0x00FFFFFF);
 	Desc.InstanceContributionToHitGroupIndexAndFlags = (InstanceFlags << 24) | (InstanceContributionToHitGroupIndex & 0x00FFFFFF);
 	Desc.AccelerationStructure = AccelerationStructureAddress;
diff --git a/Engine/Shaders/Private/RayTracing/RayTracingCommonMetal.ush b/Engine/Shaders/Private/RayTracing/RayTracingCommonMetal.ush
--- a/Engine/Shaders/Private/RayTracing/RayTracingCommonMetal.ush
+++ b/Engine/Shaders/Private/RayTracing/RayTracingCommonMetal.ush
@@ -16,7 +16,7 @@
 
 // NOTE: Must be POD! (to avoid screwing up alignment for BLAS build).
 // Mirrors MTLAccelerationStructureUserIDInstanceDescriptor:
-struct FRayTracingInstanceDescriptor
+struct FPlatformRayTracingInstanceDescriptor
 {
     float Transform[4][3]; // MTLPackedFloat4x3 transformationMatrix;
     uint Options;  //  MTLAccelerationStructureInstanceOptions options;
@@ -26,36 +26,36 @@ struct FRayTracingInstanceDescriptor
     uint InstanceID; // uint32_t userID;
 };
 
-void SetRayTracingInstanceTransform(inout FRayTracingInstanceDescriptor Desc, float3x4 LocalToWorld)
+void SetRayTracingInstanceTransform(inout FPlatformRayTracingInstanceDescriptor Desc, float3x4 LocalToWorld)
 {
     for (uint j = 0; j < 4; j++)
         for (uint i = 0; i < 3; i++)
             Desc.Transform[j][i] = LocalToWorld[i][j];
 }
 
-uint GetRayTracingInstanceDescriptorInstanceId(FRayTracingInstanceDescriptor Desc)
+uint GetRayTracingInstanceDescriptorInstanceId(FPlatformRayTracingInstanceDescriptor Desc)
 {
     return Desc.InstanceID;
 }
 
-uint GetRayTracingInstanceDescriptorMask(FRayTracingInstanceDescriptor Desc)
+uint GetRayTracingInstanceDescriptorMask(FPlatformRayTracingInstanceDescriptor Desc)
 {
     return Desc.InstanceMask;
 }
 
-uint GetRayTracingInstanceDescriptorInstanceContributionToHitGroupIndex(FRayTracingInstanceDescriptor Desc)
+uint GetRayTracingInstanceDescriptorInstanceContributionToHitGroupIndex(FPlatformRayTracingInstanceDescriptor Desc)
 {
     return Desc.InstanceContributionToHitGroupIndex;
 }
 
-uint GetRayTracingInstanceDescriptorFlags(FRayTracingInstanceDescriptor Desc)
+uint GetRayTracingInstanceDescriptorFlags(FPlatformRayTracingInstanceDescriptor Desc)
 {
     return Desc.Options;
 }
 
-FRayTracingInstanceDescriptor BuildPlatformRayTracingInstanceDesc(uint InstanceMask, uint InstanceId, uint InstanceFlags, uint InstanceContributionToHitGroupIndex, float3x4 LocalToWorld, uint2 AccelerationStructureAddress)
+FPlatformRayTracingInstanceDescriptor BuildPlatformRayTracingInstanceDesc(uint InstanceMask, uint InstanceId, uint InstanceFlags, uint InstanceContributionToHitGroupIndex, float3x4 LocalToWorld, uint2 AccelerationStructureAddress)
 {
-    FRayTracingInstanceDescriptor Desc;
+    FPlatformRayTracingInstanceDescriptor Desc;
     Desc.Options = InstanceFlags;
     Desc.InstanceMask = InstanceMask;
     Desc.InstanceContributionToHitGroupIndex = InstanceContributionToHitGroupIndex;
diff --git a/Engine/Shaders/Private/RayTracing/RayTracingDebugPicking.usf b/Engine/Shaders/Private/RayTracing/RayTracingDebugPicking.usf
--- a/Engine/Shaders/Private/RayTracing/RayTracingDebugPicking.usf
+++ b/Engine/Shaders/Private/RayTracing/RayTracingDebugPicking.usf
@@ -13,7 +13,7 @@ StructuredBuffer<FRayTracingInstanceDebugData> InstancesDebugData;
 StructuredBuffer<FRayTracingPickingFeedback> PickingBuffer;
 
 RWStructuredBuffer<FRayTracingPickingFeedback> PickingOutput;
-StructuredBuffer<FRayTracingInstanceDescriptor> InstanceBuffer;
+StructuredBuffer<FPlatformRayTracingInstanceDescriptor> InstanceBuffer;
 
 RAY_TRACING_ENTRY_RAYGEN(RayTracingDebugPickingRGS)
 {
@@ -42,7 +42,7 @@ RAY_TRACING_ENTRY_RAYGEN(RayTracingDebugPickingRGS)
 	if (Payload.IsHit())
 	{
 		FRayTracingInstanceDebugData InstanceDebugData = InstancesDebugData[Payload.InstanceIndex];
-		FRayTracingInstanceDescriptor InstanceDescriptor = InstanceBuffer[Payload.InstanceIndex];
+		FPlatformRayTracingInstanceDescriptor InstanceDescriptor = InstanceBuffer[Payload.InstanceIndex];
 
 		PickingFeedback.GeometryAddress = InstanceDebugData.GeometryAddress;
 
diff --git a/Engine/Shaders/Private/RayTracing/RayTracingInstanceBufferUtil.usf b/Engine/Shaders/Private/RayTracing/RayTracingInstanceBufferUtil.usf
--- a/Engine/Shaders/Private/RayTracing/RayTracingInstanceBufferUtil.usf
+++ b/Engine/Shaders/Private/RayTracing/RayTracingInstanceBufferUtil.usf
@@ -13,7 +13,7 @@
 #define THREADGROUP_SIZE 1
 #endif
 
-struct FRayTracingInstanceDescriptorInput
+struct FRayTracingInstanceDescriptor
 {
 	uint GPUSceneInstanceOrTransformIndex;
 	uint OutputDescriptorIndex;
@@ -25,7 +25,7 @@ struct FRayTracingInstanceDescriptorInput
 };
 
 uint NumInstances;
-uint InputDescOffset;
+uint BaseInstanceDescriptorIndex;
 
 
 float3 FarFieldReferencePos;
@@ -33,9 +33,9 @@ float3 FarFieldReferencePos;
 float3 PreViewTranslationHigh;
 float3 PreViewTranslationLow;
 
-RWStructuredBuffer<FRayTracingInstanceDescriptor> InstancesDescriptors;
+RWStructuredBuffer<FPlatformRayTracingInstanceDescriptor> OutPlatformInstanceDescriptors;
 
-StructuredBuffer<FRayTracingInstanceDescriptorInput> InputInstanceDescriptors;
+StructuredBuffer<FRayTracingInstanceDescriptor> InstanceDescriptors;
 ByteAddressBuffer AccelerationStructureAddresses;
 // Transforms are float3x4 row-major matrices stored as 3 float4
 // because of FXC matrix packing issues when using StructuredBuffer<float3x4>
@@ -112,7 +112,7 @@ void RayTracingBuildInstanceBufferCS(uint GroupIndex : SV_GroupIndex, uint3 Grou
 
 	if (Index < NumInstances)
 	{
-		FRayTracingInstanceDescriptorInput InputDesc = InputInstanceDescriptors[InputDescOffset + Index];
+		FRayTracingInstanceDescriptor InputDesc = InstanceDescriptors[BaseInstanceDescriptorIndex + Index];
 
 #if USE_GPUSCENE
 		FInstanceSceneData InstanceSceneData = GetInstanceSceneData(InputDesc.GPUSceneInstanceOrTransformIndex);
@@ -191,7 +191,7 @@ void RayTracingBuildInstanceBufferCS(uint GroupIndex : SV_GroupIndex, uint3 Grou
 			BlasAddress = 0;
 		}
 
-		InstancesDescriptors[InputDesc.OutputDescriptorIndex] = BuildPlatformRayTracingInstanceDesc(
+		OutPlatformInstanceDescriptors[InputDesc.OutputDescriptorIndex] = BuildPlatformRayTracingInstanceDesc(
 			InputDesc.InstanceMaskAndFlags & 0xFF,
 			InputDesc.InstanceId,
 			TranslateRayTracingInstanceFlags((InputDesc.InstanceMaskAndFlags >> 8) & 0xFF),
diff --git a/Engine/Shaders/Private/RayTracing/RayTracingValidation.usf b/Engine/Shaders/Private/RayTracing/RayTracingValidation.usf
--- a/Engine/Shaders/Private/RayTracing/RayTracingValidation.usf
+++ b/Engine/Shaders/Private/RayTracing/RayTracingValidation.usf
@@ -110,7 +110,7 @@ void RayTracingValidateSceneBuildParamsCS(uint3 DispatchThreadId : SV_DispatchTh
 	}
 
 	uint InstanceAddress = InstanceBufferOffsetInBytes + InstanceBufferStrideInBytes * InstanceIndex;
-	FRayTracingInstanceDescriptor InstanceDesc = InstanceBuffer.Load<FRayTracingInstanceDescriptor>(InstanceAddress);
+	FPlatformRayTracingInstanceDescriptor InstanceDesc = InstanceBuffer.Load<FPlatformRayTracingInstanceDescriptor>(InstanceAddress);
 
 	// Check that the instance matrix is finite (if instance culling is required, it should be done through null BLAS address)
 	{
diff --git a/Engine/Source/Runtime/Renderer/Private/RayTracing/RayTracingInstanceBufferUtil.cpp b/Engine/Source/Runtime/Renderer/Private/RayTracing/RayTracingInstanceBufferUtil.cpp
--- a/Engine/Source/Runtime/Renderer/Private/RayTracing/RayTracingInstanceBufferUtil.cpp
+++ b/Engine/Source/Runtime/Renderer/Private/RayTracing/RayTracingInstanceBufferUtil.cpp
@@ -103,7 +103,7 @@ void FillRayTracingInstanceUploadBuffer(
 	TConstArrayView<uint32> BaseInstancePrefixSum,
 	uint32 NumNativeGPUSceneInstances,
 	uint32 NumNativeCPUInstances,
-	TArrayView<FRayTracingInstanceDescriptorInput> OutInstanceUploadData,
+	TArrayView<FRayTracingInstanceDescriptor> OutInstanceUploadData,
 	TArrayView<FVector4f> OutTransformData)
 {
 	TRACE_CPUPROFILER_EVENT_SCOPE(FillRayTracingInstanceUploadBuffer);
@@ -154,7 +154,7 @@ void FillRayTracingInstanceUploadBuffer(
 
 			for (uint32 TransformIndex = 0; TransformIndex < NumTransforms; ++TransformIndex)
 			{
-				FRayTracingInstanceDescriptorInput InstanceDesc;
+				FRayTracingInstanceDescriptor InstanceDesc;
 
 				if (bGpuSceneInstance)
 				{
@@ -221,15 +221,15 @@ struct FRayTracingBuildInstanceBufferCS : public FGlobalShader
 	BEGIN_SHADER_PARAMETER_STRUCT(FParameters, )
 		SHADER_PARAMETER_STRUCT_INCLUDE(FGPUSceneResourceParametersRHI, GPUSceneParameters)
 
-		SHADER_PARAMETER_UAV(RWStructuredBuffer, InstancesDescriptors)
-		SHADER_PARAMETER_SRV(StructuredBuffer<FRayTracingInstanceDescriptorInput>, InputInstanceDescriptors)
+		SHADER_PARAMETER_UAV(RWStructuredBuffer, OutPlatformInstanceDescriptors)
+		SHADER_PARAMETER_SRV(StructuredBuffer, InstanceDescriptors)
 		SHADER_PARAMETER_SRV(ByteAddressBuffer, AccelerationStructureAddresses)
 		SHADER_PARAMETER_SRV(StructuredBuffer, InstanceTransforms)
 
 		SHADER_PARAMETER(FVector3f, FarFieldReferencePos)
 
 		SHADER_PARAMETER(uint32, NumInstances)
-		SHADER_PARAMETER(uint32, InputDescOffset)
+		SHADER_PARAMETER(uint32, BaseInstanceDescriptorIndex)
 
 		SHADER_PARAMETER(FVector3f, PreViewTranslationHigh)
 		SHADER_PARAMETER(FVector3f, PreViewTranslationLow)
@@ -288,7 +288,7 @@ void BuildRayTracingInstanceBuffer(
 	const FGPUScene* GPUScene,
 	const FDFVector3& PreViewTranslation,
 	uint32 NumInstances,
-	uint32 InputDescOffset,
+	uint32 BaseInstanceDescriptorIndex,
 	FRHIUnorderedAccessView* InstancesUAV,
 	FRHIShaderResourceView* InstanceUploadSRV,
 	FRHIShaderResourceView* AccelerationStructureAddressesSRV,
@@ -298,13 +298,13 @@ void BuildRayTracingInstanceBuffer(
 	FRHIUnorderedAccessView* DebugInstanceGPUSceneIndexUAV)
 {
 	FRayTracingBuildInstanceBufferCS::FParameters PassParams;
-	PassParams.InstancesDescriptors = InstancesUAV;
-	PassParams.InputInstanceDescriptors = InstanceUploadSRV;
+	PassParams.OutPlatformInstanceDescriptors = InstancesUAV;
+	PassParams.InstanceDescriptors = InstanceUploadSRV;
 	PassParams.AccelerationStructureAddresses = AccelerationStructureAddressesSRV;
 	PassParams.InstanceTransforms = InstanceTransformSRV;
 	PassParams.FarFieldReferencePos = (FVector3f)Lumen::GetFarFieldReferencePos();	// LWC_TODO: Precision Loss
 	PassParams.NumInstances = NumInstances;
-	PassParams.InputDescOffset = InputDescOffset;
+	PassParams.BaseInstanceDescriptorIndex = BaseInstanceDescriptorIndex;
 	PassParams.PreViewTranslationHigh = PreViewTranslation.High;
 	PassParams.PreViewTranslationLow = PreViewTranslation.Low;
 
@@ -334,7 +334,6 @@ void BuildRayTracingInstanceBuffer(
 	PermutationVector.Set<FRayTracingBuildInstanceBufferCS::FGpuCullingDim>(CullingParameters != nullptr);
 
 	auto ComputeShader = GetGlobalShaderMap(GMaxRHIFeatureLevel)->GetShader<FRayTracingBuildInstanceBufferCS>(PermutationVector);
-	const int32 GroupSize = FMath::DivideAndRoundUp(PassParams.NumInstances, FRayTracingBuildInstanceBufferCS::ThreadGroupSize);
 
 	const FIntVector GroupCount = FComputeShaderUtils::GetGroupCountWrapped(PassParams.NumInstances, FRayTracingBuildInstanceBufferCS::ThreadGroupSize);
 
diff --git a/Engine/Source/Runtime/Renderer/Private/RayTracing/RayTracingScene.cpp b/Engine/Source/Runtime/Renderer/Private/RayTracing/RayTracingScene.cpp
--- a/Engine/Source/Runtime/Renderer/Private/RayTracing/RayTracingScene.cpp
+++ b/Engine/Source/Runtime/Renderer/Private/RayTracing/RayTracingScene.cpp
@@ -159,14 +159,14 @@ void FRayTracingScene::Create(FRDGBuilder& GraphBuilder, const FViewInfo& View,
 
 		{
 			// Create/resize instance upload buffer (if necessary)
-			const uint32 UploadBufferSize = NumNativeInstancesAligned * sizeof(FRayTracingInstanceDescriptorInput);
+			const uint32 UploadBufferSize = NumNativeInstancesAligned * sizeof(FRayTracingInstanceDescriptor);
 
 			if (!Layer.InstanceUploadBuffer.IsValid()
 				|| UploadBufferSize > Layer.InstanceUploadBuffer->GetSize()
 				|| UploadBufferSize < Layer.InstanceUploadBuffer->GetSize() / 2)
 			{
 				FRHIResourceCreateInfo CreateInfo(TEXT("FRayTracingScene::InstanceUploadBuffer"));
-				Layer.InstanceUploadBuffer = RHICmdList.CreateStructuredBuffer(sizeof(FRayTracingInstanceDescriptorInput), UploadBufferSize, BUF_ShaderResource | BUF_Volatile, CreateInfo);
+				Layer.InstanceUploadBuffer = RHICmdList.CreateStructuredBuffer(sizeof(FRayTracingInstanceDescriptor), UploadBufferSize, BUF_ShaderResource | BUF_Volatile, CreateInfo);
 				Layer.InstanceUploadSRV = RHICmdList.CreateShaderResourceView(Layer.InstanceUploadBuffer, FRHIViewDesc::CreateBufferSRV().SetTypeFromBuffer(Layer.InstanceUploadBuffer));
 			}
 		}
@@ -275,10 +275,10 @@ void FRayTracingScene::Create(FRDGBuilder& GraphBuilder, const FViewInfo& View,
 				{
 					FOptionalTaskTagScope TaskTagScope(ETaskTag::EParallelRenderingThread);
 
-					const uint32 InstanceUploadBytes = NumNativeInstances * sizeof(FRayTracingInstanceDescriptorInput);
+					const uint32 InstanceUploadBytes = NumNativeInstances * sizeof(FRayTracingInstanceDescriptor);
 					const uint32 TransformUploadBytes = NumNativeCPUInstances * 3 * sizeof(FVector4f);
 
-					FRayTracingInstanceDescriptorInput* InstanceUploadData = (FRayTracingInstanceDescriptorInput*)RHICmdList.LockBuffer(InstanceUploadBuffer, 0, InstanceUploadBytes, RLM_WriteOnly);
+					FRayTracingInstanceDescriptor* InstanceUploadData = (FRayTracingInstanceDescriptor*)RHICmdList.LockBuffer(InstanceUploadBuffer, 0, InstanceUploadBytes, RLM_WriteOnly);
 					FVector4f* TransformUploadData = (NumNativeCPUInstances > 0) ? (FVector4f*)RHICmdList.LockBuffer(TransformUploadBuffer, 0, TransformUploadBytes, RLM_WriteOnly) : nullptr;
 
 					FillRayTracingInstanceUploadBuffer(
diff --git a/Engine/Source/Runtime/Renderer/Private/Tests/RayTracingTestbed.cpp b/Engine/Source/Runtime/Renderer/Private/Tests/RayTracingTestbed.cpp
--- a/Engine/Source/Runtime/Renderer/Private/Tests/RayTracingTestbed.cpp
+++ b/Engine/Source/Runtime/Renderer/Private/Tests/RayTracingTestbed.cpp
@@ -212,12 +212,12 @@ bool RunRayTracingTestbed_RenderThread(const FString& Parameters)
 	FByteAddressBuffer AccelerationStructureAddressesBuffer;
 	AccelerationStructureAddressesBuffer.Initialize(RHICmdList, TEXT("RayTracingTestBedAccelerationStructureAddressesBuffer"), sizeof(FRayTracingAccelerationStructureAddress), BUF_Volatile | BUF_MultiGPUAllocate);
 
-	const uint32 InstanceUploadBufferSize = SceneInitializer.MaxNumInstances * sizeof(FRayTracingInstanceDescriptorInput);
+	const uint32 InstanceUploadBufferSize = SceneInitializer.MaxNumInstances * sizeof(FRayTracingInstanceDescriptor);
 	FBufferRHIRef InstanceUploadBuffer;
 	FShaderResourceViewRHIRef InstanceUploadSRV;
 	{
 		FRHIResourceCreateInfo CreateInfo(TEXT("RayTracingTestBedInstanceUploadBuffer"));
-		InstanceUploadBuffer = RHICmdList.CreateStructuredBuffer(sizeof(FRayTracingInstanceDescriptorInput), InstanceUploadBufferSize, BUF_ShaderResource | BUF_Volatile, CreateInfo);
+		InstanceUploadBuffer = RHICmdList.CreateStructuredBuffer(sizeof(FRayTracingInstanceDescriptor), InstanceUploadBufferSize, BUF_ShaderResource | BUF_Volatile, CreateInfo);
 		InstanceUploadSRV = RHICmdList.CreateShaderResourceView(InstanceUploadBuffer, FRHIViewDesc::CreateBufferSRV().SetTypeFromBuffer(InstanceUploadBuffer));
 	}
 
@@ -231,7 +231,7 @@ bool RunRayTracingTestbed_RenderThread(const FString& Parameters)
 	}
 
 	{
-		FRayTracingInstanceDescriptorInput* InstanceUploadData = (FRayTracingInstanceDescriptorInput*)RHICmdList.LockBuffer(InstanceUploadBuffer, 0, InstanceUploadBufferSize, RLM_WriteOnly);
+		FRayTracingInstanceDescriptor* InstanceUploadData = (FRayTracingInstanceDescriptor*)RHICmdList.LockBuffer(InstanceUploadBuffer, 0, InstanceUploadBufferSize, RLM_WriteOnly);
 		FVector4f* TransformUploadData = (FVector4f*)RHICmdList.LockBuffer(TransformUploadBuffer, 0, TransformUploadBufferSize, RLM_WriteOnly);
 		FillRayTracingInstanceUploadBuffer(
 			FVector::ZeroVector,
diff --git a/Engine/Source/Runtime/Renderer/Public/RayTracingInstanceBufferUtil.h b/Engine/Source/Runtime/Renderer/Public/RayTracingInstanceBufferUtil.h
--- a/Engine/Source/Runtime/Renderer/Public/RayTracingInstanceBufferUtil.h
+++ b/Engine/Source/Runtime/Renderer/Public/RayTracingInstanceBufferUtil.h
@@ -18,14 +18,14 @@ struct FDFVector3;
 * The FRayTracingGeometryInstance array (ie: FRayTracingScene::Instances) used to create FRayTracingSceneRHI
 * can have mix of instances using GPUScene or CPU transforms.
 * In order to reduce the number of dispatches to build the native RayTracing Instance Buffer,
-* the upload buffer containing FRayTracingInstanceDescriptorInput is split in 2 sections, [GPUSceneInstances] [CPUInstances].
+* the upload buffer containing FRayTracingInstanceDescriptor is split in 2 sections, [GPUSceneInstances] [CPUInstances].
 * This way native GPUScene and CPU instance descriptors can be built in a single dispatch per type.
 * 
 * If the ray tracing scene contains multiple layers, the instance buffer is divided into multiple subranges as expected by the RHI.
 * 
 */
 
-struct FRayTracingInstanceDescriptorInput
+struct FRayTracingInstanceDescriptor
 {
 	uint32 GPUSceneInstanceOrTransformIndex;
 	uint32 OutputDescriptorIndex;
@@ -36,6 +36,8 @@ struct FRayTracingInstanceDescriptorInput
 	uint32 bApplyLocalBoundsTransform;
 };
 
+using FRayTracingInstanceDescriptorInput UE_DEPRECATED(5.6, "Use FRayTracingInstanceDescriptor instead.") = FRayTracingInstanceDescriptor;
+
 struct FRayTracingSceneInitializationData
 {
 	uint32 NumNativeGPUSceneInstances;
@@ -68,7 +70,7 @@ RENDERER_API void FillRayTracingInstanceUploadBuffer(
 	TConstArrayView<uint32> BaseInstancePrefixSum,
 	uint32 NumNativeGPUSceneInstances,
 	uint32 NumNativeCPUInstances,
-	TArrayView<FRayTracingInstanceDescriptorInput> OutInstanceUploadData,
+	TArrayView<FRayTracingInstanceDescriptor> OutInstanceUploadData,
 	TArrayView<FVector4f> OutTransformData);
 
 UE_DEPRECATED(5.6, "No longer need to provide RayTracingSceneRHI.")
@@ -81,7 +83,7 @@ inline void FillRayTracingInstanceUploadBuffer(
 	TConstArrayView<uint32> BaseInstancePrefixSum,
 	uint32 NumNativeGPUSceneInstances,
 	uint32 NumNativeCPUInstances,
-	TArrayView<FRayTracingInstanceDescriptorInput> OutInstanceUploadData,
+	TArrayView<FRayTracingInstanceDescriptor> OutInstanceUploadData,
 	TArrayView<FVector4f> OutTransformData)
 {
 	FillRayTracingInstanceUploadBuffer(
