Disable persistent SBTs until minor fallout is fixed

diff --git a/Engine/Source/Runtime/Renderer/Private/RayTracing/RayTracingShaderBindingTable.cpp b/Engine/Source/Runtime/Renderer/Private/RayTracing/RayTracingShaderBindingTable.cpp
--- a/Engine/Source/Runtime/Renderer/Private/RayTracing/RayTracingShaderBindingTable.cpp
+++ b/Engine/Source/Runtime/Renderer/Private/RayTracing/RayTracingShaderBindingTable.cpp
@@ -9,7 +9,7 @@
 
 #if RHI_RAYTRACING
 
-static int32 GPersistentSBTEnabled = 1;
+static int32 GPersistentSBTEnabled = 0;
 static FAutoConsoleVariableRef CVarRayTracingPersistentSBT(
 	TEXT("r.RayTracing.PersistentSBT"),
 	GPersistentSBTEnabled,
