LFA - added support to batched translucent lighting volume injection.

#tests PC dx12 editor
#rnx

diff --git a/Engine/Shaders/Private/TranslucentLightInjectionShaders.usf b/Engine/Shaders/Private/TranslucentLightInjectionShaders.usf
--- a/Engine/Shaders/Private/TranslucentLightInjectionShaders.usf
+++ b/Engine/Shaders/Private/TranslucentLightInjectionShaders.usf
@@ -349,6 +349,8 @@ void InjectMainPS(
 #include "LightGridCommon.ush"
 #include "LightData.ush"
 
+#include "LightFunctionAtlas/LightFunctionAtlasCommon.usf"
+
 //int VirtualShadowMapId;
 StructuredBuffer<uint> BatchedLocalLights;
 uint MaxBatchedLocalLights;
@@ -397,6 +399,7 @@ void InjectBatchMainPS(
 							ShadowFactor *= VirtualShadowMapSample.ShadowFactor;
 						}
 #endif
+
 						if (VolumeCascadeIndex == 1)
 						{
 							float FinalLerpFactor = GetBorderLerpFactor(Input.Vertex.UV, Input.LayerIndex);
@@ -405,7 +408,12 @@ void InjectBatchMainPS(
 							ShadowFactor = lerp(0.0f, ShadowFactor, FinalLerpFactor);
 						}
 
-						float3 Lighting = LightData.Color / PI * Attenuation * ShadowFactor;
+						FLightFunctionColor LightFunctionColor = 1.0f;
+					#if	USE_LIGHT_FUNCTION_ATLAS
+						LightFunctionColor = GetLocalLightFunctionCommon(TranslatedWorldPosition, LightData.LightFunctionAtlasLightIndex);
+					#endif
+
+						float3 Lighting = LightFunctionColor * LightData.Color / PI * Attenuation * ShadowFactor;
 
 						float DirectionalLightContribution = 0;
 						AccumulateSHLighting(NormalizedLightVector, Lighting, DirectionalLightContribution,
diff --git a/Engine/Source/Runtime/Renderer/Private/TranslucentLighting.cpp b/Engine/Source/Runtime/Renderer/Private/TranslucentLighting.cpp
--- a/Engine/Source/Runtime/Renderer/Private/TranslucentLighting.cpp
+++ b/Engine/Source/Runtime/Renderer/Private/TranslucentLighting.cpp
@@ -55,6 +55,9 @@
 #include "HeterogeneousVolumes/HeterogeneousVolumes.h"
 #include "DeferredShadingRenderer.h"
 #include "BasePassRendering.h"
+#include "LightFunctionAtlas.h"
+
+using namespace LightFunctionAtlas;
 
 class FMaterial;
 
@@ -821,6 +824,7 @@ class FTranslucentLightingInjectBatchPS : public FGlobalShader
 		SHADER_PARAMETER_STRUCT_REF(FViewUniformShaderParameters, ViewUniformBuffer)
 		SHADER_PARAMETER_RDG_UNIFORM_BUFFER(FForwardLightUniformParameters, ForwardLightStruct)
 		SHADER_PARAMETER_STRUCT_INCLUDE(FVirtualShadowMapSamplingParameters, VirtualShadowMapSamplingParameters)
+		SHADER_PARAMETER_RDG_UNIFORM_BUFFER(FLightFunctionAtlasGlobalParameters, LightFunctionAtlas)
 		SHADER_PARAMETER_RDG_BUFFER_SRV(StructuredBuffer<uint>, BatchedLocalLights)
 		SHADER_PARAMETER(uint32, MaxBatchedLocalLights)
 		SHADER_PARAMETER(uint32, VolumeCascadeIndex)
@@ -828,7 +832,8 @@ class FTranslucentLightingInjectBatchPS : public FGlobalShader
 	END_SHADER_PARAMETER_STRUCT()
 
 	class FVirtualShadowMap : SHADER_PERMUTATION_BOOL("VIRTUAL_SHADOW_MAP");
-	using FPermutationDomain = TShaderPermutationDomain<FVirtualShadowMap>;
+	class FUseLightFunctionAtlas : SHADER_PERMUTATION_BOOL("USE_LIGHT_FUNCTION_ATLAS");
+	using FPermutationDomain = TShaderPermutationDomain<FVirtualShadowMap, FUseLightFunctionAtlas>;
 
 public:
 
@@ -1184,9 +1189,16 @@ void FTranslucentLightInjectionCollector::AddLightForInjection(
 		const FVisibleLightInfo& VisibleLightInfo = VisibleLightInfos[LightSceneInfo.Id];
 		const ERHIFeatureLevel::Type FeatureLevel = View.FeatureLevel;
 
-		const bool bApplyLightFunction = (View.Family->EngineShowFlags.LightFunctions &&
-			LightSceneInfo.Proxy->GetLightFunctionMaterial() && 
-			LightSceneInfo.Proxy->GetLightFunctionMaterial()->GetIncompleteMaterialWithFallback(FeatureLevel).IsLightFunction());
+		bool bLightFunctionUsesAtlas = false;
+		bool bApplyLightFunction = false;
+		const FMaterialRenderProxy* LightFunctionMaterialRenderProxy = LightSceneInfo.Proxy->GetLightFunctionMaterial();
+		if (LightFunctionMaterialRenderProxy && View.Family->EngineShowFlags.LightFunctions)
+		{
+			const FMaterial& LightFunctionMaterial	= LightFunctionMaterialRenderProxy->GetIncompleteMaterialWithFallback(FeatureLevel);
+
+			bApplyLightFunction						= LightFunctionMaterial.IsLightFunction();
+			bLightFunctionUsesAtlas					= bApplyLightFunction && LightFunctionMaterial.MaterialIsLightFunctionAtlasCompatible_RenderThread();
+		}
 
 		const int32 VirtualShadowMapId = VisibleLightInfo.GetVirtualShadowMapId(&View);
 		
@@ -1210,7 +1222,7 @@ void FTranslucentLightInjectionCollector::AddLightForInjection(
 			(bSupportRectLights || LightType != LightType_Rect) &&
 			!bStaticShadowing &&
 			InProjectedShadowInfo == nullptr &&
-			!bApplyLightFunction &&												// TODO: Support LF atlas
+			(!bApplyLightFunction || bLightFunctionUsesAtlas) &&
 			!bUseAdaptiveVolumetricShadowMap;
 
 		// If it would otherwise be supported, see if we can find the local light index in the light grid
@@ -1361,8 +1373,11 @@ static void InjectTranslucencyLightingVolumeBatch(
 	TShaderMapRef<FWriteToSliceVS> VertexShader(View.ShaderMap);
 	TOptionalShaderMapRef<FWriteToSliceGS> GeometryShader(View.ShaderMap);
 
+	const bool bUseLightFunctionAtlas = LightFunctionAtlas::IsEnabled(View, ELightFunctionAtlasSystem::DeferredLighting);
+
 	FTranslucentLightingInjectBatchPS::FPermutationDomain PermutationVector;
 	PermutationVector.Set< FTranslucentLightingInjectBatchPS::FVirtualShadowMap >(bSupportVirtualShadowMaps);
+	PermutationVector.Set< FTranslucentLightingInjectBatchPS::FUseLightFunctionAtlas >(bUseLightFunctionAtlas);
 	TShaderMapRef<FTranslucentLightingInjectBatchPS> PixelShader(View.ShaderMap, PermutationVector);
 
 	auto* PassParameters = GraphBuilder.AllocParameters<FTranslucentLightingInjectBatchPS::FParameters>();
@@ -1373,6 +1388,7 @@ static void InjectTranslucencyLightingVolumeBatch(
 	PassParameters->ViewUniformBuffer = View.ViewUniformBuffer;
 	PassParameters->ForwardLightStruct = View.ForwardLightingResources.ForwardLightUniformBuffer;
 	PassParameters->BatchedLocalLights = GraphBuilder.CreateSRV(BatchedLocalLightsRDG);
+	PassParameters->LightFunctionAtlas = LightFunctionAtlas::BindGlobalParameters(GraphBuilder, View);;
 	PassParameters->MaxBatchedLocalLights = MaxBatchedLocalLights;
 	PassParameters->VolumeCascadeIndex = VolumeCascadeIndex;
 
