Fix rendering velocities in depth prepass when velocities are disabled.

diff --git a/Engine/Source/Runtime/Renderer/Private/MobileShadingRenderer.cpp b/Engine/Source/Runtime/Renderer/Private/MobileShadingRenderer.cpp
--- a/Engine/Source/Runtime/Renderer/Private/MobileShadingRenderer.cpp
+++ b/Engine/Source/Runtime/Renderer/Private/MobileShadingRenderer.cpp
@@ -815,7 +815,7 @@ void FMobileSceneRenderer::RenderFullDepthPrepass(FRDGBuilder& GraphBuilder, TAr
 
 		View.BeginRenderView();
 
-		const bool bRenderVelocityInDepthPrePass = (Scene->EarlyZPassMode == DDM_AllOpaqueNoVelocity);
+		const bool bRenderVelocityInDepthPrePass = (Scene->EarlyZPassMode == DDM_AllOpaqueNoVelocity) && bShouldRenderVelocities;
 
 		struct FFullDepthPrepassParameterCollection
 		{
