Move runtime virtual texture output settings out of view uniform buffer and into a specific uniform buffer.
We would like to extend the UB data in the future and don't want to pollute the view UB with that.

diff --git a/Engine/Shaders/Private/VirtualTextureMaterial.usf b/Engine/Shaders/Private/VirtualTextureMaterial.usf
--- a/Engine/Shaders/Private/VirtualTextureMaterial.usf
+++ b/Engine/Shaders/Private/VirtualTextureMaterial.usf
@@ -113,8 +113,8 @@ void FPixelShaderInOut_MainPS(
 #endif
 
 	// Apply any fixed color.
-	BaseColor = lerp(BaseColor, ResolvedView.RuntimeVirtualTextureDebugParams.xyz, ResolvedView.RuntimeVirtualTextureDebugParams.w);
-	Displacement = lerp(Displacement, ResolvedView.RuntimeVirtualTextureDebugParams.x, ResolvedView.RuntimeVirtualTextureDebugParams.w);
+	BaseColor =  lerp(BaseColor, RuntimeVirtualTexturePassParameters.RuntimeVirtualTextureDebugParams.xyz, RuntimeVirtualTexturePassParameters.RuntimeVirtualTextureDebugParams.w);
+	Displacement = lerp(Displacement, RuntimeVirtualTexturePassParameters.RuntimeVirtualTextureDebugParams.x, RuntimeVirtualTexturePassParameters.RuntimeVirtualTextureDebugParams.w);
 
 #if defined(OUT_BASECOLOR)
 	Out.MRT[0] = float4(BaseColor, 1.f) * Opacity;
@@ -133,7 +133,7 @@ void FPixelShaderInOut_MainPS(
 	Out.MRT[0] = float4(Mask4.xyz, 1.f) * Opacity;
 	Out.MRT[1] = float4(Mask4.w, 0.f, 0.f, 1.f) * Opacity;
 #elif defined(OUT_WORLDHEIGHT)
-	float PackedHeight = PackWorldHeight(WorldHeight, ResolvedView.RuntimeVirtualTexturePackHeight);
+	float PackedHeight = PackWorldHeight(WorldHeight, RuntimeVirtualTexturePassParameters.RuntimeVirtualTexturePackHeight);
 	Out.MRT[0] = float4(PackedHeight, 0.f, 0.f, 1.f);
 #elif defined(OUT_DISPLACEMENT)
 	Out.MRT[0] = float4(Displacement, 0.f, 0.f, 1.f) * Opacity;
diff --git a/Engine/Source/Runtime/Engine/Public/SceneView.h b/Engine/Source/Runtime/Engine/Public/SceneView.h
--- a/Engine/Source/Runtime/Engine/Public/SceneView.h
+++ b/Engine/Source/Runtime/Engine/Public/SceneView.h
@@ -1041,8 +1041,6 @@ enum ETranslucencyVolumeCascade
 	VIEW_UNIFORM_BUFFER_MEMBER(uint32, VirtualTextureFeedbackJitterOffset) \
 	VIEW_UNIFORM_BUFFER_MEMBER(uint32, VirtualTextureFeedbackSampleOffset) \
 	VIEW_UNIFORM_BUFFER_MEMBER(FVector4f, RuntimeVirtualTextureMipLevel) \
-	VIEW_UNIFORM_BUFFER_MEMBER(FVector2f, RuntimeVirtualTexturePackHeight) \
-	VIEW_UNIFORM_BUFFER_MEMBER(FVector4f, RuntimeVirtualTextureDebugParams) \
 	VIEW_UNIFORM_BUFFER_MEMBER(int32, FarShadowStaticMeshLODBias) \
 	VIEW_UNIFORM_BUFFER_MEMBER(float, MinRoughness) \
 	VIEW_UNIFORM_BUFFER_MEMBER(FVector4f, HairRenderInfo) \
diff --git a/Engine/Source/Runtime/Renderer/Private/SceneRendering.cpp b/Engine/Source/Runtime/Renderer/Private/SceneRendering.cpp
--- a/Engine/Source/Runtime/Renderer/Private/SceneRendering.cpp
+++ b/Engine/Source/Runtime/Renderer/Private/SceneRendering.cpp
@@ -2014,8 +2014,6 @@ void FViewInfo::SetupUniformBufferParameters(
 	ViewUniformShaderParameters.VirtualTextureFeedbackSampleOffset = (FrameIndex % NumPixelsInTile) + (FrameIndex / NumPixelsInTile);
 
 	ViewUniformShaderParameters.RuntimeVirtualTextureMipLevel = FVector4f(ForceInitToZero);
-	ViewUniformShaderParameters.RuntimeVirtualTexturePackHeight = FVector2f(ForceInitToZero);
-	ViewUniformShaderParameters.RuntimeVirtualTextureDebugParams = FVector4f(ForceInitToZero);
 
 	// GGX/Sheen LTC (used as BSDF or for rect light integration)
 	if (GSystemTextures.GGXLTCMat.IsValid() && GSystemTextures.GGXLTCAmp.IsValid())
diff --git a/Engine/Source/Runtime/Renderer/Private/VT/RuntimeVirtualTextureRender.cpp b/Engine/Source/Runtime/Renderer/Private/VT/RuntimeVirtualTextureRender.cpp
--- a/Engine/Source/Runtime/Renderer/Private/VT/RuntimeVirtualTextureRender.cpp
+++ b/Engine/Source/Runtime/Renderer/Private/VT/RuntimeVirtualTextureRender.cpp
@@ -130,6 +130,15 @@ namespace RuntimeVirtualTexture
 		return IsOpenGLPlatform(ShaderPlatform);
 	}
 
+	/** Parameters used when writing to the virtual texture. */
+	BEGIN_GLOBAL_SHADER_PARAMETER_STRUCT(FRuntimeVirtualTexturePassParameters, )
+		SHADER_PARAMETER(FVector2f, RuntimeVirtualTexturePackHeight)
+		SHADER_PARAMETER(FVector4f, RuntimeVirtualTextureDebugParams)
+	END_GLOBAL_SHADER_PARAMETER_STRUCT()
+
+	/** Uniform buffer for writing to the virtual texture. We reuse the DeferredDecals UB slot, which can't be used at the same time. This avoids the overhead of a new slot. */
+	IMPLEMENT_STATIC_UNIFORM_BUFFER_STRUCT(FRuntimeVirtualTexturePassParameters, "RuntimeVirtualTexturePassParameters", DeferredDecals);
+
 	/** Mesh material shader for writing to the virtual texture. */
 	class FShader_VirtualTextureMaterialDraw : public FMeshMaterialShader
 	{
@@ -137,6 +146,7 @@ namespace RuntimeVirtualTexture
 		BEGIN_SHADER_PARAMETER_STRUCT(FParameters, )
 			SHADER_PARAMETER_STRUCT_REF(FViewUniformShaderParameters, View)
 			SHADER_PARAMETER_RDG_UNIFORM_BUFFER(FSceneUniformParameters, Scene)
+			SHADER_PARAMETER_RDG_UNIFORM_BUFFER(FRuntimeVirtualTexturePassParameters, RuntimeVirtualTexturePassUB)
 			SHADER_PARAMETER_STRUCT_INCLUDE(FInstanceCullingDrawParams, InstanceCullingDrawParams)
 			RDG_TEXTURE_ACCESS(PhysicalTexture0, ERHIAccess::SRVGraphics)
 			RDG_TEXTURE_ACCESS(PhysicalTexture1, ERHIAccess::SRVGraphics)
@@ -166,6 +176,7 @@ namespace RuntimeVirtualTexture
 		FShader_VirtualTextureMaterialDraw(const FMeshMaterialShaderType::CompiledShaderInitializerType& Initializer)
 			: FMeshMaterialShader(Initializer)
 		{
+			PassUniformBuffer.Bind(Initializer.ParameterMap, FRuntimeVirtualTexturePassParameters::FTypeInfo::GetStructMetadata()->GetShaderVariableName());
 		}
 	};
 
@@ -1564,8 +1575,6 @@ namespace RuntimeVirtualTexture
 		View->CachedViewUniformShaderParameters = MakeUnique<FViewUniformShaderParameters>();
 		View->SetupUniformBufferParameters(nullptr, 0, *View->CachedViewUniformShaderParameters);
 		View->CachedViewUniformShaderParameters->RuntimeVirtualTextureMipLevel = MipLevelParameter;
-		View->CachedViewUniformShaderParameters->RuntimeVirtualTexturePackHeight = FVector2f(WorldHeightPackParameter);	// LWC_TODO: Precision loss
-		View->CachedViewUniformShaderParameters->RuntimeVirtualTextureDebugParams = CVarVTMipColors.GetValueOnRenderThread() ? GetMipLevelColor(vLevel) : FixedColor;
 		View->ViewUniformBuffer = TUniformBufferRef<FViewUniformShaderParameters>::CreateUniformBufferImmediate(*View->CachedViewUniformShaderParameters, UniformBuffer_SingleFrame);
 
 		// Build graph
@@ -1587,6 +1596,12 @@ namespace RuntimeVirtualTexture
 			FShader_VirtualTextureMaterialDraw::FParameters* PassParameters = GraphBuilder.AllocParameters<FShader_VirtualTextureMaterialDraw::FParameters>();
 			PassParameters->View = View->ViewUniformBuffer;
 			PassParameters->Scene = SceneRenderer->GetSceneUniformBufferRef(GraphBuilder);
+
+			FRuntimeVirtualTexturePassParameters* RuntimeVirtualTexturePassParameters = GraphBuilder.AllocParameters<FRuntimeVirtualTexturePassParameters>();
+			RuntimeVirtualTexturePassParameters->RuntimeVirtualTexturePackHeight = FVector2f(WorldHeightPackParameter);	// LWC_TODO: Precision loss
+			RuntimeVirtualTexturePassParameters->RuntimeVirtualTextureDebugParams = CVarVTMipColors.GetValueOnRenderThread() ? GetMipLevelColor(vLevel) : FixedColor;
+			PassParameters->RuntimeVirtualTexturePassUB = GraphBuilder.CreateUniformBuffer(RuntimeVirtualTexturePassParameters);
+
 			// When direct aliasing the output physical texture for compression, the texture gets managed by RDG.
 			// The texture needs to be transitioned to UAV during the compression pass.
 			// But we may also be accessing the physical texture in a material in this render pass (using the raw RHI SRV).
