[HWRT] Removed unused GetHitRecordBaseIndex(...).

diff --git a/Engine/Source/Runtime/D3D12RHI/Private/D3D12RayTracing.cpp b/Engine/Source/Runtime/D3D12RHI/Private/D3D12RayTracing.cpp
--- a/Engine/Source/Runtime/D3D12RHI/Private/D3D12RayTracing.cpp
+++ b/Engine/Source/Runtime/D3D12RHI/Private/D3D12RayTracing.cpp
@@ -3072,8 +3072,6 @@ public:
 		ShaderTablesPerGPU[GPUIndex] = nullptr;
 	}
 
-	uint32 GetHitRecordBaseIndex(uint32 GlobalSegmentIndex) const { return GlobalSegmentIndex * Initializer.NumShaderSlotsPerGeometrySegment; }
-
 private:
 
 	void InitForDevice(FD3D12Device* Device)
diff --git a/Engine/Source/Runtime/Renderer/Private/RayTracing/RayTracingInstanceBufferUtil.cpp b/Engine/Source/Runtime/Renderer/Private/RayTracing/RayTracingInstanceBufferUtil.cpp
--- a/Engine/Source/Runtime/Renderer/Private/RayTracing/RayTracingInstanceBufferUtil.cpp
+++ b/Engine/Source/Runtime/Renderer/Private/RayTracing/RayTracingInstanceBufferUtil.cpp
@@ -40,8 +40,6 @@ FRayTracingSceneInitializationData BuildRayTracingSceneInitializationData(TConst
 
 	Experimental::TSherwoodMap<FRHIRayTracingGeometry*, uint32> UniqueGeometries;
 
-	// Compute geometry segment prefix sum used by GetHitRecordBaseIndex() during resource binding
-
 	uint32 NumNativeInstances = 0;
 
 	for (uint32 InstanceIndex = 0; InstanceIndex < NumSceneInstances; ++InstanceIndex)
