Quantize Automatic View Texture MipLODBias to significantly reduce the number of independent sampler states used when Dynamic Resolution is enabled, preventing crashes in long-running applications that hit the sampler limit.

diff --git a/Engine/Source/Runtime/Renderer/Private/SceneVisibility.cpp b/Engine/Source/Runtime/Renderer/Private/SceneVisibility.cpp
--- a/Engine/Source/Runtime/Renderer/Private/SceneVisibility.cpp
+++ b/Engine/Source/Runtime/Renderer/Private/SceneVisibility.cpp
@@ -156,18 +156,25 @@ RENDERER_API TAutoConsoleVariable<float> CVarStaticMeshLODDistanceScale(
 	TEXT("(higher values make LODs transition earlier, e.g., 2 is twice as fast / half the distance)"),
 	ECVF_Scalability | ECVF_RenderThreadSafe);
 
-static TAutoConsoleVariable<float> CVarMinAutomaticViewMipBias(
+static TAutoConsoleVariable<float> CVarAutomaticViewMipBiasMin(
 	TEXT("r.ViewTextureMipBias.Min"),
 	-2.0f,
 	TEXT("Automatic view mip bias's minimum value (default to -2)."),
 	ECVF_RenderThreadSafe);
 
-static TAutoConsoleVariable<float> CVarMinAutomaticViewMipBiasOffset(
+static TAutoConsoleVariable<float> CVarAutomaticViewMipBiasOffset(
 	TEXT("r.ViewTextureMipBias.Offset"),
 	-0.3,
 	TEXT("Automatic view mip bias's constant offset (default to -0.3)."),
 	ECVF_RenderThreadSafe);
 
+static TAutoConsoleVariable<int32> CVarAutomaticViewMipBiasQuantization(
+	TEXT("r.ViewTextureMipBias.Quantization"),
+	1024,
+	TEXT("Quantization step count to use for automatic view mip bias, measured between -4.0f and 4.0f.\n")
+	TEXT("Lowering this can significantly reduce unique sampler states when Dynamic Resolution is active."),
+	ECVF_RenderThreadSafe);
+
 static int32 GILCUpdatePrimTaskEnabled = 1;
 
 static FAutoConsoleVariableRef CVarILCUpdatePrimitivesTask(
@@ -5049,8 +5056,21 @@ void FSceneRenderer::PrepareViewStateForVisibility(const FSceneTexturesConfig& S
 				SampleX = Halton(TemporalSampleIndex + 1, 2) - 0.5f;
 				SampleY = Halton(TemporalSampleIndex + 1, 3) - 0.5f;
 
-				View.MaterialTextureMipBias = -(FMath::Max(-FMath::Log2(EffectivePrimaryResolutionFraction), 0.0f) ) + CVarMinAutomaticViewMipBiasOffset.GetValueOnRenderThread();
-				View.MaterialTextureMipBias = FMath::Max(View.MaterialTextureMipBias, CVarMinAutomaticViewMipBias.GetValueOnRenderThread());
+				// Adjust Automatic View Mip Bias
+				float MipBias = -(FMath::Max(-FMath::Log2(EffectivePrimaryResolutionFraction), 0.0f)) + CVarAutomaticViewMipBiasOffset.GetValueOnRenderThread();
+				MipBias = FMath::Clamp(FMath::Max(MipBias, CVarAutomaticViewMipBiasMin.GetValueOnRenderThread()), -4.0f, 4.0f);
+
+				// Quantize to discrete values between -4.0f and 4.0f (range of 8.0f), since each value creates a unique sampler state
+				const int32 MipBiasQuantization = CVarAutomaticViewMipBiasQuantization.GetValueOnRenderThread();
+
+				if (MipBiasQuantization > 0)
+				{
+					// Divide the 8.0f range by the number of desired steps to get the step size
+					const float QuantizationStepSize = 8.0f / ((float)MipBiasQuantization);
+					MipBias = FMath::CeilToFloat(MipBias / QuantizationStepSize) * QuantizationStepSize;
+				}
+				
+				View.MaterialTextureMipBias = MipBias;
 			}
 			else if( CVarTemporalAASamplesValue == 2 )
 			{
