Substarte - fixed stationary lights shadow not working due to clearing too many bits from the stencil buffer.

#tests PC dx12 editor
#rnx

diff --git a/Engine/Source/Runtime/Renderer/Private/ShadowRendering.cpp b/Engine/Source/Runtime/Renderer/Private/ShadowRendering.cpp
--- a/Engine/Source/Runtime/Renderer/Private/ShadowRendering.cpp
+++ b/Engine/Source/Runtime/Renderer/Private/ShadowRendering.cpp
@@ -21,6 +21,7 @@
 #include "MobileBasePassRendering.h"
 
 using namespace LightFunctionAtlas;
+using namespace ShadowRendering;
 
 ///////////////////////////////////////////////////////////////////////////////////////////////////
 // Directional light
@@ -244,9 +245,6 @@ static TAutoConsoleVariable<int32> CVarHairStrandsCullPerObjectShadowCaster(
 
 DEFINE_GPU_STAT(ShadowProjection);
 
-// Use Shadow stencil mask is set to 0x07u instead of 0xFF so that that last bit can be used for Substrate classification without clearing the stencil bit for pre-shadow/per-object static shadow mask
-constexpr uint32 ShadowStencilMask = 0x07u;
-
 // 0:off, 1:low, 2:med, 3:high, 4:very high, 5:max
 uint32 GetShadowQuality()
 {
diff --git a/Engine/Source/Runtime/Renderer/Private/ShadowRendering.h b/Engine/Source/Runtime/Renderer/Private/ShadowRendering.h
--- a/Engine/Source/Runtime/Renderer/Private/ShadowRendering.h
+++ b/Engine/Source/Runtime/Renderer/Private/ShadowRendering.h
@@ -46,6 +46,12 @@ class FVirtualShadowMapPerLightCacheEntry;
 class FLightTileIntersectionParameters;
 class FDistanceFieldCulledObjectBufferParameters;
 
+namespace ShadowRendering
+{
+	// Use Shadow stencil mask is set to 0x07u instead of 0xFF so that that highest bit can be used for Substrate classification without clearing the stencil bit for pre-shadow/per-object static shadow mask
+	constexpr uint32 ShadowStencilMask = 0x07u;
+};
+
 DECLARE_GPU_DRAWCALL_STAT_NAMED_EXTERN(ShadowDepths, TEXT("Shadow Depths"));
 
 /** Renders a cone with a spherical cap, used for rendering spot lights in deferred passes. */
diff --git a/Engine/Source/Runtime/Renderer/Private/ShadowSetup.cpp b/Engine/Source/Runtime/Renderer/Private/ShadowSetup.cpp
--- a/Engine/Source/Runtime/Renderer/Private/ShadowSetup.cpp
+++ b/Engine/Source/Runtime/Renderer/Private/ShadowSetup.cpp
@@ -44,6 +44,7 @@
 #include "ReadOnlyCVARCache.h"
 
 using namespace UE::Geometry;
+using namespace ShadowRendering;
 
 /** Number of cube map shadow depth surfaces that will be created and used for rendering one pass point light shadows. */
 static const int32 NumCubeShadowDepthSurfaces = 5;
@@ -2638,7 +2639,7 @@ void FProjectedShadowInfo::SetupMeshDrawCommandsForProjectionStenciling(FSceneRe
 						false, CF_DepthNearOrEqual,
 						true, CF_Always, SO_Keep, SO_Keep, SO_Replace,
 						false, CF_Always, SO_Keep, SO_Keep, SO_Keep,
-						0xff, 0xff
+						ShadowStencilMask, ShadowStencilMask
 						>::GetRHI());
 
 					DrawRenderState.SetStencilRef(1);
