Updated the FVolumetricRenderTargetViewStateData constructor to initialize all class members to sensible values. This is to prevent crashes in cases where these values are read with invalid data, before FVolumetricRenderTargetViewStateData::Initialize() properly sets them.

diff --git a/Engine/Source/Runtime/Renderer/Private/VolumetricRenderTarget.cpp b/Engine/Source/Runtime/Renderer/Private/VolumetricRenderTarget.cpp
--- a/Engine/Source/Runtime/Renderer/Private/VolumetricRenderTarget.cpp
+++ b/Engine/Source/Runtime/Renderer/Private/VolumetricRenderTarget.cpp
@@ -188,18 +188,24 @@ DECLARE_GPU_STAT(VolCloudComposeForVis);
 =============================================================================*/
 
 FVolumetricRenderTargetViewStateData::FVolumetricRenderTargetViewStateData()
-	: CurrentRT(1)
+	: VolumetricReconstructRTDownsampleFactor(0)
+	, VolumetricTracingRTDownsampleFactor(0)
+	, CurrentRT(1)
 	, bFirstTimeUsed(true)
 	, bHistoryValid(false)
 	, bHoldoutValid(false)
 	, PreViewExposure(1.0f)
+	, FrameId(0)
+	, NoiseFrameIndex(0)
+	, NoiseFrameIndexModPattern(0)
+	, CurrentPixelOffset(FIntPoint::ZeroValue)
 	, FullResolution(FIntPoint::ZeroValue)
 	, VolumetricReconstructRTResolution(FIntPoint::ZeroValue)
 	, VolumetricTracingRTResolution(FIntPoint::ZeroValue)
 	, VolumetricTracingViewRect(FIntPoint::ZeroValue)
+	, Mode(0)
+	, UpsamplingMode(0)
 {
-	VolumetricReconstructRTDownsampleFactor = 0;
-	VolumetricTracingRTDownsampleFactor = 0;
 	for (uint32 i = 0; i < kRenderTargetCount; ++i)
 	{
 		VolumetricReconstructViewRect[i] = FIntPoint::ZeroValue;
