Fix UserSceneTexture Post Process Material renders incorrectly when screen percentage is below 100.
1. The intermediate texture size is the absolute screen resolution by default. This aligns to the user scene texture tutorial.
2. TSR.Output is required to be writable by PS if the user scene texture is used in the SSR Input pass.

#tests Follow scene texture tutorial test, and the repro step. Both renders correctly
[RN] Fix UserSceneTexture renders incorrectly when screenpercentage is below 100.

diff --git a/Engine/Source/Runtime/Renderer/Private/PostProcess/TemporalSuperResolution.cpp b/Engine/Source/Runtime/Renderer/Private/PostProcess/TemporalSuperResolution.cpp
--- a/Engine/Source/Runtime/Renderer/Private/PostProcess/TemporalSuperResolution.cpp
+++ b/Engine/Source/Runtime/Renderer/Private/PostProcess/TemporalSuperResolution.cpp
@@ -2862,7 +2862,7 @@ FDefaultTemporalUpscaler::FOutputs AddTemporalSuperResolutionPasses(
 				OutputExtent,
 				ColorFormat,
 				FClearValueBinding::None,
-				/* InFlags = */ TexCreate_ShaderResource | TexCreate_UAV,
+				/* InFlags = */ TexCreate_ShaderResource | TexCreate_UAV | TexCreate_RenderTargetable,
 				/* NumMips = */ PassInputs.bGenerateOutputMip1 ? 2 : 1);
 			SceneColorOutputTexture = GraphBuilder.CreateTexture(OutputDesc, TEXT("TSR.Output"));
 
diff --git a/Engine/Source/Runtime/Renderer/Private/SceneTextures.cpp b/Engine/Source/Runtime/Renderer/Private/SceneTextures.cpp
--- a/Engine/Source/Runtime/Renderer/Private/SceneTextures.cpp
+++ b/Engine/Source/Runtime/Renderer/Private/SceneTextures.cpp
@@ -548,8 +548,8 @@ FRDGTextureRef FMinimalSceneTextures::FindOrAddUserSceneTexture(FRDGBuilder& Gra
 		TransientUserTexture.AllocationOrder = TransientTextures->Num();
 		TransientUserTexture.bUsed = false;
 		TransientUserTexture.ViewMask = 0;
-	
-		FIntPoint Extent = FIntPoint::DivideAndRoundUp(Config.Extent, ResolutionDivisor);
+
+		FIntPoint Extent = OutputRect.Size();
 
 		FRDGTextureDesc Desc(Config.bRequireMultiView ?
 			FRDGTextureDesc::Create2DArray(Extent, Config.ColorFormat, Config.ColorClearValue, Config.ColorCreateFlags, 2) :
