Added warning to user when a realtime capture is going to be expenssive due to volumetric cloud + sky light resolution.

#tests PC editor dx12
#rnx

diff --git a/Engine/Source/Runtime/Renderer/Private/ReflectionEnvironmentRealTimeCapture.cpp b/Engine/Source/Runtime/Renderer/Private/ReflectionEnvironmentRealTimeCapture.cpp
--- a/Engine/Source/Runtime/Renderer/Private/ReflectionEnvironmentRealTimeCapture.cpp
+++ b/Engine/Source/Runtime/Renderer/Private/ReflectionEnvironmentRealTimeCapture.cpp
@@ -55,9 +55,14 @@ static TAutoConsoleVariable<int32> CVarRealTimeReflectionCaptureDepthBuffer(
 
 static TAutoConsoleVariable<int32> CVarRealTimeReflectionCaptureResolutionOverride(
 	TEXT("r.SkyLight.RealTimeReflectionCapture.ResolutionOverride"), 0,
-	TEXT("Override the real time sky light capture resolution if >=16."),
+	TEXT("Override the real time sky light capture resolution if this CVar is greater than 16."),
 	ECVF_ReadOnly | ECVF_RenderThreadSafe); // Read only because it is not really toggable mid-update when time slicing is used
 
+// If this message on screen message is not enough to prevent TDRs, we should implement half resolution + upsampling tracing for cloud in sky light capture.
+static TAutoConsoleVariable<int32> CVarRealTimeReflectionCaptureDisableExpenssiveCaptureMessage(
+	TEXT("r.SkyLight.RealTimeReflectionCapture.DisableExpenssiveCaptureMessage"), 0,
+	TEXT("Disable the message reporting expenssive sky light capture due to high resolution with volumetric cloud tracing that could cause TDR."),
+	ECVF_RenderThreadSafe);
 
 class FDownsampleCubeFaceCS : public FGlobalShader
 {
diff --git a/Engine/Source/Runtime/Renderer/Private/SceneRendering.cpp b/Engine/Source/Runtime/Renderer/Private/SceneRendering.cpp
--- a/Engine/Source/Runtime/Renderer/Private/SceneRendering.cpp
+++ b/Engine/Source/Runtime/Renderer/Private/SceneRendering.cpp
@@ -4151,12 +4151,23 @@ void FSceneRenderer::OnRenderFinish(FRDGBuilder& GraphBuilder, FRDGTextureRef Vi
 		const bool bLightFunctionAtlasOutOfSlotWarning = LightFunctionAtlas.IsLightFunctionAtlasEnabled() ? LightFunctionAtlas.IsOutOfSlots() : false;
 
 		bool bShowWaitingSkylight = false;
+		bool bExpenssiveSkyLightRealTimeCaptureWithCloud = false;
 #if WITH_EDITOR
 		FSkyLightSceneProxy* SkyLight = Scene->SkyLight;
 		if (SkyLight && !SkyLight->bRealTimeCaptureEnabled)
 		{
 			bShowWaitingSkylight = SkyLight->bCubemapSkyLightWaitingForCubeMapTexture || SkyLight->bCaptureSkyLightWaitingForShaders || SkyLight->bCaptureSkyLightWaitingForMeshesOrTextures;
 		}
+
+		if (SkyLight && SkyLight->bRealTimeCaptureEnabled && SkyLight->CaptureCubeMapResolution >= 512 && Scene->HasVolumetricCloud())
+		{
+			static const auto* CVarSkyCloudCubeFacePerFrame = IConsoleManager::Get().FindTConsoleVariableDataInt(TEXT("r.SkyLight.RealTimeReflectionCapture.TimeSlice.SkyCloudCubeFacePerFrame"));
+			static const auto* CVarDisableExpenssiveCaptureMessage = IConsoleManager::Get().FindTConsoleVariableDataInt(TEXT("r.SkyLight.RealTimeReflectionCapture.DisableExpenssiveCaptureMessage"));
+
+			bExpenssiveSkyLightRealTimeCaptureWithCloud = 
+				CVarSkyCloudCubeFacePerFrame && CVarSkyCloudCubeFacePerFrame->GetValueOnRenderThread() == 6 
+				&& CVarDisableExpenssiveCaptureMessage && CVarDisableExpenssiveCaptureMessage->GetValueOnRenderThread() <= 0;
+		}
 #endif
 
 		FFXSystemInterface* FXInterface = Scene->GetFXSystem();
@@ -4167,7 +4178,8 @@ void FSceneRenderer::OnRenderFinish(FRDGBuilder& GraphBuilder, FRDGTextureRef Vi
 		const bool bAnyWarning = bShowPrecomputedVisibilityWarning || bShowDemotedLocalMemoryWarning || bShowGlobalClipPlaneWarning || bShowSkylightWarning || bShowPointLightWarning
 			|| bShowDFAODisabledWarning || bShowShadowedLightOverflowWarning || bShowMobileDynamicCSMWarning || bShowMobileLowQualityLightmapWarning
 			|| bMobileMissingSkyMaterial || bShowSkinCacheOOM || bSingleLayerWaterWarning || bLightFunctionAtlasOutOfSlotWarning || bShowDFDisabledWarning || bShowNoSkyAtmosphereComponentWarning || bFxDebugDraw
-			|| bLumenEnabledButHasNoDataForTracing || bLumenEnabledButDisabledForTheProject || bNaniteEnabledButNoAtomics || bNaniteEnabledButDisabledInProject || bRealTimeSkyCaptureButNothingToCapture || bShowWaitingSkylight
+			|| bLumenEnabledButHasNoDataForTracing || bLumenEnabledButDisabledForTheProject || bNaniteEnabledButNoAtomics || bNaniteEnabledButDisabledInProject 
+			|| bRealTimeSkyCaptureButNothingToCapture || bShowWaitingSkylight || bExpenssiveSkyLightRealTimeCaptureWithCloud
 			|| bShowLocalExposureDisabledWarning || bHasDelegateWarnings || bContactShadowIntensityCvarUsed || bLocalFogVolumeInSceneButProjectDisabled || bMegaLightsHasWarning
 			;
 
@@ -4218,7 +4230,7 @@ void FSceneRenderer::OnRenderFinish(FRDGBuilder& GraphBuilder, FRDGTextureRef Vi
 						bShowSkinCacheOOM, bSingleLayerWaterWarning, bLightFunctionAtlasOutOfSlotWarning, bShowNoSkyAtmosphereComponentWarning, bFxDebugDraw, FXInterface, bShowLocalExposureDisabledWarning,
 						bLumenEnabledButHasNoDataForTracing, bLumenEnabledButDisabledForTheProject, bNaniteEnabledButNoAtomics, bNaniteEnabledButDisabledInProject, 
 						bRealTimeSkyCaptureButNothingToCapture, bShowWaitingSkylight, bShowAmbientCubemapMipGenSettingsWarning, bLocalFogVolumeInSceneButProjectDisabled, bMegaLightsHasWarning,
-						bStereoView, bPrimaryStereoView, bIsInstancedStereoEnabled, bIsMultiViewportEnabled, bIsMobileMultiViewEnabled, bContactShadowIntensityCvarUsed]
+						bStereoView, bPrimaryStereoView, bIsInstancedStereoEnabled, bIsMultiViewportEnabled, bIsMobileMultiViewEnabled, bContactShadowIntensityCvarUsed, bExpenssiveSkyLightRealTimeCaptureWithCloud]
 						(FCanvas& Canvas)
 					{
 						// so it can get the screen size
@@ -4260,9 +4272,18 @@ void FSceneRenderer::OnRenderFinish(FRDGBuilder& GraphBuilder, FRDGTextureRef Vi
 							static const FText Message = NSLOCTEXT("Renderer", "SkylightNotSuppported", "PROJECT DOES NOT SUPPORT STATIONARY SKYLIGHT: ");
 							Writer.DrawLine(Message);
 						}
+						if (bExpenssiveSkyLightRealTimeCaptureWithCloud)
+						{
+							// This can happen because cloud in the real time capture are rendered at the face capture resolution
+							static const FText Message = NSLOCTEXT("Renderer", "Skylight", "The sky light resolution is too high for 6 faces captured per frame with volumetric cloud tracing (done at face resolution).\n\
+																							It could take a long time and cause a GPU TDR or crash on older GPUs.\n\
+																							Please reduce the sky light resolution to less than 512 or the r.SkyLight.RealTimeReflectionCapture.TimeSlice.SkyCloudCubeFacePerFrame to less than 6.\n\
+																							You may use r.SkyLight.RealTimeReflectionCapture.DisableExpenssiveCaptureMessage 1 to disable that message.");
+							Writer.DrawLine(Message);
+						}
 						if (bRealTimeSkyCaptureButNothingToCapture)
 						{
-							static const FText Message = NSLOCTEXT("Renderer", "Skylight", "A sky light with real-time capture enable is in the scene. It requires at least a SkyAtmosphere component, A volumetricCloud component or a mesh with a material tagged as IsSky. Otherwise it will be black.");
+							static const FText Message = NSLOCTEXT("Renderer", "Skylight", "A sky light with real-time capture enable is in the scene. It requires at least a SkyAtmosphere component, A volumetricCloud component or a mesh with a material tagged as IsSky. Otherwise it will be black");
 							Writer.DrawLine(Message);
 						}
 						if (bShowPointLightWarning)
