Add support for material custom data on a runtime virtual texture component.
The custom data is a float4 that can be set on the volume and is then routed to the material shader.
A new material expression RuntimeVirtualTextureCustomData is added to access the value.
It returns 0 when not writing to runtime virtual texture.
The intention is to support output differences according to volume.
One use case is to modify the output alpha from a material that is used in multiple volumes.
So for example we might want a different alpha blend curve for Displacement and Material runtime virtual textures that othewise share the same material runtime virtual texture output node.

diff --git a/Engine/Shaders/Private/MaterialTemplate.ush b/Engine/Shaders/Private/MaterialTemplate.ush
--- a/Engine/Shaders/Private/MaterialTemplate.ush
+++ b/Engine/Shaders/Private/MaterialTemplate.ush
@@ -32,6 +32,7 @@
 #include "/Engine/Shared/EnvironmentComponentsFlags.h"
 #include "/Engine/Private/MeshPaintTextureCommon.ush"
 #include "/Engine/Private/ColorSpace.ush"
+#include "/Engine/Private/VirtualTextureMaterial.ush"
 
 // Reduce latency of Nanite material base pass pixel shaders caused by register pressure due to the calculation of WPO
 #if PIXELSHADER && IS_NANITE_SHADING_PASS && IS_BASE_PASS
diff --git a/Engine/Shaders/Private/VirtualTextureMaterial.usf b/Engine/Shaders/Private/VirtualTextureMaterial.usf
--- a/Engine/Shaders/Private/VirtualTextureMaterial.usf
+++ b/Engine/Shaders/Private/VirtualTextureMaterial.usf
@@ -4,6 +4,7 @@
 
 #include "Common.ush"
 #include "GammaCorrectionCommon.ush"
+#include "VirtualTextureMaterial.ush"
 #include "/Engine/Generated/Material.ush"
 #include "/Engine/Generated/VertexFactory.ush"
 
@@ -113,29 +114,29 @@ void FPixelShaderInOut_MainPS(
 #endif
 
 	// Apply any fixed color.
-	BaseColor =  lerp(BaseColor, RuntimeVirtualTexturePassParameters.RuntimeVirtualTextureDebugParams.xyz, RuntimeVirtualTexturePassParameters.RuntimeVirtualTextureDebugParams.w);
-	Displacement = lerp(Displacement, RuntimeVirtualTexturePassParameters.RuntimeVirtualTextureDebugParams.x, RuntimeVirtualTexturePassParameters.RuntimeVirtualTextureDebugParams.w);
+	BaseColor =  lerp(BaseColor, RuntimeVirtualTexturePassParameters.DebugParams.xyz, RuntimeVirtualTexturePassParameters.DebugParams.w);
+	Displacement = lerp(Displacement, RuntimeVirtualTexturePassParameters.DebugParams.x, RuntimeVirtualTexturePassParameters.DebugParams.w);
 
-#if defined(OUT_BASECOLOR)
+#if OUT_BASECOLOR
 	Out.MRT[0] = float4(BaseColor, 1.f) * Opacity;
-#elif defined(OUT_BASECOLOR_NORMAL_ROUGHNESS)
+#elif OUT_BASECOLOR_NORMAL_ROUGHNESS
 	float3 PackedNormal = PackNormal(Normal);
 	//store basecolor in srgb space to improve image quality
 	Out.MRT[0] = float4(LinearToSrgb(BaseColor), 1.f) * Opacity;
 	//not enough channels for sign of normal  z
 	Out.MRT[1] = float4(PackedNormal.x, Roughness, PackedNormal.y, 1.f) * Opacity;
-#elif defined(OUT_BASECOLOR_NORMAL_SPECULAR)
+#elif OUT_BASECOLOR_NORMAL_SPECULAR
 	float3 PackedNormal = PackNormal(Normal);
 	Out.MRT[0] = float4(BaseColor, 1.f) * Opacity;
 	Out.MRT[1] = float4(PackedNormal.xy, Mask, 1.f) * Opacity;
 	Out.MRT[2] = float4(Specular, Roughness, PackedNormal.z, 1.f) * Opacity;
-#elif defined(OUT_MASK4)
+#elif OUT_MASK4
 	Out.MRT[0] = float4(Mask4.xyz, 1.f) * Opacity;
 	Out.MRT[1] = float4(Mask4.w, 0.f, 0.f, 1.f) * Opacity;
-#elif defined(OUT_WORLDHEIGHT)
-	float PackedHeight = PackWorldHeight(WorldHeight, RuntimeVirtualTexturePassParameters.RuntimeVirtualTexturePackHeight);
+#elif OUT_WORLDHEIGHT
+	float PackedHeight = PackWorldHeight(WorldHeight, RuntimeVirtualTexturePassParameters.PackHeight);
 	Out.MRT[0] = float4(PackedHeight, 0.f, 0.f, 1.f);
-#elif defined(OUT_DISPLACEMENT)
+#elif OUT_DISPLACEMENT
 	Out.MRT[0] = float4(Displacement, 0.f, 0.f, 1.f) * Opacity;
 #endif
 }
diff --git a/Engine/Shaders/Private/VirtualTextureMaterial.ush b/Engine/Shaders/Private/VirtualTextureMaterial.ush
new file mode 100644
--- /dev/null
+++ b/Engine/Shaders/Private/VirtualTextureMaterial.ush
@@ -0,0 +1,66 @@
+// Copyright Epic Games, Inc. All Rights Reserved.
+
+#pragma once
+
+#include "Common.ush"
+
+#if !defined(IS_VIRTUAL_TEXTURE_MATERIAL)
+#define IS_VIRTUAL_TEXTURE_MATERIAL 0
+#endif
+#if !defined(OUT_BASECOLOR)
+#define OUT_BASECOLOR 0
+#endif
+#if !defined(OUT_BASECOLOR_NORMAL_ROUGHNESS)
+#define OUT_BASECOLOR_NORMAL_ROUGHNESS 0
+#endif
+#if !defined(OUT_BASECOLOR_NORMAL_SPECULAR)
+#define OUT_BASECOLOR_NORMAL_SPECULAR 0
+#endif
+#if !defined(OUT_MASK4)
+#define OUT_MASK4 0
+#endif
+#if !defined(OUT_WORLDHEIGHT)
+#define OUT_WORLDHEIGHT 0
+#endif
+#if !defined(OUT_DISPLACEMENT)
+#define OUT_DISPLACEMENT 0
+#endif
+
+float4 GetRuntimeVirtualTextureCustomData()
+{
+#if IS_VIRTUAL_TEXTURE_MATERIAL
+	return RuntimeVirtualTexturePassParameters.CustomMaterialData;
+#else
+	return 0;
+#endif
+}
+
+bool IsRuntimeVirtualTextureBaseColor()
+{
+	return (IS_VIRTUAL_TEXTURE_MATERIAL && OUT_BASECOLOR);
+}
+
+bool IsRuntimeVirtualTextureBaseColorNormalRoughness()
+{
+	return (IS_VIRTUAL_TEXTURE_MATERIAL && OUT_BASECOLOR_NORMAL_ROUGHNESS);
+}
+
+bool IsRuntimeVirtualTextureBaseColorNormalSpecular()
+{
+	return (IS_VIRTUAL_TEXTURE_MATERIAL && OUT_BASECOLOR_NORMAL_SPECULAR);
+}
+
+bool IsRuntimeVirtualTextureMask4()
+{
+	return (IS_VIRTUAL_TEXTURE_MATERIAL && OUT_MASK4);
+}
+
+bool IsRuntimeVirtualTextureWorldHeight()
+{
+	return (IS_VIRTUAL_TEXTURE_MATERIAL && OUT_WORLDHEIGHT);
+}
+
+bool IsRuntimeVirtualTextureDisplacement()
+{
+	return (IS_VIRTUAL_TEXTURE_MATERIAL && OUT_DISPLACEMENT);
+}
diff --git a/Engine/Source/Editor/VirtualTexturingEditor/Private/RuntimeVirtualTextureBuildStreamingMips.cpp b/Engine/Source/Editor/VirtualTexturingEditor/Private/RuntimeVirtualTextureBuildStreamingMips.cpp
--- a/Engine/Source/Editor/VirtualTexturingEditor/Private/RuntimeVirtualTextureBuildStreamingMips.cpp
+++ b/Engine/Source/Editor/VirtualTexturingEditor/Private/RuntimeVirtualTextureBuildStreamingMips.cpp
@@ -181,6 +181,7 @@ namespace RuntimeVirtualTexture
 		const uint32 VirtualTextureSceneIndex = RuntimeVirtualTexture::GetRuntimeVirtualTextureSceneIndex_GameThread(InComponent);
 		const FTransform Transform = InComponent->GetComponentTransform();
 		const FBox Bounds = InComponent->Bounds.GetBox();
+		const FVector4f CustomMaterialData = InComponent->GetCustomMaterialData();
 		const FLinearColor FixedColor = InComponent->GetStreamingMipsFixedColor();
 
 		FVTProducerDescription VTDesc;
@@ -255,7 +256,8 @@ namespace RuntimeVirtualTexture
 					TileX, TileY,
 					TileSize, ImageSizeX, ImageSizeY, 
 					&FinalPixels,
-					FixedColor, &RenderTargetsCurrentState](FRHICommandListImmediate& RHICmdList)
+					CustomMaterialData, FixedColor,
+					&RenderTargetsCurrentState](FRHICommandListImmediate& RHICmdList)
 				{
 					const FBox2D TileBox(FVector2D(0, 0), FVector2D(TileSize, TileSize));
 					const FIntRect TileRect(0, 0, TileSize, TileSize);
@@ -273,6 +275,7 @@ namespace RuntimeVirtualTexture
 						Desc.bClearTextures = true;
 						Desc.bIsThumbnails = false;
 						Desc.FixedColor = FixedColor;
+						Desc.CustomMaterialData = CustomMaterialData;
 						Desc.NumPageDescs = 1;
 						for (int32 Layer = 0; Layer < NumLayers; Layer++)
 						{
diff --git a/Engine/Source/Editor/VirtualTexturingEditor/Private/RuntimeVirtualTextureThumbnailRenderer.cpp b/Engine/Source/Editor/VirtualTexturingEditor/Private/RuntimeVirtualTextureThumbnailRenderer.cpp
--- a/Engine/Source/Editor/VirtualTexturingEditor/Private/RuntimeVirtualTextureThumbnailRenderer.cpp
+++ b/Engine/Source/Editor/VirtualTexturingEditor/Private/RuntimeVirtualTextureThumbnailRenderer.cpp
@@ -90,6 +90,7 @@ void URuntimeVirtualTextureThumbnailRenderer::Draw(UObject* Object, int32 X, int
 	const FBox Bounds = RuntimeVirtualTextureComponent->Bounds.GetBox();
 	const uint32 VirtualTextureSceneIndex = RuntimeVirtualTexture::GetRuntimeVirtualTextureSceneIndex_GameThread(RuntimeVirtualTextureComponent);
 	const ERuntimeVirtualTextureMaterialType MaterialType = RuntimeVirtualTexture->GetMaterialType();
+	const FVector4f CustomMaterialData = RuntimeVirtualTextureComponent->GetCustomMaterialData();
 
 	FVTProducerDescription VTDesc;
 	RuntimeVirtualTexture->GetProducerDescription(VTDesc, URuntimeVirtualTexture::FInitSettings(), Transform);
@@ -98,7 +99,7 @@ void URuntimeVirtualTextureThumbnailRenderer::Draw(UObject* Object, int32 X, int
 	UE::RenderCommandPipe::FSyncScope SyncScope;
 
 	ENQUEUE_RENDER_COMMAND(BakeStreamingTextureTileCommand)(
-		[Scene, VirtualTextureSceneIndex, MaterialType, RenderTarget, DestBox, Transform, Bounds, MaxLevel](FRHICommandListImmediate& RHICmdList)
+		[Scene, VirtualTextureSceneIndex, MaterialType, RenderTarget, DestBox, Transform, Bounds, MaxLevel, CustomMaterialData](FRHICommandListImmediate& RHICmdList)
 	{
 		FMaterialRenderProxy::UpdateDeferredCachedUniformExpressions();
 
@@ -114,6 +115,7 @@ void URuntimeVirtualTextureThumbnailRenderer::Draw(UObject* Object, int32 X, int
 		Desc.bClearTextures = true;
 		Desc.bIsThumbnails = true;
 		Desc.FixedColor = FLinearColor::Transparent;
+		Desc.CustomMaterialData = CustomMaterialData;
 		Desc.NumPageDescs = 1;
 		Desc.Targets[0].Texture = RenderTarget->GetRenderTargetTexture();
 		// The target is in RTV state and should be left in RTV state (RenderPage needs the input/output state) : 
diff --git a/Engine/Source/Runtime/Engine/Classes/Components/RuntimeVirtualTextureComponent.h b/Engine/Source/Runtime/Engine/Classes/Components/RuntimeVirtualTextureComponent.h
--- a/Engine/Source/Runtime/Engine/Classes/Components/RuntimeVirtualTextureComponent.h
+++ b/Engine/Source/Runtime/Engine/Classes/Components/RuntimeVirtualTextureComponent.h
@@ -52,6 +52,10 @@ protected:
 	UPROPERTY(EditAnywhere, BlueprintReadOnly, TextExportTransient, Category = RuntimeVirtualTexture)
 	TObjectPtr<URuntimeVirtualTexture> VirtualTexture = nullptr;
 
+	/** A float4 custom value that can be read in the material that writes this virtual texture. */
+	UPROPERTY(EditAnywhere, BlueprintReadOnly, AdvancedDisplay, Category = RuntimeVirtualTexture)
+	FVector4f CustomMaterialData;
+
 	/** Per platform overrides for enabling the virtual texture. Only affects In-Game and PIE. */
 	UPROPERTY(EditAnywhere, AdvancedDisplay, Category = RuntimeVirtualTexture)
 	FPerPlatformBool EnableInGamePerPlatform;
@@ -145,6 +149,9 @@ public:
 	/** Get if the runtime virtual texture should be fully instantiated by it's render proxy. */
 	ENGINE_API bool IsEnabledInScene() const;
 
+	/** Public getter for the custom material data on this component. */
+	ENGINE_API FVector4f GetCustomMaterialData() const { return CustomMaterialData; }
+
 	/** Get if scalability settings are enabled. */
 	bool IsScalable() const { return bEnableScalability; }
 
diff --git a/Engine/Source/Runtime/Engine/Private/Components/RuntimeVirtualTextureComponent.cpp b/Engine/Source/Runtime/Engine/Private/Components/RuntimeVirtualTextureComponent.cpp
--- a/Engine/Source/Runtime/Engine/Private/Components/RuntimeVirtualTextureComponent.cpp
+++ b/Engine/Source/Runtime/Engine/Private/Components/RuntimeVirtualTextureComponent.cpp
@@ -58,6 +58,7 @@ static TAutoConsoleVariable<bool> CVarVTStreamingMipsUseAlways(
 
 URuntimeVirtualTextureComponent::URuntimeVirtualTextureComponent(const FObjectInitializer& ObjectInitializer)
 	: Super(ObjectInitializer)
+	, CustomMaterialData(FVector4f::Zero())
 	, EnableInGamePerPlatform(true)
 	, SceneProxy(nullptr)
 {
diff --git a/Engine/Source/Runtime/Engine/Private/Materials/HLSLMaterialTranslator.cpp b/Engine/Source/Runtime/Engine/Private/Materials/HLSLMaterialTranslator.cpp
--- a/Engine/Source/Runtime/Engine/Private/Materials/HLSLMaterialTranslator.cpp
+++ b/Engine/Source/Runtime/Engine/Private/Materials/HLSLMaterialTranslator.cpp
@@ -8757,6 +8757,17 @@ int32 FHLSLMaterialTranslator::VirtualTextureUnpack(int32 CodeIndex0, int32 Code
 	return CodeIndex0;
 }
 
+int32 FHLSLMaterialTranslator::VirtualTextureCustomData()
+{
+	if (bIsInRuntimeVirtualTextureOutput)
+	{
+		// Call hlsl helper function.
+		return AddInlinedCodeChunkZeroDeriv(MCT_Float4, TEXT("GetRuntimeVirtualTextureCustomData()"));
+	}
+	// Set zero if we are not writing RVT here.
+	return Constant4(0, 0, 0, 0);
+}
+
 int32 FHLSLMaterialTranslator::ExternalTexture(const FGuid& ExternalTextureGuid)
 {
 	bool bOnlyInPixelShader = GetFeatureLevel() < ERHIFeatureLevel::SM5;
diff --git a/Engine/Source/Runtime/Engine/Private/Materials/HLSLMaterialTranslator.h b/Engine/Source/Runtime/Engine/Private/Materials/HLSLMaterialTranslator.h
--- a/Engine/Source/Runtime/Engine/Private/Materials/HLSLMaterialTranslator.h
+++ b/Engine/Source/Runtime/Engine/Private/Materials/HLSLMaterialTranslator.h
@@ -1097,6 +1097,7 @@ protected:
 	virtual int32 VirtualTextureUniform(FName ParameterName, int32 TextureIndex, int32 VectorIndex, UE::Shader::EValueType Type) override;
 	virtual int32 VirtualTextureWorldToUV(int32 WorldPositionIndex, int32 P0, int32 P1, int32 P2, EPositionOrigin PositionOrigin) override;
 	virtual int32 VirtualTextureUnpack(int32 CodeIndex0, int32 CodeIndex1, int32 CodeIndex2, int32 P0, EVirtualTextureUnpackType UnpackType) override;
+	virtual int32 VirtualTextureCustomData() override;
 
 	virtual int32 ExternalTexture(const FGuid& ExternalTextureGuid) override;
 	virtual int32 ExternalTexture(UTexture* InTexture, int32& TextureReferenceIndex) override;
diff --git a/Engine/Source/Runtime/Engine/Private/Materials/MaterialExpressions.cpp b/Engine/Source/Runtime/Engine/Private/Materials/MaterialExpressions.cpp
--- a/Engine/Source/Runtime/Engine/Private/Materials/MaterialExpressions.cpp
+++ b/Engine/Source/Runtime/Engine/Private/Materials/MaterialExpressions.cpp
@@ -202,6 +202,7 @@
 #include "Materials/MaterialExpressionRotateAboutAxis.h"
 #include "Materials/MaterialExpressionRotator.h"
 #include "Materials/MaterialExpressionRound.h"
+#include "Materials/MaterialExpressionRuntimeVirtualTextureCustomData.h"
 #include "Materials/MaterialExpressionRuntimeVirtualTextureOutput.h"
 #include "Materials/MaterialExpressionRuntimeVirtualTextureReplace.h"
 #include "Materials/MaterialExpressionRuntimeVirtualTextureSample.h"
@@ -3777,6 +3778,39 @@ void UMaterialExpressionRuntimeVirtualTextureReplace::GetCaption(TArray<FString>
 
 #endif // WITH_EDITOR
 
+UMaterialExpressionRuntimeVirtualTextureCustomData::UMaterialExpressionRuntimeVirtualTextureCustomData(const FObjectInitializer& ObjectInitializer)
+	: Super(ObjectInitializer)
+{
+	// Structure to hold one-time initialization
+	struct FConstructorStatics
+	{
+		FText NAME_VirtualTexture;
+		FConstructorStatics()
+			: NAME_VirtualTexture(LOCTEXT("VirtualTexture", "VirtualTexture"))
+		{
+		}
+	};
+	static FConstructorStatics ConstructorStatics;
+
+#if WITH_EDITORONLY_DATA
+	MenuCategories.Add(ConstructorStatics.NAME_VirtualTexture);
+#endif
+}
+
+#if WITH_EDITOR
+
+int32 UMaterialExpressionRuntimeVirtualTextureCustomData::Compile(class FMaterialCompiler* Compiler, int32 OutputIndex)
+{
+	return Compiler->VirtualTextureCustomData();
+}
+
+void UMaterialExpressionRuntimeVirtualTextureCustomData::GetCaption(TArray<FString>& OutCaptions) const
+{
+	OutCaptions.Add(TEXT("RuntimeVirtualTextureCustomData"));
+}
+
+#endif // WITH_EDITOR
+
 UMaterialExpressionVirtualTextureFeatureSwitch::UMaterialExpressionVirtualTextureFeatureSwitch(const FObjectInitializer& ObjectInitializer)
 	: Super(ObjectInitializer)
 {
diff --git a/Engine/Source/Runtime/Engine/Public/MaterialCompiler.h b/Engine/Source/Runtime/Engine/Public/MaterialCompiler.h
--- a/Engine/Source/Runtime/Engine/Public/MaterialCompiler.h
+++ b/Engine/Source/Runtime/Engine/Public/MaterialCompiler.h
@@ -326,6 +326,7 @@ public:
 	UE_DEPRECATED(5.4, "Use VirtualTextureWorldToUV(int32, int32, int32, int32, EPositionOrigin) instead")
 	int32 VirtualTextureWorldToUV(int32 WorldPositionIndex, int32 P0, int32 P1, int32 P2) { return VirtualTextureWorldToUV(WorldPositionIndex, P0, P1, P2, EPositionOrigin::Absolute); }
 	virtual int32 VirtualTextureUnpack(int32 CodeIndex0, int32 CodeIndex1, int32 CodeIndex2, int32 P0, EVirtualTextureUnpackType UnpackType) = 0;
+	virtual int32 VirtualTextureCustomData() = 0;
 
 	virtual int32 ExternalTexture(const FGuid& ExternalTextureGuid) = 0;
 	virtual int32 ExternalTexture(UTexture* InTexture, int32& TextureReferenceIndex) = 0;
@@ -1024,6 +1025,7 @@ public:
 	virtual int32 VirtualTextureUniform(FName ParameterName, int32 TextureIndex, int32 VectorIndex, UE::Shader::EValueType Type) override { return Compiler->VirtualTextureUniform(ParameterName, TextureIndex, VectorIndex, Type); }
 	virtual int32 VirtualTextureWorldToUV(int32 WorldPositionIndex, int32 P0, int32 P1, int32 P2, EPositionOrigin PositionOrigin) override { return Compiler->VirtualTextureWorldToUV(WorldPositionIndex, P0, P1, P2, PositionOrigin); }
 	virtual int32 VirtualTextureUnpack(int32 CodeIndex0, int32 CodeIndex1, int32 CodeIndex2, int32 P0, EVirtualTextureUnpackType UnpackType) override { return Compiler->VirtualTextureUnpack(CodeIndex0, CodeIndex1, CodeIndex2, P0, UnpackType); }
+	virtual int32 VirtualTextureCustomData() override { return Compiler->VirtualTextureCustomData(); }
 
 	virtual int32 ExternalTexture(const FGuid& ExternalTextureGuid) override { return Compiler->ExternalTexture(ExternalTextureGuid); }
 	virtual int32 ExternalTexture(UTexture* InTexture, int32& TextureReferenceIndex) override { return Compiler->ExternalTexture(InTexture, TextureReferenceIndex); }
diff --git a/Engine/Source/Runtime/Engine/Public/Materials/MaterialExpressionRuntimeVirtualTextureCustomData.h b/Engine/Source/Runtime/Engine/Public/Materials/MaterialExpressionRuntimeVirtualTextureCustomData.h
new file mode 100644
--- /dev/null
+++ b/Engine/Source/Runtime/Engine/Public/Materials/MaterialExpressionRuntimeVirtualTextureCustomData.h
@@ -0,0 +1,20 @@
+// Copyright Epic Games, Inc. All Rights Reserved.
+
+#pragma once
+
+#include "CoreMinimal.h"
+#include "Materials/MaterialExpression.h"
+#include "MaterialExpressionRuntimeVirtualTextureCustomData.generated.h"
+
+UCLASS(collapsecategories, hidecategories=Object)
+class UMaterialExpressionRuntimeVirtualTextureCustomData : public UMaterialExpression
+{
+	GENERATED_UCLASS_BODY()
+
+	//~ Begin UMaterialExpression Interface
+#if WITH_EDITOR
+	virtual int32 Compile(class FMaterialCompiler* Compiler, int32 OutputIndex) override;
+	virtual void GetCaption(TArray<FString>& OutCaptions) const override;
+#endif
+	//~ End UMaterialExpression Interface
+};
diff --git a/Engine/Source/Runtime/Renderer/Private/VT/RuntimeVirtualTextureProducer.cpp b/Engine/Source/Runtime/Renderer/Private/VT/RuntimeVirtualTextureProducer.cpp
--- a/Engine/Source/Runtime/Renderer/Private/VT/RuntimeVirtualTextureProducer.cpp
+++ b/Engine/Source/Runtime/Renderer/Private/VT/RuntimeVirtualTextureProducer.cpp
@@ -15,7 +15,8 @@ FRuntimeVirtualTextureFinalizer::FRuntimeVirtualTextureFinalizer(
 	bool InClearTextures, 
 	FSceneInterface* InScene, 
 	FTransform const& InUVToWorld,
-	FBox const& InWorldBounds)
+	FBox const& InWorldBounds,
+	FVector4f const& InCustomMaterialData)
 	: Desc(InDesc)
 	, ProducerId(InProducerId)
 	, RuntimeVirtualTextureMask(0)
@@ -24,6 +25,7 @@ FRuntimeVirtualTextureFinalizer::FRuntimeVirtualTextureFinalizer(
 	, Scene(InScene)
 	, UVToWorld(InUVToWorld)
 	, WorldBounds(InWorldBounds)
+	, CustomMaterialData(InCustomMaterialData)
 {
 }
 
@@ -71,6 +73,7 @@ void FRuntimeVirtualTextureFinalizer::Finalize(FRDGBuilder& GraphBuilder)
 	RenderPageBatchDesc.bClearTextures = bClearTextures;
 	RenderPageBatchDesc.bIsThumbnails = false;
 	RenderPageBatchDesc.FixedColor = FLinearColor::Transparent;
+	RenderPageBatchDesc.CustomMaterialData = CustomMaterialData;
 	
 	for (int LayerIndex = 0; LayerIndex < RuntimeVirtualTexture::MaxTextureLayers; ++LayerIndex)
 	{
@@ -137,6 +140,18 @@ void FRuntimeVirtualTextureFinalizer::Finalize(FRDGBuilder& GraphBuilder)
 	Tiles.SetNumUnsafeInternal(0);
 }
 
+FRuntimeVirtualTextureProducer::FRuntimeVirtualTextureProducer(
+	FVTProducerDescription const& InDesc,
+	uint32 InProducerId,
+	ERuntimeVirtualTextureMaterialType InMaterialType,
+	bool InClearTextures,
+	FSceneInterface* InScene,
+	FTransform const& InUVToWorld,
+	FBox const& InWorldBounds)
+	: Finalizer(InDesc, InProducerId, InMaterialType, InClearTextures, InScene, InUVToWorld, InWorldBounds, FVector4f::Zero())
+{
+}
+
 FRuntimeVirtualTextureProducer::FRuntimeVirtualTextureProducer(
 	FVTProducerDescription const& InDesc, 
 	uint32 InProducerId, 
@@ -144,8 +159,9 @@ FRuntimeVirtualTextureProducer::FRuntimeVirtualTextureProducer(
 	bool InClearTextures, 
 	FSceneInterface* InScene, 
 	FTransform const& InUVToWorld,
-	FBox const& InWorldBounds)
-	: Finalizer(InDesc, InProducerId, InMaterialType, InClearTextures, InScene, InUVToWorld, InWorldBounds)
+	FBox const& InWorldBounds,
+	FVector4f const& InCustomMaterialData)
+	: Finalizer(InDesc, InProducerId, InMaterialType, InClearTextures, InScene, InUVToWorld, InWorldBounds, InCustomMaterialData)
 {
 }
 
diff --git a/Engine/Source/Runtime/Renderer/Private/VT/RuntimeVirtualTextureProducer.h b/Engine/Source/Runtime/Renderer/Private/VT/RuntimeVirtualTextureProducer.h
--- a/Engine/Source/Runtime/Renderer/Private/VT/RuntimeVirtualTextureProducer.h
+++ b/Engine/Source/Runtime/Renderer/Private/VT/RuntimeVirtualTextureProducer.h
@@ -22,7 +22,8 @@ public:
 		bool InClearTextures, 
 		FSceneInterface* InScene, 
 		FTransform const& InUVToWorld,
-		FBox const& InWorldBounds);
+		FBox const& InWorldBounds,
+		FVector4f const& InCustomMaterialData);
 
 	virtual ~FRuntimeVirtualTextureFinalizer() {}
 
@@ -64,6 +65,8 @@ private:
 	FTransform UVToWorld;
 	/** Bounds of runtime virtual texture volume in world space. */
 	FBox WorldBounds;
+	/** Custom material data for the runtime virtual texture. */
+	FVector4f CustomMaterialData;
 	/** Array of tiles in the queue to finalize. */
 	TArray<FTileEntry> Tiles;
 };
@@ -72,6 +75,7 @@ private:
 class FRuntimeVirtualTextureProducer : public IVirtualTexture
 {
 public:
+	UE_DEPRECATED(5.6, "Use constructor that takes InCustomMaterialData.")
 	RENDERER_API FRuntimeVirtualTextureProducer(
 		FVTProducerDescription const& InDesc, 
 		uint32 InProducerId, 
@@ -80,7 +84,17 @@ public:
 		FSceneInterface* InScene, 
 		FTransform const& InUVToWorld,
 		FBox const& InWorldBounds);
-	
+
+	RENDERER_API FRuntimeVirtualTextureProducer(
+		FVTProducerDescription const& InDesc,
+		uint32 InProducerId,
+		ERuntimeVirtualTextureMaterialType InMaterialType,
+		bool InClearTextures,
+		FSceneInterface* InScene,
+		FTransform const& InUVToWorld,
+		FBox const& InWorldBounds,
+		FVector4f const& InCustomMaterialData);
+
 	virtual ~FRuntimeVirtualTextureProducer() {}
 
 	//~ Begin IVirtualTexture Interface.
diff --git a/Engine/Source/Runtime/Renderer/Private/VT/RuntimeVirtualTextureRender.cpp b/Engine/Source/Runtime/Renderer/Private/VT/RuntimeVirtualTextureRender.cpp
--- a/Engine/Source/Runtime/Renderer/Private/VT/RuntimeVirtualTextureRender.cpp
+++ b/Engine/Source/Runtime/Renderer/Private/VT/RuntimeVirtualTextureRender.cpp
@@ -131,10 +131,11 @@ namespace RuntimeVirtualTexture
 	}
 
 	/** Parameters used when writing to the virtual texture. */
-	BEGIN_GLOBAL_SHADER_PARAMETER_STRUCT(FRuntimeVirtualTexturePassParameters, )
-		SHADER_PARAMETER(FVector2f, RuntimeVirtualTexturePackHeight)
-		SHADER_PARAMETER(FVector4f, RuntimeVirtualTextureDebugParams)
-	END_GLOBAL_SHADER_PARAMETER_STRUCT()
+	BEGIN_UNIFORM_BUFFER_STRUCT(FRuntimeVirtualTexturePassParameters, )
+		SHADER_PARAMETER(FVector4f, CustomMaterialData)
+		SHADER_PARAMETER(FVector4f, DebugParams)
+		SHADER_PARAMETER(FVector2f, PackHeight)
+	END_UNIFORM_BUFFER_STRUCT()
 
 	/** Uniform buffer for writing to the virtual texture. We reuse the DeferredDecals UB slot, which can't be used at the same time. This avoids the overhead of a new slot. */
 	IMPLEMENT_STATIC_UNIFORM_BUFFER_STRUCT(FRuntimeVirtualTexturePassParameters, "RuntimeVirtualTexturePassParameters", DeferredDecals);
@@ -146,7 +147,7 @@ namespace RuntimeVirtualTexture
 		BEGIN_SHADER_PARAMETER_STRUCT(FParameters, )
 			SHADER_PARAMETER_STRUCT_REF(FViewUniformShaderParameters, View)
 			SHADER_PARAMETER_RDG_UNIFORM_BUFFER(FSceneUniformParameters, Scene)
-			SHADER_PARAMETER_RDG_UNIFORM_BUFFER(FRuntimeVirtualTexturePassParameters, RuntimeVirtualTexturePassUB)
+			SHADER_PARAMETER_RDG_UNIFORM_BUFFER(FRuntimeVirtualTexturePassParameters, RuntimeVirtualTexturePassParameters)
 			SHADER_PARAMETER_STRUCT_INCLUDE(FInstanceCullingDrawParams, InstanceCullingDrawParams)
 			RDG_TEXTURE_ACCESS(PhysicalTexture0, ERHIAccess::SRVGraphics)
 			RDG_TEXTURE_ACCESS(PhysicalTexture1, ERHIAccess::SRVGraphics)
@@ -176,7 +177,6 @@ namespace RuntimeVirtualTexture
 		FShader_VirtualTextureMaterialDraw(const FMeshMaterialShaderType::CompiledShaderInitializerType& Initializer)
 			: FMeshMaterialShader(Initializer)
 		{
-			PassUniformBuffer.Bind(Initializer.ParameterMap, FRuntimeVirtualTexturePassParameters::FTypeInfo::GetStructMetadata()->GetShaderVariableName());
 		}
 	};
 
@@ -1517,7 +1517,8 @@ namespace RuntimeVirtualTexture
 		FBox2D const& UVRange,
 		uint8 vLevel,
 		uint8 MaxLevel,
-		FLinearColor const& FixedColor)
+		FLinearColor const& FixedColor, 
+		FVector4f const& CustomMaterialData)
 	{
 		RDG_EVENT_SCOPE(GraphBuilder, "VirtualTextureDynamicCache");
 		CSV_CUSTOM_STAT(VirtualTexturing, RenderedPages, 1, ECsvCustomStatOp::Accumulate);
@@ -1598,9 +1599,10 @@ namespace RuntimeVirtualTexture
 			PassParameters->Scene = SceneRenderer->GetSceneUniformBufferRef(GraphBuilder);
 
 			FRuntimeVirtualTexturePassParameters* RuntimeVirtualTexturePassParameters = GraphBuilder.AllocParameters<FRuntimeVirtualTexturePassParameters>();
-			RuntimeVirtualTexturePassParameters->RuntimeVirtualTexturePackHeight = FVector2f(WorldHeightPackParameter);	// LWC_TODO: Precision loss
-			RuntimeVirtualTexturePassParameters->RuntimeVirtualTextureDebugParams = CVarVTMipColors.GetValueOnRenderThread() ? GetMipLevelColor(vLevel) : FixedColor;
-			PassParameters->RuntimeVirtualTexturePassUB = GraphBuilder.CreateUniformBuffer(RuntimeVirtualTexturePassParameters);
+			RuntimeVirtualTexturePassParameters->CustomMaterialData = CustomMaterialData;
+			RuntimeVirtualTexturePassParameters->DebugParams = CVarVTMipColors.GetValueOnRenderThread() ? GetMipLevelColor(vLevel) : FixedColor;
+			RuntimeVirtualTexturePassParameters->PackHeight = FVector2f(WorldHeightPackParameter);	// LWC_TODO: Precision loss
+ 			PassParameters->RuntimeVirtualTexturePassParameters = GraphBuilder.CreateUniformBuffer(RuntimeVirtualTexturePassParameters);
 
 			// When direct aliasing the output physical texture for compression, the texture gets managed by RDG.
 			// The texture needs to be transitioned to UAV during the compression pass.
@@ -1716,7 +1718,8 @@ namespace RuntimeVirtualTexture
 					PageDesc.UVRange,
 					PageDesc.vLevel,
 					InDesc.MaxLevel,
-					InDesc.FixedColor);
+					InDesc.FixedColor,
+					InDesc.CustomMaterialData);
 			}
 		}
 	}
diff --git a/Engine/Source/Runtime/Renderer/Private/VT/RuntimeVirtualTextureSceneProxy.cpp b/Engine/Source/Runtime/Renderer/Private/VT/RuntimeVirtualTextureSceneProxy.cpp
--- a/Engine/Source/Runtime/Renderer/Private/VT/RuntimeVirtualTextureSceneProxy.cpp
+++ b/Engine/Source/Runtime/Renderer/Private/VT/RuntimeVirtualTextureSceneProxy.cpp
@@ -42,6 +42,7 @@ FRuntimeVirtualTextureSceneProxy::FRuntimeVirtualTextureSceneProxy(URuntimeVirtu
 			VirtualTexture = InComponent->GetVirtualTexture();
 			Transform = InComponent->GetComponentTransform();
 			const FBox Bounds = InComponent->Bounds.GetBox();
+			const FVector4f CustomMaterialData = InComponent->GetCustomMaterialData();
 
 			// The producer description is calculated using the transform to determine the aspect ratio
 			FVTProducerDescription ProducerDesc;
@@ -84,7 +85,7 @@ FRuntimeVirtualTextureSceneProxy::FRuntimeVirtualTextureSceneProxy(URuntimeVirtu
 			if (StreamingTexture == nullptr)
 			{
 				// Create the runtime virtual texture producer.
-				Producer = new FRuntimeVirtualTextureProducer(ProducerDesc, ProducerId, MaterialType, bClearTextures, SceneInterface, Transform, Bounds);
+				Producer = new FRuntimeVirtualTextureProducer(ProducerDesc, ProducerId, MaterialType, bClearTextures, SceneInterface, Transform, Bounds, CustomMaterialData);
 
 				// We only need to dirty flush up to the producer description MaxLevel which accounts for the RemoveLowMips
 				MaxDirtyLevel = ProducerDesc.MaxLevel;
@@ -111,7 +112,7 @@ FRuntimeVirtualTextureSceneProxy::FRuntimeVirtualTextureSceneProxy(URuntimeVirtu
 				}
 
 				// Create the runtime virtual texture producer.
-				Producer = new FRuntimeVirtualTextureProducer(ProducerDesc, ProducerId, MaterialType, bClearTextures, SceneInterface, Transform, Bounds);
+				Producer = new FRuntimeVirtualTextureProducer(ProducerDesc, ProducerId, MaterialType, bClearTextures, SceneInterface, Transform, Bounds, CustomMaterialData);
 
 				// Bind the runtime virtual texture producer to the streaming producer.
 				const int32 NumLevels = (int32)FMath::CeilLogTwo(FMath::Max(ProducerDesc.BlockWidthInTiles, ProducerDesc.BlockHeightInTiles));
diff --git a/Engine/Source/Runtime/Renderer/Public/VT/RuntimeVirtualTextureRender.h b/Engine/Source/Runtime/Renderer/Public/VT/RuntimeVirtualTextureRender.h
--- a/Engine/Source/Runtime/Renderer/Public/VT/RuntimeVirtualTextureRender.h
+++ b/Engine/Source/Runtime/Renderer/Public/VT/RuntimeVirtualTextureRender.h
@@ -78,6 +78,8 @@ namespace RuntimeVirtualTexture
 		bool bIsThumbnails;
 		/** Fixed BaseColor to apply. Uses alpha channel to blend with material output. */
 		FLinearColor FixedColor;
+		/** CustomData that can be read in the material. */
+		FVector4f CustomMaterialData;
 
 		/** Physical texture targets to render to. */
 		FRenderPageTarget Targets[RuntimeVirtualTexture::MaxTextureLayers];
