VolumetricCloud - Disables rendering when volumetric render target is used to trace at lower resolution but a resolution of 0 is specified on a view.
This should fix some random crashes observed recently.

#tests PC editor dx12
#rnx

diff --git a/Engine/Source/Runtime/Renderer/Private/TranslucentRendering.cpp b/Engine/Source/Runtime/Renderer/Private/TranslucentRendering.cpp
--- a/Engine/Source/Runtime/Renderer/Private/TranslucentRendering.cpp
+++ b/Engine/Source/Runtime/Renderer/Private/TranslucentRendering.cpp
@@ -1068,7 +1068,8 @@ TRDGUniformBufferRef<FTranslucentBasePassUniformParameters> CreateTranslucentBas
 		BasePassParameters.VolumetricCloudDepth = nullptr;
 		BasePassParameters.VolumetricCloudColorSampler = TStaticSamplerState<SF_Bilinear, AM_Clamp, AM_Clamp, AM_Clamp>::GetRHI();
 		BasePassParameters.VolumetricCloudDepthSampler = TStaticSamplerState<SF_Bilinear, AM_Clamp, AM_Clamp, AM_Clamp>::GetRHI();
-		if (IsVolumetricRenderTargetEnabled() && View.ViewState && ShouldRenderVolumetricCloud(Scene, View.Family->EngineShowFlags))
+		const bool bVRTValid = View.ViewState != nullptr && View.ViewState->VolumetricCloudRenderTarget.IsValid();
+		if (bVRTValid && IsVolumetricRenderTargetEnabled() && ShouldRenderVolumetricCloud(Scene, View.Family->EngineShowFlags))
 		{
 			int32 VRTMode = View.ViewState->VolumetricCloudRenderTarget.GetMode();
 			if (VRTMode == 1 || VRTMode == 3)
diff --git a/Engine/Source/Runtime/Renderer/Private/VolumetricCloudRendering.cpp b/Engine/Source/Runtime/Renderer/Private/VolumetricCloudRendering.cpp
--- a/Engine/Source/Runtime/Renderer/Private/VolumetricCloudRendering.cpp
+++ b/Engine/Source/Runtime/Renderer/Private/VolumetricCloudRendering.cpp
@@ -2679,12 +2679,16 @@ bool FSceneRenderer::RenderVolumetricCloud(
 					continue;
 				}
 
-				FLocalCloudView& ViewToProcess = ViewsToProcess[ViewsToProcess.AddUninitialized()];
-				ViewToProcess.ViewInfo = &ViewInfo;
-				ViewToProcess.bShouldViewRenderVolumetricCloudRenderTarget = bShouldViewRenderVolumetricCloudRenderTarget;
-				ViewToProcess.bEnableAerialPerspectiveSampling = bEnableAerialPerspectiveSampling;
-				ViewToProcess.bShouldUseHighQualityAerialPerspective = bShouldUseHighQualityAerialPerspective;
-				ViewToProcess.bCloudDebugViewModeEnabled = bCloudDebugViewModeEnabled;
+				const bool bVRTValid = ViewInfo.ViewState != nullptr && ViewInfo.ViewState->VolumetricCloudRenderTarget.IsValid();
+				if (!bShouldViewRenderVolumetricCloudRenderTarget || bVRTValid)
+				{
+					FLocalCloudView& ViewToProcess = ViewsToProcess[ViewsToProcess.AddUninitialized()];
+					ViewToProcess.ViewInfo = &ViewInfo;
+					ViewToProcess.bShouldViewRenderVolumetricCloudRenderTarget = bVRTValid && bShouldViewRenderVolumetricCloudRenderTarget;
+					ViewToProcess.bEnableAerialPerspectiveSampling = bEnableAerialPerspectiveSampling;
+					ViewToProcess.bShouldUseHighQualityAerialPerspective = bShouldUseHighQualityAerialPerspective;
+					ViewToProcess.bCloudDebugViewModeEnabled = bCloudDebugViewModeEnabled;
+				}
 			}
 
 			const bool bHeightFogInScene = Scene->ExponentialFogs.Num() > 0;
diff --git a/Engine/Source/Runtime/Renderer/Private/VolumetricRenderTarget.cpp b/Engine/Source/Runtime/Renderer/Private/VolumetricRenderTarget.cpp
--- a/Engine/Source/Runtime/Renderer/Private/VolumetricRenderTarget.cpp
+++ b/Engine/Source/Runtime/Renderer/Private/VolumetricRenderTarget.cpp
@@ -15,6 +15,7 @@
 #include "RendererUtils.h"
 #include "PostProcess/PostProcessing.h" // IsPostProcessingWithAlphaChannelSupported
 #include "EnvironmentComponentsFlags.h"
+#include "RendererModule.h"
 
 static TAutoConsoleVariable<int32> CVarVolumetricRenderTarget(
 	TEXT("r.VolumetricRenderTarget"), 1,
@@ -71,6 +72,11 @@ static TAutoConsoleVariable<int32> CVarImprovedLargeScaleDistanceReconstruction(
 	TEXT("Enable extra floating point precision (32-bit) for the depth buffer that is used for reconstruction. Disabled by default."),
 	ECVF_RenderThreadSafe | ECVF_Scalability);
 
+static TAutoConsoleVariable<int32> CVarVolumetricRenderTargetSimulateNullResolution(
+	TEXT("r.VolumetricRenderTarget.SimulateNullResolution"), 0,
+	TEXT("Simulate a view port resolution of 0x0 for debugging."),
+	ECVF_RenderThreadSafe);
+
 static float GetUvNoiseSampleAcceptanceWeight()
 {
 	return FMath::Max(0.0f, CVarVolumetricRenderTargetUvNoiseSampleAcceptanceWeight.GetValueOnRenderThread());
@@ -194,6 +200,7 @@ FVolumetricRenderTargetViewStateData::FVolumetricRenderTargetViewStateData()
 	, bFirstTimeUsed(true)
 	, bHistoryValid(false)
 	, bHoldoutValid(false)
+	, bValid(false)
 	, PreViewExposure(1.0f)
 	, FrameId(0)
 	, NoiseFrameIndex(0)
@@ -235,6 +242,13 @@ void FVolumetricRenderTargetViewStateData::Initialise(
 
 	bHoldoutValid = false;
 
+	bValid = TextureResolutionInCopy.X > 0 && TextureResolutionInCopy.Y > 0;
+	if (!bValid)
+	{
+		UE_LOG(LogRenderer, Warning, TEXT("Warning: A viewport of resolution 0x0 was specified - VolumetricCloud not rendered."));
+		return;
+	}
+
 	if (bFirstTimeUsed || bCameraCut)
 	{
 		bFirstTimeUsed = false;
@@ -582,6 +596,13 @@ void InitVolumetricRenderTargetForViews(FRDGBuilder& GraphBuilder, TArrayView<FV
 		FIntPoint SceneTexturesExtent = ComputeVolumetricBufferResolution(SceneTextures.Config.Extent, VRTMode);
 		FIntPoint ViewRect = ComputeVolumetricBufferResolution(ViewInfo.ViewRect.Size(), VRTMode);
 
+#if !UE_BUILD_SHIPPING
+		if (CVarVolumetricRenderTargetSimulateNullResolution.GetValueOnRenderThread() > 0.0f)
+		{
+			SceneTexturesExtent = ViewRect = FIntPoint::ZeroValue;
+		}
+#endif
+
 		VolumetricCloudRT.Initialise(	// TODO this is going to reallocate a buffer each time dynamic resolution scaling is applied 
 			SceneTexturesExtent,
 			ViewRect,
@@ -725,6 +746,11 @@ void ReconstructVolumetricRenderTarget(
 
 		FVolumetricRenderTargetViewStateData& VolumetricCloudRT = ViewInfo.ViewState->VolumetricCloudRenderTarget;
 
+		if (!VolumetricCloudRT.IsValid())
+		{
+			continue;
+		}
+
 		if (VolumetricCloudRT.GetMode() == 1 || VolumetricCloudRT.GetMode() == 3)
 		{
 			// In this case, we trace at half resolution using checker boarded min max depth.
@@ -945,6 +971,11 @@ void ComposeVolumetricRenderTargetOverScene(
 
 		FVolumetricRenderTargetViewStateData& VolumetricCloudRT = ViewInfo.ViewState->VolumetricCloudRenderTarget;
 
+		if (!VolumetricCloudRT.IsValid())
+		{
+			continue;
+		}
+
 		// When reconstructed and back buffer resolution matches, force using a pixel perfect upsampling.
 		const uint32 VRTMode = VolumetricCloudRT.GetMode();
 		int32 UpsamplingMode = GetCompositionUpsamplingMode(VRTMode, VolumetricCloudRT.GetUpsamplingMode());
@@ -1060,6 +1091,11 @@ void ComposeVolumetricRenderTargetOverSceneUnderWater(
 		const FSceneWithoutWaterTextures::FView& WaterPassViewData = WaterPassData.Views[ViewIndex];
 		FVolumetricRenderTargetViewStateData& VolumetricCloudRT = ViewInfo.ViewState->VolumetricCloudRenderTarget;
 
+		if (!VolumetricCloudRT.IsValid())
+		{
+			continue;
+		}
+
 		// When reconstructed and back buffer resolution matches, force using a pixel perfect upsampling.
 		const uint32 VRTMode = VolumetricCloudRT.GetMode();
 		int32 UpsamplingMode = GetCompositionUpsamplingMode(VRTMode, VolumetricCloudRT.GetUpsamplingMode());
@@ -1130,6 +1166,11 @@ void ComposeVolumetricRenderTargetOverSceneForVisualization(
 		}
 		FVolumetricRenderTargetViewStateData& VolumetricCloudRT = ViewInfo.ViewState->VolumetricCloudRenderTarget;
 
+		if (!VolumetricCloudRT.IsValid())
+		{
+			continue;
+		}
+
 		// When reconstructed and back buffer resolution matches, force using a pixel perfect upsampling.
 		const uint32 VRTMode = VolumetricCloudRT.GetMode();
 		int32 UpsamplingMode = GetCompositionUpsamplingMode(VRTMode, VolumetricCloudRT.GetUpsamplingMode());
diff --git a/Engine/Source/Runtime/Renderer/Private/VolumetricRenderTargetViewStateData.h b/Engine/Source/Runtime/Renderer/Private/VolumetricRenderTargetViewStateData.h
--- a/Engine/Source/Runtime/Renderer/Private/VolumetricRenderTargetViewStateData.h
+++ b/Engine/Source/Runtime/Renderer/Private/VolumetricRenderTargetViewStateData.h
@@ -61,6 +61,7 @@ public:
 	FRDGTextureRef GetOrCreateSrcVolumetricReconstructRTDepth(FRDGBuilder& GraphBuilder);
 	const FIntPoint& GetSrcVolumetricReconstructViewRect() const;
 
+	bool IsValid() const { return bValid; }
 
 	bool GetHistoryValid() const { return bHistoryValid; }
 	bool GetHoldoutValid() const { return bHoldoutValid; }
@@ -90,6 +91,7 @@ private:
 	bool bFirstTimeUsed;
 	bool bHistoryValid;
 	bool bHoldoutValid;
+	bool bValid;
 	float PreViewExposure;
 
 	int32 FrameId;
