VolumetricCloud - variable name.

#tests pc editor dx12
#rnx

diff --git a/Engine/Source/Runtime/Renderer/Private/VolumetricRenderTarget.cpp b/Engine/Source/Runtime/Renderer/Private/VolumetricRenderTarget.cpp
--- a/Engine/Source/Runtime/Renderer/Private/VolumetricRenderTarget.cpp
+++ b/Engine/Source/Runtime/Renderer/Private/VolumetricRenderTarget.cpp
@@ -67,7 +67,7 @@ static TAutoConsoleVariable<float> CVarVolumetricRenderTargetThresholdToNearFarU
 	TEXT("The distance threshold in centimeters that will trigger upsampling. Enabled when the near to far tracing distance is larger than the threshold. Use in conjunction with r.VolumetricRenderTarget.ImprovedLargeScaleDistanceReconstruction for higher quality results."),
 	ECVF_RenderThreadSafe | ECVF_Scalability);
 
-static TAutoConsoleVariable<int32> CVarImprovedLargeScaleDistanceReconstruction(
+static TAutoConsoleVariable<int32> CVarVolumetricRenderTargetImprovedLargeScaleDistanceReconstruction(
 	TEXT("r.VolumetricRenderTarget.ImprovedLargeScaleDistanceReconstruction"), 0,
 	TEXT("Enable extra floating point precision (32-bit) for the depth buffer that is used for reconstruction. Disabled by default."),
 	ECVF_RenderThreadSafe | ECVF_Scalability);
@@ -109,7 +109,7 @@ bool IsVolumetricRenderTargetAsyncCompute()
 
 bool IsImprovedLargeScaleDistanceReconstructionEnabled()
 {
-	return CVarImprovedLargeScaleDistanceReconstruction.GetValueOnRenderThread() > 0;
+	return CVarVolumetricRenderTargetImprovedLargeScaleDistanceReconstruction.GetValueOnRenderThread() > 0;
 }
 
 static bool ShouldViewComposeVolumetricRenderTarget(const FViewInfo& ViewInfo)
