[HWRT] Use RLE when assembling data to build ray tracing instance buffer.

Context:
- Currently FillRayTracingInstanceUploadBuffer(...) uses ParallelFor over scene primitives, but each thread then iterates over the instances of each primitive.
- This is problematic when using ISMs that have many instances since a single thread will loop over the instances of each ISM.
- The GPU also has to read a lot of data (one FRayTracingInstanceDescriptor per instance)
- In the most common usage, the FRayTracingInstanceDescriptor for each instance in a primitive are similar, the only difference being some members being incremented based on the instance index.
- We can use a scheme similar to RLE to reuse FRayTracingInstanceDescriptor per primitive and apply necessary adjustments for each instance in the shader.

Changes:
- Modified RayTracingBuildInstanceBufferCS to load a FRayTracingInstanceGroup per thread group which indicates whether to load one FRayTracingInstanceDescriptor per thread, or share FRayTracingInstanceDescriptor for all threads in the group and apply some patching based on thread index in the group.
- Modified BuildRayTracingSceneInitializationData(...) and FillRayTracingInstanceUploadBuffer(...) to generate "RLE packed" groups when possible (see bRLECompatible).
- In ideal circuntances, FillRayTracingInstanceUploadBuffer(...) only needs to iterate over NumTransforms/GroupSize instead of iterating over all instances. Although in current implementation we pack primitives/instances into groups aggresively so we have to write individual instance descriptors for head/tail groups.

Results:
- Project with 254 "ray tracing relevant" primitives -> 2.5M instances:
- PC (Development Editor):
- FillRayTracingInstanceUploadBuffer (CPU):
- Critical path: 17.67ms -> 0.36ms
- Total: 30.65ms -> 0.61ms
- RayTracingBuildInstanceBufferCS (GPU): 3.35ms -> 0.66ms
- Console (Test):
- FillRayTracingInstanceUploadBuffer:
- Critical path: 4.86ms -> 0.12ms
- Total: 12.92ms -> 0.30ms
- RayTracingBuildInstanceBufferCS: 1.46ms -> 0.98ms

- CitySample (2583 "ray tracing relevant" primitives -> 150K instances):
- PC (Test):
- FillRayTracingInstanceUploadBuffer (critical path): 0.11ms -> 0.05ms
- RayTracingBuildInstanceBufferCS: 0.23ms -> 0.07ms
- Console (Test):
- FillRayTracingInstanceUploadBuffer (critical path): 0.23ms -> 0.10ms
- RayTracingBuildInstanceBufferCS: 0.1ms -> 0.08ms

diff --git a/Engine/Plugins/Experimental/GPULightmass/Source/GPULightmass/Private/LightmapRenderer.cpp b/Engine/Plugins/Experimental/GPULightmass/Source/GPULightmass/Private/LightmapRenderer.cpp
--- a/Engine/Plugins/Experimental/GPULightmass/Source/GPULightmass/Private/LightmapRenderer.cpp
+++ b/Engine/Plugins/Experimental/GPULightmass/Source/GPULightmass/Private/LightmapRenderer.cpp
@@ -1190,7 +1190,18 @@ bool FSceneRenderState::SetupRayTracingScene(FRDGBuilder& GraphBuilder, FSceneUn
 				RayTracingSceneInitializationData.ReferencedGeometries.Num() * sizeof(FRayTracingAccelerationStructureAddress),
 				BUF_Volatile | BUF_MultiGPUAllocate);
 
-			const uint32 InstanceUploadBufferSize = SceneInitializer.MaxNumInstances * sizeof(FRayTracingInstanceDescriptor);
+			const uint32 NumInstanceGroups = RayTracingSceneInitializationData.NumGPUInstanceGroups + RayTracingSceneInitializationData.NumCPUInstanceGroups;
+			const uint32 InstanceGroupUploadBytes = NumInstanceGroups * sizeof(FRayTracingInstanceGroup);
+			FBufferRHIRef InstanceGroupUploadBuffer;
+			FShaderResourceViewRHIRef InstanceGroupUploadSRV;
+			{
+				FRHIResourceCreateInfo CreateInfo(TEXT("RayTracingTestBedInstanceGroupUploadBuffer"));
+				InstanceGroupUploadBuffer = RHICmdList.CreateStructuredBuffer(sizeof(FRayTracingInstanceGroup), InstanceGroupUploadBytes, BUF_ShaderResource | BUF_Volatile, CreateInfo);
+				InstanceGroupUploadSRV = RHICmdList.CreateShaderResourceView(InstanceGroupUploadBuffer, FRHIViewDesc::CreateBufferSRV().SetTypeFromBuffer(InstanceGroupUploadBuffer));
+			}
+
+			const uint32 NumInstanceDescriptors = RayTracingSceneInitializationData.NumGPUInstanceDescriptors + RayTracingSceneInitializationData.NumCPUInstanceDescriptors;
+			const uint32 InstanceUploadBufferSize = NumInstanceDescriptors * sizeof(FRayTracingInstanceDescriptor);
 			FBufferRHIRef InstanceUploadBuffer;
 			FShaderResourceViewRHIRef InstanceUploadSRV;
 			{
@@ -1209,6 +1220,7 @@ bool FSceneRenderState::SetupRayTracingScene(FRDGBuilder& GraphBuilder, FSceneUn
 			}
 
 			{
+				FRayTracingInstanceGroup* InstanceGroupUploadData = (FRayTracingInstanceGroup*)RHICmdList.LockBuffer(InstanceGroupUploadBuffer, 0, InstanceGroupUploadBytes, RLM_WriteOnly);
 				FRayTracingInstanceDescriptor* InstanceUploadData = (FRayTracingInstanceDescriptor*)RHICmdList.LockBuffer(InstanceUploadBuffer, 0, InstanceUploadBufferSize, RLM_WriteOnly);
 				FVector4f* TransformUploadData = (FVector4f*)RHICmdList.LockBuffer(TransformUploadBuffer, 0, TransformUploadBufferSize, RLM_WriteOnly);
 				FillRayTracingInstanceUploadBuffer(
@@ -1217,9 +1229,13 @@ bool FSceneRenderState::SetupRayTracingScene(FRDGBuilder& GraphBuilder, FSceneUn
 					RayTracingSceneInitializationData.InstanceGeometryIndices,
 					RayTracingSceneInitializationData.BaseUploadBufferOffsets,
 					RayTracingSceneInitializationData.BaseInstancePrefixSum,
-					RayTracingSceneInitializationData.NumNativeGPUSceneInstances,
-					RayTracingSceneInitializationData.NumNativeCPUInstances,
-					MakeArrayView(InstanceUploadData, SceneInitializer.MaxNumInstances),
+					RayTracingSceneInitializationData.InstanceGroupEntryRefs,
+					RayTracingSceneInitializationData.NumGPUInstanceGroups,
+					RayTracingSceneInitializationData.NumCPUInstanceGroups,
+					RayTracingSceneInitializationData.NumGPUInstanceDescriptors,
+					RayTracingSceneInitializationData.NumCPUInstanceDescriptors,
+					MakeArrayView(InstanceGroupUploadData, NumInstanceGroups),
+					MakeArrayView(InstanceUploadData, NumInstanceDescriptors),
 					MakeArrayView(TransformUploadData, RayTracingSceneInitializationData.NumNativeCPUInstances * 3));
 				RHICmdList.UnlockBuffer(TransformUploadBuffer);
 				RHICmdList.UnlockBuffer(InstanceUploadBuffer);
@@ -1253,11 +1269,16 @@ bool FSceneRenderState::SetupRayTracingScene(FRDGBuilder& GraphBuilder, FSceneUn
 				nullptr,
 				PreViewTranslation,
 				InstanceBuffer.UAV,
+				InstanceGroupUploadSRV,
 				InstanceUploadSRV,
 				AccelerationStructureAddressesBuffer.SRV,
 				TransformUploadSRV,
 				RayTracingSceneInitializationData.NumNativeGPUSceneInstances,
 				RayTracingSceneInitializationData.NumNativeCPUInstances,
+				RayTracingSceneInitializationData.NumGPUInstanceGroups,
+				RayTracingSceneInitializationData.NumCPUInstanceGroups,
+				RayTracingSceneInitializationData.NumGPUInstanceDescriptors,
+				RayTracingSceneInitializationData.NumCPUInstanceDescriptors,
 				nullptr,
 				nullptr,
 				nullptr);
diff --git a/Engine/Shaders/Private/RayTracing/RayTracingInstanceBufferUtil.usf b/Engine/Shaders/Private/RayTracing/RayTracingInstanceBufferUtil.usf
--- a/Engine/Shaders/Private/RayTracing/RayTracingInstanceBufferUtil.usf
+++ b/Engine/Shaders/Private/RayTracing/RayTracingInstanceBufferUtil.usf
@@ -25,6 +25,9 @@ struct FRayTracingInstanceDescriptor
 };
 
 uint NumInstances;
+uint NumGroups;
+uint NumInstanceDescriptors;
+uint BaseGroupDescriptorIndex;
 uint BaseInstanceDescriptorIndex;
 
 
@@ -35,6 +38,7 @@ float3 PreViewTranslationLow;
 
 RWStructuredBuffer<FPlatformRayTracingInstanceDescriptor> OutPlatformInstanceDescriptors;
 
+StructuredBuffer<uint> InstanceGroupDescriptors;
 StructuredBuffer<FRayTracingInstanceDescriptor> InstanceDescriptors;
 ByteAddressBuffer AccelerationStructureAddresses;
 // Transforms are float3x4 row-major matrices stored as 3 float4
@@ -110,6 +114,11 @@ bool IsTransformElementUnsafe(float Element)
 
 void ProcessInstance(FRayTracingInstanceDescriptor InputDesc)
 {
+	if (InputDesc.OutputDescriptorIndex >= NumInstances)
+	{
+		return;
+	}
+
 	const FDFVector3 PreViewTranslation = MakeDFVector3(PreViewTranslationHigh, PreViewTranslationLow);
 
 	bool bCulled = false;
@@ -217,11 +226,51 @@ void RayTracingBuildInstanceBufferCS(uint GroupThreadIndex : SV_GroupIndex, uint
 	const uint GroupIndex = GetUnWrappedDispatchGroupId(GroupId);
 	const uint Index = GetUnWrappedDispatchThreadId(GroupId, GroupThreadIndex, THREADGROUP_SIZE);
 
-	if (Index >= NumInstances)
+#if SUPPORT_INSTANCE_GROUPS
+	if (GroupIndex >= NumGroups)
+	{
+		return;
+	}
+
+	const uint PackedInstanceGroupDesc = InstanceGroupDescriptors[BaseGroupDescriptorIndex + GroupIndex];
+
+	const bool bIsRLEPacked = (PackedInstanceGroupDesc & 0x80000000) != 0;
+	const bool bIncrementUserDataPerInstance = (PackedInstanceGroupDesc & 0x40000000) != 0;
+	const uint GroupBaseInstanceDescIndex = PackedInstanceGroupDesc & 0x3FFFFFFF;
+
+	uint InstanceDescIndex = GroupBaseInstanceDescIndex;
+
+	if (!bIsRLEPacked)
+	{
+		InstanceDescIndex += GroupThreadIndex;
+	}
+
+	if (InstanceDescIndex >= NumInstanceDescriptors)
+	{
+		return;
+	}
+
+	FRayTracingInstanceDescriptor InputDesc = InstanceDescriptors[BaseInstanceDescriptorIndex + InstanceDescIndex];
+
+	if (bIsRLEPacked)
+	{
+		InputDesc.GPUSceneInstanceOrTransformIndex += GroupThreadIndex;
+		InputDesc.OutputDescriptorIndex += GroupThreadIndex;
+
+		if (bIncrementUserDataPerInstance)
+		{
+			InputDesc.InstanceId += GroupThreadIndex;
+		}
+	}
+
+	ProcessInstance(InputDesc);
+#else
+	if (Index >= NumInstanceDescriptors)
 	{
 		return;
 	}
 
 	FRayTracingInstanceDescriptor InputDesc = InstanceDescriptors[BaseInstanceDescriptorIndex + Index];
 	ProcessInstance(InputDesc);
+#endif
 }
diff --git a/Engine/Source/Runtime/Renderer/Private/RayTracing/RayTracingInstanceBufferUtil.cpp b/Engine/Source/Runtime/Renderer/Private/RayTracing/RayTracingInstanceBufferUtil.cpp
--- a/Engine/Source/Runtime/Renderer/Private/RayTracing/RayTracingInstanceBufferUtil.cpp
+++ b/Engine/Source/Runtime/Renderer/Private/RayTracing/RayTracingInstanceBufferUtil.cpp
@@ -25,8 +25,82 @@
 
 #if RHI_RAYTRACING
 
+static TAutoConsoleVariable<bool> CVarRayTracingInstanceBufferRLE(
+	TEXT("r.RayTracing.InstanceBuffer.RLE"),
+	true,
+	TEXT("Whether to use RLE to build ray tracing instance buffer."),
+	ECVF_RenderThreadSafe | ECVF_Scalability
+);
+
+static const uint32 GRayTracingInstanceGroupSize = 64;
+
+// Helper structure to assign instances to FRayTracingInstanceGroup depending on whether the primitive is compatible with RLE
+// TODO: Investigate better schemes to maximize RLE usage
+// The current implementation fills incomplete "head" group (before generating RLE groups) and might also generate a "tail" group (neither of which can use RLE since they contain instances from different primitives)
+// which means in practice only ISMs with >128 instances benefit from RLE unless they happen to end up at group boundaries.
+// An alternative approach is to allow incomplete groups instead of packing so aggressively to maximize the number of groups using RLE,
+// although that can lead to a lot of inactive threads depending on specific heuristics.
+// Primitives could also be sorted by number of instances to reduce fragmentation, etc.
+struct FGroupHelper
+{
+	uint32 CurrentGroupIndex = 0;
+	uint32 CurrentIndexInGroup = 0;
+
+	uint32 NumInstanceDescriptors = 0;
+	uint32 OptimalNumInstanceDescriptors = 0;
+
+	void AddInstances(uint32 NumInstances, bool bRLECompatible)
+	{
+		if (bRLECompatible)
+		{
+			uint32 NumInstancesRemaining = NumInstances;
+
+			if (CurrentIndexInGroup != 0)
+			{
+				// first N instances are used to fill the current (partial) group
+				const uint32 N = FMath::Min(GRayTracingInstanceGroupSize - CurrentIndexInGroup, NumInstancesRemaining);
+				NumInstancesRemaining -= N;
+
+				CurrentIndexInGroup += N;
+				CurrentGroupIndex += CurrentIndexInGroup / GRayTracingInstanceGroupSize;
+				CurrentIndexInGroup %= GRayTracingInstanceGroupSize;
+
+				NumInstanceDescriptors += N;
+			}
+
+			if (NumInstancesRemaining > 0)
+			{
+				check(CurrentIndexInGroup == 0);
+
+				// remaining instances go into packed groups + tail group
+
+				CurrentIndexInGroup += NumInstancesRemaining;
+				CurrentGroupIndex += CurrentIndexInGroup / GRayTracingInstanceGroupSize;
+				CurrentIndexInGroup %= GRayTracingInstanceGroupSize;
+
+				const uint32 NumPackedGroups = NumInstancesRemaining / GRayTracingInstanceGroupSize;
+				NumInstanceDescriptors += NumPackedGroups;
+				NumInstanceDescriptors += CurrentIndexInGroup;
+			}
+
+			OptimalNumInstanceDescriptors += FMath::DivideAndRoundUp(NumInstances, GRayTracingInstanceGroupSize);
+		}
+		else
+		{
+			NumInstanceDescriptors += NumInstances;
+			OptimalNumInstanceDescriptors += NumInstances;
+
+			CurrentIndexInGroup += NumInstances;
+			CurrentGroupIndex += CurrentIndexInGroup / GRayTracingInstanceGroupSize;
+			CurrentIndexInGroup %= GRayTracingInstanceGroupSize;
+		}
+	}
+};
+
 FRayTracingSceneInitializationData BuildRayTracingSceneInitializationData(TConstArrayView<FRayTracingGeometryInstance> Instances)
 {
+	const bool bRLEAllowed = CVarRayTracingInstanceBufferRLE.GetValueOnRenderThread();
+
 	const uint32 NumSceneInstances = Instances.Num();
 	
 	FRayTracingSceneInitializationData Output;
@@ -36,12 +110,18 @@ FRayTracingSceneInitializationData BuildRayTracingSceneInitializationData(TConst
 	Output.InstanceGeometryIndices.SetNumUninitialized(NumSceneInstances);
 	Output.BaseUploadBufferOffsets.SetNumUninitialized(NumSceneInstances);
 	Output.BaseInstancePrefixSum.SetNumUninitialized(NumSceneInstances);
+	Output.InstanceGroupEntryRefs.SetNumUninitialized(NumSceneInstances);
 	Output.PerInstanceGeometries.SetNumUninitialized(NumSceneInstances);
 
+	TArray<uint32> InstanceGroups;
+
 	Experimental::TSherwoodMap<FRHIRayTracingGeometry*, uint32> UniqueGeometries;
 
 	uint32 NumNativeInstances = 0;
 
+	FGroupHelper GPUGroupHelper;
+	FGroupHelper CPUGroupHelper;
+
 	for (uint32 InstanceIndex = 0; InstanceIndex < NumSceneInstances; ++InstanceIndex)
 	{
 		const FRayTracingGeometryInstance& InstanceDesc = Instances[InstanceIndex];
@@ -72,12 +152,12 @@ FRayTracingSceneInitializationData BuildRayTracingSceneInitializationData(TConst
 		if (bGpuSceneInstance)
 		{
 			check(InstanceDesc.Transforms.IsEmpty());
-			Output.BaseUploadBufferOffsets[InstanceIndex] = Output.NumNativeGPUSceneInstances;
+			Output.BaseUploadBufferOffsets[InstanceIndex] = GPUGroupHelper.NumInstanceDescriptors;
 			Output.NumNativeGPUSceneInstances += InstanceDesc.NumTransforms;
 		}
 		else if (bCpuInstance)
 		{
-			Output.BaseUploadBufferOffsets[InstanceIndex] = Output.NumNativeCPUInstances;
+			Output.BaseUploadBufferOffsets[InstanceIndex] = CPUGroupHelper.NumInstanceDescriptors;
 			Output.NumNativeCPUInstances += InstanceDesc.NumTransforms;
 		}
 		else
@@ -88,8 +168,33 @@ FRayTracingSceneInitializationData BuildRayTracingSceneInitializationData(TConst
 		Output.BaseInstancePrefixSum[InstanceIndex] = NumNativeInstances;
 
 		NumNativeInstances += InstanceDesc.NumTransforms;
+
+		const bool bUseUniqueUserData = InstanceDesc.UserData.Num() != 0;
+
+		if (bGpuSceneInstance)
+		{
+			Output.InstanceGroupEntryRefs[InstanceIndex].GroupIndex = GPUGroupHelper.CurrentGroupIndex;
+			Output.InstanceGroupEntryRefs[InstanceIndex].BaseIndexInGroup = GPUGroupHelper.CurrentIndexInGroup;
+
+			const bool bRLECompatible = bRLEAllowed && (InstanceDesc.BaseInstanceSceneDataOffset != -1) && !bUseUniqueUserData;
+			GPUGroupHelper.AddInstances(InstanceDesc.NumTransforms, bRLECompatible);
+		}
+		else
+		{
+			Output.InstanceGroupEntryRefs[InstanceIndex].GroupIndex = CPUGroupHelper.CurrentGroupIndex;
+			Output.InstanceGroupEntryRefs[InstanceIndex].BaseIndexInGroup = CPUGroupHelper.CurrentIndexInGroup;
+
+			const bool bRLECompatible = bRLEAllowed && !bUseUniqueUserData;
+			CPUGroupHelper.AddInstances(InstanceDesc.NumTransforms, bRLECompatible);
+		}
 	}
 
+	Output.NumGPUInstanceGroups = GPUGroupHelper.CurrentGroupIndex + (GPUGroupHelper.CurrentIndexInGroup > 0 ? 1 : 0);
+	Output.NumCPUInstanceGroups = CPUGroupHelper.CurrentGroupIndex + (CPUGroupHelper.CurrentIndexInGroup > 0 ? 1 : 0);
+
+	Output.NumGPUInstanceDescriptors = GPUGroupHelper.NumInstanceDescriptors;
+	Output.NumCPUInstanceDescriptors = CPUGroupHelper.NumInstanceDescriptors;
+
 	return MoveTemp(Output);
 }
 
@@ -157,13 +262,19 @@ void FillRayTracingInstanceUploadBuffer(
 	TConstArrayView<uint32> InstanceGeometryIndices,
 	TConstArrayView<uint32> BaseUploadBufferOffsets,
 	TConstArrayView<uint32> BaseInstancePrefixSum,
-	uint32 NumNativeGPUSceneInstances,
-	uint32 NumNativeCPUInstances,
+	TConstArrayView<FRayTracingInstanceGroupEntryRef> InstanceGroupEntryRefs,
+	uint32 NumGPUInstanceGroups,
+	uint32 NumCPUInstanceGroups,
+	uint32 NumGPUInstanceDescriptors,
+	uint32 NumCPUInstanceDescriptors,
+	TArrayView<FRayTracingInstanceGroup> OutInstanceGroupUploadData,
 	TArrayView<FRayTracingInstanceDescriptor> OutInstanceUploadData,
 	TArrayView<FVector4f> OutTransformData)
 {
 	TRACE_CPUPROFILER_EVENT_SCOPE(FillRayTracingInstanceUploadBuffer);
 
+	const bool bRLEAllowed = CVarRayTracingInstanceBufferRLE.GetValueOnRenderThread();
+
 	const int32 NumSceneInstances = Instances.Num();
 	const int32 MinBatchSize = 128;
 	ParallelFor(
@@ -171,15 +282,20 @@ void FillRayTracingInstanceUploadBuffer(
 		NumSceneInstances,
 		MinBatchSize,
 		[
+			OutInstanceGroupUploadData,
 			OutInstanceUploadData,
 			OutTransformData,
-			NumNativeGPUSceneInstances,
-			NumNativeCPUInstances,
+			NumGPUInstanceGroups,
+			NumCPUInstanceGroups,
+			NumGPUInstanceDescriptors,
+			NumCPUInstanceDescriptors,
 			Instances,
 			InstanceGeometryIndices,
 			BaseUploadBufferOffsets,
 			BaseInstancePrefixSum,
-			PreViewTranslation
+			InstanceGroupEntryRefs,
+			PreViewTranslation,
+			bRLEAllowed
 		](int32 SceneInstanceIndex)
 		{
 			const FRayTracingGeometryInstance& SceneInstance = Instances[SceneInstanceIndex];
@@ -200,17 +316,70 @@ void FillRayTracingInstanceUploadBuffer(
 			const uint32 BaseInstanceIndex = BaseInstancePrefixSum[SceneInstanceIndex];
 			const uint32 BaseTransformIndex = bCpuInstance ? BaseUploadBufferOffsets[SceneInstanceIndex] : 0;
 
+			uint32 BaseDescriptorIndex = BaseUploadBufferOffsets[SceneInstanceIndex];
+			uint32 BaseDescriptorOffset = 0;
+
+			// Upload buffer is split into 2 sections [GPUSceneInstances][CPUInstances]
+			if (!bGpuSceneInstance)
 			{
-				uint32 BaseDescriptorIndex = BaseUploadBufferOffsets[SceneInstanceIndex];
+				BaseDescriptorOffset += NumGPUInstanceDescriptors;
+			}
+
+			const bool bRLECompatible = bRLEAllowed && (!bGpuSceneInstance || SceneInstance.BaseInstanceSceneDataOffset != -1) && !bUseUniqueUserData;
 
-				// Upload buffer is split into 2 sections [GPUSceneInstances][CPUInstances]
-				if (!bGpuSceneInstance)
+			const FRayTracingInstanceGroupEntryRef& GroupEntryRef = InstanceGroupEntryRefs[SceneInstanceIndex];
+
+			uint32 GroupIndex = GroupEntryRef.GroupIndex;
+			uint32 BaseIndexInGroup = GroupEntryRef.BaseIndexInGroup;
+
+			if (bCpuInstance)
+			{
+				GroupIndex += NumGPUInstanceGroups;
+			}
+
+			uint32 TransformIndex = 0;
+
+			if (BaseIndexInGroup > 0)
+			{
+				// write N instances to fill (partial) head group
+				const uint32 N = FMath::Min(GRayTracingInstanceGroupSize - BaseIndexInGroup, NumTransforms);
+
+				for (; TransformIndex < N; ++TransformIndex)
 				{
-					BaseDescriptorIndex += NumNativeGPUSceneInstances;
+					WriteInstanceDescriptor(
+						SceneInstance,
+						TransformIndex,
+						AccelerationStructureIndex,
+						bGpuSceneInstance,
+						bUseUniqueUserData,
+						BaseInstanceIndex,
+						BaseTransformIndex,
+						OutInstanceUploadData[BaseDescriptorOffset + BaseDescriptorIndex]);
+
+					++BaseDescriptorIndex;
 				}
 
-				for (uint32 TransformIndex = 0; TransformIndex < NumTransforms; ++TransformIndex)
+				++GroupIndex;
+			}
+
+			if (bRLECompatible)
+			{
+				const uint32 NumPackedGroups = (NumTransforms - TransformIndex) / GRayTracingInstanceGroupSize;
+
+				// write packed groups
+				for (uint32 PackedGroupIndex = 0; PackedGroupIndex < NumPackedGroups; ++PackedGroupIndex)
 				{
+					// write packed group
+					FRayTracingInstanceGroup Group;
+					Group.BaseInstanceIndex = BaseDescriptorIndex;
+					Group.bIncrementUserDataPerInstance = SceneInstance.bIncrementUserDataPerInstance ? 1 : 0;
+					Group.bReuseInstance = 1;
+
+					OutInstanceGroupUploadData[GroupIndex] = Group;
+
+					++GroupIndex;
+
+					// and corresponding instance
 					WriteInstanceDescriptor(
 						SceneInstance,
 						TransformIndex,
@@ -219,16 +388,55 @@ void FillRayTracingInstanceUploadBuffer(
 						bUseUniqueUserData,
 						BaseInstanceIndex,
 						BaseTransformIndex,
-						OutInstanceUploadData[BaseDescriptorIndex + TransformIndex]);
+						OutInstanceUploadData[BaseDescriptorOffset + BaseDescriptorIndex]);
+
+					++BaseDescriptorIndex;
+
+					TransformIndex += GRayTracingInstanceGroupSize;
+				}
+			}
+
+			if (TransformIndex < NumTransforms)
+			{
+				// write tail groups (not packed)
+
+				const uint32 NumTailGroups = FMath::DivideAndRoundUp(NumTransforms - TransformIndex, GRayTracingInstanceGroupSize);
+
+				for (uint32 TailGroupIndex = 0; TailGroupIndex < NumTailGroups; ++TailGroupIndex)
+				{
+					FRayTracingInstanceGroup Group;
+					Group.BaseInstanceIndex = BaseDescriptorIndex + TailGroupIndex * GRayTracingInstanceGroupSize;
+					Group.bIncrementUserDataPerInstance = 0;
+					Group.bReuseInstance = 0;
+
+					OutInstanceGroupUploadData[GroupIndex] = Group;
+
+					++GroupIndex;
+				}
+
+				// and instances
+				for (; TransformIndex < NumTransforms; ++TransformIndex)
+				{
+					WriteInstanceDescriptor(
+						SceneInstance,
+						TransformIndex,
+						AccelerationStructureIndex,
+						bGpuSceneInstance,
+						bUseUniqueUserData,
+						BaseInstanceIndex,
+						BaseTransformIndex,
+						OutInstanceUploadData[BaseDescriptorOffset + BaseDescriptorIndex]);
+
+					++BaseDescriptorIndex;
 				}
 			}
 
 			if (bCpuInstance)
 			{
-				for (uint32 TransformIndex = 0; TransformIndex < NumTransforms; ++TransformIndex)
+				for (uint32 TransformIndex2 = 0; TransformIndex2 < NumTransforms; ++TransformIndex2)
 				{
-					const uint32 TransformDataOffset = (BaseTransformIndex + TransformIndex) * 3;
-					FMatrix LocalToTranslatedWorld = SceneInstance.Transforms[TransformIndex].ConcatTranslation(PreViewTranslation);
+					const uint32 TransformDataOffset = (BaseTransformIndex + TransformIndex2) * 3;
+					FMatrix LocalToTranslatedWorld = SceneInstance.Transforms[TransformIndex2].ConcatTranslation(PreViewTranslation);
 					const FMatrix44f LocalToTranslatedWorldF = FMatrix44f(LocalToTranslatedWorld.GetTransposed());
 					OutTransformData[TransformDataOffset + 0] = *(FVector4f*)&LocalToTranslatedWorldF.M[0];
 					OutTransformData[TransformDataOffset + 1] = *(FVector4f*)&LocalToTranslatedWorldF.M[1];
@@ -247,6 +455,8 @@ struct FRayTracingBuildInstanceBufferCS : public FGlobalShader
 		SHADER_PARAMETER_STRUCT_INCLUDE(FGPUSceneResourceParametersRHI, GPUSceneParameters)
 
 		SHADER_PARAMETER_UAV(RWStructuredBuffer, OutPlatformInstanceDescriptors)
+
+		SHADER_PARAMETER_SRV(StructuredBuffer, InstanceGroupDescriptors)
 		SHADER_PARAMETER_SRV(StructuredBuffer, InstanceDescriptors)
 		SHADER_PARAMETER_SRV(ByteAddressBuffer, AccelerationStructureAddresses)
 		SHADER_PARAMETER_SRV(StructuredBuffer, InstanceTransforms)
@@ -254,6 +464,9 @@ struct FRayTracingBuildInstanceBufferCS : public FGlobalShader
 		SHADER_PARAMETER(FVector3f, FarFieldReferencePos)
 
 		SHADER_PARAMETER(uint32, NumInstances)
+		SHADER_PARAMETER(uint32, NumGroups)
+		SHADER_PARAMETER(uint32, NumInstanceDescriptors)
+		SHADER_PARAMETER(uint32, BaseGroupDescriptorIndex)
 		SHADER_PARAMETER(uint32, BaseInstanceDescriptorIndex)
 
 		SHADER_PARAMETER(FVector3f, PreViewTranslationHigh)
@@ -272,14 +485,15 @@ struct FRayTracingBuildInstanceBufferCS : public FGlobalShader
 		SHADER_PARAMETER_UAV(RWStructuredBuffer<uint>, RWDebugInstanceGPUSceneIndices)
 	END_SHADER_PARAMETER_STRUCT()
 
+	class FSupportInstanceGroupsDim : SHADER_PERMUTATION_BOOL("SUPPORT_INSTANCE_GROUPS");
 	class FUseGPUSceneDim : SHADER_PERMUTATION_BOOL("USE_GPUSCENE");
 	class FOutputInstanceGPUSceneIndexDim : SHADER_PERMUTATION_BOOL("OUTPUT_INSTANCE_GPUSCENE_INDEX");
 	class FGpuCullingDim : SHADER_PERMUTATION_BOOL("GPU_CULLING");
 	class FOutputStatsDim : SHADER_PERMUTATION_BOOL("OUTPUT_STATS");
 	class FUseWaveOpsDim : SHADER_PERMUTATION_BOOL("USE_WAVE_OPS");
-	using FPermutationDomain = TShaderPermutationDomain<FUseGPUSceneDim, FOutputInstanceGPUSceneIndexDim, FGpuCullingDim, FOutputStatsDim, FUseWaveOpsDim>;
+	using FPermutationDomain = TShaderPermutationDomain<FSupportInstanceGroupsDim, FUseGPUSceneDim, FOutputInstanceGPUSceneIndexDim, FGpuCullingDim, FOutputStatsDim, FUseWaveOpsDim>;
 		
-	static constexpr uint32 ThreadGroupSize = 64;
+	static constexpr uint32 ThreadGroupSize = GRayTracingInstanceGroupSize;
 
 	static inline void ModifyCompilationEnvironment(const FGlobalShaderPermutationParameters& Parameters, FShaderCompilerEnvironment& OutEnvironment)
 	{
@@ -313,9 +527,13 @@ void BuildRayTracingInstanceBuffer(
 	const FGPUScene* GPUScene,
 	const FDFVector3& PreViewTranslation,
 	uint32 NumInstances,
-	uint32 BaseInstanceDescriptorIndex,
+	uint32 NumGroups,
+	uint32 NumInstanceDescriptors,
 	FRHIUnorderedAccessView* InstancesUAV,
+	FRHIShaderResourceView* InstanceGroupUploadSRV,
+	uint32 InstanceGroupUploadOffset,
 	FRHIShaderResourceView* InstanceUploadSRV,
+	uint32 InstanceUploadOffset,
 	FRHIShaderResourceView* AccelerationStructureAddressesSRV,
 	FRHIShaderResourceView* InstanceTransformSRV,
 	const FRayTracingCullingParameters* CullingParameters,
@@ -324,12 +542,16 @@ void BuildRayTracingInstanceBuffer(
 {
 	FRayTracingBuildInstanceBufferCS::FParameters PassParams;
 	PassParams.OutPlatformInstanceDescriptors = InstancesUAV;
+	PassParams.InstanceGroupDescriptors = InstanceGroupUploadSRV;
 	PassParams.InstanceDescriptors = InstanceUploadSRV;
 	PassParams.AccelerationStructureAddresses = AccelerationStructureAddressesSRV;
 	PassParams.InstanceTransforms = InstanceTransformSRV;
 	PassParams.FarFieldReferencePos = (FVector3f)Lumen::GetFarFieldReferencePos();	// LWC_TODO: Precision Loss
 	PassParams.NumInstances = NumInstances;
-	PassParams.BaseInstanceDescriptorIndex = BaseInstanceDescriptorIndex;
+	PassParams.NumGroups = NumGroups;
+	PassParams.NumInstanceDescriptors = NumInstanceDescriptors;
+	PassParams.BaseGroupDescriptorIndex = InstanceGroupUploadOffset;
+	PassParams.BaseInstanceDescriptorIndex = InstanceUploadOffset;
 	PassParams.PreViewTranslationHigh = PreViewTranslation.High;
 	PassParams.PreViewTranslationLow = PreViewTranslation.Low;
 
@@ -352,6 +574,7 @@ void BuildRayTracingInstanceBuffer(
 	PassParams.RWDebugInstanceGPUSceneIndices = DebugInstanceGPUSceneIndexUAV;
 
 	FRayTracingBuildInstanceBufferCS::FPermutationDomain PermutationVector;
+	PermutationVector.Set<FRayTracingBuildInstanceBufferCS::FSupportInstanceGroupsDim>(InstanceGroupUploadSRV != nullptr);
 	PermutationVector.Set<FRayTracingBuildInstanceBufferCS::FUseGPUSceneDim>(InstanceTransformSRV == nullptr);
 	PermutationVector.Set<FRayTracingBuildInstanceBufferCS::FOutputInstanceGPUSceneIndexDim>(DebugInstanceGPUSceneIndexUAV != nullptr);
 	PermutationVector.Set<FRayTracingBuildInstanceBufferCS::FOutputStatsDim>(OutputStatsUAV != nullptr);
@@ -360,7 +583,7 @@ void BuildRayTracingInstanceBuffer(
 
 	auto ComputeShader = GetGlobalShaderMap(GMaxRHIFeatureLevel)->GetShader<FRayTracingBuildInstanceBufferCS>(PermutationVector);
 
-	const FIntVector GroupCount = FComputeShaderUtils::GetGroupCountWrapped(PassParams.NumInstances, FRayTracingBuildInstanceBufferCS::ThreadGroupSize);
+	const FIntVector GroupCount = FComputeShaderUtils::GetGroupCountWrapped(PassParams.NumGroups);
 
 	//ClearUnusedGraphResources(ComputeShader, &PassParams);
 
@@ -378,11 +601,16 @@ void BuildRayTracingInstanceBuffer(
 	const FGPUScene* GPUScene,
 	const FDFVector3& PreViewTranslation,
 	FRHIUnorderedAccessView* InstancesUAV,
+	FRHIShaderResourceView* InstanceGroupUploadSRV,
 	FRHIShaderResourceView* InstanceUploadSRV,
 	FRHIShaderResourceView* AccelerationStructureAddressesSRV,
 	FRHIShaderResourceView* CPUInstanceTransformSRV,
 	uint32 NumNativeGPUSceneInstances,
 	uint32 NumNativeCPUInstances,
+	uint32 NumGPUGroups,
+	uint32 NumCPUGroups,
+	uint32 NumGPUInstanceDescriptors,
+	uint32 NumCPUInstanceDescriptors,
 	const FRayTracingCullingParameters* CullingParameters,
 	FRHIUnorderedAccessView* OutputStatsUAV,
 	FRHIUnorderedAccessView* DebugInstanceGPUSceneIndexUAV)
@@ -393,10 +621,14 @@ void BuildRayTracingInstanceBuffer(
 			RHICmdList,
 			GPUScene,
 			PreViewTranslation,
-			NumNativeGPUSceneInstances,
-			0,
+			NumNativeGPUSceneInstances + NumNativeCPUInstances,
+			NumGPUGroups,
+			NumGPUInstanceDescriptors,
 			InstancesUAV,
+			InstanceGroupUploadSRV,
+			0,
 			InstanceUploadSRV,
+			0,
 			AccelerationStructureAddressesSRV,
 			nullptr,
 			CullingParameters,
@@ -410,10 +642,14 @@ void BuildRayTracingInstanceBuffer(
 			RHICmdList,
 			GPUScene,
 			PreViewTranslation,
-			NumNativeCPUInstances,
-			NumNativeGPUSceneInstances, // CPU instance input descriptors are stored after GPU Scene instances
+			NumNativeGPUSceneInstances + NumNativeCPUInstances,
+			NumCPUGroups,
+			NumCPUInstanceDescriptors,
 			InstancesUAV,
+			InstanceGroupUploadSRV,
+			NumGPUGroups, // CPU instance group descriptors are stored after GPU Scene instance groups
 			InstanceUploadSRV,
+			NumGPUInstanceDescriptors, // CPU input instance descriptors are stored after GPU Scene instances
 			AccelerationStructureAddressesSRV,
 			CPUInstanceTransformSRV,
 			nullptr,
diff --git a/Engine/Source/Runtime/Renderer/Private/RayTracing/RayTracingScene.cpp b/Engine/Source/Runtime/Renderer/Private/RayTracing/RayTracingScene.cpp
--- a/Engine/Source/Runtime/Renderer/Private/RayTracing/RayTracingScene.cpp
+++ b/Engine/Source/Runtime/Renderer/Private/RayTracing/RayTracingScene.cpp
@@ -104,8 +104,15 @@ void FRayTracingScene::Create(FRDGBuilder& GraphBuilder, const FViewInfo& View,
 			Layer.RayTracingSceneRHI = RHICreateRayTracingScene(MoveTemp(Initializer));
 		}
 
+		const uint32 NumInstanceGroups = Layer.InitializationData.NumGPUInstanceGroups + Layer.InitializationData.NumCPUInstanceGroups;
+		const uint32 NumInstanceGroupsAligned = FMath::DivideAndRoundUp(FMath::Max(NumInstanceGroups, 1U), AllocationGranularity) * AllocationGranularity;
+
+		const uint32 NumInstanceDescriptors = Layer.InitializationData.NumGPUInstanceDescriptors + Layer.InitializationData.NumCPUInstanceDescriptors;
+		const uint32 NumInstanceDescriptorsAligned = FMath::DivideAndRoundUp(FMath::Max(NumInstanceDescriptors, 1U), AllocationGranularity) * AllocationGranularity;
+
 		const uint32 NumNativeInstances = Layer.InitializationData.NumNativeGPUSceneInstances + Layer.InitializationData.NumNativeCPUInstances;
 		const uint32 NumNativeInstancesAligned = FMath::DivideAndRoundUp(FMath::Max(NumNativeInstances, 1U), AllocationGranularity) * AllocationGranularity;
+
 		const uint32 NumTransformsAligned = FMath::DivideAndRoundUp(FMath::Max(Layer.InitializationData.NumNativeCPUInstances, 1U), AllocationGranularity) * AllocationGranularity;
 
 		FRayTracingAccelerationStructureSize SizeInfo = Layer.RayTracingSceneRHI->GetSizeInfo();
@@ -157,9 +164,23 @@ void FRayTracingScene::Create(FRDGBuilder& GraphBuilder, const FViewInfo& View,
 			}
 		}
 
+		{
+			// Create/resize instance group upload buffer (if necessary)
+			const uint32 UploadBufferSize = NumInstanceGroupsAligned * sizeof(FRayTracingInstanceGroup);
+
+			if (!Layer.InstanceGroupUploadBuffer.IsValid()
+				|| UploadBufferSize > Layer.InstanceGroupUploadBuffer->GetSize()
+				|| UploadBufferSize < Layer.InstanceGroupUploadBuffer->GetSize() / 2)
+			{
+				FRHIResourceCreateInfo CreateInfo(TEXT("FRayTracingScene::InstanceGroupUploadBuffer"));
+				Layer.InstanceGroupUploadBuffer = RHICmdList.CreateStructuredBuffer(sizeof(FRayTracingInstanceGroup), UploadBufferSize, BUF_ShaderResource | BUF_Volatile, CreateInfo);
+				Layer.InstanceGroupUploadSRV = RHICmdList.CreateShaderResourceView(Layer.InstanceGroupUploadBuffer, FRHIViewDesc::CreateBufferSRV().SetTypeFromBuffer(Layer.InstanceGroupUploadBuffer));
+			}
+		}
+
 		{
 			// Create/resize instance upload buffer (if necessary)
-			const uint32 UploadBufferSize = NumNativeInstancesAligned * sizeof(FRayTracingInstanceDescriptor);
+			const uint32 UploadBufferSize = NumInstanceDescriptorsAligned * sizeof(FRayTracingInstanceDescriptor);
 
 			if (!Layer.InstanceUploadBuffer.IsValid()
 				|| UploadBufferSize > Layer.InstanceUploadBuffer->GetSize()
@@ -261,23 +282,33 @@ void FRayTracingScene::Create(FRDGBuilder& GraphBuilder, const FViewInfo& View,
 		{
 			// Fill instance upload buffer on separate thread since results are only needed in RHI thread
 			GraphBuilder.AddCommandListSetupTask(
-				[NumNativeInstances,
+				[NumInstanceDescriptors,
+				NumGPUInstanceGroups = Layer.InitializationData.NumGPUInstanceGroups,
+				NumCPUInstanceGroups = Layer.InitializationData.NumCPUInstanceGroups,
 				NumNativeGPUSceneInstances = Layer.InitializationData.NumNativeGPUSceneInstances,
 				NumNativeCPUInstances = Layer.InitializationData.NumNativeCPUInstances,
+				NumGPUInstanceDescriptors = Layer.InitializationData.NumGPUInstanceDescriptors,
+				NumCPUInstanceDescriptors = Layer.InitializationData.NumCPUInstanceDescriptors,
 				Instances = MakeArrayView(Layer.Instances),
 				InstanceGeometryIndices = MakeArrayView(Layer.InitializationData.InstanceGeometryIndices),
 				BaseUploadBufferOffsets = MakeArrayView(Layer.InitializationData.BaseUploadBufferOffsets),
 				BaseInstancePrefixSum = MakeArrayView(Layer.InitializationData.BaseInstancePrefixSum),
+				InstanceGroupEntryRefs = MakeArrayView(Layer.InitializationData.InstanceGroupEntryRefs),
 				RayTracingSceneRHI = Layer.RayTracingSceneRHI,
+				InstanceGroupUploadBuffer = Layer.InstanceGroupUploadBuffer,
 				InstanceUploadBuffer = Layer.InstanceUploadBuffer,
 				TransformUploadBuffer = Layer.TransformUploadBuffer,
 				PreViewTranslation = this->PreViewTranslation](FRHICommandList& RHICmdList)
 				{
 					FTaskTagScope TaskTagScope(ETaskTag::EParallelRenderingThread);
 
-					const uint32 InstanceUploadBytes = NumNativeInstances * sizeof(FRayTracingInstanceDescriptor);
+					const uint32 NumInstanceGroups = NumGPUInstanceGroups + NumCPUInstanceGroups;
+
+					const uint32 InstanceGroupUploadBytes = NumInstanceGroups * sizeof(FRayTracingInstanceGroup);
+					const uint32 InstanceUploadBytes = NumInstanceDescriptors * sizeof(FRayTracingInstanceDescriptor);
 					const uint32 TransformUploadBytes = NumNativeCPUInstances * 3 * sizeof(FVector4f);
 
+					FRayTracingInstanceGroup* InstanceGroupUploadData = (FRayTracingInstanceGroup*)RHICmdList.LockBuffer(InstanceGroupUploadBuffer, 0, InstanceGroupUploadBytes, RLM_WriteOnly);
 					FRayTracingInstanceDescriptor* InstanceUploadData = (FRayTracingInstanceDescriptor*)RHICmdList.LockBuffer(InstanceUploadBuffer, 0, InstanceUploadBytes, RLM_WriteOnly);
 					FVector4f* TransformUploadData = (NumNativeCPUInstances > 0) ? (FVector4f*)RHICmdList.LockBuffer(TransformUploadBuffer, 0, TransformUploadBytes, RLM_WriteOnly) : nullptr;
 
@@ -287,11 +318,16 @@ void FRayTracingScene::Create(FRDGBuilder& GraphBuilder, const FViewInfo& View,
 						InstanceGeometryIndices,
 						BaseUploadBufferOffsets,
 						BaseInstancePrefixSum,
-						NumNativeGPUSceneInstances,
-						NumNativeCPUInstances,
-						MakeArrayView(InstanceUploadData, NumNativeInstances),
+						InstanceGroupEntryRefs,
+						NumGPUInstanceGroups,
+						NumCPUInstanceGroups,
+						NumGPUInstanceDescriptors,
+						NumCPUInstanceDescriptors,
+						MakeArrayView(InstanceGroupUploadData, NumInstanceGroups),
+						MakeArrayView(InstanceUploadData, NumInstanceDescriptors),
 						MakeArrayView(TransformUploadData, NumNativeCPUInstances * 3));
 
+					RHICmdList.UnlockBuffer(InstanceGroupUploadBuffer);
 					RHICmdList.UnlockBuffer(InstanceUploadBuffer);
 
 					if (NumNativeCPUInstances > 0)
@@ -345,6 +381,10 @@ void FRayTracingScene::Create(FRDGBuilder& GraphBuilder, const FViewInfo& View,
 				GPUScene,
 				NumNativeGPUSceneInstances = Layer.InitializationData.NumNativeGPUSceneInstances,
 				NumNativeCPUInstances = Layer.InitializationData.NumNativeCPUInstances,
+				NumGPUInstanceGroups = Layer.InitializationData.NumGPUInstanceGroups,
+				NumCPUInstanceGroups = Layer.InitializationData.NumCPUInstanceGroups,
+				NumGPUInstanceDescriptors = Layer.InitializationData.NumGPUInstanceDescriptors,
+				NumCPUInstanceDescriptors = Layer.InitializationData.NumCPUInstanceDescriptors,
 				CullingParameters = View.RayTracingCullingParameters
 				](FRHICommandList& RHICmdList)
 				{
@@ -353,11 +393,16 @@ void FRayTracingScene::Create(FRDGBuilder& GraphBuilder, const FViewInfo& View,
 						GPUScene,
 						FDFVector3(PreViewTranslation),
 						PassParams->InstanceBuffer->GetRHI(),
+						Layer.InstanceGroupUploadSRV,
 						Layer.InstanceUploadSRV,
 						Layer.AccelerationStructureAddressesBuffer.SRV,
 						Layer.TransformUploadSRV,
 						NumNativeGPUSceneInstances,
 						NumNativeCPUInstances,
+						NumGPUInstanceGroups,
+						NumCPUInstanceGroups,
+						NumGPUInstanceDescriptors,
+						NumCPUInstanceDescriptors,
 						&CullingParameters,
 						PassParams->OutputStats ? PassParams->OutputStats->GetRHI() : nullptr,
 						PassParams->DebugInstanceGPUSceneIndexBuffer ? PassParams->DebugInstanceGPUSceneIndexBuffer->GetRHI() : nullptr);
diff --git a/Engine/Source/Runtime/Renderer/Private/RayTracing/RayTracingScene.h b/Engine/Source/Runtime/Renderer/Private/RayTracing/RayTracingScene.h
--- a/Engine/Source/Runtime/Renderer/Private/RayTracing/RayTracingScene.h
+++ b/Engine/Source/Runtime/Renderer/Private/RayTracing/RayTracingScene.h
@@ -180,6 +180,9 @@ private:
 		FRDGBufferRef RayTracingSceneBufferRDG;
 		FRDGBufferSRVRef RayTracingSceneBufferSRV;
 
+		FBufferRHIRef InstanceGroupUploadBuffer;
+		FShaderResourceViewRHIRef InstanceGroupUploadSRV;
+
 		FBufferRHIRef InstanceUploadBuffer;
 		FShaderResourceViewRHIRef InstanceUploadSRV;
 
diff --git a/Engine/Source/Runtime/Renderer/Private/Tests/RayTracingTestbed.cpp b/Engine/Source/Runtime/Renderer/Private/Tests/RayTracingTestbed.cpp
--- a/Engine/Source/Runtime/Renderer/Private/Tests/RayTracingTestbed.cpp
+++ b/Engine/Source/Runtime/Renderer/Private/Tests/RayTracingTestbed.cpp
@@ -212,7 +212,18 @@ bool RunRayTracingTestbed_RenderThread(const FString& Parameters)
 	FByteAddressBuffer AccelerationStructureAddressesBuffer;
 	AccelerationStructureAddressesBuffer.Initialize(RHICmdList, TEXT("RayTracingTestBedAccelerationStructureAddressesBuffer"), sizeof(FRayTracingAccelerationStructureAddress), BUF_Volatile | BUF_MultiGPUAllocate);
 
-	const uint32 InstanceUploadBufferSize = SceneInitializer.MaxNumInstances * sizeof(FRayTracingInstanceDescriptor);
+	const uint32 NumInstanceGroups = RayTracingSceneInitializationData.NumGPUInstanceGroups + RayTracingSceneInitializationData.NumCPUInstanceGroups;
+	const uint32 InstanceGroupUploadBytes = NumInstanceGroups * sizeof(FRayTracingInstanceGroup);
+	FBufferRHIRef InstanceGroupUploadBuffer;
+	FShaderResourceViewRHIRef InstanceGroupUploadSRV;
+	{
+		FRHIResourceCreateInfo CreateInfo(TEXT("RayTracingTestBedInstanceGroupUploadBuffer"));
+		InstanceGroupUploadBuffer = RHICmdList.CreateStructuredBuffer(sizeof(FRayTracingInstanceGroup), InstanceGroupUploadBytes, BUF_ShaderResource | BUF_Volatile, CreateInfo);
+		InstanceGroupUploadSRV = RHICmdList.CreateShaderResourceView(InstanceGroupUploadBuffer, FRHIViewDesc::CreateBufferSRV().SetTypeFromBuffer(InstanceGroupUploadBuffer));
+	}
+
+	const uint32 NumInstanceDescriptors = RayTracingSceneInitializationData.NumGPUInstanceDescriptors + RayTracingSceneInitializationData.NumCPUInstanceDescriptors;
+	const uint32 InstanceUploadBufferSize = NumInstanceDescriptors * sizeof(FRayTracingInstanceDescriptor);
 	FBufferRHIRef InstanceUploadBuffer;
 	FShaderResourceViewRHIRef InstanceUploadSRV;
 	{
@@ -231,6 +242,7 @@ bool RunRayTracingTestbed_RenderThread(const FString& Parameters)
 	}
 
 	{
+		FRayTracingInstanceGroup* InstanceGroupUploadData = (FRayTracingInstanceGroup*)RHICmdList.LockBuffer(InstanceGroupUploadBuffer, 0, InstanceGroupUploadBytes, RLM_WriteOnly);
 		FRayTracingInstanceDescriptor* InstanceUploadData = (FRayTracingInstanceDescriptor*)RHICmdList.LockBuffer(InstanceUploadBuffer, 0, InstanceUploadBufferSize, RLM_WriteOnly);
 		FVector4f* TransformUploadData = (FVector4f*)RHICmdList.LockBuffer(TransformUploadBuffer, 0, TransformUploadBufferSize, RLM_WriteOnly);
 		FillRayTracingInstanceUploadBuffer(
@@ -239,9 +251,13 @@ bool RunRayTracingTestbed_RenderThread(const FString& Parameters)
 			RayTracingSceneInitializationData.InstanceGeometryIndices,
 			RayTracingSceneInitializationData.BaseUploadBufferOffsets,
 			RayTracingSceneInitializationData.BaseInstancePrefixSum,
-			RayTracingSceneInitializationData.NumNativeGPUSceneInstances,
-			RayTracingSceneInitializationData.NumNativeCPUInstances,
-			MakeArrayView(InstanceUploadData, SceneInitializer.MaxNumInstances),
+			RayTracingSceneInitializationData.InstanceGroupEntryRefs,
+			RayTracingSceneInitializationData.NumGPUInstanceGroups,
+			RayTracingSceneInitializationData.NumCPUInstanceGroups,
+			RayTracingSceneInitializationData.NumGPUInstanceDescriptors,
+			RayTracingSceneInitializationData.NumCPUInstanceDescriptors,
+			MakeArrayView(InstanceGroupUploadData, NumInstanceGroups),
+			MakeArrayView(InstanceUploadData, NumInstanceDescriptors),
 			MakeArrayView(TransformUploadData, RayTracingSceneInitializationData.NumNativeCPUInstances * 3));
 		RHICmdList.UnlockBuffer(TransformUploadBuffer);
 		RHICmdList.UnlockBuffer(InstanceUploadBuffer);
@@ -273,11 +289,16 @@ bool RunRayTracingTestbed_RenderThread(const FString& Parameters)
 		nullptr,
 		FDFVector3{},
 		InstanceBuffer.UAV,
+		InstanceGroupUploadSRV,
 		InstanceUploadSRV,
 		AccelerationStructureAddressesBuffer.SRV,
 		TransformUploadSRV,
 		RayTracingSceneInitializationData.NumNativeGPUSceneInstances,
 		RayTracingSceneInitializationData.NumNativeCPUInstances,
+		RayTracingSceneInitializationData.NumGPUInstanceGroups,
+		RayTracingSceneInitializationData.NumCPUInstanceGroups,
+		RayTracingSceneInitializationData.NumGPUInstanceDescriptors,
+		RayTracingSceneInitializationData.NumCPUInstanceDescriptors,
 		nullptr,
 		nullptr,
 		nullptr);
diff --git a/Engine/Source/Runtime/Renderer/Public/RayTracingInstanceBufferUtil.h b/Engine/Source/Runtime/Renderer/Public/RayTracingInstanceBufferUtil.h
--- a/Engine/Source/Runtime/Renderer/Public/RayTracingInstanceBufferUtil.h
+++ b/Engine/Source/Runtime/Renderer/Public/RayTracingInstanceBufferUtil.h
@@ -36,10 +36,27 @@ struct FRayTracingInstanceDescriptor
 	uint32 bApplyLocalBoundsTransform;
 };
 
-using FRayTracingInstanceDescriptorInput UE_DEPRECATED(5.6, "Use FRayTracingInstanceDescriptor instead.") = FRayTracingInstanceDescriptor;
+struct FRayTracingInstanceGroup
+{
+	uint32 BaseInstanceIndex : 30;
+	uint32 bIncrementUserDataPerInstance : 1;
+	uint32 bReuseInstance : 1;
+};
+
+static_assert(sizeof(FRayTracingInstanceGroup) == sizeof(uint32), "FRayTracingInstanceGroup is expected be same size as uint32.");
+
+struct FRayTracingInstanceGroupEntryRef
+{
+	uint32 GroupIndex;
+	uint32 BaseIndexInGroup;
+};
 
 struct FRayTracingSceneInitializationData
 {
+	uint32 NumGPUInstanceGroups;
+	uint32 NumCPUInstanceGroups;
+	uint32 NumGPUInstanceDescriptors;
+	uint32 NumCPUInstanceDescriptors;
 	uint32 NumNativeGPUSceneInstances;
 	uint32 NumNativeCPUInstances;
 	uint32 TotalNumSegments;
@@ -50,6 +67,8 @@ struct FRayTracingSceneInitializationData
 	TArray<uint32> BaseUploadBufferOffsets;
 	// prefix sum of `Instance.NumTransforms` for all instances in this scene
 	TArray<uint32> BaseInstancePrefixSum;
+	// reference to corresponding instance group entry
+	TArray<FRayTracingInstanceGroupEntryRef> InstanceGroupEntryRefs;
 
 	// Unique list of geometries referenced by all instances in this scene.
 	// Any referenced geometry is kept alive while the scene is alive.
@@ -68,8 +87,12 @@ RENDERER_API void FillRayTracingInstanceUploadBuffer(
 	TConstArrayView<uint32> InstanceGeometryIndices,
 	TConstArrayView<uint32> BaseUploadBufferOffsets,
 	TConstArrayView<uint32> BaseInstancePrefixSum,
-	uint32 NumNativeGPUSceneInstances,
-	uint32 NumNativeCPUInstances,
+	TConstArrayView<FRayTracingInstanceGroupEntryRef> InstanceGroupEntryRefs,
+	uint32 NumGPUInstanceGroups,
+	uint32 NumCPUInstanceGroups,
+	uint32 NumGPUInstanceDescriptors,
+	uint32 NumCPUInstanceDescriptors,
+	TArrayView<FRayTracingInstanceGroup> OutInstanceGroupUploadData,
 	TArrayView<FRayTracingInstanceDescriptor> OutInstanceUploadData,
 	TArrayView<FVector4f> OutTransformData);
 
@@ -92,8 +115,12 @@ inline void FillRayTracingInstanceUploadBuffer(
 		InstanceGeometryIndices,
 		BaseUploadBufferOffsets,
 		BaseInstancePrefixSum,
+		{},
+		0,
+		0,
 		NumNativeGPUSceneInstances,
 		NumNativeCPUInstances,
+		{},
 		OutInstanceUploadData,
 		OutTransformData);
 }
@@ -103,13 +130,53 @@ RENDERER_API void BuildRayTracingInstanceBuffer(
 	const FGPUScene* GPUScene,
 	const FDFVector3& PreViewTranslation,
 	FRHIUnorderedAccessView* InstancesUAV,
+	FRHIShaderResourceView* InstanceGroupUploadSRV,
 	FRHIShaderResourceView* InstanceUploadSRV,
 	FRHIShaderResourceView* AccelerationStructureAddressesSRV,
 	FRHIShaderResourceView* CPUInstanceTransformSRV,
 	uint32 NumNativeGPUSceneInstances,
 	uint32 NumNativeCPUInstances,
+	uint32 NumGPUGroups,
+	uint32 NumCPUGroups,
+	uint32 NumGPUInstanceDescriptors,
+	uint32 NumCPUInstanceDescriptors,
 	const FRayTracingCullingParameters* CullingParameters,
 	FRHIUnorderedAccessView* OutputStatsUAV,
 	FRHIUnorderedAccessView* DebugInstanceGPUSceneIndexUAV);
 
+UE_DEPRECATED(5.6, "Need to provide InstanceGroupUploadSRV.")
+inline void BuildRayTracingInstanceBuffer(
+	FRHICommandList& RHICmdList,
+	const FGPUScene* GPUScene,
+	const FDFVector3& PreViewTranslation,
+	FRHIUnorderedAccessView* InstancesUAV,
+	FRHIShaderResourceView* InstanceUploadSRV,
+	FRHIShaderResourceView* AccelerationStructureAddressesSRV,
+	FRHIShaderResourceView* CPUInstanceTransformSRV,
+	uint32 NumNativeGPUSceneInstances,
+	uint32 NumNativeCPUInstances,
+	const FRayTracingCullingParameters* CullingParameters,
+	FRHIUnorderedAccessView* OutputStatsUAV,
+	FRHIUnorderedAccessView* DebugInstanceGPUSceneIndexUAV)
+{
+	BuildRayTracingInstanceBuffer(
+		RHICmdList,
+		GPUScene,
+		PreViewTranslation,
+		InstancesUAV,
+		nullptr,
+		InstanceUploadSRV,
+		AccelerationStructureAddressesSRV,
+		CPUInstanceTransformSRV,
+		NumNativeGPUSceneInstances,
+		NumNativeCPUInstances,
+		0,
+		0,
+		0,
+		0,
+		CullingParameters,
+		OutputStatsUAV,
+		DebugInstanceGPUSceneIndexUAV);
+}
+
 #endif // RHI_RAYTRACING
\ No newline at end of file
