VolumetricCloud - handle volumetric render target null viewport from application start point (no crash in this case).

#tests PC editor dx12
#rnx

diff --git a/Engine/Source/Runtime/Renderer/Private/VolumetricRenderTarget.cpp b/Engine/Source/Runtime/Renderer/Private/VolumetricRenderTarget.cpp
--- a/Engine/Source/Runtime/Renderer/Private/VolumetricRenderTarget.cpp
+++ b/Engine/Source/Runtime/Renderer/Private/VolumetricRenderTarget.cpp
@@ -610,6 +610,11 @@ void InitVolumetricRenderTargetForViews(FRDGBuilder& GraphBuilder, TArrayView<FV
 			CVarVolumetricRenderTargetUpsamplingMode.GetValueOnAnyThread(),
 			bCameraCut);
 
+		if (!VolumetricCloudRT.IsValid())
+		{
+			continue;
+		}
+
 		FViewUniformShaderParameters ViewVolumetricCloudRTParameters = *ViewInfo.CachedViewUniformShaderParameters;
 		{
 			const FIntPoint& VolumetricTracingResolution = VolumetricCloudRT.GetCurrentVolumetricTracingRTResolution();
