Fix missing indirect args access on shader parameter declaration, causing passes to be incorrectly culled.

diff --git a/Engine/Source/Runtime/Renderer/Private/HairStrands/RenderCurveRaster.cpp b/Engine/Source/Runtime/Renderer/Private/HairStrands/RenderCurveRaster.cpp
--- a/Engine/Source/Runtime/Renderer/Private/HairStrands/RenderCurveRaster.cpp
+++ b/Engine/Source/Runtime/Renderer/Private/HairStrands/RenderCurveRaster.cpp
@@ -179,6 +179,7 @@ class FRenderCurveClusterCullingCS : public FGlobalShader
 		SHADER_PARAMETER_RDG_BUFFER_UAV(RWBuffer<uint>, RWVisibleClusterArgs)
 		SHADER_PARAMETER_RDG_BUFFER_UAV(RWStructuredBuffer<uint2>, RWVisibleClusters)
 		SHADER_PARAMETER_RDG_BUFFER_UAV(RWStructuredBuffer<uint>, RWMinMaxZ)
+		RDG_BUFFER_ACCESS(VisibleInstanceIndirectArgs, ERHIAccess::IndirectArgs)
 	END_SHADER_PARAMETER_STRUCT()
 
 public:
@@ -639,6 +640,7 @@ void AddRenderCurveRasterPipeline(
 		Parameters->RWVisibleClusterArgs = RWVisibleClusterArgs;
 		Parameters->RWVisibleClusters = RWVisibleClusters;
 		Parameters->RWMinMaxZ = TransientData.RWMinMaxZ;
+		Parameters->VisibleInstanceIndirectArgs = TransientData.VisibleInstanceArgs;
 
 		FComputeShaderUtils::AddPass(GraphBuilder, RDG_EVENT_NAME("RenderCurve::ClusterCulling"), ComputeShaderRaster, Parameters, TransientData.VisibleInstanceArgs, 0u);
 	}
