Move runtime virtual texture mip level uniform out of the view uniform buffer and into the specific runtime virtual texture uniform buffer.
We added the new runtime virtual texture uniform buffer so that we can add new params without continuing to pollute the view uniform buffer.
This is last of the existing uniforms that were still in the view uniform buffer that would be better in the runtime virtual texture uniform buffer.

diff --git a/Engine/Shaders/Private/VirtualTextureMaterial.ush b/Engine/Shaders/Private/VirtualTextureMaterial.ush
--- a/Engine/Shaders/Private/VirtualTextureMaterial.ush
+++ b/Engine/Shaders/Private/VirtualTextureMaterial.ush
@@ -1,5 +1,7 @@
 // Copyright Epic Games, Inc. All Rights Reserved.
 
+// Expose public values from the runtime virtual texture output pass to the material shaders.
+
 #pragma once
 
 #include "Common.ush"
@@ -26,6 +28,15 @@
 #define OUT_DISPLACEMENT 0
 #endif
 
+float4 GetRuntimeVirtualTextureMipLevel()
+{
+#if IS_VIRTUAL_TEXTURE_MATERIAL
+	return RuntimeVirtualTexturePassParameters.MipLevel;
+#else
+	return 0;
+#endif
+}
+
 float4 GetRuntimeVirtualTextureCustomData()
 {
 #if IS_VIRTUAL_TEXTURE_MATERIAL
diff --git a/Engine/Source/Runtime/Core/Private/UObject/FortniteMainBranchObjectVersion.cpp b/Engine/Source/Runtime/Core/Private/UObject/FortniteMainBranchObjectVersion.cpp
--- a/Engine/Source/Runtime/Core/Private/UObject/FortniteMainBranchObjectVersion.cpp
+++ b/Engine/Source/Runtime/Core/Private/UObject/FortniteMainBranchObjectVersion.cpp
@@ -14,6 +14,6 @@ TMap<FGuid, FGuid> FFortniteMainBranchObjectVersion::GetSystemGuids()
 	SystemGuids.Add(DevGuids.Niagara_LatestScriptCompileVersion, FGuid("393081597FA7234C92F00869DE262D73"));
 	SystemGuids.Add(DevGuids.SkeletalMeshDerivedDataVersion, FGuid("2F90180769024483BB7D90AACEB2798D"));
 	SystemGuids.Add(DevGuids.STATICMESH_DERIVEDDATA_VER, FGuid("656E338C78C34641A5D3BFAE303997AF"));
-	SystemGuids.Add(DevGuids.MaterialTranslationDDCVersion, FGuid("0B923787D28841B8A53F4B28C8063EBA"));
+	SystemGuids.Add(DevGuids.MaterialTranslationDDCVersion, FGuid("2AAE29CCC8BB4553A3CEB289E14D75E7"));
 	return SystemGuids;
 }
diff --git a/Engine/Source/Runtime/Engine/Private/Materials/HLSLMaterialTranslator.cpp b/Engine/Source/Runtime/Engine/Private/Materials/HLSLMaterialTranslator.cpp
--- a/Engine/Source/Runtime/Engine/Private/Materials/HLSLMaterialTranslator.cpp
+++ b/Engine/Source/Runtime/Engine/Private/Materials/HLSLMaterialTranslator.cpp
@@ -5333,10 +5333,10 @@ int32 FHLSLMaterialTranslator::ViewProperty(EMaterialExposedViewProperty Propert
 		{MEVP_TemporalSampleCount, MCT_Float1, TEXT("View.TemporalAAParams.y"), nullptr},
 		{MEVP_TemporalSampleIndex, MCT_Float1, TEXT("View.TemporalAAParams.x"), nullptr},
 		{MEVP_TemporalSampleOffset, MCT_Float2, TEXT("View.TemporalAAParams.zw"), nullptr},
-		{MEVP_RuntimeVirtualTextureOutputLevel, MCT_Float1, TEXT("View.RuntimeVirtualTextureMipLevel.x"), nullptr},
-		{MEVP_RuntimeVirtualTextureOutputDerivative, MCT_Float2, TEXT("View.RuntimeVirtualTextureMipLevel.zw"), nullptr},
+		{MEVP_RuntimeVirtualTextureOutputLevel, MCT_Float1, TEXT("GetRuntimeVirtualTextureMipLevel().x"), nullptr},
+		{MEVP_RuntimeVirtualTextureOutputDerivative, MCT_Float2, TEXT("GetRuntimeVirtualTextureMipLevel().zw"), nullptr},
 		{MEVP_PreExposure, MCT_Float1, TEXT("View.PreExposure.x"), TEXT("View.OneOverPreExposure.x")},
-		{MEVP_RuntimeVirtualTextureMaxLevel, MCT_Float1, TEXT("View.RuntimeVirtualTextureMipLevel.y"), nullptr},
+		{MEVP_RuntimeVirtualTextureMaxLevel, MCT_Float1, TEXT("GetRuntimeVirtualTextureMipLevel().y"), nullptr},
 		{MEVP_ResolutionFraction, MCT_Float1, TEXT("View.ResolutionFractionAndInv.x"), TEXT("View.ResolutionFractionAndInv.y")},
 		{MEVP_PostVolumeUserFlags, MCT_Float1, TEXT("View.PostVolumeUserFlags"), nullptr},
 		{MEVP_FirstPersonFieldOfView, MCT_Float2, TEXT("View.<PREV>FirstPersonFieldOfViewWideAngles"), nullptr},
diff --git a/Engine/Source/Runtime/Engine/Public/Materials/MaterialExpressionViewProperty.h b/Engine/Source/Runtime/Engine/Public/Materials/MaterialExpressionViewProperty.h
--- a/Engine/Source/Runtime/Engine/Public/Materials/MaterialExpressionViewProperty.h
+++ b/Engine/Source/Runtime/Engine/Public/Materials/MaterialExpressionViewProperty.h
@@ -32,13 +32,13 @@ enum EMaterialExposedViewProperty : int
 	/** Offset of the temporal sample for this frame in pixel size. */
 	MEVP_TemporalSampleOffset UMETA(DisplayName = "Temporal Sample Offset"),
 	/** Mip Level that Runtime Virtual Texture Output is rendering to. */
-	MEVP_RuntimeVirtualTextureOutputLevel UMETA(DisplayName = "Virtual Texture Output Level"),
+	MEVP_RuntimeVirtualTextureOutputLevel UMETA(DisplayName = "Runtime Virtual Texture Output Level"),
 	/** World space derivatives for Runtime Virtual Texture Output. */
-	MEVP_RuntimeVirtualTextureOutputDerivative UMETA(DisplayName = "Virtual Texture Output Derivative"),
+	MEVP_RuntimeVirtualTextureOutputDerivative UMETA(DisplayName = "Runtime Virtual Texture Output Derivative"),
 	/** Pre Exposure */
 	MEVP_PreExposure UMETA(DisplayName = "Pre-Exposure"),
 	/** Maximum mip level of Runtime Virtual Texture that Runtime Virtual Texture Output is rendering to. */
-	MEVP_RuntimeVirtualTextureMaxLevel UMETA(DisplayName = "Virtual Texture Max Level"),
+	MEVP_RuntimeVirtualTextureMaxLevel UMETA(DisplayName = "Runtime Virtual Texture Max Level"),
 	/** Screen percentage at which the rendering resolution happens, to allow tech-art to remain consistent with dynamic resolution. */
 	MEVP_ResolutionFraction UMETA(DisplayName = "ScreenPercentage / 100"),
 	/** Post process volume user flags, useful for varying the behavior of a material per view */
diff --git a/Engine/Source/Runtime/Engine/Public/SceneView.h b/Engine/Source/Runtime/Engine/Public/SceneView.h
--- a/Engine/Source/Runtime/Engine/Public/SceneView.h
+++ b/Engine/Source/Runtime/Engine/Public/SceneView.h
@@ -1040,7 +1040,6 @@ enum ETranslucencyVolumeCascade
 	VIEW_UNIFORM_BUFFER_MEMBER(uint32, VirtualTextureFeedbackStride) \
 	VIEW_UNIFORM_BUFFER_MEMBER(uint32, VirtualTextureFeedbackJitterOffset) \
 	VIEW_UNIFORM_BUFFER_MEMBER(uint32, VirtualTextureFeedbackSampleOffset) \
-	VIEW_UNIFORM_BUFFER_MEMBER(FVector4f, RuntimeVirtualTextureMipLevel) \
 	VIEW_UNIFORM_BUFFER_MEMBER(int32, FarShadowStaticMeshLODBias) \
 	VIEW_UNIFORM_BUFFER_MEMBER(float, MinRoughness) \
 	VIEW_UNIFORM_BUFFER_MEMBER(FVector4f, HairRenderInfo) \
diff --git a/Engine/Source/Runtime/Renderer/Private/SceneRendering.cpp b/Engine/Source/Runtime/Renderer/Private/SceneRendering.cpp
--- a/Engine/Source/Runtime/Renderer/Private/SceneRendering.cpp
+++ b/Engine/Source/Runtime/Renderer/Private/SceneRendering.cpp
@@ -2013,8 +2013,6 @@ void FViewInfo::SetupUniformBufferParameters(
 	const uint32 NumPixelsInTile = FMath::Square(VirtualTextureFeedbackScale);
 	ViewUniformShaderParameters.VirtualTextureFeedbackSampleOffset = (FrameIndex % NumPixelsInTile) + (FrameIndex / NumPixelsInTile);
 
-	ViewUniformShaderParameters.RuntimeVirtualTextureMipLevel = FVector4f(ForceInitToZero);
-
 	// GGX/Sheen LTC (used as BSDF or for rect light integration)
 	if (GSystemTextures.GGXLTCMat.IsValid() && GSystemTextures.GGXLTCAmp.IsValid())
 	{
diff --git a/Engine/Source/Runtime/Renderer/Private/VT/RuntimeVirtualTextureRender.cpp b/Engine/Source/Runtime/Renderer/Private/VT/RuntimeVirtualTextureRender.cpp
--- a/Engine/Source/Runtime/Renderer/Private/VT/RuntimeVirtualTextureRender.cpp
+++ b/Engine/Source/Runtime/Renderer/Private/VT/RuntimeVirtualTextureRender.cpp
@@ -132,6 +132,7 @@ namespace RuntimeVirtualTexture
 
 	/** Parameters used when writing to the virtual texture. */
 	BEGIN_UNIFORM_BUFFER_STRUCT(FRuntimeVirtualTexturePassParameters, )
+		SHADER_PARAMETER(FVector4f, MipLevel)
 		SHADER_PARAMETER(FVector4f, CustomMaterialData)
 		SHADER_PARAMETER(FVector4f, DebugParams)
 		SHADER_PARAMETER(FVector2f, PackHeight)
@@ -1575,7 +1576,6 @@ namespace RuntimeVirtualTexture
 		View->ViewRect = View->UnconstrainedViewRect;
 		View->CachedViewUniformShaderParameters = MakeUnique<FViewUniformShaderParameters>();
 		View->SetupUniformBufferParameters(nullptr, 0, *View->CachedViewUniformShaderParameters);
-		View->CachedViewUniformShaderParameters->RuntimeVirtualTextureMipLevel = MipLevelParameter;
 		View->ViewUniformBuffer = TUniformBufferRef<FViewUniformShaderParameters>::CreateUniformBufferImmediate(*View->CachedViewUniformShaderParameters, UniformBuffer_SingleFrame);
 
 		// Build graph
@@ -1599,6 +1599,7 @@ namespace RuntimeVirtualTexture
 			PassParameters->Scene = SceneRenderer->GetSceneUniformBufferRef(GraphBuilder);
 
 			FRuntimeVirtualTexturePassParameters* RuntimeVirtualTexturePassParameters = GraphBuilder.AllocParameters<FRuntimeVirtualTexturePassParameters>();
+			RuntimeVirtualTexturePassParameters->MipLevel = MipLevelParameter;
 			RuntimeVirtualTexturePassParameters->CustomMaterialData = CustomMaterialData;
 			RuntimeVirtualTexturePassParameters->DebugParams = CVarVTMipColors.GetValueOnRenderThread() ? GetMipLevelColor(vLevel) : FixedColor;
 			RuntimeVirtualTexturePassParameters->PackHeight = FVector2f(WorldHeightPackParameter);	// LWC_TODO: Precision loss
