When in instanced stereo, share a single ForwardLightBuffer between both eyes. Add PreViewTranslationEyeAdjustment to the ForwardLightStruct to apply the appropriate translation offset for instanced secondary views.

diff --git a/Engine/Shaders/Private/LightGridCommon.ush b/Engine/Shaders/Private/LightGridCommon.ush
--- a/Engine/Shaders/Private/LightGridCommon.ush
+++ b/Engine/Shaders/Private/LightGridCommon.ush
@@ -275,7 +275,7 @@ FDirectionalLightData GetDirectionalLightData(uint EyeIndex)
 }
 
 FForwardLightData GetForwardLightData_Internal(
-	uint LightIndex,
+	uint LightIndex, uint EyeIndex,
 	float4 InData0, float4 InData1, float4 InData2, float4 InData3, float4 InData4, float4 InData5)
 {
 	FForwardLightData Out = (FForwardLightData)0;
@@ -291,6 +291,22 @@ FForwardLightData GetForwardLightData_Internal(
 	Out.PrevLocalLightIndex							= int(asuint(InData5.w) & 0xFFFF) - 1;
 	Out.LightSceneId								= int(Out.LightColorAndIdAndFalloffExponent.z);
 
+#if (INSTANCED_STEREO || MOBILE_MULTI_VIEW)
+	BRANCH
+	if (EyeIndex == 0)
+	{
+#endif
+		// Should always be (0,0,0) for non-instanced or primary views, but necessary for VSMs because they use secondary instanced view UBs directly.
+		Out.LightPositionAndInvRadius.xyz += ForwardLightStruct.PreViewTranslationOffsetISR.xyz;
+#if (INSTANCED_STEREO || MOBILE_MULTI_VIEW)
+	}
+	else
+	{
+		// Instanced views share one ForwardLightBuffer, and light positions (which are stored relative to the primary view) need to be offset to account for parallax
+		Out.LightPositionAndInvRadius.xyz += ForwardLightStructISR.PreViewTranslationOffsetISR.xyz;
+	}
+#endif
+
 	return Out;
 }
 
@@ -299,7 +315,7 @@ FForwardLightData GetForwardLightDataNonStereo(uint LightIndex)
 	const uint LightBaseIndex = LightIndex * FORWARD_LIGHT_DATA_STRIDE;
 
 	return GetForwardLightData_Internal(
-		LightIndex,
+		LightIndex, 0,
 		ForwardLightStruct.ForwardLightBuffer[LightBaseIndex + 0],
 		ForwardLightStruct.ForwardLightBuffer[LightBaseIndex + 1],
 		ForwardLightStruct.ForwardLightBuffer[LightBaseIndex + 2],
@@ -328,7 +344,7 @@ FLocalLightData GetLocalLightData(uint GridIndex, uint EyeIndex)
 		uint LocalLightBaseIndex = LocalLightIndex * FORWARD_LIGHT_DATA_STRIDE;
 
 		Out.Internal = GetForwardLightData_Internal(
-			LocalLightIndex,
+			LocalLightIndex, EyeIndex,
 			ForwardLightStruct.ForwardLightBuffer[LocalLightBaseIndex + 0],
 			ForwardLightStruct.ForwardLightBuffer[LocalLightBaseIndex + 1],
 			ForwardLightStruct.ForwardLightBuffer[LocalLightBaseIndex + 2],
@@ -343,7 +359,7 @@ FLocalLightData GetLocalLightData(uint GridIndex, uint EyeIndex)
 		uint LocalLightBaseIndex = LocalLightIndex * FORWARD_LIGHT_DATA_STRIDE;
 
 		Out.Internal = GetForwardLightData_Internal(
-			LocalLightIndex,
+			LocalLightIndex, EyeIndex,
 			ForwardLightStructISR.ForwardLightBuffer[LocalLightBaseIndex + 0],
 			ForwardLightStructISR.ForwardLightBuffer[LocalLightBaseIndex + 1],
 			ForwardLightStructISR.ForwardLightBuffer[LocalLightBaseIndex + 2],
diff --git a/Engine/Source/Runtime/Renderer/Private/LightGridInjection.cpp b/Engine/Source/Runtime/Renderer/Private/LightGridInjection.cpp
--- a/Engine/Source/Runtime/Renderer/Private/LightGridInjection.cpp
+++ b/Engine/Source/Runtime/Renderer/Private/LightGridInjection.cpp
@@ -777,6 +777,21 @@ FComputeLightGridOutput FSceneRenderer::ComputeLightGrid(FRDGBuilder& GraphBuild
 	{
 		FViewInfo& View = Views[ViewIndex];
 
+		// Get the associated secondary view if this is an instanced primary view, or vice/versa
+		const FViewInfo* SecondaryView = nullptr;
+		const FViewInfo* PrimaryView = nullptr;
+		if (View.bIsMobileMultiViewEnabled || View.bIsInstancedStereoEnabled)
+		{
+			if (View.StereoPass == EStereoscopicPass::eSSP_PRIMARY)
+			{
+				SecondaryView = View.GetInstancedView();
+			}
+			else if (View.StereoPass == EStereoscopicPass::eSSP_SECONDARY)
+			{
+				PrimaryView = View.GetPrimaryView();
+			}
+		}
+
 		View.ForwardLightingResources.SelectedForwardDirectionalLightProxy = nullptr;
 
 		FForwardLightUniformParameters* ForwardLightUniformParameters = PerViewForwardLightUniformParameters[ViewIndex];
@@ -842,7 +857,9 @@ FComputeLightGridOutput FSceneRenderer::ComputeLightGrid(FRDGBuilder& GraphBuild
 				const FLightSceneInfo* const LightSceneInfo = SortedLightInfo.LightSceneInfo;
 				const FLightSceneProxy* LightProxy = LightSceneInfo->Proxy;
 
-				if (LightSceneInfo->ShouldRenderLight(View))
+				if (LightSceneInfo->ShouldRenderLight(View) ||
+					(SecondaryView && LightSceneInfo->ShouldRenderLight(*SecondaryView)) ||
+					(PrimaryView && LightSceneInfo->ShouldRenderLight(*PrimaryView)))
 				{
 					FLightRenderParameters LightParameters;
 					LightProxy->GetLightShaderParameters(LightParameters, LightShaderParameterFlags);
@@ -1132,16 +1149,29 @@ FComputeLightGridOutput FSceneRenderer::ComputeLightGrid(FRDGBuilder& GraphBuild
 		static_assert(sizeof(FForwardLightData) % sizeof(FVector4f) == 0, "ForwardLightBuffer is used as a StructuredBuffer<float4> in the shader");
 		const uint32 ForwardLightDataSizeNumFloat4 = (NumLightsFinal * sizeof(FForwardLightData)) / sizeof(FVector4f);
 
-		FRDGBufferRef ForwardLightBuffer = CreateStructuredBuffer(GraphBuilder, TEXT("ForwardLightBuffer"),
-			TConstArrayView<FVector4f>(reinterpret_cast<const FVector4f*>(ForwardLightData.GetData()), ForwardLightDataSizeNumFloat4));
-
 		FRDGBufferRef DirectionalLightIndicesBuffer = CreateStructuredBuffer(GraphBuilder, TEXT("DirectionalLightIndicesBuffer"), DirectionalLightIndices);
 
 		View.bLightGridHasRectLights = bHasRectLights;
 		View.bLightGridHasTexturedLights = bHasTexturedLights;
 
 		const FIntPoint LightGridSizeXY = FIntPoint::DivideAndRoundUp(View.ViewRect.Size(), GLightGridPixelSize);
-		ForwardLightUniformParameters->ForwardLightBuffer = GraphBuilder.CreateSRV(FRDGBufferSRVDesc(ForwardLightBuffer));
+
+		if (PrimaryView)
+		{
+			// The visibility lists should be exactly the same
+			//checkSlow(PrimaryView->ForwardLightingResources.LocalLightVisibleLightInfosIndex == LocalLightVisibleLightInfosIndex);
+			check(PerViewForwardLightUniformParameters[View.PrimaryViewIndex]->NumLocalLights == NumLightsFinal);
+			ForwardLightUniformParameters->ForwardLightBuffer = PerViewForwardLightUniformParameters[View.PrimaryViewIndex]->ForwardLightBuffer;
+			ForwardLightUniformParameters->PreViewTranslationOffsetISR = FVector4f(FVector3f(View.ViewMatrices.GetPreViewTranslation() - PrimaryView->ViewMatrices.GetPreViewTranslation()), 0.0f);
+		}
+		else
+		{
+			FRDGBufferRef ForwardLightBuffer = CreateStructuredBuffer(GraphBuilder, TEXT("ForwardLocalLightBuffer"),
+				TConstArrayView<FVector4f>(reinterpret_cast<const FVector4f*>(ForwardLightData.GetData()), ForwardLightDataSizeNumFloat4));
+			ForwardLightUniformParameters->ForwardLightBuffer = GraphBuilder.CreateSRV(FRDGBufferSRVDesc(ForwardLightBuffer));
+			ForwardLightUniformParameters->PreViewTranslationOffsetISR = FVector4f(0.0f, 0.0f, 0.0f, 0.0f); // Light positions are stored relative to the primary view
+		}
+
 		ForwardLightUniformParameters->DirectionalLightIndices = GraphBuilder.CreateSRV(FRDGBufferSRVDesc(DirectionalLightIndicesBuffer));
 		ForwardLightUniformParameters->NumLocalLights = NumVisibleLocalLights;
 		ForwardLightUniformParameters->NumDirectionalLights = DirectionalLightIndices.Num();
diff --git a/Engine/Source/Runtime/Renderer/Private/SceneRendering.h b/Engine/Source/Runtime/Renderer/Private/SceneRendering.h
--- a/Engine/Source/Runtime/Renderer/Private/SceneRendering.h
+++ b/Engine/Source/Runtime/Renderer/Private/SceneRendering.h
@@ -692,6 +692,7 @@ BEGIN_GLOBAL_SHADER_PARAMETER_STRUCT_WITH_CONSTRUCTOR(FForwardLightUniformParame
 	SHADER_PARAMETER(uint32, DirectLightingShowFlag)
 	SHADER_PARAMETER(uint32, LightFunctionAtlasLightIndex)
 	SHADER_PARAMETER(uint32, bAffectsTranslucentLighting)
+	SHADER_PARAMETER(FVector4f, PreViewTranslationOffsetISR)
 	SHADER_PARAMETER_RDG_TEXTURE(Texture2D, DirectionalLightShadowmapAtlas)
 	SHADER_PARAMETER_SAMPLER(SamplerState, ShadowmapSampler)
 	SHADER_PARAMETER_TEXTURE(Texture2D, DirectionalLightStaticShadowmap)
