Re-enable viewrect support for volumetric RTs now that the underlying issue is fixed with 38829217

[Backout] - CL38615127
Original CL Desc
-----------------------------------------------------------------
Default r.VolumetricRenderTarget.ViewRect to 0 which fallbacks to the old behaviour when dynamic res changes.
Doing this temporarily while we investigate why we're seeing GPU hangs with it set to 1.

diff --git a/Engine/Source/Runtime/Renderer/Private/VolumetricRenderTarget.cpp b/Engine/Source/Runtime/Renderer/Private/VolumetricRenderTarget.cpp
--- a/Engine/Source/Runtime/Renderer/Private/VolumetricRenderTarget.cpp
+++ b/Engine/Source/Runtime/Renderer/Private/VolumetricRenderTarget.cpp
@@ -48,7 +48,7 @@ static TAutoConsoleVariable<int32> CVarVolumetricRenderTargetPreferAsyncCompute(
 	ECVF_RenderThreadSafe | ECVF_Scalability);
 
 static TAutoConsoleVariable<int32> CVarVolumetricRenderTargetViewRect(
-	TEXT("r.VolumetricRenderTarget.ViewRect"), 0,
+	TEXT("r.VolumetricRenderTarget.ViewRect"), 1,
 	TEXT("Enable ViewRect support: does not reallocate new render targets when dynamic resolution changes"),
 	ECVF_RenderThreadSafe | ECVF_Scalability);
 
