Substrate - fixed assert when deferred tried to render mesh decals beforelighting when it is not supported (only DBuffer decals are).

#tests PC editor dx12 with/without mobile preview with/without dbuffer
#rnx

diff --git a/Engine/Source/Runtime/Renderer/Private/CompositionLighting/PostProcessDeferredDecals.cpp b/Engine/Source/Runtime/Renderer/Private/CompositionLighting/PostProcessDeferredDecals.cpp
--- a/Engine/Source/Runtime/Renderer/Private/CompositionLighting/PostProcessDeferredDecals.cpp
+++ b/Engine/Source/Runtime/Renderer/Private/CompositionLighting/PostProcessDeferredDecals.cpp
@@ -678,8 +678,8 @@ void AddDeferredDecalPass(
 		// Sanity check - Substrate only support DBuffer, SceneColor, or AO decals
 		if (Substrate::IsSubstrateEnabled())
 		{
-			const bool bSupported = RenderTargetMode == EDecalRenderTargetMode::DBuffer || RenderTargetMode == EDecalRenderTargetMode::SceneColor || RenderTargetMode == EDecalRenderTargetMode::AmbientOcclusion;
-			if (!bSupported)
+			const bool bDecalSupported = RenderTargetMode == EDecalRenderTargetMode::DBuffer || RenderTargetMode == EDecalRenderTargetMode::SceneColor || RenderTargetMode == EDecalRenderTargetMode::AmbientOcclusion;
+			if (!bDecalSupported)
 			{
 				return;
 			}
@@ -760,7 +760,17 @@ void AddDeferredDecalPass(
 
 		if (bHasAnyDrawCommandDecalCount && (DecalRenderStage == EDecalRenderStage::BeforeBasePass || DecalRenderStage == EDecalRenderStage::BeforeLighting || DecalRenderStage == EDecalRenderStage::Emissive || DecalRenderStage == EDecalRenderStage::AmbientOcclusion))
 		{
-			RenderMeshDecals(GraphBuilder, Scene, View, PassTextures, InstanceCullingManager, DecalRenderStage);
+			// Sanity check - Substrate only support DBuffer, SceneColor, or AO decals
+			bool bDecalSupported = true;
+			if (Substrate::IsSubstrateEnabled() && DecalRenderStage == EDecalRenderStage::BeforeLighting)
+			{
+				bDecalSupported = false;
+			}
+
+			if (bDecalSupported)
+			{
+				RenderMeshDecals(GraphBuilder, Scene, View, PassTextures, InstanceCullingManager, DecalRenderStage);
+			}
 		}
 
 		if (SortedDecalCount > 0)
