Removed left behind commented code from 38850935

#rnx

diff --git a/Engine/Source/Runtime/Renderer/Private/MobileShadingRenderer.cpp b/Engine/Source/Runtime/Renderer/Private/MobileShadingRenderer.cpp
--- a/Engine/Source/Runtime/Renderer/Private/MobileShadingRenderer.cpp
+++ b/Engine/Source/Runtime/Renderer/Private/MobileShadingRenderer.cpp
@@ -1292,20 +1292,6 @@ void FMobileSceneRenderer::Render(FRDGBuilder& GraphBuilder, const FSceneRenderU
 		}
 		ExternalAccessQueue.Submit(GraphBuilder);
 
-//		if (bRendererOutputFinalSceneColor && Views.Num() > 0)
-//		{
-//			FViewInfo& MainView = Views[0];
-//			const bool bRealTimeSkyCaptureEnabled = (bShouldRenderSkyAtmosphere || MainView.SkyMeshBatches.Num() > 0) && Scene->SkyLight && Scene->SkyLight->bRealTimeCaptureEnabled && ViewFamily.EngineShowFlags.SkyLighting;
-//			if (bRealTimeSkyCaptureEnabled)
-//			{
-//				const bool bShouldRenderVolumetricCloud = false; // Not supported on the mobile renderer.
-//				Scene->AllocateAndCaptureFrameSkyEnvMap(GraphBuilder, *this, MainView, bShouldRenderSkyAtmosphere, bShouldRenderVolumetricCloud, InstanceCullingManager, ExternalAccessQueue);
-//				ExternalAccessQueue.Submit(GraphBuilder);
-//
-//				UpdateSkyReflectionUniformBuffer(GraphBuilder.RHICmdList);
-//			}
-//		}
-
 		// Custom depth
 		// bShouldRenderCustomDepth has been initialized in InitViews on mobile platform
 		if (bShouldRenderCustomDepth)
