Fixed DOF visualization bug with non-unit viewport scaling causing overlay layers to be stretched and misaligned.
Fixed DOF visualization bug with overlay text not scaling properly with DPI.

diff --git a/Engine/Shaders/Private/PostProcessVisualizeDOF.usf b/Engine/Shaders/Private/PostProcessVisualizeDOF.usf
--- a/Engine/Shaders/Private/PostProcessVisualizeDOF.usf
+++ b/Engine/Shaders/Private/PostProcessVisualizeDOF.usf
@@ -24,8 +24,8 @@ void VisualizeDOFPS(
 {
 	OutColor = 0;
 
-	float SceneDeviceZ = DepthTexture.SampleLevel(DepthSampler, ApplyScreenTransform(SvPosition.xy, SvPositionToColorUV), 0).r;
-	half3 SceneColor = Texture2DSample(ColorTexture, ColorSampler, ApplyScreenTransform(SvPosition.xy, SvPositionToDepthUV)).rgb;
+	float SceneDeviceZ = DepthTexture.SampleLevel(DepthSampler, ApplyScreenTransform(SvPosition.xy, SvPositionToDepthUV), 0).r;
+	half3 SceneColor = Texture2DSample(ColorTexture, ColorSampler, ApplyScreenTransform(SvPosition.xy, SvPositionToColorUV)).rgb;
 
 	float SceneDepth = ConvertFromDeviceZ(SceneDeviceZ);
 
diff --git a/Engine/Source/Runtime/Renderer/Private/PostProcess/PostProcessBokehDOF.cpp b/Engine/Source/Runtime/Renderer/Private/PostProcess/PostProcessBokehDOF.cpp
--- a/Engine/Source/Runtime/Renderer/Private/PostProcess/PostProcessBokehDOF.cpp
+++ b/Engine/Source/Runtime/Renderer/Private/PostProcess/PostProcessBokehDOF.cpp
@@ -86,6 +86,14 @@ FScreenPassTexture AddVisualizeDOFPass(FRDGBuilder& GraphBuilder, const FViewInf
 
 	AddDrawCanvasPass(GraphBuilder, RDG_EVENT_NAME("Overlay"), View, Output, [&View](FCanvas& Canvas)
 	{
+		const float DPIScale = Canvas.GetDPIScale();
+		Canvas.SetBaseTransform(FMatrix(FScaleMatrix(DPIScale) * Canvas.CalcBaseTransform2D(Canvas.GetViewRect().Width(), Canvas.GetViewRect().Height())));
+
+		auto DrawShadowedString = [&](float X, float Y, FStringView Text, const FLinearColor& Color = FLinearColor::White)
+		{
+			Canvas.DrawShadowedString(X / DPIScale, Y / DPIScale, Text, GetStatsFont(), Color);
+		};
+
 		float X = 30;
 		float Y = 18;
 		const float YStep = 14;
@@ -94,25 +102,25 @@ FScreenPassTexture AddVisualizeDOFPass(FRDGBuilder& GraphBuilder, const FViewInf
 		FString Line;
 
 		Line = FString::Printf(TEXT("Visualize Depth of Field"));
-		Canvas.DrawShadowedString(20, Y += YStep, *Line, GetStatsFont(), FLinearColor(1, 1, 0));
+		DrawShadowedString(20, Y += YStep, Line, FLinearColor(1, 1, 0));
 		Y += YStep;
 
 		Line = FString::Printf(TEXT("Cinematic DOF (blue is far, green is near, black is in focus, cross hair shows Depth and CoC radius in pixel)"));
-		Canvas.DrawShadowedString(X, Y += YStep, *Line, GetStatsFont(), FLinearColor(1, 1, 1));
+		DrawShadowedString(X, Y += YStep, Line);
 		Y += YStep;
 		Line = FString::Printf(TEXT("FocalDistance: %.2f"), View.FinalPostProcessSettings.DepthOfFieldFocalDistance);
-		Canvas.DrawShadowedString(X, Y += YStep, *Line, GetStatsFont(), FLinearColor(1, 1, 1));
+		DrawShadowedString(X, Y += YStep, Line);
 		Line = FString::Printf(TEXT("Aperture F-stop: %.2f"), View.FinalPostProcessSettings.DepthOfFieldFstop);
-		Canvas.DrawShadowedString(X, Y += YStep, *Line, GetStatsFont(), FLinearColor(1, 1, 1));
+		DrawShadowedString(X, Y += YStep, Line);
 		Line = FString::Printf(TEXT("Aperture: f/%.2f"), View.FinalPostProcessSettings.DepthOfFieldFstop);
-		Canvas.DrawShadowedString(X, Y += YStep, *Line, GetStatsFont(), FLinearColor(1, 1, 1));
+		DrawShadowedString(X, Y += YStep, Line);
 		Y += YStep;
 		Line = FString::Printf(TEXT("DepthBlur (not related to Depth of Field, due to light traveling long distances in atmosphere)"));
-		Canvas.DrawShadowedString(X, Y += YStep, *Line, GetStatsFont(), FLinearColor(1, 1, 1));
+		DrawShadowedString(X, Y += YStep, Line);
 		Line = FString::Printf(TEXT("  km for 50%%: %.2f"), View.FinalPostProcessSettings.DepthOfFieldDepthBlurAmount);
-		Canvas.DrawShadowedString(X, Y += YStep, *Line, GetStatsFont(), FLinearColor(1, 1, 1));
+		DrawShadowedString(X, Y += YStep, Line);
 		Line = FString::Printf(TEXT("  Radius (pixels in 1920x): %.2f"), View.FinalPostProcessSettings.DepthOfFieldDepthBlurRadius);
-		Canvas.DrawShadowedString(X, Y += YStep, *Line, GetStatsFont(), FLinearColor(1, 1, 1));
+		DrawShadowedString(X, Y += YStep, Line);
 		Y += YStep;
 
 		const FVector2D Fov = View.ViewMatrices.ComputeHalfFieldOfViewPerAxis();
@@ -123,13 +131,13 @@ FScreenPassTexture AddVisualizeDOFPass(FRDGBuilder& GraphBuilder, const FViewInf
 		const float UUToMM = 10.0f;
 
 		Line = FString::Printf(TEXT("Field Of View in deg.: %.1f x %.1f"), FMath::RadiansToDegrees(Fov.X) * 2.0f, FMath::RadiansToDegrees(Fov.Y) * 2.0f);
-		Canvas.DrawShadowedString(X, Y += YStep, *Line, GetStatsFont(), FLinearColor(0.5, 0.5, 1));
+		DrawShadowedString(X, Y += YStep, Line, FLinearColor(0.5, 0.5, 1));
 		Line = FString::Printf(TEXT("Focal Length: %.1f mm"), UUToMM * CocModel.VerticalFocalLength);
-		Canvas.DrawShadowedString(X, Y += YStep, *Line, GetStatsFont(), FLinearColor(0.5, 0.5, 1));
+		DrawShadowedString(X, Y += YStep, Line, FLinearColor(0.5, 0.5, 1));
 		Line = FString::Printf(TEXT("Squeeze Factor: %.1f"), CocModel.Squeeze);
-		Canvas.DrawShadowedString(X, Y += YStep, *Line, GetStatsFont(), FLinearColor(0.5, 0.5, 1));
+		DrawShadowedString(X, Y += YStep, Line, FLinearColor(0.5, 0.5, 1));
 		Line = FString::Printf(TEXT("Sensor: %.1f x %.1f mm"), UUToMM * CocModel.SensorWidth, UUToMM * CocModel.SensorHeight);
-		Canvas.DrawShadowedString(X, Y += YStep, *Line, GetStatsFont(), FLinearColor(0.5, 0.5, 1));
+		DrawShadowedString(X, Y += YStep, Line, FLinearColor(0.5, 0.5, 1));
 	});
 
 	return MoveTemp(Output);
