Fix duplicated localization name.

diff --git a/Engine/Source/Runtime/Renderer/Private/SceneRendering.cpp b/Engine/Source/Runtime/Renderer/Private/SceneRendering.cpp
--- a/Engine/Source/Runtime/Renderer/Private/SceneRendering.cpp
+++ b/Engine/Source/Runtime/Renderer/Private/SceneRendering.cpp
@@ -4272,7 +4272,7 @@ void FSceneRenderer::OnRenderFinish(FRDGBuilder& GraphBuilder, FRDGTextureRef Vi
 						if (bExpenssiveSkyLightRealTimeCaptureWithCloud)
 						{
 							// This can happen because cloud in the real time capture are rendered at the face capture resolution
-							static const FText Message = NSLOCTEXT("Renderer", "Skylight", "The sky light resolution is too high for 6 faces captured per frame with volumetric cloud tracing (done at face resolution).\n\
+							static const FText Message = NSLOCTEXT("Renderer", "SkylightResolutionTooHigh", "The sky light resolution is too high for 6 faces captured per frame with volumetric cloud tracing (done at face resolution).\n\
 																							It could take a long time and cause a GPU TDR or crash on older GPUs.\n\
 																							Please reduce the sky light resolution to less than 512 or the r.SkyLight.RealTimeReflectionCapture.TimeSlice.SkyCloudCubeFacePerFrame to less than 6.\n\
 																							You may use r.SkyLight.RealTimeReflectionCapture.DisableExpenssiveCaptureMessage 1 to disable that message.");
@@ -4280,7 +4280,7 @@ void FSceneRenderer::OnRenderFinish(FRDGBuilder& GraphBuilder, FRDGTextureRef Vi
 						}
 						if (bRealTimeSkyCaptureButNothingToCapture)
 						{
-							static const FText Message = NSLOCTEXT("Renderer", "Skylight", "A sky light with real-time capture enable is in the scene. It requires at least a SkyAtmosphere component, A volumetricCloud component or a mesh with a material tagged as IsSky. Otherwise it will be black");
+							static const FText Message = NSLOCTEXT("Renderer", "SkylightRequiresSkyAtmosphere", "A sky light with real-time capture enable is in the scene. It requires at least a SkyAtmosphere component, A volumetricCloud component or a mesh with a material tagged as IsSky. Otherwise it will be black");
 							Writer.DrawLine(Message);
 						}
 						if (bShowPointLightWarning)
