Fix mobile CSM merging issue when using GPU scene

diff --git a/Engine/Source/Runtime/Renderer/Private/MeshDrawCommands.cpp b/Engine/Source/Runtime/Renderer/Private/MeshDrawCommands.cpp
--- a/Engine/Source/Runtime/Renderer/Private/MeshDrawCommands.cpp
+++ b/Engine/Source/Runtime/Renderer/Private/MeshDrawCommands.cpp
@@ -241,11 +241,17 @@ void MergeMobileBasePassMeshDrawCommands(
 				FVisibleMeshDrawCommand& MeshCommand = MeshCommands[i];
 				FVisibleMeshDrawCommand& MeshCommandCSM = MeshCommandsCSM[i];
 
+				const int32 CSMBit = 1 << 30;
+				checkf((MeshCommand.StateBucketId & CSMBit) == 0, TEXT("Non-CSM draw command's StateBucketID should not have the CSM bit set."));
+
 				if (MeshCommand.PrimitiveIdInfo.ScenePrimitiveId < ScenePrimitiveNum && MobileCSMVisibilityInfo.MobilePrimitiveCSMReceiverVisibilityMap[MeshCommand.PrimitiveIdInfo.ScenePrimitiveId])
 				{
 					checkf(MeshCommand.PrimitiveIdInfo.ScenePrimitiveId == MeshCommandCSM.PrimitiveIdInfo.ScenePrimitiveId, TEXT("VisibleMeshDrawCommands of BasePass and MobileBasePassCSM are expected to match."));
 					// Use CSM's VisibleMeshDrawCommand.
 					MeshCommand = MeshCommandCSM;
+
+					// Make sure that CSM and non-CSM draw commands don't have the same StateBucketId, which would potentially get them merged later.
+					MeshCommand.StateBucketId |= CSMBit;
 				}
 			}
 			MeshCommandsCSM.Reset();
