View / Light Channel Masking:  Support for controlling which views a light renders in via a mask, similar to the actor / primitive lighting channel mask.  A default channel is enabled for all lights and views, and scene captures expose the view / light channel mask in the component.

diff --git a/Engine/Source/Runtime/Engine/Classes/Components/LightComponent.h b/Engine/Source/Runtime/Engine/Classes/Components/LightComponent.h
--- a/Engine/Source/Runtime/Engine/Classes/Components/LightComponent.h
+++ b/Engine/Source/Runtime/Engine/Classes/Components/LightComponent.h
@@ -178,6 +178,10 @@ class ULightComponent : public ULightComponentBase
 	UPROPERTY(EditAnywhere, AdvancedDisplay, BlueprintReadOnly, Category=Light)
 	FLightingChannels LightingChannels;
 
+	/** View / light masking support.  Controls which views this light should affect. */
+	UPROPERTY(EditAnywhere, AdvancedDisplay, BlueprintReadOnly, Category=Light)
+	FViewLightingChannels ViewLightingChannels;
+
 	/** 
 	 * The light function material to be applied to this light.
 	 * Note that only non-lightmapped lights (UseDirectLightMap=False) can have a light function. 
diff --git a/Engine/Source/Runtime/Engine/Classes/Components/SceneCaptureComponent.h b/Engine/Source/Runtime/Engine/Classes/Components/SceneCaptureComponent.h
--- a/Engine/Source/Runtime/Engine/Classes/Components/SceneCaptureComponent.h
+++ b/Engine/Source/Runtime/Engine/Classes/Components/SceneCaptureComponent.h
@@ -153,6 +153,10 @@ class USceneCaptureComponent : public USceneComponent
 	UPROPERTY(EditAnywhere, BlueprintReadWrite, Category=SceneCapture, AdvancedDisplay)
 	FName CollectionTransformProjection;
 
+	/** View / light masking support.  Controls which lights should affect this view. */
+	UPROPERTY(EditAnywhere, BlueprintReadWrite, Category=SceneCapture)
+	FViewLightingChannels ViewLightingChannels;
+
 	/** ShowFlags for the SceneCapture's ViewFamily, to control rendering settings for this view. Hidden but accessible through details customization */
 	UE_DEPRECATED_FORGAME(5.5, "Public access to this property is deprecated, and it will become private in a future release. Please use SetShowFlagSettings and GetShowFlagSettings instead.")
 	UPROPERTY(EditAnywhere, interp, Blueprintgetter = GetShowFlagSettings, BlueprintSetter = SetShowFlagSettings, Category=SceneCapture)
diff --git a/Engine/Source/Runtime/Engine/Classes/Engine/EngineTypes.h b/Engine/Source/Runtime/Engine/Classes/Engine/EngineTypes.h
--- a/Engine/Source/Runtime/Engine/Classes/Engine/EngineTypes.h
+++ b/Engine/Source/Runtime/Engine/Classes/Engine/EngineTypes.h
@@ -593,6 +593,49 @@ inline int32 GetFirstLightingChannelFromMask(uint8 Mask)
 	return Mask ? FPlatformMath::CountTrailingZeros(Mask) : -1;
 }
 
+/** View / light masking support */
+USTRUCT(BlueprintType)
+struct FViewLightingChannels
+{
+	GENERATED_BODY()
+
+	FViewLightingChannels() :
+		bViewChannel0(true),
+		bViewChannel1(false),
+		bViewChannel2(false),
+		bViewChannel3(false),
+		bViewChannel4(false)
+	{}
+
+	/** View specific lighting channel 0 (enabled by default). */
+	UPROPERTY(EditAnywhere, BlueprintReadWrite, Category=Channels)
+	uint8 bViewChannel0:1;
+
+	/** View specific lighting channel 1. */
+	UPROPERTY(EditAnywhere, BlueprintReadWrite, Category=Channels)
+	uint8 bViewChannel1:1;
+
+	/** View specific lighting channel 2. */
+	UPROPERTY(EditAnywhere, BlueprintReadWrite, Category=Channels)
+	uint8 bViewChannel2:1;
+
+	/** View specific lighting channel 3. */
+	UPROPERTY(EditAnywhere, BlueprintReadWrite, Category=Channels)
+	uint8 bViewChannel3:1;
+
+	/** View specific lighting channel 4. */
+	UPROPERTY(EditAnywhere, BlueprintReadWrite, Category = Channels)
+	uint8 bViewChannel4:1;
+
+	/** Converts view lighting channels into a bitfield for use in FLightSceneProxy */
+	inline uint8 GetMaskForStruct() const
+	{
+		return (uint8)((bViewChannel0 << 0) | (bViewChannel1 << 1) | (bViewChannel2 << 2) | (bViewChannel3 << 3) | (bViewChannel4 << 4));
+	}
+
+	static inline constexpr uint8 GetMaskDefault() { return 0x1; }
+};
+
 /** 
  * Enumerates available GBufferFormats. 
  * @warning When this enum is updated please update CVarGBufferFormat comments 
diff --git a/Engine/Source/Runtime/Engine/Private/LightSceneProxy.cpp b/Engine/Source/Runtime/Engine/Private/LightSceneProxy.cpp
--- a/Engine/Source/Runtime/Engine/Private/LightSceneProxy.cpp
+++ b/Engine/Source/Runtime/Engine/Private/LightSceneProxy.cpp
@@ -62,6 +62,7 @@ FLightSceneProxy::FLightSceneProxy(const ULightComponent* InLightComponent)
 	, AtmosphereSunDiskColorScale(InLightComponent->GetAtmosphereSunDiskColorScale())
 	, LightType(InLightComponent->GetLightType())
 	, LightingChannelMask(GetLightingChannelMaskForStruct(InLightComponent->LightingChannels))
+	, ViewLightingChannelMask(InLightComponent->ViewLightingChannels.GetMaskForStruct())
 	, StatId(InLightComponent->GetStatID(true))
 	, ComponentName(InLightComponent->GetFName())
 	, LevelName(InLightComponent->GetOwner() ? InLightComponent->GetOwner()->GetLevel()->GetOutermost()->GetFName() : NAME_None)
diff --git a/Engine/Source/Runtime/Engine/Public/LightSceneProxy.h b/Engine/Source/Runtime/Engine/Public/LightSceneProxy.h
--- a/Engine/Source/Runtime/Engine/Public/LightSceneProxy.h
+++ b/Engine/Source/Runtime/Engine/Public/LightSceneProxy.h
@@ -230,6 +230,7 @@ public:
 	inline float GetRayStartOffsetDepthScale() const { return RayStartOffsetDepthScale; }
 	inline uint8 GetLightType() const { return LightType; }
 	inline uint8 GetLightingChannelMask() const { return LightingChannelMask; }
+	inline uint8 GetViewLightingChannelMask() const { return ViewLightingChannelMask; }
 	inline FName GetComponentFName() const { return ComponentName; }
 
 	inline bool IsSelected() const { return bSelected; }
@@ -482,6 +483,7 @@ protected:
 	const uint8 LightType;
 
 	uint8 LightingChannelMask;
+	uint8 ViewLightingChannelMask;
 
 	/** Used for dynamic stats */
 	TStatId StatId;
diff --git a/Engine/Source/Runtime/Engine/Public/SceneView.h b/Engine/Source/Runtime/Engine/Public/SceneView.h
--- a/Engine/Source/Runtime/Engine/Public/SceneView.h
+++ b/Engine/Source/Runtime/Engine/Public/SceneView.h
@@ -1749,6 +1749,9 @@ public:
 	/** Set to false to disable ray tracing on this view.  Use IsRayTracingAllowedForView() function to test for ray tracing, not this bool. */
 	bool bAllowRayTracing = true;
 
+	/** View light masking support.  Some systems assume the same mask setting for all views within a family, with potential artifacts if views use different masks. */
+	uint8 ViewLightingChannelMask = FViewLightingChannels::GetMaskDefault();
+
 	/**  Stereo aspects of the shader pipeline based on this view's shader platform */
 	UE::StereoRenderUtils::FStereoShaderAspects Aspects;
 
diff --git a/Engine/Source/Runtime/Renderer/Private/HeterogeneousVolumes/HeterogeneousVolumesLiveShadingPipeline.cpp b/Engine/Source/Runtime/Renderer/Private/HeterogeneousVolumes/HeterogeneousVolumesLiveShadingPipeline.cpp
--- a/Engine/Source/Runtime/Renderer/Private/HeterogeneousVolumes/HeterogeneousVolumesLiveShadingPipeline.cpp
+++ b/Engine/Source/Runtime/Renderer/Private/HeterogeneousVolumes/HeterogeneousVolumesLiveShadingPipeline.cpp
@@ -1359,6 +1359,7 @@ static void RenderWithTransmittanceVolumePipeline(
 	for (auto LightIt = Scene->Lights.CreateConstIterator(); LightIt; ++LightIt)
 	{
 		if (HeterogeneousVolumes::SupportsLightType(LightIt->LightType) &&
+			(View.ViewLightingChannelMask & LightIt->LightSceneInfo->Proxy->GetViewLightingChannelMask()) &&
 			LightIt->AffectsPrimitive(HeterogeneousVolumeInterface->GetBounds(), HeterogeneousVolumeInterface->GetPrimitiveSceneProxy()))
 		{
 			LightSceneInfoCompact.Add(*LightIt);
@@ -1517,6 +1518,7 @@ static void RenderWithInscatteringVolumePipeline(
 		for (auto LightIt = Scene->Lights.CreateConstIterator(); LightIt; ++LightIt)
 		{
 			if (HeterogeneousVolumes::SupportsLightType(LightIt->LightType) &&
+				(View.ViewLightingChannelMask & LightIt->LightSceneInfo->Proxy->GetViewLightingChannelMask()) &&
 				LightIt->AffectsPrimitive(HeterogeneousVolumeInterface->GetBounds(), HeterogeneousVolumeInterface->GetPrimitiveSceneProxy()))
 			{
 				LightSceneInfoCompact.Add(*LightIt);
@@ -2396,7 +2398,8 @@ void RenderAdaptiveVolumetricShadowMapWithLiveShading(
 	{
 		// TODO: Use global bounds information..
 		//if (LightIt->AffectsPrimitive(HeterogeneousVolumeInterface->GetBounds(), HeterogeneousVolumeInterface->GetPrimitiveSceneProxy()))
-		if (HeterogeneousVolumes::SupportsShadowForLightType(LightIt->LightType))
+		if (HeterogeneousVolumes::SupportsShadowForLightType(LightIt->LightType) &&
+			(View.ViewLightingChannelMask & LightIt->LightSceneInfo->Proxy->GetViewLightingChannelMask()))
 		{
 			LightSceneInfoCompact.Add(*LightIt);
 		}
diff --git a/Engine/Source/Runtime/Renderer/Private/HeterogeneousVolumes/HeterogeneousVolumesPreshadingPipeline.cpp b/Engine/Source/Runtime/Renderer/Private/HeterogeneousVolumes/HeterogeneousVolumesPreshadingPipeline.cpp
--- a/Engine/Source/Runtime/Renderer/Private/HeterogeneousVolumes/HeterogeneousVolumesPreshadingPipeline.cpp
+++ b/Engine/Source/Runtime/Renderer/Private/HeterogeneousVolumes/HeterogeneousVolumesPreshadingPipeline.cpp
@@ -795,6 +795,7 @@ void RenderWithInscatteringVolumePipelineWithPreshadingCompute(
 		for (auto LightIt = Scene->Lights.CreateConstIterator(); LightIt; ++LightIt)
 		{
 			if (HeterogeneousVolumes::SupportsLightType(LightIt->LightType) &&
+				(LightIt->LightSceneInfo->Proxy->GetViewLightingChannelMask() & View.ViewLightingChannelMask) &&
 				LightIt->AffectsPrimitive(HeterogeneousVolumeInterface->GetBounds(), HeterogeneousVolumeInterface->GetPrimitiveSceneProxy()))
 			{
 				LightSceneInfoCompact.Add(*LightIt);
@@ -932,6 +933,7 @@ void RenderWithTransmittanceVolumePipelineWithPreshadingCompute(
 	for (auto LightIt = Scene->Lights.CreateConstIterator(); LightIt; ++LightIt)
 	{
 		if (HeterogeneousVolumes::SupportsLightType(LightIt->LightType) &&
+			(LightIt->LightSceneInfo->Proxy->GetViewLightingChannelMask() & View.ViewLightingChannelMask) &&
 			LightIt->AffectsPrimitive(HeterogeneousVolumeInterface->GetBounds(), HeterogeneousVolumeInterface->GetPrimitiveSceneProxy()))
 		{
 			LightSceneInfoCompact.Add(*LightIt);
@@ -1158,6 +1160,7 @@ void RenderWithInscatteringVolumePipelineWithPreshadingHardwareRayTracing(
 	for (auto LightIt = Scene->Lights.CreateConstIterator(); LightIt; ++LightIt)
 	{
 		if (HeterogeneousVolumes::SupportsLightType(LightIt->LightType) &&
+			(LightIt->LightSceneInfo->Proxy->GetViewLightingChannelMask() & View.ViewLightingChannelMask) &&
 			LightIt->AffectsPrimitive(HeterogeneousVolumeInterface->GetBounds(), HeterogeneousVolumeInterface->GetPrimitiveSceneProxy()))
 		{
 			LightSceneInfoCompact.Add(*LightIt);
@@ -1288,6 +1291,7 @@ void RenderWithTransmittanceVolumePipelineWithPreshadingHardwareRayTracing(
 	for (auto LightIt = Scene->Lights.CreateConstIterator(); LightIt; ++LightIt)
 	{
 		if (HeterogeneousVolumes::SupportsLightType(LightIt->LightType) &&
+			(LightIt->LightSceneInfo->Proxy->GetViewLightingChannelMask() & View.ViewLightingChannelMask) &&
 			LightIt->AffectsPrimitive(HeterogeneousVolumeInterface->GetBounds(), HeterogeneousVolumeInterface->GetPrimitiveSceneProxy()))
 		{
 			LightSceneInfoCompact.Add(*LightIt);
diff --git a/Engine/Source/Runtime/Renderer/Private/HeterogeneousVolumes/HeterogeneousVolumesVoxelGridPipeline.cpp b/Engine/Source/Runtime/Renderer/Private/HeterogeneousVolumes/HeterogeneousVolumesVoxelGridPipeline.cpp
--- a/Engine/Source/Runtime/Renderer/Private/HeterogeneousVolumes/HeterogeneousVolumesVoxelGridPipeline.cpp
+++ b/Engine/Source/Runtime/Renderer/Private/HeterogeneousVolumes/HeterogeneousVolumesVoxelGridPipeline.cpp
@@ -2735,7 +2735,12 @@ void CollectHeterogeneousVolumeMeshBatches(
 	{
 		// TODO: Use global bounds information..
 		//if (LightIt->AffectsPrimitive(HeterogeneousVolumeInterface->GetBounds(), HeterogeneousVolumeInterface->GetPrimitiveSceneProxy()))
-		if (HeterogeneousVolumes::SupportsShadowForLightType(LightIt->LightType))
+
+		// TODO: For view / light masking, should we check if the light is visible in any view, instead of just the first view?  What kind of artifacts would
+		// result with mixed settings on views?  For now, we'll assume all views in a family have the same channel mask (true for scene captures where
+		// view / light masks are currently exposed), so this doesn't matter.
+		if (HeterogeneousVolumes::SupportsShadowForLightType(LightIt->LightType) &&
+			(LightIt->LightSceneInfo->Proxy->GetViewLightingChannelMask() & Views[0].ViewLightingChannelMask))
 		{
 			LightSceneInfoCompact.Add(*LightIt);
 		}
@@ -3835,7 +3840,8 @@ void RenderAdaptiveVolumetricShadowMapWithVoxelGrid(
 	{
 		// TODO: Use global bounds information..
 		//if (LightIt->AffectsPrimitive(HeterogeneousVolumeInterface->GetBounds(), HeterogeneousVolumeInterface->GetPrimitiveSceneProxy()))
-		if (HeterogeneousVolumes::SupportsShadowForLightType(LightIt->LightType))
+		if (HeterogeneousVolumes::SupportsShadowForLightType(LightIt->LightType) &&
+			(LightIt->LightSceneInfo->Proxy->GetViewLightingChannelMask() & View.ViewLightingChannelMask))
 		{
 			LightSceneInfoCompact.Add(*LightIt);
 		}
diff --git a/Engine/Source/Runtime/Renderer/Private/LightSceneInfo.cpp b/Engine/Source/Runtime/Renderer/Private/LightSceneInfo.cpp
--- a/Engine/Source/Runtime/Renderer/Private/LightSceneInfo.cpp
+++ b/Engine/Source/Runtime/Renderer/Private/LightSceneInfo.cpp
@@ -242,6 +242,7 @@ bool FLightSceneInfo::ShouldRenderLight(const FViewInfo& View, bool bOffscreen)
 	return bLocalVisible
 		// Only render lights with static shadowing for reflection captures, since they are only captured at edit time
 		&& IsInDesiredSceneSubset(View, Proxy)
+		&& (Proxy->GetViewLightingChannelMask() & View.ViewLightingChannelMask)
 		// Only render lights in the default channel, or if there are any primitives outside the default channel
 		&& (Proxy->GetLightingChannelMask() & GetDefaultLightingChannelMask() || View.bUsesLightingChannels || bOffscreen);
 }
diff --git a/Engine/Source/Runtime/Renderer/Private/SceneCaptureRendering.cpp b/Engine/Source/Runtime/Renderer/Private/SceneCaptureRendering.cpp
--- a/Engine/Source/Runtime/Renderer/Private/SceneCaptureRendering.cpp
+++ b/Engine/Source/Runtime/Renderer/Private/SceneCaptureRendering.cpp
@@ -762,6 +762,11 @@ TArray<FSceneView*> SetupViewFamilyForSceneCapture(
 			View->OverridePostProcessSettings(*PostProcessSettings, PostProcessBlendWeight);
 		}
 		View->EndFinalPostprocessSettings(ViewInitOptions);
+
+		if (SceneCaptureComponent)
+		{
+			View->ViewLightingChannelMask = SceneCaptureComponent->ViewLightingChannels.GetMaskForStruct();
+		}
 	}
 
 	return ViewPtrArray;
