Persistent SBT fixes:
- Make sure no duplicate record indices are added to the SetBindingsOnShaderBindingTable (processed multithreaded and not thread safe). Can happen with multiple clears on the same record before flush to RHI SBT is called
- Make sure only valid SBT record indices are used during clear on SetBindingsOnShaderBindingTable. Can happen when adding & removing records happens before new RHI SBT is allocated
- Track last used base shader resources (textures & buffers) when removing listeners from SBT record - otherwise base shader resource might be deleted while SBT still has a residency tracked reference to the underlying d3d12 resource
- Track LastRayTracingMaterialPipelineSignature on FScene instead of FSceneViewState to be in sync with the FShaderBindingLayout dirty record states - if cached on the FSceneViewState which isn't used each frame then a very old signature could be used when it needs to create a new RTPSO unblocking (falls back to finding RTPSO in cache for last valid signature stored in the FSceneViewState - this RTPSO might not actually contain all used ray tracing shaders stored in the SBT)
- Add some extra validation and checks

diff --git a/Engine/Source/Runtime/D3D12RHI/Private/D3D12RayTracing.cpp b/Engine/Source/Runtime/D3D12RHI/Private/D3D12RayTracing.cpp
--- a/Engine/Source/Runtime/D3D12RHI/Private/D3D12RayTracing.cpp
+++ b/Engine/Source/Runtime/D3D12RHI/Private/D3D12RayTracing.cpp
@@ -1605,6 +1605,7 @@ public:
 				WorkerData[WorkerIndex].DynamicReferencedResources.Empty();
 				WorkerData[WorkerIndex].DynamicReferencedResourceSet.Empty();
 				WorkerData[WorkerIndex].PersistentReferencedResources.Empty();
+				WorkerData[WorkerIndex].LastTrackedBaseShaderResources.Empty();
 			}
 		}
 
@@ -1624,6 +1625,7 @@ public:
 			ReferencedResources.Append(WorkerData[0].DynamicReferencedResources);
 			WorkerData[0].DynamicReferencedResources.Empty();
 			WorkerData[0].DynamicReferencedResourceSet.Empty();
+			WorkerData[0].LastTrackedBaseShaderResources.Empty();
 		}
 #endif // ENABLE_RESIDENCY_MANAGEMENT
 
@@ -1899,6 +1901,7 @@ public:
 			if (CurrentResource)
 			{
 				ShaderTable.UpdatePersistentReferencedResource(WorkerIndex, CurrentResource, -1);
+				ShaderTable.AddLastTrackedBaseShaderResource(WorkerIndex, ShaderResource);
 			}
 			ShaderResource->RemoveRenameListener(this);
 			ShaderResource = nullptr;
@@ -1965,6 +1968,8 @@ public:
 			ReferencedTexture->RemoveRenameListener(this);
 			TextureReference->RemoveReplaceListener(this);
 
+			ShaderTable.AddLastTrackedBaseShaderResource(WorkerIndex, ReferencedTexture);
+
 			TextureReference = nullptr;
 			RecordIndex = 0;
 			CurrentResource = 0;
@@ -2090,7 +2095,12 @@ public:
 		}
 	}
 		
-#if ENABLE_RESIDENCY_MANAGEMENT	
+#if ENABLE_RESIDENCY_MANAGEMENT
+	void AddLastTrackedBaseShaderResource(uint32 WorkerIndex, FD3D12BaseShaderResource* BaseShaderResource)
+	{
+		WorkerData[WorkerIndex].LastTrackedBaseShaderResources.Add(BaseShaderResource);
+	}
+
 	void UpdatePersistentReferencedResource(uint32 WorkerIndex, FD3D12Resource* Resource, int32 Count)
 	{
 		check(Resource);
@@ -2452,6 +2462,10 @@ public:
 		// Track all persistent resources with number of refs accross all records (when ref reaches 0 then it doesn't need be tracked anymore)
 		TMap<TRefCountPtr<FD3D12Resource>, int32> PersistentReferencedResources;
 
+		// Track last reference to base shader resources because PersistentReferencedResources could have the last reference to the d3d12 resource
+		// and we can't delete d3d12baseshaderresource with committed resource when it doesn't hold the last reference to it
+		TSet<TRefCountPtr<FD3D12BaseShaderResource>> LastTrackedBaseShaderResources;
+
 		// A set of all resources referenced by this shader table for the purpose of updating residency before ray tracing work dispatch.
 		Experimental::TSherwoodSet<void*> DynamicReferencedResourceSet;
 		TArray<TRefCountPtr<FD3D12Resource>> DynamicReferencedResources;
diff --git a/Engine/Source/Runtime/Renderer/Private/RayTracing/RayTracingMaterialHitShaders.cpp b/Engine/Source/Runtime/Renderer/Private/RayTracing/RayTracingMaterialHitShaders.cpp
--- a/Engine/Source/Runtime/Renderer/Private/RayTracing/RayTracingMaterialHitShaders.cpp
+++ b/Engine/Source/Runtime/Renderer/Private/RayTracing/RayTracingMaterialHitShaders.cpp
@@ -760,13 +760,13 @@ void FDeferredShadingSceneRenderer::CreateMaterialRayTracingMaterialPipeline(
 
 	const bool bAllowNonBlockingPipelineCreation = GRayTracingNonBlockingPipelineCreation && !View.bIsOfflineRender;
 	FRayTracingPipelineState* FallbackPipelineState = bAllowNonBlockingPipelineCreation && View.ViewState
-		? PipelineStateCache::GetRayTracingPipelineState(View.ViewState->LastRayTracingMaterialPipelineSignature)
+		? PipelineStateCache::GetRayTracingPipelineState(Scene->LastRayTracingMaterialPipelineSignature)
 		: nullptr;
 
 	ERayTracingPipelineCacheFlags PipelineCacheFlags = ERayTracingPipelineCacheFlags::Default;
 	if (FallbackPipelineState
 		&& View.ViewState
-		&& IsCompatibleFallbackPipelineSignature(View.ViewState->LastRayTracingMaterialPipelineSignature, Initializer)
+		&& IsCompatibleFallbackPipelineSignature(Scene->LastRayTracingMaterialPipelineSignature, Initializer)
 		&& PipelineContainsHitShaders(FallbackPipelineState, RequiredHitShaders)
 		&& FindRayTracingMissShaderIndex(FallbackPipelineState, DefaultMissShader, false) != INDEX_NONE
 		&& (!bCallableShadersRequired || FindRayTracingCallableShaderIndex(FallbackPipelineState, DefaultCallableShader, false) != INDEX_NONE))
@@ -781,7 +781,7 @@ void FDeferredShadingSceneRenderer::CreateMaterialRayTracingMaterialPipeline(
 		if (View.ViewState)
 		{
 			// Save the current pipeline to be used as fallback in future frames
-			View.ViewState->LastRayTracingMaterialPipelineSignature = static_cast<FRayTracingPipelineStateSignature&>(Initializer);
+			Scene->LastRayTracingMaterialPipelineSignature = static_cast<FRayTracingPipelineStateSignature&>(Initializer);
 		}
 	}
 	else
@@ -790,6 +790,7 @@ void FDeferredShadingSceneRenderer::CreateMaterialRayTracingMaterialPipeline(
 		check(FallbackPipelineState);
 		PipelineState = FallbackPipelineState;
 		bOutIsUsingFallbackRTPSO = true;
+		UE_LOG(LogRenderer, Log, TEXT("Using fallback RTPSO"));
 	}
 
 	// Retrieve the binding data size from the actual used RTPSO because the requested RTPSO could still be non blocking async compiling
@@ -854,7 +855,7 @@ void FDeferredShadingSceneRenderer::SetupMaterialRaytracingHitGroupBindings(FRDG
 			View.RayTracingMaterialBindings[TaskIndex] = BindingWriter;
 
 			GraphBuilder.AddSetupTask(
-				[&View, SceneUB, bIsPathTracing, PipelineState, BindingWriter, DirtyShaderBindings, NumBindings, bEnableMaterials, bEnableShadowMaterials, bSupportMeshDecals,
+				[&View, &RayTracingSBT = Scene->RayTracingSBT, SceneUB, bIsPathTracing, PipelineState, BindingWriter, DirtyShaderBindings, NumBindings, bEnableMaterials, bEnableShadowMaterials, bSupportMeshDecals,
 				OpaqueShadowMaterialIndex, HiddenMaterialIndex, OpaqueMeshDecalHitGroupIndex, HiddenMeshDecalHitGroupIndex, TaskIndex]()
 				{
 					TRACE_CPUPROFILER_EVENT_SCOPE(BindRayTracingMaterialPipelineTask);
@@ -893,6 +894,11 @@ void FDeferredShadingSceneRenderer::SetupMaterialRaytracingHitGroupBindings(FRDG
 							{
 								HitGroupIndex = FoundIndex;
 							}
+							else if (RayTracingSBT.IsPersistent())
+							{
+								check(DirtyShaderBinding.BindingType == ERayTracingLocalShaderBindingType::Transient);
+								check(RayTracingSBT.IsDirty(DirtyShaderBinding.SBTRecordIndex));
+							}
 						}
 
 						uint32 BaseRecordIndex = DirtyShaderBinding.SBTRecordIndex;
diff --git a/Engine/Source/Runtime/Renderer/Private/RayTracing/RayTracingShaderBindingTable.cpp b/Engine/Source/Runtime/Renderer/Private/RayTracing/RayTracingShaderBindingTable.cpp
--- a/Engine/Source/Runtime/Renderer/Private/RayTracing/RayTracingShaderBindingTable.cpp
+++ b/Engine/Source/Runtime/Renderer/Private/RayTracing/RayTracingShaderBindingTable.cpp
@@ -486,15 +486,51 @@ FRayTracingSBTAllocation* FRayTracingShaderBindingTable::AllocateDynamicRange(ER
 	return Allocation;
 }
 
+bool FRayTracingShaderBindingTable::IsDirty(uint32 RecordIndex) const
+{
+	return !ValidPersistentRecords[RecordIndex];
+}
+
+bool FRayTracingShaderBindingTable::IsPersistent() const
+{
+	return PersistentSBTInitializer.Lifetime == ERayTracingShaderBindingTableLifetime::Persistent;
+}
+
 void FRayTracingShaderBindingTable::FlushAllocationsToClear(FRHICommandList& RHICmdList)
 {
 	FScopeLock ScopeLock(&StaticAllocationCS);
 
+	// Don't clear outside of current allocated range - could have allocated and cleared ranges before calling CheckPersistentRHI
+	// and then the RHI SBT doesn't have those new ranges allocated yet
+	uint32 MaxNumValidRecords = PersistentSBTInitializer.NumGeometrySegments * PersistentSBTInitializer.NumShaderSlotsPerGeometrySegment;
+	check(MaxNumValidRecords <= (uint32)ValidPersistentRecords.Num());
+
+	// Build bit array if records to clear (make sure we don't clear the same record twice)
+	TBitArray<> PersistentRecordsToClear;
+	PersistentRecordsToClear.Init(false, MaxNumValidRecords);
+	
 	uint32 TotalRecordsToClear = 0;
 	for (const FRayTracingSBTAllocation* Allocation : PersistentAllocationsToClear)
 	{
-		TotalRecordsToClear += Allocation->NumRecords;
+		for (uint32 RangeIndex = 0; RangeIndex < Allocation->NumRecords; ++RangeIndex)
+		{
+			uint32 RecordIndex = Allocation->BaseRecordIndex + RangeIndex;
+			if (RecordIndex < MaxNumValidRecords)
+			{
+				FBitReference RecordFlag = PersistentRecordsToClear[RecordIndex];
+				if (!RecordFlag)
+				{
+					TotalRecordsToClear++;
+					RecordFlag = true;
+				}
+			}
+		}
+
+		// TODO: pool static allocations as well?
+		delete Allocation;
 	}
+	PersistentAllocationsToClear.Reset(PersistentAllocationsToClear.Num());
+
 	if (TotalRecordsToClear == 0)
 	{
 		return;
@@ -507,20 +543,13 @@ void FRayTracingShaderBindingTable::FlushAllocationsToClear(FRHICommandList& RHI
 	: RHICmdList.Alloc(LocalBindingsSize, alignof(FRayTracingLocalShaderBindings)));
 
 	uint32 LocalBindingIndex = 0;
-	for (const FRayTracingSBTAllocation* Allocation : PersistentAllocationsToClear)
+	for (TConstSetBitIterator<> It(PersistentRecordsToClear); It; ++It)
 	{
-		for (uint32 RangeIndex = 0; RangeIndex < Allocation->NumRecords; ++RangeIndex)
-		{
-			FRayTracingLocalShaderBindings& RecordBindingToClear = LocalBindings[LocalBindingIndex++];
-			RecordBindingToClear = FRayTracingLocalShaderBindings();
-			RecordBindingToClear.RecordIndex = Allocation->BaseRecordIndex + RangeIndex;
-			RecordBindingToClear.BindingType = ERayTracingLocalShaderBindingType::Clear;
-		}
-
-		// TODO: pool static allocations as well?
-		delete Allocation;
+		FRayTracingLocalShaderBindings& RecordBindingToClear = LocalBindings[LocalBindingIndex++];
+		RecordBindingToClear = FRayTracingLocalShaderBindings();
+		RecordBindingToClear.RecordIndex = It.GetIndex();
+		RecordBindingToClear.BindingType = ERayTracingLocalShaderBindingType::Clear;
 	}
-	PersistentAllocationsToClear.Reset(PersistentAllocationsToClear.Num());
 
 	const bool bCopyDataToInlineStorage = false; // Storage is already allocated from RHICmdList, no extra copy necessary
 	for (FRHIShaderBindingTable* PersistentSBT : PersistentSBTs)
diff --git a/Engine/Source/Runtime/Renderer/Private/RayTracing/RayTracingShaderBindingTable.h b/Engine/Source/Runtime/Renderer/Private/RayTracing/RayTracingShaderBindingTable.h
--- a/Engine/Source/Runtime/Renderer/Private/RayTracing/RayTracingShaderBindingTable.h
+++ b/Engine/Source/Runtime/Renderer/Private/RayTracing/RayTracingShaderBindingTable.h
@@ -192,6 +192,16 @@ public:
 	 */	
 	void FlushAllocationsToClear(FRHICommandList& RHICmdList);
 
+	/**
+	 * Check if given record index is dirty - used for validation
+	 */	
+	bool IsDirty(uint32 RecordIndex) const;
+
+	/**
+	 * Check if SBTs are persistently allocated
+	 */
+	bool IsPersistent() const;
+
 private:
 	/**
 	 * Get the maximum amount of static allocated segments (highest allocation index with free ranges included)
