Allow some ISceneViewExtension postprocess delegates to extend mobile postprocessing
#12664

diff --git a/Engine/Source/Runtime/Renderer/Private/PostProcess/PostProcessing.cpp b/Engine/Source/Runtime/Renderer/Private/PostProcess/PostProcessing.cpp
--- a/Engine/Source/Runtime/Renderer/Private/PostProcess/PostProcessing.cpp
+++ b/Engine/Source/Runtime/Renderer/Private/PostProcess/PostProcessing.cpp
@@ -2364,6 +2364,7 @@ void AddMobilePostProcessingPasses(FRDGBuilder& GraphBuilder, FScene* Scene, con
 {
 	RDG_CSV_STAT_EXCLUSIVE_SCOPE(GraphBuilder, RenderPostProcessing);
 	QUICK_SCOPE_CYCLE_COUNTER(STAT_PostProcessing_Process);
+	using namespace UE::Renderer::PostProcess;
 
 	check(IsInRenderingThread());
 	Inputs.Validate();
@@ -2413,6 +2414,22 @@ void AddMobilePostProcessingPasses(FRDGBuilder& GraphBuilder, FScene* Scene, con
 		MAX
 	};
 
+	// Mobile unsupported passes return EPass::MAX
+	const auto TranslatePass = [](ISceneViewExtension::EPostProcessingPass Pass) -> EPass
+		{
+			switch (Pass)
+			{
+			case ISceneViewExtension::EPostProcessingPass::MotionBlur: return EPass::MAX;
+			case ISceneViewExtension::EPostProcessingPass::Tonemap: return EPass::Tonemap;
+			case ISceneViewExtension::EPostProcessingPass::FXAA: return EPass::FXAA;
+			case ISceneViewExtension::EPostProcessingPass::VisualizeDepthOfField: return EPass::MAX;
+
+			default:
+				check(false);
+				return EPass::MAX;
+			};
+		};
+
 	static const TCHAR* PassNames[] =
 	{
 		TEXT("Distortion"),
@@ -2500,6 +2517,41 @@ void AddMobilePostProcessingPasses(FRDGBuilder& GraphBuilder, FScene* Scene, con
 	PassSequence.SetEnabled(EPass::VisualizeLightGrid, ShouldVisualizeLightGrid());
 	PassSequence.SetEnabled(EPass::HMDDistortion, View.Family->EngineShowFlags.StereoRendering && View.Family->EngineShowFlags.HMDDistortion);
 
+	const auto GetPostProcessMaterialInputs = [&](FScreenPassTexture InSceneColor)
+		{
+			FPostProcessMaterialInputs PostProcessMaterialInputs;
+
+			PostProcessMaterialInputs.SetInput(GraphBuilder, EPostProcessMaterialInput::SceneColor, InSceneColor);
+			PostProcessMaterialInputs.SetInput(GraphBuilder, EPostProcessMaterialInput::Velocity, Velocity);
+			PostProcessMaterialInputs.SceneTextures = GetSceneTextureShaderParameters(Inputs.SceneTextures);
+			PostProcessMaterialInputs.CustomDepthTexture = CustomDepth.Texture;
+
+			return PostProcessMaterialInputs;
+		};
+
+	const auto AddAfterPass = [&](EPass InPass, FScreenPassTexture InSceneColor) -> FScreenPassTexture
+	{
+		// In some cases (e.g. OCIO color conversion) we want View Extensions to be able to add extra custom post processing after the pass.
+
+		FPostProcessingPassDelegateArray& PassCallbacks = PassSequence.GetAfterPassCallbacks(InPass);
+
+		if (PassCallbacks.Num())
+		{
+			FPostProcessMaterialInputs InOutPostProcessAfterPassInputs = GetPostProcessMaterialInputs(InSceneColor);
+
+			for (int32 AfterPassCallbackIndex = 0; AfterPassCallbackIndex < PassCallbacks.Num(); AfterPassCallbackIndex++)
+			{
+				InOutPostProcessAfterPassInputs.SetInput(GraphBuilder, EPostProcessMaterialInput::SceneColor, InSceneColor);
+
+				FAfterPassCallbackDelegate& AfterPassCallback = PassCallbacks[AfterPassCallbackIndex];
+				PassSequence.AcceptOverrideIfLastPass(InPass, InOutPostProcessAfterPassInputs.OverrideOutput, AfterPassCallbackIndex);
+				InSceneColor = AfterPassCallback.Execute(GraphBuilder, View, InOutPostProcessAfterPassInputs);
+			}
+		}
+
+		return MoveTemp(InSceneColor);
+	};
+
 	// Always evaluate custom post processes
 	// The scene color will be decoded at the first post-process material and output linear color space for the following passes
 	// bMetalMSAAHDRDecode will be set to false if there is any post-process material exist
@@ -2547,6 +2599,11 @@ void AddMobilePostProcessingPasses(FRDGBuilder& GraphBuilder, FScene* Scene, con
 		}
 	};
 	
+
+	constexpr int32 FirstAfterPass = static_cast<int32>(ISceneViewExtension::EPostProcessingPass::MotionBlur);
+	// Scene view extension delegates that precede the override pass sequence are to be called directly.
+	TStaticArray<FPostProcessingPassDelegateArray, static_cast<uint32>(FirstAfterPass)> SceneViewExtensionDelegates;
+
 	if (bIsPostProcessingEnabled)
 	{
 		bool bUseSun = View.MobileLightShaft.IsSet();
@@ -2578,6 +2635,33 @@ void AddMobilePostProcessingPasses(FRDGBuilder& GraphBuilder, FScene* Scene, con
 		PassSequence.SetEnabled(EPass::TAA, bUseTAA);
 		PassSequence.SetEnabled(EPass::FXAA, View.AntiAliasingMethod == AAM_FXAA);
 		PassSequence.SetEnabled(EPass::PostProcessMaterialAfterTonemapping, PostProcessMaterialAfterTonemappingChain.Num() != 0);
+
+		for (const TSharedRef<ISceneViewExtension>& ViewExtension : View.Family->ViewExtensions)
+		{
+			for (int32 SceneViewPassId = 0; SceneViewPassId < FirstAfterPass; SceneViewPassId++)
+			{
+				const ISceneViewExtension::EPostProcessingPass SceneViewPass = static_cast<ISceneViewExtension::EPostProcessingPass>(SceneViewPassId);
+				const bool bIsEnabled = (SceneViewPass == ISceneViewExtension::EPostProcessingPass::ReplacingTonemapper) ? PassSequence.IsEnabled(EPass::Tonemap) : true;
+
+				ViewExtension->SubscribeToPostProcessingPass(SceneViewPass, View, SceneViewExtensionDelegates[SceneViewPassId], bIsEnabled);
+			}
+
+			for (int32 SceneViewPassId = FirstAfterPass; SceneViewPassId < static_cast<int32>(ISceneViewExtension::EPostProcessingPass::MAX); SceneViewPassId++)
+			{
+				const ISceneViewExtension::EPostProcessingPass SceneViewPass = static_cast<ISceneViewExtension::EPostProcessingPass>(SceneViewPassId);
+				const EPass PostProcessingPass = TranslatePass(SceneViewPass);
+
+				if (PostProcessingPass != EPass::MAX)
+				{
+					ViewExtension->SubscribeToPostProcessingPass(
+						SceneViewPass,
+						View,
+						PassSequence.GetAfterPassCallbacks(PostProcessingPass),
+						PassSequence.IsEnabled(PostProcessingPass));
+				}
+			}
+		}
+
 		PassSequence.Finalize();
 			
 		if (PassSequence.IsEnabled(EPass::Distortion))
@@ -2595,6 +2679,12 @@ void AddMobilePostProcessingPasses(FRDGBuilder& GraphBuilder, FScene* Scene, con
 			SceneColor = AddMobileDistortionMergePass(GraphBuilder, View, DistortionMergeInputs);
 		}
 
+		if (SceneViewExtensionDelegates[static_cast<uint32>(ISceneViewExtension::EPostProcessingPass::BeforeDOF)].Num())
+		{
+			SceneColor = AddSceneViewExtensionPassChain(GraphBuilder, View, GetPostProcessMaterialInputs(SceneColor),
+				SceneViewExtensionDelegates[static_cast<uint32>(ISceneViewExtension::EPostProcessingPass::BeforeDOF)]);
+		}
+
 		AddPostProcessMaterialPass(BL_SceneColorBeforeDOF);
 
 		// Optional fixed pass processes
@@ -2908,6 +2998,12 @@ void AddMobilePostProcessingPasses(FRDGBuilder& GraphBuilder, FScene* Scene, con
 			AddMobileSeparateTranslucencyPass(GraphBuilder, Scene, View, SeparateTranslucencyInputs);
 		}
 
+		if (SceneViewExtensionDelegates[static_cast<uint32>(ISceneViewExtension::EPostProcessingPass::AfterDOF)].Num())
+		{
+			SceneColor = AddSceneViewExtensionPassChain(GraphBuilder, View, GetPostProcessMaterialInputs(SceneColor),
+				SceneViewExtensionDelegates[static_cast<uint32>(ISceneViewExtension::EPostProcessingPass::AfterDOF)]);
+		}
+
 		AddPostProcessMaterialPass(BL_SceneColorAfterDOF);
 
 		// Temporal Anti-aliasing. Also may perform a temporal upsample from primary to secondary view rect.
@@ -3056,6 +3152,8 @@ void AddMobilePostProcessingPasses(FRDGBuilder& GraphBuilder, FScene* Scene, con
 		bMetalMSAAHDRDecode = false;
 	}
 
+	SceneColor = AddAfterPass(EPass::Tonemap, SceneColor);
+
 	if (IsPostProcessingEnabled(View))
 	{
 		if (PassSequence.IsEnabled(EPass::FXAA))
@@ -3068,6 +3166,8 @@ void AddMobilePostProcessingPasses(FRDGBuilder& GraphBuilder, FScene* Scene, con
 			SceneColor = AddFXAAPass(GraphBuilder, View, PassInputs);
 		}
 
+		SceneColor = AddAfterPass(EPass::FXAA, SceneColor);
+
 		if (PassSequence.IsEnabled(EPass::PostProcessMaterialAfterTonemapping))
 		{
 			AddPostProcessMaterialPass(BL_SceneColorAfterTonemapping);
