Disable ISM dynamic merging if instances require ordering preservation as it does not work correctly atm

diff --git a/Engine/Source/Runtime/Renderer/Private/InstanceCulling/InstanceCullingContext.cpp b/Engine/Source/Runtime/Renderer/Private/InstanceCulling/InstanceCullingContext.cpp
--- a/Engine/Source/Runtime/Renderer/Private/InstanceCulling/InstanceCullingContext.cpp
+++ b/Engine/Source/Runtime/Renderer/Private/InstanceCulling/InstanceCullingContext.cpp
@@ -1490,8 +1490,9 @@ void FInstanceCullingContext::SetupDrawCommands(
 		const bool bForceInstanceCulling = EnumHasAnyFlags(VisibleMeshDrawCommand.Flags, EFVisibleMeshDrawCommandFlags::ForceInstanceCulling) || (GOcclusionForceInstanceCulling != 0);
 		const bool bPreserveInstanceOrder = bOrderPreservationEnabled && EnumHasAnyFlags(VisibleMeshDrawCommand.Flags, EFVisibleMeshDrawCommandFlags::PreserveInstanceOrder);
 		const bool bUseIndirectDraw = bFetchInstanceCountFromScene || bAlwaysUseIndirectDraws || bForceInstanceCulling || (VisibleMeshDrawCommand.NumRuns > 0 || MeshDrawCommand->NumInstances > 1);
-		// UniformBufferView path does not support merging ISM draws atm
-		const bool bCompactIdenticalCommands = bInCompactIdenticalCommands && (bUsesUniformBufferView ? (CurrentAutoInstanceCount < MaxPrimitiveBatchSize && !bUseIndirectDraw) : true);
+		// UniformBufferView path and instance order preservation do not support merging ISM draws atm
+		const bool bCompactIdenticalCommands = bInCompactIdenticalCommands && !(bUseIndirectDraw && bPreserveInstanceOrder) 
+			&& (bUsesUniformBufferView ? (CurrentAutoInstanceCount < MaxPrimitiveBatchSize && !bUseIndirectDraw) : true);
 
 		if (bCompactIdenticalCommands && CurrentStateBucketId != -1 && VisibleMeshDrawCommand.StateBucketId == CurrentStateBucketId && VisibleMeshDrawCommand.CullingPayloadFlags == CurrentCullingPayloadFlags)
 		{
