Various changes to improve RVT responsiveness in the editor :
* Added priority system to VT producers (8 possible values : Lowest -> Highest)
* Added "prioritize" flag to VT tile requests
* Allowed RVT invalidations to be be prioritized so that the affected RVT pages get processed before others. This maintains the throttling behavior but allows the marked areas to be updated first
* RVTs can specify what their VT producer's priority should be, either automatically, based on the virtual texture's material type (BaseColor and its variant = High, Mask4 = BelowNormal, WorldHeigh = Low, Displacement = Lower) or using a custom priority
* RVT components can also override their RVT's priority by using a custom priority
* Added CVar to enable/disable these 2 priority systems : r.VT.SortTileRequestsByPriority (enabled by default)
* Added CVar to drop a visual log whenever a RVT component invalidation occurs (r.VT.RVT.VisLogManualInvalidations)

Notes :
* Because of the fact that the producer priority system needs to be defined in RenderCore, which cannot depend on UObjects, there are 2 identical enums for specifying the priority (EVTProducerPriority in RenderCore, EVirtualTextureProducerPriority, UENUM in Runtime)
* The producer priority supports any type of virtual texture (RVT, standard VT, etc.) but is currently only exposed to users on RVTs (/RVT components). It would be doable to expose it to standard VT assets via TextureSettings and/or using a per-texture asset custom priority override
* The streamed virtual textures producers cannot have a different priority value than their corresponding runtime virtual texture producers but the former still has a slightly priority in order to minimize low-res pages (as it used to)
* Fixed potential hash truncation issue with locked pages removals
* FVirtualTextureLocalTile was indifferently for storing VT tile info and VT tile request info and there's now a FVirtualTextureLocalTileRequest for the latter, which contains priority information (in the 4 available bits)

#tests editor, cooked game on all platforms

diff --git a/Engine/Source/Editor/LandscapeEditor/Private/LandscapeEdModeComponentTools.cpp b/Engine/Source/Editor/LandscapeEditor/Private/LandscapeEdModeComponentTools.cpp
--- a/Engine/Source/Editor/LandscapeEditor/Private/LandscapeEdModeComponentTools.cpp
+++ b/Engine/Source/Editor/LandscapeEditor/Private/LandscapeEdModeComponentTools.cpp
@@ -331,6 +331,9 @@ public:
 
 			Cache.SetCachedData(X1, Y1, X2, Y2, Data);
 			Cache.Flush();
+
+			// Dirty any runtime virtual textures that our landscape components write to.
+			LandscapeInfo->DirtyRuntimeVirtualTextureForLandscapeArea(X1, Y1, X2, Y2);
 		}
 	}
 
@@ -1781,6 +1784,9 @@ public:
 				Cache.Flush();
 			}
 
+			// Dirty any runtime virtual textures that our landscape components write to.
+			LandscapeInfo->DirtyRuntimeVirtualTextureForLandscapeArea(X1, Y1, X2, Y2);
+
 			GEngine->BroadcastLevelActorListChanged();
 		}
 	}
diff --git a/Engine/Source/Editor/LandscapeEditor/Private/LandscapeEdModeErosionTools.cpp b/Engine/Source/Editor/LandscapeEditor/Private/LandscapeEdModeErosionTools.cpp
--- a/Engine/Source/Editor/LandscapeEditor/Private/LandscapeEdModeErosionTools.cpp
+++ b/Engine/Source/Editor/LandscapeEditor/Private/LandscapeEdModeErosionTools.cpp
@@ -288,6 +288,9 @@ public:
 			WeightCache.SetCachedData(X1, Y1, X2, Y2, WeightDatas, LayerNum, ELandscapeLayerPaintingRestriction::None);
 		}
 		WeightCache.Flush();
+
+		// Dirty any runtime virtual textures that our landscape components write to.
+		LandscapeInfo->DirtyRuntimeVirtualTextureForLandscapeArea(X1, Y1, X2, Y2);
 	}
 };
 
@@ -519,6 +522,9 @@ public:
 		}
 
 		LayerHeightDataCache.Write(X1, Y1, X2, Y2, HeightData);
+
+		// Dirty any runtime virtual textures that our landscape components write to.
+		LandscapeInfo->DirtyRuntimeVirtualTextureForLandscapeArea(X1, Y1, X2, Y2);
 	}
 };
 
diff --git a/Engine/Source/Editor/LandscapeEditor/Private/LandscapeEdModePaintTools.cpp b/Engine/Source/Editor/LandscapeEditor/Private/LandscapeEdModePaintTools.cpp
--- a/Engine/Source/Editor/LandscapeEditor/Private/LandscapeEdModePaintTools.cpp
+++ b/Engine/Source/Editor/LandscapeEditor/Private/LandscapeEdModePaintTools.cpp
@@ -4,7 +4,6 @@
 #include "Misc/App.h"
 #include "InputCoreTypes.h"
 #include "Engine/EngineTypes.h"
-#include "Components/RuntimeVirtualTextureComponent.h"
 #include "Components/StaticMeshComponent.h"
 #include "Engine/StaticMesh.h"
 #include "LandscapeToolInterface.h"
@@ -43,53 +42,6 @@ const int32 FNoiseParameter::Permutations[256] =
 	138, 236, 205, 93, 222, 114, 67, 29, 24, 72, 243, 141, 128, 195, 78, 66, 215, 61, 156, 180
 };
 
-namespace
-{
-	/** Helper function to invalidate the runtime virtual texture data associated with the passed in landscape region. */
-	void DirtyRuntimeVirtualTextureForLandscapeArea(ULandscapeInfo* LandscapeInfo, int32 X1, int32 Y1, int32 X2, int32 Y2)
-	{
-		FBox DirtyWorldBounds(ForceInit);
-
-		// Iterate touched components to find touched runtime virtual textures.
-		TSet<ULandscapeComponent*> Components;
-		LandscapeInfo->GetComponentsInRegion(X1 + 1, Y1 + 1, X2 - 1, Y2 - 1, Components);
-
-		TArray<URuntimeVirtualTexture*> RuntimeVirtualTextures;
-		for (ULandscapeComponent* Component : Components)
-		{
-			ALandscapeProxy* Landscape = Component->GetLandscapeProxy();
-			if (Landscape != nullptr && Landscape->RuntimeVirtualTextures.Num() > 0)
-			{
-				for (URuntimeVirtualTexture* RuntimeVirtualTexture : Landscape->RuntimeVirtualTextures)
-				{
-					RuntimeVirtualTextures.AddUnique(RuntimeVirtualTexture);
-				}
-
-				// Also accumulate bounds in world space.
-				const int32 LocalX1 = FMath::Max(X1, Component->GetSectionBase().X) - Component->GetSectionBase().X;
-				const int32 LocalY1 = FMath::Max(Y1, Component->GetSectionBase().Y) - Component->GetSectionBase().Y;
-				const int32 LocalX2 = FMath::Min(X2, Component->GetSectionBase().X + LandscapeInfo->ComponentSizeQuads) - Component->GetSectionBase().X;
-				const int32 LocalY2 = FMath::Min(Y2, Component->GetSectionBase().Y + LandscapeInfo->ComponentSizeQuads) - Component->GetSectionBase().Y;
-				const FBox LocalDirtyBounds(FVector(LocalX1, LocalY1, 0), FVector(LocalX2, LocalY2, 1));
-
-				DirtyWorldBounds += LocalDirtyBounds.TransformBy(Component->GetComponentToWorld());
-			}
-		}
-
-		// Find matching runtime virtual texture components and invalidate dirty region.
-		if (RuntimeVirtualTextures.Num() > 0)
-		{
-			for (TObjectIterator<URuntimeVirtualTextureComponent> It(/*AdditionalExclusionFlags = */RF_ClassDefaultObject, /*bIncludeDerivedClasses = */true, /*InInternalExclusionFlags = */EInternalObjectFlags::Garbage); It; ++It)
-			{
-				if (It->GetVirtualTexture() != nullptr && RuntimeVirtualTextures.Contains(It->GetVirtualTexture()))
-				{
-					It->Invalidate(FBoxSphereBounds(DirtyWorldBounds));
-				}
-			}
-		}
-	}
-}
-
 // 
 // FLandscapeToolPaintBase
 //
@@ -331,7 +283,7 @@ public:
 		this->Cache.Flush();
 		
 		// Dirty any runtime virtual textures that our landscape components write to.
-		DirtyRuntimeVirtualTextureForLandscapeArea(LandscapeInfo, X1, Y1, X2, Y2);
+		this->LandscapeInfo->DirtyRuntimeVirtualTextureForLandscapeArea(X1, Y1, X2, Y2);
 	}
 };
 
@@ -436,6 +388,9 @@ public:
 
 		this->Cache.SetCachedData(X1, Y1, X2, Y2, Data, UISettings->PaintingRestriction);
 		this->Cache.Flush();
+
+		// Dirty any runtime virtual textures that our landscape components write to.
+		this->LandscapeInfo->DirtyRuntimeVirtualTextureForLandscapeArea(X1, Y1, X2, Y2);
 	}
 };
 
@@ -664,6 +619,9 @@ public:
 
 		this->Cache.SetCachedData(X1, Y1, X2, Y2, Data);
 		this->Cache.Flush();
+
+		// Dirty any runtime virtual textures that our landscape components write to.
+		this->LandscapeInfo->DirtyRuntimeVirtualTextureForLandscapeArea(X1, Y1, X2, Y2);
 	}
 };
 
@@ -824,11 +782,8 @@ public:
 
 		LayerDataCache.Write(X1, Y1, X2, Y2, Data, UISettings->PaintingRestriction);
 
-		if (ToolTarget::TargetType == ELandscapeToolTargetType::Weightmap)
-		{
-			// Dirty any runtime virtual textures that our landscape components write to.
-			DirtyRuntimeVirtualTextureForLandscapeArea(this->LandscapeInfo, X1, Y1, X2, Y2);
-		}
+		// Dirty any runtime virtual textures that our landscape components write to.
+		this->LandscapeInfo->DirtyRuntimeVirtualTextureForLandscapeArea(X1, Y1, X2, Y2);
 	}
 };
 
@@ -1125,11 +1080,8 @@ public:
 
 		LayerDataCache.Write(X1, Y1, X2, Y2, Data, UISettings->PaintingRestriction);
 
-		if (ToolTarget::TargetType == ELandscapeToolTargetType::Weightmap)
-		{
-			// Dirty any runtime virtual textures that our landscape components write to.
-			DirtyRuntimeVirtualTextureForLandscapeArea(this->LandscapeInfo, X1, Y1, X2, Y2);
-		}
+		// Dirty any runtime virtual textures that our landscape components write to.
+		this->LandscapeInfo->DirtyRuntimeVirtualTextureForLandscapeArea(X1, Y1, X2, Y2);
 	}
 };
 
@@ -1416,11 +1368,8 @@ public:
 		this->Cache.SetCachedData(X1, Y1, X2, Y2, Data, UISettings->PaintingRestriction);
 		this->Cache.Flush();
 
-		if (ToolTarget::TargetType == ELandscapeToolTargetType::Weightmap)
-		{
-			// Dirty any runtime virtual textures that our landscape components write to.
-			DirtyRuntimeVirtualTextureForLandscapeArea(this->LandscapeInfo, X1, Y1, X2, Y2);
-		}
+		// Dirty any runtime virtual textures that our landscape components write to.
+		this->LandscapeInfo->DirtyRuntimeVirtualTextureForLandscapeArea(X1, Y1, X2, Y2);
 	}
 };
 
diff --git a/Engine/Source/Editor/LandscapeEditor/Private/LandscapeEdModeRampTool.cpp b/Engine/Source/Editor/LandscapeEditor/Private/LandscapeEdModeRampTool.cpp
--- a/Engine/Source/Editor/LandscapeEditor/Private/LandscapeEdModeRampTool.cpp
+++ b/Engine/Source/Editor/LandscapeEditor/Private/LandscapeEdModeRampTool.cpp
@@ -592,6 +592,9 @@ public:
 
 			LayerHeightDataCache.Write(MinX, MinY, MaxX, MaxY, Data);
 
+			// Dirty any runtime virtual textures that our landscape components write to.
+			LandscapeInfo->DirtyRuntimeVirtualTextureForLandscapeArea(MinX, MinY, MaxX, MaxY);
+
 			if (!EdMode->HasLandscapeLayersContent())
 			{
 				TSet<ULandscapeComponent*> Components;
diff --git a/Engine/Source/Editor/VirtualTexturingEditor/Private/RuntimeVirtualTextureDetailsCustomization.cpp b/Engine/Source/Editor/VirtualTexturingEditor/Private/RuntimeVirtualTextureDetailsCustomization.cpp
--- a/Engine/Source/Editor/VirtualTexturingEditor/Private/RuntimeVirtualTextureDetailsCustomization.cpp
+++ b/Engine/Source/Editor/VirtualTexturingEditor/Private/RuntimeVirtualTextureDetailsCustomization.cpp
@@ -107,32 +107,32 @@ void FRuntimeVirtualTextureDetailsCustomization::CustomizeDetails(IDetailLayoutB
 
 	TSharedRef<IPropertyHandle> MaterialTypePropertyHandle = DetailBuilder.GetProperty(TEXT("MaterialType"));
 	DetailBuilder.EditDefaultProperty(MaterialTypePropertyHandle)->CustomWidget()
-	.NameContent()
-	[
-		MaterialTypePropertyHandle->CreatePropertyNameWidget()
-	]
-	.ValueContent()
-	[
-		SNew(SEnumComboBox, StaticEnum<ERuntimeVirtualTextureMaterialType>())
-			.Font(FAppStyle::GetFontStyle(TEXT("MenuItem.Font")))
-			.EnumValueSubset(SupportedMaterialTypes)
-			.CurrentValue_Lambda([this]()
-			{
-				if (URuntimeVirtualTexture* Texture = VirtualTexture.Get())
+		.NameContent()
+		[
+			MaterialTypePropertyHandle->CreatePropertyNameWidget()
+		]
+		.ValueContent()
+		[
+			SNew(SEnumComboBox, StaticEnum<ERuntimeVirtualTextureMaterialType>())
+				.Font(FAppStyle::GetFontStyle(TEXT("MenuItem.Font")))
+				.EnumValueSubset(SupportedMaterialTypes)
+				.CurrentValue_Lambda([this]()
 				{
-					return (int32)Texture->GetMaterialType();
-				}
-				return 0;
-			})
-			.OnEnumSelectionChanged_Lambda([this](uint32 NewValue, ESelectInfo::Type)
-			{
-				if (URuntimeVirtualTexture* Texture = VirtualTexture.Get())
+					if (URuntimeVirtualTexture* Texture = VirtualTexture.Get())
+					{
+						return (int32)Texture->GetMaterialType();
+					}
+					return 0;
+				})
+				.OnEnumSelectionChanged_Lambda([this](uint32 NewValue, ESelectInfo::Type)
 				{
-					Texture->MaterialType = (ERuntimeVirtualTextureMaterialType)NewValue;
-					RefreshDetailsView();
-				}
-			})
-	];
+					if (URuntimeVirtualTexture* Texture = VirtualTexture.Get())
+					{
+						Texture->MaterialType = (ERuntimeVirtualTextureMaterialType)NewValue;
+						RefreshDetailsView();
+					}
+				})
+		];
 
 	// Set UIMax dependent on adaptive page table setting
 	FString MaxTileCountString = FString::Printf(TEXT("%d"), URuntimeVirtualTexture::GetMaxTileCountLog2(VirtualTexture->GetAdaptivePageTable()));
@@ -145,33 +145,79 @@ void FRuntimeVirtualTextureDetailsCustomization::CustomizeDetails(IDetailLayoutB
 	AddTextToProperty(DetailBuilder, SizeCategory, "TileBorderSize", TileBorderSizeText);
 
 	// Add details block
-	IDetailCategoryBuilder& DetailsCategory = DetailBuilder.EditCategory("Details", FText::GetEmpty(), ECategoryPriority::Important);
-	static const FText CustomRowSizeText = LOCTEXT("Details_RowFilter_Size", "Virtual Size");
-	DetailsCategory.AddCustomRow(CustomRowSizeText)
-	.NameContent()
-	[
-		SNew(STextBlock)
-		.Font(IDetailLayoutBuilder::GetDetailFont())
-		.Text(LOCTEXT("Details_Size", "Virtual Texture Size"))
-		.ToolTipText(LOCTEXT("Details_Size_Tooltip", "Virtual resolution derived from Size properties."))
-	]
-	.ValueContent()
-	[
-		SAssignNew(SizeText, STextBlock)
-	];
-	static const FText CustomRowPageTableSizeText = LOCTEXT("Details_RowFilter_PageTableSize", "Page Table Size");
-	DetailsCategory.AddCustomRow(CustomRowPageTableSizeText)
-	.NameContent()
-	[
-		SNew(STextBlock)
-		.Font(IDetailLayoutBuilder::GetDetailFont())
-		.Text(LOCTEXT("Details_PageTableSize", "Page Table Size"))
-		.ToolTipText(LOCTEXT("Details_PageTableSize_Tooltip", "Final page table size. This can vary according to the adaptive page table setting."))
-	]
-	.ValueContent()
-	[
-		SAssignNew(PageTableSizeText, STextBlock)
-	];
+	{
+		IDetailCategoryBuilder& DetailsCategory = DetailBuilder.EditCategory("Details", FText::GetEmpty(), ECategoryPriority::Important);
+		static const FText CustomRowSizeText = LOCTEXT("Details_RowFilter_Size", "Virtual Size");
+		DetailsCategory.AddCustomRow(CustomRowSizeText)
+			.NameContent()
+			[
+				SNew(STextBlock)
+				.Font(IDetailLayoutBuilder::GetDetailFont())
+				.Text(LOCTEXT("Details_Size", "Virtual Texture Size"))
+				.ToolTipText(LOCTEXT("Details_Size_Tooltip", "Virtual resolution derived from Size properties."))
+			]
+			.ValueContent()
+			[
+				SAssignNew(SizeText, STextBlock)
+			];
+		static const FText CustomRowPageTableSizeText = LOCTEXT("Details_RowFilter_PageTableSize", "Page Table Size");
+		DetailsCategory.AddCustomRow(CustomRowPageTableSizeText)
+			.NameContent()
+			[
+				SNew(STextBlock)
+				.Font(IDetailLayoutBuilder::GetDetailFont())
+				.Text(LOCTEXT("Details_PageTableSize", "Page Table Size"))
+				.ToolTipText(LOCTEXT("Details_PageTableSize_Tooltip", "Final page table size. This can vary according to the adaptive page table setting."))
+			]
+			.ValueContent()
+			[
+				SAssignNew(PageTableSizeText, STextBlock)
+			];
+	}
+
+	// Priority property :
+	{
+		IDetailCategoryBuilder& PerformanceCategory = DetailBuilder.EditCategory("Performance", FText::GetEmpty());
+		static const FText CurrentPriorityText = LOCTEXT("CurrentPriority", "Current Priority");
+		// Add the current priority before the custom priority : 
+		TSharedRef<IPropertyHandle> CustomPriorityHandle = DetailBuilder.GetProperty(GET_MEMBER_NAME_CHECKED(URuntimeVirtualTexture, CustomPriority));
+		CustomPriorityHandle->MarkHiddenByCustomization();
+		PerformanceCategory.AddCustomRow(/*FilterString = */CurrentPriorityText, /*bForAdvanced = */true)
+			.NameContent()
+			[
+				SNew(STextBlock)
+					.Font(IDetailLayoutBuilder::GetDetailFont())
+					.Text(CurrentPriorityText)
+					.IsEnabled(false)
+					.ToolTipText(LOCTEXT("CurrentPriority_Tooltip", "Defines the relative priority that this Runtime Virtual Texture has relative to other virtual texture producers.\n"
+						"This allows to get the pages from this virtual texture to update faster than others in case of high contention.\n"
+						"By default, it is derived from the material type but can be overridden with the Custom Priority property."))
+			]
+			.ValueContent()
+			[
+				SNew(STextBlock)
+					.Font(IDetailLayoutBuilder::GetDetailFont())
+					.Text_Lambda([this]()
+						{
+							if (URuntimeVirtualTexture* Texture = VirtualTexture.Get())
+							{
+								return UEnum::GetDisplayValueAsText(Texture->GetPriority());
+							}
+							return FText();
+						})
+					.ToolTipText_Lambda([this]()
+						{
+							if (URuntimeVirtualTexture* Texture = VirtualTexture.Get())
+							{
+								return Texture->GetUseCustomPriority() ? 
+									LOCTEXT("CurrentPriorityCustom_Tooltip", "Custom Priority") :
+									LOCTEXT("CurrentPriorityDefault_Tooltip", "Default Priority (inferred from the virtual texture content property)");
+							}
+							return FText();
+						})
+			];
+		PerformanceCategory.AddProperty(CustomPriorityHandle, EPropertyLocation::Advanced);
+	}
 
 	// Cache detail builder to refresh view updates
 	CachedDetailBuilder = &DetailBuilder;
diff --git a/Engine/Source/Editor/VirtualTexturingEditor/Private/RuntimeVirtualTextureThumbnailRenderer.cpp b/Engine/Source/Editor/VirtualTexturingEditor/Private/RuntimeVirtualTextureThumbnailRenderer.cpp
--- a/Engine/Source/Editor/VirtualTexturingEditor/Private/RuntimeVirtualTextureThumbnailRenderer.cpp
+++ b/Engine/Source/Editor/VirtualTexturingEditor/Private/RuntimeVirtualTextureThumbnailRenderer.cpp
@@ -25,6 +25,7 @@
 #include "VT/RuntimeVirtualTexture.h"
 #include "VT/RuntimeVirtualTextureEnum.h"
 #include "VT/RuntimeVirtualTextureRender.h"
+#include "VirtualTextureEnum.h"
 #include "VirtualTexturing.h"
 
 class FCanvas;
@@ -94,6 +95,8 @@ void URuntimeVirtualTextureThumbnailRenderer::Draw(UObject* Object, int32 X, int
 
 	FVTProducerDescription VTDesc;
 	RuntimeVirtualTexture->GetProducerDescription(VTDesc, URuntimeVirtualTexture::FInitSettings(), Transform);
+	// No need for thumbnails to be produced before anything else : 
+	VTDesc.Priority = EVTProducerPriority::Lowest;
 	const int32 MaxLevel = (int32)FMath::CeilLogTwo(FMath::Max(VTDesc.BlockWidthInTiles, VTDesc.BlockHeightInTiles));
 
 	UE::RenderCommandPipe::FSyncScope SyncScope;
diff --git a/Engine/Source/Runtime/Engine/Classes/Components/RuntimeVirtualTextureComponent.h b/Engine/Source/Runtime/Engine/Classes/Components/RuntimeVirtualTextureComponent.h
--- a/Engine/Source/Runtime/Engine/Classes/Components/RuntimeVirtualTextureComponent.h
+++ b/Engine/Source/Runtime/Engine/Classes/Components/RuntimeVirtualTextureComponent.h
@@ -8,6 +8,8 @@
 #include "RenderCommandFence.h"
 #include "SceneComponent.h"
 #include "SceneTypes.h"
+#include "VT/RuntimeVirtualTextureEnum.h"
+#include "VirtualTextureEnum.h"
 #include "RuntimeVirtualTextureComponent.generated.h"
 
 class URuntimeVirtualTexture;
@@ -31,6 +33,8 @@ class URuntimeVirtualTextureComponent : public USceneComponent
 {
 	GENERATED_UCLASS_BODY()
 
+	friend class FRuntimeVirtualTextureComponentDetailsCustomization;
+
 protected:
 	/** Actor to align rotation to. If set this actor is always included in the bounds calculation. */
 	UPROPERTY(EditAnywhere, Category = VolumeBounds)
@@ -136,9 +140,11 @@ public:
 	/**
 	 * This function marks an area of the runtime virtual texture as dirty.
 	 * @param WorldBounds : The world space bounds of the pages to invalidate.
+	 * @param InvalidatePriority Allows the pages affected by this area to get processed in priority. 
+	 * This allows increased responsiveness when there are more pages being updated than can be handled in a given frame (throttling) 
 	 */
 	UFUNCTION(BlueprintCallable, Category = VirtualTexture)
-	ENGINE_API void Invalidate(FBoxSphereBounds const& WorldBounds);
+	ENGINE_API void Invalidate(FBoxSphereBounds const& WorldBounds, EVTInvalidatePriority InvalidatePriority = EVTInvalidatePriority::High);
 
 	/** Set the runtime virtual texture object on this component. */
 	ENGINE_API void SetVirtualTexture(URuntimeVirtualTexture* InVirtualTexture);
@@ -185,6 +191,8 @@ public:
 	/** Returns whether the StreamingTexure contents is valid for use for the specified shading path */
 	ENGINE_API bool IsStreamingTextureInvalid(EShadingPath ShadingPath) const;
 
+	ENGINE_API EVTProducerPriority GetPriority() const;
+
 	enum class EStreamingTextureStatusFlags : uint8
 	{
 		None = 0,
diff --git a/Engine/Source/Runtime/Engine/Classes/VT/RuntimeVirtualTexture.h b/Engine/Source/Runtime/Engine/Classes/VT/RuntimeVirtualTexture.h
--- a/Engine/Source/Runtime/Engine/Classes/VT/RuntimeVirtualTexture.h
+++ b/Engine/Source/Runtime/Engine/Classes/VT/RuntimeVirtualTexture.h
@@ -6,6 +6,7 @@
 #include "Engine/TextureDefines.h"
 #include "VirtualTexturing.h"
 #include "VT/RuntimeVirtualTextureEnum.h"
+#include "VirtualTextureEnum.h"
 #include "RuntimeVirtualTexture.generated.h"
 
 enum EShaderPlatform : uint16;
@@ -50,6 +51,13 @@ protected:
 	UPROPERTY(EditAnywhere, BlueprintReadOnly, Category = Performance, meta = (DisplayName = "Use Low Quality Compression", editcondition = "MaterialType == ERuntimeVirtualTextureMaterialType::BaseColor_Normal_Roughness && bCompressTextures == true"))
 	bool bUseLowQualityCompression = false;
 
+	/** Allows to override the default priority that this runtime virtual texture has, relative to other virtual texture producers. This allows to get the pages from this virtual texture to update faster than others in case of high contention. */
+	UPROPERTY(EditAnywhere, AdvancedDisplay, Category = Performance, meta = (EditCondition = bUseCustomPriority))
+	EVTProducerPriority CustomPriority;
+
+	UPROPERTY(EditAnywhere, AdvancedDisplay, Category = Performance, meta = (InlineEditConditionToggle)) 
+	bool bUseCustomPriority = false;
+
 	/** Enable clear before rendering a page of the virtual texture. Disabling this can be an optimization if you know that the texture will always be fully covered by rendering.  */
 	UPROPERTY(EditAnywhere, AdvancedDisplay, Category = Layout, meta = (DisplayName = "Enable clear before render"))
 	bool bClearTextures = true;
@@ -166,6 +174,13 @@ public:
 	/** Getter for the shader uniform parameter type. */
 	static ENGINE_API UE::Shader::EValueType GetUniformParameterType(int32 Index);
 
+	ENGINE_API EVTProducerPriority GetPriority() const;
+	ENGINE_API EVTProducerPriority GetDefaultPriority() const;
+	ENGINE_API bool GetUseCustomPriority() const
+	{
+		return bUseCustomPriority;
+	}
+
 protected:
 	/** Initialize the render resources. This kicks off render thread work. */
 	ENGINE_API void InitResource(IVirtualTexture* InProducer, FVTProducerDescription const& InProducerDesc);
diff --git a/Engine/Source/Runtime/Engine/Private/Components/RuntimeVirtualTextureComponent.cpp b/Engine/Source/Runtime/Engine/Private/Components/RuntimeVirtualTextureComponent.cpp
--- a/Engine/Source/Runtime/Engine/Private/Components/RuntimeVirtualTextureComponent.cpp
+++ b/Engine/Source/Runtime/Engine/Private/Components/RuntimeVirtualTextureComponent.cpp
@@ -22,6 +22,7 @@
 #include "VT/VirtualTexture.h"
 #include "VT/VirtualTextureBuilder.h"
 #include "VT/VirtualTextureBuiltData.h"
+#include "VisualLogger/VisualLogger.h"
 
 #include UE_INLINE_GENERATED_CPP_BY_NAME(RuntimeVirtualTextureComponent)
 
@@ -55,6 +56,12 @@ static TAutoConsoleVariable<bool> CVarVTStreamingMipsUseAlways(
 	FConsoleVariableDelegate::CreateLambda([](IConsoleVariable* InVariable) { FGlobalComponentRecreateRenderStateContext Context; }),
 	ECVF_Default);
 
+static TAutoConsoleVariable<bool> CVarVTVisLogManualInvalidations(
+	TEXT("r.VT.RVT.VisLogManualInvalidations"),
+	false,
+	TEXT("Outputs the RVT component invalidations and their bounds in the Visual Logger"),
+	ECVF_Default);
+
 
 URuntimeVirtualTextureComponent::URuntimeVirtualTextureComponent(const FObjectInitializer& ObjectInitializer)
 	: Super(ObjectInitializer)
@@ -220,11 +227,14 @@ bool URuntimeVirtualTextureComponent::IsEnabledInScene() const
 	return true;
 }
 
-void URuntimeVirtualTextureComponent::Invalidate(FBoxSphereBounds const& InWorldBounds)
+void URuntimeVirtualTextureComponent::Invalidate(FBoxSphereBounds const& InWorldBounds, EVTInvalidatePriority InvalidatePriority)
 {
 	if (GetScene() != nullptr)
 	{
-		GetScene()->InvalidateRuntimeVirtualTexture(this, InWorldBounds);
+		GetScene()->InvalidateRuntimeVirtualTexture(this, InWorldBounds, InvalidatePriority);
+
+		UE_CVLOG_BOX(CVarVTVisLogManualInvalidations.GetValueOnGameThread(), this, LogVirtualTexturing, Log, InWorldBounds.GetBox(), (InvalidatePriority == EVTInvalidatePriority::High) ? FColor::Red : FColor::Green,
+			TEXT("Invalidate RVT for %s(%s) %s"), *GetOwner()->GetActorNameOrLabel(), VirtualTexture ? *VirtualTexture->GetName() : TEXT_EMPTY, (InvalidatePriority == EVTInvalidatePriority::High) ? TEXT("(prioritized)"): TEXT_EMPTY);
 	}
 }
 
@@ -371,6 +381,12 @@ URuntimeVirtualTextureComponent::EStreamingTextureStatusFlags URuntimeVirtualTex
 	return Status;
 }
 
+EVTProducerPriority URuntimeVirtualTextureComponent::GetPriority() const
+{
+	return (VirtualTexture != nullptr) ? VirtualTexture->GetPriority() : EVTProducerPriority::Normal;
+}
+
+
 #if WITH_EDITOR
 
 bool URuntimeVirtualTextureComponent::IsStreamingTextureInvalid() const
diff --git a/Engine/Source/Runtime/Engine/Private/Texture2D.cpp b/Engine/Source/Runtime/Engine/Private/Texture2D.cpp
--- a/Engine/Source/Runtime/Engine/Private/Texture2D.cpp
+++ b/Engine/Source/Runtime/Engine/Private/Texture2D.cpp
@@ -46,6 +46,7 @@
 #include "Misc/ScopedSlowTask.h"
 #include "UObject/ObjectSaveContext.h"
 #include "UObject/StrongObjectPtr.h"
+#include "VirtualTextureEnum.h"
 #include "ProfilingDebugging/AssetMetadataTrace.h"
 
 #if WITH_EDITOR
@@ -1597,6 +1598,8 @@ void FVirtualTexture2DResource::InitRHI(FRHICommandListBase& RHICmdList)
 		ProducerDesc.PhysicalGroupIndex[LayerIndex] = bSinglePhysicalSpace ? 0 : LayerIndex;
 		ProducerDesc.bIsLayerSRGB[LayerIndex] = bSRGB;
 	}
+	// TODO [jonathan.bard] : For now, let's use the normal priority, but that could come from the texture asset and/or the texture group too : 
+	ProducerDesc.Priority = EVTProducerPriority::Normal;
 
 	FUploadingVirtualTexture* VirtualTexture = new FUploadingVirtualTexture(ProducerDesc.Name, VTData, FirstMipToUse);
 	ProducerHandle = GetRendererModule().RegisterVirtualTextureProducer(RHICmdList, ProducerDesc, VirtualTexture);
diff --git a/Engine/Source/Runtime/Engine/Private/VT/RuntimeVirtualTexture.cpp b/Engine/Source/Runtime/Engine/Private/VT/RuntimeVirtualTexture.cpp
--- a/Engine/Source/Runtime/Engine/Private/VT/RuntimeVirtualTexture.cpp
+++ b/Engine/Source/Runtime/Engine/Private/VT/RuntimeVirtualTexture.cpp
@@ -364,6 +364,8 @@ void URuntimeVirtualTexture::GetProducerDescription(FVTProducerDescription& OutD
 		OutDesc.PhysicalGroupIndex[Layer] = bSinglePhysicalSpace ? 0 : Layer;
 		OutDesc.bIsLayerSRGB[Layer] = IsLayerSRGB(Layer);
 	}
+
+	OutDesc.Priority = GetPriority();
 }
 
 int32 URuntimeVirtualTexture::GetLayerCount(ERuntimeVirtualTextureMaterialType InMaterialType)
@@ -572,6 +574,34 @@ UE::Shader::EValueType URuntimeVirtualTexture::GetUniformParameterType(int32 Ind
 	return UE::Shader::EValueType::Float4;
 }
 
+EVTProducerPriority URuntimeVirtualTexture::GetPriority() const
+{
+	return bUseCustomPriority ? CustomPriority : GetDefaultPriority();
+}
+
+EVTProducerPriority URuntimeVirtualTexture::GetDefaultPriority() const
+{
+	switch (MaterialType)
+	{
+	case ERuntimeVirtualTextureMaterialType::BaseColor:
+	case ERuntimeVirtualTextureMaterialType::BaseColor_Normal_Roughness:
+	case ERuntimeVirtualTextureMaterialType::BaseColor_Normal_Specular:
+	case ERuntimeVirtualTextureMaterialType::BaseColor_Normal_Specular_YCoCg:
+	case ERuntimeVirtualTextureMaterialType::BaseColor_Normal_Specular_Mask_YCoCg:
+		return EVTProducerPriority::High;
+	case ERuntimeVirtualTextureMaterialType::Mask4:
+		return EVTProducerPriority::BelowNormal;
+	case ERuntimeVirtualTextureMaterialType::WorldHeight:
+		return EVTProducerPriority::Low;
+	case ERuntimeVirtualTextureMaterialType::Displacement:
+		return EVTProducerPriority::Lower;
+	default:
+		checkNoEntry();
+	}
+
+	return EVTProducerPriority::Normal;
+}
+
 void URuntimeVirtualTexture::Initialize(IVirtualTexture* InProducer, FVTProducerDescription const& InProducerDesc, FTransform const& InVolumeToWorld, FBox const& InWorldBounds)
 {
 	WorldToUVTransformParameters[0] = InVolumeToWorld.GetTranslation();
diff --git a/Engine/Source/Runtime/Engine/Private/VT/VirtualTextureAdapter.cpp b/Engine/Source/Runtime/Engine/Private/VT/VirtualTextureAdapter.cpp
--- a/Engine/Source/Runtime/Engine/Private/VT/VirtualTextureAdapter.cpp
+++ b/Engine/Source/Runtime/Engine/Private/VT/VirtualTextureAdapter.cpp
@@ -12,6 +12,7 @@
 #include "RHIStaticStates.h"
 #include "ShaderParameterStruct.h"
 #include "TextureResource.h"
+#include "VirtualTextureEnum.h"
 #include "VirtualTexturing.h"
 #include "VT/VirtualTextureBuildSettings.h"
 
@@ -537,6 +538,8 @@ public:
 		ProducerDesc.NumTextureLayers = 1;
 		ProducerDesc.NumPhysicalGroups = 1;
 		ProducerDesc.LayerFormat[0] = Format;
+		// TODO [jonathan.bard] : For now, let's use the normal priority, but that could come from the texture asset and/or the texture group too : 
+		ProducerDesc.Priority = EVTProducerPriority::Normal;
 
 		FVirtualTextureAdapterProducer* VirtualTexture = new FVirtualTextureAdapterProducer(RHICmdList, SourceResource->GetTexture2DRHI(), UnstreamedMipCount, ProducerDesc);
 		ProducerHandle = GetRendererModule().RegisterVirtualTextureProducer(RHICmdList, ProducerDesc, VirtualTexture);
diff --git a/Engine/Source/Runtime/Engine/Public/SceneInterface.h b/Engine/Source/Runtime/Engine/Public/SceneInterface.h
--- a/Engine/Source/Runtime/Engine/Public/SceneInterface.h
+++ b/Engine/Source/Runtime/Engine/Public/SceneInterface.h
@@ -54,6 +54,7 @@ class USkyLightComponent;
 class UStaticMeshComponent;
 class UTextureCube;
 enum class EPrimitiveDirtyState : uint8;
+enum class EVTInvalidatePriority : uint8;
 enum class EShadingPath;
 enum EShaderPlatform : uint16;
 namespace ERHIFeatureLevel { enum Type : int; }
@@ -293,9 +294,17 @@ public:
 	/* Get the bitmasks describing which virtual texture objects will hide the associated primitives. */
 	virtual void GetRuntimeVirtualTextureHidePrimitiveMask(uint8& bHideMaskEditor, uint8& bHideMaskGame) const {}
 
-	/** Invalidates pages in a runtime virtual texture object. */
+	UE_DEPRECATED(5.6, "Use the alternate prototype with the Priority param")
 	virtual void InvalidateRuntimeVirtualTexture(class URuntimeVirtualTextureComponent* Component, FBoxSphereBounds const& WorldBounds) {}
 
+	/** Invalidates pages in a runtime virtual texture object. 
+	 * @param Component RVT component containing the runtime virtual texture to invalidate
+	 * @param WorldBounds World-space bounds of the area of the runtime virtual texture to invalidate
+	 * @param InvalidatePriority Ensures the corresponding pages get processed in priority. This allows increased responsiveness when 
+	 *  there are more pages being updated than can be handled in a given frame (when throttling) 
+	 */
+	virtual void InvalidateRuntimeVirtualTexture(class URuntimeVirtualTextureComponent* Component, FBoxSphereBounds const& WorldBounds, EVTInvalidatePriority InvalidatePriority) {}
+
 	/** Mark scene as needing to restart path tracer accumulation. */
 	virtual void InvalidatePathTracedOutput() {}
 
diff --git a/Engine/Source/Runtime/Landscape/Classes/LandscapeInfo.h b/Engine/Source/Runtime/Landscape/Classes/LandscapeInfo.h
--- a/Engine/Source/Runtime/Landscape/Classes/LandscapeInfo.h
+++ b/Engine/Source/Runtime/Landscape/Classes/LandscapeInfo.h
@@ -296,6 +296,11 @@ public:
 	 * @return true if the object could not be marked dirty but has been added to the list of ModifiedPackages
 	 */
 	LANDSCAPE_API bool MarkObjectDirty(UObject* InObject, bool bInForceResave = false, const ALandscape* InLandscapeOverride = nullptr);
+
+	/** Helper function to invalidate the runtime virtual texture data associated with the specified landscape region. 
+	 *   This forces the dirtied RVT pages to be updated in priority
+	 */
+	LANDSCAPE_API void DirtyRuntimeVirtualTextureForLandscapeArea(int32 X1, int32 Y1, int32 X2, int32 Y2) const;
 #endif //WITH_EDITOR
 
 	/**
diff --git a/Engine/Source/Runtime/Landscape/Private/Landscape.cpp b/Engine/Source/Runtime/Landscape/Private/Landscape.cpp
--- a/Engine/Source/Runtime/Landscape/Private/Landscape.cpp
+++ b/Engine/Source/Runtime/Landscape/Private/Landscape.cpp
@@ -16,6 +16,7 @@ Landscape.cpp: Terrain rendering
 #include "UObject/DevObjectVersion.h"
 #include "UObject/LinkerLoad.h"
 #include "UObject/Package.h"
+#include "Components/RuntimeVirtualTextureComponent.h"
 #include "Framework/Application/SlateApplication.h"
 #include "LandscapePrivate.h"
 #include "LandscapeStreamingProxy.h"
@@ -237,6 +238,12 @@ static FAutoConsoleVariableRef CVarLandscapeHeightmapCompressionMode(
 	TEXT(" -1: force disable heightmap compression on all landscapes\n"),
 	ECVF_Preview | ECVF_ReadOnly);
 
+bool GLandscapePrioritizeDirtyRVTPages = true;
+static FAutoConsoleVariableRef CVarPrioritizeDirtyRVTPages(
+	TEXT("landscape.PrioritizeDirtyRVTPages"),
+	GLandscapePrioritizeDirtyRVTPages,
+	TEXT("Prioritize RVT pages affected by the landscape tools, so that they get updated prior to others. Improves reactiveness when invalidating large areas of the RVT.")
+);
 #endif // WITH_EDITOR
 
 TAutoConsoleVariable<int32> CVarRenderNaniteLandscape(
@@ -5187,6 +5194,49 @@ bool ULandscapeInfo::ModifyObject(UObject* InObject, bool bAlwaysMarkDirty)
 	return bWasAddedToModifiedPackages;
 }
 
+void ULandscapeInfo::DirtyRuntimeVirtualTextureForLandscapeArea(int32 X1, int32 Y1, int32 X2, int32 Y2) const
+{
+	FBox DirtyWorldBounds(ForceInit);
+
+	// Iterate touched components to find touched runtime virtual textures.
+	TSet<ULandscapeComponent*> Components;
+	GetComponentsInRegion(X1 + 1, Y1 + 1, X2 - 1, Y2 - 1, Components);
+
+	TArray<URuntimeVirtualTexture*> RuntimeVirtualTextures;
+	for (ULandscapeComponent* Component : Components)
+	{
+		ALandscapeProxy* Landscape = Component->GetLandscapeProxy();
+		if (Landscape != nullptr && Landscape->RuntimeVirtualTextures.Num() > 0)
+		{
+			for (URuntimeVirtualTexture* RuntimeVirtualTexture : Landscape->RuntimeVirtualTextures)
+			{
+				RuntimeVirtualTextures.AddUnique(RuntimeVirtualTexture);
+			}
+
+			// Also accumulate bounds in world space.
+			const int32 LocalX1 = FMath::Max(X1, Component->GetSectionBase().X) - Component->GetSectionBase().X;
+			const int32 LocalY1 = FMath::Max(Y1, Component->GetSectionBase().Y) - Component->GetSectionBase().Y;
+			const int32 LocalX2 = FMath::Min(X2, Component->GetSectionBase().X + ComponentSizeQuads) - Component->GetSectionBase().X;
+			const int32 LocalY2 = FMath::Min(Y2, Component->GetSectionBase().Y + ComponentSizeQuads) - Component->GetSectionBase().Y;
+			const FBox LocalDirtyBounds(FVector(LocalX1, LocalY1, 0), FVector(LocalX2, LocalY2, 1));
+
+			DirtyWorldBounds += LocalDirtyBounds.TransformBy(Component->GetComponentToWorld());
+		}
+	}
+
+	// Find matching runtime virtual texture components and invalidate dirty region.
+	if (RuntimeVirtualTextures.Num() > 0)
+	{
+		for (TObjectIterator<URuntimeVirtualTextureComponent> It(/*AdditionalExclusionFlags = */RF_ClassDefaultObject, /*bIncludeDerivedClasses = */true, /*InInternalExclusionFlags = */EInternalObjectFlags::Garbage); It; ++It)
+		{
+			if (It->GetVirtualTexture() != nullptr && RuntimeVirtualTextures.Contains(It->GetVirtualTexture()))
+			{
+				It->Invalidate(FBoxSphereBounds(DirtyWorldBounds), GLandscapePrioritizeDirtyRVTPages ? EVTInvalidatePriority::High : EVTInvalidatePriority::Normal);
+			}
+		}
+	}
+}
+
 ALandscapeProxy* ULandscapeInfo::GetLandscapeProxyForLevel(ULevel* Level) const
 {
 	TRACE_CPUPROFILER_EVENT_SCOPE(ULandscapeInfo::GetLandscapeProxyForLevel);
diff --git a/Engine/Source/Runtime/Landscape/Private/LandscapeEditLayers.cpp b/Engine/Source/Runtime/Landscape/Private/LandscapeEditLayers.cpp
--- a/Engine/Source/Runtime/Landscape/Private/LandscapeEditLayers.cpp
+++ b/Engine/Source/Runtime/Landscape/Private/LandscapeEditLayers.cpp
@@ -4235,14 +4235,14 @@ void ALandscape::InvalidateRVTForTextures(const TSet<UTexture2D*>& InTextures)
 						{
 							// Landscape only supports UPrimitiveComponent for the moment
 							if (UPrimitiveComponent* PrimitiveComponent = PrimitiveComponentInterface->GetUObject<UPrimitiveComponent>())
-						{
-							PrimitiveComponentsToInvalidate.Add(PrimitiveComponent);
+							{
+								PrimitiveComponentsToInvalidate.Add(PrimitiveComponent);
+							}
 						}
 					}
 				}
 			}
 		}
-		}
 
 		if (!PrimitiveComponentsToInvalidate.IsEmpty())
 		{
@@ -4253,7 +4253,7 @@ void ALandscape::InvalidateRVTForTextures(const TSet<UTexture2D*>& InTextures)
 				{
 					if (PrimitiveComponent->GetRuntimeVirtualTextures().Contains(It->GetVirtualTexture()))
 					{
-						It->Invalidate(FBoxSphereBounds(PrimitiveComponent->Bounds));
+						It->Invalidate(FBoxSphereBounds(PrimitiveComponent->Bounds), EVTInvalidatePriority::Normal);
 					}
 				}
 			}
diff --git a/Engine/Source/Runtime/RenderCore/Private/VirtualTexturing.cpp b/Engine/Source/Runtime/RenderCore/Private/VirtualTexturing.cpp
--- a/Engine/Source/Runtime/RenderCore/Private/VirtualTexturing.cpp
+++ b/Engine/Source/Runtime/RenderCore/Private/VirtualTexturing.cpp
@@ -1,5 +1,24 @@
 // Copyright Epic Games, Inc. All Rights Reserved.
 
 #include "VirtualTexturing.h"
+#include "VirtualTextureEnum.h"
 
 DEFINE_LOG_CATEGORY(LogVirtualTexturing);
+
+
+// ----------------------------------------------------------------------------------
+
+FVTProducerDescription::FVTProducerDescription()
+	: Priority(EVTProducerPriority::Normal) // Cannot inline-initialize because EVTProducerPriority is forward-declared
+{}
+
+
+// ----------------------------------------------------------------------------------
+
+void FVirtualTextureLocalTileRequest::ValidatePriorities() const
+{
+	static_assert(static_cast<uint64>(EVTProducerPriority::Count) <= (1 << 3), "EVTProducerPriority should be packable on 3 bits");
+	static_assert(static_cast<uint64>(EVTInvalidatePriority::Count) <= (1 << 1), "EVTInvalidatePriority should be packable on 1 bit");
+	checkSlow(ProducerPriority < (1 << 3));
+	checkSlow(InvalidatePriority < (1 << 1));
+}
diff --git a/Engine/Source/Runtime/RenderCore/Public/VirtualTextureEnum.h b/Engine/Source/Runtime/RenderCore/Public/VirtualTextureEnum.h
new file mode 100644
--- /dev/null
+++ b/Engine/Source/Runtime/RenderCore/Public/VirtualTextureEnum.h
@@ -0,0 +1,37 @@
+// Copyright Epic Games, Inc. All Rights Reserved.
+
+#pragma once
+
+#include "Misc/EnumRange.h"
+#include "VirtualTextureEnum.generated.h"
+
+/** Enumeration of the priority assigned to a given virtual texture producer. Must match EVTProducerPriority in RenderCore (they are duplicate because RenderCore cannot declare UENUMs) */
+UENUM()
+enum class EVTProducerPriority : uint8
+{
+	Lowest = 0,
+	Lower,
+	Low,
+	BelowNormal,
+	Normal,
+	AboveNormal,
+	High,
+	Highest, 
+
+	Count UMETA(Hidden)
+};
+static_assert(static_cast<uint8>(EVTProducerPriority::Count) == static_cast<uint8>(EVTProducerPriority::Count), "EVTProducerPriority is the UENUM counterpart to EVTProducerPriority");
+ENUM_RANGE_BY_COUNT(EVTProducerPriority, EVTProducerPriority::Count);
+
+
+/** Describes whether an invalidated VT area should be prioritized against others. Allows to improve reactiveness when invalidating a large number of pages */
+UENUM()
+enum class EVTInvalidatePriority : uint8
+{
+	Normal = 0,
+	High,
+
+	Count UMETA(Hidden)
+};
+static_assert(static_cast<uint8>(EVTInvalidatePriority::Count) == static_cast<uint8>(EVTInvalidatePriority::Count), "EVTInvalidatePriority is the UENUM counterpart to EVTInvalidatePriority");
+ENUM_RANGE_BY_COUNT(EVTInvalidatePriority, EVTInvalidatePriority::Count);
diff --git a/Engine/Source/Runtime/RenderCore/Public/VirtualTexturing.h b/Engine/Source/Runtime/RenderCore/Public/VirtualTexturing.h
--- a/Engine/Source/Runtime/RenderCore/Public/VirtualTexturing.h
+++ b/Engine/Source/Runtime/RenderCore/Public/VirtualTexturing.h
@@ -4,6 +4,7 @@
 
 #include "Async/TaskGraphInterfaces.h"
 #include "Containers/EnumAsByte.h"
+#include "Containers/HashTable.h"
 #include "CoreMinimal.h"
 #include "Logging/LogMacros.h"
 #include "Math/Color.h"
@@ -24,6 +25,8 @@ class FRHIShaderResourceView;
 class FRHITexture;
 class FRHIUnorderedAccessView;
 namespace ERHIFeatureLevel { enum Type : int; }
+enum class EVTProducerPriority : uint8;
+enum class EVTInvalidatePriority : uint8;
 
 union FVirtualTextureProducerHandle
 {
@@ -139,6 +142,8 @@ struct FAllocatedVTDescription
 
 struct FVTProducerDescription
 {
+	RENDERCORE_API FVTProducerDescription();
+
 	FName Name; /** Will be name of UTexture for streaming VTs, mostly here for debugging */
 	uint32 FullNameHash;
 	
@@ -176,6 +181,9 @@ struct FVTProducerDescription
 	
 	uint8 NumPhysicalGroups = 0u;
 	uint8 PhysicalGroupIndex[VIRTUALTEXTURE_SPACE_MAXLAYERS] = { 0 };
+
+	/** Priority allows to process a producer before another which increases responsiveness in high contention situations */
+	EVTProducerPriority Priority;
 };
 
 typedef void (FVTProducerDestroyedFunction)(const FVirtualTextureProducerHandle& InHandle, void* Baton);
@@ -238,17 +246,13 @@ struct FVTRequestPageResult
 /** Describes a location to write a single layer of a VT tile */
 struct FVTProduceTargetLayer
 {
-PRAGMA_DISABLE_DEPRECATION_WARNINGS
 	FVTProduceTargetLayer() = default;
 	FVTProduceTargetLayer(const FVTProduceTargetLayer&) = default;
 	FVTProduceTargetLayer& operator=(const FVTProduceTargetLayer&) = default;
-PRAGMA_ENABLE_DEPRECATION_WARNINGS
 
 	/** The texture to write to. */
 	FRHITexture* TextureRHI = nullptr;
 	
-	UE_DEPRECATED(5.1, "UnorderedAccessViewRHI is deprecated. Register the pooled render target with RDG instead.")
-	FRHIUnorderedAccessView* UnorderedAccessViewRHI = nullptr;
 	/**
 	 * Pooled render target. For FRDGBuilder::RegisterExternalTexture() which only accepts pooled render targets.
 	 * To avoid cost of manipulating ref counting pointers a raw pointer is used instead - it is valid until returning from your Finalize().
@@ -298,9 +302,6 @@ public:
 	*/
 	virtual FVTRequestPageResult RequestPageData(FRHICommandList& RHICmdList, const FVirtualTextureProducerHandle& ProducerHandle, uint8 LayerMask, uint8 vLevel, uint64 vAddress, EVTRequestPagePriority Priority) = 0;
 
-	UE_DEPRECATED(5.2, "RequestPageData now requires an RHI command list.")
-	virtual FVTRequestPageResult RequestPageData(const FVirtualTextureProducerHandle& ProducerHandle, uint8 LayerMask, uint8 vLevel, uint64 vAddress, EVTRequestPagePriority Priority) { return {}; }
-
 	/**
 	* Upload page data to the cache, data must have been previously requested, and reported either 'Available' or 'Pending'
 	* The system will attempt to call RequestPageData/ProducePageData only once for a given vLevel/vAddress, with all the requested layers set in LayerMask,
@@ -475,27 +476,105 @@ struct FAdaptiveVTDescription
 union FVirtualTextureLocalTile
 {
 	inline FVirtualTextureLocalTile() {}
-	inline FVirtualTextureLocalTile(const FVirtualTextureProducerHandle& InProducerHandle, uint32 InLocal_vAddress, uint8 InLocal_vLevel)
-		: PackedProducerHandle(InProducerHandle.PackedValue), Local_vAddress(InLocal_vAddress), Local_vLevel(InLocal_vLevel), Pad(0)
-	{}
+	explicit inline FVirtualTextureLocalTile(const FVirtualTextureProducerHandle& InProducerHandle, uint32 InLocal_vAddress, uint8 InLocal_vLevel)
+		: PackedProducerHandle(InProducerHandle.PackedValue)
+		, Local_vAddress(InLocal_vAddress)
+		, Local_vLevel(InLocal_vLevel)
+		, Pad(0)
+	{
+	}
+	explicit inline FVirtualTextureLocalTile(uint64 InPackedValue)
+		: PackedValue(InPackedValue)
+	{
+		Pad = 0;
+	}
 
 	inline FVirtualTextureProducerHandle GetProducerHandle() const { return FVirtualTextureProducerHandle(PackedProducerHandle); }
 
-	uint64 PackedValue;
+	uint64 PackedValue = 0;
 	struct
 	{
+		// Core "identity" data that identify a given tile
 		uint32 PackedProducerHandle;
 		uint32 Local_vAddress : 24;
 		uint32 Local_vLevel : 4;
-		uint32 Pad : 4;
+		uint32 Pad : 4;  // Important : this has to be kept in sync with FVirtualTextureLocalTileRequest : both are packed in a uint64 and the latter stores additional info in these bits
 	};
 
 	friend inline bool operator==(const FVirtualTextureLocalTile& Lhs, const FVirtualTextureLocalTile& Rhs) { return Lhs.PackedValue == Rhs.PackedValue; }
 	friend inline bool operator!=(const FVirtualTextureLocalTile& Lhs, const FVirtualTextureLocalTile& Rhs) { return Lhs.PackedValue != Rhs.PackedValue; }
-	friend inline uint64 GetTypeHash(const FVirtualTextureLocalTile& T) { return T.PackedValue; }
+	friend inline uint32 GetTypeHash(const FVirtualTextureLocalTile& T) { return GetTypeHash(T.PackedValue); } 
+	inline uint16 GetMurmurHash() const { return static_cast<uint16>(MurmurFinalize64(PackedValue)); }
 };
 static_assert(sizeof(FVirtualTextureLocalTile) == sizeof(uint64), "Bad packing");
 
+/**
+ * Identifies a VT tile request : a VT tile, along with priority information for processing this request
+ */
+union FVirtualTextureLocalTileRequest
+{
+	explicit inline FVirtualTextureLocalTileRequest() {}
+	explicit inline FVirtualTextureLocalTileRequest(const FVirtualTextureLocalTile& InTile, EVTProducerPriority InProducerPriority, EVTInvalidatePriority InInvalidatePriority)
+		: TilePackedValue(InTile.PackedValue)
+		, ProducerPriority(static_cast<uint64>(InProducerPriority))
+		, InvalidatePriority(static_cast<uint64>(InInvalidatePriority))
+	{
+		checkSlow(InTile == GetTile());
+		ValidatePriorities();
+	}
+
+	inline FVirtualTextureLocalTile GetTile() const { return FVirtualTextureLocalTile(TilePackedValue); }
+	inline EVTProducerPriority GetProducerPriority() const { return static_cast<EVTProducerPriority>(ProducerPriority); }
+	inline EVTInvalidatePriority GetInvalidatePriority() const { return static_cast<EVTInvalidatePriority>(InvalidatePriority); }
+
+	inline FVirtualTextureLocalTileRequest& MergeWith(const FVirtualTextureLocalTileRequest& InOther) 
+	{	
+		// Keep the highest priorities of the 2 :
+		ProducerPriority = FMath::Max(ProducerPriority, InOther.ProducerPriority);
+		InvalidatePriority = FMath::Max(InvalidatePriority, InOther.InvalidatePriority);
+
+		checkSlow(*this == InOther);
+		checkSlow(GetTile() == InOther.GetTile());
+		ValidatePriorities();
+		return *this;
+	}
+
+	uint64 PackedValue;
+	struct
+	{
+		// Core "identity" data that identifies a given tile : 
+		uint64 TilePackedValue : 60;
+		// Additional data about the request :
+		uint64 ProducerPriority : 3; // See EVTProducerPriority 
+		uint64 InvalidatePriority : 1; // Indicates whether the tile should get prioritized over others
+	};
+
+	friend inline bool operator==(const FVirtualTextureLocalTileRequest& Lhs, const FVirtualTextureLocalTileRequest& Rhs) { return Lhs.GetIdentity() == Rhs.GetIdentity(); }
+	friend inline bool operator!=(const FVirtualTextureLocalTileRequest& Lhs, const FVirtualTextureLocalTileRequest& Rhs) { return Lhs.GetIdentity() != Rhs.GetIdentity(); }
+	friend inline uint32 GetTypeHash(const FVirtualTextureLocalTileRequest& T) { return GetTypeHash(T.GetIdentity()); }
+	inline uint64 GetIdentity() const { return TilePackedValue; }  // 2 requests from the same tile are considered identical (it's up to the user to "merge" them if necessary)
+	inline uint16 GetMurmurHash() const { return static_cast<uint16>(MurmurFinalize64(GetIdentity())); } 
+
+private:
+	RENDERCORE_API void ValidatePriorities() const;
+};
+static_assert(sizeof(FVirtualTextureLocalTileRequest) == sizeof(uint64), "Bad packing");
+
+/**
+ * Helper function that tries to add a new FVirtualTextureLocalTileRequest to a TSet but, in case this tile is already in the set, merges it with the existing request :
+ */
+inline FVirtualTextureLocalTileRequest& AddOrMergeTileRequest(const FVirtualTextureLocalTileRequest& InTileRequest, TSet<FVirtualTextureLocalTileRequest>& TileRequests)
+{
+	bool bAlreadyInSet = false;
+	FVirtualTextureLocalTileRequest& Result = TileRequests.FindOrAdd(InTileRequest, &bAlreadyInSet);
+	if (bAlreadyInSet)
+	{
+		// If the request already exists, merge with this new request, to keep the highest priority :
+		Result.MergeWith(InTileRequest);
+	}
+	return Result;
+}
+
 RENDERCORE_API DECLARE_LOG_CATEGORY_EXTERN(LogVirtualTexturing, Log, All);
 
 DECLARE_STATS_GROUP(TEXT("Virtual Texturing"), STATGROUP_VirtualTexturing, STATCAT_Advanced);
diff --git a/Engine/Source/Runtime/RenderCore/RenderCore.Build.cs b/Engine/Source/Runtime/RenderCore/RenderCore.Build.cs
--- a/Engine/Source/Runtime/RenderCore/RenderCore.Build.cs
+++ b/Engine/Source/Runtime/RenderCore/RenderCore.Build.cs
@@ -9,7 +9,7 @@ public class RenderCore : ModuleRules
 {
 	public RenderCore(ReadOnlyTargetRules Target) : base(Target)
 	{
-		PublicDependencyModuleNames.AddRange(new string[] { "RHI" });
+		PublicDependencyModuleNames.AddRange(new string[] { "RHI", "CoreUObject" });
 
 		PrivateIncludePathModuleNames.AddRange(new string[] { "Shaders", "TargetPlatform" });
 
@@ -21,8 +21,6 @@ public class RenderCore : ModuleRules
 		if (Target.bBuildEditor == true)
         {
 			DynamicallyLoadedModuleNames.Add("TargetPlatform");
-			// UObjects are used to produce the full path of the asset by which the shaders are identified
-			PrivateDependencyModuleNames.Add("CoreUObject");
 			PrivateIncludePathModuleNames.Add("IoStoreUtilities");
 		}
 		// shader runtime usage visualization requires ability to create images - it is only used in non-editor desktop development targets
diff --git a/Engine/Source/Runtime/Renderer/Private/PrimitiveSceneInfo.cpp b/Engine/Source/Runtime/Renderer/Private/PrimitiveSceneInfo.cpp
--- a/Engine/Source/Runtime/Renderer/Private/PrimitiveSceneInfo.cpp
+++ b/Engine/Source/Runtime/Renderer/Private/PrimitiveSceneInfo.cpp
@@ -40,6 +40,7 @@
 #include "DecalRenderingCommon.h"
 #include "RendererPrivateUtils.h"
 #include "ODSC/ODSCManager.h"
+#include "VirtualTextureEnum.h"
 
 extern int32 GGPUSceneInstanceClearList;
 
@@ -2102,7 +2103,7 @@ void FPrimitiveSceneInfo::FlushRuntimeVirtualTexture()
 		{
 			if (Mask & 1)
 			{
-				Scene->RuntimeVirtualTextures[RuntimeVirtualTextureIndex]->Dirty(Proxy->GetBounds());
+				Scene->RuntimeVirtualTextures[RuntimeVirtualTextureIndex]->Dirty(Proxy->GetBounds(), EVTInvalidatePriority::Normal);
 			}
 			Mask >>= 1;
 			RuntimeVirtualTextureIndex++;
diff --git a/Engine/Source/Runtime/Renderer/Private/RendererScene.cpp b/Engine/Source/Runtime/Renderer/Private/RendererScene.cpp
--- a/Engine/Source/Runtime/Renderer/Private/RendererScene.cpp
+++ b/Engine/Source/Runtime/Renderer/Private/RendererScene.cpp
@@ -73,6 +73,7 @@
 #include "GPUScene.h"
 #include "HAL/LowLevelMemStats.h"
 #include "HAL/LowLevelMemTracker.h"
+#include "VT/RuntimeVirtualTextureEnum.h"
 #include "VT/RuntimeVirtualTextureSceneProxy.h"
 #include "VT/VirtualTextureSystem.h"
 #include "HairStrandsInterface.h"
@@ -3662,15 +3663,15 @@ void FScene::GetRuntimeVirtualTextureHidePrimitiveMask(uint8& bHideMaskEditor, u
 	bHideMaskGame = RuntimeVirtualTexturePrimitiveHideMaskGame;
 }
 
-void FScene::InvalidateRuntimeVirtualTexture(class URuntimeVirtualTextureComponent* Component, FBoxSphereBounds const& WorldBounds)
+void FScene::InvalidateRuntimeVirtualTexture(class URuntimeVirtualTextureComponent* Component, FBoxSphereBounds const& WorldBounds, EVTInvalidatePriority InvalidatePriority)
 {
 	if (Component->SceneProxy != nullptr)
 	{
 		FRuntimeVirtualTextureSceneProxy* SceneProxy = Component->SceneProxy;
 		ENQUEUE_RENDER_COMMAND(RuntimeVirtualTextureComponent_SetDirty)(
-			[SceneProxy, WorldBounds] (FRHICommandListBase&)
+			[SceneProxy, WorldBounds, InvalidatePriority] (FRHICommandListBase&)
 		{
-			SceneProxy->Dirty(WorldBounds);
+			SceneProxy->Dirty(WorldBounds, InvalidatePriority);
 		});
 	}
 }
diff --git a/Engine/Source/Runtime/Renderer/Private/ScenePrivate.h b/Engine/Source/Runtime/Renderer/Private/ScenePrivate.h
--- a/Engine/Source/Runtime/Renderer/Private/ScenePrivate.h
+++ b/Engine/Source/Runtime/Renderer/Private/ScenePrivate.h
@@ -98,6 +98,7 @@ struct FPathTracingState;
 class FSparseVolumeTextureViewerSceneProxy;
 class FExponentialHeightFogSceneInfo;
 class FStaticMeshBatch;
+enum class EVTInvalidatePriority : uint8;
 
 /** Holds information about a single primitive's occlusion. */
 class FPrimitiveOcclusionHistory
@@ -3413,7 +3414,7 @@ public:
 	virtual void AddRuntimeVirtualTexture(class URuntimeVirtualTextureComponent* Component) override;
 	virtual void RemoveRuntimeVirtualTexture(class URuntimeVirtualTextureComponent* Component) override;
 	virtual void GetRuntimeVirtualTextureHidePrimitiveMask(uint8& bHideMaskEditor, uint8& bHideMaskGame) const override;
-	virtual void InvalidateRuntimeVirtualTexture(class URuntimeVirtualTextureComponent* Component, FBoxSphereBounds const& WorldBounds) override;
+	virtual void InvalidateRuntimeVirtualTexture(class URuntimeVirtualTextureComponent* Component, FBoxSphereBounds const& WorldBounds, EVTInvalidatePriority InvalidatePriority) override;
 	virtual void InvalidatePathTracedOutput() override;
 	virtual void InvalidateLumenSurfaceCache_GameThread(UPrimitiveComponent* Component) override;
 	virtual void GetPrimitiveUniformShaderParameters_RenderThread(const FPrimitiveSceneInfo* PrimitiveSceneInfo, bool& bHasPrecomputedVolumetricLightmap, FMatrix& PreviousLocalToWorld, int32& SingleCaptureIndex, bool& bOutputVelocity) const override;
diff --git a/Engine/Source/Runtime/Renderer/Private/SceneRendering.cpp b/Engine/Source/Runtime/Renderer/Private/SceneRendering.cpp
--- a/Engine/Source/Runtime/Renderer/Private/SceneRendering.cpp
+++ b/Engine/Source/Runtime/Renderer/Private/SceneRendering.cpp
@@ -49,6 +49,7 @@
 #include "DeviceProfiles/DeviceProfileManager.h"
 #include "DeviceProfiles/DeviceProfile.h"
 #include "PostProcess/PostProcessing.h"
+#include "VirtualTextureEnum.h"
 #include "VirtualTexturing.h"
 #include "VisualizeTexturePresent.h"
 #include "GPUScene.h"
@@ -5428,7 +5429,7 @@ void FRendererModule::FlushVirtualTextureCache(IAllocatedVirtualTexture* Allocat
 
 		for (uint32 LayerIndex = 0; LayerIndex < NumLayers;  ++LayerIndex)
 		{
-			FVirtualTextureSystem::Get().FlushCache(AllocatedVT->GetProducerHandle(LayerIndex), SpaceID, TextureRect, MaxLevel, MaxAgeToKeepMapped);
+			FVirtualTextureSystem::Get().FlushCache(AllocatedVT->GetProducerHandle(LayerIndex), SpaceID, TextureRect, MaxLevel, MaxAgeToKeepMapped, EVTInvalidatePriority::Normal);
 		}
 	}
 }
diff --git a/Engine/Source/Runtime/Renderer/Private/Tests/VirtualTextureTestbed.cpp b/Engine/Source/Runtime/Renderer/Private/Tests/VirtualTextureTestbed.cpp
new file mode 100644
--- /dev/null
+++ b/Engine/Source/Runtime/Renderer/Private/Tests/VirtualTextureTestbed.cpp
@@ -0,0 +1,95 @@
+// Copyright Epic Games, Inc. All Rights Reserved.
+
+#include "CoreMinimal.h"
+#include "VT/VirtualTextureShared.h"
+#include "VirtualTextureEnum.h"
+#include "VirtualTexturing.h"
+#include "Misc/AutomationTest.h"
+
+#if WITH_DEV_AUTOMATION_TESTS || WITH_EDITOR
+
+IMPLEMENT_CUSTOM_SIMPLE_AUTOMATION_TEST(FVirtualTextureTestbed, FAutomationTestBase, "System.Renderer.VirtualTexture", EAutomationTestFlags::EditorContext | EAutomationTestFlags::ClientContext | EAutomationTestFlags::EngineFilter)
+
+// TODO [jonathan.bard] : when it'll be possible, remove this and change FVirtualTextureTestbed into a Low-Level Test (LLT) : 
+#if WITH_LOW_LEVEL_TESTS
+#define LOCAL_CHECK(test) CHECK(test)
+#define LOCAL_REQUIRE_CHECK_SLOW(test) REQUIRE_CHECK_SLOW(test)
+#else // WITH_LOW_LEVEL_TESTS
+#define LOCAL_CHECK(test) check(test)
+#define LOCAL_REQUIRE_CHECK_SLOW(test) {}
+#endif // !WITH_LOW_LEVEL_TESTS
+
+void PerformVTTilePriorityAndIndexTests()
+{
+	{
+		FVTRequestPriority PriorityKey(/*bInLocked = */false, /*bInStreaming = */false, EVTProducerPriority::Highest, EVTInvalidatePriority::High, /*InPagePriority = */0x1234);
+		PriorityKey.PackedValue |= 1ull << 63;
+		LOCAL_REQUIRE_CHECK_SLOW(FVTRequestPriorityAndIndex A(/*InIndex = */42, PriorityKey));
+	}
+
+	{
+		FVTRequestPriority PriorityKey(/*bInLocked = */false, /*bInStreaming = */false, EVTProducerPriority::Highest, EVTInvalidatePriority::High, /*InPagePriority = */0x1234);
+		FVTRequestPriorityAndIndex A(/*InIndex = */(uint16)~0, PriorityKey);
+		LOCAL_REQUIRE_CHECK_SLOW(FVTRequestPriorityAndIndex B(/*InIndex = */(uint16)~0 + 1, PriorityKey));
+	}
+
+	{
+		FVTRequestPriorityAndIndex A(/*InIndex = */42, /*bInLocked = */false, /*bInStreaming = */false, EVTProducerPriority::Highest, EVTInvalidatePriority::High, /*InPagePriority = */0x1234);
+		FVTRequestPriority APriority = A.GetPriorityKey();
+		FVTRequestPriorityAndIndex B(/*InIndex = */42, /*bInLocked = */true, /*bInStreaming = */true, EVTProducerPriority::Normal, EVTInvalidatePriority::Normal, /*InPagePriority = */0x3456);
+		LOCAL_CHECK(A.Index == B.Index);
+	}
+
+	{
+		FVTLocalTilePriorityAndIndex A(/*InIndex = */0, EVTProducerPriority::Highest, EVTInvalidatePriority::High, /*InMipLevel = */5);
+		FVTLocalTilePriorityAndIndex B(/*InIndex = */0, EVTProducerPriority::Normal, EVTInvalidatePriority::Normal, /*InMipLevel = */0);
+		LOCAL_CHECK(A.Index == B.Index);
+	}
+
+	{
+		LOCAL_CHECK(FVTRequestPriorityAndIndex().SortablePackedValue == 0);
+		LOCAL_CHECK(
+			FVTRequestPriorityAndIndex(/*InIndex = */0, /*bInLocked = */false, /*bInStreaming = */false, EVTProducerPriority::Normal, EVTInvalidatePriority::Normal, /*InPagePriority = */0) <
+			FVTRequestPriorityAndIndex(/*InIndex = */0, /*bInLocked = */false, /*bInStreaming = */false, EVTProducerPriority::BelowNormal, EVTInvalidatePriority::Normal, /*InPagePriority = */0));
+		LOCAL_CHECK(
+			FVTRequestPriorityAndIndex(/*InIndex = */0, /*bInLocked = */true, /*bInStreaming = */false, EVTProducerPriority::Normal, EVTInvalidatePriority::Normal, /*InPagePriority = */0) <
+			FVTRequestPriorityAndIndex(/*InIndex = */0, /*bInLocked = */false, /*bInStreaming = */false, EVTProducerPriority::Normal, EVTInvalidatePriority::Normal, /*InPagePriority = */0));
+		LOCAL_CHECK(
+			FVTRequestPriorityAndIndex(/*InIndex = */0, /*bInLocked = */false, /*bInStreaming = */false, EVTProducerPriority::Normal, EVTInvalidatePriority::High, /*InPagePriority = */0) <
+			FVTRequestPriorityAndIndex(/*InIndex = */0, /*bInLocked = */false, /*bInStreaming = */false, EVTProducerPriority::Normal, EVTInvalidatePriority::Normal, /*InPagePriority = */0));
+		LOCAL_CHECK(
+			FVTRequestPriorityAndIndex(/*InIndex = */0, /*bInLocked = */false, /*bInStreaming = */true, EVTProducerPriority::Normal, EVTInvalidatePriority::Normal, /*InPagePriority = */0) <
+			FVTRequestPriorityAndIndex(/*InIndex = */0, /*bInLocked = */false, /*bInStreaming = */false, EVTProducerPriority::Normal, EVTInvalidatePriority::Normal, /*InPagePriority = */0));
+		LOCAL_CHECK(
+			FVTRequestPriorityAndIndex(/*InIndex = */0, /*bInLocked = */false, /*bInStreaming = */false, EVTProducerPriority::Normal, EVTInvalidatePriority::Normal, /*InPagePriority = */1) <
+			FVTRequestPriorityAndIndex(/*InIndex = */0, /*bInLocked = */false, /*bInStreaming = */false, EVTProducerPriority::Normal, EVTInvalidatePriority::Normal, /*InPagePriority = */0));
+		LOCAL_CHECK(
+			FVTRequestPriorityAndIndex(/*InIndex = */0, /*bInLocked = */false, /*bInStreaming = */true, EVTProducerPriority::Normal, EVTInvalidatePriority::Normal, /*InPagePriority = */0) <
+			FVTRequestPriorityAndIndex(/*InIndex = */0, /*bInLocked = */false, /*bInStreaming = */false, EVTProducerPriority::Normal, EVTInvalidatePriority::Normal, /*InPagePriority = */0xffffffff));
+		LOCAL_CHECK(
+			FVTRequestPriorityAndIndex(/*InIndex = */0, /*bInLocked = */true, /*bInStreaming = */false, EVTProducerPriority::Highest, EVTInvalidatePriority::High, /*InPagePriority = */0x1234) <
+			FVTRequestPriorityAndIndex(/*InIndex = */0, /*bInLocked = */false, /*bInStreaming = */false, EVTProducerPriority::Normal, EVTInvalidatePriority::Normal, /*InPagePriority = */0x3456));
+
+		LOCAL_CHECK(FVTLocalTilePriorityAndIndex().SortablePackedValue == 0);
+		LOCAL_CHECK(
+			FVTLocalTilePriorityAndIndex(/*InIndex = */0, EVTProducerPriority::Normal, EVTInvalidatePriority::Normal, /*InMipLevel = */0) <
+			FVTLocalTilePriorityAndIndex(/*InIndex = */0, EVTProducerPriority::BelowNormal, EVTInvalidatePriority::Normal, /*InMipLevel = */0));
+		LOCAL_CHECK(
+			FVTLocalTilePriorityAndIndex(/*InIndex = */0, EVTProducerPriority::Normal, EVTInvalidatePriority::High, /*InMipLevel = */0) <
+			FVTLocalTilePriorityAndIndex(/*InIndex = */0, EVTProducerPriority::Normal, EVTInvalidatePriority::Normal, /*InMipLevel = */0));
+		LOCAL_CHECK(
+			FVTLocalTilePriorityAndIndex(/*InIndex = */0, EVTProducerPriority::Normal, EVTInvalidatePriority::Normal, /*InMipLevel = */5) <
+			FVTLocalTilePriorityAndIndex(/*InIndex = */0, EVTProducerPriority::Normal, EVTInvalidatePriority::Normal, /*InMipLevel = */0));
+	}
+}
+#undef LOCAL_CHECK
+#undef LOCAL_REQUIRE_CHECK_SLOW
+
+bool FVirtualTextureTestbed::RunTest(const FString& Parameters)
+{
+	PerformVTTilePriorityAndIndexTests();
+
+	return true;
+}
+
+#endif //WITH_DEV_AUTOMATION_TESTS
diff --git a/Engine/Source/Runtime/Renderer/Private/VT/RuntimeVirtualTextureSceneProxy.cpp b/Engine/Source/Runtime/Renderer/Private/VT/RuntimeVirtualTextureSceneProxy.cpp
--- a/Engine/Source/Runtime/Renderer/Private/VT/RuntimeVirtualTextureSceneProxy.cpp
+++ b/Engine/Source/Runtime/Renderer/Private/VT/RuntimeVirtualTextureSceneProxy.cpp
@@ -6,6 +6,7 @@
 #include "Components/RuntimeVirtualTextureComponent.h"
 #include "RendererOnScreenNotification.h"
 #include "VirtualTextureSystem.h"
+#include "VT/RuntimeVirtualTextureEnum.h"
 #include "VT/RuntimeVirtualTexture.h"
 #include "VT/RuntimeVirtualTextureProducer.h"
 #include "VT/VirtualTexture.h"
@@ -23,7 +24,6 @@ FRuntimeVirtualTextureSceneProxy::FRuntimeVirtualTextureSceneProxy(URuntimeVirtu
 	, VirtualTexture(nullptr)
 	, bHidePrimitivesInEditor(false)
 	, bHidePrimitivesInGame(false)
-	, CombinedDirtyRect(0, 0, 0, 0)
 {
 	// Evaluate the flags used to hide primitives writing to this virtual texture.
 	InComponent->GetHidePrimitiveSettings(bHidePrimitivesInEditor, bHidePrimitivesInGame);
@@ -165,7 +165,7 @@ void FRuntimeVirtualTextureSceneProxy::Release()
 	}
 }
 
-void FRuntimeVirtualTextureSceneProxy::Dirty(FBoxSphereBounds const& Bounds)
+void FRuntimeVirtualTextureSceneProxy::Dirty(FBoxSphereBounds const& Bounds, EVTInvalidatePriority InvalidatePriority)
 {
 	// Transform world bounds into Virtual Texture UV space
 	const FVector O = Transform.GetTranslation();
@@ -190,14 +190,15 @@ void FRuntimeVirtualTextureSceneProxy::Dirty(FBoxSphereBounds const& Bounds)
 	// Only add rect if it has some area
 	if (TexelRect.Min != TexelRect.Max)
 	{
-		const bool bFirst = DirtyRects.Add(TexelRect) == 0;
+		FDirtyRect DirtyRect{ .Rect = TexelRect, .InvalidatePriority = InvalidatePriority };
+		const bool bFirst = DirtyRects.Add(DirtyRect) == 0;
 		if (bFirst)
 		{
-			CombinedDirtyRect = TexelRect;
+			CombinedDirtyRect = DirtyRect;
 		}
 		else
 		{
-			CombinedDirtyRect.Union(TexelRect);
+			CombinedDirtyRect.Union(DirtyRect);
 		}
 	}
 }
@@ -208,7 +209,7 @@ void FRuntimeVirtualTextureSceneProxy::FlushDirtyPages()
 	if (ProducerHandle.PackedValue != 0)
 	{
 		// Don't do any work if we won't mark anything dirty.
-		if (MaxDirtyLevel >= 0 && CombinedDirtyRect.Width() != 0 && CombinedDirtyRect.Height() != 0)
+		if (MaxDirtyLevel >= 0 && CombinedDirtyRect.Rect.Width() != 0 && CombinedDirtyRect.Rect.Height() != 0)
 		{
 			// Keeping visible pages mapped reduces update flicker due to the latency in the unmap/feedback/map sequence.
 			// But it potentially creates more page update work since more pages may get updated.
@@ -217,24 +218,27 @@ void FRuntimeVirtualTextureSceneProxy::FlushDirtyPages()
 			//todo[vt]: 
 			// Profile to work out best heuristic for when we should use the CombinedDirtyRect
 			// Also consider using some other structure to represent dirty area such as a course 2D bitfield
-			const bool bCombinedFlush = (DirtyRects.Num() > 2 || CombinedDirtyRect == FIntRect(0, 0, VirtualTextureSize.X, VirtualTextureSize.Y));
+			bool bCombinedFlush = (DirtyRects.Num() > 2 || CombinedDirtyRect.Rect == FIntRect(0, 0, VirtualTextureSize.X, VirtualTextureSize.Y));
+			// Don't use the combined rect if one of the dirty rects is prioritized, because that would would leave all of the pages being covered by the combined rect to get prioritized, 
+			//  which would give exactly the opposite result of what we're trying to achieve, since we need to keep the number of prioritized pages to remain low :  
+			bCombinedFlush &= (CombinedDirtyRect.InvalidatePriority == EVTInvalidatePriority::Normal);
 
 			if (bCombinedFlush)
 			{
-				FVirtualTextureSystem::Get().FlushCache(ProducerHandle, SpaceID, CombinedDirtyRect, MaxDirtyLevel, MaxAgeToKeepMapped);
+				FVirtualTextureSystem::Get().FlushCache(ProducerHandle, SpaceID, CombinedDirtyRect.Rect, MaxDirtyLevel, MaxAgeToKeepMapped, EVTInvalidatePriority::Normal);
 			}
 			else
 			{
-				for (FIntRect Rect : DirtyRects)
+				for (const FDirtyRect& DirtyRect : DirtyRects)
 				{
-					FVirtualTextureSystem::Get().FlushCache(ProducerHandle, SpaceID, Rect, MaxDirtyLevel, MaxAgeToKeepMapped);
+					FVirtualTextureSystem::Get().FlushCache(ProducerHandle, SpaceID, DirtyRect.Rect, MaxDirtyLevel, MaxAgeToKeepMapped, DirtyRect.InvalidatePriority);
 				}
 			}
 		}
 	}
 
 	DirtyRects.Reset();
-	CombinedDirtyRect = FIntRect(0, 0, 0, 0);
+	CombinedDirtyRect = FDirtyRect();
 }
 
 #undef LOCTEXT_NAMESPACE
diff --git a/Engine/Source/Runtime/Renderer/Private/VT/RuntimeVirtualTextureSceneProxy.h b/Engine/Source/Runtime/Renderer/Private/VT/RuntimeVirtualTextureSceneProxy.h
--- a/Engine/Source/Runtime/Renderer/Private/VT/RuntimeVirtualTextureSceneProxy.h
+++ b/Engine/Source/Runtime/Renderer/Private/VT/RuntimeVirtualTextureSceneProxy.h
@@ -4,6 +4,7 @@
 
 #include "CoreMinimal.h"
 #include "VirtualTexturing.h"
+#include "VirtualTextureEnum.h"
 
 class URuntimeVirtualTexture;
 class URuntimeVirtualTextureComponent;
@@ -22,8 +23,13 @@ public:
 	 */
 	void Release();
 
-	/** Mark an area of the associated runtime virtual texture as dirty. */
-	void Dirty(FBoxSphereBounds const& Bounds);
+	/**
+	 * Mark an area of the associated runtime virtual texture as dirty. 
+	 * @param Bounds World-space bounds of the area of the runtime virtual texture to invalidate
+	 * @param Priority Allows the affected pages to get processed in priority. This allows increased responsiveness when 
+	 *  there are more pages being updated than can be handled in a given frame (when throttling)
+	 */
+	void Dirty(FBoxSphereBounds const& Bounds, EVTInvalidatePriority InvalidatePriority);
 
 	/** Flush the cached physical pages of the virtual texture at the given world space bounds. */
 	void FlushDirtyPages();
@@ -54,10 +60,22 @@ private:
 
 	/** Maximum mip level to mark dirty. Can be less than the virtual texture's MaxLevel if we have streaming mips. */
 	int32 MaxDirtyLevel;
+
+	struct FDirtyRect
+	{
+		FIntRect Rect;
+		EVTInvalidatePriority InvalidatePriority = EVTInvalidatePriority::Normal;
+
+		void Union(const FDirtyRect& InOther)
+		{
+			Rect.Union(InOther.Rect);
+			InvalidatePriority = static_cast<EVTInvalidatePriority>(FMath::Max(static_cast<uint8>(InvalidatePriority), static_cast<uint8>(InOther.InvalidatePriority)));
+		}
+	};
 	/** Array of dirty rectangles to process at the next flush. */
-	TArray<FIntRect> DirtyRects;
+	TArray<FDirtyRect> DirtyRects;
 	/** Combined dirty rectangle to process at the next flush. */
-	FIntRect CombinedDirtyRect;
+	FDirtyRect CombinedDirtyRect;
 
 	/** DelegateHandle for on screen warning messages. */
 	FDelegateHandle OnScreenWarningDelegateHandle;
diff --git a/Engine/Source/Runtime/Renderer/Private/VT/TexturePageLocks.cpp b/Engine/Source/Runtime/Renderer/Private/VT/TexturePageLocks.cpp
--- a/Engine/Source/Runtime/Renderer/Private/VT/TexturePageLocks.cpp
+++ b/Engine/Source/Runtime/Renderer/Private/VT/TexturePageLocks.cpp
@@ -20,7 +20,7 @@ uint32 FTexturePageLocks::AcquireIndex()
 
 bool FTexturePageLocks::Lock(const FVirtualTextureLocalTile& Tile)
 {
-	const uint16 Hash = (uint16)MurmurFinalize64(Tile.PackedValue);
+	const uint16 Hash = Tile.GetMurmurHash();
 	for (uint32 Index = TileHash.First(Hash); TileHash.IsValid(Index); Index = TileHash.Next(Index))
 	{
 		if (LockedTiles[Index] == Tile)
@@ -43,7 +43,7 @@ bool FTexturePageLocks::Lock(const FVirtualTextureLocalTile& Tile)
 
 bool FTexturePageLocks::Unlock(const FVirtualTextureLocalTile& Tile)
 {
-	const uint16 Hash = (uint16)MurmurFinalize64(Tile.PackedValue);
+	const uint16 Hash = Tile.GetMurmurHash();
 	for (uint32 Index = TileHash.First(Hash); TileHash.IsValid(Index); Index = TileHash.Next(Index))
 	{
 		if (LockedTiles[Index] == Tile)
@@ -90,7 +90,7 @@ void FTexturePageLocks::ForceUnlockAll(const FVirtualTextureProducerHandle& Prod
 		check(Tile.GetProducerHandle() == ProducerHandle);
 		OutUnlockedTiles.Add(Tile);
 		FreeIndices.Add(Index);
-		TileHash.Remove(MurmurFinalize64(Tile.PackedValue), Index);
+		TileHash.Remove(Tile.GetMurmurHash(), Index); 
 		ProducerToTileIndex.Remove(MurmurFinalize32(Tile.PackedProducerHandle), Index);
 		LockCounts[Index] = 0u;
 		LockedTiles[Index].PackedValue = 0u;
@@ -99,7 +99,7 @@ void FTexturePageLocks::ForceUnlockAll(const FVirtualTextureProducerHandle& Prod
 
 bool FTexturePageLocks::IsLocked(const FVirtualTextureLocalTile& Tile) const
 {
-	const uint16 Hash = (uint16)MurmurFinalize64(Tile.PackedValue);
+	const uint16 Hash = Tile.GetMurmurHash();
 	for (uint32 Index = TileHash.First(Hash); TileHash.IsValid(Index); Index = TileHash.Next(Index))
 	{
 		if (LockedTiles[Index] == Tile)
diff --git a/Engine/Source/Runtime/Renderer/Private/VT/TexturePagePool.cpp b/Engine/Source/Runtime/Renderer/Private/VT/TexturePagePool.cpp
--- a/Engine/Source/Runtime/Renderer/Private/VT/TexturePagePool.cpp
+++ b/Engine/Source/Runtime/Renderer/Private/VT/TexturePagePool.cpp
@@ -4,6 +4,7 @@
 
 #include "VirtualTextureSpace.h"
 #include "VirtualTextureSystem.h"
+#include "VirtualTextureEnum.h"
 
 // Reserve pAddress=0 to indicate unmapped page table entry.
 const uint32 FTexturePagePool::NumReservedPages = 1u;
@@ -19,7 +20,6 @@ FTexturePagePool::FTexturePagePool()
 FTexturePagePool::~FTexturePagePool()
 {}
 
-
 void FTexturePagePool::Initialize(uint32 InNumPages)
 {
 	NumPages = InNumPages;
@@ -118,6 +118,7 @@ void FTexturePagePool::EvictPages(FVirtualTextureSystem* System, const FVirtualT
 	for (uint32 pAddress : ToEvict)
 	{
 		UnmapAllPages(System, pAddress, false);
+		checkf(FreeHeap.IsPresent(pAddress), TEXT("If the address is present in the FreeHeap already, then there's an issue because we expect all locked pages to have been unlocked earlier by a call to ForceUnlockAllTiles"));
 		FreeHeap.Update(0, pAddress);
 	}
 }
@@ -129,7 +130,8 @@ void FTexturePagePool::EvictPages(
 	FIntRect const& TextureRegion, 
 	uint32 MaxLevelToEvict, 
 	uint32 MinFrameToKeepMapped,
-	TArray<union FVirtualTextureLocalTile>& OutDirtyMapped)
+	EVTInvalidatePriority InvalidatePriority,
+	TSet<FVirtualTextureLocalTileRequest>& OutDirtyMapped)
 {
 	TArray<uint32, TInlineAllocator<256>> ToEvict;
 	const uint32 Hash = MurmurFinalize32(ProducerHandle.PackedValue);
@@ -154,7 +156,7 @@ void FTexturePagePool::EvictPages(
 					if (!FreeHeap.IsPresent(pAddress))
 					{
 						// Locked pages aren't unmapped but are added to the dirty output array.
-						OutDirtyMapped.Add(FVirtualTextureLocalTile(ProducerHandle, vAddress, vLevel));
+						AddOrMergeTileRequest(FVirtualTextureLocalTileRequest(FVirtualTextureLocalTile(ProducerHandle, vAddress, vLevel), Desc.Priority, InvalidatePriority), OutDirtyMapped);
 					}
 					else
 					{
@@ -162,7 +164,7 @@ void FTexturePagePool::EvictPages(
 						if (PageFrame >= MinFrameToKeepMapped)
 						{
 							// Visible pages aren't unmapped but are added to the dirty output array.
-							OutDirtyMapped.Add(FVirtualTextureLocalTile(ProducerHandle, vAddress, vLevel));
+							AddOrMergeTileRequest(FVirtualTextureLocalTileRequest(FVirtualTextureLocalTile(ProducerHandle, vAddress, vLevel), Desc.Priority, InvalidatePriority), OutDirtyMapped);
 						}
 						else
 						{
diff --git a/Engine/Source/Runtime/Renderer/Private/VT/TexturePagePool.h b/Engine/Source/Runtime/Renderer/Private/VT/TexturePagePool.h
--- a/Engine/Source/Runtime/Renderer/Private/VT/TexturePagePool.h
+++ b/Engine/Source/Runtime/Renderer/Private/VT/TexturePagePool.h
@@ -49,7 +49,8 @@ public:
 	 * Unmap/remove any pages that were allocated by the given producer and are inside the TextureRegion.
 	 * Outputs the locked pages that can't be unmapped to the OutLocked array.
 	 */
-	void EvictPages(FVirtualTextureSystem* System, FVirtualTextureProducerHandle const& ProducerHandle, FVTProducerDescription const& Desc, FIntRect const& TextureRegion, uint32 MaxLevelToEvict, uint32 MinFrameToKeepMapped, TArray<union FVirtualTextureLocalTile>& OutDirtyMapped);
+	void EvictPages(FVirtualTextureSystem* System, FVirtualTextureProducerHandle const& ProducerHandle, FVTProducerDescription const& Desc, FIntRect const& TextureRegion, uint32 MaxLevelToEvict, 
+					uint32 MinFrameToKeepMapped, EVTInvalidatePriority InvalidatePriority, TSet<union FVirtualTextureLocalTileRequest>& OutDirtyMapped);
 
 	/**
 	* Unmap all pages from the given address range in a space. Pages will remain resident in the pool, but no longer by mapped to any page table.
@@ -74,7 +75,7 @@ public:
 	/**
 	 * Check if there are any free pages available at the moment.
 	 */
-	bool		AnyFreeAvailable( uint32 Frame, uint32 FreeThreshold ) const;
+	bool AnyFreeAvailable( uint32 Frame, uint32 FreeThreshold ) const;
 
 	/**
 	 * Find physical address of the page allocated for the given VT address, or ~0 if not allocated
diff --git a/Engine/Source/Runtime/Renderer/Private/VT/UniqueRequestList.h b/Engine/Source/Runtime/Renderer/Private/VT/UniqueRequestList.h
--- a/Engine/Source/Runtime/Renderer/Private/VT/UniqueRequestList.h
+++ b/Engine/Source/Runtime/Renderer/Private/VT/UniqueRequestList.h
@@ -67,10 +67,10 @@ public:
 		: LoadRequestHash(NoInit)
 		, MappingRequestHash(NoInit)
 		, DirectMappingRequestHash(NoInit)
-		, LoadRequests(Allocator.CreateArray<FVirtualTextureLocalTile>(LoadRequestCapacity))
+		, LoadRequests(Allocator.CreateArray<FVirtualTextureLocalTileRequest>(LoadRequestCapacity))
 		, MappingRequests(Allocator.CreateArray<FMappingRequest>(MappingRequestCapacity))
 		, DirectMappingRequests(Allocator.CreateArray<FDirectMappingRequest>(DirectMappingRequestCapacity))
-		, ContinuousUpdateRequests(Allocator.CreateArray<FVirtualTextureLocalTile>(LoadRequestCapacity))
+		, ContinuousUpdateRequests(Allocator.CreateArray<FVirtualTextureLocalTileRequest>(LoadRequestCapacity))
 		, AdaptiveAllocationsRequests(Allocator.MallocArray<uint32>(LoadRequestCapacity))
 		, LoadRequestCount(Allocator.MallocArray<uint16>(LoadRequestCapacity))
 		, LoadRequestGroupMask(Allocator.MallocArray<uint8>(LoadRequestCapacity))
@@ -119,30 +119,30 @@ public:
 	inline uint32 GetNumContinuousUpdateRequests() const { return NumContinuousUpdateRequests; }
 	inline uint32 GetNumAdaptiveAllocationRequests() const { return NumAdaptiveAllocationRequests; }
 
-	inline const FVirtualTextureLocalTile& GetLoadRequest(uint32 i) const { checkSlow(i < NumLoadRequests); return LoadRequests[i]; }
+	inline const FVirtualTextureLocalTileRequest& GetLoadRequest(uint32 i) const { checkSlow(i < NumLoadRequests); return LoadRequests[i]; }
 	inline const FMappingRequest& GetMappingRequest(uint32 i) const { checkSlow(i < NumMappingRequests); return MappingRequests[i]; }
 	inline const FDirectMappingRequest& GetDirectMappingRequest(uint32 i) const { checkSlow(i < NumDirectMappingRequests); return DirectMappingRequests[i]; }
-	inline const FVirtualTextureLocalTile& GetContinuousUpdateRequest(uint32 i) const { checkSlow(i < NumContinuousUpdateRequests); return ContinuousUpdateRequests[i]; }
+	inline const FVirtualTextureLocalTileRequest& GetContinuousUpdateRequest(uint32 i) const { checkSlow(i < NumContinuousUpdateRequests); return ContinuousUpdateRequests[i]; }
 	inline const uint32& GetAdaptiveAllocationRequest(uint32 i) const { checkSlow(i < NumAdaptiveAllocationRequests); return AdaptiveAllocationsRequests[i]; }
 	
 	inline uint8 GetGroupMask(uint32 i) const { checkSlow(i < NumLoadRequests); return LoadRequestGroupMask[i]; }
 	inline bool IsLocked(uint32 i) const { checkSlow(i < NumLoadRequests); return i < NumLockRequests; }
 
-	uint16 AddLoadRequest(const FVirtualTextureLocalTile& Tile, uint8 GroupMask, uint16 Count, bool bStreamingRequest);
-	uint16 LockLoadRequest(const FVirtualTextureLocalTile& Tile, uint8 GroupMask, bool bStreamingRequest);
+	uint16 AddLoadRequest(const FVirtualTextureLocalTileRequest& TileRequest, uint8 GroupMask, uint16 Count, bool bStreamingRequest);
+	uint16 LockLoadRequest(const FVirtualTextureLocalTileRequest& TileRequest, uint8 GroupMask, bool bStreamingRequest);
 
 	void AddMappingRequest(uint16 LoadRequestIndex, uint8 ProducerPhysicalGroupIndex, uint8 SpaceID, uint8 PageTableLayerIndex, uint32 MaxLevel, uint32 vAddress, uint8 vLevel, uint8 Local_vLevel);
 
 	void AddDirectMappingRequest(uint8 InSpaceID, uint16 InPhysicalSpaceID, uint8 InPageTableLayerIndex, uint32 MaxLevel, uint32 InAddress, uint8 InLevel, uint8 InLocalLevel, uint16 InPhysicalAddress);
 	void AddDirectMappingRequest(const FDirectMappingRequest& Request);
 
-	void AddContinuousUpdateRequest(const FVirtualTextureLocalTile& Request);
+	void AddContinuousUpdateRequest(const FVirtualTextureLocalTileRequest& Request);
 
 	void AddAdaptiveAllocationRequest(uint32 Request);
 
 	void MergeRequests(const FUniqueRequestList* RESTRICT Other, FConcurrentLinearBulkObjectAllocator& Allocator);
 
-	void SortRequests(FVirtualTextureProducerCollection& Producers, FConcurrentLinearBulkObjectAllocator& Allocator, uint32 MaxNonStreamingLoadRequests, uint32 MaxStreamingLoadRequests, bool bUseCombinedLimit);
+	void SortRequests(FVirtualTextureProducerCollection& Producers, FConcurrentLinearBulkObjectAllocator& Allocator, uint32 MaxNonStreamingLoadRequests, uint32 MaxStreamingLoadRequests, bool bUseCombinedLimit, bool bSortByPriority);
 
 private:
 	static const uint32 LoadRequestCapacity = 4u * 1024;
@@ -156,10 +156,10 @@ private:
 	TStaticHashTable<512u, DirectMappingRequestCapacity> DirectMappingRequestHash;
 	TStaticHashTable<1024u, ContinuousUpdateRequestCapacity> ContinuousUpdateRequestHash;
 
-	FVirtualTextureLocalTile* LoadRequests;
+	FVirtualTextureLocalTileRequest* LoadRequests;
 	FMappingRequest* MappingRequests;
 	FDirectMappingRequest* DirectMappingRequests;
-	FVirtualTextureLocalTile* ContinuousUpdateRequests;
+	FVirtualTextureLocalTileRequest* ContinuousUpdateRequests;
 	uint32* AdaptiveAllocationsRequests;
 	
 	uint16* LoadRequestCount;
@@ -187,16 +187,16 @@ private:
 	uint32 NumAdaptiveAllocationRequests;
 };
 
-
-inline uint16 FUniqueRequestList::AddLoadRequest(const FVirtualTextureLocalTile& Tile, uint8 GroupMask, uint16 Count, bool bStreamingRequest)
+uint16 FUniqueRequestList::AddLoadRequest(const FVirtualTextureLocalTileRequest& TileRequest, uint8 GroupMask, uint16 Count, bool bStreamingRequest)
 {
-	const uint16 Hash = MurmurFinalize64(Tile.PackedValue);
+	const uint16 Hash = TileRequest.GetMurmurHash();
 	check(GroupMask != 0u);
-
 	for (uint16 Index = LoadRequestHash.First(Hash); LoadRequestHash.IsValid(Index); Index = LoadRequestHash.Next(Index))
 	{
-		if (Tile == LoadRequests[Index])
+		if (TileRequest == LoadRequests[Index])
 		{
+			LoadRequests[Index].MergeWith(TileRequest);
+
 			check(LoadRequestFlags[Index].bStreaming == bStreamingRequest);
 
 			LoadRequestCount[Index] = FMath::Min<uint32>((uint32)LoadRequestCount[Index] + Count, MAX_uint16);
@@ -204,29 +204,31 @@ inline uint16 FUniqueRequestList::AddLoadRequest(const FVirtualTextureLocalTile&
 			return Index;
 		}
 	}
-	
+
 	if (NumLoadRequests < LoadRequestCapacity)
 	{
 		const uint32 Index = NumLoadRequests++;
 		LoadRequestHash.Add(Hash, Index);
-		LoadRequests[Index] = Tile;
+		LoadRequests[Index] = TileRequest;
 		LoadRequestCount[Index] = Count;
 		LoadRequestGroupMask[Index] = GroupMask;
-		LoadRequestFlags[Index] = FLoadRequestFlags(false, bStreamingRequest);
+		LoadRequestFlags[Index] = FLoadRequestFlags(/*bInLocked = */false, bStreamingRequest);
 		return Index;
 	}
 	return 0xffff;
 }
 
-inline uint16 FUniqueRequestList::LockLoadRequest(const FVirtualTextureLocalTile& Tile, uint8 GroupMask, bool bStreamingRequest)
+uint16 FUniqueRequestList::LockLoadRequest(const FVirtualTextureLocalTileRequest& TileRequest, uint8 GroupMask, bool bStreamingRequest)
 {
-	const uint16 Hash = MurmurFinalize64(Tile.PackedValue);
+	const uint16 Hash = TileRequest.GetMurmurHash();
 	check(GroupMask != 0u);
 
 	for (uint16 Index = LoadRequestHash.First(Hash); LoadRequestHash.IsValid(Index); Index = LoadRequestHash.Next(Index))
 	{
-		if (Tile == LoadRequests[Index])
+		if (TileRequest == LoadRequests[Index])
 		{
+			LoadRequests[Index].MergeWith(TileRequest);
+
 			FLoadRequestFlags& Flags = LoadRequestFlags[Index];
 			check(bStreamingRequest == Flags.bStreaming);
 
@@ -245,10 +247,10 @@ inline uint16 FUniqueRequestList::LockLoadRequest(const FVirtualTextureLocalTile
 	{
 		const uint32 Index = NumLoadRequests++;
 		LoadRequestHash.Add(Hash, Index);
-		LoadRequests[Index] = Tile;
+		LoadRequests[Index] = TileRequest;
 		LoadRequestCount[Index] = MAX_uint16;
 		LoadRequestGroupMask[Index] = GroupMask;
-		LoadRequestFlags[Index] = FLoadRequestFlags(true, bStreamingRequest);
+		LoadRequestFlags[Index] = FLoadRequestFlags(/*bInLocked = */true, bStreamingRequest);
 		++NumLockRequests;
 		return Index;
 	}
@@ -256,11 +258,11 @@ inline uint16 FUniqueRequestList::LockLoadRequest(const FVirtualTextureLocalTile
 	return 0xffff;
 }
 
-inline void FUniqueRequestList::AddMappingRequest(uint16 LoadRequestIndex, uint8 ProducerPhysicalGroupIndex, uint8 SpaceID, uint8 PageTableLayerIndex, uint32 MaxLevel, uint32 vAddress, uint8 vLevel, uint8 Local_vLevel)
+void FUniqueRequestList::AddMappingRequest(uint16 LoadRequestIndex, uint8 ProducerPhysicalGroupIndex, uint8 SpaceID, uint8 PageTableLayerIndex, uint32 MaxLevel, uint32 vAddress, uint8 vLevel, uint8 Local_vLevel)
 {
 	check(LoadRequestIndex < NumLoadRequests);
 	const FMappingRequest Request(LoadRequestIndex, ProducerPhysicalGroupIndex, SpaceID, PageTableLayerIndex, MaxLevel, vAddress, vLevel, Local_vLevel);
-	const uint16 Hash = MurmurFinalize64(Request.PackedValue);
+	const uint16 Hash = static_cast<uint16>(MurmurFinalize64(Request.PackedValue));
 
 	for (uint16 Index = MappingRequestHash.First(Hash); MappingRequestHash.IsValid(Index); Index = MappingRequestHash.Next(Index))
 	{
@@ -278,13 +280,13 @@ inline void FUniqueRequestList::AddMappingRequest(uint16 LoadRequestIndex, uint8
 	}
 }
 
-inline void FUniqueRequestList::AddDirectMappingRequest(uint8 InSpaceID, uint16 InPhysicalSpaceID, uint8 InPageTableLayerIndex, uint32 InMaxLevel, uint32 InAddress, uint8 InLevel, uint8 InLocalLevel, uint16 InPhysicalAddress)
+void FUniqueRequestList::AddDirectMappingRequest(uint8 InSpaceID, uint16 InPhysicalSpaceID, uint8 InPageTableLayerIndex, uint32 InMaxLevel, uint32 InAddress, uint8 InLevel, uint8 InLocalLevel, uint16 InPhysicalAddress)
 {
 	const FDirectMappingRequest Request(InSpaceID, InPhysicalSpaceID, InPageTableLayerIndex, InMaxLevel, InAddress, InLevel, InLocalLevel, InPhysicalAddress);
 	AddDirectMappingRequest(Request);
 }
 
-inline void FUniqueRequestList::AddDirectMappingRequest(const FDirectMappingRequest& Request)
+void FUniqueRequestList::AddDirectMappingRequest(const FDirectMappingRequest& Request)
 {
 	const uint16 Hash = Murmur32({ Request.PackedValue[0], Request.PackedValue[1], Request.PackedValue[2] });
 	for (uint16 Index = DirectMappingRequestHash.First(Hash); DirectMappingRequestHash.IsValid(Index); Index = DirectMappingRequestHash.Next(Index))
@@ -303,13 +305,14 @@ inline void FUniqueRequestList::AddDirectMappingRequest(const FDirectMappingRequ
 	}
 }
 
-inline void FUniqueRequestList::AddContinuousUpdateRequest(const FVirtualTextureLocalTile& Request)
+void FUniqueRequestList::AddContinuousUpdateRequest(const FVirtualTextureLocalTileRequest& Request)
 {
-	const uint16 Hash = MurmurFinalize64(Request.PackedValue);
+	const uint16 Hash = Request.GetMurmurHash();
 	for (uint16 Index = ContinuousUpdateRequestHash.First(Hash); ContinuousUpdateRequestHash.IsValid(Index); Index = ContinuousUpdateRequestHash.Next(Index))
 	{
 		if (Request == ContinuousUpdateRequests[Index])
 		{
+			ContinuousUpdateRequests[Index].MergeWith(Request);
 			return;
 		}
 	}
@@ -330,7 +333,7 @@ void FUniqueRequestList::AddAdaptiveAllocationRequest(uint32 Request)
 	}
 }
 
-inline void FUniqueRequestList::MergeRequests(const FUniqueRequestList* RESTRICT Other, FConcurrentLinearBulkObjectAllocator& Allocator)
+void FUniqueRequestList::MergeRequests(const FUniqueRequestList* RESTRICT Other, FConcurrentLinearBulkObjectAllocator& Allocator)
 {
 	uint16* LoadRequestIndexRemap = Allocator.MallocArray<uint16>(Other->NumLoadRequests);
 
@@ -373,17 +376,8 @@ inline void FUniqueRequestList::MergeRequests(const FUniqueRequestList* RESTRICT
 	}
 }
 
-inline void FUniqueRequestList::SortRequests(FVirtualTextureProducerCollection& Producers, FConcurrentLinearBulkObjectAllocator& Allocator, uint32 MaxNonStreamingLoadRequests, uint32 MaxStreamingLoadRequests, bool bUseCombinedLimit)
+void FUniqueRequestList::SortRequests(FVirtualTextureProducerCollection& Producers, FConcurrentLinearBulkObjectAllocator& Allocator, uint32 MaxNonStreamingLoadRequests, uint32 MaxStreamingLoadRequests, bool bUseCombinedLimit, bool bSortByPriority)
 {
-	struct FPriorityAndIndex
-	{
-		uint32 Priority;
-		uint16 Index;
-
-		// sort from largest to smallest
-		inline bool operator<(const FPriorityAndIndex& Rhs) const { return Priority > Rhs.Priority; }
-	};
-
 	if (bUseCombinedLimit)
 	{
 		MaxNonStreamingLoadRequests += MaxStreamingLoadRequests;
@@ -394,39 +388,39 @@ inline void FUniqueRequestList::SortRequests(FVirtualTextureProducerCollection&
 	uint32 CheckNumLockRequests = 0u;
 	uint32 NumNonStreamingLockRequests = 0u;
 	uint32 NumStreamingNonLockRequests = 0u;
-	FPriorityAndIndex* SortedKeys = Allocator.CreateArray<FPriorityAndIndex>(NumLoadRequests);
 
+	FVTRequestPriorityAndIndex* SortedKeys = Allocator.CreateArray<FVTRequestPriorityAndIndex>(NumLoadRequests);
 	for (uint32 i = 0u; i < NumLoadRequests; ++i)
 	{
 		const uint32 Count = LoadRequestCount[i];
 		const FLoadRequestFlags Flags = LoadRequestFlags[i];
 
 		// Try to load higher mips first
-		const FVirtualTextureLocalTile& TileToLoad = GetLoadRequest(i);
-		const uint32 Priority = (Count * (1u + TileToLoad.Local_vLevel)) & 0x0fffffff;
+		const FVirtualTextureLocalTileRequest& LoadRequest = GetLoadRequest(i);
+		const uint32 PagePriority = (Count * (1u + LoadRequest.GetTile().Local_vLevel));
 
-		SortedKeys[i].Index = i;
+		const bool bStreaming = !bUseCombinedLimit && Flags.bStreaming;
 		if (Flags.bLocked)
 		{
-			// Lock request, use max priority
-			SortedKeys[i].Priority = MAX_uint32;
-			NumNonStreamingLockRequests += (bUseCombinedLimit || !Flags.bStreaming ? 1 : 0);
+			NumNonStreamingLockRequests += bStreaming ? 0 : 1;
 			++CheckNumLockRequests;
 		}
-		else if (bUseCombinedLimit || !Flags.bStreaming)
+		else if (bStreaming)
 		{
-			SortedKeys[i].Priority = 0xd0000000 | Priority;
-		}
-		else
-		{
-			SortedKeys[i].Priority = 0xe0000000 | Priority;
 			++NumStreamingNonLockRequests;
 		}
+		SortedKeys[i] = FVTRequestPriorityAndIndex(
+			i,
+			/*bInLocked = */(Flags.bLocked != 0),
+			/*bInStreaming = */bStreaming,
+			bSortByPriority ? LoadRequest.GetProducerPriority() : static_cast<EVTProducerPriority>(0),
+			bSortByPriority ? LoadRequest.GetInvalidatePriority() : static_cast<EVTInvalidatePriority>(0),
+			PagePriority);
 	}
 	checkSlow(CheckNumLockRequests == NumLockRequests);
 
-	// Sort so highest priority requests are at the front of the list
-	// [ Lock requests | streaming non-lock requests | non-streaming non-lock requests ]
+	// Sort so highest priority requests are at the front of the list. 
+	//  Important note : the rest of the algorithm assumes locked requests come first, then streaming requests, then the rest :
 	Algo::Sort(MakeArrayView(SortedKeys, NumLoadRequests));
 
 	// Clamp number of load requests to maximum, but also ensure all lock requests are considered
@@ -439,7 +433,7 @@ inline void FUniqueRequestList::SortRequests(FVirtualTextureProducerCollection&
 	const uint32 NewNumLoadRequests = NewNumNonStreamingRequests + NewNumStreamingRequests;
 
 	// Re-index load request list, using sorted indices
-	FVirtualTextureLocalTile* SortedLoadRequests = Allocator.CreateArray<FVirtualTextureLocalTile>(NewNumLoadRequests);
+	FVirtualTextureLocalTileRequest* SortedLoadRequests = Allocator.CreateArray<FVirtualTextureLocalTileRequest>(NewNumLoadRequests);
 	uint8* SortedGroupMask = Allocator.MallocArray<uint8>(NewNumLoadRequests);
 	FLoadRequestFlags* SortedFlags = Allocator.MallocArray<FLoadRequestFlags>(NewNumLoadRequests);
 	uint16* LoadIndexToSortedLoadIndex = Allocator.MallocArray<uint16>(NumLoadRequests);
@@ -447,33 +441,38 @@ inline void FUniqueRequestList::SortRequests(FVirtualTextureProducerCollection&
 
 	uint32 WriteIndex = 0;
 	auto CopyRequestToSorted = [this, &WriteIndex, SortedKeys, SortedLoadRequests, SortedGroupMask, SortedFlags, LoadIndexToSortedLoadIndex](uint32 SortedIndex)
-	{
-		const uint32 OldIndex = SortedKeys[SortedIndex].Index;
-		checkSlow(OldIndex < NumLoadRequests);
+		{
+			const uint32 OldIndex = SortedKeys[SortedIndex].Index;
+			checkSlow(OldIndex < NumLoadRequests);
 
-		SortedLoadRequests[WriteIndex] = LoadRequests[OldIndex];
-		SortedGroupMask[WriteIndex] = LoadRequestGroupMask[OldIndex];
-		SortedFlags[WriteIndex] = LoadRequestFlags[OldIndex];
-		LoadIndexToSortedLoadIndex[OldIndex] = WriteIndex;
+			SortedLoadRequests[WriteIndex] = LoadRequests[OldIndex];
+			SortedGroupMask[WriteIndex] = LoadRequestGroupMask[OldIndex];
+			SortedFlags[WriteIndex] = LoadRequestFlags[OldIndex];
+			LoadIndexToSortedLoadIndex[OldIndex] = WriteIndex;
 
-		++WriteIndex;
-	};
+			++WriteIndex;
+		};
 
-	for (uint32 i = 0u; i < NumLockRequests; ++i)
+	for (uint32 SortedIndex = 0u; SortedIndex < NumLockRequests; ++SortedIndex)
 	{
-		CopyRequestToSorted(i);
+		CopyRequestToSorted(SortedIndex);
+		checkfSlow((SortedKeys[SortedIndex].GetPriorityKey().Locked != 0) && SortedFlags[WriteIndex - 1].bLocked, TEXT("If this asserts, then the sorting is invalid : the code assumes locked requests are the ones at the beginning of the sorted list"));
 	}
 
 	for (uint32 i = 0u; i < NewNumStreamingRequests - NumStreamingLockRequests; ++i)
 	{
-		CopyRequestToSorted(NumLockRequests + i);
+		const uint32 SortedIndex = NumLockRequests + i;
+		CopyRequestToSorted(SortedIndex);
+		checkfSlow((SortedKeys[SortedIndex].GetPriorityKey().Streaming != 0) && !SortedFlags[WriteIndex - 1].bLocked, TEXT("If this asserts, then the sorting is invalid : the code assumes locked requests are the ones at the beginning of the sorted list"));
 	}
 
 	for (uint32 i = 0u; i < NewNumNonStreamingRequests - NumNonStreamingLockRequests; ++i)
 	{
-		CopyRequestToSorted(NumLockRequests + NumStreamingNonLockRequests + i);
+		const uint32 SortedIndex = NumLockRequests + NumStreamingNonLockRequests + i;
+		CopyRequestToSorted(SortedIndex);
+		checkfSlow((SortedKeys[SortedIndex].GetPriorityKey().Streaming == 0) && !SortedFlags[WriteIndex - 1].bLocked, TEXT("If this asserts, then the sorting is invalid : the code assumes locked requests are the ones at the beginning of the sorted list"));
 	}
-	
+
 	check(NewNumLoadRequests == WriteIndex);
 	LoadRequests = SortedLoadRequests;
 	LoadRequestGroupMask = SortedGroupMask;
@@ -500,3 +499,4 @@ inline void FUniqueRequestList::SortRequests(FVirtualTextureProducerCollection&
 	check(!bUseCombinedLimit || NumLoadRequests == NumNonStreamingLoadRequests);
 	NumMappingRequests = NewNumMappingRequests;
 }
+
diff --git a/Engine/Source/Runtime/Renderer/Private/VT/VirtualTextureShared.h b/Engine/Source/Runtime/Renderer/Private/VT/VirtualTextureShared.h
--- a/Engine/Source/Runtime/Renderer/Private/VT/VirtualTextureShared.h
+++ b/Engine/Source/Runtime/Renderer/Private/VT/VirtualTextureShared.h
@@ -4,6 +4,9 @@
 
 #include "CoreMinimal.h"
 
+enum class EVTProducerPriority : uint8;
+enum class EVTInvalidatePriority : uint8;
+
 union FPhysicalTileLocation
 {
 	FPhysicalTileLocation() {}
@@ -48,3 +51,109 @@ struct FPageTableUpdate
 		//checkSlow(vLogSize <= vLevel);
 	}
 };
+
+// Little utility struct that allows to quickly sort different VT tile-related containers with different policies (size is limited to 64 bits, including the N bits-index, where N = NumBitsForIndex)
+//  PriorityKeyType the priority value on the remaining bits (64 - NumBitsForIndex)
+//  The bigger the priority value, the more important the tile
+template <typename PriorityKeyType, uint8 NumBitsForIndex = 16>
+union TVTTilePriorityAndIndex final
+{
+	TVTTilePriorityAndIndex() = default;
+
+	TVTTilePriorityAndIndex(uint64 InIndex, PriorityKeyType&& InPriorityKey)
+		// Shift by NumBitsForIndex the priority key so it occupies the most significant bits and thus defines the sorting :
+		: SortablePackedValue(InIndex | (InPriorityKey.PackedValue << NumBitsForIndex))
+	{
+		checkfSlow(InIndex < (1ull << NumBitsForIndex), TEXT("The index and priority are packed onto a single uint64. Only NumBitsForIndex bits are allowed for the index"));
+		checkfSlow((InPriorityKey.PackedValue & ~(~0ull >> NumBitsForIndex)) == 0, TEXT("In TVTTilePriorityAndIndex, the priority key is merged with the index into a single uint64. Thus the first N bits of the priority key will be ignored."));
+	}
+
+	template<typename... TArgs>
+	TVTTilePriorityAndIndex(uint64 InIndex, TArgs&&... InArgs)
+		: TVTTilePriorityAndIndex(InIndex, PriorityKeyType(std::forward<TArgs>(InArgs) ...))
+	{}
+
+	// sort from largest to smallest
+	inline bool operator<(const TVTTilePriorityAndIndex<PriorityKeyType, NumBitsForIndex>& InOther) const { return SortablePackedValue > InOther.SortablePackedValue; }
+
+	PriorityKeyType GetPriorityKey() { return PriorityKeyType(SortablePackedValue >> NumBitsForIndex); }
+
+	uint64 SortablePackedValue = 0;
+	// The index are the first N (least significant) bits, so that they don't affect sorting
+	uint64 Index : NumBitsForIndex;
+
+	static_assert(sizeof(PriorityKeyType) == sizeof(uint64), "Unexpected size for TPriorityAndIndex. This struct should be kept as small as possible for the sorting to remain efficient");
+};
+
+// Sorting key for VT requests
+struct FVTRequestPriority final
+{
+	FVTRequestPriority() = default;
+
+	FVTRequestPriority(bool bInLocked, bool bInStreaming, EVTProducerPriority InProducerPriority, EVTInvalidatePriority InInvalidatePriority, uint32 InPagePriority)
+		: PagePriority(InPagePriority)
+		, InvalidatePriority(static_cast<uint8>(InInvalidatePriority))
+		, ProducerPriority(static_cast<uint8>(InProducerPriority))
+		, Streaming((uint64)(bInStreaming ? 1 : 0))
+		, Locked((uint64)(bInLocked ? 1 : 0))
+		, Pad(0)
+	{
+		checkfSlow(static_cast<uint64>(InInvalidatePriority) <= (1 << 1), TEXT("EVTInvalidatePriority should be packable on 1 bit"));
+		checkfSlow(static_cast<uint64>(InProducerPriority) <= (1 << 3), TEXT("EVTProducerPriority should be packable on 3 bits"));
+	}
+
+	FVTRequestPriority(uint64 InPackedValue)
+		: PackedValue(InPackedValue)
+	{}
+
+	union
+	{
+		uint64 PackedValue = 0;
+		struct
+		{
+			// Important note : The order of these members is important : it defines the sort order (last member first)
+			uint64 PagePriority : 32; // Page priority depends on the number of requests and the mip level (higher mips come first)
+			uint64 InvalidatePriority : 1; // Manually-prioritized pages get processed before others
+			uint64 ProducerPriority : 3; // Sort by producer priority first
+			uint64 Streaming : 1; // Streaming pages get processed before others : important note : this needs to remain the second most significant bit, since FUniqueRequestList::SortRequests assumes streaming pages come next
+			uint64 Locked : 1; // Locked pages get processed before others : important note : this needs to remain the most significant bit, since FUniqueRequestList::SortRequests assumes locked pages come first 
+			uint64 Pad : 26;
+		};
+	};
+};
+
+// Sorting key for plain FVirtualTextureLocalTile
+struct FVTLocalTilePriority final
+{
+	FVTLocalTilePriority() = default;
+
+	FVTLocalTilePriority(EVTProducerPriority InProducerPriority, EVTInvalidatePriority InInvalidatePriority, uint8 InMipLevel)
+		: MipLevel(InMipLevel)
+		, InvalidatePriority(static_cast<uint8>(InInvalidatePriority))
+		, ProducerPriority(static_cast<uint8>(InProducerPriority))
+		, Pad(0)
+	{
+		checkfSlow(static_cast<uint64>(InProducerPriority) <= (1 << 3), TEXT("EVTProducerPriority should be packable on 3 bits"));
+		checkfSlow(static_cast<uint64>(InMipLevel) <= (1 << 4), TEXT("Mip level should be packable on 4 bits"));
+	}
+
+	FVTLocalTilePriority(uint64 InPackedValue)
+		: PackedValue(InPackedValue)
+	{}
+
+	union
+	{
+		uint64 PackedValue = 0;
+		struct
+		{
+			// Important note : The order of these members is important : it defines the sort order (last member first)
+			uint64 MipLevel : 4; // Page priority depends on the number of requests and the mip level (higher mips come first)
+			uint64 InvalidatePriority : 1; // Prioritized pages get processed before others
+			uint64 ProducerPriority : 3; // Sort by producer priority first
+			uint64 Pad : 56;
+		};
+	};
+};
+
+using FVTRequestPriorityAndIndex = TVTTilePriorityAndIndex<FVTRequestPriority>;
+using FVTLocalTilePriorityAndIndex = TVTTilePriorityAndIndex<FVTLocalTilePriority>;
diff --git a/Engine/Source/Runtime/Renderer/Private/VT/VirtualTextureSystem.cpp b/Engine/Source/Runtime/Renderer/Private/VT/VirtualTextureSystem.cpp
--- a/Engine/Source/Runtime/Renderer/Private/VT/VirtualTextureSystem.cpp
+++ b/Engine/Source/Runtime/Renderer/Private/VT/VirtualTextureSystem.cpp
@@ -18,6 +18,7 @@
 #include "SceneUtils.h"
 #include "Stats/Stats.h"
 #include "VirtualTexturing.h"
+#include "VirtualTextureEnum.h"
 #include "VT/AdaptiveVirtualTexture.h"
 #include "VT/TexturePagePool.h"
 #include "VT/UniquePageList.h"
@@ -172,6 +173,13 @@ static TAutoConsoleVariable<int32> CVarVTWaitBeforeUpdate(
 	ECVF_RenderThreadSafe
 );
 
+TAutoConsoleVariable<int32> CVarVTSortTileRequestsByPriority(
+	TEXT("r.VT.SortTileRequestsByPriority"),
+	1,
+	TEXT("Sorts the list of tile requests by priority before submitting. This allows priority requests to get processed before others in case of throttling"),
+	ECVF_RenderThreadSafe
+);
+
 FVirtualTextureUpdateSettings::FVirtualTextureUpdateSettings()
 {
 	bEnableFeedback = CVarVTEnableFeedback.GetValueOnRenderThread() != 0;
@@ -417,7 +425,7 @@ void FVirtualTextureSystem::FlushCache()
 	bFlushCaches = true;
 }
 
-void FVirtualTextureSystem::FlushCache(FVirtualTextureProducerHandle const& ProducerHandle, int32 SpaceID, FIntRect const& TextureRegion, uint32 MaxLevelToEvict, uint32 MaxAgeToKeepMapped)
+void FVirtualTextureSystem::FlushCache(FVirtualTextureProducerHandle const& ProducerHandle, int32 SpaceID, FIntRect const& TextureRegion, uint32 MaxLevelToEvict, uint32 MaxAgeToKeepMapped, EVTInvalidatePriority InvalidatePriority)
 {
 	check(!bUpdating);
 	UE::TScopeLock Lock(Mutex);
@@ -436,7 +444,7 @@ void FVirtualTextureSystem::FlushCache(FVirtualTextureProducerHandle const& Prod
 			PhysicalSpacesForProducer.AddUnique(Producer->GetPhysicalSpaceForPhysicalGroup(i));
 		}
 
-		check(TransientCollectedPages.Num() == 0);
+		check(TransientCollectedTilesToProduce.Num() == 0);
 
 		const uint32 MinFrameToKeepMapped = Frame > MaxAgeToKeepMapped ? Frame - MaxAgeToKeepMapped : 0;
 
@@ -456,23 +464,23 @@ void FVirtualTextureSystem::FlushCache(FVirtualTextureProducerHandle const& Prod
 				// Adaptive VT flushes.
 				for (FAdaptiveVirtualTexture::FProducerInfo& Info : ProducerInfos)
 				{
-					Pool.EvictPages(this, Info.ProducerHandle, ProducerDescription, Info.RemappedTextureRegion, Info.RemappedMaxLevel, MinFrameToKeepMapped, TransientCollectedPages);
+					Pool.EvictPages(this, Info.ProducerHandle, ProducerDescription, Info.RemappedTextureRegion, Info.RemappedMaxLevel, MinFrameToKeepMapped, InvalidatePriority, TransientCollectedTilesToProduce);
 				}
 			}
 			else
 			{
 				// Regular flush.
-				Pool.EvictPages(this, ProducerHandle, ProducerDescription, TextureRegion, MaxLevelToEvict, MinFrameToKeepMapped, TransientCollectedPages);
+				Pool.EvictPages(this, ProducerHandle, ProducerDescription, TextureRegion, MaxLevelToEvict, MinFrameToKeepMapped, InvalidatePriority, TransientCollectedTilesToProduce);
 			}
 		}
 
-		for (auto& Page : TransientCollectedPages)
+		for (const FVirtualTextureLocalTileRequest& TileRequest : TransientCollectedTilesToProduce)
 		{
-			MappedTilesToProduce.Add(Page);
+			AddOrMergeTileRequest(TileRequest, MappedTilesToProduce);
 		}
 
 		// Don't resize to allow this container to grow as needed (avoid allocations when collecting)
-		TransientCollectedPages.Reset();
+		TransientCollectedTilesToProduce.Reset();
 	}
 }
 
@@ -1597,7 +1605,9 @@ void FVirtualTextureSystem::GatherRequestsTask(const FGatherRequestsParameters&
 					{
 						FTexturePagePool& RESTRICT PagePool = GetPhysicalSpace(PhysicalSpaceIDAndAddress.PhysicalSpaceID)->GetPagePool();
 						const FVirtualTextureLocalTile LocalTile = PagePool.GetLocalTileFromPhysicalAddress(PhysicalSpaceIDAndAddress.pAddress);
-						RequestList->AddContinuousUpdateRequest(LocalTile);
+						//todo[vt]: the FVirtualTextureLocalTileRequest produced here doesn't have a ProducerPriority set. Technically it would be possible to retrieve it from the producer 
+						// for correctness, but this would require extra indirection so need to profile first
+						RequestList->AddContinuousUpdateRequest(FVirtualTextureLocalTileRequest(LocalTile, EVTProducerPriority::Normal, EVTInvalidatePriority::Normal));
 					}
 
 					++PageUpdateBuffers[PhysicalSpaceIDAndAddress.PhysicalSpaceID].WorkingSetSize;
@@ -1704,6 +1714,7 @@ void FVirtualTextureSystem::GatherRequestsTask(const FGatherRequestsParameters&
 
 			const uint32 ProducerMipBias = AllocatedVT->GetUniqueProducerMipBias(ProducerIndex);
 			const uint32 ProducerMaxLevel = Producer->GetMaxLevel();
+			const EVTProducerPriority ProducerPriority = Producer->GetDescription().Priority;
 
 			// here vLevel is clamped against ProducerMipBias, as ProducerMipBias represents the most detailed level of this producer, relative to the allocated VT
 			// used to rescale vAddress to the correct tile within the given mip level
@@ -1843,7 +1854,10 @@ void FVirtualTextureSystem::GatherRequestsTask(const FGatherRequestsParameters&
 					}
 
 					const bool bStreamingRequest = Producer->GetVirtualTexture()->IsPageStreamed(PrefetchLocal_vLevel, PrefetchLocal_vAddress);
-					const uint16 LoadRequestIndex = RequestList->AddLoadRequest(FVirtualTextureLocalTile(ProducerHandle, PrefetchLocal_vAddress, PrefetchLocal_vLevel), ProducerPhysicalGroupMaskToPrefetch, PageCount, bStreamingRequest);
+					const FVirtualTextureLocalTile Tile(ProducerHandle, PrefetchLocal_vAddress, PrefetchLocal_vLevel);
+					const uint16 LoadRequestIndex = RequestList->AddLoadRequest(
+						FVirtualTextureLocalTileRequest(Tile, ProducerPriority, EVTInvalidatePriority::Normal),
+						ProducerPhysicalGroupMaskToPrefetch, PageCount, bStreamingRequest);
 					if (LoadRequestIndex != 0xffff)
 					{
 						const uint32 PrefetchMapping_vLevel = PrefetchLocal_vLevel + ProducerMipBias;
@@ -1868,7 +1882,10 @@ void FVirtualTextureSystem::GatherRequestsTask(const FGatherRequestsParameters&
 			if (GroupMaskToLoad != 0u)
 			{
 				const bool bStreamingRequest = Producer->GetVirtualTexture()->IsPageStreamed(Local_vLevel, Local_vAddress);
-				const uint16 LoadRequestIndex = RequestList->AddLoadRequest(FVirtualTextureLocalTile(ProducerHandle, Local_vAddress, Local_vLevel), GroupMaskToLoad, PageCount, bStreamingRequest);
+				const FVirtualTextureLocalTile Tile(ProducerHandle, Local_vAddress, Local_vLevel);
+				const uint16 LoadRequestIndex = RequestList->AddLoadRequest(
+					FVirtualTextureLocalTileRequest(Tile, ProducerPriority, EVTInvalidatePriority::Normal),
+					GroupMaskToLoad, PageCount, bStreamingRequest);
 				if (LoadRequestIndex != 0xffff)
 				{
 					for (uint32 LoadLayerIndex = 0u; LoadLayerIndex < NumPageTableLayersToLoad; ++LoadLayerIndex)
@@ -1928,7 +1945,7 @@ void FVirtualTextureSystem::GetContinuousUpdatesToProduce(FUniqueRequestList con
 	{
 		for (int32 i = 0; i < NumContinuousUpdateRequests; ++i)
 		{
-			ContinuousUpdateTilesToProduce.Add(RequestList->GetContinuousUpdateRequest(i));
+			AddOrMergeTileRequest(RequestList->GetContinuousUpdateRequest(i), ContinuousUpdateTilesToProduce);
 		}
 	}
 	else
@@ -1941,7 +1958,7 @@ void FVirtualTextureSystem::GetContinuousUpdatesToProduce(FUniqueRequestList con
 			// Note it's possible that we add a duplicate value to the TSet here, and so MappedTilesToProduce doesn't grow.
 			// But ending up with fewer continuous updates then the maximum is OK.
 			int32 RandomIndex = FMath::Rand() % NumContinuousUpdateRequests;
-			ContinuousUpdateTilesToProduce.Add(RequestList->GetContinuousUpdateRequest(RandomIndex));
+			AddOrMergeTileRequest(RequestList->GetContinuousUpdateRequest(RandomIndex), ContinuousUpdateTilesToProduce);
 			NumContinuousUpdates++;
 		}
 	}
@@ -1997,20 +2014,76 @@ void FVirtualTextureSystem::GrowPhysicalPools() const
 	}
 }
 
-int32 FVirtualTextureSystem::SubmitRequestsFromLocalTileList(FRHICommandList& RHICmdList, TArray<FVirtualTextureLocalTile>& OutDeferredTiles, const TSet<FVirtualTextureLocalTile>& LocalTileList, EVTProducePageFlags Flags, ERHIFeatureLevel::Type FeatureLevel, uint32 MaxRequestsToProduce)
+TArrayView<FVTLocalTilePriorityAndIndex> FVirtualTextureSystem::SortLocalTileRequests(FConcurrentLinearBulkObjectAllocator& Allocator, const TConstArrayView<FVirtualTextureLocalTileRequest>& InLocalTileRequests)
+{
+	TArrayView<FVTLocalTilePriorityAndIndex> Result;
+	if (InLocalTileRequests.IsEmpty())
+	{
+		return Result;
+	}
+
+	const bool bSortByPriority = CVarVTSortTileRequestsByPriority.GetValueOnRenderThread() != 0;
+	FVTLocalTilePriorityAndIndex* SortedKeys = Allocator.CreateArray<FVTLocalTilePriorityAndIndex>(InLocalTileRequests.Num());
+	const uint16 NumTiles = IntCastChecked<uint16>(InLocalTileRequests.Num());
+	FVirtualTextureProducerHandle CurrentProducerHandle;
+	FVirtualTextureProducer const* CurrentProducer = nullptr;
+	int32 CurrentKeyIndex = 0;
+	for (uint16 TileIndex = 0; TileIndex < NumTiles; ++TileIndex)
+	{
+		const FVirtualTextureLocalTileRequest& TileRequest = InLocalTileRequests[TileIndex];
+		const FVirtualTextureLocalTile Tile = TileRequest.GetTile();
+		const FVirtualTextureProducerHandle ProducerHandle = Tile.GetProducerHandle();
+		// No need to fetch the producer if we share the same handle as the last tile (which is very likely to happen because we sort by producer): 
+		if (CurrentProducerHandle.IsNull() || (ProducerHandle != CurrentProducerHandle))
+		{
+			CurrentProducerHandle = ProducerHandle;
+			CurrentProducer = Producers.FindProducer(ProducerHandle);
+		}
+
+		// If we didn't process the tile last frame and deferred processing and if the producer got removed in the meantime we end up with a nullptr producer (just throw away all the requests from this one)
+		if (CurrentProducer != nullptr)
+		{
+			SortedKeys[CurrentKeyIndex++] = FVTLocalTilePriorityAndIndex(
+				TileIndex, 
+				bSortByPriority ? TileRequest.GetProducerPriority() : static_cast<EVTProducerPriority>(0),
+				bSortByPriority ? TileRequest.GetInvalidatePriority() : static_cast<EVTInvalidatePriority>(0),
+				/*InMipLevel = */Tile.Local_vLevel);
+		}
+	}
+
+	if (CurrentKeyIndex > 0)
+	{
+		Result = MakeArrayView(SortedKeys, CurrentKeyIndex);
+
+		if (bSortByPriority)
+		{
+			// Sort by decreasing priority : 
+			Algo::Sort(Result);
+		}
+	}
+	return Result;
+}
+
+int32 FVirtualTextureSystem::SubmitRequestsFromLocalTileRequests(FConcurrentLinearBulkObjectAllocator& Allocator, FRHICommandList& RHICmdList, TSet<FVirtualTextureLocalTileRequest>& OutDeferredTileRequests, const TSet<FVirtualTextureLocalTileRequest>& LocalTileRequests, EVTProducePageFlags Flags, ERHIFeatureLevel::Type FeatureLevel, uint32 MaxRequestsToProduce)
 {
+	TRACE_CPUPROFILER_EVENT_SCOPE(FVirtualTextureSystem::SubmitRequestsFromLocalTileList);
 	LLM_SCOPE(ELLMTag::VirtualTextureSystem);
 
 	uint32 NumPagesProduced = 0;
-	for (const FVirtualTextureLocalTile& Tile : LocalTileList)
+
+	const TArray<FVirtualTextureLocalTileRequest> LocalTileRequestsArray = LocalTileRequests.Array();
+	TArrayView<FVTLocalTilePriorityAndIndex> SortedKeys = SortLocalTileRequests(Allocator, LocalTileRequestsArray);
+	FVirtualTextureProducerHandle CurrentProducerHandle;
+	FVirtualTextureProducer const* CurrentProducer = nullptr;
+	for (const FVTLocalTilePriorityAndIndex& SortedTileKeyIndex : SortedKeys)
 	{
+		const FVirtualTextureLocalTileRequest& TileRequest = LocalTileRequestsArray[SortedTileKeyIndex.Index];
+		const FVirtualTextureLocalTile Tile = TileRequest.GetTile();
 		const FVirtualTextureProducerHandle ProducerHandle = Tile.GetProducerHandle();
-		FVirtualTextureProducer const* Producer = Producers.FindProducer(ProducerHandle);
-		if (Producer == nullptr)
+		// No need to fetch the producer if we share the same handle as the last tile (which is very likely to happen because we sort by producer): 
+		if (CurrentProducerHandle.IsNull() || (ProducerHandle != CurrentProducerHandle))
 		{
-			// If we didn't process the tile last frame and deferred processing and if the producer got removed in the meantime we end up here.
-			// Just throw away the request.
-			continue;
+			CurrentProducer = Producers.FindProducer(ProducerHandle);
 		}
 
 		// Fill targets for each layer
@@ -2018,17 +2091,17 @@ int32 FVirtualTextureSystem::SubmitRequestsFromLocalTileList(FRHICommandList& RH
 		// If the phys layer is mapped then we get the textures it owns and map them into the producer local slots and set the flags
 		uint32 LayerMask = 0;
 		FVTProduceTargetLayer ProduceTarget[VIRTUALTEXTURE_SPACE_MAXLAYERS];
-		for (uint32 ProducerPhysicalGroupIndex = 0u; ProducerPhysicalGroupIndex < Producer->GetNumPhysicalGroups(); ++ProducerPhysicalGroupIndex)
+		for (uint32 ProducerPhysicalGroupIndex = 0u; ProducerPhysicalGroupIndex < CurrentProducer->GetNumPhysicalGroups(); ++ProducerPhysicalGroupIndex)
 		{
-			FVirtualTexturePhysicalSpace* RESTRICT PhysicalSpace = Producer->GetPhysicalSpaceForPhysicalGroup(ProducerPhysicalGroupIndex);
+			FVirtualTexturePhysicalSpace* RESTRICT PhysicalSpace = CurrentProducer->GetPhysicalSpaceForPhysicalGroup(ProducerPhysicalGroupIndex);
 			FTexturePagePool& RESTRICT PagePool = PhysicalSpace->GetPagePool();
 			const uint32 pAddress = PagePool.FindPageAddress(ProducerHandle, ProducerPhysicalGroupIndex, Tile.Local_vAddress, Tile.Local_vLevel);
 			if (pAddress != ~0u)
 			{
 				int32 PhysicalLocalTextureIndex = 0;
-				for (uint32 ProducerLayerIndex = 0u; ProducerLayerIndex < Producer->GetNumTextureLayers(); ++ProducerLayerIndex)
+				for (uint32 ProducerLayerIndex = 0u; ProducerLayerIndex < CurrentProducer->GetNumTextureLayers(); ++ProducerLayerIndex)
 				{
-					if (Producer->GetPhysicalGroupIndexForTextureLayer(ProducerLayerIndex) == ProducerPhysicalGroupIndex)
+					if (CurrentProducer->GetPhysicalGroupIndexForTextureLayer(ProducerLayerIndex) == ProducerPhysicalGroupIndex)
 					{
 						ProduceTarget[ProducerLayerIndex].TextureRHI = PhysicalSpace->GetPhysicalTexture(PhysicalLocalTextureIndex);
 						ProduceTarget[ProducerLayerIndex].PooledRenderTarget = PhysicalSpace->GetPhysicalTexturePooledRenderTarget(PhysicalLocalTextureIndex);
@@ -2049,21 +2122,21 @@ int32 FVirtualTextureSystem::SubmitRequestsFromLocalTileList(FRHICommandList& RH
 		if (MaxRequestsToProduce > 0 && NumPagesProduced >= MaxRequestsToProduce)
 		{
 			// Keep the request for the next frame?
-			OutDeferredTiles.Add(Tile);
+			AddOrMergeTileRequest(TileRequest, OutDeferredTileRequests);
 			continue;
 		}
 
-		FVTRequestPageResult RequestPageResult = Producer->GetVirtualTexture()->RequestPageData(
+		FVTRequestPageResult RequestPageResult = CurrentProducer->GetVirtualTexture()->RequestPageData(
 			RHICmdList, ProducerHandle, LayerMask, Tile.Local_vLevel, Tile.Local_vAddress, EVTRequestPagePriority::High);
 
 		if (RequestPageResult.Status != EVTRequestPageStatus::Available)
 		{
 			// Keep the request for the next frame?
-			OutDeferredTiles.Add(Tile);
+			AddOrMergeTileRequest(TileRequest, OutDeferredTileRequests);
 			continue;
 		}
 
-		IVirtualTextureFinalizer* VTFinalizer = Producer->GetVirtualTexture()->ProducePageData(
+		IVirtualTextureFinalizer* VTFinalizer = CurrentProducer->GetVirtualTexture()->ProducePageData(
 			RHICmdList, FeatureLevel,
 			Flags,
 			ProducerHandle, LayerMask, Tile.Local_vLevel, Tile.Local_vAddress,
@@ -2088,6 +2161,7 @@ void FVirtualTextureSystem::SubmitThrottledRequests(FRHICommandList& RHICmdList,
 	const ERHIFeatureLevel::Type  FeatureLevel      = Updater->FeatureLevel;
 	FConcurrentLinearBulkObjectAllocator& Allocator = Updater->Allocator;
 	FUniqueRequestList* MergedRequestList           = Updater->MergedRequestList;
+	const bool bSortByPriority = CVarVTSortTileRequestsByPriority.GetValueOnRenderThread() != 0;
 
 	if (MergedRequestList->GetNumAdaptiveAllocationRequests() > 0)
 	{
@@ -2099,13 +2173,13 @@ void FVirtualTextureSystem::SubmitThrottledRequests(FRHICommandList& RHICmdList,
 	if (UpdatePhase == EUpdatePhase::Begin)
 	{
 		// Only take runtime generated page budget into account since we expect that only runtime generated pages should have been dirtied.
-		const int32 NumPagesProduced = SubmitRequestsFromLocalTileList(RHICmdList, TransientCollectedPages, MappedTilesToProduce, EVTProducePageFlags::None, FeatureLevel, Updater->PageUploadBudgetRVT);
+		int32 NumPagesProduced = SubmitRequestsFromLocalTileRequests(Allocator, RHICmdList, TransientCollectedTilesToProduce, MappedTilesToProduce, EVTProducePageFlags::None, FeatureLevel, Updater->PageUploadBudgetRVT);
 		Updater->PageUploadBudgetRVT = FMath::Max(Updater->PageUploadBudgetRVT - NumPagesProduced, 0);
 
 		INC_DWORD_STAT_BY(STAT_NumMappedPageUpdate, MappedTilesToProduce.Num());
 		MappedTilesToProduce.Reset();
-		MappedTilesToProduce.Append(TransientCollectedPages);
-		TransientCollectedPages.Reset();
+		MappedTilesToProduce.Append(TransientCollectedTilesToProduce);
+		TransientCollectedTilesToProduce.Reset();
 	}
 
 	{
@@ -2118,7 +2192,7 @@ void FVirtualTextureSystem::SubmitThrottledRequests(FRHICommandList& RHICmdList,
 		// SVT budget of 0 is a special value that enables the old behavior where all pages are limited by a single budget.
 		const bool bUseCombinedLimit = Settings.MaxSVTPageUploads == 0;
 		
-		MergedRequestList->SortRequests(Producers, Allocator, Updater->PageUploadBudgetRVT, Updater->PageUploadBudgetSVT, bUseCombinedLimit);
+		MergedRequestList->SortRequests(Producers, Allocator, Updater->PageUploadBudgetRVT, Updater->PageUploadBudgetSVT, bUseCombinedLimit, bSortByPriority);
 		
 		// Subtract sorted request count from the remaining budgets.
 		const int32 NumRequestsRVT = (int32)MergedRequestList->GetNumNonStreamingLoadRequests();
@@ -2134,12 +2208,12 @@ void FVirtualTextureSystem::SubmitThrottledRequests(FRHICommandList& RHICmdList,
 		// Don't take streaming page budget into account since they async load and likely won't be produced this frame.
 		// Also we expect only runtime generated pages to need continous updates (since continuous updates are there to handle stale pages coming from rendering before ready etc).
 		GetContinuousUpdatesToProduce(MergedRequestList, Updater->PageUploadBudgetRVT, Settings.MaxContinuousUpdates);
-		const int32 NumPagesProduced = SubmitRequestsFromLocalTileList(RHICmdList, TransientCollectedPages, ContinuousUpdateTilesToProduce, EVTProducePageFlags::ContinuousUpdate, FeatureLevel, 0);
+		int32 NumPagesProduced = SubmitRequestsFromLocalTileRequests(Allocator, RHICmdList, TransientCollectedTilesToProduce, ContinuousUpdateTilesToProduce, EVTProducePageFlags::ContinuousUpdate, FeatureLevel, 0);
 		Updater->PageUploadBudgetRVT = FMath::Max(Updater->PageUploadBudgetRVT - NumPagesProduced, 0);
 
 		INC_DWORD_STAT_BY(STAT_NumContinuousPageUpdate, ContinuousUpdateTilesToProduce.Num());
 		ContinuousUpdateTilesToProduce.Reset();
-		TransientCollectedPages.Reset();
+		TransientCollectedTilesToProduce.Reset();
 	}
 
 	// Submit the merged, sorted and throttled page load requests.
@@ -2185,7 +2259,7 @@ void FVirtualTextureSystem::SubmitRequests(FRHICommandList& RHICmdList, ERHIFeat
 		{
 			const bool bLockTile = RequestList->IsLocked(RequestIndex);
 			const bool bForceProduceTile = !bAsync || (bLockTile && bSyncProduceLockedTiles);
-			const FVirtualTextureLocalTile TileToLoad = RequestList->GetLoadRequest(RequestIndex);
+			const FVirtualTextureLocalTile TileToLoad = RequestList->GetLoadRequest(RequestIndex).GetTile();
 			const FVirtualTextureProducerHandle ProducerHandle = TileToLoad.GetProducerHandle();
 			const FVirtualTextureProducer& Producer = Producers.GetProducer(ProducerHandle);
 
@@ -2410,7 +2484,7 @@ void FVirtualTextureSystem::SubmitRequests(FRHICommandList& RHICmdList, ERHIFeat
 			const uint32 pAddress = RequestPhysicalAddress[MappingRequest.LoadRequestIndex * VIRTUALTEXTURE_SPACE_MAXLAYERS + MappingRequest.ProducerPhysicalGroupIndex];
 			if (pAddress != ~0u)
 			{
-				const FVirtualTextureLocalTile& TileToLoad = RequestList->GetLoadRequest(MappingRequest.LoadRequestIndex);
+				const FVirtualTextureLocalTile TileToLoad = RequestList->GetLoadRequest(MappingRequest.LoadRequestIndex).GetTile();
 				const FVirtualTextureProducerHandle ProducerHandle = TileToLoad.GetProducerHandle();
 				FVirtualTextureProducer& Producer = Producers.GetProducer(ProducerHandle);
 				FVirtualTexturePhysicalSpace* PhysicalSpace = Producer.GetPhysicalSpaceForPhysicalGroup(MappingRequest.ProducerPhysicalGroupIndex);
@@ -2651,7 +2725,9 @@ void FVirtualTextureSystem::GatherLockedTileRequests(FUniqueRequestList* MergedR
 			if (ProducerLayerMaskToLoad != 0u)
 			{
 				const bool bStreamingRequest = Producer->GetVirtualTexture()->IsPageStreamed(Tile.Local_vLevel, Tile.Local_vAddress);
-				const uint16 LoadRequestIndex = MergedRequestList->LockLoadRequest(FVirtualTextureLocalTile(Tile.GetProducerHandle(), Tile.Local_vAddress, Tile.Local_vLevel), ProducerLayerMaskToLoad, bStreamingRequest);
+				EVTProducerPriority ProducerPriority = Producer->GetDescription().Priority;
+				FVirtualTextureLocalTileRequest TileRequest(Tile, ProducerPriority, EVTInvalidatePriority::Normal);
+				const uint16 LoadRequestIndex = MergedRequestList->LockLoadRequest(TileRequest, ProducerLayerMaskToLoad, bStreamingRequest);
 				if (LoadRequestIndex == 0xffff)
 				{
 					// Overflowed the request list...try to lock the tile again next frame
@@ -2790,7 +2866,14 @@ TUniquePtr<FVirtualTextureUpdater> FVirtualTextureSystem::BeginUpdate(FRDGBuilde
 				if (CVarVTProduceLockedTilesOnFlush.GetValueOnRenderThread())
 				{
 					// Collect locked pages to be produced again
-					PhysicalSpace->GetPagePool().GetAllLockedPages(this, MappedTilesToProduce);
+					TSet<FVirtualTextureLocalTile> LockedTiles;
+					PhysicalSpace->GetPagePool().GetAllLockedPages(this, LockedTiles);
+					for (FVirtualTextureLocalTile LockedTile : LockedTiles)
+					{
+						// This is not really runtime-relevant (debug cache flush) so it's fine to re-produce the tile with normal priority again : 
+						AddOrMergeTileRequest(FVirtualTextureLocalTileRequest(LockedTile, EVTProducerPriority::Normal, EVTInvalidatePriority::Normal), MappedTilesToProduce);
+					}
+					
 				}
 				// Flush unlocked pages
 				PhysicalSpace->GetPagePool().EvictAllPages(this);
diff --git a/Engine/Source/Runtime/Renderer/Private/VT/VirtualTextureSystem.h b/Engine/Source/Runtime/Renderer/Private/VT/VirtualTextureSystem.h
--- a/Engine/Source/Runtime/Renderer/Private/VT/VirtualTextureSystem.h
+++ b/Engine/Source/Runtime/Renderer/Private/VT/VirtualTextureSystem.h
@@ -162,7 +162,7 @@ public:
 	void RequestRecordedTiles(TArray<uint64>&& InPageRequests);
 
 	void FlushCache();
-	void FlushCache(FVirtualTextureProducerHandle const& ProducerHandle, int32 SpaceID, FIntRect const& TextureRegion, uint32 MaxLevelToEvict, uint32 MaxAgeToKeepMapped);
+	void FlushCache(FVirtualTextureProducerHandle const& ProducerHandle, int32 SpaceID, FIntRect const& TextureRegion, uint32 MaxLevelToEvict, uint32 MaxAgeToKeepMapped, EVTInvalidatePriority InvalidatePriority);
 
 	float GetGlobalMipBias() const;
 
@@ -205,8 +205,9 @@ private:
 	void RequestTilesForRegionInternal(const IAllocatedVirtualTexture* AllocatedVT, const FVector2D& InScreenSpaceSize, const FVector2D& InViewportPosition, const FVector2D& InViewportSize, const FVector2D& InUV0, const FVector2D& InUV1, int32 InMipLevel);
 	void RequestTilesInternal(const IAllocatedVirtualTexture* AllocatedVT, int32 InMipLevel);
 	void RequestTilesInternal(const IAllocatedVirtualTexture* AllocatedVT, const FVector2D& InScreenSpaceSize, int32 InMipLevel);
-	
-	int32 SubmitRequestsFromLocalTileList(FRHICommandList& RHICmdList, TArray<FVirtualTextureLocalTile>& OutDeferredTiles, const TSet<FVirtualTextureLocalTile>& LocalTileList, EVTProducePageFlags Flags, ERHIFeatureLevel::Type FeatureLevel, uint32 MaxRequestsToProduce);
+
+	int32 SubmitRequestsFromLocalTileRequests(FConcurrentLinearBulkObjectAllocator& Allocator, FRHICommandList& RHICmdList, TSet<FVirtualTextureLocalTileRequest>& OutDeferredTileRequests, const TSet<FVirtualTextureLocalTileRequest>& LocalTileRequests, EVTProducePageFlags Flags, ERHIFeatureLevel::Type FeatureLevel, uint32 MaxRequestsToProduce);
+	TArrayView<FVTLocalTilePriorityAndIndex> SortLocalTileRequests(FConcurrentLinearBulkObjectAllocator& Allocator, const TConstArrayView<FVirtualTextureLocalTileRequest>& InLocalTileRequests);
 
 	void GatherFeedbackRequests(FConcurrentLinearBulkObjectAllocator& Allocator, const FVirtualTextureUpdateSettings& Settings, const FVirtualTextureFeedback::FMapResult& FeedbackResult, FUniqueRequestList* MergedRequestList);
 	void GatherLockedTileRequests(FUniqueRequestList* MergedRequestList);
@@ -283,9 +284,9 @@ private:
 	TArray<FVirtualTextureLocalTile> TilesToLockForNextFrame;
 	FTexturePageLocks TileLocks;
 
-	TSet<FVirtualTextureLocalTile> ContinuousUpdateTilesToProduce;
-	TSet<FVirtualTextureLocalTile> MappedTilesToProduce;
-	TArray<FVirtualTextureLocalTile> TransientCollectedPages;
+	TSet<FVirtualTextureLocalTileRequest> ContinuousUpdateTilesToProduce;  
+	TSet<FVirtualTextureLocalTileRequest> MappedTilesToProduce; 
+	TSet<FVirtualTextureLocalTileRequest> TransientCollectedTilesToProduce; 
 	TArray<IVirtualTextureFinalizer*> Finalizers;
 
 #if WITH_EDITOR
