Fix r.TSR.AlphaChannel typo

#tests the old CVar r.TSR.AplhaChannel testing in ini triggers ensure and deprecate warning.
#rnx

diff --git a/Engine/Source/Runtime/Renderer/Private/PostProcess/TemporalSuperResolution.cpp b/Engine/Source/Runtime/Renderer/Private/PostProcess/TemporalSuperResolution.cpp
--- a/Engine/Source/Runtime/Renderer/Private/PostProcess/TemporalSuperResolution.cpp
+++ b/Engine/Source/Runtime/Renderer/Private/PostProcess/TemporalSuperResolution.cpp
@@ -31,13 +31,15 @@ TAutoConsoleVariable<int32> CVarTSRSupportLensDistortion(
 	ECVF_ReadOnly);
 
 TAutoConsoleVariable<int32> CVarTSRAlphaChannel(
-	TEXT("r.TSR.AplhaChannel"), -1,
+	TEXT("r.TSR.AlphaChannel"), -1,
 	TEXT("Controls whether TSR should process the scene color's alpha channel.\n")
 	TEXT(" -1: based of r.PostProcessing.PropagateAlpha (default);\n")
 	TEXT("  0: disabled;\n")
 	TEXT("  1: enabled.\n"),
 	ECVF_RenderThreadSafe);
 
+FAutoConsoleVariableDeprecated ShadowCVarTSRAplhaChannel(TEXT("r.TSR.AplhaChannel"), TEXT("r.TSR.AlphaChannel"), TEXT("5.6"));
+
 TAutoConsoleVariable<float> CVarTSRHistorySampleCount(
 	TEXT("r.TSR.History.SampleCount"), 16.0f,
 	TEXT("Maximum number sample for each output pixel in the history. Higher values means more stability on highlights on static images, ")
