Fixed leak due to not deallocating view snapshots when running in immediate clean up mode.

diff --git a/Engine/Source/Runtime/Renderer/Private/SceneRenderBuilder.cpp b/Engine/Source/Runtime/Renderer/Private/SceneRenderBuilder.cpp
--- a/Engine/Source/Runtime/Renderer/Private/SceneRenderBuilder.cpp
+++ b/Engine/Source/Runtime/Renderer/Private/SceneRenderBuilder.cpp
@@ -711,6 +711,7 @@ void FSceneRenderProcessor::Cleanup(FRHICommandListImmediate& RHICmdList, FScene
 	{
 		WaitForAsyncDeleteTask(); // This is to handle cases where as switch from async to immediate.
 		RHICmdList.ImmediateFlush(EImmediateFlushType::WaitForOutstandingTasksOnly);
+		ViewSnapshotCache::Deallocate();
 		CleanupSceneRenderer(Renderer);
 	}
 	else
