Use inline task for compatibility task trigger to avoid spawning a new async task for no reason.

diff --git a/Engine/Source/Runtime/Renderer/Private/SceneRendering.h b/Engine/Source/Runtime/Renderer/Private/SceneRendering.h
--- a/Engine/Source/Runtime/Renderer/Private/SceneRendering.h
+++ b/Engine/Source/Runtime/Renderer/Private/SceneRendering.h
@@ -2729,7 +2729,7 @@ FGraphEventRef CreateCompatibilityGraphEvent(PrerequisiteTaskCollectionType&& Pr
 	{
 		GraphEvent->DispatchSubsequents();
 
-	}, Prerequisites, UE::Tasks::ETaskPriority::High);
+	}, Prerequisites, UE::Tasks::ETaskPriority::High, UE::Tasks::EExtendedTaskPriority::Inline);
 	return GraphEvent;
 }
 
