Ray Tracing: Fix path tracing invalidation happening when toggling between Lit and Path Tracing modes

This also was leading to constant invalidation when using camera previews in the editor.

diff --git a/Engine/Source/Runtime/Renderer/Private/RayTracing/RayTracingMaterialHitShaders.cpp b/Engine/Source/Runtime/Renderer/Private/RayTracing/RayTracingMaterialHitShaders.cpp
--- a/Engine/Source/Runtime/Renderer/Private/RayTracing/RayTracingMaterialHitShaders.cpp
+++ b/Engine/Source/Runtime/Renderer/Private/RayTracing/RayTracingMaterialHitShaders.cpp
@@ -764,9 +764,10 @@ void FDeferredShadingSceneRenderer::CreateMaterialRayTracingMaterialPipeline(
 		: nullptr;
 
 	ERayTracingPipelineCacheFlags PipelineCacheFlags = ERayTracingPipelineCacheFlags::Default;
+	const bool bCompatiblePipelineSignatures = FallbackPipelineState ? IsCompatibleFallbackPipelineSignature(Scene->LastRayTracingMaterialPipelineSignature, Initializer) : false;
 	if (FallbackPipelineState
 		&& View.ViewState
-		&& IsCompatibleFallbackPipelineSignature(Scene->LastRayTracingMaterialPipelineSignature, Initializer)
+		&& bCompatiblePipelineSignatures
 		&& PipelineContainsHitShaders(FallbackPipelineState, RequiredHitShaders)
 		&& FindRayTracingMissShaderIndex(FallbackPipelineState, DefaultMissShader, false) != INDEX_NONE
 		&& (!bCallableShadersRequired || FindRayTracingCallableShaderIndex(FallbackPipelineState, DefaultCallableShader, false) != INDEX_NONE))
@@ -801,7 +802,11 @@ void FDeferredShadingSceneRenderer::CreateMaterialRayTracingMaterialPipeline(
 	{
 		// When using path tracing, a change in pipeline state compared to the previous frame means some new materials got added to the RTPSO
 		// and we should restart sample accumulation.
-		Scene->InvalidatePathTracedOutput();
+		// Only do this if the pipeline signatures are compatible, otherwise we might be toggling between Lit and PathTraced views and don't want to invalidate the state
+		if (bCompatiblePipelineSignatures)
+		{
+			Scene->InvalidatePathTracedOutput();
+		}
 	}
 
 	check(PipelineState);
