Helper function CalculateGridZParams(...) to be shared between LightGrid, Volumetric Fog, etc.

diff --git a/Engine/Source/Runtime/RenderCore/Public/RenderUtils.h b/Engine/Source/Runtime/RenderCore/Public/RenderUtils.h
--- a/Engine/Source/Runtime/RenderCore/Public/RenderUtils.h
+++ b/Engine/Source/Runtime/RenderCore/Public/RenderUtils.h
@@ -623,6 +623,29 @@ RENDERCORE_API bool GetTranslucentUsesLightIESProfiles();
 RENDERCORE_API bool GetHairStrandsUsesTriangleStrips();
 RENDERCORE_API uint32 GetHairStrandsLODMode();
 
+// Helper function to calculate Z params for froxel grids
+// Used in shaders with ComputeDepthFromZSlice(...)
+inline FVector CalculateGridZParams(float NearPlane, float FarPlane, float DepthDistributionScale, int32 GridSizeZ)
+{
+	// S = distribution scale
+	// B, O are solved for given the z distances of the first+last slice, and the # of slices.
+	//
+	// slice = log2(z*B + O) * S
+
+	// Don't spend lots of resolution right in front of the near plane
+	double NearOffset = .095 * 100;
+	// Space out the slices so they aren't all clustered at the near plane
+	double S = DepthDistributionScale;
+
+	double N = NearPlane + NearOffset;
+	double F = FarPlane;
+
+	double O = (F - N * FMath::Exp2(GridSizeZ / S)) / (F - N);
+	double B = (1 - O) / N;
+
+	return FVector(B, O, S);
+}
+
 // LuminanceMax is the amount of light that will cause the sensor to saturate at EV100.
 //  See also https://en.wikipedia.org/wiki/Film_speed and https://en.wikipedia.org/wiki/Exposure_value for more info.
 FORCEINLINE float EV100ToLuminance(float LuminanceMax, float EV100)
diff --git a/Engine/Source/Runtime/Renderer/Private/LightGridInjection.cpp b/Engine/Source/Runtime/Renderer/Private/LightGridInjection.cpp
--- a/Engine/Source/Runtime/Renderer/Private/LightGridInjection.cpp
+++ b/Engine/Source/Runtime/Renderer/Private/LightGridInjection.cpp
@@ -377,23 +377,11 @@ FORCEINLINE float GetTanRadAngleOrZero(float coneAngle)
 
 FVector GetLightGridZParams(float NearPlane, float FarPlane)
 {
-	// S = distribution scale
-	// B, O are solved for given the z distances of the first+last slice, and the # of slices.
-	//
-	// slice = log2(z*B + O) * S
-
-	// Don't spend lots of resolution right in front of the near plane
-	double NearOffset = .095 * 100;
 	// Space out the slices so they aren't all clustered at the near plane
-	double S = 4.05;
-
-	double N = NearPlane + NearOffset;
-	double F = FarPlane;
-
-	double O = (F - N * exp2((GLightGridSizeZ - 1) / S)) / (F - N);
-	double B = (1 - O) / N;
+	float DepthDistributionScale = 4.05f;
 
-	return FVector(B, O, S);
+	// reserve last slice to cover a larger range (see LightGridInjection.usf)
+	return CalculateGridZParams(NearPlane, FarPlane, DepthDistributionScale, GLightGridSizeZ - 1);
 }
 
 uint32 PackRG16(float In0, float In1)
diff --git a/Engine/Source/Runtime/Renderer/Private/VolumetricFog.cpp b/Engine/Source/Runtime/Renderer/Private/VolumetricFog.cpp
--- a/Engine/Source/Runtime/Renderer/Private/VolumetricFog.cpp
+++ b/Engine/Source/Runtime/Renderer/Private/VolumetricFog.cpp
@@ -1192,50 +1192,10 @@ bool ShouldRenderVolumetricFog(const FScene* Scene, const FSceneViewFamily& View
 
 FVector GetVolumetricFogGridZParams(float VolumetricFogStartDistance, float NearPlane, float FarPlane, int32 GridSizeZ)
 {
-	// S = distribution scale
-	// B, O are solved for given the z distances of the first+last slice, and the # of slices.
-	//
-	// slice = log2(z*B + O) * S
-
 	// Don't spend lots of resolution right in front of the near plane
-
 	NearPlane = FMath::Max(NearPlane, double(VolumetricFogStartDistance));
 
-	double NearOffset = .095 * 100.0;
-	// Space out the slices so they aren't all clustered at the near plane
-	double S = GVolumetricFogDepthDistributionScale;
-
-	double N = NearPlane + NearOffset;
-	double F = FarPlane;
-
-	double O = (F - N * FMath::Exp2(GridSizeZ / S)) / (F - N);
-	double B = (1 - O) / N;
-
-// 	double O2 = (FMath::Exp2((GridSizeZ - 1) / S) - F / N) / (-F / N + 1);
-// 
-// 	float FloatN = (float)N;
-// 	float FloatF = (float)F;
-// 	float FloatB = (float)B;
-// 	float FloatO = (float)O;
-// 	float FloatS = (float)S;
-// 
-// 	float NSlice = FMath::Log2(FloatN*FloatB + FloatO) * FloatS;
-// 	float NearPlaneSlice = FMath::Log2(NearPlane*FloatB + FloatO) * FloatS;
-// 	float FSlice = FMath::Log2(FloatF*FloatB + FloatO) * FloatS;
-	// y = log2(z*B + O) * S
-	// f(N) = 0 = log2(N*B + O) * S
-	// 1 = N*B + O
-	// O = 1 - N*B
-	// B = (1 - O) / N
-
-	// f(F) = GLightGridSizeZ - 1 = log2(F*B + O) * S
-	// exp2((GLightGridSizeZ - 1) / S) = F*B + O
-	// exp2((GLightGridSizeZ - 1) / S) = F * (1 - O) / N + O
-	// exp2((GLightGridSizeZ - 1) / S) = F / N - F / N * O + O
-	// exp2((GLightGridSizeZ - 1) / S) = F / N + (-F / N + 1) * O
-	// O = (exp2((GLightGridSizeZ - 1) / S) - F / N) / (-F / N + 1)
-
-	return FVector(B, O, S);
+	return CalculateGridZParams(NearPlane, FarPlane, GVolumetricFogDepthDistributionScale, GridSizeZ);
 }
 
 static FIntVector GetVolumetricFogGridSize(const FIntPoint& TargetResolution, int32& OutVolumetricFogGridPixelSize)
