Fixed typo in volumetric clouds.

#rnx
#tests PC dx12 editor

diff --git a/Engine/Source/Runtime/Renderer/Private/VolumetricCloudRendering.cpp b/Engine/Source/Runtime/Renderer/Private/VolumetricCloudRendering.cpp
--- a/Engine/Source/Runtime/Renderer/Private/VolumetricCloudRendering.cpp
+++ b/Engine/Source/Runtime/Renderer/Private/VolumetricCloudRendering.cpp
@@ -79,7 +79,7 @@ static TAutoConsoleVariable<int32> CVarVolumetricCloudDisableCompute(
 
 static TAutoConsoleVariable<int32> CVarVolumetricCloudApplyFogLate(
 	TEXT("r.VolumetricCloud.ApplyFogLate"), 0,
-	TEXT("When true, fog will be applied on cloud later in the frame. Tis helps hiding latency in fog upsate on cloud due to reprojection and reconstruction. It is however more expenssive."),
+	TEXT("When true, fog will be applied on cloud later in the frame. This applyies fog on cloud without visible latency due to reprojection and reconstruction. It is however more expenssive."),
 	ECVF_RenderThreadSafe);
 
 ////////////////////////////////////////////////////////////////////////// Shadow tracing
