Fix assertion that makes sure that we can use bit 30 of StateBucketId for mobile-specific draw command handling.

StateBucketId == -1 is a special case. It is never merged anyway.

diff --git a/Engine/Source/Runtime/Renderer/Private/MeshDrawCommands.cpp b/Engine/Source/Runtime/Renderer/Private/MeshDrawCommands.cpp
--- a/Engine/Source/Runtime/Renderer/Private/MeshDrawCommands.cpp
+++ b/Engine/Source/Runtime/Renderer/Private/MeshDrawCommands.cpp
@@ -247,7 +247,7 @@ void MergeMobileBasePassMeshDrawCommands(
 				FVisibleMeshDrawCommand& MeshCommandCSM = MeshCommandsCSM[i];
 
 				const int32 CSMBit = 1 << 30;
-				checkf((MeshCommand.StateBucketId & CSMBit) == 0, TEXT("Non-CSM draw command's StateBucketID should not have the CSM bit set."));
+				checkf(MeshCommand.StateBucketId == -1 || (MeshCommand.StateBucketId & CSMBit) == 0, TEXT("Non-CSM draw command's StateBucketID should not have the CSM bit set."));
 
 				if (MeshCommand.PrimitiveIdInfo.ScenePrimitiveId < ScenePrimitiveNum && MobileCSMVisibilityInfo.MobilePrimitiveCSMReceiverVisibilityMap[MeshCommand.PrimitiveIdInfo.ScenePrimitiveId])
 				{
