Fix shadow mask on gbuffer input being run several times if the VSM shadow projection pass if not culled by RDG.

The change tries to clarify/isolate what shadow projection runs for gbuffer input and for hair input.

diff --git a/Engine/Source/Runtime/Renderer/Private/LightRendering.cpp b/Engine/Source/Runtime/Renderer/Private/LightRendering.cpp
--- a/Engine/Source/Runtime/Renderer/Private/LightRendering.cpp
+++ b/Engine/Source/Runtime/Renderer/Private/LightRendering.cpp
@@ -2156,7 +2156,14 @@ void FDeferredShadingSceneRenderer::RenderLights(
 							ClearShadowMask(ScreenShadowMaskSubPixelTexture);
 						}
 
-						RenderDeferredShadowProjections(GraphBuilder, SceneTextures, TranslucencyLightingVolumeTextures, &LightSceneInfo, ScreenShadowMaskTexture, ScreenShadowMaskSubPixelTexture);
+						if (bIsHairPass)
+						{
+							RenderDeferredShadowProjections(GraphBuilder, SceneTextures, TranslucencyLightingVolumeTextures, &LightSceneInfo, nullptr, ScreenShadowMaskSubPixelTexture);
+						}
+						else
+						{
+							RenderDeferredShadowProjections(GraphBuilder, SceneTextures, TranslucencyLightingVolumeTextures, &LightSceneInfo, ScreenShadowMaskTexture, nullptr);
+						}
 					}
 
 					bUsedShadowMaskTexture = true;
diff --git a/Engine/Source/Runtime/Renderer/Private/ShadowRendering.cpp b/Engine/Source/Runtime/Renderer/Private/ShadowRendering.cpp
--- a/Engine/Source/Runtime/Renderer/Private/ShadowRendering.cpp
+++ b/Engine/Source/Runtime/Renderer/Private/ShadowRendering.cpp
@@ -2157,6 +2157,7 @@ void FSceneRenderer::RenderShadowProjections(
 				bProjectingForForwardShading);
 		};
 
+		if (ScreenShadowMaskTexture)
 		{
 			RDG_EVENT_SCOPE(GraphBuilder, "Shadows");
 			RenderNormalShadows(ScreenShadowMaskTexture, false);
@@ -2169,7 +2170,7 @@ void FSceneRenderer::RenderShadowProjections(
 		}
 	}
 
-	if (DistanceFieldShadows.Num() > 0)
+	if (ScreenShadowMaskTexture && DistanceFieldShadows.Num() > 0)
 	{
 		RDG_EVENT_SCOPE(GraphBuilder, "DistanceFieldShadows");
 
@@ -2320,28 +2321,38 @@ void FDeferredShadingSceneRenderer::RenderDeferredShadowProjections(
 	RDG_EVENT_SCOPE_STAT(GraphBuilder, ShadowProjection, "ShadowProjectionOnOpaque");
 	RDG_GPU_STAT_SCOPE(GraphBuilder, ShadowProjection);
 
-	const FVisibleLightInfo& VisibleLightInfo = VisibleLightInfos[LightSceneInfo->Id];
-	
 	const bool bProjectingForForwardShading = false;
-	RenderShadowProjections(GraphBuilder, SceneTextures, ScreenShadowMaskTexture, ScreenShadowMaskSubPixelTexture, LightSceneInfo, bProjectingForForwardShading);
-	GetSceneExtensionsRenderers().GetRenderer<FShadowSceneRenderer>().ApplyVirtualShadowMapProjectionForLight(GraphBuilder, SceneTextures, LightSceneInfo, ScreenShadowMaskTexture, ScreenShadowMaskSubPixelTexture);
+	const bool bIsGBufferInput = ScreenShadowMaskTexture != nullptr;
+	const bool bIsHairInput = ScreenShadowMaskSubPixelTexture != nullptr;
 
-	RenderCapsuleDirectShadows(GraphBuilder, *LightSceneInfo, ScreenShadowMaskTexture, VisibleLightInfo.CapsuleShadowsToProject, bProjectingForForwardShading);
-
-	// Inject deep shadow mask for regular shadow map. When using virtual shadow map, it is directly handled in the shadow kernel.
-	if (HairStrands::HasViewHairStrandsData(Views))
+	if (bIsGBufferInput)
 	{
-		bool bNeedHairShadowMaskPass = false;
-		const bool bVirtualShadowOnePass = VisibleLightInfo.VirtualShadowMapClipmaps.Num() > 0;
-		if (!bVirtualShadowOnePass && VisibleLightInfo.ShadowsToProject.Num() > 0)
-		{
-			bNeedHairShadowMaskPass = !VisibleLightInfo.ShadowsToProject[0]->HasVirtualShadowMap();
-		}
-		if (bNeedHairShadowMaskPass)
+		RenderShadowProjections(GraphBuilder, SceneTextures, ScreenShadowMaskTexture, nullptr, LightSceneInfo, bProjectingForForwardShading);
+		GetSceneExtensionsRenderers().GetRenderer<FShadowSceneRenderer>().ApplyVirtualShadowMapProjectionForLight(GraphBuilder, SceneTextures, LightSceneInfo, EVirtualShadowMapProjectionInputType::GBuffer, ScreenShadowMaskTexture);
+
+		const FVisibleLightInfo& VisibleLightInfo = VisibleLightInfos[LightSceneInfo->Id];
+		RenderCapsuleDirectShadows(GraphBuilder, *LightSceneInfo, ScreenShadowMaskTexture, VisibleLightInfo.CapsuleShadowsToProject, bProjectingForForwardShading);
+
+		// Inject deep shadow mask for regular shadow map. When using virtual shadow map, it is directly handled in the shadow kernel.
+		if (HairStrands::HasViewHairStrandsData(Views))
 		{
-			RenderHairStrandsShadowMask(GraphBuilder, Views, LightSceneInfo, VisibleLightInfos, false /*bForward*/, ScreenShadowMaskTexture);
+			bool bNeedHairShadowMaskPass = false;
+			const bool bVirtualShadowOnePass = VisibleLightInfo.VirtualShadowMapClipmaps.Num() > 0;
+			if (!bVirtualShadowOnePass && VisibleLightInfo.ShadowsToProject.Num() > 0)
+			{
+				bNeedHairShadowMaskPass = !VisibleLightInfo.ShadowsToProject[0]->HasVirtualShadowMap();
+			}
+			if (bNeedHairShadowMaskPass)
+			{
+				RenderHairStrandsShadowMask(GraphBuilder, Views, LightSceneInfo, VisibleLightInfos, false /*bForward*/, ScreenShadowMaskTexture);
+			}
 		}
 	}
+	else if (bIsHairInput)
+	{
+		RenderShadowProjections(GraphBuilder, SceneTextures, nullptr, ScreenShadowMaskSubPixelTexture, LightSceneInfo, bProjectingForForwardShading);
+		GetSceneExtensionsRenderers().GetRenderer<FShadowSceneRenderer>().ApplyVirtualShadowMapProjectionForLight(GraphBuilder, SceneTextures, LightSceneInfo, EVirtualShadowMapProjectionInputType::HairStrands, ScreenShadowMaskSubPixelTexture);
+	}
 }
 
 void FMobileSceneRenderer::RenderModulatedShadowProjections(FRHICommandList& RHICmdList, int32 ViewIndex, const FViewInfo& View)
diff --git a/Engine/Source/Runtime/Renderer/Private/Shadows/ShadowSceneRenderer.cpp b/Engine/Source/Runtime/Renderer/Private/Shadows/ShadowSceneRenderer.cpp
--- a/Engine/Source/Runtime/Renderer/Private/Shadows/ShadowSceneRenderer.cpp
+++ b/Engine/Source/Runtime/Renderer/Private/Shadows/ShadowSceneRenderer.cpp
@@ -748,8 +748,8 @@ void FShadowSceneRenderer::ApplyVirtualShadowMapProjectionForLight(
 	FRDGBuilder& GraphBuilder,
 	const FMinimalSceneTextures& SceneTextures,
 	const FLightSceneInfo* LightSceneInfo,
-	FRDGTextureRef OutputScreenShadowMaskTexture,
-	FRDGTextureRef OutputScreenShadowMaskSubPixelTexture)
+	const EVirtualShadowMapProjectionInputType InputType,
+	FRDGTextureRef OutputScreenShadowMaskTexture)
 {
 	if (!VirtualShadowMapArray.HasAnyShadowData())
 	{
@@ -779,7 +779,7 @@ void FShadowSceneRenderer::ApplyVirtualShadowMapProjectionForLight(
 			int32 VirtualShadowMapId = VisibleLightInfo.GetVirtualShadowMapId(&View);
 			
 			// Some lights can elide the screen shadow mask entirely, in which case they will be sampled directly in the lighting shader
-			if (OutputScreenShadowMaskTexture)
+			if (OutputScreenShadowMaskTexture && InputType == EVirtualShadowMapProjectionInputType::GBuffer)
 			{
 				if (VisibleLightInfo.VirtualShadowMapClipmaps.Num() > 0)
 				{
@@ -834,7 +834,7 @@ void FShadowSceneRenderer::ApplyVirtualShadowMapProjectionForLight(
 			}
 
 			// Sub-pixel shadow (no denoising for hair)
-			if (HairStrands::HasViewHairStrandsData(View) && OutputScreenShadowMaskSubPixelTexture)
+			if (HairStrands::HasViewHairStrandsData(View) && OutputScreenShadowMaskTexture && InputType == EVirtualShadowMapProjectionInputType::HairStrands)
 			{
 				if (VisibleLightInfo.VirtualShadowMapClipmaps.Num() > 0)
 				{
@@ -848,7 +848,7 @@ void FShadowSceneRenderer::ApplyVirtualShadowMapProjectionForLight(
 						VisibleLightInfo.FindShadowClipmapForView(&View),
 						false, // bModulateRGB
 						nullptr, // TiledVSMProjection
-						OutputScreenShadowMaskSubPixelTexture,
+						OutputScreenShadowMaskTexture,
 						VisibleLightInfo.FindFirstPersonShadowClipmapForView(&View));
 				}
 				else if (bShouldUseVirtualShadowMapOnePassProjection)
@@ -863,7 +863,7 @@ void FShadowSceneRenderer::ApplyVirtualShadowMapProjectionForLight(
 						EVirtualShadowMapProjectionInputType::HairStrands,
 						VirtualShadowMapId,
 						VirtualShadowMapMaskBitsHairStrands,
-						OutputScreenShadowMaskSubPixelTexture);
+						OutputScreenShadowMaskTexture);
 				}
 				else
 				{
@@ -876,7 +876,7 @@ void FShadowSceneRenderer::ApplyVirtualShadowMapProjectionForLight(
 						EVirtualShadowMapProjectionInputType::HairStrands,
 						*LightSceneInfo,
 						VirtualShadowMapId,
-						OutputScreenShadowMaskSubPixelTexture);
+						OutputScreenShadowMaskTexture);
 				}
 			}
 		}
diff --git a/Engine/Source/Runtime/Renderer/Private/Shadows/ShadowSceneRenderer.h b/Engine/Source/Runtime/Renderer/Private/Shadows/ShadowSceneRenderer.h
--- a/Engine/Source/Runtime/Renderer/Private/Shadows/ShadowSceneRenderer.h
+++ b/Engine/Source/Runtime/Renderer/Private/Shadows/ShadowSceneRenderer.h
@@ -16,6 +16,7 @@ class FWholeSceneProjectedShadowInitializer;
 class FRDGBuilder;
 class FVirtualShadowMapPerLightCacheEntry;
 struct FNaniteVisibilityQuery;
+enum class EVirtualShadowMapProjectionInputType;
 
 namespace UE::Renderer::Private
 {
@@ -85,8 +86,8 @@ public:
 		FRDGBuilder& GraphBuilder,
 		const FMinimalSceneTextures& SceneTextures,
 		const FLightSceneInfo* LightSceneInfo,
-		FRDGTextureRef OutputScreenShadowMaskTexture,
-		FRDGTextureRef OutputScreenShadowMaskSubPixelTexture);
+		const EVirtualShadowMapProjectionInputType InputType,
+		FRDGTextureRef OutputScreenShadowMaskTexture);
 
 	// One pass projection stuff. Set up in RenderVitualShadowMapProjectionMaskBits
 	FRDGTextureRef VirtualShadowMapMaskBits = nullptr;
