[Lit Translucency] Fix local lights being culled too early in Translucency Lighting Volume.
- incorrect logic was causing only inner volume bounds to be taken into account.

diff --git a/Engine/Source/Runtime/Renderer/Private/TranslucentLighting.cpp b/Engine/Source/Runtime/Renderer/Private/TranslucentLighting.cpp
--- a/Engine/Source/Runtime/Renderer/Private/TranslucentLighting.cpp
+++ b/Engine/Source/Runtime/Renderer/Private/TranslucentLighting.cpp
@@ -1179,7 +1179,7 @@ void FTranslucentLightInjectionCollector::AddLightForInjection(
 		for (uint32 VolumeCascadeIndex = 0; VolumeCascadeIndex < TVC_MAX; ++VolumeCascadeIndex)
 		{
 			VolumeBounds[VolumeCascadeIndex] = CalculateLightVolumeBounds(LightSceneInfo.Proxy->GetBoundingSphere(), View, VolumeCascadeIndex, LightType == LightType_Directional);
-			bAnyBoundsValid = bAnyBoundsValid || VolumeBounds->IsValid();
+			bAnyBoundsValid = bAnyBoundsValid || VolumeBounds[VolumeCascadeIndex].IsValid();
 		}
 		if (!bAnyBoundsValid)
 		{
