Fix depth pre-pass (missing skinned meshes) when velocity rendering support is globally enabled but then disabled at runtime.

In this case EarlyZPassMode is DDM_AllOpaqueNoVelocity since EarlyZPassMode is not updated when velocity is disabled at runtime.
We still need the velocity pass during the depth-prepass to fill the depth buffer but we don't need to write to the velocity target.

Without passing FFullDepthPrepassParameterCollection to GraphBuilder.AddPass, RDG would not know that VelocityInstanceCullingDrawParams are used by the depth/velocity pass and trigger a RDG validation error.

diff --git a/Engine/Source/Runtime/Renderer/Private/MobileShadingRenderer.cpp b/Engine/Source/Runtime/Renderer/Private/MobileShadingRenderer.cpp
--- a/Engine/Source/Runtime/Renderer/Private/MobileShadingRenderer.cpp
+++ b/Engine/Source/Runtime/Renderer/Private/MobileShadingRenderer.cpp
@@ -815,7 +815,7 @@ void FMobileSceneRenderer::RenderFullDepthPrepass(FRDGBuilder& GraphBuilder, TAr
 
 		View.BeginRenderView();
 
-		const bool bRenderVelocityInDepthPrePass = (Scene->EarlyZPassMode == DDM_AllOpaqueNoVelocity) && bShouldRenderVelocities;
+		const bool bRenderVelocityInDepthPrePass = (Scene->EarlyZPassMode == DDM_AllOpaqueNoVelocity);
 
 		struct FFullDepthPrepassParameterCollection
 		{
@@ -825,7 +825,7 @@ void FMobileSceneRenderer::RenderFullDepthPrepass(FRDGBuilder& GraphBuilder, TAr
 		auto ParameterCollection = GraphBuilder.AllocParameters<FFullDepthPrepassParameterCollection>();
 
 		auto PassParameters = &ParameterCollection->PassParameters;
-		if (bRenderVelocityInDepthPrePass)
+		if (bShouldRenderVelocities)
 		{
 			PassParameters->RenderTargets[0] = FRenderTargetBinding(SceneTextures.Velocity, ERenderTargetLoadAction::EClear);
 		}
@@ -844,29 +844,30 @@ void FMobileSceneRenderer::RenderFullDepthPrepass(FRDGBuilder& GraphBuilder, TAr
 			PassParameters->RenderTargets.DepthStencil.SetDepthLoadAction(ERenderTargetLoadAction::ELoad);
 			PassParameters->RenderTargets.DepthStencil.SetStencilLoadAction(ERenderTargetLoadAction::ELoad);
 			PassParameters->RenderTargets.DepthStencil.SetDepthStencilAccess(FExclusiveDepthStencil::DepthWrite_StencilWrite);
-			if (bRenderVelocityInDepthPrePass)
+			if (bShouldRenderVelocities)
 			{
 				PassParameters->RenderTargets[0].SetLoadAction(ERenderTargetLoadAction::ELoad);
 			}
 		}
 
 		BuildMeshRenderingCommands(GraphBuilder, EMeshPass::DepthPass, View, Scene->GPUScene, InstanceCullingManager, PassParameters->InstanceCullingDrawParams);
+		FInstanceCullingDrawParams* VelocityInstanceCullingDrawParams = nullptr;
 		if (bRenderVelocityInDepthPrePass)
 		{
-			BuildMeshRenderingCommands(GraphBuilder, EMeshPass::Velocity, View, Scene->GPUScene, InstanceCullingManager, ParameterCollection->VelocityInstanceCullingDrawParams);
-		}
+			VelocityInstanceCullingDrawParams = View.ParallelMeshDrawCommandPasses[EMeshPass::DepthPass] ? &ParameterCollection->VelocityInstanceCullingDrawParams : &PassParameters->InstanceCullingDrawParams;
+			BuildMeshRenderingCommands(GraphBuilder, EMeshPass::Velocity, View, Scene->GPUScene, InstanceCullingManager, *VelocityInstanceCullingDrawParams);		}
 
 		GraphBuilder.AddPass(
 			RDG_EVENT_NAME("FullDepthPrepass"),
 			PassParameters,
 			ERDGPassFlags::Raster,
-			[this, ParameterCollection, &View, bRenderVelocityInDepthPrePass](FRHICommandList& RHICmdList)
+			[this, ParameterCollection, &View, VelocityInstanceCullingDrawParams](FRHICommandList& RHICmdList)
 			{
 				auto PassParameters = &ParameterCollection->PassParameters;
 				RenderPrePass(RHICmdList, View, &PassParameters->InstanceCullingDrawParams);
-				if (bRenderVelocityInDepthPrePass)
+				if (VelocityInstanceCullingDrawParams)
 				{
-					RenderVelocityPass(RHICmdList, View, &ParameterCollection->VelocityInstanceCullingDrawParams);
+					RenderVelocityPass(RHICmdList, View, VelocityInstanceCullingDrawParams);
 				}
 			});
 	}
