Added a cvar to turn off the TAA bilateral filter (r.DOF.TAA.CoCBilateralFilterStrength).
The history rejection from this filter can cause high-contrast aliasing on out-of-focus edges, which leads to very noticable flickering bokeh.

diff --git a/Engine/Shaders/Private/TemporalAA.usf b/Engine/Shaders/Private/TemporalAA.usf
--- a/Engine/Shaders/Private/TemporalAA.usf
+++ b/Engine/Shaders/Private/TemporalAA.usf
@@ -536,6 +536,8 @@ SamplerState HistoryBufferSampler_1;
 float4 HistoryBufferSize;
 float4 HistoryBufferUVMinMax;
 
+float CoCBilateralFilterStrength;
+
 float4 OutputViewportSize;
 float4 OutputViewportRect;
 float3 OutputQuantizationError;
@@ -2259,7 +2261,7 @@ FTAAHistoryPayload TemporalAASample(uint2 GroupId, uint2 GroupThreadId, uint Gro
 		{
 			float BilateralWeight = ComputeBilateralWeight(IntermediaryResult.Filtered.CocRadius, History.CocRadius);
 
-			BlendFinal = lerp(1, BlendFinal, BilateralWeight);
+			BlendFinal = lerp(BlendFinal, 1, (1-BilateralWeight)*CoCBilateralFilterStrength);
 		}
 		#endif
 		
diff --git a/Engine/Source/Runtime/Renderer/Private/PostProcess/DiaphragmDOF.cpp b/Engine/Source/Runtime/Renderer/Private/PostProcess/DiaphragmDOF.cpp
--- a/Engine/Source/Runtime/Renderer/Private/PostProcess/DiaphragmDOF.cpp
+++ b/Engine/Source/Runtime/Renderer/Private/PostProcess/DiaphragmDOF.cpp
@@ -157,6 +157,18 @@ TAutoConsoleVariable<float> CVarScatterNeighborCompareMaxColor(
 	TEXT(" (Default: 10)."),
 	ECVF_RenderThreadSafe);
 
+TAutoConsoleVariable<float> CVarCoCBilateralFilterStrength(
+	TEXT("r.DOF.TAA.CoCBilateralFilterStrength"),
+	0.0,
+	TEXT("If 1, a bilateral filter based on the circle-of-confusion is used to reject history in the pre-DoF TAA pass.\n")
+	TEXT("If 0, this is disabled.\n")
+	TEXT("When enabled, you may get stronger noise on edges of objects that separate foreground and background.\n")
+	TEXT("When disabled, you may get stronger ghosting along those same edges.\n")
+	TEXT(" (Default: 0)."),
+	ECVF_RenderThreadSafe);
+
+
+
 } // namespace
 
 
@@ -1741,6 +1753,7 @@ bool DiaphragmDOF::AddPasses(
 		TAAParameters.SceneVelocityTexture = SceneTextures.GBufferVelocityTexture;
 		TAAParameters.SceneColorInput = HalfResGatherInputTextures.SceneColor;
 		TAAParameters.SceneMetadataInput = HalfResGatherInputTextures.SeparateCoc;
+		TAAParameters.CoCBilateralFilterStrength = FMath::Clamp(CVarCoCBilateralFilterStrength.GetValueOnRenderThread(), 0.0f, 1.0f);
 
 		FTAAOutputs TAAOutputs = AddTemporalAAPass(
 			GraphBuilder,
diff --git a/Engine/Source/Runtime/Renderer/Private/PostProcess/TemporalAA.cpp b/Engine/Source/Runtime/Renderer/Private/PostProcess/TemporalAA.cpp
--- a/Engine/Source/Runtime/Renderer/Private/PostProcess/TemporalAA.cpp
+++ b/Engine/Source/Runtime/Renderer/Private/PostProcess/TemporalAA.cpp
@@ -127,6 +127,8 @@ public:
 		SHADER_PARAMETER(FVector4f, HistoryBufferUVMinMax)
 		SHADER_PARAMETER(FVector4f, ScreenPosToHistoryBufferUV)
 
+		SHADER_PARAMETER(float, CoCBilateralFilterStrength)
+
 		// Temporal upsample specific parameters.
 		SHADER_PARAMETER(FVector4f, InputViewSize)
 		SHADER_PARAMETER(FVector2f, InputViewMin)
@@ -773,6 +775,8 @@ FTAAOutputs AddTemporalAAPass(
 				InputViewSizeInvScale / InputViewRect.Width(), InputViewSizeInvScale / InputViewRect.Height());
 		}
 
+		PassParameters->CoCBilateralFilterStrength = Inputs.CoCBilateralFilterStrength;
+
 		PassParameters->OutputViewportSize = FVector4f(
 			PracticableDestRect.Width(), PracticableDestRect.Height(), 1.0f / float(PracticableDestRect.Width()), 1.0f / float(PracticableDestRect.Height()));
 		PassParameters->OutputViewportRect = FVector4f(PracticableDestRect.Min.X, PracticableDestRect.Min.Y, PracticableDestRect.Max.X, PracticableDestRect.Max.Y);
diff --git a/Engine/Source/Runtime/Renderer/Private/PostProcess/TemporalAA.h b/Engine/Source/Runtime/Renderer/Private/PostProcess/TemporalAA.h
--- a/Engine/Source/Runtime/Renderer/Private/PostProcess/TemporalAA.h
+++ b/Engine/Source/Runtime/Renderer/Private/PostProcess/TemporalAA.h
@@ -125,6 +125,10 @@ struct FTAAPassParameters
 	// Optional information that get anti aliased, such as separate CoC for DOF.
 	FRDGTexture* SceneMetadataInput = nullptr;
 
+	// If 1, a bilateral filter based on the circle-of-confusion for depth-of-field is used to reject history.
+	// If 0, this is disabled.
+	float CoCBilateralFilterStrength = 1.0;
+
 
 	FTAAPassParameters(const FViewInfo& View)
 		: InputViewRect(View.ViewRect)
