- clear update array after processing them
- Ensure geometry to update is valid (check it's not evicted)

diff --git a/Engine/Source/Runtime/Renderer/Private/RayTracing/RayTracingDynamicGeometryUpdateManager.cpp b/Engine/Source/Runtime/Renderer/Private/RayTracing/RayTracingDynamicGeometryUpdateManager.cpp
--- a/Engine/Source/Runtime/Renderer/Private/RayTracing/RayTracingDynamicGeometryUpdateManager.cpp
+++ b/Engine/Source/Runtime/Renderer/Private/RayTracing/RayTracingDynamicGeometryUpdateManager.cpp
@@ -535,8 +535,14 @@ uint32 FRayTracingDynamicGeometryUpdateManager::Update()
 
 		for (auto Iter = RTGeometriesToUpdate.CreateIterator(); Iter; ++Iter)
 		{
+			FRayTracingGeometry* RTGeometry = Iter.Key();
+			if (!ensureMsgf(RTGeometry->GetRHI(), TEXT("Skipping request with invalid ray tracing geometry in FRayTracingDynamicGeometryUpdateManager. Geometry->IsEvicted(): %d."), RTGeometry->IsEvicted()))
+			{
+				continue;
+			}
+
 			FRayTracingDynamicGeometryBuildParams GeometryBuildParams;
-			GeometryBuildParams.Geometry = Iter.Key();
+			GeometryBuildParams.Geometry = RTGeometry;
 
 			EAccelerationStructureBuildMode BuildMode = Iter.Value();
 			if (BuildMode == EAccelerationStructureBuildMode::Build)
@@ -548,6 +554,9 @@ uint32 FRayTracingDynamicGeometryUpdateManager::Update()
 				DynamicGeometryUpdates.Add(GeometryBuildParams);
 			}
 		}
+
+		// Clear working data
+		RTGeometriesToUpdate.Reset();
 	}
 
 	// Early out of nothing to do
