Fixed "unbound parameters not represented in the parameter struct" when disabling shader optimizations.

diff --git a/Engine/Source/Runtime/Renderer/Private/RayTracing/RayTracingDebug.cpp b/Engine/Source/Runtime/Renderer/Private/RayTracing/RayTracingDebug.cpp
--- a/Engine/Source/Runtime/Renderer/Private/RayTracing/RayTracingDebug.cpp
+++ b/Engine/Source/Runtime/Renderer/Private/RayTracing/RayTracingDebug.cpp
@@ -344,6 +344,14 @@ public:
 		return ShouldCompileRayTracingShadersForProject(Parameters.Platform);
 	}
 
+	static void ModifyCompilationEnvironment(const FGlobalShaderPermutationParameters& Parameters, FShaderCompilerEnvironment& OutEnvironment)
+	{
+		FGlobalShader::ModifyCompilationEnvironment(Parameters, OutEnvironment);
+
+		// Temporary workaround for "unbound parameters not represented in the parameter struct" when disabling optimizations
+		OutEnvironment.CompilerFlags.Add(CFLAG_ForceOptimization);
+	}
+
 	static ERayTracingPayloadType GetRayTracingPayloadType(const int32 PermutationId)
 	{
 		return ERayTracingPayloadType::RayTracingDebug;
