Fixed race condition in FSceneViewState for GPU readback buffers that are enqueued with AddEnqueueCopyPass.

diff --git a/Engine/Source/Runtime/Renderer/Private/RendererScene.cpp b/Engine/Source/Runtime/Renderer/Private/RendererScene.cpp
--- a/Engine/Source/Runtime/Renderer/Private/RendererScene.cpp
+++ b/Engine/Source/Runtime/Renderer/Private/RendererScene.cpp
@@ -382,6 +382,23 @@ FSceneViewState::~FSceneViewState()
 	}
 }
 
+void FSceneViewState::ReleaseRHI()
+{
+	// FSceneViewState contains multiple sub-components that have GPU readback buffers which can have in-flight async copy
+	// passes registered with RDG (e.g. AddEnqueueCopyPass), so we need to sync the execution tasks prior to releasing those buffers.
+	FRDGBuilder::WaitForAsyncExecuteTask();
+
+	HZBOcclusionTests.ReleaseRHI();
+	EyeAdaptationManager.SafeRelease(); 
+	SubstrateViewDebugData.SafeRelease();
+	OcclusionFeedback.ReleaseResource();
+PRAGMA_DISABLE_DEPRECATION_WARNINGS
+	bValidEyeAdaptationTexture = false;
+PRAGMA_ENABLE_DEPRECATION_WARNINGS
+	bValidEyeAdaptationBuffer = false;
+}
+
+
 inline FScene::FPrimitiveSceneProxyType::FPrimitiveSceneProxyType(const FPrimitiveSceneProxy *PrimitiveSceneProxy) 
 	: ProxyTypeHash(PrimitiveSceneProxy->GetTypeHash())
 	, bIsAlwaysVisible(PrimitiveSceneProxy->IsAlwaysVisible())
diff --git a/Engine/Source/Runtime/Renderer/Private/ScenePrivate.h b/Engine/Source/Runtime/Renderer/Private/ScenePrivate.h
--- a/Engine/Source/Runtime/Renderer/Private/ScenePrivate.h
+++ b/Engine/Source/Runtime/Renderer/Private/ScenePrivate.h
@@ -1412,17 +1412,7 @@ PRAGMA_ENABLE_DEPRECATION_WARNINGS
 		HZBOcclusionTests.InitRHI(RHICmdList);
 	}
 
-	virtual void ReleaseRHI() override
-	{
-		HZBOcclusionTests.ReleaseRHI();
-		EyeAdaptationManager.SafeRelease(); 
-		SubstrateViewDebugData.SafeRelease();
-		OcclusionFeedback.ReleaseResource();
-PRAGMA_DISABLE_DEPRECATION_WARNINGS
-		bValidEyeAdaptationTexture = false;
-PRAGMA_ENABLE_DEPRECATION_WARNINGS
-		bValidEyeAdaptationBuffer = false;
-	}
+	virtual void ReleaseRHI() override;
 
 	// FSceneViewStateInterface
 	RENDERER_API virtual void Destroy() override;
