Scene capture improvements for post-process settings.

* Fix previously "lying" post-process setting defaults, which were overriden in code but invisible to the user. Default values now reflect the code.
* Expose new main view camera post-processing option to automatically inherit the camera post-processing settings. Local defaults are ignored in this case, but local active overrides will function as usual (per tooltip).

diff --git a/Engine/Source/Runtime/Core/Public/UObject/FortniteMainBranchObjectVersions.inl b/Engine/Source/Runtime/Core/Public/UObject/FortniteMainBranchObjectVersions.inl
--- a/Engine/Source/Runtime/Core/Public/UObject/FortniteMainBranchObjectVersions.inl
+++ b/Engine/Source/Runtime/Core/Public/UObject/FortniteMainBranchObjectVersions.inl
@@ -565,5 +565,8 @@ FORTNITE_MAIN_VERSION(AnimSequenceRawDataOnlyFlagRemoval, 183)
 // HLOD relevancy of Level Instances was previously ignored, now taken into account. Reset to the default behavior.
 FORTNITE_MAIN_VERSION(ResetLevelInstanceHLODRelevancy, 184)
 
+// Updated default scene capture post-processing settings to reflect the underlying implementation overrides
+FORTNITE_MAIN_VERSION(SceneCaptureDefaultSettings, 185)
+
 // -----<new versions can be added above this line>-------------------------------------------------
 #undef FORTNITE_MAIN_VERSION
\ No newline at end of file
diff --git a/Engine/Source/Runtime/Engine/Classes/Components/SceneCaptureComponent2D.h b/Engine/Source/Runtime/Engine/Classes/Components/SceneCaptureComponent2D.h
--- a/Engine/Source/Runtime/Engine/Classes/Components/SceneCaptureComponent2D.h
+++ b/Engine/Source/Runtime/Engine/Classes/Components/SceneCaptureComponent2D.h
@@ -155,7 +155,7 @@ public:
 	UPROPERTY(EditAnywhere, BlueprintReadWrite, Category = SceneCapture, meta = (DisplayName = "Fog only on rendered pixels"))
 	uint32 bConsiderUnrenderedOpaquePixelAsFullyTranslucent : 1;
 
-	/** Render with main view family (where bIsMainViewFamily == true) */
+	/** Render with main view family, for example with the main editor or game viewport which mark their view families as "main". */
 	UPROPERTY(EditAnywhere, BlueprintReadWrite, Category=PassInfo)
 	uint32 bMainViewFamily : 1;
 
@@ -167,6 +167,10 @@ public:
 	UPROPERTY(EditAnywhere, BlueprintReadWrite, Category=PassInfo)
 	uint32 bMainViewCamera : 1;
 
+	/** Inherit the main view camera post-process settings and ignore local default values. Local active overrides will function as usual. */
+	UPROPERTY(EditAnywhere, BlueprintReadWrite, Category = PassInfo, meta = (EditCondition = "bMainViewCamera"))
+	uint32 bInheritMainViewCameraPostProcessSettings : 1;
+
 	/** When rendering with main view resolution, ignore screen percentage scale and render at full resolution.  Temporal AA jitter is also disabled. */
 	UPROPERTY(EditAnywhere, BlueprintReadWrite, Category = PassInfo)
 	uint32 bIgnoreScreenPercentage : 1;
diff --git a/Engine/Source/Runtime/Engine/Classes/Components/SceneCaptureComponentCube.h b/Engine/Source/Runtime/Engine/Classes/Components/SceneCaptureComponentCube.h
--- a/Engine/Source/Runtime/Engine/Classes/Components/SceneCaptureComponentCube.h
+++ b/Engine/Source/Runtime/Engine/Classes/Components/SceneCaptureComponentCube.h
@@ -43,6 +43,7 @@ public:
 
 	//~ Begin UObject Interface
 	static ENGINE_API void AddReferencedObjects(UObject* InThis, FReferenceCollector& Collector);
+	ENGINE_API virtual void Serialize(FArchive& Ar);
 #if WITH_EDITOR
 	ENGINE_API virtual void PostEditChangeProperty(FPropertyChangedEvent& PropertyChangedEvent) override;
 #endif // WITH_EDITOR
diff --git a/Engine/Source/Runtime/Engine/Private/Components/SceneCaptureComponent.cpp b/Engine/Source/Runtime/Engine/Private/Components/SceneCaptureComponent.cpp
--- a/Engine/Source/Runtime/Engine/Private/Components/SceneCaptureComponent.cpp
+++ b/Engine/Source/Runtime/Engine/Private/Components/SceneCaptureComponent.cpp
@@ -655,10 +655,16 @@ USceneCaptureComponent2D::USceneCaptureComponent2D(const FObjectInitializer& Obj
 	bMainViewFamily = false;
 	bMainViewResolution = false;
 	bMainViewCamera = false;
+	bInheritMainViewCameraPostProcessSettings = false;
 	bIgnoreScreenPercentage = false;
 	
 	TileID = 0;
 
+	// Update default post-processing values to reflect default values (previously and still) applied in code.
+	PostProcessSettings.DynamicGlobalIlluminationMethod = EDynamicGlobalIlluminationMethod::None;
+	PostProcessSettings.ReflectionMethod = EReflectionMethod::None;
+	PostProcessSettings.LumenSurfaceCacheResolution = 0.5f;
+
 	// Legacy initialization.
 	{
 		// previous behavior was to capture 2d scene captures before cube scene captures.
@@ -933,17 +939,24 @@ void USceneCaptureComponent2D::PostEditChangeProperty(FPropertyChangedEvent& Pro
 void USceneCaptureComponent2D::Serialize(FArchive& Ar)
 {
 	Ar.UsingCustomVersion(FFortniteMainBranchObjectVersion::GUID);
+	Ar.UsingCustomVersion(FUE5ReleaseStreamObjectVersion::GUID);
+
 	if (Ar.CustomVer(FFortniteMainBranchObjectVersion::GUID) < FFortniteMainBranchObjectVersion::OrthographicCameraDefaultSettings)
 	{
 		OrthoWidth = 512.0f;
 	}
 
-	Ar.UsingCustomVersion(FUE5ReleaseStreamObjectVersion::GUID);
 	if (Ar.CustomVer(FUE5ReleaseStreamObjectVersion::GUID) < FUE5ReleaseStreamObjectVersion::OrthographicAutoNearFarPlane)
 	{
 		bAutoCalculateOrthoPlanes = false;
 	}
 
+	if (Ar.CustomVer(FFortniteMainBranchObjectVersion::GUID) < FFortniteMainBranchObjectVersion::SceneCaptureDefaultSettings)
+	{
+		// For backwards compatibility, we must preserve the old defaults that do not reflect code overrides in SetupViewFamilyForSceneCapture.
+		PostProcessSettings = FPostProcessSettings::GetDefault();
+	}
+
 	Super::Serialize(Ar);
 
 	if (Ar.IsLoading())
@@ -1285,6 +1298,14 @@ USceneCaptureComponentCube::USceneCaptureComponentCube(const FObjectInitializer&
 	bTickInEditor = true;
 	bCaptureRotation = false;
 
+	// Update default post-processing values to reflect default values (previously and still) applied in code.
+	PostProcessSettings.DynamicGlobalIlluminationMethod = EDynamicGlobalIlluminationMethod::None;
+	PostProcessSettings.ReflectionMethod = EReflectionMethod::None;
+	PostProcessSettings.LumenSurfaceCacheResolution = 0.5f;
+	PostProcessSettings.VignetteIntensity = 0.0f;
+	PostProcessSettings.LumenReflectionsScreenTraces = 0;
+	PostProcessSettings.LumenFinalGatherScreenTraces = 0;
+
 #if WITH_EDITORONLY_DATA
 	if (!IsRunningCommandlet() && !IsTemplate())
 	{
@@ -1430,6 +1451,19 @@ void USceneCaptureComponentCube::UpdateSceneCaptureContents(FSceneInterface* Sce
 	Scene->UpdateSceneCaptureContents(this, SceneRenderBuilder);
 }
 
+void USceneCaptureComponentCube::Serialize(FArchive& Ar)
+{
+	Ar.UsingCustomVersion(FFortniteMainBranchObjectVersion::GUID);
+
+	if (Ar.CustomVer(FFortniteMainBranchObjectVersion::GUID) < FFortniteMainBranchObjectVersion::SceneCaptureDefaultSettings)
+	{
+		// For backwards compatibility, we must preserve the old defaults that do not reflect code overrides in SetupViewFamilyForSceneCapture.
+		PostProcessSettings = FPostProcessSettings::GetDefault();
+	}
+
+	Super::Serialize(Ar);
+}
+
 #if WITH_EDITOR
 void USceneCaptureComponentCube::PostEditChangeProperty(FPropertyChangedEvent& PropertyChangedEvent)
 {
diff --git a/Engine/Source/Runtime/Renderer/Private/PlanarReflectionRendering.cpp b/Engine/Source/Runtime/Renderer/Private/PlanarReflectionRendering.cpp
--- a/Engine/Source/Runtime/Renderer/Private/PlanarReflectionRendering.cpp
+++ b/Engine/Source/Runtime/Renderer/Private/PlanarReflectionRendering.cpp
@@ -428,6 +428,7 @@ extern TArray<FSceneView*> SetupViewFamilyForSceneCapture(
 	bool bCaptureSceneColor,
 	bool bIsPlanarReflection,
 	FPostProcessSettings* PostProcessSettings,
+	const FFinalPostProcessSettings* InheritedMainViewPostProcessSettings,
 	float PostProcessBlendWeight,
 	const AActor* ViewActor,
 	int32 CubemapFaceIndex);
@@ -547,7 +548,9 @@ void FScene::UpdatePlanarReflectionContents(UPlanarReflectionComponent* CaptureC
 		SceneCaptureViewInfo, CaptureComponent->MaxViewDistanceOverride,
 		/* bCaptureSceneColor = */ true,
 		/* bIsPlanarReflection = */ true,
-		&PostProcessSettings, 1.0f,
+		&PostProcessSettings,
+		/* InheritedMainViewPostProcessSettings = */ nullptr,
+		1.0f,
 		/*ViewActor =*/ nullptr,
 		/*CubemapFaceIndex =*/ INDEX_NONE);
 
diff --git a/Engine/Source/Runtime/Renderer/Private/SceneCaptureRendering.cpp b/Engine/Source/Runtime/Renderer/Private/SceneCaptureRendering.cpp
--- a/Engine/Source/Runtime/Renderer/Private/SceneCaptureRendering.cpp
+++ b/Engine/Source/Runtime/Renderer/Private/SceneCaptureRendering.cpp
@@ -50,6 +50,7 @@
 #include "SceneRenderBuilder.h"
 #include "Materials/MaterialParameterCollection.h"
 #include "Materials/MaterialParameterCollectionInstance.h"
+#include "Engine/PostProcessUtils.h"
 
 bool GSceneCaptureAllowRenderInMainRenderer = true;
 static FAutoConsoleVariableRef CVarSceneCaptureAllowRenderInMainRenderer(
@@ -668,6 +669,7 @@ TArray<FSceneView*> SetupViewFamilyForSceneCapture(
 	bool bCaptureSceneColor,
 	bool bIsPlanarReflection,
 	FPostProcessSettings* PostProcessSettings,
+	const FFinalPostProcessSettings* InheritedMainViewPostProcessSettings,
 	float PostProcessBlendWeight,
 	const AActor* ViewActor,
 	int32 CubemapFaceIndex)
@@ -739,22 +741,31 @@ TArray<FSceneView*> SetupViewFamilyForSceneCapture(
 
 		View->StartFinalPostprocessSettings(SceneCaptureViewInfo.ViewOrigin);
 
-		// By default, Lumen is disabled in scene captures, but can be re-enabled with the post process settings in the component.
-		View->FinalPostProcessSettings.DynamicGlobalIlluminationMethod = EDynamicGlobalIlluminationMethod::None;
-		View->FinalPostProcessSettings.ReflectionMethod = EReflectionMethod::None;
+		if (InheritedMainViewPostProcessSettings)
+		{
+			View->FinalPostProcessSettings = *InheritedMainViewPostProcessSettings;
+		}
+		else
+		{
+			// Note: Future update to defaults should be reflected in the component constructors with backward-compatible serialization logic.
 
-		// Default surface cache to lower resolution for Scene Capture.  Can be overridden via post process settings.
-		View->FinalPostProcessSettings.LumenSurfaceCacheResolution = 0.5f;
+			// By default, Lumen is disabled in scene captures, but can be re-enabled with the post process settings in the component.
+			View->FinalPostProcessSettings.DynamicGlobalIlluminationMethod = EDynamicGlobalIlluminationMethod::None;
+			View->FinalPostProcessSettings.ReflectionMethod = EReflectionMethod::None;
 
-		if (SceneCaptureComponent && SceneCaptureComponent->IsCube())
-		{
-			// Disable vignette by default for cube maps -- darkened borders don't make sense for an omnidirectional projection.
-			View->FinalPostProcessSettings.VignetteIntensity = 0.0f;
+			// Default surface cache to lower resolution for Scene Capture.  Can be overridden via post process settings.
+			View->FinalPostProcessSettings.LumenSurfaceCacheResolution = 0.5f;
+
+			if (SceneCaptureComponent && SceneCaptureComponent->IsCube())
+			{
+				// Disable vignette by default for cube maps -- darkened borders don't make sense for an omnidirectional projection.
+				View->FinalPostProcessSettings.VignetteIntensity = 0.0f;
 
-			// Disable screen traces by default for cube maps -- these don't blend well across face boundaries, creating major lighting seams.
-			// Lumen lighting still has some seams with these disabled, but it's an order of magnitude better.
-			View->FinalPostProcessSettings.LumenReflectionsScreenTraces = 0;
-			View->FinalPostProcessSettings.LumenFinalGatherScreenTraces = 0;
+				// Disable screen traces by default for cube maps -- these don't blend well across face boundaries, creating major lighting seams.
+				// Lumen lighting still has some seams with these disabled, but it's an order of magnitude better.
+				View->FinalPostProcessSettings.LumenReflectionsScreenTraces = 0;
+				View->FinalPostProcessSettings.LumenFinalGatherScreenTraces = 0;
+			}
 		}
 
 		if (PostProcessSettings)
@@ -819,20 +830,28 @@ static FSceneRenderer* CreateSceneRendererForSceneCapture(
 	SceneCaptureViewInfo.FOV = InFOV;
 
 	bool bInheritMainViewScreenPercentage = false;
+	const FFinalPostProcessSettings* InheritedMainViewPostProcessSettings = nullptr;
 	USceneCaptureComponent2D* SceneCaptureComponent2D = Cast<USceneCaptureComponent2D>(SceneCaptureComponent);
 
 	// Use camera position correction for ortho scene captures
 	if(IsValid(SceneCaptureComponent2D))
 	{
+		const FSceneViewFamily* MainViewFamily = SceneCaptureComponent2D->MainViewFamily;
+
 		if (!SceneCaptureViewInfo.IsPerspectiveProjection() && SceneCaptureComponent2D->bUpdateOrthoPlanes)
 		{
 			SceneCaptureViewInfo.UpdateOrthoPlanes(SceneCaptureComponent2D->bUseCameraHeightAsViewTarget);
 		}
 
-		if (SceneCaptureComponent2D->ShouldRenderWithMainViewResolution() && SceneCaptureComponent2D->MainViewFamily && !SceneCaptureComponent2D->ShouldIgnoreScreenPercentage())
+		if (SceneCaptureComponent2D->ShouldRenderWithMainViewResolution() && MainViewFamily && !SceneCaptureComponent2D->ShouldIgnoreScreenPercentage())
 		{
 			bInheritMainViewScreenPercentage = true;
 		}
+
+		if (SceneCaptureComponent2D->bInheritMainViewCameraPostProcessSettings && MainViewFamily)
+		{
+			InheritedMainViewPostProcessSettings = &MainViewFamily->Views[0]->FinalPostProcessSettings;
+		}
 	}
 
 	FSceneViewFamilyContext ViewFamily(FSceneViewFamily::ConstructionValues(
@@ -844,7 +863,7 @@ static FSceneRenderer* CreateSceneRendererForSceneCapture(
 
 	FSceneViewExtensionContext ViewExtensionContext(Scene);
 	ViewFamily.ViewExtensions = GEngine->ViewExtensions->GatherActiveExtensions(ViewExtensionContext);
-	
+
 	TArray<FSceneView*> Views = SetupViewFamilyForSceneCapture(
 		ViewFamily,
 		SceneCaptureComponent,
@@ -852,7 +871,8 @@ static FSceneRenderer* CreateSceneRendererForSceneCapture(
 		MaxViewDistance, 
 		bCaptureSceneColor,
 		/* bIsPlanarReflection = */ false,
-		PostProcessSettings, 
+		PostProcessSettings,
+		InheritedMainViewPostProcessSettings,
 		PostProcessBlendWeight,
 		ViewActor,
 		CubemapFaceIndex);
@@ -1790,6 +1810,7 @@ void FScene::UpdateSceneCaptureContents(USceneCaptureComponentCube* CaptureCompo
 				bCaptureSceneColor,
 				/* bIsPlanarReflection = */ false,
 				&CaptureComponent->PostProcessSettings,
+				nullptr,
 				CaptureComponent->PostProcessBlendWeight,
 				CaptureComponent->GetViewOwner(),
 				(int32)ECubeFace::CubeFace_MAX);			// Passing max cube face count indicates a view family with all faces
