Added sane defaults for three FVolumetricRenderTargetViewStateData class members when calling FVolumetricRenderTargetViewStateData::Reset. This is to prevent some rare crashes where these values contain uninitialized data that can lead to crashes.

diff --git a/Engine/Source/Runtime/Renderer/Private/VolumetricRenderTarget.cpp b/Engine/Source/Runtime/Renderer/Private/VolumetricRenderTarget.cpp
--- a/Engine/Source/Runtime/Renderer/Private/VolumetricRenderTarget.cpp
+++ b/Engine/Source/Runtime/Renderer/Private/VolumetricRenderTarget.cpp
@@ -322,12 +322,15 @@ void FVolumetricRenderTargetViewStateData::Reset()
 	bFirstTimeUsed = false;
 	bHistoryValid = false;
 	bHoldoutValid = false;
+	bValid = false;
 	PreViewExposure = 1.0f;
 	FrameId = 0;
 	NoiseFrameIndex = 0;
 	NoiseFrameIndexModPattern = 0;
 	CurrentPixelOffset = FIntPoint::ZeroValue;
 	CurrentRT = 0;
+	Mode = 0;
+	UpsamplingMode = 0;
 
 	// Release GPU resources
 	VolumetricTracingRT.SafeRelease();
