Disable normal output for GBuffer decals on mobile deferred.
Mobile deferred uses octahedral encoding for normals which isn't blending.

diff --git a/Engine/Source/Runtime/Renderer/Private/DecalRenderingCommon.cpp b/Engine/Source/Runtime/Renderer/Private/DecalRenderingCommon.cpp
--- a/Engine/Source/Runtime/Renderer/Private/DecalRenderingCommon.cpp
+++ b/Engine/Source/Runtime/Renderer/Private/DecalRenderingCommon.cpp
@@ -45,9 +45,14 @@ namespace DecalRendering
 		// Enforce platform output limitations.
 		if (bIsMobilePlatform && !bIsMobileDeferredPlatform && !bIsDBufferPlatform)
 		{
-			Desc.bWriteNormal = false;
 			Desc.bWriteRoughnessSpecularMetallic = false;
 		}
+		if (bIsMobilePlatform && !bIsDBufferPlatform)
+		{
+			// Note that we disable writing normals for mobile AND mobile deferred.
+			// Mobile deferred has octahedral encoding for normals and so doesn't support decal normal blending.
+			Desc.bWriteNormal = false;
+		}
 		if (bIsMobilePlatform)
 		{
 			Desc.bWriteAmbientOcclusion = false;
@@ -60,7 +65,7 @@ namespace DecalRendering
 		}
 
 		// Calculate main decal render stage. We set only one (or none) of these for any decal.
-		if (bIsMobileDeferredPlatform && (Desc.bWriteEmissive || Desc.bWriteBaseColor || Desc.bWriteNormal || Desc.bWriteRoughnessSpecularMetallic))
+		if ((bIsMobilePlatform && !bIsDBufferPlatform) && (Desc.bWriteEmissive || Desc.bWriteBaseColor || Desc.bWriteNormal || Desc.bWriteRoughnessSpecularMetallic))
 		{
 			Desc.RenderStageMask |= 1 << (uint32)EDecalRenderStage::MobileBeforeLighting;
 		}
