Moved scene extension BeginRenderViewFamily call to happen prior the FSceneRenderer constructor to fix third party upscalers.

diff --git a/Engine/Source/Runtime/Engine/Public/SceneRenderBuilderInterface.h b/Engine/Source/Runtime/Engine/Public/SceneRenderBuilderInterface.h
--- a/Engine/Source/Runtime/Engine/Public/SceneRenderBuilderInterface.h
+++ b/Engine/Source/Runtime/Engine/Public/SceneRenderBuilderInterface.h
@@ -58,11 +58,11 @@ public:
 	static ENGINE_API TUniquePtr<ISceneRenderBuilder> Create(FSceneInterface* SceneInterface);
 
 	/** Creates a scene renderer based on the current feature level. */
-	virtual FSceneRenderer* CreateSceneRenderer(const FSceneViewFamily* ViewFamily) = 0;
+	virtual FSceneRenderer* CreateSceneRenderer(FSceneViewFamily* ViewFamily) = 0;
 
 	/** Creates a linked set of scene renderers with related view families based on the current feature level. */
 	virtual TArray<FSceneRenderer*, FConcurrentLinearArrayAllocator> CreateLinkedSceneRenderers(
-		TConstArrayView<const FSceneViewFamily*> ViewFamilies,
+		TConstArrayView<FSceneViewFamily*> ViewFamilies,
 		FHitProxyConsumer* HitProxyConsumer) = 0;
 
 	/** Adds a command that executes on the game thread and is pipelined with scene render commands on the render thread.
diff --git a/Engine/Source/Runtime/Renderer/Private/SceneRenderBuilder.cpp b/Engine/Source/Runtime/Renderer/Private/SceneRenderBuilder.cpp
--- a/Engine/Source/Runtime/Renderer/Private/SceneRenderBuilder.cpp
+++ b/Engine/Source/Runtime/Renderer/Private/SceneRenderBuilder.cpp
@@ -78,7 +78,7 @@ static bool IsSplitScreenDebugEnabled(TConstArrayView<const FSceneViewFamily*> V
 	return CVarSplitScreenDebugEnable.GetValueOnGameThread() > 0 && ViewFamilies.Num() == 1 && ViewFamilies[0]->bSplitScreenDebugAllowed && ViewFamilies[0]->Views.Num() == 1;
 }
 
-static void CreateSplitScreenDebugViewFamilies(const FSceneViewFamily& InFamily, TConstArrayView<const FSceneViewFamily*>& OutFamilies, TArray<const FSceneViewFamily*, TFixedAllocator<2>>& OutFamiliesStorage)
+static void CreateSplitScreenDebugViewFamilies(const FSceneViewFamily& InFamily, TConstArrayView<FSceneViewFamily*>& OutFamilies, TArray<FSceneViewFamily*, TFixedAllocator<2>>& OutFamiliesStorage)
 {
 	// We either generate a single family with 2 views, or two families with 1 view each for MGPU
 	int32 NumFamilies;
@@ -472,7 +472,7 @@ public:
 	{}
 
 	TArray<FSceneRenderer*, FConcurrentLinearArrayAllocator> CreateSceneRenderers(
-		TConstArrayView<const FSceneViewFamily*> ViewFamilies,
+		TConstArrayView<FSceneViewFamily*> ViewFamilies,
 		FHitProxyConsumer* HitProxyConsumer,
 		bool bAllowSplitScreenDebug) const
 	{
@@ -483,7 +483,7 @@ public:
 
 	#if !UE_BUILD_SHIPPING
 		bool bSplitScreenDebug = false;
-		TArray<const FSceneViewFamily*, TFixedAllocator<2>> DebugViewFamilyStorage;
+		TArray<FSceneViewFamily*, TFixedAllocator<2>> DebugViewFamilyStorage;
 		if (bAllowSplitScreenDebug && IsSplitScreenDebugEnabled(ViewFamilies))
 		{
 			CreateSplitScreenDebugViewFamilies(*ViewFamilies[0], ViewFamilies, DebugViewFamilyStorage);
@@ -500,11 +500,16 @@ public:
 		TArray<FSceneRenderer*, FConcurrentLinearArrayAllocator> OutRenderers;
 		OutRenderers.Reserve(ViewFamilies.Num());
 
-		for (const FSceneViewFamily* ViewFamily : ViewFamilies)
+		for (FSceneViewFamily* ViewFamily : ViewFamilies)
 		{
 			check(ViewFamily);
 			check(ViewFamily->Scene == SceneInterface);
 
+			for (auto& ViewExtension : ViewFamily->ViewExtensions)
+			{
+				ViewExtension->BeginRenderViewFamily(*ViewFamily);
+			}
+
 			if (ShadingPath == EShadingPath::Deferred)
 			{
 				OutRenderers.Add(new FDeferredShadingSceneRenderer(ViewFamily, HitProxyConsumer));
@@ -787,15 +792,6 @@ void FSceneRenderProcessor::Execute()
 	TOptional<UE::RenderCore::DumpGPU::FDumpScope> GpuDumpScope;
 	TOptional<RenderCaptureInterface::FScopedCapture> GpuCaptureScope;
 
-	// This has to be called first because it can actually recursively invoke scene renders.
-	for (FSceneRenderer* Renderer : ActiveRenderers)
-	{
-		for (const FSceneViewExtensionRef& Extension : Renderer->ViewFamily.ViewExtensions)
-		{
-			Extension->BeginRenderViewFamily(Renderer->ViewFamily);
-		}
-	}
-
 	for (FOp Op : Ops)
 	{
 		switch (Op.Type)
@@ -964,7 +960,7 @@ struct FSceneRenderBuilder::FPersistentState
 };
 
 TArray<FSceneRenderer*, FConcurrentLinearArrayAllocator> FSceneRenderBuilder::CreateSceneRenderers(
-	TConstArrayView<const FSceneViewFamily*> ViewFamilies,
+	TConstArrayView<FSceneViewFamily*> ViewFamilies,
 	FHitProxyConsumer* HitProxyConsumer,
 	bool bAllowSplitScreenDebug)
 {
@@ -978,14 +974,14 @@ TArray<FSceneRenderer*, FConcurrentLinearArrayAllocator> FSceneRenderBuilder::Cr
 	return Renderers;
 }
 
-FSceneRenderer* FSceneRenderBuilder::CreateSceneRenderer(const FSceneViewFamily* ViewFamily)
+FSceneRenderer* FSceneRenderBuilder::CreateSceneRenderer(FSceneViewFamily* ViewFamily)
 {
 	const bool bAllowSplitScreenDebug = false;
 	return CreateSceneRenderers(MakeArrayView({ ViewFamily }), nullptr, bAllowSplitScreenDebug)[0];
 }
 
 TArray<FSceneRenderer*, FConcurrentLinearArrayAllocator> FSceneRenderBuilder::CreateLinkedSceneRenderers(
-	TConstArrayView<const FSceneViewFamily*> ViewFamilies,
+	TConstArrayView<FSceneViewFamily*> ViewFamilies,
 	FHitProxyConsumer* HitProxyConsumer)
 {
 	const bool bAllowSplitScreenDebug = true;
diff --git a/Engine/Source/Runtime/Renderer/Private/SceneRenderBuilder.h b/Engine/Source/Runtime/Renderer/Private/SceneRenderBuilder.h
--- a/Engine/Source/Runtime/Renderer/Private/SceneRenderBuilder.h
+++ b/Engine/Source/Runtime/Renderer/Private/SceneRenderBuilder.h
@@ -37,9 +37,9 @@ public:
 	//////////////////////////////////////////////////////////////////////////
 	// ISceneRenderBuilder Overrides
 
-	RENDERER_API FSceneRenderer* CreateSceneRenderer(const FSceneViewFamily* ViewFamily) override;
+	RENDERER_API FSceneRenderer* CreateSceneRenderer(FSceneViewFamily* ViewFamily) override;
 	RENDERER_API TArray<FSceneRenderer*, FConcurrentLinearArrayAllocator> CreateLinkedSceneRenderers(
-		TConstArrayView<const FSceneViewFamily*> ViewFamilies,
+		TConstArrayView<FSceneViewFamily*> ViewFamilies,
 		FHitProxyConsumer* HitProxyConsumer) override;
 	RENDERER_API void AddCommand(TUniqueFunction<void()>&& Function) override;
 	RENDERER_API void AddRenderer(FSceneRenderer* Renderer, FString&& Name, FSceneRenderFunction&& Function) override;
@@ -66,7 +66,7 @@ private:
 	void EndGroup() override;
 
 	TArray<FSceneRenderer*, FConcurrentLinearArrayAllocator> CreateSceneRenderers(
-		TConstArrayView<const FSceneViewFamily*> ViewFamilies,
+		TConstArrayView<FSceneViewFamily*> ViewFamilies,
 		FHitProxyConsumer* HitProxyConsumer,
 		bool bAllowSplitScreenDebug);
 
