Sync the GPU skin mesh update task prior to processing hair pass that can invoke FSkeletalMeshDeformerHelpers::GetBoneBufferForReading.

diff --git a/Engine/Source/Runtime/Renderer/Private/RendererScene.cpp b/Engine/Source/Runtime/Renderer/Private/RendererScene.cpp
--- a/Engine/Source/Runtime/Renderer/Private/RendererScene.cpp
+++ b/Engine/Source/Runtime/Renderer/Private/RendererScene.cpp
@@ -5589,6 +5589,7 @@ void FScene::Update(FRDGBuilder& GraphBuilder, const FUpdateParameters& Paramete
 	if (SkeletalMeshUpdater)
 	{
 		GPUSkinTasks = SkeletalMeshUpdater->Submit(GraphBuilder);
+		GPUSkinUpdateTask = GPUSkinTasks.SkinCache;
 	}
 
 	if (GPUSkinCache && (GPUSkinTasks.SkinCache.IsValid() || GPUSkinCache->HasWork()))
diff --git a/Engine/Source/Runtime/Renderer/Private/ScenePrivate.h b/Engine/Source/Runtime/Renderer/Private/ScenePrivate.h
--- a/Engine/Source/Runtime/Renderer/Private/ScenePrivate.h
+++ b/Engine/Source/Runtime/Renderer/Private/ScenePrivate.h
@@ -3768,6 +3768,15 @@ public:
 		return GPUSkinCacheTask;
 	}
 
+	void WaitForGPUSkinUpdateTask()
+	{
+		if (GPUSkinUpdateTask.IsValid())
+		{
+			GPUSkinUpdateTask.Wait();
+			GPUSkinUpdateTask = {};
+		}
+	}
+
 	void WaitForCacheMeshDrawCommandsTask()
 	{
 		CSV_SCOPED_SET_WAIT_STAT(CacheMeshDrawCommands);
@@ -4001,6 +4010,7 @@ private:
 	TArray<FLevelCommand> LevelCommands;
 
 	UE::Tasks::FTask CreateLightPrimitiveInteractionsTask;
+	UE::Tasks::FTask GPUSkinUpdateTask;
 	UE::Tasks::FTask GPUSkinCacheTask;
 	UE::Tasks::FTask CacheMeshDrawCommandsTask;
 	UE::Tasks::FTask CacheNaniteMaterialBinsTask;
diff --git a/Engine/Source/Runtime/Renderer/Private/SceneVisibility.cpp b/Engine/Source/Runtime/Renderer/Private/SceneVisibility.cpp
--- a/Engine/Source/Runtime/Renderer/Private/SceneVisibility.cpp
+++ b/Engine/Source/Runtime/Renderer/Private/SceneVisibility.cpp
@@ -4745,6 +4745,8 @@ void FSceneRenderer::PreVisibilityFrameSetup(FRDGBuilder& GraphBuilder)
 			{
 				if (Scene && IsHairStrandsEnabled(EHairStrandsShaderType::All, Scene->GetShaderPlatform()))
 				{
+					Scene->WaitForGPUSkinUpdateTask();
+
 					// 1.Update binding surfaces
 					// Prepare surface data (MeshLODData)
 					RunHairStrandsBookmark(GraphBuilder, EHairStrandsBookmark::ProcessBindingSurfaceUpdate, Parameters);
