Volumetric Cloud - added support for ShadowPassSwitch node.

#tests PC editor dx12
[RN] Volumetric Clouds - added support for the ShadowPassSwitch when tracing beers shadow map to case shadow onto world and atmaosphere (not self shadow).

diff --git a/Engine/Shaders/Private/VolumetricCloud.usf b/Engine/Shaders/Private/VolumetricCloud.usf
--- a/Engine/Shaders/Private/VolumetricCloud.usf
+++ b/Engine/Shaders/Private/VolumetricCloud.usf
@@ -71,6 +71,10 @@
 #define PERMUTATION_SUPPORT_LOCAL_FOG_VOLUME 0
 #endif
 
+#ifndef SHADOW_DEPTH_SHADER
+#define SHADOW_DEPTH_SHADER 0
+#endif
+
 // Draw debug sample positions
 #if CLOUD_DEBUG_SAMPLES || EMPTY_SPACE_SKIPPING_DEBUG
 #include "ShaderPrint.ush"
diff --git a/Engine/Source/Runtime/Renderer/Private/VolumetricCloudRendering.cpp b/Engine/Source/Runtime/Renderer/Private/VolumetricCloudRendering.cpp
--- a/Engine/Source/Runtime/Renderer/Private/VolumetricCloudRendering.cpp
+++ b/Engine/Source/Runtime/Renderer/Private/VolumetricCloudRendering.cpp
@@ -1347,6 +1347,9 @@ public:
 
 		// Force texture fetches to not use automatic mip generation because the pixel shader is using a dynamic loop to evaluate the material multiple times.
 		OutEnvironment.SetDefine(TEXT("USE_FORCE_TEXTURE_MIP"), TEXT("1"));
+
+		// Make the Shadow Pass Switch node work for volumetric cloud shadow tracing.
+		OutEnvironment.SetDefine(TEXT("SHADOW_DEPTH_SHADER"), TEXT("1"));
 	}
 
 private:
