Fixed: Missing SingleLayerWater when moible deferred shading is enabled

diff --git a/Engine/Source/Runtime/Renderer/Private/LightMapRendering.cpp b/Engine/Source/Runtime/Renderer/Private/LightMapRendering.cpp
--- a/Engine/Source/Runtime/Renderer/Private/LightMapRendering.cpp
+++ b/Engine/Source/Runtime/Renderer/Private/LightMapRendering.cpp
@@ -159,7 +159,8 @@ void FMobileDirectionalLightAndCSMPolicy::ModifyCompilationEnvironment(const FMa
 
 bool FMobileDirectionalLightAndCSMPolicy::ShouldCompilePermutation(const FMeshMaterialShaderPermutationParameters& Parameters)
 {
-	if (IsMobileDeferredShadingEnabled(Parameters.Platform))
+	if (IsMobileDeferredShadingEnabled(Parameters.Platform)
+		&& !Parameters.MaterialParameters.ShadingModels.HasShadingModel(MSM_SingleLayerWater))
 	{
 		return false;
 	}
@@ -235,7 +236,8 @@ void FMobileDirectionalLightAndSHIndirectPolicy::ModifyCompilationEnvironment(co
 
 bool FMobileDirectionalLightCSMAndSHIndirectPolicy::ShouldCompilePermutation(const FMeshMaterialShaderPermutationParameters& Parameters)
 {
-	if (IsMobileDeferredShadingEnabled(Parameters.Platform))
+	if (IsMobileDeferredShadingEnabled(Parameters.Platform) 
+		&& !Parameters.MaterialParameters.ShadingModels.HasShadingModel(MSM_SingleLayerWater))
 	{
 		return false;
 	}
