Improved RealTimeSkyCaptures mobile support so it now also work on GLES platforms

Original review (34.00): https://p4-swarm.epicgames.net/reviews/39708314

diff --git a/Engine/Source/Runtime/Renderer/Private/MobileReflectionEnvironmentCapture.h b/Engine/Source/Runtime/Renderer/Private/MobileReflectionEnvironmentCapture.h
--- a/Engine/Source/Runtime/Renderer/Private/MobileReflectionEnvironmentCapture.h
+++ b/Engine/Source/Runtime/Renderer/Private/MobileReflectionEnvironmentCapture.h
@@ -18,4 +18,5 @@ namespace MobileReflectionEnvironmentCapture
 	void ComputeAverageBrightness(FRDGBuilder& GraphBuilder, FGlobalShaderMap* ShaderMap, FRDGTexture* CubemapTexture, float* OutAverageBrightness);
 	FRDGTexture* FilterReflectionEnvironment(FRDGBuilder& GraphBuilder, FGlobalShaderMap* ShaderMap, FRDGTexture* CubemapTexture, FSHVectorRGB3* OutIrradianceEnvironmentMap);
 	void CopyToSkyTexture(FRDGBuilder& GraphBuilder, FScene* Scene, FRDGTexture* CubemapTexture, FTexture* ProcessedTexture);
+	void CreateCubeMips(FRDGBuilder& GraphBuilder, FGlobalShaderMap* ShaderMap, FRDGTexture* CubemapTexture);
 }
diff --git a/Engine/Source/Runtime/Renderer/Private/ReflectionEnvironmentCapture.cpp b/Engine/Source/Runtime/Renderer/Private/ReflectionEnvironmentCapture.cpp
--- a/Engine/Source/Runtime/Renderer/Private/ReflectionEnvironmentCapture.cpp
+++ b/Engine/Source/Runtime/Renderer/Private/ReflectionEnvironmentCapture.cpp
@@ -426,6 +426,39 @@ FRDGTexture* FilterCubeMap(FRDGBuilder& GraphBuilder, FGlobalShaderMap* ShaderMa
 	return FilteredCubemapTexture;
 }
 
+void ConvolveCubeMap(FRDGBuilder& GraphBuilder, FGlobalShaderMap* ShaderMap, uint32 CubeMipStart, uint32 CubeMipEnd, uint32 FaceStart, uint32 FaceCount, FRDGTexture* RDGSrcRenderTarget, FRDGTexture* RDGDstRenderTarget)
+{
+	TShaderMapRef<FCubeFilterPS> PixelShader(ShaderMap);
+
+	const int32 NumMips = RDGSrcRenderTarget->Desc.NumMips;
+
+	for (uint32 MipIndex = CubeMipStart; MipIndex <= CubeMipEnd; MipIndex++)
+	{
+		const int32 MipSize = 1 << (NumMips - MipIndex - 1);
+		const FIntRect ViewRect(0, 0, MipSize, MipSize);
+
+		for (uint32 CubeFace = FaceStart; CubeFace < FaceStart+FaceCount; CubeFace++)
+		{
+			auto* PassParameters = GraphBuilder.AllocParameters<FCubeFilterPS::FParameters>();
+			PassParameters->SourceCubemapTexture = GraphBuilder.CreateSRV(FRDGTextureSRVDesc::Create(RDGSrcRenderTarget));
+			PassParameters->SourceCubemapSampler = TStaticSamplerState<SF_Trilinear, AM_Clamp, AM_Clamp, AM_Clamp>::GetRHI();
+			PassParameters->CubeFace = CubeFace;
+			PassParameters->MipIndex = MipIndex;
+			PassParameters->NumMips = NumMips;
+			PassParameters->SvPositionToUVScale = FVector2f(1.0f / MipSize, 1.0f / MipSize);
+			PassParameters->RenderTargets[0] = FRenderTargetBinding(RDGDstRenderTarget, ERenderTargetLoadAction::ENoAction, MipIndex, CubeFace);
+
+			FPixelShaderUtils::AddFullscreenPass(
+				GraphBuilder,
+				ShaderMap,
+				RDG_EVENT_NAME("ConvolveCubeMap (Mip: %u, CubeFace: %u)", MipIndex, CubeFace),
+				PixelShader,
+				PassParameters,
+				ViewRect);
+		}
+	}
+}
+
 // Premultiply alpha in-place using alpha blending
 
 void PremultiplyCubeMipAlpha(FRDGBuilder& GraphBuilder, FGlobalShaderMap* ShaderMap, FRDGTexture* CubemapTexture, int32 MipIndex)
diff --git a/Engine/Source/Runtime/Renderer/Private/ReflectionEnvironmentCapture.h b/Engine/Source/Runtime/Renderer/Private/ReflectionEnvironmentCapture.h
--- a/Engine/Source/Runtime/Renderer/Private/ReflectionEnvironmentCapture.h
+++ b/Engine/Source/Runtime/Renderer/Private/ReflectionEnvironmentCapture.h
@@ -20,4 +20,6 @@ FMatrix GetCubeProjectionMatrix(float HalfFovDeg, float CubeMapSize, float NearP
 inline uint32 GetNumMips(uint32 MipSize)
 {
 	return FMath::CeilLogTwo(MipSize) + 1;
-}
\ No newline at end of file
+}
+
+void ConvolveCubeMap(FRDGBuilder& GraphBuilder, FGlobalShaderMap* ShaderMap, uint32 CubeMipStart, uint32 CubeMipEnd, uint32 FaceStart, uint32 FaceCount, FRDGTexture* RDGSrcRenderTarget, FRDGTexture* RDGDstRenderTarget);
diff --git a/Engine/Source/Runtime/Renderer/Private/ReflectionEnvironmentRealTimeCapture.cpp b/Engine/Source/Runtime/Renderer/Private/ReflectionEnvironmentRealTimeCapture.cpp
--- a/Engine/Source/Runtime/Renderer/Private/ReflectionEnvironmentRealTimeCapture.cpp
+++ b/Engine/Source/Runtime/Renderer/Private/ReflectionEnvironmentRealTimeCapture.cpp
@@ -22,6 +22,7 @@
 #include "ScreenPass.h"
 #include "SkyAtmosphereRendering.h"
 #include "MobileBasePassRendering.h"
+#include "MobileReflectionEnvironmentCapture.h"
 
 #if WITH_EDITOR
 #include "CanvasTypes.h"
@@ -910,49 +911,60 @@ void FScene::AllocateAndCaptureFrameSkyEnvMap(
 
 	auto RenderCubeFaces_GenCubeMips = [&](uint32 CubeMipStart, uint32 CubeMipEnd, TRefCountPtr<IPooledRenderTarget>& SkyRenderTarget)
 	{
-		check(CubeMipStart > 0);	// Never write to mip0 as it has just been redered into
+		// on mobile platforms use the pixel shader implementation as compute shaders might not be optimal
+		const bool bIsMobilePlatform = IsMobilePlatform(MainView.GetShaderPlatform());
+		if (bIsMobilePlatform)
+		{
+			FRDGTextureRef SkyCubeTexture = GraphBuilder.RegisterExternalTexture(SkyRenderTarget, TEXT("SkyRenderTarget"));
+			MobileReflectionEnvironmentCapture::CreateCubeMips(GraphBuilder, GetGlobalShaderMap(FeatureLevel), SkyCubeTexture);
+		}
+		else
+		{
+			check(CubeMipStart > 0);	// Never write to mip0 as it has just been redered into
 
-		FRDGTextureRef SkyCubeTexture = GraphBuilder.RegisterExternalTexture(SkyRenderTarget, TEXT("SkyRenderTarget"));
+			FRDGTextureRef SkyCubeTexture = GraphBuilder.RegisterExternalTexture(SkyRenderTarget, TEXT("SkyRenderTarget"));
 
-		FDownsampleCubeFaceCS::FPermutationDomain PermutationVector;
-		TShaderMapRef<FDownsampleCubeFaceCS> ComputeShader(GetGlobalShaderMap(FeatureLevel), PermutationVector);
+			FDownsampleCubeFaceCS::FPermutationDomain PermutationVector;
+			TShaderMapRef<FDownsampleCubeFaceCS> ComputeShader(GetGlobalShaderMap(FeatureLevel), PermutationVector);
 
-		for (uint32 MipIndex = CubeMipStart; MipIndex <= CubeMipEnd; MipIndex++)
-		{
-			const uint32 MipResolution = 1 << (CubeMipCount - MipIndex - 1);
-			FRDGTextureSRVRef SkyCubeTextureSRV = GraphBuilder.CreateSRV(FRDGTextureSRVDesc::CreateForMipLevel(SkyCubeTexture, MipIndex - 1)); // slice/face selection is useless so remove from CreateForMipLevel
-
-			FDownsampleCubeFaceCS::FParameters* PassParameters = GraphBuilder.AllocParameters<FDownsampleCubeFaceCS::FParameters>();
-			PassParameters->MipIndex = MipIndex;
-			PassParameters->NumMips = CubeMipCount;
-			PassParameters->CubeFace = 0; // unused
-			PassParameters->ValidDispatchCoord = FIntPoint(MipResolution, MipResolution);
-			PassParameters->SourceCubemapSampler = TStaticSamplerState<SF_Point>::GetRHI();
-
-			PassParameters->SourceCubemapTexture = SkyCubeTextureSRV;
-			FRDGTextureUAVDesc OutTextureMipColorDesc(SkyCubeTexture, MipIndex);
-			OutTextureMipColorDesc.DimensionOverride = ETextureDimension::Texture2DArray;
-			PassParameters->OutTextureMipColor = GraphBuilder.CreateUAV(OutTextureMipColorDesc);
-
-			FIntVector NumGroups = FIntVector::DivideAndRoundUp(FIntVector(MipResolution, MipResolution, 1), FIntVector(FDownsampleCubeFaceCS::ThreadGroupSize, FDownsampleCubeFaceCS::ThreadGroupSize, 1));
-
-			// The groupd size per face with padding
-			PassParameters->FaceThreadGroupSize = NumGroups.X * FDownsampleCubeFaceCS::ThreadGroupSize;
-
-			// We are going to dispatch once for all faces 
-			NumGroups.X *= 6;
-
-			// Dispatch with GenerateMips: reading from a slice through SRV and writing into lower mip through UAV.
-			ClearUnusedGraphResources(ComputeShader, PassParameters);
-			GraphBuilder.AddPass(
-				Forward<FRDGEventName>(RDG_EVENT_NAME("MipGen")),
-				PassParameters,
-				ERDGPassFlags::Compute,
-			[PassParameters, ComputeShader, NumGroups](FRDGAsyncTask, FRHICommandList& RHICmdList)
+			for (uint32 MipIndex = CubeMipStart; MipIndex <= CubeMipEnd; MipIndex++)
 			{
-				FComputeShaderUtils::Dispatch(RHICmdList, ComputeShader, *PassParameters, NumGroups);
-			});
+				const uint32 MipResolution = 1 << (CubeMipCount - MipIndex - 1);
+				FRDGTextureSRVRef SkyCubeTextureSRV = GraphBuilder.CreateSRV(FRDGTextureSRVDesc::CreateForMipLevel(SkyCubeTexture, MipIndex - 1)); // slice/face selection is useless so remove from CreateForMipLevel
+
+				FDownsampleCubeFaceCS::FParameters* PassParameters = GraphBuilder.AllocParameters<FDownsampleCubeFaceCS::FParameters>();
+				PassParameters->MipIndex = MipIndex;
+				PassParameters->NumMips = CubeMipCount;
+				PassParameters->CubeFace = 0; // unused
+				PassParameters->ValidDispatchCoord = FIntPoint(MipResolution, MipResolution);
+				PassParameters->SourceCubemapSampler = TStaticSamplerState<SF_Point>::GetRHI();
+
+				PassParameters->SourceCubemapTexture = SkyCubeTextureSRV;
+				FRDGTextureUAVDesc OutTextureMipColorDesc(SkyCubeTexture, MipIndex);
+				OutTextureMipColorDesc.DimensionOverride = ETextureDimension::Texture2DArray;
+				PassParameters->OutTextureMipColor = GraphBuilder.CreateUAV(OutTextureMipColorDesc);
+
+				FIntVector NumGroups = FIntVector::DivideAndRoundUp(FIntVector(MipResolution, MipResolution, 1), FIntVector(FDownsampleCubeFaceCS::ThreadGroupSize, FDownsampleCubeFaceCS::ThreadGroupSize, 1));
+
+				// The groupd size per face with padding
+				PassParameters->FaceThreadGroupSize = NumGroups.X * FDownsampleCubeFaceCS::ThreadGroupSize;
+
+				// We are going to dispatch once for all faces 
+				NumGroups.X *= 6;
+
+				// Dispatch with GenerateMips: reading from a slice through SRV and writing into lower mip through UAV.
+				ClearUnusedGraphResources(ComputeShader, PassParameters);
+				GraphBuilder.AddPass(
+					Forward<FRDGEventName>(RDG_EVENT_NAME("MipGen")),
+					PassParameters,
+					ERDGPassFlags::Compute,
+				[PassParameters, ComputeShader, NumGroups](FRDGAsyncTask, FRHICommandList& RHICmdList)
+				{
+					FComputeShaderUtils::Dispatch(RHICmdList, ComputeShader, *PassParameters, NumGroups);
+				});
+			}
 		}
+		
 	};
 
 	auto RenderCubeFaces_SpecularConvolution = [&](uint32 CubeMipStart, uint32 CubeMipEnd, uint32 FaceStart, uint32 FaceCount, TRefCountPtr<IPooledRenderTarget>& DstRenderTarget, TRefCountPtr<IPooledRenderTarget>& SrcRenderTarget)
@@ -961,36 +973,45 @@ void FScene::AllocateAndCaptureFrameSkyEnvMap(
 		FRDGTextureRef RDGSrcRenderTarget = GraphBuilder.RegisterExternalTexture(SrcRenderTarget);
 		FRDGTextureRef RDGDstRenderTarget = GraphBuilder.RegisterExternalTexture(DstRenderTarget);
 
-		FRDGTextureSRVRef RDGSrcRenderTargetSRV = GraphBuilder.CreateSRV(FRDGTextureSRVDesc::Create(RDGSrcRenderTarget));
-
-		FDownsampleCubeFaceCS::FPermutationDomain PermutationVector;
-		TShaderMapRef<FConvolveSpecularFaceCS> ComputeShader(GetGlobalShaderMap(FeatureLevel), PermutationVector);
-		for (uint32 MipIndex = CubeMipStart; MipIndex <= CubeMipEnd; MipIndex++)
+		// on mobile platforms use the pixel shader convolution as compute shaders might not be optimal
+		const bool bIsMobilePlatform = IsMobilePlatform(MainView.GetShaderPlatform());
+		if (bIsMobilePlatform)
 		{
-			const uint32 MipResolution = 1 << (CubeMipCount - MipIndex - 1);
-
-			FConvolveSpecularFaceCS::FParameters* PassParameters = GraphBuilder.AllocParameters<FConvolveSpecularFaceCS::FParameters>();
-			PassParameters->MipIndex = MipIndex;
-			PassParameters->NumMips = CubeMipCount;
-			PassParameters->CubeFace = 0; // unused
-			PassParameters->CubeFaceOffset = int(FaceStart);
-			PassParameters->ValidDispatchCoord = FIntPoint(MipResolution, MipResolution);
-			PassParameters->SourceCubemapSampler = TStaticSamplerState<SF_Point>::GetRHI();
-
-			PassParameters->SourceCubemapTexture = RDGSrcRenderTargetSRV;
-			FRDGTextureUAVDesc OutTextureMipColorDesc(RDGDstRenderTarget, MipIndex);
-			OutTextureMipColorDesc.DimensionOverride = ETextureDimension::Texture2DArray;
-			PassParameters->OutTextureMipColor = GraphBuilder.CreateUAV(OutTextureMipColorDesc);
-
-			FIntVector NumGroups = FIntVector::DivideAndRoundUp(FIntVector(MipResolution, MipResolution, 1), FIntVector(FConvolveSpecularFaceCS::ThreadGroupSize, FConvolveSpecularFaceCS::ThreadGroupSize, 1));
-
-			// The groupd size per face with padding
-			PassParameters->FaceThreadGroupSize = NumGroups.X * FConvolveSpecularFaceCS::ThreadGroupSize;
-
-			// We are going to dispatch once for all faces 
-			NumGroups.X *= FaceCount;
+			ConvolveCubeMap(GraphBuilder, GetGlobalShaderMap(FeatureLevel), CubeMipStart, CubeMipEnd, FaceStart, FaceCount, RDGSrcRenderTarget, RDGDstRenderTarget);
+		}
+		else
+		{
+			FRDGTextureSRVRef RDGSrcRenderTargetSRV = GraphBuilder.CreateSRV(FRDGTextureSRVDesc::Create(RDGSrcRenderTarget));
 
-			FComputeShaderUtils::AddPass(GraphBuilder, RDG_EVENT_NAME("Convolve"), ComputeShader, PassParameters, NumGroups);
+			FDownsampleCubeFaceCS::FPermutationDomain PermutationVector;
+			TShaderMapRef<FConvolveSpecularFaceCS> ComputeShader(GetGlobalShaderMap(FeatureLevel), PermutationVector);
+			for (uint32 MipIndex = CubeMipStart; MipIndex <= CubeMipEnd; MipIndex++)
+			{
+				const uint32 MipResolution = 1 << (CubeMipCount - MipIndex - 1);
+		
+				FConvolveSpecularFaceCS::FParameters* PassParameters = GraphBuilder.AllocParameters<FConvolveSpecularFaceCS::FParameters>();
+				PassParameters->MipIndex = MipIndex;
+				PassParameters->NumMips = CubeMipCount;
+				PassParameters->CubeFace = 0; // unused
+				PassParameters->CubeFaceOffset = int(FaceStart);
+				PassParameters->ValidDispatchCoord = FIntPoint(MipResolution, MipResolution);
+				PassParameters->SourceCubemapSampler = TStaticSamplerState<SF_Point>::GetRHI();
+		
+				PassParameters->SourceCubemapTexture = RDGSrcRenderTargetSRV;
+				FRDGTextureUAVDesc OutTextureMipColorDesc(RDGDstRenderTarget, MipIndex);
+				OutTextureMipColorDesc.DimensionOverride = ETextureDimension::Texture2DArray;
+				PassParameters->OutTextureMipColor = GraphBuilder.CreateUAV(OutTextureMipColorDesc);
+		
+				FIntVector NumGroups = FIntVector::DivideAndRoundUp(FIntVector(MipResolution, MipResolution, 1), FIntVector(FConvolveSpecularFaceCS::ThreadGroupSize, FConvolveSpecularFaceCS::ThreadGroupSize, 1));
+		
+				// The groupd size per face with padding
+				PassParameters->FaceThreadGroupSize = NumGroups.X * FConvolveSpecularFaceCS::ThreadGroupSize;
+		
+				// We are going to dispatch once for all faces 
+				NumGroups.X *= FaceCount;
+		
+				FComputeShaderUtils::AddPass(GraphBuilder, RDG_EVENT_NAME("Convolve"), ComputeShader, PassParameters, NumGroups);
+			}
 		}
 	};
 
