Set clipmap far culling plane to just beyond the visible range when force invalidate is enabled, greatly reducing rendered geometry in some cases. Can be disabled using r.Shadow.Virtual.Clipmap.CullDynamicTightly (default true).

diff --git a/Engine/Source/Runtime/Renderer/Private/VirtualShadowMaps/VirtualShadowMapClipmap.cpp b/Engine/Source/Runtime/Renderer/Private/VirtualShadowMaps/VirtualShadowMapClipmap.cpp
--- a/Engine/Source/Runtime/Renderer/Private/VirtualShadowMaps/VirtualShadowMapClipmap.cpp
+++ b/Engine/Source/Runtime/Renderer/Private/VirtualShadowMaps/VirtualShadowMapClipmap.cpp
@@ -134,14 +134,22 @@ TAutoConsoleVariable<int32> CVarMarkCoarsePagesDirectional(
 	ECVF_Scalability | ECVF_RenderThreadSafe
 );
 
+static TAutoConsoleVariable<bool> CVarVSMClipmapCullDynamicTightly(
+	TEXT("r.Shadow.Virtual.Clipmap.CullDynamicTightly"),
+	true,
+	TEXT("When enabled(default) the far culling plane for uncached clipmpap levels is set to the size of the clipmap level.\n")
+	TEXT("  Currently, this is only used when r.Shadow.Virtual.Cache.ForceInvalidateDirectional is enabled, as cached levels need a larger range."),
+	ECVF_RenderThreadSafe
+);
+
 // "Virtual" clipmap level to clipmap radius
 // NOTE: This is the radius of around the clipmap origin that this level must cover
 // The actual clipmap dimensions will be larger due to snapping and other accomodations
-static float GetLevelRadius(int32 Level)
+float FVirtualShadowMapClipmap::GetLevelRadius(float AbsoluteLevel)
 {
 	// NOTE: Virtual clipmap indices can be negative (although not commonly)
 	// Clipmap level rounds *down*, so radius needs to cover out to 2^(Level+1), where it flips
-	return FMath::Pow(2.0f, static_cast<float>(Level + 1));
+	return FMath::Pow(2.0f, AbsoluteLevel + 1.0f);
 }
 
 FVirtualShadowMapClipmapConfig FVirtualShadowMapClipmapConfig::GetGlobal()
@@ -293,7 +301,7 @@ FVirtualShadowMapClipmap::FVirtualShadowMapClipmap(
 	PerLightCacheEntry = VirtualShadowMapArrayCacheManager->FindCreateLightCacheEntry(LightSceneInfo.Id, UniqueViewKey, LevelCount, Config.ShadowTypeId);
 	PerLightCacheEntry->UpdateClipmap(LightDirection, FirstLevel, Config.bForceInvalidate);
 
-	const int RadiusLn = GetLevelRadius(LastLevel);
+	const int RadiusLn = GetLevelRadius(static_cast<float>(LastLevel));
 	const int RadiiPerLevel = 4;
 	FVector SnappedOriginLn = WorldToLightViewRotationMatrix.TransformPosition(WorldOrigin);
 	{
@@ -317,7 +325,7 @@ FVirtualShadowMapClipmap::FVirtualShadowMapClipmap(
 		FLevelData& Level = LevelData[Index];
 		const int32 AbsoluteLevel = Index + FirstLevel;		// Absolute (virtual) level index
 
-		const double RawLevelRadius = GetLevelRadius(AbsoluteLevel);
+		const double RawLevelRadius = GetLevelRadius(static_cast<float>(AbsoluteLevel));
 
 		double HalfLevelDim = 2.0 * RawLevelRadius;
 		double SnapSize = RawLevelRadius;
@@ -398,11 +406,22 @@ FVirtualShadowMapClipmap::FVirtualShadowMapClipmap(
 		ViewCenterDeltaZ = ViewCenter.Z - ClipmapLevelEntry->Clipmap.ViewCenterZ;
 		ViewRadiusZ = ClipmapLevelEntry->Clipmap.ViewRadiusZ;
 
-		// NOTE: These values are all in regular ranges after being offset
-		const double ZScale = 0.5 / ViewRadiusZ;
-		const double ZOffset = ViewRadiusZ + ViewCenterDeltaZ;
-		Level.ViewToClip = FReversedZOrthoMatrix(HalfLevelDim, HalfLevelDim, ZScale, ZOffset);
-
+		if (Config.bForceInvalidate && CVarVSMClipmapCullDynamicTightly.GetValueOnRenderThread())
+		{
+			// Far plane in the clipmap should be just the end of the visible range
+			float CullingEndDistance = GetLevelRadius(static_cast<float>(AbsoluteLevel) - ResolutionLodBias); 
+			const double ViewRangeZ = CullingEndDistance + ViewRadiusZ;
+			const double ZScale = 1.0 / ViewRangeZ;
+			const double ZOffset = ViewRadiusZ;
+			Level.ViewToClip = FReversedZOrthoMatrix(HalfLevelDim, HalfLevelDim, ZScale, ZOffset);
+		}
+		else
+		{
+			// NOTE: These values are all in regular ranges after being offset
+			const double ZScale = 0.5 / ViewRadiusZ;
+			const double ZOffset = ViewRadiusZ + ViewCenterDeltaZ;
+			Level.ViewToClip = FReversedZOrthoMatrix(HalfLevelDim, HalfLevelDim, ZScale, ZOffset);
+		}
 		// Update the entry with the new projection data
 		ClipmapLevelEntry->ProjectionData = ComputeProjectionShaderData(Index);
 	}
@@ -431,7 +450,7 @@ void FVirtualShadowMapClipmap::ComputeBoundingVolumes(const FVector ViewOrigin)
 
 float FVirtualShadowMapClipmap::GetMaxRadius() const
 {
-	return GetLevelRadius(GetClipmapLevel(GetLevelCount() - 1));
+	return GetLevelRadius(static_cast<float>(GetClipmapLevel(GetLevelCount() - 1)));
 }
 
 FViewMatrices FVirtualShadowMapClipmap::GetViewMatrices(int32 ClipmapIndex) const
diff --git a/Engine/Source/Runtime/Renderer/Private/VirtualShadowMaps/VirtualShadowMapClipmap.h b/Engine/Source/Runtime/Renderer/Private/VirtualShadowMaps/VirtualShadowMapClipmap.h
--- a/Engine/Source/Runtime/Renderer/Private/VirtualShadowMaps/VirtualShadowMapClipmap.h
+++ b/Engine/Source/Runtime/Renderer/Private/VirtualShadowMaps/VirtualShadowMapClipmap.h
@@ -91,6 +91,8 @@ public:
 		return WorldOrigin;
 	}
 
+	static float GetLevelRadius(float AbsoluteLevel);
+
 	// Returns the max radius the clipmap is guaranteed to cover (i.e. the radius of the last clipmap level)
 	// Note that this is not a conservative radius of the level projection, which is snapped
 	float GetMaxRadius() const;
