Fixed ensure on translucent postdof modulatre texture that is in fact not required.

#tests PC editor dx12 with with/without separate transluceny 100 or 50
#rnx

diff --git a/Engine/Source/Runtime/Renderer/Private/TranslucentRendering.cpp b/Engine/Source/Runtime/Renderer/Private/TranslucentRendering.cpp
--- a/Engine/Source/Runtime/Renderer/Private/TranslucentRendering.cpp
+++ b/Engine/Source/Runtime/Renderer/Private/TranslucentRendering.cpp
@@ -630,7 +630,9 @@ FScreenPassTexture FTranslucencyComposition::AddPass(
 	if (Operation == EOperation::UpscaleOnly)
 	{
 		check(!SceneColor.IsValid());
-		ensure(!TranslucencyTextures.ColorModulateTexture.IsValid());
+
+		//ensure(!TranslucencyTextures.ColorModulateTexture.IsValid());
+		// It is valid to have a valid TranslucencyTextures.ColorModulateTexture in "UpscaleTranslucency". Since it can will be upscaled later in the ComposeTranslucencyToNewSceneColor pass (apply onto scene color with added translucency color on top).
 
 		OpName = TEXT("UpscaleTranslucency");
 
