First Person: Fixed missing uniforms in FirstPersonSelfShadow shaders when Substrate is enabled.

diff --git a/Engine/Source/Runtime/Renderer/Private/Shadows/FirstPersonSelfShadow.cpp b/Engine/Source/Runtime/Renderer/Private/Shadows/FirstPersonSelfShadow.cpp
--- a/Engine/Source/Runtime/Renderer/Private/Shadows/FirstPersonSelfShadow.cpp
+++ b/Engine/Source/Runtime/Renderer/Private/Shadows/FirstPersonSelfShadow.cpp
@@ -76,6 +76,7 @@ class FFirstPersonSelfShadowInputsDownsamplePS : public FGlobalShader
 	BEGIN_SHADER_PARAMETER_STRUCT(FParameters, )
 		SHADER_PARAMETER_STRUCT_REF(FViewUniformShaderParameters, View)
 		SHADER_PARAMETER_RDG_UNIFORM_BUFFER(FSceneTextureUniformParameters, SceneTextures)
+		SHADER_PARAMETER_RDG_UNIFORM_BUFFER(FSubstrateGlobalUniformParameters, Substrate)
 		SHADER_PARAMETER(uint32, CheckerboardMode)
 		RENDER_TARGET_BINDING_SLOTS()
 	END_SHADER_PARAMETER_STRUCT()
@@ -105,6 +106,7 @@ class FFirstPersonSelfShadowTracingPS : public FGlobalShader
 	BEGIN_SHADER_PARAMETER_STRUCT(FParameters, )
 		SHADER_PARAMETER_STRUCT_REF(FViewUniformShaderParameters, View)
 		SHADER_PARAMETER_RDG_UNIFORM_BUFFER(FSceneTextureUniformParameters, SceneTextures)
+		SHADER_PARAMETER_RDG_UNIFORM_BUFFER(FSubstrateGlobalUniformParameters, Substrate)
 		SHADER_PARAMETER_RDG_UNIFORM_BUFFER(FDeferredLightUniformStruct, DeferredLight)
 		SHADER_PARAMETER_RDG_TEXTURE_SRV(Texture2D, DownsampledInputsTexture)
 		SHADER_PARAMETER_RDG_TEXTURE_SRV(Texture2D, DownsampledDepthTexture)
@@ -172,6 +174,7 @@ class FFirstPersonSelfShadowUpsamplePS : public FGlobalShader
 	BEGIN_SHADER_PARAMETER_STRUCT(FParameters, )
 		SHADER_PARAMETER_STRUCT_REF(FViewUniformShaderParameters, View)
 		SHADER_PARAMETER_RDG_UNIFORM_BUFFER(FSceneTextureUniformParameters, SceneTextures)
+		SHADER_PARAMETER_RDG_UNIFORM_BUFFER(FSubstrateGlobalUniformParameters, Substrate)
 		SHADER_PARAMETER_RDG_TEXTURE_SRV(Texture2D, DownsampledDepthTexture)
 		SHADER_PARAMETER_RDG_TEXTURE_SRV(Texture2D, ShadowFactorsTexture)
 		SHADER_PARAMETER(FVector2f, DownsampledInvBufferSize)
@@ -242,6 +245,7 @@ FFirstPersonSelfShadowInputs CreateFirstPersonSelfShadowInputs(FRDGBuilder& Grap
 		FFirstPersonSelfShadowInputsDownsamplePS::FParameters* PassParameters = GraphBuilder.AllocParameters<FFirstPersonSelfShadowInputsDownsamplePS::FParameters>();
 		PassParameters->View = View.ViewUniformBuffer;
 		PassParameters->SceneTextures = SceneTextures.UniformBuffer;
+		PassParameters->Substrate = Substrate::BindSubstrateGlobalUniformParameters(View);
 		PassParameters->CheckerboardMode = CheckerboardMode;
 		PassParameters->RenderTargets[0] = FRenderTargetBinding(DownsampledTextures.Normals, ERenderTargetLoadAction::EClear);
 		PassParameters->RenderTargets.DepthStencil = FDepthStencilBinding(DownsampledTextures.DepthStencil, ERenderTargetLoadAction::EClear, ERenderTargetLoadAction::EClear, FExclusiveDepthStencil::DepthWrite_StencilWrite);
@@ -316,6 +320,7 @@ void RenderFirstPersonSelfShadow(FRDGBuilder& GraphBuilder, const TArray<FViewIn
 			FFirstPersonSelfShadowTracingPS::FParameters* PassParameters = GraphBuilder.AllocParameters<FFirstPersonSelfShadowTracingPS::FParameters>();
 			PassParameters->View = View.ViewUniformBuffer;
 			PassParameters->SceneTextures = Inputs.SceneTextures->UniformBuffer;
+			PassParameters->Substrate = Substrate::BindSubstrateGlobalUniformParameters(View);
 			PassParameters->DeferredLight = DeferredLightUniformBuffer;
 			PassParameters->ClosestHZBTexture = View.ClosestHZB;
 			PassParameters->HZBBaseTexelSize = FVector2f(1.0f / View.ClosestHZB->Desc.Extent.X, 1.0f / View.ClosestHZB->Desc.Extent.Y);
@@ -418,6 +423,7 @@ void RenderFirstPersonSelfShadow(FRDGBuilder& GraphBuilder, const TArray<FViewIn
 			FFirstPersonSelfShadowUpsamplePS::FParameters* PassParameters = GraphBuilder.AllocParameters<FFirstPersonSelfShadowUpsamplePS::FParameters>();
 			PassParameters->View = View.ViewUniformBuffer;
 			PassParameters->SceneTextures = Inputs.SceneTextures->UniformBuffer;
+			PassParameters->Substrate = Substrate::BindSubstrateGlobalUniformParameters(View);
 			PassParameters->DownsampledDepthTexture = GraphBuilder.CreateSRV(Inputs.DownsampledInputs[ViewIndex].DepthStencil);
 			PassParameters->ShadowFactorsTexture = GraphBuilder.CreateSRV(FPDenoisedShadowsTexture);
 			PassParameters->DownsampledInvBufferSize = FVector2f(1.0f) / FVector2f(DownsampledResolution);
