Ray Tracing: Add debug visualization for flags and masks

Add new picker domains to visualize the unique flag and mask settings throughout the scene.

Tweak the fake lighting intensity slightly and fix the segment picker domain to color by segment.

diff --git a/Engine/Shaders/Private/RayTracing/RayTracingDebug.usf b/Engine/Shaders/Private/RayTracing/RayTracingDebug.usf
--- a/Engine/Shaders/Private/RayTracing/RayTracingDebug.usf
+++ b/Engine/Shaders/Private/RayTracing/RayTracingDebug.usf
@@ -44,6 +44,7 @@ uint TopKMostHitInstances;
 uint NumTotalInstances;
 
 StructuredBuffer<FRayTracingInstanceDebugData> InstancesDebugData;
+StructuredBuffer<FPlatformRayTracingInstanceDescriptor> InstanceBuffer;
 StructuredBuffer<FRayTracingPickingFeedback> PickingBuffer;
 
 #if USE_DEBUG_CHS == 0 // add functions that require shaded payload
@@ -228,7 +229,8 @@ float4 DebugRayTracingInstance(FRayDesc Ray, out float OutDepth)
 	if (Payload.IsHit())
 	{
 		const float3 InstanceColor = IntToColor(Payload.InstanceHash);
-		const float3 TriangleColor = IntToColor(MurmurMix(Payload.InstanceHash + Payload.TriangleIndex));
+		const float3 GeometryColor = IntToColor(Payload.InstanceHash + Payload.GeometryIndex);
+		const float3 TriangleColor = IntToColor(Payload.InstanceHash + Payload.TriangleIndex);
 
 		switch (VisualizationMode)
 		{
@@ -269,24 +271,34 @@ float4 DebugRayTracingInstance(FRayDesc Ray, out float OutDepth)
 			const bool bSameSegment = (PickingBuffer[0].GeometryIndex == Payload.GeometryIndex);
 			const bool bSameInstance = (PickingBuffer[0].InstanceIndex == Payload.InstanceIndex);
 
+			FPlatformRayTracingInstanceDescriptor InstanceDescriptor = InstanceBuffer[Payload.InstanceIndex];
+			uint Flags = TranslateRayTracingPlatformInstanceFlags(GetRayTracingInstanceDescriptorFlags(InstanceDescriptor));
+			uint Mask = GetRayTracingInstanceDescriptorMask(InstanceDescriptor);
+
 			switch (PickerDomain)
 			{
 			case RAY_TRACING_DEBUG_PICKER_DOMAIN_TRIANGLE:
 				Result.rgb = bSameInstance ? (bSameSegment && bSameTriangle ? PickedColor : TriangleColor * BrighterColorScale) : TriangleColor * ColorScale;
 				break;				
 			case RAY_TRACING_DEBUG_PICKER_DOMAIN_SEGMENT:
-				Result.rgb = bSameInstance ? (bSameSegment ? PickedColor : InstanceColor * BrighterColorScale) : InstanceColor * ColorScale; // TODO
+				Result.rgb = bSameInstance ? (bSameSegment ? PickedColor : GeometryColor * BrighterColorScale) : GeometryColor * ColorScale;
 				break;
 			case RAY_TRACING_DEBUG_PICKER_DOMAIN_INSTANCE:
 				Result.rgb = bSameInstance ? PickedColor : InstanceColor * ColorScale;
 				break;
+			case RAY_TRACING_DEBUG_PICKER_DOMAIN_FLAGS:
+				Result.rgb = (PickingBuffer[0].Flags == Flags ? BrighterColorScale : ColorScale) * IntToColor(Flags + 1);
+				break;
+			case RAY_TRACING_DEBUG_PICKER_DOMAIN_MASK:
+				Result.rgb = (PickingBuffer[0].Mask == Mask ? BrighterColorScale : ColorScale) * IntToColor(Mask);
+				break;
 			}
 		}
 
 		// apply some shading based on the primitive normal (map dot product result from [0,1] to [0.75,1] unless we are showing instance heatmap then just use [0,1]).
 		const float FakeLighting = dot(Payload.WorldNormal, normalize(float3(1.0f, 1.0f, 1.0f)));
-		const float FakeLightingScale = (VisualizationMode == RAY_TRACING_DEBUG_VIZ_INSTANCE_OVERLAP) ? 0.5f : 0.125f;
-		const float FakeLightingBias = (VisualizationMode == RAY_TRACING_DEBUG_VIZ_INSTANCE_OVERLAP) ? 0.5f : 0.875f;
+		const float FakeLightingScale = (VisualizationMode == RAY_TRACING_DEBUG_VIZ_INSTANCE_OVERLAP) ? 0.5f : 0.3f;
+		const float FakeLightingBias = (VisualizationMode == RAY_TRACING_DEBUG_VIZ_INSTANCE_OVERLAP) ? 0.5f : 0.7f;
 		const int NoFakeLighting = (VisualizationMode == RAY_TRACING_DEBUG_VIZ_TRIANGLE_HITCOUNT) || (VisualizationMode == RAY_TRACING_DEBUG_VIZ_HITCOUNT_PER_INSTANCE);
 
 		if (!NoFakeLighting)
diff --git a/Engine/Shaders/Shared/RaytracingDebugDefinitions.h b/Engine/Shaders/Shared/RaytracingDebugDefinitions.h
--- a/Engine/Shaders/Shared/RaytracingDebugDefinitions.h
+++ b/Engine/Shaders/Shared/RaytracingDebugDefinitions.h
@@ -55,4 +55,5 @@
 #define RAY_TRACING_DEBUG_PICKER_DOMAIN_TRIANGLE				0
 #define RAY_TRACING_DEBUG_PICKER_DOMAIN_INSTANCE				1
 #define RAY_TRACING_DEBUG_PICKER_DOMAIN_SEGMENT					2
-
+#define RAY_TRACING_DEBUG_PICKER_DOMAIN_FLAGS                   3
+#define RAY_TRACING_DEBUG_PICKER_DOMAIN_MASK                    4
diff --git a/Engine/Source/Runtime/Renderer/Private/RayTracing/RayTracing.cpp b/Engine/Source/Runtime/Renderer/Private/RayTracing/RayTracing.cpp
--- a/Engine/Source/Runtime/Renderer/Private/RayTracing/RayTracing.cpp
+++ b/Engine/Source/Runtime/Renderer/Private/RayTracing/RayTracing.cpp
@@ -171,6 +171,11 @@ namespace RayTracing
 			bTranslucentGeometry = true;
 		}
 
+		if (ViewFamily.EngineShowFlags.RayTracingDebug)
+		{
+			bTranslucentGeometry = true; // could check r.RayTracing.Visualize.OpaqueOnly, but not critical as this is only for debugging purposes
+		}
+
 		if (ViewFamily.EngineShowFlags.PathTracing
 			&& FDataDrivenShaderPlatformInfo::GetSupportsPathTracing(Scene.GetShaderPlatform()))
 		{
diff --git a/Engine/Source/Runtime/Renderer/Private/RayTracing/RayTracingDebug.cpp b/Engine/Source/Runtime/Renderer/Private/RayTracing/RayTracingDebug.cpp
--- a/Engine/Source/Runtime/Renderer/Private/RayTracing/RayTracingDebug.cpp
+++ b/Engine/Source/Runtime/Renderer/Private/RayTracing/RayTracingDebug.cpp
@@ -44,7 +44,10 @@ TAutoConsoleVariable<int32> CVarRayTracingVisualizePickerDomain(
 	0,
 	TEXT("Changes the picker domain to highlight:\n")
 	TEXT("0 - Triangles (default)\n")
-	TEXT("1 - Instances\n"),
+	TEXT("1 - Instances\n")
+	TEXT("2 - Segment\n")
+	TEXT("3 - Flags\n")
+	TEXT("4 - Mask\n"),
 	ECVF_RenderThreadSafe
 );
 
@@ -171,12 +174,12 @@ class FRayTracingDebugRGS : public FGlobalShader
 	BEGIN_SHADER_PARAMETER_STRUCT(FParameters, )
 		SHADER_PARAMETER(uint32, VisualizationMode)
 		SHADER_PARAMETER(uint32, PickerDomain)
-		SHADER_PARAMETER(int32, ShouldUsePreExposure)
+		SHADER_PARAMETER(uint32, ShouldUsePreExposure)
 		SHADER_PARAMETER(float, TimingScale)
 		SHADER_PARAMETER(float, MaxTraceDistance)
 		SHADER_PARAMETER(float, FarFieldMaxTraceDistance)
 		SHADER_PARAMETER(FVector3f, FarFieldReferencePos)
-		SHADER_PARAMETER(int32, OpaqueOnly)
+		SHADER_PARAMETER(uint32, OpaqueOnly)
 		SHADER_PARAMETER(float, TriangleHitCountMaxThreshold)
 		SHADER_PARAMETER(float, TriangleHitCountPerInstanceMaxThreshold)
 		SHADER_PARAMETER(uint32, TopKMostHitInstances)
@@ -185,6 +188,7 @@ class FRayTracingDebugRGS : public FGlobalShader
 		SHADER_PARAMETER_RDG_TEXTURE_UAV(RWTexture2D<float4>, Output)
 		SHADER_PARAMETER_RDG_TEXTURE_UAV(RWTexture2D<float>, OutputDepth)
 		SHADER_PARAMETER_RDG_BUFFER_SRV(StructuredBuffer, InstancesDebugData)
+		SHADER_PARAMETER_RDG_BUFFER_SRV(StructuredBuffer, InstanceBuffer)
 		SHADER_PARAMETER_RDG_BUFFER_SRV(StructuredBuffer<FRayTracingPickingFeedback>, PickingBuffer)
 		SHADER_PARAMETER_RDG_BUFFER_SRV(StructuredBuffer<uint2>, TopKHitStats)
 		SHADER_PARAMETER_STRUCT_REF(FViewUniformShaderParameters, ViewUniformBuffer)
@@ -1572,6 +1576,11 @@ void RenderRayTracingDebug(FRDGBuilder& GraphBuilder, const FScene& Scene, const
 	{
 		InstanceDebugBuffer = GSystemTextures.GetDefaultStructuredBuffer(GraphBuilder, sizeof(uint32));
 	}
+	FRDGBufferRef InstanceBuffer = RayTracingScene.GetInstanceBuffer(ERayTracingSceneLayer::Base);
+	if (InstanceBuffer == nullptr)
+	{
+		InstanceBuffer = GSystemTextures.GetDefaultStructuredBuffer(GraphBuilder, sizeof(uint32));
+	}
 
 	const bool bRequiresDebugCHS = RequiresRayTracingDebugCHS(DebugVisualizationMode);
 
@@ -1667,6 +1676,7 @@ void RenderRayTracingDebug(FRDGBuilder& GraphBuilder, const FScene& Scene, const
 	}
 	
 	RayGenParameters->InstancesDebugData = GraphBuilder.CreateSRV(InstanceDebugBuffer);
+	RayGenParameters->InstanceBuffer = GraphBuilder.CreateSRV(InstanceBuffer);
 	RayGenParameters->PickingBuffer = GraphBuilder.CreateSRV(PickingBuffer);
 	RayGenParameters->TLAS = RayTracingScene.GetLayerView(ERayTracingSceneLayer::Base);
 	RayGenParameters->ViewUniformBuffer = View.ViewUniformBuffer;
@@ -1735,13 +1745,21 @@ void FDeferredShadingSceneRenderer::RayTracingDisplayPicking(const FRayTracingPi
 		Writer.DrawLine(FText::FromString(TEXT("Domain [Instance]")), 10, FColor::Yellow);
 		break;
 
+	case RAY_TRACING_DEBUG_PICKER_DOMAIN_FLAGS:
+		Writer.DrawLine(FText::FromString(TEXT("Domain [Flags]")), 10, FColor::Yellow);
+		break;
+
+	case RAY_TRACING_DEBUG_PICKER_DOMAIN_MASK:
+		Writer.DrawLine(FText::FromString(TEXT("Domain [Mask]")), 10, FColor::Yellow);
+		break;
+
 	default:
 		break; // Invalid picking domain
 	}
 
 	Writer.EmptyLine();
 
-	Writer.DrawLine(FText::FromString(TEXT("(Use r.RayTracing.Debug.PickerDomain to change domain)")), 10, FColor::Yellow);
+	Writer.DrawLine(FText::FromString(TEXT("(Use r.RayTracing.Visualize.PickerDomain to change domain)")), 10, FColor::Yellow);
 
 	Writer.EmptyLine();
 
@@ -1777,6 +1795,15 @@ void FDeferredShadingSceneRenderer::RayTracingDisplayPicking(const FRayTracingPi
 		const FRayTracingGeometryInitializer& Initializer = Geometry->GetInitializer();
 		Writer.DrawLine(FText::FromString(FString::Printf(TEXT("Name: %s"), *Initializer.DebugName.ToString())), 10, FColor::Yellow);
 		Writer.DrawLine(FText::FromString(FString::Printf(TEXT("Num Segments: %d"), Initializer.Segments.Num())), 10, FColor::Yellow);
+		if (PickingFeedback.GeometryIndex < uint32(Initializer.Segments.Num()))
+		{
+			const FRayTracingGeometrySegment& Segment = Initializer.Segments[PickingFeedback.GeometryIndex];
+			Writer.DrawLine(FText::FromString(FString::Printf(TEXT("Segment %d Primitive Count: %u"), PickingFeedback.GeometryIndex, Segment.NumPrimitives)), 10, FColor::Yellow);
+		}
+		else
+		{
+			Writer.DrawLine(FText::FromString(FString::Printf(TEXT("Segment %d UNKNOWN"), PickingFeedback.GeometryIndex)), 10, FColor::Yellow);
+		}
 		Writer.DrawLine(FText::FromString(FString::Printf(TEXT("Total Primitive Count: %u"), Initializer.TotalPrimitiveCount)), 10, FColor::Yellow);
 		Writer.DrawLine(FText::FromString(FString::Printf(TEXT("Fast Build: %d"), Initializer.bFastBuild)), 10, FColor::Yellow);
 		Writer.DrawLine(FText::FromString(FString::Printf(TEXT("Allow Update: %d"), Initializer.bAllowUpdate)), 10, FColor::Yellow);
@@ -1802,6 +1829,39 @@ void FDeferredShadingSceneRenderer::RayTracingDisplayPicking(const FRayTracingPi
 
 	Writer.DrawLine(FText::FromString(FString::Printf(TEXT("InstanceId: %u"), PickingFeedback.InstanceId)), 10, FColor::Yellow);
 	Writer.DrawLine(FText::FromString(FString::Printf(TEXT("Mask: %u"), PickingFeedback.Mask)), 10, FColor::Yellow);
+	if (PickingFeedback.Mask & RAY_TRACING_MASK_OPAQUE)
+	{
+		Writer.DrawLine(FText::FromString(TEXT("   RAY_TRACING_MASK_OPAQUE")), 10, FColor::Yellow);
+	}
+	if (PickingFeedback.Mask & RAY_TRACING_MASK_TRANSLUCENT)
+	{
+		Writer.DrawLine(FText::FromString(TEXT("   RAY_TRACING_MASK_TRANSLUCENT")), 10, FColor::Yellow);
+	}
+	if (PickingFeedback.Mask & RAY_TRACING_MASK_OPAQUE_SHADOW)
+	{
+		Writer.DrawLine(FText::FromString(TEXT("   RAY_TRACING_MASK_OPAQUE_SHADOW")), 10, FColor::Yellow);
+	}
+	if (PickingFeedback.Mask & RAY_TRACING_MASK_TRANSLUCENT_SHADOW)
+	{
+		Writer.DrawLine(FText::FromString(TEXT("   RAY_TRACING_MASK_TRANSLUCENT_SHADOW")), 10, FColor::Yellow);
+	}
+	if (PickingFeedback.Mask & RAY_TRACING_MASK_THIN_SHADOW)
+	{
+		Writer.DrawLine(FText::FromString(TEXT("   RAY_TRACING_MASK_THIN_SHADOW")), 10, FColor::Yellow);
+	}
+	if (PickingFeedback.Mask & RAY_TRACING_MASK_FAR_FIELD)
+	{
+		Writer.DrawLine(FText::FromString(TEXT("   RAY_TRACING_MASK_FAR_FIELD")), 10, FColor::Yellow);
+	}
+	if (PickingFeedback.Mask & RAY_TRACING_MASK_HAIR_STRANDS)
+	{
+		Writer.DrawLine(FText::FromString(TEXT("   RAY_TRACING_MASK_HAIR_STRANDS")), 10, FColor::Yellow);
+	}
+	if (PickingFeedback.Mask & RAY_TRACING_MASK_OPAQUE_FP_WORLD_SPACE)
+	{
+		Writer.DrawLine(FText::FromString(TEXT("   RAY_TRACING_MASK_OPAQUE_FP_WORLD_SPACE")), 10, FColor::Yellow);
+	}
+
 	Writer.DrawLine(FText::FromString(FString::Printf(TEXT("ContributionToHitGroup: %u"), PickingFeedback.InstanceContributionToHitGroupIndex)), 10, FColor::Yellow);
 	{
 		const ERayTracingInstanceFlags Flags = (ERayTracingInstanceFlags)PickingFeedback.Flags;
