Fix shader compilation error on mobile platform due to typo in previous changelist.

diff --git a/Engine/Source/Runtime/Renderer/Private/DecalRenderingCommon.cpp b/Engine/Source/Runtime/Renderer/Private/DecalRenderingCommon.cpp
--- a/Engine/Source/Runtime/Renderer/Private/DecalRenderingCommon.cpp
+++ b/Engine/Source/Runtime/Renderer/Private/DecalRenderingCommon.cpp
@@ -65,7 +65,7 @@ namespace DecalRendering
 		}
 
 		// Calculate main decal render stage. We set only one (or none) of these for any decal.
-		if ((bIsMobilePlatform && !bIsDBufferPlatform) && (Desc.bWriteEmissive || Desc.bWriteBaseColor || Desc.bWriteNormal || Desc.bWriteRoughnessSpecularMetallic))
+		if ((bIsMobileDeferredPlatform && !bIsDBufferPlatform) && (Desc.bWriteEmissive || Desc.bWriteBaseColor || Desc.bWriteNormal || Desc.bWriteRoughnessSpecularMetallic))
 		{
 			Desc.RenderStageMask |= 1 << (uint32)EDecalRenderStage::MobileBeforeLighting;
 		}
