PR #12842: Fix typos

#12842

#rnx

diff --git a/Engine/Source/Runtime/Renderer/Private/SceneVisibilityPrivate.h b/Engine/Source/Runtime/Renderer/Private/SceneVisibilityPrivate.h
--- a/Engine/Source/Runtime/Renderer/Private/SceneVisibilityPrivate.h
+++ b/Engine/Source/Runtime/Renderer/Private/SceneVisibilityPrivate.h
@@ -477,7 +477,7 @@ private:
 
 		Frustum Cull          - Primitives are frustum / distance culled and visible primitives are emitted.
 		Occlusion Cull        - Primitives are culled against occluders in the scene and visible primitives are emitted.
-		Compute Relevance     - Primitives are queried for view relevance informationm dynamic primitives are identified and emitted, and static meshes are filtered
+		Compute Relevance     - Primitives are queried for view relevance information dynamic primitives are identified and emitted, and static meshes are filtered
 								into various mesh passes
 
 	The following stages are then performed for all views:
