RealTime sky capture - added optional clear step before sky/fog rendering in case we need to debug.

#tests PC editor dx12
#rnx

diff --git a/Engine/Source/Runtime/Renderer/Private/ReflectionEnvironmentRealTimeCapture.cpp b/Engine/Source/Runtime/Renderer/Private/ReflectionEnvironmentRealTimeCapture.cpp
--- a/Engine/Source/Runtime/Renderer/Private/ReflectionEnvironmentRealTimeCapture.cpp
+++ b/Engine/Source/Runtime/Renderer/Private/ReflectionEnvironmentRealTimeCapture.cpp
@@ -54,6 +54,11 @@ static TAutoConsoleVariable<int32> CVarRealTimeReflectionCaptureDepthBuffer(
 	TEXT("When enabled, the real-time sky light capture will have a depth buffer, this is for multiple meshes to be cover each other correctly. The height fog will also be applied according to the depth buffer."),
 	ECVF_RenderThreadSafe);
 
+static TAutoConsoleVariable<int32> CVarRealTimeReflectionCaptureAlwaysClearColorBuffer(
+	TEXT("r.SkyLight.RealTimeReflectionCapture.AlwaysClearColorBuffer"), 0,
+	TEXT("Always clear the color buffer to black before rendering SkyAtmosphere, SkyPass, Fog, etc."),
+	ECVF_RenderThreadSafe);
+
 static TAutoConsoleVariable<int32> CVarRealTimeReflectionCaptureResolutionOverride(
 	TEXT("r.SkyLight.RealTimeReflectionCapture.ResolutionOverride"), 0,
 	TEXT("Override the real time sky light capture resolution if this CVar is greater than 16."),
@@ -520,7 +525,7 @@ void FScene::AllocateAndCaptureFrameSkyEnvMap(
 				}
 			}
 
-
+			const bool bAlwaysClearColorBuffer = CVarRealTimeReflectionCaptureAlwaysClearColorBuffer.GetValueOnRenderThread() > 0;
 			for (int32 CubeFace = StartCubeFace; CubeFace < EndCubeFace; CubeFace++)
 			{
 				RDG_EVENT_SCOPE(GraphBuilder, "Capture Face=%d", CubeFace);
@@ -625,6 +630,11 @@ void FScene::AllocateAndCaptureFrameSkyEnvMap(
 				const bool bIsMobilePlatform = IsMobilePlatform(MainView.GetShaderPlatform());
 				if (bExecuteSky)
 				{
+					if (bAlwaysClearColorBuffer)
+					{
+						ClearCubeFace(SkyCubeTexture, CubeFace);
+					}
+
 					if(MainView.bSceneHasSkyMaterial || bShouldRenderSkyAtmosphere)
 					{
 						// If there are any mesh tagged as IsSky then we render them only, otherwise we simply render the sky atmosphere itself.
@@ -824,7 +834,7 @@ void FScene::AllocateAndCaptureFrameSkyEnvMap(
 								});
 						}
 					}
-					else
+					else if(!bAlwaysClearColorBuffer) // Only clear if not done before.
 					{
 						ClearCubeFace(SkyCubeTexture, CubeFace);
 					}
