Support reserved resources for the GPU-Scene instance data buffer to remove the need to perform copies on grow/shrink, reducing hitches.
- Changed the instance scene data layoyt to a tiled SOA such that we can grow/shrink the buffer without moving everything.
- r.GPUScene.InstanceDataTileSizeLog2 controls the new behavior and is disabled ( <0 ) by default, if enabled, >= 0, it will use reserved resources if available on the platform.

diff --git a/Engine/Plugins/Experimental/GPULightmass/Source/GPULightmass/Private/LightmapRenderer.cpp b/Engine/Plugins/Experimental/GPULightmass/Source/GPULightmass/Private/LightmapRenderer.cpp
--- a/Engine/Plugins/Experimental/GPULightmass/Source/GPULightmass/Private/LightmapRenderer.cpp
+++ b/Engine/Plugins/Experimental/GPULightmass/Source/GPULightmass/Private/LightmapRenderer.cpp
@@ -576,28 +576,37 @@ void FCachedRayTracingSceneData::SetupViewAndSceneUniformBufferFromSceneRenderSt
 		{
 			TRACE_CPUPROFILER_EVENT_SCOPE(InstanceSceneData);
 
-			TArray<FVector4f> InstanceSceneDataSOA;
-			InstanceSceneDataSOA.AddZeroed(FInstanceSceneShaderData::GetDataStrideInFloat4s() * InstanceSceneData.Num());
-			for (uint32 ArrayIndex = 0; ArrayIndex < FInstanceSceneShaderData::GetDataStrideInFloat4s(); ArrayIndex++)
+			GPUSceneNumInstances = InstanceSceneData.Num();
+			GPUSceneInstanceDataSOAStride = FMath::Max(1u, FMath::RoundUpToPowerOfTwo(InstanceSceneData.Num()));
+
+			FRDGBufferRef RDGInstanceSceneDataBuffer = nullptr;
 			{
-				for (int32 DataIndex = 0; DataIndex < InstanceSceneData.Num(); DataIndex++)
+				TArray<FVector4f> InstanceSceneDataSOA;
+				InstanceSceneDataSOA.AddZeroed(FInstanceSceneShaderData::GetDataStrideInFloat4s() * GPUSceneInstanceDataSOAStride);
+				for (uint32 ArrayIndex = 0; ArrayIndex < FInstanceSceneShaderData::GetDataStrideInFloat4s(); ArrayIndex++)
 				{
-					InstanceSceneDataSOA[ArrayIndex * InstanceSceneData.Num() + DataIndex] = InstanceSceneData[DataIndex].Data[ArrayIndex];
+					for (int32 DataIndex = 0; DataIndex < InstanceSceneData.Num(); DataIndex++)
+					{
+						InstanceSceneDataSOA[ArrayIndex * GPUSceneInstanceDataSOAStride + DataIndex] = InstanceSceneData[DataIndex].Data[ArrayIndex];
+					}
 				}
-			}
 
-			if (InstanceSceneDataSOA.Num() == 0)
-			{
-				InstanceSceneDataSOA.AddZeroed(FInstanceSceneShaderData::GetDataStrideInFloat4s());
-			}
+				if (InstanceSceneDataSOA.Num() == 0)
+				{
+					InstanceSceneDataSOA.AddZeroed(FInstanceSceneShaderData::GetDataStrideInFloat4s());
+				}
 
-			FRDGBufferRef RDGInstanceSceneDataBuffer = CreateStructuredBuffer(GraphBuilder, TEXT("InstanceSceneDataBuffer"), InstanceSceneDataSOA);
+				RDGInstanceSceneDataBuffer = CreateStructuredBuffer(GraphBuilder, TEXT("InstanceSceneDataBuffer"), MoveTemp(InstanceSceneDataSOA));
+			}
 			GPUSceneInstanceDataBuffer = GraphBuilder.ConvertToExternalBuffer(RDGInstanceSceneDataBuffer);
 			GPUScene.GPUSceneInstanceSceneData = GraphBuilder.CreateSRV(RDGInstanceSceneDataBuffer);
-			GPUSceneInstanceDataSOAStride = InstanceSceneData.Num();
-			GPUScene.CommonParameters.GPUSceneInstanceDataSOAStride = GPUSceneInstanceDataSOAStride;
+
+			check(FMath::IsPowerOfTwo(GPUSceneInstanceDataSOAStride));
+			GPUScene.CommonParameters.GPUSceneInstanceDataTileSizeLog2 = FMath::FloorLog2(GPUSceneInstanceDataSOAStride);
+			GPUScene.CommonParameters.GPUSceneInstanceDataTileSizeMask = (1u << GPUScene.CommonParameters.GPUSceneInstanceDataTileSizeLog2) - 1u;
+			GPUScene.CommonParameters.GPUSceneInstanceDataTileStride = FInstanceSceneShaderData::GetDataStrideInFloat4s() << GPUScene.CommonParameters.GPUSceneInstanceDataTileSizeLog2;
 			GPUScene.CommonParameters.GPUSceneFrameNumber = 0u;
-			GPUScene.CommonParameters.GPUSceneMaxAllocatedInstanceId = InstanceSceneData.Num();
+			GPUScene.CommonParameters.GPUSceneMaxAllocatedInstanceId = GPUSceneNumInstances;
 			GPUScene.CommonParameters.GPUSceneMaxPersistentPrimitiveIndex = 0;
 		}
 
@@ -816,9 +825,13 @@ void FCachedRayTracingSceneData::RestoreCachedBuffers(FRDGBuilder& GraphBuilder,
 
 	check(GPUSceneInstanceDataBuffer.IsValid())
 	GPUScene.GPUSceneInstanceSceneData = GraphBuilder.CreateSRV(GraphBuilder.RegisterExternalBuffer(GPUSceneInstanceDataBuffer));
-	GPUScene.CommonParameters.GPUSceneInstanceDataSOAStride = GPUSceneInstanceDataSOAStride;
+
+	check(FMath::IsPowerOfTwo(GPUSceneInstanceDataSOAStride));
+	GPUScene.CommonParameters.GPUSceneInstanceDataTileSizeLog2 = FMath::FloorLog2(GPUSceneInstanceDataSOAStride);
+	GPUScene.CommonParameters.GPUSceneInstanceDataTileSizeMask = (1u << GPUScene.CommonParameters.GPUSceneInstanceDataTileSizeLog2) - 1u;
+	GPUScene.CommonParameters.GPUSceneInstanceDataTileStride = FInstanceSceneShaderData::GetDataStrideInFloat4s() << GPUScene.CommonParameters.GPUSceneInstanceDataTileSizeLog2;
 	GPUScene.CommonParameters.GPUSceneFrameNumber = 0u;
-	GPUScene.CommonParameters.GPUSceneMaxAllocatedInstanceId = GPUSceneInstanceDataSOAStride;
+	GPUScene.CommonParameters.GPUSceneMaxAllocatedInstanceId = GPUSceneNumInstances;
 	GPUScene.CommonParameters.GPUSceneMaxPersistentPrimitiveIndex = 0;
 
 	check(GPUSceneInstancePayloadDataBuffer.IsValid())
@@ -2189,7 +2202,7 @@ void FLightmapRenderer::Finalize(FRDGBuilder& GraphBuilder)
 			TRACE_CPUPROFILER_EVENT_SCOPE(InstanceIdsIdentityBuffer);
 
 			TArray<uint32, SceneRenderingAllocator> InstanceIdsIdentity;
-			for (uint32 Index = 0U; Index < FMath::Max(1U, SceneUniforms.Get(SceneUB::GPUScene).CommonParameters.GPUSceneInstanceDataSOAStride); ++Index)
+			for (uint32 Index = 0U; Index < FMath::Max(1U, uint32(Scene->CachedRayTracingScene->GPUSceneInstanceDataSOAStride)); ++Index)
 			{
 				InstanceIdsIdentity.Add(Index);
 			}
diff --git a/Engine/Plugins/Experimental/GPULightmass/Source/GPULightmass/Private/Scene/Scene.h b/Engine/Plugins/Experimental/GPULightmass/Source/GPULightmass/Private/Scene/Scene.h
--- a/Engine/Plugins/Experimental/GPULightmass/Source/GPULightmass/Private/Scene/Scene.h
+++ b/Engine/Plugins/Experimental/GPULightmass/Source/GPULightmass/Private/Scene/Scene.h
@@ -87,6 +87,7 @@ struct FCachedRayTracingSceneData
 	TRefCountPtr<FRDGPooledBuffer> GPUSceneLightmapDataBuffer;
 	TRefCountPtr<FRDGPooledBuffer> GPUSceneInstanceDataBuffer;
 	int							   GPUSceneInstanceDataSOAStride;
+	int32						   GPUSceneNumInstances;
 	TRefCountPtr<FRDGPooledBuffer> GPUSceneInstancePayloadDataBuffer;
 	TRefCountPtr<FRDGPooledBuffer> GPUSceneLightDataBuffer;
 
diff --git a/Engine/Shaders/Private/SceneData.ush b/Engine/Shaders/Private/SceneData.ush
--- a/Engine/Shaders/Private/SceneData.ush
+++ b/Engine/Shaders/Private/SceneData.ush
@@ -590,7 +590,10 @@ StructuredBuffer<float4> GPUSceneInstanceSceneData;
 StructuredBuffer<float4> GPUSceneInstancePayloadData;
 uint GPUSceneFrameNumber;
 uint GPUSceneMaxPersistentPrimitiveIndex;
-uint GPUSceneInstanceDataSOAStride;
+uint GPUSceneInstanceDataTileSizeLog2;
+uint GPUSceneInstanceDataTileSizeMask;
+uint GPUSceneInstanceDataTileStride;
+
 uint GPUSceneMaxAllocatedInstanceId;
 uint GPUSceneNumLightmapDataItems;
 #elif USE_GPU_SCENE_DATA_RW
@@ -598,7 +601,9 @@ uint GPUSceneNumLightmapDataItems;
 RWStructuredBuffer<float4> GPUSceneInstanceSceneDataRW;
 RWStructuredBuffer<float4> GPUSceneInstancePayloadDataRW;
 uint GPUSceneFrameNumber;
-uint GPUSceneInstanceDataSOAStride;
+uint GPUSceneInstanceDataTileSizeLog2;
+uint GPUSceneInstanceDataTileSizeMask;
+uint GPUSceneInstanceDataTileStride;
 uint GPUSceneMaxPersistentPrimitiveIndex;
 uint GPUSceneMaxAllocatedInstanceId;
 uint GPUSceneNumLightmapDataItems;
@@ -613,8 +618,12 @@ struct FSceneData
 	/** Upper bound that includes all the IDs of ALL allocated instances (scene & dynamic). This is always >= the number of instances in the scene as there may be holes. */
 	uint MaxAllocatedInstanceId;
 
-	uint InstanceDataSOAStride;
 	uint NumLightmapDataItems;
+
+	// implementation details, use at own risk
+	uint InstanceDataTileSizeLog2;
+	uint InstanceDataTileSizeMask;
+	uint InstanceDataTileStride;
 };
 
 FSceneData GetSceneData()
@@ -622,7 +631,9 @@ FSceneData GetSceneData()
 	FSceneData SceneData;
 #if USE_GLOBAL_GPU_SCENE_DATA
 	SceneData.FrameNumber = GPUSceneFrameNumber;
-	SceneData.InstanceDataSOAStride = GPUSceneInstanceDataSOAStride;
+	SceneData.InstanceDataTileSizeLog2 = GPUSceneInstanceDataTileSizeLog2;
+	SceneData.InstanceDataTileSizeMask = GPUSceneInstanceDataTileSizeMask;
+	SceneData.InstanceDataTileStride = GPUSceneInstanceDataTileStride;
 	SceneData.MaxPersistentPrimitiveIndex = GPUSceneMaxPersistentPrimitiveIndex;
 	SceneData.MaxAllocatedInstanceId = GPUSceneMaxAllocatedInstanceId;
 	SceneData.NumLightmapDataItems = GPUSceneNumLightmapDataItems;
@@ -630,13 +641,17 @@ FSceneData GetSceneData()
 
 #if ENABLE_SCENE_DATA_DX11_UB_ERROR_WORKAROUND
 	SceneData.FrameNumber = GPUSceneFrameNumber;
-	SceneData.InstanceDataSOAStride = GPUSceneInstanceDataSOAStride;
+	SceneData.InstanceDataTileSizeLog2 = GPUSceneInstanceDataTileSizeLog2;
+	SceneData.InstanceDataTileSizeMask = GPUSceneInstanceDataTileSizeMask;
+	SceneData.InstanceDataTileStride = GPUSceneInstanceDataTileStride;
 	SceneData.MaxPersistentPrimitiveIndex = GPUSceneMaxPersistentPrimitiveIndex;
 	SceneData.MaxAllocatedInstanceId = GPUSceneMaxAllocatedInstanceId;
 	SceneData.NumLightmapDataItems = GPUSceneNumLightmapDataItems;
 #else
 	SceneData.FrameNumber = GPUSceneWriter.GPUSceneFrameNumber;
-	SceneData.InstanceDataSOAStride = GPUSceneWriter.GPUSceneInstanceDataSOAStride;
+	SceneData.InstanceDataTileSizeLog2 = GPUSceneWriter.InstanceDataTileSizeLog2;
+	SceneData.InstanceDataTileSizeMask = GPUSceneWriter.InstanceDataTileSizeMask;
+	SceneData.InstanceDataTileStride = GPUSceneWriter.InstanceDataTileStride;
 	SceneData.MaxPersistentPrimitiveIndex = GPUSceneWriter.GPUSceneMaxPersistentPrimitiveIndex;
 	SceneData.MaxAllocatedInstanceId = GPUSceneWriter.GPUSceneMaxAllocatedInstanceId;
 	SceneData.NumLightmapDataItems = GPUSceneWriter.GPUSceneNumLightmapDataItems;
@@ -644,7 +659,9 @@ FSceneData GetSceneData()
 
 #else
 	SceneData.FrameNumber = Scene.GPUScene.GPUSceneFrameNumber;
-	SceneData.InstanceDataSOAStride = Scene.GPUScene.GPUSceneInstanceDataSOAStride;
+	SceneData.InstanceDataTileSizeLog2 = Scene.GPUScene.GPUSceneInstanceDataTileSizeLog2;
+	SceneData.InstanceDataTileSizeMask = Scene.GPUScene.GPUSceneInstanceDataTileSizeMask;
+	SceneData.InstanceDataTileStride = Scene.GPUScene.GPUSceneInstanceDataTileStride;
 	SceneData.MaxPersistentPrimitiveIndex = Scene.GPUScene.GPUSceneMaxPersistentPrimitiveIndex;
 	SceneData.MaxAllocatedInstanceId = Scene.GPUScene.GPUSceneMaxAllocatedInstanceId;
 	SceneData.NumLightmapDataItems = Scene.GPUScene.GPUSceneNumLightmapDataItems;
@@ -657,20 +674,15 @@ uint GetGPUSceneFrameNumber()
 	return GetSceneData().FrameNumber;
 }
 
-/**
- * Function to calculate the instance data index in the instance data buffer given the instance ID and array index.
- */
-uint CalcInstanceDataIndex(uint InstanceId, uint ArrayIndex, uint InstanceDataSOAStride)
-{
-	return InstanceDataSOAStride * ArrayIndex + InstanceId;
-}
-
 /**
  * Function to calculate the instance data index in the instance data buffer given the instance ID and array index
  */
 uint CalcInstanceDataIndex(uint InstanceId, uint ArrayIndex)
 {
-	return CalcInstanceDataIndex(InstanceId, ArrayIndex, GetSceneData().InstanceDataSOAStride);
+	FSceneData SceneData = GetSceneData();
+	uint TileId = InstanceId >> SceneData.InstanceDataTileSizeLog2;
+	uint IdInTile = InstanceId & SceneData.InstanceDataTileSizeMask;
+	return SceneData.InstanceDataTileStride * TileId + (ArrayIndex << SceneData.InstanceDataTileSizeLog2) + IdInTile;
 }
 
 float4 LoadInstanceSceneDataElement(uint Index)
diff --git a/Engine/Source/Runtime/Renderer/Private/GPUScene.cpp b/Engine/Source/Runtime/Renderer/Private/GPUScene.cpp
--- a/Engine/Source/Runtime/Renderer/Private/GPUScene.cpp
+++ b/Engine/Source/Runtime/Renderer/Private/GPUScene.cpp
@@ -34,7 +34,7 @@
 
 // Useful for debugging
 #define FORCEINLINE_GPUSCENE FORCEINLINE
-//#define FORCEINLINE_GPUSCENE
+//#define FORCEINLINE_GPUSCENE FORCENOINLINE
 
 static_assert(
 	sizeof(FLightSceneData) == SIZEOF_LIGHT_SCENE_DATA,
@@ -59,10 +59,12 @@ static void ConstructDefault(FGPUSceneResourceParameters& GPUScene, FRDGBuilder&
 	FRDGBufferRef DummyBufferLight = GSystemTextures.GetDefaultStructuredBuffer(GraphBuilder, sizeof(FLightSceneData));
 	GPUScene.GPUSceneLightData = GraphBuilder.CreateSRV(DummyBufferLight);
 	GPUScene.CommonParameters.GPUSceneFrameNumber = 0u;
-	GPUScene.CommonParameters.GPUSceneInstanceDataSOAStride = 0u;
 	GPUScene.CommonParameters.GPUSceneMaxAllocatedInstanceId = 0;
 	GPUScene.CommonParameters.GPUSceneMaxPersistentPrimitiveIndex = 0;
 	GPUScene.CommonParameters.GPUSceneNumLightmapDataItems = 0;
+	GPUScene.CommonParameters.GPUSceneInstanceDataTileSizeLog2 = 0u;
+	GPUScene.CommonParameters.GPUSceneInstanceDataTileSizeMask = 0u;
+	GPUScene.CommonParameters.GPUSceneInstanceDataTileStride = 0u;
 }
 
 IMPLEMENT_SCENE_UB_STRUCT(FGPUSceneResourceParameters, GPUScene, ConstructDefault);
@@ -131,6 +133,15 @@ static TAutoConsoleVariable<int32> CVarGPUSceneLightsAsyncSetup(
 	ECVF_RenderThreadSafe
 );
 
+static TAutoConsoleVariable<int32> CVarGPUSceneInstanceDataTileSizeLog2(
+	TEXT("r.GPUScene.InstanceDataTileSizeLog2"),
+	-1,
+	TEXT("Log2 of tile size used for the instance data in number of instances. \n")
+	TEXT("  <0 (default) indicates no tiling and disables reserved resources. \n")
+	TEXT("  >0 will use the given tile-size and attempt to use reserved resources, but falls back to a regular buffer if not available."),
+	ECVF_RenderThreadSafe
+);
+
 LLM_DECLARE_TAG_API(GPUScene, RENDERER_API);
 DECLARE_LLM_MEMORY_STAT(TEXT("GPUScene"), STAT_GPUSceneLLM, STATGROUP_LLMFULL);
 DECLARE_LLM_MEMORY_STAT(TEXT("GPUScene"), STAT_GPUSceneSummaryLLM, STATGROUP_LLM);
@@ -437,12 +448,12 @@ struct FUploadDataSourceAdapterScenePrimitives
 	/**
 	 * Return the number of primitives to upload N, GetPrimitiveInfoHeader/GetPrimitiveShaderData will be called with ItemIndex in [0,N).
 	 */
-	FORCEINLINE int32 NumPrimitivesToUpload() const 
+	FORCEINLINE_GPUSCENE int32 NumPrimitivesToUpload() const 
 	{ 
 		return PrimitivesToUpdate.Num(); 
 	}
 
-	FORCEINLINE TArrayView<const uint32> GetItemPrimitiveIds() const
+	FORCEINLINE_GPUSCENE TArrayView<const uint32> GetItemPrimitiveIds() const
 	{
 		return TArrayView<const uint32>(reinterpret_cast<const uint32 *>(PrimitivesToUpdate.GetData()), PrimitivesToUpdate.Num());
 	}
@@ -478,7 +489,7 @@ struct FUploadDataSourceAdapterScenePrimitives
 		}
 	}
 
-	FORCEINLINE uint32 PackFlags(FInstanceDataFlags Flags) const
+	FORCEINLINE_GPUSCENE uint32 PackFlags(FInstanceDataFlags Flags) const
 	{
 		uint32 PackedFlags = 0x0;
 		PackedFlags |= Flags.bHasPerInstanceRandom				? INSTANCE_SCENE_DATA_FLAG_HAS_RANDOM              : 0u;
@@ -646,7 +657,6 @@ FGPUScene::FGPUScene(FScene &InScene)
 	: bUpdateAllPrimitives(false)
 	, InstancePayloadDataAllocator(CVarGPUSceneUseGrowOnlyAllocationPolicy.GetValueOnAnyThread() != 0)
 	, LightmapDataAllocator(CVarGPUSceneUseGrowOnlyAllocationPolicy.GetValueOnAnyThread() != 0)
-	, InstanceSceneDataSOAStride(0)
 	, LightDataBuffer(InitialBufferSize, TEXT("GPUScene.LightData"))
 	, Scene(InScene)
 	, InstanceSceneDataAllocator(CVarGPUSceneUseGrowOnlyAllocationPolicy.GetValueOnAnyThread() != 0)
@@ -654,14 +664,14 @@ FGPUScene::FGPUScene(FScene &InScene)
 #if !UE_BUILD_SHIPPING
 	ScreenMessageDelegate = FRendererOnScreenNotification::Get().AddLambda([this](TMultiMap<FCoreDelegates::EOnScreenMessageSeverity, FText >& OutMessages)
 	{
-		if (InstanceSceneDataSOAStride < MaxInstancesDuringPrevUpdate)
+		if (MaxInstancesDuringPrevUpdate >= MAX_INSTANCE_ID)
 		{
 			OutMessages.Add(FCoreDelegates::EOnScreenMessageSeverity::Warning, FText::FromString(FString::Printf(TEXT("GPU-Scene Instance data overflow detected, reduce instance count to avoid rendering artifacts"))));
-			OutMessages.Add(FCoreDelegates::EOnScreenMessageSeverity::Warning, FText::FromString(FString::Printf(TEXT("  Max allocated ID %d, instance buffer size: %dM"), MaxInstancesDuringPrevUpdate, InstanceSceneDataSOAStride >> 20)));
+			OutMessages.Add(FCoreDelegates::EOnScreenMessageSeverity::Warning, FText::FromString(FString::Printf(TEXT("  Max allocated ID %d, max instance buffer size: %dM"), MaxInstancesDuringPrevUpdate, MAX_INSTANCE_ID >> 20)));
 			if (!bLoggedInstanceOverflow )
 			{
 				UE_LOG(LogRenderer, Warning, TEXT("GPU-Scene Instance data overflow detected, reduce instance count to avoid rendering artifacts.\n")
-					TEXT(" Max allocated ID %d (%0.3fM), instance buffer size: %dM"), MaxInstancesDuringPrevUpdate, double(MaxInstancesDuringPrevUpdate) / (1024.0 * 1024.0),  InstanceSceneDataSOAStride >> 20);
+					TEXT(" Max allocated ID %d (%0.3fM), instance buffer size: %dM"), MaxInstancesDuringPrevUpdate, double(MaxInstancesDuringPrevUpdate) / (1024.0 * 1024.0),  MAX_INSTANCE_ID >> 20);
 				bLoggedInstanceOverflow = true;
 			}
 		}
@@ -873,6 +883,38 @@ void FGPUScene::UpdateInternal(FRDGBuilder& GraphBuilder, FSceneUniformBuffer& S
 	MaxInstancesDuringPrevUpdate = uint32(InstanceSceneDataAllocator.GetMaxSize());
 #endif // UE_BUILD_SHIPPING
 
+	bool bPrevUseTiledInstanceDataLayout = bUseTiledInstanceDataLayout;
+	bUseTiledInstanceDataLayout = CVarGPUSceneInstanceDataTileSizeLog2.GetValueOnRenderThread() >= 0;
+	int32 NewTileSizeLog2 = FMath::Max(0, CVarGPUSceneInstanceDataTileSizeLog2.GetValueOnRenderThread());
+	// if it toggled tiled<->SOA OR the tile size changed, reupload everything 
+	if (bPrevUseTiledInstanceDataLayout != bUseTiledInstanceDataLayout)
+	{
+		// reallocate the buffer on switch to/from reserved.
+		// Alywas drop the old buffer - this simplifies the logic here and the tile mode should only be changed as a testing thing anyway, probably make read only before submit.
+		InstanceSceneDataBuffer.SafeRelease();
+
+		// Trigger full upload, the tile size has changed!
+		bUpdateAllPrimitives = true;
+	}
+	
+	// if the tile size changed we need to do a full reupload. 
+	if (bUseTiledInstanceDataLayout && (NewTileSizeLog2 != InstanceDataTileSizeLog2))
+	{
+		InstanceDataTileSizeLog2 = NewTileSizeLog2;
+		bUpdateAllPrimitives = true;
+	}
+
+	const bool bUseReservedResources = bUseTiledInstanceDataLayout && GRHIGlobals.ReservedResources.Supported;
+	// For the reserved case, we reserve a "max-sized" buffer up front, otherwise leave it to the resize.
+	if (bUseReservedResources && !InstanceSceneDataBuffer.IsValid())
+	{
+		// allocate ~2GB of address space
+		const uint64 MaxSizeInBytes = uint64(2048) << 20;
+		FRDGBufferDesc InstanceDataBufferDesc = FRDGBufferDesc::CreateStructuredDesc(sizeof(FVector4f), uint32(MaxSizeInBytes / sizeof(FVector4f)));
+		InstanceDataBufferDesc.Usage |= EBufferUsageFlags::ReservedResource;
+		InstanceSceneDataBuffer = AllocatePooledBuffer(InstanceDataBufferDesc, TEXT("GPUScene.InstanceSceneData"));
+	}
+
 	if ((CVarGPUSceneUploadEveryFrame.GetValueOnRenderThread() != 0) || bUpdateAllPrimitives)
 	{
 		PrimitivesToUpdate.Reset();
@@ -921,7 +963,7 @@ void FGPUScene::UpdateInternal(FRDGBuilder& GraphBuilder, FSceneUniformBuffer& S
 
 	FUploadDataSourceAdapterScenePrimitives& Adapter = *GraphBuilder.AllocObject<FUploadDataSourceAdapterScenePrimitives>(Scene, SceneFrameNumber, MoveTemp(PrimitivesToUpdate), MoveTemp(PrimitiveDirtyState));
 	
-	FRegisteredBuffers BufferState = UpdateBufferAllocations(GraphBuilder, SceneUB, Adapter);
+	FRegisteredBuffers BufferState = UpdateBufferAllocations(GraphBuilder, SceneUB);
 
 	// Run a pass that clears (Sets ID to invalid) any instances that need it
 	AddClearInstancesPass(GraphBuilder, Scene.InstanceCullingOcclusionQueryRenderer);
@@ -963,8 +1005,7 @@ void FGPUScene::UpdateInternal(FRDGBuilder& GraphBuilder, FSceneUniformBuffer& S
 	}
 }
 
-template<typename FUploadDataSourceAdapter>
-FGPUScene::FRegisteredBuffers FGPUScene::UpdateBufferAllocations(FRDGBuilder& GraphBuilder, FSceneUniformBuffer& SceneUB, const FUploadDataSourceAdapter& UploadDataSourceAdapter)
+FGPUScene::FRegisteredBuffers FGPUScene::UpdateBufferAllocations(FRDGBuilder& GraphBuilder, FSceneUniformBuffer& SceneUB)
 {
 	LLM_SCOPE_BYTAG(GPUScene);
 
@@ -977,15 +1018,34 @@ FGPUScene::FRegisteredBuffers FGPUScene::UpdateBufferAllocations(FRDGBuilder& Gr
 	const uint32 SizeReserve = FMath::RoundUpToPowerOfTwo(FMath::Max(DynamicPrimitivesOffset, InitialBufferSize));
 	BufferState.PrimitiveBuffer = ResizeStructuredBufferIfNeeded(GraphBuilder, PrimitiveBuffer, SizeReserve * sizeof(FPrimitiveSceneShaderData::Data), TEXT("GPUScene.PrimitiveData"));
 
-	// Clamp buffer to be smaller than the MAX_INSTANCE_ID.
-	const uint32 InstanceSceneDataSizeReserve = FMath::Min(MAX_INSTANCE_ID, FMath::RoundUpToPowerOfTwo(FMath::Max(InstanceSceneDataAllocator.GetMaxSize(), InitialBufferSize)));
+	// clamped instance count / high watermark (any overflowing will not render!)
+	const uint32 ClampedMaxInstanceId = FMath::Min(MAX_INSTANCE_ID, uint32(FMath::Max(InstanceSceneDataAllocator.GetMaxSize(), InitialBufferSize)));
+
+	// emulate the old functioning, deriving the stride from the instance count
+	if (!bUseTiledInstanceDataLayout)
+	{
+		// Round to next POT
+		const uint32 InstanceSceneDataSizeReserve = FMath::RoundUpToPowerOfTwo(ClampedMaxInstanceId);
+
+		FResizeResourceSOAParams ResizeParams;
+		ResizeParams.NumBytes = InstanceSceneDataSizeReserve * FInstanceSceneShaderData::GetEffectiveNumBytes();
+		ResizeParams.NumArrays = FInstanceSceneShaderData::GetDataStrideInFloat4s();
 
-	FResizeResourceSOAParams ResizeParams;
-	ResizeParams.NumBytes = InstanceSceneDataSizeReserve * FInstanceSceneShaderData::GetEffectiveNumBytes();
-	ResizeParams.NumArrays = FInstanceSceneShaderData::GetDataStrideInFloat4s();
+		BufferState.InstanceSceneDataBuffer = ResizeStructuredBufferSOAIfNeeded(GraphBuilder, InstanceSceneDataBuffer, ResizeParams, TEXT("GPUScene.InstanceSceneData"));
+		check(FMath::IsPowerOfTwo(InstanceSceneDataSizeReserve));
+		InstanceDataTileSizeLog2 = FMath::FloorLog2(InstanceSceneDataSizeReserve);
+	}
+	else
+	{
+		// Figure out the size required for the buffer when rounded to number of tiles.
+		const uint32 TileSizeInBytes = sizeof(FVector4f) * GetInstanceDataTileStride();
+		const uint32 NumTiles = FMath::DivideAndRoundUp(ClampedMaxInstanceId, GetInstanceDataTileSize());
 
-	BufferState.InstanceSceneDataBuffer = ResizeStructuredBufferSOAIfNeeded(GraphBuilder, InstanceSceneDataBuffer, ResizeParams, TEXT("GPUScene.InstanceSceneData"));
-	InstanceSceneDataSOAStride = InstanceSceneDataSizeReserve;
+		// Grow/shrink in 16MB chunks to reduce the number of page table updates or resizes
+		const uint32 AllocationStrideInBytes = 16u << 20;
+		const uint32 InstanceSceneDataSizeReserve = FMath::DivideAndRoundUp(NumTiles * TileSizeInBytes, AllocationStrideInBytes) * AllocationStrideInBytes;
+		BufferState.InstanceSceneDataBuffer = ResizeStructuredBufferIfNeeded(GraphBuilder, InstanceSceneDataBuffer, InstanceSceneDataSizeReserve, TEXT("GPUScene.InstanceSceneData"));
+	}
 
 	const uint32 PayloadFloat4Count = FMath::Max(InstancePayloadDataAllocator.GetMaxSize(), InitialBufferSize);
 	const uint32 InstancePayloadDataSizeReserve = FMath::RoundUpToPowerOfTwo(PayloadFloat4Count * sizeof(FVector4f));
@@ -996,7 +1056,10 @@ FGPUScene::FRegisteredBuffers FGPUScene::UpdateBufferAllocations(FRDGBuilder& Gr
 
 	BufferState.LightDataBuffer = GraphBuilder.RegisterExternalBuffer(LightDataBuffer.GetPooledBuffer());
 
-	CommonParameters.GPUSceneInstanceDataSOAStride = InstanceSceneDataSOAStride;
+	// Set up with no tiling, i.e, behaves as before.
+	CommonParameters.GPUSceneInstanceDataTileSizeLog2 = GetInstanceDataTileSizeLog2();
+	CommonParameters.GPUSceneInstanceDataTileSizeMask =  GetInstanceDataTileSizeMask();
+	CommonParameters.GPUSceneInstanceDataTileStride = GetInstanceDataTileStride();
 	CommonParameters.GPUSceneFrameNumber = GetSceneFrameNumber();
 	CommonParameters.GPUSceneMaxAllocatedInstanceId = InstanceSceneDataAllocator.GetMaxSize();
 	CommonParameters.GPUSceneMaxPersistentPrimitiveIndex = Scene.GetMaxPersistentPrimitiveIndex();
@@ -1125,6 +1188,35 @@ struct FInstanceBatcher
 	}
 };
 
+namespace
+{
+	struct FGPUSceneUploadTaskContext
+	{
+		TArray<FPrimitiveUploadInfoHeader, FSceneRenderingArrayAllocator> PrimitiveUploadInfos;
+
+		FRDGScatterUploader* PrimitiveUploader = nullptr;
+		FRDGScatterUploader* InstancePayloadUploader = nullptr;
+		FRDGScatterUploader* InstanceSceneUploader = nullptr;
+		FRDGScatterUploader* LightmapUploader = nullptr;
+
+		int32 NumPrimitiveDataUploads = 0;
+		int32 NumLightmapDataUploads = 0;
+		int32 NumInstanceSceneDataUploads = 0;
+		int32 NumInstancePayloadDataUploads = 0; // Count of float4s
+		uint32 InstanceDataTileSizeLog2 = 0u;
+		uint32 InstanceDataTileSizeMask = 0u;
+		uint32 InstanceDataTileStride = 0u;
+
+		FORCEINLINE_GPUSCENE uint32 CalcInstanceDataIndex(uint32 InstanceId, uint32 ArrayIndex)
+		{
+			uint32 TileId = InstanceId >> InstanceDataTileSizeLog2;
+			uint32 IdInTile = InstanceId & InstanceDataTileSizeMask;
+			return InstanceDataTileStride * TileId + (ArrayIndex << InstanceDataTileSizeLog2) + IdInTile;
+		}
+	};
+
+};
+
 template<typename FUploadDataSourceAdapter>
 void FGPUScene::UploadGeneral(FRDGBuilder& GraphBuilder, const FRegisteredBuffers &BufferState, FRDGExternalAccessQueue* ExternalAccessQueue, const FUploadDataSourceAdapter& UploadDataSourceAdapter, const UE::Tasks::FTask& PrerequisiteTask)
 {
@@ -1140,34 +1232,18 @@ void FGPUScene::UploadGeneral(FRDGBuilder& GraphBuilder, const FRegisteredBuffer
 	}
 
 	const bool bExecuteInParallel = CVarGPUSceneParallelUpdate.GetValueOnRenderThread() != 0 && FApp::ShouldUseThreadingForPerformance();
-	const bool bNaniteEnabled = DoesPlatformSupportNanite(GMaxRHIShaderPlatform);
 
 	SCOPED_NAMED_EVENT(UpdateGPUScene, FColor::Green);
 
 	RDG_GPU_MASK_SCOPE(GraphBuilder, FRHIGPUMask::All());
 	RDG_EVENT_SCOPE(GraphBuilder, "UpdateGPUScene NumPrimitiveDataUploads %u", NumPrimitiveDataUploads);
 
-	struct FTaskContext
-	{
-		TArray<FPrimitiveUploadInfoHeader, FSceneRenderingArrayAllocator> PrimitiveUploadInfos;
-
-		FRDGScatterUploader* PrimitiveUploader = nullptr;
-		FRDGScatterUploader* InstancePayloadUploader = nullptr;
-		FRDGScatterUploader* InstanceSceneUploader = nullptr;
-		FRDGScatterUploader* LightmapUploader = nullptr;
-
-		int32 NumPrimitiveDataUploads = 0;
-		int32 NumLightmapDataUploads = 0;
-		int32 NumInstanceSceneDataUploads = 0;
-		int32 NumInstancePayloadDataUploads = 0; // Count of float4s
-
-		uint32 InstanceSceneDataSOAStride = 1;
-	};
-
-	FTaskContext& TaskContext = *GraphBuilder.AllocObject<FTaskContext>();
+	FGPUSceneUploadTaskContext& TaskContext = *GraphBuilder.AllocObject<FGPUSceneUploadTaskContext>();
 
 	TaskContext.NumPrimitiveDataUploads = NumPrimitiveDataUploads;
-	TaskContext.InstanceSceneDataSOAStride = InstanceSceneDataSOAStride;
+	TaskContext.InstanceDataTileSizeLog2 = GetInstanceDataTileSizeLog2();
+	TaskContext.InstanceDataTileSizeMask = GetInstanceDataTileSizeMask();
+	TaskContext.InstanceDataTileStride   = GetInstanceDataTileStride();
 	TaskContext.PrimitiveUploadInfos.SetNumUninitialized(NumPrimitiveDataUploads);
 
 	for (int32 ItemIndex = 0; ItemIndex < NumPrimitiveDataUploads; ++ItemIndex)
@@ -1197,7 +1273,7 @@ void FGPUScene::UploadGeneral(FRDGBuilder& GraphBuilder, const FRegisteredBuffer
 		TaskContext.LightmapUploader = LightmapUploadBuffer.Begin(GraphBuilder, BufferState.LightmapDataBuffer, TaskContext.NumLightmapDataUploads, sizeof(FLightmapSceneShaderData::Data), TEXT("LightmapUploadBuffer"));
 	}
 
-	GraphBuilder.AddCommandListSetupTask([&TaskContext, &UploadDataSourceAdapter, bNaniteEnabled, bExecuteInParallel, FeatureLevel = FeatureLevel](FRHICommandListBase& RHICmdList)
+	GraphBuilder.AddCommandListSetupTask([&TaskContext, &UploadDataSourceAdapter, bExecuteInParallel, FeatureLevel = FeatureLevel](FRHICommandListBase& RHICmdList)
 	{
 		SCOPED_NAMED_EVENT(UpdateGPUScene_Primitives, FColor::Green);
 
@@ -1316,12 +1392,12 @@ void FGPUScene::UploadGeneral(FRDGBuilder& GraphBuilder, const FRegisteredBuffer
 								!UploadInfo.bIsPrimitiveForceHidden
 							);
 						}
-						// RefIndex* BufferState.InstanceSceneDataSOAStride + UploadInfo.InstanceSceneDataOffset + InstanceIndex
 						const uint32 UploadInstanceItemOffset = (PrimitiveItemInfo.InstanceSceneDataUploadOffset + InstanceIndex) * InstanceDataStrideInFloat4s;
 
 						for (uint32 RefIndex = 0; RefIndex < InstanceDataStrideInFloat4s; ++RefIndex)
 						{
-							FVector4f* DstVector = static_cast<FVector4f*>(TaskContext.InstanceSceneUploader->Set_GetRef(UploadInstanceItemOffset + RefIndex, RefIndex * TaskContext.InstanceSceneDataSOAStride + UploadInfo.InstanceSceneDataOffset + InstanceIndex));
+							uint32 ScatterOffset = TaskContext.CalcInstanceDataIndex(UploadInfo.InstanceSceneDataOffset + InstanceIndex, RefIndex);
+							FVector4f* DstVector = static_cast<FVector4f*>(TaskContext.InstanceSceneUploader->Set_GetRef(UploadInstanceItemOffset + RefIndex, ScatterOffset));
 							*DstVector = InstanceSceneData.Data[RefIndex];
 						}
 
@@ -1685,7 +1761,7 @@ void FGPUScene::UploadDynamicPrimitiveShaderDataForViewInternal(FRDGBuilder& Gra
 			Collector.UploadData->InstancePayloadDataOffset,
 			SceneFrameNumber);
 
-		FRegisteredBuffers BufferState = UpdateBufferAllocations(GraphBuilder, View.GetSceneUniforms(), UploadAdapter);
+		FRegisteredBuffers BufferState = UpdateBufferAllocations(GraphBuilder, View.GetSceneUniforms());
 
 		// Run a pass that clears (Sets ID to invalid) any instances that need it.
 		AddClearInstancesPass(GraphBuilder, Scene.InstanceCullingOcclusionQueryRenderer);
@@ -2282,7 +2358,8 @@ TRDGUniformBufferRef<FGPUSceneWriterUniformParameters> FGPUScene::GetWriteParame
 
 	if (LegacyParameters != nullptr)
 	{
-		LegacyParameters->GPUSceneInstanceSceneDataSOAStride = UBParams->CommonParameters.GPUSceneInstanceDataSOAStride;
+		// Note, this is only correct if the tiling is disabled, but the parameter is only there because it is exposed in public API (and that only an accident, really), it is not used internally.
+		LegacyParameters->GPUSceneInstanceSceneDataSOAStride = 1u << InstanceDataTileSizeLog2;
 		LegacyParameters->GPUSceneNumAllocatedInstances = UBParams->CommonParameters.GPUSceneMaxAllocatedInstanceId;
 		LegacyParameters->GPUSceneNumAllocatedPrimitives = UBParams->CommonParameters.GPUSceneMaxPersistentPrimitiveIndex;
 		LegacyParameters->GPUSceneInstanceSceneDataRW = UBParams->GPUSceneInstanceSceneDataRW;
@@ -2298,6 +2375,16 @@ TRDGUniformBufferRef<FGPUSceneWriterUniformParameters> FGPUScene::GetWriteParame
 	return Result;
 }
 
+uint32 FGPUScene::GetInstanceDataTileSizeMask() const 
+{ 
+	return (1u << InstanceDataTileSizeLog2) - 1u; 
+}
+
+uint32 FGPUScene::GetInstanceDataTileStride() const 
+{ 
+	return FInstanceSceneShaderData::GetDataStrideInFloat4s() << InstanceDataTileSizeLog2; 
+}
+
 class FGPUSceneSetInstancePrimitiveIdCS : public FGlobalShader
 {
 	DECLARE_GLOBAL_SHADER(FGPUSceneSetInstancePrimitiveIdCS);
diff --git a/Engine/Source/Runtime/Renderer/Private/GPUScene.h b/Engine/Source/Runtime/Renderer/Private/GPUScene.h
--- a/Engine/Source/Runtime/Renderer/Private/GPUScene.h
+++ b/Engine/Source/Runtime/Renderer/Private/GPUScene.h
@@ -366,7 +366,16 @@ private:
 	 */
 	TRDGUniformBufferRef<FGPUSceneWriterUniformParameters> GetWriteParameters(FRDGBuilder& GraphBuilder, FGPUSceneWriterParameters* LegacyParameters = nullptr);
 
-	uint32                         InstanceSceneDataSOAStride;	// Distance between arrays in float4s
+	// set to true if the requested (cvar) tile size is >0, otherwise the tile size is re-derived from the instance count
+	bool bUseTiledInstanceDataLayout = true;
+
+	// log2(TileSize) in number of instances, thus the tile-size in memory is #float4s per instance * tilesize
+	int32 InstanceDataTileSizeLog2 = -1;
+
+	uint32 GetInstanceDataTileSizeLog2() const { return InstanceDataTileSizeLog2; }
+	uint32 GetInstanceDataTileSize() const { return 1u << InstanceDataTileSizeLog2; }
+	uint32 GetInstanceDataTileSizeMask() const;
+	uint32 GetInstanceDataTileStride() const;
 
 	static constexpr int32 InitialBufferSize = 256;
 
@@ -455,8 +464,7 @@ private:
 
 	ERHIFeatureLevel::Type FeatureLevel;
 
-	template<typename FUploadDataSourceAdapter>
-	FRegisteredBuffers UpdateBufferAllocations(FRDGBuilder& GraphBuilder, FSceneUniformBuffer& SceneUB, const FUploadDataSourceAdapter& UploadDataSourceAdapter);
+	FRegisteredBuffers UpdateBufferAllocations(FRDGBuilder& GraphBuilder, FSceneUniformBuffer& SceneUB);
 
 	/**
 	 * Register the external buffers with the graphbuilder.
diff --git a/Engine/Source/Runtime/Renderer/Private/Renderer.cpp b/Engine/Source/Runtime/Renderer/Private/Renderer.cpp
--- a/Engine/Source/Runtime/Renderer/Private/Renderer.cpp
+++ b/Engine/Source/Runtime/Renderer/Private/Renderer.cpp
@@ -239,10 +239,12 @@ FSceneUniformBuffer* FRendererModule::CreateSinglePrimitiveSceneUniformBuffer(FR
 			ShaderParameters.GPUSceneInstancePayloadData = GraphBuilder.CreateSRV(InstancePayloadDataBuffer);
 			ShaderParameters.GPUSceneLightmapData = GraphBuilder.CreateSRV(LightmapSceneDataBuffer);
 			ShaderParameters.GPUSceneLightData = GraphBuilder.CreateSRV(DummyBufferLight);
-			ShaderParameters.CommonParameters.GPUSceneInstanceDataSOAStride = 1;
 			ShaderParameters.CommonParameters.GPUSceneMaxAllocatedInstanceId = 1;
 			ShaderParameters.CommonParameters.GPUSceneMaxPersistentPrimitiveIndex = 1;
-			
+			ShaderParameters.CommonParameters.GPUSceneInstanceDataTileSizeLog2 = 0;
+			ShaderParameters.CommonParameters.GPUSceneInstanceDataTileSizeMask = 1;
+			ShaderParameters.CommonParameters.GPUSceneInstanceDataTileStride = 0;
+	
 			SceneUniforms.Set(SceneUB::GPUScene, ShaderParameters);
 		}
 	}
diff --git a/Engine/Source/Runtime/Renderer/Public/GPUSceneWriter.h b/Engine/Source/Runtime/Renderer/Public/GPUSceneWriter.h
--- a/Engine/Source/Runtime/Renderer/Public/GPUSceneWriter.h
+++ b/Engine/Source/Runtime/Renderer/Public/GPUSceneWriter.h
@@ -23,7 +23,9 @@ DECLARE_UNIFORM_BUFFER_STRUCT(FGPUSceneWriterUniformParameters, RENDERER_API)
 
 // Parameter sub struct that is common to all parameter use-cases
 BEGIN_SHADER_PARAMETER_STRUCT(FGPUSceneCommonParameters, RENDERER_API)
-	SHADER_PARAMETER(uint32, GPUSceneInstanceDataSOAStride)
+	SHADER_PARAMETER(uint32, GPUSceneInstanceDataTileSizeLog2)
+	SHADER_PARAMETER(uint32, GPUSceneInstanceDataTileSizeMask)
+	SHADER_PARAMETER(uint32, GPUSceneInstanceDataTileStride)
 	SHADER_PARAMETER(uint32, GPUSceneFrameNumber)
 	SHADER_PARAMETER(int32, GPUSceneMaxAllocatedInstanceId)
 	SHADER_PARAMETER(int32, GPUSceneMaxPersistentPrimitiveIndex)
