Process cluster lights after directional so that directional isn't applied to local lights. Additionally, ensure that cluster deferred shading doesn't run in the absence of directional lights.

diff --git a/Engine/Source/Runtime/Renderer/Private/MobileDeferredShadingPass.cpp b/Engine/Source/Runtime/Renderer/Private/MobileDeferredShadingPass.cpp
--- a/Engine/Source/Runtime/Renderer/Private/MobileDeferredShadingPass.cpp
+++ b/Engine/Source/Runtime/Renderer/Private/MobileDeferredShadingPass.cpp
@@ -682,7 +682,7 @@ static void RenderDirectionalLight(FRHICommandList& RHICmdList, const FScene& Sc
 	}
 }
 
-static void RenderDirectionalLights(FRHICommandList& RHICmdList, const FScene& Scene, const FViewInfo& View, const FCachedLightMaterial& DefaultLightMaterial, EMobileSSRQuality MobileSSRQuality)
+static int RenderDirectionalLights(FRHICommandList& RHICmdList, const FScene& Scene, const FViewInfo& View, const FCachedLightMaterial& DefaultLightMaterial, EMobileSSRQuality MobileSSRQuality)
 {
 	uint32 NumLights = 0;
 	for (uint32 ChannelIdx = 0; ChannelIdx < UE_ARRAY_COUNT(Scene.MobileDirectionalLights); ChannelIdx++)
@@ -706,6 +706,7 @@ static void RenderDirectionalLights(FRHICommandList& RHICmdList, const FScene& S
 			RenderDirectionalLight(RHICmdList, Scene, View, DefaultLightMaterial, *DirectionalLight, ChannelIdx, bInlineReflectionAndSky, MobileSSRQuality);
 		}
 	}
+	return NumLights;
 }
 
 template<uint32 LightingChannel, bool bWithStencilCulling, bool bEnableShadingModelSupport>
@@ -1082,9 +1083,9 @@ void MobileDeferredShadingPass(
 	DefaultMaterial.Material = DefaultMaterial.MaterialProxy->GetMaterialNoFallback(ERHIFeatureLevel::ES3_1);
 	check(DefaultMaterial.Material);
 
-	RenderDirectionalLights(RHICmdList, Scene, View, DefaultMaterial, MobileSSRQuality);
+	int NumDirLights = RenderDirectionalLights(RHICmdList, Scene, View, DefaultMaterial, MobileSSRQuality);
 	
-	const bool bMobileUseClusteredDeferredShading = UseClusteredDeferredShading(View.GetShaderPlatform());
+	const bool bMobileUseClusteredDeferredShading = UseClusteredDeferredShading(View.GetShaderPlatform()) && NumDirLights > 0;
 	if (!bMobileUseClusteredDeferredShading)
 	{
 		// Render non-clustered simple lights
@@ -1115,4 +1116,4 @@ void MobileDeferredShadingPass(
 		const FLightSceneInfo& LightSceneInfo = *SortedLight.LightSceneInfo;
 		RenderLocalLight(RHICmdList, Scene, View, LightSceneInfo, DefaultMaterial, VisibleLightInfos);
 	}
-}
\ No newline at end of file
+}
