Ray Tracing: Fix potential crash in material binding when disabling materials via cvar

Reported via UDN

diff --git a/Engine/Source/Runtime/Renderer/Private/RayTracing/RayTracingMaterialHitShaders.cpp b/Engine/Source/Runtime/Renderer/Private/RayTracing/RayTracingMaterialHitShaders.cpp
--- a/Engine/Source/Runtime/Renderer/Private/RayTracing/RayTracingMaterialHitShaders.cpp
+++ b/Engine/Source/Runtime/Renderer/Private/RayTracing/RayTracingMaterialHitShaders.cpp
@@ -709,8 +709,8 @@ void FDeferredShadingSceneRenderer::CreateMaterialRayTracingMaterialPipeline(
 	FRHIRayTracingShader* OpaqueShadowShader   = bIsPathTracing ? GetPathTracingDefaultOpaqueHitShader(View.ShaderMap) : GetRayTracingDefaultOpaqueShader(View.ShaderMap);
 	FRHIRayTracingShader* HiddenMaterialShader = bIsPathTracing ? GetPathTracingDefaultHiddenHitShader(View.ShaderMap) : GetRayTracingDefaultHiddenShader(View.ShaderMap);
 
-	FRHIRayTracingShader* OpaqueMeshDecalHitShader = GetDefaultOpaqueMeshDecalHitShader(View.ShaderMap);
-	FRHIRayTracingShader* HiddenMeshDecalHitShader = GetDefaultHiddenMeshDecalHitShader(View.ShaderMap);
+	FRHIRayTracingShader* OpaqueMeshDecalHitShader = bSupportMeshDecals ? GetDefaultOpaqueMeshDecalHitShader(View.ShaderMap) : nullptr;
+	FRHIRayTracingShader* HiddenMeshDecalHitShader = bSupportMeshDecals ? GetDefaultHiddenMeshDecalHitShader(View.ShaderMap) : nullptr;
 	
 	TArray<FRHIRayTracingShader*> RayTracingHitGroupLibrary;
 	if (bEnableMaterials)
@@ -723,15 +723,13 @@ void FDeferredShadingSceneRenderer::CreateMaterialRayTracingMaterialPipeline(
 		}
 	}
 
-	FRHIRayTracingShader* RequiredHitShaders[] =
-	{
-		OpaqueShadowShader,
-		HiddenMaterialShader
-	};
+	FRHIRayTracingShader* RequiredHitShaders[] = { OpaqueShadowShader, HiddenMaterialShader };
+	FRHIRayTracingShader* RequiredHitDecalShaders[] = { OpaqueMeshDecalHitShader, HiddenMeshDecalHitShader };
 
-	for (FRHIRayTracingShader* Shader : RequiredHitShaders)
+	RayTracingHitGroupLibrary.Append(RequiredHitShaders);
+	if (bSupportMeshDecals)
 	{
-		RayTracingHitGroupLibrary.Add(Shader);
+		RayTracingHitGroupLibrary.Append(RequiredHitDecalShaders);
 	}
 
 	Initializer.SetHitGroupTable(RayTracingHitGroupLibrary);
@@ -769,6 +767,7 @@ void FDeferredShadingSceneRenderer::CreateMaterialRayTracingMaterialPipeline(
 		&& View.ViewState
 		&& bCompatiblePipelineSignatures
 		&& PipelineContainsHitShaders(FallbackPipelineState, RequiredHitShaders)
+		&& (!bSupportMeshDecals || PipelineContainsHitShaders(FallbackPipelineState, RequiredHitDecalShaders))
 		&& FindRayTracingMissShaderIndex(FallbackPipelineState, DefaultMissShader, false) != INDEX_NONE
 		&& (!bCallableShadersRequired || FindRayTracingCallableShaderIndex(FallbackPipelineState, DefaultCallableShader, false) != INDEX_NONE))
 	{
@@ -831,14 +830,11 @@ void FDeferredShadingSceneRenderer::SetupMaterialRaytracingHitGroupBindings(FRDG
 	FRHIRayTracingShader* OpaqueShadowShader = bIsPathTracing ? GetPathTracingDefaultOpaqueHitShader(View.ShaderMap) : GetRayTracingDefaultOpaqueShader(View.ShaderMap);
 	FRHIRayTracingShader* HiddenMaterialShader = bIsPathTracing ? GetPathTracingDefaultHiddenHitShader(View.ShaderMap) : GetRayTracingDefaultHiddenShader(View.ShaderMap);
 
-	FRHIRayTracingShader* OpaqueMeshDecalHitShader = GetDefaultOpaqueMeshDecalHitShader(View.ShaderMap);
-	FRHIRayTracingShader* HiddenMeshDecalHitShader = GetDefaultHiddenMeshDecalHitShader(View.ShaderMap);
-
 	const int32 OpaqueShadowMaterialIndex = FindRayTracingHitGroupIndex(PipelineState, OpaqueShadowShader, true);
 	const int32 HiddenMaterialIndex = FindRayTracingHitGroupIndex(PipelineState, HiddenMaterialShader, true);
 
-	const int32 OpaqueMeshDecalHitGroupIndex = bSupportMeshDecals ? FindRayTracingHitGroupIndex(PipelineState, OpaqueMeshDecalHitShader, true) : INDEX_NONE;
-	const int32 HiddenMeshDecalHitGroupIndex = bSupportMeshDecals ? FindRayTracingHitGroupIndex(PipelineState, HiddenMeshDecalHitShader, true) : INDEX_NONE;
+	const int32 OpaqueMeshDecalHitGroupIndex = bSupportMeshDecals ? FindRayTracingHitGroupIndex(PipelineState, GetDefaultOpaqueMeshDecalHitShader(View.ShaderMap), true) : INDEX_NONE;
+	const int32 HiddenMeshDecalHitGroupIndex = bSupportMeshDecals ? FindRayTracingHitGroupIndex(PipelineState, GetDefaultHiddenMeshDecalHitShader(View.ShaderMap), true) : INDEX_NONE;
 
 	// material hit groups
 	{
