Fix CIS warning.

diff --git a/Engine/Source/Runtime/Renderer/Private/HZB.cpp b/Engine/Source/Runtime/Renderer/Private/HZB.cpp
--- a/Engine/Source/Runtime/Renderer/Private/HZB.cpp
+++ b/Engine/Source/Runtime/Renderer/Private/HZB.cpp
@@ -93,7 +93,7 @@ FHZBParameters GetHZBParameters(FRDGBuilder& GraphBuilder, const FViewInfo& View
 	const FRDGTextureRef HZB = bUseCurrent ? View.HZB : GraphBuilder.RegisterExternalTexture(View.PrevViewInfo.HZB);
 	const FIntRect ViewRect = bUseCurrent ? View.ViewRect : View.PrevViewInfo.ViewRect;
 
-	FHZBParameters Out = InitHZBCommonParameter(GraphBuilder, &View, ViewRect, HZB->Desc.Extent);
+	FHZBParameters Out = InitHZBCommonParameter(GraphBuilder, &View, ViewRect, HZB ? HZB->Desc.Extent : FIntPoint(1,1));
 	if (HZB)
 	{
 		Out.HZBTexture = HZB;
