Small refactor of render target creation code.
Centralised logic to create render targets instead of having it repeated in several places. This should make it easier to create render targets that will adapt to both 2D and stereo rendering.
Also added param in SceneTextureConfig for storing the number of layers of the render target texture array for XR.

diff --git a/Engine/Source/Runtime/Engine/Public/SceneTexturesConfig.h b/Engine/Source/Runtime/Engine/Public/SceneTexturesConfig.h
--- a/Engine/Source/Runtime/Engine/Public/SceneTexturesConfig.h
+++ b/Engine/Source/Runtime/Engine/Public/SceneTexturesConfig.h
@@ -217,6 +217,9 @@ struct FSceneTexturesConfig
     // True if we require an alpha channel for scene color
     bool bRequiresAlphaChannel = false;
 
+	// (Mobile XR) Number of render target layers in mobile multi view.
+	uint32 MobileMultiViewRenderTargetNumLayers = 2;
+
 private:
 	static ENGINE_API FSceneTexturesConfig GlobalInstance;
 };
diff --git a/Engine/Source/Runtime/RenderCore/Public/RenderGraphDefinitions.h b/Engine/Source/Runtime/RenderCore/Public/RenderGraphDefinitions.h
--- a/Engine/Source/Runtime/RenderCore/Public/RenderGraphDefinitions.h
+++ b/Engine/Source/Runtime/RenderCore/Public/RenderGraphDefinitions.h
@@ -706,6 +706,21 @@ struct FRDGTextureDesc : public FRHITextureDesc
 		return FRDGTextureDesc(ETextureDimension::TextureCubeArray, Flags, Format, ClearValue, { (int32)Size, (int32)Size }, Depth, ArraySize, NumMips, NumSamples, ExtData);
 	}
 
+	static FRDGTextureDesc CreateRenderTargetTextureDesc(
+		FIntPoint Size, EPixelFormat Format
+		, FClearValueBinding ClearValue
+		, ETextureCreateFlags Flags
+		, const bool bRequireMultiView
+		, uint16 MobileMultiViewRenderTargetNumLayers = 2)
+	{
+		if (bRequireMultiView)
+		{
+			return FRDGTextureDesc::Create2DArray(Size, Format, ClearValue, Flags, MobileMultiViewRenderTargetNumLayers);
+		}
+
+		return FRDGTextureDesc::Create2D(Size, Format, ClearValue, Flags);
+	}
+
 	FRDGTextureDesc() = default;
 	FRDGTextureDesc(
 		ETextureDimension   InDimension
diff --git a/Engine/Source/Runtime/Renderer/Internal/CustomDepthRendering.h b/Engine/Source/Runtime/Renderer/Internal/CustomDepthRendering.h
--- a/Engine/Source/Runtime/Renderer/Internal/CustomDepthRendering.h
+++ b/Engine/Source/Runtime/Renderer/Internal/CustomDepthRendering.h
@@ -21,7 +21,7 @@ extern ECustomDepthPassLocation GetCustomDepthPassLocation(EShaderPlatform Platf
 
 struct FCustomDepthTextures
 {
-	static FCustomDepthTextures Create(FRDGBuilder& GraphBuilder, FIntPoint CustomDepthExtent, EShaderPlatform ShaderPlatform, bool bRequireMultiView);
+	static FCustomDepthTextures Create(FRDGBuilder& GraphBuilder, FIntPoint CustomDepthExtent, EShaderPlatform ShaderPlatform, bool bRequireMultiView, uint16 MobileMultiViewRenderTargetNumLayers);
 	bool IsValid() const
 	{
 		return Depth != nullptr;
diff --git a/Engine/Source/Runtime/Renderer/Private/CustomDepthRendering.cpp b/Engine/Source/Runtime/Renderer/Private/CustomDepthRendering.cpp
--- a/Engine/Source/Runtime/Renderer/Private/CustomDepthRendering.cpp
+++ b/Engine/Source/Runtime/Renderer/Private/CustomDepthRendering.cpp
@@ -60,7 +60,7 @@ bool IsCustomDepthPassWritingStencil()
 	return GetCustomDepthMode() == ECustomDepthMode::EnabledWithStencil;
 }
 
-FCustomDepthTextures FCustomDepthTextures::Create(FRDGBuilder& GraphBuilder, FIntPoint CustomDepthExtent, EShaderPlatform ShaderPlatform, bool bRequireMultiView)
+FCustomDepthTextures FCustomDepthTextures::Create(FRDGBuilder& GraphBuilder, FIntPoint CustomDepthExtent, EShaderPlatform ShaderPlatform, bool bRequireMultiView, uint16 MobileMultiViewRenderTargetNumLayers)
 {
 	if (!IsCustomDepthPassEnabled())
 	{
@@ -82,9 +82,7 @@ FCustomDepthTextures FCustomDepthTextures::Create(FRDGBuilder& GraphBuilder, FIn
 	{
 		CreateFlags |= TexCreate_NoFastClear;
 	}
-	FRDGTextureDesc CustomDepthDesc(bRequireMultiView ?
-		FRDGTextureDesc::Create2DArray(CustomDepthExtent, PF_DepthStencil, FClearValueBinding::DepthFar, CreateFlags, 2) :
-		FRDGTextureDesc::Create2D(CustomDepthExtent, PF_DepthStencil, FClearValueBinding::DepthFar, CreateFlags));
+	FRDGTextureDesc CustomDepthDesc = FRDGTextureDesc::CreateRenderTargetTextureDesc(CustomDepthExtent, PF_DepthStencil, FClearValueBinding::DepthFar, CreateFlags, bRequireMultiView, MobileMultiViewRenderTargetNumLayers);
 
 	CustomDepthTextures.Depth = GraphBuilder.CreateTexture(CustomDepthDesc, TEXT("CustomDepth"));
 
diff --git a/Engine/Source/Runtime/Renderer/Private/SceneTextures.cpp b/Engine/Source/Runtime/Renderer/Private/SceneTextures.cpp
--- a/Engine/Source/Runtime/Renderer/Private/SceneTextures.cpp
+++ b/Engine/Source/Runtime/Renderer/Private/SceneTextures.cpp
@@ -450,10 +450,7 @@ void FMinimalSceneTextures::InitializeViewFamily(FRDGBuilder& GraphBuilder, FVie
 	}
 	else
 	{
-		// TODO: Array-size could be values > 2, on upcoming XR devices. This should be part of the config.
-		FRDGTextureDesc Desc(Config.bRequireMultiView ?
-							 FRDGTextureDesc::Create2DArray(SceneTextures.Config.Extent, PF_DepthStencil, Config.DepthClearValue, Config.DepthCreateFlags, 2) :
-							 FRDGTextureDesc::Create2D(SceneTextures.Config.Extent, PF_DepthStencil, Config.DepthClearValue, Config.DepthCreateFlags));
+		FRDGTextureDesc Desc = FRDGTextureDesc::CreateRenderTargetTextureDesc(SceneTextures.Config.Extent, PF_DepthStencil, Config.DepthClearValue, Config.DepthCreateFlags, Config.bRequireMultiView, Config.MobileMultiViewRenderTargetNumLayers);
 		Desc.NumSamples = Config.NumSamples;
 		SceneTextures.Depth = GraphBuilder.CreateTexture(Desc, TEXT("SceneDepthZ"));
 
@@ -482,16 +479,13 @@ void FMinimalSceneTextures::InitializeViewFamily(FRDGBuilder& GraphBuilder, FVie
 		const ETextureCreateFlags sRGBFlag = (bIsMobilePlatform && IsMobileColorsRGB()) ? TexCreate_SRGB : TexCreate_None;
 
 		// Create the scene color.
-		// TODO: Array-size could be values > 2, on upcoming XR devices. This should be part of the config.
-		FRDGTextureDesc Desc(Config.bRequireMultiView ?
-							 FRDGTextureDesc::Create2DArray(Config.Extent, Config.ColorFormat, Config.ColorClearValue, Config.ColorCreateFlags, 2) :
-							 FRDGTextureDesc::Create2D(Config.Extent, Config.ColorFormat, Config.ColorClearValue, Config.ColorCreateFlags));
+		FRDGTextureDesc Desc = FRDGTextureDesc::CreateRenderTargetTextureDesc(Config.Extent, Config.ColorFormat, Config.ColorClearValue, Config.ColorCreateFlags, Config.bRequireMultiView, Config.MobileMultiViewRenderTargetNumLayers);
 		Desc.NumSamples = Config.NumSamples;
 		SceneTextures.Color = CreateTextureMSAA(GraphBuilder, Desc, TEXT("SceneColorMS"), TEXT("SceneColor"), GFastVRamConfig.SceneColor | sRGBFlag);
 	}
 
 	// Custom Depth
-	SceneTextures.CustomDepth = FCustomDepthTextures::Create(GraphBuilder, Config.Extent, Config.ShaderPlatform, Config.bRequireMultiView);
+	SceneTextures.CustomDepth = FCustomDepthTextures::Create(GraphBuilder, Config.Extent, Config.ShaderPlatform, Config.bRequireMultiView, Config.MobileMultiViewRenderTargetNumLayers);
 
 	SceneTextures.bIsSceneTexturesInitialized = true;
 }
@@ -551,9 +545,7 @@ FRDGTextureRef FMinimalSceneTextures::FindOrAddUserSceneTexture(FRDGBuilder& Gra
 
 		FIntPoint Extent = OutputRect.Size();
 
-		FRDGTextureDesc Desc(Config.bRequireMultiView ?
-			FRDGTextureDesc::Create2DArray(Extent, Config.ColorFormat, Config.ColorClearValue, Config.ColorCreateFlags, 2) :
-			FRDGTextureDesc::Create2D(Extent, Config.ColorFormat, Config.ColorClearValue, Config.ColorCreateFlags));
+		FRDGTextureDesc Desc = FRDGTextureDesc::CreateRenderTargetTextureDesc(Extent, Config.ColorFormat, Config.ColorClearValue, Config.ColorCreateFlags, Config.bRequireMultiView, Config.MobileMultiViewRenderTargetNumLayers);
 
 #if !(UE_BUILD_SHIPPING || UE_BUILD_TEST)
 		// Generate a heap allocated debug name for the resource, required to make UserSceneTextures accessible in VisualizeTexture,
@@ -692,41 +684,31 @@ void FSceneTextures::InitializeViewFamily(FRDGBuilder& GraphBuilder, FViewFamily
 
 		if (Bindings.GBufferA.Index >= 0)
 		{
-			const FRDGTextureDesc Desc(Config.bRequireMultiView ?
-				FRDGTextureDesc::Create2DArray(Config.Extent, Bindings.GBufferA.Format, FClearValueBinding::Transparent, Bindings.GBufferA.Flags | FlagsToAdd | GFastVRamConfig.GBufferA, 2) :
-				FRDGTextureDesc::Create2D(Config.Extent, Bindings.GBufferA.Format, FClearValueBinding::Transparent, Bindings.GBufferA.Flags | FlagsToAdd | GFastVRamConfig.GBufferA));
+			const FRDGTextureDesc Desc = FRDGTextureDesc::CreateRenderTargetTextureDesc(Config.Extent, Bindings.GBufferA.Format, FClearValueBinding::Transparent, Bindings.GBufferA.Flags | FlagsToAdd | GFastVRamConfig.GBufferA, Config.bRequireMultiView, Config.MobileMultiViewRenderTargetNumLayers);
 			SceneTextures.GBufferA = GraphBuilder.CreateTexture(Desc, TEXT("GBufferA"));
 		}
 
 		if (Bindings.GBufferB.Index >= 0)
 		{
-			const FRDGTextureDesc Desc(Config.bRequireMultiView ?
-				FRDGTextureDesc::Create2DArray(Config.Extent, Bindings.GBufferB.Format, FClearValueBinding::Transparent, Bindings.GBufferB.Flags | FlagsToAdd | GFastVRamConfig.GBufferB, 2) :
-				FRDGTextureDesc::Create2D(Config.Extent, Bindings.GBufferB.Format, FClearValueBinding::Transparent, Bindings.GBufferB.Flags | FlagsToAdd | GFastVRamConfig.GBufferB));
+			const FRDGTextureDesc Desc = FRDGTextureDesc::CreateRenderTargetTextureDesc(Config.Extent, Bindings.GBufferB.Format, FClearValueBinding::Transparent, Bindings.GBufferB.Flags | FlagsToAdd | GFastVRamConfig.GBufferB, Config.bRequireMultiView, Config.MobileMultiViewRenderTargetNumLayers);
 			SceneTextures.GBufferB = GraphBuilder.CreateTexture(Desc, TEXT("GBufferB"));
 		}
 
 		if (Bindings.GBufferC.Index >= 0)
 		{
-			const FRDGTextureDesc Desc(Config.bRequireMultiView ?
-				FRDGTextureDesc::Create2DArray(Config.Extent, Bindings.GBufferC.Format, FClearValueBinding::Transparent, Bindings.GBufferC.Flags | FlagsToAdd | GFastVRamConfig.GBufferC, 2) :
-				FRDGTextureDesc::Create2D(Config.Extent, Bindings.GBufferC.Format, FClearValueBinding::Transparent, Bindings.GBufferC.Flags | FlagsToAdd | GFastVRamConfig.GBufferC));
+			const FRDGTextureDesc Desc = FRDGTextureDesc::CreateRenderTargetTextureDesc(Config.Extent, Bindings.GBufferC.Format, FClearValueBinding::Transparent, Bindings.GBufferC.Flags | FlagsToAdd | GFastVRamConfig.GBufferC, Config.bRequireMultiView, Config.MobileMultiViewRenderTargetNumLayers);
 			SceneTextures.GBufferC = GraphBuilder.CreateTexture(Desc, TEXT("GBufferC"));
 		}
 
 		if (Bindings.GBufferD.Index >= 0)
 		{
-			const FRDGTextureDesc Desc(Config.bRequireMultiView ?
-				FRDGTextureDesc::Create2DArray(Config.Extent, Bindings.GBufferD.Format, FClearValueBinding::Transparent, Bindings.GBufferD.Flags | FlagsToAdd | GFastVRamConfig.GBufferD, 2) :
-				FRDGTextureDesc::Create2D(Config.Extent, Bindings.GBufferD.Format, FClearValueBinding::Transparent, Bindings.GBufferD.Flags | FlagsToAdd | GFastVRamConfig.GBufferD));
+			const FRDGTextureDesc Desc = FRDGTextureDesc::CreateRenderTargetTextureDesc(Config.Extent, Bindings.GBufferD.Format, FClearValueBinding::Transparent, Bindings.GBufferD.Flags | FlagsToAdd | GFastVRamConfig.GBufferD, Config.bRequireMultiView, Config.MobileMultiViewRenderTargetNumLayers);
 			SceneTextures.GBufferD = GraphBuilder.CreateTexture(Desc, TEXT("GBufferD"));
 		}
 
 		if (Bindings.GBufferE.Index >= 0)
 		{
-			const FRDGTextureDesc Desc(Config.bRequireMultiView ?
-				FRDGTextureDesc::Create2DArray(Config.Extent, Bindings.GBufferE.Format, FClearValueBinding::Transparent, Bindings.GBufferE.Flags | FlagsToAdd | GFastVRamConfig.GBufferE, 2) :
-				FRDGTextureDesc::Create2D(Config.Extent, Bindings.GBufferE.Format, FClearValueBinding::Transparent, Bindings.GBufferE.Flags | FlagsToAdd | GFastVRamConfig.GBufferE));
+			const FRDGTextureDesc Desc = FRDGTextureDesc::CreateRenderTargetTextureDesc(Config.Extent, Bindings.GBufferE.Format, FClearValueBinding::Transparent, Bindings.GBufferE.Flags | FlagsToAdd | GFastVRamConfig.GBufferE, Config.bRequireMultiView, Config.MobileMultiViewRenderTargetNumLayers);
 			SceneTextures.GBufferE = GraphBuilder.CreateTexture(Desc, TEXT("GBufferE"));
 		}
 
@@ -735,9 +717,7 @@ void FSceneTextures::InitializeViewFamily(FRDGBuilder& GraphBuilder, FViewFamily
 		{
 			ETextureCreateFlags GBufferFCreateFlags;
 			EPixelFormat GBufferFPixelFormat = GetGBufferFFormatAndCreateFlags(GBufferFCreateFlags);
-			const FRDGTextureDesc Desc(Config.bRequireMultiView ?
-				FRDGTextureDesc::Create2DArray(Config.Extent, GBufferFPixelFormat, FClearValueBinding({ 0.5f, 0.5f, 0.5f, 0.5f }), GBufferFCreateFlags | FlagsToAdd, 2) :
-				FRDGTextureDesc::Create2D(Config.Extent, GBufferFPixelFormat, FClearValueBinding({ 0.5f, 0.5f, 0.5f, 0.5f }), GBufferFCreateFlags | FlagsToAdd));
+			const FRDGTextureDesc Desc = FRDGTextureDesc::CreateRenderTargetTextureDesc(Config.Extent, GBufferFPixelFormat, FClearValueBinding({ 0.5f, 0.5f, 0.5f, 0.5f }), GBufferFCreateFlags | FlagsToAdd, Config.bRequireMultiView, Config.MobileMultiViewRenderTargetNumLayers);
 			SceneTextures.GBufferF = GraphBuilder.CreateTexture(Desc, TEXT("GBufferF"));
 		}
 	}
@@ -756,9 +736,7 @@ void FSceneTextures::InitializeViewFamily(FRDGBuilder& GraphBuilder, FViewFamily
 		}
 
 		EPixelFormat DepthAuxFormat = Config.bPreciseDepthAux ? PF_R32_FLOAT : PF_R16F;
-		FRDGTextureDesc Desc = Config.bRequireMultiView ? 
-			FRDGTextureDesc::Create2DArray(Config.Extent, DepthAuxFormat, FClearValueBinding(FarDepthColor), TexCreate_RenderTargetable | TexCreate_ShaderResource | TexCreate_InputAttachmentRead | MemorylessFlag, 2) :
-			FRDGTextureDesc::Create2D(Config.Extent, DepthAuxFormat, FClearValueBinding(FarDepthColor), TexCreate_RenderTargetable | TexCreate_ShaderResource | TexCreate_InputAttachmentRead| MemorylessFlag);
+		FRDGTextureDesc Desc = FRDGTextureDesc::CreateRenderTargetTextureDesc(Config.Extent, DepthAuxFormat, FClearValueBinding(FarDepthColor), TexCreate_RenderTargetable | TexCreate_ShaderResource | TexCreate_InputAttachmentRead | MemorylessFlag, Config.bRequireMultiView, Config.MobileMultiViewRenderTargetNumLayers);
 		Desc.NumSamples = Config.NumSamples;
 		SceneTextures.DepthAux = CreateTextureMSAA(GraphBuilder, Desc, TEXT("SceneDepthAuxMS"), TEXT("SceneDepthAux"));
 	}
diff --git a/Engine/Source/Runtime/Renderer/Private/VelocityRendering.cpp b/Engine/Source/Runtime/Renderer/Private/VelocityRendering.cpp
--- a/Engine/Source/Runtime/Renderer/Private/VelocityRendering.cpp
+++ b/Engine/Source/Runtime/Renderer/Private/VelocityRendering.cpp
@@ -403,9 +403,7 @@ ETextureCreateFlags FVelocityRendering::GetCreateFlags(EShaderPlatform ShaderPla
 
 FRDGTextureDesc FVelocityRendering::GetRenderTargetDesc(EShaderPlatform ShaderPlatform, FIntPoint Extent, const bool bRequireMultiView)
 {
-	return bRequireMultiView ?
-		FRDGTextureDesc::Create2DArray(Extent, GetFormat(ShaderPlatform), FClearValueBinding::Transparent, GetCreateFlags(ShaderPlatform), 2) :
-		FRDGTextureDesc::Create2D(Extent, GetFormat(ShaderPlatform), FClearValueBinding::Transparent, GetCreateFlags(ShaderPlatform));
+	return FRDGTextureDesc::CreateRenderTargetTextureDesc(Extent, GetFormat(ShaderPlatform), FClearValueBinding::Transparent, GetCreateFlags(ShaderPlatform), bRequireMultiView);
 }
 
 bool FVelocityRendering::IsVelocityPassSupported(EShaderPlatform ShaderPlatform)
