PR #12762: Fixed Primitive LastRenderTime not being updated if Occlusion is disabled

#rnx

diff --git a/Engine/Source/Runtime/Renderer/Private/SceneVisibility.cpp b/Engine/Source/Runtime/Renderer/Private/SceneVisibility.cpp
--- a/Engine/Source/Runtime/Renderer/Private/SceneVisibility.cpp
+++ b/Engine/Source/Runtime/Renderer/Private/SceneVisibility.cpp
@@ -3610,7 +3610,7 @@ void FVisibilityViewPacket::BeginInitVisibility()
 	// Allocate the view's visibility maps.
 	View.PrimitiveVisibilityMap.Init(!bShouldVisibilityCull, Scene.Primitives.Num());
 	View.PrimitiveRayTracingVisibilityMap.Init(false, Scene.Primitives.Num());
-	View.PrimitiveDefinitelyUnoccludedMap.Init(false, Scene.Primitives.Num());
+	View.PrimitiveDefinitelyUnoccludedMap.Init(!GOcclusionCullEnabled, Scene.Primitives.Num());
 	View.PotentiallyFadingPrimitiveMap.Init(false, Scene.Primitives.Num());
 	View.PrimitiveFadeUniformBuffers.AddZeroed(Scene.Primitives.Num());
 	View.PrimitiveFadeUniformBufferMap.Init(false, Scene.Primitives.Num());
