[HWRT] Decouple FRayTracingShaderBindingTable from ERayTracingSceneLayer.
- using it's own enum ERayTracingShaderBindingLayer.

diff --git a/Engine/Plugins/Experimental/GPULightmass/Source/GPULightmass/Private/LightmapRenderer.cpp b/Engine/Plugins/Experimental/GPULightmass/Source/GPULightmass/Private/LightmapRenderer.cpp
--- a/Engine/Plugins/Experimental/GPULightmass/Source/GPULightmass/Private/LightmapRenderer.cpp
+++ b/Engine/Plugins/Experimental/GPULightmass/Source/GPULightmass/Private/LightmapRenderer.cpp
@@ -717,7 +717,7 @@ void FCachedRayTracingSceneData::SetupFromSceneRenderState(FSceneRenderState& Sc
 
 			if (!RayTracingMaskAndStatus.bAllSegmentsUnlit)
 			{
-				FRayTracingSBTAllocation* SBTAllocation = RaytracingSBT.AllocateStaticRange(ERayTracingSceneLayerMask::Base, MeshBatches.Num());
+				FRayTracingSBTAllocation* SBTAllocation = RaytracingSBT.AllocateStaticRange(ERayTracingShaderBindingLayerMask::Base, MeshBatches.Num());
 				StaticSBTAllocations.Add(SBTAllocation);
 
 				RayTracingNumSegmentsPerLOD[LODIndex] += MeshBatches.Num();
@@ -727,7 +727,7 @@ void FCachedRayTracingSceneData::SetupFromSceneRenderState(FSceneRenderState& Sc
 				RayTracingInstance.GeometryRHI = Instance.RenderData->RayTracingProxy->LODs[LODIndexToUse].RayTracingGeometry->GetRHI();
 				RayTracingInstance.Transforms = MakeArrayView(&Instance.LocalToWorld, 1);
 				RayTracingInstance.NumTransforms = 1;
-				RayTracingInstance.InstanceContributionToHitGroupIndex = SBTAllocation->GetInstanceContributionToHitGroupIndex(ERayTracingSceneLayer::Base);
+				RayTracingInstance.InstanceContributionToHitGroupIndex = SBTAllocation->GetInstanceContributionToHitGroupIndex(ERayTracingShaderBindingLayer::Base);
 				RayTracingInstance.DefaultUserData = (uint32)StaticMeshIndex;
 				RayTracingInstance.Mask = RayTracingMaskAndStatus.InstanceMask;
 				if (RayTracingMaskAndStatus.bAllSegmentsOpaque)
@@ -762,7 +762,7 @@ void FCachedRayTracingSceneData::SetupFromSceneRenderState(FSceneRenderState& Sc
 
 			if (!RayTracingMaskAndStatus.bAllSegmentsUnlit)
 			{
-				FRayTracingSBTAllocation* SBTAllocation = RaytracingSBT.AllocateStaticRange(ERayTracingSceneLayerMask::Base, MeshBatches.Num());
+				FRayTracingSBTAllocation* SBTAllocation = RaytracingSBT.AllocateStaticRange(ERayTracingShaderBindingLayerMask::Base, MeshBatches.Num());
 				StaticSBTAllocations.Add(SBTAllocation);
 
 				RayTracingNumSegmentsPerLOD[LODIndex] += MeshBatches.Num();
@@ -784,7 +784,7 @@ void FCachedRayTracingSceneData::SetupFromSceneRenderState(FSceneRenderState& Sc
 				RayTracingInstance.Transforms = NewTransforms;
 				RayTracingInstance.NumTransforms = NumInstances;
 
-				RayTracingInstance.InstanceContributionToHitGroupIndex = SBTAllocation->GetInstanceContributionToHitGroupIndex(ERayTracingSceneLayer::Base);
+				RayTracingInstance.InstanceContributionToHitGroupIndex = SBTAllocation->GetInstanceContributionToHitGroupIndex(ERayTracingShaderBindingLayer::Base);
 
 				RayTracingInstance.DefaultUserData = (uint32)(Scene.StaticMeshInstanceRenderStates.Elements.Num() + InstanceGroupIndex);
 				RayTracingInstance.Mask = RayTracingMaskAndStatus.InstanceMask;
@@ -1125,14 +1125,14 @@ bool FSceneRenderState::SetupRayTracingScene(FRDGBuilder& GraphBuilder, FSceneUn
 						ensure(MeshBatches.Num() == 1);
 
 						uint32 SegmentCount = MeshBatches.Num();
-						FRayTracingSBTAllocation* SBTAllocation = RayTracingSBT.AllocateDynamicRange(ERayTracingSceneLayerMask::Base, SegmentCount);
+						FRayTracingSBTAllocation* SBTAllocation = RayTracingSBT.AllocateDynamicRange(ERayTracingShaderBindingLayerMask::Base, SegmentCount);
 						RayTracingNumSegments += SegmentCount;
 
 						FRayTracingGeometryInstance& RayTracingInstance = RayTracingGeometryInstances[InstanceIndex];
 						RayTracingInstance.GeometryRHI = Landscape.SectionRayTracingStates[SubSectionIdx]->Geometry.GetRHI();
 						RayTracingInstance.Transforms = MakeArrayView(LandscapeTransforms.Emplace_GetRef(MakeUnique<FMatrix>(Landscape.LocalToWorld)).Get(), 1);
 						RayTracingInstance.NumTransforms = 1;
-						RayTracingInstance.InstanceContributionToHitGroupIndex = SBTAllocation->GetInstanceContributionToHitGroupIndex(ERayTracingSceneLayer::Base);
+						RayTracingInstance.InstanceContributionToHitGroupIndex = SBTAllocation->GetInstanceContributionToHitGroupIndex(ERayTracingShaderBindingLayer::Base);
 						RayTracingInstance.DefaultUserData = (uint32)(StaticMeshInstanceRenderStates.Elements.Num() + InstanceGroupRenderStates.Elements.Num() + LandscapeIndex);
 
 						FLandscapeBatchElementParams& BatchElementParams = *(FLandscapeBatchElementParams*)MeshBatches[0].Elements[0].UserData;
diff --git a/Engine/Plugins/Experimental/GPULightmass/Source/GPULightmass/Private/Scene/Scene.h b/Engine/Plugins/Experimental/GPULightmass/Source/GPULightmass/Private/Scene/Scene.h
--- a/Engine/Plugins/Experimental/GPULightmass/Source/GPULightmass/Private/Scene/Scene.h
+++ b/Engine/Plugins/Experimental/GPULightmass/Source/GPULightmass/Private/Scene/Scene.h
@@ -49,7 +49,7 @@ public:
 	{
 		check(GeometrySegmentIndex == RayTracingMeshCommand.GeometrySegmentIndex);
 		bool bHidden = false;
-		const uint32 RecordIndex = SBTAllocation->GetRecordIndex(ERayTracingSceneLayer::Base, RayTracingMeshCommand.GeometrySegmentIndex);
+		const uint32 RecordIndex = SBTAllocation->GetRecordIndex(ERayTracingShaderBindingLayer::Base, RayTracingMeshCommand.GeometrySegmentIndex);
 		FRayTracingShaderBindingData DirtyShaderBinding(&RayTracingMeshCommand, RayTracingGeometry, RecordIndex, ERayTracingLocalShaderBindingType::Transient, bHidden);
 		DirtyShaderBindingsStorage.Add(DirtyShaderBinding);
 		check(DirtyShaderBinding.RayTracingMeshCommand);
diff --git a/Engine/Source/Runtime/Renderer/Private/RayTracing/RayTracing.cpp b/Engine/Source/Runtime/Renderer/Private/RayTracing/RayTracing.cpp
--- a/Engine/Source/Runtime/Renderer/Private/RayTracing/RayTracing.cpp
+++ b/Engine/Source/Runtime/Renderer/Private/RayTracing/RayTracing.cpp
@@ -867,14 +867,14 @@ namespace RayTracing
 						RayTracingInstance.BaseInstanceSceneDataOffset = InstanceSceneDataOffset;
 					}
 
-					ERayTracingSceneLayerMask ActiveLayers = ERayTracingSceneLayerMask::None;
+					ERayTracingShaderBindingLayerMask ActiveLayers = ERayTracingShaderBindingLayerMask::None;
 					if (bNeedMainInstance)
 					{
-						EnumAddFlags(ActiveLayers, ERayTracingSceneLayerMask::Base);
+						EnumAddFlags(ActiveLayers, ERayTracingShaderBindingLayerMask::Base);
 					}
 					if (bNeedDecalInstance)
 					{
-						EnumAddFlags(ActiveLayers, ERayTracingSceneLayerMask::Decals);
+						EnumAddFlags(ActiveLayers, ERayTracingShaderBindingLayerMask::Decals);
 					}
 
 					FRayTracingSBTAllocation* SBTAllocation = RayTracingSBT.AllocateDynamicRange(ActiveLayers, Instance.GetMaterials().Num());
@@ -882,7 +882,7 @@ namespace RayTracing
 					{
 						RayTracingScene.NumSegments += Instance.GetMaterials().Num();
 
-						RayTracingInstance.InstanceContributionToHitGroupIndex = SBTAllocation->GetInstanceContributionToHitGroupIndex(ERayTracingSceneLayer::Base);
+						RayTracingInstance.InstanceContributionToHitGroupIndex = SBTAllocation->GetInstanceContributionToHitGroupIndex(ERayTracingShaderBindingLayer::Base);
 						RayTracingScene.AddInstance(RayTracingInstance, ERayTracingSceneLayer::Base, SceneProxy, true, Geometry->GetGeometryHandle());
 					}
 
@@ -891,7 +891,7 @@ namespace RayTracing
 						RayTracingScene.NumSegments += Instance.GetMaterials().Num();
 
 						FRayTracingGeometryInstance DecalRayTracingInstance = RayTracingInstance;
-						DecalRayTracingInstance.InstanceContributionToHitGroupIndex = SBTAllocation->GetInstanceContributionToHitGroupIndex(ERayTracingSceneLayer::Decals);
+						DecalRayTracingInstance.InstanceContributionToHitGroupIndex = SBTAllocation->GetInstanceContributionToHitGroupIndex(ERayTracingShaderBindingLayer::Decals);
 						RayTracingScene.AddInstance(MoveTemp(DecalRayTracingInstance), ERayTracingSceneLayer::Decals, SceneProxy, true, Geometry->GetGeometryHandle());
 					}
 
@@ -1536,7 +1536,7 @@ namespace RayTracing
 				if (bNeedMainInstance)
 				{
 					RayTracingScene.NumSegments += SegmentCount;
-					RayTracingInstance.InstanceContributionToHitGroupIndex = RTLODData.SBTAllocation->GetInstanceContributionToHitGroupIndex(ERayTracingSceneLayer::Base);
+					RayTracingInstance.InstanceContributionToHitGroupIndex = RTLODData.SBTAllocation->GetInstanceContributionToHitGroupIndex(ERayTracingShaderBindingLayer::Base);
 					InstanceBatch.InstanceHandle = RayTracingScene.AddInstance(RayTracingInstance, ERayTracingSceneLayer::Base, SceneProxy, false, RelevantPrimitive.RayTracingGeometry->GetGeometryHandle());
 				}
 
@@ -1545,7 +1545,7 @@ namespace RayTracing
 				{
 					RayTracingScene.NumSegments += SegmentCount;
 					FRayTracingGeometryInstance DecalRayTracingInstance = RayTracingInstance;
-					DecalRayTracingInstance.InstanceContributionToHitGroupIndex = RTLODData.SBTAllocation->GetInstanceContributionToHitGroupIndex(ERayTracingSceneLayer::Decals);
+					DecalRayTracingInstance.InstanceContributionToHitGroupIndex = RTLODData.SBTAllocation->GetInstanceContributionToHitGroupIndex(ERayTracingShaderBindingLayer::Decals);
 					InstanceBatch.DecalInstanceHandle = RayTracingScene.AddInstance(MoveTemp(DecalRayTracingInstance), ERayTracingSceneLayer::Decals, SceneProxy, false, RelevantPrimitive.RayTracingGeometry->GetGeometryHandle());
 				}
 
@@ -1560,14 +1560,14 @@ namespace RayTracing
 						if (bNeedMainInstance)
 						{
 							const bool bHidden = MeshCommand.bDecal;
-							const uint32 RecordIndex = RTLODData.SBTAllocation->GetRecordIndex(ERayTracingSceneLayer::Base, MeshCommand.GeometrySegmentIndex);
+							const uint32 RecordIndex = RTLODData.SBTAllocation->GetRecordIndex(ERayTracingShaderBindingLayer::Base, MeshCommand.GeometrySegmentIndex);
 							FRayTracingShaderBindingData RTShaderBindingData(&MeshCommand, RelevantPrimitive.RayTracingGeometry->GetRHI(), RecordIndex, BindingType, bHidden);
 							VisibleShaderBindingData.Add(RTShaderBindingData);
 						}
 						if (bNeedDecalInstance)
 						{
 							const bool bHidden = !MeshCommand.bDecal;
-							const uint32 RecordIndex = RTLODData.SBTAllocation->GetRecordIndex(ERayTracingSceneLayer::Decals, MeshCommand.GeometrySegmentIndex);
+							const uint32 RecordIndex = RTLODData.SBTAllocation->GetRecordIndex(ERayTracingShaderBindingLayer::Decals, MeshCommand.GeometrySegmentIndex);
 							FRayTracingShaderBindingData RTShaderBindingData(&MeshCommand, RelevantPrimitive.RayTracingGeometry->GetRHI(), RecordIndex, BindingType, bHidden);
 							VisibleShaderBindingData.Add(RTShaderBindingData);
 						}
@@ -1644,7 +1644,7 @@ namespace RayTracing
 					const int32 InstanceIndexInRange = GatherContexts[RelevantPrimitive.ContextIndex].InstanceOffset + RelevantPrimitive.RelativeInstanceOffset;
 
 					FRayTracingGeometryInstance RayTracingInstance = *RelevantPrimitive.CachedRayTracingInstance;
-					RayTracingInstance.InstanceContributionToHitGroupIndex = RTLODData.SBTAllocation->GetInstanceContributionToHitGroupIndex(ERayTracingSceneLayer::Base);
+					RayTracingInstance.InstanceContributionToHitGroupIndex = RTLODData.SBTAllocation->GetInstanceContributionToHitGroupIndex(ERayTracingShaderBindingLayer::Base);
 					AddDebugRayTracingInstanceFlags(RayTracingInstance.Flags);
 					
 					// TODO: verify this
@@ -1657,7 +1657,7 @@ namespace RayTracing
 					const int32 DecalInstanceIndexInRange = GatherContexts[RelevantPrimitive.ContextIndex].DecalInstanceOffset + RelevantPrimitive.RelativeDecalInstanceOffset;
 
 					FRayTracingGeometryInstance DecalRayTracingInstance = *RelevantPrimitive.CachedRayTracingInstance;
-					DecalRayTracingInstance.InstanceContributionToHitGroupIndex = RTLODData.SBTAllocation->GetInstanceContributionToHitGroupIndex(ERayTracingSceneLayer::Decals);
+					DecalRayTracingInstance.InstanceContributionToHitGroupIndex = RTLODData.SBTAllocation->GetInstanceContributionToHitGroupIndex(ERayTracingShaderBindingLayer::Decals);
 					AddDebugRayTracingInstanceFlags(DecalRayTracingInstance.Flags);
 
 					// TODO: verify this
@@ -1680,7 +1680,7 @@ namespace RayTracing
 						if (bNeedMainInstance)
 						{
 							const bool bHidden = MeshCommand.bDecal;
-							const uint32 RecordIndex = RTLODData.SBTAllocation->GetRecordIndex(ERayTracingSceneLayer::Base, MeshCommand.GeometrySegmentIndex);
+							const uint32 RecordIndex = RTLODData.SBTAllocation->GetRecordIndex(ERayTracingShaderBindingLayer::Base, MeshCommand.GeometrySegmentIndex);
 							VisibleShaderBindingData[MainCommandOffset] = FRayTracingShaderBindingData(&MeshCommand, RelevantPrimitive.CachedRayTracingInstance->GeometryRHI, RecordIndex, BindingType, bHidden);
 							++MainCommandOffset;
 						}
@@ -1688,7 +1688,7 @@ namespace RayTracing
 						if (bNeedDecalInstance)
 						{
 							const bool bHidden = !MeshCommand.bDecal;
-							const uint32 RecordIndex = RTLODData.SBTAllocation->GetRecordIndex(ERayTracingSceneLayer::Decals, MeshCommand.GeometrySegmentIndex);
+							const uint32 RecordIndex = RTLODData.SBTAllocation->GetRecordIndex(ERayTracingShaderBindingLayer::Decals, MeshCommand.GeometrySegmentIndex);
 							VisibleShaderBindingData[DecalCommandOffset] = FRayTracingShaderBindingData(&MeshCommand, RelevantPrimitive.CachedRayTracingInstance->GeometryRHI, RecordIndex, BindingType, bHidden);
 							++DecalCommandOffset;
 						}
diff --git a/Engine/Source/Runtime/Renderer/Private/RayTracing/RayTracingMeshDrawCommands.cpp b/Engine/Source/Runtime/Renderer/Private/RayTracing/RayTracingMeshDrawCommands.cpp
--- a/Engine/Source/Runtime/Renderer/Private/RayTracing/RayTracingMeshDrawCommands.cpp
+++ b/Engine/Source/Runtime/Renderer/Private/RayTracing/RayTracingMeshDrawCommands.cpp
@@ -18,18 +18,18 @@ void FDynamicRayTracingMeshCommandContext::FinalizeCommand(FRayTracingMeshComman
 
 	if (SBTAllocation)
 	{
-		if (SBTAllocation->HasLayer(ERayTracingSceneLayer::Base))
+		if (SBTAllocation->HasLayer(ERayTracingShaderBindingLayer::Base))
 		{
 			const bool bHidden = RayTracingMeshCommand.bDecal;
-			const uint32 RecordIndex = SBTAllocation->GetRecordIndex(ERayTracingSceneLayer::Base, RayTracingMeshCommand.GeometrySegmentIndex);
+			const uint32 RecordIndex = SBTAllocation->GetRecordIndex(ERayTracingShaderBindingLayer::Base, RayTracingMeshCommand.GeometrySegmentIndex);
 			FRayTracingShaderBindingData ShaderBinding(&RayTracingMeshCommand, RayTracingGeometry, RecordIndex, ERayTracingLocalShaderBindingType::Transient, bHidden);
 			ShaderBindings.Add(ShaderBinding);
 		}
 
-		if (SBTAllocation->HasLayer(ERayTracingSceneLayer::Decals))
+		if (SBTAllocation->HasLayer(ERayTracingShaderBindingLayer::Decals))
 		{
 			const bool bHidden = !RayTracingMeshCommand.bDecal;
-			const uint32 RecordIndex = SBTAllocation->GetRecordIndex(ERayTracingSceneLayer::Decals, RayTracingMeshCommand.GeometrySegmentIndex);
+			const uint32 RecordIndex = SBTAllocation->GetRecordIndex(ERayTracingShaderBindingLayer::Decals, RayTracingMeshCommand.GeometrySegmentIndex);
 			FRayTracingShaderBindingData ShaderBinding(&RayTracingMeshCommand, RayTracingGeometry, RecordIndex, ERayTracingLocalShaderBindingType::Transient, bHidden);
 			ShaderBindings.Add(ShaderBinding);
 		}
diff --git a/Engine/Source/Runtime/Renderer/Private/RayTracing/RayTracingScene.h b/Engine/Source/Runtime/Renderer/Private/RayTracing/RayTracingScene.h
--- a/Engine/Source/Runtime/Renderer/Private/RayTracing/RayTracingScene.h
+++ b/Engine/Source/Runtime/Renderer/Private/RayTracing/RayTracingScene.h
@@ -32,6 +32,14 @@ namespace Nanite
 	using CoarseMeshStreamingHandle = int16;
 }
 
+enum class ERayTracingSceneLayer : uint8
+{
+	Base = 0,
+	Decals,
+
+	NUM
+};
+
 /**
 * Persistent representation of the scene for ray tracing.
 * Manages top level acceleration structure instances, memory and build process.
diff --git a/Engine/Source/Runtime/Renderer/Private/RayTracing/RayTracingShaderBindingTable.cpp b/Engine/Source/Runtime/Renderer/Private/RayTracing/RayTracingShaderBindingTable.cpp
--- a/Engine/Source/Runtime/Renderer/Private/RayTracing/RayTracingShaderBindingTable.cpp
+++ b/Engine/Source/Runtime/Renderer/Private/RayTracing/RayTracingShaderBindingTable.cpp
@@ -74,7 +74,7 @@ static bool UseRayTracingPersistentSBTs()
 	return GRHIGlobals.RayTracing.SupportsPersistentSBTs && GPersistentSBTEnabled;
 }
 
-uint32 FRayTracingSBTAllocation::GetRecordIndex(ERayTracingSceneLayer Layer, uint32 SegmentIndex) const
+uint32 FRayTracingSBTAllocation::GetRecordIndex(ERayTracingShaderBindingLayer Layer, uint32 SegmentIndex) const
 {
 	check(HasLayer(Layer));
 
@@ -349,7 +349,7 @@ void FRayTracingShaderBindingTable::MarkDirty(FRayTracingSBTAllocation* SBTAlloc
 }
 
 FRayTracingSBTAllocation* FRayTracingShaderBindingTable::AllocateStaticRangeInternal(
-	ERayTracingSceneLayerMask AllocatedLayers, 
+	ERayTracingShaderBindingLayerMask AllocatedLayers,
 	uint32 SegmentCount, 
 	const FRHIRayTracingGeometry* Geometry, 
 	FRayTracingCachedMeshCommandFlags Flags)
@@ -382,16 +382,16 @@ FRayTracingSBTAllocation* FRayTracingShaderBindingTable::AllocateStaticRange(uin
 		return nullptr;
 	}
 
-	ERayTracingSceneLayerMask AllocatedLayers = ERayTracingSceneLayerMask::None;
+	ERayTracingShaderBindingLayerMask AllocatedLayers = ERayTracingShaderBindingLayerMask::None;
 	if (!Flags.bAllSegmentsDecal)
 	{
-		EnumAddFlags(AllocatedLayers, ERayTracingSceneLayerMask::Base);
+		EnumAddFlags(AllocatedLayers, ERayTracingShaderBindingLayerMask::Base);
 	}
 	if (Flags.bAnySegmentsDecal && !RayTracing::ShouldExcludeDecals())
 	{
-		EnumAddFlags(AllocatedLayers, ERayTracingSceneLayerMask::Decals);
+		EnumAddFlags(AllocatedLayers, ERayTracingShaderBindingLayerMask::Decals);
 	}
-	if (AllocatedLayers == ERayTracingSceneLayerMask::None)
+	if (AllocatedLayers == ERayTracingShaderBindingLayerMask::None)
 	{
 		return nullptr;
 	}
@@ -458,7 +458,7 @@ void FRayTracingShaderBindingTable::FreeStaticRange(const FRayTracingSBTAllocati
 	}
 }
 
-FRayTracingSBTAllocation* FRayTracingShaderBindingTable::AllocateDynamicRange(ERayTracingSceneLayerMask AllocatedLayers, uint32 SegmentCount)
+FRayTracingSBTAllocation* FRayTracingShaderBindingTable::AllocateDynamicRange(ERayTracingShaderBindingLayerMask AllocatedLayers, uint32 SegmentCount)
 {	
 	// Don't need lock right now because all dynamic allocation are allocated linearly on the same thread
 	// So the FreeDynamicAllocationPool can't be shared with the static allocations right now because those would require a lock
diff --git a/Engine/Source/Runtime/Renderer/Private/RayTracing/RayTracingShaderBindingTable.h b/Engine/Source/Runtime/Renderer/Private/RayTracing/RayTracingShaderBindingTable.h
--- a/Engine/Source/Runtime/Renderer/Private/RayTracing/RayTracingShaderBindingTable.h
+++ b/Engine/Source/Runtime/Renderer/Private/RayTracing/RayTracingShaderBindingTable.h
@@ -9,7 +9,7 @@
 
 #if RHI_RAYTRACING
 
-enum class ERayTracingSceneLayer : uint8
+enum class ERayTracingShaderBindingLayer : uint8
 {
 	Base = 0,
 	Decals,
@@ -17,11 +17,11 @@ enum class ERayTracingSceneLayer : uint8
 	NUM
 };
 
-enum class ERayTracingSceneLayerMask
+enum class ERayTracingShaderBindingLayerMask
 {
 	None = 0,
-	Base = 1 << (uint32)ERayTracingSceneLayer::Base,
-	Decals = 1 << (uint32)ERayTracingSceneLayer::Decals,
+	Base = 1 << (uint32)ERayTracingShaderBindingLayer::Base,
+	Decals = 1 << (uint32)ERayTracingShaderBindingLayer::Decals,
 	All = Base | Decals,
 };
 
@@ -39,7 +39,7 @@ public:
 	/**
 	 * Get the InstanceContributionToHitGroupIndex for the given layer which is stored in the RayTracingInstance data
 	 */
-	uint32 GetInstanceContributionToHitGroupIndex(ERayTracingSceneLayer Layer) const
+	uint32 GetInstanceContributionToHitGroupIndex(ERayTracingShaderBindingLayer Layer) const
 	{
 		// InstanceContributionToHitGroupIndex is stored at the first segment index because all other segments are directly allocated after this one
 		return GetRecordIndex(Layer, 0);
@@ -48,15 +48,15 @@ public:
 	/**
 	 * Get the base SBT record index for the given layer and segment index
 	 */
-	RENDERER_API uint32 GetRecordIndex(ERayTracingSceneLayer Layer, uint32 SegmentIndex) const;
+	RENDERER_API uint32 GetRecordIndex(ERayTracingShaderBindingLayer Layer, uint32 SegmentIndex) const;
 	RENDERER_API int32 GetSegmentCount() const;
-	bool HasLayer(ERayTracingSceneLayer Layer) const { return EnumHasAnyFlags(ERayTracingSceneLayerMask(1 << (uint32)Layer), AllocatedLayers); }
+	bool HasLayer(ERayTracingShaderBindingLayer Layer) const { return EnumHasAnyFlags(ERayTracingShaderBindingLayerMask(1 << (uint32)Layer), AllocatedLayers); }
 
 private:
 
-	void InitStatic(ERayTracingSceneLayerMask InAllocatedLayers, uint32 InBaseRecordIndex, uint32 InRecordsPerLayer, uint32 InNumRecords, const FRHIRayTracingGeometry* InGeometry, const FRayTracingCachedMeshCommandFlags& InFlags)
+	void InitStatic(ERayTracingShaderBindingLayerMask InAllocatedLayers, uint32 InBaseRecordIndex, uint32 InRecordsPerLayer, uint32 InNumRecords, const FRHIRayTracingGeometry* InGeometry, const FRayTracingCachedMeshCommandFlags& InFlags)
 	{
-		check(InAllocatedLayers != ERayTracingSceneLayerMask::None);
+		check(InAllocatedLayers != ERayTracingShaderBindingLayerMask::None);
 		AllocatedLayers = InAllocatedLayers;
 		BaseRecordIndex = InBaseRecordIndex;
 		RecordsPerLayer = InRecordsPerLayer;
@@ -65,9 +65,9 @@ private:
 		Flags = InFlags;
 	}
 
-	void InitDynamic(ERayTracingSceneLayerMask InAllocatedLayers, uint32 InBaseRecordIndex, uint32 InRecordsPerLayer, uint32 InNumRecords)
+	void InitDynamic(ERayTracingShaderBindingLayerMask InAllocatedLayers, uint32 InBaseRecordIndex, uint32 InRecordsPerLayer, uint32 InNumRecords)
 	{
-		check(InAllocatedLayers != ERayTracingSceneLayerMask::None);
+		check(InAllocatedLayers != ERayTracingShaderBindingLayerMask::None);
 		AllocatedLayers = InAllocatedLayers;
 		BaseRecordIndex = InBaseRecordIndex;
 		RecordsPerLayer = InRecordsPerLayer;
@@ -79,7 +79,7 @@ private:
 	uint32 BaseRecordIndex = 0;
 	uint32 RecordsPerLayer = 0;
 	uint32 NumRecords = 0;
-	ERayTracingSceneLayerMask AllocatedLayers;
+	ERayTracingShaderBindingLayerMask AllocatedLayers;
 
 	// Store the original geometry and flags in the allocation object so it can be used to build the lookup key again used for deduplication
 	const FRHIRayTracingGeometry* Geometry = nullptr;
@@ -156,7 +156,7 @@ public:
 	/**
 	 * Allocate single static range of records for the given SegmentCount for all layers in the AllocatedLayersMask
 	 */
-	FRayTracingSBTAllocation* AllocateStaticRange(ERayTracingSceneLayerMask AllocatedLayers, uint32 SegmentCount)
+	FRayTracingSBTAllocation* AllocateStaticRange(ERayTracingShaderBindingLayerMask AllocatedLayers, uint32 SegmentCount)
 	{
 		FScopeLock ScopeLock(&StaticAllocationCS);
 		FRayTracingCachedMeshCommandFlags DefaultFlags;
@@ -172,7 +172,7 @@ public:
 	/**
 	 * Allocate dynamic SBT range which can be reused again when ResetDynamicAllocationData is called
 	 */	
-	RENDERER_API FRayTracingSBTAllocation* AllocateDynamicRange(ERayTracingSceneLayerMask AllocatedLayers, uint32 SegmentCount);
+	RENDERER_API FRayTracingSBTAllocation* AllocateDynamicRange(ERayTracingShaderBindingLayerMask AllocatedLayers, uint32 SegmentCount);
 
 	/**
 	 * Mark all currently allocated dynamic ranges as free again so they can be allocated
@@ -211,7 +211,7 @@ private:
 	/**
 	 * Allocate single static range of records for the given SegmentCount for all layers in the AllocatedLayersMask
 	 */
-	RENDERER_API FRayTracingSBTAllocation* AllocateStaticRangeInternal(ERayTracingSceneLayerMask AllocatedLayers, uint32 SegmentCount, const FRHIRayTracingGeometry* Geometry, FRayTracingCachedMeshCommandFlags Flags);
+	RENDERER_API FRayTracingSBTAllocation* AllocateStaticRangeInternal(ERayTracingShaderBindingLayerMask AllocatedLayers, uint32 SegmentCount, const FRHIRayTracingGeometry* Geometry, FRayTracingCachedMeshCommandFlags Flags);
 
 	/**
 	 * Mark all records used by the SBTAllocation as dirty again
