Mobile deferred DBuffer decals

diff --git a/Engine/Source/Runtime/RenderCore/Private/RenderUtils.cpp b/Engine/Source/Runtime/RenderCore/Private/RenderUtils.cpp
--- a/Engine/Source/Runtime/RenderCore/Private/RenderUtils.cpp
+++ b/Engine/Source/Runtime/RenderCore/Private/RenderUtils.cpp
@@ -813,7 +813,7 @@ RENDERCORE_API void RenderUtilsInit()
 
 					GForwardShadingPlatformMask[ShaderPlatformIndex] = TargetPlatformSettings->UsesForwardShading();
 
-					GDBufferPlatformMask[ShaderPlatformIndex] = IsMobilePlatform(ShaderPlatformToEdit) ? (TargetPlatformSettings->UsesMobileDBuffer() && !IsMobileDeferredShadingEnabled(ShaderPlatformToEdit)) : TargetPlatformSettings->UsesDBuffer();
+					GDBufferPlatformMask[ShaderPlatformIndex] = IsMobilePlatform(ShaderPlatformToEdit) ? TargetPlatformSettings->UsesMobileDBuffer() : TargetPlatformSettings->UsesDBuffer();
 
 					GSelectiveBasePassOutputsPlatformMask[ShaderPlatformIndex] = TargetPlatformSettings->UsesSelectiveBasePassOutputs();
 
@@ -860,7 +860,7 @@ RENDERCORE_API void RenderUtilsInit()
 	if (IsMobilePlatform(GMaxRHIShaderPlatform))
 	{
 		static IConsoleVariable* MobileDBufferCVar = IConsoleManager::Get().FindConsoleVariable(TEXT("r.Mobile.DBuffer"));
-		GDBufferPlatformMask.Init(MobileDBufferCVar && MobileDBufferCVar->GetInt() && !IsMobileDeferredShadingEnabled(GMaxRHIShaderPlatform), EShaderPlatform::SP_NumPlatforms);
+		GDBufferPlatformMask.Init(MobileDBufferCVar && MobileDBufferCVar->GetInt(), EShaderPlatform::SP_NumPlatforms);
 	}
 
 	if (RayTracingCVar && RayTracingCVar->GetInt() && GRHISupportsRayTracing)
diff --git a/Engine/Source/Runtime/Renderer/Private/MobileDecalRendering.cpp b/Engine/Source/Runtime/Renderer/Private/MobileDecalRendering.cpp
--- a/Engine/Source/Runtime/Renderer/Private/MobileDecalRendering.cpp
+++ b/Engine/Source/Runtime/Renderer/Private/MobileDecalRendering.cpp
@@ -55,7 +55,7 @@ void FMobileSceneRenderer::RenderDecals(FRHICommandList& RHICmdList, FViewInfo&
 	SCOPE_CYCLE_COUNTER(STAT_DecalsDrawTime);
 
 	const bool bIsMobileDeferred = IsMobileDeferredShadingEnabled(View.GetShaderPlatform());
-	const EDecalRenderStage DecalRenderStage = bIsMobileDeferred ? EDecalRenderStage::MobileBeforeLighting : bRequiresDBufferDecals ? EDecalRenderStage::Emissive : EDecalRenderStage::Mobile;
+	const EDecalRenderStage DecalRenderStage = bRequiresDBufferDecals ? EDecalRenderStage::Emissive : bIsMobileDeferred ? EDecalRenderStage::MobileBeforeLighting : EDecalRenderStage::Mobile;
 	const EDecalRenderTargetMode RenderTargetMode = bIsMobileDeferred ? EDecalRenderTargetMode::SceneColorAndGBuffer : EDecalRenderTargetMode::SceneColor;
 
 	// Deferred decals
diff --git a/Engine/Source/Runtime/Renderer/Private/MobileShadingRenderer.cpp b/Engine/Source/Runtime/Renderer/Private/MobileShadingRenderer.cpp
--- a/Engine/Source/Runtime/Renderer/Private/MobileShadingRenderer.cpp
+++ b/Engine/Source/Runtime/Renderer/Private/MobileShadingRenderer.cpp
@@ -300,7 +300,7 @@ FMobileSceneRenderer::FMobileSceneRenderer(const FSceneViewFamily* InViewFamily,
 	: FSceneRenderer(InViewFamily, HitProxyConsumer)
 	, bGammaSpace(!IsMobileHDR())
 	, bDeferredShading(IsMobileDeferredShadingEnabled(ShaderPlatform))
-	, bRequiresDBufferDecals(bDeferredShading ? false : IsUsingDBuffers(ShaderPlatform))
+	, bRequiresDBufferDecals(IsUsingDBuffers(ShaderPlatform))
 	, bUseVirtualTexturing(UseVirtualTexturing(ShaderPlatform) && GetRendererOutput() != FSceneRenderer::ERendererOutput::DepthPrepassOnly)
 	, bSupportsSimpleLights(bDeferredShading || MobileForwardEnableParticleLights(ShaderPlatform))
 {
@@ -1456,11 +1456,11 @@ void FMobileSceneRenderer::Render(FRDGBuilder& GraphBuilder, const FSceneRenderU
 		{
 			if (bRequiresMultiPass)
 			{
-				RenderDeferredMultiPass(GraphBuilder, SceneTextures, SortedLightSet, InstanceCullingManager);
+				RenderDeferredMultiPass(GraphBuilder, SceneTextures, SortedLightSet, DBufferTextures, InstanceCullingManager);
 			}
 			else
 			{
-				RenderDeferredSinglePass(GraphBuilder, SceneTextures, SortedLightSet, InstanceCullingManager);
+				RenderDeferredSinglePass(GraphBuilder, SceneTextures, SortedLightSet, DBufferTextures, InstanceCullingManager);
 			}
 		}
 		else
@@ -2097,7 +2097,7 @@ FRenderTargetBindingSlots FMobileSceneRenderer::InitRenderTargetBindings_Deferre
 	return BasePassRenderTargets;
 }
 
-void FMobileSceneRenderer::RenderDeferredSinglePass(FRDGBuilder& GraphBuilder, FSceneTextures& SceneTextures, const FSortedLightSetSceneInfo& SortedLightSet, FInstanceCullingManager& InstanceCullingManager)
+void FMobileSceneRenderer::RenderDeferredSinglePass(FRDGBuilder& GraphBuilder, FSceneTextures& SceneTextures, const FSortedLightSetSceneInfo& SortedLightSet, FDBufferTextures& DBufferTextures, FInstanceCullingManager& InstanceCullingManager)
 {
 	bool bUsingPixelLocalStorage = UsingPixelLocalStorage(ShaderPlatform);
 	FColorTargets ColorTargets = GetColorTargets_Deferred(SceneTextures);
@@ -2143,9 +2143,11 @@ void FMobileSceneRenderer::RenderDeferredSinglePass(FRDGBuilder& GraphBuilder, F
 		};
 		auto ParameterCollection = GraphBuilder.AllocParameters<FDeferredSinglePassParameterCollection>();
 
+		FMobileBasePassTextures MobileBasePassTextures{};
+		MobileBasePassTextures.DBufferTextures = DBufferTextures;
 		auto PassParameters = &ParameterCollection->PassParameters;
 		PassParameters->View = View.GetShaderParameters();
-		PassParameters->MobileBasePass = CreateMobileBasePassUniformBuffer(GraphBuilder, View, EMobileBasePass::Opaque, SetupMode);
+		PassParameters->MobileBasePass = CreateMobileBasePassUniformBuffer(GraphBuilder, View, EMobileBasePass::Opaque, SetupMode, MobileBasePassTextures);
 		PassParameters->ReflectionCapture = View.MobileReflectionCaptureUniformBuffer;
 		PassParameters->LocalFogVolumeInstances = View.LocalFogVolumeViewData.GPUInstanceDataBufferSRV;
 		PassParameters->LocalFogVolumeTileDrawIndirectBuffer = View.LocalFogVolumeViewData.GPUTileDrawIndirectBuffer;
@@ -2213,7 +2215,7 @@ void FMobileSceneRenderer::RenderDeferredSinglePass(FRDGBuilder& GraphBuilder, F
 	}
 }
 
-void FMobileSceneRenderer::RenderDeferredMultiPass(FRDGBuilder& GraphBuilder, FSceneTextures& SceneTextures, const FSortedLightSetSceneInfo& SortedLightSet, FInstanceCullingManager& InstanceCullingManager)
+void FMobileSceneRenderer::RenderDeferredMultiPass(FRDGBuilder& GraphBuilder, FSceneTextures& SceneTextures, const FSortedLightSetSceneInfo& SortedLightSet, FDBufferTextures& DBufferTextures, FInstanceCullingManager& InstanceCullingManager)
 {
 	FColorTargets ColorTargets = GetColorTargets_Deferred(SceneTextures);
 	FRenderTargetBindingSlots BasePassRenderTargets = InitRenderTargetBindings_Deferred(SceneTextures, ColorTargets);
@@ -2244,7 +2246,9 @@ void FMobileSceneRenderer::RenderDeferredMultiPass(FRDGBuilder& GraphBuilder, FS
 		auto PassParameters = &ParameterCollection->PassParameters;
 		PassParameters->View = View.GetShaderParameters();
 		EMobileSceneTextureSetupMode SetupMode = bIsFullDepthPrepassEnabled ? EMobileSceneTextureSetupMode::SceneDepth : EMobileSceneTextureSetupMode::None;
-		PassParameters->MobileBasePass = CreateMobileBasePassUniformBuffer(GraphBuilder, View, EMobileBasePass::Opaque, SetupMode);
+		FMobileBasePassTextures MobileBasePassTextures{};
+		MobileBasePassTextures.DBufferTextures = DBufferTextures;
+		PassParameters->MobileBasePass = CreateMobileBasePassUniformBuffer(GraphBuilder, View, EMobileBasePass::Opaque, SetupMode, MobileBasePassTextures);
 		PassParameters->RenderTargets = BasePassRenderTargets;
 		if (!ViewContext.bIsFirstView)
 		{
diff --git a/Engine/Source/Runtime/Renderer/Private/SceneRendering.h b/Engine/Source/Runtime/Renderer/Private/SceneRendering.h
--- a/Engine/Source/Runtime/Renderer/Private/SceneRendering.h
+++ b/Engine/Source/Runtime/Renderer/Private/SceneRendering.h
@@ -2869,8 +2869,8 @@ protected:
 	void RenderForwardSinglePass(FRDGBuilder& GraphBuilder, class FMobileRenderPassParameters* PassParameters, struct FRenderViewContext& ViewContext, FSceneTextures& SceneTextures, FInstanceCullingManager& InstanceCullingManager);
 	void RenderForwardMultiPass(FRDGBuilder& GraphBuilder, class FMobileRenderPassParameters* PassParameters, struct FRenderViewContext& ViewContext, FSceneTextures& SceneTextures, FInstanceCullingManager& InstanceCullingManager);
 	
-	void RenderDeferredSinglePass(FRDGBuilder& GraphBuilder, FSceneTextures& SceneTextures, const FSortedLightSetSceneInfo& SortedLightSet, FInstanceCullingManager& InstanceCullingManager);
-	void RenderDeferredMultiPass(FRDGBuilder& GraphBuilder, FSceneTextures& SceneTextures, const FSortedLightSetSceneInfo& SortedLightSet, FInstanceCullingManager& InstanceCullingManager);
+	void RenderDeferredSinglePass(FRDGBuilder& GraphBuilder, FSceneTextures& SceneTextures, const FSortedLightSetSceneInfo& SortedLightSet, FDBufferTextures& DBufferTextures, FInstanceCullingManager& InstanceCullingManager);
+	void RenderDeferredMultiPass(FRDGBuilder& GraphBuilder, FSceneTextures& SceneTextures, const FSortedLightSetSceneInfo& SortedLightSet, FDBufferTextures& DBufferTextures, FInstanceCullingManager& InstanceCullingManager);
 
 	void RenderAmbientOcclusion(FRDGBuilder& GraphBuilder, FRDGTextureRef SceneDepthTexture, FRDGTextureRef AmbientOcclusionTexture);
 
