[HWRT] DynamicGeometryUpdateManager - Add support for forcing a rebuild of a dynamic BLAS:
* r.RayTracing.DynamicGeometry.ForceBuild.MaxPrimitivesPerFrame - controls how many primitives per frame can be rebuild (default 0)
* r.RayTracing.DynamicGeometry.ForceBuild.MinUpdatesSinceLastBuild - controls after how many updates we want to consider rebuilding a BLAS (default INT_MAX)

#tests Firefly
#rnx

diff --git a/Engine/Source/Runtime/RenderCore/Public/RayTracingGeometry.h b/Engine/Source/Runtime/RenderCore/Public/RayTracingGeometry.h
--- a/Engine/Source/Runtime/RenderCore/Public/RayTracingGeometry.h
+++ b/Engine/Source/Runtime/RenderCore/Public/RayTracingGeometry.h
@@ -54,8 +54,11 @@ public:
 	*/
 	int64 DynamicGeometrySharedBufferGenerationID = NonSharedVertexBuffers;
 
-	// Last frame when geometry was updated (only skinned geometry)
-	uint64 LastUpdatedFrame = 0;	
+	// Last frame when geometry was updated (only dynamic geometry)
+	uint64 LastUpdatedFrame = 0;
+
+	// How many updates since the last build (only dynamic geometry)
+	uint32 NumUpdatesSinceLastBuild = 0;
 
 	FRayTracingGeometryInitializer Initializer;
 
diff --git a/Engine/Source/Runtime/Renderer/Private/RayTracing/RayTracingDynamicGeometryUpdateManager.cpp b/Engine/Source/Runtime/Renderer/Private/RayTracing/RayTracingDynamicGeometryUpdateManager.cpp
--- a/Engine/Source/Runtime/Renderer/Private/RayTracing/RayTracingDynamicGeometryUpdateManager.cpp
+++ b/Engine/Source/Runtime/Renderer/Private/RayTracing/RayTracingDynamicGeometryUpdateManager.cpp
@@ -44,6 +44,20 @@ static TAutoConsoleVariable<int32> CVarRTDynGeomMaxUpdatePrimitivesPerFrame(
 	ECVF_RenderThreadSafe
 );
 
+static TAutoConsoleVariable<int32> CVarRTDynGeomForceBuildMaxUpdatePrimitivesPerFrame(
+	TEXT("r.RayTracing.DynamicGeometry.ForceBuild.MaxPrimitivesPerFrame"),
+	0,
+	TEXT("Sets the dynamic ray tracing acceleration structure build budget in terms of maximum number of triangles that are rebuild per frame (default 0)"),
+	ECVF_RenderThreadSafe
+);
+
+static TAutoConsoleVariable<int32> CVarRTDynGeomForceBuildMinUpdatesSinceLastBuild(
+	TEXT("r.RayTracing.DynamicGeometry.ForceBuild.MinUpdatesSinceLastBuild"),
+	-1,
+	TEXT("Sets minimum number of updates before the dynamic geometry acceleration structure will be considered for rebuild (default INT_MAX)"),
+	ECVF_RenderThreadSafe
+);
+
 class FRayTracingDynamicGeometryConverterCS : public FMeshMaterialShader
 {
 	DECLARE_SHADER_TYPE(FRayTracingDynamicGeometryConverterCS, MeshMaterial);
@@ -583,103 +597,108 @@ void FRayTracingDynamicGeometryUpdateManager::Update(const FViewInfo& View)
 	uint32 BLASScratchSize = 0;
 	int32 NumBuildPrimitives = 0;
 
-	for (const FRayTracingDynamicGeometryBuildParams& Build : DynamicGeometryBuilds)
-	{
-		FRHIRayTracingGeometry* RayTracingGeometry = Build.Geometry->GetRHI();
-
-		NumBuildPrimitives += Build.Geometry->Initializer.TotalPrimitiveCount;
-
-		const uint32 ScratchSize = RayTracingGeometry->GetSizeInfo().BuildScratchSize;
-		BLASScratchSize = Align(BLASScratchSize + ScratchSize, ScratchAlignment);
-
-		FRayTracingGeometryBuildParams BuildParam;
-		BuildParam.Geometry = RayTracingGeometry;
-		BuildParam.BuildMode = EAccelerationStructureBuildMode::Build;
-
-		if (Build.SegmentOffset >= 0)
+	auto AddGeometryBuildParamsToBuildList = [](const FRayTracingDynamicGeometryBuildParams& BuildParams,
+		FRayTracingGeometrySegment* SegmentData,
+		uint32& OutScratchSize,
+		TArray<FMeshComputeDispatchCommand>& OutDispatchCommands,
+		TArray<FRayTracingGeometryBuildParams>& OutBuildParams
+		)
 		{
-			BuildParam.Segments = MakeArrayView(&SegmentData[Build.SegmentOffset], Build.Geometry->Initializer.Segments.Num());
-		}
-
-		BuildParams.Add(MoveTemp(BuildParam));
+			FRHIRayTracingGeometry* RayTracingGeometry = BuildParams.Geometry->GetRHI();
 
-		DispatchCommands.Append(Build.DispatchCommands);
-	}
+			const uint32 ScratchAlignment = GRHIRayTracingScratchBufferAlignment;
+			const uint32 NumBuildPrimitives = BuildParams.Geometry->Initializer.TotalPrimitiveCount;
 
-	const int32 MaxUpdatePrimitivesPerFrame = CVarRTDynGeomMaxUpdatePrimitivesPerFrame.GetValueOnRenderThread();
+			const uint32 ScratchSize = RayTracingGeometry->GetSizeInfo().BuildScratchSize;
+			OutScratchSize = Align(OutScratchSize + ScratchSize, ScratchAlignment);
 
-	int32 NumUpdatedPrimitives = 0;
+			BuildParams.Geometry->NumUpdatesSinceLastBuild = 0;
 
-	if (MaxUpdatePrimitivesPerFrame <= 0)
-	{
-		for (const FRayTracingDynamicGeometryBuildParams& Update : DynamicGeometryUpdates)
-		{
-			FRHIRayTracingGeometry* RayTracingGeometry = Update.Geometry->GetRHI();
+			FRayTracingGeometryBuildParams BuildParam;
+			BuildParam.Geometry = RayTracingGeometry;
+			BuildParam.BuildMode = EAccelerationStructureBuildMode::Build;
 
-			if (View.bRayTracingFeedbackEnabled && !GRayTracingGeometryManager->IsGeometryVisible(Update.Geometry->GetGeometryHandle()))
+			if (BuildParams.SegmentOffset >= 0)
 			{
-				INC_DWORD_STAT_BY(STAT_RayTracingDynamicSkippedPrimitives, Update.Geometry->Initializer.TotalPrimitiveCount);
-				continue;
+				BuildParam.Segments = MakeArrayView(&SegmentData[BuildParams.SegmentOffset], BuildParams.Geometry->Initializer.Segments.Num());
 			}
 
-			Update.Geometry->LastUpdatedFrame = GFrameCounterRenderThread;
+			OutBuildParams.Add(MoveTemp(BuildParam));
 
-			NumUpdatedPrimitives += Update.Geometry->Initializer.TotalPrimitiveCount;
+			OutDispatchCommands.Append(BuildParams.DispatchCommands);
+		};
 
-			const uint32 ScratchSize = RayTracingGeometry->GetSizeInfo().UpdateScratchSize;
-			BLASScratchSize = Align(BLASScratchSize + ScratchSize, ScratchAlignment);
+	for (const FRayTracingDynamicGeometryBuildParams& Build : DynamicGeometryBuilds)
+	{
+		AddGeometryBuildParamsToBuildList(Build, SegmentData, BLASScratchSize, DispatchCommands, BuildParams);
+		NumBuildPrimitives += Build.Geometry->Initializer.TotalPrimitiveCount;
+	}
 
-			FRayTracingGeometryBuildParams BuildParam;
-			BuildParam.Geometry = RayTracingGeometry;
-			BuildParam.BuildMode = EAccelerationStructureBuildMode::Update;
-			if (Update.SegmentOffset >= 0)
-			{
-				BuildParam.Segments = MakeArrayView(&SegmentData[Update.SegmentOffset], Update.Geometry->Initializer.Segments.Num());
-			}
-			BuildParams.Add(MoveTemp(BuildParam));
+	const uint32 MaxUpdatePrimitivesPerFrame = uint32(CVarRTDynGeomMaxUpdatePrimitivesPerFrame.GetValueOnRenderThread());
+	const uint32 MaxForceBuildPrimitivesPerFrame = uint32(CVarRTDynGeomForceBuildMaxUpdatePrimitivesPerFrame.GetValueOnRenderThread());
+	const uint32 MinUpdatesSinceLastBuild = uint32(CVarRTDynGeomForceBuildMinUpdatesSinceLastBuild.GetValueOnRenderThread());
 
-			DispatchCommands.Append(Update.DispatchCommands);
-		}
-	}
-	else
+	uint32 NumUpdatedPrimitives = 0;
+	uint32 NumForceBuildPrimitives = 0;
+
+	const bool bNeedsSorting = (int32(MaxUpdatePrimitivesPerFrame) != -1) || (MaxForceBuildPrimitivesPerFrame != 0);
+	if (bNeedsSorting)
 	{
 		DynamicGeometryUpdates.Sort([](const FRayTracingDynamicGeometryBuildParams& InLHS, const FRayTracingDynamicGeometryBuildParams& InRHS)
 			{
+				if (InLHS.Geometry->LastUpdatedFrame == InRHS.Geometry->LastUpdatedFrame)
+				{
+					return InLHS.Geometry->NumUpdatesSinceLastBuild > InRHS.Geometry->NumUpdatesSinceLastBuild;
+				}
+
 				return InLHS.Geometry->LastUpdatedFrame < InRHS.Geometry->LastUpdatedFrame;
 			});
+	}	
+
+	for (const FRayTracingDynamicGeometryBuildParams& Update : DynamicGeometryUpdates)
+	{
+		FRHIRayTracingGeometry* RayTracingGeometry = Update.Geometry->GetRHI();
+		const uint32 TotalPrimitiveCount = Update.Geometry->Initializer.TotalPrimitiveCount;
 
-		for (const FRayTracingDynamicGeometryBuildParams& Update : DynamicGeometryUpdates)
+		if (View.bRayTracingFeedbackEnabled && !GRayTracingGeometryManager->IsGeometryVisible(Update.Geometry->GetGeometryHandle()))
 		{
-			FRHIRayTracingGeometry* RayTracingGeometry = Update.Geometry->GetRHI();
+			INC_DWORD_STAT_BY(STAT_RayTracingDynamicSkippedPrimitives, TotalPrimitiveCount);
+			continue;
+		}
 
-			if (View.bRayTracingFeedbackEnabled && !GRayTracingGeometryManager->IsGeometryVisible(Update.Geometry->GetGeometryHandle()))
+		if (MaxForceBuildPrimitivesPerFrame > 0)
+		{
+			if (Update.Geometry->NumUpdatesSinceLastBuild > MinUpdatesSinceLastBuild && NumForceBuildPrimitives <= MaxForceBuildPrimitivesPerFrame)
 			{
-				INC_DWORD_STAT_BY(STAT_RayTracingDynamicSkippedPrimitives, Update.Geometry->Initializer.TotalPrimitiveCount);
+				AddGeometryBuildParamsToBuildList(Update, SegmentData, BLASScratchSize, DispatchCommands, BuildParams);
+				NumBuildPrimitives += TotalPrimitiveCount;
+				NumForceBuildPrimitives += TotalPrimitiveCount;
 				continue;
 			}
+		}
 
-			Update.Geometry->LastUpdatedFrame = GFrameCounterRenderThread;
+		Update.Geometry->LastUpdatedFrame = GFrameCounterRenderThread;
+		Update.Geometry->NumUpdatesSinceLastBuild++;
 
-			NumUpdatedPrimitives += Update.Geometry->Initializer.TotalPrimitiveCount;
+		NumUpdatedPrimitives += TotalPrimitiveCount;
 
-			const uint32 ScratchSize = RayTracingGeometry->GetSizeInfo().UpdateScratchSize;
-			BLASScratchSize = Align(BLASScratchSize + ScratchSize, ScratchAlignment);
+		const uint32 ScratchSize = RayTracingGeometry->GetSizeInfo().UpdateScratchSize;
+		BLASScratchSize = Align(BLASScratchSize + ScratchSize, ScratchAlignment);
 
-			FRayTracingGeometryBuildParams BuildParam;
-			BuildParam.Geometry = RayTracingGeometry;
-			BuildParam.BuildMode = EAccelerationStructureBuildMode::Update;
-			if (Update.SegmentOffset >= 0)
-			{
-				BuildParam.Segments = MakeArrayView(&SegmentData[Update.SegmentOffset], Update.Geometry->Initializer.Segments.Num());
-			}
-			BuildParams.Add(MoveTemp(BuildParam));
+		FRayTracingGeometryBuildParams BuildParam;
+		BuildParam.Geometry = RayTracingGeometry;
+		BuildParam.BuildMode = EAccelerationStructureBuildMode::Update;
+		if (Update.SegmentOffset >= 0)
+		{
+			BuildParam.Segments = MakeArrayView(&SegmentData[Update.SegmentOffset], Update.Geometry->Initializer.Segments.Num());
+		}
+		BuildParams.Add(MoveTemp(BuildParam));
 
-			DispatchCommands.Append(Update.DispatchCommands);
+		DispatchCommands.Append(Update.DispatchCommands);
 
-			if (NumUpdatedPrimitives > MaxUpdatePrimitivesPerFrame)
-			{
-				break;
-			}
+		if (NumUpdatedPrimitives > MaxUpdatePrimitivesPerFrame)
+		{
+			break;
 		}
 	}
 
