Support for pipeline VS->PS combined shader compilation added for BasePass material shaders. This allows for detecting unused interpolators in PS and removing them in both VS and PS in the same Pipeline type.

diff --git a/Engine/Shaders/Public/ShaderVersion.ush b/Engine/Shaders/Public/ShaderVersion.ush
--- a/Engine/Shaders/Public/ShaderVersion.ush
+++ b/Engine/Shaders/Public/ShaderVersion.ush
@@ -3,4 +3,4 @@
 // Each time the console command is executed it generates a new GUID. As this file is included
 // in Platform.ush (which should be included in any shader) it allows to invalidate the shader DDC.
 // 
-#pragma message("UESHADERMETADATA_VERSION 8F3EFE47-05C8-4AFE-A9D7-50C55D3ED855")
+#pragma message("UESHADERMETADATA_VERSION AAF7C14D-1BF2-4F8E-A286-C910215B5C9E")
diff --git a/Engine/Source/Runtime/RenderCore/Private/Shader.cpp b/Engine/Source/Runtime/RenderCore/Private/Shader.cpp
--- a/Engine/Source/Runtime/RenderCore/Private/Shader.cpp
+++ b/Engine/Source/Runtime/RenderCore/Private/Shader.cpp
@@ -1026,14 +1026,6 @@ FShaderPipelineType::FShaderPipelineType(
 	Stages.Add(InVertexOrMeshShader);
 	AllStages[bInIsMeshPipeline ? SF_Mesh : SF_Vertex] = InVertexOrMeshShader;
 
-	for (uint32 FrequencyIndex = 0; FrequencyIndex < SF_NumStandardFrequencies; ++FrequencyIndex)
-	{
-		if (const FShaderType* ShaderType = AllStages[FrequencyIndex])
-		{
-			checkf(ShaderType->GetPermutationCount() == 1, TEXT("Shader '%s' has multiple shader permutations. Shader pipelines only support a single permutation."), ShaderType->GetName())
-		}
-	}
-
 	static uint32 TypeHashCounter = 0;
 	++TypeHashCounter;
 	HashIndex = TypeHashCounter;
diff --git a/Engine/Source/Runtime/Renderer/Private/BasePassRendering.cpp b/Engine/Source/Runtime/Renderer/Private/BasePassRendering.cpp
--- a/Engine/Source/Runtime/Renderer/Private/BasePassRendering.cpp
+++ b/Engine/Source/Runtime/Renderer/Private/BasePassRendering.cpp
@@ -128,9 +128,13 @@ IMPLEMENT_STATIC_UNIFORM_BUFFER_STRUCT(FTranslucentBasePassUniformParameters, "T
 	typedef TBasePassVS< LightMapPolicyType > TBasePassVS##LightMapPolicyName ; \
 	IMPLEMENT_MATERIAL_SHADER_TYPE(template<>,TBasePassVS##LightMapPolicyName,TEXT("/Engine/Private/BasePassVertexShader.usf"),TEXT("Main"),SF_Vertex);
 
+// This macro allows templated arguments usage in macro definitions
+#define WRAP(...) __VA_ARGS__
+
 #define IMPLEMENT_BASEPASS_PIXELSHADER_TYPE(LightMapPolicyType,LightMapPolicyName,bEnableSkyLight,SkyLightName,GBufferLayout,LayoutName) \
 	typedef TBasePassPS<LightMapPolicyType, bEnableSkyLight, GBufferLayout> TBasePassPS##LightMapPolicyName##SkyLightName##LayoutName; \
-	IMPLEMENT_MATERIAL_SHADER_TYPE(template<>,TBasePassPS##LightMapPolicyName##SkyLightName##LayoutName,TEXT("/Engine/Private/BasePassPixelShader.usf"),TEXT("MainPS"),SF_Pixel);
+	IMPLEMENT_MATERIAL_SHADER_TYPE(template<>,TBasePassPS##LightMapPolicyName##SkyLightName##LayoutName,TEXT("/Engine/Private/BasePassPixelShader.usf"),TEXT("MainPS"),SF_Pixel); \
+	IMPLEMENT_SHADERPIPELINE_TYPE_VSPS(BasePassPipeline##LightMapPolicyName##bEnableSkyLight##SkyLightName##GBufferLayout##LayoutName, WRAP(TBasePassVS<LightMapPolicyType>), WRAP(TBasePassPS<LightMapPolicyType, bEnableSkyLight, GBufferLayout>), true);
 
 #define IMPLEMENT_BASEPASS_COMPUTESHADER_TYPE(LightMapPolicyType,LightMapPolicyName,bEnableSkyLight,SkyLightName) \
 	typedef TBasePassCS<LightMapPolicyType, bEnableSkyLight, SF_Compute> TBasePassCS##LightMapPolicyName##SkyLightName; \
@@ -152,6 +156,7 @@ IMPLEMENT_STATIC_UNIFORM_BUFFER_STRUCT(FTranslucentBasePassUniformParameters, "T
 	IMPLEMENT_BASEPASS_WORKGRAPHSHADER_TYPE(LightMapPolicyType,LightMapPolicyName,true,Skylight) \
 	IMPLEMENT_BASEPASS_WORKGRAPHSHADER_TYPE(LightMapPolicyType,LightMapPolicyName,false,)
 
+// New IMPLEMENT_BASEPASS_LIGHTMAPPED_SHADER_TYPE statements should be also handled in GetUniformBasePassShaders and GetBasePassShaders
 // Implement shader types per lightmap policy
 // If renaming or refactoring these, remember to update FMaterialResource::GetRepresentativeInstructionCounts and FPreviewMaterial::ShouldCache().
 IMPLEMENT_BASEPASS_LIGHTMAPPED_SHADER_TYPE( FSelfShadowedTranslucencyPolicy, FSelfShadowedTranslucencyPolicy );
@@ -183,6 +188,8 @@ void F128BitRTBasePassPS::ModifyCompilationEnvironment(const FMeshMaterialShader
 }
 
 IMPLEMENT_MATERIAL_SHADER_TYPE(, F128BitRTBasePassPS, TEXT("/Engine/Private/BasePassPixelShader.usf"), TEXT("MainPS"), SF_Pixel);
+// New 128 bit RT Base Pass statements should be also handled in GetUniformBasePassShaders and GetBasePassShaders
+IMPLEMENT_SHADERPIPELINE_TYPE_VSPS(BasePassPipeline128BitRT, WRAP(TBasePassVS<TUniformLightMapPolicy<LMP_NO_LIGHTMAP>>), F128BitRTBasePassPS, true);
 
 DEFINE_GPU_DRAWCALL_STAT(Basepass);
 
@@ -190,6 +197,33 @@ DECLARE_CYCLE_STAT(TEXT("DeferredShadingSceneRenderer ClearGBufferAtMaxZ"), STAT
 DECLARE_CYCLE_STAT(TEXT("DeferredShadingSceneRenderer ViewExtensionPostRenderBasePass"), STAT_FDeferredShadingSceneRenderer_ViewExtensionPostRenderBasePass, STATGROUP_SceneRendering);
 
 DEFINE_GPU_STAT(NaniteBasePass);
+#define LIGHTMAPPED_SHADER_PIPELINE_FILL(Policy, DesiredPolicy, IndirectPolicyType, DesiredIndirectPolicyType, DesiredPolicyName) \
+	constexpr (std::is_same<Policy, DesiredPolicy>::value) \
+	{ \
+		if (IndirectPolicyType == DesiredIndirectPolicyType) \
+		{ \
+			if (bUse128bitRT) \
+			{ \
+				checkf(bEnableSkyLight == false && GBufferLayout == GBL_Default, TEXT("New 128bitRT Base Pass was added? / \
+							Currently supporting just the F128BitRTBasePassPS. Handle new ones analogically.")); \
+				ShaderTypes.PipelineType = &(BasePassPipeline128BitRT); \
+			} \
+			else if (bEnableSkyLight) \
+			{ \
+				if (GBufferLayout == GBL_Default) \
+					ShaderTypes.PipelineType = &(BasePassPipeline##DesiredPolicyName##trueSkylightGBL_Default); \
+				else if (GBufferLayout == GBL_ForceVelocity) \
+					ShaderTypes.PipelineType = &(BasePassPipeline##DesiredPolicyName##trueSkylightGBL_ForceVelocityForceVelocity); \
+			} \
+			else \
+			{ \
+				if (GBufferLayout == GBL_Default) \
+					ShaderTypes.PipelineType = &(BasePassPipeline##DesiredPolicyName##falseGBL_Default); \
+				else if (GBufferLayout == GBL_ForceVelocity) \
+					ShaderTypes.PipelineType = &(BasePassPipeline##DesiredPolicyName##falseGBL_ForceVelocityForceVelocity); \
+			} \
+		} \
+	}
 
 static bool IsStandardTranslucenyPassSeparated()
 {
@@ -601,6 +635,15 @@ bool GetUniformBasePassShaders(
 )
 {
 	FMaterialShaderTypes ShaderTypes;
+
+	if LIGHTMAPPED_SHADER_PIPELINE_FILL(TUniformLightMapPolicy<Policy>, TUniformLightMapPolicy<LMP_NO_LIGHTMAP>, LMP_DUMMY, LMP_DUMMY, FNoLightMapPolicy)
+	else if LIGHTMAPPED_SHADER_PIPELINE_FILL(TUniformLightMapPolicy<Policy>, TUniformLightMapPolicy<LMP_PRECOMPUTED_IRRADIANCE_VOLUME_INDIRECT_LIGHTING>, LMP_DUMMY, LMP_DUMMY, FPrecomputedVolumetricLightmapLightingPolicy)
+	else if LIGHTMAPPED_SHADER_PIPELINE_FILL(TUniformLightMapPolicy<Policy>, TUniformLightMapPolicy<LMP_CACHED_VOLUME_INDIRECT_LIGHTING>, LMP_DUMMY, LMP_DUMMY, FCachedVolumeIndirectLightingPolicy)
+	else if LIGHTMAPPED_SHADER_PIPELINE_FILL(TUniformLightMapPolicy<Policy>, TUniformLightMapPolicy<LMP_CACHED_POINT_INDIRECT_LIGHTING>, LMP_DUMMY, LMP_DUMMY, FCachedPointIndirectLightingPolicy)
+	else if LIGHTMAPPED_SHADER_PIPELINE_FILL(TUniformLightMapPolicy<Policy>, TUniformLightMapPolicy<LMP_LQ_LIGHTMAP>, LMP_DUMMY, LMP_DUMMY, TLightMapPolicyLQ)
+	else if LIGHTMAPPED_SHADER_PIPELINE_FILL(TUniformLightMapPolicy<Policy>, TUniformLightMapPolicy<LMP_HQ_LIGHTMAP>, LMP_DUMMY, LMP_DUMMY, TLightMapPolicyHQ)
+	else if LIGHTMAPPED_SHADER_PIPELINE_FILL(TUniformLightMapPolicy<Policy>, TUniformLightMapPolicy<LMP_DISTANCE_FIELD_SHADOWS_AND_HQ_LIGHTMAP>, LMP_DUMMY, LMP_DUMMY, TDistanceFieldShadowsAndLightMapPolicyHQ)
+	
 	if (VertexShader)
 	{
 		ShaderTypes.AddShaderType<TBasePassVS<TUniformLightMapPolicy<Policy>>>();
@@ -633,6 +676,64 @@ bool GetUniformBasePassShaders(
 	return true;
 }
 
+template <typename LightMapPolicyType>
+bool GetBasePassShaders(
+	const FMaterial& Material,
+	const FVertexFactoryType* VertexFactoryType,
+	LightMapPolicyType LightMapPolicy,
+	ERHIFeatureLevel::Type FeatureLevel,
+	bool bEnableSkyLight,
+	bool bUse128bitRT,
+	EGBufferLayout GBufferLayout,
+	TShaderRef<TBasePassVertexShaderPolicyParamType<LightMapPolicyType>>* VertexShader,
+	TShaderRef<TBasePassPixelShaderPolicyParamType<LightMapPolicyType>>* PixelShader,
+	bool bIsForOITPass)
+{
+	FMaterialShaderTypes ShaderTypes;
+
+	if LIGHTMAPPED_SHADER_PIPELINE_FILL(LightMapPolicyType, FSelfShadowedTranslucencyPolicy, LMP_DUMMY, LMP_DUMMY, FSelfShadowedTranslucencyPolicy)
+	else if LIGHTMAPPED_SHADER_PIPELINE_FILL(LightMapPolicyType, FSelfShadowedCachedPointIndirectLightingPolicy, LMP_DUMMY, LMP_DUMMY, FSelfShadowedCachedPointIndirectLightingPolicy) 
+	else if LIGHTMAPPED_SHADER_PIPELINE_FILL(LightMapPolicyType, FSelfShadowedVolumetricLightmapPolicy, LMP_DUMMY, LMP_DUMMY, FSelfShadowedVolumetricLightmapPolicy)
+	else if LIGHTMAPPED_SHADER_PIPELINE_FILL(LightMapPolicyType, FUniformLightMapPolicy, LightMapPolicy.GetIndirectPolicy(), LMP_NO_LIGHTMAP, FNoLightMapPolicy)
+	else if LIGHTMAPPED_SHADER_PIPELINE_FILL(LightMapPolicyType, FUniformLightMapPolicy, LightMapPolicy.GetIndirectPolicy(), LMP_PRECOMPUTED_IRRADIANCE_VOLUME_INDIRECT_LIGHTING, FPrecomputedVolumetricLightmapLightingPolicy)
+	else if LIGHTMAPPED_SHADER_PIPELINE_FILL(LightMapPolicyType, FUniformLightMapPolicy, LightMapPolicy.GetIndirectPolicy(), LMP_CACHED_VOLUME_INDIRECT_LIGHTING, FCachedVolumeIndirectLightingPolicy)
+	else if LIGHTMAPPED_SHADER_PIPELINE_FILL(LightMapPolicyType, FUniformLightMapPolicy, LightMapPolicy.GetIndirectPolicy(), LMP_CACHED_POINT_INDIRECT_LIGHTING, FCachedPointIndirectLightingPolicy)
+	else if LIGHTMAPPED_SHADER_PIPELINE_FILL(LightMapPolicyType, FUniformLightMapPolicy, LightMapPolicy.GetIndirectPolicy(), LMP_LQ_LIGHTMAP, TLightMapPolicyLQ)
+	else if LIGHTMAPPED_SHADER_PIPELINE_FILL(LightMapPolicyType, FUniformLightMapPolicy, LightMapPolicy.GetIndirectPolicy(), LMP_HQ_LIGHTMAP, TLightMapPolicyHQ)
+	else if LIGHTMAPPED_SHADER_PIPELINE_FILL(LightMapPolicyType, FUniformLightMapPolicy, LightMapPolicy.GetIndirectPolicy(), LMP_DISTANCE_FIELD_SHADOWS_AND_HQ_LIGHTMAP, TDistanceFieldShadowsAndLightMapPolicyHQ)
+
+	if (VertexShader)
+	{
+		ShaderTypes.AddShaderType<TBasePassVS<LightMapPolicyType>>();
+	}
+
+	if (PixelShader)
+	{
+		switch (GBufferLayout)
+		{
+		case GBL_Default:
+			AddBasePassPixelShader<LightMapPolicyType, GBL_Default>(bEnableSkyLight, ShaderTypes, bIsForOITPass);
+			break;
+		case GBL_ForceVelocity:
+			AddBasePassPixelShader<LightMapPolicyType, GBL_ForceVelocity>(bEnableSkyLight, ShaderTypes, bIsForOITPass);
+			break;
+		default:
+			check(false);
+			break;
+		}
+	}
+
+	FMaterialShaders Shaders;
+	if (!Material.TryGetShaders(ShaderTypes, VertexFactoryType, Shaders))
+	{
+		return false;
+	}
+
+	Shaders.TryGetVertexShader(VertexShader);
+	Shaders.TryGetPixelShader(PixelShader);
+	return true;
+}
+
 template <>
 bool GetBasePassShaders<FUniformLightMapPolicy>(
 	const FMaterial& Material, 
diff --git a/Engine/Source/Runtime/Renderer/Private/BasePassRendering.h b/Engine/Source/Runtime/Renderer/Private/BasePassRendering.h
--- a/Engine/Source/Runtime/Renderer/Private/BasePassRendering.h
+++ b/Engine/Source/Runtime/Renderer/Private/BasePassRendering.h
@@ -797,40 +797,7 @@ bool GetBasePassShaders(
 	TShaderRef<TBasePassVertexShaderPolicyParamType<LightMapPolicyType>>* VertexShader,
 	TShaderRef<TBasePassPixelShaderPolicyParamType<LightMapPolicyType>>* PixelShader,
 	bool bIsForOITPass = false
-)
-{
-	FMaterialShaderTypes ShaderTypes;
-	if (VertexShader)
-	{
-		ShaderTypes.AddShaderType<TBasePassVS<LightMapPolicyType>>();
-	}
-
-	if (PixelShader)
-	{
-		switch (GBufferLayout)
-		{
-		case GBL_Default:
-			AddBasePassPixelShader<LightMapPolicyType, GBL_Default>(bEnableSkyLight, ShaderTypes, bIsForOITPass);
-			break;
-		case GBL_ForceVelocity:
-			AddBasePassPixelShader<LightMapPolicyType, GBL_ForceVelocity>(bEnableSkyLight, ShaderTypes, bIsForOITPass);
-			break;
-		default:
-			check(false);
-			break;
-		}
-	}
-
-	FMaterialShaders Shaders;
-	if (!Material.TryGetShaders(ShaderTypes, VertexFactoryType, Shaders))
-	{
-		return false;
-	}
-
-	Shaders.TryGetVertexShader(VertexShader);
-	Shaders.TryGetPixelShader(PixelShader);
-	return true;
-}
+);
 
 template <>
 bool GetBasePassShaders<FUniformLightMapPolicy>(
