Make CustomDepth behave same between Editor and EditorPreview on mobile platform
#12805

diff --git a/Engine/Source/Runtime/Renderer/Private/MobileShadingRenderer.cpp b/Engine/Source/Runtime/Renderer/Private/MobileShadingRenderer.cpp
--- a/Engine/Source/Runtime/Renderer/Private/MobileShadingRenderer.cpp
+++ b/Engine/Source/Runtime/Renderer/Private/MobileShadingRenderer.cpp
@@ -642,7 +642,7 @@ void FMobileSceneRenderer::InitViews(
 	
 	// Find out whether custom depth pass should be rendered.
 	{
-		bool bCouldUseCustomDepthStencil = (!Scene->World || (Scene->World->WorldType != EWorldType::EditorPreview && Scene->World->WorldType != EWorldType::Inactive));
+		bool bCouldUseCustomDepthStencil = (!Scene->World || (Scene->World->WorldType != EWorldType::Inactive));
 		FMobileCustomDepthStencilUsage CustomDepthStencilUsage;
 		for (int32 ViewIndex = 0; ViewIndex < Views.Num(); ViewIndex++)
 		{
