First Person: Changed First Person Self-Shadows implementation to use the depth test to early out on non-FP pixels instead of stencil. This appears to be marginally faster. Also replaced DrawScreenPass with the FPixelShaderUtils::AddFullscreenPass API for the simpler setup and improved readability in both the .cpp file and the shaders. Finally, this CL unified the naming of the FP self shadow cvars (*Shadow instead of *Shadows) and changed some of them to have a better default value out of the box.

diff --git a/Engine/Shaders/Private/FirstPersonSelfShadow.usf b/Engine/Shaders/Private/FirstPersonSelfShadow.usf
--- a/Engine/Shaders/Private/FirstPersonSelfShadow.usf
+++ b/Engine/Shaders/Private/FirstPersonSelfShadow.usf
@@ -41,8 +41,6 @@ float3 GetWorldNormal(uint2 PixelPos)
 uint CheckerboardMode; // 0: always min, 1: always max, 2: checkerboard
 
 void FirstPersonSelfShadowDownsamplePS(
-	noperspective float4 OutUVAndScreenPos : TEXCOORD0,
-	in FStereoPSInput StereoInput,
 	in float4 SVPos : SV_Position,
 	out float4 OutColor : SV_Target0,
 	out float OutDepth : SV_Depth
@@ -293,8 +291,6 @@ float TraceFirstPersonScreenSpaceShadows(FFirstPersonScreenTracingInputs Inputs,
 }
 
 void FirstPersonSelfShadowTracePS(
-	noperspective float4 UVAndScreenPos : TEXCOORD0,
-	FStereoPSInput StereoInput,
 	float4 SVPos : SV_Position,
 	out float4 OutColor : SV_Target0)
 {
@@ -378,8 +374,6 @@ void AccumulateSample(inout FSampleAccumulator Accumulator, FFirstPersonBlurInpu
 }
 
 void FirstPersonSelfShadowBlurPS(
-	noperspective float4 UVAndScreenPos : TEXCOORD0,
-	FStereoPSInput StereoInput,
 	float4 SVPos : SV_Position,
 	out float4 OutColor : SV_Target0)
 {
@@ -472,8 +466,6 @@ float4 GetDepthWeights(float FullResSceneDepth, float4 DownsampledDeviceZ)
 }
 
 void FirstPersonSelfShadowUpsamplePS(
-	noperspective float4 UVAndScreenPos : TEXCOORD0,
-	FStereoPSInput StereoInput,
 	float4 SVPos : SV_Position,
 	out float4 OutColor : SV_Target0)
 {
diff --git a/Engine/Source/Runtime/Renderer/Private/Shadows/FirstPersonSelfShadow.cpp b/Engine/Source/Runtime/Renderer/Private/Shadows/FirstPersonSelfShadow.cpp
--- a/Engine/Source/Runtime/Renderer/Private/Shadows/FirstPersonSelfShadow.cpp
+++ b/Engine/Source/Runtime/Renderer/Private/Shadows/FirstPersonSelfShadow.cpp
@@ -6,21 +6,18 @@
 #include "../LightRendering.h"
 #include "../ShadowRendering.h"
 #include "RenderGraphUtils.h"
-#include "ScreenPass.h"
+#include "PixelShaderUtils.h"
 #include "HZB.h"
 
-// Stencil value used to tag first person pixels in downsampled depth-stencil buffer.
-static const uint8 FirstPersonSelfShadowStencilRef = 1u;
-
-static TAutoConsoleVariable<int32> CVarFirstPersonSelfShadows(
+static TAutoConsoleVariable<int32> CVarFirstPersonSelfShadow(
 	TEXT("r.FirstPerson.SelfShadow"),
 	1,
 	TEXT("Enables self shadows for first person primitives. Self shadows are achieved with HZB screen space traces."),
 	ECVF_RenderThreadSafe);
 
-static TAutoConsoleVariable<int32> CVarFirstPersonSelfShadowsCheckerboardMode(
+static TAutoConsoleVariable<int32> CVarFirstPersonSelfShadowCheckerboardMode(
 	TEXT("r.FirstPerson.SelfShadow.DownsampleCheckerboardMode"),
-	2,
+	0,
 	TEXT("Controls how to downsample depth and normals for first person self shadows. 0: always pick closest depth, 1: always pick farthest depth, 2: alternate closest/farthest in a checkerboard pattern."),
 	ECVF_RenderThreadSafe);
 
@@ -30,31 +27,31 @@ static TAutoConsoleVariable<float> CVarFirstPersonSelfShadowMaxTraceDistance(
 	TEXT("Maximum world space trace distance for shadow rays."),
 	ECVF_RenderThreadSafe);
 
-static TAutoConsoleVariable<int32> CVarFirstPersonSelfShadowsMaxIterations(
+static TAutoConsoleVariable<int32> CVarFirstPersonSelfShadowMaxIterations(
 	TEXT("r.FirstPerson.SelfShadow.MaxHZBTraceIterations"),
 	512,
 	TEXT("Maximum number of HZB traversal iterations during the first person self shadow screen trace. Lowering this number can improve performance at the cost of potential shadow artifacts."),
 	ECVF_RenderThreadSafe);
 
-static TAutoConsoleVariable<float> CVarFirstPersonSelfShadowsRelativeDepthThickness(
+static TAutoConsoleVariable<float> CVarFirstPersonSelfShadowRelativeDepthThickness(
 	TEXT("r.FirstPerson.SelfShadow.RelativeDepthThickness"),
 	0.2f,
 	TEXT("Relative depth thickness behind which a screen space tracing hit is ignored."),
 	ECVF_RenderThreadSafe);
 
-static TAutoConsoleVariable<int32> CVarFirstPersonSelfShadowsMinimumHZBTraceOccupancy(
+static TAutoConsoleVariable<int32> CVarFirstPersonSelfShadowMinimumHZBTraceOccupancy(
 	TEXT("r.FirstPerson.SelfShadow.MinimumHZBTraceOccupancy"),
 	0,
 	TEXT("Minimum wave thread occupancy below which HZB tracing is aborted. Setting this to a value higher than 0 can improve performance at the cost of potential shadow artifacts."),
 	ECVF_RenderThreadSafe);
 
-static TAutoConsoleVariable<float> CVarFirstPersonSelfShadowsBilateralFilterDepthThreshold(
+static TAutoConsoleVariable<float> CVarFirstPersonSelfShadowBilateralFilterDepthThreshold(
 	TEXT("r.FirstPerson.SelfShadow.BilateralFilterDepthThreshold"),
-	0.01f,
+	1.0f,
 	TEXT("Scale applied to depth differences used to weigh sample contributions when filtering and upsampling first person self shadows. A higher value makes the result softer but may lead to leaking of light/shadow across geometric edges."),
 	ECVF_RenderThreadSafe);
 
-static TAutoConsoleVariable<int32> CVarFirstPersonSelfShadowsRawFullResolution(
+static TAutoConsoleVariable<int32> CVarFirstPersonSelfShadowRawFullResolution(
 	TEXT("r.FirstPerson.SelfShadow.RawFullResolution"),
 	0,
 	TEXT("Runs first person self shadows at full resolution and without filtering, resulting in pixel perfect shadows. This mostly serves as a ground truth to compare the half-resolution shadow to."),
@@ -199,7 +196,7 @@ IMPLEMENT_GLOBAL_SHADER(FFirstPersonSelfShadowUpsamplePS, "/Engine/Private/First
 bool ShouldRenderFirstPersonSelfShadow(const FSceneViewFamily& ViewFamily)
 {
 	return ViewFamily.EngineShowFlags.DirectLighting
-		&& CVarFirstPersonSelfShadows.GetValueOnRenderThread() != 0 
+		&& CVarFirstPersonSelfShadow.GetValueOnRenderThread() != 0 
 		&& HasFirstPersonGBufferBit(ViewFamily.GetShaderPlatform());
 }
 
@@ -215,7 +212,7 @@ FFirstPersonSelfShadowInputs CreateFirstPersonSelfShadowInputs(FRDGBuilder& Grap
 	FFirstPersonSelfShadowInputs Result{};
 	Result.SceneTextures = &SceneTextures;
 
-	const bool bRawFullResolution = CVarFirstPersonSelfShadowsRawFullResolution.GetValueOnRenderThread() != 0;
+	const bool bRawFullResolution = CVarFirstPersonSelfShadowRawFullResolution.GetValueOnRenderThread() != 0;
 	if (bRawFullResolution)
 	{
 		// No need to downsample anything in this case; just store a pointer to the scene textures for later
@@ -233,7 +230,7 @@ FFirstPersonSelfShadowInputs CreateFirstPersonSelfShadowInputs(FRDGBuilder& Grap
 
 		// Remap cvar (closest/furthest/CB depth) to min/max/CB for the shader.
 		uint32 CheckerboardMode = 0;
-		switch (FMath::Clamp(CVarFirstPersonSelfShadowsCheckerboardMode.GetValueOnRenderThread(), 0, 2))
+		switch (FMath::Clamp(CVarFirstPersonSelfShadowCheckerboardMode.GetValueOnRenderThread(), 0, 2))
 		{
 		case 0: CheckerboardMode = (bool)ERHIZBuffer::IsInverted ? 1 : 0; break;
 		case 1: CheckerboardMode = (bool)ERHIZBuffer::IsInverted ? 0 : 1; break;
@@ -245,31 +242,22 @@ FFirstPersonSelfShadowInputs CreateFirstPersonSelfShadowInputs(FRDGBuilder& Grap
 		PassParameters->SceneTextures = SceneTextures.UniformBuffer;
 		PassParameters->Substrate = Substrate::BindSubstrateGlobalUniformParameters(View);
 		PassParameters->CheckerboardMode = CheckerboardMode;
-		PassParameters->RenderTargets[0] = FRenderTargetBinding(DownsampledTextures.Normals, ERenderTargetLoadAction::EClear);
-		PassParameters->RenderTargets.DepthStencil = FDepthStencilBinding(DownsampledTextures.DepthStencil, ERenderTargetLoadAction::EClear, ERenderTargetLoadAction::EClear, FExclusiveDepthStencil::DepthWrite_StencilWrite);
+		PassParameters->RenderTargets[0] = FRenderTargetBinding(DownsampledTextures.Normals, ERenderTargetLoadAction::ENoAction);
+		PassParameters->RenderTargets.DepthStencil = FDepthStencilBinding(DownsampledTextures.DepthStencil, ERenderTargetLoadAction::EClear, ERenderTargetLoadAction::ENoAction, FExclusiveDepthStencil::DepthWrite_StencilNop);
 
 		TShaderMapRef<FFirstPersonSelfShadowInputsDownsamplePS> PixelShader(View.ShaderMap);
-		TShaderMapRef<FScreenPassVS> VertexShader(View.ShaderMap);
-
-		FRHIBlendState* BlendState = TStaticBlendState<>::GetRHI();
-		FRHIDepthStencilState* DepthStencilState = TStaticDepthStencilState<true, CF_Always, true, CF_Always, SO_Keep, SO_Keep, SO_Replace, true, CF_Always, SO_Keep, SO_Keep, SO_Replace>::GetRHI();
-
-		// Mark first person pixels with a stencil bit in the downsampled depth-stencil buffer so we can use the hardware to only do work for first person pixels.
-		const FScreenPassPipelineState PipelineState(VertexShader, PixelShader, BlendState, DepthStencilState, FirstPersonSelfShadowStencilRef);
 
-		ClearUnusedGraphResources(PixelShader, PassParameters);
-		GraphBuilder.AddPass(
+		// Set all non-first person pixels to the far depth value in the downsampled depth buffer so we can use the hardware to only do work for first person pixels.
+		FPixelShaderUtils::AddFullscreenPass(
+			GraphBuilder,
+			View.ShaderMap,
 			RDG_EVENT_NAME("FPSelfShadowsDownsampleDepthNormal (View: %i)", ViewIndex),
+			PixelShader,
 			PassParameters,
-			ERDGPassFlags::Raster,
-			[&View, PipelineState, PixelShader, PassParameters](FRDGAsyncTask, FRHICommandList& RHICmdList)
-			{
-				FScreenPassTextureViewport OutputViewport(PassParameters->RenderTargets.DepthStencil.GetTexture()->Desc.Extent);
-				DrawScreenPass(RHICmdList, View, OutputViewport, OutputViewport, PipelineState, EScreenPassDrawFlags::None, [&](FRHICommandList&)
-				{
-					SetShaderParameters(RHICmdList, PixelShader, PixelShader.GetPixelShader(), *PassParameters);
-				});
-			});
+			FIntRect(FIntPoint::ZeroValue, DownsampledTextures.Resolution),
+			nullptr /*BlendState*/,
+			nullptr /*RasterizerState*/,
+			TStaticDepthStencilState<true, CF_Always>::GetRHI());
 	}
 
 	return Result;
@@ -277,12 +265,13 @@ FFirstPersonSelfShadowInputs CreateFirstPersonSelfShadowInputs(FRDGBuilder& Grap
 
 void RenderFirstPersonSelfShadow(FRDGBuilder& GraphBuilder, const TArray<FViewInfo>& Views, FRDGTextureRef ScreenShadowMaskTexture, const FFirstPersonSelfShadowInputs& Inputs, const FLightSceneInfo& LightSceneInfo)
 {
+	static_assert((bool)ERHIZBuffer::IsInverted, "Inverted depth buffer is assumed because FPixelShaderUtils::AddFullscreenPass is drawing at depth 0!");
 	check(ScreenShadowMaskTexture);
 	const FLightSceneProxy* RESTRICT LightProxy = LightSceneInfo.Proxy;
 	const ELightComponentType LightType = (ELightComponentType)LightProxy->GetLightType();
 	const bool bIsRadial = LightType != LightType_Directional;
-	const bool bRawFullResolution = CVarFirstPersonSelfShadowsRawFullResolution.GetValueOnRenderThread() != 0;
-	const float InvDepthThreshold = 1.0f / FMath::Max(UE_SMALL_NUMBER, CVarFirstPersonSelfShadowsBilateralFilterDepthThreshold.GetValueOnRenderThread());
+	const bool bRawFullResolution = CVarFirstPersonSelfShadowRawFullResolution.GetValueOnRenderThread() != 0;
+	const float InvDepthThreshold = 1.0f / FMath::Max(UE_SMALL_NUMBER, CVarFirstPersonSelfShadowBilateralFilterDepthThreshold.GetValueOnRenderThread());
 
 	for (int32 ViewIndex = 0, ViewCount = Views.Num(); ViewIndex < ViewCount; ++ViewIndex)
 	{
@@ -318,15 +307,15 @@ void RenderFirstPersonSelfShadow(FRDGBuilder& GraphBuilder, const TArray<FViewIn
 			PassParameters->DeferredLight = DeferredLightUniformBuffer;
 			PassParameters->HZBParameters = GetHZBParameters(GraphBuilder, View, EHZBType::ClosestHZB);
 			PassParameters->HZBMaxTraceDistance = FMath::Max(0.0f, CVarFirstPersonSelfShadowMaxTraceDistance.GetValueOnRenderThread());
-			PassParameters->HZBMaxIterations = FMath::Max(1.0f, CVarFirstPersonSelfShadowsMaxIterations.GetValueOnRenderThread());
-			PassParameters->HZBRelativeDepthThickness = FMath::Max(UE_SMALL_NUMBER, CVarFirstPersonSelfShadowsRelativeDepthThickness.GetValueOnRenderThread());
-			PassParameters->HZBMinimumTracingThreadOccupancy = FMath::Max(0, CVarFirstPersonSelfShadowsMinimumHZBTraceOccupancy.GetValueOnRenderThread());
+			PassParameters->HZBMaxIterations = FMath::Max(1.0f, CVarFirstPersonSelfShadowMaxIterations.GetValueOnRenderThread());
+			PassParameters->HZBRelativeDepthThickness = FMath::Max(UE_SMALL_NUMBER, CVarFirstPersonSelfShadowRelativeDepthThickness.GetValueOnRenderThread());
+			PassParameters->HZBMinimumTracingThreadOccupancy = FMath::Max(0, CVarFirstPersonSelfShadowMinimumHZBTraceOccupancy.GetValueOnRenderThread());
 			if (!bRawFullResolution)
 			{
 				PassParameters->DownsampledInputsTexture = GraphBuilder.CreateSRV(Inputs.DownsampledInputs[ViewIndex].Normals);
 				PassParameters->DownsampledDepthTexture = GraphBuilder.CreateSRV(Inputs.DownsampledInputs[ViewIndex].DepthStencil);
 				PassParameters->RenderTargets[0] = FRenderTargetBinding(FPShadowsTexture, ERenderTargetLoadAction::ENoAction);
-				PassParameters->RenderTargets.DepthStencil = FDepthStencilBinding(Inputs.DownsampledInputs[ViewIndex].DepthStencil, ERenderTargetLoadAction::ELoad, ERenderTargetLoadAction::ELoad, FExclusiveDepthStencil::DepthNop_StencilRead);
+				PassParameters->RenderTargets.DepthStencil = FDepthStencilBinding(Inputs.DownsampledInputs[ViewIndex].DepthStencil, ERenderTargetLoadAction::ELoad, ERenderTargetLoadAction::ENoAction, FExclusiveDepthStencil::DepthRead_StencilNop);
 			}
 			else
 			{
@@ -339,15 +328,14 @@ void RenderFirstPersonSelfShadow(FRDGBuilder& GraphBuilder, const TArray<FViewIn
 			PermutationVector.Set<FFirstPersonSelfShadowTracingPS::FSourceShapeDim>(bIsRadial ? (LightProxy->IsRectLight() ? EFPLightSourceShape::Rect : EFPLightSourceShape::Point) : EFPLightSourceShape::Directional);
 			PermutationVector.Set<FFirstPersonSelfShadowTracingPS::FRawFullResolution>(bRawFullResolution);
 			TShaderMapRef<FFirstPersonSelfShadowTracingPS> PixelShader(View.ShaderMap, PermutationVector);
-			TShaderMapRef<FScreenPassVS> VertexShader(View.ShaderMap);
 
 			FRHIBlendState* BlendState = nullptr;
 			FRHIDepthStencilState* DepthStencilState = nullptr;
 			if (!bRawFullResolution)
 			{
-				// Default blend state, but early out on stencil such that we only process first person pixels.
+				// Default blend state, but early out on depth such that we only process first person pixels.
 				BlendState = TStaticBlendState<>::GetRHI();
-				DepthStencilState = TStaticDepthStencilState<false, CF_Always, true, CF_Equal, SO_Keep, SO_Keep, SO_Keep, true, CF_Equal, SO_Keep, SO_Keep, SO_Keep>::GetRHI();
+				DepthStencilState = TStaticDepthStencilState<false, CF_NotEqual>::GetRHI();
 			}
 			else
 			{
@@ -356,21 +344,16 @@ void RenderFirstPersonSelfShadow(FRDGBuilder& GraphBuilder, const TArray<FViewIn
 				DepthStencilState = TStaticDepthStencilState<false, CF_Always>::GetRHI();
 			}
 
-			const FScreenPassPipelineState PipelineState(VertexShader, PixelShader, BlendState, DepthStencilState, FirstPersonSelfShadowStencilRef);
-
-			ClearUnusedGraphResources(PixelShader, PassParameters);
-			GraphBuilder.AddPass(
+			FPixelShaderUtils::AddFullscreenPass(
+				GraphBuilder,
+				View.ShaderMap,
 				RDG_EVENT_NAME("FirstPersonSelfShadowTracing: %s (View: %i)", *LightProxy->GetOwnerNameOrLabel(), ViewIndex),
+				PixelShader,
 				PassParameters,
-				ERDGPassFlags::Raster,
-				[&View, PipelineState, PixelShader, PassParameters](FRDGAsyncTask, FRHICommandList& RHICmdList)
-				{
-					FScreenPassTextureViewport OutputViewport(PassParameters->RenderTargets[0].GetTexture()->Desc.Extent);
-					DrawScreenPass(RHICmdList, View, OutputViewport, OutputViewport, PipelineState, EScreenPassDrawFlags::None, [&](FRHICommandList&)
-					{
-						SetShaderParameters(RHICmdList, PixelShader, PixelShader.GetPixelShader(), *PassParameters);
-					});
-				});
+				!bRawFullResolution ? FIntRect(FIntPoint::ZeroValue, DownsampledResolution) : View.ViewRect,
+				BlendState,
+				nullptr /*RasterizerState*/,
+				DepthStencilState);
 		}
 
 		// Apply a 3x3 blur with some contact hardening depending on shadow caster distance. This helps with achieving a nicer upsampled shadow while still giving some small scale details.
@@ -382,30 +365,21 @@ void RenderFirstPersonSelfShadow(FRDGBuilder& GraphBuilder, const TArray<FViewIn
 			PassParameters->DepthTexture = GraphBuilder.CreateSRV(Inputs.DownsampledInputs[ViewIndex].DepthStencil);
 			PassParameters->InvDepthThreshold = InvDepthThreshold;
 			PassParameters->RenderTargets[0] = FRenderTargetBinding(FPDenoisedShadowsTexture, ERenderTargetLoadAction::ENoAction);
-			PassParameters->RenderTargets.DepthStencil = FDepthStencilBinding(Inputs.DownsampledInputs[ViewIndex].DepthStencil, ERenderTargetLoadAction::ELoad, ERenderTargetLoadAction::ELoad, FExclusiveDepthStencil::DepthNop_StencilRead);
+			PassParameters->RenderTargets.DepthStencil = FDepthStencilBinding(Inputs.DownsampledInputs[ViewIndex].DepthStencil, ERenderTargetLoadAction::ELoad, ERenderTargetLoadAction::ENoAction, FExclusiveDepthStencil::DepthRead_StencilNop);
 
 			TShaderMapRef<FFirstPersonSelfShadowBlurPS> PixelShader(View.ShaderMap);
-			TShaderMapRef<FScreenPassVS> VertexShader(View.ShaderMap);
-
-			// Default blend state, but early out on stencil such that we only process first person pixels.
-			FRHIBlendState* BlendState = TStaticBlendState<>::GetRHI();
-			FRHIDepthStencilState* DepthStencilState = TStaticDepthStencilState<false, CF_Always, true, CF_Equal, SO_Keep, SO_Keep, SO_Keep, true, CF_Equal, SO_Keep, SO_Keep, SO_Keep>::GetRHI();
-
-			const FScreenPassPipelineState PipelineState(VertexShader, PixelShader, BlendState, DepthStencilState, FirstPersonSelfShadowStencilRef);
 
-			ClearUnusedGraphResources(PixelShader, PassParameters);
-			GraphBuilder.AddPass(
+			// Default blend state, but early out on depth such that we only process first person pixels.
+			FPixelShaderUtils::AddFullscreenPass(
+				GraphBuilder,
+				View.ShaderMap,
 				RDG_EVENT_NAME("FirstPersonSelfShadowBlur: %s (View: %i)", *LightProxy->GetOwnerNameOrLabel(), ViewIndex),
+				PixelShader,
 				PassParameters,
-				ERDGPassFlags::Raster,
-				[&View, PipelineState, PixelShader, PassParameters](FRDGAsyncTask, FRHICommandList& RHICmdList)
-				{
-					FScreenPassTextureViewport OutputViewport(PassParameters->RenderTargets.DepthStencil.GetTexture()->Desc.Extent);
-					DrawScreenPass(RHICmdList, View, OutputViewport, OutputViewport, PipelineState, EScreenPassDrawFlags::None, [&](FRHICommandList&)
-					{
-						SetShaderParameters(RHICmdList, PixelShader, PixelShader.GetPixelShader(), *PassParameters);
-					});
-				});
+				FIntRect(FIntPoint::ZeroValue, DownsampledResolution),
+				nullptr /*BlendState*/,
+				nullptr /*RasterizerState*/,
+				TStaticDepthStencilState<false, CF_NotEqual>::GetRHI());
 		}
 
 		// Upsample into ScreenShadowMaskTexture
@@ -422,19 +396,17 @@ void RenderFirstPersonSelfShadow(FRDGBuilder& GraphBuilder, const TArray<FViewIn
 			PassParameters->RenderTargets[0] = FRenderTargetBinding(ScreenShadowMaskTexture, ERenderTargetLoadAction::ELoad);
 
 			TShaderMapRef<FFirstPersonSelfShadowUpsamplePS> PixelShader(View.ShaderMap);
-			TShaderMapRef<FScreenPassVS> VertexShader(View.ShaderMap);
 
 			FRHIBlendState* BlendState = FProjectedShadowInfo::GetBlendStateForProjection(LightSceneInfo.GetDynamicShadowMapChannel(), LightType == LightType_Directional, false, false, false);
 
-			AddDrawScreenPass(GraphBuilder,
+			FPixelShaderUtils::AddFullscreenPass(
+				GraphBuilder,
+				View.ShaderMap,
 				RDG_EVENT_NAME("FirstPersonSelfShadowUpsample: %s (View: %i)", *LightProxy->GetOwnerNameOrLabel(), ViewIndex),
-				FScreenPassViewInfo(View),
-				FScreenPassTextureViewport(ScreenShadowMaskTexture, View.ViewRect),
-				FScreenPassTextureViewport(ScreenShadowMaskTexture, View.ViewRect),
-				VertexShader,
 				PixelShader,
-				BlendState,
-				PassParameters);
+				PassParameters,
+				View.ViewRect,
+				BlendState);
 		}
 	}
 }
