[CIS] Fix CIS warning about redundant if-statements in LIGHTMAPPED_SHADER_PIPELINE_FILL() macro.

#rnx

diff --git a/Engine/Source/Runtime/Renderer/Private/BasePassRendering.cpp b/Engine/Source/Runtime/Renderer/Private/BasePassRendering.cpp
--- a/Engine/Source/Runtime/Renderer/Private/BasePassRendering.cpp
+++ b/Engine/Source/Runtime/Renderer/Private/BasePassRendering.cpp
@@ -197,34 +197,38 @@ DECLARE_CYCLE_STAT(TEXT("DeferredShadingSceneRenderer ClearGBufferAtMaxZ"), STAT
 DECLARE_CYCLE_STAT(TEXT("DeferredShadingSceneRenderer ViewExtensionPostRenderBasePass"), STAT_FDeferredShadingSceneRenderer_ViewExtensionPostRenderBasePass, STATGROUP_SceneRendering);
 
 DEFINE_GPU_STAT(NaniteBasePass);
-#define LIGHTMAPPED_SHADER_PIPELINE_FILL(Policy, DesiredPolicy, IndirectPolicyType, DesiredIndirectPolicyType, DesiredPolicyName) \
-	constexpr (std::is_same<Policy, DesiredPolicy>::value) \
+
+#define LIGHTMAPPED_SHADER_PIPELINE_FILL_INTERNAL(Condition, DesiredPolicy, DesiredPolicyName) \
+	constexpr (Condition) \
 	{ \
-		if (IndirectPolicyType == DesiredIndirectPolicyType) \
+		if (bUse128bitRT) \
+		{ \
+			checkf(bEnableSkyLight == false && GBufferLayout == GBL_Default, TEXT("New 128bitRT Base Pass was added? / \
+						Currently supporting just the F128BitRTBasePassPS. Handle new ones analogically.")); \
+			ShaderTypes.PipelineType = &(BasePassPipeline128BitRT); \
+		} \
+		else if (bEnableSkyLight) \
 		{ \
-			if (bUse128bitRT) \
-			{ \
-				checkf(bEnableSkyLight == false && GBufferLayout == GBL_Default, TEXT("New 128bitRT Base Pass was added? / \
-							Currently supporting just the F128BitRTBasePassPS. Handle new ones analogically.")); \
-				ShaderTypes.PipelineType = &(BasePassPipeline128BitRT); \
-			} \
-			else if (bEnableSkyLight) \
-			{ \
-				if (GBufferLayout == GBL_Default) \
-					ShaderTypes.PipelineType = &(BasePassPipeline##DesiredPolicyName##trueSkylightGBL_Default); \
-				else if (GBufferLayout == GBL_ForceVelocity) \
-					ShaderTypes.PipelineType = &(BasePassPipeline##DesiredPolicyName##trueSkylightGBL_ForceVelocityForceVelocity); \
-			} \
-			else \
-			{ \
-				if (GBufferLayout == GBL_Default) \
-					ShaderTypes.PipelineType = &(BasePassPipeline##DesiredPolicyName##falseGBL_Default); \
-				else if (GBufferLayout == GBL_ForceVelocity) \
-					ShaderTypes.PipelineType = &(BasePassPipeline##DesiredPolicyName##falseGBL_ForceVelocityForceVelocity); \
-			} \
+			if (GBufferLayout == GBL_Default) \
+				ShaderTypes.PipelineType = &(BasePassPipeline##DesiredPolicyName##trueSkylightGBL_Default); \
+			else if (GBufferLayout == GBL_ForceVelocity) \
+				ShaderTypes.PipelineType = &(BasePassPipeline##DesiredPolicyName##trueSkylightGBL_ForceVelocityForceVelocity); \
+		} \
+		else \
+		{ \
+			if (GBufferLayout == GBL_Default) \
+				ShaderTypes.PipelineType = &(BasePassPipeline##DesiredPolicyName##falseGBL_Default); \
+			else if (GBufferLayout == GBL_ForceVelocity) \
+				ShaderTypes.PipelineType = &(BasePassPipeline##DesiredPolicyName##falseGBL_ForceVelocityForceVelocity); \
 		} \
 	}
 
+#define LIGHTMAPPED_SHADER_PIPELINE_FILL(Policy, DesiredPolicy, DesiredPolicyName) \
+	LIGHTMAPPED_SHADER_PIPELINE_FILL_INTERNAL((std::is_same<Policy, DesiredPolicy>::value), DesiredPolicy, DesiredPolicyName)
+
+#define LIGHTMAPPED_SHADER_PIPELINE_FILL_INDIRECT_POLICY_TYPES(Policy, DesiredPolicy, DesiredPolicyName, IndirectPolicyType, DesiredIndirectPolicyType) \
+	LIGHTMAPPED_SHADER_PIPELINE_FILL_INTERNAL((std::is_same<Policy, DesiredPolicy>::value && IndirectPolicyType == DesiredIndirectPolicyType), DesiredPolicy, DesiredPolicyName)
+
 static bool IsStandardTranslucenyPassSeparated()
 {
 	static const auto TranslucencyStandardSeparatedCVar = IConsoleManager::Get().FindTConsoleVariableDataInt(TEXT("r.Translucency.StandardSeparated"));
@@ -636,13 +640,13 @@ bool GetUniformBasePassShaders(
 {
 	FMaterialShaderTypes ShaderTypes;
 
-	if LIGHTMAPPED_SHADER_PIPELINE_FILL(TUniformLightMapPolicy<Policy>, TUniformLightMapPolicy<LMP_NO_LIGHTMAP>, LMP_DUMMY, LMP_DUMMY, FNoLightMapPolicy)
-	else if LIGHTMAPPED_SHADER_PIPELINE_FILL(TUniformLightMapPolicy<Policy>, TUniformLightMapPolicy<LMP_PRECOMPUTED_IRRADIANCE_VOLUME_INDIRECT_LIGHTING>, LMP_DUMMY, LMP_DUMMY, FPrecomputedVolumetricLightmapLightingPolicy)
-	else if LIGHTMAPPED_SHADER_PIPELINE_FILL(TUniformLightMapPolicy<Policy>, TUniformLightMapPolicy<LMP_CACHED_VOLUME_INDIRECT_LIGHTING>, LMP_DUMMY, LMP_DUMMY, FCachedVolumeIndirectLightingPolicy)
-	else if LIGHTMAPPED_SHADER_PIPELINE_FILL(TUniformLightMapPolicy<Policy>, TUniformLightMapPolicy<LMP_CACHED_POINT_INDIRECT_LIGHTING>, LMP_DUMMY, LMP_DUMMY, FCachedPointIndirectLightingPolicy)
-	else if LIGHTMAPPED_SHADER_PIPELINE_FILL(TUniformLightMapPolicy<Policy>, TUniformLightMapPolicy<LMP_LQ_LIGHTMAP>, LMP_DUMMY, LMP_DUMMY, TLightMapPolicyLQ)
-	else if LIGHTMAPPED_SHADER_PIPELINE_FILL(TUniformLightMapPolicy<Policy>, TUniformLightMapPolicy<LMP_HQ_LIGHTMAP>, LMP_DUMMY, LMP_DUMMY, TLightMapPolicyHQ)
-	else if LIGHTMAPPED_SHADER_PIPELINE_FILL(TUniformLightMapPolicy<Policy>, TUniformLightMapPolicy<LMP_DISTANCE_FIELD_SHADOWS_AND_HQ_LIGHTMAP>, LMP_DUMMY, LMP_DUMMY, TDistanceFieldShadowsAndLightMapPolicyHQ)
+	if LIGHTMAPPED_SHADER_PIPELINE_FILL(TUniformLightMapPolicy<Policy>, TUniformLightMapPolicy<LMP_NO_LIGHTMAP>, FNoLightMapPolicy)
+	else if LIGHTMAPPED_SHADER_PIPELINE_FILL(TUniformLightMapPolicy<Policy>, TUniformLightMapPolicy<LMP_PRECOMPUTED_IRRADIANCE_VOLUME_INDIRECT_LIGHTING>, FPrecomputedVolumetricLightmapLightingPolicy)
+	else if LIGHTMAPPED_SHADER_PIPELINE_FILL(TUniformLightMapPolicy<Policy>, TUniformLightMapPolicy<LMP_CACHED_VOLUME_INDIRECT_LIGHTING>, FCachedVolumeIndirectLightingPolicy)
+	else if LIGHTMAPPED_SHADER_PIPELINE_FILL(TUniformLightMapPolicy<Policy>, TUniformLightMapPolicy<LMP_CACHED_POINT_INDIRECT_LIGHTING>, FCachedPointIndirectLightingPolicy)
+	else if LIGHTMAPPED_SHADER_PIPELINE_FILL(TUniformLightMapPolicy<Policy>, TUniformLightMapPolicy<LMP_LQ_LIGHTMAP>, TLightMapPolicyLQ)
+	else if LIGHTMAPPED_SHADER_PIPELINE_FILL(TUniformLightMapPolicy<Policy>, TUniformLightMapPolicy<LMP_HQ_LIGHTMAP>, TLightMapPolicyHQ)
+	else if LIGHTMAPPED_SHADER_PIPELINE_FILL(TUniformLightMapPolicy<Policy>, TUniformLightMapPolicy<LMP_DISTANCE_FIELD_SHADOWS_AND_HQ_LIGHTMAP>, TDistanceFieldShadowsAndLightMapPolicyHQ)
 	
 	if (VertexShader)
 	{
@@ -691,16 +695,16 @@ bool GetBasePassShaders(
 {
 	FMaterialShaderTypes ShaderTypes;
 
-	if LIGHTMAPPED_SHADER_PIPELINE_FILL(LightMapPolicyType, FSelfShadowedTranslucencyPolicy, LMP_DUMMY, LMP_DUMMY, FSelfShadowedTranslucencyPolicy)
-	else if LIGHTMAPPED_SHADER_PIPELINE_FILL(LightMapPolicyType, FSelfShadowedCachedPointIndirectLightingPolicy, LMP_DUMMY, LMP_DUMMY, FSelfShadowedCachedPointIndirectLightingPolicy) 
-	else if LIGHTMAPPED_SHADER_PIPELINE_FILL(LightMapPolicyType, FSelfShadowedVolumetricLightmapPolicy, LMP_DUMMY, LMP_DUMMY, FSelfShadowedVolumetricLightmapPolicy)
-	else if LIGHTMAPPED_SHADER_PIPELINE_FILL(LightMapPolicyType, FUniformLightMapPolicy, LightMapPolicy.GetIndirectPolicy(), LMP_NO_LIGHTMAP, FNoLightMapPolicy)
-	else if LIGHTMAPPED_SHADER_PIPELINE_FILL(LightMapPolicyType, FUniformLightMapPolicy, LightMapPolicy.GetIndirectPolicy(), LMP_PRECOMPUTED_IRRADIANCE_VOLUME_INDIRECT_LIGHTING, FPrecomputedVolumetricLightmapLightingPolicy)
-	else if LIGHTMAPPED_SHADER_PIPELINE_FILL(LightMapPolicyType, FUniformLightMapPolicy, LightMapPolicy.GetIndirectPolicy(), LMP_CACHED_VOLUME_INDIRECT_LIGHTING, FCachedVolumeIndirectLightingPolicy)
-	else if LIGHTMAPPED_SHADER_PIPELINE_FILL(LightMapPolicyType, FUniformLightMapPolicy, LightMapPolicy.GetIndirectPolicy(), LMP_CACHED_POINT_INDIRECT_LIGHTING, FCachedPointIndirectLightingPolicy)
-	else if LIGHTMAPPED_SHADER_PIPELINE_FILL(LightMapPolicyType, FUniformLightMapPolicy, LightMapPolicy.GetIndirectPolicy(), LMP_LQ_LIGHTMAP, TLightMapPolicyLQ)
-	else if LIGHTMAPPED_SHADER_PIPELINE_FILL(LightMapPolicyType, FUniformLightMapPolicy, LightMapPolicy.GetIndirectPolicy(), LMP_HQ_LIGHTMAP, TLightMapPolicyHQ)
-	else if LIGHTMAPPED_SHADER_PIPELINE_FILL(LightMapPolicyType, FUniformLightMapPolicy, LightMapPolicy.GetIndirectPolicy(), LMP_DISTANCE_FIELD_SHADOWS_AND_HQ_LIGHTMAP, TDistanceFieldShadowsAndLightMapPolicyHQ)
+	if LIGHTMAPPED_SHADER_PIPELINE_FILL(LightMapPolicyType, FSelfShadowedTranslucencyPolicy, FSelfShadowedTranslucencyPolicy)
+	else if LIGHTMAPPED_SHADER_PIPELINE_FILL(LightMapPolicyType, FSelfShadowedCachedPointIndirectLightingPolicy, FSelfShadowedCachedPointIndirectLightingPolicy) 
+	else if LIGHTMAPPED_SHADER_PIPELINE_FILL(LightMapPolicyType, FSelfShadowedVolumetricLightmapPolicy, FSelfShadowedVolumetricLightmapPolicy)
+	else if LIGHTMAPPED_SHADER_PIPELINE_FILL_INDIRECT_POLICY_TYPES(LightMapPolicyType, FUniformLightMapPolicy, FNoLightMapPolicy, LightMapPolicy.GetIndirectPolicy(), LMP_NO_LIGHTMAP)
+	else if LIGHTMAPPED_SHADER_PIPELINE_FILL_INDIRECT_POLICY_TYPES(LightMapPolicyType, FUniformLightMapPolicy, FPrecomputedVolumetricLightmapLightingPolicy, LightMapPolicy.GetIndirectPolicy(), LMP_PRECOMPUTED_IRRADIANCE_VOLUME_INDIRECT_LIGHTING)
+	else if LIGHTMAPPED_SHADER_PIPELINE_FILL_INDIRECT_POLICY_TYPES(LightMapPolicyType, FUniformLightMapPolicy, FCachedVolumeIndirectLightingPolicy, LightMapPolicy.GetIndirectPolicy(), LMP_CACHED_VOLUME_INDIRECT_LIGHTING)
+	else if LIGHTMAPPED_SHADER_PIPELINE_FILL_INDIRECT_POLICY_TYPES(LightMapPolicyType, FUniformLightMapPolicy, FCachedPointIndirectLightingPolicy, LightMapPolicy.GetIndirectPolicy(), LMP_CACHED_POINT_INDIRECT_LIGHTING)
+	else if LIGHTMAPPED_SHADER_PIPELINE_FILL_INDIRECT_POLICY_TYPES(LightMapPolicyType, FUniformLightMapPolicy, TLightMapPolicyLQ, LightMapPolicy.GetIndirectPolicy(), LMP_LQ_LIGHTMAP)
+	else if LIGHTMAPPED_SHADER_PIPELINE_FILL_INDIRECT_POLICY_TYPES(LightMapPolicyType, FUniformLightMapPolicy, TLightMapPolicyHQ, LightMapPolicy.GetIndirectPolicy(), LMP_HQ_LIGHTMAP)
+	else if LIGHTMAPPED_SHADER_PIPELINE_FILL_INDIRECT_POLICY_TYPES(LightMapPolicyType, FUniformLightMapPolicy, TDistanceFieldShadowsAndLightMapPolicyHQ, LightMapPolicy.GetIndirectPolicy(), LMP_DISTANCE_FIELD_SHADOWS_AND_HQ_LIGHTMAP)
 
 	if (VertexShader)
 	{
