[HWRT] Refactored FRayTracingDynamicGeometryUpdateManager to remove View from AddDynamicGeometryUpdatePass(...) parameters
- keep track of the view uniform buffer used by each dispatch command when the update is requested and dispatch work using one pass per view.

diff --git a/Engine/Plugins/Experimental/GPULightmass/Source/GPULightmass/Private/LightmapRenderer.cpp b/Engine/Plugins/Experimental/GPULightmass/Source/GPULightmass/Private/LightmapRenderer.cpp
--- a/Engine/Plugins/Experimental/GPULightmass/Source/GPULightmass/Private/LightmapRenderer.cpp
+++ b/Engine/Plugins/Experimental/GPULightmass/Source/GPULightmass/Private/LightmapRenderer.cpp
@@ -1100,9 +1100,10 @@ bool FSceneRenderState::SetupRayTracingScene(FRDGBuilder& GraphBuilder, FSceneUn
 								StaticMeshInstanceRenderStates.Elements.Num() + InstanceGroupRenderStates.Elements.Num() + LandscapeIndex
 							);
 
+							PRAGMA_DISABLE_DEPRECATION_WARNINGS
+
 							DynamicGeometryUpdateManager.Update(*ReferenceView.Get());
 
-							PRAGMA_DISABLE_DEPRECATION_WARNINGS
 							const uint32 BLASScratchSize = DynamicGeometryUpdateManager.ComputeScratchBufferSize();
 
 							const FRHIBufferCreateDesc CreateDesc =
@@ -1114,6 +1115,8 @@ bool FSceneRenderState::SetupRayTracingScene(FRDGBuilder& GraphBuilder, FSceneUn
 
 							RHICmdList.SetStaticUniformBuffers({ReferenceView->ViewUniformBuffer, SceneUniforms.GetBufferRHI(GraphBuilder)});
 							DynamicGeometryUpdateManager.DispatchUpdates(RHICmdList, ScratchBuffer);
+							DynamicGeometryUpdateManager.EndUpdate();
+
 							PRAGMA_ENABLE_DEPRECATION_WARNINGS
 
 							// Landscape VF doesn't really use the vertex buffer in HitGroupSystemParameters
diff --git a/Engine/Source/Runtime/Renderer/Private/DeferredShadingRenderer.cpp b/Engine/Source/Runtime/Renderer/Private/DeferredShadingRenderer.cpp
--- a/Engine/Source/Runtime/Renderer/Private/DeferredShadingRenderer.cpp
+++ b/Engine/Source/Runtime/Renderer/Private/DeferredShadingRenderer.cpp
@@ -885,7 +885,7 @@ bool FDeferredShadingSceneRenderer::DispatchRayTracingWorldUpdates(FRDGBuilder&
 	RDG_GPU_MASK_SCOPE(GraphBuilder, GRayTracingMultiGpuTLASMask ? GraphBuilder.RHICmdList.GetGPUMask() : FRHIGPUMask::All());
 
 	FRayTracingDynamicGeometryUpdateManager* DynamicGeometryUpdateManager = (FRayTracingDynamicGeometryUpdateManager*)Scene->GetRayTracingDynamicGeometryUpdateManager();
-	DynamicGeometryUpdateManager->AddDynamicGeometryUpdatePass(ReferenceView, GraphBuilder, ComputePassFlags, ResourceAccessPipelines, OutDynamicGeometryScratchBuffer);
+	DynamicGeometryUpdateManager->AddDynamicGeometryUpdatePass(GraphBuilder, ComputePassFlags, GetSceneUniformBufferRef(GraphBuilder), RayTracingScene.bTracingFeedbackEnabled, ResourceAccessPipelines, OutDynamicGeometryScratchBuffer);
 
 	((FRayTracingGeometryManager*)GRayTracingGeometryManager)->ResetVisibleGeometries();
 
diff --git a/Engine/Source/Runtime/Renderer/Private/RayTracing/RayTracingDynamicGeometryUpdateManager.cpp b/Engine/Source/Runtime/Renderer/Private/RayTracing/RayTracingDynamicGeometryUpdateManager.cpp
--- a/Engine/Source/Runtime/Renderer/Private/RayTracing/RayTracingDynamicGeometryUpdateManager.cpp
+++ b/Engine/Source/Runtime/Renderer/Private/RayTracing/RayTracingDynamicGeometryUpdateManager.cpp
@@ -217,8 +217,9 @@ FRayTracingDynamicGeometryUpdateManager::~FRayTracingDynamicGeometryUpdateManage
 
 void FRayTracingDynamicGeometryUpdateManager::Clear()
 {
+	DispatchCommandsPerView = {};
+
 	// Clear working arrays - keep max size allocated
-	DispatchCommands.Empty(DispatchCommands.Max());
 	BuildParams.Empty(BuildParams.Max());
 	Segments.Empty(Segments.Max());
 
@@ -230,7 +231,7 @@ void FRayTracingDynamicGeometryUpdateManager::Clear()
 
 int64 FRayTracingDynamicGeometryUpdateManager::BeginUpdate()
 {
-	check(DispatchCommands.IsEmpty());
+	check(DispatchCommandsPerView.IsEmpty());
 	check(BuildParams.IsEmpty());
 	check(Segments.IsEmpty());
 	check(ReferencedUniformBuffers.IsEmpty());
@@ -353,6 +354,7 @@ void FRayTracingDynamicGeometryUpdateManager::AddDynamicGeometryToUpdate(
 	check(IsAligned(VertexBufferOffset, RHI_RAW_VIEW_ALIGNMENT));
 
 	FRayTracingDynamicGeometryBuildParams GeometryBuildParams;
+	GeometryBuildParams.ViewUniformBuffer = View->ViewUniformBuffer;
 	GeometryBuildParams.DispatchCommands.Reserve(UpdateParams.MeshBatches.Num());
 
 	const int32 PSOCollectorIndex = FPSOCollectorCreateManager::GetIndex(EShadingPath::Deferred, RayTracingDynamicGeometryPSOCollectorName);
@@ -543,7 +545,7 @@ BEGIN_SHADER_PARAMETER_STRUCT(FRayTracingDynamicGeometryUpdatePassParams, )
 	SHADER_PARAMETER_RDG_UNIFORM_BUFFER(FSceneUniformParameters, Scene)
 END_SHADER_PARAMETER_STRUCT()
 
-void FRayTracingDynamicGeometryUpdateManager::Update(const FViewInfo& View)
+void FRayTracingDynamicGeometryUpdateManager::ScheduleUpdates(bool bUseTracingFeedback)
 {
 	TRACE_CPUPROFILER_EVENT_SCOPE(FRayTracingDynamicGeometryUpdateManager::Update);
 
@@ -584,10 +586,16 @@ void FRayTracingDynamicGeometryUpdateManager::Update(const FViewInfo& View)
 		return;
 	}
 
-	checkf(DispatchCommands.IsEmpty(), TEXT("DispatchCommands is not empty. Previous frame updates were not dispatched."));
+	checkf(DispatchCommandsPerView.IsEmpty(), TEXT("DispatchCommandsPerView is not empty. Previous frame updates were not dispatched."));
 	checkf(BuildParams.IsEmpty(), TEXT("BuildParams is not empty. Previous frame updates were not dispatched."));
 
-	DispatchCommands.Reserve(TotalNumGeometryBuilds);
+	// reserve for worst case (might be wasteful if there are too many views)
+	for (FRHIUniformBuffer* ViewUniformBuffer : ReferencedUniformBuffers)
+	{
+		TArray<FMeshComputeDispatchCommand>& ViewDispatchCommands = DispatchCommandsPerView.FindOrAdd(ViewUniformBuffer);
+		ViewDispatchCommands.Reserve(TotalNumGeometryBuilds);
+	}
+
 	BuildParams.Reserve(TotalNumGeometryBuilds);
 
 	FRayTracingGeometrySegment* SegmentData = Segments.GetData();
@@ -597,40 +605,38 @@ void FRayTracingDynamicGeometryUpdateManager::Update(const FViewInfo& View)
 	uint32 BLASScratchSize = 0;
 	int32 NumBuildPrimitives = 0;
 
-	auto AddGeometryBuildParamsToBuildList = [](const FRayTracingDynamicGeometryBuildParams& BuildParams,
-		FRayTracingGeometrySegment* SegmentData,
-		uint32& OutScratchSize,
-		TArray<FMeshComputeDispatchCommand>& OutDispatchCommands,
-		TArray<FRayTracingGeometryBuildParams>& OutBuildParams
-		)
+	auto AddGeometryBuildParamsToBuildList = [this, SegmentData, &BLASScratchSize](const FRayTracingDynamicGeometryBuildParams& InBuildParams)
 		{
-			FRHIRayTracingGeometry* RayTracingGeometry = BuildParams.Geometry->GetRHI();
+			FRHIRayTracingGeometry* RayTracingGeometry = InBuildParams.Geometry->GetRHI();
 
 			const uint32 ScratchAlignment = GRHIRayTracingScratchBufferAlignment;
-			const uint32 NumBuildPrimitives = BuildParams.Geometry->Initializer.TotalPrimitiveCount;
+			const uint32 NumBuildPrimitives = InBuildParams.Geometry->Initializer.TotalPrimitiveCount;
 
 			const uint32 ScratchSize = RayTracingGeometry->GetSizeInfo().BuildScratchSize;
-			OutScratchSize = Align(OutScratchSize + ScratchSize, ScratchAlignment);
+			BLASScratchSize = Align(BLASScratchSize + ScratchSize, ScratchAlignment);
 
-			BuildParams.Geometry->NumUpdatesSinceLastBuild = 0;
+			InBuildParams.Geometry->NumUpdatesSinceLastBuild = 0;
 
-			FRayTracingGeometryBuildParams BuildParam;
-			BuildParam.Geometry = RayTracingGeometry;
-			BuildParam.BuildMode = EAccelerationStructureBuildMode::Build;
+			FRayTracingGeometryBuildParams RTGeoBuildParams;
+			RTGeoBuildParams.Geometry = RayTracingGeometry;
+			RTGeoBuildParams.BuildMode = EAccelerationStructureBuildMode::Build;
 
-			if (BuildParams.SegmentOffset >= 0)
+			if (InBuildParams.SegmentOffset >= 0)
 			{
-				BuildParam.Segments = MakeArrayView(&SegmentData[BuildParams.SegmentOffset], BuildParams.Geometry->Initializer.Segments.Num());
+				RTGeoBuildParams.Segments = MakeArrayView(&SegmentData[InBuildParams.SegmentOffset], InBuildParams.Geometry->Initializer.Segments.Num());
 			}
 
-			OutBuildParams.Add(MoveTemp(BuildParam));
+			BuildParams.Add(MoveTemp(RTGeoBuildParams));
 
-			OutDispatchCommands.Append(BuildParams.DispatchCommands);
+			if (!InBuildParams.DispatchCommands.IsEmpty())
+			{
+				DispatchCommandsPerView[InBuildParams.ViewUniformBuffer].Append(InBuildParams.DispatchCommands);
+			}
 		};
 
 	for (const FRayTracingDynamicGeometryBuildParams& Build : DynamicGeometryBuilds)
 	{
-		AddGeometryBuildParamsToBuildList(Build, SegmentData, BLASScratchSize, DispatchCommands, BuildParams);
+		AddGeometryBuildParamsToBuildList(Build);
 		NumBuildPrimitives += Build.Geometry->Initializer.TotalPrimitiveCount;
 	}
 
@@ -660,7 +666,7 @@ void FRayTracingDynamicGeometryUpdateManager::Update(const FViewInfo& View)
 		FRHIRayTracingGeometry* RayTracingGeometry = Update.Geometry->GetRHI();
 		const uint32 TotalPrimitiveCount = Update.Geometry->Initializer.TotalPrimitiveCount;
 
-		if (View.bRayTracingFeedbackEnabled && !GRayTracingGeometryManager->IsGeometryVisible(Update.Geometry->GetGeometryHandle()))
+		if (bUseTracingFeedback && !GRayTracingGeometryManager->IsGeometryVisible(Update.Geometry->GetGeometryHandle()))
 		{
 			INC_DWORD_STAT_BY(STAT_RayTracingDynamicSkippedPrimitives, TotalPrimitiveCount);
 			continue;
@@ -670,7 +676,7 @@ void FRayTracingDynamicGeometryUpdateManager::Update(const FViewInfo& View)
 		{
 			if (Update.Geometry->NumUpdatesSinceLastBuild > MinUpdatesSinceLastBuild && NumForceBuildPrimitives <= MaxForceBuildPrimitivesPerFrame)
 			{
-				AddGeometryBuildParamsToBuildList(Update, SegmentData, BLASScratchSize, DispatchCommands, BuildParams);
+				AddGeometryBuildParamsToBuildList(Update);
 				NumBuildPrimitives += TotalPrimitiveCount;
 				NumForceBuildPrimitives += TotalPrimitiveCount;
 				continue;
@@ -694,7 +700,10 @@ void FRayTracingDynamicGeometryUpdateManager::Update(const FViewInfo& View)
 		}
 		BuildParams.Add(MoveTemp(BuildParam));
 
-		DispatchCommands.Append(Update.DispatchCommands);
+		if (!Update.DispatchCommands.IsEmpty())
+		{
+			DispatchCommandsPerView[Update.ViewUniformBuffer].Append(Update.DispatchCommands);
+		}
 
 		if (NumUpdatedPrimitives > MaxUpdatePrimitivesPerFrame)
 		{
@@ -708,13 +717,18 @@ void FRayTracingDynamicGeometryUpdateManager::Update(const FViewInfo& View)
 	ScratchBufferSize = BLASScratchSize;
 }
 
-void FRayTracingDynamicGeometryUpdateManager::AddDynamicGeometryUpdatePass(const FViewInfo& View, FRDGBuilder& GraphBuilder, ERDGPassFlags ComputePassFlags, ERHIPipeline ResourceAccessPipelines, FRDGBufferRef& OutDynamicGeometryScratchBuffer)
+void FRayTracingDynamicGeometryUpdateManager::Update(const FViewInfo& View)
+{
+	ScheduleUpdates(View.bRayTracingFeedbackEnabled);
+}
+
+void FRayTracingDynamicGeometryUpdateManager::AddDynamicGeometryUpdatePass(FRDGBuilder& GraphBuilder, ERDGPassFlags ComputePassFlags, const TRDGUniformBufferRef<FSceneUniformParameters>& SceneUB, bool bUseTracingFeedback, ERHIPipeline ResourceAccessPipelines, FRDGBufferRef& OutDynamicGeometryScratchBuffer)
 {
 	RDG_GPU_MASK_SCOPE(GraphBuilder, FRHIGPUMask::All());
 	RDG_EVENT_SCOPE_STAT(GraphBuilder, RayTracingDynamicGeometry, "RayTracingDynamicGeometry");
 	RDG_GPU_STAT_SCOPE(GraphBuilder, RayTracingDynamicGeometry);
 
-	Update(View);
+	ScheduleUpdates(bUseTracingFeedback);
 
 	const uint32 ScratchAlignment = GRHIRayTracingScratchBufferAlignment;
 	const uint32 BLASScratchSize = ScratchBufferSize;
@@ -729,26 +743,70 @@ void FRayTracingDynamicGeometryUpdateManager::AddDynamicGeometryUpdatePass(const
 		OutDynamicGeometryScratchBuffer = GraphBuilder.CreateBuffer(ScratchBufferDesc, TEXT("DynamicGeometry.BLASSharedScratchBuffer"));
 	}
 
-	FRayTracingDynamicGeometryUpdatePassParams* PassParams = GraphBuilder.AllocParameters<FRayTracingDynamicGeometryUpdatePassParams>();
-	PassParams->View = View.GetShaderParameters();
-	PassParams->Scene = View.GetSceneUniforms().GetBuffer(GraphBuilder);
-	PassParams->DynamicGeometryScratchBuffer = OutDynamicGeometryScratchBuffer;	
+	const ERHIPipeline SrcResourceAccessPipelines = ComputePassFlags == ERDGPassFlags::AsyncCompute ? ERHIPipeline::AsyncCompute : ERHIPipeline::Graphics;
 
-	ERHIPipeline SrcResourceAccessPipelines = ComputePassFlags == ERDGPassFlags::AsyncCompute ? ERHIPipeline::AsyncCompute : ERHIPipeline::Graphics;
-
-	GraphBuilder.AddPass(RDG_EVENT_NAME("RayTracingDynamicUpdate"), PassParams, ComputePassFlags | ERDGPassFlags::NeverCull,
-		[this, SrcResourceAccessPipelines, ResourceAccessPipelines, PassParams](FRDGAsyncTask, FRHICommandList& RHICmdList)
+	for(TPair<FRHIUniformBuffer*, TArray<FMeshComputeDispatchCommand>>& ViewDispatchCommands : DispatchCommandsPerView)
+	{
+		if (ViewDispatchCommands.Value.IsEmpty())
 		{
-			FRHIBuffer* DynamicGeometryScratchBuffer = PassParams->DynamicGeometryScratchBuffer ? PassParams->DynamicGeometryScratchBuffer->GetRHI() : nullptr;
+			continue;
+		}
+
+		FRayTracingDynamicGeometryUpdatePassParams* PassParams = GraphBuilder.AllocParameters<FRayTracingDynamicGeometryUpdatePassParams>();
+		PassParams->View.View = TUniformBufferRef<FViewUniformShaderParameters>(ViewDispatchCommands.Key);
+		PassParams->View.InstancedView = TUniformBufferRef<FInstancedViewUniformShaderParameters>(nullptr); // TODO: Is this necessary?
+		PassParams->Scene = SceneUB;
+
+		// DynamicGeometryScratchBuffer is not directly used in this pass but set so RDG orders passes correctly
+		// (TODO: this might also prevent dispatches for different views from overlapping, so investigate better solution)
+		PassParams->DynamicGeometryScratchBuffer = OutDynamicGeometryScratchBuffer; 
 
+		GraphBuilder.AddPass(RDG_EVENT_NAME("RayTracingDynamicUpdate"), PassParams, ComputePassFlags | ERDGPassFlags::NeverCull,
+			[this, SrcResourceAccessPipelines, ResourceAccessPipelines, DispatchCommands = MakeArrayView(ViewDispatchCommands.Value)](FRHICommandList& RHICmdList)
+			{
+				DispatchUpdates(RHICmdList, DispatchCommands, SrcResourceAccessPipelines, ResourceAccessPipelines);
+			});
+	}
+
+	if (BuildParams.Num() > 0)
+	{
+		FRayTracingDynamicGeometryUpdatePassParams* PassParams = GraphBuilder.AllocParameters<FRayTracingDynamicGeometryUpdatePassParams>();
+		PassParams->View = {};
+		PassParams->Scene = nullptr;
+		PassParams->DynamicGeometryScratchBuffer = OutDynamicGeometryScratchBuffer;
+
+		GraphBuilder.AddPass(RDG_EVENT_NAME("RayTracingDynamicUpdateBuild"), PassParams, ComputePassFlags | ERDGPassFlags::NeverCull,
+			[this, PassParams](FRHICommandList& RHICmdList)
+			{
+				// Can't use parallel command list because we have to make sure we are not building BVH data
+				// on the same RTGeometry on multiple threads at the same time. Ideally move the build
+				// requests over to the RaytracingGeometry manager so they can be correctly scheduled
+				// with other build requests in the engine (see UE-106982)
+				SCOPED_DRAW_EVENT(RHICmdList, Build);
+
+				FRHIBufferRange ScratchBufferRange;
+				ScratchBufferRange.Buffer = PassParams->DynamicGeometryScratchBuffer ? PassParams->DynamicGeometryScratchBuffer->GetRHI() : nullptr;
+				ScratchBufferRange.Offset = 0;
+				RHICmdList.BuildAccelerationStructures(BuildParams, ScratchBufferRange);
+			});
+	}
+
+	// TODO: Is it safe to use a regular task that waits on FRDGBuilder::GetAsyncExecuteTask() here instead?
+	// which would allow the passes above to be tagged with FRDGAsyncTask
+	GraphBuilder.AddPostExecuteCallback([this]
+		{
 			PRAGMA_DISABLE_DEPRECATION_WARNINGS
-			DispatchUpdates(RHICmdList, DynamicGeometryScratchBuffer, SrcResourceAccessPipelines, ResourceAccessPipelines);
 			EndUpdate();
 			PRAGMA_ENABLE_DEPRECATION_WARNINGS
 		});
 }
 
-void FRayTracingDynamicGeometryUpdateManager::DispatchUpdates(FRHICommandList& RHICmdList, FRHIBuffer* ScratchBuffer, ERHIPipeline SrcResourceAccessPipelines, ERHIPipeline DstResourceAccessPipelines)
+void FRayTracingDynamicGeometryUpdateManager::AddDynamicGeometryUpdatePass(const FViewInfo& View, FRDGBuilder& GraphBuilder, ERDGPassFlags ComputePassFlags, ERHIPipeline ResourceAccessPipelines, FRDGBufferRef& OutDynamicGeometryScratchBuffer)
+{
+	AddDynamicGeometryUpdatePass(GraphBuilder, ComputePassFlags, View.GetSceneUniforms().GetBuffer(GraphBuilder), View.bRayTracingFeedbackEnabled, ResourceAccessPipelines, OutDynamicGeometryScratchBuffer);
+}
+
+void FRayTracingDynamicGeometryUpdateManager::DispatchUpdates(FRHICommandList& RHICmdList, TConstArrayView<FMeshComputeDispatchCommand> DispatchCommands, ERHIPipeline SrcResourceAccessPipelines, ERHIPipeline DstResourceAccessPipelines)
 {
 	if (DispatchCommands.Num() > 0)
 	{
@@ -778,7 +836,7 @@ void FRayTracingDynamicGeometryUpdateManager::DispatchUpdates(FRHICommandList& R
 		OverlapUAVs.Reserve(DispatchCommands.Num());
 		const FRWBuffer* LastBuffer = nullptr;
 		TSet<const FRWBuffer*> TransitionedBuffers;
-		for (FMeshComputeDispatchCommand& Cmd : DispatchCommands)
+		for (const FMeshComputeDispatchCommand& Cmd : DispatchCommands)
 		{
 			if (Cmd.TargetBuffer == nullptr)
 			{
@@ -822,7 +880,7 @@ void FRayTracingDynamicGeometryUpdateManager::DispatchUpdates(FRHICommandList& R
 			// Cache the bound uniform buffers because a lot are the same between dispatches
 			FShaderBindingState ShaderBindingState;
 
-			for (FMeshComputeDispatchCommand& Cmd : DispatchCommands)
+			for (const FMeshComputeDispatchCommand& Cmd : DispatchCommands)
 			{
 				const TShaderRef<FRayTracingDynamicGeometryConverterCS>& Shader = Cmd.MaterialShader;
 				FRHIComputeShader* ComputeShader = Shader.GetComputeShader();
@@ -860,7 +918,15 @@ void FRayTracingDynamicGeometryUpdateManager::DispatchUpdates(FRHICommandList& R
 			RHICmdList.Transition(TransitionsAfter, SrcResourceAccessPipelines, DstResourceAccessPipelines);
 		}
 	}
+}
 
+void FRayTracingDynamicGeometryUpdateManager::DispatchUpdates(FRHICommandList& RHICmdList, FRHIBuffer* ScratchBuffer, ERHIPipeline SrcResourceAccessPipelines, ERHIPipeline DstResourceAccessPipelines)
+{
+	for (TPair<FRHIUniformBuffer*, TArray<FMeshComputeDispatchCommand>>& ViewDispatchCommands : DispatchCommandsPerView)
+	{
+		DispatchUpdates(RHICmdList, ViewDispatchCommands.Value, SrcResourceAccessPipelines, DstResourceAccessPipelines);
+	}
+			
 	if (BuildParams.Num() > 0)
 	{
 		// Can't use parallel command list because we have to make sure we are not building BVH data
@@ -874,6 +940,10 @@ void FRayTracingDynamicGeometryUpdateManager::DispatchUpdates(FRHICommandList& R
 		ScratchBufferRange.Offset = 0;
 		RHICmdList.BuildAccelerationStructures(BuildParams, ScratchBufferRange);
 	}
+
+	PRAGMA_DISABLE_DEPRECATION_WARNINGS
+	EndUpdate();
+	PRAGMA_ENABLE_DEPRECATION_WARNINGS
 }
 
 void FRayTracingDynamicGeometryUpdateManager::EndUpdate()
diff --git a/Engine/Source/Runtime/Renderer/Public/RayTracingDynamicGeometryUpdateManager.h b/Engine/Source/Runtime/Renderer/Public/RayTracingDynamicGeometryUpdateManager.h
--- a/Engine/Source/Runtime/Renderer/Public/RayTracingDynamicGeometryUpdateManager.h
+++ b/Engine/Source/Runtime/Renderer/Public/RayTracingDynamicGeometryUpdateManager.h
@@ -18,12 +18,15 @@ class FSceneView;
 class FViewInfo;
 class FRDGBuilder;
 class FRayTracingGeometry;
+class FRayTracingDynamicGeometryConverterCS;
 struct FRayTracingDynamicGeometryUpdateParams;
 
+DECLARE_UNIFORM_BUFFER_STRUCT(FSceneUniformParameters, RENDERER_API)
+
 struct FMeshComputeDispatchCommand
 {
 	FMeshDrawShaderBindings ShaderBindings;
-	TShaderRef<class FRayTracingDynamicGeometryConverterCS> MaterialShader;
+	TShaderRef<FRayTracingDynamicGeometryConverterCS> MaterialShader;
 
 	uint32 NumMaxVertices;
 	uint32 NumCPUVertices;
@@ -57,6 +60,15 @@ public:
 	RENDERER_API int64 BeginUpdate();
 
 	/** Add RDG pass which will dispatch all dynamic geometry vertex updates and then request BLAS build and update for all pending requests */
+	RENDERER_API void AddDynamicGeometryUpdatePass(
+		FRDGBuilder& GraphBuilder,
+		ERDGPassFlags ComputePassFlags,
+		const TRDGUniformBufferRef<FSceneUniformParameters>& SceneUB,
+		bool bUseTracingFeedback, 
+		ERHIPipeline ResourceAccessPipelines,
+		FRDGBufferRef& OutDynamicGeometryScratchBuffer);
+
+	UE_DEPRECATED(5.6, "Provide SceneUB and bUseTracingFeedback instead of View.")
 	RENDERER_API void AddDynamicGeometryUpdatePass(const FViewInfo& View, FRDGBuilder& GraphBuilder, ERDGPassFlags ComputePassFlags, ERHIPipeline ResourceAccessPipelines, FRDGBufferRef& OutDynamicGeometryScratchBuffer);
 
 	/** Clears the working arrays to not hold any references. */
@@ -71,22 +83,29 @@ public:
 	UE_DEPRECATED(5.5, "Use AddDynamicGeometryUpdatePass instead which allocates scratch buffer internally.")
 	RENDERER_API uint32 ComputeScratchBufferSize();
 
+	UE_DEPRECATED(5.6, "Use AddDynamicGeometryUpdatePass(...) which internally handles the full update.")
 	RENDERER_API void Update(const FViewInfo& View);
 private:
+
+	void ScheduleUpdates(bool bUseTracingFeedback);
+	static void DispatchUpdates(FRHICommandList& RHICmdList, TConstArrayView<FMeshComputeDispatchCommand> DispatchCommands, ERHIPipeline SrcResourceAccessPipelines = ERHIPipeline::Graphics, ERHIPipeline DstResourceAccessPipelines = ERHIPipeline::Graphics);
+
 	FCriticalSection RTGeometriesToUpdateCS;
 	TMap<FRayTracingGeometry*, EAccelerationStructureBuildMode> RTGeometriesToUpdate;
 
 	struct FRayTracingDynamicGeometryBuildParams
 	{
 		TArray<FMeshComputeDispatchCommand> DispatchCommands;
-		FRayTracingGeometry* Geometry;
+		FRHIUniformBuffer* ViewUniformBuffer = nullptr;
+		FRayTracingGeometry* Geometry = nullptr;
 		int32 SegmentOffset = -1;
 	};
 
 	TArray<FRayTracingDynamicGeometryBuildParams> DynamicGeometryBuilds;
 	TArray<FRayTracingDynamicGeometryBuildParams> DynamicGeometryUpdates;
 
-	TArray<FMeshComputeDispatchCommand> DispatchCommands;
+	// Group dispatch commands per view uniform buffer since it is specified when creating the RDG passes
+	TMap<FRHIUniformBuffer*, TArray<FMeshComputeDispatchCommand>> DispatchCommandsPerView;
 	TArray<FRayTracingGeometryBuildParams> BuildParams;
 	TArray<FRayTracingGeometrySegment> Segments;
 
