PR #12864: Fix copy-paste in comment.

#rnx

diff --git a/Engine/Source/Runtime/Renderer/Private/ScenePrivate.h b/Engine/Source/Runtime/Renderer/Private/ScenePrivate.h
--- a/Engine/Source/Runtime/Renderer/Private/ScenePrivate.h
+++ b/Engine/Source/Runtime/Renderer/Private/ScenePrivate.h
@@ -2487,7 +2487,7 @@ namespace EOcclusionFlags
 		AllowApproximateOcclusion = 0x4,
 		/** Indicates the primitive has a valid ID for precomputed visibility. */
 		HasPrecomputedVisibility = 0x8,
-		/** Indicates the primitive has a valid ID for precomputed visibility. */
+		/** Indicates the primitive supports subprimitive occlusion queries. */
 		HasSubprimitiveQueries = 0x10,
 	};
 };
