Use RHI fences for virtual texture feedback on windows.
Previously we used queries on windows because of a lack of fence support on D3D11. But fence support is now implemented in the D3D11 RHI.
Also added a r.vt.FeedbackLatency cvar to allow some intentional slop in the feedback pipeline.
This smooths oscillating feedback counts which were observed when testing the fence usage, and actually occured before and on all platforms.

diff --git a/Engine/Source/Runtime/Renderer/Private/VT/VirtualTextureFeedback.cpp b/Engine/Source/Runtime/Renderer/Private/VT/VirtualTextureFeedback.cpp
--- a/Engine/Source/Runtime/Renderer/Private/VT/VirtualTextureFeedback.cpp
+++ b/Engine/Source/Runtime/Renderer/Private/VT/VirtualTextureFeedback.cpp
@@ -5,14 +5,12 @@
 #include "RenderGraphBuilder.h"
 #include "RenderGraphUtils.h"
 
-#if PLATFORM_WINDOWS
-// Use Query objects until RHI has a good fence on D3D11
-#define USE_RHI_FENCES 0
-#else
-#define USE_RHI_FENCES 1
-#endif
-
-#if USE_RHI_FENCES
+int32 GVirtualTextureFeedbackLatency = 3;
+static FAutoConsoleVariableRef CVarVirtualTextureFeedbackLatency(
+	TEXT("r.vt.FeedbackLatency"),
+	GVirtualTextureFeedbackLatency,
+	TEXT("How much latency to allow in the GPU feedback pipeline before we start mapping multiple buffers to catch up."),
+	ECVF_RenderThreadSafe);
 
 /** Container for GPU fences. */
 class FFeedbackGPUFencePool
@@ -67,85 +65,6 @@ public:
 	}
 };
 
-#else // USE_RHI_FENCES
-
-/** Container for GPU fences. Implemented as GPU Queries. */
-class FFeedbackGPUFencePool
-{
-public:
-	FGPUFenceRHIRef DummyFence;
-	TArray<FRenderQueryRHIRef> Fences;
-	bool bDummyFenceWritten = false;
-
-	FFeedbackGPUFencePool(int32 InSize)
-	{
-		Fences.AddDefaulted(InSize);
-	}
-
-	void InitRHI(FRHICommandListBase&)
-	{
-		if (!DummyFence.IsValid())
-		{
-			DummyFence = RHICreateGPUFence(FName());
-			bDummyFenceWritten = false;
-		}
-	}
-
-	void ReleaseRHI()
-	{
-		for (int i = 0; i < Fences.Num(); ++i)
-		{
-			if (Fences[i].IsValid())
-			{
-				Fences[i].SafeRelease();
-			}
-		}
-
-		DummyFence.SafeRelease();
-		bDummyFenceWritten = false;
-	}
-
-	void Allocate(FRHICommandList& RHICmdList, int32 Index)
-	{
-		if (Fences[Index].IsValid())
-		{
-			Fences[Index].SafeRelease();
-		}
-		Fences[Index] = GDynamicRHI->RHICreateRenderQuery(RQT_AbsoluteTime);
-
-		if (!bDummyFenceWritten && DummyFence.IsValid())
-		{
-			// Write dummy fence one time on first Allocate
-			// After that we want it to always Poll() true
-			RHICmdList.WriteGPUFence(DummyFence);
-			bDummyFenceWritten = true;
-		}
-	}
-	
-	void Write(FRHICommandList& RHICmdList, int32 Index)
-	{
-		RHICmdList.EndRenderQuery(Fences[Index]);
-	}
-
-	bool Poll(FRHICommandList& RHICmdList, int32 Index)
-	{
-		uint64 Dummy;
-		return RHIGetRenderQueryResult(Fences[Index], Dummy, false, RHICmdList.GetGPUMask().ToIndex());
-	}
-
-	FGPUFenceRHIRef GetMapFence(int32 Index)
-	{
-		return DummyFence;
-	}
-
-	void Release(int32 Index)
-	{
-		Fences[Index].SafeRelease();
-	}
-};
-
-#endif // USE_RHI_FENCES
-
 
 FVirtualTextureFeedback::FVirtualTextureFeedback()
 	: NumPending(0)
@@ -353,7 +272,11 @@ FVirtualTextureFeedback::FMapResult FVirtualTextureFeedback::Map(FRHICommandList
 
 FVirtualTextureFeedback::FMapResult FVirtualTextureFeedback::Map(FRHICommandListImmediate& RHICmdList)
 {
-	return Map(RHICmdList, MaxTransfers);
+	// Allow some slack in the pipeline before we start mapping more than one buffer.
+	// Otherwise we can get into an oscillating pattern of mapping 2 buffers, then 0 buffers, then 2 again etc.
+	const uint32 MaxTransfersToMap = NumPending < GVirtualTextureFeedbackLatency ? 1 : MaxTransfers;
+
+	return Map(RHICmdList, MaxTransfersToMap);
 }
 
 void FVirtualTextureFeedback::Unmap(FRHICommandListImmediate& RHICmdList, int32 MapHandle)
