[HWRT] Fix race condition accessing VisibleRayTracingShaderBindings.
- moved merging dynamic bindings out of SetupRayTracingPipelineStatesAndSBT(...) and after waiting for VisibleRayTracingShaderBindingsFinalizeTask.

diff --git a/Engine/Source/Runtime/Renderer/Private/DeferredShadingRenderer.cpp b/Engine/Source/Runtime/Renderer/Private/DeferredShadingRenderer.cpp
--- a/Engine/Source/Runtime/Renderer/Private/DeferredShadingRenderer.cpp
+++ b/Engine/Source/Runtime/Renderer/Private/DeferredShadingRenderer.cpp
@@ -394,7 +394,7 @@ FDeferredShadingSceneRenderer::FDeferredShadingSceneRenderer(const FSceneViewFam
 /** 
 * Renders the view family. 
 */
-DECLARE_CYCLE_STAT(TEXT("Wait RayTracing Add Mesh Batch"), STAT_WaitRayTracingAddMesh, STATGROUP_SceneRendering);
+DECLARE_CYCLE_STAT(TEXT("Wait RayTracing Dynamic Bindings"), STAT_WaitRayTracingDynamicBindings, STATGROUP_SceneRendering);
 DECLARE_CYCLE_STAT(TEXT("Wait Ray Tracing Scene Initialization"), STAT_WaitRayTracingSceneInitTask, STATGROUP_SceneRendering);
 DECLARE_CYCLE_STAT(TEXT("Wait Ray Tracing Visible Shader Bindings Finalize"), STAT_WaitRayTracingVisibleShaderBindingsFinalizeTask, STATGROUP_SceneRendering);
 
@@ -603,20 +603,6 @@ bool FDeferredShadingSceneRenderer::SetupRayTracingPipelineStatesAndSBT(FRDGBuil
 	const int32 ReferenceViewIndex = 0;
 	FViewInfo& ReferenceView = Views[ReferenceViewIndex];
 
-	if (ReferenceView.AddRayTracingMeshBatchTaskList.Num() > 0)
-	{
-		SCOPE_CYCLE_COUNTER(STAT_WaitRayTracingAddMesh);
-
-		UE::Tasks::Wait(ReferenceView.AddRayTracingMeshBatchTaskList);
-
-		for (int32 TaskIndex = 0; TaskIndex < ReferenceView.AddRayTracingMeshBatchTaskList.Num(); TaskIndex++)
-		{
-			ReferenceView.VisibleRayTracingShaderBindings.Append(*ReferenceView.DynamicRayTracingShaderBindingsPerTask[TaskIndex]);
-		}
-
-		ReferenceView.AddRayTracingMeshBatchTaskList.Empty();
-	}
-
 	if (!GRHISupportsRayTracingShaders && !GRHISupportsInlineRayTracing)
 	{
 		return false;
@@ -945,6 +931,22 @@ void FDeferredShadingSceneRenderer::SetupRayTracingRenderingData(FRDGBuilder& Gr
 		SCOPE_CYCLE_COUNTER(STAT_WaitRayTracingVisibleShaderBindingsFinalizeTask);
 		ReferenceView.VisibleRayTracingShaderBindingsFinalizeTask.Wait();
 	}
+
+	if (ReferenceView.AddRayTracingMeshBatchTaskList.Num() > 0)
+	{
+		{
+			SCOPE_CYCLE_COUNTER(STAT_WaitRayTracingDynamicBindings);
+
+			UE::Tasks::Wait(ReferenceView.AddRayTracingMeshBatchTaskList);
+		}
+
+		for (int32 TaskIndex = 0; TaskIndex < ReferenceView.AddRayTracingMeshBatchTaskList.Num(); TaskIndex++)
+		{
+			ReferenceView.VisibleRayTracingShaderBindings.Append(*ReferenceView.DynamicRayTracingShaderBindingsPerTask[TaskIndex]);
+		}
+
+		ReferenceView.AddRayTracingMeshBatchTaskList.Empty();
+	}
 	
 	// Build the dirty persistent shader bindings from the visible shader bindings (SBT version is updated after the PSO and SBTs have been created)
 	// Inline lumen buffer bindings don't use the dirty bindings yet - so don't need to build it here (uses the visible binding set)
