Change default values for r.RayTracing.PersistentSBT.MinStaticGeometrySegments and r.RayTracing.PersistentSBT.MinDynamicGeometrySegments to reduce memory footprint in small projects with not a very high amount of segments

diff --git a/Engine/Source/Runtime/Renderer/Private/RayTracing/RayTracingShaderBindingTable.cpp b/Engine/Source/Runtime/Renderer/Private/RayTracing/RayTracingShaderBindingTable.cpp
--- a/Engine/Source/Runtime/Renderer/Private/RayTracing/RayTracingShaderBindingTable.cpp
+++ b/Engine/Source/Runtime/Renderer/Private/RayTracing/RayTracingShaderBindingTable.cpp
@@ -41,7 +41,7 @@ static FAutoConsoleVariableRef CVarRayTracingPersistentSBTMinMissShaderSlots(
 	ECVF_ReadOnly
 );
 
-static int32 GMinStaticGeometrySegments = 4096;
+static int32 GMinStaticGeometrySegments = 256;
 static FAutoConsoleVariableRef CVarRayTracingPersistentSBTMinStaticGeometrySegments(
 	TEXT("r.RayTracing.PersistentSBT.MinStaticGeometrySegments"),
 	GMinStaticGeometrySegments,
@@ -49,7 +49,7 @@ static FAutoConsoleVariableRef CVarRayTracingPersistentSBTMinStaticGeometrySegme
 	ECVF_ReadOnly
 );
 
-static int32 GMinDynamicGeometrySegments = 1024;
+static int32 GMinDynamicGeometrySegments = 256;
 static FAutoConsoleVariableRef CVarRayTracingPersistentSBTMinDynamicGeometrySegments(
 	TEXT("r.RayTracing.PersistentSBT.MinDynamicGeometrySegments"),
 	GMinDynamicGeometrySegments,
