[HWRT] Fix crash with Nanite.Mode 1 -- RayTracingGeometry is not valid at the point when we are getting the geometry handle.

#rnx

diff --git a/Engine/Source/Runtime/Renderer/Private/RayTracing/RayTracing.cpp b/Engine/Source/Runtime/Renderer/Private/RayTracing/RayTracing.cpp
--- a/Engine/Source/Runtime/Renderer/Private/RayTracing/RayTracing.cpp
+++ b/Engine/Source/Runtime/Renderer/Private/RayTracing/RayTracing.cpp
@@ -1646,9 +1646,8 @@ namespace RayTracing
 					FRayTracingGeometryInstance RayTracingInstance = *RelevantPrimitive.CachedRayTracingInstance;
 					RayTracingInstance.InstanceContributionToHitGroupIndex = RTLODData.SBTAllocation->GetInstanceContributionToHitGroupIndex(ERayTracingShaderBindingLayer::Base);
 					AddDebugRayTracingInstanceFlags(RayTracingInstance.Flags);
-					
-					// TODO: verify this
-					const int32 GeometryHandle = RelevantPrimitive.RayTracingGeometry->GetGeometryHandle();
+
+					const int32 GeometryHandle = RelevantPrimitive.RayTracingGeometry ? RelevantPrimitive.RayTracingGeometry->GetGeometryHandle() : INDEX_NONE;
 					RayTracingScene.SetInstance(CachedStaticInstanceRange, InstanceIndexInRange, MoveTemp(RayTracingInstance), SceneProxy, false, GeometryHandle);
 				}
 
@@ -1659,9 +1658,8 @@ namespace RayTracing
 					FRayTracingGeometryInstance DecalRayTracingInstance = *RelevantPrimitive.CachedRayTracingInstance;
 					DecalRayTracingInstance.InstanceContributionToHitGroupIndex = RTLODData.SBTAllocation->GetInstanceContributionToHitGroupIndex(ERayTracingShaderBindingLayer::Decals);
 					AddDebugRayTracingInstanceFlags(DecalRayTracingInstance.Flags);
-
-					// TODO: verify this
-					const int32 GeometryHandle = RelevantPrimitive.RayTracingGeometry->GetGeometryHandle();
+					
+					const int32 GeometryHandle = RelevantPrimitive.RayTracingGeometry ? RelevantPrimitive.RayTracingGeometry->GetGeometryHandle() : INDEX_NONE;
 					RayTracingScene.SetInstance(CachedStaticDecalInstanceRange, DecalInstanceIndexInRange, MoveTemp(DecalRayTracingInstance), SceneProxy, false, GeometryHandle);
 				}
 
