Added optional scene captures logging for diagnostic purposes.

diff --git a/Engine/Source/Runtime/Renderer/Private/PlanarReflectionRendering.cpp b/Engine/Source/Runtime/Renderer/Private/PlanarReflectionRendering.cpp
--- a/Engine/Source/Runtime/Renderer/Private/PlanarReflectionRendering.cpp
+++ b/Engine/Source/Runtime/Renderer/Private/PlanarReflectionRendering.cpp
@@ -394,6 +394,8 @@ static bool UpdatePlanarReflectionContents_RenderThread(
 		}
 	}
 
+	UE_CLOG(FSceneCaptureLogUtils::bEnableSceneCaptureLogging, LogSceneCapture, Log, TEXT("Running UpdatePlanarReflectionContents_RenderThread."));
+
 	SceneRenderer->Render(GraphBuilder, SceneUpdateInputs);
 
 	FRDGTextureRef ReflectionOutputTexture = GraphBuilder.RegisterExternalTexture(CreateRenderTarget(Inputs.RenderTarget->TextureRHI, TEXT("ReflectionOutputTexture")));
diff --git a/Engine/Source/Runtime/Renderer/Private/ReflectionEnvironmentCapture.cpp b/Engine/Source/Runtime/Renderer/Private/ReflectionEnvironmentCapture.cpp
--- a/Engine/Source/Runtime/Renderer/Private/ReflectionEnvironmentCapture.cpp
+++ b/Engine/Source/Runtime/Renderer/Private/ReflectionEnvironmentCapture.cpp
@@ -549,6 +549,8 @@ void CaptureSceneToScratchCubemap(
 {
 	const ERHIFeatureLevel::Type FeatureLevel = SceneRenderer->FeatureLevel;
 
+	UE_CLOG(FSceneCaptureLogUtils::bEnableSceneCaptureLogging, LogSceneCapture, Log, TEXT("Starting CaptureSceneToScratchCubemap Face: %d, Size: %d, bCapturingForSkylight: %d"), CubeFace, CubemapSize, (int32) bCapturingForSkyLight);
+
 	// Render the scene normally for one face of the cubemap
 	SceneRenderer->Render(GraphBuilder, SceneUpdateInputs);
 
@@ -594,6 +596,8 @@ void CaptureSceneToScratchCubemap(
 		}
 
 		PassParameters->SkyLightCaptureParameters = SkyLightParametersValue;
+
+		UE_CLOG(FSceneCaptureLogUtils::bEnableSceneCaptureLogging, LogSceneCapture, Log, TEXT("Ending CaptureSceneToScratchCubemap Face: %d, Size: %d, bCapturingForSkylight: %d"), CubeFace, CubemapSize, (int32) bCapturingForSkyLight);
 	}
 
 	const FMinimalSceneTextures& SceneTextures = View.GetSceneTextures();
diff --git a/Engine/Source/Runtime/Renderer/Private/SceneCaptureRendering.cpp b/Engine/Source/Runtime/Renderer/Private/SceneCaptureRendering.cpp
--- a/Engine/Source/Runtime/Renderer/Private/SceneCaptureRendering.cpp
+++ b/Engine/Source/Runtime/Renderer/Private/SceneCaptureRendering.cpp
@@ -471,6 +471,8 @@ static void UpdateSceneCaptureContent_RenderThread(
 		View.bIgnoreExistingQueries = true;
 	}
 
+	UE_CLOG(FSceneCaptureLogUtils::bEnableSceneCaptureLogging, LogSceneCapture, Log, TEXT("Running UpdateSceneCaptureContent_RenderThread."));
+
 	SceneRenderer->Render(GraphBuilder, SceneUpdateInputs);
 
 	if (ShadingPath == EShadingPath::Mobile)
diff --git a/Engine/Source/Runtime/Renderer/Private/SceneRendering.cpp b/Engine/Source/Runtime/Renderer/Private/SceneRendering.cpp
--- a/Engine/Source/Runtime/Renderer/Private/SceneRendering.cpp
+++ b/Engine/Source/Runtime/Renderer/Private/SceneRendering.cpp
@@ -120,6 +120,8 @@
 	Globals
 -----------------------------------------------------------------------------*/
 
+DEFINE_LOG_CATEGORY(LogSceneCapture);
+
 static TGlobalResource<FVirtualTextureFeedbackBuffer> GVirtualTextureFeedbackBuffer;
 
 // This is an experimental optimization switch to render pre-pass depth for scene capture without calling the entire FDeferredShadingSceneRenderer::Render()
@@ -173,6 +175,13 @@ static TAutoConsoleVariable<int32> CVarMeshDrawCommandsDynamicInstancing(
 	TEXT("Whether to dynamically combine multiple compatible visible Mesh Draw Commands into one instanced draw on vertex factories that support it."),
 	ECVF_RenderThreadSafe);
 
+bool FSceneCaptureLogUtils::bEnableSceneCaptureLogging = false;
+
+FAutoConsoleVariableRef CVarEnableSceneCaptureLogging(
+	TEXT("r.SceneCapture.EnableLogging"),
+	FSceneCaptureLogUtils::bEnableSceneCaptureLogging,
+	TEXT("Enable logging of scene captures."));
+
 bool IsDynamicInstancingEnabled(ERHIFeatureLevel::Type FeatureLevel)
 {
 	return CVarMeshDrawCommandsDynamicInstancing.GetValueOnRenderThread() > 0
diff --git a/Engine/Source/Runtime/Renderer/Private/SceneRendering.h b/Engine/Source/Runtime/Renderer/Private/SceneRendering.h
--- a/Engine/Source/Runtime/Renderer/Private/SceneRendering.h
+++ b/Engine/Source/Runtime/Renderer/Private/SceneRendering.h
@@ -107,6 +107,12 @@ struct FPersistentViewId
 	int32 Index = INDEX_NONE;
 };
 
+DECLARE_LOG_CATEGORY_EXTERN(LogSceneCapture, Log, All);
+
+struct FSceneCaptureLogUtils
+{
+	static bool bEnableSceneCaptureLogging;
+};
 
 extern bool ShouldUseStereoLumenOptimizations();
 
