[HWRT] Fix race condition between building dynamic ray tracing mesh batches and gathering shaders for RTPSO.

diff --git a/Engine/Source/Runtime/Renderer/Private/DeferredShadingRenderer.cpp b/Engine/Source/Runtime/Renderer/Private/DeferredShadingRenderer.cpp
--- a/Engine/Source/Runtime/Renderer/Private/DeferredShadingRenderer.cpp
+++ b/Engine/Source/Runtime/Renderer/Private/DeferredShadingRenderer.cpp
@@ -924,6 +924,15 @@ void FDeferredShadingSceneRenderer::SetupRayTracingRenderingData(FRDGBuilder& Gr
 		}
 	}
 
+	if (ReferenceView.AddDynamicRayTracingMeshBatchTaskList.Num() > 0)
+	{
+		SCOPE_CYCLE_COUNTER(STAT_WaitRayTracingDynamicBindings);
+
+		// need to wait for dynamic mesh batches tasks to finish before executing SetupRayTracingPipelineStatesAndSBT(...)
+		// since they can request new materials that need to be included in RTPSO 
+		UE::Tasks::Wait(ReferenceView.AddDynamicRayTracingMeshBatchTaskList);
+	}
+
 	bool bIsUsingFallbackRTPSO = false;
 	SetupRayTracingPipelineStatesAndSBT(GraphBuilder, bAnyLumenHardwareInlineRayTracingPassEnabled, bIsUsingFallbackRTPSO);
 
@@ -932,20 +941,14 @@ void FDeferredShadingSceneRenderer::SetupRayTracingRenderingData(FRDGBuilder& Gr
 		ReferenceView.VisibleRayTracingShaderBindingsFinalizeTask.Wait();
 	}
 
-	if (ReferenceView.AddRayTracingMeshBatchTaskList.Num() > 0)
+	if (ReferenceView.AddDynamicRayTracingMeshBatchTaskList.Num() > 0)
 	{
-		{
-			SCOPE_CYCLE_COUNTER(STAT_WaitRayTracingDynamicBindings);
-
-			UE::Tasks::Wait(ReferenceView.AddRayTracingMeshBatchTaskList);
-		}
-
-		for (int32 TaskIndex = 0; TaskIndex < ReferenceView.AddRayTracingMeshBatchTaskList.Num(); TaskIndex++)
+		for (int32 TaskIndex = 0; TaskIndex < ReferenceView.AddDynamicRayTracingMeshBatchTaskList.Num(); TaskIndex++)
 		{
 			ReferenceView.VisibleRayTracingShaderBindings.Append(*ReferenceView.DynamicRayTracingShaderBindingsPerTask[TaskIndex]);
 		}
 
-		ReferenceView.AddRayTracingMeshBatchTaskList.Empty();
+		ReferenceView.AddDynamicRayTracingMeshBatchTaskList.Empty();
 	}
 	
 	// Build the dirty persistent shader bindings from the visible shader bindings (SBT version is updated after the PSO and SBTs have been created)
diff --git a/Engine/Source/Runtime/Renderer/Private/RayTracing/RayTracing.cpp b/Engine/Source/Runtime/Renderer/Private/RayTracing/RayTracing.cpp
--- a/Engine/Source/Runtime/Renderer/Private/RayTracing/RayTracing.cpp
+++ b/Engine/Source/Runtime/Renderer/Private/RayTracing/RayTracing.cpp
@@ -494,7 +494,7 @@ namespace RayTracing
 		TaskVisibleShaderBindings->Reserve(NumPendingMeshBatches);
 		View.DynamicRayTracingShaderBindingsPerTask.Add(TaskVisibleShaderBindings);
 
-		View.AddRayTracingMeshBatchTaskList.Add(UE::Tasks::Launch(UE_SOURCE_LOCATION,
+		View.AddDynamicRayTracingMeshBatchTaskList.Add(UE::Tasks::Launch(UE_SOURCE_LOCATION,
 			[TaskDataHead = MeshBatchTaskHead, &View, &Scene, TaskDynamicCommandStorage, TaskVisibleShaderBindings]()
 			{
 				FTaskTagScope TaskTagScope(ETaskTag::EParallelRenderingThread);
diff --git a/Engine/Source/Runtime/Renderer/Private/SceneRendering.h b/Engine/Source/Runtime/Renderer/Private/SceneRendering.h
--- a/Engine/Source/Runtime/Renderer/Private/SceneRendering.h
+++ b/Engine/Source/Runtime/Renderer/Private/SceneRendering.h
@@ -1304,7 +1304,7 @@ public:
 	UE::Tasks::FTask RayTracingSceneInitTask; // Task to asynchronously call RayTracingScene.BuildInitializationData()
 	UE::Tasks::FTask VisibleRayTracingShaderBindingsFinalizeTask;
 
-	TArray<UE::Tasks::FTask> AddRayTracingMeshBatchTaskList;
+	TArray<UE::Tasks::FTask> AddDynamicRayTracingMeshBatchTaskList;
 	TArray<FRayTracingShaderBindingDataOneFrameArray*, SceneRenderingAllocator> DynamicRayTracingShaderBindingsPerTask;
 	TArray<FDynamicRayTracingMeshCommandStorage*, SceneRenderingAllocator> DynamicRayTracingMeshCommandStoragePerTask;
 #endif
