Fully disable RDG tracking of VirtutalTexture FeedbackBuffer since we are managing the state of that explicitly.

diff --git a/Engine/Source/Runtime/Renderer/Private/VT/VirtualTextureFeedbackBuffer.cpp b/Engine/Source/Runtime/Renderer/Private/VT/VirtualTextureFeedbackBuffer.cpp
--- a/Engine/Source/Runtime/Renderer/Private/VT/VirtualTextureFeedbackBuffer.cpp
+++ b/Engine/Source/Runtime/Renderer/Private/VT/VirtualTextureFeedbackBuffer.cpp
@@ -115,15 +115,15 @@ void FVirtualTextureFeedbackBuffer::Begin(FRDGBuilder& GraphBuilder, uint32 InFe
 
 void FVirtualTextureFeedbackBuffer::End(FRDGBuilder& GraphBuilder)
 {
-	FRDGBufferRef FeedbackBuffer = GraphBuilder.RegisterExternalBuffer(PooledBuffer);
-	FRDGBufferSRVRef FeedbackBufferSRV = GraphBuilder.CreateSRV(FeedbackBuffer);
-
 	AddPass(GraphBuilder, RDG_EVENT_NAME("VirtualTextureFeedbackTransition"), [UAV = UAV](FRHICommandListImmediate& RHICmdList)
 	{
 		RHICmdList.EndUAVOverlap(UAV);
 		RHICmdList.Transition(FRHITransitionInfo(UAV, ERHIAccess::UAVMask, ERHIAccess::SRVCompute));
 	});
 
+	FRDGBufferRef FeedbackBuffer = GraphBuilder.RegisterExternalBuffer(PooledBuffer, ERDGBufferFlags::SkipTracking);
+	FRDGBufferSRVRef FeedbackBufferSRV = GraphBuilder.CreateSRV(FeedbackBuffer);
+
 	// Compact feedback before queuing for readback.
 	const uint32 CompactedFeedbackBufferSize = GetVirtualTextureCompactedFeedbackBufferSize(FeedbackBufferSize);
 	const uint32 HashTableSize = 2 * CompactedFeedbackBufferSize;
