Add check in FRayTracingDynamicGeometryUpdateManager::RemoveGeometryToUpdate because we shouldn't reach this code anymore after refactor of skinned geometry collection during GetDynamicRayTracingInstances instead of from the skinned cache updates itself

diff --git a/Engine/Source/Runtime/Renderer/Private/RayTracing/RayTracingDynamicGeometryUpdateManager.cpp b/Engine/Source/Runtime/Renderer/Private/RayTracing/RayTracingDynamicGeometryUpdateManager.cpp
--- a/Engine/Source/Runtime/Renderer/Private/RayTracing/RayTracingDynamicGeometryUpdateManager.cpp
+++ b/Engine/Source/Runtime/Renderer/Private/RayTracing/RayTracingDynamicGeometryUpdateManager.cpp
@@ -285,6 +285,7 @@ void FRayTracingDynamicGeometryUpdateManager::RemoveGeometryToUpdate(FRayTracing
 	FScopeLock Lock(&RTGeometriesToUpdateCS);
 	if (RTGeometriesToUpdate.Find(RayTracingGeometry) != nullptr)
 	{
+		checkf(false, TEXT("Shouldn't find active raytracing geometries anymore in the dynamic update manager, because all pending ones should have been either processed or cleared already"));
 		RTGeometriesToUpdate.Remove(RayTracingGeometry);
 	}
 }
