Refactor some of the virtual texture feedback buffer API so that it will be easier to reuse for Slate or Canvas renderers.

diff --git a/Engine/Source/Runtime/Engine/Private/SceneView.cpp b/Engine/Source/Runtime/Engine/Private/SceneView.cpp
--- a/Engine/Source/Runtime/Engine/Private/SceneView.cpp
+++ b/Engine/Source/Runtime/Engine/Private/SceneView.cpp
@@ -29,6 +29,7 @@
 #include "UObject/Interface.h"
 #include "TextureResource.h"
 #include "HDRHelper.h"
+#include "VT/VirtualTextureScalability.h"
 
 DEFINE_LOG_CATEGORY(LogBufferVisualization);
 DEFINE_LOG_CATEGORY(LogNaniteVisualization);
@@ -361,15 +362,6 @@ static TAutoConsoleVariable<bool> CVarOrthoCameraHeightAsViewTarget(
 	TEXT("Primarily helps with VSM clipmap selection and avoids overcorrecting NearPlanes.\n"),
 	ECVF_Scalability | ECVF_RenderThreadSafe);
 
-int32 GVirtualTextureFeedbackFactor = 16;
-static FAutoConsoleVariableRef CVarVirtualTextureFeedbackFactor(
-	TEXT("r.vt.FeedbackFactor"),
-	GVirtualTextureFeedbackFactor,
-	TEXT("The size of the VT feedback buffer is calculated by dividing the render resolution by this factor.")
-	TEXT("The value set here is rounded up to the nearest power of two before use."),
-	ECVF_RenderThreadSafe);
-
-
 static int32 GHairStrandsComposeDOFDepth = 1;
 static FAutoConsoleVariableRef CVarHairStrandsComposeDOFDepth(
 	TEXT("r.HairStrands.DOFDepth"), 
@@ -3015,7 +3007,7 @@ PRAGMA_ENABLE_DEPRECATION_WARNINGS
 	bCurrentlyBeingEdited = false;
 
 	bOverrideVirtualTextureThrottle = false;
-	VirtualTextureFeedbackFactor = GVirtualTextureFeedbackFactor;
+	VirtualTextureFeedbackFactor = VirtualTextureScalability::GetVirtualTextureFeedbackFactor();
 
 	if (IsDefaultBackBufferLinearSDR())
 	{
diff --git a/Engine/Source/Runtime/Engine/Private/VT/VirtualTextureScalability.cpp b/Engine/Source/Runtime/Engine/Private/VT/VirtualTextureScalability.cpp
--- a/Engine/Source/Runtime/Engine/Private/VT/VirtualTextureScalability.cpp
+++ b/Engine/Source/Runtime/Engine/Private/VT/VirtualTextureScalability.cpp
@@ -63,6 +63,14 @@ namespace VirtualTextureScalability
 		ECVF_RenderThreadSafe | ECVF_Scalability
 	);
 
+	static TAutoConsoleVariable<int32> CVarVirtualTextureFeedbackFactor(
+		TEXT("r.vt.FeedbackFactor"),
+		16,
+		TEXT("The size of the VT feedback buffer is calculated by dividing the render resolution by this factor.")
+		TEXT("The value set here is rounded up to the nearest power of two before use."),
+		ECVF_RenderThreadSafe | ECVF_Scalability
+	);
+
 	static TAutoConsoleVariable<int32> CVarVTPageFreeThreshold(
 		TEXT("r.VT.PageFreeThreshold"),
 		60,
@@ -213,6 +221,11 @@ namespace VirtualTextureScalability
 		return CVarVTMaxReleasedPerFrame.GetValueOnAnyThread();
 	}
 
+	uint32 GetVirtualTextureFeedbackFactor()
+	{
+		return CVarVirtualTextureFeedbackFactor.GetValueOnAnyThread();
+	}
+
 	uint32 GetPageFreeThreshold()
 	{
 		return FMath::Max(CVarVTPageFreeThreshold.GetValueOnRenderThread(), 0);
diff --git a/Engine/Source/Runtime/Engine/Public/VT/VirtualTextureScalability.h b/Engine/Source/Runtime/Engine/Public/VT/VirtualTextureScalability.h
--- a/Engine/Source/Runtime/Engine/Public/VT/VirtualTextureScalability.h
+++ b/Engine/Source/Runtime/Engine/Public/VT/VirtualTextureScalability.h
@@ -16,6 +16,8 @@ namespace VirtualTextureScalability
 	ENGINE_API int32 GetMaxContinuousUpdatesPerFrame();
 	/** Get max allocated virtual textures to release per frame. */
 	ENGINE_API int32 GetMaxAllocatedVTReleasedPerFrame();
+	/** Get the size of the virtual texture feedback tiles. */
+	ENGINE_API uint32 GetVirtualTextureFeedbackFactor();
 	/** Get the number of frames a page must be unused, before it's considered free */
 	ENGINE_API uint32 GetPageFreeThreshold();
 	/** Get the number of frames a page must be unused, before it is unmapped instead of updated during RVT page invalidation. */
diff --git a/Engine/Source/Runtime/Renderer/Private/SceneRendering.cpp b/Engine/Source/Runtime/Renderer/Private/SceneRendering.cpp
--- a/Engine/Source/Runtime/Renderer/Private/SceneRendering.cpp
+++ b/Engine/Source/Runtime/Renderer/Private/SceneRendering.cpp
@@ -1151,6 +1151,33 @@ void UpdateNoiseTextureParameters(FViewUniformShaderParameters& ViewUniformShade
 	check(ViewUniformShaderParameters.SobolSamplingTexture);
 }
 
+static uint32 AlignVirtualTextureFeedbackTileSize(uint32 InTileSize)
+{
+	// Round to nearest power of two to ensure that shader maths is efficient and sampling sequence logic is simple.
+	return FMath::RoundUpToPowerOfTwo(FMath::Max(InTileSize, 1u));
+}
+
+static void UpdateVirtualTextureFeedbackParameters(uint32 InFrameIndex, uint32 InVirtualTextureFeedbackTileSize, FViewUniformShaderParameters& ViewUniformShaderParameters)
+{
+	ViewUniformShaderParameters.VTFeedbackBuffer = GVirtualTextureFeedbackBuffer.GetUAV();
+	ViewUniformShaderParameters.VirtualTextureFeedbackBufferSize = GVirtualTextureFeedbackBuffer.GetBufferSize();
+	
+	ViewUniformShaderParameters.VirtualTextureFeedbackShift = FMath::FloorLog2(InVirtualTextureFeedbackTileSize);
+	ViewUniformShaderParameters.VirtualTextureFeedbackMask = InVirtualTextureFeedbackTileSize - 1;
+	// Use some low(ish) discrepancy sequence to run over every pixel in the virtual texture feedback tile.
+	ViewUniformShaderParameters.VirtualTextureFeedbackJitterOffset = SampleVirtualTextureFeedbackSequence(InFrameIndex, InVirtualTextureFeedbackTileSize);
+	// SampleOffset is used to cycle through all VT samples in a material. It just needs to monotonically increase.
+	ViewUniformShaderParameters.VirtualTextureFeedbackSampleOffset = InFrameIndex;
+
+	ViewUniformShaderParameters.GlobalVirtualTextureMipBias = FVirtualTextureSystem::Get().GetGlobalMipBias();
+}
+
+void UpdateVirtualTextureFeedbackParameters(FViewUniformShaderParameters& ViewUniformShaderParameters)
+{
+	const uint32 VirtualTextureFeedbackTileSize = AlignVirtualTextureFeedbackTileSize(VirtualTextureScalability::GetVirtualTextureFeedbackFactor());
+	UpdateVirtualTextureFeedbackParameters(GFrameNumber, VirtualTextureFeedbackTileSize, ViewUniformShaderParameters);
+}
+
 void SetupPrecomputedVolumetricLightmapUniformBufferParameters(const FScene* Scene, FEngineShowFlags EngineShowFlags, FViewUniformShaderParameters& ViewUniformShaderParameters)
 {
 	if (Scene && Scene->VolumetricLightmapSceneData.HasData() && EngineShowFlags.VolumetricLightmap)
@@ -2023,16 +2050,9 @@ void FViewInfo::SetupUniformBufferParameters(
 		}
 	}
 
-	ViewUniformShaderParameters.GlobalVirtualTextureMipBias = FVirtualTextureSystem::Get().GetGlobalMipBias();
-
-	ViewUniformShaderParameters.VirtualTextureFeedbackBufferSize = GetVirtualTextureFeedbackBufferSize(Family, EffectiveBufferSize);
-	const uint32 VirtualTextureFeedbackScale = GetVirtualTextureFeedbackScale(Family);
-	check(VirtualTextureFeedbackScale == 1 << FMath::FloorLog2(VirtualTextureFeedbackScale));
-	ViewUniformShaderParameters.VirtualTextureFeedbackShift = FMath::FloorLog2(VirtualTextureFeedbackScale);
-	ViewUniformShaderParameters.VirtualTextureFeedbackMask = VirtualTextureFeedbackScale - 1;
-	// Use some low(ish) discrepancy sequence to run over every pixel in the virtual texture feedback tile.
-	ViewUniformShaderParameters.VirtualTextureFeedbackJitterOffset = SampleVirtualTextureFeedbackSequence(Family, FrameIndex);
-	ViewUniformShaderParameters.VirtualTextureFeedbackSampleOffset = FrameIndex % 1024;
+	const uint32 VirtualTextureFrameIndex = ViewState ? ViewState->GetFrameIndex() : Family->FrameNumber;
+	const uint32 VirtualTextureFeedbackTileSize = AlignVirtualTextureFeedbackTileSize(Family->VirtualTextureFeedbackFactor);
+	UpdateVirtualTextureFeedbackParameters(VirtualTextureFrameIndex, VirtualTextureFeedbackTileSize, ViewUniformShaderParameters);
 
 	// GGX/Sheen LTC (used as BSDF or for rect light integration)
 	if (GSystemTextures.GGXLTCMat.IsValid() && GSystemTextures.GGXLTCAmp.IsValid())
@@ -2177,8 +2197,6 @@ void FViewInfo::SetupUniformBufferParameters(
 		ViewUniformShaderParameters.LandscapePerComponentData = GWhiteVertexBufferWithSRV->ShaderResourceViewRHI;
 	}
 
-	ViewUniformShaderParameters.VTFeedbackBuffer = GVirtualTextureFeedbackBuffer.GetUAV();
-
 	ViewUniformShaderParameters.GPUSceneViewId = GPUSceneViewId;
 
 	{
@@ -6053,7 +6071,8 @@ void VirtualTextureFeedbackBegin(FRDGBuilder& GraphBuilder, TArrayView<const FVi
 	FIntPoint EffectiveBufferSize = SceneTextureExtent;
 	EffectiveBufferSize.X = FMath::Max(EffectiveBufferSize.X, 1);
 	EffectiveBufferSize.Y = FMath::Max(EffectiveBufferSize.Y, 1);
-	const int32 FeedbackBufferSize = GetVirtualTextureFeedbackBufferSize(Views.Num() > 0 ? Views[0].Family : nullptr, EffectiveBufferSize);
+	const uint32 FeedbackTileSize = Views.Num() > 0 ? Views[0].Family->VirtualTextureFeedbackFactor : VirtualTextureScalability::GetVirtualTextureFeedbackFactor();
+	const int32 FeedbackBufferSize = GetVirtualTextureFeedbackBufferSize(EffectiveBufferSize, AlignVirtualTextureFeedbackTileSize(FeedbackTileSize));
 	const ERHIFeatureLevel::Type FeatureLevel = Views.Num() > 0 ? Views[0].GetFeatureLevel() : GMaxRHIFeatureLevel;
 	
 	GVirtualTextureFeedbackBuffer.Begin(GraphBuilder, FeedbackBufferSize, FeatureLevel);
diff --git a/Engine/Source/Runtime/Renderer/Private/SceneRendering.h b/Engine/Source/Runtime/Renderer/Private/SceneRendering.h
--- a/Engine/Source/Runtime/Renderer/Private/SceneRendering.h
+++ b/Engine/Source/Runtime/Renderer/Private/SceneRendering.h
@@ -2921,6 +2921,8 @@ private:
 
 // The noise textures need to be set in Slate too.
 RENDERER_API void UpdateNoiseTextureParameters(FViewUniformShaderParameters& ViewUniformShaderParameters);
+// The virtual texture feedback also needs to be set in Slate.
+RENDERER_API void UpdateVirtualTextureFeedbackParameters(FViewUniformShaderParameters& ViewUniformShaderParameters);
 
 struct FFastVramConfig
 {
diff --git a/Engine/Source/Runtime/Renderer/Private/VT/VirtualTextureFeedbackBuffer.cpp b/Engine/Source/Runtime/Renderer/Private/VT/VirtualTextureFeedbackBuffer.cpp
--- a/Engine/Source/Runtime/Renderer/Private/VT/VirtualTextureFeedbackBuffer.cpp
+++ b/Engine/Source/Runtime/Renderer/Private/VT/VirtualTextureFeedbackBuffer.cpp
@@ -24,31 +24,30 @@ static FAutoConsoleVariableRef CVarVirtualTextureFeedbackCompactionFactor(
 	TEXT("to account for compaction of duplicate page ids."),
 	ECVF_RenderThreadSafe);
 
-int32 GetVirtualTextureFeedbackScale(FSceneViewFamily const* InViewFamily)
+uint32 GetVirtualTextureFeedbackBufferSize(FIntPoint InSceneTextureExtent, uint32 InVirtualTextureFeedbackTileSize)
 {
-	// Round to nearest power of two to ensure that shader maths is efficient and sampling sequence logic is simple.
-	const int32 VirtualTextureFeedbackFactor = InViewFamily ? InViewFamily->VirtualTextureFeedbackFactor : 0;
-	return FMath::RoundUpToPowerOfTwo(FMath::Max(VirtualTextureFeedbackFactor, 1));
-}
+	// Tile size must be power or 2.
+	check(InVirtualTextureFeedbackTileSize == 1 << FMath::FloorLog2(InVirtualTextureFeedbackTileSize));
 
-int32 GetVirtualTextureFeedbackBufferSize(FSceneViewFamily const* InViewFamily, FIntPoint InSceneTextureExtent)
-{
-	FIntPoint ScaledExtent = FIntPoint::DivideAndRoundUp(InSceneTextureExtent, GetVirtualTextureFeedbackScale(InViewFamily));
-	return ScaledExtent.X * ScaledExtent.Y * GVirtualTextureFeedbackOverdrawFactor;
+	FIntPoint ScaledExtent = FIntPoint::DivideAndRoundUp(InSceneTextureExtent, InVirtualTextureFeedbackTileSize);
+	return ScaledExtent.X * ScaledExtent.Y * FMath::Max(GVirtualTextureFeedbackOverdrawFactor, 1);
 }
 
 static uint32 GetVirtualTextureCompactedFeedbackBufferSize(uint32 InSourceBufferSize)
 {
 	// Could possibly do this dynamically according to some tracked recent high watermark?
-	const int32 CompactedBufferSize = InSourceBufferSize / GVirtualTextureFeedbackCompactionFactor;
+	const uint32 CompactedBufferSize = InSourceBufferSize / FMath::Max(GVirtualTextureFeedbackCompactionFactor, 1);
 	// Size needs to be power of two for hash table wrapping.
-	const int32 AlignedCompactedBufferSize = FMath::RoundUpToPowerOfTwo(FMath::Clamp(CompactedBufferSize, 16, 16 * 1024));
+	const uint32 AlignedCompactedBufferSize = FMath::RoundUpToPowerOfTwo(FMath::Clamp(CompactedBufferSize, 16u, 16u * 1024u));
 	return AlignedCompactedBufferSize;
 }
 
-uint32 SampleVirtualTextureFeedbackSequence(FSceneViewFamily const* InViewFamily, uint32 InFrameIndex)
+uint32 SampleVirtualTextureFeedbackSequence(uint32 InFrameIndex, uint32 InVirtualTextureFeedbackTileSize)
 {
-	const uint32 TileSize = GetVirtualTextureFeedbackScale(InViewFamily);
+	// Tile size must be power or 2.
+	check(InVirtualTextureFeedbackTileSize == 1 << FMath::FloorLog2(InVirtualTextureFeedbackTileSize));
+
+	const uint32 TileSize = InVirtualTextureFeedbackTileSize;
 	const uint32 TileSizeLog2 = FMath::CeilLogTwo(TileSize);
 	const uint32 SequenceSize = FMath::Square(TileSize);
 	const uint32 PixelIndex = InFrameIndex % SequenceSize;
@@ -97,11 +96,10 @@ void FVirtualTextureFeedbackBuffer::Begin(FRDGBuilder& GraphBuilder, uint32 InFe
 	FRDGBufferDesc BufferDesc(FRDGBufferDesc::CreateStructuredDesc(sizeof(uint32), FeedbackBufferSize));
 	BufferDesc.Usage |= BUF_SourceCopy;
 
-	if (AllocatePooledBuffer(BufferDesc, PooledBuffer, TEXT("VirtualTexture_FeedbackBuffer")))
-	{
-		FRDGBufferUAVDesc UAVDesc{};
-		UAV = PooledBuffer->GetOrCreateUAV(GraphBuilder.RHICmdList, UAVDesc);
-	}
+	AllocatePooledBuffer(BufferDesc, PooledBuffer, TEXT("VirtualTexture_FeedbackBuffer"));
+	
+	FRDGBufferUAVDesc UAVDesc{};
+	UAV = PooledBuffer->GetOrCreateUAV(GraphBuilder.RHICmdList, UAVDesc);
 
 	AddPass(GraphBuilder, RDG_EVENT_NAME("VirtualTextureClear"), [UAV = UAV](FRDGAsyncTask, FRHICommandList& RHICmdList)
 	{
@@ -111,10 +109,18 @@ void FVirtualTextureFeedbackBuffer::Begin(FRDGBuilder& GraphBuilder, uint32 InFe
 		RHICmdList.Transition(FRHITransitionInfo(UAV, ERHIAccess::UAVCompute, ERHIAccess::UAVMask));
 		RHICmdList.BeginUAVOverlap(UAV);
 	});
+
+	bIsInBeginEndScope = true;
 }
 
 void FVirtualTextureFeedbackBuffer::End(FRDGBuilder& GraphBuilder)
 {
+	if (!bIsInBeginEndScope)
+	{
+		return;
+	}
+	bIsInBeginEndScope = false;
+
 	AddPass(GraphBuilder, RDG_EVENT_NAME("VirtualTextureFeedbackTransition"), [UAV = UAV](FRHICommandListImmediate& RHICmdList)
 	{
 		RHICmdList.EndUAVOverlap(UAV);
@@ -238,9 +244,14 @@ void FVirtualTextureFeedbackBuffer::End(FRDGBuilder& GraphBuilder)
 	SubmitVirtualTextureFeedbackBuffer(GraphBuilder, CompactedFeedbackBuffer, CompactedDesc);
 }
 
+uint32 FVirtualTextureFeedbackBuffer::GetBufferSize() const
+{
+	return bIsInBeginEndScope ? FeedbackBufferSize : 0;
+}
+
 FRHIUnorderedAccessView* FVirtualTextureFeedbackBuffer::GetUAV() const
 {
-	return UAV ? UAV : GEmptyStructuredBufferWithUAV->UnorderedAccessViewRHI.GetReference();
+	return UAV != nullptr && bIsInBeginEndScope ? UAV : GEmptyStructuredBufferWithUAV->UnorderedAccessViewRHI.GetReference();
 }
 
 void FVirtualTextureFeedbackBuffer::ReleaseRHI()
diff --git a/Engine/Source/Runtime/Renderer/Public/VT/VirtualTextureFeedbackBuffer.h b/Engine/Source/Runtime/Renderer/Public/VT/VirtualTextureFeedbackBuffer.h
--- a/Engine/Source/Runtime/Renderer/Public/VT/VirtualTextureFeedbackBuffer.h
+++ b/Engine/Source/Runtime/Renderer/Public/VT/VirtualTextureFeedbackBuffer.h
@@ -13,16 +13,12 @@
 #include "Templates/RefCounting.h"
 
 class FRDGBuilder;
-class FSceneViewFamily;
 
-/** Returns the scale used to divide the scene texture resolution to get the virtual texture feedback resolution. */
-extern int32 GetVirtualTextureFeedbackScale(FSceneViewFamily const* InViewFamily);
-
-/** Get size for a virtual texture feedback buffer. This is calculated internally by applying GetVirtualTextureFeedbackScale() to the SceneTextureExtent. */
-extern int32 GetVirtualTextureFeedbackBufferSize(FSceneViewFamily const* InViewFamily, FIntPoint InSceneTextureExtent);
+/** Get size for a virtual texture feedback buffer. */
+extern uint32 GetVirtualTextureFeedbackBufferSize(FIntPoint InSceneTextureExtent, uint32 InVirtualTextureTileSize);
 
 /** Returns a value from the sampling sequence used to traverse each feedback tile. */
-extern uint32 SampleVirtualTextureFeedbackSequence(FSceneViewFamily const* InViewFamily, uint32 InFrameIndex);
+extern uint32 SampleVirtualTextureFeedbackSequence(uint32 InFrameIndex, uint32 InVirtualTextureFeedbackTileSize);
 
 /** 
  * Description of how to interpret an RHIBuffer that is being fed to the virtual texture feedback system.
@@ -47,12 +43,15 @@ struct FVirtualTextureFeedbackBufferDesc
 class FVirtualTextureFeedbackBuffer : public FRenderResource
 {
 public:
-	// Allocates and prepares the feedback buffer for submission.
+	// Allocates and prepares the feedback buffer for write access.
 	void Begin(FRDGBuilder& GraphBuilder, uint32 InFeedbackBufferSize, ERHIFeatureLevel::Type InFeatureLevel);
 
 	// Submits the feedback buffer for readback.
 	void End(FRDGBuilder& GraphBuilder);
 
+	// Get the size of the feedback buffer from the last call to Begin().
+	uint32 GetBufferSize() const;
+	// Get the UAV of the feedback buffer.
 	FRHIUnorderedAccessView* GetUAV() const;
 
 	//~ Begin FRenderResource Interface
@@ -60,6 +59,7 @@ public:
 	//~ End FRenderResource Interface
 
 private:
+	bool bIsInBeginEndScope = false;
 	uint32 FeedbackBufferSize = 0;
 	ERHIFeatureLevel::Type FeatureLevel = ERHIFeatureLevel::Num;
 	TRefCountPtr<FRDGPooledBuffer> PooledBuffer;
