Fixed missing dummy instance culling parameters for rendering debug primitives.

diff --git a/Engine/Source/Runtime/Renderer/Private/BasePassRendering.cpp b/Engine/Source/Runtime/Renderer/Private/BasePassRendering.cpp
--- a/Engine/Source/Runtime/Renderer/Private/BasePassRendering.cpp
+++ b/Engine/Source/Runtime/Renderer/Private/BasePassRendering.cpp
@@ -1648,6 +1648,10 @@ void FDeferredShadingSceneRenderer::RenderBasePassInternal(
 						Pass->Draw(RHICmdList, &PassParameters->InstanceCullingDrawParams);
 					});
 				}
+				else
+				{
+					InstanceCullingManager.SetDummyCullingParams(GraphBuilder, PassParameters->InstanceCullingDrawParams);
+				}
 
 				const bool bShouldRenderViewForNanite = bNaniteEnabled && (!bDrawSceneViewsInOneNanitePass || ViewIndex == 0); // when bDrawSceneViewsInOneNanitePass, the first view should cover all the other atlased ones
 				if (bShouldRenderViewForNanite)
diff --git a/Engine/Source/Runtime/Renderer/Private/DepthRendering.cpp b/Engine/Source/Runtime/Renderer/Private/DepthRendering.cpp
--- a/Engine/Source/Runtime/Renderer/Private/DepthRendering.cpp
+++ b/Engine/Source/Runtime/Renderer/Private/DepthRendering.cpp
@@ -531,6 +531,10 @@ void FDeferredShadingSceneRenderer::RenderPrePass(FRDGBuilder& GraphBuilder, TAr
 							Pass->Dispatch(DispatchPassBuilder, &PassParameters->InstanceCullingDrawParams);
 						});
 					}
+					else
+					{
+						InstanceCullingManager.SetDummyCullingParams(GraphBuilder, PassParameters->InstanceCullingDrawParams);
+					}
 
 					RenderPrePassEditorPrimitives(GraphBuilder, View, PassParameters, DrawRenderState, DepthPass.EarlyZPassMode);
 				}
@@ -568,6 +572,10 @@ void FDeferredShadingSceneRenderer::RenderPrePass(FRDGBuilder& GraphBuilder, TAr
 							Pass->Draw(RHICmdList, &PassParameters->InstanceCullingDrawParams);
 						});
 					}
+					else
+					{
+						InstanceCullingManager.SetDummyCullingParams(GraphBuilder, PassParameters->InstanceCullingDrawParams);
+					}
 
 					RenderPrePassEditorPrimitives(GraphBuilder, View, PassParameters, DrawRenderState, DepthPass.EarlyZPassMode);
 				}
