[HWRT] Improvements to ray tracing scene instance buffer stats.
- gather stats for all layers.
- moved readback pass to the end of the frame so it doesn't cause a sync point between graphics and compute queues.
- support stats on all builds (not just #if STATS).

diff --git a/Engine/Plugins/Experimental/GPULightmass/Source/GPULightmass/Private/LightmapRenderer.cpp b/Engine/Plugins/Experimental/GPULightmass/Source/GPULightmass/Private/LightmapRenderer.cpp
--- a/Engine/Plugins/Experimental/GPULightmass/Source/GPULightmass/Private/LightmapRenderer.cpp
+++ b/Engine/Plugins/Experimental/GPULightmass/Source/GPULightmass/Private/LightmapRenderer.cpp
@@ -1311,6 +1311,7 @@ bool FSceneRenderState::SetupRayTracingScene(FRDGBuilder& GraphBuilder, FSceneUn
 				RayTracingSceneInitializationData.NumCPUInstanceDescriptors,
 				nullptr,
 				nullptr,
+				0,
 				nullptr);
 
 			RHICmdList.BindAccelerationStructureMemory(RayTracingScene, RayTracingSceneBuffer, 0);
diff --git a/Engine/Shaders/Private/RayTracing/RayTracingInstanceBufferUtil.usf b/Engine/Shaders/Private/RayTracing/RayTracingInstanceBufferUtil.usf
--- a/Engine/Shaders/Private/RayTracing/RayTracingInstanceBufferUtil.usf
+++ b/Engine/Shaders/Private/RayTracing/RayTracingInstanceBufferUtil.usf
@@ -50,6 +50,7 @@ RWStructuredBuffer<uint> RWDebugInstanceGPUSceneIndices;
 
 #if OUTPUT_STATS
 RWStructuredBuffer<uint> RWOutputStats;
+uint OutputStatsOffset;
 #endif
 
 #if GPU_CULLING
@@ -212,9 +213,9 @@ void ProcessInstance(FRayTracingInstanceDescriptor InputDesc)
 	if (!bCulled)
 	{
 #if USE_WAVE_OPS
-		WaveInterlockedAddScalar(RWOutputStats[0], 1);
+		WaveInterlockedAddScalar(RWOutputStats[OutputStatsOffset + 0], 1);
 #else
-		InterlockedAdd(RWOutputStats[0], 1);
+		InterlockedAdd(RWOutputStats[OutputStatsOffset + 0], 1);
 #endif
 	}
 #endif
diff --git a/Engine/Source/Runtime/Renderer/Private/DeferredShadingRenderer.cpp b/Engine/Source/Runtime/Renderer/Private/DeferredShadingRenderer.cpp
--- a/Engine/Source/Runtime/Renderer/Private/DeferredShadingRenderer.cpp
+++ b/Engine/Source/Runtime/Renderer/Private/DeferredShadingRenderer.cpp
@@ -3517,7 +3517,7 @@ void FDeferredShadingSceneRenderer::Render(FRDGBuilder& GraphBuilder, const FSce
 		FinishGatheringLumenSurfaceCacheFeedback(GraphBuilder, Views[0], LumenFrameTemporaries);
 
 #if RHI_RAYTRACING
-		RayTracingScene.FinishTracingFeedback(GraphBuilder);
+		RayTracingScene.PostRender(GraphBuilder);
 #endif
 
 		if (ViewFamily.bResolveScene && ViewFamilyTexture)
diff --git a/Engine/Source/Runtime/Renderer/Private/RayTracing/RayTracingInstanceBufferUtil.cpp b/Engine/Source/Runtime/Renderer/Private/RayTracing/RayTracingInstanceBufferUtil.cpp
--- a/Engine/Source/Runtime/Renderer/Private/RayTracing/RayTracingInstanceBufferUtil.cpp
+++ b/Engine/Source/Runtime/Renderer/Private/RayTracing/RayTracingInstanceBufferUtil.cpp
@@ -480,6 +480,7 @@ struct FRayTracingBuildInstanceBufferCS : public FGlobalShader
 		SHADER_PARAMETER(uint32, CullingMode)
 
 		SHADER_PARAMETER_UAV(RWStructuredBuffer<uint>, RWOutputStats)
+		SHADER_PARAMETER(uint32, OutputStatsOffset)
 
 		// Debug parameters
 		SHADER_PARAMETER_UAV(RWStructuredBuffer<uint>, RWDebugInstanceGPUSceneIndices)
@@ -538,6 +539,7 @@ void BuildRayTracingInstanceBuffer(
 	FRHIShaderResourceView* InstanceTransformSRV,
 	const FRayTracingCullingParameters* CullingParameters,
 	FRHIUnorderedAccessView* OutputStatsUAV,
+	uint32 OutputStatsOffset,
 	FRHIUnorderedAccessView* DebugInstanceGPUSceneIndexUAV)
 {
 	FRayTracingBuildInstanceBufferCS::FParameters PassParams;
@@ -570,6 +572,7 @@ void BuildRayTracingInstanceBuffer(
 	}
 
 	PassParams.RWOutputStats = OutputStatsUAV;
+	PassParams.OutputStatsOffset = OutputStatsOffset;
 
 	PassParams.RWDebugInstanceGPUSceneIndices = DebugInstanceGPUSceneIndexUAV;
 
@@ -613,6 +616,7 @@ void BuildRayTracingInstanceBuffer(
 	uint32 NumCPUInstanceDescriptors,
 	const FRayTracingCullingParameters* CullingParameters,
 	FRHIUnorderedAccessView* OutputStatsUAV,
+	uint32 OutputStatsOffset,
 	FRHIUnorderedAccessView* DebugInstanceGPUSceneIndexUAV)
 {
 	if (NumNativeGPUSceneInstances > 0)
@@ -633,6 +637,7 @@ void BuildRayTracingInstanceBuffer(
 			nullptr,
 			CullingParameters,
 			OutputStatsUAV,
+			OutputStatsOffset,
 			DebugInstanceGPUSceneIndexUAV);
 	}
 
@@ -654,6 +659,7 @@ void BuildRayTracingInstanceBuffer(
 			CPUInstanceTransformSRV,
 			nullptr,
 			OutputStatsUAV,
+			OutputStatsOffset,
 			nullptr);
 	}
 }
diff --git a/Engine/Source/Runtime/Renderer/Private/RayTracing/RayTracingScene.cpp b/Engine/Source/Runtime/Renderer/Private/RayTracing/RayTracingScene.cpp
--- a/Engine/Source/Runtime/Renderer/Private/RayTracing/RayTracingScene.cpp
+++ b/Engine/Source/Runtime/Renderer/Private/RayTracing/RayTracingScene.cpp
@@ -51,6 +51,8 @@ const FRayTracingScene::FInstanceHandle FRayTracingScene::INVALID_INSTANCE_HANDL
 static constexpr uint32 AllocationGranularity = 8 * 1024;
 static constexpr uint64 BufferAllocationGranularity = 16 * 1024 * 1024;
 
+using FInstanceBufferStats = uint32;
+
 FRayTracingScene::FRayTracingScene()
 {
 	const uint8 NumLayers = uint8(ERayTracingSceneLayer::NUM);
@@ -102,6 +104,18 @@ void FRayTracingScene::Create(FRDGBuilder& GraphBuilder, const FViewInfo& View,
 
 	const uint8 NumLayers = uint8(ERayTracingSceneLayer::NUM);
 
+	FRDGBufferUAVRef InstanceStatsBufferUAV = nullptr;
+	{
+		// one counter per layer in the stats buffer
+		FRDGBufferDesc InstanceStatsBufferDesc = FRDGBufferDesc::CreateStructuredDesc(sizeof(FInstanceBufferStats), NumLayers);
+		InstanceStatsBufferDesc.Usage |= BUF_SourceCopy;
+
+		InstanceStatsBuffer = GraphBuilder.CreateBuffer(InstanceStatsBufferDesc, TEXT("FRayTracingScene::InstanceStatsBuffer"));
+		InstanceStatsBufferUAV = GraphBuilder.CreateUAV(InstanceStatsBuffer);
+
+		AddClearUAVPass(GraphBuilder, InstanceStatsBufferUAV, 0, ComputePassFlags);
+	}
+
 	for (uint32 LayerIndex = 0; LayerIndex < NumLayers; ++LayerIndex)
 	{
 		FLayer& Layer = Layers[LayerIndex];
@@ -272,15 +286,6 @@ void FRayTracingScene::Create(FRDGBuilder& GraphBuilder, const FViewInfo& View,
 			}
 		}
 
-#if STATS
-		FRDGBufferDesc OutputStatsBufferDesc(FRDGBufferDesc::CreateStructuredDesc(sizeof(uint32), 1));
-		OutputStatsBufferDesc.Usage |= BUF_SourceCopy;
-
-		FRDGBufferRef OutputStatsBuffer = GraphBuilder.CreateBuffer(OutputStatsBufferDesc, TEXT("FRayTracingScene::OutputStatsBuffer"));
-		FRDGBufferUAVRef OutputStatsBufferUAV = GraphBuilder.CreateUAV(OutputStatsBuffer);
-		AddClearUAVPass(GraphBuilder, OutputStatsBufferUAV, 0, ComputePassFlags);
-#endif
-
 		FRDGBufferUAVRef DebugInstanceGPUSceneIndexBufferUAV = nullptr;
 		if (bNeedsDebugInstanceGPUSceneIndexBuffer)
 		{
@@ -432,10 +437,7 @@ void FRayTracingScene::Create(FRDGBuilder& GraphBuilder, const FViewInfo& View,
 			PassParams->InstanceBuffer = GraphBuilder.CreateUAV(Layer.InstanceBuffer);
 			PassParams->DebugInstanceGPUSceneIndexBuffer = DebugInstanceGPUSceneIndexBufferUAV;
 			PassParams->Scene = View.GetSceneUniforms().GetBuffer(GraphBuilder);
-
-#if STATS
-			PassParams->OutputStats = OutputStatsBufferUAV;
-#endif
+			PassParams->OutputStats = STATS ? InstanceStatsBufferUAV : nullptr;
 
 			GraphBuilder.AddPass(
 				RDG_EVENT_NAME("RayTracingBuildInstanceBuffer"),
@@ -444,6 +446,7 @@ void FRayTracingScene::Create(FRDGBuilder& GraphBuilder, const FViewInfo& View,
 				[PassParams,
 				this,
 				&Layer,
+				LayerIndex,
 				GPUScene,
 				NumNativeGPUSceneInstances = Layer.InitializationData.NumNativeGPUSceneInstances,
 				NumNativeCPUInstances = Layer.InitializationData.NumNativeCPUInstances,
@@ -471,63 +474,10 @@ void FRayTracingScene::Create(FRDGBuilder& GraphBuilder, const FViewInfo& View,
 						NumCPUInstanceDescriptors,
 						&CullingParameters,
 						PassParams->OutputStats ? PassParams->OutputStats->GetRHI() : nullptr,
+						LayerIndex,
 						PassParams->DebugInstanceGPUSceneIndexBuffer ? PassParams->DebugInstanceGPUSceneIndexBuffer->GetRHI() : nullptr);
 				});
 		}
-
-#if STATS
-		// Update stats
-		// Currently only supported for Base Layer
-		if(LayerIndex == uint8(ERayTracingSceneLayer::Base))
-		{
-			//  if necessary create readback buffers
-			if (StatsReadbackBuffers.IsEmpty())
-			{
-				StatsReadbackBuffers.SetNum(MaxReadbackBuffers);
-
-				for (uint32 Index = 0; Index < MaxReadbackBuffers; ++Index)
-				{
-					StatsReadbackBuffers[Index] = new FRHIGPUBufferReadback(TEXT("FRayTracingScene::StatsReadbackBuffer"));
-				}
-			}
-
-			// copy stats to readback buffer
-			{
-				AddReadbackBufferPass(GraphBuilder, RDG_EVENT_NAME("FRayTracingScene::StatsReadback"), OutputStatsBuffer,
-					[ReadbackBuffer = StatsReadbackBuffers[StatsReadbackBuffersWriteIndex], OutputStatsBuffer](FRDGAsyncTask, FRHICommandList& RHICmdList)
-					{
-						ReadbackBuffer->EnqueueCopy(RHICmdList, OutputStatsBuffer->GetRHI(), 0u);
-					});
-
-				StatsReadbackBuffersWriteIndex = (StatsReadbackBuffersWriteIndex + 1u) % MaxReadbackBuffers;
-				StatsReadbackBuffersNumPending = FMath::Min(StatsReadbackBuffersNumPending + 1u, MaxReadbackBuffers);
-			}
-
-			// process ready results
-			while (StatsReadbackBuffersNumPending > 0)
-			{
-				uint32 Index = (StatsReadbackBuffersWriteIndex + MaxReadbackBuffers - StatsReadbackBuffersNumPending) % MaxReadbackBuffers;
-				FRHIGPUBufferReadback* ReadbackBuffer = StatsReadbackBuffers[Index];
-				if (ReadbackBuffer->IsReady())
-				{
-					StatsReadbackBuffersNumPending--;
-
-					auto ReadbackBufferPtr = (const uint32*)ReadbackBuffer->Lock(sizeof(uint32));
-
-					NumActiveInstances = ReadbackBufferPtr[0];
-
-					ReadbackBuffer->Unlock();
-				}
-				else
-				{
-					break;
-				}
-			}
-
-			SET_DWORD_STAT(STAT_RayTracingTotalInstances, NumNativeInstances);
-			SET_DWORD_STAT(STAT_RayTracingActiveInstances, FMath::Min(NumActiveInstances, NumNativeInstances));
-		}
-#endif
 	}
 
 #if DO_CHECK
@@ -802,6 +752,78 @@ void FRayTracingScene::FinishTracingFeedback(FRDGBuilder& GraphBuilder, ERDGPass
 	}
 }
 
+void FRayTracingScene::FinishStats(FRDGBuilder& GraphBuilder, ERDGPassFlags ComputePassFlags)
+{
+	//  if necessary create readback buffers
+	if (StatsReadbackBuffers.IsEmpty())
+	{
+		StatsReadbackBuffers.SetNum(MaxReadbackBuffers);
+
+		for (uint32 Index = 0; Index < MaxReadbackBuffers; ++Index)
+		{
+			StatsReadbackBuffers[Index] = new FRHIGPUBufferReadback(TEXT("FRayTracingScene::StatsReadbackBuffer"));
+		}
+	}
+
+	// copy stats to readback buffer
+	if (InstanceStatsBuffer != nullptr)
+	{
+		AddReadbackBufferPass(GraphBuilder, RDG_EVENT_NAME("FRayTracingScene::StatsReadback"), InstanceStatsBuffer,
+			[ReadbackBuffer = StatsReadbackBuffers[StatsReadbackBuffersWriteIndex], InstanceStatsBuffer = InstanceStatsBuffer](FRDGAsyncTask, FRHICommandList& RHICmdList)
+			{
+				ReadbackBuffer->EnqueueCopy(RHICmdList, InstanceStatsBuffer->GetRHI(), 0u);
+			});
+
+		StatsReadbackBuffersWriteIndex = (StatsReadbackBuffersWriteIndex + 1u) % MaxReadbackBuffers;
+		StatsReadbackBuffersNumPending = FMath::Min(StatsReadbackBuffersNumPending + 1u, MaxReadbackBuffers);
+	}
+
+	uint32 TotalNumNativeInstances = 0;
+	uint32 TotalNumActiveInstances = 0;
+
+	const uint8 NumLayers = uint8(ERayTracingSceneLayer::NUM);
+
+	// process ready results
+	while (StatsReadbackBuffersNumPending > 0)
+	{
+		uint32 Index = (StatsReadbackBuffersWriteIndex + MaxReadbackBuffers - StatsReadbackBuffersNumPending) % MaxReadbackBuffers;
+		FRHIGPUBufferReadback* ReadbackBuffer = StatsReadbackBuffers[Index];
+		if (ReadbackBuffer->IsReady())
+		{
+			StatsReadbackBuffersNumPending--;
+
+			auto ReadbackBufferPtr = (const FInstanceBufferStats*)ReadbackBuffer->Lock(sizeof(FInstanceBufferStats) * NumLayers);
+
+			for (uint32 LayerIndex = 0; LayerIndex < NumLayers; ++LayerIndex)
+			{
+				FLayer& Layer = Layers[LayerIndex];
+
+				const uint32 LayerNumNativeInstances = GetNumNativeInstances(ERayTracingSceneLayer(LayerIndex));
+
+				Layer.NumActiveInstances = FMath::Min(ReadbackBufferPtr[LayerIndex], LayerNumNativeInstances);
+
+				TotalNumNativeInstances += LayerNumNativeInstances;
+				TotalNumActiveInstances += Layer.NumActiveInstances;
+			}
+
+			ReadbackBuffer->Unlock();
+		}
+		else
+		{
+			break;
+		}
+	}
+
+	SET_DWORD_STAT(STAT_RayTracingTotalInstances, TotalNumNativeInstances);
+	SET_DWORD_STAT(STAT_RayTracingActiveInstances, TotalNumActiveInstances);
+}
+
+void FRayTracingScene::PostRender(FRDGBuilder& GraphBuilder, ERDGPassFlags ComputePassFlags)
+{
+	FinishTracingFeedback(GraphBuilder, ComputePassFlags);
+	FinishStats(GraphBuilder, ComputePassFlags);
+}
+
 bool FRayTracingScene::IsCreated() const
 {
 	return bUsedThisFrame;
@@ -958,6 +980,8 @@ void FRayTracingScene::Reset(bool bInInstanceDebugDataEnabled)
 		Layer.GeometryHandles.Reset();		
 	}
 
+	InstanceStatsBuffer = nullptr;
+
 	CallableCommands.Reset();
 	UniformBuffers.Reset();
 	GeometriesToBuild.Reset();
@@ -990,14 +1014,7 @@ void FRayTracingScene::EndFrame()
 		GeometriesToBuild.Empty();
 
 		ReleaseFeedbackReadbackBuffers();
-#if STATS
 		ReleaseReadbackBuffers();
-
-		StatsReadbackBuffersWriteIndex = 0;
-		StatsReadbackBuffersNumPending = 0;
-
-		NumActiveInstances = 0;
-#endif
 	}
 
 	bUsedThisFrame = false;
@@ -1006,14 +1023,15 @@ void FRayTracingScene::EndFrame()
 
 void FRayTracingScene::ReleaseReadbackBuffers()
 {
-#if STATS
 	for (auto& ReadbackBuffer : StatsReadbackBuffers)
 	{
 		delete ReadbackBuffer;
 	}
 
 	StatsReadbackBuffers.Empty();
-#endif
+
+	StatsReadbackBuffersWriteIndex = 0;
+	StatsReadbackBuffersNumPending = 0;
 }
 
 void FRayTracingScene::ReleaseFeedbackReadbackBuffers()
diff --git a/Engine/Source/Runtime/Renderer/Private/RayTracing/RayTracingScene.h b/Engine/Source/Runtime/Renderer/Private/RayTracing/RayTracingScene.h
--- a/Engine/Source/Runtime/Renderer/Private/RayTracing/RayTracingScene.h
+++ b/Engine/Source/Runtime/Renderer/Private/RayTracing/RayTracingScene.h
@@ -109,7 +109,7 @@ public:
 
 	void Build(FRDGBuilder& GraphBuilder, ERDGPassFlags ComputePassFlags, FRDGBufferRef DynamicGeometryScratchBuffer);
 
-	void FinishTracingFeedback(FRDGBuilder& GraphBuilder, ERDGPassFlags ComputePassFlags = ERDGPassFlags::Compute);
+	void PostRender(FRDGBuilder& GraphBuilder, ERDGPassFlags ComputePassFlags = ERDGPassFlags::Compute);
 
 	// Resets the instance list and reserves memory for this frame.
 	void Reset(bool bInstanceDebugDataEnabled);
@@ -178,7 +178,12 @@ public:
 
 	// Used for transforming to translated world space in which TLAS was built.
 	FVector PreViewTranslation {};
+
 private:
+
+	void FinishTracingFeedback(FRDGBuilder& GraphBuilder, ERDGPassFlags ComputePassFlags);
+	void FinishStats(FRDGBuilder& GraphBuilder, ERDGPassFlags ComputePassFlags);
+
 	void ReleaseReadbackBuffers();
 	void ReleaseFeedbackReadbackBuffers();
 
@@ -223,6 +228,8 @@ private:
 		TArray<FRayTracingGeometryInstance> Instances;
 
 		TArray<FRayTracingInstanceDebugData> InstancesDebugData;
+
+		uint32 NumActiveInstances = 0;
 	};
 
 	TArray<FLayer> Layers;
@@ -243,17 +250,15 @@ private:
 
 	const uint32 MaxReadbackBuffers = 4;
 
+	FRDGBufferRef InstanceStatsBuffer;
+
 	TArray<FFeedbackReadbackData> FeedbackReadback;
 	uint32 FeedbackReadbackWriteIndex = 0;
 	uint32 FeedbackReadbackNumPending = 0;
 
-#if STATS
 	TArray<FRHIGPUBufferReadback*> StatsReadbackBuffers;
 	uint32 StatsReadbackBuffersWriteIndex = 0;
 	uint32 StatsReadbackBuffersNumPending = 0;
-
-	uint32 NumActiveInstances = 0;
-#endif
 };
 
 #endif // RHI_RAYTRACING
diff --git a/Engine/Source/Runtime/Renderer/Private/Tests/RayTracingTestbed.cpp b/Engine/Source/Runtime/Renderer/Private/Tests/RayTracingTestbed.cpp
--- a/Engine/Source/Runtime/Renderer/Private/Tests/RayTracingTestbed.cpp
+++ b/Engine/Source/Runtime/Renderer/Private/Tests/RayTracingTestbed.cpp
@@ -297,6 +297,7 @@ bool RunRayTracingTestbed_RenderThread(const FString& Parameters)
 		RayTracingSceneInitializationData.NumCPUInstanceDescriptors,
 		nullptr,
 		nullptr,
+		0,
 		nullptr);
 
 	RHICmdList.BindAccelerationStructureMemory(RayTracingSceneRHI, SceneBuffer, 0);
diff --git a/Engine/Source/Runtime/Renderer/Public/RayTracingInstanceBufferUtil.h b/Engine/Source/Runtime/Renderer/Public/RayTracingInstanceBufferUtil.h
--- a/Engine/Source/Runtime/Renderer/Public/RayTracingInstanceBufferUtil.h
+++ b/Engine/Source/Runtime/Renderer/Public/RayTracingInstanceBufferUtil.h
@@ -142,6 +142,7 @@ RENDERER_API void BuildRayTracingInstanceBuffer(
 	uint32 NumCPUInstanceDescriptors,
 	const FRayTracingCullingParameters* CullingParameters,
 	FRHIUnorderedAccessView* OutputStatsUAV,
+	uint32 OutputStatsOffset,
 	FRHIUnorderedAccessView* DebugInstanceGPUSceneIndexUAV);
 
 UE_DEPRECATED(5.6, "Need to provide InstanceGroupUploadSRV.")
@@ -176,6 +177,7 @@ inline void BuildRayTracingInstanceBuffer(
 		0,
 		CullingParameters,
 		OutputStatsUAV,
+		0,
 		DebugInstanceGPUSceneIndexUAV);
 }
 
