Fix rect. light atlas leaking.

* Aligned texture resolution to ensure proper bilinear filtering and avoid leaking on neighbor slot
* Move rect border computation into shader to ensure it works for all MIP level.
* Add debug clear pass to visualize filtering error

diff --git a/Engine/Shaders/Private/RectLight.ush b/Engine/Shaders/Private/RectLight.ush
--- a/Engine/Shaders/Private/RectLight.ush
+++ b/Engine/Shaders/Private/RectLight.ush
@@ -37,6 +37,15 @@ float3 SampleRectTexture(FRectTexture RectTexture, float2 RectUV, float Level, b
 	Level  = min(Level, RectTexture.AtlasMaxLevel);
 
 	RectUV = saturate(RectUV) * RectTexture.AtlasUVScale + RectTexture.AtlasUVOffset;
+
+	// Compute rect border to prevent leaking during tri-linear filtering.
+	// Border are aligned on the coarsest MIP value
+	const uint2 MippedResoluton = uint2(View.RectLightAtlasSizeAndInvSize.xy) >> uint(ceil(Level));
+	const float2 UVBorder = 0.5f / float2(MippedResoluton);
+	const float2 MinRectUV =  UVBorder + RectTexture.AtlasUVOffset;
+	const float2 MaxRectUV = -UVBorder + RectTexture.AtlasUVOffset + RectTexture.AtlasUVScale;
+	RectUV = clamp(RectUV, MinRectUV, MaxRectUV);
+
 	return bIsValid ? View.RectLightAtlasTexture.SampleLevel(View.SharedTrilinearClampedSampler, RectUV, bIsReference ? 0 : Level).rgb : 1.f;
 #else
 	return 1;
diff --git a/Engine/Shaders/Private/RectLightAtlas.usf b/Engine/Shaders/Private/RectLightAtlas.usf
--- a/Engine/Shaders/Private/RectLightAtlas.usf
+++ b/Engine/Shaders/Private/RectLightAtlas.usf
@@ -4,6 +4,33 @@
 #include "ShaderPrint.ush"
 #include "ColorMap.ush"
 
+////////////////////////////////////////////////////////////////////////////////////////////////
+// Clear texture
+#if SHADER_CLEAR_TEXTURE
+
+int2 Resolution;
+
+RWTexture2D<float4> OutputTexture;
+
+[numthreads(8, 8, 1)]
+void MainCS(uint3 DispatchThreadId : SV_DispatchThreadID)
+{
+	const uint2 PixelCoord = DispatchThreadId.xy;
+	if (all(PixelCoord < Resolution))
+	{
+		float4 Random = 0; 
+		Random.x = InterleavedGradientNoise(PixelCoord + 0.5f, 0);
+		Random.y = InterleavedGradientNoise(PixelCoord + 0.5f, 1);
+		Random.z = InterleavedGradientNoise(PixelCoord + 0.5f, 2);
+		Random.w = 1;
+
+		Random = float4(1, 0, 1, 1);
+		OutputTexture[PixelCoord] = Random;
+	}
+}
+
+#endif // SHADER_CLEAR_TEXTURE
+
 ////////////////////////////////////////////////////////////////////////////////////////////////
 // Add texture
 #if SHADER_RECT_VS
diff --git a/Engine/Source/Runtime/Renderer/Private/RectLightTextureManager.cpp b/Engine/Source/Runtime/Renderer/Private/RectLightTextureManager.cpp
--- a/Engine/Source/Runtime/Renderer/Private/RectLightTextureManager.cpp
+++ b/Engine/Source/Runtime/Renderer/Private/RectLightTextureManager.cpp
@@ -182,6 +182,7 @@ struct FAtlasSlot
 	FAtlasRect Rect;
 	FVector4f SourceScaleOffset;
 	FTextureReferenceRHIRef SourceTexture = nullptr;
+	FIntPoint CachedSourceTextureResolution = FIntPoint(0, 0);
 	uint32 RefCount = 0;
 	bool bForceRefresh = false;
 	bool IsValid() const { return SourceTexture != nullptr; }
@@ -249,6 +250,14 @@ static int32 GetSlotMaxMIPLevel(const FAtlasSlot& In)
 	return MaxMIP > 0u ? MaxMIP-1u : 0u;
 }
 
+// Align the slot resolution to prevent bilinear filtering to leak on neighbors' slot
+static void AlignSlotResolution(FAtlasSlot& In)
+{
+	const uint32 MaxMIP = GetSlotMaxMIPLevel(In);
+	const FIntPoint ResolutionMIP = ToMIP(In.Rect.Resolution, MaxMIP);
+	In.Rect.Resolution = FIntPoint(ResolutionMIP.X << MaxMIP, ResolutionMIP.Y << MaxMIP);
+}
+
 static EPixelFormat GetRectLightAtlasFormat()
 {
 	return CVarRectLighAtlasFormat.GetValueOnRenderThread() > 0 ? PF_FloatRGBA : PF_FloatR11G11B10;
@@ -333,6 +342,58 @@ static FRDGBufferRef CreateScaleOffsetBuffer(FRDGBuilder& GraphBuilder, const TA
 	return Buffer;
 }
 
+///////////////////////////////////////////////////////////////////////////////////////////////////
+// Clear
+#define RECT_ATLAS_DEBUG_CLEAR 0
+
+#if RECT_ATLAS_DEBUG_CLEAR
+class FRectLightAtlasClearCS : public FGlobalShader
+{
+	DECLARE_GLOBAL_SHADER(FRectLightAtlasClearCS);
+	SHADER_USE_PARAMETER_STRUCT(FRectLightAtlasClearCS, FGlobalShader);
+
+	BEGIN_SHADER_PARAMETER_STRUCT(FParameters, )
+		SHADER_PARAMETER(FIntPoint, Resolution)
+		SHADER_PARAMETER_RDG_TEXTURE_UAV(RWTexture2D, OutputTexture)
+	END_SHADER_PARAMETER_STRUCT()
+
+public:
+		static bool ShouldCompilePermutation(const FGlobalShaderPermutationParameters& Parameters) { return true; }
+	static void ModifyCompilationEnvironment(const FGlobalShaderPermutationParameters& Parameters, FShaderCompilerEnvironment& OutEnvironment)
+	{
+		FGlobalShader::ModifyCompilationEnvironment(Parameters, OutEnvironment);
+		OutEnvironment.SetDefine(TEXT("SHADER_CLEAR_TEXTURE"), 1);
+	}
+	static EShaderPermutationPrecacheRequest ShouldPrecachePermutation(const FGlobalShaderPermutationParameters& Parameters)
+	{
+		return EShaderPermutationPrecacheRequest::NotPrecached;
+	}
+};
+
+IMPLEMENT_GLOBAL_SHADER(FRectLightAtlasClearCS, "/Engine/Private/RectLightAtlas.usf", "MainCS", SF_Compute);
+
+static FRDGTextureRef AddRectLightClearPass(
+	FRDGBuilder& GraphBuilder,
+	FGlobalShaderMap* ShaderMap,
+	FRDGTextureRef OutputTexture)
+{
+	const FIntPoint OutputResolution(OutputTexture->Desc.Extent);
+	const uint32 MipCount = FMath::FloorLog2(OutputResolution.X);
+	for (uint32 MipIt = 0; MipIt < MipCount; ++MipIt)
+	{
+		FRectLightAtlasClearCS::FParameters* Parameters = GraphBuilder.AllocParameters<FRectLightAtlasClearCS::FParameters>();
+		Parameters->Resolution.X = OutputResolution.X>>MipIt;
+		Parameters->Resolution.Y = OutputResolution.Y>>MipIt;
+		Parameters->OutputTexture = GraphBuilder.CreateUAV(FRDGTextureUAVDesc(OutputTexture, MipIt));
+
+		TShaderMapRef<FRectLightAtlasClearCS> ComputeShader(ShaderMap);
+		const FIntVector DispatchCount = DispatchCount.DivideAndRoundUp(FIntVector(OutputResolution.X, OutputResolution.Y, 1), FIntVector(8, 8, 1));
+		FComputeShaderUtils::AddPass(GraphBuilder, RDG_EVENT_NAME("RectLightAtlas::Clear"), ComputeShader, Parameters, DispatchCount);
+	}
+
+	return OutputTexture;
+}
+#endif
 
 ///////////////////////////////////////////////////////////////////////////////////////////////////
 // Debug
@@ -1390,11 +1451,13 @@ static void PackAtlas(
 			// Check if the texture resolution has changed (due to streaming)
 			// If the texture resolution has change, reset the slot to be handled as a new slot
 			const FIntPoint TextureResolution = ValidSlot.GetSourceResolution();
-			const bool bHasTextureResolutionChanged = ValidSlot.Rect.Resolution != TextureResolution;
+			const bool bHasTextureResolutionChanged = ValidSlot.CachedSourceTextureResolution != TextureResolution;
 			if (bHasTextureResolutionChanged || ValidSlot.bForceRefresh)
 			{
 				ValidSlot.Rect.Origin = InvalidOrigin;
-				ValidSlot.Rect.Resolution = FIntPoint(TextureResolution.X, TextureResolution.Y);
+				ValidSlot.Rect.Resolution = TextureResolution;
+				AlignSlotResolution(ValidSlot);
+				ValidSlot.CachedSourceTextureResolution = TextureResolution;
 			}
 
 			TargetPixelCount += ValidSlot.Rect.Resolution.X * ValidSlot.Rect.Resolution.Y;
@@ -1499,6 +1562,8 @@ static void PackAtlas(
 						const FIntPoint SourceResolution = Slot.GetSourceResolution();
 						Slot.Rect.Origin = InvalidOrigin;
 						Slot.Rect.Resolution = ToMIP(SourceResolution, CurrentSourceTextureMIPBias);
+						Slot.CachedSourceTextureResolution = Slot.Rect.Resolution;
+						AlignSlotResolution(Slot);
 					}
 				}
 
@@ -1559,8 +1624,12 @@ uint32 AddTexture(UTexture* In, const FVector4f& ScaleOffset)
 			Slot->Id = SlotIndex;
 			Slot->Rect.Origin = InvalidOrigin;
 			Slot->Rect.Resolution = GetSourceResolution(Tex, ScaleOffset);
+			Slot->CachedSourceTextureResolution = Slot->Rect.Resolution;
 			Slot->RefCount = 1;
 
+			// Force alignement based on the max. MIP-level
+			AlignSlotResolution(*Slot);
+
 			GRectLightTextureManager.bHasPendingAdds = true;
 		}
 	}
@@ -1608,10 +1677,10 @@ FAtlasSlotDesc GetAtlasSlot(uint32 InSlotIndex)
 		const FIntPoint Resolution = GRectLightTextureManager.AtlasTexture->GetDesc().Extent;
 		const FVector2f InvResolution(1.f / Resolution.X, 1.f / Resolution.Y);
 
-		// Shrink rect, texture by 1 pixels (0.5px offset, and -1 on resolution) 
-		// so that adjacent slots don't leak during bilinear filtering
-		Out.UVOffset = FVector2f((Slot.Rect.Origin.X + 0.5f)  * InvResolution.X, (Slot.Rect.Origin.Y + 0.5f)   * InvResolution.Y);
-		Out.UVScale  = FVector2f((Slot.Rect.Resolution.X-1.f) * InvResolution.X, (Slot.Rect.Resolution.Y-1.0f) * InvResolution.Y);
+		// UV rect shrinking, to avoid leaking during tri/bi-linear filtering, is done within the shader, as it depends on the selected MIP value.
+		// UVOffset/UVScale gives the rect value without any border/clamping
+		Out.UVOffset = FVector2f(Slot.Rect.Origin.X * InvResolution.X, Slot.Rect.Origin.Y * InvResolution.Y);
+		Out.UVScale  = FVector2f(Slot.Rect.Resolution.X * InvResolution.X, Slot.Rect.Resolution.Y * InvResolution.Y);
 		Out.MaxMipLevel = GetSlotMaxMIPLevel(Slot);
 	}
 
@@ -1725,6 +1794,11 @@ void UpdateAtlasTexture(FRDGBuilder& GraphBuilder, const ERHIFeatureLevel::Type
 			if (bRecreateAtlasTexture)
 			{
 				AtlasTexture = CreateRectLightAtlasTexture(GraphBuilder, GRectLightTextureManager.AtlasLayout.AtlasResolution, AtlasFormat);
+
+				#if RECT_ATLAS_DEBUG_CLEAR
+				AddRectLightClearPass(GraphBuilder, ShaderMap, AtlasTexture);
+				#endif
+
 				bNeedExtraction = true;
 				// Sanity check
 				check(CopySlots.Num() == 0);
@@ -1736,7 +1810,11 @@ void UpdateAtlasTexture(FRDGBuilder& GraphBuilder, const ERHIFeatureLevel::Type
 				// Clear the atlas (optional)
 				if (bClearAtlas && bResetOnChange && CopySlots.Num() == 0)
 				{
+					#if RECT_ATLAS_DEBUG_CLEAR
+					AddRectLightClearPass(GraphBuilder, ShaderMap, AtlasTexture);
+					#else
 					AddClearRenderTargetPass(GraphBuilder, AtlasTexture, FLinearColor::Transparent);
+					#endif
 				}
 			}
 
@@ -1746,6 +1824,10 @@ void UpdateAtlasTexture(FRDGBuilder& GraphBuilder, const ERHIFeatureLevel::Type
 				FRDGTextureRef NewAtlasTexture = CreateRectLightAtlasTexture(GraphBuilder, GRectLightTextureManager.AtlasLayout.AtlasResolution, AtlasFormat);
 				bNeedExtraction = true;
 
+				#if RECT_ATLAS_DEBUG_CLEAR
+				AddRectLightClearPass(GraphBuilder, ShaderMap, NewAtlasTexture);
+				#endif
+
 				CopySlotsPass(GraphBuilder, ShaderMap, CopySlots, AtlasTexture, NewAtlasTexture);
 
 				for (FAtlasCopySlot& Slot : CopySlots)
