Fix check failing when ray tracing is toggled off at runtime.

diff --git a/Engine/Source/Runtime/Renderer/Private/RayTracing/RayTracingScene.cpp b/Engine/Source/Runtime/Renderer/Private/RayTracing/RayTracingScene.cpp
--- a/Engine/Source/Runtime/Renderer/Private/RayTracing/RayTracingScene.cpp
+++ b/Engine/Source/Runtime/Renderer/Private/RayTracing/RayTracingScene.cpp
@@ -794,7 +794,7 @@ void FRayTracingScene::FinishStats(FRDGBuilder& GraphBuilder, ERDGPassFlags Comp
 			{
 				FLayer& Layer = Layers[LayerIndex];
 
-				const uint32 LayerNumNativeInstances = GetNumNativeInstances(ERayTracingSceneLayer(LayerIndex));
+				const uint32 LayerNumNativeInstances = Layer.InitializationData.NumNativeGPUSceneInstances + Layer.InitializationData.NumNativeCPUInstances;
 
 				Layer.NumActiveInstances = FMath::Min(ReadbackBufferPtr[LayerIndex], LayerNumNativeInstances);
 
