Fix assert in HZB setup code, when there is no HZB texture.

diff --git a/Engine/Source/Runtime/Renderer/Private/HZB.cpp b/Engine/Source/Runtime/Renderer/Private/HZB.cpp
--- a/Engine/Source/Runtime/Renderer/Private/HZB.cpp
+++ b/Engine/Source/Runtime/Renderer/Private/HZB.cpp
@@ -55,13 +55,12 @@ FHZBParameters GetHZBParameters(FRDGBuilder& GraphBuilder, const FViewInfo& View
 	FIntPoint TextureExtent;
 	if (InHZBTypes & EHZBType::FurthestHZB)
 	{
-		TextureExtent = View.HZB ? View.HZB->Desc.Extent : FIntPoint::ZeroValue;
+		TextureExtent = View.HZB ? View.HZB->Desc.Extent : FIntPoint(1,1);
 	}
 	else
 	{
-		TextureExtent = View.ClosestHZB ? View.ClosestHZB->Desc.Extent : FIntPoint::ZeroValue;
+		TextureExtent = View.ClosestHZB ? View.ClosestHZB->Desc.Extent : FIntPoint(1,1);
 	}
-	check(TextureExtent != FIntPoint::ZeroValue);
 
 	FHZBParameters Out = InitHZBCommonParameter(GraphBuilder, &View, View.ViewRect, TextureExtent);
 
