Eye Adaptation UDN crash workaround when no EyeAdaptationViewState is present.  I'm not able to reproduce this locally, as eye adaptation code should only be reached when the bResolveScene is set on the view family, and the PostProcessing flag is set in the engine show flags.  When both of those are set, the FSceneView constructor fills in EyeAdaptationViewState.  The workaround is to explicitly check that EyeAdaptationViewState has been set before enabling eye adaptation, so if that scenario somehow occurs (perhaps due to local code changes, scene view extensions, or some other unforeseen code path), it won't crash.

#rnx

diff --git a/Engine/Source/Runtime/Engine/Public/SceneView.h b/Engine/Source/Runtime/Engine/Public/SceneView.h
--- a/Engine/Source/Runtime/Engine/Public/SceneView.h
+++ b/Engine/Source/Runtime/Engine/Public/SceneView.h
@@ -2009,6 +2009,11 @@ public:
 		return bAllowRayTracing && (!bIsSceneCapture || bSceneCaptureUsesRayTracing);
 	}
 
+	inline bool HasEyeAdaptationViewState() const
+	{
+		return EyeAdaptationViewState != nullptr;
+	}
+
 protected:
 	friend class FViewSnapshotCache;
 
diff --git a/Engine/Source/Runtime/Renderer/Private/PostProcess/PostProcessMobile.cpp b/Engine/Source/Runtime/Renderer/Private/PostProcess/PostProcessMobile.cpp
--- a/Engine/Source/Runtime/Renderer/Private/PostProcess/PostProcessMobile.cpp
+++ b/Engine/Source/Runtime/Renderer/Private/PostProcess/PostProcessMobile.cpp
@@ -55,7 +55,7 @@ float GetMobileDepthOfFieldScale(const FViewInfo& View)
 
 bool IsMobileEyeAdaptationEnabled(const FViewInfo& View)
 {
-	return View.ViewState != nullptr && View.Family->EngineShowFlags.EyeAdaptation && CVarMobileEyeAdaptation.GetValueOnRenderThread() == 1 && IsMobileHDR();
+	return View.ViewState != nullptr && View.HasEyeAdaptationViewState() && View.Family->EngineShowFlags.EyeAdaptation && CVarMobileEyeAdaptation.GetValueOnRenderThread() == 1 && IsMobileHDR();
 }
 
 //Following variations are always generated
diff --git a/Engine/Source/Runtime/Renderer/Private/PostProcess/PostProcessing.cpp b/Engine/Source/Runtime/Renderer/Private/PostProcess/PostProcessing.cpp
--- a/Engine/Source/Runtime/Renderer/Private/PostProcess/PostProcessing.cpp
+++ b/Engine/Source/Runtime/Renderer/Private/PostProcess/PostProcessing.cpp
@@ -694,6 +694,7 @@ void AddPostProcessingPasses(
 		const bool bEyeAdaptationEnabled =
 			// Skip for transient views.
 			bHasViewState &&
+			View.HasEyeAdaptationViewState() &&
 			// Skip for secondary views in a stereo setup.
 			bPrimaryView;
 
diff --git a/Engine/Source/Runtime/Renderer/Private/SceneRendering.cpp b/Engine/Source/Runtime/Renderer/Private/SceneRendering.cpp
--- a/Engine/Source/Runtime/Renderer/Private/SceneRendering.cpp
+++ b/Engine/Source/Runtime/Renderer/Private/SceneRendering.cpp
@@ -2419,8 +2419,11 @@ void FViewInfo::EnqueueEyeAdaptationExposureBufferReadback(FRDGBuilder& GraphBui
 
 bool FViewInfo::ShouldUpdateEyeAdaptationBuffer() const
 {
-	// If this view owns its eye adaptation view state (no secondary view state set, or secondary view state is equal), it should update
-	if (!EyeAdaptationViewState || EyeAdaptationViewState == reinterpret_cast<const FSceneViewStateInterface*>(ViewState))
+	// This code should only be reached if eye adaptation is enabled (calling code should check HasEyeAdaptationViewState())
+	check(EyeAdaptationViewState);
+
+	// If this view owns its eye adaptation view state (equal to main view state), it should update
+	if (EyeAdaptationViewState == reinterpret_cast<const FSceneViewStateInterface*>(ViewState))
 	{
 		return true;
 	}
