Use high priority task for visibility command pipe to avoid stalls with high task graph contention.

diff --git a/Engine/Source/Runtime/Renderer/Private/SceneVisibilityPrivate.h b/Engine/Source/Runtime/Renderer/Private/SceneVisibilityPrivate.h
--- a/Engine/Source/Runtime/Renderer/Private/SceneVisibilityPrivate.h
+++ b/Engine/Source/Runtime/Renderer/Private/SceneVisibilityPrivate.h
@@ -112,7 +112,7 @@ public:
 					ReleaseNumCommands(NumProcessedCommands);
 				}
 
-			}, PrerequisiteTask);
+			}, PrerequisiteTask, UE::Tasks::ETaskPriority::High);
 		}
 	}
 
