Fixed distance culling bug for Nanite rendering into CSM, where it would not set up the culling view overrides correctly leading to issues with e.g., per instance cull distance (foliage).

diff --git a/Engine/Source/Runtime/Renderer/Private/ShadowDepthRendering.cpp b/Engine/Source/Runtime/Renderer/Private/ShadowDepthRendering.cpp
--- a/Engine/Source/Runtime/Renderer/Private/ShadowDepthRendering.cpp
+++ b/Engine/Source/Runtime/Renderer/Private/ShadowDepthRendering.cpp
@@ -1446,6 +1446,7 @@ static void RenderShadowDepthAtlasNanite(
 						UpdatePackedViewParamsFromPrevShadowState(Initializer, PrevShadowState);
 						UpdateCurrentFrameHZB(LightSceneInfo, ShadowKey, ProjectedShadowInfo, nullptr);
 	
+						Nanite::SetCullingViewOverrides(ProjectedShadowInfo->ShadowDepthView, Initializer);
 						PackedViews.Add(Nanite::CreatePackedView(Initializer));
 					}
 	
