Removed unused flag bIsFarField from FRayTracingCachedMeshCommandFlags.

diff --git a/Engine/Source/Runtime/Renderer/Private/PrimitiveSceneInfo.cpp b/Engine/Source/Runtime/Renderer/Private/PrimitiveSceneInfo.cpp
--- a/Engine/Source/Runtime/Renderer/Private/PrimitiveSceneInfo.cpp
+++ b/Engine/Source/Runtime/Renderer/Private/PrimitiveSceneInfo.cpp
@@ -1160,7 +1160,6 @@ void CacheRayTracingPrimitive(
 			{
 				FPrimitiveSceneInfo::FRayTracingLODData& LODData = RayTracingLODData[LODIndex];
 				LODData.CachedMeshCommandFlags.InstanceMask = GetFarFieldRayTracingInstanceMask(RayTracingType);
-				LODData.CachedMeshCommandFlags.bIsFarField = true;
 			}
 		}
 
diff --git a/Engine/Source/Runtime/Renderer/Public/RayTracingMeshDrawCommands.h b/Engine/Source/Runtime/Renderer/Public/RayTracingMeshDrawCommands.h
--- a/Engine/Source/Runtime/Renderer/Public/RayTracingMeshDrawCommands.h
+++ b/Engine/Source/Runtime/Renderer/Public/RayTracingMeshDrawCommands.h
@@ -24,7 +24,6 @@ struct FRayTracingCachedMeshCommandFlags
 		bAllSegmentsDecal = true;
 		bTwoSided = false;
 		bIsSky = false;
-		bIsFarField = false;
 		bAllSegmentsTranslucent = true;
 		bAllSegmentsReverseCulling = true;
 	}
@@ -58,7 +57,6 @@ struct FRayTracingCachedMeshCommandFlags
 			bool bAllSegmentsDecal : 1;
 			bool bTwoSided : 1;
 			bool bIsSky : 1;
-			bool bIsFarField : 1;
 			bool bAllSegmentsTranslucent : 1;
 			bool bAllSegmentsReverseCulling : 1;
 		};
