[HWRT] Deprecate/delete NumSegments tracking at Ray Tracing Scene level.
- it is now tracked and managed by FShaderBindingTable.

diff --git a/Engine/Plugins/Experimental/GPULightmass/Source/GPULightmass/Private/LightmapRenderer.cpp b/Engine/Plugins/Experimental/GPULightmass/Source/GPULightmass/Private/LightmapRenderer.cpp
--- a/Engine/Plugins/Experimental/GPULightmass/Source/GPULightmass/Private/LightmapRenderer.cpp
+++ b/Engine/Plugins/Experimental/GPULightmass/Source/GPULightmass/Private/LightmapRenderer.cpp
@@ -1187,7 +1187,6 @@ bool FSceneRenderState::SetupRayTracingScene(FRDGBuilder& GraphBuilder, FSceneUn
 				FRayTracingSceneInitializer Initializer;
 				Initializer.DebugName = FName(TEXT("LightmapRendererRayTracingScene"));
 				Initializer.MaxNumInstances = RayTracingSceneInitializationData.NumNativeGPUSceneInstances + RayTracingSceneInitializationData.NumNativeCPUInstances;
-				Initializer.NumTotalSegments = RayTracingSceneInitializationData.TotalNumSegments;
 				Initializer.BuildFlags = SceneBuildFlags;
 
 				RayTracingScene = RHICreateRayTracingScene(MoveTemp(Initializer));
diff --git a/Engine/Source/Runtime/RHI/Public/RHIResources.h b/Engine/Source/Runtime/RHI/Public/RHIResources.h
--- a/Engine/Source/Runtime/RHI/Public/RHIResources.h
+++ b/Engine/Source/Runtime/RHI/Public/RHIResources.h
@@ -3674,6 +3674,7 @@ struct FRayTracingSceneInitializer
 	// Maximum number of instances in this scene. Actual number of instances is specified in FRayTracingSceneBuildParams.
 	uint32 MaxNumInstances = 0;
 
+	UE_DEPRECATED(5.6, "Use FRayTracingShaderBindingTableInitializer instead.")
 	uint32 NumTotalSegments = 0;
 
 	UE_DEPRECATED(5.5, "Use FRayTracingShaderBindingTableInitializer instead.")
diff --git a/Engine/Source/Runtime/Renderer/Private/HeterogeneousVolumes/HeterogeneousVolumesHardwareRayTracing.cpp b/Engine/Source/Runtime/Renderer/Private/HeterogeneousVolumes/HeterogeneousVolumesHardwareRayTracing.cpp
--- a/Engine/Source/Runtime/Renderer/Private/HeterogeneousVolumes/HeterogeneousVolumesHardwareRayTracing.cpp
+++ b/Engine/Source/Runtime/Renderer/Private/HeterogeneousVolumes/HeterogeneousVolumesHardwareRayTracing.cpp
@@ -712,10 +712,10 @@ void RenderLightingCacheWithPreshadingHardwareRayTracing(
 			// WARNING: Currently hit-group indexing is required to bind uniform buffers to hit-group shaders.
 			SBTInitializer.HitGroupIndexingMode = ERayTracingHitGroupIndexingMode::Allow;
 			SBTInitializer.ShaderBindingMode = ERayTracingShaderBindingMode::RTPSO;
-			SBTInitializer.NumGeometrySegments = RayTracingScene.GetTotalNumSegments();
+			SBTInitializer.NumGeometrySegments = 1;
 			SBTInitializer.NumShaderSlotsPerGeometrySegment = RAY_TRACING_NUM_SHADER_SLOTS;
-			SBTInitializer.NumMissShaderSlots = RayTracingScene.NumMissShaderSlots;
-			SBTInitializer.NumCallableShaderSlots = RayTracingScene.NumCallableShaderSlots;
+			SBTInitializer.NumMissShaderSlots = 1;
+			SBTInitializer.NumCallableShaderSlots = 0;
 			SBTInitializer.LocalBindingDataSize = MaxLocalBindingDataSize;
 
 			FShaderBindingTableRHIRef SBT = RHICmdList.CreateRayTracingShaderBindingTable(SBTInitializer);
@@ -888,10 +888,10 @@ void RenderSingleScatteringWithPreshadingHardwareRayTracing(
 			// WARNING: Currently hit-group indexing is required to bind uniform buffers to hit-group shaders.
 			SBTInitializer.HitGroupIndexingMode = ERayTracingHitGroupIndexingMode::Allow;
 			SBTInitializer.ShaderBindingMode = ERayTracingShaderBindingMode::RTPSO;
-			SBTInitializer.NumGeometrySegments = RayTracingScene.GetTotalNumSegments();
+			SBTInitializer.NumGeometrySegments = 1;
 			SBTInitializer.NumShaderSlotsPerGeometrySegment = RAY_TRACING_NUM_SHADER_SLOTS;
-			SBTInitializer.NumMissShaderSlots = RayTracingScene.NumMissShaderSlots;
-			SBTInitializer.NumCallableShaderSlots = RayTracingScene.NumCallableShaderSlots;
+			SBTInitializer.NumMissShaderSlots = 0;
+			SBTInitializer.NumCallableShaderSlots = 0;
 			SBTInitializer.LocalBindingDataSize = MaxLocalBindingDataSize;
 
 			FShaderBindingTableRHIRef SBT = RHICmdList.CreateRayTracingShaderBindingTable(SBTInitializer);
diff --git a/Engine/Source/Runtime/Renderer/Private/RayTracing/RayTracing.cpp b/Engine/Source/Runtime/Renderer/Private/RayTracing/RayTracing.cpp
--- a/Engine/Source/Runtime/Renderer/Private/RayTracing/RayTracing.cpp
+++ b/Engine/Source/Runtime/Renderer/Private/RayTracing/RayTracing.cpp
@@ -323,7 +323,6 @@ namespace RayTracing
 		int32 NumCachedStaticInstances = 0;
 		int32 NumCachedStaticDecalInstances = 0;
 		int32 NumCachedStaticFarFieldInstances = 0;
-		int32 NumCachedStaticSegments = 0; 
 		int32 NumCachedStaticVisibleShaderBindings = 0;
 
 		// Indicates that this object has been fully produced (for validation)
@@ -884,8 +883,6 @@ namespace RayTracing
 					FRayTracingSBTAllocation* SBTAllocation = RayTracingSBT.AllocateDynamicRange(ActiveLayers, Instance.GetMaterials().Num());
 					if (bNeedMainInstance)
 					{
-						RayTracingScene.NumSegments += Instance.GetMaterials().Num();
-
 						RayTracingInstance.InstanceContributionToHitGroupIndex = SBTAllocation->GetInstanceContributionToHitGroupIndex(ERayTracingShaderBindingLayer::Base);
 
 						ERayTracingPrimitiveFlags Flags = Scene.PrimitiveRayTracingFlags[PrimitiveIndex];
@@ -894,8 +891,6 @@ namespace RayTracing
 
 					if (bNeedDecalInstance)
 					{
-						RayTracingScene.NumSegments += Instance.GetMaterials().Num();
-
 						FRayTracingGeometryInstance DecalRayTracingInstance = RayTracingInstance;
 						DecalRayTracingInstance.InstanceContributionToHitGroupIndex = SBTAllocation->GetInstanceContributionToHitGroupIndex(ERayTracingShaderBindingLayer::Decals);
 						RayTracingScene.AddInstance(MoveTemp(DecalRayTracingInstance), ERayTracingSceneLayer::Decals, SceneProxy, true, Geometry->GetGeometryHandle());
@@ -1131,7 +1126,6 @@ namespace RayTracing
 			int32 NumCachedStaticInstances = 0;
 			int32 NumCachedStaticDecalInstances = 0;
 			int32 NumCachedStaticFarFieldInstances = 0;
-			int32 NumCachedStaticSegments = 0;
 			int32 NumCachedStaticVisibleShaderBindings = 0;
 
 			int32 ContextIndex = INDEX_NONE;
@@ -1266,22 +1260,6 @@ namespace RayTracing
 						++Context.NumCachedStaticFarFieldInstances;
 					}
 
-					uint32 NumTLASInstances = 0;
-					if (bNeedMainInstance)
-					{
-						++NumTLASInstances;
-					}
-					if (bNeedDecalInstance)
-					{
-						++NumTLASInstances;
-					}
-					if (bNeedFarFieldInstance)
-					{
-						++NumTLASInstances;
-					}
-
-					Context.NumCachedStaticSegments += RTLODData.CachedMeshCommandIndices.Num() * NumTLASInstances;
-
 					// Setup the shader bindings write offset which is computed here for all cached instances so the writing can happen in parallel to the visible bindings array
 					RelevantPrimitive->RelativeVisibleShaderBindingsOffset = Context.NumCachedStaticVisibleShaderBindings;
 					Context.NumCachedStaticVisibleShaderBindings += RTLODData.SBTAllocation->GetSegmentCount();
@@ -1366,7 +1344,6 @@ namespace RayTracing
 				TaskData.NumCachedStaticInstances += Context.NumCachedStaticInstances;
 				TaskData.NumCachedStaticDecalInstances += Context.NumCachedStaticDecalInstances;
 				TaskData.NumCachedStaticFarFieldInstances += Context.NumCachedStaticFarFieldInstances;
-				TaskData.NumCachedStaticSegments += Context.NumCachedStaticSegments;
 				TaskData.NumCachedStaticVisibleShaderBindings += Context.NumCachedStaticVisibleShaderBindings;
 
 				for (const FPrimitiveSceneInfo* SceneInfo : Context.VisibleNaniteRayTracingPrimitives)
@@ -1564,7 +1541,6 @@ namespace RayTracing
 				uint32 SegmentCount = RelevantPrimitive.RayTracingGeometry->Initializer.Segments.Num();
 				if (bNeedMainInstance)
 				{
-					RayTracingScene.NumSegments += SegmentCount;
 					RayTracingInstance.InstanceContributionToHitGroupIndex = RTLODData.SBTAllocation->GetInstanceContributionToHitGroupIndex(ERayTracingShaderBindingLayer::Base);
 					InstanceBatch.InstanceHandle = RayTracingScene.AddInstance(RayTracingInstance, EnumHasAllFlags(Flags, ERayTracingPrimitiveFlags::FarField) ? ERayTracingSceneLayer::FarField : ERayTracingSceneLayer::Base, SceneProxy, false, RelevantPrimitive.RayTracingGeometry->GetGeometryHandle());
 				}
@@ -1572,7 +1548,6 @@ namespace RayTracing
 				InstanceBatch.DecalInstanceHandle = FRayTracingScene::INVALID_INSTANCE_HANDLE;
 				if (bNeedDecalInstance)
 				{
-					RayTracingScene.NumSegments += SegmentCount;
 					FRayTracingGeometryInstance DecalRayTracingInstance = RayTracingInstance;
 					DecalRayTracingInstance.InstanceContributionToHitGroupIndex = RTLODData.SBTAllocation->GetInstanceContributionToHitGroupIndex(ERayTracingShaderBindingLayer::Decals);
 					InstanceBatch.DecalInstanceHandle = RayTracingScene.AddInstance(MoveTemp(DecalRayTracingInstance), ERayTracingSceneLayer::Decals, SceneProxy, false, RelevantPrimitive.RayTracingGeometry->GetGeometryHandle());
@@ -1619,7 +1594,6 @@ namespace RayTracing
 		const int32 NumCachedStaticInstances,
 		const int32 NumCachedStaticDecalInstances,
 		const int32 NumCachedStaticFarFieldInstances,
-		const int32 NumCachedStaticSegments,
 		const int32 NumCachedStaticVisibleShaderBindings,
 		FRayTracingScene& RayTracingScene,
 		TArray<FRayTracingShaderBindingData>& VisibleShaderBindingData)
@@ -1630,7 +1604,6 @@ namespace RayTracing
 		const FRayTracingScene::FInstanceRange CachedStaticDecalInstanceRange = RayTracingScene.AllocateInstanceRangeUninitialized(NumCachedStaticDecalInstances, ERayTracingSceneLayer::Decals);
 		const FRayTracingScene::FInstanceRange CachedStaticFarFieldInstanceRange = RayTracingScene.AllocateInstanceRangeUninitialized(NumCachedStaticFarFieldInstances, ERayTracingSceneLayer::FarField);
 		const uint32 BaseCachedVisibleShaderBindingsIndex = VisibleShaderBindingData.AddUninitialized(NumCachedStaticVisibleShaderBindings);
-		RayTracingScene.NumSegments += NumCachedStaticSegments;
 
 		const int32 MinBatchSize = 128;
 		ParallelFor(
@@ -1872,7 +1845,6 @@ namespace RayTracing
 					TaskData.NumCachedStaticInstances,
 					TaskData.NumCachedStaticDecalInstances,
 					TaskData.NumCachedStaticFarFieldInstances,
-					TaskData.NumCachedStaticSegments,
 					TaskData.NumCachedStaticVisibleShaderBindings,
 					RayTracingScene,
 					View.VisibleRayTracingShaderBindings);
diff --git a/Engine/Source/Runtime/Renderer/Private/RayTracing/RayTracingInstanceBufferUtil.cpp b/Engine/Source/Runtime/Renderer/Private/RayTracing/RayTracingInstanceBufferUtil.cpp
--- a/Engine/Source/Runtime/Renderer/Private/RayTracing/RayTracingInstanceBufferUtil.cpp
+++ b/Engine/Source/Runtime/Renderer/Private/RayTracing/RayTracingInstanceBufferUtil.cpp
@@ -106,7 +106,6 @@ FRayTracingSceneInitializationData BuildRayTracingSceneInitializationData(TConst
 	FRayTracingSceneInitializationData Output;
 	Output.NumNativeGPUSceneInstances = 0;
 	Output.NumNativeCPUInstances = 0;
-	Output.TotalNumSegments = 0;
 	Output.InstanceGeometryIndices.SetNumUninitialized(NumSceneInstances);
 	Output.BaseUploadBufferOffsets.SetNumUninitialized(NumSceneInstances);
 	Output.BaseInstancePrefixSum.SetNumUninitialized(NumSceneInstances);
@@ -137,8 +136,6 @@ FRayTracingSceneInitializationData BuildRayTracingSceneInitializationData(TConst
 
 		checkf(InstanceDesc.GeometryRHI, TEXT("Ray tracing instance must have a valid geometry."));
 
-		Output.TotalNumSegments += InstanceDesc.GeometryRHI->GetNumSegments();
-
 		uint32 GeometryIndex = UniqueGeometries.FindOrAdd(InstanceDesc.GeometryRHI, Output.ReferencedGeometries.Num());
 		Output.InstanceGeometryIndices[InstanceIndex] = GeometryIndex;
 		if (GeometryIndex == Output.ReferencedGeometries.Num())
diff --git a/Engine/Source/Runtime/Renderer/Private/RayTracing/RayTracingScene.cpp b/Engine/Source/Runtime/Renderer/Private/RayTracing/RayTracingScene.cpp
--- a/Engine/Source/Runtime/Renderer/Private/RayTracing/RayTracingScene.cpp
+++ b/Engine/Source/Runtime/Renderer/Private/RayTracing/RayTracingScene.cpp
@@ -137,7 +137,6 @@ void FRayTracingScene::Create(FRDGBuilder& GraphBuilder, const FViewInfo& View,
 			FRayTracingSceneInitializer Initializer;
 			Initializer.DebugName = FName(TEXT("FRayTracingScene"));
 			Initializer.MaxNumInstances = Layer.InitializationData.NumNativeGPUSceneInstances + Layer.InitializationData.NumNativeCPUInstances;
-			Initializer.NumTotalSegments = Layer.InitializationData.TotalNumSegments;
 			Initializer.BuildFlags = LayerIndex == uint8(ERayTracingSceneLayer::FarField) ? FarFieldBuildFlags : BuildFlags;
 			Layer.RayTracingSceneRHI = RHICreateRayTracingScene(MoveTemp(Initializer));
 		}
@@ -488,16 +487,6 @@ void FRayTracingScene::Create(FRDGBuilder& GraphBuilder, const FViewInfo& View,
 				});
 		}
 	}
-
-#if DO_CHECK
-	uint32 LayersTotalNumSegments = 0;
-	for (uint32 LayerIndex = 0; LayerIndex < NumLayers; ++LayerIndex)
-	{
-		LayersTotalNumSegments += Layers[LayerIndex].InitializationData.TotalNumSegments;
-	}
-	
-	checkf(LayersTotalNumSegments <= NumSegments, TEXT("Ray tracing scene layers use more segments than the number used to create SBTs"));
-#endif
 }
 
 BEGIN_SHADER_PARAMETER_STRUCT(FRayTracingSceneBuildPassParams, )
@@ -994,7 +983,6 @@ void FRayTracingScene::Reset(bool bInInstanceDebugDataEnabled)
 	UniformBuffers.Reset();
 	GeometriesToBuild.Reset();
 
-	NumSegments = 0;
 	NumMissShaderSlots = 1;
 	NumCallableShaderSlots = 0;
 
diff --git a/Engine/Source/Runtime/Renderer/Private/RayTracing/RayTracingScene.h b/Engine/Source/Runtime/Renderer/Private/RayTracing/RayTracingScene.h
--- a/Engine/Source/Runtime/Renderer/Private/RayTracing/RayTracingScene.h
+++ b/Engine/Source/Runtime/Renderer/Private/RayTracing/RayTracingScene.h
@@ -150,8 +150,6 @@ public:
 
 	FRayTracingGeometryInstance& GetInstance(FInstanceHandle Handle) { return Layers[uint8(Handle.Layer)].Instances[Handle.Index]; }
 
-	uint32 GetTotalNumSegments() const { return NumSegments; }
-
 	uint32 GetNumNativeInstances(ERayTracingSceneLayer Layer) const;
 
 	FRDGBufferRef GetInstanceDebugBuffer(ERayTracingSceneLayer Layer) const { return Layers[uint8(Layer)].InstanceDebugBuffer; }
@@ -163,7 +161,6 @@ public:
 
 	// Public members for initial refactoring step (previously were public members of FViewInfo).
 
-	uint32 NumSegments = 0;
 	uint32 NumMissShaderSlots = 1; // we must have a default miss shader, so always include it from the start
 	uint32 NumCallableShaderSlots = 0;
 	TArray<FRayTracingShaderCommand> CallableCommands;
diff --git a/Engine/Source/Runtime/Renderer/Private/Tests/RayTracingTestbed.cpp b/Engine/Source/Runtime/Renderer/Private/Tests/RayTracingTestbed.cpp
--- a/Engine/Source/Runtime/Renderer/Private/Tests/RayTracingTestbed.cpp
+++ b/Engine/Source/Runtime/Renderer/Private/Tests/RayTracingTestbed.cpp
@@ -183,7 +183,6 @@ bool RunRayTracingTestbed_RenderThread(const FString& Parameters)
 		FRayTracingSceneInitializer Initializer;
 		Initializer.DebugName = FName(TEXT("FRayTracingScene"));
 		Initializer.MaxNumInstances = RayTracingSceneInitializationData.NumNativeGPUSceneInstances + RayTracingSceneInitializationData.NumNativeCPUInstances;
-		Initializer.NumTotalSegments = RayTracingSceneInitializationData.TotalNumSegments;
 		Initializer.BuildFlags = SceneBuildFlags;
 
 		RayTracingSceneRHI = RHICreateRayTracingScene(MoveTemp(Initializer));
diff --git a/Engine/Source/Runtime/Renderer/Public/RayTracingInstanceBufferUtil.h b/Engine/Source/Runtime/Renderer/Public/RayTracingInstanceBufferUtil.h
--- a/Engine/Source/Runtime/Renderer/Public/RayTracingInstanceBufferUtil.h
+++ b/Engine/Source/Runtime/Renderer/Public/RayTracingInstanceBufferUtil.h
@@ -59,6 +59,8 @@ struct FRayTracingSceneInitializationData
 	uint32 NumCPUInstanceDescriptors;
 	uint32 NumNativeGPUSceneInstances;
 	uint32 NumNativeCPUInstances;
+
+	UE_DEPRECATED(5.6, "No longer used. Use FShaderBindingTable instead.")
 	uint32 TotalNumSegments;
 
 	// index of each instance geometry in ReferencedGeometries
