[Lit Translucency] Use compute shader instead of VS+GS+PS to filter TLV.

diff --git a/Engine/Shaders/Private/TranslucentLightingShaders.usf b/Engine/Shaders/Private/TranslucentLightingShaders.usf
--- a/Engine/Shaders/Private/TranslucentLightingShaders.usf
+++ b/Engine/Shaders/Private/TranslucentLightingShaders.usf
@@ -82,25 +82,35 @@ void WriteToSliceMainGS(triangle FWriteToSliceVertexOutput Input[3], inout Trian
 	OutStream.Append(Vertex2);
 }
 
+#ifdef FilterTranslucentVolumeCS
+
+/** Filter pass ouputs. */
+RWTexture3D<float4> RWTranslucencyLightingVolumeAmbient;
+RWTexture3D<float4> RWTranslucencyLightingVolumeDirectional;
+
 /** Filter pass inputs. */
 Texture3D TranslucencyLightingVolumeAmbient;
 SamplerState TranslucencyLightingVolumeAmbientSampler;
 Texture3D TranslucencyLightingVolumeDirectional;
 SamplerState TranslucencyLightingVolumeDirectionalSampler;
 
+uint3 VolumeSize;
 float TexelSize;
 
 /** Filters the volume lighting to reduce aliasing. */
-void FilterMainPS(
-	FWriteToSliceGeometryOutput Input,
-	out float4 OutColor0 : SV_Target0,
-	out float4 OutColor1 : SV_Target1
-	)
+[numthreads(THREADGROUP_SIZE_X, THREADGROUP_SIZE_Y, THREADGROUP_SIZE_Z)]
+void FilterTranslucentVolumeCS(uint3 DispatchThreadId : SV_DispatchThreadID)
 {
+	uint3 VolumeCoord = DispatchThreadId.xyz;
+
+	if (any(VolumeCoord >= VolumeSize))
+	{
+		return;
+	}
 	float4 TextureValue0 = 0;
 	float4 TextureValue1 = 0;
 
-	float3 VolumeUV = float3(Input.Vertex.UV, (Input.LayerIndex + .5f) * TexelSize);
+	float3 VolumeUV = (VolumeCoord + 0.5f) * TexelSize;
 
 #define USE_FILTER 1
 
@@ -135,10 +145,12 @@ void FilterMainPS(
 
 #endif
 	
-	OutColor0 = TextureValue0 * InvWeight;
-	OutColor1 = TextureValue1 * InvWeight;
+	RWTranslucencyLightingVolumeAmbient[DispatchThreadId] = TextureValue0 * InvWeight;
+	RWTranslucencyLightingVolumeDirectional[DispatchThreadId] = TextureValue1 * InvWeight;
 }
 
+#endif // FilterTranslucentVolumeCS
+
 #include "CubemapCommon.ush"
 
 /** Add AmbientCubemap color to the volume. */
@@ -155,7 +167,7 @@ void InjectAmbientCubemapMainPS(
 #endif // FEATURE_LEVEL
 
 
-#ifdef DEBUG_TRANSLUCENCY_LIGHTING_VOLUME_CS
+#ifdef DebugTranslucencyLightingVolumeCS
 
 #include "/Engine/Private/ShaderPrint.ush"
 
@@ -222,5 +234,5 @@ void DebugTranslucencyLightingVolumeCS(
 
 	}
 }
-#endif // DEBUG_TRANSLUCENCY_LIGHTING_VOLUME_CS
+#endif // DebugTranslucencyLightingVolumeCS
 
diff --git a/Engine/Source/Runtime/Renderer/Private/TranslucentLighting.cpp b/Engine/Source/Runtime/Renderer/Private/TranslucentLighting.cpp
--- a/Engine/Source/Runtime/Renderer/Private/TranslucentLighting.cpp
+++ b/Engine/Source/Runtime/Renderer/Private/TranslucentLighting.cpp
@@ -694,30 +694,48 @@ void FProjectedShadowInfo::RenderTranslucencyDepths(FRDGBuilder& GraphBuilder, F
 	});
 }
 
-/** Pixel shader used to filter a single volume lighting cascade. */
-class FFilterTranslucentVolumePS : public FGlobalShader
+/** Compute shader used to filter a single volume lighting cascade. */
+class FFilterTranslucentVolumeCS : public FGlobalShader
 {
 public:
-	DECLARE_GLOBAL_SHADER(FFilterTranslucentVolumePS);
-	SHADER_USE_PARAMETER_STRUCT(FFilterTranslucentVolumePS, FGlobalShader);
+	DECLARE_GLOBAL_SHADER(FFilterTranslucentVolumeCS);
+	SHADER_USE_PARAMETER_STRUCT(FFilterTranslucentVolumeCS, FGlobalShader);
 
 	BEGIN_SHADER_PARAMETER_STRUCT(FParameters, )
 		SHADER_PARAMETER_STRUCT_REF(FViewUniformShaderParameters, View)
+		SHADER_PARAMETER_RDG_TEXTURE_UAV(RWTexture3D, RWTranslucencyLightingVolumeAmbient)
+		SHADER_PARAMETER_RDG_TEXTURE_UAV(RWTexture3D, RWTranslucencyLightingVolumeDirectional)
 		SHADER_PARAMETER_RDG_TEXTURE(Texture3D, TranslucencyLightingVolumeAmbient)
 		SHADER_PARAMETER_RDG_TEXTURE(Texture3D, TranslucencyLightingVolumeDirectional)
 		SHADER_PARAMETER_SAMPLER(SamplerState, TranslucencyLightingVolumeAmbientSampler)
 		SHADER_PARAMETER_SAMPLER(SamplerState, TranslucencyLightingVolumeDirectionalSampler)
+		SHADER_PARAMETER(FIntVector, VolumeSize)
 		SHADER_PARAMETER(float, TexelSize)
-		RENDER_TARGET_BINDING_SLOTS()
 	END_SHADER_PARAMETER_STRUCT()
 
 	static bool ShouldCompilePermutation(const FGlobalShaderPermutationParameters& Parameters) 
 	{ 
-		return IsFeatureLevelSupported(Parameters.Platform, ERHIFeatureLevel::SM5) && (RHISupportsGeometryShaders(Parameters.Platform) || RHISupportsVertexShaderLayer(Parameters.Platform));
+		return IsFeatureLevelSupported(Parameters.Platform, ERHIFeatureLevel::SM5);
+	}
+
+	static FIntVector GetGroupSize()
+	{
+		return FIntVector(4, 4, 4);
+	}
+
+	static void ModifyCompilationEnvironment(const FGlobalShaderPermutationParameters& Parameters, FShaderCompilerEnvironment& OutEnvironment)
+	{
+		FGlobalShader::ModifyCompilationEnvironment(Parameters, OutEnvironment);
+
+		FIntVector GroupSize = GetGroupSize();
+
+		OutEnvironment.SetDefine(TEXT("THREADGROUP_SIZE_X"), GroupSize.X);
+		OutEnvironment.SetDefine(TEXT("THREADGROUP_SIZE_Y"), GroupSize.Y);
+		OutEnvironment.SetDefine(TEXT("THREADGROUP_SIZE_Z"), GroupSize.Z);
 	}
 };
 
-IMPLEMENT_GLOBAL_SHADER(FFilterTranslucentVolumePS, "/Engine/Private/TranslucentLightingShaders.usf", "FilterMainPS", SF_Pixel);
+IMPLEMENT_GLOBAL_SHADER(FFilterTranslucentVolumeCS, "/Engine/Private/TranslucentLightingShaders.usf", "FilterTranslucentVolumeCS", SF_Compute);
 
 /** Shader that adds direct lighting contribution from the given light to the current volume lighting cascade. */
 class FTranslucentLightingInjectPS : public FMaterialShader
@@ -1747,7 +1765,6 @@ class FDebugTranslucencyLightingVolumeCS : public FGlobalShader
 	{
 		FGlobalShader::ModifyCompilationEnvironment(Parameters, OutEnvironment);
 		OutEnvironment.SetDefine(TEXT("THREADGROUP_SIZE"), GetGroupSize());
-		OutEnvironment.SetDefine(TEXT("DEBUG_TRANSLUCENCY_LIGHTING_VOLUME_CS"), 1);
 	}
 
 	static EShaderPermutationPrecacheRequest ShouldPrecachePermutation(const FGlobalShaderPermutationParameters& Parameters)
@@ -1819,7 +1836,7 @@ void FilterTranslucencyLightingVolume(
 	{
 		const FViewInfo& View = Views[ViewIndex];
 		RDG_GPU_MASK_SCOPE(GraphBuilder, View.GPUMask);
-		RDG_EVENT_SCOPE_CONDITIONAL(GraphBuilder, Views.Num() > 1, "View%d", ViewIndex);
+		RDG_EVENT_SCOPE_CONDITIONAL(GraphBuilder, Views.Num() > 1, "View %d", ViewIndex);
 
 		for (int32 VolumeCascadeIndex = 0; VolumeCascadeIndex < TVC_MAX; VolumeCascadeIndex++)
 		{
@@ -1834,48 +1851,29 @@ void FilterTranslucencyLightingVolume(
 			Textures.Ambient[TextureIndex] = OutputVolumeAmbientTexture;
 			Textures.Directional[TextureIndex] = OutputVolumeDirectionalTexture;
 
-			auto* PassParameters = GraphBuilder.AllocParameters<FFilterTranslucentVolumePS::FParameters>();
+			FIntVector VolumeSize = FIntVector(TranslucencyLightingVolumeDim);
+
+			auto* PassParameters = GraphBuilder.AllocParameters<FFilterTranslucentVolumeCS::FParameters>();
 			PassParameters->View = View.ViewUniformBuffer;
+			PassParameters->VolumeSize = VolumeSize;
 			PassParameters->TexelSize = 1.0f / TranslucencyLightingVolumeDim;
 			PassParameters->TranslucencyLightingVolumeAmbient = InputVolumeAmbientTexture;
 			PassParameters->TranslucencyLightingVolumeDirectional = InputVolumeDirectionalTexture;
 			PassParameters->TranslucencyLightingVolumeAmbientSampler = SamplerStateRHI;
 			PassParameters->TranslucencyLightingVolumeDirectionalSampler = SamplerStateRHI;
-			PassParameters->RenderTargets[0] = FRenderTargetBinding(OutputVolumeAmbientTexture, ERenderTargetLoadAction::ENoAction);
-			PassParameters->RenderTargets[1] = FRenderTargetBinding(OutputVolumeDirectionalTexture, ERenderTargetLoadAction::ENoAction);
+			PassParameters->RWTranslucencyLightingVolumeAmbient = GraphBuilder.CreateUAV(OutputVolumeAmbientTexture);
+			PassParameters->RWTranslucencyLightingVolumeDirectional = GraphBuilder.CreateUAV(OutputVolumeDirectionalTexture);
 
-			const FVolumeBounds VolumeBounds(TranslucencyLightingVolumeDim);
-			TShaderMapRef<FWriteToSliceVS> VertexShader(View.ShaderMap);
-			TOptionalShaderMapRef<FWriteToSliceGS> GeometryShader(View.ShaderMap);
-			TShaderMapRef<FFilterTranslucentVolumePS> PixelShader(View.ShaderMap);
+			auto ComputeShader = View.ShaderMap->GetShader<FFilterTranslucentVolumeCS>();
 
-			GraphBuilder.AddPass(
-				RDG_EVENT_NAME("Cascade%d", VolumeCascadeIndex),
+			FIntVector NumGroups = FComputeShaderUtils::GetGroupCount(VolumeSize, FFilterTranslucentVolumeCS::GetGroupSize());
+
+			FComputeShaderUtils::AddPass(
+				GraphBuilder,
+				RDG_EVENT_NAME("Cascade %d", VolumeCascadeIndex),
+				ComputeShader,
 				PassParameters,
-				ERDGPassFlags::Raster,
-				[VertexShader, GeometryShader, PixelShader, PassParameters, VolumeBounds, TranslucencyLightingVolumeDim](FRDGAsyncTask, FRHICommandList& RHICmdList)
-			{
-				FGraphicsPipelineStateInitializer GraphicsPSOInit;
-				RHICmdList.ApplyCachedRenderTargets(GraphicsPSOInit);
-
-				GraphicsPSOInit.RasterizerState = TStaticRasterizerState<FM_Solid, CM_None>::GetRHI();
-				GraphicsPSOInit.DepthStencilState = TStaticDepthStencilState<false, CF_Always>::GetRHI();
-				GraphicsPSOInit.BlendState = TStaticBlendState<>::GetRHI();
-				GraphicsPSOInit.BoundShaderState.VertexDeclarationRHI = GScreenVertexDeclaration.VertexDeclarationRHI;
-				GraphicsPSOInit.BoundShaderState.VertexShaderRHI = VertexShader.GetVertexShader();
-				GraphicsPSOInit.BoundShaderState.SetGeometryShader(GeometryShader.GetGeometryShader());
-				GraphicsPSOInit.BoundShaderState.PixelShaderRHI = PixelShader.GetPixelShader();
-				GraphicsPSOInit.PrimitiveType = PT_TriangleStrip;
-				SetGraphicsPipelineState(RHICmdList, GraphicsPSOInit, 0);
-
-				SetShaderParametersLegacyVS(RHICmdList, VertexShader, VolumeBounds, FIntVector(TranslucencyLightingVolumeDim));
-				if (GeometryShader.IsValid())
-				{
-					SetShaderParametersLegacyGS(RHICmdList, GeometryShader, VolumeBounds.MinZ);
-				}
-				SetShaderParameters(RHICmdList, PixelShader, PixelShader.GetPixelShader(), *PassParameters);
-				RasterizeToVolumeTexture(RHICmdList, VolumeBounds);
-			});
+				NumGroups);
 		}
 	}
 }
