[HWRT] Fix Lumen / debug visualizations using incorrect shaders when Nanite Ray Tracing is enabled and platform uses global shader bindings.

diff --git a/Engine/Source/Runtime/Renderer/Private/RayTracing/RayTracingInstanceMask.cpp b/Engine/Source/Runtime/Renderer/Private/RayTracing/RayTracingInstanceMask.cpp
--- a/Engine/Source/Runtime/Renderer/Private/RayTracing/RayTracingInstanceMask.cpp
+++ b/Engine/Source/Runtime/Renderer/Private/RayTracing/RayTracingInstanceMask.cpp
@@ -296,6 +296,8 @@ void SetupRayTracingMeshCommandMaskAndStatus(FRayTracingMeshCommand& MeshCommand
 	// In the case that that this mesh command is not associated with a primitive, the mesh batch value will still apply.
 	MeshCommand.bReverseCulling = PrimitiveSceneProxy->IsCullingReversedByComponent();
 
+	MeshCommand.bNaniteRayTracing = (Nanite::GetRayTracingMode() != Nanite::ERayTracingMode::Fallback) && PrimitiveSceneProxy->IsNaniteMesh();
+
 	FSceneProxyRayTracingMaskInfo MaskInfo = GetSceneProxyRayTracingMaskInfo(*PrimitiveSceneProxy);
 
 	MeshCommand.InstanceMask = ApplyFirstPersonRayTracingInstanceMaskFlag(MeshCommand.InstanceMask, RayTracingType, MaskInfo.bIsFirstPersonWorldSpaceRepresentation);
diff --git a/Engine/Source/Runtime/Renderer/Private/RayTracing/RayTracingMeshDrawCommands.cpp b/Engine/Source/Runtime/Renderer/Private/RayTracing/RayTracingMeshDrawCommands.cpp
--- a/Engine/Source/Runtime/Renderer/Private/RayTracing/RayTracingMeshDrawCommands.cpp
+++ b/Engine/Source/Runtime/Renderer/Private/RayTracing/RayTracingMeshDrawCommands.cpp
@@ -128,6 +128,11 @@ void FRayTracingMeshCommand::SetShader(const TShaderRef<FShader>& Shader)
 	SceneUniformBufferParameter = Shader->GetUniformBufferParameter<FSceneUniformParameters>();
 	NaniteUniformBufferParameter = Shader->GetUniformBufferParameter<FNaniteRayTracingUniformParameters>();
 	ShaderBindings.Initialize(Shader);
+
+	if (NaniteUniformBufferParameter.IsBound())
+	{
+		check(bNaniteRayTracing);
+	}
 }
 
 void FRayTracingMeshCommand::SetShaders(const FMeshProcessorShaders& Shaders)
@@ -137,7 +142,7 @@ void FRayTracingMeshCommand::SetShaders(const FMeshProcessorShaders& Shaders)
 
 bool FRayTracingMeshCommand::IsUsingNaniteRayTracing() const
 {
-	return NaniteUniformBufferParameter.IsBound();
+	return bNaniteRayTracing;
 }
 
 void FRayTracingMeshCommand::UpdateFlags(FRayTracingCachedMeshCommandFlags& Flags) const
diff --git a/Engine/Source/Runtime/Renderer/Public/RayTracingMeshDrawCommands.h b/Engine/Source/Runtime/Renderer/Public/RayTracingMeshDrawCommands.h
--- a/Engine/Source/Runtime/Renderer/Public/RayTracingMeshDrawCommands.h
+++ b/Engine/Source/Runtime/Renderer/Public/RayTracingMeshDrawCommands.h
@@ -82,6 +82,7 @@ public:
 	bool bIsTranslucent : 1 = false;
 	bool bTwoSided : 1 = false;
 	bool bReverseCulling : 1 = false;
+	bool bNaniteRayTracing : 1 = false;
 	bool bCanBeCached : 1 = false;
 
 	RENDERER_API void SetRayTracingShaderBindingsForHitGroup(
