[Lit Translucency] Improvements to Translucency Lighting Volume
- support temporal filtering (r.TranslucencyLightingVolume.Temporal)
- support per-pixel position offset for sampling (r.TranslucencyLightingVolume.PositionOffsetRadius)
- fix position offset not being adjusted for specific cascade size.

diff --git a/Engine/Plugins/FX/Niagara/Shaders/Private/NiagaraSpriteVertexFactory.ush b/Engine/Plugins/FX/Niagara/Shaders/Private/NiagaraSpriteVertexFactory.ush
--- a/Engine/Plugins/FX/Niagara/Shaders/Private/NiagaraSpriteVertexFactory.ush
+++ b/Engine/Plugins/FX/Niagara/Shaders/Private/NiagaraSpriteVertexFactory.ush
@@ -1083,7 +1083,7 @@ FVertexFactoryIntermediates GetVertexFactoryIntermediates(FVertexFactoryInput In
 #if USE_PARTICLE_LIGHTING_OFFSET && USE_PARTICLE_RANDOM
 	// Hash function based on the particle ID to generate a uniformly distributed 3d offset
 	float3 RandomParticleOffset = frac(Square(Intermediates.ParticleRandom * 47 + 10) * float3(1361.456345, 2333.578, 3623.983)) * 2 - 1;
-	Intermediates.LightingPositionOffset = .5f * View.TranslucencyLightingVolumeInvSize[0].w * RandomParticleOffset;
+	Intermediates.LightingPositionOffset = .5f * RandomParticleOffset;
 #endif
 
 #if USE_PARTICLE_SIZE
diff --git a/Engine/Shaders/Private/BasePassCommon.ush b/Engine/Shaders/Private/BasePassCommon.ush
--- a/Engine/Shaders/Private/BasePassCommon.ush
+++ b/Engine/Shaders/Private/BasePassCommon.ush
@@ -120,6 +120,8 @@ struct FSharedBasePassInterpolants
 #define TranslucencyLightingVolumeDirectionalInner TranslucentBasePass.TranslucencyLightingVolumeDirectionalInner
 #define TranslucencyLightingVolumeDirectionalOuter TranslucentBasePass.TranslucencyLightingVolumeDirectionalOuter
 
+#define TranslucencyLightingRandomPositionOffsetRadius TranslucentBasePass.TranslucencyLightingRandomPositionOffsetRadius
+
 #include "TranslucencyVolumeCommon.ush"
 
 #include "SphericalGaussian.ush"
diff --git a/Engine/Shaders/Private/BasePassPixelShader.usf b/Engine/Shaders/Private/BasePassPixelShader.usf
--- a/Engine/Shaders/Private/BasePassPixelShader.usf
+++ b/Engine/Shaders/Private/BasePassPixelShader.usf
@@ -222,6 +222,8 @@
 #include "LocalFogVolumes/LocalFogVolumeCommon.ush"
 #endif
 
+#include "BlueNoise.ush"
+
 #if !FORWARD_SHADING
 
 bool ManualDepthTestEqual(float4 SvPosition, float PSDeviceZWithOffset)
@@ -248,11 +250,21 @@ float3 GetTranslucencyVolumeLighting(
 {
 	float4 VolumeLighting;
 	float3 InterpolatedLighting = 0;
+	
+	float3 LightingPositionOffset = MaterialParameters.LightingPositionOffset;
+
+	if (TranslucencyLightingRandomPositionOffsetRadius > 0.0f)
+	{
+		// override offset set by VF
+		// TODO: evaluate different noise functions, maybe use actual 3D Spatiotemporal Blue Noise?
+		const float3 Random = float3(BlueNoiseVec2(MaterialParameters.SvPosition.xy, ResolvedView.StateFrameIndex), BlueNoiseScalar(MaterialParameters.SvPosition.xy, ResolvedView.StateFrameIndex));
+		LightingPositionOffset = (Random - 0.5f) * TranslucencyLightingRandomPositionOffsetRadius;
+	}
 
 	float3 InnerVolumeUVs;
 	float3 OuterVolumeUVs;
 	float FinalLerpFactor;
-	ComputeVolumeUVs(MaterialParameters.WorldPosition_CamRelative, MaterialParameters.LightingPositionOffset, InnerVolumeUVs, OuterVolumeUVs, FinalLerpFactor);
+	ComputeVolumeUVs(MaterialParameters.WorldPosition_CamRelative, LightingPositionOffset, InnerVolumeUVs, OuterVolumeUVs, FinalLerpFactor);
 
 	#if TRANSLUCENCY_LIGHTING_VOLUMETRIC_PERVERTEX_DIRECTIONAL
 
@@ -282,7 +294,7 @@ float3 GetTranslucencyVolumeLighting(
 		{
 			// Determine the shadow space position
 			// Apply a stable offset to the world position used for shadowing, which blurs out high frequency details in the shadowmap with many layers
-			float4 HomogeneousShadowPosition = mul(float4(WSHackToFloat(GetWorldPosition(MaterialParameters)) + MaterialParameters.LightingPositionOffset, 1), TranslucentSelfShadow.WorldToShadowMatrix);
+			float4 HomogeneousShadowPosition = mul(float4(WSHackToFloat(GetWorldPosition(MaterialParameters)) + MaterialParameters.LightingPositionOffset * ResolvedView.TranslucencyLightingVolumeInvSize[0].w, 1), TranslucentSelfShadow.WorldToShadowMatrix);
 			float2 ShadowUVs = HomogeneousShadowPosition.xy / HomogeneousShadowPosition.w;
 			float ShadowZ = 1 - HomogeneousShadowPosition.z;
 			// Lookup the shadow density at the point being shaded
diff --git a/Engine/Shaders/Private/ParticleGPUSpriteVertexFactory.ush b/Engine/Shaders/Private/ParticleGPUSpriteVertexFactory.ush
--- a/Engine/Shaders/Private/ParticleGPUSpriteVertexFactory.ush
+++ b/Engine/Shaders/Private/ParticleGPUSpriteVertexFactory.ush
@@ -364,7 +364,7 @@ FVertexFactoryIntermediates GetVertexFactoryIntermediates( FVertexFactoryInput I
 #if USE_PARTICLE_LIGHTING_OFFSET
 	// Hash function based on the particle ID to generate a uniformly distributed 3d offset
 	float3 RandomParticleOffset = frac((ParticleIndex.x + 10) * (ParticleIndex.y + 10) * float3(1341.456345, 2633.578, 5623.983)) * 2 - 1;
-	Intermediates.LightingPositionOffset = .5 * View.TranslucencyLightingVolumeInvSize[0].w * RandomParticleOffset;
+	Intermediates.LightingPositionOffset = .5 * RandomParticleOffset;
 #endif
 
 	Intermediates.TexCoord = TexCoord;
diff --git a/Engine/Shaders/Private/ParticleSpriteVertexFactory.ush b/Engine/Shaders/Private/ParticleSpriteVertexFactory.ush
--- a/Engine/Shaders/Private/ParticleSpriteVertexFactory.ush
+++ b/Engine/Shaders/Private/ParticleSpriteVertexFactory.ush
@@ -600,7 +600,7 @@ FVertexFactoryIntermediates GetVertexFactoryIntermediates(FVertexFactoryInput In
 #if USE_PARTICLE_LIGHTING_OFFSET
 	// Hash function based on the particle ID to generate a uniformly distributed 3d offset
 	float3 RandomParticleOffset = frac(Square(Input.OldPosition.w + 10) * float3(1361.456345, 2333.578, 3623.983)) * 2 - 1;
-	Intermediates.LightingPositionOffset = .5f * View.TranslucencyLightingVolumeInvSize[0].w * RandomParticleOffset;
+	Intermediates.LightingPositionOffset = .5f * RandomParticleOffset;
 #endif
 
 #if USE_PARTICLE_SIZE
diff --git a/Engine/Shaders/Private/Substrate/SubstrateForwardLighting.ush b/Engine/Shaders/Private/Substrate/SubstrateForwardLighting.ush
--- a/Engine/Shaders/Private/Substrate/SubstrateForwardLighting.ush
+++ b/Engine/Shaders/Private/Substrate/SubstrateForwardLighting.ush
@@ -367,7 +367,7 @@ float3 SubstrateForwardLighting(
 				{
 					// Determine the shadow space position
 					// Apply a stable offset to the world position used for shadowing, which blurs out high frequency details in the shadowmap with many layers
-					float4 HomogeneousShadowPosition = mul(float4(AbsoluteWorldPosition + LightingPositionOffset, 1), TranslucentSelfShadow.WorldToShadowMatrix);
+					float4 HomogeneousShadowPosition = mul(float4(AbsoluteWorldPosition + LightingPositionOffset * View.TranslucencyLightingVolumeInvSize[0].w, 1), TranslucentSelfShadow.WorldToShadowMatrix);
 					float2 ShadowUVs = HomogeneousShadowPosition.xy / HomogeneousShadowPosition.w;
 					float ShadowZ = 1 - HomogeneousShadowPosition.z;
 					// Lookup the shadow density at the point being shaded
diff --git a/Engine/Shaders/Private/TranslucencyVolumeCommon.ush b/Engine/Shaders/Private/TranslucencyVolumeCommon.ush
--- a/Engine/Shaders/Private/TranslucencyVolumeCommon.ush
+++ b/Engine/Shaders/Private/TranslucencyVolumeCommon.ush
@@ -9,9 +9,9 @@
 
 void ComputeVolumeUVs(float3 TranslatedWorldPosition, float3 LightingPositionOffset, out float3 InnerVolumeUVs, out float3 OuterVolumeUVs, out float FinalLerpFactor)
 {
-	// Apply a stable offset to the world position used for lighting, which breaks up artifacts from using a low res volume texture
-	InnerVolumeUVs = (TranslatedWorldPosition + LightingPositionOffset - View.TranslucencyLightingVolumeMin[0].xyz) * View.TranslucencyLightingVolumeInvSize[0].xyz;
-	OuterVolumeUVs = (TranslatedWorldPosition + LightingPositionOffset - View.TranslucencyLightingVolumeMin[1].xyz) * View.TranslucencyLightingVolumeInvSize[1].xyz;
+	// Apply an offset to the world position used for lighting, which breaks up artifacts from using a low res volume texture
+	InnerVolumeUVs = (TranslatedWorldPosition + LightingPositionOffset * View.TranslucencyLightingVolumeInvSize[0].w - View.TranslucencyLightingVolumeMin[0].xyz) * View.TranslucencyLightingVolumeInvSize[0].xyz;
+	OuterVolumeUVs = (TranslatedWorldPosition + LightingPositionOffset * View.TranslucencyLightingVolumeInvSize[1].w - View.TranslucencyLightingVolumeMin[1].xyz) * View.TranslucencyLightingVolumeInvSize[1].xyz;
 
 	// Controls how fast the lerp between the inner and outer cascades happens
 	// Larger values result in a shorter transition distance
diff --git a/Engine/Shaders/Private/TranslucentLightingShaders.usf b/Engine/Shaders/Private/TranslucentLightingShaders.usf
--- a/Engine/Shaders/Private/TranslucentLightingShaders.usf
+++ b/Engine/Shaders/Private/TranslucentLightingShaders.usf
@@ -8,10 +8,7 @@
 #include "Common.ush"
 #include "SHCommon.ush"
 #include "SceneTexturesCommon.ush"
-
-#if INJECTION_PIXEL_SHADER
-	#include "ShadowProjectionCommon.ush"
-#endif
+#include "BlueNoise.ush"
 
 struct FWriteToSliceVertexOutput
 {
@@ -94,9 +91,41 @@ SamplerState TranslucencyLightingVolumeAmbientSampler;
 Texture3D TranslucencyLightingVolumeDirectional;
 SamplerState TranslucencyLightingVolumeDirectionalSampler;
 
+Texture3D HistoryAmbient;
+Texture3D HistoryDirectional;
+
+SamplerState HistoryAmbientSampler;
+SamplerState HistoryDirectionalSampler;
+
 uint3 VolumeSize;
 float TexelSize;
 
+float4 PrevTranslucencyLightingVolumeMin;
+float4 PrevTranslucencyLightingVolumeInvSize;
+
+float3 HistoryTextureBilinearUVMin;
+float3 HistoryTextureBilinearUVMax;
+
+float3 GetTranslatedWorldPositionFromTranslucencyLightingVolumeUVW(float3 UVW)
+{
+	return View.TranslucencyLightingVolumeMin[VolumeCascadeIndex].xyz + UVW / View.TranslucencyLightingVolumeInvSize[VolumeCascadeIndex].xyz;
+}
+
+float GetBorderLerpFactor(float3 UVW)
+{
+	// Larger values result in a shorter transition distance
+	float TransitionScale = 10;
+
+	// Rescale the UVs to make the fade go to 0 before the edge of the volume
+	float3 FadeUVs = UVW * (1 + 4 * View.TranslucencyLightingVolumeMin[VolumeCascadeIndex].w) - 2 * View.TranslucencyLightingVolumeMin[VolumeCascadeIndex].w;
+	// Setup a 3d lerp factor going to 0 near the edge of the outer volume
+	float3 LerpFactors = saturate((.5f - abs(FadeUVs - .5f)) * TransitionScale);
+	float FinalLerpFactor = LerpFactors.x * LerpFactors.y * LerpFactors.z;
+	
+	//FinalLerpFactor = 1.0f;
+	return FinalLerpFactor;
+}
+
 /** Filters the volume lighting to reduce aliasing. */
 [numthreads(THREADGROUP_SIZE_X, THREADGROUP_SIZE_Y, THREADGROUP_SIZE_Z)]
 void FilterTranslucentVolumeCS(uint3 DispatchThreadId : SV_DispatchThreadID)
@@ -116,6 +145,38 @@ void FilterTranslucentVolumeCS(uint3 DispatchThreadId : SV_DispatchThreadID)
 
 #if USE_FILTER
 
+#if USE_TEMPORAL_REPROJECTION
+
+	TextureValue0 = Texture3DSample(TranslucencyLightingVolumeAmbient, TranslucencyLightingVolumeAmbientSampler, VolumeUV);
+	TextureValue1 = Texture3DSample(TranslucencyLightingVolumeDirectional, TranslucencyLightingVolumeDirectionalSampler, VolumeUV);
+
+	float3 TranslatedWorldPosition = GetTranslatedWorldPositionFromTranslucencyLightingVolumeUVW(VolumeUV);
+	float3 HistoryUV = (TranslatedWorldPosition - PrevTranslucencyLightingVolumeMin.xyz) * PrevTranslucencyLightingVolumeInvSize.xyz;
+
+	float4 HistoryTextureValue0 = Texture3DSample(HistoryAmbient, HistoryAmbientSampler, HistoryUV);
+	float4 HistoryTextureValue1 = Texture3DSample(HistoryDirectional, HistoryDirectionalSampler, HistoryUV);
+
+	float HistoryAlpha = 0.9f;
+
+	if (any(HistoryUV < HistoryTextureBilinearUVMin) || any(HistoryUV > HistoryTextureBilinearUVMax))
+	{
+		HistoryAlpha = 0;
+	}
+
+	if (VolumeCascadeIndex == 1)
+	{
+		// fade out at border to prevent issues with usage of clamped sampler during shading passes
+		HistoryAlpha *= GetBorderLerpFactor(VolumeUV);
+	}
+
+
+	TextureValue0 = lerp(TextureValue0, HistoryTextureValue0, HistoryAlpha);
+	TextureValue1 = lerp(TextureValue1, HistoryTextureValue1, HistoryAlpha);
+
+	float InvWeight = 1;
+
+#else // !USE_TEMPORAL_REPROJECTION
+
 	// Use trilinear filtering to filter neighbors to the current voxel with minimal texture fetches
 	TextureValue0 += Texture3DSample(TranslucencyLightingVolumeAmbient, TranslucencyLightingVolumeAmbientSampler, VolumeUV + .5f * TexelSize * float3(1, 1, 1));
 	TextureValue0 += Texture3DSample(TranslucencyLightingVolumeAmbient, TranslucencyLightingVolumeAmbientSampler, VolumeUV + .5f * TexelSize * float3(-1, 1, 1));
@@ -137,6 +198,8 @@ void FilterTranslucentVolumeCS(uint3 DispatchThreadId : SV_DispatchThreadID)
 
 	float InvWeight = 1.0f / 8;
 
+#endif // !USE_TEMPORAL_REPROJECTION
+
 #else
 
 	TextureValue0 = Texture3DSample(TranslucencyLightingVolumeAmbient, TranslucencyLightingVolumeAmbientSampler, VolumeUV);
diff --git a/Engine/Source/Runtime/Renderer/Private/RendererScene.cpp b/Engine/Source/Runtime/Renderer/Private/RendererScene.cpp
--- a/Engine/Source/Runtime/Renderer/Private/RendererScene.cpp
+++ b/Engine/Source/Runtime/Renderer/Private/RendererScene.cpp
@@ -1000,6 +1000,19 @@ uint64 FStochasticLightingViewState::GetGPUSizeBytes(bool bLogSizes) const
 		GetRenderTargetGPUSizeBytes(SceneNormalHistory, bLogSizes);
 }
 
+uint64 FTranslucencyLightingViewState::GetGPUSizeBytes(bool bLogSizes) const
+{
+	uint64 TotalSize = 0;
+
+	for (int32 Index = 0; Index < TVC_MAX; Index++)
+	{
+		TotalSize += GetRenderTargetGPUSizeBytes(HistoryAmbient[Index], bLogSizes);
+		TotalSize += GetRenderTargetGPUSizeBytes(HistoryDirectional[Index], bLogSizes);
+	}
+
+	return TotalSize;
+}
+
 uint64 FPersistentGlobalDistanceFieldData::GetGPUSizeBytes(bool bLogSizes) const
 {
 	uint64 TotalSize =
@@ -1096,6 +1109,7 @@ uint64 FSceneViewState::GetGPUSizeBytes(bool bLogSizes) const
 	TotalSize += GetRenderTargetGPUSizeBytes(SubsurfaceScatteringQualityHistoryRT, bLogSizes);
 	TotalSize += Lumen.GetGPUSizeBytes(bLogSizes);
 	TotalSize += MegaLights.GetGPUSizeBytes(bLogSizes);
+	TotalSize += TranslucencyLighting.GetGPUSizeBytes(bLogSizes);
 	TotalSize += GetRenderTargetGPUSizeBytes(BloomFFTKernel.Spectral, bLogSizes);
 	TotalSize += GetBufferGPUSizeBytes(BloomFFTKernel.ConstantsBuffer, bLogSizes);
 	TotalSize += GetBufferGPUSizeBytes(FilmGrainCache.ConstantsBuffer, bLogSizes);
diff --git a/Engine/Source/Runtime/Renderer/Private/ScenePrivate.h b/Engine/Source/Runtime/Renderer/Private/ScenePrivate.h
--- a/Engine/Source/Runtime/Renderer/Private/ScenePrivate.h
+++ b/Engine/Source/Runtime/Renderer/Private/ScenePrivate.h
@@ -48,6 +48,7 @@
 #include "MegaLights/MegaLightsViewState.h"
 #include "StochasticLighting/StochasticLightingViewState.h"
 #include "VolumetricRenderTargetViewStateData.h"
+#include "TranslucentLightingViewState.h"
 #include "GPUScene.h"
 #include "DynamicBVH.h"
 #include "OIT/OIT.h"
@@ -1004,6 +1005,8 @@ public:
 	// Shared by Lumen and Mega Lights
 	FStochasticLightingViewState StochasticLighting;
 
+	FTranslucencyLightingViewState TranslucencyLighting;
+
 	// Heterogeneous Volumes cached data stores
 	TRDGUniformBufferRef<FOrthoVoxelGridUniformBufferParameters> OrthoVoxelGridUniformBuffer = nullptr;
 	TRDGUniformBufferRef<FFrustumVoxelGridUniformBufferParameters> FrustumVoxelGridUniformBuffer = nullptr;
diff --git a/Engine/Source/Runtime/Renderer/Private/SceneRendering.cpp b/Engine/Source/Runtime/Renderer/Private/SceneRendering.cpp
--- a/Engine/Source/Runtime/Renderer/Private/SceneRendering.cpp
+++ b/Engine/Source/Runtime/Renderer/Private/SceneRendering.cpp
@@ -102,6 +102,7 @@
 #include "MeshDrawCommandStats.h"
 #include "LocalFogVolumeRendering.h"
 #include "OIT/OIT.h"
+#include "TranslucentLighting.h"
 #include "Rendering/CustomRenderPass.h"
 #include "Stats/ThreadIdleStats.h"
 #include "CustomRenderPassSceneCapture.h"
@@ -2716,7 +2717,8 @@ FSceneRenderer::FSceneRenderer(const FSceneViewFamily* InViewFamily, FHitProxyCo
 			ShouldRenderVolumetricCloudWithBlueNoise_GameThread(Scene, *ViewInfo) ||
 			UseVirtualShadowMaps(Scene->GetShaderPlatform(), Scene->GetFeatureLevel()) ||
 			Substrate::IsGlintEnabled(ViewInfo->GetShaderPlatform()) ||
-			IsHairStrandsSupported(EHairStrandsShaderType::Strands, ViewInfo->GetShaderPlatform());
+			IsHairStrandsSupported(EHairStrandsShaderType::Strands, ViewInfo->GetShaderPlatform()) ||
+			IsTranslucencyLightingVolumeUsingBlueNoise();
 		GEngine->LoadBlueNoiseTexture(LoadVector2BlueNoiseTexture);
 
 		if (Substrate::IsGlintEnabled(ViewInfo->GetShaderPlatform()))
diff --git a/Engine/Source/Runtime/Renderer/Private/TranslucentLighting.cpp b/Engine/Source/Runtime/Renderer/Private/TranslucentLighting.cpp
--- a/Engine/Source/Runtime/Renderer/Private/TranslucentLighting.cpp
+++ b/Engine/Source/Runtime/Renderer/Private/TranslucentLighting.cpp
@@ -49,6 +49,7 @@
 #include "MeshPassProcessor.inl"
 #include "SkyAtmosphereRendering.h"
 #include "VolumetricCloudRendering.h"
+#include "TranslucentLightingViewState.h"
 #include "RenderCore.h"
 #include "StaticMeshBatch.h"
 #include "LightFunctionAtlas.h"
@@ -56,6 +57,7 @@
 #include "DeferredShadingRenderer.h"
 #include "BasePassRendering.h"
 #include "LightFunctionAtlas.h"
+#include "BlueNoise.h"
 
 using namespace LightFunctionAtlas;
 
@@ -120,6 +122,19 @@ static TAutoConsoleVariable<float> CVarTranslucencyLightingVolumeOuterDistance(
 	TEXT("Distance from the camera that the second volume cascade should end"),
 	ECVF_RenderThreadSafe | ECVF_Scalability);
 
+static TAutoConsoleVariable<float> CVarTranslucencyLightingVolumePositionOffsetRadius(
+	TEXT("r.TranslucencyLightingVolume.PositionOffsetRadius"),
+	0.0f,
+	TEXT("Radius of per-pixel offset applied to position when sampling translucency lighting volume."),
+	ECVF_RenderThreadSafe);
+
+static TAutoConsoleVariable<bool> CVarTranslucencyLightingVolumeTemporal(
+	TEXT("r.TranslucencyLightingVolume.Temporal"),
+	false,
+	TEXT("Whether to use temporal accumulation instead of spatial filter when updating the translucency lighting volume."),
+	ECVF_Scalability | ECVF_RenderThreadSafe
+);
+
 static TAutoConsoleVariable<int32> CVarTranslucencyLightingVolumeBatch(
 	TEXT("r.TranslucencyLightingVolume.Batch"),
 	1,
@@ -166,10 +181,14 @@ static FAutoConsoleVariableRef CVarGTranslucencyLightingVolumeMaterialPSOPrecach
 
 static const TCHAR* TranslucentLightingMaterialPSOCollectorName = TEXT("TranslucentLightingMaterialPSOCollector");
 
+float GetTranslucencyLightingVolumePositionOffsetRadius()
+{
+	return FMath::Max(0.0f, CVarTranslucencyLightingVolumePositionOffsetRadius.GetValueOnRenderThread());
+}
+
 /** Function returning current translucency lighting volume dimensions. */
 int32 GetTranslucencyLightingVolumeDim()
 {
-	extern int32 GTranslucencyLightingVolumeDim;
 	return FMath::Clamp(GTranslucencyLightingVolumeDim, 4, 2048);
 }
 
@@ -709,10 +728,27 @@ public:
 		SHADER_PARAMETER_RDG_TEXTURE(Texture3D, TranslucencyLightingVolumeDirectional)
 		SHADER_PARAMETER_SAMPLER(SamplerState, TranslucencyLightingVolumeAmbientSampler)
 		SHADER_PARAMETER_SAMPLER(SamplerState, TranslucencyLightingVolumeDirectionalSampler)
+
+		SHADER_PARAMETER_RDG_TEXTURE(Texture3D, HistoryAmbient)
+		SHADER_PARAMETER_RDG_TEXTURE(Texture3D, HistoryDirectional)
+		SHADER_PARAMETER_SAMPLER(SamplerState, HistoryAmbientSampler)
+		SHADER_PARAMETER_SAMPLER(SamplerState, HistoryDirectionalSampler)
+
 		SHADER_PARAMETER(FIntVector, VolumeSize)
 		SHADER_PARAMETER(float, TexelSize)
+		SHADER_PARAMETER(uint32, VolumeCascadeIndex)
+
+		SHADER_PARAMETER(FVector4f, PrevTranslucencyLightingVolumeMin)
+		SHADER_PARAMETER(FVector4f, PrevTranslucencyLightingVolumeInvSize)
+
+		SHADER_PARAMETER(FVector3f, HistoryTextureBilinearUVMin)
+		SHADER_PARAMETER(FVector3f, HistoryTextureBilinearUVMax)
+
 	END_SHADER_PARAMETER_STRUCT()
 
+	class FUseTemporalReprojection : SHADER_PERMUTATION_BOOL("USE_TEMPORAL_REPROJECTION");
+	using FPermutationDomain = TShaderPermutationDomain<FUseTemporalReprojection>;
+
 	static bool ShouldCompilePermutation(const FGlobalShaderPermutationParameters& Parameters) 
 	{ 
 		return IsFeatureLevelSupported(Parameters.Platform, ERHIFeatureLevel::SM5);
@@ -1042,9 +1078,17 @@ FTranslucencyLightingVolumeParameters GetTranslucencyLightingVolumeParameters(FR
 		Parameters.TranslucencyLightingVolumeDirectionalInner = SystemTextures.VolumetricBlack;
 		Parameters.TranslucencyLightingVolumeDirectionalOuter = SystemTextures.VolumetricBlack;
 	}
+
+	Parameters.TranslucencyLightingRandomPositionOffsetRadius = GetTranslucencyLightingVolumePositionOffsetRadius();
+
 	return Parameters;
 }
 
+bool IsTranslucencyLightingVolumeUsingBlueNoise()
+{
+	return GUseTranslucentLightingVolumes && GetTranslucencyLightingVolumePositionOffsetRadius();
+}
+
 class FInjectAmbientCubemapPS : public FGlobalShader
 {
 public:
@@ -1821,8 +1865,26 @@ void FilterTranslucencyLightingVolume(
 	// Draw debug information if needed.
 	DrawDebugTranslucencyLightingVolume(GraphBuilder, Views);
 
-	if (!GUseTranslucentLightingVolumes || Views.Num() == 0 || !RHISupportsVolumeTextureRendering(Views[0].GetShaderPlatform()) || !GUseTranslucencyVolumeBlur)
+	const bool bTemporal = CVarTranslucencyLightingVolumeTemporal.GetValueOnRenderThread();
+
+	if (!GUseTranslucentLightingVolumes || Views.Num() == 0 || !RHISupportsVolumeTextureRendering(Views[0].GetShaderPlatform()) || (!GUseTranslucencyVolumeBlur && !bTemporal))
 	{
+		// discard history since we are not updating it this frame
+
+		for (int32 ViewIndex = 0; ViewIndex < Views.Num(); ++ViewIndex)
+		{
+			const FViewInfo& View = Views[ViewIndex];
+			if (View.ViewState && !View.bStatePrevViewInfoIsReadOnly)
+			{
+				FTranslucencyLightingViewState& TranslucencyLightingViewState = View.ViewState->TranslucencyLighting;
+				for (int32 Index = 0; Index < TVC_MAX; Index++)
+				{
+					TranslucencyLightingViewState.HistoryAmbient[Index] = nullptr;
+					TranslucencyLightingViewState.HistoryDirectional[Index] = nullptr;
+				}
+			}
+		}
+
 		return;
 	}
 
@@ -1840,40 +1902,117 @@ void FilterTranslucencyLightingVolume(
 
 		for (int32 VolumeCascadeIndex = 0; VolumeCascadeIndex < TVC_MAX; VolumeCascadeIndex++)
 		{
+			bool bHasValidHistory = false;
+
+			FRDGTextureRef HistoryAmbient;
+			FRDGTextureRef HistoryDirectional;
+
+			FVector4f PrevTranslucencyLightingVolumeMin;
+			FVector4f PrevTranslucencyLightingVolumeInvSize;
+
+			if (bTemporal && View.ViewState && !View.bCameraCut && !View.bPrevTransformsReset)
+			{
+				FTranslucencyLightingViewState& TranslucencyLightingViewState = View.ViewState->TranslucencyLighting;
+
+				if (TranslucencyLightingViewState.HistoryAmbient[VolumeCascadeIndex] && TranslucencyLightingViewState.HistoryDirectional[VolumeCascadeIndex])
+				{
+					HistoryAmbient = GraphBuilder.RegisterExternalTexture(TranslucencyLightingViewState.HistoryAmbient[VolumeCascadeIndex]);
+					HistoryDirectional = GraphBuilder.RegisterExternalTexture(TranslucencyLightingViewState.HistoryDirectional[VolumeCascadeIndex]);
+
+					bHasValidHistory = true;
+				}
+				else
+				{
+					const FRDGSystemTextures& SystemTextures = FRDGSystemTextures::Get(GraphBuilder);
+					HistoryAmbient = SystemTextures.VolumetricBlack;
+					HistoryDirectional = SystemTextures.VolumetricBlack;
+				}
+
+				{
+					const FVector VolumeWorldMin = TranslucencyLightingViewState.HistoryVolumeMin[VolumeCascadeIndex];
+					const float VolumeVoxelSize = TranslucencyLightingViewState.HistoryVoxelSize[VolumeCascadeIndex];
+					const FVector3f VolumeSize = FVector3f(TranslucencyLightingViewState.HistoryVolumeSize[VolumeCascadeIndex]);
+
+					const FVector3f VolumeTranslatedWorldMin = FVector3f(VolumeWorldMin + View.ViewMatrices.GetPreViewTranslation());
+
+					PrevTranslucencyLightingVolumeMin = FVector4f(VolumeTranslatedWorldMin, 1.0f / TranslucencyLightingVolumeDim);
+					PrevTranslucencyLightingVolumeInvSize = FVector4f(FVector3f(1.0f) / VolumeSize, VolumeVoxelSize);
+				}
+			}
+			else
+			{
+				const FRDGSystemTextures& SystemTextures = FRDGSystemTextures::Get(GraphBuilder);
+				HistoryAmbient = SystemTextures.VolumetricBlack;
+				HistoryDirectional = SystemTextures.VolumetricBlack;
+
+				PrevTranslucencyLightingVolumeMin = FVector4f::Zero();
+				PrevTranslucencyLightingVolumeInvSize = FVector4f::Zero();
+			}
+
 			const uint32 TextureIndex = Textures.GetIndex(View, VolumeCascadeIndex);
 
-			FRDGTextureRef InputVolumeAmbientTexture = Textures.Ambient[TextureIndex];
-			FRDGTextureRef InputVolumeDirectionalTexture = Textures.Directional[TextureIndex];
+			if (!bTemporal || bHasValidHistory)
+			{
+				FRDGTextureRef InputVolumeAmbientTexture = Textures.Ambient[TextureIndex];
+				FRDGTextureRef InputVolumeDirectionalTexture = Textures.Directional[TextureIndex];
+
+				FRDGTextureRef OutputVolumeAmbientTexture = GraphBuilder.CreateTexture(InputVolumeAmbientTexture->Desc, InputVolumeAmbientTexture->Name);
+				FRDGTextureRef OutputVolumeDirectionalTexture = GraphBuilder.CreateTexture(InputVolumeDirectionalTexture->Desc, InputVolumeDirectionalTexture->Name);
+
+				Textures.Ambient[TextureIndex] = OutputVolumeAmbientTexture;
+				Textures.Directional[TextureIndex] = OutputVolumeDirectionalTexture;
+
+				const FIntVector VolumeSize = FIntVector(TranslucencyLightingVolumeDim);
+
+				auto* PassParameters = GraphBuilder.AllocParameters<FFilterTranslucentVolumeCS::FParameters>();
+				PassParameters->View = View.ViewUniformBuffer;
+				PassParameters->VolumeSize = VolumeSize;
+				PassParameters->TexelSize = 1.0f / TranslucencyLightingVolumeDim;
+				PassParameters->VolumeCascadeIndex = VolumeCascadeIndex;
+				PassParameters->TranslucencyLightingVolumeAmbient = InputVolumeAmbientTexture;
+				PassParameters->TranslucencyLightingVolumeDirectional = InputVolumeDirectionalTexture;
+				PassParameters->TranslucencyLightingVolumeAmbientSampler = SamplerStateRHI;
+				PassParameters->TranslucencyLightingVolumeDirectionalSampler = SamplerStateRHI;
+				PassParameters->RWTranslucencyLightingVolumeAmbient = GraphBuilder.CreateUAV(OutputVolumeAmbientTexture);
+				PassParameters->RWTranslucencyLightingVolumeDirectional = GraphBuilder.CreateUAV(OutputVolumeDirectionalTexture);
+
+				PassParameters->PrevTranslucencyLightingVolumeMin = PrevTranslucencyLightingVolumeMin;
+				PassParameters->PrevTranslucencyLightingVolumeInvSize = PrevTranslucencyLightingVolumeInvSize;
 
-			FRDGTextureRef OutputVolumeAmbientTexture = GraphBuilder.CreateTexture(InputVolumeAmbientTexture->Desc, InputVolumeAmbientTexture->Name);
-			FRDGTextureRef OutputVolumeDirectionalTexture = GraphBuilder.CreateTexture(InputVolumeDirectionalTexture->Desc, InputVolumeDirectionalTexture->Name);
+				PassParameters->HistoryTextureBilinearUVMin = FVector3f(0.5f / TranslucencyLightingVolumeDim);
+				PassParameters->HistoryTextureBilinearUVMax = FVector3f((TranslucencyLightingVolumeDim - 0.5f) / TranslucencyLightingVolumeDim);
 
-			Textures.Ambient[TextureIndex] = OutputVolumeAmbientTexture;
-			Textures.Directional[TextureIndex] = OutputVolumeDirectionalTexture;
+				PassParameters->HistoryAmbient = HistoryAmbient;
+				PassParameters->HistoryDirectional = HistoryDirectional;
+				PassParameters->HistoryAmbientSampler = SamplerStateRHI;
+				PassParameters->HistoryDirectionalSampler = SamplerStateRHI;
 
-			FIntVector VolumeSize = FIntVector(TranslucencyLightingVolumeDim);
+				FFilterTranslucentVolumeCS::FPermutationDomain PermutationVector;
+				PermutationVector.Set<FFilterTranslucentVolumeCS::FUseTemporalReprojection>(bTemporal);
 
-			auto* PassParameters = GraphBuilder.AllocParameters<FFilterTranslucentVolumeCS::FParameters>();
-			PassParameters->View = View.ViewUniformBuffer;
-			PassParameters->VolumeSize = VolumeSize;
-			PassParameters->TexelSize = 1.0f / TranslucencyLightingVolumeDim;
-			PassParameters->TranslucencyLightingVolumeAmbient = InputVolumeAmbientTexture;
-			PassParameters->TranslucencyLightingVolumeDirectional = InputVolumeDirectionalTexture;
-			PassParameters->TranslucencyLightingVolumeAmbientSampler = SamplerStateRHI;
-			PassParameters->TranslucencyLightingVolumeDirectionalSampler = SamplerStateRHI;
-			PassParameters->RWTranslucencyLightingVolumeAmbient = GraphBuilder.CreateUAV(OutputVolumeAmbientTexture);
-			PassParameters->RWTranslucencyLightingVolumeDirectional = GraphBuilder.CreateUAV(OutputVolumeDirectionalTexture);
+				auto ComputeShader = View.ShaderMap->GetShader<FFilterTranslucentVolumeCS>(PermutationVector);
 
-			auto ComputeShader = View.ShaderMap->GetShader<FFilterTranslucentVolumeCS>();
+				FIntVector NumGroups = FComputeShaderUtils::GetGroupCount(VolumeSize, FFilterTranslucentVolumeCS::GetGroupSize());
 
-			FIntVector NumGroups = FComputeShaderUtils::GetGroupCount(VolumeSize, FFilterTranslucentVolumeCS::GetGroupSize());
+				FComputeShaderUtils::AddPass(
+					GraphBuilder,
+					RDG_EVENT_NAME("Cascade %d", VolumeCascadeIndex),
+					ComputeShader,
+					PassParameters,
+					NumGroups);
+			}
 
-			FComputeShaderUtils::AddPass(
-				GraphBuilder,
-				RDG_EVENT_NAME("Cascade %d", VolumeCascadeIndex),
-				ComputeShader,
-				PassParameters,
-				NumGroups);
+			if (bTemporal && View.ViewState && !View.bStatePrevViewInfoIsReadOnly)
+			{
+				FTranslucencyLightingViewState& TranslucencyLightingViewState = View.ViewState->TranslucencyLighting;
+
+				GraphBuilder.QueueTextureExtraction(Textures.Ambient[TextureIndex], &TranslucencyLightingViewState.HistoryAmbient[VolumeCascadeIndex]);
+				GraphBuilder.QueueTextureExtraction(Textures.Directional[TextureIndex], &TranslucencyLightingViewState.HistoryDirectional[VolumeCascadeIndex]);
+
+				TranslucencyLightingViewState.HistoryVolumeMin[VolumeCascadeIndex] = View.TranslucencyLightingVolumeMin[VolumeCascadeIndex];
+				TranslucencyLightingViewState.HistoryVoxelSize[VolumeCascadeIndex] = View.TranslucencyVolumeVoxelSize[VolumeCascadeIndex];
+				TranslucencyLightingViewState.HistoryVolumeSize[VolumeCascadeIndex] = View.TranslucencyLightingVolumeSize[VolumeCascadeIndex];
+			}
 		}
 	}
 }
diff --git a/Engine/Source/Runtime/Renderer/Private/TranslucentLighting.h b/Engine/Source/Runtime/Renderer/Private/TranslucentLighting.h
--- a/Engine/Source/Runtime/Renderer/Private/TranslucentLighting.h
+++ b/Engine/Source/Runtime/Renderer/Private/TranslucentLighting.h
@@ -48,10 +48,11 @@ struct FTranslucencyLightingVolumeTextures
 };
 
 BEGIN_SHADER_PARAMETER_STRUCT(FTranslucencyLightingVolumeParameters, )
-SHADER_PARAMETER_RDG_TEXTURE(Texture3D, TranslucencyLightingVolumeAmbientInner)
-SHADER_PARAMETER_RDG_TEXTURE(Texture3D, TranslucencyLightingVolumeAmbientOuter)
-SHADER_PARAMETER_RDG_TEXTURE(Texture3D, TranslucencyLightingVolumeDirectionalInner)
-SHADER_PARAMETER_RDG_TEXTURE(Texture3D, TranslucencyLightingVolumeDirectionalOuter)
+	SHADER_PARAMETER_RDG_TEXTURE(Texture3D, TranslucencyLightingVolumeAmbientInner)
+	SHADER_PARAMETER_RDG_TEXTURE(Texture3D, TranslucencyLightingVolumeAmbientOuter)
+	SHADER_PARAMETER_RDG_TEXTURE(Texture3D, TranslucencyLightingVolumeDirectionalInner)
+	SHADER_PARAMETER_RDG_TEXTURE(Texture3D, TranslucencyLightingVolumeDirectionalOuter)
+	SHADER_PARAMETER(float, TranslucencyLightingRandomPositionOffsetRadius)
 END_SHADER_PARAMETER_STRUCT()
 
 /** Utility for batching together multiple lighting injections */
@@ -113,6 +114,8 @@ public:
 /** Initializes translucency volume lighting shader parameters from an optional textures struct. If null or uninitialized, fallback textures are used. */
 FTranslucencyLightingVolumeParameters GetTranslucencyLightingVolumeParameters(FRDGBuilder& GraphBuilder, const FTranslucencyLightingVolumeTextures& Textures, const FViewInfo& View);
 
+bool IsTranslucencyLightingVolumeUsingBlueNoise();
+
 void InjectTranslucencyLightingVolume(
 	FRDGBuilder& GraphBuilder,
 	const FViewInfo& View,
@@ -138,4 +141,4 @@ void InjectTranslucencyLightingVolumeAmbientCubemap(
 void FilterTranslucencyLightingVolume(
 	FRDGBuilder& GraphBuilder,
 	const TArrayView<const FViewInfo> Views,
-	FTranslucencyLightingVolumeTextures& Textures);
\ No newline at end of file
+	FTranslucencyLightingVolumeTextures& Textures);
diff --git a/Engine/Source/Runtime/Renderer/Private/TranslucentLightingViewState.h b/Engine/Source/Runtime/Renderer/Private/TranslucentLightingViewState.h
new file mode 100644
--- /dev/null
+++ b/Engine/Source/Runtime/Renderer/Private/TranslucentLightingViewState.h
@@ -0,0 +1,30 @@
+// Copyright Epic Games, Inc. All Rights Reserved.
+
+#pragma once
+
+#include "SceneView.h"
+#include "Math/Vector4.h"
+#include "RendererInterface.h"
+#include "RenderGraphResources.h"
+
+class FTranslucencyLightingViewState
+{
+public:
+	TStaticArray<TRefCountPtr<IPooledRenderTarget>, TVC_MAX> HistoryAmbient;
+	TStaticArray<TRefCountPtr<IPooledRenderTarget>, TVC_MAX> HistoryDirectional;
+
+	TStaticArray<FVector, TVC_MAX> HistoryVolumeMin;
+	TStaticArray<float, TVC_MAX> HistoryVoxelSize;
+	TStaticArray<FVector, TVC_MAX> HistoryVolumeSize;
+
+	void SafeRelease()
+	{
+		for (int32 Index = 0; Index < TVC_MAX; Index++)
+		{
+			HistoryAmbient[Index].SafeRelease();
+			HistoryDirectional[Index].SafeRelease();
+		}
+	}
+
+	uint64 GetGPUSizeBytes(bool bLogSizes) const;
+};
\ No newline at end of file
diff --git a/Engine/Source/Runtime/Renderer/Private/TranslucentRendering.cpp b/Engine/Source/Runtime/Renderer/Private/TranslucentRendering.cpp
--- a/Engine/Source/Runtime/Renderer/Private/TranslucentRendering.cpp
+++ b/Engine/Source/Runtime/Renderer/Private/TranslucentRendering.cpp
@@ -1154,7 +1154,7 @@ TRDGUniformBufferRef<FTranslucentBasePassUniformParameters> CreateTranslucentBas
 	OIT::SetOITParameters(GraphBuilder, View, BasePassParameters.OIT, OITData);
 
 	// Only use blue noise resources if VSM quality is set to high
-	if (IsVSMTranslucentHighQualityEnabled())
+	if (IsVSMTranslucentHighQualityEnabled() || IsTranslucencyLightingVolumeUsingBlueNoise())
 	{
 		BasePassParameters.BlueNoise = GetBlueNoiseParameters();
 	}
