Always rebind the target buffer because we have to make sure that the bindless index is always written otherwise it might miss in the cbuffer data

diff --git a/Engine/Source/Runtime/Renderer/Private/RayTracing/RayTracingDynamicGeometryUpdateManager.cpp b/Engine/Source/Runtime/Renderer/Private/RayTracing/RayTracingDynamicGeometryUpdateManager.cpp
--- a/Engine/Source/Runtime/Renderer/Private/RayTracing/RayTracingDynamicGeometryUpdateManager.cpp
+++ b/Engine/Source/Runtime/Renderer/Private/RayTracing/RayTracingDynamicGeometryUpdateManager.cpp
@@ -864,7 +864,10 @@ void FRayTracingDynamicGeometryUpdateManager::DispatchUpdates(FRHICommandList& R
 				FRHIBatchedShaderParameters& BatchedParameters = RHICmdList.GetScratchShaderParameters();
 
 				FRWBuffer* TargetBuffer = Cmd.TargetBuffer;
-				if (CurrentBuffer != TargetBuffer)
+
+				// Always rebind the target buffer because we have to make sure that the bindless index is always written
+				// otherwise it might miss in the cbuffer data
+				//if (CurrentBuffer != TargetBuffer)
 				{
 					CurrentBuffer = TargetBuffer;
 
