Remove unrequired casts because interface has been removed

diff --git a/Engine/Source/Runtime/Renderer/Private/RayTracing/RayTracing.cpp b/Engine/Source/Runtime/Renderer/Private/RayTracing/RayTracing.cpp
--- a/Engine/Source/Runtime/Renderer/Private/RayTracing/RayTracing.cpp
+++ b/Engine/Source/Runtime/Renderer/Private/RayTracing/RayTracing.cpp
@@ -630,7 +630,7 @@ namespace RayTracing
 
 		// TODO: Could process RayTracingGeometriesToUpdate in parallel thread after merging multiple tasks
 
-		FRayTracingDynamicGeometryUpdateManager* DynamicGeometryUpdateManager = (FRayTracingDynamicGeometryUpdateManager*)Scene.GetRayTracingDynamicGeometryUpdateManager();
+		FRayTracingDynamicGeometryUpdateManager* DynamicGeometryUpdateManager = Scene.GetRayTracingDynamicGeometryUpdateManager();
 
 		// Can't use RayTracingGeometriesToUpdate directly because need SceneProxy and PersistentPrimitiveIndex
 		// TODO: Move those parameters into FRayTracingDynamicGeometryUpdateParams
@@ -1809,7 +1809,7 @@ namespace RayTracing
 
 		INC_DWORD_STAT_BY(STAT_VisibleRayTracingPrimitives, TaskData.StaticPrimitives.Num() + TaskData.CachedStaticPrimitives.Num() + TaskData.DynamicPrimitivesIndices.Num());
 
-		FRayTracingDynamicGeometryUpdateManager* DynamicGeometryUpdateManager = (FRayTracingDynamicGeometryUpdateManager*)Scene.GetRayTracingDynamicGeometryUpdateManager();
+		FRayTracingDynamicGeometryUpdateManager* DynamicGeometryUpdateManager = Scene.GetRayTracingDynamicGeometryUpdateManager();
 		const int64 SharedBufferGenerationID = DynamicGeometryUpdateManager->BeginUpdate();
 
 		TaskData.GatherDynamicRayTracingInstancesTask.Wait();
