Fixed bug with removal of GPU Skin Cache task from GDME. While the skin cache is not a direct prerequisite, the mesh updater tasks are and were prerequisites of the skin cache work. Added those prerequisites explicitly.

diff --git a/Engine/Source/Runtime/Renderer/Private/RendererScene.cpp b/Engine/Source/Runtime/Renderer/Private/RendererScene.cpp
--- a/Engine/Source/Runtime/Renderer/Private/RendererScene.cpp
+++ b/Engine/Source/Runtime/Renderer/Private/RendererScene.cpp
@@ -5611,6 +5611,7 @@ void FScene::Update(FRDGBuilder& GraphBuilder, const FUpdateParameters& Paramete
 	if (SkeletalMeshUpdater)
 	{
 		GPUSkinTasks = SkeletalMeshUpdater->Submit(GraphBuilder);
+		GPUSkinUpdateTask = GPUSkinTasks.SkinCache;
 	}
 
 	if (GPUSkinCache && (GPUSkinTasks.SkinCache.IsValid() || GPUSkinCache->HasWork()))
diff --git a/Engine/Source/Runtime/Renderer/Private/ScenePrivate.h b/Engine/Source/Runtime/Renderer/Private/ScenePrivate.h
--- a/Engine/Source/Runtime/Renderer/Private/ScenePrivate.h
+++ b/Engine/Source/Runtime/Renderer/Private/ScenePrivate.h
@@ -3765,12 +3765,17 @@ public:
 		CreateLightPrimitiveInteractionsTask.Wait();
 	}
 
-	UE::Tasks::FTask GetCreateLightPrimitiveInteractionsTask() const
+	const UE::Tasks::FTask& GetCreateLightPrimitiveInteractionsTask() const
 	{
 		return CreateLightPrimitiveInteractionsTask;
 	}
 
-	UE::Tasks::FTask GetGPUSkinCacheTask() const
+	const UE::Tasks::FTask& GetGPUSkinUpdateTask() const
+	{
+		return GPUSkinUpdateTask;
+	}
+
+	const UE::Tasks::FTask& GetGPUSkinCacheTask() const
 	{
 		return GPUSkinCacheTask;
 	}
@@ -3783,7 +3788,7 @@ public:
 		CacheMeshDrawCommandsTask.Wait();
 	}
 
-	UE::Tasks::FTask GetCacheMeshDrawCommandsTask() const
+	const UE::Tasks::FTask& GetCacheMeshDrawCommandsTask() const
 	{
 		return CacheMeshDrawCommandsTask;
 	}
@@ -3794,7 +3799,7 @@ public:
 		CacheNaniteMaterialBinsTask.Wait();
 	}
 
-	UE::Tasks::FTask GetCacheNaniteMaterialBinsTask() const
+	const UE::Tasks::FTask& GetCacheNaniteMaterialBinsTask() const
 	{
 		return CacheNaniteMaterialBinsTask;
 	}
@@ -3806,7 +3811,7 @@ public:
 		CacheRayTracingPrimitivesTask.Wait();
 	}
 
-	UE::Tasks::FTask GetCacheRayTracingPrimitivesTask()
+	const UE::Tasks::FTask& GetCacheRayTracingPrimitivesTask()
 	{
 		return CacheRayTracingPrimitivesTask;
 	}
@@ -4010,6 +4015,7 @@ private:
 	TArray<FLevelCommand> LevelCommands;
 
 	UE::Tasks::FTask CreateLightPrimitiveInteractionsTask;
+	UE::Tasks::FTask GPUSkinUpdateTask;
 	UE::Tasks::FTask GPUSkinCacheTask;
 	UE::Tasks::FTask CacheMeshDrawCommandsTask;
 	UE::Tasks::FTask CacheNaniteMaterialBinsTask;
diff --git a/Engine/Source/Runtime/Renderer/Private/SceneVisibility.cpp b/Engine/Source/Runtime/Renderer/Private/SceneVisibility.cpp
--- a/Engine/Source/Runtime/Renderer/Private/SceneVisibility.cpp
+++ b/Engine/Source/Runtime/Renderer/Private/SceneVisibility.cpp
@@ -4268,6 +4268,9 @@ void FVisibilityTaskData::LaunchVisibilityTasks(const UE::Tasks::FTask& BeginIni
 	// Each relevance task should have this as a prerequisite, but in case there aren't any tasks we make it explicit.
 	Tasks.ComputeRelevance.AddPrerequisites(Scene.GetCacheMeshDrawCommandsTask());
 
+	// Wait on the GPU skin update task prior to GDME.
+	Tasks.DynamicMeshElementsPrerequisites.AddPrerequisites(Scene.GetGPUSkinUpdateTask());
+
 	bool bAllocatePrimitiveViewMasks = true;
 
 	if (TaskConfig.Schedule == EVisibilityTaskSchedule::Parallel)
