Fixed: broken rendering when -rdgimmediate mode is enabled in a mobile renderer
Added RHISetStaticUniformBuffer implementation to all RHIs

diff --git a/Engine/Source/Runtime/Apple/MetalRHI/Private/MetalCommands.cpp b/Engine/Source/Runtime/Apple/MetalRHI/Private/MetalCommands.cpp
--- a/Engine/Source/Runtime/Apple/MetalRHI/Private/MetalCommands.cpp
+++ b/Engine/Source/Runtime/Apple/MetalRHI/Private/MetalCommands.cpp
@@ -476,6 +476,11 @@ void FMetalRHICommandContext::RHISetStaticUniformBuffers(const FUniformBufferSta
 	}
 }
 
+void FMetalRHICommandContext::RHISetStaticUniformBuffer(FUniformBufferStaticSlot InSlot, FRHIUniformBuffer* InBuffer)
+{
+	GlobalUniformBuffers[InSlot] = InBuffer;
+}
+
 struct FMetalShaderBinder
 {
     FMetalStateCache& StateCache;
diff --git a/Engine/Source/Runtime/Apple/MetalRHI/Public/MetalRHIContext.h b/Engine/Source/Runtime/Apple/MetalRHI/Public/MetalRHIContext.h
--- a/Engine/Source/Runtime/Apple/MetalRHI/Public/MetalRHIContext.h
+++ b/Engine/Source/Runtime/Apple/MetalRHI/Public/MetalRHIContext.h
@@ -95,6 +95,7 @@ public:
 	virtual void RHISetGraphicsPipelineState(FRHIGraphicsPipelineState* GraphicsState, uint32 StencilRef, bool bApplyAdditionalState) final override;
 	
 	virtual void RHISetStaticUniformBuffers(const FUniformBufferStaticBindings& InUniformBuffers) final override;
+	virtual void RHISetStaticUniformBuffer(FUniformBufferStaticSlot Slot, FRHIUniformBuffer* Buffer) final override;
 
 	virtual void RHISetShaderParameters(FRHIGraphicsShader* Shader, TConstArrayView<uint8> InParametersData, TConstArrayView<FRHIShaderParameter> InParameters, TConstArrayView<FRHIShaderParameterResource> InResourceParameters, TConstArrayView<FRHIShaderParameterResource> InBindlessParameters) final override;
 	virtual void RHISetShaderParameters(FRHIComputeShader* Shader, TConstArrayView<uint8> InParametersData, TConstArrayView<FRHIShaderParameter> InParameters, TConstArrayView<FRHIShaderParameterResource> InResourceParameters, TConstArrayView<FRHIShaderParameterResource> InBindlessParameters) final override;
diff --git a/Engine/Source/Runtime/D3D12RHI/Private/D3D12CommandContext.h b/Engine/Source/Runtime/D3D12RHI/Private/D3D12CommandContext.h
--- a/Engine/Source/Runtime/D3D12RHI/Private/D3D12CommandContext.h
+++ b/Engine/Source/Runtime/D3D12RHI/Private/D3D12CommandContext.h
@@ -634,6 +634,7 @@ public:
 	virtual void RHIDispatchComputeShader(uint32 ThreadGroupCountX, uint32 ThreadGroupCountY, uint32 ThreadGroupCountZ) final override;
 	virtual void RHIDispatchIndirectComputeShader(FRHIBuffer* ArgumentBuffer, uint32 ArgumentOffset) final override;
 	virtual void RHISetStaticUniformBuffers(const FUniformBufferStaticBindings& InUniformBuffers) final override;
+	virtual void RHISetStaticUniformBuffer(FUniformBufferStaticSlot Slot, FRHIUniformBuffer* Buffer) final override;
 	virtual void RHISetShaderParameters(FRHIComputeShader* Shader, TConstArrayView<uint8> InParametersData, TConstArrayView<FRHIShaderParameter> InParameters, TConstArrayView<FRHIShaderParameterResource> InResourceParameters, TConstArrayView<FRHIShaderParameterResource> InBindlessParameters) final override;
 	virtual void RHISetShaderUnbinds(FRHIComputeShader* Shader, TConstArrayView<FRHIShaderParameterUnbind> InUnbinds) final override;
 	virtual void RHISetShaderUnbinds(FRHIGraphicsShader* Shader, TConstArrayView<FRHIShaderParameterUnbind> InUnbinds) final override;
@@ -1001,6 +1002,10 @@ public:
 	{
 		ContextRedirect(RHISetStaticUniformBuffers(InUniformBuffers));
 	}
+	FORCEINLINE virtual void RHISetStaticUniformBuffer(FUniformBufferStaticSlot Slot, FRHIUniformBuffer* Buffer) final override
+	{
+		ContextRedirect(RHISetStaticUniformBuffer(Slot, Buffer));
+	}
 	FORCEINLINE virtual void RHISetShaderParameters(FRHIGraphicsShader* Shader, TConstArrayView<uint8> InParametersData, TConstArrayView<FRHIShaderParameter> InParameters, TConstArrayView<FRHIShaderParameterResource> InResourceParameters, TConstArrayView<FRHIShaderParameterResource> InBindlessParameters) final override
 	{
 		ContextRedirect(RHISetShaderParameters(Shader, InParametersData, InParameters, InResourceParameters, InBindlessParameters));
diff --git a/Engine/Source/Runtime/D3D12RHI/Private/D3D12Commands.cpp b/Engine/Source/Runtime/D3D12RHI/Private/D3D12Commands.cpp
--- a/Engine/Source/Runtime/D3D12RHI/Private/D3D12Commands.cpp
+++ b/Engine/Source/Runtime/D3D12RHI/Private/D3D12Commands.cpp
@@ -812,6 +812,11 @@ void FD3D12CommandContext::RHISetStaticUniformBuffers(const FUniformBufferStatic
 
 }
 
+void FD3D12CommandContext::RHISetStaticUniformBuffer(FUniformBufferStaticSlot InSlot, FRHIUniformBuffer* InBuffer)
+{
+	StaticUniformBuffers[InSlot] = InBuffer;
+}
+
 void FD3D12CommandContext::RHICopyToStagingBuffer(FRHIBuffer* SourceBufferRHI, FRHIStagingBuffer* StagingBufferRHI, uint32 Offset, uint32 NumBytes)
 {
 	SCOPE_CYCLE_COUNTER(STAT_D3D12CopyToStagingBufferTime);
diff --git a/Engine/Source/Runtime/NullDrv/Public/NullRHI.h b/Engine/Source/Runtime/NullDrv/Public/NullRHI.h
--- a/Engine/Source/Runtime/NullDrv/Public/NullRHI.h
+++ b/Engine/Source/Runtime/NullDrv/Public/NullRHI.h
@@ -526,6 +526,14 @@ public:
 		return 0;
 	}
 
+	virtual void RHISetStaticUniformBuffers(const FUniformBufferStaticBindings& InUniformBuffers) final override
+	{
+	}
+
+	virtual void RHISetStaticUniformBuffer(FUniformBufferStaticSlot Slot, FRHIUniformBuffer* UniformBuffer) final override
+	{
+	}
+
 #if WITH_RHI_BREADCRUMBS
 	virtual void RHIBeginBreadcrumbGPU(FRHIBreadcrumbNode* Breadcrumb) final override
 	{
diff --git a/Engine/Source/Runtime/RHI/Public/RHIContext.h b/Engine/Source/Runtime/RHI/Public/RHIContext.h
--- a/Engine/Source/Runtime/RHI/Public/RHIContext.h
+++ b/Engine/Source/Runtime/RHI/Public/RHIContext.h
@@ -384,15 +384,9 @@ public:
 		checkf(false, TEXT("RHISetShaderUnbinds called when the active RHI hasn't overridden it and GRHIGlobals.NeedsShaderUnbinds is set."));
 	}
 
-	virtual void RHISetStaticUniformBuffers(const FUniformBufferStaticBindings& InUniformBuffers)
-	{
-		/** empty default implementation. */
-	}
+	virtual void RHISetStaticUniformBuffers(const FUniformBufferStaticBindings& InUniformBuffers) = 0;
 
-	virtual void RHISetStaticUniformBuffer(FUniformBufferStaticSlot Slot, FRHIUniformBuffer* UniformBuffer)
-	{
-		/* empty default implementation */
-	}
+	virtual void RHISetStaticUniformBuffer(FUniformBufferStaticSlot Slot, FRHIUniformBuffer* UniformBuffer) = 0;
 
 	virtual void RHISetUniformBufferDynamicOffset(FUniformBufferStaticSlot Slot, uint32 Offset)
 	{
diff --git a/Engine/Source/Runtime/RHI/Public/RHIValidationContext.h b/Engine/Source/Runtime/RHI/Public/RHIValidationContext.h
--- a/Engine/Source/Runtime/RHI/Public/RHIValidationContext.h
+++ b/Engine/Source/Runtime/RHI/Public/RHIValidationContext.h
@@ -317,6 +317,11 @@ public:
 		RHIContext->RHISetStaticUniformBuffers(InUniformBuffers);
 	}
 
+	virtual void RHISetStaticUniformBuffer(FUniformBufferStaticSlot Slot, FRHIUniformBuffer* UniformBuffer) override final
+	{
+		RHIContext->RHISetStaticUniformBuffer(Slot, UniformBuffer);
+	}
+
 #if WITH_RHI_BREADCRUMBS
 	virtual void RHIBeginBreadcrumbGPU(FRHIBreadcrumbNode* Breadcrumb) final override
 	{
@@ -909,6 +914,11 @@ public:
 		RHIContext->RHISetStaticUniformBuffers(InUniformBuffers);
 	}
 
+	virtual void RHISetStaticUniformBuffer(FUniformBufferStaticSlot Slot, FRHIUniformBuffer* UniformBuffer) override final
+	{
+		RHIContext->RHISetStaticUniformBuffer(Slot, UniformBuffer);
+	}
+
 	virtual void RHISetStencilRef(uint32 StencilRef) override final
 	{
 		//checkf(State.bGfxPSOSet, TEXT("A Graphics PSO has to be set to change stencil ref!"));
diff --git a/Engine/Source/Runtime/Renderer/Private/InstanceCulling/InstanceCullingContext.cpp b/Engine/Source/Runtime/Renderer/Private/InstanceCulling/InstanceCullingContext.cpp
--- a/Engine/Source/Runtime/Renderer/Private/InstanceCulling/InstanceCullingContext.cpp
+++ b/Engine/Source/Runtime/Renderer/Private/InstanceCulling/InstanceCullingContext.cpp
@@ -93,23 +93,29 @@ static uint32 PackDrawCommandDesc(bool bMaterialUsesWorldPositionOffset, bool bM
 FMeshDrawCommandOverrideArgs GetMeshDrawCommandOverrideArgs(const FInstanceCullingDrawParams& InstanceCullingDrawParams)
 {
 	FMeshDrawCommandOverrideArgs Result;
-	Result.InstanceBuffer = InstanceCullingDrawParams.InstanceIdOffsetBuffer.GetBuffer() != nullptr ? InstanceCullingDrawParams.InstanceIdOffsetBuffer.GetBuffer()->GetRHI() : nullptr;
 	Result.IndirectArgsBuffer = InstanceCullingDrawParams.DrawIndirectArgsBuffer.GetBuffer() != nullptr ? InstanceCullingDrawParams.DrawIndirectArgsBuffer.GetBuffer()->GetRHI() : nullptr;
-	Result.BatchedPrimitiveUB = InstanceCullingDrawParams.BatchedPrimitive.GetUniformBuffer() != nullptr ? InstanceCullingDrawParams.BatchedPrimitive.GetUniformBuffer()->GetRHI() : nullptr;
 	Result.InstanceDataByteOffset = InstanceCullingDrawParams.InstanceDataByteOffset;
 	Result.IndirectArgsByteOffset = InstanceCullingDrawParams.IndirectArgsByteOffset;
+	
+	if (InstanceCullingDrawParams.BatchedPrimitive.GetUniformBuffer() != nullptr)
+	{
+		Result.InstanceBuffer = nullptr;
+		Result.InstanceCullingStaticUB = InstanceCullingDrawParams.BatchedPrimitive.GetUniformBuffer()->GetRHI();
+	}
+	else
+	{
+		Result.InstanceBuffer = InstanceCullingDrawParams.InstanceIdOffsetBuffer.GetBuffer() != nullptr ? InstanceCullingDrawParams.InstanceIdOffsetBuffer.GetBuffer()->GetRHI() : nullptr;
+		Result.InstanceCullingStaticUB = InstanceCullingDrawParams.InstanceCulling.GetUniformBuffer() != nullptr ? InstanceCullingDrawParams.InstanceCulling.GetUniformBuffer()->GetRHI() : nullptr;
+	}
+
 	return Result;
 }
 
-FUniformBufferStaticSlot FInstanceCullingContext::GetUniformBufferViewStaticSlot(EShaderPlatform ShaderPlatform)
+FUniformBufferStaticSlot FInstanceCullingContext::GetStaticUniformBufferSlot(EShaderPlatform ShaderPlatform)
 {
-	FUniformBufferStaticSlot StaticSlot = MAX_UNIFORM_BUFFER_STATIC_SLOTS;
-	if (PlatformGPUSceneUsesUniformBufferView(ShaderPlatform))
-	{
-		static FName BatchedPrimitiveSlotName = "BatchedPrimitive";
-		StaticSlot = FUniformBufferStaticSlotRegistry::Get().FindSlotByName(BatchedPrimitiveSlotName);
-	}
-	return StaticSlot;
+	const static FName InstanceCullingSlotName = "InstanceCullingUbSlot";
+	const static FName BatchedPrimitiveSlotName = "BatchedPrimitive";
+	return FUniformBufferStaticSlotRegistry::Get().FindSlotByName(PlatformGPUSceneUsesUniformBufferView(ShaderPlatform) ? BatchedPrimitiveSlotName : InstanceCullingSlotName);
 }
 
 static uint32 GetInstanceDataStrideElements(EShaderPlatform ShaderPlatform, EBatchProcessingMode Mode)
@@ -171,11 +177,11 @@ uint32 FInstanceCullingContext::GetInstanceIdNumElements() const
 FInstanceCullingContext::FInstanceCullingContext(
 	const TCHAR* PassName,
 	EShaderPlatform InShaderPlatform,
-	FInstanceCullingManager* InInstanceCullingManager, 
-	TArrayView<const int32> InViewIds, 
-	const TRefCountPtr<IPooledRenderTarget>& InPrevHZB, 
-	EInstanceCullingMode InInstanceCullingMode, 
-	EInstanceCullingFlags InFlags, 
+	FInstanceCullingManager* InInstanceCullingManager,
+	TArrayView<const int32> InViewIds,
+	const TRefCountPtr<IPooledRenderTarget>& InPrevHZB,
+	EInstanceCullingMode InInstanceCullingMode,
+	EInstanceCullingFlags InFlags,
 	EBatchProcessingMode InSingleInstanceProcessingMode) :
 	InstanceCullingManager(InInstanceCullingManager),
 	ShaderPlatform(InShaderPlatform),
@@ -185,7 +191,7 @@ FInstanceCullingContext::FInstanceCullingContext(
 	InstanceCullingMode(InInstanceCullingMode),
 	Flags(InFlags),
 	SingleInstanceProcessingMode(InSingleInstanceProcessingMode),
-	BatchedPrimitiveSlot(GetUniformBufferViewStaticSlot(InShaderPlatform)),
+	InstanceCullingStaticSlot(IsMobilePlatform(InShaderPlatform) ? GetStaticUniformBufferSlot(InShaderPlatform) : MAX_UNIFORM_BUFFER_STATIC_SLOTS),
 	bUsesUniformBufferView(PlatformGPUSceneUsesUniformBufferView(InShaderPlatform))
 {
 #if MESH_DRAW_COMMAND_STATS
@@ -947,19 +953,21 @@ void FInstanceCullingContext::BuildRenderingCommandsInternal(
 
 	InstanceCullingDrawParams->DrawIndirectArgsBuffer = DrawIndirectArgsRDG;
 	InstanceCullingDrawParams->InstanceIdOffsetBuffer = InstanceIdOffsetBufferRDG;
-
-	FInstanceCullingGlobalUniforms* UniformParameters = GraphBuilder.AllocParameters<FInstanceCullingGlobalUniforms>();
-	UniformParameters->InstanceIdsBuffer = GraphBuilder.CreateSRV(InstanceIdsBuffer);
-	UniformParameters->PageInfoBuffer = GraphBuilder.CreateSRV(InstanceIdsBuffer);
-	UniformParameters->BufferCapacity = InstanceIdBufferSize;
-	InstanceCullingDrawParams->InstanceCulling = GraphBuilder.CreateUniformBuffer(UniformParameters);
-
+	
 	if (PlatformGPUSceneUsesUniformBufferView(ShaderPlatform))
 	{
 		FBatchedPrimitiveParameters* BatchedPrimitiveParameters = GraphBuilder.AllocParameters<FBatchedPrimitiveParameters>();
 		BatchedPrimitiveParameters->Data = GraphBuilder.CreateSRV(InstanceIdsBuffer);
 		InstanceCullingDrawParams->BatchedPrimitive = GraphBuilder.CreateUniformBuffer(BatchedPrimitiveParameters);
 	}
+	else
+	{
+		FInstanceCullingGlobalUniforms* UniformParameters = GraphBuilder.AllocParameters<FInstanceCullingGlobalUniforms>();
+		UniformParameters->InstanceIdsBuffer = GraphBuilder.CreateSRV(InstanceIdsBuffer);
+		UniformParameters->PageInfoBuffer = GraphBuilder.CreateSRV(InstanceIdsBuffer);
+		UniformParameters->BufferCapacity = InstanceIdBufferSize;
+		InstanceCullingDrawParams->InstanceCulling = GraphBuilder.CreateUniformBuffer(UniformParameters);
+	}
 
 #if MESH_DRAW_COMMAND_STATS
 	if (MeshDrawCommandPassStats)
@@ -1304,18 +1312,20 @@ FInstanceCullingDeferredContext *FInstanceCullingContext::CreateDeferredContext(
 		}
 	}
 
-	FInstanceCullingGlobalUniforms* UniformParameters = GraphBuilder.AllocParameters<FInstanceCullingGlobalUniforms>();
-	UniformParameters->InstanceIdsBuffer = GraphBuilder.CreateSRV(InstanceIdsBuffer);
-	UniformParameters->PageInfoBuffer = GraphBuilder.CreateSRV(InstanceIdsBuffer);
-	UniformParameters->BufferCapacity = 0U; // TODO: this is not used at the moment, but is intended for range checks so would have been good.
-	DeferredContext->UniformBuffer = GraphBuilder.CreateUniformBuffer(UniformParameters);
-
 	if (PlatformGPUSceneUsesUniformBufferView(ShaderPlatform))
 	{
 		FBatchedPrimitiveParameters* BatchedPrimitiveParameters = GraphBuilder.AllocParameters<FBatchedPrimitiveParameters>();
 		BatchedPrimitiveParameters->Data = GraphBuilder.CreateSRV(InstanceIdsBuffer);
 		DeferredContext->BatchedPrimitive = GraphBuilder.CreateUniformBuffer(BatchedPrimitiveParameters);
 	}
+	else
+	{
+		FInstanceCullingGlobalUniforms* UniformParameters = GraphBuilder.AllocParameters<FInstanceCullingGlobalUniforms>();
+		UniformParameters->InstanceIdsBuffer = GraphBuilder.CreateSRV(InstanceIdsBuffer);
+		UniformParameters->PageInfoBuffer = GraphBuilder.CreateSRV(InstanceIdsBuffer);
+		UniformParameters->BufferCapacity = 0U; // TODO: this is not used at the moment, but is intended for range checks so would have been good.
+		DeferredContext->UniformBuffer = GraphBuilder.CreateUniformBuffer(UniformParameters);
+	}
 
 #undef INST_CULL_CREATE_STRUCT_BUFF_ARGS
 #undef INST_CULL_CALLBACK
@@ -1669,13 +1679,16 @@ void FInstanceCullingContext::SubmitDrawCommands(
 				
 		FMeshDrawCommandSceneArgs SceneArgs;
 		SceneArgs.PrimitiveIdsBuffer = OverrideArgs.InstanceBuffer;
-		if (IsUniformBufferStaticSlotValid(BatchedPrimitiveSlot))
+		
+		if (IsUniformBufferStaticSlotValid(InstanceCullingStaticSlot))
 		{
-			// UniformBufferView suplies instance data through global UB binding
-			SceneArgs.PrimitiveIdsBuffer = nullptr;
-			SceneArgs.BatchedPrimitiveSlot = BatchedPrimitiveSlot;
-			// Only does anything when batching is disabled and each mesh pass has an unique BatchedPrimitive UB
-			RHICmdList.SetStaticUniformBuffer(BatchedPrimitiveSlot, OverrideArgs.BatchedPrimitiveUB);
+			if (bUsesUniformBufferView)
+			{
+				SceneArgs.BatchedPrimitiveSlot = InstanceCullingStaticSlot;
+			}
+			// Only does anything when batching is disabled and each mesh pass has an unique InstanceCulling/BatchedPrimitive UB
+			// otherwise static UBs set once when we start render-pass
+			RHICmdList.SetStaticUniformBuffer(InstanceCullingStaticSlot, OverrideArgs.InstanceCullingStaticUB);
 		}
 							
 		FMeshDrawCommandStateCache StateCache;
diff --git a/Engine/Source/Runtime/Renderer/Private/MeshPassProcessor.cpp b/Engine/Source/Runtime/Renderer/Private/MeshPassProcessor.cpp
--- a/Engine/Source/Runtime/Renderer/Private/MeshPassProcessor.cpp
+++ b/Engine/Source/Runtime/Renderer/Private/MeshPassProcessor.cpp
@@ -1842,32 +1842,28 @@ void DrawDynamicMeshPassPrivate(
 		//               This is because the stride on the instance data buffer is set to 0 so only the first will ever be fetched.
 		const bool bDynamicInstancing = false;
 		const ERHIFeatureLevel::Type FeatureLevel = View.GetFeatureLevel();
-		const bool bUseGPUScene = UseGPUScene(GMaxRHIShaderPlatform, FeatureLevel);
-		
-		FMeshDrawCommandSceneArgs SceneArgs;
-		if (bUseGPUScene)
-		{
-			SceneArgs.BatchedPrimitiveSlot = FInstanceCullingContext::GetUniformBufferViewStaticSlot(View.GetShaderPlatform());
-		}
-	
-		const uint32 PrimitiveIdBufferStride = FInstanceCullingContext::GetInstanceIdBufferStride(View.GetShaderPlatform());
+		const EShaderPlatform ShaderPlatform = View.GetShaderPlatform();
+		const bool bUseGPUScene = UseGPUScene(ShaderPlatform, FeatureLevel);
 		
 		ApplyViewOverridesToMeshDrawCommands(View, VisibleMeshDrawCommands, DynamicMeshDrawCommandStorage, GraphicsMinimalPipelineStateSet, InNeedsShaderInitialisation);
 
 		check(View.bIsViewInfo);
 		const FViewInfo* ViewInfo = static_cast<const FViewInfo*>(&View);
 
+		FMeshDrawCommandSceneArgs SceneArgs;
 		SortAndMergeDynamicPassMeshDrawCommands(View, RHICmdList, VisibleMeshDrawCommands, DynamicMeshDrawCommandStorage, SceneArgs.PrimitiveIdsBuffer, InstanceFactor, &ViewInfo->DynamicPrimitiveCollector);
 
-		if (IsUniformBufferStaticSlotValid(SceneArgs.BatchedPrimitiveSlot))
+		if (bUseGPUScene && PlatformGPUSceneUsesUniformBufferView(ShaderPlatform))
 		{
 			FShaderResourceViewRHIRef BatchedPrimitiveSRV = RHICmdList.CreateShaderResourceView(SceneArgs.PrimitiveIdsBuffer, FRHIViewDesc::CreateBufferSRV().SetTypeFromBuffer(SceneArgs.PrimitiveIdsBuffer));
 			SceneArgs.PrimitiveIdsBuffer = nullptr;
 			FRHIShaderResourceView* SRV = BatchedPrimitiveSRV.GetReference();
 			FUniformBufferRHIRef UBRef = RHICreateUniformBuffer(&SRV, GDynamicBatchedPrimitiveLayout.LayoutRHI, UniformBuffer_SingleFrame, EUniformBufferValidation::None);
+			SceneArgs.BatchedPrimitiveSlot = FInstanceCullingContext::GetStaticUniformBufferSlot(ShaderPlatform);
 			RHICmdList.SetStaticUniformBuffer(SceneArgs.BatchedPrimitiveSlot, UBRef);
 		}
-
+		
+		const uint32 PrimitiveIdBufferStride = FInstanceCullingContext::GetInstanceIdBufferStride(ShaderPlatform);
 		SubmitMeshDrawCommandsRange(VisibleMeshDrawCommands, GraphicsMinimalPipelineStateSet, SceneArgs, PrimitiveIdBufferStride, bDynamicInstancing, 0, VisibleMeshDrawCommands.Num(), InstanceFactor, RHICmdList);
 	}
 }
diff --git a/Engine/Source/Runtime/Renderer/Private/MobileBasePassRendering.cpp b/Engine/Source/Runtime/Renderer/Private/MobileBasePassRendering.cpp
--- a/Engine/Source/Runtime/Renderer/Private/MobileBasePassRendering.cpp
+++ b/Engine/Source/Runtime/Renderer/Private/MobileBasePassRendering.cpp
@@ -598,13 +598,22 @@ void FMobileSceneRenderer::RenderMobileEditorPrimitives(FRHICommandList& RHICmdL
 			}
 		});
 		
-		// DrawDynamicMeshPass may change BatchedPrimitive binding, so we need to restore it
-		FUniformBufferStaticSlot BatchedPrimitiveSlot = FInstanceCullingContext::GetUniformBufferViewStaticSlot(View.GetShaderPlatform());
-		if (IsUniformBufferStaticSlotValid(BatchedPrimitiveSlot) && InstanceCullingDrawParams->BatchedPrimitive)
+		// DrawDynamicMeshPass may change global InstanceCulling binding, so we need to restore it
+		if (UseGPUScene(View.GetShaderPlatform(), FeatureLevel))
 		{
-			FRHIUniformBuffer* BatchedPrimitiveBufferRHI = InstanceCullingDrawParams->BatchedPrimitive.GetUniformBuffer()->GetRHI();
-			check(BatchedPrimitiveBufferRHI);
-			RHICmdList.SetStaticUniformBuffer(BatchedPrimitiveSlot, BatchedPrimitiveBufferRHI);
+			FRHIUniformBuffer* InstanceCullingBufferRHI = nullptr;
+			if (PlatformGPUSceneUsesUniformBufferView(View.GetShaderPlatform()))
+			{
+				InstanceCullingBufferRHI = InstanceCullingDrawParams->BatchedPrimitive.GetUniformBuffer()->GetRHI();
+			}
+			else
+			{
+				InstanceCullingBufferRHI = InstanceCullingDrawParams->InstanceCulling.GetUniformBuffer()->GetRHI();
+			}
+			check(InstanceCullingBufferRHI);
+			
+			FUniformBufferStaticSlot InstanceCullingStaticSlot = FInstanceCullingContext::GetStaticUniformBufferSlot(View.GetShaderPlatform());
+			RHICmdList.SetStaticUniformBuffer(InstanceCullingStaticSlot, InstanceCullingBufferRHI);
 		}
 
 		// Draw the view's batched simple elements(lines, sprites, etc).
diff --git a/Engine/Source/Runtime/Renderer/Private/MobileShadingRenderer.cpp b/Engine/Source/Runtime/Renderer/Private/MobileShadingRenderer.cpp
--- a/Engine/Source/Runtime/Renderer/Private/MobileShadingRenderer.cpp
+++ b/Engine/Source/Runtime/Renderer/Private/MobileShadingRenderer.cpp
@@ -1910,7 +1910,6 @@ void FMobileSceneRenderer::RenderForwardMultiPass(FRDGBuilder& GraphBuilder, FMo
 	{
 		FMobileRenderPassParameters PassParameters;
 		FInstanceCullingDrawParams MeshDecalSceneColorInstanceCullingDrawParams;
-		FInstanceCullingDrawParams TranslucencyInstanceCullingDrawParams;
 	};
 	auto SecondParameterCollection = GraphBuilder.AllocParameters<FForwardSecondPassParameterCollection>();
 
@@ -1930,7 +1929,7 @@ void FMobileSceneRenderer::RenderForwardMultiPass(FRDGBuilder& GraphBuilder, FMo
 	if (Scene->GPUScene.IsEnabled())
 	{
 		BuildMeshRenderingCommands(GraphBuilder, EMeshPass::MeshDecal_SceneColor, View, Scene->GPUScene, InstanceCullingManager, SecondParameterCollection->MeshDecalSceneColorInstanceCullingDrawParams);
-		BuildMeshRenderingCommands(GraphBuilder, StandardTranslucencyMeshPass, View, Scene->GPUScene, InstanceCullingManager, SecondParameterCollection->TranslucencyInstanceCullingDrawParams);
+		BuildMeshRenderingCommands(GraphBuilder, StandardTranslucencyMeshPass, View, Scene->GPUScene, InstanceCullingManager, SecondPassParameters->InstanceCullingDrawParams);
 	}
 
 	GraphBuilder.AddPass(
@@ -1947,7 +1946,7 @@ void FMobileSceneRenderer::RenderForwardMultiPass(FRDGBuilder& GraphBuilder, FMo
 		RenderModulatedShadowProjections(RHICmdList, ViewContext.ViewIndex, View);
 		RenderFog(RHICmdList, View);
 		// Draw translucency.
-		RenderTranslucency(RHICmdList, View, Views, StandardTranslucencyPass, StandardTranslucencyMeshPass, &SecondParameterCollection->TranslucencyInstanceCullingDrawParams);
+		RenderTranslucency(RHICmdList, View, Views, StandardTranslucencyPass, StandardTranslucencyMeshPass, &SecondPassParameters->InstanceCullingDrawParams);
 
 		if (bDoOcclusionQueries)
 		{
@@ -2334,31 +2333,23 @@ void FMobileSceneRenderer::RenderDeferredMultiPass(FRDGBuilder& GraphBuilder, FS
 
 			View.BeginRenderView();
 
-			struct FDeferredMultiPassDecalParameterCollection
-			{
-				FMobileRenderPassParameters PassParameters;
-				FInstanceCullingDrawParams MeshDecalSceneColorAndGBufferInstanceCullingDrawParams;
-			};
-			auto ParameterCollection = GraphBuilder.AllocParameters<FDeferredMultiPassDecalParameterCollection>();
-
-			auto PassParameters = &ParameterCollection->PassParameters;
+			auto PassParameters = GraphBuilder.AllocParameters<FMobileRenderPassParameters>();
 			PassParameters->View = View.GetShaderParameters();
 			PassParameters->MobileBasePass = CreateMobileBasePassUniformBuffer(GraphBuilder, View, EMobileBasePass::Opaque, EMobileSceneTextureSetupMode::SceneDepth);
 			PassParameters->RenderTargets = BasePassRenderTargets;
 
 			if (Scene->GPUScene.IsEnabled())
 			{
-				BuildMeshRenderingCommands(GraphBuilder, EMeshPass::MeshDecal_SceneColor, View, Scene->GPUScene, InstanceCullingManager, PassParameters->InstanceCullingDrawParams);
-				BuildMeshRenderingCommands(GraphBuilder, EMeshPass::MeshDecal_SceneColorAndGBuffer, View, Scene->GPUScene, InstanceCullingManager, ParameterCollection->MeshDecalSceneColorAndGBufferInstanceCullingDrawParams);
+				BuildMeshRenderingCommands(GraphBuilder, EMeshPass::MeshDecal_SceneColorAndGBuffer, View, Scene->GPUScene, InstanceCullingManager, PassParameters->InstanceCullingDrawParams);
 			}
 
 			GraphBuilder.AddPass(
 				RDG_EVENT_NAME("Decals"),
 				PassParameters,
 				ERDGPassFlags::Raster,
-				[this, ParameterCollection, &View](FRHICommandList& RHICmdList)
+				[this, PassParameters, &View](FRHICommandList& RHICmdList)
 				{
-					RenderDecals(RHICmdList, View, &ParameterCollection->MeshDecalSceneColorAndGBufferInstanceCullingDrawParams);
+					RenderDecals(RHICmdList, View, &PassParameters->InstanceCullingDrawParams);
 				});
 		}
 	}
diff --git a/Engine/Source/Runtime/Renderer/Private/SimpleMeshDrawCommandPass.cpp b/Engine/Source/Runtime/Renderer/Private/SimpleMeshDrawCommandPass.cpp
--- a/Engine/Source/Runtime/Renderer/Private/SimpleMeshDrawCommandPass.cpp
+++ b/Engine/Source/Runtime/Renderer/Private/SimpleMeshDrawCommandPass.cpp
@@ -80,9 +80,15 @@ void FSimpleMeshDrawCommandPass::SubmitDraw(FRHICommandList& RHICmdList, const F
 		else
 		{
 			FMeshDrawCommandSceneArgs SceneArgs;
-			SceneArgs.PrimitiveIdsBuffer = GetMeshDrawCommandOverrideArgs(InstanceCullingDrawParams).InstanceBuffer;
+			if (InstanceCullingContext.bUsesUniformBufferView)
+			{
+				SceneArgs.BatchedPrimitiveSlot = InstanceCullingContext.InstanceCullingStaticSlot;
+			}
+			else
+			{
+				SceneArgs.PrimitiveIdsBuffer = GetMeshDrawCommandOverrideArgs(InstanceCullingDrawParams).InstanceBuffer;
+			}
 			SceneArgs.PrimitiveIdOffset = 0u;
-			SceneArgs.BatchedPrimitiveSlot = InstanceCullingContext.BatchedPrimitiveSlot;
 			const uint32 PrimitiveIdBufferStride = FInstanceCullingContext::GetInstanceIdBufferStride(InstanceCullingContext.ShaderPlatform);
 
 			SubmitMeshDrawCommandsRange(VisibleMeshDrawCommands, GraphicsMinimalPipelineStateSet, SceneArgs, PrimitiveIdBufferStride, false, 0, VisibleMeshDrawCommands.Num(), InstanceFactor, RHICmdList);
diff --git a/Engine/Source/Runtime/Renderer/Public/InstanceCulling/InstanceCullingContext.h b/Engine/Source/Runtime/Renderer/Public/InstanceCulling/InstanceCullingContext.h
--- a/Engine/Source/Runtime/Renderer/Public/InstanceCulling/InstanceCullingContext.h
+++ b/Engine/Source/Runtime/Renderer/Public/InstanceCulling/InstanceCullingContext.h
@@ -90,7 +90,7 @@ public:
 	static constexpr uint32 IndirectArgsNumWords = 5;
 	static constexpr uint32 CompactionBlockNumInstances = 64;
 	RENDERER_API static uint32 GetInstanceIdBufferStride(EShaderPlatform ShaderPlatform);
-	RENDERER_API static FUniformBufferStaticSlot GetUniformBufferViewStaticSlot(EShaderPlatform ShaderPlatform);
+	RENDERER_API static FUniformBufferStaticSlot GetStaticUniformBufferSlot(EShaderPlatform ShaderPlatform);
 	
 	FInstanceCullingContext() {}
 
@@ -350,9 +350,9 @@ public:
 	// Processing mode to use for single-instance primitives, default to skip culling, as this is already done on CPU. 
 	EBatchProcessingMode SingleInstanceProcessingMode = EBatchProcessingMode::UnCulled;
 
-	// A valid uniform buffer slot in case shader platform supplies instance data trhough a global UB binding
-	FUniformBufferStaticSlot BatchedPrimitiveSlot;
-	// Whether current platform uses Uniform Buffer View path
+	// Static uniform buffer slot for InstanceCulling/BatchedPrimitive UB. Only valid for a mobile renderer, which needs to handle cases with a manualy merged render-passes
+	FUniformBufferStaticSlot InstanceCullingStaticSlot;
+	// Whether current platform uses 'Uniform Buffer View' path
 	bool bUsesUniformBufferView;
 	
 #if MESH_DRAW_COMMAND_STATS
diff --git a/Engine/Source/Runtime/Renderer/Public/MeshPassProcessor.h b/Engine/Source/Runtime/Renderer/Public/MeshPassProcessor.h
--- a/Engine/Source/Runtime/Renderer/Public/MeshPassProcessor.h
+++ b/Engine/Source/Runtime/Renderer/Public/MeshPassProcessor.h
@@ -1175,7 +1175,7 @@ struct FMeshDrawCommandOverrideArgs
 {
 	FRHIBuffer* InstanceBuffer;
 	FRHIBuffer* IndirectArgsBuffer;
-	FRHIUniformBuffer* BatchedPrimitiveUB;
+	FRHIUniformBuffer* InstanceCullingStaticUB;
 	uint32 InstanceDataByteOffset;
 	uint32 IndirectArgsByteOffset;
 	
@@ -1183,7 +1183,7 @@ struct FMeshDrawCommandOverrideArgs
 	{
 		InstanceBuffer = nullptr;
 		IndirectArgsBuffer = nullptr;
-		BatchedPrimitiveUB = nullptr;
+		InstanceCullingStaticUB = nullptr;
 		InstanceDataByteOffset = 0u;
 		IndirectArgsByteOffset = 0u;
 	}
diff --git a/Engine/Source/Runtime/Windows/D3D11RHI/Private/D3D11Commands.cpp b/Engine/Source/Runtime/Windows/D3D11RHI/Private/D3D11Commands.cpp
--- a/Engine/Source/Runtime/Windows/D3D11RHI/Private/D3D11Commands.cpp
+++ b/Engine/Source/Runtime/Windows/D3D11RHI/Private/D3D11Commands.cpp
@@ -346,6 +346,11 @@ void FD3D11DynamicRHI::RHISetStaticUniformBuffers(const FUniformBufferStaticBind
 	}
 }
 
+void FD3D11DynamicRHI::RHISetStaticUniformBuffer(FUniformBufferStaticSlot InSlot, FRHIUniformBuffer* InBuffer)
+{
+	StaticUniformBuffers[InSlot] = InBuffer;
+}
+
 template<EShaderFrequency ShaderFrequency>
 struct FD3D11ResourceBinder
 {
diff --git a/Engine/Source/Runtime/Windows/D3D11RHI/Private/D3D11RHIPrivate.h b/Engine/Source/Runtime/Windows/D3D11RHI/Private/D3D11RHIPrivate.h
--- a/Engine/Source/Runtime/Windows/D3D11RHI/Private/D3D11RHIPrivate.h
+++ b/Engine/Source/Runtime/Windows/D3D11RHI/Private/D3D11RHIPrivate.h
@@ -546,6 +546,7 @@ public:
 	virtual void RHISetBoundShaderState(FRHIBoundShaderState* BoundShaderState) final override;
 	virtual void RHISetGraphicsPipelineState(FRHIGraphicsPipelineState* GraphicsState, uint32 StencilRef, bool bApplyAdditionalState) final override;
 	virtual void RHISetStaticUniformBuffers(const FUniformBufferStaticBindings& InUniformBuffers) final override;
+	virtual void RHISetStaticUniformBuffer(FUniformBufferStaticSlot Slot, FRHIUniformBuffer* Buffer) final override;
 	virtual void RHISetShaderParameters(FRHIComputeShader* Shader, TConstArrayView<uint8> InParametersData, TConstArrayView<FRHIShaderParameter> InParameters, TConstArrayView<FRHIShaderParameterResource> InResourceParameters, TConstArrayView<FRHIShaderParameterResource> InBindlessParameters) final override;
 	virtual void RHISetShaderParameters(FRHIGraphicsShader* Shader, TConstArrayView<uint8> InParametersData, TConstArrayView<FRHIShaderParameter> InParameters, TConstArrayView<FRHIShaderParameterResource> InResourceParameters, TConstArrayView<FRHIShaderParameterResource> InBindlessParameters) final override;
 	virtual void RHISetShaderUnbinds(FRHIComputeShader* Shader, TConstArrayView<FRHIShaderParameterUnbind> InUnbinds) final override;
