Fixed: Assert about wrong uniform buffer bound when rendering primitices using a "Dynamic" mesh pass.
Related to a recent changes to UB hash computations.

diff --git a/Engine/Source/Runtime/Renderer/Private/MeshPassProcessor.cpp b/Engine/Source/Runtime/Renderer/Private/MeshPassProcessor.cpp
--- a/Engine/Source/Runtime/Renderer/Private/MeshPassProcessor.cpp
+++ b/Engine/Source/Runtime/Renderer/Private/MeshPassProcessor.cpp
@@ -1802,8 +1802,8 @@ public:
 	{
 		// This Layout fully replicates BatchedPrimitive UB
 		// we replace RDG_SRV with a regular SRV to be able to update UB inside RDG passes
-		static FName BatchedPrimitiveSlotName = "BatchedPrimitive";
-		FRHIUniformBufferLayoutInitializer Initializer(TEXT("DynamicBatchedPrimitive"), 16u);
+		const static FName BatchedPrimitiveSlotName = "BatchedPrimitive";
+		FRHIUniformBufferLayoutInitializer Initializer(TEXT("FBatchedPrimitiveParameters"), 16u);
 		EnumAddFlags(Initializer.Flags, ERHIUniformBufferFlags::UniformView);
 		EnumAddFlags(Initializer.Flags, ERHIUniformBufferFlags::NoEmulatedUniformBuffer); // implicit for UniformView
 		Initializer.Resources.Add({0, UBMT_RDG_BUFFER_SRV});
