[HWRT] Refactor RayTracingInstanceBufferUtil
- New FRayTracingInstanceBufferBuilder class to wrap and help manage all related logic.
- Deprecated all the existing functions.

diff --git a/Engine/Plugins/Experimental/GPULightmass/Source/GPULightmass/Private/LightmapRenderer.cpp b/Engine/Plugins/Experimental/GPULightmass/Source/GPULightmass/Private/LightmapRenderer.cpp
--- a/Engine/Plugins/Experimental/GPULightmass/Source/GPULightmass/Private/LightmapRenderer.cpp
+++ b/Engine/Plugins/Experimental/GPULightmass/Source/GPULightmass/Private/LightmapRenderer.cpp
@@ -1180,14 +1180,14 @@ bool FSceneRenderState::SetupRayTracingScene(FRDGBuilder& GraphBuilder, FSceneUn
 		{
 			SCOPED_GPU_MASK(RHICmdList, FRHIGPUMask::All());
 
-			const FRayTracingSceneInitializationData RayTracingSceneInitializationData = BuildRayTracingSceneInitializationData(RayTracingGeometryInstances);
-			const ERayTracingAccelerationStructureFlags SceneBuildFlags = ERayTracingAccelerationStructureFlags::FastTrace;
+			FRayTracingInstanceBufferBuilder RayTracingInstanceBufferBuilder;
+			RayTracingInstanceBufferBuilder.Init(RayTracingGeometryInstances, View.ViewMatrices.GetPreViewTranslation());
 
 			{
 				FRayTracingSceneInitializer Initializer;
 				Initializer.DebugName = FName(TEXT("LightmapRendererRayTracingScene"));
-				Initializer.MaxNumInstances = RayTracingSceneInitializationData.NumNativeGPUSceneInstances + RayTracingSceneInitializationData.NumNativeCPUInstances;
-				Initializer.BuildFlags = SceneBuildFlags;
+				Initializer.MaxNumInstances = RayTracingInstanceBufferBuilder.GetMaxNumInstances();
+				Initializer.BuildFlags = ERayTracingAccelerationStructureFlags::FastTrace;
 
 				RayTracingScene = RHICreateRayTracingScene(MoveTemp(Initializer));
 			}
@@ -1212,103 +1212,14 @@ bool FSceneRenderState::SetupRayTracingScene(FRDGBuilder& GraphBuilder, FSceneUn
 			FRWBufferStructured InstanceBuffer;
 			InstanceBuffer.Initialize(RHICmdList, TEXT("LightmassRayTracingInstanceBuffer"), GRHIRayTracingInstanceDescriptorSize, SceneInitializer.MaxNumInstances);
 
-			// Need to pass "BUF_MultiGPUAllocate", as acceleration structure virtual addresses are different per GPU
-			FByteAddressBuffer AccelerationStructureAddressesBuffer;
-			AccelerationStructureAddressesBuffer.Initialize(
-				RHICmdList,
-				TEXT("LightmassRayTracingAccelerationStructureAddressesBuffer"),
-				RayTracingSceneInitializationData.ReferencedGeometries.Num() * sizeof(FRayTracingAccelerationStructureAddress),
-				BUF_Volatile | BUF_MultiGPUAllocate);
-
-			const uint32 NumInstanceGroups = RayTracingSceneInitializationData.NumGPUInstanceGroups + RayTracingSceneInitializationData.NumCPUInstanceGroups;
-			const FRHIBufferCreateDesc InstanceGroupUploadCreateDesc =
-				FRHIBufferCreateDesc::CreateStructured<FRayTracingInstanceGroup>(TEXT("RayTracingTestBedInstanceGroupUploadBuffer"), NumInstanceGroups)
-				.AddUsage(EBufferUsageFlags::ShaderResource | EBufferUsageFlags::Volatile)
-				.SetInitActionInitializer()
-				.DetermineInitialState();
-			TRHIBufferInitializer<FRayTracingInstanceGroup> InstanceGroupUploadData = RHICmdList.CreateBufferInitializer(InstanceGroupUploadCreateDesc);
-
-			const uint32 NumInstanceDescriptors = RayTracingSceneInitializationData.NumGPUInstanceDescriptors + RayTracingSceneInitializationData.NumCPUInstanceDescriptors;
-			const FRHIBufferCreateDesc InstanceUploadCreateDesc =
-				FRHIBufferCreateDesc::CreateStructured<FRayTracingInstanceDescriptor>(TEXT("LightmassRayTracingInstanceUploadBuffer"), NumInstanceDescriptors)
-				.AddUsage(EBufferUsageFlags::ShaderResource | EBufferUsageFlags::Volatile)
-				.SetInitActionInitializer()
-				.DetermineInitialState();
-			TRHIBufferInitializer<FRayTracingInstanceDescriptor> InstanceUploadData = RHICmdList.CreateBufferInitializer(InstanceUploadCreateDesc);
-
-			const FRHIBufferCreateDesc TransformUploadCreateDesc =
-				FRHIBufferCreateDesc::CreateStructured<FVector4f>(TEXT("LightmassRayTracingTransformUploadBuffer"), RayTracingSceneInitializationData.NumNativeCPUInstances * 3)
-				.AddUsage(EBufferUsageFlags::ShaderResource | EBufferUsageFlags::Volatile)
-				.SetInitActionInitializer()
-				.DetermineInitialState();
-			TRHIBufferInitializer<FVector4f> TransformUploadData = RHICmdList.CreateBufferInitializer(TransformUploadCreateDesc);
-
-			{
-				FillRayTracingInstanceUploadBuffer(
-					View.ViewMatrices.GetPreViewTranslation(),
-					RayTracingGeometryInstances,
-					RayTracingSceneInitializationData.InstanceGeometryIndices,
-					RayTracingSceneInitializationData.BaseUploadBufferOffsets,
-					RayTracingSceneInitializationData.BaseInstancePrefixSum,
-					RayTracingSceneInitializationData.InstanceGroupEntryRefs,
-					RayTracingSceneInitializationData.NumGPUInstanceGroups,
-					RayTracingSceneInitializationData.NumCPUInstanceGroups,
-					RayTracingSceneInitializationData.NumGPUInstanceDescriptors,
-					RayTracingSceneInitializationData.NumCPUInstanceDescriptors,
-					InstanceGroupUploadData.GetWriteView(),
-					InstanceUploadData.GetWriteView(),
-					TransformUploadData.GetWriteView()
-				);
-			}
-
-			FBufferRHIRef InstanceGroupUploadBuffer = InstanceGroupUploadData.Finalize();
-			FShaderResourceViewRHIRef InstanceGroupUploadSRV = RHICmdList.CreateShaderResourceView(InstanceGroupUploadBuffer, FRHIViewDesc::CreateBufferSRV().SetTypeFromBuffer(InstanceGroupUploadBuffer));
-
-			FBufferRHIRef InstanceUploadBuffer = InstanceUploadData.Finalize();
-			FShaderResourceViewRHIRef InstanceUploadSRV = RHICmdList.CreateShaderResourceView(InstanceUploadBuffer, FRHIViewDesc::CreateBufferSRV().SetTypeFromBuffer(InstanceUploadBuffer));
-
-			FBufferRHIRef TransformUploadBuffer = TransformUploadData.Finalize();
-			FShaderResourceViewRHIRef TransformUploadSRV = RHICmdList.CreateShaderResourceView(TransformUploadBuffer, FRHIViewDesc::CreateBufferSRV().SetTypeFromBuffer(TransformUploadBuffer));
-
-			for (uint32 GPUIndex : RHICmdList.GetGPUMask())
-			{
-				FRayTracingAccelerationStructureAddress* AddressesPtr = (FRayTracingAccelerationStructureAddress*)RHICmdList.LockBufferMGPU(
-					AccelerationStructureAddressesBuffer.Buffer,
-					GPUIndex,
-					0,
-					RayTracingSceneInitializationData.ReferencedGeometries.Num() * sizeof(FRayTracingAccelerationStructureAddress), RLM_WriteOnly);
-
-				const TArrayView<FRHIRayTracingGeometry*> ReferencedGeometries = RHICmdList.AllocArray(MakeConstArrayView(RayTracingSceneInitializationData.ReferencedGeometries));
-
-				RHICmdList.EnqueueLambda([AddressesPtr, ReferencedGeometries, GPUIndex](FRHICommandListBase&)
-				{
-					for (int32 GeometryIndex = 0; GeometryIndex < ReferencedGeometries.Num(); ++GeometryIndex)
-					{
-						AddressesPtr[GeometryIndex] = ReferencedGeometries[GeometryIndex]->GetAccelerationStructureAddress(GPUIndex);
-					}
-				});
-
-				RHICmdList.UnlockBufferMGPU(AccelerationStructureAddressesBuffer.Buffer, GPUIndex);
-			}
+			RayTracingInstanceBufferBuilder.FillRayTracingInstanceUploadBuffer(RHICmdList);
+			RayTracingInstanceBufferBuilder.FillAccelerationStructureAddressesBuffer(RHICmdList);
 
-			const FDFVector3 PreViewTranslation { View.ViewMatrices.GetPreViewTranslation() };
-
-			BuildRayTracingInstanceBuffer(
+			RayTracingInstanceBufferBuilder.BuildRayTracingInstanceBuffer(
 				RHICmdList,
 				nullptr,
-				PreViewTranslation,
-				InstanceBuffer.UAV,
-				InstanceGroupUploadSRV,
-				InstanceUploadSRV,
-				AccelerationStructureAddressesBuffer.SRV,
-				TransformUploadSRV,
-				RayTracingSceneInitializationData.NumNativeGPUSceneInstances,
-				RayTracingSceneInitializationData.NumNativeCPUInstances,
-				RayTracingSceneInitializationData.NumGPUInstanceGroups,
-				RayTracingSceneInitializationData.NumCPUInstanceGroups,
-				RayTracingSceneInitializationData.NumGPUInstanceDescriptors,
-				RayTracingSceneInitializationData.NumCPUInstanceDescriptors,
 				nullptr,
+				InstanceBuffer.UAV,
 				nullptr,
 				0,
 				nullptr);
@@ -1322,8 +1233,8 @@ bool FSceneRenderState::SetupRayTracingScene(FRDGBuilder& GraphBuilder, FSceneUn
 				BuildParams.ScratchBufferOffset = 0;
 				BuildParams.InstanceBuffer = InstanceBuffer.Buffer;
 				BuildParams.InstanceBufferOffset = 0;
-				BuildParams.ReferencedGeometries = RayTracingSceneInitializationData.ReferencedGeometries;
-				BuildParams.NumInstances = RayTracingSceneInitializationData.NumNativeGPUSceneInstances + RayTracingSceneInitializationData.NumNativeCPUInstances;
+				BuildParams.ReferencedGeometries = RayTracingInstanceBufferBuilder.GetReferencedGeometries();
+				BuildParams.NumInstances = RayTracingInstanceBufferBuilder.GetMaxNumInstances();
 
 				RHICmdList.BuildAccelerationStructure(BuildParams);
 			}
diff --git a/Engine/Source/Runtime/Renderer/Private/RayTracing/RayTracingInstanceBufferUtil.cpp b/Engine/Source/Runtime/Renderer/Private/RayTracing/RayTracingInstanceBufferUtil.cpp
--- a/Engine/Source/Runtime/Renderer/Private/RayTracing/RayTracingInstanceBufferUtil.cpp
+++ b/Engine/Source/Runtime/Renderer/Private/RayTracing/RayTracingInstanceBufferUtil.cpp
@@ -25,6 +25,22 @@
 
 #if RHI_RAYTRACING
 
+PRAGMA_DISABLE_DEPRECATION_WARNINGS
+
+/*
+*
+* Each FRayTracingGeometryInstance can translate to multiple native TLAS instances (see FRayTracingGeometryInstance::NumTransforms).
+*
+* The FRayTracingGeometryInstance array (ie: FRayTracingScene::Instances) used to create FRayTracingSceneRHI
+* can have mix of instances using GPUScene or CPU transforms.
+* In order to reduce the number of dispatches to build the native RayTracing Instance Buffer,
+* the upload buffer containing FRayTracingInstanceDescriptor is split in 2 sections, [GPUSceneInstances] [CPUInstances].
+* This way native GPUScene and CPU instance descriptors can be built in a single dispatch per type.
+*
+* If the ray tracing scene contains multiple layers, the instance buffer is divided into multiple subranges as expected by the RHI.
+*
+*/
+
 static TAutoConsoleVariable<bool> CVarRayTracingInstanceBufferRLE(
 	TEXT("r.RayTracing.InstanceBuffer.RLE"),
 	true,
@@ -32,6 +48,15 @@ static TAutoConsoleVariable<bool> CVarRayTracingInstanceBufferRLE(
 	ECVF_RenderThreadSafe | ECVF_Scalability
 );
 
+struct FRayTracingInstanceGroup
+{
+	uint32 BaseInstanceIndex : 30;
+	uint32 bIncrementUserDataPerInstance : 1;
+	uint32 bReuseInstance : 1;
+};
+
+static_assert(sizeof(FRayTracingInstanceGroup) == sizeof(uint32), "FRayTracingInstanceGroup is expected be same size as uint32.");
+
 static const uint32 GRayTracingInstanceGroupSize = 64;
 
 // Helper structure to assign instances to FRayTracingInstanceGroup depending on whether the primitive is compatible with RLE
@@ -250,6 +275,9 @@ void WriteInstanceDescriptor(
 	OutInstanceDescriptor = InstanceDesc;
 }
 
+
+// Helper function to fill upload buffers required by BuildRayTracingInstanceBuffer with instance descriptors
+// Transforms of CPU instances are copied to OutTransformData
 void FillRayTracingInstanceUploadBuffer(
 	FVector PreViewTranslation,
 	TConstArrayView<FRayTracingGeometryInstance> Instances,
@@ -440,6 +468,34 @@ void FillRayTracingInstanceUploadBuffer(
 		});
 }
 
+void FillRayTracingInstanceUploadBuffer(
+	FRayTracingSceneRHIRef RayTracingSceneRHI,
+	FVector PreViewTranslation,
+	TConstArrayView<FRayTracingGeometryInstance> Instances,
+	TConstArrayView<uint32> InstanceGeometryIndices,
+	TConstArrayView<uint32> BaseUploadBufferOffsets,
+	TConstArrayView<uint32> BaseInstancePrefixSum,
+	uint32 NumNativeGPUSceneInstances,
+	uint32 NumNativeCPUInstances,
+	TArrayView<FRayTracingInstanceDescriptor> OutInstanceUploadData,
+	TArrayView<FVector4f> OutTransformData)
+{
+	FillRayTracingInstanceUploadBuffer(
+		PreViewTranslation,
+		Instances,
+		InstanceGeometryIndices,
+		BaseUploadBufferOffsets,
+		BaseInstancePrefixSum,
+		{},
+		0,
+		0,
+		NumNativeGPUSceneInstances,
+		NumNativeCPUInstances,
+		{},
+		OutInstanceUploadData,
+		OutTransformData);
+}
+
 struct FRayTracingBuildInstanceBufferCS : public FGlobalShader
 {
 	DECLARE_GLOBAL_SHADER(FRayTracingBuildInstanceBufferCS);
@@ -655,4 +711,219 @@ void BuildRayTracingInstanceBuffer(
 	}
 }
 
+void BuildRayTracingInstanceBuffer(
+	FRHICommandList& RHICmdList,
+	const FGPUScene* GPUScene,
+	const FDFVector3& PreViewTranslation,
+	FRHIUnorderedAccessView* InstancesUAV,
+	FRHIShaderResourceView* InstanceUploadSRV,
+	FRHIShaderResourceView* AccelerationStructureAddressesSRV,
+	FRHIShaderResourceView* CPUInstanceTransformSRV,
+	uint32 NumNativeGPUSceneInstances,
+	uint32 NumNativeCPUInstances,
+	const FRayTracingCullingParameters* CullingParameters,
+	FRHIUnorderedAccessView* OutputStatsUAV,
+	FRHIUnorderedAccessView* DebugInstanceGPUSceneIndexUAV)
+{
+	BuildRayTracingInstanceBuffer(
+		RHICmdList,
+		GPUScene,
+		PreViewTranslation,
+		InstancesUAV,
+		nullptr,
+		InstanceUploadSRV,
+		AccelerationStructureAddressesSRV,
+		CPUInstanceTransformSRV,
+		NumNativeGPUSceneInstances,
+		NumNativeCPUInstances,
+		0,
+		0,
+		0,
+		0,
+		CullingParameters,
+		OutputStatsUAV,
+		0,
+		DebugInstanceGPUSceneIndexUAV);
+}
+
+// Round up buffer sizes to some multiple to avoid pathological growth reallocations.
+static constexpr uint32 AllocationGranularity = 8 * 1024;
+static constexpr uint64 BufferAllocationGranularity = 16 * 1024 * 1024;
+
+void FRayTracingInstanceBufferBuilder::Init(TConstArrayView<FRayTracingGeometryInstance> InInstances, FVector InPreViewTranslation)
+{
+	Instances = InInstances;
+	PreViewTranslation = InPreViewTranslation;
+
+	Data = BuildRayTracingSceneInitializationData(Instances);
+}
+
+void FRayTracingInstanceBufferBuilder::FillRayTracingInstanceUploadBuffer(FRHICommandList& RHICmdList)
+{
+	const uint32 NumInstanceGroups = Data.NumGPUInstanceGroups + Data.NumCPUInstanceGroups;
+	const uint32 NumInstanceGroupsAligned = FMath::DivideAndRoundUp(FMath::Max(NumInstanceGroups, 1U), AllocationGranularity) * AllocationGranularity;
+
+	const uint32 NumInstanceDescriptors = Data.NumGPUInstanceDescriptors + Data.NumCPUInstanceDescriptors;
+	const uint32 NumInstanceDescriptorsAligned = FMath::DivideAndRoundUp(FMath::Max(NumInstanceDescriptors, 1U), AllocationGranularity) * AllocationGranularity;
+
+	const uint32 NumTransformsAligned = FMath::DivideAndRoundUp(FMath::Max(Data.NumNativeCPUInstances, 1U), AllocationGranularity) * AllocationGranularity;
+
+	{
+		// Create/resize instance group upload buffer (if necessary)
+		const uint32 UploadBufferSize = NumInstanceGroupsAligned * sizeof(FRayTracingInstanceGroup);
+
+		if (!InstanceGroupUploadBuffer.IsValid()
+			|| UploadBufferSize > InstanceGroupUploadBuffer->GetSize()
+			|| UploadBufferSize < InstanceGroupUploadBuffer->GetSize() / 2)
+		{
+			const FRHIBufferCreateDesc CreateDesc =
+				FRHIBufferCreateDesc::CreateStructured(TEXT("FRayTracingScene::InstanceGroupUploadBuffer"), UploadBufferSize, sizeof(FRayTracingInstanceGroup))
+				.AddUsage(EBufferUsageFlags::ShaderResource | EBufferUsageFlags::Volatile)
+				.DetermineInitialState();
+
+			InstanceGroupUploadBuffer = RHICmdList.CreateBuffer(CreateDesc);
+			InstanceGroupUploadSRV = RHICmdList.CreateShaderResourceView(InstanceGroupUploadBuffer, FRHIViewDesc::CreateBufferSRV().SetTypeFromBuffer(InstanceGroupUploadBuffer));
+		}
+	}
+
+	{
+		// Create/resize instance upload buffer (if necessary)
+		const uint32 UploadBufferSize = NumInstanceDescriptorsAligned * sizeof(FRayTracingInstanceDescriptor);
+
+		if (!InstanceUploadBuffer.IsValid()
+			|| UploadBufferSize > InstanceUploadBuffer->GetSize()
+			|| UploadBufferSize < InstanceUploadBuffer->GetSize() / 2)
+		{
+			const FRHIBufferCreateDesc CreateDesc =
+				FRHIBufferCreateDesc::CreateStructured(TEXT("FRayTracingScene::InstanceUploadBuffer"), UploadBufferSize, sizeof(FRayTracingInstanceDescriptor))
+				.AddUsage(EBufferUsageFlags::ShaderResource | EBufferUsageFlags::Volatile)
+				.DetermineInitialState();
+
+			InstanceUploadBuffer = RHICmdList.CreateBuffer(CreateDesc);
+			InstanceUploadSRV = RHICmdList.CreateShaderResourceView(InstanceUploadBuffer, FRHIViewDesc::CreateBufferSRV().SetTypeFromBuffer(InstanceUploadBuffer));
+		}
+	}
+
+	{
+		const uint32 UploadBufferSize = NumTransformsAligned * sizeof(FVector4f) * 3;
+
+		// Create/resize transform upload buffer (if necessary)
+		if (!TransformUploadBuffer.IsValid()
+			|| UploadBufferSize > TransformUploadBuffer->GetSize()
+			|| UploadBufferSize < TransformUploadBuffer->GetSize() / 2)
+		{
+			const FRHIBufferCreateDesc CreateDesc =
+				FRHIBufferCreateDesc::CreateStructured(TEXT("FRayTracingScene::TransformUploadBuffer"), UploadBufferSize, sizeof(FVector4f))
+				.AddUsage(EBufferUsageFlags::ShaderResource | EBufferUsageFlags::Volatile)
+				.DetermineInitialState();
+
+			TransformUploadBuffer = RHICmdList.CreateBuffer(CreateDesc);
+			TransformUploadSRV = RHICmdList.CreateShaderResourceView(TransformUploadBuffer, FRHIViewDesc::CreateBufferSRV().SetTypeFromBuffer(TransformUploadBuffer));
+		}
+	}
+
+	const uint32 InstanceGroupUploadBytes = NumInstanceGroups * sizeof(FRayTracingInstanceGroup);
+	const uint32 InstanceUploadBytes = NumInstanceDescriptors * sizeof(FRayTracingInstanceDescriptor);
+	const uint32 TransformUploadBytes = Data.NumNativeCPUInstances * 3 * sizeof(FVector4f);
+
+	FRayTracingInstanceGroup* InstanceGroupUploadData = (FRayTracingInstanceGroup*)RHICmdList.LockBuffer(InstanceGroupUploadBuffer, 0, InstanceGroupUploadBytes, RLM_WriteOnly);
+	FRayTracingInstanceDescriptor* InstanceUploadData = (FRayTracingInstanceDescriptor*)RHICmdList.LockBuffer(InstanceUploadBuffer, 0, InstanceUploadBytes, RLM_WriteOnly);
+	FVector4f* TransformUploadData = (Data.NumNativeCPUInstances > 0) ? (FVector4f*)RHICmdList.LockBuffer(TransformUploadBuffer, 0, TransformUploadBytes, RLM_WriteOnly) : nullptr;
+
+	::FillRayTracingInstanceUploadBuffer(
+		PreViewTranslation,
+		Instances,
+		Data.InstanceGeometryIndices,
+		Data.BaseUploadBufferOffsets,
+		Data.BaseInstancePrefixSum,
+		Data.InstanceGroupEntryRefs,
+		Data.NumGPUInstanceGroups,
+		Data.NumCPUInstanceGroups,
+		Data.NumGPUInstanceDescriptors,
+		Data.NumCPUInstanceDescriptors,
+		MakeArrayView(InstanceGroupUploadData, NumInstanceGroups),
+		MakeArrayView(InstanceUploadData, NumInstanceDescriptors),
+		MakeArrayView(TransformUploadData, Data.NumNativeCPUInstances * 3));
+
+	RHICmdList.UnlockBuffer(InstanceGroupUploadBuffer);
+	RHICmdList.UnlockBuffer(InstanceUploadBuffer);
+
+	if (Data.NumNativeCPUInstances > 0)
+	{
+		RHICmdList.UnlockBuffer(TransformUploadBuffer);
+	}
+}
+
+void FRayTracingInstanceBufferBuilder::FillAccelerationStructureAddressesBuffer(FRHICommandList& RHICmdList)
+{
+	const uint32 NumGeometries = FMath::RoundUpToPowerOfTwo(Data.ReferencedGeometries.Num());
+
+	{
+		// Round to PoT to avoid resizing too often
+		const uint32 NumGeometriesAligned = FMath::RoundUpToPowerOfTwo(NumGeometries);
+		const uint32 AccelerationStructureAddressesBufferSize = NumGeometriesAligned * sizeof(FRayTracingAccelerationStructureAddress);
+
+		if (AccelerationStructureAddressesBuffer.NumBytes < AccelerationStructureAddressesBufferSize)
+		{
+			// Need to pass "BUF_MultiGPUAllocate", as virtual addresses are different per GPU
+			AccelerationStructureAddressesBuffer.Initialize(RHICmdList, TEXT("FRayTracingScene::AccelerationStructureAddressesBuffer"), AccelerationStructureAddressesBufferSize, BUF_Volatile | BUF_MultiGPUAllocate);
+		}
+	}
+
+	for (uint32 GPUIndex : RHICmdList.GetGPUMask())
+	{
+		FRayTracingAccelerationStructureAddress* AddressesPtr = (FRayTracingAccelerationStructureAddress*)RHICmdList.LockBufferMGPU(
+			AccelerationStructureAddressesBuffer.Buffer,
+			GPUIndex,
+			0,
+			NumGeometries * sizeof(FRayTracingAccelerationStructureAddress), RLM_WriteOnly);
+
+		const TArrayView<FRHIRayTracingGeometry*> ReferencedGeometries = RHICmdList.AllocArray(MakeConstArrayView(Data.ReferencedGeometries));
+
+		RHICmdList.EnqueueLambda([AddressesPtr, ReferencedGeometries, GPUIndex](FRHICommandListBase&)
+			{
+				TRACE_CPUPROFILER_EVENT_SCOPE(GetAccelerationStructuresAddresses);
+
+				for (int32 GeometryIndex = 0; GeometryIndex < ReferencedGeometries.Num(); ++GeometryIndex)
+				{
+					AddressesPtr[GeometryIndex] = ReferencedGeometries[GeometryIndex]->GetAccelerationStructureAddress(GPUIndex);
+				}
+			});
+
+		RHICmdList.UnlockBufferMGPU(AccelerationStructureAddressesBuffer.Buffer, GPUIndex);
+	}
+}
+
+void FRayTracingInstanceBufferBuilder::BuildRayTracingInstanceBuffer(
+	FRHICommandList& RHICmdList,
+	const FGPUScene* GPUScene,
+	const FRayTracingCullingParameters* CullingParameters,
+	FRHIUnorderedAccessView* InstancesUAV,
+	FRHIUnorderedAccessView* OutputStatsUAV,
+	uint32 OutputStatsOffset,
+	FRHIUnorderedAccessView* DebugInstanceGPUSceneIndexUAV)
+{
+	::BuildRayTracingInstanceBuffer(
+		RHICmdList,
+		GPUScene,
+		FDFVector3(PreViewTranslation),
+		InstancesUAV,
+		InstanceGroupUploadSRV,
+		InstanceUploadSRV,
+		AccelerationStructureAddressesBuffer.SRV,
+		TransformUploadSRV,
+		Data.NumNativeGPUSceneInstances,
+		Data.NumNativeCPUInstances,
+		Data.NumGPUInstanceGroups,
+		Data.NumCPUInstanceGroups,
+		Data.NumGPUInstanceDescriptors,
+		Data.NumCPUInstanceDescriptors,
+		CullingParameters,
+		OutputStatsUAV,
+		OutputStatsOffset,
+		DebugInstanceGPUSceneIndexUAV);
+}
+
+PRAGMA_ENABLE_DEPRECATION_WARNINGS
+
 #endif //RHI_RAYTRACING
diff --git a/Engine/Source/Runtime/Renderer/Private/RayTracing/RayTracingScene.cpp b/Engine/Source/Runtime/Renderer/Private/RayTracing/RayTracingScene.cpp
--- a/Engine/Source/Runtime/Renderer/Private/RayTracing/RayTracingScene.cpp
+++ b/Engine/Source/Runtime/Renderer/Private/RayTracing/RayTracingScene.cpp
@@ -82,7 +82,7 @@ void FRayTracingScene::BuildInitializationData()
 	{
 		FLayer& Layer = Layers[LayerIndex];
 
-		Layer.InitializationData = BuildRayTracingSceneInitializationData(Layer.Instances);
+		Layer.InstanceBufferBuilder.Init(Layer.Instances, PreViewTranslation);
 	}
 
 	bInitializationDataBuilt = true;
@@ -133,25 +133,17 @@ void FRayTracingScene::Create(FRDGBuilder& GraphBuilder, const FViewInfo& View,
 	{
 		FLayer& Layer = Layers[LayerIndex];
 
+		const uint32 NumNativeInstances = Layer.InstanceBufferBuilder.GetMaxNumInstances();
+		const uint32 NumNativeInstancesAligned = FMath::DivideAndRoundUp(FMath::Max(NumNativeInstances, 1U), AllocationGranularity) * AllocationGranularity;
+
 		{
 			FRayTracingSceneInitializer Initializer;
 			Initializer.DebugName = FName(TEXT("FRayTracingScene"));
-			Initializer.MaxNumInstances = Layer.InitializationData.NumNativeGPUSceneInstances + Layer.InitializationData.NumNativeCPUInstances;
+			Initializer.MaxNumInstances = NumNativeInstances;
 			Initializer.BuildFlags = LayerIndex == uint8(ERayTracingSceneLayer::FarField) ? FarFieldBuildFlags : BuildFlags;
 			Layer.RayTracingSceneRHI = RHICreateRayTracingScene(MoveTemp(Initializer));
 		}
 
-		const uint32 NumInstanceGroups = Layer.InitializationData.NumGPUInstanceGroups + Layer.InitializationData.NumCPUInstanceGroups;
-		const uint32 NumInstanceGroupsAligned = FMath::DivideAndRoundUp(FMath::Max(NumInstanceGroups, 1U), AllocationGranularity) * AllocationGranularity;
-
-		const uint32 NumInstanceDescriptors = Layer.InitializationData.NumGPUInstanceDescriptors + Layer.InitializationData.NumCPUInstanceDescriptors;
-		const uint32 NumInstanceDescriptorsAligned = FMath::DivideAndRoundUp(FMath::Max(NumInstanceDescriptors, 1U), AllocationGranularity) * AllocationGranularity;
-
-		const uint32 NumNativeInstances = Layer.InitializationData.NumNativeGPUSceneInstances + Layer.InitializationData.NumNativeCPUInstances;
-		const uint32 NumNativeInstancesAligned = FMath::DivideAndRoundUp(FMath::Max(NumNativeInstances, 1U), AllocationGranularity) * AllocationGranularity;
-
-		const uint32 NumTransformsAligned = FMath::DivideAndRoundUp(FMath::Max(Layer.InitializationData.NumNativeCPUInstances, 1U), AllocationGranularity) * AllocationGranularity;
-
 		FRayTracingAccelerationStructureSize SizeInfo = Layer.RayTracingSceneRHI->GetSizeInfo();
 		SizeInfo.ResultSize = FMath::DivideAndRoundUp(FMath::Max(SizeInfo.ResultSize, 1ull), BufferAllocationGranularity) * BufferAllocationGranularity;
 
@@ -188,109 +180,40 @@ void FRayTracingScene::Create(FRDGBuilder& GraphBuilder, const FViewInfo& View,
 			Layer.InstanceBuffer = GraphBuilder.CreateBuffer(InstanceBufferDesc, TEXT("FRayTracingScene::InstanceBuffer"));
 		}
 
-		{
-			// Round to PoT to avoid resizing too often
-			const uint32 NumGeometries = FMath::RoundUpToPowerOfTwo(Layer.InitializationData.ReferencedGeometries.Num());
-			const uint32 AccelerationStructureAddressesBufferSize = NumGeometries * sizeof(FRayTracingAccelerationStructureAddress);
-
-			if (Layer.AccelerationStructureAddressesBuffer.NumBytes < AccelerationStructureAddressesBufferSize)
+		// Feedback
+		if (bTracingFeedbackEnabled)
+		{				
 			{
-				// Need to pass "BUF_MultiGPUAllocate", as virtual addresses are different per GPU
-				Layer.AccelerationStructureAddressesBuffer.Initialize(
-					GraphBuilder.RHICmdList, TEXT("FRayTracingScene::AccelerationStructureAddressesBuffer"), AccelerationStructureAddressesBufferSize, BUF_Volatile | BUF_MultiGPUAllocate);
-			}
-
-			// Feedback
-			if (bTracingFeedbackEnabled)
-			{				
-				{
-					FRDGBufferDesc InstanceHitCountBufferDesc;
-					InstanceHitCountBufferDesc.Usage = EBufferUsageFlags::UnorderedAccess | EBufferUsageFlags::ShaderResource | EBufferUsageFlags::StructuredBuffer;
-					InstanceHitCountBufferDesc.BytesPerElement = sizeof(uint32);
-					InstanceHitCountBufferDesc.NumElements = NumNativeInstancesAligned;
-
-					Layer.InstanceHitCountBuffer = GraphBuilder.CreateBuffer(InstanceHitCountBufferDesc, TEXT("FRayTracingScene::InstanceHitCount"));
-					Layer.InstanceHitCountBufferUAV = GraphBuilder.CreateUAV(Layer.InstanceHitCountBuffer, ERDGUnorderedAccessViewFlags::SkipBarrier);
-					AddClearUAVPass(GraphBuilder, Layer.InstanceHitCountBufferUAV, 0, ComputePassFlags);
-				}
-
-				{
-					FRDGBufferDesc AccelerationStructureIndexBufferDesc;
-					AccelerationStructureIndexBufferDesc.Usage = EBufferUsageFlags::ShaderResource | EBufferUsageFlags::StructuredBuffer;
-					AccelerationStructureIndexBufferDesc.BytesPerElement = sizeof(uint32);
-					AccelerationStructureIndexBufferDesc.NumElements = NumNativeInstancesAligned;
-
-					Layer.AccelerationStructureIndexBuffer = GraphBuilder.CreateBuffer(AccelerationStructureIndexBufferDesc, TEXT("FRayTracingScene::AccelerationStructureIndexBuffer"));
-					GraphBuilder.QueueBufferUpload(Layer.AccelerationStructureIndexBuffer,
-						Layer.InitializationData.InstanceGeometryIndices.GetData(),
-						Layer.InitializationData.InstanceGeometryIndices.GetTypeSize() * Layer.InitializationData.InstanceGeometryIndices.Num());
-				}
-
-				{
-					FRDGBufferDesc GeometryHandleBufferDesc;
-					GeometryHandleBufferDesc.Usage = EBufferUsageFlags::ShaderResource | EBufferUsageFlags::StructuredBuffer;
-					GeometryHandleBufferDesc.BytesPerElement = sizeof(int32);
-					GeometryHandleBufferDesc.NumElements = FMath::Max(Layer.GeometryHandles.Num(), 1);
-
-					Layer.GeometryHandleBuffer = GraphBuilder.CreateBuffer(GeometryHandleBufferDesc, TEXT("FRayTracingScene::GeometryHandleBuffer"));
-					GraphBuilder.QueueBufferUpload(Layer.GeometryHandleBuffer,
-						Layer.GeometryHandles.GetData(),
-						Layer.GeometryHandles.GetTypeSize() * Layer.GeometryHandles.Num());
-				}
+				FRDGBufferDesc InstanceHitCountBufferDesc;
+				InstanceHitCountBufferDesc.Usage = EBufferUsageFlags::UnorderedAccess | EBufferUsageFlags::ShaderResource | EBufferUsageFlags::StructuredBuffer;
+				InstanceHitCountBufferDesc.BytesPerElement = sizeof(uint32);
+				InstanceHitCountBufferDesc.NumElements = NumNativeInstancesAligned;
+
+				Layer.InstanceHitCountBuffer = GraphBuilder.CreateBuffer(InstanceHitCountBufferDesc, TEXT("FRayTracingScene::InstanceHitCount"));
+				Layer.InstanceHitCountBufferUAV = GraphBuilder.CreateUAV(Layer.InstanceHitCountBuffer, ERDGUnorderedAccessViewFlags::SkipBarrier);
+				AddClearUAVPass(GraphBuilder, Layer.InstanceHitCountBufferUAV, 0, ComputePassFlags);
 			}
-		}
 
-		{
-			// Create/resize instance group upload buffer (if necessary)
-			const uint32 UploadBufferSize = NumInstanceGroupsAligned * sizeof(FRayTracingInstanceGroup);
-
-			if (!Layer.InstanceGroupUploadBuffer.IsValid()
-				|| UploadBufferSize > Layer.InstanceGroupUploadBuffer->GetSize()
-				|| UploadBufferSize < Layer.InstanceGroupUploadBuffer->GetSize() / 2)
 			{
-				const FRHIBufferCreateDesc CreateDesc =
-					FRHIBufferCreateDesc::CreateStructured(TEXT("FRayTracingScene::InstanceGroupUploadBuffer"), UploadBufferSize, sizeof(FRayTracingInstanceGroup))
-					.AddUsage(EBufferUsageFlags::ShaderResource | EBufferUsageFlags::Volatile)
-					.DetermineInitialState();
-
-				Layer.InstanceGroupUploadBuffer = RHICmdList.CreateBuffer(CreateDesc);
-				Layer.InstanceGroupUploadSRV = RHICmdList.CreateShaderResourceView(Layer.InstanceGroupUploadBuffer, FRHIViewDesc::CreateBufferSRV().SetTypeFromBuffer(Layer.InstanceGroupUploadBuffer));
-			}
-		}
+				FRDGBufferDesc AccelerationStructureIndexBufferDesc;
+				AccelerationStructureIndexBufferDesc.Usage = EBufferUsageFlags::ShaderResource | EBufferUsageFlags::StructuredBuffer;
+				AccelerationStructureIndexBufferDesc.BytesPerElement = sizeof(uint32);
+				AccelerationStructureIndexBufferDesc.NumElements = NumNativeInstancesAligned;
 
-		{
-			// Create/resize instance upload buffer (if necessary)
-			const uint32 UploadBufferSize = NumInstanceDescriptorsAligned * sizeof(FRayTracingInstanceDescriptor);
-
-			if (!Layer.InstanceUploadBuffer.IsValid()
-				|| UploadBufferSize > Layer.InstanceUploadBuffer->GetSize()
-				|| UploadBufferSize < Layer.InstanceUploadBuffer->GetSize() / 2)
-			{
-				const FRHIBufferCreateDesc CreateDesc =
-					FRHIBufferCreateDesc::CreateStructured(TEXT("FRayTracingScene::InstanceUploadBuffer"), UploadBufferSize, sizeof(FRayTracingInstanceDescriptor))
-					.AddUsage(EBufferUsageFlags::ShaderResource | EBufferUsageFlags::Volatile)
-					.DetermineInitialState();
+				Layer.AccelerationStructureIndexBuffer = GraphBuilder.CreateBuffer(AccelerationStructureIndexBufferDesc, TEXT("FRayTracingScene::AccelerationStructureIndexBuffer"));
 
-				Layer.InstanceUploadBuffer = RHICmdList.CreateBuffer(CreateDesc);
-				Layer.InstanceUploadSRV = RHICmdList.CreateShaderResourceView(Layer.InstanceUploadBuffer, FRHIViewDesc::CreateBufferSRV().SetTypeFromBuffer(Layer.InstanceUploadBuffer));
+				TConstArrayView<uint32> InstanceGeometryIndices = Layer.InstanceBufferBuilder.GetInstanceGeometryIndices();
+				GraphBuilder.QueueBufferUpload(Layer.AccelerationStructureIndexBuffer, InstanceGeometryIndices.GetData(), InstanceGeometryIndices.GetTypeSize() * InstanceGeometryIndices.Num());
 			}
-		}
-
-		{
-			const uint32 UploadBufferSize = NumTransformsAligned * sizeof(FVector4f) * 3;
 
-			// Create/resize transform upload buffer (if necessary)
-			if (!Layer.TransformUploadBuffer.IsValid()
-				|| UploadBufferSize > Layer.TransformUploadBuffer->GetSize()
-				|| UploadBufferSize < Layer.TransformUploadBuffer->GetSize() / 2)
 			{
-				const FRHIBufferCreateDesc CreateDesc =
-					FRHIBufferCreateDesc::CreateStructured(TEXT("FRayTracingScene::TransformUploadBuffer"), UploadBufferSize, sizeof(FVector4f))
-					.AddUsage(EBufferUsageFlags::ShaderResource | EBufferUsageFlags::Volatile)
-					.DetermineInitialState();
+				FRDGBufferDesc GeometryHandleBufferDesc;
+				GeometryHandleBufferDesc.Usage = EBufferUsageFlags::ShaderResource | EBufferUsageFlags::StructuredBuffer;
+				GeometryHandleBufferDesc.BytesPerElement = sizeof(int32);
+				GeometryHandleBufferDesc.NumElements = FMath::Max(Layer.GeometryHandles.Num(), 1);
 
-				Layer.TransformUploadBuffer = RHICmdList.CreateBuffer(CreateDesc);
-				Layer.TransformUploadSRV = RHICmdList.CreateShaderResourceView(Layer.TransformUploadBuffer, FRHIViewDesc::CreateBufferSRV().SetTypeFromBuffer(Layer.TransformUploadBuffer));
+				Layer.GeometryHandleBuffer = GraphBuilder.CreateBuffer(GeometryHandleBufferDesc, TEXT("FRayTracingScene::GeometryHandleBuffer"));
+				GraphBuilder.QueueBufferUpload(Layer.GeometryHandleBuffer, Layer.GeometryHandles.GetData(), Layer.GeometryHandles.GetTypeSize() * Layer.GeometryHandles.Num());
 			}
 		}
 
@@ -333,7 +256,7 @@ void FRayTracingScene::Create(FRDGBuilder& GraphBuilder, const FViewInfo& View,
 
 					TConstArrayView<FRayTracingGeometryInstance> TaskInstancesData(Layer.Instances.GetData() + TaskFirstItemIndex, TaskNumItems);
 					TConstArrayView<FRayTracingInstanceDebugData> TaskInstancesDebugData(Layer.InstancesDebugData.GetData() + TaskFirstItemIndex, TaskNumItems);
-					TConstArrayView<uint32> TaskBaseInstancePrefixSum(Layer.InitializationData.BaseInstancePrefixSum.GetData() + TaskFirstItemIndex, TaskNumItems);
+					TConstArrayView<uint32> TaskBaseInstancePrefixSum(Layer.InstanceBufferBuilder.GetBaseInstancePrefixSum().GetData() + TaskFirstItemIndex, TaskNumItems);
 
 					GraphBuilder.AddSetupTask([UploadData, TaskInstancesDebugData, TaskInstancesData, TaskBaseInstancePrefixSum]()
 						{
@@ -360,85 +283,18 @@ void FRayTracingScene::Create(FRDGBuilder& GraphBuilder, const FViewInfo& View,
 		if (NumNativeInstances > 0)
 		{
 			// Fill instance upload buffer on separate thread since results are only needed in RHI thread
-			GraphBuilder.AddCommandListSetupTask(
-				[NumInstanceDescriptors,
-				NumGPUInstanceGroups = Layer.InitializationData.NumGPUInstanceGroups,
-				NumCPUInstanceGroups = Layer.InitializationData.NumCPUInstanceGroups,
-				NumNativeGPUSceneInstances = Layer.InitializationData.NumNativeGPUSceneInstances,
-				NumNativeCPUInstances = Layer.InitializationData.NumNativeCPUInstances,
-				NumGPUInstanceDescriptors = Layer.InitializationData.NumGPUInstanceDescriptors,
-				NumCPUInstanceDescriptors = Layer.InitializationData.NumCPUInstanceDescriptors,
-				Instances = MakeArrayView(Layer.Instances),
-				InstanceGeometryIndices = MakeArrayView(Layer.InitializationData.InstanceGeometryIndices),
-				BaseUploadBufferOffsets = MakeArrayView(Layer.InitializationData.BaseUploadBufferOffsets),
-				BaseInstancePrefixSum = MakeArrayView(Layer.InitializationData.BaseInstancePrefixSum),
-				InstanceGroupEntryRefs = MakeArrayView(Layer.InitializationData.InstanceGroupEntryRefs),
-				RayTracingSceneRHI = Layer.RayTracingSceneRHI,
-				InstanceGroupUploadBuffer = Layer.InstanceGroupUploadBuffer,
-				InstanceUploadBuffer = Layer.InstanceUploadBuffer,
-				TransformUploadBuffer = Layer.TransformUploadBuffer,
-				PreViewTranslation = this->PreViewTranslation](FRHICommandList& RHICmdList)
+			GraphBuilder.AddCommandListSetupTask([&InstanceBufferBuilder = Layer.InstanceBufferBuilder](FRHICommandList& RHICmdList)
 				{
 					FTaskTagScope TaskTagScope(ETaskTag::EParallelRenderingThread);
 
-					const uint32 NumInstanceGroups = NumGPUInstanceGroups + NumCPUInstanceGroups;
-
-					const uint32 InstanceGroupUploadBytes = NumInstanceGroups * sizeof(FRayTracingInstanceGroup);
-					const uint32 InstanceUploadBytes = NumInstanceDescriptors * sizeof(FRayTracingInstanceDescriptor);
-					const uint32 TransformUploadBytes = NumNativeCPUInstances * 3 * sizeof(FVector4f);
-
-					FRayTracingInstanceGroup* InstanceGroupUploadData = (FRayTracingInstanceGroup*)RHICmdList.LockBuffer(InstanceGroupUploadBuffer, 0, InstanceGroupUploadBytes, RLM_WriteOnly);
-					FRayTracingInstanceDescriptor* InstanceUploadData = (FRayTracingInstanceDescriptor*)RHICmdList.LockBuffer(InstanceUploadBuffer, 0, InstanceUploadBytes, RLM_WriteOnly);
-					FVector4f* TransformUploadData = (NumNativeCPUInstances > 0) ? (FVector4f*)RHICmdList.LockBuffer(TransformUploadBuffer, 0, TransformUploadBytes, RLM_WriteOnly) : nullptr;
-
-					FillRayTracingInstanceUploadBuffer(
-						PreViewTranslation,
-						Instances,
-						InstanceGeometryIndices,
-						BaseUploadBufferOffsets,
-						BaseInstancePrefixSum,
-						InstanceGroupEntryRefs,
-						NumGPUInstanceGroups,
-						NumCPUInstanceGroups,
-						NumGPUInstanceDescriptors,
-						NumCPUInstanceDescriptors,
-						MakeArrayView(InstanceGroupUploadData, NumInstanceGroups),
-						MakeArrayView(InstanceUploadData, NumInstanceDescriptors),
-						MakeArrayView(TransformUploadData, NumNativeCPUInstances * 3));
-
-					RHICmdList.UnlockBuffer(InstanceGroupUploadBuffer);
-					RHICmdList.UnlockBuffer(InstanceUploadBuffer);
-
-					if (NumNativeCPUInstances > 0)
-					{
-						RHICmdList.UnlockBuffer(TransformUploadBuffer);
-					}
+					InstanceBufferBuilder.FillRayTracingInstanceUploadBuffer(RHICmdList);
 				});
 
-			GraphBuilder.AddCommandListSetupTask([&Layer](FRHICommandList& RHICmdList)
+			GraphBuilder.AddCommandListSetupTask([&InstanceBufferBuilder = Layer.InstanceBufferBuilder](FRHICommandList& RHICmdList)
 				{
-					for (uint32 GPUIndex : RHICmdList.GetGPUMask())
-					{
-						FRayTracingAccelerationStructureAddress* AddressesPtr = (FRayTracingAccelerationStructureAddress*)RHICmdList.LockBufferMGPU(
-							Layer.AccelerationStructureAddressesBuffer.Buffer,
-							GPUIndex,
-							0,
-							Layer.InitializationData.ReferencedGeometries.Num() * sizeof(FRayTracingAccelerationStructureAddress), RLM_WriteOnly);
-
-						const TArrayView<FRHIRayTracingGeometry*> ReferencedGeometries = RHICmdList.AllocArray(MakeConstArrayView(Layer.InitializationData.ReferencedGeometries));
-
-						RHICmdList.EnqueueLambda([AddressesPtr, ReferencedGeometries, GPUIndex](FRHICommandListBase&)
-							{
-								TRACE_CPUPROFILER_EVENT_SCOPE(GetAccelerationStructuresAddresses);
-
-								for (int32 GeometryIndex = 0; GeometryIndex < ReferencedGeometries.Num(); ++GeometryIndex)
-								{
-									AddressesPtr[GeometryIndex] = ReferencedGeometries[GeometryIndex]->GetAccelerationStructureAddress(GPUIndex);
-								}
-							});
+					FTaskTagScope TaskTagScope(ETaskTag::EParallelRenderingThread);
 
-						RHICmdList.UnlockBufferMGPU(Layer.AccelerationStructureAddressesBuffer.Buffer, GPUIndex);
-					}
+					InstanceBufferBuilder.FillAccelerationStructureAddressesBuffer(RHICmdList);
 				});
 
 			FBuildInstanceBufferPassParams* PassParams = GraphBuilder.AllocParameters<FBuildInstanceBufferPassParams>();
@@ -452,35 +308,17 @@ void FRayTracingScene::Create(FRDGBuilder& GraphBuilder, const FViewInfo& View,
 				PassParams,
 				ComputePassFlags,
 				[PassParams,
-				this,
-				&Layer,
+				&InstanceBufferBuilder = Layer.InstanceBufferBuilder,
 				LayerIndex,
 				GPUScene,
-				NumNativeGPUSceneInstances = Layer.InitializationData.NumNativeGPUSceneInstances,
-				NumNativeCPUInstances = Layer.InitializationData.NumNativeCPUInstances,
-				NumGPUInstanceGroups = Layer.InitializationData.NumGPUInstanceGroups,
-				NumCPUInstanceGroups = Layer.InitializationData.NumCPUInstanceGroups,
-				NumGPUInstanceDescriptors = Layer.InitializationData.NumGPUInstanceDescriptors,
-				NumCPUInstanceDescriptors = Layer.InitializationData.NumCPUInstanceDescriptors,
-				CullingParameters = View.RayTracingCullingParameters
+				CullingParameters = &View.RayTracingCullingParameters
 				](FRHICommandList& RHICmdList)
 				{
-					BuildRayTracingInstanceBuffer(
+					InstanceBufferBuilder.BuildRayTracingInstanceBuffer(
 						RHICmdList,
 						GPUScene,
-						FDFVector3(PreViewTranslation),
+						CullingParameters,
 						PassParams->InstanceBuffer->GetRHI(),
-						Layer.InstanceGroupUploadSRV,
-						Layer.InstanceUploadSRV,
-						Layer.AccelerationStructureAddressesBuffer.SRV,
-						Layer.TransformUploadSRV,
-						NumNativeGPUSceneInstances,
-						NumNativeCPUInstances,
-						NumGPUInstanceGroups,
-						NumCPUInstanceGroups,
-						NumGPUInstanceDescriptors,
-						NumCPUInstanceDescriptors,
-						&CullingParameters,
 						PassParams->OutputStats ? PassParams->OutputStats->GetRHI() : nullptr,
 						LayerIndex,
 						PassParams->DebugInstanceGPUSceneIndexBuffer ? PassParams->DebugInstanceGPUSceneIndexBuffer->GetRHI() : nullptr);
@@ -520,8 +358,8 @@ void FRayTracingScene::Build(FRDGBuilder& GraphBuilder, ERDGPassFlags ComputePas
 				BuildParams.ScratchBufferOffset = 0;
 				BuildParams.InstanceBuffer = PassParams->InstanceBuffer->GetRHI();
 				BuildParams.InstanceBufferOffset = 0;
-				BuildParams.NumInstances = Layer.InitializationData.NumNativeCPUInstances + Layer.InitializationData.NumNativeGPUSceneInstances;
-				BuildParams.ReferencedGeometries = Layer.InitializationData.ReferencedGeometries;
+				BuildParams.NumInstances = Layer.InstanceBufferBuilder.GetMaxNumInstances();
+				BuildParams.ReferencedGeometries = Layer.InstanceBufferBuilder.GetReferencedGeometries();
 
 				RHICmdList.BindAccelerationStructureMemory(Layer.RayTracingSceneRHI, PassParams->TLASBuffer->GetRHI(), 0);
 				RHICmdList.BuildAccelerationStructure(BuildParams);
@@ -607,8 +445,8 @@ void FRayTracingScene::FinishTracingFeedback(FRDGBuilder& GraphBuilder, ERDGPass
 	RDG_EVENT_SCOPE(GraphBuilder, "RayTracingScene::FinishTracingFeedback");
 	
 	const FLayer& Layer = Layers[0];
-	const uint32 NumGeometries = (uint32)Layer.InitializationData.ReferencedGeometries.Num();
-	const uint32 NumInstances = Layer.InitializationData.NumNativeCPUInstances + Layer.InitializationData.NumNativeGPUSceneInstances;
+	const uint32 NumGeometries = (uint32)Layer.InstanceBufferBuilder.GetReferencedGeometries().Num();
+	const uint32 NumInstances = Layer.InstanceBufferBuilder.GetMaxNumInstances();
 
 	if (NumGeometries == 0 || !bTracingFeedbackEnabled)
 	{
@@ -795,7 +633,7 @@ void FRayTracingScene::FinishStats(FRDGBuilder& GraphBuilder, ERDGPassFlags Comp
 			{
 				FLayer& Layer = Layers[LayerIndex];
 
-				const uint32 LayerNumNativeInstances = Layer.InitializationData.NumNativeGPUSceneInstances + Layer.InitializationData.NumNativeCPUInstances;
+				const uint32 LayerNumNativeInstances = Layer.InstanceBufferBuilder.GetMaxNumInstances();
 
 				Layer.NumActiveInstances = FMath::Min(ReadbackBufferPtr[LayerIndex], LayerNumNativeInstances);
 
@@ -860,7 +698,7 @@ uint32 FRayTracingScene::GetNumNativeInstances(ERayTracingSceneLayer InLayer) co
 {
 	const FLayer& Layer = Layers[uint8(InLayer)];
 	checkf(bInitializationDataBuilt, TEXT("Must call BuildInitializationData() or Create() before using GetNumNativeInstances()."));
-	return Layer.InitializationData.NumNativeCPUInstances + Layer.InitializationData.NumNativeGPUSceneInstances;
+	return Layer.InstanceBufferBuilder.GetMaxNumInstances();
 }
 
 FRayTracingScene::FInstanceHandle FRayTracingScene::AddInstance(FRayTracingGeometryInstance Instance, ERayTracingSceneLayer InLayer, const FPrimitiveSceneProxy* Proxy, bool bDynamic, int32 GeometryHandle)
diff --git a/Engine/Source/Runtime/Renderer/Private/RayTracing/RayTracingScene.h b/Engine/Source/Runtime/Renderer/Private/RayTracing/RayTracingScene.h
--- a/Engine/Source/Runtime/Renderer/Private/RayTracing/RayTracingScene.h
+++ b/Engine/Source/Runtime/Renderer/Private/RayTracing/RayTracingScene.h
@@ -180,7 +180,7 @@ private:
 
 	struct FLayer
 	{
-		FRayTracingSceneInitializationData InitializationData;
+		FRayTracingInstanceBufferBuilder InstanceBufferBuilder;
 
 		FRayTracingSceneRHIRef RayTracingSceneRHI;
 
@@ -198,17 +198,6 @@ private:
 		FRDGBufferRef RayTracingSceneBufferRDG;
 		FRDGBufferSRVRef RayTracingSceneBufferSRV;
 
-		FBufferRHIRef InstanceGroupUploadBuffer;
-		FShaderResourceViewRHIRef InstanceGroupUploadSRV;
-
-		FBufferRHIRef InstanceUploadBuffer;
-		FShaderResourceViewRHIRef InstanceUploadSRV;
-
-		FBufferRHIRef TransformUploadBuffer;
-		FShaderResourceViewRHIRef TransformUploadSRV;
-
-		FByteAddressBuffer AccelerationStructureAddressesBuffer;
-
 		// Special data for debugging purposes
 		FRDGBufferRef InstanceDebugBuffer = nullptr;
 		FRDGBufferRef DebugInstanceGPUSceneIndexBuffer = nullptr;
diff --git a/Engine/Source/Runtime/Renderer/Private/Tests/RayTracingTestbed.cpp b/Engine/Source/Runtime/Renderer/Private/Tests/RayTracingTestbed.cpp
--- a/Engine/Source/Runtime/Renderer/Private/Tests/RayTracingTestbed.cpp
+++ b/Engine/Source/Runtime/Renderer/Private/Tests/RayTracingTestbed.cpp
@@ -175,15 +175,15 @@ bool RunRayTracingTestbed_RenderThread(const FString& Parameters)
 	Instances[0].Transforms = MakeArrayView(&FMatrix::Identity, 1);
 	Instances[0].InstanceContributionToHitGroupIndex = 0;
 
-	const FRayTracingSceneInitializationData RayTracingSceneInitializationData = BuildRayTracingSceneInitializationData(Instances);
-	const ERayTracingAccelerationStructureFlags SceneBuildFlags = ERayTracingAccelerationStructureFlags::FastTrace;
+	FRayTracingInstanceBufferBuilder RayTracingInstanceBufferBuilder;
+	RayTracingInstanceBufferBuilder.Init(Instances, FVector::ZeroVector);
 
 	FRayTracingSceneRHIRef RayTracingSceneRHI;
 	{
 		FRayTracingSceneInitializer Initializer;
 		Initializer.DebugName = FName(TEXT("FRayTracingScene"));
-		Initializer.MaxNumInstances = RayTracingSceneInitializationData.NumNativeGPUSceneInstances + RayTracingSceneInitializationData.NumNativeCPUInstances;
-		Initializer.BuildFlags = SceneBuildFlags;
+		Initializer.MaxNumInstances = RayTracingInstanceBufferBuilder.GetMaxNumInstances();
+		Initializer.BuildFlags = ERayTracingAccelerationStructureFlags::FastTrace;
 
 		RayTracingSceneRHI = RHICreateRayTracingScene(MoveTemp(Initializer));
 	}
@@ -205,96 +205,14 @@ bool RunRayTracingTestbed_RenderThread(const FString& Parameters)
 	FRWBufferStructured InstanceBuffer;
 	InstanceBuffer.Initialize(RHICmdList, TEXT("RayTracingTestBedInstanceBuffer"), GRHIRayTracingInstanceDescriptorSize, SceneInitializer.MaxNumInstances);
 
-	FByteAddressBuffer AccelerationStructureAddressesBuffer;
-	AccelerationStructureAddressesBuffer.Initialize(RHICmdList, TEXT("RayTracingTestBedAccelerationStructureAddressesBuffer"), sizeof(FRayTracingAccelerationStructureAddress), BUF_Volatile | BUF_MultiGPUAllocate);
-
-	const uint32 NumInstanceGroups = RayTracingSceneInitializationData.NumGPUInstanceGroups + RayTracingSceneInitializationData.NumCPUInstanceGroups;
-	const FRHIBufferCreateDesc InstanceGroupUploadCreateDesc =
-		FRHIBufferCreateDesc::CreateStructured<FRayTracingInstanceGroup>(TEXT("RayTracingTestBedInstanceGroupUploadBuffer"), NumInstanceGroups)
-		.AddUsage(EBufferUsageFlags::ShaderResource | EBufferUsageFlags::Volatile)
-		.SetInitActionInitializer()
-		.DetermineInitialState();
-	TRHIBufferInitializer<FRayTracingInstanceGroup> InstanceGroupUploadData = RHICmdList.CreateBufferInitializer(InstanceGroupUploadCreateDesc);
-
-	const uint32 NumInstanceDescriptors = RayTracingSceneInitializationData.NumGPUInstanceDescriptors + RayTracingSceneInitializationData.NumCPUInstanceDescriptors;
-	const FRHIBufferCreateDesc InstanceUploadCreateDesc =
-		FRHIBufferCreateDesc::CreateStructured<FRayTracingInstanceDescriptor>(TEXT("RayTracingTestBedInstanceUploadBuffer"), NumInstanceDescriptors)
-		.AddUsage(EBufferUsageFlags::ShaderResource | EBufferUsageFlags::Volatile)
-		.SetInitActionInitializer()
-		.DetermineInitialState();
-	TRHIBufferInitializer<FRayTracingInstanceDescriptor> InstanceUploadData = RHICmdList.CreateBufferInitializer(InstanceUploadCreateDesc);
-
-	const FRHIBufferCreateDesc TransformUploadCreateDesc =
-		FRHIBufferCreateDesc::CreateStructured<FVector4f>(TEXT("RayTracingTestBedTransformUploadBuffer"), RayTracingSceneInitializationData.NumNativeCPUInstances * 3)
-		.AddUsage(EBufferUsageFlags::ShaderResource | EBufferUsageFlags::Volatile)
-		.SetInitActionInitializer()
-		.DetermineInitialState();
-	TRHIBufferInitializer<FVector4f> TransformUploadData = RHICmdList.CreateBufferInitializer(TransformUploadCreateDesc);
+	RayTracingInstanceBufferBuilder.FillRayTracingInstanceUploadBuffer(RHICmdList);
+	RayTracingInstanceBufferBuilder.FillAccelerationStructureAddressesBuffer(RHICmdList);
 
-	{
-		FillRayTracingInstanceUploadBuffer(
-			FVector::ZeroVector,
-			Instances,
-			RayTracingSceneInitializationData.InstanceGeometryIndices,
-			RayTracingSceneInitializationData.BaseUploadBufferOffsets,
-			RayTracingSceneInitializationData.BaseInstancePrefixSum,
-			RayTracingSceneInitializationData.InstanceGroupEntryRefs,
-			RayTracingSceneInitializationData.NumGPUInstanceGroups,
-			RayTracingSceneInitializationData.NumCPUInstanceGroups,
-			RayTracingSceneInitializationData.NumGPUInstanceDescriptors,
-			RayTracingSceneInitializationData.NumCPUInstanceDescriptors,
-			InstanceGroupUploadData.GetWriteView(),
-			InstanceUploadData.GetWriteView(),
-			TransformUploadData.GetWriteView()
-		);
-	}
-
-	FBufferRHIRef InstanceGroupUploadBuffer = InstanceGroupUploadData.Finalize();
-	FShaderResourceViewRHIRef InstanceGroupUploadSRV = RHICmdList.CreateShaderResourceView(InstanceGroupUploadBuffer, FRHIViewDesc::CreateBufferSRV().SetTypeFromBuffer(InstanceGroupUploadBuffer));
-
-	FBufferRHIRef InstanceUploadBuffer = InstanceUploadData.Finalize();
-	FShaderResourceViewRHIRef InstanceUploadSRV = RHICmdList.CreateShaderResourceView(InstanceUploadBuffer, FRHIViewDesc::CreateBufferSRV().SetTypeFromBuffer(InstanceUploadBuffer));
-
-	FBufferRHIRef TransformUploadBuffer = TransformUploadData.Finalize();
-	FShaderResourceViewRHIRef TransformUploadSRV = RHICmdList.CreateShaderResourceView(TransformUploadBuffer, FRHIViewDesc::CreateBufferSRV().SetTypeFromBuffer(TransformUploadBuffer));
-
-	for (uint32 GPUIndex : RHICmdList.GetGPUMask())
-	{
-		FRayTracingAccelerationStructureAddress* AddressesPtr = (FRayTracingAccelerationStructureAddress*)RHICmdList.LockBufferMGPU(
-			AccelerationStructureAddressesBuffer.Buffer,
-			GPUIndex,
-			0,
-			RayTracingSceneInitializationData.ReferencedGeometries.Num() * sizeof(FRayTracingAccelerationStructureAddress), RLM_WriteOnly);
-
-		const TArrayView<FRHIRayTracingGeometry*> ReferencedGeometries = RHICmdList.AllocArray(MakeConstArrayView(RayTracingSceneInitializationData.ReferencedGeometries));
-
-		RHICmdList.EnqueueLambda([AddressesPtr, ReferencedGeometries, GPUIndex](FRHICommandListBase&)
-		{
-			for (int32 GeometryIndex = 0; GeometryIndex < ReferencedGeometries.Num(); ++GeometryIndex)
-			{
-				AddressesPtr[GeometryIndex] = ReferencedGeometries[GeometryIndex]->GetAccelerationStructureAddress(GPUIndex);
-			}
-		});
-
-		RHICmdList.UnlockBufferMGPU(AccelerationStructureAddressesBuffer.Buffer, GPUIndex);
-	}
-
-	BuildRayTracingInstanceBuffer(
+	RayTracingInstanceBufferBuilder.BuildRayTracingInstanceBuffer(
 		RHICmdList,
 		nullptr,
-		FDFVector3{},
-		InstanceBuffer.UAV,
-		InstanceGroupUploadSRV,
-		InstanceUploadSRV,
-		AccelerationStructureAddressesBuffer.SRV,
-		TransformUploadSRV,
-		RayTracingSceneInitializationData.NumNativeGPUSceneInstances,
-		RayTracingSceneInitializationData.NumNativeCPUInstances,
-		RayTracingSceneInitializationData.NumGPUInstanceGroups,
-		RayTracingSceneInitializationData.NumCPUInstanceGroups,
-		RayTracingSceneInitializationData.NumGPUInstanceDescriptors,
-		RayTracingSceneInitializationData.NumCPUInstanceDescriptors,
 		nullptr,
+		InstanceBuffer.UAV,
 		nullptr,
 		0,
 		nullptr);
@@ -312,8 +230,8 @@ bool RunRayTracingTestbed_RenderThread(const FString& Parameters)
 	BuildParams.ScratchBufferOffset = 0;
 	BuildParams.InstanceBuffer = InstanceBuffer.Buffer;
 	BuildParams.InstanceBufferOffset = 0;
-	BuildParams.ReferencedGeometries = RayTracingSceneInitializationData.ReferencedGeometries;
-	BuildParams.NumInstances = RayTracingSceneInitializationData.NumNativeGPUSceneInstances + RayTracingSceneInitializationData.NumNativeCPUInstances;
+	BuildParams.ReferencedGeometries = RayTracingInstanceBufferBuilder.GetReferencedGeometries();
+	BuildParams.NumInstances = RayTracingInstanceBufferBuilder.GetMaxNumInstances();
 
 	RHICmdList.Transition(FRHITransitionInfo(InstanceBuffer.Buffer, ERHIAccess::UAVMask, ERHIAccess::SRVCompute));
 
diff --git a/Engine/Source/Runtime/Renderer/Public/RayTracingInstanceBufferUtil.h b/Engine/Source/Runtime/Renderer/Public/RayTracingInstanceBufferUtil.h
--- a/Engine/Source/Runtime/Renderer/Public/RayTracingInstanceBufferUtil.h
+++ b/Engine/Source/Runtime/Renderer/Public/RayTracingInstanceBufferUtil.h
@@ -3,6 +3,7 @@
 #pragma once
 
 #include "RenderGraphDefinitions.h"
+#include "RHIUtilities.h"
 #include "CoreMinimal.h"
 
 #if RHI_RAYTRACING
@@ -11,21 +12,7 @@ class FGPUScene;
 struct FRayTracingCullingParameters;
 struct FDFVector3;
 
-/*
-* 
-* Each FRayTracingGeometryInstance can translate to multiple native TLAS instances (see FRayTracingGeometryInstance::NumTransforms).
-* 
-* The FRayTracingGeometryInstance array (ie: FRayTracingScene::Instances) used to create FRayTracingSceneRHI
-* can have mix of instances using GPUScene or CPU transforms.
-* In order to reduce the number of dispatches to build the native RayTracing Instance Buffer,
-* the upload buffer containing FRayTracingInstanceDescriptor is split in 2 sections, [GPUSceneInstances] [CPUInstances].
-* This way native GPUScene and CPU instance descriptors can be built in a single dispatch per type.
-* 
-* If the ray tracing scene contains multiple layers, the instance buffer is divided into multiple subranges as expected by the RHI.
-* 
-*/
-
-struct FRayTracingInstanceDescriptor
+struct UE_DEPRECATED(5.6, "Use FRayTracingInstanceBufferBuilder instead.") FRayTracingInstanceDescriptor
 {
 	uint32 GPUSceneInstanceOrTransformIndex;
 	uint32 OutputDescriptorIndex;
@@ -36,22 +23,14 @@ struct FRayTracingInstanceDescriptor
 	uint32 bApplyLocalBoundsTransform;
 };
 
-struct FRayTracingInstanceGroup
-{
-	uint32 BaseInstanceIndex : 30;
-	uint32 bIncrementUserDataPerInstance : 1;
-	uint32 bReuseInstance : 1;
-};
-
-static_assert(sizeof(FRayTracingInstanceGroup) == sizeof(uint32), "FRayTracingInstanceGroup is expected be same size as uint32.");
-
-struct FRayTracingInstanceGroupEntryRef
+struct UE_DEPRECATED(5.6, "Use FRayTracingInstanceBufferBuilder instead.") FRayTracingInstanceGroupEntryRef
 {
 	uint32 GroupIndex;
 	uint32 BaseIndexInGroup;
 };
 
-struct FRayTracingSceneInitializationData
+PRAGMA_DISABLE_DEPRECATION_WARNINGS
+struct UE_DEPRECATED(5.6, "Use FRayTracingInstanceBufferBuilder instead.") FRayTracingSceneInitializationData
 {
 	uint32 NumGPUInstanceGroups;
 	uint32 NumCPUInstanceGroups;
@@ -79,37 +58,19 @@ struct FRayTracingSceneInitializationData
 	UE_DEPRECATED(5.5, "No longer used as 'inline SBT' is managed using FShaderBindingTable.")
 	TArray<FRHIRayTracingGeometry*> PerInstanceGeometries;
 
-	PRAGMA_DISABLE_DEPRECATION_WARNINGS
 	FRayTracingSceneInitializationData() = default;
 	FRayTracingSceneInitializationData(const FRayTracingSceneInitializationData&) = default;
 	FRayTracingSceneInitializationData& operator=(const FRayTracingSceneInitializationData&) = default;
 	FRayTracingSceneInitializationData(FRayTracingSceneInitializationData&&) = default;
 	FRayTracingSceneInitializationData& operator=(FRayTracingSceneInitializationData&&) = default;
 	~FRayTracingSceneInitializationData() = default;
-	PRAGMA_ENABLE_DEPRECATION_WARNINGS
 };
 
+UE_DEPRECATED(5.6, "Use FRayTracingInstanceBufferBuilder instead.")
 RENDERER_API FRayTracingSceneInitializationData BuildRayTracingSceneInitializationData(TConstArrayView<FRayTracingGeometryInstance> Instances);
 
-// Helper function to fill upload buffers required by BuildRayTracingInstanceBuffer with instance descriptors
-// Transforms of CPU instances are copied to OutTransformData
+UE_DEPRECATED(5.6, "Use FRayTracingInstanceBufferBuilder instead.")
 RENDERER_API void FillRayTracingInstanceUploadBuffer(
-	FVector PreViewTranslation,
-	TConstArrayView<FRayTracingGeometryInstance> Instances,
-	TConstArrayView<uint32> InstanceGeometryIndices,
-	TConstArrayView<uint32> BaseUploadBufferOffsets,
-	TConstArrayView<uint32> BaseInstancePrefixSum,
-	TConstArrayView<FRayTracingInstanceGroupEntryRef> InstanceGroupEntryRefs,
-	uint32 NumGPUInstanceGroups,
-	uint32 NumCPUInstanceGroups,
-	uint32 NumGPUInstanceDescriptors,
-	uint32 NumCPUInstanceDescriptors,
-	TArrayView<FRayTracingInstanceGroup> OutInstanceGroupUploadData,
-	TArrayView<FRayTracingInstanceDescriptor> OutInstanceUploadData,
-	TArrayView<FVector4f> OutTransformData);
-
-UE_DEPRECATED(5.6, "No longer need to provide RayTracingSceneRHI.")
-inline void FillRayTracingInstanceUploadBuffer(
 	FRayTracingSceneRHIRef RayTracingSceneRHI,
 	FVector PreViewTranslation,
 	TConstArrayView<FRayTracingGeometryInstance> Instances,
@@ -119,78 +80,72 @@ inline void FillRayTracingInstanceUploadBuffer(
 	uint32 NumNativeGPUSceneInstances,
 	uint32 NumNativeCPUInstances,
 	TArrayView<FRayTracingInstanceDescriptor> OutInstanceUploadData,
-	TArrayView<FVector4f> OutTransformData)
-{
-	FillRayTracingInstanceUploadBuffer(
-		PreViewTranslation,
-		Instances,
-		InstanceGeometryIndices,
-		BaseUploadBufferOffsets,
-		BaseInstancePrefixSum,
-		{},
-		0,
-		0,
-		NumNativeGPUSceneInstances,
-		NumNativeCPUInstances,
-		{},
-		OutInstanceUploadData,
-		OutTransformData);
-}
+	TArrayView<FVector4f> OutTransformData);
 
+UE_DEPRECATED(5.6, "Use FRayTracingInstanceBufferBuilder instead.")
 RENDERER_API void BuildRayTracingInstanceBuffer(
 	FRHICommandList& RHICmdList,
 	const FGPUScene* GPUScene,
 	const FDFVector3& PreViewTranslation,
 	FRHIUnorderedAccessView* InstancesUAV,
-	FRHIShaderResourceView* InstanceGroupUploadSRV,
 	FRHIShaderResourceView* InstanceUploadSRV,
 	FRHIShaderResourceView* AccelerationStructureAddressesSRV,
 	FRHIShaderResourceView* CPUInstanceTransformSRV,
 	uint32 NumNativeGPUSceneInstances,
 	uint32 NumNativeCPUInstances,
-	uint32 NumGPUGroups,
-	uint32 NumCPUGroups,
-	uint32 NumGPUInstanceDescriptors,
-	uint32 NumCPUInstanceDescriptors,
 	const FRayTracingCullingParameters* CullingParameters,
 	FRHIUnorderedAccessView* OutputStatsUAV,
-	uint32 OutputStatsOffset,
 	FRHIUnorderedAccessView* DebugInstanceGPUSceneIndexUAV);
 
-UE_DEPRECATED(5.6, "Need to provide InstanceGroupUploadSRV.")
-inline void BuildRayTracingInstanceBuffer(
-	FRHICommandList& RHICmdList,
-	const FGPUScene* GPUScene,
-	const FDFVector3& PreViewTranslation,
-	FRHIUnorderedAccessView* InstancesUAV,
-	FRHIShaderResourceView* InstanceUploadSRV,
-	FRHIShaderResourceView* AccelerationStructureAddressesSRV,
-	FRHIShaderResourceView* CPUInstanceTransformSRV,
-	uint32 NumNativeGPUSceneInstances,
-	uint32 NumNativeCPUInstances,
-	const FRayTracingCullingParameters* CullingParameters,
-	FRHIUnorderedAccessView* OutputStatsUAV,
-	FRHIUnorderedAccessView* DebugInstanceGPUSceneIndexUAV)
+PRAGMA_ENABLE_DEPRECATION_WARNINGS
+
+class FRayTracingInstanceBufferBuilder
 {
-	BuildRayTracingInstanceBuffer(
-		RHICmdList,
-		GPUScene,
-		PreViewTranslation,
-		InstancesUAV,
-		nullptr,
-		InstanceUploadSRV,
-		AccelerationStructureAddressesSRV,
-		CPUInstanceTransformSRV,
-		NumNativeGPUSceneInstances,
-		NumNativeCPUInstances,
-		0,
-		0,
-		0,
-		0,
-		CullingParameters,
-		OutputStatsUAV,
-		0,
-		DebugInstanceGPUSceneIndexUAV);
-}
+public:
+
+	RENDERER_API void Init(TConstArrayView<FRayTracingGeometryInstance> InInstances, FVector InPreViewTranslation);
+
+	RENDERER_API void FillRayTracingInstanceUploadBuffer(FRHICommandList& RHICmdList);
+	RENDERER_API void FillAccelerationStructureAddressesBuffer(FRHICommandList& RHICmdList);
+
+	RENDERER_API void BuildRayTracingInstanceBuffer(
+		FRHICommandList& RHICmdList,
+		const FGPUScene* GPUScene,
+		const FRayTracingCullingParameters* CullingParameters,
+		FRHIUnorderedAccessView* InstancesUAV,
+		FRHIUnorderedAccessView* OutputStatsUAV,
+		uint32 OutputStatsOffset,
+		FRHIUnorderedAccessView* DebugInstanceGPUSceneIndexUAV);
+
+	PRAGMA_DISABLE_DEPRECATION_WARNINGS
+
+	RENDERER_API uint32 GetMaxNumInstances() const { return Data.NumNativeGPUSceneInstances + Data.NumNativeCPUInstances; }
+
+	RENDERER_API TConstArrayView<FRHIRayTracingGeometry*> GetReferencedGeometries() const { return Data.ReferencedGeometries; }
+	RENDERER_API TConstArrayView<uint32> GetInstanceGeometryIndices() const { return Data.InstanceGeometryIndices; }
+	RENDERER_API TConstArrayView<uint32> GetBaseInstancePrefixSum() const { return Data.BaseInstancePrefixSum; }
+
+	PRAGMA_ENABLE_DEPRECATION_WARNINGS
+
+private:
+
+	TConstArrayView<FRayTracingGeometryInstance> Instances;
+	FVector PreViewTranslation;
+
+	PRAGMA_DISABLE_DEPRECATION_WARNINGS
+	FRayTracingSceneInitializationData Data;
+	PRAGMA_ENABLE_DEPRECATION_WARNINGS
+
+	FBufferRHIRef InstanceGroupUploadBuffer;
+	FShaderResourceViewRHIRef InstanceGroupUploadSRV;
+
+	FBufferRHIRef InstanceUploadBuffer;
+	FShaderResourceViewRHIRef InstanceUploadSRV;
+
+	FBufferRHIRef TransformUploadBuffer;
+	FShaderResourceViewRHIRef TransformUploadSRV;
+
+	FByteAddressBuffer AccelerationStructureAddressesBuffer;
+};
 
 #endif // RHI_RAYTRACING
\ No newline at end of file
