Refactor virtual texture feedback buffer code.
Moved some code that wasn't exposed with RENDERER_API out of public and into internal folder.
Moved some related code out of SceneRenderer and into the new internal file.
Changed the internal API so that it will be suitable to use for slate rendering.

diff --git a/Engine/Source/Runtime/Renderer/Internal/VT/VirtualTextureFeedbackResource.h b/Engine/Source/Runtime/Renderer/Internal/VT/VirtualTextureFeedbackResource.h
new file mode 100644
--- /dev/null
+++ b/Engine/Source/Runtime/Renderer/Internal/VT/VirtualTextureFeedbackResource.h
@@ -0,0 +1,72 @@
+// Copyright Epic Games, Inc. All Rights Reserved.
+
+#pragma once
+
+#include "CoreMinimal.h"
+#include "RenderResource.h"
+#include "RHI.h"
+
+class FRDGBuilder;
+class FViewUniformShaderParameters;
+
+/** GPU feedback buffer for tracking virtual texture requests from the GPU. */
+class FVirtualTextureFeedbackBufferResource : public FRenderResource
+{
+public:
+	/** Allocates and prepares a new feedback buffer for write access. */
+	void Begin(FRDGBuilder& GraphBuilder, uint32 InFeedbackBufferSize, ERHIFeatureLevel::Type InFeatureLevel);
+	/** Compacts the feedback buffer on the GPU and copies the compacted buffer for readback. */
+	void End(FRDGBuilder& GraphBuilder);
+
+	/** Get the size of the feedback buffer from the last call to Begin(). */
+	uint32 GetBufferSize() const;
+	/** Get the UAV of the feedback buffer. */
+	FRHIUnorderedAccessView* GetUAV() const;
+
+	//~ Begin FRenderResource Interface
+	virtual void ReleaseRHI() override;
+	//~ End FRenderResource Interface
+
+private:
+	bool bIsInBeginEndScope = false;
+	uint32 FeedbackBufferSize = 0;
+	ERHIFeatureLevel::Type FeatureLevel = ERHIFeatureLevel::Num;
+	TRefCountPtr<FRDGPooledBuffer> PooledBuffer;
+	FRHIUnorderedAccessView* UAV = nullptr;
+};
+
+namespace VirtualTexture
+{
+	/** Shader parameters needed when writing virtual texture feedback from GPU. */
+	struct RENDERER_API FFeedbackShaderParams
+	{
+		FRHIUnorderedAccessView* BufferUAV = nullptr;
+		uint32 BufferSize = 0;
+		uint32 Shift = 0;
+		uint32 Mask = 0;
+		uint32 JitterOffset = 0;
+		uint32 SampleOffset = 0;
+	};
+
+	/** 
+	 * Begin a virtual texture feedback scope. 
+	 * Note that this uses a global FVirtualTextureFeedbackBufferResource object which is then read when calling GetFeedbackShaderParams().
+	 */
+	RENDERER_API void BeginFeedback(FRDGBuilder& GraphBuilder, uint32 InBufferSize);
+	RENDERER_API void BeginFeedback(FRDGBuilder& GraphBuilder, uint32 InBufferSize, ERHIFeatureLevel::Type InFeatureLevel);
+	void BeginFeedback(FRDGBuilder& GraphBuilder, FIntPoint InViewportSize, uint32 InVirtualTextureFeedbackTileSize, ERHIFeatureLevel::Type InFeatureLevel);
+
+	/** End a virtual texture feedback scope. */
+	RENDERER_API void EndFeedback(FRDGBuilder& GraphBuilder);
+
+	/** Get virtual texture feedback parameters to use for binding to any shader that is sampling virtual textures. */
+	void GetFeedbackShaderParams(int32 InFrameIndex, int32 InVirtualTextureFeedbackTileSize, FFeedbackShaderParams& OutParams);
+	/** 
+	 * Get virtual texture feedback parameters to use for binding to any shader that is sampling virtual textures. 
+	 * This version uses an internal scene frame counter, and the project default feedback tile size.
+	 */
+	RENDERER_API void GetFeedbackShaderParams(FFeedbackShaderParams& OutParams);
+	
+	/** Helper to copy the virtual texture feedback shader parameters into the view parameters to be bound in the view uniform buffer. */
+	RENDERER_API void UpdateViewUniformShaderParameters(FFeedbackShaderParams const& InParams, FViewUniformShaderParameters& ViewUniformShaderParameters);
+}
diff --git a/Engine/Source/Runtime/Renderer/Private/DeferredShadingRenderer.cpp b/Engine/Source/Runtime/Renderer/Private/DeferredShadingRenderer.cpp
--- a/Engine/Source/Runtime/Renderer/Private/DeferredShadingRenderer.cpp
+++ b/Engine/Source/Runtime/Renderer/Private/DeferredShadingRenderer.cpp
@@ -35,7 +35,7 @@
 #include "PipelineStateCache.h"
 #include "ClearQuad.h"
 #include "RendererModule.h"
-#include "VT/VirtualTextureFeedback.h"
+#include "VT/VirtualTextureFeedbackResource.h"
 #include "VT/VirtualTextureSystem.h"
 #include "GPUScene.h"
 #include "PathTracing.h"
@@ -3563,7 +3563,7 @@ void FDeferredShadingSceneRenderer::Render(FRDGBuilder& GraphBuilder, const FSce
 			RDG_EVENT_SCOPE_STAT(GraphBuilder, VirtualTextureUpdate, "VirtualTextureUpdate");
 			RDG_GPU_STAT_SCOPE(GraphBuilder, VirtualTextureUpdate);
 
-			VirtualTextureFeedbackEnd(GraphBuilder);
+			VirtualTexture::EndFeedback(GraphBuilder);
 		}
 
 		// After AddPostProcessingPasses in case of Lumen Visualizations writing to feedback
diff --git a/Engine/Source/Runtime/Renderer/Private/MobileShadingRenderer.cpp b/Engine/Source/Runtime/Renderer/Private/MobileShadingRenderer.cpp
--- a/Engine/Source/Runtime/Renderer/Private/MobileShadingRenderer.cpp
+++ b/Engine/Source/Runtime/Renderer/Private/MobileShadingRenderer.cpp
@@ -53,7 +53,7 @@
 #include "MaterialSceneTextureId.h"
 #include "SkyAtmosphereRendering.h"
 #include "VisualizeTexture.h"
-#include "VT/VirtualTextureFeedback.h"
+#include "VT/VirtualTextureFeedbackResource.h"
 #include "VT/VirtualTextureSystem.h"
 #include "GPUSortManager.h"
 #include "MobileBasePassRendering.h"
@@ -1558,7 +1558,7 @@ void FMobileSceneRenderer::Render(FRDGBuilder& GraphBuilder, const FSceneRenderU
 	{
 		RDG_EVENT_SCOPE_STAT(GraphBuilder, VirtualTextureUpdate, "VirtualTextureUpdate");
 		RDG_GPU_STAT_SCOPE(GraphBuilder, VirtualTextureUpdate);
-		VirtualTextureFeedbackEnd(GraphBuilder);
+		VirtualTexture::EndFeedback(GraphBuilder);
 	}
 
 	if (bRendererOutputFinalSceneColor && bShouldRenderHZB && !bRequiresMultiPass)
diff --git a/Engine/Source/Runtime/Renderer/Private/Renderer.cpp b/Engine/Source/Runtime/Renderer/Private/Renderer.cpp
--- a/Engine/Source/Runtime/Renderer/Private/Renderer.cpp
+++ b/Engine/Source/Runtime/Renderer/Private/Renderer.cpp
@@ -32,6 +32,7 @@
 #include "EditorPrimitivesRendering.h"
 #include "VisualizeTexturePresent.h"
 #include "ScreenSpaceDenoise.h"
+#include "VT/VirtualTextureFeedbackResource.h"
 #include "VT/VirtualTextureSystem.h"
 #include "PostProcess/TemporalAA.h"
 #include "CanvasRender.h"
@@ -575,7 +576,7 @@ void FRendererModule::DrawTileMesh(FCanvasRenderContext& RenderContext, FMeshPas
 			RDG_EVENT_SCOPE_STAT(GraphBuilder, VirtualTextureUpdate, "VirtualTextureUpdate");
 			RDG_GPU_STAT_SCOPE(GraphBuilder, VirtualTextureUpdate);
 
-			VirtualTextureFeedbackEnd(GraphBuilder);
+			VirtualTexture::EndFeedback(GraphBuilder);
 		}
 	}
 }
diff --git a/Engine/Source/Runtime/Renderer/Private/SceneRendering.cpp b/Engine/Source/Runtime/Renderer/Private/SceneRendering.cpp
--- a/Engine/Source/Runtime/Renderer/Private/SceneRendering.cpp
+++ b/Engine/Source/Runtime/Renderer/Private/SceneRendering.cpp
@@ -57,8 +57,6 @@
 #include "VisualizeTexture.h"
 #include "VisualizeTexturePresent.h"
 #include "MeshDrawCommands.h"
-#include "VT/VirtualTextureSystem.h"
-#include "VT/VirtualTextureFeedback.h"
 #include "HAL/LowLevelMemTracker.h"
 #include "IXRTrackingSystem.h"
 #include "IXRCamera.h"
@@ -110,6 +108,7 @@
 #include "EnvironmentComponentsFlags.h"
 #include "Math/RotationMatrix.h"
 #include "VolumetricCloudProxy.h"
+#include "VT/VirtualTextureFeedbackResource.h"
 #include "VT/VirtualTextureScalability.h"
 #include "VT/VirtualTextureSystem.h"
 #include "SceneRenderBuilder.h"
@@ -123,8 +122,6 @@
 
 DEFINE_LOG_CATEGORY(LogSceneCapture);
 
-static TGlobalResource<FVirtualTextureFeedbackBuffer> GVirtualTextureFeedbackBuffer;
-
 // This is an experimental optimization switch to render pre-pass depth for scene capture without calling the entire FDeferredShadingSceneRenderer::Render()
 int32 GSceneCaptureDepthPrepassOptimization = 0;
 static FAutoConsoleVariableRef CVarSceneCaptureDepthPrepassOptimization(
@@ -1152,33 +1149,6 @@ void UpdateNoiseTextureParameters(FViewUniformShaderParameters& ViewUniformShade
 	check(ViewUniformShaderParameters.SobolSamplingTexture);
 }
 
-static uint32 AlignVirtualTextureFeedbackTileSize(uint32 InTileSize)
-{
-	// Round to nearest power of two to ensure that shader maths is efficient and sampling sequence logic is simple.
-	return FMath::RoundUpToPowerOfTwo(FMath::Max(InTileSize, 1u));
-}
-
-static void UpdateVirtualTextureFeedbackParameters(uint32 InFrameIndex, uint32 InVirtualTextureFeedbackTileSize, FViewUniformShaderParameters& ViewUniformShaderParameters)
-{
-	ViewUniformShaderParameters.VTFeedbackBuffer = GVirtualTextureFeedbackBuffer.GetUAV();
-	ViewUniformShaderParameters.VirtualTextureFeedbackBufferSize = GVirtualTextureFeedbackBuffer.GetBufferSize();
-	
-	ViewUniformShaderParameters.VirtualTextureFeedbackShift = FMath::FloorLog2(InVirtualTextureFeedbackTileSize);
-	ViewUniformShaderParameters.VirtualTextureFeedbackMask = InVirtualTextureFeedbackTileSize - 1;
-	// Use some low(ish) discrepancy sequence to run over every pixel in the virtual texture feedback tile.
-	ViewUniformShaderParameters.VirtualTextureFeedbackJitterOffset = SampleVirtualTextureFeedbackSequence(InFrameIndex, InVirtualTextureFeedbackTileSize);
-	// SampleOffset is used to cycle through all VT samples in a material. It just needs to monotonically increase.
-	ViewUniformShaderParameters.VirtualTextureFeedbackSampleOffset = InFrameIndex;
-
-	ViewUniformShaderParameters.GlobalVirtualTextureMipBias = FVirtualTextureSystem::Get().GetGlobalMipBias();
-}
-
-void UpdateVirtualTextureFeedbackParameters(FViewUniformShaderParameters& ViewUniformShaderParameters)
-{
-	const uint32 VirtualTextureFeedbackTileSize = AlignVirtualTextureFeedbackTileSize(VirtualTextureScalability::GetVirtualTextureFeedbackFactor());
-	UpdateVirtualTextureFeedbackParameters(GFrameNumber, VirtualTextureFeedbackTileSize, ViewUniformShaderParameters);
-}
-
 void SetupPrecomputedVolumetricLightmapUniformBufferParameters(const FScene* Scene, FEngineShowFlags EngineShowFlags, FViewUniformShaderParameters& ViewUniformShaderParameters)
 {
 	if (Scene && Scene->VolumetricLightmapSceneData.HasData() && EngineShowFlags.VolumetricLightmap)
@@ -2056,8 +2026,12 @@ void FViewInfo::SetupUniformBufferParameters(
 	}
 
 	const uint32 VirtualTextureFrameIndex = ViewState ? ViewState->GetFrameIndex() : Family->FrameNumber;
-	const uint32 VirtualTextureFeedbackTileSize = AlignVirtualTextureFeedbackTileSize(Family->VirtualTextureFeedbackFactor);
-	UpdateVirtualTextureFeedbackParameters(VirtualTextureFrameIndex, VirtualTextureFeedbackTileSize, ViewUniformShaderParameters);
+	const uint32 VirtualTextureFeedbackTileSize = Family->VirtualTextureFeedbackFactor;
+	VirtualTexture::FFeedbackShaderParams Params;
+	VirtualTexture::GetFeedbackShaderParams(VirtualTextureFrameIndex, VirtualTextureFeedbackTileSize, Params);
+	VirtualTexture::UpdateViewUniformShaderParameters(Params, ViewUniformShaderParameters);
+
+	ViewUniformShaderParameters.GlobalVirtualTextureMipBias = FVirtualTextureSystem::Get().GetGlobalMipBias();
 
 	// GGX/Sheen LTC (used as BSDF or for rect light integration)
 	if (GSystemTextures.GGXLTCMat.IsValid() && GSystemTextures.GGXLTCAmp.IsValid())
@@ -6078,19 +6052,11 @@ void AddResolveSceneDepthPass(FRDGBuilder& GraphBuilder, TArrayView<const FViewI
 
 void VirtualTextureFeedbackBegin(FRDGBuilder& GraphBuilder, TArrayView<const FViewInfo> Views, FIntPoint SceneTextureExtent)
 {
-	FIntPoint EffectiveBufferSize = SceneTextureExtent;
-	EffectiveBufferSize.X = FMath::Max(EffectiveBufferSize.X, 1);
-	EffectiveBufferSize.Y = FMath::Max(EffectiveBufferSize.Y, 1);
+	const FIntPoint FeedbackViewSize = FIntPoint(FMath::Max(SceneTextureExtent.X, 1), FMath::Max(SceneTextureExtent.Y, 1));
 	const uint32 FeedbackTileSize = Views.Num() > 0 ? Views[0].Family->VirtualTextureFeedbackFactor : VirtualTextureScalability::GetVirtualTextureFeedbackFactor();
-	const int32 FeedbackBufferSize = GetVirtualTextureFeedbackBufferSize(EffectiveBufferSize, AlignVirtualTextureFeedbackTileSize(FeedbackTileSize));
 	const ERHIFeatureLevel::Type FeatureLevel = Views.Num() > 0 ? Views[0].GetFeatureLevel() : GMaxRHIFeatureLevel;
 	
-	GVirtualTextureFeedbackBuffer.Begin(GraphBuilder, FeedbackBufferSize, FeatureLevel);
-}
-
-void VirtualTextureFeedbackEnd(FRDGBuilder& GraphBuilder)
-{
-	GVirtualTextureFeedbackBuffer.End(GraphBuilder);
+	VirtualTexture::BeginFeedback(GraphBuilder, FeedbackViewSize, FeedbackTileSize, FeatureLevel);
 }
 
 static TAutoConsoleVariable<int32> CVarHalfResDepthNoFastClear(
diff --git a/Engine/Source/Runtime/Renderer/Private/SceneRendering.h b/Engine/Source/Runtime/Renderer/Private/SceneRendering.h
--- a/Engine/Source/Runtime/Renderer/Private/SceneRendering.h
+++ b/Engine/Source/Runtime/Renderer/Private/SceneRendering.h
@@ -2923,8 +2923,6 @@ private:
 
 // The noise textures need to be set in Slate too.
 RENDERER_API void UpdateNoiseTextureParameters(FViewUniformShaderParameters& ViewUniformShaderParameters);
-// The virtual texture feedback also needs to be set in Slate.
-RENDERER_API void UpdateVirtualTextureFeedbackParameters(FViewUniformShaderParameters& ViewUniformShaderParameters);
 
 struct FFastVramConfig
 {
@@ -3025,9 +3023,6 @@ void AddResolveSceneDepthPass(FRDGBuilder& GraphBuilder, TArrayView<const FViewI
 /** Prepares virtual textures for feedback updates. */
 void VirtualTextureFeedbackBegin(FRDGBuilder& GraphBuilder, TArrayView<const FViewInfo> Views, FIntPoint SceneTextureExtent);
 
-/** Finalizes feedback and submits for readback. */
-void VirtualTextureFeedbackEnd(FRDGBuilder& GraphBuilder);
-
 /** Creates a half resolution checkerboard min / max depth buffer from the input full resolution depth buffer. */
 FRDGTextureRef CreateHalfResolutionDepthCheckerboardMinMax(FRDGBuilder& GraphBuilder, TArrayView<const FViewInfo> Views, FRDGTextureRef SceneDepth);
 
diff --git a/Engine/Source/Runtime/Renderer/Private/VT/VirtualTextureFeedbackBuffer.cpp b/Engine/Source/Runtime/Renderer/Private/VT/VirtualTextureFeedbackBuffer.cpp
--- a/Engine/Source/Runtime/Renderer/Private/VT/VirtualTextureFeedbackBuffer.cpp
+++ b/Engine/Source/Runtime/Renderer/Private/VT/VirtualTextureFeedbackBuffer.cpp
@@ -2,62 +2,8 @@
 
 #include "VT/VirtualTextureFeedbackBuffer.h"
 
-#include "RenderGraphUtils.h"
-#include "SceneView.h"
-#include "GlobalRenderResources.h"
-#include "GPUFeedbackCompaction.h"
 #include "VT/VirtualTextureFeedback.h"
 
-int32 GVirtualTextureFeedbackOverdrawFactor = 2;
-static FAutoConsoleVariableRef CVarVirtualTextureFeedbackOverdrawFactor(
-	TEXT("r.vt.FeedbackOverdrawFactor"),
-	GVirtualTextureFeedbackOverdrawFactor,
-	TEXT("A multiplicative factor to apply to virtual texture feedback buffer sizes to account for ")
-	TEXT("multiple layers of virtual texture feedback from Decal/Transparency/PostFX etc."),
-	ECVF_RenderThreadSafe);
-
-int32 GVirtualTextureFeedbackCompactionFactor = 16;
-static FAutoConsoleVariableRef CVarVirtualTextureFeedbackCompactionFactor(
-	TEXT("r.vt.FeedbackCompactionFactor"),
-	GVirtualTextureFeedbackCompactionFactor,
-	TEXT("A division factor to apply to the size of the virtual texture feedback compaction buffer ")
-	TEXT("to account for compaction of duplicate page ids."),
-	ECVF_RenderThreadSafe);
-
-uint32 GetVirtualTextureFeedbackBufferSize(FIntPoint InSceneTextureExtent, uint32 InVirtualTextureFeedbackTileSize)
-{
-	// Tile size must be power or 2.
-	check(InVirtualTextureFeedbackTileSize == 1 << FMath::FloorLog2(InVirtualTextureFeedbackTileSize));
-
-	FIntPoint ScaledExtent = FIntPoint::DivideAndRoundUp(InSceneTextureExtent, InVirtualTextureFeedbackTileSize);
-	return ScaledExtent.X * ScaledExtent.Y * FMath::Max(GVirtualTextureFeedbackOverdrawFactor, 1);
-}
-
-static uint32 GetVirtualTextureCompactedFeedbackBufferSize(uint32 InSourceBufferSize)
-{
-	// Could possibly do this dynamically according to some tracked recent high watermark?
-	const uint32 CompactedBufferSize = InSourceBufferSize / FMath::Max(GVirtualTextureFeedbackCompactionFactor, 1);
-	// Size needs to be power of two for hash table wrapping.
-	const uint32 AlignedCompactedBufferSize = FMath::RoundUpToPowerOfTwo(FMath::Clamp(CompactedBufferSize, 16u, 16u * 1024u));
-	return AlignedCompactedBufferSize;
-}
-
-uint32 SampleVirtualTextureFeedbackSequence(uint32 InFrameIndex, uint32 InVirtualTextureFeedbackTileSize)
-{
-	// Tile size must be power or 2.
-	check(InVirtualTextureFeedbackTileSize == 1 << FMath::FloorLog2(InVirtualTextureFeedbackTileSize));
-
-	const uint32 TileSize = InVirtualTextureFeedbackTileSize;
-	const uint32 TileSizeLog2 = FMath::CeilLogTwo(TileSize);
-	const uint32 SequenceSize = FMath::Square(TileSize);
-	const uint32 PixelIndex = InFrameIndex % SequenceSize;
-	const uint32 PixelAddress = ReverseBits(PixelIndex) >> (32U - 2 * TileSizeLog2);
-	const uint32 X = FMath::ReverseMortonCode2(PixelAddress);
-	const uint32 Y = FMath::ReverseMortonCode2(PixelAddress >> 1);
-	const uint32 PixelSequenceIndex = X + Y * TileSize;
-	return PixelSequenceIndex;
-}
-
 void FVirtualTextureFeedbackBufferDesc::Init(uint32 InBufferSize)
 {
 	BufferSize = InBufferSize;
@@ -84,182 +30,6 @@ void FVirtualTextureFeedbackBufferDesc::Init2D(FIntPoint InUnscaledBufferSize, T
 	bSizeInHeader = false;
 }
 
-void FVirtualTextureFeedbackBuffer::Begin(FRDGBuilder& GraphBuilder, uint32 InFeedbackBufferSize, ERHIFeatureLevel::Type InFeatureLevel)
-{
-	// NOTE: Transitions and allocations are handled manually right now, because the VT feedback UAV is used by
-	// the view uniform buffer, which is not an RDG uniform buffer. If it can be factored out into its own RDG
-	// uniform buffer (or put on the pass uniform buffers), then the resource can be fully converted to RDG.
-
-	FeedbackBufferSize = InFeedbackBufferSize;
-	FeatureLevel = InFeatureLevel;
-
-	FRDGBufferDesc BufferDesc(FRDGBufferDesc::CreateStructuredDesc(sizeof(uint32), FeedbackBufferSize));
-	BufferDesc.Usage |= BUF_SourceCopy;
-
-	AllocatePooledBuffer(BufferDesc, PooledBuffer, TEXT("VirtualTexture_FeedbackBuffer"));
-	
-	FRDGBufferUAVDesc UAVDesc{};
-	UAV = PooledBuffer->GetOrCreateUAV(GraphBuilder.RHICmdList, UAVDesc);
-
-	AddPass(GraphBuilder, RDG_EVENT_NAME("VirtualTextureClear"), [UAV = UAV](FRDGAsyncTask, FRHICommandList& RHICmdList)
-	{
-		// Clear virtual texture feedback to default value
-		RHICmdList.Transition(FRHITransitionInfo(UAV, ERHIAccess::Unknown, ERHIAccess::UAVCompute));
-		RHICmdList.ClearUAVUint(UAV, FUintVector4(0u, 0u, 0u, 0u));
-		RHICmdList.Transition(FRHITransitionInfo(UAV, ERHIAccess::UAVCompute, ERHIAccess::UAVMask));
-		RHICmdList.BeginUAVOverlap(UAV);
-	});
-
-	bIsInBeginEndScope = true;
-}
-
-void FVirtualTextureFeedbackBuffer::End(FRDGBuilder& GraphBuilder)
-{
-	if (!bIsInBeginEndScope)
-	{
-		return;
-	}
-	bIsInBeginEndScope = false;
-
-	AddPass(GraphBuilder, RDG_EVENT_NAME("VirtualTextureFeedbackTransition"), [UAV = UAV](FRHICommandListImmediate& RHICmdList)
-	{
-		RHICmdList.EndUAVOverlap(UAV);
-		RHICmdList.Transition(FRHITransitionInfo(UAV, ERHIAccess::UAVMask, ERHIAccess::SRVCompute));
-	});
-
-	FRDGBufferRef FeedbackBuffer = GraphBuilder.RegisterExternalBuffer(PooledBuffer, ERDGBufferFlags::SkipTracking);
-	FRDGBufferSRVRef FeedbackBufferSRV = GraphBuilder.CreateSRV(FeedbackBuffer);
-
-	// Compact feedback before queuing for readback.
-	const uint32 CompactedFeedbackBufferSize = GetVirtualTextureCompactedFeedbackBufferSize(FeedbackBufferSize);
-	const uint32 HashTableSize = 2 * CompactedFeedbackBufferSize;
-	const uint32 HashTableIndexWrapMask = HashTableSize - 1;
-
-	FRDGBufferDesc CompactedFeedbackBufferDesc(FRDGBufferDesc::CreateStructuredDesc(sizeof(uint32), CompactedFeedbackBufferSize));
-	CompactedFeedbackBufferDesc.Usage |= BUF_SourceCopy;
-	FRDGBufferRef CompactedFeedbackBuffer = GraphBuilder.CreateBuffer(CompactedFeedbackBufferDesc, TEXT("VirtualTexture.CompactedFeedback"));
-
-	// Need to clear this buffer, as first element will be used as an allocator
-	AddClearUAVPass(GraphBuilder, GraphBuilder.CreateUAV(CompactedFeedbackBuffer, PF_R32_UINT), 0);
-
-	FRDGBufferDesc HashTableBufferDesc(FRDGBufferDesc::CreateStructuredDesc(sizeof(uint32), HashTableSize));
-	FRDGBufferRef HashTableKeyBuffer = GraphBuilder.CreateBuffer(HashTableBufferDesc, TEXT("VirtualTexture.HashTableKeys"));
-	FRDGBufferRef HashTableElementIndexBuffer = GraphBuilder.CreateBuffer(HashTableBufferDesc, TEXT("VirtualTexture.HashTableElementIndices"));
-	FRDGBufferRef HashTableElementCountBuffer = GraphBuilder.CreateBuffer(HashTableBufferDesc, TEXT("VirtualTexture.HashTableElementCounts"));
-
-	// Hash table depends on empty slots to be 0
-	AddClearUAVPass(GraphBuilder, GraphBuilder.CreateUAV(HashTableKeyBuffer, PF_R32_UINT), 0);
-	AddClearUAVPass(GraphBuilder, GraphBuilder.CreateUAV(HashTableElementCountBuffer, PF_R32_UINT), 0);
-
-	FRDGBufferRef BuildHashTableIndirectArgBuffer = GraphBuilder.CreateBuffer(FRDGBufferDesc::CreateIndirectDesc<FRHIDispatchIndirectParameters>(1), TEXT("VirtualTexture.BuildHashTableIndirectArgs"));
-
-	FGlobalShaderMap* GlobalShaderMap = GetGlobalShaderMap(FeatureLevel);
-
-	// Set indirect dispatch arguments for hash table building
-	{
-		FBuildFeedbackHashTableIndirectArgsCS::FParameters* PassParameters = GraphBuilder.AllocParameters<FBuildFeedbackHashTableIndirectArgsCS::FParameters>();
-
-		PassParameters->RWBuildHashTableIndirectArgs = GraphBuilder.CreateUAV(BuildHashTableIndirectArgBuffer, PF_R32_UINT);
-
-		PassParameters->FeedbackBufferAllocator = FeedbackBufferSRV;
-		PassParameters->FeedbackBuffer = FeedbackBufferSRV;
-		PassParameters->FeedbackBufferSize = FeedbackBufferSize;
-
-		FBuildFeedbackHashTableIndirectArgsCS::FPermutationDomain PermutationVector;
-		PermutationVector.Set<FBuildFeedbackHashTableIndirectArgsCS::FFeedbackBufferStride>(1);
-
-		auto ComputeShader = GlobalShaderMap->GetShader<FBuildFeedbackHashTableIndirectArgsCS>(PermutationVector);
-		const FIntVector GroupSize = FIntVector(1, 1, 1);
-
-		FComputeShaderUtils::AddPass(
-			GraphBuilder,
-			RDG_EVENT_NAME("Hash table indirect arguments"),
-			ComputeShader,
-			PassParameters,
-			GroupSize);
-	}
-
-	// Build hash table of feedback elements
-	{
-		FBuildFeedbackHashTableCS::FParameters* PassParameters = GraphBuilder.AllocParameters<FBuildFeedbackHashTableCS::FParameters>();
-
-		PassParameters->BuildHashTableIndirectArgs = BuildHashTableIndirectArgBuffer;
-
-		PassParameters->RWHashTableKeys = GraphBuilder.CreateUAV(HashTableKeyBuffer);
-		PassParameters->RWHashTableElementIndices = GraphBuilder.CreateUAV(HashTableElementIndexBuffer);
-		PassParameters->RWHashTableElementCounts = GraphBuilder.CreateUAV(HashTableElementCountBuffer);
-		PassParameters->HashTableSize = HashTableSize;
-		PassParameters->HashTableIndexWrapMask = HashTableIndexWrapMask;
-
-		PassParameters->FeedbackBufferAllocator = FeedbackBufferSRV;
-		PassParameters->FeedbackBuffer = FeedbackBufferSRV;
-		PassParameters->FeedbackBufferSize = FeedbackBufferSize;
-
-		FBuildFeedbackHashTableCS::FPermutationDomain PermutationVector;
-		PermutationVector.Set<FBuildFeedbackHashTableCS::FFeedbackBufferStride>(1);
-
-		auto ComputeShader = GlobalShaderMap->GetShader<FBuildFeedbackHashTableCS>(PermutationVector);
-
-		FComputeShaderUtils::AddPass(
-			GraphBuilder,
-			RDG_EVENT_NAME("Build feedback hash table"),
-			ComputeShader,
-			PassParameters,
-			BuildHashTableIndirectArgBuffer,
-			0);
-	}
-
-	// Compact hash table into an array of unique feedback elements
-	{
-		FCompactFeedbackHashTableCS::FParameters* PassParameters = GraphBuilder.AllocParameters<FCompactFeedbackHashTableCS::FParameters>();
-
-		PassParameters->RWCompactedFeedbackBuffer = GraphBuilder.CreateUAV(CompactedFeedbackBuffer);
-		PassParameters->CompactedFeedbackBufferSize = CompactedFeedbackBufferSize;
-
-		PassParameters->HashTableElementIndices = GraphBuilder.CreateSRV(HashTableElementIndexBuffer, PF_R32_UINT);
-		PassParameters->HashTableElementCounts = GraphBuilder.CreateSRV(HashTableElementCountBuffer, PF_R32_UINT);
-		PassParameters->HashTableSize = HashTableSize;
-		PassParameters->HashTableIndexWrapMask = HashTableIndexWrapMask;
-
-		PassParameters->FeedbackBufferAllocator = FeedbackBufferSRV;
-		PassParameters->FeedbackBuffer = FeedbackBufferSRV;
-		PassParameters->FeedbackBufferSize = FeedbackBufferSize;
-
-		FCompactFeedbackHashTableCS::FPermutationDomain PermutationVector;
-		PermutationVector.Set<FCompactFeedbackHashTableCS::FFeedbackBufferStride>(1);
-
-		auto ComputeShader = GlobalShaderMap->GetShader<FCompactFeedbackHashTableCS>(PermutationVector);
-		const FIntVector GroupSize = FComputeShaderUtils::GetGroupCount(HashTableSize, FCompactFeedbackHashTableCS::GetGroupSize());
-
-		FComputeShaderUtils::AddPass(
-			GraphBuilder,
-			RDG_EVENT_NAME("Compact feedback hash table"),
-			ComputeShader,
-			PassParameters,
-			GroupSize);
-	}
-
-	FVirtualTextureFeedbackBufferDesc CompactedDesc;
-	CompactedDesc.Init(CompactedFeedbackBufferSize, true);
-	SubmitVirtualTextureFeedbackBuffer(GraphBuilder, CompactedFeedbackBuffer, CompactedDesc);
-}
-
-uint32 FVirtualTextureFeedbackBuffer::GetBufferSize() const
-{
-	return bIsInBeginEndScope ? FeedbackBufferSize : 0;
-}
-
-FRHIUnorderedAccessView* FVirtualTextureFeedbackBuffer::GetUAV() const
-{
-	return UAV != nullptr && bIsInBeginEndScope ? UAV : GEmptyStructuredBufferWithUAV->UnorderedAccessViewRHI.GetReference();
-}
-
-void FVirtualTextureFeedbackBuffer::ReleaseRHI()
-{
-	PooledBuffer = nullptr;
-	UAV = nullptr;
-}
-
 void SubmitVirtualTextureFeedbackBuffer(FRHICommandList& RHICmdList, FBufferRHIRef const& InBuffer, FVirtualTextureFeedbackBufferDesc const& InDesc)
 {
 	GVirtualTextureFeedback.TransferGPUToCPU(RHICmdList, InBuffer, InDesc);
diff --git a/Engine/Source/Runtime/Renderer/Private/VT/VirtualTextureFeedbackResource.cpp b/Engine/Source/Runtime/Renderer/Private/VT/VirtualTextureFeedbackResource.cpp
new file mode 100644
--- /dev/null
+++ b/Engine/Source/Runtime/Renderer/Private/VT/VirtualTextureFeedbackResource.cpp
@@ -0,0 +1,313 @@
+// Copyright Epic Games, Inc. All Rights Reserved.
+
+#include "VT/VirtualTextureFeedbackResource.h"
+
+#include "GlobalRenderResources.h"
+#include "GPUFeedbackCompaction.h"
+#include "RenderGraphUtils.h"
+#include "SceneView.h"
+#include "VT/VirtualTextureFeedbackBuffer.h"
+#include "VT/VirtualTextureScalability.h"
+#include "VT/VirtualTextureSystem.h"
+
+int32 GVirtualTextureFeedbackOverdrawFactor = 2;
+static FAutoConsoleVariableRef CVarVirtualTextureFeedbackOverdrawFactor(
+	TEXT("r.vt.FeedbackOverdrawFactor"),
+	GVirtualTextureFeedbackOverdrawFactor,
+	TEXT("A multiplicative factor to apply to virtual texture feedback buffer sizes to account for ")
+	TEXT("multiple layers of virtual texture feedback from Decal/Transparency/PostFX etc."),
+	ECVF_RenderThreadSafe);
+
+int32 GVirtualTextureFeedbackCompactionFactor = 16;
+static FAutoConsoleVariableRef CVarVirtualTextureFeedbackCompactionFactor(
+	TEXT("r.vt.FeedbackCompactionFactor"),
+	GVirtualTextureFeedbackCompactionFactor,
+	TEXT("A division factor to apply to the size of the virtual texture feedback compaction buffer ")
+	TEXT("to account for compaction of duplicate page ids."),
+	ECVF_RenderThreadSafe);
+
+static uint32 GetVirtualTextureCompactedFeedbackBufferSize(uint32 InSourceBufferSize);
+
+/** Single global feedback buffer resource used when calling the BeginFeedback()/EndFeedback free functions. */
+static TGlobalResource<FVirtualTextureFeedbackBufferResource> GVirtualTextureFeedbackBufferResource;
+
+void FVirtualTextureFeedbackBufferResource::Begin(FRDGBuilder& GraphBuilder, uint32 InFeedbackBufferSize, ERHIFeatureLevel::Type InFeatureLevel)
+{
+	// NOTE: Transitions and allocations are handled manually right now, because the VT feedback UAV is used by
+	// the view uniform buffer, which is not an RDG uniform buffer. If it can be factored out into its own RDG
+	// uniform buffer (or put on the pass uniform buffers), then the resource can be fully converted to RDG.
+
+	FeedbackBufferSize = InFeedbackBufferSize;
+	FeatureLevel = InFeatureLevel;
+
+	FRDGBufferDesc BufferDesc(FRDGBufferDesc::CreateStructuredDesc(sizeof(uint32), FeedbackBufferSize));
+	BufferDesc.Usage |= BUF_SourceCopy;
+
+	AllocatePooledBuffer(BufferDesc, PooledBuffer, TEXT("VirtualTexture_FeedbackBuffer"));
+
+	FRDGBufferUAVDesc UAVDesc{};
+	UAV = PooledBuffer->GetOrCreateUAV(GraphBuilder.RHICmdList, UAVDesc);
+
+	AddPass(GraphBuilder, RDG_EVENT_NAME("VirtualTextureClear"), [UAV = UAV](FRDGAsyncTask, FRHICommandList& RHICmdList)
+		{
+			// Clear virtual texture feedback to default value
+			RHICmdList.Transition(FRHITransitionInfo(UAV, ERHIAccess::Unknown, ERHIAccess::UAVCompute));
+			RHICmdList.ClearUAVUint(UAV, FUintVector4(0u, 0u, 0u, 0u));
+			RHICmdList.Transition(FRHITransitionInfo(UAV, ERHIAccess::UAVCompute, ERHIAccess::UAVMask));
+			RHICmdList.BeginUAVOverlap(UAV);
+		});
+
+	bIsInBeginEndScope = true;
+}
+
+void FVirtualTextureFeedbackBufferResource::End(FRDGBuilder& GraphBuilder)
+{
+	if (!bIsInBeginEndScope)
+	{
+		return;
+	}
+	bIsInBeginEndScope = false;
+
+	AddPass(GraphBuilder, RDG_EVENT_NAME("VirtualTextureFeedbackTransition"), [UAV = UAV](FRHICommandListImmediate& RHICmdList)
+		{
+			RHICmdList.EndUAVOverlap(UAV);
+			RHICmdList.Transition(FRHITransitionInfo(UAV, ERHIAccess::UAVMask, ERHIAccess::SRVCompute));
+		});
+
+	FRDGBufferRef FeedbackBuffer = GraphBuilder.RegisterExternalBuffer(PooledBuffer, ERDGBufferFlags::SkipTracking);
+	FRDGBufferSRVRef FeedbackBufferSRV = GraphBuilder.CreateSRV(FeedbackBuffer);
+
+	// Compact feedback before queuing for readback.
+	const uint32 CompactedFeedbackBufferSize = GetVirtualTextureCompactedFeedbackBufferSize(FeedbackBufferSize);
+	const uint32 HashTableSize = 2 * CompactedFeedbackBufferSize;
+	const uint32 HashTableIndexWrapMask = HashTableSize - 1;
+
+	FRDGBufferDesc CompactedFeedbackBufferDesc(FRDGBufferDesc::CreateStructuredDesc(sizeof(uint32), CompactedFeedbackBufferSize));
+	CompactedFeedbackBufferDesc.Usage |= BUF_SourceCopy;
+	FRDGBufferRef CompactedFeedbackBuffer = GraphBuilder.CreateBuffer(CompactedFeedbackBufferDesc, TEXT("VirtualTexture.CompactedFeedback"));
+
+	// Need to clear this buffer, as first element will be used as an allocator
+	AddClearUAVPass(GraphBuilder, GraphBuilder.CreateUAV(CompactedFeedbackBuffer, PF_R32_UINT), 0);
+
+	FRDGBufferDesc HashTableBufferDesc(FRDGBufferDesc::CreateStructuredDesc(sizeof(uint32), HashTableSize));
+	FRDGBufferRef HashTableKeyBuffer = GraphBuilder.CreateBuffer(HashTableBufferDesc, TEXT("VirtualTexture.HashTableKeys"));
+	FRDGBufferRef HashTableElementIndexBuffer = GraphBuilder.CreateBuffer(HashTableBufferDesc, TEXT("VirtualTexture.HashTableElementIndices"));
+	FRDGBufferRef HashTableElementCountBuffer = GraphBuilder.CreateBuffer(HashTableBufferDesc, TEXT("VirtualTexture.HashTableElementCounts"));
+
+	// Hash table depends on empty slots to be 0
+	AddClearUAVPass(GraphBuilder, GraphBuilder.CreateUAV(HashTableKeyBuffer, PF_R32_UINT), 0);
+	AddClearUAVPass(GraphBuilder, GraphBuilder.CreateUAV(HashTableElementCountBuffer, PF_R32_UINT), 0);
+
+	FRDGBufferRef BuildHashTableIndirectArgBuffer = GraphBuilder.CreateBuffer(FRDGBufferDesc::CreateIndirectDesc<FRHIDispatchIndirectParameters>(1), TEXT("VirtualTexture.BuildHashTableIndirectArgs"));
+
+	FGlobalShaderMap* GlobalShaderMap = GetGlobalShaderMap(FeatureLevel);
+
+	// Set indirect dispatch arguments for hash table building
+	{
+		FBuildFeedbackHashTableIndirectArgsCS::FParameters* PassParameters = GraphBuilder.AllocParameters<FBuildFeedbackHashTableIndirectArgsCS::FParameters>();
+
+		PassParameters->RWBuildHashTableIndirectArgs = GraphBuilder.CreateUAV(BuildHashTableIndirectArgBuffer, PF_R32_UINT);
+
+		PassParameters->FeedbackBufferAllocator = FeedbackBufferSRV;
+		PassParameters->FeedbackBuffer = FeedbackBufferSRV;
+		PassParameters->FeedbackBufferSize = FeedbackBufferSize;
+
+		FBuildFeedbackHashTableIndirectArgsCS::FPermutationDomain PermutationVector;
+		PermutationVector.Set<FBuildFeedbackHashTableIndirectArgsCS::FFeedbackBufferStride>(1);
+
+		auto ComputeShader = GlobalShaderMap->GetShader<FBuildFeedbackHashTableIndirectArgsCS>(PermutationVector);
+		const FIntVector GroupSize = FIntVector(1, 1, 1);
+
+		FComputeShaderUtils::AddPass(
+			GraphBuilder,
+			RDG_EVENT_NAME("Hash table indirect arguments"),
+			ComputeShader,
+			PassParameters,
+			GroupSize);
+	}
+
+	// Build hash table of feedback elements
+	{
+		FBuildFeedbackHashTableCS::FParameters* PassParameters = GraphBuilder.AllocParameters<FBuildFeedbackHashTableCS::FParameters>();
+
+		PassParameters->BuildHashTableIndirectArgs = BuildHashTableIndirectArgBuffer;
+
+		PassParameters->RWHashTableKeys = GraphBuilder.CreateUAV(HashTableKeyBuffer);
+		PassParameters->RWHashTableElementIndices = GraphBuilder.CreateUAV(HashTableElementIndexBuffer);
+		PassParameters->RWHashTableElementCounts = GraphBuilder.CreateUAV(HashTableElementCountBuffer);
+		PassParameters->HashTableSize = HashTableSize;
+		PassParameters->HashTableIndexWrapMask = HashTableIndexWrapMask;
+
+		PassParameters->FeedbackBufferAllocator = FeedbackBufferSRV;
+		PassParameters->FeedbackBuffer = FeedbackBufferSRV;
+		PassParameters->FeedbackBufferSize = FeedbackBufferSize;
+
+		FBuildFeedbackHashTableCS::FPermutationDomain PermutationVector;
+		PermutationVector.Set<FBuildFeedbackHashTableCS::FFeedbackBufferStride>(1);
+
+		auto ComputeShader = GlobalShaderMap->GetShader<FBuildFeedbackHashTableCS>(PermutationVector);
+
+		FComputeShaderUtils::AddPass(
+			GraphBuilder,
+			RDG_EVENT_NAME("Build feedback hash table"),
+			ComputeShader,
+			PassParameters,
+			BuildHashTableIndirectArgBuffer,
+			0);
+	}
+
+	// Compact hash table into an array of unique feedback elements
+	{
+		FCompactFeedbackHashTableCS::FParameters* PassParameters = GraphBuilder.AllocParameters<FCompactFeedbackHashTableCS::FParameters>();
+
+		PassParameters->RWCompactedFeedbackBuffer = GraphBuilder.CreateUAV(CompactedFeedbackBuffer);
+		PassParameters->CompactedFeedbackBufferSize = CompactedFeedbackBufferSize;
+
+		PassParameters->HashTableElementIndices = GraphBuilder.CreateSRV(HashTableElementIndexBuffer, PF_R32_UINT);
+		PassParameters->HashTableElementCounts = GraphBuilder.CreateSRV(HashTableElementCountBuffer, PF_R32_UINT);
+		PassParameters->HashTableSize = HashTableSize;
+		PassParameters->HashTableIndexWrapMask = HashTableIndexWrapMask;
+
+		PassParameters->FeedbackBufferAllocator = FeedbackBufferSRV;
+		PassParameters->FeedbackBuffer = FeedbackBufferSRV;
+		PassParameters->FeedbackBufferSize = FeedbackBufferSize;
+
+		FCompactFeedbackHashTableCS::FPermutationDomain PermutationVector;
+		PermutationVector.Set<FCompactFeedbackHashTableCS::FFeedbackBufferStride>(1);
+
+		auto ComputeShader = GlobalShaderMap->GetShader<FCompactFeedbackHashTableCS>(PermutationVector);
+		const FIntVector GroupSize = FComputeShaderUtils::GetGroupCount(HashTableSize, FCompactFeedbackHashTableCS::GetGroupSize());
+
+		FComputeShaderUtils::AddPass(
+			GraphBuilder,
+			RDG_EVENT_NAME("Compact feedback hash table"),
+			ComputeShader,
+			PassParameters,
+			GroupSize);
+	}
+
+	FVirtualTextureFeedbackBufferDesc CompactedDesc;
+	CompactedDesc.Init(CompactedFeedbackBufferSize, true);
+	SubmitVirtualTextureFeedbackBuffer(GraphBuilder, CompactedFeedbackBuffer, CompactedDesc);
+}
+
+uint32 FVirtualTextureFeedbackBufferResource::GetBufferSize() const
+{
+	return bIsInBeginEndScope ? FeedbackBufferSize : 0;
+}
+
+FRHIUnorderedAccessView* FVirtualTextureFeedbackBufferResource::GetUAV() const
+{
+	return UAV != nullptr && bIsInBeginEndScope ? UAV : GEmptyStructuredBufferWithUAV->UnorderedAccessViewRHI.GetReference();
+}
+
+void FVirtualTextureFeedbackBufferResource::ReleaseRHI()
+{
+	PooledBuffer = nullptr;
+	UAV = nullptr;
+}
+
+/**
+ * Get the feedback buffer size in dwords based on view size.
+ * Takes into account the feedback tile size and overdraw factors.
+ */
+static uint32 GetVirtualTextureFeedbackBufferSize(FIntPoint InViewportSize, uint32 InVirtualTextureFeedbackTileSize)
+{
+	// Tile size must be power or 2.
+	check(InVirtualTextureFeedbackTileSize == 1 << FMath::FloorLog2(InVirtualTextureFeedbackTileSize));
+
+	FIntPoint ScaledExtent = FIntPoint::DivideAndRoundUp(InViewportSize, InVirtualTextureFeedbackTileSize);
+	return ScaledExtent.X * ScaledExtent.Y * FMath::Max(GVirtualTextureFeedbackOverdrawFactor, 1);
+}
+
+/** Get compacted feedback buffer size in dwords based on the original feedback buffer size. */
+static uint32 GetVirtualTextureCompactedFeedbackBufferSize(uint32 InSourceBufferSize)
+{
+	// Could possibly do this dynamically according to some tracked recent high watermark?
+	const uint32 CompactedBufferSize = InSourceBufferSize / FMath::Max(GVirtualTextureFeedbackCompactionFactor, 1);
+	// Size needs to be power of two for hash table wrapping.
+	const uint32 AlignedCompactedBufferSize = FMath::RoundUpToPowerOfTwo(FMath::Clamp(CompactedBufferSize, 16u, 16u * 1024u));
+	return AlignedCompactedBufferSize;
+}
+
+/** Apply alignment rules to the feedback tile size. */
+static uint32 AlignVirtualTextureFeedbackTileSize(uint32 InTileSize)
+{
+	// Round to nearest power of two to ensure that shader maths is efficient and sampling sequence logic is simple.
+	return FMath::RoundUpToPowerOfTwo(FMath::Max(InTileSize, 1u));
+}
+
+/** Get jittered pixel index within feedback tile. */
+static uint32 SampleVirtualTextureFeedbackSequence(uint32 InFrameIndex, uint32 InVirtualTextureFeedbackTileSize)
+{
+	// Tile size must be power or 2.
+	check(InVirtualTextureFeedbackTileSize == 1 << FMath::FloorLog2(InVirtualTextureFeedbackTileSize));
+
+	const uint32 TileSize = InVirtualTextureFeedbackTileSize;
+	const uint32 TileSizeLog2 = FMath::CeilLogTwo(TileSize);
+	const uint32 SequenceSize = FMath::Square(TileSize);
+	const uint32 PixelIndex = InFrameIndex % SequenceSize;
+	const uint32 PixelAddress = ReverseBits(PixelIndex) >> (32U - 2 * TileSizeLog2);
+	const uint32 X = FMath::ReverseMortonCode2(PixelAddress);
+	const uint32 Y = FMath::ReverseMortonCode2(PixelAddress >> 1);
+	const uint32 PixelSequenceIndex = X + Y * TileSize;
+	return PixelSequenceIndex;
+}
+
+namespace VirtualTexture {
+
+void GetFeedbackShaderParams(int32 InFrameIndex, int32 InVirtualTextureFeedbackTileSize, FFeedbackShaderParams& OutParams)
+{
+	OutParams.BufferUAV = GVirtualTextureFeedbackBufferResource.GetUAV();
+	OutParams.BufferSize = GVirtualTextureFeedbackBufferResource.GetBufferSize();
+
+	// Round to nearest power of two to ensure that shader maths is efficient and sampling sequence logic is simple.
+	InVirtualTextureFeedbackTileSize = AlignVirtualTextureFeedbackTileSize(InVirtualTextureFeedbackTileSize);
+
+	OutParams.Shift = FMath::FloorLog2(InVirtualTextureFeedbackTileSize);
+	OutParams.Mask = InVirtualTextureFeedbackTileSize - 1;
+	// Use some low(ish) discrepancy sequence to run over every pixel in the virtual texture feedback tile.
+	OutParams.JitterOffset = SampleVirtualTextureFeedbackSequence(InFrameIndex, InVirtualTextureFeedbackTileSize);
+	// SampleOffset is used to cycle through all VT samples in a material. It just needs to monotonically increase.
+	OutParams.SampleOffset = InFrameIndex;
+}
+
+void GetFeedbackShaderParams(FFeedbackShaderParams& OutParams)
+{
+	const uint32 VirtualTextureFeedbackTileSize = AlignVirtualTextureFeedbackTileSize(VirtualTextureScalability::GetVirtualTextureFeedbackFactor());
+	VirtualTexture::GetFeedbackShaderParams(GFrameNumber, VirtualTextureFeedbackTileSize, OutParams);
+}
+
+void UpdateViewUniformShaderParameters(FFeedbackShaderParams const& InParams, FViewUniformShaderParameters& ViewUniformShaderParameters)
+{
+	ViewUniformShaderParameters.VTFeedbackBuffer = InParams.BufferUAV;
+	ViewUniformShaderParameters.VirtualTextureFeedbackBufferSize = InParams.BufferSize;
+	ViewUniformShaderParameters.VirtualTextureFeedbackShift = InParams.Shift;
+	ViewUniformShaderParameters.VirtualTextureFeedbackMask = InParams.Mask;
+	ViewUniformShaderParameters.VirtualTextureFeedbackJitterOffset = InParams.JitterOffset;
+	ViewUniformShaderParameters.VirtualTextureFeedbackSampleOffset = InParams.SampleOffset;
+}
+
+void BeginFeedback(FRDGBuilder& GraphBuilder, uint32 InBufferSize, ERHIFeatureLevel::Type InFeatureLevel)
+{
+	GVirtualTextureFeedbackBufferResource.Begin(GraphBuilder, InBufferSize, InFeatureLevel);
+}
+
+void BeginFeedback(FRDGBuilder& GraphBuilder, uint32 InBufferSize)
+{
+	BeginFeedback(GraphBuilder, InBufferSize, GMaxRHIFeatureLevel);
+}
+
+void BeginFeedback(FRDGBuilder& GraphBuilder, FIntPoint InViewportSize, uint32 InVirtualTextureFeedbackTileSize, ERHIFeatureLevel::Type InFeatureLevel)
+{
+	const int32 FeedbackBufferSize = GetVirtualTextureFeedbackBufferSize(InViewportSize, AlignVirtualTextureFeedbackTileSize(InVirtualTextureFeedbackTileSize));
+	BeginFeedback(GraphBuilder, FeedbackBufferSize, InFeatureLevel);
+
+}
+
+void EndFeedback(FRDGBuilder& GraphBuilder)
+{
+	GVirtualTextureFeedbackBufferResource.End(GraphBuilder);
+}
+
+} // namespace VirtualTexture
diff --git a/Engine/Source/Runtime/Renderer/Public/VT/VirtualTextureFeedbackBuffer.h b/Engine/Source/Runtime/Renderer/Public/VT/VirtualTextureFeedbackBuffer.h
--- a/Engine/Source/Runtime/Renderer/Public/VT/VirtualTextureFeedbackBuffer.h
+++ b/Engine/Source/Runtime/Renderer/Public/VT/VirtualTextureFeedbackBuffer.h
@@ -2,23 +2,12 @@
 
 #pragma once
 
-#include "Containers/ArrayView.h"
-#include "Math/IntPoint.h"
-#include "Math/IntRect.h"
-#include "Math/UnrealMathSSE.h"
+#include "CoreMinimal.h"
 #include "RHI.h"
-#include "RenderGraphResources.h"
-#include "RenderResource.h"
-#include "RenderUtils.h"
-#include "Templates/RefCounting.h"
 
+class FRDGBuffer;
 class FRDGBuilder;
-
-/** Get size for a virtual texture feedback buffer. */
-extern uint32 GetVirtualTextureFeedbackBufferSize(FIntPoint InSceneTextureExtent, uint32 InVirtualTextureTileSize);
-
-/** Returns a value from the sampling sequence used to traverse each feedback tile. */
-extern uint32 SampleVirtualTextureFeedbackSequence(uint32 InFrameIndex, uint32 InVirtualTextureFeedbackTileSize);
+class FRHICommandList;
 
 /** 
  * Description of how to interpret an RHIBuffer that is being fed to the virtual texture feedback system.
@@ -39,33 +28,6 @@ struct FVirtualTextureFeedbackBufferDesc
 	bool bSizeInHeader = false;
 };
 
-/** GPU feedback buffer for tracking virtual texture requests from the GPU. */
-class FVirtualTextureFeedbackBuffer : public FRenderResource
-{
-public:
-	// Allocates and prepares the feedback buffer for write access.
-	void Begin(FRDGBuilder& GraphBuilder, uint32 InFeedbackBufferSize, ERHIFeatureLevel::Type InFeatureLevel);
-
-	// Submits the feedback buffer for readback.
-	void End(FRDGBuilder& GraphBuilder);
-
-	// Get the size of the feedback buffer from the last call to Begin().
-	uint32 GetBufferSize() const;
-	// Get the UAV of the feedback buffer.
-	FRHIUnorderedAccessView* GetUAV() const;
-
-	//~ Begin FRenderResource Interface
-	virtual void ReleaseRHI() override;
-	//~ End FRenderResource Interface
-
-private:
-	bool bIsInBeginEndScope = false;
-	uint32 FeedbackBufferSize = 0;
-	ERHIFeatureLevel::Type FeatureLevel = ERHIFeatureLevel::Num;
-	TRefCountPtr<FRDGPooledBuffer> PooledBuffer;
-	FRHIUnorderedAccessView* UAV = nullptr;
-};
-
 /** 
  * Submit an RHIBuffer containing virtual texture feedback data to the virtual texture system.
  * The buffer is internally copied to the CPU and parsed to determine which virtual texture pages need to be mapped.
@@ -73,5 +35,5 @@ private:
  * Multiple buffers can be transferred per frame using this function.
  * The function can be called from the render thread only.
 */
-RENDERER_API void SubmitVirtualTextureFeedbackBuffer(class FRHICommandList& RHICmdList, FBufferRHIRef const& InBuffer, FVirtualTextureFeedbackBufferDesc const& InDesc);
-RENDERER_API void SubmitVirtualTextureFeedbackBuffer(class FRDGBuilder& GraphBuilder, class FRDGBuffer* InBuffer, FVirtualTextureFeedbackBufferDesc const& InDesc);
\ No newline at end of file
+RENDERER_API void SubmitVirtualTextureFeedbackBuffer(FRHICommandList& RHICmdList, FBufferRHIRef const& InBuffer, FVirtualTextureFeedbackBufferDesc const& InDesc);
+RENDERER_API void SubmitVirtualTextureFeedbackBuffer(FRDGBuilder& GraphBuilder, FRDGBuffer* InBuffer, FVirtualTextureFeedbackBufferDesc const& InDesc);
\ No newline at end of file
