Add some stats to virtual texture feedback system.
Summon with stat virtualtexturing.
Shows number of feedback buffer read/writes and number of pages lost due to buffer compaction.

diff --git a/Engine/Source/Runtime/Renderer/Private/VT/VirtualTextureFeedback.cpp b/Engine/Source/Runtime/Renderer/Private/VT/VirtualTextureFeedback.cpp
--- a/Engine/Source/Runtime/Renderer/Private/VT/VirtualTextureFeedback.cpp
+++ b/Engine/Source/Runtime/Renderer/Private/VT/VirtualTextureFeedback.cpp
@@ -4,6 +4,7 @@
 
 #include "RenderGraphBuilder.h"
 #include "RenderGraphUtils.h"
+#include "VirtualTexturing.h"
 
 int32 GVirtualTextureFeedbackLatency = 3;
 static FAutoConsoleVariableRef CVarVirtualTextureFeedbackLatency(
@@ -12,6 +13,12 @@ static FAutoConsoleVariableRef CVarVirtualTextureFeedbackLatency(
 	TEXT("How much latency to allow in the GPU feedback pipeline before we start mapping multiple buffers to catch up."),
 	ECVF_RenderThreadSafe);
 
+DECLARE_DWORD_COUNTER_STAT(TEXT("Num Feedback Pending"), STAT_VirtualTexture_PendingFeedback, STATGROUP_VirtualTexturing);
+DECLARE_DWORD_COUNTER_STAT(TEXT("Num Feedback Writes"), STAT_VirtualTexture_WriteFeedback, STATGROUP_VirtualTexturing);
+DECLARE_DWORD_COUNTER_STAT(TEXT("Num Feedback Reads"), STAT_VirtualTexture_ReadFeedback, STATGROUP_VirtualTexturing);
+DECLARE_DWORD_COUNTER_STAT(TEXT("Num Feedback Lost Buffers"), STAT_VirtualTexture_LostFeedback, STATGROUP_VirtualTexturing);
+DECLARE_DWORD_COUNTER_STAT(TEXT("Num Feedback Lost Pages"), STAT_VirtualTexture_ReadFeedbackLostPage, STATGROUP_VirtualTexturing);
+
 /** Container for GPU fences. */
 class FFeedbackGPUFencePool
 {
@@ -107,10 +114,14 @@ void FVirtualTextureFeedback::TransferGPUToCPU(FRHICommandList& RHICmdList, FBuf
 		return;
 	}
 
+	INC_DWORD_STAT(STAT_VirtualTexture_WriteFeedback);
+
 	if (NumPending >= MaxTransfers)
 	{
 		// If we have too many pending transfers, start throwing away the oldest in the ring buffer.
 		// We will need to allocate a new fence, since the previous fence will still be set on the old CopyToResolveTarget command (which we will now ignore/discard).
+		INC_DWORD_STAT(STAT_VirtualTexture_LostFeedback);
+
 		Fences->Release(ReadIndex);
 		NumPending --;
 		ReadIndex = (ReadIndex + 1) % MaxTransfers;
@@ -209,8 +220,10 @@ FVirtualTextureFeedback::FMapResult FVirtualTextureFeedback::Map(FRHICommandList
 
 			if (FeedbackItemDesc.bSizeInHeader)
 			{
-				MapResult.Size = FMath::Min(MapResult.Data[0], FeedbackItemDesc.BufferSize - 1);
-				MapResult.Data += 1;
+				const uint32 BufferWriteCount = *MapResult.Data++;
+				MapResult.Size = FMath::Min(BufferWriteCount, FeedbackItemDesc.BufferSize - 1);
+
+				INC_DWORD_STAT_BY(STAT_VirtualTexture_ReadFeedbackLostPage, BufferWriteCount - MapResult.Size);
 			}
 
 			// Store index so that we can unlock staging buffer when we call Unmap().
@@ -240,17 +253,22 @@ FVirtualTextureFeedback::FMapResult FVirtualTextureFeedback::Map(FRHICommandList
 				}
 				else
 				{
-					const int32 BufferSize = FMath::Min(Data[0], FeedbackItemDesc.BufferSize - 1);
-					FMemory::Memcpy(MapResult.Data + MapResult.Size, Data + 1, BufferSize * sizeof(uint32));
+					const uint32 BufferWriteCount = *Data++;
+					const int32 BufferSize = FMath::Min(BufferWriteCount, FeedbackItemDesc.BufferSize - 1);
+					FMemory::Memcpy(MapResult.Data + MapResult.Size, Data, BufferSize * sizeof(uint32));
 					MapResult.Size += BufferSize;
+
+					INC_DWORD_STAT_BY(STAT_VirtualTexture_ReadFeedbackLostPage, BufferWriteCount - BufferSize);
 				}
 
 				RHICmdList.UnlockStagingBuffer(StagingBuffer);
 			}
 		}
 
+		INC_DWORD_STAT_BY(STAT_VirtualTexture_ReadFeedback, NumResults);
+
 		check(MapResult.Size <= TotalReadSize)
-	
+
 		// Increment the ring buffer read position.
 		NumPending -= NumResults;
 		ReadIndex = (ReadIndex + NumResults) % MaxTransfers;
@@ -261,6 +279,9 @@ FVirtualTextureFeedback::FMapResult FVirtualTextureFeedback::Map(FRHICommandList
 
 FVirtualTextureFeedback::FMapResult FVirtualTextureFeedback::Map(FRHICommandListImmediate& RHICmdList)
 {
+	// Note that this stat for pending could vary over the frame, particularly if we Map() more than once.
+	SET_DWORD_STAT(STAT_VirtualTexture_PendingFeedback, NumPending);
+
 	// Allow some slack in the pipeline before we start mapping more than one buffer.
 	// Otherwise we can get into an oscillating pattern of mapping 2 buffers, then 0 buffers, then 2 again etc.
 	const uint32 MaxTransfersToMap = NumPending < GVirtualTextureFeedbackLatency ? 1 : MaxTransfers;
