Small change to Virtual Texture feedback API to allow for a default feedback buffer size r.vt.FeedbackDefaultBufferSize
This will only be used if the code does not have a valid screen size to drive the calculation.

diff --git a/Engine/Source/Runtime/Renderer/Internal/VT/VirtualTextureFeedbackResource.h b/Engine/Source/Runtime/Renderer/Internal/VT/VirtualTextureFeedbackResource.h
--- a/Engine/Source/Runtime/Renderer/Internal/VT/VirtualTextureFeedbackResource.h
+++ b/Engine/Source/Runtime/Renderer/Internal/VT/VirtualTextureFeedbackResource.h
@@ -42,9 +42,9 @@ namespace VirtualTexture
 	{
 		FRHIUnorderedAccessView* BufferUAV = nullptr;
 		uint32 BufferSize = 0;
-		uint32 Shift = 0;
-		uint32 Mask = 0;
-		uint32 JitterOffset = 0;
+		uint32 TileShift = 0;
+		uint32 TileMask = 0;
+		uint32 TileJitterOffset = 0;
 		uint32 SampleOffset = 0;
 	};
 
@@ -52,8 +52,7 @@ namespace VirtualTexture
 	 * Begin a virtual texture feedback scope. 
 	 * Note that this uses a global FVirtualTextureFeedbackBufferResource object which is then read when calling GetFeedbackShaderParams().
 	 */
-	RENDERER_API void BeginFeedback(FRDGBuilder& GraphBuilder, uint32 InBufferSize);
-	RENDERER_API void BeginFeedback(FRDGBuilder& GraphBuilder, uint32 InBufferSize, ERHIFeatureLevel::Type InFeatureLevel);
+	RENDERER_API void BeginFeedback(FRDGBuilder& GraphBuilder, uint32 InBufferSize = 0, ERHIFeatureLevel::Type InFeatureLevel = GMaxRHIFeatureLevel);
 	void BeginFeedback(FRDGBuilder& GraphBuilder, FIntPoint InViewportSize, uint32 InVirtualTextureFeedbackTileSize, ERHIFeatureLevel::Type InFeatureLevel);
 
 	/** End a virtual texture feedback scope. */
diff --git a/Engine/Source/Runtime/Renderer/Private/SceneRendering.cpp b/Engine/Source/Runtime/Renderer/Private/SceneRendering.cpp
--- a/Engine/Source/Runtime/Renderer/Private/SceneRendering.cpp
+++ b/Engine/Source/Runtime/Renderer/Private/SceneRendering.cpp
@@ -6052,11 +6052,9 @@ void AddResolveSceneDepthPass(FRDGBuilder& GraphBuilder, TArrayView<const FViewI
 
 void VirtualTextureFeedbackBegin(FRDGBuilder& GraphBuilder, TArrayView<const FViewInfo> Views, FIntPoint SceneTextureExtent)
 {
-	const FIntPoint FeedbackViewSize = FIntPoint(FMath::Max(SceneTextureExtent.X, 1), FMath::Max(SceneTextureExtent.Y, 1));
-	const uint32 FeedbackTileSize = Views.Num() > 0 ? Views[0].Family->VirtualTextureFeedbackFactor : VirtualTextureScalability::GetVirtualTextureFeedbackFactor();
+	const uint32 FeedbackTileSize = Views.Num() > 0 ? Views[0].Family->VirtualTextureFeedbackFactor : 0;
 	const ERHIFeatureLevel::Type FeatureLevel = Views.Num() > 0 ? Views[0].GetFeatureLevel() : GMaxRHIFeatureLevel;
-	
-	VirtualTexture::BeginFeedback(GraphBuilder, FeedbackViewSize, FeedbackTileSize, FeatureLevel);
+	VirtualTexture::BeginFeedback(GraphBuilder, SceneTextureExtent, FeedbackTileSize, FeatureLevel);
 }
 
 static TAutoConsoleVariable<int32> CVarHalfResDepthNoFastClear(
diff --git a/Engine/Source/Runtime/Renderer/Private/VT/VirtualTextureFeedbackResource.cpp b/Engine/Source/Runtime/Renderer/Private/VT/VirtualTextureFeedbackResource.cpp
--- a/Engine/Source/Runtime/Renderer/Private/VT/VirtualTextureFeedbackResource.cpp
+++ b/Engine/Source/Runtime/Renderer/Private/VT/VirtualTextureFeedbackResource.cpp
@@ -10,6 +10,13 @@
 #include "VT/VirtualTextureScalability.h"
 #include "VT/VirtualTextureSystem.h"
 
+int32 GVirtualTextureFeedbackDefaultBufferSize = 4 * 1024;
+static FAutoConsoleVariableRef CVarVirtualTextureFeedbackDefaultBufferSize(
+	TEXT("r.vt.FeedbackDefaultBufferSize"),
+	GVirtualTextureFeedbackDefaultBufferSize,
+	TEXT("Virtual texture feedback buffer size for cases where we don't calculate it based on screen size."),
+	ECVF_RenderThreadSafe);
+
 int32 GVirtualTextureFeedbackOverdrawFactor = 2;
 static FAutoConsoleVariableRef CVarVirtualTextureFeedbackOverdrawFactor(
 	TEXT("r.vt.FeedbackOverdrawFactor"),
@@ -264,10 +271,10 @@ void GetFeedbackShaderParams(int32 InFrameIndex, int32 InVirtualTextureFeedbackT
 	// Round to nearest power of two to ensure that shader maths is efficient and sampling sequence logic is simple.
 	InVirtualTextureFeedbackTileSize = AlignVirtualTextureFeedbackTileSize(InVirtualTextureFeedbackTileSize);
 
-	OutParams.Shift = FMath::FloorLog2(InVirtualTextureFeedbackTileSize);
-	OutParams.Mask = InVirtualTextureFeedbackTileSize - 1;
+	OutParams.TileShift = FMath::FloorLog2(InVirtualTextureFeedbackTileSize);
+	OutParams.TileMask = InVirtualTextureFeedbackTileSize - 1;
 	// Use some low(ish) discrepancy sequence to run over every pixel in the virtual texture feedback tile.
-	OutParams.JitterOffset = SampleVirtualTextureFeedbackSequence(InFrameIndex, InVirtualTextureFeedbackTileSize);
+	OutParams.TileJitterOffset = SampleVirtualTextureFeedbackSequence(InFrameIndex, InVirtualTextureFeedbackTileSize);
 	// SampleOffset is used to cycle through all VT samples in a material. It just needs to monotonically increase.
 	OutParams.SampleOffset = InFrameIndex;
 }
@@ -282,27 +289,26 @@ void UpdateViewUniformShaderParameters(FFeedbackShaderParams const& InParams, FV
 {
 	ViewUniformShaderParameters.VTFeedbackBuffer = InParams.BufferUAV;
 	ViewUniformShaderParameters.VirtualTextureFeedbackBufferSize = InParams.BufferSize;
-	ViewUniformShaderParameters.VirtualTextureFeedbackShift = InParams.Shift;
-	ViewUniformShaderParameters.VirtualTextureFeedbackMask = InParams.Mask;
-	ViewUniformShaderParameters.VirtualTextureFeedbackJitterOffset = InParams.JitterOffset;
+	ViewUniformShaderParameters.VirtualTextureFeedbackShift = InParams.TileShift;
+	ViewUniformShaderParameters.VirtualTextureFeedbackMask = InParams.TileMask;
+	ViewUniformShaderParameters.VirtualTextureFeedbackJitterOffset = InParams.TileJitterOffset;
 	ViewUniformShaderParameters.VirtualTextureFeedbackSampleOffset = InParams.SampleOffset;
 }
 
 void BeginFeedback(FRDGBuilder& GraphBuilder, uint32 InBufferSize, ERHIFeatureLevel::Type InFeatureLevel)
 {
-	GVirtualTextureFeedbackBufferResource.Begin(GraphBuilder, InBufferSize, InFeatureLevel);
-}
-
-void BeginFeedback(FRDGBuilder& GraphBuilder, uint32 InBufferSize)
-{
-	BeginFeedback(GraphBuilder, InBufferSize, GMaxRHIFeatureLevel);
+	const uint32 BufferSize = InBufferSize > 0 ? InBufferSize : GVirtualTextureFeedbackDefaultBufferSize;
+	const ERHIFeatureLevel::Type FeatureLevel = InFeatureLevel < ERHIFeatureLevel::Num ? InFeatureLevel : GMaxRHIFeatureLevel;
+	GVirtualTextureFeedbackBufferResource.Begin(GraphBuilder, BufferSize, FeatureLevel);
 }
 
 void BeginFeedback(FRDGBuilder& GraphBuilder, FIntPoint InViewportSize, uint32 InVirtualTextureFeedbackTileSize, ERHIFeatureLevel::Type InFeatureLevel)
 {
-	const int32 FeedbackBufferSize = GetVirtualTextureFeedbackBufferSize(InViewportSize, AlignVirtualTextureFeedbackTileSize(InVirtualTextureFeedbackTileSize));
-	BeginFeedback(GraphBuilder, FeedbackBufferSize, InFeatureLevel);
-
+	const FIntPoint ViewportSize = FIntPoint(FMath::Max(InViewportSize.X, 1), FMath::Max(InViewportSize.Y, 1));
+	const uint32 TileSize = InVirtualTextureFeedbackTileSize > 0 ? InVirtualTextureFeedbackTileSize : VirtualTextureScalability::GetVirtualTextureFeedbackFactor();
+	const uint32 AlignedTileSize = AlignVirtualTextureFeedbackTileSize(TileSize);
+	const uint32 BufferSize = GetVirtualTextureFeedbackBufferSize(ViewportSize, AlignedTileSize);
+	BeginFeedback(GraphBuilder, BufferSize, InFeatureLevel);
 }
 
 void EndFeedback(FRDGBuilder& GraphBuilder)
diff --git a/Engine/Source/Runtime/SlateRHIRenderer/Private/SlateShaders.cpp b/Engine/Source/Runtime/SlateRHIRenderer/Private/SlateShaders.cpp
--- a/Engine/Source/Runtime/SlateRHIRenderer/Private/SlateShaders.cpp
+++ b/Engine/Source/Runtime/SlateRHIRenderer/Private/SlateShaders.cpp
@@ -179,9 +179,9 @@ void FSlateElementPS::SetVirtualTextureParameters(FMeshDrawSingleShaderBindings&
 	VirtualTexture::GetFeedbackShaderParams(VirtualTextureFeedbackShaderParams);
 
 	const FUintVector4 FeedbackParams{ 
-		VirtualTextureFeedbackShaderParams.Mask, 
-		VirtualTextureFeedbackShaderParams.Shift, 
-		VirtualTextureFeedbackShaderParams.JitterOffset, 
+		VirtualTextureFeedbackShaderParams.TileMask, 
+		VirtualTextureFeedbackShaderParams.TileShift, 
+		VirtualTextureFeedbackShaderParams.TileJitterOffset, 
 		VirtualTextureFeedbackShaderParams.BufferSize };
 
 	FSlateElementVirtualTextureParams VTParams;
