Move FVirtualTextureFeedbackBufferResource into private cpp file implementation.
If we want to expose more than one object in future we can do it via allocating to an internal array and returning handles.
Added some better comments.

diff --git a/Engine/Source/Runtime/Renderer/Internal/VT/VirtualTextureFeedbackResource.h b/Engine/Source/Runtime/Renderer/Internal/VT/VirtualTextureFeedbackResource.h
--- a/Engine/Source/Runtime/Renderer/Internal/VT/VirtualTextureFeedbackResource.h
+++ b/Engine/Source/Runtime/Renderer/Internal/VT/VirtualTextureFeedbackResource.h
@@ -3,38 +3,11 @@
 #pragma once
 
 #include "CoreMinimal.h"
-#include "RenderResource.h"
 #include "RHI.h"
 
 class FRDGBuilder;
 class FViewUniformShaderParameters;
 
-/** GPU feedback buffer for tracking virtual texture requests from the GPU. */
-class FVirtualTextureFeedbackBufferResource : public FRenderResource
-{
-public:
-	/** Allocates and prepares a new feedback buffer for write access. */
-	void Begin(FRDGBuilder& GraphBuilder, uint32 InFeedbackBufferSize, ERHIFeatureLevel::Type InFeatureLevel);
-	/** Compacts the feedback buffer on the GPU and copies the compacted buffer for readback. */
-	void End(FRDGBuilder& GraphBuilder);
-
-	/** Get the size of the feedback buffer from the last call to Begin(). */
-	uint32 GetBufferSize() const;
-	/** Get the UAV of the feedback buffer. */
-	FRHIUnorderedAccessView* GetUAV() const;
-
-	//~ Begin FRenderResource Interface
-	virtual void ReleaseRHI() override;
-	//~ End FRenderResource Interface
-
-private:
-	bool bIsInBeginEndScope = false;
-	uint32 FeedbackBufferSize = 0;
-	ERHIFeatureLevel::Type FeatureLevel = ERHIFeatureLevel::Num;
-	TRefCountPtr<FRDGPooledBuffer> PooledBuffer;
-	FRHIUnorderedAccessView* UAV = nullptr;
-};
-
 namespace VirtualTexture
 {
 	/** Shader parameters needed when writing virtual texture feedback from GPU. */
@@ -50,19 +23,31 @@ namespace VirtualTexture
 
 	/** 
 	 * Begin a virtual texture feedback scope. 
-	 * Note that this uses a global FVirtualTextureFeedbackBufferResource object which is then read when calling GetFeedbackShaderParams().
+	 * The feedback buffer size will be InBufferSize. If InBufferSize is 0 a default size is used.
+	 * Note that there is only one global feedback buffer alive at any time and its buffer size is later used internally by GetFeedbackShaderParams().
 	 */
 	RENDERER_API void BeginFeedback(FRDGBuilder& GraphBuilder, uint32 InBufferSize = 0, ERHIFeatureLevel::Type InFeatureLevel = GMaxRHIFeatureLevel);
+
+	/** 
+	 * Begin a virtual texture feedback scope. 
+	 * InViewportSize and InVirtualTextureFeedbackTileSize are used to calculate the feedback buffer size. 
+	 * Note that there is only one global feedback buffer alive at any time and its buffer size is later used internally by GetFeedbackShaderParams().
+	 */
 	void BeginFeedback(FRDGBuilder& GraphBuilder, FIntPoint InViewportSize, uint32 InVirtualTextureFeedbackTileSize, ERHIFeatureLevel::Type InFeatureLevel);
 
 	/** End a virtual texture feedback scope. */
 	RENDERER_API void EndFeedback(FRDGBuilder& GraphBuilder);
 
-	/** Get virtual texture feedback parameters to use for binding to any shader that is sampling virtual textures. */
+	/** 
+	 * Get virtual texture feedback parameters to use for binding to any shader that is sampling virtual textures. 
+	 * This is only valid within a BeginFeedback()/EndFeedback() scope.
+	 */
 	void GetFeedbackShaderParams(int32 InFrameIndex, int32 InVirtualTextureFeedbackTileSize, FFeedbackShaderParams& OutParams);
+	
 	/** 
 	 * Get virtual texture feedback parameters to use for binding to any shader that is sampling virtual textures. 
 	 * This version uses an internal scene frame counter, and the project default feedback tile size.
+	 * This is only valid within a BeginFeedback()/EndFeedback() scope.
 	 */
 	RENDERER_API void GetFeedbackShaderParams(FFeedbackShaderParams& OutParams);
 	
diff --git a/Engine/Source/Runtime/Renderer/Private/VT/VirtualTextureFeedbackResource.cpp b/Engine/Source/Runtime/Renderer/Private/VT/VirtualTextureFeedbackResource.cpp
--- a/Engine/Source/Runtime/Renderer/Private/VT/VirtualTextureFeedbackResource.cpp
+++ b/Engine/Source/Runtime/Renderer/Private/VT/VirtualTextureFeedbackResource.cpp
@@ -5,6 +5,7 @@
 #include "GlobalRenderResources.h"
 #include "GPUFeedbackCompaction.h"
 #include "RenderGraphUtils.h"
+#include "RenderResource.h"
 #include "SceneView.h"
 #include "VT/VirtualTextureFeedbackBuffer.h"
 #include "VT/VirtualTextureScalability.h"
@@ -33,11 +34,37 @@ static FAutoConsoleVariableRef CVarVirtualTextureFeedbackCompactionFactor(
 	TEXT("to account for compaction of duplicate page ids."),
 	ECVF_RenderThreadSafe);
 
-static uint32 GetVirtualTextureCompactedFeedbackBufferSize(uint32 InSourceBufferSize);
+/** GPU feedback buffer for tracking virtual texture requests from the GPU. */
+class FVirtualTextureFeedbackBufferResource : public FRenderResource
+{
+public:
+	/** Allocates and prepares a new feedback buffer for write access. */
+	void Begin(FRDGBuilder& GraphBuilder, uint32 InFeedbackBufferSize, ERHIFeatureLevel::Type InFeatureLevel);
+	/** Compacts the feedback buffer on the GPU and copies the compacted buffer for readback. */
+	void End(FRDGBuilder& GraphBuilder);
+
+	/** Get the size of the feedback buffer from the last call to Begin(). */
+	uint32 GetBufferSize() const;
+	/** Get the UAV of the feedback buffer. */
+	FRHIUnorderedAccessView* GetUAV() const;
+
+	//~ Begin FRenderResource Interface
+	virtual void ReleaseRHI() override;
+	//~ End FRenderResource Interface
+
+private:
+	bool bIsInBeginEndScope = false;
+	uint32 FeedbackBufferSize = 0;
+	ERHIFeatureLevel::Type FeatureLevel = ERHIFeatureLevel::Num;
+	TRefCountPtr<FRDGPooledBuffer> PooledBuffer;
+	FRHIUnorderedAccessView* UAV = nullptr;
+};
 
 /** Single global feedback buffer resource used when calling the BeginFeedback()/EndFeedback free functions. */
 static TGlobalResource<FVirtualTextureFeedbackBufferResource> GVirtualTextureFeedbackBufferResource;
 
+static uint32 GetVirtualTextureCompactedFeedbackBufferSize(uint32 InSourceBufferSize);
+
 void FVirtualTextureFeedbackBufferResource::Begin(FRDGBuilder& GraphBuilder, uint32 InFeedbackBufferSize, ERHIFeatureLevel::Type InFeatureLevel)
 {
 	// NOTE: Transitions and allocations are handled manually right now, because the VT feedback UAV is used by
