Early out of feedback page processing for adaptive page requests. We only need to store the adaptive request for processing and then we're done. By definition these pages aren't mappable until the adaptive request completes.

diff --git a/Engine/Source/Runtime/Renderer/Private/VT/VirtualTextureSystem.cpp b/Engine/Source/Runtime/Renderer/Private/VT/VirtualTextureSystem.cpp
--- a/Engine/Source/Runtime/Renderer/Private/VT/VirtualTextureSystem.cpp
+++ b/Engine/Source/Runtime/Renderer/Private/VT/VirtualTextureSystem.cpp
@@ -1575,6 +1575,7 @@ void FVirtualTextureSystem::GatherRequestsTask(const FGatherRequestsParameters&
 			if (AdaptiveAllocationRequest != 0)
 			{
 				RequestList->AddAdaptiveAllocationRequest(AdaptiveAllocationRequest);
+				continue;
 			}
 		}
 
