[HWRT] Support outputting SceneInstanceIndex when building ray tracing instance buffer.
- This is necessary for debug visualizations and other systems (such as tracing feedback) to work when instance buffer is compacted after GPU culling, since we need a way to map from compacted instances to the scene instances in FRayTracingScene.
- Replace the exisiting InstanceGPUSceneIndices with a new InstanceExtraDataBuffer which is filled by RayTracingBuildInstanceBufferCS (so it matches the instance buffer order) and includes both GPUSceneInstanceId and the SceneInstanceIndex.
- Simplified InstanceDebugBuffer so it only includes one entry per "SceneInstance" instead of "NativeInstance" since we can now use FRayTracingInstanceExtraData to map from "NativeInstance" to "SceneInstance".

diff --git a/Engine/Shaders/Private/RayTracing/RayTracingDebug.usf b/Engine/Shaders/Private/RayTracing/RayTracingDebug.usf
--- a/Engine/Shaders/Private/RayTracing/RayTracingDebug.usf
+++ b/Engine/Shaders/Private/RayTracing/RayTracingDebug.usf
@@ -43,6 +43,7 @@ float TriangleHitCountPerInstanceMaxThreshold;
 uint TopKMostHitInstances;
 uint NumTotalInstances;
 
+StructuredBuffer<FRayTracingInstanceExtraData> InstancesExtraData;
 StructuredBuffer<FRayTracingInstanceDebugData> InstancesDebugData;
 StructuredBuffer<FPlatformRayTracingInstanceDescriptor> InstanceBuffer;
 StructuredBuffer<FRayTracingPickingFeedback> PickingBuffer;
@@ -234,6 +235,9 @@ float4 DebugRayTracingInstance(FRayDesc Ray, out float OutDepth)
 		const float3 GeometryColor = IntToColor(Payload.InstanceHash + Payload.GeometryIndex);
 		const float3 TriangleColor = IntToColor(Payload.InstanceHash + Payload.TriangleIndex);
 
+		const FRayTracingInstanceExtraData InstanceExtraData = InstancesExtraData[Payload.InstanceIndex];
+		const FRayTracingInstanceDebugData InstanceDebugData = InstancesDebugData[InstanceExtraData.SceneInstanceIndex];
+
 		switch (VisualizationMode)
 		{
 		default:
@@ -256,10 +260,10 @@ float4 DebugRayTracingInstance(FRayDesc Ray, out float OutDepth)
 			Result.rgb = TriangleColor;
 			break;
 		case RAY_TRACING_DEBUG_VIZ_DYNAMIC_INSTANCES:
-			Result.rgb = InstancesDebugData[Payload.InstanceIndex].Flags ? float3(0.25f, 0.50f, 0.25f) : float3(0.50f, 0.25f, 0.25f);
+			Result.rgb = InstanceDebugData.Flags ? float3(0.25f, 0.50f, 0.25f) : float3(0.50f, 0.25f, 0.25f);
 			break;
 		case RAY_TRACING_DEBUG_VIZ_PROXY_TYPE:
-			Result.rgb = IntToColor(InstancesDebugData[Payload.InstanceIndex].ProxyHash);
+			Result.rgb = IntToColor(InstanceDebugData.ProxyHash);
 			break;
 		}
 
diff --git a/Engine/Shaders/Private/RayTracing/RayTracingDebugInstanceOverlap.usf b/Engine/Shaders/Private/RayTracing/RayTracingDebugInstanceOverlap.usf
--- a/Engine/Shaders/Private/RayTracing/RayTracingDebugInstanceOverlap.usf
+++ b/Engine/Shaders/Private/RayTracing/RayTracingDebugInstanceOverlap.usf
@@ -6,6 +6,7 @@
 #include "../Hash.ush"
 #include "../Visualization.ush"
 #include "../SceneData.ush"
+#include "/Engine/Shared/RayTracingDebugTypes.h"
 
 #define MAX_FLT (+3.402823466e+38f)
 #define MIN_FLT (-3.402823466e+38f)
@@ -13,7 +14,7 @@
 
 float BoundingBoxExtentScale;
 
-StructuredBuffer<uint> InstanceGPUSceneIndices;
+StructuredBuffer<FRayTracingInstanceExtraData> InstanceExtraDataBuffer;
 
 struct FInstanceBounds
 {
@@ -21,9 +22,9 @@ struct FInstanceBounds
 	float3 Extent;
 };
 
-FInstanceBounds GetTranslatedWorldAABBForInstance(uint GPUSceneInstanceIndex)
+FInstanceBounds GetTranslatedWorldAABBForInstance(uint GPUSceneInstanceId)
 {
-	FInstanceSceneData InstanceSceneData = GetInstanceSceneData(GPUSceneInstanceIndex);
+	FInstanceSceneData InstanceSceneData = GetInstanceSceneData(GPUSceneInstanceId);
 
 	const FDFVector3 PreViewTranslation = MakeDFVector3(View.TLASPreViewTranslationHigh, View.TLASPreViewTranslationLow);
 	const float4x4 LocalToTranslatedWorld = DFFastToTranslatedWorld(InstanceSceneData.LocalToWorld, PreViewTranslation);	
@@ -68,11 +69,11 @@ void InstanceOverlapMainVS(
 	out float4 OutPosition : SV_POSITION
 )
 {
-	const uint GPUSceneIndex = InstanceGPUSceneIndices[InstanceIndex];
+	const FRayTracingInstanceExtraData InstanceExtraData = InstanceExtraDataBuffer[InstanceIndex];
 
-	if (GPUSceneIndex != INVALID_INSTANCE_GPUSCENE_INDEX)
+	if (InstanceExtraData.GPUSceneInstanceId != INVALID_INSTANCE_GPUSCENE_INDEX)
 	{
-		const FInstanceBounds InstanceBounds = GetTranslatedWorldAABBForInstance(GPUSceneIndex);
+		const FInstanceBounds InstanceBounds = GetTranslatedWorldAABBForInstance(InstanceExtraData.GPUSceneInstanceId);
 
 		// Make sure extent is non-zero to avoid flickering.
 		const float3 Extent = max(InstanceBounds.Extent, BoundingBoxExtentScale);
diff --git a/Engine/Shaders/Private/RayTracing/RayTracingDebugPicking.usf b/Engine/Shaders/Private/RayTracing/RayTracingDebugPicking.usf
--- a/Engine/Shaders/Private/RayTracing/RayTracingDebugPicking.usf
+++ b/Engine/Shaders/Private/RayTracing/RayTracingDebugPicking.usf
@@ -9,6 +9,7 @@ uint OpaqueOnly;
 
 RaytracingAccelerationStructure TLAS;
 
+StructuredBuffer<FRayTracingInstanceExtraData> InstancesExtraData;
 StructuredBuffer<FRayTracingInstanceDebugData> InstancesDebugData;
 StructuredBuffer<FRayTracingPickingFeedback> PickingBuffer;
 
@@ -41,8 +42,9 @@ RAY_TRACING_ENTRY_RAYGEN(RayTracingDebugPickingRGS)
 
 	if (Payload.IsHit())
 	{
-		FRayTracingInstanceDebugData InstanceDebugData = InstancesDebugData[Payload.InstanceIndex];
-		FPlatformRayTracingInstanceDescriptor InstanceDescriptor = InstanceBuffer[Payload.InstanceIndex];
+		const FRayTracingInstanceExtraData InstanceExtraData = InstancesExtraData[Payload.InstanceIndex];
+		const FRayTracingInstanceDebugData InstanceDebugData = InstancesDebugData[InstanceExtraData.SceneInstanceIndex];
+		const FPlatformRayTracingInstanceDescriptor InstanceDescriptor = InstanceBuffer[Payload.InstanceIndex];
 
 		PickingFeedback.GeometryAddress = InstanceDebugData.GeometryAddress;
 
diff --git a/Engine/Shaders/Private/RayTracing/RayTracingInstanceBufferUtil.usf b/Engine/Shaders/Private/RayTracing/RayTracingInstanceBufferUtil.usf
--- a/Engine/Shaders/Private/RayTracing/RayTracingInstanceBufferUtil.usf
+++ b/Engine/Shaders/Private/RayTracing/RayTracingInstanceBufferUtil.usf
@@ -8,6 +8,7 @@
 #include "../SceneData.ush"
 #include "../ComputeShaderUtils.ush"
 #include "../Common.ush"
+#include "/Engine/Shared/RayTracingDebugTypes.h"
 
 #ifndef THREADGROUP_SIZE
 #define THREADGROUP_SIZE 1
@@ -21,7 +22,7 @@ struct FRayTracingInstanceDescriptor
 	uint InstanceId;
 	uint InstanceMaskAndFlags;
 	uint InstanceContributionToHitGroupIndex;
-	uint bApplyLocalBoundsTransform;
+	uint SceneInstanceIndexAndApplyLocalBoundsTransform;
 };
 
 uint NumInstances;
@@ -42,8 +43,8 @@ ByteAddressBuffer AccelerationStructureAddresses;
 // because of FXC matrix packing issues when using StructuredBuffer<float3x4>
 StructuredBuffer<float4> InstanceTransforms;
 
-// Extra per instance debug data
-RWStructuredBuffer<uint> RWDebugInstanceGPUSceneIndices;
+// Extra per instance data
+RWStructuredBuffer<FRayTracingInstanceExtraData> RWInstanceExtraData;
 
 #if OUTPUT_STATS
 RWStructuredBuffer<uint> RWOutputStats;
@@ -132,7 +133,7 @@ void ProcessInstance(FRayTracingInstanceDescriptor InputDesc)
 	}
 #endif // GPU_CULLING
 
-	if (InputDesc.bApplyLocalBoundsTransform)
+	if ((InputDesc.SceneInstanceIndexAndApplyLocalBoundsTransform & 0x80000000) != 0)
 	{
 		float4x4 ProcT;
 		ProcT[0] = float4(InstanceSceneData.LocalBoundsExtent.x * 2, 0, 0, 0);
@@ -149,24 +150,30 @@ void ProcessInstance(FRayTracingInstanceDescriptor InputDesc)
 	LocalToTranslatedWorld[0] = LocalToTranslatedWorld4x4[0];
 	LocalToTranslatedWorld[1] = LocalToTranslatedWorld4x4[1];
 	LocalToTranslatedWorld[2] = LocalToTranslatedWorld4x4[2];
-
-#if OUTPUT_INSTANCE_GPUSCENE_INDEX
-	if (InputDesc.AccelerationStructureIndex != 0xFFFFFFFF && !bCulled)
-	{
-		RWDebugInstanceGPUSceneIndices[InputDesc.OutputDescriptorIndex] = InputDesc.GPUSceneInstanceOrTransformIndex;
-	}
-#endif // OUTPUT_INSTANCE_GPUSCENE_INDEX
-
-#else
+#else // !USE_GPUSCENE
 	float3x4 LocalToTranslatedWorld;
 	LocalToTranslatedWorld[0] = InstanceTransforms[InputDesc.GPUSceneInstanceOrTransformIndex * 3 + 0];
 	LocalToTranslatedWorld[1] = InstanceTransforms[InputDesc.GPUSceneInstanceOrTransformIndex * 3 + 1];
 	LocalToTranslatedWorld[2] = InstanceTransforms[InputDesc.GPUSceneInstanceOrTransformIndex * 3 + 2];
-#endif // USE_GPUSCENE
+#endif // !USE_GPUSCENE
+
+#if OUTPUT_INSTANCE_EXTRA_DATA
+	if (!bCulled)
+	{
+		FRayTracingInstanceExtraData ExtraData;
+		ExtraData.SceneInstanceIndex = InputDesc.SceneInstanceIndexAndApplyLocalBoundsTransform & 0x7fffffff;
+#if USE_GPUSCENE
+		ExtraData.GPUSceneInstanceId = InputDesc.GPUSceneInstanceOrTransformIndex;
+#else
+		ExtraData.GPUSceneInstanceId = -1;
+#endif
+		RWInstanceExtraData[InputDesc.OutputDescriptorIndex] = ExtraData;
+	}
+#endif // OUTPUT_INSTANCE_EXTRA_DATA
 
 	uint2 BlasAddress = 0;
 
-	if (InputDesc.AccelerationStructureIndex != 0xFFFFFFFF && !bCulled)
+	if (!bCulled)
 	{
 		BlasAddress = AccelerationStructureAddresses.Load2(InputDesc.AccelerationStructureIndex * 8);
 	}
diff --git a/Engine/Shaders/Shared/RayTracingDebugTypes.h b/Engine/Shaders/Shared/RayTracingDebugTypes.h
--- a/Engine/Shaders/Shared/RayTracingDebugTypes.h
+++ b/Engine/Shaders/Shared/RayTracingDebugTypes.h
@@ -30,6 +30,19 @@ struct FRayTracingInstanceDebugData
 #endif
 };
 
+struct FRayTracingInstanceExtraData
+{
+	UINT32_TYPE SceneInstanceIndex;
+	UINT32_TYPE GPUSceneInstanceId;
+
+#ifdef __cplusplus
+	FRayTracingInstanceExtraData()
+		: SceneInstanceIndex(-1)
+		, GPUSceneInstanceId(-1)
+	{}
+#endif
+};
+
 struct FRayTracingPickingFeedback
 {
 	// Hit data
@@ -38,7 +51,7 @@ struct FRayTracingPickingFeedback
 	UINT32_TYPE TriangleIndex;
 
 	// Geometry data
-	UINT64_TYPE GeometryAddress;
+	UINT64_TYPE GeometryAddress; // TODO: Usage GeometryHandle instead
 
 	// Instance data
 	UINT32_TYPE InstanceId;
diff --git a/Engine/Source/Runtime/Renderer/Private/RayTracing/RayTracingDebug.cpp b/Engine/Source/Runtime/Renderer/Private/RayTracing/RayTracingDebug.cpp
--- a/Engine/Source/Runtime/Renderer/Private/RayTracing/RayTracingDebug.cpp
+++ b/Engine/Source/Runtime/Renderer/Private/RayTracing/RayTracingDebug.cpp
@@ -187,6 +187,7 @@ class FRayTracingDebugRGS : public FGlobalShader
 		SHADER_PARAMETER_RDG_BUFFER_SRV(RaytracingAccelerationStructure, FarFieldTLAS)
 		SHADER_PARAMETER_RDG_TEXTURE_UAV(RWTexture2D<float4>, Output)
 		SHADER_PARAMETER_RDG_TEXTURE_UAV(RWTexture2D<float>, OutputDepth)
+		SHADER_PARAMETER_RDG_BUFFER_SRV(StructuredBuffer, InstancesExtraData)
 		SHADER_PARAMETER_RDG_BUFFER_SRV(StructuredBuffer, InstancesDebugData)
 		SHADER_PARAMETER_RDG_BUFFER_SRV(StructuredBuffer, InstanceBuffer)
 		SHADER_PARAMETER_RDG_BUFFER_SRV(StructuredBuffer<FRayTracingPickingFeedback>, PickingBuffer)
@@ -444,6 +445,7 @@ class FRayTracingPickingRGS : public FGlobalShader
 		SHADER_PARAMETER_RDG_BUFFER_SRV(RaytracingAccelerationStructure, TLAS)
 		SHADER_PARAMETER(int32, OpaqueOnly)
 		SHADER_PARAMETER_RDG_BUFFER_UAV(RWStructuredBuffer, PickingOutput)
+		SHADER_PARAMETER_RDG_BUFFER_SRV(StructuredBuffer, InstancesExtraData)
 		SHADER_PARAMETER_RDG_BUFFER_SRV(StructuredBuffer, InstancesDebugData)
 		SHADER_PARAMETER_RDG_BUFFER_SRV(StructuredBuffer, InstanceBuffer)
 		SHADER_PARAMETER_STRUCT_REF(FViewUniformShaderParameters, ViewUniformBuffer)
@@ -475,7 +477,7 @@ class FRayTracingDebugInstanceOverlapVS : public FGlobalShader
 	BEGIN_SHADER_PARAMETER_STRUCT(FParameters, )
 		SHADER_PARAMETER_STRUCT_REF(FViewUniformShaderParameters, View)
 		SHADER_PARAMETER_RDG_UNIFORM_BUFFER(FSceneUniformParameters, Scene)
-		SHADER_PARAMETER_RDG_BUFFER_SRV(StructuredBuffer<uint>, InstanceGPUSceneIndices)
+		SHADER_PARAMETER_RDG_BUFFER_SRV(StructuredBuffer, InstanceExtraDataBuffer)
 		SHADER_PARAMETER(float, BoundingBoxExtentScale)
 	END_SHADER_PARAMETER_STRUCT()
 
@@ -912,6 +914,7 @@ static FRDGBufferRef RayTracingPerformPicking(FRDGBuilder& GraphBuilder, const F
 	FRDGBufferRef PickingBuffer = GraphBuilder.CreateBuffer(PickingBufferDesc, TEXT("RayTracingDebug.PickingBuffer"));
 
 	FRayTracingPickingRGS::FParameters* RayGenParameters = GraphBuilder.AllocParameters<FRayTracingPickingRGS::FParameters>();
+	RayGenParameters->InstancesExtraData = GraphBuilder.CreateSRV(RayTracingScene.GetInstanceExtraDataBuffer(ERayTracingSceneLayer::Base));
 	RayGenParameters->InstancesDebugData = GraphBuilder.CreateSRV(RayTracingScene.GetInstanceDebugBuffer(ERayTracingSceneLayer::Base));
 	RayGenParameters->TLAS = RayTracingScene.GetLayerView(ERayTracingSceneLayer::Base);
 	RayGenParameters->OpaqueOnly = CVarRayTracingVisualizeOpaqueOnly.GetValueOnRenderThread();
@@ -1268,7 +1271,7 @@ static FRDGBufferRef RayTracingPerformHitStatsPerPrimitive(FRDGBuilder& GraphBui
 }
 
 
-static void RayTracingDrawInstances(FRDGBuilder& GraphBuilder, const FViewInfo& View, FRDGTextureRef OutputTexture, FRDGTextureRef SceneDepthTexture, FRDGBufferRef InstanceGPUSceneIndexBuffer, bool bWireframe)
+static void RayTracingDrawInstances(FRDGBuilder& GraphBuilder, const FViewInfo& View, FRDGTextureRef OutputTexture, FRDGTextureRef SceneDepthTexture, FRDGBufferRef InstanceExtraDataBuffer, uint32 NumInstances, bool bWireframe)
 {
 	TShaderMapRef<FRayTracingDebugInstanceOverlapVS> VertexShader(View.ShaderMap);
 	TShaderMapRef<FRayTracingDebugInstanceOverlapPS> PixelShader(View.ShaderMap);
@@ -1276,13 +1279,11 @@ static void RayTracingDrawInstances(FRDGBuilder& GraphBuilder, const FViewInfo&
 	FRayTracingDebugInstanceOverlapVSPSParameters* PassParameters = GraphBuilder.AllocParameters<FRayTracingDebugInstanceOverlapVSPSParameters>();
 	PassParameters->VS.View = View.ViewUniformBuffer;
 	PassParameters->VS.Scene = View.GetSceneUniforms().GetBuffer(GraphBuilder);
-	PassParameters->VS.InstanceGPUSceneIndices = GraphBuilder.CreateSRV(InstanceGPUSceneIndexBuffer);
+	PassParameters->VS.InstanceExtraDataBuffer = GraphBuilder.CreateSRV(InstanceExtraDataBuffer);
 	PassParameters->VS.BoundingBoxExtentScale = CVarRayTracingVisualizeInstanceOverlapBoundingBoxScale.GetValueOnRenderThread();
 
 	PassParameters->PS.View = View.ViewUniformBuffer;
 
-	const uint32 NumInstances = InstanceGPUSceneIndexBuffer->GetSize() / InstanceGPUSceneIndexBuffer->GetStride();
-
 	PassParameters->RenderTargets[0] = FRenderTargetBinding(OutputTexture, bWireframe ? ERenderTargetLoadAction::ELoad : ERenderTargetLoadAction::EClear);
 	PassParameters->RenderTargets.DepthStencil = FDepthStencilBinding(SceneDepthTexture, ERenderTargetLoadAction::ELoad, FExclusiveDepthStencil::DepthRead_StencilNop);
 
@@ -1356,7 +1357,14 @@ static void DrawInstanceOverlap(FRDGBuilder& GraphBuilder, const FScene& Scene,
 	FRDGTextureDesc InstanceOverlapTextureDesc = FRDGTextureDesc::Create2D(SceneColorTexture->Desc.Extent, PF_R32_FLOAT, FClearValueBinding::Black, TexCreate_ShaderResource | TexCreate_RenderTargetable);
 	FRDGTextureRef InstanceOverlapTexture = GraphBuilder.CreateTexture(InstanceOverlapTextureDesc, TEXT("RayTracingDebug::InstanceOverlap"));
 	
-	RayTracingDrawInstances(GraphBuilder, View, InstanceOverlapTexture, SceneDepthTexture, Scene.RayTracingScene.GetDebugInstanceGPUSceneIndexBuffer(ERayTracingSceneLayer::Base), false);
+	RayTracingDrawInstances(
+		GraphBuilder,
+		View,
+		InstanceOverlapTexture,
+		SceneDepthTexture,
+		Scene.RayTracingScene.GetInstanceExtraDataBuffer(ERayTracingSceneLayer::Base),
+		Scene.RayTracingScene.GetNumNativeInstances(ERayTracingSceneLayer::Base),
+		false);
 
 	// Calculate heatmap of instance overlap and blend it on top of ray tracing debug output
 	{
@@ -1384,7 +1392,14 @@ static void DrawInstanceOverlap(FRDGBuilder& GraphBuilder, const FScene& Scene,
 	// Draw instance AABB with lines
 	if (CVarRayTracingVisualizeInstanceOverlapShowWireframe.GetValueOnRenderThread() != 0)
 	{
-		RayTracingDrawInstances(GraphBuilder, View, SceneColorTexture, SceneDepthTexture, Scene.RayTracingScene.GetDebugInstanceGPUSceneIndexBuffer(ERayTracingSceneLayer::Base), true);
+		RayTracingDrawInstances(
+			GraphBuilder,
+			View,
+			SceneColorTexture,
+			SceneDepthTexture,
+			Scene.RayTracingScene.GetInstanceExtraDataBuffer(ERayTracingSceneLayer::Base),
+			Scene.RayTracingScene.GetNumNativeInstances(ERayTracingSceneLayer::Base),
+			true);
 	}
 }
 
@@ -1547,7 +1562,9 @@ void RenderRayTracingDebug(FRDGBuilder& GraphBuilder, const FScene& Scene, const
 
 	FRDGBufferRef PickingBuffer = nullptr;
 	FRDGBufferRef StatsBuffer = nullptr;
-	if (IsRayTracingPickingEnabled(DebugVisualizationMode) && RayTracingScene.GetInstanceDebugBuffer(ERayTracingSceneLayer::Base) != nullptr)
+	if (IsRayTracingPickingEnabled(DebugVisualizationMode)
+		&& RayTracingScene.GetInstanceExtraDataBuffer(ERayTracingSceneLayer::Base) != nullptr
+		&& RayTracingScene.GetInstanceDebugBuffer(ERayTracingSceneLayer::Base) != nullptr)
 	{
 		PickingBuffer = RayTracingPerformPicking(GraphBuilder, Scene, View, PickingFeedback);
 	}
@@ -1571,11 +1588,18 @@ void RenderRayTracingDebug(FRDGBuilder& GraphBuilder, const FScene& Scene, const
 		AddClearUAVPass(GraphBuilder, DebugHitStatsUniformBufferParameters->HitStatsOutput, 0);
 	}
 
+	FRDGBufferRef InstanceExtraDataBuffer = RayTracingScene.GetInstanceExtraDataBuffer(ERayTracingSceneLayer::Base);
+	if (InstanceExtraDataBuffer == nullptr)
+	{
+		InstanceExtraDataBuffer = GSystemTextures.GetDefaultStructuredBuffer(GraphBuilder, sizeof(uint32));
+	}
+
 	FRDGBufferRef InstanceDebugBuffer = RayTracingScene.GetInstanceDebugBuffer(ERayTracingSceneLayer::Base);
 	if (InstanceDebugBuffer == nullptr)
 	{
 		InstanceDebugBuffer = GSystemTextures.GetDefaultStructuredBuffer(GraphBuilder, sizeof(uint32));
 	}
+
 	FRDGBufferRef InstanceBuffer = RayTracingScene.GetInstanceBuffer(ERayTracingSceneLayer::Base);
 	if (InstanceBuffer == nullptr)
 	{
@@ -1673,6 +1697,7 @@ void RenderRayTracingDebug(FRDGBuilder& GraphBuilder, const FScene& Scene, const
 		RayGenParameters->FarFieldMaxTraceDistance = 0.0f;
 	}
 	
+	RayGenParameters->InstancesExtraData = GraphBuilder.CreateSRV(InstanceExtraDataBuffer);
 	RayGenParameters->InstancesDebugData = GraphBuilder.CreateSRV(InstanceDebugBuffer);
 	RayGenParameters->InstanceBuffer = GraphBuilder.CreateSRV(InstanceBuffer);
 	RayGenParameters->PickingBuffer = GraphBuilder.CreateSRV(PickingBuffer);
diff --git a/Engine/Source/Runtime/Renderer/Private/RayTracing/RayTracingInstanceBufferUtil.cpp b/Engine/Source/Runtime/Renderer/Private/RayTracing/RayTracingInstanceBufferUtil.cpp
--- a/Engine/Source/Runtime/Renderer/Private/RayTracing/RayTracingInstanceBufferUtil.cpp
+++ b/Engine/Source/Runtime/Renderer/Private/RayTracing/RayTracingInstanceBufferUtil.cpp
@@ -219,6 +219,7 @@ FRayTracingSceneInitializationData BuildRayTracingSceneInitializationData(TConst
 
 void WriteInstanceDescriptor(
 	const FRayTracingGeometryInstance& SceneInstance,
+	uint32 SceneInstanceIndex,
 	uint32 TransformIndex,
 	uint32 AccelerationStructureIndex,
 	bool bGpuSceneInstance,
@@ -266,7 +267,7 @@ void WriteInstanceDescriptor(
 	InstanceDesc.InstanceId = UserData;
 	InstanceDesc.InstanceMaskAndFlags = SceneInstance.Mask | ((uint32)SceneInstance.Flags << 8);
 	InstanceDesc.InstanceContributionToHitGroupIndex = SceneInstance.InstanceContributionToHitGroupIndex;
-	InstanceDesc.bApplyLocalBoundsTransform = SceneInstance.bApplyLocalBoundsTransform;
+	InstanceDesc.SceneInstanceIndexAndApplyLocalBoundsTransform = (SceneInstance.bApplyLocalBoundsTransform ? 0x80000000 : 0) | SceneInstanceIndex;
 
 	checkf(InstanceDesc.InstanceId <= 0xFFFFFF, TEXT("InstanceId must fit in 24 bits."));
 	checkf(InstanceDesc.InstanceContributionToHitGroupIndex <= 0xFFFFFF, TEXT("InstanceContributionToHitGroupIndex must fit in 24 bits."));
@@ -370,6 +371,7 @@ void FillRayTracingInstanceUploadBuffer(
 				{
 					WriteInstanceDescriptor(
 						SceneInstance,
+						SceneInstanceIndex,
 						TransformIndex,
 						AccelerationStructureIndex,
 						bGpuSceneInstance,
@@ -404,6 +406,7 @@ void FillRayTracingInstanceUploadBuffer(
 					// and corresponding instance
 					WriteInstanceDescriptor(
 						SceneInstance,
+						SceneInstanceIndex,
 						TransformIndex,
 						AccelerationStructureIndex,
 						bGpuSceneInstance,
@@ -441,6 +444,7 @@ void FillRayTracingInstanceUploadBuffer(
 				{
 					WriteInstanceDescriptor(
 						SceneInstance,
+						SceneInstanceIndex,
 						TransformIndex,
 						AccelerationStructureIndex,
 						bGpuSceneInstance,
@@ -531,16 +535,16 @@ struct FRayTracingBuildInstanceBufferCS : public FGlobalShader
 		SHADER_PARAMETER(uint32, OutputStatsOffset)
 
 		// Debug parameters
-		SHADER_PARAMETER_UAV(RWStructuredBuffer<uint>, RWDebugInstanceGPUSceneIndices)
+		SHADER_PARAMETER_UAV(RWStructuredBuffer, RWInstanceExtraData)
 	END_SHADER_PARAMETER_STRUCT()
 
 	class FSupportInstanceGroupsDim : SHADER_PERMUTATION_BOOL("SUPPORT_INSTANCE_GROUPS");
 	class FUseGPUSceneDim : SHADER_PERMUTATION_BOOL("USE_GPUSCENE");
-	class FOutputInstanceGPUSceneIndexDim : SHADER_PERMUTATION_BOOL("OUTPUT_INSTANCE_GPUSCENE_INDEX");
+	class FOutputInstanceExtraDataDim : SHADER_PERMUTATION_BOOL("OUTPUT_INSTANCE_EXTRA_DATA");
 	class FGpuCullingDim : SHADER_PERMUTATION_BOOL("GPU_CULLING");
 	class FOutputStatsDim : SHADER_PERMUTATION_BOOL("OUTPUT_STATS");
 	class FUseWaveOpsDim : SHADER_PERMUTATION_BOOL("USE_WAVE_OPS");
-	using FPermutationDomain = TShaderPermutationDomain<FSupportInstanceGroupsDim, FUseGPUSceneDim, FOutputInstanceGPUSceneIndexDim, FGpuCullingDim, FOutputStatsDim, FUseWaveOpsDim>;
+	using FPermutationDomain = TShaderPermutationDomain<FSupportInstanceGroupsDim, FUseGPUSceneDim, FOutputInstanceExtraDataDim, FGpuCullingDim, FOutputStatsDim, FUseWaveOpsDim>;
 		
 	static constexpr uint32 ThreadGroupSize = GRayTracingInstanceGroupSize;
 
@@ -588,7 +592,7 @@ void BuildRayTracingInstanceBuffer(
 	const FRayTracingCullingParameters* CullingParameters,
 	FRHIUnorderedAccessView* OutputStatsUAV,
 	uint32 OutputStatsOffset,
-	FRHIUnorderedAccessView* DebugInstanceGPUSceneIndexUAV)
+	FRHIUnorderedAccessView* InstanceExtraDataUAV)
 {
 	FRayTracingBuildInstanceBufferCS::FParameters PassParams;
 	PassParams.OutPlatformInstanceDescriptors = InstancesUAV;
@@ -621,12 +625,12 @@ void BuildRayTracingInstanceBuffer(
 	PassParams.RWOutputStats = OutputStatsUAV;
 	PassParams.OutputStatsOffset = OutputStatsOffset;
 
-	PassParams.RWDebugInstanceGPUSceneIndices = DebugInstanceGPUSceneIndexUAV;
+	PassParams.RWInstanceExtraData = InstanceExtraDataUAV;
 
 	FRayTracingBuildInstanceBufferCS::FPermutationDomain PermutationVector;
 	PermutationVector.Set<FRayTracingBuildInstanceBufferCS::FSupportInstanceGroupsDim>(InstanceGroupUploadSRV != nullptr);
 	PermutationVector.Set<FRayTracingBuildInstanceBufferCS::FUseGPUSceneDim>(InstanceTransformSRV == nullptr);
-	PermutationVector.Set<FRayTracingBuildInstanceBufferCS::FOutputInstanceGPUSceneIndexDim>(DebugInstanceGPUSceneIndexUAV != nullptr);
+	PermutationVector.Set<FRayTracingBuildInstanceBufferCS::FOutputInstanceExtraDataDim>(InstanceExtraDataUAV != nullptr);
 	PermutationVector.Set<FRayTracingBuildInstanceBufferCS::FOutputStatsDim>(OutputStatsUAV != nullptr);
 	PermutationVector.Set<FRayTracingBuildInstanceBufferCS::FUseWaveOpsDim>(GRHISupportsWaveOperations);
 	PermutationVector.Set<FRayTracingBuildInstanceBufferCS::FGpuCullingDim>(CullingParameters != nullptr);
@@ -664,7 +668,7 @@ void BuildRayTracingInstanceBuffer(
 	const FRayTracingCullingParameters* CullingParameters,
 	FRHIUnorderedAccessView* OutputStatsUAV,
 	uint32 OutputStatsOffset,
-	FRHIUnorderedAccessView* DebugInstanceGPUSceneIndexUAV)
+	FRHIUnorderedAccessView* InstanceExtraDataUAV)
 {
 	if (NumNativeGPUSceneInstances > 0)
 	{
@@ -685,7 +689,7 @@ void BuildRayTracingInstanceBuffer(
 			CullingParameters,
 			OutputStatsUAV,
 			OutputStatsOffset,
-			DebugInstanceGPUSceneIndexUAV);
+			InstanceExtraDataUAV);
 	}
 
 	if (NumNativeCPUInstances > 0)
@@ -707,7 +711,7 @@ void BuildRayTracingInstanceBuffer(
 			nullptr,
 			OutputStatsUAV,
 			OutputStatsOffset,
-			nullptr);
+			InstanceExtraDataUAV);
 	}
 }
 
@@ -723,7 +727,7 @@ void BuildRayTracingInstanceBuffer(
 	uint32 NumNativeCPUInstances,
 	const FRayTracingCullingParameters* CullingParameters,
 	FRHIUnorderedAccessView* OutputStatsUAV,
-	FRHIUnorderedAccessView* DebugInstanceGPUSceneIndexUAV)
+	FRHIUnorderedAccessView* InstanceExtraDataUAV)
 {
 	BuildRayTracingInstanceBuffer(
 		RHICmdList,
@@ -743,7 +747,7 @@ void BuildRayTracingInstanceBuffer(
 		CullingParameters,
 		OutputStatsUAV,
 		0,
-		DebugInstanceGPUSceneIndexUAV);
+		InstanceExtraDataUAV);
 }
 
 void FRayTracingInstanceBufferBuilder::Init(TConstArrayView<FRayTracingGeometryInstance> InInstances, FVector InPreViewTranslation)
@@ -901,7 +905,7 @@ void FRayTracingInstanceBufferBuilder::BuildRayTracingInstanceBuffer(
 	FRHIUnorderedAccessView* InstancesUAV,
 	FRHIUnorderedAccessView* OutputStatsUAV,
 	uint32 OutputStatsOffset,
-	FRHIUnorderedAccessView* DebugInstanceGPUSceneIndexUAV)
+	FRHIUnorderedAccessView* InstanceExtraDataUAV)
 {
 	::BuildRayTracingInstanceBuffer(
 		RHICmdList,
@@ -921,7 +925,7 @@ void FRayTracingInstanceBufferBuilder::BuildRayTracingInstanceBuffer(
 		CullingParameters,
 		OutputStatsUAV,
 		OutputStatsOffset,
-		DebugInstanceGPUSceneIndexUAV);
+		InstanceExtraDataUAV);
 }
 
 PRAGMA_ENABLE_DEPRECATION_WARNINGS
diff --git a/Engine/Source/Runtime/Renderer/Private/RayTracing/RayTracingScene.cpp b/Engine/Source/Runtime/Renderer/Private/RayTracing/RayTracingScene.cpp
--- a/Engine/Source/Runtime/Renderer/Private/RayTracing/RayTracingScene.cpp
+++ b/Engine/Source/Runtime/Renderer/Private/RayTracing/RayTracingScene.cpp
@@ -50,7 +50,7 @@ bool IsRayTracingFeedbackEnabled(const FSceneViewFamily& ViewFamily)
 BEGIN_SHADER_PARAMETER_STRUCT(FBuildInstanceBufferPassParams, )
 	SHADER_PARAMETER_RDG_BUFFER_UAV(RWStructuredBuffer, InstanceBuffer)
 	SHADER_PARAMETER_RDG_BUFFER_UAV(RWStructuredBuffer, OutputStats)
-	SHADER_PARAMETER_RDG_BUFFER_UAV(RWStructuredBuffer, DebugInstanceGPUSceneIndexBuffer)
+	SHADER_PARAMETER_RDG_BUFFER_UAV(RWStructuredBuffer, InstanceExtraDataBuffer)
 	SHADER_PARAMETER_RDG_UNIFORM_BUFFER(FSceneUniformParameters, Scene)
 END_SHADER_PARAMETER_STRUCT()
 
@@ -217,67 +217,25 @@ void FRayTracingScene::Create(FRDGBuilder& GraphBuilder, const FViewInfo& View,
 			}
 		}
 
-		FRDGBufferUAVRef DebugInstanceGPUSceneIndexBufferUAV = nullptr;
-		if (bNeedsDebugInstanceGPUSceneIndexBuffer)
+		FRDGBufferUAVRef InstanceExtraDataBufferUAV = nullptr;
+		if (bNeedsInstanceExtraDataBuffer || bTracingFeedbackEnabled || bInstanceDebugDataEnabled)
 		{
-			FRDGBufferDesc DebugInstanceGPUSceneIndexBufferDesc;
-			DebugInstanceGPUSceneIndexBufferDesc.Usage = EBufferUsageFlags::UnorderedAccess | EBufferUsageFlags::ShaderResource | EBufferUsageFlags::StructuredBuffer;
-			DebugInstanceGPUSceneIndexBufferDesc.BytesPerElement = sizeof(uint32);
-			DebugInstanceGPUSceneIndexBufferDesc.NumElements = FMath::Max(NumNativeInstances, 1u);
+			FRDGBufferDesc InstanceExtraDataBufferDesc;
+			InstanceExtraDataBufferDesc.Usage = EBufferUsageFlags::UnorderedAccess | EBufferUsageFlags::ShaderResource | EBufferUsageFlags::StructuredBuffer;
+			InstanceExtraDataBufferDesc.BytesPerElement = sizeof(FRayTracingInstanceExtraData);
+			InstanceExtraDataBufferDesc.NumElements = FMath::Max(NumNativeInstances, 1u);
 
-			Layer.DebugInstanceGPUSceneIndexBuffer = GraphBuilder.CreateBuffer(DebugInstanceGPUSceneIndexBufferDesc, TEXT("FRayTracingScene::DebugInstanceGPUSceneIndexBuffer"));
-			DebugInstanceGPUSceneIndexBufferUAV = GraphBuilder.CreateUAV(Layer.DebugInstanceGPUSceneIndexBuffer);
+			Layer.InstanceExtraDataBuffer = GraphBuilder.CreateBuffer(InstanceExtraDataBufferDesc, TEXT("FRayTracingScene::InstanceExtraDataBuffer"));
+			InstanceExtraDataBufferUAV = GraphBuilder.CreateUAV(Layer.InstanceExtraDataBuffer);
 
-			AddClearUAVPass(GraphBuilder, DebugInstanceGPUSceneIndexBufferUAV, 0xFFFFFFFF);
+			AddClearUAVPass(GraphBuilder, InstanceExtraDataBufferUAV, 0xFFFFFFFF);
 		}
 
-		if (Layer.InstancesDebugData.Num() > 0 && NumNativeInstances > 0)
+		if (Layer.InstancesDebugData.Num() > 0)
 		{
-			// Create InstanceDebugBuffer (one entry per instance in TLAS)
-			// This requires replicating the data in InstancesDebugData (one entry per FRayTracingGeometryInstance) according to NumTransforms in each geometry instance
-
 			check(Layer.InstancesDebugData.Num() == Layer.Instances.Num());
 
-			FRDGUploadData<FRayTracingInstanceDebugData> UploadData(GraphBuilder, NumNativeInstances);
-
-			{
-				const uint32 NumItems = Layer.InstancesDebugData.Num();
-
-				// Distribute work evenly to the available task graph workers based on NumItems.
-				const uint32 TargetItemsPerTask = 512;
-				const uint32 NumThreads = FMath::Min(FTaskGraphInterface::Get().GetNumWorkerThreads(), CVarRHICmdWidth.GetValueOnRenderThread());
-				const uint32 NumTasks = FMath::Min(NumThreads, FMath::DivideAndRoundUp(NumItems, TargetItemsPerTask));
-				const uint32 NumItemsPerTask = FMath::DivideAndRoundUp(NumItems, NumTasks);
-
-				for (uint32 TaskIndex = 0; TaskIndex < NumTasks; ++TaskIndex)
-				{
-					const uint32 TaskFirstItemIndex = TaskIndex * NumItemsPerTask;
-					const uint32 TaskNumItems = FMath::Min(NumItemsPerTask, NumItems - TaskFirstItemIndex);
-
-					TConstArrayView<FRayTracingGeometryInstance> TaskInstancesData(Layer.Instances.GetData() + TaskFirstItemIndex, TaskNumItems);
-					TConstArrayView<FRayTracingInstanceDebugData> TaskInstancesDebugData(Layer.InstancesDebugData.GetData() + TaskFirstItemIndex, TaskNumItems);
-					TConstArrayView<uint32> TaskBaseInstancePrefixSum(Layer.InstanceBufferBuilder.GetBaseInstancePrefixSum().GetData() + TaskFirstItemIndex, TaskNumItems);
-
-					GraphBuilder.AddSetupTask([UploadData, TaskInstancesDebugData, TaskInstancesData, TaskBaseInstancePrefixSum]()
-						{
-							TRACE_CPUPROFILER_EVENT_SCOPE(FillRayTracingInstanceDebugBuffer);
-
-							for (int32 Index = 0; Index < TaskInstancesDebugData.Num(); ++Index)
-							{
-								const FRayTracingGeometryInstance& SceneInstance = TaskInstancesData[Index];
-								const uint32 BaseInstanceIndex = TaskBaseInstancePrefixSum[Index];
-
-								for (uint32 TransformIndex = 0; TransformIndex < SceneInstance.NumTransforms; ++TransformIndex)
-								{
-									// write data in the same order used in InstanceBuffer used to build TLAS / InstanceIndex() in hit shaders
-									UploadData[BaseInstanceIndex + TransformIndex] = TaskInstancesDebugData[Index];
-								}
-							}
-						});
-				}
-			}
-
-			Layer.InstanceDebugBuffer = CreateStructuredBuffer(GraphBuilder, TEXT("FRayTracingScene::InstanceDebugData"), UploadData);
+			Layer.InstanceDebugBuffer = CreateStructuredBuffer(GraphBuilder, TEXT("FRayTracingScene::InstanceDebugData"), Layer.InstancesDebugData);
 		}
 
 		if (NumNativeInstances > 0)
@@ -299,7 +257,7 @@ void FRayTracingScene::Create(FRDGBuilder& GraphBuilder, const FViewInfo& View,
 
 			FBuildInstanceBufferPassParams* PassParams = GraphBuilder.AllocParameters<FBuildInstanceBufferPassParams>();
 			PassParams->InstanceBuffer = GraphBuilder.CreateUAV(Layer.InstanceBuffer);
-			PassParams->DebugInstanceGPUSceneIndexBuffer = DebugInstanceGPUSceneIndexBufferUAV;
+			PassParams->InstanceExtraDataBuffer = InstanceExtraDataBufferUAV;
 			PassParams->Scene = View.GetSceneUniforms().GetBuffer(GraphBuilder);
 			PassParams->OutputStats = STATS ? InstanceStatsBufferUAV : nullptr;
 
@@ -321,7 +279,7 @@ void FRayTracingScene::Create(FRDGBuilder& GraphBuilder, const FViewInfo& View,
 						PassParams->InstanceBuffer->GetRHI(),
 						PassParams->OutputStats ? PassParams->OutputStats->GetRHI() : nullptr,
 						LayerIndex,
-						PassParams->DebugInstanceGPUSceneIndexBuffer ? PassParams->DebugInstanceGPUSceneIndexBuffer->GetRHI() : nullptr);
+						PassParams->InstanceExtraDataBuffer ? PassParams->InstanceExtraDataBuffer->GetRHI() : nullptr);
 				});
 		}
 	}
@@ -806,8 +764,8 @@ void FRayTracingScene::Reset(bool bInInstanceDebugDataEnabled)
 
 		Layer.InstanceBuffer = nullptr;
 		Layer.BuildScratchBuffer = nullptr;
+		Layer.InstanceExtraDataBuffer = nullptr;
 		Layer.InstanceDebugBuffer = nullptr;
-		Layer.DebugInstanceGPUSceneIndexBuffer = nullptr;
 
 		Layer.InstanceHitCountBuffer = nullptr;
 		Layer.AccelerationStructureIndexBuffer = nullptr;
diff --git a/Engine/Source/Runtime/Renderer/Private/RayTracing/RayTracingScene.h b/Engine/Source/Runtime/Renderer/Private/RayTracing/RayTracingScene.h
--- a/Engine/Source/Runtime/Renderer/Private/RayTracing/RayTracingScene.h
+++ b/Engine/Source/Runtime/Renderer/Private/RayTracing/RayTracingScene.h
@@ -153,7 +153,7 @@ public:
 	uint32 GetNumNativeInstances(ERayTracingSceneLayer Layer) const;
 
 	FRDGBufferRef GetInstanceDebugBuffer(ERayTracingSceneLayer Layer) const { return Layers[uint8(Layer)].InstanceDebugBuffer; }
-	FRDGBufferRef GetDebugInstanceGPUSceneIndexBuffer(ERayTracingSceneLayer Layer) const { return Layers[uint8(Layer)].DebugInstanceGPUSceneIndexBuffer; }
+	FRDGBufferRef GetInstanceExtraDataBuffer(ERayTracingSceneLayer Layer) const { return Layers[uint8(Layer)].InstanceExtraDataBuffer; }
 
 	void InitPreViewTranslation(const FViewMatrices& ViewMatrices);
 
@@ -164,7 +164,7 @@ public:
 	// Geometries which still have a pending build request but are used this frame and require a force build.
 	TArray<const FRayTracingGeometry*> GeometriesToBuild;
 
-	bool bNeedsDebugInstanceGPUSceneIndexBuffer = false;
+	bool bNeedsInstanceExtraDataBuffer = false;
 	bool bTracingFeedbackEnabled = false;
 
 	// Used for transforming to translated world space in which TLAS was built.
@@ -198,9 +198,10 @@ private:
 		FRDGBufferRef RayTracingSceneBufferRDG;
 		FRDGBufferSRVRef RayTracingSceneBufferSRV;
 
+		FRDGBufferRef InstanceExtraDataBuffer = nullptr;
+
 		// Special data for debugging purposes
 		FRDGBufferRef InstanceDebugBuffer = nullptr;
-		FRDGBufferRef DebugInstanceGPUSceneIndexBuffer = nullptr;
 
 		// Persistent storage for ray tracing instance descriptors.
 		// Cleared every frame without releasing memory to avoid large heap allocations.
diff --git a/Engine/Source/Runtime/Renderer/Private/SceneVisibility.cpp b/Engine/Source/Runtime/Renderer/Private/SceneVisibility.cpp
--- a/Engine/Source/Runtime/Renderer/Private/SceneVisibility.cpp
+++ b/Engine/Source/Runtime/Renderer/Private/SceneVisibility.cpp
@@ -4832,7 +4832,7 @@ void FSceneRenderer::PreVisibilityFrameSetup(FRDGBuilder& GraphBuilder)
 			const FViewInfo& ReferenceView = Views[ReferenceViewIndex];
 
 			Scene->RayTracingScene.InitPreViewTranslation(ReferenceView.ViewMatrices);
-			Scene->RayTracingScene.bNeedsDebugInstanceGPUSceneIndexBuffer = IsRayTracingInstanceOverlapEnabled(ReferenceView);
+			Scene->RayTracingScene.bNeedsInstanceExtraDataBuffer = IsRayTracingInstanceOverlapEnabled(ReferenceView);
 			Scene->RayTracingScene.bTracingFeedbackEnabled = IsRayTracingFeedbackEnabled(ViewFamily);
 
 			for (int32 ViewIndex = 0; ViewIndex < Views.Num(); ViewIndex++)
diff --git a/Engine/Source/Runtime/Renderer/Public/RayTracingInstanceBufferUtil.h b/Engine/Source/Runtime/Renderer/Public/RayTracingInstanceBufferUtil.h
--- a/Engine/Source/Runtime/Renderer/Public/RayTracingInstanceBufferUtil.h
+++ b/Engine/Source/Runtime/Renderer/Public/RayTracingInstanceBufferUtil.h
@@ -20,7 +20,11 @@ struct UE_DEPRECATED(5.6, "Use FRayTracingInstanceBufferBuilder instead.") FRayT
 	uint32 InstanceId;
 	uint32 InstanceMaskAndFlags;
 	uint32 InstanceContributionToHitGroupIndex;
-	uint32 bApplyLocalBoundsTransform;
+	union
+	{
+		uint32 SceneInstanceIndexAndApplyLocalBoundsTransform;
+		uint32 bApplyLocalBoundsTransform;
+	};
 };
 
 struct UE_DEPRECATED(5.6, "Use FRayTracingInstanceBufferBuilder instead.") FRayTracingInstanceGroupEntryRef
@@ -95,7 +99,7 @@ RENDERER_API void BuildRayTracingInstanceBuffer(
 	uint32 NumNativeCPUInstances,
 	const FRayTracingCullingParameters* CullingParameters,
 	FRHIUnorderedAccessView* OutputStatsUAV,
-	FRHIUnorderedAccessView* DebugInstanceGPUSceneIndexUAV);
+	FRHIUnorderedAccessView* InstanceExtraDataUAV);
 
 PRAGMA_ENABLE_DEPRECATION_WARNINGS
 
@@ -115,7 +119,7 @@ public:
 		FRHIUnorderedAccessView* InstancesUAV,
 		FRHIUnorderedAccessView* OutputStatsUAV,
 		uint32 OutputStatsOffset,
-		FRHIUnorderedAccessView* DebugInstanceGPUSceneIndexUAV);
+		FRHIUnorderedAccessView* InstanceExtraDataUAV);
 
 	PRAGMA_DISABLE_DEPRECATION_WARNINGS
 
