First Person: Removed ensure on GPUScene being enabled in FirstPersonSceneExtension and instead simply made it a condition of enabling the extension. This works around some failed CI tests, but should really be addressed by making GPUSceneViewId independently available  of GPUScene being enabled.

diff --git a/Engine/Source/Runtime/Renderer/Private/FirstPersonSceneExtension.cpp b/Engine/Source/Runtime/Renderer/Private/FirstPersonSceneExtension.cpp
--- a/Engine/Source/Runtime/Renderer/Private/FirstPersonSceneExtension.cpp
+++ b/Engine/Source/Runtime/Renderer/Private/FirstPersonSceneExtension.cpp
@@ -9,10 +9,10 @@ IMPLEMENT_SCENE_EXTENSION(FFirstPersonSceneExtension);
 
 bool FFirstPersonSceneExtension::ShouldCreateExtension(FScene& Scene)
 {
-	ensureMsgf(Scene.GPUScene.IsEnabled(), TEXT("FFirstPersonSceneExtension relies on GPUScene to be enabled so that GPUSceneViewId on FViewInfo objects is properly initialized!"));
 	// For now, the bounds computed by this extension are only needed for First Person Self-Shadow and Lumen HWRT reflections of
 	// FirstPersonWorldSpaceRepresentation primitives. Both these features require the first person gbuffer bit.
-	return HasFirstPersonGBufferBit(GetFeatureLevelShaderPlatform(Scene.GetFeatureLevel()));
+	// Further, this extension relies on GPUSceneViewId on FViewInfo, which is initialized by GPUScene and therefore requires that to be enabled as well.
+	return HasFirstPersonGBufferBit(GetFeatureLevelShaderPlatform(Scene.GetFeatureLevel())) && Scene.GPUScene.IsEnabled();
 }
 
 ISceneExtensionUpdater* FFirstPersonSceneExtension::CreateUpdater()
