[TSR ] Improve thin geometry (especially foliage) temporal stability.

TSR now has a thin geometry pass to classify foliage thin geometry into three different regions:
1. Edge line. Vertical, horizontal and diagnal single pixel lines based on depth. This allows immediate response to  history rejection.
2. Thin geometry cluster hole region. Pixel region that has parrtial coverage. This classification is based a history partial coverage mask, where the closer to 0.5, the higher chance an instability happens. The partial coverage mask is accumulated across frame and rejected based on statistic test (p value=0.05) of neighbor 3x3 to history. Newly disoccluded pixels are rejected as other region to use the regular history rectifier.
3. Other region. All regions not belowing to 1 and 2.

For edge line and cluster hole region, we apply a different history rectificer that allows us to perform history relaxation where the relaxation factor (r.TSR.ThinGeometryDetection.Coverage.MaxRelaxationWeight, default 0.037) controls the relaxation. Other region uses the default rectifying logic. Use r.tsr.visualize 15 to see the region classification: Edge line (red), cluster hole (green), and other (yellow).

Use r.tsr.thingeometrydetection 1/0 to turn on/off.

#tests Test done in naniteResearch and project stream on both PC and console

diff --git a/Engine/Shaders/Private/TemporalSuperResolution/TSRConvolutionNetwork.ush b/Engine/Shaders/Private/TemporalSuperResolution/TSRConvolutionNetwork.ush
--- a/Engine/Shaders/Private/TemporalSuperResolution/TSRConvolutionNetwork.ush
+++ b/Engine/Shaders/Private/TemporalSuperResolution/TSRConvolutionNetwork.ush
@@ -882,6 +882,45 @@ TLaneVector2D<ScalarType, VectorSize, LaneStrideX, LaneStrideY> Clamp3x3(
 	return fastClamp(ToClamp, BoundaryMin, BoundaryMax);
 }
 
+CALL_SITE_DEBUGLOC
+template<typename ScalarType, uint VectorSize, uint LaneStrideX, uint LaneStrideY>
+TLaneVector2D<ScalarType, VectorSize, LaneStrideX, LaneStrideY> LerpClamp3x3(
+	TLaneVector2D<ScalarType, VectorSize, LaneStrideX, LaneStrideY> ToClamp,
+	TLaneVector2D<ScalarType, VectorSize, LaneStrideX, LaneStrideY> Boundary /*Boundary inside min max of ToClamp)*/,
+	TLaneVector2D<ScalarType, 1, LaneStrideX, LaneStrideY> LerpFactor)
+{
+    #if 0
+	TLaneVector2D<ScalarType, VectorSize, LaneStrideX, LaneStrideY> BoundaryMin;
+	TLaneVector2D<ScalarType, VectorSize, LaneStrideX, LaneStrideY> BoundaryMax;
+	MinMax3x3(Boundary, /* out */ BoundaryMin, /* out */ BoundaryMax);
+
+	TLaneVector2D<ScalarType, VectorSize, LaneStrideX, LaneStrideY> ToClampBoundaryMin;
+	TLaneVector2D<ScalarType, VectorSize, LaneStrideX, LaneStrideY> ToClampBoundaryMax;
+	MinMax3x3(ToClamp, /* out */ ToClampBoundaryMin, /* out */ ToClampBoundaryMax);
+
+	TLaneVector2D<ScalarType, VectorSize, LaneStrideX, LaneStrideY> LerpFactorVector = TLaneVector2D<ScalarType, VectorSize, LaneStrideX, LaneStrideY>::Vectorize(LerpFactor);
+	BoundaryMin = lerp(BoundaryMin, ToClampBoundaryMin, LerpFactorVector);
+	BoundaryMax = lerp(BoundaryMax, ToClampBoundaryMax, LerpFactorVector);
+    #else
+	// register pressure is high
+	TLaneVector2D<ScalarType, VectorSize, LaneStrideX, LaneStrideY> BoundaryMin;
+	TLaneVector2D<ScalarType, VectorSize, LaneStrideX, LaneStrideY> BoundaryMax;
+	TLaneVector2D<ScalarType, VectorSize, LaneStrideX, LaneStrideY> BoundaryLimit;
+	TLaneVector2D<ScalarType, VectorSize, LaneStrideX, LaneStrideY> ToClampBoundaryLimit;
+	TLaneVector2D<ScalarType, VectorSize, LaneStrideX, LaneStrideY> LerpFactorVector = TLaneVector2D<ScalarType, VectorSize, LaneStrideX, LaneStrideY>::Vectorize(LerpFactor);
+
+	BoundaryLimit = Min3x3(Boundary);
+	ToClampBoundaryLimit = Min3x3(ToClamp);
+	BoundaryMin = lerp(BoundaryLimit, ToClampBoundaryLimit, LerpFactorVector);
+
+	BoundaryLimit = Max3x3(Boundary);
+	ToClampBoundaryLimit = Max3x3(ToClamp);
+	BoundaryMax = lerp(BoundaryLimit, ToClampBoundaryLimit, LerpFactorVector);
+    #endif
+
+	return fastClamp(ToClamp, BoundaryMin, BoundaryMax);
+}
+
 CALL_SITE_DEBUGLOC
 template<typename ScalarType, uint VectorSize, uint LaneStrideX, uint LaneStrideY>
 TLaneVector2D<ScalarType, VectorSize, LaneStrideX, LaneStrideY> AnnihilateToGuide3x3(
diff --git a/Engine/Shaders/Private/TemporalSuperResolution/TSRDecimateHistory.usf b/Engine/Shaders/Private/TemporalSuperResolution/TSRDecimateHistory.usf
--- a/Engine/Shaders/Private/TemporalSuperResolution/TSRDecimateHistory.usf
+++ b/Engine/Shaders/Private/TemporalSuperResolution/TSRDecimateHistory.usf
@@ -20,6 +20,8 @@
 /** Whether to compile the resurrection reprojection. */
 #define CONFIG_RESURRECTION_REPROJECTION (DIM_RESURRECTION_REPROJECTION)
 
+/** Whether to compile the thin geometry coverage reprojection.*/
+#define CONFIG_THIN_GEOMETRY_COVERAGE_REPROJECTION (DIM_THIN_GEOMETRY_COVERAGE_REPROJECTION)
 
 //------------------------------------------------------- PARAMETERS
 
@@ -37,6 +39,9 @@ Texture2DArray<uint> PrevAtomicTextureArray;
 
 Texture2DArray<tsr_half4> PrevHistoryGuide;
 Texture2DArray<tsr_half4> PrevHistoryMoire;
+#if CONFIG_THIN_GEOMETRY_COVERAGE_REPROJECTION
+    Texture2DArray<tsr_half> PrevHistoryCoverage;
+#endif
 float2 PrevGuideInfo_Extent;
 float2 PrevGuideInfo_ExtentInverse;
 uint2  PrevGuideInfo_ViewportMin;
@@ -57,7 +62,9 @@ float4x4 ClipToResurrectionClip;
 
 RWTexture2DArray<tsr_half4> ReprojectedHistoryGuideOutput;
 RWTexture2DArray<tsr_half4> ReprojectedHistoryMoireOutput;
-
+#if CONFIG_THIN_GEOMETRY_COVERAGE_REPROJECTION
+    RWTexture2DArray<tsr_half> ReprojectedHistoryCoverageOutput;
+#endif
 RWTexture2DArray<uint> ReprojectionFieldOutput;
 RWTexture2D<tsr_half2> DecimateMaskOutput;
 
@@ -114,6 +121,13 @@ tsr_half4 CorrectMoireExposure(tsr_half4 ReprojectedHistoryMoireColor, const flo
 }
 #endif
 
+#if CONFIG_THIN_GEOMETRY_COVERAGE_REPROJECTION
+tsr_half ReprojectPrevHistoryCoverage(tsr_ushort2 InputPixelPos, float2 TextureUV)
+{
+	return PrevHistoryCoverage.SampleLevel(GlobalPointClampedSampler, float3(TextureUV, 0), 0);
+}
+#endif
+
 //ISOLATE
 void ReprojectAllPrevTextures(
 	uint2 GroupId,
@@ -129,6 +143,9 @@ void ReprojectAllPrevTextures(
 	#if CONFIG_MOIRE_REPROJECTION
 		out tsr_half4 ReprojectedHistoryMoireColor,
 	#endif
+    #if CONFIG_THIN_GEOMETRY_COVERAGE_REPROJECTION
+	    out tsr_half ReprojectedHistoryCoverage,
+    #endif
 	#if CONFIG_RESURRECTION_REPROJECTION
 		out tsr_halfC ResurrectedHistoryGuideColor,
 		out tsr_half  ResurrectedHistoryGuideUncertainty,
@@ -197,6 +214,12 @@ void ReprojectAllPrevTextures(
 			ReprojectedHistoryMoireColor = ReprojectPrevHistoryMoire(InputPixelPos, KernelHistoryBufferUV);
 		}
 		#endif
+
+        #if CONFIG_THIN_GEOMETRY_COVERAGE_REPROJECTION
+		{
+			ReprojectedHistoryCoverage = ReprojectPrevHistoryCoverage(InputPixelPos, KernelHistoryBufferUV);
+		}
+        #endif
 	}
 
 	#if CONFIG_RESURRECTION_REPROJECTION
@@ -314,6 +337,9 @@ void MainCS(
 	#if CONFIG_MOIRE_REPROJECTION
 		tsr_half4 ReprojectedHistoryMoireColor;
 	#endif
+    #if CONFIG_THIN_GEOMETRY_COVERAGE_REPROJECTION
+		tsr_half ReprojectedHistoryCoverage;
+    #endif
 	#if CONFIG_RESURRECTION_REPROJECTION
 		tsr_halfC ResurrectedHistoryGuideColor;
 		tsr_half  ResurrectedHistoryGuideUncertainty;
@@ -337,6 +363,9 @@ void MainCS(
 			#if CONFIG_MOIRE_REPROJECTION
 				/* out */ ReprojectedHistoryMoireColor,
 			#endif
+			#if CONFIG_THIN_GEOMETRY_COVERAGE_REPROJECTION
+				/* out */ ReprojectedHistoryCoverage,
+			#endif
 			#if CONFIG_RESURRECTION_REPROJECTION
 				/* out */ ResurrectedHistoryGuideColor,
 				/* out */ ResurrectedHistoryGuideUncertainty,
@@ -463,6 +492,9 @@ void MainCS(
 			#if CONFIG_MOIRE_REPROJECTION
 				ReprojectedHistoryMoireOutput[tsr_short3(OutputPixelPos, 0)] = ReprojectedHistoryMoireColor;
 			#endif
+			#if CONFIG_THIN_GEOMETRY_COVERAGE_REPROJECTION
+				ReprojectedHistoryCoverageOutput[tsr_short3(OutputPixelPos, 0)] = ReprojectedHistoryCoverage;
+			#endif
 		}
 
 		#if CONFIG_RESURRECTION_REPROJECTION
diff --git a/Engine/Shaders/Private/TemporalSuperResolution/TSRDetectThinGeometry.usf b/Engine/Shaders/Private/TemporalSuperResolution/TSRDetectThinGeometry.usf
--- a/Engine/Shaders/Private/TemporalSuperResolution/TSRDetectThinGeometry.usf
+++ b/Engine/Shaders/Private/TemporalSuperResolution/TSRDetectThinGeometry.usf
@@ -19,6 +19,8 @@
 
 #define EDGE_TYPE EDGE_HORIZONTAL_VERTICAL_DIAGNAL
 
+#define CONFIG_THIN_GEOMETRY_COVERAGE 1
+
 //------------------------------------------------------- CONSTANTS
 
 #if CONFIG_LDS_DESPILL
@@ -31,8 +33,16 @@
 int TileOverscan;
 int ThinGeometryTextureIndex;
 float ErrorMultiplier;
+float MaxRelaxationWeight;
 
 Texture2D<float> SceneDepthTexture;
+
+#if CONFIG_THIN_GEOMETRY_COVERAGE
+    Texture2D<tsr_half> CurrentCoverageTexture;
+    Texture2D<tsr_half> ReprojectedHistoryCoverageTexture;
+
+    RWTexture2DArray<tsr_half> HistoryCoverageOutput;
+#endif
 RWTexture2DArray<uint> R8Output;
 
 //------------------------------------------------------- GLOBALS
@@ -84,9 +94,14 @@ tsr_short2 ComputeElementOutputPixelPos(const uint ElementIndex)
 }
 
 void FetchInput(
-	tsr_tensor_short2 PixelPos,
-	out tsr_tensor_half SceneDepth,
-	out tsr_tensor_half DeviceZThinGeometryDetectionError)
+	tsr_tensor_short2 PixelPos
+	,out tsr_tensor_half SceneDepth
+	,out tsr_tensor_half DeviceZThinGeometryDetectionError
+#if CONFIG_THIN_GEOMETRY_COVERAGE
+	,out tsr_tensor_half CurrentCoverage
+	,out tsr_tensor_half ReprojectedHistoryCoverage
+#endif
+	)
 {
 	tsr_tensor_half LumaAbsDelta;
 
@@ -95,6 +110,11 @@ void FetchInput(
 	{
 		tsr_short2 ClampedFetchPixelPos = PixelPos.GetElement(ElementIndex);
 
+        #if CONFIG_THIN_GEOMETRY_COVERAGE
+		    CurrentCoverage.SetElement(ElementIndex, CurrentCoverageTexture[ClampedFetchPixelPos]);
+		    ReprojectedHistoryCoverage.SetElement(ElementIndex, ReprojectedHistoryCoverageTexture[ClampedFetchPixelPos]);
+        #endif
+
 		float Depth = SceneDepthTexture[ClampedFetchPixelPos];
 		float DeviceZError = ComputePixelDeviceZError(Depth);
 
@@ -198,6 +218,27 @@ TLaneVector2D<ScalarType, VectorSize, LaneStrideX, LaneStrideY> VerticalThinEdge
 	return ThinGeometryEdge;
 }
 
+tsr_tensor_uint EncodeThinGeometryWeight(tsr_tensor_half ThinGeometryEdge
+	, tsr_tensor_bool IsClusterHoleRegion,
+	tsr_tensor_half Weight
+)
+{
+#if CONFIG_THIN_GEOMETRY_COVERAGE
+	// 7bit   | 1bit
+	// weight | 0 : Edge line
+	//		  | 1 : Cluster hole region
+	Weight = select(IsClusterHoleRegion, Weight, tsr_tensor_half::Const(0));
+	Weight = select(ThinGeometryEdge > tsr_tensor_half::Const(0), ThinGeometryEdge, Weight);
+	IsClusterHoleRegion = select(ThinGeometryEdge > tsr_tensor_half::Const(0), tsr_tensor_bool::Const(false), IsClusterHoleRegion);
+	tsr_tensor_uint Encoded = bit_shift_left(tsr_tensor_uint::CastFrom<tsr_half>(Weight * tsr_half(127.0)),tsr_tensor_uint::Const(0x1u));
+	Encoded = bit_or(Encoded,select(IsClusterHoleRegion,tsr_tensor_uint::Const(0x1u),tsr_tensor_uint::Const(0)));
+	return Encoded;
+#else
+	tsr_tensor_uint Encoded = bit_shift_left(tsr_tensor_uint::CastFrom<tsr_half>(ThinGeometryEdge * tsr_half(127.0)),tsr_tensor_uint::Const(0x1u));
+	return Encoded;
+#endif
+}
+
 //------------------------------------------------------- ENTRY POINT
 
 #if COMPILER_SUPPORTS_WAVE_SIZE && DIM_WAVE_SIZE > 0
@@ -225,8 +266,18 @@ void MainCS(
 	tsr_tensor_short2 PixelPos = ComputePixelPos();
 	tsr_tensor_half DeviceZ;
 	tsr_tensor_half DeviceZThinGeometryDetectionError;
+#if CONFIG_THIN_GEOMETRY_COVERAGE
+	tsr_tensor_half CurrentCoverage;
+	tsr_tensor_half ReprojectedHistoryCoverage;
+#endif
 	FetchInput(PixelPos,
-		DeviceZ, DeviceZThinGeometryDetectionError);
+		DeviceZ, 
+		DeviceZThinGeometryDetectionError
+#if CONFIG_THIN_GEOMETRY_COVERAGE
+		,CurrentCoverage
+		,ReprojectedHistoryCoverage
+#endif
+	);
 	
 	//Debug[0].SetComponent(0, DeviceZ);
 	//Debug[0].SetComponent(1, DeviceZThinGeometryDetectionError);
@@ -247,6 +298,92 @@ void MainCS(
 		//Debug[0].SetComponent(3, tsr_tensor_half(ThinGeometryEdge));
 	}
 
+	tsr_tensor_half HistoryRelaxationWeight = tsr_tensor_half::Const(0);
+	tsr_tensor_bool IsClusterHoleRegion = tsr_tensor_bool::Const(false);
+	tsr_tensor_half NewSampleWeight = tsr_tensor_half::Const(0.1);
+#if CONFIG_THIN_GEOMETRY_COVERAGE
+	{
+		// Student t test of current observation and history to check if they are of the same distribution.
+		tsr_tensor_bool bSameDistribution = tsr_tensor_bool::Const(false);
+		tsr_tensor_bool NotAbsoluteMatch = tsr_tensor_bool::Const(false);
+		tsr_tensor_half T;
+		{ // Student T test on Bernouli distribution of 3x3 neighbours. If T > T critical, reject weights due to low similarity.
+
+			tsr_tensor_half3 ExpectObservationAndVariance = Concatenate(Concatenate(ReprojectedHistoryCoverage,CurrentCoverage),ReprojectedHistoryCoverage * (tsr_tensor_half::Const(1.0) - ReprojectedHistoryCoverage));
+			ExpectObservationAndVariance = Sum3x3(ExpectObservationAndVariance);
+
+			tsr_tensor_half XAverage = ExpectObservationAndVariance[0];
+			XAverage = (XAverage - ExpectObservationAndVariance[1]);
+			tsr_tensor_half S = sqrt(ExpectObservationAndVariance[2] + tsr_tensor_half::Const(tsr_half(0.000001)))
+				* sqrt(NewSampleWeight * rcp(tsr_tensor_half::Const(2.0) - NewSampleWeight) * tsr_half(72.0));
+			// When observation is equivalent to the expectation, t = 0.
+			T = abs(select( XAverage > tsr_tensor_half::Const(0), XAverage * rcp(S), tsr_tensor_half::Const(0)));
+			bSameDistribution = T < tsr_tensor_half::Const(20.754); //2.306 (p=0.05) * 9
+			
+			Debug[0].SetComponent(0,abs(T));
+			Debug[0].SetComponent(1,tsr_tensor_half::CastFrom<bool>(abs(T) > tsr_tensor_half::Const(20.754)));
+		}
+
+		tsr_tensor_bool ValidCurrentThinGeometryRegion = tsr_tensor_bool::Const(false);
+		{
+			tsr_tensor_half2 TAndCoverage = Concatenate(T, CurrentCoverage);
+			TAndCoverage = Max3x3(TAndCoverage);
+			TAndCoverage = Max3x3(TAndCoverage);
+			NotAbsoluteMatch =TAndCoverage[0]/* Max3x3(Max3x3(T))*/ > tsr_tensor_half::Const(0);
+			ValidCurrentThinGeometryRegion = TAndCoverage[1]/*Max3x3(Max3x3(CurrentCoverage))*/ > tsr_tensor_half::Const(0);
+		}
+
+		// Sample weight for coverage update.
+		{
+			// Reject history coverage if history is not valid or absolute match.
+			// TODO: Absolute coverage match might has similar color, then no need to reset history coverage.
+			// e.g., 5x5 flat surface disappear and 5x5 flat surface appear. This is uncommon in thin geometry
+			// yet need to address to be less aggressive.
+			tsr_tensor_bool AcceptHistory = and(ValidCurrentThinGeometryRegion, and(bSameDistribution, NotAbsoluteMatch));
+			NewSampleWeight = select(AcceptHistory, NewSampleWeight, tsr_tensor_half::Const(1.0));
+
+			// Update reprojected history and use it for weight calculation.
+			ReprojectedHistoryCoverage = lerp(ReprojectedHistoryCoverage,CurrentCoverage,NewSampleWeight);
+			Debug[0].SetComponent(1, NewSampleWeight);
+		}
+
+		// Update new history relaxation weight.
+		{
+			// Calculate the weights.
+			{
+				HistoryRelaxationWeight = abs(ReprojectedHistoryCoverage - tsr_tensor_half::Const(0.5))+tsr_tensor_half::Const(0.02);
+				HistoryRelaxationWeight = HistoryRelaxationWeight * exp(HistoryRelaxationWeight * tsr_half(-20.0)) *  tsr_half(10.0/0.184); // delta = 1 maximum
+			}
+			tsr_tensor_half Min5x5HistoryCoverage = Min3x3(Min3x3(ReprojectedHistoryCoverage));
+			tsr_tensor_bool bThinGeometryHistoryRegion = Min5x5HistoryCoverage > tsr_tensor_half::Const(0);
+			
+			// Reject samples that has min being background
+			// Set relaxation weights to 0 if the history is rejected to avoid overblur.
+			tsr_tensor_bool bValidHistoryRelaxation = and(bThinGeometryHistoryRegion,and(ValidCurrentThinGeometryRegion, bSameDistribution));
+			HistoryRelaxationWeight = select(bValidHistoryRelaxation, HistoryRelaxationWeight, tsr_tensor_half::Const(0));
+			
+			//Diffuse a little bit to suppress the history clamping around the valid pixel and apply max relaxation weight.
+			HistoryRelaxationWeight = saturate(WeightedAvg3x3(HistoryRelaxationWeight, float3(0.8, 1.0, 0.8))*tsr_half(MaxRelaxationWeight));
+			HistoryRelaxationWeight = saturate(WeightedAvg3x3(HistoryRelaxationWeight, float3(0.5, 1.0, 0.5)));
+
+			// Not similar patches should reset relaxation weight to 0
+			// to use more strict history clamping.
+			HistoryRelaxationWeight = select(bSameDistribution, HistoryRelaxationWeight, tsr_tensor_half::Const(0));
+			
+			Debug[0].SetComponent(2,HistoryRelaxationWeight);
+		}
+
+		// Valid pixels for cluster hole stability improvement.
+		IsClusterHoleRegion = ReprojectedHistoryCoverage > tsr_tensor_half::Const(half(0.99f/127.0f));
+	}
+#endif
+
+	tsr_tensor_uint EncodedThinGeometryEdge = EncodeThinGeometryWeight(
+		ThinGeometryEdge
+		,IsClusterHoleRegion
+		,HistoryRelaxationWeight
+	);
+
 	// Output the thin geometry mask
 	UNROLL_N(SIMD_SIZE)
 	for(uint ElementIndex = 0; ElementIndex < SIMD_SIZE; ElementIndex++)
@@ -254,8 +391,11 @@ void MainCS(
 		const tsr_short2 LaneSimdPixelOffset = GetElementPixelOffsetInTile(LANE_STRIDE_X, LANE_STRIDE_Y, GGroupThreadIndex, ElementIndex);
 		const tsr_short2 OutputPixelPos = ComputeElementOutputPixelPos(ElementIndex);
 		
-		tsr_half ThinGeometryFlag = ThinGeometryEdge.GetElement(ElementIndex);
-		R8Output[uint3(OutputPixelPos,ThinGeometryTextureIndex)] = uint(ThinGeometryFlag * tsr_half(255.0));
+		R8Output[uint3(OutputPixelPos,ThinGeometryTextureIndex)] = EncodedThinGeometryEdge.GetElement(ElementIndex);//uint(ThinGeometryFlag * tsr_half(255.0));
+
+        #if CONFIG_THIN_GEOMETRY_COVERAGE
+		    HistoryCoverageOutput[uint3(OutputPixelPos,0)] = ReprojectedHistoryCoverage.GetElement(ElementIndex);
+		#endif
 	}
 	
 	#if DEBUG_OUTPUT
diff --git a/Engine/Shaders/Private/TemporalSuperResolution/TSRMeasureCoverage.usf b/Engine/Shaders/Private/TemporalSuperResolution/TSRMeasureCoverage.usf
new file mode 100644
--- /dev/null
+++ b/Engine/Shaders/Private/TemporalSuperResolution/TSRMeasureCoverage.usf
@@ -0,0 +1,68 @@
+// Copyright Epic Games, Inc. All Rights Reserved.
+
+
+#include "../DeferredShadingCommon.ush"  
+#include "TSRColorSpace.ush"
+
+//------------------------------------------------------- CONFIG
+
+#define TILE_SIZE 8
+
+#define TILE_ROWS 2
+
+//------------------------------------------------------- PARAMETERS
+
+int OutputArrayIndex;
+RWTexture2DArray<float> ThinGeometryCoverageOutput;
+
+uint GetThinGeometry(const FGBufferData GBufferData)
+{
+	return GBufferData.ShadingModelID == SHADINGMODELID_TWOSIDED_FOLIAGE ? 1u : 0u;	
+}
+
+//------------------------------------------------------- ENTRY POINT
+
+[numthreads(TILE_SIZE * TILE_SIZE, 1, 1)]
+void MainCS(
+	uint2 GroupId : SV_GroupID,
+	uint GroupThreadIndex : SV_GroupIndex)
+{
+	uint2 GroupThreadId = uint2(GroupThreadIndex % TILE_SIZE, GroupThreadIndex / TILE_SIZE);
+	uint2 OutputPixelOrigin = (InputInfo_ViewportMin + GroupId * uint(TILE_SIZE * TILE_ROWS)) + GroupThreadId;
+
+	uint ThinGeometryArray[TILE_ROWS * TILE_ROWS];
+
+	ISOLATE
+	{
+		UNROLL_N(TILE_ROWS)
+		for (uint y = 0; y < TILE_ROWS; y++)
+		{
+			UNROLL_N(TILE_ROWS)
+			for (uint x = 0; x < TILE_ROWS; x++)
+			{
+				const uint2 PixelOffset = uint2(x, y) * TILE_SIZE;
+				uint2 InputPixelPos = OutputPixelOrigin + PixelOffset;
+
+				ThinGeometryArray[x + y * TILE_ROWS] = GetThinGeometry(GetGBufferDataUint(InputPixelPos));
+			}
+		}
+	}
+
+	ISOLATE
+	{
+		UNROLL_N(TILE_ROWS)
+		for (uint y = 0; y < TILE_ROWS; y++)
+		{
+			UNROLL_N(TILE_ROWS)
+			for (uint x = 0; x < TILE_ROWS; x++)
+			{
+				const uint2 PixelOffset = uint2(x, y) * TILE_SIZE;
+				uint2 InputPixelPos = OutputPixelOrigin + PixelOffset;
+				float ThinGeometry = (float)ThinGeometryArray[x + y * TILE_ROWS];
+
+				InputPixelPos.x = select(all(InputPixelPos < InputInfo_ViewportMax), InputPixelPos.x, uint(~0));
+				ThinGeometryCoverageOutput[int3(InputPixelPos, OutputArrayIndex)] = ThinGeometry;
+			}
+		}
+	}
+}
diff --git a/Engine/Shaders/Private/TemporalSuperResolution/TSRMeasureFlickeringLuma.usf b/Engine/Shaders/Private/TemporalSuperResolution/TSRMeasureFlickeringLuma.usf
--- a/Engine/Shaders/Private/TemporalSuperResolution/TSRMeasureFlickeringLuma.usf
+++ b/Engine/Shaders/Private/TemporalSuperResolution/TSRMeasureFlickeringLuma.usf
@@ -13,7 +13,7 @@
 //------------------------------------------------------- PARAMETERS
 
 Texture2D<tsr_half3> SceneColorTexture;
-RWTexture2D<tsr_half> FlickeringLumaOutput;
+RWTexture2DArray<tsr_half> FlickeringLumaOutput;
 
 
 //------------------------------------------------------- ENTRY POINT
@@ -62,7 +62,7 @@ void MainCS(
 
 				InputPixelPos.x = select(all(InputPixelPos < InputInfo_ViewportMax), InputPixelPos.x, uint(~0));
 
-				FlickeringLumaOutput[InputPixelPos] = GLuma;
+				FlickeringLumaOutput[int3(InputPixelPos, 0)] = GLuma;
 			}
 		}
 	}
diff --git a/Engine/Shaders/Private/TemporalSuperResolution/TSRRejectShading.usf b/Engine/Shaders/Private/TemporalSuperResolution/TSRRejectShading.usf
--- a/Engine/Shaders/Private/TemporalSuperResolution/TSRRejectShading.usf
+++ b/Engine/Shaders/Private/TemporalSuperResolution/TSRRejectShading.usf
@@ -2,6 +2,8 @@
 
 //------------------------------------------------------- INCLUDES
 
+#define CONFIG_THIN_GEOMETRY_DETECTION (DIM_THIN_GEOMETRY_DETECTION)
+
 #include "TSRShadingAnalysis.ush"
 #include "TSRReprojectionField.ush"
 
@@ -19,7 +21,6 @@
 
 #define CONFIG_HISTORY_RESURRECTION (DIM_HISTORY_RESURRECTION)
 
-#define CONFIG_THIN_GEOMETRY_DETECTION (DIM_THIN_GEOMETRY_DETECTION)
 
 //------------------------------------------------------- CONSTANTS
 
@@ -120,14 +121,20 @@ tsr_short2 ComputeElementOutputPixelPos(const uint ElementIndex)
 
 void FetchThinGeometryEdge(
 	tsr_tensor_short2 PixelPos,
-	out tsr_tensor_half IsThinGeometryEdge)
+	out tsr_tensor_half HistoryRelaxationWeight,
+	out tsr_tensor_bool IsClusterHoleRegion)
 {
 	UNROLL_N(SIMD_SIZE)
 	for (uint ElementIndex = 0; ElementIndex < SIMD_SIZE; ElementIndex++)
 	{
 		tsr_short2 ClampedFetchPixelPos = PixelPos.GetElement(ElementIndex);
+		uint EncodedThinGeometryWeight = ThinGeometryTexture[ClampedFetchPixelPos];
 		
-		IsThinGeometryEdge.SetElement(ElementIndex, tsr_half(ThinGeometryTexture[ClampedFetchPixelPos]) * rcp(tsr_half(255.0)));
+		tsr_half Weight = tsr_half(EncodedThinGeometryWeight >> 1) * rcp(tsr_half(127.0));
+		bool UseMultiplyOperation = bool(EncodedThinGeometryWeight & 0x1u);
+
+		HistoryRelaxationWeight.SetElement(ElementIndex, Weight);
+		IsClusterHoleRegion.SetElement(ElementIndex, UseMultiplyOperation);
 	}
 } // FetchMoireInput()
 
@@ -604,6 +611,14 @@ void MainCS(
 	tsr_tensor_half RejectionBlendFinal = tsr_tensor_half::Const(1.0);
 	tsr_tensor_half RejectionClampBlend = tsr_tensor_half::Const(1.0);
 	tsr_tensor_bool bAliasingIsVisible = tsr_tensor_bool::Const(false);
+#if CONFIG_THIN_GEOMETRY_DETECTION
+	tsr_tensor_half HistoryRelaxationFactor = tsr_tensor_half::Const(0);
+	{
+		tsr_tensor_short2 PixelPos = ComputePixelPos();
+		tsr_tensor_bool IsClusterHoleRegion;
+		FetchThinGeometryEdge(PixelPos, HistoryRelaxationFactor, IsClusterHoleRegion);
+	}
+#endif
 	#if CONFIG_HISTORY_RESURRECTION
 		tsr_tensor_bool bResurrectionIsCloser;
 	#endif
@@ -653,6 +668,10 @@ void MainCS(
 			/* InputC2 = */    ClampedInput,
 			/* HistoryC2 = */  ClampedHistory,
 			/* MoireError = */ MoireError,
+#if CONFIG_THIN_GEOMETRY_DETECTION
+			HistoryRelaxationFactor ,
+			true,
+#endif
 			/* bEnableMoireError = */ CONFIG_MOIRE_DETECTION != 0 && bEnableFlickeringHeuristic != 0,
 			/* out */ RejectionBlendFinal,
 			/* out */ RejectionClampBlend);
@@ -709,6 +728,10 @@ void MainCS(
 			/* InputC2 = */    ClampedResurrectedInput,
 			/* HistoryC2 = */  ClampedResurrectedHistory,
 			/* MoireError = */ tsr_tensor_half::Const(0.0),
+#if CONFIG_THIN_GEOMETRY_DETECTION
+			/*HistoryRelaxationFactor */tsr_tensor_half::Const(0),
+			false,
+#endif
 			/* bEnableMoireError = */ CONFIG_MOIRE_DETECTION != 0 && bEnableFlickeringHeuristic != 0,
 			/* out */ ResurrectionRejectionBlendFinal,
 			/* out */ ResurrectionRejectionClampBlend);
@@ -720,17 +743,6 @@ void MainCS(
 	} // if (bComputeResurrection)
 	#endif // CONFIG_HISTORY_RESURRECTION
 	
-    #if CONFIG_THIN_GEOMETRY_DETECTION
-	PLATFORM_SPECIFIC_ISOLATE
-	{
-		tsr_tensor_short2 PixelPos = ComputePixelPos();
-		tsr_tensor_half IsThinGeometryEdge;
-		FetchThinGeometryEdge(PixelPos, IsThinGeometryEdge);
-		RejectionBlendFinal = lerp(RejectionBlendFinal, tsr_tensor_half::Const(tsr_half(1.0)), IsThinGeometryEdge);
-		RejectionClampBlend = lerp(RejectionClampBlend, tsr_tensor_half::Const(tsr_half(1.0)), IsThinGeometryEdge);
-	}
-    #endif
-	
 	// Update the guide for next frame.
 	tsr_tensor_halfC FinalGuide;
 	tsr_tensor_half  InputUncertainty;
diff --git a/Engine/Shaders/Private/TemporalSuperResolution/TSRShadingAnalysis.ush b/Engine/Shaders/Private/TemporalSuperResolution/TSRShadingAnalysis.ush
--- a/Engine/Shaders/Private/TemporalSuperResolution/TSRShadingAnalysis.ush
+++ b/Engine/Shaders/Private/TemporalSuperResolution/TSRShadingAnalysis.ush
@@ -9,6 +9,11 @@
 #include "TSRSpatialAntiAliasing.ush"
 #include "TSRConvolutionNetworkPass.ush"
 
+//------------------------------------------------------- CONFIG
+
+#ifndef CONFIG_THIN_GEOMETRY_DETECTION
+#define CONFIG_THIN_GEOMETRY_DETECTION 0
+#endif
 
 //------------------------------------------------------- CONSTANTS
 
@@ -448,6 +453,10 @@ void MeasureRejection(
 	tsr_tensor_halfC InputC2,
 	tsr_tensor_halfC HistoryC2,
 	tsr_tensor_half  MoireError,
+#if CONFIG_THIN_GEOMETRY_DETECTION
+	tsr_tensor_half HistoryRelaxationFactor,
+	const bool bEnableHistoryControl,
+#endif
 	bool bEnableMoireError,
 	out tsr_tensor_half OutRejectionBlendFinal,
 	out tsr_tensor_half OutRejectionClampBlend)
@@ -472,6 +481,19 @@ void MeasureRejection(
 	tsr_tensor_halfC FilteredBoxMin;
 	tsr_tensor_halfC FilteredBoxMax;
 	MinMax3x3(FilteredInput, /* out */ FilteredBoxMin, /* out */ FilteredBoxMax);
+
+#if CONFIG_THIN_GEOMETRY_DETECTION  
+	if (bEnableHistoryControl)
+	{
+		tsr_tensor_halfC LerpClampedHistoryInput = LerpClamp3x3(FilteredHistory, FilteredInput, HistoryRelaxationFactor);
+		tsr_tensor_halfC FilteredHistoryBoxMin = Min3x3(LerpClampedHistoryInput);
+		tsr_tensor_halfC FilteredHistoryBoxMax = Max3x3(LerpClampedHistoryInput);
+		tsr_tensor_halfC ScalingFactor = tsr_tensor_halfC::Vectorize((HistoryRelaxationFactor));
+		FilteredBoxMin = select(ScalingFactor > tsr_tensor_halfC::Const(tsr_half(1.0f / 127.0f)), FilteredHistoryBoxMin, FilteredBoxMin);
+		FilteredBoxMax = select(ScalingFactor > tsr_tensor_halfC::Const(tsr_half(1.0f / 127.0f)), FilteredHistoryBoxMax, FilteredBoxMax);
+	}
+#endif
+
 	FilteredBoxMin = FilteredBoxMin - ClampError;
 	FilteredBoxMax = FilteredBoxMax + ClampError;
 		
diff --git a/Engine/Shaders/Private/TemporalSuperResolution/TSRVisualize.usf b/Engine/Shaders/Private/TemporalSuperResolution/TSRVisualize.usf
--- a/Engine/Shaders/Private/TemporalSuperResolution/TSRVisualize.usf
+++ b/Engine/Shaders/Private/TemporalSuperResolution/TSRVisualize.usf
@@ -203,13 +203,17 @@ tsr_half3 DisplayHistorySampleCount(FContext Ctx)
 
 tsr_half3 DisplayThinGeometry(FContext Ctx)
 {
-	tsr_half ThinGeometryType = saturate(tsr_half(ThinGeometryTexture[Ctx.InputPixelPos] * rcp(255.0f)));
-	ThinGeometryType = (View.FrameNumber % 240) > 120 ? ThinGeometryType : tsr_half(0);
-	tsr_half3 Color = kRed;
-	Color = lerp(kWhite,    Color, saturate((ThinGeometryType - 0.0) * 2.0));
+	tsr_ushort EncodedThinGeometryFlag = ThinGeometryTexture[Ctx.InputPixelPos];
+	bool bClusterHole = bool(EncodedThinGeometryFlag & 0x1u);
+	tsr_half HistoryRelaxation = saturate(tsr_half(EncodedThinGeometryFlag >> 0x1u)*rcp(tsr_half(127.0)));
+	bool bUseV2Rectifier = HistoryRelaxation > tsr_half(1.0 / 127.0);
+	tsr_half3 Color = select(bClusterHole, kWhite, kRed);
+	Color = lerp(kWhite,    Color, saturate((HistoryRelaxation - 0.0) * 2.0));
+	Color = lerp(Color,    kGreen, saturate(tsr_half(select(bClusterHole && bUseV2Rectifier, 0.5f, 0.0f))));
+	Color = lerp(Color,   kYellow, saturate(tsr_half(select(bClusterHole && !bUseV2Rectifier, 0.5f, 0.0f))));
 	Color = select(bCameraCut == 0, Color, kPink);
 
-	return SampleTSROutputColor(Ctx) * Color * tsr_half(View.PreExposure);   
+	return SampleTSROutputColor(Ctx) * Color;   
 }
 
 tsr_half3 DisplayHistoryRejection(FContext Ctx)
diff --git a/Engine/Shaders/Public/LaneVectorization.ush b/Engine/Shaders/Public/LaneVectorization.ush
--- a/Engine/Shaders/Public/LaneVectorization.ush
+++ b/Engine/Shaders/Public/LaneVectorization.ush
@@ -48,6 +48,29 @@ bool4 not(bool4 x)
 	return !x;
 }
 
+uint bit_not(uint x)
+{
+	return ~x;
+}
+
+uint2 bit_not(uint2 x)
+{
+	return ~x;
+}
+
+#if PLATFORM_SUPPORTS_REAL_TYPES
+
+uint16_t bit_not(uint16_t x)
+{
+	return ~x;
+}
+
+uint16_t2 bit_not(uint16_t2 x)
+{
+	return ~x;
+}
+#endif
+
 bool2 v_pack_b32_b16(bool a, bool b)
 {
 	return bool2(a, b);
@@ -1436,7 +1459,7 @@ struct TLaneVector2D
 	CALL_SITE_DEBUGLOC
 	template<typename SourceScalarType>
 	static TLaneVector2D<ScalarType, VectorSize, SimdSizeX, SimdSizeY> CastFrom(
-		TLaneVector2D<SourceScalarType, 1, SimdSizeX, SimdSizeY> A)
+		TLaneVector2D<SourceScalarType, VectorSize, SimdSizeX, SimdSizeY> A)
 	{
 		TLaneVector2D<ScalarType, VectorSize, SimdSizeX, SimdSizeY> R;
 		if(SimdSize == 2)
@@ -1808,6 +1831,61 @@ struct TLaneVector2D
 		return R;
 	}
 
+	CALL_SITE_DEBUGLOC
+	TLaneVector2D<bool, VectorSize, SimdSizeX, SimdSizeY> operator != (TLaneVector2D<ScalarType, VectorSize, SimdSizeX, SimdSizeY> B)
+	{
+		TLaneVector2D<bool, VectorSize, SimdSizeX, SimdSizeY> R;
+
+		if(RegisterRowCount == 2)
+		{
+			R.Registers.SetRegisterRow(0, Registers.GetRegisterRow(0) != B.Registers.GetRegisterRow(0));
+			R.Registers.SetRegisterRow(1, Registers.GetRegisterRow(1) != B.Registers.GetRegisterRow(1));
+		}
+		else if(RegisterRowCount == 4)
+		{
+			R.Registers.SetRegisterRow(0, Registers.GetRegisterRow(0) != B.Registers.GetRegisterRow(0));
+			R.Registers.SetRegisterRow(1, Registers.GetRegisterRow(1) != B.Registers.GetRegisterRow(1));
+			R.Registers.SetRegisterRow(2, Registers.GetRegisterRow(2) != B.Registers.GetRegisterRow(2));
+			R.Registers.SetRegisterRow(3, Registers.GetRegisterRow(3) != B.Registers.GetRegisterRow(3));
+		}
+		else
+		{
+			UNROLL
+				for (uint RegisterRowIndex = 0; RegisterRowIndex < RegisterRowCount; RegisterRowIndex++)
+				{
+					R.Registers.SetRegisterRow(RegisterRowIndex, Registers.GetRegisterRow(RegisterRowIndex) != B.Registers.GetRegisterRow(RegisterRowIndex));
+				}
+		}
+		return R;
+	}
+
+	CALL_SITE_DEBUGLOC
+		TLaneVector2D<bool, VectorSize, SimdSizeX, SimdSizeY> operator == (TLaneVector2D<ScalarType, VectorSize, SimdSizeX, SimdSizeY> B)
+	{
+		TLaneVector2D<bool, VectorSize, SimdSizeX, SimdSizeY> R;
+
+		if(RegisterRowCount == 2)
+		{
+			R.Registers.SetRegisterRow(0, Registers.GetRegisterRow(0) == B.Registers.GetRegisterRow(0));
+			R.Registers.SetRegisterRow(1, Registers.GetRegisterRow(1) == B.Registers.GetRegisterRow(1));
+		}
+		else if(RegisterRowCount == 4)
+		{
+			R.Registers.SetRegisterRow(0, Registers.GetRegisterRow(0) == B.Registers.GetRegisterRow(0));
+			R.Registers.SetRegisterRow(1, Registers.GetRegisterRow(1) == B.Registers.GetRegisterRow(1));
+			R.Registers.SetRegisterRow(2, Registers.GetRegisterRow(2) == B.Registers.GetRegisterRow(2));
+			R.Registers.SetRegisterRow(3, Registers.GetRegisterRow(3) == B.Registers.GetRegisterRow(3));
+		}
+		else
+		{
+			UNROLL
+				for (uint RegisterRowIndex = 0; RegisterRowIndex < RegisterRowCount; RegisterRowIndex++)
+				{
+					R.Registers.SetRegisterRow(RegisterRowIndex, Registers.GetRegisterRow(RegisterRowIndex) == B.Registers.GetRegisterRow(RegisterRowIndex));
+				}
+		}
+		return R;
+	}
 
 }; // TLaneVector2D
 
@@ -1945,11 +2023,11 @@ void Deconcatenate(
 
 
 // ------------- 1 parameter
-
-#define TVECTOR_FUNCTION_1PARAMS(ReturnScalarType, FunctionName) \
+	
+#define TVECTOR_FUNCTION_1PARAMS_ALIASE_FUNCTION_NAME(ReturnScalarType, FunctionName,AliasFunctionName) \
 	CALL_SITE_DEBUGLOC \
 	template<typename ScalarType, uint VectorSize, uint SimdSizeX, uint SimdSizeY> \
-	TLaneVector2D<ReturnScalarType, VectorSize, SimdSizeX, SimdSizeY> FunctionName( \
+	TLaneVector2D<ReturnScalarType, VectorSize, SimdSizeX, SimdSizeY> AliasFunctionName( \
 		TLaneVector2D<ScalarType, VectorSize, SimdSizeX, SimdSizeY> A) \
 	{ \
 		TLaneVector2D<ReturnScalarType, VectorSize, SimdSizeX, SimdSizeY> R; \
@@ -1973,7 +2051,9 @@ void Deconcatenate(
 		} \
 		return R; \
 	} \
-	
+
+#define TVECTOR_FUNCTION_1PARAMS(ReturnScalarType, FunctionName) TVECTOR_FUNCTION_1PARAMS_ALIASE_FUNCTION_NAME(ReturnScalarType, FunctionName, FunctionName)
+
 TVECTOR_FUNCTION_1PARAMS(ScalarType, log);
 TVECTOR_FUNCTION_1PARAMS(ScalarType, log2);
 TVECTOR_FUNCTION_1PARAMS(ScalarType, sqrt);
@@ -1987,11 +2067,13 @@ TVECTOR_FUNCTION_1PARAMS(ScalarType, ceil);
 TVECTOR_FUNCTION_1PARAMS(ScalarType, round);
 TVECTOR_FUNCTION_1PARAMS(ScalarType, fast_sign);
 TVECTOR_FUNCTION_1PARAMS(bool, not);
-TVECTOR_FUNCTION_1PARAMS(uint, asuint);
+TVECTOR_FUNCTION_1PARAMS(uint, bit_not);
 #if PLATFORM_SUPPORTS_REAL_TYPES
 	TVECTOR_FUNCTION_1PARAMS(uint16_t, asuint16);
 #endif
-
+//Work around: cannot use keyword asuint and asfloat on struct.
+TVECTOR_FUNCTION_1PARAMS_ALIASE_FUNCTION_NAME(uint, asuint, azuint);
+TVECTOR_FUNCTION_1PARAMS_ALIASE_FUNCTION_NAME(float,asfloat, azfloat);
 
 // ------------- 2 parameters
 
@@ -2112,6 +2194,56 @@ TVECTOR_FUNCTION_3PARAMS(max3, ScalarType);
 
 // ------------- AnyElement & AllElement
 
+CALL_SITE_DEBUGLOC
+template<uint VectorSize, uint SimdSizeX, uint SimdSizeY>
+TLaneVector2D<bool, 1, SimdSizeX, SimdSizeY> AnyElement(TLaneVector2D<bool, VectorSize, SimdSizeX, SimdSizeY> A)
+{
+	TLaneVector2D<bool, 1, SimdSizeX, SimdSizeY> R;
+	R.SetElement(0, any(A.GetElement(0)));
+
+	if (SimdSizeX * SimdSizeY == 2)
+	{
+		R.SetElement(1, any(A.GetElement(1)));
+	}	
+	else if (SimdSizeX * SimdSizeY == 4)
+	{
+		R.SetElement(1, any(A.GetElement(1)));
+		R.SetElement(2, any(A.GetElement(2)));
+		R.SetElement(3, any(A.GetElement(3)));
+	}	
+	else
+	{
+		UNROLL for (uint SimdIndex = 1; SimdIndex < SimdSizeX * SimdSizeY; SimdIndex++) { R.SetElement(SimdIndex, any(A.GetElement(SimdIndex))); }
+	}
+
+	return R;
+}
+
+CALL_SITE_DEBUGLOC
+template<uint VectorSize, uint SimdSizeX, uint SimdSizeY>
+TLaneVector2D<bool, 1, SimdSizeX, SimdSizeY> AllElement(TLaneVector2D<bool, VectorSize, SimdSizeX, SimdSizeY> A)
+{
+	TLaneVector2D<bool, 1, SimdSizeX, SimdSizeY> R;
+	R.SetElement(0, all(A.GetElement(0)));
+
+	if (SimdSizeX * SimdSizeY == 2)
+	{
+		R.SetElement(1, all(A.GetElement(1)));
+	}	
+	else if (SimdSizeX * SimdSizeY == 4)
+	{
+		R.SetElement(1, all(A.GetElement(1)));
+		R.SetElement(2, all(A.GetElement(2)));
+		R.SetElement(3, all(A.GetElement(3)));
+	}	
+	else
+	{
+		UNROLL for (uint SimdIndex = 1; SimdIndex < SimdSizeX * SimdSizeY; SimdIndex++) { R.SetElement(SimdIndex, all(A.GetElement(SimdIndex))); }
+	}
+
+	return R;
+}
+
 CALL_SITE_DEBUGLOC
 template<uint VectorSize, uint SimdSizeX, uint SimdSizeY>
 bool AnyComponent(TLaneVector2D<bool, VectorSize, SimdSizeX, SimdSizeY> A)
diff --git a/Engine/Source/Runtime/Renderer/Private/DeferredShadingRenderer.cpp b/Engine/Source/Runtime/Renderer/Private/DeferredShadingRenderer.cpp
--- a/Engine/Source/Runtime/Renderer/Private/DeferredShadingRenderer.cpp
+++ b/Engine/Source/Runtime/Renderer/Private/DeferredShadingRenderer.cpp
@@ -3049,18 +3049,18 @@ void FDeferredShadingSceneRenderer::Render(FRDGBuilder& GraphBuilder, const FSce
 		TArray<FScreenPassTexture, TInlineAllocator<4>> TSRFlickeringInputTextures;
 		if (!bHasRayTracedOverlay)
 		{
-			// Extract TSR's moire heuristic luminance before rendering translucency into the scene color.
+			// Extract TSR's moire heuristic luminance, thin geometry coverage before rendering translucency into the scene color.
 			for (int32 ViewIndex = 0; ViewIndex < Views.Num(); ++ViewIndex)
 			{
 				FViewInfo& View = Views[ViewIndex];
-				if (NeedTSRMoireLuma(View))
+				if (NeedTSRAntiFlickeringPass(View))
 				{
 					if (TSRFlickeringInputTextures.Num() == 0)
 					{
 						TSRFlickeringInputTextures.SetNum(Views.Num());
 					}
 
-					TSRFlickeringInputTextures[ViewIndex] = AddTSRMeasureFlickeringLuma(GraphBuilder, View.ShaderMap, FScreenPassTexture(SceneTextures.Color.Target, View.ViewRect));
+					TSRFlickeringInputTextures[ViewIndex] = AddTSRMainAntiFlickeringPass(GraphBuilder, View.ShaderMap, FScreenPassTexture(SceneTextures.Color.Target, View.ViewRect), SceneTextures);
 				}
 			}
 		}
diff --git a/Engine/Source/Runtime/Renderer/Private/PostProcess/TemporalAA.h b/Engine/Source/Runtime/Renderer/Private/PostProcess/TemporalAA.h
--- a/Engine/Source/Runtime/Renderer/Private/PostProcess/TemporalAA.h
+++ b/Engine/Source/Runtime/Renderer/Private/PostProcess/TemporalAA.h
@@ -176,15 +176,14 @@ bool IsTSRLensDistortionSupported(EShaderPlatform ShaderPlatform);
 /** Returns whether TSR lens distortion is enabled (for runtime toggle). */
 bool IsTSRLensDistortionEnabled(EShaderPlatform ShaderPlatform);
 
-/** Returns whether a given view need to measure luminance of the scene color for moire anti-flickering. */
-bool NeedTSRMoireLuma(const FViewInfo& View);
+/** Returns whether a given view need to measure the scene for moire anti-flickering. */
+bool NeedTSRAntiFlickeringPass(const FViewInfo& View);
 
 /** Returns whether TSR internal visualization is enabled on the view. */
 bool IsVisualizeTSREnabled(const FViewInfo& View);
 
-/** Measure luminance of the scene color for anti-flickering. */
-FScreenPassTexture AddTSRMeasureFlickeringLuma(FRDGBuilder& GraphBuilder, FGlobalShaderMap* ShaderMap, FScreenPassTexture SceneColor);
-
+/** Measure different metrics of the scene for anti-flickering */
+FScreenPassTexture AddTSRMainAntiFlickeringPass(FRDGBuilder& GraphBuilder, FGlobalShaderMap* ShaderMap, FScreenPassTexture SceneColor, const FSceneTextures& SceneTextures);
 
 EMainTAAPassConfig GetMainTAAPassConfig(const FViewInfo& View);
 
diff --git a/Engine/Source/Runtime/Renderer/Private/PostProcess/TemporalSuperResolution.cpp b/Engine/Source/Runtime/Renderer/Private/PostProcess/TemporalSuperResolution.cpp
--- a/Engine/Source/Runtime/Renderer/Private/PostProcess/TemporalSuperResolution.cpp
+++ b/Engine/Source/Runtime/Renderer/Private/PostProcess/TemporalSuperResolution.cpp
@@ -306,9 +306,16 @@
 
 TAutoConsoleVariable<int32> CVarTSRThinGeometryDetection(
 	TEXT("r.TSR.ThinGeometryDetection"), 0,
-	TEXT("Define if we should perform another pass to detect thin geometies (sub-pixel in frame buffer) either due to sampling algorithm before TSR or the geometry being too thin. ")
-	TEXT("When enabled thin geometry pixels will under-reject if current pixel misses.")
-	TEXT("(Default = 1)."),
+	TEXT("Define if we should perform another pass to detect thin geometies (sub-pixel in frame buffer) either due to sampling algorithm before TSR or the geometry being too thin (Default = 0). ")
+	TEXT("When enabled thin geometry pixels will relax history rejection based on the types.\n")
+	TEXT("Edge line: single pixel line over non foliage (Red in r.TSR.Visualize 15)\n")
+	TEXT("Cluster hole region: foliage pixels with partial coverage against background material (Green)\n")
+	TEXT("Other: no history relaxation (Yellow)."),
 	ECVF_RenderThreadSafe);
 
+TAutoConsoleVariable<float> CVarTSRThinGeometryCoverageMaxRelaxationWeight(
+	TEXT("r.TSR.ThinGeometryDetection.Coverage.MaxRelaxationWeight"), 0.037,
+	TEXT("The max history clamping box relaxation weight due to thin geometry detection (0 to 1)"),
+	ECVF_RenderThreadSafe);
+
 TAutoConsoleVariable<float> CVarTSRThinGeometryDetectionErrorMultiplier(
@@ -340,7 +347,7 @@
 	TEXT(" 12: Mask where the pixel's jacobian is null in the reprojection field in orange;\n")
 	TEXT(" 13: Mask where the pixel's jacobian has reached its encoding limit in the reprojection field in red;\n")
 	TEXT(" 14: Mask where the reprojected history is upscaled (in red) or downscaled (in green) by the reprojection field's jacobian (for instance due to getting closer or further away from camera respectively, or an object is getting scaled dynamicaly);\n")
-	TEXT(" 15: Mask where thin geometry is detected by depth and luma difference (in red);\n"),
+	TEXT(" 15: Mask where thin geometry is detected by (edge line:red, potential cluster hole with partial coverage:green, other: yellow);\n"),
 	ECVF_RenderThreadSafe);
 
 #endif
@@ -380,6 +387,7 @@
 	SHADER_PARAMETER_RDG_TEXTURE(Texture2DArray, MetadataArray)
 	SHADER_PARAMETER_RDG_TEXTURE(Texture2DArray, GuideArray)
 	SHADER_PARAMETER_RDG_TEXTURE(Texture2DArray, MoireArray)
+	SHADER_PARAMETER_RDG_TEXTURE(Texture2DArray, CoverageArray)
 END_SHADER_PARAMETER_STRUCT()
 
 BEGIN_SHADER_PARAMETER_STRUCT(FTSRPrevHistoryParameters, )
@@ -647,10 +655,31 @@
 	BEGIN_SHADER_PARAMETER_STRUCT(FParameters, )
 		SHADER_PARAMETER_STRUCT(FScreenPassTextureViewportParameters, InputInfo)
 		SHADER_PARAMETER_RDG_TEXTURE(Texture2D, SceneColorTexture)
-		SHADER_PARAMETER_RDG_TEXTURE_UAV(RWTexture2D, FlickeringLumaOutput)
+		SHADER_PARAMETER_RDG_TEXTURE_UAV(RWTexture2DArray, FlickeringLumaOutput)
 	END_SHADER_PARAMETER_STRUCT()
 }; // class FTSRMeasureFlickeringLumaCS
 
+class FTSRMeasureThinGeometryCoverageCS : public FTSRShader
+{
+	DECLARE_GLOBAL_SHADER(FTSRMeasureThinGeometryCoverageCS);
+	SHADER_USE_PARAMETER_STRUCT(FTSRMeasureThinGeometryCoverageCS, FTSRShader);
+
+	BEGIN_SHADER_PARAMETER_STRUCT(FParameters, )
+		SHADER_PARAMETER_STRUCT(FScreenPassTextureViewportParameters, InputInfo)
+		SHADER_PARAMETER_RDG_UNIFORM_BUFFER(FSceneTextureUniformParameters, SceneTextures)
+
+		SHADER_PARAMETER(int32, OutputArrayIndex)
+		SHADER_PARAMETER_RDG_TEXTURE_UAV(RWTexture2DArray, ThinGeometryCoverageOutput)
+	END_SHADER_PARAMETER_STRUCT()
+
+	static void ModifyCompilationEnvironment(const FGlobalShaderPermutationParameters& Parameters, FShaderCompilerEnvironment& OutEnvironment)
+	{
+		FTSRShader::ModifyCompilationEnvironment(Parameters, OutEnvironment);
+		//TODO: Fix half type cast warning in GBuffer function before removing.
+		OutEnvironment.CompilerFlags.Remove(CFLAG_AllowRealTypes);
+	}
+}; // class FTSRMeasureThinGeometryCoverageCS
+
 class FTSRClearPrevTexturesCS : public FTSRShader
 {
 	DECLARE_GLOBAL_SHADER(FTSRClearPrevTexturesCS);
@@ -726,7 +755,8 @@
 
 	class FMoireReprojectionDim : SHADER_PERMUTATION_BOOL("DIM_MOIRE_REPROJECTION");
 	class FResurrectionReprojectionDim : SHADER_PERMUTATION_BOOL("DIM_RESURRECTION_REPROJECTION");
-	using FPermutationDomain = TShaderPermutationDomain<FMoireReprojectionDim, FResurrectionReprojectionDim, FTSRShader::F16BitVALUDim, FTSRShader::FAlphaChannelDim>;
+	class FThinGeometryCoverageDim: SHADER_PERMUTATION_BOOL("DIM_THIN_GEOMETRY_COVERAGE_REPROJECTION");
+	using FPermutationDomain = TShaderPermutationDomain<FMoireReprojectionDim, FResurrectionReprojectionDim, FThinGeometryCoverageDim, FTSRShader::F16BitVALUDim, FTSRShader::FAlphaChannelDim>;
 
 	BEGIN_SHADER_PARAMETER_STRUCT(FParameters, )
 		SHADER_PARAMETER_STRUCT_INCLUDE(FTSRCommonParameters, CommonParameters)
@@ -741,6 +771,7 @@
 		SHADER_PARAMETER_STRUCT_INCLUDE(FTSRPrevHistoryParameters, PrevHistoryParameters)
 		SHADER_PARAMETER_RDG_TEXTURE(Texture2D, PrevHistoryGuide)
 		SHADER_PARAMETER_RDG_TEXTURE(Texture2D, PrevHistoryMoire)
+		SHADER_PARAMETER_RDG_TEXTURE(Texture2D, PrevHistoryCoverage)
 		SHADER_PARAMETER_STRUCT(FScreenPassTextureViewportParameters, PrevGuideInfo)
 		SHADER_PARAMETER(FScreenTransform, InputPixelPosToReprojectScreenPos)
 		SHADER_PARAMETER(FScreenTransform, ScreenPosToPrevHistoryGuideBufferUV)
@@ -754,6 +785,7 @@
 
 		SHADER_PARAMETER_RDG_TEXTURE_UAV(RWTexture2DArray, ReprojectedHistoryGuideOutput)
 		SHADER_PARAMETER_RDG_TEXTURE_UAV(RWTexture2DArray, ReprojectedHistoryMoireOutput)
+		SHADER_PARAMETER_RDG_TEXTURE_UAV(RWTexture2DArray, ReprojectedHistoryCoverageOutput)
 		SHADER_PARAMETER_RDG_TEXTURE_UAV(RWTexture2DArray, ReprojectionFieldOutput)
 		SHADER_PARAMETER_RDG_TEXTURE_UAV(RWTexture2D, DecimateMaskOutput)
 
@@ -823,7 +855,7 @@
 		SHADER_PARAMETER(int32, bPassthroughAlpha)
 
 		SHADER_PARAMETER_RDG_TEXTURE(Texture2D, InputTexture)
-		SHADER_PARAMETER_RDG_TEXTURE(Texture2D, InputMoireLumaTexture)
+		SHADER_PARAMETER_RDG_TEXTURE_SRV(Texture2D, InputMoireLumaTexture)
 		SHADER_PARAMETER_RDG_TEXTURE(Texture2D, InputSceneTranslucencyTexture)
 		SHADER_PARAMETER_RDG_TEXTURE_SRV(Texture2D, ReprojectedHistoryGuideTexture)
 		SHADER_PARAMETER_RDG_TEXTURE_SRV(Texture2D, ReprojectedHistoryGuideMetadataTexture)
@@ -911,8 +943,7 @@
 	DECLARE_GLOBAL_SHADER(FTSRDetectThinGeometryCS);
 	SHADER_USE_PARAMETER_STRUCT(FTSRDetectThinGeometryCS, FTSRConvolutionNetworkShader);
 
-	using FPermutationDomain = TShaderPermutationDomain<
-		FTSRConvolutionNetworkShader::FPermutationDomain>;
+	using FPermutationDomain = TShaderPermutationDomain<FTSRConvolutionNetworkShader::FPermutationDomain>;
 
 	BEGIN_SHADER_PARAMETER_STRUCT(FParameters, )
 		SHADER_PARAMETER_STRUCT_INCLUDE(FTSRCommonParameters, CommonParameters)
@@ -920,9 +951,14 @@
 		SHADER_PARAMETER(int32, TileOverscan)
 		SHADER_PARAMETER(int32, ThinGeometryTextureIndex)
 		SHADER_PARAMETER(float, ErrorMultiplier)
+		SHADER_PARAMETER(float, MaxRelaxationWeight)
 
 		SHADER_PARAMETER_RDG_TEXTURE(Texture2D, SceneDepthTexture)
+		SHADER_PARAMETER_RDG_TEXTURE_SRV(Texture2D, CurrentCoverageTexture)
+		SHADER_PARAMETER_RDG_TEXTURE_SRV(Texture2D, ReprojectedHistoryCoverageTexture)
+
 		SHADER_PARAMETER_RDG_TEXTURE_UAV(RWTexture2DArray, R8Output)
+		SHADER_PARAMETER_RDG_TEXTURE_UAV(RWTexture2DArray, HistoryCoverageOutput)
 		SHADER_PARAMETER_RDG_TEXTURE_UAV(RWTexture2D, DebugOutput)
 	END_SHADER_PARAMETER_STRUCT()
 
@@ -1224,7 +1260,7 @@
 		SHADER_PARAMETER_RDG_TEXTURE(Texture2D, ResurrectedDistortingDisplacementTexture)
 		SHADER_PARAMETER_RDG_TEXTURE(Texture2D, UndistortingDisplacementTexture)
 		SHADER_PARAMETER_RDG_TEXTURE(Texture2D, InputTexture)
-		SHADER_PARAMETER_RDG_TEXTURE(Texture2D, InputMoireLumaTexture)
+		SHADER_PARAMETER_RDG_TEXTURE_SRV(Texture2D, InputMoireLumaTexture)
 		SHADER_PARAMETER_RDG_TEXTURE(Texture2D, InputSceneTranslucencyTexture)
 		SHADER_PARAMETER_RDG_TEXTURE_SRV(Texture2D, SceneColorTexture)
 		SHADER_PARAMETER_RDG_TEXTURE(Texture2D, ClosestDepthTexture)
@@ -1246,6 +1282,7 @@
 }; // class FTSRVisualizeCS
 
 IMPLEMENT_GLOBAL_SHADER(FTSRMeasureFlickeringLumaCS, "/Engine/Private/TemporalSuperResolution/TSRMeasureFlickeringLuma.usf", "MainCS", SF_Compute);
+IMPLEMENT_GLOBAL_SHADER(FTSRMeasureThinGeometryCoverageCS, "/Engine/Private/TemporalSuperResolution/TSRMeasureCoverage.usf", "MainCS", SF_Compute);
 IMPLEMENT_GLOBAL_SHADER(FTSRClearPrevTexturesCS,     "/Engine/Private/TemporalSuperResolution/TSRClearPrevTextures.usf",     "MainCS", SF_Compute);
 IMPLEMENT_GLOBAL_SHADER(FTSRDilateVelocityCS,        "/Engine/Private/TemporalSuperResolution/TSRDilateVelocity.usf",        "MainCS", SF_Compute);
 IMPLEMENT_GLOBAL_SHADER(FTSRDetectThinGeometryCS,	 "/Engine/Private/TemporalSuperResolution/TSRDetectThinGeometry.usf",	 "MainCS", SF_Compute);
@@ -1430,12 +1467,27 @@
 	return CVarTSRLensDistortion.GetValueOnRenderThread() != 0;
 }
 
-bool NeedTSRMoireLuma(const FViewInfo& View)
+bool NeedTSRAntiFlickeringPass(const FViewInfo& View)
 {
 	// Need to also check PostProcessing flag, as scene captures may run with temporal AA jitter matching the main view, but post processing disabled.
 	return GetMainTAAPassConfig(View) == EMainTAAPassConfig::TSR && View.Family->EngineShowFlags.PostProcessing;
 }
 
+bool ShouldAddTSRMainFlickeringLumaPass()
+{
+	return  CVarTSRFlickeringEnable.GetValueOnRenderThread()!=0 || CVarTSRFlickeringPeriod.GetValueOnRenderThread() == 0.0f;
+}
+
+bool ShouldAddTSRMainThinGeometryCoveragePass()
+{
+	return CVarTSRThinGeometryDetection.GetValueOnRenderThread() != 0;
+}
+
+int32 GetTSRMainFlickeringLumaTextureArraySize()
+{
+	return FMath::Max(1, static_cast<int32>(ShouldAddTSRMainFlickeringLumaPass()) + static_cast<int32>(ShouldAddTSRMainThinGeometryCoveragePass()));
+};
+
 bool IsVisualizeTSREnabled(const FViewInfo& View)
 #if UE_BUILD_OPTIMIZED_SHOWFLAGS
 {
@@ -1456,31 +1508,75 @@
 
 	FScreenPassTexture FlickeringLuma;
 	{
-		FRDGTextureDesc Desc = FRDGTextureDesc::Create2D(
+		FRDGTextureDesc Desc = FRDGTextureDesc::Create2DArray(
 			SceneColor.Texture->Desc.Extent,
 			PF_R8,
 			FClearValueBinding::None,
-			/* InFlags = */ TexCreate_ShaderResource | TexCreate_UAV);
+			/* InFlags = */ TexCreate_ShaderResource | TexCreate_UAV,
+			GetTSRMainFlickeringLumaTextureArraySize());
 
 		FlickeringLuma.Texture = GraphBuilder.CreateTexture(Desc, TEXT("TSR.Flickering.Luminance"));
 		FlickeringLuma.ViewRect = SceneColor.ViewRect;
 	}
 
-	FTSRMeasureFlickeringLumaCS::FParameters* PassParameters = GraphBuilder.AllocParameters<FTSRMeasureFlickeringLumaCS::FParameters>();
-	PassParameters->InputInfo = GetScreenPassTextureViewportParameters(FScreenPassTextureViewport(
-		SceneColor.Texture->Desc.Extent, SceneColor.ViewRect));
-	PassParameters->SceneColorTexture = SceneColor.Texture;
-	PassParameters->FlickeringLumaOutput = GraphBuilder.CreateUAV(FlickeringLuma.Texture);
+	if (ShouldAddTSRMainFlickeringLumaPass() || !ShouldAddTSRMainThinGeometryCoveragePass())
+	{
+		FTSRMeasureFlickeringLumaCS::FParameters* PassParameters = GraphBuilder.AllocParameters<FTSRMeasureFlickeringLumaCS::FParameters>();
+		PassParameters->InputInfo = GetScreenPassTextureViewportParameters(FScreenPassTextureViewport(
+			SceneColor.Texture->Desc.Extent, SceneColor.ViewRect));
+		PassParameters->SceneColorTexture = SceneColor.Texture;
+		PassParameters->FlickeringLumaOutput = GraphBuilder.CreateUAV(FlickeringLuma.Texture);
 
-	TShaderMapRef<FTSRMeasureFlickeringLumaCS> ComputeShader(ShaderMap);
+		TShaderMapRef<FTSRMeasureFlickeringLumaCS> ComputeShader(ShaderMap);
+		FComputeShaderUtils::AddPass(
+			GraphBuilder,
+			RDG_EVENT_NAME("TSR MeasureFlickeringLuma %dx%d", SceneColor.ViewRect.Width(), SceneColor.ViewRect.Height()),
+			ComputeShader,
+			PassParameters,
+			FComputeShaderUtils::GetGroupCount(FlickeringLuma.ViewRect.Size(), 8 * 2));
+	}
+
+	return FlickeringLuma;
+}
+
+void AddTSRMeasureThinGeometryCoverage(FRDGBuilder& GraphBuilder, FGlobalShaderMap* ShaderMap, const FSceneTextures& SceneTextures, const FScreenPassTexture& ThinGeometryCoverage)
+{
+	const bool bRecordThinGeometryCoverage = CVarTSRThinGeometryDetection.GetValueOnRenderThread() != 0; 
+
+	if (!bRecordThinGeometryCoverage)
+	{
+		return;
+	}
+
+	RDG_EVENT_SCOPE_STAT(GraphBuilder, TemporalSuperResolution, "TemporalSuperResolution");
+	RDG_GPU_STAT_SCOPE(GraphBuilder, TemporalSuperResolution);
+
+	FIntRect ViewRect = ThinGeometryCoverage.ViewRect;
+	FIntPoint ScreenpassExtent = ThinGeometryCoverage.Texture->Desc.Extent;
+
+	FTSRMeasureThinGeometryCoverageCS::FParameters* PassParameters = GraphBuilder.AllocParameters<FTSRMeasureThinGeometryCoverageCS::FParameters>();
+	PassParameters->InputInfo = GetScreenPassTextureViewportParameters(FScreenPassTextureViewport(ScreenpassExtent, ViewRect));
+	PassParameters->SceneTextures = SceneTextures.UniformBuffer;
+	PassParameters->OutputArrayIndex = GetTSRMainFlickeringLumaTextureArraySize() - 1;
+	
+	PassParameters->ThinGeometryCoverageOutput = GraphBuilder.CreateUAV(ThinGeometryCoverage.Texture);
+
+	TShaderMapRef<FTSRMeasureThinGeometryCoverageCS> ComputeShader(ShaderMap);
 	FComputeShaderUtils::AddPass(
 		GraphBuilder,
-		RDG_EVENT_NAME("TSR MeasureFlickeringLuma %dx%d", SceneColor.ViewRect.Width(), SceneColor.ViewRect.Height()),
+		RDG_EVENT_NAME("TSR MeasureThinGeometryCoverage %dx%d", ViewRect.Width(), ViewRect.Height()),
 		ComputeShader,
 		PassParameters,
-		FComputeShaderUtils::GetGroupCount(FlickeringLuma.ViewRect.Size(), 8 * 2));
-
-	return FlickeringLuma;
+		FComputeShaderUtils::GetGroupCount(ThinGeometryCoverage.ViewRect.Size(), 8 * 2));
 }
 
+FScreenPassTexture AddTSRMainAntiFlickeringPass(FRDGBuilder& GraphBuilder, FGlobalShaderMap* ShaderMap, FScreenPassTexture SceneColor, const FSceneTextures& SceneTextures)
+{
+	FScreenPassTexture AntiFlickeringTexture = AddTSRMeasureFlickeringLuma(GraphBuilder, ShaderMap, SceneColor);
+
+	AddTSRMeasureThinGeometryCoverage(GraphBuilder, ShaderMap, SceneTextures, AntiFlickeringTexture);
+
+	return AntiFlickeringTexture;
+}
+
 FDefaultTemporalUpscaler::FOutputs AddMainTemporalSuperResolutionPasses(
@@ -1869,6 +1965,16 @@
 				/* ArraySize = */ 1);
 			History.MoireArray = GraphBuilder.CreateTexture(Desc, TEXT("TSR.History.Moire"));
 		}
+
+		{
+			FRDGTextureDesc Desc = FRDGTextureDesc::Create2DArray(
+				HistoryGuideExtent,
+				PF_R8,
+				FClearValueBinding::None,
+				TexCreate_ShaderResource | TexCreate_UAV,
+				/* ArraySize = */ 1);
+			History.CoverageArray = GraphBuilder.CreateTexture(Desc, TEXT("TSR.History.Coverage"));
+		}
 	}
 
 	// Whether to camera cut the history Resurrection
@@ -1884,6 +1990,8 @@
 		History.ColorArray->Desc.ArraySize != InputHistory.ColorArray->GetDesc().ArraySize ||
 		History.GuideArray->Desc.Extent != InputHistory.GuideArray->GetDesc().Extent ||
 		History.GuideArray->Desc.Format != InputHistory.GuideArray->GetDesc().Format ||
+		(History.CoverageArray && InputHistory.CoverageArray && History.CoverageArray->Desc.Extent != InputHistory.CoverageArray->GetDesc().Extent) ||
+		(History.CoverageArray && InputHistory.CoverageArray && History.CoverageArray->Desc.Format != InputHistory.CoverageArray->GetDesc().Format) ||
 		false;
 
 	// Current and previous frame histories
@@ -1901,6 +2009,7 @@
 		PrevHistory.MetadataArray = BlackArrayDummy;
 		PrevHistory.GuideArray = BlackArrayDummy;
 		PrevHistory.MoireArray = BlackArrayDummy;
+		PrevHistory.CoverageArray = BlackArrayDummy;
 
 		if (HistorySliceSequence.GetRollingIndexCount() > 1)
 		{
@@ -1926,6 +2035,9 @@
 		PrevHistory.MoireArray = InputHistory.MoireArray.IsValid()
 			? GraphBuilder.RegisterExternalTexture(InputHistory.MoireArray)
 			: BlackArrayDummy;
+		PrevHistory.CoverageArray = InputHistory.CoverageArray.IsValid()
+			? GraphBuilder.RegisterExternalTexture(InputHistory.CoverageArray)
+			: BlackArrayDummy;
 
 		int32 ResurrectionFrameRollingIndex = 0;
 		int32 PrevFrameRollingIndex = 0;
@@ -1948,6 +2060,12 @@
 				History.MetadataArray = PrevHistory.MetadataArray;
 				History.GuideArray = PrevHistory.GuideArray;
 				History.MoireArray = PrevHistory.MoireArray;
+				
+				// Reuse same history for coverage when it has already been allocated or we do not need coverage pass
+				if (PrevHistory.CoverageArray->Desc.Extent.Size() != 1 ||  (!ShouldAddTSRMainThinGeometryCoveragePass()))
+				{
+					History.CoverageArray = PrevHistory.CoverageArray;
+				}
 			}
 
 			ResurrectionFrameRollingIndex = PrevHistorySliceSequence.GetResurrectionFrameRollingIndex(InputHistory.AccumulatedFrameCount, InputHistory.LastFrameRollingIndex);
@@ -2203,52 +2321,11 @@
 			FComputeShaderUtils::GetGroupCount(InputRect.Size(), TileSize));
 	}
 
-	//Create the thin geometry mask to avoid over rejection
-	if(PassConfig.ThinGeometryDetectionEnable)
-	{
-		FTSRConvolutionNetworkShader::FPermutationDomain ConvolutionNetworkPermutationVector;
-		ConvolutionNetworkPermutationVector.Set<FTSRConvolutionNetworkShader::FWaveSizeOps>(SelectWaveSize(View.GetShaderPlatform(), { 16, 32, 64 }));
-		ConvolutionNetworkPermutationVector.Set<FTSRShader::F16BitVALUDim>(bUse16BitVALU);
-		ConvolutionNetworkPermutationVector.Set<FTSRShader::FAlphaChannelDim>(bSupportsAlpha);
-
-		FTSRDetectThinGeometryCS::FPermutationDomain PermutationVector;
-		PermutationVector.Set<FTSRConvolutionNetworkShader::FPermutationDomain>(ConvolutionNetworkPermutationVector);
-		PermutationVector = FTSRDetectThinGeometryCS::RemapPermutation(PermutationVector);
-
-		const int32 GroupTileSize = 32;
-		const int32 TileOverscan = FMath::Clamp(CVarTSRShadingTileOverscan.GetValueOnRenderThread(), 3, GroupTileSize / 2 - 1);
-		const int32 TileSize = GroupTileSize - 2 * TileOverscan;
-
-		FTSRDetectThinGeometryCS::FParameters* PassParameters = GraphBuilder.AllocParameters<FTSRDetectThinGeometryCS::FParameters>();
-		PassParameters->CommonParameters = CommonParameters;
-		PassParameters->TileOverscan = TileOverscan;
-		PassParameters->ThinGeometryTextureIndex = ThinGeometryTextureIndex;
-		PassParameters->ErrorMultiplier = PassConfig.ThinGeometryErrorMultiplier;
-		PassParameters->SceneDepthTexture = PassInputs.SceneDepth.Texture;
-		PassParameters->R8Output = GraphBuilder.CreateUAV(R8OutputTexture);
-		PassParameters->DebugOutput = CreateDebugUAV(InputExtent, TEXT("Debug.TSR.DetectThinGeometry"));
-
-		TShaderMapRef<FTSRDetectThinGeometryCS> ComputeShader(View.ShaderMap, PermutationVector);
-		FComputeShaderUtils::AddPass(
-			GraphBuilder,
-			RDG_EVENT_NAME("TSR DetectThinGeometry(#%d TileSize=%d PaddingCostMultiplier=%1.1f WaveSize=%d VALU=%s%s) %dx%d",
-				PermutationVector.ToDimensionValueId(),
-				TileSize,
-				FMath::Pow(float(GroupTileSize) / float(TileSize), 2),
-				int32(PermutationVector.Get<FTSRConvolutionNetworkShader::FPermutationDomain>().Get<FTSRDetectThinGeometryCS::FWaveSizeOps>()),
-				PermutationVector.Get<FTSRConvolutionNetworkShader::FPermutationDomain>().Get<FTSRShader::F16BitVALUDim>() ? TEXT("16bit") : TEXT("32bit"),
-				PermutationVector.Get<FTSRConvolutionNetworkShader::FPermutationDomain>().Get<FTSRShader::FAlphaChannelDim>() ? TEXT(" AlphaChannel") : TEXT(""),
-				InputRect.Width(), InputRect.Height()),
-			AsyncComputePasses >= 2 ? ERDGPassFlags::AsyncCompute : ERDGPassFlags::Compute,
-			ComputeShader,
-			PassParameters,
-			FComputeShaderUtils::GetGroupCount(InputRect.Size(), TileSize));
-
-	}
-
 	// Decimate input to flicker at same frequency as input.
 	FRDGTextureRef ReprojectedHistoryGuideTexture = nullptr;
 	FRDGTextureRef ReprojectedHistoryMoireTexture = nullptr;
+	FRDGTextureRef ReprojectedHistoryCoverageTexture = nullptr;
+
 	FRDGTextureRef DecimateMaskTexture = nullptr;
 	{
 		{
@@ -2282,6 +2359,17 @@
 			ReprojectedHistoryMoireTexture = GraphBuilder.CreateTexture(Desc, TEXT("TSR.ReprojectedHistoryMoire"));
 		}
 
+		if (PassConfig.ThinGeometryDetectionEnable)
+		{
+			FRDGTextureDesc Desc = FRDGTextureDesc::Create2DArray(
+				InputExtent,
+				History.CoverageArray->Desc.Format,
+				FClearValueBinding::None,
+				/* InFlags = */ TexCreate_ShaderResource | TexCreate_UAV,
+				/* InArraySize = */ 1);
+			ReprojectedHistoryCoverageTexture = GraphBuilder.CreateTexture(Desc, TEXT("TSR.ReprojectedHistoryCoverage"));
+		}
+
 		FTSRDecimateHistoryCS::FParameters* PassParameters = GraphBuilder.AllocParameters<FTSRDecimateHistoryCS::FParameters>();
 		PassParameters->CommonParameters = CommonParameters;
 		PassParameters->RotationalClipToPrevClip = RotationalClipToPrevClip;
@@ -2298,6 +2386,7 @@
 			FScreenPassTextureViewport PrevHistoryGuideViewport(PrevHistory.GuideArray->Desc.Extent, InputHistory.InputViewportRect - InputHistory.InputViewportRect.Min);
 			PassParameters->PrevHistoryGuide = PrevHistory.GuideArray;
 			PassParameters->PrevHistoryMoire = PrevHistory.MoireArray;
+			PassParameters->PrevHistoryCoverage = PrevHistory.CoverageArray;
 			PassParameters->PrevGuideInfo = GetScreenPassTextureViewportParameters(PrevHistoryGuideViewport);
 			PassParameters->InputPixelPosToReprojectScreenPos = ((FScreenTransform::Identity - InputRect.Min + 0.5f) / InputRect.Size()) * FScreenTransform::ViewportUVToScreenPos;
 			PassParameters->ScreenPosToPrevHistoryGuideBufferUV = FScreenTransform::ChangeTextureBasisFromTo(
@@ -2322,6 +2411,12 @@
 		{
 			PassParameters->ReprojectedHistoryMoireOutput = GraphBuilder.CreateUAV(ReprojectedHistoryMoireTexture);
 		}
+
+		if (ReprojectedHistoryCoverageTexture)
+		{
+			PassParameters->ReprojectedHistoryCoverageOutput = GraphBuilder.CreateUAV(ReprojectedHistoryCoverageTexture);
+		}
+
 		if (bReprojectionField)
 		{
 			FRDGTextureUAVDesc ReprojectionFieldUAVDesc(ReprojectionFieldTexture);
@@ -2348,16 +2443,18 @@
 		FTSRDecimateHistoryCS::FPermutationDomain PermutationVector;
 		PermutationVector.Set<FTSRDecimateHistoryCS::FMoireReprojectionDim>(FlickeringFramePeriod > 0.0f);
 		PermutationVector.Set<FTSRDecimateHistoryCS::FResurrectionReprojectionDim>(bCanResurrectHistory);
+		PermutationVector.Set<FTSRDecimateHistoryCS::FThinGeometryCoverageDim>(PassConfig.ThinGeometryDetectionEnable);
 		PermutationVector.Set<FTSRShader::F16BitVALUDim>(bUse16BitVALU);
 		PermutationVector.Set<FTSRShader::FAlphaChannelDim>(bSupportsAlpha);
 
 		TShaderMapRef<FTSRDecimateHistoryCS> ComputeShader(View.ShaderMap, PermutationVector);
 		FComputeShaderUtils::AddPass(
 			GraphBuilder,
-			RDG_EVENT_NAME("TSR DecimateHistory(#%d%s%s%s%s) %dx%d",
+			RDG_EVENT_NAME("TSR DecimateHistory(#%d%s%s%s%s%s) %dx%d",
 				PermutationVector.ToDimensionValueId(),
 				PermutationVector.Get<FTSRDecimateHistoryCS::FMoireReprojectionDim>() ? TEXT(" ReprojectMoire") : TEXT(""),
 				PermutationVector.Get<FTSRDecimateHistoryCS::FResurrectionReprojectionDim>() ? TEXT(" ReprojectResurrection") : TEXT(""),
+				PermutationVector.Get<FTSRDecimateHistoryCS::FThinGeometryCoverageDim>() ? TEXT(" ThinGeometryCoverage") : TEXT(""),
 				PermutationVector.Get<FTSRShader::F16BitVALUDim>() ? TEXT(" 16bit") : TEXT(""),
 				PermutationVector.Get<FTSRShader::FAlphaChannelDim>() ? TEXT(" AlphaChannel") : TEXT(""),
 				InputRect.Width(), InputRect.Height()),
@@ -2367,6 +2464,88 @@
 			FComputeShaderUtils::GetGroupCount(InputRect.Size(), 8));
 	}
 
+	//Create the thin geometry mask to avoid over rejection
+	if(PassConfig.ThinGeometryDetectionEnable)
+	{
+		FTSRConvolutionNetworkShader::FPermutationDomain ConvolutionNetworkPermutationVector;
+		ConvolutionNetworkPermutationVector.Set<FTSRConvolutionNetworkShader::FWaveSizeOps>(SelectWaveSize(View.GetShaderPlatform(), { 16, 32, 64 }));
+		ConvolutionNetworkPermutationVector.Set<FTSRShader::F16BitVALUDim>(bUse16BitVALU);
+		ConvolutionNetworkPermutationVector.Set<FTSRShader::FAlphaChannelDim>(bSupportsAlpha);
+
+		FTSRDetectThinGeometryCS::FPermutationDomain PermutationVector;
+		PermutationVector.Set<FTSRConvolutionNetworkShader::FPermutationDomain>(ConvolutionNetworkPermutationVector);
+		PermutationVector = FTSRDetectThinGeometryCS::RemapPermutation(PermutationVector);
+
+		const int32 GroupTileSize = 32;
+		const int32 TileOverscan = FMath::Clamp(CVarTSRShadingTileOverscan.GetValueOnRenderThread(), 3, GroupTileSize / 2 - 1);
+		const int32 TileSize = GroupTileSize - 2 * TileOverscan;
+
+		FTSRDetectThinGeometryCS::FParameters* PassParameters = GraphBuilder.AllocParameters<FTSRDetectThinGeometryCS::FParameters>();
+		PassParameters->CommonParameters = CommonParameters;
+		PassParameters->TileOverscan = TileOverscan;
+		PassParameters->ThinGeometryTextureIndex = ThinGeometryTextureIndex;
+		PassParameters->ErrorMultiplier = PassConfig.ThinGeometryErrorMultiplier;
+		PassParameters->MaxRelaxationWeight = FMath::Clamp(CVarTSRThinGeometryCoverageMaxRelaxationWeight.GetValueOnRenderThread(), 0.0f, 1.0f);
+		PassParameters->SceneDepthTexture = PassInputs.SceneDepth.Texture;
+
+		// Coverage texture
+		{
+			if (bCameraCut)
+			{
+				PassParameters->ReprojectedHistoryCoverageTexture = GraphBuilder.CreateSRV(BlackArrayDummy);
+				PassParameters->CurrentCoverageTexture = GraphBuilder.CreateSRV(FRDGTextureSRVDesc(BlackDummy));
+			}
+			else
+			{
+				PassParameters->ReprojectedHistoryCoverageTexture = ReprojectedHistoryCoverageTexture ? GraphBuilder.CreateSRV(FRDGTextureSRVDesc::CreateForSlice(
+					ReprojectedHistoryCoverageTexture, /* SliceIndex = */ 0)) : GraphBuilder.CreateSRV(FRDGTextureSRVDesc(BlackDummy));;
+
+				if (PassInputs.FlickeringInputTexture.IsValid())
+				{
+					ensure(InputRect == PassInputs.FlickeringInputTexture.ViewRect);
+					int32 ArrayIndex = GetTSRMainFlickeringLumaTextureArraySize() - 1;
+					PassParameters->CurrentCoverageTexture = GraphBuilder.CreateSRV(
+						FRDGTextureSRVDesc::CreateForSlice(PassInputs.FlickeringInputTexture.Texture,  /* SliceIndex = */ ArrayIndex));
+				}
+				else
+				{
+					PassParameters->CurrentCoverageTexture = GraphBuilder.CreateSRV(FRDGTextureSRVDesc(BlackDummy));
+				}
+			}
+
+			//Output
+			{
+				if (View.bStatePrevViewInfoIsReadOnly)
+				{
+					PassParameters->HistoryCoverageOutput = CreateDummyUAVArray(GraphBuilder, History.CoverageArray->Desc.Format);
+				}
+				else
+				{
+					FRDGTextureUAVDesc CoverageUAVDesc(History.CoverageArray);
+					PassParameters->HistoryCoverageOutput = GraphBuilder.CreateUAV(CoverageUAVDesc);
+				}
+			}
+		}
+
+		PassParameters->R8Output = GraphBuilder.CreateUAV(R8OutputTexture);
+		PassParameters->DebugOutput = CreateDebugUAV(InputExtent, TEXT("Debug.TSR.DetectThinGeometry"));
+
+		TShaderMapRef<FTSRDetectThinGeometryCS> ComputeShader(View.ShaderMap, PermutationVector);
+		FComputeShaderUtils::AddPass(
+			GraphBuilder,
+			RDG_EVENT_NAME("TSR DetectThinGeometry(#%d TileSize=%d PaddingCostMultiplier=%1.1f WaveSize=%d VALU=%s%s) %dx%d",
+				PermutationVector.ToDimensionValueId(),
+				TileSize,
+				FMath::Pow(float(GroupTileSize) / float(TileSize), 2),
+				int32(PermutationVector.Get<FTSRConvolutionNetworkShader::FPermutationDomain>().Get<FTSRDetectThinGeometryCS::FWaveSizeOps>()),
+				PermutationVector.Get<FTSRConvolutionNetworkShader::FPermutationDomain>().Get<FTSRShader::F16BitVALUDim>() ? TEXT("16bit") : TEXT("32bit"),
+				PermutationVector.Get<FTSRConvolutionNetworkShader::FPermutationDomain>().Get<FTSRShader::FAlphaChannelDim>() ? TEXT(" AlphaChannel") : TEXT(""),
+				InputRect.Width(), InputRect.Height()),
+			AsyncComputePasses >= 2 ? ERDGPassFlags::AsyncCompute : ERDGPassFlags::Compute,
+			ComputeShader,
+			PassParameters,
+			FComputeShaderUtils::GetGroupCount(InputRect.Size(), TileSize));
+	}
 	// Merge PostDOF translucency within same scene color.
 	FRDGTextureRef InputSceneColorTexture = nullptr;
 	if (!bHasSeparateTranslucency)
@@ -2474,11 +2653,12 @@
 		if (PassInputs.FlickeringInputTexture.IsValid())
 		{
 			ensure(InputRect == PassInputs.FlickeringInputTexture.ViewRect);
-			PassParameters->InputMoireLumaTexture = PassInputs.FlickeringInputTexture.Texture;
+			PassParameters->InputMoireLumaTexture = GraphBuilder.CreateSRV(
+				FRDGTextureSRVDesc::CreateForSlice(PassInputs.FlickeringInputTexture.Texture, 0));
 		}
 		else
 		{
-			PassParameters->InputMoireLumaTexture = BlackDummy;
+			PassParameters->InputMoireLumaTexture = GraphBuilder.CreateSRV(FRDGTextureSRVDesc(BlackDummy));
 		}
 		PassParameters->InputSceneTranslucencyTexture = SeparateTranslucencyTexture;
 		PassParameters->ReprojectedHistoryGuideTexture = GraphBuilder.CreateSRV(FRDGTextureSRVDesc::CreateForSlice(
@@ -2966,6 +3146,11 @@
 			GraphBuilder.QueueTextureExtraction(History.MoireArray, &OutputHistory.MoireArray);
 		}
 
+		if (PassConfig.ThinGeometryDetectionEnable)
+		{
+			GraphBuilder.QueueTextureExtraction(History.CoverageArray, &OutputHistory.CoverageArray);
+		}
+
 		if (bLensDistortion && PassInputs.LensDistortionLUT.IsEnabled())
 		{
 			GraphBuilder.QueueTextureExtraction(PassInputs.LensDistortionLUT.DistortingDisplacementTexture, &OutputHistory.DistortingDisplacementTextures[CurrentFrameSliceIndex]);
@@ -3088,11 +3273,12 @@
 			if (PassInputs.FlickeringInputTexture.IsValid())
 			{
 				ensure(InputRect == PassInputs.FlickeringInputTexture.ViewRect);
-				PassParameters->InputMoireLumaTexture = PassInputs.FlickeringInputTexture.Texture;
+				PassParameters->InputMoireLumaTexture = GraphBuilder.CreateSRV(
+					FRDGTextureSRVDesc::CreateForSlice(PassInputs.FlickeringInputTexture.Texture, 0));
 			}
 			else
 			{
-				PassParameters->InputMoireLumaTexture = BlackDummy;
+				PassParameters->InputMoireLumaTexture = GraphBuilder.CreateSRV(FRDGTextureSRVDesc(BlackDummy));
 			}
 			PassParameters->InputSceneTranslucencyTexture = SeparateTranslucencyTexture;
 			PassParameters->SceneColorTexture = SceneColorOutputTextureSRV;
diff --git a/Engine/Source/Runtime/Renderer/Private/RendererScene.cpp b/Engine/Source/Runtime/Renderer/Private/RendererScene.cpp
--- a/Engine/Source/Runtime/Renderer/Private/RendererScene.cpp
+++ b/Engine/Source/Runtime/Renderer/Private/RendererScene.cpp
@@ -804,6 +804,11 @@ uint64 FTSRHistory::GetGPUSizeBytes(bool bLogSizes) const
 		GetRenderTargetGPUSizeBytes(GuideArray, bLogSizes) +
 		GetRenderTargetGPUSizeBytes(MoireArray, bLogSizes);
 
+	if (CoverageArray.IsValid())
+	{
+		TotalSize += GetRenderTargetGPUSizeBytes(CoverageArray, bLogSizes);
+	}
+
 	for (const TRefCountPtr <IPooledRenderTarget>& Texture : DistortingDisplacementTextures)
 	{
 		if (Texture.IsValid())
diff --git a/Engine/Source/Runtime/Renderer/Private/SceneRendering.h b/Engine/Source/Runtime/Renderer/Private/SceneRendering.h
--- a/Engine/Source/Runtime/Renderer/Private/SceneRendering.h
+++ b/Engine/Source/Runtime/Renderer/Private/SceneRendering.h
@@ -887,6 +887,7 @@ struct FTSRHistory
 	// Input resolution representation of the output
 	TRefCountPtr<IPooledRenderTarget> GuideArray;
 	TRefCountPtr<IPooledRenderTarget> MoireArray;
+	TRefCountPtr<IPooledRenderTarget> CoverageArray;
 
 	// Frame's input and output resolution.
 	FIntRect InputViewportRect;
