Add hair scope in DeferredShadingRenderer to regroup events.

diff --git a/Engine/Source/Runtime/Renderer/Private/DeferredShadingRenderer.cpp b/Engine/Source/Runtime/Renderer/Private/DeferredShadingRenderer.cpp
--- a/Engine/Source/Runtime/Renderer/Private/DeferredShadingRenderer.cpp
+++ b/Engine/Source/Runtime/Renderer/Private/DeferredShadingRenderer.cpp
@@ -2517,6 +2517,8 @@ void FDeferredShadingSceneRenderer::Render(FRDGBuilder& GraphBuilder, const FSce
 
 			if (bHairStrandsEnable)
 			{
+				RDG_EVENT_SCOPE(GraphBuilder, "Hair");
+
 				RunHairStrandsBookmark(GraphBuilder, EHairStrandsBookmark::ProcessStrandsInterpolation, HairStrandsBookmarkParameters);
 				if (!bHasRayTracedOverlay)
 				{
@@ -2684,6 +2686,8 @@ void FDeferredShadingSceneRenderer::Render(FRDGBuilder& GraphBuilder, const FSce
 		// Render hair
 		if (bHairStrandsEnable && !IsForwardShadingEnabled(ShaderPlatform))
 		{
+			RDG_EVENT_SCOPE(GraphBuilder, "Hair");
+
 			RunHairStrandsBookmark(GraphBuilder, EHairStrandsBookmark::ProcessStrandsInterpolation, HairStrandsBookmarkParameters);
 			if (!bHasRayTracedOverlay)
 			{
