[HWRT] Support "Performance" ray tracing visualization on Nvidia GPUs.
- using NVAPI timestamp.

diff --git a/Engine/Shaders/Private/RayTracing/RayTracingDebug.usf b/Engine/Shaders/Private/RayTracing/RayTracingDebug.usf
--- a/Engine/Shaders/Private/RayTracing/RayTracingDebug.usf
+++ b/Engine/Shaders/Private/RayTracing/RayTracingDebug.usf
@@ -75,7 +75,7 @@ float4 DebugRayTracingMaterial(uint2 PixelCoord, FRayDesc Ray, bool bIsFarField)
 	const bool bCameraRay = true;
 
 #if PLATFORM_SUPPORTS_SHADER_TIMESTAMP
-	uint64_t TimeBegin = GetShaderTimestamp();
+	FTimestamp TimeBegin = GetShaderTimestamp();
 #endif
 
 	RaytracingAccelerationStructure TLASToTrace = TLAS;
@@ -102,8 +102,9 @@ float4 DebugRayTracingMaterial(uint2 PixelCoord, FRayDesc Ray, bool bIsFarField)
 		bCameraRay);
 
 #if PLATFORM_SUPPORTS_SHADER_TIMESTAMP
-	uint64_t TimeEnd = GetShaderTimestamp();
+	FTimestamp TimeEnd = GetShaderTimestamp();
 #endif
+
 	if (Payload.IsHit())
 	{
 
@@ -124,7 +125,7 @@ float4 DebugRayTracingMaterial(uint2 PixelCoord, FRayDesc Ray, bool bIsFarField)
 		#endif
 
 		#if PLATFORM_SUPPORTS_SHADER_TIMESTAMP
-		case RAY_TRACING_DEBUG_VIZ_PERFORMANCE:				return float4(ColorMapInferno(float(TimeEnd - TimeBegin) * TimingScale), 0.0f);
+		case RAY_TRACING_DEBUG_VIZ_PERFORMANCE:				return float4(ColorMapInferno(float(ShaderTimestampDiff(TimeBegin, TimeEnd)) * TimingScale), 0.0f);
 		#endif // PLATFORM_SUPPORTS_SHADER_TIMESTAMP
 		}
 
diff --git a/Engine/Shaders/Public/Platform.ush b/Engine/Shaders/Public/Platform.ush
--- a/Engine/Shaders/Public/Platform.ush
+++ b/Engine/Shaders/Public/Platform.ush
@@ -590,8 +590,27 @@
 #endif
 
 #if !defined(PLATFORM_SUPPORTS_SHADER_TIMESTAMP)
+#if USE_NVAPI_TIMESTAMP
+	#include "/Engine/Shared/ThirdParty/NVIDIA/nvHLSLExtns.h"
+
+	#define PLATFORM_SUPPORTS_SHADER_TIMESTAMP 1
+
+	#define FTimestamp uint
+
+	FTimestamp GetShaderTimestamp()
+	{
+		return NvGetSpecial( NV_SPECIALOP_GLOBAL_TIMER_LO );
+	}
+
+	uint ShaderTimestampDiff(FTimestamp TimeBegin, FTimestamp TimeEnd)
+	{
+	  // Account for (at most one) overflow
+	  return TimeEnd >= TimeBegin ? (TimeEnd - TimeBegin) : (~0u - (TimeBegin - TimeEnd));
+	}
+#else // !USE_NVAPI_TIMESTAMP
 	#define PLATFORM_SUPPORTS_SHADER_TIMESTAMP 0
-#endif
+#endif // !USE_NVAPI_TIMESTAMP
+#endif // !defined(PLATFORM_SUPPORTS_SHADER_TIMESTAMP)
 
 // Hlslcc platforms ignore the uniform keyword as it can't properly optimize flow
 #if COMPILER_HLSLCC
diff --git a/Engine/Source/Runtime/D3D12RHI/Private/D3D12Adapter.cpp b/Engine/Source/Runtime/D3D12RHI/Private/D3D12Adapter.cpp
--- a/Engine/Source/Runtime/D3D12RHI/Private/D3D12Adapter.cpp
+++ b/Engine/Source/Runtime/D3D12RHI/Private/D3D12Adapter.cpp
@@ -1198,6 +1198,8 @@ void FD3D12Adapter::InitializeDevices()
 				}
 #endif // D3D12_RHI_WORKGRAPHS_GRAPHICS
 #endif // D3D12_MAX_FEATURE_OPTIONS
+
+				GRHIGlobals.SupportsShaderTimestamp = IsRHIDeviceNVIDIA();
 			}
 #endif // PLATFORM_WINDOWS
 		}
diff --git a/Engine/Source/Runtime/Renderer/Private/DeferredShadingRenderer.cpp b/Engine/Source/Runtime/Renderer/Private/DeferredShadingRenderer.cpp
--- a/Engine/Source/Runtime/Renderer/Private/DeferredShadingRenderer.cpp
+++ b/Engine/Source/Runtime/Renderer/Private/DeferredShadingRenderer.cpp
@@ -1460,6 +1460,7 @@ void FDeferredShadingSceneRenderer::RenderNanite(FRDGBuilder& GraphBuilder, cons
 
 #if RHI_RAYTRACING
 extern void RenderRayTracingDebug(FRDGBuilder& GraphBuilder, const FScene& Scene, const FViewInfo& View, FSceneTextures& SceneTextures, FRayTracingPickingFeedback& PickingFeedback);
+extern void RayTracingDebugDisplayOnScreenMessages(FScreenMessageWriter& Writer, const FViewInfo& View);
 #endif // RHI_RAYTRACING
 
 void FDeferredShadingSceneRenderer::Render(FRDGBuilder& GraphBuilder, const FSceneRenderUpdateInputs* SceneUpdateInputs)
@@ -3380,8 +3381,9 @@ void FDeferredShadingSceneRenderer::Render(FRDGBuilder& GraphBuilder, const FSce
 					FRayTracingPickingFeedback PickingFeedback = {};
 					RenderRayTracingDebug(GraphBuilder, *Scene, View, SceneTextures, PickingFeedback);
 
-					OnGetOnScreenMessages.AddLambda([this, PickingFeedback](FScreenMessageWriter& ScreenMessageWriter)->void
+					OnGetOnScreenMessages.AddLambda([this, &View, PickingFeedback](FScreenMessageWriter& ScreenMessageWriter)->void
 						{
+							RayTracingDebugDisplayOnScreenMessages(ScreenMessageWriter, View);
 							RayTracingDisplayPicking(PickingFeedback, ScreenMessageWriter);
 						});
 				}
diff --git a/Engine/Source/Runtime/Renderer/Private/RayTracing/RayTracingDebug.cpp b/Engine/Source/Runtime/Renderer/Private/RayTracing/RayTracingDebug.cpp
--- a/Engine/Source/Runtime/Renderer/Private/RayTracing/RayTracingDebug.cpp
+++ b/Engine/Source/Runtime/Renderer/Private/RayTracing/RayTracingDebug.cpp
@@ -168,8 +168,8 @@ class FRayTracingDebugRGS : public FGlobalShader
 	SHADER_USE_ROOT_PARAMETER_STRUCT(FRayTracingDebugRGS, FGlobalShader)
 	
 	class FUseDebugCHSType : SHADER_PERMUTATION_BOOL("USE_DEBUG_CHS");
-
-	using FPermutationDomain = TShaderPermutationDomain<FUseDebugCHSType>;
+	class FUseNvAPITimestamp : SHADER_PERMUTATION_BOOL("USE_NVAPI_TIMESTAMP");
+	using FPermutationDomain = TShaderPermutationDomain<FUseDebugCHSType, FUseNvAPITimestamp>;
 
 	BEGIN_SHADER_PARAMETER_STRUCT(FParameters, )
 		SHADER_PARAMETER(uint32, VisualizationMode)
@@ -205,6 +205,11 @@ class FRayTracingDebugRGS : public FGlobalShader
 		return ShouldCompileRayTracingShadersForProject(Parameters.Platform);
 	}
 
+	static void ModifyCompilationEnvironment(const FGlobalShaderPermutationParameters& Parameters, FShaderCompilerEnvironment& OutEnvironment)
+	{
+		FGlobalShader::ModifyCompilationEnvironment(Parameters, OutEnvironment);
+	}
+
 	static ERayTracingPayloadType GetRayTracingPayloadType(const int32 PermutationId)
 	{
 		FPermutationDomain PermutationVector(PermutationId);
@@ -865,6 +870,7 @@ void FDeferredShadingSceneRenderer::PrepareRayTracingDebug(const FSceneViewFamil
 		{
 			FRayTracingDebugRGS::FPermutationDomain PermutationVector;
 			PermutationVector.Set<FRayTracingDebugRGS::FUseDebugCHSType>(false);
+			PermutationVector.Set<FRayTracingDebugRGS::FUseNvAPITimestamp>(IsRHIDeviceNVIDIA());
 			auto RayGenShader = GetGlobalShaderMap(ViewFamily.GetShaderPlatform())->GetShader<FRayTracingDebugRGS>(PermutationVector);
 			OutRayGenShaders.Add(RayGenShader.GetRayTracingShader());
 		}
@@ -1610,6 +1616,7 @@ void RenderRayTracingDebug(FRDGBuilder& GraphBuilder, const FScene& Scene, const
 
 	FRayTracingDebugRGS::FPermutationDomain PermutationVector;
 	PermutationVector.Set<FRayTracingDebugRGS::FUseDebugCHSType>(bRequiresDebugCHS);
+	PermutationVector.Set<FRayTracingDebugRGS::FUseNvAPITimestamp>(IsRHIDeviceNVIDIA());
 
 	auto RayGenShader = View.ShaderMap->GetShader<FRayTracingDebugRGS>(PermutationVector);
 
@@ -1747,6 +1754,23 @@ void RenderRayTracingDebug(FRDGBuilder& GraphBuilder, const FScene& Scene, const
 	}
 }
 
+extern void RayTracingDebugDisplayOnScreenMessages(FScreenMessageWriter& Writer, const FViewInfo& View)
+{
+	const uint32 DebugVisualizationMode = GetRaytracingDebugViewMode(View);
+
+	switch (DebugVisualizationMode)
+	{
+	case RAY_TRACING_DEBUG_VIZ_PERFORMANCE:
+	{
+		static const FText Message = NSLOCTEXT("Renderer", "RayTracingDebugVizPerformance", "Use r.RayTracing.Visualize.TimingScale to adjust visualization.");
+		Writer.DrawLine(Message, 10, FColor::White);
+		break;
+	}
+	default:
+		break;
+	}
+}
+
 void FDeferredShadingSceneRenderer::RayTracingDisplayPicking(const FRayTracingPickingFeedback& PickingFeedback, FScreenMessageWriter& Writer)
 {
 	if (PickingFeedback.InstanceIndex == ~uint32(0))
