[HWRT] Support compacting active instances in instance buffer.

diff --git a/Engine/Plugins/Experimental/GPULightmass/Source/GPULightmass/Private/LightmapRenderer.cpp b/Engine/Plugins/Experimental/GPULightmass/Source/GPULightmass/Private/LightmapRenderer.cpp
--- a/Engine/Plugins/Experimental/GPULightmass/Source/GPULightmass/Private/LightmapRenderer.cpp
+++ b/Engine/Plugins/Experimental/GPULightmass/Source/GPULightmass/Private/LightmapRenderer.cpp
@@ -1220,6 +1220,8 @@ bool FSceneRenderState::SetupRayTracingScene(FRDGBuilder& GraphBuilder, FSceneUn
 				nullptr,
 				nullptr,
 				InstanceBuffer.UAV,
+				SceneInitializer.MaxNumInstances,
+				/*bCompactOutput*/ false,
 				nullptr,
 				0,
 				nullptr);
diff --git a/Engine/Shaders/Private/RayTracing/RayTracingInstanceBufferUtil.usf b/Engine/Shaders/Private/RayTracing/RayTracingInstanceBufferUtil.usf
--- a/Engine/Shaders/Private/RayTracing/RayTracingInstanceBufferUtil.usf
+++ b/Engine/Shaders/Private/RayTracing/RayTracingInstanceBufferUtil.usf
@@ -25,7 +25,7 @@ struct FRayTracingInstanceDescriptor
 	uint SceneInstanceIndexAndApplyLocalBoundsTransform;
 };
 
-uint NumInstances;
+uint MaxNumInstances;
 uint NumGroups;
 uint NumInstanceDescriptors;
 uint BaseGroupDescriptorIndex;
@@ -108,11 +108,6 @@ bool CullInstance(FInstanceSceneData InstanceSceneData, FPrimitiveSceneData Prim
 
 void ProcessInstance(FRayTracingInstanceDescriptor InputDesc)
 {
-	if (InputDesc.OutputDescriptorIndex >= NumInstances)
-	{
-		return;
-	}
-
 	const FDFVector3 PreViewTranslation = MakeDFVector3(PreViewTranslationHigh, PreViewTranslationLow);
 
 	bool bCulled = false;
@@ -157,6 +152,33 @@ void ProcessInstance(FRayTracingInstanceDescriptor InputDesc)
 	LocalToTranslatedWorld[2] = InstanceTransforms[InputDesc.GPUSceneInstanceOrTransformIndex * 3 + 2];
 #endif // !USE_GPUSCENE
 
+	if (bCulled)
+	{
+#	if COMPACT_OUTPUT
+		// when using compaction don't need to write inactive instances
+		return;
+#	endif
+	}
+	else
+	{
+#	if OUTPUT_STATS
+		uint CompactedIndex;
+#		if USE_WAVE_OPS
+			WaveInterlockedAddScalar_(RWOutputStats[OutputStatsOffset + 0], 1, CompactedIndex);
+#		else
+			InterlockedAdd(RWOutputStats[OutputStatsOffset + 0], 1, CompactedIndex);
+#		endif
+#		if COMPACT_OUTPUT
+			InputDesc.OutputDescriptorIndex = CompactedIndex;
+#		endif // COMPACT_OUTPUT
+#	endif // OUTPUT_STATS
+	}
+
+	if (InputDesc.OutputDescriptorIndex >= MaxNumInstances)
+	{
+		return;
+	}
+
 #if OUTPUT_INSTANCE_EXTRA_DATA
 	if (!bCulled)
 	{
@@ -185,17 +207,6 @@ void ProcessInstance(FRayTracingInstanceDescriptor InputDesc)
 		InputDesc.InstanceContributionToHitGroupIndex,
 		LocalToTranslatedWorld,
 		BlasAddress);
-
-#if OUTPUT_STATS
-	if (!bCulled)
-	{
-#if USE_WAVE_OPS
-		WaveInterlockedAddScalar(RWOutputStats[OutputStatsOffset + 0], 1);
-#else
-		InterlockedAdd(RWOutputStats[OutputStatsOffset + 0], 1);
-#endif
-	}
-#endif
 }
 
 [numthreads(THREADGROUP_SIZE, 1, 1)]
diff --git a/Engine/Source/Runtime/Renderer/Private/RayTracing/RayTracingInstanceBufferUtil.cpp b/Engine/Source/Runtime/Renderer/Private/RayTracing/RayTracingInstanceBufferUtil.cpp
--- a/Engine/Source/Runtime/Renderer/Private/RayTracing/RayTracingInstanceBufferUtil.cpp
+++ b/Engine/Source/Runtime/Renderer/Private/RayTracing/RayTracingInstanceBufferUtil.cpp
@@ -515,7 +515,7 @@ struct FRayTracingBuildInstanceBufferCS : public FGlobalShader
 		SHADER_PARAMETER_SRV(ByteAddressBuffer, AccelerationStructureAddresses)
 		SHADER_PARAMETER_SRV(StructuredBuffer, InstanceTransforms)
 
-		SHADER_PARAMETER(uint32, NumInstances)
+		SHADER_PARAMETER(uint32, MaxNumInstances)
 		SHADER_PARAMETER(uint32, NumGroups)
 		SHADER_PARAMETER(uint32, NumInstanceDescriptors)
 		SHADER_PARAMETER(uint32, BaseGroupDescriptorIndex)
@@ -543,8 +543,9 @@ struct FRayTracingBuildInstanceBufferCS : public FGlobalShader
 	class FOutputInstanceExtraDataDim : SHADER_PERMUTATION_BOOL("OUTPUT_INSTANCE_EXTRA_DATA");
 	class FGpuCullingDim : SHADER_PERMUTATION_BOOL("GPU_CULLING");
 	class FOutputStatsDim : SHADER_PERMUTATION_BOOL("OUTPUT_STATS");
+	class FCompactOutputDim : SHADER_PERMUTATION_BOOL("COMPACT_OUTPUT");
 	class FUseWaveOpsDim : SHADER_PERMUTATION_BOOL("USE_WAVE_OPS");
-	using FPermutationDomain = TShaderPermutationDomain<FSupportInstanceGroupsDim, FUseGPUSceneDim, FOutputInstanceExtraDataDim, FGpuCullingDim, FOutputStatsDim, FUseWaveOpsDim>;
+	using FPermutationDomain = TShaderPermutationDomain<FSupportInstanceGroupsDim, FUseGPUSceneDim, FOutputInstanceExtraDataDim, FGpuCullingDim, FOutputStatsDim, FCompactOutputDim, FUseWaveOpsDim>;
 		
 	static constexpr uint32 ThreadGroupSize = GRayTracingInstanceGroupSize;
 
@@ -579,7 +580,7 @@ void BuildRayTracingInstanceBuffer(
 	FRHICommandList& RHICmdList,
 	const FGPUScene* GPUScene,
 	const FDFVector3& PreViewTranslation,
-	uint32 NumInstances,
+	uint32 MaxNumInstances,
 	uint32 NumGroups,
 	uint32 NumInstanceDescriptors,
 	FRHIUnorderedAccessView* InstancesUAV,
@@ -590,6 +591,7 @@ void BuildRayTracingInstanceBuffer(
 	FRHIShaderResourceView* AccelerationStructureAddressesSRV,
 	FRHIShaderResourceView* InstanceTransformSRV,
 	const FRayTracingCullingParameters* CullingParameters,
+	bool bCompactOutput,
 	FRHIUnorderedAccessView* OutputStatsUAV,
 	uint32 OutputStatsOffset,
 	FRHIUnorderedAccessView* InstanceExtraDataUAV)
@@ -600,7 +602,7 @@ void BuildRayTracingInstanceBuffer(
 	PassParams.InstanceDescriptors = InstanceUploadSRV;
 	PassParams.AccelerationStructureAddresses = AccelerationStructureAddressesSRV;
 	PassParams.InstanceTransforms = InstanceTransformSRV;
-	PassParams.NumInstances = NumInstances;
+	PassParams.MaxNumInstances = MaxNumInstances;
 	PassParams.NumGroups = NumGroups;
 	PassParams.NumInstanceDescriptors = NumInstanceDescriptors;
 	PassParams.BaseGroupDescriptorIndex = InstanceGroupUploadOffset;
@@ -632,6 +634,7 @@ void BuildRayTracingInstanceBuffer(
 	PermutationVector.Set<FRayTracingBuildInstanceBufferCS::FUseGPUSceneDim>(InstanceTransformSRV == nullptr);
 	PermutationVector.Set<FRayTracingBuildInstanceBufferCS::FOutputInstanceExtraDataDim>(InstanceExtraDataUAV != nullptr);
 	PermutationVector.Set<FRayTracingBuildInstanceBufferCS::FOutputStatsDim>(OutputStatsUAV != nullptr);
+	PermutationVector.Set<FRayTracingBuildInstanceBufferCS::FCompactOutputDim>(bCompactOutput);
 	PermutationVector.Set<FRayTracingBuildInstanceBufferCS::FUseWaveOpsDim>(GRHISupportsWaveOperations);
 	PermutationVector.Set<FRayTracingBuildInstanceBufferCS::FGpuCullingDim>(CullingParameters != nullptr);
 
@@ -659,24 +662,24 @@ void BuildRayTracingInstanceBuffer(
 	FRHIShaderResourceView* InstanceUploadSRV,
 	FRHIShaderResourceView* AccelerationStructureAddressesSRV,
 	FRHIShaderResourceView* CPUInstanceTransformSRV,
-	uint32 NumNativeGPUSceneInstances,
-	uint32 NumNativeCPUInstances,
+	uint32 MaxNumInstances,
 	uint32 NumGPUGroups,
 	uint32 NumCPUGroups,
 	uint32 NumGPUInstanceDescriptors,
 	uint32 NumCPUInstanceDescriptors,
 	const FRayTracingCullingParameters* CullingParameters,
+	bool bCompactOutput,
 	FRHIUnorderedAccessView* OutputStatsUAV,
 	uint32 OutputStatsOffset,
 	FRHIUnorderedAccessView* InstanceExtraDataUAV)
 {
-	if (NumNativeGPUSceneInstances > 0)
+	if (NumGPUInstanceDescriptors > 0)
 	{
 		BuildRayTracingInstanceBuffer(
 			RHICmdList,
 			GPUScene,
 			PreViewTranslation,
-			NumNativeGPUSceneInstances + NumNativeCPUInstances,
+			MaxNumInstances,
 			NumGPUGroups,
 			NumGPUInstanceDescriptors,
 			InstancesUAV,
@@ -687,18 +690,19 @@ void BuildRayTracingInstanceBuffer(
 			AccelerationStructureAddressesSRV,
 			nullptr,
 			CullingParameters,
+			bCompactOutput,
 			OutputStatsUAV,
 			OutputStatsOffset,
 			InstanceExtraDataUAV);
 	}
 
-	if (NumNativeCPUInstances > 0)
+	if (NumCPUInstanceDescriptors > 0)
 	{
 		BuildRayTracingInstanceBuffer(
 			RHICmdList,
 			GPUScene,
 			PreViewTranslation,
-			NumNativeGPUSceneInstances + NumNativeCPUInstances,
+			MaxNumInstances,
 			NumCPUGroups,
 			NumCPUInstanceDescriptors,
 			InstancesUAV,
@@ -709,6 +713,7 @@ void BuildRayTracingInstanceBuffer(
 			AccelerationStructureAddressesSRV,
 			CPUInstanceTransformSRV,
 			nullptr,
+			bCompactOutput,
 			OutputStatsUAV,
 			OutputStatsOffset,
 			InstanceExtraDataUAV);
@@ -738,13 +743,13 @@ void BuildRayTracingInstanceBuffer(
 		InstanceUploadSRV,
 		AccelerationStructureAddressesSRV,
 		CPUInstanceTransformSRV,
-		NumNativeGPUSceneInstances,
-		NumNativeCPUInstances,
-		0,
-		0,
+		NumNativeGPUSceneInstances + NumNativeCPUInstances,
 		0,
 		0,
+		NumNativeGPUSceneInstances,
+		NumNativeCPUInstances,
 		CullingParameters,
+		/*bCompactOutput*/ false,
 		OutputStatsUAV,
 		0,
 		InstanceExtraDataUAV);
@@ -903,6 +908,8 @@ void FRayTracingInstanceBufferBuilder::BuildRayTracingInstanceBuffer(
 	const FGPUScene* GPUScene,
 	const FRayTracingCullingParameters* CullingParameters,
 	FRHIUnorderedAccessView* InstancesUAV,
+	uint32 MaxNumInstances,
+	bool bCompactOutput,
 	FRHIUnorderedAccessView* OutputStatsUAV,
 	uint32 OutputStatsOffset,
 	FRHIUnorderedAccessView* InstanceExtraDataUAV)
@@ -916,13 +923,13 @@ void FRayTracingInstanceBufferBuilder::BuildRayTracingInstanceBuffer(
 		InstanceUploadSRV,
 		AccelerationStructureAddressesBuffer.SRV,
 		TransformUploadSRV,
-		Data.NumNativeGPUSceneInstances,
-		Data.NumNativeCPUInstances,
+		MaxNumInstances,
 		Data.NumGPUInstanceGroups,
 		Data.NumCPUInstanceGroups,
 		Data.NumGPUInstanceDescriptors,
 		Data.NumCPUInstanceDescriptors,
 		CullingParameters,
+		bCompactOutput,
 		OutputStatsUAV,
 		OutputStatsOffset,
 		InstanceExtraDataUAV);
diff --git a/Engine/Source/Runtime/Renderer/Private/RayTracing/RayTracingScene.cpp b/Engine/Source/Runtime/Renderer/Private/RayTracing/RayTracingScene.cpp
--- a/Engine/Source/Runtime/Renderer/Private/RayTracing/RayTracingScene.cpp
+++ b/Engine/Source/Runtime/Renderer/Private/RayTracing/RayTracingScene.cpp
@@ -41,6 +41,31 @@ static TAutoConsoleVariable<bool> CVarRayTracingSceneUseTracingFeedback(
 	ECVF_RenderThreadSafe
 );
 
+static TAutoConsoleVariable<bool> CVarRayTracingSceneCompactInstances(
+	TEXT("r.RayTracing.Scene.CompactInstances"),
+	false,
+	TEXT("Whether to compact the instance buffer so it only contains active instances.\n")
+	TEXT("On platforms that don't support indirect TLAS build this requires doing a GPU->CPU readback, ")
+	TEXT("which lead so instances missing from TLAS due to the extra latency.\n")
+	TEXT("r.RayTracing.Scene.CompactInstances.Min and r.RayTracing.Scene.CompactInstances.Margin can be used to avoid those issues."),
+	ECVF_RenderThreadSafe | ECVF_Scalability
+);
+
+static TAutoConsoleVariable<int32> CVarRayTracingSceneCompactInstancesMin(
+	TEXT("r.RayTracing.Scene.CompactInstances.Min"),
+	0,
+	TEXT("Minimum of instances in the instance buffer when using compaction.\n")
+	TEXT("Should be set to the expected high water mark to avoid issues on platforms that don't support indirect TLAS build."),
+	ECVF_RenderThreadSafe
+);
+
+static TAutoConsoleVariable<int32> CVarRayTracingSceneCompactInstancesMargin(
+	TEXT("r.RayTracing.Scene.CompactInstances.Margin"),
+	5000,
+	TEXT("Margin applied on top of lastest number of active instances readback from GPU to avoid issues when number instances increases from frame to frame."),
+	ECVF_RenderThreadSafe
+);
+
 bool IsRayTracingFeedbackEnabled(const FSceneViewFamily& ViewFamily)
 {
 	// TODO: For now Feedback is limited to inline passes
@@ -129,11 +154,22 @@ void FRayTracingScene::Create(FRDGBuilder& GraphBuilder, const FViewInfo& View,
 		AddClearUAVPass(GraphBuilder, InstanceStatsBufferUAV, 0, ComputePassFlags);
 	}
 
+	const bool bCompactInstanceBuffer = CVarRayTracingSceneCompactInstances.GetValueOnRenderThread();
+
 	for (uint32 LayerIndex = 0; LayerIndex < NumLayers; ++LayerIndex)
 	{
 		FLayer& Layer = Layers[LayerIndex];
 
-		const uint32 NumNativeInstances = Layer.InstanceBufferBuilder.GetMaxNumInstances();
+		uint32 NumNativeInstances = Layer.InstanceBufferBuilder.GetMaxNumInstances();
+
+		if (bCompactInstanceBuffer)
+		{
+			NumNativeInstances = FMath::Max<uint32>(CVarRayTracingSceneCompactInstancesMin.GetValueOnRenderThread(), Layer.NumActiveInstances + CVarRayTracingSceneCompactInstancesMargin.GetValueOnRenderThread());
+			NumNativeInstances = FMath::Min<uint32>(NumNativeInstances, Layer.InstanceBufferBuilder.GetMaxNumInstances());
+		}
+
+		Layer.MaxNumInstances = NumNativeInstances;
+
 		const uint32 NumNativeInstancesAligned = FMath::DivideAndRoundUp(FMath::Max(NumNativeInstances, 1U), AllocationGranularity) * AllocationGranularity;
 
 		{
@@ -178,11 +214,17 @@ void FRayTracingScene::Create(FRDGBuilder& GraphBuilder, const FViewInfo& View,
 			InstanceBufferDesc.NumElements = NumNativeInstancesAligned;
 
 			Layer.InstanceBuffer = GraphBuilder.CreateBuffer(InstanceBufferDesc, TEXT("FRayTracingScene::InstanceBuffer"));
+
+			if (bCompactInstanceBuffer)
+			{
+				// need to clear since FRayTracingBuildInstanceBufferCS will only write active instances
+				AddClearUAVPass(GraphBuilder, GraphBuilder.CreateUAV(Layer.InstanceBuffer), 0, ComputePassFlags);
+			}
 		}
 
 		// Feedback
 		if (bTracingFeedbackEnabled)
-		{				
+		{
 			{
 				FRDGBufferDesc InstanceHitCountBufferDesc;
 				InstanceHitCountBufferDesc.Usage = EBufferUsageFlags::UnorderedAccess | EBufferUsageFlags::ShaderResource | EBufferUsageFlags::StructuredBuffer;
@@ -255,11 +297,17 @@ void FRayTracingScene::Create(FRDGBuilder& GraphBuilder, const FViewInfo& View,
 					InstanceBufferBuilder.FillAccelerationStructureAddressesBuffer(RHICmdList);
 				});
 
+#if STATS
+			const bool bStatsEnabled = true;
+#else
+			const bool bStatsEnabled = false;
+#endif
+
 			FBuildInstanceBufferPassParams* PassParams = GraphBuilder.AllocParameters<FBuildInstanceBufferPassParams>();
 			PassParams->InstanceBuffer = GraphBuilder.CreateUAV(Layer.InstanceBuffer);
 			PassParams->InstanceExtraDataBuffer = InstanceExtraDataBufferUAV;
 			PassParams->Scene = View.GetSceneUniforms().GetBuffer(GraphBuilder);
-			PassParams->OutputStats = STATS ? InstanceStatsBufferUAV : nullptr;
+			PassParams->OutputStats = bCompactInstanceBuffer || bStatsEnabled ? InstanceStatsBufferUAV : nullptr;
 
 			GraphBuilder.AddPass(
 				RDG_EVENT_NAME("RayTracingBuildInstanceBuffer"),
@@ -269,7 +317,9 @@ void FRayTracingScene::Create(FRDGBuilder& GraphBuilder, const FViewInfo& View,
 				&InstanceBufferBuilder = Layer.InstanceBufferBuilder,
 				LayerIndex,
 				GPUScene,
-				CullingParameters = &View.RayTracingCullingParameters
+				CullingParameters = &View.RayTracingCullingParameters,
+				NumNativeInstances,
+				bCompactInstanceBuffer
 				](FRHICommandList& RHICmdList)
 				{
 					InstanceBufferBuilder.BuildRayTracingInstanceBuffer(
@@ -277,6 +327,8 @@ void FRayTracingScene::Create(FRDGBuilder& GraphBuilder, const FViewInfo& View,
 						GPUScene,
 						CullingParameters,
 						PassParams->InstanceBuffer->GetRHI(),
+						NumNativeInstances,
+						bCompactInstanceBuffer,
 						PassParams->OutputStats ? PassParams->OutputStats->GetRHI() : nullptr,
 						LayerIndex,
 						PassParams->InstanceExtraDataBuffer ? PassParams->InstanceExtraDataBuffer->GetRHI() : nullptr);
@@ -316,7 +368,7 @@ void FRayTracingScene::Build(FRDGBuilder& GraphBuilder, ERDGPassFlags ComputePas
 				BuildParams.ScratchBufferOffset = 0;
 				BuildParams.InstanceBuffer = PassParams->InstanceBuffer->GetRHI();
 				BuildParams.InstanceBufferOffset = 0;
-				BuildParams.NumInstances = Layer.InstanceBufferBuilder.GetMaxNumInstances();
+				BuildParams.NumInstances = Layer.MaxNumInstances;
 				BuildParams.ReferencedGeometries = Layer.InstanceBufferBuilder.GetReferencedGeometries();
 
 				RHICmdList.BindAccelerationStructureMemory(Layer.RayTracingSceneRHI, PassParams->TLASBuffer->GetRHI(), 0);
diff --git a/Engine/Source/Runtime/Renderer/Private/RayTracing/RayTracingScene.h b/Engine/Source/Runtime/Renderer/Private/RayTracing/RayTracingScene.h
--- a/Engine/Source/Runtime/Renderer/Private/RayTracing/RayTracingScene.h
+++ b/Engine/Source/Runtime/Renderer/Private/RayTracing/RayTracingScene.h
@@ -211,6 +211,7 @@ private:
 		TArray<FRayTracingInstanceDebugData> InstancesDebugData;
 
 		uint32 NumActiveInstances = 0;
+		uint32 MaxNumInstances = 0;
 	};
 
 	TArray<FLayer> Layers;
diff --git a/Engine/Source/Runtime/Renderer/Private/Tests/RayTracingTestbed.cpp b/Engine/Source/Runtime/Renderer/Private/Tests/RayTracingTestbed.cpp
--- a/Engine/Source/Runtime/Renderer/Private/Tests/RayTracingTestbed.cpp
+++ b/Engine/Source/Runtime/Renderer/Private/Tests/RayTracingTestbed.cpp
@@ -213,6 +213,8 @@ bool RunRayTracingTestbed_RenderThread(const FString& Parameters)
 		nullptr,
 		nullptr,
 		InstanceBuffer.UAV,
+		SceneInitializer.MaxNumInstances,
+		/*bCompactOutput*/ false,
 		nullptr,
 		0,
 		nullptr);
diff --git a/Engine/Source/Runtime/Renderer/Public/RayTracingInstanceBufferUtil.h b/Engine/Source/Runtime/Renderer/Public/RayTracingInstanceBufferUtil.h
--- a/Engine/Source/Runtime/Renderer/Public/RayTracingInstanceBufferUtil.h
+++ b/Engine/Source/Runtime/Renderer/Public/RayTracingInstanceBufferUtil.h
@@ -117,6 +117,8 @@ public:
 		const FGPUScene* GPUScene,
 		const FRayTracingCullingParameters* CullingParameters,
 		FRHIUnorderedAccessView* InstancesUAV,
+		uint32 MaxNumInstances,
+		bool bCompactOutput,
 		FRHIUnorderedAccessView* OutputStatsUAV,
 		uint32 OutputStatsOffset,
 		FRHIUnorderedAccessView* InstanceExtraDataUAV);
