Fix error compiling FRayTracingDebugRGS when using Vulkan.

diff --git a/Engine/Source/Runtime/Renderer/Private/RayTracing/RayTracingDebug.cpp b/Engine/Source/Runtime/Renderer/Private/RayTracing/RayTracingDebug.cpp
--- a/Engine/Source/Runtime/Renderer/Private/RayTracing/RayTracingDebug.cpp
+++ b/Engine/Source/Runtime/Renderer/Private/RayTracing/RayTracingDebug.cpp
@@ -202,6 +202,15 @@ class FRayTracingDebugRGS : public FGlobalShader
 
 	static bool ShouldCompilePermutation(const FGlobalShaderPermutationParameters& Parameters)
 	{
+		FPermutationDomain PermutationVector(Parameters.PermutationId);
+
+		// TODO: Check this using DDPI
+		const bool bUseNvAPITimestamp = PermutationVector.Get<FUseNvAPITimestamp>();
+		if (bUseNvAPITimestamp && IsVulkanPlatform(Parameters.Platform))
+		{
+			return false;
+		}
+
 		return ShouldCompileRayTracingShadersForProject(Parameters.Platform);
 	}
 
@@ -870,7 +879,7 @@ void FDeferredShadingSceneRenderer::PrepareRayTracingDebug(const FSceneViewFamil
 		{
 			FRayTracingDebugRGS::FPermutationDomain PermutationVector;
 			PermutationVector.Set<FRayTracingDebugRGS::FUseDebugCHSType>(false);
-			PermutationVector.Set<FRayTracingDebugRGS::FUseNvAPITimestamp>(IsRHIDeviceNVIDIA());
+			PermutationVector.Set<FRayTracingDebugRGS::FUseNvAPITimestamp>(GRHIGlobals.SupportsShaderTimestamp && IsRHIDeviceNVIDIA());
 			auto RayGenShader = GetGlobalShaderMap(ViewFamily.GetShaderPlatform())->GetShader<FRayTracingDebugRGS>(PermutationVector);
 			OutRayGenShaders.Add(RayGenShader.GetRayTracingShader());
 		}
@@ -1616,7 +1625,7 @@ void RenderRayTracingDebug(FRDGBuilder& GraphBuilder, const FScene& Scene, const
 
 	FRayTracingDebugRGS::FPermutationDomain PermutationVector;
 	PermutationVector.Set<FRayTracingDebugRGS::FUseDebugCHSType>(bRequiresDebugCHS);
-	PermutationVector.Set<FRayTracingDebugRGS::FUseNvAPITimestamp>(IsRHIDeviceNVIDIA());
+	PermutationVector.Set<FRayTracingDebugRGS::FUseNvAPITimestamp>(GRHIGlobals.SupportsShaderTimestamp && IsRHIDeviceNVIDIA());
 
 	auto RayGenShader = View.ShaderMap->GetShader<FRayTracingDebugRGS>(PermutationVector);
 
