[HWRT] Add inline modes for ray tracing visualizations that don't use materials but use debug closest-hit shader.

#rnx

diff --git a/Engine/Shaders/Private/RayTracing/RayTracingDebug.usf b/Engine/Shaders/Private/RayTracing/RayTracingDebug.usf
--- a/Engine/Shaders/Private/RayTracing/RayTracingDebug.usf
+++ b/Engine/Shaders/Private/RayTracing/RayTracingDebug.usf
@@ -6,6 +6,8 @@
 #define SUBSTRATE_COMPLEXSPECIALPATH 1				// Allow to read data from complex special path
 #define SUBSTRATE_INLINE_SHADING 0
 
+#define RAYTRACING_DEBUG_INLINE COMPUTESHADER
+
 #include "/Engine/Shared/RayTracingDebugTypes.h"
 #include "/Engine/Shared/RayTracingTypes.h"
 #include "../Common.ush"
@@ -48,7 +50,7 @@ StructuredBuffer<FRayTracingInstanceDebugData> InstancesDebugData;
 StructuredBuffer<FPlatformRayTracingInstanceDescriptor> InstanceBuffer;
 StructuredBuffer<FRayTracingPickingFeedback> PickingBuffer;
 
-#if USE_DEBUG_CHS == 0 // add functions that require shaded payload
+#if (USE_DEBUG_CHS == 0) && (RAYTRACING_DEBUG_INLINE == 0) // add functions that require shaded payload
 
 float4 AnisoDebugColor(float InAniso)
 {
@@ -220,6 +222,33 @@ float4 DebugRayTracingInstance(FRayDesc Ray, out float OutDepth)
 	{
 		RayFlags |= RAY_FLAG_FORCE_OPAQUE;
 	}
+
+#if RAYTRACING_DEBUG_INLINE
+	FTraceRayInlineContext TraceRayInlineContext = CreateTraceRayInlineContext();
+
+	FTraceRayInlineResult TraceRayResult = TraceRayInline(
+		TLAS,
+		RayFlags,
+		InstanceInclusionMask,
+		Ray.GetNativeDesc(),
+		TraceRayInlineContext
+	);
+	
+	Payload.HitT = TraceRayResult.HitT;
+
+	const uint GPUSceneInstanceId = TraceRayResult.InstanceID;
+
+	const FInstanceSceneData InstanceSceneData = GetInstanceSceneData(GPUSceneInstanceId);
+
+	Payload.InstanceHash = MurmurAdd(InstanceSceneData.PrimitiveId, InstanceSceneData.RelativeId);
+
+	Payload.TriangleIndex = TraceRayResult.PrimitiveIndex;
+
+	Payload.WorldNormal = TraceRayResult.WorldGeometryNormal;
+	Payload.InstanceIndex = TraceRayResult.InstanceIndex;
+	Payload.GeometryIndex = TraceRayResult.GeometryIndex;
+	Payload.ScenePrimitiveIndex = InstanceSceneData.PrimitiveId;
+#else
 	TraceRay(
 		TLAS,
 		RayFlags  /*RayFlags*/,
@@ -229,6 +258,7 @@ float4 DebugRayTracingInstance(FRayDesc Ray, out float OutDepth)
 		0 /*MissShaderIndex*/,
 		Ray.GetNativeDesc(),
 		Payload);
+#endif
 
 	if (Payload.IsHit())
 	{
@@ -328,9 +358,9 @@ float4 DebugRayTracingInstance(FRayDesc Ray, out float OutDepth)
 #endif
 
 
-RAY_TRACING_ENTRY_RAYGEN(RayTracingDebugMainRGS)
+RAY_TRACING_ENTRY_RAYGEN_OR_INLINE(RayTracingDebugMain)
 {
-	uint2 PixelCoord = DispatchRaysIndex().xy + View.ViewRectMin.xy;
+	uint2 PixelCoord = DispatchThreadIndex.xy + View.ViewRectMin.xy;
 
 	float2 RenderTargetUV = (float2(PixelCoord) + .5f) * View.BufferSizeAndInvSize.zw;
 
@@ -346,7 +376,7 @@ RAY_TRACING_ENTRY_RAYGEN(RayTracingDebugMainRGS)
 
 	bool bWriteColorOuput = VisualizationMode != RAY_TRACING_DEBUG_VIZ_SUBSTRATE_DATA;
 	bool bWriteDepthOutput = VisualizationMode == RAY_TRACING_DEBUG_VIZ_INSTANCE_OVERLAP;
-#if USE_DEBUG_CHS
+#if USE_DEBUG_CHS || RAYTRACING_DEBUG_INLINE
 	{
 		Result = DebugRayTracingInstance(Ray, OutDepth);
 	}
diff --git a/Engine/Shaders/Private/RayTracing/RayTracingDebugPicking.usf b/Engine/Shaders/Private/RayTracing/RayTracingDebugPicking.usf
--- a/Engine/Shaders/Private/RayTracing/RayTracingDebugPicking.usf
+++ b/Engine/Shaders/Private/RayTracing/RayTracingDebugPicking.usf
@@ -16,7 +16,7 @@ StructuredBuffer<FRayTracingPickingFeedback> PickingBuffer;
 RWStructuredBuffer<FRayTracingPickingFeedback> PickingOutput;
 StructuredBuffer<FPlatformRayTracingInstanceDescriptor> InstanceBuffer;
 
-RAY_TRACING_ENTRY_RAYGEN(RayTracingDebugPickingRGS)
+RAY_TRACING_ENTRY_RAYGEN_OR_INLINE(RayTracingDebugPicking)
 {
 	const uint2 PixelPos = (float2(View.CursorPosition) * View.ViewResolutionFraction);
 	float2 RenderTargetUV = (float2(PixelPos) + .5f) * View.BufferSizeAndInvSize.zw;
@@ -28,6 +28,23 @@ RAY_TRACING_ENTRY_RAYGEN(RayTracingDebugPickingRGS)
 	FRayTracingDebugPayload Payload = (FRayTracingDebugPayload)0;
 	Payload.SetMiss();
 
+#if COMPUTESHADER
+	FTraceRayInlineContext TraceRayInlineContext = CreateTraceRayInlineContext();
+
+	FTraceRayInlineResult TraceRayResult = TraceRayInline(
+		TLAS,
+		RAY_FLAG_CULL_BACK_FACING_TRIANGLES,
+		InstanceInclusionMask,
+		Ray.GetNativeDesc(),
+		TraceRayInlineContext
+	);
+	
+	Payload.HitT = TraceRayResult.HitT;
+
+	Payload.TriangleIndex = TraceRayResult.PrimitiveIndex;
+	Payload.InstanceIndex = TraceRayResult.InstanceIndex;
+	Payload.GeometryIndex = TraceRayResult.GeometryIndex;
+#else
 	TraceRay(
 		TLAS,
 		RAY_FLAG_CULL_BACK_FACING_TRIANGLES  /*RayFlags*/,
@@ -37,6 +54,7 @@ RAY_TRACING_ENTRY_RAYGEN(RayTracingDebugPickingRGS)
 		0 /*MissShaderIndex*/,
 		Ray.GetNativeDesc(),
 		Payload);
+#endif
 
 	FRayTracingPickingFeedback PickingFeedback = (FRayTracingPickingFeedback)0xFFFFFFFFu;
 
diff --git a/Engine/Source/Runtime/Renderer/Private/RayTracing/RayTracingDebug.cpp b/Engine/Source/Runtime/Renderer/Private/RayTracing/RayTracingDebug.cpp
--- a/Engine/Source/Runtime/Renderer/Private/RayTracing/RayTracingDebug.cpp
+++ b/Engine/Source/Runtime/Renderer/Private/RayTracing/RayTracingDebug.cpp
@@ -161,6 +161,34 @@ END_UNIFORM_BUFFER_STRUCT()
 
 IMPLEMENT_UNIFORM_BUFFER_STRUCT(FRayTracingDebugHitStatsUniformBufferParameters, "RayTracingDebugHitStatsUniformBuffer");
 
+BEGIN_SHADER_PARAMETER_STRUCT(FRayTracingDebugParameters, )
+	SHADER_PARAMETER(uint32, VisualizationMode)
+	SHADER_PARAMETER(uint32, PickerDomain)
+	SHADER_PARAMETER(uint32, ShouldUsePreExposure)
+	SHADER_PARAMETER(float, TimingScale)
+	SHADER_PARAMETER(float, MaxTraceDistance)
+	SHADER_PARAMETER(float, FarFieldMaxTraceDistance)
+	SHADER_PARAMETER(uint32, OpaqueOnly)
+	SHADER_PARAMETER(float, TriangleHitCountMaxThreshold)
+	SHADER_PARAMETER(float, TriangleHitCountPerInstanceMaxThreshold)
+	SHADER_PARAMETER(uint32, TopKMostHitInstances)
+	SHADER_PARAMETER(uint32, NumTotalInstances)
+	SHADER_PARAMETER_RDG_BUFFER_SRV(RaytracingAccelerationStructure, TLAS)
+	SHADER_PARAMETER_RDG_BUFFER_SRV(RaytracingAccelerationStructure, FarFieldTLAS)
+	SHADER_PARAMETER_RDG_TEXTURE_UAV(RWTexture2D<float4>, Output)
+	SHADER_PARAMETER_RDG_TEXTURE_UAV(RWTexture2D<float>, OutputDepth)
+	SHADER_PARAMETER_RDG_BUFFER_SRV(StructuredBuffer, InstancesExtraData)
+	SHADER_PARAMETER_RDG_BUFFER_SRV(StructuredBuffer, InstancesDebugData)
+	SHADER_PARAMETER_RDG_BUFFER_SRV(StructuredBuffer, InstanceBuffer)
+	SHADER_PARAMETER_RDG_BUFFER_SRV(StructuredBuffer<FRayTracingPickingFeedback>, PickingBuffer)
+	SHADER_PARAMETER_RDG_BUFFER_SRV(StructuredBuffer<uint2>, TopKHitStats)
+	SHADER_PARAMETER_STRUCT_REF(FViewUniformShaderParameters, ViewUniformBuffer)
+	SHADER_PARAMETER_RDG_UNIFORM_BUFFER(FSceneUniformParameters, SceneUniformBuffer)
+	SHADER_PARAMETER_RDG_UNIFORM_BUFFER(FRayTracingDebugHitStatsUniformBufferParameters, RayTracingDebugHitStatsUniformBuffer)
+	SHADER_PARAMETER_RDG_UNIFORM_BUFFER(FRayTracingLightGrid, LightGridPacked)
+	SHADER_PARAMETER_STRUCT_INCLUDE(ShaderPrint::FShaderParameters, ShaderPrintUniformBuffer)
+	SHADER_PARAMETER_RDG_BUFFER_UAV(RWBuffer<int>, SubstrateDebugDataUAV)
+END_SHADER_PARAMETER_STRUCT()
 
 class FRayTracingDebugRGS : public FGlobalShader
 {
@@ -172,32 +200,7 @@ class FRayTracingDebugRGS : public FGlobalShader
 	using FPermutationDomain = TShaderPermutationDomain<FUseDebugCHSType, FUseNvAPITimestamp>;
 
 	BEGIN_SHADER_PARAMETER_STRUCT(FParameters, )
-		SHADER_PARAMETER(uint32, VisualizationMode)
-		SHADER_PARAMETER(uint32, PickerDomain)
-		SHADER_PARAMETER(uint32, ShouldUsePreExposure)
-		SHADER_PARAMETER(float, TimingScale)
-		SHADER_PARAMETER(float, MaxTraceDistance)
-		SHADER_PARAMETER(float, FarFieldMaxTraceDistance)
-		SHADER_PARAMETER(uint32, OpaqueOnly)
-		SHADER_PARAMETER(float, TriangleHitCountMaxThreshold)
-		SHADER_PARAMETER(float, TriangleHitCountPerInstanceMaxThreshold)
-		SHADER_PARAMETER(uint32, TopKMostHitInstances)
-		SHADER_PARAMETER(uint32, NumTotalInstances)
-		SHADER_PARAMETER_RDG_BUFFER_SRV(RaytracingAccelerationStructure, TLAS)
-		SHADER_PARAMETER_RDG_BUFFER_SRV(RaytracingAccelerationStructure, FarFieldTLAS)
-		SHADER_PARAMETER_RDG_TEXTURE_UAV(RWTexture2D<float4>, Output)
-		SHADER_PARAMETER_RDG_TEXTURE_UAV(RWTexture2D<float>, OutputDepth)
-		SHADER_PARAMETER_RDG_BUFFER_SRV(StructuredBuffer, InstancesExtraData)
-		SHADER_PARAMETER_RDG_BUFFER_SRV(StructuredBuffer, InstancesDebugData)
-		SHADER_PARAMETER_RDG_BUFFER_SRV(StructuredBuffer, InstanceBuffer)
-		SHADER_PARAMETER_RDG_BUFFER_SRV(StructuredBuffer<FRayTracingPickingFeedback>, PickingBuffer)
-		SHADER_PARAMETER_RDG_BUFFER_SRV(StructuredBuffer<uint2>, TopKHitStats)
-		SHADER_PARAMETER_STRUCT_REF(FViewUniformShaderParameters, ViewUniformBuffer)
-		SHADER_PARAMETER_RDG_UNIFORM_BUFFER(FSceneUniformParameters, SceneUniformBuffer)
-		SHADER_PARAMETER_RDG_UNIFORM_BUFFER(FRayTracingDebugHitStatsUniformBufferParameters, RayTracingDebugHitStatsUniformBuffer)
-		SHADER_PARAMETER_RDG_UNIFORM_BUFFER(FRayTracingLightGrid, LightGridPacked)
-		SHADER_PARAMETER_STRUCT_INCLUDE(ShaderPrint::FShaderParameters, ShaderPrintUniformBuffer)
-		SHADER_PARAMETER_RDG_BUFFER_UAV(RWBuffer<int>, SubstrateDebugDataUAV)
+		SHADER_PARAMETER_STRUCT_INCLUDE(FRayTracingDebugParameters, SharedParameters)
 	END_SHADER_PARAMETER_STRUCT()
 
 	static bool ShouldCompilePermutation(const FGlobalShaderPermutationParameters& Parameters)
@@ -239,6 +242,38 @@ class FRayTracingDebugRGS : public FGlobalShader
 };
 IMPLEMENT_GLOBAL_SHADER(FRayTracingDebugRGS, "/Engine/Private/RayTracing/RayTracingDebug.usf", "RayTracingDebugMainRGS", SF_RayGen);
 
+class FRayTracingDebugCS : public FGlobalShader
+{
+	DECLARE_GLOBAL_SHADER(FRayTracingDebugCS)
+	SHADER_USE_PARAMETER_STRUCT(FRayTracingDebugCS, FGlobalShader)
+
+	static bool ShouldCompilePermutation(const FGlobalShaderPermutationParameters& Parameters)
+	{
+		return IsRayTracingEnabledForProject(Parameters.Platform) && FDataDrivenShaderPlatformInfo::GetSupportsInlineRayTracing(Parameters.Platform);
+	}
+
+	static void ModifyCompilationEnvironment(const FGlobalShaderPermutationParameters& Parameters, FShaderCompilerEnvironment& OutEnvironment)
+	{
+		FGlobalShader::ModifyCompilationEnvironment(Parameters, OutEnvironment);
+
+		OutEnvironment.CompilerFlags.Add(CFLAG_Wave32);
+		OutEnvironment.CompilerFlags.Add(CFLAG_InlineRayTracing);
+
+		OutEnvironment.SetDefine(TEXT("INLINE_RAY_TRACING_THREAD_GROUP_SIZE_X"), ThreadGroupSizeX);
+		OutEnvironment.SetDefine(TEXT("INLINE_RAY_TRACING_THREAD_GROUP_SIZE_Y"), ThreadGroupSizeY);
+	}
+
+	BEGIN_SHADER_PARAMETER_STRUCT(FParameters, )
+		SHADER_PARAMETER_STRUCT_INCLUDE(FRayTracingDebugParameters, SharedParameters)
+	END_SHADER_PARAMETER_STRUCT()
+
+	// Current inline ray tracing implementation requires 1:1 mapping between thread groups and waves and only supports wave32 mode.
+	static constexpr uint32 ThreadGroupSizeX = 8;
+	static constexpr uint32 ThreadGroupSizeY = 4;
+};
+
+IMPLEMENT_GLOBAL_SHADER(FRayTracingDebugCS, "/Engine/Private/RayTracing/RayTracingDebug.usf", "RayTracingDebugMainCS", SF_Compute);
+
 class FRayTracingDebugCHS : public FGlobalShader
 {
 	DECLARE_GLOBAL_SHADER(FRayTracingDebugCHS);
@@ -450,20 +485,24 @@ class FRayTracingDebugTraversalCS : public FGlobalShader
 };
 IMPLEMENT_GLOBAL_SHADER(FRayTracingDebugTraversalCS, "/Engine/Private/RayTracing/RayTracingDebugTraversal.usf", "RayTracingDebugTraversalCS", SF_Compute);
 
+BEGIN_SHADER_PARAMETER_STRUCT(FRayTracingPickingParameters, )
+	SHADER_PARAMETER_RDG_BUFFER_SRV(RaytracingAccelerationStructure, TLAS)
+	SHADER_PARAMETER(int32, OpaqueOnly)
+	SHADER_PARAMETER_RDG_BUFFER_UAV(RWStructuredBuffer, PickingOutput)
+	SHADER_PARAMETER_RDG_BUFFER_SRV(StructuredBuffer, InstancesExtraData)
+	SHADER_PARAMETER_RDG_BUFFER_SRV(StructuredBuffer, InstancesDebugData)
+	SHADER_PARAMETER_RDG_BUFFER_SRV(StructuredBuffer, InstanceBuffer)
+	SHADER_PARAMETER_STRUCT_REF(FViewUniformShaderParameters, ViewUniformBuffer)
+	SHADER_PARAMETER_RDG_UNIFORM_BUFFER(FSceneUniformParameters, SceneUniformBuffer)
+END_SHADER_PARAMETER_STRUCT()
+
 class FRayTracingPickingRGS : public FGlobalShader
 {
 	DECLARE_GLOBAL_SHADER(FRayTracingPickingRGS)
 	SHADER_USE_ROOT_PARAMETER_STRUCT(FRayTracingPickingRGS, FGlobalShader)
 
 	BEGIN_SHADER_PARAMETER_STRUCT(FParameters, )
-		SHADER_PARAMETER_RDG_BUFFER_SRV(RaytracingAccelerationStructure, TLAS)
-		SHADER_PARAMETER(int32, OpaqueOnly)
-		SHADER_PARAMETER_RDG_BUFFER_UAV(RWStructuredBuffer, PickingOutput)
-		SHADER_PARAMETER_RDG_BUFFER_SRV(StructuredBuffer, InstancesExtraData)
-		SHADER_PARAMETER_RDG_BUFFER_SRV(StructuredBuffer, InstancesDebugData)
-		SHADER_PARAMETER_RDG_BUFFER_SRV(StructuredBuffer, InstanceBuffer)
-		SHADER_PARAMETER_STRUCT_REF(FViewUniformShaderParameters, ViewUniformBuffer)
-		SHADER_PARAMETER_RDG_UNIFORM_BUFFER(FSceneUniformParameters, SceneUniformBuffer)
+		SHADER_PARAMETER_STRUCT_INCLUDE(FRayTracingPickingParameters, SharedParameters)
 	END_SHADER_PARAMETER_STRUCT()
 
 	static bool ShouldCompilePermutation(const FGlobalShaderPermutationParameters& Parameters)
@@ -483,6 +522,36 @@ class FRayTracingPickingRGS : public FGlobalShader
 };
 IMPLEMENT_GLOBAL_SHADER(FRayTracingPickingRGS, "/Engine/Private/RayTracing/RayTracingDebugPicking.usf", "RayTracingDebugPickingRGS", SF_RayGen);
 
+class FRayTracingPickingCS : public FGlobalShader
+{
+	DECLARE_GLOBAL_SHADER(FRayTracingPickingCS)
+	SHADER_USE_PARAMETER_STRUCT(FRayTracingPickingCS, FGlobalShader)
+
+	static bool ShouldCompilePermutation(const FGlobalShaderPermutationParameters& Parameters)
+	{
+		return IsRayTracingEnabledForProject(Parameters.Platform) && FDataDrivenShaderPlatformInfo::GetSupportsInlineRayTracing(Parameters.Platform);
+	}
+
+	static void ModifyCompilationEnvironment(const FGlobalShaderPermutationParameters& Parameters, FShaderCompilerEnvironment& OutEnvironment)
+	{
+		FGlobalShader::ModifyCompilationEnvironment(Parameters, OutEnvironment);
+
+		OutEnvironment.CompilerFlags.Add(CFLAG_Wave32);
+		OutEnvironment.CompilerFlags.Add(CFLAG_InlineRayTracing);
+
+		OutEnvironment.SetDefine(TEXT("INLINE_RAY_TRACING_THREAD_GROUP_SIZE_X"), ThreadGroupSizeX);
+		OutEnvironment.SetDefine(TEXT("INLINE_RAY_TRACING_THREAD_GROUP_SIZE_Y"), ThreadGroupSizeY);
+	}
+
+	BEGIN_SHADER_PARAMETER_STRUCT(FParameters, )
+		SHADER_PARAMETER_STRUCT_INCLUDE(FRayTracingPickingParameters, SharedParameters)
+	END_SHADER_PARAMETER_STRUCT()
+	
+	static constexpr uint32 ThreadGroupSizeX = 1;
+	static constexpr uint32 ThreadGroupSizeY = 1;
+};
+IMPLEMENT_GLOBAL_SHADER(FRayTracingPickingCS, "/Engine/Private/RayTracing/RayTracingDebugPicking.usf", "RayTracingDebugPickingCS", SF_Compute);
+
 class FRayTracingDebugInstanceOverlapVS : public FGlobalShader
 {
 	DECLARE_GLOBAL_SHADER(FRayTracingDebugInstanceOverlapVS);
@@ -851,6 +920,13 @@ static bool RequiresRayTracingDebugCHS(uint32 DebugVisualizationMode)
 		|| DebugVisualizationMode == RAY_TRACING_DEBUG_VIZ_HITCOUNT_PER_INSTANCE;
 }
 
+static bool RayTracingDebugModeSupportsInline(uint32 DebugVisualizationMode)
+{
+	return RequiresRayTracingDebugCHS(DebugVisualizationMode)
+		&& (DebugVisualizationMode != RAY_TRACING_DEBUG_VIZ_TRIANGLE_HITCOUNT)
+		&& (DebugVisualizationMode != RAY_TRACING_DEBUG_VIZ_HITCOUNT_PER_INSTANCE);
+}
+
 static bool IsRayTracingDebugTraversalMode(uint32 DebugVisualizationMode)
 {
 	return DebugVisualizationMode == RAY_TRACING_DEBUG_VIZ_TRAVERSAL_PRIMARY_NODE 
@@ -886,76 +962,99 @@ void FDeferredShadingSceneRenderer::PrepareRayTracingDebug(const FSceneViewFamil
 	}
 }
 
-static FRDGBufferRef RayTracingPerformPicking(FRDGBuilder& GraphBuilder, const FScene& Scene, const FViewInfo& View, FRayTracingPickingFeedback& PickingFeedback)
+static FRDGBufferRef RayTracingPerformPicking(FRDGBuilder& GraphBuilder, const FScene& Scene, const FViewInfo& View, FRayTracingPickingFeedback& PickingFeedback, bool bInlineRayTracing)
 {
 	const FRayTracingScene& RayTracingScene = Scene.RayTracingScene;
 
-	FGlobalShaderMap* ShaderMap = GetGlobalShaderMap(GMaxRHIFeatureLevel);
-	auto RayGenShader = ShaderMap->GetShader<FRayTracingPickingRGS>();
+	FRDGBufferDesc PickingBufferDesc = FRDGBufferDesc::CreateStructuredDesc(sizeof(FRayTracingPickingFeedback), 1);
+	PickingBufferDesc.Usage = EBufferUsageFlags(PickingBufferDesc.Usage | BUF_SourceCopy);
+	FRDGBufferRef PickingBuffer = GraphBuilder.CreateBuffer(PickingBufferDesc, TEXT("RayTracingDebug.PickingBuffer"));
 
-	FRayTracingPipelineStateInitializer Initializer;
-	Initializer.MaxPayloadSizeInBytes = GetRayTracingPayloadTypeMaxSize(ERayTracingPayloadType::RayTracingDebug);
+	FRayTracingPickingParameters SharedParameters = {};
+	SharedParameters.InstancesExtraData = GraphBuilder.CreateSRV(RayTracingScene.GetInstanceExtraDataBuffer(ERayTracingSceneLayer::Base));
+	SharedParameters.InstancesDebugData = GraphBuilder.CreateSRV(RayTracingScene.GetInstanceDebugBuffer(ERayTracingSceneLayer::Base));
+	SharedParameters.TLAS = RayTracingScene.GetLayerView(ERayTracingSceneLayer::Base);
+	SharedParameters.OpaqueOnly = CVarRayTracingVisualizeOpaqueOnly.GetValueOnRenderThread();
+	SharedParameters.InstanceBuffer = GraphBuilder.CreateSRV(RayTracingScene.GetInstanceBuffer(ERayTracingSceneLayer::Base));
+	SharedParameters.ViewUniformBuffer = View.ViewUniformBuffer;
+	SharedParameters.SceneUniformBuffer = GetSceneUniformBufferRef(GraphBuilder, View); // TODO: use a separate params structure
+	SharedParameters.PickingOutput = GraphBuilder.CreateUAV(PickingBuffer);
 
-	const FShaderBindingLayout* ShaderBindingLayout = RayTracing::GetShaderBindingLayout(Scene.GetShaderPlatform());
-	if (ShaderBindingLayout)
+	if (bInlineRayTracing)
 	{
-		Initializer.ShaderBindingLayout = &ShaderBindingLayout->RHILayout;
+		FRayTracingPickingCS::FParameters* InlinePassParameters = GraphBuilder.AllocParameters<FRayTracingPickingCS::FParameters>();
+
+		InlinePassParameters->SharedParameters = SharedParameters;
+
+		TShaderRef<FRayTracingPickingCS> ComputeShader = View.ShaderMap->GetShader<FRayTracingPickingCS>();
+
+		GraphBuilder.AddPass(
+			RDG_EVENT_NAME("RayTracingPicking (INLINE)"),
+			InlinePassParameters,
+			ERDGPassFlags::Compute,
+			[InlinePassParameters, ComputeShader](FRDGAsyncTask, FRHIComputeCommandList& RHICmdList)
+			{
+				FComputeShaderUtils::Dispatch(RHICmdList, ComputeShader, *InlinePassParameters, FIntVector(1,1,1));
+			});
 	}
+	else
+	{
+		FRayTracingPickingRGS::FParameters* RayGenParameters = GraphBuilder.AllocParameters<FRayTracingPickingRGS::FParameters>();
+		RayGenParameters->SharedParameters = SharedParameters;
 
-	FRHIRayTracingShader* RayGenShaderTable[] = { RayGenShader.GetRayTracingShader() };
-	Initializer.SetRayGenShaderTable(RayGenShaderTable);
+		FGlobalShaderMap* ShaderMap = GetGlobalShaderMap(GMaxRHIFeatureLevel);
+		auto RayGenShader = ShaderMap->GetShader<FRayTracingPickingRGS>();
 
-	FRayTracingDebugCHS::FPermutationDomain PermutationVector;
+		FRayTracingPipelineStateInitializer Initializer;
+		Initializer.MaxPayloadSizeInBytes = GetRayTracingPayloadTypeMaxSize(ERayTracingPayloadType::RayTracingDebug);
 
-	PermutationVector.Set<FRayTracingDebugCHS::FNaniteRayTracing>(false);
-	auto HitGroupShader = View.ShaderMap->GetShader<FRayTracingDebugCHS>(PermutationVector);
+		const FShaderBindingLayout* ShaderBindingLayout = RayTracing::GetShaderBindingLayout(Scene.GetShaderPlatform());
+		if (ShaderBindingLayout)
+		{
+			Initializer.ShaderBindingLayout = &ShaderBindingLayout->RHILayout;
+		}
 
-	PermutationVector.Set<FRayTracingDebugCHS::FNaniteRayTracing>(true);
-	auto HitGroupShaderNaniteRT = View.ShaderMap->GetShader<FRayTracingDebugCHS>(PermutationVector);
+		FRHIRayTracingShader* RayGenShaderTable[] = { RayGenShader.GetRayTracingShader() };
+		Initializer.SetRayGenShaderTable(RayGenShaderTable);
 
-	FRHIRayTracingShader* HitGroupTable[] = { HitGroupShader.GetRayTracingShader(), HitGroupShaderNaniteRT.GetRayTracingShader() };
-	Initializer.SetHitGroupTable(HitGroupTable);
+		FRayTracingDebugCHS::FPermutationDomain PermutationVector;
 
-	auto MissShader = ShaderMap->GetShader<FRayTracingDebugMS>();
-	FRHIRayTracingShader* MissTable[] = { MissShader.GetRayTracingShader() };
-	Initializer.SetMissShaderTable(MissTable);
+		PermutationVector.Set<FRayTracingDebugCHS::FNaniteRayTracing>(false);
+		auto HitGroupShader = View.ShaderMap->GetShader<FRayTracingDebugCHS>(PermutationVector);
 
-	FRayTracingPipelineState* PickingPipeline = PipelineStateCache::GetAndOrCreateRayTracingPipelineState(GraphBuilder.RHICmdList, Initializer);
-		
-	FShaderBindingTableRHIRef PickingSBT = Scene.RayTracingSBT.AllocateTransientRHI(GraphBuilder.RHICmdList, ERayTracingShaderBindingMode::RTPSO, ERayTracingHitGroupIndexingMode::Allow, Initializer.GetMaxLocalBindingDataSize());
-	   
-	FRDGBufferDesc PickingBufferDesc = FRDGBufferDesc::CreateStructuredDesc(sizeof(FRayTracingPickingFeedback), 1);
-	PickingBufferDesc.Usage = EBufferUsageFlags(PickingBufferDesc.Usage | BUF_SourceCopy);
-	FRDGBufferRef PickingBuffer = GraphBuilder.CreateBuffer(PickingBufferDesc, TEXT("RayTracingDebug.PickingBuffer"));
+		PermutationVector.Set<FRayTracingDebugCHS::FNaniteRayTracing>(true);
+		auto HitGroupShaderNaniteRT = View.ShaderMap->GetShader<FRayTracingDebugCHS>(PermutationVector);
 
-	FRayTracingPickingRGS::FParameters* RayGenParameters = GraphBuilder.AllocParameters<FRayTracingPickingRGS::FParameters>();
-	RayGenParameters->InstancesExtraData = GraphBuilder.CreateSRV(RayTracingScene.GetInstanceExtraDataBuffer(ERayTracingSceneLayer::Base));
-	RayGenParameters->InstancesDebugData = GraphBuilder.CreateSRV(RayTracingScene.GetInstanceDebugBuffer(ERayTracingSceneLayer::Base));
-	RayGenParameters->TLAS = RayTracingScene.GetLayerView(ERayTracingSceneLayer::Base);
-	RayGenParameters->OpaqueOnly = CVarRayTracingVisualizeOpaqueOnly.GetValueOnRenderThread();
-	RayGenParameters->InstanceBuffer = GraphBuilder.CreateSRV(RayTracingScene.GetInstanceBuffer(ERayTracingSceneLayer::Base));
-	RayGenParameters->ViewUniformBuffer = View.ViewUniformBuffer;
-	RayGenParameters->SceneUniformBuffer = GetSceneUniformBufferRef(GraphBuilder, View); // TODO: use a separate params structure
-	RayGenParameters->PickingOutput = GraphBuilder.CreateUAV(PickingBuffer);
-	
-	GraphBuilder.AddPass(
-		RDG_EVENT_NAME("RayTracingPicking"),
-		RayGenParameters,
-		ERDGPassFlags::Compute,
-		[RayGenParameters, RayGenShader, &View, PickingSBT, PickingPipeline](FRDGAsyncTask, FRHICommandList& RHICmdList)
-		{
-			FRHIBatchedShaderParameters& GlobalResources = RHICmdList.GetScratchShaderParameters();
-			SetShaderParameters(GlobalResources, RayGenShader, *RayGenParameters);
+		FRHIRayTracingShader* HitGroupTable[] = { HitGroupShader.GetRayTracingShader(), HitGroupShaderNaniteRT.GetRayTracingShader() };
+		Initializer.SetHitGroupTable(HitGroupTable);
 
-			FRHIUniformBuffer* SceneUniformBuffer = RayGenParameters->SceneUniformBuffer.GetUniformBuffer()->GetRHI();
-			TOptional<FScopedUniformBufferStaticBindings> StaticUniformBufferScope = RayTracing::BindStaticUniformBufferBindings(View, SceneUniformBuffer, RHICmdList);
+		auto MissShader = ShaderMap->GetShader<FRayTracingDebugMS>();
+		FRHIRayTracingShader* MissTable[] = { MissShader.GetRayTracingShader() };
+		Initializer.SetMissShaderTable(MissTable);
 
-			BindRayTracingDebugCHSMaterialBindings(RHICmdList, PickingSBT, View, RayGenParameters->SceneUniformBuffer->GetRHI(), PickingPipeline);
-			RHICmdList.SetRayTracingMissShader(PickingSBT, 0, PickingPipeline, 0 /* ShaderIndexInPipeline */, 0, nullptr, 0);
-			RHICmdList.CommitShaderBindingTable(PickingSBT);
+		FRayTracingPipelineState* PickingPipeline = PipelineStateCache::GetAndOrCreateRayTracingPipelineState(GraphBuilder.RHICmdList, Initializer);
 
-			RHICmdList.RayTraceDispatch(PickingPipeline, RayGenShader.GetRayTracingShader(), PickingSBT, GlobalResources, 1, 1);
-		});
+		FShaderBindingTableRHIRef PickingSBT = Scene.RayTracingSBT.AllocateTransientRHI(GraphBuilder.RHICmdList, ERayTracingShaderBindingMode::RTPSO, ERayTracingHitGroupIndexingMode::Allow, Initializer.GetMaxLocalBindingDataSize());
+
+		GraphBuilder.AddPass(
+			RDG_EVENT_NAME("RayTracingPicking"),
+			RayGenParameters,
+			ERDGPassFlags::Compute,
+			[RayGenParameters, RayGenShader, &View, PickingSBT, PickingPipeline](FRDGAsyncTask, FRHICommandList& RHICmdList)
+			{
+				FRHIBatchedShaderParameters& GlobalResources = RHICmdList.GetScratchShaderParameters();
+				SetShaderParameters(GlobalResources, RayGenShader, *RayGenParameters);
+
+				FRHIUniformBuffer* SceneUniformBuffer = RayGenParameters->SharedParameters.SceneUniformBuffer.GetUniformBuffer()->GetRHI();
+				TOptional<FScopedUniformBufferStaticBindings> StaticUniformBufferScope = RayTracing::BindStaticUniformBufferBindings(View, SceneUniformBuffer, RHICmdList);
+
+				BindRayTracingDebugCHSMaterialBindings(RHICmdList, PickingSBT, View, RayGenParameters->SharedParameters.SceneUniformBuffer->GetRHI(), PickingPipeline);
+				RHICmdList.SetRayTracingMissShader(PickingSBT, 0, PickingPipeline, 0 /* ShaderIndexInPipeline */, 0, nullptr, 0);
+				RHICmdList.CommitShaderBindingTable(PickingSBT);
+
+				RHICmdList.RayTraceDispatch(PickingPipeline, RayGenShader.GetRayTracingShader(), PickingSBT, GlobalResources, 1, 1);
+			});
+	}	
 
 	const int32 MaxPickingBuffers = GRayTracingDebugResources.MaxPickingBuffers;
 
@@ -1554,15 +1653,15 @@ void RenderRayTracingDebug(FRDGBuilder& GraphBuilder, const FScene& Scene, const
 		return;
 	}
 
-	// Debug modes other than barycentrics and traversal require pipeline support.
+	const bool bInlineRayTracing = ShouldRenderRayTracingEffect(true, ERayTracingPipelineCompatibilityFlags::Inline, View);
 	const bool bRayTracingPipeline = ShouldRenderRayTracingEffect(true, ERayTracingPipelineCompatibilityFlags::FullPipeline, View);
-	if (!bRayTracingPipeline)
+	if (!bRayTracingPipeline && !(bInlineRayTracing && RayTracingDebugModeSupportsInline(DebugVisualizationMode)))
 	{
 		AddClearUAVPass(GraphBuilder, GraphBuilder.CreateUAV(SceneColorTexture), FLinearColor::Black);
 		return;
 	}
 
-	if (DebugVisualizationMode == RAY_TRACING_DEBUG_VIZ_PRIMARY_RAYS) 
+	if (DebugVisualizationMode == RAY_TRACING_DEBUG_VIZ_PRIMARY_RAYS)
 	{
 		FRDGTextureRef OutputColor = nullptr;
 		FRDGTextureRef HitDistanceTexture = nullptr;
@@ -1581,7 +1680,7 @@ void RenderRayTracingDebug(FRDGBuilder& GraphBuilder, const FScene& Scene, const
 		&& RayTracingScene.GetInstanceExtraDataBuffer(ERayTracingSceneLayer::Base) != nullptr
 		&& RayTracingScene.GetInstanceDebugBuffer(ERayTracingSceneLayer::Base) != nullptr)
 	{
-		PickingBuffer = RayTracingPerformPicking(GraphBuilder, Scene, View, PickingFeedback);
+		PickingBuffer = RayTracingPerformPicking(GraphBuilder, Scene, View, PickingFeedback, bInlineRayTracing);
 	}
 	else
 	{
@@ -1623,134 +1722,163 @@ void RenderRayTracingDebug(FRDGBuilder& GraphBuilder, const FScene& Scene, const
 
 	const bool bRequiresDebugCHS = RequiresRayTracingDebugCHS(DebugVisualizationMode);
 
-	FRayTracingDebugRGS::FPermutationDomain PermutationVector;
-	PermutationVector.Set<FRayTracingDebugRGS::FUseDebugCHSType>(bRequiresDebugCHS);
-	PermutationVector.Set<FRayTracingDebugRGS::FUseNvAPITimestamp>(GRHIGlobals.SupportsShaderTimestamp && IsRHIDeviceNVIDIA());
-
-	auto RayGenShader = View.ShaderMap->GetShader<FRayTracingDebugRGS>(PermutationVector);
-
-	FRayTracingPipelineState* Pipeline = View.MaterialRayTracingData.PipelineState;
-	FShaderBindingTableRHIRef SBT = View.MaterialRayTracingData.ShaderBindingTable;
-	bool bRequiresBindings = false;
-
-	if (bRequiresDebugCHS)
-	{
-		FRayTracingPipelineStateInitializer Initializer;
-		Initializer.MaxPayloadSizeInBytes = GetRayTracingPayloadTypeMaxSize(ERayTracingPayloadType::RayTracingDebug);
-		
-		const FShaderBindingLayout* ShaderBindingLayout = RayTracing::GetShaderBindingLayout(ShaderPlatform);
-		if (ShaderBindingLayout)
-		{
-			Initializer.ShaderBindingLayout = &ShaderBindingLayout->RHILayout;
-		}
-
-		FRHIRayTracingShader* RayGenShaderTable[] = { RayGenShader.GetRayTracingShader() };
-		Initializer.SetRayGenShaderTable(RayGenShaderTable);
-
-		FRayTracingDebugCHS::FPermutationDomain PermutationVectorCHS;
-
-		PermutationVectorCHS.Set<FRayTracingDebugCHS::FNaniteRayTracing>(false);
-		auto HitGroupShader = View.ShaderMap->GetShader<FRayTracingDebugCHS>(PermutationVectorCHS);
-
-		PermutationVectorCHS.Set<FRayTracingDebugCHS::FNaniteRayTracing>(true);
-		auto HitGroupShaderNaniteRT = View.ShaderMap->GetShader<FRayTracingDebugCHS>(PermutationVectorCHS);
-
-		// auto AnyHitShader = ShaderMap->GetShader<FRayTracingDebugAHS>();
-
-		FRHIRayTracingShader* HitGroupTable[] = { HitGroupShader.GetRayTracingShader(), HitGroupShaderNaniteRT.GetRayTracingShader()};
-		Initializer.SetHitGroupTable(HitGroupTable);
-
-		auto MissShader = View.ShaderMap->GetShader<FRayTracingDebugMS>();
-		FRHIRayTracingShader* MissTable[] = { MissShader.GetRayTracingShader() };
-		Initializer.SetMissShaderTable(MissTable);
-
-		Pipeline = PipelineStateCache::GetAndOrCreateRayTracingPipelineState(GraphBuilder.RHICmdList, Initializer);
-
-		SBT = Scene.RayTracingSBT.AllocateTransientRHI(GraphBuilder.RHICmdList, ERayTracingShaderBindingMode::RTPSO, ERayTracingHitGroupIndexingMode::Allow, Initializer.GetMaxLocalBindingDataSize());
-
-		bRequiresBindings = true;
-	}
-
 	const uint32 NumInstances = RayTracingScene.GetNumNativeInstances(ERayTracingSceneLayer::Base);
 
-	FRayTracingDebugRGS::FParameters* RayGenParameters = GraphBuilder.AllocParameters<FRayTracingDebugRGS::FParameters>();
-	RayGenParameters->VisualizationMode = DebugVisualizationMode;
-	RayGenParameters->PickerDomain = CVarRayTracingVisualizePickerDomain.GetValueOnRenderThread();
-	RayGenParameters->ShouldUsePreExposure = View.Family->EngineShowFlags.Tonemapper;
-	RayGenParameters->TimingScale = CVarRayTracingVisualizeTimingScale.GetValueOnAnyThread() / 25000.0f;
-	RayGenParameters->OpaqueOnly = CVarRayTracingVisualizeOpaqueOnly.GetValueOnRenderThread();
-	RayGenParameters->TriangleHitCountMaxThreshold = FMath::Clamp((float)CVarRayTracingVisualizeHitCountMaxThreshold.GetValueOnRenderThread(), 1, 100000);
-	RayGenParameters->TriangleHitCountPerInstanceMaxThreshold = FMath::Max(1, CVarRayTracingVisualizeHitCountPerInstanceMaxThreshold.GetValueOnRenderThread());
-	RayGenParameters->RayTracingDebugHitStatsUniformBuffer = DebugHitStatsUniformBuffer;
-	RayGenParameters->LightGridPacked = View.RayTracingLightGridUniformBuffer;
-	RayGenParameters->TopKMostHitInstances = CVarRayTracingVisualizeHitCountTopKHits.GetValueOnRenderThread();
-	RayGenParameters->NumTotalInstances = NumInstances;
-	ShaderPrint::SetParameters(GraphBuilder, View.ShaderPrintData, RayGenParameters->ShaderPrintUniformBuffer);
+	FRayTracingDebugParameters SharedParameters;
+	SharedParameters.VisualizationMode = DebugVisualizationMode;
+	SharedParameters.PickerDomain = CVarRayTracingVisualizePickerDomain.GetValueOnRenderThread();
+	SharedParameters.ShouldUsePreExposure = View.Family->EngineShowFlags.Tonemapper;
+	SharedParameters.TimingScale = CVarRayTracingVisualizeTimingScale.GetValueOnAnyThread() / 25000.0f;
+	SharedParameters.OpaqueOnly = CVarRayTracingVisualizeOpaqueOnly.GetValueOnRenderThread();
+	SharedParameters.TriangleHitCountMaxThreshold = FMath::Clamp((float)CVarRayTracingVisualizeHitCountMaxThreshold.GetValueOnRenderThread(), 1, 100000);
+	SharedParameters.TriangleHitCountPerInstanceMaxThreshold = FMath::Max(1, CVarRayTracingVisualizeHitCountPerInstanceMaxThreshold.GetValueOnRenderThread());
+	SharedParameters.RayTracingDebugHitStatsUniformBuffer = DebugHitStatsUniformBuffer;
+	SharedParameters.LightGridPacked = View.RayTracingLightGridUniformBuffer;
+	SharedParameters.TopKMostHitInstances = CVarRayTracingVisualizeHitCountTopKHits.GetValueOnRenderThread();
+	SharedParameters.NumTotalInstances = NumInstances;
+	ShaderPrint::SetParameters(GraphBuilder, View.ShaderPrintData, SharedParameters.ShaderPrintUniformBuffer);
 
 	FSubstrateViewDebugData& SubstrateViewDebugData = View.ViewState->GetSubstrateViewDebugData();
 	FSubstrateViewDebugData::FTransientPixelDebugBuffer NewSubstratePixelDebugBuffer = SubstrateViewDebugData.CreateTransientPixelDebugBuffer(GraphBuilder);
-	RayGenParameters->SubstrateDebugDataUAV = NewSubstratePixelDebugBuffer.PixelMaterialDebugDataUAV;
-	
+	SharedParameters.SubstrateDebugDataUAV = NewSubstratePixelDebugBuffer.PixelMaterialDebugDataUAV;
+
 	// If we don't output depth, create dummy 1x1 texture
 	const bool bOutputDepth = DebugVisualizationMode == RAY_TRACING_DEBUG_VIZ_INSTANCE_OVERLAP;
 
 	FRDGTextureDesc OutputDepthTextureDesc = FRDGTextureDesc::Create2D(
 		bOutputDepth ? SceneColorTexture->Desc.Extent : FIntPoint(1, 1),
-		PF_R32_FLOAT, 
-		FClearValueBinding::Black, 
+		PF_R32_FLOAT,
+		FClearValueBinding::Black,
 		TexCreate_ShaderResource | TexCreate_UAV);
 	FRDGTextureRef OutputDepthTexture = GraphBuilder.CreateTexture(OutputDepthTextureDesc, TEXT("RayTracingDebug::Depth"));
-	RayGenParameters->OutputDepth = GraphBuilder.CreateUAV(OutputDepthTexture);
+	SharedParameters.OutputDepth = GraphBuilder.CreateUAV(OutputDepthTexture);
 
 	if (Lumen::UseFarField(*View.Family) || MegaLights::UseFarField(*View.Family))
 	{
-		RayGenParameters->MaxTraceDistance = Lumen::GetMaxTraceDistance(View);
-		RayGenParameters->FarFieldMaxTraceDistance = Lumen::GetFarFieldMaxTraceDistance();
+		SharedParameters.MaxTraceDistance = Lumen::GetMaxTraceDistance(View);
+		SharedParameters.FarFieldMaxTraceDistance = Lumen::GetFarFieldMaxTraceDistance();
 	}
 	else
 	{
-		RayGenParameters->MaxTraceDistance = 0.0f;
-		RayGenParameters->FarFieldMaxTraceDistance = 0.0f;
+		SharedParameters.MaxTraceDistance = 0.0f;
+		SharedParameters.FarFieldMaxTraceDistance = 0.0f;
 	}
 	
-	RayGenParameters->InstancesExtraData = GraphBuilder.CreateSRV(InstanceExtraDataBuffer);
-	RayGenParameters->InstancesDebugData = GraphBuilder.CreateSRV(InstanceDebugBuffer);
-	RayGenParameters->InstanceBuffer = GraphBuilder.CreateSRV(InstanceBuffer);
-	RayGenParameters->PickingBuffer = GraphBuilder.CreateSRV(PickingBuffer);
-	RayGenParameters->TLAS = RayTracingScene.GetLayerView(ERayTracingSceneLayer::Base);
-	RayGenParameters->FarFieldTLAS = RayTracingScene.GetLayerView(ERayTracingSceneLayer::FarField);
-	RayGenParameters->ViewUniformBuffer = View.ViewUniformBuffer;
-	RayGenParameters->Output = GraphBuilder.CreateUAV(SceneColorTexture);
-	RayGenParameters->TopKHitStats = GraphBuilder.CreateSRV(StatsBuffer);
-
-	RayGenParameters->SceneUniformBuffer = GetSceneUniformBufferRef(GraphBuilder, View); // TODO: use a separate params structure
+	SharedParameters.InstancesExtraData = GraphBuilder.CreateSRV(InstanceExtraDataBuffer);
+	SharedParameters.InstancesDebugData = GraphBuilder.CreateSRV(InstanceDebugBuffer);
+	SharedParameters.InstanceBuffer = GraphBuilder.CreateSRV(InstanceBuffer);
+	SharedParameters.PickingBuffer = GraphBuilder.CreateSRV(PickingBuffer);
+	SharedParameters.TLAS = RayTracingScene.GetLayerView(ERayTracingSceneLayer::Base);
+	SharedParameters.FarFieldTLAS = RayTracingScene.GetLayerView(ERayTracingSceneLayer::FarField);
+	SharedParameters.ViewUniformBuffer = View.ViewUniformBuffer;
+	SharedParameters.Output = GraphBuilder.CreateUAV(SceneColorTexture);
+	SharedParameters.TopKHitStats = GraphBuilder.CreateSRV(StatsBuffer);
+
+	SharedParameters.SceneUniformBuffer = GetSceneUniformBufferRef(GraphBuilder, View); // TODO: use a separate params structure
 
 	FIntRect ViewRect = View.ViewRect;
-	
+
 	RDG_EVENT_SCOPE_STAT(GraphBuilder, RayTracingDebug, "RayTracingDebug");
 	RDG_GPU_STAT_SCOPE(GraphBuilder, RayTracingDebug);
 
-	GraphBuilder.AddPass(
-		RDG_EVENT_NAME("RayTracingDebug"),
-		RayGenParameters,
-		ERDGPassFlags::Compute,
-		[RayGenParameters, RayGenShader, &View, Pipeline, SBT, ViewRect, bRequiresBindings](FRDGAsyncTask, FRHICommandList& RHICmdList)
+	if (bInlineRayTracing && RayTracingDebugModeSupportsInline(DebugVisualizationMode))
 	{
-		FRHIBatchedShaderParameters& GlobalResources = RHICmdList.GetScratchShaderParameters();
-		SetShaderParameters(GlobalResources, RayGenShader, *RayGenParameters);
+		FRayTracingDebugCS::FParameters* InlinePassParameters = GraphBuilder.AllocParameters<FRayTracingDebugCS::FParameters>();
+
+		InlinePassParameters->SharedParameters = SharedParameters;
+
+		TShaderRef<FRayTracingDebugCS> ComputeShader = View.ShaderMap->GetShader<FRayTracingDebugCS>();
+
+		FIntPoint Resolution(View.ViewRect.Width(), View.ViewRect.Height());
 
-		FRHIUniformBuffer* SceneUniformBuffer = RayGenParameters->SceneUniformBuffer.GetUniformBuffer()->GetRHI();
-		TOptional<FScopedUniformBufferStaticBindings> StaticUniformBufferScope = RayTracing::BindStaticUniformBufferBindings(View, SceneUniformBuffer, RHICmdList);
+		const FIntPoint GroupSize(FRayTracingDebugCS::ThreadGroupSizeX, FRayTracingDebugCS::ThreadGroupSizeY);
+		const FIntVector GroupCount = FComputeShaderUtils::GetGroupCount(ViewRect.Size(), GroupSize);
 
-		if (bRequiresBindings)
+		GraphBuilder.AddPass(
+			RDG_EVENT_NAME("RayTracingDebug (INLINE)"),
+			InlinePassParameters,
+			ERDGPassFlags::Compute,
+			[InlinePassParameters, ComputeShader, GroupCount](FRDGAsyncTask, FRHIComputeCommandList& RHICmdList)
+			{
+				FComputeShaderUtils::Dispatch(RHICmdList, ComputeShader, *InlinePassParameters, GroupCount);
+			});
+	}
+	else
+	{
+		FRayTracingDebugRGS::FParameters* RayGenParameters = GraphBuilder.AllocParameters<FRayTracingDebugRGS::FParameters>();
+
+		RayGenParameters->SharedParameters = SharedParameters;
+
+		FRayTracingDebugRGS::FPermutationDomain PermutationVector;
+		PermutationVector.Set<FRayTracingDebugRGS::FUseDebugCHSType>(bRequiresDebugCHS);
+		PermutationVector.Set<FRayTracingDebugRGS::FUseNvAPITimestamp>(GRHIGlobals.SupportsShaderTimestamp && IsRHIDeviceNVIDIA());
+
+		auto RayGenShader = View.ShaderMap->GetShader<FRayTracingDebugRGS>(PermutationVector);
+
+		FRayTracingPipelineState* Pipeline = View.MaterialRayTracingData.PipelineState;
+		FShaderBindingTableRHIRef SBT = View.MaterialRayTracingData.ShaderBindingTable;
+		bool bRequiresBindings = false;
+
+		if (bRequiresDebugCHS)
 		{
-			BindRayTracingDebugCHSMaterialBindings(RHICmdList, SBT, View, RayGenParameters->SceneUniformBuffer->GetRHI(), Pipeline);
-			RHICmdList.SetRayTracingMissShader(SBT, 0, Pipeline, 0 /* ShaderIndexInPipeline */, 0, nullptr, 0);
-			RHICmdList.CommitShaderBindingTable(SBT);
+			FRayTracingPipelineStateInitializer Initializer;
+			Initializer.MaxPayloadSizeInBytes = GetRayTracingPayloadTypeMaxSize(ERayTracingPayloadType::RayTracingDebug);
+
+			const FShaderBindingLayout* ShaderBindingLayout = RayTracing::GetShaderBindingLayout(ShaderPlatform);
+			if (ShaderBindingLayout)
+			{
+				Initializer.ShaderBindingLayout = &ShaderBindingLayout->RHILayout;
+			}
+
+			FRHIRayTracingShader* RayGenShaderTable[] = { RayGenShader.GetRayTracingShader() };
+			Initializer.SetRayGenShaderTable(RayGenShaderTable);
+
+			FRayTracingDebugCHS::FPermutationDomain PermutationVectorCHS;
+
+			PermutationVectorCHS.Set<FRayTracingDebugCHS::FNaniteRayTracing>(false);
+			auto HitGroupShader = View.ShaderMap->GetShader<FRayTracingDebugCHS>(PermutationVectorCHS);
+
+			PermutationVectorCHS.Set<FRayTracingDebugCHS::FNaniteRayTracing>(true);
+			auto HitGroupShaderNaniteRT = View.ShaderMap->GetShader<FRayTracingDebugCHS>(PermutationVectorCHS);
+
+			// auto AnyHitShader = ShaderMap->GetShader<FRayTracingDebugAHS>();
+
+			FRHIRayTracingShader* HitGroupTable[] = { HitGroupShader.GetRayTracingShader(), HitGroupShaderNaniteRT.GetRayTracingShader() };
+			Initializer.SetHitGroupTable(HitGroupTable);
+
+			auto MissShader = View.ShaderMap->GetShader<FRayTracingDebugMS>();
+			FRHIRayTracingShader* MissTable[] = { MissShader.GetRayTracingShader() };
+			Initializer.SetMissShaderTable(MissTable);
+
+			Pipeline = PipelineStateCache::GetAndOrCreateRayTracingPipelineState(GraphBuilder.RHICmdList, Initializer);
+
+			SBT = Scene.RayTracingSBT.AllocateTransientRHI(GraphBuilder.RHICmdList, ERayTracingShaderBindingMode::RTPSO, ERayTracingHitGroupIndexingMode::Allow, Initializer.GetMaxLocalBindingDataSize());
+
+			bRequiresBindings = true;
 		}
 
-		RHICmdList.RayTraceDispatch(Pipeline, RayGenShader.GetRayTracingShader(), SBT, GlobalResources, ViewRect.Size().X, ViewRect.Size().Y);
-	});
+		GraphBuilder.AddPass(
+			RDG_EVENT_NAME("RayTracingDebug"),
+			RayGenParameters,
+			ERDGPassFlags::Compute,
+			[RayGenParameters, RayGenShader, &View, Pipeline, SBT, ViewRect, bRequiresBindings](FRDGAsyncTask, FRHICommandList& RHICmdList)
+			{
+				FRHIBatchedShaderParameters& GlobalResources = RHICmdList.GetScratchShaderParameters();
+				SetShaderParameters(GlobalResources, RayGenShader, *RayGenParameters);
+
+				FRHIUniformBuffer* SceneUniformBuffer = RayGenParameters->SharedParameters.SceneUniformBuffer.GetUniformBuffer()->GetRHI();
+				TOptional<FScopedUniformBufferStaticBindings> StaticUniformBufferScope = RayTracing::BindStaticUniformBufferBindings(View, SceneUniformBuffer, RHICmdList);
+
+				if (bRequiresBindings)
+				{
+					BindRayTracingDebugCHSMaterialBindings(RHICmdList, SBT, View, RayGenParameters->SharedParameters.SceneUniformBuffer->GetRHI(), Pipeline);
+					RHICmdList.SetRayTracingMissShader(SBT, 0, Pipeline, 0 /* ShaderIndexInPipeline */, 0, nullptr, 0);
+					RHICmdList.CommitShaderBindingTable(SBT);
+				}
+
+				RHICmdList.RayTraceDispatch(Pipeline, RayGenShader.GetRayTracingShader(), SBT, GlobalResources, ViewRect.Size().X, ViewRect.Size().Y);
+			});
+	}	
 
 	if (DebugVisualizationMode == RAY_TRACING_DEBUG_VIZ_INSTANCE_OVERLAP)
 	{
