[HWRT] Replace "Performance" visualization with three different "timing" modes.
- Traversal - timing of TraceRay() in barycentrics visualization (no CHS or AHS evaluation)
- Material  - timing of TraceRay + material CHS and AHS (same as previous "Performance" visualization)
- Material (Alpha Only) - timing of TraceRay using shadow hitgroups so only AHS is used

- for now we won't show these modes in editor by default (instead requires enabling r.RayTracing.Visualize.Timing) since they're quite noisy and will require better documentation, etc.

diff --git a/Engine/Shaders/Private/RayTracing/RayTracingBarycentrics.usf b/Engine/Shaders/Private/RayTracing/RayTracingBarycentrics.usf
--- a/Engine/Shaders/Private/RayTracing/RayTracingBarycentrics.usf
+++ b/Engine/Shaders/Private/RayTracing/RayTracingBarycentrics.usf
@@ -1,6 +1,7 @@
 // Copyright Epic Games, Inc. All Rights Reserved.
 
 #include "../Common.ush"
+#include "../Visualization.ush"
 #include "RayTracingCommon.ush"
 
 #if COMPUTESHADER
@@ -12,6 +13,8 @@ RWTexture2D<float4> Output;
 RaytracingAccelerationStructure TLAS;
 float RTDebugVisualizationNaniteCutError;
 
+float TimingScale;
+
 #if COMPUTESHADER
 struct FNaniteTraceRayInlineCallback : FDefaultTraceRayInlineCallback
 {
@@ -47,6 +50,10 @@ void RayTracingBarycentricsCommon(uint2 DispatchThreadIndex)
 		| RAY_FLAG_CULL_BACK_FACING_TRIANGLES      // use back face culling
 		| RAY_FLAG_FORCE_OPAQUE;                   // don't run anyhit shader
 
+#if PLATFORM_SUPPORTS_SHADER_TIMESTAMP
+	FTimestamp TimeBegin = GetShaderTimestamp();
+#endif
+
 	FDefaultPayload Payload = (FDefaultPayload)0;
 #if COMPUTESHADER
 	FTraceRayInlineContext TraceRayInlineContext = CreateTraceRayInlineContext();
@@ -78,10 +85,18 @@ void RayTracingBarycentricsCommon(uint2 DispatchThreadIndex)
 		Payload // Payload
 	);
 #endif // COMPUTESHADER
+
+#if PLATFORM_SUPPORTS_SHADER_TIMESTAMP
+	FTimestamp TimeEnd = GetShaderTimestamp();
+#endif
+
 	float4 Result = float4(0, 0, 0, 1);
 	if (Payload.IsHit())
 	{
 		Result = float4(Payload.Barycentrics, 1.0f - Payload.Barycentrics.x - Payload.Barycentrics.y, 0.0f);
+#if OUTPUT_TIMING && PLATFORM_SUPPORTS_SHADER_TIMESTAMP
+		Result = float4(ColorMapInferno(float(ShaderTimestampDiff(TimeBegin, TimeEnd)) * TimingScale), 0.0f);
+#endif
 	}
 	Output[PixelCoord] = Result;
 }
diff --git a/Engine/Shaders/Private/RayTracing/RayTracingDebug.usf b/Engine/Shaders/Private/RayTracing/RayTracingDebug.usf
--- a/Engine/Shaders/Private/RayTracing/RayTracingDebug.usf
+++ b/Engine/Shaders/Private/RayTracing/RayTracingDebug.usf
@@ -76,10 +76,6 @@ float4 DebugRayTracingMaterial(uint2 PixelCoord, FRayDesc Ray, bool bIsFarField)
 	const bool bIgnoreTranslucentMaterials = false;
 	const bool bCameraRay = true;
 
-#if PLATFORM_SUPPORTS_SHADER_TIMESTAMP
-	FTimestamp TimeBegin = GetShaderTimestamp();
-#endif
-
 	RaytracingAccelerationStructure TLASToTrace = TLAS;
 	if (bIsFarField)
 	{
@@ -93,19 +89,48 @@ float4 DebugRayTracingMaterial(uint2 PixelCoord, FRayDesc Ray, bool bIsFarField)
 		Ray.TMax = (MaxTraceDistance > 0.0f) ? MaxTraceDistance : Ray.TMax;
 	}
 	
-	FMaterialClosestHitPayload Payload = TraceMaterialRay(
-		TLASToTrace,
-		RayFlags,
-		InstanceInclusionMask,
-		Ray,
-		RayCone,
-		bEnableSkyLightContribution,
-		bIgnoreTranslucentMaterials,
-		bCameraRay);
+	FMaterialClosestHitPayload Payload = (FMaterialClosestHitPayload)0;
 
 #if PLATFORM_SUPPORTS_SHADER_TIMESTAMP
-	FTimestamp TimeEnd = GetShaderTimestamp();
+	uint TimestampDiff = 0;
+
+	if (VisualizationMode == RAY_TRACING_DEBUG_VIZ_TIMING_ANY_HIT)
+	{
+		FTimestamp TimeBegin = GetShaderTimestamp();
+
+		TraceVisibilityRay(
+			TLAS,
+			RayFlags,
+			InstanceInclusionMask,
+			Ray);
+
+		FTimestamp TimeEnd = GetShaderTimestamp();
+
+		TimestampDiff = ShaderTimestampDiff(TimeBegin, TimeEnd);
+	}
+	else
+#endif
+	{
+#if PLATFORM_SUPPORTS_SHADER_TIMESTAMP
+		FTimestamp TimeBegin = GetShaderTimestamp();
+#endif
+
+		Payload = TraceMaterialRay(
+			TLASToTrace,
+			RayFlags,
+			InstanceInclusionMask,
+			Ray,
+			RayCone,
+			bEnableSkyLightContribution,
+			bIgnoreTranslucentMaterials,
+			bCameraRay);
+
+#if PLATFORM_SUPPORTS_SHADER_TIMESTAMP
+		FTimestamp TimeEnd = GetShaderTimestamp();
+
+		TimestampDiff = ShaderTimestampDiff(TimeBegin, TimeEnd);
 #endif
+	}
 
 	if (Payload.IsHit())
 	{
@@ -127,7 +152,8 @@ float4 DebugRayTracingMaterial(uint2 PixelCoord, FRayDesc Ray, bool bIsFarField)
 		#endif
 
 		#if PLATFORM_SUPPORTS_SHADER_TIMESTAMP
-		case RAY_TRACING_DEBUG_VIZ_PERFORMANCE:				return float4(ColorMapInferno(float(ShaderTimestampDiff(TimeBegin, TimeEnd)) * TimingScale), 0.0f);
+		case RAY_TRACING_DEBUG_VIZ_TIMING_ANY_HIT:
+		case RAY_TRACING_DEBUG_VIZ_TIMING_MATERIAL:				return float4(ColorMapInferno(float(TimestampDiff) * TimingScale), 0.0f);
 		#endif // PLATFORM_SUPPORTS_SHADER_TIMESTAMP
 		}
 
diff --git a/Engine/Shaders/Shared/RaytracingDebugDefinitions.h b/Engine/Shaders/Shared/RaytracingDebugDefinitions.h
--- a/Engine/Shaders/Shared/RaytracingDebugDefinitions.h
+++ b/Engine/Shaders/Shared/RaytracingDebugDefinitions.h
@@ -30,27 +30,29 @@
 #define RAY_TRACING_DEBUG_VIZ_WORLD_TANGENT						20 
 #define RAY_TRACING_DEBUG_VIZ_ANISOTROPY						21 
 #define RAY_TRACING_DEBUG_VIZ_INSTANCES							22
-#define RAY_TRACING_DEBUG_VIZ_PERFORMANCE						23 
-#define RAY_TRACING_DEBUG_VIZ_TRIANGLES							24
-#define RAY_TRACING_DEBUG_VIZ_FAR_FIELD							25
-#define RAY_TRACING_DEBUG_VIZ_TRAVERSAL_PRIMARY_NODE			26
-#define RAY_TRACING_DEBUG_VIZ_TRAVERSAL_PRIMARY_CLUSTER			27
-#define RAY_TRACING_DEBUG_VIZ_TRAVERSAL_PRIMARY_TRIANGLE		28
-#define RAY_TRACING_DEBUG_VIZ_TRAVERSAL_PRIMARY_ALL				29
-#define RAY_TRACING_DEBUG_VIZ_TRAVERSAL_PRIMARY_STATISTICS		30
-#define RAY_TRACING_DEBUG_VIZ_TRAVERSAL_SECONDARY_NODE			31
-#define RAY_TRACING_DEBUG_VIZ_TRAVERSAL_SECONDARY_CLUSTER		32
-#define RAY_TRACING_DEBUG_VIZ_TRAVERSAL_SECONDARY_TRIANGLE		33
-#define RAY_TRACING_DEBUG_VIZ_TRAVERSAL_SECONDARY_ALL			34
-#define RAY_TRACING_DEBUG_VIZ_TRAVERSAL_SECONDARY_STATISTICS	35
-#define RAY_TRACING_DEBUG_VIZ_DYNAMIC_INSTANCES					36
-#define RAY_TRACING_DEBUG_VIZ_PROXY_TYPE						37
-#define RAY_TRACING_DEBUG_VIZ_PICKER							38
-#define RAY_TRACING_DEBUG_VIZ_INSTANCE_OVERLAP					39
-#define RAY_TRACING_DEBUG_VIZ_TRIANGLE_HITCOUNT					40
-#define RAY_TRACING_DEBUG_VIZ_HITCOUNT_PER_INSTANCE				41
-#define RAY_TRACING_DEBUG_VIZ_LIGHT_GRID_COUNT					42
-#define RAY_TRACING_DEBUG_VIZ_SUBSTRATE_DATA					43
+#define RAY_TRACING_DEBUG_VIZ_TIMING_TRAVERSAL					23 
+#define RAY_TRACING_DEBUG_VIZ_TIMING_ANY_HIT					24 
+#define RAY_TRACING_DEBUG_VIZ_TIMING_MATERIAL					25 
+#define RAY_TRACING_DEBUG_VIZ_TRIANGLES							26
+#define RAY_TRACING_DEBUG_VIZ_FAR_FIELD							27
+#define RAY_TRACING_DEBUG_VIZ_TRAVERSAL_PRIMARY_NODE			28
+#define RAY_TRACING_DEBUG_VIZ_TRAVERSAL_PRIMARY_CLUSTER			29
+#define RAY_TRACING_DEBUG_VIZ_TRAVERSAL_PRIMARY_TRIANGLE		30
+#define RAY_TRACING_DEBUG_VIZ_TRAVERSAL_PRIMARY_ALL				31
+#define RAY_TRACING_DEBUG_VIZ_TRAVERSAL_PRIMARY_STATISTICS		32
+#define RAY_TRACING_DEBUG_VIZ_TRAVERSAL_SECONDARY_NODE			33
+#define RAY_TRACING_DEBUG_VIZ_TRAVERSAL_SECONDARY_CLUSTER		34
+#define RAY_TRACING_DEBUG_VIZ_TRAVERSAL_SECONDARY_TRIANGLE		35
+#define RAY_TRACING_DEBUG_VIZ_TRAVERSAL_SECONDARY_ALL			36
+#define RAY_TRACING_DEBUG_VIZ_TRAVERSAL_SECONDARY_STATISTICS	37
+#define RAY_TRACING_DEBUG_VIZ_DYNAMIC_INSTANCES					38
+#define RAY_TRACING_DEBUG_VIZ_PROXY_TYPE						39
+#define RAY_TRACING_DEBUG_VIZ_PICKER							40
+#define RAY_TRACING_DEBUG_VIZ_INSTANCE_OVERLAP					41
+#define RAY_TRACING_DEBUG_VIZ_TRIANGLE_HITCOUNT					42
+#define RAY_TRACING_DEBUG_VIZ_HITCOUNT_PER_INSTANCE				43
+#define RAY_TRACING_DEBUG_VIZ_LIGHT_GRID_COUNT					44
+#define RAY_TRACING_DEBUG_VIZ_SUBSTRATE_DATA					45
 
 #define RAY_TRACING_DEBUG_PICKER_DOMAIN_TRIANGLE				0
 #define RAY_TRACING_DEBUG_PICKER_DOMAIN_INSTANCE				1
diff --git a/Engine/Source/Editor/UnrealEd/Private/RayTracingDebugVisualizationMenuCommands.cpp b/Engine/Source/Editor/UnrealEd/Private/RayTracingDebugVisualizationMenuCommands.cpp
--- a/Engine/Source/Editor/UnrealEd/Private/RayTracingDebugVisualizationMenuCommands.cpp
+++ b/Engine/Source/Editor/UnrealEd/Private/RayTracingDebugVisualizationMenuCommands.cpp
@@ -11,6 +11,13 @@
 #include "EditorViewportClient.h"
 #include "ToolMenu.h"
 
+int32 GRayTracingVisualizeTiming = 0;
+static FAutoConsoleVariableRef CVarRayTracingVisualizeTiming(
+	TEXT("r.RayTracing.Visualize.Timing"),
+	GRayTracingVisualizeTiming,
+	TEXT("Include 'Timing' visualization modes in the in-editor 'Ray Tracing Debug' drop down menu.")
+);
+
 int32 GRayTracingVisualizeOther = 0;
 static FAutoConsoleVariableRef CVarRayTracingVisualizeOther(
 	TEXT("r.RayTracing.Visualize.Other"),
@@ -72,6 +79,10 @@ void FRayTracingDebugVisualizationMenuCommands::BuildCommandMap()
 			Record.Type = FVisualizationType::Performance;
 			break;
 
+		case FRayTracingVisualizationData::FModeType::Timing:
+			Record.Type = FVisualizationType::Timing;
+			break;
+
 		case FRayTracingVisualizationData::FModeType::Other:
 			Record.Type = FVisualizationType::Other;
 			break;
@@ -82,6 +93,7 @@ void FRayTracingDebugVisualizationMenuCommands::BuildCommandMap()
 // Deprecated. When removed, also remove AddCommandTypeToMenu(FMenuBuilder& Menu, const FVisualizationType Type)
 void FRayTracingDebugVisualizationMenuCommands::BuildVisualisationSubMenu(FMenuBuilder& Menu)
 {
+	const bool bShowTiming = GRayTracingVisualizeTiming != 0;
 	const bool bShowOther = GRayTracingVisualizeOther != 0;
 
 	const FRayTracingDebugVisualizationMenuCommands& Commands = FRayTracingDebugVisualizationMenuCommands::Get();
@@ -104,6 +116,13 @@ void FRayTracingDebugVisualizationMenuCommands::BuildVisualisationSubMenu(FMenuB
 			Menu.EndSection();
 		}
 
+		if (bShowTiming)
+		{
+			Menu.BeginSection("LevelViewportRayTracingVisualizationTiming", LOCTEXT("RayTracingVisualizationTiming", "Timing"));
+			Commands.AddCommandTypeToMenu(Menu, FVisualizationType::Timing);
+			Menu.EndSection();
+		}
+
 		if (bShowOther)
 		{
 			Menu.BeginSection("LevelViewportRayTracingVisualizationOther", LOCTEXT("RayTracingVisualizationOther", "Other"));
@@ -143,6 +162,13 @@ void FRayTracingDebugVisualizationMenuCommands::BuildVisualisationSubMenu(UToolM
 		Commands.AddCommandTypeToSection(PerformanceSection, FVisualizationType::Performance);
 	}
 
+	const bool bShowTiming = GRayTracingVisualizeTiming != 0;
+	if (bShowTiming)
+	{
+		FToolMenuSection& TimingSection = InMenu->AddSection("LevelViewportRayTracingVisualizationTiming", LOCTEXT("RayTracingVisualizationTiming", "Timing"));
+		Commands.AddCommandTypeToSection(TimingSection, FVisualizationType::Timing);
+	}
+
 	const bool bShowOther = GRayTracingVisualizeOther != 0;
 	if (bShowOther)
 	{
diff --git a/Engine/Source/Editor/UnrealEd/Public/RayTracingDebugVisualizationMenuCommands.h b/Engine/Source/Editor/UnrealEd/Public/RayTracingDebugVisualizationMenuCommands.h
--- a/Engine/Source/Editor/UnrealEd/Public/RayTracingDebugVisualizationMenuCommands.h
+++ b/Engine/Source/Editor/UnrealEd/Public/RayTracingDebugVisualizationMenuCommands.h
@@ -21,7 +21,7 @@ public:
 		Overview,
 		Standard,
 		Performance,
-		Material,
+		Timing,
 		Other,
 	};
 
diff --git a/Engine/Source/Runtime/Engine/Private/RayTracingVisualizationData.cpp b/Engine/Source/Runtime/Engine/Private/RayTracingVisualizationData.cpp
--- a/Engine/Source/Runtime/Engine/Private/RayTracingVisualizationData.cpp
+++ b/Engine/Source/Runtime/Engine/Private/RayTracingVisualizationData.cpp
@@ -57,11 +57,6 @@ FRayTracingVisualizationData::FRayTracingVisualizationData()
 		AddVisualizationMode(TEXT("TriangleHitCount"), LOCTEXT("TriangleHitCount", "Triangle Hit Count"), FModeType::Performance, RAY_TRACING_DEBUG_VIZ_TRIANGLE_HITCOUNT, false);
 		AddVisualizationMode(TEXT("HitCountPerInstance"), LOCTEXT("HitCountPerInstance", "Hit Count Per Instance"), FModeType::Performance, RAY_TRACING_DEBUG_VIZ_HITCOUNT_PER_INSTANCE, false);
 
-		if (GRHISupportsShaderTimestamp)
-		{
-			AddVisualizationMode(TEXT("Performance"), LOCTEXT("Performance", "Performance"), FModeType::Performance, RAY_TRACING_DEBUG_VIZ_PERFORMANCE, false);
-		}
-
 		AddVisualizationMode(TEXT("Traversal Primary Node"), LOCTEXT("TraversalPrimaryNode", "Traversal Primary Node"), FModeType::Performance, RAY_TRACING_DEBUG_VIZ_TRAVERSAL_PRIMARY_NODE, false, /*bHiddenInEditor*/ true);
 		AddVisualizationMode(TEXT("Traversal Primary Cluster"), LOCTEXT("TraversalPrimaryCluster", "Traversal Primary Cluster"), FModeType::Performance, RAY_TRACING_DEBUG_VIZ_TRAVERSAL_PRIMARY_CLUSTER, false, /*bHiddenInEditor*/ true);
 		AddVisualizationMode(TEXT("Traversal Primary Triangle"), LOCTEXT("TraversalPrimaryTriangle", "Traversal Primary Triangle"), FModeType::Performance, RAY_TRACING_DEBUG_VIZ_TRAVERSAL_PRIMARY_TRIANGLE, false, /*bHiddenInEditor*/ true);
@@ -74,6 +69,13 @@ FRayTracingVisualizationData::FRayTracingVisualizationData()
 		AddVisualizationMode(TEXT("Traversal Secondary All"), LOCTEXT("TraversalSecondaryAll", "Traversal Secondary All"), FModeType::Performance, RAY_TRACING_DEBUG_VIZ_TRAVERSAL_SECONDARY_ALL, false, /*bHiddenInEditor*/ true);
 		AddVisualizationMode(TEXT("Traversal Secondary Statistics"), LOCTEXT("TraversalSecondaryStatistics", "Traversal Secondary Statistics"), FModeType::Performance, RAY_TRACING_DEBUG_VIZ_TRAVERSAL_SECONDARY_STATISTICS, false, /*bHiddenInEditor*/ true);
 
+		if (GRHISupportsShaderTimestamp)
+		{
+			AddVisualizationMode(TEXT("Timing Traversal"), LOCTEXT("TimingTraversal", "Timing - Traversal"), FModeType::Timing, RAY_TRACING_DEBUG_VIZ_TIMING_TRAVERSAL, false);
+			AddVisualizationMode(TEXT("Timing Material"), LOCTEXT("TimingMaterial", "Timing - Material"), FModeType::Timing, RAY_TRACING_DEBUG_VIZ_TIMING_MATERIAL, false);
+			AddVisualizationMode(TEXT("Timing AHS"), LOCTEXT("TimingAHS", "Timing - Material (Alpha Only)"), FModeType::Timing, RAY_TRACING_DEBUG_VIZ_TIMING_ANY_HIT, false);
+		}
+
 		if (Substrate::IsSubstrateEnabled())
 		{
 			AddVisualizationMode(TEXT("SubstrateMaterialProperties"), LOCTEXT("SubstrateMaterialProperties", "Substrate Material Properties"), FModeType::Other, RAY_TRACING_DEBUG_VIZ_SUBSTRATE_DATA, true);
diff --git a/Engine/Source/Runtime/Engine/Public/RayTracingVisualizationData.h b/Engine/Source/Runtime/Engine/Public/RayTracingVisualizationData.h
--- a/Engine/Source/Runtime/Engine/Public/RayTracingVisualizationData.h
+++ b/Engine/Source/Runtime/Engine/Public/RayTracingVisualizationData.h
@@ -12,6 +12,7 @@ public:
 		Overview,
 		Standard,
 		Performance,
+		Timing,
 		Other,
 	};
 
diff --git a/Engine/Source/Runtime/Renderer/Private/RayTracing/RayTracingBarycentrics.cpp b/Engine/Source/Runtime/Renderer/Private/RayTracing/RayTracingBarycentrics.cpp
--- a/Engine/Source/Runtime/Renderer/Private/RayTracing/RayTracingBarycentrics.cpp
+++ b/Engine/Source/Runtime/Renderer/Private/RayTracing/RayTracingBarycentrics.cpp
@@ -28,8 +28,13 @@ class FRayTracingBarycentricsRGS : public FGlobalShader
 		SHADER_PARAMETER_RDG_TEXTURE_UAV(RWTexture2D<float4>, Output)
 		SHADER_PARAMETER_STRUCT_REF(FViewUniformShaderParameters, ViewUniformBuffer)
 		SHADER_PARAMETER_RDG_UNIFORM_BUFFER(FSceneUniformParameters, Scene)
+		SHADER_PARAMETER(float, TimingScale)
 	END_SHADER_PARAMETER_STRUCT()
 
+	class FOutputTiming : SHADER_PERMUTATION_BOOL("OUTPUT_TIMING");
+	class FUseNvAPITimestamp : SHADER_PERMUTATION_BOOL("USE_NVAPI_TIMESTAMP");
+	using FPermutationDomain = TShaderPermutationDomain<FOutputTiming, FUseNvAPITimestamp>;
+
 	static bool ShouldCompilePermutation(const FGlobalShaderPermutationParameters& Parameters)
 	{
 		return ShouldCompileRayTracingShadersForProject(Parameters.Platform);
@@ -95,7 +100,8 @@ class FRayTracingBarycentricsCS : public FGlobalShader
 	END_SHADER_PARAMETER_STRUCT()
 
 	class FSupportProceduralPrimitive : SHADER_PERMUTATION_BOOL("ENABLE_TRACE_RAY_INLINE_PROCEDURAL_PRIMITIVE");
-	using FPermutationDomain = TShaderPermutationDomain<FSupportProceduralPrimitive>;
+	class FOutputTiming : SHADER_PERMUTATION_BOOL("OUTPUT_TIMING");
+	using FPermutationDomain = TShaderPermutationDomain<FSupportProceduralPrimitive, FOutputTiming>;
 
 	static void ModifyCompilationEnvironment(const FGlobalShaderPermutationParameters& Parameters, FShaderCompilerEnvironment& OutEnvironment)
 	{
@@ -121,7 +127,7 @@ class FRayTracingBarycentricsCS : public FGlobalShader
 };
 IMPLEMENT_GLOBAL_SHADER(FRayTracingBarycentricsCS, "/Engine/Private/RayTracing/RayTracingBarycentrics.usf", "RayTracingBarycentricsMainCS", SF_Compute);
 
-void RenderRayTracingBarycentricsCS(FRDGBuilder& GraphBuilder, const FViewInfo& View, FRDGTextureRef SceneColor, bool bVisualizeProceduralPrimitives)
+void RenderRayTracingBarycentricsCS(FRDGBuilder& GraphBuilder, const FViewInfo& View, FRDGTextureRef SceneColor, bool bVisualizeProceduralPrimitives, bool bOutputTiming)
 {
 	FRayTracingBarycentricsCS::FParameters* PassParameters = GraphBuilder.AllocParameters<FRayTracingBarycentricsCS::FParameters>();
 
@@ -138,6 +144,7 @@ void RenderRayTracingBarycentricsCS(FRDGBuilder& GraphBuilder, const FViewInfo&
 
 	FRayTracingBarycentricsCS::FPermutationDomain PermutationVector;
 	PermutationVector.Set<FRayTracingBarycentricsCS::FSupportProceduralPrimitive>(bVisualizeProceduralPrimitives);
+	PermutationVector.Set<FRayTracingBarycentricsCS::FOutputTiming>(bOutputTiming);
 
 	auto ComputeShader = View.ShaderMap->GetShader<FRayTracingBarycentricsCS>(PermutationVector);
 
@@ -147,9 +154,13 @@ void RenderRayTracingBarycentricsCS(FRDGBuilder& GraphBuilder, const FViewInfo&
 	FComputeShaderUtils::AddPass(GraphBuilder, RDG_EVENT_NAME("Barycentrics"), ComputeShader, PassParameters, GroupCount);
 }
 
-void RenderRayTracingBarycentricsRGS(FRDGBuilder& GraphBuilder, const FScene& Scene, const FViewInfo& View, FRDGTextureRef SceneColor)
+void RenderRayTracingBarycentricsRGS(FRDGBuilder& GraphBuilder, const FScene& Scene, const FViewInfo& View, FRDGTextureRef SceneColor, bool bOutputTiming)
 {
-	auto RayGenShader = View.ShaderMap->GetShader<FRayTracingBarycentricsRGS>();
+	FRayTracingBarycentricsRGS::FPermutationDomain PermutationVector;
+	PermutationVector.Set<FRayTracingBarycentricsRGS::FOutputTiming>(bOutputTiming);
+	PermutationVector.Set<FRayTracingBarycentricsRGS::FUseNvAPITimestamp>(IsRHIDeviceNVIDIA());
+
+	auto RayGenShader = View.ShaderMap->GetShader<FRayTracingBarycentricsRGS>(PermutationVector);
 	auto ClosestHitShader = View.ShaderMap->GetShader<FRayTracingBarycentricsCHS>();
 
 	FRayTracingPipelineStateInitializer Initializer;
@@ -180,6 +191,9 @@ void RenderRayTracingBarycentricsRGS(FRDGBuilder& GraphBuilder, const FScene& Sc
 	RayGenParameters->Scene = View.GetSceneUniforms().GetBuffer(GraphBuilder);
 	RayGenParameters->Output = GraphBuilder.CreateUAV(SceneColor);
 
+	static const auto TimingScaleCvar = IConsoleManager::Get().FindTConsoleVariableDataFloat(TEXT("r.RayTracing.Visualize.TimingScale"));
+	RayGenParameters->TimingScale = (TimingScaleCvar ? TimingScaleCvar->GetValueOnRenderThread() : 1.0f) / 25000.0f;
+
 	FIntRect ViewRect = View.ViewRect;
 	
 	GraphBuilder.AddPass(
@@ -202,18 +216,19 @@ void RenderRayTracingBarycentricsRGS(FRDGBuilder& GraphBuilder, const FScene& Sc
 	});
 }
 
-void RenderRayTracingBarycentrics(FRDGBuilder& GraphBuilder, const FScene& Scene, const FViewInfo& View, FRDGTextureRef SceneColor, bool bVisualizeProceduralPrimitives)
+void RenderRayTracingBarycentrics(FRDGBuilder& GraphBuilder, const FScene& Scene, const FViewInfo& View, FRDGTextureRef SceneColor, bool bVisualizeProceduralPrimitives, bool bOutputTiming)
 {
 	const bool bRayTracingInline = ShouldRenderRayTracingEffect(true, ERayTracingPipelineCompatibilityFlags::Inline, View);
 	const bool bRayTracingPipeline = ShouldRenderRayTracingEffect(true, ERayTracingPipelineCompatibilityFlags::FullPipeline, View);
 
-	if (bRayTracingInline)
+	// currently bOutputTiming is only supported on Nvidia when using RGS
+	if (bRayTracingInline && (!bOutputTiming || !IsRHIDeviceNVIDIA()))
 	{
-		RenderRayTracingBarycentricsCS(GraphBuilder, View, SceneColor, bVisualizeProceduralPrimitives);
+		RenderRayTracingBarycentricsCS(GraphBuilder, View, SceneColor, bVisualizeProceduralPrimitives, bOutputTiming);
 	}
 	else if (bRayTracingPipeline)
 	{
-		RenderRayTracingBarycentricsRGS(GraphBuilder, Scene, View, SceneColor);
+		RenderRayTracingBarycentricsRGS(GraphBuilder, Scene, View, SceneColor, bOutputTiming);
 	}
 }
 #endif
diff --git a/Engine/Source/Runtime/Renderer/Private/RayTracing/RayTracingDebug.cpp b/Engine/Source/Runtime/Renderer/Private/RayTracing/RayTracingDebug.cpp
--- a/Engine/Source/Runtime/Renderer/Private/RayTracing/RayTracingDebug.cpp
+++ b/Engine/Source/Runtime/Renderer/Private/RayTracing/RayTracingDebug.cpp
@@ -1564,7 +1564,7 @@ bool HasRaytracingDebugViewModeRaytracedOverlay(const FSceneViewFamily& ViewFami
 	return !bAnySubstrate && !bAnyTraversalSecondary;
 }
 
-extern void RenderRayTracingBarycentrics(FRDGBuilder& GraphBuilder, const FScene& Scene, const FViewInfo& View, FRDGTextureRef SceneColor, bool bVisualizeProceduralPrimitives);
+extern void RenderRayTracingBarycentrics(FRDGBuilder& GraphBuilder, const FScene& Scene, const FViewInfo& View, FRDGTextureRef SceneColor, bool bVisualizeProceduralPrimitives, bool bOutputTiming);
 
 extern void RenderRayTracingPrimaryRaysView(
 	FRDGBuilder& GraphBuilder, const FViewInfo& View, const FSceneTextures& SceneTextures,
@@ -1587,9 +1587,9 @@ void RenderRayTracingDebug(FRDGBuilder& GraphBuilder, const FScene& Scene, const
 		ShaderPrint::RequestSpaceForCharacters(1024);
 	}
 
-	if (DebugVisualizationMode == RAY_TRACING_DEBUG_VIZ_BARYCENTRICS)
+	if (DebugVisualizationMode == RAY_TRACING_DEBUG_VIZ_BARYCENTRICS || DebugVisualizationMode == RAY_TRACING_DEBUG_VIZ_TIMING_TRAVERSAL)
 	{
-		return RenderRayTracingBarycentrics(GraphBuilder, Scene, View, SceneColorTexture, (bool)GVisualizeProceduralPrimitives);
+		return RenderRayTracingBarycentrics(GraphBuilder, Scene, View, SceneColorTexture, (bool)GVisualizeProceduralPrimitives, /*bOutputTiming*/ DebugVisualizationMode == RAY_TRACING_DEBUG_VIZ_TIMING_TRAVERSAL);
 	}
 
 	const FRayTracingScene& RayTracingScene = Scene.RayTracingScene;
@@ -1897,7 +1897,9 @@ extern void RayTracingDebugDisplayOnScreenMessages(FScreenMessageWriter& Writer,
 
 	switch (DebugVisualizationMode)
 	{
-	case RAY_TRACING_DEBUG_VIZ_PERFORMANCE:
+	case RAY_TRACING_DEBUG_VIZ_TIMING_TRAVERSAL:
+	case RAY_TRACING_DEBUG_VIZ_TIMING_ANY_HIT:
+	case RAY_TRACING_DEBUG_VIZ_TIMING_MATERIAL:
 	{
 		static const FText Message = NSLOCTEXT("Renderer", "RayTracingDebugVizPerformance", "Use r.RayTracing.Visualize.TimingScale to adjust visualization.");
 		Writer.DrawLine(Message, 10, FColor::White);
