First Person: Skip Nanite rendering to FP VSM since we don't support appropriate culling yet.

diff --git a/Engine/Source/Runtime/Renderer/Private/Shadows/ShadowSceneRenderer.cpp b/Engine/Source/Runtime/Renderer/Private/Shadows/ShadowSceneRenderer.cpp
--- a/Engine/Source/Runtime/Renderer/Private/Shadows/ShadowSceneRenderer.cpp
+++ b/Engine/Source/Runtime/Renderer/Private/Shadows/ShadowSceneRenderer.cpp
@@ -993,6 +993,14 @@ void FShadowSceneRenderer::CreateNaniteRenderPasses(
 
 		for (FProjectedShadowInfo* ProjectedShadowInfo : Shadows)
 		{
+			// First person VSM clipmaps currently do not support Nanite primitives because Nanite does not yet support filtering for such primitives
+			// so that they'll end up in the correct clipmap. Without this check, all Nanite primitives will end being rendered into the first person clipmap,
+			// when we actually only want primitives tagged as FirstPersonWorldSpaceRepresentation to be rendered into it.
+			if (ProjectedShadowInfo->VirtualShadowMapClipmap && ProjectedShadowInfo->VirtualShadowMapClipmap->IsFirstPersonShadow())
+			{
+				continue;
+			}
+
 			if (ProjectedShadowInfo->bShouldRenderVSM)
 			{
 				FVSMRenderViewInfo Info = GetRenderViewInfo(ProjectedShadowInfo);
diff --git a/Engine/Source/Runtime/Renderer/Private/VirtualShadowMaps/VirtualShadowMapClipmap.h b/Engine/Source/Runtime/Renderer/Private/VirtualShadowMaps/VirtualShadowMapClipmap.h
--- a/Engine/Source/Runtime/Renderer/Private/VirtualShadowMaps/VirtualShadowMapClipmap.h
+++ b/Engine/Source/Runtime/Renderer/Private/VirtualShadowMaps/VirtualShadowMapClipmap.h
@@ -113,6 +113,11 @@ public:
 
 	TSharedPtr<FVirtualShadowMapPerLightCacheEntry>& GetCacheEntry() { return PerLightCacheEntry; }
 
+	/**
+	 * Whether the clipmap is for casting shadow of FirstPersonWorldSpaceRepresentation primitives onto the scene.
+	 */
+	bool IsFirstPersonShadow() const { return Config.bIsFirstPersonShadow; }
+
 private:
 	FVirtualShadowMapProjectionShaderData ComputeProjectionShaderData(int32 ClipmapIndex) const;
 	void ComputeBoundingVolumes(const FVector CameraOrigin);
