Fix error compiling FRayTracingBarycentricsRGS when using Vulkan.

diff --git a/Engine/Source/Runtime/Renderer/Private/RayTracing/RayTracingBarycentrics.cpp b/Engine/Source/Runtime/Renderer/Private/RayTracing/RayTracingBarycentrics.cpp
--- a/Engine/Source/Runtime/Renderer/Private/RayTracing/RayTracingBarycentrics.cpp
+++ b/Engine/Source/Runtime/Renderer/Private/RayTracing/RayTracingBarycentrics.cpp
@@ -37,6 +37,15 @@ class FRayTracingBarycentricsRGS : public FGlobalShader
 
 	static bool ShouldCompilePermutation(const FGlobalShaderPermutationParameters& Parameters)
 	{
+		FPermutationDomain PermutationVector(Parameters.PermutationId);
+
+		// TODO: Check this using DDPI
+		const bool bUseNvAPITimestamp = PermutationVector.Get<FUseNvAPITimestamp>();
+		if (bUseNvAPITimestamp && IsVulkanPlatform(Parameters.Platform))
+		{
+			return false;
+		}
+
 		return ShouldCompileRayTracingShadersForProject(Parameters.Platform);
 	}
 
