Fix error compiling FRayTracingBarycentricsCS.
- added helper function GetRayTracingDebugTimingScale().

diff --git a/Engine/Source/Runtime/Renderer/Private/RayTracing/RayTracingBarycentrics.cpp b/Engine/Source/Runtime/Renderer/Private/RayTracing/RayTracingBarycentrics.cpp
--- a/Engine/Source/Runtime/Renderer/Private/RayTracing/RayTracingBarycentrics.cpp
+++ b/Engine/Source/Runtime/Renderer/Private/RayTracing/RayTracingBarycentrics.cpp
@@ -106,6 +106,7 @@ class FRayTracingBarycentricsCS : public FGlobalShader
 		SHADER_PARAMETER_RDG_UNIFORM_BUFFER(FNaniteRasterUniformParameters, NaniteRasterUniformBuffer)
 		SHADER_PARAMETER_RDG_UNIFORM_BUFFER(FNaniteShadingUniformParameters, NaniteShadingUniformBuffer)
 		SHADER_PARAMETER(float, RTDebugVisualizationNaniteCutError)
+		SHADER_PARAMETER(float, TimingScale)
 	END_SHADER_PARAMETER_STRUCT()
 
 	class FSupportProceduralPrimitive : SHADER_PERMUTATION_BOOL("ENABLE_TRACE_RAY_INLINE_PROCEDURAL_PRIMITIVE");
@@ -136,6 +137,8 @@ class FRayTracingBarycentricsCS : public FGlobalShader
 };
 IMPLEMENT_GLOBAL_SHADER(FRayTracingBarycentricsCS, "/Engine/Private/RayTracing/RayTracingBarycentrics.usf", "RayTracingBarycentricsMainCS", SF_Compute);
 
+extern float GetRayTracingDebugTimingScale();
+
 void RenderRayTracingBarycentricsCS(FRDGBuilder& GraphBuilder, const FViewInfo& View, FRDGTextureRef SceneColor, bool bVisualizeProceduralPrimitives, bool bOutputTiming)
 {
 	FRayTracingBarycentricsCS::FParameters* PassParameters = GraphBuilder.AllocParameters<FRayTracingBarycentricsCS::FParameters>();
@@ -148,6 +151,7 @@ void RenderRayTracingBarycentricsCS(FRDGBuilder& GraphBuilder, const FViewInfo&
 	PassParameters->NaniteShadingUniformBuffer = CreateDebugNaniteShadingUniformBuffer(GraphBuilder);
 
 	PassParameters->RTDebugVisualizationNaniteCutError = 0.0f;
+	PassParameters->TimingScale = GetRayTracingDebugTimingScale();
 
 	FIntRect ViewRect = View.ViewRect;
 
@@ -200,8 +204,7 @@ void RenderRayTracingBarycentricsRGS(FRDGBuilder& GraphBuilder, const FScene& Sc
 	RayGenParameters->Scene = View.GetSceneUniforms().GetBuffer(GraphBuilder);
 	RayGenParameters->Output = GraphBuilder.CreateUAV(SceneColor);
 
-	static const auto TimingScaleCvar = IConsoleManager::Get().FindTConsoleVariableDataFloat(TEXT("r.RayTracing.Visualize.TimingScale"));
-	RayGenParameters->TimingScale = (TimingScaleCvar ? TimingScaleCvar->GetValueOnRenderThread() : 1.0f) / 25000.0f;
+	RayGenParameters->TimingScale = GetRayTracingDebugTimingScale();
 
 	FIntRect ViewRect = View.ViewRect;
 	
diff --git a/Engine/Source/Runtime/Renderer/Private/RayTracing/RayTracingDebug.cpp b/Engine/Source/Runtime/Renderer/Private/RayTracing/RayTracingDebug.cpp
--- a/Engine/Source/Runtime/Renderer/Private/RayTracing/RayTracingDebug.cpp
+++ b/Engine/Source/Runtime/Renderer/Private/RayTracing/RayTracingDebug.cpp
@@ -153,6 +153,11 @@ static FAutoConsoleVariableRef CVarVisualizeProceduralPrimitives(
 
 static FAutoConsoleVariableDeprecated CVarVisualizeProceduralPrimitives_Deprecated(TEXT("r.RayTracing.DebugVisualizationMode.ProceduralPrimitives"), TEXT("r.RayTracing.Visualize.ProceduralPrimitives"), TEXT("5.6"));
 
+float GetRayTracingDebugTimingScale()
+{
+	return CVarRayTracingVisualizeTimingScale.GetValueOnRenderThread() / 25000.0f;
+}
+
 IMPLEMENT_RT_PAYLOAD_TYPE(ERayTracingPayloadType::RayTracingDebug, 44);
 
 BEGIN_UNIFORM_BUFFER_STRUCT(FRayTracingDebugHitStatsUniformBufferParameters, )
@@ -1728,7 +1733,7 @@ void RenderRayTracingDebug(FRDGBuilder& GraphBuilder, const FScene& Scene, const
 	SharedParameters.VisualizationMode = DebugVisualizationMode;
 	SharedParameters.PickerDomain = CVarRayTracingVisualizePickerDomain.GetValueOnRenderThread();
 	SharedParameters.ShouldUsePreExposure = View.Family->EngineShowFlags.Tonemapper;
-	SharedParameters.TimingScale = CVarRayTracingVisualizeTimingScale.GetValueOnAnyThread() / 25000.0f;
+	SharedParameters.TimingScale = GetRayTracingDebugTimingScale();
 	SharedParameters.OpaqueOnly = CVarRayTracingVisualizeOpaqueOnly.GetValueOnRenderThread();
 	SharedParameters.TriangleHitCountMaxThreshold = FMath::Clamp((float)CVarRayTracingVisualizeHitCountMaxThreshold.GetValueOnRenderThread(), 1, 100000);
 	SharedParameters.TriangleHitCountPerInstanceMaxThreshold = FMath::Max(1, CVarRayTracingVisualizeHitCountPerInstanceMaxThreshold.GetValueOnRenderThread());
