Path Tracer: Remove cloud map implementation and single-point callable shader

This is deprecated by the new mode which can perform cloud ray marching with a single shader call.

#rnx

diff --git a/Engine/Shaders/Private/PathTracing/PathTracingBuildCloud.usf b/Engine/Shaders/Private/PathTracing/PathTracingBuildCloud.usf
deleted file mode 100644
diff --git a/Engine/Shaders/Private/PathTracing/PathTracingBuildCloudAccelerationMap.usf b/Engine/Shaders/Private/PathTracing/PathTracingBuildCloudAccelerationMap.usf
--- a/Engine/Shaders/Private/PathTracing/PathTracingBuildCloudAccelerationMap.usf
+++ b/Engine/Shaders/Private/PathTracing/PathTracingBuildCloudAccelerationMap.usf
@@ -11,9 +11,7 @@
 uint NumSamples;
 uint TemporalSeed;
 uint Iteration;
-int UseCloudMap;
 
-Texture3D<uint4> CloudMap;
 RWTexture2D<float4> CloudAccelerationMap;
 
 [numthreads(THREADGROUPSIZE_X, THREADGROUPSIZE_Y, 1)]
@@ -62,22 +60,10 @@ void PathTracingBuildCloudAccelerationMapCS(uint2 DispatchThreadId : SV_Dispatch
 		float Z = ClipPos.z - CloudClipRadiusKm; // remember height of the point (before rotation, relative to center)
 		float H = length(ClipPos);
 
-		uint4 Data = 0;
-		float3 Extinction = 0;
-		if (UseCloudMap)
-		{
-			const float3 UVW = GetUVWFromCloudClipPos(ClipPos);
-			const int3 GridCoord = int3(floor(float3(CloudMapResolution) * UVW));
-			Data = CloudMap.Load(int4(GridCoord, 0));
-			Extinction = DecodeCloudResult(Data).Extinction;
-		}
-		else
-		{
-			float3 WorldPos = mul(ClipPos, GetCloudClipBasis()) + CloudClipCenterKm; // TangentToWorld
-			FDFVector3 AbsoluteWorldPosition = DFMultiply(WorldPos, SKY_UNIT_TO_CM);
-			// run the material and encode the result, which includes the effect that quantization will have during rendering
-			Extinction = SampleCloudMaterial(AbsoluteWorldPosition, H, B, T).Extinction;
-		}
+		float3 WorldPos = mul(ClipPos, GetCloudClipBasis()) + CloudClipCenterKm; // TangentToWorld
+		FDFVector3 AbsoluteWorldPosition = DFMultiply(WorldPos, SKY_UNIT_TO_CM);
+		// run the material and encode the result, which includes the effect that quantization will have during rendering
+		float3 Extinction = SampleCloudMaterial(AbsoluteWorldPosition, H, B, T).Extinction;
 		
 		const float Density   = max3(Extinction.x, Extinction.y, Extinction.z);
 		const float DensityLo = min3(Extinction.x, Extinction.y, Extinction.z);
diff --git a/Engine/Shaders/Private/PathTracing/PathTracingCommon.ush b/Engine/Shaders/Private/PathTracing/PathTracingCommon.ush
--- a/Engine/Shaders/Private/PathTracing/PathTracingCommon.ush
+++ b/Engine/Shaders/Private/PathTracing/PathTracingCommon.ush
@@ -431,34 +431,7 @@ struct FPackedPathTracingPayload : FMinimalPayload
 		PackedData[8] = asuint(SceneZ);
 	}
 
-
-#if CLOUD_CALLABLE_SHADERS_PER_POINT
-	// Volumetric callable shaders -- used when calling per point
-	void SetVolumetricCallableShaderInput(float3 TranslatedWorldPos, float CloudHeightKm)
-	{
-		PackedData[0] = asuint(TranslatedWorldPos.x);
-		PackedData[1] = asuint(TranslatedWorldPos.y);
-		PackedData[2] = asuint(TranslatedWorldPos.z);
-		PackedData[3] = asuint(CloudHeightKm);
-	}
-	float3 GetVolumetricCallableShaderInputTranslatedWorldPos()
-	{
-		return float3(asfloat(PackedData[0]), asfloat(PackedData[1]), asfloat(PackedData[2]));
-	}
-	float GetVolumetricCallableShaderInputCloudHeightKm() { return asfloat(PackedData[3]); }
-
-	void SetVolumetricCallableShaderOutput(uint4 Data)
-	{
-		PackedData[0] = Data.x;
-		PackedData[1] = Data.y;
-		PackedData[2] = Data.z;
-		PackedData[3] = Data.w;
-	}
-
-	uint4 GetVolumetricCallableShaderOutput() { return uint4(PackedData[0], PackedData[1], PackedData[2], PackedData[3]); }
-
-#else
-
+#if PATH_TRACER_USE_CLOUD_SHADER
 	// Volumetric callable shaders
 	void SetVolumetricCallableShaderInput(float3 Origin, float3 Direction, float TMin, float TMax, RandomSequence RandSeq, uint Flags, float3 Throughput, float CloudDensity)
 	{
@@ -484,7 +457,6 @@ struct FPackedPathTracingPayload : FMinimalPayload
 		PackedData[14] = asuint(SurfaceContext.RandSample);
 	}
 
-#if !COMPUTESHADER  // FRayDesc only defined for non compute shaders or inline ray tracing
 	FRayDesc GetVolumetricCallableShaderInputRay()
 	{
 		FRayDesc Ray;
@@ -494,7 +466,6 @@ struct FPackedPathTracingPayload : FMinimalPayload
 		Ray.TMax      = asfloat(PackedData[0]);
 		return Ray;
 	}
-#endif
 
 	RandomSequence GetVolumetricCallableShaderInputRandSequence()
 	{
@@ -606,7 +577,7 @@ struct FPackedPathTracingPayload : FMinimalPayload
 		return asfloat(PackedData[12]);
 	}
 
-#endif 
+#endif // PATH_TRACER_USE_CLOUD_SHADER
 
 #endif // IS_PAYLOAD_ENABLED(RT_PAYLOAD_TYPE_GPULIGHTMASS)
 };
diff --git a/Engine/Shaders/Private/PathTracing/PathTracingCore.ush b/Engine/Shaders/Private/PathTracing/PathTracingCore.ush
--- a/Engine/Shaders/Private/PathTracing/PathTracingCore.ush
+++ b/Engine/Shaders/Private/PathTracing/PathTracingCore.ush
@@ -730,7 +730,7 @@ FPackedPathTracingPayload TraceTransparentRay(inout FPathState PathState, int Bo
 				const bool bHaveHoldouts =
 					(Interval.VolumeMask << PATH_TRACER_MAX_VOLUMES) & VolumeFlags & PATH_TRACER_VOLUME_HOLDOUT_MASK != 0;
 
-#if CLOUD_CALLABLE_SHADERS_PER_POINT == 0 && PATH_TRACER_USE_CLOUD_SHADER
+#if PATH_TRACER_USE_CLOUD_SHADER
 				if (Interval.VolumeMask & PATH_TRACER_VOLUME_ENABLE_CLOUDS)
 				{
 					// let callable shader handle cloud sampling
@@ -785,7 +785,7 @@ FPackedPathTracingPayload TraceTransparentRay(inout FPathState PathState, int Bo
 					}
 				}
 				else
-#endif
+#endif // PATH_TRACER_USE_CLOUD_SHADER
 				// reasons to enter the ray marching loop:
 				//  - we can't use analytic transmittance (need ray marching for throughput or emission)
 				//  - we are a camera ray and some volumes are flagged as holdouts
diff --git a/Engine/Shaders/Private/PathTracing/PathTracingVolumetricCloudMaterialShader.usf b/Engine/Shaders/Private/PathTracing/PathTracingVolumetricCloudMaterialShader.usf
--- a/Engine/Shaders/Private/PathTracing/PathTracingVolumetricCloudMaterialShader.usf
+++ b/Engine/Shaders/Private/PathTracing/PathTracingVolumetricCloudMaterialShader.usf
@@ -1,6 +1,7 @@
 // Copyright Epic Games, Inc. All Rights Reserved.
 
 #define PATH_TRACING 1
+#define PATH_TRACER_USE_CLOUD_SHADER 1
 
 #include "/Engine/Private/Common.ush"
 #include "/Engine/Private/RayTracing/RayTracingCommon.ush"
@@ -56,19 +57,9 @@ RAY_TRACING_ENTRY_CALLABLE(PathTracingVolumetricCloudMaterialShader, FPackedPath
 {
 	ResolvedView = ResolveView();
 
-#if CLOUD_CALLABLE_SHADERS_PER_POINT
-	// Simple model, the Raygen contains all integrator logic and we just get called once per step to figure out the local cloud density
-
-	float3 TranslatedWorldPos = PackedPayload.GetVolumetricCallableShaderInputTranslatedWorldPos();
-	FDFVector3 AbsoluteWorldPosition = DFFastSubtract(TranslatedWorldPos, ResolvedView.PreViewTranslation);
-	float H = PackedPayload.GetVolumetricCallableShaderInputCloudHeightKm();
-	float B = PathTracingCloudParameters.CloudLayerBotKm;
-	float T = PathTracingCloudParameters.CloudLayerTopKm;
-
-	FSampleCloudMaterialResult Result = SampleCloudMaterial(AbsoluteWorldPosition, H, B, T);
-	PackedPayload.SetVolumetricCallableShaderOutput(EncodeCloudResult(Result));
-#else
-	// Advanced model - the ray marching happens inside here so that we only have to pay for the cost of callable shading _once_
+	// Ray marching happens inside here so that we only have to pay for the cost of callable shading _once_
+	// This comes at the cost of some duplication of the logic
+	// TODO: Can we use HLSL templates to merge some of the duplicated logic?
 
 	RandomSequence RandSeq = PackedPayload.GetVolumetricCallableShaderInputRandSequence();
 	FRayDesc Ray = PackedPayload.GetVolumetricCallableShaderInputRay();
@@ -76,9 +67,11 @@ RAY_TRACING_ENTRY_CALLABLE(PathTracingVolumetricCloudMaterialShader, FPackedPath
 	float CloudDensity = PackedPayload.GetVolumetricCallableShaderInputCloudDensity();
 	uint Flags = PackedPayload.GetVolumetricCallableShaderInputFlags();
 
-	if (Flags & PATH_TRACER_VOLUME_CALLABLE_FLAGS_TRANSMITTANCE) // TODO: evaluate if it would be better to make a second shader for this (so that we don't get two copies of the shader eval in inner loops)
+	// TODO: evaluate if it would be better to make a second shader for this (so that we don't get two copies of the shader eval in inner loops)
+	if (Flags & PATH_TRACER_VOLUME_CALLABLE_FLAGS_TRANSMITTANCE)
 	{
 		// Transmittance loop for the cloud shader only
+		// NOTE: We have specialized this since we know our upper bound is a scalar value, most of the pdf math cancels out and we can take a faster path
 		float TMin = Ray.TMin;
 		float TMax = Ray.TMax;
 		float SigmaBar = CloudDensity, InvSigmaBar = rcp(SigmaBar);
@@ -220,11 +213,9 @@ RAY_TRACING_ENTRY_CALLABLE(PathTracingVolumetricCloudMaterialShader, FPackedPath
 			}
 		}
 
-
 		// Write state back here
 		PackedPayload.SetVolumetricCallableShaderOutputRadiance(VolumeRadiance, VolumeAlpha);
 		PackedPayload.SetVolumetricCallableShaderOutput(RandSeq, PathThroughput, VolumeAlbedo);
 		PackedPayload.SetVolumetricCallableShaderRISContext(HitSample);
 	}
-#endif // CLOUD_CALLABLE_SHADERS_PER_POINT
 }
diff --git a/Engine/Shaders/Private/PathTracing/Volume/PathTracingClouds.ush b/Engine/Shaders/Private/PathTracing/Volume/PathTracingClouds.ush
--- a/Engine/Shaders/Private/PathTracing/Volume/PathTracingClouds.ush
+++ b/Engine/Shaders/Private/PathTracing/Volume/PathTracingClouds.ush
@@ -9,7 +9,6 @@
 
 Texture2D<float4> CloudAccelerationMap;
 SamplerState CloudAccelerationMapSampler;
-Texture3D<uint4> CloudMap;
 
 // intersect the cloud layer with high precision
 float4 RayCloudIntersectPrecise(float3 Ro, float3 Rd)
@@ -181,53 +180,21 @@ FVolumeIntersection CloudsIntersect(float3 Origin, float3 Direction, float TMin,
 	return CreateEmptyVolumeIntersection();
 }
 
+#if PATH_TRACER_USE_CLOUD_SHADER == 0
 FVolumeShadedResult CloudsGetDensity(float3 TranslatedWorldPos)
 {
-	FVolumeShadedResult Result = (FVolumeShadedResult)0;
-
-	if ((VolumeFlags & PATH_TRACER_VOLUME_VIEW_CLOUD_MAP_ACCEL) != 0)
-	{
-		// debug preview path to visualize the cloud accel map
-		float3 P = mul(GetCloudClipBasis(), TranslatedWorldPos); // WorldToTangent
-		float2 UV = (P.xy * CM_TO_SKY_UNIT + CloudClipDistKm.xx) / (2 * CloudClipDistKm);
-		const int Level = 0;
-		float4 Data = CloudAccelerationMap.SampleLevel(CloudAccelerationMapSampler, UV, Level);
+	// debug preview path to visualize the cloud accel map
+	float3 P = mul(GetCloudClipBasis(), TranslatedWorldPos); // WorldToTangent
+	float2 UV = (P.xy * CM_TO_SKY_UNIT + CloudClipDistKm.xx) / (2 * CloudClipDistKm);
+	const int Level = 0;
+	float4 Data = CloudAccelerationMap.SampleLevel(CloudAccelerationMapSampler, UV, Level);
 	
-		float Z = P.z * CM_TO_SKY_UNIT;
-		float Density = (Z > Data.z && Z < Data.w) ? Data.x : 0.0;
-		Result.SigmaT = Density;
-		Result.SigmaSHG = 0.75 * Density;
-		return Result;
-	}
-
-	uint4 Data = 0;
-#if CLOUD_CALLABLE_SHADERS_PER_POINT
-#if PATH_TRACER_USE_CLOUD_SHADER
-	{
-		// Reference path -- invoke the actual callable shader
-		FPackedPathTracingPayload Payload = (FPackedPathTracingPayload)0;
-		Payload.SetVolumetricCallableShaderInput(
-			TranslatedWorldPos,
-			GetSphericalPlanetParameters(TranslatedWorldPos, float3(0, 0, 1)).HeightKm
-		);
-		CallShader(CloudCallableShaderId, Payload);
-		Data = Payload.GetVolumetricCallableShaderOutput();
-	}
-#else
-	{
-		// Main path -- via baked cloud map
-		const float3 P = mul(GetCloudClipBasis(), TranslatedWorldPos * CM_TO_SKY_UNIT); // WorldToTangent
-		const float3 UVW = GetUVWFromCloudClipPos(P + float3(0, 0, CloudClipRadiusKm));
-		const int3 GridCoord = int3(floor(float3(CloudMapResolution) * UVW));
-		Data = CloudMap.Load(int4(GridCoord, 0));
-	}
-#endif
-#endif
+	float Z = P.z * CM_TO_SKY_UNIT;
+	float Density = (Z > Data.z && Z < Data.w) ? Data.x : 0.0;
+	FVolumeShadedResult Result = (FVolumeShadedResult)0;
 
-	const FSampleCloudMaterialResult CloudResult = DecodeCloudResult(Data);
-	Result.SigmaT        = CloudResult.Extinction;
-	Result.SigmaSDualHG  = CloudResult.Albedo * Result.SigmaT;
-	Result.Emission      = CloudResult.EmissiveColor;
-	Result.DualPhaseData = float3(CloudResult.PhaseG1, CloudResult.PhaseG2, CloudResult.PhaseBlend);
+	Result.SigmaT = Density;
+	Result.SigmaSHG = 0.85 * Density;
 	return Result;
 }
+#endif
\ No newline at end of file
diff --git a/Engine/Shaders/Private/PathTracing/Volume/PathTracingCloudsCommon.ush b/Engine/Shaders/Private/PathTracing/Volume/PathTracingCloudsCommon.ush
--- a/Engine/Shaders/Private/PathTracing/Volume/PathTracingCloudsCommon.ush
+++ b/Engine/Shaders/Private/PathTracing/Volume/PathTracingCloudsCommon.ush
@@ -19,7 +19,6 @@ float CloudTracingMaxDistance;
 float CloudRoughnessCutoff;
 int CloudAccelMapResolution;
 int CloudCallableShaderId;
-uint3 CloudMapResolution;
 
 float3x3 GetCloudClipBasis()
 {
@@ -51,82 +50,6 @@ float2 GetRadiusBounds()
 	return RBounds;
 }
 
-#define CLOUD_MAPPING_TYPE   2
-
-
-float3 GetCloudClipPosFromUVW(float3 UVW)
-{
-#if CLOUD_MAPPING_TYPE == 0
-	// super basic voxelization for debugging
-	const float3 Corner0 = float3(-CloudClipDistKm.xx, GetLowestCloudZ());
-	const float3 Corner1 = float3(+CloudClipDistKm.xx, CloudLayerTopKm  );
-	return lerp(Corner0, Corner1, UVW);
-#elif CLOUD_MAPPING_TYPE == 1 || CLOUD_MAPPING_TYPE == 2
-	const float3 Corner0 = float3(-CloudClipDistKm.xx, GetLowestCloudZ());
-	const float3 Corner1 = float3(+CloudClipDistKm.xx, CloudLayerTopKm  );
-
-
-	float SliceDist = abs(GetLowestCloudZ() - CloudClipRadiusKm);
-
-	float2 Angle0 = atan2(Corner0.xy, SliceDist);
-	float2 Angle1 = atan2(Corner1.xy, SliceDist);
-	float3 CloudClipPos = float3(SliceDist * tan(lerp(Angle0, Angle1, UVW.xy)), 0.0);
-
-#if CLOUD_MAPPING_TYPE == 1
-	// angular mapping
-	const float ColumnDist = length(CloudClipPos.xy);
-	const float AngleZ0 = atan2(Corner0.z, ColumnDist);
-	const float AngleZ1 = atan2(Corner1.z, ColumnDist);
-	CloudClipPos.z = ColumnDist * tan(lerp(AngleZ0, AngleZ1, UVW.z));
-#else
-	// linear mapping
-	CloudClipPos.z = lerp(Corner0.z, Corner1.z, UVW.z);
-#endif
-
-
-	return CloudClipPos;
-
-#else
-#error "Unknown cloud mapping type"
-#endif
-}
-
-float3 GetUVWFromCloudClipPos(float3 CloudClipPos)
-{
-#if CLOUD_MAPPING_TYPE == 0
-	// super basic voxelization for debugging
-	const float3 Corner0 = float3(-CloudClipDistKm.xx, GetLowestCloudZ());
-	const float3 Corner1 = float3(+CloudClipDistKm.xx, CloudLayerTopKm  );
-	return saturate((CloudClipPos - Corner0) / (Corner1 - Corner0)) * 0.99999994;
-#elif CLOUD_MAPPING_TYPE == 1 || CLOUD_MAPPING_TYPE == 2
-	const float3 Corner0 = float3(-CloudClipDistKm.xx, GetLowestCloudZ());
-	const float3 Corner1 = float3(+CloudClipDistKm.xx, CloudLayerTopKm  );
-
-	float SliceDist = abs(GetLowestCloudZ() - CloudClipRadiusKm);
-
-	float2 Angle0 = atan2(Corner0.xy, SliceDist);
-	float2 Angle1 = atan2(Corner1.xy, SliceDist);
-
-	float2 UV = (atan2(CloudClipPos.xy, SliceDist) - Angle0) / (Angle1 - Angle0);
-
-#if CLOUD_MAPPING_TYPE == 1
-	// angular mapping
-	const float ColumnDist = length(CloudClipPos.xy);
-	const float AngleZ0 = atan2(Corner0.z, ColumnDist);
-	const float AngleZ1 = atan2(Corner1.z, ColumnDist);
-	const float W = (atan2(CloudClipPos.z, ColumnDist) - AngleZ0) / (AngleZ1 - AngleZ0);
-#else
-	// linear mapping
-	const float W = (CloudClipPos.z - Corner0.z) / (Corner1.z - Corner0.z);
-#endif
-
-	return saturate(float3(UV, W));
-#else
-#error "Unknown cloud mapping type"
-#endif
-}
-
-
 
 // Simplified struct only used in the cloud case (to reduce the amount of data needed in the payload for the callable shader)
 struct FSampleCloudMaterialResult
diff --git a/Engine/Shaders/Private/PathTracing/Volume/PathTracingVolume.ush b/Engine/Shaders/Private/PathTracing/Volume/PathTracingVolume.ush
--- a/Engine/Shaders/Private/PathTracing/Volume/PathTracingVolume.ush
+++ b/Engine/Shaders/Private/PathTracing/Volume/PathTracingVolume.ush
@@ -75,10 +75,13 @@ FVolumeShadedResult VolumeGetDensity(float3 TranslatedWorldPos, uint VolumeMask,
 	{
 		MergeVolumeShadedResult(Result, AtmosphereGetDensity(TranslatedWorldPos), bIsCameraRay && (VolumeFlags & PATH_TRACER_VOLUME_HOLDOUT_ATMOSPHERE) != 0);
 	}
+#if PATH_TRACER_USE_CLOUD_SHADER == 0
+	// If we are using the CallableShader Path, this function is not called with clouds in VolumeMask
 	if (VolumeMask & PATH_TRACER_VOLUME_ENABLE_CLOUDS)
 	{
 		MergeVolumeShadedResult(Result, CloudsGetDensity(TranslatedWorldPos), bIsCameraRay && (VolumeFlags & PATH_TRACER_VOLUME_HOLDOUT_CLOUDS) != 0);
 	}
+#endif
 	if (VolumeMask & PATH_TRACER_VOLUME_ENABLE_FOG)
 	{
 		MergeVolumeShadedResult(Result, FogGetDensity(TranslatedWorldPos), bIsCameraRay && (VolumeFlags & PATH_TRACER_VOLUME_HOLDOUT_FOG) != 0);
@@ -172,7 +175,7 @@ float3 VolumeGetTransmittance(float3 StartThroughput, float3 Origin, float3 Dire
 		float TMin = Interval.VolumeTMin;
 		float TMax = Interval.VolumeTMax;
 		float SigmaBar = Interval.CloudDensity;
-#if CLOUD_CALLABLE_SHADERS_PER_POINT == 0 && PATH_TRACER_USE_CLOUD_SHADER
+#if PATH_TRACER_USE_CLOUD_SHADER
 		FPackedPathTracingPayload CloudPayload = (FPackedPathTracingPayload) 0;
 		CloudPayload.SetVolumetricCallableShaderInput(Origin, Direction, TMin, TMax, RandSequence, PATH_TRACER_VOLUME_CALLABLE_FLAGS_TRANSMITTANCE, Throughput, SigmaBar);
 		CallShader(CloudCallableShaderId, CloudPayload);
@@ -226,7 +229,7 @@ float3 VolumeGetTransmittance(float3 StartThroughput, float3 Origin, float3 Dire
 		float TMin = IntersectionList.VolumeTMin[VOLUMEID_CLOUDS];
 		float TMax = IntersectionList.VolumeTMax[VOLUMEID_CLOUDS];
 		float SigmaBar = IntersectionList.CloudDensity;
-#if CLOUD_CALLABLE_SHADERS_PER_POINT == 0 && PATH_TRACER_USE_CLOUD_SHADER
+#if PATH_TRACER_USE_CLOUD_SHADER
 		FPackedPathTracingPayload CloudPayload = (FPackedPathTracingPayload) 0;
 		CloudPayload.SetVolumetricCallableShaderInput(Origin, Direction, TMin, TMax, RandSequence, PATH_TRACER_VOLUME_CALLABLE_FLAGS_TRANSMITTANCE, Throughput, SigmaBar);
 		CallShader(CloudCallableShaderId, CloudPayload);
diff --git a/Engine/Shaders/Shared/PathTracingDefinitions.h b/Engine/Shaders/Shared/PathTracingDefinitions.h
--- a/Engine/Shaders/Shared/PathTracingDefinitions.h
+++ b/Engine/Shaders/Shared/PathTracingDefinitions.h
@@ -68,16 +68,11 @@
 #define PATH_TRACER_VOLUME_HOLDOUT_MASK                     (PATH_TRACER_VOLUME_ENABLE_MASK << PATH_TRACER_MAX_VOLUMES)
 
 #define PATH_TRACER_VOLUME_USE_ANALYTIC_TRANSMITTANCE       (PATH_TRACER_VOLUME_HOLDOUT_BIT << PATH_TRACER_MAX_VOLUMES)
-#define PATH_TRACER_VOLUME_USE_CLOUD_MAP                    (PATH_TRACER_VOLUME_USE_ANALYTIC_TRANSMITTANCE << 1)
-#define PATH_TRACER_VOLUME_VIEW_CLOUD_MAP_ACCEL             (PATH_TRACER_VOLUME_USE_ANALYTIC_TRANSMITTANCE << 2)
 
 // These flags are reserved for passing information to the cloud callable shader
-#define PATH_TRACER_VOLUME_CALLABLE_FLAGS_TRANSMITTANCE     (PATH_TRACER_VOLUME_VIEW_CLOUD_MAP_ACCEL << 1) // only compute transmittance for the cloud portion
-#define PATH_TRACER_VOLUME_CALLABLE_FLAGS_GET_SAMPLE        (PATH_TRACER_VOLUME_VIEW_CLOUD_MAP_ACCEL << 2) // perform RIS sampling for a scatter location (else do transmittance/emission/alpha calc only)
-#define PATH_TRACER_VOLUME_CALLABLE_FLAGS_IS_CAMERA_RAY     (PATH_TRACER_VOLUME_VIEW_CLOUD_MAP_ACCEL << 3) // this is a camera ray (so apply holdouts if enabled)
-
-
-#define CLOUD_CALLABLE_SHADERS_PER_POINT 0 // TODO: remove the per-point codepath
+#define PATH_TRACER_VOLUME_CALLABLE_FLAGS_TRANSMITTANCE     (PATH_TRACER_VOLUME_USE_ANALYTIC_TRANSMITTANCE << 1) // only compute transmittance for the cloud portion
+#define PATH_TRACER_VOLUME_CALLABLE_FLAGS_GET_SAMPLE        (PATH_TRACER_VOLUME_USE_ANALYTIC_TRANSMITTANCE << 2) // perform RIS sampling for a scatter location (else do transmittance/emission/alpha calc only)
+#define PATH_TRACER_VOLUME_CALLABLE_FLAGS_IS_CAMERA_RAY     (PATH_TRACER_VOLUME_USE_ANALYTIC_TRANSMITTANCE << 3) // this is a camera ray (so apply holdouts if enabled)
 
 
 // Constants related to debugging
diff --git a/Engine/Source/Runtime/Renderer/Private/PathTracing.cpp b/Engine/Source/Runtime/Renderer/Private/PathTracing.cpp
--- a/Engine/Source/Runtime/Renderer/Private/PathTracing.cpp
+++ b/Engine/Source/Runtime/Renderer/Private/PathTracing.cpp
@@ -296,30 +296,6 @@ TAutoConsoleVariable<float> CVarPathTracingCloudRoughnessCutoff(
 	ECVF_RenderThreadSafe
 );
 
-TAutoConsoleVariable<bool> CVarPathTracingCloudMapEnabled(
-	TEXT("r.PathTracing.CloudMapEnable"),
-	true,
-	TEXT("Controls the use of a texture to cache the cloud material.\n")
-	TEXT(" true : Clouds will be voxelized into a texture for faster evaluation (default)\n")
-	TEXT(" false: Cloud material will be invoked per ray march step via callable shaders (more accurate, but more expensive)\n")
-	TEXT("Requires r.PathTracing.Experimental=true to modify.\n"),
-	ECVF_RenderThreadSafe
-);
-
-TAutoConsoleVariable<int32> CVarPathTracingCloudMapResolution(
-	TEXT("r.PathTracing.CloudMapResolution"),
-	512,
-	TEXT("Size of the texture used to cache volumetric cloud appearance in reference atmosphere mode.  (default = 512)"),
-	ECVF_RenderThreadSafe
-);
-
-TAutoConsoleVariable<int32> CVarPathTracingCloudMapDepth(
-	TEXT("r.PathTracing.CloudMapDepth"),
-	64,
-	TEXT("Depth of the texture used to cache volumetric cloud appearance in reference atmosphere mode.  (default = 64)"),
-	ECVF_RenderThreadSafe
-);
-
 TAutoConsoleVariable<float> CVarPathTracingFogDensityClamp(
 	TEXT("r.PathTracing.FogDensityClamp"),
 	8.0f,
@@ -589,8 +565,7 @@ struct FPathTracingConfig
 	float AdaptiveSamplingThreshold;
 	int CloudAccelerationMapNumSamples;
 	int CloudAccelerationMapResolution;
-	int CloudMapResolution;
-	int CloudMapDepth;
+	bool CloudAccelerationMapVisualize;
 
 	bool IsDifferent(const FPathTracingConfig& Other) const
 	{
@@ -630,8 +605,7 @@ struct FPathTracingConfig
 			AdaptiveSamplingThreshold != Other.AdaptiveSamplingThreshold ||
 			CloudAccelerationMapNumSamples != Other.CloudAccelerationMapNumSamples ||
 			CloudAccelerationMapResolution != Other.CloudAccelerationMapResolution ||
-			CloudMapResolution != Other.CloudMapResolution ||
-			CloudMapDepth != Other.CloudMapDepth ||
+			CloudAccelerationMapVisualize != Other.CloudAccelerationMapVisualize ||
 			UseMultiGPU != Other.UseMultiGPU;
 	}
 
@@ -702,7 +676,6 @@ struct FPathTracingState
 	TRefCountPtr<IPooledRenderTarget> DepthRT;
 	TRefCountPtr<FRDGPooledBuffer> VarianceBuffer;
 	TRefCountPtr<IPooledRenderTarget> CloudAccelerationMap;
-	TRefCountPtr<IPooledRenderTarget> CloudMap;
 
 	// Cache to improve the stability when frame denoising (SPP=r.pathtracing.SamplesPerPixel) is used in animation rendering
 	TRefCountPtr<IPooledRenderTarget> LastDenoisedRadianceRT;
@@ -865,11 +838,10 @@ static void PreparePathTracingData(const FScene* Scene, const FViewInfo& View, F
 	
 	const FVolumetricCloudRenderSceneInfo* CloudInfo = Scene->GetVolumetricCloudSceneInfo();
 	const bool bExperimental = CVarPathTracingExperimental.GetValueOnRenderThread();
-	const bool bVolumeCloudMapEnabled = !ShouldCompileRayTracingCallableShadersForProject(View.GetShaderPlatform()) || (bExperimental == false || CVarPathTracingCloudMapEnabled.GetValueOnRenderThread());
 	const bool bVolumeCloudsVisible = ShouldRenderVolumetricCloud(Scene, ShowFlags)
 		&& bUseReferenceAtmosphere
 		&& CVarPathTracingEnableReferenceClouds.GetValueOnRenderThread() != 0
-		&& (bVolumeCloudMapEnabled || uint32(View.PathTracingVolumetricCloudCallableShaderIndex) < Scene->RayTracingSBT.NumCallableShaderSlots);
+		&& (uint32(View.PathTracingVolumetricCloudCallableShaderIndex) < Scene->RayTracingSBT.NumCallableShaderSlots);
 
 	// Merge all volume flags into one uint
 	PathTracingData.VolumeFlags = 0;
@@ -888,8 +860,6 @@ static void PreparePathTracingData(const FScene* Scene, const FViewInfo& View, F
 	PathTracingData.VolumeFlags |= bVolumeCloudsVisible && IsVolumetricCloudHoldout(View.CachedViewUniformShaderParameters->EnvironmentComponentsFlags) ? PATH_TRACER_VOLUME_HOLDOUT_CLOUDS : 0;
 	PathTracingData.VolumeFlags |= Scene->ExponentialFogs.Num() > 0 && IsExponentialFogHoldout(View.CachedViewUniformShaderParameters->EnvironmentComponentsFlags) ? PATH_TRACER_VOLUME_HOLDOUT_FOG : 0;
 	PathTracingData.VolumeFlags |= CVarPathTracingUseAnalyticTransmittance.GetValueOnRenderThread() ? PATH_TRACER_VOLUME_USE_ANALYTIC_TRANSMITTANCE : 0;
-	PathTracingData.VolumeFlags |= bVolumeCloudMapEnabled ? PATH_TRACER_VOLUME_USE_CLOUD_MAP : 0;
-	PathTracingData.VolumeFlags |= CVarPathTracingCloudAccelerationMapVisualize.GetValueOnRenderThread() ? PATH_TRACER_VOLUME_VIEW_CLOUD_MAP_ACCEL : 0;
 
 	PathTracingData.CloudRoughnessCutoff = bVolumeCloudsVisible ? CVarPathTracingCloudRoughnessCutoff.GetValueOnRenderThread() : -1.0f;
 
@@ -1157,7 +1127,6 @@ BEGIN_SHADER_PARAMETER_STRUCT(FPathTracingCloudParameters, )
 	SHADER_PARAMETER(float, CloudTracingMaxDistance) // limit ray lengths (to avoid slowing down when a ray crosses all clouds)
 	SHADER_PARAMETER(int32, CloudAccelMapResolution)
 	SHADER_PARAMETER(int32, CloudCallableShaderId)
-	SHADER_PARAMETER(FIntVector, CloudMapResolution)
 END_SHADER_PARAMETER_STRUCT()
 
 BEGIN_GLOBAL_SHADER_PARAMETER_STRUCT(FPathTracingCloudParameterGlobals,)
@@ -1293,7 +1262,6 @@ class FPathTracingRG : public FGlobalShader
 		SHADER_PARAMETER_STRUCT_INCLUDE(FPathTracingCloudParameters, CloudParameters)
 		SHADER_PARAMETER_RDG_TEXTURE(Texture2D, CloudAccelerationMap)
 		SHADER_PARAMETER_SAMPLER(SamplerState, CloudAccelerationMapSampler)
-		SHADER_PARAMETER_RDG_TEXTURE(Texture3D, CloudMap)
 
 		// exponential height fog
 		SHADER_PARAMETER_STRUCT_INCLUDE(FPathTracingFogParameters, FogParameters)
@@ -1553,7 +1521,6 @@ FPathTracingCloudParameters PrepareCloudParameters(const FScene* Scene, const FV
 	if (Config != nullptr)
 	{
 		Params.CloudAccelMapResolution = Config->CloudAccelerationMapResolution;
-		Params.CloudMapResolution = FIntVector(Config->CloudMapResolution, Config->CloudMapResolution, Config->CloudMapDepth);
 	}
 	return Params;
 }
@@ -1565,7 +1532,10 @@ class FPathTracingBuildCloudAccelerationMapCS : public FMeshMaterialShader
 
 	static bool ShouldCompilePermutation(const FMeshMaterialShaderPermutationParameters& Parameters)
 	{
+		const bool bUseExperimental = CVarPathTracingExperimental.GetValueOnAnyThread();
+
 		return ShouldCompilePathTracingShadersForProject(Parameters.Platform) &&
+			bUseExperimental &&
 			Parameters.MaterialParameters.bIsUsedWithVolumetricCloud &&
 			Parameters.MaterialParameters.MaterialDomain == MD_Volume;
 	}
@@ -1583,49 +1553,13 @@ class FPathTracingBuildCloudAccelerationMapCS : public FMeshMaterialShader
 		SHADER_PARAMETER(uint32, NumSamples)
 		SHADER_PARAMETER(uint32, Iteration)
 		SHADER_PARAMETER(uint32, TemporalSeed)
-		SHADER_PARAMETER(int32 , UseCloudMap)
 		SHADER_PARAMETER_STRUCT_INCLUDE(FPathTracingCloudParameters, CloudParameters)
 		SHADER_PARAMETER_STRUCT_REF(FViewUniformShaderParameters, ViewUniformBuffer)
 		SHADER_PARAMETER_RDG_TEXTURE_UAV(RWTexture2D, CloudAccelerationMap)
-		SHADER_PARAMETER_RDG_TEXTURE(Texture3D, CloudMap)
 	END_SHADER_PARAMETER_STRUCT()
 };
 IMPLEMENT_MATERIAL_SHADER_TYPE(, FPathTracingBuildCloudAccelerationMapCS, TEXT("/Engine/Private/PathTracing/PathTracingBuildCloudAccelerationMap.usf"), TEXT("PathTracingBuildCloudAccelerationMapCS"), SF_Compute);
 
-class FPathTracingBuildCloudCS : public FMeshMaterialShader
-{
-	DECLARE_SHADER_TYPE(FPathTracingBuildCloudCS, MeshMaterial)
-	SHADER_USE_PARAMETER_STRUCT_WITH_LEGACY_BASE(FPathTracingBuildCloudCS, FMeshMaterialShader) // TODO: following pattern used in VolumetricCloudRender.cpp -- is there a proper modern alternative?
-
-	static bool ShouldCompilePermutation(const FMeshMaterialShaderPermutationParameters& Parameters)
-	{
-		return ShouldCompilePathTracingShadersForProject(Parameters.Platform) &&
-		Parameters.MaterialParameters.bIsUsedWithVolumetricCloud &&
-		Parameters.MaterialParameters.MaterialDomain == MD_Volume;
-	}
-
-	static void ModifyCompilationEnvironment(const FMaterialShaderPermutationParameters& Parameters, FShaderCompilerEnvironment& OutEnvironment)
-	{
-		FMaterialShader::ModifyCompilationEnvironment(Parameters, OutEnvironment);
-
-		OutEnvironment.SetDefine(TEXT("THREADGROUPSIZE_X"), 4);
-		OutEnvironment.SetDefine(TEXT("THREADGROUPSIZE_Y"), 4);
-		OutEnvironment.SetDefine(TEXT("THREADGROUPSIZE_Z"), 4);
-		OutEnvironment.SetDefine(TEXT("CLOUD_LAYER_PIXEL_SHADER"), 1);
-	}
-
-	BEGIN_SHADER_PARAMETER_STRUCT(FParameters, )
-		SHADER_PARAMETER(float , BlendFactor)
-		SHADER_PARAMETER(uint32, Iteration)
-		SHADER_PARAMETER(uint32, TemporalSeed)
-		SHADER_PARAMETER_STRUCT_INCLUDE(FPathTracingCloudParameters, CloudParameters)
-		SHADER_PARAMETER_STRUCT_REF(FViewUniformShaderParameters, ViewUniformBuffer)
-		SHADER_PARAMETER_RDG_TEXTURE_UAV(RWTexture3D, CloudMap)
-	END_SHADER_PARAMETER_STRUCT()
-};
-IMPLEMENT_MATERIAL_SHADER_TYPE(, FPathTracingBuildCloudCS, TEXT("/Engine/Private/PathTracing/PathTracingBuildCloud.usf"), TEXT("PathTracingBuildCloudCS"), SF_Compute);
-
-
 class FPathTracingVolumetricCloudMaterial : public FMaterialShader
 {
 	DECLARE_SHADER_TYPE(FPathTracingVolumetricCloudMaterial, Material);
@@ -1723,7 +1657,7 @@ void PreparePathTracingCloudMaterial(FScene* Scene, TArrayView<FViewInfo> Views)
 
 	// if we are using reference clouds, or if we are using the cloud map, no need to prepare the callable shader version
 	if (CVarPathTracingEnableReferenceClouds.GetValueOnRenderThread() == 0 ||
-		(CVarPathTracingExperimental.GetValueOnRenderThread() == false || CVarPathTracingCloudMapEnabled.GetValueOnRenderThread()))
+		CVarPathTracingExperimental.GetValueOnRenderThread() == false)
 	{
 		return;
 	}
@@ -3043,7 +2977,11 @@ static FPathTracingRG::FPermutationDomain GetPathTracingRGPermutation(const FSce
 {
 	const bool bUseExperimental = CVarPathTracingExperimental.GetValueOnRenderThread();
 	const bool bUseAdaptiveSampling = bUseExperimental && CVarPathTracingAdaptiveSampling.GetValueOnRenderThread() != 0;
-	const bool bUseCloudShader = bUseExperimental && ShouldCompileRayTracingCallableShadersForProject(Scene.GetShaderPlatform()) && !CVarPathTracingCloudMapEnabled.GetValueOnRenderThread();
+	const bool bUseCloudShader = bUseExperimental
+		&& ShouldCompileRayTracingCallableShadersForProject(Scene.GetShaderPlatform())
+		&& CVarPathTracingEnableReferenceClouds.GetValueOnRenderThread()
+		&& !CVarPathTracingCloudAccelerationMapVisualize.GetValueOnRenderThread()
+		&& Scene.GetVolumetricCloudSceneInfo() != nullptr;
 	const bool bHasComplexSpecialRenderPath = Substrate::IsSubstrateEnabled() && Scene.SubstrateSceneData.bUsesComplexSpecialRenderPath;
 	const bool bUseSER = GRHIGlobals.SupportsShaderExecutionReordering && CVarPathTracingShaderExecutionReordering.GetValueOnRenderThread();
 	const bool bUseTraceOpaqueFirst = CVarPathTracingTraceOpaqueFirst.GetValueOnRenderThread();
@@ -3177,7 +3115,6 @@ void FSceneViewState::PathTracingInvalidate(bool InvalidateAnimationStates)
 		State->DepthRT.SafeRelease();
 		State->VarianceBuffer.SafeRelease();
 		State->CloudAccelerationMap.SafeRelease();
-		State->CloudMap.SafeRelease();
 		State->SampleIndex = 0;
 
 		State->AdaptiveFrustumGridParameterCache.TopLevelGridBuffer.SafeRelease();
@@ -3347,8 +3284,7 @@ void FDeferredShadingSceneRenderer::RenderPathTracing(
 	Config.AdaptiveSamplingThreshold = CVarPathTracingAdaptiveSamplingErrorThreshold.GetValueOnRenderThread();
 	Config.CloudAccelerationMapNumSamples = FMath::Clamp(CVarPathTracingCloudAccelerationMapNumSamples.GetValueOnRenderThread(), 1, 65536);
 	Config.CloudAccelerationMapResolution = FMath::Clamp(CVarPathTracingCloudAccelerationMapResolution.GetValueOnRenderThread(), 1, 4096);
-	Config.CloudMapResolution = FMath::Clamp(CVarPathTracingCloudMapResolution.GetValueOnRenderThread(), 1, 4096);
-	Config.CloudMapDepth = FMath::Clamp(CVarPathTracingCloudMapDepth.GetValueOnRenderThread(), 1, 256);
+	Config.CloudAccelerationMapVisualize = CVarPathTracingCloudAccelerationMapVisualize.GetValueOnRenderThread();
 
 	// compute an integer code of what show flags and booleans related to lights are currently enabled so we can detect changes
 	Config.LightShowFlags = 0;
@@ -3677,10 +3613,8 @@ void FDeferredShadingSceneRenderer::RenderPathTracing(
 			}
 
 			FRDGTexture* CloudAccelerationMap = nullptr;
-			FRDGTexture* CloudMap = nullptr;
 
 			const bool bEnableClouds   = (Config.PathTracingData.VolumeFlags & PATH_TRACER_VOLUME_ENABLE_CLOUDS) != 0;
-			const bool bEnableCloudMap = (Config.PathTracingData.VolumeFlags & PATH_TRACER_VOLUME_USE_CLOUD_MAP) != 0;
 
 			if (bEnableClouds)
 			{
@@ -3705,26 +3639,6 @@ void FDeferredShadingSceneRenderer::RenderPathTracing(
 						TexCreate_ShaderResource | TexCreate_UAV);
 					CloudAccelerationMap = GraphBuilder.CreateTexture(Desc, TEXT("PathTracer.CloudAccelerationMap"), ERDGTextureFlags::None);
 				}
-				if (bEnableCloudMap)
-				{
-					if (PathTracingState->CloudMap)
-					{
-						CloudMap = GraphBuilder.RegisterExternalTexture(PathTracingState->CloudMap, TEXT("PathTracer.CloudMap"));
-					}
-					else
-					{
-						FRDGTextureDesc Desc = FRDGTextureDesc::Create3D(
-							FIntVector(Config.CloudMapResolution, Config.CloudMapResolution, Config.CloudMapDepth),
-							PF_R32G32B32A32_UINT,
-							FClearValueBinding::None,
-							TexCreate_ShaderResource | TexCreate_UAV);
-						CloudMap = GraphBuilder.CreateTexture(Desc, TEXT("PathTracer.CloudMap"), ERDGTextureFlags::None);
-					}
-				}
-				else
-				{
-					CloudMap = GraphBuilder.RegisterExternalTexture(GSystemTextures.VolumetricBlackUintDummy);
-				}
 				FVolumetricCloudRenderSceneInfo* CloudRenderSceneInfo = Scene->GetVolumetricCloudSceneInfo();
 				check(CloudRenderSceneInfo != nullptr);
 				UMaterialInterface* CloudMaterialInterface = CloudRenderSceneInfo->GetVolumetricCloudSceneProxy().GetCloudVolumeMaterial();
@@ -3735,38 +3649,6 @@ void FDeferredShadingSceneRenderer::RenderPathTracing(
 				
 
 				const FPathTracingCloudParameters CloudParameters = PrepareCloudParameters(Scene, View, &Config);
-				// build cloud map
-				if (bEnableCloudMap)
-				{
-					typename FPathTracingBuildCloudCS::FPermutationDomain PermutationVector;
-					TShaderRef<FPathTracingBuildCloudCS> ComputeShader = MaterialResource->GetShader<FPathTracingBuildCloudCS>(&FLocalVertexFactory::StaticType, PermutationVector, false);
-
-					FPathTracingBuildCloudCS::FParameters* PassParameters = GraphBuilder.AllocParameters<FPathTracingBuildCloudCS::FParameters>();
-					PassParameters->BlendFactor = Config.PathTracingData.BlendFactor;;
-					PassParameters->Iteration = Config.PathTracingData.Iteration;
-					PassParameters->TemporalSeed = Config.PathTracingData.TemporalSeed;
-					PassParameters->CloudParameters = CloudParameters;
-					PassParameters->ViewUniformBuffer = View.ViewUniformBuffer;
-					PassParameters->CloudMap = GraphBuilder.CreateUAV(CloudMap);
-
-					check(CloudMap->Desc.GetSize() == CloudParameters.CloudMapResolution);
-
-					FIntVector GroupCount = FComputeShaderUtils::GetGroupCount(
-						CloudMap->Desc.GetSize(),
-						FIntVector(4, 4, 4)
-					);
-					GraphBuilder.AddPass(
-						RDG_EVENT_NAME("Path Tracing Cloud Build (%dx%dx%d)", CloudMap->Desc.GetSize().X, CloudMap->Desc.GetSize().Y, CloudMap->Desc.GetSize().Z),
-						PassParameters,
-						ERDGPassFlags::Compute,
-						[Scene = Scene, CloudVolumeMaterialProxy, MaterialResource, PassParameters, ComputeShader, GroupCount](FRDGAsyncTask, FRHIComputeCommandList& RHICmdList)
-						{
-							FMeshDrawShaderBindings ShaderBindings;
-							UE::MeshPassUtils::SetupComputeBindings(ComputeShader, Scene, Scene->GetFeatureLevel(), nullptr, *CloudVolumeMaterialProxy, *MaterialResource, ShaderBindings);
-							UE::MeshPassUtils::Dispatch(RHICmdList, ComputeShader, ShaderBindings, *PassParameters, GroupCount);
-						});
-					GraphBuilder.QueueTextureExtraction(CloudMap, &PathTracingState->CloudMap);
-				}
 				// build cloud accel map
 				{
 					typename FPathTracingBuildCloudAccelerationMapCS::FPermutationDomain PermutationVector;
@@ -3779,8 +3661,6 @@ void FDeferredShadingSceneRenderer::RenderPathTracing(
 					PassParameters->CloudParameters = CloudParameters;
 					PassParameters->ViewUniformBuffer = View.ViewUniformBuffer;
 					PassParameters->CloudAccelerationMap = GraphBuilder.CreateUAV(CloudAccelerationMap);
-					PassParameters->UseCloudMap = bEnableCloudMap;
-					PassParameters->CloudMap = CloudMap;
 
 					FIntVector GroupCount = FComputeShaderUtils::GetGroupCount(
 						FIntPoint(Resolution, Resolution),
@@ -3802,10 +3682,6 @@ void FDeferredShadingSceneRenderer::RenderPathTracing(
 			{
 				CloudAccelerationMap = GraphBuilder.RegisterExternalTexture(GSystemTextures.BlackDummy);
 			}
-			if (CloudMap == nullptr)
-			{
-				CloudMap = GraphBuilder.RegisterExternalTexture(GSystemTextures.VolumetricBlackUintDummy);
-			}
 
 			// We are writing to the texture, we'll need to extract it...
 			bNeedsTextureExtract = true;
@@ -3978,7 +3854,6 @@ void FDeferredShadingSceneRenderer::RenderPathTracing(
 							PassParameters->CloudParameters.CloudCallableShaderId = View.PathTracingVolumetricCloudCallableShaderIndex;
 							PassParameters->CloudAccelerationMap = CloudAccelerationMap;
 							PassParameters->CloudAccelerationMapSampler = TStaticSamplerState<SF_Point, AM_Clamp, AM_Clamp, AM_Clamp>::GetRHI();
-							PassParameters->CloudMap = CloudMap;
 
 							if ((Config.PathTracingData.VolumeFlags & PATH_TRACER_VOLUME_ENABLE_FOG) != 0)
 							{
