[HWRT] Respect r.RayTracing.Culling.UseGroupIds in GPU instance culling.

diff --git a/Engine/Shaders/Private/RayTracing/RayTracingInstanceBufferUtil.usf b/Engine/Shaders/Private/RayTracing/RayTracingInstanceBufferUtil.usf
--- a/Engine/Shaders/Private/RayTracing/RayTracingInstanceBufferUtil.usf
+++ b/Engine/Shaders/Private/RayTracing/RayTracingInstanceBufferUtil.usf
@@ -58,6 +58,7 @@ float FarFieldCullingRadius;
 float AngleThresholdRatioSq;
 float3 ViewOrigin;
 uint CullingMode;
+uint CullUsingGroups;
 
 bool CullInstance(FInstanceSceneData InstanceSceneData, FPrimitiveSceneData PrimitiveData, float4x4 LocalToTranslatedWorld4x4, bool bFarFieldInstance)
 {
@@ -122,7 +123,7 @@ void ProcessInstance(FRayTracingInstanceDescriptor InputDesc)
 
 #if GPU_CULLING
 	const bool bHasGroupId = PrimitiveData.Flags & PRIMITIVE_SCENE_DATA_FLAG_RAYTRACING_HAS_GROUPID;
-	if (!bHasGroupId)
+	if (!bHasGroupId || !CullUsingGroups)
 	{
 		bCulled = CullInstance(InstanceSceneData, PrimitiveData, LocalToTranslatedWorld4x4, bFarFieldInstance);
 	}
diff --git a/Engine/Source/Runtime/Renderer/Private/RayTracing/RayTracingInstanceBufferUtil.cpp b/Engine/Source/Runtime/Renderer/Private/RayTracing/RayTracingInstanceBufferUtil.cpp
--- a/Engine/Source/Runtime/Renderer/Private/RayTracing/RayTracingInstanceBufferUtil.cpp
+++ b/Engine/Source/Runtime/Renderer/Private/RayTracing/RayTracingInstanceBufferUtil.cpp
@@ -530,6 +530,7 @@ struct FRayTracingBuildInstanceBufferCS : public FGlobalShader
 		SHADER_PARAMETER(float, AngleThresholdRatioSq)
 		SHADER_PARAMETER(FVector3f, ViewOrigin)
 		SHADER_PARAMETER(uint32, CullingMode)
+		SHADER_PARAMETER(uint32, CullUsingGroups)
 
 		SHADER_PARAMETER_UAV(RWStructuredBuffer<uint>, RWOutputStats)
 		SHADER_PARAMETER(uint32, OutputStatsOffset)
@@ -622,6 +623,7 @@ void BuildRayTracingInstanceBuffer(
 		PassParams.AngleThresholdRatioSq = CullingParameters->AngleThresholdRatioSq;
 		PassParams.ViewOrigin = CullingParameters->TranslatedViewOrigin;
 		PassParams.CullingMode = uint32(CullingParameters->CullingMode);
+		PassParams.CullUsingGroups = uint32(CullingParameters->bCullUsingGroupIds);
 	}
 
 	PassParams.RWOutputStats = OutputStatsUAV;
