Path Tracer: Refactor debug visualization settings

Align the naming of the cvars with the raytracing case. The new name for this feature is r.PathTracing.Visualize (with related options using that prefix).

Add some logic to populate the cvar text from the enum list.

Use the FName cvar mechanism to avoid paying a FName lookup on every query

diff --git a/Engine/Source/Runtime/Renderer/Private/PathTracing.cpp b/Engine/Source/Runtime/Renderer/Private/PathTracing.cpp
--- a/Engine/Source/Runtime/Renderer/Private/PathTracing.cpp
+++ b/Engine/Source/Runtime/Renderer/Private/PathTracing.cpp
@@ -49,14 +49,68 @@ TAutoConsoleVariable<bool> CVarPathTracingExperimental(
 	ECVF_RenderThreadSafe | ECVF_ReadOnly
 );
 
-TAutoConsoleVariable<FString> CVarPathTracingDebugVisualize(
-	TEXT("r.PathTracing.DebugVisualize"),
-	TEXT(""),
-	TEXT("Enables a debug mode for the path tracer which can visualize payload fields and perform basic lighting without sample accumulation. (default: "")")
+
+
+struct FPathTracingDebugModeInfo
+{
+	FPathTracingDebugModeInfo()
+	{
+		CVarHelpText = TEXT("Enables a debug mode for the path tracer which can visualize payload fields and perform basic lighting without sample accumulation.\n")
+					   TEXT("Allowed values:\n");
+#define REGISTER_VIZ_TYPE(name, value)	Map.Emplace(FName(name), value); CVarHelpText += TEXT(" - ") name TEXT("\n")
+
+		REGISTER_VIZ_TYPE(TEXT("Radiance") 				, PATH_TRACER_DEBUG_VIZ_RADIANCE				);
+		REGISTER_VIZ_TYPE(TEXT("WorldNormal") 			, PATH_TRACER_DEBUG_VIZ_WORLD_NORMAL			);
+		REGISTER_VIZ_TYPE(TEXT("WorldSmoothNormal") 	, PATH_TRACER_DEBUG_VIZ_WORLD_SMOOTH_NORMAL		);
+		REGISTER_VIZ_TYPE(TEXT("WorldGeoNormal") 		, PATH_TRACER_DEBUG_VIZ_WORLD_GEO_NORMAL		);
+		REGISTER_VIZ_TYPE(TEXT("BaseColor") 	    	, PATH_TRACER_DEBUG_VIZ_BASE_COLOR				);
+		REGISTER_VIZ_TYPE(TEXT("DiffuseColor") 	   		, PATH_TRACER_DEBUG_VIZ_DIFFUSE_COLOR			);
+		REGISTER_VIZ_TYPE(TEXT("SpecularColor") 	   	, PATH_TRACER_DEBUG_VIZ_SPECULAR_COLOR			);
+		REGISTER_VIZ_TYPE(TEXT("Opacity") 	     		, PATH_TRACER_DEBUG_VIZ_OPACITY					);
+		REGISTER_VIZ_TYPE(TEXT("Metallic") 	     		, PATH_TRACER_DEBUG_VIZ_METALLIC				);
+		REGISTER_VIZ_TYPE(TEXT("Specular") 	     		, PATH_TRACER_DEBUG_VIZ_SPECULAR				);
+		REGISTER_VIZ_TYPE(TEXT("Roughness") 	     	, PATH_TRACER_DEBUG_VIZ_ROUGHNESS				);
+		REGISTER_VIZ_TYPE(TEXT("IOR") 	    	 		, PATH_TRACER_DEBUG_VIZ_IOR						);
+		REGISTER_VIZ_TYPE(TEXT("ShadingModel")   	 	, PATH_TRACER_DEBUG_VIZ_SHADING_MODEL			);
+		REGISTER_VIZ_TYPE(TEXT("LightingChannelMask") 	, PATH_TRACER_DEBUG_VIZ_LIGHTING_CHANNEL_MASK	);
+		REGISTER_VIZ_TYPE(TEXT("CustomData0") 			, PATH_TRACER_DEBUG_VIZ_CUSTOM_DATA0 			);
+		REGISTER_VIZ_TYPE(TEXT("CustomData1") 			, PATH_TRACER_DEBUG_VIZ_CUSTOM_DATA1 			);
+		REGISTER_VIZ_TYPE(TEXT("WorldPosition") 		, PATH_TRACER_DEBUG_VIZ_WORLD_POSITION 			);
+		REGISTER_VIZ_TYPE(TEXT("PrimaryRays") 			, PATH_TRACER_DEBUG_VIZ_PRIMARY_RAYS 			);
+		REGISTER_VIZ_TYPE(TEXT("WorldTangent") 			, PATH_TRACER_DEBUG_VIZ_WORLD_TANGENT 			);
+		REGISTER_VIZ_TYPE(TEXT("Anisotropy") 	    	, PATH_TRACER_DEBUG_VIZ_ANISOTROPY 				);
+		REGISTER_VIZ_TYPE(TEXT("LightGridCount") 	   	, PATH_TRACER_DEBUG_VIZ_LIGHT_GRID_COUNT		);
+		REGISTER_VIZ_TYPE(TEXT("LightGridAxis") 	   	, PATH_TRACER_DEBUG_VIZ_LIGHT_GRID_AXIS			);
+		REGISTER_VIZ_TYPE(TEXT("DecalGridCount")	   	, PATH_TRACER_DEBUG_VIZ_DECAL_GRID_COUNT		);
+		REGISTER_VIZ_TYPE(TEXT("DecalGridAxis") 	   	, PATH_TRACER_DEBUG_VIZ_DECAL_GRID_AXIS			);
+		REGISTER_VIZ_TYPE(TEXT("VolumeLightCount") 		, PATH_TRACER_DEBUG_VIZ_VOLUME_LIGHT_COUNT		);
+		REGISTER_VIZ_TYPE(TEXT("HitKind")            	, PATH_TRACER_DEBUG_VIZ_HITKIND             	);
+#undef REGISTER_VIZ_TYPE
+
+		CVarHelpText += TEXT("Default value: \"\"\n");
+	}
+
+	static FPathTracingDebugModeInfo& Get()
+	{
+		static FPathTracingDebugModeInfo Singleton;
+		return Singleton;
+	}
+
+	TMap<FName, uint32> Map;
+	FString CVarHelpText;
+};
+
+
+static FName GPathTracingVisualizeMode = NAME_None;
+FAutoConsoleVariableRef CVarPathTracingVisualize(
+	TEXT("r.PathTracing.Visualize"),
+	GPathTracingVisualizeMode,
+	*FPathTracingDebugModeInfo::Get().CVarHelpText,
+	ECVF_RenderThreadSafe
 );
 
-TAutoConsoleVariable<int32> CVarPathTracingDebugVisualizeLightSamples(
-	TEXT("r.PathTracing.DebugVisualizeLightSamples"),
+TAutoConsoleVariable<int32> CVarPathTracingVisualizeLightSamples(
+	TEXT("r.PathTracing.Visualize.LightSamples"),
 	1,
 	TEXT("Number of light samples for the path tracing light sampling debug mode (default: 1)"),
 	ECVF_RenderThreadSafe
@@ -707,44 +761,9 @@ struct FPathTracingState
 	uint32_t FrameIndex = 0;
 };
 
-int GetPathTracingDebugVisualizationMode() {
-	static TMap<FName, uint32> PathTracingDebugModes;
-	if (PathTracingDebugModes.Num() == 0)
-	{
-		PathTracingDebugModes.Emplace(FName(TEXT("Radiance"))			, PATH_TRACER_DEBUG_VIZ_RADIANCE				);
-		PathTracingDebugModes.Emplace(FName(TEXT("WorldNormal"))		, PATH_TRACER_DEBUG_VIZ_WORLD_NORMAL			);
-		PathTracingDebugModes.Emplace(FName(TEXT("WorldSmoothNormal"))	, PATH_TRACER_DEBUG_VIZ_WORLD_SMOOTH_NORMAL		);
-		PathTracingDebugModes.Emplace(FName(TEXT("WorldGeoNormal"))		, PATH_TRACER_DEBUG_VIZ_WORLD_GEO_NORMAL		);
-		PathTracingDebugModes.Emplace(FName(TEXT("BaseColor"))	    	, PATH_TRACER_DEBUG_VIZ_BASE_COLOR				);
-		PathTracingDebugModes.Emplace(FName(TEXT("DiffuseColor"))	   	, PATH_TRACER_DEBUG_VIZ_DIFFUSE_COLOR			);
-		PathTracingDebugModes.Emplace(FName(TEXT("SpecularColor"))	   	, PATH_TRACER_DEBUG_VIZ_SPECULAR_COLOR			);
-		PathTracingDebugModes.Emplace(FName(TEXT("Opacity"))	     	, PATH_TRACER_DEBUG_VIZ_OPACITY					);
-		PathTracingDebugModes.Emplace(FName(TEXT("Metallic"))	     	, PATH_TRACER_DEBUG_VIZ_METALLIC				);
-		PathTracingDebugModes.Emplace(FName(TEXT("Specular"))	     	, PATH_TRACER_DEBUG_VIZ_SPECULAR				);
-		PathTracingDebugModes.Emplace(FName(TEXT("Roughness"))	     	, PATH_TRACER_DEBUG_VIZ_ROUGHNESS				);
-		PathTracingDebugModes.Emplace(FName(TEXT("IOR"))	    	 	, PATH_TRACER_DEBUG_VIZ_IOR						);
-		PathTracingDebugModes.Emplace(FName(TEXT("ShadingModel"))  	 	, PATH_TRACER_DEBUG_VIZ_SHADING_MODEL			);
-		PathTracingDebugModes.Emplace(FName(TEXT("LightingChannelMask")), PATH_TRACER_DEBUG_VIZ_LIGHTING_CHANNEL_MASK 	);
-		PathTracingDebugModes.Emplace(FName(TEXT("CustomData0"))		, PATH_TRACER_DEBUG_VIZ_CUSTOM_DATA0 			);
-		PathTracingDebugModes.Emplace(FName(TEXT("CustomData1"))		, PATH_TRACER_DEBUG_VIZ_CUSTOM_DATA1 			);
-		PathTracingDebugModes.Emplace(FName(TEXT("WorldPosition"))		, PATH_TRACER_DEBUG_VIZ_WORLD_POSITION 			);
-		PathTracingDebugModes.Emplace(FName(TEXT("PrimaryRays"))		, PATH_TRACER_DEBUG_VIZ_PRIMARY_RAYS 			);
-		PathTracingDebugModes.Emplace(FName(TEXT("WorldTangent"))		, PATH_TRACER_DEBUG_VIZ_WORLD_TANGENT 			);
-		PathTracingDebugModes.Emplace(FName(TEXT("Anisotropy"))	    	, PATH_TRACER_DEBUG_VIZ_ANISOTROPY 				);
-		PathTracingDebugModes.Emplace(FName(TEXT("LightGridCount"))	   	, PATH_TRACER_DEBUG_VIZ_LIGHT_GRID_COUNT		);
-		PathTracingDebugModes.Emplace(FName(TEXT("LightGridAxis"))	   	, PATH_TRACER_DEBUG_VIZ_LIGHT_GRID_AXIS			);
-		PathTracingDebugModes.Emplace(FName(TEXT("DecalGridCount"))	   	, PATH_TRACER_DEBUG_VIZ_DECAL_GRID_COUNT		);
-		PathTracingDebugModes.Emplace(FName(TEXT("DecalGridAxis"))	   	, PATH_TRACER_DEBUG_VIZ_DECAL_GRID_AXIS			);
-		PathTracingDebugModes.Emplace(FName(TEXT("VolumeLightCount"))	, PATH_TRACER_DEBUG_VIZ_VOLUME_LIGHT_COUNT		);
-		PathTracingDebugModes.Emplace(FName(TEXT("HitKind"))            , PATH_TRACER_DEBUG_VIZ_HITKIND                 );
-	}
-	FString DebugMode = CVarPathTracingDebugVisualize.GetValueOnRenderThread();
-	if (!DebugMode.IsEmpty())
-	{
-		return PathTracingDebugModes.FindRef(FName(*DebugMode, FNAME_Find), -1);
-	}
-	// not doing debugging
-	return -1;
+int GetPathTracingVisualizationMode()
+{
+	return GPathTracingVisualizeMode.IsNone() ? -1 : FPathTracingDebugModeInfo::Get().Map.FindRef(GPathTracingVisualizeMode, -1);
 }
 
 namespace PathTracing
@@ -767,12 +786,12 @@ namespace PathTracing
 
 	bool NeedsAntiAliasing(const FViewInfo& View)
 	{
-		return GetPathTracingDebugVisualizationMode() >= 0;
+		return GetPathTracingVisualizationMode() >= 0;
 	}
 
 	bool NeedsTonemapping()
 	{
-		const int DebugMode = GetPathTracingDebugVisualizationMode();
+		const int DebugMode = GetPathTracingVisualizationMode();
 		return DebugMode < 0 ||
 			   DebugMode == PATH_TRACER_DEBUG_VIZ_RADIANCE ||
 			   DebugMode == PATH_TRACER_DEBUG_VIZ_PRIMARY_RAYS;
@@ -3001,7 +3020,7 @@ void FDeferredShadingSceneRenderer::PreparePathTracing(const FSceneViewFamily& V
 	if (ViewFamily.EngineShowFlags.PathTracing
 		&& ShouldCompilePathTracingShadersForProject(ViewFamily.GetShaderPlatform()))
 	{
-		if (GetPathTracingDebugVisualizationMode() >= 0)
+		if (GetPathTracingVisualizationMode() >= 0)
 		{
 			auto RayGenShader = GetGlobalShaderMap(ViewFamily.GetShaderPlatform())->GetShader<FPathTracingDebugRG>();
 			OutRayGenShaders.Add(RayGenShader.GetRayTracingShader());
@@ -3177,7 +3196,7 @@ void FDeferredShadingSceneRenderer::RenderPathTracing(
 	}
 
 
-	if (int DebugMode = GetPathTracingDebugVisualizationMode(); DebugMode >= 0)
+	if (int DebugMode = GetPathTracingVisualizationMode(); DebugMode >= 0)
 	{
 		const int32 DispatchResX = View.ViewRect.Size().X;
 		const int32 DispatchResY = View.ViewRect.Size().Y;
@@ -3199,7 +3218,7 @@ void FDeferredShadingSceneRenderer::RenderPathTracing(
 		PassParameters->ViewUniformBuffer = View.ViewUniformBuffer;
 		PassParameters->DebugMode = DebugMode;
 
-		PassParameters->NumLightSamples = FMath::Max(CVarPathTracingDebugVisualizeLightSamples.GetValueOnRenderThread(), 1);
+		PassParameters->NumLightSamples = FMath::Max(CVarPathTracingVisualizeLightSamples.GetValueOnRenderThread(), 1);
 		SetLightParameters(
 			GraphBuilder,
 			Scene,
