FSceneInterface - Added BatchRemovePrimitives()
* Efficient removal of an array of SceneProxies

diff --git a/Engine/Source/Runtime/Engine/Public/SceneInterface.h b/Engine/Source/Runtime/Engine/Public/SceneInterface.h
--- a/Engine/Source/Runtime/Engine/Public/SceneInterface.h
+++ b/Engine/Source/Runtime/Engine/Public/SceneInterface.h
@@ -543,6 +543,8 @@ public:
 	virtual void BatchAddPrimitives(TArrayView<FPrimitiveSceneDesc*> InPrimitives) = 0;
 	virtual void BatchRemovePrimitives(TArrayView<FPrimitiveSceneDesc*> InPrimitives) = 0;
 	virtual void BatchReleasePrimitives(TArrayView<FPrimitiveSceneDesc*> InPrimitives) = 0;
+
+	virtual void BatchRemovePrimitives(TArray<FPrimitiveSceneProxy*>&& InPrimitives) = 0;
 	
 	virtual void UpdateCustomPrimitiveData(FPrimitiveSceneDesc* Primitive, const FCustomPrimitiveData& CustomPrimitiveData) = 0;
 
diff --git a/Engine/Source/Runtime/Renderer/Private/RendererScene.cpp b/Engine/Source/Runtime/Renderer/Private/RendererScene.cpp
--- a/Engine/Source/Runtime/Renderer/Private/RendererScene.cpp
+++ b/Engine/Source/Runtime/Renderer/Private/RendererScene.cpp
@@ -2447,7 +2447,6 @@ void FScene::BatchRemovePrimitivesInternal(TArrayView<T*> InPrimitives)
 
 	for (T* Primitive : InPrimitives)
 	{
-		FPrimitiveSceneInfoData& SceneData = Primitive->GetSceneData();
 		FPrimitiveSceneProxy* PrimitiveSceneProxy = Primitive->GetSceneProxy();
 		if (PrimitiveSceneProxy)
 		{
@@ -2488,6 +2487,22 @@ void FScene::BatchRemovePrimitives(TArrayView<FPrimitiveSceneDesc*> InPrimitives
 	BatchRemovePrimitivesInternal(InPrimitives);
 }
 
+void FScene::BatchRemovePrimitives(TArray<FPrimitiveSceneProxy*>&& InPrimitives)
+{
+	if (!InPrimitives.IsEmpty())
+	{
+		ENQUEUE_RENDER_COMMAND(BatchRemovePrimitives)(
+			[this, InPrimitives = MoveTemp(InPrimitives)](FRHICommandListBase&)
+		{
+			for (FPrimitiveSceneProxy* PrimitiveSceneProxy : InPrimitives)
+			{
+				RemovePrimitiveSceneInfo_RenderThread(PrimitiveSceneProxy->GetPrimitiveSceneInfo());
+				PrimitiveSceneProxy->DestroyRenderThreadResources();
+			}
+		});
+	}
+}
+
 void FScene::ReleasePrimitive(UPrimitiveComponent* PrimitiveComponent)
 {
 	// Check if this components was already bulk released on the render side
@@ -7138,6 +7153,8 @@ public:
 	virtual void BatchRemovePrimitives(TArrayView<FPrimitiveSceneDesc*> InPrimitives) override {};
 	virtual void BatchReleasePrimitives(TArrayView<FPrimitiveSceneDesc*> InPrimitives) override {};
 
+	virtual void BatchRemovePrimitives(TArray<FPrimitiveSceneProxy*>&& InPrimitives) override {};
+
 	virtual void UpdateCustomPrimitiveData(FPrimitiveSceneDesc* Primitive, const FCustomPrimitiveData&) override {}
 	virtual void UpdatePrimitiveInstances(FPrimitiveSceneDesc* Primitive) override {};
 	virtual void UpdatePrimitiveInstances(FInstancedStaticMeshSceneDesc* Primitive) override {};	
diff --git a/Engine/Source/Runtime/Renderer/Private/ScenePrivate.h b/Engine/Source/Runtime/Renderer/Private/ScenePrivate.h
--- a/Engine/Source/Runtime/Renderer/Private/ScenePrivate.h
+++ b/Engine/Source/Runtime/Renderer/Private/ScenePrivate.h
@@ -3498,6 +3498,8 @@ public:
 	virtual void BatchAddPrimitives(TArrayView<FPrimitiveSceneDesc*> InPrimitives) override;
 	virtual void BatchRemovePrimitives(TArrayView<FPrimitiveSceneDesc*> InPrimitives) override;
 	virtual void BatchReleasePrimitives(TArrayView<FPrimitiveSceneDesc*> InPrimitives) override;
+
+	virtual void BatchRemovePrimitives(TArray<FPrimitiveSceneProxy*>&& InPrimitives) override;
 		
 	virtual void UpdateCustomPrimitiveData(FPrimitiveSceneDesc* Primitive, const FCustomPrimitiveData& CustomPrimitiveData) override;
 
