PR #12219: Allow light culling to be run on async compute.

diff --git a/Engine/Source/Runtime/Renderer/Private/LightGridInjection.cpp b/Engine/Source/Runtime/Renderer/Private/LightGridInjection.cpp
--- a/Engine/Source/Runtime/Renderer/Private/LightGridInjection.cpp
+++ b/Engine/Source/Runtime/Renderer/Private/LightGridInjection.cpp
@@ -40,6 +40,13 @@ LightGridInjection.cpp
 #include "VolumetricFog.h"
 #include "LightViewData.h"
 
+static TAutoConsoleVariable<bool> CVarLightGridAsyncCompute(
+	TEXT("r.Forward.LightGridAsyncCompute"),
+	false,
+	TEXT("Run the light culling passes in async compute."),
+	ECVF_Scalability | ECVF_RenderThreadSafe
+);
+
 enum class ELightBufferMode
 {
 	VisibleLocalLights			= 0,
@@ -636,9 +643,10 @@ FLightGrid LightGridInjection(
 		CulledLightDataGridUAV = GraphBuilder.CreateUAV(CulledLightDataGrid);
 	}
 
-	AddClearUAVPass(GraphBuilder, CulledLightLinkAllocatorUAV, 0);
-	AddClearUAVPass(GraphBuilder, CulledLightDataAllocatorUAV, 0);
-	AddClearUAVPass(GraphBuilder, NumCulledLightsGridUAV, 0);
+	const ERDGPassFlags RDGPassFlags = CVarLightGridAsyncCompute.GetValueOnRenderThread() ? ERDGPassFlags::AsyncCompute : ERDGPassFlags::Compute;
+	AddClearUAVPass(GraphBuilder, CulledLightLinkAllocatorUAV, 0, RDGPassFlags);
+	AddClearUAVPass(GraphBuilder, CulledLightDataAllocatorUAV, 0, RDGPassFlags);
+	AddClearUAVPass(GraphBuilder, NumCulledLightsGridUAV, 0, RDGPassFlags);
 
 	FLightGridInjectionCS::FParameters* PassParameters = GraphBuilder.AllocParameters<FLightGridInjectionCS::FParameters>();
 
@@ -721,6 +729,7 @@ FLightGrid LightGridInjection(
 		RDG_EVENT_NAME("LightGridInject %s %s",
 			PermutationVector.Get<FLightGridInjectionCS::FUseLinkedList>() ? TEXT("LinkedList") : TEXT("NoLinkedList"),
 			PermutationVector.Get<FLightGridInjectionCS::FUseThreadGroupPerCell>() ? TEXT("ThreadGroup") : TEXT("SingleThread")),
+    	RDGPassFlags,
 		ComputeShader,
 		PassParameters,
 		NumGroups);
