Make the light grid feedback pass also run on the async queue to prevent it from causing a sync.

diff --git a/Engine/Source/Runtime/Renderer/Private/LightGrid.h b/Engine/Source/Runtime/Renderer/Private/LightGrid.h
--- a/Engine/Source/Runtime/Renderer/Private/LightGrid.h
+++ b/Engine/Source/Runtime/Renderer/Private/LightGrid.h
@@ -17,7 +17,7 @@ public:
 
 	FLightGridViewState();
 
-	void FeedbackStatus(FRDGBuilder& GraphBuilder, FViewInfo& View, FRDGBufferRef CulledLightDataAllocatorBuffer, uint32 NumCulledLightDataEntries, FRDGBufferRef CulledLightLinkAllocatorBuffer, uint32 NumCulledLightLinks);
+	void FeedbackStatus(FRDGBuilder& GraphBuilder, FViewInfo& View, FRDGBufferRef CulledLightDataAllocatorBuffer, uint32 NumCulledLightDataEntries, FRDGBufferRef CulledLightLinkAllocatorBuffer, uint32 NumCulledLightLinks, bool bUseAsyncCompute);
 
 #if !UE_BUILD_SHIPPING
 	uint32 GetStatusMessageId() const { return StatusFeedbackSocket.GetMessageId().GetIndex(); }
diff --git a/Engine/Source/Runtime/Renderer/Private/LightGridInjection.cpp b/Engine/Source/Runtime/Renderer/Private/LightGridInjection.cpp
--- a/Engine/Source/Runtime/Renderer/Private/LightGridInjection.cpp
+++ b/Engine/Source/Runtime/Renderer/Private/LightGridInjection.cpp
@@ -643,7 +643,8 @@ FLightGrid LightGridInjection(
 		CulledLightDataGridUAV = GraphBuilder.CreateUAV(CulledLightDataGrid);
 	}
 
-	const ERDGPassFlags RDGPassFlags = CVarLightGridAsyncCompute.GetValueOnRenderThread() ? ERDGPassFlags::AsyncCompute : ERDGPassFlags::Compute;
+	const bool bUseAsyncCompute = CVarLightGridAsyncCompute.GetValueOnRenderThread();
+	const ERDGPassFlags RDGPassFlags = bUseAsyncCompute ? ERDGPassFlags::AsyncCompute : ERDGPassFlags::Compute;
 	AddClearUAVPass(GraphBuilder, CulledLightLinkAllocatorUAV, 0, RDGPassFlags);
 	AddClearUAVPass(GraphBuilder, CulledLightDataAllocatorUAV, 0, RDGPassFlags);
 	AddClearUAVPass(GraphBuilder, NumCulledLightsGridUAV, 0, RDGPassFlags);
@@ -740,7 +741,7 @@ FLightGrid LightGridInjection(
 
 	if (LightGridViewState != nullptr)
 	{
-		LightGridViewState->FeedbackStatus(GraphBuilder, View, CulledLightDataAllocatorBuffer, NumCulledLightEntries, CulledLightLinkAllocatorBuffer, NumCulledLightLinks);
+		LightGridViewState->FeedbackStatus(GraphBuilder, View, CulledLightDataAllocatorBuffer, NumCulledLightEntries, CulledLightLinkAllocatorBuffer, NumCulledLightLinks, bUseAsyncCompute);
 	}
 
 	return Output;
@@ -1704,7 +1705,7 @@ FLightGridViewState::FLightGridViewState()
 #endif
 }
 
-void FLightGridViewState::FeedbackStatus(FRDGBuilder& GraphBuilder, FViewInfo& View, FRDGBufferRef CulledLightDataAllocatorBuffer, uint32 NumCulledLightDataEntries, FRDGBufferRef CulledLightLinkAllocatorBuffer, uint32 NumCulledLightLinks)
+void FLightGridViewState::FeedbackStatus(FRDGBuilder& GraphBuilder, FViewInfo& View, FRDGBufferRef CulledLightDataAllocatorBuffer, uint32 NumCulledLightDataEntries, FRDGBufferRef CulledLightLinkAllocatorBuffer, uint32 NumCulledLightLinks, bool bUseAsyncCompute)
 {
 #if !UE_BUILD_SHIPPING
 	FLightGridFeedbackStatusCS::FParameters* PassParameters = GraphBuilder.AllocParameters<FLightGridFeedbackStatusCS::FParameters>();
@@ -1723,6 +1724,7 @@ void FLightGridViewState::FeedbackStatus(FRDGBuilder& GraphBuilder, FViewInfo& V
 	FComputeShaderUtils::AddPass(
 		GraphBuilder,
 		RDG_EVENT_NAME("LightGridFeedbackStatus"),
+		bUseAsyncCompute ? ERDGPassFlags::AsyncCompute : ERDGPassFlags::Compute,
 		ComputeShader,
 		PassParameters,
 		FIntVector(1, 1, 1)
