[Lit Translucency] Use compute shader instead of VS+GS+PS for InjectBatchMain.
- use UAV typed load when possible otherwise fallback to using auxiliary textures.

diff --git a/Engine/Shaders/Private/TranslucentLightInjectionShaders.usf b/Engine/Shaders/Private/TranslucentLightInjectionShaders.usf
--- a/Engine/Shaders/Private/TranslucentLightInjectionShaders.usf
+++ b/Engine/Shaders/Private/TranslucentLightInjectionShaders.usf
@@ -57,6 +57,11 @@ float3 GetTranslatedWorldPosition(float2 UV, uint LayerIndex)
 	return TranslatedWorldPosition;
 }
 
+float3 GetTranslatedWorldPosition(uint3 VolumeCoord)
+{
+	return View.TranslucencyLightingVolumeMin[VolumeCascadeIndex].xyz + (VolumeCoord + .5f) * View.TranslucencyLightingVolumeInvSize[VolumeCascadeIndex].w;
+}
+
 void AccumulateSHLighting(
 	float3 NormalizedLightVector,
 	float3 Lighting,
@@ -131,13 +136,11 @@ void SimpleLightInjectMainPS(
 		OutColor0, OutColor1);
 }
 
-float GetBorderLerpFactor(float2 UV, uint LayerIndex)
+float GetBorderLerpFactorFromVolumeUVs(float3 VolumeUVs)
 {
 	// Larger values result in a shorter transition distance
 	float TransitionScale = 10;
 
-	float3 VolumeUVs = float3(UV, (LayerIndex + .5f) * View.TranslucencyLightingVolumeMin[VolumeCascadeIndex].w);
-
 	// Rescale the UVs to make the fade go to 0 before the edge of the volume
 	float3 FadeUVs = VolumeUVs * (1 + 4 * View.TranslucencyLightingVolumeMin[VolumeCascadeIndex].w) - 2 * View.TranslucencyLightingVolumeMin[VolumeCascadeIndex].w;
 	// Setup a 3d lerp factor going to 0 near the edge of the outer volume
@@ -147,6 +150,18 @@ float GetBorderLerpFactor(float2 UV, uint LayerIndex)
 	return FinalLerpFactor;
 }
 
+float GetBorderLerpFactor(float2 UV, uint LayerIndex)
+{
+	const float3 VolumeUVs = float3(UV, (LayerIndex + .5f) * View.TranslucencyLightingVolumeMin[VolumeCascadeIndex].w);
+	return GetBorderLerpFactorFromVolumeUVs(VolumeUVs);
+}
+
+float GetBorderLerpFactor(uint3 VolumeCoord)
+{
+	const float3 VolumeUVs = (VolumeCoord + 0.5f) * View.TranslucencyLightingVolumeMin[VolumeCascadeIndex].w;
+	return GetBorderLerpFactorFromVolumeUVs(VolumeUVs);
+}
+
 float GetTranslucencyVolumeLocalLightAttenuation(float3 TranslatedWorldPosition, FDeferredLightData LightData)
 {
 	float3 ToLight, L;
@@ -344,7 +359,7 @@ void InjectMainPS(
 
 #endif // #if INJECTION_PIXEL_SHADER
 
-#if INJECTION_BATCH_PIXEL_SHADER
+#ifdef InjectBatchMainCS
 
 #include "LightGridCommon.ush"
 #include "LightData.ush"
@@ -355,17 +370,31 @@ void InjectMainPS(
 StructuredBuffer<uint> BatchedLocalLights;
 uint MaxBatchedLocalLights;
 
-/** Pixel shader that calculates direct lighting for all the affected voxels of a volume texture. */
-void InjectBatchMainPS(
-	FWriteToSliceGeometryOutput Input,
-	out float4 RTOutColor0 : SV_Target0,
-	out float4 RTOutColor1 : SV_Target1
-)
+uint3 VolumeSize;
+
+RWTexture3D<float4> RWTranslucencyLightingVolumeAmbient;
+RWTexture3D<float4> RWTranslucencyLightingVolumeDirectional;
+
+#if !USE_UAV_TYPED_LOAD
+Texture3D<float4> TranslucencyLightingVolumeAmbient;
+Texture3D<float4> TranslucencyLightingVolumeDirectional;
+#endif
+
+/** Compute shader that calculates direct lighting for all the affected voxels of a volume texture. */
+[numthreads(THREADGROUP_SIZE_X, THREADGROUP_SIZE_Y, THREADGROUP_SIZE_Z)]
+void InjectBatchMainCS(uint3 DispatchThreadId : SV_DispatchThreadID)
 {
+	uint3 VolumeCoord = DispatchThreadId.xyz;
+
+	if (any(VolumeCoord >= VolumeSize))
+	{
+		return;
+	}
+
 	float4 OutColor0 = 0;
 	float4 OutColor1 = 0;
 
-	float3 TranslatedWorldPosition = GetTranslatedWorldPosition(Input.Vertex.UV, Input.LayerIndex);
+	float3 TranslatedWorldPosition = GetTranslatedWorldPosition(VolumeCoord);
 	
 	const uint NumUints = (MaxBatchedLocalLights + 31U) / 32U;
 	for (uint UintIndex = 0; UintIndex < NumUints; ++UintIndex)
@@ -402,7 +431,7 @@ void InjectBatchMainPS(
 
 						if (VolumeCascadeIndex == 1)
 						{
-							float FinalLerpFactor = GetBorderLerpFactor(Input.Vertex.UV, Input.LayerIndex);
+							float FinalLerpFactor = GetBorderLerpFactor(VolumeCoord);
 							// For local lights, fade attenuation to 0 for the border voxels
 							Attenuation = lerp(0, Attenuation, FinalLerpFactor);
 							ShadowFactor = lerp(0.0f, ShadowFactor, FinalLerpFactor);
@@ -424,11 +453,16 @@ void InjectBatchMainPS(
 		}
 	}
 
-	RTOutColor0 = OutColor0;
-	RTOutColor1 = OutColor1;
+#if USE_UAV_TYPED_LOAD
+	RWTranslucencyLightingVolumeAmbient[DispatchThreadId] = RWTranslucencyLightingVolumeAmbient[DispatchThreadId] + OutColor0;
+	RWTranslucencyLightingVolumeDirectional[DispatchThreadId] = RWTranslucencyLightingVolumeDirectional[DispatchThreadId] + OutColor1;
+#else
+	RWTranslucencyLightingVolumeAmbient[DispatchThreadId] = TranslucencyLightingVolumeAmbient[DispatchThreadId] + OutColor0;
+	RWTranslucencyLightingVolumeDirectional[DispatchThreadId] = TranslucencyLightingVolumeDirectional[DispatchThreadId] + OutColor1;
+#endif
 }
 
-#endif // #if INJECTION_BATCH_PIXEL_SHADER
+#endif // InjectBatchMainCS
 
 
 #if CLEAR_COMPUTE_SHADER
diff --git a/Engine/Source/Runtime/Renderer/Private/DeferredShadingRenderer.h b/Engine/Source/Runtime/Renderer/Private/DeferredShadingRenderer.h
--- a/Engine/Source/Runtime/Renderer/Private/DeferredShadingRenderer.h
+++ b/Engine/Source/Runtime/Renderer/Private/DeferredShadingRenderer.h
@@ -843,7 +843,7 @@ private:
 	void RenderLights(
 		FRDGBuilder& GraphBuilder,
 		FMinimalSceneTextures& SceneTextures,
-		const FTranslucencyLightingVolumeTextures& TranslucencyLightingVolumeTextures,
+		FTranslucencyLightingVolumeTextures& TranslucencyLightingVolumeTextures,
 		FRDGTextureRef LightingChannelsTexture,
 		const FSortedLightSetSceneInfo& SortedLightSet);
 
diff --git a/Engine/Source/Runtime/Renderer/Private/LightRendering.cpp b/Engine/Source/Runtime/Renderer/Private/LightRendering.cpp
--- a/Engine/Source/Runtime/Renderer/Private/LightRendering.cpp
+++ b/Engine/Source/Runtime/Renderer/Private/LightRendering.cpp
@@ -1431,7 +1431,7 @@ FHairStrandsTransmittanceMaskData CreateDummyHairStrandsTransmittanceMaskData(FR
 void FDeferredShadingSceneRenderer::RenderLights(
 	FRDGBuilder& GraphBuilder,
 	FMinimalSceneTextures& SceneTextures,
-	const FTranslucencyLightingVolumeTextures& TranslucencyLightingVolumeTextures,
+	FTranslucencyLightingVolumeTextures& TranslucencyLightingVolumeTextures,
 	FRDGTextureRef LightingChannelsTexture,
 	const FSortedLightSetSceneInfo& SortedLightSet)
 {
@@ -1524,7 +1524,7 @@ void FDeferredShadingSceneRenderer::RenderLights(
 					FViewInfo& View = Views[ViewIndex];
 					RDG_GPU_MASK_SCOPE(GraphBuilder, View.GPUMask);
 					RDG_EVENT_SCOPE(GraphBuilder, "InjectTranslucencyLightingVolume(View=%d)", ViewIndex);
-					InjectTranslucencyLightingVolume(GraphBuilder, View, ViewIndex, Scene, *this, TranslucencyLightingVolumeTextures, Collector);
+					InjectTranslucencyLightingVolume(GraphBuilder, View, ViewIndex, Scene, *this, Collector, TranslucencyLightingVolumeTextures);
 				}
 			}
 		}
diff --git a/Engine/Source/Runtime/Renderer/Private/TranslucentLighting.cpp b/Engine/Source/Runtime/Renderer/Private/TranslucentLighting.cpp
--- a/Engine/Source/Runtime/Renderer/Private/TranslucentLighting.cpp
+++ b/Engine/Source/Runtime/Renderer/Private/TranslucentLighting.cpp
@@ -869,10 +869,10 @@ public:
 IMPLEMENT_MATERIAL_SHADER_TYPE(,FTranslucentLightingInjectPS, TEXT("/Engine/Private/TranslucentLightInjectionShaders.usf"), TEXT("InjectMainPS"), SF_Pixel);
 
 /** Shader that adds direct lighting contribution from multiple lights. */
-class FTranslucentLightingInjectBatchPS : public FGlobalShader
+class FTranslucentLightingInjectBatchCS : public FGlobalShader
 {
-	DECLARE_GLOBAL_SHADER(FTranslucentLightingInjectBatchPS);
-	SHADER_USE_PARAMETER_STRUCT(FTranslucentLightingInjectBatchPS, FGlobalShader);
+	DECLARE_GLOBAL_SHADER(FTranslucentLightingInjectBatchCS);
+	SHADER_USE_PARAMETER_STRUCT(FTranslucentLightingInjectBatchCS, FGlobalShader);
 
 	BEGIN_SHADER_PARAMETER_STRUCT(FParameters, )
 		SHADER_PARAMETER_STRUCT_REF(FViewUniformShaderParameters, ViewUniformBuffer)
@@ -882,22 +882,43 @@ class FTranslucentLightingInjectBatchPS : public FGlobalShader
 		SHADER_PARAMETER_RDG_BUFFER_SRV(StructuredBuffer<uint>, BatchedLocalLights)
 		SHADER_PARAMETER(uint32, MaxBatchedLocalLights)
 		SHADER_PARAMETER(uint32, VolumeCascadeIndex)
-		RENDER_TARGET_BINDING_SLOTS()
+		SHADER_PARAMETER(FIntVector, VolumeSize)
+
+		SHADER_PARAMETER_RDG_TEXTURE_UAV(RWTexture3D, RWTranslucencyLightingVolumeAmbient)
+		SHADER_PARAMETER_RDG_TEXTURE_UAV(RWTexture3D, RWTranslucencyLightingVolumeDirectional)
+
+		// used when UAV typed loads are not supported
+		SHADER_PARAMETER_RDG_TEXTURE(Texture3D, TranslucencyLightingVolumeAmbient)
+		SHADER_PARAMETER_RDG_TEXTURE(Texture3D, TranslucencyLightingVolumeDirectional)
 	END_SHADER_PARAMETER_STRUCT()
 
 	class FVirtualShadowMap : SHADER_PERMUTATION_BOOL("VIRTUAL_SHADOW_MAP");
 	class FUseLightFunctionAtlas : SHADER_PERMUTATION_BOOL("USE_LIGHT_FUNCTION_ATLAS");
-	using FPermutationDomain = TShaderPermutationDomain<FVirtualShadowMap, FUseLightFunctionAtlas>;
+	class FUseUAVTypedLoad : SHADER_PERMUTATION_BOOL("USE_UAV_TYPED_LOAD");
+	using FPermutationDomain = TShaderPermutationDomain<FVirtualShadowMap, FUseLightFunctionAtlas, FUseUAVTypedLoad>;
 
 public:
 
+	static FIntVector GetGroupSize()
+	{
+		return FIntVector(4, 4, 4);
+	}
+
 	static void ModifyCompilationEnvironment(const FGlobalShaderPermutationParameters& Parameters, FShaderCompilerEnvironment& OutEnvironment)
 	{
 		FGlobalShader::ModifyCompilationEnvironment(Parameters, OutEnvironment);
 		FVirtualShadowMapArray::SetShaderDefines(OutEnvironment);
 		FForwardLightingParameters::ModifyCompilationEnvironment(Parameters.Platform, OutEnvironment);
 
-		OutEnvironment.SetDefine(TEXT("INJECTION_BATCH_PIXEL_SHADER"), 1);
+		FIntVector GroupSize = GetGroupSize();
+
+		OutEnvironment.SetDefine(TEXT("THREADGROUP_SIZE_X"), GroupSize.X);
+		OutEnvironment.SetDefine(TEXT("THREADGROUP_SIZE_Y"), GroupSize.Y);
+		OutEnvironment.SetDefine(TEXT("THREADGROUP_SIZE_Z"), GroupSize.Z);
+
+		// This shader must support typed UAV load and we are testing if it is supported at runtime
+		// using UE::PixelFormat::HasCapabilities(..., EPixelFormatCapabilities::TypedUAVLoad)
+		OutEnvironment.CompilerFlags.Add(CFLAG_AllowTypedUAVLoads);
 	}
 
 	static bool ShouldCompilePermutation(const FGlobalShaderPermutationParameters& Parameters)
@@ -909,12 +930,11 @@ public:
 			return false;
 		}
 
-		return (IsFeatureLevelSupported(Parameters.Platform, ERHIFeatureLevel::SM5) &&
-			(RHISupportsGeometryShaders(Parameters.Platform) || RHISupportsVertexShaderLayer(Parameters.Platform)));
+		return IsFeatureLevelSupported(Parameters.Platform, ERHIFeatureLevel::SM5);
 	}
 };
 
-IMPLEMENT_MATERIAL_SHADER_TYPE(,FTranslucentLightingInjectBatchPS, TEXT("/Engine/Private/TranslucentLightInjectionShaders.usf"), TEXT("InjectBatchMainPS"), SF_Pixel);
+IMPLEMENT_GLOBAL_SHADER(FTranslucentLightingInjectBatchCS, "/Engine/Private/TranslucentLightInjectionShaders.usf", "InjectBatchMainCS", SF_Compute);
 
 
 
@@ -1425,24 +1445,23 @@ static void InjectTranslucencyLightingVolumeBatch(
 	FRDGBufferRef BatchedLocalLightsRDG,
 	bool bSupportVirtualShadowMaps,
 	uint32 VolumeCascadeIndex,
-	FRDGTextureRef VolumeAmbientTexture,
-	FRDGTextureRef VolumeDirectionalTexture,
+	FRDGTextureRef& VolumeAmbientTexture,
+	FRDGTextureRef& VolumeDirectionalTexture,
 	const FSceneRenderer& Renderer)
 {
 	check(MaxBatchedLocalLights > 0);
 	check(BatchedLocalLightsRDG);
 
-	TShaderMapRef<FWriteToSliceVS> VertexShader(View.ShaderMap);
-	TOptionalShaderMapRef<FWriteToSliceGS> GeometryShader(View.ShaderMap);
+	EPixelFormat TranslucencyPixelFormat;
+	ETextureCreateFlags TranslucencyTargetFlags;
+	FTranslucencyLightingVolumeTextures::GetTextureFormatAndCreationFlags(TranslucencyPixelFormat, TranslucencyTargetFlags);
 
-	const bool bUseLightFunctionAtlas = LightFunctionAtlas::IsEnabled(View, ELightFunctionAtlasSystem::DeferredLighting);
+	const FIntVector VolumeSize = FIntVector(GetTranslucencyLightingVolumeDim());
 
-	FTranslucentLightingInjectBatchPS::FPermutationDomain PermutationVector;
-	PermutationVector.Set< FTranslucentLightingInjectBatchPS::FVirtualShadowMap >(bSupportVirtualShadowMaps);
-	PermutationVector.Set< FTranslucentLightingInjectBatchPS::FUseLightFunctionAtlas >(bUseLightFunctionAtlas);
-	TShaderMapRef<FTranslucentLightingInjectBatchPS> PixelShader(View.ShaderMap, PermutationVector);
+	const bool bUseLightFunctionAtlas = LightFunctionAtlas::IsEnabled(View, ELightFunctionAtlasSystem::DeferredLighting);
+	const bool bUseUAVTypedLoad = UE::PixelFormat::HasCapabilities(TranslucencyPixelFormat, EPixelFormatCapabilities::TypedUAVLoad);
 
-	auto* PassParameters = GraphBuilder.AllocParameters<FTranslucentLightingInjectBatchPS::FParameters>();
+	auto* PassParameters = GraphBuilder.AllocParameters<FTranslucentLightingInjectBatchCS::FParameters>();
 	if (bSupportVirtualShadowMaps)
 	{
 		PassParameters->VirtualShadowMapSamplingParameters = Renderer.VirtualShadowMapArray.GetSamplingParameters(GraphBuilder, ViewIndex);
@@ -1453,28 +1472,47 @@ static void InjectTranslucencyLightingVolumeBatch(
 	PassParameters->LightFunctionAtlas = LightFunctionAtlas::BindGlobalParameters(GraphBuilder, View);;
 	PassParameters->MaxBatchedLocalLights = MaxBatchedLocalLights;
 	PassParameters->VolumeCascadeIndex = VolumeCascadeIndex;
+	PassParameters->VolumeSize = VolumeSize;
 
-	PassParameters->RenderTargets[0] = FRenderTargetBinding(VolumeAmbientTexture, ERenderTargetLoadAction::ELoad);
-	PassParameters->RenderTargets[1] = FRenderTargetBinding(VolumeDirectionalTexture, ERenderTargetLoadAction::ELoad);
-			
-	GraphBuilder.AddPass(
+	if (bUseUAVTypedLoad)
+	{
+		PassParameters->RWTranslucencyLightingVolumeAmbient = GraphBuilder.CreateUAV(VolumeAmbientTexture);
+		PassParameters->RWTranslucencyLightingVolumeDirectional = GraphBuilder.CreateUAV(VolumeDirectionalTexture);
+	}
+	else
+	{
+		// need to output to auxiliary textures to be able to read from existing volumes
+
+		FRDGTextureRef OutputVolumeAmbientTexture = GraphBuilder.CreateTexture(VolumeAmbientTexture->Desc, VolumeAmbientTexture->Name);
+		FRDGTextureRef OutputVolumeDirectionalTexture = GraphBuilder.CreateTexture(VolumeDirectionalTexture->Desc, VolumeDirectionalTexture->Name);
+
+		PassParameters->RWTranslucencyLightingVolumeAmbient = GraphBuilder.CreateUAV(OutputVolumeAmbientTexture);
+		PassParameters->RWTranslucencyLightingVolumeDirectional = GraphBuilder.CreateUAV(OutputVolumeDirectionalTexture);
+
+		PassParameters->TranslucencyLightingVolumeAmbient = VolumeAmbientTexture;
+		PassParameters->TranslucencyLightingVolumeDirectional = VolumeDirectionalTexture;
+
+		VolumeAmbientTexture = OutputVolumeAmbientTexture;
+		VolumeDirectionalTexture = OutputVolumeDirectionalTexture;
+	}
+
+	FTranslucentLightingInjectBatchCS::FPermutationDomain PermutationVector;
+	PermutationVector.Set<FTranslucentLightingInjectBatchCS::FVirtualShadowMap>(bSupportVirtualShadowMaps);
+	PermutationVector.Set<FTranslucentLightingInjectBatchCS::FUseLightFunctionAtlas>(bUseLightFunctionAtlas);
+	PermutationVector.Set<FTranslucentLightingInjectBatchCS::FUseUAVTypedLoad>(bUseUAVTypedLoad);
+
+	auto ComputeShader = View.ShaderMap->GetShader<FTranslucentLightingInjectBatchCS>(PermutationVector);
+
+	FIntVector NumGroups = FComputeShaderUtils::GetGroupCount(VolumeSize, FTranslucentLightingInjectBatchCS::GetGroupSize());
+
+	FComputeShaderUtils::AddPass(
+		GraphBuilder,
 		RDG_EVENT_NAME("InjectTranslucencyLightingVolumeBatch(VolumeCascade=%d,Max=%d%s)",
 			VolumeCascadeIndex, MaxBatchedLocalLights,
 			bSupportVirtualShadowMaps ? TEXT(",VirtualShadowMap") : TEXT("")),
+		ComputeShader,
 		PassParameters,
-		ERDGPassFlags::Raster,
-		[PassParameters, VertexShader, GeometryShader, PixelShader, VolumeCascadeIndex](FRDGAsyncTask, FRHICommandList& RHICmdList)
-	{
-		FGraphicsPipelineStateInitializer GraphicsPSOInit;
-
-		const FVolumeBounds VolumeBounds = FVolumeBounds(GetTranslucencyLightingVolumeDim());
-		SetPSOStateForVolumeInjection(VertexShader, GeometryShader, PixelShader.GetPixelShader(),
-			VolumeBounds, false, RHICmdList, GraphicsPSOInit);
-				
-		SetShaderParameters(RHICmdList, PixelShader, PixelShader.GetPixelShader(), *PassParameters);
-
-		RasterizeToVolumeTexture(RHICmdList, VolumeBounds);
-	});
+		NumGroups);
 }
 
 
@@ -1493,8 +1531,8 @@ void InjectTranslucencyLightingVolume(
 	const uint32 ViewIndex,
 	const FScene* Scene,
 	const FSceneRenderer& Renderer,
-	const FTranslucencyLightingVolumeTextures& Textures,
-	const FTranslucentLightInjectionCollector& Collector)
+	const FTranslucentLightInjectionCollector& Collector,
+	FTranslucencyLightingVolumeTextures& Textures)
 {
 	if (!GUseTranslucentLightingVolumes || !RHISupportsVolumeTextureRendering(View.GetShaderPlatform()))
 	{
@@ -1671,6 +1709,9 @@ void InjectTranslucencyLightingVolume(
 				});
 			}
 		}
+
+		Textures.Ambient[TextureIndex] = VolumeAmbientTexture;
+		Textures.Directional[TextureIndex] = VolumeDirectionalTexture;
 	}
 
 	GraphBuilder.FlushSetupQueue();
diff --git a/Engine/Source/Runtime/Renderer/Private/TranslucentLighting.h b/Engine/Source/Runtime/Renderer/Private/TranslucentLighting.h
--- a/Engine/Source/Runtime/Renderer/Private/TranslucentLighting.h
+++ b/Engine/Source/Runtime/Renderer/Private/TranslucentLighting.h
@@ -122,8 +122,8 @@ void InjectTranslucencyLightingVolume(
 	const uint32 ViewIndex,
 	const FScene* Scene,
 	const FSceneRenderer& Renderer,
-	const FTranslucencyLightingVolumeTextures& Textures,
-	const FTranslucentLightInjectionCollector& Collector);
+	const FTranslucentLightInjectionCollector& Collector,
+	FTranslucencyLightingVolumeTextures& Textures);
 
 void InjectSimpleTranslucencyLightingVolumeArray(
 	FRDGBuilder& GraphBuilder,
