Move management of runtime virtual texture state for primitive scene proxies to a new scene extension.
Use runtime virtual texture unique object id to handle the association between the runtime virtual texture and primitive scene proxies.
This also allows us to tidy up the runtime virtual texture scene proxy management and remove the need to track producer ids etc between game and render thread.
It also lifts the old restriction on the number of runtimve virtual textures allowed in a single scene.
Also change hide primitives logic to be global rather than tracking per runtime virtual texture volume, since that is the common intention.
Also made some code naming changes (where possible without requiring deprecation) to explitly name runtime/lightmap etc variations of virtual texture. This is more important now that we have many flavors of virtual texturing.

diff --git a/Engine/Plugins/Experimental/VirtualHeightfieldMesh/Source/VirtualHeightfieldMeshEditor/Private/HeightfieldMinMaxTextureBuild.cpp b/Engine/Plugins/Experimental/VirtualHeightfieldMesh/Source/VirtualHeightfieldMeshEditor/Private/HeightfieldMinMaxTextureBuild.cpp
--- a/Engine/Plugins/Experimental/VirtualHeightfieldMesh/Source/VirtualHeightfieldMeshEditor/Private/HeightfieldMinMaxTextureBuild.cpp
+++ b/Engine/Plugins/Experimental/VirtualHeightfieldMesh/Source/VirtualHeightfieldMeshEditor/Private/HeightfieldMinMaxTextureBuild.cpp
@@ -130,11 +130,11 @@ namespace VirtualHeightfieldMesh
 		}
 
 		FSceneInterface* Scene = VirtualTextureComponent->GetScene();
-		const uint32 VirtualTextureSceneIndex = RuntimeVirtualTexture::GetRuntimeVirtualTextureSceneIndex_GameThread(VirtualTextureComponent);
 		const FTransform Transform = VirtualTextureComponent->GetComponentTransform();
 		const FBox Bounds = VirtualTextureComponent->Bounds.GetBox();
 
-		URuntimeVirtualTexture const* VirtualTexture = InComponent->GetVirtualTexture();
+		URuntimeVirtualTexture* VirtualTexture = InComponent->GetVirtualTexture();
+		const int32 RuntimeVirtualTextureId = VirtualTexture->GetUniqueID();
 
 		FVTProducerDescription VTDesc;
 		VirtualTexture->GetProducerDescription(VTDesc, URuntimeVirtualTexture::FInitSettings(), Transform);
@@ -185,7 +185,7 @@ namespace VirtualHeightfieldMesh
 				IStreamingManager::Get().StreamAllResources(0);
 
 				ENQUEUE_RENDER_COMMAND(MinMaxTextureTileCommand)([
-					Scene, VirtualTextureSceneIndex,
+					Scene, RuntimeVirtualTextureId,
 					&RenderTileResources,
 					Transform, Bounds, UVRange,
 					TileX, TileY, TileSize,
@@ -199,7 +199,7 @@ namespace VirtualHeightfieldMesh
 
 					RuntimeVirtualTexture::FRenderPageBatchDesc Desc;
 					Desc.Scene = Scene->GetRenderScene();
-					Desc.RuntimeVirtualTextureMask = 1 << VirtualTextureSceneIndex;
+					Desc.RuntimeVirtualTextureId = RuntimeVirtualTextureId;
 					Desc.UVToWorld = Transform;
 					Desc.WorldBounds = Bounds;
 					Desc.MaterialType = ERuntimeVirtualTextureMaterialType::WorldHeight;
diff --git a/Engine/Source/Editor/VirtualTexturingEditor/Private/RuntimeVirtualTextureBuildStreamingMips.cpp b/Engine/Source/Editor/VirtualTexturingEditor/Private/RuntimeVirtualTextureBuildStreamingMips.cpp
--- a/Engine/Source/Editor/VirtualTexturingEditor/Private/RuntimeVirtualTextureBuildStreamingMips.cpp
+++ b/Engine/Source/Editor/VirtualTexturingEditor/Private/RuntimeVirtualTextureBuildStreamingMips.cpp
@@ -175,10 +175,10 @@ namespace RuntimeVirtualTexture
 		}
 
 		URuntimeVirtualTexture* RuntimeVirtualTexture = InComponent->GetVirtualTexture();
+		const int32 RuntimeVirtualTextureId = RuntimeVirtualTexture->GetUniqueID();
 		FSceneInterface* Scene = InComponent->GetScene();
 		checkf(Scene != nullptr, TEXT("Trying to build streaming mips for a component (from actor %s) that doesn't have a valid scene"), *InComponent->GetOwner()->GetActorNameOrLabel());
 
-		const uint32 VirtualTextureSceneIndex = RuntimeVirtualTexture::GetRuntimeVirtualTextureSceneIndex_GameThread(InComponent);
 		const FTransform Transform = InComponent->GetComponentTransform();
 		const FBox Bounds = InComponent->Bounds.GetBox();
 		const FVector4f CustomMaterialData = InComponent->GetCustomMaterialData();
@@ -248,7 +248,7 @@ namespace RuntimeVirtualTexture
 				IStreamingManager::Get().StreamAllResources(0);
 
 				ENQUEUE_RENDER_COMMAND(BakeStreamingTextureTileCommand)([
-					Scene, VirtualTextureSceneIndex, 
+					Scene, RuntimeVirtualTextureId, 
 					&RenderTileResources,
 					MaterialType, NumLayers,
 					Transform, Bounds, UVRange,
@@ -267,7 +267,7 @@ namespace RuntimeVirtualTexture
 
 						RuntimeVirtualTexture::FRenderPageBatchDesc Desc;
 						Desc.Scene = Scene->GetRenderScene();
-						Desc.RuntimeVirtualTextureMask = 1 << VirtualTextureSceneIndex;
+						Desc.RuntimeVirtualTextureId = RuntimeVirtualTextureId;
 						Desc.UVToWorld = Transform;
 						Desc.WorldBounds = Bounds;
 						Desc.MaterialType = MaterialType;
diff --git a/Engine/Source/Editor/VirtualTexturingEditor/Private/RuntimeVirtualTextureThumbnailRenderer.cpp b/Engine/Source/Editor/VirtualTexturingEditor/Private/RuntimeVirtualTextureThumbnailRenderer.cpp
--- a/Engine/Source/Editor/VirtualTexturingEditor/Private/RuntimeVirtualTextureThumbnailRenderer.cpp
+++ b/Engine/Source/Editor/VirtualTexturingEditor/Private/RuntimeVirtualTextureThumbnailRenderer.cpp
@@ -77,6 +77,7 @@ void URuntimeVirtualTextureThumbnailRenderer::Draw(UObject* Object, int32 X, int
 	}
 
 	URuntimeVirtualTexture* RuntimeVirtualTexture = Cast<URuntimeVirtualTexture>(Object);
+	const int32 RuntimeVirtualTextureId = RuntimeVirtualTexture->GetUniqueID();
 	URuntimeVirtualTextureComponent* RuntimeVirtualTextureComponent = FindComponent(RuntimeVirtualTexture);
 	FSceneInterface* Scene = RuntimeVirtualTextureComponent != nullptr ? RuntimeVirtualTextureComponent->GetScene() : nullptr;
 	check(Scene != nullptr);
@@ -89,7 +90,6 @@ void URuntimeVirtualTextureThumbnailRenderer::Draw(UObject* Object, int32 X, int
 	const FBox2D DestBox = FBox2D(FVector2D(X, Y), FVector2D(Width, Height));
 	const FTransform Transform = RuntimeVirtualTextureComponent->GetComponentTransform();
 	const FBox Bounds = RuntimeVirtualTextureComponent->Bounds.GetBox();
-	const uint32 VirtualTextureSceneIndex = RuntimeVirtualTexture::GetRuntimeVirtualTextureSceneIndex_GameThread(RuntimeVirtualTextureComponent);
 	const ERuntimeVirtualTextureMaterialType MaterialType = RuntimeVirtualTexture->GetMaterialType();
 	const FVector4f CustomMaterialData = RuntimeVirtualTextureComponent->GetCustomMaterialData();
 
@@ -102,7 +102,7 @@ void URuntimeVirtualTextureThumbnailRenderer::Draw(UObject* Object, int32 X, int
 	UE::RenderCommandPipe::FSyncScope SyncScope;
 
 	ENQUEUE_RENDER_COMMAND(BakeStreamingTextureTileCommand)(
-		[Scene, VirtualTextureSceneIndex, MaterialType, RenderTarget, DestBox, Transform, Bounds, MaxLevel, CustomMaterialData](FRHICommandListImmediate& RHICmdList)
+		[Scene, RuntimeVirtualTextureId, MaterialType, RenderTarget, DestBox, Transform, Bounds, MaxLevel, CustomMaterialData](FRHICommandListImmediate& RHICmdList)
 	{
 		FMaterialRenderProxy::UpdateDeferredCachedUniformExpressions();
 
@@ -110,7 +110,7 @@ void URuntimeVirtualTextureThumbnailRenderer::Draw(UObject* Object, int32 X, int
 
 		RuntimeVirtualTexture::FRenderPageBatchDesc Desc;
 		Desc.Scene = Scene->GetRenderScene();
-		Desc.RuntimeVirtualTextureMask = 1 << VirtualTextureSceneIndex;
+		Desc.RuntimeVirtualTextureId = RuntimeVirtualTextureId;
 		Desc.UVToWorld = Transform;
 		Desc.WorldBounds = Bounds;
 		Desc.MaterialType = MaterialType;
diff --git a/Engine/Source/Runtime/Engine/Private/PrimitiveSceneProxy.cpp b/Engine/Source/Runtime/Engine/Private/PrimitiveSceneProxy.cpp
--- a/Engine/Source/Runtime/Engine/Private/PrimitiveSceneProxy.cpp
+++ b/Engine/Source/Runtime/Engine/Private/PrimitiveSceneProxy.cpp
@@ -517,6 +517,7 @@ FPrimitiveSceneProxy::FPrimitiveSceneProxy(const FPrimitiveSceneProxyDesc& InPro
 ,	bSkinnedMesh(false)
 ,	bIsFirstPerson(InProxyDesc.bIsFirstPerson)
 ,	bIsFirstPersonWorldSpaceRepresentation(InProxyDesc.bIsFirstPersonWorldSpaceRepresentation)
+,	bSupportsRuntimeVirtualTexture(false)
 ,	bUseAsOccluder(InProxyDesc.bUseAsOccluder)
 ,	bSelectable(InProxyDesc.bSelectable)
 ,	bHasPerInstanceHitProxies(InProxyDesc.bHasPerInstanceHitProxies)
@@ -631,12 +632,12 @@ FPrimitiveSceneProxy::FPrimitiveSceneProxy(const FPrimitiveSceneProxyDesc& InPro
 	// Setup the runtime virtual texture information
 	if (UseVirtualTexturing(GetScene().GetShaderPlatform()))
 	{
-		for (URuntimeVirtualTexture* VirtualTexture : InProxyDesc.GetRuntimeVirtualTextures())
+		for (URuntimeVirtualTexture* RuntimeVirtualTexture : InProxyDesc.GetRuntimeVirtualTextures())
 		{
-			if (VirtualTexture != nullptr)
+			if (RuntimeVirtualTexture != nullptr)
 			{
-				RuntimeVirtualTextures.Add(VirtualTexture);
-				RuntimeVirtualTextureMaterialTypes.AddUnique(VirtualTexture->GetMaterialType());
+				RuntimeVirtualTextures.Add(RuntimeVirtualTexture->GetUniqueID());
+				RuntimeVirtualTextureMaterialTypes.AddUnique(RuntimeVirtualTexture->GetMaterialType());
 			}
 		}
 	}
@@ -1264,8 +1265,7 @@ void FPrimitiveSceneProxy::SetDrawDistance_RenderThread(float InMinDrawDistance,
 	MinDrawDistance = InMinDrawDistance;
 	MaxDrawDistance = InMaxDrawDistance > 0.0f ? InMaxDrawDistance : FLT_MAX;
 	// Modify max draw distance for main pass if we are using virtual texturing
-	const bool bUseVirtualTexture = RuntimeVirtualTextures.Num() > 0;
-	if (bUseVirtualTexture && InVirtualTextureMaxDrawDistance > 0.f)
+	if (WritesVirtualTexture() && InVirtualTextureMaxDrawDistance > 0.f)
 	{
 		MaxDrawDistance = FMath::Min(MaxDrawDistance, InVirtualTextureMaxDrawDistance);
 	}
@@ -1394,7 +1394,7 @@ void FPrimitiveSceneProxy::SetupInstanceSceneDataBuffers(const FInstanceSceneDat
 	InstanceSceneDataBuffersInternal = InInstanceSceneDataBuffers;
 }
 
-bool FPrimitiveSceneProxy::DrawInVirtualTextureOnly(bool bEditor) const
+bool FPrimitiveSceneProxy::DrawInRuntimeVirtualTextureOnly(bool bEditor) const
 {
 	if (bVirtualTextureMainPassDrawAlways)
 	{
@@ -1404,12 +1404,11 @@ bool FPrimitiveSceneProxy::DrawInVirtualTextureOnly(bool bEditor) const
 	{
 		return true;
 	}
-	// Conditional path tests the flags stored on scene virtual texture.
-	uint8 bHideMaskEditor, bHideMaskGame;
-	Scene->GetRuntimeVirtualTextureHidePrimitiveMask(bHideMaskEditor, bHideMaskGame);
-	const uint8 bHideMask = bEditor ? bHideMaskEditor : bHideMaskGame;
-	const uint8 RuntimeVirtualTextureMask = GetPrimitiveSceneInfo()->GetRuntimeVirtualTextureFlags().RuntimeVirtualTextureMask;
-	return (RuntimeVirtualTextureMask & bHideMask) != 0;
+	// Conditional path tests against the collected flags from all runtime virtual textures in the scene.
+ 	uint8 bHideMaskEditor, bHideMaskGame;
+ 	Scene->GetRuntimeVirtualTextureHidePrimitiveMask(bHideMaskEditor, bHideMaskGame);
+ 	const uint8 bHideMask = bEditor ? bHideMaskEditor : bHideMaskGame;
+ 	return bHideMask != 0;
 }
 
 void FPrimitiveSceneProxy::EnableGPUSceneSupportFlags()
@@ -1630,7 +1629,7 @@ bool FPrimitiveSceneProxy::IsShown(const FSceneView* View) const
 			return false;
 		}
 
-		if (DrawInVirtualTextureOnly(true) && !View->bIsVirtualTexture && !View->Family->EngineShowFlags.VirtualTexturePrimitives && !IsSelected())
+		if (DrawInRuntimeVirtualTextureOnly(true) && !View->bIsVirtualTexture && !View->Family->EngineShowFlags.VirtualTexturePrimitives && !IsSelected())
 		{
 			return false;
 		}
@@ -1648,7 +1647,7 @@ bool FPrimitiveSceneProxy::IsShown(const FSceneView* View) const
 			return false;
 		}
 
-		if (DrawInVirtualTextureOnly(false) && !View->bIsVirtualTexture)
+		if (DrawInRuntimeVirtualTextureOnly(false) && !View->bIsVirtualTexture)
 		{
 			return false;
 		}
@@ -1730,7 +1729,7 @@ bool FPrimitiveSceneProxy::IsShadowCast(const FSceneView* View) const
 			return false;
 		}
 
-		if (DrawInVirtualTextureOnly(View->Family->EngineShowFlags.Editor) && !View->bIsVirtualTexture)
+		if (DrawInRuntimeVirtualTextureOnly(View->Family->EngineShowFlags.Editor) && !View->bIsVirtualTexture)
 		{
 			return false;
 		}
diff --git a/Engine/Source/Runtime/Engine/Private/StaticMeshSceneProxy.cpp b/Engine/Source/Runtime/Engine/Private/StaticMeshSceneProxy.cpp
--- a/Engine/Source/Runtime/Engine/Private/StaticMeshSceneProxy.cpp
+++ b/Engine/Source/Runtime/Engine/Private/StaticMeshSceneProxy.cpp
@@ -277,6 +277,9 @@ FStaticMeshSceneProxy::FStaticMeshSceneProxy(const FStaticMeshSceneProxyDesc& In
 	// Static meshes do not deform internally (save by material effects such as WPO and PDO, which is allowed).
 	bHasDeformableMesh = false;
 
+	// Static meshes can write to runtime virtual texture if they are set to do so.
+	bSupportsRuntimeVirtualTexture = true;
+
 	bEvaluateWorldPositionOffset = InProxyDesc.bEvaluateWorldPositionOffset;
 
 	const auto FeatureLevel = GetScene().GetFeatureLevel();
@@ -2115,7 +2118,7 @@ void FStaticMeshSceneProxy::GetLCIs(FLCIArray& LCIs)
 
 bool FStaticMeshSceneProxy::CanBeOccluded() const
 {
-	return !MaterialRelevance.bDisableDepthTest && !MaterialRelevance.bPostMotionBlurTranslucency && !ShouldRenderCustomDepth() && !IsVirtualTextureOnly();
+	return !MaterialRelevance.bDisableDepthTest && !MaterialRelevance.bPostMotionBlurTranslucency && !ShouldRenderCustomDepth() && !IsRuntimeVirtualTextureOnly();
 }
 
 bool FStaticMeshSceneProxy::IsUsingDistanceCullFade() const
diff --git a/Engine/Source/Runtime/Engine/Public/PrimitiveSceneProxy.h b/Engine/Source/Runtime/Engine/Public/PrimitiveSceneProxy.h
--- a/Engine/Source/Runtime/Engine/Public/PrimitiveSceneProxy.h
+++ b/Engine/Source/Runtime/Engine/Public/PrimitiveSceneProxy.h
@@ -688,6 +688,7 @@ public:
 	inline int32 GetVirtualTextureLodBias() const { return VirtualTextureLodBias; }
 	inline int32 GetVirtualTextureCullMips() const { return VirtualTextureCullMips; }
 	inline int32 GetVirtualTextureMinCoverage() const {	return VirtualTextureMinCoverage; }
+	inline const TArray<int32>& GetRuntimeVirtualTextureIds() const { return RuntimeVirtualTextures; }
 
 	inline bool IsMovable() const 
 	{ 
@@ -741,8 +742,9 @@ public:
 	inline bool CastsStaticShadow() const { return bCastStaticShadow; }
 	inline bool CastsDynamicShadow() const { return bCastDynamicShadow; }
 	inline bool IsEmissiveLightSource() const { return bEmissiveLightSource; }
-	inline bool WritesVirtualTexture() const{ return RuntimeVirtualTextures.Num() > 0; }
-	inline bool WritesVirtualTexture(URuntimeVirtualTexture* VirtualTexture) const { return RuntimeVirtualTextures.Find(VirtualTexture) != INDEX_NONE; }
+	inline bool WritesVirtualTexture() const{ return bSupportsRuntimeVirtualTexture && RuntimeVirtualTextures.Num() > 0; }
+	UE_DEPRECATED(5.6, "Use GetRuntimeVirtualTextureIds() instead.")
+	inline bool WritesVirtualTexture(URuntimeVirtualTexture* VirtualTexture) const { return false; }
 	inline bool AffectsDynamicIndirectLighting() const { return bAffectDynamicIndirectLighting; }
 	inline bool AffectsDistanceFieldLighting() const { return bAffectDistanceFieldLighting; }
 	inline bool AffectsIndirectLightingWhileHidden() const { return bAffectIndirectLightingWhileHidden; }
@@ -1157,10 +1159,10 @@ protected:
 	ENGINE_API void UpdateVisibleInLumenScene();
 
 	/** Returns true if a primitive can never be rendered outside of a runtime virtual texture. */
-	bool IsVirtualTextureOnly() const { return bVirtualTextureMainPassDrawNever; }
+	bool IsRuntimeVirtualTextureOnly() const { return bVirtualTextureMainPassDrawNever; }
 
 	/** Returns true if a primitive should currently be hidden because it is drawn only to the runtime virtual texture. The result can depend on the current scene state. */
-	bool DrawInVirtualTextureOnly(bool bEditor) const;
+	bool DrawInRuntimeVirtualTextureOnly(bool bEditor) const;
 
 	/** Allow subclasses to override the primitive name. Used primarily by BSP. */
 	void OverrideOwnerName(FName InOwnerName)
@@ -1489,6 +1491,9 @@ protected:
 
 	/** If this is True, this primitive represents another first person primitive in world space. Mutually exclusive with bIsFirstPerson. */
 	uint8 bIsFirstPersonWorldSpaceRepresentation : 1;
+
+	/** This is only set true by the constructors of proxy types which support writing to runtime virtual texture. */
+	uint8 bSupportsRuntimeVirtualTexture : 1;
 	
 private:
 
@@ -1565,8 +1570,8 @@ protected:
 
 	FVector2f MinMaxMaterialDisplacement;
 
-	/** Array of runtime virtual textures that this proxy should render to. */
-	TArray<URuntimeVirtualTexture*> RuntimeVirtualTextures;
+	/** Array of runtime virtual textures unique ids that this proxy should render to. */
+	TArray<int32> RuntimeVirtualTextures;
 	/** Set of unique runtime virtual texture material types referenced by RuntimeVirtualTextures. */
 	TArray<ERuntimeVirtualTextureMaterialType> RuntimeVirtualTextureMaterialTypes;
 
diff --git a/Engine/Source/Runtime/Landscape/Private/LandscapeRender.cpp b/Engine/Source/Runtime/Landscape/Private/LandscapeRender.cpp
--- a/Engine/Source/Runtime/Landscape/Private/LandscapeRender.cpp
+++ b/Engine/Source/Runtime/Landscape/Private/LandscapeRender.cpp
@@ -1392,6 +1392,9 @@ FLandscapeComponentSceneProxy::FLandscapeComponentSceneProxy(ULandscapeComponent
 	bNaniteActive = InComponent->IsNaniteActive();
 	bUsesLandscapeCulling = (!bNaniteActive && Culling::UseCulling(GetScene().GetShaderPlatform()));
 
+	// Landscape meshes can write to runtime virtual texture if they are set to do so.
+	bSupportsRuntimeVirtualTexture = true;
+
 	EnableGPUSceneSupportFlags();
 
 	const ERHIFeatureLevel::Type FeatureLevel = GetScene().GetFeatureLevel();
@@ -1967,7 +1970,7 @@ FLandscapeComponentSceneProxy::~FLandscapeComponentSceneProxy()
 
 bool FLandscapeComponentSceneProxy::CanBeOccluded() const
 {
-	if (IsVirtualTextureOnly())
+	if (IsRuntimeVirtualTextureOnly())
 	{
 		return false;
 	}
diff --git a/Engine/Source/Runtime/Renderer/Private/PrimitiveSceneInfo.cpp b/Engine/Source/Runtime/Renderer/Private/PrimitiveSceneInfo.cpp
--- a/Engine/Source/Runtime/Renderer/Private/PrimitiveSceneInfo.cpp
+++ b/Engine/Source/Runtime/Renderer/Private/PrimitiveSceneInfo.cpp
@@ -297,16 +297,17 @@ FPrimitiveSceneInfo::FPrimitiveSceneInfo(const FPrimitiveSceneInfoAdapter& InAda
 	SceneData(InAdapter.SceneData),	
 	bNeedsUniformBufferUpdate(false),
 	bIndirectLightingCacheBufferDirty(false),
-	bRegisteredVirtualTextureProducerCallback(false),
+	bRegisteredLightmapVirtualTextureProducerCallback(false),
 	bRegisteredWithVelocityData(false),
 	bCacheShadowAsStatic((InAdapter.Mobility != EComponentMobility::Movable && InAdapter.SceneProxy->GetShadowCacheInvalidationBehavior() != EShadowCacheInvalidationBehavior::Always) || InAdapter.SceneProxy->GetShadowCacheInvalidationBehavior() == EShadowCacheInvalidationBehavior::Static),
 	bNaniteRasterBinsRenderCustomDepth(false),
 	bPendingAddToScene(false),
 	bPendingAddStaticMeshes(false),
-	bPendingFlushVirtualTexture(false),
+	bPendingFlushRuntimeVirtualTexture(false),
 	bNeedsCachedReflectionCaptureUpdate(true),
 	bShouldRenderInMainPass(InAdapter.SceneProxy->ShouldRenderInMainPass()),
 	bVisibleInRealTimeSkyCapture(InAdapter.SceneProxy->IsVisibleInRealTimeSkyCaptures()),
+	bWritesRuntimeVirtualTexture(InAdapter.SceneProxy->WritesVirtualTexture()),
 #if RHI_RAYTRACING
 	bDrawInGame(Proxy->IsDrawnInGame()),
 	bRayTracingFarField(Proxy->IsRayTracingFarField()),
@@ -1387,6 +1388,46 @@ void FPrimitiveSceneInfo::RemoveCachedRayTracingPrimitives()
 }
 #endif
 
+static bool GetRuntimeVirtualTextureLODRange(TArray<class FStaticMeshBatchRelevance> const& MeshRelevances, int8& OutMinLOD, int8& OutMaxLOD)
+{
+	OutMinLOD = MAX_int8;
+	OutMaxLOD = 0;
+
+	for (int32 MeshIndex = 0; MeshIndex < MeshRelevances.Num(); ++MeshIndex)
+	{
+		const FStaticMeshBatchRelevance& MeshRelevance = MeshRelevances[MeshIndex];
+		if (MeshRelevance.bRenderToVirtualTexture)
+		{
+			OutMinLOD = FMath::Min(OutMinLOD, MeshRelevance.GetLODIndex());
+			OutMaxLOD = FMath::Max(OutMaxLOD, MeshRelevance.GetLODIndex());
+		}
+	}
+
+	return OutMinLOD <= OutMaxLOD;
+}
+
+static FPrimitiveRuntimeVirtualTextureLodInfo BuildRuntimeVirtualTextureLodInfo(FPrimitiveSceneInfo const& InPrimitveSceneInfo)
+{
+	FPrimitiveRuntimeVirtualTextureLodInfo LodInfo;
+
+	if (InPrimitveSceneInfo.bWritesRuntimeVirtualTexture)
+	{
+		int8 MinLod, MaxLod;
+		if (GetRuntimeVirtualTextureLODRange(InPrimitveSceneInfo.StaticMeshRelevances, MinLod, MaxLod))
+		{
+			FPrimitiveSceneProxy* Proxy = InPrimitveSceneInfo.Proxy;
+
+			LodInfo.MinLod = FMath::Clamp((int32)MinLod, 0, 15);
+			LodInfo.MaxLod = FMath::Clamp((int32)MaxLod, 0, 15);
+			LodInfo.LodBias = FMath::Clamp(Proxy->GetVirtualTextureLodBias() + FPrimitiveRuntimeVirtualTextureLodInfo::LodBiasOffset, 0, 15);
+			LodInfo.CullMethod = Proxy->GetVirtualTextureMinCoverage() == 0 ? 0 : 1;
+			LodInfo.CullValue = LodInfo.CullMethod == 0 ? Proxy->GetVirtualTextureCullMips() : Proxy->GetVirtualTextureMinCoverage();
+		}
+	}
+
+	return LodInfo;
+}
+
 void FPrimitiveSceneInfo::AddStaticMeshes(FRHICommandListBase& RHICmdList, FScene* Scene, TArrayView<FPrimitiveSceneInfo*> SceneInfos, bool bCacheMeshDrawCommands)
 {
 	LLM_SCOPE(ELLMTag::StaticMesh);
@@ -1403,6 +1444,7 @@ void FPrimitiveSceneInfo::AddStaticMeshes(FRHICommandListBase& RHICmdList, FScen
 			SceneInfo->Proxy->DrawStaticElements(&BatchingSPDI);
 			SceneInfo->StaticMeshes.Shrink();
 			SceneInfo->StaticMeshRelevances.Shrink();
+			SceneInfo->RuntimeVirtualTextureLodInfo = BuildRuntimeVirtualTextureLodInfo(*SceneInfo);
 			SceneInfo->bPendingAddStaticMeshes = false;
 
 			check(SceneInfo->StaticMeshRelevances.Num() == SceneInfo->StaticMeshes.Num());
@@ -1447,7 +1489,7 @@ void FPrimitiveSceneInfo::AddStaticMeshes(FRHICommandListBase& RHICmdList, FScen
 	}
 }
 
-static void OnVirtualTextureDestroyed(const FVirtualTextureProducerHandle& InHandle, void* Baton)
+static void OnLightmapVirtualTextureDestroyed(const FVirtualTextureProducerHandle& InHandle, void* Baton)
 {
 	FPrimitiveSceneInfo* PrimitiveSceneInfo = static_cast<FPrimitiveSceneInfo*>(Baton);
 
@@ -1458,33 +1500,15 @@ static void OnVirtualTextureDestroyed(const FVirtualTextureProducerHandle& InHan
 	PrimitiveSceneInfo->Scene->GPUScene.AddPrimitiveToUpdate(PrimitiveSceneInfo->GetPersistentIndex(), EPrimitiveDirtyState::ChangedStaticLighting);
 }
 
-static void GetRuntimeVirtualTextureLODRange(TArray<class FStaticMeshBatchRelevance> const& MeshRelevances, int8& OutMinLOD, int8& OutMaxLOD)
-{
-	OutMinLOD = MAX_int8;
-	OutMaxLOD = 0;
-
-	for (int32 MeshIndex = 0; MeshIndex < MeshRelevances.Num(); ++MeshIndex)
-	{
-		const FStaticMeshBatchRelevance& MeshRelevance = MeshRelevances[MeshIndex];
-		if (MeshRelevance.bRenderToVirtualTexture)
-		{
-			OutMinLOD = FMath::Min(OutMinLOD, MeshRelevance.GetLODIndex());
-			OutMaxLOD = FMath::Max(OutMaxLOD, MeshRelevance.GetLODIndex());
-		}
-	}
-
-	check(OutMinLOD <= OutMaxLOD);
-}
-
 int32 FPrimitiveSceneInfo::UpdateStaticLightingBuffer()
 {
 	checkSlow(IsInRenderingThread());
 
-	if (bRegisteredVirtualTextureProducerCallback)
+	if (bRegisteredLightmapVirtualTextureProducerCallback)
 	{
 		// Remove any previous VT callbacks
 		FVirtualTextureSystem::Get().RemoveAllProducerDestroyedCallbacks(this);
-		bRegisteredVirtualTextureProducerCallback = false;
+		bRegisteredLightmapVirtualTextureProducerCallback = false;
 	}
 
 	FPrimitiveSceneProxy::FLCIArray LCIs;
@@ -1502,8 +1526,8 @@ int32 FPrimitiveSceneInfo::UpdateStaticLightingBuffer()
 			FVirtualTextureProducerHandle VTProducerHandle;
 			if (LCI->GetVirtualTextureLightmapProducer(Scene->GetFeatureLevel(), VTProducerHandle))
 			{
-				FVirtualTextureSystem::Get().AddProducerDestroyedCallback(VTProducerHandle, &OnVirtualTextureDestroyed, this);
-				bRegisteredVirtualTextureProducerCallback = true;
+				FVirtualTextureSystem::Get().AddProducerDestroyedCallback(VTProducerHandle, &OnLightmapVirtualTextureDestroyed, this);
+				bRegisteredLightmapVirtualTextureProducerCallback = true;
 			}
 		}
 	}
@@ -1812,32 +1836,6 @@ void FPrimitiveSceneInfo::AddToScene(FScene* Scene, TArrayView<FPrimitiveSceneIn
 	}
 }
 
-void FPrimitiveSceneInfo::UpdateVirtualTextures(FScene* Scene, TArrayView<FPrimitiveSceneInfo*> SceneInfos)
-{
-	SCOPED_NAMED_EVENT(FPrimitiveSceneInfo_AddToScene_UpdateVirtualTexture, FColor::Emerald);
-	for (FPrimitiveSceneInfo* SceneInfo : SceneInfos)
-	{
-		FPrimitiveSceneProxy* Proxy = SceneInfo->Proxy;
-		// Store the runtime virtual texture flags.
-		SceneInfo->UpdateRuntimeVirtualTextureFlags();
-		Scene->PrimitiveVirtualTextureFlags[SceneInfo->PackedIndex] = SceneInfo->RuntimeVirtualTextureFlags;
-
-		// Store the runtime virtual texture Lod info.
-		if (SceneInfo->RuntimeVirtualTextureFlags.bRenderToVirtualTexture)
-		{
-			int8 MinLod, MaxLod;
-			GetRuntimeVirtualTextureLODRange(SceneInfo->StaticMeshRelevances, MinLod, MaxLod);
-
-			FPrimitiveVirtualTextureLodInfo& LodInfo = Scene->PrimitiveVirtualTextureLod[SceneInfo->PackedIndex];
-			LodInfo.MinLod = FMath::Clamp((int32)MinLod, 0, 15);
-			LodInfo.MaxLod = FMath::Clamp((int32)MaxLod, 0, 15);
-			LodInfo.LodBias = FMath::Clamp(Proxy->GetVirtualTextureLodBias() + FPrimitiveVirtualTextureLodInfo::LodBiasOffset, 0, 15);
-			LodInfo.CullMethod = Proxy->GetVirtualTextureMinCoverage() == 0 ? 0 : 1;
-			LodInfo.CullValue = LodInfo.CullMethod == 0 ? Proxy->GetVirtualTextureCullMips() : Proxy->GetVirtualTextureMinCoverage();
-		}
-	}
-}
-
 void FPrimitiveSceneInfo::RemoveStaticMeshes()
 {
 	// Deallocate potential OIT dynamic index buffer
@@ -1915,10 +1913,10 @@ void FPrimitiveSceneInfo::RemoveFromScene(bool bUpdateStaticDrawLists)
 		RemoveStaticMeshes();
 	}
 
-	if (bRegisteredVirtualTextureProducerCallback)
+	if (bRegisteredLightmapVirtualTextureProducerCallback)
 	{
 		FVirtualTextureSystem::Get().RemoveAllProducerDestroyedCallbacks(this);
-		bRegisteredVirtualTextureProducerCallback = false;
+		bRegisteredLightmapVirtualTextureProducerCallback = false;
 	}
 
 	if (Proxy->ShouldNotifyOnWorldAddRemove())
@@ -1941,38 +1939,6 @@ void FPrimitiveSceneInfo::RemoveFromScene(bool bUpdateStaticDrawLists)
 	}
 }
 
-void FPrimitiveSceneInfo::UpdateRuntimeVirtualTextureFlags()
-{
-	RuntimeVirtualTextureFlags.bRenderToVirtualTexture = false;
-	RuntimeVirtualTextureFlags.RuntimeVirtualTextureMask = 0;
-
-	if (Proxy->WritesVirtualTexture())
-	{
-		if (Proxy->IsNaniteMesh())
-		{
-			UE_LOG(LogRenderer, Warning, TEXT("Rendering a nanite mesh to a runtime virtual texture isn't yet supported. Please disable this option on primitive component : %s"), *Proxy->GetOwnerName().ToString());
-		}
-		else if (StaticMeshes.Num() > 0)
-		{
-			RuntimeVirtualTextureFlags.bRenderToVirtualTexture = true;
-
-			// Performance assumption: The arrays of runtime virtual textures are small (less that 5?) so that O(n^2) scan isn't expensive
-			for (TSparseArray<FRuntimeVirtualTextureSceneProxy*>::TConstIterator It(Scene->RuntimeVirtualTextures); It; ++It)
-			{
-				int32 SceneIndex = It.GetIndex();
-				if (SceneIndex < FPrimitiveVirtualTextureFlags::RuntimeVirtualTexture_BitCount)
-				{
-					URuntimeVirtualTexture* SceneVirtualTexture = (*It)->VirtualTexture;
-					if (Proxy->WritesVirtualTexture(SceneVirtualTexture))
-					{
-						RuntimeVirtualTextureFlags.RuntimeVirtualTextureMask |= 1 << SceneIndex;
-					}
-				}
-			}
-		}
-	}
-}
-
 void FPrimitiveSceneInfo::UpdateStaticMeshes(FScene* Scene, TArrayView<FPrimitiveSceneInfo*> SceneInfos, EUpdateStaticMeshFlags UpdateFlags, bool bReAddToDrawLists)
 {
 	QUICK_SCOPE_CYCLE_COUNTER(STAT_FPrimitiveSceneInfo_UpdateStaticMeshes);
@@ -2095,18 +2061,14 @@ FInstanceDataBufferHeader FPrimitiveSceneInfo::GetInstanceDataHeader() const
 
 void FPrimitiveSceneInfo::FlushRuntimeVirtualTexture()
 {
-	if (RuntimeVirtualTextureFlags.bRenderToVirtualTexture)
+	if (bWritesRuntimeVirtualTexture)
 	{
-		uint32 RuntimeVirtualTextureIndex = 0;
-		uint32 Mask = RuntimeVirtualTextureFlags.RuntimeVirtualTextureMask;
-		while (Mask != 0)
+		for (TSparseArray<FRuntimeVirtualTextureSceneProxy*>::TIterator It(Scene->RuntimeVirtualTextures); It; ++It)
 		{
-			if (Mask & 1)
+			if (Proxy->GetRuntimeVirtualTextureIds().Contains((*It)->RuntimeVirtualTextureId))
 			{
-				Scene->RuntimeVirtualTextures[RuntimeVirtualTextureIndex]->Dirty(Proxy->GetBounds(), EVTInvalidatePriority::Normal);
+				(*It)->Dirty(Proxy->GetBounds(), EVTInvalidatePriority::Normal);
 			}
-			Mask >>= 1;
-			RuntimeVirtualTextureIndex++;
 		}
 	}
 }
diff --git a/Engine/Source/Runtime/Renderer/Private/RendererScene.cpp b/Engine/Source/Runtime/Renderer/Private/RendererScene.cpp
--- a/Engine/Source/Runtime/Renderer/Private/RendererScene.cpp
+++ b/Engine/Source/Runtime/Renderer/Private/RendererScene.cpp
@@ -1331,8 +1331,6 @@ void FScene::CheckPrimitiveArrays(int MaxTypeOffsetIndex)
 	check(Primitives.Num() == PrimitiveOctreeIndex.Num());
 	check(Primitives.Num() == PrimitiveOcclusionFlags.Num());
 	check(Primitives.Num() == PrimitiveComponentIds.Num());
-	check(Primitives.Num() == PrimitiveVirtualTextureFlags.Num());
-	check(Primitives.Num() == PrimitiveVirtualTextureLod.Num());
 	check(Primitives.Num() == PrimitiveOcclusionBounds.Num());
 #if WITH_EDITOR
 	check(Primitives.Num() == PrimitivesSelected.Num());
@@ -1647,14 +1645,14 @@ FScene::FScene(UWorld* InWorld, bool bInRequiresHitProxies, bool bInIsEditorScen
 ,	PrimitiveOctree(FVector::ZeroVector, UE_OLD_HALF_WORLD_MAX)
 ,	bRequiresHitProxies(bInRequiresHitProxies)
 ,	bIsEditorScene(bInIsEditorScene)
+,	bRuntimeVirtualTexturePrimitiveHideEditor(false)
+,	bRuntimeVirtualTexturePrimitiveHideGame(false)
 ,	NumUncachedStaticLightingInteractions(0)
 ,	NumUnbuiltReflectionCaptures(0)
 ,	NumMobileStaticAndCSMLights_RenderThread(0)
 ,	NumMobileMovableDirectionalLights_RenderThread(0)
 ,	GPUSkinCache(nullptr)
 ,	SceneLODHierarchy(this)
-,	RuntimeVirtualTexturePrimitiveHideMaskEditor(0)
-,	RuntimeVirtualTexturePrimitiveHideMaskGame(0)
 ,	DefaultMaxDistanceFieldOcclusionDistance(InWorld->GetWorldSettings()->DefaultMaxDistanceFieldOcclusionDistance)
 ,	GlobalDistanceFieldViewDistance(InWorld->GetWorldSettings()->GlobalDistanceFieldViewDistance)
 ,	DynamicIndirectShadowsSelfShadowingIntensity(FMath::Clamp(InWorld->GetWorldSettings()->DynamicIndirectShadowsSelfShadowingIntensity, 0.0f, 1.0f))
@@ -3580,7 +3578,19 @@ void FScene::RemovePrecomputedVolumetricLightmap(const FPrecomputedVolumetricLig
 	});
 }
 
-void FScene::AddRuntimeVirtualTexture(class URuntimeVirtualTextureComponent* Component)
+static void UpdateRuntimeVirtualTextureHidePrimitiveMasks(TSparseArray<FRuntimeVirtualTextureSceneProxy*> const& InProxies, bool& bOutEditor, bool& bOutGame)
+{
+	// Set the global flag to hide when any one volume sets it.
+	bOutEditor = false;
+	bOutGame = false;
+	for (TSparseArray<FRuntimeVirtualTextureSceneProxy*>::TConstIterator It(InProxies); It && !bOutEditor && !bOutGame; ++It)
+	{
+		bOutEditor |= (*It)->bHidePrimitivesInEditor;
+		bOutGame |= (*It)->bHidePrimitivesInGame;
+	}
+}
+
+void FScene::AddRuntimeVirtualTexture(URuntimeVirtualTextureComponent* Component)
 {
 	if (Component->SceneProxy == nullptr)
 	{
@@ -3592,35 +3602,35 @@ void FScene::AddRuntimeVirtualTexture(class URuntimeVirtualTextureComponent* Com
 		ENQUEUE_RENDER_COMMAND(AddRuntimeVirtualTextureCommand)(
 			[Scene, SceneProxy] (FRHICommandListBase&)
 		{
-			Scene->AddRuntimeVirtualTexture_RenderThread(SceneProxy);
-			Scene->UpdateRuntimeVirtualTextureForAllPrimitives_RenderThread();
+			SceneProxy->SceneIndex = Scene->RuntimeVirtualTextures.Add(SceneProxy);
+			
+			UpdateRuntimeVirtualTextureHidePrimitiveMasks(Scene->RuntimeVirtualTextures, Scene->bRuntimeVirtualTexturePrimitiveHideEditor, Scene->bRuntimeVirtualTexturePrimitiveHideGame);
 		});
 	}
 	else
 	{
 		// This is a component update.
-		// Store the new FRuntimeVirtualTextureSceneProxy at the same index as the old (to avoid needing to update any associated primitives).
-		// Defer old proxy deletion to the render thread.
-		FRuntimeVirtualTextureSceneProxy* SceneProxyToReplace = Component->SceneProxy;
+		// Delete old proxy and add new one on the render thread.
+		FRuntimeVirtualTextureSceneProxy* SceneProxyToRemove = Component->SceneProxy;
 		Component->SceneProxy = new FRuntimeVirtualTextureSceneProxy(Component);
 
 		FScene* Scene = this;
-		FRuntimeVirtualTextureSceneProxy* SceneProxy = Component->SceneProxy;
+		FRuntimeVirtualTextureSceneProxy* SceneProxyToAdd = Component->SceneProxy;
 
-		ENQUEUE_RENDER_COMMAND(AddRuntimeVirtualTextureCommand)(
-			[Scene, SceneProxy, SceneProxyToReplace] (FRHICommandListBase&)
+		ENQUEUE_RENDER_COMMAND(UpdateRuntimeVirtualTextureCommand)(
+			[Scene, SceneProxyToAdd, SceneProxyToRemove] (FRHICommandListBase&)
 		{
-			const bool bUpdatePrimitives = SceneProxy->VirtualTexture != SceneProxyToReplace->VirtualTexture;
-			Scene->UpdateRuntimeVirtualTexture_RenderThread(SceneProxy, SceneProxyToReplace);
-			if (bUpdatePrimitives)
-			{
-				Scene->UpdateRuntimeVirtualTextureForAllPrimitives_RenderThread();
-			}
+			Scene->RuntimeVirtualTextures.RemoveAt(SceneProxyToRemove->SceneIndex);
+			delete SceneProxyToRemove;
+
+			SceneProxyToAdd->SceneIndex = Scene->RuntimeVirtualTextures.Add(SceneProxyToAdd);
+
+			UpdateRuntimeVirtualTextureHidePrimitiveMasks(Scene->RuntimeVirtualTextures, Scene->bRuntimeVirtualTexturePrimitiveHideEditor, Scene->bRuntimeVirtualTexturePrimitiveHideGame);
 		});
 	}
 }
 
-void FScene::RemoveRuntimeVirtualTexture(class URuntimeVirtualTextureComponent* Component)
+void FScene::RemoveRuntimeVirtualTexture(URuntimeVirtualTextureComponent* Component)
 {
 	FRuntimeVirtualTextureSceneProxy* SceneProxy = Component->SceneProxy;
 	if (SceneProxy != nullptr)
@@ -3633,95 +3643,26 @@ void FScene::RemoveRuntimeVirtualTexture(class URuntimeVirtualTextureComponent*
 		ENQUEUE_RENDER_COMMAND(RemoveRuntimeVirtualTextureCommand)(
 			[Scene, SceneProxy] (FRHICommandListBase&)
 		{
-			Scene->RemoveRuntimeVirtualTexture_RenderThread(SceneProxy);
-			Scene->UpdateRuntimeVirtualTextureForAllPrimitives_RenderThread();
-		});
-	}
-}
-
-void FScene::AddRuntimeVirtualTexture_RenderThread(FRuntimeVirtualTextureSceneProxy* SceneProxy)
-{
-	SceneProxy->SceneIndex = RuntimeVirtualTextures.Add(SceneProxy);
-
-	const uint8 HideFlagBit = 1 << SceneProxy->SceneIndex;
-	RuntimeVirtualTexturePrimitiveHideMaskEditor &= ~HideFlagBit;
-	RuntimeVirtualTexturePrimitiveHideMaskEditor |= (SceneProxy->bHidePrimitivesInEditor ? HideFlagBit : 0);
-	RuntimeVirtualTexturePrimitiveHideMaskGame &= ~HideFlagBit;
-	RuntimeVirtualTexturePrimitiveHideMaskGame |= (SceneProxy->bHidePrimitivesInGame ? HideFlagBit : 0);
-}
-
-void FScene::UpdateRuntimeVirtualTexture_RenderThread(FRuntimeVirtualTextureSceneProxy* SceneProxy, FRuntimeVirtualTextureSceneProxy* SceneProxyToReplace)
-{
-	for (TSparseArray<FRuntimeVirtualTextureSceneProxy*>::TIterator It(RuntimeVirtualTextures); It; ++It)
-	{
-		if (*It == SceneProxyToReplace)
-		{
-			SceneProxy->SceneIndex = It.GetIndex();
-			*It = SceneProxy;
-			delete SceneProxyToReplace;
-
-			const uint8 HideFlagBit = 1 << SceneProxy->SceneIndex;
-			RuntimeVirtualTexturePrimitiveHideMaskEditor &= ~HideFlagBit;
-			RuntimeVirtualTexturePrimitiveHideMaskEditor |= (SceneProxy->bHidePrimitivesInEditor ? HideFlagBit : 0);
-			RuntimeVirtualTexturePrimitiveHideMaskGame &= ~HideFlagBit;
-			RuntimeVirtualTexturePrimitiveHideMaskGame |= (SceneProxy->bHidePrimitivesInGame ? HideFlagBit : 0);
-
-			return;
-		}
-	}
-	// If we get here then we didn't find the object to replace!
-	check(false);
-}
-
-void FScene::RemoveRuntimeVirtualTexture_RenderThread(FRuntimeVirtualTextureSceneProxy* SceneProxy)
-{
-	const uint8 HideFlagBit = 1 << SceneProxy->SceneIndex;
-	RuntimeVirtualTexturePrimitiveHideMaskEditor &= ~HideFlagBit;
-	RuntimeVirtualTexturePrimitiveHideMaskGame &= ~HideFlagBit;
-
-	RuntimeVirtualTextures.RemoveAt(SceneProxy->SceneIndex);
-	delete SceneProxy;
-}
-
-void FScene::UpdateRuntimeVirtualTextureForAllPrimitives_RenderThread()
-{
-	for (int32 Index = 0; Index < Primitives.Num(); ++Index)
-	{
-		if (PrimitiveVirtualTextureFlags[Index].bRenderToVirtualTexture)
-		{
-			Primitives[Index]->UpdateRuntimeVirtualTextureFlags();
-			PrimitiveVirtualTextureFlags[Index] = Primitives[Index]->GetRuntimeVirtualTextureFlags();
-		}
-	}
-}
+			Scene->RuntimeVirtualTextures.RemoveAt(SceneProxy->SceneIndex);
+			delete SceneProxy;
 
-uint32 FScene::GetRuntimeVirtualTextureSceneIndex(uint32 ProducerId)
-{
-	checkSlow(IsInParallelRenderingThread());
-	for (FRuntimeVirtualTextureSceneProxy const* Proxy : RuntimeVirtualTextures)
-	{
-		if (Proxy->ProducerId == ProducerId)
-		{
-			return Proxy->SceneIndex;
-		}
+			UpdateRuntimeVirtualTextureHidePrimitiveMasks(Scene->RuntimeVirtualTextures, Scene->bRuntimeVirtualTexturePrimitiveHideEditor, Scene->bRuntimeVirtualTexturePrimitiveHideGame);
+		});
 	}
-	// Should not get here
-	check(false);
-	return 0;
 }
 
 void FScene::GetRuntimeVirtualTextureHidePrimitiveMask(uint8& bHideMaskEditor, uint8& bHideMaskGame) const
 {
-	bHideMaskEditor = RuntimeVirtualTexturePrimitiveHideMaskEditor;
-	bHideMaskGame = RuntimeVirtualTexturePrimitiveHideMaskGame;
+	bHideMaskEditor = bRuntimeVirtualTexturePrimitiveHideEditor ? 0xff : 0;
+	bHideMaskGame = bRuntimeVirtualTexturePrimitiveHideGame ? 0xff : 0;
 }
 
-void FScene::InvalidateRuntimeVirtualTexture(class URuntimeVirtualTextureComponent* Component, FBoxSphereBounds const& WorldBounds, EVTInvalidatePriority InvalidatePriority)
+void FScene::InvalidateRuntimeVirtualTexture(URuntimeVirtualTextureComponent* Component, FBoxSphereBounds const& WorldBounds, EVTInvalidatePriority InvalidatePriority)
 {
 	if (Component->SceneProxy != nullptr)
 	{
 		FRuntimeVirtualTextureSceneProxy* SceneProxy = Component->SceneProxy;
-		ENQUEUE_RENDER_COMMAND(RuntimeVirtualTextureComponent_SetDirty)(
+		ENQUEUE_RENDER_COMMAND(InvalidateRuntimeVirtualTextureCommand)(
 			[SceneProxy, WorldBounds, InvalidatePriority] (FRHICommandListBase&)
 		{
 			SceneProxy->Dirty(WorldBounds, InvalidatePriority);
@@ -3729,6 +3670,15 @@ void FScene::InvalidateRuntimeVirtualTexture(class URuntimeVirtualTextureCompone
 	}
 }
 
+void FScene::FlushDirtyRuntimeVirtualTextures()
+{
+	checkSlow(IsInRenderingThread());
+	for (TSparseArray<FRuntimeVirtualTextureSceneProxy*>::TIterator It(RuntimeVirtualTextures); It; ++It)
+	{
+		(*It)->FlushDirtyPages();
+	}
+}
+
 void FScene::InvalidatePathTracedOutput()
 {
 	// NOTE: this is an atomic, so this function is ok to call from any thread
@@ -3752,15 +3702,6 @@ void FScene::InvalidateLumenSurfaceCache_GameThread(UPrimitiveComponent* Compone
 	}
 }
 
-void FScene::FlushDirtyRuntimeVirtualTextures()
-{
-	checkSlow(IsInRenderingThread());
-	for (TSparseArray<FRuntimeVirtualTextureSceneProxy*>::TIterator It(RuntimeVirtualTextures); It; ++It)
-	{
-		(*It)->FlushDirtyPages();
-	}
-}
-
 bool FScene::GetPreviousLocalToWorld(const FPrimitiveSceneInfo* PrimitiveSceneInfo, FMatrix& OutPreviousLocalToWorld) const
 {
 	return VelocityData.GetComponentPreviousLocalToWorld(PrimitiveSceneInfo->PrimitiveComponentId, OutPreviousLocalToWorld);
@@ -5823,8 +5764,6 @@ void FScene::Update(FRDGBuilder& GraphBuilder, const FUpdateParameters& Paramete
 							TArraySwapElements(PrimitiveOctreeIndex, DestIndex, SourceIndex);
 							TArraySwapElements(PrimitiveOcclusionFlags, DestIndex, SourceIndex);
 							TArraySwapElements(PrimitiveComponentIds, DestIndex, SourceIndex);
-							TArraySwapElements(PrimitiveVirtualTextureFlags, DestIndex, SourceIndex);
-							TArraySwapElements(PrimitiveVirtualTextureLod, DestIndex, SourceIndex);
 							TArraySwapElements(PrimitiveOcclusionBounds, DestIndex, SourceIndex);
 						#if WITH_EDITOR
 							TBitArraySwapElements(PrimitivesSelected, DestIndex, SourceIndex);
@@ -5883,8 +5822,6 @@ void FScene::Update(FRDGBuilder& GraphBuilder, const FUpdateParameters& Paramete
 			PrimitiveOctreeIndex.RemoveAt(SourceIndex, RemoveCount, EAllowShrinking::No);
 			PrimitiveOcclusionFlags.RemoveAt(SourceIndex, RemoveCount, EAllowShrinking::No);
 			PrimitiveComponentIds.RemoveAt(SourceIndex, RemoveCount, EAllowShrinking::No);
-			PrimitiveVirtualTextureFlags.RemoveAt(SourceIndex, RemoveCount, EAllowShrinking::No);
-			PrimitiveVirtualTextureLod.RemoveAt(SourceIndex, RemoveCount, EAllowShrinking::No);
 			PrimitiveOcclusionBounds.Remove(RemoveCount, EAllowShrinking::No);
 
 			#if WITH_EDITOR
@@ -5962,18 +5899,18 @@ void FScene::Update(FRDGBuilder& GraphBuilder, const FUpdateParameters& Paramete
 	const int32 SceneInfosContainerReservedSize = AddedPrimitiveSceneInfos.Num() + UpdatedTransforms.Num() + UpdatedInstances.Num();
 
 	TArray<FPrimitiveSceneInfo*, SceneRenderingAllocator>& SceneInfosWithAddToScene = *GraphBuilder.AllocObject<TArray<FPrimitiveSceneInfo*, SceneRenderingAllocator>>();
-	TArray<FPrimitiveSceneInfo*, SceneRenderingAllocator> SceneInfosWithFlushVirtualTexture;
+	TArray<FPrimitiveSceneInfo*, SceneRenderingAllocator> SceneInfosWithFlushRuntimeVirtualTexture;
 	TArray<FPrimitiveSceneInfo*, SceneRenderingAllocator>& SceneInfosWithStaticDrawListUpdate = *GraphBuilder.AllocObject<TArray<FPrimitiveSceneInfo*, SceneRenderingAllocator>>();
 	SceneInfosWithAddToScene.Reserve(SceneInfosContainerReservedSize);
-	SceneInfosWithFlushVirtualTexture.Reserve(SceneInfosContainerReservedSize);
+	SceneInfosWithFlushRuntimeVirtualTexture.Reserve(SceneInfosContainerReservedSize);
 	SceneInfosWithStaticDrawListUpdate.Reserve(SceneInfosContainerReservedSize);
 
-	const auto QueueFlushVirtualTexture = [&](FPrimitiveSceneInfo* SceneInfo) -> bool
+	const auto QueueFlushRuntimeVirtualTexture = [&](FPrimitiveSceneInfo* SceneInfo) -> bool
 	{
-		if (!SceneInfo->bPendingFlushVirtualTexture)
+		if (SceneInfo->bWritesRuntimeVirtualTexture && !SceneInfo->bPendingFlushRuntimeVirtualTexture)
 		{
-			SceneInfo->bPendingFlushVirtualTexture = true;
-			SceneInfosWithFlushVirtualTexture.Push(SceneInfo);
+			SceneInfo->bPendingFlushRuntimeVirtualTexture = true;
+			SceneInfosWithFlushRuntimeVirtualTexture.Push(SceneInfo);
 			return true;
 		}
 		return false;
@@ -6070,8 +6007,6 @@ void FScene::Update(FRDGBuilder& GraphBuilder, const FUpdateParameters& Paramete
 			PrimitiveVisibilityIds.Reserve(PrimitiveVisibilityIds.Num() + AddedLocalPrimitiveSceneInfos.Num());
 			PrimitiveOcclusionFlags.Reserve(PrimitiveOcclusionFlags.Num() + AddedLocalPrimitiveSceneInfos.Num());
 			PrimitiveComponentIds.Reserve(PrimitiveComponentIds.Num() + AddedLocalPrimitiveSceneInfos.Num());
-			PrimitiveVirtualTextureFlags.Reserve(PrimitiveVirtualTextureFlags.Num() + AddedLocalPrimitiveSceneInfos.Num());
-			PrimitiveVirtualTextureLod.Reserve(PrimitiveVirtualTextureLod.Num() + AddedLocalPrimitiveSceneInfos.Num());
 			PrimitiveOcclusionBounds.Reserve(PrimitiveOcclusionBounds.Num() + AddedLocalPrimitiveSceneInfos.Num());
 		#if WITH_EDITOR
 			PrimitivesSelected.Reserve(PrimitivesSelected.Num() + AddedLocalPrimitiveSceneInfos.Num());
@@ -6110,8 +6045,6 @@ void FScene::Update(FRDGBuilder& GraphBuilder, const FUpdateParameters& Paramete
 					PrimitiveOctreeIndex.Add(0);
 					PrimitiveOcclusionFlags.AddUninitialized();
 					PrimitiveComponentIds.AddUninitialized();
-					PrimitiveVirtualTextureFlags.AddUninitialized();
-					PrimitiveVirtualTextureLod.AddUninitialized();
 					PrimitiveOcclusionBounds.AddUninitialized();
 				#if WITH_EDITOR
 					PrimitivesSelected.Add(PrimitiveSceneInfo->Proxy->WantsEditorEffects() || PrimitiveSceneInfo->Proxy->IsSelected());
@@ -6219,8 +6152,6 @@ void FScene::Update(FRDGBuilder& GraphBuilder, const FUpdateParameters& Paramete
 							TArraySwapElements(PrimitiveOctreeIndex, DestIndex, SourceIndex);
 							TArraySwapElements(PrimitiveOcclusionFlags, DestIndex, SourceIndex);
 							TArraySwapElements(PrimitiveComponentIds, DestIndex, SourceIndex);
-							TArraySwapElements(PrimitiveVirtualTextureFlags, DestIndex, SourceIndex);
-							TArraySwapElements(PrimitiveVirtualTextureLod, DestIndex, SourceIndex);
 							TArraySwapElements(PrimitiveOcclusionBounds, DestIndex, SourceIndex);
 						#if WITH_EDITOR
 							TBitArraySwapElements(PrimitivesSelected, DestIndex, SourceIndex);
@@ -6284,7 +6215,7 @@ void FScene::Update(FRDGBuilder& GraphBuilder, const FUpdateParameters& Paramete
 				}
 #endif
 
-				QueueFlushVirtualTexture(PrimitiveSceneInfo);
+				QueueFlushRuntimeVirtualTexture(PrimitiveSceneInfo);
 				QueueAddToScene(PrimitiveSceneInfo);
 				QueueAddStaticMeshes(PrimitiveSceneInfo);
 
@@ -6326,7 +6257,7 @@ void FScene::Update(FRDGBuilder& GraphBuilder, const FUpdateParameters& Paramete
 				}
 			}
 
-			QueueFlushVirtualTexture(PrimitiveSceneInfo);
+			QueueFlushRuntimeVirtualTexture(PrimitiveSceneInfo);
 
 			if (ShouldPrimitiveOutputVelocity(PrimitiveSceneInfo->Proxy, GetShaderPlatform()))
 			{
@@ -6379,7 +6310,7 @@ void FScene::Update(FRDGBuilder& GraphBuilder, const FUpdateParameters& Paramete
 			const FUpdateInstanceCommand &UpdateInstance = Item.Payload;
 			FScopeCycleCounter Context(PrimitiveSceneProxy->GetStatId());
 
-			QueueFlushVirtualTexture(PrimitiveSceneInfo);
+			QueueFlushRuntimeVirtualTexture(PrimitiveSceneInfo);
 
 			// TODO: no need to do this if only the payload size changed, we only need it because the MDC stores the instance count!
 			//       Better yet: don't update MDCs on instance data change as we can pull it from elsewhere.
@@ -6570,7 +6501,7 @@ void FScene::Update(FRDGBuilder& GraphBuilder, const FUpdateParameters& Paramete
 	const bool bAsyncCacheMeshDrawCommands = EnumHasAnyFlags(Parameters.AsyncOps, EUpdateAllPrimitiveSceneInfosAsyncOps::CacheMeshDrawCommands) && GRHISupportsMultithreadedShaderCreation;
 
 	UE::Tasks::FTask AddStaticMeshesTask = GraphBuilder.AddCommandListSetupTask(
-		[this, AddStaticMeshes = CopyTemp(SceneInfosWithStaticDrawListUpdate), SceneInfosWithFlushVirtualTexture = MoveTemp(SceneInfosWithFlushVirtualTexture), &SceneInfosWithAddToScene]
+		[this, AddStaticMeshes = CopyTemp(SceneInfosWithStaticDrawListUpdate), SceneInfosWithFlushRuntimeVirtualTexture = MoveTemp(SceneInfosWithFlushRuntimeVirtualTexture), &SceneInfosWithAddToScene]
 			(FRHICommandListBase& RHICmdList) mutable
 	{
 		SCOPED_NAMED_EVENT(StaticMeshUpdate, FColor::Emerald);
@@ -6580,12 +6511,10 @@ void FScene::Update(FRDGBuilder& GraphBuilder, const FUpdateParameters& Paramete
 			FPrimitiveSceneInfo::AddStaticMeshes(RHICmdList, this, AddStaticMeshes, false);
 		}
 
-		FPrimitiveSceneInfo::UpdateVirtualTextures(this, SceneInfosWithAddToScene);
-
-		for (FPrimitiveSceneInfo* PrimitiveSceneInfo : SceneInfosWithFlushVirtualTexture)
+		for (FPrimitiveSceneInfo* PrimitiveSceneInfo : SceneInfosWithFlushRuntimeVirtualTexture)
 		{
 			PrimitiveSceneInfo->FlushRuntimeVirtualTexture();
-			PrimitiveSceneInfo->bPendingFlushVirtualTexture = false;
+			PrimitiveSceneInfo->bPendingFlushRuntimeVirtualTexture = false;
 		}
 
 	}, MakeArrayView({ UpdateUniformExpressionsTask, GPUSkinTasks.Inline }), UE::Tasks::ETaskPriority::High, bAsyncCacheMeshDrawCommands);
diff --git a/Engine/Source/Runtime/Renderer/Private/ScenePrivate.h b/Engine/Source/Runtime/Renderer/Private/ScenePrivate.h
--- a/Engine/Source/Runtime/Renderer/Private/ScenePrivate.h
+++ b/Engine/Source/Runtime/Renderer/Private/ScenePrivate.h
@@ -2955,10 +2955,6 @@ public:
 	TScenePrimitiveArray<FBoxSphereBounds> PrimitiveOcclusionBounds;
 	/** Packed array of primitive components associated with the primitive. */
 	TArray<FPrimitiveComponentId> PrimitiveComponentIds;
-	/** Packed array of runtime virtual texture flags. */
-	TArray<FPrimitiveVirtualTextureFlags> PrimitiveVirtualTextureFlags;
-	/** Packed array of runtime virtual texture lod info. */
-	TArray<FPrimitiveVirtualTextureLodInfo> PrimitiveVirtualTextureLod;
 #if RHI_RAYTRACING
 	/** Packed array of ray tracing primitive caching flags*/
 	TArray<ERayTracingPrimitiveFlags> PrimitiveRayTracingFlags;
@@ -3255,6 +3251,10 @@ public:
 	/** Whether this is an editor scene. */
 	bool bIsEditorScene;
 
+	/** Determines whether primitives that draw to a runtime virtual texture should also be drawn in the main pass. */
+	bool bRuntimeVirtualTexturePrimitiveHideEditor;
+	bool bRuntimeVirtualTexturePrimitiveHideGame;
+
 	/** Set by the rendering thread to signal to the game thread that the scene needs a static lighting build. */
 	volatile mutable int32 NumUncachedStaticLightingInteractions;
 
@@ -3287,10 +3287,6 @@ public:
 
 	LightFunctionAtlas::FLightFunctionAtlasSceneData LightFunctionAtlasSceneData;
 
-	/** Mask used to determine whether primitives that draw to a runtime virtual texture should also be drawn in the main pass. */
-	uint8 RuntimeVirtualTexturePrimitiveHideMaskEditor;
-	uint8 RuntimeVirtualTexturePrimitiveHideMaskGame;
-
 	float DefaultMaxDistanceFieldOcclusionDistance;
 
 	float GlobalDistanceFieldViewDistance;
@@ -3667,9 +3663,6 @@ public:
 		return PrimitiveComponentIds;
 	}
 
-	/** Get the scene index of the FRuntimeVirtualTextureSceneProxy associated with the producer. */
-	uint32 GetRuntimeVirtualTextureSceneIndex(uint32 ProducerId);
-
 	/** Flush any dirty runtime virtual texture pages */
 	void FlushDirtyRuntimeVirtualTextures();
 
@@ -3952,15 +3945,6 @@ private:
 	/** Updates all static draw lists. */
 	void UpdateStaticDrawLists_RenderThread(FRHICommandListImmediate& RHICmdList);
 
-	/** Add a runtime virtual texture proxy to the scene. Called in the rendering thread by AddRuntimeVirtualTexture. */
-	void AddRuntimeVirtualTexture_RenderThread(FRuntimeVirtualTextureSceneProxy* SceneProxy);
-	/** Update a runtime virtual texture proxy to the scene. Called in the rendering thread by AddRuntimeVirtualTexture. */
-	void UpdateRuntimeVirtualTexture_RenderThread(FRuntimeVirtualTextureSceneProxy* SceneProxy, FRuntimeVirtualTextureSceneProxy* SceneProxyToReplace);
-	/** Remove a runtime virtual texture proxy from the scene. Called in the rendering thread by RemoveRuntimeVirtualTexture. */
-	void RemoveRuntimeVirtualTexture_RenderThread(FRuntimeVirtualTextureSceneProxy* SceneProxy);
-	/** Update the scene primitive data after completing operations that add or remove runtime virtual textures. This can be slow and should be called rarely. */
-	void UpdateRuntimeVirtualTextureForAllPrimitives_RenderThread();
-
 	/**
 	 * Shifts scene data by provided delta
 	 * Called on world origin changes
diff --git a/Engine/Source/Runtime/Renderer/Private/VT/RuntimeVirtualTextureProducer.cpp b/Engine/Source/Runtime/Renderer/Private/VT/RuntimeVirtualTextureProducer.cpp
--- a/Engine/Source/Runtime/Renderer/Private/VT/RuntimeVirtualTextureProducer.cpp
+++ b/Engine/Source/Runtime/Renderer/Private/VT/RuntimeVirtualTextureProducer.cpp
@@ -10,7 +10,7 @@
 
 FRuntimeVirtualTextureFinalizer::FRuntimeVirtualTextureFinalizer(
 	FVTProducerDescription const& InDesc, 
-	uint32 InProducerId, 
+	int32 InRuntimeVirtualTextureId,
 	ERuntimeVirtualTextureMaterialType InMaterialType, 
 	bool InClearTextures, 
 	FSceneInterface* InScene, 
@@ -18,8 +18,7 @@ FRuntimeVirtualTextureFinalizer::FRuntimeVirtualTextureFinalizer(
 	FBox const& InWorldBounds,
 	FVector4f const& InCustomMaterialData)
 	: Desc(InDesc)
-	, ProducerId(InProducerId)
-	, RuntimeVirtualTextureMask(0)
+	, RuntimeVirtualTextureId(InRuntimeVirtualTextureId)
 	, MaterialType(InMaterialType)
 	, bClearTextures(InClearTextures)
 	, Scene(InScene)
@@ -34,25 +33,6 @@ bool FRuntimeVirtualTextureFinalizer::IsReady()
 	return RuntimeVirtualTexture::IsSceneReadyToRender(Scene);
 }
 
-void FRuntimeVirtualTextureFinalizer::InitProducer(const FVirtualTextureProducerHandle& ProducerHandle)
-{
-	if (RuntimeVirtualTextureMask == 0)
-	{
-		FScene* RenderScene = Scene->GetRenderScene();
-
-		// Initialize the RuntimeVirtualTextureMask by matching this producer with those registered in the scene's runtime virtual textures.
-		// We only need to do this once. If the associated scene proxy is removed this finalizer will also be destroyed.
-		const uint32 VirtualTextureSceneIndex = RenderScene->GetRuntimeVirtualTextureSceneIndex(ProducerId);
-		RuntimeVirtualTextureMask = 1 << VirtualTextureSceneIndex;
-
-		//todo[vt]: 
-		// Add a slow render path inside RenderPage() when this check fails. 
-		// It will need to iterate the virtual textures on each primitive instead of using the RuntimeVirtualTextureMask.
-		// Currently nothing will render for this finalizer when the check fails.
-		checkSlow(VirtualTextureSceneIndex < FPrimitiveVirtualTextureFlags::RuntimeVirtualTexture_BitCount);
-	}
-}
-
 void FRuntimeVirtualTextureFinalizer::AddTile(FTileEntry& Tile)
 {
 	Tiles.Add(Tile);
@@ -65,7 +45,7 @@ void FRuntimeVirtualTextureFinalizer::Finalize(FRDGBuilder& GraphBuilder)
 
 	RuntimeVirtualTexture::FRenderPageBatchDesc RenderPageBatchDesc;
 	RenderPageBatchDesc.Scene = Scene->GetRenderScene();
-	RenderPageBatchDesc.RuntimeVirtualTextureMask = RuntimeVirtualTextureMask;
+	RenderPageBatchDesc.RuntimeVirtualTextureId = RuntimeVirtualTextureId;
 	RenderPageBatchDesc.UVToWorld = UVToWorld;
 	RenderPageBatchDesc.WorldBounds = WorldBounds;
 	RenderPageBatchDesc.MaterialType = MaterialType;
@@ -142,26 +122,26 @@ void FRuntimeVirtualTextureFinalizer::Finalize(FRDGBuilder& GraphBuilder)
 
 FRuntimeVirtualTextureProducer::FRuntimeVirtualTextureProducer(
 	FVTProducerDescription const& InDesc,
-	uint32 InProducerId,
+	int32 InRuntimeVirtualTextureId,
 	ERuntimeVirtualTextureMaterialType InMaterialType,
 	bool InClearTextures,
 	FSceneInterface* InScene,
 	FTransform const& InUVToWorld,
 	FBox const& InWorldBounds)
-	: Finalizer(InDesc, InProducerId, InMaterialType, InClearTextures, InScene, InUVToWorld, InWorldBounds, FVector4f::Zero())
+	: Finalizer(InDesc, InRuntimeVirtualTextureId, InMaterialType, InClearTextures, InScene, InUVToWorld, InWorldBounds, FVector4f::Zero())
 {
 }
 
 FRuntimeVirtualTextureProducer::FRuntimeVirtualTextureProducer(
 	FVTProducerDescription const& InDesc, 
-	uint32 InProducerId, 
-	ERuntimeVirtualTextureMaterialType InMaterialType, 
+	int32 InRuntimeVirtualTextureId,
+	ERuntimeVirtualTextureMaterialType InMaterialType,
 	bool InClearTextures, 
 	FSceneInterface* InScene, 
 	FTransform const& InUVToWorld,
 	FBox const& InWorldBounds,
 	FVector4f const& InCustomMaterialData)
-	: Finalizer(InDesc, InProducerId, InMaterialType, InClearTextures, InScene, InUVToWorld, InWorldBounds, InCustomMaterialData)
+	: Finalizer(InDesc, InRuntimeVirtualTextureId, InMaterialType, InClearTextures, InScene, InUVToWorld, InWorldBounds, InCustomMaterialData)
 {
 }
 
@@ -210,7 +190,6 @@ IVirtualTextureFinalizer* FRuntimeVirtualTextureProducer::ProducePageData(
 		}
 	}
 
-	Finalizer.InitProducer(ProducerHandle);
 	Finalizer.AddTile(Tile);
 
 	return &Finalizer;
diff --git a/Engine/Source/Runtime/Renderer/Private/VT/RuntimeVirtualTextureProducer.h b/Engine/Source/Runtime/Renderer/Private/VT/RuntimeVirtualTextureProducer.h
--- a/Engine/Source/Runtime/Renderer/Private/VT/RuntimeVirtualTextureProducer.h
+++ b/Engine/Source/Runtime/Renderer/Private/VT/RuntimeVirtualTextureProducer.h
@@ -17,7 +17,7 @@ class FRuntimeVirtualTextureFinalizer : public IVirtualTextureFinalizer
 public:
 	FRuntimeVirtualTextureFinalizer(
 		FVTProducerDescription const& InDesc, 
-		uint32 InProducerId, 
+		int32 InRuntimeVirtualTextureId,
 		ERuntimeVirtualTextureMaterialType InMaterialType, 
 		bool InClearTextures, 
 		FSceneInterface* InScene, 
@@ -38,9 +38,6 @@ public:
 	/** Returns false if we don't yet have everything we need to render a VT page. */
 	bool IsReady();
 
-	/** Does some one time work at the first call to set up the Producer */
-	void InitProducer(const FVirtualTextureProducerHandle& ProducerHandle);
-
 	/** Add a tile to the finalize queue. */
 	void AddTile(FTileEntry& Tile);
 
@@ -51,10 +48,8 @@ public:
 private:
 	/** Description of our virtual texture. */
 	const FVTProducerDescription Desc;
-	/** Producer index used once to build the RuntimeVirtualTextureMask for any rendering. */
-	uint32 ProducerId;
-	/** The mask of runtime virtual textures in the scene that we should render to. */
-	uint32 RuntimeVirtualTextureMask;
+	/** Object ID of our virtual texture. */
+	int32 RuntimeVirtualTextureId;
 	/** Contents of virtual texture layer stack. */
 	ERuntimeVirtualTextureMaterialType MaterialType;
 	/** Clear before render flag. */
@@ -78,7 +73,7 @@ public:
 	UE_DEPRECATED(5.6, "Use constructor that takes InCustomMaterialData.")
 	RENDERER_API FRuntimeVirtualTextureProducer(
 		FVTProducerDescription const& InDesc, 
-		uint32 InProducerId, 
+		int32 InRuntimeVirtualTextureId,
 		ERuntimeVirtualTextureMaterialType InMaterialType, 
 		bool InClearTextures, 
 		FSceneInterface* InScene, 
@@ -87,7 +82,7 @@ public:
 
 	RENDERER_API FRuntimeVirtualTextureProducer(
 		FVTProducerDescription const& InDesc,
-		uint32 InProducerId,
+		int32 InRuntimeVirtualTextureId,
 		ERuntimeVirtualTextureMaterialType InMaterialType,
 		bool InClearTextures,
 		FSceneInterface* InScene,
diff --git a/Engine/Source/Runtime/Renderer/Private/VT/RuntimeVirtualTextureRender.cpp b/Engine/Source/Runtime/Renderer/Private/VT/RuntimeVirtualTextureRender.cpp
--- a/Engine/Source/Runtime/Renderer/Private/VT/RuntimeVirtualTextureRender.cpp
+++ b/Engine/Source/Runtime/Renderer/Private/VT/RuntimeVirtualTextureRender.cpp
@@ -18,6 +18,7 @@
 #include "PostProcess/SceneRenderTargets.h"
 #include "ShaderBaseClasses.h"
 #include "VT/RuntimeVirtualTexture.h"
+#include "VT/RuntimeVirtualTextureSceneExtension.h"
 #include "VT/RuntimeVirtualTextureSceneProxy.h"
 #include "MeshPassProcessor.inl"
 #include "RenderCaptureInterface.h"
@@ -1084,7 +1085,7 @@ namespace RuntimeVirtualTexture
 
 
 	/** Collect meshes to draw. */
-	void GatherMeshesToDraw(FDynamicPassMeshDrawListContext* DynamicMeshPassContext, FScene const* Scene, FViewInfo* View, ERuntimeVirtualTextureMaterialType MaterialType, uint32 RuntimeVirtualTextureMask, uint8 vLevel, uint8 MaxLevel, bool bAllowCachedMeshDrawCommands)
+	void GatherMeshesToDraw(FDynamicPassMeshDrawListContext* DynamicMeshPassContext, FScene const* Scene, FViewInfo* View, ERuntimeVirtualTextureMaterialType MaterialType, int32 RuntimeVirtualTextureId, uint8 vLevel, uint8 MaxLevel, bool bAllowCachedMeshDrawCommands)
 	{
 		// Cached draw command collectors
 		const FCachedPassMeshDrawList& SceneDrawList = Scene->CachedDrawLists[EMeshPass::VirtualTexture];
@@ -1096,20 +1097,11 @@ namespace RuntimeVirtualTexture
 		const float RcpWorldSize = 1.f / (View->ViewMatrices.GetInvProjectionMatrix().M[0][0]);
 		const float WorldToPixel = View->ViewRect.Width() * RcpWorldSize;
 
-		// Iterate over scene and collect visible virtual texture draw commands for this view
-		//todo: Consider a broad phase (quad tree etc?) here. (But only if running over PrimitiveVirtualTextureFlags shows up as a bottleneck.)
-		for (int32 PrimitiveIndex = 0; PrimitiveIndex < Scene->Primitives.Num(); ++PrimitiveIndex)
-		{
-			const FPrimitiveVirtualTextureFlags Flags = Scene->PrimitiveVirtualTextureFlags[PrimitiveIndex];
-			if (!Flags.bRenderToVirtualTexture)
-			{
-				continue;
-			}
-			if ((Flags.RuntimeVirtualTextureMask & RuntimeVirtualTextureMask) == 0)
-			{
-				continue;
-			}
+		TArray<int32> PrimitiveIndices;
+		Scene->GetExtensionPtr<FRuntimeVirtualTextureSceneExtension>()->GetPrimitivesForRuntimeVirtualTexture(Scene, RuntimeVirtualTextureId, PrimitiveIndices);
 
+		for (const int32 PrimitiveIndex : PrimitiveIndices)
+ 		{
 			//todo[vt]: In our case we know that frustum is an oriented box so investigate cheaper test for intersecting that
 			const FSphere SphereBounds = Scene->PrimitiveBounds[PrimitiveIndex].BoxSphereBounds.GetSphere();
 			if (!View->ViewFrustum.IntersectSphere(SphereBounds.Center, SphereBounds.W))
@@ -1117,8 +1109,10 @@ namespace RuntimeVirtualTexture
 				continue;
 			}
 
+			FPrimitiveSceneInfo* RESTRICT PrimitiveSceneInfo = Scene->Primitives[PrimitiveIndex];
+
 			// Cull primitives according to mip level or pixel coverage
-			const FPrimitiveVirtualTextureLodInfo LodInfo = Scene->PrimitiveVirtualTextureLod[PrimitiveIndex];
+			const FPrimitiveRuntimeVirtualTextureLodInfo LodInfo = PrimitiveSceneInfo->GetRuntimeVirtualTextureLodInfo();
 			if (LodInfo.CullMethod == 0)
 			{
 				if (MaxLevel - vLevel < LodInfo.CullValue)
@@ -1136,7 +1130,6 @@ namespace RuntimeVirtualTexture
 				}
 			}
 
-			FPrimitiveSceneInfo* RESTRICT PrimitiveSceneInfo = Scene->Primitives[PrimitiveIndex];
 			FMeshDrawCommandPrimitiveIdInfo IdInfo = PrimitiveSceneInfo->GetMDCIdInfo();
 
 			// Calculate Lod for current mip
@@ -1144,7 +1137,7 @@ namespace RuntimeVirtualTexture
 			const int32 CurFirstLODIdx = PrimitiveSceneInfo->Proxy->GetCurrentFirstLODIdx_RenderThread();
 			const int32 MinLODIdx = FMath::Max((int32)LodInfo.MinLod, CurFirstLODIdx);
 			const int32 MaxLODIdx = FMath::Max((int32)LodInfo.MaxLod, CurFirstLODIdx);
-			const int32 LodBias = (int32)LodInfo.LodBias - FPrimitiveVirtualTextureLodInfo::LodBiasOffset;
+			const int32 LodBias = (int32)LodInfo.LodBias - FPrimitiveRuntimeVirtualTextureLodInfo::LodBiasOffset;
 			const int32 LodIndex = FMath::Clamp<int32>(LodBias - FMath::FloorToInt(FMath::Log2(AreaRatio)), MinLODIdx, MaxLODIdx);
 
 			// Process meshes
@@ -1493,7 +1486,7 @@ namespace RuntimeVirtualTexture
 		FRDGBuilder& GraphBuilder,
 		FScene* Scene,
 		ISceneRenderer* SceneRenderer,
-		uint32 RuntimeVirtualTextureMask,
+		int32 RuntimeVirtualTextureId,
 		ERuntimeVirtualTextureMaterialType MaterialType,
 		bool bClearTextures,
 		bool bIsThumbnails,
@@ -1618,9 +1611,9 @@ namespace RuntimeVirtualTexture
 			PassParameters->RenderTargets[2] = GraphSetup.RenderTexture2 ? FRenderTargetBinding(GraphSetup.RenderTexture2, LoadAction) : FRenderTargetBinding();
     		
 			AddSimpleMeshPass(GraphBuilder, PassParameters, Scene, *View, nullptr, RDG_EVENT_NAME("VirtualTextureDraw"), View->ViewRect,
-	    	[Scene, View, MaterialType, RuntimeVirtualTextureMask, vLevel, MaxLevel, bAllowCachedMeshDrawCommands](FDynamicPassMeshDrawListContext* DynamicMeshPassContext)
+	    	[Scene, View, MaterialType, RuntimeVirtualTextureId, vLevel, MaxLevel, bAllowCachedMeshDrawCommands](FDynamicPassMeshDrawListContext* DynamicMeshPassContext)
 	    	{
-	    		GatherMeshesToDraw(DynamicMeshPassContext, Scene, View, MaterialType, RuntimeVirtualTextureMask, vLevel, MaxLevel, bAllowCachedMeshDrawCommands);
+	    		GatherMeshesToDraw(DynamicMeshPassContext, Scene, View, MaterialType, RuntimeVirtualTextureId, vLevel, MaxLevel, bAllowCachedMeshDrawCommands);
 	    	});
 		}
 
@@ -1706,7 +1699,7 @@ namespace RuntimeVirtualTexture
 					GraphBuilder,
 					InDesc.Scene,
 					SceneRenderer,
-					InDesc.RuntimeVirtualTextureMask,
+					InDesc.RuntimeVirtualTextureId,
 					InDesc.MaterialType,
 					InDesc.bClearTextures,
 					InDesc.bIsThumbnails,
@@ -1759,19 +1752,27 @@ namespace RuntimeVirtualTexture
 		}
 	}
 
-	uint32 GetRuntimeVirtualTextureSceneIndex_GameThread(class URuntimeVirtualTextureComponent* InComponent)
+	/** This function is deprecated! */
+	uint32 GetRuntimeVirtualTextureSceneIndex_GameThread(URuntimeVirtualTextureComponent* InComponent)
 	{
+		if (InComponent == nullptr || InComponent->GetScene() == nullptr || InComponent->GetVirtualTexture() == nullptr)
+		{
+			return ~0u;
+		}
+
 		int32 SceneIndex = 0;
+		FSceneInterface* SceneInterface = InComponent->GetScene();
+		int32 RuntimeVirtualTextureId = InComponent->GetVirtualTexture()->GetUniqueID();
+
 		ENQUEUE_RENDER_COMMAND(GetSceneIndexCommand)(
-			[&SceneIndex, InComponent](FRHICommandListImmediate& RHICmdList)
+			[&SceneIndex, SceneInterface, RuntimeVirtualTextureId](FRHICommandListImmediate& RHICmdList)
 		{
-			if (InComponent->GetScene() != nullptr)
+			if (FScene* Scene = SceneInterface->GetRenderScene())
 			{
-				FScene* Scene = InComponent->GetScene()->GetRenderScene();
-				if (Scene != nullptr && InComponent->SceneProxy != nullptr)
+				SceneIndex = Scene->RuntimeVirtualTextures.IndexOfByPredicate([RuntimeVirtualTextureId] (FRuntimeVirtualTextureSceneProxy const* SceneProxy)
 				{
-					SceneIndex = Scene->GetRuntimeVirtualTextureSceneIndex(InComponent->SceneProxy->ProducerId);
-				}
+					return SceneProxy->RuntimeVirtualTextureId == RuntimeVirtualTextureId;
+				});
 			}
 		});
 		FlushRenderingCommands();
diff --git a/Engine/Source/Runtime/Renderer/Private/VT/RuntimeVirtualTextureSceneExtension.cpp b/Engine/Source/Runtime/Renderer/Private/VT/RuntimeVirtualTextureSceneExtension.cpp
new file mode 100644
--- /dev/null
+++ b/Engine/Source/Runtime/Renderer/Private/VT/RuntimeVirtualTextureSceneExtension.cpp
@@ -0,0 +1,113 @@
+// Copyright Epic Games, Inc. All Rights Reserved.
+
+#include "VT/RuntimeVirtualTextureSceneExtension.h"
+
+#include "RenderUtils.h"
+#include "ScenePrivate.h"
+
+IMPLEMENT_SCENE_EXTENSION(FRuntimeVirtualTextureSceneExtension);
+
+/** Persistant scene data stored by the extension. */
+struct FRuntimeVirtualTextureSceneExtensionData
+{
+	/** A bit array and count of active persistant primitive ids for a single runtime virtual texture. */
+	struct FActivePrimitives
+	{
+		TBitArray<> BitArray;
+		int32 SetBitCount = 0;
+	};
+	
+	/** Map of runtime virtual texture id to entry data. */
+	TMap<int32, FActivePrimitives> RuntimeVirtualTextureMap;
+};
+
+FRuntimeVirtualTextureSceneExtension::FRuntimeVirtualTextureSceneExtension(FScene& InScene)
+	: ISceneExtension(InScene)
+	, Data(new FRuntimeVirtualTextureSceneExtensionData)
+{
+}
+
+bool FRuntimeVirtualTextureSceneExtension::ShouldCreateExtension(FScene& InScene)
+{
+	return UseVirtualTexturing(InScene.GetShaderPlatform());
+}
+
+void FRuntimeVirtualTextureSceneExtension::GetPrimitivesForRuntimeVirtualTexture(FScene const* InScene, int32 InRuntimeVirtualTextureId, TArray<int32>& OutPrimitiveIndices) const
+{
+	if (FRuntimeVirtualTextureSceneExtensionData::FActivePrimitives* Found = Data->RuntimeVirtualTextureMap.Find(InRuntimeVirtualTextureId))
+	{
+		OutPrimitiveIndices.Reserve(Found->SetBitCount);
+
+		int32 Index = Found->BitArray.Find(true);
+		while (Index != INDEX_NONE)
+		{
+			OutPrimitiveIndices.Add(InScene->GetPrimitiveIndex(FPersistentPrimitiveIndex{Index}));
+			Index = Found->BitArray.FindFrom(true, Index + 1);
+		}
+	}
+}
+
+class FRuntimeVirtualTextureSceneExtensionUpdater : public ISceneExtensionUpdater
+{
+	DECLARE_SCENE_EXTENSION_UPDATER(FRenderer, FRuntimeVirtualTextureSceneExtension);
+
+public:
+	FRuntimeVirtualTextureSceneExtensionUpdater(FRuntimeVirtualTextureSceneExtensionData& InData) 
+		: Data(InData)
+	{
+	}
+
+	virtual void PreSceneUpdate(FRDGBuilder& GraphBuilder, const FScenePreUpdateChangeSet& ChangeSet, FSceneUniformBuffer& SceneUniforms) override
+	{
+		for (FPrimitiveSceneInfo* PrimitiveSceneInfo : ChangeSet.RemovedPrimitiveSceneInfos)
+		{
+			if (PrimitiveSceneInfo->bWritesRuntimeVirtualTexture)
+			{
+				if (FPrimitiveSceneProxy* Proxy = PrimitiveSceneInfo->Proxy)
+				{
+					const FPersistentPrimitiveIndex PersistentIndex = PrimitiveSceneInfo->GetPersistentIndex();
+
+					for (int32 RuntimeVirtualTextureId : Proxy->GetRuntimeVirtualTextureIds())
+					{
+						FRuntimeVirtualTextureSceneExtensionData::FActivePrimitives& Entry = Data.RuntimeVirtualTextureMap.FindChecked(RuntimeVirtualTextureId);
+						ensure(Entry.BitArray[PersistentIndex.Index] == true);
+						Entry.BitArray[PersistentIndex.Index] = false;
+						Entry.SetBitCount--;
+					}
+				}
+			}
+		}
+	}
+
+	virtual void PostSceneUpdate(FRDGBuilder& GraphBuilder, const FScenePostUpdateChangeSet& ChangeSet) override
+	{
+		for (FPrimitiveSceneInfo* PrimitiveSceneInfo : ChangeSet.AddedPrimitiveSceneInfos)
+		{
+			if (PrimitiveSceneInfo->bWritesRuntimeVirtualTexture)
+			{
+				if (FPrimitiveSceneProxy* Proxy = PrimitiveSceneInfo->Proxy)
+				{
+					const FPersistentPrimitiveIndex PersistentIndex = PrimitiveSceneInfo->GetPersistentIndex();
+
+					for (int32 RuntimeVirtualTextureId : Proxy->GetRuntimeVirtualTextureIds())
+					{
+						FRuntimeVirtualTextureSceneExtensionData::FActivePrimitives& Entry = Data.RuntimeVirtualTextureMap.FindOrAdd(RuntimeVirtualTextureId);
+						Entry.BitArray.PadToNum(PersistentIndex.Index + 1, false);
+						ensure(Entry.BitArray[PersistentIndex.Index] == false);
+						Entry.BitArray[PersistentIndex.Index] = true;
+						Entry.SetBitCount++;
+					}
+				}
+			}
+		}
+	}
+
+private:
+	/** Reference to the owner's extension data. */
+	FRuntimeVirtualTextureSceneExtensionData& Data;
+};
+
+ISceneExtensionUpdater* FRuntimeVirtualTextureSceneExtension::CreateUpdater()
+{
+	return new FRuntimeVirtualTextureSceneExtensionUpdater(*Data);
+}
diff --git a/Engine/Source/Runtime/Renderer/Private/VT/RuntimeVirtualTextureSceneExtension.h b/Engine/Source/Runtime/Renderer/Private/VT/RuntimeVirtualTextureSceneExtension.h
new file mode 100644
--- /dev/null
+++ b/Engine/Source/Runtime/Renderer/Private/VT/RuntimeVirtualTextureSceneExtension.h
@@ -0,0 +1,27 @@
+// Copyright Epic Games, Inc. All Rights Reserved.
+
+#pragma once
+
+#include "SceneExtensions.h"
+
+/** Scene extension for storing information about runtime virtual textures in the scene. */
+class FRuntimeVirtualTextureSceneExtension : public ISceneExtension
+{
+	DECLARE_SCENE_EXTENSION(RENDERER_API, FRuntimeVirtualTextureSceneExtension);
+
+public:
+	using ISceneExtension::ISceneExtension;
+
+	FRuntimeVirtualTextureSceneExtension(FScene& InScene);
+
+	//~ Begin ISceneExtension Interface.
+	static bool ShouldCreateExtension(FScene& Scene);
+	virtual ISceneExtensionUpdater* CreateUpdater() override;
+	//~ End ISceneExtension Interface.
+
+	/** Get a list of scene primitive indices for a given runtime virtual texture. */
+	void GetPrimitivesForRuntimeVirtualTexture(FScene const* InScene, int32 InRuntimeVirtualTextureId, TArray<int32>& OutPrimitiveIndices) const;
+
+private:
+	TUniquePtr<struct FRuntimeVirtualTextureSceneExtensionData> Data;
+};
diff --git a/Engine/Source/Runtime/Renderer/Private/VT/RuntimeVirtualTextureSceneProxy.cpp b/Engine/Source/Runtime/Renderer/Private/VT/RuntimeVirtualTextureSceneProxy.cpp
--- a/Engine/Source/Runtime/Renderer/Private/VT/RuntimeVirtualTextureSceneProxy.cpp
+++ b/Engine/Source/Runtime/Renderer/Private/VT/RuntimeVirtualTextureSceneProxy.cpp
@@ -16,14 +16,7 @@
 
 #define LOCTEXT_NAMESPACE "VirtualTexture"
 
-int32 FRuntimeVirtualTextureSceneProxy::ProducerIdGenerator = 1;
-
 FRuntimeVirtualTextureSceneProxy::FRuntimeVirtualTextureSceneProxy(URuntimeVirtualTextureComponent* InComponent)
-	: SceneIndex(0)
-	, ProducerId(0)
-	, VirtualTexture(nullptr)
-	, bHidePrimitivesInEditor(false)
-	, bHidePrimitivesInGame(false)
 {
 	// Evaluate the flags used to hide primitives writing to this virtual texture.
 	InComponent->GetHidePrimitiveSettings(bHidePrimitivesInEditor, bHidePrimitivesInGame);
@@ -32,14 +25,12 @@ FRuntimeVirtualTextureSceneProxy::FRuntimeVirtualTextureSceneProxy(URuntimeVirtu
 	{
 		if (InComponent->IsEnabledInScene())
 		{
-			// We store a ProducerId here so that we will be able to find our SceneIndex from the Producer during rendering.
-			// We will need the SceneIndex to determine which primitives should render to this Producer.
-			ProducerId = ProducerIdGenerator++;
-
 			URuntimeVirtualTexture::FInitSettings InitSettings;
 			InitSettings.TileCountBias = InComponent->IsScalable() ? VirtualTextureScalability::GetRuntimeVirtualTextureSizeBias(InComponent->GetScalabilityGroup()) : 0;
 
 			VirtualTexture = InComponent->GetVirtualTexture();
+			RuntimeVirtualTextureId = VirtualTexture->GetUniqueID();
+
 			Transform = InComponent->GetComponentTransform();
 			const FBox Bounds = InComponent->Bounds.GetBox();
 			const FVector4f CustomMaterialData = InComponent->GetCustomMaterialData();
@@ -85,7 +76,7 @@ FRuntimeVirtualTextureSceneProxy::FRuntimeVirtualTextureSceneProxy(URuntimeVirtu
 			if (StreamingTexture == nullptr)
 			{
 				// Create the runtime virtual texture producer.
-				Producer = new FRuntimeVirtualTextureProducer(ProducerDesc, ProducerId, MaterialType, bClearTextures, SceneInterface, Transform, Bounds, CustomMaterialData);
+				Producer = new FRuntimeVirtualTextureProducer(ProducerDesc, RuntimeVirtualTextureId, MaterialType, bClearTextures, SceneInterface, Transform, Bounds, CustomMaterialData);
 
 				// We only need to dirty flush up to the producer description MaxLevel which accounts for the RemoveLowMips
 				MaxDirtyLevel = ProducerDesc.MaxLevel;
@@ -112,7 +103,7 @@ FRuntimeVirtualTextureSceneProxy::FRuntimeVirtualTextureSceneProxy(URuntimeVirtu
 				}
 
 				// Create the runtime virtual texture producer.
-				Producer = new FRuntimeVirtualTextureProducer(ProducerDesc, ProducerId, MaterialType, bClearTextures, SceneInterface, Transform, Bounds, CustomMaterialData);
+				Producer = new FRuntimeVirtualTextureProducer(ProducerDesc, RuntimeVirtualTextureId, MaterialType, bClearTextures, SceneInterface, Transform, Bounds, CustomMaterialData);
 
 				// Bind the runtime virtual texture producer to the streaming producer.
 				const int32 NumLevels = (int32)FMath::CeilLogTwo(FMath::Max(ProducerDesc.BlockWidthInTiles, ProducerDesc.BlockHeightInTiles));
diff --git a/Engine/Source/Runtime/Renderer/Private/VT/RuntimeVirtualTextureSceneProxy.h b/Engine/Source/Runtime/Renderer/Private/VT/RuntimeVirtualTextureSceneProxy.h
--- a/Engine/Source/Runtime/Renderer/Private/VT/RuntimeVirtualTextureSceneProxy.h
+++ b/Engine/Source/Runtime/Renderer/Private/VT/RuntimeVirtualTextureSceneProxy.h
@@ -35,31 +35,32 @@ public:
 	void FlushDirtyPages();
 
 	/** Index in FScene::RuntimeVirtualTextures. */
-	int32 SceneIndex;
-	/** Unique ID for producer that this Proxy created. Used for finding this object and it's SceneIndex from the producer.  */
-	int32 ProducerId;
-
-	/** Pointer to linked URuntimeVirtualTexture. Not for dereferencing, just for pointer comparison. */
-	URuntimeVirtualTexture* VirtualTexture;
+	int32 SceneIndex = -1;
+	
+	/** Unique object ID of the runtime virtual texture used to filter proxies to render to it.  */
+	int32 RuntimeVirtualTextureId = -1;
 
 	/** Hide primitives in the main pass in editor mode. */
-	bool bHidePrimitivesInEditor;
+	bool bHidePrimitivesInEditor = false;
 	/** Hide primitives in the main pass in game mode. */
-	bool bHidePrimitivesInGame;
+	bool bHidePrimitivesInGame = false;
 
 private:
+	/** Pointer to linked URuntimeVirtualTexture. Not for dereferencing, just for pointer comparison. */
+	URuntimeVirtualTexture* VirtualTexture = nullptr;
+
 	/** UVToWorld transform for the URuntimeVirtualTexture object. */
-	FTransform Transform;
+	FTransform Transform = FTransform::Identity;
 	/** Virtual texture size of the URuntimeVirtualTexture object. */
-	FIntPoint VirtualTextureSize;
+	FIntPoint VirtualTextureSize = FIntPoint::ZeroValue;
 
 	/** Handle for the producer that this Proxy initialized. Used only for invalidation logic. */
 	FVirtualTextureProducerHandle ProducerHandle;
 	/** Space ID used by the virtual texture. Used only for invalidation logic. */
-	int32 SpaceID;
+	int32 SpaceID = -1;
 
 	/** Maximum mip level to mark dirty. Can be less than the virtual texture's MaxLevel if we have streaming mips. */
-	int32 MaxDirtyLevel;
+	int32 MaxDirtyLevel = 0;
 
 	struct FDirtyRect
 	{
@@ -79,7 +80,4 @@ private:
 
 	/** DelegateHandle for on screen warning messages. */
 	FDelegateHandle OnScreenWarningDelegateHandle;
-
-	/** Static ProducerId counter. */
-	static int32 ProducerIdGenerator;
 };
diff --git a/Engine/Source/Runtime/Renderer/Public/PrimitiveSceneInfo.h b/Engine/Source/Runtime/Renderer/Public/PrimitiveSceneInfo.h
--- a/Engine/Source/Runtime/Renderer/Public/PrimitiveSceneInfo.h
+++ b/Engine/Source/Runtime/Renderer/Public/PrimitiveSceneInfo.h
@@ -213,38 +213,26 @@ public:
 	FPrimitiveSceneInfoCompact(FPrimitiveSceneInfo* InPrimitiveSceneInfo);
 };
 
-/** Flags needed for broad phase culling of runtime virtual texture page rendering. */
-struct FPrimitiveVirtualTextureFlags
-{
-	/** True if the primitive can render to virtual texture */
-	uint8 bRenderToVirtualTexture : 1;
-
-	/** Number of bits to reserve for the RuntimeVirtualTextureMask. If we use more than this number of runtime virtual textures in a scene we will trigger a slower path in rendering the VT pages. */
-	enum { RuntimeVirtualTexture_BitCount = 7 };
-	/** Mask of the allocated runtime virtual textures in the scene to render to. */
-	uint8 RuntimeVirtualTextureMask : RuntimeVirtualTexture_BitCount;
-};
-
 /** Lod data used for runtime virtual texture page rendering. Packed to reduce memory overhead since one of these is allocated per primitive. */
-struct FPrimitiveVirtualTextureLodInfo
+struct FPrimitiveRuntimeVirtualTextureLodInfo
 {
 	/** LodBias is in range [-7,8] so is stored with this offset. */
 	enum { LodBiasOffset = 7 };
 
 	/** Minimum Lod for primitive in the runtime virtual texture. */
-	uint16 MinLod : 4;
+	uint16 MinLod : 4 = 0;
 	/** Maximum Lod for primitive in the runtime virtual texture. */
-	uint16 MaxLod : 4;
+	uint16 MaxLod : 4 = 0;
 	/** Bias to use for Lod calculation in the runtime virtual texture. */
-	uint16 LodBias : 4;
+	uint16 LodBias : 4 = 0;
 	/** 
 	 * Culling method used to remove the primitive from low mips of the runtime virtual texture.
 	 * 0: CullValue is the number of low mips for which we cull the primitive from the runtime virtual texture.
 	 * 1: CullValue is the pixel coverage threshold at which we cull the primitive from the runtime virtual texture. 
 	 */
-	uint16 CullMethod : 1;
+	uint16 CullMethod : 1 = 0;
 	/** Value used according to the CullMethod. */
-	uint16 CullValue : 3;
+	uint16 CullValue : 3 = 0;
 };
 
 /** The type of the octree used by FScene to find primitives. */
@@ -536,11 +524,8 @@ public:
 	/** Updates static lighting uniform buffer, returns the number of entries needed for GPUScene */
 	RENDERER_API int32 UpdateStaticLightingBuffer();
 
-	/** Update the cached runtime virtual texture flags for this primitive. Do this when runtime virtual textures are created or destroyed. */
-	void UpdateRuntimeVirtualTextureFlags();
-
-	/** Get the cached runtime virtual texture flags for this primitive. */
-	FPrimitiveVirtualTextureFlags GetRuntimeVirtualTextureFlags() const { return RuntimeVirtualTextureFlags; }
+	/** Get the packed LodInfo. */
+	FPrimitiveRuntimeVirtualTextureLodInfo const& GetRuntimeVirtualTextureLodInfo() const { return RuntimeVirtualTextureLodInfo; }
 
 	/** Mark the runtime virtual textures covered by this primitive as dirty. */
 	void FlushRuntimeVirtualTexture();
@@ -643,8 +628,8 @@ private:
 	 */
 	FPrimitiveSceneInfoData* SceneData;
 
-	/** These flags carry information about which runtime virtual textures are bound to this primitive. */
-	FPrimitiveVirtualTextureFlags RuntimeVirtualTextureFlags;
+	/** Packed LOD info used for selecting LOD when rendering to a runtime virtual texture. */
+	FPrimitiveRuntimeVirtualTextureLodInfo RuntimeVirtualTextureLodInfo;
 
 	/** If this is TRUE, this primitive's uniform buffer needs to be updated before it can be rendered. */
 	bool bNeedsUniformBufferUpdate : 1;
@@ -653,7 +638,7 @@ private:
 	bool bIndirectLightingCacheBufferDirty : 1;
 
 	/** If this is TRUE, this primitive has registered with the virtual texture system for a callback on virtual texture changes. */
-	bool bRegisteredVirtualTextureProducerCallback : 1;
+	bool bRegisteredLightmapVirtualTextureProducerCallback : 1;
 
 	/** True if the primitive registered with velocity data and needs to remove itself when being removed from the scene. */
 	bool bRegisteredWithVelocityData : 1;
@@ -673,8 +658,8 @@ private:
 	/** True if the primitive is queued for deletion. */
 	bool bPendingDelete : 1 = false;
 
-	/** True if the primitive is queued to have its virtual texture flushed. */
-	bool bPendingFlushVirtualTexture : 1;
+	/** True if the primitive is queued to have its runtime virtual texture flushed. */
+	bool bPendingFlushRuntimeVirtualTexture : 1;
 
 public:
 	/** Whether the primitive is newly registered or moved and CachedReflectionCaptureProxy needs to be updated on the next render. */
@@ -686,6 +671,9 @@ public:
 	/** Set to true for the primitive to be rendered into the real-time sky light reflection capture. */
 	bool bVisibleInRealTimeSkyCapture : 1;
 
+	/** True if the primitive can write to runtime virtual texture. */
+	bool bWritesRuntimeVirtualTexture : 1;
+
 #if RHI_RAYTRACING
 	bool bDrawInGame : 1;
 	bool bRayTracingFarField : 1;
@@ -780,9 +768,6 @@ private:
 	/** Creates cached mesh draw commands for all meshes. */
 	static void CacheMeshDrawCommands(FScene* Scene, TArrayView<FPrimitiveSceneInfo*> SceneInfos);
 
-	/** Updates virtual texture state on the scene primitives. */
-	static void UpdateVirtualTextures(FScene* Scene, TArrayView<FPrimitiveSceneInfo*> SceneInfos);
-
 	/** Removes cached mesh draw commands for all meshes. */
 	void RemoveCachedMeshDrawCommands();
 
diff --git a/Engine/Source/Runtime/Renderer/Public/VT/RuntimeVirtualTextureRender.h b/Engine/Source/Runtime/Renderer/Public/VT/RuntimeVirtualTextureRender.h
--- a/Engine/Source/Runtime/Renderer/Public/VT/RuntimeVirtualTextureRender.h
+++ b/Engine/Source/Runtime/Renderer/Public/VT/RuntimeVirtualTextureRender.h
@@ -6,12 +6,8 @@
 #include "VT/RuntimeVirtualTextureEnum.h"
 #include "RHIAccess.h"
 
-class FRHICommandListImmediate;
 class FRHITexture;
-class FRHIUnorderedAccessView;
 class FRDGBuilder;
-class FRDGTexture;
-class FRDGTextureUAV;
 class FScene;
 class FSceneInterface;
 struct IPooledRenderTarget;
@@ -19,15 +15,6 @@ class URuntimeVirtualTextureComponent;
 
 namespace RuntimeVirtualTexture
 {
-#if WITH_EDITOR
-	/**
-	 * Get the scene index of the FRuntimeVirtualTextureSceneProxy associated with a URuntimeVirtualTextureComponent.
-	 * This is needed when rendering runtime virtual texture pages in alternative contexts such as when building previews etc.
-	 * This function is slow because it needs to flush render commands.
-	 */
-	RENDERER_API uint32 GetRuntimeVirtualTextureSceneIndex_GameThread(URuntimeVirtualTextureComponent* InComponent);
-#endif
-
 	/** Enum for our maximum RenderPages() batch size. */
 	enum { EMaxRenderPageBatch = 8 };
 
@@ -62,7 +49,9 @@ namespace RuntimeVirtualTexture
 	{
 		/** Scene to use when rendering the batch. */
 		FScene* Scene;
-		/** Mask created from the target runtime virtual texture's index within the scene. */
+		/** Unique object ID of the runtime virtual texture that we are rendering. */
+		int32 RuntimeVirtualTextureId;
+		UE_DEPRECATED(5.6, "Use RuntimeVirtualTextureId instead.")
 		uint32 RuntimeVirtualTextureMask;
 		/** Virtual texture UV space to world space transform. */
 		FTransform UVToWorld;
@@ -101,4 +90,9 @@ namespace RuntimeVirtualTexture
 	 * Performs the required scene rendering setup such that it can be called from a render thread task.
 	 */
 	RENDERER_API void RenderPagesStandAlone(FRDGBuilder& GraphBuilder, FRenderPageBatchDesc const& InDesc);
+
+#if WITH_EDITOR
+	UE_DEPRECATED(5.6, "Use RuntimeVirtualTextureId everywhere that you would previously use a SceneIndex.")
+	RENDERER_API uint32 GetRuntimeVirtualTextureSceneIndex_GameThread(URuntimeVirtualTextureComponent* InComponent);
+#endif
 }
