Only apply RVT hide primitive flags for active RVTs.
Before the change RVTs that were culled per platform were still instantiated as (inactive) RVT scene objects and had their flags taken into account.
#rnx

diff --git a/Engine/Source/Runtime/Renderer/Private/RendererScene.cpp b/Engine/Source/Runtime/Renderer/Private/RendererScene.cpp
--- a/Engine/Source/Runtime/Renderer/Private/RendererScene.cpp
+++ b/Engine/Source/Runtime/Renderer/Private/RendererScene.cpp
@@ -3585,8 +3585,11 @@ static void UpdateRuntimeVirtualTextureHidePrimitiveMasks(TSparseArray<FRuntimeV
 	bOutGame = false;
 	for (TSparseArray<FRuntimeVirtualTextureSceneProxy*>::TConstIterator It(InProxies); It && !bOutEditor && !bOutGame; ++It)
 	{
-		bOutEditor |= (*It)->bHidePrimitivesInEditor;
-		bOutGame |= (*It)->bHidePrimitivesInGame;
+		if ((*It)->RuntimeVirtualTextureId != -1)
+		{
+			bOutEditor |= (*It)->bHidePrimitivesInEditor;
+			bOutGame |= (*It)->bHidePrimitivesInGame;
+		}
 	}
 }
 
