Make DBuffer BaseColor texture sRGB.
Fixes banding issues in some content.

diff --git a/Engine/Source/Runtime/Renderer/Private/DBufferTextures.cpp b/Engine/Source/Runtime/Renderer/Private/DBufferTextures.cpp
--- a/Engine/Source/Runtime/Renderer/Private/DBufferTextures.cpp
+++ b/Engine/Source/Runtime/Renderer/Private/DBufferTextures.cpp
@@ -49,6 +49,8 @@ FDBufferTexturesDesc GetDBufferTexturesDesc(FIntPoint Extent, EShaderPlatform Sh
 		Desc.ClearValue = ArrayDesc.ClearValue = FClearValueBinding::Black;
 		DBufferTexturesDesc.DBufferADesc = Desc;
 		DBufferTexturesDesc.DBufferATexArrayDesc = ArrayDesc;
+		DBufferTexturesDesc.DBufferADesc.Flags |= TexCreate_SRGB;
+		DBufferTexturesDesc.DBufferATexArrayDesc.Flags |= TexCreate_SRGB;
 
 		Desc.Flags = ArrayDesc.Flags = BaseFlags | GFastVRamConfig.DBufferB;
 		Desc.ClearValue = ArrayDesc.ClearValue = FClearValueBinding(FLinearColor(128.0f / 255.0f, 128.0f / 255.0f, 128.0f / 255.0f, 1));
