Added missing copyright headers

diff --git a/Engine/Source/Runtime/Renderer/Private/MaterialCache/MaterialCache.cpp b/Engine/Source/Runtime/Renderer/Private/MaterialCache/MaterialCache.cpp
--- a/Engine/Source/Runtime/Renderer/Private/MaterialCache/MaterialCache.cpp
+++ b/Engine/Source/Runtime/Renderer/Private/MaterialCache/MaterialCache.cpp
@@ -1,4 +1,6 @@
-﻿#include "MaterialCache/MaterialCache.h"
+﻿// Copyright Epic Games, Inc. All Rights Reserved.
+
+#include "MaterialCache/MaterialCache.h"
 #include "ComponentRecreateRenderStateContext.h"
 #include "VT/VirtualTextureBuildSettings.h"
 #include "HAL/IConsoleManager.h"
diff --git a/Engine/Source/Runtime/Renderer/Private/MaterialCache/MaterialCacheMeshProcessor.cpp b/Engine/Source/Runtime/Renderer/Private/MaterialCache/MaterialCacheMeshProcessor.cpp
--- a/Engine/Source/Runtime/Renderer/Private/MaterialCache/MaterialCacheMeshProcessor.cpp
+++ b/Engine/Source/Runtime/Renderer/Private/MaterialCache/MaterialCacheMeshProcessor.cpp
@@ -1,4 +1,6 @@
-﻿#include "MaterialCache/MaterialCacheMeshProcessor.h"
+﻿// Copyright Epic Games, Inc. All Rights Reserved.
+
+#include "MaterialCache/MaterialCacheMeshProcessor.h"
 #include "MeshPassUtils.h"
 #include "NaniteSceneProxy.h"
 #include "NaniteVertexFactory.h"
diff --git a/Engine/Source/Runtime/Renderer/Private/MaterialCache/MaterialCacheRenderer.cpp b/Engine/Source/Runtime/Renderer/Private/MaterialCache/MaterialCacheRenderer.cpp
--- a/Engine/Source/Runtime/Renderer/Private/MaterialCache/MaterialCacheRenderer.cpp
+++ b/Engine/Source/Runtime/Renderer/Private/MaterialCache/MaterialCacheRenderer.cpp
@@ -1,3 +1,4 @@
+// Copyright Epic Games, Inc. All Rights Reserved.
 
 #include "MaterialCache/MaterialCacheRenderer.h"
 #include "DeferredShadingRenderer.h"
diff --git a/Engine/Source/Runtime/Renderer/Private/MaterialCache/MaterialCacheShaders.cpp b/Engine/Source/Runtime/Renderer/Private/MaterialCache/MaterialCacheShaders.cpp
--- a/Engine/Source/Runtime/Renderer/Private/MaterialCache/MaterialCacheShaders.cpp
+++ b/Engine/Source/Runtime/Renderer/Private/MaterialCache/MaterialCacheShaders.cpp
@@ -1,4 +1,6 @@
-﻿#include "MaterialCache/MaterialCacheShaders.h"
+﻿// Copyright Epic Games, Inc. All Rights Reserved.
+
+#include "MaterialCache/MaterialCacheShaders.h"
 #include "MaterialCache/MaterialCache.h"
 #include "DataDrivenShaderPlatformInfo.h"
 
