Potential fix for TSR mip2 texture size too small (1x1) due to viewport or secondary scaling too small (0x0, or 0).

#tests Limited test in engine, as this is triggered by OG guided missile. However, it does not have any potential drawback

diff --git a/Engine/Source/Runtime/Renderer/Private/PostProcess/TemporalSuperResolution.cpp b/Engine/Source/Runtime/Renderer/Private/PostProcess/TemporalSuperResolution.cpp
--- a/Engine/Source/Runtime/Renderer/Private/PostProcess/TemporalSuperResolution.cpp
+++ b/Engine/Source/Runtime/Renderer/Private/PostProcess/TemporalSuperResolution.cpp
@@ -1925,9 +1925,10 @@ FDefaultTemporalUpscaler::FOutputs AddTemporalSuperResolutionPasses(
 	{
 		{
 			bool bRequires2Mips = HistorySize == OutputRect.Size() && PassInputs.bGenerateOutputMip1;
-
+			FIntPoint MipClampedHistoryExtent = FIntPoint(FMath::Max(HistoryExtent.X, bRequires2Mips ? 2: 1), 
+				                                          FMath::Max(HistoryExtent.Y, bRequires2Mips ? 2: 1));
 			FRDGTextureDesc ArrayDesc = FRDGTextureDesc::Create2DArray(
-				HistoryExtent,
+				MipClampedHistoryExtent,
 				HistoryColorFormat,
 				FClearValueBinding::None,
 				TexCreate_ShaderResource | TexCreate_UAV,
@@ -3044,8 +3045,10 @@ FDefaultTemporalUpscaler::FOutputs AddTemporalSuperResolutionPasses(
 		
 		FRDGTextureRef SceneColorOutputTexture;
 		{
+			FIntPoint MipClampedOutputExtent = FIntPoint(FMath::Max(OutputExtent.X, PassInputs.bGenerateOutputMip1 ? 2 : 1), 
+				                                         FMath::Max(OutputExtent.Y, PassInputs.bGenerateOutputMip1 ? 2 : 1));
 			FRDGTextureDesc OutputDesc = FRDGTextureDesc::Create2D(
-				OutputExtent,
+				MipClampedOutputExtent,
 				ColorFormat,
 				FClearValueBinding::None,
 				/* InFlags = */ TexCreate_ShaderResource | TexCreate_UAV | TexCreate_RenderTargetable,
