Fixed "after motion blur" translucent pass to have valid previous matrices. (for niagara effect to be able to simulate motion blur using clip and prevclip matrices)
Follow up UE-215155    CL 33768987

#rnx
#tests PC editor dx12 - with without separate transluceny and difference screen percentage

diff --git a/Engine/Source/Runtime/Renderer/Private/TranslucentRendering.cpp b/Engine/Source/Runtime/Renderer/Private/TranslucentRendering.cpp
--- a/Engine/Source/Runtime/Renderer/Private/TranslucentRendering.cpp
+++ b/Engine/Source/Runtime/Renderer/Private/TranslucentRendering.cpp
@@ -893,8 +893,11 @@ void SetupPostMotionBlurTranslucencyViewParameters(const FViewInfo& View, FViewU
 	FViewMatrices ModifiedViewMatrices = View.ViewMatrices;
 	ModifiedViewMatrices.HackRemoveTemporalAAProjectionJitter();
 
+	FViewMatrices ModifiedPrevViewViewMatrices = View.PrevViewInfo.ViewMatrices;
+	ModifiedPrevViewViewMatrices.HackRemoveTemporalAAProjectionJitter();
+
 	Parameters = *View.CachedViewUniformShaderParameters;
-	View.SetupUniformBufferParameters(ModifiedViewMatrices, ModifiedViewMatrices, VolumeBounds, TVC_MAX, Parameters);
+	View.SetupUniformBufferParameters(ModifiedViewMatrices, ModifiedPrevViewViewMatrices, VolumeBounds, TVC_MAX, Parameters);
 }
 
 const FRDGTextureDesc GetPostDOFTranslucentTextureDesc(
@@ -950,7 +953,7 @@ void SetupDownsampledTranslucencyViewParameters(
 	{
 		// Remove jitter from this pass
 		ViewMatrices.HackRemoveTemporalAAProjectionJitter();
-		PrevViewMatrices = ViewMatrices;
+		PrevViewMatrices.HackRemoveTemporalAAProjectionJitter();
 
 		FBox VolumeBounds[TVC_MAX];
 		View.SetupUniformBufferParameters(ViewMatrices, PrevViewMatrices, VolumeBounds, TVC_MAX, DownsampledTranslucencyViewParameters);
