Locallised some warning strings.

#tests PC editor
#rnx

diff --git a/Engine/Source/Runtime/Renderer/Private/ReflectionEnvironmentRealTimeCapture.cpp b/Engine/Source/Runtime/Renderer/Private/ReflectionEnvironmentRealTimeCapture.cpp
--- a/Engine/Source/Runtime/Renderer/Private/ReflectionEnvironmentRealTimeCapture.cpp
+++ b/Engine/Source/Runtime/Renderer/Private/ReflectionEnvironmentRealTimeCapture.cpp
@@ -298,11 +298,13 @@ void FScene::ValidateSkyLightRealTimeCapture(
 
 			if (View.bSceneHasSkyMaterial && !bSkyMeshInMainPassExist)
 			{
-				Canvas.DrawShadowedString(100.0f, 100.0f, TEXT("At least one mesh with a sky material is in the scene but none are rendered in main view."), GetStatsFont(), TextColor);
+				const FString WarningString = NSLOCTEXT("RealTimeReflectionCapture", "RealTimeReflectionCaptureMainView",			"At least one mesh with a sky material is in the scene but none are rendered in main view.").ToString();
+				Canvas.DrawShadowedString(100.0f, 100.0f, WarningString, GetStatsFont(), TextColor);
 			}
 			if (View.bSceneHasSkyMaterial && !bSkyMeshInRealTimeSkyCaptureExtist && SkyLight && SkyLight->bRealTimeCaptureEnabled)
 			{
-				Canvas.DrawShadowedString(100.0f, 110.0f, TEXT("At least one mesh with a sky material is in the scene but none are rendered in the real-time sky light reflection."), GetStatsFont(), TextColor);
+				const FString WarningString = NSLOCTEXT("RealTimeReflectionCapture", "RealTimeReflectionCaptureREReflectionView",	"At least one mesh with a sky material is in the scene but none are rendered in the real-time sky light reflection.").ToString();
+				Canvas.DrawShadowedString(100.0f, 100.0f, WarningString, GetStatsFont(), TextColor);
 			}
 		});
 	}
