[DOF] Added SCATTER_QUALITY shader permutation to vertex shader.

diff --git a/Engine/Source/Runtime/Renderer/Private/PostProcess/DiaphragmDOF.cpp b/Engine/Source/Runtime/Renderer/Private/PostProcess/DiaphragmDOF.cpp
--- a/Engine/Source/Runtime/Renderer/Private/PostProcess/DiaphragmDOF.cpp
+++ b/Engine/Source/Runtime/Renderer/Private/PostProcess/DiaphragmDOF.cpp
@@ -464,11 +464,17 @@ public:
 
 	static void ModifyCompilationEnvironment(const FGlobalShaderPermutationParameters& Parameters, FShaderCompilerEnvironment& OutEnvironment)
 	{
+		if (FDataDrivenShaderPlatformInfo::GetSupportsRealTypes(Parameters.Platform) == ERHIFeatureSupport::RuntimeGuaranteed)
+		{
+			OutEnvironment.CompilerFlags.Add(CFLAG_AllowRealTypes);
+		}
+
 		OutEnvironment.SetDefine(TEXT("COC_TILE_SIZE"), kCocTileSize);
 
 		// If on console, don't compile any dynamic CoC offset to avoid performance differences
 		const bool bIsConsole = FDataDrivenShaderPlatformInfo::GetIsConsole(Parameters.Platform);
 		OutEnvironment.SetDefine(TEXT("WITH_DYNAMIC_COC_OFFSET"), bIsConsole ? 0 : 1);
+
 	}
 
 	FDiaphragmDOFShader() {}
@@ -1159,11 +1165,6 @@ class FDiaphragmDOFGatherCS : public FDiaphragmDOFShader
 
 		// This shader takes a very long time to compile with FXC, so we pre-compile it with DXC first and then forward the optimized HLSL to FXC.
 		OutEnvironment.CompilerFlags.Add(CFLAG_PrecompileWithDXC);
-
-		if (FDataDrivenShaderPlatformInfo::GetSupportsRealTypes(Parameters.Platform) == ERHIFeatureSupport::RuntimeGuaranteed)
-		{
-			OutEnvironment.CompilerFlags.Add(CFLAG_AllowRealTypes);
-		}
 	}
 
 	
@@ -1275,12 +1276,14 @@ BEGIN_SHADER_PARAMETER_STRUCT(FDOFHybridScatterParameters, )
 	RENDER_TARGET_BINDING_SLOTS()
 END_SHADER_PARAMETER_STRUCT()
 
+class FDDOFScatterQualityDim	: SHADER_PERMUTATION_ENUM_CLASS("DIM_SCATTER_QUALITY", EDiaphragmDOFScatterQuality);
+
 class FDiaphragmDOFHybridScatterVS : public FDiaphragmDOFShader
 {
 	DECLARE_GLOBAL_SHADER(FDiaphragmDOFHybridScatterVS);
 	SHADER_USE_PARAMETER_STRUCT(FDiaphragmDOFHybridScatterVS, FDiaphragmDOFShader);
 
-	using FPermutationDomain = TShaderPermutationDomain<>;
+	using FPermutationDomain = TShaderPermutationDomain<FDDOFScatterQualityDim>;
 	using FParameters = FDOFHybridScatterParameters;
 
 	static bool ShouldCompilePermutation(const FGlobalShaderPermutationParameters& Parameters)
@@ -1291,6 +1294,14 @@ class FDiaphragmDOFHybridScatterVS : public FDiaphragmDOFShader
 		}
 		return FDiaphragmDOFShader::ShouldCompilePermutation(Parameters);
 	}
+
+	static void ModifyCompilationEnvironment(const FGlobalShaderPermutationParameters& Parameters, FShaderCompilerEnvironment& OutEnvironment)
+	{
+		FDiaphragmDOFShader::ModifyCompilationEnvironment(Parameters, OutEnvironment);
+
+		OutEnvironment.SetDefine(TEXT("SCATTER_QUALITY_MEDIUM"), (int)EDiaphragmDOFScatterQuality::Medium);
+		OutEnvironment.SetDefine(TEXT("SCATTER_QUALITY_HIGH"), (int)EDiaphragmDOFScatterQuality::High);
+	}
 };
 
 class FDiaphragmDOFHybridScatterPS : public FDiaphragmDOFShader
@@ -1299,7 +1310,6 @@ class FDiaphragmDOFHybridScatterPS : public FDiaphragmDOFShader
 	SHADER_USE_PARAMETER_STRUCT(FDiaphragmDOFHybridScatterPS, FDiaphragmDOFShader);
 
 	class FBokehSimulationDim		: SHADER_PERMUTATION_BOOL("DIM_BOKEH_SIMULATION");
-	class FDDOFScatterQualityDim	: SHADER_PERMUTATION_ENUM_CLASS("DIM_SCATTER_QUALITY", EDiaphragmDOFScatterQuality);
 
 	using FPermutationDomain = TShaderPermutationDomain<FDDOFLayerProcessingDim, FBokehSimulationDim, FDDOFScatterOcclusionDim, FDDOFScatterQualityDim>;
 	using FParameters = FDOFHybridScatterParameters;
@@ -1616,6 +1626,7 @@ bool DiaphragmDOF::AddPasses(
 		CVarEnableGatherBokehSettings.GetValueOnRenderThread() == 1);
 	const bool bEnableScatterBokehSettings = CVarEnableScatterBokehSettings.GetValueOnRenderThread() == 1;
 	const bool bEnableSlightOutOfFocusBokeh = bSupportGatheringBokehSimulation && bRecombineDoesSlightOutOfFocus && CVarEnableRecombineBokehSettings.GetValueOnRenderThread();
+
 		
 	// Setup all the descriptors.
 	FRDGTextureDesc FullResDesc;
@@ -2306,6 +2317,7 @@ bool DiaphragmDOF::AddPasses(
 			FDOFConvolutionTextures* ConvolutionOutputTextures,
 			FRDGTextureRef* ScatterOcclusionTexture)
 		{
+
 			bool bUsesScatterOcclusionTexture =
 				ConvolutionSettings.QualityConfig == EDiaphragmDOFGatherQuality::HighQualityWithHybridScatterOcclusion ||
 				ConvolutionSettings.QualityConfig == EDiaphragmDOFGatherQuality::Cinematic;
@@ -2580,15 +2592,18 @@ bool DiaphragmDOF::AddPasses(
 				Quality = EDiaphragmDOFScatterQuality::Medium;
 			}
 
-			FDiaphragmDOFHybridScatterPS::FPermutationDomain PermutationVector;
-			PermutationVector.Set<FDDOFLayerProcessingDim>(ConvolutionSettings.LayerProcessing);
-			PermutationVector.Set<FDiaphragmDOFHybridScatterPS::FBokehSimulationDim>(ScatteringBokehLUT ? true : false);
-			PermutationVector.Set<FDDOFScatterOcclusionDim>(ScatterOcclusionTexture ? true : false);
-			PermutationVector.Set<FDiaphragmDOFHybridScatterPS::FDDOFScatterQualityDim>(Quality);
-			PermutationVector = FDiaphragmDOFHybridScatterPS::RemapPermutation(PermutationVector);
+			FDiaphragmDOFHybridScatterVS::FPermutationDomain VSPermutationVector;
+			VSPermutationVector.Set<FDDOFScatterQualityDim>(Quality);
 
-			TShaderMapRef<FDiaphragmDOFHybridScatterVS> VertexShader(View.ShaderMap);
-			TShaderMapRef<FDiaphragmDOFHybridScatterPS> PixelShader(View.ShaderMap, PermutationVector);
+			FDiaphragmDOFHybridScatterPS::FPermutationDomain PSPermutationVector;
+			PSPermutationVector.Set<FDDOFLayerProcessingDim>(ConvolutionSettings.LayerProcessing);
+			PSPermutationVector.Set<FDiaphragmDOFHybridScatterPS::FBokehSimulationDim>(ScatteringBokehLUT ? true : false);
+			PSPermutationVector.Set<FDDOFScatterOcclusionDim>(ScatterOcclusionTexture ? true : false);
+			PSPermutationVector.Set<FDDOFScatterQualityDim>(Quality);
+			PSPermutationVector = FDiaphragmDOFHybridScatterPS::RemapPermutation(PSPermutationVector);
+
+			TShaderMapRef<FDiaphragmDOFHybridScatterVS> VertexShader(View.ShaderMap, VSPermutationVector);
+			TShaderMapRef<FDiaphragmDOFHybridScatterPS> PixelShader(View.ShaderMap, PSPermutationVector);
 
 			FDOFHybridScatterParameters* PassParameters = GraphBuilder.AllocParameters<FDOFHybridScatterParameters>();
 			PassParameters->ViewportSize = FVector4f(GatheringViewSize.X, GatheringViewSize.Y, 1.0f / GatheringViewSize.X, 1.0f / GatheringViewSize.Y);
@@ -2609,13 +2624,15 @@ bool DiaphragmDOF::AddPasses(
 			PassParameters->RenderTargets[0] = FRenderTargetBinding(
 				ConvolutionTextures->SceneColor, ERenderTargetLoadAction::ELoad);
 
+				
+
 			ClearUnusedGraphResources(VertexShader, PixelShader, PassParameters);
 
 			GraphBuilder.AddPass(
 				RDG_EVENT_NAME("DOF IndirectScatter(%s Bokeh=%s Occlusion=%s) %dx%d",
 					GetEventName(bIsForeground ? EDiaphragmDOFLayerProcessing::ForegroundOnly : EDiaphragmDOFLayerProcessing::BackgroundOnly),
-					PermutationVector.Get<FDiaphragmDOFHybridScatterPS::FBokehSimulationDim>() ? TEXT("Generic") : TEXT("None"),
-					PermutationVector.Get<FDDOFScatterOcclusionDim>() ? TEXT("Yes") : TEXT("No"),
+					PSPermutationVector.Get<FDiaphragmDOFHybridScatterPS::FBokehSimulationDim>() ? TEXT("Generic") : TEXT("None"),
+					PSPermutationVector.Get<FDDOFScatterOcclusionDim>() ? TEXT("Yes") : TEXT("No"),
 					GatheringViewSize.X, GatheringViewSize.Y),
 				PassParameters,
 				ERDGPassFlags::Raster,
