Fix HairOnlyDepthTexture nullptr access in RenderLightFunctionForMaterial with bUseHairStrands at true (by skipping views that do not have valid HairStrandsViewData)

#rnx

diff --git a/Engine/Source/Runtime/Renderer/Private/LightFunctionRendering.cpp b/Engine/Source/Runtime/Renderer/Private/LightFunctionRendering.cpp
--- a/Engine/Source/Runtime/Renderer/Private/LightFunctionRendering.cpp
+++ b/Engine/Source/Runtime/Renderer/Private/LightFunctionRendering.cpp
@@ -363,6 +363,13 @@ bool FDeferredShadingSceneRenderer::RenderLightFunctionForMaterial(
 
 		if (View.VisibleLightInfos[LightSceneInfo->Id].bInViewFrustum)
 		{
+			if (bUseHairStrands && !View.HairStrandsViewData.bIsValid)
+			{
+				// bUseHairStrands is true if one of the view has valid hair strands data.
+				// Skip views that do not have valid hair strands data to prevent HairOnlyDepthTexture nullptr access
+				continue;
+			}
+
 			FRenderLightFunctionForMaterialParameters* PassParameters = GraphBuilder.AllocParameters<FRenderLightFunctionForMaterialParameters>();
 			PassParameters->SceneTextures = SceneTextures.UniformBuffer;
 			PassParameters->HairOnlyDepthTexture = (View.HairStrandsViewData.bIsValid && View.HairStrandsViewData.VisibilityData.HairOnlyDepthTexture) ? View.HairStrandsViewData.VisibilityData.HairOnlyDepthTexture : GSystemTextures.GetDepthDummy(GraphBuilder);
