Fix for mesh edges rendering not releasing resources.

diff --git a/Engine/Source/Runtime/Renderer/Private/MeshEdgesRendering.cpp b/Engine/Source/Runtime/Renderer/Private/MeshEdgesRendering.cpp
--- a/Engine/Source/Runtime/Renderer/Private/MeshEdgesRendering.cpp
+++ b/Engine/Source/Runtime/Renderer/Private/MeshEdgesRendering.cpp
@@ -125,6 +125,19 @@ private:
 class FMeshEdgesViewFamilyData : public ISceneViewFamilyExtentionData
 {
 public:
+	~FMeshEdgesViewFamilyData()
+	{
+		if (WireframeColor)
+		{
+			WireframeColor->ReleaseResource();
+		}
+
+		if (WireframeDepth)
+		{
+			WireframeDepth->ReleaseResource();
+		}
+	}
+
 	inline static const TCHAR* const GSubclassIdentifier = TEXT("FMeshEdgesViewFamilyData");
 	virtual const TCHAR* GetSubclassIdentifier() const override
 	{
@@ -439,14 +452,5 @@ void ComposeMeshEdges(FRDGBuilder& GraphBuilder,
 	FRHIDepthStencilState* DepthStencilState = TStaticDepthStencilState<true, CF_DepthNearOrEqual>::GetRHI();
 
 	FPixelShaderUtils::AddFullscreenPass(GraphBuilder, View.ShaderMap, RDG_EVENT_NAME("ComposeMeshEdges"), PixelShader, PassParameters, Viewport, BlendState, nullptr, DepthStencilState);
-	
-	if (View.IsLastInFamily())
-	{
-		GraphBuilder.AddPostExecuteCallback([&WireframeTextureColor, &WireframeTextureDepth]
-			{
-				WireframeTextureColor.ReleaseResource();
-				WireframeTextureDepth.ReleaseResource();
-			});
-	}
 }
 
