Only allow gpu skin cache async compute for scene updates associated with a parallel renderer.

diff --git a/Engine/Source/Runtime/Renderer/Private/RendererScene.cpp b/Engine/Source/Runtime/Renderer/Private/RendererScene.cpp
--- a/Engine/Source/Runtime/Renderer/Private/RendererScene.cpp
+++ b/Engine/Source/Runtime/Renderer/Private/RendererScene.cpp
@@ -5569,7 +5569,8 @@ void FScene::Update(FRDGBuilder& GraphBuilder, const FUpdateParameters& Paramete
 
 	FSkeletalMeshUpdater::FSubmitTasks GPUSkinTasks;
 
-	const ERHIPipeline GPUSkinCachePipeline = GPUSkinCache ? FGPUSkinCache::GetDispatchPipeline(GraphBuilder) : ERHIPipeline::Graphics;
+	// Only allow async compute when we are updating as part of a render and are allowing async ops to occur.
+	const ERHIPipeline GPUSkinCachePipeline = GPUSkinCache && Parameters.AsyncOps != EUpdateAllPrimitiveSceneInfosAsyncOps::None ? FGPUSkinCache::GetDispatchPipeline(GraphBuilder) : ERHIPipeline::Graphics;
 
 	if (SkeletalMeshUpdater)
 	{
