Change default virtual texture behavior in cooked builds to not wait on render thread for root pages to be mapped.
Instead we read from 1x1 fallback color generated during texture compilation.
This removes hitches at the with the downside of potenially seeing the low res fallback color sometimes.
Note that runtime virutal textures use the fallback color from any baked streaming mips.
This setting has been used several internal projects, so is well tested to become the default.

diff --git a/Engine/Source/Runtime/Renderer/Private/VT/VirtualTextureSystem.cpp b/Engine/Source/Runtime/Renderer/Private/VT/VirtualTextureSystem.cpp
--- a/Engine/Source/Runtime/Renderer/Private/VT/VirtualTextureSystem.cpp
+++ b/Engine/Source/Runtime/Renderer/Private/VT/VirtualTextureSystem.cpp
@@ -131,7 +131,7 @@ static TAutoConsoleVariable<int32> CVarVTForceContinuousUpdate(
 );
 static TAutoConsoleVariable<int32> CVarVTSyncProduceLockedTiles(
 	TEXT("r.VT.SyncProduceLockedTiles"),
-	1,
+	0,
 	TEXT("Should we sync loading when producing locked tiles"),
 	ECVF_RenderThreadSafe
 );
