Fix iOS/Metal assertion and termination when Scene Captures are in use but the Distortion pass is disabled.

diff --git a/Engine/Source/Runtime/Renderer/Private/MobileShadingRenderer.cpp b/Engine/Source/Runtime/Renderer/Private/MobileShadingRenderer.cpp
--- a/Engine/Source/Runtime/Renderer/Private/MobileShadingRenderer.cpp
+++ b/Engine/Source/Runtime/Renderer/Private/MobileShadingRenderer.cpp
@@ -261,7 +261,8 @@ static bool PostProcessUsesSceneDepth(const FViewInfo& View)
 			}
 		}
 	}
-	return false;
+	
+	return IsMobileDistortionActive(View) || View.bIsSceneCapture;
 }
 
 struct FRenderViewContext
@@ -590,7 +591,7 @@ void FMobileSceneRenderer::InitViews(
 	// On PowerVR we see flickering of shadows and depths not updating correctly if targets are discarded.
 	const bool bForceDepthResolve = (CVarMobileForceDepthResolve.GetValueOnRenderThread() == 1);
 	const bool bSeparateTranslucencyActive = IsMobileSeparateTranslucencyActive(Views.GetData(), Views.Num()); 
-	const bool bPostProcessUsesSceneDepth = PostProcessUsesSceneDepth(Views[0]) || IsMobileDistortionActive(Views[0]);
+	const bool bPostProcessUsesSceneDepth = PostProcessUsesSceneDepth(Views[0]);
 	const bool bRequireSeparateViewPass = Views.Num() > 1 && !Views[0].bIsMobileMultiViewEnabled;
 	bRequiresMultiPass = RequiresMultiPass(NumMSAASamples, ShaderPlatform);
 
