Made cvar scalable for real time capture control per platform.

#tests PC editor dx12
#rnx

diff --git a/Engine/Source/Runtime/Renderer/Private/ReflectionEnvironmentRealTimeCapture.cpp b/Engine/Source/Runtime/Renderer/Private/ReflectionEnvironmentRealTimeCapture.cpp
--- a/Engine/Source/Runtime/Renderer/Private/ReflectionEnvironmentRealTimeCapture.cpp
+++ b/Engine/Source/Runtime/Renderer/Private/ReflectionEnvironmentRealTimeCapture.cpp
@@ -37,12 +37,12 @@ DECLARE_GPU_STAT(CaptureConvolveSkyEnvMap);
 static TAutoConsoleVariable<int32> CVarRealTimeReflectionCaptureTimeSlicing(
 	TEXT("r.SkyLight.RealTimeReflectionCapture.TimeSlice"), 1,
 	TEXT("When enabled, the real-time sky light capture and convolutions will by distributed over several frames to lower the per-frame cost."),
-	ECVF_RenderThreadSafe);
+	ECVF_RenderThreadSafe | ECVF_Scalability);
 
 static TAutoConsoleVariable<int32> CVarRealTimeReflectionCaptureTimeSlicingSkyCloudCubeFacePerFrame(
 	TEXT("r.SkyLight.RealTimeReflectionCapture.TimeSlice.SkyCloudCubeFacePerFrame"), 6,
 	TEXT("When enabled, the real-time sky light capture, when time sliced, will not render cloud in all cube face in a single frame; but one face per frame. That is to distribute the cloud tracing cost even more, but will add latency and potentially can result in lighting discrepancy between faces if the sun is moving fast. Value in [1,6]."),
-	ECVF_RenderThreadSafe);
+	ECVF_RenderThreadSafe | ECVF_Scalability);
 
 static TAutoConsoleVariable<int32> CVarRealTimeReflectionCaptureShadowFromOpaque(
 	TEXT("r.SkyLight.RealTimeReflectionCapture.ShadowFromOpaque"), 0,
@@ -62,7 +62,8 @@ static TAutoConsoleVariable<int32> CVarRealTimeReflectionCaptureAlwaysClearColor
 static TAutoConsoleVariable<int32> CVarRealTimeReflectionCaptureResolutionOverride(
 	TEXT("r.SkyLight.RealTimeReflectionCapture.ResolutionOverride"), 0,
 	TEXT("Override the real time sky light capture resolution if this CVar is greater than 16."),
-	ECVF_ReadOnly | ECVF_RenderThreadSafe); // Read only because it is not really toggable mid-update when time slicing is used
+	ECVF_RenderThreadSafe | ECVF_Scalability);	// Should be ECVF_ReadOnly because it is not really toggable mid-update when time slicing is used (a black convolved sky target will be used instead). 
+												// TODO update on FRealTimeSlicedReflectionCapture only of first frame of capture.
 
 // If this message on screen message is not enough to prevent TDRs, we should implement half resolution + upsampling tracing for cloud in sky light capture.
 static TAutoConsoleVariable<int32> CVarRealTimeReflectionCaptureDisableExpenssiveCaptureMessage(
