Fix RHI validation transition errors for the adaptive runtime virtual texture indirection texture.

diff --git a/Engine/Source/Runtime/Renderer/Private/VT/AdaptiveVirtualTexture.cpp b/Engine/Source/Runtime/Renderer/Private/VT/AdaptiveVirtualTexture.cpp
--- a/Engine/Source/Runtime/Renderer/Private/VT/AdaptiveVirtualTexture.cpp
+++ b/Engine/Source/Runtime/Renderer/Private/VT/AdaptiveVirtualTexture.cpp
@@ -751,13 +751,13 @@ void FAdaptiveVirtualTexture::UpdateAllocations(FVirtualTextureSystem* InSystem,
 		FRHITexture* Texture = Space->GetPageTableIndirectionTexture()->GetReferencedTexture();
 
 		//todo[vt]: If we have more than 1 or 2 updates per frame then add a shader to batch updates.
-		RHICmdList.Transition(FRHITransitionInfo(Texture, ERHIAccess::SRVMask, ERHIAccess::UAVCompute));
+		RHICmdList.Transition({ FRHITransitionInfo(Texture, ERHIAccess::SRVMask, ERHIAccess::UAVCompute) }, ERHIPipeline::All, ERHIPipeline::Graphics);
 		for (FIndirectionTextureUpdate& TextureUpdate : TextureUpdates)
 		{
 			const FUpdateTextureRegion2D Region(TextureUpdate.X, TextureUpdate.Y, 0, 0, 1, 1);
 			RHIUpdateTexture2D((FRHITexture*)Texture, 0, Region, 4, (uint8*)&TextureUpdate.Value);
 		}
-		RHICmdList.Transition(FRHITransitionInfo(Texture, ERHIAccess::UAVCompute, ERHIAccess::SRVMask));
+		RHICmdList.Transition({ FRHITransitionInfo(Texture, ERHIAccess::UAVCompute, ERHIAccess::SRVMask) }, ERHIPipeline::Graphics, ERHIPipeline::All);
 	}
 
 	// Clear requests
diff --git a/Engine/Source/Runtime/Renderer/Private/VT/VirtualTextureSpace.cpp b/Engine/Source/Runtime/Renderer/Private/VT/VirtualTextureSpace.cpp
--- a/Engine/Source/Runtime/Renderer/Private/VT/VirtualTextureSpace.cpp
+++ b/Engine/Source/Runtime/Renderer/Private/VT/VirtualTextureSpace.cpp
@@ -360,6 +360,7 @@ void FVirtualTextureSpace::AllocateTextures(FRDGBuilder& GraphBuilder)
 			FRDGTextureRef PageTableIndirectionTexture = GraphBuilder.CreateTexture(Desc, TEXT("VirtualTexture_PageTableAdaptiveIndirection"));
 			AddClearUAVPass(GraphBuilder, GraphBuilder.CreateUAV(PageTableIndirectionTexture), FUintVector4(ForceInitToZero));
 			PageTableIndirection.RenderTarget = GraphBuilder.ConvertToExternalTexture(PageTableIndirectionTexture);
+			GraphBuilder.UseExternalAccessMode(PageTableIndirectionTexture, ERHIAccess::SRVMask, ERHIPipeline::All);
 			RHIUpdateTextureReference(PageTableIndirection.TextureReferenceRHI, PageTableIndirection.RenderTarget->GetRHI());
 		}
 
