[SingleLayerWater] SLW can now always output velocity, independent of whether the Velocity Ouput Pass is set to "Write during base pass" or not. There's a new r.Water.SingleLayer.ForceVelocity cvar, which defaults to 1 and can be used as to revert to the old behavior.

diff --git a/Engine/Source/Runtime/Renderer/Private/BasePassRendering.cpp b/Engine/Source/Runtime/Renderer/Private/BasePassRendering.cpp
--- a/Engine/Source/Runtime/Renderer/Private/BasePassRendering.cpp
+++ b/Engine/Source/Runtime/Renderer/Private/BasePassRendering.cpp
@@ -947,13 +947,14 @@ bool IsGBufferLayoutSupportedForMaterial(EGBufferLayout Layout, const FMeshMater
 		return true;
 
 	case GBL_ForceVelocity:
-		// Only Nanite materials with WPO support this GBuffer layout
+		// Only Nanite materials with WPO and SingleLayerWater materials support this GBuffer layout
 		// NOTE: FMaterialShaderParameters::bHasVertexPositionOffsetConnected means that the material *could* have WPO.
 		// It's still possible for the material to disable WPO after translation.
-		return !IsUsingBasePassVelocity(Parameters.Platform) && 
+		return (!IsUsingBasePassVelocity(Parameters.Platform) && 
 			Parameters.VertexFactoryType->SupportsNaniteRendering() &&
 			Parameters.MaterialParameters.bIsUsedWithNanite &&
-			Parameters.MaterialParameters.bHasVertexPositionOffsetConnected;
+			Parameters.MaterialParameters.bHasVertexPositionOffsetConnected)
+			|| Parameters.MaterialParameters.ShadingModels.HasShadingModel(MSM_SingleLayerWater);
 
 	default:
 		checkf(false, TEXT("Unhandled GBuffer Layout!"));
diff --git a/Engine/Source/Runtime/Renderer/Private/BasePassRendering.h b/Engine/Source/Runtime/Renderer/Private/BasePassRendering.h
--- a/Engine/Source/Runtime/Renderer/Private/BasePassRendering.h
+++ b/Engine/Source/Runtime/Renderer/Private/BasePassRendering.h
@@ -314,6 +314,16 @@ public:
 
 		// @todo MetalMRT: Remove this hack and implement proper atmospheric-fog solution for Metal MRT...
 		OutEnvironment.SetDefine(TEXT("BASEPASS_SKYATMOSPHERE_AERIALPERSPECTIVE"), !IsMetalMRTPlatform(Parameters.Platform) ? 1 : 0);
+
+		// GBUFFER_LAYOUT is used to derive GBUFFER_HAS_VELOCITY, which influences whether VelocityPrevScreenPosition is computed by the vertex shader.
+		// SLW can force velocity to be drawn in the water base pass even if the rest of the scene outputs in the depth prepass. SLW achieves this override
+		// by using GBL_ForceVelocity, which requires setting the corresponding define on the vertex shader as well, as the calculation for VelocityPrevScreenPosition will be stripped otherwise.
+		static const auto SLWForceVelocityCVar = IConsoleManager::Get().FindTConsoleVariableDataInt(TEXT("r.Water.SingleLayer.ForceVelocity"));
+		const bool bSLWForceVelocity = SLWForceVelocityCVar && SLWForceVelocityCVar->GetValueOnAnyThread() != 0;
+		if (bSLWForceVelocity && Parameters.MaterialParameters.ShadingModels.HasShadingModel(MSM_SingleLayerWater))
+		{
+			OutEnvironment.SetDefine(TEXT("GBUFFER_LAYOUT"), GBL_ForceVelocity);
+		}
 	}
 };
 
diff --git a/Engine/Source/Runtime/Renderer/Private/MeshPassProcessor.cpp b/Engine/Source/Runtime/Renderer/Private/MeshPassProcessor.cpp
--- a/Engine/Source/Runtime/Renderer/Private/MeshPassProcessor.cpp
+++ b/Engine/Source/Runtime/Renderer/Private/MeshPassProcessor.cpp
@@ -2322,9 +2322,9 @@ void SetupDepthStencilInfo(
 	RenderTargetsInfo.StencilTargetStoreAction = RenderTargetsInfo.DepthStencilAccess.IsUsingStencil() ? StoreAction : ERenderTargetStoreAction::ENoAction;
 }
 
-void SetupGBufferRenderTargetInfo(const FSceneTexturesConfig& SceneTexturesConfig, FGraphicsPipelineRenderTargetsInfo& RenderTargetsInfo, bool bSetupDepthStencil)
+void SetupGBufferRenderTargetInfo(const FSceneTexturesConfig& SceneTexturesConfig, FGraphicsPipelineRenderTargetsInfo& RenderTargetsInfo, bool bSetupDepthStencil, EGBufferLayout Layout)
 {
-	SceneTexturesConfig.GetGBufferRenderTargetsInfo(RenderTargetsInfo);
+	SceneTexturesConfig.GetGBufferRenderTargetsInfo(RenderTargetsInfo, Layout);
 	if (bSetupDepthStencil)
 	{
 		SetupDepthStencilInfo(PF_DepthStencil, SceneTexturesConfig.DepthCreateFlags, ERenderTargetLoadAction::ELoad,
diff --git a/Engine/Source/Runtime/Renderer/Private/SingleLayerWaterRendering.cpp b/Engine/Source/Runtime/Renderer/Private/SingleLayerWaterRendering.cpp
--- a/Engine/Source/Runtime/Renderer/Private/SingleLayerWaterRendering.cpp
+++ b/Engine/Source/Runtime/Renderer/Private/SingleLayerWaterRendering.cpp
@@ -138,12 +138,22 @@ static TAutoConsoleVariable<int32> CVarSingleLayerWaterPassOptimizedClear(
 	TEXT("Toggles optimized depth clear"),
 	ECVF_RenderThreadSafe);
 
+static TAutoConsoleVariable<int32> CVarWaterSingleLayerForceVelocity(
+	TEXT("r.Water.SingleLayer.ForceVelocity"), 1,
+	TEXT("Whether to always output velocity, even if the velocity pass is not set to \"Write during base pass\"."),
+	ECVF_ReadOnly | ECVF_RenderThreadSafe);
+
 static int32 GetSingleLayerWaterReflectionTechnique()
 {
 	const int32 Value = CVarWaterSingleLayerReflection.GetValueOnRenderThread();
 	return FMath::Clamp(Value, 0, ESingleLayerWaterReflections::MaxValue);
 }
 
+static EGBufferLayout GetSingleLayerWaterGBufferLayout()
+{
+	return CVarWaterSingleLayerForceVelocity.GetValueOnRenderThread() != 0 ? GBL_ForceVelocity : GBL_Default;
+}
+
 // This is to have platforms use the simple single layer water shading similar to mobile: no dynamic lights, only sun and sky, no distortion, no colored transmittance on background, no custom depth read.
 bool SingleLayerWaterUsesSimpleShading(EShaderPlatform ShaderPlatform)
 {
@@ -1272,8 +1282,9 @@ void FDeferredShadingSceneRenderer::RenderSingleLayerWaterInner(
 
 	const FRDGSystemTextures& SystemTextures = FRDGSystemTextures::Get(GraphBuilder);
 
+	const EGBufferLayout GBufferLayout = GetSingleLayerWaterGBufferLayout();
 	TStaticArray<FTextureRenderTargetBinding, MaxSimultaneousRenderTargets> BasePassTextures;
-	uint32 BasePassTextureCount = SceneTextures.GetGBufferRenderTargets(BasePassTextures);
+	uint32 BasePassTextureCount = SceneTextures.GetGBufferRenderTargets(BasePassTextures, GBufferLayout);
 	if(IsWaterSeparateMainDirLightEnabled(Scene->GetShaderPlatform()))
 	{
 		const bool bNeverClear = true;
@@ -1352,6 +1363,13 @@ void FDeferredShadingSceneRenderer::RenderSingleLayerWaterInner(
 		PassParameters->RenderTargets = GetRenderTargetBindings(ERenderTargetLoadAction::ELoad, BasePassTexturesView);
 		PassParameters->RenderTargets.DepthStencil = DepthStencilBinding;
 
+		// Make sure to clear the velocity texture if it wasn't already written to. This can be the case if the velocity pass is set to "Write after base pass".
+		const FGBufferBindings& GBufferBindings = SceneTextures.Config.GBufferBindings[GBufferLayout];
+		if (GBufferBindings.GBufferVelocity.Index > 0 && GBufferBindings.GBufferVelocity.Index < (int32)BasePassTextureCount && !HasBeenProduced(SceneTextures.Velocity))
+		{
+			PassParameters->RenderTargets[GBufferBindings.GBufferVelocity.Index].SetLoadAction(ERenderTargetLoadAction::EClear);
+		}
+
 		Pass->BuildRenderingCommands(GraphBuilder, Scene->GPUScene, PassParameters->InstanceCullingDrawParams);
 
 		if (bRenderInParallel)
@@ -1491,7 +1509,7 @@ bool FSingleLayerWaterPassMeshProcessor::Process(
 		FeatureLevel,
 		bRenderSkylight,
 		false,
-		GBL_Default,
+		GetSingleLayerWaterGBufferLayout(),
 		&WaterPassShaders.VertexShader,
 		&WaterPassShaders.PixelShader
 		))
@@ -1535,6 +1553,7 @@ void FSingleLayerWaterPassMeshProcessor::CollectPSOInitializers(const FSceneText
 			TBasePassVertexShaderPolicyParamType<LightMapPolicyType>,
 			TBasePassPixelShaderPolicyParamType<LightMapPolicyType>> WaterPassShaders;
 
+		const EGBufferLayout GBufferLayout = GetSingleLayerWaterGBufferLayout();
 		const bool bRenderSkylight = true;
 		if (!GetBasePassShaders<LightMapPolicyType>(
 			Material,
@@ -1543,7 +1562,7 @@ void FSingleLayerWaterPassMeshProcessor::CollectPSOInitializers(const FSceneText
 			FeatureLevel,
 			bRenderSkylight,
 			false,
-			GBL_Default,
+			GBufferLayout,
 			&WaterPassShaders.VertexShader,
 			&WaterPassShaders.PixelShader
 			))
@@ -1552,7 +1571,7 @@ void FSingleLayerWaterPassMeshProcessor::CollectPSOInitializers(const FSceneText
 		}
 
 		FGraphicsPipelineRenderTargetsInfo RenderTargetsInfo;
-		SetupGBufferRenderTargetInfo(SceneTexturesConfig, RenderTargetsInfo, true /*bSetupDepthStencil*/);
+		SetupGBufferRenderTargetInfo(SceneTexturesConfig, RenderTargetsInfo, true /*bSetupDepthStencil*/, GBufferLayout);
 		if (IsWaterSeparateMainDirLightEnabled(GMaxRHIShaderPlatform))
 		{
 			AddRenderTargetInfo(PF_FloatR11G11B10, TexCreate_ShaderResource | TexCreate_RenderTargetable, RenderTargetsInfo);
diff --git a/Engine/Source/Runtime/Renderer/Public/MeshPassProcessor.h b/Engine/Source/Runtime/Renderer/Public/MeshPassProcessor.h
--- a/Engine/Source/Runtime/Renderer/Public/MeshPassProcessor.h
+++ b/Engine/Source/Runtime/Renderer/Public/MeshPassProcessor.h
@@ -10,6 +10,7 @@
 #include "MeshDrawShaderBindings.h"
 #include "MeshMaterialShader.h"
 #include "RendererInterface.h"
+#include "GBufferInfo.h"
 #include "SceneUtils.h"
 #include "MeshBatch.h"
 #include "MeshDrawCommandStatsDefines.h"
@@ -2462,7 +2463,7 @@ RENDERER_API extern FMeshDrawCommandSortKey CalculateMeshStaticSortKey(const FMe
 
 RENDERER_API extern void AddRenderTargetInfo(EPixelFormat PixelFormat, ETextureCreateFlags CreateFlags, FGraphicsPipelineRenderTargetsInfo& RenderTargetsInfo);
 RENDERER_API extern void SetupDepthStencilInfo(EPixelFormat DepthStencilFormat, ETextureCreateFlags DepthStencilCreateFlags, ERenderTargetLoadAction DepthTargetLoadAction, ERenderTargetLoadAction StencilTargetLoadAction, FExclusiveDepthStencil DepthStencilAccess, FGraphicsPipelineRenderTargetsInfo& RenderTargetsInfo);
-RENDERER_API extern void SetupGBufferRenderTargetInfo(const FSceneTexturesConfig& SceneTexturesConfig, FGraphicsPipelineRenderTargetsInfo& RenderTargetsInfo, bool bSetupDepthStencil);
+RENDERER_API extern void SetupGBufferRenderTargetInfo(const FSceneTexturesConfig& SceneTexturesConfig, FGraphicsPipelineRenderTargetsInfo& RenderTargetsInfo, bool bSetupDepthStencil, EGBufferLayout Layout = GBL_Default);
 RENDERER_API extern void ApplyTargetsInfo(FGraphicsPipelineStateInitializer& GraphicsPSOInit, const FGraphicsPipelineRenderTargetsInfo& RenderTargetsInfo);
 RENDERER_API extern ESubpassHint GetSubpassHint(const FStaticShaderPlatform Platform, bool bIsUsingGBuffers, bool bMultiViewRendering, uint32 NumSamples);
 
