[VSM] Fix RDG transition error for the hierarchical page flag build (using SHADER_PARAMETER_RDG_TEXTURE_SRV for mip 0 read access).

#rnx

diff --git a/Engine/Source/Runtime/Renderer/Private/VirtualShadowMaps/VirtualShadowMapArray.cpp b/Engine/Source/Runtime/Renderer/Private/VirtualShadowMaps/VirtualShadowMapArray.cpp
--- a/Engine/Source/Runtime/Renderer/Private/VirtualShadowMaps/VirtualShadowMapArray.cpp
+++ b/Engine/Source/Runtime/Renderer/Private/VirtualShadowMaps/VirtualShadowMapArray.cpp
@@ -1030,8 +1030,8 @@ class FGenerateHierarchicalPageFlagsCS : public FVirtualShadowMapPageManagementS
 		SHADER_PARAMETER_RDG_BUFFER_SRV(StructuredBuffer< FPhysicalPageMetaData >, PhysicalPageMetaData)
 		SHADER_PARAMETER_RDG_TEXTURE_UAV_ARRAY(RWTexture2D<uint>, OutPageFlagMips, [VSM_LOG2_PAGE_SIZE - 1u])
 		SHADER_PARAMETER_RDG_TEXTURE_UAV_ARRAY(RWTexture2D<uint>, OutPageReceiverMaskMips, [VSM_LOG2_PAGE_SIZE])
-		SHADER_PARAMETER_RDG_TEXTURE(Texture2D<uint>, InPageFlags)
-		SHADER_PARAMETER_RDG_TEXTURE(Texture2D<uint>, InPageReceiverMasks)
+		SHADER_PARAMETER_RDG_TEXTURE_SRV(Texture2D<uint>, InPageFlags)
+		SHADER_PARAMETER_RDG_TEXTURE_SRV(Texture2D<uint>, InPageReceiverMasks)
 		SHADER_PARAMETER_RDG_BUFFER_UAV(RWStructuredBuffer<FIntVector4>, OutUncachedPageRectBounds)
 		SHADER_PARAMETER_RDG_BUFFER_UAV(RWStructuredBuffer<FIntVector4>, OutAllocatedPageRectBounds)
 	END_SHADER_PARAMETER_STRUCT()
@@ -2447,8 +2447,8 @@ void FVirtualShadowMapArray::BuildPageAllocations(
 		FGenerateHierarchicalPageFlagsCS::FParameters* PassParameters = GraphBuilder.AllocParameters<FGenerateHierarchicalPageFlagsCS::FParameters>();
 		PassParameters->VirtualShadowMap = GetUncachedUniformBuffer(GraphBuilder);
 		PassParameters->PhysicalPageMetaData = GraphBuilder.CreateSRV(PhysicalPageMetaDataRDG);
-		PassParameters->InPageFlags = PageFlagsRDG;
-		PassParameters->InPageReceiverMasks = PageReceiverMasksRDG;
+		PassParameters->InPageFlags = GraphBuilder.CreateSRV(GRHISupportsTextureViews ? FRDGTextureSRVDesc::CreateForMipLevel(PageFlagsRDG, 0) : FRDGTextureSRVDesc::Create(PageFlagsRDG));
+		PassParameters->InPageReceiverMasks = GraphBuilder.CreateSRV(GRHISupportsTextureViews ? FRDGTextureSRVDesc::CreateForMipLevel(PageReceiverMasksRDG, 0) : FRDGTextureSRVDesc::Create(PageReceiverMasksRDG));
 		PassParameters->OutUncachedPageRectBounds = GraphBuilder.CreateUAV(UncachedPageRectBoundsRDG);
 		PassParameters->OutAllocatedPageRectBounds = GraphBuilder.CreateUAV(AllocatedPageRectBoundsRDG);
 
