Fix crash at Game runtime due to bad access to cvar.

#tests Game win64 development
#rnx

diff --git a/Engine/Source/Runtime/Renderer/Private/TranslucentLighting.cpp b/Engine/Source/Runtime/Renderer/Private/TranslucentLighting.cpp
--- a/Engine/Source/Runtime/Renderer/Private/TranslucentLighting.cpp
+++ b/Engine/Source/Runtime/Renderer/Private/TranslucentLighting.cpp
@@ -1220,7 +1220,7 @@ bool IsTranslucencyLightingVolumeUsingVoxelMarkingSupported()
 
 bool IsTranslucencyLightingVolumeUsingVoxelMarking()
 {
-	return GUseTranslucentLightingVolumes && IsTranslucencyLightingVolumeUsingVoxelMarkingSupported() && CVarTranslucencyLightingVolumeMarkVoxels.GetValueOnRenderThread();
+	return GUseTranslucentLightingVolumes && IsTranslucencyLightingVolumeUsingVoxelMarkingSupported() && CVarTranslucencyLightingVolumeMarkVoxels.GetValueOnAnyThread();
 }
 
 bool IsTranslucencyLightingVolumeUsingBlueNoise()
