- Fix race with Post Init Views particle systems
- This call must be made before we send of GDME tasks and should hopefully never move again!

#rrb

diff --git a/Engine/Source/Runtime/Renderer/Private/SceneVisibility.cpp b/Engine/Source/Runtime/Renderer/Private/SceneVisibility.cpp
--- a/Engine/Source/Runtime/Renderer/Private/SceneVisibility.cpp
+++ b/Engine/Source/Runtime/Renderer/Private/SceneVisibility.cpp
@@ -5654,6 +5654,22 @@ void FDeferredShadingSceneRenderer::BeginInitViews(
 
 	PreVisibilityFrameSetup(GraphBuilder);
 
+	{
+		// This is to init the ViewUniformBuffer before rendering for the Niagara compute shader.
+		// This needs to run before ComputeViewVisibility() is called, but the views normally initialize the ViewUniformBuffer after that (at the end of this method).
+		// Note: This MUST happen before we start to execute GDME tasks of any kind otherwise we will race with PostInitViews renderers!
+		if (FXSystem && FXSystem->RequiresEarlyViewUniformBuffer() && Views.IsValidIndex(0) && bRendererOutputFinalSceneColor)
+		{
+			// during ISR, instanced view RHI resources need to be initialized first.
+			if (FViewInfo* InstancedView = const_cast<FViewInfo*>(Views[0].GetInstancedView()))
+			{
+				InstancedView->InitRHIResources();
+			}
+			Views[0].InitRHIResources();
+			FXSystem->PostInitViews(GraphBuilder, GetSceneViews(), !ViewFamily.EngineShowFlags.HitProxies);
+		}
+	}
+
 	// Start processing dynamic mesh elements tasks early enough to overlap with shadows and GPU scene update.
 	TaskDatas.VisibilityTaskData->StartGatherDynamicMeshElements();
 
@@ -5687,21 +5703,6 @@ void FDeferredShadingSceneRenderer::BeginInitViews(
 		ViewInfo->InstanceFactor = InstanceFactor > 0 ? InstanceFactor : 1;
 	}
 
-	{
-		// This is to init the ViewUniformBuffer before rendering for the Niagara compute shader.
-		// This needs to run before ComputeViewVisibility() is called, but the views normally initialize the ViewUniformBuffer after that (at the end of this method).
-		if (FXSystem && FXSystem->RequiresEarlyViewUniformBuffer() && Views.IsValidIndex(0) && bRendererOutputFinalSceneColor)
-		{
-			// during ISR, instanced view RHI resources need to be initialized first.
-			if (FViewInfo* InstancedView = const_cast<FViewInfo*>(Views[0].GetInstancedView()))
-			{
-				InstancedView->InitRHIResources();
-			}
-			Views[0].InitRHIResources();
-			FXSystem->PostInitViews(GraphBuilder, GetSceneViews(), !ViewFamily.EngineShowFlags.HitProxies);
-		}
-	}
-
 	LumenScenePDIVisualization();
 
 	// This must happen before we start initialising and using views.
