[HWRT] For tracing feedback use SceneInstanceIndex from InstanceExtraData to correctly map instance to geometry instance.

#rnx

diff --git a/Engine/Shaders/Private/RayTracing/RayTracingFeedback.usf b/Engine/Shaders/Private/RayTracing/RayTracingFeedback.usf
--- a/Engine/Shaders/Private/RayTracing/RayTracingFeedback.usf
+++ b/Engine/Shaders/Private/RayTracing/RayTracingFeedback.usf
@@ -6,6 +6,7 @@
 
 #include "../Common.ush"
 #include "../ComputeShaderUtils.ush"
+#include "/Engine/Shared/RayTracingDebugTypes.h"
 
 #ifndef THREADGROUP_SIZE
 #define THREADGROUP_SIZE 1
@@ -40,6 +41,7 @@ void RayTracingProcessFeedbackCS(uint3 GroupId : SV_GroupID, int GroupThreadInde
 
 StructuredBuffer<uint> InstanceHitCountBuffer;
 RWStructuredBuffer<uint> RWGeometryHitCountBuffer;
+StructuredBuffer<FRayTracingInstanceExtraData> InstanceExtraDataBuffer;
 
 [numthreads(THREADGROUP_SIZE, 1, 1)]
 void RayTracingUpdateGeometryHitCountCS(uint3 GroupId : SV_GroupID, int GroupThreadIndex : SV_GroupIndex)
@@ -48,10 +50,13 @@ void RayTracingUpdateGeometryHitCountCS(uint3 GroupId : SV_GroupID, int GroupThr
 
 	if (Index < NumInstances)
 	{
+		const uint SceneInstanceIndex = InstanceExtraDataBuffer[Index].SceneInstanceIndex;
 		const uint InstanceHitCount = InstanceHitCountBuffer[Index];
-		const uint GeometryIndex = AccelerationStructureIndexBuffer[Index];
+
+		const uint GeometryIndex = AccelerationStructureIndexBuffer[SceneInstanceIndex];
+
 		InterlockedAdd(RWGeometryHitCountBuffer[GeometryIndex], InstanceHitCount);
 		// TODO: look into improving this
-		RWGeometryHandleBuffer[GeometryIndex] = GeometryHandleBuffer[Index];
+		RWGeometryHandleBuffer[GeometryIndex] = GeometryHandleBuffer[SceneInstanceIndex];
 	}
 }
diff --git a/Engine/Source/Runtime/Renderer/Private/RayTracing/RayTracingScene.cpp b/Engine/Source/Runtime/Renderer/Private/RayTracing/RayTracingScene.cpp
--- a/Engine/Source/Runtime/Renderer/Private/RayTracing/RayTracingScene.cpp
+++ b/Engine/Source/Runtime/Renderer/Private/RayTracing/RayTracingScene.cpp
@@ -270,7 +270,7 @@ void FRayTracingScene::Create(FRDGBuilder& GraphBuilder, const FViewInfo& View,
 			Layer.InstanceExtraDataBuffer = GraphBuilder.CreateBuffer(InstanceExtraDataBufferDesc, TEXT("FRayTracingScene::InstanceExtraDataBuffer"));
 			InstanceExtraDataBufferUAV = GraphBuilder.CreateUAV(Layer.InstanceExtraDataBuffer);
 
-			AddClearUAVPass(GraphBuilder, InstanceExtraDataBufferUAV, 0xFFFFFFFF);
+			AddClearUAVPass(GraphBuilder, InstanceExtraDataBufferUAV, 0xFFFFFFFF, ComputePassFlags);
 		}
 
 		if (Layer.InstancesDebugData.Num() > 0)
@@ -385,7 +385,6 @@ struct FRayTracingProcessFeedbackCS : public FGlobalShader
 	BEGIN_SHADER_PARAMETER_STRUCT(FParameters, )		
 		SHADER_PARAMETER_RDG_BUFFER_SRV(StructuredBuffer, GeometryHitCountBuffer)
 		SHADER_PARAMETER_RDG_BUFFER_UAV(RWStructuredBuffer, RWGeometryHandleBuffer)
-		SHADER_PARAMETER_RDG_BUFFER_SRV(StructuredBuffer, AccelerationStructureIndexBuffer)
 		SHADER_PARAMETER_RDG_BUFFER_UAV(RWBuffer<uint>, RWGeometryHandleAllocator)
 		SHADER_PARAMETER(uint32, NumGeometries)
 	END_SHADER_PARAMETER_STRUCT()
@@ -421,6 +420,7 @@ struct FRayTracingUpdateGeometryHitCountCS : public FGlobalShader
 		SHADER_PARAMETER_RDG_BUFFER_UAV(RWStructuredBuffer, RWGeometryHitCountBuffer)
 		SHADER_PARAMETER_RDG_BUFFER_UAV(RWStructuredBuffer, RWGeometryHandleBuffer)
 		SHADER_PARAMETER_RDG_BUFFER_SRV(StructuredBuffer, AccelerationStructureIndexBuffer)
+		SHADER_PARAMETER_RDG_BUFFER_SRV(StructuredBuffer, InstanceExtraDataBuffer)
 
 		SHADER_PARAMETER(uint32, NumInstances)
 	END_SHADER_PARAMETER_STRUCT()
@@ -481,7 +481,7 @@ void FRayTracingScene::FinishTracingFeedback(FRDGBuilder& GraphBuilder, ERDGPass
 		GeometryHitCountBufferDesc.NumElements = NumGeometries;
 
 		GeometryHitCountBuffer = GraphBuilder.CreateBuffer(GeometryHitCountBufferDesc, TEXT("FRayTracingScene::GeometryHitCountBuffer"));
-		AddClearUAVPass(GraphBuilder, GraphBuilder.CreateUAV(GeometryHitCountBuffer), 0);
+		AddClearUAVPass(GraphBuilder, GraphBuilder.CreateUAV(GeometryHitCountBuffer), 0, ComputePassFlags);
 
 		FRayTracingUpdateGeometryHitCountCS::FParameters* PassParameters = GraphBuilder.AllocParameters<FRayTracingUpdateGeometryHitCountCS::FParameters>();
 
@@ -490,6 +490,7 @@ void FRayTracingScene::FinishTracingFeedback(FRDGBuilder& GraphBuilder, ERDGPass
 		PassParameters->InstanceHitCountBuffer = GraphBuilder.CreateSRV(Layer.InstanceHitCountBuffer);
 		PassParameters->RWGeometryHitCountBuffer = GraphBuilder.CreateUAV(GeometryHitCountBuffer);
 		PassParameters->RWGeometryHandleBuffer = GraphBuilder.CreateUAV(GeometryHandleBuffer);
+		PassParameters->InstanceExtraDataBuffer = GraphBuilder.CreateSRV(Layer.InstanceExtraDataBuffer);
 		PassParameters->NumInstances = NumInstances;
 
 		const FIntVector GroupSize = FComputeShaderUtils::GetGroupCountWrapped(NumInstances, FRayTracingUpdateGeometryHitCountCS::ThreadGroupSize);
@@ -510,10 +511,9 @@ void FRayTracingScene::FinishTracingFeedback(FRDGBuilder& GraphBuilder, ERDGPass
 		FRDGBufferDesc GeometryHandleAllocatorBufferDesc = FRDGBufferDesc::CreateBufferDesc(sizeof(uint32), 1);
 		GeometryHandleAllocatorBufferDesc.Usage = EBufferUsageFlags(GeometryHandleAllocatorBufferDesc.Usage | BUF_SourceCopy);
 		GeometryHandleAllocatorBuffer = GraphBuilder.CreateBuffer(GeometryHandleAllocatorBufferDesc, TEXT("FRayTracingScene::GeometryHandleAllocator"));
-		AddClearUAVPass(GraphBuilder, GraphBuilder.CreateUAV(GeometryHandleAllocatorBuffer, PF_R32_UINT), 0);
+		AddClearUAVPass(GraphBuilder, GraphBuilder.CreateUAV(GeometryHandleAllocatorBuffer, PF_R32_UINT), 0, ComputePassFlags);
 
 		FRayTracingProcessFeedbackCS::FParameters* PassParameters = GraphBuilder.AllocParameters<FRayTracingProcessFeedbackCS::FParameters>();
-		PassParameters->AccelerationStructureIndexBuffer = GraphBuilder.CreateSRV(Layer.AccelerationStructureIndexBuffer);
 		PassParameters->GeometryHitCountBuffer = GraphBuilder.CreateSRV(GeometryHitCountBuffer);
 		PassParameters->RWGeometryHandleBuffer = GraphBuilder.CreateUAV(GeometryHandleBuffer);
 		PassParameters->RWGeometryHandleAllocator = GraphBuilder.CreateUAV(GeometryHandleAllocatorBuffer, PF_R32_UINT);
