Always copy over the rt geometry initializer segments to the build params to make sure updated vertex buffers are used

diff --git a/Engine/Source/Runtime/Renderer/Private/RayTracing/RayTracingDynamicGeometryUpdateManager.cpp b/Engine/Source/Runtime/Renderer/Private/RayTracing/RayTracingDynamicGeometryUpdateManager.cpp
--- a/Engine/Source/Runtime/Renderer/Private/RayTracing/RayTracingDynamicGeometryUpdateManager.cpp
+++ b/Engine/Source/Runtime/Renderer/Private/RayTracing/RayTracingDynamicGeometryUpdateManager.cpp
@@ -590,6 +590,10 @@ void FRayTracingDynamicGeometryUpdateManager::ScheduleUpdates(bool bUseTracingFe
 			{
 				RTGeoBuildParams.Segments = MakeArrayView(&SegmentData[InBuildParams.SegmentOffset], InBuildParams.Geometry->Initializer.Segments.Num());
 			}
+			else
+			{
+				RTGeoBuildParams.Segments = InBuildParams.Geometry->Initializer.Segments;
+			}
 
 			BuildParams.Add(MoveTemp(RTGeoBuildParams));
 
@@ -664,6 +668,10 @@ void FRayTracingDynamicGeometryUpdateManager::ScheduleUpdates(bool bUseTracingFe
 		{
 			BuildParam.Segments = MakeArrayView(&SegmentData[Update.SegmentOffset], Update.Geometry->Initializer.Segments.Num());
 		}
+		else
+		{
+			BuildParam.Segments = Update.Geometry->Initializer.Segments;
+		}
 		BuildParams.Add(MoveTemp(BuildParam));
 
 		if (!Update.DispatchCommands.IsEmpty())
