[Lighting] Fix anisotropy rendering with split screen.

diff --git a/Engine/Source/Runtime/Renderer/Private/AnisotropyRendering.cpp b/Engine/Source/Runtime/Renderer/Private/AnisotropyRendering.cpp
--- a/Engine/Source/Runtime/Renderer/Private/AnisotropyRendering.cpp
+++ b/Engine/Source/Runtime/Renderer/Private/AnisotropyRendering.cpp
@@ -372,7 +372,10 @@ void FDeferredShadingSceneRenderer::RenderAnisotropyPass(
 			ParallelMeshPass.BuildRenderingCommands(GraphBuilder, Scene->GPUScene, PassParameters->InstanceCullingDrawParams);
 			if (bDoParallelPass)
 			{
-				AddClearRenderTargetPass(GraphBuilder, SceneTextures.GBufferF);
+				if (ViewIndex == 0)
+				{
+					AddClearRenderTargetPass(GraphBuilder, SceneTextures.GBufferF);
+				}
 
 				PassParameters->RenderTargets[0] = FRenderTargetBinding(SceneTextures.GBufferF, ERenderTargetLoadAction::ELoad);
 
@@ -387,7 +390,7 @@ void FDeferredShadingSceneRenderer::RenderAnisotropyPass(
 			}
 			else
 			{
-				PassParameters->RenderTargets[0] = FRenderTargetBinding(SceneTextures.GBufferF, ERenderTargetLoadAction::EClear);
+				PassParameters->RenderTargets[0] = FRenderTargetBinding(SceneTextures.GBufferF, ViewIndex == 0 ? ERenderTargetLoadAction::EClear : ERenderTargetLoadAction::ELoad);
 
 				GraphBuilder.AddPass(
 					RDG_EVENT_NAME("AnisotropyPass"),
