Fix crash with incorrect sized staging buffer.
#rnx

diff --git a/Engine/Source/Runtime/Renderer/Private/VT/VirtualTextureFeedback.cpp b/Engine/Source/Runtime/Renderer/Private/VT/VirtualTextureFeedback.cpp
--- a/Engine/Source/Runtime/Renderer/Private/VT/VirtualTextureFeedback.cpp
+++ b/Engine/Source/Runtime/Renderer/Private/VT/VirtualTextureFeedback.cpp
@@ -134,7 +134,8 @@ void FVirtualTextureFeedback::TransferGPUToCPU(FRHICommandList& RHICmdList, FBuf
 	FeedbackItem.GPUMask = FRHIGPUMask::FromIndex(RHICmdList.GetGPUMask().GetFirstIndex());
 	SCOPED_GPU_MASK(RHICmdList, FeedbackItem.GPUMask);
 
-	RHICmdList.CopyToStagingBuffer(Buffer, FeedbackItem.StagingBuffer, 0, Desc.BufferSize * sizeof(uint32));
+	const uint32 FeedbackStride = Desc.bPageAndCount ? 2 : 1;
+	RHICmdList.CopyToStagingBuffer(Buffer, FeedbackItem.StagingBuffer, 0, Desc.BufferSize * FeedbackStride * sizeof(uint32));
 
 	Fences->Allocate(RHICmdList, WriteIndex);
 	Fences->Write(RHICmdList, WriteIndex);
