Also add the material shader index to the shader hash of the RT flags because it's used to deduplicate SBT allocations and the
material shader index is stored in the user data of the SBT binding data (same shader hash can be moved to another material shader index value)

diff --git a/Engine/Source/Runtime/Renderer/Private/PrimitiveSceneInfo.cpp b/Engine/Source/Runtime/Renderer/Private/PrimitiveSceneInfo.cpp
--- a/Engine/Source/Runtime/Renderer/Private/PrimitiveSceneInfo.cpp
+++ b/Engine/Source/Runtime/Renderer/Private/PrimitiveSceneInfo.cpp
@@ -1077,6 +1077,10 @@ void CacheRayTracingMeshBatch(
 		uint32 ShaderHash = Shader != nullptr ? GetTypeHash(Shader->GetHash()) : 0;
 		uint32 ShaderBindingsHash = RTMeshCommand.ShaderBindings.GetDynamicInstancingHash();
 
+		// Also add the material shader index to the hash because it's used to deduplicate SBT allocations and the 
+		// material shader index is stored in the user data of the SBT binding data (same shader hash can be moved to another material shader index value)
+		ShaderHash = HashCombine(ShaderHash, RTMeshCommand.MaterialShaderIndex);
+
 		Hash <<= 1; // TODO: It would probably be better to use some kind of proper 64 bit mix here?
 		Hash ^= (uint64(ShaderBindingsHash) << 32) | uint64(ShaderHash);
 		
