Fix translucency alpha not accessable in TAA mode in post process material.
* Add rendertargable tex create flag for TAA and Motion blur texture
* Move the Buffer visualization PPM to later before bloom so translucency is blended in.

#tests Done in EngineTest where it fails SeparateTranslucencyAlpha

#virtualized

diff --git a/Engine/Source/Runtime/Renderer/Private/PostProcess/PostProcessMotionBlur.cpp b/Engine/Source/Runtime/Renderer/Private/PostProcess/PostProcessMotionBlur.cpp
--- a/Engine/Source/Runtime/Renderer/Private/PostProcess/PostProcessMotionBlur.cpp
+++ b/Engine/Source/Runtime/Renderer/Private/PostProcess/PostProcessMotionBlur.cpp
@@ -1063,7 +1063,7 @@ FMotionBlurOutputs AddMotionBlurFilterPass(
 				ColorTexture->Desc.Texture->Desc.Extent,
 				bAlphaChannel ? PF_FloatRGBA : PF_FloatRGB,
 				FClearValueBinding::None,
-				TexCreate_UAV | TexCreate_ShaderResource | GFastVRamConfig.MotionBlur);
+				TexCreate_UAV | TexCreate_ShaderResource | TexCreate_RenderTargetable | GFastVRamConfig.MotionBlur);
 
 			FRDGTextureRef FullResTexture = GraphBuilder.CreateTexture(Desc, TEXT("MotionBlur.SceneColor"));
 
diff --git a/Engine/Source/Runtime/Renderer/Private/PostProcess/TemporalAA.cpp b/Engine/Source/Runtime/Renderer/Private/PostProcess/TemporalAA.cpp
--- a/Engine/Source/Runtime/Renderer/Private/PostProcess/TemporalAA.cpp
+++ b/Engine/Source/Runtime/Renderer/Private/PostProcess/TemporalAA.cpp
@@ -1077,6 +1077,11 @@ FDefaultTemporalUpscaler::FOutputs AddGen4MainTemporalAAPasses(
 		TAAParameters.OutputViewRect.Max = HistoryViewSize;
 	}
 
+	// Post process material access the alpha after TAA pass requires the buffer to be rendertargetable in the main pass.
+	TAAParameters.bOutputRenderTargetable = TAAParameters.Pass == ETAAPassConfig::Main || 
+		                                    TAAParameters.Pass == ETAAPassConfig::MainUpsampling || 
+		                                    TAAParameters.Pass == ETAAPassConfig::MainSuperSampling;
+
 	TAAParameters.DownsampleOverrideFormat = PassInputs.DownsampleOverrideFormat;
 
 	TAAParameters.bDownsample = (PassInputs.bGenerateSceneColorHalfRes || PassInputs.bGenerateSceneColorQuarterRes) && TAAParameters.Quality != ETAAQuality::High;
