Mobile: Fixed case where a non-default directional light channels were applied incorrectly with a deferred shading enabled.

diff --git a/Engine/Source/Runtime/Renderer/Private/MobileDeferredShadingPass.cpp b/Engine/Source/Runtime/Renderer/Private/MobileDeferredShadingPass.cpp
--- a/Engine/Source/Runtime/Renderer/Private/MobileDeferredShadingPass.cpp
+++ b/Engine/Source/Runtime/Renderer/Private/MobileDeferredShadingPass.cpp
@@ -561,8 +561,6 @@ static void RenderDirectionalLight(FRHICommandList& RHICmdList, const FScene& Sc
 	RHICmdList.ApplyCachedRenderTargets(GraphicsPSOInit);
 	GraphicsPSOInit.RasterizerState = TStaticRasterizerState<>::GetRHI();
 
-	const uint8 LightingChannelStencilValue = GetLightingChannelStencilValue(LightingChannel);
-	
 	TShaderMapRef<FPostProcessVS> VertexShader(View.ShaderMap);
 	
 	const FMaterialRenderProxy* LightFunctionMaterialProxy = nullptr;
@@ -640,10 +638,12 @@ static void RenderDirectionalLight(FRHICommandList& RHICmdList, const FScene& Sc
 		PassEnableSSR |= (1 << PassIndex);
 		ShadingModelStencilRef[PassIndex] = STENCIL_MOBILE_DEFAULTLIT_MASK | STENCIL_MOBILE_REFLECTIVE_MASK;
 	}
+
+	const uint8 LightingChannelStencilValue = GetLightingChannelStencilValue(LightingChannel);
 	
 	for (int32 PassIndex = 0; PassIndex < NumPasses; ++PassIndex)
 	{
-		SetDirectionalLightDepthStencilState(GraphicsPSOInit, LightingChannelStencilValue, (PassShadingModelStencilMaskAnyDefaultLit & (1 << PassIndex)) == 0);
+		SetDirectionalLightDepthStencilState(GraphicsPSOInit, LightingChannel, (PassShadingModelStencilMaskAnyDefaultLit & (1 << PassIndex)) == 0);
 		
 		FMobileDirectionalLightFunctionPS::FPermutationDomain PermutationVector = FMobileDirectionalLightFunctionPS::BuildPermutationVector(
 			View,
