[TSR] Add partial coverage detection for sparse thin geometry clusters.

* The edge is tracked if it's reprojected pixel is thin geometry or it is itself. It is controlled by r.TSR.ThinGeometryDetection.Coverage.EdgeReprojectiob (default on).
* Reduce shading rejection time in FP16 mode. More vector register use than more barriers, leading to net time reduced for TSR.

#tests Test done in nanite research level needles and project

diff --git a/Engine/Shaders/Private/TemporalSuperResolution/TSRConvolutionNetwork.ush b/Engine/Shaders/Private/TemporalSuperResolution/TSRConvolutionNetwork.ush
--- a/Engine/Shaders/Private/TemporalSuperResolution/TSRConvolutionNetwork.ush
+++ b/Engine/Shaders/Private/TemporalSuperResolution/TSRConvolutionNetwork.ush
@@ -889,33 +889,33 @@ TLaneVector2D<ScalarType, VectorSize, LaneStrideX, LaneStrideY> LerpClamp3x3(
 	TLaneVector2D<ScalarType, VectorSize, LaneStrideX, LaneStrideY> Boundary /*Boundary inside min max of ToClamp)*/,
 	TLaneVector2D<ScalarType, 1, LaneStrideX, LaneStrideY> LerpFactor)
 {
-    #if 0
-	TLaneVector2D<ScalarType, VectorSize, LaneStrideX, LaneStrideY> BoundaryMin;
-	TLaneVector2D<ScalarType, VectorSize, LaneStrideX, LaneStrideY> BoundaryMax;
-	MinMax3x3(Boundary, /* out */ BoundaryMin, /* out */ BoundaryMax);
-
-	TLaneVector2D<ScalarType, VectorSize, LaneStrideX, LaneStrideY> ToClampBoundaryMin;
-	TLaneVector2D<ScalarType, VectorSize, LaneStrideX, LaneStrideY> ToClampBoundaryMax;
-	MinMax3x3(ToClamp, /* out */ ToClampBoundaryMin, /* out */ ToClampBoundaryMax);
-
-	TLaneVector2D<ScalarType, VectorSize, LaneStrideX, LaneStrideY> LerpFactorVector = TLaneVector2D<ScalarType, VectorSize, LaneStrideX, LaneStrideY>::Vectorize(LerpFactor);
-	BoundaryMin = lerp(BoundaryMin, ToClampBoundaryMin, LerpFactorVector);
-	BoundaryMax = lerp(BoundaryMax, ToClampBoundaryMax, LerpFactorVector);
+    #if CONFIG_COMPILE_FP16
+		TLaneVector2D<ScalarType, VectorSize, LaneStrideX, LaneStrideY> BoundaryMin;
+		TLaneVector2D<ScalarType, VectorSize, LaneStrideX, LaneStrideY> BoundaryMax;
+		MinMax3x3(Boundary, /* out */ BoundaryMin, /* out */ BoundaryMax);
+
+		TLaneVector2D<ScalarType, VectorSize, LaneStrideX, LaneStrideY> ToClampBoundaryMin;
+		TLaneVector2D<ScalarType, VectorSize, LaneStrideX, LaneStrideY> ToClampBoundaryMax;
+		MinMax3x3(ToClamp, /* out */ ToClampBoundaryMin, /* out */ ToClampBoundaryMax);
+
+		TLaneVector2D<ScalarType, VectorSize, LaneStrideX, LaneStrideY> LerpFactorVector = TLaneVector2D<ScalarType, VectorSize, LaneStrideX, LaneStrideY>::Vectorize(LerpFactor);
+		BoundaryMin = lerp(BoundaryMin, ToClampBoundaryMin, LerpFactorVector);
+		BoundaryMax = lerp(BoundaryMax, ToClampBoundaryMax, LerpFactorVector);
     #else
-	// register pressure is high
-	TLaneVector2D<ScalarType, VectorSize, LaneStrideX, LaneStrideY> BoundaryMin;
-	TLaneVector2D<ScalarType, VectorSize, LaneStrideX, LaneStrideY> BoundaryMax;
-	TLaneVector2D<ScalarType, VectorSize, LaneStrideX, LaneStrideY> BoundaryLimit;
-	TLaneVector2D<ScalarType, VectorSize, LaneStrideX, LaneStrideY> ToClampBoundaryLimit;
-	TLaneVector2D<ScalarType, VectorSize, LaneStrideX, LaneStrideY> LerpFactorVector = TLaneVector2D<ScalarType, VectorSize, LaneStrideX, LaneStrideY>::Vectorize(LerpFactor);
-
-	BoundaryLimit = Min3x3(Boundary);
-	ToClampBoundaryLimit = Min3x3(ToClamp);
-	BoundaryMin = lerp(BoundaryLimit, ToClampBoundaryLimit, LerpFactorVector);
-
-	BoundaryLimit = Max3x3(Boundary);
-	ToClampBoundaryLimit = Max3x3(ToClamp);
-	BoundaryMax = lerp(BoundaryLimit, ToClampBoundaryLimit, LerpFactorVector);
+		// register pressure is high
+		TLaneVector2D<ScalarType, VectorSize, LaneStrideX, LaneStrideY> BoundaryMin;
+		TLaneVector2D<ScalarType, VectorSize, LaneStrideX, LaneStrideY> BoundaryMax;
+		TLaneVector2D<ScalarType, VectorSize, LaneStrideX, LaneStrideY> BoundaryLimit;
+		TLaneVector2D<ScalarType, VectorSize, LaneStrideX, LaneStrideY> ToClampBoundaryLimit;
+		TLaneVector2D<ScalarType, VectorSize, LaneStrideX, LaneStrideY> LerpFactorVector = TLaneVector2D<ScalarType, VectorSize, LaneStrideX, LaneStrideY>::Vectorize(LerpFactor);
+
+		BoundaryLimit = Min3x3(Boundary);
+		ToClampBoundaryLimit = Min3x3(ToClamp);
+		BoundaryMin = lerp(BoundaryLimit, ToClampBoundaryLimit, LerpFactorVector);
+
+		BoundaryLimit = Max3x3(Boundary);
+		ToClampBoundaryLimit = Max3x3(ToClamp);
+		BoundaryMax = lerp(BoundaryLimit, ToClampBoundaryLimit, LerpFactorVector);
     #endif
 
 	return fastClamp(ToClamp, BoundaryMin, BoundaryMax);
diff --git a/Engine/Shaders/Private/TemporalSuperResolution/TSRDetectThinGeometry.usf b/Engine/Shaders/Private/TemporalSuperResolution/TSRDetectThinGeometry.usf
--- a/Engine/Shaders/Private/TemporalSuperResolution/TSRDetectThinGeometry.usf
+++ b/Engine/Shaders/Private/TemporalSuperResolution/TSRDetectThinGeometry.usf
@@ -5,6 +5,7 @@
 #include "TSRColorSpace.ush"
 #include "TSRDepthVelocityAnalysis.ush"
 #include "TSRConvolutionNetworkPass.ush"
+#include "TSRReprojectionField.ush"
 
 //------------------------------------------------------- CONFIG
 
@@ -19,6 +20,8 @@
 
 #define EDGE_TYPE EDGE_HORIZONTAL_VERTICAL_DIAGNAL
 
+#define CONFIG_THIN_GEOMETRY_EDGE_REPROJECTION (DIM_THIN_GEOMETRY_EDGE_REPROJECTION)
+
 //------------------------------------------------------- CONSTANTS
 
 #if CONFIG_LDS_DESPILL
@@ -98,6 +101,10 @@ void FetchInput(
 {
 	tsr_tensor_float DepthCache;
 	tsr_tensor_bool bHasReprojectionOffset;
+#if CONFIG_THIN_GEOMETRY_EDGE_REPROJECTION
+	tsr_tensor_short2 ReprojectionOffset;
+#endif
+
 	UNROLL_N(SIMD_SIZE)
 	for (uint ElementIndex = 0; ElementIndex < SIMD_SIZE; ElementIndex++)
 	{
@@ -108,8 +115,24 @@ void FetchInput(
 
 		uint EncodedDilateMask = R8Output[int3(ClampedFetchPixelPos,0)];
 		bHasReprojectionOffset.SetElement(ElementIndex, (EncodedDilateMask & 0x80u) != 0u);
+#if CONFIG_THIN_GEOMETRY_EDGE_REPROJECTION
+		ReprojectionOffset.SetElement(ElementIndex, DecodeReprojectionOffset(R8Output[int3(ClampedFetchPixelPos, ThinGeometryTextureIndex)]));
+#endif
 	}
 
+#if CONFIG_THIN_GEOMETRY_EDGE_REPROJECTION
+	tsr_tensor_half ReprojectedCurrentCoverage;
+	UNROLL_N(SIMD_SIZE)
+	for (uint ElementIndex = 0; ElementIndex < SIMD_SIZE; ElementIndex++)
+	{
+		tsr_short2 ClampedFetchPixelPos = PixelPos.GetElement(ElementIndex) + ReprojectionOffset.GetElement(ElementIndex);
+		ReprojectedCurrentCoverage.SetElement(ElementIndex, CurrentCoverageTexture[ClampedFetchPixelPos]);
+	}
+	CurrentCoverage = select(bHasReprojectionOffset, max(CurrentCoverage, ReprojectedCurrentCoverage), tsr_tensor_half::Const(0));
+#else
+	CurrentCoverage = select(bHasReprojectionOffset, CurrentCoverage, tsr_tensor_half::Const(0));
+#endif
+
 	UNROLL_N(SIMD_SIZE)
 	for (uint ElementIndex = 0; ElementIndex < SIMD_SIZE; ElementIndex++)
 	{
@@ -119,8 +142,6 @@ void FetchInput(
 		SceneDepth.SetElement(ElementIndex, tsr_half(Depth));
 		DeviceZThinGeometryDetectionError.SetElement(ElementIndex, tsr_half(ErrorMultiplier * DeviceZError));
 	}
-
-	CurrentCoverage = select(bHasReprojectionOffset, CurrentCoverage, tsr_tensor_half::Const(0));
 } // FetchInput
 
 
diff --git a/Engine/Shaders/Private/TemporalSuperResolution/TSRDilateVelocity.usf b/Engine/Shaders/Private/TemporalSuperResolution/TSRDilateVelocity.usf
--- a/Engine/Shaders/Private/TemporalSuperResolution/TSRDilateVelocity.usf
+++ b/Engine/Shaders/Private/TemporalSuperResolution/TSRDilateVelocity.usf
@@ -19,6 +19,7 @@
 	#define TILE_SIZE 8
 #endif
 
+#define CONFIG_THIN_GEOMETRY_EDGE_REPROJECTION (DIM_THIN_GEOMETRY_EDGE_REPROJECTION)
 
 //------------------------------------------------------- PARAMETERS
 
@@ -28,6 +29,9 @@ float ReprojectionFieldAntiAliasVelocityThreshold;
 uint bReprojectionField;
 uint bOutputIsMovingTexture;
 uint ReprojectionVectorOutputIndex;
+#if CONFIG_THIN_GEOMETRY_EDGE_REPROJECTION
+uint ThinGeometryTextureIndex;
+#endif
 
 Texture2D<float> SceneDepthTexture;
 Texture2D<float4> SceneVelocityTexture;
@@ -309,7 +313,11 @@ void MainCS(
 			EncodedDeviceZError = EncodeDeviceZError(PrevClosestDeviceZ, ParallaxDepthError);
 		}
 
+#if CONFIG_THIN_GEOMETRY_EDGE_REPROJECTION
+		EncodedReprojectionBoundary = EncodeReprojectionOffset(ReprojectionOffset);
+#else
 		EncodedReprojectionBoundary = 0u;
+#endif
 		EncodedReprojectionJacobian = 0u;
 	}
 
@@ -365,6 +373,11 @@ void MainCS(
 
 		R8Output[uint3(OutputPixelPos, 0)] = uint(round(ReprojectionEdge * 127.0)) | select(bHasReprojectionOffset, 0x80u, 0x00u);
 		R8Output[uint3(OutputPixelPos, 1)] = EncodedDeviceZError;
+
+        #if CONFIG_THIN_GEOMETRY_EDGE_REPROJECTION
+		    R8Output[uint3(OutputPixelPos, ThinGeometryTextureIndex)] = EncodedReprojectionBoundary & 0x0fu;
+        #endif
+
 		BRANCH
 		if (bOutputIsMovingTexture)
 		{
diff --git a/Engine/Shaders/Private/TemporalSuperResolution/TSRReprojectionField.ush b/Engine/Shaders/Private/TemporalSuperResolution/TSRReprojectionField.ush
--- a/Engine/Shaders/Private/TemporalSuperResolution/TSRReprojectionField.ush
+++ b/Engine/Shaders/Private/TemporalSuperResolution/TSRReprojectionField.ush
@@ -16,11 +16,15 @@ static const uint kReprojectionBoundaryOutputIndex = 2;
 
 static const tsr_half2 kFullDilateBoundary = tsr_half2(0.0, 1.0);
 
+uint EncodeReprojectionOffset(tsr_short2 ReprojectionOffset)
+{
+	return  uint(dot(tsr_ushort2(ReprojectionOffset + tsr_ushort(1)), tsr_ushort2(1, 4)));
+}
 
 /** Encodes the reprojection field's dilating offset [-1;1]^2 in first 4bits, and ReprojectionBoundary.xy in 11bit eachs, 26bits total. */
 uint EncodeReprojectionBoundary(tsr_short2 ReprojectionOffset, tsr_half2 ReprojectionBoundary)
 {
-	uint EncodedReprojectionOffset = uint(dot(tsr_ushort2(ReprojectionOffset + tsr_ushort(1)), tsr_ushort2(1, 4)));
+	uint EncodedReprojectionOffset = EncodeReprojectionOffset(ReprojectionOffset);
 
 	tsr_short2 BoundaryBits = tsr_short2(round(ReprojectionBoundary * tsr_half(1023))) + 1024;
 	BoundaryBits = fastClamp(BoundaryBits, tsr_short(1).xx, tsr_short(2047).xx);
diff --git a/Engine/Shaders/Private/TemporalSuperResolution/TSRShadingAnalysis.ush b/Engine/Shaders/Private/TemporalSuperResolution/TSRShadingAnalysis.ush
--- a/Engine/Shaders/Private/TemporalSuperResolution/TSRShadingAnalysis.ush
+++ b/Engine/Shaders/Private/TemporalSuperResolution/TSRShadingAnalysis.ush
@@ -486,8 +486,14 @@ void MeasureRejection(
 	if (bEnableHistoryControl)
 	{
 		tsr_tensor_halfC LerpClampedHistoryInput = LerpClamp3x3(FilteredHistory, FilteredInput, HistoryRelaxationFactor);
-		tsr_tensor_halfC FilteredHistoryBoxMin = Min3x3(LerpClampedHistoryInput);
-		tsr_tensor_halfC FilteredHistoryBoxMax = Max3x3(LerpClampedHistoryInput);
+		#if CONFIG_COMPILE_FP16
+			tsr_tensor_halfC FilteredHistoryBoxMin;
+			tsr_tensor_halfC FilteredHistoryBoxMax;
+			MinMax3x3(LerpClampedHistoryInput, /* out */ FilteredHistoryBoxMin, /* out */ FilteredHistoryBoxMax);
+        #else
+			tsr_tensor_halfC FilteredHistoryBoxMin = Min3x3(LerpClampedHistoryInput);
+			tsr_tensor_halfC FilteredHistoryBoxMax = Max3x3(LerpClampedHistoryInput);
+        #endif
 		tsr_tensor_halfC ScalingFactor = tsr_tensor_halfC::Vectorize((HistoryRelaxationFactor));
 		FilteredBoxMin = select(ScalingFactor > tsr_tensor_halfC::Const(tsr_half(1.0f / 127.0f)), FilteredHistoryBoxMin, FilteredBoxMin);
 		FilteredBoxMax = select(ScalingFactor > tsr_tensor_halfC::Const(tsr_half(1.0f / 127.0f)), FilteredHistoryBoxMax, FilteredBoxMax);
diff --git a/Engine/Source/Runtime/Renderer/Private/PostProcess/TemporalSuperResolution.cpp b/Engine/Source/Runtime/Renderer/Private/PostProcess/TemporalSuperResolution.cpp
--- a/Engine/Source/Runtime/Renderer/Private/PostProcess/TemporalSuperResolution.cpp
+++ b/Engine/Source/Runtime/Renderer/Private/PostProcess/TemporalSuperResolution.cpp
@@ -313,6 +313,12 @@ TAutoConsoleVariable<int32> CVarTSRThinGeometryDetection(
 	TEXT("Other: no history relaxation (Yellow)."),
 	ECVF_RenderThreadSafe);
 
+
+TAutoConsoleVariable<bool> CVarTSRThinGeometryCoverageEdgeReprojection(
+	TEXT("r.TSR.ThinGeometryDetection.Coverage.EdgeReprojection"), true,
+	TEXT("Whether thin geometry edge against not thin geometry should be considered for coverage. Better stability for sparse thin geometry clusters."),
+	ECVF_RenderThreadSafe);
+
 TAutoConsoleVariable<float> CVarTSRThinGeometryCoverageMaxRelaxationWeight(
 	TEXT("r.TSR.ThinGeometryDetection.Coverage.MaxRelaxationWeight"), 0.037,
 	TEXT("The max history clamping box relaxation weight due to thin geometry detection (0 to 1)"),
@@ -704,7 +710,8 @@ class FTSRDilateVelocityCS : public FTSRShader
 	SHADER_USE_PARAMETER_STRUCT(FTSRDilateVelocityCS, FTSRShader);
 
 	class FMotionBlurDirectionsDim : SHADER_PERMUTATION_INT("DIM_MOTION_BLUR_DIRECTIONS", 3);
-	using FPermutationDomain = TShaderPermutationDomain<FMotionBlurDirectionsDim>;
+	class FThinGeometryEdgeReprojectionDim : SHADER_PERMUTATION_BOOL("DIM_THIN_GEOMETRY_EDGE_REPROJECTION");
+	using FPermutationDomain = TShaderPermutationDomain<FMotionBlurDirectionsDim, FThinGeometryEdgeReprojectionDim>;
 
 	BEGIN_SHADER_PARAMETER_STRUCT(FParameters, )
 		SHADER_PARAMETER_STRUCT_INCLUDE(FTSRCommonParameters, CommonParameters)
@@ -718,6 +725,7 @@ class FTSRDilateVelocityCS : public FTSRShader
 		SHADER_PARAMETER(int32, bReprojectionField)
 		SHADER_PARAMETER(int32, bOutputIsMovingTexture)
 		SHADER_PARAMETER(int32, ReprojectionVectorOutputIndex)
+		SHADER_PARAMETER(int32, ThinGeometryTextureIndex)
 
 		SHADER_PARAMETER_RDG_TEXTURE(Texture2D, SceneDepthTexture)
 		SHADER_PARAMETER_RDG_TEXTURE(Texture2D, SceneVelocityTexture)
@@ -949,7 +957,8 @@ class FTSRDetectThinGeometryCS : public FTSRConvolutionNetworkShader
 	DECLARE_GLOBAL_SHADER(FTSRDetectThinGeometryCS);
 	SHADER_USE_PARAMETER_STRUCT(FTSRDetectThinGeometryCS, FTSRConvolutionNetworkShader);
 
-	using FPermutationDomain = TShaderPermutationDomain<FTSRConvolutionNetworkShader::FPermutationDomain>;
+	class FThinGeometryEdgeReprojectionDim : SHADER_PERMUTATION_BOOL("DIM_THIN_GEOMETRY_EDGE_REPROJECTION");
+	using FPermutationDomain = TShaderPermutationDomain<FTSRConvolutionNetworkShader::FPermutationDomain,FThinGeometryEdgeReprojectionDim>;
 
 	BEGIN_SHADER_PARAMETER_STRUCT(FParameters, )
 		SHADER_PARAMETER_STRUCT_INCLUDE(FTSRCommonParameters, CommonParameters)
@@ -2252,6 +2261,7 @@ FDefaultTemporalUpscaler::FOutputs AddTemporalSuperResolutionPasses(
 
 		int32 TileSize = 8;
 		FTSRDilateVelocityCS::FPermutationDomain PermutationVector;
+		PermutationVector.Set<FTSRDilateVelocityCS::FThinGeometryEdgeReprojectionDim>(PassConfig.ThinGeometryDetectionEnable && CVarTSRThinGeometryCoverageEdgeReprojection.GetValueOnRenderThread());
 
 		FTSRDilateVelocityCS::FParameters* PassParameters = GraphBuilder.AllocParameters<FTSRDilateVelocityCS::FParameters>();
 		PassParameters->CommonParameters = CommonParameters;
@@ -2265,7 +2275,8 @@ FDefaultTemporalUpscaler::FOutputs AddTemporalSuperResolutionPasses(
 		PassParameters->ReprojectionFieldAntiAliasVelocityThreshold = FMath::Square(FMath::Max(PassConfig.ReprojectionFieldAntiAliasPixelSpeed / OutputToInputResolutionFraction, 1.0f / 64.0f));
 		PassParameters->bReprojectionField = bReprojectionField;
 		PassParameters->bOutputIsMovingTexture = bOutputIsMovingTexture;
-		
+		PassParameters->ThinGeometryTextureIndex = ThinGeometryTextureIndex;
+
 		PassParameters->SceneDepthTexture = PassInputs.SceneDepth.Texture;
 		PassParameters->SceneVelocityTexture = PassInputs.SceneVelocity.Texture;
 		PassParameters->ReprojectionVectorOutputIndex = DilatedReprojectionVectorTexture->Desc.FirstArraySlice;
@@ -2480,6 +2491,7 @@ FDefaultTemporalUpscaler::FOutputs AddTemporalSuperResolutionPasses(
 		ConvolutionNetworkPermutationVector.Set<FTSRShader::FAlphaChannelDim>(bSupportsAlpha);
 
 		FTSRDetectThinGeometryCS::FPermutationDomain PermutationVector;
+		PermutationVector.Set<FTSRDetectThinGeometryCS::FThinGeometryEdgeReprojectionDim>(CVarTSRThinGeometryCoverageEdgeReprojection.GetValueOnRenderThread());
 		PermutationVector.Set<FTSRConvolutionNetworkShader::FPermutationDomain>(ConvolutionNetworkPermutationVector);
 		PermutationVector = FTSRDetectThinGeometryCS::RemapPermutation(PermutationVector);
 
