Mobile: Fixed case where SkyLight was not correctly applied for a non DefaultLit shading models
#12929

diff --git a/Engine/Source/Runtime/Renderer/Private/MobileDeferredShadingPass.cpp b/Engine/Source/Runtime/Renderer/Private/MobileDeferredShadingPass.cpp
--- a/Engine/Source/Runtime/Renderer/Private/MobileDeferredShadingPass.cpp
+++ b/Engine/Source/Runtime/Renderer/Private/MobileDeferredShadingPass.cpp
@@ -483,7 +483,7 @@ void RenderReflectionEnvironmentSkyLighting(FRHICommandList& RHICmdList, const F
 		GraphicsPSOInit.DepthStencilState = StencilState[PassIndex];
 		
 		FMobileReflectionEnvironmentSkyLightingPS::FPermutationDomain PermutationVector;
-		PermutationVector.Set<FMobileReflectionEnvironmentSkyLightingPS::FEnableShadingModelSupport>(PassEnableSSR & (1 << PassIndex));
+		PermutationVector.Set<FMobileReflectionEnvironmentSkyLightingPS::FEnableShadingModelSupport>(PassEnableShadingModelSupport & (1 << PassIndex));
 		PermutationVector.Set<FMobileReflectionEnvironmentSkyLightingPS::FEnableClustredReflection>(bClustredReflection);
 		PermutationVector.Set<FMobileReflectionEnvironmentSkyLightingPS::FEnablePlanarReflection>(bPlanarReflection);
 		PermutationVector.Set<FMobileReflectionEnvironmentSkyLightingPS::FEnableSkyLight>(bEnableSkyLight);
