Fix Wrong texture format used in velocity pass with mobile preview

diff --git a/Engine/Source/Runtime/Renderer/Private/VelocityRendering.cpp b/Engine/Source/Runtime/Renderer/Private/VelocityRendering.cpp
--- a/Engine/Source/Runtime/Renderer/Private/VelocityRendering.cpp
+++ b/Engine/Source/Runtime/Renderer/Private/VelocityRendering.cpp
@@ -384,8 +384,8 @@ EPixelFormat FVelocityRendering::GetFormat(EShaderPlatform ShaderPlatform)
 		|| FDataDrivenShaderPlatformInfo::GetSupportsRayTracing(ShaderPlatform);
 
 	// Android GLES platform doesn't support R16G16_UNORM and R16G16B16A16_UNORM format, use R16G16_UINT or R16G16B16A16_UINT instead.
-	const bool bIsAndroidOpenGLESPlatform = IsAndroidOpenGLESPlatform(ShaderPlatform);
-	if (bIsAndroidOpenGLESPlatform)
+	const bool bIsOpenGLPlatform = IsOpenGLPlatform(ShaderPlatform);
+	if (bIsOpenGLPlatform)
 	{
 		return bNeedVelocityDepth ? PF_R16G16B16A16_UINT : PF_R16G16_UINT;
 	}
