[DOF] Implement bokeh vignetting in scatter pass.
Geometry close to the lens, such as the lens barrel or matte box, can mask out regions of the bokeh shape and produce vignetting.

diff --git a/Engine/Shaders/Private/DiaphragmDOF/DOFHybridScatterCommon.ush b/Engine/Shaders/Private/DiaphragmDOF/DOFHybridScatterCommon.ush
--- a/Engine/Shaders/Private/DiaphragmDOF/DOFHybridScatterCommon.ush
+++ b/Engine/Shaders/Private/DiaphragmDOF/DOFHybridScatterCommon.ush
@@ -40,6 +40,14 @@
 // Number of pixels per scattering group.
 #define PIXEL_COUNT_PER_SCATTER_GROUP (HYBRID_SCATTERING_GROUP_BORDER_SIZE * HYBRID_SCATTERING_GROUP_BORDER_SIZE)
 
+// 2x float4: 2 floats + (3 flags * 2 floats)
+#define VIGNETTE_DATA_PER_PIXEL (2)
+
+#define SCATTER_DRAW_LIST_DATA_PER_PIXEL (1 + VIGNETTE_DATA_PER_PIXEL)
+
+// Number of float4s in the draw list per scattering group
+#define SCATTER_DRAW_LIST_GROUP_STRIDE (1 + (PIXEL_COUNT_PER_SCATTER_GROUP * SCATTER_DRAW_LIST_DATA_PER_PIXEL))
+
 
 // Configures the drawing method used for scattering bokeh.
 #if PLATFORM_SUPPORTS_RECT_LIST
@@ -87,6 +95,9 @@
 #endif
 
 #if DIM_SCATTER_QUALITY == SCATTER_QUALITY_MEDIUM
+	#define CONFIG_PETZVAL 0
+	#define CONFIG_VIGNETTE 0
 #elif DIM_SCATTER_QUALITY == SCATTER_QUALITY_HIGH
 	#define CONFIG_PETZVAL 1
+	#define CONFIG_VIGNETTE 1
 #endif //DIM_SCATTER_QUALITY
\ No newline at end of file
diff --git a/Engine/Shaders/Private/DiaphragmDOF/DOFHybridScatterCompilation.usf b/Engine/Shaders/Private/DiaphragmDOF/DOFHybridScatterCompilation.usf
--- a/Engine/Shaders/Private/DiaphragmDOF/DOFHybridScatterCompilation.usf
+++ b/Engine/Shaders/Private/DiaphragmDOF/DOFHybridScatterCompilation.usf
@@ -81,18 +81,18 @@ void ScatterGroupPackMainCS(
 		if (GroupId == 0)
 		{
 			#if DIM_HYBRID_SCATTER_FGD
-			for (uint i = 0; i < 5; i++)
+			for (uint i = 0; i < SCATTER_DRAW_LIST_GROUP_STRIDE; i++)
 			{
-				OutForegroundScatterDrawList[ScatterGroupCount * 5 + GroupThreadIndex + i * SCATTER_GROUP_PER_INSTANCE] = float4(0, 0, 0, 0);
+				OutForegroundScatterDrawList[ScatterGroupCount * SCATTER_DRAW_LIST_GROUP_STRIDE + GroupThreadIndex + i * SCATTER_GROUP_PER_INSTANCE] = float4(0, 0, 0, 0);
 			}
 			#endif
 		}
 		else
 		{	
 			#if DIM_HYBRID_SCATTER_BGD
-			for (uint i = 0; i < 5; i++)
+			for (uint i = 0; i < SCATTER_DRAW_LIST_GROUP_STRIDE; i++)
 			{
-				OutBackgroundScatterDrawList[ScatterGroupCount * 5 + GroupThreadIndex + i * SCATTER_GROUP_PER_INSTANCE] = float4(0, 0, 0, 0);
+				OutBackgroundScatterDrawList[ScatterGroupCount * SCATTER_DRAW_LIST_GROUP_STRIDE + GroupThreadIndex + i * SCATTER_GROUP_PER_INSTANCE] = float4(0, 0, 0, 0);
 			}
 			#endif
 		}
diff --git a/Engine/Shaders/Private/DiaphragmDOF/DOFHybridScatterPixelShader.usf b/Engine/Shaders/Private/DiaphragmDOF/DOFHybridScatterPixelShader.usf
--- a/Engine/Shaders/Private/DiaphragmDOF/DOFHybridScatterPixelShader.usf
+++ b/Engine/Shaders/Private/DiaphragmDOF/DOFHybridScatterPixelShader.usf
@@ -5,6 +5,7 @@
 =============================================================================*/
 
 #include "DOFHybridScatterCommon.ush"
+#include "DOFVignetteCommon.ush"
 #include "../FastMath.ush"
 #include "/Engine/Public/DualPixelVectorization.ush"
 
@@ -22,11 +23,15 @@
 
 float ScatteringScaling;
 
+StructuredBuffer<float4> ScatterDrawList;
+
 Texture2D BokehLUT;
 
 float4 ScatterOcclusionSize;
 Texture2D ScatterOcclusion;
 
+float BarrelRadius;
+float BarrelLength;
 
 //------------------------------------------------------- FUNCTIONS
 
@@ -47,16 +52,67 @@ half2x2 dpv_mul_vec_mat(half2x2 InVector, half2x2 InMatrix)
 	);
 }
 
+/// Calculate the visibility of the pixel for a specific bokeh through the lens barrel
+half2 GetPixelVisibility_DPV(half2x2 PixelOffset, uint2 VignetteDataOffset)
+{
+#if CONFIG_VIGNETTE
+	bool bOpticalVignetting = BarrelLength >= 0;
+	if(bOpticalVignetting)
+	{
+		float4 VignetteData[VIGNETTE_DATA_PER_PIXEL][DPV_PIXEL_PER_LANE] = {
+			{
+				ScatterDrawList[VignetteDataOffset[0] + 0],
+				ScatterDrawList[VignetteDataOffset[1] + 0]
+			},
+			{
+				ScatterDrawList[VignetteDataOffset[0] + 1],
+				ScatterDrawList[VignetteDataOffset[1] + 1]
+			}
+		};
+
+		uint2 BarrelData[DPV_PIXEL_PER_LANE] = {
+			asuint(VignetteData[0][0].xy),
+			asuint(VignetteData[0][1].xy)
+		};
+
+		uint2 ProjectedMatteBoxFlags[MAX_MATTE_BOX_FLAGS][DPV_PIXEL_PER_LANE] = {
+			{
+				asuint(VignetteData[0][0].zw),
+				asuint(VignetteData[0][1].zw)
+			},
+			{
+				asuint(VignetteData[1][0].xy),
+				asuint(VignetteData[1][1].xy)
+			},
+			{
+				asuint(VignetteData[1][0].zw),
+				asuint(VignetteData[1][1].zw)
+			}
+		};
+
+		half2x2 SamplePositionAtBarrelEnd;
+		half2 Visibility = HitTestProjectedCylinder_DPV(PixelOffset, BarrelData, BarrelRadius, /* out */ SamplePositionAtBarrelEnd);
+
+		UNROLL
+		for (int Index = 0; Index < MAX_MATTE_BOX_FLAGS; ++Index)
+		{
+			Visibility *= HitTestProjectedMatteBoxFlag_DPV(SamplePositionAtBarrelEnd, ProjectedMatteBoxFlags[Index]);
+		}
+
+		return Visibility;
+	}
+#endif
+	return half2(1.0, 1.0);
+}
+
 void ScatterMainPS(
 	in nointerpolation float4 BokehCenterScreenAndClipPos : TEXCOORD0,
 	in nointerpolation float4 ColorAndAbsCocRadius0 : TEXCOORD1,
 	in nointerpolation float4 ColorAndAbsCocRadius1 : TEXCOORD2,
+	in nointerpolation uint3 GlobalSpriteIdAndPetzvalTransformPacked : TEXCOORD3,
 #if TAYLOR_EXPENDED_COC_INTERSECTION
-	in nointerpolation float4 TaylorExpensionMAD0   : TEXCOORD3,
-	in nointerpolation float4 TaylorExpensionMAD1   : TEXCOORD4,
-#endif
-#if CONFIG_PETZVAL
-	in nointerpolation float2 PetzvalTransformPacked : TEXCOORD5,
+	in nointerpolation float4 TaylorExpensionMAD0   : TEXCOORD4,
+	in nointerpolation float4 TaylorExpensionMAD1   : TEXCOORD5,
 #endif
 	in float4 SvPosition : SV_Position,
 	out float4 OutColor : SV_Target0) 
@@ -112,10 +168,11 @@ void ScatterMainPS(
 	
 	half2 PixelOffset0 = half2(SvPosition.xy - BokehCenter0);
 
+	uint GlobalSpriteId = GlobalSpriteIdAndPetzvalTransformPacked.x;
 #if CONFIG_PETZVAL
 	half2x2 PetzvalTransform = half2x2(
-		UnpackFloat2FromUInt(asuint(PetzvalTransformPacked.x)),
-		UnpackFloat2FromUInt(asuint(PetzvalTransformPacked.y))
+		UnpackFloat2FromUInt(GlobalSpriteIdAndPetzvalTransformPacked.y),
+		UnpackFloat2FromUInt(GlobalSpriteIdAndPetzvalTransformPacked.z)
 	);
 #endif
 
@@ -163,18 +220,30 @@ void ScatterMainPS(
 		}
 		#endif
 
-		half2 Visibility = 1;
-		#if DIM_SCATTER_OCCLUSION
+		BRANCH
+		if (any(Intersection[0] > 0))
 		{
-			half2 CocRadiusComparison = max(PixelAbsCocRadius * half(1 - FAST_GATHERING_COC_ERROR) - CocRadiusAvg, (half2)0);
-			Visibility = CocRadiusVariance / (CocRadiusVariance + CocRadiusComparison * CocRadiusComparison);
+			half2 Visibility = 1;
+			#if DIM_SCATTER_OCCLUSION
+			{
+				half2 CocRadiusComparison = max(PixelAbsCocRadius * half(1 - FAST_GATHERING_COC_ERROR) - CocRadiusAvg, (half2)0);
+				Visibility = CocRadiusVariance / (CocRadiusVariance + CocRadiusComparison * CocRadiusComparison);
+			}
+			#endif
+
+			uint2 ScatterDrawListOffset = {
+				SCATTER_DRAW_LIST_GROUP_STRIDE * GlobalSpriteId + 1 + ((i + 0) * SCATTER_DRAW_LIST_DATA_PER_PIXEL),
+				SCATTER_DRAW_LIST_GROUP_STRIDE * GlobalSpriteId + 1 + ((i + 1) * SCATTER_DRAW_LIST_DATA_PER_PIXEL)
+			};
+			uint2 VignetteDataOffset = ScatterDrawListOffset + 1;
+		
+			Visibility *= GetPixelVisibility_DPV(PixelOffset, VignetteDataOffset);
+
+			half3x2 PixelColor = dpv_interleave_registers(Color[i].rgb, Color[i + 1].rgb);
+			PixelColor = dpv_scale(PixelColor, Intersection * Visibility);
+			half3 ColorAcc = dpv_lo(PixelColor) + dpv_hi(PixelColor);
+			OutColor.rgb += ColorAcc;
 		}
-		#endif
-
-		half3x2 PixelColor = dpv_interleave_registers(Color[i].rgb, Color[i + 1].rgb);
-		PixelColor = dpv_scale(PixelColor, Intersection * Visibility);
-		half3 ColorAcc = dpv_lo(PixelColor) + dpv_hi(PixelColor);
-		OutColor.rgb += ColorAcc;
 	}
 	
 	#if DEBUG_HYBRID_SCATTERING_SPRITE == 1
diff --git a/Engine/Shaders/Private/DiaphragmDOF/DOFHybridScatterVertexShader.usf b/Engine/Shaders/Private/DiaphragmDOF/DOFHybridScatterVertexShader.usf
--- a/Engine/Shaders/Private/DiaphragmDOF/DOFHybridScatterVertexShader.usf
+++ b/Engine/Shaders/Private/DiaphragmDOF/DOFHybridScatterVertexShader.usf
@@ -76,12 +76,10 @@ void ScatterMainVS(
 	out nointerpolation float4 BokehCenterScreenAndClipPos : TEXCOORD0,
 	out nointerpolation float4 ColorAndAbsCocRadius0 : TEXCOORD1,
 	out nointerpolation float4 ColorAndAbsCocRadius1 : TEXCOORD2,
+	out nointerpolation uint3 GlobalSpriteIdAndPetzvalTransformPacked : TEXCOORD3,
 #if TAYLOR_EXPENDED_COC_INTERSECTION
-	out nointerpolation float4 TaylorExpensionMAD0   : TEXCOORD3,
-	out nointerpolation float4 TaylorExpensionMAD1   : TEXCOORD4,
-#endif
-#if CONFIG_PETZVAL
-	out nointerpolation float2 PetzvalTransformPacked : TEXCOORD5,
+	out nointerpolation float4 TaylorExpensionMAD0   : TEXCOORD4,
+	out nointerpolation float4 TaylorExpensionMAD1   : TEXCOORD5,
 #endif
 	out float4 OutPosition : SV_POSITION)
 {
@@ -108,7 +106,9 @@ void ScatterMainVS(
 	UNROLL
 	for (uint i = 0; i < PIXEL_COUNT_PER_SCATTER_GROUP; i++)
 	{
-		float4 E = ScatterDrawList[5 * GlobalSpriteId + i + 1];
+		uint ScatterDrawListOffset = SCATTER_DRAW_LIST_GROUP_STRIDE * GlobalSpriteId + 1 + (i * SCATTER_DRAW_LIST_DATA_PER_PIXEL);
+		float4 E = ScatterDrawList[ScatterDrawListOffset];
+		
 		Color[i] = float4(E.rgb, 0);
 		AbsCocRadius[i] = E.w;
 
@@ -141,7 +141,7 @@ void ScatterMainVS(
 	}
 
 	// Output interpolators.
-	const float2 BokehCenterScreenPos = ScatteringScaling * ScatterDrawList[5 * GlobalSpriteId].xy;
+	const float2 BokehCenterScreenPos = ScatteringScaling * ScatterDrawList[SCATTER_DRAW_LIST_GROUP_STRIDE * GlobalSpriteId].xy;
 	const float2 BokehCenterClipPos = (BokehCenterScreenPos * ViewportSize.zw) * 2.0 - 1.0;
 	{
 		ColorAndAbsCocRadius0 = float4(
@@ -170,12 +170,14 @@ void ScatterMainVS(
 		#endif
 	}
 
+	GlobalSpriteIdAndPetzvalTransformPacked = 0;
+	GlobalSpriteIdAndPetzvalTransformPacked.x = GlobalSpriteId;
 #if CONFIG_PETZVAL
 	float2x2 PetzvalTransform = CalcPetzvalTransform(BokehCenterClipPos);
 	half4 PetzvalTransformHalf = half4(float4(PetzvalTransform));
-	PetzvalTransformPacked = float2(
-		asfloat(PackFloat2ToUInt(PetzvalTransformHalf[0], PetzvalTransformHalf[1])),
-		asfloat(PackFloat2ToUInt(PetzvalTransformHalf[2], PetzvalTransformHalf[3]))
+	GlobalSpriteIdAndPetzvalTransformPacked.yz = uint2(
+		PackFloat2ToUInt(PetzvalTransformHalf[0], PetzvalTransformHalf[1]),
+		PackFloat2ToUInt(PetzvalTransformHalf[2], PetzvalTransformHalf[3])
 	);
 #endif
 
diff --git a/Engine/Shaders/Private/DiaphragmDOF/DOFReduce.usf b/Engine/Shaders/Private/DiaphragmDOF/DOFReduce.usf
--- a/Engine/Shaders/Private/DiaphragmDOF/DOFReduce.usf
+++ b/Engine/Shaders/Private/DiaphragmDOF/DOFReduce.usf
@@ -8,6 +8,7 @@
 
 #include "DOFDownsample.ush"
 #include "DOFHybridScatterCompilation.ush"
+#include "DOFVignetteCommon.ush"
 #include "../ReductionCommon.ush"
 
 
@@ -63,12 +64,21 @@
 //------------------------------------------------------- PARAMETERS
 
 uint4 ViewportRect;
+float4 ScatteringViewportSize;
+float ScatteringScaling;
 uint MaxScatteringGroupCount;
 float2 MaxInputBufferUV;
 float MinScatteringCocRadius;
 float PreProcessingToProcessingCocRadiusFactor;
 float BokehGatherDistinctionLimit;
 
+float2 SensorSize;
+float Aperture;
+float LensImageDistance;
+float BarrelRadius;
+float BarrelLength;
+uint4 MatteBoxPlanes[MAX_MATTE_BOX_FLAGS];
+
 float4 GatherInputSize;
 Texture2D GatherInput_SceneColor;
 Texture2D GatherInput_SeparateCoc;
@@ -263,6 +273,36 @@ float ApplyBokehGatherDistinctionFilter(float CocRadius, float4 Color, float Fra
 	}
 }
 
+void ComputeVignetteData(uint2 DispatchThreadId, float CocRadius, out float4 Output[VIGNETTE_DATA_PER_PIXEL])
+{
+	const float2 DrawPosition = (DispatchThreadId & 0xFFFFFFFE) + 0.5;
+	const float2 BokehCenterScreenPos = ScatteringScaling * DrawPosition;
+	const float2 BokehCenterClipPos = (BokehCenterScreenPos * ScatteringViewportSize.zw) * 2.0 - 1.0;
+	const float3 ObjectDirection = normalize(-float3(BokehCenterClipPos * (SensorSize*0.5), -LensImageDistance));
+
+	const float Focus = View.DepthOfFieldFocalDistance;
+	const float2 Barrel = asfloat(ProjectAndPackCylinder(
+		CocRadius, 
+		BokehCenterClipPos, 
+		BarrelRadius, 
+		BarrelLength, 
+		SensorSize, 
+		Focus, 
+		Aperture, 
+		CocInfinityRadius, 
+		LensImageDistance));
+
+	float2 ProjectedMatteBoxFlags[MAX_MATTE_BOX_FLAGS];
+	UNROLL
+	for (int Index = 0; Index < MAX_MATTE_BOX_FLAGS; ++Index)
+	{
+		ProjectedMatteBoxFlags[Index] = asfloat(ProjectAndPackMatteBoxFlag(ObjectDirection, BarrelRadius, BarrelLength, MatteBoxPlanes[Index]));
+	}
+
+	Output[0] = float4(Barrel, ProjectedMatteBoxFlags[0]);
+	Output[1] = float4(ProjectedMatteBoxFlags[1], ProjectedMatteBoxFlags[2]);
+}
+
 [numthreads(THREADGROUP_TILE_SIZE, THREADGROUP_TILE_SIZE, 1)]
 void ReduceCS(
 	uint2 GroupId : SV_GroupID,
@@ -509,6 +549,9 @@ void ReduceCS(
 		// TODO: no need min()
 		float IntensityLossFactor = min(1, SafeRcp(PI * CocRadius * CocRadius, 1.0) * CocSqueeze);
 
+		
+		
+
 		// Enqueue scattering group in foreground scatter draw list.
 		BRANCH
 		if (DIM_HYBRID_SCATTER_FGD && (ScatteringGroupMask & ForegroundScatterGroupMask))
@@ -522,11 +565,24 @@ void ReduceCS(
 				BRANCH
 				if (ScatteringGroupThreadIndex == 0)
 				{
-					OutForegroundScatterDrawList[5 * IndirectDrawIndex + 0] = float4(DispatchThreadId + 0.5, 0.0, 0.0);
+					OutForegroundScatterDrawList[SCATTER_DRAW_LIST_GROUP_STRIDE * IndirectDrawIndex + 0] = float4(DispatchThreadId + 0.5, 0.0, 0.0);
 				}
 				
 				float4 ScatterColor = GatheredColor * (ScatterFactor * IntensityLossFactor * (IsForeground(ScatterCocRadius) ? 1 : 0));
-				OutForegroundScatterDrawList[5 * IndirectDrawIndex + ScatteringGroupThreadIndex + 1] = float4(ScatterColor.rgb, abs(ScatterCocRadius));
+				uint Offset = SCATTER_DRAW_LIST_GROUP_STRIDE * IndirectDrawIndex + 1 + (ScatteringGroupThreadIndex * SCATTER_DRAW_LIST_DATA_PER_PIXEL);
+				OutForegroundScatterDrawList[Offset] = float4(ScatterColor.rgb, abs(ScatterCocRadius));
+
+				if (BarrelLength >= 0)
+				{
+					float4 VignetteData[VIGNETTE_DATA_PER_PIXEL];
+					ComputeVignetteData(DispatchThreadId, ScatterCocRadius, VignetteData);
+
+					UNROLL
+					for (uint Index = 0; Index < VIGNETTE_DATA_PER_PIXEL; ++Index)
+					{
+						OutForegroundScatterDrawList[Offset + 1 + Index] = VignetteData[Index];
+					}
+				}
 			}
 			else
 			{
@@ -548,11 +604,24 @@ void ReduceCS(
 				BRANCH
 				if (ScatteringGroupThreadIndex == 0)
 				{
-					OutBackgroundScatterDrawList[5 * IndirectDrawIndex + 0] = float4(DispatchThreadId + 0.5, 0.0, 0.0);
+					OutBackgroundScatterDrawList[SCATTER_DRAW_LIST_GROUP_STRIDE * IndirectDrawIndex + 0] = float4(DispatchThreadId + 0.5, 0.0, 0.0);
 				}
 				
 				float4 ScatterColor = GatheredColor * (ScatterFactor * IntensityLossFactor * (IsForeground(ScatterCocRadius) ? 0 : 1));
-				OutBackgroundScatterDrawList[5 * IndirectDrawIndex + ScatteringGroupThreadIndex + 1] = float4(ScatterColor.rgb, abs(ScatterCocRadius));
+				uint Offset = SCATTER_DRAW_LIST_GROUP_STRIDE * IndirectDrawIndex + 1 + (ScatteringGroupThreadIndex * SCATTER_DRAW_LIST_DATA_PER_PIXEL);
+				OutBackgroundScatterDrawList[Offset] = float4(ScatterColor.rgb, abs(ScatterCocRadius));
+
+				if (BarrelLength >= 0)
+				{
+					float4 VignetteData[VIGNETTE_DATA_PER_PIXEL];
+					ComputeVignetteData(DispatchThreadId, ScatterCocRadius, VignetteData);
+
+					UNROLL
+					for (uint Index = 0; Index < VIGNETTE_DATA_PER_PIXEL; ++Index)
+					{
+						OutBackgroundScatterDrawList[Offset + 1 + Index] = VignetteData[Index];
+					}
+				}
 			}
 			else
 			{
diff --git a/Engine/Shaders/Private/DiaphragmDOF/DOFVignetteCommon.ush b/Engine/Shaders/Private/DiaphragmDOF/DOFVignetteCommon.ush
new file mode 100644
--- /dev/null
+++ b/Engine/Shaders/Private/DiaphragmDOF/DOFVignetteCommon.ush
@@ -0,0 +1,209 @@
+// Copyright Epic Games, Inc. All Rights Reserved.
+
+/*=============================================================================
+	DiaphragmDOF/DOFVignetteCommon.ush: Various functions to calculate intersections
+	between DOF rays and close occluders such as the lens barrel or matte box.
+=============================================================================*/
+
+#include "/Engine/Public/DualPixelVectorization.ush"
+
+//------------------------------------------------------- BARREL
+
+/// Given a ray and a cylinder (barrel) around the origin, calculate the height on the barrel where the ray hits.
+float HitTestCylinder(float2 RayOrigin, float3 RayDirection, float BarrelRadius, float BarrelLength)
+{
+	const float2 P = RayOrigin;
+	const float3 D = RayDirection;
+
+	const float A = dot(D.xy, D.xy);
+	const float B = 2.0 * dot(D.xy, P.xy);
+	const float C = dot(P.xy, P.xy) - BarrelRadius*BarrelRadius;
+	const float Det = B * B - 4 * A * C;
+	const float SqrtDet = sqrt(Det);
+	const float Inv2A = rcp(2 * A);
+	const float T1 = (-B + SqrtDet) * Inv2A;
+	const float T2 = (-B - SqrtDet) * Inv2A;
+	const float L1 = T1 * D.z;
+	const float L2 = T2 * D.z;
+	const float L = min(T1 > 0 ? L1 : 1e20, T2 > 0 ? L2 : 1e20);
+	const float HitOnBarrel = BarrelLength - L;
+	return HitOnBarrel;
+}
+
+/// Project the aperture onto the end of a cylinder (lens barrel).
+/// Use the output for fast vignette occlusion testing.
+uint2 ProjectAndPackCylinder(
+	const float CocRadius,
+	const float2 BokehCenterClipPos,
+	const float BarrelRadius,
+	const float BarrelLength,
+	const float2 SensorSize,
+	const float Focus,
+	const float Aperture,
+	const float CocInfinityRadius,
+	const float LensImageDistance
+)
+{
+	const float SampleDepth = Focus / (1.0 - (CocRadius / CocInfinityRadius));
+	const float CocRadiusAtBarrelEnd = lerp(Aperture/2, 0, BarrelRadius/SampleDepth);
+	const float Factor = -1.0 * CocRadiusAtBarrelEnd / abs(CocRadius);
+
+	const float2 CocCenterAtBarrelEnd = BokehCenterClipPos * (SensorSize * 0.5) * (BarrelLength / LensImageDistance) * -1;
+
+	uint Packed0 = PackFloat2ToUInt(CocCenterAtBarrelEnd);
+	uint Packed1 = asuint(Factor);
+	uint2 Packed = uint2(Packed0, Packed1);
+
+	return Packed;
+}
+
+float HitTestProjectedCylinder(
+	const float2 PixelOffset, 
+	const float2 BarrelData, 
+	const float BarrelRadius,
+	out float2 SamplePositionAtBarrelEnd
+)
+{
+	const float2 CocCenterAtBarrelEnd = UnpackFloat2FromUInt(asuint(BarrelData.x));
+	const float Factor = BarrelData.y;
+
+	SamplePositionAtBarrelEnd = CocCenterAtBarrelEnd + /*sign(CocRadius) */ PixelOffset * Factor;
+	return dot(SamplePositionAtBarrelEnd, SamplePositionAtBarrelEnd) < BarrelRadius*BarrelRadius;
+}
+
+half2 HitTestProjectedCylinder_DPV(
+	const half2x2 PixelOffset, 
+	const uint2 BarrelData[2], 
+	const float BarrelRadius,
+	out half2x2 SamplePositionAtBarrelEnd
+)
+{
+	const half2x2 CocCenterAtBarrelEnd = dpv_interleave_registers(
+		half2(UnpackFloat2FromUInt(BarrelData[0].x)), 
+		half2(UnpackFloat2FromUInt(BarrelData[1].x))
+	);
+	const half2 Factor = dpv_interleave_registers(half(asfloat(BarrelData[0].y)), half(asfloat(BarrelData[1].y)));
+	
+	SamplePositionAtBarrelEnd = dpv_add(CocCenterAtBarrelEnd, dpv_mul(PixelOffset, Factor));
+	const half BarrelRadiusSqr = half(BarrelRadius*BarrelRadius);
+	return saturate(BarrelRadiusSqr - dpv_dot(SamplePositionAtBarrelEnd, SamplePositionAtBarrelEnd));
+}
+
+//------------------------------------------------------- MATTE BOX
+
+/// Test a ray for intersection with a single panel on a matte box at the end of the lens barrel.
+float HitTestMatteBoxFlag(float3 RayOrigin, float3 RayDirection, float BarrelRadius, float BarrelLength, uint4 FlagDefinition)
+{
+	float2 Unpacked0 = UnpackFloat2FromUInt(FlagDefinition.x);
+	float2 Unpacked1 = UnpackFloat2FromUInt(FlagDefinition.y);
+
+	float SinPlanePitch = Unpacked0.x;
+	float CosPlanePitch = Unpacked0.y;
+	float SinPlaneRoll = Unpacked1.x;
+	float CosPlaneRoll = Unpacked1.y;
+
+	float RcpSinPlanePitch = asfloat(FlagDefinition.z);
+	float PlaneLength = asfloat(FlagDefinition.w);
+
+	if (PlaneLength > 0)
+	{
+		float3 PlaneNormal = float3(CosPlaneRoll*SinPlanePitch, SinPlaneRoll*SinPlanePitch, CosPlanePitch);
+		float PlaneDist = SinPlanePitch*BarrelRadius + CosPlanePitch*BarrelLength;
+
+		float HitAngle = dot(RayDirection, PlaneNormal);
+		if (HitAngle != 0)
+		{
+			float Hit = (-dot(PlaneNormal, RayOrigin) + PlaneDist) / HitAngle;
+			if (Hit > 0)
+			{
+				float3 HitPoint = RayOrigin + Hit * RayDirection;
+
+				bool bIsHit = BarrelLength < HitPoint.z && HitPoint.z < (BarrelLength + SinPlanePitch*PlaneLength);
+				return !bIsHit;
+			}
+		}
+	}
+
+	return 1;
+}
+
+/// Project a matte box flag for onto the end of the lens barrel along a given object direction.
+/// Use the output for fast vignette occlusion testing.
+uint2 ProjectAndPackMatteBoxFlag(
+	const float3 ObjectDirection,
+	const float BarrelRadius,
+	const float BarrelLength,
+	const uint4 FlagDefinition
+)
+{
+	float2 Unpacked0 = UnpackFloat2FromUInt(FlagDefinition.x);
+	float2 Unpacked1 = UnpackFloat2FromUInt(FlagDefinition.y);
+
+	float SinPitch = Unpacked0.x;
+	float CosPitch = Unpacked0.y;
+	float SinRoll = Unpacked1.x;
+	float CosRoll = Unpacked1.y;
+
+	float RcpSinPlanePitch = asfloat(FlagDefinition.z);
+	float FlagLength = asfloat(FlagDefinition.w);
+
+	// Find point on the end of the flag
+	float3 FlagPoint = float3(
+		CosRoll * (BarrelRadius - CosPitch * FlagLength),
+		SinRoll * (BarrelRadius - CosPitch * FlagLength),
+		BarrelLength + FlagLength * SinPitch);
+
+	// Project point onto the barrel end
+	float Step = (FlagLength * SinPitch) / ObjectDirection.z;
+	float2 BarrelPoint = (FlagPoint - ObjectDirection * Step).xy;
+
+	// Find projected flag edge direction
+	float3 FlagPlaneNormal = float3(
+		CosRoll * SinPitch,
+		SinRoll * SinPitch,
+		CosPitch);
+	float2 EdgeDirection = float2(FlagPlaneNormal.y, -FlagPlaneNormal.x); // cross(FlagPlaneNormal, (0,0,1))
+
+	uint Packed0 = PackFloat2ToUInt(BarrelPoint);
+	uint Packed1 = PackFloat2ToUInt(EdgeDirection);
+	uint2 Packed = uint2(Packed0, Packed1);
+	
+	return Packed;
+}
+
+float HitTestProjectedMatteBoxFlag(
+	const float2 SamplePositionAtBarrelEnd, 
+	const uint2 PackedProjectedMatteBoxFlag
+)
+{
+	float2 BarrelPoint   = UnpackFloat2FromUInt(PackedProjectedMatteBoxFlag.x);
+	float2 EdgeDirection = UnpackFloat2FromUInt(PackedProjectedMatteBoxFlag.y);
+
+	// Check if sample point is on the left or right side of the line (BarrelPoint + t * EdgeDirection)
+	float2 Delta = SamplePositionAtBarrelEnd - BarrelPoint;
+	float Angle = Delta.x * EdgeDirection.y - Delta.y * EdgeDirection.x;
+
+	return Angle > 0;
+}
+
+half2 HitTestProjectedMatteBoxFlag_DPV(
+	const half2x2 SamplePositionAtBarrelEnd, 
+	const uint2 PackedProjectedMatteBoxFlag[2]
+)
+{
+	half2x2 BarrelPoint = dpv_interleave_registers(
+		(half2)UnpackFloat2FromUInt(PackedProjectedMatteBoxFlag[0].x),
+		(half2)UnpackFloat2FromUInt(PackedProjectedMatteBoxFlag[1].x)
+	);
+	half2x2 EdgeDirection = dpv_interleave_registers(
+		(half2)UnpackFloat2FromUInt(PackedProjectedMatteBoxFlag[0].y),
+		(half2)UnpackFloat2FromUInt(PackedProjectedMatteBoxFlag[1].y)
+	);
+
+	// Check if sample point is on the left or right side of the line (BarrelPoint + t * EdgeDirection)
+	half2x2 Delta = dpv_sub(SamplePositionAtBarrelEnd, BarrelPoint);
+	half2 Angle = dpv_cross(Delta, EdgeDirection);
+
+	const half Sharpness = 128.0;
+	return saturate(Angle * Sharpness); // Angle > 0
+}
diff --git a/Engine/Shaders/Public/DualPixelVectorization.ush b/Engine/Shaders/Public/DualPixelVectorization.ush
--- a/Engine/Shaders/Public/DualPixelVectorization.ush
+++ b/Engine/Shaders/Public/DualPixelVectorization.ush
@@ -133,6 +133,9 @@
 		DEFINE_DPV_FUNCTION(float16_t, 3) \
 		DEFINE_DPV_FUNCTION(float16_t, 4) \
 
+	#define DEFINE_DPV_FUNCTION_FOR_HALF_VECTOR2(DEFINE_DPV_FUNCTION) \
+		DEFINE_DPV_FUNCTION(float16_t, 2) \
+
 	#define DEFINE_DPV_FUNCTION_FOR_SHORT_VECTORS(DEFINE_DPV_FUNCTION) \
 		DEFINE_DPV_FUNCTION(uint16_t, 2) \
 		DEFINE_DPV_FUNCTION(uint16_t, 3) \
@@ -149,6 +152,8 @@
 		
 	#define DEFINE_DPV_FUNCTION_FOR_HALF_VECTORS(DEFINE_DPV_FUNCTION)
 
+	#define DEFINE_DPV_FUNCTION_FOR_HALF_VECTOR2(DEFINE_DPV_FUNCTION)
+
 	#define DEFINE_DPV_FUNCTION_FOR_SHORT_VECTORS(DEFINE_DPV_FUNCTION)
 		
 #endif
@@ -421,6 +426,22 @@ DEFINE_DPV_FUNCTION_FOR_ALL_VECTORS(DEFINE_DPV_FUNCTION_VECTOR)
 #undef DEFINE_DPV_FUNCTION_VECTOR
 
 
+// 2D cross product.
+#define DEFINE_DPV_FUNCTION_VECTOR(TYPE, COMPONENTS) \
+	CALL_SITE_DEBUGLOC __DPV_VECTOR(TYPE, COMPONENTS) dpv_cross(__DPV_MULTI_VECTOR(TYPE, COMPONENTS) a, __DPV_MULTI_VECTOR(TYPE, COMPONENTS) b) { \
+		__DPV_VECTOR(TYPE, COMPONENTS) lo = dpv_lo(b); \
+		__DPV_VECTOR(TYPE, COMPONENTS) lo_rotated = __DPV_VECTOR(TYPE, COMPONENTS)(lo.y, -lo.x); \
+		__DPV_VECTOR(TYPE, COMPONENTS) hi = dpv_hi(b); \
+		__DPV_VECTOR(TYPE, COMPONENTS) hi_rotated = __DPV_VECTOR(TYPE, COMPONENTS)(hi.y, -hi.x); \
+		__DPV_MULTI_VECTOR(TYPE, COMPONENTS) rotated = dpv_interleave_registers(lo_rotated, hi_rotated); \
+		return dpv_dot(a, rotated); \
+	}
+
+DEFINE_DPV_FUNCTION_FOR_FLOAT_VECTOR(DEFINE_DPV_FUNCTION_VECTOR, 2)
+DEFINE_DPV_FUNCTION_FOR_HALF_VECTOR2(DEFINE_DPV_FUNCTION_VECTOR)
+#undef DEFINE_DPV_FUNCTION_VECTOR
+
+
 /* per vector vectorial operations
  *
  * Example:
diff --git a/Engine/Source/Runtime/Engine/Classes/Engine/Scene.h b/Engine/Source/Runtime/Engine/Classes/Engine/Scene.h
--- a/Engine/Source/Runtime/Engine/Classes/Engine/Scene.h
+++ b/Engine/Source/Runtime/Engine/Classes/Engine/Scene.h
@@ -581,6 +581,30 @@ struct FCameraExposureSettings
 	ENGINE_API void ExportToPostProcessSettings(FPostProcessSettings* OutPostProcessSettings) const;
 };
 
+USTRUCT(BlueprintType)
+struct FMatteBoxFlag
+{
+	GENERATED_USTRUCT_BODY()
+
+	/** Angle of the panel with the lens barrel */
+	UPROPERTY(interp, BlueprintReadWrite, Category=FMatteBoxFlag, meta=(ClampMin = "-90.0", ClampMax = "90.0"))
+	float Pitch;
+
+	/** Rotation of the panel around the lens barrel */
+	UPROPERTY(interp, BlueprintReadWrite, Category=FMatteBoxFlag, meta=(ClampMin = "0.0", ClampMax = "360.0"))
+	float Roll;
+
+	/** How long is the panel */
+	UPROPERTY(interp, BlueprintReadWrite, Category=FMatteBoxFlag, meta=(ClampMin = "0.0", UIMax = "100.0"))
+	float Length;
+
+	FMatteBoxFlag()
+	: Pitch(0)
+	, Roll(0)
+	, Length(0)
+	{ }
+};
+
 USTRUCT(BlueprintType)
 struct FWeightedBlendable
 {
@@ -1121,6 +1145,15 @@ struct FPostProcessSettings
 	UPROPERTY(EditAnywhere, BlueprintReadWrite, Category=Overrides, meta=(PinHiddenByDefault, InlineEditConditionToggle))
 	uint8 bOverride_DepthOfFieldAspectRatioScalar:1;
 	
+	UPROPERTY(EditAnywhere, BlueprintReadWrite, Category=Overrides, meta=(PinHiddenByDefault, InlineEditConditionToggle))
+	uint8 bOverride_DepthOfFieldMatteBoxFlags:1;
+
+	UPROPERTY(EditAnywhere, BlueprintReadWrite, Category=Overrides, meta=(PinHiddenByDefault, InlineEditConditionToggle))
+	uint8 bOverride_DepthOfFieldBarrelRadius:1;
+	
+	UPROPERTY(EditAnywhere, BlueprintReadWrite, Category=Overrides, meta=(PinHiddenByDefault, InlineEditConditionToggle))
+	uint8 bOverride_DepthOfFieldBarrelLength:1;
+
 	UPROPERTY(EditAnywhere, BlueprintReadWrite, Category=Overrides, meta=(PinHiddenByDefault, InlineEditConditionToggle))
 	uint8 bOverride_DepthOfFieldDepthBlurAmount:1;
 
@@ -2157,6 +2190,18 @@ struct FPostProcessSettings
 	/** Amount to scale the output viewport's aspect ratio by when computing depth of field properties.  */
 	UPROPERTY(Interp, BlueprintReadWrite, Category="Lens|Depth of Field", meta=(ClampMin = "0.0", editcondition = "bOverride_DepthOfFieldAspectRatioScalar", DisplayName = "Aspect Ratio Scalar"))
 	float DepthOfFieldAspectRatioScalar;
+		
+	/** The lens barrel creates vignetting and occludes bokeh, i.e. cat's eye bokeh */
+	UPROPERTY(interp, EditAnywhere, BlueprintReadWrite, Category="Lens|Depth of Field", meta=(ClampMin = "0.0", UIMax = "100.0", editcondition = "bOverride_DepthOfFieldBarrelRadius", DisplayName = "Barrel Radius"))
+	float DepthOfFieldBarrelRadius;
+
+	/** The lens barrel creates vignetting and occludes bokeh, i.e. cat's eye bokeh */
+	UPROPERTY(interp, EditAnywhere, BlueprintReadWrite, Category="Lens|Depth of Field", meta=(ClampMin = "0.0", UIMax = "150.0", editcondition = "bOverride_DepthOfFieldBarrelLength", DisplayName = "Barrel Length"))
+	float DepthOfFieldBarrelLength;
+
+	/** Panels around the front of the lens barrel that occlude bokeh */
+	UPROPERTY(EditAnywhere, Category="Lens|Depth of Field", meta=(editcondition = "bOverride_DepthOfFieldMatteBoxFlags", DisplayName = "Matte Box Flags"))
+	FMatteBoxFlag DepthOfFieldMatteBoxFlags[3];
 	
 	/** Artificial region where all content is in focus, starting after DepthOfFieldFocalDistance, in unreal units  (cm) */
 	UPROPERTY(interp, EditAnywhere, BlueprintReadWrite, Category="Lens|Mobile Depth of Field", meta=(UIMin = "0.0", UIMax = "10000.0", editcondition = "bOverride_DepthOfFieldFocalRegion", DisplayName = "Focal Region"))
diff --git a/Engine/Source/Runtime/Engine/Private/PostProcessUtils.cpp b/Engine/Source/Runtime/Engine/Private/PostProcessUtils.cpp
--- a/Engine/Source/Runtime/Engine/Private/PostProcessUtils.cpp
+++ b/Engine/Source/Runtime/Engine/Private/PostProcessUtils.cpp
@@ -149,6 +149,14 @@ bool FPostProcessUtils::OverridePostProcessSettings(FPostProcessSettings& ThisFr
 		UE_SET_PP(DepthOfFieldDepthBlurRadius);
 		UE_SET_PP(DepthOfFieldUseHairDepth)
 		UE_SET_PP(DepthOfFieldDepthBlurAmount);
+		if (OtherTo.bOverride_DepthOfFieldMatteBoxFlags)
+		{
+			for (uint32 Index = 0; Index < UE_ARRAY_COUNT(OtherTo.DepthOfFieldMatteBoxFlags); ++Index)
+			{
+				ThisFrom.DepthOfFieldMatteBoxFlags[Index] = OtherTo.DepthOfFieldMatteBoxFlags[Index];
+			}
+			bAnyOverwritten = true;
+		}
 		UE_SET_PP(DepthOfFieldFocalRegion);
 		UE_SET_PP(DepthOfFieldNearTransitionRegion);
 		UE_SET_PP(DepthOfFieldFarTransitionRegion);
@@ -159,6 +167,21 @@ bool FPostProcessUtils::OverridePostProcessSettings(FPostProcessSettings& ThisFr
 		UE_SET_PP(DepthOfFieldSkyFocusDistance);
 		UE_SET_PP(DepthOfFieldVignetteSize);
 		UE_SET_PP(DepthOfFieldAspectRatioScalar);
+		UE_SET_PP(DepthOfFieldPetzvalBokeh);
+		UE_SET_PP(DepthOfFieldPetzvalBokehFalloff);
+		UE_SET_PP(DepthOfFieldPetzvalExclusionBoxExtents);
+		UE_SET_PP(DepthOfFieldPetzvalExclusionBoxRadius);
+		UE_SET_PP(DepthOfFieldBarrelRadius);
+		UE_SET_PP(DepthOfFieldBarrelLength);
+		if (OtherTo.bOverride_DepthOfFieldMatteBoxFlags)
+		{
+			for (uint32 Index = 0; Index < UE_ARRAY_COUNT(OtherTo.DepthOfFieldMatteBoxFlags); ++Index)
+			{
+				ThisFrom.DepthOfFieldMatteBoxFlags[Index] = OtherTo.DepthOfFieldMatteBoxFlags[Index];
+			}
+			bAnyOverwritten = true;
+		}
+
 		UE_SET_PP(MotionBlurAmount);
 		UE_SET_PP(MotionBlurMax);
 		UE_SET_PP(MotionBlurPerObjectSize);
@@ -440,6 +463,20 @@ bool FPostProcessUtils::BlendPostProcessSettings(FPostProcessSettings& ThisFrom,
 		UE_LERP_PP(DepthOfFieldDepthBlurRadius);
 		UE_SET_PP(DepthOfFieldUseHairDepth)
 		UE_LERP_PP(DepthOfFieldDepthBlurAmount);
+		UE_LERP_PP(DepthOfFieldPetzvalBokeh);
+		UE_LERP_PP(DepthOfFieldPetzvalBokehFalloff);
+		UE_LERP_PP(DepthOfFieldPetzvalExclusionBoxExtents);
+		UE_LERP_PP(DepthOfFieldPetzvalExclusionBoxRadius);
+		UE_LERP_PP(DepthOfFieldBarrelRadius);
+		UE_LERP_PP(DepthOfFieldBarrelLength);
+		if (OtherTo.bOverride_DepthOfFieldMatteBoxFlags)
+		{
+			for (uint32 Index = 0; Index < UE_ARRAY_COUNT(OtherTo.DepthOfFieldMatteBoxFlags); ++Index)
+			{
+				ThisFrom.DepthOfFieldMatteBoxFlags[Index] = OtherTo.DepthOfFieldMatteBoxFlags[Index];
+			}
+			bAnyOverwritten = true;
+		}
 		UE_LERP_PP(DepthOfFieldFocalRegion);
 		UE_LERP_PP(DepthOfFieldNearTransitionRegion);
 		UE_LERP_PP(DepthOfFieldFarTransitionRegion);
diff --git a/Engine/Source/Runtime/Engine/Private/Scene.cpp b/Engine/Source/Runtime/Engine/Private/Scene.cpp
--- a/Engine/Source/Runtime/Engine/Private/Scene.cpp
+++ b/Engine/Source/Runtime/Engine/Private/Scene.cpp
@@ -570,6 +570,12 @@ FPostProcessSettings::FPostProcessSettings()
 	DepthOfFieldPetzvalExclusionBoxExtents = FVector2f::ZeroVector;
 	DepthOfFieldPetzvalExclusionBoxRadius = 0.0f;
 	DepthOfFieldAspectRatioScalar = 1.0f;
+	DepthOfFieldBarrelRadius = 5;
+	DepthOfFieldBarrelLength = 0;
+	for (int Index = 0; Index < UE_ARRAY_COUNT(DepthOfFieldMatteBoxFlags); ++Index)
+	{
+		DepthOfFieldMatteBoxFlags[Index] = {};
+	}
 	DepthOfFieldFocalRegion = 0.0f;
 	DepthOfFieldNearTransitionRegion = 300.0f;
 	DepthOfFieldFarTransitionRegion = 500.0f;
diff --git a/Engine/Source/Runtime/Engine/Private/SceneView.cpp b/Engine/Source/Runtime/Engine/Private/SceneView.cpp
--- a/Engine/Source/Runtime/Engine/Private/SceneView.cpp
+++ b/Engine/Source/Runtime/Engine/Private/SceneView.cpp
@@ -1720,6 +1720,19 @@ void FSceneView::OverridePostProcessSettings(const FPostProcessSettings& Src, fl
 		LERP_PP(DepthOfFieldPetzvalBokehFalloff);
 		LERP_PP(DepthOfFieldPetzvalExclusionBoxExtents);
 		LERP_PP(DepthOfFieldPetzvalExclusionBoxRadius);
+		LERP_PP(DepthOfFieldBarrelRadius);
+		LERP_PP(DepthOfFieldBarrelLength);
+		if (Src.bOverride_DepthOfFieldMatteBoxFlags)
+		{
+			for (uint32 Index = 0; Index < UE_ARRAY_COUNT(Src.DepthOfFieldMatteBoxFlags); ++Index)
+			{
+				const FMatteBoxFlag& SrcFlag = Src.DepthOfFieldMatteBoxFlags[Index];
+				FMatteBoxFlag& DestFlag = Dest.DepthOfFieldMatteBoxFlags[Index];
+				DestFlag.Pitch = FMath::Lerp(DestFlag.Pitch, SrcFlag.Pitch, Weight);
+				DestFlag.Roll = FMath::Lerp(DestFlag.Roll, SrcFlag.Roll, Weight);
+				DestFlag.Length = FMath::Lerp(DestFlag.Length, SrcFlag.Length, Weight);
+			}
+		}
 		LERP_PP(DepthOfFieldDepthBlurAmount);
 		LERP_PP(DepthOfFieldFocalRegion);
 		LERP_PP(DepthOfFieldNearTransitionRegion);
diff --git a/Engine/Source/Runtime/Renderer/Private/PostProcess/DiaphragmDOF.cpp b/Engine/Source/Runtime/Renderer/Private/PostProcess/DiaphragmDOF.cpp
--- a/Engine/Source/Runtime/Renderer/Private/PostProcess/DiaphragmDOF.cpp
+++ b/Engine/Source/Runtime/Renderer/Private/PostProcess/DiaphragmDOF.cpp
@@ -314,6 +314,8 @@ const int32 kMaxMipLevelCount = 4;
 /** Minimum number of ring. */
 const int32 kMinGatheringRingCount = 3;
 
+const int32 kMaxGatheringRingCount = 5;
+
 /** Maximum number of ring for slight out of focus pass. Same as USH's MAX_RECOMBINE_ABS_COC_RADIUS. */
 const int32 kMaxSlightOutOfFocusCocRadius = 3;
 
@@ -475,6 +477,7 @@ public:
 		const bool bIsConsole = FDataDrivenShaderPlatformInfo::GetIsConsole(Parameters.Platform);
 		OutEnvironment.SetDefine(TEXT("WITH_DYNAMIC_COC_OFFSET"), bIsConsole ? 0 : 1);
 
+		OutEnvironment.SetDefine(TEXT("MAX_MATTE_BOX_FLAGS"), (int)DiaphragmDOF::FPhysicalCocModel::MaxMatteBoxFlags);
 	}
 
 	FDiaphragmDOFShader() {}
@@ -530,7 +533,7 @@ class FDDOFDilateRadiusDim     : SHADER_PERMUTATION_RANGE_INT("DIM_DILATE_RADIUS
 class FDDOFDilateModeDim       : SHADER_PERMUTATION_ENUM_CLASS("DIM_DILATE_MODE", EDiaphragmDOFDilateCocMode);
 
 class FDDOFLayerProcessingDim  : SHADER_PERMUTATION_ENUM_CLASS("DIM_LAYER_PROCESSING", EDiaphragmDOFLayerProcessing);
-class FDDOFGatherRingCountDim  : SHADER_PERMUTATION_RANGE_INT("DIM_GATHER_RING_COUNT", kMinGatheringRingCount, 3);
+class FDDOFGatherRingCountDim  : SHADER_PERMUTATION_RANGE_INT("DIM_GATHER_RING_COUNT", kMinGatheringRingCount, kMaxGatheringRingCount - kMinGatheringRingCount + 1);
 class FDDOFGatherQualityDim    : SHADER_PERMUTATION_ENUM_CLASS("DIM_GATHER_QUALITY", EDiaphragmDOFGatherQuality);
 class FDDOFPostfilterMethodDim : SHADER_PERMUTATION_ENUM_CLASS("DIM_POSTFILTER_METHOD", EDiaphragmDOFPostfilterMethod);
 class FDDOFRGBColorBufferDim   : SHADER_PERMUTATION_BOOL("DIM_RGB_COLOR_BUFFER");
@@ -748,7 +751,7 @@ static FORCEINLINE int32 MaxGatheringRingCount(EShaderPlatform ShaderPlatform)
 {
 	if (IsPCPlatform(ShaderPlatform))
 	{
-		return 5;
+		return kMaxGatheringRingCount;
 	}
 	return 4;
 }
@@ -949,12 +952,23 @@ class FDiaphragmDOFReduceCS : public FDiaphragmDOFShader
 	BEGIN_SHADER_PARAMETER_STRUCT(FParameters, )
 		SHADER_PARAMETER(FIntRect, ViewportRect)
 		SHADER_PARAMETER(FVector2f, MaxInputBufferUV)
+		SHADER_PARAMETER(FVector4f, ScatteringViewportSize)
+		SHADER_PARAMETER(float, ScatteringScaling)
 		SHADER_PARAMETER(int32, MaxScatteringGroupCount)
 		SHADER_PARAMETER(float, PreProcessingToProcessingCocRadiusFactor)
 		SHADER_PARAMETER(float, BokehGatherDistinctionLimit)
 		SHADER_PARAMETER(float, MinScatteringCocRadius)
 		SHADER_PARAMETER(float, NeighborCompareMaxColor)
-		SHADER_PARAMETER(float, CocSqueeze)
+
+		SHADER_PARAMETER_STRUCT_INCLUDE(FDOFCocModelShaderParameters, CocModel)
+
+		SHADER_PARAMETER(FVector2f, SensorSize)
+		SHADER_PARAMETER(float, Aperture)
+		SHADER_PARAMETER(float, LensImageDistance)
+	
+		SHADER_PARAMETER(float, BarrelRadius)
+		SHADER_PARAMETER(float, BarrelLength)
+		SHADER_PARAMETER_ARRAY(FUintVector4, MatteBoxPlanes, [DiaphragmDOF::FPhysicalCocModel::MaxMatteBoxFlags])
 		
 		SHADER_PARAMETER_RDG_BUFFER_SRV(StructuredBuffer<float4>, EyeAdaptationBuffer)
 		SHADER_PARAMETER_STRUCT_INCLUDE(FDOFCommonShaderParameters, CommonParameters)
@@ -1255,8 +1269,6 @@ BEGIN_SHADER_PARAMETER_STRUCT(FDOFHybridScatterParameters, )
 	SHADER_PARAMETER(FVector4f, ViewportSize)
 	SHADER_PARAMETER(float, CocRadiusToCircumscribedRadius)
 	SHADER_PARAMETER(float, ScatteringScaling)
-	SHADER_PARAMETER(float, CocSqueeze)
-	SHADER_PARAMETER(float, CocInvSqueeze)
 	
 	SHADER_PARAMETER(float, Petzval)
 	SHADER_PARAMETER(float, PetzvalFalloffPower)
@@ -1264,6 +1276,7 @@ BEGIN_SHADER_PARAMETER_STRUCT(FDOFHybridScatterParameters, )
 	SHADER_PARAMETER(float, PetzvalExclusionBoxRadius)
 	
 	SHADER_PARAMETER_STRUCT_INCLUDE(FDOFCommonShaderParameters, CommonParameters)
+	SHADER_PARAMETER_STRUCT_INCLUDE(FDOFCocModelShaderParameters, CocModel)
 	
 	SHADER_PARAMETER(FVector4f, ScatterOcclusionSize)
 	SHADER_PARAMETER_RDG_TEXTURE(Texture2D, ScatterOcclusion)
@@ -1273,6 +1286,14 @@ BEGIN_SHADER_PARAMETER_STRUCT(FDOFHybridScatterParameters, )
 	SHADER_PARAMETER_RDG_BUFFER_SRV(StructuredBuffer<float4>, ScatterDrawList)
 	RDG_BUFFER_ACCESS(IndirectDrawParameter, ERHIAccess::IndirectArgs)
 
+	SHADER_PARAMETER(FVector2f, SensorSize)
+	SHADER_PARAMETER(float, Aperture)
+	SHADER_PARAMETER(float, LensImageDistance)
+	
+	SHADER_PARAMETER(float, BarrelRadius)
+	SHADER_PARAMETER(float, BarrelLength)
+	SHADER_PARAMETER_ARRAY(FUintVector4, MatteBoxPlanes, [DiaphragmDOF::FPhysicalCocModel::MaxMatteBoxFlags])
+
 	RENDER_TARGET_BINDING_SLOTS()
 END_SHADER_PARAMETER_STRUCT()
 
@@ -1627,6 +1648,7 @@ bool DiaphragmDOF::AddPasses(
 	const bool bEnableScatterBokehSettings = CVarEnableScatterBokehSettings.GetValueOnRenderThread() == 1;
 	const bool bEnableSlightOutOfFocusBokeh = bSupportGatheringBokehSimulation && bRecombineDoesSlightOutOfFocus && CVarEnableRecombineBokehSettings.GetValueOnRenderThread();
 
+	bool bEnableVignette = true;
 		
 	// Setup all the descriptors.
 	FRDGTextureDesc FullResDesc;
@@ -1978,6 +2000,35 @@ bool DiaphragmDOF::AddPasses(
 		}
 	}
 
+	bool bMatteBoxActive = false;
+	TStaticArray<FUintVector4, FPhysicalCocModel::MaxMatteBoxFlags, SHADER_PARAMETER_ARRAY_ELEMENT_ALIGNMENT> MatteBoxPlanes;
+
+	for (int Index = 0; Index < FPhysicalCocModel::MaxMatteBoxFlags; ++Index)
+	{
+		float PlanePitch = (CocModel.MatteBoxFlags[Index].Pitch + 90.0f) * PI / 180.0;
+		float PlaneRoll = CocModel.MatteBoxFlags[Index].Roll * PI / 180.0;
+		float PlaneLength = CocModel.MatteBoxFlags[Index].Length;
+
+		float SinPlanePitch;
+		float CosPlanePitch;
+		float SinPlaneRoll;
+		float CosPlaneRoll;
+		float RcpSinPlanePitch;
+
+		FMath::SinCos(&SinPlanePitch, &CosPlanePitch, PlanePitch);
+		FMath::SinCos(&SinPlaneRoll, &CosPlaneRoll, PlaneRoll);
+		RcpSinPlanePitch = 1.0f / SinPlanePitch;
+
+		uint32 Packed0 = FVector2DHalf(SinPlanePitch, CosPlanePitch).AsUInt32();
+		uint32 Packed1 = FVector2DHalf(SinPlaneRoll, CosPlaneRoll).AsUInt32();
+		uint32 Packed2 = *reinterpret_cast<uint32*>(&RcpSinPlanePitch);
+		uint32 Packed3 = *reinterpret_cast<uint32*>(&PlaneLength);
+
+		MatteBoxPlanes[Index] = FUintVector4(Packed0, Packed1, Packed2, Packed3);
+		bMatteBoxActive = bMatteBoxActive || PlaneLength > 0;
+	}
+	bool bVignetteActive = CocModel.BarrelLength > 0 || bMatteBoxActive;
+
 	// Add the reduce pass
 	FDOFGatherInputTextures ReducedGatherInputTextures;
 	FRDGBufferRef DrawIndirectParametersBuffer = nullptr;
@@ -2072,7 +2123,12 @@ bool DiaphragmDOF::AddPasses(
 		{
 			DrawIndirectParametersBuffer = GraphBuilder.CreateBuffer(FRDGBufferDesc::CreateIndirectDesc<FRHIDrawIndexedIndirectParameters>(2), TEXT("DOFIndirectDrawParameters"));
 
-			FRDGBufferDesc DrawListDescs = FRDGBufferDesc::CreateStructuredDesc(sizeof(float) * 4, 5 * MaxScatteringGroupCount);
+			const uint32 PixelCountPerScatterGroup = 2 * 2;
+			const uint32 VignetteDataPerPixel = 3;
+			const uint32 ScatterDrawListDataPerPixel = 1 + VignetteDataPerPixel;
+			const uint32 ScatterDrawListGroupStride = 1 + (PixelCountPerScatterGroup * ScatterDrawListDataPerPixel);
+
+			FRDGBufferDesc DrawListDescs = FRDGBufferDesc::CreateStructuredDesc(sizeof(float) * 4, ScatterDrawListGroupStride * MaxScatteringGroupCount);
 			if (bForegroundHybridScattering)
 				ForegroundScatterDrawListBuffer = GraphBuilder.CreateBuffer(DrawListDescs, TEXT("DOF.ForegroundDrawList"));
 			if (bBackgroundHybridScattering)
@@ -2100,12 +2156,21 @@ bool DiaphragmDOF::AddPasses(
 			PassParameters->MaxInputBufferUV = FVector2f(
 				(PreprocessViewSize.X - 0.5f) / SrcSize.X,
 				(PreprocessViewSize.Y - 0.5f) / SrcSize.Y);
+			PassParameters->ScatteringViewportSize = FVector4f(GatheringViewSize.X, GatheringViewSize.Y, 1.0f / GatheringViewSize.X, 1.0f / GatheringViewSize.Y);
+			PassParameters->ScatteringScaling = float(GatheringViewSize.X) / float(PreprocessViewSize.X);
 			PassParameters->MaxScatteringGroupCount = MaxScatteringGroupCount;
 			PassParameters->PreProcessingToProcessingCocRadiusFactor = GatheringCocRadiusBasis / EncodedCocRadiusBasis;
 			PassParameters->BokehGatherDistinctionLimit = CVarGatherDistinctionLimit.GetValueOnRenderThread();
 			PassParameters->MinScatteringCocRadius = MinScatteringCocRadius;
 			PassParameters->NeighborCompareMaxColor = CVarScatterNeighborCompareMaxColor.GetValueOnRenderThread();
-			PassParameters->CocSqueeze = CocModel.Squeeze;
+			SetCocModelParameters(GraphBuilder, &PassParameters->CocModel, CocModel, float(GatheringViewSize.X));
+
+			PassParameters->LensImageDistance = 1.0f / (1.0f/CocModel.VerticalFocalLength - 1.0f/CocModel.FocusDistance);
+			PassParameters->SensorSize = FVector2f(CocModel.SensorWidth, CocModel.SensorHeight);
+			PassParameters->Aperture = CocModel.GetLensRadius().Y * 2.0f;
+			PassParameters->BarrelRadius = CocModel.BarrelRadius;
+			PassParameters->BarrelLength = bVignetteActive ? CocModel.BarrelLength : -1.0f;
+			PassParameters->MatteBoxPlanes = MatteBoxPlanes;
 			
 			PassParameters->EyeAdaptationBuffer = GraphBuilder.CreateSRV(GetEyeAdaptationBuffer(GraphBuilder, View));
 			PassParameters->CommonParameters = CommonParameters;
@@ -2609,13 +2674,12 @@ bool DiaphragmDOF::AddPasses(
 			PassParameters->ViewportSize = FVector4f(GatheringViewSize.X, GatheringViewSize.Y, 1.0f / GatheringViewSize.X, 1.0f / GatheringViewSize.Y);
 			PassParameters->CocRadiusToCircumscribedRadius = BokehModel.CocRadiusToCircumscribedRadius;
 			PassParameters->ScatteringScaling = float(GatheringViewSize.X) / float(PreprocessViewSize.X);
-			PassParameters->CocSqueeze = CocModel.Squeeze;
-			PassParameters->CocInvSqueeze = 1.0 / CocModel.Squeeze;
 			PassParameters->Petzval = CocModel.Petzval;
 			PassParameters->PetzvalFalloffPower = CocModel.PetzvalFalloffPower;
 			PassParameters->PetzvalExclusionBoxExtents = CocModel.PetzvalExclusionBoxExtents;
 			PassParameters->PetzvalExclusionBoxRadius = CocModel.PetzvalExclusionBoxRadius;
 			PassParameters->CommonParameters = CommonParameters;
+			SetCocModelParameters(GraphBuilder, &PassParameters->CocModel, CocModel, float(GatheringViewSize.X));
 			PassParameters->BokehLUT = ScatteringBokehLUT;
 			PassParameters->ScatterOcclusionSize = FVector4f(RefBufferSize.X, RefBufferSize.Y, 1.0f / RefBufferSize.X, 1.0f / RefBufferSize.Y);
 			PassParameters->ScatterOcclusion = ScatterOcclusionTexture;
@@ -2624,7 +2688,13 @@ bool DiaphragmDOF::AddPasses(
 			PassParameters->RenderTargets[0] = FRenderTargetBinding(
 				ConvolutionTextures->SceneColor, ERenderTargetLoadAction::ELoad);
 
+			PassParameters->LensImageDistance = 1.0f / (1.0f/CocModel.VerticalFocalLength - 1.0f/CocModel.FocusDistance);
+			PassParameters->SensorSize = FVector2f(CocModel.SensorWidth, CocModel.SensorHeight);
+			PassParameters->Aperture = CocModel.GetLensRadius().Y * 2.0f;
 				
+			PassParameters->BarrelRadius = CocModel.BarrelRadius;
+			PassParameters->BarrelLength = bVignetteActive ? CocModel.BarrelLength : -1.0f;
+			PassParameters->MatteBoxPlanes = MatteBoxPlanes;
 
 			ClearUnusedGraphResources(VertexShader, PixelShader, PassParameters);
 
diff --git a/Engine/Source/Runtime/Renderer/Private/PostProcess/DiaphragmDOF.h b/Engine/Source/Runtime/Renderer/Private/PostProcess/DiaphragmDOF.h
--- a/Engine/Source/Runtime/Renderer/Private/PostProcess/DiaphragmDOF.h
+++ b/Engine/Source/Runtime/Renderer/Private/PostProcess/DiaphragmDOF.h
@@ -72,6 +72,11 @@ struct FPhysicalCocModel
 	float PetzvalFalloffPower;
 	FVector2f PetzvalExclusionBoxExtents;
 	float PetzvalExclusionBoxRadius;
+	float BarrelRadius;
+	float BarrelLength;
+
+	static const uint32 MaxMatteBoxFlags = 3;
+	TStaticArray<FMatteBoxFlag, MaxMatteBoxFlags> MatteBoxFlags;
 
 	// The maximum radius of depth blur.
 	float MaxDepthBlurRadius;
diff --git a/Engine/Source/Runtime/Renderer/Private/PostProcess/DiaphragmDOFUtils.cpp b/Engine/Source/Runtime/Renderer/Private/PostProcess/DiaphragmDOFUtils.cpp
--- a/Engine/Source/Runtime/Renderer/Private/PostProcess/DiaphragmDOFUtils.cpp
+++ b/Engine/Source/Runtime/Renderer/Private/PostProcess/DiaphragmDOFUtils.cpp
@@ -112,6 +112,12 @@ void DiaphragmDOF::FPhysicalCocModel::Compile(const FViewInfo& View)
 		PetzvalFalloffPower = FMath::Clamp(View.FinalPostProcessSettings.DepthOfFieldPetzvalBokehFalloff, 0, 100.0);
 		PetzvalExclusionBoxExtents = View.FinalPostProcessSettings.DepthOfFieldPetzvalExclusionBoxExtents;
 		PetzvalExclusionBoxRadius = View.FinalPostProcessSettings.DepthOfFieldPetzvalExclusionBoxRadius;
+		
+		static_assert(MaxMatteBoxFlags == UE_ARRAY_COUNT(View.FinalPostProcessSettings.DepthOfFieldMatteBoxFlags));
+		for (int Index = 0; Index < MaxMatteBoxFlags; ++Index)
+		{
+			MatteBoxFlags[Index] = View.FinalPostProcessSettings.DepthOfFieldMatteBoxFlags[Index];
+		}
 
 		RenderingAspectRatio = float(View.UnscaledViewRect.Width()) / float(View.UnscaledViewRect.Height()) * View.FinalPostProcessSettings.DepthOfFieldAspectRatioScalar;
 		const float HorizontalHalfFOV = FMath::Atan(1.0f / View.ViewMatrices.GetProjectionMatrix().M[0][0]);
@@ -122,6 +128,10 @@ void DiaphragmDOF::FPhysicalCocModel::Compile(const FViewInfo& View)
 		SensorWidth  = View.FinalPostProcessSettings.DepthOfFieldSensorWidth * MMToUU;
 		SensorHeight = SensorWidth / SensorAspectRatio;
 		VerticalFocalLength = 0.5f * SensorHeight * (1.0f / FMath::Tan(VerticalHalfFOV));
+
+		float Aperture = VerticalFocalLength / FStops;
+		BarrelRadius = FMath::Max(View.FinalPostProcessSettings.DepthOfFieldBarrelRadius, (Aperture / 2) + 0.5f);
+		BarrelLength = View.FinalPostProcessSettings.DepthOfFieldBarrelLength;
 	}
 
 	// Fetch the max bluring radius
