Fix high peak memory usage when caching MDCs.

diff --git a/Engine/Source/Runtime/Renderer/Private/PrimitiveSceneInfo.cpp b/Engine/Source/Runtime/Renderer/Private/PrimitiveSceneInfo.cpp
--- a/Engine/Source/Runtime/Renderer/Private/PrimitiveSceneInfo.cpp
+++ b/Engine/Source/Runtime/Renderer/Private/PrimitiveSceneInfo.cpp
@@ -543,6 +543,9 @@ void FPrimitiveSceneInfo::CacheMeshDrawCommands(FScene* Scene, TArrayView<FPrimi
 			}
 		}
 
+		const int32 NumCommonPassesExpected = 4;	// to avoid reserving too much, account only for heuristically likely passes like depth / velocity / base / something
+		DrawListContext.ReserveMemoryForCommands(NumCommonPassesExpected * MeshBatches.Num());
+
 		for (int32 PassIndex = 0; PassIndex < EMeshPass::Num; PassIndex++)
 		{
 			const EShadingPath ShadingPath = GetFeatureLevelShadingPath(Scene->GetFeatureLevel());
diff --git a/Engine/Source/Runtime/Renderer/Public/MeshPassProcessor.h b/Engine/Source/Runtime/Renderer/Public/MeshPassProcessor.h
--- a/Engine/Source/Runtime/Renderer/Public/MeshPassProcessor.h
+++ b/Engine/Source/Runtime/Renderer/Public/MeshPassProcessor.h
@@ -1984,6 +1984,7 @@ public:
 
 	FCachedPassMeshDrawListContext(FScene& InScene);
 
+	virtual void ReserveMemoryForCommands(int32 MaxCommandsExpected) {};
 	virtual FMeshDrawCommand& AddCommand(FMeshDrawCommand& Initializer, uint32 NumElements) override final;
 
 	void BeginMeshPass(EMeshPass::Type MeshPass);
@@ -2036,11 +2037,17 @@ public:
 class FCachedPassMeshDrawListContextDeferred : public FCachedPassMeshDrawListContext
 {
 public:
-	FCachedPassMeshDrawListContextDeferred(FScene& InScene) : FCachedPassMeshDrawListContext(InScene)
+	FCachedPassMeshDrawListContextDeferred(FScene& InScene) : FCachedPassMeshDrawListContext(InScene) {}
+
+	virtual void ReserveMemoryForCommands(int32 MaxCommandsExpected) override final
 	{
-		if (LIKELY(bUseGPUScene))
+		if (LIKELY(MaxCommandsExpected))
 		{
-			DeferredCommandHashes.Reserve(DeferredCommands.Max());
+			DeferredCommands.Reserve(MaxCommandsExpected);
+			if (LIKELY(bUseGPUScene))
+			{
+				DeferredCommandHashes.Reserve(DeferredCommands.Max());
+			}
 		}
 	}
 
@@ -2059,7 +2066,7 @@ public:
 	void DeferredFinalizeMeshDrawCommands(const TArrayView<FPrimitiveSceneInfo*>& SceneInfos, int32 Start, int32 End);
 
 private:
-	TArray<FMeshDrawCommand, TInlineAllocator<192>> DeferredCommands;
+	TArray<FMeshDrawCommand> DeferredCommands;
 	TArray<Experimental::FHashType> DeferredCommandHashes;
 };
 
