[SingleLayerWater] Fixed ensure triggered by accessing a cvar on the GT with the RT accessor function.

diff --git a/Engine/Source/Runtime/Renderer/Private/SingleLayerWaterRendering.cpp b/Engine/Source/Runtime/Renderer/Private/SingleLayerWaterRendering.cpp
--- a/Engine/Source/Runtime/Renderer/Private/SingleLayerWaterRendering.cpp
+++ b/Engine/Source/Runtime/Renderer/Private/SingleLayerWaterRendering.cpp
@@ -155,9 +155,16 @@ static int32 GetSingleLayerWaterReflectionTechnique()
 	return FMath::Clamp(Value, 0, ESingleLayerWaterReflections::MaxValue);
 }
 
-static EGBufferLayout GetSingleLayerWaterGBufferLayout()
+static EGBufferLayout GetSingleLayerWaterGBufferLayout(bool bIsGameThread = false)
 {
-	return CVarWaterSingleLayerForceVelocity.GetValueOnRenderThread() != 0 ? GBL_ForceVelocity : GBL_Default;
+	if (!bIsGameThread)
+	{
+		return CVarWaterSingleLayerForceVelocity.GetValueOnRenderThread() != 0 ? GBL_ForceVelocity : GBL_Default;
+	}
+	else
+	{
+		return CVarWaterSingleLayerForceVelocity.GetValueOnGameThread() != 0 ? GBL_ForceVelocity : GBL_Default;
+	}
 }
 
 // This is to have platforms use the simple single layer water shading similar to mobile: no dynamic lights, only sun and sky, no distortion, no colored transmittance on background, no custom depth read.
@@ -1563,7 +1570,7 @@ void FSingleLayerWaterPassMeshProcessor::CollectPSOInitializers(const FSceneText
 			TBasePassVertexShaderPolicyParamType<LightMapPolicyType>,
 			TBasePassPixelShaderPolicyParamType<LightMapPolicyType>> WaterPassShaders;
 
-		const EGBufferLayout GBufferLayout = GetSingleLayerWaterGBufferLayout();
+		const EGBufferLayout GBufferLayout = GetSingleLayerWaterGBufferLayout(true /*bIsGameThread*/);
 		const bool bRenderSkylight = true;
 		if (!GetBasePassShaders<LightMapPolicyType>(
 			Material,
