Reduce CombineLUTs permutations given recently added one for ACES 2.0.

diff --git a/Engine/Shaders/Private/PostProcessCombineLUTs.usf b/Engine/Shaders/Private/PostProcessCombineLUTs.usf
--- a/Engine/Shaders/Private/PostProcessCombineLUTs.usf
+++ b/Engine/Shaders/Private/PostProcessCombineLUTs.usf
@@ -81,6 +81,10 @@ Texture2D<float> ACESReachTable;
 Texture2D<float3> ACESGamutTable;
 Texture2D<float> ACESGammaTable;
 
+static const float DefaultTemperature = 6500.0;
+static const float DefaultTint = 0.0;
+static const float Epsilon = 1e-8;
+
 float3 ColorCorrect( float3 WorkingColor,
 	float4 ColorSaturation,
 	float4 ColorContrast,
@@ -208,11 +212,13 @@ float4 CombineLUTsCommon(float2 InUV, uint InLayerIndex)
 		LinearColor = LogToLin(LUTEncodedColor) - LogToLin(0);
 	}
 
-#if !SKIP_TEMPERATURE
-	float3 BalancedColor = WhiteBalance(LinearColor, WhiteTemp, WhiteTint, bIsTemperatureWhiteBalance, (float3x3)WorkingColorSpace.ToXYZ, (float3x3)WorkingColorSpace.FromXYZ);
-#else
 	float3 BalancedColor = LinearColor;
-#endif
+	
+	BRANCH
+	if( abs(WhiteTemp - DefaultTemperature) > Epsilon || abs(WhiteTint - DefaultTint) > Epsilon )
+	{
+		BalancedColor = WhiteBalance(LinearColor, WhiteTemp, WhiteTint, bIsTemperatureWhiteBalance, (float3x3)WorkingColorSpace.ToXYZ, (float3x3)WorkingColorSpace.FromXYZ);
+	}
 
 	float3 ColorAP1 = mul( (float3x3)WorkingColorSpace.ToAP1, BalancedColor );
 
diff --git a/Engine/Source/Runtime/Renderer/Private/PostProcess/PostProcessCombineLUTs.cpp b/Engine/Source/Runtime/Renderer/Private/PostProcess/PostProcessCombineLUTs.cpp
--- a/Engine/Source/Runtime/Renderer/Private/PostProcess/PostProcessCombineLUTs.cpp
+++ b/Engine/Source/Runtime/Renderer/Private/PostProcess/PostProcessCombineLUTs.cpp
@@ -321,10 +321,9 @@ public:
 	static const int32 GroupSize = 8;
 
 	class FBlendCount : SHADER_PERMUTATION_RANGE_INT("BLENDCOUNT", 1, 5);
-	class FSkipTemperature : SHADER_PERMUTATION_BOOL("SKIP_TEMPERATURE");
 	class FOutputDeviceSRGB : SHADER_PERMUTATION_BOOL("OUTPUT_DEVICE_SRGB");
 	class FUseACES2 : SHADER_PERMUTATION_BOOL("USE_ACES_2");
-	using FPermutationDomain = TShaderPermutationDomain<FBlendCount, FSkipTemperature, FOutputDeviceSRGB, FUseACES2>;
+	using FPermutationDomain = TShaderPermutationDomain<FBlendCount, FOutputDeviceSRGB, FUseACES2>;
 
 	static bool ShouldCompilePermutation(const FGlobalShaderPermutationParameters& Parameters)
 	{
@@ -545,9 +544,6 @@ FRDGTextureRef AddCombineLUTPass(FRDGBuilder& GraphBuilder, const FViewInfo& Vie
 	static const auto CVarAcesVersion = IConsoleManager::Get().FindTConsoleVariableDataInt(TEXT("r.HDR.Aces.Version"));
 	PermutationVector.Set<FLUTBlenderShader::FUseACES2>(CVarAcesVersion->GetValueOnRenderThread() > 1);
 
-	const float DefaultTemperature = 6500;
-	const float DefaultTint = 0;
-
 	if (bUseComputePass)
 	{
 		FLUTBlenderCS::FParameters* PassParameters = GraphBuilder.AllocParameters<FLUTBlenderCS::FParameters>();
@@ -555,10 +551,8 @@ FRDGTextureRef AddCombineLUTPass(FRDGBuilder& GraphBuilder, const FViewInfo& Vie
 		PassParameters->OutputExtentInverse = FVector2f(1.0f, 1.0f) / FVector2f(OutputViewSize);
 		PassParameters->RWOutputTexture = GraphBuilder.CreateUAV(OutputTexture);
 
-		const bool ShouldSkipTemperature = FMath::IsNearlyEqual(PassParameters->CombineLUT.WhiteTemp, DefaultTemperature) && FMath::IsNearlyEqual(PassParameters->CombineLUT.WhiteTint, DefaultTint);
 		const bool bOutputDeviceSRGB = (PassParameters->CombineLUT.OutputDevice.OutputDevice == (uint32)EDisplayOutputFormat::SDR_sRGB);
 
-		PermutationVector.Set<FLUTBlenderShader::FSkipTemperature>(ShouldSkipTemperature);
 		PermutationVector.Set<FLUTBlenderShader::FOutputDeviceSRGB>(bOutputDeviceSRGB);
 		TShaderMapRef<FLUTBlenderCS> ComputeShader(View.ShaderMap, PermutationVector);
 
@@ -578,10 +572,8 @@ FRDGTextureRef AddCombineLUTPass(FRDGBuilder& GraphBuilder, const FViewInfo& Vie
 		PassParameters->CombineLUT = CachedLUTSettings.Parameters;
 		PassParameters->RenderTargets[0] = FRenderTargetBinding(OutputTexture, ERenderTargetLoadAction::ENoAction);
 
-		const bool ShouldSkipTemperature = FMath::IsNearlyEqual(PassParameters->CombineLUT.WhiteTemp, DefaultTemperature) && FMath::IsNearlyEqual(PassParameters->CombineLUT.WhiteTint, DefaultTint);
 		const bool bOutputDeviceSRGB = (PassParameters->CombineLUT.OutputDevice.OutputDevice == (uint32)EDisplayOutputFormat::SDR_sRGB);
 
-		PermutationVector.Set<FLUTBlenderShader::FSkipTemperature>(ShouldSkipTemperature);
 		PermutationVector.Set<FLUTBlenderShader::FOutputDeviceSRGB>(bOutputDeviceSRGB);
 		TShaderMapRef<FLUTBlenderPS> PixelShader(View.ShaderMap, PermutationVector);
 
