Change virtual texture continuous update to be set per page table space instead of per physical space.
This removes a restriction that a runtime virtual texture that is set to continuous update can't reuse the same physical pool as other generic streaming virtual textures (which don't continuous update).
This can matter for runtime virtual textures with a single layer and a common texture format such as Displacement (BC4), World Height (G16), and Mask4 (DXT5).
The new restriction that runtime virtual textures might now not share the same page table space is much less problematic since runtime virtual textures typically use private page table spaces.
Note that the bContinuousUpdate flag can now be removed from the runtime virtual texture streaming mips. It was only needed to ensure that both runtime and streaming pages always ended in the same physical space.

diff --git a/Engine/Source/Runtime/Engine/Classes/VT/VirtualTexture.h b/Engine/Source/Runtime/Engine/Classes/VT/VirtualTexture.h
--- a/Engine/Source/Runtime/Engine/Classes/VT/VirtualTexture.h
+++ b/Engine/Source/Runtime/Engine/Classes/VT/VirtualTexture.h
@@ -43,14 +43,13 @@ class UVirtualTexture2D : public UTexture2D
 	FVirtualTextureBuildSettings Settings;
 
 	UPROPERTY()
-	bool bContinuousUpdate;
+	bool bContinuousUpdate_DEPRECATED;
 
 	UPROPERTY()
 	bool bSinglePhysicalSpace;
 
 	//~ Begin UTexture Interface.
 	virtual void GetVirtualTextureBuildSettings(FVirtualTextureBuildSettings& OutSettings) const override { OutSettings = Settings; }
-	virtual bool IsVirtualTexturedWithContinuousUpdate() const override { return bContinuousUpdate; }
 	virtual bool IsVirtualTexturedWithSinglePhysicalSpace() const override { return bSinglePhysicalSpace; }
 #if WITH_EDITOR
 	virtual void BeginCacheForCookedPlatformData(const ITargetPlatform* TargetPlatform) override;
diff --git a/Engine/Source/Runtime/Engine/Classes/VT/VirtualTextureBuilder.h b/Engine/Source/Runtime/Engine/Classes/VT/VirtualTextureBuilder.h
--- a/Engine/Source/Runtime/Engine/Classes/VT/VirtualTextureBuilder.h
+++ b/Engine/Source/Runtime/Engine/Classes/VT/VirtualTextureBuilder.h
@@ -25,6 +25,7 @@ struct FVirtualTextureBuildDesc
 	TEnumAsByte<enum TextureGroup> LODGroup = TEXTUREGROUP_World;
 	ETextureLossyCompressionAmount LossyCompressionAmount = TLCA_Default;
 
+	UE_DEPRECATED(5.6, "bContinuousUpdate is not used.")
 	bool bContinuousUpdate = false;
 	bool bSinglePhysicalSpace = false;
 
diff --git a/Engine/Source/Runtime/Engine/Private/Components/RuntimeVirtualTextureComponent.cpp b/Engine/Source/Runtime/Engine/Private/Components/RuntimeVirtualTextureComponent.cpp
--- a/Engine/Source/Runtime/Engine/Private/Components/RuntimeVirtualTextureComponent.cpp
+++ b/Engine/Source/Runtime/Engine/Private/Components/RuntimeVirtualTextureComponent.cpp
@@ -511,7 +511,6 @@ void URuntimeVirtualTextureComponent::InitializeStreamingTexture(EShadingPath Sh
 		FScopedRuntimeVirtualTextureRecreate ProducerRecreate(StreamingTexture);
 
 		FVirtualTextureBuildDesc BuildDesc;
-		BuildDesc.bContinuousUpdate = VirtualTexture->GetContinuousUpdate();
 		BuildDesc.bSinglePhysicalSpace = VirtualTexture->GetSinglePhysicalSpace();
 
 		BuildDesc.TileSize = VirtualTexture->GetTileSize();
diff --git a/Engine/Source/Runtime/Engine/Private/Texture2D.cpp b/Engine/Source/Runtime/Engine/Private/Texture2D.cpp
--- a/Engine/Source/Runtime/Engine/Private/Texture2D.cpp
+++ b/Engine/Source/Runtime/Engine/Private/Texture2D.cpp
@@ -1522,7 +1522,6 @@ FVirtualTexture2DResource::FVirtualTexture2DResource(const UTexture2D* InOwner,
 	TexCreateFlags |= InOwner->bNotOfflineProcessed ? ETextureCreateFlags::None : ETextureCreateFlags::OfflineProcessed;
 	TexCreateFlags |= InOwner->bNoTiling ? ETextureCreateFlags::NoTiling : ETextureCreateFlags::None;
 
-	bContinuousUpdate = InOwner->IsVirtualTexturedWithContinuousUpdate();
 	bSinglePhysicalSpace = InOwner->IsVirtualTexturedWithSinglePhysicalSpace();
 
 	check(InVTData);
@@ -1587,7 +1586,6 @@ void FVirtualTexture2DResource::InitRHI(FRHICommandListBase& RHICmdList)
 	FVTProducerDescription ProducerDesc;
 	ProducerDesc.Name = TextureName;
 	ProducerDesc.FullNameHash = GetTypeHash(TextureName);
-	ProducerDesc.bContinuousUpdate = bContinuousUpdate;
 	ProducerDesc.Dimensions = 2;
 	ProducerDesc.TileSize = VTData->TileSize;
 	ProducerDesc.TileBorderSize = VTData->TileBorderSize;
diff --git a/Engine/Source/Runtime/Engine/Private/VT/VirtualTexture.cpp b/Engine/Source/Runtime/Engine/Private/VT/VirtualTexture.cpp
--- a/Engine/Source/Runtime/Engine/Private/VT/VirtualTexture.cpp
+++ b/Engine/Source/Runtime/Engine/Private/VT/VirtualTexture.cpp
@@ -184,7 +184,6 @@ URuntimeVirtualTextureStreamingProxy::URuntimeVirtualTextureStreamingProxy(const
 
 UVirtualTexture2D::UVirtualTexture2D(const FObjectInitializer& ObjectInitializer)
 	: UTexture2D(ObjectInitializer)
-	, bContinuousUpdate(false)
 	, bSinglePhysicalSpace(false)
 {
 }
diff --git a/Engine/Source/Runtime/Engine/Private/VT/VirtualTextureBuilder.cpp b/Engine/Source/Runtime/Engine/Private/VT/VirtualTextureBuilder.cpp
--- a/Engine/Source/Runtime/Engine/Private/VT/VirtualTextureBuilder.cpp
+++ b/Engine/Source/Runtime/Engine/Private/VT/VirtualTextureBuilder.cpp
@@ -120,7 +120,6 @@ static void BuildVirtualTexture2D(UVirtualTexture2D* Texture, FVirtualTextureBui
 	Texture->Settings.TileBorderSize = BuildDesc.TileBorderSize;
 	Texture->LossyCompressionAmount = BuildDesc.LossyCompressionAmount;
 
-	Texture->bContinuousUpdate = BuildDesc.bContinuousUpdate;
 	Texture->bSinglePhysicalSpace = BuildDesc.bSinglePhysicalSpace;
 
 	for (int32 Layer = 0; Layer < BuildDesc.LayerCount; Layer++)
diff --git a/Engine/Source/Runtime/Engine/Public/TextureResource.h b/Engine/Source/Runtime/Engine/Public/TextureResource.h
--- a/Engine/Source/Runtime/Engine/Public/TextureResource.h
+++ b/Engine/Source/Runtime/Engine/Public/TextureResource.h
@@ -230,7 +230,6 @@ protected:
 	/** Cached flags for texture creation. */
 	ETextureCreateFlags TexCreateFlags = ETextureCreateFlags::None;
 	/** Cached runtime virtual texture settings */
-	bool bContinuousUpdate = false;
 	bool bSinglePhysicalSpace = false;
 	/** Mip offset */
 	int32 FirstMipToUse = 0;
diff --git a/Engine/Source/Runtime/Renderer/Private/VT/AllocatedVirtualTexture.cpp b/Engine/Source/Runtime/Renderer/Private/VT/AllocatedVirtualTexture.cpp
--- a/Engine/Source/Runtime/Renderer/Private/VT/AllocatedVirtualTexture.cpp
+++ b/Engine/Source/Runtime/Renderer/Private/VT/AllocatedVirtualTexture.cpp
@@ -36,6 +36,8 @@ FAllocatedVirtualTexture::FAllocatedVirtualTexture(
 	FMemory::Memzero(TextureLayers);
 	FMemory::Memzero(FallbackColorPerTextureLayer);
 
+	bool bContinuousUpdate = false;
+
 	for (uint32 LayerIndex = 0u; LayerIndex < Description.NumTextureLayers; ++LayerIndex)
 	{
 		FVirtualTextureProducer const* Producer = InProducers[LayerIndex];
@@ -50,6 +52,7 @@ FAllocatedVirtualTexture::FAllocatedVirtualTexture(
 			{
 				ProducerPhysicalGroupIndex = Producer->GetPhysicalGroupIndexForTextureLayer(ProducerLayerIndex);
 				PhysicalSpace = Producer->GetPhysicalSpaceForPhysicalGroup(ProducerPhysicalGroupIndex);
+				bContinuousUpdate |= Producer->GetDescription().bContinuousUpdate;
 			}
 			const uint32 UniquePhysicalSpaceIndex = AddUniquePhysicalSpace(RHICmdList, PhysicalSpace, UniqueProducerIndex, ProducerPhysicalGroupIndex);
 			UniquePageTableLayers[UniquePhysicalSpaceIndex].ProducerTextureLayerMask |= 1 << ProducerLayerIndex;
@@ -101,6 +104,7 @@ FAllocatedVirtualTexture::FAllocatedVirtualTexture(
 	SpaceDesc.TileSize = InDesc.TileSize;
 	SpaceDesc.TileBorderSize = InDesc.TileBorderSize;
 	SpaceDesc.bPrivateSpace = InDesc.bPrivateSpace;
+	SpaceDesc.bContinuousUpdate = bContinuousUpdate;
 	SpaceDesc.MaxSpaceSize = InDesc.MaxSpaceSize > 0 ? InDesc.MaxSpaceSize : SpaceDesc.MaxSpaceSize;
 	SpaceDesc.IndirectionTextureSize = InDesc.IndirectionTextureSize;
 	SpaceDesc.PageTableFormat = bSupport16BitPageTable ? EVTPageTableFormat::UInt16 : EVTPageTableFormat::UInt32;
diff --git a/Engine/Source/Runtime/Renderer/Private/VT/VirtualTexturePhysicalSpace.h b/Engine/Source/Runtime/Renderer/Private/VT/VirtualTexturePhysicalSpace.h
--- a/Engine/Source/Runtime/Renderer/Private/VT/VirtualTexturePhysicalSpace.h
+++ b/Engine/Source/Runtime/Renderer/Private/VT/VirtualTexturePhysicalSpace.h
@@ -20,15 +20,13 @@ struct FVTPhysicalSpaceDescription
 	uint8 NumLayers;
 	TEnumAsByte<EPixelFormat> Format[VIRTUALTEXTURE_SPACE_MAXLAYERS];
 	bool bHasLayerSrgbView[VIRTUALTEXTURE_SPACE_MAXLAYERS];
-	bool bContinuousUpdate;
 };
 
 inline bool operator==(const FVTPhysicalSpaceDescription& Lhs, const FVTPhysicalSpaceDescription& Rhs)
 {
 	if (Lhs.TileSize != Rhs.TileSize ||
 		Lhs.NumLayers != Rhs.NumLayers || 
-		Lhs.Dimensions != Rhs.Dimensions || 
-		Lhs.bContinuousUpdate != Rhs.bContinuousUpdate)
+		Lhs.Dimensions != Rhs.Dimensions)
 	{
 		return false;
 	}
@@ -57,7 +55,6 @@ inline uint32 GetTypeHash(const FVTPhysicalSpaceDescription& Desc)
 	uint32 Hash = Desc.TileSize;
 	Hash = HashCombine(Hash, GetTypeHash(Desc.Dimensions));
 	Hash = HashCombine(Hash, GetTypeHash(Desc.NumLayers));
-	Hash = HashCombine(Hash, GetTypeHash(Desc.bContinuousUpdate));
 	for (int32 Layer = 0; Layer < Desc.NumLayers; ++Layer)
 	{
 		Hash = HashCombine(Hash, GetTypeHash(Desc.Format[Layer]));
diff --git a/Engine/Source/Runtime/Renderer/Private/VT/VirtualTextureProducer.cpp b/Engine/Source/Runtime/Renderer/Private/VT/VirtualTextureProducer.cpp
--- a/Engine/Source/Runtime/Renderer/Private/VT/VirtualTextureProducer.cpp
+++ b/Engine/Source/Runtime/Renderer/Private/VT/VirtualTextureProducer.cpp
@@ -74,7 +74,6 @@ FVirtualTextureProducerHandle FVirtualTextureProducerCollection::RegisterProduce
 	FVTPhysicalSpaceDescription PhysicalSpaceDesc;
 	PhysicalSpaceDesc.Dimensions = InDesc.Dimensions;
 	PhysicalSpaceDesc.TileSize = InDesc.TileSize + InDesc.TileBorderSize * 2u;
-	PhysicalSpaceDesc.bContinuousUpdate = InDesc.bContinuousUpdate;
 
 	Entry.Producer.PhysicalGroups.SetNum(InDesc.NumPhysicalGroups);
 	for (uint32 PhysicalGroupIndex = 0u; PhysicalGroupIndex < InDesc.NumPhysicalGroups; ++PhysicalGroupIndex)
diff --git a/Engine/Source/Runtime/Renderer/Private/VT/VirtualTextureSpace.h b/Engine/Source/Runtime/Renderer/Private/VT/VirtualTextureSpace.h
--- a/Engine/Source/Runtime/Renderer/Private/VT/VirtualTextureSpace.h
+++ b/Engine/Source/Runtime/Renderer/Private/VT/VirtualTextureSpace.h
@@ -21,6 +21,7 @@ struct FVTSpaceDescription
 	EVTPageTableFormat PageTableFormat = EVTPageTableFormat::UInt16;
 	uint8 NumPageTableLayers = 0u;
 	uint8 bPrivateSpace = false;
+	uint8 bContinuousUpdate = false;
 	uint32 MaxSpaceSize = VIRTUALTEXTURE_MAX_PAGETABLE_SIZE;
 	uint32 IndirectionTextureSize = 0u;
 };
@@ -33,6 +34,7 @@ inline bool operator==(const FVTSpaceDescription& Lhs, const FVTSpaceDescription
 		Lhs.NumPageTableLayers == Rhs.NumPageTableLayers &&
 		Lhs.PageTableFormat == Rhs.PageTableFormat &&
 		Lhs.bPrivateSpace == Rhs.bPrivateSpace &&
+		Lhs.bContinuousUpdate == Rhs.bContinuousUpdate &&
 		Lhs.MaxSpaceSize == Rhs.MaxSpaceSize &&
 		Lhs.IndirectionTextureSize == Rhs.IndirectionTextureSize;
 }
diff --git a/Engine/Source/Runtime/Renderer/Private/VT/VirtualTextureSystem.cpp b/Engine/Source/Runtime/Renderer/Private/VT/VirtualTextureSystem.cpp
--- a/Engine/Source/Runtime/Renderer/Private/VT/VirtualTextureSystem.cpp
+++ b/Engine/Source/Runtime/Renderer/Private/VT/VirtualTextureSystem.cpp
@@ -1602,8 +1602,7 @@ void FVirtualTextureSystem::GatherRequestsTask(const FGatherRequestsParameters&
 					AddPageUpdate(PageUpdateBuffers, PageUpdateFlushCount, PhysicalSpaceIDAndAddress.PhysicalSpaceID, PhysicalSpaceIDAndAddress.pAddress);
 
 					// If continuous update flag is set then add this to pages which can be potentially updated if we have spare upload bandwidth
-					//todo[vt]: Would be better to test continuous update flag *per producer*, but this would require extra indirection so need to profile first
-					if (bForceContinuousUpdate || GetPhysicalSpace(PhysicalSpaceIDAndAddress.PhysicalSpaceID)->GetDescription().bContinuousUpdate)
+					if (bForceContinuousUpdate || Space->GetDescription().bContinuousUpdate)
 					{
 						FTexturePagePool& RESTRICT PagePool = GetPhysicalSpace(PhysicalSpaceIDAndAddress.PhysicalSpaceID)->GetPagePool();
 						const FVirtualTextureLocalTile LocalTile = PagePool.GetLocalTileFromPhysicalAddress(PhysicalSpaceIDAndAddress.pAddress);
