Ray Tracing: Revert changes to dither masked opacity

After some discussion internally, we have decided the previous behavior is more desireable in the realtime cases as there are better workarounds for artifacts. Users that need translucent behavior can simply use the translucent blend modes.

Dither masking now no longer triggers any special behavior and gets treated the same as masked materials in RayTracing codepaths. If materials are not being run, this means they will act as opaque shadow casters (which can be ignored by simply letting them not cast shadows).

#rnx

diff --git a/Engine/Shaders/Private/RayTracing/RayTracingMaterialHitShaders.usf b/Engine/Shaders/Private/RayTracing/RayTracingMaterialHitShaders.usf
--- a/Engine/Shaders/Private/RayTracing/RayTracingMaterialHitShaders.usf
+++ b/Engine/Shaders/Private/RayTracing/RayTracingMaterialHitShaders.usf
@@ -324,9 +324,6 @@ void GetMaterialPayload(
 			TransmittancePreCoverage);
 
 	}
-	#elif MATERIALBLENDING_MASKED && MATERIAL_DITHER_OPACITY_MASK
-	// If we have dithered masking, convert to regular transparency here
-	TotalCoverage *= saturate(GetMaterialMaskInputRaw(PixelMaterialInputs));
 	#endif
 
 	Payload.SubstrateData = InitialiseRWSubstrateMaterialContainer();
@@ -404,7 +401,7 @@ void GetMaterialPayload(
 	Payload.Opacity = TotalCoverage;
 	Payload.TransmittancePreCoverage = TransmittancePreCoverage;
 
-	#if MATERIALBLENDING_MASKED && !MATERIAL_DITHER_OPACITY_MASK
+	#if MATERIALBLENDING_MASKED
 	// regular masked mode - binary decision
 	if (GetMaterialMask(PixelMaterialInputs) < 0)
 	{
@@ -434,11 +431,7 @@ void GetMaterialPayload(
 	Payload.Specular = Specular;
 	Payload.Roughness = Roughness;
 	Payload.Metallic = Metallic;
-#if MATERIALBLENDING_MASKED && MATERIAL_DITHER_OPACITY_MASK
-	// dithered opacity doesn't make sense for pure raytracing, so replace it with the equivalent translucent behavior
-	Payload.Opacity = saturate(GetMaterialMaskInputRaw(PixelMaterialInputs));
-
-#elif MATERIALBLENDING_MODULATE
+#if MATERIALBLENDING_MODULATE
 	Payload.Opacity = 0.0;
 #else
 	// All other blending modes
@@ -665,9 +658,6 @@ RAY_TRACING_ENTRY_CLOSEST_HIT(MaterialCHS,
 		uint MaterialBlendingMode = RAY_TRACING_BLEND_MODE_ADDITIVE;
 	#elif MATERIALBLENDING_MODULATE
 		uint MaterialBlendingMode = RAY_TRACING_BLEND_MODE_MODULATE;
-	#elif MATERIALBLENDING_MASKED && MATERIAL_DITHER_OPACITY_MASK
-		// dithered opacity doesn't make sense for pure raytracing, so replace it with the equivalent translucent behavior
-		uint MaterialBlendingMode = RAY_TRACING_BLEND_MODE_TRANSLUCENT;
 	#else
 		uint MaterialBlendingMode = RAY_TRACING_BLEND_MODE_OPAQUE;
 	#endif
@@ -759,7 +749,7 @@ RAY_TRACING_ENTRY_ANY_HIT(MaterialAHS,
 	// All other blending modes _might_ need material evaluation - so we need
 	// to generate some code
 
-#if MATERIALBLENDING_MASKED && !MATERIAL_DITHER_OPACITY_MASK
+#if MATERIALBLENDING_MASKED
 	// NOTE: Masked mode execute always - regardless of ray flags
 	// This ensures that both material rays and shadow rays "see" the same thing
 #else
@@ -828,17 +818,15 @@ RAY_TRACING_ENTRY_ANY_HIT(MaterialAHS,
 	}
 
 	// Now extract the relevant info from PixelMaterialInputs according to the blending mode
-#if MATERIALBLENDING_MASKED && !MATERIAL_DITHER_OPACITY_MASK
+#if MATERIALBLENDING_MASKED
 	// Regardless of payload flags -- we always apply this 
 	if (GetMaterialMask(PixelMaterialInputs) < 0)
 	{
 		IgnoreHit();
 	}
-#else // MATERIALBLENDING_MASKED && !MATERIAL_DITHER_OPACITY_MASK
+#else // MATERIALBLENDING_MASKED
 
-#if MATERIALBLENDING_MASKED && MATERIAL_DITHER_OPACITY_MASK
-	const float Transparency = 1 - saturate(GetMaterialMaskInputRaw(PixelMaterialInputs));
-#elif SUBSTRATE_ENABLED
+#if SUBSTRATE_ENABLED
 		#if MATERIAL_IS_SUBSTRATE
 		float3 EmissiveLuminance = 0.0f;
 		float TotalCoverage = 1.f;
@@ -964,7 +952,7 @@ RAY_TRACING_ENTRY_ANY_HIT(MaterialAHS,
 	{
 		IgnoreHit();
 	}
-#endif // !MATERIALBLENDING_MASKED || MATERIAL_DITHER_OPACITY_MASK
+#endif // !MATERIALBLENDING_MASKED
 
 #endif // MATERIALBLENDING_SOLID
 
diff --git a/Engine/Source/Runtime/Renderer/Private/RayTracing/RayTracingInstanceMask.cpp b/Engine/Source/Runtime/Renderer/Private/RayTracing/RayTracingInstanceMask.cpp
--- a/Engine/Source/Runtime/Renderer/Private/RayTracing/RayTracingInstanceMask.cpp
+++ b/Engine/Source/Runtime/Renderer/Private/RayTracing/RayTracingInstanceMask.cpp
@@ -110,20 +110,21 @@ uint8 GetFarFieldRayTracingInstanceMask(ERayTracingType RayTracingType)
 
 uint8 BlendModeToRayTracingInstanceMask(const EBlendMode BlendMode, bool bIsDitherMasked, bool bCastShadow, ERayTracingType RayTracingType)
 {
-	uint8 InstanceMask = 0;
+	ERayTracingInstanceMaskType MaskTypePrimary = ERayTracingInstanceMaskType::Opaque;
+	ERayTracingInstanceMaskType MaskTypeShadows = ERayTracingInstanceMaskType::OpaqueShadow;
 
-	if (IsOpaqueOrMaskedBlendMode(BlendMode) && !bIsDitherMasked)
-	{
-		InstanceMask |= ComputeRayTracingInstanceMask(ERayTracingInstanceMaskType::Opaque, RayTracingType);
-		InstanceMask |= bCastShadow ? ComputeRayTracingInstanceMask(ERayTracingInstanceMaskType::OpaqueShadow, RayTracingType) : 0;
-	}
-	else
+	const bool bIsOpaqueOrMasked = IsOpaqueOrMaskedBlendMode(BlendMode);
+
+	// RayTracing treats dithered masked materials the same as regular masked materials for speed
+	// PathTracing/LightmapTracing both upgrade dithered masking to translucent internally and therefore need to take them with those bits
+	if ((RayTracingType == ERayTracingType::RayTracing && !bIsOpaqueOrMasked) ||
+		(RayTracingType != ERayTracingType::RayTracing && (!bIsOpaqueOrMasked || bIsDitherMasked)))
 	{
-		InstanceMask |= ComputeRayTracingInstanceMask(ERayTracingInstanceMaskType::Translucent, RayTracingType);
-		InstanceMask |= bCastShadow ? ComputeRayTracingInstanceMask(ERayTracingInstanceMaskType::TranslucentShadow, RayTracingType) : 0;
+		MaskTypePrimary = ERayTracingInstanceMaskType::Translucent;
+		MaskTypeShadows = ERayTracingInstanceMaskType::TranslucentShadow;
 	}
 
-	return InstanceMask;
+	return ComputeRayTracingInstanceMask(MaskTypePrimary, RayTracingType) |	(bCastShadow ? ComputeRayTracingInstanceMask(MaskTypeShadows, RayTracingType) : 0);
 }
 
 namespace {
@@ -190,7 +191,6 @@ FRayTracingMaskAndFlags BuildRayTracingInstanceMaskAndFlags(const FRayTracingIns
 	ensureMsgf(Instance.GetMaterials().Num() > 0, TEXT("You need to add MeshBatches first for instance mask and flags to build upon."));
 
 	bool bAllSegmentsOpaque = true;
-	bool bAllSegmentsTranslucent = true;
 	bool bAnySegmentsCastShadow = false;
 	bool bAllSegmentsCastShadow = true;
 	bool bAnySegmentsDecal = false;
@@ -210,7 +210,6 @@ FRayTracingMaskAndFlags BuildRayTracingInstanceMaskAndFlags(const FRayTracingIns
 
 			Result.Mask |= BlendModeToRayTracingInstanceMask(BlendMode, Material.IsDitherMasked(), bSegmentCastsShadow, RayTracingType);
 			bAllSegmentsOpaque &= BlendMode == EBlendMode::BLEND_Opaque;
-			bAllSegmentsTranslucent &= !(PrimitiveSceneProxy.IsOpaqueOrMasked() && !Material.IsDitherMasked());
 			bAnySegmentsCastShadow |= bSegmentCastsShadow;
 			bAllSegmentsCastShadow &= bSegmentCastsShadow;
 			bAnySegmentsDecal |= Material.IsDeferredDecal();
@@ -222,10 +221,15 @@ FRayTracingMaskAndFlags BuildRayTracingInstanceMaskAndFlags(const FRayTracingIns
 
 	// Run AHS for alpha masked and meshes with only some sections casting shadows, which require per mesh section filtering in AHS
 	Result.bForceOpaque = bAllSegmentsOpaque && (bAllSegmentsCastShadow || !bAnySegmentsCastShadow);
+	// Consider that all segments are translucent if none of the mask bits contain Opaque or OpaqueShadow
+	const uint8 OpaqueMask =
+		ComputeRayTracingInstanceMask(ERayTracingInstanceMaskType::Opaque      , RayTracingType) |
+		ComputeRayTracingInstanceMask(ERayTracingInstanceMaskType::OpaqueShadow, RayTracingType);
+	Result.bAllSegmentsTranslucent = Result.Mask != 0 && (Result.Mask & OpaqueMask) == 0;
+
 	Result.bDoubleSided = bDoubleSided;	
 	Result.bAnySegmentsDecal = bAnySegmentsDecal;
 	Result.bAllSegmentsDecal = bAllSegmentsDecal;
-	Result.bAllSegmentsTranslucent = bAllSegmentsTranslucent;
 	Result.bReverseCulling = bAllSegmentsReverseCulling;
 
 	const bool bIsHairStrands = Instance.bThinGeometry;
@@ -233,6 +237,8 @@ FRayTracingMaskAndFlags BuildRayTracingInstanceMaskAndFlags(const FRayTracingIns
 	{
 		// Reset all hair strands bits "on"
 		Result.Mask = ComputeRayTracingInstanceMask(ERayTracingInstanceMaskType::HairStrands, RayTracingType);
+		Result.bForceOpaque = true;
+		Result.bAllSegmentsTranslucent = false;
 	}
 
 	if (!MaskInfo.bVisibleToCamera)
@@ -257,7 +263,6 @@ FRayTracingMaskAndFlags BuildRayTracingInstanceMaskAndFlags(const FRayTracingIns
 		);
 	}
 
-
 	Result.Mask = ApplyFirstPersonRayTracingInstanceMaskFlag(Result.Mask, RayTracingType, MaskInfo.bIsFirstPersonWorldSpaceRepresentation);
 
 	if (Instance.InstanceLayer == ERayTracingInstanceLayer::FarField)
@@ -276,7 +281,7 @@ void SetupRayTracingMeshCommandMaskAndStatus(FRayTracingMeshCommand& MeshCommand
 
 	MeshCommand.bCastRayTracedShadows = MeshBatch.CastRayTracedShadow && MaterialResource.CastsRayTracedShadows() && MaterialResource.GetBlendMode() != BLEND_Additive;
 	MeshCommand.bOpaque = MaterialResource.GetBlendMode() == EBlendMode::BLEND_Opaque && !(VertexFactory->GetType()->SupportsRayTracingProceduralPrimitive() && FDataDrivenShaderPlatformInfo::GetSupportsRayTracingProceduralPrimitive(GMaxRHIShaderPlatform));
-	MeshCommand.bAlphaMasked = MaterialResource.GetBlendMode() == EBlendMode::BLEND_Masked && !MaterialResource.IsDitherMasked();
+	MeshCommand.bAlphaMasked = MaterialResource.GetBlendMode() == EBlendMode::BLEND_Masked; // Used by Lumen only
 	MeshCommand.bDecal = MaterialResource.IsDeferredDecal();
 	MeshCommand.bIsSky = MaterialResource.IsSky();
 	MeshCommand.bTwoSided = MaterialResource.IsTwoSided();
