[TSR] Add hair as thin geometry.

It can be controlled by r.TSR.ThinGeometryDetection.Coverage.Hair (default on).

#tests Test on project silver and console

diff --git a/Engine/Shaders/Private/TemporalSuperResolution/TSRMeasureCoverage.usf b/Engine/Shaders/Private/TemporalSuperResolution/TSRMeasureCoverage.usf
--- a/Engine/Shaders/Private/TemporalSuperResolution/TSRMeasureCoverage.usf
+++ b/Engine/Shaders/Private/TemporalSuperResolution/TSRMeasureCoverage.usf
@@ -10,6 +10,8 @@
 
 #define TILE_ROWS 2
 
+#define CONFIG_HAIR_AS_THIN_GEOMETRY (DIM_HAIR_AS_THIN_GEOMETRY)
+
 //------------------------------------------------------- PARAMETERS
 
 int OutputArrayIndex;
@@ -17,7 +19,11 @@ RWTexture2DArray<float> ThinGeometryCoverageOutput;
 
 uint GetThinGeometry(const FGBufferData GBufferData)
 {
-	return GBufferData.ShadingModelID == SHADINGMODELID_TWOSIDED_FOLIAGE ? 1u : 0u;	
+	return ((GBufferData.ShadingModelID == SHADINGMODELID_TWOSIDED_FOLIAGE) 
+#if CONFIG_HAIR_AS_THIN_GEOMETRY
+		    ||(GBufferData.ShadingModelID == SHADINGMODELID_HAIR)
+#endif
+		    ) ? 1u : 0u;
 }
 
 //------------------------------------------------------- ENTRY POINT
diff --git a/Engine/Source/Runtime/Renderer/Private/PostProcess/TemporalSuperResolution.cpp b/Engine/Source/Runtime/Renderer/Private/PostProcess/TemporalSuperResolution.cpp
--- a/Engine/Source/Runtime/Renderer/Private/PostProcess/TemporalSuperResolution.cpp
+++ b/Engine/Source/Runtime/Renderer/Private/PostProcess/TemporalSuperResolution.cpp
@@ -313,12 +313,16 @@ TAutoConsoleVariable<int32> CVarTSRThinGeometryDetection(
 	TEXT("Other: no history relaxation (Yellow)."),
 	ECVF_RenderThreadSafe);
 
-
 TAutoConsoleVariable<bool> CVarTSRThinGeometryCoverageEdgeReprojection(
 	TEXT("r.TSR.ThinGeometryDetection.Coverage.EdgeReprojection"), true,
 	TEXT("Whether thin geometry edge against not thin geometry should be considered for coverage. Better stability for sparse thin geometry clusters."),
 	ECVF_RenderThreadSafe);
 
+TAutoConsoleVariable<bool> CVarTSRThinGeometryCoverageHair(
+	TEXT("r.TSR.ThinGeometryDetection.Coverage.Hair"), true,
+	TEXT("Include hair as thin geometry and accumulate the coverage."),
+	ECVF_RenderThreadSafe);
+
 TAutoConsoleVariable<float> CVarTSRThinGeometryCoverageMaxRelaxationWeight(
 	TEXT("r.TSR.ThinGeometryDetection.Coverage.MaxRelaxationWeight"), 0.037,
 	TEXT("The max history clamping box relaxation weight due to thin geometry detection (0 to 1)"),
@@ -676,6 +680,9 @@ class FTSRMeasureThinGeometryCoverageCS : public FTSRShader
 	DECLARE_GLOBAL_SHADER(FTSRMeasureThinGeometryCoverageCS);
 	SHADER_USE_PARAMETER_STRUCT(FTSRMeasureThinGeometryCoverageCS, FTSRShader);
 
+	class FHairAsThinGeometryDim : SHADER_PERMUTATION_BOOL("DIM_HAIR_AS_THIN_GEOMETRY");
+	using FPermutationDomain = TShaderPermutationDomain<FHairAsThinGeometryDim>;
+
 	BEGIN_SHADER_PARAMETER_STRUCT(FParameters, )
 		SHADER_PARAMETER_STRUCT(FScreenPassTextureViewportParameters, InputInfo)
 		SHADER_PARAMETER_RDG_UNIFORM_BUFFER(FSceneTextureUniformParameters, SceneTextures)
@@ -1576,10 +1583,18 @@ void AddTSRMeasureThinGeometryCoverage(FRDGBuilder& GraphBuilder, FGlobalShaderM
 	
 	PassParameters->ThinGeometryCoverageOutput = GraphBuilder.CreateUAV(ThinGeometryCoverage.Texture);
 
-	TShaderMapRef<FTSRMeasureThinGeometryCoverageCS> ComputeShader(ShaderMap);
+	bool HairAsThinGeometry = CVarTSRThinGeometryCoverageHair.GetValueOnRenderThread();
+	FTSRMeasureThinGeometryCoverageCS::FPermutationDomain PermutationVector;
+	PermutationVector.Set<FTSRMeasureThinGeometryCoverageCS::FHairAsThinGeometryDim>(HairAsThinGeometry);
+
+	TShaderMapRef<FTSRMeasureThinGeometryCoverageCS> ComputeShader(ShaderMap,PermutationVector);
 	FComputeShaderUtils::AddPass(
 		GraphBuilder,
-		RDG_EVENT_NAME("TSR MeasureThinGeometryCoverage %dx%d", ViewRect.Width(), ViewRect.Height()),
+		RDG_EVENT_NAME("TSR MeasureThinGeometryCoverage(#%d%s) %dx%d", 
+			PermutationVector.ToDimensionValueId(),
+			HairAsThinGeometry ? TEXT(" Hair"): TEXT(""),
+			ViewRect.Width(), 
+			ViewRect.Height()),
 		ComputeShader,
 		PassParameters,
 		FComputeShaderUtils::GetGroupCount(ThinGeometryCoverage.ViewRect.Size(), 8 * 2));
