Refactored the FLightSceneChangeSet to include a bit mask of removed lights and to be uncopiable but persistent (RDG lifetime), and thus more similar to the scene update change set.
- Also removed an unfortunate task dependency between ShadowScene & cache manager (that relied on the removed mask being computed in the former).

#rnx

diff --git a/Engine/Source/Runtime/Renderer/Private/GPUScene.cpp b/Engine/Source/Runtime/Renderer/Private/GPUScene.cpp
--- a/Engine/Source/Runtime/Renderer/Private/GPUScene.cpp
+++ b/Engine/Source/Runtime/Renderer/Private/GPUScene.cpp
@@ -752,7 +752,7 @@ void FGPUScene::UpdateGPULights(FRDGBuilder& GraphBuilder, const UE::Tasks::FTas
 	LightDataArray.SetNum(NumGPULights);
 #endif
 
-	auto* LightDataUploader = TAsyncByteAddressBufferScatterUploader<FLightSceneData>::Allocate(GraphBuilder, LightsPostUpdateData.SceneLightInfoUpdates->NumCommands());
+	auto* LightDataUploader = TAsyncByteAddressBufferScatterUploader<FLightSceneData>::Allocate(GraphBuilder, LightsPostUpdateData.SceneLightInfoUpdates.NumCommands());
 
 	const bool bAsyncSetup = CVarGPUSceneLightsAsyncSetup.GetValueOnRenderThread() != 0;
 
@@ -762,7 +762,7 @@ void FGPUScene::UpdateGPULights(FRDGBuilder& GraphBuilder, const UE::Tasks::FTas
 			const bool bAllowStaticLighting = IsStaticLightingAllowed();
 			const uint32 Flags = RenderRectLightsAsSpotLights(FeatureLevel) ? ELightShaderParameterFlags::RectAsSpotLight : 0u;
 
-			ChangeSet.SceneLightInfoUpdates->ForEachCommand(ESceneUpdateCommandFilter::AddedUpdated, [&](const FUpdateLightCommand& UpdateLightCommand)
+			ChangeSet.SceneLightInfoUpdates.ForEachCommand(ESceneUpdateCommandFilter::AddedUpdated, [&](const FUpdateLightCommand& UpdateLightCommand)
 				{
 					int32 Index = UpdateLightCommand.GetSceneInfo()->Id;
 
@@ -810,7 +810,7 @@ void FGPUScene::UpdateGPULights(FRDGBuilder& GraphBuilder, const UE::Tasks::FTas
 		UE::Tasks::ETaskPriority::Normal,
 		bAsyncSetup);
 
-	if (!LightsPostUpdateData.SceneLightInfoUpdates->IsEmpty())
+	if (!LightsPostUpdateData.SceneLightInfoUpdates.IsEmpty())
 	{
 		if (bAsyncSetup)
 		{
diff --git a/Engine/Source/Runtime/Renderer/Private/RendererScene.cpp b/Engine/Source/Runtime/Renderer/Private/RendererScene.cpp
--- a/Engine/Source/Runtime/Renderer/Private/RendererScene.cpp
+++ b/Engine/Source/Runtime/Renderer/Private/RendererScene.cpp
@@ -2632,7 +2632,7 @@ void FScene::AddLightSceneInfo_RenderThread(FLightSceneInfo* LightSceneInfo)
 	// Add the light to the light list.
 	LightSceneInfo->Id = Lights.Add(FLightSceneInfoCompact(LightSceneInfo));
 	const FLightSceneInfoCompact& LightSceneInfoCompact = Lights[LightSceneInfo->Id];
-	const ELightComponentType LightType = (ELightComponentType)LightSceneInfo->Proxy->GetLightType();
+	const ELightComponentType LightType = ELightComponentType(LightSceneInfoCompact.LightType);
 	const bool bDirectionalLight = LightType == LightType_Directional;
 
 	if (bDirectionalLight)
@@ -5162,25 +5162,35 @@ void FScene::UpdateLights(FRDGBuilder& GraphBuilder, FSceneExtensionsUpdaters& S
 	{
 		TArray<int32, SceneRenderingAllocator> RemovedLightIds;
 		TArray<int32, SceneRenderingAllocator> AddedLightIds;
-		FFLightSceneChangeSetAllocation(FSceneLightInfoUpdates&& InSceneLightInfoUpdates) : SceneLightInfoUpdates(MoveTemp(InSceneLightInfoUpdates)) {}
+		FFLightSceneChangeSetAllocation(FSceneLightInfoUpdates&& InSceneLightInfoUpdates, int32 InPreUpdateMaxIndex) 
+			: SceneLightInfoUpdates(MoveTemp(InSceneLightInfoUpdates)) 
+			, RemovedLightsMask(false, InPreUpdateMaxIndex)
+			, PreUpdateChangeSet{TConstArrayView<int32>(RemovedLightIds), TConstArrayView<int32>(), SceneLightInfoUpdates, RemovedLightsMask, InPreUpdateMaxIndex, -1 }
+			, PostUpdateChangeSet{TConstArrayView<int32>(RemovedLightIds), TConstArrayView<int32>(AddedLightIds), SceneLightInfoUpdates, RemovedLightsMask, InPreUpdateMaxIndex, -1 }		
+		{}
 		FSceneLightInfoUpdates SceneLightInfoUpdates;
+		TBitArray<SceneRenderingAllocator> RemovedLightsMask;
+
+		// These are the accessors to the data that are read only and can be passed around by pointer or reference as they are kept alive in the alloc.
+		FLightSceneChangeSet PreUpdateChangeSet;
+		FLightSceneChangeSet PostUpdateChangeSet;
 	};
 	// Allocate change set storage with graph builder lifetime such that we can safely pass it to async tasks.
-	FFLightSceneChangeSetAllocation& ChangeSet = *GraphBuilder.AllocObject<FFLightSceneChangeSetAllocation>(MoveTemp(*SceneLightInfoUpdates));
-
+	FFLightSceneChangeSetAllocation& ChangeSetAlloc = *GraphBuilder.AllocObject<FFLightSceneChangeSetAllocation>(MoveTemp(*SceneLightInfoUpdates), Lights.GetMaxIndex());
 #if DO_CHECK
-	FSceneLightInfoUpdates::FReadAccessScope ReadAccessScope(ChangeSet.SceneLightInfoUpdates);
+	FSceneLightInfoUpdates::FReadAccessScope ReadAccessScope(ChangeSetAlloc.SceneLightInfoUpdates);
 #endif
 
 	// Filter out removes & never-adds
-	ChangeSet.SceneLightInfoUpdates.ForEachCommand([&](const FUpdateLightCommand& UpdateLightCommand)
+	ChangeSetAlloc.SceneLightInfoUpdates.ForEachCommand([&](const FUpdateLightCommand& UpdateLightCommand)
 	{
 		int32 Id = UpdateLightCommand.GetSceneInfo()->Id;
 		if (UpdateLightCommand.IsDelete())
 		{
 			if (Id != INDEX_NONE)
 			{
-				ChangeSet.RemovedLightIds.Add(Id);
+				ChangeSetAlloc.RemovedLightIds.Add(Id);
+				ChangeSetAlloc.RemovedLightsMask[Id] = true;
 			}
 			else
 			{
@@ -5189,11 +5199,14 @@ void FScene::UpdateLights(FRDGBuilder& GraphBuilder, FSceneExtensionsUpdaters& S
 			}
 		}
 	});
-	SceneExtensionsUpdaters.PreLightsUpdate(GraphBuilder, FLightSceneChangeSet{ ChangeSet.RemovedLightIds, TConstArrayView<int32>(), &ChangeSet.SceneLightInfoUpdates });
+	ChangeSetAlloc.PreUpdateChangeSet.RemovedLightIds = ChangeSetAlloc.RemovedLightIds;
+	ChangeSetAlloc.PostUpdateChangeSet.RemovedLightIds = ChangeSetAlloc.RemovedLightIds;
+
+	SceneExtensionsUpdaters.PreLightsUpdate(GraphBuilder, ChangeSetAlloc.PreUpdateChangeSet);
 	// This can't access the scene light data if done async since it happens before the actual removals.
-	OnPreLightSceneInfoUpdate.Broadcast(GraphBuilder, FLightSceneChangeSet{ ChangeSet.RemovedLightIds, TConstArrayView<int32>(), &ChangeSet.SceneLightInfoUpdates });
+	OnPreLightSceneInfoUpdate.Broadcast(GraphBuilder, ChangeSetAlloc.PreUpdateChangeSet);
 	// Batch process all light removes
-	for (int32 LightId : ChangeSet.RemovedLightIds)
+	for (int32 LightId : ChangeSetAlloc.RemovedLightIds)
 	{
 		FLightSceneInfo* LightSceneInfo = Lights[LightId].LightSceneInfo;
 		FScopeCycleCounter Context(LightSceneInfo->Proxy->GetStatId());
@@ -5201,7 +5214,7 @@ void FScene::UpdateLights(FRDGBuilder& GraphBuilder, FSceneExtensionsUpdaters& S
 	}
 
 	// Process all light adds & updates
-	ChangeSet.SceneLightInfoUpdates.ForEachCommand(ESceneUpdateCommandFilter::AddedUpdated, [&](const FUpdateLightCommand& UpdateLightCommand)
+	ChangeSetAlloc.SceneLightInfoUpdates.ForEachCommand(ESceneUpdateCommandFilter::AddedUpdated, [&](const FUpdateLightCommand& UpdateLightCommand)
 	{
 		FLightSceneInfo* LightSceneInfo = UpdateLightCommand.GetSceneInfo();
 		FScopeCycleCounter Context(LightSceneInfo->Proxy->GetStatId());
@@ -5210,14 +5223,14 @@ void FScene::UpdateLights(FRDGBuilder& GraphBuilder, FSceneExtensionsUpdaters& S
 		const bool bHasId = Id != INDEX_NONE;
 		check(bHasId == !UpdateLightCommand.IsAdd());
 		// Directly process updates.
-		if (FUpdateLightTransformParameters* TransformParameters = ChangeSet.SceneLightInfoUpdates.GetPayloadPtr<FUpdateLightTransformParameters>(UpdateLightCommand))
+		if (FUpdateLightTransformParameters* TransformParameters = ChangeSetAlloc.SceneLightInfoUpdates.GetPayloadPtr<FUpdateLightTransformParameters>(UpdateLightCommand))
 		{
 			SCOPE_CYCLE_COUNTER(STAT_UpdateSceneLightTime);
 
 			UpdateLightTransform_RenderThread(Id, LightSceneInfo, *TransformParameters);
 		}
 
-		if (FUpdateLightColorParameters* ColorParameters = ChangeSet.SceneLightInfoUpdates.GetPayloadPtr<FUpdateLightColorParameters>(UpdateLightCommand))
+		if (FUpdateLightColorParameters* ColorParameters = ChangeSetAlloc.SceneLightInfoUpdates.GetPayloadPtr<FUpdateLightColorParameters>(UpdateLightCommand))
 		{
 			const FUpdateLightColorParameters& NewParameters = *ColorParameters;
 			SCOPE_CYCLE_COUNTER(STAT_UpdateSceneLightTime);
@@ -5252,13 +5265,14 @@ void FScene::UpdateLights(FRDGBuilder& GraphBuilder, FSceneExtensionsUpdaters& S
 		{
 			AddLightSceneInfo_RenderThread(LightSceneInfo);
 			// Note: Id is set in AddLightSceneInfo_RenderThread so we must fetch it again
-			ChangeSet.AddedLightIds.Add(LightSceneInfo->Id);
+			ChangeSetAlloc.AddedLightIds.Add(LightSceneInfo->Id);
 		}
 	});
-
-	SceneExtensionsUpdaters.PostLightsUpdate(GraphBuilder, FLightSceneChangeSet{ ChangeSet.RemovedLightIds, ChangeSet.AddedLightIds, &ChangeSet.SceneLightInfoUpdates });
-	OnPostLightSceneInfoUpdate.Broadcast(GraphBuilder, FLightSceneChangeSet{ ChangeSet.RemovedLightIds, ChangeSet.AddedLightIds, &ChangeSet.SceneLightInfoUpdates });
-	GPUScene.OnPostLightSceneInfoUpdate(GraphBuilder, FLightSceneChangeSet{ ChangeSet.RemovedLightIds, ChangeSet.AddedLightIds, &ChangeSet.SceneLightInfoUpdates });
+	ChangeSetAlloc.PostUpdateChangeSet.PostUpdateMaxIndex = Lights.GetMaxIndex();
+	ChangeSetAlloc.PostUpdateChangeSet.AddedLightIds = ChangeSetAlloc.AddedLightIds;
+	SceneExtensionsUpdaters.PostLightsUpdate(GraphBuilder, ChangeSetAlloc.PostUpdateChangeSet);
+	OnPostLightSceneInfoUpdate.Broadcast(GraphBuilder, ChangeSetAlloc.PostUpdateChangeSet);
+	GPUScene.OnPostLightSceneInfoUpdate(GraphBuilder, ChangeSetAlloc.PostUpdateChangeSet);
 }
 
 template<class T>
diff --git a/Engine/Source/Runtime/Renderer/Private/ScenePrivate.h b/Engine/Source/Runtime/Renderer/Private/ScenePrivate.h
--- a/Engine/Source/Runtime/Renderer/Private/ScenePrivate.h
+++ b/Engine/Source/Runtime/Renderer/Private/ScenePrivate.h
@@ -2838,8 +2838,6 @@ struct FUpdateLightColorParameters: public FSceneLightInfoUpdates::TPayloadBase<
 
 /**
  * Describes all light modifications to the scene by recording the light scene IDs.
- * TODO: If needed, we could add a reference to the FLightUpdates (which contains the commands) since this would enable systems to consume out the delta updates as they come in.
- *       If this is useful we must ensure FLightUpdates are kept alive as long as any async tasks might require.
  */
 struct FLightSceneChangeSet
 {
@@ -2847,7 +2845,11 @@ struct FLightSceneChangeSet
 	TConstArrayView<int32> RemovedLightIds;
 	// IDs of all lights added to the scene, only available after all lights are added to the scene, may contain the same ID's as removed, as they may be reused.
 	TConstArrayView<int32> AddedLightIds;
-	FSceneLightInfoUpdates *SceneLightInfoUpdates = nullptr;
+	FSceneLightInfoUpdates& SceneLightInfoUpdates;
+	// Bit-array of at least size PreUpdateMaxIndex, with a bit set for each removed light
+	const TBitArray<SceneRenderingAllocator>& RemovedLightsMask;
+	int32 PreUpdateMaxIndex = 0;
+	int32 PostUpdateMaxIndex = 0;
 };
 
 struct FViewSceneChangeSet
diff --git a/Engine/Source/Runtime/Renderer/Private/Shadows/ShadowScene.cpp b/Engine/Source/Runtime/Renderer/Private/Shadows/ShadowScene.cpp
--- a/Engine/Source/Runtime/Renderer/Private/Shadows/ShadowScene.cpp
+++ b/Engine/Source/Runtime/Renderer/Private/Shadows/ShadowScene.cpp
@@ -3,7 +3,6 @@
 #include "ShadowScene.h"
 #include "ScenePrivate.h"
 #include "Algo/ForEach.h"
-#include "VirtualShadowMaps/VirtualShadowMapCacheManager.h"
 #include "ShadowSceneRenderer.h"
 
 #if !(UE_BUILD_SHIPPING || UE_BUILD_TEST)
@@ -49,7 +48,7 @@ void FShadowScene::FUpdater::PostLightsUpdate(FRDGBuilder& GraphBuilder, const F
 	// don't spawn async work for no good reason.
 	constexpr int32 kMinWorkSizeForAsync = 64;
 
-	int32 WorkSize = LightSceneChangeSet.SceneLightInfoUpdates->NumCommands();
+	int32 WorkSize = LightSceneChangeSet.SceneLightInfoUpdates.NumCommands();
 
 	// Don't sync, or kick off a new job if there is no work to do.
 	if (WorkSize > 0)
@@ -58,22 +57,12 @@ void FShadowScene::FUpdater::PostLightsUpdate(FRDGBuilder& GraphBuilder, const F
 		ShadowScene.SceneChangeUpdateTask.Wait();
 		ShadowScene.SceneChangeUpdateTask = GraphBuilder.AddSetupTask([this, LightSceneChangeSet]
 		{
-			// Collect a bit mask of all removed lights so we can prune the cache entries in one sweep
-			TBitArray<> RemovedLightMask;
-			RemovedLightMask.SetNum(ShadowScene.Lights.GetMaxIndex(), false);
-
 			// Oust all removed Ids.
 			for (int32 Id : LightSceneChangeSet.RemovedLightIds)
 			{
 				check(ShadowScene.Lights.IsValidIndex(Id));
 				ShadowScene.ActiveLights[Id] = false;
 				ShadowScene.Lights.RemoveAt(Id);
-				RemovedLightMask[Id] = true;
-			}
-
-			if (auto CacheManager = ShadowScene.Scene.GetVirtualShadowMapCache())
-			{
-				CacheManager->ProcessRemovedLights(RemovedLightMask);
 			}
 	
 			// Track active lights (those that are or were moving, and thus need updating)
@@ -99,7 +88,7 @@ void FShadowScene::FUpdater::PostLightsUpdate(FRDGBuilder& GraphBuilder, const F
 			};
 
 			Algo::ForEach(LightSceneChangeSet.AddedLightIds, UpdateLight);
-			for (const auto& Item : LightSceneChangeSet.SceneLightInfoUpdates->GetRangeView<FUpdateLightTransformParameters>())
+			for (const auto& Item : LightSceneChangeSet.SceneLightInfoUpdates.GetRangeView<FUpdateLightTransformParameters>())
 			{ 
 				UpdateLight(Item.SceneInfo->Id);
 			}
diff --git a/Engine/Source/Runtime/Renderer/Private/VirtualShadowMaps/VirtualShadowMapCacheManager.cpp b/Engine/Source/Runtime/Renderer/Private/VirtualShadowMaps/VirtualShadowMapCacheManager.cpp
--- a/Engine/Source/Runtime/Renderer/Private/VirtualShadowMaps/VirtualShadowMapCacheManager.cpp
+++ b/Engine/Source/Runtime/Renderer/Private/VirtualShadowMaps/VirtualShadowMapCacheManager.cpp
@@ -1608,7 +1608,7 @@ void FVirtualShadowMapArrayCacheManager::ProcessInvalidations(FRDGBuilder& Graph
 	}
 }
 
-void FVirtualShadowMapArrayCacheManager::ProcessRemovedLights(const TBitArray<>& RemovedLightMask)
+void FVirtualShadowMapArrayCacheManager::ProcessRemovedLights(const TBitArray<SceneRenderingAllocator>& RemovedLightMask)
 {
 	for (auto It = CacheEntries.CreateIterator(); It; ++It)
 	{
@@ -1731,6 +1731,14 @@ FVirtualShadowMapInvalidationSceneUpdater::FVirtualShadowMapInvalidationSceneUpd
 	: CacheManager(InCacheManager)
 {}
 
+void FVirtualShadowMapInvalidationSceneUpdater::PreLightsUpdate(FRDGBuilder& GraphBuilder, const FLightSceneChangeSet& LightSceneChangeSet)
+{
+	if (LightSceneChangeSet.RemovedLightIds.Num() > 0)
+	{
+		CacheManager.ProcessRemovedLights(LightSceneChangeSet.RemovedLightsMask);
+	}
+}
+
 void FVirtualShadowMapInvalidationSceneUpdater::PreSceneUpdate(FRDGBuilder& GraphBuilder, const FScenePreUpdateChangeSet& ChangeSet, FSceneUniformBuffer& SceneUniforms)
 {
 	SCOPED_NAMED_EVENT(FScene_VirtualShadowCacheUpdate, FColor::Orange);
@@ -1747,12 +1755,8 @@ void FVirtualShadowMapInvalidationSceneUpdater::PreSceneUpdate(FRDGBuilder& Grap
 		}
 	}
 
-
 	FShadowScene& ShadowScene = CacheManager.Scene.GetExtension<FShadowScene>();
 
-	// Note: waits for the light update task that the shadow scene uses because it calls ProcessRemovedLights from that thread.
-	ShadowScene.WaitForSceneLightsUpdateTask();
-
 	// Needs to be called before the first time we start adding invalidations.
 	// There may be a way to avoid doing this both in pre and post, but it is pretty light if there is nothing to do anyways.
 	CacheManager.ReallocatePersistentPrimitiveIndices();
@@ -1761,6 +1765,8 @@ void FVirtualShadowMapInvalidationSceneUpdater::PreSceneUpdate(FRDGBuilder& Grap
 	{
 		FVirtualShadowMapArrayCacheManager::FInvalidatingPrimitiveCollector InvalidatingPrimitiveCollector(&CacheManager);
 
+		// TODO: Note that there is an ordering dependency here on ShadowScene, which may not have updated the AlwaysInvalidatingPrimitives list at this point. 
+		//       Probably harmless (as they will get invalidated next frame etc) but something to think about.
 		// Primitives that are tracked as always invalidating shadows, pipe through as transform updates
 		for (FPrimitiveSceneInfo* PrimitiveSceneInfo : ShadowScene.GetAlwaysInvalidatingPrimitives())
 		{
diff --git a/Engine/Source/Runtime/Renderer/Private/VirtualShadowMaps/VirtualShadowMapCacheManager.h b/Engine/Source/Runtime/Renderer/Private/VirtualShadowMaps/VirtualShadowMapCacheManager.h
--- a/Engine/Source/Runtime/Renderer/Private/VirtualShadowMaps/VirtualShadowMapCacheManager.h
+++ b/Engine/Source/Runtime/Renderer/Private/VirtualShadowMaps/VirtualShadowMapCacheManager.h
@@ -394,11 +394,6 @@ public:
 		FSceneUniformBuffer &SceneUniformBuffer,
 		FInvalidatingPrimitiveCollector& InvalidatingPrimitiveCollector);
 
-	/**
-	 * Handle light removal, need to clear out cache entries as the ID may be reused after this.
-	 */
-	void ProcessRemovedLights(const TBitArray<>& RemovedLightMask);
-
 	uint64 GetGPUSizeBytes(bool bLogSizes) const;
 
 	const FVirtualShadowMapArrayFrameData& GetPrevBuffers() const { return PrevBuffers; }
@@ -428,6 +423,12 @@ public:
 	void ReallocatePersistentPrimitiveIndices();
 
 private:
+
+	/**
+	 * Handle light removal, need to clear out cache entries as the ID may be reused after this.
+	 */
+	void ProcessRemovedLights(const TBitArray<SceneRenderingAllocator>& RemovedLightMask);
+
 	friend class FVirtualShadowMapInvalidationSceneRenderer;
 	friend FVirtualShadowMapArray;
 
@@ -563,6 +564,7 @@ class FVirtualShadowMapInvalidationSceneUpdater : public ISceneExtensionUpdater
 public:
 	FVirtualShadowMapInvalidationSceneUpdater(FVirtualShadowMapArrayCacheManager& InCacheManager);
 
+	virtual void PreLightsUpdate(FRDGBuilder& GraphBuilder, const FLightSceneChangeSet& LightSceneChangeSet);
 	virtual void PreSceneUpdate(FRDGBuilder& GraphBuilder, const FScenePreUpdateChangeSet& ChangeSet, FSceneUniformBuffer& SceneUniforms) override;
 	virtual void PostSceneUpdate(FRDGBuilder& GraphBuilder, const FScenePostUpdateChangeSet& ChangeSet) override;
 	virtual void PostGPUSceneUpdate(FRDGBuilder& GraphBuilder, FSceneUniformBuffer& SceneUniforms) override;
