Moved more shadow rendering logic inside the FShadowScene[Renderer]

diff --git a/Engine/Source/Runtime/Renderer/Private/DeferredShadingRenderer.cpp b/Engine/Source/Runtime/Renderer/Private/DeferredShadingRenderer.cpp
--- a/Engine/Source/Runtime/Renderer/Private/DeferredShadingRenderer.cpp
+++ b/Engine/Source/Runtime/Renderer/Private/DeferredShadingRenderer.cpp
@@ -2755,29 +2755,18 @@ void FDeferredShadingSceneRenderer::Render(FRDGBuilder& GraphBuilder, const FSce
 			// If we haven't already rendered shadow maps, render them now (due to forward shading or r.shadow.ShadowMapsRenderEarly)
 			if (!bShadowMapsRenderedEarly)
 			{
-				if (VirtualShadowMapArray.IsEnabled())
+				// We use this somewhat awkward callback to get this into the right scope & place in the timeline (relying on some syncs which we can probably clean up in the future).
+				auto RenderVirtualShadowMapsFunc = [&](bool bNaniteEnabled)
 				{
-					// TODO: actually move this inside RenderShadowDepthMaps instead of this extra scope to make it 1:1 with profiling captures/traces
-					RDG_EVENT_SCOPE_STAT(GraphBuilder, ShadowDepths, "ShadowDepths");
-					RDG_GPU_STAT_SCOPE(GraphBuilder, ShadowDepths);
-
-					ensureMsgf(AreLightsInLightGrid(), TEXT("Virtual shadow map setup requires local lights to be injected into the light grid (this may be caused by 'r.LightCulling.Quality=0')."));
-
-					FFrontLayerTranslucencyData FrontLayerTranslucencyData = RenderFrontLayerTranslucency(GraphBuilder, Views, SceneTextures, true /*VSM page marking*/);
-
-					VirtualShadowMapArray.BuildPageAllocations(GraphBuilder,
-						GetActiveSceneTextures(),
-						Views,
-						*GatherAndSortLightsTask.GetResult(),
-						VisibleLightInfos,
-						SingleLayerWaterPrePassResult,
-						FrontLayerTranslucencyData,
-						FroxelRenderer,
-						ShadowSceneRenderer.AreAnyLocalLightsPreset()
-					);
-				}
+					FFrontLayerTranslucencyData FrontLayerTranslucencyData;
+					if (ShadowSceneRenderer.GetVirtualShadowMapArray().IsEnabled())
+					{
+						FrontLayerTranslucencyData = RenderFrontLayerTranslucency(GraphBuilder, Views, SceneTextures, true /*VSM page marking*/);
+					}
+					ShadowSceneRenderer.RenderVirtualShadowMaps(GraphBuilder, bNaniteEnabled, SingleLayerWaterPrePassResult, FrontLayerTranslucencyData, FroxelRenderer);
+				};
 
-				RenderShadowDepthMaps(GraphBuilder, InitViewTaskDatas.DynamicShadows, InstanceCullingManager, ExternalAccessQueue);
+				RenderShadowDepthMaps(GraphBuilder, InitViewTaskDatas.DynamicShadows, InstanceCullingManager, ExternalAccessQueue, RenderVirtualShadowMapsFunc);
 			}
 			CheckShadowDepthRenderCompleted();
 
diff --git a/Engine/Source/Runtime/Renderer/Private/SceneRendering.cpp b/Engine/Source/Runtime/Renderer/Private/SceneRendering.cpp
--- a/Engine/Source/Runtime/Renderer/Private/SceneRendering.cpp
+++ b/Engine/Source/Runtime/Renderer/Private/SceneRendering.cpp
@@ -566,21 +566,11 @@ static TSharedPtr<FVirtualShadowMapClipmap> FindShadowClipmapForView(const TArra
 	return TSharedPtr<FVirtualShadowMapClipmap>();
 }
 
-TSharedPtr<FVirtualShadowMapClipmap> FVisibleLightInfo::FindShadowClipmapForView(const FViewInfo* View) const
-{
-	return ::FindShadowClipmapForView(VirtualShadowMapClipmaps, View);
-}
-
-TSharedPtr<FVirtualShadowMapClipmap> FVisibleLightInfo::FindFirstPersonShadowClipmapForView(const FViewInfo* View) const
-{
-	return ::FindShadowClipmapForView(FirstPersonVirtualShadowMapClipmaps, View);
-}
-
 int32 FVisibleLightInfo::GetVirtualShadowMapId(const FViewInfo* View) const
 {
 	if (VirtualShadowMapClipmaps.Num())
 	{
-		return FindShadowClipmapForView(View)->GetVirtualShadowMapId();
+		return FindShadowClipmapForView(VirtualShadowMapClipmaps, View)->GetVirtualShadowMapId();
 	}
 	else
 	{
diff --git a/Engine/Source/Runtime/Renderer/Private/SceneRendering.h b/Engine/Source/Runtime/Renderer/Private/SceneRendering.h
--- a/Engine/Source/Runtime/Renderer/Private/SceneRendering.h
+++ b/Engine/Source/Runtime/Renderer/Private/SceneRendering.h
@@ -221,22 +221,7 @@ public:
 	/** A list of per-object shadows that were occluded. We need to track these so we can issue occlusion queries for them. */
 	TArray<FProjectedShadowInfo*,SceneRenderingAllocator> OccludedPerObjectShadows;
 
-	/** Virtual shadow map clipmap shadows */
-	TArray<TSharedPtr<FVirtualShadowMapClipmap>,SceneRenderingAllocator> VirtualShadowMapClipmaps;
-
-	/** First-person clipmap shadows */
-	TArray<TSharedPtr<FVirtualShadowMapClipmap>,SceneRenderingAllocator> FirstPersonVirtualShadowMapClipmaps;
-
-	/**
-	 * Search VirtualShadowMapClipmaps for the one that was set up with the given view. 
-	 * TODO: should VirtualShadowMapClipmap move to FVisibleLightViewInfo?
-	 */
-	TSharedPtr<FVirtualShadowMapClipmap> FindShadowClipmapForView(const FViewInfo* View) const;
-
-	/**
-	 * Find (if any) first-person special clipmap is present for the given view.
-	 */
-	TSharedPtr<FVirtualShadowMapClipmap> FindFirstPersonShadowClipmapForView(const FViewInfo* View) const;
+	TArray<TSharedPtr<FVirtualShadowMapClipmap>, SceneRenderingAllocator> VirtualShadowMapClipmaps;
 
 	/** 
 	* Returns true if the light has only virtual shadow maps (or no shadows at all), i.e. no conventional shadow maps that are allocated
@@ -1905,8 +1890,6 @@ struct FSortedShadowMaps
 
 	TArray<FProjectedShadowInfo*, SceneRenderingAllocator> VirtualShadowMapShadows;
 
-	TArray<TSharedPtr<FVirtualShadowMapClipmap>, SceneRenderingAllocator> VirtualShadowMapClipmaps;
-
 	TArray<FSortedShadowMapAtlas,SceneRenderingAllocator> CompleteShadowMapAtlases;
 
 	void Release();
@@ -1938,7 +1921,7 @@ struct FSortedShadowMaps
 	bool IsEmpty() const
 	{
 		return ShadowMapAtlases.IsEmpty() && ShadowMapCubemaps.IsEmpty() && PreshadowCache.Shadows.IsEmpty() && TranslucencyShadowMapAtlases.IsEmpty() &&
-			VirtualShadowMapShadows.IsEmpty() && VirtualShadowMapClipmaps.IsEmpty() && CompleteShadowMapAtlases.IsEmpty();
+			VirtualShadowMapShadows.IsEmpty() && CompleteShadowMapAtlases.IsEmpty();
 	}
 };
 
@@ -2321,6 +2304,11 @@ public:
 	static TGlobalResource<FGlobalDynamicReadBuffer> DynamicReadBufferForRayTracing;
 	static TGlobalResource<FGlobalDynamicReadBuffer> DynamicReadBufferForShadows;
 
+	bool IsRenderingStereo() const
+	{
+		return Views.Num() == 2 && IStereoRendering::IsASecondaryView(Views[1]);
+	}
+
 protected:
 
 	/** Size of the family. */
@@ -2455,7 +2443,7 @@ protected:
 	void FinishGatherShadowPrimitives(FDynamicShadowsTaskData* TaskData);
 
 	void RenderShadowDepthMaps(FRDGBuilder& GraphBuilder, FDynamicShadowsTaskData* DynamicShadowsTaskData, FInstanceCullingManager& InstanceCullingManager, FRDGExternalAccessQueue& ExternalAccessQueue);
-	void RenderVirtualShadowMaps(FRDGBuilder& GraphBuilder, bool bNaniteEnabled);
+	void RenderShadowDepthMaps(FRDGBuilder& GraphBuilder, FDynamicShadowsTaskData* DynamicShadowsTaskData, FInstanceCullingManager& InstanceCullingManager, FRDGExternalAccessQueue& ExternalAccessQueue, TFunctionRef<void(bool bNaniteEnabled)> RenderAdditionalShadowDepths);
 
 	void RenderShadowDepthMapAtlases(FRDGBuilder& GraphBuilder);
 
diff --git a/Engine/Source/Runtime/Renderer/Private/ShadowDepthRendering.cpp b/Engine/Source/Runtime/Renderer/Private/ShadowDepthRendering.cpp
--- a/Engine/Source/Runtime/Renderer/Private/ShadowDepthRendering.cpp
+++ b/Engine/Source/Runtime/Renderer/Private/ShadowDepthRendering.cpp
@@ -144,7 +144,7 @@ static TAutoConsoleVariable<int32> CVarNaniteShadowsUseHZB(
 	TEXT("Enables HZB for Nanite shadows."),
 	ECVF_RenderThreadSafe);
 
-static TAutoConsoleVariable<int32> CVarNaniteShadowsUpdateStreaming(
+TAutoConsoleVariable<int32> CVarNaniteShadowsUpdateStreaming(
 	TEXT("r.Shadow.NaniteUpdateStreaming"),
 	1,
 	TEXT("Produce Nanite geometry streaming requests from shadow map rendering."),
@@ -1681,15 +1681,13 @@ void FSceneRenderer::RenderShadowDepthMapAtlases(FRDGBuilder& GraphBuilder)
 	ExternalAccessQueue.Submit(GraphBuilder);
 }
 
-void FSceneRenderer::RenderVirtualShadowMaps(FRDGBuilder& GraphBuilder, bool bNaniteEnabled)
+void FSceneRenderer::RenderShadowDepthMaps(FRDGBuilder& GraphBuilder, FDynamicShadowsTaskData* DynamicShadowsTaskData, FInstanceCullingManager& InstanceCullingManager, FRDGExternalAccessQueue& ExternalAccessQueue)
 {
-	if (FShadowSceneRenderer* ShadowSceneRenderer = GetSceneExtensionsRenderers().GetRendererPtr<FShadowSceneRenderer>())
-	{
-		ShadowSceneRenderer->RenderVirtualShadowMaps(GraphBuilder, bNaniteEnabled, CVarNaniteShadowsUpdateStreaming.GetValueOnRenderThread() != 0);
-	}
+	auto DummyExtraShadowDepthRenderFunc = [](bool){}; 
+	RenderShadowDepthMaps(GraphBuilder, DynamicShadowsTaskData, InstanceCullingManager, ExternalAccessQueue, DummyExtraShadowDepthRenderFunc);
 }
 
-void FSceneRenderer::RenderShadowDepthMaps(FRDGBuilder& GraphBuilder, FDynamicShadowsTaskData* DynamicShadowsTaskData, FInstanceCullingManager &InstanceCullingManager, FRDGExternalAccessQueue& ExternalAccessQueue)
+void FSceneRenderer::RenderShadowDepthMaps(FRDGBuilder& GraphBuilder, FDynamicShadowsTaskData* DynamicShadowsTaskData, FInstanceCullingManager &InstanceCullingManager, FRDGExternalAccessQueue& ExternalAccessQueue, TFunctionRef<void(bool bNaniteEnabled)> RenderAdditionalShadowDepths)
 {
 	ensureMsgf(!bShadowDepthRenderCompleted, TEXT("RenderShadowDepthMaps called twice in the same frame"));
 
@@ -1727,7 +1725,7 @@ void FSceneRenderer::RenderShadowDepthMaps(FRDGBuilder& GraphBuilder, FDynamicSh
 		ViewFamily.EngineShowFlags.NaniteMeshes &&
 		Nanite::GStreamingManager.HasResourceEntries();
 
-	RenderVirtualShadowMaps(GraphBuilder, bNaniteEnabled);
+	RenderAdditionalShadowDepths(bNaniteEnabled);
 
 	// Render non-VSM shadows
 	RenderShadowDepthMapAtlases(GraphBuilder);
diff --git a/Engine/Source/Runtime/Renderer/Private/ShadowSetup.cpp b/Engine/Source/Runtime/Renderer/Private/ShadowSetup.cpp
--- a/Engine/Source/Runtime/Renderer/Private/ShadowSetup.cpp
+++ b/Engine/Source/Runtime/Renderer/Private/ShadowSetup.cpp
@@ -5668,38 +5668,8 @@ void FSceneRenderer::AddViewDependentWholeSceneShadowsForView(
 
 			if (bNeedsVirtualShadowMap)
 			{
-				TSharedPtr<FVirtualShadowMapClipmap> VirtualShadowMapClipmap = TSharedPtr<FVirtualShadowMapClipmap>(new FVirtualShadowMapClipmap(
-					VirtualShadowMapArray,
-					LightSceneInfo,
-					View.ViewMatrices,
-					View.ViewRect.Size(),
-					&View,
-					Scene->GetExtension<FShadowScene>().GetLightMobilityFactor(LightSceneInfo.Id)
-				));
-
-				// TODO: Not clear we need both of these in this path, but keep it consistent for now
-				VisibleLightInfo.VirtualShadowMapClipmaps.Add(VirtualShadowMapClipmap);
-				SortedShadowsForShadowDepthPass.VirtualShadowMapClipmaps.Add(VirtualShadowMapClipmap);
-
-				if (GEnableNonNaniteVSM != 0)
-				{
-					FShadowSceneRenderer& ShadowSceneRenderer = GetSceneExtensionsRenderers().GetRenderer<FShadowSceneRenderer>();
-
-					// Create the projected shadow info to make sure that culling happens.
-					FProjectedShadowInfo* ProjectedShadowInfo = Allocator.Create<FProjectedShadowInfo>();
-					ProjectedShadowInfo->SetupClipmapProjection(&LightSceneInfo, &View, VirtualShadowMapClipmap, CVarVsmUseFarShadowRules.GetValueOnRenderThread() != 0 ? MaxNonFarCascadeDistance : -1.0f);
-					VisibleLightInfo.AllProjectedShadows.Add(ProjectedShadowInfo);
-					ShadowInfosThatNeedCulling.Add(ProjectedShadowInfo);
-
-					ProjectedShadowInfo->VirtualShadowMapPerLightCacheEntry = VirtualShadowMapClipmap->GetCacheEntry();
-
-					// TODO: disentangle from legacy setup alltogether
-					// TODO2: This needs to not depend on the non-nanite shadows!
-					ShadowSceneRenderer.AddDirectionalLightShadow(ProjectedShadowInfo);
-				}
-
-				// NOTE: If there are multiple camera views this will simply be associated with "one of them"
-				VisibleLightInfo.VirtualShadowMapId = VirtualShadowMapClipmap->GetVirtualShadowMapId();
+				FShadowSceneRenderer& ShadowSceneRenderer = GetSceneExtensionsRenderers().GetRenderer<FShadowSceneRenderer>();
+				ShadowSceneRenderer.AddDirectionalLightShadow(LightSceneInfo, View, ShadowInfosThatNeedCulling);
 			}
 		}
 	}
diff --git a/Engine/Source/Runtime/Renderer/Private/Shadows/ShadowScene.cpp b/Engine/Source/Runtime/Renderer/Private/Shadows/ShadowScene.cpp
--- a/Engine/Source/Runtime/Renderer/Private/Shadows/ShadowScene.cpp
+++ b/Engine/Source/Runtime/Renderer/Private/Shadows/ShadowScene.cpp
@@ -43,6 +43,39 @@ DECLARE_DWORD_COUNTER_STAT(TEXT("Active Light Count"), STAT_ActiveLightCount, ST
 
 IMPLEMENT_SCENE_EXTENSION(FShadowScene);
 
+
+void FShadowScene::FUpdater::PreLightsUpdate(FRDGBuilder& GraphBuilder, const FLightSceneChangeSet& LightSceneChangeSet)
+{
+	// Don't sync if there is no work to do.
+	if (!LightSceneChangeSet.RemovedLightIds.IsEmpty())
+	{
+		// Need to wait in case the update is performed several times in a row for some reason.
+		ShadowScene.SceneChangeUpdateTask.Wait();
+
+		// TODO: maybe make async?
+		{
+			// Oust all removed Ids.
+			for (int32 Id : LightSceneChangeSet.RemovedLightIds)
+			{
+				ShadowScene.ActiveLights[Id] = false;
+				if (ShadowScene.LightsCommonData.IsValidIndex(Id))
+				{
+					ShadowScene.LightsCommonData.RemoveAt(Id);
+				}
+
+				if (ShadowScene.Scene.Lights[Id].LightType == LightType_Directional)
+				{
+					int32 DirLightIndex = ShadowScene.DirectionalLights.IndexOfByPredicate([Id](const FDirectionalLightData& DirectionalLightData) { return DirectionalLightData.LightId == Id; });
+					if (DirLightIndex != INDEX_NONE)
+					{
+						ShadowScene.DirectionalLights.RemoveAtSwap(DirLightIndex);
+					}
+				}
+			}
+		}
+	}
+}
+
 void FShadowScene::FUpdater::PostLightsUpdate(FRDGBuilder& GraphBuilder, const FLightSceneChangeSet& LightSceneChangeSet)
 {
 	// don't spawn async work for no good reason.
@@ -57,37 +90,40 @@ void FShadowScene::FUpdater::PostLightsUpdate(FRDGBuilder& GraphBuilder, const F
 		ShadowScene.SceneChangeUpdateTask.Wait();
 		ShadowScene.SceneChangeUpdateTask = GraphBuilder.AddSetupTask([this, LightSceneChangeSet]
 		{
-			// Oust all removed Ids.
-			for (int32 Id : LightSceneChangeSet.RemovedLightIds)
-			{
-				check(ShadowScene.Lights.IsValidIndex(Id));
-				ShadowScene.ActiveLights[Id] = false;
-				ShadowScene.Lights.RemoveAt(Id);
-			}
-	
 			// Track active lights (those that are or were moving, and thus need updating)
-			ShadowScene.ActiveLights.SetNum(FMath::Max(ShadowScene.Lights.GetMaxIndex(), ShadowScene.Lights.Num() + LightSceneChangeSet.AddedLightIds.Num()), false);
+			ShadowScene.ActiveLights.SetNum(FMath::Max(LightSceneChangeSet.PreUpdateMaxIndex, LightSceneChangeSet.PostUpdateMaxIndex), false);
 
 			auto UpdateLight = [&](int32 Id)
 			{
-				FLightData& Light = ShadowScene.GetOrAddLight(Id);
+				const FLightSceneInfoCompact& LightSceneInfoCompact = ShadowScene.Scene.Lights[Id];
+				FLightCommonData& LightCommonData = ShadowScene.GetOrAddLightCommon(Id);
+
 				// Mark as not rendered
-				Light.FirstActiveFrameNumber = INDEX_NONE;
+				LightCommonData.FirstActiveFrameNumber = INDEX_NONE;
 				// Only movable lights can become "active" (i.e., having moved recently and thus needing active update)
-				if (ShadowScene.Scene.Lights[Id].bIsMovable)
+				if (LightSceneInfoCompact.bIsMovable)
 				{
 					ShadowScene.ActiveLights[Id] = true;
-					Light.MobilityFactor = 1.0f;
+					LightCommonData.MobilityFactor = 1.0f;
 				}
 				else
 				{
 					ShadowScene.ActiveLights[Id] = false;
-					Light.MobilityFactor = 0.0f;
+					LightCommonData.MobilityFactor = 0.0f;
 					// Go straight to not active
-				}
+				}	
 			};
 
-			Algo::ForEach(LightSceneChangeSet.AddedLightIds, UpdateLight);
+			for (int32 LightId : LightSceneChangeSet.AddedLightIds)
+			{ 
+				if (ShadowScene.Scene.Lights[LightId].LightType == LightType_Directional)
+				{
+					check(ShadowScene.DirectionalLights.FindByPredicate([LightId](const FDirectionalLightData& DirectionalLightData) { return DirectionalLightData.LightId == LightId; }) == nullptr);
+					ShadowScene.DirectionalLights.Emplace(FDirectionalLightData{LightId});
+				}
+				UpdateLight(LightId);
+			}
+
 			for (const auto& Item : LightSceneChangeSet.SceneLightInfoUpdates.GetRangeView<FUpdateLightTransformParameters>())
 			{ 
 				UpdateLight(Item.SceneInfo->Id);
@@ -161,7 +197,7 @@ void FShadowScene::UpdateForRenderedFrame(FRDGBuilder& GraphBuilder)
 	{
 		int32 Id = BitIt.GetIndex();
 		// No need to process if we already decided it is active (i.e., it was moved again this frame)
-		FLightData& Light = Lights[Id];
+		FLightCommonData& Light = LightsCommonData[Id];
 		
 		// If it was not rendered before, we record the current scene frame number
 		if (Light.FirstActiveFrameNumber == INDEX_NONE)
@@ -193,7 +229,7 @@ float FShadowScene::GetLightMobilityFactor(int32 LightId) const
 
 	if (IsActive(LightId))
 	{
-		return Lights[LightId].MobilityFactor;
+		return LightsCommonData[LightId].MobilityFactor;
 	}
 	return 0.0f;
 }
@@ -251,7 +287,7 @@ void FShadowScene::WaitForSceneLightsUpdateTask()
 
 bool FShadowScene::ShouldCreateExtension(FScene& InScene)
 {
-	return true;
+	return DoesPlatformSupportVirtualShadowMaps(GetFeatureLevelShaderPlatform(InScene.GetFeatureLevel()));
 }
 
 void FShadowScene::InitExtension(FScene& InScene)
diff --git a/Engine/Source/Runtime/Renderer/Private/Shadows/ShadowScene.h b/Engine/Source/Runtime/Renderer/Private/Shadows/ShadowScene.h
--- a/Engine/Source/Runtime/Renderer/Private/Shadows/ShadowScene.h
+++ b/Engine/Source/Runtime/Renderer/Private/Shadows/ShadowScene.h
@@ -7,6 +7,7 @@
 #include "Containers/Array.h"
 #include "Tasks/Task.h"
 #include "SceneExtensions.h"
+#include "LightSceneInfo.h"
 
 class FProjectedShadowInfo;
 class FScene;
@@ -61,6 +62,7 @@ private:
 		FUpdater(FShadowScene& InShadowScene) : ShadowScene(InShadowScene) {}
 
 		// Handle all scene changes wrt lights
+		virtual void PreLightsUpdate(FRDGBuilder& GraphBuilder, const FLightSceneChangeSet& LightSceneChangeSet) override;
 		virtual void PostLightsUpdate(FRDGBuilder& GraphBuilder, const FLightSceneChangeSet &LightSceneChangeSet) override;
 
 		/**
@@ -81,7 +83,10 @@ private:
 
 	inline bool IsActive(int32 LightId) const { return ActiveLights[LightId]; }
 
-	struct FLightData
+	/**
+	 * Data common to all light types, indexed by light scene ID
+	 */
+	struct FLightCommonData
 	{
 		/**
 		 * Scene rendering frame number of the first frame that the scene was rendered after being modified (moved/added).
@@ -89,16 +94,26 @@ private:
 		int32 FirstActiveFrameNumber;
 		float MobilityFactor;
 	};
-	TSparseArray<FLightData> Lights;
+	TSparseArray<FLightCommonData> LightsCommonData;
 
-	inline FLightData& GetOrAddLight(int32 LightId)
+	FLightCommonData& GetOrAddLightCommon(int32 LightId)
 	{
-		if (!Lights.IsValidIndex(LightId))
+		if (!LightsCommonData.IsValidIndex(LightId))
 		{
-			Lights.EmplaceAt(LightId, FLightData{});
+			LightsCommonData.EmplaceAt(LightId, FLightCommonData{});
 		}
-		return Lights[LightId];
-	}
+		return LightsCommonData[LightId];
+	}	
+
+	/**
+	 * Directional light data, not indexed by light scene ID but instead linearly searched as there are typically very few.
+	 */
+	struct FDirectionalLightData
+	{
+		FLightSceneInfo::FPersistentId LightId = -1;
+	};
+	TArray<FDirectionalLightData> DirectionalLights;
+
 	// Bit-array marking active lights, those we deem active are ones that have been modified in a recent frame and thus need some kind of active update.
 	TBitArray<> ActiveLights;
 
diff --git a/Engine/Source/Runtime/Renderer/Private/Shadows/ShadowSceneRenderer.cpp b/Engine/Source/Runtime/Renderer/Private/Shadows/ShadowSceneRenderer.cpp
--- a/Engine/Source/Runtime/Renderer/Private/Shadows/ShadowSceneRenderer.cpp
+++ b/Engine/Source/Runtime/Renderer/Private/Shadows/ShadowSceneRenderer.cpp
@@ -18,6 +18,8 @@
 
 CSV_DECLARE_CATEGORY_EXTERN(VSM);
 
+extern TAutoConsoleVariable<int32> CVarNaniteShadowsUpdateStreaming;
+
 TAutoConsoleVariable<int32> CVarVSMMaterialVisibility(
 	TEXT("r.Shadow.Virtual.Nanite.MaterialVisibility"),
 	0,
@@ -163,6 +165,35 @@ struct FHeapPair
 
 void FShadowSceneRenderer::PreInitViews(FRDGBuilder& GraphBuilder)
 {
+	// Clear the frame setups to indicate that nothing is allocated for this frame
+	CommonSetups.SetNumZeroed(ShadowScene.LightsCommonData.GetMaxIndex());
+
+	// Allocate space for each directional light in the scene, one for each view.
+	bool bIsStereo = SceneRenderer.IsRenderingStereo();
+	const int32 ShadowNumViews = bIsStereo ? 1 : SceneRenderer.Views.Num();
+
+	DirectionalLights.Reserve(ShadowScene.DirectionalLights.Num() * ShadowNumViews);
+	// pre-allocate indexes and setups for each directional light, strided by the view count.
+	for (const auto& DirectionalLight : ShadowScene.DirectionalLights)
+	{
+		FLightCommonFrameSetup& CommonSetup = CommonSetups[DirectionalLight.LightId];
+		CommonSetup.bIsDirectional = true;
+		CommonSetup.SetupIndex = DirectionalLights.Num();
+		if (bIsStereo)
+		{
+			// only set up one for both
+			DirectionalLights.Emplace(DirectionalLight.LightId, 3u);
+			check(SceneRenderer.Views.Num() == 2)
+		}
+		else
+		{
+			for (int32 ViewIndex = 0; ViewIndex < ShadowNumViews; ++ViewIndex)
+			{
+				DirectionalLights.Emplace(DirectionalLight.LightId, 1u << ViewIndex);
+			}
+		}
+	}
+
 	bNeedVSMOnePassProjection = false;
 	bNeedMegaLightsProjection = false;
 
@@ -269,6 +300,20 @@ void FShadowSceneRenderer::UpdateLocalLightProjectionShaderDataMatrices(
 	OutProjectionShaderData->ClipmapLevel_ClipmapLevelCountRemaining	= -1;					// Not a clipmap
 }
 
+FShadowSceneRenderer::FDirectionalLightShadowFrameSetup* FShadowSceneRenderer::FindDirectional(const FLightSceneInfo::FPersistentId LightId, int32 ViewIndex)
+{
+	// may alternatively be indexed through CommonSetups[LightId].SetupIndex but it is probably more efficient to do a linear search.
+	for (FDirectionalLightShadowFrameSetup& Setup : DirectionalLights)
+	{
+		if (LightId == Setup.LightId && (Setup.ViewMask & (1u << ViewIndex)) != 0u)
+		{
+			return &Setup;
+		}
+	}
+	check(false);
+	return nullptr;
+}
+
 /**
  * Calculate the radius in world-space units of a single pixel at a given depth.
  */
@@ -297,20 +342,28 @@ static uint32 GetConservativeMipLevelLocal(const FViewInfo& View, float ClipToVi
 TSharedPtr<FVirtualShadowMapPerLightCacheEntry> FShadowSceneRenderer::AddLocalLightShadow(const FWholeSceneProjectedShadowInitializer& ProjectedShadowInitializer, FProjectedShadowInfo* ProjectedShadowInfo, FLightSceneInfo* LightSceneInfo, float MaxScreenRadius)
 {
 	FVirtualShadowMapArrayCacheManager* CacheManager = VirtualShadowMapArray.CacheManager;
+	FLightSceneInfo::FPersistentId LightId = LightSceneInfo->Id;
+
+	FLightCommonFrameSetup& CommonSetup = CommonSetups[LightId];
+	// prevent double allocation
+	check(!CommonSetup.bHasVirtualShadowMap);
+	CommonSetup.bHasVirtualShadowMap = true;
+	// link from ID to the allocated local shadow slot
+	CommonSetup.SetupIndex = uint32(LocalLights.Num());
+
+	FLocalLightShadowFrameSetup& LocalLightSetup = LocalLights.AddDefaulted_GetRef();
 
-	const int32 LocalLightShadowIndex = LocalLights.Num();
-	FLocalLightShadowFrameSetup& LocalLightShadowFrameSetup = LocalLights.AddDefaulted_GetRef();
-	LocalLightShadowFrameSetup.ProjectedShadowInfo = ProjectedShadowInfo;
-	LocalLightShadowFrameSetup.LightSceneInfo = LightSceneInfo;
+	LocalLightSetup.ProjectedShadowInfo = ProjectedShadowInfo;
+	LocalLightSetup.LightSceneInfo = LightSceneInfo;
 
 	const FLightSceneProxy* LightSceneProxy = ProjectedShadowInfo->GetLightSceneInfo().Proxy;
-	const float ResolutionLODBiasLocal = GetResolutionLODBiasLocal(ShadowScene.GetLightMobilityFactor(LightSceneInfo->Id), LightSceneProxy->GetVSMResolutionLodBias());
+	const float ResolutionLODBiasLocal = GetResolutionLODBiasLocal(ShadowScene.GetLightMobilityFactor(LightId), LightSceneProxy->GetVSMResolutionLodBias());
 
 	// Compute conservative mip level estimate based on radius of the bounding sphere.
 	// TODO: can probably do better by finding  closest point on cone for certain scenarios? Not as important as it might seem as the worst case is for a narrow cone, but then the narrow FOV limits the required resolution.
 
 	const FVector2f ShadowViewSize = FVector2f(FVirtualShadowMap::VirtualMaxResolutionXY, FVirtualShadowMap::VirtualMaxResolutionXY);
-	const FMatrix &ShadowViewToClip = ProjectedShadowInfo->bOnePassPointLightShadow ? ProjectedShadowInfo->OnePassShadowFaceProjectionMatrix : ProjectedShadowInfo->ViewToClipOuter;
+	const FMatrix& ShadowViewToClip = ProjectedShadowInfo->bOnePassPointLightShadow ? ProjectedShadowInfo->OnePassShadowFaceProjectionMatrix : ProjectedShadowInfo->ViewToClipOuter;
 	float ShadowProjScale = ShadowViewToClip.M[0][0]; // always symmetrical
 	const float WorldToShadowFootprintScale = ShadowProjScale * ShadowViewSize.X;
 
@@ -369,8 +422,8 @@ TSharedPtr<FVirtualShadowMapPerLightCacheEntry> FShadowSceneRenderer::AddLocalLi
 
 
 	const int32 NumMaps = ProjectedShadowInitializer.bOnePassPointLightShadow ? 6 : 1;
-	TSharedPtr<FVirtualShadowMapPerLightCacheEntry> PerLightCacheEntry = CacheManager->FindCreateLightCacheEntry(LightSceneInfo->Id, 0, NumMaps);
-	LocalLightShadowFrameSetup.PerLightCacheEntry = PerLightCacheEntry;
+	TSharedPtr<FVirtualShadowMapPerLightCacheEntry> PerLightCacheEntry = CacheManager->FindCreateLightCacheEntry(LightId, 0, NumMaps);
+	LocalLightSetup.PerLightCacheEntry = PerLightCacheEntry;
 
 	PerLightCacheEntry->UpdateLocal(
 		ProjectedShadowInitializer,
@@ -387,6 +440,7 @@ TSharedPtr<FVirtualShadowMapPerLightCacheEntry> FShadowSceneRenderer::AddLocalLi
 
 	// Update info on the ProjectionShadowInfo; eventually this should all move into local data structures here
 	const int32 VirtualShadowMapId = VirtualShadowMapArray.Allocate(bIsDistantLight, NumMaps);
+	LocalLightSetup.VirtualShadowMapId = VirtualShadowMapId;
 	ProjectedShadowInfo->VirtualShadowMapId = VirtualShadowMapId;
 	ProjectedShadowInfo->VirtualShadowMapPerLightCacheEntry = PerLightCacheEntry;
 	ProjectedShadowInfo->bShouldRenderVSM = !PerLightCacheEntry->IsFullyCached();
@@ -445,24 +499,72 @@ TSharedPtr<FVirtualShadowMapPerLightCacheEntry> FShadowSceneRenderer::AddLocalLi
 	return PerLightCacheEntry;
 }
 
-void FShadowSceneRenderer::AddDirectionalLightShadow(FProjectedShadowInfo* ProjectedShadowInfo)
+void FShadowSceneRenderer::AddDirectionalLightShadow(FLightSceneInfo& LightSceneInfo, FViewInfo& View, TArray<FProjectedShadowInfo*, SceneRenderingAllocator>& OutShadowInfosThatNeedCulling)
 {
-	FDirectionalLightShadowFrameSetup& DirectionalLightShadowFrameSetup = DirectionalLights.AddDefaulted_GetRef();
-	DirectionalLightShadowFrameSetup.ProjectedShadowInfo = ProjectedShadowInfo;
+	const int32 ViewIndex = View.SceneRendererPrimaryViewId;
+	FLightSceneInfo::FPersistentId LightId = LightSceneInfo.Id;
+	FVisibleLightInfo& VisibleLightInfo = SceneRenderer.VisibleLightInfos[LightId];
 
-	FLightSceneInfo& LightSceneInfo = ProjectedShadowInfo->GetLightSceneInfo();
+	FLightCommonFrameSetup& CommonSetup = CommonSetups[LightId];
+	check(CommonSetup.bIsDirectional);
+	CommonSetup.bHasVirtualShadowMap = true;
 
-	if (!ShadowScene.FirstPersonWorldSpacePrimitives.IsEmpty())
+	// Helper function to create a projected shadow info. This is needed to:
+	//  * Get the matrices included in the shadow rendering pass setup, driving nanite VSM rendering (which VisibleLightInfo.AllProjectedShadows is appended to)
+	auto AddLegacySetup = [&](const TSharedPtr<FVirtualShadowMapClipmap>& Clipmap, bool bQueueforNonNaniteCulling)
 	{
-		FViewInfo& View = *ProjectedShadowInfo->DependentView;
-		FVisibleLightInfo& VisibleLightInfo = SceneRenderer.VisibleLightInfos[LightSceneInfo.Id];
+		FProjectedShadowInfo* ProjectedShadowInfo = SceneRenderer.Allocator.Create<FProjectedShadowInfo>();
+		ProjectedShadowInfo->SetupClipmapProjection(&LightSceneInfo, &View, Clipmap, ProjectedShadowInfo->GetMaxNonFarCascadeDistance());
+	
+		// This is needed to get it into the line for ending up in the SortedShadowsForShadowDepthPass.VirtualShadowMapShadows which is what drives the shadow rendering.
+		VisibleLightInfo.AllProjectedShadows.Add(ProjectedShadowInfo);
+		
+		ProjectedShadowInfo->VirtualShadowMapPerLightCacheEntry = Clipmap->GetCacheEntry();
+
+		if (bQueueforNonNaniteCulling)
+		{
+			OutShadowInfosThatNeedCulling.Emplace(ProjectedShadowInfo);
+		}
+
+		return ProjectedShadowInfo;
+	};
+
+	// a secondary view should not allocate shadow for itself!
+	check(View.GetPrimaryView() == &View);
+	FDirectionalLightShadowFrameSetup& Setup = DirectionalLights[CommonSetup.SetupIndex + ViewIndex];
+	check(Setup.LightId == LightId);
+	check(Setup.ViewMask & (1u << ViewIndex));
+
+	// Set up regular clipmap
+	FDirectionalLightShadowFrameSetup::FClipmapInfo& RegularInfo = Setup.ClipmapInfos[EVirtualShadowTypeId::Regular];
+	{
+		FVirtualShadowMapClipmapConfig Config = FVirtualShadowMapClipmapConfig::GetGlobal();
 
-		// Arbitrary type ID for first person shadow
-		constexpr uint32 FirstPersonShadowTypeId = 123u;
+		TSharedPtr<FVirtualShadowMapClipmap> VirtualShadowMapClipmap = TSharedPtr<FVirtualShadowMapClipmap>(new FVirtualShadowMapClipmap(
+			VirtualShadowMapArray,
+			LightSceneInfo,
+			View.ViewMatrices,
+			View.ViewRect.Size(),
+			&View,
+			ShadowScene.GetLightMobilityFactor(LightSceneInfo.Id),
+			Config));
+
+		// NOTE: only contains "regular" clipmaps, the alternate types are internal to the system and needs to be queried for.
+		VisibleLightInfo.VirtualShadowMapClipmaps.Add(VirtualShadowMapClipmap);
+		// NOTE: If there are multiple camera views this will simply be associated with "one of them"
+		VisibleLightInfo.VirtualShadowMapId = VirtualShadowMapClipmap->GetVirtualShadowMapId();
+
+		RegularInfo.ProjectedShadowInfo = AddLegacySetup(VirtualShadowMapClipmap, true);
+		RegularInfo.Clipmap = VirtualShadowMapClipmap;
+	}
+
+	if (!ShadowScene.FirstPersonWorldSpacePrimitives.IsEmpty())
+	{
+		FDirectionalLightShadowFrameSetup::FClipmapInfo& FPInfo = Setup.ClipmapInfos[EVirtualShadowTypeId::FirstPerson];
 
 		// Clone the setup from the regular VSM clipmap
 		FVirtualShadowMapClipmapConfig Config = FVirtualShadowMapClipmapConfig::GetGlobal();
-		Config.ShadowTypeId = FirstPersonShadowTypeId;
+		Config.ShadowTypeId = EVirtualShadowTypeId::FirstPerson;
 		Config.bForceInvalidate = true;
 		Config.FirstCoarseLevel = -1;
 		Config.LastCoarseLevel = -1;
@@ -470,7 +572,7 @@ void FShadowSceneRenderer::AddDirectionalLightShadow(FProjectedShadowInfo* Proje
 		Config.LastLevel = CVarVirtualShadowMapFirstPersonClipmapLastLevel.GetValueOnRenderThread();
 		Config.bIsFirstPersonShadow = true;
 
-		TSharedPtr<FVirtualShadowMapClipmap> FPClipmap = TSharedPtr<FVirtualShadowMapClipmap>(new FVirtualShadowMapClipmap(
+		FPInfo.Clipmap = TSharedPtr<FVirtualShadowMapClipmap>(new FVirtualShadowMapClipmap(
 			VirtualShadowMapArray,
 			LightSceneInfo,
 			View.ViewMatrices,
@@ -480,21 +582,12 @@ void FShadowSceneRenderer::AddDirectionalLightShadow(FProjectedShadowInfo* Proje
 			Config
 		));
 
-		FProjectedShadowInfo* FPProjectedShadowInfo = SceneRenderer.Allocator.Create<FProjectedShadowInfo>();
-		FPProjectedShadowInfo->SetupClipmapProjection(&LightSceneInfo, &View, FPClipmap, ProjectedShadowInfo->GetMaxNonFarCascadeDistance());
-		
-		// Need their own array or they'll be found at the wrong times, and we also need them in here to be able to find them during page marking (sigh)
-		VisibleLightInfo.FirstPersonVirtualShadowMapClipmaps.Add(FPClipmap);
-		
-		// This is needed to get it into the line for ending up in the SortedShadowsForShadowDepthPass.VirtualShadowMapShadows which is what drives the shadow rendering.
-		VisibleLightInfo.AllProjectedShadows.Add(FPProjectedShadowInfo);
-		
-		FPProjectedShadowInfo->VirtualShadowMapPerLightCacheEntry = FPClipmap->GetCacheEntry();
+		FPInfo.ProjectedShadowInfo = AddLegacySetup(FPInfo.Clipmap, false);
 
 		for (FPrimitiveSceneInfo* PrimitiveSceneInfo : ShadowScene.FirstPersonWorldSpacePrimitives)
 		{
 			// _invert_ the shadow relevance because we want to render the thing despite being owner no see, and NOT anyting that is the opposite (owned by other views).
-			FPProjectedShadowInfo->AddSubjectPrimitive(PrimitiveSceneInfo, SceneRenderer.Views, false, true);
+			FPInfo.ProjectedShadowInfo->AddSubjectPrimitive(PrimitiveSceneInfo, SceneRenderer.Views, false, true);
 		}
 	}
 
@@ -519,7 +612,7 @@ void FShadowSceneRenderer::PostInitDynamicShadowsSetup()
 		
 		for (const FDirectionalLightShadowFrameSetup& DirectionalLightShadowFrameSetup : DirectionalLights)
 		{
-			FProjectedShadowInfo* ProjectedShadowInfo = DirectionalLightShadowFrameSetup.ProjectedShadowInfo;
+			FProjectedShadowInfo* ProjectedShadowInfo = DirectionalLightShadowFrameSetup.ClipmapInfos[EVirtualShadowTypeId::Regular].ProjectedShadowInfo;
 			if (!bUnboundedClipmap && ProjectedShadowInfo->bShouldRenderVSM)
 			{
 				const bool bIsCached = !ProjectedShadowInfo->VirtualShadowMapPerLightCacheEntry->IsUncached();
@@ -597,10 +690,7 @@ void FShadowSceneRenderer::RenderVirtualShadowMaps(FRDGBuilder& GraphBuilder, bo
 #endif
 	}
 
-	const TArray<FProjectedShadowInfo*, SceneRenderingAllocator>& VirtualShadowMapShadows = SceneRenderer.SortedShadowsForShadowDepthPass.VirtualShadowMapShadows;
-	TArray<TSharedPtr<FVirtualShadowMapClipmap>, SceneRenderingAllocator>& VirtualShadowMapClipmaps = SceneRenderer.SortedShadowsForShadowDepthPass.VirtualShadowMapClipmaps;
-
-	if (VirtualShadowMapShadows.Num() == 0 && VirtualShadowMapClipmaps.Num() == 0)
+	if (VirtualShadowMapArray.GetNumShadowMaps() == 0)
 	{
 		return;
 	}
@@ -617,6 +707,7 @@ void FShadowSceneRenderer::RenderVirtualShadowMaps(FRDGBuilder& GraphBuilder, bo
 
 	if (UseNonNaniteVirtualShadowMaps(SceneRenderer.ShaderPlatform, SceneRenderer.FeatureLevel))
 	{
+		const TArray<FProjectedShadowInfo*, SceneRenderingAllocator>& VirtualShadowMapShadows = SceneRenderer.SortedShadowsForShadowDepthPass.VirtualShadowMapShadows;
 		VirtualShadowMapArray.RenderVirtualShadowMapsNonNanite(GraphBuilder, SceneRenderer.GetSceneUniforms(), VirtualShadowMapShadows, SceneRenderer.Views);
 	}
 
@@ -690,6 +781,32 @@ void FShadowSceneRenderer::PostSetupDebugRender()
 #endif
 }
 
+void FShadowSceneRenderer::RenderVirtualShadowMaps(
+	FRDGBuilder& GraphBuilder, 
+	bool bNaniteEnabled, 
+	const FSingleLayerWaterPrePassResult* SingleLayerWaterPrePassResult,
+	const FFrontLayerTranslucencyData& FrontLayerTranslucencyData,
+	const Froxel::FRenderer& FroxelRenderer)
+{
+	if (!VirtualShadowMapArray.IsEnabled())
+	{
+		return;
+	}
+
+	VirtualShadowMapArray.BuildPageAllocations(GraphBuilder,
+		SceneRenderer.GetActiveSceneTextures(),
+		SceneRenderer.Views,
+		SceneRenderer,
+		SceneRenderer.VisibleLightInfos,
+		SingleLayerWaterPrePassResult,
+		FrontLayerTranslucencyData,
+		FroxelRenderer,
+		AreAnyLocalLightsPreset()
+	);
+
+	RenderVirtualShadowMaps(GraphBuilder, bNaniteEnabled, CVarNaniteShadowsUpdateStreaming.GetValueOnRenderThread() != 0);
+}
+
 void FShadowSceneRenderer::RenderVirtualShadowMapProjectionMaskBits(
 	FRDGBuilder& GraphBuilder,
 	FMinimalSceneTextures& SceneTextures)
@@ -744,6 +861,56 @@ void FShadowSceneRenderer::RenderVirtualShadowMapProjectionMaskBits(
 	}
 }
 
+void FShadowSceneRenderer::RenderVirtualShadowMapProjection(
+	FRDGBuilder& GraphBuilder,
+	const FMinimalSceneTextures& SceneTextures,
+	FLightSceneInfo::FPersistentId LightId,
+	const FViewInfo& View, int32 ViewIndex,
+	const FIntRect ScissorRect,
+	EVirtualShadowMapProjectionInputType InputType,
+	bool bModulateRGB,
+	FTiledVSMProjection* TiledVSMProjection,
+	FRDGTextureRef OutputShadowMaskTexture)
+{
+	FShadowSceneRenderer::FDirectionalLightShadowFrameSetup* DirSetup = FindDirectional(LightId, ViewIndex);
+	if (DirSetup != nullptr)
+	{
+		::RenderVirtualShadowMapProjection(
+			GraphBuilder,
+			SceneTextures,
+			View, ViewIndex,
+			VirtualShadowMapArray,
+			ScissorRect,
+			EVirtualShadowMapProjectionInputType::GBuffer,
+			DirSetup->ClipmapInfos[EVirtualShadowTypeId::Regular].Clipmap,
+			true, // bModulateRGB
+			TiledVSMProjection,
+			OutputShadowMaskTexture,
+			DirSetup->ClipmapInfos[EVirtualShadowTypeId::FirstPerson].Clipmap);
+	}
+}
+
+TArray<int32, SceneRenderingAllocator> FShadowSceneRenderer::GatherClipmapIds(int32 ViewIndex) const
+{
+	TArray<int32, SceneRenderingAllocator> Result;
+	Result.Reserve(DirectionalLights.Num() * EVirtualShadowTypeId::Max );
+
+	for (const FDirectionalLightShadowFrameSetup& DirSetup : DirectionalLights)
+	{
+		if (DirSetup.ViewMask & (1u << ViewIndex))
+		{
+			for (const auto& Info: DirSetup.ClipmapInfos)
+			{
+				if (Info.Clipmap)
+				{
+					Result.Emplace(Info.Clipmap->GetVirtualShadowMapId());
+				}
+			}
+		}
+	}
+	return Result;
+}
+
 void FShadowSceneRenderer::ApplyVirtualShadowMapProjectionForLight(
 	FRDGBuilder& GraphBuilder,
 	const FMinimalSceneTextures& SceneTextures,
@@ -756,18 +923,34 @@ void FShadowSceneRenderer::ApplyVirtualShadowMapProjectionForLight(
 		return;
 	}
 
-	const FVisibleLightInfo& VisibleLightInfo = SceneRenderer.VisibleLightInfos[LightSceneInfo->Id];
+	FLightSceneInfo::FPersistentId LightId = LightSceneInfo->Id;
+
+	const FLightCommonFrameSetup& CommonSetup = CommonSetups[LightId];
+
+	// No VSM set up this frame
+	if (!CommonSetup.bHasVirtualShadowMap)
+	{
+		return;
+	}
+	const FVisibleLightInfo& VisibleLightInfo = SceneRenderer.VisibleLightInfos[LightId];
 
 	if (!VisibleLightInfo.HasVirtualShadowMap())
 	{
 		return;
 	}
 
+	// Some lights can elide the screen shadow mask entirely, in which case they will be sampled directly in the lighting shader
+	if (!OutputScreenShadowMaskTexture)
+	{
+		return;
+	}
+
 	for (int32 ViewIndex = 0; ViewIndex < SceneRenderer.Views.Num(); ++ViewIndex)
 	{
 		RDG_EVENT_SCOPE_CONDITIONAL(GraphBuilder, SceneRenderer.Views.Num() > 1, "View%d", ViewIndex);
 
 		FViewInfo& View = SceneRenderer.Views[ViewIndex];
+
 		FIntRect ScissorRect;
 		if (!LightSceneInfo->Proxy->GetScissorRect(ScissorRect, View, View.ViewRect))
 		{
@@ -776,33 +959,41 @@ void FShadowSceneRenderer::ApplyVirtualShadowMapProjectionForLight(
 
 		if (ScissorRect.Area() > 0)
 		{
-			int32 VirtualShadowMapId = VisibleLightInfo.GetVirtualShadowMapId(&View);
-			
-			// Some lights can elide the screen shadow mask entirely, in which case they will be sampled directly in the lighting shader
-			if (OutputScreenShadowMaskTexture && InputType == EVirtualShadowMapProjectionInputType::GBuffer)
+			if (InputType == EVirtualShadowMapProjectionInputType::HairStrands && !HairStrands::HasViewHairStrandsData(View))
 			{
-				if (VisibleLightInfo.VirtualShadowMapClipmaps.Num() > 0)
-				{
-					INC_DWORD_STAT(STAT_VSMDirectionalProjectionFull);
+				continue;
+			}
 
-					// Project directional light virtual shadow map
-					RenderVirtualShadowMapProjection(
-						GraphBuilder,
-						SceneTextures,
-						View, ViewIndex,
-						VirtualShadowMapArray,
-						ScissorRect,
-						EVirtualShadowMapProjectionInputType::GBuffer,
-						VisibleLightInfo.FindShadowClipmapForView(&View),
-						false, // bModulateRGB
-						nullptr, // TiledVSMProjection
-						OutputScreenShadowMaskTexture,
-						VisibleLightInfo.FindFirstPersonShadowClipmapForView(&View));				
-				}
-				else if (bShouldUseVirtualShadowMapOnePassProjection)
+			if (CommonSetup.bIsDirectional)
+			{
+				// remap to use the primary view index for stereo rendering.
+				int32 ShadowViewIndex = View.GetPrimaryView()->SceneRendererPrimaryViewId;
+				const FDirectionalLightShadowFrameSetup& DirSetup = DirectionalLights[CommonSetup.SetupIndex + ShadowViewIndex];
+
+				INC_DWORD_STAT(STAT_VSMDirectionalProjectionFull);
+
+				// Project directional light virtual shadow map
+				::RenderVirtualShadowMapProjection(
+					GraphBuilder,
+					SceneTextures,
+					View, ViewIndex,
+					VirtualShadowMapArray,
+					ScissorRect,
+					InputType,
+					DirSetup.ClipmapInfos[EVirtualShadowTypeId::Regular].Clipmap,
+					false, // bModulateRGB
+					nullptr, // TiledVSMProjection
+					OutputScreenShadowMaskTexture,
+					DirSetup.ClipmapInfos[EVirtualShadowTypeId::FirstPerson].Clipmap);
+			}
+			else 
+			{
+				const FLocalLightShadowFrameSetup& LocalLightSetup = LocalLights[CommonSetup.SetupIndex];
+				if (bShouldUseVirtualShadowMapOnePassProjection)
 				{
 					INC_DWORD_STAT(STAT_VSMLocalProjectionOnePassCopy);
 
+
 					// Copy local light from one pass projection output
 					CompositeVirtualShadowMapFromMaskBits(
 						GraphBuilder,
@@ -810,8 +1001,8 @@ void FShadowSceneRenderer::ApplyVirtualShadowMapProjectionForLight(
 						View, ViewIndex,
 						ScissorRect,
 						VirtualShadowMapArray,
-						EVirtualShadowMapProjectionInputType::GBuffer,
-						VirtualShadowMapId,
+						InputType,
+						LocalLightSetup.VirtualShadowMapId,
 						VirtualShadowMapMaskBits,
 						OutputScreenShadowMaskTexture);
 				}
@@ -820,65 +1011,18 @@ void FShadowSceneRenderer::ApplyVirtualShadowMapProjectionForLight(
 					INC_DWORD_STAT(STAT_VSMLocalProjectionFull);
 
 					// Project local light virtual shadow map
-					RenderVirtualShadowMapProjection(
+					::RenderVirtualShadowMapProjection(
 						GraphBuilder,
 						SceneTextures,
 						View, ViewIndex,
 						VirtualShadowMapArray,
 						ScissorRect,
-						EVirtualShadowMapProjectionInputType::GBuffer,
+						InputType,
 						*LightSceneInfo,
-						VirtualShadowMapId,
+						LocalLightSetup.VirtualShadowMapId,
 						OutputScreenShadowMaskTexture);
 				}
-			}
-
-			// Sub-pixel shadow (no denoising for hair)
-			if (HairStrands::HasViewHairStrandsData(View) && OutputScreenShadowMaskTexture && InputType == EVirtualShadowMapProjectionInputType::HairStrands)
-			{
-				if (VisibleLightInfo.VirtualShadowMapClipmaps.Num() > 0)
-				{
-					RenderVirtualShadowMapProjection(
-						GraphBuilder,
-						SceneTextures,
-						View, ViewIndex,
-						VirtualShadowMapArray,
-						ScissorRect,
-						EVirtualShadowMapProjectionInputType::HairStrands,
-						VisibleLightInfo.FindShadowClipmapForView(&View),
-						false, // bModulateRGB
-						nullptr, // TiledVSMProjection
-						OutputScreenShadowMaskTexture,
-						VisibleLightInfo.FindFirstPersonShadowClipmapForView(&View));
-				}
-				else if (bShouldUseVirtualShadowMapOnePassProjection)
-				{
-					// Copy local light from one pass projection output
-					CompositeVirtualShadowMapFromMaskBits(
-						GraphBuilder,
-						SceneTextures,
-						View, ViewIndex,
-						ScissorRect,
-						VirtualShadowMapArray,
-						EVirtualShadowMapProjectionInputType::HairStrands,
-						VirtualShadowMapId,
-						VirtualShadowMapMaskBitsHairStrands,
-						OutputScreenShadowMaskTexture);
-				}
-				else
-				{
-					RenderVirtualShadowMapProjection(
-						GraphBuilder,
-						SceneTextures,
-						View, ViewIndex,
-						VirtualShadowMapArray,
-						ScissorRect,
-						EVirtualShadowMapProjectionInputType::HairStrands,
-						*LightSceneInfo,
-						VirtualShadowMapId,
-						OutputScreenShadowMaskTexture);
-				}
-			}
+			}			
 		}
 	}
 }
diff --git a/Engine/Source/Runtime/Renderer/Private/Shadows/ShadowSceneRenderer.h b/Engine/Source/Runtime/Renderer/Private/Shadows/ShadowSceneRenderer.h
--- a/Engine/Source/Runtime/Renderer/Private/Shadows/ShadowSceneRenderer.h
+++ b/Engine/Source/Runtime/Renderer/Private/Shadows/ShadowSceneRenderer.h
@@ -6,9 +6,10 @@
 #include "Containers/ArrayView.h"
 #include "Containers/Array.h"
 #include "Tasks/Task.h"
-
+#include "LightSceneInfo.h"
 #include "ShadowScene.h"
 #include "VirtualShadowMaps/VirtualShadowMapArray.h"
+#include "VirtualShadowMaps/VirtualShadowMapProjection.h"
 
 class FProjectedShadowInfo;
 class FDeferredShadingSceneRenderer;
@@ -23,6 +24,22 @@ namespace UE::Renderer::Private
 	class IShadowInvalidatingInstances;
 }
 
+/**
+ * Type Id used to differentiate when there are multiple clipmaps for the same light & view.
+ */
+namespace EVirtualShadowTypeId
+{
+	enum Type : uint32
+	{
+		// Regular clipmap
+		Regular,
+		// First-person (3rd person) shadow map,
+		FirstPerson,
+		
+		Max
+	};
+}
+
 /**
  * Transient scope for per-frame rendering resources for the shadow rendering.
  */
@@ -50,10 +67,8 @@ public:
 	TSharedPtr<FVirtualShadowMapPerLightCacheEntry> AddLocalLightShadow(const FWholeSceneProjectedShadowInitializer& Initializer, FProjectedShadowInfo* ProjectedShadowInfo, FLightSceneInfo* LightSceneInfo, float MaxScreenRadius);
 	/**
 	 * Add a directional light for processing this frame.
-	 * TODO: Don't use legacy FProjectedShadowInfo or other params, instead info should flow from persistent setup & update.
-	 * TODO: Return reference to FLocalLightShadowFrameSetup ?
 	 */
-	void AddDirectionalLightShadow(FProjectedShadowInfo* ProjectedShadowInfo);
+	void AddDirectionalLightShadow(FLightSceneInfo& LightSceneInfo, FViewInfo& View, TArray<FProjectedShadowInfo*, SceneRenderingAllocator>& OutShadowInfosThatNeedCulling);
 
 	/**
 	 * Call after view-dependent setup has been processed (InitView etc) but before any rendering activity has been kicked off.
@@ -67,9 +82,24 @@ public:
 
 	void PostSetupDebugRender();
 
+	/** 
+	 * returns true if the given light has any VSM set up (for any view in the case of view dependent)
+	 */
+	bool HasVirtualShadowMap(FLightSceneInfo::FPersistentId LightId) const { return CommonSetups[LightId].bHasVirtualShadowMap; }
+
+	/** 
+	 * returns true if the given light has a clipmap VSM set up for any view
+	 */
+	bool HasVirtualClipMap(FLightSceneInfo::FPersistentId LightId) const { return CommonSetups[LightId].bHasVirtualShadowMap && CommonSetups[LightId].bIsDirectional; }
+
 	/**
 	 */
-	void RenderVirtualShadowMaps(FRDGBuilder& GraphBuilder, bool bNaniteEnabled, bool bUpdateNaniteStreaming);
+	void RenderVirtualShadowMaps(
+		FRDGBuilder& GraphBuilder, 
+		bool bNaniteEnabled, 
+		const FSingleLayerWaterPrePassResult* SingleLayerWaterPrePassResult,
+		const FFrontLayerTranslucencyData& FrontLayerTranslucencyData,
+		const Froxel::FRenderer& FroxelRenderer);
 
 	/* Does any one pass shadow projection and generates screen space shadow mask bits
 	 * Call before beginning light loop/shadow projection, but after shadow map rendering
@@ -78,6 +108,23 @@ public:
 		FRDGBuilder& GraphBuilder,
 		FMinimalSceneTextures& SceneTextures);
 
+	void RenderVirtualShadowMapProjection(
+		FRDGBuilder& GraphBuilder,
+		const FMinimalSceneTextures& SceneTextures,
+		FLightSceneInfo::FPersistentId LightId,
+		const FViewInfo& View, int32 ViewIndex,
+		const FIntRect ScissorRect,
+		EVirtualShadowMapProjectionInputType InputType,
+		bool bModulateRGB,
+		FTiledVSMProjection* TiledVSMProjection,
+		FRDGTextureRef OutputShadowMaskTexture);
+
+
+	/**
+	 * Get the clipmap IDs for the view (for all lights and features), which in case of stereo maps to the primary view internally.
+	 */
+	TArray<int32, SceneRenderingAllocator> GatherClipmapIds(int32 ViewIndex) const;
+
 	/**
 	 * Renders virtual shadow map projection for a given light into the shadow mask.
 	 * If one pass projection is enabled, this may be a simple composite from the shadow mask bits.
@@ -125,6 +172,10 @@ public:
 	const FVirtualShadowMapArray& GetVirtualShadowMapArray() const { return VirtualShadowMapArray; }
 
 private:
+	/**
+	 */
+	void RenderVirtualShadowMaps(FRDGBuilder& GraphBuilder, bool bNaniteEnabled, bool bUpdateNaniteStreaming);
+
 	struct FViewData
 	{
 		float ClipToViewSizeScale = 0.0f;
@@ -165,24 +216,53 @@ private:
 		float ShadowsLODScaleFactor,
 		FNaniteVirtualShadowMapRenderPass& InOutRenderPass);
 
+
+	struct FLightCommonFrameSetup
+	{
+		uint32 bHasVirtualShadowMap : 1;
+		uint32 bIsDirectional : 1;
+		// index into respective setup array, for a directional light this is the offset to the setup for the first view index 
+		uint32 SetupIndex : 30;
+	};
+
+	// Indexed by light scene ID
+	TArray<FLightCommonFrameSetup, SceneRenderingAllocator> CommonSetups;
+
 	struct FLocalLightShadowFrameSetup
 	{
+		int32 VirtualShadowMapId = INDEX_NONE;
 		TSharedPtr<FVirtualShadowMapPerLightCacheEntry> PerLightCacheEntry;
 		// link to legacy system stuff, to be removed in due time
 		FProjectedShadowInfo* ProjectedShadowInfo = nullptr;
 		FLightSceneInfo* LightSceneInfo = nullptr;
 	};
 
-	// TODO: maybe we want to keep these in a 1:1 sparse array wrt the light scene infos, for easy crossreference & GPU access (maybe)?
-	//       tradeoff is easy to look up (given light ID) but not compact, but OTOH can keep compact lists of indices for various purposes
+	// Indexed by order of allocation, represented in CommonSetups[LightId].SetupIndex
 	TArray<FLocalLightShadowFrameSetup, SceneRenderingAllocator> LocalLights;
 
 	struct FDirectionalLightShadowFrameSetup
 	{
-		FProjectedShadowInfo* ProjectedShadowInfo = nullptr;
+		// Search key
+		FLightSceneInfo::FPersistentId LightId = -1;
+		// A clipmap may belong to more than one view (in stereo mode, specifically)
+		uint32 ViewMask = 0u;
+
+		struct FClipmapInfo
+		{
+			TSharedPtr<FVirtualShadowMapClipmap> Clipmap;
+			// for culling and other misc reasons.
+			FProjectedShadowInfo* ProjectedShadowInfo = nullptr;
+		};
+
+		// clipmaps, Indexed by the EVirtualShadowTypeId.
+		TStaticArray<FClipmapInfo, EVirtualShadowTypeId::Max> ClipmapInfos;
 	};
+
+	// Indexed by CommonSetups[LightId].SetupIndex + ViewIndex
 	TArray<FDirectionalLightShadowFrameSetup, SceneRenderingAllocator> DirectionalLights;
 
+	FDirectionalLightShadowFrameSetup* FindDirectional(const FLightSceneInfo::FPersistentId LightId, int32 ViewIndex);
+
 	// Links to other systems etc.
 	FDeferredShadingSceneRenderer& SceneRenderer;
 	FScene& Scene;
diff --git a/Engine/Source/Runtime/Renderer/Private/SingleLayerWaterRendering.cpp b/Engine/Source/Runtime/Renderer/Private/SingleLayerWaterRendering.cpp
--- a/Engine/Source/Runtime/Renderer/Private/SingleLayerWaterRendering.cpp
+++ b/Engine/Source/Runtime/Renderer/Private/SingleLayerWaterRendering.cpp
@@ -17,8 +17,7 @@
 #include "ScreenSpaceRayTracing.h"
 #include "SceneTextureParameters.h"
 #include "Substrate/Substrate.h"
-#include "VirtualShadowMaps/VirtualShadowMapArray.h"
-#include "VirtualShadowMaps/VirtualShadowMapProjection.h"
+#include "Shadows/ShadowSceneRenderer.h"
 #include "Lumen/LumenSceneData.h"
 #include "Lumen/LumenTracingUtils.h"
 #include "RenderCore.h"
@@ -1064,6 +1063,7 @@ void FDeferredShadingSceneRenderer::RenderSingleLayerWaterReflections(
 					ScissorRect = View.ViewRect;
 				}
 
+				FLightSceneInfo::FPersistentId LightId = SelectedForwardDirectionalLightProxy->GetLightSceneInfo()->Id;
 				const FVisibleLightInfo& VisibleLightInfo = VisibleLightInfos[SelectedForwardDirectionalLightProxy->GetLightSceneInfo()->Id];
 
 				if (VisibleLightInfo.VirtualShadowMapClipmaps.Num() > 0)
@@ -1077,18 +1077,16 @@ void FDeferredShadingSceneRenderer::RenderSingleLayerWaterReflections(
 						TiledVSMProjection.TileSize = TiledScreenSpaceReflection.TileSize;
 					}
 
-					RenderVirtualShadowMapProjection(
+					GetSceneExtensionsRenderers().GetRenderer<FShadowSceneRenderer>().RenderVirtualShadowMapProjection(
 						GraphBuilder,
 						SceneTextures,
+						LightId,
 						View, ViewIndex,
-						VirtualShadowMapArray,
 						ScissorRect,
 						EVirtualShadowMapProjectionInputType::GBuffer,
-						VisibleLightInfo.FindShadowClipmapForView(&View),
 						true, // bModulateRGB
 						bRunTiled  ? &TiledVSMProjection : nullptr,
-						SceneWithoutWaterTextures.SeparatedMainDirLightTexture,
-						VisibleLightInfo.FindFirstPersonShadowClipmapForView(&View));
+						SceneWithoutWaterTextures.SeparatedMainDirLightTexture);
 				}
 			}
 
diff --git a/Engine/Source/Runtime/Renderer/Private/VirtualShadowMaps/VirtualShadowMapArray.cpp b/Engine/Source/Runtime/Renderer/Private/VirtualShadowMaps/VirtualShadowMapArray.cpp
--- a/Engine/Source/Runtime/Renderer/Private/VirtualShadowMaps/VirtualShadowMapArray.cpp
+++ b/Engine/Source/Runtime/Renderer/Private/VirtualShadowMaps/VirtualShadowMapArray.cpp
@@ -23,6 +23,7 @@
 #include "SingleLayerWaterRendering.h"
 #include "RenderUtils.h"
 #include "SceneCulling/SceneCullingRenderer.h"
+#include "Shadows/ShadowSceneRenderer.h"
 #include "Froxel/Froxel.h"
 
 #define DEBUG_ALLOW_STATIC_SEPARATE_WITHOUT_CACHING 0
@@ -683,7 +684,7 @@ void FVirtualShadowMapArray::Initialize(
 		// Fixed physical page pool width, we adjust the height to accomodate the requested maximum
 		// NOTE: Row size in pages has to be POT since we use mask & shift in place of integer ops
 		// NOTE: This assumes GetMax2DTextureDimension() is a power of two on supported platforms
-		const uint32 PhysicalPagesX = FMath::DivideAndRoundDown(GetMax2DTextureDimension(), FVirtualShadowMap::PageSize);
+		const uint32 PhysicalPagesX = FMath::DivideAndRoundDown(CacheManager->GetPhysicalMaxWidth(), FVirtualShadowMap::PageSize);
 		check(FMath::IsPowerOfTwo(PhysicalPagesX));
 		const int32 MaxPhysicalPages = CVarMaxPhysicalPages.GetValueOnRenderThread();
 		uint32 PhysicalPagesY = FMath::DivideAndRoundUp((uint32)FMath::Max(1, MaxPhysicalPages), PhysicalPagesX);	
@@ -1812,7 +1813,7 @@ void FVirtualShadowMapArray::BuildPageAllocations(
 	FRDGBuilder& GraphBuilder,
 	const FMinimalSceneTextures& SceneTextures,
 	const TConstArrayView<FViewInfo>& Views,
-	const FSortedLightSetSceneInfo& SortedLightsInfo,
+	const FSceneRenderer& SceneRenderer,
 	const TConstArrayView<FVisibleLightInfo>& VisibleLightInfos,
 	const FSingleLayerWaterPrePassResult* SingleLayerWaterPrePassResult,
 	const FFrontLayerTranslucencyData& FrontLayerTranslucencyData,
@@ -1827,6 +1828,8 @@ void FVirtualShadowMapArray::BuildPageAllocations(
 		return;
 	}
 
+	const FShadowSceneRenderer& ShadowSceneRenderer = SceneRenderer.GetSceneExtensionsRenderers().GetRenderer<FShadowSceneRenderer>();
+
 	RDG_EVENT_SCOPE(GraphBuilder, "FVirtualShadowMapArray::BuildPageAllocation");
 	SCOPED_NAMED_EVENT(FVirtualShadowMapArray_BuildPageAllocation, FColor::Emerald);
 
@@ -2063,29 +2066,7 @@ void FVirtualShadowMapArray::BuildPageAllocations(
 		FVirtualShadowMapPerViewParameters& PerViewData = PerViewParameters[ViewIndex];
 
 		// Gather directional light virtual shadow maps
-		TArray<int32, SceneRenderingAllocator> DirectionalLightIds;
-		for (const FVisibleLightInfo& VisibleLightInfo : VisibleLightInfos)
-		{
-			for (const TSharedPtr<FVirtualShadowMapClipmap>& Clipmap : VisibleLightInfo.VirtualShadowMapClipmaps)
-			{
-				// mark pages from both views in stereo for the left (primary) view clipmap - assumes GetPrimaryView returns self for non-stereo views.
-				if (Clipmap->GetDependentView() == Views[ViewIndex].GetPrimaryView())
-				{
-					// NOTE: Shader assumes all levels from a given clipmap are contiguous
-					int32 ClipmapID = Clipmap->GetVirtualShadowMapId();
-					DirectionalLightIds.Add(ClipmapID);
-				}
-			}
-
-			for (const TSharedPtr<FVirtualShadowMapClipmap>& Clipmap : VisibleLightInfo.FirstPersonVirtualShadowMapClipmaps)
-			{
-				// mark pages from both views in stereo for the left (primary) view clipmap - assumes GetPrimaryView returns self for non-stereo views.
-				if (Clipmap->GetDependentView() == Views[ViewIndex].GetPrimaryView())
-				{
-					DirectionalLightIds.Add(Clipmap->GetVirtualShadowMapId());
-				}
-			}
-		}	
+		TArray<int32, SceneRenderingAllocator> DirectionalLightIds = ShadowSceneRenderer.GatherClipmapIds(ViewIndex);
 		
 		// This view contained no local lights (that were stored in the light grid), and no directional lights, so nothing to do.
 		if (View.ForwardLightingResources.ForwardLightUniformParameters->NumLocalLights + DirectionalLightIds.Num() == 0)
diff --git a/Engine/Source/Runtime/Renderer/Private/VirtualShadowMaps/VirtualShadowMapArray.h b/Engine/Source/Runtime/Renderer/Private/VirtualShadowMaps/VirtualShadowMapArray.h
--- a/Engine/Source/Runtime/Renderer/Private/VirtualShadowMaps/VirtualShadowMapArray.h
+++ b/Engine/Source/Runtime/Renderer/Private/VirtualShadowMaps/VirtualShadowMapArray.h
@@ -47,7 +47,6 @@ struct FNaniteVirtualShadowMapRenderPass
 	FSceneInstanceCullingQuery *SceneInstanceCullingQuery = nullptr;
 	TArray<FProjectedShadowInfo*, SceneRenderingAllocator> Shadows;
 	uint32 TotalPrimaryViews = 0u;
-	uint32 MaxNumMips = 0u;
 	uint32 MaxCullingViews = 0u;
 	Nanite::FPackedViewArray* VirtualShadowMapViews = nullptr;
 };
@@ -342,7 +341,7 @@ public:
 		FRDGBuilder& GraphBuilder,
 		const FMinimalSceneTextures& SceneTextures,
 		const TConstArrayView<FViewInfo> &Views,
-		const FSortedLightSetSceneInfo& SortedLights, 
+		const FSceneRenderer& SceneRenderer,
 		const TConstArrayView<FVisibleLightInfo>& VisibleLightInfos,
 		const FSingleLayerWaterPrePassResult* SingleLayerWaterPrePassResult,
 		const FFrontLayerTranslucencyData& FrontLayerTranslucencyData,
diff --git a/Engine/Source/Runtime/Renderer/Private/VirtualShadowMaps/VirtualShadowMapCacheManager.cpp b/Engine/Source/Runtime/Renderer/Private/VirtualShadowMaps/VirtualShadowMapCacheManager.cpp
--- a/Engine/Source/Runtime/Renderer/Private/VirtualShadowMaps/VirtualShadowMapCacheManager.cpp
+++ b/Engine/Source/Runtime/Renderer/Private/VirtualShadowMaps/VirtualShadowMapCacheManager.cpp
@@ -1403,6 +1403,12 @@ void FVirtualShadowMapArrayCacheManager::ReallocatePersistentPrimitiveIndices()
 	// Do instance-based GPU allocations here too? For now we do them lazily each frame when the FVirtualShadowMapArray gets constructed
 }
 
+uint32 FVirtualShadowMapArrayCacheManager::GetPhysicalMaxWidth()
+{
+	return GetMax2DTextureDimension();
+	//return 2048u;
+}
+
 BEGIN_SHADER_PARAMETER_STRUCT(FInvalidatePagesParameters, )
 	SHADER_PARAMETER_RDG_UNIFORM_BUFFER(FVirtualShadowMapUniformParameters, VirtualShadowMap)
 	SHADER_PARAMETER_RDG_UNIFORM_BUFFER(FSceneUniformParameters, Scene)
diff --git a/Engine/Source/Runtime/Renderer/Private/VirtualShadowMaps/VirtualShadowMapCacheManager.h b/Engine/Source/Runtime/Renderer/Private/VirtualShadowMaps/VirtualShadowMapCacheManager.h
--- a/Engine/Source/Runtime/Renderer/Private/VirtualShadowMaps/VirtualShadowMapCacheManager.h
+++ b/Engine/Source/Runtime/Renderer/Private/VirtualShadowMaps/VirtualShadowMapCacheManager.h
@@ -422,6 +422,7 @@ public:
 	// NOTE: Can move to private after we remove old invalidations path
 	void ReallocatePersistentPrimitiveIndices();
 
+	uint32 GetPhysicalMaxWidth();
 private:
 
 	/**
diff --git a/Engine/Source/Runtime/Renderer/Private/VirtualShadowMaps/VirtualShadowMapClipmap.cpp b/Engine/Source/Runtime/Renderer/Private/VirtualShadowMaps/VirtualShadowMapClipmap.cpp
--- a/Engine/Source/Runtime/Renderer/Private/VirtualShadowMaps/VirtualShadowMapClipmap.cpp
+++ b/Engine/Source/Runtime/Renderer/Private/VirtualShadowMaps/VirtualShadowMapClipmap.cpp
@@ -161,7 +161,10 @@ FVirtualShadowMapClipmapConfig FVirtualShadowMapClipmapConfig::GetGlobal()
 		.LastLevel = CVarVirtualShadowMapClipmapLastLevel.GetValueOnRenderThread(),
 		.FirstCoarseLevel = bMarkCoarsePagesDirectional ? CVarVirtualShadowMapClipmapFirstCoarseLevel.GetValueOnRenderThread() : -1,
 		.LastCoarseLevel = bMarkCoarsePagesDirectional ? CVarVirtualShadowMapClipmapLastCoarseLevel.GetValueOnRenderThread() : -1,
-		.bForceInvalidate = CVarForceInvalidateDirectionalVSM.GetValueOnRenderThread() != 0
+		.ResolutionLodBias = CVarVirtualShadowMapResolutionLodBiasDirectional.GetValueOnRenderThread(),
+		.ResolutionLodBiasMoving = CVarVirtualShadowMapResolutionLodBiasDirectionalMoving.GetValueOnRenderThread(),
+		.bForceInvalidate = CVarForceInvalidateDirectionalVSM.GetValueOnRenderThread() != 0,
+		.bCullDynamicTightly = CVarVSMClipmapCullDynamicTightly.GetValueOnRenderThread()
 	};
 }
 
@@ -217,7 +220,7 @@ FVirtualShadowMapClipmap::FVirtualShadowMapClipmap(
 	// For now we adjust resolution by just biasing the page we look up in. This is wasteful in terms of page table vs.
 	// just resizing the virtual shadow maps for each clipmap, but convenient for now. This means we need to additionally bias
 	// which levels are present.
-	ResolutionLodBias = FVirtualShadowMapArray::InterpolateResolutionBias(CVarVirtualShadowMapResolutionLodBiasDirectional.GetValueOnRenderThread(), CVarVirtualShadowMapResolutionLodBiasDirectionalMoving.GetValueOnRenderThread(), LightMobilityFactor) + FMath::Log2(LodScale);
+	ResolutionLodBias = FVirtualShadowMapArray::InterpolateResolutionBias(Config.ResolutionLodBias, Config.ResolutionLodBiasMoving, LightMobilityFactor) + FMath::Log2(LodScale);
 	ResolutionLodBias += GetLightSceneInfo().Proxy->GetVSMResolutionLodBias();
 	// Clamp negative absolute resolution biases as they would exceed the maximum resolution/ranges allocated
 	ResolutionLodBias = FMath::Max(0.0f, ResolutionLodBias);
@@ -406,7 +409,7 @@ FVirtualShadowMapClipmap::FVirtualShadowMapClipmap(
 		ViewCenterDeltaZ = ViewCenter.Z - ClipmapLevelEntry->Clipmap.ViewCenterZ;
 		ViewRadiusZ = ClipmapLevelEntry->Clipmap.ViewRadiusZ;
 
-		if (Config.bForceInvalidate && CVarVSMClipmapCullDynamicTightly.GetValueOnRenderThread())
+		if (Config.bForceInvalidate && Config.bCullDynamicTightly )
 		{
 			// Far plane in the clipmap should be just the end of the visible range
 			float CullingEndDistance = GetLevelRadius(static_cast<float>(AbsoluteLevel) - ResolutionLodBias); 
diff --git a/Engine/Source/Runtime/Renderer/Private/VirtualShadowMaps/VirtualShadowMapClipmap.h b/Engine/Source/Runtime/Renderer/Private/VirtualShadowMaps/VirtualShadowMapClipmap.h
--- a/Engine/Source/Runtime/Renderer/Private/VirtualShadowMaps/VirtualShadowMapClipmap.h
+++ b/Engine/Source/Runtime/Renderer/Private/VirtualShadowMaps/VirtualShadowMapClipmap.h
@@ -29,8 +29,11 @@ public:
 	int32 FirstCoarseLevel = -1;
 	int32 LastCoarseLevel = -1;
 	uint32 ShadowTypeId = 0u;
+	float ResolutionLodBias = 0.0f;
+	float ResolutionLodBiasMoving = 0.0f;
 	bool bForceInvalidate = false;
 	bool bIsFirstPersonShadow = false;
+	bool bCullDynamicTightly = false;
 };
 
 class FVirtualShadowMapClipmap : FRefCountedObject
@@ -146,7 +149,6 @@ private:
 
 	int32 FirstLevel;
 	float ResolutionLodBias;
-	float MaxRadius;
 
 	struct FLevelData
 	{
