CIS fix for less forgiving platforms.

#rnx

diff --git a/Engine/Source/Runtime/Renderer/Private/Shadows/ShadowSceneRenderer.h b/Engine/Source/Runtime/Renderer/Private/Shadows/ShadowSceneRenderer.h
--- a/Engine/Source/Runtime/Renderer/Private/Shadows/ShadowSceneRenderer.h
+++ b/Engine/Source/Runtime/Renderer/Private/Shadows/ShadowSceneRenderer.h
@@ -247,6 +247,12 @@ private:
 		// A clipmap may belong to more than one view (in stereo mode, specifically)
 		uint32 ViewMask = 0u;
 
+		FDirectionalLightShadowFrameSetup(FLightSceneInfo::FPersistentId InLightId, uint32 InViewMask) 
+			: LightId(InLightId)
+			, ViewMask(InViewMask)
+		{
+		}
+
 		struct FClipmapInfo
 		{
 			TSharedPtr<FVirtualShadowMapClipmap> Clipmap;
