[SingleLayerWater] Optimized SLW to only run ComposeVolumetricRenderTargetOverSceneUnderWater() if the camera is possibly intersecting the water surface. While it's functionally fine to always run this pass when the renderer thinks that the camera is not underwater, it costs performance and is only really necessary in cases where the the camera is underwater (or intersects the water surface) but the FSceneView::IsUnderwater() is false. The IsUnderwater() is not really precise and considering the intersection case, can't possibly be precise at all. To somewhat remedy the issue, this CL introduces a new enum on the view, which tells the renderer whether the view is completley inside water, possibly intersecting it or completely outside. This enum should be conservative, so even if the view is entirely underwater, it could still be set to PossiblyIntersecting, so it should probably be used in conjuction with the existing IsUnderwater() function, which while also imprecise, models the water surface more closely than the new enum. In the future, it might be worth it to change how volumetrics are rendered such that they make further use of this new enum, allowing them to always composite correctly (as far as reasonably possible).

diff --git a/Engine/Plugins/Experimental/Water/Source/Runtime/Private/WaterSubsystem.cpp b/Engine/Plugins/Experimental/Water/Source/Runtime/Private/WaterSubsystem.cpp
--- a/Engine/Plugins/Experimental/Water/Source/Runtime/Private/WaterSubsystem.cpp
+++ b/Engine/Plugins/Experimental/Water/Source/Runtime/Private/WaterSubsystem.cpp
@@ -780,6 +780,12 @@ void UWaterSubsystem::ComputeUnderwaterPostProcess(FVector ViewLocation, FSceneV
 	const float PrevDepthUnderwater = CachedDepthUnderwater;
 	CachedDepthUnderwater = -1;
 
+	// Compute distance from view origin to the corner of the near plane. This distance needs to be taken into account when computing whether the view intersects the water surface.
+	const FVector4f NearPlaneCornerViewSpace = FVector4f(SceneView->ViewMatrices.GetInvProjectionMatrix().TransformFVector4(FVector4(1.0f, 1.0f, (bool)ERHIZBuffer::IsInverted ? 1.0f : 0.0f, 1.0f)));
+	const float ViewToNearPlaneCornerDistance = FVector2f(NearPlaneCornerViewSpace / NearPlaneCornerViewSpace.W).Length();
+	bool bAnyDefinitelyUnderwater = false;
+	bool bAnyPossiblyUnderwater = false;
+
 	bool bUnderwaterForPostProcess = false;
 
 	// Trace just a small distance extra from the viewpoint to account for waves since the waves wont traced against
@@ -865,6 +871,11 @@ void UWaterSubsystem::ComputeUnderwaterPostProcess(FVector ViewLocation, FSceneV
 
 		for (const FWaterBodyPostProcessQuery& Query : WaterBodyQueriesToProcess)
 		{
+			const float MaxWaveHeight = Query.QueryResult.GetWaveInfo().MaxHeight;
+			const double ImmersionWithoutWaves = Query.QueryResult.GetWaterPlaneLocation().Z - ViewLocation.Z; // Positive is under water, negative above water
+			bAnyDefinitelyUnderwater |= (ImmersionWithoutWaves - ViewToNearPlaneCornerDistance) > MaxWaveHeight;
+			bAnyPossiblyUnderwater |= (FMath::Abs(ImmersionWithoutWaves) - ViewToNearPlaneCornerDistance) <= MaxWaveHeight;
+
 			float LocalDepthUnderwater = 0.0f;
 
 			// Underwater is fudged a bit for post process so its possible to get a true return here but depth underwater is < 0
@@ -885,6 +896,18 @@ void UWaterSubsystem::ComputeUnderwaterPostProcess(FVector ViewLocation, FSceneV
 	}
 
 	SceneView->UnderwaterDepth = CachedDepthUnderwater;
+	if (bAnyPossiblyUnderwater)
+	{
+		SceneView->WaterIntersection = EViewWaterIntersection::PossiblyIntersectingWater;
+	}
+	else if (bAnyDefinitelyUnderwater)
+	{
+		SceneView->WaterIntersection = EViewWaterIntersection::InsideWater;
+	}
+	else
+	{
+		SceneView->WaterIntersection = EViewWaterIntersection::OutsideWater;
+	}
 
 	if (!bUnderwaterForPostProcess || !IsUnderwaterPostProcessEnabled() || SceneView->Family->EngineShowFlags.PathTracing)
 	{
diff --git a/Engine/Source/Runtime/Engine/Public/SceneView.h b/Engine/Source/Runtime/Engine/Public/SceneView.h
--- a/Engine/Source/Runtime/Engine/Public/SceneView.h
+++ b/Engine/Source/Runtime/Engine/Public/SceneView.h
@@ -1349,6 +1349,13 @@ namespace EDrawDynamicFlags
 	};
 }
 
+enum class EViewWaterIntersection : uint8
+{
+	OutsideWater,
+	PossiblyIntersectingWater,
+	InsideWater
+};
+
 // Hacky base class to avoid 8 bytes of padding after the vtable
 class FSceneViewFixLayout
 {
@@ -1634,6 +1641,9 @@ public:
 	/** How far below the water surface this view is. -1 means the view is out of water. */
 	float UnderwaterDepth;
 
+	/** Whether the view intersects water. */
+	EViewWaterIntersection WaterIntersection = EViewWaterIntersection::OutsideWater;
+
 	/** Global clipping plane being applied to the scene, or all 0's if disabled.  This is used when rendering the planar reflection pass. */
 	FPlane GlobalClippingPlane;
 
diff --git a/Engine/Source/Runtime/Renderer/Private/SingleLayerWaterRendering.cpp b/Engine/Source/Runtime/Renderer/Private/SingleLayerWaterRendering.cpp
--- a/Engine/Source/Runtime/Renderer/Private/SingleLayerWaterRendering.cpp
+++ b/Engine/Source/Runtime/Renderer/Private/SingleLayerWaterRendering.cpp
@@ -1336,7 +1336,10 @@ void FDeferredShadingSceneRenderer::RenderSingleLayerWater(
 		}
 		if (!bIsCameraUnderWater && SceneWithoutWaterTextures.ColorTexture && bShouldRenderVolumetricCloud)
 		{
-			// This path is only taken when rendering the clouds in a render target that can be composited
+			// This path is only taken when rendering the clouds in a render target that can be composited and when the view is possibly intersecting the water surface.
+			// The !bIsCameraUnderWater check is a bit misleading: In this case, volumetrics (including clouds) are rendered after water, but since bIsCameraUnderWater is
+			// somewhat imprecise, it's possible for the camera to be fully or partially below the water surface and bIsCameraUnderWater being false. Without this call,
+			// clouds would not be visible when looking up from under the water surface in such a case.
 			ComposeVolumetricRenderTargetOverSceneUnderWater(GraphBuilder, InViews, SceneWithoutWaterTextures, SceneTextures);
 		}
 	}
diff --git a/Engine/Source/Runtime/Renderer/Private/VolumetricRenderTarget.cpp b/Engine/Source/Runtime/Renderer/Private/VolumetricRenderTarget.cpp
--- a/Engine/Source/Runtime/Renderer/Private/VolumetricRenderTarget.cpp
+++ b/Engine/Source/Runtime/Renderer/Private/VolumetricRenderTarget.cpp
@@ -97,9 +97,13 @@ bool IsVolumetricRenderTargetAsyncCompute()
 	return GSupportsEfficientAsyncCompute && CVarVolumetricRenderTargetPreferAsyncCompute.GetValueOnRenderThread() > 0 && !bCloudComputePathDisabled;
 }
 
-static bool ShouldViewComposeVolumetricRenderTarget(const FViewInfo& ViewInfo)
+static bool ShouldViewComposeVolumetricRenderTarget(const FViewInfo& ViewInfo, bool bComposeCameraIntersectingWater = false)
 {
-	return ShouldViewRenderVolumetricCloudRenderTarget(ViewInfo);
+	// IsUnderwater() is imprecise, but is currently what's used to decide whether to render clouds into SceneColor before or after water. In order to handle cases where
+	// IsUnderwater() == false but the camera is still fully or partially underwater, we need to composite clouds into SceneColorWithoutWater in a special pass. Since this is more of an edge case,
+	// we rely on FSceneView::WaterIntersection to only do run this pass when absolutely necessary.
+	const bool bComposeCameraIntersectingWaterRelevant = (ViewInfo.WaterIntersection == EViewWaterIntersection::PossiblyIntersectingWater) && !ViewInfo.IsUnderwater();
+	return ShouldViewRenderVolumetricCloudRenderTarget(ViewInfo) && (!bComposeCameraIntersectingWater || bComposeCameraIntersectingWaterRelevant);
 }
 
 static uint32 GetMainDownsampleFactor(int32 Mode)
@@ -151,13 +155,13 @@ static void GetTextureSafeUvCoordBound(const FIntPoint& TextureViewRect, FRDGTex
 	TextureValidUvRect.W = (float(TextureViewRect.Y) - 0.51f) / float(TexSize.Y);
 };
 
-static bool AnyViewRequiresProcessing(TArrayView<FViewInfo> Views)
+static bool AnyViewRequiresProcessing(TArrayView<FViewInfo> Views, bool bComposeCameraIntersectingWater = false)
 {
 	bool bAnyViewRequiresProcessing = false;
 	for (int32 ViewIndex = 0; ViewIndex < Views.Num(); ViewIndex++)
 	{
 		FViewInfo& ViewInfo = Views[ViewIndex];
-		bAnyViewRequiresProcessing |= ShouldViewComposeVolumetricRenderTarget(ViewInfo);
+		bAnyViewRequiresProcessing |= ShouldViewComposeVolumetricRenderTarget(ViewInfo, bComposeCameraIntersectingWater);
 	}
 	return bAnyViewRequiresProcessing;
 }
@@ -1050,7 +1054,7 @@ void ComposeVolumetricRenderTargetOverSceneUnderWater(
 	const FSceneWithoutWaterTextures& WaterPassData,
 	const FMinimalSceneTextures& SceneTextures)
 {
-	if (!AnyViewRequiresProcessing(Views))
+	if (!AnyViewRequiresProcessing(Views, true /*bComposeCameraIntersectingWater*/))
 	{
 		return;
 	}
@@ -1064,7 +1068,7 @@ void ComposeVolumetricRenderTargetOverSceneUnderWater(
 	for (int32 ViewIndex = 0; ViewIndex < Views.Num(); ViewIndex++)
 	{
 		FViewInfo& ViewInfo = Views[ViewIndex];
-		if (!ShouldViewRenderVolumetricCloudRenderTarget(ViewInfo) || !ViewInfo.ShouldRenderView())
+		if (!ShouldViewComposeVolumetricRenderTarget(ViewInfo, true /*bComposeCameraIntersectingWater*/) || !ViewInfo.ShouldRenderView())
 		{
 			continue;
 		}
