Unconditionally create the FShadowScene extension as the Renderer is unconditionally requested in a number of places.

#rnx

diff --git a/Engine/Source/Runtime/Renderer/Private/Shadows/ShadowScene.cpp b/Engine/Source/Runtime/Renderer/Private/Shadows/ShadowScene.cpp
--- a/Engine/Source/Runtime/Renderer/Private/Shadows/ShadowScene.cpp
+++ b/Engine/Source/Runtime/Renderer/Private/Shadows/ShadowScene.cpp
@@ -287,7 +287,7 @@ void FShadowScene::WaitForSceneLightsUpdateTask()
 
 bool FShadowScene::ShouldCreateExtension(FScene& InScene)
 {
-	return DoesPlatformSupportVirtualShadowMaps(GetFeatureLevelShaderPlatform(InScene.GetFeatureLevel()));
+	return true;
 }
 
 void FShadowScene::InitExtension(FScene& InScene)
