Path Tracer: Use cloud acceleration map to optimize ray marching

Pass the cloud acceleration map texture to the callable shader so that it can be used to acceleration raymarching of clouds. This allows us to skip empty space between hits and greatly reduces the amount of ray marching needed around the "edges" of clouds in most cases.

Fix a subtle bug introduced in the last refactor to callable shaders that was leading to the cloud density being used for all volumetric segments, not just the ones overlapping clouds.

Refactor the early termination heuristic on transmittance to be driven by a single function for consistency.

Fix intensity of emissive clouds

diff --git a/Engine/Shaders/Private/PathTracing/PathTracingCore.ush b/Engine/Shaders/Private/PathTracing/PathTracingCore.ush
--- a/Engine/Shaders/Private/PathTracing/PathTracingCore.ush
+++ b/Engine/Shaders/Private/PathTracing/PathTracingCore.ush
@@ -201,6 +201,12 @@ struct FPathState
 			     FirstScatterType == PATHTRACER_SCATTER_REFRACT);
 	}
 
+	bool HasMadeDiffuseOrSpecularScatter()
+	{
+		uint FirstScatterType = GetFirstScatterType();
+		return FirstScatterType == PATHTRACER_SCATTER_DIFFUSE || FirstScatterType == PATHTRACER_SCATTER_SPECULAR;
+	}
+
 	bool UpdateScatterType(uint ScatterType)
 	{
 		// Keep track of the first non-refractive scatter type, and leave it "locked" beyond this
@@ -744,7 +750,7 @@ FPackedPathTracingPayload TraceTransparentRay(inout FPathState PathState, int Bo
 						PathState.RandSequence,
 						Interval.VolumeMask |
 						  (VolumeFlags & ~PATH_TRACER_VOLUME_ENABLE_MASK) |
-						  (bIsCameraRay ? PATH_TRACER_VOLUME_CALLABLE_FLAGS_IS_CAMERA_RAY : 0) |
+						  (min(Bounce, 31) << PATH_TRACER_VOLUME_CALLABLE_FLAGS_BOUNCE_SHIFT) |
 						  (bNeedSample  ? PATH_TRACER_VOLUME_CALLABLE_FLAGS_GET_SAMPLE    : 0),
 						PathThroughput,
 						Interval.CloudDensity
@@ -845,8 +851,7 @@ FPackedPathTracingPayload TraceTransparentRay(inout FPathState PathState, int Bo
 							VolumeRadiance += Result.Emission * PathThroughput;
 							// keep tracing through the volume
 							PathThroughput *= SigmaN;
-							// early termination - fades off from 0.002 to 0.001 and then will be 0
-							PathThroughput *= saturate(10000.0 * max3(PathThroughput.x, PathThroughput.y, PathThroughput.z) - 1.0);
+							PathThroughput *= LowThroughputClampFactor(PathThroughput);
 							if (!any(PathThroughput > 0))
 							{
 								break;
@@ -1699,7 +1704,7 @@ bool PathTracingKernel(inout FPathState PathState, int Bounce)
 
 #if 0
 	// visualize the shadow ray handling directly
-	PathState.Radiance = TraceTransparentVisibilityRay(PathState.Ray, 0.0, 0, true, true, PathState.RandSequence);
+	PathState.Radiance = TraceTransparentVisibilityRay(PathState.Ray, 0.0, 0.0, 0, true, true, PathState.RandSequence);
 	return false;
 #endif
 
@@ -1923,7 +1928,10 @@ bool PathTracingKernel(inout FPathState PathState, int Bounce)
 			{
 				LightEstimate *= GetVolumetricScatteringIntensity(LightIndex);
 			}
-			LightEstimate *= GetLightIndirectScale(LightIndex, PathState.GetPathRoughness());
+			if (PathState.HasMadeDiffuseOrSpecularScatter())
+			{
+				LightEstimate *= GetLightIndirectScale(LightIndex, PathState.GetPathRoughness());
+			}
 			LightPickingCdfSum += LightEstimate;
 			LightPickingCdf[Index] = LightPickingCdfSum;
 		}
@@ -1960,7 +1968,10 @@ bool PathTracingKernel(inout FPathState PathState, int Bounce)
 					LightSample.RadianceOverPdf *= GetVolumetricScatteringIntensity(LightId);
 				}
 
-				LightSample.RadianceOverPdf *= GetLightIndirectScale(LightId, PathState.GetPathRoughness());
+				if (PathState.HasMadeDiffuseOrSpecularScatter())
+				{
+					LightSample.RadianceOverPdf *= GetLightIndirectScale(LightId, PathState.GetPathRoughness());
+				}
 
 				// #dxr_todo: Is it cheaper to fire the ray first? Or eval the material first?
 				if (any(LightSample.RadianceOverPdf > 0))
@@ -2117,7 +2128,10 @@ bool PathTracingKernel(inout FPathState PathState, int Bounce)
 				LightContrib *= MISWeightPower(MaterialSample.Pdf, LightResult.Pdf * LightPickPdf);
 			}
 
-			LightContrib *= GetLightIndirectScale(LightId, PathState.GetPathRoughness());
+			if (PathState.HasMadeDiffuseOrSpecularScatter())
+			{
+				LightContrib *= GetLightIndirectScale(LightId, PreviousPathRoughness);
+			}
 
 			if (any(LightContrib > 0))
 			{
diff --git a/Engine/Shaders/Private/PathTracing/PathTracingVolumetricCloudMaterialShader.usf b/Engine/Shaders/Private/PathTracing/PathTracingVolumetricCloudMaterialShader.usf
--- a/Engine/Shaders/Private/PathTracing/PathTracingVolumetricCloudMaterialShader.usf
+++ b/Engine/Shaders/Private/PathTracing/PathTracingVolumetricCloudMaterialShader.usf
@@ -19,9 +19,20 @@
 #define FogRadius 		PathTracingCloudParameters.FogRadius
 #define FogFalloffClamp	PathTracingCloudParameters.FogFalloffClamp
 
-#define PlanetCenterTranslatedWorldHi PathTracingCloudParameters.PlanetCenterTranslatedWorldHi
-#define PlanetCenterTranslatedWorldLo PathTracingCloudParameters.PlanetCenterTranslatedWorldLo
-#define Atmosphere 		PathTracingCloudParameters
+#define PlanetCenterTranslatedWorldHi 	PathTracingCloudParameters.PlanetCenterTranslatedWorldHi
+#define PlanetCenterTranslatedWorldLo 	PathTracingCloudParameters.PlanetCenterTranslatedWorldLo
+#define Atmosphere 						PathTracingCloudParameters
+
+#define CloudAccelerationMap 			PathTracingCloudParameters.CloudAccelerationMap
+#define CloudAccelerationMapSampler 	PathTracingCloudParameters.CloudAccelerationMapSampler
+
+#define CloudClipX						PathTracingCloudParameters.CloudClipX						
+#define CloudClipY						PathTracingCloudParameters.CloudClipY
+#define CloudClipZ						PathTracingCloudParameters.CloudClipZ
+#define CloudClipDistKm                 PathTracingCloudParameters.CloudClipDistKm
+#define CloudAccelMapResolution			PathTracingCloudParameters.CloudAccelMapResolution
+#define CloudVoxelWidth					PathTracingCloudParameters.CloudVoxelWidth
+#define CloudInvVoxelWidth				PathTracingCloudParameters.CloudInvVoxelWidth
 
 
 #include "./Volume/PathTracingCloudsMaterialCommon.ush"
@@ -33,7 +44,6 @@
 
 uint MaxRaymarchSteps;
 
-
 FVolumeShadedResult CloudsGetDensity(float3 TranslatedWorldPos)
 {
 	FDFVector3 AbsoluteWorldPosition = DFFastSubtract(TranslatedWorldPos, ResolvedView.PreViewTranslation);
@@ -79,29 +89,24 @@ RAY_TRACING_ENTRY_CALLABLE(PathTracingVolumetricCloudMaterialShader, FPackedPath
 		for (int Step = 0; Step < PathTracingCloudParameters.MaxRaymarchSteps; Step++)
 		{
 			/* take stochastic steps along the ray to estimate transmittance (null scattering) */
-			float3 ColorChannelPdf = Throughput;													
 			/* Sample the distance to the next interaction */										
 			float RandValue = RandomSequence_GenerateSample1D(RandSeq);						
-			float DeltaT = -log(1 - RandValue) * InvSigmaBar;
-			if (TMin + DeltaT < TMax)
+			TMin -= log(1 - RandValue) * InvSigmaBar;
+			if (TMin >= TMax)
 			{
-				TMin += DeltaT;
-				/* our ray marching step stayed inside the atmo and is still in front of the next hit */
-				float3 WorldPos = Ray.Origin + TMin * Ray.Direction;
-				/* clamp to make sure we never exceed the majorant (should not be the case, but need to avoid any possible numerical issues) */
-				float3 SigmaT = min(CloudsGetDensity(WorldPos).SigmaT, SigmaBar);
-				/* keep tracing through the volume */
-				Throughput *= 1.0 - SigmaT * InvSigmaBar;
-				/* early termination - fades off from 0.002 to 0.001 and then will be 0 */
-				Throughput *= saturate(1000.0 * max3(Throughput.x, Throughput.y, Throughput.z) - 1.0);
-				if (!any(Throughput > 0))
-				{
-					break;
-				}
+				/* exit the ray marching loop */
+				break;
 			}
-			else
+			/* our ray marching step stayed inside the atmo and is still in front of the next hit */
+			float3 WorldPos = Ray.Origin + TMin * Ray.Direction;
+			/* clamp to make sure we never exceed the majorant (should not be the case, but need to avoid any possible numerical issues) */
+			float3 SigmaT = min(CloudsGetDensity(WorldPos).SigmaT, SigmaBar);
+			/* keep tracing through the volume */
+			Throughput *= 1.0 - SigmaT * InvSigmaBar;
+			/* Encourage early termination of paths with low contribution */
+			Throughput *= LowThroughputClampFactor(Throughput);
+			if (!any(Throughput > 0))
 			{
-				/* exit the ray marching loop */
 				break;
 			}
 		}
@@ -113,19 +118,80 @@ RAY_TRACING_ENTRY_CALLABLE(PathTracingVolumetricCloudMaterialShader, FPackedPath
 	{
 		FRISContext HitSample = PackedPayload.GetVolumetricCallableShaderRISContext();
 		// Loop through combined density for sampling
-		const bool bIsCameraRay = (Flags & PATH_TRACER_VOLUME_CALLABLE_FLAGS_IS_CAMERA_RAY) != 0;
+		const int Bounce = (Flags & PATH_TRACER_VOLUME_CALLABLE_FLAGS_BOUNCE_MASK) >> PATH_TRACER_VOLUME_CALLABLE_FLAGS_BOUNCE_SHIFT;
+		const bool bIsCameraRay = Bounce == 0;
 
 		// zero out the payload so we can write our output to it
 		PackedPayload = (FPackedPathTracingPayload) 0;
 		PackedPayload.HitT = -1.0; // invalid sample until we are sure we got one
 
+#define USE_VOXEL_TRAVERSAL	1
+
+#if USE_VOXEL_TRAVERSAL
+		// isect clip slabs, the AABB our clouds are enclosed in
+		float3 Ro = mul(GetCloudClipBasis(), Ray.Origin);
+		float3 Rd = mul(GetCloudClipBasis(), Ray.Direction);
+		float3 InvRd = rcp(Rd);
+
+		// now that we know the ray length, lets march through the grid to see if we can clip away anything
+		float2 P = Ro.xy + Ray.TMin * Rd.xy;
+		int2 Idx = PositionToVoxel(P);
+		float2 NextCrossingT = Ray.TMin + (VoxelToPosition(Idx + int2(Rd.xy >= 0)) - P) * InvRd.xy;
+		float2 DeltaT = abs(CloudVoxelWidth * InvRd.xy);
+		int2 Step = select_internal(Rd.xy >= 0, 1, -1);
+		int2 Stop = select_internal(Rd.xy >= 0, CloudAccelMapResolution, -1);
+#endif
+
 		float3 VolumeRadiance = 0;
 		float  VolumeAlpha = 0;
 		float3 VolumeAlbedo = 0;
 		// Limit number of steps to prevent timeouts
 		// TODO: This biases the result! Is there a better way to limit the number of steps?
-		for (int Step = 0; Step < PathTracingCloudParameters.MaxRaymarchSteps; Step++)
+		for (int RaymarchStep = 0; RaymarchStep < PathTracingCloudParameters.MaxRaymarchSteps; RaymarchStep++)
 		{
+#if USE_VOXEL_TRAVERSAL
+			// visit current voxel
+			float4 Data = CloudAccelerationMap.Load(uint3(Idx.xy, 0));
+			CloudDensity = Data.x;
+			// because each cell has slabs, we can have at most 3 segments (empty, filled, empty) per grid cell
+			float CellMinT = Ray.TMin;
+			float CellMaxT = min3(Ray.TMax, NextCrossingT.x, NextCrossingT.y);
+			float IntegrationMaxT = CellMaxT;
+			if (CloudDensity > 0)
+			{
+				float SlabT0 = (Data.z * SKY_UNIT_TO_CM - Ro.z) * InvRd.z;
+				float SlabT1 = (Data.w * SKY_UNIT_TO_CM - Ro.z) * InvRd.z;
+				float InsideMinT = max(CellMinT, min(SlabT0, SlabT1));
+				float InsideMaxT = min(CellMaxT, max(SlabT0, SlabT1));
+				// do we still have a valid interval?
+				if (InsideMinT < InsideMaxT)
+				{
+					// this interval must be going through at least 3 portions (some of which may be empty)
+					// Figure out in which we are at the moment
+					if (Ray.TMin < InsideMinT)
+					{
+						// we are outside the slab, take steps up to the slab boundary
+						CloudDensity = 0;
+						IntegrationMaxT = InsideMinT;
+					} else if (Ray.TMin < InsideMaxT)
+					{
+						// we are inside the slab, take steps up to the end of the slab
+						IntegrationMaxT = InsideMaxT;
+					} else {
+						// we are past the slab, take steps up to the cell boundary
+						CloudDensity = 0;
+					}
+				}
+				else
+				{
+					// we missed the slabs, so only keep the one segment, going through empty space
+					CloudDensity = 0;
+				}
+			}
+#else
+			float IntegrationMaxT = Ray.TMax;
+#endif
+
 			float3 SigmaBar = CloudDensity;
 			if (Flags & PATH_TRACER_VOLUME_ENABLE_ATMOSPHERE)
 			{
@@ -144,8 +210,7 @@ RAY_TRACING_ENTRY_CALLABLE(PathTracingVolumetricCloudMaterialShader, FPackedPath
 				// no energy left
 				break;
 			}
-
-			if (Ray.TMin + Distance < Ray.TMax)
+			if (Ray.TMin + Distance < IntegrationMaxT)
 			{
 				float4 Evaluation = EvaluateSpectralTransmittanceHit(Distance, SigmaBar, PathThroughput);
 				PathThroughput *= Evaluation.xyz;
@@ -196,8 +261,7 @@ RAY_TRACING_ENTRY_CALLABLE(PathTracingVolumetricCloudMaterialShader, FPackedPath
 				VolumeRadiance += Result.Emission * PathThroughput;
 				// keep tracing through the volume
 				PathThroughput *= SigmaN;
-				// early termination - fades off from 0.002 to 0.001 and then will be 0
-				PathThroughput *= saturate(10000.0 * max3(PathThroughput.x, PathThroughput.y, PathThroughput.z) - 1.0);
+				PathThroughput *= LowThroughputClampFactor(PathThroughput);
 				if (!any(PathThroughput > 0))
 				{
 					break;
@@ -207,10 +271,43 @@ RAY_TRACING_ENTRY_CALLABLE(PathTracingVolumetricCloudMaterialShader, FPackedPath
 			{
 				// update the path throughput, knowing that we escaped the medium
 				// exit the ray marching loop
-				float4 Evaluation = EvaluateSpectralTransmittanceMiss(Ray.TMax - Ray.TMin, SigmaBar, PathThroughput);
+				float4 Evaluation = EvaluateSpectralTransmittanceMiss(IntegrationMaxT - Ray.TMin, SigmaBar, PathThroughput);
 				PathThroughput *= Evaluation.xyz;
+
+#if USE_VOXEL_TRAVERSAL
+				Ray.TMin = IntegrationMaxT;
+#else
 				break;
+#endif
+			}
+
+#if USE_VOXEL_TRAVERSAL
+			// Advance to next grid cell if we have reached the end of the cell
+			if (Ray.TMin == CellMaxT)
+			{
+				if (NextCrossingT.x < NextCrossingT.y)
+				{
+					// step x first
+					Idx.x += Step.x;
+					if (Idx.x == Stop.x || Ray.TMax < NextCrossingT.x)
+					{
+						break;
+					}
+					Ray.TMin = NextCrossingT.x;
+					NextCrossingT.x += DeltaT.x;
+				}
+				else
+				{
+					Idx.y += Step.y;
+					if (Idx.y == Stop.y || Ray.TMax < NextCrossingT.y)
+					{
+						break;
+					}
+					Ray.TMin = NextCrossingT.y;
+					NextCrossingT.y += DeltaT.y;
+				}
 			}
+#endif
 		}
 
 		// Write state back here
diff --git a/Engine/Shaders/Private/PathTracing/Volume/PathTracingClouds.ush b/Engine/Shaders/Private/PathTracing/Volume/PathTracingClouds.ush
--- a/Engine/Shaders/Private/PathTracing/Volume/PathTracingClouds.ush
+++ b/Engine/Shaders/Private/PathTracing/Volume/PathTracingClouds.ush
@@ -4,7 +4,6 @@
 
 #include "./PathTracingAtmosphereCommon.ush"
 #include "./PathTracingCloudsCommon.ush"
-#include "../../SkyAtmosphereCommon.ush"
 
 
 Texture2D<float4> CloudAccelerationMap;
@@ -59,19 +58,6 @@ float4 RayCloudIntersectPrecise(float3 Ro, float3 Rd)
 	return Result;
 }
 
-int2 PositionToVoxel(float2 P)
-{
-	float CloudVoxelWidth = 2.0 * SKY_UNIT_TO_CM * CloudClipDistKm / CloudAccelMapResolution;
-	int2 Idx = int2((P + CloudClipDistKm * SKY_UNIT_TO_CM) * rcp(CloudVoxelWidth));
-	return clamp(Idx, 0, CloudAccelMapResolution - 1);
-}
-
-float2 VoxelToPosition(int2 Idx)
-{
-	float CloudVoxelWidth = 2.0 * SKY_UNIT_TO_CM * CloudClipDistKm / CloudAccelMapResolution;
-	return (Idx * CloudVoxelWidth) - CloudClipDistKm * SKY_UNIT_TO_CM;
-}
-
 FVolumeIntersection CloudsIntersect(float3 Origin, float3 Direction, float TMin, float TMax, float PathRoughness)
 {
 	if (PathRoughness > CloudRoughnessCutoff)
@@ -111,12 +97,10 @@ FVolumeIntersection CloudsIntersect(float3 Origin, float3 Direction, float TMin,
 #if 1 // use cloud accel map to tighten intersection bounds
 
 		// now that we know the ray length, lets march through the grid to see if we can clip away anything
-		float VoxelWidth = (2 * SKY_UNIT_TO_CM * CloudClipDistKm) / CloudAccelMapResolution;
-		
 		float2 P = Ro.xy + VolumeTMin * Rd.xy;
 		int2 Idx = PositionToVoxel(P);
 		float2 NextCrossingT = VolumeTMin + (VoxelToPosition(Idx + int2(Rd.xy >= 0)) - P) * InvRd.xy;
-		float2 DeltaT = abs(VoxelWidth * InvRd.xy);
+		float2 DeltaT = abs(CloudVoxelWidth * InvRd.xy);
 		int2 Step = select(Rd.xy >= 0, 1, -1);
 		int2 Stop = select(Rd.xy >= 0, CloudAccelMapResolution, -1);
 		
diff --git a/Engine/Shaders/Private/PathTracing/Volume/PathTracingCloudsCommon.ush b/Engine/Shaders/Private/PathTracing/Volume/PathTracingCloudsCommon.ush
--- a/Engine/Shaders/Private/PathTracing/Volume/PathTracingCloudsCommon.ush
+++ b/Engine/Shaders/Private/PathTracing/Volume/PathTracingCloudsCommon.ush
@@ -3,7 +3,9 @@
 #pragma once
 
 #include "../PathTracingCommon.ush"
+#include "../../SkyAtmosphereCommon.ush"
 
+#ifndef CloudClipX
 float3 CloudClipX;
 float3 CloudClipY;
 float3 CloudClipZ;
@@ -17,8 +19,11 @@ float CloudClipRadiusKm;
 
 float CloudTracingMaxDistance;
 float CloudRoughnessCutoff;
+float CloudVoxelWidth;
+float CloudInvVoxelWidth;
 int CloudAccelMapResolution;
 int CloudCallableShaderId;
+#endif
 
 float3x3 GetCloudClipBasis()
 {
@@ -50,6 +55,17 @@ float2 GetRadiusBounds()
 	return RBounds;
 }
 
+int2 PositionToVoxel(float2 P)
+{
+	int2 Idx = int2((P + CloudClipDistKm * SKY_UNIT_TO_CM) * CloudInvVoxelWidth);
+	return clamp(Idx, 0, CloudAccelMapResolution - 1);
+}
+
+float2 VoxelToPosition(int2 Idx)
+{
+	return (Idx * CloudVoxelWidth) - CloudClipDistKm * SKY_UNIT_TO_CM;
+}
+
 
 // Simplified struct only used in the cloud case (to reduce the amount of data needed in the payload for the callable shader)
 struct FSampleCloudMaterialResult
@@ -61,26 +77,3 @@ struct FSampleCloudMaterialResult
 	float PhaseG2;
 	float PhaseBlend;
 };
-
-FSampleCloudMaterialResult DecodeCloudResult(uint4 Data)
-{
-	FSampleCloudMaterialResult Result = (FSampleCloudMaterialResult) 0;
-	Result.Extinction    = PayloadDecodeHDRColor(Data.x) * (1.0 / 64.0);
-	Result.Albedo        = PayloadDecodeLDRColor(Data.y);
-	Result.EmissiveColor = PayloadDecodeHDRColor(Data.z);
-	float3 PhaseData     = PayloadDecodeLDRColor(Data.w);
-	Result.PhaseG1       = PhaseData.x * 2.0 - 1.0;
-	Result.PhaseG2       = PhaseData.y * 2.0 - 1.0;
-	Result.PhaseBlend    = PhaseData.z;
-	return Result;
-}
-
-uint4 EncodeCloudResult(FSampleCloudMaterialResult Result)
-{
-	return uint4(
-		PayloadEncodeHDRColor(Result.Extinction * 64.0), // scale up so we don't loose any density when clipping bits and make better use of the higher range of the encoding
-		PayloadEncodeLDRColor(Result.Albedo),
-		PayloadEncodeHDRColor(Result.EmissiveColor),
-		PayloadEncodeLDRColor(float3(Result.PhaseG1 * 0.5 + 0.5, Result.PhaseG2 * 0.5 + 0.5, Result.PhaseBlend))
-	);
-}
diff --git a/Engine/Shaders/Private/PathTracing/Volume/PathTracingCloudsMaterialCommon.ush b/Engine/Shaders/Private/PathTracing/Volume/PathTracingCloudsMaterialCommon.ush
--- a/Engine/Shaders/Private/PathTracing/Volume/PathTracingCloudsMaterialCommon.ush
+++ b/Engine/Shaders/Private/PathTracing/Volume/PathTracingCloudsMaterialCommon.ush
@@ -59,7 +59,7 @@ FSampleCloudMaterialResult SampleCloudMaterial(FDFVector3 AbsoluteWorldPosition,
 #endif
 
 	Result.Extinction = clamp(Result.Extinction * CENTIMETER_TO_METER, 0.0f, 65000.0f);
-	Result.EmissiveColor = clamp(Result.EmissiveColor, 0.0f, 6500.0f);
+	Result.EmissiveColor = clamp(Result.EmissiveColor * CENTIMETER_TO_METER, 0.0f, 6500.0f);
 	Result.Albedo = saturate(Result.Albedo);
 
 #if MATERIAL_VOLUMETRIC_ADVANCED
diff --git a/Engine/Shaders/Private/PathTracing/Volume/PathTracingVolume.ush b/Engine/Shaders/Private/PathTracing/Volume/PathTracingVolume.ush
--- a/Engine/Shaders/Private/PathTracing/Volume/PathTracingVolume.ush
+++ b/Engine/Shaders/Private/PathTracing/Volume/PathTracingVolume.ush
@@ -117,8 +117,7 @@ FVolumeShadedResult VolumeGetDensity(float3 TranslatedWorldPos, uint VolumeMask,
 			float3 SigmaN = SigmaBar - SigmaT;														\
 			/* keep tracing through the volume */													\
 			Throughput *= SigmaN;																	\
-			/* early termination - fades off from 0.002 to 0.001 and then will be 0 */				\
-			Throughput *= saturate(1000.0 * max3(Throughput.x, Throughput.y, Throughput.z) - 1.0);	\
+			Throughput *= LowThroughputClampFactor(Throughput); 									\
 			if (!any(Throughput > 0))																\
 			{																						\
 				break;																				\
diff --git a/Engine/Shaders/Private/PathTracing/Volume/PathTracingVolumeCommon.ush b/Engine/Shaders/Private/PathTracing/Volume/PathTracingVolumeCommon.ush
--- a/Engine/Shaders/Private/PathTracing/Volume/PathTracingVolumeCommon.ush
+++ b/Engine/Shaders/Private/PathTracing/Volume/PathTracingVolumeCommon.ush
@@ -154,7 +154,7 @@ struct FVolumeIntersectionList
 			}
 		}
 		// Transfer CloudDensity to Interval if this interval includes clouds
-		Result.CloudDensity = Result.VolumeMask & PATH_TRACER_VOLUME_ENABLE_BIT ? CloudDensity : 0; 
+		Result.CloudDensity = Result.VolumeMask & PATH_TRACER_VOLUME_ENABLE_CLOUDS ? CloudDensity : 0; 
 		return Result;
 	}
 
diff --git a/Engine/Shaders/Private/PathTracing/Volume/PathTracingVolumeSampling.ush b/Engine/Shaders/Private/PathTracing/Volume/PathTracingVolumeSampling.ush
--- a/Engine/Shaders/Private/PathTracing/Volume/PathTracingVolumeSampling.ush
+++ b/Engine/Shaders/Private/PathTracing/Volume/PathTracingVolumeSampling.ush
@@ -54,3 +54,10 @@ float4 EvaluateSpectralTransmittanceMiss(float MaxT, float3 Sigma, float3 ColorC
 	float MisPdf = dot(ColorChannelPdf, TransmittancePdf);
 	return MisPdf > 0 ? float4(Transmittance / MisPdf, MisPdf) : 0.0;
 }
+
+// Factor to multiply the path throughput by so that paths that have low overall contribution can get terminated earlier
+float LowThroughputClampFactor(float3 Throughput)
+{
+	// early termination - fades off from 0.002 to 0.001 and then will be 0
+	return saturate(1000.0 * max3(Throughput.x, Throughput.y, Throughput.z) - 1.0);
+}
\ No newline at end of file
diff --git a/Engine/Shaders/Shared/PathTracingDefinitions.h b/Engine/Shaders/Shared/PathTracingDefinitions.h
--- a/Engine/Shaders/Shared/PathTracingDefinitions.h
+++ b/Engine/Shaders/Shared/PathTracingDefinitions.h
@@ -67,12 +67,13 @@
 #define PATH_TRACER_VOLUME_HOLDOUT_HETEROGENEOUS_VOLUMES	(PATH_TRACER_VOLUME_HOLDOUT_BIT << VOLUMEID_HETEROGENEOUS_VOLUMES)
 #define PATH_TRACER_VOLUME_HOLDOUT_MASK                     (PATH_TRACER_VOLUME_ENABLE_MASK << PATH_TRACER_MAX_VOLUMES)
 
-#define PATH_TRACER_VOLUME_USE_ANALYTIC_TRANSMITTANCE       (PATH_TRACER_VOLUME_HOLDOUT_BIT << PATH_TRACER_MAX_VOLUMES)
+#define PATH_TRACER_VOLUME_USE_ANALYTIC_TRANSMITTANCE       (256u)
 
 // These flags are reserved for passing information to the cloud callable shader
-#define PATH_TRACER_VOLUME_CALLABLE_FLAGS_TRANSMITTANCE     (PATH_TRACER_VOLUME_USE_ANALYTIC_TRANSMITTANCE << 1) // only compute transmittance for the cloud portion
-#define PATH_TRACER_VOLUME_CALLABLE_FLAGS_GET_SAMPLE        (PATH_TRACER_VOLUME_USE_ANALYTIC_TRANSMITTANCE << 2) // perform RIS sampling for a scatter location (else do transmittance/emission/alpha calc only)
-#define PATH_TRACER_VOLUME_CALLABLE_FLAGS_IS_CAMERA_RAY     (PATH_TRACER_VOLUME_USE_ANALYTIC_TRANSMITTANCE << 3) // this is a camera ray (so apply holdouts if enabled)
+#define PATH_TRACER_VOLUME_CALLABLE_FLAGS_TRANSMITTANCE     (512u) // only compute transmittance for the cloud portion
+#define PATH_TRACER_VOLUME_CALLABLE_FLAGS_GET_SAMPLE        (1024u) // perform RIS sampling for a scatter location (else do transmittance/emission/alpha calc only)
+#define PATH_TRACER_VOLUME_CALLABLE_FLAGS_BOUNCE_MASK       (2048u * 31) // What bounce are we rendering? (5 bits, so bounces>31 are clamped)
+#define PATH_TRACER_VOLUME_CALLABLE_FLAGS_BOUNCE_SHIFT      (11)
 
 
 // Constants related to debugging
diff --git a/Engine/Source/Runtime/Renderer/Private/DeferredShadingRenderer.cpp b/Engine/Source/Runtime/Renderer/Private/DeferredShadingRenderer.cpp
--- a/Engine/Source/Runtime/Renderer/Private/DeferredShadingRenderer.cpp
+++ b/Engine/Source/Runtime/Renderer/Private/DeferredShadingRenderer.cpp
@@ -1757,9 +1757,8 @@ void FDeferredShadingSceneRenderer::Render(FRDGBuilder& GraphBuilder, const FSce
 				// Calculate decal grid for ray tracing per view since decal fade is view dependent
 				// TODO: investigate reusing the same grid for all views (ie: different callable shader SBT entries for each view so fade alpha is still correct for each view)
 
-				for (int32 ViewIndex = 0; ViewIndex < Views.Num(); ViewIndex++)
+				for (FViewInfo& View : Views)
 				{
-					FViewInfo& View = Views[ViewIndex];
 					View.RayTracingDecalUniformBuffer = CreateRayTracingDecalData(GraphBuilder, *Scene, View, RayTracingSBT.NumCallableShaderSlots);
 					View.bHasRayTracingDecals = true;
 					RayTracingSBT.NumCallableShaderSlots += Scene->Decals.Num();
@@ -1769,9 +1768,8 @@ void FDeferredShadingSceneRenderer::Render(FRDGBuilder& GraphBuilder, const FSce
 			{
 				TRDGUniformBufferRef<FRayTracingDecals> NullRayTracingDecalUniformBuffer = CreateNullRayTracingDecalsUniformBuffer(GraphBuilder);
 
-				for (int32 ViewIndex = 0; ViewIndex < Views.Num(); ViewIndex++)
+				for (FViewInfo& View : Views)
 				{
-					FViewInfo& View = Views[ViewIndex];
 					View.RayTracingDecalUniformBuffer = NullRayTracingDecalUniformBuffer;
 					View.bHasRayTracingDecals = false;
 				}
@@ -1779,7 +1777,7 @@ void FDeferredShadingSceneRenderer::Render(FRDGBuilder& GraphBuilder, const FSce
 
 			// If we might be path tracing the clouds -- call the path tracer's method for cloud callable shader setup
 			// this will skip work if cloud rendering is not being used
-			PreparePathTracingCloudMaterial(Scene, Views);
+			PreparePathTracingCloudMaterial(GraphBuilder, Scene, Views);
 		}
 
 		if (IsRayTracingEnabled(ViewFamily.GetShaderPlatform()) && GRHISupportsRayTracingShaders)
diff --git a/Engine/Source/Runtime/Renderer/Private/PathTracing.cpp b/Engine/Source/Runtime/Renderer/Private/PathTracing.cpp
--- a/Engine/Source/Runtime/Renderer/Private/PathTracing.cpp
+++ b/Engine/Source/Runtime/Renderer/Private/PathTracing.cpp
@@ -763,6 +763,17 @@ struct FPathTracingState
 	uint32_t FrameIndex = 0;
 };
 
+FPathTracingState* GetPathTracingStateFromView(const FViewInfo& View)
+{
+	if (!View.ViewState->PathTracingState.IsValid())
+	{
+		View.ViewState->PathTracingState = MakePimpl<FPathTracingState>();
+	}
+	check(View.ViewState->PathTracingState.IsValid());
+	return View.ViewState->PathTracingState.Get();
+}
+
+
 int GetPathTracingVisualizationMode()
 {
 	return GPathTracingVisualizeMode.IsNone() ? -1 : FPathTracingDebugModeInfo::Get().Map.FindRef(GPathTracingVisualizeMode, -1);
@@ -1146,6 +1157,8 @@ BEGIN_SHADER_PARAMETER_STRUCT(FPathTracingCloudParameters, )
 	SHADER_PARAMETER(float, CloudClipRadiusKm) // distance from origin to planet center in Km
 
 	SHADER_PARAMETER(float, CloudTracingMaxDistance) // limit ray lengths (to avoid slowing down when a ray crosses all clouds)
+	SHADER_PARAMETER(float, CloudVoxelWidth)
+	SHADER_PARAMETER(float, CloudInvVoxelWidth)
 	SHADER_PARAMETER(int32, CloudAccelMapResolution)
 	SHADER_PARAMETER(int32, CloudCallableShaderId)
 END_SHADER_PARAMETER_STRUCT()
@@ -1154,6 +1167,8 @@ BEGIN_GLOBAL_SHADER_PARAMETER_STRUCT(FPathTracingCloudParameterGlobals,)
 	SHADER_PARAMETER_STRUCT_INCLUDE(FPathTracingCloudParameters, CloudParameters)
 	SHADER_PARAMETER_STRUCT_INCLUDE(FPathTracingFogParameters, FogParameters)
 	SHADER_PARAMETER_STRUCT_INCLUDE(FAtmosphereUniformShaderParameters, AtmosphereParameters)
+	SHADER_PARAMETER_RDG_TEXTURE(Texture2D, CloudAccelerationMap)
+	SHADER_PARAMETER_SAMPLER(SamplerState, CloudAccelerationMapSampler)
 	SHADER_PARAMETER(FVector3f, PlanetCenterTranslatedWorldHi)
 	SHADER_PARAMETER(FVector3f, PlanetCenterTranslatedWorldLo)
 	SHADER_PARAMETER(uint32, MaxRaymarchSteps)
@@ -1483,7 +1498,7 @@ class FPathTracingBuildAtmosphereOpticalDepthLUTCS : public FGlobalShader
 IMPLEMENT_SHADER_TYPE(, FPathTracingBuildAtmosphereOpticalDepthLUTCS, TEXT("/Engine/Private/PathTracing/PathTracingBuildAtmosphereLUT.usf"), TEXT("PathTracingBuildAtmosphereOpticalDepthLUTCS"), SF_Compute);
 
 
-FPathTracingCloudParameters PrepareCloudParameters(const FScene* Scene, const FViewInfo& View, const FPathTracingConfig* Config)
+FPathTracingCloudParameters PrepareCloudParameters(const FScene* Scene, const FViewInfo& View, const int CloudAccelerationMapResolution)
 {
 	check(Scene != nullptr);;
 	check(Scene->GetVolumetricCloudSceneInfo() != nullptr);
@@ -1543,10 +1558,9 @@ FPathTracingCloudParameters PrepareCloudParameters(const FScene* Scene, const FV
 	
 	Params.CloudClipRadiusKm = float(ViewToPlanet * double(FAtmosphereSetup::CmToSkyUnit));
 	Params.CloudCallableShaderId = -1;
-	if (Config != nullptr)
-	{
-		Params.CloudAccelMapResolution = Config->CloudAccelerationMapResolution;
-	}
+	Params.CloudAccelMapResolution = CloudAccelerationMapResolution;
+	Params.CloudVoxelWidth = float(2.0 * Params.CloudClipDistKm * FAtmosphereSetup::SkyUnitToCm / CloudAccelerationMapResolution);
+	Params.CloudInvVoxelWidth = 1.0f / Params.CloudVoxelWidth;
 	return Params;
 }
 
@@ -1672,7 +1686,7 @@ private:
 };
 IMPLEMENT_SHADER_TYPE(, FPathTracingVolumetricCloudMaterial, TEXT("/Engine/Private/PathTracing/PathTracingVolumetricCloudMaterialShader.usf"), TEXT("PathTracingVolumetricCloudMaterialShader"), SF_RayCallable);
 
-void PreparePathTracingCloudMaterial(FScene* Scene, TArrayView<FViewInfo> Views)
+void PreparePathTracingCloudMaterial(FRDGBuilder& GraphBuilder, FScene* Scene, TArrayView<FViewInfo> Views)
 {
 	// make sure all views have an invalid callable shader index (unless proven otherwise below)
 	for (FViewInfo& View : Views)
@@ -1733,8 +1747,10 @@ void PreparePathTracingCloudMaterial(FScene* Scene, TArrayView<FViewInfo> Views)
 
 		View.PathTracingVolumetricCloudCallableShaderIndex = BaseCallableSlotIndex;
 
+		const int CloudAccelMapResolution = CVarPathTracingCloudAccelerationMapResolution.GetValueOnRenderThread();
+
 		FPathTracingCloudParameterGlobals Params = {};
-		Params.CloudParameters = PrepareCloudParameters(Scene, View, nullptr);
+		Params.CloudParameters = PrepareCloudParameters(Scene, View, CloudAccelMapResolution);
 		if (!Scene->ExponentialFogs.IsEmpty())
 		{
 			Params.FogParameters = PrepareFogParameters(View, Scene->ExponentialFogs[0]);
@@ -1744,13 +1760,31 @@ void PreparePathTracingCloudMaterial(FScene* Scene, TArrayView<FViewInfo> Views)
 			Params.AtmosphereParameters = *Scene->GetSkyAtmosphereSceneInfo()->GetAtmosphereShaderParameters();
 		}
 
+		const int Resolution = Params.CloudParameters.CloudAccelMapResolution;
+		FPathTracingState* PathTracingState = GetPathTracingStateFromView(View);
+		if (PathTracingState->CloudAccelerationMap && PathTracingState->CloudAccelerationMap->GetDesc().Extent.X == Resolution)
+		{
+			// we already have a map from a previous iteration, re-use it
+		}
+		else
+		{
+			// Either we don't have a map yet, or the resolution cvar changed
+			PathTracingState->CloudAccelerationMap.SafeRelease();
+			EPixelFormat CloudAccelerationMapFormat = PF_FloatRGBA; // 16 bit should be good enough for typical density/z ranges
+			const FRHITextureCreateDesc Desc =
+				FRHITextureCreateDesc::Create2D(TEXT("PathTracer.CloudAccelerationMap"), Resolution, Resolution, CloudAccelerationMapFormat)
+			.SetFlags(ETextureCreateFlags::ShaderResource | ETextureCreateFlags::UAV);
+
+			PathTracingState->CloudAccelerationMap = CreateRenderTarget(RHICreateTexture(Desc), Desc.DebugName);
+		}
+
+		Params.CloudAccelerationMap = GraphBuilder.RegisterExternalTexture(PathTracingState->CloudAccelerationMap, TEXT("PathTracer.CloudAccelerationMap"));
+		Params.CloudAccelerationMapSampler = TStaticSamplerState<SF_Point, AM_Clamp, AM_Clamp, AM_Clamp>::GetRHI();
+
 		PreparePlanetCenter(View, Scene->GetSkyAtmosphereSceneInfo(), &Params.PlanetCenterTranslatedWorldHi, &Params.PlanetCenterTranslatedWorldLo);
 		Params.MaxRaymarchSteps = CVarPathTracingMaxRaymarchSteps.GetValueOnRenderThread();
 
-		TUniformBufferRef<FPathTracingCloudParameterGlobals> CloudParametersUB = CreateUniformBufferImmediate(
-			Params,
-			EUniformBufferUsage::UniformBuffer_SingleFrame
-		);
+		TUniformBufferRef<FPathTracingCloudParameterGlobals> CloudParametersUB = CreateUniformBufferImmediate(Params, EUniformBufferUsage::UniformBuffer_SingleFrame);
 		Scene->RayTracingSBT.TransientUniformBuffers.Add(CloudParametersUB); // Hold uniform buffer ref in RayTracingSBT since FMeshDrawSingleShaderBindings doesn't
 
 		FMeshDrawSingleShaderBindings SingleShaderBindings = Command.ShaderBindings.GetSingleShaderBindings(SF_RayCallable);
@@ -3139,7 +3173,6 @@ void FSceneViewState::PathTracingInvalidate(bool InvalidateAnimationStates)
 		State->NormalRT.SafeRelease();
 		State->DepthRT.SafeRelease();
 		State->VarianceBuffer.SafeRelease();
-		State->CloudAccelerationMap.SafeRelease();
 		State->SampleIndex = 0;
 
 		State->AdaptiveFrustumGridParameterCache.TopLevelGridBuffer.SafeRelease();
@@ -3347,16 +3380,10 @@ void FDeferredShadingSceneRenderer::RenderPathTracing(
 
 	Config.PathTracingData.MaxSamples = MaxSPP;
 
-	bool FirstTime = false;
-	if (!View.ViewState->PathTracingState.IsValid())
-	{
-		View.ViewState->PathTracingState = MakePimpl<FPathTracingState>();
-		FirstTime = true; // we just initialized the option state for this view -- don't bother comparing in this case
-	}
-	check(View.ViewState->PathTracingState.IsValid());
-	FPathTracingState* PathTracingState = View.ViewState->PathTracingState.Get();
+	FPathTracingState* PathTracingState = GetPathTracingStateFromView(View);
+	const bool bFirstTime = !PathTracingState->RadianceRT.IsValid(); // we just initialized (or reset) the option state for this view -- don't bother comparing in this case
 
-	if (FirstTime || Config.UseMISCompensation != PathTracingState->LastConfig.UseMISCompensation)
+	if (!bFirstTime && Config.UseMISCompensation != PathTracingState->LastConfig.UseMISCompensation)
 	{
 		// if the mode changes we need to rebuild the importance table
 		Scene->PathTracingSkylightTexture.SafeRelease();
@@ -3389,7 +3416,6 @@ void FDeferredShadingSceneRenderer::RenderPathTracing(
 		// camera medium tracking is not enabled - just make a temp buffer and set it to 0
 		StartingExtinctionCoefficient = GraphBuilder.CreateBuffer(FRDGBufferDesc::CreateBufferDesc(sizeof(float), 3), TEXT("PathTracer.StartingExtinctionCoefficient"));
 		AddClearUAVPass(GraphBuilder, GraphBuilder.CreateUAV(StartingExtinctionCoefficient, PF_R32_FLOAT), 0);
-		
 	}
 	else if (!PathTracingState->StartingExtinctionCoefficient.IsValid() || PathTracingState->SampleIndex == 0 || View.bIsOfflineRender)
 	{
@@ -3494,7 +3520,7 @@ void FDeferredShadingSceneRenderer::RenderPathTracing(
 	// we must invalidate the ViewState to start over from scratch
 	// NOTE: only check things like hair position changes for interactive viewports, for offline renders we don't want any chance of mid-render invalidation
 	// NOTE: same for DOF changes, these parameters could be animated which should not automatically invalidate a render in progress
-	if (FirstTime ||
+	if (bFirstTime ||
 		Config.IsDifferent(PathTracingState->LastConfig) ||
 		(!View.bIsOfflineRender && Config.IsExposureDifferentEnough(PathTracingState->LastConfig)) ||
 		(!View.bIsOfflineRender && Config.IsDOFDifferent(PathTracingState->LastConfig)) ||
@@ -3639,33 +3665,21 @@ void FDeferredShadingSceneRenderer::RenderPathTracing(
 				);
 			}
 
-			FRDGTexture* CloudAccelerationMap = nullptr;
+			FRDGTexture* CloudAccelerationMap = PathTracingState->CloudAccelerationMap.IsValid() ? GraphBuilder.RegisterExternalTexture(PathTracingState->CloudAccelerationMap, TEXT("PathTracer.CloudAccelerationMap")) : nullptr;
 
 			const bool bEnableClouds   = (Config.PathTracingData.VolumeFlags & PATH_TRACER_VOLUME_ENABLE_CLOUDS) != 0;
 
 			if (bEnableClouds)
 			{
-				// clouds are enabled, build an accel map (do this every frame as clouds are usually animating)
+				// clouds are enabled, build an accel map (do this every frame as clouds are usually animating, and so that the bounds improve during sampling)
 				RDG_GPU_MASK_SCOPE(GraphBuilder, FRHIGPUMask::All());
 
 				const int32 NumSamples = Config.CloudAccelerationMapNumSamples;
 				const int32 Resolution = Config.CloudAccelerationMapResolution;
-				if (PathTracingState->CloudAccelerationMap)
-				{
-					// we already have a map from a previous iteration, re-use it
-					CloudAccelerationMap = GraphBuilder.RegisterExternalTexture(PathTracingState->CloudAccelerationMap, TEXT("PathTracer.CloudAccelerationMap"));
-				}
-				else
-				{
-					// need to create a new LUT
-					EPixelFormat CloudAccelerationMapFormat = PF_FloatRGBA; // 16 bit should be good enough for typical density/z ranges
-					FRDGTextureDesc Desc = FRDGTextureDesc::Create2D(
-						FIntPoint(Resolution, Resolution),
-						CloudAccelerationMapFormat,
-						FClearValueBinding::None,
-						TexCreate_ShaderResource | TexCreate_UAV);
-					CloudAccelerationMap = GraphBuilder.CreateTexture(Desc, TEXT("PathTracer.CloudAccelerationMap"), ERDGTextureFlags::None);
-				}
+
+				// If we got here, the acceleration map texture should have already been created
+				check(CloudAccelerationMap != nullptr);
+
 				FVolumetricCloudRenderSceneInfo* CloudRenderSceneInfo = Scene->GetVolumetricCloudSceneInfo();
 				check(CloudRenderSceneInfo != nullptr);
 				UMaterialInterface* CloudMaterialInterface = CloudRenderSceneInfo->GetVolumetricCloudSceneProxy().GetCloudVolumeMaterial();
@@ -3675,7 +3689,7 @@ void FDeferredShadingSceneRenderer::RenderPathTracing(
 				const FMaterial* MaterialResource = &CloudVolumeMaterialProxy->GetMaterialWithFallback(Scene->GetFeatureLevel(), CloudVolumeMaterialProxy);
 				
 
-				const FPathTracingCloudParameters CloudParameters = PrepareCloudParameters(Scene, View, &Config);
+				const FPathTracingCloudParameters CloudParameters = PrepareCloudParameters(Scene, View, Config.CloudAccelerationMapResolution);
 				// build cloud accel map
 				{
 					typename FPathTracingBuildCloudAccelerationMapCS::FPermutationDomain PermutationVector;
@@ -3871,9 +3885,10 @@ void FDeferredShadingSceneRenderer::RenderPathTracing(
 							if (PreviousPassParameters != nullptr)
 							{
 								PassParameters->CloudParameters = PreviousPassParameters->CloudParameters;
-							} else if ((Config.PathTracingData.VolumeFlags & PATH_TRACER_VOLUME_ENABLE_CLOUDS) != 0)
+							}
+							else if ((Config.PathTracingData.VolumeFlags & PATH_TRACER_VOLUME_ENABLE_CLOUDS) != 0)
 							{
-								PassParameters->CloudParameters = PrepareCloudParameters(Scene, View, &Config);
+								PassParameters->CloudParameters = PrepareCloudParameters(Scene, View, Config.CloudAccelerationMapResolution);
 							}
 							else
 							{
diff --git a/Engine/Source/Runtime/Renderer/Private/PathTracing.h b/Engine/Source/Runtime/Renderer/Private/PathTracing.h
--- a/Engine/Source/Runtime/Renderer/Private/PathTracing.h
+++ b/Engine/Source/Runtime/Renderer/Private/PathTracing.h
@@ -31,7 +31,7 @@ RENDERER_API FRHIRayTracingShader* GetGPULightmassDefaultOpaqueHitShader(const F
 RENDERER_API FRHIRayTracingShader* GetGPULightmassDefaultHiddenHitShader(const FGlobalShaderMap* ShaderMap);
 
 void PreparePathTracingRTPSO();
-void PreparePathTracingCloudMaterial(FScene* Scene, TArrayView<FViewInfo> Views);
+void PreparePathTracingCloudMaterial(FRDGBuilder& GraphBuilder, FScene* Scene, TArrayView<FViewInfo> Views);
 
 namespace PathTracing
 {
