Force HLSL2021 in FVirtualShadowMapPageManagementShader to avoid compilation errors when using platform preview in editor.

diff --git a/Engine/Source/Runtime/Renderer/Private/VirtualShadowMaps/VirtualShadowMapShaders.h b/Engine/Source/Runtime/Renderer/Private/VirtualShadowMaps/VirtualShadowMapShaders.h
--- a/Engine/Source/Runtime/Renderer/Private/VirtualShadowMaps/VirtualShadowMapShaders.h
+++ b/Engine/Source/Runtime/Renderer/Private/VirtualShadowMaps/VirtualShadowMapShaders.h
@@ -45,6 +45,8 @@ public:
 
 		FForwardLightingParameters::ModifyCompilationEnvironment(Parameters.Platform, OutEnvironment);
 		OutEnvironment.SetDefine(TEXT("VF_SUPPORTS_PRIMITIVE_SCENE_DATA"), 1);
+
+		OutEnvironment.CompilerFlags.Add(CFLAG_HLSL2021);
 	}
 };
 
