HairStrands Voxelize - Fix missing AtomicCompatible flag on Hair.VoxelPageTexture causing corruption on Mac

#tests Local testing on Mac + SM6 in QAGame TM-HairAndFur

diff --git a/Engine/Source/Runtime/Renderer/Private/HairStrands/HairStrandsVoxelization.cpp b/Engine/Source/Runtime/Renderer/Private/HairStrands/HairStrandsVoxelization.cpp
--- a/Engine/Source/Runtime/Renderer/Private/HairStrands/HairStrandsVoxelization.cpp
+++ b/Engine/Source/Runtime/Renderer/Private/HairStrands/HairStrandsVoxelization.cpp
@@ -924,7 +924,7 @@ static FHairStrandsVoxelResources AllocateVirtualVoxelResources(
 			Out.Parameters.Common.PageTextureResolution,
 			PF_R32_UINT, 
 			FClearValueBinding::Black, 
-			TexCreate_UAV | TexCreate_ShaderResource | TexCreate_NoFastClear | TexCreate_3DTiling,
+			TexCreate_UAV | TexCreate_ShaderResource | TexCreate_NoFastClear | TexCreate_3DTiling | TexCreate_AtomicCompatible,
 			MipCount);
 		Out.PageTexture = GraphBuilder.CreateTexture(Desc, TEXT("Hair.VoxelPageTexture"));
 	}
