[VSM] Added null check as the new directional shadow setup always sets up slots for all directional lights, but may not populate them.

#rnx

diff --git a/Engine/Source/Runtime/Renderer/Private/Shadows/ShadowSceneRenderer.cpp b/Engine/Source/Runtime/Renderer/Private/Shadows/ShadowSceneRenderer.cpp
--- a/Engine/Source/Runtime/Renderer/Private/Shadows/ShadowSceneRenderer.cpp
+++ b/Engine/Source/Runtime/Renderer/Private/Shadows/ShadowSceneRenderer.cpp
@@ -613,7 +613,7 @@ void FShadowSceneRenderer::PostInitDynamicShadowsSetup()
 		for (const FDirectionalLightShadowFrameSetup& DirectionalLightShadowFrameSetup : DirectionalLights)
 		{
 			FProjectedShadowInfo* ProjectedShadowInfo = DirectionalLightShadowFrameSetup.ClipmapInfos[EVirtualShadowTypeId::Regular].ProjectedShadowInfo;
-			if (!bUnboundedClipmap && ProjectedShadowInfo->bShouldRenderVSM)
+			if (!bUnboundedClipmap && ProjectedShadowInfo && ProjectedShadowInfo->bShouldRenderVSM)
 			{
 				const bool bIsCached = !ProjectedShadowInfo->VirtualShadowMapPerLightCacheEntry->IsUncached();
 
