[TSR] Add a weight relaxation pass to adjust weight.

* Currently, for translucency overlay, we reduce the weight based on the translcuency contribution. Fully disable the relaxation when translucency contribution is larger than 60%. This helps make foliage and hair converge faster without lag due to translucency change. It has a cost around 0.05ms on consoles but can be hidden with async compute.

#tests Test on silver project
#[review] aleksander.netzel
#[fyi] Brian.Karis

diff --git a/Engine/Shaders/Private/TemporalSuperResolution/TSRThinGeometryCommon.ush b/Engine/Shaders/Private/TemporalSuperResolution/TSRThinGeometryCommon.ush
new file mode 100644
--- /dev/null
+++ b/Engine/Shaders/Private/TemporalSuperResolution/TSRThinGeometryCommon.ush
@@ -0,0 +1,32 @@
+// Copyright Epic Games, Inc. All Rights Reserved.
+
+#pragma once
+
+//------------------------------------------------------- INCLUDES
+#include "TSRColorSpace.ush"
+#include "TSRConvolutionNetworkPass.ush"
+
+
+//-------------------------------------------------------UTIL FUNCTION
+CALL_SITE_DEBUGLOC
+void DecodeThinGeometry(
+	tsr_tensor_ushort EncodedThinGeometryWeight,
+	out tsr_tensor_half HistoryRelaxationWeight,
+	out tsr_tensor_bool IsClusterHoleRegion)
+{
+	// use sign to encode?
+	HistoryRelaxationWeight = tsr_tensor_half::CastFrom<tsr_ushort>(
+		bit_shift_right(EncodedThinGeometryWeight, tsr_tensor_ushort::Const(tsr_ushort(0x1u)))) * rcp(tsr_half(127.0));
+	IsClusterHoleRegion = tsr_tensor_bool::CastFrom<tsr_ushort>(bit_and(EncodedThinGeometryWeight, tsr_tensor_ushort::Const(tsr_ushort(0x1u))));
+}
+
+CALL_SITE_DEBUGLOC
+tsr_tensor_ushort EncodeThinGeometry(
+	tsr_tensor_half HistoryRelaxationWeight,
+	tsr_tensor_bool IsClusterHoleRegion)
+{
+	tsr_tensor_ushort EncodedThinGeometry = bit_shift_left(
+		tsr_tensor_ushort::CastFrom<tsr_half>(HistoryRelaxationWeight * tsr_half(127.0)), tsr_tensor_ushort::Const(tsr_ushort(0x1u)));
+	EncodedThinGeometry = bit_or(EncodedThinGeometry,tsr_tensor_ushort::CastFrom<bool>(IsClusterHoleRegion));
+	return EncodedThinGeometry;
+}
diff --git a/Engine/Shaders/Private/TemporalSuperResolution/TSRWeightRelaxation.usf b/Engine/Shaders/Private/TemporalSuperResolution/TSRWeightRelaxation.usf
new file mode 100644
--- /dev/null
+++ b/Engine/Shaders/Private/TemporalSuperResolution/TSRWeightRelaxation.usf
@@ -0,0 +1,233 @@
+// Copyright Epic Games, Inc. All Rights Reserved.
+
+
+//------------------------------------------------------- INCLUDES
+#include "TSRColorSpace.ush"
+//#include "TSRDepthVelocityAnalysis.ush"
+#include "TSRThinGeometryCommon.ush"
+
+//------------------------------------------------------- CONFIG
+
+#define DEBUG_ARRAY_SIZE 8
+
+#define DEBUG_DISABLE_PADDING 0
+
+#define DEBUG_DIGESTABLE 0
+
+//------------------------------------------------------- CONSTANTS
+
+#if CONFIG_LDS_DESPILL
+static const uint kDespillInput = LDS_DESPILL_THREAD_COUNT * LDS_DESPILL_OFFSET * 0;
+static const uint kDespillHistory = LDS_DESPILL_THREAD_COUNT * LDS_DESPILL_OFFSET * 1;
+#endif
+
+//------------------------------------------------------- PARAMETERS
+
+int TileOverscan;
+int ThinGeometryTextureIndex;
+float ErrorMultiplier;
+float MaxRelaxationWeight;
+
+Texture2D<tsr_halfC> InputTexture;
+Texture2D<tsr_half4> InputSceneTranslucencyTexture;
+
+RWTexture2DArray<uint> R8Output;
+
+//------------------------------------------------------- GLOBALS
+static uint2 GGroupId = 0;
+
+//------------------------------------------------------- DEBUG
+
+static tsr_tensor_half4 Debug[DEBUG_ARRAY_SIZE];
+
+//------------------------------------------------------- INPUT FETCHES
+
+tsr_short2 ComputeElementPixelPos(const uint ElementIndex)
+{
+	tsr_short2 ClampedFetchPixelPos = ComputeElementInputPixelPos(
+		/* LaneStride = */ uint2(LANE_STRIDE_X, LANE_STRIDE_Y),
+		TileOverscan,
+		InputInfo_ViewportMin,
+		InputInfo_ViewportMax,
+		GGroupId,
+		GGroupThreadIndex,
+		ElementIndex);
+
+	return ClampedFetchPixelPos;
+}
+
+tsr_tensor_short2 ComputePixelPos()
+{
+	tsr_tensor_short2 PixelPos;
+
+	UNROLL_N(SIMD_SIZE)
+		for (uint ElementIndex = 0; ElementIndex < SIMD_SIZE; ElementIndex++)
+		{
+			PixelPos.SetElement(ElementIndex, ComputeElementPixelPos(ElementIndex));
+		}
+
+	return PixelPos;
+}
+
+tsr_short2 ComputeElementOutputPixelPos(const uint ElementIndex)
+{
+	return ComputeElementOutputPixelPos(
+		/* LaneStride = */ uint2(LANE_STRIDE_X, LANE_STRIDE_Y),
+		TileOverscan,
+		InputInfo_ViewportMin,
+		InputInfo_ViewportMax,
+		GGroupId,
+		GGroupThreadIndex,
+		ElementIndex);
+}
+
+CALL_SITE_DEBUGLOC
+tsr_half ComputeRGBLuma(tsr_half3 Input)
+{
+	// aligns to the logic in MeasureFlickeringLuma.usf so taht we can compare the 
+	// difference between LumaBeforeComposition that comes from that pass.
+	tsr_half3 LinearColor = Input;
+	tsr_half3 SMColor = LinearToSMCS(LinearColor);
+	tsr_half  SMLuma = dot(SMColor, kMoireSMCSWeights);
+	tsr_half  GLuma = SMCSToGCS(SMLuma.rrr).r;
+
+	return GLuma;
+}
+
+CALL_SITE_DEBUGLOC
+tsr_tensor_half ComputeRGBLuma(tsr_tensor_half3 Input)
+{
+	tsr_tensor_half Luma;
+	UNROLL_N(SIMD_SIZE)
+	for (uint ElementIndex = 0; ElementIndex < SIMD_SIZE; ElementIndex++)
+	{
+		// aligns to the logic in MeasureFlickeringLuma.usf so taht we can compare the 
+		// difference between LumaBeforeComposition that comes from that pass.
+		tsr_half3 LinearColor = Input.GetElement(ElementIndex);
+		Luma.SetElement(ElementIndex, ComputeRGBLuma(LinearColor));
+	}
+
+	return Luma;
+}
+
+void FetchInput(
+	tsr_tensor_short2 PixelPos,
+	out tsr_tensor_half InputLuma,
+	out tsr_tensor_half2 Translucency,
+	out tsr_tensor_half HistoryRelaxationWeight,
+	out tsr_tensor_bool IsClusterHoleRegion)
+{
+	tsr_tensor_ushort ThinGeometryCache;
+	
+	UNROLL_N(SIMD_SIZE)
+	for (uint ElementIndex = 0; ElementIndex < SIMD_SIZE; ElementIndex++)
+	{
+		tsr_short2 ClampedFetchPixelPos = PixelPos.GetElement(ElementIndex);
+		InputLuma.SetElement(ElementIndex,ComputeRGBLuma(InputTexture[ClampedFetchPixelPos].xyz));
+		tsr_half4 FullTranslucency = InputSceneTranslucencyTexture[ClampedFetchPixelPos];
+		Translucency.SetElement(ElementIndex, tsr_half2(ComputeRGBLuma(FullTranslucency.rgb),FullTranslucency.a));
+		ThinGeometryCache.SetElement(ElementIndex, tsr_ushort(R8Output[int3(ClampedFetchPixelPos, ThinGeometryTextureIndex)]));
+	}
+
+	DecodeThinGeometry(ThinGeometryCache, HistoryRelaxationWeight, IsClusterHoleRegion);	
+}
+
+//------------------------------------------------------- ENTRY POINT
+
+#if COMPILER_SUPPORTS_WAVE_SIZE && DIM_WAVE_SIZE > 0
+WAVESIZE(DIM_WAVE_SIZE)
+#endif
+[numthreads(LANE_COUNT * WAVE_COUNT, 1, 1)]
+void MainCS(
+	uint2 GroupId : SV_GroupID,
+	uint GroupThreadIndex : SV_GroupIndex)
+{
+	GGroupId = GroupId;
+	GGroupThreadIndex = GroupThreadIndex;
+
+#if DEBUG_OUTPUT
+	{
+		UNROLL_N(DEBUG_ARRAY_SIZE)
+			for (uint DebugId = 0; DebugId < DEBUG_ARRAY_SIZE; DebugId++)
+			{
+				Debug[DebugId].SetAllElements(0.0);
+			}
+	}
+#endif
+	tsr_tensor_short2 PixelPos = ComputePixelPos();
+	tsr_tensor_half InputLuma;
+	tsr_tensor_half2 Translucency;
+	tsr_tensor_half HistoryRelaxationWeight;
+	tsr_tensor_bool IsPartialCoverageRegion;
+
+	FetchInput(
+		PixelPos,
+		/*out */ InputLuma,
+		/*out */ Translucency,
+		/*out */ HistoryRelaxationWeight,
+		/*out */ IsPartialCoverageRegion);
+
+	// If translucency has data turn down the weights aggresively. if alpha < 0.95
+	tsr_tensor_half TranslucencyRatio = Translucency[0] * rcp(Translucency[0] + InputLuma * Translucency[1] + tsr_tensor_half::Const(0.0005));
+#if DEBUG_OUTPUT
+	Debug[0].SetComponent(0, TranslucencyRatio);
+	Debug[0].SetComponent(1, HistoryRelaxationWeight);
+#endif
+	// Weights decline starts around TranslucencyRatio = 0.05, and declines fast to 0 when Translucency ratio goes to 0.6. 
+	tsr_tensor_half RatioToLerpWeight = max(tsr_tensor_half::Const(0.6)-TranslucencyRatio,tsr_tensor_half::Const(0));
+	RatioToLerpWeight = min(RatioToLerpWeight * RatioToLerpWeight * tsr_tensor_half::Const(3.33), tsr_tensor_half::Const(1));
+	RatioToLerpWeight = select(or(IsPartialCoverageRegion,TranslucencyRatio > tsr_tensor_half::Const(0.6)), RatioToLerpWeight, tsr_tensor_half::Const(1));
+	HistoryRelaxationWeight = lerp(tsr_tensor_half::Const(0), HistoryRelaxationWeight, RatioToLerpWeight); 
+	
+#if DEBUG_OUTPUT
+	Debug[0].SetComponent(2, HistoryRelaxationWeight);
+#endif
+
+	tsr_tensor_ushort EncodedThinGeometry = EncodeThinGeometry(HistoryRelaxationWeight,IsPartialCoverageRegion);
+	// Output
+	UNROLL_N(SIMD_SIZE)
+		for(uint ElementIndex = 0; ElementIndex < SIMD_SIZE; ElementIndex++)
+		{
+			tsr_short2 OutputPixelPos = ComputeElementOutputPixelPos(ElementIndex);
+			R8Output[int3(OutputPixelPos, ThinGeometryTextureIndex)] = EncodedThinGeometry.GetElement(ElementIndex);
+		}
+
+#if DEBUG_OUTPUT
+	{
+		// Mark the top left corner of the tile for debugging purposes.
+#if DEBUG_DISABLE_PADDING == 2
+		if (GGroupThreadIndex == 0)
+		{
+			UNROLL_N(DEBUG_ARRAY_SIZE)
+				for (uint DebugId = 0; DebugId < DEBUG_ARRAY_SIZE; DebugId++)
+				{
+					Debug[DebugId].SetElement(/* ElementIndex = */ 0, 1.0);
+				}
+		}
+#endif // DEBUG_DISABLE_PADDING == 2
+
+		UNROLL_N(SIMD_SIZE)
+			for (uint ElementIndex = 0; ElementIndex < SIMD_SIZE; ElementIndex++)
+			{
+				const tsr_short2 OutputPixelPos = ComputeElementOutputPixelPos(
+					/* LaneStride = */ uint2(LANE_STRIDE_X, LANE_STRIDE_Y),
+					TileOverscan,
+					InputInfo_ViewportMin,
+					InputInfo_ViewportMax,
+					GGroupId,
+					GGroupThreadIndex,
+					ElementIndex,
+					/* OutputDataOffset = */ int2(0, 0),
+					/* OutputResDivisor = */ 1,
+					/* bDebugDisablePadding = */ DEBUG_DISABLE_PADDING != 0);
+
+				UNROLL_N(DEBUG_ARRAY_SIZE)
+					for (uint DebugId = 0; DebugId < DEBUG_ARRAY_SIZE; DebugId++)
+					{
+						DebugOutput[tsr_short3(OutputPixelPos, DebugId)] = float4(Debug[DebugId].GetElement(ElementIndex));
+					}
+			}
+	}
+#endif // DEBUG_OUTPUT
+
+}
diff --git a/Engine/Source/Runtime/Renderer/Private/PostProcess/TemporalSuperResolution.cpp b/Engine/Source/Runtime/Renderer/Private/PostProcess/TemporalSuperResolution.cpp
--- a/Engine/Source/Runtime/Renderer/Private/PostProcess/TemporalSuperResolution.cpp
+++ b/Engine/Source/Runtime/Renderer/Private/PostProcess/TemporalSuperResolution.cpp
@@ -328,6 +328,11 @@ TAutoConsoleVariable<float> CVarTSRThinGeometryCoverageMaxRelaxationWeight(
 	TEXT("The max history clamping box relaxation weight due to thin geometry detection (0 to 1)"),
 	ECVF_RenderThreadSafe);
 
+TAutoConsoleVariable<int32> CVarTSRThinGeometryDetectionWeightRelaxation(
+	TEXT("r.TSR.ThinGeometryDetection.WeightRelaxation"), 1,
+	TEXT("Adaptively trim the history relaxation to avoid ghosting."),
+	ECVF_RenderThreadSafe);
+
 TAutoConsoleVariable<float> CVarTSRThinGeometryDetectionErrorMultiplier(
 	TEXT("r.TSR.ThinGeometryDetection.ErrorMultiplier"), 200.0,
 	TEXT("Define the depth difference multiplier between neighbors will be used to detect thin geometry. The larger the upscaler, the smaller the multiplier should be for temporal stability. ")
@@ -1029,7 +1034,70 @@ class FTSRDetectThinGeometryCS : public FTSRConvolutionNetworkShader
 	{
 		return EShaderCompileJobPriority::High;
 	}
-}; // class FTSRRejectShadingCS
+}; // class FTSRDetectThinGeometryCS 
+
+class FTSRTSRWeightRelaxationCS :  public FTSRConvolutionNetworkShader
+{
+	DECLARE_GLOBAL_SHADER(FTSRTSRWeightRelaxationCS);
+	SHADER_USE_PARAMETER_STRUCT(FTSRTSRWeightRelaxationCS, FTSRConvolutionNetworkShader);
+
+	using FPermutationDomain = TShaderPermutationDomain<FTSRConvolutionNetworkShader::FPermutationDomain>;
+
+	BEGIN_SHADER_PARAMETER_STRUCT(FParameters,)
+		SHADER_PARAMETER_STRUCT_INCLUDE(FTSRCommonParameters, CommonParameters)
+		SHADER_PARAMETER(int32, TileOverscan)
+		SHADER_PARAMETER(int32, ThinGeometryTextureIndex)
+		SHADER_PARAMETER(float, MaxRelaxationWeight)
+
+		SHADER_PARAMETER_RDG_TEXTURE(Texture2D, InputTexture)                  // Scene with
+		SHADER_PARAMETER_RDG_TEXTURE(Texture2D, InputSceneTranslucencyTexture) // Translucency texture
+		SHADER_PARAMETER_RDG_TEXTURE_UAV(RWTexture2DArray, R8Output)
+		SHADER_PARAMETER_RDG_TEXTURE_UAV(RWTexture2D, DebugOutput)
+	END_SHADER_PARAMETER_STRUCT()
+
+	static FPermutationDomain RemapPermutation(FPermutationDomain PermutationVector)
+	{
+		// Remap redondant convolution permutations.
+		PermutationVector.Set<FTSRConvolutionNetworkShader::FPermutationDomain>(
+			FTSRConvolutionNetworkShader::RemapPermutation(PermutationVector.Get<FTSRConvolutionNetworkShader::FPermutationDomain>()));
+
+		return PermutationVector;
+	}
+
+	static bool ShouldCompilePermutation(const FGlobalShaderPermutationParameters& Parameters)
+	{
+		FPermutationDomain PermutationVector(Parameters.PermutationId);
+		if (PermutationVector != RemapPermutation(PermutationVector))
+		{
+			return false;
+		}
+
+		if (!FTSRConvolutionNetworkShader::ShouldCompilePermutation(Parameters, PermutationVector.Get<FTSRConvolutionNetworkShader::FPermutationDomain>()))
+		{
+			return false;
+		}
+
+		return true;
+	}
+
+	static EShaderPermutationPrecacheRequest ShouldPrecachePermutation(const FGlobalShaderPermutationParameters& Parameters)
+	{
+		FPermutationDomain PermutationVector(Parameters.PermutationId);
+		return FTSRConvolutionNetworkShader::ShouldPrecachePermutation(Parameters, PermutationVector.Get<FTSRConvolutionNetworkShader::FPermutationDomain>());
+	}
+
+	static void ModifyCompilationEnvironment(const FGlobalShaderPermutationParameters& Parameters, FShaderCompilerEnvironment& OutEnvironment)
+	{
+		FPermutationDomain PermutationVector(Parameters.PermutationId);
+
+		FTSRConvolutionNetworkShader::ModifyCompilationEnvironment(
+			Parameters,
+			PermutationVector.Get<FTSRConvolutionNetworkShader::FPermutationDomain>(),
+			OutEnvironment);
+		OutEnvironment.CompilerFlags.Add(CFLAG_Wave32);
+	}
+
+}; // class FTSRTSRWeightRelaxationCS
 
 class FTSRSpatialAntiAliasingCS : public FTSRShader
 {
@@ -1308,6 +1376,7 @@ IMPLEMENT_GLOBAL_SHADER(FTSRMeasureThinGeometryCoverageCS, "/Engine/Private/Temp
 IMPLEMENT_GLOBAL_SHADER(FTSRClearPrevTexturesCS,     "/Engine/Private/TemporalSuperResolution/TSRClearPrevTextures.usf",     "MainCS", SF_Compute);
 IMPLEMENT_GLOBAL_SHADER(FTSRDilateVelocityCS,        "/Engine/Private/TemporalSuperResolution/TSRDilateVelocity.usf",        "MainCS", SF_Compute);
 IMPLEMENT_GLOBAL_SHADER(FTSRDetectThinGeometryCS,	 "/Engine/Private/TemporalSuperResolution/TSRDetectThinGeometry.usf",	 "MainCS", SF_Compute);
+IMPLEMENT_GLOBAL_SHADER(FTSRTSRWeightRelaxationCS,	 "/Engine/Private/TemporalSuperResolution/TSRWeightRelaxation.usf",		 "MainCS", SF_Compute);
 IMPLEMENT_GLOBAL_SHADER(FTSRDecimateHistoryCS,       "/Engine/Private/TemporalSuperResolution/TSRDecimateHistory.usf",       "MainCS", SF_Compute);
 IMPLEMENT_GLOBAL_SHADER(FTSRRejectShadingCS,         "/Engine/Private/TemporalSuperResolution/TSRRejectShading.usf",         "MainCS", SF_Compute);
 IMPLEMENT_GLOBAL_SHADER(FTSRSpatialAntiAliasingCS,   "/Engine/Private/TemporalSuperResolution/TSRSpatialAntiAliasing.usf",   "MainCS", SF_Compute);
@@ -2498,87 +2567,134 @@ FDefaultTemporalUpscaler::FOutputs AddTemporalSuperResolutionPasses(
 	}
 
 	//Create the thin geometry mask to avoid over rejection
-	if(PassConfig.ThinGeometryDetectionEnable)
+	if (PassConfig.ThinGeometryDetectionEnable)
 	{
-		FTSRConvolutionNetworkShader::FPermutationDomain ConvolutionNetworkPermutationVector;
-		ConvolutionNetworkPermutationVector.Set<FTSRConvolutionNetworkShader::FWaveSizeOps>(SelectWaveSize(View.GetShaderPlatform(), { 16, 32, 64 }));
-		ConvolutionNetworkPermutationVector.Set<FTSRShader::F16BitVALUDim>(bUse16BitVALU);
-		ConvolutionNetworkPermutationVector.Set<FTSRShader::FAlphaChannelDim>(bSupportsAlpha);
+		{
+			FTSRConvolutionNetworkShader::FPermutationDomain ConvolutionNetworkPermutationVector;
+			ConvolutionNetworkPermutationVector.Set<FTSRConvolutionNetworkShader::FWaveSizeOps>(SelectWaveSize(View.GetShaderPlatform(), { 16, 32, 64 }));
+			ConvolutionNetworkPermutationVector.Set<FTSRShader::F16BitVALUDim>(bUse16BitVALU);
+			ConvolutionNetworkPermutationVector.Set<FTSRShader::FAlphaChannelDim>(bSupportsAlpha);
 
-		FTSRDetectThinGeometryCS::FPermutationDomain PermutationVector;
-		PermutationVector.Set<FTSRDetectThinGeometryCS::FThinGeometryEdgeReprojectionDim>(CVarTSRThinGeometryCoverageEdgeReprojection.GetValueOnRenderThread());
-		PermutationVector.Set<FTSRConvolutionNetworkShader::FPermutationDomain>(ConvolutionNetworkPermutationVector);
-		PermutationVector = FTSRDetectThinGeometryCS::RemapPermutation(PermutationVector);
+			FTSRDetectThinGeometryCS::FPermutationDomain PermutationVector;
+            PermutationVector.Set<FTSRDetectThinGeometryCS::FThinGeometryEdgeReprojectionDim>(CVarTSRThinGeometryCoverageEdgeReprojection.GetValueOnRenderThread());
+			PermutationVector.Set<FTSRConvolutionNetworkShader::FPermutationDomain>(ConvolutionNetworkPermutationVector);
+			PermutationVector = FTSRDetectThinGeometryCS::RemapPermutation(PermutationVector);
 
-		const int32 GroupTileSize = 32;
-		const int32 TileOverscan = FMath::Clamp(CVarTSRShadingTileOverscan.GetValueOnRenderThread(), 3, GroupTileSize / 2 - 1);
-		const int32 TileSize = GroupTileSize - 2 * TileOverscan;
+			const int32 GroupTileSize = 32;
+			const int32 TileOverscan = FMath::Clamp(CVarTSRShadingTileOverscan.GetValueOnRenderThread(), 3, GroupTileSize / 2 - 1);
+			const int32 TileSize = GroupTileSize - 2 * TileOverscan;
 
-		FTSRDetectThinGeometryCS::FParameters* PassParameters = GraphBuilder.AllocParameters<FTSRDetectThinGeometryCS::FParameters>();
-		PassParameters->CommonParameters = CommonParameters;
-		PassParameters->TileOverscan = TileOverscan;
-		PassParameters->ThinGeometryTextureIndex = ThinGeometryTextureIndex;
-		PassParameters->ErrorMultiplier = PassConfig.ThinGeometryErrorMultiplier;
-		PassParameters->MaxRelaxationWeight = FMath::Clamp(CVarTSRThinGeometryCoverageMaxRelaxationWeight.GetValueOnRenderThread(), 0.0f, 1.0f);
-		PassParameters->SceneDepthTexture = PassInputs.SceneDepth.Texture;
+			FTSRDetectThinGeometryCS::FParameters* PassParameters = GraphBuilder.AllocParameters<FTSRDetectThinGeometryCS::FParameters>();
+			PassParameters->CommonParameters = CommonParameters;
+			PassParameters->TileOverscan = TileOverscan;
+			PassParameters->ThinGeometryTextureIndex = ThinGeometryTextureIndex;
+			PassParameters->ErrorMultiplier = PassConfig.ThinGeometryErrorMultiplier;
+			PassParameters->MaxRelaxationWeight = FMath::Clamp(CVarTSRThinGeometryCoverageMaxRelaxationWeight.GetValueOnRenderThread(), 0.0f, 1.0f);
+			PassParameters->SceneDepthTexture = PassInputs.SceneDepth.Texture;
 
-		// Coverage texture
-		{
-			if (bCameraCut)
+			// Coverage texture
 			{
-				PassParameters->ReprojectedHistoryCoverageTexture = GraphBuilder.CreateSRV(BlackArrayDummy);
-				PassParameters->CurrentCoverageTexture = GraphBuilder.CreateSRV(FRDGTextureSRVDesc(BlackDummy));
-			}
-			else
-			{
-				PassParameters->ReprojectedHistoryCoverageTexture = ReprojectedHistoryCoverageTexture ? GraphBuilder.CreateSRV(FRDGTextureSRVDesc::CreateForSlice(
-					ReprojectedHistoryCoverageTexture, /* SliceIndex = */ 0)) : GraphBuilder.CreateSRV(FRDGTextureSRVDesc(BlackDummy));;
-
-				if (PassInputs.FlickeringInputTexture.IsValid())
+				if (bCameraCut)
 				{
-					ensure(InputRect == PassInputs.FlickeringInputTexture.ViewRect);
-					int32 ArrayIndex = GetTSRMainFlickeringLumaTextureArraySize() - 1;
-					PassParameters->CurrentCoverageTexture = GraphBuilder.CreateSRV(
-						FRDGTextureSRVDesc::CreateForSlice(PassInputs.FlickeringInputTexture.Texture,  /* SliceIndex = */ ArrayIndex));
+					PassParameters->ReprojectedHistoryCoverageTexture = GraphBuilder.CreateSRV(BlackArrayDummy);
+					PassParameters->CurrentCoverageTexture = GraphBuilder.CreateSRV(FRDGTextureSRVDesc(BlackDummy));
 				}
 				else
 				{
-					PassParameters->CurrentCoverageTexture = GraphBuilder.CreateSRV(FRDGTextureSRVDesc(BlackDummy));
+					PassParameters->ReprojectedHistoryCoverageTexture = ReprojectedHistoryCoverageTexture ? GraphBuilder.CreateSRV(FRDGTextureSRVDesc::CreateForSlice(
+						ReprojectedHistoryCoverageTexture, /* SliceIndex = */ 0)) : GraphBuilder.CreateSRV(FRDGTextureSRVDesc(BlackDummy));;
+
+					if (PassInputs.FlickeringInputTexture.IsValid())
+					{
+						ensure(InputRect == PassInputs.FlickeringInputTexture.ViewRect);
+						int32 ArrayIndex = GetTSRMainFlickeringLumaTextureArraySize() - 1;
+						PassParameters->CurrentCoverageTexture = GraphBuilder.CreateSRV(
+							FRDGTextureSRVDesc::CreateForSlice(PassInputs.FlickeringInputTexture.Texture,  /* SliceIndex = */ ArrayIndex));
+					}
+					else
+					{
+						PassParameters->CurrentCoverageTexture = GraphBuilder.CreateSRV(FRDGTextureSRVDesc(BlackDummy));
+					}
 				}
-			}
 
-			//Output
-			{
-				if (View.bStatePrevViewInfoIsReadOnly)
+				//Output
 				{
-					PassParameters->HistoryCoverageOutput = CreateDummyUAVArray(GraphBuilder, History.CoverageArray->Desc.Format);
-				}
-				else
-				{
-					FRDGTextureUAVDesc CoverageUAVDesc(History.CoverageArray);
-					PassParameters->HistoryCoverageOutput = GraphBuilder.CreateUAV(CoverageUAVDesc);
+					if (View.bStatePrevViewInfoIsReadOnly)
+					{
+						PassParameters->HistoryCoverageOutput = CreateDummyUAVArray(GraphBuilder, History.CoverageArray->Desc.Format);
+					}
+					else
+					{
+						FRDGTextureUAVDesc CoverageUAVDesc(History.CoverageArray);
+						PassParameters->HistoryCoverageOutput = GraphBuilder.CreateUAV(CoverageUAVDesc);
+					}
 				}
 			}
+
+			PassParameters->R8Output = GraphBuilder.CreateUAV(R8OutputTexture);
+			PassParameters->DebugOutput = CreateDebugUAV(InputExtent, TEXT("Debug.TSR.DetectThinGeometry"));
+
+			TShaderMapRef<FTSRDetectThinGeometryCS> ComputeShader(View.ShaderMap, PermutationVector);
+			FComputeShaderUtils::AddPass(
+				GraphBuilder,
+				RDG_EVENT_NAME("TSR DetectThinGeometry(#%d TileSize=%d PaddingCostMultiplier=%1.1f WaveSize=%d VALU=%s%s) %dx%d",
+					PermutationVector.ToDimensionValueId(),
+					TileSize,
+					FMath::Pow(float(GroupTileSize) / float(TileSize), 2),
+					int32(PermutationVector.Get<FTSRConvolutionNetworkShader::FPermutationDomain>().Get<FTSRDetectThinGeometryCS::FWaveSizeOps>()),
+					PermutationVector.Get<FTSRConvolutionNetworkShader::FPermutationDomain>().Get<FTSRShader::F16BitVALUDim>() ? TEXT("16bit") : TEXT("32bit"),
+					PermutationVector.Get<FTSRConvolutionNetworkShader::FPermutationDomain>().Get<FTSRShader::FAlphaChannelDim>() ? TEXT(" AlphaChannel") : TEXT(""),
+					InputRect.Width(), InputRect.Height()),
+				AsyncComputePasses >= 2 ? ERDGPassFlags::AsyncCompute : ERDGPassFlags::Compute,
+				ComputeShader,
+				PassParameters,
+				FComputeShaderUtils::GetGroupCount(InputRect.Size(), TileSize));
 		}
 
-		PassParameters->R8Output = GraphBuilder.CreateUAV(R8OutputTexture);
-		PassParameters->DebugOutput = CreateDebugUAV(InputExtent, TEXT("Debug.TSR.DetectThinGeometry"));
+		// Trim history relaxation weight
+		if(CVarTSRThinGeometryDetectionWeightRelaxation.GetValueOnRenderThread() != 0)
+		{
+			FTSRConvolutionNetworkShader::FPermutationDomain ConvolutionNetworkPermutationVector;
+			ConvolutionNetworkPermutationVector.Set<FTSRConvolutionNetworkShader::FWaveSizeOps>(SelectWaveSize(View.GetShaderPlatform(), { 16, 32, 64 }));
+			ConvolutionNetworkPermutationVector.Set<FTSRShader::F16BitVALUDim>(bUse16BitVALU);
+			ConvolutionNetworkPermutationVector.Set<FTSRShader::FAlphaChannelDim>(bSupportsAlpha);
+
+			FTSRTSRWeightRelaxationCS::FPermutationDomain PermutationVector;
+			PermutationVector.Set<FTSRConvolutionNetworkShader::FPermutationDomain>(ConvolutionNetworkPermutationVector);
+			PermutationVector = FTSRTSRWeightRelaxationCS::RemapPermutation(PermutationVector);
+
+			const int32 GroupTileSize = 32;
+			const int32 TileOverscan = 0;//FMath::Clamp(CVarTSRShadingTileOverscan.GetValueOnRenderThread(), 3, GroupTileSize / 2 - 1);
+			const int32 TileSize = GroupTileSize - 2 * TileOverscan;
+
+			FTSRTSRWeightRelaxationCS::FParameters* PassParameters = GraphBuilder.AllocParameters<FTSRTSRWeightRelaxationCS::FParameters>();
+			PassParameters->CommonParameters = CommonParameters;
+			PassParameters->TileOverscan = TileOverscan;
+			PassParameters->ThinGeometryTextureIndex = ThinGeometryTextureIndex;
+			PassParameters->MaxRelaxationWeight = FMath::Clamp(CVarTSRThinGeometryCoverageMaxRelaxationWeight.GetValueOnRenderThread(), 0.0f, 1.0f);
+
+			PassParameters->InputTexture = PassInputs.SceneColor.Texture;
+			PassParameters->InputSceneTranslucencyTexture = SeparateTranslucencyTexture;
+			PassParameters->R8Output = GraphBuilder.CreateUAV(R8OutputTexture);
+			PassParameters->DebugOutput = CreateDebugUAV(InputExtent, TEXT("Debug.TSR.WeightRelaxation"));
+
+			TShaderMapRef<FTSRTSRWeightRelaxationCS> ComputeShader(View.ShaderMap, PermutationVector);
+			FComputeShaderUtils::AddPass(
+				GraphBuilder,
+				RDG_EVENT_NAME("TSR WeightRelaxation(#%d TileSize=%d PaddingCostMultiplier=%1.1f WaveSize=%d VALU=%s%s) %dx%d",
+					PermutationVector.ToDimensionValueId(),
+					TileSize,
+					FMath::Pow(float(GroupTileSize) / float(TileSize), 2),
+					int32(PermutationVector.Get<FTSRConvolutionNetworkShader::FPermutationDomain>().Get<FTSRDetectThinGeometryCS::FWaveSizeOps>()),
+					PermutationVector.Get<FTSRConvolutionNetworkShader::FPermutationDomain>().Get<FTSRShader::F16BitVALUDim>() ? TEXT("16bit") : TEXT("32bit"),
+					PermutationVector.Get<FTSRConvolutionNetworkShader::FPermutationDomain>().Get<FTSRShader::FAlphaChannelDim>() ? TEXT(" AlphaChannel") : TEXT(""),
+					InputRect.Width(), InputRect.Height()),
+				    AsyncComputePasses >= 2 ? ERDGPassFlags::AsyncCompute : ERDGPassFlags::Compute,
+				ComputeShader,
+				PassParameters,
+				FComputeShaderUtils::GetGroupCount(InputRect.Size(), TileSize));
+		}
 
-		TShaderMapRef<FTSRDetectThinGeometryCS> ComputeShader(View.ShaderMap, PermutationVector);
-		FComputeShaderUtils::AddPass(
-			GraphBuilder,
-			RDG_EVENT_NAME("TSR DetectThinGeometry(#%d TileSize=%d PaddingCostMultiplier=%1.1f WaveSize=%d VALU=%s%s) %dx%d",
-				PermutationVector.ToDimensionValueId(),
-				TileSize,
-				FMath::Pow(float(GroupTileSize) / float(TileSize), 2),
-				int32(PermutationVector.Get<FTSRConvolutionNetworkShader::FPermutationDomain>().Get<FTSRDetectThinGeometryCS::FWaveSizeOps>()),
-				PermutationVector.Get<FTSRConvolutionNetworkShader::FPermutationDomain>().Get<FTSRShader::F16BitVALUDim>() ? TEXT("16bit") : TEXT("32bit"),
-				PermutationVector.Get<FTSRConvolutionNetworkShader::FPermutationDomain>().Get<FTSRShader::FAlphaChannelDim>() ? TEXT(" AlphaChannel") : TEXT(""),
-				InputRect.Width(), InputRect.Height()),
-			AsyncComputePasses >= 2 ? ERDGPassFlags::AsyncCompute : ERDGPassFlags::Compute,
-			ComputeShader,
-			PassParameters,
-			FComputeShaderUtils::GetGroupCount(InputRect.Size(), TileSize));
 	}
 	// Merge PostDOF translucency within same scene color.
 	FRDGTextureRef InputSceneColorTexture = nullptr;
