Substrate - Force OpaqueRoughRefraction disabled when Blendable GBuffer is on.
Updated SystemInfo vis.

#rnx
#tests PC editor dx12, BlendableGBuffer 0/1, OpaqueMaterialRoughRefraction 0/1

diff --git a/Engine/Shaders/Public/ShaderVersion.ush b/Engine/Shaders/Public/ShaderVersion.ush
--- a/Engine/Shaders/Public/ShaderVersion.ush
+++ b/Engine/Shaders/Public/ShaderVersion.ush
@@ -3,4 +3,4 @@
 // Each time the console command is executed it generates a new GUID. As this file is included
 // in Platform.ush (which should be included in any shader) it allows to invalidate the shader DDC.
 // 
-#pragma message("UESHADERMETADATA_VERSION 9FE56796-34AB-4080-9B31-5FD7214FAF3A")
+#pragma message("UESHADERMETADATA_VERSION 9FE56796-34AB-4080-9B31-5FD7214FAF3B")
diff --git a/Engine/Source/Runtime/Engine/Private/ShaderCompiler/ShaderCompiler.cpp b/Engine/Source/Runtime/Engine/Private/ShaderCompiler/ShaderCompiler.cpp
--- a/Engine/Source/Runtime/Engine/Private/ShaderCompiler/ShaderCompiler.cpp
+++ b/Engine/Source/Runtime/Engine/Private/ShaderCompiler/ShaderCompiler.cpp
@@ -3907,7 +3907,7 @@ void GlobalBeginCompileShader(
 		const bool bSubstrateBackCompatibility = bSubstrate && Substrate::IsBackCompatibilityEnabled();
 		SET_SHADER_DEFINE(Input.Environment, PROJECT_SUBSTRATE_BACKCOMPATIBILITY, bSubstrateBackCompatibility ? 1 : 0);
 
-		const bool bSubstrateOpaqueRoughRefrac = bSubstrate && Substrate::IsOpaqueRoughRefractionEnabled();
+		const bool bSubstrateOpaqueRoughRefrac = bSubstrate && Substrate::IsOpaqueRoughRefractionEnabled(Target.GetPlatform());
 		SET_SHADER_DEFINE(Input.Environment, SUBSTRATE_OPAQUE_ROUGH_REFRACTION_ENABLED, bSubstrateOpaqueRoughRefrac ? 1 : 0);
 
 		const bool bSubstrateAdvDebug = bSubstrate && Substrate::IsAdvancedVisualizationEnabled();
diff --git a/Engine/Source/Runtime/Engine/Private/SubstrateVisualizationData.cpp b/Engine/Source/Runtime/Engine/Private/SubstrateVisualizationData.cpp
--- a/Engine/Source/Runtime/Engine/Private/SubstrateVisualizationData.cpp
+++ b/Engine/Source/Runtime/Engine/Private/SubstrateVisualizationData.cpp
@@ -125,7 +125,7 @@ void FSubstrateVisualizationData::Initialize()
 			LOCTEXT("OpaqueRoughRefractionClassificationDesc", "Visualizes Substrate Opaque Rough Refraction Classification"),
 			FSubstrateViewMode::RoughRefractionClassification,
 			true,
-			Substrate::IsOpaqueRoughRefractionEnabled(),
+			Substrate::IsOpaqueRoughRefractionEnabled(GMaxRHIShaderPlatform),
 			LOCTEXT("IsSubstrateRoughRefractionEnabled", "Substrate rough refraction r.Substrate.OpaqueMaterialRoughRefraction is disabled"));
 
 		AddVisualizationMode(
diff --git a/Engine/Source/Runtime/RenderCore/Private/RenderUtils.cpp b/Engine/Source/Runtime/RenderCore/Private/RenderUtils.cpp
--- a/Engine/Source/Runtime/RenderCore/Private/RenderUtils.cpp
+++ b/Engine/Source/Runtime/RenderCore/Private/RenderUtils.cpp
@@ -2147,9 +2147,9 @@ namespace Substrate
 		return IsSubstrateEnabled() && !IsSubstrateBlendableGBufferEnabled(InPlatform) && IsUsingDBuffers(InPlatform) && bDBufferPassSupported && CVarSubstrateDBufferPassPlatform.Get(InPlatform) > 0;
 	}
 
-	bool IsOpaqueRoughRefractionEnabled()
+	bool IsOpaqueRoughRefractionEnabled(EShaderPlatform InPlatform)
 	{
-		return IsSubstrateEnabled() && CVarSubstrateOpaqueMaterialRoughRefraction.GetValueOnAnyThread() > 0;
+		return IsSubstrateEnabled() && CVarSubstrateOpaqueMaterialRoughRefraction.GetValueOnAnyThread() > 0 && !IsSubstrateBlendableGBufferEnabled(InPlatform);
 	}
 
 	bool IsAdvancedVisualizationEnabled()
diff --git a/Engine/Source/Runtime/RenderCore/Private/Shader.cpp b/Engine/Source/Runtime/RenderCore/Private/Shader.cpp
--- a/Engine/Source/Runtime/RenderCore/Private/Shader.cpp
+++ b/Engine/Source/Runtime/RenderCore/Private/Shader.cpp
@@ -2400,7 +2400,7 @@ void ShaderMapAppendKey(EShaderPlatform Platform, FShaderKeyGenerator& KeyGen)
 			KeyGen.Append(TEXT("BACKCOMPAT"));
 		}
 
-		if (Substrate::IsOpaqueRoughRefractionEnabled())
+		if (Substrate::IsOpaqueRoughRefractionEnabled(Platform))
 		{
 			KeyGen.AppendSeparator();
 			KeyGen.Append(TEXT("ROUGHDIFF"));
diff --git a/Engine/Source/Runtime/RenderCore/Public/RenderUtils.h b/Engine/Source/Runtime/RenderCore/Public/RenderUtils.h
--- a/Engine/Source/Runtime/RenderCore/Public/RenderUtils.h
+++ b/Engine/Source/Runtime/RenderCore/Public/RenderUtils.h
@@ -600,7 +600,7 @@ namespace Substrate
 	RENDERCORE_API bool IsSpecularProfileEnabled(EShaderPlatform InPlatform);
 	RENDERCORE_API bool IsBackCompatibilityEnabled();
 	RENDERCORE_API bool IsDBufferPassEnabled(EShaderPlatform InPlatform);
-	RENDERCORE_API bool IsOpaqueRoughRefractionEnabled();
+	RENDERCORE_API bool IsOpaqueRoughRefractionEnabled(EShaderPlatform InPlatform);
 	RENDERCORE_API bool IsAdvancedVisualizationEnabled();
 	RENDERCORE_API bool IsMaterialLayeringSupportEnabled();
 
diff --git a/Engine/Source/Runtime/Renderer/Private/ClusteredDeferredShadingPass.cpp b/Engine/Source/Runtime/Renderer/Private/ClusteredDeferredShadingPass.cpp
--- a/Engine/Source/Runtime/Renderer/Private/ClusteredDeferredShadingPass.cpp
+++ b/Engine/Source/Runtime/Renderer/Private/ClusteredDeferredShadingPass.cpp
@@ -208,7 +208,7 @@ static void InternalAddClusteredDeferredShadingPass(
 	{
 		PSParameters->RenderTargets[0] = FRenderTargetBinding(View.HairStrandsViewData.VisibilityData.SampleLightingTexture, ERenderTargetLoadAction::ELoad);
 	}
-	if (Substrate::IsOpaqueRoughRefractionEnabled())
+	if (Substrate::IsOpaqueRoughRefractionEnabled(View.GetShaderPlatform()))
 	{
 		check(SubstrateSceneData);
 		PSParameters->RenderTargets[1] = FRenderTargetBinding(SubstrateSceneData->SeparatedOpaqueRoughRefractionSceneColor, ERenderTargetLoadAction::ELoad);
@@ -250,7 +250,7 @@ static void InternalAddClusteredDeferredShadingPass(
 			InRHICmdList.ApplyCachedRenderTargets(GraphicsPSOInit);
 
 			// Additive blend to accumulate lighting contributions.
-			if (Substrate::IsOpaqueRoughRefractionEnabled())
+			if (Substrate::IsOpaqueRoughRefractionEnabled(View.GetShaderPlatform()))
 			{
 				GraphicsPSOInit.BlendState = TStaticBlendState<
 					CW_RGBA, BO_Add, BF_One, BF_One, BO_Add, BF_One, BF_One,
diff --git a/Engine/Source/Runtime/Renderer/Private/IndirectLightRendering.cpp b/Engine/Source/Runtime/Renderer/Private/IndirectLightRendering.cpp
--- a/Engine/Source/Runtime/Renderer/Private/IndirectLightRendering.cpp
+++ b/Engine/Source/Runtime/Renderer/Private/IndirectLightRendering.cpp
@@ -213,7 +213,7 @@ class FDiffuseIndirectCompositePS : public FGlobalShader
 	static void ModifyCompilationEnvironment(const FGlobalShaderPermutationParameters& Parameters, FShaderCompilerEnvironment& OutEnvironment)
 	{
 		OutEnvironment.CompilerFlags.Add(CFLAG_ForceOptimization);
-		if (Substrate::IsSubstrateEnabled() && Substrate::IsOpaqueRoughRefractionEnabled())
+		if (Substrate::IsSubstrateEnabled() && Substrate::IsOpaqueRoughRefractionEnabled(Parameters.Platform))
 		{
 			OutEnvironment.CompilerFlags.Add(CFLAG_AllowTypedUAVLoads);
 		}
@@ -1541,7 +1541,7 @@ void FDeferredShadingSceneRenderer::RenderDiffuseIndirectAndAmbientOcclusion(
 				check(Substrate::IsSubstrateEnabled());
 
 				// Rough refraction targets
-				if (Substrate::IsOpaqueRoughRefractionEnabled())
+				if (Substrate::IsOpaqueRoughRefractionEnabled(View.GetShaderPlatform()))
 				{
 					PSParameters->OutOpaqueRoughRefractionSceneColor = GraphBuilder.CreateUAV(Scene->SubstrateSceneData.SeparatedOpaqueRoughRefractionSceneColor);
 					PSParameters->OutSubSurfaceSceneColor = GraphBuilder.CreateUAV(Scene->SubstrateSceneData.SeparatedSubSurfaceSceneColor);
@@ -1633,7 +1633,7 @@ void FDeferredShadingSceneRenderer::RenderDiffuseIndirectAndAmbientOcclusion(
 				PSParameters->ViewUniformBuffer = View.ViewUniformBuffer;
 
 				PSParameters->RenderTargets[0] = FRenderTargetBinding(SceneColorTexture, ERenderTargetLoadAction::ELoad);
-				if (Substrate::IsOpaqueRoughRefractionEnabled())
+				if (Substrate::IsOpaqueRoughRefractionEnabled(View.GetShaderPlatform()))
 				{
 					PSParameters->RenderTargets[1] = FRenderTargetBinding(Scene->SubstrateSceneData.SeparatedOpaqueRoughRefractionSceneColor, ERenderTargetLoadAction::ELoad);
 					PSParameters->RenderTargets[2] = FRenderTargetBinding(Scene->SubstrateSceneData.SeparatedSubSurfaceSceneColor, ERenderTargetLoadAction::ELoad);
@@ -1843,14 +1843,14 @@ static void AddSkyReflectionPass(
 	}
 
 	PassParameters->RenderTargets[0] = FRenderTargetBinding(SceneColorTexture.Target, ERenderTargetLoadAction::ELoad);
-	if (Substrate::IsOpaqueRoughRefractionEnabled())
+	if (Substrate::IsOpaqueRoughRefractionEnabled(View.GetShaderPlatform()))
 	{
 		PassParameters->RenderTargets[1] = FRenderTargetBinding(Scene->SubstrateSceneData.SeparatedOpaqueRoughRefractionSceneColor, ERenderTargetLoadAction::ELoad);
 		PassParameters->RenderTargets[2] = FRenderTargetBinding(Scene->SubstrateSceneData.SeparatedSubSurfaceSceneColor, ERenderTargetLoadAction::ELoad);
 	}
 
 	// Bind hair data
-	const bool bCheckerboardSubsurfaceRendering = IsSubsurfaceCheckerboardFormat(SceneColorTexture.Target->Desc.Format);
+	const bool bCheckerboardSubsurfaceRendering = IsSubsurfaceCheckerboardFormat(SceneColorTexture.Target->Desc.Format, View);
 	const bool bHasHairCardsOrMeshes = View.HairCardsMeshElements.Num() > 0;
 
 	auto PermutationVector = FReflectionEnvironmentSkyLightingPS::BuildPermutationVector(
@@ -1905,7 +1905,7 @@ static void AddSkyReflectionPass(
 			}
 			else
 			{
-				if (Substrate::IsOpaqueRoughRefractionEnabled())
+				if (Substrate::IsOpaqueRoughRefractionEnabled(View.GetShaderPlatform()))
 				{
 					GraphicsPSOInit.BlendState = TStaticBlendState<
 						CW_RGBA, BO_Add, BF_One, BF_One, BO_Add, BF_One, BF_One,
diff --git a/Engine/Source/Runtime/Renderer/Private/LightRendering.cpp b/Engine/Source/Runtime/Renderer/Private/LightRendering.cpp
--- a/Engine/Source/Runtime/Renderer/Private/LightRendering.cpp
+++ b/Engine/Source/Runtime/Renderer/Private/LightRendering.cpp
@@ -2516,7 +2516,7 @@ static FDeferredLightPS::FParameters GetDeferredLightPSParameters(
 
 	// PS - Render Targets
 	Out.RenderTargets[0] = FRenderTargetBinding(SceneColorTexture, ERenderTargetLoadAction::ELoad);
-	if (Substrate::IsOpaqueRoughRefractionEnabled() && Substrate::UsesSubstrateMaterialBuffer(View.GetShaderPlatform()))
+	if (Substrate::IsOpaqueRoughRefractionEnabled(View.GetShaderPlatform()) && Substrate::UsesSubstrateMaterialBuffer(View.GetShaderPlatform()))
 	{
 		Out.RenderTargets[1] = FRenderTargetBinding(Scene->SubstrateSceneData.SeparatedOpaqueRoughRefractionSceneColor, ERenderTargetLoadAction::ELoad);
 		Out.RenderTargets[2] = FRenderTargetBinding(Scene->SubstrateSceneData.SeparatedSubSurfaceSceneColor, ERenderTargetLoadAction::ELoad);
@@ -2534,10 +2534,11 @@ uint32 SetupLightGraphicsPSOState(
 	bool bCameraInsideLightGeometry, 
 	bool bReverseCulling, 
 	ESubstrateTileType SubstrateTileMaterialType,
-	FGraphicsPipelineStateInitializer& GraphicsPSOInit)
+	FGraphicsPipelineStateInitializer& GraphicsPSOInit,
+	EShaderPlatform ShaderPlatform)
 {
 	GraphicsPSOInit.PrimitiveType = PT_TriangleList;
-	if (Substrate::IsOpaqueRoughRefractionEnabled())
+	if (Substrate::IsOpaqueRoughRefractionEnabled(ShaderPlatform))
 	{
 		GraphicsPSOInit.BlendState = TStaticBlendState<
 			CW_RGBA, BO_Add, BF_One, BF_One, BO_Add, BF_One, BF_One,
@@ -2608,7 +2609,7 @@ static void InternalRenderLight(
 				//@todo - accurate ortho camera / light intersection
 			|| !View.IsPerspectiveProjection();
 
-		const uint32 StencilRef = SetupLightGraphicsPSOState(LightType == LightType_Directional, bCameraInsideLightGeometry, View.bReverseCulling, SubstrateTileMaterialType, GraphicsPSOInit);
+		const uint32 StencilRef = SetupLightGraphicsPSOState(LightType == LightType_Directional, bCameraInsideLightGeometry, View.bReverseCulling, SubstrateTileMaterialType, GraphicsPSOInit, View.GetShaderPlatform());
 
 		// Set the device viewport for the view.
 		RHICmdList.SetViewport(View.ViewRect.Min.X, View.ViewRect.Min.Y, 0.0f, View.ViewRect.Max.X, View.ViewRect.Max.Y, 1.0f);
@@ -3435,7 +3436,7 @@ void DeferredLightGlobalPSOCollector(const FSceneTexturesConfig& SceneTexturesCo
 		}
 		else
 		{
-			SetupLightGraphicsPSOState(bIsDirectional, bCameraInsideLightGeometry, bReverseCulling, SubstrateTileMaterialType, GraphicsPSOInit);
+			SetupLightGraphicsPSOState(bIsDirectional, bCameraInsideLightGeometry, bReverseCulling, SubstrateTileMaterialType, GraphicsPSOInit, ShaderPlatform);
 
 			if (bIsDirectional)
 			{
@@ -3463,7 +3464,7 @@ void DeferredLightGlobalPSOCollector(const FSceneTexturesConfig& SceneTexturesCo
 			}
 
 			AddRenderTargetInfo(SceneTexturesConfig.ColorFormat, SceneTexturesConfig.ColorCreateFlags, RenderTargetsInfo);
-			if (Substrate::IsOpaqueRoughRefractionEnabled() && Substrate::UsesSubstrateMaterialBuffer(ShaderPlatform))
+			if (Substrate::IsOpaqueRoughRefractionEnabled(ShaderPlatform) && Substrate::UsesSubstrateMaterialBuffer(ShaderPlatform))
 			{
 				// TODO: find out the render target info for subtrate here
 				//Out.RenderTargets[1] = FRenderTargetBinding(Scene->SubstrateSceneData.SeparatedOpaqueRoughRefractionSceneColor, ERenderTargetLoadAction::ELoad);
diff --git a/Engine/Source/Runtime/Renderer/Private/PostProcess/PostProcessSubsurface.cpp b/Engine/Source/Runtime/Renderer/Private/PostProcess/PostProcessSubsurface.cpp
--- a/Engine/Source/Runtime/Renderer/Private/PostProcess/PostProcessSubsurface.cpp
+++ b/Engine/Source/Runtime/Renderer/Private/PostProcess/PostProcessSubsurface.cpp
@@ -323,9 +323,9 @@ uint32 GetSubsurfaceRequiredViewMask(TArrayView<const FViewInfo> Views)
 	return ViewMask;
 }
 
-bool IsSubsurfaceCheckerboardFormat(EPixelFormat SceneColorFormat)
+bool IsSubsurfaceCheckerboardFormat(EPixelFormat SceneColorFormat, const FViewInfo& View)
 {
-	if (Substrate::IsOpaqueRoughRefractionEnabled())
+	if (Substrate::IsOpaqueRoughRefractionEnabled(View.GetShaderPlatform()))
 	{
 		// With this mode, specular and subsurface colors are correctly separated so checkboard is not required.
 		return false;
@@ -1000,7 +1000,7 @@ void AddSubsurfaceViewPass(
 
 	const bool bHalfRes = (SubsurfaceMode == ESubsurfaceMode::HalfRes);
 
-	const bool bCheckerboard = IsSubsurfaceCheckerboardFormat(SceneColorTextureDesc.Format);
+	const bool bCheckerboard = IsSubsurfaceCheckerboardFormat(SceneColorTextureDesc.Format, View);
 
 	const uint32 ScaleFactor = bHalfRes ? 2 : 1;
 	
@@ -1049,8 +1049,8 @@ void AddSubsurfaceViewPass(
 	const FScreenPassTextureViewportParameters SubsurfaceViewportParameters = GetScreenPassTextureViewportParameters(SubsurfaceViewport);
 	const FScreenPassTextureViewportParameters SceneViewportParameters = GetScreenPassTextureViewportParameters(SceneViewport);
 
-	const bool bReadSeparatedSubSurfaceSceneColor = Substrate::IsOpaqueRoughRefractionEnabled();
-	const bool bWriteSeparatedOpaqueRoughRefractionSceneColor = Substrate::IsOpaqueRoughRefractionEnabled();
+	const bool bReadSeparatedSubSurfaceSceneColor = Substrate::IsOpaqueRoughRefractionEnabled(View.GetShaderPlatform());
+	const bool bWriteSeparatedOpaqueRoughRefractionSceneColor = Substrate::IsOpaqueRoughRefractionEnabled(View.GetShaderPlatform());
 	FRDGTextureRef SeparatedSubSurfaceSceneColor = View.SubstrateViewData.SceneData->SeparatedSubSurfaceSceneColor;
 	FRDGTextureRef SeparatedOpaqueRoughRefractionSceneColor = View.SubstrateViewData.SceneData->SeparatedOpaqueRoughRefractionSceneColor;
 
diff --git a/Engine/Source/Runtime/Renderer/Private/PostProcess/PostProcessSubsurface.h b/Engine/Source/Runtime/Renderer/Private/PostProcess/PostProcessSubsurface.h
--- a/Engine/Source/Runtime/Renderer/Private/PostProcess/PostProcessSubsurface.h
+++ b/Engine/Source/Runtime/Renderer/Private/PostProcess/PostProcessSubsurface.h
@@ -12,7 +12,7 @@ bool IsSubsurfaceEnabled();
 bool IsSubsurfaceRequiredForView(const FViewInfo& View);
 
 // Returns whether checkerboard rendering is enabled for the provided format.
-bool IsSubsurfaceCheckerboardFormat(EPixelFormat SceneColorFormat);
+bool IsSubsurfaceCheckerboardFormat(EPixelFormat SceneColorFormat, const FViewInfo& View);
 
 class FSceneTextureParameters;
 
diff --git a/Engine/Source/Runtime/Renderer/Private/SceneRendering.cpp b/Engine/Source/Runtime/Renderer/Private/SceneRendering.cpp
--- a/Engine/Source/Runtime/Renderer/Private/SceneRendering.cpp
+++ b/Engine/Source/Runtime/Renderer/Private/SceneRendering.cpp
@@ -1301,7 +1301,7 @@ void FViewInfo::SetupUniformBufferParameters(
 		InPrevViewMatrices
 	);
 
-	const bool bCheckerboardSubsurfaceRendering = IsSubsurfaceCheckerboardFormat(SceneTexturesConfig.ColorFormat);
+	const bool bCheckerboardSubsurfaceRendering = IsSubsurfaceCheckerboardFormat(SceneTexturesConfig.ColorFormat, *this);
 	ViewUniformShaderParameters.bCheckerboardSubsurfaceProfileRendering = bCheckerboardSubsurfaceRendering ? 1.0f : 0.0f;
 
 	ViewUniformShaderParameters.IndirectLightingCacheShowFlag = Family->EngineShowFlags.IndirectLightingCache;
diff --git a/Engine/Source/Runtime/Renderer/Private/Substrate/Substrate.cpp b/Engine/Source/Runtime/Renderer/Private/Substrate/Substrate.cpp
--- a/Engine/Source/Runtime/Renderer/Private/Substrate/Substrate.cpp
+++ b/Engine/Source/Runtime/Renderer/Private/Substrate/Substrate.cpp
@@ -240,7 +240,7 @@ static void InitialiseSubstrateViewData(FRDGBuilder& GraphBuilder, FViewInfo& Vi
 			AddClearUAVPass(GraphBuilder, Out.ClassificationTileDispatchIndirectBufferUAV, 0);
 
 			// Separated subsurface & rough refraction textures (tile data)
-			const uint32 RoughTileCount = IsOpaqueRoughRefractionEnabled() ? TileResolution.X * TileResolution.Y : 4;
+			const uint32 RoughTileCount = IsOpaqueRoughRefractionEnabled(View.GetShaderPlatform()) ? TileResolution.X * TileResolution.Y : 4;
 			const uint32 DecalTileCount = IsDBufferPassEnabled(View.GetShaderPlatform()) ? TileResolution.X * TileResolution.Y : 4;
 			const uint32 RegularTileCount = TileResolution.X * TileResolution.Y;
 
@@ -495,7 +495,7 @@ void InitialiseSubstrateFrameSceneData(FRDGBuilder& GraphBuilder, FSceneRenderer
 		// Separated subsurface and rough refraction textures
 		if (bNeedsMaterialBuffer)
 		{
-			const bool bIsSubstrateOpaqueMaterialRoughRefractionEnabled = IsOpaqueRoughRefractionEnabled();
+			const bool bIsSubstrateOpaqueMaterialRoughRefractionEnabled = IsOpaqueRoughRefractionEnabled(SceneRenderer.ShaderPlatform);
 			const FIntPoint OpaqueRoughRefractionSceneExtent		 = bIsSubstrateOpaqueMaterialRoughRefractionEnabled ? SceneTextureExtent : FIntPoint(4, 4);
 			
 			Out.OpaqueRoughRefractionTexture = GraphBuilder.CreateTexture(
diff --git a/Engine/Source/Runtime/Renderer/Private/Substrate/SubstrateRoughRefraction.cpp b/Engine/Source/Runtime/Renderer/Private/Substrate/SubstrateRoughRefraction.cpp
--- a/Engine/Source/Runtime/Renderer/Private/Substrate/SubstrateRoughRefraction.cpp
+++ b/Engine/Source/Runtime/Renderer/Private/Substrate/SubstrateRoughRefraction.cpp
@@ -69,12 +69,18 @@ void AddSubstrateOpaqueRoughRefractionPasses(
 	TArrayView<const FViewInfo> Views)
 {
 	const uint32 ViewCount = Views.Num();
-	const bool bOpaqueRoughRefractionEnabled = IsOpaqueRoughRefractionEnabled() && ViewCount > 0;
-	RDG_EVENT_SCOPE_CONDITIONAL(GraphBuilder, bOpaqueRoughRefractionEnabled, "Substrate::OpaqueRoughRefraction");
+
+	bool bOpaqueRoughRefractionEnabled = false;
+	for (uint32 ViewIndex = 0; ViewIndex < ViewCount; ++ViewIndex)
+	{
+		const FViewInfo& View = Views[ViewIndex];
+		bOpaqueRoughRefractionEnabled |= IsOpaqueRoughRefractionEnabled(View.GetShaderPlatform());;
+	}
 	if (!bOpaqueRoughRefractionEnabled)
 	{
 		return;
 	}
+	RDG_EVENT_SCOPE_CONDITIONAL(GraphBuilder, bOpaqueRoughRefractionEnabled, "Substrate::OpaqueRoughRefraction");
 
 	FRDGTextureRef SceneColorTexture = SceneTextures.Color.Target;
 	FRDGTextureRef TempTexture = nullptr;
diff --git a/Engine/Source/Runtime/Renderer/Private/Substrate/SubstrateVisualize.cpp b/Engine/Source/Runtime/Renderer/Private/Substrate/SubstrateVisualize.cpp
--- a/Engine/Source/Runtime/Renderer/Private/Substrate/SubstrateVisualize.cpp
+++ b/Engine/Source/Runtime/Renderer/Private/Substrate/SubstrateVisualize.cpp
@@ -186,6 +186,7 @@ class FMaterialDebugSubstrateTreeCS : public FGlobalShader
 	static bool IsSupported(EShaderPlatform InPlatform, EShaderPermutationFlags InFlags=EShaderPermutationFlags::None)
 	{
 		return IsSubstrateDebugVisualizationSupported(InPlatform, true, InFlags)
+			&& !IsSubstrateBlendableGBufferEnabled(InPlatform)
 			&& !IsVulkanPlatform(InPlatform)  // SUBSTRATE_TODO Move to CPU debug visualisation and it should then work on all platforms
 			&& !IsMetalPlatform(InPlatform);
 	}
@@ -233,6 +234,7 @@ class FMaterialDebugSubstrateTreePS : public FGlobalShader
 	static bool IsSupported(EShaderPlatform InPlatform, EShaderPermutationFlags InFlags=EShaderPermutationFlags::None)
 	{
 		return IsSubstrateDebugVisualizationSupported(InPlatform, true, InFlags)
+			&& !IsSubstrateBlendableGBufferEnabled(InPlatform)
 			&& !IsVulkanPlatform(InPlatform)  // SUBSTRATE_TODO Move to CPU debug visualisation and it should then work on all platforms
 			&& !IsMetalPlatform(InPlatform);
 	}
@@ -1066,7 +1068,7 @@ static void AddVisualizeSystemInfoPasses(FRDGBuilder& GraphBuilder, const FViewI
 	SystemInformation.bUseClosureCountFromMaterialData = UsesSubstrateClosureCountFromMaterialData() ? 1 : 0;
 	SystemInformation.TileClosureCount = GetSubstrateMaxClosureCount(View);
 	SystemInformation.bDbufferPass = IsDBufferPassEnabled(View.GetShaderPlatform()) ? 1 : 0;
-	SystemInformation.bRoughRefraction = IsOpaqueRoughRefractionEnabled() ? 1 : 0;
+	SystemInformation.bRoughRefraction = IsOpaqueRoughRefractionEnabled(View.GetShaderPlatform()) ? 1 : 0;
 	SystemInformation.TileSize = SUBSTRATE_TILE_SIZE;
 	SystemInformation.TileCountX = View.SubstrateViewData.TileCount.X;
 	SystemInformation.TileCountY = View.SubstrateViewData.TileCount.Y;
@@ -1140,44 +1142,56 @@ static void AddVisualizeSystemInfoPasses(FRDGBuilder& GraphBuilder, const FViewI
 			PrintS(DrawPosX, DrawPosY, TEXT("General"), Yellow);
 			NewLine();
 			{
-				PrintS(DrawPosX, DrawPosY, TEXT("  Bytes Per Pixel"), DarkYellow);
-				NewLine();
-				PrintS(DrawPosX +   0, DrawPosY, TEXT("    Max"), White);
-				PrintI(DrawPosX + 150, DrawPosY, SystemInformation.CurrentBytesPerPixel, Orange);
-				PrintS(DrawPosX + 180, DrawPosY, TEXT("/"), DarkYellow);
-				PrintI(DrawPosX + 190, DrawPosY, SystemInformation.ProjectMaxBytesPerPixel, Orange);
-				Text = FString::Printf(TEXT("[%i MB]"), SystemInformation.MaterialBufferAllocationInBytes/(1024*1024));
-				PrintS(DrawPosX + 220, DrawPosY, *Text, White);
-				NewLine();
-				PrintS(DrawPosX + 0, DrawPosY, TEXT("    Views"), White);
-				PrintI(DrawPosX + 150, DrawPosY, SystemInformation.ViewsMaxBytesPerPixel, Orange);
-				PrintS(DrawPosX + 180, DrawPosY, TEXT("/ "), DarkYellow);
-				PrintI(DrawPosX + 190, DrawPosY, SystemInformation.ProjectMaxBytesPerPixel, Orange);
-				NewLine();
-				NewLine();
+				if (IsSubstrateBlendableGBufferEnabled(View.GetShaderPlatform()))
+				{
+					PrintS(DrawPosX, DrawPosY, TEXT("  Using blendable GBuffer with fixed layout."), DarkYellow);
+					NewLine();
+					PrintS(DrawPosX, DrawPosY, TEXT("    - 4 x 32 bits textures."), White);
+					NewLine();
+					PrintS(DrawPosX, DrawPosY, TEXT("    - 1 closures per pixel."), White);
+				}
+				else
+				{
+					PrintS(DrawPosX, DrawPosY, TEXT("  Bytes Per Pixel"), DarkYellow);
+					NewLine();
+					PrintS(DrawPosX + 0, DrawPosY, TEXT("    Max"), White);
+					PrintI(DrawPosX + 150, DrawPosY, SystemInformation.CurrentBytesPerPixel, Orange);
+					PrintS(DrawPosX + 180, DrawPosY, TEXT("/"), DarkYellow);
+					PrintI(DrawPosX + 190, DrawPosY, SystemInformation.ProjectMaxBytesPerPixel, Orange);
+					Text = FString::Printf(TEXT("[%i MB]"), SystemInformation.MaterialBufferAllocationInBytes / (1024 * 1024));
+					PrintS(DrawPosX + 220, DrawPosY, *Text, White);
+					NewLine();
+					PrintS(DrawPosX + 0, DrawPosY, TEXT("    Views"), White);
+					PrintI(DrawPosX + 150, DrawPosY, SystemInformation.ViewsMaxBytesPerPixel, Orange);
+					PrintS(DrawPosX + 180, DrawPosY, TEXT("/ "), DarkYellow);
+					PrintI(DrawPosX + 190, DrawPosY, SystemInformation.ProjectMaxBytesPerPixel, Orange);
+					NewLine();
+					NewLine();
 
-				PrintS(DrawPosX, DrawPosY, TEXT("  Closures Per Pixel"), DarkYellow);
-				NewLine();
-				PrintS(DrawPosX + 0, DrawPosY, TEXT("    Max"), White);
-				PrintI(DrawPosX + 150, DrawPosY, SystemInformation.CurrentClosuresPerPixel, Orange);
-				PrintS(DrawPosX + 180, DrawPosY, TEXT("/ "), DarkYellow);
-				PrintI(DrawPosX + 190, DrawPosY, SystemInformation.ProjectMaxClosuresPerPixel, Orange);
-				NewLine();
-				PrintS(DrawPosX + 0, DrawPosY, TEXT("    Views"), White);
-				PrintI(DrawPosX + 150, DrawPosY, SystemInformation.ViewsMaxClosuresPerPixel, Orange);
-				PrintS(DrawPosX + 180, DrawPosY, TEXT("/ "), DarkYellow);
-				PrintI(DrawPosX + 190, DrawPosY, SystemInformation.ProjectMaxClosuresPerPixel, Orange);
-				NewLine();
-				NewLine();
+					PrintS(DrawPosX, DrawPosY, TEXT("  Closures Per Pixel"), DarkYellow);
+					NewLine();
+					PrintS(DrawPosX + 0, DrawPosY, TEXT("    Max"), White);
+					PrintI(DrawPosX + 150, DrawPosY, SystemInformation.CurrentClosuresPerPixel, Orange);
+					PrintS(DrawPosX + 180, DrawPosY, TEXT("/ "), DarkYellow);
+					PrintI(DrawPosX + 190, DrawPosY, SystemInformation.ProjectMaxClosuresPerPixel, Orange);
+					NewLine();
+					PrintS(DrawPosX + 0, DrawPosY, TEXT("    Views"), White);
+					PrintI(DrawPosX + 150, DrawPosY, SystemInformation.ViewsMaxClosuresPerPixel, Orange);
+					PrintS(DrawPosX + 180, DrawPosY, TEXT("/ "), DarkYellow);
+					PrintI(DrawPosX + 190, DrawPosY, SystemInformation.ProjectMaxClosuresPerPixel, Orange);
+					NewLine();
+					NewLine();
 
-				PrintS(DrawPosX, DrawPosY, TEXT("  Allocation mode"), White);
-				switch (SystemInformation.MaterialBufferAllocationMode)
-				{
-				case 0: PrintS(DrawPosX + 150, DrawPosY, TEXT("View based"), Orange); break;
-				case 1: PrintS(DrawPosX + 150, DrawPosY, TEXT("View based | Growing-only"), Orange); break;
-				case 2: PrintS(DrawPosX + 150, DrawPosY, TEXT("Setting based"), Orange); break;
-				default: PrintS(DrawPosX + 150, DrawPosY, TEXT("Unkown"), Orange); break;
+					PrintS(DrawPosX, DrawPosY, TEXT("  Allocation mode"), White);
+					switch (SystemInformation.MaterialBufferAllocationMode)
+					{
+					case 0: PrintS(DrawPosX + 150, DrawPosY, TEXT("View based"), Orange); break;
+					case 1: PrintS(DrawPosX + 150, DrawPosY, TEXT("View based | Growing-only"), Orange); break;
+					case 2: PrintS(DrawPosX + 150, DrawPosY, TEXT("Setting based"), Orange); break;
+					default: PrintS(DrawPosX + 150, DrawPosY, TEXT("Unkown"), Orange); break;
+					}
 				}
+
 				NewLine();
 				PrintS(DrawPosX, DrawPosY, TEXT("  Shading quality"), White);
 				PrintI(DrawPosX + 150, DrawPosY, SystemInformation.ShadingQuality, Orange);
@@ -1417,7 +1431,7 @@ FScreenPassTexture AddSubstrateDebugPasses(FRDGBuilder& GraphBuilder, const FVie
 		}
 		else if (DebugMode == FSubstrateViewMode::RoughRefractionClassification)
 		{
-			if (IsOpaqueRoughRefractionEnabled())
+			if (IsOpaqueRoughRefractionEnabled(View.GetShaderPlatform()))
 			{
 				AddSubstrateInternalClassificationTilePass(GraphBuilder, View, nullptr, &ScreenPassSceneColor.Texture, ESubstrateTileType::EOpaqueRoughRefraction, bDebugPass);
 				AddSubstrateInternalClassificationTilePass(GraphBuilder, View, nullptr, &ScreenPassSceneColor.Texture, ESubstrateTileType::EOpaqueRoughRefractionSSSWithout, bDebugPass);
