[DOF] Cache BokehLUT between frames.
This LUT, which is used for bokeh shape when applying aperture blades, usually does not change between frames and has a small memory footprint.

diff --git a/Engine/Source/Runtime/Renderer/Private/PostProcess/DiaphragmDOF.cpp b/Engine/Source/Runtime/Renderer/Private/PostProcess/DiaphragmDOF.cpp
--- a/Engine/Source/Runtime/Renderer/Private/PostProcess/DiaphragmDOF.cpp
+++ b/Engine/Source/Runtime/Renderer/Private/PostProcess/DiaphragmDOF.cpp
@@ -1462,6 +1462,22 @@ IMPLEMENT_GLOBAL_SHADER(FDiaphragmDOFRecombineCS,        "/Engine/Private/Diaphr
 
 } // namespace
 
+namespace DiaphragmDOF
+{
+	class FDOFViewState
+	{
+	public:
+		TStaticArray<TRefCountPtr<IPooledRenderTarget>, (int)EDiaphragmDOFBokehLUTFormat::MAX> BokehLUTs;
+		FBokehModel BokehModel;
+
+		FDOFViewState()
+			: BokehLUTs()
+			, BokehModel()
+		{
+		}
+	};
+}
+
 bool DiaphragmDOF::IsEnabled(const FViewInfo& View)
 {
 	static const auto CVar = IConsoleManager::Get().FindTConsoleVariableDataInt(TEXT("r.DepthOfFieldQuality"));
@@ -1573,7 +1589,32 @@ bool DiaphragmDOF::AddPasses(
 	CocModel.Compile(View);
 
 	FBokehModel BokehModel;
-	BokehModel.Compile(View);
+	BokehModel.Compile(View, CocModel);
+
+	FDOFViewState* DOFViewState;
+	TStaticArray<bool, (int)EDiaphragmDOFBokehLUTFormat::MAX> HasCachedBokehLUT {};
+	if (View.ViewState)
+	{
+		if (!View.ViewState->DepthOfFieldState.IsValid())
+		{
+			View.ViewState->DepthOfFieldState = MakePimpl<FDOFViewState>();
+		}
+		DOFViewState = View.ViewState->DepthOfFieldState.Get();
+
+		if (DOFViewState->BokehModel != BokehModel)
+		{
+			for (int Index = 0; Index < (int)EDiaphragmDOFBokehLUTFormat::MAX; ++Index)
+			{
+				DOFViewState->BokehLUTs[Index] = nullptr;
+			}
+			DOFViewState->BokehModel = BokehModel;
+		}
+
+		for (int Index = 0; Index < (int)EDiaphragmDOFBokehLUTFormat::MAX; ++Index)
+		{
+			HasCachedBokehLUT[Index] = DOFViewState->BokehLUTs[Index] != nullptr;
+		}
+	}
 
 	// Prepare preprocessing TAA pass.
 	FTAAPassParameters TAAParameters(View);
@@ -2308,6 +2349,12 @@ bool DiaphragmDOF::AddPasses(
 			FClearValueBinding::None,
 			TexCreate_ShaderResource | TexCreate_UAV);
 
+		// Retrieve cached BokehLUT from DOFViewState, if any
+		if (HasCachedBokehLUT[int32(LUTFormat)])
+		{
+			return GraphBuilder.RegisterExternalTexture(DOFViewState->BokehLUTs[int32(LUTFormat)], DebugNames[int32(LUTFormat)]);
+		}
+
 		FRDGTextureRef BokehLUT = GraphBuilder.CreateTexture(BokehLUTDesc, DebugNames[int32(LUTFormat)]);
 
 		FDiaphragmDOFBuildBokehLUTCS::FPermutationDomain PermutationVector;
@@ -2321,8 +2368,8 @@ bool DiaphragmDOF::AddPasses(
 		PassParameters->CocRadiusToIncircleRadius = BokehModel.CocRadiusToIncircleRadius;
 		PassParameters->DiaphragmBladeRadius = BokehModel.RoundedBlades.DiaphragmBladeRadius;
 		PassParameters->DiaphragmBladeCenterOffset = BokehModel.RoundedBlades.DiaphragmBladeCenterOffset;
-		PassParameters->CocSqueeze = CocModel.Squeeze;
-		PassParameters->CocInvSqueeze = 1.0f / CocModel.Squeeze;
+		PassParameters->CocSqueeze = BokehModel.Squeeze;
+		PassParameters->CocInvSqueeze = 1.0f / BokehModel.Squeeze;
 		PassParameters->BokehLUTOutput = GraphBuilder.CreateUAV(BokehLUT);
 			
 		TShaderMapRef<FDiaphragmDOFBuildBokehLUTCS> ComputeShader(View.ShaderMap, PermutationVector);
@@ -2337,6 +2384,11 @@ bool DiaphragmDOF::AddPasses(
 			PassParameters,
 			FComputeShaderUtils::GetGroupCount(BokehLUTDesc.Extent, kDefaultGroupSize));
 
+		if (DOFViewState)
+		{
+			GraphBuilder.QueueTextureExtraction(BokehLUT, &DOFViewState->BokehLUTs[int32(LUTFormat)]);
+		}
+
 		return BokehLUT;
 	};
 	
diff --git a/Engine/Source/Runtime/Renderer/Private/PostProcess/DiaphragmDOF.h b/Engine/Source/Runtime/Renderer/Private/PostProcess/DiaphragmDOF.h
--- a/Engine/Source/Runtime/Renderer/Private/PostProcess/DiaphragmDOF.h
+++ b/Engine/Source/Runtime/Renderer/Private/PostProcess/DiaphragmDOF.h
@@ -155,9 +155,27 @@ struct FBokehModel
 		float DiaphragmBladeCenterOffset;
 	} RoundedBlades;
 
+	float Squeeze;
 
 	/** Compile the model from a view. */
-	void Compile(const FViewInfo& View);
+	void Compile(const FViewInfo& View, const FPhysicalCocModel& CocModel);
+
+	bool operator==(const FBokehModel& Other) const
+	{
+		return BokehShape == Other.BokehShape
+			&& CocRadiusToCircumscribedRadius == Other.CocRadiusToCircumscribedRadius
+			&& CocRadiusToIncircleRadius == Other.CocRadiusToIncircleRadius
+			&& DiaphragmBladeCount == Other.DiaphragmBladeCount
+			&& DiaphragmRotation == Other.DiaphragmRotation
+			&& RoundedBlades.DiaphragmBladeRadius == Other.RoundedBlades.DiaphragmBladeRadius
+			&& RoundedBlades.DiaphragmBladeCenterOffset == Other.RoundedBlades.DiaphragmBladeCenterOffset
+			&& Squeeze == Other.Squeeze;
+	}
+
+	bool operator!=(const FBokehModel& Other) const
+	{
+		return !(*this == Other);
+	}
 };
 
 /** Returns whether DOF is supported. */
diff --git a/Engine/Source/Runtime/Renderer/Private/PostProcess/DiaphragmDOFUtils.cpp b/Engine/Source/Runtime/Renderer/Private/PostProcess/DiaphragmDOFUtils.cpp
--- a/Engine/Source/Runtime/Renderer/Private/PostProcess/DiaphragmDOFUtils.cpp
+++ b/Engine/Source/Runtime/Renderer/Private/PostProcess/DiaphragmDOFUtils.cpp
@@ -238,7 +238,7 @@ float DiaphragmDOF::FPhysicalCocModel::GetCocOffset(float CocRadius) const
 	return DynamicOffset;
 }
 
-void DiaphragmDOF::FBokehModel::Compile(const FViewInfo& View)
+void DiaphragmDOF::FBokehModel::Compile(const FViewInfo& View, const FPhysicalCocModel& CocModel)
 {
 	{
 		DiaphragmBladeCount = FMath::Clamp(View.FinalPostProcessSettings.DepthOfFieldBladeCount, 4, 16);
@@ -252,6 +252,8 @@ void DiaphragmDOF::FBokehModel::Compile(const FViewInfo& View)
 	// Target a constant bokeh area to be eenergy preservative.
 	const float TargetedBokehArea = PI * (CircumscribedRadius * CircumscribedRadius);
 
+	Squeeze = CocModel.Squeeze;
+
 	// Always uses circle if max aparture is smaller or equal to aperture. 
 	if (Fstop <= MinFstop)
 	{
diff --git a/Engine/Source/Runtime/Renderer/Private/ScenePrivate.h b/Engine/Source/Runtime/Renderer/Private/ScenePrivate.h
--- a/Engine/Source/Runtime/Renderer/Private/ScenePrivate.h
+++ b/Engine/Source/Runtime/Renderer/Private/ScenePrivate.h
@@ -100,6 +100,10 @@ class FSparseVolumeTextureViewerSceneProxy;
 class FExponentialHeightFogSceneInfo;
 class FStaticMeshBatch;
 enum class EVTInvalidatePriority : uint8;
+namespace DiaphragmDOF
+{
+	class FDOFViewState;
+}
 
 /** Holds information about a single primitive's occlusion. */
 class FPrimitiveOcclusionHistory
@@ -1063,6 +1067,8 @@ public:
 		TRefCountPtr<FRDGPooledBuffer> ConstantsBuffer;
 	} FilmGrainCache;
 
+	TPimplPtr<DiaphragmDOF::FDOFViewState> DepthOfFieldState;
+
 	// Cached material texture samplers
 	float MaterialTextureCachedMipBias;
 	float LandscapeCachedMipBias;
