[DOF] Batch passes by shader to better overlap execution of dispatches, avoiding stalls and somewhat mitigating long tails.

diff --git a/Engine/Source/Runtime/Renderer/Private/PostProcess/DiaphragmDOF.cpp b/Engine/Source/Runtime/Renderer/Private/PostProcess/DiaphragmDOF.cpp
--- a/Engine/Source/Runtime/Renderer/Private/PostProcess/DiaphragmDOF.cpp
+++ b/Engine/Source/Runtime/Renderer/Private/PostProcess/DiaphragmDOF.cpp
@@ -2797,6 +2797,20 @@ bool DiaphragmDOF::AddPasses(
 			});
 		}; // AddHybridScatterPass()
 
+		struct FDOFPassDescriptor
+		{
+			FConvolutionSettings ConvolutionSettings;
+			FRDGTextureRef BokehLUT;
+			FDOFConvolutionTextures* ConvolutionTextures;
+			FRDGBufferRef ScatterDrawList;
+			FRDGTextureRef ScatterOcclusionTexture = nullptr;
+
+			uint32 bDoGather : 1;
+			uint32 bDoPostFilter : 1;
+			uint32 bDoScatter : 1;
+		};
+		TArray<FDOFPassDescriptor, TInlineAllocator<4>> Passes;
+
 		// Wire foreground gathering passes.
 		if (bGatherForeground)
 		{
@@ -2814,12 +2828,16 @@ bool DiaphragmDOF::AddPasses(
 				ConvolutionSettings.QualityConfig = EDiaphragmDOFGatherQuality::LowQualityAccumulator;
 			}
 
-			FRDGTextureRef ScatterOcclusionTexture = nullptr;
-			AddGatherPass(ConvolutionSettings, GatheringBokehLUT, &ForegroundConvolutionTextures, &ScatterOcclusionTexture);
-			AddPostFilterPass(ConvolutionSettings, &ForegroundConvolutionTextures);
-
-			if (bForegroundHybridScattering)
-				AddHybridScatterPass(ConvolutionSettings, &ForegroundConvolutionTextures, ScatterOcclusionTexture, ForegroundScatterDrawListBuffer);
+			Passes.Emplace(FDOFPassDescriptor {
+				.ConvolutionSettings = ConvolutionSettings,
+				.BokehLUT = GatheringBokehLUT,
+				.ConvolutionTextures = &ForegroundConvolutionTextures,
+				.ScatterDrawList = ForegroundScatterDrawListBuffer,
+				
+				.bDoGather = true,
+				.bDoPostFilter = true,
+				.bDoScatter = bForegroundHybridScattering
+			});
 		}
 
 		// Wire hole filling gathering passes.
@@ -2830,8 +2848,15 @@ bool DiaphragmDOF::AddPasses(
 			ConvolutionSettings.PostfilterMethod = PostfilterMethod;
 			ConvolutionSettings.RingCount = HalfResRingCount;
 			
-			FRDGTextureRef ScatterOcclusionTexture = nullptr;
-			AddGatherPass(ConvolutionSettings, /* BokehLUT = */ nullptr, &ForegroundHoleFillingConvolutionTextures, &ScatterOcclusionTexture);
+			Passes.Emplace(FDOFPassDescriptor {
+				.ConvolutionSettings = ConvolutionSettings,
+				.BokehLUT = nullptr,
+				.ConvolutionTextures = &ForegroundHoleFillingConvolutionTextures,
+				
+				.bDoGather = true,
+				.bDoPostFilter = false,
+				.bDoScatter = false
+			});
 		}
 		
 		// Gather slight out of focus.
@@ -2845,11 +2870,15 @@ bool DiaphragmDOF::AddPasses(
 			// Number of rings is dynamic in shader.
 			ConvolutionSettings.RingCount = kMinGatheringRingCount;
 
-			FRDGTextureRef ScatterOcclusionTexture = nullptr;
-			AddGatherPass(
-				ConvolutionSettings,
-				/* BokehLUT = */ bEnableSlightOutOfFocusBokeh ? ScatteringBokehLUT : nullptr,
-				&SlightOutOfFocusConvolutionTextures, &ScatterOcclusionTexture);
+			Passes.Emplace(FDOFPassDescriptor {
+				.ConvolutionSettings = ConvolutionSettings,
+				.BokehLUT = bEnableSlightOutOfFocusBokeh ? ScatteringBokehLUT : nullptr,
+				.ConvolutionTextures = &SlightOutOfFocusConvolutionTextures,
+				
+				.bDoGather = true,
+				.bDoPostFilter = false,
+				.bDoScatter = false
+			});
 		}
 
 		// Wire background gathering passes.
@@ -2877,12 +2906,50 @@ bool DiaphragmDOF::AddPasses(
 				}
 			}
 			
-			FRDGTextureRef ScatterOcclusionTexture = nullptr;
-			AddGatherPass(ConvolutionSettings, GatheringBokehLUT, &BackgroundConvolutionTextures, &ScatterOcclusionTexture);
-			AddPostFilterPass(ConvolutionSettings, &BackgroundConvolutionTextures);
+			Passes.Emplace(FDOFPassDescriptor {
+				.ConvolutionSettings = ConvolutionSettings,
+				.BokehLUT = GatheringBokehLUT,
+				.ConvolutionTextures = &BackgroundConvolutionTextures,
+				.ScatterDrawList = BackgroundScatterDrawListBuffer,
+				
+				.bDoGather = true,
+				.bDoPostFilter = true,
+				.bDoScatter = bBackgroundHybridScattering
+			});
+		}
+
+		const bool bAvoidPSOSwitch = true;
+		if (bAvoidPSOSwitch)
+		{
+			for (FDOFPassDescriptor& Pass : Passes)
+			{
+				if (Pass.bDoGather)
+					AddGatherPass(Pass.ConvolutionSettings, Pass.BokehLUT, Pass.ConvolutionTextures, &Pass.ScatterOcclusionTexture);
+			}
+			for (FDOFPassDescriptor& Pass : Passes)
+			{
+				if (Pass.bDoPostFilter)
+					AddPostFilterPass(Pass.ConvolutionSettings, Pass.ConvolutionTextures);
+			}
+			for (FDOFPassDescriptor& Pass : Passes)
+			{
+				if (Pass.bDoScatter)
+					AddHybridScatterPass(Pass.ConvolutionSettings, Pass.ConvolutionTextures, Pass.ScatterOcclusionTexture, Pass.ScatterDrawList);
+			}
+		}
+		else
+		{
+			for (FDOFPassDescriptor& Pass : Passes)
+			{
+				if (Pass.bDoGather)
+					AddGatherPass(Pass.ConvolutionSettings, Pass.BokehLUT, Pass.ConvolutionTextures, &Pass.ScatterOcclusionTexture);
+
+				if (Pass.bDoPostFilter)
+					AddPostFilterPass(Pass.ConvolutionSettings, Pass.ConvolutionTextures);
 			
-			if (bBackgroundHybridScattering)
-				AddHybridScatterPass(ConvolutionSettings, &BackgroundConvolutionTextures, ScatterOcclusionTexture, BackgroundScatterDrawListBuffer);
+				if (Pass.bDoScatter)
+					AddHybridScatterPass(Pass.ConvolutionSettings, Pass.ConvolutionTextures, Pass.ScatterOcclusionTexture, Pass.ScatterDrawList);
+			}
 		}
 	}
 	
