Removed sky atmosphere async compute overlap with occlusion as it's no longer an option. Default instead to overlapping with the depth prepass when async compute is enabled.

diff --git a/Engine/Source/Runtime/Renderer/Private/DeferredShadingRenderer.cpp b/Engine/Source/Runtime/Renderer/Private/DeferredShadingRenderer.cpp
--- a/Engine/Source/Runtime/Renderer/Private/DeferredShadingRenderer.cpp
+++ b/Engine/Source/Runtime/Renderer/Private/DeferredShadingRenderer.cpp
@@ -2401,15 +2401,6 @@ void FDeferredShadingSceneRenderer::Render(FRDGBuilder& GraphBuilder, const FSce
 		// Run before RenderSkyAtmosphereLookUpTables for cloud shadows to be valid.
 		InitVolumetricCloudsForViews(GraphBuilder, bShouldRenderVolumetricCloudBase, InstanceCullingManager);
 
-		if (SkyAtmospherePassLocation == ESkyAtmospherePassLocation::BeforeOcclusion && bShouldRenderSkyAtmosphere)
-		{
-			// Generate the Sky/Atmosphere look up tables
-			RenderSkyAtmosphereLookUpTables(GraphBuilder, /* out */ SkyAtmospherePendingRDGResources);
-
-			SkyAtmospherePendingRDGResources.CommitToSceneAndViewUniformBuffers(GraphBuilder, /* out */ ExternalAccessQueue);
-			ExternalAccessQueue.Submit(GraphBuilder);
-		}
-
 		BeginAsyncDistanceFieldShadowProjections(GraphBuilder, SceneTextures, InitViewTaskDatas.DynamicShadows);
 
 		// Run local fog volume culling before base pass and after HZB generation to benefit from more culling.
@@ -2433,12 +2424,10 @@ void FDeferredShadingSceneRenderer::Render(FRDGBuilder& GraphBuilder, const FSce
 			RenderSkyAtmosphereLookUpTables(GraphBuilder, /* out */ SkyAtmospherePendingRDGResources);
 
 			SkyAtmospherePendingRDGResources.CommitToSceneAndViewUniformBuffers(GraphBuilder, /* out */ ExternalAccessQueue);
-			ExternalAccessQueue.Submit(GraphBuilder);
 		}
 		else if (SkyAtmospherePassLocation == ESkyAtmospherePassLocation::BeforePrePass && bShouldRenderSkyAtmosphere)
 		{
 			SkyAtmospherePendingRDGResources.CommitToSceneAndViewUniformBuffers(GraphBuilder, /* out */ ExternalAccessQueue);
-			ExternalAccessQueue.Submit(GraphBuilder);
 		}
 
 		// Capture the SkyLight using the SkyAtmosphere and VolumetricCloud component if available.
@@ -2446,6 +2435,9 @@ void FDeferredShadingSceneRenderer::Render(FRDGBuilder& GraphBuilder, const FSce
 		const bool bPathTracedAtmosphere = ViewFamily.EngineShowFlags.PathTracing && Views.Num() > 0 && PathTracing::UsesReferenceAtmosphere(Views[0]);
 		if (bRealTimeSkyCaptureEnabled && !bPathTracedAtmosphere)
 		{
+			// Sky capture accesses the view uniform buffer which uses LUT's.
+			ExternalAccessQueue.Submit(GraphBuilder);
+
 			FViewInfo& MainView = Views[0];
 			Scene->AllocateAndCaptureFrameSkyEnvMap(GraphBuilder, *this, MainView, bShouldRenderSkyAtmosphere, bShouldRenderVolumetricCloud, InstanceCullingManager, ExternalAccessQueue);
 		}
diff --git a/Engine/Source/Runtime/Renderer/Private/SkyAtmosphereRendering.cpp b/Engine/Source/Runtime/Renderer/Private/SkyAtmosphereRendering.cpp
--- a/Engine/Source/Runtime/Renderer/Private/SkyAtmosphereRendering.cpp
+++ b/Engine/Source/Runtime/Renderer/Private/SkyAtmosphereRendering.cpp
@@ -219,8 +219,8 @@ static TAutoConsoleVariable<int32> CVarSkyAtmosphereEditorNotifications(
 	TEXT("Enable the rendering of in editor notification to warn the user about missing sky dome pixels on screen. It is better to keep it enabled and will be removed when shipping.\n"),
 	ECVF_RenderThreadSafe | ECVF_Scalability);
 
-static TAutoConsoleVariable<int32> CVarSkyAtmosphereASyncCompute(
-	TEXT("r.SkyAtmosphereASyncCompute"), 0,
+static TAutoConsoleVariable<int32> CVarSkyAtmosphereAsyncCompute(
+	TEXT("r.SkyAtmosphereAsyncCompute"), 0,
 	TEXT("SkyAtmosphere on async compute (default: false). When running on the async pipe, SkyAtmosphere lut generation will overlap with the occlusion pass.\n"),
 	ECVF_RenderThreadSafe);
 
@@ -287,16 +287,11 @@ IMPLEMENT_GLOBAL_SHADER_PARAMETER_STRUCT(FSkyAtmosphereInternalCommonParameters,
 
 ESkyAtmospherePassLocation GetSkyAtmospherePassLocation()
 {
-	if (CVarSkyAtmosphereASyncCompute.GetValueOnAnyThread() == 1)
-	{
-		// When SkyAtmLUT is running on async compute, try to kick it before occlusion pass to have better wave occupancy
-		return ESkyAtmospherePassLocation::BeforeOcclusion;
-	}
-	else if (CVarSkyAtmosphereASyncCompute.GetValueOnAnyThread() == 2)
+	if (CVarSkyAtmosphereAsyncCompute.GetValueOnAnyThread() != 0)
 	{
 		return ESkyAtmospherePassLocation::BeforePrePass;
 	}
-	// When SkyAtmLUT is running on graphics pipe, kickbefore BasePass to have a better overlap when DistanceFieldShadows is running async
+	// When SkyAtmLUT is running on graphics pipe, kick before BasePass to have a better overlap when DistanceFieldShadows is running async
 	return ESkyAtmospherePassLocation::BeforeBasePass;
 }
 
@@ -1364,7 +1359,7 @@ void FSceneRenderer::RenderSkyAtmosphereLookUpTables(FRDGBuilder& GraphBuilder,
 	FRDGTextureRef MultiScatteredLuminanceLut = GraphBuilder.RegisterExternalTexture(SkyInfo.GetMultiScatteredLuminanceLutTexture());
 	FRDGTextureUAVRef MultiScatteredLuminanceLutUAV = GraphBuilder.CreateUAV(FRDGTextureUAVDesc(MultiScatteredLuminanceLut, 0));
 
-	ERDGPassFlags PassFlag = CVarSkyAtmosphereASyncCompute.GetValueOnAnyThread() ? ERDGPassFlags::AsyncCompute : ERDGPassFlags::Compute;
+	ERDGPassFlags PassFlag = CVarSkyAtmosphereAsyncCompute.GetValueOnAnyThread() ? ERDGPassFlags::AsyncCompute : ERDGPassFlags::Compute;
 
 	// Transmittance LUT
 	FGlobalShaderMap* GlobalShaderMap = GetGlobalShaderMap(FeatureLevel);
diff --git a/Engine/Source/Runtime/Renderer/Private/SkyAtmosphereRendering.h b/Engine/Source/Runtime/Renderer/Private/SkyAtmosphereRendering.h
--- a/Engine/Source/Runtime/Renderer/Private/SkyAtmosphereRendering.h
+++ b/Engine/Source/Runtime/Renderer/Private/SkyAtmosphereRendering.h
@@ -205,13 +205,10 @@ private:
 
 enum class ESkyAtmospherePassLocation : uint32
 {
-	// Renders just before the Occlusion queries. Good for wave occuppency when SkyAtmosphere is async
-	BeforeOcclusion,
-
 	// Renders just before the BasePass.
 	BeforeBasePass,
 
-	// Renders just before prepass but can overlap on async compute with it too.
+	// Renders just before the PrePass but can overlap on async compute with it too.
 	BeforePrePass,
 };
 
