Fixed crash when trying to render stamped probes on SM5 when the shader is not compiled.

#tests PC editor dx12 sm5 and sm6
#rnx

diff --git a/Engine/Source/Runtime/Renderer/Private/DebugProbeRendering.cpp b/Engine/Source/Runtime/Renderer/Private/DebugProbeRendering.cpp
--- a/Engine/Source/Runtime/Renderer/Private/DebugProbeRendering.cpp
+++ b/Engine/Source/Runtime/Renderer/Private/DebugProbeRendering.cpp
@@ -38,6 +38,11 @@ extern bool IsIlluminanceMeterSupportedByView(const FViewInfo& View);
 #define RENDER_BASEPASS	     1
 #define RENDER_VELOCITYPASS  2
 
+static bool IsDeferredDebugProbeSupportedByShaderPlatform(const FStaticShaderPlatform ShaderPlatform)
+{
+	// Only SM6 since SM5 is limited to 8 UAV.
+	return IsFeatureLevelSupported(ShaderPlatform, ERHIFeatureLevel::SM6);
+}
 
 class FStampDeferredDebugProbePS : public FGlobalShader
 {
@@ -67,8 +72,7 @@ public:
 
 	static bool ShouldCompilePermutation(const FGlobalShaderPermutationParameters& Parameters)
 	{
-		// Only SM6 since SM5 is limited to 8 UAV.
-		return IsFeatureLevelSupported(Parameters.Platform, ERHIFeatureLevel::SM6) && EnumHasAllFlags(Parameters.Flags, EShaderPermutationFlags::HasEditorOnlyData);
+		return IsDeferredDebugProbeSupportedByShaderPlatform(Parameters.Platform) && EnumHasAllFlags(Parameters.Flags, EShaderPermutationFlags::HasEditorOnlyData);
 	}
 	static void ModifyCompilationEnvironment(const FGlobalShaderPermutationParameters& Parameters, FShaderCompilerEnvironment& OutEnvironment)
 	{
@@ -140,7 +144,8 @@ void StampDeferredDebugProbeDepthPS(
 	for (const FViewInfo& View : Views)
 	{
 		const bool bViewRequiresAndSupportsIlluminanceMeter = ViewRequiresAndSupportsIlluminanceMeter(View);
-		if (!(bVisualizeLightingOnProbes || View.Family->EngineShowFlags.VisualizeLightingOnProbes) || View.bIsReflectionCapture || bViewRequiresAndSupportsIlluminanceMeter)
+		if (!(bVisualizeLightingOnProbes || View.Family->EngineShowFlags.VisualizeLightingOnProbes) || View.bIsReflectionCapture || bViewRequiresAndSupportsIlluminanceMeter
+			|| !IsDeferredDebugProbeSupportedByShaderPlatform(View.GetShaderPlatform()))
 		{
 			// When "Visualizing HDR with Illuminance Meter", we want to evaluated the illuminance on the surface behind the target square patch.  
 			// So we do not want write depth or velocity in this case.
@@ -174,7 +179,8 @@ void StampDeferredDebugProbeMaterialPS(
 	for (const FViewInfo& View : Views)
 	{
 		const bool bViewRequiresAndSupportsIlluminanceMeter = ViewRequiresAndSupportsIlluminanceMeter(View);
-		if (!(bVisualizeLightingOnProbes || View.Family->EngineShowFlags.VisualizeLightingOnProbes || bViewRequiresAndSupportsIlluminanceMeter) || View.bIsReflectionCapture)
+		if (!(bVisualizeLightingOnProbes || View.Family->EngineShowFlags.VisualizeLightingOnProbes || bViewRequiresAndSupportsIlluminanceMeter) || View.bIsReflectionCapture
+			|| !IsDeferredDebugProbeSupportedByShaderPlatform(View.GetShaderPlatform()))
 		{
 			continue;
 		}
@@ -239,7 +245,8 @@ void StampDeferredDebugProbeVelocityPS(
 	for (const FViewInfo& View : Views)
 	{
 		const bool bViewRequiresAndSupportsIlluminanceMeter = ViewRequiresAndSupportsIlluminanceMeter(View);
-		if (!(bVisualizeLightingOnProbes || View.Family->EngineShowFlags.VisualizeLightingOnProbes) || View.bIsReflectionCapture || bViewRequiresAndSupportsIlluminanceMeter)
+		if (!(bVisualizeLightingOnProbes || View.Family->EngineShowFlags.VisualizeLightingOnProbes) || View.bIsReflectionCapture || bViewRequiresAndSupportsIlluminanceMeter
+			|| !IsDeferredDebugProbeSupportedByShaderPlatform(View.GetShaderPlatform()))
 		{
 			// When "Visualizing HDR with Illuminance Meter", we want to evaluated the illuminance on the surface behind the target square patch. 
 			// So we do not want write depth or velocity in this case.
