[HWRT] Include LayerIndex in FRHIRayTracingScene DebugName.

diff --git a/Engine/Source/Runtime/Renderer/Private/RayTracing/RayTracingScene.cpp b/Engine/Source/Runtime/Renderer/Private/RayTracing/RayTracingScene.cpp
--- a/Engine/Source/Runtime/Renderer/Private/RayTracing/RayTracingScene.cpp
+++ b/Engine/Source/Runtime/Renderer/Private/RayTracing/RayTracingScene.cpp
@@ -87,6 +87,12 @@ FRayTracingScene::FRayTracingScene()
 {
 	const uint8 NumLayers = uint8(ERayTracingSceneLayer::NUM);
 	Layers.AddDefaulted(NumLayers);
+
+	for (uint32 LayerIndex = 0; LayerIndex < NumLayers; ++LayerIndex)
+	{
+		FLayer& Layer = Layers[LayerIndex];
+		Layer.Name = FName(FString::Printf(TEXT("RayTracingScene_Layer%u"), LayerIndex));
+	}
 }
 
 FRayTracingScene::~FRayTracingScene()
@@ -174,7 +180,7 @@ void FRayTracingScene::Create(FRDGBuilder& GraphBuilder, const FViewInfo& View,
 
 		{
 			FRayTracingSceneInitializer Initializer;
-			Initializer.DebugName = FName(TEXT("FRayTracingScene"));
+			Initializer.DebugName = Layer.Name;
 			Initializer.MaxNumInstances = NumNativeInstances;
 			Initializer.BuildFlags = LayerIndex == uint8(ERayTracingSceneLayer::FarField) ? FarFieldBuildFlags : BuildFlags;
 			Layer.RayTracingSceneRHI = RHICreateRayTracingScene(MoveTemp(Initializer));
diff --git a/Engine/Source/Runtime/Renderer/Private/RayTracing/RayTracingScene.h b/Engine/Source/Runtime/Renderer/Private/RayTracing/RayTracingScene.h
--- a/Engine/Source/Runtime/Renderer/Private/RayTracing/RayTracingScene.h
+++ b/Engine/Source/Runtime/Renderer/Private/RayTracing/RayTracingScene.h
@@ -212,6 +212,8 @@ private:
 
 		uint32 NumActiveInstances = 0;
 		uint32 MaxNumInstances = 0;
+
+		FName Name;
 	};
 
 	TArray<FLayer> Layers;
