[HWRT] Support serializing FRHIRayTracingScene to disk.

diff --git a/Engine/Source/Runtime/RHI/Public/DynamicRHI.h b/Engine/Source/Runtime/RHI/Public/DynamicRHI.h
--- a/Engine/Source/Runtime/RHI/Public/DynamicRHI.h
+++ b/Engine/Source/Runtime/RHI/Public/DynamicRHI.h
@@ -992,6 +992,13 @@ public:
 		return nullptr;
 	}
 
+#if !UE_BUILD_SHIPPING
+	virtual void RHISerializeAccelerationStructure(FRHICommandListImmediate& RHICmdList, FRHIRayTracingScene* Scene, const TCHAR* Path)
+	{
+		checkNoEntry();
+	}
+#endif
+
 	virtual FShaderBundleRHIRef RHICreateShaderBundle(const FShaderBundleCreateInfo& CreateInfo)
 	{
 		checkNoEntry();
diff --git a/Engine/Source/Runtime/RHI/Public/RHICommandList.h b/Engine/Source/Runtime/RHI/Public/RHICommandList.h
--- a/Engine/Source/Runtime/RHI/Public/RHICommandList.h
+++ b/Engine/Source/Runtime/RHI/Public/RHICommandList.h
@@ -5363,6 +5363,13 @@ public:
 	{
 		GDynamicRHI->RHIVirtualTextureSetFirstMipVisible(*this, Texture, FirstMip);
 	}
+
+#if !UE_BUILD_SHIPPING
+	FORCEINLINE void SerializeAccelerationStructure(FRHIRayTracingScene* Scene, const TCHAR* Path)
+	{
+		GDynamicRHI->RHISerializeAccelerationStructure(*this, Scene, Path);
+	}
+#endif
 	
 	FORCEINLINE void* GetNativeDevice()
 	{
diff --git a/Engine/Source/Runtime/RHI/Public/RHIGlobals.h b/Engine/Source/Runtime/RHI/Public/RHIGlobals.h
--- a/Engine/Source/Runtime/RHI/Public/RHIGlobals.h
+++ b/Engine/Source/Runtime/RHI/Public/RHIGlobals.h
@@ -458,6 +458,9 @@ struct FRHIGlobals
 		/** Wether an extra uniform buffer parameter is required when loose parameters are present, or if they are stored directly in the shader record. */
 		bool SupportsLooseParamsInShaderRecord = false;
 
+		/** Whether or not the RHI supports serializing acceleration structures to a file */
+		bool SupportsSerializeAccelerationStructure = false;
+
 		/** Required alignment for ray tracing acceleration structures. */
 		uint32 AccelerationStructureAlignment = 0;
 
diff --git a/Engine/Source/Runtime/RHI/Public/RHIValidation.h b/Engine/Source/Runtime/RHI/Public/RHIValidation.h
--- a/Engine/Source/Runtime/RHI/Public/RHIValidation.h
+++ b/Engine/Source/Runtime/RHI/Public/RHIValidation.h
@@ -1110,6 +1110,13 @@ public:
 	{
 		return RHI->RHICreateShaderBindingTable(RHICmdList, Initializer);
 	}
+
+#if !UE_BUILD_SHIPPING
+	void RHISerializeAccelerationStructure(FRHICommandListImmediate& RHICmdList, FRHIRayTracingScene* Scene, const TCHAR* Path)
+	{
+		return RHI->RHISerializeAccelerationStructure(RHICmdList, Scene, Path);
+	}
+#endif
 #endif // RHI_RAYTRACING
 
 	virtual FShaderBundleRHIRef RHICreateShaderBundle(const FShaderBundleCreateInfo& CreateInfo) override final
diff --git a/Engine/Source/Runtime/Renderer/Private/RayTracing/RayTracingScene.cpp b/Engine/Source/Runtime/Renderer/Private/RayTracing/RayTracingScene.cpp
--- a/Engine/Source/Runtime/Renderer/Private/RayTracing/RayTracingScene.cpp
+++ b/Engine/Source/Runtime/Renderer/Private/RayTracing/RayTracingScene.cpp
@@ -66,6 +66,17 @@ static TAutoConsoleVariable<int32> CVarRayTracingSceneCompactInstancesMargin(
 	ECVF_RenderThreadSafe
 );
 
+#if !UE_BUILD_SHIPPING
+
+static bool GRayTracingSerializeSceneNextFrame = false;
+
+static FAutoConsoleCommand RayTracingSerializeSceneCmd(
+	TEXT("r.RayTracing.Scene.SerializeOnce"),
+	TEXT("Serialize Ray Tracing Scene to disk."),
+	FConsoleCommandDelegate::CreateStatic([] { GRayTracingSerializeSceneNextFrame = true; }));
+
+#endif
+
 bool IsRayTracingFeedbackEnabled(const FSceneViewFamily& ViewFamily)
 {
 	// TODO: For now Feedback is limited to inline passes
@@ -351,6 +362,10 @@ BEGIN_SHADER_PARAMETER_STRUCT(FRayTracingSceneBuildPassParams, )
 	RDG_BUFFER_ACCESS(DynamicGeometryScratchBuffer, ERHIAccess::UAVCompute)
 END_SHADER_PARAMETER_STRUCT()
 
+BEGIN_SHADER_PARAMETER_STRUCT(FRayTracingSceneSerializePassParams, )
+	RDG_BUFFER_ACCESS(TLASBuffer, ERHIAccess::BVHRead)
+END_SHADER_PARAMETER_STRUCT()
+
 void FRayTracingScene::Build(FRDGBuilder& GraphBuilder, ERDGPassFlags ComputePassFlags, FRDGBufferRef DynamicGeometryScratchBuffer)
 {
 	const uint8 NumLayers = uint8(ERayTracingSceneLayer::NUM);
@@ -366,7 +381,7 @@ void FRayTracingScene::Build(FRDGBuilder& GraphBuilder, ERDGPassFlags ComputePas
 		PassParams->DynamicGeometryScratchBuffer = DynamicGeometryScratchBuffer; // TODO: Is this necessary?
 
 		GraphBuilder.AddPass(RDG_EVENT_NAME("RayTracingBuildScene"), PassParams, ComputePassFlags,
-			[PassParams, this, &Layer](FRHICommandList& RHICmdList)
+			[PassParams, &Layer](FRHICommandList& RHICmdList)
 			{
 				FRayTracingSceneBuildParams BuildParams;
 				BuildParams.Scene = Layer.RayTracingSceneRHI;
@@ -380,7 +395,29 @@ void FRayTracingScene::Build(FRDGBuilder& GraphBuilder, ERDGPassFlags ComputePas
 				RHICmdList.BindAccelerationStructureMemory(Layer.RayTracingSceneRHI, PassParams->TLASBuffer->GetRHI(), 0);
 				RHICmdList.BuildAccelerationStructure(BuildParams);
 			});
+
+#if !UE_BUILD_SHIPPING
+		if (GRayTracingSerializeSceneNextFrame && GRHIGlobals.RayTracing.SupportsSerializeAccelerationStructure)
+		{
+			FRayTracingSceneSerializePassParams* SerializePassParams = GraphBuilder.AllocParameters<FRayTracingSceneSerializePassParams>();
+			SerializePassParams->TLASBuffer = Layer.RayTracingSceneBufferRDG;
+
+			GraphBuilder.AddPass(RDG_EVENT_NAME("RayTracingSerializeScene"), SerializePassParams, ERDGPassFlags::Readback,
+				[SerializePassParams, &Layer](FRHICommandListImmediate& RHICmdList)
+				{
+					FString Filename = FString::Printf(TEXT("%s_(%s)"), *Layer.Name.ToString(), *FDateTime::Now().ToString(TEXT("%Y%m%d_%H%M%S")));
+					FString RootPath = FPaths::ScreenShotDir() + TEXT("BVH/"); // Save BVH dumps to ScreenShot directory
+					FString OutputFilename = RootPath + Filename + TEXT(".bvh");
+
+					RHICmdList.SerializeAccelerationStructure(Layer.RayTracingSceneRHI, *OutputFilename);
+				});
+		}
+#endif
 	}
+
+#if !UE_BUILD_SHIPPING
+	GRayTracingSerializeSceneNextFrame = false;
+#endif
 }
 
 struct FRayTracingProcessFeedbackCS : public FGlobalShader
