Fix issue where adaptive virtual texture could get into a state where one of it's allocated VTs reaches MaxLevel very quickly and continues to receive allocation requests.
It's possible that we then never fulfil the individual page requests because all allocatation requests were being dropped.
Then we would see low res results never being cleared.
This could be repro'd by going very close to a textured surface and flushing the VT.

diff --git a/Engine/Source/Runtime/Renderer/Private/VT/VirtualTextureSystem.cpp b/Engine/Source/Runtime/Renderer/Private/VT/VirtualTextureSystem.cpp
--- a/Engine/Source/Runtime/Renderer/Private/VT/VirtualTextureSystem.cpp
+++ b/Engine/Source/Runtime/Renderer/Private/VT/VirtualTextureSystem.cpp
@@ -1574,11 +1574,11 @@ void FVirtualTextureSystem::GatherRequestsTask(const FGatherRequestsParameters&
 		if (AdaptiveVT != nullptr && vLevelPlusOne <= 1)
 		{
 			uint32 AdaptiveAllocationRequest = AdaptiveVT->GetPackedAllocationRequest(vAddress, vLevelPlusOne, Frame);
-			if (AdaptiveAllocationRequest != 0)
+			if (AdaptiveAllocationRequest == 0)
 			{
-				RequestList->AddAdaptiveAllocationRequest(AdaptiveAllocationRequest);
 				continue;
 			}
+			RequestList->AddAdaptiveAllocationRequest(AdaptiveAllocationRequest);
 		}
 
 		uint32 PageTableLayersToLoad[VIRTUALTEXTURE_SPACE_MAXLAYERS] = { 0 };
