Path Tracer: The cvar r.RayTracing.ExcludeSky was being applied incorrectly in the path tracing case

A recent refactor changed the behavior and made this cvar apply unconditionally instead of being limited to the ray tracing case.

diff --git a/Engine/Source/Runtime/Renderer/Private/RayTracing/RayTracing.cpp b/Engine/Source/Runtime/Renderer/Private/RayTracing/RayTracing.cpp
--- a/Engine/Source/Runtime/Renderer/Private/RayTracing/RayTracing.cpp
+++ b/Engine/Source/Runtime/Renderer/Private/RayTracing/RayTracing.cpp
@@ -186,6 +186,8 @@ namespace RayTracing
 		{
 			bTranslucentGeometry = false;
 		}
+
+		bIncludeSky = GRayTracingExcludeSky == 0 || ViewFamily.EngineShowFlags.PathTracing;
 	}
 
 	struct FRelevantPrimitive
@@ -1456,7 +1458,7 @@ namespace RayTracing
 				continue;
 			}
 
-			if (GRayTracingExcludeSky && RelevantPrimitive.CachedMeshCommandFlags.bIsSky)
+			if (!SceneOptions.bIncludeSky && RelevantPrimitive.CachedMeshCommandFlags.bIsSky)
 			{
 				continue;
 			}
diff --git a/Engine/Source/Runtime/Renderer/Private/RayTracing/RayTracing.h b/Engine/Source/Runtime/Renderer/Private/RayTracing/RayTracing.h
--- a/Engine/Source/Runtime/Renderer/Private/RayTracing/RayTracing.h
+++ b/Engine/Source/Runtime/Renderer/Private/RayTracing/RayTracing.h
@@ -25,6 +25,7 @@ namespace RayTracing
 	struct FSceneOptions
 	{
 		bool bTranslucentGeometry = true;
+		bool bIncludeSky = true;
 
 		FSceneOptions(
 			const FScene& Scene,
