PR #13016: [VT] Fix PoolAutoGrow when SplitPhysicalPoolSize is nonzero

diff --git a/Engine/Source/Runtime/Renderer/Private/VT/VirtualTextureSystem.cpp b/Engine/Source/Runtime/Renderer/Private/VT/VirtualTextureSystem.cpp
--- a/Engine/Source/Runtime/Renderer/Private/VT/VirtualTextureSystem.cpp
+++ b/Engine/Source/Runtime/Renderer/Private/VT/VirtualTextureSystem.cpp
@@ -2002,8 +2002,9 @@ void FVirtualTextureSystem::GrowPhysicalPools() const
 			// Increase pool by 1 tile or 4MB, whichever is greater.
 			const int32 TileCount = PhysicalSpace->GetSizeInTiles();
 			const int32 TileSizeInBytes = PhysicalSpace->GetTileSizeInBytes();
-			const int32 CurrentSizeInBytes = TileCount * TileCount * TileSizeInBytes;
-			const int32 NextSizeInBytes = (TileCount + 1) * (TileCount + 1) * TileSizeInBytes;
+			const int32 PhysicalPoolCount = PhysicalSpace->GetDescriptionExt().PoolCount;
+			const int32 CurrentSizeInBytes = TileCount * TileCount * TileSizeInBytes * PhysicalPoolCount;
+			const int32 NextSizeInBytes = (TileCount + 1) * (TileCount + 1) * TileSizeInBytes * PhysicalPoolCount;
 			const int32 MinIncreaseInBytes = 4 * 1024 * 1024;
 			const int32 ClampedNextSizeInBytes = FMath::Max(NextSizeInBytes, CurrentSizeInBytes + MinIncreaseInBytes);
 			Config.SizeInMegabyte = FMath::DivideAndRoundUp(ClampedNextSizeInBytes, 1024 * 1024);
