Add support for "r.vt.rvt.mipcolor 2" which will only highlight RVT mip 0 usage.
Can be useful for debugging mip 0 sampling performance issues.
Also force render context update for that cvar so that we don't need to r.vt.flush after setting it.

diff --git a/Engine/Source/Runtime/Renderer/Private/VT/RuntimeVirtualTextureRender.cpp b/Engine/Source/Runtime/Renderer/Private/VT/RuntimeVirtualTextureRender.cpp
--- a/Engine/Source/Runtime/Renderer/Private/VT/RuntimeVirtualTextureRender.cpp
+++ b/Engine/Source/Runtime/Renderer/Private/VT/RuntimeVirtualTextureRender.cpp
@@ -2,6 +2,7 @@
 
 #include "VT/RuntimeVirtualTextureRender.h"
 
+#include "ComponentRecreateRenderStateContext.h"
 #include "Components/RuntimeVirtualTextureComponent.h"
 #include "DataDrivenShaderPlatformInfo.h"
 #include "EngineModule.h"
@@ -32,10 +33,11 @@ DECLARE_DWORD_COUNTER_STAT(TEXT("Num pages rendered"), STAT_RenderedPages, STATG
 
 namespace RuntimeVirtualTexture
 {
-	static TAutoConsoleVariable<bool> CVarVTMipColors(
+	static TAutoConsoleVariable<int32> CVarVTMipColors(
 		TEXT("r.VT.RVT.MipColors"),
-		false,
-		TEXT("Render mip colors to RVT BaseColor."),
+		0,
+		TEXT("Render mip colors to RVT BaseColor.\n 0 off. 1 shows all mip colors. 2 shows only mip 0."),
+		FConsoleVariableDelegate::CreateLambda([](IConsoleVariable* InVariable) { FGlobalComponentRecreateRenderStateContext Context; }),
 		ECVF_RenderThreadSafe);
 
 	static TAutoConsoleVariable<int32> CVarVTHighQualityPerPixelHeight(
@@ -1472,14 +1474,22 @@ namespace RuntimeVirtualTexture
 		return Scene != nullptr && Scene->GetRenderScene() != nullptr && Scene->GetRenderScene()->GPUScene.IsRendering();
 	}
 
-	FLinearColor GetMipLevelColor(uint32 InLevel)
+	FLinearColor GetDebugMipLevelColor(uint32 InLevel, FLinearColor InDefaultColor)
 	{
-		static const uint32 MipColors[] = {
-			0xC0FFFFFF, 0xC0FFFF00, 0xC000FFFF, 0xC000FF00, 0xC0FF00FF, 0xC0FF0000, 0xC00000FF,
-			0xC0808080, 0xC0808000, 0xC0008080, 0xC0008000, 0xC0800080, 0xC0800000, 0xC0000080 };
+		const int32 MipColorMode = CVarVTMipColors.GetValueOnRenderThread();
+		const bool bShowMipColor = MipColorMode == 1 || (MipColorMode == 2 && InLevel == 0);
 
-		InLevel = FMath::Min<uint32>(InLevel, sizeof(MipColors) / sizeof(MipColors[0]) - 1);
-		return FLinearColor(FColor(MipColors[InLevel]));
+		if (bShowMipColor)
+		{
+			static const uint32 MipColors[] = {
+				0xC0FFFFFF, 0xC0FFFF00, 0xC000FFFF, 0xC000FF00, 0xC0FF00FF, 0xC0FF0000, 0xC00000FF,
+				0xC0808080, 0xC0808000, 0xC0008080, 0xC0008000, 0xC0800080, 0xC0800000, 0xC0000080 };
+
+			InLevel = FMath::Min<uint32>(InLevel, sizeof(MipColors) / sizeof(MipColors[0]) - 1);
+			return FLinearColor(FColor(MipColors[InLevel]));
+		}
+
+		return InDefaultColor;
 	}
 
 	void RenderPage(
@@ -1594,7 +1604,7 @@ namespace RuntimeVirtualTexture
 			FRuntimeVirtualTexturePassParameters* RuntimeVirtualTexturePassParameters = GraphBuilder.AllocParameters<FRuntimeVirtualTexturePassParameters>();
 			RuntimeVirtualTexturePassParameters->MipLevel = MipLevelParameter;
 			RuntimeVirtualTexturePassParameters->CustomMaterialData = CustomMaterialData;
-			RuntimeVirtualTexturePassParameters->DebugParams = CVarVTMipColors.GetValueOnRenderThread() ? GetMipLevelColor(vLevel) : FixedColor;
+			RuntimeVirtualTexturePassParameters->DebugParams = GetDebugMipLevelColor(vLevel, FixedColor);
 			RuntimeVirtualTexturePassParameters->PackHeight = FVector2f(WorldHeightPackParameter);	// LWC_TODO: Precision loss
  			PassParameters->RuntimeVirtualTexturePassParameters = GraphBuilder.CreateUniformBuffer(RuntimeVirtualTexturePassParameters);
 
