Fixed material cache shader type constructor

diff --git a/Engine/Source/Runtime/Renderer/Internal/MaterialCache/MaterialCacheShaders.h b/Engine/Source/Runtime/Renderer/Internal/MaterialCache/MaterialCacheShaders.h
--- a/Engine/Source/Runtime/Renderer/Internal/MaterialCache/MaterialCacheShaders.h
+++ b/Engine/Source/Runtime/Renderer/Internal/MaterialCache/MaterialCacheShaders.h
@@ -23,7 +23,14 @@ class FMaterialCacheUnwrapVS : public FMaterialCacheUnwrapVSBase
 {
 	DECLARE_SHADER_TYPE(FMaterialCacheUnwrapVS, MeshMaterial);
 
-public:
+public:	
+	FMaterialCacheUnwrapVS(const ShaderMetaType::CompiledShaderInitializerType& Initializer) : FMaterialCacheUnwrapVSBase(Initializer)
+	{
+		
+	}
+	
+	FMaterialCacheUnwrapVS() = default;
+	
 	static bool ShouldCompilePermutation(const FMeshMaterialShaderPermutationParameters& Parameters);
 	static void ModifyCompilationEnvironment(const FMaterialShaderPermutationParameters& Parameters, FShaderCompilerEnvironment& OutEnvironment);
 };
