Suppress SA-ownerless warning

diff --git a/Engine/Source/Runtime/Renderer/Private/MaterialCache/MaterialCacheRenderer.cpp b/Engine/Source/Runtime/Renderer/Private/MaterialCache/MaterialCacheRenderer.cpp
--- a/Engine/Source/Runtime/Renderer/Private/MaterialCache/MaterialCacheRenderer.cpp
+++ b/Engine/Source/Runtime/Renderer/Private/MaterialCache/MaterialCacheRenderer.cpp
@@ -25,7 +25,7 @@
 
 static void MaterialCacheInvalidateRenderStates(IConsoleVariable*)
 {
-	FGlobalComponentRecreateRenderStateContext{};
+	FGlobalComponentRecreateRenderStateContext{}; //-V607
 }
 
 bool GMaterialCacheStaticMeshEnableViewportFromVS = true;
