Use magenta instead of white for debug mode r.vt.rvt.mipcolor=2.
It's much easier to see for most content.
#rnx

diff --git a/Engine/Source/Runtime/Renderer/Private/VT/RuntimeVirtualTextureRender.cpp b/Engine/Source/Runtime/Renderer/Private/VT/RuntimeVirtualTextureRender.cpp
--- a/Engine/Source/Runtime/Renderer/Private/VT/RuntimeVirtualTextureRender.cpp
+++ b/Engine/Source/Runtime/Renderer/Private/VT/RuntimeVirtualTextureRender.cpp
@@ -1477,9 +1477,7 @@ namespace RuntimeVirtualTexture
 	FLinearColor GetDebugMipLevelColor(uint32 InLevel, FLinearColor InDefaultColor)
 	{
 		const int32 MipColorMode = CVarVTMipColors.GetValueOnRenderThread();
-		const bool bShowMipColor = MipColorMode == 1 || (MipColorMode == 2 && InLevel == 0);
-
-		if (bShowMipColor)
+		if (MipColorMode == 1)
 		{
 			static const uint32 MipColors[] = {
 				0xC0FFFFFF, 0xC0FFFF00, 0xC000FFFF, 0xC000FF00, 0xC0FF00FF, 0xC0FF0000, 0xC00000FF,
@@ -1488,6 +1486,10 @@ namespace RuntimeVirtualTexture
 			InLevel = FMath::Min<uint32>(InLevel, sizeof(MipColors) / sizeof(MipColors[0]) - 1);
 			return FLinearColor(FColor(MipColors[InLevel]));
 		}
+		else if (MipColorMode == 2 && InLevel == 0)
+		{
+			return FLinearColor(1.f, 0.f, 1.f, 0.5f);
+		}
 
 		return InDefaultColor;
 	}
