Setup dummy fog uniform parameters in real time sky capture if height fog is hidden for real time sky capture.

diff --git a/Engine/Source/Runtime/Renderer/Private/BasePassRendering.cpp b/Engine/Source/Runtime/Renderer/Private/BasePassRendering.cpp
--- a/Engine/Source/Runtime/Renderer/Private/BasePassRendering.cpp
+++ b/Engine/Source/Runtime/Renderer/Private/BasePassRendering.cpp
@@ -763,7 +763,8 @@ void SetupSharedBasePassParameters(
 	const FViewInfo& View,
 	const int32 ViewIndex,
 	bool bLumenGIEnabled,
-	FSharedBasePassUniformParameters& SharedParameters)
+	FSharedBasePassUniformParameters& SharedParameters,
+	bool bForRealtimeSkyCapture)
 {
 	if (View.ForwardLightingResources.ForwardLightUniformParameters)
 	{
@@ -774,13 +775,13 @@ void SetupSharedBasePassParameters(
 		SetupDummyForwardLightUniformParameters(GraphBuilder, SharedParameters.Forward, View.GetShaderPlatform());
 	}
 
-	SetupFogUniformParameters(GraphBuilder, View, SharedParameters.Fog);
+	SetupFogUniformParameters(GraphBuilder, View, SharedParameters.Fog, bForRealtimeSkyCapture);
 
 	if (View.IsInstancedStereoPass())
 	{
 		const FViewInfo& InstancedView = *View.GetInstancedView();
 		SharedParameters.ForwardISR = *InstancedView.ForwardLightingResources.ForwardLightUniformParameters;
-		SetupFogUniformParameters(GraphBuilder, (FViewInfo&)InstancedView, SharedParameters.FogISR);
+		SetupFogUniformParameters(GraphBuilder, (FViewInfo&)InstancedView, SharedParameters.FogISR, bForRealtimeSkyCapture);
 	}
 	else
 	{
@@ -816,10 +817,11 @@ TRDGUniformBufferRef<FOpaqueBasePassUniformParameters> CreateOpaqueBasePassUnifo
 	const int32 ViewIndex,
 	const FForwardBasePassTextures& ForwardBasePassTextures,
 	const FDBufferTextures& DBufferTextures,
-	bool bLumenGIEnabled)
+	bool bLumenGIEnabled,
+	bool bForRealtimeSkyCapture)
 {
 	FOpaqueBasePassUniformParameters& BasePassParameters = *GraphBuilder.AllocParameters<FOpaqueBasePassUniformParameters>();
-	SetupSharedBasePassParameters(GraphBuilder, View, ViewIndex, bLumenGIEnabled, BasePassParameters.Shared);
+	SetupSharedBasePassParameters(GraphBuilder, View, ViewIndex, bLumenGIEnabled, BasePassParameters.Shared, bForRealtimeSkyCapture);
 
 	const FRDGSystemTextures& SystemTextures = FRDGSystemTextures::Get(GraphBuilder);
 
diff --git a/Engine/Source/Runtime/Renderer/Private/BasePassRendering.h b/Engine/Source/Runtime/Renderer/Private/BasePassRendering.h
--- a/Engine/Source/Runtime/Renderer/Private/BasePassRendering.h
+++ b/Engine/Source/Runtime/Renderer/Private/BasePassRendering.h
@@ -151,7 +151,8 @@ extern void SetupSharedBasePassParameters(
 	const FViewInfo& View,
 	const int32 ViewIndex,
 	bool bLumenGIEnabled,
-	FSharedBasePassUniformParameters& BasePassParameters);
+	FSharedBasePassUniformParameters& BasePassParameters,
+	bool bForRealtimeSkyCapture = false);
 
 extern TRDGUniformBufferRef<FOpaqueBasePassUniformParameters> CreateOpaqueBasePassUniformBuffer(
 	FRDGBuilder& GraphBuilder,
@@ -159,7 +160,8 @@ extern TRDGUniformBufferRef<FOpaqueBasePassUniformParameters> CreateOpaqueBasePa
 	const int32 ViewIndex = 0,
 	const FForwardBasePassTextures& ForwardBasePassTextures = {},
 	const FDBufferTextures& DBufferTextures = {},
-	bool bLumenGIEnabled = false);
+	bool bLumenGIEnabled = false,
+	bool bForRealtimeSkyCapture = false);
 
 extern TRDGUniformBufferRef<FTranslucentBasePassUniformParameters> CreateTranslucentBasePassUniformBuffer(
 	FRDGBuilder& GraphBuilder,
diff --git a/Engine/Source/Runtime/Renderer/Private/FogRendering.cpp b/Engine/Source/Runtime/Renderer/Private/FogRendering.cpp
--- a/Engine/Source/Runtime/Renderer/Private/FogRendering.cpp
+++ b/Engine/Source/Runtime/Renderer/Private/FogRendering.cpp
@@ -51,9 +51,25 @@ static TAutoConsoleVariable<bool> CVarOrthoFogHeightAdjustment(
 
 IMPLEMENT_GLOBAL_SHADER_PARAMETER_STRUCT(FFogUniformParameters, "FogStruct");
 
-void SetupFogUniformParameters(FRDGBuilder& GraphBuilder, const FViewInfo& View, FFogUniformParameters& OutParameters)
+static void SetupDummyExponentialHeightFogUniformParameters(FFogUniformParameters& OutParameters)
+{
+	FMemory::Memzero(OutParameters);
+	OutParameters.InscatteringLightDirection.W = -1;
+	OutParameters.DirectionalInscatteringColor.W = -1;
+	OutParameters.FogInscatteringColorCubemap = GBlackTextureCube->TextureRHI;
+	OutParameters.FogInscatteringColorSampler = TStaticSamplerState<SF_Trilinear, AM_Clamp, AM_Clamp, AM_Clamp>::GetRHI();
+}
+
+void SetupFogUniformParameters(FRDGBuilder& GraphBuilder, const FViewInfo& View, FFogUniformParameters& OutParameters, bool bForRealtimeSkyCapture)
 {
 	// Exponential Height Fog
+	const FScene* Scene = View.Family->Scene ? View.Family->Scene->GetRenderScene() : nullptr;
+	if (bForRealtimeSkyCapture && Scene && Scene->HasAnyExponentialHeightFog() && !Scene->ExponentialFogs[0].bVisibleInRealTimeSkyCaptures)
+	{
+		// If Height Fog is hidden in Real Time Sky capture then setup the FogData with dummy values.
+		SetupDummyExponentialHeightFogUniformParameters(OutParameters);
+	}
+	else
 	{
 		const FTexture* Cubemap = GWhiteTextureCube;
 
diff --git a/Engine/Source/Runtime/Renderer/Private/FogRendering.h b/Engine/Source/Runtime/Renderer/Private/FogRendering.h
--- a/Engine/Source/Runtime/Renderer/Private/FogRendering.h
+++ b/Engine/Source/Runtime/Renderer/Private/FogRendering.h
@@ -35,7 +35,7 @@ BEGIN_GLOBAL_SHADER_PARAMETER_STRUCT(FFogUniformParameters,)
 	SHADER_PARAMETER_SAMPLER(SamplerState, IntegratedLightScatteringSampler)
 END_GLOBAL_SHADER_PARAMETER_STRUCT()
 
-extern void SetupFogUniformParameters(FRDGBuilder& GraphBuilder, const FViewInfo& View, FFogUniformParameters& OutParameters);
+extern void SetupFogUniformParameters(FRDGBuilder& GraphBuilder, const FViewInfo& View, FFogUniformParameters& OutParameters, bool bForRealtimeSkyCapture = false);
 TRDGUniformBufferRef<FFogUniformParameters> CreateFogUniformBuffer(FRDGBuilder& GraphBuilder, const FViewInfo& View);
 
 extern bool ShouldRenderFog(const FSceneViewFamily& Family);
diff --git a/Engine/Source/Runtime/Renderer/Private/MobileBasePassRendering.cpp b/Engine/Source/Runtime/Renderer/Private/MobileBasePassRendering.cpp
--- a/Engine/Source/Runtime/Renderer/Private/MobileBasePassRendering.cpp
+++ b/Engine/Source/Runtime/Renderer/Private/MobileBasePassRendering.cpp
@@ -263,14 +263,15 @@ void SetupMobileBasePassUniformParameters(
 	EMobileBasePass BasePass,
 	EMobileSceneTextureSetupMode SetupMode,
 	const FMobileBasePassTextures& MobileBasePassTextures,
-	FMobileBasePassUniformParameters& BasePassParameters)
+	FMobileBasePassUniformParameters& BasePassParameters,
+	bool bForRealtimeSkyCapture)
 {
 	const FViewInfo* InstancedView = View.GetInstancedView();
 
-	SetupFogUniformParameters(GraphBuilder, View, BasePassParameters.Fog);
+	SetupFogUniformParameters(GraphBuilder, View, BasePassParameters.Fog, bForRealtimeSkyCapture);
 	if (InstancedView && (View.bIsMobileMultiViewEnabled || View.Aspects.IsMobileMultiViewEnabled()))
 	{
-		SetupFogUniformParameters(GraphBuilder, *InstancedView, BasePassParameters.FogMMV);
+		SetupFogUniformParameters(GraphBuilder, *InstancedView, BasePassParameters.FogMMV, bForRealtimeSkyCapture);
 	}
 	else
 	{
@@ -382,10 +383,11 @@ TRDGUniformBufferRef<FMobileBasePassUniformParameters> CreateMobileBasePassUnifo
 	const FViewInfo& View,
 	EMobileBasePass BasePass,
 	EMobileSceneTextureSetupMode SetupMode,
-	const FMobileBasePassTextures& MobileBasePassTextures)
+	const FMobileBasePassTextures& MobileBasePassTextures,
+	bool bForRealtimeSkyCapture)
 {
 	FMobileBasePassUniformParameters* BasePassParameters = GraphBuilder.AllocParameters<FMobileBasePassUniformParameters>();
-	SetupMobileBasePassUniformParameters(GraphBuilder, View, BasePass, SetupMode, MobileBasePassTextures, *BasePassParameters);
+	SetupMobileBasePassUniformParameters(GraphBuilder, View, BasePass, SetupMode, MobileBasePassTextures, *BasePassParameters, bForRealtimeSkyCapture);
 #if WITH_DEBUG_VIEW_MODES
 	if (View.Family->UseDebugViewPS())
 	{
diff --git a/Engine/Source/Runtime/Renderer/Private/MobileBasePassRendering.h b/Engine/Source/Runtime/Renderer/Private/MobileBasePassRendering.h
--- a/Engine/Source/Runtime/Renderer/Private/MobileBasePassRendering.h
+++ b/Engine/Source/Runtime/Renderer/Private/MobileBasePassRendering.h
@@ -80,14 +80,16 @@ extern void SetupMobileBasePassUniformParameters(
 	const FViewInfo& View,
 	EMobileBasePass BasePass,
 	const FMobileBasePassTextures& MobileBasePassTextures,
-	FMobileBasePassUniformParameters& BasePassParameters);
+	FMobileBasePassUniformParameters& BasePassParameters,
+	bool bForRealtimeSkyCapture);
 
 extern TRDGUniformBufferRef<FMobileBasePassUniformParameters> CreateMobileBasePassUniformBuffer(
 	FRDGBuilder& GraphBuilder,
 	const FViewInfo& View,
 	EMobileBasePass BasePass,
 	EMobileSceneTextureSetupMode SetupMode,
-	const FMobileBasePassTextures& MobileBasePassTextures = {});
+	const FMobileBasePassTextures& MobileBasePassTextures = {},
+	bool bForRealtimeSkyCapture = false);
 
 extern void SetupMobileDirectionalLightUniformParameters(
 	const FScene& Scene,
diff --git a/Engine/Source/Runtime/Renderer/Private/ReflectionEnvironmentRealTimeCapture.cpp b/Engine/Source/Runtime/Renderer/Private/ReflectionEnvironmentRealTimeCapture.cpp
--- a/Engine/Source/Runtime/Renderer/Private/ReflectionEnvironmentRealTimeCapture.cpp
+++ b/Engine/Source/Runtime/Renderer/Private/ReflectionEnvironmentRealTimeCapture.cpp
@@ -650,7 +650,7 @@ void FScene::AllocateAndCaptureFrameSkyEnvMap(
 								PassParameters->View = CubeView.GetShaderParameters();
 								PassParameters->RenderTargets = SkyRC.RenderTargets;
 								const EMobileBasePass BasePass = EMobileBasePass::DepthPrePass;
-								PassParameters->BasePass = CreateMobileBasePassUniformBuffer(GraphBuilder, MainView, BasePass, EMobileSceneTextureSetupMode::SceneDepth);
+								PassParameters->BasePass = CreateMobileBasePassUniformBuffer(GraphBuilder, MainView, BasePass, EMobileSceneTextureSetupMode::SceneDepth, {}, true);
 					
 								// Setup the depth buffer
 								if (bUseDepthBuffer)
@@ -721,7 +721,7 @@ void FScene::AllocateAndCaptureFrameSkyEnvMap(
 								auto* PassParameters = GraphBuilder.AllocParameters<FCaptureSkyMeshReflectionPassParameters>();
 								PassParameters->View = CubeView.GetShaderParameters();
 								PassParameters->RenderTargets = SkyRC.RenderTargets;
-								PassParameters->BasePass = CreateOpaqueBasePassUniformBuffer(GraphBuilder, MainView, 0);
+								PassParameters->BasePass = CreateOpaqueBasePassUniformBuffer(GraphBuilder, MainView, 0, {}, {}, false, true);
 					
 								// Setup the depth buffer
 								if (bUseDepthBuffer)
