Path Tracer: Implement multiple scattering approximation for clouds

The settings provided by the advanced volumetric cloud node in the material graph are now respected in the path tracer. This mode is enabled by default as it improves the alignment between raster and path tracer in most cases.

Add a cvar to control how multiple scattering on cloud is done. Possible values include (0 -- none), (1 -- default, use the shader settings), or (2 -- reference, ignore shader settings and force CloudRoughnessCutoff=1 to simulate the ground truth (which is generally too expensive to be practical).

diff --git a/Engine/Shaders/Private/PathTracing/PathTracingBuildCloudAccelerationMap.usf b/Engine/Shaders/Private/PathTracing/PathTracingBuildCloudAccelerationMap.usf
--- a/Engine/Shaders/Private/PathTracing/PathTracingBuildCloudAccelerationMap.usf
+++ b/Engine/Shaders/Private/PathTracing/PathTracingBuildCloudAccelerationMap.usf
@@ -64,13 +64,12 @@ void PathTracingBuildCloudAccelerationMapCS(uint2 DispatchThreadId : SV_Dispatch
 		FDFVector3 AbsoluteWorldPosition = DFMultiply(WorldPos, SKY_UNIT_TO_CM);
 		// run the material and encode the result, which includes the effect that quantization will have during rendering
 		FSampleCloudMaterialResult Result = SampleCloudMaterial(AbsoluteWorldPosition, H, B, T);
-		float3 Extinction = Result.Extinction;
 		
-		float Density   = max3(Extinction.x, Extinction.y, Extinction.z);
-		float DensityLo = min3(Extinction.x, Extinction.y, Extinction.z);
+		float Density   = MATVEC_Max(Result.Extinction);
+		float DensityLo = MATVEC_Min(Result.Extinction);
 
 		// If shader had no density but some emission, we need to artificially boost the density so that we get at least _some_ raymarch steps in the empty regions
-		if (Density == 0 && any(Result.EmissiveColor > 0))
+		if (Density == 0 && any(Result.Emission > 0))
 		{
 			Density = CloudInvVoxelWidth; // ~1 step per voxel
 		}
diff --git a/Engine/Shaders/Private/PathTracing/PathTracingCommon.ush b/Engine/Shaders/Private/PathTracing/PathTracingCommon.ush
--- a/Engine/Shaders/Private/PathTracing/PathTracingCommon.ush
+++ b/Engine/Shaders/Private/PathTracing/PathTracingCommon.ush
@@ -450,6 +450,11 @@ struct FPackedPathTracingPayload : FMinimalPayload
 		PackedData[11] = asuint(Throughput.y);
 		PackedData[12] = asuint(Throughput.z);
 	}
+	
+	void SetVolumetricCallableShaderInputCloudFactor(float CloudFactor)
+	{
+		PackedData[13] = asuint(CloudFactor);
+	}
 
 	void SetVolumetricCallableShaderRISContext(FRISContext SurfaceContext)
 	{
@@ -485,6 +490,11 @@ struct FPackedPathTracingPayload : FMinimalPayload
 		return asfloat(uint3(PackedData[10], PackedData[11], PackedData[12]));
 	}
 
+	float GetVolumetricCallableShaderInputCloudFactor()
+	{
+		return asfloat(PackedData[13]);
+	}
+
 	float GetVolumetricCallableShaderInputCloudDensity()
 	{
 		return f16tof32(PackedData[9] >> 16);
@@ -518,10 +528,17 @@ struct FPackedPathTracingPayload : FMinimalPayload
 		PackedData[10] = PayloadEncodeLDRColor(DualPhaseData * float3(0.5, 0.5, 1.0) + float3(0.5, 0.5, 0.0));
 	}
 
+	void SetVolumetricCallableShaderOutputCloudDensityFactor(float Factor)
+	{
+		PackedData[12] &= 0xFFFFu; // clear upper bits
+		PackedData[12] |= f32tof16(Factor) << 16;
+	}
+
 	void SetVolumetricCallableShaderOutputRadiance(float3 Radiance, float Alpha)
 	{
 		PackedData[11] = PayloadEncodeHDRColor(Radiance);
-		PackedData[12] = asuint(Alpha);
+		PackedData[12] &= 0xFFFF0000u; // clear lower bits
+		PackedData[12] |= f32tof16(Alpha);
 	}
 
 	RandomSequence GetVolumetricCallableShaderOutputRandSeq()
@@ -574,7 +591,12 @@ struct FPackedPathTracingPayload : FMinimalPayload
 
 	float GetVolumetricCallableShaderOutputAlpha()
 	{
-		return asfloat(PackedData[12]);
+		return f16tof32(PackedData[12]);
+	}
+
+	float GetVolumetricCallableCloudDensityFactor()
+	{
+		return f16tof32(PackedData[12] >> 16);
 	}
 
 #endif // PATH_TRACER_USE_CLOUD_SHADER
@@ -708,6 +730,8 @@ struct FPathTracingPayload : FMinimalPayload
 		SpecularEdgeColor = Weight;
 		RoughnessData = PhaseData;
 	}
+	float GetCloudFactor() { return Anisotropy; }
+	void SetCloudFactor(float Factor) { Anisotropy = Factor; }
 
 	// SHADINGMODELID_EYE
 	float GetEyeRoughness() { return RoughnessData.x; }
@@ -820,6 +844,8 @@ struct FPathTracingPayload : FMinimalPayload
 		Specular = PhaseData.y;
 		Roughness = PhaseData.z;
 	}
+	float GetCloudFactor() { return CustomData1.w; }
+	void SetCloudFactor(float Factor) { CustomData1.w = Factor; }
 #endif // PATHTRACING_SUBSTRATE_PAYLOAD
 
 	// Methods used by the integrator
diff --git a/Engine/Shaders/Private/PathTracing/PathTracingCore.ush b/Engine/Shaders/Private/PathTracing/PathTracingCore.ush
--- a/Engine/Shaders/Private/PathTracing/PathTracingCore.ush
+++ b/Engine/Shaders/Private/PathTracing/PathTracingCore.ush
@@ -772,6 +772,7 @@ FPackedPathTracingPayload TraceTransparentRay(inout FPathState PathState, int Bo
 						CloudPayload.GetVolumetricCallableShaderOutputDualHGWeight(),
 						CloudPayload.GetVolumetricCallableShaderOutputDualHGPhaseData()
 					);
+					Result.SetCloudFactor(CloudPayload.GetVolumetricCallableCloudDensityFactor());
 					Payload = PackPathTracingPayload(Result);
 				}
 			}
@@ -1018,7 +1019,7 @@ FPackedPathTracingPayload TraceTransparentRay(inout FPathState PathState, int Bo
 	return Payload;
 }
 
-float3 TraceTransparentVisibilityRay(FRayDesc Ray, float PathRoughness, float PreviousPathRoughness, uint MissShaderIndex, bool bCastShadows, bool bCloudCastShadows, inout RandomSequence RandSequence)
+float3 TraceTransparentVisibilityRay(FRayDesc Ray, float PathRoughness, float PreviousPathRoughness, uint MissShaderIndex, bool bCastShadows, bool bCloudCastShadows, float CloudFactor, inout RandomSequence RandSequence)
 {
 	if (!bCastShadows && MissShaderIndex == 0)
 	{
@@ -1060,7 +1061,7 @@ float3 TraceTransparentVisibilityRay(FRayDesc Ray, float PathRoughness, float Pr
 	{
 		// NOTE: use the path roughness before scattering off the current vertex (so we can have clouds casting shadows onto diffuse surfaces)
 		uint VolumeMask = VolumeFlags;
-		if (!bCloudCastShadows)
+		if (!bCloudCastShadows || CloudFactor <= 0)
 		{
 			// if skipping clouds on shadows, clear the contribution here
 			VolumeMask &= ~PATH_TRACER_VOLUME_ENABLE_CLOUDS;
@@ -1074,7 +1075,7 @@ float3 TraceTransparentVisibilityRay(FRayDesc Ray, float PathRoughness, float Pr
 		}
 		if (VolumeIsectList.HitVolume())
 		{
-			Throughput = VolumeGetTransmittance(Throughput, Ray.Origin, Ray.Direction, VolumeIsectList, RandSequence);
+			Throughput = VolumeGetTransmittance(Throughput, Ray.Origin, Ray.Direction, VolumeIsectList, RandSequence, CloudFactor);
 		}
 		// add in the medium extinction
 		Throughput *= exp(-max(PackedPayload.GetTau(), 0.0));
@@ -1778,6 +1779,7 @@ bool PathTracingKernel(inout FPathState PathState, int Bounce)
 
 				const bool bCastShadows = CastsVolumeShadow(LightSampleLightId);
 				const bool bCloudCastShadows = true; // volumes should get shadows from clouds
+				const float CloudFactor = 1.0; // TODO: the handling of local lights does not yet account for clouds correctly, need to fetch this from another call to the callable shader
 				const uint MissShaderIndex = GetLightMissShaderIndex(LightSampleLightId);
 				// compute direct light sampling?
 				if (MISMode != 0)
@@ -1801,7 +1803,7 @@ bool PathTracingKernel(inout FPathState PathState, int Bounce)
 							{
 								LightContrib *= MISWeightPower(LightSample.Pdf, MaterialEval.Pdf);
 							}
-							LightContrib *= TraceTransparentVisibilityRay(LightRay, 1.0, PathState.GetPathRoughness(), MissShaderIndex, bCastShadows, bCloudCastShadows, PathState.RandSequence);
+							LightContrib *= TraceTransparentVisibilityRay(LightRay, 1.0, PathState.GetPathRoughness(), MissShaderIndex, bCastShadows, bCloudCastShadows, CloudFactor, PathState.RandSequence);
 
 							AccumulateRadiance(PathState.Radiance, LightContrib, bIsCameraRay);
 						}
@@ -1827,7 +1829,7 @@ bool PathTracingKernel(inout FPathState PathState, int Bounce)
 								LightContrib *= MISWeightPower(MaterialSample.Pdf, LightResult.Pdf);
 							}
 							MaterialRay.TMax = LightResult.HitT;
-							LightContrib *= TraceTransparentVisibilityRay(MaterialRay, 1.0, PathState.GetPathRoughness(), MissShaderIndex, bCastShadows, bCloudCastShadows, PathState.RandSequence);
+							LightContrib *= TraceTransparentVisibilityRay(MaterialRay, 1.0, PathState.GetPathRoughness(), MissShaderIndex, bCastShadows, bCloudCastShadows, CloudFactor, PathState.RandSequence);
 							AccumulateRadiance(PathState.Radiance, LightContrib, bIsCameraRay);
 						}
 					}
@@ -1979,8 +1981,9 @@ bool PathTracingKernel(inout FPathState PathState, int Bounce)
 					}
 					const bool bCastShadows = bIsVolumeSample ? CastsVolumeShadow(LightId) : CastsShadow(LightId);
 					const bool bCloudCastShadows = bIsVolumeSample ? true : CastsCloudShadow(LightId);
+					const float CloudFactor = bIsVolumeSample ? Payload.GetCloudFactor() : 1.0;
 					const uint MissShaderIndex = GetLightMissShaderIndex(LightId);
-					LightContrib *= TraceTransparentVisibilityRay(LightRay, AvgRoughness, PathState.GetPathRoughness(), MissShaderIndex, bCastShadows, bCloudCastShadows, PathState.RandSequence);
+					LightContrib *= TraceTransparentVisibilityRay(LightRay, AvgRoughness, PathState.GetPathRoughness(), MissShaderIndex, bCastShadows, bCloudCastShadows, CloudFactor, PathState.RandSequence);
 
 					AccumulateRadiance(PathState.Radiance, LightContrib, bIsCameraRay);
 				}
@@ -2129,15 +2132,24 @@ bool PathTracingKernel(inout FPathState PathState, int Bounce)
 				LightRay.TMax = LightResult.HitT;
 				const bool bCastShadows = bIsVolumeSample ? CastsVolumeShadow(LightId) : CastsShadow(LightId);
 				const bool bCloudCastShadows = bIsVolumeSample ? true : CastsCloudShadow(LightId);
+				const float CloudFactor = bIsVolumeSample ? Payload.GetCloudFactor() : 1.0;
 				const uint MissShaderIndex = GetLightMissShaderIndex(LightId);
 				const float ShadowRayRoughness = ApproximateCaustics ? PathState.GetPathRoughness() : 0.0;
-				LightContrib *= TraceTransparentVisibilityRay(LightRay, ShadowRayRoughness, PreviousPathRoughness, MissShaderIndex, bCastShadows, bCloudCastShadows, PathState.RandSequence);
+				LightContrib *= TraceTransparentVisibilityRay(LightRay, ShadowRayRoughness, PreviousPathRoughness, MissShaderIndex, bCastShadows, bCloudCastShadows, CloudFactor, PathState.RandSequence);
 				AccumulateRadiance(PathState.Radiance, LightContrib, bIsCameraRay);
 			}
 		}
 
 	}
 
+#if PATH_TRACER_USE_CLOUD_SHADER
+	if (bIsVolumeSample && Payload.GetCloudFactor() < 1)
+	{
+		// Cloud MS approximation octaves should not bounce
+		return false;
+	}
+#endif
+
 	// continue the path only if the relevant indirect contribution is enabled
 	bool bKeepGoing = PathState.UpdateScatterType(MaterialSample.ScatterType);
 	// also check for last bounce if emissive contributions are not needed
diff --git a/Engine/Shaders/Private/PathTracing/PathTracingVolumetricCloudMaterialShader.usf b/Engine/Shaders/Private/PathTracing/PathTracingVolumetricCloudMaterialShader.usf
--- a/Engine/Shaders/Private/PathTracing/PathTracingVolumetricCloudMaterialShader.usf
+++ b/Engine/Shaders/Private/PathTracing/PathTracingVolumetricCloudMaterialShader.usf
@@ -56,8 +56,10 @@ FVolumeShadedResult CloudsGetDensity(float3 TranslatedWorldPos)
 	FVolumeShadedResult Result = (FVolumeShadedResult)0;
 	Result.SigmaT        = CloudResult.Extinction;
 	Result.SigmaSDualHG  = CloudResult.Albedo * Result.SigmaT;
-	Result.Emission      = CloudResult.EmissiveColor;
+	Result.Emission      = CloudResult.Emission;
+#if MATERIAL_VOLUMETRIC_ADVANCED
 	Result.DualPhaseData = float3(CloudResult.PhaseG1, CloudResult.PhaseG2, CloudResult.PhaseBlend);
+#endif
 	return Result;
 }
 
@@ -76,13 +78,15 @@ RAY_TRACING_ENTRY_CALLABLE(PathTracingVolumetricCloudMaterialShader, FPackedPath
 	uint Flags = PackedPayload.GetVolumetricCallableShaderInputFlags();
 
 	// TODO: evaluate if it would be better to make a second shader for this (so that we don't get two copies of the shader eval in inner loops)
+	// TODO: evaluate if voxel traversal is worth it for transmittance
 	if (Flags & PATH_TRACER_VOLUME_CALLABLE_FLAGS_TRANSMITTANCE)
 	{
+		float CloudFactor = PackedPayload.GetVolumetricCallableShaderInputCloudFactor();
 		// Transmittance loop for the cloud shader only
 		// NOTE: We have specialized this since we know our upper bound is a scalar value, most of the pdf math cancels out and we can take a faster path
 		float TMin = Ray.TMin;
 		float TMax = Ray.TMax;
-		float SigmaBar = CloudDensity, InvSigmaBar = rcp(SigmaBar);
+		float SigmaBar = CloudDensity * CloudFactor, InvSigmaBar = rcp(SigmaBar);
 		float3 Throughput = PathThroughput;
 		for (int Step = 0; Step < PathTracingCloudParameters.MaxRaymarchSteps; Step++)
 		{
@@ -98,7 +102,7 @@ RAY_TRACING_ENTRY_CALLABLE(PathTracingVolumetricCloudMaterialShader, FPackedPath
 			/* our ray marching step stayed inside the atmo and is still in front of the next hit */
 			float3 WorldPos = Ray.Origin + TMin * Ray.Direction;
 			/* clamp to make sure we never exceed the majorant (should not be the case, but need to avoid any possible numerical issues) */
-			float3 SigmaT = min(CloudsGetDensity(WorldPos).SigmaT, SigmaBar);
+			float3 SigmaT = min(CloudFactor * CloudsGetDensity(WorldPos).SigmaT, SigmaBar);
 			/* keep tracing through the volume */
 			Throughput *= 1.0 - SigmaT * InvSigmaBar;
 			/* Encourage early termination of paths with low contribution */
@@ -119,6 +123,14 @@ RAY_TRACING_ENTRY_CALLABLE(PathTracingVolumetricCloudMaterialShader, FPackedPath
 		const int Bounce = (Flags & PATH_TRACER_VOLUME_CALLABLE_FLAGS_BOUNCE_MASK) >> PATH_TRACER_VOLUME_CALLABLE_FLAGS_BOUNCE_SHIFT;
 		const bool bIsCameraRay = Bounce == 0;
 
+#if MATERIAL_VOLUMETRIC_ADVANCED_MULTISCATTERING_OCTAVE_COUNT > 0
+		FCloudMaterialMSParams MS = (FCloudMaterialMSParams)0; 
+		if (PathTracingCloudParameters.CloudShaderMultipleScatterApproxEnabled)
+		{
+			MS = CloudMaterialGetMSParams();
+		}
+#endif
+
 		// zero out the payload so we can write our output to it
 		PackedPayload = (FPackedPathTracingPayload) 0;
 		PackedPayload.HitT = -1.0; // invalid sample until we are sure we got one
@@ -220,6 +232,9 @@ RAY_TRACING_ENTRY_CALLABLE(PathTracingVolumetricCloudMaterialShader, FPackedPath
 				float3 TranslatedWorldPos = Ro + Ray.TMin * Rd;
 				FVolumeShadedResult Result = (FVolumeShadedResult)0; 
 				MergeVolumeShadedResult(Result, CloudsGetDensity(TranslatedWorldPos), bIsCameraRay && (Flags & PATH_TRACER_VOLUME_HOLDOUT_CLOUDS) != 0);
+#if MATERIAL_VOLUMETRIC_ADVANCED_MULTISCATTERING_OCTAVE_COUNT > 0
+				FVolumeShadedResult CloudResult = Result; // capture just the Cloud contribution (including holdouts) so we can scale it for the MS approx
+#endif
 				if (Flags & PATH_TRACER_VOLUME_ENABLE_ATMOSPHERE)
 				{
 					MergeVolumeShadedResult(Result, AtmosphereGetDensity(TranslatedWorldPos), bIsCameraRay && (Flags & PATH_TRACER_VOLUME_HOLDOUT_ATMOSPHERE) != 0);
@@ -253,7 +268,36 @@ RAY_TRACING_ENTRY_CALLABLE(PathTracingVolumetricCloudMaterialShader, FPackedPath
 						const float3 HGWeight       = select(SigmaS > 0.0, Result.SigmaSHG       / SigmaS, 0.0);
 						const float3 DualHGWeight   = select(SigmaS > 0.0, Result.SigmaSDualHG   / SigmaS, 0.0);
 						PackedPayload.SetVolumetricCallableShaderOutputSample(Ray.TMin, RayleighWeight, HGWeight, Result.PhaseG, DualHGWeight, Result.DualPhaseData, PayloadThroughput);
+						PackedPayload.SetVolumetricCallableShaderOutputCloudDensityFactor(1.0);
+					}
+#if MATERIAL_VOLUMETRIC_ADVANCED_MULTISCATTERING_OCTAVE_COUNT > 0
+					// The material has enabled fake-multiple scattering. We can take this into account here by stochastically returning additional lobes
+					// NOTE: This is not energy conserving! TODO: how to integrate with indirect scattering?
+					
+					float S = MS.ScattFactor, E = MS.ExtinFactor, P = MS.PhaseFactor;
+					for (int ms = 1; ms <= MATERIAL_VOLUMETRIC_ADVANCED_MULTISCATTERING_OCTAVE_COUNT; ms++)
+					{
+						CloudResult.SigmaSDualHG *= S;
+
+						// Propose a candidate hit with just the scaled cloud contribution
+						Contrib = PathThroughput * CloudResult.SigmaSDualHG;
+						SelectionWeight = max3(Contrib.x, Contrib.y, Contrib.z);
+						if (HitSample.Accept(SelectionWeight))
+						{
+							// stash this hit for next time
+							float3 PayloadThroughput = Contrib / SelectionWeight;
+							// We create a blend of SigmaHG (set to isotropic) and SigmaSDualHG (set to the original phase function) to match what raster does (lerp the eval)
+							PackedPayload.SetVolumetricCallableShaderOutputSample(Ray.TMin, 0.0, P, 0.0, 1.0 - P, CloudResult.DualPhaseData, PayloadThroughput);
+							PackedPayload.SetVolumetricCallableShaderOutputCloudDensityFactor(E); // so that future cloud shadow rays get attenuated
+						}
+
+						// Update scale factors for next iteration
+						S *= S;
+						P *= P;
+						E *= E;
 					}
+#endif
+
 				}
 
 				VolumeRadiance += Result.Emission * PathThroughput;
diff --git a/Engine/Shaders/Private/PathTracing/Volume/PathTracingCloudsCommon.ush b/Engine/Shaders/Private/PathTracing/Volume/PathTracingCloudsCommon.ush
--- a/Engine/Shaders/Private/PathTracing/Volume/PathTracingCloudsCommon.ush
+++ b/Engine/Shaders/Private/PathTracing/Volume/PathTracingCloudsCommon.ush
@@ -65,15 +65,3 @@ float2 VoxelToPosition(int2 Idx)
 {
 	return (Idx * CloudVoxelWidth) - CloudClipDistKm * SKY_UNIT_TO_CM;
 }
-
-
-// Simplified struct only used in the cloud case (to reduce the amount of data needed in the payload for the callable shader)
-struct FSampleCloudMaterialResult
-{
-	float3 Extinction;
-	float3 Albedo;
-	float3 EmissiveColor;
-	float PhaseG1;
-	float PhaseG2;
-	float PhaseBlend;
-};
diff --git a/Engine/Shaders/Private/PathTracing/Volume/PathTracingCloudsMaterialCommon.ush b/Engine/Shaders/Private/PathTracing/Volume/PathTracingCloudsMaterialCommon.ush
--- a/Engine/Shaders/Private/PathTracing/Volume/PathTracingCloudsMaterialCommon.ush
+++ b/Engine/Shaders/Private/PathTracing/Volume/PathTracingCloudsMaterialCommon.ush
@@ -13,6 +13,67 @@
 #include "../../SkyAtmosphereCommon.ush"
 
 
+#if MATERIAL_VOLUMETRIC_ADVANCED_GRAYSCALE_MATERIAL
+
+#define MATVEC              float
+#define ColorToMATVEC(C)    ((C).r)   // Assume all 3 floats are the same
+#define MATVEC_Max(C)       (C)
+#define MATVEC_Min(C)       (C)
+#define MATVEC_Avg(C)       (C)
+
+#else
+
+#define MATVEC              float3
+#define ColorToMATVEC(C)    ((C).rgb)
+float MATVEC_Max(float3 C) { return max3(C.r, C.g, C.b); }
+float MATVEC_Min(float3 C) { return min3(C.r, C.g, C.b); }
+float MATVEC_Avg(float3 C) { return dot(C, 1.0f / 3.0f); }
+
+#endif
+
+// Simplified struct that represents what a CloudMaterial can return
+struct FSampleCloudMaterialResult
+{
+	MATVEC Extinction;
+	MATVEC Albedo;
+	MATVEC Emission;
+#if MATERIAL_VOLUMETRIC_ADVANCED
+	float PhaseG1;
+	float PhaseG2;
+	float PhaseBlend;
+#endif
+};
+
+#if MATERIAL_VOLUMETRIC_ADVANCED_MULTISCATTERING_OCTAVE_COUNT > 0
+struct FCloudMaterialMSParams
+{
+	float ScattFactor;
+	float ExtinFactor;
+	float PhaseFactor;
+};
+
+FCloudMaterialMSParams CloudMaterialGetMSParams()
+{
+	FMaterialPixelParameters MaterialParameters = MakeInitializedMaterialPixelParameters();
+	float4 SvPosition = 0; // TODO: should initialize this from the ray somehow?
+	FPixelMaterialInputs PixelMaterialInputs = (FPixelMaterialInputs)0;
+	CalcMaterialParameters(MaterialParameters, PixelMaterialInputs, SvPosition, true);
+
+	FCloudMaterialMSParams Result = (FCloudMaterialMSParams)0;
+
+	#if MATERIAL_VOLUMETRIC_ADVANCED_CLAMP_MULTISCATTERING_CONTRIBUTION
+	Result.ScattFactor = saturate(GetVolumetricAdvancedMaterialOutput3(MaterialParameters));
+	#else
+	Result.ScattFactor = GetVolumetricAdvancedMaterialOutput3(MaterialParameters);
+	#endif
+	Result.ExtinFactor = saturate(GetVolumetricAdvancedMaterialOutput4(MaterialParameters));
+	Result.PhaseFactor = saturate(GetVolumetricAdvancedMaterialOutput5(MaterialParameters));
+
+	return Result;
+}
+#endif
+
+
 // For path tracing, we need to be able to evaluate the material starting from a single point in space
 // In all contexts, we have a slightly more accurate Height already available so pass that in instead of measuring it from the position
 FSampleCloudMaterialResult SampleCloudMaterial(FDFVector3 AbsoluteWorldPosition, float HeightKm, float CloudBotKm, float CloudTopKm)
@@ -34,7 +95,8 @@ FSampleCloudMaterialResult SampleCloudMaterial(FDFVector3 AbsoluteWorldPosition,
 	
 	MaterialParameters.VolumeSampleConservativeDensity = 1.0f;
 #if MATERIAL_VOLUMETRIC_ADVANCED_CONSERVATIVE_DENSITY
-	MaterialParameters.VolumeSampleConservativeDensity = GetVolumetricAdvancedMaterialOutput6(MaterialParameters);	// Evaluate conservative density
+	// Evaluate conservative density - this may be influenced by WorldPosition
+	MaterialParameters.VolumeSampleConservativeDensity = GetVolumetricAdvancedMaterialOutput6(MaterialParameters);
 	if (MaterialParameters.VolumeSampleConservativeDensity.x <= 0.0f)
 	{
 		return (FSampleCloudMaterialResult) 0;
@@ -47,26 +109,29 @@ FSampleCloudMaterialResult SampleCloudMaterial(FDFVector3 AbsoluteWorldPosition,
 
 #if SUBSTRATE_ENABLED
 	FSubstrateBSDF BSDF = PixelMaterialInputs.FrontMaterial.InlinedBSDF;
-	Result.Extinction = VOLUMETRICFOGCLOUD_EXTINCTION(BSDF).rgb;
-	Result.EmissiveColor = BSDF_GETEMISSIVE(BSDF).rgb;
-	Result.Albedo = VOLUMETRICFOGCLOUD_ALBEDO(BSDF).rgb;
+	Result.Extinction = ColorToMATVEC(VOLUMETRICFOGCLOUD_EXTINCTION(BSDF));
+	Result.Albedo     = ColorToMATVEC(VOLUMETRICFOGCLOUD_ALBEDO(BSDF));
+	Result.Emission   = ColorToMATVEC(BSDF_GETEMISSIVE(BSDF));
+		
 #else
-#if !MATERIAL_SHADINGMODEL_UNLIT
-	Result.Extinction = GetMaterialSubsurfaceDataRaw(PixelMaterialInputs).rgb;
-	Result.Albedo = GetMaterialBaseColor(PixelMaterialInputs).rgb * View.DiffuseOverrideParameter.w + View.DiffuseOverrideParameter.xyz;
-#endif
-	Result.EmissiveColor = GetMaterialEmissiveRaw(PixelMaterialInputs).rgb;
+	#if !MATERIAL_SHADINGMODEL_UNLIT
+	Result.Extinction = ColorToMATVEC(GetMaterialSubsurfaceDataRaw(PixelMaterialInputs));
+	Result.Albedo     = ColorToMATVEC(GetMaterialBaseColor(PixelMaterialInputs).rgb * View.DiffuseOverrideParameter.w + View.DiffuseOverrideParameter.xyz);
+	#endif
+	Result.Emission   = ColorToMATVEC(GetMaterialEmissiveRaw(PixelMaterialInputs));
 #endif
 
+	// Clamp to valid ranges
 	Result.Extinction = clamp(Result.Extinction * CENTIMETER_TO_METER, 0.0f, 65000.0f);
-	Result.EmissiveColor = clamp(Result.EmissiveColor * CENTIMETER_TO_METER, 0.0f, 6500.0f);
-	Result.Albedo = saturate(Result.Albedo);
+	Result.Emission   = clamp(Result.Emission   * CENTIMETER_TO_METER, 0.0f, 65000.0f);
+	Result.Albedo     = saturate(Result.Albedo);
 
 #if MATERIAL_VOLUMETRIC_ADVANCED
+	// Advanced node allows configuring the phase function
 	Result.PhaseG1 = GetVolumetricAdvancedMaterialOutput0(MaterialParameters);
 	Result.PhaseG2 = GetVolumetricAdvancedMaterialOutput1(MaterialParameters);
 	Result.PhaseBlend = GetVolumetricAdvancedMaterialOutput2(MaterialParameters);
 #endif
 
 	return Result;
-}
\ No newline at end of file
+}
diff --git a/Engine/Shaders/Private/PathTracing/Volume/PathTracingVolume.ush b/Engine/Shaders/Private/PathTracing/Volume/PathTracingVolume.ush
--- a/Engine/Shaders/Private/PathTracing/Volume/PathTracingVolume.ush
+++ b/Engine/Shaders/Private/PathTracing/Volume/PathTracingVolume.ush
@@ -177,6 +177,7 @@ float3 VolumeGetTransmittance(float3 StartThroughput, float3 Origin, float3 Dire
 #if PATH_TRACER_USE_CLOUD_SHADER
 		FPackedPathTracingPayload CloudPayload = (FPackedPathTracingPayload) 0;
 		CloudPayload.SetVolumetricCallableShaderInput(Origin, Direction, TMin, TMax, RandSequence, PATH_TRACER_VOLUME_CALLABLE_FLAGS_TRANSMITTANCE, Throughput, SigmaBar);
+		CloudPayload.SetVolumetricCallableShaderInputCloudFactor(1.0);
 		CallShader(CloudCallableShaderId, CloudPayload);
 		RandSequence = CloudPayload.GetVolumetricCallableShaderOutputRandSeq();
 		Throughput = CloudPayload.GetVolumetricCallableShaderOutputThroughput();
@@ -207,7 +208,8 @@ float3 VolumeGetTransmittance(float3 StartThroughput, float3 Origin, float3 Dire
 }
 
 // Return the transmittance along a whole intersection list
-float3 VolumeGetTransmittance(float3 StartThroughput, float3 Origin, float3 Direction, FVolumeIntersectionList IntersectionList, inout RandomSequence RandSequence)
+// Each volume type has its own range unlike the function above. This is used for shadow rays, so this version adds the CloudFactor for cloud shadows
+float3 VolumeGetTransmittance(float3 StartThroughput, float3 Origin, float3 Direction, FVolumeIntersectionList IntersectionList, inout RandomSequence RandSequence, float CloudFactor)
 {
 	float3 Throughput = StartThroughput;
 	if (IntersectionList.VolumeMask & PATH_TRACER_VOLUME_ENABLE_ATMOSPHERE)
@@ -231,6 +233,7 @@ float3 VolumeGetTransmittance(float3 StartThroughput, float3 Origin, float3 Dire
 #if PATH_TRACER_USE_CLOUD_SHADER
 		FPackedPathTracingPayload CloudPayload = (FPackedPathTracingPayload) 0;
 		CloudPayload.SetVolumetricCallableShaderInput(Origin, Direction, TMin, TMax, RandSequence, PATH_TRACER_VOLUME_CALLABLE_FLAGS_TRANSMITTANCE, Throughput, SigmaBar);
+		CloudPayload.SetVolumetricCallableShaderInputCloudFactor(CloudFactor);
 		CallShader(CloudCallableShaderId, CloudPayload);
 		RandSequence = CloudPayload.GetVolumetricCallableShaderOutputRandSeq();
 		Throughput = CloudPayload.GetVolumetricCallableShaderOutputThroughput();
diff --git a/Engine/Shaders/Private/PathTracing/Volume/PathTracingVolumeCommon.ush b/Engine/Shaders/Private/PathTracing/Volume/PathTracingVolumeCommon.ush
--- a/Engine/Shaders/Private/PathTracing/Volume/PathTracingVolumeCommon.ush
+++ b/Engine/Shaders/Private/PathTracing/Volume/PathTracingVolumeCommon.ush
@@ -245,6 +245,7 @@ FPathTracingPayload CreateMediumHitPayload(float HitT, float3 TranslatedWorldPos
 	Result.SetBaseColor(RayleighWeight);
 	Result.SetHG(HGWeight, ShadedResult.PhaseG);
 	Result.SetDualHG(DualHGWeight, ShadedResult.DualPhaseData);
+	Result.SetCloudFactor(1.0);
 	return Result;
 }
 
diff --git a/Engine/Source/Runtime/Renderer/Private/PathTracing.cpp b/Engine/Source/Runtime/Renderer/Private/PathTracing.cpp
--- a/Engine/Source/Runtime/Renderer/Private/PathTracing.cpp
+++ b/Engine/Source/Runtime/Renderer/Private/PathTracing.cpp
@@ -325,26 +325,36 @@ TAutoConsoleVariable<int32> CVarPathTracingAtmosphereOpticalDepthLutNumSamples(
 );
 
 TAutoConsoleVariable<int32> CVarPathTracingCloudAccelerationMapResolution(
-	TEXT("r.PathTracing.CloudAccelerationMapResolution"),
+	TEXT("r.PathTracing.CloudAccelerationMap.Resolution"),
 	512,
 	TEXT("Size of the square texture used to accelerate cloud ray marching for the path tracer in reference atmosphere mode.  (default = 512)"),
 	ECVF_RenderThreadSafe
 );
 
 TAutoConsoleVariable<int32> CVarPathTracingCloudAccelerationMapNumSamples(
-	TEXT("r.PathTracing.CloudAccelerationMapNumSamples"),
+	TEXT("r.PathTracing.CloudAccelerationMap.NumSamples"),
 	64,
 	TEXT("Number of ray marching samples used when building the cloud acceleration map for the path tracer in reference atmosphere mode.  (default = 64)"),
 	ECVF_RenderThreadSafe
 );
 
 TAutoConsoleVariable<bool> CVarPathTracingCloudAccelerationMapVisualize(
-	TEXT("r.PathTracing.CloudAccelerationMapVisualize"),
+	TEXT("r.PathTracing.CloudAccelerationMap.Visualize"),
 	false,
 	TEXT("If true, replace clouds with a visualization of the acceleration map to help visualize it and fine tune its resolution (default = false)"),
 	ECVF_RenderThreadSafe
 );
 
+TAutoConsoleVariable<int> CVarPathTracingCloudMultipleScatterMode(
+	TEXT("r.PathTracing.CloudMultipleScatterMode"),
+	1,
+	TEXT("Selects the multiple scattering mode for rendering of volumetric clouds in the path tracer.\n")
+	TEXT("  0: None      - multiple scattering settings inside the material are ignored, CloudRoughnessCutoff is applied\n")
+	TEXT("  1: Approx    - multiple scattering settings inside the material are used, CloudRoughnessCutoff is applied (default)\n")
+	TEXT("  2: Reference - multiple scattering settings inside the material are ignored, CloudRoughnessCutoff is not applied\n"),
+	ECVF_RenderThreadSafe
+);
+
 TAutoConsoleVariable<float> CVarPathTracingCloudRoughnessCutoff(
 	TEXT("r.PathTracing.CloudRoughnessCutoff"),
 	0.05f,
@@ -622,6 +632,7 @@ struct FPathTracingConfig
 	int CloudAccelerationMapNumSamples;
 	int CloudAccelerationMapResolution;
 	bool CloudAccelerationMapVisualize;
+	int CloudMultipleScatterMode;
 
 	bool IsDifferent(const FPathTracingConfig& Other) const
 	{
@@ -662,6 +673,7 @@ struct FPathTracingConfig
 			CloudAccelerationMapNumSamples != Other.CloudAccelerationMapNumSamples ||
 			CloudAccelerationMapResolution != Other.CloudAccelerationMapResolution ||
 			CloudAccelerationMapVisualize != Other.CloudAccelerationMapVisualize ||
+			CloudMultipleScatterMode != Other.CloudMultipleScatterMode ||
 			UseMultiGPU != Other.UseMultiGPU;
 	}
 
@@ -869,7 +881,6 @@ static void PreparePathTracingData(const FScene* Scene, const FViewInfo& View, F
 		PathTracing::UsesReferenceAtmosphere(View);
 	
 	const FVolumetricCloudRenderSceneInfo* CloudInfo = Scene->GetVolumetricCloudSceneInfo();
-	const bool bExperimental = CVarPathTracingExperimental.GetValueOnRenderThread();
 	const bool bVolumeCloudsVisible = ShouldRenderVolumetricCloud(Scene, ShowFlags)
 		&& bUseReferenceAtmosphere
 		&& CVarPathTracingEnableReferenceClouds.GetValueOnRenderThread() != 0
@@ -894,6 +905,14 @@ static void PreparePathTracingData(const FScene* Scene, const FViewInfo& View, F
 	PathTracingData.VolumeFlags |= CVarPathTracingUseAnalyticTransmittance.GetValueOnRenderThread() ? PATH_TRACER_VOLUME_USE_ANALYTIC_TRANSMITTANCE : 0;
 
 	PathTracingData.CloudRoughnessCutoff = bVolumeCloudsVisible ? CVarPathTracingCloudRoughnessCutoff.GetValueOnRenderThread() : -1.0f;
+	if (bVolumeCloudsVisible &&
+		PathTracingData.CloudRoughnessCutoff > 0 &&
+		PathTracingData.CloudRoughnessCutoff < 1 &&
+		CVarPathTracingCloudMultipleScatterMode.GetValueOnRenderThread() == 2)
+	{
+		// User had clouds visible, but wants multiple scattering to be done by brute force -- enable this here
+		PathTracingData.CloudRoughnessCutoff = 1.0f;
+	}
 
 	PathTracingData.EnableDBuffer = CVarPathTracingUseDBuffer.GetValueOnRenderThread();
 
@@ -1172,6 +1191,7 @@ BEGIN_GLOBAL_SHADER_PARAMETER_STRUCT(FPathTracingCloudParameterGlobals,)
 	SHADER_PARAMETER(FVector3f, PlanetCenterTranslatedWorldHi)
 	SHADER_PARAMETER(FVector3f, PlanetCenterTranslatedWorldLo)
 	SHADER_PARAMETER(uint32, MaxRaymarchSteps)
+	SHADER_PARAMETER(int32, CloudShaderMultipleScatterApproxEnabled)
 END_GLOBAL_SHADER_PARAMETER_STRUCT()
 
 IMPLEMENT_GLOBAL_SHADER_PARAMETER_STRUCT(FPathTracingCloudParameterGlobals, "PathTracingCloudParameters");
@@ -1210,11 +1230,6 @@ class FPathTracingRG : public FGlobalShader
 				// adaptive sampling is experimental
 				return false;
 			}
-			if (PermutationVector.Get<FCloudShader>())
-			{
-				// using the cloud shader via callable shaders is experimental
-				return false;
-			}
 		}
 		else
 		{
@@ -1571,10 +1586,7 @@ class FPathTracingBuildCloudAccelerationMapCS : public FMeshMaterialShader
 
 	static bool ShouldCompilePermutation(const FMeshMaterialShaderPermutationParameters& Parameters)
 	{
-		const bool bUseExperimental = CVarPathTracingExperimental.GetValueOnAnyThread();
-
 		return ShouldCompilePathTracingShadersForProject(Parameters.Platform) &&
-			bUseExperimental &&
 			Parameters.MaterialParameters.bIsUsedWithVolumetricCloud &&
 			Parameters.MaterialParameters.MaterialDomain == MD_Volume;
 	}
@@ -1613,11 +1625,8 @@ public:
 
 	static bool ShouldCompilePermutation(const FMaterialShaderPermutationParameters& Parameters)
 	{
-		const bool bUseExperimental = CVarPathTracingExperimental.GetValueOnAnyThread();
-
 		return ShouldCompilePathTracingShadersForProject(Parameters.Platform) &&
 			ShouldCompileRayTracingCallableShadersForProject(Parameters.Platform) &&
-			bUseExperimental &&
 			Parameters.MaterialParameters.bIsUsedWithVolumetricCloud &&
 			Parameters.MaterialParameters.MaterialDomain == MD_Volume;
 	}
@@ -1694,9 +1703,8 @@ void PreparePathTracingCloudMaterial(FRDGBuilder& GraphBuilder, FScene* Scene, T
 		View.PathTracingVolumetricCloudCallableShaderIndex = -1;
 	}
 
-	// if we are using reference clouds, or if we are using the cloud map, no need to prepare the callable shader version
-	if (CVarPathTracingEnableReferenceClouds.GetValueOnRenderThread() == 0 ||
-		CVarPathTracingExperimental.GetValueOnRenderThread() == false)
+	// if we are not using reference clouds, no need to prepare the callable shader
+	if (CVarPathTracingEnableReferenceClouds.GetValueOnRenderThread() == 0)
 	{
 		return;
 	}
@@ -1783,6 +1791,7 @@ void PreparePathTracingCloudMaterial(FRDGBuilder& GraphBuilder, FScene* Scene, T
 
 		PreparePlanetCenter(View, Scene->GetSkyAtmosphereSceneInfo(), &Params.PlanetCenterTranslatedWorldHi, &Params.PlanetCenterTranslatedWorldLo);
 		Params.MaxRaymarchSteps = CVarPathTracingMaxRaymarchSteps.GetValueOnRenderThread();
+		Params.CloudShaderMultipleScatterApproxEnabled = CVarPathTracingCloudMultipleScatterMode.GetValueOnRenderThread() == 1;
 
 		TUniformBufferRef<FPathTracingCloudParameterGlobals> CloudParametersUB = CreateUniformBufferImmediate(Params, EUniformBufferUsage::UniformBuffer_SingleFrame);
 		Scene->RayTracingSBT.TransientUniformBuffers.Add(CloudParametersUB); // Hold uniform buffer ref in RayTracingSBT since FMeshDrawSingleShaderBindings doesn't
@@ -3036,8 +3045,7 @@ static FPathTracingRG::FPermutationDomain GetPathTracingRGPermutation(const FSce
 {
 	const bool bUseExperimental = CVarPathTracingExperimental.GetValueOnRenderThread();
 	const bool bUseAdaptiveSampling = bUseExperimental && CVarPathTracingAdaptiveSampling.GetValueOnRenderThread() != 0;
-	const bool bUseCloudShader = bUseExperimental
-		&& ShouldCompileRayTracingCallableShadersForProject(Scene.GetShaderPlatform())
+	const bool bUseCloudShader = ShouldCompileRayTracingCallableShadersForProject(Scene.GetShaderPlatform())
 		&& CVarPathTracingEnableReferenceClouds.GetValueOnRenderThread()
 		&& !CVarPathTracingCloudAccelerationMapVisualize.GetValueOnRenderThread()
 		&& Scene.GetVolumetricCloudSceneInfo() != nullptr;
@@ -3344,6 +3352,7 @@ void FDeferredShadingSceneRenderer::RenderPathTracing(
 	Config.CloudAccelerationMapNumSamples = FMath::Clamp(CVarPathTracingCloudAccelerationMapNumSamples.GetValueOnRenderThread(), 1, 65536);
 	Config.CloudAccelerationMapResolution = FMath::Clamp(CVarPathTracingCloudAccelerationMapResolution.GetValueOnRenderThread(), 1, 4096);
 	Config.CloudAccelerationMapVisualize = CVarPathTracingCloudAccelerationMapVisualize.GetValueOnRenderThread();
+	Config.CloudMultipleScatterMode = CVarPathTracingCloudMultipleScatterMode.GetValueOnRenderThread();
 
 	// compute an integer code of what show flags and booleans related to lights are currently enabled so we can detect changes
 	Config.LightShowFlags = 0;
