Fixed race condition with dynamic cached shadows and the scene render clean up task which releases cached preshadow mesh passes.

diff --git a/Engine/Source/Runtime/Renderer/Private/SceneRenderBuilder.cpp b/Engine/Source/Runtime/Renderer/Private/SceneRenderBuilder.cpp
--- a/Engine/Source/Runtime/Renderer/Private/SceneRenderBuilder.cpp
+++ b/Engine/Source/Runtime/Renderer/Private/SceneRenderBuilder.cpp
@@ -639,6 +639,16 @@ public:
 		return EnumHasAnyFlags(CommonShowFlags, ESceneRenderCommonShowFlags::HitProxies)  == EngineShowFlags.HitProxies
 			&& EnumHasAnyFlags(CommonShowFlags, ESceneRenderCommonShowFlags::PathTracing) == EngineShowFlags.PathTracing;
 	}
+	
+	static void WaitForAsyncCleanupTask()
+	{
+		check(IsInRenderingThread());
+		if (AsyncTasks.Cleanup)
+		{
+			AsyncTasks.Cleanup->Wait();
+			AsyncTasks.Cleanup = {};
+		}
+	}
 
 	static void WaitForAsyncDeleteTask()
 	{
@@ -650,6 +660,11 @@ public:
 		}
 	}
 
+	static const FGraphEventRef& GetAsyncCleanupTask()
+	{
+		return AsyncTasks.Cleanup;
+	}
+
 private:
 	struct FAsyncTasks
 	{
@@ -1105,4 +1120,14 @@ bool FSceneRenderBuilder::IsCompatible(const FEngineShowFlags& EngineShowFlags)
 void FSceneRenderBuilder::WaitForAsyncDeleteTask()
 {
 	FSceneRenderProcessor::WaitForAsyncDeleteTask();
+}
+
+void FSceneRenderBuilder::WaitForAsyncCleanupTask()
+{
+	FSceneRenderProcessor::WaitForAsyncCleanupTask();
+}
+
+const FGraphEventRef& FSceneRenderBuilder::GetAsyncCleanupTask()
+{
+	return FSceneRenderProcessor::GetAsyncCleanupTask();
 }
\ No newline at end of file
diff --git a/Engine/Source/Runtime/Renderer/Private/SceneRenderBuilder.h b/Engine/Source/Runtime/Renderer/Private/SceneRenderBuilder.h
--- a/Engine/Source/Runtime/Renderer/Private/SceneRenderBuilder.h
+++ b/Engine/Source/Runtime/Renderer/Private/SceneRenderBuilder.h
@@ -58,8 +58,11 @@ public:
 		}
 	}
 
+	static void WaitForAsyncCleanupTask();
 	static void WaitForAsyncDeleteTask();
 
+	static const FGraphEventRef& GetAsyncCleanupTask();
+
 private:
 	void LazyInit();
 	void BeginGroup(FString&& Name, ESceneRenderGroupFlags Flags) override;
diff --git a/Engine/Source/Runtime/Renderer/Private/ShadowSetup.cpp b/Engine/Source/Runtime/Renderer/Private/ShadowSetup.cpp
--- a/Engine/Source/Runtime/Renderer/Private/ShadowSetup.cpp
+++ b/Engine/Source/Runtime/Renderer/Private/ShadowSetup.cpp
@@ -42,6 +42,7 @@
 #include "LineTypes.h"
 #include "SceneCulling/SceneCulling.h"
 #include "ReadOnlyCVARCache.h"
+#include "SceneRenderBuilder.h"
 
 using namespace UE::Geometry;
 using namespace ShadowRendering;
@@ -5302,9 +5303,13 @@ void FSceneRenderer::BeginGatherShadowPrimitives(FDynamicShadowsTaskData* TaskDa
 	ShadowSetupPrerequisites.AddPrerequisites(Scene->GetCacheMeshDrawCommandsTask());
 	ShadowSetupPrerequisites.Trigger();
 
+	// The FSceneRenderBuilder clean-up task clears shadow mesh passes. Cached pre-shadows are the only shadow type that is persistently
+	// reused from render-to-render and those need to be cleaned up prior to generating new mesh draw commands for the current scene.
+
 	if (!TaskData->bMultithreadedCreateAndFilterShadows)
 	{
 		ShadowSetupPrerequisites.Wait();
+		FSceneRenderBuilder::WaitForAsyncCleanupTask();
 		CreateDynamicShadows(*TaskData);
 	}
 
@@ -5313,6 +5318,7 @@ void FSceneRenderer::BeginGatherShadowPrimitives(FDynamicShadowsTaskData* TaskDa
 		if (TaskData->PreShadows.Num() || TaskData->ViewDependentWholeSceneShadows.Num())
 		{
 			ShadowSetupPrerequisites.Wait();
+			FSceneRenderBuilder::WaitForAsyncCleanupTask();
 			FGatherShadowPrimitivesPrepareTask(*TaskData).AnyThreadTask();
 
 			if (!TaskData->bRunningEarly)
@@ -5337,6 +5343,7 @@ void FSceneRenderer::BeginGatherShadowPrimitives(FDynamicShadowsTaskData* TaskDa
 	{
 		FGraphEventArray Prereqs;
 		Prereqs.Emplace(CreateCompatibilityGraphEvent(ShadowSetupPrerequisites));
+		Prereqs.Emplace(FSceneRenderBuilder::GetAsyncCleanupTask());
 
 		TaskData->TaskEvent = TGraphTask<FGatherShadowPrimitivesPrepareTask>::CreateTask(&Prereqs).ConstructAndDispatchWhenReady(*TaskData);
 	}
