Added more descriptive CPU scopes to help improve readability of insights when visibility is waiting on the GPU.

diff --git a/Engine/Source/Runtime/Renderer/Private/SceneVisibility.cpp b/Engine/Source/Runtime/Renderer/Private/SceneVisibility.cpp
--- a/Engine/Source/Runtime/Renderer/Private/SceneVisibility.cpp
+++ b/Engine/Source/Runtime/Renderer/Private/SceneVisibility.cpp
@@ -3105,6 +3105,8 @@ void FGPUOcclusion::WaitForLastOcclusionQuery()
 	const uint32 QueryIndex = FOcclusionQueryHelpers::GetQueryLookupIndex(ViewState.OcclusionFrameCounter, State.NumBufferedFrames);
 	if (ViewState.Occlusion.LastOcclusionQueryArray[QueryIndex])
 	{
+		TRACE_CPUPROFILER_EVENT_SCOPE(GPUBound_WaitingForGPUForOcclusionQueries_SeeGPUTrack);
+
 		uint64 Result;
 		const bool bWait = true;
 		RHIGetRenderQueryResult(ViewState.Occlusion.LastOcclusionQueryArray[QueryIndex], Result, bWait);
@@ -4643,7 +4645,7 @@ void FVisibilityTaskData::ProcessRenderThreadTasks()
 	{
 		if (DynamicMeshElements.CommandPipe)
 		{
-			SCOPED_NAMED_EVENT(WaitForGatherDynamicMeshElements, FColor::Magenta);
+			SCOPED_NAMED_EVENT(WaitForVisibilityTasks, FColor::Magenta);
 			CSV_SCOPED_SET_WAIT_STAT(Visibility);
 
 			// Wait on the command pipe first as it will be continually updating the render thread event (and process tasks while we wait).
