Path Tracer: Expose "Reference Clouds" in PPV

Now that this setting is out of experimental, add it to PPV so end users can discover it more easily. The old cvar default is changed to -1 to mean "inherit from PPV" like the others.

This requires changing the raygen discovery to run per view since the PPV setting could be different for each view.  Tweak the logic to avoid deduplication when running on a single view (common case).

#rnx

diff --git a/Engine/Source/Runtime/Engine/Classes/Engine/Scene.h b/Engine/Source/Runtime/Engine/Classes/Engine/Scene.h
--- a/Engine/Source/Runtime/Engine/Classes/Engine/Scene.h
+++ b/Engine/Source/Runtime/Engine/Classes/Engine/Scene.h
@@ -1350,6 +1350,9 @@ struct FPostProcessSettings
 	UPROPERTY(EditAnywhere, BlueprintReadWrite, Category = Overrides, meta = (PinHiddenByDefault, InlineEditConditionToggle))
 	uint32 bOverride_PathTracingEnableReferenceAtmosphere : 1;
 
+	UPROPERTY(EditAnywhere, BlueprintReadWrite, Category = Overrides, meta = (PinHiddenByDefault, InlineEditConditionToggle))
+	uint32 bOverride_PathTracingEnableReferenceClouds : 1;
+
 	UPROPERTY(EditAnywhere, BlueprintReadWrite, Category = Overrides, meta = (PinHiddenByDefault, InlineEditConditionToggle))
 	uint32 bOverride_PathTracingEnableDenoiser : 1;
 
@@ -2349,6 +2352,10 @@ struct FPostProcessSettings
 	UPROPERTY(EditAnywhere, BlueprintReadWrite, Category = "Path Tracing", meta = (editcondition = "bOverride_PathTracingEnableReferenceAtmosphere", DisplayName = "Reference Atmosphere"))
 	uint32 PathTracingEnableReferenceAtmosphere : 1;
 
+	/** Enables path tracing of volumetric clouds. This requires Reference Atmosphere. */
+	UPROPERTY(EditAnywhere, BlueprintReadWrite, Category = "Path Tracing", meta = (editcondition = "bOverride_PathTracingEnableReferenceClouds", DisplayName = "Reference Clouds"))
+	uint32 PathTracingEnableReferenceClouds : 1;
+
 	/** Run the currently loaded denoiser plugin on the last sample to remove noise from the output. Has no effect if a plug-in is not loaded. */
 	UPROPERTY(EditAnywhere, BlueprintReadWrite, Category = "Path Tracing", meta = (editcondition = "bOverride_PathTracingEnableDenoiser", DisplayName = "Denoiser"))
 	uint32 PathTracingEnableDenoiser : 1;
diff --git a/Engine/Source/Runtime/Engine/Private/PostProcessUtils.cpp b/Engine/Source/Runtime/Engine/Private/PostProcessUtils.cpp
--- a/Engine/Source/Runtime/Engine/Private/PostProcessUtils.cpp
+++ b/Engine/Source/Runtime/Engine/Private/PostProcessUtils.cpp
@@ -231,6 +231,7 @@ bool FPostProcessUtils::OverridePostProcessSettings(FPostProcessSettings& ThisFr
 		UE_SET_PP(PathTracingEnableEmissiveMaterials);
 		UE_SET_PP(PathTracingEnableReferenceDOF);
 		UE_SET_PP(PathTracingEnableReferenceAtmosphere);
+		UE_SET_PP(PathTracingEnableReferenceClouds);
 		UE_SET_PP(PathTracingEnableDenoiser);
 		UE_SET_PP(PathTracingIncludeEmissive);
 		UE_SET_PP(PathTracingIncludeDiffuse);
@@ -536,6 +537,7 @@ bool FPostProcessUtils::BlendPostProcessSettings(FPostProcessSettings& ThisFrom,
 		UE_SET_PP(PathTracingEnableEmissiveMaterials);
 		UE_SET_PP(PathTracingEnableReferenceDOF);
 		UE_SET_PP(PathTracingEnableReferenceAtmosphere);
+		UE_SET_PP(PathTracingEnableReferenceClouds);
 		UE_SET_PP(PathTracingEnableDenoiser);
 		UE_SET_PP(PathTracingIncludeEmissive);
 		UE_SET_PP(PathTracingIncludeDiffuse);
diff --git a/Engine/Source/Runtime/Engine/Private/Scene.cpp b/Engine/Source/Runtime/Engine/Private/Scene.cpp
--- a/Engine/Source/Runtime/Engine/Private/Scene.cpp
+++ b/Engine/Source/Runtime/Engine/Private/Scene.cpp
@@ -629,6 +629,7 @@ FPostProcessSettings::FPostProcessSettings()
 	PathTracingEnableEmissiveMaterials = 1;
 	PathTracingEnableReferenceDOF = 0;
 	PathTracingEnableReferenceAtmosphere = 0;
+	PathTracingEnableReferenceClouds = 0;
 	PathTracingEnableDenoiser = 1;
 
 	PathTracingIncludeEmissive = 1;
diff --git a/Engine/Source/Runtime/Engine/Private/SceneView.cpp b/Engine/Source/Runtime/Engine/Private/SceneView.cpp
--- a/Engine/Source/Runtime/Engine/Private/SceneView.cpp
+++ b/Engine/Source/Runtime/Engine/Private/SceneView.cpp
@@ -1794,6 +1794,7 @@ void FSceneView::OverridePostProcessSettings(const FPostProcessSettings& Src, fl
 		SET_PP(PathTracingEnableEmissiveMaterials);
 		SET_PP(PathTracingEnableReferenceDOF);
 		SET_PP(PathTracingEnableReferenceAtmosphere);
+		SET_PP(PathTracingEnableReferenceClouds);
 		SET_PP(PathTracingEnableDenoiser);
 		SET_PP(PathTracingIncludeEmissive);
 		SET_PP(PathTracingIncludeDiffuse);
diff --git a/Engine/Source/Runtime/Renderer/Private/DeferredShadingRenderer.cpp b/Engine/Source/Runtime/Renderer/Private/DeferredShadingRenderer.cpp
--- a/Engine/Source/Runtime/Renderer/Private/DeferredShadingRenderer.cpp
+++ b/Engine/Source/Runtime/Renderer/Private/DeferredShadingRenderer.cpp
@@ -625,7 +625,10 @@ bool FDeferredShadingSceneRenderer::SetupRayTracingPipelineStatesAndSBT(FRDGBuil
 		{
 			// This view only needs the path tracing raygen shaders as all other
 			// passes should be disabled.
-			PreparePathTracing(ViewFamily, *Scene, RayGenShaders);
+			for (const FViewInfo& View : Views)
+			{
+				PreparePathTracing(View, *Scene, RayGenShaders);
+			}
 		}
 		else
 		{
@@ -659,6 +662,11 @@ bool FDeferredShadingSceneRenderer::SetupRayTracingPipelineStatesAndSBT(FRDGBuil
 					PrepareMegaLightsHardwareRayTracing(View, RayGenShaders);
 				}
 			}
+		}
+
+		if (Views.Num() > 1)
+		{
+			// If we have more than one View, chances are we got many duplicates, so compact the list here
 			DeduplicateRayGenerationShaders(RayGenShaders);
 		}
 
diff --git a/Engine/Source/Runtime/Renderer/Private/DeferredShadingRenderer.h b/Engine/Source/Runtime/Renderer/Private/DeferredShadingRenderer.h
--- a/Engine/Source/Runtime/Renderer/Private/DeferredShadingRenderer.h
+++ b/Engine/Source/Runtime/Renderer/Private/DeferredShadingRenderer.h
@@ -1139,7 +1139,7 @@ private:
 	static void PrepareRayTracingTranslucency(const FViewInfo& View, TArray<FRHIRayTracingShader*>& OutRayGenShaders);
 	static void PrepareRayTracingVolumetricFogShadows(const FViewInfo& View, const FScene& Scene, TArray<FRHIRayTracingShader*>& OutRayGenShaders);
 	static void PrepareRayTracingDebug(const FSceneViewFamily& ViewFamily, TArray<FRHIRayTracingShader*>& OutRayGenShaders);
-	static void PreparePathTracing(const FSceneViewFamily& ViewFamily, const FScene& Scene, TArray<FRHIRayTracingShader*>& OutRayGenShaders);
+	static void PreparePathTracing(const FViewInfo& View, const FScene& Scene, TArray<FRHIRayTracingShader*>& OutRayGenShaders);
 	static void PrepareLumenHardwareRayTracingShortRangeAO(const FViewInfo& View, TArray<FRHIRayTracingShader*>& OutRayGenShaders);
 	void PrepareLumenHardwareRayTracingScreenProbeGather(const FViewInfo& View, TArray<FRHIRayTracingShader*>& OutRayGenShaders);
 	void PrepareLumenHardwareRayTracingRadianceCache(const FViewInfo& View, TArray<FRHIRayTracingShader*>& OutRayGenShaders);
diff --git a/Engine/Source/Runtime/Renderer/Private/PathTracing.cpp b/Engine/Source/Runtime/Renderer/Private/PathTracing.cpp
--- a/Engine/Source/Runtime/Renderer/Private/PathTracing.cpp
+++ b/Engine/Source/Runtime/Renderer/Private/PathTracing.cpp
@@ -287,7 +287,7 @@ FAutoConsoleVariableRef CVarPathTracingEnableReferenceDOF(
 	TEXT("Should the path tracer ray trace the depth-of-field effect instead of the post-processed effect?\n")
 	TEXT("-1: Inherit from PostProcess settings (default)\n")
 	TEXT(" 0: Disabled\n")
-	TEXT(" 1: Enable\n"),
+	TEXT(" 1: Enabled\n"),
 	ECVF_RenderThreadSafe
 );
 
@@ -297,16 +297,17 @@ TAutoConsoleVariable<int32> CVarPathTracingEnableReferenceAtmosphere(
 	TEXT("Should the path tracer use a volumetric calculation to represent the sky atmosphere?\n")
 	TEXT("-1: Inherit from PostProcess settings (default)\n")
 	TEXT(" 0: Disabled\n")
-	TEXT(" 1: Enable\n"),
+	TEXT(" 1: Enabled\n"),
 	ECVF_RenderThreadSafe
 );
 
 TAutoConsoleVariable<int32> CVarPathTracingEnableReferenceClouds(
 	TEXT("r.PathTracing.EnableReferenceClouds"),
-	0,
+	-1,
 	TEXT("Should the path tracer use a volumetric calculation to represent volumetric clouds? (This requires Reference Atmosphere to be enabled)\n")
-	TEXT(" 0: Disabled (default)\n")
-	TEXT(" 1: Enable\n"),
+	TEXT("-1: Inherit from PostProcess settings (default)\n")
+	TEXT(" 0: Disabled\n")
+	TEXT(" 1: Enabled\n"),
 	ECVF_RenderThreadSafe
 );
 
@@ -804,6 +805,17 @@ namespace PathTracing
 		return EnableReferenceAtmosphereCVar < 0 ? View.FinalPostProcessSettings.PathTracingEnableReferenceAtmosphere != 0 : EnableReferenceAtmosphereCVar != 0;
 	}
 
+	bool UsesReferenceClouds(const FViewInfo& View)
+	{
+		// Reference Clouds are only available if reference atmosphere is on
+		if (UsesReferenceAtmosphere(View))
+		{
+			const int32 EnableReferenceCloudsCVar = CVarPathTracingEnableReferenceClouds.GetValueOnRenderThread();
+			return EnableReferenceCloudsCVar < 0 ? View.FinalPostProcessSettings.PathTracingEnableReferenceClouds != 0 : EnableReferenceCloudsCVar != 0;
+		}
+		return false;
+	}
+
 	bool UsesReferenceDOF(const FViewInfo& View)
 	{
 		return GEnableReferenceDOF < 0 ? View.FinalPostProcessSettings.PathTracingEnableReferenceDOF != 0 : GEnableReferenceDOF != 0;
@@ -880,11 +892,9 @@ static void PreparePathTracingData(const FScene* Scene, const FViewInfo& View, F
 		View.SkyAtmosphereUniformShaderParameters != nullptr &&
 		PathTracing::UsesReferenceAtmosphere(View);
 	
+	// NOTE: the callable shader is populated only when clouds are active, so no need to check PPV/cvar again
 	const FVolumetricCloudRenderSceneInfo* CloudInfo = Scene->GetVolumetricCloudSceneInfo();
-	const bool bVolumeCloudsVisible = ShouldRenderVolumetricCloud(Scene, ShowFlags)
-		&& bUseReferenceAtmosphere
-		&& CVarPathTracingEnableReferenceClouds.GetValueOnRenderThread() != 0
-		&& (uint32(View.PathTracingVolumetricCloudCallableShaderIndex) < Scene->RayTracingSBT.NumCallableShaderSlots);
+	const bool bVolumeCloudsVisible = uint32(View.PathTracingVolumetricCloudCallableShaderIndex) < Scene->RayTracingSBT.NumCallableShaderSlots;
 
 	// Merge all volume flags into one uint
 	PathTracingData.VolumeFlags = 0;
@@ -1703,12 +1713,6 @@ void PreparePathTracingCloudMaterial(FRDGBuilder& GraphBuilder, FScene* Scene, T
 		View.PathTracingVolumetricCloudCallableShaderIndex = -1;
 	}
 
-	// if we are not using reference clouds, no need to prepare the callable shader
-	if (CVarPathTracingEnableReferenceClouds.GetValueOnRenderThread() == 0)
-	{
-		return;
-	}
-
 	if (!ShouldCompileRayTracingCallableShadersForProject(Scene->GetShaderPlatform()))
 	{
 		return;
@@ -1747,6 +1751,13 @@ void PreparePathTracingCloudMaterial(FRDGBuilder& GraphBuilder, FScene* Scene, T
 
 	for (FViewInfo& View : Views)
 	{
+		if (!PathTracing::UsesReferenceClouds(View) ||
+			!ShouldRenderVolumetricCloud(Scene, View.Family->EngineShowFlags))
+		{
+			// reference cloud mode disabled for this view, or clouds disabled for this view
+			continue;
+		}
+
 		uint32 BaseCallableSlotIndex = Scene->RayTracingSBT.NumCallableShaderSlots;
 		FRayTracingShaderCommand& Command = Scene->RayTracingSBT.CallableCommands.AddDefaulted_GetRef();
 
@@ -3041,14 +3052,13 @@ class FPathTracingCompositorPS : public FGlobalShader
 };
 IMPLEMENT_SHADER_TYPE(, FPathTracingCompositorPS, TEXT("/Engine/Private/PathTracing/PathTracingCompositingPixelShader.usf"), TEXT("CompositeMain"), SF_Pixel);
 
-static FPathTracingRG::FPermutationDomain GetPathTracingRGPermutation(const FScene& Scene, const bool bUseCompaction)
+static FPathTracingRG::FPermutationDomain GetPathTracingRGPermutation(const FViewInfo& View, const FScene& Scene, const bool bUseCompaction)
 {
 	const bool bUseExperimental = CVarPathTracingExperimental.GetValueOnRenderThread();
 	const bool bUseAdaptiveSampling = bUseExperimental && CVarPathTracingAdaptiveSampling.GetValueOnRenderThread() != 0;
-	const bool bUseCloudShader = ShouldCompileRayTracingCallableShadersForProject(Scene.GetShaderPlatform())
-		&& CVarPathTracingEnableReferenceClouds.GetValueOnRenderThread()
-		&& !CVarPathTracingCloudAccelerationMapVisualize.GetValueOnRenderThread()
-		&& Scene.GetVolumetricCloudSceneInfo() != nullptr;
+	// NOTE: the decision about when to enable clouds involves checking lots of things, so rely on the presence of a valid index to signify that clouds are ready
+	const bool bUseCloudShader = uint32(View.PathTracingVolumetricCloudCallableShaderIndex) < Scene.RayTracingSBT.NumCallableShaderSlots
+		&& !CVarPathTracingCloudAccelerationMapVisualize.GetValueOnRenderThread();
 	const bool bHasComplexSpecialRenderPath = Substrate::IsSubstrateEnabled() && Scene.SubstrateSceneData.bUsesComplexSpecialRenderPath;
 	const bool bUseSER = GRHIGlobals.SupportsShaderExecutionReordering && CVarPathTracingShaderExecutionReordering.GetValueOnRenderThread();
 	const bool bUseTraceOpaqueFirst = CVarPathTracingTraceOpaqueFirst.GetValueOnRenderThread();
@@ -3063,8 +3073,11 @@ static FPathTracingRG::FPermutationDomain GetPathTracingRGPermutation(const FSce
 	return Out;
 }
 
-void FDeferredShadingSceneRenderer::PreparePathTracing(const FSceneViewFamily& ViewFamily, const FScene& Scene, TArray<FRHIRayTracingShader*>& OutRayGenShaders)
+void FDeferredShadingSceneRenderer::PreparePathTracing(const FViewInfo& View, const FScene& Scene, TArray<FRHIRayTracingShader*>& OutRayGenShaders)
 {
+	check(View.Family != nullptr);
+
+	const FSceneViewFamily& ViewFamily = *View.Family;
 	if (ViewFamily.EngineShowFlags.PathTracing
 		&& ShouldCompilePathTracingShadersForProject(ViewFamily.GetShaderPlatform()))
 	{
@@ -3079,13 +3092,13 @@ void FDeferredShadingSceneRenderer::PreparePathTracing(const FSceneViewFamily& V
 			// Declare all RayGen shaders that require material closest hit shaders to be bound
 			if (CompactionDepth >= 0)
 			{
-				FPathTracingRG::FPermutationDomain PermutationVector = GetPathTracingRGPermutation(Scene, false);
+				FPathTracingRG::FPermutationDomain PermutationVector = GetPathTracingRGPermutation(View, Scene, false);
 				auto RayGenShader = GetGlobalShaderMap(ViewFamily.GetShaderPlatform())->GetShader<FPathTracingRG>(PermutationVector);
 				OutRayGenShaders.Add(RayGenShader.GetRayTracingShader());
 			}
 			if (CompactionDepth != 0)
 			{
-				FPathTracingRG::FPermutationDomain PermutationVector = GetPathTracingRGPermutation(Scene, true);
+				FPathTracingRG::FPermutationDomain PermutationVector = GetPathTracingRGPermutation(View, Scene, true);
 				auto RayGenShader = GetGlobalShaderMap(ViewFamily.GetShaderPlatform())->GetShader<FPathTracingRG>(PermutationVector);
 				OutRayGenShaders.Add(RayGenShader.GetRayTracingShader());
 			}
@@ -3947,7 +3960,7 @@ void FDeferredShadingSceneRenderer::RenderPathTracing(
 
 
 							const bool bUseCompaction = Bounce < CompactionDepth || Bounce == Config.PathTracingData.MaxBounces;
-							TShaderMapRef<FPathTracingRG> RayGenShader(View.ShaderMap, GetPathTracingRGPermutation(*Scene, bUseCompaction));
+							TShaderMapRef<FPathTracingRG> RayGenShader(View.ShaderMap, GetPathTracingRGPermutation(View, *Scene, bUseCompaction));
 							ClearUnusedGraphResources(RayGenShader, PassParameters);
 							const bool bFlushRenderingCommands = FlushRenderingCommands == 1 || (FlushRenderingCommands == 2 && Bounce == CompactionDepth);
 							const bool bUse1DDispatch = (Config.UseAdaptiveSampling && Config.PathTracingData.Iteration > 0) || Bounce > 0;
