[HWRT] Improvements to r.RayTracing.Visualize
- automally set r.RayTracing.Visualize help message to show all available modes.
- for convenience, automally enable RayTracingDebug show flag when specific mode is set using cvar.

diff --git a/Engine/Source/Editor/UnrealEd/Classes/Settings/LevelEditorViewportSettings.h b/Engine/Source/Editor/UnrealEd/Classes/Settings/LevelEditorViewportSettings.h
--- a/Engine/Source/Editor/UnrealEd/Classes/Settings/LevelEditorViewportSettings.h
+++ b/Engine/Source/Editor/UnrealEd/Classes/Settings/LevelEditorViewportSettings.h
@@ -102,6 +102,7 @@ struct FLevelEditorViewportInstanceSettings
 		, SubstrateVisualizationMode()
 		, GroomVisualizationMode()
 		, VirtualShadowMapVisualizationMode()
+		, RayTracingDebugVisualizationMode()
 		, GPUSkinCacheVisualizationMode()
 		, ExposureSettings()
 		, FOVAngle(EditorViewportDefs::DefaultPerspectiveFOVAngle)
diff --git a/Engine/Source/Runtime/Engine/Private/RayTracingVisualizationData.cpp b/Engine/Source/Runtime/Engine/Private/RayTracingVisualizationData.cpp
--- a/Engine/Source/Runtime/Engine/Private/RayTracingVisualizationData.cpp
+++ b/Engine/Source/Runtime/Engine/Private/RayTracingVisualizationData.cpp
@@ -9,6 +9,16 @@
 
 #define LOCTEXT_NAMESPACE "FRayTracingVisualizationData"
 
+static FName GRayTracingVisualizeMode = NAME_None;
+static FAutoConsoleVariableRef CVarRayTracingVisualize(
+	TEXT("r.RayTracing.Visualize"),
+	GRayTracingVisualizeMode,
+	TEXT("Sets the ray tracing debug visualization mode (default = None - Driven by viewport menu) .\n"),
+	ECVF_RenderThreadSafe
+);
+
+static FAutoConsoleVariableDeprecated CVarRayTracingDebugMode_Deprecated(TEXT("r.RayTracing.DebugVisualizationMode"), TEXT("r.RayTracing.Visualize"), TEXT("5.6"));
+
 FRayTracingVisualizationData::FRayTracingVisualizationData()
 {
 	const bool bRayTracingShadersSupported = FDataDrivenShaderPlatformInfo::GetSupportsRayTracingShaders(GMaxRHIShaderPlatform);
@@ -81,6 +91,8 @@ FRayTracingVisualizationData::FRayTracingVisualizationData()
 			AddVisualizationMode(TEXT("SubstrateMaterialProperties"), LOCTEXT("SubstrateMaterialProperties", "Substrate Material Properties"), FModeType::Other, RAY_TRACING_DEBUG_VIZ_SUBSTRATE_DATA, true);
 		}
 	}
+
+	ConfigureConsoleCommand();
 }
 
 void FRayTracingVisualizationData::AddVisualizationMode(
@@ -141,6 +153,33 @@ bool FRayTracingVisualizationData::GetModeTonemapped(const FName& InModeName) co
 	}
 }
 
+bool FRayTracingVisualizationData::HasOverrides() const
+{
+	return GRayTracingVisualizeMode != NAME_None && GetModeID(GRayTracingVisualizeMode) != INDEX_NONE;
+}
+
+FName FRayTracingVisualizationData::ApplyOverrides(const FName& InModeName) const
+{
+	check(IsInParallelRenderingThread());
+	return GRayTracingVisualizeMode != NAME_None ? GRayTracingVisualizeMode : InModeName;
+}
+
+void FRayTracingVisualizationData::ConfigureConsoleCommand()
+{
+	FString AvailableVisualizationModes;
+	for (TModeMap::TConstIterator It = ModeMap.CreateConstIterator(); It; ++It)
+	{
+		const FModeRecord& Record = It.Value();
+		AvailableVisualizationModes += FString(TEXT("\n  "));
+		AvailableVisualizationModes += Record.ModeString;
+	}
+
+	ConsoleDocumentationVisualizationMode = TEXT("When the viewport view-mode is set to 'Ray Tracing Debug', this command specifies which of the various modes to display. Values entered other than the allowed values shown below will be ignored.");
+	ConsoleDocumentationVisualizationMode += AvailableVisualizationModes;
+
+	CVarRayTracingVisualize->SetHelp(*ConsoleDocumentationVisualizationMode);
+}
+
 FRayTracingVisualizationData& GetRayTracingVisualizationData()
 {
 	static FRayTracingVisualizationData GRayTracingVisualizationData;
diff --git a/Engine/Source/Runtime/Engine/Private/UnrealEngine.cpp b/Engine/Source/Runtime/Engine/Private/UnrealEngine.cpp
--- a/Engine/Source/Runtime/Engine/Private/UnrealEngine.cpp
+++ b/Engine/Source/Runtime/Engine/Private/UnrealEngine.cpp
@@ -175,6 +175,7 @@ UnrealEngine.cpp: Implements the UEngine class and helpers.
 #include "BufferVisualizationData.h"
 #include "NaniteVisualizationData.h"
 #include "VirtualShadowMapVisualizationData.h"
+#include "RayTracingVisualizationData.h"
 #include "Misc/HotReloadInterface.h"
 #include "Widgets/Testing/STestSuite.h"
 #include "Engine/DemoNetDriver.h"
@@ -2266,6 +2267,7 @@ void UEngine::Init(IEngineLoop* InEngineLoop)
 	GetBufferVisualizationData().Initialize();
 	GetNaniteVisualizationData().Initialize();
 	GetVirtualShadowMapVisualizationData().Initialize();
+	GetRayTracingVisualizationData(); // force initialization
 
 	// Initialize Portal services
 	if (!IsRunningCommandlet() && !IsRunningDedicatedServer())
diff --git a/Engine/Source/Runtime/Engine/Public/RayTracingVisualizationData.h b/Engine/Source/Runtime/Engine/Public/RayTracingVisualizationData.h
--- a/Engine/Source/Runtime/Engine/Public/RayTracingVisualizationData.h
+++ b/Engine/Source/Runtime/Engine/Public/RayTracingVisualizationData.h
@@ -36,6 +36,9 @@ public:
 public:
 	FRayTracingVisualizationData();
 
+	ENGINE_API bool HasOverrides() const;
+	ENGINE_API FName ApplyOverrides(const FName& InModeName) const;
+
 	/** Get the display name of a named mode from the available mode map. **/
 	ENGINE_API FText GetModeDisplayName(const FName& InModeName) const;
 
@@ -59,6 +62,9 @@ private:
 		bool bHiddenInEditor = false
 	);
 
+	/** Internal helper function for creating the visualization system console commands. */
+	void ConfigureConsoleCommand();
+
 private:
 	/** The name->mode mapping table */
 	TModeMap ModeMap;
diff --git a/Engine/Source/Runtime/Renderer/Private/DeferredShadingRenderer.cpp b/Engine/Source/Runtime/Renderer/Private/DeferredShadingRenderer.cpp
--- a/Engine/Source/Runtime/Renderer/Private/DeferredShadingRenderer.cpp
+++ b/Engine/Source/Runtime/Renderer/Private/DeferredShadingRenderer.cpp
@@ -46,6 +46,7 @@
 #include "RayTracing/RayTracingScene.h"
 #include "RayTracing/RayTracingInstanceMask.h"
 #include "RayTracingDynamicGeometryUpdateManager.h"
+#include "RayTracingVisualizationData.h"
 #include "SceneTextureParameters.h"
 #include "ScreenSpaceDenoise.h"
 #include "ScreenSpaceRayTracing.h"
@@ -1446,6 +1447,16 @@ void FDeferredShadingSceneRenderer::Render(FRDGBuilder& GraphBuilder, const FSce
 		return;
 	}
 
+	{
+		FRayTracingVisualizationData& RayTracingVisualizationData = GetRayTracingVisualizationData();
+
+		if (RayTracingVisualizationData.HasOverrides())
+		{
+			// When activating the view modes from the command line, automatically enable the RayTracingDebug show flag for convenience.
+			ViewFamily.EngineShowFlags.SetRayTracingDebug(true);
+		}
+	}
+
 	// If this is scene capture rendering depth pre-pass, we'll take the shortcut function RenderSceneCaptureDepth if optimization switch is on.
 	const ERendererOutput RendererOutput = GetRendererOutput();
 
diff --git a/Engine/Source/Runtime/Renderer/Private/RayTracing/RayTracingDebug.cpp b/Engine/Source/Runtime/Renderer/Private/RayTracing/RayTracingDebug.cpp
--- a/Engine/Source/Runtime/Renderer/Private/RayTracing/RayTracingDebug.cpp
+++ b/Engine/Source/Runtime/Renderer/Private/RayTracing/RayTracingDebug.cpp
@@ -29,16 +29,6 @@
 
 DECLARE_GPU_STAT(RayTracingDebug);
 
-static FName GRayTracingVisualizeMode = NAME_None;
-static FAutoConsoleVariableRef CVarRayTracingVisualize(
-	TEXT("r.RayTracing.Visualize"),
-	GRayTracingVisualizeMode,
-	TEXT("Sets the ray tracing debug visualization mode (default = None - Driven by viewport menu) .\n"),
-	ECVF_RenderThreadSafe
-);
-
-static FAutoConsoleVariableDeprecated CVarRayTracingDebugMode_Deprecated(TEXT("r.RayTracing.DebugVisualizationMode"), TEXT("r.RayTracing.Visualize"), TEXT("5.6"));
-
 TAutoConsoleVariable<int32> CVarRayTracingVisualizePickerDomain(
 	TEXT("r.RayTracing.Visualize.PickerDomain"),
 	0,
@@ -1540,30 +1530,25 @@ static void DrawInstanceOverlap(FRDGBuilder& GraphBuilder, const FScene& Scene,
 	}
 }
 
-static uint32 GetRaytracingDebugViewMode(const FSceneView& View)
+static FName GetRaytracingDebugViewMode(const FSceneView& View)
 {
 	check(IsInParallelRenderingThread());
-	FName CurrentMode = GRayTracingVisualizeMode != NAME_None ? GRayTracingVisualizeMode : View.CurrentRayTracingDebugVisualizationMode;
-	if (CurrentMode != NAME_None)
-	{
-		return GetRayTracingVisualizationData().GetModeID(CurrentMode);
-	}
-	// Set useful default value
-	return RAY_TRACING_DEBUG_VIZ_BARYCENTRICS;
+	FRayTracingVisualizationData& RayTracingVisualizationData = GetRayTracingVisualizationData();
+	FName CurrentMode = RayTracingVisualizationData.ApplyOverrides(View.CurrentRayTracingDebugVisualizationMode);
+
+	// Use barycentrics as default when mode is not specified
+	return CurrentMode != NAME_None ? CurrentMode : TEXT("Barycentrics");
 }
 
-bool RaytracingDebugViewModeNeedsTonemapping(const FSceneView& View)
+static uint32 GetRaytracingDebugViewModeID(const FSceneView& View)
 {
-	check(IsInParallelRenderingThread());
-	FName CurrentMode = GRayTracingVisualizeMode != NAME_None ? GRayTracingVisualizeMode : View.CurrentRayTracingDebugVisualizationMode;
-	if (CurrentMode != NAME_None)
-	{
-		return GetRayTracingVisualizationData().GetModeTonemapped(CurrentMode);
-	}
-	// unknown mode, assume no tone-mapping is needed
-	return false;
+	return GetRayTracingVisualizationData().GetModeID(GetRaytracingDebugViewMode(View));
 }
 
+bool RaytracingDebugViewModeNeedsTonemapping(const FSceneView& View)
+{
+	return GetRayTracingVisualizationData().GetModeTonemapped(GetRaytracingDebugViewMode(View));
+}
 
 bool HasRaytracingDebugViewModeRaytracedOverlay(const FSceneViewFamily& ViewFamily)
 {
@@ -1572,7 +1557,7 @@ bool HasRaytracingDebugViewModeRaytracedOverlay(const FSceneViewFamily& ViewFami
 
 	for (const FSceneView* View : ViewFamily.Views)
 	{
-		const uint32 Mode = View != nullptr ? GetRaytracingDebugViewMode(*View) : UINT32_MAX;
+		const uint32 Mode = View != nullptr ? GetRaytracingDebugViewModeID(*View) : UINT32_MAX;
 
 		bAnySubstrate |= Mode == RAY_TRACING_DEBUG_VIZ_SUBSTRATE_DATA;
 
@@ -1601,7 +1586,7 @@ void RenderRayTracingDebug(FRDGBuilder& GraphBuilder, const FScene& Scene, const
 	FSceneTextureParameters SceneTextureParameters = GetSceneTextureParameters(GraphBuilder, SceneTextures.UniformBuffer);
 	FRDGTextureRef SceneColorTexture = SceneTextures.Color.Target;
 
-	const uint32 DebugVisualizationMode = GetRaytracingDebugViewMode(View);
+	const uint32 DebugVisualizationMode = GetRaytracingDebugViewModeID(View);
 
 	if (DebugVisualizationMode == RAY_TRACING_DEBUG_VIZ_SUBSTRATE_DATA)
 	{
@@ -1918,7 +1903,7 @@ void RenderRayTracingDebug(FRDGBuilder& GraphBuilder, const FScene& Scene, const
 
 extern void RayTracingDebugDisplayOnScreenMessages(FScreenMessageWriter& Writer, const FViewInfo& View)
 {
-	const uint32 DebugVisualizationMode = GetRaytracingDebugViewMode(View);
+	const uint32 DebugVisualizationMode = GetRaytracingDebugViewModeID(View);
 
 	switch (DebugVisualizationMode)
 	{
@@ -2107,7 +2092,7 @@ void FDeferredShadingSceneRenderer::RayTracingDisplayPicking(const FRayTracingPi
 bool IsRayTracingInstanceDebugDataEnabled(const FViewInfo& View)
 {
 #if !UE_BUILD_SHIPPING
-	const uint32 Mode = GetRaytracingDebugViewMode(View);
+	const uint32 Mode = GetRaytracingDebugViewModeID(View);
 
 	return Mode == RAY_TRACING_DEBUG_VIZ_DYNAMIC_INSTANCES
 		|| Mode == RAY_TRACING_DEBUG_VIZ_PROXY_TYPE
@@ -2120,7 +2105,7 @@ bool IsRayTracingInstanceDebugDataEnabled(const FViewInfo& View)
 bool IsRayTracingInstanceOverlapEnabled(const FViewInfo& View)
 {
 #if !UE_BUILD_SHIPPING	
-	return GetRaytracingDebugViewMode(View) == RAY_TRACING_DEBUG_VIZ_INSTANCE_OVERLAP;
+	return GetRaytracingDebugViewModeID(View) == RAY_TRACING_DEBUG_VIZ_INSTANCE_OVERLAP;
 #else
 	return false;
 #endif
