Prevent flickering on the first frame when showing primitives that were previously force hidden
* When the AddToWorld() for a level was performed over multiple frames, the occlusion culling "history" logic for these force hidden primitives would mark them as occluded as they were not rendered

diff --git a/Engine/Source/Runtime/Renderer/Private/PrimitiveSceneInfo.cpp b/Engine/Source/Runtime/Renderer/Private/PrimitiveSceneInfo.cpp
--- a/Engine/Source/Runtime/Renderer/Private/PrimitiveSceneInfo.cpp
+++ b/Engine/Source/Runtime/Renderer/Private/PrimitiveSceneInfo.cpp
@@ -1658,6 +1658,40 @@ void FPrimitiveSceneInfo::FreeGPUSceneInstances()
 	}
 }
 
+void FPrimitiveSceneInfo::UpdateOcclusionFlags()
+{
+	if (IsIndexValid())
+	{
+		uint8 OcclusionFlags = EOcclusionFlags::None;
+		// First person primitives potentially deform the geometry outside of its bounds in a view dependent way. They are very unlikely to be occluded anyways,
+		// so to avoid falsely culling them, it is better to simply don't occlusion cull them at all.
+		if (Proxy->CanBeOccluded() && !Proxy->IsFirstPerson())
+		{
+			OcclusionFlags |= EOcclusionFlags::CanBeOccluded;
+		}
+		if (Proxy->HasSubprimitiveOcclusionQueries())
+		{
+			OcclusionFlags |= EOcclusionFlags::HasSubprimitiveQueries;
+		}
+		if (Proxy->AllowApproximateOcclusion()
+			// Allow approximate occlusion if attached, even if the parent does not have bLightAttachmentsAsGroup enabled
+			|| LightingAttachmentRoot.IsValid())
+		{
+			OcclusionFlags |= EOcclusionFlags::AllowApproximateOcclusion;
+		}
+		if (Proxy->GetVisibilityId() >= 0)
+		{
+			OcclusionFlags |= EOcclusionFlags::HasPrecomputedVisibility;
+		}
+		if (Proxy->IsForceHidden())
+		{
+			OcclusionFlags |= EOcclusionFlags::IsForceHidden;
+		}
+
+		Scene->PrimitiveOcclusionFlags[PackedIndex] = OcclusionFlags;
+	}	
+}
+
 void FPrimitiveSceneInfo::AddToScene(FScene* Scene, TArrayView<FPrimitiveSceneInfo*> SceneInfos)
 {
 	check(IsInRenderingThread());
@@ -1803,28 +1837,7 @@ void FPrimitiveSceneInfo::AddToScene(FScene* Scene, TArrayView<FPrimitiveSceneIn
 			VisibilityId.BitMask = (1 << (VisibilityBitIndex & 0x7));
 
 			// Store occlusion flags.
-			uint8 OcclusionFlags = EOcclusionFlags::None;
-			// First person primitives potentially deform the geometry outside of its bounds in a view dependent way. They are very unlikely to be occluded anyways,
-			// so to avoid falsely culling them, it is better to simply don't occlusion cull them at all.
-			if (Proxy->CanBeOccluded() && !Proxy->IsFirstPerson())
-			{
-				OcclusionFlags |= EOcclusionFlags::CanBeOccluded;
-			}
-			if (Proxy->HasSubprimitiveOcclusionQueries())
-			{
-				OcclusionFlags |= EOcclusionFlags::HasSubprimitiveQueries;
-			}
-			if (Proxy->AllowApproximateOcclusion()
-				// Allow approximate occlusion if attached, even if the parent does not have bLightAttachmentsAsGroup enabled
-				|| SceneInfo->LightingAttachmentRoot.IsValid())
-			{
-				OcclusionFlags |= EOcclusionFlags::AllowApproximateOcclusion;
-			}
-			if (VisibilityBitIndex >= 0)
-			{
-				OcclusionFlags |= EOcclusionFlags::HasPrecomputedVisibility;
-			}
-			Scene->PrimitiveOcclusionFlags[PackedIndex] = OcclusionFlags;
+			SceneInfo->UpdateOcclusionFlags();
 
 			// Store occlusion bounds.
 			FBoxSphereBounds OcclusionBounds = BoxSphereBounds;
diff --git a/Engine/Source/Runtime/Renderer/Private/RendererScene.cpp b/Engine/Source/Runtime/Renderer/Private/RendererScene.cpp
--- a/Engine/Source/Runtime/Renderer/Private/RendererScene.cpp
+++ b/Engine/Source/Runtime/Renderer/Private/RendererScene.cpp
@@ -6502,6 +6502,11 @@ void FScene::Update(FRDGBuilder& GraphBuilder, const FUpdateParameters& Paramete
 				for (FPrimitiveSceneInfo* Primitive : *LevelPrimitives)
 				{
 					Primitive->Proxy->OnLevelAddedToWorld_RenderThread();
+
+					// Primitive was force hidden & excluded from the occlusion culling processing
+					// Now that it's going to be rendered, update it's occlusion flags so that it is considered.
+					Primitive->UpdateOcclusionFlags();
+
 					InvalidateNaniteGPUScene(Primitive);
 				}
 			}
@@ -6513,6 +6518,10 @@ void FScene::Update(FRDGBuilder& GraphBuilder, const FUpdateParameters& Paramete
 				for (FPrimitiveSceneInfo* Primitive : *LevelPrimitives)
 				{
 					Primitive->Proxy->OnLevelRemovedFromWorld_RenderThread();
+
+					// Primitive is now force hidden & should be excluded from the occlusion culling processing
+					Primitive->UpdateOcclusionFlags();
+
 					InvalidateNaniteGPUScene(Primitive);
 				}
 			}
diff --git a/Engine/Source/Runtime/Renderer/Private/ScenePrivate.h b/Engine/Source/Runtime/Renderer/Private/ScenePrivate.h
--- a/Engine/Source/Runtime/Renderer/Private/ScenePrivate.h
+++ b/Engine/Source/Runtime/Renderer/Private/ScenePrivate.h
@@ -2498,6 +2498,8 @@ namespace EOcclusionFlags
 		HasPrecomputedVisibility = 0x8,
 		/** Indicates the primitive supports subprimitive occlusion queries. */
 		HasSubprimitiveQueries = 0x10,
+		/** Indicates that the primitive is forced hidden and does not need to be occlusion tested. */
+		IsForceHidden = 0x20,
 	};
 };
 
diff --git a/Engine/Source/Runtime/Renderer/Private/SceneVisibility.cpp b/Engine/Source/Runtime/Renderer/Private/SceneVisibility.cpp
--- a/Engine/Source/Runtime/Renderer/Private/SceneVisibility.cpp
+++ b/Engine/Source/Runtime/Renderer/Private/SceneVisibility.cpp
@@ -2602,6 +2602,12 @@ bool FGPUOcclusionPacket::OcclusionCullPrimitive(VisitorType& Visitor, FOcclusio
 		}
 	};
 
+	if (OcclusionFlags & EOcclusionFlags::IsForceHidden)
+	{
+		// Primitive is not visible and does not need to be occlusion tested
+		return false;
+	}
+
 	if ((OcclusionFlags & EOcclusionFlags::HasSubprimitiveQueries) && OcclusionState.bAllowSubQueries)
 	{
 		FPrimitiveSceneProxy* Proxy = Scene.PrimitiveSceneProxies[Index];
diff --git a/Engine/Source/Runtime/Renderer/Public/PrimitiveSceneInfo.h b/Engine/Source/Runtime/Renderer/Public/PrimitiveSceneInfo.h
--- a/Engine/Source/Runtime/Renderer/Public/PrimitiveSceneInfo.h
+++ b/Engine/Source/Runtime/Renderer/Public/PrimitiveSceneInfo.h
@@ -434,6 +434,9 @@ public:
 	/** Refreshes a primitive's references to raster bins. To be called after changes that might have invalidated them. */
 	RENDERER_API void RefreshNaniteRasterBins();
 
+	/** Update the occlusion flags for this primitive. */
+	void UpdateOcclusionFlags();
+
 	/** 
 	 * Builds an array of all primitive scene info's in this primitive's attachment group. 
 	 * This only works on potential parents (!LightingAttachmentRoot.IsValid()) and will include the current primitive in the output array.
