Add RequestPreload() function to runtime virtual texture compoent.
This attempts to fill the pages at the given mip level in the given world space bounds.

diff --git a/Engine/Source/Runtime/Engine/Classes/Components/RuntimeVirtualTextureComponent.h b/Engine/Source/Runtime/Engine/Classes/Components/RuntimeVirtualTextureComponent.h
--- a/Engine/Source/Runtime/Engine/Classes/Components/RuntimeVirtualTextureComponent.h
+++ b/Engine/Source/Runtime/Engine/Classes/Components/RuntimeVirtualTextureComponent.h
@@ -156,6 +156,14 @@ public:
 	UFUNCTION(BlueprintCallable, Category = VirtualTexture)
 	ENGINE_API void Invalidate(FBoxSphereBounds const& WorldBounds, EVTInvalidatePriority InvalidatePriority = EVTInvalidatePriority::High);
 
+	/** 
+	 * Request preload of an area of the runtime virtual texture at a given mip level. 
+	 * @param WorldBounds : The world space bounds of the pages to preload.
+	 * @param Level : The mip map level to preload.
+	 */
+	UFUNCTION(BlueprintCallable, Category = VirtualTexture)
+	ENGINE_API void RequestPreload(FBoxSphereBounds const& WorldBounds, int32 Level);
+
 	/** Set the runtime virtual texture object on this component. */
 	ENGINE_API void SetVirtualTexture(URuntimeVirtualTexture* InVirtualTexture);
 
diff --git a/Engine/Source/Runtime/Engine/Private/Components/RuntimeVirtualTextureComponent.cpp b/Engine/Source/Runtime/Engine/Private/Components/RuntimeVirtualTextureComponent.cpp
--- a/Engine/Source/Runtime/Engine/Private/Components/RuntimeVirtualTextureComponent.cpp
+++ b/Engine/Source/Runtime/Engine/Private/Components/RuntimeVirtualTextureComponent.cpp
@@ -252,6 +252,14 @@ void URuntimeVirtualTextureComponent::Invalidate(FBoxSphereBounds const& InWorld
 	}
 }
 
+void URuntimeVirtualTextureComponent::RequestPreload(FBoxSphereBounds const& InWorldBounds, int32 InLevel)
+{
+	if (GetScene() != nullptr)
+	{
+		GetScene()->RequestPreloadRuntimeVirtualTexture(this, InWorldBounds, InLevel);
+	}
+}
+
 FBoxSphereBounds URuntimeVirtualTextureComponent::CalcBounds(const FTransform& LocalToWorld) const
 {
 	return FBoxSphereBounds(FBox(FVector(0.f, 0.f, 0.f), FVector(1.f, 1.f, 1.f))).TransformBy(LocalToWorld);
diff --git a/Engine/Source/Runtime/Engine/Public/SceneInterface.h b/Engine/Source/Runtime/Engine/Public/SceneInterface.h
--- a/Engine/Source/Runtime/Engine/Public/SceneInterface.h
+++ b/Engine/Source/Runtime/Engine/Public/SceneInterface.h
@@ -50,6 +50,7 @@ class ULightComponent;
 class UPlanarReflectionComponent;
 class UPrimitiveComponent;
 class UReflectionCaptureComponent;
+class URuntimeVirtualTextureComponent;
 class USkyLightComponent;
 class UStaticMeshComponent;
 class UTextureCube;
@@ -286,24 +287,25 @@ public:
 	virtual void RemovePrecomputedVolumetricLightmap(const class FPrecomputedVolumetricLightmap* Volume) {}
 
 	/** Add a runtime virtual texture object to the scene. */
-	virtual void AddRuntimeVirtualTexture(class URuntimeVirtualTextureComponent* Component) {}
-
+	virtual void AddRuntimeVirtualTexture(URuntimeVirtualTextureComponent* Component) {}
 	/** Removes a runtime virtual texture object from the scene. */
-	virtual void RemoveRuntimeVirtualTexture(class URuntimeVirtualTextureComponent* Component) {}
-
+	virtual void RemoveRuntimeVirtualTexture(URuntimeVirtualTextureComponent* Component) {}
 	/* Get the bitmasks describing which virtual texture objects will hide the associated primitives. */
 	virtual void GetRuntimeVirtualTextureHidePrimitiveMask(uint8& bHideMaskEditor, uint8& bHideMaskGame) const {}
-
-	UE_DEPRECATED(5.6, "Use the alternate prototype with the Priority param")
-	virtual void InvalidateRuntimeVirtualTexture(class URuntimeVirtualTextureComponent* Component, FBoxSphereBounds const& WorldBounds) {}
-
-	/** Invalidates pages in a runtime virtual texture object. 
+	
+	/** 
+	 * Invalidates pages in a runtime virtual texture object. 
 	 * @param Component RVT component containing the runtime virtual texture to invalidate
 	 * @param WorldBounds World-space bounds of the area of the runtime virtual texture to invalidate
 	 * @param InvalidatePriority Ensures the corresponding pages get processed in priority. This allows increased responsiveness when 
 	 *  there are more pages being updated than can be handled in a given frame (when throttling) 
 	 */
-	virtual void InvalidateRuntimeVirtualTexture(class URuntimeVirtualTextureComponent* Component, FBoxSphereBounds const& WorldBounds, EVTInvalidatePriority InvalidatePriority) {}
+	virtual void InvalidateRuntimeVirtualTexture(URuntimeVirtualTextureComponent* Component, FBoxSphereBounds const& WorldBounds, EVTInvalidatePriority InvalidatePriority) {}
+	UE_DEPRECATED(5.6, "Use the alternate prototype with the Priority param")
+	virtual void InvalidateRuntimeVirtualTexture(URuntimeVirtualTextureComponent* Component, FBoxSphereBounds const& WorldBounds) {}
+
+	/** Request preload of an area of a runtime virtual texture at a given mip level. */
+	virtual void RequestPreloadRuntimeVirtualTexture(class URuntimeVirtualTextureComponent* Component, FBoxSphereBounds const& WorldBounds, int32 Level) {}
 
 	/** Mark scene as needing to restart path tracer accumulation. */
 	virtual void InvalidatePathTracedOutput() {}
diff --git a/Engine/Source/Runtime/Renderer/Private/RendererScene.cpp b/Engine/Source/Runtime/Renderer/Private/RendererScene.cpp
--- a/Engine/Source/Runtime/Renderer/Private/RendererScene.cpp
+++ b/Engine/Source/Runtime/Renderer/Private/RendererScene.cpp
@@ -3699,6 +3699,19 @@ void FScene::FlushDirtyRuntimeVirtualTextures()
 	}
 }
 
+void FScene::RequestPreloadRuntimeVirtualTexture(class URuntimeVirtualTextureComponent* Component, FBoxSphereBounds const& WorldBounds, int32 Level)
+{
+	if (Component->SceneProxy != nullptr)
+	{
+		FRuntimeVirtualTextureSceneProxy* SceneProxy = Component->SceneProxy;
+		ENQUEUE_RENDER_COMMAND(PreloadRuntimeVirtualTextureCommand)(
+			[SceneProxy, WorldBounds, Level](FRHICommandListBase&)
+		{
+			SceneProxy->RequestPreload(WorldBounds, Level);
+		});
+	}
+}
+
 void FScene::InvalidatePathTracedOutput()
 {
 	// NOTE: this is an atomic, so this function is ok to call from any thread
diff --git a/Engine/Source/Runtime/Renderer/Private/ScenePrivate.h b/Engine/Source/Runtime/Renderer/Private/ScenePrivate.h
--- a/Engine/Source/Runtime/Renderer/Private/ScenePrivate.h
+++ b/Engine/Source/Runtime/Renderer/Private/ScenePrivate.h
@@ -3410,6 +3410,7 @@ public:
 	virtual void RemoveRuntimeVirtualTexture(class URuntimeVirtualTextureComponent* Component) override;
 	virtual void GetRuntimeVirtualTextureHidePrimitiveMask(uint8& bHideMaskEditor, uint8& bHideMaskGame) const override;
 	virtual void InvalidateRuntimeVirtualTexture(class URuntimeVirtualTextureComponent* Component, FBoxSphereBounds const& WorldBounds, EVTInvalidatePriority InvalidatePriority) override;
+	virtual void RequestPreloadRuntimeVirtualTexture(class URuntimeVirtualTextureComponent* Component, FBoxSphereBounds const& WorldBounds, int32 Level) override;
 	virtual void InvalidatePathTracedOutput() override;
 	virtual void InvalidateLumenSurfaceCache_GameThread(UPrimitiveComponent* Component) override;
 	virtual void GetPrimitiveUniformShaderParameters_RenderThread(const FPrimitiveSceneInfo* PrimitiveSceneInfo, bool& bHasPrecomputedVolumetricLightmap, FMatrix& PreviousLocalToWorld, int32& SingleCaptureIndex, bool& bOutputVelocity) const override;
diff --git a/Engine/Source/Runtime/Renderer/Private/VT/RuntimeVirtualTextureSceneProxy.cpp b/Engine/Source/Runtime/Renderer/Private/VT/RuntimeVirtualTextureSceneProxy.cpp
--- a/Engine/Source/Runtime/Renderer/Private/VT/RuntimeVirtualTextureSceneProxy.cpp
+++ b/Engine/Source/Runtime/Renderer/Private/VT/RuntimeVirtualTextureSceneProxy.cpp
@@ -123,12 +123,13 @@ FRuntimeVirtualTextureSceneProxy::FRuntimeVirtualTextureSceneProxy(URuntimeVirtu
 			// The Initialize() call will allocate the virtual texture by spawning work on the render thread.
 			VirtualTexture->Initialize(Producer, ProducerDesc, Transform, Bounds);
 
-			// Store the ProducerHandle and SpaceID immediately after virtual texture is initialized.
+			// Store the ProducerHandle, SpaceID and AllocatedVirtualTexture immediately after virtual texture is initialized.
 			ENQUEUE_RENDER_COMMAND(GetProducerHandle)(
 				[this, VirtualTexturePtr = VirtualTexture](FRHICommandList& RHICmdList)
 				{
 					ProducerHandle = VirtualTexturePtr->GetProducerHandle();
-					SpaceID = VirtualTexturePtr->GetAllocatedVirtualTexture()->GetSpaceID();
+					AllocatedVirtualTexture = VirtualTexturePtr->GetAllocatedVirtualTexture();
+					SpaceID = AllocatedVirtualTexture->GetSpaceID();
 				});
 		}
 		else
@@ -156,18 +157,24 @@ void FRuntimeVirtualTextureSceneProxy::Release()
 	}
 }
 
-void FRuntimeVirtualTextureSceneProxy::Dirty(FBoxSphereBounds const& Bounds, EVTInvalidatePriority InvalidatePriority)
+/** Transform world bounds into Virtual Texture UV space. */
+static FBox2D GetUVRectFromWorldBounds(FTransform const& InTransform, FBoxSphereBounds const& InBounds)
 {
-	// Transform world bounds into Virtual Texture UV space
-	const FVector O = Transform.GetTranslation();
-	const FVector U = Transform.GetUnitAxis(EAxis::X) * 1.f / Transform.GetScale3D().X;
-	const FVector V = Transform.GetUnitAxis(EAxis::Y) * 1.f / Transform.GetScale3D().Y;
-	const FVector P = Bounds.GetSphere().Center - O;
+	const FVector O = InTransform.GetTranslation();
+	const FVector U = InTransform.GetUnitAxis(EAxis::X) * 1.f / InTransform.GetScale3D().X;
+	const FVector V = InTransform.GetUnitAxis(EAxis::Y) * 1.f / InTransform.GetScale3D().Y;
+	const FVector P = InBounds.GetSphere().Center - O;
 	const FVector2D UVCenter = FVector2D(FVector::DotProduct(P, U), FVector::DotProduct(P, V));
-	const float Scale = FMath::Max(1.f / Transform.GetScale3D().X, 1.f / Transform.GetScale3D().Y);
-	const float UVRadius = Bounds.GetSphere().W * Scale;
+	const float Scale = FMath::Max(1.f / InTransform.GetScale3D().X, 1.f / InTransform.GetScale3D().Y);
+	const float UVRadius = InBounds.GetSphere().W * Scale;
 	const FVector2D UVExtent(UVRadius, UVRadius);
-	const FBox2D UVRect = FBox2D(UVCenter - UVExtent, UVCenter + UVExtent);
+
+	return FBox2D(UVCenter - UVExtent, UVCenter + UVExtent);
+}
+
+void FRuntimeVirtualTextureSceneProxy::Dirty(FBoxSphereBounds const& InBounds, EVTInvalidatePriority InInvalidatePriority)
+{
+	const FBox2D UVRect = GetUVRectFromWorldBounds(Transform, InBounds);
 
 	// Convert to Texel coordinate space
 	const FIntRect TextureRect(0, 0, VirtualTextureSize.X, VirtualTextureSize.Y);
@@ -181,7 +188,7 @@ void FRuntimeVirtualTextureSceneProxy::Dirty(FBoxSphereBounds const& Bounds, EVT
 	// Only add rect if it has some area
 	if (TexelRect.Min != TexelRect.Max)
 	{
-		FDirtyRect DirtyRect{ .Rect = TexelRect, .InvalidatePriority = InvalidatePriority };
+		FDirtyRect DirtyRect{ .Rect = TexelRect, .InvalidatePriority = InInvalidatePriority };
 		const bool bFirst = DirtyRects.Add(DirtyRect) == 0;
 		if (bFirst)
 		{
@@ -232,4 +239,11 @@ void FRuntimeVirtualTextureSceneProxy::FlushDirtyPages()
 	CombinedDirtyRect = FDirtyRect();
 }
 
+void FRuntimeVirtualTextureSceneProxy::RequestPreload(FBoxSphereBounds const& InBounds, int32 InLevel)
+{
+	const FBox2D UVRect = GetUVRectFromWorldBounds(Transform, InBounds);
+
+	FVirtualTextureSystem::Get().RequestTiles(AllocatedVirtualTexture, FVector2D::One(), FVector2D::Zero(), FVector2D::One(), UVRect.Min, UVRect.Max, InLevel);
+}
+
 #undef LOCTEXT_NAMESPACE
diff --git a/Engine/Source/Runtime/Renderer/Private/VT/RuntimeVirtualTextureSceneProxy.h b/Engine/Source/Runtime/Renderer/Private/VT/RuntimeVirtualTextureSceneProxy.h
--- a/Engine/Source/Runtime/Renderer/Private/VT/RuntimeVirtualTextureSceneProxy.h
+++ b/Engine/Source/Runtime/Renderer/Private/VT/RuntimeVirtualTextureSceneProxy.h
@@ -34,6 +34,9 @@ public:
 	/** Flush the cached physical pages of the virtual texture at the given world space bounds. */
 	void FlushDirtyPages();
 
+	/** Request preload of an area of the associated runtime virtual texture at a given mip level. */
+	void RequestPreload(FBoxSphereBounds const& Bounds, int32 Level);
+
 	/** Index in FScene::RuntimeVirtualTextures. */
 	int32 SceneIndex = -1;
 	
@@ -56,6 +59,8 @@ private:
 
 	/** Handle for the producer that this Proxy initialized. Used only for invalidation logic. */
 	FVirtualTextureProducerHandle ProducerHandle;
+	/** Allocated virtual texture that this Proxy initialized. Used only for preload logic. */
+	IAllocatedVirtualTexture* AllocatedVirtualTexture = nullptr;
 	/** Space ID used by the virtual texture. Used only for invalidation logic. */
 	int32 SpaceID = -1;
 
