Path Tracer: Remove Reference Clouds checkbox from PPV and enable by default instead (when reference atmosphere is enabled)

This seems to be the preferred behavior after talking to some internal users. The Cloud actor can simply be hidden if they are not needed.

#rnx

diff --git a/Engine/Source/Runtime/Engine/Classes/Engine/Scene.h b/Engine/Source/Runtime/Engine/Classes/Engine/Scene.h
--- a/Engine/Source/Runtime/Engine/Classes/Engine/Scene.h
+++ b/Engine/Source/Runtime/Engine/Classes/Engine/Scene.h
@@ -1350,9 +1350,6 @@ struct FPostProcessSettings
 	UPROPERTY(EditAnywhere, BlueprintReadWrite, Category = Overrides, meta = (PinHiddenByDefault, InlineEditConditionToggle))
 	uint32 bOverride_PathTracingEnableReferenceAtmosphere : 1;
 
-	UPROPERTY(EditAnywhere, BlueprintReadWrite, Category = Overrides, meta = (PinHiddenByDefault, InlineEditConditionToggle))
-	uint32 bOverride_PathTracingEnableReferenceClouds : 1;
-
 	UPROPERTY(EditAnywhere, BlueprintReadWrite, Category = Overrides, meta = (PinHiddenByDefault, InlineEditConditionToggle))
 	uint32 bOverride_PathTracingEnableDenoiser : 1;
 
@@ -2352,10 +2349,6 @@ struct FPostProcessSettings
 	UPROPERTY(EditAnywhere, BlueprintReadWrite, Category = "Path Tracing", meta = (editcondition = "bOverride_PathTracingEnableReferenceAtmosphere", DisplayName = "Reference Atmosphere"))
 	uint32 PathTracingEnableReferenceAtmosphere : 1;
 
-	/** Enables path tracing of volumetric clouds. This requires Reference Atmosphere. */
-	UPROPERTY(EditAnywhere, BlueprintReadWrite, Category = "Path Tracing", meta = (editcondition = "bOverride_PathTracingEnableReferenceClouds", DisplayName = "Reference Clouds"))
-	uint32 PathTracingEnableReferenceClouds : 1;
-
 	/** Run the currently loaded denoiser plugin on the last sample to remove noise from the output. Has no effect if a plug-in is not loaded. */
 	UPROPERTY(EditAnywhere, BlueprintReadWrite, Category = "Path Tracing", meta = (editcondition = "bOverride_PathTracingEnableDenoiser", DisplayName = "Denoiser"))
 	uint32 PathTracingEnableDenoiser : 1;
diff --git a/Engine/Source/Runtime/Engine/Private/PostProcessUtils.cpp b/Engine/Source/Runtime/Engine/Private/PostProcessUtils.cpp
--- a/Engine/Source/Runtime/Engine/Private/PostProcessUtils.cpp
+++ b/Engine/Source/Runtime/Engine/Private/PostProcessUtils.cpp
@@ -231,7 +231,6 @@ bool FPostProcessUtils::OverridePostProcessSettings(FPostProcessSettings& ThisFr
 		UE_SET_PP(PathTracingEnableEmissiveMaterials);
 		UE_SET_PP(PathTracingEnableReferenceDOF);
 		UE_SET_PP(PathTracingEnableReferenceAtmosphere);
-		UE_SET_PP(PathTracingEnableReferenceClouds);
 		UE_SET_PP(PathTracingEnableDenoiser);
 		UE_SET_PP(PathTracingIncludeEmissive);
 		UE_SET_PP(PathTracingIncludeDiffuse);
@@ -537,7 +536,6 @@ bool FPostProcessUtils::BlendPostProcessSettings(FPostProcessSettings& ThisFrom,
 		UE_SET_PP(PathTracingEnableEmissiveMaterials);
 		UE_SET_PP(PathTracingEnableReferenceDOF);
 		UE_SET_PP(PathTracingEnableReferenceAtmosphere);
-		UE_SET_PP(PathTracingEnableReferenceClouds);
 		UE_SET_PP(PathTracingEnableDenoiser);
 		UE_SET_PP(PathTracingIncludeEmissive);
 		UE_SET_PP(PathTracingIncludeDiffuse);
diff --git a/Engine/Source/Runtime/Engine/Private/Scene.cpp b/Engine/Source/Runtime/Engine/Private/Scene.cpp
--- a/Engine/Source/Runtime/Engine/Private/Scene.cpp
+++ b/Engine/Source/Runtime/Engine/Private/Scene.cpp
@@ -629,7 +629,6 @@ FPostProcessSettings::FPostProcessSettings()
 	PathTracingEnableEmissiveMaterials = 1;
 	PathTracingEnableReferenceDOF = 0;
 	PathTracingEnableReferenceAtmosphere = 0;
-	PathTracingEnableReferenceClouds = 0;
 	PathTracingEnableDenoiser = 1;
 
 	PathTracingIncludeEmissive = 1;
diff --git a/Engine/Source/Runtime/Engine/Private/SceneView.cpp b/Engine/Source/Runtime/Engine/Private/SceneView.cpp
--- a/Engine/Source/Runtime/Engine/Private/SceneView.cpp
+++ b/Engine/Source/Runtime/Engine/Private/SceneView.cpp
@@ -1794,7 +1794,6 @@ void FSceneView::OverridePostProcessSettings(const FPostProcessSettings& Src, fl
 		SET_PP(PathTracingEnableEmissiveMaterials);
 		SET_PP(PathTracingEnableReferenceDOF);
 		SET_PP(PathTracingEnableReferenceAtmosphere);
-		SET_PP(PathTracingEnableReferenceClouds);
 		SET_PP(PathTracingEnableDenoiser);
 		SET_PP(PathTracingIncludeEmissive);
 		SET_PP(PathTracingIncludeDiffuse);
diff --git a/Engine/Source/Runtime/Renderer/Private/PathTracing.cpp b/Engine/Source/Runtime/Renderer/Private/PathTracing.cpp
--- a/Engine/Source/Runtime/Renderer/Private/PathTracing.cpp
+++ b/Engine/Source/Runtime/Renderer/Private/PathTracing.cpp
@@ -303,11 +303,10 @@ TAutoConsoleVariable<int32> CVarPathTracingEnableReferenceAtmosphere(
 
 TAutoConsoleVariable<int32> CVarPathTracingEnableReferenceClouds(
 	TEXT("r.PathTracing.EnableReferenceClouds"),
-	-1,
+	1,
 	TEXT("Should the path tracer use a volumetric calculation to represent volumetric clouds? (This requires Reference Atmosphere to be enabled)\n")
-	TEXT("-1: Inherit from PostProcess settings (default)\n")
 	TEXT(" 0: Disabled\n")
-	TEXT(" 1: Enabled\n"),
+	TEXT(" 1: Enabled (default)\n"),
 	ECVF_RenderThreadSafe
 );
 
@@ -805,17 +804,6 @@ namespace PathTracing
 		return EnableReferenceAtmosphereCVar < 0 ? View.FinalPostProcessSettings.PathTracingEnableReferenceAtmosphere != 0 : EnableReferenceAtmosphereCVar != 0;
 	}
 
-	bool UsesReferenceClouds(const FViewInfo& View)
-	{
-		// Reference Clouds are only available if reference atmosphere is on
-		if (UsesReferenceAtmosphere(View))
-		{
-			const int32 EnableReferenceCloudsCVar = CVarPathTracingEnableReferenceClouds.GetValueOnRenderThread();
-			return EnableReferenceCloudsCVar < 0 ? View.FinalPostProcessSettings.PathTracingEnableReferenceClouds != 0 : EnableReferenceCloudsCVar != 0;
-		}
-		return false;
-	}
-
 	bool UsesReferenceDOF(const FViewInfo& View)
 	{
 		return GEnableReferenceDOF < 0 ? View.FinalPostProcessSettings.PathTracingEnableReferenceDOF != 0 : GEnableReferenceDOF != 0;
@@ -1718,6 +1706,11 @@ void PreparePathTracingCloudMaterial(FRDGBuilder& GraphBuilder, FScene* Scene, T
 		return;
 	}
 
+	if (CVarPathTracingEnableReferenceClouds.GetValueOnRenderThread() == 0)
+	{
+		return;
+	}
+
 	FVolumetricCloudRenderSceneInfo* CloudRenderSceneInfo = Scene->GetVolumetricCloudSceneInfo();
 	if (CloudRenderSceneInfo == nullptr)
 	{
@@ -1751,10 +1744,9 @@ void PreparePathTracingCloudMaterial(FRDGBuilder& GraphBuilder, FScene* Scene, T
 
 	for (FViewInfo& View : Views)
 	{
-		if (!PathTracing::UsesReferenceClouds(View) ||
-			!ShouldRenderVolumetricCloud(Scene, View.Family->EngineShowFlags))
+		if (!PathTracing::UsesReferenceAtmosphere(View) || !ShouldRenderVolumetricCloud(Scene, View.Family->EngineShowFlags))
 		{
-			// reference cloud mode disabled for this view, or clouds disabled for this view
+			// reference atmosphere mode disabled for this view, or clouds disabled for this view
 			continue;
 		}
 
