Prevent Water Body River component from greatly increasing memory usage in the editor with large landscapes.
Add flag to USceneCaptureComponent to allow it to instruct that it shouldn't be used to influence scene extent decisions for other rendering.

diff --git a/Engine/Plugins/Experimental/Water/Source/Editor/Private/WaterBrushManager.cpp b/Engine/Plugins/Experimental/Water/Source/Editor/Private/WaterBrushManager.cpp
--- a/Engine/Plugins/Experimental/Water/Source/Editor/Private/WaterBrushManager.cpp
+++ b/Engine/Plugins/Experimental/Water/Source/Editor/Private/WaterBrushManager.cpp
@@ -49,6 +49,7 @@ AWaterBrushManager::AWaterBrushManager(const FObjectInitializer& ObjectInitializ
 	SceneCaptureComponent2D->CaptureSource = ESceneCaptureSource::SCS_SceneDepth;
 	SceneCaptureComponent2D->bCaptureEveryFrame = false;
 	SceneCaptureComponent2D->bCaptureOnMovement = false;
+	SceneCaptureComponent2D->bExcludeFromSceneTextureExtents = true;  // Don't let this influence the render target size used for other purposes.
 	SceneCaptureComponent2D->SetRelativeRotation(FRotator(-90.0f, 0.0f, -90.0f));
 	SceneCaptureComponent2D->SetRelativeScale3D(FVector(0.01f, 0.01f, 0.01f));
 
diff --git a/Engine/Source/Runtime/Engine/Classes/Components/SceneCaptureComponent.h b/Engine/Source/Runtime/Engine/Classes/Components/SceneCaptureComponent.h
--- a/Engine/Source/Runtime/Engine/Classes/Components/SceneCaptureComponent.h
+++ b/Engine/Source/Runtime/Engine/Classes/Components/SceneCaptureComponent.h
@@ -105,6 +105,14 @@ class USceneCaptureComponent : public USceneComponent
 	 */
 	uint8 bSuppressGpuCaptureOrDump : 1;
 
+	/**
+	 * Whether this capture should be excluded from tracking scene texture extents.  This should be set when this capture is not expected to be
+	 * frequently used, especially if the capture resolution is very large.  Setting this for a single-use capture will avoid influencing other
+	 * scene texture extent decisions and avoid a possible ongoing increase in memory usage.
+	 */
+	UPROPERTY(EditAnywhere, BlueprintReadWrite, AdvancedDisplay, Category = SceneCapture)
+	uint8 bExcludeFromSceneTextureExtents : 1;
+
 	/** Whether to persist the rendering state even if bCaptureEveryFrame==false.  This allows velocities for Motion Blur and Temporal AA to be computed. */
 	UPROPERTY(EditAnywhere, BlueprintReadWrite, Category = SceneCapture, meta = (editcondition = "!bCaptureEveryFrame"))
 	bool bAlwaysPersistRenderingState;
diff --git a/Engine/Source/Runtime/Engine/Private/SceneView.cpp b/Engine/Source/Runtime/Engine/Private/SceneView.cpp
--- a/Engine/Source/Runtime/Engine/Private/SceneView.cpp
+++ b/Engine/Source/Runtime/Engine/Private/SceneView.cpp
@@ -911,6 +911,7 @@ FSceneView::FSceneView(const FSceneViewInitOptions& InitOptions)
 	, bIsSinglePassStereo(false)
 	, bForceCameraVisibilityReset(false)
 	, bDisableDistanceBasedFadeTransitions(false)
+	, bExcludeFromSceneTextureExtents(InitOptions.bExcludeFromSceneTextureExtents)
 	, bForcePathTracerReset(false)
 	, PathTracerSeedOffset(0)
 	, UnderwaterDepth(-1.0f)
diff --git a/Engine/Source/Runtime/Engine/Public/SceneView.h b/Engine/Source/Runtime/Engine/Public/SceneView.h
--- a/Engine/Source/Runtime/Engine/Public/SceneView.h
+++ b/Engine/Source/Runtime/Engine/Public/SceneView.h
@@ -241,6 +241,9 @@ struct FSceneViewInitOptions : public FSceneViewProjectionData
 	/** Whether this view is being used to render a planar reflection. */
 	bool bIsPlanarReflection;
 
+	/** Whether this view should be excluded from extent history tracking. */
+	bool bExcludeFromSceneTextureExtents = false;
+
 	/** If > 0, overrides the view's resolution fraction. */
 	float OverridePrimaryResolutionFraction;
 
@@ -1632,6 +1635,9 @@ public:
 	/** Whether we should disable distance-based fade transitions for this frame (usually after a large camera movement.) */
 	bool bDisableDistanceBasedFadeTransitions;
 
+	/** Whether this view should be excluded from extent history tracking. */
+	bool bExcludeFromSceneTextureExtents = false;
+
 	/** True if we should force the path tracer to reset its internal accumulation state */
 	bool bForcePathTracerReset;
 
diff --git a/Engine/Source/Runtime/Renderer/Private/SceneCaptureRendering.cpp b/Engine/Source/Runtime/Renderer/Private/SceneCaptureRendering.cpp
--- a/Engine/Source/Runtime/Renderer/Private/SceneCaptureRendering.cpp
+++ b/Engine/Source/Runtime/Renderer/Private/SceneCaptureRendering.cpp
@@ -729,6 +729,7 @@ TArray<FSceneView*> SetupViewFamilyForSceneCapture(
 			ViewInitOptions.LODDistanceFactor = FMath::Clamp(SceneCaptureComponent->LODDistanceFactor, .01f, 100.0f);
 			ViewInitOptions.bIsSceneCaptureCube = SceneCaptureComponent->IsCube();
 			ViewInitOptions.bSceneCaptureUsesRayTracing = GRayTracingSceneCaptures == -1 ? SceneCaptureComponent->bUseRayTracingIfEnabled : GRayTracingSceneCaptures > 0;
+			ViewInitOptions.bExcludeFromSceneTextureExtents = SceneCaptureComponent->bExcludeFromSceneTextureExtents;
 
 			USceneCaptureComponent2D* SceneCaptureComponent2D = Cast<USceneCaptureComponent2D>(SceneCaptureComponent);
 			if (IsValid(SceneCaptureComponent2D))
diff --git a/Engine/Source/Runtime/Renderer/Private/SceneTextures.cpp b/Engine/Source/Runtime/Renderer/Private/SceneTextures.cpp
--- a/Engine/Source/Runtime/Renderer/Private/SceneTextures.cpp
+++ b/Engine/Source/Runtime/Renderer/Private/SceneTextures.cpp
@@ -141,12 +141,15 @@ public:
 		bool bIsSceneCapture = false;
 		bool bIsReflectionCapture = false;
 		bool bIsVRScene = false;
+		// If explicitly requested, don't include this in history tracking for future requests.  This extent request can still be adjusted based on existing history.
+		bool bExcludeFromHistoryUpdate = false;
 
 		for (const FSceneView* View : ViewFamily.AllViews)
 		{
 			bIsSceneCapture |= View->bIsSceneCapture;
 			bIsReflectionCapture |= View->bIsReflectionCapture;
 			bIsVRScene |= (IStereoRendering::IsStereoEyeView(*View) && GEngine->XRSystem.IsValid());
+			bExcludeFromHistoryUpdate |= View->bExcludeFromSceneTextureExtents;
 		}
 
 		FIntPoint DesiredExtent = FIntPoint::ZeroValue;
@@ -254,12 +257,15 @@ public:
 
 		// this allows The extent to not grow below the SceneCapture requests (happen before scene rendering, in the same frame with a Grow request)
 		FIntPoint& LargestDesiredExtentThisFrame = LargestDesiredExtents[DesiredExtentIndex];
-		LargestDesiredExtentThisFrame = LargestDesiredExtentThisFrame.ComponentMax(DesiredExtent);
 		bool bIsHighResScreenshot = GIsHighResScreenshot;
-		UpdateHistoryFlags(HistoryFlags[DesiredExtentIndex], bIsSceneCapture, bIsReflectionCapture, bIsHighResScreenshot);
+
+		if (!bExcludeFromHistoryUpdate)
+		{
+			LargestDesiredExtentThisFrame = LargestDesiredExtentThisFrame.ComponentMax(DesiredExtent);
+			UpdateHistoryFlags(HistoryFlags[DesiredExtentIndex], bIsSceneCapture, bIsReflectionCapture, bIsHighResScreenshot);
+		}
 
 		// We want to shrink the buffer but as we can have multiple scene captures per frame we have to delay that a frame to get all size requests
-		// Don't save buffer size in history while making high-res screenshot.
 		// We have to use the requested size when allocating an hmd depth target to ensure it matches the hmd allocated render target size.
 		bool bAllowDelayResize = !GIsHighResScreenshot && !bIsVRScene;
 
@@ -267,7 +273,7 @@ public:
 		// This prevents problems when orientation changes on mobile in particular.
 		// bIsReflectionCapture is explicitly checked on all platforms to prevent aspect ratio change detection from forcing the immediate buffer resize.
 		// This ensures that 1) buffers are not resized spuriously during reflection rendering 2) all cubemap faces use the same render target size.
-		if (bAllowDelayResize && !bIsReflectionCapture && !AnyCaptureRenderedRecently<ExtentHistoryCount>(ERenderTargetHistory::MaskAll))
+		if (bAllowDelayResize && !bIsReflectionCapture && !bExcludeFromHistoryUpdate && !AnyCaptureRenderedRecently<ExtentHistoryCount>(ERenderTargetHistory::MaskAll))
 		{
 			const bool bAspectRatioChanged =
 				!LastExtent.Y ||
@@ -306,7 +312,10 @@ public:
 
 		check(DesiredExtent.X > 0 && DesiredExtent.Y > 0);
 		QuantizeSceneBufferSize(DesiredExtent, DesiredExtent);
-		LastExtent = DesiredExtent;
+		if (!bExcludeFromHistoryUpdate)
+		{
+			LastExtent = DesiredExtent;
+		}
 		return DesiredExtent;
 	}
 
