Resolve thread timing issue that was causinng invalid Scene texture config. Move TaskDatas.VisibilityTaskData->Finish(); to precede IsMobileSeparateTranslucencyActive calls as it depends on results  from TaskDatas.VisibilityTaskData, namely SetupMeshPasses and the collection of DynamicMeshRelevance data.

diff --git a/Engine/Source/Runtime/Renderer/Private/MobileShadingRenderer.cpp b/Engine/Source/Runtime/Renderer/Private/MobileShadingRenderer.cpp
--- a/Engine/Source/Runtime/Renderer/Private/MobileShadingRenderer.cpp
+++ b/Engine/Source/Runtime/Renderer/Private/MobileShadingRenderer.cpp
@@ -584,6 +584,9 @@ void FMobileSceneRenderer::InitViews(
 
 	bShouldRenderHZB = ShouldRenderHZB(Views) && bRendererOutputFinalSceneColor;
 
+	// Wait for visibilityTaskData to finish as IsMobileSeparateTranslucencyActive depends on results from SetupMeshPasses.
+	TaskDatas.VisibilityTaskData->Finish();
+	
 	bool bUsingOcclusionFeedback = Views[0].ViewState && Views[0].ViewState->OcclusionFeedback.IsInitialized();
 	bAdrenoOcclusionMode = (DoOcclusionQueries() && !bUsingOcclusionFeedback && !Views[0].bDisableQuerySubmissions && CVarMobileAdrenoOcclusionMode.GetValueOnAnyThread() != 0);
 
@@ -714,9 +717,6 @@ void FMobileSceneRenderer::InitViews(
 		}
 	}
 
-	// Trigger shadow tasks after the main visibility tasks are fully executed. Projection stencil shadows reference the main view dynamic elements result.
-	TaskDatas.VisibilityTaskData->Finish();
-
 	if (bRendererOutputFinalSceneColor)
 	{
 		const bool bDynamicShadows = ViewFamily.EngineShowFlags.DynamicShadows;
