Only use async deletion task when cleanup mode is async. This avoids an edge case where the GPU scene deferred context is not cleared on the render thread, causing a race condition.

diff --git a/Engine/Source/Runtime/Renderer/Private/SceneRenderBuilder.cpp b/Engine/Source/Runtime/Renderer/Private/SceneRenderBuilder.cpp
--- a/Engine/Source/Runtime/Renderer/Private/SceneRenderBuilder.cpp
+++ b/Engine/Source/Runtime/Renderer/Private/SceneRenderBuilder.cpp
@@ -951,11 +951,7 @@ void FSceneRenderProcessor::Execute()
 
 	ENQUEUE_RENDER_COMMAND(SceneRenderBuilder_End)([this](FRHICommandListImmediate& RHICmdList) mutable
 	{
-		FGraphEventArray Prereqs;
-		Prereqs.Add(AsyncTasks.Cleanup);
-		Prereqs.Add(AsyncTasks.Delete);
-
-		AsyncTasks.Delete = FFunctionGraphTask::CreateAndDispatchWhenReady([this]
+		const auto DeleteLambda = [this]
 		{
 			TRACE_CPUPROFILER_EVENT_SCOPE(FSceneRenderProcessor::DeleteSceneRenderers);
 			for (FSceneRenderer* Renderer : Renderers)
@@ -963,8 +959,20 @@ void FSceneRenderProcessor::Execute()
 				delete Renderer;
 			}
 			delete this;
+		};
+
+		if (GetSceneRenderCleanUpMode() == ESceneRenderCleanUpMode::Async)
+		{
+			FGraphEventArray Prereqs;
+			Prereqs.Add(AsyncTasks.Cleanup);
+			Prereqs.Add(AsyncTasks.Delete);
 
-		}, TStatId(), &Prereqs);
+			AsyncTasks.Delete = FFunctionGraphTask::CreateAndDispatchWhenReady(DeleteLambda, TStatId(), &Prereqs);
+		}
+		else
+		{
+			DeleteLambda();
+		}
 	});
 
 	// NOTE: 'this' is queued for deletion and is no longer valid!
