[TSR] Improve temporal stability of thin geometries against sky.

* Sky instability pixels are detected based on intensity (The instability is caused by volumetric render target mode). Once detected, we relax the history clamping around the patch to not aggressively reject the history.
* Add a CVar (r.TSR.ThinGeometryDetection.WeightRelaxation.Sky) to change between a general detector (0) and this specialized detector (1).

#[review] Brian.Karis, aleksander.netzel
#tests Needles on consoles and internal project

diff --git a/Engine/Shaders/Private/TemporalSuperResolution/TSRDetectThinGeometry.usf b/Engine/Shaders/Private/TemporalSuperResolution/TSRDetectThinGeometry.usf
--- a/Engine/Shaders/Private/TemporalSuperResolution/TSRDetectThinGeometry.usf
+++ b/Engine/Shaders/Private/TemporalSuperResolution/TSRDetectThinGeometry.usf
@@ -6,7 +6,7 @@
 #include "TSRDepthVelocityAnalysis.ush"
 #include "TSRConvolutionNetworkPass.ush"
 #include "TSRReprojectionField.ush"
-
+#include "TSRThinGeometryCommon.ush"
 //------------------------------------------------------- CONFIG
 
 #define DEBUG_ARRAY_SIZE 8
@@ -21,7 +21,7 @@
 #define EDGE_TYPE EDGE_HORIZONTAL_VERTICAL_DIAGNAL
 
 #define CONFIG_THIN_GEOMETRY_EDGE_REPROJECTION (DIM_THIN_GEOMETRY_EDGE_REPROJECTION)
-
+#define CONFIG_SKY_RELAXATION (DIM_SKY_RELAXATION)
 //------------------------------------------------------- CONSTANTS
 
 #if CONFIG_LDS_DESPILL
@@ -132,6 +132,8 @@ void FetchInput(
 #else
 	CurrentCoverage = select(bHasReprojectionOffset, CurrentCoverage, tsr_tensor_half::Const(0));
 #endif
+	
+	CurrentCoverage = ToBinaryCoverage(CurrentCoverage);
 
 	UNROLL_N(SIMD_SIZE)
 	for (uint ElementIndex = 0; ElementIndex < SIMD_SIZE; ElementIndex++)
@@ -147,6 +149,7 @@ void FetchInput(
 
 //------------------------------------------------------- FUNCTIONS
 
+#if !CONFIG_SKY_RELAXATION
 // Horizontal thin edge detection.
 // 0 0 0
 // x x x
@@ -236,6 +239,8 @@ TLaneVector2D<ScalarType, VectorSize, LaneStrideX, LaneStrideY> VerticalThinEdge
 	return ThinGeometryEdge;
 }
 
+#endif // !CONFIG_SKY_RELAXATION
+
 tsr_tensor_uint EncodeThinGeometryWeight(tsr_tensor_half ThinGeometryEdge
 	, tsr_tensor_bool IsClusterHoleRegion,
 	tsr_tensor_half Weight
@@ -293,11 +298,13 @@ void MainCS(
 	
 	tsr_tensor_half ThinGeometryEdge = tsr_tensor_half::Const(0.0);
 
+#if !CONFIG_SKY_RELAXATION
 	{
 		ThinGeometryEdge = HorizontalThinEdge3x3(DeviceZ, DeviceZThinGeometryDetectionError);
 		ThinGeometryEdge = ThinGeometryEdge + VerticalThinEdge3x3(DeviceZ, DeviceZThinGeometryDetectionError);
 		//Debug[0].SetComponent(2, tsr_tensor_half(ThinGeometryEdge));
 	}
+#endif
 
 	tsr_tensor_half HistoryRelaxationWeight = tsr_tensor_half::Const(0);
 	tsr_tensor_bool IsClusterHoleRegion = tsr_tensor_bool::Const(false);
diff --git a/Engine/Shaders/Private/TemporalSuperResolution/TSRMeasureCoverage.usf b/Engine/Shaders/Private/TemporalSuperResolution/TSRMeasureCoverage.usf
--- a/Engine/Shaders/Private/TemporalSuperResolution/TSRMeasureCoverage.usf
+++ b/Engine/Shaders/Private/TemporalSuperResolution/TSRMeasureCoverage.usf
@@ -3,7 +3,7 @@
 
 #include "../DeferredShadingCommon.ush"  
 #include "TSRColorSpace.ush"
-
+#include "TSRThinGeometryCommon.ush"
 //------------------------------------------------------- CONFIG
 
 #define TILE_SIZE 8
@@ -19,11 +19,12 @@ RWTexture2DArray<float> ThinGeometryCoverageOutput;
 
 uint GetThinGeometry(const FGBufferData GBufferData)
 {
-	return ((GBufferData.ShadingModelID == SHADINGMODELID_TWOSIDED_FOLIAGE) 
+	uint ThinGeometryFlag = ((GBufferData.ShadingModelID == SHADINGMODELID_TWOSIDED_FOLIAGE) 
 #if CONFIG_HAIR_AS_THIN_GEOMETRY
 		    ||(GBufferData.ShadingModelID == SHADINGMODELID_HAIR)
 #endif
-		    ) ? 1u : 0u;
+		    ) ? 2u : 1u;
+	return select(GBufferData.ShadingModelID == SHADINGMODELID_UNLIT, 0u, ThinGeometryFlag);
 }
 
 //------------------------------------------------------- ENTRY POINT
@@ -64,7 +65,7 @@ void MainCS(
 			{
 				const uint2 PixelOffset = uint2(x, y) * TILE_SIZE;
 				uint2 InputPixelPos = OutputPixelOrigin + PixelOffset;
-				float ThinGeometry = (float)ThinGeometryArray[x + y * TILE_ROWS];
+				float ThinGeometry = EncodeThinGeometryMaterial(ThinGeometryArray[x + y * TILE_ROWS]);
 
 				InputPixelPos.x = select(all(InputPixelPos < InputInfo_ViewportMax), InputPixelPos.x, uint(~0));
 				ThinGeometryCoverageOutput[int3(InputPixelPos, OutputArrayIndex)] = ThinGeometry;
diff --git a/Engine/Shaders/Private/TemporalSuperResolution/TSRThinGeometryCommon.ush b/Engine/Shaders/Private/TemporalSuperResolution/TSRThinGeometryCommon.ush
--- a/Engine/Shaders/Private/TemporalSuperResolution/TSRThinGeometryCommon.ush
+++ b/Engine/Shaders/Private/TemporalSuperResolution/TSRThinGeometryCommon.ush
@@ -8,6 +8,24 @@
 
 
 //-------------------------------------------------------UTIL FUNCTION
+
+float EncodeThinGeometryMaterial(uint ThinGeometryFlag)
+{
+	return ((float)ThinGeometryFlag) * rcp(2.0f);
+}
+
+CALL_SITE_DEBUGLOC
+tsr_tensor_bool IsThinGeometry(tsr_tensor_half CurrentCoverage)
+{
+	return CurrentCoverage == tsr_tensor_half::Const(1.0);
+}
+
+CALL_SITE_DEBUGLOC
+tsr_tensor_half ToBinaryCoverage(tsr_tensor_half EncodedThinGeometryMaterial)
+{
+	return floor(EncodedThinGeometryMaterial);
+}
+
 CALL_SITE_DEBUGLOC
 void DecodeThinGeometry(
 	tsr_tensor_ushort EncodedThinGeometryWeight,
diff --git a/Engine/Shaders/Private/TemporalSuperResolution/TSRWeightRelaxation.usf b/Engine/Shaders/Private/TemporalSuperResolution/TSRWeightRelaxation.usf
--- a/Engine/Shaders/Private/TemporalSuperResolution/TSRWeightRelaxation.usf
+++ b/Engine/Shaders/Private/TemporalSuperResolution/TSRWeightRelaxation.usf
@@ -14,6 +14,8 @@
 
 #define DEBUG_DIGESTABLE 0
 
+#define CONFIG_SKY_RELAXATION (DIM_SKY_RELAXATION)
+
 //------------------------------------------------------- CONSTANTS
 
 #if CONFIG_LDS_DESPILL
@@ -28,6 +30,8 @@ int ThinGeometryTextureIndex;
 float ErrorMultiplier;
 float MaxRelaxationWeight;
 
+Texture2D<tsr_half>  CurrentCoverageTexture;
+Texture2D<tsr_half>  InputMoireLumaTexture;
 Texture2D<tsr_halfC> InputTexture;
 Texture2D<tsr_half4> InputSceneTranslucencyTexture;
 
@@ -115,7 +119,12 @@ void FetchInput(
 	out tsr_tensor_half InputLuma,
 	out tsr_tensor_half2 Translucency,
 	out tsr_tensor_half HistoryRelaxationWeight,
-	out tsr_tensor_bool IsClusterHoleRegion)
+	out tsr_tensor_bool IsClusterHoleRegion
+#if CONFIG_SKY_RELAXATION
+	,out tsr_tensor_half InputLumaBeforeComposition
+	,out tsr_tensor_half CurrentCoverage
+#endif
+)
 {
 	tsr_tensor_ushort ThinGeometryCache;
 	
@@ -127,6 +136,11 @@ void FetchInput(
 		tsr_half4 FullTranslucency = InputSceneTranslucencyTexture[ClampedFetchPixelPos];
 		Translucency.SetElement(ElementIndex, tsr_half2(ComputeRGBLuma(FullTranslucency.rgb),FullTranslucency.a));
 		ThinGeometryCache.SetElement(ElementIndex, tsr_ushort(R8Output[int3(ClampedFetchPixelPos, ThinGeometryTextureIndex)]));
+
+#if CONFIG_SKY_RELAXATION
+		CurrentCoverage.SetElement(ElementIndex,CurrentCoverageTexture[ClampedFetchPixelPos]);
+		InputLumaBeforeComposition.SetElement(ElementIndex,InputMoireLumaTexture[ClampedFetchPixelPos]);
+#endif
 	}
 
 	DecodeThinGeometry(ThinGeometryCache, HistoryRelaxationWeight, IsClusterHoleRegion);	
@@ -159,30 +173,56 @@ void MainCS(
 	tsr_tensor_half2 Translucency;
 	tsr_tensor_half HistoryRelaxationWeight;
 	tsr_tensor_bool IsPartialCoverageRegion;
-
+#if CONFIG_SKY_RELAXATION
+	tsr_tensor_half InputLumaBeforeComposition;
+	tsr_tensor_half CurrentCoverage;
+#endif
 	FetchInput(
 		PixelPos,
-		/*out */ InputLuma,
+ 		/*out */ InputLuma,
 		/*out */ Translucency,
 		/*out */ HistoryRelaxationWeight,
-		/*out */ IsPartialCoverageRegion);
+		/*out */ IsPartialCoverageRegion
+#if CONFIG_SKY_RELAXATION
+		,/*out */ InputLumaBeforeComposition
+		,/*out */ CurrentCoverage
+#endif
+	);
 
-	// If translucency has data turn down the weights aggresively. if alpha < 0.95
-	tsr_tensor_half TranslucencyRatio = Translucency[0] * rcp(Translucency[0] + InputLuma * Translucency[1] + tsr_tensor_half::Const(0.0005));
+	// Edge case: thin geometry against volumetric clouds.
+#if CONFIG_SKY_RELAXATION
+	tsr_tensor_bool IsEarlyTranslucencyEffect;
+	{
+		tsr_tensor_half EarlyTranslucencyAE = abs(InputLuma - InputLumaBeforeComposition);
 #if DEBUG_OUTPUT
-	Debug[0].SetComponent(0, TranslucencyRatio);
-	Debug[0].SetComponent(1, HistoryRelaxationWeight);
+		Debug[0].SetComponent(0, EarlyTranslucencyAE);
+#endif
+
+		tsr_tensor_half MinCurrentCoverageNeighbor;
+		MinCurrentCoverageNeighbor = Min3x3(CurrentCoverage);
+
+		tsr_tensor_bool bIsLeafAgainstSky = and(and(EarlyTranslucencyAE > tsr_tensor_half::Const(0.3), InputLumaBeforeComposition > tsr_tensor_half::Const(0.01)), 
+			and(MinCurrentCoverageNeighbor == tsr_tensor_half::Const(0), IsThinGeometry(CurrentCoverage)));
+#if DEBUG_OUTPUT		
+		Debug[0].SetComponent(1, select(bIsLeafAgainstSky,abs(InputLuma - InputLumaBeforeComposition), tsr_tensor_half::Const(0)));
+#endif
+		IsEarlyTranslucencyEffect = bIsLeafAgainstSky;
+		tsr_tensor_half EarlyTranslucencyAddWeight = select(IsEarlyTranslucencyEffect, tsr_tensor_half::Const(1.0), tsr_tensor_half::Const(0));
+		EarlyTranslucencyAddWeight = WeightedAvg3x3(EarlyTranslucencyAddWeight, float3(0.8, 1.0, 0.8));
+		IsPartialCoverageRegion = select(EarlyTranslucencyAddWeight > tsr_tensor_half::Const(0),tsr_tensor_bool::Const(false),IsPartialCoverageRegion);
+		HistoryRelaxationWeight = saturate(HistoryRelaxationWeight + EarlyTranslucencyAddWeight);
+	}
 #endif
+
+	// If translucency has data turn down the weights aggresively. if alpha < 0.95
+	tsr_tensor_half TranslucencyRatio = Translucency[0] * rcp(Translucency[0] + InputLuma * Translucency[1] + tsr_tensor_half::Const(0.0005));
+
 	// Weights decline starts around TranslucencyRatio = 0.05, and declines fast to 0 when Translucency ratio goes to 0.6. 
 	tsr_tensor_half RatioToLerpWeight = max(tsr_tensor_half::Const(0.6)-TranslucencyRatio,tsr_tensor_half::Const(0));
 	RatioToLerpWeight = min(RatioToLerpWeight * RatioToLerpWeight * tsr_tensor_half::Const(3.33), tsr_tensor_half::Const(1));
 	RatioToLerpWeight = select(or(IsPartialCoverageRegion,TranslucencyRatio > tsr_tensor_half::Const(0.6)), RatioToLerpWeight, tsr_tensor_half::Const(1));
 	HistoryRelaxationWeight = lerp(tsr_tensor_half::Const(0), HistoryRelaxationWeight, RatioToLerpWeight); 
 	
-#if DEBUG_OUTPUT
-	Debug[0].SetComponent(2, HistoryRelaxationWeight);
-#endif
-
 	tsr_tensor_ushort EncodedThinGeometry = EncodeThinGeometry(HistoryRelaxationWeight,IsPartialCoverageRegion);
 	// Output
 	UNROLL_N(SIMD_SIZE)
diff --git a/Engine/Source/Runtime/Renderer/Private/DeferredShadingRenderer.cpp b/Engine/Source/Runtime/Renderer/Private/DeferredShadingRenderer.cpp
--- a/Engine/Source/Runtime/Renderer/Private/DeferredShadingRenderer.cpp
+++ b/Engine/Source/Runtime/Renderer/Private/DeferredShadingRenderer.cpp
@@ -3048,7 +3048,7 @@ void FDeferredShadingSceneRenderer::Render(FRDGBuilder& GraphBuilder, const FSce
 		TArray<FScreenPassTexture, TInlineAllocator<4>> TSRFlickeringInputTextures;
 		if (!bHasRayTracedOverlay)
 		{
-			// Extract TSR's moire heuristic luminance, thin geometry coverage before rendering translucency into the scene color.
+			// Extract TSR's moire heuristic luminance before rendering translucency into the scene color.
 			for (int32 ViewIndex = 0; ViewIndex < Views.Num(); ++ViewIndex)
 			{
 				FViewInfo& View = Views[ViewIndex];
@@ -3059,7 +3059,7 @@ void FDeferredShadingSceneRenderer::Render(FRDGBuilder& GraphBuilder, const FSce
 						TSRFlickeringInputTextures.SetNum(Views.Num());
 					}
 
-					TSRFlickeringInputTextures[ViewIndex] = AddTSRMainAntiFlickeringPass(GraphBuilder, View.ShaderMap, FScreenPassTexture(SceneTextures.Color.Target, View.ViewRect), SceneTextures);
+					TSRFlickeringInputTextures[ViewIndex] = AddTSRMeasureFlickeringLuma(GraphBuilder, View.ShaderMap, FScreenPassTexture(SceneTextures.Color.Target, View.ViewRect));
 				}
 			}
 		}
@@ -3163,6 +3163,24 @@ void FDeferredShadingSceneRenderer::Render(FRDGBuilder& GraphBuilder, const FSce
 		// Rebuild scene textures to include scene color.
 		SceneTextures.UniformBuffer = CreateSceneTextureUniformBuffer(GraphBuilder, &SceneTextures, FeatureLevel, SceneTextures.SetupMode);
 
+		if (!bHasRayTracedOverlay)
+		{
+			// Extract TSR's thin geometry coverage after SLW but before rendering translucency into the scene color.
+			for (int32 ViewIndex = 0; ViewIndex < Views.Num(); ++ViewIndex)
+			{
+				FViewInfo& View = Views[ViewIndex];
+				if (NeedTSRAntiFlickeringPass(View))
+				{
+					if (TSRFlickeringInputTextures.Num() == 0)
+					{
+						TSRFlickeringInputTextures.SetNum(Views.Num());
+					}
+
+					AddTSRMeasureThinGeometryCoverage(GraphBuilder, View.ShaderMap, SceneTextures, TSRFlickeringInputTextures[ViewIndex]);
+				}
+			}
+		}
+
 		if (!bIsCameraUnderWater)
 		{
 			RenderLightShaftSkyFogAndCloud();
diff --git a/Engine/Source/Runtime/Renderer/Private/PostProcess/TemporalAA.h b/Engine/Source/Runtime/Renderer/Private/PostProcess/TemporalAA.h
--- a/Engine/Source/Runtime/Renderer/Private/PostProcess/TemporalAA.h
+++ b/Engine/Source/Runtime/Renderer/Private/PostProcess/TemporalAA.h
@@ -185,6 +185,10 @@ bool IsVisualizeTSREnabled(const FViewInfo& View);
 /** Measure different metrics of the scene for anti-flickering */
 FScreenPassTexture AddTSRMainAntiFlickeringPass(FRDGBuilder& GraphBuilder, FGlobalShaderMap* ShaderMap, FScreenPassTexture SceneColor, const FSceneTextures& SceneTextures);
 
+FScreenPassTexture AddTSRMeasureFlickeringLuma(FRDGBuilder& GraphBuilder, FGlobalShaderMap* ShaderMap, FScreenPassTexture SceneColor);
+
+void AddTSRMeasureThinGeometryCoverage(FRDGBuilder& GraphBuilder, FGlobalShaderMap* ShaderMap, const FSceneTextures& SceneTextures, const FScreenPassTexture& ThinGeometryCoverage); 
+
 EMainTAAPassConfig GetMainTAAPassConfig(const FViewInfo& View);
 
 /** Interface for the default temporal upscaling algorithm. */
diff --git a/Engine/Source/Runtime/Renderer/Private/PostProcess/TemporalSuperResolution.cpp b/Engine/Source/Runtime/Renderer/Private/PostProcess/TemporalSuperResolution.cpp
--- a/Engine/Source/Runtime/Renderer/Private/PostProcess/TemporalSuperResolution.cpp
+++ b/Engine/Source/Runtime/Renderer/Private/PostProcess/TemporalSuperResolution.cpp
@@ -333,6 +333,11 @@ TAutoConsoleVariable<int32> CVarTSRThinGeometryDetectionWeightRelaxation(
 	TEXT("Adaptively trim the history relaxation to avoid ghosting."),
 	ECVF_RenderThreadSafe);
 
+TAutoConsoleVariable<int32> CVarTSRThinGeometryDetectionWeightRelaxationSky(
+	TEXT("r.TSR.ThinGeometryDetection.WeightRelaxation.Sky"), 1,
+	TEXT("Apply weight adjust to deal with thin geometry shimmering against sky."),
+	ECVF_RenderThreadSafe);
+
 TAutoConsoleVariable<float> CVarTSRThinGeometryDetectionErrorMultiplier(
 	TEXT("r.TSR.ThinGeometryDetection.ErrorMultiplier"), 200.0,
 	TEXT("Define the depth difference multiplier between neighbors will be used to detect thin geometry. The larger the upscaler, the smaller the multiplier should be for temporal stability. ")
@@ -428,6 +433,12 @@ FTSRHistoryArrayIndices TranslateHistoryFormatBitsToArrayIndices(ETSRHistoryForm
 	return ArrayIndices;
 }
 
+bool ShouldApplySkyRelaxation()
+{
+	return (CVarTSRThinGeometryDetection.GetValueOnRenderThread()!=0 && CVarTSRThinGeometryDetectionWeightRelaxation.GetValueOnRenderThread()!=0
+		&& CVarTSRThinGeometryDetectionWeightRelaxationSky.GetValueOnRenderThread()!=0);
+}
+
 class FTSRShader : public FGlobalShader
 {
 public:
@@ -969,8 +980,9 @@ class FTSRDetectThinGeometryCS : public FTSRConvolutionNetworkShader
 	DECLARE_GLOBAL_SHADER(FTSRDetectThinGeometryCS);
 	SHADER_USE_PARAMETER_STRUCT(FTSRDetectThinGeometryCS, FTSRConvolutionNetworkShader);
 
+	class FSkyRelaxationDim : SHADER_PERMUTATION_BOOL("DIM_SKY_RELAXATION");
 	class FThinGeometryEdgeReprojectionDim : SHADER_PERMUTATION_BOOL("DIM_THIN_GEOMETRY_EDGE_REPROJECTION");
-	using FPermutationDomain = TShaderPermutationDomain<FTSRConvolutionNetworkShader::FPermutationDomain,FThinGeometryEdgeReprojectionDim>;
+	using FPermutationDomain = TShaderPermutationDomain<FTSRConvolutionNetworkShader::FPermutationDomain,FThinGeometryEdgeReprojectionDim, FSkyRelaxationDim>;
 
 	BEGIN_SHADER_PARAMETER_STRUCT(FParameters, )
 		SHADER_PARAMETER_STRUCT_INCLUDE(FTSRCommonParameters, CommonParameters)
@@ -1041,7 +1053,8 @@ class FTSRTSRWeightRelaxationCS :  public FTSRConvolutionNetworkShader
 	DECLARE_GLOBAL_SHADER(FTSRTSRWeightRelaxationCS);
 	SHADER_USE_PARAMETER_STRUCT(FTSRTSRWeightRelaxationCS, FTSRConvolutionNetworkShader);
 
-	using FPermutationDomain = TShaderPermutationDomain<FTSRConvolutionNetworkShader::FPermutationDomain>;
+	class FSkyRelaxationDim : SHADER_PERMUTATION_BOOL("DIM_SKY_RELAXATION");
+	using FPermutationDomain = TShaderPermutationDomain<FTSRConvolutionNetworkShader::FPermutationDomain,FSkyRelaxationDim>;
 
 	BEGIN_SHADER_PARAMETER_STRUCT(FParameters,)
 		SHADER_PARAMETER_STRUCT_INCLUDE(FTSRCommonParameters, CommonParameters)
@@ -1049,6 +1062,8 @@ class FTSRTSRWeightRelaxationCS :  public FTSRConvolutionNetworkShader
 		SHADER_PARAMETER(int32, ThinGeometryTextureIndex)
 		SHADER_PARAMETER(float, MaxRelaxationWeight)
 
+		SHADER_PARAMETER_RDG_TEXTURE_SRV(Texture2D, CurrentCoverageTexture)
+		SHADER_PARAMETER_RDG_TEXTURE_SRV(Texture2D, InputMoireLumaTexture)
 		SHADER_PARAMETER_RDG_TEXTURE(Texture2D, InputTexture)                  // Scene with
 		SHADER_PARAMETER_RDG_TEXTURE(Texture2D, InputSceneTranslucencyTexture) // Translucency texture
 		SHADER_PARAMETER_RDG_TEXTURE_UAV(RWTexture2DArray, R8Output)
@@ -1566,7 +1581,7 @@ bool NeedTSRAntiFlickeringPass(const FViewInfo& View)
 
 bool ShouldAddTSRMainFlickeringLumaPass()
 {
-	return  CVarTSRFlickeringEnable.GetValueOnRenderThread()!=0 || CVarTSRFlickeringPeriod.GetValueOnRenderThread() == 0.0f;
+	return  CVarTSRFlickeringEnable.GetValueOnRenderThread()!=0 || CVarTSRFlickeringPeriod.GetValueOnRenderThread() == 0.0f || ShouldApplySkyRelaxation();
 }
 
 bool ShouldAddTSRMainThinGeometryCoveragePass()
@@ -1656,6 +1671,9 @@ void AddTSRMeasureThinGeometryCoverage(FRDGBuilder& GraphBuilder, FGlobalShaderM
 	FTSRMeasureThinGeometryCoverageCS::FPermutationDomain PermutationVector;
 	PermutationVector.Set<FTSRMeasureThinGeometryCoverageCS::FHairAsThinGeometryDim>(HairAsThinGeometry);
 
+	// whether TSR passes can run on async compute.
+	int32 AsyncComputePasses = GSupportsEfficientAsyncCompute ? CVarTSRAsyncCompute.GetValueOnRenderThread() : 0;
+
 	TShaderMapRef<FTSRMeasureThinGeometryCoverageCS> ComputeShader(ShaderMap,PermutationVector);
 	FComputeShaderUtils::AddPass(
 		GraphBuilder,
@@ -1664,6 +1682,7 @@ void AddTSRMeasureThinGeometryCoverage(FRDGBuilder& GraphBuilder, FGlobalShaderM
 			HairAsThinGeometry ? TEXT(" Hair"): TEXT(""),
 			ViewRect.Width(), 
 			ViewRect.Height()),
+		AsyncComputePasses >= 2 ? ERDGPassFlags::AsyncCompute : ERDGPassFlags::Compute,
 		ComputeShader,
 		PassParameters,
 		FComputeShaderUtils::GetGroupCount(ThinGeometryCoverage.ViewRect.Size(), 8 * 2));
@@ -2576,6 +2595,7 @@ FDefaultTemporalUpscaler::FOutputs AddTemporalSuperResolutionPasses(
 			ConvolutionNetworkPermutationVector.Set<FTSRShader::FAlphaChannelDim>(bSupportsAlpha);
 
 			FTSRDetectThinGeometryCS::FPermutationDomain PermutationVector;
+			PermutationVector.Set<FTSRDetectThinGeometryCS::FSkyRelaxationDim>(ShouldApplySkyRelaxation());
             PermutationVector.Set<FTSRDetectThinGeometryCS::FThinGeometryEdgeReprojectionDim>(CVarTSRThinGeometryCoverageEdgeReprojection.GetValueOnRenderThread());
 			PermutationVector.Set<FTSRConvolutionNetworkShader::FPermutationDomain>(ConvolutionNetworkPermutationVector);
 			PermutationVector = FTSRDetectThinGeometryCS::RemapPermutation(PermutationVector);
@@ -2660,11 +2680,12 @@ FDefaultTemporalUpscaler::FOutputs AddTemporalSuperResolutionPasses(
 			ConvolutionNetworkPermutationVector.Set<FTSRShader::FAlphaChannelDim>(bSupportsAlpha);
 
 			FTSRTSRWeightRelaxationCS::FPermutationDomain PermutationVector;
+			PermutationVector.Set<FTSRTSRWeightRelaxationCS::FSkyRelaxationDim>(ShouldApplySkyRelaxation());
 			PermutationVector.Set<FTSRConvolutionNetworkShader::FPermutationDomain>(ConvolutionNetworkPermutationVector);
 			PermutationVector = FTSRTSRWeightRelaxationCS::RemapPermutation(PermutationVector);
 
 			const int32 GroupTileSize = 32;
-			const int32 TileOverscan = 0;//FMath::Clamp(CVarTSRShadingTileOverscan.GetValueOnRenderThread(), 3, GroupTileSize / 2 - 1);
+			const int32 TileOverscan = ShouldApplySkyRelaxation() ? 2 : 0;//FMath::Clamp(CVarTSRShadingTileOverscan.GetValueOnRenderThread(), 3, GroupTileSize / 2 - 1);
 			const int32 TileSize = GroupTileSize - 2 * TileOverscan;
 
 			FTSRTSRWeightRelaxationCS::FParameters* PassParameters = GraphBuilder.AllocParameters<FTSRTSRWeightRelaxationCS::FParameters>();
@@ -2673,6 +2694,31 @@ FDefaultTemporalUpscaler::FOutputs AddTemporalSuperResolutionPasses(
 			PassParameters->ThinGeometryTextureIndex = ThinGeometryTextureIndex;
 			PassParameters->MaxRelaxationWeight = FMath::Clamp(CVarTSRThinGeometryCoverageMaxRelaxationWeight.GetValueOnRenderThread(), 0.0f, 1.0f);
 
+			
+			if (PassInputs.FlickeringInputTexture.IsValid() && !bCameraCut)
+			{
+				ensure(InputRect == PassInputs.FlickeringInputTexture.ViewRect);
+				int32 ArrayIndex = GetTSRMainFlickeringLumaTextureArraySize() - 1;
+				PassParameters->CurrentCoverageTexture = GraphBuilder.CreateSRV(
+					FRDGTextureSRVDesc::CreateForSlice(PassInputs.FlickeringInputTexture.Texture,  /* SliceIndex = */ ArrayIndex));
+			}
+			else
+			{
+				PassParameters->CurrentCoverageTexture = GraphBuilder.CreateSRV(FRDGTextureSRVDesc(BlackDummy));
+			}
+			
+
+			if (PassInputs.FlickeringInputTexture.IsValid())
+			{
+				ensure(InputRect == PassInputs.FlickeringInputTexture.ViewRect);
+				PassParameters->InputMoireLumaTexture = GraphBuilder.CreateSRV(
+					FRDGTextureSRVDesc::CreateForSlice(PassInputs.FlickeringInputTexture.Texture, 0));
+			}
+			else
+			{
+				PassParameters->InputMoireLumaTexture = GraphBuilder.CreateSRV(FRDGTextureSRVDesc(BlackDummy));
+			}
+
 			PassParameters->InputTexture = PassInputs.SceneColor.Texture;
 			PassParameters->InputSceneTranslucencyTexture = SeparateTranslucencyTexture;
 			PassParameters->R8Output = GraphBuilder.CreateUAV(R8OutputTexture);
