Real time sky capture and volumetric cloud - compute cloud tracing now overlaps when capture for real time refelction pature faces. It is also possible to reduce the resolution of each faces, now composited later over the real cube map resolution, reusing the bottom sky color pass.

#rnx
#tests PC editor dx12 and console.

diff --git a/Engine/Shaders/Private/ReflectionEnvironmentShaders.usf b/Engine/Shaders/Private/ReflectionEnvironmentShaders.usf
--- a/Engine/Shaders/Private/ReflectionEnvironmentShaders.usf
+++ b/Engine/Shaders/Private/ReflectionEnvironmentShaders.usf
@@ -395,8 +395,13 @@ void DownsampleIntegrateCubemapIlluminancePS(
 
 #if APPLY_LOWER_HEMISPHERE_COLOR_PIXELSHADER
 
+int ApplyLowerHemisphereColor;
 float4 LowerHemisphereSolidColor;
 
+int ApplyLowResCloudTexture;
+Texture2D<float4> LowResCloudTexture;
+SamplerState LowResCloudSampler;
+
 // This pixel shader is designed to work with RenderRealTimeReflectionHeightFogVS
 void ApplyLowerHemisphereColorPS(
 	in float4 SVPos : SV_POSITION,
@@ -409,12 +414,18 @@ void ApplyLowerHemisphereColorPS(
 	float3 N = normalize(CubeCoordinates);
 
 	OutLuminanceAlpha = float4(0.0f, 0.0f, 0.0f, 1.0f);
-	if (N.z < 0.0f)
+
+	if (ApplyLowResCloudTexture > 0)
+	{
+		OutLuminanceAlpha = LowResCloudTexture.Sample(LowResCloudSampler, UV);
+	}
+
+	if (N.z < 0.0f && ApplyLowerHemisphereColor > 0)
 	{
 		float Coverage = 0.0f;
 		Coverage = pow(saturate(LowerHemisphereSolidColor.a), 2.2);
 		Coverage = saturate(LowerHemisphereSolidColor.a);
-		OutLuminanceAlpha = float4(LowerHemisphereSolidColor.rgb * Coverage, 1.0 - Coverage);
+		OutLuminanceAlpha = float4(lerp(OutLuminanceAlpha.rgb, LowerHemisphereSolidColor.rgb ,Coverage), OutLuminanceAlpha.a * (1.0 - Coverage));
 	}
 }
 
diff --git a/Engine/Source/Runtime/Renderer/Private/ReflectionEnvironmentRealTimeCapture.cpp b/Engine/Source/Runtime/Renderer/Private/ReflectionEnvironmentRealTimeCapture.cpp
--- a/Engine/Source/Runtime/Renderer/Private/ReflectionEnvironmentRealTimeCapture.cpp
+++ b/Engine/Source/Runtime/Renderer/Private/ReflectionEnvironmentRealTimeCapture.cpp
@@ -59,6 +59,11 @@ static TAutoConsoleVariable<int32> CVarRealTimeReflectionCaptureAlwaysClearColor
 	TEXT("Always clear the color buffer to black before rendering SkyAtmosphere, SkyPass, Fog, etc."),
 	ECVF_RenderThreadSafe);
 
+static TAutoConsoleVariable<int32> CVarRealTimeReflectionCaptureVolumetricCloudResolutionDivider(
+	TEXT("r.SkyLight.RealTimeReflectionCapture.VolumetricCloudResolutionDivider"), 1,
+	TEXT("The divider applied on the resolution when capturing cloud in the real time sky light. For instance, a value of 2 will render cloud at half resolution and thus will be faster to render."),
+	ECVF_RenderThreadSafe | ECVF_Scalability);
+
 static TAutoConsoleVariable<int32> CVarRealTimeReflectionCaptureResolutionOverride(
 	TEXT("r.SkyLight.RealTimeReflectionCapture.ResolutionOverride"), 0,
 	TEXT("Override the real time sky light capture resolution if this CVar is greater than 16."),
@@ -180,9 +185,13 @@ class FApplyLowerHemisphereColorPS : public FGlobalShader
 
 	BEGIN_SHADER_PARAMETER_STRUCT(FParameters, )
 		SHADER_PARAMETER_STRUCT_REF(FViewUniformShaderParameters, ViewUniformBuffer)
+		SHADER_PARAMETER(int32, ApplyLowerHemisphereColor)
 		SHADER_PARAMETER(FVector4f, LowerHemisphereSolidColor)
 		SHADER_PARAMETER(FVector2f, SvPositionToUVScale)
 		SHADER_PARAMETER(int32, CubeFace)
+		SHADER_PARAMETER(int32, ApplyLowResCloudTexture)
+		SHADER_PARAMETER_RDG_TEXTURE(Texture2D<float4>, LowResCloudTexture)
+		SHADER_PARAMETER_SAMPLER(SamplerState, LowResCloudSampler)
 		RENDER_TARGET_BINDING_SLOTS()
 	END_SHADER_PARAMETER_STRUCT()
 
@@ -370,29 +379,35 @@ void FScene::AllocateAndCaptureFrameSkyEnvMap(
 	const uint32 CubeWidth = CaptureResolutionOverride >= 16 ? CaptureResolutionOverride : SkyLight->CaptureCubeMapResolution;
 	const uint32 CubeMipCount = FMath::CeilLogTwo(CubeWidth) + 1;
 
-	// Make a snapshot we are going to use for the 6 cubemap faces and set it up.
-	// Note: cube view is not meant to be sent to lambdas because we only create a single one. You should only send the ViewUniformBuffer around.
-	FViewInfo& CubeView = *MainView.CreateSnapshot();
-	CubeView.FOV = 90.0f;
-	// Note: We cannot override exposure because sky input texture are using exposure
-
-	// Other view data clean up
-	CubeView.StereoPass = EStereoscopicPass::eSSP_FULL;
-	CubeView.DrawDynamicFlags = EDrawDynamicFlags::ForceLowestLOD;
-	CubeView.MaterialTextureMipBias = 0;
-
-	FViewMatrices::FMinimalInitializer SceneCubeViewInitOptions;
-	SceneCubeViewInitOptions.ConstrainedViewRect = FIntRect(FIntPoint(0, 0), FIntPoint(CubeWidth, CubeWidth));
-
-	FBox VolumeBounds[TVC_MAX];
-	CubeView.CachedViewUniformShaderParameters = MakeUnique<FViewUniformShaderParameters>();
-	CubeView.SetupUniformBufferParameters(
-		VolumeBounds,
-		TVC_MAX,
-		*CubeView.CachedViewUniformShaderParameters);
+	auto CreateMainViewSnapshotForRealTimeCapture = [&MainView](FViewInfo*& OutView, FMatrix& OutCubeProjectionMatrix, float CubeViewWidth)
+	{
+		// Make a snapshot we are going to use for the 6 cubemap faces and set it up.
+		// Note: cube view is not meant to be sent to lambdas because we only create a single one. You should only send the ViewUniformBuffer around.
+		OutView = MainView.CreateSnapshot();
+		FViewInfo& CubeView = *OutView;
+		CubeView.FOV = 90.0f;
+		// Note: We cannot override exposure because sky input texture are using exposure
+
+		// Other view data clean up
+		CubeView.StereoPass = EStereoscopicPass::eSSP_FULL;
+		CubeView.DrawDynamicFlags = EDrawDynamicFlags::ForceLowestLOD;
+		CubeView.MaterialTextureMipBias = 0;
+
+		FBox VolumeBounds[TVC_MAX];
+		CubeView.CachedViewUniformShaderParameters = MakeUnique<FViewUniformShaderParameters>();
+		CubeView.SetupUniformBufferParameters(
+			VolumeBounds,
+			TVC_MAX,
+			*CubeView.CachedViewUniformShaderParameters);
+
+		OutCubeProjectionMatrix = GetCubeProjectionMatrix(CubeView.FOV * 0.5f, (float)CubeViewWidth, GReflectionCaptureNearPlane);
+		CubeView.UpdateProjectionMatrix(OutCubeProjectionMatrix);
+	};
 
-	const FMatrix CubeProjectionMatrix = GetCubeProjectionMatrix(CubeView.FOV * 0.5f, (float)CubeWidth, GReflectionCaptureNearPlane);
-	CubeView.UpdateProjectionMatrix(CubeProjectionMatrix);
+	FMatrix CubeProjectionMatrix;
+	FViewInfo* CubeViewPtr = nullptr;
+	CreateMainViewSnapshotForRealTimeCapture(CubeViewPtr, CubeProjectionMatrix, CubeWidth);
+	FViewInfo& CubeView = *CubeViewPtr;
 
 	FPooledRenderTargetDesc SkyCubeTexDesc = Translate(FSkyPassMeshProcessor::GetCaptureFrameSkyEnvMapTextureDesc(CubeWidth, CubeMipCount));
 
@@ -528,32 +543,34 @@ void FScene::AllocateAndCaptureFrameSkyEnvMap(
 				}
 			}
 
-			const bool bAlwaysClearColorBuffer = CVarRealTimeReflectionCaptureAlwaysClearColorBuffer.GetValueOnRenderThread() > 0;
-			for (int32 CubeFace = StartCubeFace; CubeFace < EndCubeFace; CubeFace++)
+			const FSkyLightSceneProxy* SkyLightProxy = SkyLight;
+			auto SetupCommonViewUniformBufferParameters = [SkyLightProxy](FViewInfo& CubeView, FMatrix CubeProjectionMatrix, FViewMatrices& OutCubeViewMatrices, float CubeViewWidth, int32 CubeFace)
 			{
-				RDG_EVENT_SCOPE(GraphBuilder, "Capture Face=%d", CubeFace);
-
-				SkyRC.RenderTargets[0] = FRenderTargetBinding(SkyCubeTexture, ERenderTargetLoadAction::ELoad, 0, CubeFace);
-
 				const FMatrix CubeViewRotationMatrix = CalcCubeFaceViewRotationMatrix((ECubeFace)CubeFace);
 
+				FViewMatrices::FMinimalInitializer SceneCubeViewInitOptions;
+				SceneCubeViewInitOptions.ConstrainedViewRect = FIntRect(FIntPoint(0, 0), FIntPoint(CubeViewWidth, CubeViewWidth));
 				SceneCubeViewInitOptions.ViewRotationMatrix = CubeViewRotationMatrix;
-				SceneCubeViewInitOptions.ViewOrigin = SkyLight->CapturePosition;
+				SceneCubeViewInitOptions.ViewOrigin = SkyLightProxy->CapturePosition;
 				SceneCubeViewInitOptions.ProjectionMatrix = CubeProjectionMatrix;
-				FViewMatrices CubeViewMatrices = FViewMatrices(SceneCubeViewInitOptions);
+				OutCubeViewMatrices = FViewMatrices(SceneCubeViewInitOptions);
+
 				CubeView.SetupCommonViewUniformBufferParameters(
 					*CubeView.CachedViewUniformShaderParameters,
-					FIntPoint(CubeWidth, CubeWidth),
+					FIntPoint(CubeViewWidth, CubeViewWidth),
 					1,
-					FIntRect(FIntPoint(0, 0), FIntPoint(CubeWidth, CubeWidth)),
-					CubeViewMatrices,
-					CubeViewMatrices);
+					FIntRect(FIntPoint(0, 0), FIntPoint(CubeViewWidth, CubeViewWidth)),
+					OutCubeViewMatrices,
+					OutCubeViewMatrices);
 
 				// Notify the fact that we render a reflection, e.g. remove sun disk.
 				CubeView.CachedViewUniformShaderParameters->RenderingReflectionCaptureMask = 1.0f;
 				// Notify the fact that we render a reflection, e.g. use special exposure.
 				CubeView.CachedViewUniformShaderParameters->RealTimeReflectionCapture = 1.0f;
+			};
 
+			auto SetupViewSkyAtmosphereParametersAndResources = [&](FViewInfo& OutView)
+			{
 				// We have rendered a sky dome with identity rotation at the SkyLight position for the capture.
 				if (AtmosphereSetup)
 				{
@@ -569,7 +586,7 @@ void FScene::AllocateAndCaptureFrameSkyEnvMap(
 						AtmosphereSetup->ComputeViewData(
 							SkyLight->CapturePosition, MainView.ViewMatrices.GetPreViewTranslation(), SkyViewLutReferentialForward, SkyViewLutReferentialRight,
 							SkyCameraTranslatedWorldOrigin, TempSkyPlanetData, SkyViewLutReferential);
-						CubeView.CachedViewUniformShaderParameters->SkyViewLutTexture = RealTimeReflectionCaptureSkyAtmosphereViewLutTexture->GetRHI();
+						OutView.CachedViewUniformShaderParameters->SkyViewLutTexture = RealTimeReflectionCaptureSkyAtmosphereViewLutTexture->GetRHI();
 					}
 					else
 					{
@@ -583,20 +600,86 @@ void FScene::AllocateAndCaptureFrameSkyEnvMap(
 							SkyCameraTranslatedWorldOrigin, TempSkyPlanetData, SkyViewLutReferential);
 					}
 
-					CubeView.CachedViewUniformShaderParameters->SkyPlanetTranslatedWorldCenterAndViewHeight = TempSkyPlanetData;
-					CubeView.CachedViewUniformShaderParameters->SkyCameraTranslatedWorldOrigin = SkyCameraTranslatedWorldOrigin;
-					CubeView.CachedViewUniformShaderParameters->SkyViewLutReferential = SkyViewLutReferential;
+					OutView.CachedViewUniformShaderParameters->SkyPlanetTranslatedWorldCenterAndViewHeight = TempSkyPlanetData;
+					OutView.CachedViewUniformShaderParameters->SkyCameraTranslatedWorldOrigin = SkyCameraTranslatedWorldOrigin;
+					OutView.CachedViewUniformShaderParameters->SkyViewLutReferential = SkyViewLutReferential;
 				}
 
-				if (HasSkyAtmosphere()&& (MainView.bSceneHasSkyMaterial || HasVolumetricCloud())
+				if (HasSkyAtmosphere() && (MainView.bSceneHasSkyMaterial || HasVolumetricCloud())
 					&& RealTimeReflectionCaptureCamera360APLutTexture.IsValid())	// we also check that because it seems it can happen for some view setup UE-107270, TODO find a repro for a proper fix.
 				{
-					CubeView.CachedViewUniformShaderParameters->CameraAerialPerspectiveVolume = RealTimeReflectionCaptureCamera360APLutTexture->GetRHI();
+					OutView.CachedViewUniformShaderParameters->CameraAerialPerspectiveVolume = RealTimeReflectionCaptureCamera360APLutTexture->GetRHI();
 				}
 				else
 				{
-					CubeView.CachedViewUniformShaderParameters->CameraAerialPerspectiveVolume = GSystemTextures.VolumetricBlackDummy->GetRHI();
+					OutView.CachedViewUniformShaderParameters->CameraAerialPerspectiveVolume = GSystemTextures.VolumetricBlackDummy->GetRHI();
+				}
+			};
+
+			// Render clouds in separate textures all at once in parallel with compute work that overlaps on GPU to be faster.
+			// The texture will be composited later using the ApplyLowerHemisphereColor pass.
+			const int32 CloudResolutionDivider = FMath::Clamp(CVarRealTimeReflectionCaptureVolumetricCloudResolutionDivider.GetValueOnRenderThread(), 1, 8);
+			FRDGTextureRef LowResCloudTextureArray[6] = { nullptr, nullptr, nullptr, nullptr, nullptr, nullptr };
+			if (bShouldRenderVolumetricCloud && bExecuteCloud)
+			{
+				for (int32 CubeFace = StartCubeFace; CubeFace < EndCubeFace; CubeFace++)
+				{
+					RDG_EVENT_SCOPE(GraphBuilder, "Capture Volumetric Cloud Face=%d", CubeFace);
+
+					FRDGTextureRef& LowResCloudTexture = LowResCloudTextureArray[CubeFace];
+					const uint32 CubeWidthLow = CubeWidth / CloudResolutionDivider;
+
+					FMatrix LowResCubeProjectionMatrix;
+					FViewInfo* LowResCubeViewPtr = nullptr;
+					CreateMainViewSnapshotForRealTimeCapture(LowResCubeViewPtr, LowResCubeProjectionMatrix, CubeWidthLow);
+					FViewInfo& LowResCubeView = *CubeViewPtr;
+
+					FViewMatrices CubeViewMatrices;
+					SetupCommonViewUniformBufferParameters(LowResCubeView, LowResCubeProjectionMatrix, CubeViewMatrices, CubeWidthLow, CubeFace);
+
+					// We have rendered a sky dome with identity rotation at the SkyLight position for the capture.
+					SetupViewSkyAtmosphereParametersAndResources(LowResCubeView);
+
+					LowResCubeView.CreateViewUniformBuffers(*LowResCubeView.CachedViewUniformShaderParameters);
+					CloudRC.ViewUniformBuffer = LowResCubeView.ViewUniformBuffer;
+
+					// Create a texture to render the clouds at a given resolution
+					FRDGTextureDesc LowResCloudTexDesc = FRDGTextureDesc::Create2D(FIntPoint(CubeWidthLow, CubeWidthLow), PF_FloatRGBA,
+						FClearValueBinding::Black,
+						TexCreate_RenderTargetable | TexCreate_UAV | TexCreate_ShaderResource);
+					LowResCloudTexture = GraphBuilder.CreateTexture(LowResCloudTexDesc, TEXT("SkyLight.LowResCloudTexture"));
+
+					// Render
+					CloudRC.RenderTargets[0] = FRenderTargetBinding(LowResCloudTexture, ERenderTargetLoadAction::EClear);
+					CloudRC.bDisableCloudBlending = true;
+
+					FCloudShadowAOData CloudShadowAOData;
+					GetCloudShadowAOData(GetVolumetricCloudSceneInfo(), CubeView, GraphBuilder, CloudShadowAOData);
+					SkyRC.bShouldSampleCloudShadow = CloudShadowAOData.bShouldSampleCloudShadow;
+					SkyRC.VolumetricCloudShadowMap[0] = CloudShadowAOData.VolumetricCloudShadowMap[0];
+					SkyRC.VolumetricCloudShadowMap[1] = CloudShadowAOData.VolumetricCloudShadowMap[1];
+					SkyRC.bShouldSampleCloudSkyAO = CloudShadowAOData.bShouldSampleCloudSkyAO;
+					SkyRC.VolumetricCloudSkyAO = CloudShadowAOData.VolumetricCloudSkyAO;
+					// TODO this CloudShadowAOData management looks a bit heavy, simplify or make common
+					CloudRC.VolumetricCloudShadowTexture[0] = CloudShadowAOData.VolumetricCloudShadowMap[0];
+					CloudRC.VolumetricCloudShadowTexture[1] = CloudShadowAOData.VolumetricCloudShadowMap[1];
+
+					SceneRenderer.RenderVolumetricCloudsInternal(GraphBuilder, CloudRC, InstanceCullingManager, FIntPoint(CubeWidthLow, CubeWidthLow));
 				}
+			}
+
+			const bool bAlwaysClearColorBuffer = CVarRealTimeReflectionCaptureAlwaysClearColorBuffer.GetValueOnRenderThread() > 0;
+			for (int32 CubeFace = StartCubeFace; CubeFace < EndCubeFace; CubeFace++)
+			{
+				RDG_EVENT_SCOPE(GraphBuilder, "Capture Face=%d", CubeFace);
+
+				SkyRC.RenderTargets[0] = FRenderTargetBinding(SkyCubeTexture, ERenderTargetLoadAction::ELoad, 0, CubeFace);
+
+				FViewMatrices CubeViewMatrices;
+				SetupCommonViewUniformBufferParameters(CubeView, CubeProjectionMatrix, CubeViewMatrices, CubeWidth, CubeFace);
+
+				// We have rendered a sky dome with identity rotation at the SkyLight position for the capture.
+				SetupViewSkyAtmosphereParametersAndResources(CubeView);
 
 				CubeView.CreateViewUniformBuffers(*CubeView.CachedViewUniformShaderParameters);
 
@@ -843,30 +926,9 @@ void FScene::AllocateAndCaptureFrameSkyEnvMap(
 					}
 				}
 
-				if (bShouldRenderVolumetricCloud && bExecuteCloud)
-				{
-					CloudRC.ViewUniformBuffer = CubeView.ViewUniformBuffer;
-
-					CloudRC.RenderTargets[0] = SkyRC.RenderTargets[0];
-					//	CloudRC.RenderTargets[1] = Null target will skip export
-
-					if (CloudRC.ComputeOverlapCloudColorCubeTextureUAVWithoutBarrier == nullptr)
-					{
-						// Create a UAV skipping barrier to make sure each compute pass processing each faces can overlap. 
-						// Here in ReflectionEnvironmentRealTimeCapture we know it is a cube texture
-						FRDGTextureRef CubeTexture = CloudRC.RenderTargets[0].GetTexture();
-						check(CubeTexture->Desc.IsTextureCube());
-						CloudRC.ComputeOverlapCloudColorCubeTextureUAVWithoutBarrier = GraphBuilder.CreateUAV(CubeTexture, ERDGUnorderedAccessViewFlags::SkipBarrier);
-					}
-
-					CloudRC.VolumetricCloudShadowTexture[0] = CloudShadowAOData.VolumetricCloudShadowMap[0];
-					CloudRC.VolumetricCloudShadowTexture[1] = CloudShadowAOData.VolumetricCloudShadowMap[1];
-
-					SceneRenderer.RenderVolumetricCloudsInternal(GraphBuilder, CloudRC, InstanceCullingManager, FIntPoint(CubeWidth, CubeWidth));
-				}
-
-				// Now blend over the lower hemisphere color, during the second pass for clouds
-				if (SkyLight->bLowerHemisphereIsSolidColor && bExecuteCloud)
+				// Now composited: (1) clouds rendered in isolated texture, and (2) blend over the lower hemisphere color.
+				FRDGTextureRef& LowResCloudTexture = LowResCloudTextureArray[CubeFace];
+				if ((SkyLight->bLowerHemisphereIsSolidColor || LowResCloudTexture !=nullptr) && bExecuteCloud)
 				{
 					FRenderRealTimeReflectionHeightFogVS::FPermutationDomain VsPermutationVector;
 					TShaderMapRef<FRenderRealTimeReflectionHeightFogVS> VertexShader(GetGlobalShaderMap(SkyRC.FeatureLevel), VsPermutationVector);
@@ -877,7 +939,11 @@ void FScene::AllocateAndCaptureFrameSkyEnvMap(
 					FApplyLowerHemisphereColorPS::FParameters* PsPassParameters = GraphBuilder.AllocParameters<FApplyLowerHemisphereColorPS::FParameters>();
 					PsPassParameters->RenderTargets = SkyRC.RenderTargets;
 					PsPassParameters->ViewUniformBuffer = CubeView.ViewUniformBuffer;
+					PsPassParameters->ApplyLowerHemisphereColor = SkyLight->bLowerHemisphereIsSolidColor ? 1 : 0;
 					PsPassParameters->LowerHemisphereSolidColor = SkyLight->LowerHemisphereColor;
+					PsPassParameters->ApplyLowResCloudTexture = LowResCloudTexture ? 1 : 0;
+					PsPassParameters->LowResCloudTexture = LowResCloudTexture ? LowResCloudTexture : GSystemTextures.GetBlackAlphaOneDummy(GraphBuilder);
+					PsPassParameters->LowResCloudSampler = TStaticSamplerState<SF_Bilinear>::GetRHI();;
 					PsPassParameters->CubeFace = CubeFace;
 					PsPassParameters->SvPositionToUVScale = FVector2f(1.0f / float(CubeWidth), 1.0f / float(CubeWidth));
 
