Substrate - fix for mesh decal rendering wrong logic.
Follow up to 39322119

#rnx
#tests pc editor dx12/dx11

diff --git a/Engine/Source/Runtime/Renderer/Private/CompositionLighting/PostProcessDeferredDecals.cpp b/Engine/Source/Runtime/Renderer/Private/CompositionLighting/PostProcessDeferredDecals.cpp
--- a/Engine/Source/Runtime/Renderer/Private/CompositionLighting/PostProcessDeferredDecals.cpp
+++ b/Engine/Source/Runtime/Renderer/Private/CompositionLighting/PostProcessDeferredDecals.cpp
@@ -762,7 +762,7 @@ void AddDeferredDecalPass(
 		{
 			// Sanity check - Substrate only support DBuffer, SceneColor, or AO decals
 			bool bDecalSupported = true;
-			if (Substrate::IsSubstrateEnabled() && (!Substrate::IsSubstrateBlendableGBufferEnabled(ShaderPlatform) || DecalRenderStage == EDecalRenderStage::BeforeLighting))
+			if (Substrate::IsSubstrateEnabled() && !Substrate::IsSubstrateBlendableGBufferEnabled(ShaderPlatform) && DecalRenderStage == EDecalRenderStage::BeforeLighting)
 			{
 				bDecalSupported = false;
 			}
