[Lit Translucency] Refactor TranslucentLightingVolume update functions.
- dispatch all TLV update logic from a single function, FDeferredShadingSceneRenderer::RenderTranslucencyLightingVolume(...).
- clean up unecessary parameters.

diff --git a/Engine/Source/Runtime/Renderer/Private/DeferredShadingRenderer.cpp b/Engine/Source/Runtime/Renderer/Private/DeferredShadingRenderer.cpp
--- a/Engine/Source/Runtime/Renderer/Private/DeferredShadingRenderer.cpp
+++ b/Engine/Source/Runtime/Renderer/Private/DeferredShadingRenderer.cpp
@@ -2944,7 +2944,7 @@ void FDeferredShadingSceneRenderer::Render(FRDGBuilder& GraphBuilder, const FSce
 
 			const FSortedLightSetSceneInfo& SortedLightSet = *GatherAndSortLightsTask.GetResult();
 
-			RenderLights(GraphBuilder, SceneTextures, TranslucencyLightingVolumeTextures, LightingChannelsTexture, SortedLightSet);
+			RenderLights(GraphBuilder, SceneTextures, LightingChannelsTexture, SortedLightSet);
 
 			if (SortedLightSet.MegaLightsLightStart < SortedLightSet.SortedLights.Num())
 			{
@@ -2958,8 +2958,7 @@ void FDeferredShadingSceneRenderer::Render(FRDGBuilder& GraphBuilder, const FSce
 			// Copy depth history without water and translucency for ray traced lighting denoising
 			StoreStochasticLightingSceneHistory(GraphBuilder, LumenFrameTemporaries, SceneTextures);
 
-			InjectTranslucencyLightingVolumeAmbientCubemap(GraphBuilder, Views, TranslucencyLightingVolumeTextures);
-			FilterTranslucencyLightingVolume(GraphBuilder, Views, TranslucencyLightingVolumeTextures);
+			RenderTranslucencyLightingVolume(GraphBuilder, TranslucencyLightingVolumeTextures, SortedLightSet);
 
 			// Do DiffuseIndirectComposite after Lights so that async Lumen work can overlap
 			RenderDiffuseIndirectAndAmbientOcclusion(
diff --git a/Engine/Source/Runtime/Renderer/Private/DeferredShadingRenderer.h b/Engine/Source/Runtime/Renderer/Private/DeferredShadingRenderer.h
--- a/Engine/Source/Runtime/Renderer/Private/DeferredShadingRenderer.h
+++ b/Engine/Source/Runtime/Renderer/Private/DeferredShadingRenderer.h
@@ -848,10 +848,14 @@ private:
 	void RenderLights(
 		FRDGBuilder& GraphBuilder,
 		FMinimalSceneTextures& SceneTextures,
-		FTranslucencyLightingVolumeTextures& TranslucencyLightingVolumeTextures,
 		FRDGTextureRef LightingChannelsTexture,
 		const FSortedLightSetSceneInfo& SortedLightSet);
 
+	void RenderTranslucencyLightingVolume(
+		FRDGBuilder& GraphBuilder,
+		FTranslucencyLightingVolumeTextures& Textures,
+		const FSortedLightSetSceneInfo& SortedLightSet);
+
 	/** Render stationary light overlap as complexity to scene color. */
 	void RenderStationaryLightOverlap(
 		FRDGBuilder& GraphBuilder,
@@ -913,9 +917,6 @@ private:
 		FTranslucencyPassResourcesMap& TranslucencyResourceMap);
 
 	void CollectLightForTranslucencyLightingVolumeInjection(
-		FRDGBuilder& GraphBuilder,
-		const FMinimalSceneTextures& SceneTextures,
-		const FTranslucencyLightingVolumeTextures& TranslucencyLightingVolumeTextures,
 		const FLightSceneInfo* LightSceneInfo,
 		bool bSupportShadowMaps,
 		FTranslucentLightInjectionCollector& Collector);
@@ -933,7 +934,6 @@ private:
 	void RenderDeferredShadowProjections(
 		FRDGBuilder& GraphBuilder,
 		const FMinimalSceneTextures& SceneTextures,
-		const FTranslucencyLightingVolumeTextures& TranslucencyLightingVolumeTextures,
 		const FLightSceneInfo* LightSceneInfo,
 		FRDGTextureRef ScreenShadowMaskTexture,
 		FRDGTextureRef ScreenShadowMaskSubPixelTexture);
diff --git a/Engine/Source/Runtime/Renderer/Private/LightRendering.cpp b/Engine/Source/Runtime/Renderer/Private/LightRendering.cpp
--- a/Engine/Source/Runtime/Renderer/Private/LightRendering.cpp
+++ b/Engine/Source/Runtime/Renderer/Private/LightRendering.cpp
@@ -1096,21 +1096,6 @@ public:
 
 IMPLEMENT_GLOBAL_SHADER(FDeferredLightOverlapPS, "/Engine/Private/StationaryLightOverlapShaders.usf", "OverlapPixelMain", SF_Pixel);
 
-static void SplitSimpleLightsByView(TArrayView<const FViewInfo> Views, const FSimpleLightArray& SimpleLights, TArrayView<FSimpleLightArray> SimpleLightsByView)
-{
-	check(SimpleLightsByView.Num() == Views.Num());
-
-	for (int32 LightIndex = 0; LightIndex < SimpleLights.InstanceData.Num(); ++LightIndex)
-	{
-		for (int32 ViewIndex = 0; ViewIndex < Views.Num(); ++ViewIndex)
-		{
-			FSimpleLightPerViewEntry PerViewEntry = SimpleLights.GetViewDependentData(LightIndex, ViewIndex, Views.Num());
-			SimpleLightsByView[ViewIndex].InstanceData.Add(SimpleLights.InstanceData[LightIndex]);
-			SimpleLightsByView[ViewIndex].PerViewData.Add(PerViewEntry);
-		}
-	}
-}
-
 /** Gathers simple lights from visible primtives in the passed in views. */
 void FSceneRenderer::GatherSimpleLights(const FSceneViewFamily& ViewFamily, const TArray<FViewInfo>& Views, FSimpleLightArray& SimpleLights)
 {
@@ -1431,7 +1416,6 @@ FHairStrandsTransmittanceMaskData CreateDummyHairStrandsTransmittanceMaskData(FR
 void FDeferredShadingSceneRenderer::RenderLights(
 	FRDGBuilder& GraphBuilder,
 	FMinimalSceneTextures& SceneTextures,
-	FTranslucencyLightingVolumeTextures& TranslucencyLightingVolumeTextures,
 	FRDGTextureRef LightingChannelsTexture,
 	const FSortedLightSetSceneInfo& SortedLightSet)
 {
@@ -1450,7 +1434,6 @@ void FDeferredShadingSceneRenderer::RenderLights(
 	SCOPE_CYCLE_COUNTER(STAT_LightingDrawTime);
 	SCOPE_CYCLE_COUNTER(STAT_LightRendering);
 
-	const FSimpleLightArray& SimpleLights = SortedLightSet.SimpleLights;
 	const TArray<FSortedLightSceneInfo, SceneRenderingAllocator>& SortedLights = SortedLightSet.SortedLights;
 	const int32 SimpleLightsEnd = SortedLightSet.SimpleLightsEnd;
 	const int32 UnbatchedLightStart = SortedLightSet.UnbatchedLightStart;
@@ -1472,63 +1455,6 @@ void FDeferredShadingSceneRenderer::RenderLights(
 			Substrate::AddSubstrateStencilPass(GraphBuilder, Views, SceneTextures);
 		}
 
-		if (ViewFamily.EngineShowFlags.DirectLighting && GUseTranslucentLightingVolumes && RHISupportsVolumeTextureRendering(ViewFamily.GetShaderPlatform()))
-		{
-			RDG_EVENT_SCOPE(GraphBuilder, "InjectTranslucencyLightingVolume");
-			SCOPE_CYCLE_COUNTER(STAT_TranslucentInjectTime);
-
-			if (SimpleLights.InstanceData.Num() > 0)
-			{
-				auto& SimpleLightsByView = *GraphBuilder.AllocObject<TArray<FSimpleLightArray, SceneRenderingAllocator>>();
-				SimpleLightsByView.SetNum(Views.Num());
-
-				SplitSimpleLightsByView(Views, SimpleLights, SimpleLightsByView);
-
-				for (int32 ViewIndex = 0; ViewIndex < Views.Num(); ++ViewIndex)
-				{
-					FSimpleLightArray& SimpleLightArray = SimpleLightsByView[ViewIndex];
-
-					if (SimpleLightArray.InstanceData.Num() > 0)
-					{
-						FViewInfo& View = Views[ViewIndex];
-						RDG_GPU_MASK_SCOPE(GraphBuilder, View.GPUMask);
-						RDG_EVENT_SCOPE(GraphBuilder, "InjectSimpleLightsTranslucentLighting");
-						InjectSimpleTranslucencyLightingVolumeArray(GraphBuilder, View, ViewIndex, Views.Num(), TranslucencyLightingVolumeTextures, SimpleLightArray);
-					}
-				}
-			}
-
-			// Shadowed and light function lights
-			{
-				FTranslucentLightInjectionCollector Collector(GraphBuilder, Views, bAreLightsInLightGrid);
-
-				// Collect all the light injection data
-				for (int32 LightIndex = SimpleLightsEnd; LightIndex < SortedLights.Num(); LightIndex++)
-				{
-					const FSortedLightSceneInfo& SortedLightInfo = SortedLights[LightIndex];
-					const FLightSceneInfo& LightSceneInfo = *SortedLightInfo.LightSceneInfo;
-					const FLightSceneProxy& LightSceneProxy = *LightSceneInfo.Proxy;
-
-					const bool bDrawShadows = SortedLightInfo.SortKey.Fields.bShadowed;
-					const FLightOcclusionType OcclusionType = GetLightOcclusionType(LightSceneProxy, ViewFamily);
-					const bool bSupportShadowMaps = bDrawShadows && OcclusionType == FLightOcclusionType::Shadowmap;
-
-					// Collect all the light injection data
-					CollectLightForTranslucencyLightingVolumeInjection(
-						GraphBuilder, SceneTextures, TranslucencyLightingVolumeTextures, &LightSceneInfo, bSupportShadowMaps, Collector);
-				}
-
-				// Run light injection for each view
-				for (int32 ViewIndex = 0; ViewIndex < Views.Num(); ViewIndex++)
-				{
-					FViewInfo& View = Views[ViewIndex];
-					RDG_GPU_MASK_SCOPE(GraphBuilder, View.GPUMask);
-					RDG_EVENT_SCOPE(GraphBuilder, "InjectTranslucencyLightingVolume(View=%d)", ViewIndex);
-					InjectTranslucencyLightingVolume(GraphBuilder, View, ViewIndex, Scene, *this, Collector, TranslucencyLightingVolumeTextures);
-				}
-			}
-		}
-
 		// The shadow scene renderer is not optional in the deferred renderer
 		FShadowSceneRenderer& ShadowSceneRenderer = GetSceneExtensionsRenderers().GetRenderer<FShadowSceneRenderer>();
 
@@ -2169,11 +2095,11 @@ void FDeferredShadingSceneRenderer::RenderLights(
 
 						if (bIsHairPass)
 						{
-							RenderDeferredShadowProjections(GraphBuilder, SceneTextures, TranslucencyLightingVolumeTextures, &LightSceneInfo, nullptr, ScreenShadowMaskSubPixelTexture);
+							RenderDeferredShadowProjections(GraphBuilder, SceneTextures, &LightSceneInfo, nullptr, ScreenShadowMaskSubPixelTexture);
 						}
 						else
 						{
-							RenderDeferredShadowProjections(GraphBuilder, SceneTextures, TranslucencyLightingVolumeTextures, &LightSceneInfo, ScreenShadowMaskTexture, nullptr);
+							RenderDeferredShadowProjections(GraphBuilder, SceneTextures, &LightSceneInfo, ScreenShadowMaskTexture, nullptr);
 						}
 					}
 
diff --git a/Engine/Source/Runtime/Renderer/Private/ShadowRendering.cpp b/Engine/Source/Runtime/Renderer/Private/ShadowRendering.cpp
--- a/Engine/Source/Runtime/Renderer/Private/ShadowRendering.cpp
+++ b/Engine/Source/Runtime/Renderer/Private/ShadowRendering.cpp
@@ -2260,9 +2260,6 @@ void FSceneRenderer::BeginAsyncDistanceFieldShadowProjections(FRDGBuilder& Graph
 
 
 void FDeferredShadingSceneRenderer::CollectLightForTranslucencyLightingVolumeInjection(
-	FRDGBuilder& GraphBuilder,
-	const FMinimalSceneTextures& SceneTextures,
-	const FTranslucencyLightingVolumeTextures& TranslucencyLightingVolumeTextures,
 	const FLightSceneInfo* LightSceneInfo,
 	bool bSupportShadowMaps,
 	FTranslucentLightInjectionCollector& Collector)
@@ -2321,7 +2318,6 @@ void FDeferredShadingSceneRenderer::CollectLightForTranslucencyLightingVolumeInj
 void FDeferredShadingSceneRenderer::RenderDeferredShadowProjections(
 	FRDGBuilder& GraphBuilder,
 	const FMinimalSceneTextures& SceneTextures,
-	const FTranslucencyLightingVolumeTextures& TranslucencyLightingVolumeTextures,
 	const FLightSceneInfo* LightSceneInfo,
 	FRDGTextureRef ScreenShadowMaskTexture,
 	FRDGTextureRef ScreenShadowMaskSubPixelTexture)
diff --git a/Engine/Source/Runtime/Renderer/Private/TranslucentLighting.cpp b/Engine/Source/Runtime/Renderer/Private/TranslucentLighting.cpp
--- a/Engine/Source/Runtime/Renderer/Private/TranslucentLighting.cpp
+++ b/Engine/Source/Runtime/Renderer/Private/TranslucentLighting.cpp
@@ -1137,9 +1137,10 @@ void FTranslucencyLightingVolumeTextures::Init(FRDGBuilder& GraphBuilder, TArray
 
 		for (int32 ViewIndex = 0; ViewIndex < ViewCount; ++ViewIndex)
 		{
+			const FViewInfo& View = Views[ViewIndex];
 			for (int32 CascadeIndex = 0; CascadeIndex < TVC_MAX; ++CascadeIndex)
 			{
-				const uint32 TextureIndex = FTranslucencyLightingVolumeTextures::GetIndex(Views[ViewIndex], CascadeIndex);
+				const uint32 TextureIndex = FTranslucencyLightingVolumeTextures::GetIndex(View, CascadeIndex);
 				check(TextureIndex <= NumViewsWithTextures * TVC_MAX);
 
 				FRDGTextureRef AmbientTexture = GraphBuilder.CreateTexture(
@@ -1777,7 +1778,7 @@ void InjectTranslucencyLightingVolume(
 
 				InjectTranslucencyLightingVolumeBatch(GraphBuilder, View, ViewIndex,
 					BatchedLocalLights.Num(), BatchedLocalLightsRDG,
-					Collector.InjectionDataPerView[ViewIndex].bAnyBatchedLightsWithVirtualShadowMaps,
+					LightInjectionData.bAnyBatchedLightsWithVirtualShadowMaps,
 					VolumeCascadeIndex, VolumeAmbientTexture, VolumeDirectionalTexture,
 					Renderer);
 			}
@@ -2287,6 +2288,89 @@ void FilterTranslucencyLightingVolume(
 	}
 }
 
+static void SplitSimpleLightsByView(TConstArrayView<FViewInfo> Views, const FSimpleLightArray& SimpleLights, TArrayView<FSimpleLightArray> OutSimpleLightsByView)
+{
+	check(OutSimpleLightsByView.Num() == Views.Num());
+
+	for (int32 LightIndex = 0; LightIndex < SimpleLights.InstanceData.Num(); ++LightIndex)
+	{
+		for (int32 ViewIndex = 0; ViewIndex < Views.Num(); ++ViewIndex)
+		{
+			FSimpleLightPerViewEntry PerViewEntry = SimpleLights.GetViewDependentData(LightIndex, ViewIndex, Views.Num());
+			OutSimpleLightsByView[ViewIndex].InstanceData.Add(SimpleLights.InstanceData[LightIndex]);
+			OutSimpleLightsByView[ViewIndex].PerViewData.Add(PerViewEntry);
+		}
+	}
+}
+
+void FDeferredShadingSceneRenderer::RenderTranslucencyLightingVolume(FRDGBuilder& GraphBuilder, FTranslucencyLightingVolumeTextures& Textures, const FSortedLightSetSceneInfo& SortedLightSet)
+{
+	if (!GUseTranslucentLightingVolumes || !ViewFamily.EngineShowFlags.DirectLighting || !RHISupportsVolumeTextureRendering(ViewFamily.GetShaderPlatform()))
+	{
+		return;
+	}
+
+	RDG_EVENT_SCOPE(GraphBuilder, "TranslucencyLightingVolume");
+	SCOPE_CYCLE_COUNTER(STAT_TranslucentInjectTime);
+
+	const TArray<FSortedLightSceneInfo, SceneRenderingAllocator>& SortedLights = SortedLightSet.SortedLights;
+	const FSimpleLightArray& SimpleLights = SortedLightSet.SimpleLights;
+	const int32 SimpleLightsEnd = SortedLightSet.SimpleLightsEnd;
+
+	if (SimpleLights.InstanceData.Num() > 0)
+	{
+		auto& SimpleLightsByView = *GraphBuilder.AllocObject<TArray<FSimpleLightArray, SceneRenderingAllocator>>();
+		SimpleLightsByView.SetNum(Views.Num());
+
+		SplitSimpleLightsByView(Views, SimpleLights, SimpleLightsByView);
+
+		for (int32 ViewIndex = 0; ViewIndex < Views.Num(); ++ViewIndex)
+		{
+			FSimpleLightArray& SimpleLightArray = SimpleLightsByView[ViewIndex];
+
+			if (SimpleLightArray.InstanceData.Num() > 0)
+			{
+				const FViewInfo& View = Views[ViewIndex];
+				RDG_GPU_MASK_SCOPE(GraphBuilder, View.GPUMask);
+				RDG_EVENT_SCOPE(GraphBuilder, "InjectSimpleLightsTranslucentLighting");
+				InjectSimpleTranslucencyLightingVolumeArray(GraphBuilder, View, ViewIndex, Views.Num(), Textures, SimpleLightArray);
+			}
+		}
+	}
+
+	// Shadowed and light function lights
+	{
+		FTranslucentLightInjectionCollector Collector(GraphBuilder, Views, bAreLightsInLightGrid);
+
+		// Collect all the light injection data
+		for (int32 LightIndex = SimpleLightsEnd; LightIndex < SortedLights.Num(); LightIndex++)
+		{
+			const FSortedLightSceneInfo& SortedLightInfo = SortedLights[LightIndex];
+			const FLightSceneInfo& LightSceneInfo = *SortedLightInfo.LightSceneInfo;
+			const FLightSceneProxy& LightSceneProxy = *LightSceneInfo.Proxy;
+
+			const bool bDrawShadows = SortedLightInfo.SortKey.Fields.bShadowed;
+			const FLightOcclusionType OcclusionType = GetLightOcclusionType(LightSceneProxy, ViewFamily);
+			const bool bSupportShadowMaps = bDrawShadows && OcclusionType == FLightOcclusionType::Shadowmap;
+
+			// Collect all the light injection data
+			CollectLightForTranslucencyLightingVolumeInjection(&LightSceneInfo, bSupportShadowMaps, Collector);
+		}
+
+		// Run light injection for each view
+		for (int32 ViewIndex = 0; ViewIndex < Views.Num(); ViewIndex++)
+		{
+			FViewInfo& View = Views[ViewIndex];
+			RDG_GPU_MASK_SCOPE(GraphBuilder, View.GPUMask);
+			RDG_EVENT_SCOPE(GraphBuilder, "InjectTranslucencyLightingVolume(View=%d)", ViewIndex);
+			InjectTranslucencyLightingVolume(GraphBuilder, View, ViewIndex, Scene, *this, Collector, Textures);
+		}
+	}
+
+	InjectTranslucencyLightingVolumeAmbientCubemap(GraphBuilder, Views, Textures);
+	FilterTranslucencyLightingVolume(GraphBuilder, Views, Textures);
+}
+
 class FTranslucentLightingMaterialPSOCollector : public IPSOCollector
 {
 public:
diff --git a/Engine/Source/Runtime/Renderer/Private/TranslucentLighting.h b/Engine/Source/Runtime/Renderer/Private/TranslucentLighting.h
--- a/Engine/Source/Runtime/Renderer/Private/TranslucentLighting.h
+++ b/Engine/Source/Runtime/Renderer/Private/TranslucentLighting.h
@@ -108,6 +108,7 @@ public:
 		const FLightSceneInfo& LightSceneInfo,
 		const FProjectedShadowInfo* InProjectedShadowInfo = nullptr);
 
+private:
 	bool bCollectorSupportsBatching = false;
 };
 
@@ -118,30 +119,3 @@ bool IsTranslucencyLightingVolumeUsingVoxelMarkingSupported();
 bool IsTranslucencyLightingVolumeUsingVoxelMarking();
 
 bool IsTranslucencyLightingVolumeUsingBlueNoise();
-
-void InjectTranslucencyLightingVolume(
-	FRDGBuilder& GraphBuilder,
-	const FViewInfo& View,
-	const uint32 ViewIndex,
-	const FScene* Scene,
-	const FSceneRenderer& Renderer,
-	const FTranslucentLightInjectionCollector& Collector,
-	FTranslucencyLightingVolumeTextures& Textures);
-
-void InjectSimpleTranslucencyLightingVolumeArray(
-	FRDGBuilder& GraphBuilder,
-	const FViewInfo& View,
-	const uint32 ViewIndex,
-	const uint32 ViewCount,
-	const FTranslucencyLightingVolumeTextures& Textures,
-	const FSimpleLightArray& SimpleLights);
-
-void InjectTranslucencyLightingVolumeAmbientCubemap(
-	FRDGBuilder& GraphBuilder,
-	const TArrayView<const FViewInfo> Views,
-	const FTranslucencyLightingVolumeTextures& Textures);
-
-void FilterTranslucencyLightingVolume(
-	FRDGBuilder& GraphBuilder,
-	const TArrayView<const FViewInfo> Views,
-	FTranslucencyLightingVolumeTextures& Textures);
