[SingleLayerWater] Refraction Distance/Depth/Fresnel Culling. This CL introduces the required machinery for skipping shading of the scene behind SLW. This is achieved by first generating a mask texture directly after the water depth prepass and then writing near plane depth for all water pixels where we want to skip refraction. The water depth prepass runs before the base pass, so not only lighting, but also material evaluation should get skipped with this approach. Then during the water base pass, the mask texture is sampled and refractions/translucency is faded out. The way the mask is generated is still fairly experimental, but it should be easy to iterate: There's a single global shader which computes the mask based on water depth and scene depth. Currently it supports distance, depth and fresnel based fading.
This CL also makes some changes to how the SceneDepth/ColorWithoutWater textures are created/copied. Since the refraction culling writes near plane depth to the main depth buffer, it is unusable as SceneDepthWithoutWater, so the depth copy needs to be done before the water prepass while the color copy continues to be done immediately before the SLW base pass. A small optimization that was done here was to take advantage of the water tile list produced in the prepass. This list allows copying only relevant parts of the SceneColorTexture. In fact, even this is probably too conservative, but then we'd need to build yet another list of tiles for water pixels where we want valid refraction data.

diff --git a/Engine/Shaders/Private/BasePassPixelShader.usf b/Engine/Shaders/Private/BasePassPixelShader.usf
--- a/Engine/Shaders/Private/BasePassPixelShader.usf
+++ b/Engine/Shaders/Private/BasePassPixelShader.usf
@@ -1584,6 +1584,7 @@ void FPixelShaderInOut_MainPS(
 			SingleLayerWater.SceneColorWithoutSingleLayerWaterTexture, SingleLayerWaterSceneColorSampler,
 			SingleLayerWater.SceneWithoutSingleLayerWaterMinMaxUV.xy,
 			SingleLayerWater.SceneWithoutSingleLayerWaterMinMaxUV.zw,
+			SingleLayerWater.RefractionMaskTexture,
 			Specular, SingleLayerWater.DistortionParams,
 			SunIlluminance, WaterDiffuseIndirectIlluminance, EnvBrdf,
 			CameraIsUnderWater, WaterVisibility, EyeIndex,
diff --git a/Engine/Shaders/Private/SingleLayerWaterCommon.ush b/Engine/Shaders/Private/SingleLayerWaterCommon.ush
--- a/Engine/Shaders/Private/SingleLayerWaterCommon.ush
+++ b/Engine/Shaders/Private/SingleLayerWaterCommon.ush
@@ -61,4 +61,31 @@ float WaterSampleSceneDepthWithoutWater(Texture2D SceneDepthWithoutWaterTexture,
 	// A linear sampler should be passed in here if downsampling was used.
 	return SceneDepthWithoutWaterTexture.SampleLevel(SceneDepthWithoutWaterSampler, UV, 0).x;
 #endif // METAL_PROFILE
-}
\ No newline at end of file
+}
+
+#define SLW_MASK_IS_WATER_BITMASK (128.0f / 255.0f) // Unorm8 mask with 8th bit set
+#define SLW_MASK_REFRACTION_BITMASK (127.0f / 255.0f) // Unorm8 mask with the first 7 bits set
+
+float PackWaterRefractionMask(bool bIsWater, float RefractionMask)
+{
+	if (bIsWater)
+	{
+		return saturate(RefractionMask) * SLW_MASK_REFRACTION_BITMASK + SLW_MASK_IS_WATER_BITMASK; // Encode refraction mask in first 7 bits and water mask in 8th bit, assuming a R8 render target
+	}
+	return 0.0f;
+}
+
+void UnpackWaterRefractionMask(float PackedRefractionMask, out bool bOutIsWater, out float OutRefractionMask)
+{
+	bOutIsWater = PackedRefractionMask >= SLW_MASK_IS_WATER_BITMASK;
+	OutRefractionMask = bOutIsWater ? (PackedRefractionMask - SLW_MASK_IS_WATER_BITMASK) / SLW_MASK_REFRACTION_BITMASK : 0.0f;
+}
+
+void UnpackWaterRefractionMask(float4 PackedRefractionMask, out bool4 bOutIsWater, out float4 OutRefractionMask)
+{
+	bOutIsWater = PackedRefractionMask >= SLW_MASK_IS_WATER_BITMASK;
+	OutRefractionMask = select(bOutIsWater, (PackedRefractionMask - SLW_MASK_IS_WATER_BITMASK) / SLW_MASK_REFRACTION_BITMASK, 0.0f);
+}
+
+#undef SLW_MASK_IS_WATER_BITMASK
+#undef SLW_MASK_REFRACTION_BITMASK
diff --git a/Engine/Shaders/Private/SingleLayerWaterComposite.usf b/Engine/Shaders/Private/SingleLayerWaterComposite.usf
--- a/Engine/Shaders/Private/SingleLayerWaterComposite.usf
+++ b/Engine/Shaders/Private/SingleLayerWaterComposite.usf
@@ -662,16 +662,31 @@ SamplerState      SceneDepthCopyDownsampleSampler;
 float2 SVPositionToSourceTextureUV;
 
 void WaterRefractionCopyPS(
-	in float4 SVPosition : SV_Position,
-	out float OutSceneDepth : SV_Target0
-#if DOWNSAMPLE_COLOR
+	in float4 SVPosition : SV_Position
+#if COPY_DEPTH
+	, out float OutSceneDepth : SV_Target0
+#elif COPY_COLOR
+	, out float4 OutSceneColor : SV_Target0
+#endif
+#if COPY_DEPTH && COPY_COLOR
 	, out float4 OutSceneColor : SV_Target1
 #endif
 )
 {
 	float2 BufferUV = SVPosition.xy * SVPositionToSourceTextureUV;
 
-#if DOWNSAMPLE_COLOR
+#if COPY_DEPTH
+	#if DOWNSAMPLE_REFRACTION
+		// Pick furthermost depth to minimize downsampling artifacts (UE4 is using reverse-z)
+		float4 Depth4 = SceneDepthCopyDownsampleTexture.Gather(SceneDepthCopyDownsampleSampler, BufferUV);
+		float DeviceDepthZ = min(min(min(Depth4.x, Depth4.y), Depth4.z), Depth4.w);
+	#else
+		float DeviceDepthZ = Texture2DSampleLevel(SceneDepthCopyDownsampleTexture, SceneDepthCopyDownsampleSampler, BufferUV, 0).x;
+	#endif
+	OutSceneDepth = DeviceDepthZ;
+#endif
+
+#if COPY_COLOR
 	#if DOWNSAMPLE_REFRACTION || !SUPPORTS_INDEPENDENT_SAMPLERS
 		// Downsample and average all 2x2 quads
 		OutSceneColor.xyz = Texture2DSampleLevel(SceneColorCopyDownsampleTexture, SceneColorCopyDownsampleSampler, BufferUV, 0).xyz;
@@ -682,20 +697,5 @@ void WaterRefractionCopyPS(
 
 	OutSceneColor.w = 1.0f;
 #endif
-
-	float DeviceDepthZ = 0.0f;
-#if DOWNSAMPLE_REFRACTION
-	// Pick furthermost depth to minimize downsampling artifacts (UE4 is using reverse-z)
-	float4 Depth4 = SceneDepthCopyDownsampleTexture.Gather(SceneDepthCopyDownsampleSampler, BufferUV);
-	DeviceDepthZ = min(min(min(Depth4.x, Depth4.y), Depth4.z), Depth4.w);
-#else
-	DeviceDepthZ = Texture2DSampleLevel(SceneDepthCopyDownsampleTexture, SceneDepthCopyDownsampleSampler, BufferUV, 0).x;
-#endif
-	OutSceneDepth = DeviceDepthZ;
-
-#if METAL_PROFILE
-	// Clamp on iOS to avoid #inf
-    OutSceneDepth = min(OutSceneDepth, 65500.0f);
-#endif
 }
 
diff --git a/Engine/Shaders/Private/SingleLayerWaterRefractionCulling.usf b/Engine/Shaders/Private/SingleLayerWaterRefractionCulling.usf
new file mode 100644
--- /dev/null
+++ b/Engine/Shaders/Private/SingleLayerWaterRefractionCulling.usf
@@ -0,0 +1,95 @@
+// Copyright Epic Games, Inc. All Rights Reserved.
+
+#pragma once
+#include "Common.ush"
+#include "SingleLayerWaterCommon.ush"
+
+#ifdef SingleLayerWaterRefractionMaskPS
+
+#include "BRDF.ush"
+
+Texture2D<float> SceneDepthTexture;
+Texture2D<float> WaterDepthTexture;
+float DistanceCullingRangeBegin;
+float DistanceCullingRangeEnd;
+float FresnelCullingRangeBegin;
+float FresnelCullingRangeEnd;
+float DepthCullingRangeBegin;
+float DepthCullingRangeEnd;
+
+void SingleLayerWaterRefractionMaskPS(
+	in float4 SvPosition : SV_Position,
+	out float OutRefractionMask : SV_Target0
+)
+{
+	const float WaterDeviceZ = WaterDepthTexture.Load(int3(SvPosition.xy, 0)).x;
+	const float3 WaterTranslatedWorldPos = SvPositionToTranslatedWorld(float4(SvPosition.xy, WaterDeviceZ, 1.0f));
+	const float3 CameraTranslatedWorldPos = GetTranslatedWorldCameraPosFromView(SvPosition.xy);
+	const float3 CameraToWater = WaterTranslatedWorldPos - CameraTranslatedWorldPos;
+	
+	float ResultMask = 1.0f;
+
+	if (DistanceCullingRangeEnd > 0.0f)
+	{
+		const float Distance2 = length2(CameraToWater);
+		const float Distance = Distance2 > 1e-5f ? sqrt(Distance2) : 0.0f;
+		ResultMask *= 1.0f - smoothstep(DistanceCullingRangeBegin, DistanceCullingRangeEnd, Distance);
+	}
+
+	if (FresnelCullingRangeEnd > 0.0f)
+	{
+		const float3 WorldPosDDX = ddx_fine(WaterTranslatedWorldPos);
+		const float3 WorldPosDDY = ddy_fine(WaterTranslatedWorldPos);
+		const float3 WorldNormal = normalize(cross(WorldPosDDX, WorldPosDDY));
+
+		const float3 V = normalize(-CameraToWater);
+		const float3 L = reflect(-V, WorldNormal);
+		const float3 H = normalize(V + L);
+		const float Fresnel = F_Schlick(0.04f, abs(dot(V, H))).g;
+
+		ResultMask *= 1.0f - smoothstep(FresnelCullingRangeEnd, FresnelCullingRangeBegin, Fresnel);
+	}
+
+	if (DepthCullingRangeEnd > 0.0f)
+	{
+		const float LinearWaterDepth = ConvertFromDeviceZ(WaterDeviceZ);
+		const float LinearSceneDepth = ConvertFromDeviceZ(SceneDepthTexture.Load(int3(SvPosition.xy, 0)).x);
+		const float DepthDiff = LinearSceneDepth - LinearWaterDepth;
+
+		ResultMask *= 1.0f - smoothstep(DepthCullingRangeBegin, DepthCullingRangeEnd, DepthDiff);
+	}
+	
+	OutRefractionMask = PackWaterRefractionMask(true /*bIsWater*/, ResultMask);
+}
+
+#endif //SingleLayerWaterRefractionMaskPS
+
+#ifdef SingleLayerWaterRefractionCullingPS
+
+Texture2D<float> WaterRefractionCullingTexture;
+
+void SingleLayerWaterRefractionCullingPS(
+	in float4 SvPosition : SV_Position,
+	out float OutDepth : SV_Depth
+#if SHADING_MASK_EXPORT
+	, out uint OutShadingMask : SV_Target0
+#endif
+)
+{
+	const float PackedRefractionMask = WaterRefractionCullingTexture.Load(int3(SvPosition.xy, 0)).x;
+	bool bIsWater = false;
+	float RefractionMask = 0.0f;
+	UnpackWaterRefractionMask(PackedRefractionMask, bIsWater, RefractionMask);
+	if (!bIsWater || RefractionMask > 0.0f)
+	{
+		discard;
+	}
+
+	OutDepth = NearDepthValue;
+
+#if SHADING_MASK_EXPORT
+	OutShadingMask = 0;
+#endif
+}
+
+#endif // SingleLayerWaterRefractionCullingPS
\ No newline at end of file
diff --git a/Engine/Shaders/Private/SingleLayerWaterShading.ush b/Engine/Shaders/Private/SingleLayerWaterShading.ush
--- a/Engine/Shaders/Private/SingleLayerWaterShading.ush
+++ b/Engine/Shaders/Private/SingleLayerWaterShading.ush
@@ -41,6 +41,7 @@ WaterVolumeLightingOutput EvaluateWaterVolumeLighting(
 	Texture2D SceneColorWithoutSingleLayerWaterTexture, SamplerState SceneColorWithoutSingleLayerWaterSampler,
 	float2 SceneWithoutSingleLayerWaterMinUV,
 	float2 SceneWithoutSingleLayerWaterMaxUV,
+	Texture2D RefractionMaskTexture,
 #endif
 	float Specular, const float4 DistortionParams,
 	float3 SunIlluminance, float3 AmbiantIlluminance, float3 EnvBrdf,
@@ -101,7 +102,44 @@ WaterVolumeLightingOutput EvaluateWaterVolumeLighting(
 		// Also reduce refraction according to water thickness
 		BufferUVDistortion *= saturate((PixelSceneDepth - WaterDepth) * 1.0 / 30.0f);// *(1 - saturate(WaterDepth * 1.0 / 2000.0f));
 
-		const float2 DistortedUV = clamp(ViewportUV + BufferUVDistortion, SceneWithoutSingleLayerWaterMinUV, SceneWithoutSingleLayerWaterMaxUV);
+		float2 DistortedUV = clamp(ViewportUV + BufferUVDistortion, SceneWithoutSingleLayerWaterMinUV, SceneWithoutSingleLayerWaterMaxUV);
+
+		float RefractionMask = 1.0f;
+#if SINGLE_LAYER_WATER_SHADING_QUALITY == SINGLE_LAYER_WATER_SHADING_QUALITY_FULL
+		{
+			// Gather the four refraction mask samples at the distorted UV to see if they are valid for refraction.
+			const float4 PackedRefractionMask4 = RefractionMaskTexture.Gather(SceneColorWithoutSingleLayerWaterSampler, DistortedUV);
+			bool4 bIsWater4 = false;
+			float4 RefractionMask4 = 0.0f;
+			UnpackWaterRefractionMask(PackedRefractionMask4, bIsWater4, RefractionMask4);
+			
+			const float2 PixelCoord = DistortedUV * ResolvedView.BufferSizeAndInvSize.xy;
+			const float2 BilinearInterp = frac(PixelCoord - 0.5f);
+			const float4 BilinearWeights = float4(
+				(1.0f - BilinearInterp.x) * BilinearInterp.y,
+				BilinearInterp.x * BilinearInterp.y,
+				BilinearInterp.x * (1.0f - BilinearInterp.y),
+				(1.0f - BilinearInterp.x) * (1.0f - BilinearInterp.y));
+
+			const float4 IsWater4 = select(bIsWater4, 1.0f, 0.0f);
+			const float IsWater = dot(IsWater4, BilinearWeights);
+			RefractionMask = dot(RefractionMask4, BilinearWeights);
+			DistortedUV = lerp(ViewportUV, DistortedUV, RefractionMask);
+			
+			// Some samples don't belong to water, so we skip the distortion. However, we must still check if there is valid refraction data at the viewport UV pixel.
+			// IsWater is bilenarly interpolated, so even if a non-water pixel is in the 2x2 footprint, if the UV does not or only barely "touch" it, it shouldn't cause
+			// it to take this branch.
+			if (IsWater <= 0.99f)
+			{
+				float PackedRefractionMaskLocal = RefractionMaskTexture.SampleLevel(SceneColorWithoutSingleLayerWaterSampler, ViewportUV, 0.0f).x;
+				bool bIsWaterLocal = false;
+				float RefractionMaskLocal = 0.0f;
+				UnpackWaterRefractionMask(PackedRefractionMaskLocal, bIsWaterLocal, RefractionMaskLocal);
+				RefractionMask = RefractionMaskLocal;
+				DistortedUV = ViewportUV;
+			}
+		}
+#endif
 
 		// Gather the four depth samples that would be used for bilinear interpolation. We need to know if any of these are invalid (above the water). If so, we have to discard this distorted sample and fall back. Otherwise we end up sampling foreground color from the bilinear color sample.
 		// This is only guaranteed to be artifact free if the input depth/color resolution is the same as the output. If the input is downsampled, this will still do the right thing but might still sample leaking color
@@ -149,6 +187,7 @@ WaterVolumeLightingOutput EvaluateWaterVolumeLighting(
 		//Sample the optional Material Input ColorScaleBehindWater and fade it out at shorelines to avoid hard edge intersections
 		float3 ColorScaleBehindWater = lerp(1.0f, max(0.0f, DFDemote(GetSingleLayerWaterMaterialOutput3(MaterialParameters))), saturate(BehindWaterDeltaDepth * 0.02f));
 		const float3 ExtinctionCoeff = ScatteringCoeff + AbsorptionCoeff;
+		const float3 Albedo = ScatteringCoeff / max(ExtinctionCoeff, 1e-7f);
 #endif
 		
 		// Max to avoid division by 0 with the analytical integral below.
@@ -183,6 +222,7 @@ WaterVolumeLightingOutput EvaluateWaterVolumeLighting(
 		// We use the pixel SvPosition instead of the scene one pre refraction/distortion to avoid those extra ALUs.
 		float3 BehindWaterSceneLuminance = SceneColorWithoutSingleLayerWaterTexture.SampleLevel(SceneColorWithoutSingleLayerWaterSampler, ViewportUV, 0).rgb;
 		BehindWaterSceneLuminance = MeanTransmittanceToLightSources * ResolvedView.OneOverPreExposure * BehindWaterSceneLuminance;
+		BehindWaterSceneLuminance = lerp(Albedo * ResolvedView.OneOverPreExposure, BehindWaterSceneLuminance, RefractionMask);
 #endif
 
 		float3 SunScattLuminance = DirLightPhaseValue * SunIlluminance;
diff --git a/Engine/Source/Runtime/Renderer/Private/DeferredShadingRenderer.cpp b/Engine/Source/Runtime/Renderer/Private/DeferredShadingRenderer.cpp
--- a/Engine/Source/Runtime/Renderer/Private/DeferredShadingRenderer.cpp
+++ b/Engine/Source/Runtime/Renderer/Private/DeferredShadingRenderer.cpp
@@ -2164,7 +2164,7 @@ void FDeferredShadingSceneRenderer::Render(FRDGBuilder& GraphBuilder, const FSce
 				const FSingleLayerWaterPrePassResult* SingleLayerWaterPrePassResult = nullptr;
 				if (ShouldRenderSingleLayerWaterDepthPrepass(CustomRenderPassViews))
 				{
-					SingleLayerWaterPrePassResult = RenderSingleLayerWaterDepthPrepass(GraphBuilder, CustomRenderPassViews, SceneTextures);
+					SingleLayerWaterPrePassResult = RenderSingleLayerWaterDepthPrepass(GraphBuilder, CustomRenderPassViews, SceneTextures, ESingleLayerWaterPrepassLocation::BeforeBasePass, NaniteRasterResults);
 				}
 
 				const FSceneCaptureCustomRenderPassUserData& SceneCaptureUserData = FSceneCaptureCustomRenderPassUserData::Get(CustomRenderPass);
@@ -2337,7 +2337,7 @@ void FDeferredShadingSceneRenderer::Render(FRDGBuilder& GraphBuilder, const FSce
 		const FSingleLayerWaterPrePassResult* SingleLayerWaterPrePassResult = nullptr;
 		if (ShouldRenderSingleLayerWaterDepthPrepass(Views))
 		{
-			SingleLayerWaterPrePassResult = RenderSingleLayerWaterDepthPrepass(GraphBuilder, Views, SceneTextures);
+			SingleLayerWaterPrePassResult = RenderSingleLayerWaterDepthPrepass(GraphBuilder, Views, SceneTextures, ESingleLayerWaterPrepassLocation::BeforeBasePass, NaniteRasterResults);
 		}
 
 		FroxelRenderer = RenderOcclusionLambda();
@@ -2387,6 +2387,16 @@ void FDeferredShadingSceneRenderer::Render(FRDGBuilder& GraphBuilder, const FSce
 			ViewExtension->PreRenderBasePass_RenderThread(GraphBuilder, ShouldRenderPrePass() /*bDepthBufferIsPopulated*/);
 		}
 
+		// Single layer water depth prepass. Needs to run before VSM page allocation. If there's a full depth prepass, it can run before the base pass, otherwise after.
+		// Running before the base pass allows for some optimizations to save work in the base pass and lighting stages.
+		const FSingleLayerWaterPrePassResult* SingleLayerWaterPrePassResult = nullptr;
+		const ESingleLayerWaterPrepassLocation SingleLayerWaterPrepassLocation = GetSingleLayerWaterDepthPrepassLocation(bIsEarlyDepthComplete);
+		const bool bShouldRenderSingleLayerWaterDepthPrepass = !bHasRayTracedOverlay && ShouldRenderSingleLayerWaterDepthPrepass(Views);
+		if (bShouldRenderSingleLayerWaterDepthPrepass && SingleLayerWaterPrepassLocation == ESingleLayerWaterPrepassLocation::BeforeBasePass)
+		{
+			SingleLayerWaterPrePassResult = RenderSingleLayerWaterDepthPrepass(GraphBuilder, Views, SceneTextures, SingleLayerWaterPrepassLocation, NaniteRasterResults);
+		}
+
 		{
 			SCOPE_CYCLE_COUNTER(STAT_WaitGatherAndSortLightsTask);
 			GatherAndSortLightsTask.Wait();
@@ -2718,12 +2728,9 @@ void FDeferredShadingSceneRenderer::Render(FRDGBuilder& GraphBuilder, const FSce
 		FRDGTextureRef LightingChannelsTexture = CopyStencilToLightingChannelTexture(GraphBuilder, SceneTextures.Stencil, NaniteShadingMask);
 
 		// Single layer water depth prepass. Needs to run before VSM page allocation.
-		const FSingleLayerWaterPrePassResult* SingleLayerWaterPrePassResult = nullptr;
-
-		const bool bShouldRenderSingleLayerWaterDepthPrepass = !bHasRayTracedOverlay && ShouldRenderSingleLayerWaterDepthPrepass(Views);
-		if (bShouldRenderSingleLayerWaterDepthPrepass)
+		if (bShouldRenderSingleLayerWaterDepthPrepass && SingleLayerWaterPrepassLocation == ESingleLayerWaterPrepassLocation::AfterBasePass)
 		{
-			SingleLayerWaterPrePassResult = RenderSingleLayerWaterDepthPrepass(GraphBuilder, Views, SceneTextures);
+			SingleLayerWaterPrePassResult = RenderSingleLayerWaterDepthPrepass(GraphBuilder, Views, SceneTextures, SingleLayerWaterPrepassLocation, NaniteRasterResults);
 		}
 
 		FAsyncLumenIndirectLightingOutputs AsyncLumenIndirectLightingOutputs;
diff --git a/Engine/Source/Runtime/Renderer/Private/DeferredShadingRenderer.h b/Engine/Source/Runtime/Renderer/Private/DeferredShadingRenderer.h
--- a/Engine/Source/Runtime/Renderer/Private/DeferredShadingRenderer.h
+++ b/Engine/Source/Runtime/Renderer/Private/DeferredShadingRenderer.h
@@ -25,6 +25,7 @@
 #include "LightFunctionAtlas.h"
 
 enum class ERayTracingPrimaryRaysFlag : uint32;
+enum class ESingleLayerWaterPrepassLocation : uint8;
 
 class FLumenCardUpdateContext;
 class FSceneTextureParameters;
@@ -393,7 +394,9 @@ public:
 	FSingleLayerWaterPrePassResult* RenderSingleLayerWaterDepthPrepass(
 		FRDGBuilder& GraphBuilder,
 		TArrayView<FViewInfo> InViews,
-		const FSceneTextures& SceneTextures);
+		const FSceneTextures& SceneTextures,
+		ESingleLayerWaterPrepassLocation Location, 
+		TConstArrayView<Nanite::FRasterResults> NaniteRasterResults);
 
 	void RenderSingleLayerWater(
 		FRDGBuilder& GraphBuilder,
diff --git a/Engine/Source/Runtime/Renderer/Private/SingleLayerWaterRendering.cpp b/Engine/Source/Runtime/Renderer/Private/SingleLayerWaterRendering.cpp
--- a/Engine/Source/Runtime/Renderer/Private/SingleLayerWaterRendering.cpp
+++ b/Engine/Source/Runtime/Renderer/Private/SingleLayerWaterRendering.cpp
@@ -121,10 +121,8 @@ static TAutoConsoleVariable<int32> CVarWaterSingleLayerVSMFiltering(
 //
 // Misc
 
-int32 GSingleLayerWaterRefractionDownsampleFactor = 1;
-static FAutoConsoleVariableRef CVarWaterSingleLayerRefractionDownsampleFactor(
-	TEXT("r.Water.SingleLayer.RefractionDownsampleFactor"),
-	GSingleLayerWaterRefractionDownsampleFactor,
+static TAutoConsoleVariable<int32> CVarWaterSingleLayerRefractionDownsampleFactor(
+	TEXT("r.Water.SingleLayer.RefractionDownsampleFactor"), 1,
 	TEXT("Resolution divider for the water refraction buffer."),
 	ECVF_Scalability | ECVF_RenderThreadSafe);
 
@@ -148,12 +146,58 @@ static TAutoConsoleVariable<int32> CVarWaterSingleLayerForceVelocity(
 	TEXT("Whether to always output velocity, even if the velocity pass is not set to \"Write during base pass\"."),
 	ECVF_ReadOnly | ECVF_RenderThreadSafe);
 
+static TAutoConsoleVariable<int32> CVarSingleLayerWaterRefractionCulling(
+	TEXT("r.Water.SingleLayer.Refraction.Culling"), 0,
+	TEXT("Enables refraction culling on water. This allows the renderer to skip rendering portions of the scene behind water which are unlikely to be visible through water."),
+	ECVF_RenderThreadSafe);
+
+static TAutoConsoleVariable<float> CVarSingleLayerWaterRefractionDistanceCulling(
+	TEXT("r.Water.SingleLayer.Refraction.DistanceCulling"), -1.0f,
+	TEXT("Distance at which to cull refractions."),
+	ECVF_RenderThreadSafe);
+
+static TAutoConsoleVariable<float> CVarSingleLayerWaterRefractionDistanceCullingFadeRange(
+	TEXT("r.Water.SingleLayer.Refraction.DistanceCullingFadeRange"), 400.0f,
+	TEXT("Range over which to fade out refractions."),
+	ECVF_RenderThreadSafe);
+
+static TAutoConsoleVariable<float> CVarSingleLayerWaterRefractionFresnelCulling(
+	TEXT("r.Water.SingleLayer.Refraction.FresnelCulling"), -1.0f,
+	TEXT("Fresnel value below which to cull refractions."),
+	ECVF_RenderThreadSafe);
+
+static TAutoConsoleVariable<float> CVarSingleLayerWaterRefractionFresnelCullingFadeRange(
+	TEXT("r.Water.SingleLayer.Refraction.FresnelCullingFadeRange"), 0.2f,
+	TEXT("Fresnel range over which to fade out refractions."),
+	ECVF_RenderThreadSafe);
+
+static TAutoConsoleVariable<float> CVarSingleLayerWaterRefractionDepthCulling(
+	TEXT("r.Water.SingleLayer.Refraction.DepthCulling"), -1.0f,
+	TEXT("Depth below which to cull refractions."),
+	ECVF_RenderThreadSafe);
+
+static TAutoConsoleVariable<float> CVarSingleLayerWaterRefractionDepthCullingFadeRange(
+	TEXT("r.Water.SingleLayer.Refraction.DepthCullingFadeRange"), 100.0f,
+	TEXT("Depth range over which to fade out refractions."),
+	ECVF_RenderThreadSafe);
+
+static TAutoConsoleVariable<int32> CVarWaterSingleLayerTiledSceneColorCopy(
+	TEXT("r.Water.SingleLayer.TiledSceneColorCopy"), 1,
+	TEXT("Use indirect draws and a list of water pixel tiles to copy only relevant parts of the SceneColor texture for refraction behind the water surface."),
+	ECVF_RenderThreadSafe);
+
 static int32 GetSingleLayerWaterReflectionTechnique()
 {
 	const int32 Value = CVarWaterSingleLayerReflection.GetValueOnRenderThread();
 	return FMath::Clamp(Value, 0, ESingleLayerWaterReflections::MaxValue);
 }
 
+static int32 GetSingleLayerWaterRefractionDownsampleFactor()
+{
+	const int32 RefractionDownsampleFactor = FMath::Clamp(CVarWaterSingleLayerRefractionDownsampleFactor.GetValueOnRenderThread(), 1, 8);
+	return RefractionDownsampleFactor;
+}
+
 static EGBufferLayout GetSingleLayerWaterGBufferLayout(bool bIsGameThread = false)
 {
 	if (!bIsGameThread)
@@ -221,6 +265,15 @@ bool ShouldRenderSingleLayerWaterDepthPrepass(TArrayView<const FViewInfo> Views)
 	return bPrepassEnabled && bShouldRenderWater;
 }
 
+ESingleLayerWaterPrepassLocation GetSingleLayerWaterDepthPrepassLocation(bool bFullDepthPrepass)
+{
+	if (bFullDepthPrepass)
+	{
+		return ESingleLayerWaterPrepassLocation::BeforeBasePass;
+	}
+	return ESingleLayerWaterPrepassLocation::AfterBasePass;
+}
+
 namespace ScreenSpaceRayTracing
 {
 bool ShouldRenderScreenSpaceReflectionsWater(const FViewInfo& View)
@@ -263,7 +316,7 @@ bool ShouldRenderLumenReflectionsWater(const FViewInfo& View, bool bSkipTracingD
 
 bool ShouldUseBilinearSamplerForDepthWithoutSingleLayerWater(EPixelFormat DepthTextureFormat)
 {
-	const bool bHasDownsampling = GSingleLayerWaterRefractionDownsampleFactor > 1;
+	const bool bHasDownsampling = GetSingleLayerWaterRefractionDownsampleFactor() > 1;
 	const bool bSupportsLinearSampling = !!(GPixelFormats[DepthTextureFormat].Capabilities & EPixelFormatCapabilities::TextureSample);
 	
 	// Linear sampling is only required if the depth texture has been downsampled.
@@ -345,6 +398,82 @@ class FSingleLayerWaterCompositePS : public FGlobalShader
 
 IMPLEMENT_GLOBAL_SHADER(FSingleLayerWaterCompositePS, "/Engine/Private/SingleLayerWaterComposite.usf", "SingleLayerWaterCompositePS", SF_Pixel);
 
+class FSingleLayerWaterRefractionMaskPS : public FGlobalShader
+{
+	DECLARE_GLOBAL_SHADER(FSingleLayerWaterRefractionMaskPS);
+	SHADER_USE_PARAMETER_STRUCT(FSingleLayerWaterRefractionMaskPS, FGlobalShader)
+
+	using FPermutationDomain = TShaderPermutationDomain<>;
+
+	BEGIN_SHADER_PARAMETER_STRUCT(FParameters, )
+		SHADER_PARAMETER_STRUCT_REF(FViewUniformShaderParameters, View)
+		SHADER_PARAMETER_RDG_TEXTURE(Texture2D, SceneDepthTexture)
+		SHADER_PARAMETER_RDG_TEXTURE(Texture2D, WaterDepthTexture)
+		SHADER_PARAMETER(float, DistanceCullingRangeBegin)
+		SHADER_PARAMETER(float, DistanceCullingRangeEnd)
+		SHADER_PARAMETER(float, FresnelCullingRangeBegin)
+		SHADER_PARAMETER(float, FresnelCullingRangeEnd)
+		SHADER_PARAMETER(float, DepthCullingRangeBegin)
+		SHADER_PARAMETER(float, DepthCullingRangeEnd)
+		RENDER_TARGET_BINDING_SLOTS()
+	END_SHADER_PARAMETER_STRUCT()
+
+	static FPermutationDomain RemapPermutation(FPermutationDomain PermutationVector)
+	{
+		return PermutationVector;
+	}
+
+	static bool ShouldCompilePermutation(const FGlobalShaderPermutationParameters& Parameters)
+	{
+		return IsFeatureLevelSupported(Parameters.Platform, ERHIFeatureLevel::SM5);
+	}
+
+	static void ModifyCompilationEnvironment(const FGlobalShaderPermutationParameters& Parameters, FShaderCompilerEnvironment& OutEnvironment)
+	{
+		FGlobalShader::ModifyCompilationEnvironment(Parameters, OutEnvironment);
+		OutEnvironment.SetRenderTargetOutputFormat(0, PF_R8);
+	}
+};
+
+IMPLEMENT_GLOBAL_SHADER(FSingleLayerWaterRefractionMaskPS, "/Engine/Private/SingleLayerWaterRefractionCulling.usf", "SingleLayerWaterRefractionMaskPS", SF_Pixel);
+
+class FSingleLayerWaterRefractionCullingPS : public FGlobalShader
+{
+	DECLARE_GLOBAL_SHADER(FSingleLayerWaterRefractionCullingPS);
+	SHADER_USE_PARAMETER_STRUCT(FSingleLayerWaterRefractionCullingPS, FGlobalShader)
+
+	class FNaniteShadingMaskExport : SHADER_PERMUTATION_BOOL("SHADING_MASK_EXPORT");
+	using FPermutationDomain = TShaderPermutationDomain<FNaniteShadingMaskExport>;
+
+	BEGIN_SHADER_PARAMETER_STRUCT(FParameters, )
+		SHADER_PARAMETER_STRUCT_REF(FViewUniformShaderParameters, View)
+		SHADER_PARAMETER_RDG_TEXTURE(Texture2D, WaterRefractionCullingTexture)
+		RENDER_TARGET_BINDING_SLOTS()
+	END_SHADER_PARAMETER_STRUCT()
+
+	static FPermutationDomain RemapPermutation(FPermutationDomain PermutationVector)
+	{
+		return PermutationVector;
+	}
+
+	static bool ShouldCompilePermutation(const FGlobalShaderPermutationParameters& Parameters)
+	{
+		return IsFeatureLevelSupported(Parameters.Platform, ERHIFeatureLevel::SM5);
+	}
+
+	static void ModifyCompilationEnvironment(const FGlobalShaderPermutationParameters& Parameters, FShaderCompilerEnvironment& OutEnvironment)
+	{
+		FGlobalShader::ModifyCompilationEnvironment(Parameters, OutEnvironment);
+		const FPermutationDomain Permutation(Parameters.PermutationId);
+		if (Permutation.Get<FNaniteShadingMaskExport>())
+		{
+			OutEnvironment.SetRenderTargetOutputFormat(0, PF_R32_UINT);
+		}
+	}
+};
+
+IMPLEMENT_GLOBAL_SHADER(FSingleLayerWaterRefractionCullingPS, "/Engine/Private/SingleLayerWaterRefractionCulling.usf", "SingleLayerWaterRefractionCullingPS", SF_Pixel);
+
 class FWaterTileCategorisationMarkCS : public FGlobalShader
 {
 	DECLARE_GLOBAL_SHADER(FWaterTileCategorisationMarkCS);
@@ -487,19 +616,30 @@ class FWaterRefractionCopyPS : public FGlobalShader
 	END_SHADER_PARAMETER_STRUCT()
 
 	class FDownsampleRefraction : SHADER_PERMUTATION_BOOL("DOWNSAMPLE_REFRACTION");
-	class FDownsampleColor : SHADER_PERMUTATION_BOOL("DOWNSAMPLE_COLOR");
+	class FCopyDepth : SHADER_PERMUTATION_BOOL("COPY_DEPTH");
+	class FCopyColor : SHADER_PERMUTATION_BOOL("COPY_COLOR");
 
-	using FPermutationDomain = TShaderPermutationDomain<FDownsampleRefraction, FDownsampleColor>;
+	using FPermutationDomain = TShaderPermutationDomain<FDownsampleRefraction, FCopyDepth, FCopyColor>;
 
 	static bool ShouldCompilePermutation(const FGlobalShaderPermutationParameters& Parameters)
 	{
-		return IsFeatureLevelSupported(Parameters.Platform, ERHIFeatureLevel::SM5);
+		const FPermutationDomain Permutation(Parameters.PermutationId);
+		return IsFeatureLevelSupported(Parameters.Platform, ERHIFeatureLevel::SM5) && (Permutation.Get<FCopyDepth>() || Permutation.Get<FCopyColor>());
 	}
 
 	static void ModifyCompilationEnvironment(const FGlobalShaderPermutationParameters& Parameters, FShaderCompilerEnvironment& OutEnvironment)
 	{
 		FGlobalShader::ModifyCompilationEnvironment(Parameters, OutEnvironment);
-		OutEnvironment.SetRenderTargetOutputFormat(0, PF_R32_FLOAT);
+		const FPermutationDomain Permutation(Parameters.PermutationId);
+		const bool bCopyDepth = Permutation.Get<FCopyDepth>();
+		const bool bCopyColor = Permutation.Get<FCopyColor>();
+		const EPixelFormat DepthFormat = PF_R32_FLOAT;
+		const EPixelFormat ColorFormat = PF_FloatRGBA;
+		OutEnvironment.SetRenderTargetOutputFormat(0, bCopyDepth ? DepthFormat : ColorFormat);
+		if (bCopyDepth && bCopyColor)
+		{
+			OutEnvironment.SetRenderTargetOutputFormat(1, ColorFormat);
+		}
 	}
 };
 
@@ -685,7 +825,97 @@ static FSingleLayerWaterTileClassification ClassifyTiles(FRDGBuilder& GraphBuild
 	return Result;
 }
 
-FSingleLayerWaterPrePassResult* FDeferredShadingSceneRenderer::RenderSingleLayerWaterDepthPrepass(FRDGBuilder& GraphBuilder, TArrayView<FViewInfo> InViews, const FSceneTextures& SceneTextures)
+BEGIN_SHADER_PARAMETER_STRUCT(FWaterRefractionCopyParameters, )
+	SHADER_PARAMETER_STRUCT_INCLUDE(FWaterTileVS::FParameters, VS)
+	SHADER_PARAMETER_STRUCT_INCLUDE(FWaterRefractionCopyPS::FParameters, PS)
+	RDG_BUFFER_ACCESS(IndirectDrawParameter, ERHIAccess::IndirectArgs)
+END_SHADER_PARAMETER_STRUCT()
+
+static void AddCopySceneWithoutWaterPass_Internal(
+	FRDGBuilder& GraphBuilder,
+	const FSceneViewFamily& ViewFamily,
+	TArrayView<const FViewInfo> Views,
+	FRDGTextureRef DstDepthTexture,
+	FRDGTextureRef SrcDepthTexture,
+	FRDGTextureRef DstColorTexture,
+	FRDGTextureRef SrcColorTexture,
+	int32 RefractionDownsampleFactor,
+	const FSingleLayerWaterPrePassResult* SingleLayerWaterPrepassResult
+)
+{
+	const FIntPoint SceneTextureExtent = Views[0].GetSceneTextures().Config.Extent;
+	const bool bCopyDepth = DstDepthTexture && SrcDepthTexture;
+	const bool bCopyColor = DstColorTexture && SrcColorTexture;
+	check(bCopyDepth || bCopyColor);
+	const bool bDoTiledCopy = !bCopyDepth 
+		&& RefractionDownsampleFactor == 1 
+		&& SingleLayerWaterPrepassResult 
+		&& !SingleLayerWaterPrepassResult->ViewTileClassification.IsEmpty() 
+		&& CVarWaterSingleLayerTiledSceneColorCopy.GetValueOnRenderThread() != 0;
+
+	ERenderTargetLoadAction LoadAction = ERenderTargetLoadAction::ENoAction;
+
+	for (int32 ViewIndex = 0; ViewIndex < Views.Num(); ++ViewIndex)
+	{
+		const FViewInfo& View = Views[ViewIndex];
+
+		if (!View.ShouldRenderView())
+		{
+			continue;
+		}
+
+		RDG_GPU_MASK_SCOPE(GraphBuilder, View.GPUMask);
+		RDG_EVENT_SCOPE_CONDITIONAL(GraphBuilder, Views.Num() > 1, "View%d", ViewIndex);
+
+		FWaterRefractionCopyParameters* PassParameters = GraphBuilder.AllocParameters<FWaterRefractionCopyParameters>();
+		PassParameters->PS.SceneColorCopyDownsampleTexture = SrcColorTexture;
+		PassParameters->PS.SceneColorCopyDownsampleSampler = TStaticSamplerState<SF_Bilinear, AM_Clamp, AM_Clamp, AM_Clamp>::GetRHI();
+		PassParameters->PS.SceneDepthCopyDownsampleTexture = SrcDepthTexture;
+		PassParameters->PS.SceneDepthCopyDownsampleSampler = TStaticSamplerState<SF_Point, AM_Clamp, AM_Clamp, AM_Clamp>::GetRHI();
+		PassParameters->PS.SVPositionToSourceTextureUV = FVector2f(RefractionDownsampleFactor / float(SceneTextureExtent.X), RefractionDownsampleFactor / float(SceneTextureExtent.Y));
+		PassParameters->PS.RenderTargets[0] = FRenderTargetBinding(bCopyDepth ? DstDepthTexture : DstColorTexture, LoadAction);
+		if (bCopyDepth && bCopyColor)
+		{
+			PassParameters->PS.RenderTargets[1] = FRenderTargetBinding(DstColorTexture, LoadAction);
+		}
+
+		if (!View.Family->bMultiGPUForkAndJoin)
+		{
+			LoadAction = ERenderTargetLoadAction::ELoad;
+		}
+
+		FWaterRefractionCopyPS::FPermutationDomain PermutationVector;
+		PermutationVector.Set<FWaterRefractionCopyPS::FDownsampleRefraction>(RefractionDownsampleFactor > 1);
+		PermutationVector.Set<FWaterRefractionCopyPS::FCopyDepth>(bCopyDepth);
+		PermutationVector.Set<FWaterRefractionCopyPS::FCopyColor>(bCopyColor);
+		auto PixelShader = View.ShaderMap->GetShader<FWaterRefractionCopyPS>(PermutationVector);
+
+		// if we have a particular case of ISR where two views are laid out in side by side, we should copy both views at once
+		const bool bIsInstancedStereoSideBySide = View.bIsInstancedStereoEnabled && !View.bIsMobileMultiViewEnabled && IStereoRendering::IsStereoEyeView(View);
+		FIntRect RectToCopy = View.ViewRect;
+		if (bIsInstancedStereoSideBySide)
+		{
+			const FViewInfo* NeighboringStereoView = View.GetInstancedView();
+			if (ensure(NeighboringStereoView))
+			{
+				RectToCopy.Union(NeighboringStereoView->ViewRect);
+			}
+		}
+
+		const FIntRect RefractionViewRect = FIntRect(FIntPoint::DivideAndRoundDown(RectToCopy.Min, RefractionDownsampleFactor), FIntPoint::DivideAndRoundDown(RectToCopy.Max, RefractionDownsampleFactor));
+		if (bDoTiledCopy)
+		{
+			const FTiledReflection* TiledReflection = &SingleLayerWaterPrepassResult->ViewTileClassification[ViewIndex].TiledReflection;
+			SingleLayerWaterAddTiledFullscreenPass(GraphBuilder, View.ShaderMap, RDG_EVENT_NAME("SLW::Copy"), PixelShader, PassParameters, View.ViewUniformBuffer, RefractionViewRect, TiledReflection);
+		}
+		else
+		{
+			FPixelShaderUtils::AddFullscreenPass(GraphBuilder, View.ShaderMap, RDG_EVENT_NAME("SLW::Copy"), PixelShader, &PassParameters->PS, RefractionViewRect);
+		}
+	}
+}
+
+FSingleLayerWaterPrePassResult* FDeferredShadingSceneRenderer::RenderSingleLayerWaterDepthPrepass(FRDGBuilder& GraphBuilder, TArrayView<FViewInfo> InViews, const FSceneTextures& SceneTextures, ESingleLayerWaterPrepassLocation Location, TConstArrayView<Nanite::FRasterResults> NaniteRasterResults)
 {
 	RDG_CSV_STAT_EXCLUSIVE_SCOPE(GraphBuilder, Water);
 	SCOPED_NAMED_EVENT(FDeferredShadingSceneRenderer_RenderSingleLayerWaterDepthPrepass, FColor::Emerald);
@@ -769,6 +999,15 @@ FSingleLayerWaterPrePassResult* FDeferredShadingSceneRenderer::RenderSingleLayer
 		}
 	}
 
+	// Create SceneDepthWithoutWater texture
+	{
+		const int32 RefractionDownsampleFactor = GetSingleLayerWaterRefractionDownsampleFactor();
+		const FIntPoint RefractionResolution = FIntPoint::DivideAndRoundDown(SceneTextures.Config.Extent, RefractionDownsampleFactor);
+		Result->SceneDepthWithoutWater = GraphBuilder.CreateTexture(FRDGTextureDesc::Create2D(RefractionResolution, PF_R32_FLOAT, FClearValueBinding::DepthFar, TexCreate_ShaderResource | TexCreate_RenderTargetable), TEXT("SLW.SceneDepthWithout"));
+
+		AddCopySceneWithoutWaterPass_Internal(GraphBuilder, ViewFamily, Views, Result->SceneDepthWithoutWater, SceneTextures.Depth.Resolve, nullptr, nullptr, RefractionDownsampleFactor, nullptr);
+	}
+
 	const bool bRenderInParallel = GRHICommandList.UseParallelAlgorithms() && CVarParallelSingleLayerWaterPass.GetValueOnRenderThread() == 1;
 
 	for (int32 ViewIndex = 0; ViewIndex < InViews.Num(); ++ViewIndex)
@@ -818,7 +1057,7 @@ FSingleLayerWaterPrePassResult* FDeferredShadingSceneRenderer::RenderSingleLayer
 	AddResolveSceneDepthPass(GraphBuilder, InViews, OutDepthPrepassTexture);
 
 	// Run classification pass.
-	if (UseSingleLayerWaterIndirectDraw(ShaderPlatform) && CVarWaterSingleLayerTiledComposite.GetValueOnRenderThread())
+	if (UseSingleLayerWaterIndirectDraw(ShaderPlatform) && (CVarWaterSingleLayerTiledComposite.GetValueOnRenderThread() || CVarWaterSingleLayerTiledSceneColorCopy.GetValueOnRenderThread()))
 	{
 		Result->Froxels = Froxel::FRenderer(DoesVSMWantFroxels(ShaderPlatform), GraphBuilder, Views);
 
@@ -830,6 +1069,89 @@ FSingleLayerWaterPrePassResult* FDeferredShadingSceneRenderer::RenderSingleLayer
 		}
 	}
 
+	Result->RefractionMaskTexture = nullptr;
+	const float RefractionDistanceCulling = CVarSingleLayerWaterRefractionDistanceCulling.GetValueOnRenderThread();
+	const float RefractionFresnelCulling = CVarSingleLayerWaterRefractionFresnelCulling.GetValueOnRenderThread();
+	const float RefractionDepthCulling = CVarSingleLayerWaterRefractionDepthCulling.GetValueOnRenderThread();
+	const bool bDoRefractionCulling = CVarSingleLayerWaterRefractionCulling.GetValueOnRenderThread() != 0
+		&& (RefractionDistanceCulling > 0.0f || RefractionFresnelCulling > 0.0f || RefractionDepthCulling > 0.0f)
+		&& Location == ESingleLayerWaterPrepassLocation::BeforeBasePass
+		&& GetRendererOutput() != ERendererOutput::DepthPrepassOnly
+		&& !IsForwardShadingEnabled(ShaderPlatform);
+	
+	if (bDoRefractionCulling)
+	{
+		Result->RefractionMaskTexture = GraphBuilder.CreateTexture(FRDGTextureDesc::Create2D(SceneTextures.Config.Extent, PF_R8, FClearValueBinding::Black, TexCreate_RenderTargetable | TexCreate_ShaderResource), TEXT("SLW.RefractionMask"));
+
+		AddClearRenderTargetPass(GraphBuilder, Result->RefractionMaskTexture);
+
+		// Create refraction mask texture
+		for (int32 ViewIndex = 0; ViewIndex < InViews.Num(); ++ViewIndex)
+		{
+			FViewInfo& View = InViews[ViewIndex];
+
+			FSingleLayerWaterRefractionMaskPS::FParameters* PassParameters = GraphBuilder.AllocParameters<FSingleLayerWaterRefractionMaskPS::FParameters>();
+			PassParameters->View = View.ViewUniformBuffer;
+			PassParameters->SceneDepthTexture = Result->SceneDepthWithoutWater;
+			PassParameters->WaterDepthTexture = Result->DepthPrepassTexture.Resolve;
+			PassParameters->DistanceCullingRangeEnd = RefractionDistanceCulling;
+			PassParameters->DistanceCullingRangeBegin = PassParameters->DistanceCullingRangeEnd - FMath::Max(CVarSingleLayerWaterRefractionDistanceCullingFadeRange.GetValueOnRenderThread(), 0.01f);
+			PassParameters->FresnelCullingRangeEnd = RefractionFresnelCulling;
+			PassParameters->FresnelCullingRangeBegin = PassParameters->FresnelCullingRangeEnd + FMath::Max(CVarSingleLayerWaterRefractionFresnelCullingFadeRange.GetValueOnRenderThread(), 0.01f);
+			PassParameters->DepthCullingRangeEnd = RefractionDepthCulling;
+			PassParameters->DepthCullingRangeBegin = PassParameters->DepthCullingRangeEnd - FMath::Max(CVarSingleLayerWaterRefractionDepthCullingFadeRange.GetValueOnRenderThread(), 0.01f);
+			PassParameters->RenderTargets.DepthStencil = FDepthStencilBinding(Result->DepthPrepassTexture.Target, ERenderTargetLoadAction::ELoad, ERenderTargetLoadAction::ELoad, FExclusiveDepthStencil::DepthRead_StencilRead);
+			PassParameters->RenderTargets[0] = FRenderTargetBinding(Result->RefractionMaskTexture, ERenderTargetLoadAction::ELoad);
+
+			TShaderMapRef<FSingleLayerWaterRefractionMaskPS> PixelShader(View.ShaderMap);
+
+			FPixelShaderUtils::AddFullscreenPass(
+				GraphBuilder,
+				View.ShaderMap,
+				RDG_EVENT_NAME("SLW::RefractionMask (View: %i)", ViewIndex),
+				PixelShader,
+				PassParameters,
+				View.ViewRect,
+				nullptr /*BlendState*/,
+				nullptr /*RasterizerState*/,
+				TStaticDepthStencilState<false, CF_Always, true, CF_Equal>::GetRHI(),
+				1 /*StencilRef*/);
+		}
+
+		// Write near plane depth for all water pixels where we want to skip all shading for the scene behind them
+		for (int32 ViewIndex = 0; ViewIndex < InViews.Num(); ++ViewIndex)
+		{
+			FViewInfo& View = InViews[ViewIndex];
+
+			const bool bNaniteShadingMaskExport = NaniteRasterResults.IsValidIndex(ViewIndex) && HasBeenProduced(NaniteRasterResults[ViewIndex].ShadingMask);
+
+			FSingleLayerWaterRefractionCullingPS::FParameters* PassParameters = GraphBuilder.AllocParameters<FSingleLayerWaterRefractionCullingPS::FParameters>();
+			PassParameters->View = View.ViewUniformBuffer;
+			PassParameters->WaterRefractionCullingTexture = Result->RefractionMaskTexture;
+			PassParameters->RenderTargets.DepthStencil = FDepthStencilBinding(SceneTextures.Depth.Target, ERenderTargetLoadAction::ELoad, ERenderTargetLoadAction::ELoad, FExclusiveDepthStencil::DepthWrite_StencilNop);
+			if (bNaniteShadingMaskExport)
+			{
+				PassParameters->RenderTargets[0] = FRenderTargetBinding(NaniteRasterResults[ViewIndex].ShadingMask, ERenderTargetLoadAction::ELoad);
+			}
+
+			FSingleLayerWaterRefractionCullingPS::FPermutationDomain PermutationDomain;
+			PermutationDomain.Set<FSingleLayerWaterRefractionCullingPS::FNaniteShadingMaskExport>(bNaniteShadingMaskExport);
+			TShaderMapRef<FSingleLayerWaterRefractionCullingPS> PixelShader(View.ShaderMap, PermutationDomain);
+
+			FPixelShaderUtils::AddFullscreenPass(
+				GraphBuilder,
+				View.ShaderMap,
+				RDG_EVENT_NAME("SLW::RefractionCulling (View: %i)", ViewIndex),
+				PixelShader,
+				PassParameters,
+				View.ViewRect,
+				nullptr /*BlendState*/,
+				nullptr /*RasterizerState*/,
+				TStaticDepthStencilState<true, CF_Always>::GetRHI(),
+				1 /*StencilRef*/);
+		}
+	}
+
 	return Result;
 }
 
@@ -838,7 +1160,8 @@ static FSceneWithoutWaterTextures AddCopySceneWithoutWaterPass(
 	const FSceneViewFamily& ViewFamily, 
 	TArrayView<const FViewInfo> Views,
 	FRDGTextureRef SceneColorTexture,
-	FRDGTextureRef SceneDepthTexture)
+	FRDGTextureRef SceneDepthTexture,
+	const FSingleLayerWaterPrePassResult* SingleLayerWaterPrePassResult)
 {
 	RDG_EVENT_SCOPE(GraphBuilder, "SLW::CopySceneWithoutWater");
 
@@ -846,26 +1169,29 @@ static FSceneWithoutWaterTextures AddCopySceneWithoutWaterPass(
 	check(SceneColorTexture);
 	check(SceneDepthTexture);
 
-	EShaderPlatform ShaderPlatform = Views[0].GetShaderPlatform();
+	const EShaderPlatform ShaderPlatform = Views[0].GetShaderPlatform();
+	const bool bCopyDepth = !SingleLayerWaterPrePassResult || !HasBeenProduced(SingleLayerWaterPrePassResult->SceneDepthWithoutWater);
 	const bool bCopyColor = !SingleLayerWaterUsesSimpleShading(ShaderPlatform);
 
-	const FRDGTextureDesc& SceneColorDesc = SceneColorTexture->Desc;
 	const FRDGTextureDesc& SceneDepthDesc = SceneColorTexture->Desc;
+	const FRDGTextureDesc& SceneColorDesc = SceneColorTexture->Desc;
 
-	const int32 RefractionDownsampleFactor = FMath::Clamp(GSingleLayerWaterRefractionDownsampleFactor, 1, 8);
+	const int32 RefractionDownsampleFactor = GetSingleLayerWaterRefractionDownsampleFactor();
 	const FIntPoint RefractionResolution = FIntPoint::DivideAndRoundDown(SceneColorDesc.Extent, RefractionDownsampleFactor);
-	FRDGTextureRef SceneColorWithoutSingleLayerWaterTexture = GraphBuilder.RegisterExternalTexture(GSystemTextures.BlackDummy);
-
+	FRDGTextureRef SceneColorWithoutSingleLayerWaterTexture = nullptr;
+	FRDGTextureRef SceneDepthWithoutSingleLayerWaterTexture = nullptr;
+	if (bCopyDepth)
+	{
+		// Note: if changing format, also update FWaterRefractionCopyPS::ModifyCompilationEnvironment accordingly
+		const FRDGTextureDesc DepthDesc(FRDGTextureDesc::Create2D(RefractionResolution, PF_R32_FLOAT, SceneDepthDesc.ClearValue, TexCreate_ShaderResource | TexCreate_RenderTargetable));
+		SceneDepthWithoutSingleLayerWaterTexture = GraphBuilder.CreateTexture(DepthDesc, TEXT("SLW.SceneDepthWithout"));
+	}
 	if (bCopyColor)
 	{
 		const FRDGTextureDesc ColorDesc = FRDGTextureDesc::Create2D(RefractionResolution, SceneColorDesc.Format, SceneColorDesc.ClearValue, TexCreate_ShaderResource | TexCreate_RenderTargetable);
 		SceneColorWithoutSingleLayerWaterTexture = GraphBuilder.CreateTexture(ColorDesc, TEXT("SLW.SceneColorWithout"));
 	}
 
-	// Note: if changing format, also update FWaterRefractionCopyPS::ModifyCompilationEnvironment accordingly
-	const FRDGTextureDesc DepthDesc(FRDGTextureDesc::Create2D(RefractionResolution, PF_R32_FLOAT, SceneDepthDesc.ClearValue, TexCreate_ShaderResource | TexCreate_RenderTargetable));
-	FRDGTextureRef SceneDepthWithoutSingleLayerWaterTexture = GraphBuilder.CreateTexture(DepthDesc, TEXT("SLW.SceneDepthWithout"));
-
 	const FRDGTextureDesc SeparatedMainDirLightDesc(FRDGTextureDesc::Create2D(SceneColorDesc.Extent, PF_FloatR11G11B10, FClearValueBinding(FLinearColor::White), TexCreate_ShaderResource | TexCreate_RenderTargetable));
 	FRDGTextureRef SeparatedMainDirLightTexture = GraphBuilder.CreateTexture(SeparatedMainDirLightDesc, TEXT("SLW.SeparatedMainDirLight"));
 
@@ -873,7 +1199,11 @@ static FSceneWithoutWaterTextures AddCopySceneWithoutWaterPass(
 	Textures.RefractionDownsampleFactor = float(RefractionDownsampleFactor);
 	Textures.Views.SetNum(Views.Num());
 
-	ERenderTargetLoadAction LoadAction = ERenderTargetLoadAction::ENoAction;
+	AddCopySceneWithoutWaterPass_Internal(GraphBuilder, ViewFamily, Views, 
+		SceneDepthWithoutSingleLayerWaterTexture, SceneDepthTexture, 
+		SceneColorWithoutSingleLayerWaterTexture, SceneColorTexture, 
+		RefractionDownsampleFactor,
+		SingleLayerWaterPrePassResult);
 
 	for (int32 ViewIndex = 0; ViewIndex < Views.Num(); ++ViewIndex)
 	{
@@ -884,33 +1214,6 @@ static FSceneWithoutWaterTextures AddCopySceneWithoutWaterPass(
 			continue;
 		}
 
-		RDG_GPU_MASK_SCOPE(GraphBuilder, View.GPUMask);
-		RDG_EVENT_SCOPE_CONDITIONAL(GraphBuilder, Views.Num() > 1, "View%d", ViewIndex);
-
-		FWaterRefractionCopyPS::FParameters* PassParameters = GraphBuilder.AllocParameters<FWaterRefractionCopyPS::FParameters>();
-		PassParameters->SceneColorCopyDownsampleTexture = SceneColorTexture;
-		PassParameters->SceneColorCopyDownsampleSampler = TStaticSamplerState<SF_Bilinear, AM_Clamp, AM_Clamp, AM_Clamp>::GetRHI();
-		PassParameters->SceneDepthCopyDownsampleTexture = SceneDepthTexture;
-		PassParameters->SceneDepthCopyDownsampleSampler = TStaticSamplerState<SF_Point, AM_Clamp, AM_Clamp, AM_Clamp>::GetRHI();
-		PassParameters->SVPositionToSourceTextureUV = FVector2f(RefractionDownsampleFactor / float(SceneColorDesc.Extent.X), RefractionDownsampleFactor / float(SceneColorDesc.Extent.Y));
-
-		PassParameters->RenderTargets[0] = FRenderTargetBinding(SceneDepthWithoutSingleLayerWaterTexture, LoadAction);
-
-		if (bCopyColor)
-		{
-			PassParameters->RenderTargets[1] = FRenderTargetBinding(SceneColorWithoutSingleLayerWaterTexture, LoadAction);
-		}
-
-		if (!View.Family->bMultiGPUForkAndJoin)
-		{
-			LoadAction = ERenderTargetLoadAction::ELoad;
-		}
-
-		FWaterRefractionCopyPS::FPermutationDomain PermutationVector;
-		PermutationVector.Set<FWaterRefractionCopyPS::FDownsampleRefraction>(RefractionDownsampleFactor > 1);
-		PermutationVector.Set<FWaterRefractionCopyPS::FDownsampleColor>(bCopyColor);
-		auto PixelShader = View.ShaderMap->GetShader<FWaterRefractionCopyPS>(PermutationVector);
-
 		// if we have a particular case of ISR where two views are laid out in side by side, we should copy both views at once
 		const bool bIsInstancedStereoSideBySide = View.bIsInstancedStereoEnabled && !View.bIsMobileMultiViewEnabled && IStereoRendering::IsStereoEyeView(View);
 		FIntRect RectToCopy = View.ViewRect;
@@ -924,8 +1227,7 @@ static FSceneWithoutWaterTextures AddCopySceneWithoutWaterPass(
 		}
 
 		const FIntRect RefractionViewRect = FIntRect(FIntPoint::DivideAndRoundDown(RectToCopy.Min, RefractionDownsampleFactor), FIntPoint::DivideAndRoundDown(RectToCopy.Max, RefractionDownsampleFactor));
-
-		Textures.Views[ViewIndex].ViewRect   = RefractionViewRect;
+		Textures.Views[ViewIndex].ViewRect = RefractionViewRect;
 
 		// This is usually half a pixel. But it seems that when using Gather4, 0.5 is not conservative enough and can return pixel outside the guard band. 
 		// That is why it is a tiny bit higher than 0.5: for Gathre4 to always return pixels within the valid side of UVs (see EvaluateWaterVolumeLighting).
@@ -934,21 +1236,15 @@ static FSceneWithoutWaterTextures AddCopySceneWithoutWaterPass(
 		Textures.Views[ViewIndex].MinMaxUV.Y = (RefractionViewRect.Min.Y + PixelSafeGuardBand) / RefractionResolution.Y;
 		Textures.Views[ViewIndex].MinMaxUV.Z = (RefractionViewRect.Max.X - PixelSafeGuardBand) / RefractionResolution.X;
 		Textures.Views[ViewIndex].MinMaxUV.W = (RefractionViewRect.Max.Y - PixelSafeGuardBand) / RefractionResolution.Y;
-
-		FPixelShaderUtils::AddFullscreenPass(
-			GraphBuilder,
-			View.ShaderMap,
-			{},
-			PixelShader,
-			PassParameters,
-			RefractionViewRect);
 	}
 
-	check(SceneColorWithoutSingleLayerWaterTexture);
-	check(SceneDepthWithoutSingleLayerWaterTexture);
-	Textures.ColorTexture = SceneColorWithoutSingleLayerWaterTexture;
-	Textures.DepthTexture = SceneDepthWithoutSingleLayerWaterTexture;
+	Textures.DepthTexture = bCopyDepth ? SceneDepthWithoutSingleLayerWaterTexture : SingleLayerWaterPrePassResult->SceneDepthWithoutWater;
+	Textures.ColorTexture = bCopyColor ? SceneColorWithoutSingleLayerWaterTexture : GraphBuilder.RegisterExternalTexture(GSystemTextures.BlackDummy);
 	Textures.SeparatedMainDirLightTexture = SeparatedMainDirLightTexture;
+
+	check(HasBeenProduced(Textures.DepthTexture));
+	check(HasBeenProduced(Textures.ColorTexture));
+
 	return MoveTemp(Textures);
 }
 
@@ -1322,7 +1618,7 @@ void FDeferredShadingSceneRenderer::RenderSingleLayerWater(
 	SCOPED_NAMED_EVENT(SingleLayerWater, FColor::Emerald);
 
 	// Copy the texture to be available for the water surface to refract
-	SceneWithoutWaterTextures = AddCopySceneWithoutWaterPass(GraphBuilder, ViewFamily, InViews, SceneTextures.Color.Resolve, SceneTextures.Depth.Resolve);
+	SceneWithoutWaterTextures = AddCopySceneWithoutWaterPass(GraphBuilder, ViewFamily, InViews, SceneTextures.Color.Resolve, SceneTextures.Depth.Resolve, SingleLayerWaterPrePassResult);
 
 	// Check if this is depth or base pass only renderer, where the final scene color isn't relevant, and we don't need fog, clouds, or reflections
 	bool bFinalSceneColor = !InViews[0].CustomRenderPass && GetRendererOutput() == ERendererOutput::FinalSceneColor;
@@ -1372,6 +1668,7 @@ BEGIN_UNIFORM_BUFFER_STRUCT(FSingleLayerWaterPassUniformParameters,)
 	SHADER_PARAMETER_RDG_TEXTURE(Texture2D, CustomDepthTexture)
 	SHADER_PARAMETER_RDG_TEXTURE_SRV(Texture2D<uint2>, CustomStencilTexture)
 	SHADER_PARAMETER_SAMPLER(SamplerState, CustomDepthSampler)
+	SHADER_PARAMETER_RDG_TEXTURE(Texture2D, RefractionMaskTexture)
 	SHADER_PARAMETER(FVector4f, SceneWithoutSingleLayerWaterMinMaxUV)
 	SHADER_PARAMETER(FVector4f, DistortionParams)
 	SHADER_PARAMETER(FVector2f, SceneWithoutSingleLayerWaterTextureSize)
@@ -1466,6 +1763,7 @@ void FDeferredShadingSceneRenderer::RenderSingleLayerWaterInner(
 			SLWUniformParameters.CustomDepthTexture = bCustomDepthTextureProduced ? SceneTextures.CustomDepth.Depth : SystemTextures.DepthDummy;
 			SLWUniformParameters.CustomStencilTexture = bCustomDepthTextureProduced ? SceneTextures.CustomDepth.Stencil : SystemTextures.StencilDummySRV;
 			SLWUniformParameters.CustomDepthSampler = TStaticSamplerState<SF_Point>::GetRHI();
+			SLWUniformParameters.RefractionMaskTexture = SingleLayerWaterPrePassResult && SingleLayerWaterPrePassResult->RefractionMaskTexture ? SingleLayerWaterPrePassResult->RefractionMaskTexture : SystemTextures.White;
 			SLWUniformParameters.SceneWithoutSingleLayerWaterMinMaxUV = SceneWithoutWaterTextures.Views[ViewIndex].MinMaxUV;
 			SetupDistortionParams(SLWUniformParameters.DistortionParams, View);
 			SLWUniformParameters.SceneWithoutSingleLayerWaterTextureSize = FVector2f(DepthTextureSize.X, DepthTextureSize.Y);
diff --git a/Engine/Source/Runtime/Renderer/Private/SingleLayerWaterRendering.h b/Engine/Source/Runtime/Renderer/Private/SingleLayerWaterRendering.h
--- a/Engine/Source/Runtime/Renderer/Private/SingleLayerWaterRendering.h
+++ b/Engine/Source/Runtime/Renderer/Private/SingleLayerWaterRendering.h
@@ -21,6 +21,8 @@ struct FSingleLayerWaterPrePassResult
 {
 	Froxel::FRenderer Froxels;
 	FRDGTextureMSAA DepthPrepassTexture;
+	FRDGTextureRef RefractionMaskTexture;
+	FRDGTextureRef SceneDepthWithoutWater;
 	TArray<FSingleLayerWaterTileClassification> ViewTileClassification;
 };
 
@@ -40,9 +42,20 @@ struct FSceneWithoutWaterTextures
 	float RefractionDownsampleFactor = 1.0f;
 };
 
+// Location relative to the base pass when to run the water depth prepass. If a full depth prepass is available,
+// the water depth prepass can run before the base pass, allowing certain optimizations to save work in the base pass.
+enum class ESingleLayerWaterPrepassLocation : uint8
+{
+	// The water depth prepass (if enabled) runs after the regular depth prepass and before the base pass.
+	BeforeBasePass,
+	// The water depth prepass (if enabled) runs after the base pass.
+	AfterBasePass
+};
+
 bool ShouldRenderSingleLayerWater(TArrayView<const FViewInfo> Views);
 bool ShouldRenderSingleLayerWaterSkippedRenderEditorNotification(TArrayView<const FViewInfo> Views);
 bool ShouldRenderSingleLayerWaterDepthPrepass(TArrayView<const FViewInfo> Views);
+ESingleLayerWaterPrepassLocation GetSingleLayerWaterDepthPrepassLocation(bool bFullDepthPrepass);
 bool ShouldUseBilinearSamplerForDepthWithoutSingleLayerWater(EPixelFormat DepthTextureFormat);
 
 class FWaterTileVS : public FGlobalShader
@@ -76,7 +89,7 @@ void SingleLayerWaterAddTiledFullscreenPass(
 	TPassParameters* PassParameters,
 	const TUniformBufferRef<FViewUniformShaderParameters>& ViewUniformBuffer,
 	const FIntRect& Viewport,
-	FTiledReflection* TiledScreenSpaceReflection = nullptr,
+	const FTiledReflection* TiledScreenSpaceReflection = nullptr,
 	FRHIBlendState* BlendState = nullptr,
 	FRHIRasterizerState* RasterizerState = nullptr,
 	FRHIDepthStencilState* DepthStencilState = nullptr,
