Bumped task priority on render tasks that were being starved and retracted back to the render thread.

diff --git a/Engine/Source/Runtime/RenderCore/Private/RenderGraphBuilder.cpp b/Engine/Source/Runtime/RenderCore/Private/RenderGraphBuilder.cpp
--- a/Engine/Source/Runtime/RenderCore/Private/RenderGraphBuilder.cpp
+++ b/Engine/Source/Runtime/RenderCore/Private/RenderGraphBuilder.cpp
@@ -1875,7 +1875,8 @@ void FRDGBuilder::Execute()
 			AddSetupTask([this, QueryBatchData = RHICmdList.GetQueryBatchData(RQT_AbsoluteTime)]
 			{
 				SetupParallelExecute(QueryBatchData);
-			});
+
+			}, TaskPriority);
 		}
 
 		UE::Tasks::FTask AllocatePooledBuffersTask;
diff --git a/Engine/Source/Runtime/Renderer/Private/RayTracing/RayTracing.cpp b/Engine/Source/Runtime/Renderer/Private/RayTracing/RayTracing.cpp
--- a/Engine/Source/Runtime/Renderer/Private/RayTracing/RayTracing.cpp
+++ b/Engine/Source/Runtime/Renderer/Private/RayTracing/RayTracing.cpp
@@ -525,7 +525,7 @@ namespace RayTracing
 					Page = NextPage;
 				}
 				check(ExpectedMaxVisibleCommands <= TaskVisibleShaderBindings->Max());
-			}));
+			}, UE::Tasks::ETaskPriority::High));
 	};
 
 	FDynamicRayTracingInstancesContext::FDynamicRayTracingInstancesContext(FScene& Scene, FViewInfo& InView, FSceneRenderingBulkObjectAllocator& InBulkAllocator, const FSceneOptions& InSceneOptions)
diff --git a/Engine/Source/Runtime/Renderer/Private/SceneCulling/SceneCullingRenderer.cpp b/Engine/Source/Runtime/Renderer/Private/SceneCulling/SceneCullingRenderer.cpp
--- a/Engine/Source/Runtime/Renderer/Private/SceneCulling/SceneCullingRenderer.cpp
+++ b/Engine/Source/Runtime/Renderer/Private/SceneCulling/SceneCullingRenderer.cpp
@@ -244,7 +244,7 @@ void FSceneInstanceCullingQuery::Dispatch(FRDGBuilder& GraphBuilder, bool bInAll
 		{
 			ComputeResult();
 		},
-		nullptr, TArray<UE::Tasks::FTask>{ UpdateTaskHandle }, UE::Tasks::ETaskPriority::Normal, bAllowAsync);
+		nullptr, TArray<UE::Tasks::FTask>{ UpdateTaskHandle }, UE::Tasks::ETaskPriority::High, bAllowAsync);
 	}
 }
 
