[DOF] Fix crash when rendering without ViewState due to uninitialized pointer.

diff --git a/Engine/Source/Runtime/Renderer/Private/PostProcess/DiaphragmDOF.cpp b/Engine/Source/Runtime/Renderer/Private/PostProcess/DiaphragmDOF.cpp
--- a/Engine/Source/Runtime/Renderer/Private/PostProcess/DiaphragmDOF.cpp
+++ b/Engine/Source/Runtime/Renderer/Private/PostProcess/DiaphragmDOF.cpp
@@ -1579,7 +1579,7 @@ bool DiaphragmDOF::AddPasses(
 	FBokehModel BokehModel;
 	BokehModel.Compile(View, CocModel);
 
-	FDOFViewState* DOFViewState;
+	FDOFViewState* DOFViewState = nullptr;
 	TStaticArray<bool, (int)EDiaphragmDOFBokehLUTFormat::MAX> HasCachedBokehLUT {};
 	if (View.ViewState)
 	{
