[TSR] Fix performance regression due to unnecessary pass to measure flickering luma.

#tests Test done on console and performance regression fixed via perfB.
#rnx

diff --git a/Engine/Source/Runtime/Renderer/Private/PostProcess/TemporalSuperResolution.cpp b/Engine/Source/Runtime/Renderer/Private/PostProcess/TemporalSuperResolution.cpp
--- a/Engine/Source/Runtime/Renderer/Private/PostProcess/TemporalSuperResolution.cpp
+++ b/Engine/Source/Runtime/Renderer/Private/PostProcess/TemporalSuperResolution.cpp
@@ -1581,7 +1581,7 @@ bool NeedTSRAntiFlickeringPass(const FViewInfo& View)
 
 bool ShouldAddTSRMainFlickeringLumaPass()
 {
-	return  CVarTSRFlickeringEnable.GetValueOnRenderThread()!=0 || CVarTSRFlickeringPeriod.GetValueOnRenderThread() == 0.0f || ShouldApplySkyRelaxation();
+	return  (CVarTSRFlickeringEnable.GetValueOnRenderThread()!=0 && CVarTSRFlickeringPeriod.GetValueOnRenderThread() != 0.0f) || ShouldApplySkyRelaxation();
 }
 
 bool ShouldAddTSRMainThinGeometryCoveragePass()
@@ -1625,7 +1625,7 @@ FScreenPassTexture AddTSRMeasureFlickeringLuma(FRDGBuilder& GraphBuilder, FGloba
 		FlickeringLuma.ViewRect = SceneColor.ViewRect;
 	}
 
-	if (ShouldAddTSRMainFlickeringLumaPass() || !ShouldAddTSRMainThinGeometryCoveragePass())
+	if (ShouldAddTSRMainFlickeringLumaPass())
 	{
 		FTSRMeasureFlickeringLumaCS::FParameters* PassParameters = GraphBuilder.AllocParameters<FTSRMeasureFlickeringLumaCS::FParameters>();
 		PassParameters->InputInfo = GetScreenPassTextureViewportParameters(FScreenPassTextureViewport(
