Fix occlusion queries failing to set stereo viewport for instanced stereo, leading to incorrect occlusion culling in secondary instanced views.

diff --git a/Engine/Source/Runtime/Renderer/Private/SceneOcclusion.cpp b/Engine/Source/Runtime/Renderer/Private/SceneOcclusion.cpp
--- a/Engine/Source/Runtime/Renderer/Private/SceneOcclusion.cpp
+++ b/Engine/Source/Runtime/Renderer/Private/SceneOcclusion.cpp
@@ -1263,8 +1263,26 @@ static void BeginOcclusionTests(
 		// We only need to render the front-faces of the culling geometry (this halves the amount of pixels we touch)
 		GraphicsPSOInit.RasterizerState = View.bReverseCulling ? TStaticRasterizerState<FM_Solid, CM_CCW>::GetRHI() : TStaticRasterizerState<FM_Solid, CM_CW>::GetRHI();
 
-		const FIntRect ViewRect = GetDownscaledRect(View.ViewRect, DownsampleFactor);
-		RHICmdList.SetViewport(ViewRect.Min.X, ViewRect.Min.Y, 0.0f, ViewRect.Max.X, ViewRect.Max.Y, 1.0f);
+		// Set stereo viewports while accounting for the downsample factor
+		if (View.IsInstancedStereoPass())
+		{
+			if (View.bIsMultiViewportEnabled)
+			{
+				const FIntRect ViewRect = GetDownscaledRect(View.ViewRect, DownsampleFactor);
+				const FIntRect SecondaryViewRect = GetDownscaledRect(View.GetInstancedView()->ViewRect, DownsampleFactor);
+				RHICmdList.SetStereoViewport(ViewRect.Min.X, SecondaryViewRect.Min.X, ViewRect.Min.Y, SecondaryViewRect.Min.Y, 0.0f, ViewRect.Max.X, SecondaryViewRect.Max.X, ViewRect.Max.Y, SecondaryViewRect.Max.Y, 1.0f);
+			}
+			else
+			{
+				const FIntRect ViewRect = GetDownscaledRect(View.ViewRectWithSecondaryViews, DownsampleFactor);
+				RHICmdList.SetViewport(ViewRect.Min.X, ViewRect.Min.Y, 0.0f, ViewRect.Max.X, ViewRect.Max.Y, 1.0f);
+			}
+		}
+		else
+		{
+			const FIntRect ViewRect = GetDownscaledRect(View.ViewRect, DownsampleFactor);
+			RHICmdList.SetViewport(ViewRect.Min.X, ViewRect.Min.Y, 0.0f, ViewRect.Max.X, ViewRect.Max.Y, 1.0f);
+		}
 
 		// Lookup the vertex shader.
 		TShaderMapRef<FOcclusionQueryVS> VertexShader(View.ShaderMap);
