Plumb through "MaxNonFarCascadeDistance" from the old shadow setup to the ShadowSceneRenderer to match preveious behavior.

#rnx

diff --git a/Engine/Source/Runtime/Renderer/Private/ShadowSetup.cpp b/Engine/Source/Runtime/Renderer/Private/ShadowSetup.cpp
--- a/Engine/Source/Runtime/Renderer/Private/ShadowSetup.cpp
+++ b/Engine/Source/Runtime/Renderer/Private/ShadowSetup.cpp
@@ -328,7 +328,7 @@ static TAutoConsoleVariable<int32> CVarAlwaysAllocateMaxResolutionAtlases(
 	ECVF_RenderThreadSafe
 );
 
-static TAutoConsoleVariable<int32> CVarVsmUseFarShadowRules(
+TAutoConsoleVariable<int32> CVarVsmUseFarShadowRules(
 	TEXT("r.Shadow.Virtual.UseFarShadowCulling"),
 	1,
 	TEXT("Switch between implementing the far shadow culling logic for VSMs."),
@@ -5676,7 +5676,7 @@ void FSceneRenderer::AddViewDependentWholeSceneShadowsForView(
 			if (bNeedsVirtualShadowMap)
 			{
 				FShadowSceneRenderer& ShadowSceneRenderer = GetSceneExtensionsRenderers().GetRenderer<FShadowSceneRenderer>();
-				ShadowSceneRenderer.AddDirectionalLightShadow(LightSceneInfo, View, ShadowInfosThatNeedCulling);
+				ShadowSceneRenderer.AddDirectionalLightShadow(LightSceneInfo, View, MaxNonFarCascadeDistance, ShadowInfosThatNeedCulling);
 			}
 		}
 	}
diff --git a/Engine/Source/Runtime/Renderer/Private/Shadows/ShadowSceneRenderer.cpp b/Engine/Source/Runtime/Renderer/Private/Shadows/ShadowSceneRenderer.cpp
--- a/Engine/Source/Runtime/Renderer/Private/Shadows/ShadowSceneRenderer.cpp
+++ b/Engine/Source/Runtime/Renderer/Private/Shadows/ShadowSceneRenderer.cpp
@@ -19,6 +19,7 @@
 CSV_DECLARE_CATEGORY_EXTERN(VSM);
 
 extern TAutoConsoleVariable<int32> CVarNaniteShadowsUpdateStreaming;
+extern TAutoConsoleVariable<int32> CVarVsmUseFarShadowRules;
 
 TAutoConsoleVariable<int32> CVarVSMMaterialVisibility(
 	TEXT("r.Shadow.Virtual.Nanite.MaterialVisibility"),
@@ -499,7 +500,7 @@ TSharedPtr<FVirtualShadowMapPerLightCacheEntry> FShadowSceneRenderer::AddLocalLi
 	return PerLightCacheEntry;
 }
 
-void FShadowSceneRenderer::AddDirectionalLightShadow(FLightSceneInfo& LightSceneInfo, FViewInfo& View, TArray<FProjectedShadowInfo*, SceneRenderingAllocator>& OutShadowInfosThatNeedCulling)
+void FShadowSceneRenderer::AddDirectionalLightShadow(FLightSceneInfo& LightSceneInfo, FViewInfo& View, float MaxNonFarCascadeDistance, TArray<FProjectedShadowInfo*, SceneRenderingAllocator>& OutShadowInfosThatNeedCulling)
 {
 	const int32 ViewIndex = View.SceneRendererPrimaryViewId;
 	FLightSceneInfo::FPersistentId LightId = LightSceneInfo.Id;
@@ -514,7 +515,7 @@ void FShadowSceneRenderer::AddDirectionalLightShadow(FLightSceneInfo& LightScene
 	auto AddLegacySetup = [&](const TSharedPtr<FVirtualShadowMapClipmap>& Clipmap, bool bQueueforNonNaniteCulling)
 	{
 		FProjectedShadowInfo* ProjectedShadowInfo = SceneRenderer.Allocator.Create<FProjectedShadowInfo>();
-		ProjectedShadowInfo->SetupClipmapProjection(&LightSceneInfo, &View, Clipmap, ProjectedShadowInfo->GetMaxNonFarCascadeDistance());
+		ProjectedShadowInfo->SetupClipmapProjection(&LightSceneInfo, &View, Clipmap, CVarVsmUseFarShadowRules.GetValueOnRenderThread() != 0 ? MaxNonFarCascadeDistance : -1.0f);
 	
 		// This is needed to get it into the line for ending up in the SortedShadowsForShadowDepthPass.VirtualShadowMapShadows which is what drives the shadow rendering.
 		VisibleLightInfo.AllProjectedShadows.Add(ProjectedShadowInfo);
diff --git a/Engine/Source/Runtime/Renderer/Private/Shadows/ShadowSceneRenderer.h b/Engine/Source/Runtime/Renderer/Private/Shadows/ShadowSceneRenderer.h
--- a/Engine/Source/Runtime/Renderer/Private/Shadows/ShadowSceneRenderer.h
+++ b/Engine/Source/Runtime/Renderer/Private/Shadows/ShadowSceneRenderer.h
@@ -68,7 +68,7 @@ public:
 	/**
 	 * Add a directional light for processing this frame.
 	 */
-	void AddDirectionalLightShadow(FLightSceneInfo& LightSceneInfo, FViewInfo& View, TArray<FProjectedShadowInfo*, SceneRenderingAllocator>& OutShadowInfosThatNeedCulling);
+	void AddDirectionalLightShadow(FLightSceneInfo& LightSceneInfo, FViewInfo& View, float MaxNonFarCascadeDistance, TArray<FProjectedShadowInfo*, SceneRenderingAllocator>& OutShadowInfosThatNeedCulling);
 
 	/**
 	 * Call after view-dependent setup has been processed (InitView etc) but before any rendering activity has been kicked off.
