Changed packing of instance hierarchy cells to avoid overflow when > 4M instances end up in the same cell (which happens in PCG test content).

#rnx

diff --git a/Engine/Shaders/Private/SceneCulling/SceneCulling.ush b/Engine/Shaders/Private/SceneCulling/SceneCulling.ush
--- a/Engine/Shaders/Private/SceneCulling/SceneCulling.ush
+++ b/Engine/Shaders/Private/SceneCulling/SceneCulling.ush
@@ -41,6 +41,25 @@ FCellBlockData GetCellBlockData(uint BlockId)
 
 StructuredBuffer<FPackedCellHeader> InstanceHierarchyCellHeaders;
 
+FCellHeader UnpackCellHeader(FPackedCellHeader Packed)
+{
+	uint4 Bits = uint4(Packed.Packed0, Packed.Packed1, 0u, 0u);
+	FCellHeader Result;
+	uint Offset = 0u;
+	Result.ItemChunksOffset = ReadBits(Bits, Offset, INSTANCE_HIERARCHY_CELL_HEADER_OFFSET_BITS);
+	Result.NumStaticChunks  = ReadBits(Bits, Offset, INSTANCE_HIERARCHY_CELL_HEADER_COUNT_BITS);
+	Result.NumDynamicChunks = ReadBits(Bits, Offset, INSTANCE_HIERARCHY_CELL_HEADER_COUNT_BITS); 
+	
+	Result.bIsValid = Result.NumDynamicChunks != 0u;
+	if (Result.bIsValid)
+	{
+		Result.NumDynamicChunks -= 1u;
+	}	
+	Result.NumItemChunks = Result.NumDynamicChunks + Result.NumStaticChunks;
+
+	return Result;
+}
+
 FCellHeader GetCellHeader(uint CellId)
 {
 	return UnpackCellHeader(InstanceHierarchyCellHeaders[CellId]);
diff --git a/Engine/Shaders/Shared/SceneCullingDefinitions.h b/Engine/Shaders/Shared/SceneCullingDefinitions.h
--- a/Engine/Shaders/Shared/SceneCullingDefinitions.h
+++ b/Engine/Shaders/Shared/SceneCullingDefinitions.h
@@ -9,6 +9,9 @@
 #define  INSTANCE_HIERARCHY_ITEM_CHUNK_COMPRESSED_PAYLOAD_MASK  ((1u << 31u) - 1u)
 #define  INSTANCE_HIERARCHY_MAX_CHUNK_SIZE  (64u)
 
+#define  INSTANCE_HIERARCHY_CELL_HEADER_OFFSET_BITS (22u) // 22 for offset
+#define  INSTANCE_HIERARCHY_CELL_HEADER_COUNT_BITS (21u)  // 2x21 for static/dynamic
+
 #ifdef __cplusplus
 #include "HLSLTypeAliases.h"
 
@@ -22,8 +25,8 @@ namespace UE::HLSL
 // Representation in hierarchy buffer
 struct FPackedCellHeader
 {
-	uint ItemChunksOffset;
-	uint NumStaticDynamicItemChunks; // 2x16 bits: (Static << 16) + Dynamic
+	uint Packed0;
+	uint Packed1;
 };
 
 // Unpacked version of the struct
@@ -33,18 +36,9 @@ struct FCellHeader
 	uint ItemChunksOffset;
 	uint NumStaticChunks;
 	uint NumDynamicChunks;
+	bool bIsValid;
 };
 
-inline FCellHeader UnpackCellHeader(FPackedCellHeader Packed)
-{
-	FCellHeader Result;
-	Result.ItemChunksOffset = Packed.ItemChunksOffset;
-	Result.NumStaticChunks = Packed.NumStaticDynamicItemChunks >> 16;
-	Result.NumDynamicChunks = Packed.NumStaticDynamicItemChunks & 0xFFFFu;
-	Result.NumItemChunks = Result.NumDynamicChunks + Result.NumStaticChunks;
-	return Result;
-}
-
 struct FCellBlockData
 {
 	FDFVector3 WorldPos;
diff --git a/Engine/Source/Runtime/Renderer/Private/SceneCulling/SceneCulling.cpp b/Engine/Source/Runtime/Renderer/Private/SceneCulling/SceneCulling.cpp
--- a/Engine/Source/Runtime/Renderer/Private/SceneCulling/SceneCulling.cpp
+++ b/Engine/Source/Runtime/Renderer/Private/SceneCulling/SceneCulling.cpp
@@ -283,6 +283,17 @@ namespace EUpdateFrequencyCategory
 		Dynamic,
 		Num,
 	};
+
+	const TCHAR* ToString(EType Category)
+	{
+		const TCHAR* Strs[Num] =
+		{
+			TEXT("Static"),
+			TEXT("Dynamic"),
+		};
+		return Strs[Category];
+	}
+
 }
 
 FString GSceneCullingDbgPattern;// = "Cube*";
@@ -375,13 +386,25 @@ struct FSpatialHashDebugDrawer
 };
 #endif
 
-// works for both packed and not
-template <typename CellHeaderType>
-inline bool IsTempCell(const CellHeaderType& CellHeader)
+
+SC_FORCEINLINE FPackedCellHeader PackCellHeader(const FCellHeader& CellHeader)
 {
-	return (CellHeader.ItemChunksOffset & FSceneCulling::TempCellFlag) != 0;
-}
+	check(CellHeader.bIsValid);
+	check(CellHeader.ItemChunksOffset < (1u << INSTANCE_HIERARCHY_CELL_HEADER_OFFSET_BITS));
+	check(CellHeader.NumStaticChunks < (1u << INSTANCE_HIERARCHY_CELL_HEADER_COUNT_BITS));
+	// use 0 to represent an invalid cell
+	check(CellHeader.NumDynamicChunks  < ((1u << INSTANCE_HIERARCHY_CELL_HEADER_COUNT_BITS) - 1u));
+	
+	FPackedCellHeader Packed;
 
+	uint32 NumDynamicChunks = CellHeader.bIsValid ? (CellHeader.NumDynamicChunks + 1u) : 0u;
+	uint64& Bits = reinterpret_cast<uint64&>(Packed);
+	Bits = (uint64(CellHeader.ItemChunksOffset) << (2u * INSTANCE_HIERARCHY_CELL_HEADER_COUNT_BITS))
+		| (uint64(CellHeader.NumStaticChunks) << INSTANCE_HIERARCHY_CELL_HEADER_COUNT_BITS)
+		| uint64(NumDynamicChunks);
+	
+	return Packed;
+}
 
 SC_FORCEINLINE FSceneCulling::FLocation8 ToBlockLocal(const FSceneCulling::FLocation64& ItemLoc, const FSceneCulling::FBlockLoc& BlockLoc)
 {
@@ -784,8 +807,7 @@ public:
 	// Alias a bunch of types / definitions from the instance hierarchy
 	using FSpatialHash = FSceneCulling::FSpatialHash;
 	using FHashElementId = FSpatialHash::FHashElementId;
-	static constexpr uint32 InvalidCellFlag = FSceneCulling::InvalidCellFlag;
-	static constexpr uint32 TempCellFlag = FSceneCulling::TempCellFlag;
+	static constexpr uint32 TempCellMarker = 0xFFFFFFFFu;
 
 	static constexpr int32 CellBlockSize = FSpatialHash::CellBlockSize;
 	static constexpr int32 CellBlockDimLog2 = FSpatialHash::CellBlockDimLog2;
@@ -805,6 +827,13 @@ public:
 
 	EExplicitBoundsUpdateMode ExplicitBoundsMode = EExplicitBoundsUpdateMode::Disabled;
 
+	inline bool IsTempCell(int32 CellIndex)
+	{
+		return TempCellMask.IsValidIndex(CellIndex) && TempCellMask[CellIndex];
+	}
+
+	static constexpr FPackedCellHeader InvalidPackedCell = FPackedCellHeader{ 0u, 0u };
+
 #if SC_ENABLE_DETAILED_BUILDER_STATS
 	// Detail Stats
 	struct FStats
@@ -956,6 +985,7 @@ public:
 
 		SC_FORCEINLINE void OutputChunkIds(FSceneCullingBuilder &Builder, uint32 ItemChunksOffset)
 		{
+			BUILDER_LOG("OutputChunkIds Num: %u -> %d", PackedChunkIds.Num(), ItemChunksOffset);
 			int32 NumIds = PackedChunkIds.Num();
 			// 3. Copy the list of chunk IDs
 			for (uint32 Index = 0; Index < uint32(NumIds); ++Index)
@@ -1127,7 +1157,7 @@ public:
 		int32 RemovedInstanceCount[EUpdateFrequencyCategory::Num] = { 0, 0 };
 		bool  bAnyChange[EUpdateFrequencyCategory::Num] = { false, false };
 
-		FCellHeader PrevCellHeader = FCellHeader { 0u, InvalidCellFlag };
+		FCellHeader PrevCellHeader = FCellHeader { false, 0u };
 
 		SC_FORCEINLINE int32 GetRemovedInstanceCount() const
 		{
@@ -1226,7 +1256,6 @@ public:
 		return FBlockLoc(RenderingSpatialHash::TLocation<int64>(ToLevelRelative<CellBlockDimLog2>(Loc.Coord), Loc.Level + CellBlockDimLog2));
 	};
 
-
 	SC_FORCEINLINE FTempCell& GetOrAddTempCell(const FSceneCulling::FLocation64& InstanceCellLoc)
 	{
 		// Address of the cell-block touched
@@ -1252,7 +1281,7 @@ public:
 			// TODO: store the validity state in bit mask instead of with each item?
 			for (int32 Index = 0; Index < CellBlockSize; ++Index)
 			{
-				SceneCulling.CellHeaders[StartIndex + Index].ItemChunksOffset = InvalidCellFlag;
+				SceneCulling.CellHeaders[StartIndex + Index] = InvalidPackedCell;
 			}			
 			// No need to set the bits as they are maintained incrementally (or new and cleared)
 			SceneCulling.CellOccupancyMask.SetNum(NewMinSize, false);
@@ -1275,8 +1304,12 @@ public:
 	{
 		FPackedCellHeader CellHeader = SceneCulling.CellHeaders[CellIndex];
 
-		if (!IsTempCell(CellHeader))
+		if (!IsTempCell(CellIndex))
 		{
+			// Mark the cell as a temp cell
+			TempCellMask.PadToNum(SceneCulling.CellHeaders.Num(), false);
+			TempCellMask[CellIndex] = true;
+
 			int32 TempCellIndex = TempCells.AddDefaulted();
 			BUILDER_LOG("Alloc Temp Cell %d / %d", CellIndex, TempCellIndex);
 			// Store link back to the cell in question.
@@ -1288,15 +1321,15 @@ public:
 			LogCell(TempCell.PrevCellHeader);
 
 			// Hijack the items offset to store the index to the temp cell so we can add data there during construction.
-			CellHeader.ItemChunksOffset = TempCellFlag;
-			CellHeader.NumStaticDynamicItemChunks = TempCellIndex;
+			CellHeader.Packed0 = TempCellIndex;
+			CellHeader.Packed1 = TempCellMarker;
 
 			SceneCulling.CellHeaders[CellIndex] = CellHeader;
 
 			return TempCell;
 		}
-
-		return TempCells[CellHeader.NumStaticDynamicItemChunks];
+		check(CellHeader.Packed1 == TempCellMarker);
+		return TempCells[CellHeader.Packed0];
 	}
 
 	SC_FORCEINLINE int32 AddRange(const FSceneCulling::FLocation64& InstanceCellLoc, int32 InInstanceIdOffset, int32 InInstanceIdCount, EUpdateFrequencyCategory::EType UpdateFrequencyCategory)
@@ -1494,8 +1527,9 @@ public:
 		for (FTempCell& TempCell : TempCells)
 		{
 			BUILDER_LOG_SCOPE("OutputChunkIds");
-			FPackedCellHeader &CellHeader = SceneCulling.CellHeaders[TempCell.CellOffset];
-			check(IsTempCell(CellHeader));
+			FPackedCellHeader& OutPackedCellHeader = SceneCulling.CellHeaders[TempCell.CellOffset];
+			check(IsTempCell(TempCell.CellOffset));
+			check(OutPackedCellHeader.Packed1 == TempCellMarker);
 			// mark block as dirty
 			DirtyBlocks[SceneCulling.CellIndexToBlockId(TempCell.CellOffset)] = true;
 			++NumDirtyBlocks;
@@ -1503,24 +1537,28 @@ public:
 			{
 				BUILDER_LOG("Mark Empty Cell: %d", TempCell.CellOffset);
 				SceneCulling.CellOccupancyMask[TempCell.CellOffset] = false;
-				CellHeader.ItemChunksOffset = InvalidCellFlag;
-				CellHeader.NumStaticDynamicItemChunks = 0u;
+				OutPackedCellHeader = InvalidPackedCell;
 			}
 			else
 			{
 				SceneCulling.CellOccupancyMask[TempCell.CellOffset] = true;
 				// 2. Store final offsets
+				FCellHeader CellHeader;
+				CellHeader.bIsValid = true;
 				CellHeader.ItemChunksOffset = TempCell.ItemChunksOffset;
-				CellHeader.NumStaticDynamicItemChunks = (TempCell.GetNumItemChunks(EUpdateFrequencyCategory::Static) << 16u) | TempCell.GetNumItemChunks(EUpdateFrequencyCategory::Dynamic);
+				CellHeader.NumStaticChunks  = TempCell.GetNumItemChunks(EUpdateFrequencyCategory::Static);
+				CellHeader.NumDynamicChunks = TempCell.GetNumItemChunks(EUpdateFrequencyCategory::Dynamic);
+				CellHeader.NumItemChunks = CellHeader.NumDynamicChunks + CellHeader.NumStaticChunks;
+				OutPackedCellHeader = PackCellHeader(CellHeader);
 
 				// 3. Copy the list of chunk IDs
 				TempCell.OutputChunkIds(*this);
 
 				MarkCellBoundsDirty(TempCell.CellOffset, TempCell.bAnyChange[EUpdateFrequencyCategory::Static], TempCell.bAnyChange[EUpdateFrequencyCategory::Dynamic]);
 			}
-			LogCell(UE::HLSL::UnpackCellHeader(CellHeader));
-			CellHeaderUploader.Add(CellHeader, TempCell.CellOffset);
-			BUILDER_LOG("CellHeaderUploader { %u, %u} -> %d", CellHeader.ItemChunksOffset, CellHeader.NumItemChunks, TempCell.CellOffset);
+			LogCell(UnpackCellHeader(OutPackedCellHeader));
+			CellHeaderUploader.Add(OutPackedCellHeader, TempCell.CellOffset);
+			BUILDER_LOG("CellHeaderUploader { %u, %u, %u } -> %d", UnpackCellHeader(OutPackedCellHeader).ItemChunksOffset, UnpackCellHeader(OutPackedCellHeader).NumItemChunks, UnpackCellHeader(OutPackedCellHeader).NumDynamicChunks, TempCell.CellOffset);
 		}
 		}
 		TempCells.Empty();
@@ -1773,6 +1811,7 @@ public:
 
 	inline void MarkCellForRemove(int32 CellIndex, int32 NumInstances, EUpdateFrequencyCategory::EType UpdateFrequencyCategory)
 	{
+		BUILDER_LOG("MarkCellForRemove: %d, %d, %s", CellIndex, NumInstances, ToString(UpdateFrequencyCategory));
 		FTempCell &TempCell = GetOrAddTempCell(CellIndex);
 		// Should not mark remove for a cell that didn't have anything in it before...
 		check(IsValidCell(TempCell.PrevCellHeader));
@@ -1794,6 +1833,7 @@ public:
 
 	inline void MarkForRemove(int32 CellIndex, int32 InstanceDataOffset, int32 NumInstances, EUpdateFrequencyCategory::EType UpdateFrequencyCategory)
 	{
+		BUILDER_LOG("MarkForRemove: %d, %d, %d, %s", CellIndex, InstanceDataOffset, NumInstances, ToString(UpdateFrequencyCategory));
 		RemovedInstanceFlags.SetRange(InstanceDataOffset, NumInstances, true);
 		MarkCellForRemove(CellIndex, NumInstances, UpdateFrequencyCategory);
 	}
@@ -1856,6 +1896,8 @@ public:
 		// 1: mark all instances that need clearing out
 		// TODO: need only mark one bit for compressed chunks. Track knowledge of this?
 		RemovedInstanceFlags.SetRange(PrimitiveState.InstanceDataOffset, PrimitiveState.NumInstances, true);
+		BUILDER_LOG("RemovedInstanceFlags: %d, %d", PrimitiveState.InstanceDataOffset, PrimitiveState.NumInstances);
+
 		// Handle singular case
 		if (PrimitiveState.State == FPrimitiveState::SingleCell)
 		{ 
@@ -2038,6 +2080,7 @@ public:
 			BUILDER_LOG("FPrimitiveState::Precomputed");
 
 			RemovedInstanceFlags.SetRange(PrevPrimitiveState.InstanceDataOffset, PrevPrimitiveState.NumInstances, true);
+			BUILDER_LOG("RemovedInstanceFlags: %d, %d", PrevPrimitiveState.InstanceDataOffset, PrevPrimitiveState.NumInstances);
 
 			RemovePrecomputed(PrevPrimitiveState);
 
@@ -2052,6 +2095,7 @@ public:
 			FCellIndexCacheEntry &CellIndexCacheEntry = GetCacheEntry(PrevPrimitiveState.Payload);
 
 			RemovedInstanceFlags.SetRange(PrevPrimitiveState.InstanceDataOffset, PrevPrimitiveState.NumInstances, true);
+			BUILDER_LOG("RemovedInstanceFlags: %d, %d", PrevPrimitiveState.InstanceDataOffset, PrevPrimitiveState.NumInstances);
 			for (int32 ItemIndex = 0; ItemIndex < CellIndexCacheEntry.Items.Num(); ++ItemIndex)
 			{
 				FCellIndexCacheEntry::FItem Item = CellIndexCacheEntry.LoadAndStepItem(ItemIndex);
@@ -2263,8 +2307,8 @@ public:
 					// In the future (if we start doing GPU-side traversal that might need to change).
 					if (IsValidCell(HostValue))
 					{
-						check(GPUValue.ItemChunksOffset == HostValue.ItemChunksOffset); 
-						check(GPUValue.NumStaticDynamicItemChunks == HostValue.NumStaticDynamicItemChunks); 
+						check(GPUValue.Packed0 == HostValue.Packed0); 
+						check(GPUValue.Packed1 == HostValue.Packed1); 
 					}
 				});
 			SceneCulling.ItemChunksBuffer.ValidateGPUData(GraphBuilder, TConstArrayView<const uint32>(SceneCulling.PackedCellChunkData), 
@@ -2309,7 +2353,7 @@ public:
 		SCOPE_CYCLE_COUNTER(STAT_SceneCulling_Update_Post);
 		CSV_SCOPED_TIMING_STAT(SceneCulling, PostSceneUpdate);
 
-		BUILDER_LOG_SCOPE("ProcessPostSceneUpdate: %d/%d", ScenePostUpdateData.UpdatedPrimitiveIds.Num(), ScenePostUpdateData.AddedPrimitiveIds.Num());
+		BUILDER_LOG_SCOPE("ProcessPostSceneUpdate: %d/%d", ScenePostUpdateData.PrimitiveUpdates.NumCommands(), ScenePostUpdateData.AddedPrimitiveIds.Num());
 
 		SceneCulling.PrimitiveStates.SetNum(SceneCulling.Scene.GetMaxPersistentPrimitiveIndex(), EAllowShrinking::No);
 
@@ -2382,7 +2426,7 @@ public:
 			{
 				SceneCulling.CellBlockData[BlockIndex] = BlockData;
 				BlockDataUploader.Add(BlockData, BlockIndex);
-				BUILDER_LOG("BlockDataUploader {%0.0f, %0.0f, %0.0f, %0.0f} -> %d", BlockData.WorldPos.GetAbsolute().X, BlockData.WorldPos.GetAbsolute().Y, BlockData.WorldPos.GetAbsolute().Z, BlockData.LevelCellSize, BlockIndex);
+				BUILDER_LOG("BlockDataUploader {%0.0f, %0.0f, %0.0f, %0.0f} -> %d", BlockData.WorldPos.GetVector3d().X, BlockData.WorldPos.GetVector3d().Y, BlockData.WorldPos.GetVector3d().Z, BlockData.LevelCellSize, BlockIndex);
 			}
 			else
 			{
@@ -2465,6 +2509,9 @@ public:
 	FSceneCulling::FSpatialHash &SpatialHash;
 	FSceneCulling::FSpatialHash::FSpatialHashMap &SpatialHashMap;
 
+	// Tracks current temp cells (cells where we store an index to a temp cell instead of proper cell data).
+	TBitArray<SceneRenderingAllocator> TempCellMask;
+
 	// Dirty state tracking for upload
 	int32 NumDirtyBlocks = 0;
 	TBitArray<SceneRenderingAllocator> DirtyBlocks;
@@ -2502,7 +2549,7 @@ public:
 		if (bIsLoggingEnabled)
 		{
 			FString Indent = FString::ChrN(LogStrIndent, TEXT(' '));
-			UE_LOG(LogTemp, Warning, TEXT("%s%s"), *Indent, *Item);
+			UE_LOG(LogTemp, Log, TEXT("%s%s"), *Indent, *Item);
 		}
 	}
 	inline void EndLogging()
@@ -2631,7 +2678,7 @@ void FSceneCulling::FUpdater::PreSceneUpdate(FRDGBuilder& GraphBuilder, const FS
 
 	if (SceneCulling.BeginUpdate(GraphBuilder, bAnySceneUpdatesQueued))
 	{
-		SC_DETAILED_LOGGING_SCOPE(SceneCulling.ActiveUpdaterImplementation);
+		SC_DETAILED_LOGGING_SCOPE(SceneCulling.ActiveUpdaterImplementation.Get());
 
 		// Handle all removed primitives 
 		// Step 1. Mark all removed instances
@@ -2641,7 +2688,7 @@ void FSceneCulling::FUpdater::PreSceneUpdate(FRDGBuilder& GraphBuilder, const FS
 		{
 			SCOPED_NAMED_EVENT(SceneCulling_Update, FColor::Emerald);
 			SCOPE_CYCLE_COUNTER(STAT_SceneCulling_Update_Pre);
-			BUILDER_LOG_SCOPE("OnPreSceneUpdate: %d", ScenePreUpdateData.RemovedPrimitiveIds.Num());
+			BUILDER_LOG_SCOPE("OnPreSceneUpdate: %d", ChangeSet.RemovedPrimitiveIds.Num());
 			CSV_SCOPED_TIMING_STAT(SceneCulling, PreSceneUpdate);
 #if DO_CHECK
 			check(DebugTaskCounter++ == 0);
@@ -2665,7 +2712,7 @@ void FSceneCulling::FUpdater::PostSceneUpdate(FRDGBuilder& GraphBuilder, const F
 	{
 		return;
 	}
-	SC_DETAILED_LOGGING_SCOPE(SceneCulling.ActiveUpdaterImplementation);
+	SC_DETAILED_LOGGING_SCOPE(SceneCulling.ActiveUpdaterImplementation.Get());
 
 #if SC_ALLOW_ASYNC_TASKS
 	check(SceneCulling.PostUpdateTaskHandle.IsCompleted());
@@ -2694,7 +2741,7 @@ void FSceneCulling::FinalizeUpdateAndClear(FRDGBuilder& GraphBuilder, FSceneUnif
 
 	if (ActiveUpdaterImplementation.IsValid())
 	{
-		SC_DETAILED_LOGGING_SCOPE(ActiveUpdaterImplementation);
+		SC_DETAILED_LOGGING_SCOPE(ActiveUpdaterImplementation.Get());
 		SCOPED_NAMED_EVENT(SceneCulling_Update_FinalizeAndClear, FColor::Emerald);
 		SCOPE_CYCLE_COUNTER(STAT_SceneCulling_Update_FinalizeAndClear);
 		//BUILDER_LOG_SCOPE("FinalizeAndClear %d", 1);
diff --git a/Engine/Source/Runtime/Renderer/Private/SceneCulling/SceneCulling.h b/Engine/Source/Runtime/Renderer/Private/SceneCulling/SceneCulling.h
--- a/Engine/Source/Runtime/Renderer/Private/SceneCulling/SceneCulling.h
+++ b/Engine/Source/Runtime/Renderer/Private/SceneCulling/SceneCulling.h
@@ -40,9 +40,6 @@ class FSceneCulling : public ISceneExtension
 {
 	DECLARE_SCENE_EXTENSION(RENDERER_API, FSceneCulling);
 public:
-	static constexpr uint32 InvalidCellFlag = 1U << 31;
-	static constexpr uint32 TempCellFlag = 1U << 30;
-
 	friend class FSceneCullingRenderer;
 
 	FSceneCulling(FScene& InScene);
diff --git a/Engine/Source/Runtime/Renderer/Private/SceneCulling/SceneCulling.inl b/Engine/Source/Runtime/Renderer/Private/SceneCulling/SceneCulling.inl
--- a/Engine/Source/Runtime/Renderer/Private/SceneCulling/SceneCulling.inl
+++ b/Engine/Source/Runtime/Renderer/Private/SceneCulling/SceneCulling.inl
@@ -114,9 +114,29 @@ void FSceneCulling::TestSphere(const FSphere& Sphere, ResultConsumerType& Result
 	}
 }
 
-// works for both packed and not
-template <typename CellHeaderType>
-inline bool IsValidCell(const CellHeaderType& CellHeader)
+inline bool IsValidCell(const FCellHeader& CellHeader)
 {
-	return (CellHeader.ItemChunksOffset & FSceneCulling::InvalidCellFlag) == 0;
+	return CellHeader.bIsValid;
+}
+
+inline bool IsValidCell(const FPackedCellHeader& CellHeader)
+{
+	// For a valid cell the value is always nonzero
+	return CellHeader.Packed1 != 0u;
+}
+
+inline FCellHeader UnpackCellHeader(const FPackedCellHeader& Packed)
+{
+	const uint64& Bits = reinterpret_cast<const uint64&>(Packed);
+	FCellHeader CellHeader;
+	CellHeader.ItemChunksOffset = uint32(Bits >> (2u * INSTANCE_HIERARCHY_CELL_HEADER_COUNT_BITS)) & ((1u << INSTANCE_HIERARCHY_CELL_HEADER_OFFSET_BITS) - 1u);
+	CellHeader.NumStaticChunks  = uint32(Bits >> INSTANCE_HIERARCHY_CELL_HEADER_COUNT_BITS) & ((1u << INSTANCE_HIERARCHY_CELL_HEADER_COUNT_BITS) - 1u);
+	CellHeader.NumDynamicChunks  = uint32(Bits) & ((1u << INSTANCE_HIERARCHY_CELL_HEADER_COUNT_BITS) - 1u);
+	CellHeader.bIsValid = CellHeader.NumDynamicChunks != 0u;
+	if (CellHeader.bIsValid)
+	{
+		CellHeader.NumDynamicChunks -= 1u;
+	}
+	CellHeader.NumItemChunks = CellHeader.NumDynamicChunks + CellHeader.NumStaticChunks;
+	return CellHeader;
 }
diff --git a/Engine/Source/Runtime/Renderer/Private/SceneCulling/SceneCullingRenderer.cpp b/Engine/Source/Runtime/Renderer/Private/SceneCulling/SceneCullingRenderer.cpp
--- a/Engine/Source/Runtime/Renderer/Private/SceneCulling/SceneCullingRenderer.cpp
+++ b/Engine/Source/Runtime/Renderer/Private/SceneCulling/SceneCullingRenderer.cpp
@@ -143,7 +143,7 @@ void FSceneCullingRenderer::DebugRender(FRDGBuilder& GraphBuilder, TArrayView<FV
 		TBitArray<> ValidCellsMask(false, MaxCellCount);
 		for (int32 Index = 0; Index < MaxCellCount; ++Index)
 		{
-			ValidCellsMask[Index] =  (SceneCulling.CellHeaders[Index].ItemChunksOffset & FSceneCulling::InvalidCellFlag) == 0; 
+			ValidCellsMask[Index] =  IsValidCell(SceneCulling.CellHeaders[Index]); 
 		}
 		FRDGBuffer* ValidCellsMaskRdg = CreateStructuredBuffer(GraphBuilder, TEXT("SceneCulling.Debug.ValidCellsMaskRdg"), TConstArrayView<uint32>(ValidCellsMask.GetData(), FBitSet::CalculateNumWords(ValidCellsMask.Num())));
 
