[HWRT] Replace checks in WriteInstanceDescriptor(...) with ensures.
- don't need to add onscree message or logs because GPU Scene has the same InstanceId limit and already does so.
- we could add better logging for InstanceContributionToHitGroupIndex being exceeded but seems unlikely this limit is hit at the moment (should require millions of primitives/materials).

diff --git a/Engine/Source/Runtime/Renderer/Private/RayTracing/RayTracingInstanceBufferUtil.cpp b/Engine/Source/Runtime/Renderer/Private/RayTracing/RayTracingInstanceBufferUtil.cpp
--- a/Engine/Source/Runtime/Renderer/Private/RayTracing/RayTracingInstanceBufferUtil.cpp
+++ b/Engine/Source/Runtime/Renderer/Private/RayTracing/RayTracingInstanceBufferUtil.cpp
@@ -269,8 +269,8 @@ void WriteInstanceDescriptor(
 	InstanceDesc.InstanceContributionToHitGroupIndex = SceneInstance.InstanceContributionToHitGroupIndex;
 	InstanceDesc.SceneInstanceIndexAndApplyLocalBoundsTransform = (SceneInstance.bApplyLocalBoundsTransform ? 0x80000000 : 0) | SceneInstanceIndex;
 
-	checkf(InstanceDesc.InstanceId <= 0xFFFFFF, TEXT("InstanceId must fit in 24 bits."));
-	checkf(InstanceDesc.InstanceContributionToHitGroupIndex <= 0xFFFFFF, TEXT("InstanceContributionToHitGroupIndex must fit in 24 bits."));
+	ensureMsgf(InstanceDesc.InstanceId <= 0xFFFFFF, TEXT("InstanceId must fit in 24 bits."));
+	ensureMsgf(InstanceDesc.InstanceContributionToHitGroupIndex <= 0xFFFFFF, TEXT("InstanceContributionToHitGroupIndex must fit in 24 bits."));
 
 	// copy at the end to avoid reading from OutInstanceDescriptor in the checks above
 	OutInstanceDescriptor = InstanceDesc;
