VolumetricCloud - fix for missing fog in reflection when using late fog application on clouds.

#tests pc editor dx12
#rnx

diff --git a/Engine/Source/Runtime/Renderer/Private/VolumetricCloudRendering.cpp b/Engine/Source/Runtime/Renderer/Private/VolumetricCloudRendering.cpp
--- a/Engine/Source/Runtime/Renderer/Private/VolumetricCloudRendering.cpp
+++ b/Engine/Source/Runtime/Renderer/Private/VolumetricCloudRendering.cpp
@@ -667,7 +667,7 @@ uint32 GetVolumetricCloudDebugViewMode(const FEngineShowFlags& ShowFlags)
 static float GetEmptySpaceSkippingSliceDepth();
 
 // When calling this, you still need to setup Light0Shadow yourself.
-void SetupDefaultRenderVolumetricCloudGlobalParameters(FRDGBuilder& GraphBuilder, FRenderVolumetricCloudGlobalParameters& VolumetricCloudParams, FVolumetricCloudRenderSceneInfo& CloudInfo, FViewInfo& ViewInfo)
+void SetupDefaultRenderVolumetricCloudGlobalParameters(FRDGBuilder& GraphBuilder, FRenderVolumetricCloudGlobalParameters& VolumetricCloudParams, FVolumetricCloudRenderSceneInfo& CloudInfo, FViewInfo& ViewInfo, bool bForceNonLateFogOnCloud = false)
 {
 	FRDGTextureRef BlackDummy = GSystemTextures.GetBlackDummy(GraphBuilder);
 	VolumetricCloudParams.VolumetricCloud = CloudInfo.GetVolumetricCloudCommonShaderParameters();
@@ -692,7 +692,7 @@ void SetupDefaultRenderVolumetricCloudGlobalParameters(FRDGBuilder& GraphBuilder
 
 	VolumetricCloudParams.EnableHeightFog = ViewInfo.Family->Scene->HasAnyExponentialHeightFog() && ShouldRenderFog(*ViewInfo.Family);
 	SetupFogUniformParameters(GraphBuilder, ViewInfo, VolumetricCloudParams.FogStruct);
-	if (ShouldVolumetricCloudsApplyFogDuringReconstruction(ViewInfo))
+	if (ShouldVolumetricCloudsApplyFogDuringReconstruction(ViewInfo) && !bForceNonLateFogOnCloud)
 	{
 		VolumetricCloudParams.FogStruct.ApplyVolumetricFog = 0;
 		VolumetricCloudParams.EnableHeightFog = 0;
@@ -2202,7 +2202,8 @@ static TRDGUniformBufferRef<FRenderVolumetricCloudGlobalParameters> CreateCloudP
 	FVolumetricCloudRenderSceneInfo& CloudInfo = *CloudRC.CloudInfo;
 
 	FRenderVolumetricCloudGlobalParameters& VolumetricCloudParams = *GraphBuilder.AllocParameters<FRenderVolumetricCloudGlobalParameters>();
-	SetupDefaultRenderVolumetricCloudGlobalParameters(GraphBuilder, VolumetricCloudParams, CloudInfo, MainView);
+	const bool bForceNonLateFogOnCloud = CloudRC.bIsSkyRealTimeReflectionRendering;
+	SetupDefaultRenderVolumetricCloudGlobalParameters(GraphBuilder, VolumetricCloudParams, CloudInfo, MainView, bForceNonLateFogOnCloud);
 
 	VolumetricCloudParams.SceneDepthTexture = CloudRC.SceneDepthZ;
 	VolumetricCloudParams.SceneDepthMinAndMaxTexture = CloudRC.SceneDepthMinAndMax ? CloudRC.SceneDepthMinAndMax : VolumetricCloudParams.SceneDepthMinAndMaxTexture;
