Real time sky light capture fog rendering improvement
- correctly integrate fog w.r.t. skylight position (and not main view camera position).
- mobile works in forward and deferred without assert
- setup mobile correct depth stencil state
- made DeviceZ safe when capturing fog

#rnx
#tests Pc editor dx12 / mobile

diff --git a/Engine/Shaders/Private/ReflectionEnvironmentShaders.usf b/Engine/Shaders/Private/ReflectionEnvironmentShaders.usf
--- a/Engine/Shaders/Private/ReflectionEnvironmentShaders.usf
+++ b/Engine/Shaders/Private/ReflectionEnvironmentShaders.usf
@@ -892,6 +892,8 @@ void ComputeSkyEnvMapDiffuseIrradianceCS(uint3 ThreadId : SV_DispatchThreadID)
 
 #include "HeightFogCommon.ush"
 
+float3 SkyLightPosition;
+
 #if PERMUTATION_DEPTHTEXTURE
 Texture2D DepthTexture;
 #endif
@@ -919,11 +921,25 @@ void RenderRealTimeReflectionHeightFogPS(
 
 #if PERMUTATION_DEPTHTEXTURE
 	float DeviceZ = DepthTexture.Load(int3(SVPos.xy, 0)).x;
-	const float HeightFogDistance = ConvertFromDeviceZ(DeviceZ);
 #else
-	const float HeightFogDistance = 100000.0f * 10.0f;	// 10 kilometers
+	float DeviceZ = FarDepthValue;
+#endif
+
+#if HAS_INVERTED_Z_BUFFER
+	DeviceZ = max(1.0e-10, DeviceZ); // For ConvertFromDeviceZ to return a valid value (non inf).
 #endif
-	float4 HeightFogInscatteringAndOpacity = CalculateHeightFog(normalize(ScreenVector)* HeightFogDistance);
+
+	float SceneDepth = ConvertFromDeviceZ(DeviceZ);
+	float3 WorldPositionRelativeToCamera = ScreenVector.xyz * SceneDepth;
+
+	// Evaluate the height fog the actual sky light position to match the capture position.
+	const float ExcludeDistance = 0;
+	const float EyeIndex = 0;
+	const bool bOverrideOrigin = true;
+	const float3 OverrideRayStartRelativeToOrigin = SkyLightPosition;
+	const float OverrideRayLength = length(WorldPositionRelativeToCamera);
+	const float3 OverrideRayDir = WorldPositionRelativeToCamera / max(0.000001, OverrideRayLength);
+	float4 HeightFogInscatteringAndOpacity = GetExponentialHeightFog(WorldPositionRelativeToCamera, ExcludeDistance, EyeIndex, PrimaryView, bOverrideOrigin, OverrideRayStartRelativeToOrigin, OverrideRayDir, OverrideRayLength);
 
 	const float OutputPreExposure = (ResolvedView.RealTimeReflectionCapture ? ResolvedView.RealTimeReflectionCapturePreExposure : ResolvedView.PreExposure);
 	HeightFogInscatteringAndOpacity.rgb *= OutputPreExposure;
diff --git a/Engine/Source/Runtime/Renderer/Private/ReflectionEnvironmentRealTimeCapture.cpp b/Engine/Source/Runtime/Renderer/Private/ReflectionEnvironmentRealTimeCapture.cpp
--- a/Engine/Source/Runtime/Renderer/Private/ReflectionEnvironmentRealTimeCapture.cpp
+++ b/Engine/Source/Runtime/Renderer/Private/ReflectionEnvironmentRealTimeCapture.cpp
@@ -256,6 +256,7 @@ class FRenderRealTimeReflectionHeightFogPS : public FGlobalShader
 		SHADER_PARAMETER_STRUCT_REF(FViewUniformShaderParameters, ViewUniformBuffer)
 		SHADER_PARAMETER_RDG_UNIFORM_BUFFER(FFogUniformParameters, FogStruct)
 		SHADER_PARAMETER_RDG_TEXTURE(Texture2D, DepthTexture)
+		SHADER_PARAMETER(FVector3f, SkyLightPosition)
 		RENDER_TARGET_BINDING_SLOTS()
 	END_SHADER_PARAMETER_STRUCT()
 
@@ -710,6 +711,10 @@ void FScene::AllocateAndCaptureFrameSkyEnvMap(
 				SkyRC.bShouldSampleCloudSkyAO = CloudShadowAOData.bShouldSampleCloudSkyAO;
 				SkyRC.VolumetricCloudSkyAO = CloudShadowAOData.VolumetricCloudSkyAO;
 
+				// Note: 
+				//	- The depth texture is here so that the ordering of IsSky material meshes is correct when rendering the sky into cube map.
+				//	- This depth is also used to apply distance based fog (usefull when one want to capture a terrain at the bottom of the skylight).
+				//	- If that behavior is not needed, use r.SkyLight.RealTimeReflectionCapture.DepthBuffer 0.
 				const bool bUseDepthBuffer = CVarRealTimeReflectionCaptureDepthBuffer.GetValueOnRenderThread() > 0;
 				FRDGTextureRef CubeDepthTexture = nullptr;
 
@@ -732,11 +737,11 @@ void FScene::AllocateAndCaptureFrameSkyEnvMap(
 								auto* PassParameters = GraphBuilder.AllocParameters<FMobileCaptureSkyMeshReflectionPassParameters>();
 								PassParameters->View = CubeView.GetShaderParameters();
 								PassParameters->RenderTargets = SkyRC.RenderTargets;
-								const EMobileBasePass BasePass = EMobileBasePass::DepthPrePass;
+								const EMobileBasePass BasePass = EMobileBasePass::Opaque;
 								PassParameters->BasePass = CreateMobileBasePassUniformBuffer(GraphBuilder, MainView, BasePass, EMobileSceneTextureSetupMode::SceneDepth, {}, true);
 					
-								// Setup the depth buffer
-								if (bUseDepthBuffer)
+								// Setup the depth buffer, always allocate on mobile since the basepass is asking for stencil operations.
+								//if (bUseDepthBuffer)
 								{
 									FRDGTextureDesc CubeDepthTextureDesc = FRDGTextureDesc::Create2D(FIntPoint(CubeWidth, CubeWidth), PF_DepthStencil,
 										MainView.GetSceneTexturesConfig().DepthClearValue,
@@ -758,7 +763,7 @@ void FScene::AllocateAndCaptureFrameSkyEnvMap(
 									}
 									else
 									{
-										DrawRenderState.SetDepthStencilAccess(FExclusiveDepthStencil::DepthNop_StencilWrite);
+										DrawRenderState.SetDepthStencilAccess(FExclusiveDepthStencil::DepthWrite_StencilNop);
 										DrawRenderState.SetDepthStencilState(TStaticDepthStencilState<false, CF_Always>::GetRHI());
 									}
 					
@@ -824,6 +829,18 @@ void FScene::AllocateAndCaptureFrameSkyEnvMap(
 									FExclusiveDepthStencil::Type BasePassDepthStencilAccess_Sky = bUseDepthBuffer ? FExclusiveDepthStencil::Type(Scene->DefaultBasePassDepthStencilAccess | FExclusiveDepthStencil::DepthWrite)
 										: FExclusiveDepthStencil::Type(Scene->DefaultBasePassDepthStencilAccess & ~FExclusiveDepthStencil::DepthWrite);
 									SetupBasePassState(BasePassDepthStencilAccess_Sky, false, DrawRenderState);
+
+									DrawRenderState.SetBlendState(TStaticBlendStateWriteMask<CW_RGBA>::GetRHI());
+									if (bUseDepthBuffer)
+									{
+										DrawRenderState.SetDepthStencilAccess(FExclusiveDepthStencil::DepthWrite_StencilNop);
+										DrawRenderState.SetDepthStencilState(TStaticDepthStencilState<true, CF_DepthNearOrEqual>::GetRHI());
+									}
+									else
+									{
+										DrawRenderState.SetDepthStencilAccess(FExclusiveDepthStencil::DepthNop_StencilNop);
+										DrawRenderState.SetDepthStencilState(TStaticDepthStencilState<false, CF_Always>::GetRHI());
+									}
 					
 									FSkyPassMeshProcessor PassMeshProcessor(Scene, Scene->GetFeatureLevel(), nullptr, DrawRenderState, DynamicMeshPassContext);
 									const int32 SkyRealTimeReflectionOnlyMeshBatcheCount = MainView.SkyMeshBatches.Num();
@@ -878,14 +895,15 @@ void FScene::AllocateAndCaptureFrameSkyEnvMap(
 							TShaderMapRef<FRenderRealTimeReflectionHeightFogVS> VertexShader(GetGlobalShaderMap(SkyRC.FeatureLevel), VsPermutationVector);
 
 							FRenderRealTimeReflectionHeightFogPS::FPermutationDomain PsPermutationVector;
-							PsPermutationVector.Set<FRenderRealTimeReflectionHeightFogPS::FDepthTexture>(CubeDepthTexture != nullptr);
+							PsPermutationVector.Set<FRenderRealTimeReflectionHeightFogPS::FDepthTexture>(bUseDepthBuffer);
 							TShaderMapRef<FRenderRealTimeReflectionHeightFogPS> PixelShader(GetGlobalShaderMap(SkyRC.FeatureLevel), PsPermutationVector);
 
 							FRenderRealTimeReflectionHeightFogPS::FParameters* PsPassParameters = GraphBuilder.AllocParameters<FRenderRealTimeReflectionHeightFogPS::FParameters>();
 							PsPassParameters->ViewUniformBuffer = CubeView.ViewUniformBuffer;
 							PsPassParameters->RenderTargets = SkyRC.RenderTargets;
-							PsPassParameters->DepthTexture = CubeDepthTexture != nullptr ? CubeDepthTexture : BlackDummy2dTex;
+							PsPassParameters->DepthTexture = bUseDepthBuffer ? CubeDepthTexture : BlackDummy2dTex;
 							PsPassParameters->FogStruct = CreateFogUniformBuffer(GraphBuilder, CubeView);
+							PsPassParameters->SkyLightPosition = FVector3f(SkyLightProxy->CapturePosition);
 
 							ClearUnusedGraphResources(PixelShader, PsPassParameters);
 
