Fix for a crash in real time sky light capture.

#rnx
#tests pc editor dx12 without skydome mesh

diff --git a/Engine/Source/Runtime/Renderer/Private/ReflectionEnvironmentRealTimeCapture.cpp b/Engine/Source/Runtime/Renderer/Private/ReflectionEnvironmentRealTimeCapture.cpp
--- a/Engine/Source/Runtime/Renderer/Private/ReflectionEnvironmentRealTimeCapture.cpp
+++ b/Engine/Source/Runtime/Renderer/Private/ReflectionEnvironmentRealTimeCapture.cpp
@@ -895,13 +895,13 @@ void FScene::AllocateAndCaptureFrameSkyEnvMap(
 							TShaderMapRef<FRenderRealTimeReflectionHeightFogVS> VertexShader(GetGlobalShaderMap(SkyRC.FeatureLevel), VsPermutationVector);
 
 							FRenderRealTimeReflectionHeightFogPS::FPermutationDomain PsPermutationVector;
-							PsPermutationVector.Set<FRenderRealTimeReflectionHeightFogPS::FDepthTexture>(bUseDepthBuffer);
+							PsPermutationVector.Set<FRenderRealTimeReflectionHeightFogPS::FDepthTexture>(bUseDepthBuffer && CubeDepthTexture);
 							TShaderMapRef<FRenderRealTimeReflectionHeightFogPS> PixelShader(GetGlobalShaderMap(SkyRC.FeatureLevel), PsPermutationVector);
 
 							FRenderRealTimeReflectionHeightFogPS::FParameters* PsPassParameters = GraphBuilder.AllocParameters<FRenderRealTimeReflectionHeightFogPS::FParameters>();
 							PsPassParameters->ViewUniformBuffer = CubeView.ViewUniformBuffer;
 							PsPassParameters->RenderTargets = SkyRC.RenderTargets;
-							PsPassParameters->DepthTexture = bUseDepthBuffer ? CubeDepthTexture : BlackDummy2dTex;
+							PsPassParameters->DepthTexture = (bUseDepthBuffer && CubeDepthTexture) ? CubeDepthTexture : BlackDummy2dTex;
 							PsPassParameters->FogStruct = CreateFogUniformBuffer(GraphBuilder, CubeView);
 							PsPassParameters->SkyLightPosition = FVector3f(SkyLightProxy->CapturePosition);
 
