Sky light real time capture mobile improvement
- Fixed mobile not writting depth, resulting in Fog difference when r.SkyLight.RealTimeReflectionCapture.DepthBuffer=0.
- Depth buffer is now allocated on mobile only when needed as it was before.
- Capture is now cleared to black when the skypas is empty.

#tests PC editor dx12 + mobile preview
#rnx

diff --git a/Engine/Source/Runtime/Renderer/Private/ReflectionEnvironmentRealTimeCapture.cpp b/Engine/Source/Runtime/Renderer/Private/ReflectionEnvironmentRealTimeCapture.cpp
--- a/Engine/Source/Runtime/Renderer/Private/ReflectionEnvironmentRealTimeCapture.cpp
+++ b/Engine/Source/Runtime/Renderer/Private/ReflectionEnvironmentRealTimeCapture.cpp
@@ -740,34 +740,40 @@ void FScene::AllocateAndCaptureFrameSkyEnvMap(
 								const EMobileBasePass BasePass = EMobileBasePass::Opaque;
 								PassParameters->BasePass = CreateMobileBasePassUniformBuffer(GraphBuilder, MainView, BasePass, EMobileSceneTextureSetupMode::SceneDepth, {}, true);
 					
-								// Setup the depth buffer, always allocate on mobile since the basepass is asking for stencil operations.
-								//if (bUseDepthBuffer)
+								if (bUseDepthBuffer)
 								{
 									FRDGTextureDesc CubeDepthTextureDesc = FRDGTextureDesc::Create2D(FIntPoint(CubeWidth, CubeWidth), PF_DepthStencil,
 										MainView.GetSceneTexturesConfig().DepthClearValue,
 										TexCreate_DepthStencilTargetable | TexCreate_ShaderResource);
 									CubeDepthTexture = GraphBuilder.CreateTexture(CubeDepthTextureDesc, TEXT("SkyLight.CubeDepthTexture"));
-									PassParameters->RenderTargets.DepthStencil = FDepthStencilBinding(CubeDepthTexture, ERenderTargetLoadAction::EClear, FExclusiveDepthStencil::DepthWrite_StencilWrite);
+									PassParameters->RenderTargets.DepthStencil = FDepthStencilBinding(CubeDepthTexture, ERenderTargetLoadAction::EClear, FExclusiveDepthStencil::DepthWrite_StencilNop);
 								}
 					
 								AddSimpleMeshPass(GraphBuilder, PassParameters, Scene, MainView, &InstanceCullingManager, RDG_EVENT_NAME("CaptureSkyMeshReflection"), SkyRC.Viewport,
 								[&MainView, CubeViewUniformBuffer = CubeView.ViewUniformBuffer, bUseDepthBuffer, Scene](FDynamicPassMeshDrawListContext* DynamicMeshPassContext)
 								{
 									FMeshPassProcessorRenderState DrawRenderState;
+									FSkyPassMeshProcessor::ESkyPassType SkyPassType = FSkyPassMeshProcessor::ESkyPassType::SPT_Default;
 					
 									DrawRenderState.SetBlendState(TStaticBlendStateWriteMask<CW_RGBA>::GetRHI());
 									if (bUseDepthBuffer)
 									{
-										DrawRenderState.SetDepthStencilAccess(FExclusiveDepthStencil::DepthWrite_StencilWrite);
+										DrawRenderState.SetDepthStencilAccess(FExclusiveDepthStencil::DepthWrite_StencilNop);
 										DrawRenderState.SetDepthStencilState(TStaticDepthStencilState<true, CF_DepthNearOrEqual>::GetRHI());
+
+										SkyPassType = FSkyPassMeshProcessor::ESkyPassType::SPT_RealTimeCapture_DepthWrite;
 									}
 									else
 									{
-										DrawRenderState.SetDepthStencilAccess(FExclusiveDepthStencil::DepthWrite_StencilNop);
+										DrawRenderState.SetDepthStencilAccess(FExclusiveDepthStencil::DepthNop_StencilNop);
 										DrawRenderState.SetDepthStencilState(TStaticDepthStencilState<false, CF_Always>::GetRHI());
+
+										SkyPassType = FSkyPassMeshProcessor::ESkyPassType::SPT_RealTimeCapture_DepthNop;
 									}
 					
 									FSkyPassMeshProcessor PassMeshProcessor(Scene, Scene->GetFeatureLevel(), nullptr, DrawRenderState, DynamicMeshPassContext);
+									PassMeshProcessor.SkyPassType = SkyPassType;
+
 									const int32 SkyRealTimeReflectionOnlyMeshBatcheCount = MainView.SkyMeshBatches.Num();
 									for (int32 MeshBatchIndex = 0; MeshBatchIndex < SkyRealTimeReflectionOnlyMeshBatcheCount; ++MeshBatchIndex)
 									{
@@ -879,13 +885,18 @@ void FScene::AllocateAndCaptureFrameSkyEnvMap(
 								});
 								}
 						}
-						else if(!bIsMobilePlatform)	// The SkyAtmosphere full screen pass shader is not compiled on mobile, a sky dome is required.
+						else if (!bIsMobilePlatform)
 						{
 							// TODO: mobile should never get there
 							RDG_EVENT_SCOPE(GraphBuilder, "Capture Sky Raw", CubeFace);
 							FSceneTextureShaderParameters SceneTextures = CreateSceneTextureShaderParameters(GraphBuilder, &SceneRenderer.GetActiveSceneTextures(), SceneRenderer.FeatureLevel, ESceneTextureSetupMode::SceneDepth);
 							SceneRenderer.RenderSkyAtmosphereInternal(GraphBuilder, SceneTextures, SkyRC);
 						}
+						else if (bIsMobilePlatform)	// The SkyAtmosphere full screen pass shader is not compiled on mobile, so clear the capture to black.
+						{
+							// No sky dome mesh and no sky rendering on mobile so let's cleat to black.
+							ClearCubeFace(SkyCubeTexture, CubeFace);
+						}
 
 						// Also render the height fog as part of the sky render pass when time slicing is enabled.
 						
diff --git a/Engine/Source/Runtime/Renderer/Private/SkyPassRendering.cpp b/Engine/Source/Runtime/Renderer/Private/SkyPassRendering.cpp
--- a/Engine/Source/Runtime/Renderer/Private/SkyPassRendering.cpp
+++ b/Engine/Source/Runtime/Renderer/Private/SkyPassRendering.cpp
@@ -125,14 +125,34 @@ bool FSkyPassMeshProcessor::Process(
 			return false;
 		}
 
-		// Mask sky pixels so we can skip them when rendering per-pixel fog
-		PassDrawRenderState.SetDepthStencilState(TStaticDepthStencilState<
-			false, CF_DepthNearOrEqual,
-			true, CF_Always, SO_Keep, SO_Keep, SO_Replace,
-			false, CF_Always, SO_Keep, SO_Keep, SO_Keep,
-			0x00, STENCIL_MOBILE_SKY_MASK>::GetRHI());
-		
-		PassDrawRenderState.SetStencilRef(STENCIL_MOBILE_SKY_MASK); 
+		if (SkyPassType == SPT_Default)
+		{
+			// Mask sky pixels so we can skip them when rendering per-pixel fog (USE_VERTEX_FOG is applying for on sky pixels).
+			PassDrawRenderState.SetDepthStencilState(TStaticDepthStencilState<
+				false, CF_DepthNearOrEqual,
+				true, CF_Always, SO_Keep, SO_Keep, SO_Replace,
+				false, CF_Always, SO_Keep, SO_Keep, SO_Keep,
+				0x00, STENCIL_MOBILE_SKY_MASK>::GetRHI());
+
+			PassDrawRenderState.SetStencilRef(STENCIL_MOBILE_SKY_MASK);
+		}
+		else if (SkyPassType == SPT_RealTimeCapture_DepthWrite)
+		{
+			// Capturing real time sky light, writing depth only.
+			PassDrawRenderState.SetDepthStencilState(TStaticDepthStencilState<
+				true, CF_DepthNearOrEqual>::GetRHI());
+		}
+		else if (SkyPassType == SPT_RealTimeCapture_DepthNop)
+		{
+
+			// Capturing real time sky light, no depth writes.
+			PassDrawRenderState.SetDepthStencilState(TStaticDepthStencilState<
+				false, CF_Always>::GetRHI());
+		}
+		else
+		{
+			check(false);
+		}
 		
 		TMobileBasePassShaderElementData<LightMapPolicyType> ShaderElementData(nullptr, false);
 		ShaderElementData.InitializeMeshMaterialData(ViewIfDynamicMeshCommand, PrimitiveSceneProxy, MeshBatch, StaticMeshId, false);
diff --git a/Engine/Source/Runtime/Renderer/Private/SkyPassRendering.h b/Engine/Source/Runtime/Renderer/Private/SkyPassRendering.h
--- a/Engine/Source/Runtime/Renderer/Private/SkyPassRendering.h
+++ b/Engine/Source/Runtime/Renderer/Private/SkyPassRendering.h
@@ -26,6 +26,14 @@ public:
 	FMeshPassProcessorRenderState PassDrawRenderState;
 	EShadingPath ShadingPath;
 
+	enum ESkyPassType
+	{
+		SPT_Default,
+		SPT_RealTimeCapture_DepthWrite,
+		SPT_RealTimeCapture_DepthNop,
+	};
+	ESkyPassType SkyPassType = SPT_Default;
+
 private:
 	bool TryAddMeshBatch(
 		const FMeshBatch& RESTRICT MeshBatch,
