[First Person] Removed "Cast First Person Self Shadow" checkbox on light component. Self-Shadow is now driven by the light being shadow casting (dynamic shadows) and r.FirstPerson.SelfShadow and r.FirstPerson.SelfShadow.LightTypes. The first cvar enables the effect at all, the second cvar can be used to have it affect only directional lights, only local lights or all lights.

diff --git a/Engine/Source/Runtime/Engine/Classes/Components/LightComponent.h b/Engine/Source/Runtime/Engine/Classes/Components/LightComponent.h
--- a/Engine/Source/Runtime/Engine/Classes/Components/LightComponent.h
+++ b/Engine/Source/Runtime/Engine/Classes/Components/LightComponent.h
@@ -134,10 +134,6 @@ class ULightComponent : public ULightComponentBase
 	UPROPERTY(EditAnywhere, BlueprintReadWrite, Category = Light, AdvancedDisplay, meta = (ClampMin = 0.0, ClampMax = 1.0))
 	float ContactShadowNonCastingIntensity;
 
-	/** Whether the light is allowed to cast self shadows on first person primitives. */
-	UPROPERTY(EditAnywhere, BlueprintReadOnly, Category = Light, AdvancedDisplay)
-	uint32 bCastFirstPersonSelfShadow : 1;
-
 	UPROPERTY()
 	uint32 InverseSquaredFalloff_DEPRECATED:1;
 
diff --git a/Engine/Source/Runtime/Engine/Private/Components/LightComponent.cpp b/Engine/Source/Runtime/Engine/Private/Components/LightComponent.cpp
--- a/Engine/Source/Runtime/Engine/Private/Components/LightComponent.cpp
+++ b/Engine/Source/Runtime/Engine/Private/Components/LightComponent.cpp
@@ -461,7 +461,6 @@ ULightComponent::ULightComponent(const FObjectInitializer& ObjectInitializer)
 	ContactShadowLengthInWS = false;
 	ContactShadowCastingIntensity = 1.0f;
 	ContactShadowNonCastingIntensity = 0.0f;
-	bCastFirstPersonSelfShadow = false;
 	bAllowMegaLights = true;
 	MegaLightsShadowMethod = EMegaLightsShadowMethod::Default;
 	bUseIESBrightness = false;
@@ -777,7 +776,6 @@ void ULightComponent::PostEditChangeProperty(FPropertyChangedEvent& PropertyChan
 		PropertyName != GET_MEMBER_NAME_STRING_CHECKED(ULightComponent, ContactShadowCastingIntensity) &&
 		PropertyName != GET_MEMBER_NAME_STRING_CHECKED(ULightComponent, ContactShadowNonCastingIntensity) &&
 		PropertyName != GET_MEMBER_NAME_STRING_CHECKED(ULightComponent, ContactShadowLengthInWS) &&
-		PropertyName != GET_MEMBER_NAME_STRING_CHECKED(ULightComponent, bCastFirstPersonSelfShadow) &&
 		PropertyName != GET_MEMBER_NAME_STRING_CHECKED(ULightComponent, bEnableLightShaftBloom) &&
 		PropertyName != GET_MEMBER_NAME_STRING_CHECKED(ULightComponent, BloomScale) &&
 		PropertyName != GET_MEMBER_NAME_STRING_CHECKED(ULightComponent, BloomThreshold) &&
diff --git a/Engine/Source/Runtime/Engine/Private/LightSceneProxy.cpp b/Engine/Source/Runtime/Engine/Private/LightSceneProxy.cpp
--- a/Engine/Source/Runtime/Engine/Private/LightSceneProxy.cpp
+++ b/Engine/Source/Runtime/Engine/Private/LightSceneProxy.cpp
@@ -54,7 +54,6 @@ FLightSceneProxy::FLightSceneProxy(const ULightComponent* InLightComponent)
 	, bUseVirtualShadowMaps(false)	// See below
 	, bCastModulatedShadows(false)
 	, bUseWholeSceneCSMForMovableObjects(false)
-	, bCastFirstPersonSelfShadow(InLightComponent->bCastFirstPersonSelfShadow)
 	, bSelected(InLightComponent->IsSelected() || InLightComponent->IsOwnerSelected())
 	, bAllowMegaLights(InLightComponent->bAllowMegaLights)
 	, MegaLightsShadowMethod(InLightComponent->MegaLightsShadowMethod)
diff --git a/Engine/Source/Runtime/Engine/Public/LightSceneProxy.h b/Engine/Source/Runtime/Engine/Public/LightSceneProxy.h
--- a/Engine/Source/Runtime/Engine/Public/LightSceneProxy.h
+++ b/Engine/Source/Runtime/Engine/Public/LightSceneProxy.h
@@ -216,7 +216,6 @@ public:
 	inline bool CastsTranslucentShadows() const { return bCastTranslucentShadows; }
 	inline bool CastsVolumetricShadow() const { return bCastVolumetricShadow; }
 	inline bool CastsHairStrandsDeepShadow() const { return bCastHairStrandsDeepShadow; }
-	inline bool CastsFirstPersonSelfShadow() const { return bCastFirstPersonSelfShadow; }
 	inline TEnumAsByte<ECastRayTracedShadow::Type> CastsRaytracedShadow() const { return CastRaytracedShadow; }
 	inline bool AffectReflection() const { return bAffectReflection; }
 	inline bool AffectGlobalIllumination() const { return bAffectGlobalIllumination; }
@@ -460,9 +459,6 @@ protected:
 	/** Whether to render csm shadows for movable objects only (mobile). */
 	uint8 bUseWholeSceneCSMForMovableObjects : 1;
 
-	/** Whether to cast first person self shadows on first person primitives. */
-	uint8 bCastFirstPersonSelfShadow : 1;
-
     /** Is the light selected in the editor? */
 	uint8 bSelected : 1;
 
diff --git a/Engine/Source/Runtime/Renderer/Private/LightRendering.cpp b/Engine/Source/Runtime/Renderer/Private/LightRendering.cpp
--- a/Engine/Source/Runtime/Renderer/Private/LightRendering.cpp
+++ b/Engine/Source/Runtime/Renderer/Private/LightRendering.cpp
@@ -1285,6 +1285,8 @@ void FSceneRenderer::GatherAndSortLights(FSortedLightSetSceneInfo& OutSortedLigh
 					// Lights handled by Mega Lights
 					const bool bHandledByMegaLights = MegaLights::GetMegaLightsMode(ViewFamily, LightSceneInfoCompact.LightType, LightSceneInfoCompact.bAllowMegaLights, LightSceneInfoCompact.MegaLightsShadowMethod) != EMegaLightsMode::Disabled;
 
+					const bool bCastsFirstPersonSelfShadow = LightCastsFirstPersonSelfShadow(*LightSceneInfo);
+
 					// NOTE: bClusteredDeferredSupported==false means "lights cannot be batched" (tiled or clustered). When false, light will go the slower unbatched render path.
 					// Tiled and clustered deferred lighting only support certain lights that don't use any additional features (like shadow or light function not compatible with the atlas.)
 					// And also that are not directional (mostly because it doesn't make so much sense to insert them into every grid cell in the universe).
@@ -1296,14 +1298,14 @@ void FSceneRenderer::GatherAndSortLights(FSortedLightSetSceneInfo& OutSortedLigh
 						(!SortedLightInfo->SortKey.Fields.bLightFunction || (bUseLightFunctionAtlas && SortedLightInfo->bIsCompatibleWithLightFunctionAtlas)) // If not compatible with the LightFunctionAtlas, light with light function materials must go the unbatched route.
 						&& LightSceneInfoCompact.LightType != LightType_Directional
 						&& LightSceneInfo->Proxy->GetContactShadowLength() == 0
-						&& !LightSceneInfo->Proxy->CastsFirstPersonSelfShadow()
+						&& !bCastsFirstPersonSelfShadow
 						&& !bHandledByMegaLights;
 
 					// Track feature available accross all lights
 					if (SortedLightInfo->SortKey.Fields.LightType == LightType_Rect)	{ OutSortedLights.bHasRectLights = true; }
 					if (SortedLightInfo->SortKey.Fields.bUsesLightingChannels)			{ OutSortedLights.bHasLightChannels = true; }
 					if (SortedLightInfo->SortKey.Fields.bLightFunction) 				{ OutSortedLights.bHasLightFunctions = true; }
-					if (LightSceneInfo->Proxy->CastsFirstPersonSelfShadow()) 			{ OutSortedLights.bHasFirstPersonSelfShadowLights = true; }
+					if (bCastsFirstPersonSelfShadow) 									{ OutSortedLights.bHasFirstPersonSelfShadowLights = true; }
 
 					SortedLightInfo->SortKey.Fields.bClusteredDeferredNotSupported = !bClusteredDeferredSupported;
 
diff --git a/Engine/Source/Runtime/Renderer/Private/Shadows/FirstPersonSelfShadow.cpp b/Engine/Source/Runtime/Renderer/Private/Shadows/FirstPersonSelfShadow.cpp
--- a/Engine/Source/Runtime/Renderer/Private/Shadows/FirstPersonSelfShadow.cpp
+++ b/Engine/Source/Runtime/Renderer/Private/Shadows/FirstPersonSelfShadow.cpp
@@ -12,9 +12,15 @@
 
 static TAutoConsoleVariable<int32> CVarFirstPersonSelfShadow(
 	TEXT("r.FirstPerson.SelfShadow"),
-	1,
-	TEXT("Enables self shadows for first person primitives. Self shadows are achieved with HZB screen space traces."),
-	ECVF_RenderThreadSafe);
+	0,
+	TEXT("Enables self shadows for first person primitives. Self shadows are achieved with HZB screen space traces. Use r.FirstPerson.SelfShadow.LightTypes to control which shadow casting light types should cast self shadows."),
+	ECVF_RenderThreadSafe | ECVF_Scalability);
+
+static TAutoConsoleVariable<int32> CVarFirstPersonSelfShadowLightTypes(
+	TEXT("r.FirstPerson.SelfShadow.LightTypes"),
+	0,
+	TEXT("Controls which light types should cast self shadows for first person primitives. 0: Directional Lights Only, 1: Local Lights Only, 2: All Lights"),
+	ECVF_RenderThreadSafe | ECVF_Scalability);
 
 static TAutoConsoleVariable<int32> CVarFirstPersonSelfShadowCheckerboardMode(
 	TEXT("r.FirstPerson.SelfShadow.DownsampleCheckerboardMode"),
@@ -201,6 +207,17 @@ bool ShouldRenderFirstPersonSelfShadow(const FSceneViewFamily& ViewFamily)
 		&& HasFirstPersonGBufferBit(ViewFamily.GetShaderPlatform());
 }
 
+bool LightCastsFirstPersonSelfShadow(const FLightSceneInfo& LightSceneInfo)
+{
+	const bool bCastsShadow = LightSceneInfo.Proxy->CastsDynamicShadow();
+	const bool bSelfShadowEnabled = CVarFirstPersonSelfShadow.GetValueOnRenderThread() != 0;
+	const int32 EnabledSelfShadowLightTypes = FMath::Clamp(CVarFirstPersonSelfShadowLightTypes.GetValueOnRenderThread(), 0, 2);
+	const bool bLocalLightSelfShadowEnabled = EnabledSelfShadowLightTypes > 0;
+	const bool bDirectionalLightSelfShadowEnabled = EnabledSelfShadowLightTypes != 1;
+	const bool bIsDirectionalLight = LightSceneInfo.Type == LightType_Directional;
+	return bCastsShadow && bSelfShadowEnabled && ((bIsDirectionalLight && bDirectionalLightSelfShadowEnabled) || (!bIsDirectionalLight && bLocalLightSelfShadowEnabled));
+}
+
 static bool IsViewFirstPersonSelfShadowRelevant(const FViewInfo& View, const FFirstPersonViewBounds& FirstPersonViewBounds, const FLightSceneInfo& LightSceneInfo)
 {
 	// First person primitives can be expected to all be very close to one another, so lights will usually either fully affect all of them or none,
@@ -211,7 +228,7 @@ static bool IsViewFirstPersonSelfShadowRelevant(const FViewInfo& View, const FFi
 bool ShouldRenderFirstPersonSelfShadowForLight(const FSceneRendererBase& SceneRenderer, const FSceneViewFamily& ViewFamily, const TArray<FViewInfo>& Views, const FLightSceneInfo& LightSceneInfo)
 {
 	const bool bSelfShadowSupported = ShouldRenderFirstPersonSelfShadow(ViewFamily);
-	const bool bLightCastsSelfShadow = LightSceneInfo.Proxy->CastsFirstPersonSelfShadow();
+	const bool bLightCastsSelfShadow = LightCastsFirstPersonSelfShadow(LightSceneInfo);
 	if (bSelfShadowSupported && bLightCastsSelfShadow)
 	{
 		const FFirstPersonSceneExtensionRenderer* FPRenderer = SceneRenderer.GetSceneExtensionsRenderers().GetRendererPtr<FFirstPersonSceneExtensionRenderer>();
@@ -290,7 +307,7 @@ void RenderFirstPersonSelfShadow(FRDGBuilder& GraphBuilder, const FSceneRenderer
 {
 	static_assert((bool)ERHIZBuffer::IsInverted, "Inverted depth buffer is assumed because FPixelShaderUtils::AddFullscreenPass is drawing at depth 0!");
 	check(ScreenShadowMaskTexture);
-	check(LightSceneInfo.Proxy->CastsFirstPersonSelfShadow());
+	check(LightCastsFirstPersonSelfShadow(LightSceneInfo));
 
 	const FFirstPersonSceneExtensionRenderer* FPRenderer = SceneRenderer.GetSceneExtensionsRenderers().GetRendererPtr<FFirstPersonSceneExtensionRenderer>();
 	if (!ensure(FPRenderer))
diff --git a/Engine/Source/Runtime/Renderer/Private/Shadows/FirstPersonSelfShadow.h b/Engine/Source/Runtime/Renderer/Private/Shadows/FirstPersonSelfShadow.h
--- a/Engine/Source/Runtime/Renderer/Private/Shadows/FirstPersonSelfShadow.h
+++ b/Engine/Source/Runtime/Renderer/Private/Shadows/FirstPersonSelfShadow.h
@@ -26,6 +26,12 @@ struct FFirstPersonSelfShadowInputs
 // Whether to render first person self-shadow at all.
 bool ShouldRenderFirstPersonSelfShadow(const FSceneViewFamily& ViewFamily);
 
+// Whether the light could cast first person self-shadow. This is similar in spirit to calling Casts*Shadow() on the light proxy
+// and does not check for other conditions unrelated to the light itself, such as whether it is relevant for a given view or
+// if the current platform and configuration supports first person self-shadow at all. For these cases, consider calling
+// ShouldRenderFirstPersonSelfShadowForLight instead.
+bool LightCastsFirstPersonSelfShadow(const FLightSceneInfo& LightSceneInfo);
+
 // Whether to render first person self-shadow for a particular light.
 bool ShouldRenderFirstPersonSelfShadowForLight(const FSceneRendererBase& SceneRenderer, const FSceneViewFamily& ViewFamily, const TArray<FViewInfo>& Views, const FLightSceneInfo& LightSceneInfo);
 
