Extend load balancer allocations to the RDG timeline so that we can avoid copying memory.

diff --git a/Engine/Source/Runtime/Renderer/Private/InstanceCulling/InstanceCullingLoadBalancer.cpp b/Engine/Source/Runtime/Renderer/Private/InstanceCulling/InstanceCullingLoadBalancer.cpp
--- a/Engine/Source/Runtime/Renderer/Private/InstanceCulling/InstanceCullingLoadBalancer.cpp
+++ b/Engine/Source/Runtime/Renderer/Private/InstanceCulling/InstanceCullingLoadBalancer.cpp
@@ -25,9 +25,8 @@ void FInstanceCullingLoadBalancerBase::FGPUData::GetShaderParameters(FRDGBuilder
 FInstanceCullingLoadBalancerBase::FGPUData FInstanceCullingLoadBalancerBase::Upload(FRDGBuilder& GraphBuilder, TConstArrayView<FPackedBatch> Batches, TConstArrayView<FPackedItem> Items, ERDGInitialDataFlags RDGInitialDataFlags, int32 NumGroupsPerBatch) const
 {
 	FGPUData Result;
-	// TODO: Several of these load balancers are being created on the stack and the memory is going out of scope before RDG execution. Always making a copy for now.
-	Result.BatchBuffer = CreateStructuredBuffer(GraphBuilder, TEXT("InstanceCullingLoadBalancer.Batches"), Batches, /* RDGInitialDataFlags */ ERDGInitialDataFlags::None);
-	Result.ItemBuffer = CreateStructuredBuffer(GraphBuilder, TEXT("InstanceCullingLoadBalancer.Items"), Items, /* RDGInitialDataFlags */ ERDGInitialDataFlags::None);
+	Result.BatchBuffer = CreateStructuredBuffer(GraphBuilder, TEXT("InstanceCullingLoadBalancer.Batches"), Batches, RDGInitialDataFlags);
+	Result.ItemBuffer = CreateStructuredBuffer(GraphBuilder, TEXT("InstanceCullingLoadBalancer.Items"), Items, RDGInitialDataFlags);
 	Result.NumBatches = Batches.Num();
 	Result.NumItems = Items.Num();
 	Result.NumGroupsPerBatch = NumGroupsPerBatch;
diff --git a/Engine/Source/Runtime/Renderer/Private/InstanceCulling/InstanceCullingLoadBalancer.h b/Engine/Source/Runtime/Renderer/Private/InstanceCulling/InstanceCullingLoadBalancer.h
--- a/Engine/Source/Runtime/Renderer/Private/InstanceCulling/InstanceCullingLoadBalancer.h
+++ b/Engine/Source/Runtime/Renderer/Private/InstanceCulling/InstanceCullingLoadBalancer.h
@@ -10,22 +10,6 @@
 
 struct FShaderCompilerEnvironment;
 
-namespace InstanceCullingImplementationDetails
-{
-	// Helpers to make default ERDGInitialDataFlags::NoCopy for scene rendering allocator which is known to live longer than the RenderGraph
-	template <typename AllocatorType>
-	struct AllocatorTypeRDGInitialDataFlags
-	{
-		static constexpr ERDGInitialDataFlags Flags = ERDGInitialDataFlags::None;
-	};
-
-	template <>
-	struct AllocatorTypeRDGInitialDataFlags<SceneRenderingAllocator>
-	{
-		static constexpr ERDGInitialDataFlags Flags = ERDGInitialDataFlags::NoCopy;
-	};
-};
-
 /*
  * Helper to build the needed data to run per-instance operation on the GPU in a balanced way
  */
@@ -134,12 +118,11 @@ class TInstanceCullingLoadBalancer : public FInstanceCullingLoadBalancerBase
 {
 public:
 	using AllocatorType = InAllocatorType;
-	static constexpr ERDGInitialDataFlags DefaultRDGInitialDataFlags = InstanceCullingImplementationDetails::AllocatorTypeRDGInitialDataFlags<AllocatorType>::Flags;
 
 	void ReserveStorage(int32 NumBatches, int32 NumItems)
 	{
-		Batches.Empty(NumBatches);
-		Items.Empty(NumItems);
+		Data->Batches.Empty(NumBatches);
+		Data->Items.Empty(NumItems);
 	}
 
 	/**
@@ -156,7 +139,7 @@ public:
 			{
 				const uint32 NumInstancesThisItem = FMath::Min(MaxInstancesThisBatch, NumInstanceDataEntries - InstancesAdded);
 
-				Items.Add(PackItem(InstanceDataOffset + InstancesAdded, NumInstancesThisItem, Payload, CurrentBatchPrefixSum));
+				Data->Items.Add(PackItem(InstanceDataOffset + InstancesAdded, NumInstancesThisItem, Payload, CurrentBatchPrefixSum));
 				if (CurrentBatchNumItems * PrefixBits < sizeof(CurrentBatchPackedPrefixSum) * 8U)
 				{
 					CurrentBatchPackedPrefixSum |= CurrentBatchPackedPrefixSum << (PrefixBits * CurrentBatchNumItems);
@@ -170,8 +153,8 @@ public:
 			// Flush batch if it is not possible to add any more items (for one of the reasons)
 			if (MaxInstancesThisBatch <= 0U || CurrentBatchPrefixSum >= ThreadGroupSize)
 			{
-				Batches.Add(PackBatch(CurrentBatchFirstItem, CurrentBatchNumItems));
-				CurrentBatchFirstItem = uint32(Items.Num());
+				Data->Batches.Add(PackBatch(CurrentBatchFirstItem, CurrentBatchNumItems));
+				CurrentBatchFirstItem = uint32(Data->Items.Num());
 				CurrentBatchPrefixSum = 0u;
 				CurrentBatchNumItems = 0U;
 				CurrentBatchPackedPrefixSum = 0U;
@@ -182,22 +165,32 @@ public:
 
 	bool IsEmpty() const
 	{
-		return Items.IsEmpty();
+		return Data->Items.IsEmpty();
 	}
 
-	FGPUData Upload(FRDGBuilder& GraphBuilder, ERDGInitialDataFlags RDGInitialDataFlags = DefaultRDGInitialDataFlags, int32 NumGroupsPerBatch = 1)
+	FGPUData Upload(FRDGBuilder& GraphBuilder, int32 NumGroupsPerBatch = 1)
 	{
 		FinalizeBatches();
 
-		return FInstanceCullingLoadBalancerBase::Upload(GraphBuilder, Batches, Items, RDGInitialDataFlags, NumGroupsPerBatch);
+		const auto& LocalBatches = Data->Batches;
+		const auto& LocalItems   = Data->Items;
+
+		GraphBuilder.AllocObject<TRefCountPtr<FData>>(Data);
+
+		return FInstanceCullingLoadBalancerBase::Upload(GraphBuilder, LocalBatches, LocalItems, ERDGInitialDataFlags::NoCopy, NumGroupsPerBatch);
 	}
 
 	/* Const variant that assumes the batches have already been finalized */
-	FGPUData UploadFinalized(FRDGBuilder& GraphBuilder, ERDGInitialDataFlags RDGInitialDataFlags = DefaultRDGInitialDataFlags, int32 NumGroupsPerBatch = 1) const
+	FGPUData UploadFinalized(FRDGBuilder& GraphBuilder, int32 NumGroupsPerBatch = 1) const
 	{
 		check(CurrentBatchNumItems == 0);
 
-		return FInstanceCullingLoadBalancerBase::Upload(GraphBuilder, Batches, Items, RDGInitialDataFlags, NumGroupsPerBatch);
+		const auto& LocalBatches = Data->Batches;
+		const auto& LocalItems   = Data->Items;
+
+		GraphBuilder.AllocObject<TRefCountPtr<FData>>(Data);
+
+		return FInstanceCullingLoadBalancerBase::Upload(GraphBuilder, LocalBatches, LocalItems, ERDGInitialDataFlags::NoCopy, NumGroupsPerBatch);
 	}
 
 	/**
@@ -207,7 +200,7 @@ public:
 	{
 		if (CurrentBatchNumItems != 0)
 		{
-			Batches.Add(PackBatch(CurrentBatchFirstItem, CurrentBatchNumItems));
+			Data->Batches.Add(PackBatch(CurrentBatchFirstItem, CurrentBatchNumItems));
 			CurrentBatchNumItems = 0;
 		}
 	}
@@ -219,19 +212,19 @@ public:
 	 */
 	FIntVector GetWrappedCsGroupCount(int32 NumGroupsPerBatch = 1) const
 	{
-		return FInstanceCullingLoadBalancerBase::GetWrappedCsGroupCount(Batches, NumGroupsPerBatch);
+		return FInstanceCullingLoadBalancerBase::GetWrappedCsGroupCount(Data->Batches, NumGroupsPerBatch);
 	}
 
 	const TArray<FPackedBatch, AllocatorType> &GetBatches() const
 	{
 		check(CurrentBatchNumItems == 0);
-		return Batches;
+		return Data->Batches;
 	};
 
 	const TArray<FPackedItem, AllocatorType> &GetItems() const
 	{
 		check(CurrentBatchNumItems == 0);
-		return Items;
+		return Data->Items;
 	}
 
 	uint32 GetTotalNumInstances() const { return TotalInstances; }
@@ -239,19 +232,24 @@ public:
 	template <typename AllocatorType>
 	void AppendData(const TInstanceCullingLoadBalancer<AllocatorType> &Other)
 	{
-		Batches.Append(Other.GetBatches().GetData(), Other.GetBatches().Num());
-		Items.Append(Other.GetItems().GetData(), Other.GetItems().Num());
+		Data->Batches.Append(Other.GetBatches().GetData(), Other.GetBatches().Num());
+		Data->Items.Append(Other.GetItems().GetData(), Other.GetItems().Num());
 		TotalInstances += Other.GetTotalNumInstances();
 	}
 
 	bool HasSingleInstanceItemsOnly() const
 	{
-		return TotalInstances == Items.Num();
+		return TotalInstances == Data->Items.Num();
 	}
 
 protected:
-	TArray<FPackedBatch, AllocatorType>  Batches;
-	TArray<FPackedItem, AllocatorType>  Items;
+	struct FData : public TRefCountingMixin<FData>
+	{
+		TArray<FPackedBatch, AllocatorType>  Batches;
+		TArray<FPackedItem, AllocatorType>  Items;
+	};
+
+	TRefCountPtr<FData> Data{ new FData };
 
 	uint32 CurrentBatchPrefixSum = 0u;
 	uint32 CurrentBatchNumItems = 0U;
