Update toggling VRS to force an update to all mesh draw commands rather than those for meshes using a non-1x1 per-material shading rate.
Fixes potential nullptr dereference in the case of missing materials, and game thread/render thread conflict when accessing the material.

diff --git a/Engine/Source/Runtime/Renderer/Private/RendererScene.cpp b/Engine/Source/Runtime/Renderer/Private/RendererScene.cpp
--- a/Engine/Source/Runtime/Renderer/Private/RendererScene.cpp
+++ b/Engine/Source/Runtime/Renderer/Private/RendererScene.cpp
@@ -6602,7 +6602,13 @@ void FScene::Update(FRDGBuilder& GraphBuilder, const FUpdateParameters& Paramete
 			bScenesPrimitivesNeedStaticMeshElementUpdate = true;
 		}
 
-		bool bVariableRateShadingToggled = GVRSImageManager.GetNeedStaticMeshUpdate();
+		if (GVRSImageManager.GetNeedStaticMeshUpdate())
+		{
+			UE_CLOG(!GIsEditor, LogRenderer, Log, TEXT("Forcing update for all mesh draw commands: Variable rate shading toggled"));
+			bScenesPrimitivesNeedStaticMeshElementUpdate = true;
+
+			GVRSImageManager.SetNeedStaticMeshUpdate(false);
+		}
 
 		if (bScenesPrimitivesNeedStaticMeshElementUpdate || CachedDefaultBasePassDepthStencilAccess != DefaultBasePassDepthStencilAccess)
 		{
@@ -6612,26 +6618,6 @@ void FScene::Update(FRDGBuilder& GraphBuilder, const FUpdateParameters& Paramete
 			bScenesPrimitivesNeedStaticMeshElementUpdate = false;
 			CachedDefaultBasePassDepthStencilAccess = DefaultBasePassDepthStencilAccess;
 		}
-		else if (bVariableRateShadingToggled)
-		{
-			// If VRS was just toggled, mark all primitives using non-1x1 per-material shading rates as needing an update
-			for (FPrimitiveSceneInfo* PrimitiveSceneInfo : Primitives)
-			{
-				for (FMeshBatch& MeshBatch : PrimitiveSceneInfo->StaticMeshes)
-				{
-					if (MeshBatch.MaterialRenderProxy->GetMaterialNoFallback(FeatureLevel)->GetShadingRate() != EMaterialShadingRate::MSR_1x1)
-					{
-						PrimitivesNeedingStaticMeshUpdate[PrimitiveSceneInfo->PackedIndex] = true;
-						break;
-					}
-				}
-			}
-		}
-
-		if (bVariableRateShadingToggled)
-		{
-			GVRSImageManager.SetNeedStaticMeshUpdate(false);
-		}
 
 		for (FPrimitiveSceneInfo* PrimitiveSceneInfo : SceneInfosWithStaticDrawListUpdate)
 		{
