[SingleLayerWater] Fix for refraction reading garbage when the refracted ray samples a non-water pixel.

diff --git a/Engine/Source/Runtime/Renderer/Private/SingleLayerWaterRendering.cpp b/Engine/Source/Runtime/Renderer/Private/SingleLayerWaterRendering.cpp
--- a/Engine/Source/Runtime/Renderer/Private/SingleLayerWaterRendering.cpp
+++ b/Engine/Source/Runtime/Renderer/Private/SingleLayerWaterRendering.cpp
@@ -850,7 +850,8 @@ static void AddCopySceneWithoutWaterPass_Internal(
 	const bool bDoTiledCopy = !bCopyDepth 
 		&& RefractionDownsampleFactor == 1 
 		&& SingleLayerWaterPrepassResult 
-		&& !SingleLayerWaterPrepassResult->ViewTileClassification.IsEmpty() 
+		&& !SingleLayerWaterPrepassResult->ViewTileClassification.IsEmpty()
+		&& HasBeenProduced(SingleLayerWaterPrepassResult->RefractionMaskTexture)
 		&& CVarWaterSingleLayerTiledSceneColorCopy.GetValueOnRenderThread() != 0;
 
 	ERenderTargetLoadAction LoadAction = ERenderTargetLoadAction::ENoAction;
