[TSR] Fix ensure on Luma texture not written to before using.

diff --git a/Engine/Source/Runtime/Renderer/Private/PostProcess/TemporalSuperResolution.cpp b/Engine/Source/Runtime/Renderer/Private/PostProcess/TemporalSuperResolution.cpp
--- a/Engine/Source/Runtime/Renderer/Private/PostProcess/TemporalSuperResolution.cpp
+++ b/Engine/Source/Runtime/Renderer/Private/PostProcess/TemporalSuperResolution.cpp
@@ -2846,7 +2846,7 @@ FDefaultTemporalUpscaler::FOutputs AddTemporalSuperResolutionPasses(
 		PassParameters->bPassthroughAlpha = IsPrimitiveAlphaHoldoutEnabled(View);
 
 		PassParameters->InputTexture = PassInputs.SceneColor.Texture;
-		if (PassInputs.FlickeringInputTexture.IsValid())
+		if (PassInputs.FlickeringInputTexture.IsValid() && FlickeringFramePeriod > 0.0f)
 		{
 			ensure(InputRect == PassInputs.FlickeringInputTexture.ViewRect);
 			PassParameters->InputMoireLumaTexture = GraphBuilder.CreateSRV(
