LightRendering: Allow ray tracing shadows in the presence of stereo and hair strands.

This restriction was added due to a crash in 4.27 but that does not repro any more. This was causing unexpected behavior in nDisplay where RT shadows cannot be enabled via the project setting and instead requires an explicit enable on each light.

#rnx

diff --git a/Engine/Source/Runtime/Renderer/Private/LightRendering.cpp b/Engine/Source/Runtime/Renderer/Private/LightRendering.cpp
--- a/Engine/Source/Runtime/Renderer/Private/LightRendering.cpp
+++ b/Engine/Source/Runtime/Renderer/Private/LightRendering.cpp
@@ -220,26 +220,42 @@ static bool ShouldRenderRayTracingShadowsForLightType(ELightComponentType LightT
 
 bool ShouldRenderRayTracingShadows(const FSceneViewFamily& ViewFamily)
 {
-	const bool bIsStereo = GEngine->StereoRenderingDevice.IsValid() && GEngine->StereoRenderingDevice->IsStereoEnabled();
-	const bool bHairStrands = IsHairStrandsEnabled(EHairStrandsShaderType::Strands);
+	return ShouldRenderRayTracingEffect(
+		(CVarRayTracingOcclusion.GetValueOnRenderThread() > 0),
+		ERayTracingPipelineCompatibilityFlags::FullPipeline | ERayTracingPipelineCompatibilityFlags::Inline, 
+		ViewFamily
+	);
+}
+
+bool ShouldRenderRayTracingShadowsForLight(const FSceneViewFamily& ViewFamily, ELightComponentType LightType, ECastRayTracedShadow::Type CastRayTracedShadow)
+{
+	if (   !ShouldRenderRayTracingEffect(true, ERayTracingPipelineCompatibilityFlags::FullPipeline | ERayTracingPipelineCompatibilityFlags::Inline, ViewFamily)
+		|| !ShouldRenderRayTracingShadowsForLightType(LightType))
+	{
+		return false;
+	}
 
-	return ShouldRenderRayTracingEffect((CVarRayTracingOcclusion.GetValueOnRenderThread() > 0) && !(bIsStereo && bHairStrands), ERayTracingPipelineCompatibilityFlags::FullPipeline | ERayTracingPipelineCompatibilityFlags::Inline, ViewFamily);
+	switch (CastRayTracedShadow)
+	{
+	case ECastRayTracedShadow::Enabled:
+		return true;
+
+	case ECastRayTracedShadow::UseProjectSetting:
+		return ShouldRenderRayTracingShadows(ViewFamily);
+
+	default:
+		return false;
+	}
 }
 
 bool ShouldRenderRayTracingShadowsForLight(const FSceneViewFamily& ViewFamily, const FLightSceneProxy& LightProxy)
 {
-	const bool bShadowRayTracingAllowed = ShouldRenderRayTracingEffect(true, ERayTracingPipelineCompatibilityFlags::FullPipeline | ERayTracingPipelineCompatibilityFlags::Inline, ViewFamily);
-	return (LightProxy.CastsRaytracedShadow() == ECastRayTracedShadow::Enabled || (ShouldRenderRayTracingShadows(ViewFamily) && LightProxy.CastsRaytracedShadow() == ECastRayTracedShadow::UseProjectSetting))
-		&& ShouldRenderRayTracingShadowsForLightType((ELightComponentType)LightProxy.GetLightType())
-		&& bShadowRayTracingAllowed;
+	return ShouldRenderRayTracingShadowsForLight(ViewFamily, ELightComponentType(LightProxy.GetLightType()), LightProxy.CastsRaytracedShadow());
 }
 
 bool ShouldRenderRayTracingShadowsForLight(const FSceneViewFamily& ViewFamily, const FLightSceneInfoCompact& LightInfo)
 {
-	const bool bShadowRayTracingAllowed = ShouldRenderRayTracingEffect(true, ERayTracingPipelineCompatibilityFlags::FullPipeline | ERayTracingPipelineCompatibilityFlags::Inline, ViewFamily);
-	return (LightInfo.CastRaytracedShadow == ECastRayTracedShadow::Enabled || (ShouldRenderRayTracingShadows(ViewFamily) && LightInfo.CastRaytracedShadow == ECastRayTracedShadow::UseProjectSetting))
-		&& ShouldRenderRayTracingShadowsForLightType((ELightComponentType)LightInfo.LightType)
-		&& bShadowRayTracingAllowed;
+	return ShouldRenderRayTracingShadowsForLight(ViewFamily, ELightComponentType(LightInfo.LightType), LightInfo.CastRaytracedShadow);
 }
 #endif // RHI_RAYTRACING
 
