Move debug visualization of pending virtual texture mips to post process.
Data is generated in ODSC debug shader permutations in main passes by writing extended debug data to the feedback buffer.

diff --git a/Engine/Plugins/Experimental/CinematicPrestreaming/Source/CinematicPrestreamingEditor/Private/CinePrestreamingDebugRender.h b/Engine/Plugins/Experimental/CinematicPrestreaming/Source/CinematicPrestreamingEditor/Private/CinePrestreamingDebugRender.h
--- a/Engine/Plugins/Experimental/CinematicPrestreaming/Source/CinematicPrestreamingEditor/Private/CinePrestreamingDebugRender.h
+++ b/Engine/Plugins/Experimental/CinematicPrestreaming/Source/CinematicPrestreamingEditor/Private/CinePrestreamingDebugRender.h
@@ -21,8 +21,8 @@ public:
 	virtual void GetViewShowFlags(FEngineShowFlags& OutShowFlag, EViewModeIndex& OutViewModeIndex) const override
 	{
 		OutShowFlag = FEngineShowFlags(EShowFlagInitMode::ESFIM_Game);
-		OutShowFlag.SetVirtualTexturePendingMips(true);
-		OutViewModeIndex = EViewModeIndex::VMI_VirtualTexturePendingMips;
+		OutShowFlag.SetVisualizeVirtualTexture(true);
+		OutViewModeIndex = EViewModeIndex::VMI_VisualizeVirtualTexture;
 	}
 	virtual int32 GetOutputFileSortingOrder() const override { return 2; }
 };
diff --git a/Engine/Shaders/Private/BasePassPixelShader.usf b/Engine/Shaders/Private/BasePassPixelShader.usf
--- a/Engine/Shaders/Private/BasePassPixelShader.usf
+++ b/Engine/Shaders/Private/BasePassPixelShader.usf
@@ -2382,7 +2382,8 @@ void FPixelShaderInOut_MainPS(
 		FinalizeVirtualTextureFeedback(
 			MaterialParameters.VirtualTextureFeedback,
 			MaterialParameters.SvPosition,
-			View.VTFeedbackBuffer
+			View.VTFeedbackBuffer,
+			VISUALIZE
 		);
 	}
 #endif
diff --git a/Engine/Shaders/Private/DebugViewModePixelShader.usf b/Engine/Shaders/Private/DebugViewModePixelShader.usf
--- a/Engine/Shaders/Private/DebugViewModePixelShader.usf
+++ b/Engine/Shaders/Private/DebugViewModePixelShader.usf
@@ -78,7 +78,6 @@ struct FTexCoordScalesParams
 	float4 ScalesPerIndex;
 
 	float RequiredResolution;
-	int VTMissingMipCount;
 };
 
 MaterialFloat StoreTexCoordScale(in out FTexCoordScalesParams Params, float2 UV, int TextureReferenceIndex)
@@ -214,7 +213,6 @@ float4 VisualizeMaterialTexCoordScales(in FDebugPSIn DebugInputs, in bool bIsFro
 	Params.MaxScale = 0;
 	Params.ScalesPerIndex = INITIAL_GPU_SCALE;
 	Params.RequiredResolution = 0;
-	Params.VTMissingMipCount = -1;
 
 	FMaterialPixelParameters MaterialParameters = GetMaterialPixelParameters(DebugInputs, DebugInputs.SvPosition);
 
@@ -367,7 +365,6 @@ float4 VisualizeRequiredTextureResolution(in FDebugPSIn DebugInputs, in bool bIs
 
 	MaterialParameters.TexCoordScalesParams.RequiredResolution = 0;
 	MaterialParameters.TexCoordScalesParams.TexSample = 0;
-	MaterialParameters.TexCoordScalesParams.VTMissingMipCount = -1;
 
 	half3 BaseColor;
 	{
@@ -408,18 +405,6 @@ float4 VisualizeRequiredTextureResolution(in FDebugPSIn DebugInputs, in bool bIs
 			Result *= MaterialParameters.TexCoordScalesParams.TexSample;
 		}
 	}
-	else if (VisualizeMode == DVSM_VirtualTexturePendingMips)
-	{
-#if MATERIAL_VIRTUALTEXTURE_FEEDBACK || LIGHTMAP_VT_ENABLED
-		uint PendingMips = MaterialParameters.VirtualTextureFeedback.PendingMips;
-		if (PendingMips > 0)
-		{
-			int ColorIndex = PendingMips > 1 ? 0 : 1;
-			Result = DebugViewModeStruct.AccuracyColors[ColorIndex].rgb;
-			Result *= 0.5f;
-		}
-#endif
-	}
 
 #if MATERIAL_VIRTUALTEXTURE_FEEDBACK || LIGHTMAP_VT_ENABLED
 	FinalizeVirtualTextureFeedback(
@@ -553,8 +538,7 @@ void Main(
 	{
 		VisualizeColor = VisualizeMaterialTexCoordScales(DebugInputs, bIsFrontFace);
 	}
-	else if (VisualizeMode == DVSM_RequiredTextureResolution ||
-			 VisualizeMode == DVSM_VirtualTexturePendingMips)
+	else if (VisualizeMode == DVSM_RequiredTextureResolution)
 	{
 		VisualizeColor = VisualizeRequiredTextureResolution(DebugInputs, bIsFrontFace);
 	}
diff --git a/Engine/Shaders/Private/VirtualTextureCommon.ush b/Engine/Shaders/Private/VirtualTextureCommon.ush
--- a/Engine/Shaders/Private/VirtualTextureCommon.ush
+++ b/Engine/Shaders/Private/VirtualTextureCommon.ush
@@ -21,8 +21,8 @@
 #define LIGHTMAP_VT_ENABLED 0
 #endif
 
-#ifndef WRITE_PENDING_VIRTUALTEXTURE_MIPS
-#define WRITE_PENDING_VIRTUALTEXTURE_MIPS 0
+#ifndef VISUALIZE
+#define VISUALIZE 0
 #endif
 
 /** Struct used to store feedback information in. */
@@ -30,7 +30,7 @@ struct FVirtualTextureFeedbackParams
 {
 	uint RequestId;		// The id of the VT texture sample this pixel will store a request for. Only used when we have more than one sample.
 	uint Request;		// The tile to request for this pixel.
-	uint PendingMips;	// The accumulated pending mip count for this pixel from all samples.
+	uint DebugCounts;	// The accumulated debug counts for this pixel from all samples. Stack count in top bits, pending mips in lower bits.
 };
 
 /** Initializes the FVirtualTextureFeedbackParams for the pixel shader. */
@@ -40,12 +40,10 @@ void InitializeVirtualTextureFeedback(in out FVirtualTextureFeedbackParams Param
 	const uint NumVTSamplesInShader = NUM_VIRTUALTEXTURE_FEEDBACK_REQUESTS + LIGHTMAP_VT_ENABLED;
 	Params.RequestId = View.VirtualTextureFeedbackSampleOffset % NumVTSamplesInShader;
 #endif	
-
+	// Init request with "empty" value.
 	Params.Request = 0xFFFFFFFF;
-
-#if WRITE_PENDING_VIRTUALTEXTURE_MIPS
-	Params.PendingMips = 0;
-#endif
+	// Init debug count with number of feedback samples in top 16 bits.
+	Params.DebugCounts = (NUM_VIRTUALTEXTURE_SAMPLES + LIGHTMAP_VT_ENABLED) << 16;
 }
 
 /** Store feedback info for a VT request in the passed in FVirtualTextureFeedbackParams. */
@@ -56,9 +54,7 @@ void StoreVirtualTextureFeedback(in out FVirtualTextureFeedbackParams Params, ui
 #else
 	Params.Request = Request;	
 #endif
-#if WRITE_PENDING_VIRTUALTEXTURE_MIPS
-	Params.PendingMips = max(Params.PendingMips, PendingMips);
-#endif
+	Params.DebugCounts += PendingMips;
 }
 
 /** This should be called at the end of the pixel shader to write out the gathered VT feedback info to the OutputBuffer. */
@@ -78,13 +74,25 @@ void FinalizeVirtualTextureFeedback(in FVirtualTextureFeedbackParams Params, in
 		{
 			OutputBuffer[WriteOffset] = Params.Request;
 		}
-	}	
+	}
 }
 
-void FinalizeVirtualTextureFeedback(in FVirtualTextureFeedbackParams Params, float4 SvPosition, RWStructuredBuffer<uint> OutputBuffer)
+void FinalizeVirtualTextureFeedback(in FVirtualTextureFeedbackParams Params, float4 SvPosition, RWStructuredBuffer<uint> OutputBuffer, bool bEnableDebug = false)
 {
 	const uint4 FeedbackViewParams = uint4(View.VirtualTextureFeedbackMask, View.VirtualTextureFeedbackShift, View.VirtualTextureFeedbackJitterOffset, View.VirtualTextureFeedbackBufferSize);
 	FinalizeVirtualTextureFeedback(Params, FeedbackViewParams, SvPosition, OutputBuffer);
+
+#if VISUALIZE
+	// Write to extended debug buffer if it exists and if we passed in the debug flag.
+	// If the debug flag is only driven by a debug permutation then we could potentially drop VISUALIZE defines in this file because the debug logic should just compile out.
+	// But we keep the defines to keep things clear and explicit.
+	if (bEnableDebug && View.VirtualTextureExtendedDebugBufferSize > 0)
+	{
+		uint IgnoredOldValue;
+		uint WriteOffset = View.VirtualTextureFeedbackBufferSize + (uint)SvPosition.y * (uint)View.ViewSizeAndInvSize.x + (uint)SvPosition.x;
+		InterlockedAdd(OutputBuffer[WriteOffset], Params.DebugCounts, IgnoredOldValue);
+	}
+#endif
 }
 
 /**
@@ -214,7 +222,7 @@ struct VTPageTableResult
 
 uint GetVirtualTexturePendingMips(uint RequestedLevel, uint4 PackedPageTableValue, uint AccumulatedResult)
 {
-#if WRITE_PENDING_VIRTUALTEXTURE_MIPS
+#if VISUALIZE
 	// 0 is an empty value in the PageTableResult.
 	// We must also ignore a value of 1 for this test since it will be in the .w channel for page table textures with less than 4 channels.
 	// So only consider packed page table values > 1.
@@ -225,6 +233,7 @@ uint GetVirtualTexturePendingMips(uint RequestedLevel, uint4 PackedPageTableValu
 		AccumulatedResult = max(MinSampledLevel - RequestedLevel, AccumulatedResult);
 	}
 #endif
+
 	return AccumulatedResult;
 }
 
diff --git a/Engine/Shaders/Private/VisualizeVirtualTexture.usf b/Engine/Shaders/Private/VisualizeVirtualTexture.usf
new file mode 100644
--- /dev/null
+++ b/Engine/Shaders/Private/VisualizeVirtualTexture.usf
@@ -0,0 +1,54 @@
+// Copyright Epic Games, Inc. All Rights Reserved.
+
+#include "Common.ush"
+#include "PostProcessCommon.ush"
+#include "ScreenPass.ush"
+ 
+SCREEN_PASS_TEXTURE_VIEWPORT(Input)
+SCREEN_PASS_TEXTURE_VIEWPORT(Output)
+
+Texture2D InputTexture;
+SamplerState InputSampler;
+
+StructuredBuffer<uint> DebugBuffer;
+uint2 ViewSize;
+
+uint ViewMode;
+
+float4 Colors[MAX_NUM_COLORS];
+uint ColorCount;
+
+void Main(
+	noperspective float4 UVAndScreenPos : TEXCOORD0,
+	FStereoVSOutput StereoOutput,
+	float4 SvPosition : SV_POSITION,
+	out float4 OutColor : SV_Target0)
+{
+	float2 UV = UVAndScreenPos.xy;
+	float3 SceneColor = Texture2DSample(InputTexture, InputSampler, UV).rgb;
+	float SceneColorLuminance = saturate(Luminance(SceneColor));
+	OutColor.xyz = SceneColorLuminance;
+	OutColor.w = 1;
+
+	uint2 PixelPos = (uint2)SvPosition.xy;
+	uint ReadIndex = PixelPos.y * ViewSize.x + PixelPos.x;
+	uint PackedDebugValue = DebugBuffer[ReadIndex];
+	uint DebugValue = 0;
+
+	if (ViewMode == 1)
+	{
+		// Pending mips in lower 16 bits.
+		DebugValue = PackedDebugValue & 0xffff;
+	}
+	else if (ViewMode == 2)
+	{
+		// Stack count in upper 16 bits.
+		DebugValue = PackedDebugValue >> 16;
+	}
+
+	if (DebugValue != 0)
+	{
+		DebugValue = clamp(DebugValue - 1, 0, ColorCount - 1);
+		OutColor.xyz = Colors[DebugValue].rgb * lerp(0.8, 1.0, SceneColorLuminance);
+	}
+}
diff --git a/Engine/Source/Editor/CommonMenuExtensions/Private/CommonMenuExtensionsModule.cpp b/Engine/Source/Editor/CommonMenuExtensions/Private/CommonMenuExtensionsModule.cpp
--- a/Engine/Source/Editor/CommonMenuExtensions/Private/CommonMenuExtensionsModule.cpp
+++ b/Engine/Source/Editor/CommonMenuExtensions/Private/CommonMenuExtensionsModule.cpp
@@ -7,6 +7,7 @@
 #include "SubstrateVisualizationMenuCommands.h"
 #include "GroomVisualizationMenuCommands.h"
 #include "VirtualShadowMapVisualizationMenuCommands.h"
+#include "VirtualTextureVisualizationMenuCommands.h"
 #include "ShowFlagMenuCommands.h"
 #include "Modules/ModuleManager.h"
 
@@ -20,12 +21,14 @@ void FCommonMenuExtensionsModule::StartupModule()
 	FSubstrateVisualizationMenuCommands::Register();
 	FGroomVisualizationMenuCommands::Register();
 	FVirtualShadowMapVisualizationMenuCommands::Register();
+	FVirtualTextureVisualizationMenuCommands::Register();
 	FShowFlagMenuCommands::Register();
 }
 
 void FCommonMenuExtensionsModule::ShutdownModule()
 {
 	FShowFlagMenuCommands::Unregister();
+	FVirtualTextureVisualizationMenuCommands::Unregister();
 	FVirtualShadowMapVisualizationMenuCommands::Unregister();
 	FNaniteVisualizationMenuCommands::Unregister();
 	FGroomVisualizationMenuCommands::Unregister();
diff --git a/Engine/Source/Editor/CommonMenuExtensions/Private/VirtualTextureVisualizationMenuCommands.cpp b/Engine/Source/Editor/CommonMenuExtensions/Private/VirtualTextureVisualizationMenuCommands.cpp
new file mode 100644
--- /dev/null
+++ b/Engine/Source/Editor/CommonMenuExtensions/Private/VirtualTextureVisualizationMenuCommands.cpp
@@ -0,0 +1,120 @@
+// Copyright Epic Games, Inc. All Rights Reserved.
+
+#include "VirtualTextureVisualizationMenuCommands.h"
+
+#include "EditorViewportClient.h"
+#include "Framework/Commands/InputChord.h"
+#include "Framework/Commands/UIAction.h"
+#include "Framework/Commands/UICommandInfo.h"
+#include "Framework/Commands/UICommandList.h"
+#include "Framework/MultiBox/MultiBoxBuilder.h"
+#include "VT/VirtualTextureVisualizationData.h"
+
+#define LOCTEXT_NAMESPACE "VirtualTextureVisualizationMenuCommands"
+
+FVirtualTextureVisualizationMenuCommands::FVirtualTextureVisualizationMenuCommands()
+	: TCommands<FVirtualTextureVisualizationMenuCommands>
+	(
+		TEXT("VirtualTextureVisualizationMenu"), // Context name for fast lookup
+		NSLOCTEXT("Contexts", "VirtualTextureVisualizationMenu", "VirtualTexture"), // Localized context name for displaying
+		NAME_None, // Parent context name.  
+		FAppStyle::GetAppStyleSetName() // Icon Style Set
+	),
+	CommandMap()
+{
+}
+
+void FVirtualTextureVisualizationMenuCommands::BuildCommandMap()
+{
+	CommandMap.Empty();
+
+	const FVirtualTextureVisualizationData& VisualizationData = GetVirtualTextureVisualizationData();
+	for (FVirtualTextureVisualizationData::TModeArray::TConstIterator It = VisualizationData.GetModes().CreateConstIterator(); It; ++It)
+	{
+		const FVirtualTextureVisualizationData::FModeRecord& Entry = *It;
+		FVisualizationRecord& Record = CommandMap.Add(Entry.ModeName, FVisualizationRecord());
+		Record.Name = Entry.ModeName;
+		Record.ModeID = Entry.ModeID;
+		Record.Command = FUICommandInfoDecl(
+			this->AsShared(),
+			Entry.ModeName,
+			Entry.ModeText,
+			Entry.ModeDesc)
+			.UserInterfaceType(EUserInterfaceActionType::RadioButton)
+			.DefaultChord(FInputChord()
+		);
+
+	}
+}
+
+void FVirtualTextureVisualizationMenuCommands::BuildVisualisationSubMenu(FMenuBuilder& Menu)
+{
+	const FVirtualTextureVisualizationMenuCommands& Commands = FVirtualTextureVisualizationMenuCommands::Get();
+	if (Commands.IsPopulated())
+	{
+		Commands.AddCommandTypeToMenu(Menu, EVirtualTextureVisualizationMode::PendingMips);
+		Commands.AddCommandTypeToMenu(Menu, EVirtualTextureVisualizationMode::StackCount);
+	}
+}
+
+bool FVirtualTextureVisualizationMenuCommands::AddCommandTypeToMenu(FMenuBuilder& Menu, const EVirtualTextureVisualizationMode ModeID) const
+{
+	bool bAddedCommands = false;
+
+	const TVisualizationModeCommandMap& Commands = CommandMap;
+	for (TCommandConstIterator It = CreateCommandConstIterator(); It; ++It)
+	{
+		const FVisualizationRecord& Record = It.Value();
+		if (Record.ModeID == ModeID)
+		{
+			Menu.AddMenuEntry(Record.Command, NAME_None, Record.Command->GetLabel());
+			bAddedCommands = true;
+		}
+	}
+
+	return bAddedCommands;
+}
+
+FVirtualTextureVisualizationMenuCommands::TCommandConstIterator FVirtualTextureVisualizationMenuCommands::CreateCommandConstIterator() const
+{
+	return CommandMap.CreateConstIterator();
+}
+
+void FVirtualTextureVisualizationMenuCommands::RegisterCommands()
+{
+	BuildCommandMap();
+}
+
+void FVirtualTextureVisualizationMenuCommands::BindCommands(FUICommandList& CommandList, const TSharedPtr<FEditorViewportClient>& Client) const
+{
+	for (FVirtualTextureVisualizationMenuCommands::TCommandConstIterator It = FVirtualTextureVisualizationMenuCommands::Get().CreateCommandConstIterator(); It; ++It)
+	{
+		const FVirtualTextureVisualizationMenuCommands::FVisualizationRecord& Record = It.Value();
+		CommandList.MapAction(
+			Record.Command,
+			FExecuteAction::CreateStatic(&FVirtualTextureVisualizationMenuCommands::ChangeVisualizationMode, Client.ToWeakPtr(), Record.Name),
+			FCanExecuteAction(),
+			FIsActionChecked::CreateStatic(&FVirtualTextureVisualizationMenuCommands::IsVisualizationModeSelected, Client.ToWeakPtr(), Record.Name)
+		);
+	}
+}
+
+void FVirtualTextureVisualizationMenuCommands::ChangeVisualizationMode(TWeakPtr<FEditorViewportClient> WeakClient, FName InName)
+{
+	if (TSharedPtr<FEditorViewportClient> Client = WeakClient.Pin())
+	{
+		Client->ChangeVirtualTextureVisualizationMode(InName);
+	}
+}
+
+bool FVirtualTextureVisualizationMenuCommands::IsVisualizationModeSelected(TWeakPtr<FEditorViewportClient> WeakClient, FName InName)
+{
+	if (TSharedPtr<FEditorViewportClient> Client = WeakClient.Pin())
+	{
+		return Client->IsVirtualTextureVisualizationModeSelected(InName);
+	}
+	
+	return false;
+}
+
+#undef LOCTEXT_NAMESPACE
diff --git a/Engine/Source/Editor/CommonMenuExtensions/Public/VirtualTextureVisualizationMenuCommands.h b/Engine/Source/Editor/CommonMenuExtensions/Public/VirtualTextureVisualizationMenuCommands.h
new file mode 100644
--- /dev/null
+++ b/Engine/Source/Editor/CommonMenuExtensions/Public/VirtualTextureVisualizationMenuCommands.h
@@ -0,0 +1,54 @@
+// Copyright Epic Games, Inc. All Rights Reserved.
+
+#pragma once
+
+#include "Framework/Commands/Commands.h"
+#include "Templates/SharedPointer.h"
+#include "UObject/NameTypes.h"
+
+class SEditorViewport;
+class FEditorViewportClient;
+class FMenuBuilder;
+class FUICommandInfo;
+class FUICommandList;
+struct FToolMenuEntry;
+enum class EVirtualTextureVisualizationMode : uint8;
+
+class COMMONMENUEXTENSIONS_API FVirtualTextureVisualizationMenuCommands : public TCommands<FVirtualTextureVisualizationMenuCommands>
+{
+public:
+	struct FVisualizationRecord
+	{
+		FName Name;
+		TSharedPtr<FUICommandInfo> Command;
+		EVirtualTextureVisualizationMode ModeID;
+	};
+
+	typedef TMultiMap<FName, FVisualizationRecord> TVisualizationModeCommandMap;
+	typedef TVisualizationModeCommandMap::TConstIterator TCommandConstIterator;
+
+	FVirtualTextureVisualizationMenuCommands();
+
+	TCommandConstIterator CreateCommandConstIterator() const;
+
+	static void BuildVisualisationSubMenu(FMenuBuilder& Menu);
+
+	virtual void RegisterCommands() override;
+
+	void BindCommands(FUICommandList& CommandList, const TSharedPtr<FEditorViewportClient>& Client) const;
+
+	inline bool IsPopulated() const
+	{
+		return CommandMap.Num() > 0;
+	}
+
+private:
+	void BuildCommandMap();
+	bool AddCommandTypeToMenu(FMenuBuilder& Menu, const EVirtualTextureVisualizationMode ModeID) const;
+
+	static void ChangeVisualizationMode(TWeakPtr<FEditorViewportClient> WeakClient, FName InName);
+	static bool IsVisualizationModeSelected(TWeakPtr<FEditorViewportClient> WeakClient, FName InName);
+	
+private:
+	TVisualizationModeCommandMap CommandMap;
+};
diff --git a/Engine/Source/Editor/LevelEditor/Private/SLevelViewport.cpp b/Engine/Source/Editor/LevelEditor/Private/SLevelViewport.cpp
--- a/Engine/Source/Editor/LevelEditor/Private/SLevelViewport.cpp
+++ b/Engine/Source/Editor/LevelEditor/Private/SLevelViewport.cpp
@@ -81,6 +81,7 @@
 #include "LumenVisualizationMenuCommands.h"
 #include "SubstrateVisualizationMenuCommands.h"
 #include "VirtualShadowMapVisualizationMenuCommands.h"
+#include "VirtualTextureVisualizationMenuCommands.h"
 #include "EditorLevelUtils.h"
 #include "Engine/LevelStreaming.h"
 #include "WorldBrowserModule.h"
@@ -633,6 +634,7 @@ void SLevelViewport::ConstructLevelEditorViewportClient(FLevelEditorViewportInst
 	LevelViewportClient->CurrentSubstrateVisualizationMode = ViewportInstanceSettings.SubstrateVisualizationMode;
 	LevelViewportClient->CurrentGroomVisualizationMode = ViewportInstanceSettings.GroomVisualizationMode;
 	LevelViewportClient->CurrentVirtualShadowMapVisualizationMode = ViewportInstanceSettings.VirtualShadowMapVisualizationMode;
+	LevelViewportClient->CurrentVirtualTextureVisualizationMode = ViewportInstanceSettings.VirtualTextureVisualizationMode;
 	LevelViewportClient->CurrentRayTracingDebugVisualizationMode = ViewportInstanceSettings.RayTracingDebugVisualizationMode;
 	LevelViewportClient->CurrentGPUSkinCacheVisualizationMode = ViewportInstanceSettings.GPUSkinCacheVisualizationMode;
 	LevelViewportClient->ExposureSettings = ViewportInstanceSettings.ExposureSettings;
@@ -1735,6 +1737,10 @@ void SLevelViewport::BindViewCommands( FUICommandList& OutCommandList )
 		FGroomVisualizationMenuCommands::Get().BindCommands(OutCommandList, Client);
 	}
 	FVirtualShadowMapVisualizationMenuCommands::Get().BindCommands(OutCommandList, Client);
+	if (UseVirtualTexturing(GMaxRHIShaderPlatform))
+	{
+		FVirtualTextureVisualizationMenuCommands::Get().BindCommands(OutCommandList, Client);
+	}
 }
 
 
@@ -2673,6 +2679,7 @@ void SLevelViewport::SaveConfig(const FString& ConfigName) const
 		ViewportInstanceSettings.SubstrateVisualizationMode = LevelViewportClient->CurrentSubstrateVisualizationMode;
 		ViewportInstanceSettings.GroomVisualizationMode = LevelViewportClient->CurrentGroomVisualizationMode;
 		ViewportInstanceSettings.VirtualShadowMapVisualizationMode = LevelViewportClient->CurrentVirtualShadowMapVisualizationMode;
+		ViewportInstanceSettings.VirtualTextureVisualizationMode = LevelViewportClient->CurrentVirtualTextureVisualizationMode;
 		ViewportInstanceSettings.RayTracingDebugVisualizationMode = LevelViewportClient->CurrentRayTracingDebugVisualizationMode;
 		ViewportInstanceSettings.GPUSkinCacheVisualizationMode = LevelViewportClient->CurrentGPUSkinCacheVisualizationMode;
 		ViewportInstanceSettings.ExposureSettings = LevelViewportClient->ExposureSettings;
diff --git a/Engine/Source/Editor/LevelEditor/Private/ViewportToolbar/LevelEditorSubmenus.cpp b/Engine/Source/Editor/LevelEditor/Private/ViewportToolbar/LevelEditorSubmenus.cpp
--- a/Engine/Source/Editor/LevelEditor/Private/ViewportToolbar/LevelEditorSubmenus.cpp
+++ b/Engine/Source/Editor/LevelEditor/Private/ViewportToolbar/LevelEditorSubmenus.cpp
@@ -26,6 +26,7 @@
 #include "ToolMenus.h"
 #include "ViewportToolbar/LevelViewportContext.h"
 #include "ViewportToolbar/UnrealEdViewportToolbar.h"
+#include "VirtualTextureVisualizationMenuCommands.h"
 #include "Widgets/Input/SNumericEntryBox.h"
 #include "Widgets/Input/SSpinBox.h"
 #include "Widgets/Input/SVolumeControl.h"
@@ -1046,6 +1047,33 @@ void PopulateViewModesMenu(UToolMenu* InMenu)
 		);
 	}
 
+	if (UseVirtualTexturing(GMaxRHIShaderPlatform))
+	{
+		FToolMenuSection& Section = InMenu->FindOrAddSection("ViewMode");
+		Section.AddSubMenu(
+			"VirtualTextureSubMenu",
+			LOCTEXT("VirtualTexture_SubMenu", "Virtual Texture"),
+			LOCTEXT("VirtualTexure_ToolTip", "Select virtual texture visualization view modes"),
+			FNewMenuDelegate::CreateStatic(&FVirtualTextureVisualizationMenuCommands::BuildVisualisationSubMenu),
+			FUIAction(
+				FExecuteAction(),
+				FCanExecuteAction(),
+				FIsActionChecked::CreateLambda(
+					[WeakViewport = LevelViewport.ToWeakPtr()]()
+					{
+						const TSharedPtr<::SLevelViewport> Viewport = WeakViewport.Pin();
+						check(Viewport.IsValid());
+						FLevelEditorViewportClient& ViewportClient = Viewport->GetLevelViewportClient();
+						return ViewportClient.IsViewModeEnabled(VMI_VisualizeVirtualTexture);
+					}
+				)
+			),
+			EUserInterfaceActionType::RadioButton,
+			/* bInOpenSubMenuOnClick = */ false,
+			FSlateIcon(FAppStyle::GetAppStyleSetName(), "EditorViewport.VisualizeVirtualTextureMode")
+		);
+	}
+
 	{
 		auto BuildActorColorationMenu = [WeakViewport = LevelViewport.ToWeakPtr()](UToolMenu* InMenu)
 		{
diff --git a/Engine/Source/Editor/UnrealEd/Classes/Settings/LevelEditorViewportSettings.h b/Engine/Source/Editor/UnrealEd/Classes/Settings/LevelEditorViewportSettings.h
--- a/Engine/Source/Editor/UnrealEd/Classes/Settings/LevelEditorViewportSettings.h
+++ b/Engine/Source/Editor/UnrealEd/Classes/Settings/LevelEditorViewportSettings.h
@@ -102,6 +102,7 @@ struct FLevelEditorViewportInstanceSettings
 		, SubstrateVisualizationMode()
 		, GroomVisualizationMode()
 		, VirtualShadowMapVisualizationMode()
+		, VirtualTextureVisualizationMode()
 		, RayTracingDebugVisualizationMode()
 		, GPUSkinCacheVisualizationMode()
 		, ExposureSettings()
@@ -165,6 +166,10 @@ struct FLevelEditorViewportInstanceSettings
 	UPROPERTY(config)
 	FName VirtualShadowMapVisualizationMode;
 
+	/** The virtual texture visualization mode for the viewport. */
+	UPROPERTY(config)
+	FName VirtualTextureVisualizationMode;
+
 	/** The buffer visualization mode for the viewport. */
 	UPROPERTY(config)
 	FName RayTracingDebugVisualizationMode;
diff --git a/Engine/Source/Editor/UnrealEd/Private/EditorBuildUtils.cpp b/Engine/Source/Editor/UnrealEd/Private/EditorBuildUtils.cpp
--- a/Engine/Source/Editor/UnrealEd/Private/EditorBuildUtils.cpp
+++ b/Engine/Source/Editor/UnrealEd/Private/EditorBuildUtils.cpp
@@ -1461,8 +1461,6 @@ EDebugViewShaderMode ViewModeIndexToDebugViewShaderMode(EViewModeIndex SelectedV
 		return DVSM_MaterialTextureScaleAccuracy;
 	case VMI_RequiredTextureResolution:
 		return DVSM_RequiredTextureResolution;
-	case VMI_VirtualTexturePendingMips:
-		return DVSM_VirtualTexturePendingMips;
 	case VMI_LODColoration:
 	case VMI_HLODColoration:
 		return DVSM_LODColoration;
diff --git a/Engine/Source/Editor/UnrealEd/Private/EditorViewportClient.cpp b/Engine/Source/Editor/UnrealEd/Private/EditorViewportClient.cpp
--- a/Engine/Source/Editor/UnrealEd/Private/EditorViewportClient.cpp
+++ b/Engine/Source/Editor/UnrealEd/Private/EditorViewportClient.cpp
@@ -67,6 +67,7 @@
 #include "SubstrateVisualizationData.h"
 #include "GroomVisualizationData.h"
 #include "VirtualShadowMapVisualizationData.h"
+#include "VT/VirtualTextureVisualizationData.h"
 #include "GPUSkinCacheVisualizationData.h"
 #include "UnrealWidget.h"
 #include "Engine/World.h"
@@ -421,6 +422,7 @@ FEditorViewportClient::FEditorViewportClient(FEditorModeTools* InModeTools, FPre
 	, CurrentSubstrateVisualizationMode(NAME_None)
 	, CurrentGroomVisualizationMode(NAME_None)
 	, CurrentVirtualShadowMapVisualizationMode(NAME_None)
+	, CurrentVirtualTextureVisualizationMode(NAME_None)
 	, CurrentRayTracingDebugVisualizationMode(NAME_None)
 	, CurrentGPUSkinCacheVisualizationMode(NAME_None)
 	, FramesSinceLastDraw(0)
@@ -2883,6 +2885,23 @@ FText FEditorViewportClient::GetCurrentVirtualShadowMapVisualizationModeDisplayN
 	return GetVirtualShadowMapVisualizationData().GetModeDisplayName(CurrentVirtualShadowMapVisualizationMode);
 }
 
+void FEditorViewportClient::ChangeVirtualTextureVisualizationMode(FName InName)
+{
+	SetViewMode(VMI_VisualizeVirtualTexture);
+	CurrentVirtualTextureVisualizationMode = InName;
+}
+
+bool FEditorViewportClient::IsVirtualTextureVisualizationModeSelected(FName InName) const
+{
+	return IsViewModeEnabled(VMI_VisualizeVirtualTexture) && CurrentVirtualTextureVisualizationMode == InName;
+}
+
+FText FEditorViewportClient::GetCurrentVirtualTextureVisualizationModeDisplayName() const
+{
+	checkf(IsViewModeEnabled(VMI_VisualizeVirtualTexture), TEXT("In order to call GetCurrentVirtualTextureVisualizationMode(), first you must set ViewMode to VMI_VisualizeVirtualTexture."));
+	return GetVirtualTextureVisualizationData().GetModeDisplayName(CurrentVirtualTextureVisualizationMode);
+}
+
 void FEditorViewportClient::ChangeActorColorationVisualizationMode(FName InName)
 {
 	SetViewMode(VMI_VisualizeActorColoration);
@@ -4222,6 +4241,7 @@ void FEditorViewportClient::SetupViewForRendering(FSceneViewFamily& ViewFamily,
 	View.CurrentSubstrateVisualizationMode = CurrentSubstrateVisualizationMode;
 	View.CurrentGroomVisualizationMode = CurrentGroomVisualizationMode;
 	View.CurrentVirtualShadowMapVisualizationMode = CurrentVirtualShadowMapVisualizationMode;
+	View.CurrentVirtualTextureVisualizationMode = CurrentVirtualTextureVisualizationMode;
 	View.CurrentGPUSkinCacheVisualizationMode = CurrentGPUSkinCacheVisualizationMode;
 	View.CurrentRayTracingDebugVisualizationMode = CurrentRayTracingDebugVisualizationMode;
 
diff --git a/Engine/Source/Editor/UnrealEd/Private/EditorViewportCommands.cpp b/Engine/Source/Editor/UnrealEd/Private/EditorViewportCommands.cpp
--- a/Engine/Source/Editor/UnrealEd/Private/EditorViewportCommands.cpp
+++ b/Engine/Source/Editor/UnrealEd/Private/EditorViewportCommands.cpp
@@ -93,7 +93,6 @@ void FEditorViewportCommands::RegisterCommands()
 	UI_COMMAND( TexStreamAccMaterialTextureScaleMode, "Material Texture Scales Accuracy View Mode", "Visualize the accuracy of the material texture scales used for texture streaming", EUserInterfaceActionType::RadioButton, FInputChord() );
 	UI_COMMAND( TexStreamAccMaterialTextureScaleAll, "All Textures", "Visualize the scales accuracy of all textures", EUserInterfaceActionType::RadioButton, FInputChord() );
 	UI_COMMAND( RequiredTextureResolutionMode, "Required Texture Resolution View Mode", "Visualize the ratio between the currently streamed texture resolution and the resolution wanted by the GPU", EUserInterfaceActionType::RadioButton, FInputChord() );
-	UI_COMMAND( VirtualTexturePendingMipsMode, "Virtual Texture Pending Mips View Mode", "Visualize the difference between the currently streamed virtual texture level and the level wanted by the GPU", EUserInterfaceActionType::RadioButton, FInputChord() );
 
 	for (int32 TextureIndex = 0; TextureIndex < TEXSTREAM_MAX_NUM_TEXTURES_PER_MATERIAL; ++TextureIndex)
 	{
@@ -126,6 +125,7 @@ void FEditorViewportCommands::RegisterCommands()
 	UI_COMMAND( VisualizeSubstrateMode, "Substrate Visualization View Mode", "Visualizes Substrate debug views.", EUserInterfaceActionType::RadioButton, FInputChord());
 	UI_COMMAND( VisualizeGroomMode, "Groom Visualization View Mode", "Visualizes Groom debug views.", EUserInterfaceActionType::RadioButton, FInputChord());
 	UI_COMMAND( VisualizeVirtualShadowMapMode, "Virtual Shadow Map Visualization View Mode", "Visualizes various rendering aspects of virtual shadow maps.", EUserInterfaceActionType::RadioButton, FInputChord());
+	UI_COMMAND( VisualizeVirtualTextureMode, "Virtual Texture Visualization View Mode", "Visualizes various rendering aspects of virtual textures", EUserInterfaceActionType::RadioButton, FInputChord());
 
 	UI_COMMAND( ReflectionOverrideMode, "Reflections View Mode", "Renders the scene with reflections only", EUserInterfaceActionType::RadioButton, FInputChord() );
 	UI_COMMAND( CollisionPawn, "Player Collision", "Renders player collision visualization", EUserInterfaceActionType::RadioButton, FInputChord() );
diff --git a/Engine/Source/Editor/UnrealEd/Private/SEditorViewport.cpp b/Engine/Source/Editor/UnrealEd/Private/SEditorViewport.cpp
--- a/Engine/Source/Editor/UnrealEd/Private/SEditorViewport.cpp
+++ b/Engine/Source/Editor/UnrealEd/Private/SEditorViewport.cpp
@@ -561,7 +561,6 @@ void SEditorViewport::BindCommands()
 	MAP_VIEWMODE_ACTION( Commands.TexStreamAccMeshUVDensityMode, VMI_MeshUVDensityAccuracy);
 	MAP_VIEWMODE_ACTION( Commands.TexStreamAccMaterialTextureScaleMode, VMI_MaterialTextureScaleAccuracy );
 	MAP_VIEWMODE_ACTION( Commands.RequiredTextureResolutionMode, VMI_RequiredTextureResolution);
-	MAP_VIEWMODE_ACTION( Commands.VirtualTexturePendingMipsMode, VMI_VirtualTexturePendingMips );
 	MAP_VIEWMODE_ACTION( Commands.StationaryLightOverlapMode, VMI_StationaryLightOverlap );
 
 	if (IsStaticLightingAllowed())
@@ -579,6 +578,7 @@ void SEditorViewport::BindCommands()
 	MAP_VIEWMODE_ACTION( Commands.VisualizeSubstrateMode, VMI_VisualizeSubstrate);
 	MAP_VIEWMODE_ACTION( Commands.VisualizeGroomMode, VMI_VisualizeGroom);
 	MAP_VIEWMODE_ACTION( Commands.VisualizeVirtualShadowMapMode, VMI_VisualizeVirtualShadowMap );
+	MAP_VIEWMODE_ACTION( Commands.VisualizeVirtualTextureMode, VMI_VisualizeVirtualTexture);
 	MAP_VIEWMODE_ACTION( Commands.CollisionPawn, VMI_CollisionPawn);
 	MAP_VIEWMODE_ACTION( Commands.CollisionVisibility, VMI_CollisionVisibility);
 
diff --git a/Engine/Source/Editor/UnrealEd/Private/ViewportToolbar/UnrealEdViewportToolbar.cpp b/Engine/Source/Editor/UnrealEd/Private/ViewportToolbar/UnrealEdViewportToolbar.cpp
--- a/Engine/Source/Editor/UnrealEd/Private/ViewportToolbar/UnrealEdViewportToolbar.cpp
+++ b/Engine/Source/Editor/UnrealEd/Private/ViewportToolbar/UnrealEdViewportToolbar.cpp
@@ -1364,6 +1364,10 @@ FText GetViewModesSubmenuLabel(TWeakPtr<SEditorViewport> InViewport)
 		{
 			Label = ViewportClient->GetCurrentVirtualShadowMapVisualizationModeDisplayName();
 		}
+		else if (ViewMode == VMI_VisualizeVirtualTexture)
+		{
+			Label = ViewportClient->GetCurrentVirtualTextureVisualizationModeDisplayName();
+		}
 		else if (ViewMode == VMI_VisualizeActorColoration)
 		{
 			Label = ViewportClient->GetCurrentActorColorationVisualizationModeDisplayName();
@@ -1813,44 +1817,6 @@ void PopulateViewModesMenu(UToolMenu* InMenu)
 				FSlateIcon(FAppStyle::GetAppStyleSetName(), "EditorViewport.VisualizeGPUSkinCacheMode")
 			);
 		}
-
-		if (UseVirtualTexturing(GMaxRHIShaderPlatform))
-		{
-			struct Local
-			{
-				static void BuildVirtualTextureMenu(UToolMenu* Menu, IsViewModeSupportedDelegate IsViewModeSupported)
-				{
-					FToolMenuSection& Section = Menu->AddSection("VirtualTextureVisualizationSection", LOCTEXT("VirtualTexture_VisualizationSection", "Virtual Texture Visualization"));
-
-					UE::UnrealEd::Private::ViewModesSubmenu::AddModeIfSupported(
-						IsViewModeSupported, Section, FEditorViewportCommands::Get().VirtualTexturePendingMipsMode, VMI_VirtualTexturePendingMips
-					);
-				}
-			};
-
-			Section.AddSubMenu(
-				"VirtualTextureSubMenu",
-				LOCTEXT("VirtualTexture_SubMenu", "Virtual Texture"),
-				LOCTEXT("VirtualTexure_ToolTip", "Select virtual texture visualization view modes"),
-				FNewToolMenuDelegate::CreateStatic(&Local::BuildVirtualTextureMenu, IsViewModeSupported),
-				FUIAction(
-					FExecuteAction(),
-					FCanExecuteAction(),
-					FIsActionChecked::CreateLambda(
-						[WeakViewport = EditorViewport.ToWeakPtr()]()
-						{
-							const TSharedRef<SEditorViewport> ViewportRef = WeakViewport.Pin().ToSharedRef();
-							const TSharedPtr<FEditorViewportClient> ViewportClient = ViewportRef->GetViewportClient();
-							check(ViewportClient.IsValid());
-							return ViewportClient->IsViewModeEnabled(VMI_VirtualTexturePendingMips);
-						}
-					)
-				),
-				EUserInterfaceActionType::RadioButton,
-				false,
-				FSlateIcon(FAppStyle::GetAppStyleSetName(), "EditorViewport.VisualizeVirtualTextureMode")
-			);
-		}
 	}
 }
 
diff --git a/Engine/Source/Editor/UnrealEd/Public/EditorViewportClient.h b/Engine/Source/Editor/UnrealEd/Public/EditorViewportClient.h
--- a/Engine/Source/Editor/UnrealEd/Public/EditorViewportClient.h
+++ b/Engine/Source/Editor/UnrealEd/Public/EditorViewportClient.h
@@ -1444,6 +1444,26 @@ public:
 	*/
 	UNREALED_API FText GetCurrentVirtualShadowMapVisualizationModeDisplayName() const;
 
+	/**
+	* Changes the virtual texture visualization mode for this viewport.
+	*
+	* @param InName	The ID of the required visualization mode
+	*/
+	UNREALED_API void ChangeVirtualTextureVisualizationMode(FName InName);
+
+	/**
+	* Checks if a virtual texture visualization mode is selected.
+	*
+	* @param InName	The ID of the required visualization mode
+	* @return	true if the supplied virtual texture visualization mode is checked
+	*/
+	UNREALED_API bool IsVirtualTextureVisualizationModeSelected(FName InName) const;
+	
+	/**
+	* Returns the FText display name associated with CurrentVirtualTextureVisualizationMode.
+	*/
+	UNREALED_API FText GetCurrentVirtualTextureVisualizationModeDisplayName() const;
+
 	/**
 	* Changes the actor coloration visualization mode for this viewport.
 	*
@@ -1825,6 +1845,7 @@ public:
 	FName CurrentSubstrateVisualizationMode;
 	FName CurrentGroomVisualizationMode;
 	FName CurrentVirtualShadowMapVisualizationMode;
+	FName CurrentVirtualTextureVisualizationMode;
 
 	FName CurrentRayTracingDebugVisualizationMode;
 	FName CurrentGPUSkinCacheVisualizationMode;
diff --git a/Engine/Source/Editor/UnrealEd/Public/EditorViewportCommands.h b/Engine/Source/Editor/UnrealEd/Public/EditorViewportCommands.h
--- a/Engine/Source/Editor/UnrealEd/Public/EditorViewportCommands.h
+++ b/Engine/Source/Editor/UnrealEd/Public/EditorViewportCommands.h
@@ -111,8 +111,6 @@ public:
 	TSharedPtr< FUICommandInfo > RequiredTextureResolutionMode;
 	TSharedPtr< FUICommandInfo > RequiredTextureResolutionSingle[TEXSTREAM_MAX_NUM_TEXTURES_PER_MATERIAL];
 
-	TSharedPtr< FUICommandInfo > VirtualTexturePendingMipsMode;
-
 	/** Changes the viewport to stationary light overlap mode */
 	TSharedPtr< FUICommandInfo > StationaryLightOverlapMode;
 
@@ -146,6 +144,9 @@ public:
 	/** Changes the viewport to visualize virtual shadow map content */
 	TSharedPtr< FUICommandInfo > VisualizeVirtualShadowMapMode;
 
+	/** Changes the viewport to visualize virtual texture content */
+	TSharedPtr< FUICommandInfo > VisualizeVirtualTextureMode;
+
 	/** Changes the viewport to Skin Cache mode */
 	TSharedPtr< FUICommandInfo > VisualizeGPUSkinCacheMode;
 
diff --git a/Engine/Source/Runtime/Engine/Classes/Engine/EngineBaseTypes.h b/Engine/Source/Runtime/Engine/Classes/Engine/EngineBaseTypes.h
--- a/Engine/Source/Runtime/Engine/Classes/Engine/EngineBaseTypes.h
+++ b/Engine/Source/Runtime/Engine/Classes/Engine/EngineBaseTypes.h
@@ -1014,7 +1014,7 @@ enum EViewModeIndex : int
 	VMI_VisualizeNanite = 29 UMETA(DisplayName = "Nanite Visualization"),
 
 	/** Compare the required texture resolution to the actual resolution. */
-	VMI_VirtualTexturePendingMips = 30 UMETA(DisplayName = "Virtual Texture Pending Mips"),
+	VMI_VisualizeVirtualTexture = 30 UMETA(DisplayName = "Virtual Texture Visualization"),
 
 	/** Visualize Lumen debug views */
 	VMI_VisualizeLumen = 31 UMETA(DisplayName = "Lumen Visualization"),
diff --git a/Engine/Source/Runtime/Engine/Internal/VT/VirtualTextureVisualizationData.h b/Engine/Source/Runtime/Engine/Internal/VT/VirtualTextureVisualizationData.h
new file mode 100644
--- /dev/null
+++ b/Engine/Source/Runtime/Engine/Internal/VT/VirtualTextureVisualizationData.h
@@ -0,0 +1,81 @@
+// Copyright Epic Games, Inc. All Rights Reserved.
+
+#pragma once
+
+#include "CoreMinimal.h"
+
+/** The virtual texture visualization modes. */
+enum class EVirtualTextureVisualizationMode : uint8
+{
+	None,
+	PendingMips,
+	StackCount,
+};
+
+/** Manager for virtual texture visualization modes. */
+class FVirtualTextureVisualizationData
+{
+public:
+	FVirtualTextureVisualizationData()
+		: bIsInitialized(false)
+	{
+	}
+
+	/** Initialize the system. */
+	void Initialize();
+
+	/** Check if system was initialized. */
+	inline bool IsInitialized() const { return bIsInitialized; }
+
+	/** Describes a single available visualization mode. */
+	struct FModeRecord
+	{
+		FString ModeString;
+		FName ModeName;
+		FText ModeText;
+		FText ModeDesc;
+		EVirtualTextureVisualizationMode ModeID;
+	};
+
+	/** Mapping of FName to a visualization mode record. */
+	typedef TArray<FModeRecord> TModeArray;
+
+	/** Get the mode map. */
+	inline const TModeArray& GetModes() const { return ModeArray; }
+
+	/** Get the mode id of the active mode. */
+	ENGINE_API FName GetActiveMode(class FSceneView const& InView) const;
+
+	/** Get the mode id from the mode name. **/
+	ENGINE_API EVirtualTextureVisualizationMode GetModeID(FName const& InModeName) const;
+
+	/** Get the display name from the mode name. **/
+	ENGINE_API FText GetModeDisplayName(FName const& InModeName) const;
+
+	/** Get the display description from the mode name. **/
+	ENGINE_API FText GetModeDisplayDesc(FName const& InModeName) const;
+
+	/** Return the console command name for enabling single mode visualization. */
+	static const TCHAR* GetVisualizeConsoleCommandName();
+
+private:
+	void AddVisualizationMode(
+		const TCHAR* ModeString,
+		const FText& ModeText,
+		const FText& ModeDesc,
+		EVirtualTextureVisualizationMode ModeID);
+
+	void ConfigureConsoleCommand();
+
+private:
+	/** Flag indicating if system is initialized. **/
+	bool bIsInitialized;
+
+	/** The registered modes. */
+	TModeArray ModeArray;
+
+	/** Storage for console variable documentation strings. **/
+	FString ConsoleDocumentationVisualizationMode;
+};
+
+ENGINE_API FVirtualTextureVisualizationData& GetVirtualTextureVisualizationData();
diff --git a/Engine/Source/Runtime/Engine/Private/DebugViewModeHelpers.cpp b/Engine/Source/Runtime/Engine/Private/DebugViewModeHelpers.cpp
--- a/Engine/Source/Runtime/Engine/Private/DebugViewModeHelpers.cpp
+++ b/Engine/Source/Runtime/Engine/Private/DebugViewModeHelpers.cpp
@@ -46,8 +46,6 @@ const TCHAR* DebugViewShaderModeToString(EDebugViewShaderMode InShaderMode)
 		return TEXT("DVSM_OutputMaterialTextureScales");
 	case DVSM_RequiredTextureResolution:
 		return TEXT("DVSM_RequiredTextureResolution");
-	case DVSM_VirtualTexturePendingMips:
-		return TEXT("DVSM_VirtualTexturePendingMips");
 	case DVSM_LODColoration:
 		return TEXT("DVSM_LODColoration");
 	case DVSM_VisualizeGPUSkinCache:
@@ -72,8 +70,7 @@ bool IsDebugViewShaderModeODSCOnly()
 bool SupportDebugViewShaderMode(EDebugViewShaderMode ShaderMode, EShaderPlatform Platform)
 {
 	// Quad modes don't compile/work on Metal or Playstation. 
-	// Needs further testing for other platforms.
-	const bool bSupportQuadComplexityModes = !IsMobilePlatform(Platform) && 
+	const bool bSupportQuadComplexityModes =
 		!FDataDrivenShaderPlatformInfo::GetIsLanguageSony(Platform) && 
 		!FDataDrivenShaderPlatformInfo::GetIsLanguageMetal(Platform);
 	
@@ -81,7 +78,6 @@ bool SupportDebugViewShaderMode(EDebugViewShaderMode ShaderMode, EShaderPlatform
 	{
 	case DVSM_LODColoration:
 	case DVSM_VisualizeGPUSkinCache:
-	case DVSM_VirtualTexturePendingMips:
 		return true;
 
 	case DVSM_QuadComplexity:
diff --git a/Engine/Source/Runtime/Engine/Private/SceneView.cpp b/Engine/Source/Runtime/Engine/Private/SceneView.cpp
--- a/Engine/Source/Runtime/Engine/Private/SceneView.cpp
+++ b/Engine/Source/Runtime/Engine/Private/SceneView.cpp
@@ -3183,10 +3183,6 @@ EDebugViewShaderMode FSceneViewFamily::ChooseDebugViewShaderMode() const
 	{
 		return DVSM_RequiredTextureResolution;
 	}
-	else if (EngineShowFlags.VirtualTexturePendingMips)
-	{
-		return DVSM_VirtualTexturePendingMips;
-	}
 	else if (EngineShowFlags.LODColoration || EngineShowFlags.HLODColoration)
 	{
 		return DVSM_LODColoration;
diff --git a/Engine/Source/Runtime/Engine/Private/ShowFlags.cpp b/Engine/Source/Runtime/Engine/Private/ShowFlags.cpp
--- a/Engine/Source/Runtime/Engine/Private/ShowFlags.cpp
+++ b/Engine/Source/Runtime/Engine/Private/ShowFlags.cpp
@@ -327,7 +327,6 @@ void ApplyViewMode(EViewModeIndex ViewModeIndex, bool bPerspective, FEngineShowF
 		case VMI_MeshUVDensityAccuracy:
 		case VMI_MaterialTextureScaleAccuracy:
 		case VMI_RequiredTextureResolution:
-		case VMI_VirtualTexturePendingMips:
 		case VMI_LODColoration:
 		case VMI_HLODColoration:
 			bPostProcessing = false;
@@ -394,6 +393,7 @@ void ApplyViewMode(EViewModeIndex ViewModeIndex, bool bPerspective, FEngineShowF
 	EngineShowFlags.SetVisualizeSubstrate(ViewModeIndex == VMI_VisualizeSubstrate);
 	EngineShowFlags.SetVisualizeGroom(ViewModeIndex == VMI_VisualizeGroom);
 	EngineShowFlags.SetVisualizeVirtualShadowMap(ViewModeIndex == VMI_VisualizeVirtualShadowMap);
+	EngineShowFlags.SetVisualizeVirtualTexture(ViewModeIndex == VMI_VisualizeVirtualTexture);
 	EngineShowFlags.SetVisualizeLightCulling(ViewModeIndex == VMI_LightComplexity);
 	EngineShowFlags.SetShaderComplexity(ViewModeIndex == VMI_ShaderComplexity || ViewModeIndex == VMI_QuadOverdraw || ViewModeIndex == VMI_ShaderComplexityWithQuadOverdraw);
 	EngineShowFlags.SetQuadOverdraw(ViewModeIndex == VMI_QuadOverdraw);
@@ -402,7 +402,6 @@ void ApplyViewMode(EViewModeIndex ViewModeIndex, bool bPerspective, FEngineShowF
 	EngineShowFlags.SetMeshUVDensityAccuracy(ViewModeIndex == VMI_MeshUVDensityAccuracy);
 	EngineShowFlags.SetMaterialTextureScaleAccuracy(ViewModeIndex == VMI_MaterialTextureScaleAccuracy);
 	EngineShowFlags.SetRequiredTextureResolution(ViewModeIndex == VMI_RequiredTextureResolution);
-	EngineShowFlags.SetVirtualTexturePendingMips(ViewModeIndex == VMI_VirtualTexturePendingMips);
 	EngineShowFlags.SetStationaryLightOverlap(ViewModeIndex == VMI_StationaryLightOverlap);
 	EngineShowFlags.SetLightMapDensity(ViewModeIndex == VMI_LightmapDensity || ViewModeIndex == VMI_LitLightmapDensity);
 	EngineShowFlags.SetPostProcessing(bPostProcessing);
@@ -524,7 +523,6 @@ void EngineShowFlagOverride(EShowFlagInitMode ShowFlagInitMode, EViewModeIndex V
 			ViewModeIndex == VMI_MeshUVDensityAccuracy ||
 			ViewModeIndex == VMI_MaterialTextureScaleAccuracy ||
 			ViewModeIndex == VMI_RequiredTextureResolution ||
-			ViewModeIndex == VMI_VirtualTexturePendingMips ||
 			ViewModeIndex == VMI_LightmapDensity ||
 			ViewModeIndex == VMI_LitLightmapDensity)
 		{
@@ -549,7 +547,6 @@ void EngineShowFlagOverride(EShowFlagInitMode ShowFlagInitMode, EViewModeIndex V
 			ViewModeIndex == VMI_MeshUVDensityAccuracy ||
 			ViewModeIndex == VMI_MaterialTextureScaleAccuracy ||
 			ViewModeIndex == VMI_RequiredTextureResolution ||
-			ViewModeIndex == VMI_VirtualTexturePendingMips ||
 			ViewModeIndex == VMI_LODColoration ||
 			ViewModeIndex == VMI_HLODColoration ||
 			ViewModeIndex == VMI_VisualizeGPUSkinCache ||
@@ -587,8 +584,7 @@ void EngineShowFlagOverride(EShowFlagInitMode ShowFlagInitMode, EViewModeIndex V
 		if (ViewModeIndex == VMI_PrimitiveDistanceAccuracy ||
 			ViewModeIndex == VMI_MeshUVDensityAccuracy ||
 			ViewModeIndex == VMI_MaterialTextureScaleAccuracy || 
-			ViewModeIndex == VMI_RequiredTextureResolution ||
-			ViewModeIndex == VMI_VirtualTexturePendingMips)
+			ViewModeIndex == VMI_RequiredTextureResolution)
 		{
 			EngineShowFlags.SetDecals(false); // Decals require the use of FDebugPSInLean.
 			EngineShowFlags.SetParticles(false); // FX are fully streamed.
@@ -780,9 +776,9 @@ EViewModeIndex FindViewMode(const FEngineShowFlags& EngineShowFlags)
 	{
 		return VMI_RequiredTextureResolution;
 	}
-	else if (EngineShowFlags.VirtualTexturePendingMips)
+	else if (EngineShowFlags.VisualizeVirtualTexture)
 	{
-		return VMI_VirtualTexturePendingMips;
+		return VMI_VisualizeVirtualTexture;
 	}
 	else if (EngineShowFlags.ShaderComplexity)
 	{
@@ -890,7 +886,6 @@ const TCHAR* GetViewModeName(EViewModeIndex ViewModeIndex)
 		case VMI_MeshUVDensityAccuracy:		return TEXT("MeshUVDensityAccuracy");
 		case VMI_MaterialTextureScaleAccuracy: return TEXT("MaterialTextureScaleAccuracy");
 		case VMI_RequiredTextureResolution: return TEXT("RequiredTextureResolution");
-		case VMI_VirtualTexturePendingMips:	return TEXT("VirtualTexturePendingMips");
 		case VMI_StationaryLightOverlap:	return TEXT("StationaryLightOverlap");
 		case VMI_LightmapDensity:			return TEXT("LightmapDensity");
 		case VMI_LitLightmapDensity:		return TEXT("LitLightmapDensity");
@@ -901,6 +896,7 @@ const TCHAR* GetViewModeName(EViewModeIndex ViewModeIndex)
 		case VMI_VisualizeSubstrate:		return TEXT("VisualizeSubstrate");
 		case VMI_VisualizeGroom:			return TEXT("VisualizeGroom");
 		case VMI_VisualizeVirtualShadowMap:	return TEXT("VisualizeVirtualShadowMap");
+		case VMI_VisualizeVirtualTexture:	return TEXT("VisualizeVirtualTexture");
 		case VMI_RayTracingDebug:			return TEXT("RayTracingDebug");
 		case VMI_PathTracing:				return TEXT("PathTracing");
 		case VMI_CollisionPawn:				return TEXT("CollisionPawn");
diff --git a/Engine/Source/Runtime/Engine/Private/UnrealEngine.cpp b/Engine/Source/Runtime/Engine/Private/UnrealEngine.cpp
--- a/Engine/Source/Runtime/Engine/Private/UnrealEngine.cpp
+++ b/Engine/Source/Runtime/Engine/Private/UnrealEngine.cpp
@@ -176,6 +176,7 @@ UnrealEngine.cpp: Implements the UEngine class and helpers.
 #include "NaniteVisualizationData.h"
 #include "VirtualShadowMapVisualizationData.h"
 #include "RayTracingVisualizationData.h"
+#include "VT/VirtualTextureVisualizationData.h"
 #include "Misc/HotReloadInterface.h"
 #include "Widgets/Testing/STestSuite.h"
 #include "Engine/DemoNetDriver.h"
@@ -2272,6 +2273,7 @@ void UEngine::Init(IEngineLoop* InEngineLoop)
 	GetNaniteVisualizationData().Initialize();
 	GetVirtualShadowMapVisualizationData().Initialize();
 	GetRayTracingVisualizationData(); // force initialization
+	GetVirtualTextureVisualizationData().Initialize();
 
 	// Initialize Portal services
 	if (!IsRunningCommandlet() && !IsRunningDedicatedServer())
diff --git a/Engine/Source/Runtime/Engine/Private/VT/VirtualTextureVisualizationData.cpp b/Engine/Source/Runtime/Engine/Private/VT/VirtualTextureVisualizationData.cpp
new file mode 100644
--- /dev/null
+++ b/Engine/Source/Runtime/Engine/Private/VT/VirtualTextureVisualizationData.cpp
@@ -0,0 +1,152 @@
+// Copyright Epic Games, Inc. All Rights Reserved.
+
+#include "VT/VirtualTextureVisualizationData.h"
+
+#include "SceneManagement.h"
+#include "SceneView.h"
+
+#define LOCTEXT_NAMESPACE "FVirtualTextureVisualizationData"
+
+static FVirtualTextureVisualizationData GVirtualTextureVisualizationData;
+
+const TCHAR* FVirtualTextureVisualizationData::GetVisualizeConsoleCommandName()
+{
+	return TEXT("r.VT.Visualize");
+}
+
+void FVirtualTextureVisualizationData::ConfigureConsoleCommand()
+{
+	FString AvailableVisualizationModes;
+	for (TModeArray::TConstIterator It = ModeArray.CreateConstIterator(); It; ++It)
+	{
+		const FModeRecord& Record = *It;
+		AvailableVisualizationModes += FString(TEXT("\n  "));
+		AvailableVisualizationModes += Record.ModeString;
+	}
+
+	ConsoleDocumentationVisualizationMode = TEXT("When the viewport view-mode is set to 'Virtual Texture Visualization', this command specifies which of the various channels to display. Values entered other than the allowed values shown below will be ignored.");
+	ConsoleDocumentationVisualizationMode += AvailableVisualizationModes;
+
+	IConsoleManager::Get().RegisterConsoleVariable(
+		GetVisualizeConsoleCommandName(),
+		TEXT(""),
+		*ConsoleDocumentationVisualizationMode,
+		ECVF_Cheat);
+}
+
+void FVirtualTextureVisualizationData::AddVisualizationMode(
+	const TCHAR* ModeString,
+	const FText& ModeText,
+	const FText& ModeDesc,
+	EVirtualTextureVisualizationMode ModeID)
+{
+	const FName ModeName = FName(ModeString);
+
+	FModeRecord& Record = ModeArray.AddDefaulted_GetRef();
+	Record.ModeString = FString(ModeString);
+	Record.ModeName = ModeName;
+	Record.ModeText = ModeText;
+	Record.ModeDesc = ModeDesc;
+	Record.ModeID = ModeID;
+}
+
+void FVirtualTextureVisualizationData::Initialize()
+{
+	if (!bIsInitialized)
+	{
+		AddVisualizationMode(
+			TEXT("pending"),
+			LOCTEXT("PendingMips", "Pending Mips"),
+			LOCTEXT("PendingMipsDesc", "The number of pending virtual texture mips to reach the resolution wanted by the GPU at a pixel"),
+			EVirtualTextureVisualizationMode::PendingMips);
+
+		AddVisualizationMode(
+			TEXT("count"),
+			LOCTEXT("StackCount", "Stack Count"),
+			LOCTEXT("StackCountDesc", "The number of virtual texture stack (page table) samples at a pixel"),
+			EVirtualTextureVisualizationMode::StackCount);
+
+		ConfigureConsoleCommand();
+
+		bIsInitialized = true;
+	}
+}
+
+FName FVirtualTextureVisualizationData::GetActiveMode(FSceneView const& InView) const
+{
+	if (IsInitialized())
+	{
+		static TConsoleVariableData<FString>* CVarViewMode = nullptr;
+		if (CVarViewMode == nullptr)
+		{
+			if (IConsoleVariable* ConsoleVariable = IConsoleManager::Get().FindConsoleVariable(GetVisualizeConsoleCommandName()))
+			{
+				CVarViewMode = ConsoleVariable->AsVariableString();
+			}
+		}
+
+		if (CVarViewMode)
+		{
+			const FString& ViewMode = CVarViewMode->GetValueOnAnyThread();
+			if (const FModeRecord* Record = ModeArray.FindByPredicate([ViewMode](const FModeRecord& Record) { return Record.ModeString == ViewMode; }))
+			{
+				return Record->ModeName;
+			}
+		}
+
+		if (InView.Family && InView.Family->EngineShowFlags.VisualizeVirtualTexture)
+		{
+			return InView.CurrentVirtualTextureVisualizationMode;
+		}
+	}
+
+	return NAME_None;
+}
+
+EVirtualTextureVisualizationMode FVirtualTextureVisualizationData::GetModeID(FName const& InModeName) const
+{
+	if (const FModeRecord* Record = ModeArray.FindByPredicate([InModeName](const FModeRecord& Record) { return Record.ModeName == InModeName; }))
+	{
+		return Record->ModeID;
+	}
+	else
+	{
+		return EVirtualTextureVisualizationMode::None;
+	}
+}
+
+FText FVirtualTextureVisualizationData::GetModeDisplayName(FName const& InModeName) const
+{
+	if (const FModeRecord* Record = ModeArray.FindByPredicate([InModeName](const FModeRecord& Record) { return Record.ModeName == InModeName; }))
+	{
+		return Record->ModeText;
+	}
+	else
+	{
+		return FText::GetEmpty();
+	}
+}
+
+FText FVirtualTextureVisualizationData::GetModeDisplayDesc(FName const& InModeName) const
+{
+	if (const FModeRecord* Record = ModeArray.FindByPredicate([InModeName](const FModeRecord& Record) { return Record.ModeName == InModeName; }))
+	{
+		return Record->ModeDesc;
+	}
+	else
+	{
+		return FText::GetEmpty();
+	}
+}
+
+FVirtualTextureVisualizationData& GetVirtualTextureVisualizationData()
+{
+	if (!GVirtualTextureVisualizationData.IsInitialized())
+	{
+		GVirtualTextureVisualizationData.Initialize();
+	}
+
+	return GVirtualTextureVisualizationData;
+}
+
+#undef LOCTEXT_NAMESPACE
diff --git a/Engine/Source/Runtime/Engine/Private/ViewModeNames.cpp b/Engine/Source/Runtime/Engine/Private/ViewModeNames.cpp
--- a/Engine/Source/Runtime/Engine/Private/ViewModeNames.cpp
+++ b/Engine/Source/Runtime/Engine/Private/ViewModeNames.cpp
@@ -168,9 +168,9 @@ TArray<FText> FillViewModeDisplayNames()
 			ViewModeDisplayNames.Emplace(LOCTEXT("UViewModeUtils_VMI_VisualizeNanite", "Nanite Visualization"));
 		}
 
-		else if (ViewModeIndex == VMI_VirtualTexturePendingMips)
+		else if (ViewModeIndex == VMI_VisualizeVirtualTexture)
 		{
-			ViewModeDisplayNames.Emplace(LOCTEXT("UViewModeUtils_VMI_VirtualTexturePendingMips", "Pending Mips"));
+			ViewModeDisplayNames.Emplace(LOCTEXT("UViewModeUtils_VMI_VisualizeVirtualTexture", "Virtual Texture Visualization"));
 		}
 
 		else if (ViewModeIndex == VMI_VisualizeLumen)
@@ -330,7 +330,7 @@ FName GetViewModeIconName(const EViewModeIndex InViewMode)
 	case VMI_VisualizeNanite:
 		Icon = VisualizeNaniteIcon;
 		break;
-	case VMI_VirtualTexturePendingMips:
+	case VMI_VisualizeVirtualTexture:
 		Icon = VisualizeVirtualTextureIcon;
 		break;
 	case VMI_VisualizeLumen:
diff --git a/Engine/Source/Runtime/Engine/Public/DebugViewModeHelpers.h b/Engine/Source/Runtime/Engine/Public/DebugViewModeHelpers.h
--- a/Engine/Source/Runtime/Engine/Public/DebugViewModeHelpers.h
+++ b/Engine/Source/Runtime/Engine/Public/DebugViewModeHelpers.h
@@ -33,7 +33,6 @@ enum EDebugViewShaderMode
 	DVSM_MaterialTextureScaleAccuracy, // Visualize the accuracy of the material texture scales used for texture streaming.
 	DVSM_OutputMaterialTextureScales,  // Outputs the material texture scales.
 	DVSM_RequiredTextureResolution, // Visualize the accuracy of the streamed texture resolution.
-	DVSM_VirtualTexturePendingMips,	// Visualize the pending virtual texture mips.
 	DVSM_LODColoration,				// Visualize primitive LOD .
 	DVSM_VisualizeGPUSkinCache,		// Visualize various properties of Skin Cache.
 	DVSM_LWCComplexity,				// Visualize usage of LWC functions in materials.
diff --git a/Engine/Source/Runtime/Engine/Public/SceneView.h b/Engine/Source/Runtime/Engine/Public/SceneView.h
--- a/Engine/Source/Runtime/Engine/Public/SceneView.h
+++ b/Engine/Source/Runtime/Engine/Public/SceneView.h
@@ -1051,6 +1051,7 @@ enum ETranslucencyVolumeCascade
 	VIEW_UNIFORM_BUFFER_MEMBER(uint32, VirtualTextureFeedbackMask) \
 	VIEW_UNIFORM_BUFFER_MEMBER(uint32, VirtualTextureFeedbackJitterOffset) \
 	VIEW_UNIFORM_BUFFER_MEMBER(uint32, VirtualTextureFeedbackSampleOffset) \
+	VIEW_UNIFORM_BUFFER_MEMBER(uint32, VirtualTextureExtendedDebugBufferSize) \
 	VIEW_UNIFORM_BUFFER_MEMBER(int32, FarShadowStaticMeshLODBias) \
 	VIEW_UNIFORM_BUFFER_MEMBER(float, MinRoughness) \
 	VIEW_UNIFORM_BUFFER_MEMBER(FVector4f, HairRenderInfo) \
@@ -1487,6 +1488,9 @@ public:
 	/** Current Virtual Shadow Map visualization mode */
 	FName CurrentVirtualShadowMapVisualizationMode;
 
+	/** Current Virtual Texture visualization mode */
+	FName CurrentVirtualTextureVisualizationMode;
+
 	/** Current visualize calibration color material name */
 	FName CurrentVisualizeCalibrationColorMaterialName;
 
diff --git a/Engine/Source/Runtime/Engine/Public/ShowFlags.h b/Engine/Source/Runtime/Engine/Public/ShowFlags.h
--- a/Engine/Source/Runtime/Engine/Public/ShowFlags.h
+++ b/Engine/Source/Runtime/Engine/Public/ShowFlags.h
@@ -407,6 +407,7 @@ private:
 		SetVisualizeSubstrate(false);
 		SetVisualizeGroom(false);
 		SetVisualizeVirtualShadowMap(false);
+		SetVisualizeVirtualTexture(false);
 		SetVectorFields(false);
 		SetGBufferHints(false);
 		SetCompositeEditorPrimitives(InitMode == ESFIM_Editor || InitMode == ESFIM_VREditing);
@@ -496,7 +497,6 @@ private:
 		SetMaterialTextureScaleAccuracy(false);
 		SetOutputMaterialTextureScales(false);
 		SetRequiredTextureResolution(false);
-		SetVirtualTexturePendingMips(false);
 		SetMotionBlur(InitMode != ESFIM_Editor && InitMode != ESFIM_VREditing);
 		SetBones(false);
 		SetServerDrawDebug(false);
diff --git a/Engine/Source/Runtime/Engine/Public/ShowFlagsValues.inl b/Engine/Source/Runtime/Engine/Public/ShowFlagsValues.inl
--- a/Engine/Source/Runtime/Engine/Public/ShowFlagsValues.inl
+++ b/Engine/Source/Runtime/Engine/Public/ShowFlagsValues.inl
@@ -378,7 +378,7 @@ SHOWFLAG_FIXED_IN_SHIPPING(0, OutputMaterialTextureScales, SFG_Hidden, NSLOCTEXT
 /** Compare the required texture resolution to the actual resolution. */
 SHOWFLAG_FIXED_IN_SHIPPING(0, RequiredTextureResolution, SFG_Hidden, NSLOCTEXT("UnrealEd", "RequiredTextureResolution", "Required Texture Resolution"))
 /** Compare the required mip level to the resident virtual texture mip level. */
-SHOWFLAG_FIXED_IN_SHIPPING(0, VirtualTexturePendingMips, SFG_Hidden, NSLOCTEXT("UnrealEd", "VirtualTexturePendingMips", "Virtual Texture Pending Mips"))
+SHOWFLAG_FIXED_IN_SHIPPING(0, VisualizeVirtualTexture, SFG_Hidden, NSLOCTEXT("UnrealEd", "VisualizeVirtualTextureSF", "Virtual Texture Visualization"))
 /** If WidgetComponents should be rendered in the scene */
 SHOWFLAG_ALWAYS_ACCESSIBLE(WidgetComponents, SFG_Normal, NSLOCTEXT("UnrealEd", "WidgetComponentsSF", "Widget Components"))
 /** Draw the bones of all skeletal meshes */
diff --git a/Engine/Source/Runtime/Renderer/Internal/VT/VirtualTextureFeedbackResource.h b/Engine/Source/Runtime/Renderer/Internal/VT/VirtualTextureFeedbackResource.h
--- a/Engine/Source/Runtime/Renderer/Internal/VT/VirtualTextureFeedbackResource.h
+++ b/Engine/Source/Runtime/Renderer/Internal/VT/VirtualTextureFeedbackResource.h
@@ -5,6 +5,7 @@
 #include "CoreMinimal.h"
 #include "RHI.h"
 
+class FRDGBuffer;
 class FRDGBuilder;
 class FViewUniformShaderParameters;
 
@@ -19,6 +20,7 @@ namespace VirtualTexture
 		uint32 TileMask = 0;
 		uint32 TileJitterOffset = 0;
 		uint32 SampleOffset = 0;
+		uint32 ExtendedDebugBufferSize = 0;
 	};
 
 	/** 
@@ -31,9 +33,10 @@ namespace VirtualTexture
 	/** 
 	 * Begin a virtual texture feedback scope. 
 	 * InViewportSize and InVirtualTextureFeedbackTileSize are used to calculate the feedback buffer size. 
+	 * If bInExtendFeedbackForDebug is true we will allocate room at the end of the feedback buffer to capture extra debug info.
 	 * Note that there is only one global feedback buffer alive at any time and its buffer size is later used internally by GetFeedbackShaderParams().
 	 */
-	void BeginFeedback(FRDGBuilder& GraphBuilder, FIntPoint InViewportSize, uint32 InVirtualTextureFeedbackTileSize, ERHIFeatureLevel::Type InFeatureLevel);
+	void BeginFeedback(FRDGBuilder& GraphBuilder, FIntPoint InViewportSize, uint32 InVirtualTextureFeedbackTileSize, bool bInExtendFeedbackForDebug, ERHIFeatureLevel::Type InFeatureLevel);
 
 	/** End a virtual texture feedback scope. */
 	RENDERER_API void EndFeedback(FRDGBuilder& GraphBuilder);
@@ -53,4 +56,10 @@ namespace VirtualTexture
 	
 	/** Helper to copy the virtual texture feedback shader parameters into the view parameters to be bound in the view uniform buffer. */
 	RENDERER_API void UpdateViewUniformShaderParameters(FFeedbackShaderParams const& InParams, FViewUniformShaderParameters& ViewUniformShaderParameters);
+
+	/** 
+	 * This resolves any extended debug information that is currently stored at the end of the feedback buffer and returns a buffer that contains it.
+	 * This is only valid within a BeginFeedback()/EndFeedback() scope.
+	 */
+	FRDGBuffer* ResolveExtendedDebugBuffer(FRDGBuilder& GraphBuilder);
 }
diff --git a/Engine/Source/Runtime/Renderer/Private/BasePassRendering.cpp b/Engine/Source/Runtime/Renderer/Private/BasePassRendering.cpp
--- a/Engine/Source/Runtime/Renderer/Private/BasePassRendering.cpp
+++ b/Engine/Source/Runtime/Renderer/Private/BasePassRendering.cpp
@@ -1089,7 +1089,6 @@ void FDeferredShadingSceneRenderer::RenderBasePass(
 	const bool bDoParallelBasePass = bEnableParallelBasePasses && !bDebugViewMode && !bRenderLightmapDensity; // DebugView and LightmapDensity are non-parallel substitutions inside BasePass
 	const bool bNeedsBeginRender = AllowDebugViewmodes() &&
 		(ViewFamily.EngineShowFlags.RequiredTextureResolution ||
-			ViewFamily.EngineShowFlags.VirtualTexturePendingMips ||
 			ViewFamily.EngineShowFlags.MaterialTextureScaleAccuracy ||
 			ViewFamily.EngineShowFlags.MeshUVDensityAccuracy ||
 			ViewFamily.EngineShowFlags.PrimitiveDistanceAccuracy ||
diff --git a/Engine/Source/Runtime/Renderer/Private/DebugViewModeRendering.cpp b/Engine/Source/Runtime/Renderer/Private/DebugViewModeRendering.cpp
--- a/Engine/Source/Runtime/Renderer/Private/DebugViewModeRendering.cpp
+++ b/Engine/Source/Runtime/Renderer/Private/DebugViewModeRendering.cpp
@@ -282,7 +282,6 @@ void FDebugViewModePS::ModifyCompilationEnvironment(const FMaterialShaderPermuta
 	BufferRegister[1] += GetQuadOverdrawUAVIndex(Parameters.Platform);
 	OutEnvironment.SetDefine(TEXT("QUAD_BUFFER_REGISTER"), BufferRegister);
 	OutEnvironment.SetDefine(TEXT("OUTPUT_QUAD_OVERDRAW"), SupportDebugViewShaderMode(DVSM_QuadComplexity, Parameters.Platform));
-	OutEnvironment.SetDefine(TEXT("WRITE_PENDING_VIRTUALTEXTURE_MIPS"), 1u);
 
 	for (int i = 0; i < DVSM_MAX; ++i)
 	{
@@ -566,7 +565,7 @@ void FDebugViewModeImplementation::GetDebugViewModeShaderBindings(
 		const int32 AnalysisIndex = ViewModeParam >= 0 ? FMath::Clamp<int32>(ViewModeParam, 0, TEXSTREAM_MAX_NUM_TEXTURES_PER_MATERIAL - 1) : -1;
 		AnalysisParameter = FIntPoint(bOutputScales ? -1 : AnalysisIndex, bOutputScales ? 1 : 0);
 	}
-	else if (DebugViewMode == DVSM_RequiredTextureResolution || DebugViewMode == DVSM_VirtualTexturePendingMips)
+	else if (DebugViewMode == DVSM_RequiredTextureResolution)
 	{
 		int32 AnalysisIndex = INDEX_NONE;
 		int32 TextureResolution = 64;
diff --git a/Engine/Source/Runtime/Renderer/Private/PostProcess/PostProcessStreamingAccuracyLegend.cpp b/Engine/Source/Runtime/Renderer/Private/PostProcess/PostProcessStreamingAccuracyLegend.cpp
--- a/Engine/Source/Runtime/Renderer/Private/PostProcess/PostProcessStreamingAccuracyLegend.cpp
+++ b/Engine/Source/Runtime/Renderer/Private/PostProcess/PostProcessStreamingAccuracyLegend.cpp
@@ -55,10 +55,6 @@ FScreenPassTexture AddStreamingAccuracyLegendPass(FRDGBuilder& GraphBuilder, con
 		{
 			DrawDesc(Canvas, OutputViewRect.Min.X + 115, OutputViewRect.Max.Y - 75, LOCTEXT("DescRequiredTextureResolution", "Shows the ratio between the currently streamed texture resolution and the resolution wanted by the GPU."));
 		}
-		else if (ViewFamily.GetDebugViewShaderMode() == DVSM_VirtualTexturePendingMips)
-		{
-			DrawDesc(Canvas, OutputViewRect.Min.X + 115, OutputViewRect.Max.Y - 75, LOCTEXT("DescVirtualTexturePendingMips", "Shows the number of pending virtual texture mips to reach the resolution wanted by the GPU."));
-		}
 		else if (ViewFamily.GetDebugViewShaderMode() == DVSM_MaterialTextureScaleAccuracy)
 		{
 			DrawDesc(Canvas, OutputViewRect.Min.X + 115, OutputViewRect.Max.Y - 75, LOCTEXT("DescMaterialTextureScaleAccuracy", "Shows under/over texture streaming caused by the material texture scales applied when sampling."));
@@ -77,14 +73,11 @@ FScreenPassTexture AddStreamingAccuracyLegendPass(FRDGBuilder& GraphBuilder, con
 		DrawBox(Canvas, OutputViewRect.Min.X + 215, OutputViewRect.Max.Y - 25, Colors[1], LOCTEXT("1XUnder", "1X Under"));
 		DrawBox(Canvas, OutputViewRect.Min.X + 315, OutputViewRect.Max.Y - 25, Colors[2], LOCTEXT("Good", "Good"));
 
-		if (ViewFamily.GetDebugViewShaderMode() != DVSM_VirtualTexturePendingMips)
-		{
-			DrawBox(Canvas, OutputViewRect.Min.X + 415, OutputViewRect.Max.Y - 25, Colors[3], LOCTEXT("1xOver", "1X Over"));
-			DrawBox(Canvas, OutputViewRect.Min.X + 515, OutputViewRect.Max.Y - 25, Colors[4], LOCTEXT("2XOver", "2X+ Over"));
+		DrawBox(Canvas, OutputViewRect.Min.X + 415, OutputViewRect.Max.Y - 25, Colors[3], LOCTEXT("1xOver", "1X Over"));
+		DrawBox(Canvas, OutputViewRect.Min.X + 515, OutputViewRect.Max.Y - 25, Colors[4], LOCTEXT("2XOver", "2X+ Over"));
 
-			const FLinearColor UndefColor(UndefinedStreamingAccuracyIntensity, UndefinedStreamingAccuracyIntensity, UndefinedStreamingAccuracyIntensity, 1.f);
-			DrawBox(Canvas, OutputViewRect.Min.X + 615, OutputViewRect.Max.Y - 25, UndefColor, LOCTEXT("Undefined", "Undefined"));
-		}
+		const FLinearColor UndefColor(UndefinedStreamingAccuracyIntensity, UndefinedStreamingAccuracyIntensity, UndefinedStreamingAccuracyIntensity, 1.f);
+		DrawBox(Canvas, OutputViewRect.Min.X + 615, OutputViewRect.Max.Y - 25, UndefColor, LOCTEXT("Undefined", "Undefined"));
 
 		if (ViewFamily.GetDebugViewShaderMode() == DVSM_MaterialTextureScaleAccuracy || ViewFamily.GetDebugViewShaderMode() == DVSM_MeshUVDensityAccuracy)
 		{
diff --git a/Engine/Source/Runtime/Renderer/Private/PostProcess/PostProcessVisualizeVirtualTexture.cpp b/Engine/Source/Runtime/Renderer/Private/PostProcess/PostProcessVisualizeVirtualTexture.cpp
new file mode 100644
--- /dev/null
+++ b/Engine/Source/Runtime/Renderer/Private/PostProcess/PostProcessVisualizeVirtualTexture.cpp
@@ -0,0 +1,126 @@
+// Copyright Epic Games, Inc. All Rights Reserved.
+
+#include "PostProcess/PostProcessVisualizeVirtualTexture.h"
+
+#include "CanvasTypes.h"
+#include "SceneRendering.h"
+#include "UnrealEngine.h"
+#include "VT/VirtualTextureVisualizationData.h"
+
+#define LOCTEXT_NAMESPACE "VisualizeVirtualTexture"
+
+class FVisualizeVirtualTextureApplyPS : public FGlobalShader
+{
+public:
+	DECLARE_GLOBAL_SHADER(FVisualizeVirtualTextureApplyPS);
+	SHADER_USE_PARAMETER_STRUCT(FVisualizeVirtualTextureApplyPS, FGlobalShader);
+
+	static constexpr int32 MaxNumColors = 11;
+
+	BEGIN_SHADER_PARAMETER_STRUCT(FParameters, )
+		SHADER_PARAMETER_STRUCT(FScreenPassTextureViewportParameters, Input)
+		SHADER_PARAMETER_STRUCT(FScreenPassTextureViewportParameters, Output)
+		SHADER_PARAMETER_RDG_TEXTURE(Texture2D, InputTexture)
+		SHADER_PARAMETER_SAMPLER(SamplerState, InputSampler)
+		SHADER_PARAMETER_RDG_BUFFER_SRV(StructuredBuffer<uint>, DebugBuffer)
+		SHADER_PARAMETER(FIntPoint, ViewSize)
+		SHADER_PARAMETER(uint32, ViewMode)
+		SHADER_PARAMETER_ARRAY(FLinearColor, Colors, [MaxNumColors])
+		SHADER_PARAMETER(uint32, ColorCount)
+		RENDER_TARGET_BINDING_SLOTS()
+	END_SHADER_PARAMETER_STRUCT()
+
+	static bool ShouldCompilePermutation(const FGlobalShaderPermutationParameters& Parameters)
+	{
+		return true;
+	}
+
+	static void ModifyCompilationEnvironment(const FGlobalShaderPermutationParameters& Parameters, FShaderCompilerEnvironment& OutEnvironment)
+	{
+		FGlobalShader::ModifyCompilationEnvironment(Parameters, OutEnvironment);
+		OutEnvironment.SetDefine(TEXT("MAX_NUM_COLORS"), MaxNumColors);
+	}
+};
+
+IMPLEMENT_GLOBAL_SHADER(FVisualizeVirtualTextureApplyPS, "/Engine/Private/VisualizeVirtualTexture.usf", "Main", SF_Pixel);
+
+FScreenPassTexture AddVisualizeVirtualTexturePass(FRDGBuilder& GraphBuilder, const FViewInfo& View, const FVisualizeVirtualTextureInputs& Inputs)
+{
+	check(Inputs.SceneColor.IsValid());
+
+	FScreenPassRenderTarget Output = Inputs.OverrideOutput;
+
+	if (!Output.IsValid())
+	{
+		Output = FScreenPassRenderTarget::CreateFromInput(GraphBuilder, Inputs.SceneColor, View.GetOverwriteLoadAction(), TEXT("VisualizeVirtualTexture"));
+	}
+
+	const FScreenPassTextureViewport InputViewport(Inputs.SceneColor);
+	const FScreenPassTextureViewport OutputViewport(Output);
+
+	FVisualizeVirtualTextureApplyPS::FParameters* PassParameters = GraphBuilder.AllocParameters<FVisualizeVirtualTextureApplyPS::FParameters>();
+	PassParameters->RenderTargets[0] = Output.GetRenderTargetBinding();
+	PassParameters->Input = GetScreenPassTextureViewportParameters(InputViewport);
+	PassParameters->Output = GetScreenPassTextureViewportParameters(OutputViewport);
+	PassParameters->InputTexture = Inputs.SceneColor.Texture;
+	PassParameters->InputSampler = TStaticSamplerState<SF_Bilinear, AM_Clamp, AM_Clamp, AM_Clamp>::GetRHI();
+	PassParameters->DebugBuffer = GraphBuilder.CreateSRV(Inputs.DebugBuffer);
+	PassParameters->ViewSize = View.ViewRect.Size();
+
+	const int32 ClampedColorCount = FMath::Min<int32>(Inputs.Colors.Num(), FVisualizeVirtualTextureApplyPS::MaxNumColors);
+	if (ClampedColorCount > 0)
+	{
+		for (int32 ColorIndex = 0; ColorIndex < ClampedColorCount; ColorIndex++)
+		{
+			PassParameters->Colors[ColorIndex] = Inputs.Colors[ColorIndex];
+		}
+		PassParameters->ColorCount = ClampedColorCount;
+	}
+	else // Otherwise fallback to a safe value.
+	{
+		PassParameters->Colors[0] = FLinearColor::Gray;
+		PassParameters->ColorCount = 1;
+	}
+
+	const EVirtualTextureVisualizationMode Mode = GetVirtualTextureVisualizationData().GetModeID(Inputs.ModeName);
+	const FText ModeDesc = GetVirtualTextureVisualizationData().GetModeDisplayDesc(Inputs.ModeName);
+	
+	PassParameters->ViewMode = (uint32)Mode;
+
+	TShaderMapRef<FVisualizeVirtualTextureApplyPS> PixelShader(View.ShaderMap);
+
+	AddDrawScreenPass(GraphBuilder, RDG_EVENT_NAME("VisualizeVirtualTexture"), View, FScreenPassTextureViewport(Output), InputViewport, PixelShader, PassParameters);
+
+	Output.LoadAction = ERenderTargetLoadAction::ELoad;
+
+	const TArrayView<const FLinearColor> Colors = Inputs.Colors;
+	const FIntRect OutputViewRect = Output.ViewRect;
+
+	AddDrawCanvasPass(GraphBuilder, RDG_EVENT_NAME("VisualizeVirtualTextureLegend"), View, Output, [&View, ModeDesc, Colors, ClampedColorCount, OutputViewRect](FCanvas& Canvas)
+	{
+		auto DrawDesc = [&](FCanvas& Canvas, float PosX, float PosY, const FText& Text)
+		{
+			Canvas.DrawShadowedText(PosX + 18, PosY, Text, GetStatsFont(), FLinearColor(0.7f, 0.7f, 0.7f), FLinearColor::Black);
+		};
+
+		auto DrawBox = [&](FCanvas& Canvas, float PosX, float PosY, const FLinearColor& Color, const FText& Text)
+		{
+			Canvas.DrawTile(PosX, PosY, 16, 16, 0, 0, 1, 1, FLinearColor::Black);
+			Canvas.DrawTile(PosX + 1, PosY + 1, 14, 14, 0, 0, 1, 1, Color);
+			Canvas.DrawShadowedText(PosX + 18, PosY, Text, GetStatsFont(), FLinearColor(0.7f, 0.7f, 0.7f), FLinearColor::Black);
+		};
+
+		int32 CoordX = 115;
+		DrawDesc(Canvas, OutputViewRect.Min.X + CoordX, OutputViewRect.Max.Y - 75, ModeDesc);
+
+		for (int32 Index = 0; Index < ClampedColorCount; ++Index)
+		{
+			DrawBox(Canvas, OutputViewRect.Min.X + CoordX, OutputViewRect.Max.Y - 25, Colors[Index], FText::AsNumber(Index + 1));
+			CoordX += 50;
+		}
+	});
+
+	return MoveTemp(Output);
+}
+
+#undef LOCTEXT_NAMESPACE
\ No newline at end of file
diff --git a/Engine/Source/Runtime/Renderer/Private/PostProcess/PostProcessVisualizeVirtualTexture.h b/Engine/Source/Runtime/Renderer/Private/PostProcess/PostProcessVisualizeVirtualTexture.h
new file mode 100644
--- /dev/null
+++ b/Engine/Source/Runtime/Renderer/Private/PostProcess/PostProcessVisualizeVirtualTexture.h
@@ -0,0 +1,23 @@
+// Copyright Epic Games, Inc. All Rights Reserved.
+
+#pragma once
+
+#include "ScreenPass.h"
+#include "DebugViewModeHelpers.h"
+#include "OverridePassSequence.h"
+
+struct FVisualizeVirtualTextureInputs
+{
+	// [Optional] Render to the specified output. If invalid, a new texture is created and returned.
+	FScreenPassRenderTarget OverrideOutput;
+	// [Required] The input scene color and view rect.
+	FScreenPassTexture SceneColor;
+	// [Required] The buffer generated by main pass rendering which contains debug data.
+	FRDGBufferRef DebugBuffer;
+	// [Required] The visualization mode.
+	FName ModeName;
+	// [Required] The table of colors used for visualization.
+	TArrayView<const FLinearColor> Colors;
+};
+
+FScreenPassTexture AddVisualizeVirtualTexturePass(FRDGBuilder& GraphBuilder, const FViewInfo& View, const FVisualizeVirtualTextureInputs& Inputs);
\ No newline at end of file
diff --git a/Engine/Source/Runtime/Renderer/Private/PostProcess/PostProcessing.cpp b/Engine/Source/Runtime/Renderer/Private/PostProcess/PostProcessing.cpp
--- a/Engine/Source/Runtime/Renderer/Private/PostProcess/PostProcessing.cpp
+++ b/Engine/Source/Runtime/Renderer/Private/PostProcess/PostProcessing.cpp
@@ -34,6 +34,7 @@
 #include "PostProcess/PostProcessHMD.h"
 #include "PostProcess/AlphaInvert.h"
 #include "PostProcess/PostProcessVisualizeComplexity.h"
+#include "PostProcess/PostProcessVisualizeVirtualTexture.h"
 #if UE_ENABLE_DEBUG_DRAWING
 	#include "PostProcess/PostProcessCompositeDebugPrimitives.h"
 #endif
@@ -78,6 +79,8 @@
 #include "Materials/MaterialRenderProxy.h"
 #include "CustomRenderPassSceneCapture.h"
 #include "GPUSkinCache.h"
+#include "VT/VirtualTextureFeedbackResource.h"
+#include "VT/VirtualTextureVisualizationData.h"
 
 bool IsMobileEyeAdaptationEnabled(const FViewInfo& View);
 
@@ -417,6 +420,7 @@ void AddPostProcessingPasses(
 
 	FVisualizeTemporalUpscalerInputs VisualizeTemporalUpscalerInputs;
 
+	const bool bViewDebugMaterialsEnabled = View.RequiresDebugMaterials();
 	const FEngineShowFlags& EngineShowFlags = View.Family->EngineShowFlags;
 	const bool bVisualizeHDR = EngineShowFlags.VisualizeHDR;
 	const bool bViewFamilyOutputInHDR = View.Family->RenderTarget->GetSceneHDREnabled();
@@ -455,6 +459,7 @@ void AddPostProcessingPasses(
 		SelectionOutline,
 		EditorPrimitive,
 		VisualizeVirtualShadowMaps_PostEditorPrimitives,
+		VisualizeVirtualTexture,
 		VisualizeShadingModels,
 		VisualizeGBufferHints,
 		VisualizeSubsurface,
@@ -513,6 +518,7 @@ void AddPostProcessingPasses(
 		TEXT("SelectionOutline"),
 		TEXT("EditorPrimitive"),
 		TEXT("VisualizeVirtualShadowMaps_PostEditorPrimitives"),
+		TEXT("VisualizeVirtualTexture"),
 		TEXT("VisualizeShadingModels"),
 		TEXT("VisualizeGBufferHints"),
 		TEXT("VisualizeSubsurface"),
@@ -572,6 +578,7 @@ void AddPostProcessingPasses(
 	PassSequence.SetEnabled(EPass::VisualizeVirtualShadowMaps_PostEditorPrimitives, false);
 #endif
 
+	PassSequence.SetEnabled(EPass::VisualizeVirtualTexture, EngineShowFlags.VisualizeVirtualTexture && bViewDebugMaterialsEnabled);
 	PassSequence.SetEnabled(EPass::VisualizeShadingModels, EngineShowFlags.VisualizeShadingModels);
 	PassSequence.SetEnabled(EPass::VisualizeGBufferHints, EngineShowFlags.GBufferHints);
 	PassSequence.SetEnabled(EPass::VisualizeSubsurface, EngineShowFlags.VisualizeSSS);
@@ -1724,6 +1731,20 @@ void AddPostProcessingPasses(
 		SceneColor = VirtualShadowMapArray->AddVisualizePass(GraphBuilder, View, ViewIndex, EVSMVisualizationPostPass::PostEditorPrimitives, SceneColor, OverrideOutput);
 	}
 
+	if (PassSequence.IsEnabled(EPass::VisualizeVirtualTexture))
+	{
+		if (FRDGBuffer* DebugBuffer = VirtualTexture::ResolveExtendedDebugBuffer(GraphBuilder))
+		{
+			FVisualizeVirtualTextureInputs PassInputs;
+			PassSequence.AcceptOverrideIfLastPass(EPass::VisualizeVirtualTexture, PassInputs.OverrideOutput);
+			PassInputs.SceneColor = SceneColor;
+			PassInputs.DebugBuffer = DebugBuffer;
+			PassInputs.ModeName = GetVirtualTextureVisualizationData().GetActiveMode(View);
+			PassInputs.Colors = GEngine->ShaderComplexityColors;
+
+			SceneColor = AddVisualizeVirtualTexturePass(GraphBuilder, View, PassInputs);
+		}
+	}
 
 	if (PassSequence.IsEnabled(EPass::VisualizeShadingModels))
 	{
@@ -2186,7 +2207,6 @@ void AddDebugViewPostProcessingPasses(FRDGBuilder& GraphBuilder, const FViewInfo
 		case DVSM_MeshUVDensityAccuracy:
 		case DVSM_MaterialTextureScaleAccuracy:
 		case DVSM_RequiredTextureResolution:
-		case DVSM_VirtualTexturePendingMips:
 		{
 			FStreamingAccuracyLegendInputs PassInputs;
 			PassInputs.OverrideOutput = OverrideOutput;
diff --git a/Engine/Source/Runtime/Renderer/Private/SceneRendering.cpp b/Engine/Source/Runtime/Renderer/Private/SceneRendering.cpp
--- a/Engine/Source/Runtime/Renderer/Private/SceneRendering.cpp
+++ b/Engine/Source/Runtime/Renderer/Private/SceneRendering.cpp
@@ -472,12 +472,6 @@ static TAutoConsoleVariable<int32> CVarCrossGPUTransferOption(
 	TEXT(" 2: optimized pull transfer (destination GPU runs copy, with transfers delayed to last view's render); (default)"),
 	ECVF_Default);
 
-static TAutoConsoleVariable<bool> CVarForceDebugMaterials(
-	TEXT("r.ForceDebugMaterials"),
-	0,
-	TEXT("Force the scene to request the use of (slow) debug visualization shader permutations. Used only for testing shader reload logic."),
-	ECVF_RenderThreadSafe);
-
 FOcclusionSubmittedFenceState FSceneRenderer::OcclusionSubmittedFence[FOcclusionQueryHelpers::MaxBufferedOcclusionFrames];
 
 // cleanup OcclusionSubmittedFence to avoid undefined order of destruction that can destroy it after its allocator
@@ -2467,8 +2461,9 @@ bool FViewInfo::RequiresDebugMaterials() const
 {
 	// We can only use debug materials in ODSC environments.
 	static const bool bCanUseDebugMaterials = ShouldCompileODSCOnlyShaders();
-	// todo: Add debug modes that need to use scene material debug materials.
-	return bCanUseDebugMaterials && CVarForceDebugMaterials.GetValueOnAnyThread();
+	// Add other debug modes here as required.
+	return bCanUseDebugMaterials && 
+		(Family != nullptr && Family->EngineShowFlags.VisualizeVirtualTexture);
 }
 
 void FDisplayInternalsData::Setup(UWorld *World)
@@ -6069,7 +6064,9 @@ void VirtualTextureFeedbackBegin(FRDGBuilder& GraphBuilder, TArrayView<const FVi
 {
 	const uint32 FeedbackTileSize = Views.Num() > 0 ? Views[0].Family->VirtualTextureFeedbackFactor : 0;
 	const ERHIFeatureLevel::Type FeatureLevel = Views.Num() > 0 ? Views[0].GetFeatureLevel() : GMaxRHIFeatureLevel;
-	VirtualTexture::BeginFeedback(GraphBuilder, SceneTextureExtent, FeedbackTileSize, FeatureLevel);
+	static const bool bCanUseDebugMaterials = ShouldCompileODSCOnlyShaders();
+	const bool bExtendFeedbackForDebug = Views.Num() > 0 ? Views[0].Family->EngineShowFlags.VisualizeVirtualTexture && bCanUseDebugMaterials : false;
+	VirtualTexture::BeginFeedback(GraphBuilder, SceneTextureExtent, FeedbackTileSize, bExtendFeedbackForDebug, FeatureLevel);
 }
 
 static TAutoConsoleVariable<int32> CVarHalfResDepthNoFastClear(
diff --git a/Engine/Source/Runtime/Renderer/Private/VT/VirtualTextureFeedbackResource.cpp b/Engine/Source/Runtime/Renderer/Private/VT/VirtualTextureFeedbackResource.cpp
--- a/Engine/Source/Runtime/Renderer/Private/VT/VirtualTextureFeedbackResource.cpp
+++ b/Engine/Source/Runtime/Renderer/Private/VT/VirtualTextureFeedbackResource.cpp
@@ -39,12 +39,17 @@ class FVirtualTextureFeedbackBufferResource : public FRenderResource
 {
 public:
 	/** Allocates and prepares a new feedback buffer for write access. */
-	void Begin(FRDGBuilder& GraphBuilder, uint32 InFeedbackBufferSize, ERHIFeatureLevel::Type InFeatureLevel);
+	void Begin(FRDGBuilder& GraphBuilder, uint32 InFeedbackBufferSize, uint32 InExtendedDebugBufferSize, ERHIFeatureLevel::Type InFeatureLevel);
 	/** Compacts the feedback buffer on the GPU and copies the compacted buffer for readback. */
 	void End(FRDGBuilder& GraphBuilder);
 
+	/** Resolve and return any extended debug information that is currently stored at the end of the feedback buffer. */
+	FRDGBuffer* ResolveExtendedDebugBuffer(FRDGBuilder& GraphBuilder);
+
 	/** Get the size of the feedback buffer from the last call to Begin(). */
 	uint32 GetBufferSize() const;
+	/** Get the added size for the extended debug area of the feedback buffer from the last call to Begin(). */
+	uint32 GetExtendedDebugBufferSize() const;
 	/** Get the UAV of the feedback buffer. */
 	FRHIUnorderedAccessView* GetUAV() const;
 
@@ -55,6 +60,7 @@ public:
 private:
 	bool bIsInBeginEndScope = false;
 	uint32 FeedbackBufferSize = 0;
+	uint32 ExtendedDebugBufferSize = 0;
 	ERHIFeatureLevel::Type FeatureLevel = ERHIFeatureLevel::Num;
 	TRefCountPtr<FRDGPooledBuffer> PooledBuffer;
 	FRHIUnorderedAccessView* UAV = nullptr;
@@ -65,16 +71,17 @@ static TGlobalResource<FVirtualTextureFeedbackBufferResource> GVirtualTextureFee
 
 static uint32 GetVirtualTextureCompactedFeedbackBufferSize(uint32 InSourceBufferSize);
 
-void FVirtualTextureFeedbackBufferResource::Begin(FRDGBuilder& GraphBuilder, uint32 InFeedbackBufferSize, ERHIFeatureLevel::Type InFeatureLevel)
+void FVirtualTextureFeedbackBufferResource::Begin(FRDGBuilder& GraphBuilder, uint32 InFeedbackBufferSize, uint32 InExtendedDebugBufferSize, ERHIFeatureLevel::Type InFeatureLevel)
 {
 	// NOTE: Transitions and allocations are handled manually right now, because the VT feedback UAV is used by
 	// the view uniform buffer, which is not an RDG uniform buffer. If it can be factored out into its own RDG
 	// uniform buffer (or put on the pass uniform buffers), then the resource can be fully converted to RDG.
 
 	FeedbackBufferSize = InFeedbackBufferSize;
+	ExtendedDebugBufferSize = InExtendedDebugBufferSize;
 	FeatureLevel = InFeatureLevel;
 
-	FRDGBufferDesc BufferDesc(FRDGBufferDesc::CreateStructuredDesc(sizeof(uint32), FeedbackBufferSize));
+	FRDGBufferDesc BufferDesc(FRDGBufferDesc::CreateStructuredDesc(sizeof(uint32), FeedbackBufferSize + InExtendedDebugBufferSize));
 	BufferDesc.Usage |= BUF_SourceCopy;
 
 	AllocatePooledBuffer(BufferDesc, PooledBuffer, TEXT("VirtualTexture_FeedbackBuffer"));
@@ -83,13 +90,13 @@ void FVirtualTextureFeedbackBufferResource::Begin(FRDGBuilder& GraphBuilder, uin
 	UAV = PooledBuffer->GetOrCreateUAV(GraphBuilder.RHICmdList, UAVDesc);
 
 	AddPass(GraphBuilder, RDG_EVENT_NAME("VirtualTextureClear"), [UAV = UAV](FRDGAsyncTask, FRHICommandList& RHICmdList)
-		{
-			// Clear virtual texture feedback to default value
-			RHICmdList.Transition(FRHITransitionInfo(UAV, ERHIAccess::Unknown, ERHIAccess::UAVCompute));
-			RHICmdList.ClearUAVUint(UAV, FUintVector4(0u, 0u, 0u, 0u));
-			RHICmdList.Transition(FRHITransitionInfo(UAV, ERHIAccess::UAVCompute, ERHIAccess::UAVMask));
-			RHICmdList.BeginUAVOverlap(UAV);
-		});
+	{
+		// Clear virtual texture feedback to default value
+		RHICmdList.Transition(FRHITransitionInfo(UAV, ERHIAccess::Unknown, ERHIAccess::UAVCompute));
+		RHICmdList.ClearUAVUint(UAV, FUintVector4(0u, 0u, 0u, 0u));
+		RHICmdList.Transition(FRHITransitionInfo(UAV, ERHIAccess::UAVCompute, ERHIAccess::UAVMask));
+		RHICmdList.BeginUAVOverlap(UAV);
+	});
 
 	bIsInBeginEndScope = true;
 }
@@ -103,10 +110,10 @@ void FVirtualTextureFeedbackBufferResource::End(FRDGBuilder& GraphBuilder)
 	bIsInBeginEndScope = false;
 
 	AddPass(GraphBuilder, RDG_EVENT_NAME("VirtualTextureFeedbackTransition"), [UAV = UAV](FRHICommandListImmediate& RHICmdList)
-		{
-			RHICmdList.EndUAVOverlap(UAV);
-			RHICmdList.Transition(FRHITransitionInfo(UAV, ERHIAccess::UAVMask, ERHIAccess::SRVCompute));
-		});
+	{
+		RHICmdList.EndUAVOverlap(UAV);
+		RHICmdList.Transition(FRHITransitionInfo(UAV, ERHIAccess::UAVMask, ERHIAccess::SRVCompute));
+	});
 
 	FRDGBufferRef FeedbackBuffer = GraphBuilder.RegisterExternalBuffer(PooledBuffer, ERDGBufferFlags::SkipTracking);
 	FRDGBufferSRVRef FeedbackBufferSRV = GraphBuilder.CreateSRV(FeedbackBuffer);
@@ -232,11 +239,48 @@ void FVirtualTextureFeedbackBufferResource::End(FRDGBuilder& GraphBuilder)
 	SubmitVirtualTextureFeedbackBuffer(GraphBuilder, CompactedFeedbackBuffer, CompactedDesc);
 }
 
+FRDGBuffer* FVirtualTextureFeedbackBufferResource::ResolveExtendedDebugBuffer(FRDGBuilder& GraphBuilder)
+{
+	if (!bIsInBeginEndScope || ExtendedDebugBufferSize == 0)
+	{
+		return nullptr;
+	}
+
+	// Transition for reading.
+	AddPass(GraphBuilder, RDG_EVENT_NAME("VirtualTextureFeedbackTransitionBeforeExtract"), [UAV = UAV](FRHICommandList& RHICmdList)
+	{
+		RHICmdList.EndUAVOverlap(UAV);
+		RHICmdList.Transition(FRHITransitionInfo(UAV, ERHIAccess::UAVMask, ERHIAccess::SRVCompute));
+	});
+
+	// Copy the extended debug payload.
+	FRDGBufferRef FeedbackBuffer = GraphBuilder.RegisterExternalBuffer(PooledBuffer, ERDGBufferFlags::SkipTracking);
+
+	FRDGBufferDesc DebugBufferCopyDesc = FRDGBufferDesc::CreateStructuredDesc(sizeof(uint32), ExtendedDebugBufferSize);
+	FRDGBufferRef DebugBufferCopy = GraphBuilder.CreateBuffer(DebugBufferCopyDesc, TEXT("VirtualTexture.DebugBufferCopy"));
+
+	AddCopyBufferPass(GraphBuilder, DebugBufferCopy, 0, FeedbackBuffer, FeedbackBufferSize * sizeof(uint32), ExtendedDebugBufferSize * sizeof(uint32));
+
+	// Transition feedback to writing in case we have any subsequent feedback passes (that we won't have captured debug info for).
+	AddPass(GraphBuilder, RDG_EVENT_NAME("VirtualTextureFeedbackTransitionAfterExtract"), [UAV = UAV](FRHICommandList& RHICmdList)
+	{
+		RHICmdList.Transition(FRHITransitionInfo(UAV, ERHIAccess::SRVCompute, ERHIAccess::UAVMask));
+		RHICmdList.BeginUAVOverlap(UAV);
+	});
+
+	return DebugBufferCopy;
+}
+
 uint32 FVirtualTextureFeedbackBufferResource::GetBufferSize() const
 {
 	return bIsInBeginEndScope ? FeedbackBufferSize : 0;
 }
 
+uint32 FVirtualTextureFeedbackBufferResource::GetExtendedDebugBufferSize() const
+{
+	return bIsInBeginEndScope ? ExtendedDebugBufferSize : 0;
+}
+
 FRHIUnorderedAccessView* FVirtualTextureFeedbackBufferResource::GetUAV() const
 {
 	return UAV != nullptr && bIsInBeginEndScope ? UAV : GEmptyStructuredBufferWithUAV->UnorderedAccessViewRHI.GetReference();
@@ -301,6 +345,7 @@ void GetFeedbackShaderParams(int32 InFrameIndex, int32 InVirtualTextureFeedbackT
 {
 	OutParams.BufferUAV = GVirtualTextureFeedbackBufferResource.GetUAV();
 	OutParams.BufferSize = GVirtualTextureFeedbackBufferResource.GetBufferSize();
+	OutParams.ExtendedDebugBufferSize = GVirtualTextureFeedbackBufferResource.GetExtendedDebugBufferSize();
 
 	// Round to nearest power of two to ensure that shader maths is efficient and sampling sequence logic is simple.
 	InVirtualTextureFeedbackTileSize = AlignVirtualTextureFeedbackTileSize(InVirtualTextureFeedbackTileSize);
@@ -327,22 +372,24 @@ void UpdateViewUniformShaderParameters(FFeedbackShaderParams const& InParams, FV
 	ViewUniformShaderParameters.VirtualTextureFeedbackMask = InParams.TileMask;
 	ViewUniformShaderParameters.VirtualTextureFeedbackJitterOffset = InParams.TileJitterOffset;
 	ViewUniformShaderParameters.VirtualTextureFeedbackSampleOffset = InParams.SampleOffset;
+	ViewUniformShaderParameters.VirtualTextureExtendedDebugBufferSize = InParams.ExtendedDebugBufferSize;
 }
 
 void BeginFeedback(FRDGBuilder& GraphBuilder, uint32 InBufferSize, ERHIFeatureLevel::Type InFeatureLevel)
 {
 	const uint32 BufferSize = InBufferSize > 0 ? InBufferSize : GVirtualTextureFeedbackDefaultBufferSize;
 	const ERHIFeatureLevel::Type FeatureLevel = InFeatureLevel < ERHIFeatureLevel::Num ? InFeatureLevel : GMaxRHIFeatureLevel;
-	GVirtualTextureFeedbackBufferResource.Begin(GraphBuilder, BufferSize, FeatureLevel);
+	GVirtualTextureFeedbackBufferResource.Begin(GraphBuilder, BufferSize, 0, FeatureLevel);
 }
 
-void BeginFeedback(FRDGBuilder& GraphBuilder, FIntPoint InViewportSize, uint32 InVirtualTextureFeedbackTileSize, ERHIFeatureLevel::Type InFeatureLevel)
+void BeginFeedback(FRDGBuilder& GraphBuilder, FIntPoint InViewportSize, uint32 InVirtualTextureFeedbackTileSize, bool bInExtendFeedbackForDebug, ERHIFeatureLevel::Type InFeatureLevel)
 {
 	const FIntPoint ViewportSize = FIntPoint(FMath::Max(InViewportSize.X, 1), FMath::Max(InViewportSize.Y, 1));
 	const uint32 TileSize = InVirtualTextureFeedbackTileSize > 0 ? InVirtualTextureFeedbackTileSize : VirtualTextureScalability::GetVirtualTextureFeedbackFactor();
 	const uint32 AlignedTileSize = AlignVirtualTextureFeedbackTileSize(TileSize);
 	const uint32 BufferSize = GetVirtualTextureFeedbackBufferSize(ViewportSize, AlignedTileSize);
-	BeginFeedback(GraphBuilder, BufferSize, InFeatureLevel);
+	const uint32 ExtendedDebugBufferSize = bInExtendFeedbackForDebug ? ViewportSize.X * ViewportSize.Y : 0;
+	GVirtualTextureFeedbackBufferResource.Begin(GraphBuilder, BufferSize, ExtendedDebugBufferSize, InFeatureLevel);
 }
 
 void EndFeedback(FRDGBuilder& GraphBuilder)
@@ -350,4 +397,9 @@ void EndFeedback(FRDGBuilder& GraphBuilder)
 	GVirtualTextureFeedbackBufferResource.End(GraphBuilder);
 }
 
+FRDGBuffer* ResolveExtendedDebugBuffer(FRDGBuilder& GraphBuilder)
+{
+	return GVirtualTextureFeedbackBufferResource.ResolveExtendedDebugBuffer(GraphBuilder);
+}
+
 } // namespace VirtualTexture
