AmbientCubeMap contributes to translucent lighting volume again.

#tests PC editor dx12
#rnx

diff --git a/Engine/Source/Runtime/Renderer/Private/TranslucentLighting.cpp b/Engine/Source/Runtime/Renderer/Private/TranslucentLighting.cpp
--- a/Engine/Source/Runtime/Renderer/Private/TranslucentLighting.cpp
+++ b/Engine/Source/Runtime/Renderer/Private/TranslucentLighting.cpp
@@ -1311,6 +1311,8 @@ void InjectTranslucencyLightingVolumeAmbientCubemap(
 						SetShaderParametersLegacyGS(RHICmdList, GeometryShader, VolumeBounds.MinZ);
 					}
 					SetShaderParameters(RHICmdList, PixelShader, PixelShader.GetPixelShader(), *PassParameters);
+
+					RasterizeToVolumeTexture(RHICmdList, VolumeBounds);
 				});
 			}
 		}
