Path Tracer: Implement cloud acceleration map visualization through the callable shader setup

This removes the need for the cloud case from raygen completely, skipping a few additional branches and removing some unnecessary code paths from the main shader

#rnx

diff --git a/Engine/Shaders/Private/PathTracing/PathTracingVolumetricCloudMaterialShader.usf b/Engine/Shaders/Private/PathTracing/PathTracingVolumetricCloudMaterialShader.usf
--- a/Engine/Shaders/Private/PathTracing/PathTracingVolumetricCloudMaterialShader.usf
+++ b/Engine/Shaders/Private/PathTracing/PathTracingVolumetricCloudMaterialShader.usf
@@ -54,6 +54,21 @@
 
 FVolumeShadedResult CloudsGetDensity(float3 TranslatedWorldPos)
 {
+
+#if CLOUD_VISUALIZATION_SHADER == 1
+	float3 P = mul(GetCloudClipBasis(), TranslatedWorldPos); // WorldToTangent
+	float2 UV = (P.xy * CM_TO_SKY_UNIT + CloudClipDistKm.xx) / (2 * CloudClipDistKm);
+	const int Level = 0;
+	float4 Data = CloudAccelerationMap.SampleLevel(CloudAccelerationMapSampler, UV, Level);
+	
+	float Z = P.z * CM_TO_SKY_UNIT;
+	float Density = (Z > Data.z && Z < Data.w) ? Data.x : 0.0;
+	FVolumeShadedResult Result = (FVolumeShadedResult)0;
+
+	Result.SigmaT = Density;
+	Result.SigmaSDualHG = 0.85 * Density;
+	return Result;
+#else
 	FDFVector3 AbsoluteWorldPosition = DFFastSubtract(TranslatedWorldPos, ResolvedView.PreViewTranslation);
 
 	// Measure HeightKm as the distance to the planet center
@@ -71,6 +86,7 @@ FVolumeShadedResult CloudsGetDensity(float3 TranslatedWorldPos)
 	Result.DualPhaseData = float3(CloudResult.PhaseG1, CloudResult.PhaseG2, CloudResult.PhaseBlend);
 #endif
 	return Result;
+#endif
 }
 
 RAY_TRACING_ENTRY_CALLABLE(PathTracingVolumetricCloudMaterialShader, FPackedPathTracingPayload, PackedPayload)
diff --git a/Engine/Shaders/Private/PathTracing/Volume/PathTracingAtmosphere.ush b/Engine/Shaders/Private/PathTracing/Volume/PathTracingAtmosphere.ush
--- a/Engine/Shaders/Private/PathTracing/Volume/PathTracingAtmosphere.ush
+++ b/Engine/Shaders/Private/PathTracing/Volume/PathTracingAtmosphere.ush
@@ -113,7 +113,7 @@ float3 RayPlanetIntersectPrecise(float3 Ro, float3 Rd)
 
 #endif
 
-#ifndef PlanetCenterTranslatedWorldHi // TODO: find a better way to recognize we are coming from the cloud shader
+#if CLOUD_LAYER_PIXEL_SHADER == 0
 
 FVolumeIntersection AtmosphereIntersect(float3 Origin, float3 Direction, float TMin, float TMax)
 {
@@ -266,7 +266,7 @@ FVolumeShadedResult AtmosphereGetDensity(float3 TranslatedWorldPos)
 	return Result;
 }
 
-#ifndef PlanetCenterTranslatedWorldHi // TODO: find a better way to recognize we are coming from the cloud shader
+#if CLOUD_LAYER_PIXEL_SHADER == 0
 
 float3 AtmosphereGetOpticalDensity(FSphericalPlanetParameters Params)
 {
diff --git a/Engine/Shaders/Private/PathTracing/Volume/PathTracingClouds.ush b/Engine/Shaders/Private/PathTracing/Volume/PathTracingClouds.ush
--- a/Engine/Shaders/Private/PathTracing/Volume/PathTracingClouds.ush
+++ b/Engine/Shaders/Private/PathTracing/Volume/PathTracingClouds.ush
@@ -163,22 +163,3 @@ FVolumeIntersection CloudsIntersect(float3 Origin, float3 Direction, float TMin,
 	}
 	return CreateEmptyVolumeIntersection();
 }
-
-#if PATH_TRACER_USE_CLOUD_SHADER == 0
-FVolumeShadedResult CloudsGetDensity(float3 TranslatedWorldPos)
-{
-	// debug preview path to visualize the cloud accel map
-	float3 P = mul(GetCloudClipBasis(), TranslatedWorldPos); // WorldToTangent
-	float2 UV = (P.xy * CM_TO_SKY_UNIT + CloudClipDistKm.xx) / (2 * CloudClipDistKm);
-	const int Level = 0;
-	float4 Data = CloudAccelerationMap.SampleLevel(CloudAccelerationMapSampler, UV, Level);
-	
-	float Z = P.z * CM_TO_SKY_UNIT;
-	float Density = (Z > Data.z && Z < Data.w) ? Data.x : 0.0;
-	FVolumeShadedResult Result = (FVolumeShadedResult)0;
-
-	Result.SigmaT = Density;
-	Result.SigmaSHG = 0.85 * Density;
-	return Result;
-}
-#endif
\ No newline at end of file
diff --git a/Engine/Shaders/Private/PathTracing/Volume/PathTracingCloudsMaterialCommon.ush b/Engine/Shaders/Private/PathTracing/Volume/PathTracingCloudsMaterialCommon.ush
--- a/Engine/Shaders/Private/PathTracing/Volume/PathTracingCloudsMaterialCommon.ush
+++ b/Engine/Shaders/Private/PathTracing/Volume/PathTracingCloudsMaterialCommon.ush
@@ -6,7 +6,9 @@
 // This code is for sharing between the cloud acceleration map building and the callable shader version of the cloud material
 
 #include "../../Common.ush"
+#if CLOUD_VISUALIZATION_SHADER == 0
 #include "/Engine/Generated/Material.ush"
+#endif
 #include "PathTracingVolumeCommon.ush"
 #include "PathTracingCloudsCommon.ush"
 #include "../../DoubleFloat.ush"
@@ -73,7 +75,7 @@ FCloudMaterialMSParams CloudMaterialGetMSParams()
 }
 #endif
 
-
+#if CLOUD_VISUALIZATION_SHADER == 0
 // For path tracing, we need to be able to evaluate the material starting from a single point in space
 // In all contexts, we have a slightly more accurate Height already available so pass that in instead of measuring it from the position
 FSampleCloudMaterialResult SampleCloudMaterial(FDFVector3 AbsoluteWorldPosition, float HeightKm, float CloudBotKm, float CloudTopKm)
@@ -135,3 +137,4 @@ FSampleCloudMaterialResult SampleCloudMaterial(FDFVector3 AbsoluteWorldPosition,
 
 	return Result;
 }
+#endif
diff --git a/Engine/Shaders/Private/PathTracing/Volume/PathTracingVolume.ush b/Engine/Shaders/Private/PathTracing/Volume/PathTracingVolume.ush
--- a/Engine/Shaders/Private/PathTracing/Volume/PathTracingVolume.ush
+++ b/Engine/Shaders/Private/PathTracing/Volume/PathTracingVolume.ush
@@ -75,13 +75,6 @@ FVolumeShadedResult VolumeGetDensity(float3 TranslatedWorldPos, uint VolumeMask,
 	{
 		MergeVolumeShadedResult(Result, AtmosphereGetDensity(TranslatedWorldPos), bIsCameraRay && (VolumeFlags & PATH_TRACER_VOLUME_HOLDOUT_ATMOSPHERE) != 0);
 	}
-#if PATH_TRACER_USE_CLOUD_SHADER == 0
-	// If we are using the CallableShader Path, this function is not called with clouds in VolumeMask
-	if (VolumeMask & PATH_TRACER_VOLUME_ENABLE_CLOUDS)
-	{
-		MergeVolumeShadedResult(Result, CloudsGetDensity(TranslatedWorldPos), bIsCameraRay && (VolumeFlags & PATH_TRACER_VOLUME_HOLDOUT_CLOUDS) != 0);
-	}
-#endif
 	if (VolumeMask & PATH_TRACER_VOLUME_ENABLE_FOG)
 	{
 		MergeVolumeShadedResult(Result, FogGetDensity(TranslatedWorldPos), bIsCameraRay && (VolumeFlags & PATH_TRACER_VOLUME_HOLDOUT_FOG) != 0);
@@ -169,22 +162,20 @@ float3 VolumeGetTransmittance(float3 StartThroughput, float3 Origin, float3 Dire
 			PATH_TRACER_REFERENCE_TRANSMITTANCE_LOOP(AtmosphereGetDensityBounds, AtmosphereGetDensity);
 		}
 	}
+#if PATH_TRACER_USE_CLOUD_SHADER
 	if (Interval.VolumeMask & PATH_TRACER_VOLUME_ENABLE_CLOUDS)
 	{
 		float TMin = Interval.VolumeTMin;
 		float TMax = Interval.VolumeTMax;
 		float SigmaBar = Interval.CloudDensity;
-#if PATH_TRACER_USE_CLOUD_SHADER
 		FPackedPathTracingPayload CloudPayload = (FPackedPathTracingPayload) 0;
 		CloudPayload.SetVolumetricCallableShaderInput(Origin, Direction, TMin, TMax, RandSequence, PATH_TRACER_VOLUME_CALLABLE_FLAGS_TRANSMITTANCE, Throughput, SigmaBar);
 		CloudPayload.SetVolumetricCallableShaderInputCloudFactor(1.0);
 		CallShader(CloudCallableShaderId, CloudPayload);
 		RandSequence = CloudPayload.GetVolumetricCallableShaderOutputRandSeq();
 		Throughput = CloudPayload.GetVolumetricCallableShaderOutputThroughput();
-#else
-		PATH_TRACER_REFERENCE_TRANSMITTANCE_LOOP_CONSTANT_BOUNDS(CloudsGetDensity);
-#endif
 	}
+#endif
 	if (Interval.VolumeMask & PATH_TRACER_VOLUME_ENABLE_FOG)
 	{
 		float TMin = Interval.VolumeTMin;
@@ -225,22 +216,20 @@ float3 VolumeGetTransmittance(float3 StartThroughput, float3 Origin, float3 Dire
 			PATH_TRACER_REFERENCE_TRANSMITTANCE_LOOP(AtmosphereGetDensityBounds, AtmosphereGetDensity);
 		}
 	}
+#if PATH_TRACER_USE_CLOUD_SHADER
 	if (IntersectionList.VolumeMask & PATH_TRACER_VOLUME_ENABLE_CLOUDS)
 	{
 		float TMin = IntersectionList.VolumeTMin[VOLUMEID_CLOUDS];
 		float TMax = IntersectionList.VolumeTMax[VOLUMEID_CLOUDS];
 		float SigmaBar = IntersectionList.CloudDensity;
-#if PATH_TRACER_USE_CLOUD_SHADER
 		FPackedPathTracingPayload CloudPayload = (FPackedPathTracingPayload) 0;
 		CloudPayload.SetVolumetricCallableShaderInput(Origin, Direction, TMin, TMax, RandSequence, PATH_TRACER_VOLUME_CALLABLE_FLAGS_TRANSMITTANCE, Throughput, SigmaBar);
 		CloudPayload.SetVolumetricCallableShaderInputCloudFactor(CloudFactor);
 		CallShader(CloudCallableShaderId, CloudPayload);
 		RandSequence = CloudPayload.GetVolumetricCallableShaderOutputRandSeq();
 		Throughput = CloudPayload.GetVolumetricCallableShaderOutputThroughput();
-#else
-		PATH_TRACER_REFERENCE_TRANSMITTANCE_LOOP_CONSTANT_BOUNDS(CloudsGetDensity);
-#endif
 	}
+#endif
 	if (IntersectionList.VolumeMask & PATH_TRACER_VOLUME_ENABLE_FOG)
 	{
 		float TMin = IntersectionList.VolumeTMin[VOLUMEID_FOG];
diff --git a/Engine/Source/Runtime/Renderer/Private/PathTracing.cpp b/Engine/Source/Runtime/Renderer/Private/PathTracing.cpp
--- a/Engine/Source/Runtime/Renderer/Private/PathTracing.cpp
+++ b/Engine/Source/Runtime/Renderer/Private/PathTracing.cpp
@@ -1693,6 +1693,80 @@ private:
 };
 IMPLEMENT_SHADER_TYPE(, FPathTracingVolumetricCloudMaterial, TEXT("/Engine/Private/PathTracing/PathTracingVolumetricCloudMaterialShader.usf"), TEXT("PathTracingVolumetricCloudMaterialShader"), SF_RayCallable);
 
+
+class FPathTracingVolumetricCloudMaterialVisualize : public FGlobalShader
+{
+	DECLARE_GLOBAL_SHADER(FPathTracingVolumetricCloudMaterialVisualize)
+public:
+
+	FPathTracingVolumetricCloudMaterialVisualize() = default;
+	FPathTracingVolumetricCloudMaterialVisualize(const ShaderMetaType::CompiledShaderInitializerType& Initializer) : FGlobalShader(Initializer)
+	{
+		CloudParameter.Bind(Initializer.ParameterMap, TEXT("PathTracingCloudParameters"));
+	}
+
+	static bool ShouldCompilePermutation(const FGlobalShaderPermutationParameters& Parameters)
+	{
+		return ShouldCompilePathTracingShadersForProject(Parameters.Platform) &&
+			   ShouldCompileRayTracingCallableShadersForProject(Parameters.Platform);
+	}
+
+	static void ModifyCompilationEnvironment(const FGlobalShaderPermutationParameters& Parameters, FShaderCompilerEnvironment& OutEnvironment)
+	{
+		FGlobalShader::ModifyCompilationEnvironment(Parameters, OutEnvironment);
+		OutEnvironment.SetDefine(TEXT("CLOUD_LAYER_PIXEL_SHADER"), 1);
+		OutEnvironment.SetDefine(TEXT("CLOUD_VISUALIZATION_SHADER"), 1);
+	}
+
+	static bool ValidateCompiledResult(EShaderPlatform Platform, const FShaderParameterMap& ParameterMap, TArray<FString>& OutError)
+	{
+		if (ParameterMap.ContainsParameterAllocation(FSceneTextureUniformParameters::FTypeInfo::GetStructMetadata()->GetShaderVariableName()))
+		{
+			OutError.Add(TEXT("Ray tracing callable shaders cannot read from the SceneTexturesStruct."));
+			return false;
+		}
+
+		for (const auto& It : ParameterMap.GetParameterMap())
+		{
+			const FParameterAllocation& ParamAllocation = It.Value;
+			if (ParamAllocation.Type != EShaderParameterType::UniformBuffer
+				&& ParamAllocation.Type != EShaderParameterType::LooseData)
+			{
+				OutError.Add(FString::Printf(TEXT("Invalid ray tracing shader parameter '%s'. Only uniform buffers and loose data parameters are supported."), *(It.Key)));
+				return false;
+			}
+		}
+
+		return true;
+	}
+
+	static ERayTracingPayloadType GetRayTracingPayloadType(const int32 PermutationId)
+	{
+		// TODO: This isn't the payload we use, but the logic in RayTracingMaterialHitShaders.cpp needs to assume a consistent payload ID for all callable shaders
+		// The simplest solution is probably to remove FDecalShaderPayload and use the PT payload everywhere
+		return ERayTracingPayloadType::Decals;
+	}
+			
+	static const FShaderBindingLayout* GetShaderBindingLayout(const FShaderPermutationParameters& Parameters)
+	{
+		return RayTracing::GetShaderBindingLayout(Parameters.Platform);
+	}
+
+	void GetShaderBindings(
+		const FViewInfo& View,
+		const TUniformBufferRef<FPathTracingCloudParameterGlobals>& CloudParameters,
+		FMeshDrawSingleShaderBindings& ShaderBindings) const
+	{
+		ShaderBindings.Add(GetUniformBufferParameter<FViewUniformShaderParameters>(), View.ViewUniformBuffer);
+		ShaderBindings.Add(CloudParameter, CloudParameters);
+	}
+
+private:
+	LAYOUT_FIELD(FShaderUniformBufferParameter, CloudParameter);
+};
+IMPLEMENT_SHADER_TYPE(, FPathTracingVolumetricCloudMaterialVisualize, TEXT("/Engine/Private/PathTracing/PathTracingVolumetricCloudMaterialShader.usf"), TEXT("PathTracingVolumetricCloudMaterialShader"), SF_RayCallable);
+
+
 void PreparePathTracingCloudMaterial(FRDGBuilder& GraphBuilder, FScene* Scene, TArrayView<FViewInfo> Views)
 {
 	// make sure all views have an invalid callable shader index (unless proven otherwise below)
@@ -1742,6 +1816,19 @@ void PreparePathTracingCloudMaterial(FRDGBuilder& GraphBuilder, FScene* Scene, T
 
 	const FVolumetricCloudRenderSceneInfo* CloudInfo = Scene->GetVolumetricCloudSceneInfo();
 
+	const bool bCloudAccelerationMapVisualize = CVarPathTracingCloudAccelerationMapVisualize.GetValueOnRenderThread();
+
+	TShaderRef<FPathTracingVolumetricCloudMaterialVisualize> CallableShaderVisualize;
+	if (bCloudAccelerationMapVisualize)
+	{
+		CallableShaderVisualize = GetGlobalShaderMap(Scene->GetFeatureLevel())->GetShader<FPathTracingVolumetricCloudMaterialVisualize>();
+		if (!CallableShaderVisualize.IsValid())
+		{
+			// asking for visualization, but visualization shader is not ready
+			return;
+		}
+	}
+
 	for (FViewInfo& View : Views)
 	{
 		if (!PathTracing::UsesReferenceAtmosphere(View) || !ShouldRenderVolumetricCloud(Scene, View.Family->EngineShowFlags))
@@ -1753,7 +1840,14 @@ void PreparePathTracingCloudMaterial(FRDGBuilder& GraphBuilder, FScene* Scene, T
 		uint32 BaseCallableSlotIndex = Scene->RayTracingSBT.NumCallableShaderSlots;
 		FRayTracingShaderCommand& Command = Scene->RayTracingSBT.CallableCommands.AddDefaulted_GetRef();
 
-		Command.SetShader(CallableShader);
+		if (bCloudAccelerationMapVisualize)
+		{
+			Command.SetShader(CallableShaderVisualize);
+		}
+		else
+		{
+			Command.SetShader(CallableShader);
+		}
 		Command.SlotInScene = BaseCallableSlotIndex;
 
 		View.PathTracingVolumetricCloudCallableShaderIndex = BaseCallableSlotIndex;
@@ -1800,8 +1894,16 @@ void PreparePathTracingCloudMaterial(FRDGBuilder& GraphBuilder, FScene* Scene, T
 		Scene->RayTracingSBT.TransientUniformBuffers.Add(CloudParametersUB); // Hold uniform buffer ref in RayTracingSBT since FMeshDrawSingleShaderBindings doesn't
 
 		FMeshDrawSingleShaderBindings SingleShaderBindings = Command.ShaderBindings.GetSingleShaderBindings(SF_RayCallable);
-		CallableShader->GetShaderBindings(Scene, Scene->GetFeatureLevel(), *CloudVolumeMaterialProxy, *MaterialResource, View, CloudParametersUB, SingleShaderBindings);
-
+		
+		if (bCloudAccelerationMapVisualize)
+		{
+			CallableShaderVisualize->GetShaderBindings(View, CloudParametersUB, SingleShaderBindings);
+		}
+		else
+		{
+			CallableShader->GetShaderBindings(Scene, Scene->GetFeatureLevel(), *CloudVolumeMaterialProxy, *MaterialResource, View, CloudParametersUB, SingleShaderBindings);
+		}
+		
 		Scene->RayTracingSBT.NumCallableShaderSlots++;
 	}
 }
@@ -3049,8 +3151,7 @@ static FPathTracingRG::FPermutationDomain GetPathTracingRGPermutation(const FVie
 	const bool bUseExperimental = CVarPathTracingExperimental.GetValueOnRenderThread();
 	const bool bUseAdaptiveSampling = bUseExperimental && CVarPathTracingAdaptiveSampling.GetValueOnRenderThread() != 0;
 	// NOTE: the decision about when to enable clouds involves checking lots of things, so rely on the presence of a valid index to signify that clouds are ready
-	const bool bUseCloudShader = uint32(View.PathTracingVolumetricCloudCallableShaderIndex) < Scene.RayTracingSBT.NumCallableShaderSlots
-		&& !CVarPathTracingCloudAccelerationMapVisualize.GetValueOnRenderThread();
+	const bool bUseCloudShader = uint32(View.PathTracingVolumetricCloudCallableShaderIndex) < Scene.RayTracingSBT.NumCallableShaderSlots;
 	const bool bHasComplexSpecialRenderPath = Substrate::IsSubstrateEnabled() && Scene.SubstrateSceneData.bUsesComplexSpecialRenderPath;
 	const bool bUseSER = GRHIGlobals.SupportsShaderExecutionReordering && CVarPathTracingShaderExecutionReordering.GetValueOnRenderThread();
 	const bool bUseTraceOpaqueFirst = CVarPathTracingTraceOpaqueFirst.GetValueOnRenderThread();
