Improvements to Translucency compositing with Heterogeneous Volumes.

1. Added new option to "r.Translucency.HeterogeneousVolumes" to support Beer Shadow Maps:
r.Translucency.HeterogeneousVolumes ?

Enables composting with heterogeneous volumes when rendering translucency (Default = 0)
0: Off
1: Camera-based Adaptive Volumetric Shadow Map
2: Camera-based Beer-Law Shadow Map

2. Unify HV compositing with Fogging.

diff --git a/Engine/Shaders/Private/BasePassPixelShader.usf b/Engine/Shaders/Private/BasePassPixelShader.usf
--- a/Engine/Shaders/Private/BasePassPixelShader.usf
+++ b/Engine/Shaders/Private/BasePassPixelShader.usf
@@ -1415,10 +1415,12 @@ void FPixelShaderInOut_MainPS(
 #endif
 
 #if ADAPTIVE_VOLUMETRIC_SHADOW_MAP
+	float2 NDC = MaterialParameters.ScreenPosition.xy / MaterialParameters.ScreenPosition.w;
+	float2 ScreenUV = NDC * ResolvedView.ScreenPositionScaleBias.xy + ResolvedView.ScreenPositionScaleBias.wz;
 	float3 TranslatedWorldPosition = SvPositionToResolvedTranslatedWorld(In.SvPosition);
-	float HeterogeneousVolumeTransmittance = saturate(AVSM_SampleTransmittance(TranslatedWorldPosition, ResolvedView.TranslatedWorldCameraOrigin));
-	Opacity *= HeterogeneousVolumeTransmittance;
-
+	float4 HeterogeneousVolumeResult = saturate(AVSM_SampleCameraRadianceAndTransmittance4(ScreenUV, TranslatedWorldPosition, ResolvedView.TranslatedWorldCameraOrigin));
+	Fogging.rgb = HeterogeneousVolumeResult.rgb + Fogging.rgb * HeterogeneousVolumeResult.a;
+	Fogging.a *= HeterogeneousVolumeResult.a;
 #endif // ADAPTIVE_VOLUMETRIC_SHADOW_MAP
 
 #if NEEDS_BASEPASS_PIXEL_FOGGING
@@ -1599,9 +1601,6 @@ void FPixelShaderInOut_MainPS(
 
 #if MATERIAL_SHADINGMODEL_THIN_TRANSLUCENT && !SUBSTRATE_ENABLED
 	{
-#if ADAPTIVE_VOLUMETRIC_SHADOW_MAP
-		ColorSeparateSpecular *= HeterogeneousVolumeTransmittance;
-#endif
 		AccumulateThinTranslucentModel(
 								DualBlendSurfaceLuminancePostCoverage,
 								DualBlendSurfaceTransmittancePreCoverage,
@@ -1612,10 +1611,6 @@ void FPixelShaderInOut_MainPS(
 								ColorSeparateSpecular,
 								Emissive,
 								Opacity);
-#if ADAPTIVE_VOLUMETRIC_SHADOW_MAP
-		DualBlendSurfaceCoverage *= HeterogeneousVolumeTransmittance;
-		DualBlendSurfaceLuminancePostCoverage *= HeterogeneousVolumeTransmittance;
-#endif
 		Color = 0;
 		Opacity = 1.0f;
 	}
diff --git a/Engine/Shaders/Private/HeterogeneousVolumes/HeterogeneousVolumesAdaptiveVolumetricShadowMapSampling.ush b/Engine/Shaders/Private/HeterogeneousVolumes/HeterogeneousVolumesAdaptiveVolumetricShadowMapSampling.ush
--- a/Engine/Shaders/Private/HeterogeneousVolumes/HeterogeneousVolumesAdaptiveVolumetricShadowMapSampling.ush
+++ b/Engine/Shaders/Private/HeterogeneousVolumes/HeterogeneousVolumesAdaptiveVolumetricShadowMapSampling.ush
@@ -19,6 +19,7 @@ struct FAdaptiveVolumetricShadowMap
 {
 	float4x4 TranslatedWorldToShadow[6];
 	float3 TranslatedWorldOrigin;
+	float DownsampleFactor;
 	float4 TranslatedWorldPlane;
 
 	int2 Resolution;
@@ -29,6 +30,8 @@ struct FAdaptiveVolumetricShadowMap
 
 	StructuredBuffer<uint4> IndirectionBuffer;
 	StructuredBuffer<uint> SampleBuffer;
+	Texture2D<float4> RadianceTexture;
+	SamplerState TextureSampler;
 };
 
 bool AVSM_PixelAndDepth(
@@ -74,6 +77,56 @@ float SafeRcp(float Value)
 	return Value > 0 ? rcp(Value) : 0.0;
 }
 
+float SafeLog(float Value)
+{
+	return Value > 0 ? log(Value) : 0.0;
+}
+
+#define INTERPOLANT_TYPE_LINEAR 0
+#define INTERPOLANT_TYPE_EXPONENTIAL 1
+#define INTERPOLANT_TYPE_EXPONENTIAL_FIT 2
+#define INTERPOLANT_TYPE INTERPOLANT_TYPE_EXPONENTIAL_FIT
+
+float Interpolate(
+	FAVSMSampleData ShadowData[2],
+	float WorldHitT
+)
+{
+	float Transmittance = ShadowData[1].Tau;
+
+	if (INTERPOLANT_TYPE == INTERPOLANT_TYPE_LINEAR)
+	{
+		// Linear interpolant
+		float DeltaX = max(ShadowData[1].X - ShadowData[0].X, 0.0);
+		float X = clamp(WorldHitT - ShadowData[0].X, 0.0, DeltaX);
+		float BlendWeight = saturate(X * SafeRcp(DeltaX));
+		Transmittance = lerp(ShadowData[0].Tau, ShadowData[1].Tau, BlendWeight);
+	}
+	else if (INTERPOLANT_TYPE == INTERPOLANT_TYPE_EXPONENTIAL)
+	{
+		// Exponential interpolant
+		float DeltaX = max(ShadowData[1].X - ShadowData[0].X, 0.0);
+		float X = clamp(WorldHitT - ShadowData[0].X, 0.0, DeltaX);
+		float BlendWeight = saturate(X * SafeRcp(DeltaX));
+		//float SigmaT = 10.0;
+		float SigmaT = 1.0;
+		Transmittance = lerp(ShadowData[1].Tau, ShadowData[0].Tau, exp(-BlendWeight * SigmaT));
+	}
+	else // if (INTERPOLANT_TYPE == INTERPOLANT_TYPE_EXPONENTIAL_FIT)
+	{
+		// Exponential fit
+		float DeltaX = max(ShadowData[1].X - ShadowData[0].X, 0.0);
+		float SigmaT = -SafeLog(ShadowData[1].Tau * SafeRcp(ShadowData[0].Tau)) * SafeRcp(DeltaX);
+		if (SigmaT > 0.0)
+		{
+			float X = clamp(WorldHitT - ShadowData[0].X, 0.0, DeltaX);
+			Transmittance = saturate(ShadowData[0].Tau * exp(-X * SigmaT));
+		}
+	}
+
+	return Transmittance;
+}
+
 float AVSM_Sample_Reference(
 	inout FAdaptiveVolumetricShadowMap ShadowMap,
 	int Face,
@@ -119,8 +172,7 @@ float AVSM_Sample_Reference(
 		}
 	}
 
-	float BlendWeight = saturate((WorldHitT - ShadowData[0].X) * SafeRcp(ShadowData[1].X - ShadowData[0].X));
-	float Transmittance = lerp(ShadowData[0].Tau, ShadowData[1].Tau, BlendWeight);
+	float Transmittance = Interpolate(ShadowData, WorldHitT);
 	return Transmittance;
 }
 
@@ -177,8 +229,7 @@ float AVSM_Sample_BinarySearch(
 		}
 	}
 
-	float BlendWeight = saturate((WorldHitT - ShadowData[0].X) * SafeRcp(ShadowData[1].X - ShadowData[0].X));
-	float Transmittance = lerp(ShadowData[0].Tau, ShadowData[1].Tau, BlendWeight);
+	float Transmittance = Interpolate(ShadowData, WorldHitT);
 	return Transmittance;
 }
 #if 0
@@ -347,9 +398,7 @@ float AVSM_Sample_TwoLevel_Vectorized(
 		ShadowData[1] = AVSM_UnpackSampleData(PackedShadowData4[min(Index, 3)], IndirectionData);
 	}
 
-	// Lerp the final value
-	float BlendWeight = saturate((WorldHitT - ShadowData[0].X) * SafeRcp(ShadowData[1].X - ShadowData[0].X));
-	float Transmittance = lerp(ShadowData[0].Tau, ShadowData[1].Tau, BlendWeight);
+	float Transmittance = Interpolate(ShadowData, WorldHitT);
 	return Transmittance;
 }
 
@@ -460,6 +509,7 @@ FAdaptiveVolumetricShadowMap GetAdaptiveVolumetricShadowMap()
 	ShadowMap.TranslatedWorldPlane = AVSM.TranslatedWorldPlane;
 
 	ShadowMap.Resolution = AVSM.Resolution;
+	ShadowMap.DownsampleFactor = AVSM.DownsampleFactor;
 	ShadowMap.NumShadowMatrices = AVSM.NumShadowMatrices;
 	ShadowMap.MaxSampleCount = AVSM.MaxSampleCount;
 	ShadowMap.bIsEmpty = AVSM.bIsEmpty;
@@ -467,6 +517,9 @@ FAdaptiveVolumetricShadowMap GetAdaptiveVolumetricShadowMap()
 
 	ShadowMap.IndirectionBuffer = AVSM.IndirectionBuffer;
 	ShadowMap.SampleBuffer = AVSM.SampleBuffer;
+
+	ShadowMap.RadianceTexture = AVSM.RadianceTexture;
+	ShadowMap.TextureSampler = AVSM.TextureSampler;
 	return ShadowMap;
 }
 
@@ -482,6 +535,7 @@ FAdaptiveVolumetricShadowMap GetAdaptiveVolumetricShadowMapForCamera()
 	ShadowMap.TranslatedWorldPlane = CameraAVSM.TranslatedWorldPlane;
 
 	ShadowMap.Resolution = CameraAVSM.Resolution;
+	ShadowMap.DownsampleFactor = CameraAVSM.DownsampleFactor;
 	ShadowMap.NumShadowMatrices = CameraAVSM.NumShadowMatrices;
 	ShadowMap.MaxSampleCount = CameraAVSM.MaxSampleCount;
 	ShadowMap.bIsEmpty = CameraAVSM.bIsEmpty;
@@ -489,6 +543,9 @@ FAdaptiveVolumetricShadowMap GetAdaptiveVolumetricShadowMapForCamera()
 
 	ShadowMap.IndirectionBuffer = CameraAVSM.IndirectionBuffer;
 	ShadowMap.SampleBuffer = CameraAVSM.SampleBuffer;
+
+	ShadowMap.RadianceTexture = CameraAVSM.RadianceTexture;
+	ShadowMap.TextureSampler = CameraAVSM.TextureSampler;
 #endif // USE_CAMERA_AVSM
 	return ShadowMap;
 }
@@ -501,3 +558,20 @@ float AVSM_SampleTransmittance(
 	FAdaptiveVolumetricShadowMap AdaptiveVolumetricShadowMap = GetAdaptiveVolumetricShadowMap();
 	return AVSM_SampleTransmittance(AdaptiveVolumetricShadowMap, TranslatedWorldPosition, LightTranslatedWorldPosition);
 }
+
+float4 AVSM_SampleCameraRadianceAndTransmittance4(
+	float2 ScreenUV,
+	float3 TranslatedWorldPosition,
+	float3 LightTranslatedWorldPosition
+)
+{
+	FAdaptiveVolumetricShadowMap AdaptiveVolumetricShadowMap = GetAdaptiveVolumetricShadowMap();
+	float Transmittance = AVSM_SampleTransmittance(AdaptiveVolumetricShadowMap, TranslatedWorldPosition, LightTranslatedWorldPosition);
+
+	// Extract radiance from precomposited render target and weight based on relative opacity
+	float4 Result = Texture2DSample(AdaptiveVolumetricShadowMap.RadianceTexture, AdaptiveVolumetricShadowMap.TextureSampler, ScreenUV);
+	float3 Radiance = Result.rgb * View.OneOverPreExposure;
+	float Weight = saturate((1.0 - Transmittance) * SafeRcp(1.0 - Result.a));
+
+	return float4(Radiance * Weight, Transmittance);
+}
\ No newline at end of file
diff --git a/Engine/Shaders/Private/HeterogeneousVolumes/HeterogeneousVolumesComposite.usf b/Engine/Shaders/Private/HeterogeneousVolumes/HeterogeneousVolumesComposite.usf
--- a/Engine/Shaders/Private/HeterogeneousVolumes/HeterogeneousVolumesComposite.usf
+++ b/Engine/Shaders/Private/HeterogeneousVolumes/HeterogeneousVolumesComposite.usf
@@ -198,6 +198,7 @@ void HeterogeneousVolumesCompositeCS(
 	// Cache Sample values
 	float4 Radiance = HeterogeneousVolumeRadiance[PixelCoord];
 	float HeterogeneousVolumeTransmittance = bUseAVSM ? SampleTransmittance(ViewPixelCoord) : Radiance.a;
+	//float HeterogeneousVolumeTransmittance = Radiance.a;
 
 	float4 Result;
 	Result.rgb = Radiance.rgb + RWColorTexture[ViewPixelCoord].rgb * HeterogeneousVolumeTransmittance;
diff --git a/Engine/Shaders/Private/HeterogeneousVolumes/HeterogeneousVolumesLiveShadingPipeline.usf b/Engine/Shaders/Private/HeterogeneousVolumes/HeterogeneousVolumesLiveShadingPipeline.usf
--- a/Engine/Shaders/Private/HeterogeneousVolumes/HeterogeneousVolumesLiveShadingPipeline.usf
+++ b/Engine/Shaders/Private/HeterogeneousVolumes/HeterogeneousVolumesLiveShadingPipeline.usf
@@ -89,6 +89,7 @@ float ShadowStepFactor;
 // Atmosphere
 int bApplyHeightFog;
 int bApplyVolumetricFog;
+int bCreateBeerShadowMap;
 
 // Indirect
 float IndirectInscatteringFactor;
@@ -110,6 +111,7 @@ int DownsampleFactor;
 RWTexture2D<float4> RWLightingTexture;
 RWTexture2D<float4> RWVelocityTexture;
 RWTexture2D<float> RWHoldoutTexture;
+RWTexture2D<float4> RWBeerShadowMapTexture;
 
 struct FDebugOutput
 {
@@ -426,9 +428,14 @@ void RenderLightingCacheWithLiveShadingCS(
 #endif // DIM_LIGHTING_CACHE_MODE
 }
 
-int2 GetDispatchSize()
+int2 GetDownsampledResolution(int2 Resolution, int Factor)
 {
-	return ceil(View.ViewSizeAndInvSize.xy / (float)DownsampleFactor);
+	return (Resolution + Factor - 1) / Factor;
+}
+
+int2 GetScaledViewRect()
+{
+	return GetDownsampledResolution(View.ViewSizeAndInvSize.xy, DownsampleFactor);
 }
 
 #if USE_CAMERA_AVSM
@@ -529,6 +536,94 @@ void RayMarchSingleScatteringAdaptive(
 
 #endif // USE_CAMERA_AVSM
 
+struct FBeerShadowMapData
+{
+	float T[2];
+	float Tau[2];
+};
+
+FBeerShadowMapData InitializeBeerShadowMapData(float Transmittance)
+{
+	FBeerShadowMapData Data;
+	Data.T[0] = GetMaxTraceDistance();
+	Data.T[1] = 0.0;
+	Data.Tau[0] = Transmittance;
+	Data.Tau[1] = Transmittance;
+	return Data;
+}
+
+FBeerShadowMapData UnpackBeerShadowMapData(float4 PackedData)
+{
+	FBeerShadowMapData Data;
+	Data.T[0] = PackedData.r;
+	Data.T[1] = PackedData.g;
+	Data.Tau[0] = PackedData.b;
+	Data.Tau[1] = PackedData.a;
+	return Data;
+}
+
+float4 PackBeerShadowMapData(FBeerShadowMapData Data)
+{
+	float4 PackedData;
+	PackedData.r = Data.T[0];
+	PackedData.g = Data.T[1];
+	PackedData.b = Data.Tau[0];
+	PackedData.a = Data.Tau[1];
+	return PackedData;
+}
+
+#define BEER_SHADOW_MAP_COMBINE_METHOD_FIRST 0
+#define BEER_SHADOW_MAP_COMBINE_METHOD_EXPAND 1
+
+FBeerShadowMapData ExpandBeerShadowMapData(FBeerShadowMapData Data0, FBeerShadowMapData Data1)
+{
+	FBeerShadowMapData Result;
+	Result.T[0] = min(Data0.T[0], Data1.T[0]);
+	Result.T[1] = max(Data0.T[1], Data1.T[1]);
+	Result.Tau[0] = Data0.Tau[0];
+	Result.Tau[1] = Data1.Tau[1];
+	return Result;
+}
+
+FBeerShadowMapData FirstBeerShadowMapData(FBeerShadowMapData Data0, FBeerShadowMapData Data1)
+{
+	bool bIsValid0 = Data0.T[1] > Data0.T[0];
+	bool bIsValid1 = Data1.T[1] > Data1.T[0];
+
+	if (bIsValid0 && bIsValid1)
+	{
+		if (Data0.T[0] < Data1.T[0])
+		{
+			return Data0;
+		}
+		else
+		{
+			return Data1;
+		}
+	}
+	else if (bIsValid0)
+	{
+		return Data0;
+	}
+	// else
+	return Data1;
+}
+
+FBeerShadowMapData CombineBeerShadowMapData(FBeerShadowMapData Data0, FBeerShadowMapData Data1, int CombineMethod)
+{
+	FBeerShadowMapData Result;
+
+	if (CombineMethod == BEER_SHADOW_MAP_COMBINE_METHOD_FIRST)
+	{
+		Result = FirstBeerShadowMapData(Data0, Data1);
+	}
+	else // if (CombineMethod == BEER_SHADOW_MAP_COMBINE_METHOD_EXPAND)
+	{
+		Result = ExpandBeerShadowMapData(Data0, Data1);
+	}
+	return Result;
+}
+
 void RenderSingleScatteringWithLiveShading(
 	uint2 DispatchThreadId
 )
@@ -536,19 +631,27 @@ void RenderSingleScatteringWithLiveShading(
 	FDebugOutput DebugOutput = (FDebugOutput)0;
 
 	// Create screen ray
-	if (any(DispatchThreadId.xy >= GetDispatchSize()))
+	if (any(DispatchThreadId.xy >= GetScaledViewRect()))
 	{
 		return;
 	}
+
 	uint LinearIndex = DispatchThreadId.y * GroupCount.x * THREADGROUP_SIZE_2D + DispatchThreadId.x;
 	uint2 PixelCoord = DispatchThreadId.xy;
-	uint2 ViewPixelCoord = PixelCoord * DownsampleFactor + View.ViewRectMin.xy / DownsampleFactor;
-
+	//float Epsilon = 1.0e-5;
+	//if (RWLightingTexture[PixelCoord].a < Epsilon)
+	//if (RWLightingTexture[PixelCoord].a <= 0.0)
+	//{
+	//	return;
+	//}
 	float3 Radiance = 0.0;
-	float3 PrevTransmittance = RWLightingTexture[PixelCoord].a;
+	float PrevTransmittance = RWLightingTexture[PixelCoord].a;
 	float PrevOpacity = Luminance(1.0 - PrevTransmittance);
 
+	FBeerShadowMapData BeerShadowMapData = InitializeBeerShadowMapData(PrevTransmittance);
+
 	// Extract depth
+	uint2 ViewPixelCoord = PixelCoord * DownsampleFactor + View.ViewRectMin.xy / DownsampleFactor;
 	float DeviceZ = SceneDepthTexture.Load(int3(ViewPixelCoord, 0)).r;
 #if HAS_INVERTED_Z_BUFFER
 	DeviceZ = max(0.000000000001, DeviceZ);
@@ -626,7 +729,7 @@ void RenderSingleScatteringWithLiveShading(
 
 		float ScatterHitT = HitT.y;
 
-		float3 Transmittance = PrevTransmittance;
+		float3 Transmittance = float3(PrevTransmittance, PrevTransmittance, PrevTransmittance);
 
 #if USE_CAMERA_AVSM
 		FAdaptiveVolumetricShadowMap GlobalTransmittanceMap = GetAdaptiveVolumetricShadowMapForCamera();
@@ -684,6 +787,14 @@ void RenderSingleScatteringWithLiveShading(
 		Transmittance = AVSM_Sampler4_Eval(TransmittanceSampler, WorldRayLength);
 
 #else
+//	#if REFERENCE_FAST_PATH
+//		RayMarchSingleScatteringReferenceFastPath(
+//			RayMarchingContext,
+//			StepCount,
+//			Radiance,
+//			Transmittance
+//		);
+//	#else
 		RayMarchSingleScattering(
 			RayMarchingContext,
 			LightData,
@@ -693,6 +804,7 @@ void RenderSingleScatteringWithLiveShading(
 			Radiance,
 			Transmittance
 		);
+//	#endif
 #endif
 
 		// Radiance
@@ -731,9 +843,29 @@ void RenderSingleScatteringWithLiveShading(
 #endif // DIM_WRITE_VELOCITY
 		}
 
+		// TODO: Conditionalize around CameraAVSM
+		if (bCreateBeerShadowMap)
+		{
+			BeerShadowMapData.T[0] = RayMarchingContext.LocalScatterT0 * RayMarchingContext.LocalToWorldScale;
+			BeerShadowMapData.T[1] = RayMarchingContext.LocalScatterT1 * RayMarchingContext.LocalToWorldScale;
+			BeerShadowMapData.Tau[0] = Luminance(RayMarchingContext.Transmittance0);
+			BeerShadowMapData.Tau[1] = Luminance(Transmittance);
+		}
+
 		// Debug output
 		//RWDebugOutputBuffer[LinearIndex] = DebugOutput;
 	}
+
+	if (bCreateBeerShadowMap)
+	{
+		RWBeerShadowMapTexture[PixelCoord] = PackBeerShadowMapData(
+			CombineBeerShadowMapData(
+				UnpackBeerShadowMapData(RWBeerShadowMapTexture[PixelCoord]),
+				BeerShadowMapData,
+				BEER_SHADOW_MAP_COMBINE_METHOD_EXPAND)
+			// BEER_SHADOW_MAP_COMBINE_METHOD_FIRST)
+		);
+	}
 }
 
 [numthreads(THREADGROUP_SIZE_2D, THREADGROUP_SIZE_2D, 1)]
diff --git a/Engine/Shaders/Private/HeterogeneousVolumes/HeterogeneousVolumesLiveShadingShadows.usf b/Engine/Shaders/Private/HeterogeneousVolumes/HeterogeneousVolumesLiveShadingShadows.usf
--- a/Engine/Shaders/Private/HeterogeneousVolumes/HeterogeneousVolumesLiveShadingShadows.usf
+++ b/Engine/Shaders/Private/HeterogeneousVolumes/HeterogeneousVolumesLiveShadingShadows.usf
@@ -430,8 +430,7 @@ void RenderVolumetricShadowMapForLightWithLiveShadingCS(
 	float HitSpan = HitT.y - HitT.x;
 	if (HitSpan > 0.0)
 	{
-		//float Jitter = bJitter ? BlueNoiseScalar(PixelCoord, View.StateFrameIndexMod8) : 0.0;
-		float Jitter = IS_OFFLINE_RENDER ? BlueNoiseScalar(PixelCoord, View.StateFrameIndexMod8) : 0.0;
+		float Jitter = bJitter ? BlueNoiseScalar(PixelCoord, View.StateFrameIndexMod8) : 0.0;
 		int bApplyEmissionAndTransmittance = 0;
 		int bApplyDirectLighting = 0;
 		int bApplyShadowTransmittance = 0;
@@ -470,6 +469,7 @@ void RenderVolumetricShadowMapForLightWithLiveShadingCS(
 
 		float LocalHitT = HitT.x;
 		float WorldHitT = LocalHitT * RayMarchingContext.LocalToWorldScale;
+		float LastWorldHitT = WorldHitT;
 		for (uint StepIndex = 0; StepIndex < StepCount; ++StepIndex)
 		{
 			LocalHitT = min(RayMarchingContext.LocalRayTMin + StepSize * (RayMarchingContext.Jitter + StepIndex), RayMarchingContext.LocalRayTMax);
@@ -487,31 +487,35 @@ void RenderVolumetricShadowMapForLightWithLiveShadingCS(
 			float3 Extinction = SampleExtinction(SampleContext);
 			Transmittance *= exp(-Extinction * StepSize);
 
-			bool bFirstEntry = (ElementIndex == 0);
-			if (bFirstEntry)
+			if (any(Extinction > 0.0))
 			{
-				FAVSMSampleData SampleData = AVSM_CreateSampleData(WorldHitT, Transmittance);
-				AddSample(SampleData, PixelIndex, PrevElementIndex, ElementIndex);
-
-				// Set expected values
-				ExpectedExtinction = Extinction;
-				ExpectedTransmittance = Transmittance;
-			}
-			else
-			{
-				// If extrapolated transmittance is beyond relative error tolerance
-				ExpectedTransmittance *= exp(-ExpectedExtinction * StepSize);
-				float AbsoluteError = length(Transmittance - ExpectedTransmittance);
-				float RelativeError = length((Transmittance - ExpectedTransmittance) / ExpectedTransmittance);
-				if ((AbsoluteError > AbsoluteErrorThreshold) && (RelativeError > RelativeErrorThreshold))
+				LastWorldHitT = WorldHitT;
+				bool bFirstEntry = (ElementIndex == 0);
+				if (bFirstEntry)
 				{
 					FAVSMSampleData SampleData = AVSM_CreateSampleData(WorldHitT, Transmittance);
-					AddOrReplaceSample(SampleData, PixelIndex, PrevElementIndex, ElementIndex);
+					AddSample(SampleData, PixelIndex, PrevElementIndex, ElementIndex);
 
-					// Reset expected values
+					// Set expected values
 					ExpectedExtinction = Extinction;
 					ExpectedTransmittance = Transmittance;
 				}
+				else
+				{
+					// If extrapolated transmittance is beyond relative error tolerance
+					ExpectedTransmittance *= exp(-ExpectedExtinction * StepSize);
+					float AbsoluteError = length(Transmittance - ExpectedTransmittance);
+					float RelativeError = length((Transmittance - ExpectedTransmittance) / ExpectedTransmittance);
+					if ((AbsoluteError > AbsoluteErrorThreshold) && (RelativeError > RelativeErrorThreshold))
+					{
+						FAVSMSampleData SampleData = AVSM_CreateSampleData(WorldHitT, Transmittance);
+						AddOrReplaceSample(SampleData, PixelIndex, PrevElementIndex, ElementIndex);
+
+						// Reset expected values
+						ExpectedExtinction = Extinction;
+						ExpectedTransmittance = Transmittance;
+					}
+				}
 			}
 
 			// Early termination
@@ -531,7 +535,7 @@ void RenderVolumetricShadowMapForLightWithLiveShadingCS(
 		}
 
 		// Last Entry
-		AVSM_SetLinkedListSampleData(RWVolumetricShadowLinkedListBuffer[PixelIndex + ElementIndex], AVSM_CreateSampleData(WorldHitT, Transmittance));
+		AVSM_SetLinkedListSampleData(RWVolumetricShadowLinkedListBuffer[PixelIndex + ElementIndex], AVSM_CreateSampleData(LastWorldHitT, Transmittance));
 		AVSM_SetLinkedListPrevPtr(RWVolumetricShadowLinkedListBuffer[PixelIndex + ElementIndex], PrevElementIndex);
 		AVSM_SetLinkedListNextPtr(RWVolumetricShadowLinkedListBuffer[PixelIndex + ElementIndex], AVSM_NULL_PTR);
 	}
diff --git a/Engine/Shaders/Private/HeterogeneousVolumes/HeterogeneousVolumesRayMarchingTypes.ush b/Engine/Shaders/Private/HeterogeneousVolumes/HeterogeneousVolumesRayMarchingTypes.ush
--- a/Engine/Shaders/Private/HeterogeneousVolumes/HeterogeneousVolumesRayMarchingTypes.ush
+++ b/Engine/Shaders/Private/HeterogeneousVolumes/HeterogeneousVolumesRayMarchingTypes.ush
@@ -33,6 +33,9 @@ struct FRayMarchingContext
 
 	// TEMPORARY STORAGE!!
 	RandomSequence RandSequence;
+	float LocalScatterT0;
+	float LocalScatterT1;
+	float3 Transmittance0;
 };
 
 FRayMarchingContext CreateRayMarchingContext(
diff --git a/Engine/Shaders/Private/HeterogeneousVolumes/HeterogeneousVolumesRayMarchingUtils.ush b/Engine/Shaders/Private/HeterogeneousVolumes/HeterogeneousVolumesRayMarchingUtils.ush
--- a/Engine/Shaders/Private/HeterogeneousVolumes/HeterogeneousVolumesRayMarchingUtils.ush
+++ b/Engine/Shaders/Private/HeterogeneousVolumes/HeterogeneousVolumesRayMarchingUtils.ush
@@ -396,6 +396,8 @@ void RayMarchSingleScattering(
 	inout float3 Transmittance
 )
 {
+	LocalScatterTMin = RayMarchingContext.LocalRayTMax;
+
 #if FOG_INSCATTERING_MODE
 	float3 FogRadiance = 0;
 	float3 PrevTransmittance = Transmittance;
@@ -436,8 +438,6 @@ void RayMarchSingleScattering(
 	}
 #endif // FOG_INSCATTERING_MODE
 
-	LocalScatterTMin = RayMarchingContext.LocalRayTMax;
-
 	float WorldShadowBias = CalcShadowBias();
 	//float WorldShadowBias = View.GeneralPurposeTweak;
 	float StepSize = RayMarchingContext.StepSize;
@@ -445,6 +445,8 @@ void RayMarchSingleScattering(
 	// Reuse jitter as random sample
 	float RandomSample = RayMarchingContext.Jitter;
 
+	bool bFirstScatterEvent = true;
+
 	for (uint StepIndex = 0; StepIndex < StepCount; ++StepIndex)
 	{
 		// Correct the final step size
@@ -507,13 +509,13 @@ void RayMarchSingleScattering(
 				float IsotropicPhase = 1.0 / (4.0 * PI);
 				RadianceSample = Inscattering * ScatteringCoefficient * IsotropicPhase * Transmittance * StepSize;
 			}
-
+#if 0
 			// Stochastically decide to record scatter position
 			if (any(Transmittance < RandomSample))
 			{
 				LocalScatterTMin = min(LocalHitT, LocalScatterTMin);
 			}
-
+#endif
 #if FOG_INSCATTERING_MODE == FOG_INSCATTERING_MODE_REFERENCE
 			{
 				float WorldRayT = LocalHitT * RayMarchingContext.LocalToWorldScale;
@@ -569,6 +571,19 @@ void RayMarchSingleScattering(
 		//Transmittance = AVSM_Sampler4_Eval_NoInterpolation(TransmittanceSampler, WorldHitT);
 		Transmittance = AVSM_Sampler4_Eval(TransmittanceSampler, WorldHitT);
 #else
+		if (any(Extinction > 0.0))
+		{
+			// NOTE: Trimming the end point is too aggressive for exponential interpolation
+			//RayMarchingContext.LocalScatterT1 = LocalHitT;
+			RayMarchingContext.LocalScatterT1 = RayMarchingContext.LocalRayTMax;
+			if (bFirstScatterEvent)
+			{
+				RayMarchingContext.LocalScatterT0 = LocalScatterTMin = LocalHitT;
+				RayMarchingContext.Transmittance0 = Transmittance;
+				bFirstScatterEvent = false;
+			}
+		}
+
 		// Accumulate transmittance for the next evaluation
 		Transmittance *= exp(-Extinction * StepSize);
 #endif
@@ -583,7 +598,6 @@ void RayMarchSingleScattering(
 		LocalHitT += StepSize;
 	}
 
-
 #if FOG_INSCATTERING_MODE
 	// Simplifying terms to avoid vector division:
 	// float3 DeltaTransmittance = clamp(Transmittance / PrevTransmittance, 0.0, 1.0);
diff --git a/Engine/Shaders/Private/HeterogeneousVolumes/HeterogeneousVolumesVoxelGridShadows.usf b/Engine/Shaders/Private/HeterogeneousVolumes/HeterogeneousVolumesVoxelGridShadows.usf
--- a/Engine/Shaders/Private/HeterogeneousVolumes/HeterogeneousVolumesVoxelGridShadows.usf
+++ b/Engine/Shaders/Private/HeterogeneousVolumes/HeterogeneousVolumesVoxelGridShadows.usf
@@ -13,6 +13,10 @@
 
 #include "../PathTracing/Utilities/PathTracingRandomSequence.ush"
 
+#ifndef USE_AVSM_COMPRESSION
+#define USE_AVSM_COMPRESSION 0
+#endif // USE_AVSM_COMPRESSION
+
 // Ray data
 int bJitter;
 float ShadowStepSize;
@@ -501,6 +505,7 @@ void RenderVolumetricShadowMapForLightWithVoxelGridCS(
 				ExpectedTransmittance *= exp(-ExpectedExtinction * RayMarchingContext.StepSize);
 				float AbsoluteError = length(Transmittance - ExpectedTransmittance);
 				float RelativeError = length((Transmittance - ExpectedTransmittance) / ExpectedTransmittance);
+				//if ((AbsoluteError > AbsoluteErrorThreshold * (1 + RayMarchingContext.Jitter)) && (RelativeError > RelativeErrorThreshold * (1 + RayMarchingContext.Jitter)))
 				if ((AbsoluteError > AbsoluteErrorThreshold) && (RelativeError > RelativeErrorThreshold))
 				{
 					FAVSMSampleData SampleData = AVSM_CreateSampleData(WorldHitT, Transmittance);
@@ -770,3 +775,108 @@ void CombineVolumetricShadowMapsCS(
 		AVSM_SetLinkedListNextPtr(RWVolumetricShadowLinkedListBuffer[PixelIndex + PrevSampleIndex], AVSM_NULL_PTR);
 	}
 }
+
+#define MAX_SAMPLE_COUNT 2
+
+// Beer shadow map
+//Texture2D<float4> BeerShadowMapTexture;
+//int2 ShadowResolution;
+
+struct FBeerShadowMapData
+{
+	float T[2];
+	float Tau[2];
+};
+
+FBeerShadowMapData UnpackBeerShadowMapData(float4 PackedData)
+{
+	FBeerShadowMapData Data;
+	Data.T[0] = PackedData.r;
+	Data.T[1] = PackedData.g;
+	Data.Tau[0] = PackedData.b;
+	Data.Tau[1] = PackedData.a;
+	return Data;
+}
+
+[numthreads(THREADGROUP_SIZE_2D, THREADGROUP_SIZE_2D, 1)]
+void ConvertBeerLawShadowMapToVolumetricShadowMapCS(
+	uint2 GroupThreadId : SV_GroupThreadID,
+	uint3 DispatchThreadId : SV_DispatchThreadID
+)
+{
+	uint ScalarGroupThreadId = GroupThreadId.y * THREADGROUP_SIZE_2D + GroupThreadId.x;
+	if (ScalarGroupThreadId == 0)
+	{
+		GSWarpOffset = AVSM_INVALID_PIXEL_OFFSET;
+	}
+
+	uint Face = DispatchThreadId.z;
+	uint2 PixelCoord = DispatchThreadId.xy;
+	uint LinearIndex = GetLinearIndex(Face, PixelCoord);
+
+	// Allocation is fixed-size
+	int BottomLevelSampleCount = MAX_SAMPLE_COUNT;
+	int TopLevelSampleCount = CalcTopLevelSampleCount(BottomLevelSampleCount);
+	int TotalSampleCount = Align4(TopLevelSampleCount) + Align4(BottomLevelSampleCount);
+
+	uint TileSize = THREADGROUP_SIZE_2D * THREADGROUP_SIZE_2D;
+	uint SampleCountPerTile = TotalSampleCount * TileSize;
+	if ((ScalarGroupThreadId == 0) && (SampleCountPerTile > 0))
+	{
+		InterlockedAdd(RWVolumetricShadowIndirectionAllocatorBuffer[0], SampleCountPerTile, GSWarpOffset);
+	}
+	GroupMemoryBarrierWithGroupSync();
+
+	if (any(PixelCoord < 0) || any(PixelCoord >= ShadowResolution))
+	{
+		return;
+	}
+
+	uint TopLevelOffset = GSWarpOffset + ScalarGroupThreadId * TotalSampleCount;
+	uint BottomLevelOffset = TopLevelOffset + Align4(TopLevelSampleCount);
+
+	FBeerShadowMapData BeerShadowMapData = UnpackBeerShadowMapData(BeerShadowMapTexture[PixelCoord]);
+
+	// Write indirection entry
+	int PixelOffset = TopLevelOffset;
+	FAVSMIndirectionData IndirectionData = AVSM_CreateIndirectionData(
+		PixelOffset,
+		BottomLevelSampleCount,
+		BeerShadowMapData.T[0],
+		BeerShadowMapData.T[1]
+	);
+
+	RWVolumetricShadowIndirectionBuffer[LinearIndex] = AVSM_PackIndirectionData(IndirectionData);
+
+	// Convert beer shadow map samples to AVSM samples
+	FAVSMSampleData ShadowData[] = {
+		AVSM_CreateSampleData(BeerShadowMapData.T[0], BeerShadowMapData.Tau[0]),
+		AVSM_CreateSampleData(BeerShadowMapData.T[1], BeerShadowMapData.Tau[1])
+	};
+
+	for (int SampleIndex = 0; SampleIndex < BottomLevelSampleCount; ++SampleIndex)
+	{
+		// Store indexed list offset by 1 element
+		int BottomLevelIndex = max(SampleIndex - 1, 0);
+		RWVolumetricShadowTransmittanceBuffer[BottomLevelOffset + BottomLevelIndex] = AVSM_PackSampleData(ShadowData[SampleIndex], IndirectionData);
+
+		// Mark the current sample as an acceleration sample at the top-level
+		if (SampleIndex % 4 == 0)
+		{
+			uint TopLevelIndex = SampleIndex / 4;
+			RWVolumetricShadowTransmittanceBuffer[TopLevelOffset + TopLevelIndex] = AVSM_PackSampleData(ShadowData[SampleIndex], IndirectionData);
+		}
+	}
+	
+	// Pad remainder
+	for (int BottomLevelIndex = max(BottomLevelSampleCount - 1, 0); BottomLevelIndex < Align4(BottomLevelSampleCount); ++BottomLevelIndex)
+	{
+		RWVolumetricShadowTransmittanceBuffer[BottomLevelOffset + BottomLevelIndex] = AVSM_PackSampleData(AVSM_CreateSampleData(ShadowData[1].X, ShadowData[1].Tau), IndirectionData);
+	}
+
+	// Pad remainder
+	for (int TopLevelIndex = TopLevelSampleCount; TopLevelIndex < Align4(TopLevelSampleCount); ++TopLevelIndex)
+	{
+		RWVolumetricShadowTransmittanceBuffer[TopLevelOffset + TopLevelIndex] = AVSM_PackSampleData(AVSM_CreateSampleData(ShadowData[1].X, ShadowData[1].Tau), IndirectionData);
+	}
+}
diff --git a/Engine/Source/Runtime/Renderer/Private/HeterogeneousVolumes/HeterogeneousVolumes.cpp b/Engine/Source/Runtime/Renderer/Private/HeterogeneousVolumes/HeterogeneousVolumes.cpp
--- a/Engine/Source/Runtime/Renderer/Private/HeterogeneousVolumes/HeterogeneousVolumes.cpp
+++ b/Engine/Source/Runtime/Renderer/Private/HeterogeneousVolumes/HeterogeneousVolumes.cpp
@@ -29,7 +29,10 @@ static TAutoConsoleVariable<int32> CVarHeterogeneousVolumesShadows(
 static TAutoConsoleVariable<int32> CVarTranslucencyHeterogeneousVolumes(
 	TEXT("r.Translucency.HeterogeneousVolumes"),
 	0,
-	TEXT("Enables composting with heterogeneous volumes when rendering translucency (Default = 0)\n"),
+	TEXT("Enables composting with heterogeneous volumes when rendering translucency (Default = 0)\n")
+	TEXT("0: Off\n")
+	TEXT("1: Camera-based Adaptive Volumetric Shadow Map\n")
+	TEXT("2: Camera-based Beer-Law Shadow Map\n"),
 	ECVF_RenderThreadSafe
 );
 
@@ -344,6 +347,15 @@ bool ShouldCompositeHeterogeneousVolumesWithTranslucency()
 	return CVarTranslucencyHeterogeneousVolumes.GetValueOnAnyThread() != 0;
 }
 
+namespace HeterogeneousVolumes
+{
+	ETranslucencyCompositingMode GetTranslucencyCompositingMode()
+	{
+		int32 ClampedValue = FMath::Clamp(CVarTranslucencyHeterogeneousVolumes.GetValueOnRenderThread(), 0, 2);
+		return static_cast<ETranslucencyCompositingMode>(ClampedValue);
+	}
+}
+
 EHeterogeneousVolumesCompositionType GetHeterogeneousVolumesCompositionType()
 {
 	int32 CompositionOrder = CVarHeterogeneousVolumesComposition.GetValueOnRenderThread();
@@ -481,9 +493,14 @@ namespace HeterogeneousVolumes
 		return FMath::Clamp(CVarHeterogeneousVolumesFilterWidth.GetValueOnRenderThread(), 0, 11);
 	}
 
+	FIntPoint GetDownsampledResolution(FIntPoint Resolution, int32 DownsampleFactor)
+	{
+		return FIntPoint::DivideAndRoundUp(Resolution, DownsampleFactor);
+	}
+
 	FIntPoint GetScaledViewRect(FIntRect ViewRect)
 	{
-		return FIntPoint::DivideAndRoundUp(ViewRect.Size(), GetDownsampleFactor());
+		return GetDownsampledResolution(ViewRect.Size(), GetDownsampleFactor());
 	}
 
 	FIntVector GetVolumeResolution(const IHeterogeneousVolumeInterface* Interface)
@@ -895,7 +912,7 @@ void FDeferredShadingSceneRenderer::RenderHeterogeneousVolumeShadows(
 	{
 		FViewInfo& View = Views[ViewIndex];
 
-		if (ShouldCompositeHeterogeneousVolumesWithTranslucency())
+		if (HeterogeneousVolumes::GetTranslucencyCompositingMode() == HeterogeneousVolumes::ETranslucencyCompositingMode::VolumetricShadowMap)
 		{
 			if (HeterogeneousVolumes::GetShadowMode() == HeterogeneousVolumes::EShadowMode::LiveShading)
 			{
@@ -979,10 +996,11 @@ void FDeferredShadingSceneRenderer::RenderHeterogeneousVolumes(
 
 	FRDGTextureRef HeterogeneousVolumeRadiance = GSystemTextures.GetBlackDummy(GraphBuilder);
 	FRDGTextureRef HeterogeneousVolumeHoldout = GSystemTextures.GetBlackDummy(GraphBuilder);
+	FRDGTextureRef HeterogeneousVolumeBeerShadowMap = GSystemTextures.GetBlackDummy(GraphBuilder);
 	if (ShouldRenderHeterogeneousVolumesForAnyView(Views))
 	{
 		FRDGTextureDesc Desc = SceneTextures.Color.Target->Desc;
-		Desc.Extent /= HeterogeneousVolumes::GetDownsampleFactor();
+		Desc.Extent = HeterogeneousVolumes::GetDownsampledResolution(Desc.Extent, HeterogeneousVolumes::GetDownsampleFactor());
 		Desc.Format = PF_FloatRGBA;
 		Desc.Flags &= ~(TexCreate_FastVRAM);
 		HeterogeneousVolumeRadiance = GraphBuilder.CreateTexture(Desc, TEXT("HeterogeneousVolumes"));
@@ -994,6 +1012,15 @@ void FDeferredShadingSceneRenderer::RenderHeterogeneousVolumes(
 			HeterogeneousVolumeHoldout = GraphBuilder.CreateTexture(Desc, TEXT("HeterogeneousVolume.Holdout"));
 			AddClearUAVPass(GraphBuilder, GraphBuilder.CreateUAV(HeterogeneousVolumeHoldout), FLinearColor::Black);
 		}
+
+		// Satisfy binding UAV binding requirement
+		HeterogeneousVolumeBeerShadowMap = HeterogeneousVolumeRadiance;
+		if (HeterogeneousVolumes::GetTranslucencyCompositingMode() == HeterogeneousVolumes::ETranslucencyCompositingMode::BeerShadowMap)
+		{
+			Desc.Format = PF_FloatRGBA;
+			HeterogeneousVolumeBeerShadowMap = GraphBuilder.CreateTexture(Desc, TEXT("HeterogeneousVolume.BeerShadowMap"));
+			AddClearUAVPass(GraphBuilder, GraphBuilder.CreateUAV(HeterogeneousVolumeBeerShadowMap), FVector4d(HeterogeneousVolumes::GetMaxTraceDistance(), 0.0, 1.0, 1.0));
+		}
 	}
 
 	for (int32 ViewIndex = 0; ViewIndex < Views.Num(); ViewIndex++)
@@ -1148,14 +1175,96 @@ void FDeferredShadingSceneRenderer::RenderHeterogeneousVolumes(
 							LightingCacheTexture,
 							// Output
 							HeterogeneousVolumeRadiance,
-							HeterogeneousVolumeHoldout
+							HeterogeneousVolumeHoldout,
+							HeterogeneousVolumeBeerShadowMap
 						);
+
+						// Validate that a view state exists to store the AVSM
+						if (View.ViewState)
+						{
+							// Stash the temporary beer-law shadow map as a camera shadow if we aren't already creating one
+							if (HeterogeneousVolumes::GetTranslucencyCompositingMode() == HeterogeneousVolumes::ETranslucencyCompositingMode::BeerShadowMap)
+							{
+								// Resolution
+								FIntPoint ShadowMapResolution = HeterogeneousVolumes::GetDownsampledResolution(View.ViewRect.Size(), HeterogeneousVolumes::GetDownsampleFactor());
+								FRDGBufferRef VolumetricShadowIndirectionBuffer = GSystemTextures.GetDefaultStructuredBuffer(GraphBuilder, sizeof(FAVSMIndirectionPackedData));
+								FRDGBufferRef VolumetricShadowSampleBuffer = GSystemTextures.GetDefaultStructuredBuffer(GraphBuilder, sizeof(FAVSMSamplePackedData));
+
+								ConvertBeerLawShadowMapToVolumetricShadowMap(
+									GraphBuilder,
+									View,
+									//GroupCount,
+									ShadowMapResolution,
+									HeterogeneousVolumeBeerShadowMap,
+									VolumetricShadowIndirectionBuffer,
+									VolumetricShadowSampleBuffer
+								);
+
+								// Transform
+								const FMatrix ProjectionMatrix = View.ViewMatrices.GetProjectionMatrix();
+								float FOV = FMath::DegreesToRadians(View.FOV * 0.5);
+								FMatrix ViewToClip = FPerspectiveMatrix(
+									FOV,
+									ShadowMapResolution.X,
+									ShadowMapResolution.Y,
+									1.0,
+									HeterogeneousVolumes::GetMaxTraceDistance()
+								);
+								FMatrix ClipToView = ViewToClip.Inverse();
+								FMatrix ScreenMatrix = FScaleMatrix(FVector(0.5, -0.5, -0.5)) * FTranslationMatrix(FVector(0.5, 0.5, 0.5));
+
+								int32 NumShadowMatrices = 1;
+								FMatrix44f TranslatedWorldToShadow[] = {
+									FMatrix44f(View.ViewMatrices.GetTranslatedViewMatrix() * ViewToClip * ScreenMatrix)
+								};
+								FMatrix44f ShadowToTranslatedWorld[] = {
+									TranslatedWorldToShadow[0].Inverse()
+								};
+								FVector3f TranslatedWorldOrigin = ShadowToTranslatedWorld[0].GetOrigin();
+
+								// Generic data
+								int32 MaxSampleCount = 2;
+								FVector4f TranslatedWorldPlane = FVector4f::Zero();
+								float DownsampleFactor = HeterogeneousVolumes::GetDownsampleFactor();
+								bool bIsDirectionalLight = false;
+								FRDGBufferRef VolumetricShadowLinkedListBuffer = GSystemTextures.GetDefaultStructuredBuffer(GraphBuilder, sizeof(FAVSMLinkedListPackedData));
+
+								// Uniform buffer
+								CreateAdaptiveVolumetricShadowMapUniformBufferParameters(
+									GraphBuilder,
+									TranslatedWorldOrigin,
+									TranslatedWorldPlane,
+									TranslatedWorldToShadow,
+									ShadowMapResolution,
+									DownsampleFactor,
+									NumShadowMatrices,
+									MaxSampleCount,
+									bIsDirectionalLight,
+									VolumetricShadowLinkedListBuffer,
+									VolumetricShadowIndirectionBuffer,
+									VolumetricShadowSampleBuffer,
+									View.ViewState->AdaptiveVolumetricCameraMapUniformBufferParameters
+								);
+							}
+
+							// Append beauty image and build uniform buffer
+							if (ShouldCompositeHeterogeneousVolumesWithTranslucency() && View.ViewState->AdaptiveVolumetricCameraMapUniformBufferParameters)
+							{
+								// Append beauty image
+								View.ViewState->AdaptiveVolumetricCameraMapUniformBufferParameters->RadianceTexture = GraphBuilder.CreateSRV(HeterogeneousVolumeRadiance);
+								View.ViewState->AdaptiveVolumetricCameraMapUniformBufferParameters->TextureSampler = TStaticSamplerState<SF_Bilinear, AM_Clamp, AM_Clamp, AM_Clamp>::GetRHI();
+
+								View.ViewState->AdaptiveVolumetricCameraMapUniformBuffer =
+									GraphBuilder.CreateUniformBuffer(View.ViewState->AdaptiveVolumetricCameraMapUniformBufferParameters);
+							}
+						}
 					}
 				}
 			}
 
 			View.HeterogeneousVolumeRadiance = HeterogeneousVolumeRadiance;
 			View.HeterogeneousVolumeHoldout = HeterogeneousVolumeHoldout;
+			View.HeterogeneousVolumeBeerShadowMap = HeterogeneousVolumeBeerShadowMap;
 		}
 	}
 }
diff --git a/Engine/Source/Runtime/Renderer/Private/HeterogeneousVolumes/HeterogeneousVolumes.h b/Engine/Source/Runtime/Renderer/Private/HeterogeneousVolumes/HeterogeneousVolumes.h
--- a/Engine/Source/Runtime/Renderer/Private/HeterogeneousVolumes/HeterogeneousVolumes.h
+++ b/Engine/Source/Runtime/Renderer/Private/HeterogeneousVolumes/HeterogeneousVolumes.h
@@ -63,6 +63,7 @@ namespace HeterogeneousVolumes
 	FIntVector GetVolumeResolution(const IHeterogeneousVolumeInterface*);
 
 	int32 GetDownsampleFactor();
+	FIntPoint GetDownsampledResolution(FIntPoint Resolution, int32 DownsampleFactor);
 	FIntPoint GetScaledViewRect(FIntRect ViewRect);
 	float GetShadowStepSize();
 	float GetMaxTraceDistance();
@@ -100,6 +101,15 @@ namespace HeterogeneousVolumes
 	};
 	EIndirectLightingMode GetIndirectLightingMode();
 
+	// Translucency compositing
+	enum class ETranslucencyCompositingMode
+	{
+		Disabled,
+		VolumetricShadowMap,
+		BeerShadowMap
+	};
+	ETranslucencyCompositingMode GetTranslucencyCompositingMode();
+
 	int32 GetMipLevel();
 	int32 GetDebugMode();
 	int32 GetLightingCacheMode();
@@ -450,10 +460,14 @@ BEGIN_UNIFORM_BUFFER_STRUCT(FAdaptiveVolumetricShadowMapUniformBufferParameters,
 	SHADER_PARAMETER(int32, MaxSampleCount)
 	SHADER_PARAMETER(int32, bIsEmpty)
 	SHADER_PARAMETER(int32, bIsDirectionalLight)
+	SHADER_PARAMETER(float, DownsampleFactor)
 
 	SHADER_PARAMETER_RDG_BUFFER_SRV(StructuredBuffer<uint2>, LinkedListBuffer)
 	SHADER_PARAMETER_RDG_BUFFER_SRV(StructuredBuffer<uint4>, IndirectionBuffer)
 	SHADER_PARAMETER_RDG_BUFFER_SRV(StructuredBuffer<uint>, SampleBuffer)
+
+	SHADER_PARAMETER_RDG_TEXTURE_SRV(Texture2D, RadianceTexture)
+	SHADER_PARAMETER_SAMPLER(SamplerState, TextureSampler)
 END_UNIFORM_BUFFER_STRUCT()
 
 BEGIN_SHADER_PARAMETER_STRUCT(FAdaptiveVolumetricShadowMapParameters, RENDERER_API)
@@ -466,10 +480,14 @@ BEGIN_SHADER_PARAMETER_STRUCT(FAdaptiveVolumetricShadowMapParameters, RENDERER_A
 	SHADER_PARAMETER(int32, MaxSampleCount)
 	SHADER_PARAMETER(int32, bIsEmpty)
 	SHADER_PARAMETER(int32, bIsDirectionalLight)
+	SHADER_PARAMETER(float, DownsampleFactor)
 
 	SHADER_PARAMETER_RDG_BUFFER_SRV(StructuredBuffer<uint2>, LinkedListBuffer)
 	SHADER_PARAMETER_RDG_BUFFER_SRV(StructuredBuffer<uint4>, IndirectionBuffer)
 	SHADER_PARAMETER_RDG_BUFFER_SRV(StructuredBuffer<uint>, SampleBuffer)
+	
+	SHADER_PARAMETER_RDG_TEXTURE_SRV(Texture2D, RadianceTexture)
+	SHADER_PARAMETER_SAMPLER(SamplerState, TextureSampler)
 END_SHADER_PARAMETER_STRUCT()
 
 BEGIN_UNIFORM_BUFFER_STRUCT(FAdaptiveVolumetricShadowMaps, )
@@ -478,36 +496,6 @@ BEGIN_UNIFORM_BUFFER_STRUCT(FAdaptiveVolumetricShadowMaps, )
 END_UNIFORM_BUFFER_STRUCT()
 
 namespace HeterogeneousVolumes {
-	struct FAdaptiveVolumetricShadowMapParameterCache
-	{
-		FMatrix44f TranslatedWorldToShadow[6] =
-		{
-			FMatrix44f::Identity,
-			FMatrix44f::Identity,
-			FMatrix44f::Identity,
-			FMatrix44f::Identity,
-			FMatrix44f::Identity,
-			FMatrix44f::Identity
-		};
-		FVector3f TranslatedWorldOrigin = FVector3f::ZeroVector;
-		FVector4f TranslatedWorldPlane = FVector4f::Zero();
-		FIntPoint Resolution = FIntPoint::ZeroValue;
-		int32 NumShadowMatrices = 1;
-		int32 MaxSampleCount = 0;
-		bool bIsEmpty = true;
-		bool bIsDirectionalLight = false;
-
-		TRefCountPtr<FRDGPooledBuffer> LinkedListBuffer;
-		TRefCountPtr<FRDGPooledBuffer> IndirectionBuffer;
-		TRefCountPtr<FRDGPooledBuffer> SampleBuffer;
-	};
-
-	struct FAdaptiveVolumetricShadowMapState
-	{
-		TRDGUniformBufferRef<FAdaptiveVolumetricShadowMapUniformBufferParameters> UniformBuffer = nullptr;
-		FAdaptiveVolumetricShadowMapParameterCache PrevFrameParameterCache;
-	};
-
 	TRDGUniformBufferRef<FAdaptiveVolumetricShadowMapUniformBufferParameters> GetAdaptiveVolumetricShadowMapUniformBuffer(
 		FRDGBuilder& GraphBuilder,
 		FSceneViewState* ViewState,
@@ -569,7 +557,8 @@ void RenderWithLiveShading(
 	FRDGTextureRef LightingCacheTexture,
 	// Output
 	FRDGTextureRef& HeterogeneousVolumeRadiance,
-	FRDGTextureRef& HeterogeneousVolumeHoldout
+	FRDGTextureRef& HeterogeneousVolumeHoldout,
+	FRDGTextureRef& HeterogeneousVolumeBeerShadowMap
 );
 
 void RenderWithPreshading(
@@ -674,12 +663,40 @@ void CombineVolumetricShadowMap(
 	FRDGBufferRef& VolumetricShadowLinkedListBuffer
 );
 
+void ConvertBeerLawShadowMapToVolumetricShadowMap(
+	FRDGBuilder& GraphBuilder,
+	FViewInfo& View,
+	// Input
+	FIntPoint ShadowMapResolution,
+	FRDGTextureRef BeerShadowMapTexture,
+	// Output
+	FRDGBufferRef& VolumetricShadowIndirectionBuffer,
+	FRDGBufferRef& VolumetricShadowTransmittanceBuffer
+);
+
+void CreateAdaptiveVolumetricShadowMapUniformBufferParameters(
+	FRDGBuilder& GraphBuilder,
+	const FVector3f& TranslatedWorldOrigin,
+	const FVector4f& TranslatedWorldPlane,
+	const FMatrix44f* TranslatedWorldToShadow,
+	FIntPoint VolumetricShadowMapResolution,
+	float VolumetricShadowMapDownsampleFactor,
+	int32 NumShadowMatrices,
+	uint32 VolumetricShadowMapMaxSampleCount,
+	bool bIsDirectionalLight,
+	FRDGBufferRef VolumetricShadowMapLinkedListBuffer,
+	FRDGBufferRef VolumetricShadowMapIndirectionBuffer,
+	FRDGBufferRef VolumetricShadowMapSampleBuffer,
+	FAdaptiveVolumetricShadowMapUniformBufferParameters*& AdaptiveVolumetricShadowMapUniformBufferParameters
+);
+
 void CreateAdaptiveVolumetricShadowMapUniformBuffer(
 	FRDGBuilder& GraphBuilder,
 	const FVector3f& TranslatedWorldOrigin,
 	const FVector4f& TranslatedWorldPlane,
 	const FMatrix44f* TranslatedWorldToShadow,
 	FIntPoint VolumetricShadowMapResolution,
+	float VolumetricShadowMapDownsampleFactor,
 	int32 NumShadowMatrices,
 	uint32 VolumetricShadowMapMaxSampleCount,
 	bool bIsDirectionalLight,
diff --git a/Engine/Source/Runtime/Renderer/Private/HeterogeneousVolumes/HeterogeneousVolumesLiveShadingPipeline.cpp b/Engine/Source/Runtime/Renderer/Private/HeterogeneousVolumes/HeterogeneousVolumesLiveShadingPipeline.cpp
--- a/Engine/Source/Runtime/Renderer/Private/HeterogeneousVolumes/HeterogeneousVolumesLiveShadingPipeline.cpp
+++ b/Engine/Source/Runtime/Renderer/Private/HeterogeneousVolumes/HeterogeneousVolumesLiveShadingPipeline.cpp
@@ -104,6 +104,13 @@ static TAutoConsoleVariable<int32> CVarHeterogeneousVolumesAdaptiveMarching(
 );
 #endif
 
+static TAutoConsoleVariable<bool> CVarHeterogeneousVolumesReferenceFastPath(
+	TEXT("r.HeterogeneousVolumes.ReferenceFastPath"),
+	0,
+	TEXT("Enables minimal VGPR mode (Default = 0)"),
+	ECVF_RenderThreadSafe
+);
+
 namespace HeterogeneousVolumes
 {
 	bool SupportsLightType(uint32 LightType)
@@ -153,6 +160,10 @@ namespace HeterogeneousVolumes
 		return CVarHeterogeneousVolumesAdaptiveMarching.GetValueOnRenderThread() != 0;
 	}
 #endif
+	bool UseReferenceFastPath()
+	{
+		return CVarHeterogeneousVolumesReferenceFastPath.GetValueOnRenderThread() != 0;
+	}
 
 	bool ShadowsUseCameraSceneDepth()
 	{
@@ -431,6 +442,7 @@ class FRenderSingleScatteringWithLiveShadingCS : public FMeshMaterialShader
 	class FIsOfflineRender : SHADER_PERMUTATION_INT("IS_OFFLINE_RENDER", 2);
 	class FFogInscatteringMode : SHADER_PERMUTATION_INT("FOG_INSCATTERING_MODE", 3);
 	class FUseAnalyticDerivatives : SHADER_PERMUTATION_BOOL("USE_ANALYTIC_DERIVATIVES");
+	class FReferenceFastPath : SHADER_PERMUTATION_BOOL("REFERENCE_FAST_PATH");
 	using FPermutationDomain = TShaderPermutationDomain<
 		//FUseTransmittanceVolume, 
 		FUseInscatteringVolume, 
@@ -440,7 +452,8 @@ class FRenderSingleScatteringWithLiveShadingCS : public FMeshMaterialShader
 		FAVSMSampleMode, 
 		FSupportOverlappingVolumes,
 		FIsOfflineRender,
-		FUseAnalyticDerivatives
+		FUseAnalyticDerivatives,
+		FReferenceFastPath
 		//FBilinearInterpolation,
 		//FAdaptiveMarch
 	>;
@@ -477,6 +490,7 @@ class FRenderSingleScatteringWithLiveShadingCS : public FMeshMaterialShader
 		SHADER_PARAMETER_RDG_UNIFORM_BUFFER(FFogUniformParameters, FogStruct)
 		SHADER_PARAMETER(int, bApplyHeightFog)
 		SHADER_PARAMETER(int, bApplyVolumetricFog)
+		SHADER_PARAMETER(int, bCreateBeerShadowMap)
 
 		// Indirect Lighting
 		SHADER_PARAMETER_RDG_UNIFORM_BUFFER(FLumenTranslucencyLightingUniforms, LumenGIVolumeStruct)
@@ -516,6 +530,7 @@ class FRenderSingleScatteringWithLiveShadingCS : public FMeshMaterialShader
 		SHADER_PARAMETER_RDG_TEXTURE_UAV(RWTexture2D<float4>, RWLightingTexture)
 		SHADER_PARAMETER_RDG_TEXTURE_UAV(RWTexture2D<float>, RWHoldoutTexture)
 		SHADER_PARAMETER_RDG_TEXTURE_UAV(RWTexture2D<float4>, RWVelocityTexture)
+		SHADER_PARAMETER_RDG_TEXTURE_UAV(RWTexture2D<float4>, RWBeerShadowMapTexture)
 		//SHADER_PARAMETER_RDG_BUFFER_UAV(RWStructuredBuffer<Volumes::FDebugOutput>, RWDebugOutputBuffer)
 	END_SHADER_PARAMETER_STRUCT()
 
@@ -1117,7 +1132,8 @@ void RenderSingleScatteringWithLiveShading(
 	FRDGTextureRef AmbientOcclusionTexture,
 	// Output
 	FRDGTextureRef& HeterogeneousVolumeRadianceTexture,
-	FRDGTextureRef& HeterogeneousVolumeHoldoutTexture
+	FRDGTextureRef& HeterogeneousVolumeHoldoutTexture,
+	FRDGTextureRef& HeterogeneousVolumeBeerShadowMapTexture
 )
 {
 	typedef FRenderSingleScatteringWithLiveShadingCS<DispatchMode> FRenderSingleScatteringWithLiveShadingDispatchTypeCS;
@@ -1244,6 +1260,8 @@ void RenderSingleScatteringWithLiveShading(
 		PassParameters->FogStruct = FogBuffer;
 		PassParameters->bApplyHeightFog = HeterogeneousVolumes::ShouldApplyHeightFog();
 		PassParameters->bApplyVolumetricFog = HeterogeneousVolumes::ShouldApplyVolumetricFog();
+		bool bCreateBeerShadowMap = HeterogeneousVolumes::GetTranslucencyCompositingMode() == HeterogeneousVolumes::ETranslucencyCompositingMode::BeerShadowMap;
+		PassParameters->bCreateBeerShadowMap = bCreateBeerShadowMap;
 
 		// Indirect lighting data
 		PassParameters->IndirectInscatteringFactor = HeterogeneousVolumes::GetIndirectLightingFactor();
@@ -1295,6 +1313,10 @@ void RenderSingleScatteringWithLiveShading(
 		{
 			PassParameters->RWVelocityTexture = GraphBuilder.CreateUAV(SceneTextures.Velocity);
 		}
+		//if (bCreateBeerShadowMap)
+		{
+			PassParameters->RWBeerShadowMapTexture = GraphBuilder.CreateUAV(HeterogeneousVolumeBeerShadowMapTexture);
+		}
 		//PassParameters->RWDebugOutputBuffer = GraphBuilder.CreateUAV(DebugOutputBuffer);
 	}
 
@@ -1325,6 +1347,7 @@ void RenderSingleScatteringWithLiveShading(
 	PermutationVector.template Set<typename FRenderSingleScatteringWithLiveShadingDispatchTypeCS::FIsOfflineRender>(View.bIsOfflineRender);
 	PermutationVector.template Set<typename FRenderSingleScatteringWithLiveShadingDispatchTypeCS::FFogInscatteringMode>(static_cast<int32>(HeterogeneousVolumes::GetFogInscatteringMode()));
 	PermutationVector.template Set<typename FRenderSingleScatteringWithLiveShadingDispatchTypeCS::FUseAnalyticDerivatives>(HeterogeneousVolumes::UseAnalyticDerivatives());
+	PermutationVector.template Set<typename FRenderSingleScatteringWithLiveShadingDispatchTypeCS::FReferenceFastPath>(HeterogeneousVolumes::UseReferenceFastPath());
 	PermutationVector = FRenderSingleScatteringWithLiveShadingDispatchTypeCS::RemapPermutation(PermutationVector);
 	TShaderRef<FRenderSingleScatteringWithLiveShadingDispatchTypeCS> ComputeShader = Material.GetShader<FRenderSingleScatteringWithLiveShadingDispatchTypeCS>(&FLocalVertexFactory::StaticType, PermutationVector, false);
 	if (!ComputeShader.IsNull())
@@ -1351,7 +1374,8 @@ static void RenderWithTransmittanceVolumePipeline(
 	FRDGTextureRef AmbientOcclusionTexture,
 	// Output
 	FRDGTextureRef& HeterogeneousVolumeRadiance,
-	FRDGTextureRef& HeterogeneousVolumeHoldout
+	FRDGTextureRef& HeterogeneousVolumeHoldout,
+	FRDGTextureRef& HeterogeneousVolumeBeerShadowMap
 )
 {
 	// Light culling
@@ -1448,7 +1472,8 @@ static void RenderWithTransmittanceVolumePipeline(
 				AmbientOcclusionTexture,
 				// Output
 				HeterogeneousVolumeRadiance,
-				HeterogeneousVolumeHoldout
+				HeterogeneousVolumeHoldout,
+				HeterogeneousVolumeBeerShadowMap
 			);
 		}
 		else
@@ -1479,7 +1504,8 @@ static void RenderWithTransmittanceVolumePipeline(
 				AmbientOcclusionTexture,
 				// Output
 				HeterogeneousVolumeRadiance,
-				HeterogeneousVolumeHoldout
+				HeterogeneousVolumeHoldout,
+				HeterogeneousVolumeBeerShadowMap
 			);
 		}
 	}
@@ -1503,7 +1529,8 @@ static void RenderWithInscatteringVolumePipeline(
 	FRDGTextureRef AmbientOcclusionTexture,
 	// Output
 	FRDGTextureRef& HeterogeneousVolumeRadiance,
-	FRDGTextureRef& HeterogeneousVolumeHoldout
+	FRDGTextureRef& HeterogeneousVolumeHoldout,
+	FRDGTextureRef& HeterogeneousVolumeBeerShadowMap
 )
 {
 	CSV_SCOPED_TIMING_STAT_EXCLUSIVE(Effects);
@@ -1617,7 +1644,8 @@ static void RenderWithInscatteringVolumePipeline(
 				AmbientOcclusionTexture,
 				// Output
 				HeterogeneousVolumeRadiance,
-				HeterogeneousVolumeHoldout
+				HeterogeneousVolumeHoldout,
+				HeterogeneousVolumeBeerShadowMap
 			);
 		}
 		else
@@ -1648,7 +1676,8 @@ static void RenderWithInscatteringVolumePipeline(
 				AmbientOcclusionTexture,
 				// Output
 				HeterogeneousVolumeRadiance,
-				HeterogeneousVolumeHoldout
+				HeterogeneousVolumeHoldout,
+				HeterogeneousVolumeBeerShadowMap
 			);
 		}
 	}
@@ -1671,7 +1700,8 @@ void RenderWithLiveShading(
 	FRDGTextureRef LightingCacheTexture,
 	// Output
 	FRDGTextureRef& HeterogeneousVolumeRadiance,
-	FRDGTextureRef& HeterogeneousVolumeHoldout
+	FRDGTextureRef& HeterogeneousVolumeHoldout,
+	FRDGTextureRef& HeterogeneousVolumeBeerShadowMap
 )
 {
 	FRDGTextureRef AmbientOcclusionTexture;
@@ -1710,7 +1740,8 @@ void RenderWithLiveShading(
 			AmbientOcclusionTexture,
 			// Output
 			HeterogeneousVolumeRadiance,
-			HeterogeneousVolumeHoldout
+			HeterogeneousVolumeHoldout,
+			HeterogeneousVolumeBeerShadowMap
 		);
 	}
 	else
@@ -1733,7 +1764,8 @@ void RenderWithLiveShading(
 			AmbientOcclusionTexture,
 			// Output
 			HeterogeneousVolumeRadiance,
-			HeterogeneousVolumeHoldout
+			HeterogeneousVolumeHoldout,
+			HeterogeneousVolumeBeerShadowMap
 		);
 	}
 }
@@ -2499,12 +2531,14 @@ void RenderAdaptiveVolumetricShadowMapWithLiveShading(
 						VolumetricShadowMapSampleBuffer
 					);
 
+					float DownsampleFactor = 1.0f;
 					CreateAdaptiveVolumetricShadowMapUniformBuffer(
 						GraphBuilder,
 						TranslatedWorldOrigin,
 						TranslatedWorldPlane,
 						TranslatedWorldToShadow,
 						VolumetricShadowMapResolution,
+						DownsampleFactor,
 						NumShadowMatrices,
 						VolumetricShadowMapMaxSampleCount,
 						bIsDirectionalLight,
@@ -2547,8 +2581,6 @@ void RenderAdaptiveVolumetricCameraMapWithLiveShading(
 	}
 	RDG_EVENT_SCOPE(GraphBuilder, "Adaptive Volumetric Camera Map");
 
-	TRDGUniformBufferRef<FAdaptiveVolumetricShadowMapUniformBufferParameters> AdaptiveVolumetricShadowMapUniformBuffer;
-
 	// Collect all volumes for view
 	bool bCollectForShadowCasting = false;
 	TSet<FVolumetricMeshBatch> HeterogeneousVolumesMeshBatches;
@@ -2565,10 +2597,7 @@ void RenderAdaptiveVolumetricCameraMapWithLiveShading(
 	if (bShouldRenderCameraMap)
 	{
 		// Resolution
-		FIntPoint ShadowMapResolution = FIntPoint(View.ViewRect.Width(), View.ViewRect.Height());
-		float DownsampleFactor = HeterogeneousVolumes::GetCameraDownsampleFactor();
-		ShadowMapResolution.X = FMath::Max(ShadowMapResolution.X / DownsampleFactor, 1);
-		ShadowMapResolution.Y = FMath::Max(ShadowMapResolution.Y / DownsampleFactor, 1);
+		FIntPoint ShadowMapResolution = HeterogeneousVolumes::GetDownsampledResolution(View.ViewRect.Size(), HeterogeneousVolumes::GetCameraDownsampleFactor());
 
 		// Transform
 		const FMatrix ProjectionMatrix = View.ViewMatrices.GetProjectionMatrix();
@@ -2727,26 +2756,22 @@ void RenderAdaptiveVolumetricCameraMapWithLiveShading(
 		);
 
 		FVector4f TranslatedWorldPlane = FVector4f::Zero();
+		float DownsampleFactor = HeterogeneousVolumes::GetCameraDownsampleFactor();
 		bool bIsDirectionalLight = false;
-		CreateAdaptiveVolumetricShadowMapUniformBuffer(
+		CreateAdaptiveVolumetricShadowMapUniformBufferParameters(
 			GraphBuilder,
 			TranslatedWorldOrigin,
 			TranslatedWorldPlane,
 			TranslatedWorldToShadow,
 			ShadowMapResolution,
+			DownsampleFactor,
 			NumShadowMatrices,
 			MaxSampleCount,
 			bIsDirectionalLight,
 			VolumetricShadowLinkedListBuffer,
 			VolumetricShadowIndirectionBuffer,
 			VolumetricShadowSampleBuffer,
-			AdaptiveVolumetricShadowMapUniformBuffer
+			View.ViewState->AdaptiveVolumetricCameraMapUniformBufferParameters
 		);
 	}
-	else
-	{
-		AdaptiveVolumetricShadowMapUniformBuffer = HeterogeneousVolumes::CreateEmptyAdaptiveVolumetricShadowMapUniformBuffer(GraphBuilder);
-	}
-
-	View.ViewState->AdaptiveVolumetricCameraMapUniformBuffer = AdaptiveVolumetricShadowMapUniformBuffer;
 }
\ No newline at end of file
diff --git a/Engine/Source/Runtime/Renderer/Private/HeterogeneousVolumes/HeterogeneousVolumesVoxelGridPipeline.cpp b/Engine/Source/Runtime/Renderer/Private/HeterogeneousVolumes/HeterogeneousVolumesVoxelGridPipeline.cpp
--- a/Engine/Source/Runtime/Renderer/Private/HeterogeneousVolumes/HeterogeneousVolumesVoxelGridPipeline.cpp
+++ b/Engine/Source/Runtime/Renderer/Private/HeterogeneousVolumes/HeterogeneousVolumesVoxelGridPipeline.cpp
@@ -3191,44 +3191,194 @@ class FCombineVolumetricShadowMapsCS : public FGlobalShader
 
 IMPLEMENT_GLOBAL_SHADER(FCombineVolumetricShadowMapsCS, "/Engine/Private/HeterogeneousVolumes/HeterogeneousVolumesVoxelGridShadows.usf", "CombineVolumetricShadowMapsCS", SF_Compute);
 
-void CreateAdaptiveVolumetricShadowMapUniformBuffer(
+class FConvertBeerLawShadowMapToVolumetricShadowMapCS : public FGlobalShader
+{
+	DECLARE_GLOBAL_SHADER(FConvertBeerLawShadowMapToVolumetricShadowMapCS);
+	SHADER_USE_PARAMETER_STRUCT(FConvertBeerLawShadowMapToVolumetricShadowMapCS, FGlobalShader);
+
+	using FPermutationDomain = TShaderPermutationDomain<>;
+
+	BEGIN_SHADER_PARAMETER_STRUCT(FParameters, )
+		// Input
+		SHADER_PARAMETER(FIntPoint, ShadowResolution)
+		SHADER_PARAMETER_RDG_TEXTURE_SRV(Texture2D<float4>, BeerShadowMapTexture)
+
+		// Output
+		SHADER_PARAMETER_RDG_BUFFER_UAV(RWBuffer<int>, RWVolumetricShadowIndirectionAllocatorBuffer)
+		SHADER_PARAMETER_RDG_BUFFER_UAV(RWStructuredBuffer<FAVSMIndirectionPackedData>, RWVolumetricShadowIndirectionBuffer)
+		SHADER_PARAMETER_RDG_BUFFER_UAV(RWStructuredBuffer<FAVSMSamplePackedData>, RWVolumetricShadowTransmittanceBuffer)
+	END_SHADER_PARAMETER_STRUCT()
+
+	static bool ShouldCompilePermutation(
+		const FGlobalShaderPermutationParameters& Parameters
+	)
+	{
+		return DoesPlatformSupportHeterogeneousVolumes(Parameters.Platform);
+	}
+
+	static void ModifyCompilationEnvironment(
+		const FGlobalShaderPermutationParameters& Parameters,
+		FShaderCompilerEnvironment& OutEnvironment
+	)
+	{
+		FMaterialShader::ModifyCompilationEnvironment(Parameters, OutEnvironment);
+		OutEnvironment.SetDefine(TEXT("THREADGROUP_SIZE_1D"), GetThreadGroupSize1D());
+		OutEnvironment.SetDefine(TEXT("THREADGROUP_SIZE_2D"), GetThreadGroupSize2D());
+		OutEnvironment.SetDefine(TEXT("THREADGROUP_SIZE_3D"), GetThreadGroupSize3D());
+
+		OutEnvironment.SetDefine(TEXT("DIM_USE_TRANSMITTANCE_VOLUME"), 0);
+		OutEnvironment.SetDefine(TEXT("DIM_USE_INSCATTERING_VOLUME"), 0);
+		OutEnvironment.SetDefine(TEXT("DIM_USE_LUMEN_GI"), 0);
+
+		OutEnvironment.CompilerFlags.Add(CFLAG_AllowTypedUAVLoads);
+	}
+
+	static int32 GetThreadGroupSize1D() { return GetThreadGroupSize3D() * GetThreadGroupSize3D() * GetThreadGroupSize3D(); }
+	static int32 GetThreadGroupSize2D() { return 8; }
+	static int32 GetThreadGroupSize3D() { return 4; }
+};
+
+IMPLEMENT_GLOBAL_SHADER(FConvertBeerLawShadowMapToVolumetricShadowMapCS, "/Engine/Private/HeterogeneousVolumes/HeterogeneousVolumesVoxelGridShadows.usf", "ConvertBeerLawShadowMapToVolumetricShadowMapCS", SF_Compute);
+
+void ConvertBeerLawShadowMapToVolumetricShadowMap(
+	FRDGBuilder& GraphBuilder,
+	FViewInfo& View,
+	//FIntVector GroupCount,
+	// Input
+	FIntPoint ShadowMapResolution,
+	FRDGTextureRef BeerShadowMapTexture,
+	// Output
+	FRDGBufferRef& VolumetricShadowIndirectionBuffer,
+	FRDGBufferRef& VolumetricShadowTransmittanceBuffer
+)
+{
+	FIntVector GroupCount = FIntVector(1);
+	GroupCount.X = FMath::DivideAndRoundUp(ShadowMapResolution.X, FConvertBeerLawShadowMapToVolumetricShadowMapCS::GetThreadGroupSize2D());
+	GroupCount.Y = FMath::DivideAndRoundUp(ShadowMapResolution.X, FConvertBeerLawShadowMapToVolumetricShadowMapCS::GetThreadGroupSize2D());
+
+	int32 ThreadGroupSize2D = FConvertBeerLawShadowMapToVolumetricShadowMapCS::GetThreadGroupSize2D();
+	int32 VolumetricShadowPixelCount = GroupCount.X * ThreadGroupSize2D * GroupCount.Y * ThreadGroupSize2D;
+	FRDGBufferRef VolumetricShadowIndirectionAllocatorBuffer = GraphBuilder.CreateBuffer(
+		FRDGBufferDesc::CreateBufferDesc(sizeof(uint32), 1),
+		TEXT("HeterogeneousVolume.VolumetricShadowIndirectionAllocatorBuffer")
+	);
+	AddClearUAVPass(GraphBuilder, GraphBuilder.CreateUAV(VolumetricShadowIndirectionAllocatorBuffer, PF_R32_UINT), 0);
+
+	VolumetricShadowIndirectionBuffer = GraphBuilder.CreateBuffer(
+		FRDGBufferDesc::CreateStructuredDesc(sizeof(FAVSMIndirectionPackedData), VolumetricShadowPixelCount),
+		TEXT("HeterogeneousVolume.VolumetricShadowIndirectionBuffer")
+	);
+
+	// Beer-Lambert Shadow Maps have two samples
+	const int32 MaxSampleCount = 2;
+	int32 TopLevelSampleCount = FMath::CeilToInt(MaxSampleCount / 4.0);
+	int32 VolumetricShadowTransmittanceMaxCount = VolumetricShadowPixelCount * (Align(TopLevelSampleCount, 4) + Align(MaxSampleCount, 4));
+	VolumetricShadowTransmittanceBuffer = GraphBuilder.CreateBuffer(
+		FRDGBufferDesc::CreateStructuredDesc(sizeof(FAVSMSamplePackedData), VolumetricShadowTransmittanceMaxCount),
+		TEXT("HeterogeneousVolume.VolumetricShadowTransmittanceBuffer")
+	);
+
+	FConvertBeerLawShadowMapToVolumetricShadowMapCS::FParameters* PassParameters = GraphBuilder.AllocParameters<FConvertBeerLawShadowMapToVolumetricShadowMapCS::FParameters>();
+	{
+		// Input
+		PassParameters->ShadowResolution = ShadowMapResolution;
+		PassParameters->BeerShadowMapTexture = GraphBuilder.CreateSRV(BeerShadowMapTexture);
+
+		// Output
+		PassParameters->RWVolumetricShadowIndirectionAllocatorBuffer = GraphBuilder.CreateUAV(VolumetricShadowIndirectionAllocatorBuffer, PF_R32_UINT);
+		PassParameters->RWVolumetricShadowIndirectionBuffer = GraphBuilder.CreateUAV(VolumetricShadowIndirectionBuffer);
+		PassParameters->RWVolumetricShadowTransmittanceBuffer = GraphBuilder.CreateUAV(VolumetricShadowTransmittanceBuffer);
+	}
+
+	FConvertBeerLawShadowMapToVolumetricShadowMapCS::FPermutationDomain PermutationVector;
+	TShaderRef<FConvertBeerLawShadowMapToVolumetricShadowMapCS> ComputeShader = View.ShaderMap->GetShader<FConvertBeerLawShadowMapToVolumetricShadowMapCS>(PermutationVector);
+	FComputeShaderUtils::AddPass(
+		GraphBuilder,
+		// Add Light name..
+		RDG_EVENT_NAME("FConvertBeerLawShadowMapToVolumetricShadowMapCS"),
+		ERDGPassFlags::Compute,
+		ComputeShader,
+		PassParameters,
+		GroupCount
+	);
+}
+
+void CreateAdaptiveVolumetricShadowMapUniformBufferParameters(
 	FRDGBuilder& GraphBuilder,
 	const FVector3f& TranslatedWorldOrigin,
 	const FVector4f& TranslatedWorldPlane,
 	const FMatrix44f* TranslatedWorldToShadow,
 	FIntPoint VolumetricShadowMapResolution,
+	float VolumetricShadowMapDownsampleFactor,
 	int32 NumShadowMatrices,
 	uint32 VolumetricShadowMapMaxSampleCount,
 	bool bIsDirectionalLight,
 	FRDGBufferRef VolumetricShadowMapLinkedListBuffer,
 	FRDGBufferRef VolumetricShadowMapIndirectionBuffer,
 	FRDGBufferRef VolumetricShadowMapSampleBuffer,
-	TRDGUniformBufferRef<FAdaptiveVolumetricShadowMapUniformBufferParameters>& AdaptiveVolumetricShadowMapUniformBuffer
+	FAdaptiveVolumetricShadowMapUniformBufferParameters*& Parameters
 )
 {
-	FAdaptiveVolumetricShadowMapUniformBufferParameters* UniformBufferParameters = GraphBuilder.AllocParameters<FAdaptiveVolumetricShadowMapUniformBufferParameters>();
+	Parameters = GraphBuilder.AllocParameters<FAdaptiveVolumetricShadowMapUniformBufferParameters>();
 	{
-		UniformBufferParameters->NumShadowMatrices = NumShadowMatrices;
+		Parameters->NumShadowMatrices = NumShadowMatrices;
 		for (int32 i = 0; i < NumShadowMatrices; ++i)
 		{
-			UniformBufferParameters->TranslatedWorldToShadow[i] = TranslatedWorldToShadow[i];
+			Parameters->TranslatedWorldToShadow[i] = TranslatedWorldToShadow[i];
 		}
-		UniformBufferParameters->TranslatedWorldOrigin = TranslatedWorldOrigin;
-		UniformBufferParameters->TranslatedWorldPlane = TranslatedWorldPlane;
-		UniformBufferParameters->Resolution = VolumetricShadowMapResolution;
-		UniformBufferParameters->MaxSampleCount = VolumetricShadowMapMaxSampleCount;
-		UniformBufferParameters->bIsEmpty = (VolumetricShadowMapResolution == FIntPoint::ZeroValue);
-		UniformBufferParameters->bIsDirectionalLight = bIsDirectionalLight;
-		UniformBufferParameters->LinkedListBuffer = GraphBuilder.CreateSRV(VolumetricShadowMapLinkedListBuffer);
-		UniformBufferParameters->IndirectionBuffer = GraphBuilder.CreateSRV(VolumetricShadowMapIndirectionBuffer);
-		UniformBufferParameters->SampleBuffer = GraphBuilder.CreateSRV(VolumetricShadowMapSampleBuffer);
+		Parameters->TranslatedWorldOrigin = TranslatedWorldOrigin;
+		Parameters->TranslatedWorldPlane = TranslatedWorldPlane;
+		Parameters->Resolution = VolumetricShadowMapResolution;
+		Parameters->DownsampleFactor = VolumetricShadowMapDownsampleFactor;
+		Parameters->MaxSampleCount = VolumetricShadowMapMaxSampleCount;
+		Parameters->bIsEmpty = (VolumetricShadowMapResolution == FIntPoint::ZeroValue);
+		Parameters->bIsDirectionalLight = bIsDirectionalLight;
+		Parameters->LinkedListBuffer = GraphBuilder.CreateSRV(VolumetricShadowMapLinkedListBuffer);
+		Parameters->IndirectionBuffer = GraphBuilder.CreateSRV(VolumetricShadowMapIndirectionBuffer);
+		Parameters->SampleBuffer = GraphBuilder.CreateSRV(VolumetricShadowMapSampleBuffer);
+		Parameters->RadianceTexture = GraphBuilder.CreateSRV(GSystemTextures.GetBlackDummy(GraphBuilder));
+		// Test
+		Parameters->TextureSampler = TStaticSamplerState<SF_Bilinear, AM_Clamp, AM_Clamp, AM_Clamp>::GetRHI();
 	}
+}
 
+void CreateAdaptiveVolumetricShadowMapUniformBuffer(
+	FRDGBuilder& GraphBuilder,
+	const FVector3f& TranslatedWorldOrigin,
+	const FVector4f& TranslatedWorldPlane,
+	const FMatrix44f* TranslatedWorldToShadow,
+	FIntPoint VolumetricShadowMapResolution,
+	float VolumetricShadowMapDownsampleFactor,
+	int32 NumShadowMatrices,
+	uint32 VolumetricShadowMapMaxSampleCount,
+	bool bIsDirectionalLight,
+	FRDGBufferRef VolumetricShadowMapLinkedListBuffer,
+	FRDGBufferRef VolumetricShadowMapIndirectionBuffer,
+	FRDGBufferRef VolumetricShadowMapSampleBuffer,
+	TRDGUniformBufferRef<FAdaptiveVolumetricShadowMapUniformBufferParameters>& AdaptiveVolumetricShadowMapUniformBuffer
+)
+{
+	FAdaptiveVolumetricShadowMapUniformBufferParameters* UniformBufferParameters;
+	CreateAdaptiveVolumetricShadowMapUniformBufferParameters(
+		GraphBuilder,
+		TranslatedWorldOrigin,
+		TranslatedWorldPlane,
+		TranslatedWorldToShadow,
+		VolumetricShadowMapResolution,
+		VolumetricShadowMapDownsampleFactor,
+		NumShadowMatrices,
+		VolumetricShadowMapMaxSampleCount,
+		bIsDirectionalLight,
+		VolumetricShadowMapLinkedListBuffer,
+		VolumetricShadowMapIndirectionBuffer,
+		VolumetricShadowMapSampleBuffer,
+		UniformBufferParameters
+	);
 	AdaptiveVolumetricShadowMapUniformBuffer = GraphBuilder.CreateUniformBuffer(UniformBufferParameters);
 }
 
-namespace HeterogeneousVolumes {
-
+namespace HeterogeneousVolumes
+{
 	TRDGUniformBufferRef<FAdaptiveVolumetricShadowMapUniformBufferParameters> CreateEmptyAdaptiveVolumetricShadowMapUniformBuffer(
 		FRDGBuilder& GraphBuilder
 	)
@@ -3243,12 +3393,15 @@ namespace HeterogeneousVolumes {
 			UniformBufferParameters->TranslatedWorldOrigin = FVector3f::ZeroVector;
 			UniformBufferParameters->TranslatedWorldPlane = FVector4f::Zero();
 			UniformBufferParameters->Resolution = FIntPoint::ZeroValue;
+			UniformBufferParameters->DownsampleFactor = 1.0f;
 			UniformBufferParameters->MaxSampleCount = 0;
 			UniformBufferParameters->bIsEmpty = true;
 			UniformBufferParameters->bIsDirectionalLight = false;
 			UniformBufferParameters->LinkedListBuffer = GraphBuilder.CreateSRV(GSystemTextures.GetDefaultStructuredBuffer(GraphBuilder, sizeof(FAVSMLinkedListPackedData)));
 			UniformBufferParameters->IndirectionBuffer = GraphBuilder.CreateSRV(GSystemTextures.GetDefaultStructuredBuffer(GraphBuilder, sizeof(FAVSMIndirectionPackedData)));
 			UniformBufferParameters->SampleBuffer = GraphBuilder.CreateSRV(GSystemTextures.GetDefaultStructuredBuffer(GraphBuilder, sizeof(FAVSMSamplePackedData)));
+			UniformBufferParameters->RadianceTexture = GraphBuilder.CreateSRV(GSystemTextures.GetBlackDummy(GraphBuilder));
+			UniformBufferParameters->TextureSampler = TStaticSamplerState<SF_Bilinear>::GetRHI();
 		}
 
 		TRDGUniformBufferRef<FAdaptiveVolumetricShadowMapUniformBufferParameters> AdaptiveVolumetricShadowMapUniformBuffer = GraphBuilder.CreateUniformBuffer(UniformBufferParameters);
@@ -3299,9 +3452,9 @@ namespace HeterogeneousVolumes {
 	{
 		FAdaptiveVolumetricShadowMapUniformBufferParameters Parameters;
 
-		if (ViewState && ViewState->AdaptiveVolumetricCameraMapUniformBuffer)
+		TRDGUniformBufferRef<FAdaptiveVolumetricShadowMapUniformBufferParameters> UniformBuffer = GetAdaptiveVolumetricCameraMapUniformBuffer(GraphBuilder, ViewState);
+		if (ViewState && UniformBuffer)
 		{
-			TRDGUniformBufferRef<FAdaptiveVolumetricShadowMapUniformBufferParameters> UniformBuffer = ViewState->AdaptiveVolumetricCameraMapUniformBuffer;
 
 			Parameters.NumShadowMatrices = UniformBuffer->GetParameters()->NumShadowMatrices;
 			for (int32 i = 0; i < Parameters.NumShadowMatrices; ++i)
@@ -3311,12 +3464,15 @@ namespace HeterogeneousVolumes {
 			Parameters.TranslatedWorldOrigin = UniformBuffer->GetParameters()->TranslatedWorldOrigin;
 			Parameters.TranslatedWorldPlane = UniformBuffer->GetParameters()->TranslatedWorldPlane;
 			Parameters.Resolution = UniformBuffer->GetParameters()->Resolution;
+			Parameters.DownsampleFactor = UniformBuffer->GetParameters()->DownsampleFactor;
 			Parameters.MaxSampleCount = UniformBuffer->GetParameters()->MaxSampleCount;
 			Parameters.bIsEmpty = UniformBuffer->GetParameters()->bIsEmpty;
 			Parameters.bIsDirectionalLight = UniformBuffer->GetParameters()->bIsDirectionalLight;
 			Parameters.LinkedListBuffer = UniformBuffer->GetParameters()->LinkedListBuffer;
 			Parameters.IndirectionBuffer = UniformBuffer->GetParameters()->IndirectionBuffer;
 			Parameters.SampleBuffer = UniformBuffer->GetParameters()->SampleBuffer;
+			Parameters.RadianceTexture = UniformBuffer->GetParameters()->RadianceTexture;
+			Parameters.TextureSampler = UniformBuffer->GetParameters()->TextureSampler;
 		}
 		else
 		{
@@ -3328,12 +3484,15 @@ namespace HeterogeneousVolumes {
 			Parameters.TranslatedWorldOrigin = FVector3f::ZeroVector;
 			Parameters.TranslatedWorldPlane = FVector4f::Zero();
 			Parameters.Resolution = FIntPoint::ZeroValue;
+			Parameters.DownsampleFactor = 1.0f;
 			Parameters.MaxSampleCount = 0;
 			Parameters.bIsEmpty = true;
 			Parameters.bIsDirectionalLight = false;
 			Parameters.LinkedListBuffer = GraphBuilder.CreateSRV(GSystemTextures.GetDefaultStructuredBuffer(GraphBuilder, sizeof(FAVSMLinkedListPackedData)));
 			Parameters.IndirectionBuffer = GraphBuilder.CreateSRV(GSystemTextures.GetDefaultStructuredBuffer(GraphBuilder, sizeof(FAVSMIndirectionPackedData)));
 			Parameters.SampleBuffer = GraphBuilder.CreateSRV(GSystemTextures.GetDefaultStructuredBuffer(GraphBuilder, sizeof(FAVSMSamplePackedData)));
+			Parameters.RadianceTexture = GraphBuilder.CreateSRV(GSystemTextures.GetBlackDummy(GraphBuilder));
+			Parameters.TextureSampler = TStaticSamplerState<SF_Bilinear>::GetRHI();
 		}
 
 		return Parameters;
@@ -3712,12 +3871,7 @@ void RenderVolumetricShadowMapForCameraWithVoxelGrid(
 	FRDGBufferRef& VolumetricShadowLinkedListBuffer
 )
 {
-	ShadowMapResolution = FIntPoint(View.ViewRect.Width(), View.ViewRect.Height());
-
-	float DownsampleFactor = HeterogeneousVolumes::GetCameraDownsampleFactor();
-	ShadowMapResolution.X = FMath::Max(ShadowMapResolution.X / DownsampleFactor, 1);
-	ShadowMapResolution.Y = FMath::Max(ShadowMapResolution.Y / DownsampleFactor, 1);
-	//ShadowMapResolution = HeterogeneousVolumes::GetShadowMapResolution();
+	ShadowMapResolution = HeterogeneousVolumes::GetDownsampledResolution(View.ViewRect.Size(), HeterogeneousVolumes::GetCameraDownsampleFactor());
 
 	GroupCount = FIntVector(1);
 	GroupCount.X = FMath::DivideAndRoundUp(ShadowMapResolution.X, FRenderVolumetricShadowMapForLightWithVoxelGridCS::GetThreadGroupSize2D());
@@ -3948,12 +4102,14 @@ void RenderAdaptiveVolumetricShadowMapWithVoxelGrid(
 					VolumetricShadowMapSampleBuffer
 				);
 
+				float DownsampleFactor = 1.0f;
 				CreateAdaptiveVolumetricShadowMapUniformBuffer(
 					GraphBuilder,
 					TranslatedWorldOrigin,
 					TranslatedWorldPlane,
 					TranslatedWorldToShadow,
 					VolumetricShadowMapResolution,
+					DownsampleFactor,
 					NumShadowMatrices,
 					VolumetricShadowMapMaxSampleCount,
 					bIsDirectionalLight,
@@ -4009,7 +4165,6 @@ void RenderAdaptiveVolumetricCameraMapWithVoxelGrid(
 	bool bShouldRenderCameraMap = !View.ViewRect.IsEmpty()
 		&& (OrthoGridUniformBuffer->GetParameters()->bUseOrthoGrid || FrustumGridUniformBuffer->GetParameters()->bUseFrustumGrid);
 
-	TRDGUniformBufferRef<FAdaptiveVolumetricShadowMapUniformBufferParameters> AdaptiveVolumetricShadowMapUniformBuffer;
 	if (bShouldRenderCameraMap)
 	{
 		FRDGTextureDesc Desc = SceneTextures.Color.Target->Desc;
@@ -4051,25 +4206,22 @@ void RenderAdaptiveVolumetricCameraMapWithVoxelGrid(
 
 		FVector4f TranslatedWorldPlane = FVector4f::Zero();
 		bool bIsDirectionalLight = false;
-		CreateAdaptiveVolumetricShadowMapUniformBuffer(
+		float DownsampleFactor = HeterogeneousVolumes::GetCameraDownsampleFactor();
+
+		CreateAdaptiveVolumetricShadowMapUniformBufferParameters(
 			GraphBuilder,
 			TranslatedWorldOrigin,
 			TranslatedWorldPlane,
 			TranslatedWorldToShadow,
 			VolumetricShadowMapResolution,
+			DownsampleFactor,
 			NumShadowMatrices,
 			VolumetricShadowMapMaxSampleCount,
 			bIsDirectionalLight,
 			VolumetricShadowMapLinkedListBuffer,
 			VolumetricShadowMapIndirectionBuffer,
 			VolumetricShadowMapSampleBuffer,
-			AdaptiveVolumetricShadowMapUniformBuffer
+			View.ViewState->AdaptiveVolumetricCameraMapUniformBufferParameters
 		);
 	}
-	else
-	{
-		AdaptiveVolumetricShadowMapUniformBuffer = HeterogeneousVolumes::CreateEmptyAdaptiveVolumetricShadowMapUniformBuffer(GraphBuilder);
-	}
-
-	View.ViewState->AdaptiveVolumetricCameraMapUniformBuffer = AdaptiveVolumetricShadowMapUniformBuffer;
 }
diff --git a/Engine/Source/Runtime/Renderer/Private/ScenePrivate.h b/Engine/Source/Runtime/Renderer/Private/ScenePrivate.h
--- a/Engine/Source/Runtime/Renderer/Private/ScenePrivate.h
+++ b/Engine/Source/Runtime/Renderer/Private/ScenePrivate.h
@@ -1015,6 +1015,7 @@ public:
 	TRDGUniformBufferRef<FOrthoVoxelGridUniformBufferParameters> OrthoVoxelGridUniformBuffer = nullptr;
 	TRDGUniformBufferRef<FFrustumVoxelGridUniformBufferParameters> FrustumVoxelGridUniformBuffer = nullptr;
 
+	FAdaptiveVolumetricShadowMapUniformBufferParameters* AdaptiveVolumetricCameraMapUniformBufferParameters = nullptr;
 	TRDGUniformBufferRef<FAdaptiveVolumetricShadowMapUniformBufferParameters> AdaptiveVolumetricCameraMapUniformBuffer = nullptr;
 	TMap<int32, TRDGUniformBufferRef<FAdaptiveVolumetricShadowMapUniformBufferParameters>> AdaptiveVolumetricShadowMapUniformBufferMap;
 
diff --git a/Engine/Source/Runtime/Renderer/Private/SceneRendering.h b/Engine/Source/Runtime/Renderer/Private/SceneRendering.h
--- a/Engine/Source/Runtime/Renderer/Private/SceneRendering.h
+++ b/Engine/Source/Runtime/Renderer/Private/SceneRendering.h
@@ -1450,6 +1450,7 @@ public:
 
 	FRDGTextureRef HeterogeneousVolumeRadiance = nullptr;
 	FRDGTextureRef HeterogeneousVolumeHoldout = nullptr;
+	FRDGTextureRef HeterogeneousVolumeBeerShadowMap = nullptr;
 
 	// Size of the HZB's mipmap 0
 	// NOTE: the mipmap 0 is downsampled version of the depth buffer
