Fixed local fog volume debug view when volumetric fog is enabled.

#tests PC editor
#rnx

diff --git a/Engine/Shaders/Private/LocalFogVolumes/LocalFogVolumeSplat.usf b/Engine/Shaders/Private/LocalFogVolumes/LocalFogVolumeSplat.usf
--- a/Engine/Shaders/Private/LocalFogVolumes/LocalFogVolumeSplat.usf
+++ b/Engine/Shaders/Private/LocalFogVolumes/LocalFogVolumeSplat.usf
@@ -18,6 +18,9 @@
 #ifndef LFV_HALFRES_PS
 #define LFV_HALFRES_PS 0
 #endif
+#ifndef LFV_VISUALISATION_PS
+#define LFV_VISUALISATION_PS 0
+#endif
 
 
 
@@ -61,7 +64,7 @@ void LocalFogVolumeTiledVS(
 #if LFV_TILED_PS
 int LocalFogVolumeTileDebug;
 
-#if USE_DEVELOPMENT_SHADERS
+#if LFV_VISUALISATION_PS
 #include "../ColorMap.ush"
 #endif
 
@@ -108,7 +111,7 @@ void LocalFogVolumeTiledPS(
 	float4 LFVContribution = GetLFVContribution(ResolvedView, TilePos, DepthBufferTranslatedWorldPos, bPlatformSupportsVolumetricFogOntranslucent, LFVCount);
 
 	bool bDoClip = true;
-#if USE_DEVELOPMENT_SHADERS
+#if LFV_VISUALISATION_PS
 	bDoClip = LocalFogVolumeTileDebug != 1;
 #endif
 #if LFV_HALFRES_PS
@@ -121,6 +124,19 @@ void LocalFogVolumeTiledPS(
 		clip(-1.0f);
 	}
 
+#if LFV_VISUALISATION_PS == 1
+
+	if (LocalFogVolumeTileDebug > 0)
+	{
+		float Coverage = 0.5;
+		float3 DebugColor = GetHSVDebugColor(float(LFVCount) / 4.0f);
+		DebugColor = DebugColor * Coverage; // * PrimaryView.OneOverPreExposure * PrimaryView.PreExposure == 1
+
+		OutColor = float4(DebugColor, 1.0 - Coverage);
+	}
+
+#else // LFV_VISUALISATION_PS
+
 #if LFV_HALFRES_PS
 	OutColorThroughput = float4(
 		LFVContribution.rgb * PrimaryView.PreExposure,
@@ -132,16 +148,8 @@ void LocalFogVolumeTiledPS(
 		LFVContribution.a);
 #endif
 
-#if USE_DEVELOPMENT_SHADERS && !LFV_HALFRES_PS
-	if (LocalFogVolumeTileDebug > 0)
-	{
-		float DebugTransmittance = 0.5;
-		float3 DebugColor = GetHSVDebugColor(float(LFVCount)  / 4.0f);
-		DebugColor = DebugColor * (1.0f - DebugTransmittance) * PrimaryView.PreExposure;
+#endif // LFV_VISUALISATION_PS
 
-		OutColor = float4(OutColor.rgb * DebugTransmittance + DebugColor, OutColor.a * DebugTransmittance);
-	}
-#endif
 }
 #endif // LFV_TILED_PS
 
diff --git a/Engine/Source/Runtime/Renderer/Private/DeferredShadingRenderer.cpp b/Engine/Source/Runtime/Renderer/Private/DeferredShadingRenderer.cpp
--- a/Engine/Source/Runtime/Renderer/Private/DeferredShadingRenderer.cpp
+++ b/Engine/Source/Runtime/Renderer/Private/DeferredShadingRenderer.cpp
@@ -1894,6 +1894,7 @@ void FDeferredShadingSceneRenderer::Render(FRDGBuilder& GraphBuilder, const FSce
 	const bool bShouldRenderLocalFogVolume = ShouldRenderLocalFogVolume(Scene, ViewFamily);
 	const bool bShouldRenderLocalFogVolumeDuringHeightFogPass = ShouldRenderLocalFogVolumeDuringHeightFogPass(Scene, ViewFamily);
 	const bool bShouldRenderLocalFogVolumeInVolumetricFog = ShouldRenderLocalFogVolumeInVolumetricFog(Scene, ViewFamily, bShouldRenderLocalFogVolume);
+	const bool bShouldRenderLocalFogVolumeVisualisationPass = ShouldRenderLocalFogVolumeVisualisationPass(Scene, ViewFamily);
 
 	const bool bRenderDeferredLighting = ViewFamily.EngineShowFlags.Lighting
 		&& FeatureLevel >= ERHIFeatureLevel::SM5
@@ -3109,15 +3110,23 @@ void FDeferredShadingSceneRenderer::Render(FRDGBuilder& GraphBuilder, const FSce
 				bHeightFogHasComposedLocalFogVolume = bFogComposeLocalFogVolumes;
 			}
 
-			// Local Fog Volumes (LFV) rendering order is first HeightFog, then LFV, then volumetric fog on top.
-			// LFVs are rendered as part of the regular height fog + volumetric fog pass when volumetric fog is enabled and it is requested to voxelise LFVs into volumetric fog.
-			// Otherwise, they are rendered in an independent pass (this for instance make it independent of the near clip plane optimization).
-			if (!bHasRayTracedOverlay && !bHeightFogHasComposedLocalFogVolume)
+			if (!bHasRayTracedOverlay) 
 			{
-				RDG_CSV_STAT_EXCLUSIVE_SCOPE(GraphBuilder, RenderLocalFogVolume);
-				SCOPED_NAMED_EVENT(RenderLocalFogVolume, FColor::Emerald);
-				SCOPE_CYCLE_COUNTER(STAT_FDeferredShadingSceneRenderer_RenderLocalFogVolume);
-				RenderLocalFogVolume(Scene, Views, ViewFamily, GraphBuilder, SceneTextures, LightShaftOcclusionTexture);
+				// Local Fog Volumes (LFV) rendering order is first HeightFog, then LFV, then volumetric fog on top.
+				// LFVs are rendered as part of the regular height fog + volumetric fog pass when volumetric fog is enabled and it is requested to voxelise LFVs into volumetric fog.
+				// Otherwise, they are rendered in an independent pass (this for instance make it independent of the near clip plane optimization).
+				if (!bHeightFogHasComposedLocalFogVolume)
+				{
+					RDG_CSV_STAT_EXCLUSIVE_SCOPE(GraphBuilder, RenderLocalFogVolume);
+					SCOPED_NAMED_EVENT(RenderLocalFogVolume, FColor::Emerald);
+					SCOPE_CYCLE_COUNTER(STAT_FDeferredShadingSceneRenderer_RenderLocalFogVolume);
+					RenderLocalFogVolume(Scene, Views, ViewFamily, GraphBuilder, SceneTextures, LightShaftOcclusionTexture);
+				}
+				// Also compose on top the visualisation pass if enabled.
+				if (bShouldRenderLocalFogVolumeVisualisationPass)
+				{
+					RenderLocalFogVolumeVisualisation(Scene, Views, ViewFamily, GraphBuilder, SceneTextures);
+				}
 			}
 
 			// After the height fog, Draw volumetric clouds (having fog applied on them already) when using per pixel tracing,
diff --git a/Engine/Source/Runtime/Renderer/Private/LocalFogVolumeRendering.cpp b/Engine/Source/Runtime/Renderer/Private/LocalFogVolumeRendering.cpp
--- a/Engine/Source/Runtime/Renderer/Private/LocalFogVolumeRendering.cpp
+++ b/Engine/Source/Runtime/Renderer/Private/LocalFogVolumeRendering.cpp
@@ -135,6 +135,15 @@ bool ShouldRenderLocalFogVolumeInVolumetricFog(const FScene* Scene, const FScene
 	return false;
 }
 
+bool ShouldRenderLocalFogVolumeVisualisationPass(const FScene* Scene, const FSceneViewFamily& SceneViewFamily)
+{
+	if (ShouldRenderLocalFogVolume(Scene, SceneViewFamily))
+	{
+		return CVarLocalFogVolumeTileDebug.GetValueOnRenderThread() > 0;
+	}
+	return false;
+}
+
 float GetLocalFogVolumeGlobalStartDistance()
 {
 	return FMath::Max(10.0f, CVarLocalFogVolumeGlobalStartDistance.GetValueOnRenderThread());
@@ -695,7 +704,8 @@ class FLocalFogVolumeTiledRenderPS : public FGlobalShader
 		SHADER_PARAMETER(int32, LocalFogVolumeTileDebug)
 	END_SHADER_PARAMETER_STRUCT()
 
-	using FPermutationDomain = TShaderPermutationDomain<>;
+	class FVisualisationPass : SHADER_PERMUTATION_BOOL("LFV_VISUALISATION_PS");
+	using FPermutationDomain = TShaderPermutationDomain<FVisualisationPass>;
 
 	static bool ShouldCompilePermutation(const FGlobalShaderPermutationParameters& Parameters)
 	{
@@ -772,12 +782,13 @@ void RenderLocalFogVolume(
 	const FSceneViewFamily& SceneViewFamily,
 	FRDGBuilder& GraphBuilder,
 	const FMinimalSceneTextures& SceneTextures,
-	FRDGTextureRef LightShaftOcclusionTexture)
+	FRDGTextureRef LightShaftOcclusionTexture,
+	bool bVisualisationPass)
 {
 	uint32 LocalFogVolumeInstanceCount = Scene->LocalFogVolumes.Num();
 	if (LocalFogVolumeInstanceCount > 0 && ShouldRenderLocalFogVolume(Scene, SceneViewFamily))
 	{
-		RDG_EVENT_SCOPE_STAT(GraphBuilder, LocalFogVolumeVolumes, "LocalFogVolumeVolumes");
+		RDG_EVENT_SCOPE_STAT(GraphBuilder, LocalFogVolumeVolumes, "%s", bVisualisationPass ? TEXT("LocalFogVolume.Visualisation") : TEXT("LocalFogVolume.Render"));
 		RDG_GPU_STAT_SCOPE(GraphBuilder, LocalFogVolumeVolumes);
 
 		FRDGTextureRef SceneColorTexture = SceneTextures.Color.Resolve;
@@ -803,7 +814,7 @@ void RenderLocalFogVolume(
 
 			PassParameters->PS.View = GetShaderBinding(View.ViewUniformBuffer);
 			PassParameters->PS.LFV = View.LocalFogVolumeViewData.UniformParametersStruct;
-			PassParameters->PS.LocalFogVolumeTileDebug = FMath::Clamp(CVarLocalFogVolumeTileDebug.GetValueOnRenderThread(), 0, 2);
+			PassParameters->PS.LocalFogVolumeTileDebug = FMath::Clamp(CVarLocalFogVolumeTileDebug.GetValueOnRenderThread(), 0, 2); 
 
 			PassParameters->SceneTextures = SceneTextures.UniformBuffer;
 			PassParameters->TileDrawIndirectBuffer = View.LocalFogVolumeViewData.GPUTileDrawIndirectBuffer;
@@ -814,6 +825,7 @@ void RenderLocalFogVolume(
 			auto VertexShader = View.ShaderMap->GetShader< FLocalFogVolumeTiledRenderVS >(VSPermutationVector);
 
 			FLocalFogVolumeTiledRenderPS::FPermutationDomain PsPermutationVector;
+			PsPermutationVector.Set<FLocalFogVolumeTiledRenderPS::FVisualisationPass>(bVisualisationPass);
 			auto PixelShader = View.ShaderMap->GetShader< FLocalFogVolumeTiledRenderPS >(PsPermutationVector);
 
 			ClearUnusedGraphResources(VertexShader, &PassParameters->VS);
@@ -822,7 +834,7 @@ void RenderLocalFogVolume(
 			FUintVector2& LocalFogVolumeTileDataTextureResolution = View.LocalFogVolumeViewData.UniformParametersStruct.LocalFogVolumeCommon.LocalFogVolumeTileDataTextureResolution;
 
 			GraphBuilder.AddPass(
-				RDG_EVENT_NAME("LocalFogVolume.Tiled (%u X %u)", LocalFogVolumeTileDataTextureResolution.X, LocalFogVolumeTileDataTextureResolution.Y),
+				RDG_EVENT_NAME("LocalFogVolume.Tiled (%u x %u)", LocalFogVolumeTileDataTextureResolution.X, LocalFogVolumeTileDataTextureResolution.Y),
 				PassParameters,
 				ERDGPassFlags::Raster,
 				[VertexShader, PixelShader, PassParameters, ViewRect, DepthBoundSetup](FRDGAsyncTask, FRHICommandList& RHICmdList)
@@ -857,6 +869,20 @@ void RenderLocalFogVolume(
 		}
 	}
 }
+ 
+void RenderLocalFogVolumeVisualisation(
+	const FScene* Scene,
+	TArray<FViewInfo>& Views,
+	const FSceneViewFamily& SceneViewFamily,
+	FRDGBuilder& GraphBuilder,
+	const FMinimalSceneTextures& SceneTextures)
+{
+#if !UE_BUILD_SHIPPING
+	const bool bVisualisationPass = true;
+	FRDGTexture* LightShaftOcclusionTexture = GSystemTextures.GetBlackDummy(GraphBuilder);
+	RenderLocalFogVolume(Scene, Views, SceneViewFamily, GraphBuilder, SceneTextures, LightShaftOcclusionTexture, bVisualisationPass);
+#endif
+}
 
 /*=============================================================================
 	Local height fog rendering - mobile
diff --git a/Engine/Source/Runtime/Renderer/Private/LocalFogVolumeRendering.h b/Engine/Source/Runtime/Renderer/Private/LocalFogVolumeRendering.h
--- a/Engine/Source/Runtime/Renderer/Private/LocalFogVolumeRendering.h
+++ b/Engine/Source/Runtime/Renderer/Private/LocalFogVolumeRendering.h
@@ -132,6 +132,7 @@ bool ProjectSupportsLocalFogVolumes();
 bool ShouldRenderLocalFogVolume(const FScene* Scene, const FSceneViewFamily& SceneViewFamily);
 bool ShouldRenderLocalFogVolumeDuringHeightFogPass(const FScene* Scene, const FSceneViewFamily& SceneViewFamily);
 bool ShouldRenderLocalFogVolumeInVolumetricFog(const FScene* Scene, const FSceneViewFamily& SceneViewFamily, bool bShouldRenderVolumetricFog);
+bool ShouldRenderLocalFogVolumeVisualisationPass(const FScene* Scene, const FSceneViewFamily& SceneViewFamily);
 float GetLocalFogVolumeGlobalStartDistance();
 bool IsLocalFogVolumeHalfResolution();
 
@@ -155,7 +156,8 @@ void RenderLocalFogVolume(
 	const FSceneViewFamily& SceneViewFamily,
 	FRDGBuilder& GraphBuilder,
 	const FMinimalSceneTextures& SceneTextures,
-	FRDGTextureRef LightShaftOcclusionTexture);
+	FRDGTextureRef LightShaftOcclusionTexture,
+	bool bVisualisationPass = false);
 
 void RenderLocalFogVolumeMobile(
 	FRHICommandList& RHICmdList, 
@@ -164,3 +166,10 @@ void RenderLocalFogVolumeHalfResMobile(
 	FRDGBuilder& GraphBuilder,
 	const FViewInfo& View);
 
+void RenderLocalFogVolumeVisualisation(
+	const FScene* Scene,
+	TArray<FViewInfo>& Views,
+	const FSceneViewFamily& SceneViewFamily,
+	FRDGBuilder& GraphBuilder,
+	const FMinimalSceneTextures& SceneTextures);
+
