Fix crash on Groom voxelization due to incorrect array element access.

diff --git a/Engine/Source/Runtime/Renderer/Public/HairStrandsInterface.h b/Engine/Source/Runtime/Renderer/Public/HairStrandsInterface.h
--- a/Engine/Source/Runtime/Renderer/Public/HairStrandsInterface.h
+++ b/Engine/Source/Runtime/Renderer/Public/HairStrandsInterface.h
@@ -421,8 +421,8 @@ public:
 // Resources are indexed by hair instance's registered index
 struct FHairTransientResources
 {
-	FORCEINLINE uint32 GetClusterOffset(uint32 InRegisterIndex) const { check(ClusterAABBOffetAndCounts.IsValidIndex(InRegisterIndex)); return ClusterAABBOffetAndCounts[InRegisterIndex].X; }
-	FORCEINLINE uint32 GetClusterCount(uint32 InRegisterIndex)  const { check(ClusterAABBOffetAndCounts.IsValidIndex(InRegisterIndex)); return ClusterAABBOffetAndCounts[InRegisterIndex].Y; }
+	FORCEINLINE uint32 GetClusterOffset(uint32 InRegisterIndex) const { return ClusterAABBOffetAndCounts.IsValidIndex(InRegisterIndex) ? ClusterAABBOffetAndCounts[InRegisterIndex].X : 0u; }
+	FORCEINLINE uint32 GetClusterCount(uint32 InRegisterIndex)  const { return ClusterAABBOffetAndCounts.IsValidIndex(InRegisterIndex) ? ClusterAABBOffetAndCounts[InRegisterIndex].Y : 0u; }
 	FORCEINLINE bool IsClusterAABBValid(uint32 InRegisterIndex) const { return GetClusterCount(InRegisterIndex) > 0; }
 
 	// Group AABB
