LocalFogVolumes - added visualization control on menu.
Fixed naming.

#rnx
#tests pc editor

diff --git a/Engine/Shaders/Private/LocalFogVolumes/LocalFogVolumeSplat.usf b/Engine/Shaders/Private/LocalFogVolumes/LocalFogVolumeSplat.usf
--- a/Engine/Shaders/Private/LocalFogVolumes/LocalFogVolumeSplat.usf
+++ b/Engine/Shaders/Private/LocalFogVolumes/LocalFogVolumeSplat.usf
@@ -18,8 +18,8 @@
 #ifndef LFV_HALFRES_PS
 #define LFV_HALFRES_PS 0
 #endif
-#ifndef LFV_VISUALISATION_PS
-#define LFV_VISUALISATION_PS 0
+#ifndef LFV_VISUALIZATION_PS
+#define LFV_VISUALIZATION_PS 0
 #endif
 
 
@@ -64,7 +64,7 @@ void LocalFogVolumeTiledVS(
 #if LFV_TILED_PS
 int LocalFogVolumeTileDebug;
 
-#if LFV_VISUALISATION_PS
+#if LFV_VISUALIZATION_PS
 #include "../ColorMap.ush"
 #endif
 
@@ -111,7 +111,7 @@ void LocalFogVolumeTiledPS(
 	float4 LFVContribution = GetLFVContribution(ResolvedView, TilePos, DepthBufferTranslatedWorldPos, bPlatformSupportsVolumetricFogOntranslucent, LFVCount);
 
 	bool bDoClip = true;
-#if LFV_VISUALISATION_PS
+#if LFV_VISUALIZATION_PS
 	bDoClip = LocalFogVolumeTileDebug != 1;
 #endif
 #if LFV_HALFRES_PS
@@ -124,7 +124,7 @@ void LocalFogVolumeTiledPS(
 		clip(-1.0f);
 	}
 
-#if LFV_VISUALISATION_PS == 1
+#if LFV_VISUALIZATION_PS == 1
 
 	if (LocalFogVolumeTileDebug > 0)
 	{
@@ -135,7 +135,7 @@ void LocalFogVolumeTiledPS(
 		OutColor = float4(DebugColor, 1.0 - Coverage);
 	}
 
-#else // LFV_VISUALISATION_PS
+#else // LFV_VISUALIZATION_PS
 
 #if LFV_HALFRES_PS
 	OutColorThroughput = float4(
@@ -148,7 +148,7 @@ void LocalFogVolumeTiledPS(
 		LFVContribution.a);
 #endif
 
-#endif // LFV_VISUALISATION_PS
+#endif // LFV_VISUALIZATION_PS
 
 }
 #endif // LFV_TILED_PS
diff --git a/Engine/Source/Runtime/Engine/Public/ShowFlags.h b/Engine/Source/Runtime/Engine/Public/ShowFlags.h
--- a/Engine/Source/Runtime/Engine/Public/ShowFlags.h
+++ b/Engine/Source/Runtime/Engine/Public/ShowFlags.h
@@ -396,6 +396,7 @@ private:
 		// The following code sets what should be off by default.
 		SetVisualizeHDR(false);
 		SetVisualizeSkyLightIlluminance(false);
+		SetVisualizeLocalFogVolumes(false);
 		SetVisualizeLocalExposure(false);
 		SetVisualizeShadingModels(false);
 		SetOverrideDiffuseAndSpecular(false);
diff --git a/Engine/Source/Runtime/Engine/Public/ShowFlagsValues.inl b/Engine/Source/Runtime/Engine/Public/ShowFlagsValues.inl
--- a/Engine/Source/Runtime/Engine/Public/ShowFlagsValues.inl
+++ b/Engine/Source/Runtime/Engine/Public/ShowFlagsValues.inl
@@ -37,7 +37,9 @@ SHOWFLAG_ALWAYS_ACCESSIBLE(EyeAdaptation, SFG_PostProcess, NSLOCTEXT("UnrealEd",
 /** Display a histogram of the scene HDR color */
 SHOWFLAG_FIXED_IN_SHIPPING(0, VisualizeHDR, SFG_Visualize, NSLOCTEXT("UnrealEd", "VisualizeHDRSF", "HDR (Eye Adaptation)"))
 /** Display the illuminance debug view for the skylight */
-SHOWFLAG_FIXED_IN_SHIPPING(0, VisualizeSkyLightIlluminance, SFG_Visualize, NSLOCTEXT("UnrealEd", "VisualizeSkyLightIlluminance", "SkyLight Illuminance"))
+SHOWFLAG_FIXED_IN_SHIPPING(0, VisualizeSkyLightIlluminance, SFG_Visualize, NSLOCTEXT("UnrealEd", "VisualizeSkyLightIlluminance", "Sky Light Illuminance"))
+/** Display the illuminance debug view for the skylight */
+SHOWFLAG_FIXED_IN_SHIPPING(0, VisualizeLocalFogVolumes, SFG_Visualize, NSLOCTEXT("UnrealEd", "VisualizeLocalFogVolumes", "Local Fog Volumes"))
 /** Helper to tweak local expsoure settings */
 SHOWFLAG_FIXED_IN_SHIPPING(0, VisualizeLocalExposure, SFG_Visualize, NSLOCTEXT("UnrealEd", "VisualizeLocalExposureSF", "Local Exposure"))
 /** Image based lens flares (Simulate artifact of reflections within a camera system) */
diff --git a/Engine/Source/Runtime/Renderer/Private/DeferredShadingRenderer.cpp b/Engine/Source/Runtime/Renderer/Private/DeferredShadingRenderer.cpp
--- a/Engine/Source/Runtime/Renderer/Private/DeferredShadingRenderer.cpp
+++ b/Engine/Source/Runtime/Renderer/Private/DeferredShadingRenderer.cpp
@@ -1894,7 +1894,7 @@ void FDeferredShadingSceneRenderer::Render(FRDGBuilder& GraphBuilder, const FSce
 	const bool bShouldRenderLocalFogVolume = ShouldRenderLocalFogVolume(Scene, ViewFamily);
 	const bool bShouldRenderLocalFogVolumeDuringHeightFogPass = ShouldRenderLocalFogVolumeDuringHeightFogPass(Scene, ViewFamily);
 	const bool bShouldRenderLocalFogVolumeInVolumetricFog = ShouldRenderLocalFogVolumeInVolumetricFog(Scene, ViewFamily, bShouldRenderLocalFogVolume);
-	const bool bShouldRenderLocalFogVolumeVisualisationPass = ShouldRenderLocalFogVolumeVisualisationPass(Scene, ViewFamily);
+	const bool bShouldRenderLocalFogVolumeVisualizationPass = ShouldRenderLocalFogVolumeVisualizationPass(Scene, ViewFamily);
 
 	const bool bRenderDeferredLighting = ViewFamily.EngineShowFlags.Lighting
 		&& FeatureLevel >= ERHIFeatureLevel::SM5
@@ -3123,9 +3123,9 @@ void FDeferredShadingSceneRenderer::Render(FRDGBuilder& GraphBuilder, const FSce
 					RenderLocalFogVolume(Scene, Views, ViewFamily, GraphBuilder, SceneTextures, LightShaftOcclusionTexture);
 				}
 				// Also compose on top the visualisation pass if enabled.
-				if (bShouldRenderLocalFogVolumeVisualisationPass)
+				if (bShouldRenderLocalFogVolumeVisualizationPass)
 				{
-					RenderLocalFogVolumeVisualisation(Scene, Views, ViewFamily, GraphBuilder, SceneTextures);
+					RenderLocalFogVolumeVisualization(Scene, Views, ViewFamily, GraphBuilder, SceneTextures);
 				}
 			}
 
diff --git a/Engine/Source/Runtime/Renderer/Private/LocalFogVolumeRendering.cpp b/Engine/Source/Runtime/Renderer/Private/LocalFogVolumeRendering.cpp
--- a/Engine/Source/Runtime/Renderer/Private/LocalFogVolumeRendering.cpp
+++ b/Engine/Source/Runtime/Renderer/Private/LocalFogVolumeRendering.cpp
@@ -102,6 +102,11 @@ static uint32 GetLocalFogVolumeTileMaxInstanceCount()
 	return FMath::Max(1u, FMath::Min(256u, (uint32)CVarLocalFogVolumeTileMaxInstanceCount.GetValueOnRenderThread()));
 }
 
+static int32 GetLocalFogVolumeTileDebugMode(const FEngineShowFlags& EngineShowFlags)
+{
+	return EngineShowFlags.VisualizeLocalFogVolumes ? 2 : FMath::Clamp(CVarLocalFogVolumeTileDebug.GetValueOnRenderThread(), 0, 2);
+}
+
 bool ProjectSupportsLocalFogVolumes()
 {
 	return CVarSupportLocalFogVolumes.GetValueOnRenderThread() > 0;
@@ -109,7 +114,7 @@ bool ProjectSupportsLocalFogVolumes()
 
 bool ShouldRenderLocalFogVolume(const FScene* Scene, const FSceneViewFamily& SceneViewFamily)
 {
-	const FEngineShowFlags EngineShowFlags = SceneViewFamily.EngineShowFlags;
+	const FEngineShowFlags& EngineShowFlags = SceneViewFamily.EngineShowFlags;
 	if (Scene && Scene->HasAnyLocalFogVolume() && EngineShowFlags.Fog && !SceneViewFamily.UseDebugViewPS())
 	{
 		return ProjectSupportsLocalFogVolumes() && (CVarLocalFogVolume.GetValueOnRenderThread() > 0);
@@ -135,11 +140,12 @@ bool ShouldRenderLocalFogVolumeInVolumetricFog(const FScene* Scene, const FScene
 	return false;
 }
 
-bool ShouldRenderLocalFogVolumeVisualisationPass(const FScene* Scene, const FSceneViewFamily& SceneViewFamily)
+bool ShouldRenderLocalFogVolumeVisualizationPass(const FScene* Scene, const FSceneViewFamily& SceneViewFamily)
 {
 	if (ShouldRenderLocalFogVolume(Scene, SceneViewFamily))
 	{
-		return CVarLocalFogVolumeTileDebug.GetValueOnRenderThread() > 0;
+		const FEngineShowFlags& EngineShowFlags = SceneViewFamily.EngineShowFlags;
+		return GetLocalFogVolumeTileDebugMode(EngineShowFlags) > 0;
 	}
 	return false;
 }
@@ -704,8 +710,8 @@ class FLocalFogVolumeTiledRenderPS : public FGlobalShader
 		SHADER_PARAMETER(int32, LocalFogVolumeTileDebug)
 	END_SHADER_PARAMETER_STRUCT()
 
-	class FVisualisationPass : SHADER_PERMUTATION_BOOL("LFV_VISUALISATION_PS");
-	using FPermutationDomain = TShaderPermutationDomain<FVisualisationPass>;
+	class FVisualizationPass : SHADER_PERMUTATION_BOOL("LFV_VISUALIZATION_PS");
+	using FPermutationDomain = TShaderPermutationDomain<FVisualizationPass>;
 
 	static bool ShouldCompilePermutation(const FGlobalShaderPermutationParameters& Parameters)
 	{
@@ -783,12 +789,12 @@ void RenderLocalFogVolume(
 	FRDGBuilder& GraphBuilder,
 	const FMinimalSceneTextures& SceneTextures,
 	FRDGTextureRef LightShaftOcclusionTexture,
-	bool bVisualisationPass)
+	bool bVisualizationPass)
 {
 	uint32 LocalFogVolumeInstanceCount = Scene->LocalFogVolumes.Num();
 	if (LocalFogVolumeInstanceCount > 0 && ShouldRenderLocalFogVolume(Scene, SceneViewFamily))
 	{
-		RDG_EVENT_SCOPE_STAT(GraphBuilder, LocalFogVolumeVolumes, "%s", bVisualisationPass ? TEXT("LocalFogVolume.Visualisation") : TEXT("LocalFogVolume.Render"));
+		RDG_EVENT_SCOPE_STAT(GraphBuilder, LocalFogVolumeVolumes, "%s", bVisualizationPass ? TEXT("LocalFogVolume.Visualization") : TEXT("LocalFogVolume.Render"));
 		RDG_GPU_STAT_SCOPE(GraphBuilder, LocalFogVolumeVolumes);
 
 		FRDGTextureRef SceneColorTexture = SceneTextures.Color.Resolve;
@@ -814,7 +820,7 @@ void RenderLocalFogVolume(
 
 			PassParameters->PS.View = GetShaderBinding(View.ViewUniformBuffer);
 			PassParameters->PS.LFV = View.LocalFogVolumeViewData.UniformParametersStruct;
-			PassParameters->PS.LocalFogVolumeTileDebug = FMath::Clamp(CVarLocalFogVolumeTileDebug.GetValueOnRenderThread(), 0, 2); 
+			PassParameters->PS.LocalFogVolumeTileDebug = GetLocalFogVolumeTileDebugMode(SceneViewFamily.EngineShowFlags);
 
 			PassParameters->SceneTextures = SceneTextures.UniformBuffer;
 			PassParameters->TileDrawIndirectBuffer = View.LocalFogVolumeViewData.GPUTileDrawIndirectBuffer;
@@ -825,7 +831,7 @@ void RenderLocalFogVolume(
 			auto VertexShader = View.ShaderMap->GetShader< FLocalFogVolumeTiledRenderVS >(VSPermutationVector);
 
 			FLocalFogVolumeTiledRenderPS::FPermutationDomain PsPermutationVector;
-			PsPermutationVector.Set<FLocalFogVolumeTiledRenderPS::FVisualisationPass>(bVisualisationPass);
+			PsPermutationVector.Set<FLocalFogVolumeTiledRenderPS::FVisualizationPass>(bVisualizationPass);
 			auto PixelShader = View.ShaderMap->GetShader< FLocalFogVolumeTiledRenderPS >(PsPermutationVector);
 
 			ClearUnusedGraphResources(VertexShader, &PassParameters->VS);
@@ -870,7 +876,7 @@ void RenderLocalFogVolume(
 	}
 }
  
-void RenderLocalFogVolumeVisualisation(
+void RenderLocalFogVolumeVisualization(
 	const FScene* Scene,
 	TArray<FViewInfo>& Views,
 	const FSceneViewFamily& SceneViewFamily,
@@ -878,9 +884,9 @@ void RenderLocalFogVolumeVisualisation(
 	const FMinimalSceneTextures& SceneTextures)
 {
 #if !UE_BUILD_SHIPPING
-	const bool bVisualisationPass = true;
+	const bool bVisualizationPass = true;
 	FRDGTexture* LightShaftOcclusionTexture = GSystemTextures.GetBlackDummy(GraphBuilder);
-	RenderLocalFogVolume(Scene, Views, SceneViewFamily, GraphBuilder, SceneTextures, LightShaftOcclusionTexture, bVisualisationPass);
+	RenderLocalFogVolume(Scene, Views, SceneViewFamily, GraphBuilder, SceneTextures, LightShaftOcclusionTexture, bVisualizationPass);
 #endif
 }
 
@@ -1053,7 +1059,7 @@ void RenderLocalFogVolumeMobile(
 	FMobileLocalFogVolumeTiledRenderPS::FParameters PSParameters;
 	PSParameters.View = View.GetShaderParameters();
 	PSParameters.LFV = View.LocalFogVolumeViewData.UniformParametersStruct;
-	PSParameters.LocalFogVolumeTileDebug = FMath::Clamp(CVarLocalFogVolumeTileDebug.GetValueOnRenderThread(), 0, 2);
+	PSParameters.LocalFogVolumeTileDebug = GetLocalFogVolumeTileDebugMode(View.Family->EngineShowFlags);
 	// PSParameters.MobileBasePass filled up by the RDG pass parameters.
 	SetShaderParameters(RHICmdList, PixelShader, PixelShader.GetPixelShader(), PSParameters);
 
@@ -1077,7 +1083,7 @@ void RenderLocalFogVolumeHalfResMobile(
 
 	PassParameters->View = View.GetShaderParameters();
 	PassParameters->LFV = View.LocalFogVolumeViewData.UniformParametersStruct;
-	PassParameters->LocalFogVolumeTileDebug = FMath::Clamp(CVarLocalFogVolumeTileDebug.GetValueOnRenderThread(), 0, 2);
+	PassParameters->LocalFogVolumeTileDebug = GetLocalFogVolumeTileDebugMode(View.Family->EngineShowFlags);
 	PassParameters->MobileBasePass = CreateMobileBasePassUniformBuffer(GraphBuilder, View, BasePass, EMobileSceneTextureSetupMode::SceneDepth);
 	PassParameters->RenderTargets[0] = FRenderTargetBinding(View.LocalFogVolumeViewData.HalfResLocalFogVolumeView, ERenderTargetLoadAction::ENoAction);
 	PassParameters->RenderTargets[1] = FRenderTargetBinding(View.LocalFogVolumeViewData.HalfResLocalFogVolumeDepth, ERenderTargetLoadAction::ENoAction);
diff --git a/Engine/Source/Runtime/Renderer/Private/LocalFogVolumeRendering.h b/Engine/Source/Runtime/Renderer/Private/LocalFogVolumeRendering.h
--- a/Engine/Source/Runtime/Renderer/Private/LocalFogVolumeRendering.h
+++ b/Engine/Source/Runtime/Renderer/Private/LocalFogVolumeRendering.h
@@ -132,7 +132,7 @@ bool ProjectSupportsLocalFogVolumes();
 bool ShouldRenderLocalFogVolume(const FScene* Scene, const FSceneViewFamily& SceneViewFamily);
 bool ShouldRenderLocalFogVolumeDuringHeightFogPass(const FScene* Scene, const FSceneViewFamily& SceneViewFamily);
 bool ShouldRenderLocalFogVolumeInVolumetricFog(const FScene* Scene, const FSceneViewFamily& SceneViewFamily, bool bShouldRenderVolumetricFog);
-bool ShouldRenderLocalFogVolumeVisualisationPass(const FScene* Scene, const FSceneViewFamily& SceneViewFamily);
+bool ShouldRenderLocalFogVolumeVisualizationPass(const FScene* Scene, const FSceneViewFamily& SceneViewFamily);
 float GetLocalFogVolumeGlobalStartDistance();
 bool IsLocalFogVolumeHalfResolution();
 
@@ -157,7 +157,7 @@ void RenderLocalFogVolume(
 	FRDGBuilder& GraphBuilder,
 	const FMinimalSceneTextures& SceneTextures,
 	FRDGTextureRef LightShaftOcclusionTexture,
-	bool bVisualisationPass = false);
+	bool bVisualizationPass = false);
 
 void RenderLocalFogVolumeMobile(
 	FRHICommandList& RHICmdList, 
@@ -166,7 +166,7 @@ void RenderLocalFogVolumeHalfResMobile(
 	FRDGBuilder& GraphBuilder,
 	const FViewInfo& View);
 
-void RenderLocalFogVolumeVisualisation(
+void RenderLocalFogVolumeVisualization(
 	const FScene* Scene,
 	TArray<FViewInfo>& Views,
 	const FSceneViewFamily& SceneViewFamily,
