Holdout: Fix missing reflections with holdout & improved conditions for enabled state.

* Fix regression caused by my previous CL 37178806 due to missing view uniform update in ReflectionEnvironmentRealTimeCapture.cpp to force-disable primitive alpha holdout.
* Unify conditions for enabling primitive alpha holdout with renamed view uniform bPrimitiveAlphaHoldoutEnabled. This also improves the application of the "AllowPrimitiveAlphaHoldout" show flag.
* Updated IsPrimitiveAlphaHoldoutEnabled function for path tracer & reflection capture cases.

diff --git a/Engine/Shaders/Private/BasePassPixelShader.usf b/Engine/Shaders/Private/BasePassPixelShader.usf
--- a/Engine/Shaders/Private/BasePassPixelShader.usf
+++ b/Engine/Shaders/Private/BasePassPixelShader.usf
@@ -138,9 +138,9 @@
 #define RadianceCacheInterpolation TranslucentBasePass
 #include "Lumen/LumenTranslucencyVolumeShared.ush"
 
-#define IsTranslucentHoldoutPass (TranslucentBasePass.TranslucencyPass == 1)
+#define bIsTranslucentHoldoutPass (TranslucentBasePass.TranslucencyPass == 1)
 #else
-#define IsTranslucentHoldoutPass false
+#define bIsTranslucentHoldoutPass false
 #endif
 
 #if TRANSLUCENT_SELF_SHADOWING
@@ -924,7 +924,7 @@ void FPixelShaderInOut_MainPS(
 #if (MATERIALBLENDING_ALPHAHOLDOUT || SUBSTRATE_BLENDING_ALPHAHOLDOUT)
 	const bool bIsHoldout = true;
 #else
-	const bool bIsHoldout = (GetPrimitiveData(MaterialParameters).Flags & PRIMITIVE_SCENE_DATA_FLAG_HOLDOUT) && ResolvedView.bAllowPrimitiveAlphaHoldout;
+	const bool bIsHoldout = (GetPrimitiveData(MaterialParameters).Flags & PRIMITIVE_SCENE_DATA_FLAG_HOLDOUT) && ResolvedView.bPrimitiveAlphaHoldoutEnabled;
 #endif
 #endif
 
@@ -2449,8 +2449,7 @@ void FPixelShaderInOut_MainPS(
 
 #if MATERIALBLENDING_ANY_TRANSLUCENT || MATERIALBLENDING_ALPHACOMPOSITE || MATERIALBLENDING_ALPHAHOLDOUT
 	// If this is translucency holdout pass, calculate the background visibility and path throughput.
-	if (IsTranslucentHoldoutPass &&
-        ResolvedView.RenderingReflectionCaptureMask == 0)
+	if (bIsTranslucentHoldoutPass && ResolvedView.bPrimitiveAlphaHoldoutEnabled)
 	{
 		// Default behavior is to preserve the destination alpha.
 		half BackgroundVisibilityAdd = 0.0f;
diff --git a/Engine/Shaders/Private/HeightFogCommon.ush b/Engine/Shaders/Private/HeightFogCommon.ush
--- a/Engine/Shaders/Private/HeightFogCommon.ush
+++ b/Engine/Shaders/Private/HeightFogCommon.ush
@@ -376,7 +376,7 @@ half4 GetExponentialHeightFog(float3 WorldPositionRelativeToCamera, float Exclud
 	#endif
 
 #if SUPPORT_PRIMITIVE_ALPHA_HOLDOUT
-	if (IsExponentialFogHoldout(View.EnvironmentComponentsFlags) && View.RenderingReflectionCaptureMask == 0.0f)
+	if (IsExponentialFogHoldout(View.EnvironmentComponentsFlags) && View.bPrimitiveAlphaHoldoutEnabled)
 	{
 		FogColor.rgb *= 0;
 	}
@@ -459,7 +459,7 @@ float4 CombineVolumetricFog(float4 GlobalFog, float3 VolumeUV, uint EyeIndex, fl
 	//VolumetricFogLookup.rgb += .1f * frac(min(ZSlice, 1.0f) / View.VolumetricFogInvGridSize.z);
 
 #if SUPPORT_PRIMITIVE_ALPHA_HOLDOUT
-	if (IsExponentialFogHoldout(InView.EnvironmentComponentsFlags) && InView.RenderingReflectionCaptureMask == 0.0f)
+	if (IsExponentialFogHoldout(InView.EnvironmentComponentsFlags) && InView.bPrimitiveAlphaHoldoutEnabled)
 	{
 		VolumetricFogLookup.rgb *= 0;
 		GlobalFog.rgb *= 0;
diff --git a/Engine/Shaders/Private/HeterogeneousVolumes/HeterogeneousVolumesComposite.usf b/Engine/Shaders/Private/HeterogeneousVolumes/HeterogeneousVolumesComposite.usf
--- a/Engine/Shaders/Private/HeterogeneousVolumes/HeterogeneousVolumesComposite.usf
+++ b/Engine/Shaders/Private/HeterogeneousVolumes/HeterogeneousVolumesComposite.usf
@@ -203,7 +203,7 @@ void HeterogeneousVolumesCompositeCS(
 	float4 Result;
 	Result.rgb = Radiance.rgb + RWColorTexture[ViewPixelCoord].rgb * HeterogeneousVolumeTransmittance;
 #if SUPPORT_PRIMITIVE_ALPHA_HOLDOUT
-	if (View.RenderingReflectionCaptureMask == 0.0f)
+	if (View.bPrimitiveAlphaHoldoutEnabled)
 	{
 		float SurfaceOpacity = RWColorTexture[ViewPixelCoord].a * HeterogeneousVolumeTransmittance;
 		float VolumeOpacity = HeterogeneousVolumeHoldout[PixelCoord];
diff --git a/Engine/Shaders/Private/HeterogeneousVolumes/HeterogeneousVolumesLiveShadingPipeline.usf b/Engine/Shaders/Private/HeterogeneousVolumes/HeterogeneousVolumesLiveShadingPipeline.usf
--- a/Engine/Shaders/Private/HeterogeneousVolumes/HeterogeneousVolumesLiveShadingPipeline.usf
+++ b/Engine/Shaders/Private/HeterogeneousVolumes/HeterogeneousVolumesLiveShadingPipeline.usf
@@ -809,7 +809,7 @@ void RenderSingleScatteringWithLiveShading(
 
 		// Radiance
 #if SUPPORT_PRIMITIVE_ALPHA_HOLDOUT
-		if (bHoldout && View.RenderingReflectionCaptureMask == 0.0f)
+		if (bHoldout && View.bPrimitiveAlphaHoldoutEnabled)
 		{
 			Radiance *= 0;
 		}
@@ -822,7 +822,7 @@ void RenderSingleScatteringWithLiveShading(
 			RWLightingTexture[PixelCoord].a = Luminance(Transmittance);
 
 #if SUPPORT_PRIMITIVE_ALPHA_HOLDOUT
-			if (bHoldout && View.RenderingReflectionCaptureMask == 0.0f)
+			if (bHoldout && View.bPrimitiveAlphaHoldoutEnabled)
 			{
 				RWHoldoutTexture[PixelCoord] += Luminance(PrevTransmittance) - Luminance(Transmittance);
 			}
diff --git a/Engine/Shaders/Private/MaterialTemplate.ush b/Engine/Shaders/Private/MaterialTemplate.ush
--- a/Engine/Shaders/Private/MaterialTemplate.ush
+++ b/Engine/Shaders/Private/MaterialTemplate.ush
@@ -2347,7 +2347,7 @@ float3 MaterialExpressionSkyAtmosphereViewLuminance(FMaterialPixelParameters Par
 	SkyAtmosphereViewLuminance *= ResolvedView.OneOverPreExposure;
 	
 #if SUPPORT_PRIMITIVE_ALPHA_HOLDOUT
-	if (IsSkyAtmosphereHoldout(View.EnvironmentComponentsFlags) && View.RenderingReflectionCaptureMask == 0.0f)
+	if (IsSkyAtmosphereHoldout(View.EnvironmentComponentsFlags) && View.bPrimitiveAlphaHoldoutEnabled)
 	{
 		SkyAtmosphereViewLuminance = 0.0f; // Disable sky dome luminance contribution. Throughput will be 0.
 	}
@@ -2383,7 +2383,7 @@ float4 MaterialExpressionSkyAtmosphereAerialPerspective(FMaterialPixelParameters
 		NearFadeOutRangeInvDepthKm);
 	
 #if SUPPORT_PRIMITIVE_ALPHA_HOLDOUT
-	if (IsSkyAtmosphereHoldout(View.EnvironmentComponentsFlags) && View.RenderingReflectionCaptureMask == 0.0f)
+	if (IsSkyAtmosphereHoldout(View.EnvironmentComponentsFlags) && View.bPrimitiveAlphaHoldoutEnabled)
 	{
 		AerialPerspective.rgb = 0.0f;
 	}
diff --git a/Engine/Shaders/Private/PathTracing/PathTracingMaterialHitShader.usf b/Engine/Shaders/Private/PathTracing/PathTracingMaterialHitShader.usf
--- a/Engine/Shaders/Private/PathTracing/PathTracingMaterialHitShader.usf
+++ b/Engine/Shaders/Private/PathTracing/PathTracingMaterialHitShader.usf
@@ -1235,7 +1235,7 @@ RAY_TRACING_ENTRY_CLOSEST_HIT(PathTracingMaterialCHS,
 #else
 	if (PackedPayload.IsCameraRay())
 	{
-		if ((PrimitiveFlags & PRIMITIVE_SCENE_DATA_FLAG_HOLDOUT) != 0 && ResolvedView.bAllowPrimitiveAlphaHoldout)
+		if ((PrimitiveFlags & PRIMITIVE_SCENE_DATA_FLAG_HOLDOUT) != 0 && ResolvedView.bPrimitiveAlphaHoldoutEnabled)
 		{
 			if (Payload.Ior != 0.0)
 			{
diff --git a/Engine/Shaders/Private/RayTracing/RayTracingMaterialHitShaders.usf b/Engine/Shaders/Private/RayTracing/RayTracingMaterialHitShaders.usf
--- a/Engine/Shaders/Private/RayTracing/RayTracingMaterialHitShaders.usf
+++ b/Engine/Shaders/Private/RayTracing/RayTracingMaterialHitShaders.usf
@@ -734,7 +734,7 @@ RAY_TRACING_ENTRY_CLOSEST_HIT(MaterialCHS,
 #if MATERIALBLENDING_ALPHAHOLDOUT
 	const bool bIsHoldout = true;
 #else
-	const bool bIsHoldout = (PrimitiveFlags & PRIMITIVE_SCENE_DATA_FLAG_HOLDOUT) != 0 && ResolvedView.bAllowPrimitiveAlphaHoldout;
+	const bool bIsHoldout = (PrimitiveFlags & PRIMITIVE_SCENE_DATA_FLAG_HOLDOUT) != 0 && ResolvedView.bPrimitiveAlphaHoldoutEnabled;
 #endif
 	const bool bIsCameraRay = PackedPayload.IsCameraRay();
 	if (bIsHoldout && bIsCameraRay)
diff --git a/Engine/Shaders/Private/SkyAtmosphere.usf b/Engine/Shaders/Private/SkyAtmosphere.usf
--- a/Engine/Shaders/Private/SkyAtmosphere.usf
+++ b/Engine/Shaders/Private/SkyAtmosphere.usf
@@ -827,7 +827,7 @@ float4 PrepareOutput(float3 PreExposedLuminance, float3 Transmittance = float3(1
 	if(bPropagateAlphaNonReflection > 0)
 	{
 #if SUPPORT_PRIMITIVE_ALPHA_HOLDOUT
-		LuminanceAlpha.rgb	= IsSkyAtmosphereHoldout(View.EnvironmentComponentsFlags) ? 0.0f : LuminanceAlpha.rgb;
+		LuminanceAlpha.rgb	= (IsSkyAtmosphereHoldout(View.EnvironmentComponentsFlags) && View.bPrimitiveAlphaHoldoutEnabled) ? 0.0f : LuminanceAlpha.rgb;
 #endif
 		LuminanceAlpha.a	= 1 - GreyScaleTransmittance;
 	}
@@ -844,7 +844,7 @@ void UpdateVisibleSkyAlpha(in float DeviceZ, inout float4 OutLuminance)
 	// Output fully opaque alpha if looking through outer space, if atmosphere is not in holdout.
 	bool bOutputFullyOpaque = (bPropagateAlphaNonReflection > 0) && (DeviceZ == FarDepthValue);
 #if SUPPORT_PRIMITIVE_ALPHA_HOLDOUT
-	bOutputFullyOpaque &= !IsSkyAtmosphereHoldout(View.EnvironmentComponentsFlags);	
+	bOutputFullyOpaque &= !IsSkyAtmosphereHoldout(View.EnvironmentComponentsFlags) || !View.bPrimitiveAlphaHoldoutEnabled;
 #endif
 
 	FLATTEN
diff --git a/Engine/Shaders/Private/VolumetricCloud.usf b/Engine/Shaders/Private/VolumetricCloud.usf
--- a/Engine/Shaders/Private/VolumetricCloud.usf
+++ b/Engine/Shaders/Private/VolumetricCloud.usf
@@ -1472,7 +1472,7 @@ void MainCommon(in FMaterialPixelParameters MaterialParameters, in float4 SvPosi
 	const float OutputPreExposure = (ResolvedView.RealTimeReflectionCapture ? ResolvedView.RealTimeReflectionCapturePreExposure : ResolvedView.PreExposure);
 
 #if SUPPORT_PRIMITIVE_ALPHA_HOLDOUT
-	if (IsVolumetricCloudHoldout(View.EnvironmentComponentsFlags) && View.RenderingReflectionCaptureMask == 0.0f)
+	if (IsVolumetricCloudHoldout(View.EnvironmentComponentsFlags) && View.bPrimitiveAlphaHoldoutEnabled)
 	{
 		Luminance *= 0;
 	#if CLOUD_MIN_AND_MAX_DEPTH
@@ -1542,7 +1542,7 @@ void MainCommon(in FMaterialPixelParameters MaterialParameters, in float4 SvPosi
 		Luminance = AerialPerspective.rgb * CloudMeanCoverage + AerialPerspective.a * Luminance;
 
 #if SUPPORT_PRIMITIVE_ALPHA_HOLDOUT
-		if(View.RenderingReflectionCaptureMask == 0.0f)
+		if(View.bPrimitiveAlphaHoldoutEnabled)
 		{
 			const float AerialPespectiveTransmittance = AerialPerspective.a;
 			const float AerialPespectiveCoverage = saturate(1.0 - AerialPespectiveTransmittance);
@@ -1600,7 +1600,7 @@ void MainCommon(in FMaterialPixelParameters MaterialParameters, in float4 SvPosi
 		Luminance = HeightFogInscatteringAndTransmittance.rgb * CloudMeanCoverage + HeightFogInscatteringAndTransmittance.a * Luminance;
 
 #if SUPPORT_PRIMITIVE_ALPHA_HOLDOUT
-		if(View.RenderingReflectionCaptureMask == 0.0f)
+		if(View.bPrimitiveAlphaHoldoutEnabled)
 		{
 			const float ExponentialFogTransmittance = HeightFogInscatteringAndTransmittance.a;
 			const float ExponentialFogCoverage = saturate(1.0 - ExponentialFogTransmittance);
@@ -1819,7 +1819,7 @@ void MainCS(uint2 DispatchThreadId : SV_DispatchThreadID)
 			{
 				OutCloudColor0[DispatchThreadId] = OutColor0;
 			#if SUPPORT_PRIMITIVE_ALPHA_HOLDOUT
-				if (View.RenderingReflectionCaptureMask == 0.0f)
+				if (View.bPrimitiveAlphaHoldoutEnabled)
 				{
 					if (IsVolumetricCloudHoldout(View.EnvironmentComponentsFlags) && bAccumulateAlphaHoldOut > 0)
 					{
diff --git a/Engine/Source/Runtime/Engine/Public/SceneView.h b/Engine/Source/Runtime/Engine/Public/SceneView.h
--- a/Engine/Source/Runtime/Engine/Public/SceneView.h
+++ b/Engine/Source/Runtime/Engine/Public/SceneView.h
@@ -1085,7 +1085,7 @@ enum ETranslucencyVolumeCascade
 	VIEW_UNIFORM_BUFFER_MEMBER(float, MaterialMaxEmissiveValue) \
 	VIEW_UNIFORM_BUFFER_MEMBER(int32, PostVolumeUserFlags) \
 	VIEW_UNIFORM_BUFFER_MEMBER(FIntVector4, EnvironmentComponentsFlags) \
-	VIEW_UNIFORM_BUFFER_MEMBER(uint32, bAllowPrimitiveAlphaHoldout) \
+	VIEW_UNIFORM_BUFFER_MEMBER(uint32, bPrimitiveAlphaHoldoutEnabled) \
 
 /** The uniform shader parameters associated with a view. */
 #define VIEW_UNIFORM_BUFFER_MEMBER(type, identifier) \
diff --git a/Engine/Source/Runtime/Renderer/Private/FogRendering.cpp b/Engine/Source/Runtime/Renderer/Private/FogRendering.cpp
--- a/Engine/Source/Runtime/Renderer/Private/FogRendering.cpp
+++ b/Engine/Source/Runtime/Renderer/Private/FogRendering.cpp
@@ -382,7 +382,7 @@ static void RenderViewFog(
 		if (bSupportsAlpha)
 		{
 			// Coverage is the alpha output of the shader in this case.
-			if (IsExponentialFogHoldout(View.CachedViewUniformShaderParameters->EnvironmentComponentsFlags) && View.CachedViewUniformShaderParameters->RenderingReflectionCaptureMask == 0.0f)
+			if (IsExponentialFogHoldout(View.CachedViewUniformShaderParameters->EnvironmentComponentsFlags) && View.CachedViewUniformShaderParameters->bPrimitiveAlphaHoldoutEnabled)
 			{
 				// Alpha holdout: apply only when requested and when not rendering reflections. We want to punch a hole according to the Coverage. (black as throughput=0 should become brighter for see throught)
 				// SceneAlpha = Coverage*1 + (1.0-Coverage)*(SceneThroughput)
diff --git a/Engine/Source/Runtime/Renderer/Private/HeterogeneousVolumes/HeterogeneousVolumes.cpp b/Engine/Source/Runtime/Renderer/Private/HeterogeneousVolumes/HeterogeneousVolumes.cpp
--- a/Engine/Source/Runtime/Renderer/Private/HeterogeneousVolumes/HeterogeneousVolumes.cpp
+++ b/Engine/Source/Runtime/Renderer/Private/HeterogeneousVolumes/HeterogeneousVolumes.cpp
@@ -415,10 +415,8 @@ bool ShouldRenderHeterogeneousVolumesAsHoldoutForView(
 	const FViewInfo& View
 )
 {
-	const bool bAllowPrimitiveAlphaHoldout = View.Family->EngineShowFlags.AllowPrimitiveAlphaHoldout;
-
 	// This query returns true if any volume is marked as a holdout; otherwise, the query returns false
-	if (ShouldRenderHeterogeneousVolumesForView(View) && bAllowPrimitiveAlphaHoldout)
+	if (ShouldRenderHeterogeneousVolumesForView(View))
 	{
 		for (int32 MeshBatchIndex = 0; MeshBatchIndex < View.HeterogeneousVolumesMeshBatches.Num(); ++MeshBatchIndex)
 		{
diff --git a/Engine/Source/Runtime/Renderer/Private/PathTracing.cpp b/Engine/Source/Runtime/Renderer/Private/PathTracing.cpp
--- a/Engine/Source/Runtime/Renderer/Private/PathTracing.cpp
+++ b/Engine/Source/Runtime/Renderer/Private/PathTracing.cpp
@@ -896,10 +896,13 @@ static void PreparePathTracingData(const FScene* Scene, const FViewInfo& View, F
 		&& (Scene->ExponentialFogs[0].FogData[0].Density > 0 ||
 			Scene->ExponentialFogs[0].FogData[1].Density > 0) ? PATH_TRACER_VOLUME_ENABLE_FOG : 0;
 	PathTracingData.VolumeFlags |= ShouldRenderHeterogeneousVolumesForView(View) ? PATH_TRACER_VOLUME_ENABLE_HETEROGENEOUS_VOLUMES : 0;
-	PathTracingData.VolumeFlags |= ShouldRenderHeterogeneousVolumesAsHoldoutForView(View) ? PATH_TRACER_VOLUME_HOLDOUT_HETEROGENEOUS_VOLUMES : 0;
-	PathTracingData.VolumeFlags |= View.SkyAtmosphereUniformShaderParameters != nullptr && IsSkyAtmosphereHoldout(View.CachedViewUniformShaderParameters->EnvironmentComponentsFlags) ? PATH_TRACER_VOLUME_HOLDOUT_ATMOSPHERE : 0;
-	PathTracingData.VolumeFlags |= bVolumeCloudsVisible && IsVolumetricCloudHoldout(View.CachedViewUniformShaderParameters->EnvironmentComponentsFlags) ? PATH_TRACER_VOLUME_HOLDOUT_CLOUDS : 0;
-	PathTracingData.VolumeFlags |= Scene->ExponentialFogs.Num() > 0 && IsExponentialFogHoldout(View.CachedViewUniformShaderParameters->EnvironmentComponentsFlags) ? PATH_TRACER_VOLUME_HOLDOUT_FOG : 0;
+	if (View.CachedViewUniformShaderParameters->bPrimitiveAlphaHoldoutEnabled)
+	{
+		PathTracingData.VolumeFlags |= ShouldRenderHeterogeneousVolumesAsHoldoutForView(View) ? PATH_TRACER_VOLUME_HOLDOUT_HETEROGENEOUS_VOLUMES : 0;
+		PathTracingData.VolumeFlags |= View.SkyAtmosphereUniformShaderParameters != nullptr && IsSkyAtmosphereHoldout(View.CachedViewUniformShaderParameters->EnvironmentComponentsFlags) ? PATH_TRACER_VOLUME_HOLDOUT_ATMOSPHERE : 0;
+		PathTracingData.VolumeFlags |= bVolumeCloudsVisible && IsVolumetricCloudHoldout(View.CachedViewUniformShaderParameters->EnvironmentComponentsFlags) ? PATH_TRACER_VOLUME_HOLDOUT_CLOUDS : 0;
+		PathTracingData.VolumeFlags |= Scene->ExponentialFogs.Num() > 0 && IsExponentialFogHoldout(View.CachedViewUniformShaderParameters->EnvironmentComponentsFlags) ? PATH_TRACER_VOLUME_HOLDOUT_FOG : 0;
+	}
 	PathTracingData.VolumeFlags |= CVarPathTracingUseAnalyticTransmittance.GetValueOnRenderThread() ? PATH_TRACER_VOLUME_USE_ANALYTIC_TRANSMITTANCE : 0;
 
 	PathTracingData.CloudRoughnessCutoff = bVolumeCloudsVisible ? CVarPathTracingCloudRoughnessCutoff.GetValueOnRenderThread() : -1.0f;
diff --git a/Engine/Source/Runtime/Renderer/Private/ReflectionEnvironmentRealTimeCapture.cpp b/Engine/Source/Runtime/Renderer/Private/ReflectionEnvironmentRealTimeCapture.cpp
--- a/Engine/Source/Runtime/Renderer/Private/ReflectionEnvironmentRealTimeCapture.cpp
+++ b/Engine/Source/Runtime/Renderer/Private/ReflectionEnvironmentRealTimeCapture.cpp
@@ -566,6 +566,8 @@ void FScene::AllocateAndCaptureFrameSkyEnvMap(
 
 				// Notify the fact that we render a reflection, e.g. remove sun disk.
 				CubeView.CachedViewUniformShaderParameters->RenderingReflectionCaptureMask = 1.0f;
+				// Force-disable primitive alpha holdout for reflection captures
+				CubeView.CachedViewUniformShaderParameters->bPrimitiveAlphaHoldoutEnabled = 0;
 				// Notify the fact that we render a reflection, e.g. use special exposure.
 				CubeView.CachedViewUniformShaderParameters->RealTimeReflectionCapture = 1.0f;
 			};
diff --git a/Engine/Source/Runtime/Renderer/Private/SceneRendering.cpp b/Engine/Source/Runtime/Renderer/Private/SceneRendering.cpp
--- a/Engine/Source/Runtime/Renderer/Private/SceneRendering.cpp
+++ b/Engine/Source/Runtime/Renderer/Private/SceneRendering.cpp
@@ -1868,7 +1868,7 @@ void FViewInfo::SetupUniformBufferParameters(
 	ViewUniformShaderParameters.RealTimeReflectionCapture = 0.0f;
 	ViewUniformShaderParameters.RealTimeReflectionCapturePreExposure = 1.0f; // This must be 1 for now. If changed, we need to update the SkyLight AverageExposure and take it into account when sampling sky specular and diffuse irradiance.
 
-	ViewUniformShaderParameters.bAllowPrimitiveAlphaHoldout = static_cast<uint32>(Family->EngineShowFlags.AllowPrimitiveAlphaHoldout && !bIsReflectionCapture);
+	ViewUniformShaderParameters.bPrimitiveAlphaHoldoutEnabled = IsPrimitiveAlphaHoldoutEnabled(*this);
 
 	ViewUniformShaderParameters.AmbientCubemapTint = FinalPostProcessSettings.AmbientCubemapTint;
 	ViewUniformShaderParameters.AmbientCubemapIntensity = FinalPostProcessSettings.AmbientCubemapIntensity;
@@ -6151,12 +6151,14 @@ void CreateQuarterResolutionDepthMinAndMaxFromDepthTexture(FRDGBuilder& GraphBui
 
 bool IsPrimitiveAlphaHoldoutEnabled(const FViewInfo& View)
 {
-	const bool bSupportPrimitiveAlphaHoldout = CVarPrimitiveAlphaHoldoutSupport.GetValueOnRenderThread();
+	// Note: r.Deferred.SupportPrimitiveAlphaHoldout excludes the path tracer
+	const bool bSupportPrimitiveAlphaHoldout = View.Family->EngineShowFlags.PathTracing ? true : CVarPrimitiveAlphaHoldoutSupport.GetValueOnRenderThread();
 
 	return bSupportPrimitiveAlphaHoldout
 		&& (GetFeatureLevelShadingPath(View.GetFeatureLevel()) != EShadingPath::Mobile) 
 		&& IsPostProcessingWithAlphaChannelSupported()
-		&& View.Family->EngineShowFlags.AllowPrimitiveAlphaHoldout;
+		&& View.Family->EngineShowFlags.AllowPrimitiveAlphaHoldout
+		&& !View.bIsReflectionCapture; // Force-disable primitive alpha holdout during reflection captures
 }
 
 bool IsPrimitiveAlphaHoldoutEnabledForAnyView(TArrayView<const FViewInfo> Views)
diff --git a/Engine/Source/Runtime/Renderer/Private/VolumetricCloudRendering.cpp b/Engine/Source/Runtime/Renderer/Private/VolumetricCloudRendering.cpp
--- a/Engine/Source/Runtime/Renderer/Private/VolumetricCloudRendering.cpp
+++ b/Engine/Source/Runtime/Renderer/Private/VolumetricCloudRendering.cpp
@@ -2749,7 +2749,7 @@ bool FSceneRenderer::RenderVolumetricCloud(
 					DestinationRTDepth = VRT.GetOrCreateVolumetricTracingRTDepth(GraphBuilder);
 
 					// Only account for alphaholdout if any important primitive is holdout, and we are in mode 3
-					if (ViewFamily.EngineShowFlags.AllowPrimitiveAlphaHoldout && IsPostProcessingWithAlphaChannelSupported() && ViewInfo.CachedViewUniformShaderParameters->RenderingReflectionCaptureMask == 0.0f)
+					if (ViewInfo.CachedViewUniformShaderParameters->bPrimitiveAlphaHoldoutEnabled)
 					{
 						const FIntVector4& EnvironmentComponentFlags = ViewInfo.CachedViewUniformShaderParameters->EnvironmentComponentsFlags;
 						bAccumulateAlphaHoldOut = VRT.GetMode() == 3 &&  (IsSkyAtmosphereHoldout(EnvironmentComponentFlags) || IsVolumetricCloudHoldout(EnvironmentComponentFlags) || IsExponentialFogHoldout(EnvironmentComponentFlags));
