Path Tracing: Add permutation that changes how transparency is traced

We have discovered that on some GPU models, the DXR spec is not exactly obeyed and it is possible to get a valid hit at TMin. In the path tracer, this prevents the transparency tracing loop from making forward progress which causes crashes.

This patch introduces a permutation for this, which only gets used if the current hardware needs such a workaround.

diff --git a/Engine/Shaders/Private/PathTracing/PathTracingCommon.ush b/Engine/Shaders/Private/PathTracing/PathTracingCommon.ush
--- a/Engine/Shaders/Private/PathTracing/PathTracingCommon.ush
+++ b/Engine/Shaders/Private/PathTracing/PathTracingCommon.ush
@@ -47,6 +47,28 @@ HLSL_STATIC_ASSERT(SHADINGMODELID_SOLID_GLASS <= SHADINGMODELID_MASK, "");
 
 #define PATH_TRACING_ABSORPTION_SCALE  0.01        // inverse distance at which we hit the transmission color (Beer's law)
 
+
+// Given a ray with origin O in direction D, assuming we have found a valid hit at THit, compute a new TMin value that we can use to retrace
+// the same ray and find more hits. This is used for transparency, mesh decals and SSS tracing.
+float ComputeNewTMin(float3 O, float3 D, float THit)
+{
+#if NEED_TMIN_WORKAROUND
+	// Some hardware shifts the ray origin to TMin which can lead to erroneous hits against TMin when shifted back
+	// This code tries to find to smallest increment that causes the shifted origin to change
+	uint Inc = 1;
+	float TMin = asfloat(asuint(THit) + Inc);
+	while (all((mad(D, TMin, O) == mad(D, THit, O)))) {
+		Inc *= 2;
+		TMin = asfloat(asuint(THit) + Inc);
+	}
+	return TMin;
+#else
+	// On conformant hardware, the following is sufficient to avoid the previous hit
+	return asfloat(asuint(THit) + 1);
+#endif
+}
+
+
 // Utilities to help weight the contribution of the albedo/normal to the denoising buffers
 // We want cases like low roughness mirrors to "pass through" the signal
 static const float DenoiserMaxRoughness = 0.15f; // don't bother past this level
diff --git a/Engine/Shaders/Private/PathTracing/PathTracingCore.ush b/Engine/Shaders/Private/PathTracing/PathTracingCore.ush
--- a/Engine/Shaders/Private/PathTracing/PathTracingCore.ush
+++ b/Engine/Shaders/Private/PathTracing/PathTracingCore.ush
@@ -448,7 +448,12 @@ void EvaluateDecals(FPathState PathState, float HitT, inout FDecalShaderPayload
 				MissShaderIndex,
 				DecalRay.GetNativeDesc(),
 				DecalPayload);
-
+#if NEED_TMIN_WORKAROUND // extra safety - discard hit if not conforming
+			if (DecalPayload.HitT <= DecalRay.TMin)
+			{
+				DecalPayload.HitT = -1.0; 
+			}
+#endif
 			if (DecalPayload.IsMiss())
 			{
 				break;
@@ -456,7 +461,7 @@ void EvaluateDecals(FPathState PathState, float HitT, inout FDecalShaderPayload
 
 			// prepare next step around the loop
 			// retrace the exact same ray with TMin one ulp past the hit we just found
-			DecalRay.TMin = asfloat(asuint(DecalPayload.HitT) + 1);
+			DecalRay.TMin = ComputeNewTMin(DecalRay.Origin, DecalRay.Direction, DecalPayload.HitT);
 		}
 	}
 
@@ -887,7 +892,12 @@ FPackedPathTracingPayload TraceTransparentRay(inout FPathState PathState, int Bo
 			VolumeIsectList = VolumeIsectList.Update(Interval.VolumeTMax);
 		}
 		// proceed to handling the surface hit (if any)
-
+#if NEED_TMIN_WORKAROUND // extra safety - discard hit if not conforming
+		if (PackedPayload.HitT <= PathState.Ray.TMin)
+		{
+			PackedPayload.HitT = -1.0;
+		}
+#endif
 		if (PackedPayload.IsMiss())
 		{
 			// Ray didn't hit any real geometry
@@ -997,7 +1007,7 @@ FPackedPathTracingPayload TraceTransparentRay(inout FPathState PathState, int Bo
 
 		// prepare next step around the loop
 		// retrace the exact same ray with TMin one ulp past the hit we just found
-		PathState.Ray.TMin = asfloat(asuint(PackedPayload.HitT) + 1);
+		PathState.Ray.TMin = ComputeNewTMin(PathState.Ray.Origin, PathState.Ray.Direction, PackedPayload.HitT);
 		if (!(PathState.Ray.TMin < PathState.Ray.TMax))
 		{
 			break;
@@ -1128,7 +1138,12 @@ FProbeResult TraceSSSProbeRay(FRayDesc Ray, inout int InterfaceCounter)
 			MissShaderIndex,
 			Ray.GetNativeDesc(),
 			PackedPayload);
-
+#if NEED_TMIN_WORKAROUND // extra safety - discard hit if not conforming
+		if (PackedPayload.HitT <= Ray.TMin)
+		{
+			PackedPayload.HitT = -1.0;
+		}
+#endif
 		if (PackedPayload.IsMiss())
 		{
 			// we did not hit anything
@@ -1141,7 +1156,7 @@ FProbeResult TraceSSSProbeRay(FRayDesc Ray, inout int InterfaceCounter)
 		if (InterfaceCounter != 0)
 		{
 			// we have not yet crossed the right number of interfaces, so ignore the current hit and try the segment again
-			Ray.TMin = asfloat(asuint(PackedPayload.HitT) + 1);
+			Ray.TMin = ComputeNewTMin(Ray.Origin, Ray.Direction, PackedPayload.HitT);
 			continue;
 		}
 #endif
diff --git a/Engine/Shaders/Private/PathTracing/PathTracingInitExtinctionCoefficient.usf b/Engine/Shaders/Private/PathTracing/PathTracingInitExtinctionCoefficient.usf
--- a/Engine/Shaders/Private/PathTracing/PathTracingInitExtinctionCoefficient.usf
+++ b/Engine/Shaders/Private/PathTracing/PathTracingInitExtinctionCoefficient.usf
@@ -33,6 +33,12 @@ RAY_TRACING_ENTRY_RAYGEN(PathTracingInitExtinctionCoefficientRG)
 			MissShaderIndex,
 			UpRay.GetNativeDesc(),
 			PackedPayload);
+#if NEED_TMIN_WORKAROUND // extra safety - discard hit if not conforming
+		if (PackedPayload.HitT <= UpRay.TMin)
+		{
+			PackedPayload.HitT = -1.0;
+		}
+#endif
 		if (PackedPayload.IsMiss())
 		{
 			// we didn't hit anything
@@ -49,7 +55,7 @@ RAY_TRACING_ENTRY_RAYGEN(PathTracingInitExtinctionCoefficientRG)
 			StartingSigmaT += HitPayload.IsFrontFace() ? -LocalSigmaT : +LocalSigmaT;
 		}
 		// keep tracing
-		UpRay.TMin = asfloat(asuint(PackedPayload.HitT) + 1);
+		UpRay.TMin = ComputeNewTMin(UpRay.Origin, UpRay.Direction, PackedPayload.HitT);
 	}
 	// clamp against 0 in the unlikely case that we only hit front faces of glass on our way up
 	StartingSigmaT = max(StartingSigmaT, 0.0);
diff --git a/Engine/Source/Runtime/Renderer/Private/PathTracing.cpp b/Engine/Source/Runtime/Renderer/Private/PathTracing.cpp
--- a/Engine/Source/Runtime/Renderer/Private/PathTracing.cpp
+++ b/Engine/Source/Runtime/Renderer/Private/PathTracing.cpp
@@ -1217,7 +1217,8 @@ class FPathTracingRG : public FGlobalShader
 	class FSubstrateComplexSpecialMaterial : SHADER_PERMUTATION_BOOL("PATH_TRACER_USE_SUBSTRATE_SPECIAL_COMPLEX_MATERIAL");
 	class FUseSER : SHADER_PERMUTATION_BOOL("PATH_TRACER_USE_SER");
 	class FTraceOpaqueFirst : SHADER_PERMUTATION_BOOL("PATH_TRACER_TRACE_OPAQUE_FIRST");
-	using FPermutationDomain = TShaderPermutationDomain<FCompactionType, FAdaptiveSampling, FCloudShader, FSubstrateComplexSpecialMaterial, FUseSER, FTraceOpaqueFirst>;
+	class FNeedTMinWorkaround : SHADER_PERMUTATION_BOOL("NEED_TMIN_WORKAROUND");
+	using FPermutationDomain = TShaderPermutationDomain<FCompactionType, FAdaptiveSampling, FCloudShader, FSubstrateComplexSpecialMaterial, FUseSER, FTraceOpaqueFirst, FNeedTMinWorkaround>;
 
 	static bool ShouldCompilePermutation(const FGlobalShaderPermutationParameters& Parameters)
 	{
@@ -1433,6 +1434,9 @@ class FPathTracingInitExtinctionCoefficientRG : public FGlobalShader
 	DECLARE_GLOBAL_SHADER(FPathTracingInitExtinctionCoefficientRG)
 	SHADER_USE_ROOT_PARAMETER_STRUCT(FPathTracingInitExtinctionCoefficientRG, FGlobalShader)
 
+	class FNeedTMinWorkaround : SHADER_PERMUTATION_BOOL("NEED_TMIN_WORKAROUND");
+	using FPermutationDomain = TShaderPermutationDomain<FNeedTMinWorkaround>;
+
 	static bool ShouldCompilePermutation(const FGlobalShaderPermutationParameters& Parameters)
 	{
 		return ShouldCompilePathTracingShadersForProject(Parameters.Platform);
@@ -3149,6 +3153,14 @@ class FPathTracingCompositorPS : public FGlobalShader
 };
 IMPLEMENT_SHADER_TYPE(, FPathTracingCompositorPS, TEXT("/Engine/Private/PathTracing/PathTracingCompositingPixelShader.usf"), TEXT("CompositeMain"), SF_Pixel);
 
+static bool GPUNeedsTMinWorkaround()
+{
+	// See JIRA UE-213077
+	// Some AMD GPUs can return true for a hit exactly at TMin. This causes some of the loops in the path tracer that want to repeatedly trace the same ray to get stuck in infinite loops
+	// This workaround will be fine tuned in the future as drivers that fix this issue get released.
+	return IsRHIDeviceAMD();
+}
+
 static FPathTracingRG::FPermutationDomain GetPathTracingRGPermutation(const FViewInfo& View, const FScene& Scene, const bool bUseCompaction)
 {
 	const bool bUseExperimental = CVarPathTracingExperimental.GetValueOnRenderThread();
@@ -3166,6 +3178,7 @@ static FPathTracingRG::FPermutationDomain GetPathTracingRGPermutation(const FVie
 	Out.Set<FPathTracingRG::FSubstrateComplexSpecialMaterial>(bHasComplexSpecialRenderPath);
 	Out.Set<FPathTracingRG::FUseSER>(bUseSER);
 	Out.Set<FPathTracingRG::FTraceOpaqueFirst>(bUseTraceOpaqueFirst);
+	Out.Set<FPathTracingRG::FNeedTMinWorkaround>(GPUNeedsTMinWorkaround());
 	return Out;
 }
 
@@ -3199,7 +3212,9 @@ void FDeferredShadingSceneRenderer::PreparePathTracing(const FViewInfo& View, co
 				OutRayGenShaders.Add(RayGenShader.GetRayTracingShader());
 			}
 			{
-				auto RayGenShader = GetGlobalShaderMap(ViewFamily.GetShaderPlatform())->GetShader<FPathTracingInitExtinctionCoefficientRG>();
+				FPathTracingInitExtinctionCoefficientRG::FPermutationDomain PermutationVector;
+				PermutationVector.Set<FPathTracingInitExtinctionCoefficientRG::FNeedTMinWorkaround>(GPUNeedsTMinWorkaround());
+				auto RayGenShader = GetGlobalShaderMap(ViewFamily.GetShaderPlatform())->GetShader<FPathTracingInitExtinctionCoefficientRG>(PermutationVector);
 				OutRayGenShaders.Add(RayGenShader.GetRayTracingShader());
 			}
 		}
@@ -3227,6 +3242,12 @@ void PreparePathTracingRTPSO()
 
 				FPathTracingRG::FPermutationDomain PermutationVector(PermutationId);
 
+				if (PermutationVector.Get<FPathTracingRG::FNeedTMinWorkaround>() != GPUNeedsTMinWorkaround())
+				{
+					// only compile the version that will be used at runtime
+					continue;
+				}
+
 				if (PermutationVector.Get<FPathTracingRG::FUseSER>() && !GRHIGlobals.SupportsShaderExecutionReordering)
 				{
 					// we can safely skip these if the current machine does not support SER
@@ -3234,10 +3255,13 @@ void PreparePathTracingRTPSO()
 					// where the machine that decides what to compile is not the machine that will run the code
 					continue;
 				}
+
+				FPathTracingInitExtinctionCoefficientRG::FPermutationDomain PermutationVectorInitExtinctionCoeffs;
+				PermutationVectorInitExtinctionCoeffs.Set<FPathTracingInitExtinctionCoefficientRG::FNeedTMinWorkaround>(GPUNeedsTMinWorkaround());
 				
 				FRHIRayTracingShader* RayGenShaderTable[] = {
 					ShaderMap->GetShader<FPathTracingRG>(PermutationVector).GetRayTracingShader(),
-					ShaderMap->GetShader<FPathTracingInitExtinctionCoefficientRG>().GetRayTracingShader(),
+					ShaderMap->GetShader<FPathTracingInitExtinctionCoefficientRG>(PermutationVectorInitExtinctionCoeffs).GetRayTracingShader(),
 				};
 				FRHIRayTracingShader* MissShaderTable[] = {
 					GetPathTracingDefaultMissShader(ShaderMap),
@@ -3537,7 +3561,10 @@ void FDeferredShadingSceneRenderer::RenderPathTracing(
 	}
 	else if (!PathTracingState->StartingExtinctionCoefficient.IsValid() || PathTracingState->SampleIndex == 0 || View.bIsOfflineRender)
 	{
-		auto RayGenShader = GetGlobalShaderMap(View.FeatureLevel)->GetShader<FPathTracingInitExtinctionCoefficientRG>();
+		FPathTracingInitExtinctionCoefficientRG::FPermutationDomain PermutationVector;
+		PermutationVector.Set<FPathTracingInitExtinctionCoefficientRG::FNeedTMinWorkaround>(GPUNeedsTMinWorkaround());
+
+		auto RayGenShader = GetGlobalShaderMap(View.FeatureLevel)->GetShader<FPathTracingInitExtinctionCoefficientRG>(PermutationVector);
 
 		// prepare extinction coefficient for camera rays
 		StartingExtinctionCoefficient = GraphBuilder.CreateBuffer(FRDGBufferDesc::CreateBufferDesc(sizeof(float), 3), TEXT("PathTracer.StartingExtinctionCoefficient"));
