Bumped OcclusionFinalize priority to High to avoid visibility getting starved out by game thread tasks.

diff --git a/Engine/Source/Runtime/Renderer/Private/SceneVisibilityPrivate.h b/Engine/Source/Runtime/Renderer/Private/SceneVisibilityPrivate.h
--- a/Engine/Source/Runtime/Renderer/Private/SceneVisibilityPrivate.h
+++ b/Engine/Source/Runtime/Renderer/Private/SceneVisibilityPrivate.h
@@ -391,7 +391,7 @@ public:
 		static constexpr uint32 MinQueriesPerTask = 64;
 
 		const UE::Tasks::ETaskPriority TaskPriority = UE::Tasks::ETaskPriority::High;
-		const UE::Tasks::ETaskPriority FinalizeTaskPriority = UE::Tasks::ETaskPriority::Normal;
+		const UE::Tasks::ETaskPriority FinalizeTaskPriority = UE::Tasks::ETaskPriority::High;
 
 		struct FView
 		{
