Add an extra bit of precision to the instance culling screen size.
It was passed to culling shader as 3.9 fixed point and the change makes it 3.10.
This is a fix for some content that has _very_ small LOD screen sizes, and it is good to make use of available bits for extra precision.
But it is not robust against content with _very very_ small LOD screen sizes. A later fix should guard against that.
#rnx

diff --git a/Engine/Shaders/Private/InstanceCulling/InstanceCullingCommon.ush b/Engine/Shaders/Private/InstanceCulling/InstanceCullingCommon.ush
--- a/Engine/Shaders/Private/InstanceCulling/InstanceCullingCommon.ush
+++ b/Engine/Shaders/Private/InstanceCulling/InstanceCullingCommon.ush
@@ -28,8 +28,8 @@ FDrawCommandDesc UnpackDrawCommandDesc(uint PackedData)
 	Desc.bMaterialUsesWorldPositionOffset = (PackedData & 1U) != 0;
 	Desc.bMaterialAlwaysEvaluatesWorldPositionOffset = ((PackedData >> 1U) & 1U) != 0;
 	Desc.MeshLODIndex = (PackedData >> 2U) & 0xFU;
-	Desc.MinScreenSize = (float)((PackedData >> 6U) & 0xFFFU) / (float)(1 << 9);
-	Desc.MaxScreenSize = (float)((PackedData >> 18U) & 0xFFFU) / (float)(1 << 9);
+	Desc.MinScreenSize = (float)((PackedData >> 6U) & 0x1FFFU) / (float)(1 << 10);
+	Desc.MaxScreenSize = (float)((PackedData >> 19U) & 0x1FFFU) / (float)(1 << 10);
 	return Desc;
 }
 
diff --git a/Engine/Source/Runtime/Renderer/Private/InstanceCulling/InstanceCullingContext.cpp b/Engine/Source/Runtime/Renderer/Private/InstanceCulling/InstanceCullingContext.cpp
--- a/Engine/Source/Runtime/Renderer/Private/InstanceCulling/InstanceCullingContext.cpp
+++ b/Engine/Source/Runtime/Renderer/Private/InstanceCulling/InstanceCullingContext.cpp
@@ -85,7 +85,7 @@ static uint32 PackDrawCommandDesc(bool bMaterialUsesWorldPositionOffset, bool bM
 	}
 	if (EnumHasAnyFlags(CullingPayloadFlags, EMeshDrawCommandCullingPayloadFlags::MaxScreenSizeCull))
 	{
-		PackedData |= CullingPayload.MaxScreenSize << 18;
+		PackedData |= CullingPayload.MaxScreenSize << 19;
 	}
 	return PackedData;
 }
diff --git a/Engine/Source/Runtime/Renderer/Public/MeshPassProcessor.h b/Engine/Source/Runtime/Renderer/Public/MeshPassProcessor.h
--- a/Engine/Source/Runtime/Renderer/Public/MeshPassProcessor.h
+++ b/Engine/Source/Runtime/Renderer/Public/MeshPassProcessor.h
@@ -1619,22 +1619,22 @@ struct FMeshDrawCommandCullingPayload
 		struct
 		{
 			uint32 LodIndex : 4;
-			uint32 MinScreenSize : 12;
-			uint32 MaxScreenSize : 12;
+			uint32 MinScreenSize : 13;
+			uint32 MaxScreenSize : 13;
 		};
 	};
 
-	// Pack screen sizes into fixed point 3.9
+	// Pack screen sizes into fixed point 3.10
 	static uint32 PackScreenSize(float ScreenSize)
 	{
-		constexpr float MaxExpressibleValue = (float)(0xFFFU) / (float)(1 << 9);
+		constexpr float MaxExpressibleValue = (float)(0x1FFFU) / (float)(1 << 10);
 		ScreenSize = FMath::Clamp(ScreenSize, 0.f, MaxExpressibleValue);
-		return (uint32)(ScreenSize * (1 << 9)) & 0xFFF;
+		return (uint32)(ScreenSize * (1 << 10)) & 0x1FFF;
 	}
 
 	static float UnpackScreenSize(uint32 PackedScreenSize)
 	{
-		return (float)(PackedScreenSize & 0xFFFU) / (float)(1 << 9);
+		return (float)(PackedScreenSize & 0x1FFFU) / (float)(1 << 10);
 	}
 };
 
