[Substrate] Do not force VRS off when Substrate is enabled.

Save 0.24ms on NaniteBasePass during replays on PS5.

diff --git a/Engine/Source/Runtime/Renderer/Private/VariableRateShading/VariableRateShadingImageManager.cpp b/Engine/Source/Runtime/Renderer/Private/VariableRateShading/VariableRateShadingImageManager.cpp
--- a/Engine/Source/Runtime/Renderer/Private/VariableRateShading/VariableRateShadingImageManager.cpp
+++ b/Engine/Source/Runtime/Renderer/Private/VariableRateShading/VariableRateShadingImageManager.cpp
@@ -536,7 +536,6 @@ void FVariableRateShadingImageManager::PrepareImageBasedVRS(FRDGBuilder& GraphBu
 	bSoftwareVRSEnabledForFrame = 
 		CVarEnableVRSSoftwareImage.GetValueOnRenderThread() > 0 &&
 		CVarLocalNaniteSoftwareVRS->GetInt() > 0 &&
-		!Substrate::IsSubstrateEnabled() &&
 		IsFeatureLevelSupported(ShaderPlatform, ERHIFeatureLevel::SM6);
 
 	if (!IsVRSEnabledForFrame())
