Move virtual texture finalize to after nanite streaming.
This is important if we want to render nanite to RVT or Material Cache.

diff --git a/Engine/Source/Runtime/Renderer/Private/DeferredShadingRenderer.cpp b/Engine/Source/Runtime/Renderer/Private/DeferredShadingRenderer.cpp
--- a/Engine/Source/Runtime/Renderer/Private/DeferredShadingRenderer.cpp
+++ b/Engine/Source/Runtime/Renderer/Private/DeferredShadingRenderer.cpp
@@ -1905,13 +1905,10 @@ void FDeferredShadingSceneRenderer::Render(FRDGBuilder& GraphBuilder, const FSce
 	bool bAnyLumenEnabled = false;
 
 	// Virtual texturing isn't needed for depth prepass
-	if (RendererOutput != ERendererOutput::DepthPrepassOnly)
+	if (bUseVirtualTexturing && RendererOutput != ERendererOutput::DepthPrepassOnly)
 	{
-		if (bUseVirtualTexturing)
-		{
-			// Note, should happen after the GPU-Scene update to ensure rendering to runtime virtual textures is using the correctly updated scene
-			FVirtualTextureSystem::Get().EndUpdate(GraphBuilder, MoveTemp(VirtualTextureUpdater), FeatureLevel);
-		}
+		// Note, should happen after the GPU-Scene update to ensure rendering to runtime virtual textures is using the correctly updated scene
+		FVirtualTextureSystem::Get().EndUpdate(GraphBuilder, MoveTemp(VirtualTextureUpdater), FeatureLevel);
 	}
 
 	UE::Tasks::TTask<FSortedLightSetSceneInfo*> GatherAndSortLightsTask;
@@ -2009,6 +2006,12 @@ void FDeferredShadingSceneRenderer::Render(FRDGBuilder& GraphBuilder, const FSce
 		MaterialCacheRenderPages(GraphBuilder, this);
 	}
 
+	// Virtual texturing isn't needed for depth prepass
+	if (bUseVirtualTexturing && RendererOutput != ERendererOutput::DepthPrepassOnly)
+	{
+		FVirtualTextureSystem::Get().FinalizeRequests(GraphBuilder);
+	}
+
 	{
 		RDG_RHI_GPU_STAT_SCOPE(GraphBuilder, VisibilityCommands);
 		EndInitViews(GraphBuilder, LumenFrameTemporaries, InstanceCullingManager, InitViewTaskDatas);
diff --git a/Engine/Source/Runtime/Renderer/Private/MobileShadingRenderer.cpp b/Engine/Source/Runtime/Renderer/Private/MobileShadingRenderer.cpp
--- a/Engine/Source/Runtime/Renderer/Private/MobileShadingRenderer.cpp
+++ b/Engine/Source/Runtime/Renderer/Private/MobileShadingRenderer.cpp
@@ -1187,6 +1187,7 @@ void FMobileSceneRenderer::Render(FRDGBuilder& GraphBuilder, const FSceneRenderU
 		if (bUseVirtualTexturing)
 		{
 			FVirtualTextureSystem::Get().EndUpdate(GraphBuilder, MoveTemp(VirtualTextureUpdater), FeatureLevel);
+			FVirtualTextureSystem::Get().FinalizeRequests(GraphBuilder);
 		}
 
 		if (bDeferredShading ||
diff --git a/Engine/Source/Runtime/Renderer/Private/VT/VirtualTextureSystem.cpp b/Engine/Source/Runtime/Renderer/Private/VT/VirtualTextureSystem.cpp
--- a/Engine/Source/Runtime/Renderer/Private/VT/VirtualTextureSystem.cpp
+++ b/Engine/Source/Runtime/Renderer/Private/VT/VirtualTextureSystem.cpp
@@ -2567,6 +2567,8 @@ void FVirtualTextureSystem::FinalizeRequests(FRDGBuilder& GraphBuilder)
 
 	ExternalAccessQueue.Submit(GraphBuilder);
 
+	Producers.NotifyRequestsCompleted();
+
 	Frame++;
 }
 
@@ -2789,6 +2791,10 @@ void FVirtualTextureSystem::BeginUpdate(FRDGBuilder& GraphBuilder, FVirtualTextu
 	{
 		TRACE_CPUPROFILER_EVENT_SCOPE(FVirtualTextureSystem::GatherAndSubmitRequests);
 
+		// This doesn't need to be synced with requests and can be updated at any point in the frame.
+		// Do it here so that it is off the render thread.
+		UpdateResidencyTracking();
+
 		Updater->MergedRequestList = Allocator.Create<FUniqueRequestList>(Allocator);
 		Updater->MergedRequestList->Initialize();
 
@@ -2979,14 +2985,8 @@ void FVirtualTextureSystem::EndUpdate(FRDGBuilder& GraphBuilder, TUniquePtr<FVir
 		SubmitThrottledRequests(GraphBuilder.RHICmdList, Updater.Get(), EUpdatePhase::End);
 	}
 
-	FinalizeRequests(GraphBuilder);
-
-	Producers.NotifyRequestsCompleted();
-
-	UpdateResidencyTracking();
 	GrowPhysicalPools();
 
-
 #if !UE_BUILD_SHIPPING
 	UpdateCsvStats();
 #endif
@@ -3005,6 +3005,8 @@ void FVirtualTextureSystem::Update(FRDGBuilder& GraphBuilder, ERHIFeatureLevel::
 
 	TUniquePtr<FVirtualTextureUpdater> Updater = BeginUpdate(GraphBuilder, FeatureLevel, Scene, Settings);
 	EndUpdate(GraphBuilder, MoveTemp(Updater), FeatureLevel);
+
+	FinalizeRequests(GraphBuilder);
 }
 
 void FVirtualTextureSystem::ReleasePendingResources()
