Skip queuing cache invalidations for dynamic geometry when receiver mask is enabled

diff --git a/Engine/Source/Runtime/Renderer/Private/VirtualShadowMaps/VirtualShadowMapCacheManager.cpp b/Engine/Source/Runtime/Renderer/Private/VirtualShadowMaps/VirtualShadowMapCacheManager.cpp
--- a/Engine/Source/Runtime/Renderer/Private/VirtualShadowMaps/VirtualShadowMapCacheManager.cpp
+++ b/Engine/Source/Runtime/Renderer/Private/VirtualShadowMaps/VirtualShadowMapCacheManager.cpp
@@ -425,34 +425,57 @@ FVirtualShadowMapArrayCacheManager::FInvalidatingPrimitiveCollector::FInvalidati
 
 void FVirtualShadowMapArrayCacheManager::FInvalidatingPrimitiveCollector::AddPrimitivesToInvalidate()
 {	
-	for (auto& CacheEntry : Manager.CacheEntries)
+	const TBitArray<>& ManagerCachePrimitiveAsDynamic = Manager.CachePrimitiveAsDynamic;
+
+	for (auto& CacheEntryIt : Manager.CacheEntries)
 	{
-		for (const auto& SmCacheEntry : CacheEntry.Value->ShadowMapEntries)
-		{
-			const uint32 Payload = EncodeInstanceInvalidationPayload(SmCacheEntry.CurrentVirtualShadowMapId);
+		FVirtualShadowMapPerLightCacheEntry& CacheEntry = *CacheEntryIt.Value;
 
-			// Global invalidations
-			for (const FVirtualShadowMapInstanceRange& Range : Manager.ShadowInvalidatingInstancesImplementation.PrimitiveInstancesToInvalidate)
+		// Global invalidations
+		for (const FVirtualShadowMapInstanceRange& Range : Manager.ShadowInvalidatingInstancesImplementation.PrimitiveInstancesToInvalidate)
+		{
+			// If something doesn't have a valid persistent primitive index, it is treated as a dynamic primitive
+			// See GetCachePrimitiveAsDynamic
+			bool bCachePrimativeAsDynamic = true;
+			if (Range.PersistentPrimitiveIndex.IsValid())
 			{
-				Instances.Add(Range.InstanceSceneDataOffset, Range.NumInstanceSceneDataEntries, Payload);
-				if (Range.bMarkAsDynamic && Range.PersistentPrimitiveIndex.IsValid())
+				bCachePrimativeAsDynamic = ManagerCachePrimitiveAsDynamic[Range.PersistentPrimitiveIndex.Index];
+				if (Range.bMarkAsDynamic)
 				{
 					InvalidatedPrimitives[Range.PersistentPrimitiveIndex.Index] = true;
 				}
 			}
 
-			// Per-light invalidations
-			for (const FVirtualShadowMapInstanceRange& Range : CacheEntry.Value->PrimitiveInstancesToInvalidate)
+			AddInvalidation(
+				CacheEntry,
+				Range.InstanceSceneDataOffset,
+				Range.NumInstanceSceneDataEntries,
+				bCachePrimativeAsDynamic);
+		}
+
+		// Per-light invalidations
+		for (const FVirtualShadowMapInstanceRange& Range : CacheEntry.PrimitiveInstancesToInvalidate)
+		{
+			check(Range.PersistentPrimitiveIndex.IsValid());		// Should always be valid currently in this path
+
+			bool bCachePrimativeAsDynamic = true;
+			if (Range.PersistentPrimitiveIndex.IsValid())
 			{
-				Instances.Add(Range.InstanceSceneDataOffset, Range.NumInstanceSceneDataEntries, Payload);
-				check(Range.PersistentPrimitiveIndex.IsValid());		// Should always be valid currently in this path
-				if (Range.bMarkAsDynamic && Range.PersistentPrimitiveIndex.IsValid())
+				bCachePrimativeAsDynamic = ManagerCachePrimitiveAsDynamic[Range.PersistentPrimitiveIndex.Index];
+				if (Range.bMarkAsDynamic)
 				{
 					InvalidatedPrimitives[Range.PersistentPrimitiveIndex.Index] = true;
 				}
 			}
+
+			AddInvalidation(
+				CacheEntry,
+				Range.InstanceSceneDataOffset,
+				Range.NumInstanceSceneDataEntries,
+				bCachePrimativeAsDynamic);
 		}
-		CacheEntry.Value->PrimitiveInstancesToInvalidate.Reset();
+
+		CacheEntry.PrimitiveInstancesToInvalidate.Reset();
 	}
 	Manager.ShadowInvalidatingInstancesImplementation.PrimitiveInstancesToInvalidate.Reset();
 }
@@ -502,34 +525,52 @@ void FVirtualShadowMapArrayCacheManager::FInvalidatingPrimitiveCollector::AddInv
 	
 	const int32 NumInstanceSceneDataEntries = PrimitiveSceneInfo->GetNumInstanceSceneDataEntries();
 	const FBoxSphereBounds PrimitiveBounds = PrimitiveSceneInfo->Proxy->GetBounds();
-
 	const bool bCachePrimativeAsDynamic = Manager.CachePrimitiveAsDynamic[PersistentPrimitiveIndex.Index];
-
+	
 	for (auto& CacheEntry : Manager.CacheEntries)
 	{
-		// We don't need explicit invalidations for force invalidated/uncached lights
-		// We don't need explicit dynamic invalidations when using receiver mask
-		if (!CacheEntry.Value->IsUncached() &&
-			!(bCachePrimativeAsDynamic && CacheEntry.Value->ShouldUseReceiverMask())
-		   )
-		{
-			// Quick bounds overlap check to eliminate stuff that is too far away to affect a light
-			if (GVSMLightRadiusInvalidationCulling == 0 || CacheEntry.Value->AffectsBounds(PrimitiveBounds))
-			{
-				// Add item for each shadow map explicitly, inflates host data but improves load balancing,
-				for (const auto& SmCacheEntry : CacheEntry.Value->ShadowMapEntries)
-				{
-					Instances.Add(InstanceSceneDataOffset, NumInstanceSceneDataEntries, EncodeInstanceInvalidationPayload(SmCacheEntry.CurrentVirtualShadowMapId));
-				}
-			}
-			else
-			{
-				//UE_LOG(LogRenderer, Display, TEXT("VirtualShadowMapCacheManager: Skipped invalidation %s!"), *PrimitiveSceneInfo->GetOwnerActorNameOrLabelForDebuggingOnly());
-			}
-		}
+		AddInvalidation(
+			*CacheEntry.Value,
+			InstanceSceneDataOffset,
+			NumInstanceSceneDataEntries,
+			bCachePrimativeAsDynamic,
+			(GVSMLightRadiusInvalidationCulling != 0),
+			PrimitiveBounds);
 	}
 }
 
+void FVirtualShadowMapArrayCacheManager::FInvalidatingPrimitiveCollector::AddInvalidation(
+	const FVirtualShadowMapPerLightCacheEntry& CacheEntry,
+	int32 InstanceSceneDataOffset,
+	int32 NumInstanceSceneDataEntries,
+	bool bCachePrimitiveAsDynamic,
+	bool bLightRadiusCulling,
+	const FBoxSphereBounds& PrimitiveBounds)
+{
+	// We don't need explicit invalidations for force invalidated/uncached lights
+	if (CacheEntry.IsUncached())
+	{
+		return;
+	}
+	// We don't need explicit dynamic invalidations when using receiver mask
+	if (bCachePrimitiveAsDynamic && CacheEntry.ShouldUseReceiverMask())
+	{
+		return;
+	}
+	// Quick bounds overlap check to eliminate stuff that is too far away to affect a light
+	if (bLightRadiusCulling && !CacheEntry.AffectsBounds(PrimitiveBounds))
+	{
+		return;
+	}
+	
+	// Add item for each shadow map explicitly, inflates host data but improves load balancing,
+	for (const auto& SmCacheEntry : CacheEntry.ShadowMapEntries)
+	{
+		Instances.Add(InstanceSceneDataOffset, NumInstanceSceneDataEntries, EncodeInstanceInvalidationPayload(SmCacheEntry.CurrentVirtualShadowMapId));
+	}
+}
+
+
 
 FVirtualShadowMapFeedback::FVirtualShadowMapFeedback()
 {
diff --git a/Engine/Source/Runtime/Renderer/Private/VirtualShadowMaps/VirtualShadowMapCacheManager.h b/Engine/Source/Runtime/Renderer/Private/VirtualShadowMaps/VirtualShadowMapCacheManager.h
--- a/Engine/Source/Runtime/Renderer/Private/VirtualShadowMaps/VirtualShadowMapCacheManager.h
+++ b/Engine/Source/Runtime/Renderer/Private/VirtualShadowMaps/VirtualShadowMapCacheManager.h
@@ -395,6 +395,14 @@ public:
 
 		void AddInvalidation(FPrimitiveSceneInfo* PrimitiveSceneInfo, EInvalidationCause InvalidationCause);
 
+		void AddInvalidation(
+			const FVirtualShadowMapPerLightCacheEntry& CacheEntry,
+			int32 InstanceSceneDataOffset,
+			int32 NumInstanceSceneDataEntries,
+			bool bCachePrimitiveAsDynamic,
+			bool bLightRadiusCulling = false,
+			const FBoxSphereBounds& PrimitiveBounds = FBoxSphereBounds());
+
 		FScene& Scene;
 		FVirtualShadowMapArrayCacheManager& Manager;
 	};
