Avoid creating a separate simple renderer for VT. Uses the RDG blackboard system to avoid needing to plumb it all the way down. Allows access to the active scene renderer from the RDG Builder instance.

diff --git a/Engine/Source/Runtime/Renderer/Private/SceneRenderBuilder.cpp b/Engine/Source/Runtime/Renderer/Private/SceneRenderBuilder.cpp
--- a/Engine/Source/Runtime/Renderer/Private/SceneRenderBuilder.cpp
+++ b/Engine/Source/Runtime/Renderer/Private/SceneRenderBuilder.cpp
@@ -897,6 +897,7 @@ void FSceneRenderProcessor::Execute()
 				const FSceneRenderFunctionInputs FunctionInputs(&Renderer, SceneUpdateInputs ? &SceneUpdateInputs.GetValue() : nullptr, *RenderNode.Name, *RenderState.FullPath);
 
 				FRDGBuilder GraphBuilder(RHICmdList, RDG_EVENT_NAME("%s", FunctionInputs.FullPath), ERDGBuilderFlags::Parallel, Scene->GetShaderPlatform());
+				FSceneRendererBase::SetActiveInstance(GraphBuilder, &Renderer);
 
 			#if WITH_MGPU
 				if (Renderer.ViewFamily.bForceCopyCrossGPU)
diff --git a/Engine/Source/Runtime/Renderer/Private/SceneRendering.cpp b/Engine/Source/Runtime/Renderer/Private/SceneRendering.cpp
--- a/Engine/Source/Runtime/Renderer/Private/SceneRendering.cpp
+++ b/Engine/Source/Runtime/Renderer/Private/SceneRendering.cpp
@@ -2597,6 +2597,28 @@ TGlobalResource<FGlobalDynamicReadBuffer> FSceneRenderer::DynamicReadBufferForSh
 /*-----------------------------------------------------------------------------
 	FSceneRenderer
 -----------------------------------------------------------------------------*/
+
+struct FSceneUniformBufferBlackboardStruct
+{
+	FSceneRendererBase* SceneRenderer = nullptr;
+};
+
+RDG_REGISTER_BLACKBOARD_STRUCT(FSceneUniformBufferBlackboardStruct);
+
+void FSceneRendererBase::SetActiveInstance(FRDGBuilder& GraphBuilder, FSceneRendererBase* SceneRenderer)
+{
+	GraphBuilder.Blackboard.GetOrCreate<FSceneUniformBufferBlackboardStruct>().SceneRenderer = SceneRenderer;
+}
+
+FSceneRendererBase* FSceneRendererBase::GetActiveInstance(FRDGBuilder& GraphBuilder)
+{
+	if (const FSceneUniformBufferBlackboardStruct* Struct = GraphBuilder.Blackboard.Get<FSceneUniformBufferBlackboardStruct>())
+	{
+		return Struct->SceneRenderer;
+	}
+	return nullptr;
+}
+
 FSceneRenderer::FSceneRenderer(const FSceneViewFamily* InViewFamily, FHitProxyConsumer* HitProxyConsumer)
 :	FSceneRendererBase(*CheckPointer(CheckPointer(CheckPointer(InViewFamily)->Scene)->GetRenderScene()))
 ,	ViewFamily(*InViewFamily)
@@ -5273,8 +5295,8 @@ void FRendererModule::RenderPostResolvedSceneColorExtension(FRDGBuilder& GraphBu
 class FScenePrimitiveRenderingContext : public IScenePrimitiveRenderingContext
 {
 public:
-	FScenePrimitiveRenderingContext(FRDGBuilder& GraphBuilder, FScene& Scene, FSceneViewFamily* InViewFamily = nullptr) :
-		Renderer(),
+	FScenePrimitiveRenderingContext(FRDGBuilder& InGraphBuilder, FScene& Scene, FSceneViewFamily* InViewFamily = nullptr) :
+		GraphBuilder(InGraphBuilder),
 		GPUScene(Scene.GPUScene),
 		GPUSceneDynamicContext(GPUScene),
 		ViewFamily(InViewFamily)
@@ -5291,10 +5313,13 @@ public:
 
 		GPUScene.FillSceneUniformBuffer(GraphBuilder, Renderer.GetSceneUniforms());
 		Renderer.GetSceneExtensionsRenderers().UpdateSceneUniformBuffer(GraphBuilder, Renderer.GetSceneUniforms());
+
+		FSceneRendererBase::SetActiveInstance(GraphBuilder, &Renderer);
 	}
 
 	virtual ~FScenePrimitiveRenderingContext()
 	{
+		FSceneRendererBase::SetActiveInstance(GraphBuilder, nullptr);
 		GPUScene.EndRender();
 		if (ViewFamily)
 		{
@@ -5307,6 +5332,7 @@ public:
 		return &Renderer;
 	}
 
+	FRDGBuilder& GraphBuilder;
 	FSceneRendererBase Renderer;
 	FGPUScene& GPUScene;
 	FGPUSceneDynamicContext GPUSceneDynamicContext;
diff --git a/Engine/Source/Runtime/Renderer/Private/SceneRendering.h b/Engine/Source/Runtime/Renderer/Private/SceneRendering.h
--- a/Engine/Source/Runtime/Renderer/Private/SceneRendering.h
+++ b/Engine/Source/Runtime/Renderer/Private/SceneRendering.h
@@ -2000,11 +2000,15 @@ public:
 	FSceneRendererBase() {}
 	FSceneRendererBase(FScene& InScene) : Scene(&InScene) {}
 
+	// Helper functions to retrieve the currently active scene renderer from the GraphBuilder instance.
+	static void SetActiveInstance(FRDGBuilder& GraphBuilder, FSceneRendererBase* SceneRenderer);
+	static FSceneRendererBase* GetActiveInstance(FRDGBuilder& GraphBuilder);
+
 	// ISceneRenderer interface
 	const FSceneUniformBuffer& GetSceneUniforms() const final override { return SceneUniforms; }
 	FSceneUniformBuffer& GetSceneUniforms() final override { return SceneUniforms; }
 
-	TRDGUniformBufferRef<FSceneUniformParameters> GetSceneUniformBufferRef(FRDGBuilder& GraphBuilder) final
+	TRDGUniformBufferRef<FSceneUniformParameters> GetSceneUniformBufferRef(FRDGBuilder& GraphBuilder) override final
 	{
 		return GetSceneUniforms().GetBuffer(GraphBuilder);
 	}
diff --git a/Engine/Source/Runtime/Renderer/Private/VT/RuntimeVirtualTextureRender.cpp b/Engine/Source/Runtime/Renderer/Private/VT/RuntimeVirtualTextureRender.cpp
--- a/Engine/Source/Runtime/Renderer/Private/VT/RuntimeVirtualTextureRender.cpp
+++ b/Engine/Source/Runtime/Renderer/Private/VT/RuntimeVirtualTextureRender.cpp
@@ -2168,22 +2168,12 @@ namespace RuntimeVirtualTexture
 		check(InDesc.Scene != nullptr);
 		if (ensure(InDesc.Scene->GPUScene.IsRendering()))
 		{
-			// TODO: this should be replaced by piping through a reference to the scene renderer rather than just the scene, such that we can get at the already populated scene UB.
-			class FSimpleRVTRenderer : public FSceneRendererBase
-			{
-			public:
-				FSimpleRVTRenderer(FRDGBuilder& GraphBuilder, FRenderPageBatchDesc const& InDesc)
-					: FSceneRendererBase(*InDesc.Scene)
-				{
-					InitSceneExtensionsRenderers(FEngineShowFlags(ESFIM_Game));
+			FSceneRendererBase* SceneRenderer = FSceneRendererBase::GetActiveInstance(GraphBuilder);
+			checkf(SceneRenderer, TEXT("RenderPages needs to be called from within the context of a scene renderer."));
+			checkf(SceneRenderer->Scene == InDesc.Scene, TEXT("Provided scene doesn't match the parent scene renderer's scene."));
 
-					InDesc.Scene->GPUScene.FillSceneUniformBuffer(GraphBuilder, GetSceneUniforms());
-					GetSceneExtensionsRenderers().UpdateSceneUniformBuffer(GraphBuilder, GetSceneUniforms());
-			}
-			};
-			FSimpleRVTRenderer SimpleRenderer(GraphBuilder, InDesc);
 			const bool bAllowCachedMeshDrawCommands = true;
-			RenderPagesInternal(GraphBuilder, InDesc, &SimpleRenderer, bAllowCachedMeshDrawCommands);
+			RenderPagesInternal(GraphBuilder, InDesc, SceneRenderer, bAllowCachedMeshDrawCommands);
 		}
 	}
 
