[SSS] Reduce setup pass cost.
* Do a prepass based on SSP shading model only to create a memory coherent (zorder) list of tiles.
* Run SSS:Setup on the tiles with subsurface.
* If everything runs on separable, use the list of tiles built in the prepass for all following passes, otherwise, build additional list as usual for separable and burley.

#tests SubsurfaceProfile test in EngineTest pass on PC
#rnx

diff --git a/Engine/Shaders/Private/PostProcessSubsurface.usf b/Engine/Shaders/Private/PostProcessSubsurface.usf
--- a/Engine/Shaders/Private/PostProcessSubsurface.usf
+++ b/Engine/Shaders/Private/PostProcessSubsurface.usf
@@ -8,7 +8,7 @@
 #include "Substrate/Substrate.ush"
 #include "PostProcessCommon.ush"
 #include "DeferredShadingCommon.ush"
-#include "SubsurfaceProfileCommon.ush"
+#include "PostProcessSubsurfaceCommon.ush"
 
 // for VisualizeSSS
 #include "MiniFontCommon.ush"
@@ -47,12 +47,6 @@ SCREEN_PASS_TEXTURE_VIEWPORT(Output)
 
 RWTexture2D<uint> ProfileIdTexture;
 
-// the same to FSubsurfaceTiles::ETileType
-#define SUBSURFACE_TILE_TYPE_AFIS		0
-#define SUBSURFACE_TILE_TYPE_SEPARABLE	1
-#define SUBSURFACE_TILE_TYPE_ALL		2
-#define SUBSURFACE_TILE_TYPE_COUNT		3
-
 #define SUBSURFACE_PASS_ONE 0
 #define SUBSURFACE_PASS_TWO 1
 
@@ -648,7 +642,6 @@ bool ShouldGenerateMipmaps(uint TileType)
 // Indirect dispatch parameters
 RWBuffer<uint> RWSeparableGroupBuffer;
 RWBuffer<uint> RWBurleyGroupBuffer;
-RWBuffer<uint> RWAllGroupBuffer;
 RWBuffer<uint> RWIndirectDispatchArgsBuffer;
 RWBuffer<uint> RWIndirectDrawArgsBuffer;
 RWBuffer<uint> RWTileTypeCountBuffer;
@@ -679,7 +672,7 @@ void FillMipsConditionBufferCS()
 [numthreads(8, 1, 1)]
 void BuildIndirectDispatchArgsCS(uint2 DispatchThreadId : SV_DispatchThreadID)
 {
-	if (DispatchThreadId.y == 0 && DispatchThreadId.x < SUBSURFACE_TILE_TYPE_COUNT)
+	if (DispatchThreadId.y == 0 && DispatchThreadId.x < SUBSURFACE_TILE_TYPE_COUNT && DispatchThreadId.x != SUBSURFACE_TILE_TYPE_ALL)
 	{
 		uint TileType = DispatchThreadId.x;
 		const uint TileCount = TileTypeCountBuffer[TileType];
@@ -696,16 +689,6 @@ void BuildIndirectDispatchArgsCS(uint2 DispatchThreadId : SV_DispatchThreadID)
 	}
 }
 
-#define TYPE_SEPARABLE 0x1
-#define TYPE_BURLEY    0x2
-#define TYPE_IN_VIEWPORT 0x4
-
-// 0, no subsurface
-// 1, Separable
-// 2, Burley
-// 3, Both
-groupshared uint SubsurfaceTypeFlag;
-
 uint GetSubsurfaceType(uint ProfileId)
 {
 	return GetSubsurfaceProfileUseBurley(ProfileId) ? TYPE_BURLEY : TYPE_SEPARABLE;
@@ -723,8 +706,7 @@ void AddSeparableGroup(uint2 GroupXY)
 		return;
 	}
 
-	RWSeparableGroupBuffer[2 * IndexToStore + 0] = GroupXY.x;
-	RWSeparableGroupBuffer[2 * IndexToStore + 1] = GroupXY.y;
+	RWSeparableGroupBuffer[IndexToStore] = PackTileCoord12bits(GroupXY);
 }
 
 void AddBurleyGroup(uint2 GroupXY)
@@ -739,29 +721,21 @@ void AddBurleyGroup(uint2 GroupXY)
 		return;
 	}
 
-	RWBurleyGroupBuffer[2* IndexToStore + 0] = GroupXY.x;
-	RWBurleyGroupBuffer[2* IndexToStore + 1] = GroupXY.y;
+	RWBurleyGroupBuffer[IndexToStore] = PackTileCoord12bits(GroupXY);
 }
 
-void AddAllGroup(uint2 GroupXY)
+void AddSubsurfaceComputeGroup(uint GI, uint Type, uint2 GroupXY)
 {
-	//buffer counter is stored in the .Count in the first element.
-	uint IndexToStore = 0;
-
-	InterlockedAdd(RWTileTypeCountBuffer[SUBSURFACE_TILE_TYPE_ALL], 1, IndexToStore);
-
-	if (IndexToStore >= SubsurfaceUniformParameters.MaxGroupCount)
+#if COMPILER_SUPPORTS_WAVE_VOTE && COMPILER_SUPPORTS_WAVE_BIT_ORAND
+	const uint AnySSSInViewInWave = WaveActiveBitOr(Type);
+	if (WaveIsFirstLane() && AnySSSInViewInWave)
 	{
-		return;
+		InterlockedOr(SubsurfaceTypeFlag, AnySSSInViewInWave);
 	}
 
-	RWAllGroupBuffer[2 * IndexToStore + 0] = GroupXY.x;
-	RWAllGroupBuffer[2 * IndexToStore + 1] = GroupXY.y;
-}
-
-void AddSubsurfaceComputeGroup(uint GI, uint Type, uint2 GroupXY)
-{
+#else
 	InterlockedOr(SubsurfaceTypeFlag, Type);
+#endif
 
 	GroupMemoryBarrierWithGroupSync();
 
@@ -774,22 +748,14 @@ void AddSubsurfaceComputeGroup(uint GI, uint Type, uint2 GroupXY)
 	{
 		AddBurleyGroup(GroupXY);
 	}
-
-	if (GI == 0 && ((SubsurfaceTypeFlag & TYPE_IN_VIEWPORT) &&
-				   ((SubsurfaceTypeFlag &	TYPE_SEPARABLE) ||
-					(SubsurfaceTypeFlag &	   TYPE_BURLEY))))
-	{
-		AddAllGroup(GroupXY);
-	}
 }
 
 RWTexture2D<float4> SetupTexture;
 RWTexture2D<float4> ProfileEdgeMask;
 
-[numthreads(THREAD_SIZE_X, THREAD_SIZE_Y, 1)]
-void SetupIndirectCS(uint3 DT_ID : SV_DispatchThreadID,
-	uint3 G_ID : SV_GroupID,
-	uint3 GT_ID : SV_GroupThreadID,
+[numthreads(THREAD_SIZE_X*THREAD_SIZE_Y, 1, 1)]
+void SetupIndirectCS(	
+	uint3 IndirectG_ID : SV_GroupID,
 	uint  GI : SV_GroupIndex)
 {
 	// initalize the group shared variable
@@ -799,6 +765,9 @@ void SetupIndirectCS(uint3 DT_ID : SV_DispatchThreadID,
 	}
 	GroupMemoryBarrierWithGroupSync();
 
+	uint2 G_ID = UnpackTileCoord12bits(GroupBuffer[IndirectG_ID.x]);
+	uint2 DT_ID = G_ID * float2(THREAD_SIZE_X, THREAD_SIZE_Y) + float2(GI % THREAD_SIZE_X, GI / THREAD_SIZE_X);
+
 	uint2 Pos = DT_ID.xy + Output_ViewportMin;
 	float2 BufferUV = ConvertGridPos2UV(Pos);
 	uint Type = 0;
@@ -912,7 +881,9 @@ void SetupIndirectCS(uint3 DT_ID : SV_DispatchThreadID,
 #endif
 	}
 
+#if !SUBSURFACE_FORCE_SEPARABLE
 	AddSubsurfaceComputeGroup(GI, Type, G_ID.xy);
+#endif
 }
 
 [numthreads(THREAD_SIZE_X*THREAD_SIZE_Y, 1, 1)]
@@ -921,8 +892,7 @@ void MainIndirectDispatchCS(
 	uint  GI : SV_GroupIndex
 ) 
 {
-	// The first element is used for counting
-	uint2 G_ID = uint2(GroupBuffer[2 * IndirectG_ID.x + 0], GroupBuffer[2 * IndirectG_ID.x + 1]);
+	uint2 G_ID = UnpackTileCoord12bits(GroupBuffer[IndirectG_ID.x]);
 	uint2 DT_ID = G_ID * float2(THREAD_SIZE_X, THREAD_SIZE_Y) + float2(GI % THREAD_SIZE_X, GI / THREAD_SIZE_X);
 
 
diff --git a/Engine/Shaders/Private/PostProcessSubsurfaceCommon.ush b/Engine/Shaders/Private/PostProcessSubsurfaceCommon.ush
new file mode 100644
--- /dev/null
+++ b/Engine/Shaders/Private/PostProcessSubsurfaceCommon.ush
@@ -0,0 +1,22 @@
+// Copyright Epic Games, Inc. All Rights Reserved.
+
+/*=============================================================================
+	PostProcessSubsurfaceTile.usf: Screenspace subsurface scattering common shaders.
+=============================================================================*/
+
+#pragma once
+
+#include "Common.ush"
+#include "SubsurfaceProfileCommon.ush"
+
+// the same to FSubsurfaceTiles::ETileType
+#define SUBSURFACE_TILE_TYPE_AFIS		0
+#define SUBSURFACE_TILE_TYPE_SEPARABLE	1
+#define SUBSURFACE_TILE_TYPE_ALL		2
+#define SUBSURFACE_TILE_TYPE_COUNT		3
+
+#define TYPE_SEPARABLE   0x1
+#define TYPE_BURLEY      0x2
+#define TYPE_IN_VIEWPORT 0x4
+
+groupshared uint SubsurfaceTypeFlag;
\ No newline at end of file
diff --git a/Engine/Shaders/Private/PostProcessSubsurfaceTile.usf b/Engine/Shaders/Private/PostProcessSubsurfaceTile.usf
--- a/Engine/Shaders/Private/PostProcessSubsurfaceTile.usf
+++ b/Engine/Shaders/Private/PostProcessSubsurfaceTile.usf
@@ -17,6 +17,14 @@
 #define SUBSURFACE_COMPUTE 0
 #endif
 
+#ifndef TILE_CATERGORISATION_SHADER
+#define TILE_CATERGORISATION_SHADER 0
+#endif
+
+#ifndef SUBSURFACE_HALF_RES
+#define SUBSURFACE_HALF_RES 0
+#endif
+
 #define LINE_TEXEL_SIZE 1
 
 #if SUBSURFACE_COMPUTE
@@ -85,7 +93,7 @@ void MainVS(
 	out FScreenVertexOutput Out)
 {
 	Out = (FScreenVertexOutput)0;
-	const uint2 TileCoord = uint2(TileDataBuffer[InInstanceId * 2 + 0], TileDataBuffer[InInstanceId * 2 + 1]);
+	const uint2 TileCoord = UnpackTileCoord12bits(TileDataBuffer[InInstanceId]);
 
 	uint2 TileVertex = TileCoord * SUBSURFACE_TILE_SIZE;
 	TileVertex.x += InVertexId == 1 || InVertexId == 2 || InVertexId == 4 ? SUBSURFACE_TILE_SIZE : 0;
@@ -102,4 +110,225 @@ void SubsurfaceTileFallbackScreenPassVS(
 	out FScreenVertexOutput Out)
 {
 	DrawRectangle(InPosition, InTexCoord, Out.Position, Out.UV);
-}
\ No newline at end of file
+}
+
+//Code adapted from DefaultSSRTiles.usf
+#if TILE_CATERGORISATION_SHADER
+
+#include "Substrate/Substrate.ush"
+#include "Substrate/SubstrateEvaluation.ush"
+
+#include "PostProcessCommon.ush"
+#include "DeferredShadingCommon.ush"
+#include "PostProcessSubsurfaceCommon.ush"
+#include "ScreenPass.ush"
+
+SCREEN_PASS_TEXTURE_VIEWPORT(SubsurfaceInput0)
+SCREEN_PASS_TEXTURE_VIEWPORT(Output)
+
+groupshared bool ContainsSSS[SUBSURFACE_TILE_SIZE * SUBSURFACE_TILE_SIZE];
+
+int MaxGroupCount;
+int2 TiledViewRes;
+RWStructuredBuffer<uint> TileMaskBufferOut;
+
+// Only access shading model GBuffer
+float HasSubsurfaceForOnePixel(float2 UVSceneColor)
+{
+#if SUBTRATE_GBUFFER_FORMAT==1
+	const FSubstrateSubsurfaceHeader SSSHeader = LoadSubstrateSSSHeader(UVSceneColor);
+	const bool bHasSSSDiffusion = SubstrateSubSurfaceHeaderGetUseDiffusion(SSSHeader);
+#else
+	FScreenSpaceData ScreenSpaceData = GetScreenSpaceData(UVSceneColor);
+	const bool bHasSSSDiffusion = UseSubsurfaceProfile(ScreenSpaceData.GBuffer.ShadingModelID);
+#endif
+
+	return select(bHasSSSDiffusion, 1, 0);
+}
+
+// Simpler logic from the SSS::Setup pass
+[numthreads(SUBSURFACE_TILE_SIZE, SUBSURFACE_TILE_SIZE, 1)]
+void SSSTileCategorisationMarkCS(uint3 DT_ID : SV_DispatchThreadID, 
+	uint3 G_ID : SV_GroupID, 
+	uint3 GT_ID : SV_GroupThreadID,
+	uint  GI : SV_GroupIndex)
+{
+	if (GI == 0)
+	{
+		SubsurfaceTypeFlag = 0;
+	}
+	GroupMemoryBarrierWithGroupSync();
+	uint2 Pos = DT_ID.xy + Output_ViewportMin;
+	float2 BufferUV = Output_ExtentInverse * (float2(Pos) + 0.5f);
+	uint Type = 0;
+	
+#if SUBSURFACE_HALF_RES
+	float A = HasSubsurfaceForOnePixel(min(BufferUV + float2(-0.5, 0.5f) * SubsurfaceInput0_ExtentInverse, SubsurfaceInput0_UVViewportBilinearMax));
+	float B = HasSubsurfaceForOnePixel(min(BufferUV + float2( 0.5, 0.5f) * SubsurfaceInput0_ExtentInverse, SubsurfaceInput0_UVViewportBilinearMax));
+	float C = HasSubsurfaceForOnePixel(min(BufferUV + float2( 0.5,-0.5f) * SubsurfaceInput0_ExtentInverse, SubsurfaceInput0_UVViewportBilinearMax));
+	float D = HasSubsurfaceForOnePixel(min(BufferUV + float2(-0.5,-0.5f) * SubsurfaceInput0_ExtentInverse, SubsurfaceInput0_UVViewportBilinearMax));
+	float Sum = (A + B) + (C + D);
+	// Each alpha value of A/B/C/D should be either 1.0f or 0.0, so a check of 0.5 should be safe
+	bool bHasSubsurface = (Sum >= .50f);
+#else
+	bool bHasSubsurface = bool(HasSubsurfaceForOnePixel(BufferUV));
+#endif
+	bHasSubsurface = bHasSubsurface && all(Pos < Output_ViewportMax);
+	bool bWriteTile = false;
+
+#if COMPILER_SUPPORTS_WAVE_VOTE
+	const bool bAnyActiveSSSPixelsInWave = WaveActiveAnyTrue(bHasSubsurface);
+	if (WaveIsFirstLane() && bAnyActiveSSSPixelsInWave)
+	{
+		uint IndexToStore = 0;
+		InterlockedAdd(SubsurfaceTypeFlag, 1, IndexToStore);
+	}
+	GroupMemoryBarrierWithGroupSync();
+	if (GI == 0)
+	{
+		bWriteTile = SubsurfaceTypeFlag > 0;
+	}
+#else
+
+	const int LinearIndex = GT_ID.y * SUBSURFACE_TILE_SIZE + GT_ID.x;
+	ContainsSSS[LinearIndex] = bHasSubsurface;
+	GroupMemoryBarrierWithGroupSync();
+
+	if (LinearIndex < 32) // 8*8 = 64 elements to merge
+	{
+		ContainsSSS[LinearIndex] = ContainsSSS[LinearIndex] || ContainsSSS[LinearIndex + 32];
+	}
+	GroupMemoryBarrierWithGroupSync();
+	if (LinearIndex < 16)
+	{
+		ContainsSSS[LinearIndex] = ContainsSSS[LinearIndex] || ContainsSSS[LinearIndex + 16];
+	}
+	GroupMemoryBarrierWithGroupSync();
+
+	// The smallest wave size is 16 on Intel hardware. So now we can do simple math operations without group sync.
+	// EDIT: for some reason group sync is needed until the end, otherwise some pixels are missing...
+
+	if (LinearIndex < 8)
+	{
+		ContainsSSS[LinearIndex] = ContainsSSS[LinearIndex] || ContainsSSS[LinearIndex + 8];
+	}
+	GroupMemoryBarrierWithGroupSync();
+	if (LinearIndex < 4)
+	{
+		ContainsSSS[LinearIndex] = ContainsSSS[LinearIndex] || ContainsSSS[LinearIndex + 4];
+	}
+	GroupMemoryBarrierWithGroupSync();
+	if (LinearIndex < 2)
+	{
+		ContainsSSS[LinearIndex] = ContainsSSS[LinearIndex] || ContainsSSS[LinearIndex + 2];
+	}
+	GroupMemoryBarrierWithGroupSync();
+	if (LinearIndex < 1)
+	{
+		bWriteTile = ContainsSSS[LinearIndex] || ContainsSSS[LinearIndex + 1];
+	}
+#endif
+
+	if (bWriteTile)
+	{
+		// Set bit to indicate tile is occupied.
+		uint MaskLinearIndex = TiledViewRes.x * G_ID.y + G_ID.x;
+		InterlockedOr(TileMaskBufferOut[MaskLinearIndex / 32U], 1U << (MaskLinearIndex % 32U));
+	}
+}
+
+int VertexCountPerInstanceIndirect;
+
+StructuredBuffer<uint> TileMaskBuffer;
+
+RWBuffer<uint> RWSSSTileListData;
+RWBuffer<uint> RWTileTypeCountBuffer;
+RWBuffer<uint> RWIndirectDrawArgsBuffer;
+RWBuffer<uint> RWIndirectDispatchArgsBuffer;
+RWBuffer<uint> RWSSSTileListDataBuffer;
+
+groupshared uint SharedNumTiles;
+groupshared uint SharedTileData[SUBSURFACE_TILE_SIZE * SUBSURFACE_TILE_SIZE];
+groupshared uint SharedGlobalTileOffset;
+
+/**
+* Every group checks 64 tiles and builds a spatially coherent compacted list of SSS tiles
+*/
+[numthreads(SUBSURFACE_TILE_SIZE, SUBSURFACE_TILE_SIZE, 1)]
+void SSSTileClassificationBuildListsCS(
+	uint2 GroupId : SV_GroupID,
+	uint2 DispatchThreadId : SV_DispatchThreadID,
+	uint2 GroupThreadId : SV_GroupThreadID)
+{
+	const uint LinearThreadIndex = GroupThreadId.y * SUBSURFACE_TILE_SIZE + GroupThreadId.x;
+	const uint TileType = SUBSURFACE_TILE_TYPE_ALL;// Only interested in generating the SSS tiles instead of differentiating the type
+
+	if (all(DispatchThreadId == 0))
+	{
+		const uint TileCount = 0; // Instance count will be accumulated later with locked add.
+		
+		// Indirect draw
+		RWIndirectDrawArgsBuffer[TileType * DRAW_INDIRECT_UINT_COUNT + 0] = VertexCountPerInstanceIndirect;		// VertexCountPerInstance
+		RWIndirectDrawArgsBuffer[TileType * DRAW_INDIRECT_UINT_COUNT + 1] = TileCount;  // InstanceCount
+		RWIndirectDrawArgsBuffer[TileType * DRAW_INDIRECT_UINT_COUNT + 2] = 0;  // StartVertexLocation
+		RWIndirectDrawArgsBuffer[TileType * DRAW_INDIRECT_UINT_COUNT + 3] = 0;  // StartInstanceLocation
+
+		// Indirect dispatch initialize other dimensions
+		const uint IndirectDispatchSize = TileCount;
+		WriteDispatchIndirectArgs(RWIndirectDispatchArgsBuffer, TileType, IndirectDispatchSize, 1, 1);
+	}
+
+	if (LinearThreadIndex == 0)
+	{
+		SharedNumTiles = 0;
+	}
+
+	GroupMemoryBarrierWithGroupSync();
+
+	if (LinearThreadIndex == 0)
+	{
+		for (uint LocalTileIndex = 0; LocalTileIndex < SUBSURFACE_TILE_SIZE * SUBSURFACE_TILE_SIZE; ++LocalTileIndex)
+		{
+			// ZOrder tiles to maximize screen locality after converting to 1d for compaction
+			// The tile locality ultimately affects gathering from neighbor tiles.
+			uint2 ThreadOffset = ZOrder2D(LocalTileIndex, log2(SUBSURFACE_TILE_SIZE));
+			uint2 TileCoordinate = GroupId * SUBSURFACE_TILE_SIZE + ThreadOffset;
+
+			if (all(TileCoordinate < TiledViewRes))
+			{
+				uint MaskLinearIndex = TiledViewRes.x * TileCoordinate.y + TileCoordinate.x;
+				uint Mask = 1u << (MaskLinearIndex % 32u);
+				bool bTileUsed = (TileMaskBuffer[MaskLinearIndex / 32u] & Mask) != 0;
+				if (bTileUsed)
+				{
+					uint TileOffset = SharedNumTiles;
+					SharedTileData[TileOffset] = PackTileCoord12bits(TileCoordinate);
+					SharedNumTiles = TileOffset + 1;
+				}
+			}
+		}
+	}
+
+	GroupMemoryBarrierWithGroupSync();
+
+	// Allocate space in the tile list
+	if (LinearThreadIndex == 0 && SharedNumTiles > 0)
+	{
+		// Still need to track the total number of tiles ?.
+		uint IndexToStore = 0;
+		InterlockedAdd(RWTileTypeCountBuffer[TileType], SharedNumTiles, IndexToStore);
+
+		InterlockedAdd(RWIndirectDrawArgsBuffer[TileType * DRAW_INDIRECT_UINT_COUNT + 1], SharedNumTiles, SharedGlobalTileOffset);
+		InterlockedAdd(RWIndirectDispatchArgsBuffer[TileType * DISPATCH_INDIRECT_UINT_COUNT + 0], SharedNumTiles);
+	}
+
+	GroupMemoryBarrierWithGroupSync();
+
+	// Write out tiles
+	if (LinearThreadIndex < SharedNumTiles)
+	{
+		RWSSSTileListDataBuffer[SharedGlobalTileOffset + LinearThreadIndex] = SharedTileData[LinearThreadIndex];
+	}
+}
+
+#endif
\ No newline at end of file
diff --git a/Engine/Source/Runtime/Renderer/Private/PostProcess/PostProcessSubsurface.cpp b/Engine/Source/Runtime/Renderer/Private/PostProcess/PostProcessSubsurface.cpp
--- a/Engine/Source/Runtime/Renderer/Private/PostProcess/PostProcessSubsurface.cpp
+++ b/Engine/Source/Runtime/Renderer/Private/PostProcess/PostProcessSubsurface.cpp
@@ -535,12 +535,14 @@ public:
 		SHADER_PARAMETER_STRUCT(FSubsurfaceInput, SubsurfaceInput0)
 		SHADER_PARAMETER_SAMPLER(SamplerState, SubsurfaceSampler0)
 		SHADER_PARAMETER_RDG_TEXTURE_UAV(RWTexture2D, SetupTexture)
+		SHADER_PARAMETER_RDG_BUFFER_SRV(Buffer<uint>, GroupBuffer)
 		SHADER_PARAMETER_RDG_BUFFER_UAV(RWBuffer<uint>, RWSeparableGroupBuffer)
 		SHADER_PARAMETER_RDG_BUFFER_UAV(RWBuffer<uint>, RWBurleyGroupBuffer)
 		SHADER_PARAMETER_RDG_BUFFER_UAV(RWBuffer<uint>, RWAllGroupBuffer)
 		SHADER_PARAMETER_RDG_BUFFER_UAV(RWBuffer<uint>, RWTileTypeCountBuffer)
 		SHADER_PARAMETER_RDG_TEXTURE_UAV(RWTexture2D, ProfileIdTexture)
 		SHADER_PARAMETER_STRUCT_REF(FSubsurfaceUniformParameters, SubsurfaceUniformParameters)
+		RDG_BUFFER_ACCESS(IndirectDispatchArgsBuffer, ERHIAccess::IndirectArgs)
 	END_SHADER_PARAMETER_STRUCT()
 	
 	// Controls the quality of lighting reconstruction.
@@ -748,8 +750,13 @@ public:
 				static_cast<int32>(FSubsurfaceIndirectDispatchCS::EBilateralFilterKernelFunctionType::DepthAndNormal)));
 	}
 
-	static FSubsurfaceTiles::ETileType ToSubsurfaceTileType(ESubsurfaceType Type)
+	static FSubsurfaceTiles::ETileType ToSubsurfaceTileType(ESubsurfaceType Type, bool bForceSeparable)
 	{
+		if (bForceSeparable)
+		{
+			return FSubsurfaceTiles::ETileType::All;
+		}
+
 		FSubsurfaceTiles::ETileType TileType = FSubsurfaceTiles::ETileType::AFIS;
 
 		switch (Type)
@@ -1097,44 +1104,7 @@ void AddSubsurfaceViewPass(
 
 		FRDGTextureRef VelocityTexture = GetIfProduced(SceneTextures.Velocity, SystemTextures.Black);
 		
-		// Initialize tile buffers.
-		{
-			check(FSubsurfaceTiles::TilePerThread_GroupSize == 64); // If this value change, we need to update the shaders using 
-			check(FSubsurfaceTiles::TileSize == 8); // only size supported for now
-
-			Tiles.TileDimension = FIntPoint::DivideAndRoundUp(SubsurfaceViewport.Extent, FSubsurfaceTiles::TileSize);
-			Tiles.TileCount = Tiles.TileDimension.X * Tiles.TileDimension.Y;
-			Tiles.bRectPrimitive = GRHISupportsRectTopology;
-			Tiles.TileTypeCountBuffer = GraphBuilder.CreateBuffer(FRDGBufferDesc::CreateBufferDesc(sizeof(uint32), FSubsurfaceTiles::TileTypeCount), TEXT("Subsurface.TileCountBuffer"));
-			Tiles.TileIndirectDispatchBuffer = 
-				GraphBuilder.CreateBuffer(FRDGBufferDesc::CreateIndirectDesc<FRHIDispatchIndirectParameters>(FSubsurfaceTiles::TileTypeCount), TEXT("Subsurface.TileIndirectDispatchBuffer"));
-			Tiles.TileIndirectDrawBuffer =
-				GraphBuilder.CreateBuffer(FRDGBufferDesc::CreateIndirectDesc<FRHIDrawIndirectParameters>(FSubsurfaceTiles::TileTypeCount), TEXT("Subsurface.TileIndirectDrawBuffer"));
-
-			Tiles.TileDataBuffer[ToIndex(FSubsurfaceTiles::ETileType::All)] = 
-				GraphBuilder.CreateBuffer(FRDGBufferDesc::CreateBufferDesc(sizeof(uint32), 2 * Tiles.TileCount), TEXT("Subsurface.TileDataBuffer(All)"));
-			Tiles.TileDataBuffer[ToIndex(FSubsurfaceTiles::ETileType::AFIS)] = 
-				GraphBuilder.CreateBuffer(FRDGBufferDesc::CreateBufferDesc(sizeof(uint32), 2 * Tiles.TileCount), TEXT("Subsurface.TileDataBuffer(AFIS)"));
-			Tiles.TileDataBuffer[ToIndex(FSubsurfaceTiles::ETileType::SEPARABLE)] = 
-				GraphBuilder.CreateBuffer(FRDGBufferDesc::CreateBufferDesc(sizeof(uint32), 2 * Tiles.TileCount), TEXT("Subsurface.TileDataBuffer(SEPARABLE)"));
-
-			for (uint32 BufferIt = 0; BufferIt < FSubsurfaceTiles::TileTypeCount; ++BufferIt)
-			{
-				if (Tiles.TileDataBuffer[BufferIt])
-				{
-					Tiles.TileDataSRV[BufferIt] = GraphBuilder.CreateSRV(Tiles.TileDataBuffer[BufferIt], PF_R32_UINT);
-				}
-			}
-		}
-
-		// Initialize the group counters
-		{
-			typedef FSubsurfaceInitValueBufferCS SHADER;
-			TShaderMapRef<SHADER> ComputeShader(View.ShaderMap);
-			SHADER::FParameters* PassParameters = GraphBuilder.AllocParameters<SHADER::FParameters>();
-			PassParameters->RWTileTypeCountBuffer = GraphBuilder.CreateUAV(Tiles.TileTypeCountBuffer, EPixelFormat::PF_R32_UINT);
-			FComputeShaderUtils::AddPass(GraphBuilder, RDG_EVENT_NAME("SSS::InitGroupCounter"), ComputeShader, PassParameters, FIntVector(1, 1, 1));
-		}
+		Tiles = ClassifySSSTiles(GraphBuilder, View, SceneTextures, SceneViewportParameters, SubsurfaceViewportParameters, bHalfRes);
 
 		FSubsurfaceUniformRef UniformBuffer = CreateUniformBuffer(View, Tiles.TileCount);
 
@@ -1155,8 +1125,9 @@ void AddSubsurfaceViewPass(
 			}
 			PassParameters->RWBurleyGroupBuffer = GraphBuilder.CreateUAV(Tiles.TileDataBuffer[ToIndex(FSubsurfaceTiles::ETileType::AFIS)], EPixelFormat::PF_R32_UINT);
 			PassParameters->RWSeparableGroupBuffer = GraphBuilder.CreateUAV(Tiles.TileDataBuffer[ToIndex(FSubsurfaceTiles::ETileType::SEPARABLE)], EPixelFormat::PF_R32_UINT);
-			PassParameters->RWAllGroupBuffer = GraphBuilder.CreateUAV(Tiles.TileDataBuffer[ToIndex(FSubsurfaceTiles::ETileType::All)], EPixelFormat::PF_R32_UINT);
 			PassParameters->RWTileTypeCountBuffer = GraphBuilder.CreateUAV(Tiles.TileTypeCountBuffer, EPixelFormat::PF_R32_UINT);
+			PassParameters->IndirectDispatchArgsBuffer = Tiles.TileIndirectDispatchBuffer;
+			PassParameters->GroupBuffer = Tiles.GetTileBufferSRV(FSubsurfaceTiles::ETileType::All);
 
 			PassParameters->SubsurfaceUniformParameters = UniformBuffer;
 			PassParameters->Substrate = Substrate::BindSubstrateGlobalUniformParameters(View);
@@ -1182,11 +1153,15 @@ void AddSubsurfaceViewPass(
 					SubsurfaceViewport.Extent.Y),
 				ComputeShader,
 				PassParameters,
-				FComputeShaderUtils::GetGroupCount(SubsurfaceViewport.Extent, FSubsurfaceTiles::TileSize));
+				PassParameters->IndirectDispatchArgsBuffer, 
+				Tiles.GetIndirectDispatchArgOffset(FSubsurfaceTiles::ETileType::All));
 		}
 
-		// Setup the indirect dispatch & draw arguments.
+	
+		if (!bForceRunningInSeparable)
 		{
+			// Setup the indirect dispatch & draw arguments if not running in separable mode.
+			{
 				typedef FSubsurfaceBuildIndirectDispatchArgsCS ARGSETUPSHADER;
 				ARGSETUPSHADER::FParameters* PassParameters = GraphBuilder.AllocParameters<ARGSETUPSHADER::FParameters>();
 				PassParameters->SubsurfaceUniformParameters = UniformBuffer;
@@ -1195,33 +1170,34 @@ void AddSubsurfaceViewPass(
 				PassParameters->TileTypeCountBuffer = GraphBuilder.CreateSRV(Tiles.TileTypeCountBuffer, EPixelFormat::PF_R32_UINT);
 				TShaderMapRef<ARGSETUPSHADER> ComputeShader(View.ShaderMap);
 				FComputeShaderUtils::AddPass(GraphBuilder, FRDGEventName(TEXT("SSS::BuildIndirectArgs(Dispatch & Draw)")), ComputeShader, PassParameters, FIntVector(1, 1, 1));
-		}
-
-		// In half resolution, only Separable is used. We do not need this mipmap.
-		if (!bForceRunningInSeparable)
-		{
-			FRDGBufferRef MipsConditionBuffer = GraphBuilder.CreateBuffer(
-				FRDGBufferDesc::CreateBufferDesc(sizeof(uint32), FSubsurfaceTiles::TileTypeCount), TEXT("Subsurface.MipsConditionBuffer"));
-			const int32 Offset = ToIndex(FSubsurfaceTiles::ETileType::AFIS);
+			}
 
+			// Separable does not need this mipmap.
 			{
-				typedef FFillMipsConditionBufferCS SHADER;
-				TShaderMapRef<SHADER> ComputeShader(View.ShaderMap);
-				SHADER::FParameters* PassParameters = GraphBuilder.AllocParameters<SHADER::FParameters>();
-				PassParameters->MinGenerateMipsTileCount = MinGenerateMipsTileCount;
-				PassParameters->Offset = Offset;
-				PassParameters->TileTypeCountBuffer = GraphBuilder.CreateSRV(Tiles.TileTypeCountBuffer, EPixelFormat::PF_R32_UINT);
-				PassParameters->RWMipsConditionBuffer = GraphBuilder.CreateUAV(MipsConditionBuffer, EPixelFormat::PF_R32_UINT);
+				FRDGBufferRef MipsConditionBuffer = GraphBuilder.CreateBuffer(
+					FRDGBufferDesc::CreateBufferDesc(sizeof(uint32), FSubsurfaceTiles::TileTypeCount), TEXT("Subsurface.MipsConditionBuffer"));
+				const int32 Offset = ToIndex(FSubsurfaceTiles::ETileType::AFIS);
 
-				FComputeShaderUtils::AddPass(GraphBuilder, RDG_EVENT_NAME("SSS::FillMipsConditionBufferCS"), ComputeShader, PassParameters, FIntVector(1, 1, 1));
-			}
+				{
+					typedef FFillMipsConditionBufferCS SHADER;
+					TShaderMapRef<SHADER> ComputeShader(View.ShaderMap);
+					SHADER::FParameters* PassParameters = GraphBuilder.AllocParameters<SHADER::FParameters>();
+					PassParameters->MinGenerateMipsTileCount = MinGenerateMipsTileCount;
+					PassParameters->Offset = Offset;
+					PassParameters->TileTypeCountBuffer = GraphBuilder.CreateSRV(Tiles.TileTypeCountBuffer, EPixelFormat::PF_R32_UINT);
+					PassParameters->RWMipsConditionBuffer = GraphBuilder.CreateUAV(MipsConditionBuffer, EPixelFormat::PF_R32_UINT);
+
+					FComputeShaderUtils::AddPass(GraphBuilder, RDG_EVENT_NAME("SSS::FillMipsConditionBufferCS"), ComputeShader, PassParameters, FIntVector(1, 1, 1));
+				}
 
-			// Generate mipmap for the diffuse scene color and depth, use bilinear filter conditionally
-			FGenerateMips::ExecuteCompute(GraphBuilder, View.FeatureLevel, SetupTexture, BilinearBorderSampler, MipsConditionBuffer, Offset);
+				// Generate mipmap for the diffuse scene color and depth, use bilinear filter conditionally
+				FGenerateMips::ExecuteCompute(GraphBuilder, View.FeatureLevel, SetupTexture, BilinearBorderSampler, MipsConditionBuffer, Offset);
+			}
 		}
 
 		// Allocate auxiliary buffers
 		{
+			const int32 Offset = ToIndex(bForceRunningInSeparable ? FSubsurfaceTiles::ETileType::All : FSubsurfaceTiles::ETileType::SEPARABLE);
 			// Pre-allocate black UAV if we have separable tiles.
 			SubsurfaceSubpassOneTex = CreateBlackUAVTexture(
 				GraphBuilder, 
@@ -1230,7 +1206,7 @@ void AddSubsurfaceViewPass(
 				RDG_EVENT_NAME("SSS::ClearUAV(%s, %s.TileCount > 0)",TEXT("Subsurface.SubpassOneTex"), ToString(FSubsurfaceTiles::ETileType::SEPARABLE)),
 				SubsurfaceViewport,
 				Tiles.TileTypeCountBuffer,
-				ToIndex(FSubsurfaceTiles::ETileType::SEPARABLE));
+				Offset);
 
 			// Pre-allocate black UAV for the second buffer if the separable tiles are rendered in half resolution.
 			if (bHalfRes)
@@ -1242,7 +1218,7 @@ void AddSubsurfaceViewPass(
 					RDG_EVENT_NAME("SSS::ClearUAV(%s, %s.TileCount > 0)", TEXT("Subsurface.SubpassTwoTex"), ToString(FSubsurfaceTiles::ETileType::SEPARABLE)),
 					SubsurfaceViewport,
 					Tiles.TileTypeCountBuffer,
-					ToIndex(FSubsurfaceTiles::ETileType::SEPARABLE));
+					Offset);
 			}
 			else
 			{
@@ -1285,7 +1261,9 @@ void AddSubsurfaceViewPass(
 				{	 TEXT("SSS::PassFour_BVar"),	SubsurfaceSubpassOneTex, SubsurfaceSubpassTwoTex, SHADER::ESubsurfaceType::BURLEY	 , SHADER::ESubsurfacePass::PassTwo}  //Burley Variance
 			};
 
-			const FRDGBufferSRVRef SubsurfaceBufferUsage[] = { Tiles.GetTileBufferSRV(FSubsurfaceTiles::ETileType::AFIS), Tiles.GetTileBufferSRV(FSubsurfaceTiles::ETileType::SEPARABLE) };
+			const FRDGBufferSRVRef SubsurfaceBufferUsage[] = { 
+				Tiles.GetTileBufferSRV(FSubsurfaceTiles::ETileType::AFIS), 
+				Tiles.GetTileBufferSRV(bForceRunningInSeparable ? FSubsurfaceTiles::ETileType::All : FSubsurfaceTiles::ETileType::SEPARABLE) };
 
 			//Dispatch the two phase for both SSS
 			for (int PassIndex = 0; PassIndex < NumOfSubsurfacePass; ++PassIndex)
@@ -1362,7 +1340,8 @@ void AddSubsurfaceViewPass(
 						ComputeShaderPermutationVector.Get<SHADER::FRunningInSeparable>() ? TEXT(" RunningInSeparable") : TEXT("")),
 					ComputeShader,
 					PassParameters,
-					PassParameters->IndirectDispatchArgsBuffer, Tiles.GetIndirectDispatchArgOffset(SHADER::ToSubsurfaceTileType(SubsurfaceType)));
+					PassParameters->IndirectDispatchArgsBuffer, 
+					Tiles.GetIndirectDispatchArgOffset(SHADER::ToSubsurfaceTileType(SubsurfaceType,bForceRunningInSeparable)));
 			}
 		}
 	}
diff --git a/Engine/Source/Runtime/Renderer/Private/SubsurfaceTiles.cpp b/Engine/Source/Runtime/Renderer/Private/SubsurfaceTiles.cpp
--- a/Engine/Source/Runtime/Renderer/Private/SubsurfaceTiles.cpp
+++ b/Engine/Source/Runtime/Renderer/Private/SubsurfaceTiles.cpp
@@ -138,4 +138,186 @@ const TCHAR* ToString(FSubsurfaceTiles::ETileType Type)
 	case FSubsurfaceTiles::ETileType::SEPARABLE:		return TEXT("SSS(Separable)");
 	default:											return TEXT("Unknown");
 	}
+}
+
+class FSSSTileCategorisationMarkCS : public FGlobalShader
+{
+	DECLARE_GLOBAL_SHADER(FSSSTileCategorisationMarkCS);
+	SHADER_USE_PARAMETER_STRUCT(FSSSTileCategorisationMarkCS, FGlobalShader)
+
+	class FDimensionHalfRes : SHADER_PERMUTATION_BOOL("SUBSURFACE_HALF_RES");
+	using FPermutationDomain = TShaderPermutationDomain<FDimensionHalfRes>;
+
+	BEGIN_SHADER_PARAMETER_STRUCT(FParameters, )
+		SHADER_PARAMETER_RDG_UNIFORM_BUFFER(FSceneTextureUniformParameters, SceneTextures)
+		SHADER_PARAMETER_RDG_UNIFORM_BUFFER(FSubstrateGlobalUniformParameters, Substrate)
+		SHADER_PARAMETER_STRUCT_INCLUDE(FViewShaderParameters, View)
+		SHADER_PARAMETER_STRUCT(FScreenPassTextureViewportParameters, Output)
+		SHADER_PARAMETER_STRUCT(FScreenPassTextureViewportParameters, SubsurfaceInput0)
+		SHADER_PARAMETER(FIntPoint, TiledViewRes)
+		SHADER_PARAMETER_RDG_BUFFER_UAV(RWStructuredBuffer<uint>, TileMaskBufferOut)
+	END_SHADER_PARAMETER_STRUCT()
+
+	static FPermutationDomain RemapPermutation(FPermutationDomain PermutationVector)
+	{
+		return PermutationVector;
+	}
+
+	static bool ShouldCompilePermutation(const FGlobalShaderPermutationParameters& Parameters)
+	{
+		return IsFeatureLevelSupported(Parameters.Platform, ERHIFeatureLevel::SM5);
+	}
+
+	static void ModifyCompilationEnvironment(const FGlobalShaderPermutationParameters& Parameters, FShaderCompilerEnvironment& OutEnvironment)
+	{
+		OutEnvironment.SetDefine(TEXT("TILE_CATERGORISATION_SHADER"), 1);
+		OutEnvironment.SetDefine(TEXT("SUBSURFACE_TILE_SIZE"), FSubsurfaceTiles::TileSize);
+		FGlobalShader::ModifyCompilationEnvironment(Parameters, OutEnvironment);
+	}
+};
+
+IMPLEMENT_GLOBAL_SHADER(FSSSTileCategorisationMarkCS, "/Engine/Private/PostProcessSubsurfaceTile.usf", "SSSTileCategorisationMarkCS", SF_Compute);
+
+class FSSSTileClassificationBuildListsCS : public FGlobalShader
+{
+	DECLARE_GLOBAL_SHADER(FSSSTileClassificationBuildListsCS);
+	SHADER_USE_PARAMETER_STRUCT(FSSSTileClassificationBuildListsCS, FGlobalShader)
+
+	BEGIN_SHADER_PARAMETER_STRUCT(FParameters, )
+		SHADER_PARAMETER_STRUCT_INCLUDE(FViewShaderParameters, View)
+		SHADER_PARAMETER(uint32, VertexCountPerInstanceIndirect)
+		SHADER_PARAMETER(FIntPoint, TiledViewRes)
+		SHADER_PARAMETER_RDG_BUFFER_SRV(StructuredBuffer<uint>, TileMaskBuffer)
+		SHADER_PARAMETER_RDG_BUFFER_UAV(RWBuffer<uint>, RWIndirectDrawArgsBuffer)
+		SHADER_PARAMETER_RDG_BUFFER_UAV(RWBuffer<uint>, RWIndirectDispatchArgsBuffer)
+		SHADER_PARAMETER_RDG_BUFFER_UAV(RWBuffer<uint>, RWTileTypeCountBuffer)
+		SHADER_PARAMETER_RDG_BUFFER_UAV(RWBuffer<uint>, RWSSSTileListDataBuffer)
+	END_SHADER_PARAMETER_STRUCT()
+
+	static bool ShouldCompilePermutation(const FGlobalShaderPermutationParameters& Parameters)
+	{
+		return IsFeatureLevelSupported(Parameters.Platform, ERHIFeatureLevel::SM5);
+	}
+
+	static int32 GetGroupSize()
+	{
+		return FSubsurfaceTiles::TileSize;
+	}
+
+	static void ModifyCompilationEnvironment(const FGlobalShaderPermutationParameters& Parameters, FShaderCompilerEnvironment& OutEnvironment)
+	{
+		FGlobalShader::ModifyCompilationEnvironment(Parameters, OutEnvironment);
+		OutEnvironment.SetDefine(TEXT("TILE_CATERGORISATION_SHADER"), 1);
+		OutEnvironment.SetDefine(TEXT("SUBSURFACE_TILE_SIZE"), GetGroupSize());
+	}
+};
+
+IMPLEMENT_GLOBAL_SHADER(FSSSTileClassificationBuildListsCS, "/Engine/Private/PostProcessSubsurfaceTile.usf", "SSSTileClassificationBuildListsCS", SF_Compute);
+
+
+/**
+* Code adapted from ScreenSpaceReflectionTiles to reduce Setup SSS cost.
+* Build lists of 8x8 tiles used by SSS pixels
+* Mark and build list steps are separated in order to build a more coherent list (z-ordered over a larger region), which is important for the performance of future passes due to neighbor diffusion.
+*/
+FSubsurfaceTiles ClassifySSSTiles(FRDGBuilder& GraphBuilder, const FViewInfo& View, const FSceneTextures& SceneTextures, 
+	const FScreenPassTextureViewportParameters SceneViewportParameters, const FScreenPassTextureViewportParameters SubsurfaceViewportParameters, const bool IsHalfRes)
+{
+	FSubsurfaceTiles Result;
+	
+	{
+		check(FSubsurfaceTiles::TilePerThread_GroupSize == 64); // If this value change, we need to update the shaders using 
+		check(FSubsurfaceTiles::TileSize == 8); // only size supported for now
+
+		FIntPoint SubsurfaceExtent = FIntPoint(SubsurfaceViewportParameters.ViewportSize.X, SubsurfaceViewportParameters.ViewportSize.Y);
+		Result.TileDimension = FIntPoint::DivideAndRoundUp(SubsurfaceExtent, FSubsurfaceTiles::TileSize);
+		Result.TileCount = Result.TileDimension.X * Result.TileDimension.Y;
+		Result.bRectPrimitive = GRHISupportsRectTopology;
+		Result.TileTypeCountBuffer = GraphBuilder.CreateBuffer(FRDGBufferDesc::CreateBufferDesc(sizeof(uint32), FSubsurfaceTiles::TileTypeCount), TEXT("Subsurface.TileCountBuffer"));
+		Result.TileIndirectDispatchBuffer = 
+			GraphBuilder.CreateBuffer(FRDGBufferDesc::CreateIndirectDesc<FRHIDispatchIndirectParameters>(FSubsurfaceTiles::TileTypeCount), TEXT("Subsurface.TileIndirectDispatchBuffer"));
+		Result.TileIndirectDrawBuffer =
+			GraphBuilder.CreateBuffer(FRDGBufferDesc::CreateIndirectDesc<FRHIDrawIndirectParameters>(FSubsurfaceTiles::TileTypeCount), TEXT("Subsurface.TileIndirectDrawBuffer"));
+
+		Result.TileDataBuffer[ToIndex(FSubsurfaceTiles::ETileType::All)] = 
+			GraphBuilder.CreateBuffer(FRDGBufferDesc::CreateBufferDesc(sizeof(uint32), Result.TileCount), TEXT("Subsurface.TileDataBuffer(All)"));
+		Result.TileDataBuffer[ToIndex(FSubsurfaceTiles::ETileType::AFIS)] = 
+			GraphBuilder.CreateBuffer(FRDGBufferDesc::CreateBufferDesc(sizeof(uint32), Result.TileCount), TEXT("Subsurface.TileDataBuffer(AFIS)"));
+		Result.TileDataBuffer[ToIndex(FSubsurfaceTiles::ETileType::SEPARABLE)] = 
+			GraphBuilder.CreateBuffer(FRDGBufferDesc::CreateBufferDesc(sizeof(uint32), Result.TileCount), TEXT("Subsurface.TileDataBuffer(SEPARABLE)"));
+
+		for (uint32 BufferIt = 0; BufferIt < FSubsurfaceTiles::TileTypeCount; ++BufferIt)
+		{
+			if (Result.TileDataBuffer[BufferIt])
+			{
+				Result.TileDataSRV[BufferIt] = GraphBuilder.CreateSRV(Result.TileDataBuffer[BufferIt], PF_R32_UINT);
+			}
+		}
+
+		// We write to FSubsurfaceTiles::All,
+		FRDGBufferRef TileListDataBuffer = Result.TileDataBuffer[ToIndex(FSubsurfaceTiles::ETileType::All)];//GraphBuilder.CreateBuffer(FRDGBufferDesc::CreateBufferDesc(sizeof(uint32), Result.TileDimension.X * Result.TileDimension.Y), TEXT("Subsurface.TileListDataBuffer"));
+		FRDGBufferSRVRef TileListDataBufferSRV = Result.TileDataSRV[ToIndex(FSubsurfaceTiles::ETileType::All)];//GraphBuilder.CreateSRV(TileListDataBuffer, PF_R32_UINT);
+
+		FRDGBufferUAVRef DrawIndirectParametersBufferUAV = GraphBuilder.CreateUAV(Result.TileIndirectDrawBuffer, EPixelFormat::PF_R32_UINT);
+		FRDGBufferUAVRef DispatchIndirectParametersBufferUAV = GraphBuilder.CreateUAV(Result.TileIndirectDispatchBuffer, EPixelFormat::PF_R32_UINT);
+
+		// Allocate buffer with 1 bit / tile
+		FRDGBufferRef TileMaskBuffer = GraphBuilder.CreateBuffer(FRDGBufferDesc::CreateStructuredDesc(sizeof(uint32), FMath::DivideAndRoundUp(Result.TileDimension.X * Result.TileDimension.Y, 32)), TEXT("SSR.Classify.TileMaskBuffer"));
+		FRDGBufferUAVRef TileMaskBufferUAV = GraphBuilder.CreateUAV(TileMaskBuffer);
+		
+		AddClearUAVPass(GraphBuilder, TileMaskBufferUAV, 0);
+		AddClearUAVPass(GraphBuilder, GraphBuilder.CreateUAV(Result.TileTypeCountBuffer,EPixelFormat::PF_R32_UINT), 0);
+
+		// Mark used tiles based on SHADING_MODEL_ID only
+		{
+			typedef FSSSTileCategorisationMarkCS SHADER;
+			SHADER::FPermutationDomain PermutationVector;
+			PermutationVector.Set<SHADER::FDimensionHalfRes>(IsHalfRes);
+			TShaderMapRef<SHADER> ComputeShader(View.ShaderMap, PermutationVector);
+
+			SHADER::FParameters* PassParameters = GraphBuilder.AllocParameters<SHADER::FParameters>();
+			PassParameters->SceneTextures = SceneTextures.UniformBuffer;
+			PassParameters->View = View.GetShaderParameters();
+			PassParameters->Substrate = Substrate::BindSubstrateGlobalUniformParameters(View);
+			PassParameters->Output = SubsurfaceViewportParameters;
+			PassParameters->SubsurfaceInput0 = SceneViewportParameters;
+			PassParameters->TiledViewRes = Result.TileDimension;
+			PassParameters->TileMaskBufferOut = TileMaskBufferUAV;
+
+			FComputeShaderUtils::AddPass(
+				GraphBuilder,
+				RDG_EVENT_NAME("SSS::TileCategorisationMarkTiles(All)"),
+				ComputeShader,
+				PassParameters,
+				FIntVector(Result.TileDimension.X, Result.TileDimension.Y, 1)
+			);
+		}
+
+		// Build compacted and coherent z-order tile lists from bit-marked tiles
+		{
+			typedef FSSSTileClassificationBuildListsCS SHADER;
+			TShaderMapRef<SHADER> ComputeShader(View.ShaderMap);
+
+			SHADER::FParameters* PassParameters = GraphBuilder.AllocParameters<SHADER::FParameters>();
+
+			PassParameters->View = View.GetShaderParameters();
+			PassParameters->TiledViewRes = Result.TileDimension;
+			PassParameters->VertexCountPerInstanceIndirect = GRHISupportsRectTopology ? 4 : 6;
+			PassParameters->TileMaskBuffer = GraphBuilder.CreateSRV(TileMaskBuffer);
+
+			PassParameters->RWIndirectDrawArgsBuffer = DrawIndirectParametersBufferUAV;
+			PassParameters->RWIndirectDispatchArgsBuffer = DispatchIndirectParametersBufferUAV;
+			PassParameters->RWTileTypeCountBuffer = GraphBuilder.CreateUAV(Result.TileTypeCountBuffer, PF_R32_UINT);
+			PassParameters->RWSSSTileListDataBuffer = GraphBuilder.CreateUAV(TileListDataBuffer, PF_R32_UINT);
+			
+			FComputeShaderUtils::AddPass(
+				GraphBuilder,
+				RDG_EVENT_NAME("SSS::TileCategorisationBuildList(All)"),
+				ComputeShader,
+				PassParameters,
+				FComputeShaderUtils::GetGroupCount(Result.TileDimension, SHADER::GetGroupSize())
+			);
+		}
+	}
+	return Result;
 }
\ No newline at end of file
diff --git a/Engine/Source/Runtime/Renderer/Private/SubsurfaceTiles.h b/Engine/Source/Runtime/Renderer/Private/SubsurfaceTiles.h
--- a/Engine/Source/Runtime/Renderer/Private/SubsurfaceTiles.h
+++ b/Engine/Source/Runtime/Renderer/Private/SubsurfaceTiles.h
@@ -156,3 +156,12 @@ void AddSubsurfaceTiledScreenPass(
 /** Clear the UAV texture to black only when ConditionBuffer[Offset] > 0 */
 void AddConditionalClearBlackUAVPass(FRDGBuilder& GraphBuilder, FRDGEventName&& PassName, 
 	FRDGTextureRef Texture, const FScreenPassTextureViewport& ScreenPassViewport, FRDGBufferRef ConditionBuffer, uint32 Offset);
+
+/**
+* Code adapted from ScreenSpaceReflectionTiles to reduce Setup SSS cost.
+* Build lists of 8x8 tiles used by SSS pixels
+* Mark and build list steps are separated in order to build a more coherent list (z-ordered over a larger region), which is important for the performance of future passes due to neighbor diffusion.
+*/
+FSubsurfaceTiles ClassifySSSTiles(FRDGBuilder& GraphBuilder, const FViewInfo& View, const FSceneTextures& SceneTextures, 
+	const FScreenPassTextureViewportParameters SceneViewportParameters, const FScreenPassTextureViewportParameters SubsurfaceViewportParameters, const bool IsHalfRes);
+
