* Fix incorrect hair macro AABB with split screen. The clusterAABB are expressed in view0's translated world position.
* Fix incorrect hair light culling for ViewIndex > 0

diff --git a/Engine/Shaders/Private/HairStrands/HairStrandsAABB.usf b/Engine/Shaders/Private/HairStrands/HairStrandsAABB.usf
--- a/Engine/Shaders/Private/HairStrands/HairStrandsAABB.usf
+++ b/Engine/Shaders/Private/HairStrands/HairStrandsAABB.usf
@@ -12,6 +12,7 @@ float PixelSizeAtDepth1;
 float NumPixelPerVoxel;
 uint  VoxelPageResolution;
 uint  PassCount;
+float3 TranslatedWorldOffsetCorrection;
 
 // Group data packed as {uint3 AABBMin, uint3 AABBMax}
 Buffer<uint> RegisteredIndexBuffer;
@@ -37,7 +38,12 @@ void Main(uint2 DispatchThreadId : SV_DispatchThreadID, uint LinearIndex : SV_Gr
 		if (BoundIndex < MacroGroupCount)
 		{
 			const uint RegisteredIndex = RegisteredIndexBuffer[BoundIndex];
-			const FHairAABB Bound = ReadHairAABB(RegisteredIndex, InGroupAABBBuffer);
+			FHairAABB Bound = ReadHairAABB(RegisteredIndex, InGroupAABBBuffer);
+			
+			// Correct View0 translated world offset to ViewX translated world offset
+			Bound.Min += TranslatedWorldOffsetCorrection;
+			Bound.Max += TranslatedWorldOffsetCorrection;
+
 			if (PassIt == 0)
 			{
 				g_MinBounds[LinearIndex] = Bound.Min;
diff --git a/Engine/Shaders/Private/HairStrands/HairStrandsVoxelRasterCompute.usf b/Engine/Shaders/Private/HairStrands/HairStrandsVoxelRasterCompute.usf
--- a/Engine/Shaders/Private/HairStrands/HairStrandsVoxelRasterCompute.usf
+++ b/Engine/Shaders/Private/HairStrands/HairStrandsVoxelRasterCompute.usf
@@ -168,6 +168,7 @@ uint			MacroGroupId;
 uint			bUseMacroGroupBoundCPU;
 float3			MacroGroupBoundCPU_TranslatedWorldMinAABB;
 float3			MacroGroupBoundCPU_TranslatedWorldMaxAABB;
+float3			TranslatedWorldOffsetCorrection;
 
 Buffer<int>		GroupAABBsBuffer;
 Buffer<int>		ClusterAABBsBuffer;
@@ -280,6 +281,10 @@ void MarkValid_PrepareCS(uint3 DispatchThreadId : SV_DispatchThreadID)
 	{
 		MacroGroupBound = ReadHairAABB(MacroGroupId, MacroGroupVoxelAlignedAABBBuffer);
 		GroupBound = ReadHairAABB(InstanceRegisteredIndex, GroupAABBsBuffer);
+
+		// Correct View0 translated world offset to ViewX translated world offset
+		GroupBound.Min += TranslatedWorldOffsetCorrection;
+		GroupBound.Max += TranslatedWorldOffsetCorrection;
 	}
 
 	const uint3 MinCoord = PositionToCoord(GroupBound.Min, MacroGroupBound.Min, MacroGroupBound.Max, PageIndexResolution);
diff --git a/Engine/Source/Runtime/Renderer/Private/HairStrands/HairStrandsInterface.cpp b/Engine/Source/Runtime/Renderer/Private/HairStrands/HairStrandsInterface.cpp
--- a/Engine/Source/Runtime/Renderer/Private/HairStrands/HairStrandsInterface.cpp
+++ b/Engine/Source/Runtime/Renderer/Private/HairStrands/HairStrandsInterface.cpp
@@ -217,6 +217,13 @@ bool IsHairStrandsSimulationEnable()
 
 ////////////////////////////////////////////////////////////////////////////////////////////////////////////////////////
 
+FVector3f FHairTransientResources::GetTranslatedWorldOffsetCorrection(const FViewInfo& View) const
+{ 
+	return FVector3f(View.ViewMatrices.GetPreViewTranslation() - TranslatedWorldOffsetView0); 
+}
+
+////////////////////////////////////////////////////////////////////////////////////////////////////////////////////////
+
 FHairGroupPublicData::FHairGroupPublicData(uint32 InGroupIndex, const FName& InOwnerName)
 {
 	GroupIndex = InGroupIndex;
@@ -450,17 +457,22 @@ void UpdateHairStrandsBookmarkParameters(FScene* Scene, TArray<FViewInfo>& Views
 
 			if (Out.InstancesVisibilityType.IsValidIndex(Instance->RegisteredIndex) && Out.InstancesVisibilityType[Instance->RegisteredIndex] == EHairInstanceVisibilityType::NotVisible && bCompatible)
 			{
-				if (IsHairStrandsVisibleInShadows(Views[0], *Instance))
+				for (const FViewInfo& View : Views)
 				{
-					if (bStrands) 		
-					{ 
-						Out.VisibleStrands.Add(Instance);
-						Out.InstancesVisibilityType[Instance->RegisteredIndex] = EHairInstanceVisibilityType::StrandsShadowView;
-					}
-					else if (bCardsOrMeshes) 
-					{ 
-						Out.VisibleCardsOrMeshes_Shadow.Add(Instance);
-						Out.InstancesVisibilityType[Instance->RegisteredIndex] = EHairInstanceVisibilityType::CardsOrMeshesShadowView;
+					if (IsHairStrandsVisibleInShadows(View, *Instance))
+					{
+						if (bStrands)
+						{ 
+							Out.VisibleStrands.Add(Instance);
+							Out.InstancesVisibilityType[Instance->RegisteredIndex] = EHairInstanceVisibilityType::StrandsShadowView;
+							break;
+						}
+						else if (bCardsOrMeshes) 
+						{ 
+							Out.VisibleCardsOrMeshes_Shadow.Add(Instance);
+							Out.InstancesVisibilityType[Instance->RegisteredIndex] = EHairInstanceVisibilityType::CardsOrMeshesShadowView;
+							break;
+						}
 					}
 				}
 			}
@@ -554,12 +566,10 @@ void AddVisibleShadowCastingLight(const FScene& Scene, TArray<FViewInfo>& Views,
 				}
 				Data.ViewFrustumInLightSpace.Init();
 				Data.LightInfo = LightSceneInfo;
-				break;
 			}
 			else
 			{
 				View.HairStrandsViewData.VisibleShadowCastingLights.Add(LightSceneInfo);
-				break;
 			}
 		}
 	}
diff --git a/Engine/Source/Runtime/Renderer/Private/HairStrands/HairStrandsMacroGroup.cpp b/Engine/Source/Runtime/Renderer/Private/HairStrands/HairStrandsMacroGroup.cpp
--- a/Engine/Source/Runtime/Renderer/Private/HairStrands/HairStrandsMacroGroup.cpp
+++ b/Engine/Source/Runtime/Renderer/Private/HairStrands/HairStrandsMacroGroup.cpp
@@ -31,6 +31,7 @@ class FHairMacroGroupAABBCS : public FGlobalShader
 		SHADER_PARAMETER(float, PixelSizeAtDepth1)
 		SHADER_PARAMETER(float, NumPixelPerVoxel)
 		SHADER_PARAMETER(uint32, VoxelPageResolution)
+		SHADER_PARAMETER(FVector3f, TranslatedWorldOffsetCorrection)
 
 		SHADER_PARAMETER_RDG_BUFFER_SRV(Buffer<uint>, RegisteredIndexBuffer)
 		SHADER_PARAMETER_RDG_BUFFER_SRV(Buffer<int>, InGroupAABBBuffer)
@@ -98,7 +99,8 @@ static void AddHairMacroGroupAABBPass(
 	Parameters->View			    			= View.ViewUniformBuffer;
 	Parameters->MacroGroupCount					= MacroGroupCount;
 	Parameters->PassCount						= FMath::DivideAndRoundUp(MacroGroupCount, FHairMacroGroupAABBCS::GetGroupSize());
-	
+	Parameters->TranslatedWorldOffsetCorrection = TransientResources.GetTranslatedWorldOffsetCorrection(View);
+
 	TShaderMapRef<FHairMacroGroupAABBCS> ComputeShader(View.ShaderMap);
 	FComputeShaderUtils::AddPass(
 		GraphBuilder,
diff --git a/Engine/Source/Runtime/Renderer/Private/HairStrands/HairStrandsRendering.cpp b/Engine/Source/Runtime/Renderer/Private/HairStrands/HairStrandsRendering.cpp
--- a/Engine/Source/Runtime/Renderer/Private/HairStrands/HairStrandsRendering.cpp
+++ b/Engine/Source/Runtime/Renderer/Private/HairStrands/HairStrandsRendering.cpp
@@ -97,7 +97,7 @@ void AddMeshDrawTransitionPass(
 	const FViewInfo& ViewInfo,
 	const FHairStrandsMacroGroupDatas& MacroGroupDatas);
 
-FHairTransientResources* AllocateHairTransientResources(FRDGBuilder& GraphBuilder, FScene* Scene)
+FHairTransientResources* AllocateHairTransientResources(FRDGBuilder& GraphBuilder, FScene* Scene, const TArray<FViewInfo>& Views)
 {
 	if (Scene->HairStrandsSceneData.TransientResources == nullptr)
 	{
@@ -107,6 +107,7 @@ FHairTransientResources* AllocateHairTransientResources(FRDGBuilder& GraphBuilde
 	const uint32 InstanceCount = Scene->HairStrandsSceneData.RegisteredProxies.Num();
 	FHairTransientResources* Out = Scene->HairStrandsSceneData.TransientResources;
 	*Out = FHairTransientResources();
+	Out->TranslatedWorldOffsetView0 = Views.Num() > 0 ? Views[0].ViewMatrices.GetPreViewTranslation() : FVector3d::ZeroVector;
 	if (InstanceCount > 0)
 	{
 		Out->bIsGroupAABBValid.Init(false, InstanceCount);
@@ -133,6 +134,8 @@ void RenderHairPrePass(
 	// Experimental curve raster pipeline. The function early outs when it is not enabled.
 	AddRenderCurveRasterPipeline(GraphBuilder, Scene, Views, SceneTextures.Color.Target, SceneTextures.Depth.Target);
 
+	TArray<FViewInfo*> ActiveViews;
+	ActiveViews.Reserve(Views.Num());
 	for (FViewInfo& View : Views)
 	{
 		const bool bIsViewCompatible = IsHairStrandsEnabled(EHairStrandsShaderType::Strands, View.GetShaderPlatform());
@@ -162,7 +165,7 @@ void RenderHairPrePass(
 			}
 		}
 
-		const ERHIFeatureLevel::Type FeatureLevel = Scene->GetFeatureLevel();
+		ActiveViews.Add(&View);
 
 		// Allocate state/feedback data if needed
 		if (View.ViewState && !View.ViewState->HairStrandsViewStateData.IsInit())
@@ -174,11 +177,21 @@ void RenderHairPrePass(
 
 		// Voxelization and Deep Opacity Maps
 		VoxelizeHairStrands(GraphBuilder, Scene, View, InstanceCullingManager, PreViewStereoCorrection);
-		if (View.HairStrandsViewData.MacroGroupDatas.Num() > 0)
+	}
+
+	// Transit resources to be used by draw passes (e.g., primary/shadow draws). 
+	// Note: this will reset references to instances's imported buffers
+	for (FViewInfo* View : ActiveViews)
+	{
+		if (View->HairStrandsViewData.MacroGroupDatas.Num() > 0)
 		{
-			AddMeshDrawTransitionPass(GraphBuilder, View, View.HairStrandsViewData.MacroGroupDatas);
+			AddMeshDrawTransitionPass(GraphBuilder, *View, View->HairStrandsViewData.MacroGroupDatas);
 		}
-		RenderHairStrandsDeepShadows(GraphBuilder, Scene, View, InstanceCullingManager);
+	}
+
+	for (FViewInfo* View : ActiveViews)
+	{
+		RenderHairStrandsDeepShadows(GraphBuilder, Scene, *View, InstanceCullingManager);
 	}
 }
 
diff --git a/Engine/Source/Runtime/Renderer/Private/HairStrands/HairStrandsRendering.h b/Engine/Source/Runtime/Renderer/Private/HairStrands/HairStrandsRendering.h
--- a/Engine/Source/Runtime/Renderer/Private/HairStrands/HairStrandsRendering.h
+++ b/Engine/Source/Runtime/Renderer/Private/HairStrands/HairStrandsRendering.h
@@ -21,7 +21,7 @@
 #include "HairStrandsInterface.h"
 #include "HairStrandsData.h"
 
-FHairTransientResources* AllocateHairTransientResources(FRDGBuilder& GraphBuilder, FScene* Scene);
+FHairTransientResources* AllocateHairTransientResources(FRDGBuilder& GraphBuilder, FScene* Scene, const TArray<FViewInfo>& Views);
 
 void RenderHairPrePass(
 	FRDGBuilder& GraphBuilder,
diff --git a/Engine/Source/Runtime/Renderer/Private/HairStrands/HairStrandsVoxelization.cpp b/Engine/Source/Runtime/Renderer/Private/HairStrands/HairStrandsVoxelization.cpp
--- a/Engine/Source/Runtime/Renderer/Private/HairStrands/HairStrandsVoxelization.cpp
+++ b/Engine/Source/Runtime/Renderer/Private/HairStrands/HairStrandsVoxelization.cpp
@@ -288,6 +288,7 @@ class FVoxelMarkValidPageIndex_PrepareCS : public FGlobalShader
 		SHADER_PARAMETER(uint32, bUseMacroGroupBoundCPU)
 		SHADER_PARAMETER(FVector3f, MacroGroupBoundCPU_TranslatedWorldMinAABB)
 		SHADER_PARAMETER(FVector3f, MacroGroupBoundCPU_TranslatedWorldMaxAABB)
+		SHADER_PARAMETER(FVector3f, TranslatedWorldOffsetCorrection)
 
 		SHADER_PARAMETER_RDG_BUFFER_SRV(Buffer, GroupAABBsBuffer)
 		SHADER_PARAMETER_RDG_BUFFER_SRV(Buffer, ClusterAABBsBuffer)
@@ -656,6 +657,7 @@ static void AddAllocateVoxelPagesPass(
 				Parameters->bUseMacroGroupBoundCPU = bIsGPUDriven ? 0 : 1;
 				Parameters->MacroGroupBoundCPU_TranslatedWorldMinAABB = (FVector3f)CPUAllocationDesc.TranslatedMinAABB;
 				Parameters->MacroGroupBoundCPU_TranslatedWorldMaxAABB = (FVector3f)CPUAllocationDesc.TranslatedMaxAABB;
+				Parameters->TranslatedWorldOffsetCorrection = TransientResources.GetTranslatedWorldOffsetCorrection(View);
 				Parameters->OutValidPageIndexBuffer = PageIndexBufferUAVSkipBarrier;
 
 				FIntVector DispatchCount = bUseClusterAABB ?
diff --git a/Engine/Source/Runtime/Renderer/Private/SceneVisibility.cpp b/Engine/Source/Runtime/Renderer/Private/SceneVisibility.cpp
--- a/Engine/Source/Runtime/Renderer/Private/SceneVisibility.cpp
+++ b/Engine/Source/Runtime/Renderer/Private/SceneVisibility.cpp
@@ -4760,7 +4760,7 @@ void FSceneRenderer::PreVisibilityFrameSetup(FRDGBuilder& GraphBuilder)
 		if (Views.Num() > 0 && !ViewFamily.EngineShowFlags.HitProxies)
 		{
 			CreateHairStrandsBookmarkParameters(Scene, Views, AllViews, Parameters, false /*bComputeVisibleInstances*/);
-			Parameters.TransientResources = AllocateHairTransientResources(GraphBuilder, Scene);
+			Parameters.TransientResources = AllocateHairTransientResources(GraphBuilder, Scene, Views);
 			if (Parameters.HasInstances())
 			{
 				if (Scene && IsHairStrandsEnabled(EHairStrandsShaderType::All, Scene->GetShaderPlatform()))
diff --git a/Engine/Source/Runtime/Renderer/Public/HairStrandsInterface.h b/Engine/Source/Runtime/Renderer/Public/HairStrandsInterface.h
--- a/Engine/Source/Runtime/Renderer/Public/HairStrandsInterface.h
+++ b/Engine/Source/Runtime/Renderer/Public/HairStrandsInterface.h
@@ -429,6 +429,8 @@ struct FHairTransientResources
 	TBitArray<> bIsGroupAABBValid;
 	FRDGBufferRef    GroupAABBBuffer = nullptr;
 	FRDGBufferSRVRef GroupAABBSRV = nullptr;
+	FVector3d TranslatedWorldOffsetView0 = FVector3d::ZeroVector; // View[0]'s Translated World Offset, as all ClusterAABB & GroupAABB are expressed relative to this
+	FVector3f GetTranslatedWorldOffsetCorrection(const FViewInfo& View) const;
 
 	// Cluster AABB
 	TArray<FUintVector2> ClusterAABBOffetAndCounts;
