[Hair] Enable hair voxelization only on platfrom supporting atomic on volume texture.

Change Metal SM6's config to support atomic on volume texture.

diff --git a/Engine/Config/Mac/DataDrivenPlatformInfo.ini b/Engine/Config/Mac/DataDrivenPlatformInfo.ini
--- a/Engine/Config/Mac/DataDrivenPlatformInfo.ini
+++ b/Engine/Config/Mac/DataDrivenPlatformInfo.ini
@@ -86,7 +86,7 @@ bSupportsInstancedStereo=false
 SupportsMultiViewport=RuntimeGuaranteed
 bSupportsPercentageCloserShadows=true
 bSupports4ComponentUAVReadWrite=true
-bSupportsVolumeTextureAtomics=false
+bSupportsVolumeTextureAtomics=true
 bSupportsDualSourceBlending=true
 EnablesHLSL2021ByDefault=1
 bSupportsRayTracing = false
diff --git a/Engine/Source/Runtime/Renderer/Private/HairStrands/HairStrandsDeepShadow.cpp b/Engine/Source/Runtime/Renderer/Private/HairStrands/HairStrandsDeepShadow.cpp
--- a/Engine/Source/Runtime/Renderer/Private/HairStrands/HairStrandsDeepShadow.cpp
+++ b/Engine/Source/Runtime/Renderer/Private/HairStrands/HairStrandsDeepShadow.cpp
@@ -211,7 +211,7 @@ public:
 IMPLEMENT_GLOBAL_SHADER(FDeepShadowCreateViewInfoCS, "/Engine/Private/HairStrands/HairStrandsDeepShadowAllocation.usf", "CreateViewInfo", SF_Compute);
 ///////////////////////////////////////////////////////////////////////////////////////////////////
 
-bool IsHairStrandsForVoxelTransmittanceAndShadowEnable();
+bool IsHairStrandsForVoxelTransmittanceAndShadowEnable(EShaderPlatform ShaderPlatform);
 float GetDeepShadowMaxFovAngle();
 float GetDeepShadowRasterizationScale();
 float GetDeepShadowAABBScale();
@@ -248,7 +248,7 @@ void RenderHairStrandsDeepShadows(
 
 		if (MacroGroupDatas.Num() == 0 || 
 			VisibleLights.Num() == 0 ||
-			IsHairStrandsForVoxelTransmittanceAndShadowEnable()) 
+			IsHairStrandsForVoxelTransmittanceAndShadowEnable(View.GetShaderPlatform())) 
 		{
 			return; 
 		}
diff --git a/Engine/Source/Runtime/Renderer/Private/HairStrands/HairStrandsTransmittance.cpp b/Engine/Source/Runtime/Renderer/Private/HairStrands/HairStrandsTransmittance.cpp
--- a/Engine/Source/Runtime/Renderer/Private/HairStrands/HairStrandsTransmittance.cpp
+++ b/Engine/Source/Runtime/Renderer/Private/HairStrands/HairStrandsTransmittance.cpp
@@ -767,7 +767,7 @@ static FHairStrandsTransmittanceMaskData InternalRenderHairStrandsTransmittanceM
 	if (MacroGroupDatas.Num() == 0)
 		return Out;
 
-	if (!HasDeepShadowData(LightSceneInfo, MacroGroupDatas) && !IsHairStrandsVoxelizationEnable())
+	if (!HasDeepShadowData(LightSceneInfo, MacroGroupDatas) && !IsHairStrandsVoxelizationEnable(View.GetShaderPlatform()))
 		return Out;
 
 	DECLARE_GPU_STAT(HairStrandsTransmittanceMask);
@@ -778,7 +778,7 @@ static FHairStrandsTransmittanceMaskData InternalRenderHairStrandsTransmittanceM
 	FSceneTextureParameters SceneTextures = GetSceneTextureParameters(GraphBuilder, View);
 
 	bool bHasFoundLight = false;
-	if (!IsHairStrandsForVoxelTransmittanceAndShadowEnable())
+	if (!IsHairStrandsForVoxelTransmittanceAndShadowEnable(View.GetShaderPlatform()))
 	{
 		FHairStrandsDeepShadowTransmittanceLightParams Params;
 		Params.DeepShadow_DensityScale = GetDeepShadowDensityScale();
@@ -929,7 +929,7 @@ static void InternalRenderHairStrandsShadowMask(
 	if (InMacroGroupDatas.Num() == 0)
 		return;
 
-	if (!HasDeepShadowData(LightSceneInfo, InMacroGroupDatas) && !IsHairStrandsVoxelizationEnable())
+	if (!HasDeepShadowData(LightSceneInfo, InMacroGroupDatas) && !IsHairStrandsVoxelizationEnable(View.GetShaderPlatform()))
 		return;
 
 	DECLARE_GPU_STAT(HairStrandsOpaqueMask);
@@ -964,7 +964,7 @@ static void InternalRenderHairStrandsShadowMask(
 	AtlasSlotIndexData.Reserve(InMacroGroupDatas.Num());
 	float LayerDistribution = 0;
 	FIntPoint AtlasResolution = FIntPoint::ZeroValue;
-	if (!IsHairStrandsForVoxelTransmittanceAndShadowEnable())
+	if (!IsHairStrandsForVoxelTransmittanceAndShadowEnable(View.GetShaderPlatform()))
 	{
 		for (const FHairStrandsMacroGroupData& MacroGroupData : InMacroGroupDatas)
 		{
diff --git a/Engine/Source/Runtime/Renderer/Private/HairStrands/HairStrandsVoxelization.cpp b/Engine/Source/Runtime/Renderer/Private/HairStrands/HairStrandsVoxelization.cpp
--- a/Engine/Source/Runtime/Renderer/Private/HairStrands/HairStrandsVoxelization.cpp
+++ b/Engine/Source/Runtime/Renderer/Private/HairStrands/HairStrandsVoxelization.cpp
@@ -139,14 +139,14 @@ static bool IsHairStrandsAdaptiveVoxelAllocationEnable()
 	return GHairVirtualVoxelAdaptive_Enable > 0;
 }
 
-bool IsHairStrandsVoxelizationEnable()
+bool IsHairStrandsVoxelizationEnable(EShaderPlatform ShaderPlatform)
 {
-	return GHairVoxelizationEnable > 0;
+	return RHISupportsVolumeTextureAtomics(ShaderPlatform) && GHairVoxelizationEnable > 0;
 }
 
-bool IsHairStrandsForVoxelTransmittanceAndShadowEnable()
+bool IsHairStrandsForVoxelTransmittanceAndShadowEnable(EShaderPlatform ShaderPlatform)
 {
-	return IsHairStrandsVoxelizationEnable() && GHairForVoxelTransmittanceAndShadow > 0;
+	return IsHairStrandsVoxelizationEnable(ShaderPlatform) && GHairForVoxelTransmittanceAndShadow > 0;
 }
 
 bool IsHairStrandsGPUDrivenVoxelizationEnabled()
@@ -1299,7 +1299,7 @@ void VoxelizeHairStrands(
 	FHairTransientResources* TransientResources = Scene->HairStrandsSceneData.TransientResources;
 
 	// Simple early out to check if voxelization is needed
-	if (!IsHairStrandsVoxelizationEnable() || MacroGroupDatas.Num() == 0)
+	if (!IsHairStrandsVoxelizationEnable(View.GetShaderPlatform()) || MacroGroupDatas.Num() == 0)
 	{
 		VirtualVoxelResources = AllocateDummyVirtualVoxelResources(GraphBuilder, View, MacroGroupDatas);
 		return;
diff --git a/Engine/Source/Runtime/Renderer/Private/HairStrands/HairStrandsVoxelization.h b/Engine/Source/Runtime/Renderer/Private/HairStrands/HairStrandsVoxelization.h
--- a/Engine/Source/Runtime/Renderer/Private/HairStrands/HairStrandsVoxelization.h
+++ b/Engine/Source/Runtime/Renderer/Private/HairStrands/HairStrandsVoxelization.h
@@ -13,8 +13,8 @@
 
 
 /// Global enable/disable for hair voxelization
-bool IsHairStrandsVoxelizationEnable();
-bool IsHairStrandsForVoxelTransmittanceAndShadowEnable();
+bool IsHairStrandsVoxelizationEnable(EShaderPlatform ShaderPlatform);
+bool IsHairStrandsForVoxelTransmittanceAndShadowEnable(EShaderPlatform ShaderPlatform);
 
 void VoxelizeHairStrands(
 	FRDGBuilder& GraphBuilder,
