VolumetricCloud - add VisibleInRealTimeReflectionCapture control.

#tests PC editor

diff --git a/Engine/Source/Runtime/Engine/Classes/Components/VolumetricCloudComponent.h b/Engine/Source/Runtime/Engine/Classes/Components/VolumetricCloudComponent.h
--- a/Engine/Source/Runtime/Engine/Classes/Components/VolumetricCloudComponent.h
+++ b/Engine/Source/Runtime/Engine/Classes/Components/VolumetricCloudComponent.h
@@ -147,6 +147,10 @@ class UVolumetricCloudComponent : public USceneComponent
 	UPROPERTY(EditAnywhere, AdvancedDisplay, BlueprintReadOnly, Category = Rendering)
 	uint8 bRenderInMainPass : 1;
 
+	/** If true, this component will be visible in real-time sky light reflection captures. */
+	UPROPERTY(EditAnywhere, AdvancedDisplay, BlueprintReadOnly, Category = Rendering)
+	uint8 bVisibleInRealTimeSkyCaptures : 1;
+
 
 	UFUNCTION(BlueprintCallable, Category = "Rendering")
 	ENGINE_API void SetLayerBottomAltitude(float NewValue);
@@ -197,6 +201,9 @@ class UVolumetricCloudComponent : public USceneComponent
 	UFUNCTION(BlueprintCallable, Category = "Rendering")
 	ENGINE_API void SetRenderInMainPass(bool bValue);
 
+	UFUNCTION(BlueprintCallable, Category = "Rendering")
+	ENGINE_API void SetVisibleInRealTimeSkyCaptures(bool bValue);
+
 
 protected:
 	//~ Begin UActorComponent Interface.
diff --git a/Engine/Source/Runtime/Engine/Private/Components/VolumetricCloudComponent.cpp b/Engine/Source/Runtime/Engine/Private/Components/VolumetricCloudComponent.cpp
--- a/Engine/Source/Runtime/Engine/Private/Components/VolumetricCloudComponent.cpp
+++ b/Engine/Source/Runtime/Engine/Private/Components/VolumetricCloudComponent.cpp
@@ -62,6 +62,7 @@ UVolumetricCloudComponent::UVolumetricCloudComponent(const FObjectInitializer& O
 	, AerialPespectiveMieScatteringFadeDistance(0.0f)
 	, bHoldout(false)
 	, bRenderInMainPass(true)
+	, bVisibleInRealTimeSkyCaptures(true)
 	, VolumetricCloudSceneProxy(nullptr)
 	, MaterialLifetimePtr(nullptr)
 {
@@ -223,6 +224,15 @@ void UVolumetricCloudComponent::SetRenderInMainPass(bool bValue)
 	}
 }
 
+void UVolumetricCloudComponent::SetVisibleInRealTimeSkyCaptures(bool bValue)
+{
+	if (bVisibleInRealTimeSkyCaptures != bValue)
+	{
+		bVisibleInRealTimeSkyCaptures = bValue;
+		MarkRenderStateDirty();
+	}
+}
+
 void UVolumetricCloudComponent::SetReflectionViewSampleCountScale(float NewValue)
 {
 	if (AreDynamicDataChangesAllowed() && ReflectionViewSampleCountScaleValue != NewValue)
diff --git a/Engine/Source/Runtime/Engine/Private/VolumetricCloudProxy.cpp b/Engine/Source/Runtime/Engine/Private/VolumetricCloudProxy.cpp
--- a/Engine/Source/Runtime/Engine/Private/VolumetricCloudProxy.cpp
+++ b/Engine/Source/Runtime/Engine/Private/VolumetricCloudProxy.cpp
@@ -21,6 +21,7 @@ FVolumetricCloudSceneProxy::FVolumetricCloudSceneProxy(const UVolumetricCloudCom
 	, bUsePerSampleAtmosphericLightTransmittance(InComponent->bUsePerSampleAtmosphericLightTransmittance)
 	, bHoldout(InComponent->bHoldout)
 	, bRenderInMainPass(InComponent->bRenderInMainPass)
+	, bVisibleInRealTimeSkyCaptures(InComponent->bVisibleInRealTimeSkyCaptures)
 	, SkyLightCloudBottomOcclusion(InComponent->SkyLightCloudBottomOcclusion)
 	, ViewSampleCountScale(InComponent->ViewSampleCountScale)
 	, ReflectionViewSampleCountScale(InComponent->ReflectionViewSampleCountScaleValue)
diff --git a/Engine/Source/Runtime/Engine/Public/VolumetricCloudProxy.h b/Engine/Source/Runtime/Engine/Public/VolumetricCloudProxy.h
--- a/Engine/Source/Runtime/Engine/Public/VolumetricCloudProxy.h
+++ b/Engine/Source/Runtime/Engine/Public/VolumetricCloudProxy.h
@@ -42,6 +42,7 @@ public:
 	bool bUsePerSampleAtmosphericLightTransmittance;
 	bool bHoldout;
 	bool bRenderInMainPass;
+	bool bVisibleInRealTimeSkyCaptures;
 	float SkyLightCloudBottomOcclusion;
 	
 	float ViewSampleCountScale;
diff --git a/Engine/Source/Runtime/Renderer/Private/ReflectionEnvironmentRealTimeCapture.cpp b/Engine/Source/Runtime/Renderer/Private/ReflectionEnvironmentRealTimeCapture.cpp
--- a/Engine/Source/Runtime/Renderer/Private/ReflectionEnvironmentRealTimeCapture.cpp
+++ b/Engine/Source/Runtime/Renderer/Private/ReflectionEnvironmentRealTimeCapture.cpp
@@ -628,46 +628,53 @@ void FScene::AllocateAndCaptureFrameSkyEnvMap(
 				for (int32 CubeFace = StartCubeFace; CubeFace < EndCubeFace; CubeFace++)
 				{
 					RDG_EVENT_SCOPE(GraphBuilder, "Capture Volumetric Cloud Face=%d", CubeFace);
-
-					FRDGTextureRef& LowResCloudTexture = LowResCloudTextureArray[CubeFace];
 					const uint32 CubeWidthLow = CubeWidth / CloudResolutionDivider;
 
-					FMatrix LowResCubeProjectionMatrix;
-					FViewInfo* LowResCubeViewPtr = nullptr;
-					CreateMainViewSnapshotForRealTimeCapture(LowResCubeViewPtr, LowResCubeProjectionMatrix, CubeWidthLow);
-					FViewInfo& LowResCubeView = *CubeViewPtr;
-
-					FViewMatrices CubeViewMatrices;
-					SetupCommonViewUniformBufferParameters(LowResCubeView, LowResCubeProjectionMatrix, CubeViewMatrices, CubeWidthLow, CubeFace);
-
-					// We have rendered a sky dome with identity rotation at the SkyLight position for the capture.
-					SetupViewSkyAtmosphereParametersAndResources(LowResCubeView);
-
-					LowResCubeView.CreateViewUniformBuffers(*LowResCubeView.CachedViewUniformShaderParameters);
-					CloudRC.ViewUniformBuffer = LowResCubeView.ViewUniformBuffer;
-
 					// Create a texture to render the clouds at a given resolution
+					FRDGTextureRef& LowResCloudTexture = LowResCloudTextureArray[CubeFace];
 					FRDGTextureDesc LowResCloudTexDesc = FRDGTextureDesc::Create2D(FIntPoint(CubeWidthLow, CubeWidthLow), PF_FloatRGBA,
 						FClearValueBinding::Black,
 						TexCreate_RenderTargetable | TexCreate_UAV | TexCreate_ShaderResource);
 					LowResCloudTexture = GraphBuilder.CreateTexture(LowResCloudTexDesc, TEXT("SkyLight.LowResCloudTexture"));
 
-					// Render
-					CloudRC.RenderTargets[0] = FRenderTargetBinding(LowResCloudTexture, ERenderTargetLoadAction::EClear);
-					CloudRC.bDisableCloudBlending = true;
-
-					FCloudShadowAOData CloudShadowAOData;
-					GetCloudShadowAOData(GetVolumetricCloudSceneInfo(), CubeView, GraphBuilder, CloudShadowAOData);
-					SkyRC.bShouldSampleCloudShadow = CloudShadowAOData.bShouldSampleCloudShadow;
-					SkyRC.VolumetricCloudShadowMap[0] = CloudShadowAOData.VolumetricCloudShadowMap[0];
-					SkyRC.VolumetricCloudShadowMap[1] = CloudShadowAOData.VolumetricCloudShadowMap[1];
-					SkyRC.bShouldSampleCloudSkyAO = CloudShadowAOData.bShouldSampleCloudSkyAO;
-					SkyRC.VolumetricCloudSkyAO = CloudShadowAOData.VolumetricCloudSkyAO;
-					// TODO this CloudShadowAOData management looks a bit heavy, simplify or make common
-					CloudRC.VolumetricCloudShadowTexture[0] = CloudShadowAOData.VolumetricCloudShadowMap[0];
-					CloudRC.VolumetricCloudShadowTexture[1] = CloudShadowAOData.VolumetricCloudShadowMap[1];
-
-					SceneRenderer.RenderVolumetricCloudsInternal(GraphBuilder, CloudRC, InstanceCullingManager, FIntPoint(CubeWidthLow, CubeWidthLow));
+					if (CloudRC.CloudInfo->GetVolumetricCloudSceneProxy().bVisibleInRealTimeSkyCaptures)
+					{
+						FMatrix LowResCubeProjectionMatrix;
+						FViewInfo* LowResCubeViewPtr = nullptr;
+						CreateMainViewSnapshotForRealTimeCapture(LowResCubeViewPtr, LowResCubeProjectionMatrix, CubeWidthLow);
+						FViewInfo& LowResCubeView = *CubeViewPtr;
+
+						FViewMatrices CubeViewMatrices;
+						SetupCommonViewUniformBufferParameters(LowResCubeView, LowResCubeProjectionMatrix, CubeViewMatrices, CubeWidthLow, CubeFace);
+
+						// We have rendered a sky dome with identity rotation at the SkyLight position for the capture.
+						SetupViewSkyAtmosphereParametersAndResources(LowResCubeView);
+
+						LowResCubeView.CreateViewUniformBuffers(*LowResCubeView.CachedViewUniformShaderParameters);
+						CloudRC.ViewUniformBuffer = LowResCubeView.ViewUniformBuffer;
+
+
+						// Render
+						CloudRC.RenderTargets[0] = FRenderTargetBinding(LowResCloudTexture, ERenderTargetLoadAction::EClear);
+						CloudRC.bDisableCloudBlending = true;
+
+						FCloudShadowAOData CloudShadowAOData;
+						GetCloudShadowAOData(GetVolumetricCloudSceneInfo(), CubeView, GraphBuilder, CloudShadowAOData);
+						SkyRC.bShouldSampleCloudShadow = CloudShadowAOData.bShouldSampleCloudShadow;
+						SkyRC.VolumetricCloudShadowMap[0] = CloudShadowAOData.VolumetricCloudShadowMap[0];
+						SkyRC.VolumetricCloudShadowMap[1] = CloudShadowAOData.VolumetricCloudShadowMap[1];
+						SkyRC.bShouldSampleCloudSkyAO = CloudShadowAOData.bShouldSampleCloudSkyAO;
+						SkyRC.VolumetricCloudSkyAO = CloudShadowAOData.VolumetricCloudSkyAO;
+						// TODO this CloudShadowAOData management looks a bit heavy, simplify or make common
+						CloudRC.VolumetricCloudShadowTexture[0] = CloudShadowAOData.VolumetricCloudShadowMap[0];
+						CloudRC.VolumetricCloudShadowTexture[1] = CloudShadowAOData.VolumetricCloudShadowMap[1];
+
+						SceneRenderer.RenderVolumetricCloudsInternal(GraphBuilder, CloudRC, InstanceCullingManager, FIntPoint(CubeWidthLow, CubeWidthLow));
+					}
+					else
+					{
+						AddClearRenderTargetPass(GraphBuilder, LowResCloudTexture);
+					}
 				}
 			}
 
