Atmospheric dir light - fixed the look for volumetric fog and local fog volumes when per pixel transmittance is enabled.

#rnx
#tests PC editor

diff --git a/Engine/Source/Runtime/Renderer/Private/LightGridInjection.cpp b/Engine/Source/Runtime/Renderer/Private/LightGridInjection.cpp
--- a/Engine/Source/Runtime/Renderer/Private/LightGridInjection.cpp
+++ b/Engine/Source/Runtime/Renderer/Private/LightGridInjection.cpp
@@ -1012,6 +1012,13 @@ FComputeLightGridOutput FSceneRenderer::ComputeLightGrid(FRDGBuilder& GraphBuild
 
 							ForwardLightUniformParameters->HasDirectionalLight = 1;
 							ForwardLightUniformParameters->DirectionalLightColor = FVector3f(LightParameters.Color);
+							if (LightProxy->GetUsePerPixelAtmosphereTransmittance())
+							{
+								// When using PerPixelTransmittance, transmittance is evaluated per pixel by sampling the transmittance texture. It gives better gradient on large scale objects such as mountains.
+								// However, to skip doing that texture sampling in translucent using ForwardShading/VolumetricFog, we use the simple planet top ground transmittance as a simplification.
+								// That will work for most of the cases for most of the map/terrain at the top of the virtual planet.
+								ForwardLightUniformParameters->DirectionalLightColor *= FVector3f(LightProxy->GetAtmosphereTransmittanceTowardSun());
+							}
 							ForwardLightUniformParameters->DirectionalLightVolumetricScatteringIntensity = LightProxy->GetVolumetricScatteringIntensity();
 							ForwardLightUniformParameters->DirectionalLightSpecularScale = FMath::Clamp(LightProxy->GetSpecularScale(), 0.f, 1.f);
 							ForwardLightUniformParameters->DirectionalLightDiffuseScale = FMath::Clamp(LightProxy->GetDiffuseScale(), 0.f, 1.f);
diff --git a/Engine/Source/Runtime/Renderer/Private/LocalFogVolumeRendering.cpp b/Engine/Source/Runtime/Renderer/Private/LocalFogVolumeRendering.cpp
--- a/Engine/Source/Runtime/Renderer/Private/LocalFogVolumeRendering.cpp
+++ b/Engine/Source/Runtime/Renderer/Private/LocalFogVolumeRendering.cpp
@@ -554,6 +554,7 @@ void CreateViewLocalFogVolumeBufferSRV(const FScene* Scene, FViewInfo& View, FRD
 	View.LocalFogVolumeViewData.UniformParametersStruct.LocalFogVolumeCommon.ShouldRenderLocalFogVolumeInVolumetricFog	= bShouldRenderLocalFogVolumeInVolumetricFog ? 1 : 0;
 	View.LocalFogVolumeViewData.UniformParametersStruct.LocalFogVolumeCommon.LocalFogVolumeInstances					= View.LocalFogVolumeViewData.GPUInstanceDataBufferSRV;
 	View.LocalFogVolumeViewData.UniformParametersStruct.LocalFogVolumeCommon.GlobalStartDistance						= GetLocalFogVolumeGlobalStartDistance();
+	const FLightSceneProxy* ApproximatePerPixelTransmittanceProxy = nullptr; // by default ignore
 	if (IsMobilePlatform(View.GetShaderPlatform()))
 	{
 		// On mobile there is a separate FMobileDirectionalLightShaderParameters UB which holds all directional light data.
@@ -566,6 +567,7 @@ void CreateViewLocalFogVolumeBufferSRV(const FScene* Scene, FViewInfo& View, FRD
 			FLightSceneInfo* Light = Scene->MobileDirectionalLights[ChannelIdx];
 			if (Light != nullptr)
 			{
+				ApproximatePerPixelTransmittanceProxy = Light->Proxy;
 				View.LocalFogVolumeViewData.UniformParametersStruct.LocalFogVolumeCommon.DirectionalLightColor			= FVector3f( Light->Proxy->GetSunIlluminanceAccountingForSkyAtmospherePerPixelTransmittance());
 				View.LocalFogVolumeViewData.UniformParametersStruct.LocalFogVolumeCommon.DirectionalLightDirection		= FVector3f(-Light->Proxy->GetDirection());
 				break;
@@ -577,6 +579,7 @@ void CreateViewLocalFogVolumeBufferSRV(const FScene* Scene, FViewInfo& View, FRD
 		const FLightSceneProxy* SelectedForwardDirectionalLightProxy = View.ForwardLightingResources.SelectedForwardDirectionalLightProxy;
 		View.LocalFogVolumeViewData.UniformParametersStruct.LocalFogVolumeCommon.DirectionalLightColor					= FVector3f( SelectedForwardDirectionalLightProxy->GetSunIlluminanceAccountingForSkyAtmospherePerPixelTransmittance());
 		View.LocalFogVolumeViewData.UniformParametersStruct.LocalFogVolumeCommon.DirectionalLightDirection				= FVector3f(-SelectedForwardDirectionalLightProxy->GetDirection());
+		ApproximatePerPixelTransmittanceProxy = SelectedForwardDirectionalLightProxy;
 	}
 	else
 	{
@@ -584,6 +587,14 @@ void CreateViewLocalFogVolumeBufferSRV(const FScene* Scene, FViewInfo& View, FRD
 		View.LocalFogVolumeViewData.UniformParametersStruct.LocalFogVolumeCommon.DirectionalLightDirection				= View.CachedViewUniformShaderParameters->DirectionalLightDirection;
 	}
 
+	if (ApproximatePerPixelTransmittanceProxy && ApproximatePerPixelTransmittanceProxy->GetUsePerPixelAtmosphereTransmittance())
+	{
+		// When using PerPixelTransmittance, transmittance is evaluated per pixel by sampling the transmittance texture. It gives better gradient on large scale objects such as mountains.
+		// However, to skip doing that texture sampling in Local Fog Splatting shader, we use the simple planet top ground transmittance as a simplification.
+		// That will work for most of the cases for most of the map/terrain at the top of the virtual planet.
+		View.LocalFogVolumeViewData.UniformParametersStruct.LocalFogVolumeCommon.DirectionalLightColor *= FVector3f(ApproximatePerPixelTransmittanceProxy->GetAtmosphereTransmittanceTowardSun());
+	}
+
 	// This buffer must remain a basic vertex buffer for mobile to be able to read it from vertex shader
 	View.LocalFogVolumeViewData.GPUTileDataBuffer = CreateVertexBuffer(
 		GraphBuilder, TEXT("LocalFogVolume.GPUTileDataBuffer"),
