Changed the GPUSceneInstanceCount csv counter to be the actual used instance count rather than the current allocation maximum.

#tests Trivial change, no testing needed

diff --git a/Engine/Source/Runtime/Renderer/Private/GPUScene.cpp b/Engine/Source/Runtime/Renderer/Private/GPUScene.cpp
--- a/Engine/Source/Runtime/Renderer/Private/GPUScene.cpp
+++ b/Engine/Source/Runtime/Renderer/Private/GPUScene.cpp
@@ -877,7 +877,7 @@ void FGPUScene::UpdateInternal(FRDGBuilder& GraphBuilder, FSceneUniformBuffer& S
 	RDG_EVENT_SCOPE(GraphBuilder, "GPUScene.Update");
 
 	// Do this stat separately since it should always be on.
-	CSV_CUSTOM_STAT_GLOBAL(GPUSceneInstanceCount, float(InstanceSceneDataAllocator.GetMaxSize()), ECsvCustomStatOp::Set);
+	CSV_CUSTOM_STAT_GLOBAL(GPUSceneInstanceCount, float(InstanceSceneDataAllocator.GetSparselyAllocatedSize()), ECsvCustomStatOp::Set);
 
 #if !UE_BUILD_SHIPPING
 	MaxInstancesDuringPrevUpdate = uint32(InstanceSceneDataAllocator.GetMaxSize());
