Move hair instance culling data into its own struct, and add support for multiview culling info.

diff --git a/Engine/Plugins/Runtime/HairStrands/Source/HairStrandsCore/Private/GroomDebug.cpp b/Engine/Plugins/Runtime/HairStrands/Source/HairStrandsCore/Private/GroomDebug.cpp
--- a/Engine/Plugins/Runtime/HairStrands/Source/HairStrandsCore/Private/GroomDebug.cpp
+++ b/Engine/Plugins/Runtime/HairStrands/Source/HairStrandsCore/Private/GroomDebug.cpp
@@ -776,7 +776,7 @@ static void AddHairDebugPrintInstancePass(
 	FHairTransientResources& TransientResources,
 	const FShaderPrintData* ShaderPrintData,
 	const FHairStrandsInstances& Instances,
-	const TArray<EHairInstanceVisibilityType>& InstancesVisibilityType)
+	const FHairInstanceCullingResults& CullingResults)
 {
 	if (!ShaderPrint::IsSupported(View.GetShaderPlatform()))
 	{
@@ -790,7 +790,8 @@ static void AddHairDebugPrintInstancePass(
 	InstanceCountPerType.Init(0u, uint32(EHairInstanceVisibilityType::Count));
 	for (uint32 InstanceIndex = 0; InstanceIndex < InstanceCount; ++InstanceIndex)
 	{
-		InstanceCountPerType[uint32(InstancesVisibilityType[Instances[InstanceIndex]->RegisteredIndex])]++;
+		const EHairInstanceVisibilityType VisibilityType = CullingResults.Is(View, *Instances[InstanceIndex]);
+		InstanceCountPerType[uint32(VisibilityType)]++;
 	}
 
 	// Force ShaderPrint on.
@@ -824,7 +825,7 @@ static void AddHairDebugPrintInstancePass(
 		const uint32 IntLODIndex = Instance->HairGroupPublicData->LODIndex;
 		const uint32 LODCount = Instance->HairGroupPublicData->GetLODScreenSizes().Num();
 		const EGroomCacheType ActiveGroomCacheType = GetHairInstanceCacheType(Instance);
-		const EHairInstanceVisibilityType VisibilityType = InstancesVisibilityType[Instance->RegisteredIndex];
+		const EHairInstanceVisibilityType VisibilityType = CullingResults.Is(View,*Instance);
 
 		FInstanceInfos& D = Infos.AddDefaulted_GetRef();
 		D.Data0.X =
@@ -1234,7 +1235,7 @@ void RunHairStrandsDebug(
 	FSceneInterface* Scene,
 	const FSceneView& View,
 	const FHairStrandsInstances& Instances,
-	const TArray<EHairInstanceVisibilityType>& InstancesVisibilityType,
+	const FHairInstanceCullingResults& CullingResults,
 	FHairTransientResources& TransientResources,
 	const FShaderPrintData* ShaderPrintData,
 	FRDGTextureRef SceneColorTexture,
@@ -1246,7 +1247,7 @@ void RunHairStrandsDebug(
 
 	if (ViewMode == EGroomViewMode::MacroGroups)
 	{
-		AddHairDebugPrintInstancePass(GraphBuilder, ShaderMap, View, TransientResources, ShaderPrintData, Instances, InstancesVisibilityType);
+		AddHairDebugPrintInstancePass(GraphBuilder, ShaderMap, View, TransientResources, ShaderPrintData, Instances, CullingResults);
 	}
 
 	if (ViewMode == EGroomViewMode::MeshProjection)
diff --git a/Engine/Plugins/Runtime/HairStrands/Source/HairStrandsCore/Private/GroomManager.cpp b/Engine/Plugins/Runtime/HairStrands/Source/HairStrandsCore/Private/GroomManager.cpp
--- a/Engine/Plugins/Runtime/HairStrands/Source/HairStrandsCore/Private/GroomManager.cpp
+++ b/Engine/Plugins/Runtime/HairStrands/Source/HairStrandsCore/Private/GroomManager.cpp
@@ -3169,7 +3169,7 @@ void RunHairStrandsDebug(
 	FSceneInterface* Scene,
 	const FSceneView& View,
 	const FHairStrandsInstances& Instances,
-	const TArray<EHairInstanceVisibilityType>& InstancesVisibilityType,
+	const FHairInstanceCullingResults& CullingResults,
 	FHairTransientResources& TransientResources,
 	const FShaderPrintData* ShaderPrintData,
 	FRDGTextureRef SceneColor,
@@ -3205,15 +3205,15 @@ void ProcessHairStrandsBookmark(
 		Instances = nullptr;
 		if (Bookmark == EHairStrandsBookmark::ProcessCardsAndMeshesInterpolation_PrimaryView)
 		{
-			Instances = &Parameters.VisibleCardsOrMeshes_Primary;
+			Instances = &Parameters.CullingResults.VisibleInAnyViews_CardsOrMeshes_Primary;
 		}
 		else if (Bookmark == EHairStrandsBookmark::ProcessCardsAndMeshesInterpolation_ShadowView)
 		{
-			Instances = &Parameters.VisibleCardsOrMeshes_Shadow;
+			Instances = &Parameters.CullingResults.VisibleInAnyViews_CardsOrMeshes_Shadow;
 		}
 		else if (Bookmark == EHairStrandsBookmark::ProcessStrandsInterpolation)
 		{
-			Instances = &Parameters.VisibleStrands;
+			Instances = &Parameters.CullingResults.VisibleInAnyViews_Strands;
 		}
 		else
 		{
@@ -3337,9 +3337,9 @@ void ProcessHairStrandsBookmark(
 		FHairStrandsInstances DebugInstances;
 		if (IsInstanceFrustumCullingEnable())
 		{
-			DebugInstances.Append(Parameters.VisibleStrands);
-			DebugInstances.Append(Parameters.VisibleCardsOrMeshes_Primary);
-			DebugInstances.Append(Parameters.VisibleCardsOrMeshes_Shadow);
+			DebugInstances.Append(Parameters.CullingResults.VisibleInAnyViews_Strands);
+			DebugInstances.Append(Parameters.CullingResults.VisibleInAnyViews_CardsOrMeshes_Primary);
+			DebugInstances.Append(Parameters.CullingResults.VisibleInAnyViews_CardsOrMeshes_Shadow);
 			DebugInstances.Sort([](const FHairStrandsInstance& A, const FHairStrandsInstance& B)
 			{
 				return A.RegisteredIndex < B.RegisteredIndex;
@@ -3355,7 +3355,7 @@ void ProcessHairStrandsBookmark(
 			Parameters.Scene,
 			*Parameters.View,
 			*Instances,
-			Parameters.InstancesVisibilityType,
+			Parameters.CullingResults,
 			*Parameters.TransientResources,
 			Parameters.ShaderPrintData,
 			Parameters.SceneColorTexture,
diff --git a/Engine/Source/Runtime/Renderer/Private/DeferredShadingRenderer.cpp b/Engine/Source/Runtime/Renderer/Private/DeferredShadingRenderer.cpp
--- a/Engine/Source/Runtime/Renderer/Private/DeferredShadingRenderer.cpp
+++ b/Engine/Source/Runtime/Renderer/Private/DeferredShadingRenderer.cpp
@@ -2528,7 +2528,7 @@ void FDeferredShadingSceneRenderer::Render(FRDGBuilder& GraphBuilder, const FSce
 				RunHairStrandsBookmark(GraphBuilder, EHairStrandsBookmark::ProcessStrandsInterpolation, HairStrandsBookmarkParameters);
 				if (!bHasRayTracedOverlay)
 				{
-					RenderHairPrePass(GraphBuilder, Scene, SceneTextures, Views, InstanceCullingManager, HairStrandsBookmarkParameters.InstancesVisibilityType);
+					RenderHairPrePass(GraphBuilder, Scene, SceneTextures, Views, InstanceCullingManager, HairStrandsBookmarkParameters.CullingResults);
 					RenderHairBasePass(GraphBuilder, Scene, SceneTextures, Views, InstanceCullingManager);
 				}
 			}
@@ -2697,7 +2697,7 @@ void FDeferredShadingSceneRenderer::Render(FRDGBuilder& GraphBuilder, const FSce
 			RunHairStrandsBookmark(GraphBuilder, EHairStrandsBookmark::ProcessStrandsInterpolation, HairStrandsBookmarkParameters);
 			if (!bHasRayTracedOverlay)
 			{
-				RenderHairPrePass(GraphBuilder, Scene, SceneTextures, Views, InstanceCullingManager, HairStrandsBookmarkParameters.InstancesVisibilityType);
+				RenderHairPrePass(GraphBuilder, Scene, SceneTextures, Views, InstanceCullingManager, HairStrandsBookmarkParameters.CullingResults);
 				RenderHairBasePass(GraphBuilder, Scene, SceneTextures, Views, InstanceCullingManager);
 			}
 		}
diff --git a/Engine/Source/Runtime/Renderer/Private/HairStrands/HairStrandsDebug.cpp b/Engine/Source/Runtime/Renderer/Private/HairStrands/HairStrandsDebug.cpp
--- a/Engine/Source/Runtime/Renderer/Private/HairStrands/HairStrandsDebug.cpp
+++ b/Engine/Source/Runtime/Renderer/Private/HairStrands/HairStrandsDebug.cpp
@@ -1147,7 +1147,13 @@ static void InternalRenderHairStrandsDebugInfo(
 
 	// Bookmark for calling debug rendering from the plugin
 	{
+		const FSceneView* View0 = Params.View;
+		FShaderPrintData* View0ShaderPrintData = Params.ShaderPrintData;
+		Params.View = &View;
+		Params.ShaderPrintData = &View.ShaderPrintData;
 		RunHairStrandsBookmark(GraphBuilder, EHairStrandsBookmark::ProcessDebug, Params);
+		Params.View = View0;
+		Params.ShaderPrintData = View0ShaderPrintData;
 	}
 
 	const EGroomViewMode ViewMode = GetGroomViewMode(View);
diff --git a/Engine/Source/Runtime/Renderer/Private/HairStrands/HairStrandsInterface.cpp b/Engine/Source/Runtime/Renderer/Private/HairStrands/HairStrandsInterface.cpp
--- a/Engine/Source/Runtime/Renderer/Private/HairStrands/HairStrandsInterface.cpp
+++ b/Engine/Source/Runtime/Renderer/Private/HairStrands/HairStrandsInterface.cpp
@@ -216,12 +216,61 @@ bool IsHairStrandsSimulationEnable()
 }
 
 ////////////////////////////////////////////////////////////////////////////////////////////////////////////////////////
+// FHairTransientResources
 
 FVector3f FHairTransientResources::GetTranslatedWorldOffsetCorrection(const FViewInfo& View) const
 { 
 	return FVector3f(View.ViewMatrices.GetPreViewTranslation() - TranslatedWorldOffsetView0); 
 }
 
+////////////////////////////////////////////////////////////////////////////////////////////////////////////////////////
+// FHairInstanceCullingResults
+
+void FHairInstanceCullingResults::InitVisibility(uint32 InstanceCount)
+{
+	InstancesVisibilityType.Init(0u, InstanceCount);
+}
+
+void FHairInstanceCullingResults::Add(const FViewInfo& View, const FHairStrandsInstance& Instance, EHairInstanceVisibilityType Flag)
+{
+	if (InstancesVisibilityType.IsValidIndex(Instance.RegisteredIndex))
+	{
+		const uint32 ViewIndex = FMath::Min(View.SceneRendererPrimaryViewId, 7);
+		const uint32 ViewMask  = 0xF << (ViewIndex * 4u);
+		const uint32 BitsToAdd = uint32(Flag) << (ViewIndex * 4u);
+		InstancesVisibilityType[Instance.RegisteredIndex] &= ~ViewMask;
+		InstancesVisibilityType[Instance.RegisteredIndex] |=  BitsToAdd;
+	}
+}
+
+EHairInstanceVisibilityType FHairInstanceCullingResults::Is(const FViewInfo& View, const FHairStrandsInstance& Instance) const
+{
+	EHairInstanceVisibilityType Out = EHairInstanceVisibilityType::NotVisible;
+	if (InstancesVisibilityType.IsValidIndex(Instance.RegisteredIndex))
+	{
+		const uint32 ViewIndex = FMath::Min(View.SceneRendererPrimaryViewId, 7);
+		const uint32 ViewBits  = (InstancesVisibilityType[Instance.RegisteredIndex] >> (4u * ViewIndex)) & 0xF;
+//		const uint32 BitIndex = 31u - FMath::CountLeadingZeros(ViewBits);
+		Out = EHairInstanceVisibilityType(ViewBits);
+	}
+	return Out;
+}
+
+EHairInstanceVisibilityType FHairInstanceCullingResults::Is(const FSceneView& View, const FHairStrandsInstance& Instance) const
+{
+	return Is(*(FViewInfo*)(&View), Instance);
+}
+
+bool FHairInstanceCullingResults::Is(const FViewInfo& View, const FHairStrandsInstance& Instance, EHairInstanceVisibilityType Flag) const
+{
+	return Is(View, Instance) == Flag;
+}
+
+bool FHairInstanceCullingResults::IsVisibleInAnyViews(const FHairStrandsInstance& Instance) const
+{
+	return InstancesVisibilityType.IsValidIndex(Instance.RegisteredIndex) && InstancesVisibilityType[Instance.RegisteredIndex] != 0;
+}
+
 ////////////////////////////////////////////////////////////////////////////////////////////////////////////////////////
 
 FHairGroupPublicData::FHairGroupPublicData(uint32 InGroupIndex, const FName& InOwnerName)
@@ -385,16 +434,15 @@ void RunHairStrandsBookmark(EHairStrandsBookmark Bookmark, FHairStrandsBookmarkP
 	}
 }
 
-void CreateHairStrandsBookmarkParameters(FScene* Scene, FViewInfo& View, FHairStrandsBookmarkParameters& Out, bool bComputeVisibleInstances)
+void CullHairInstance(FScene* Scene, FViewInfo& View, FHairInstanceCullingResults& Out)
 {
 	// Only compute visible instances when required as this is expensive.
-	if (bComputeVisibleInstances)
 	{
 		const int32 ActiveInstanceCount = Scene->HairStrandsSceneData.RegisteredProxies.Num();
-		Out.InstancesVisibilityType.Init(EHairInstanceVisibilityType::NotVisible, ActiveInstanceCount);
+		Out.InitVisibility(ActiveInstanceCount);
 	
 		// 1. Strands - Add all visible strands instances
-		Out.VisibleStrands.Reserve(View.HairStrandsMeshElements.Num());
+		Out.VisibleInAnyViews_Strands.Reserve(View.HairStrandsMeshElements.Num());
 		for (const FMeshBatchAndRelevance& MeshBatch : View.HairStrandsMeshElements)
 		{
 			check(MeshBatch.PrimitiveSceneProxy && MeshBatch.PrimitiveSceneProxy->ShouldRenderInMainPass());
@@ -403,14 +451,17 @@ void CreateHairStrandsBookmarkParameters(FScene* Scene, FViewInfo& View, FHairSt
 				FHairGroupPublicData* HairData = HairStrands::GetHairData(MeshBatch.Mesh);
 				if (HairData && HairData->Instance)
 				{
-					Out.VisibleStrands.Add(HairData->Instance);
-					Out.InstancesVisibilityType[HairData->Instance->RegisteredIndex] = EHairInstanceVisibilityType::StrandsPrimaryView;
+					if (!Out.IsVisibleInAnyViews(*HairData->Instance))
+					{
+						Out.VisibleInAnyViews_Strands.Add(HairData->Instance);
+					}
+					Out.Add(View, *HairData->Instance, EHairInstanceVisibilityType::StrandsPrimaryView);
 				}
 			}
 		}
 	
 		// 2. Cards/Meshes - Add all visible cards instances
-		Out.VisibleCardsOrMeshes_Primary.Reserve(View.HairCardsMeshElements.Num());
+		Out.VisibleInAnyViews_CardsOrMeshes_Primary.Reserve(View.HairCardsMeshElements.Num());
 		for (const FMeshBatchAndRelevance& MeshBatch : View.HairCardsMeshElements)
 		{
 			check(MeshBatch.PrimitiveSceneProxy && MeshBatch.PrimitiveSceneProxy->ShouldRenderInMainPass());
@@ -419,12 +470,62 @@ void CreateHairStrandsBookmarkParameters(FScene* Scene, FViewInfo& View, FHairSt
 				FHairGroupPublicData* HairData = HairStrands::GetHairData(MeshBatch.Mesh);
 				if (HairData && HairData->Instance)
 				{
-					Out.VisibleCardsOrMeshes_Primary.Add(HairData->Instance);
-					Out.InstancesVisibilityType[HairData->Instance->RegisteredIndex] = EHairInstanceVisibilityType::CardsOrMeshesPrimaryView;
+					if (!Out.IsVisibleInAnyViews(*HairData->Instance))
+					{
+						Out.VisibleInAnyViews_CardsOrMeshes_Primary.Add(HairData->Instance);
+					}
+					Out.Add(View, *HairData->Instance, EHairInstanceVisibilityType::CardsOrMeshesPrimaryView);
+				}
+			}
+		}
+	} 
+}
+
+void CullHairInstanceShadow(FScene* Scene, TArray<FViewInfo>& Views, FHairInstanceCullingResults& Out)
+{
+	const int32 ActiveInstanceCount = Scene->HairStrandsSceneData.RegisteredProxies.Num();
+	for (FHairStrandsInstance* Instance : Scene->HairStrandsSceneData.RegisteredProxies)
+	{
+		// 2. Strands - Add all instances non-visible primary view(s) but visible in shadow view(s)
+		const bool bStrands = Instance->GetHairGeometry() == EHairGeometryType::Strands;
+		const bool bCardsOrMeshes = Instance->GetHairGeometry() == EHairGeometryType::Cards || Instance->GetHairGeometry() == EHairGeometryType::Meshes;
+		const bool bCompatible = bStrands || bCardsOrMeshes;
+
+		if (bCompatible)
+		{
+			for (const FViewInfo& View : Views)
+			{
+				if (Out.Is(View, *Instance) == EHairInstanceVisibilityType::NotVisible && IsHairStrandsVisibleInShadows(View, *Instance))
+				{
+					if (bStrands)
+					{ 
+						if (!Out.IsVisibleInAnyViews(*Instance))
+						{
+							Out.VisibleInAnyViews_Strands.Add(Instance);
+						}
+						Out.Add(View, *Instance, EHairInstanceVisibilityType::StrandsShadowView);
+					}
+					else if (bCardsOrMeshes) 
+					{ 
+						if (!Out.IsVisibleInAnyViews(*Instance))
+						{
+							Out.VisibleInAnyViews_CardsOrMeshes_Shadow.Add(Instance);
+						}
+						Out.Add(View, *Instance, EHairInstanceVisibilityType::CardsOrMeshesShadowView);
+					}
 				}
 			}
 		}
 	}
+}
+
+void CreateHairStrandsBookmarkParameters(FScene* Scene, FViewInfo& View, FHairStrandsBookmarkParameters& Out, bool bComputeVisibleInstances)
+{
+	// Only compute visible instances when required as this is expensive.
+	if (bComputeVisibleInstances)
+	{
+		CullHairInstance(Scene, View, Out.CullingResults);
+	}
 
 	Out.ShaderPrintData			= ShaderPrint::IsEnabled(View.ShaderPrintData) ? &View.ShaderPrintData : nullptr;
 	Out.ShaderMap				= View.ShaderMap;
@@ -447,36 +548,7 @@ void UpdateHairStrandsBookmarkParameters(FScene* Scene, TArray<FViewInfo>& Views
 
 	if (IsHairStrandsNonVisibleShadowCastingEnable())
 	{
-		const int32 ActiveInstanceCount = Scene->HairStrandsSceneData.RegisteredProxies.Num();
-		for (FHairStrandsInstance* Instance : Scene->HairStrandsSceneData.RegisteredProxies)
-		{
-			// 2. Strands - Add all instances non-visible primary view(s) but visible in shadow view(s)
-			const bool bStrands = Instance->GetHairGeometry() == EHairGeometryType::Strands;
-			const bool bCardsOrMeshes = Instance->GetHairGeometry() == EHairGeometryType::Cards || Instance->GetHairGeometry() == EHairGeometryType::Meshes;
-			const bool bCompatible = bStrands || bCardsOrMeshes;
-
-			if (Out.InstancesVisibilityType.IsValidIndex(Instance->RegisteredIndex) && Out.InstancesVisibilityType[Instance->RegisteredIndex] == EHairInstanceVisibilityType::NotVisible && bCompatible)
-			{
-				for (const FViewInfo& View : Views)
-				{
-					if (IsHairStrandsVisibleInShadows(View, *Instance))
-					{
-						if (bStrands)
-						{ 
-							Out.VisibleStrands.Add(Instance);
-							Out.InstancesVisibilityType[Instance->RegisteredIndex] = EHairInstanceVisibilityType::StrandsShadowView;
-							break;
-						}
-						else if (bCardsOrMeshes) 
-						{ 
-							Out.VisibleCardsOrMeshes_Shadow.Add(Instance);
-							Out.InstancesVisibilityType[Instance->RegisteredIndex] = EHairInstanceVisibilityType::CardsOrMeshesShadowView;
-							break;
-						}
-					}
-				}
-			}
-		}
+		CullHairInstanceShadow(Scene, Views, Out.CullingResults);
 	}
 }
 
diff --git a/Engine/Source/Runtime/Renderer/Private/HairStrands/HairStrandsMacroGroup.cpp b/Engine/Source/Runtime/Renderer/Private/HairStrands/HairStrandsMacroGroup.cpp
--- a/Engine/Source/Runtime/Renderer/Private/HairStrands/HairStrandsMacroGroup.cpp
+++ b/Engine/Source/Runtime/Renderer/Private/HairStrands/HairStrandsMacroGroup.cpp
@@ -151,7 +151,7 @@ void CreateHairStrandsMacroGroups(
 	FRDGBuilder& GraphBuilder,
 	const FScene* Scene,
 	const FViewInfo& View, 
-	const TArray<EHairInstanceVisibilityType>& InstancesVisibilityType,
+	const FHairInstanceCullingResults& CullingResults,
 	FHairStrandsViewData& OutHairStrandsViewData,
 	bool bBuildGPUAABB)
 {
@@ -266,7 +266,7 @@ void CreateHairStrandsMacroGroups(
 	{
 		for (FHairStrandsInstance* Instance : Scene->HairStrandsSceneData.RegisteredProxies)
 		{
-			if (Instance && InstancesVisibilityType.IsValidIndex(Instance->RegisteredIndex) && InstancesVisibilityType[Instance->RegisteredIndex] == EHairInstanceVisibilityType::StrandsShadowView)
+			if (Instance && CullingResults.Is(View, *Instance, EHairInstanceVisibilityType::StrandsShadowView))
 			{
 				UpdateMacroGroup(const_cast<FHairGroupPublicData*>(Instance->GetHairData()), nullptr, nullptr, Instance->GetBounds());
 			}
diff --git a/Engine/Source/Runtime/Renderer/Private/HairStrands/HairStrandsMacroGroup.h b/Engine/Source/Runtime/Renderer/Private/HairStrands/HairStrandsMacroGroup.h
--- a/Engine/Source/Runtime/Renderer/Private/HairStrands/HairStrandsMacroGroup.h
+++ b/Engine/Source/Runtime/Renderer/Private/HairStrands/HairStrandsMacroGroup.h
@@ -18,6 +18,6 @@ void CreateHairStrandsMacroGroups(
 	FRDGBuilder& GraphBuilder,
 	const FScene* Scene,
 	const FViewInfo& View, 
-	const TArray<EHairInstanceVisibilityType>& InstancesVisibilityType,
+	const FHairInstanceCullingResults& CullingResults,
 	FHairStrandsViewData& OutHairStrandsViewData,
 	bool bBuildGPUAABB=true);
\ No newline at end of file
diff --git a/Engine/Source/Runtime/Renderer/Private/HairStrands/HairStrandsRendering.cpp b/Engine/Source/Runtime/Renderer/Private/HairStrands/HairStrandsRendering.cpp
--- a/Engine/Source/Runtime/Renderer/Private/HairStrands/HairStrandsRendering.cpp
+++ b/Engine/Source/Runtime/Renderer/Private/HairStrands/HairStrandsRendering.cpp
@@ -129,7 +129,7 @@ void RenderHairPrePass(
 	const FSceneTextures& SceneTextures,
 	TArray<FViewInfo>& Views,
 	FInstanceCullingManager& InstanceCullingManager,
-	const TArray<EHairInstanceVisibilityType>& InstancesVisibilityType)
+	const FHairInstanceCullingResults& CullingResults)
 {
 	// Experimental curve raster pipeline. The function early outs when it is not enabled.
 	AddRenderCurveRasterPipeline(GraphBuilder, Scene, Views, SceneTextures.Color.Target, SceneTextures.Depth.Target);
@@ -173,7 +173,7 @@ void RenderHairPrePass(
 			View.ViewState->HairStrandsViewStateData.Init();
 		}
 
-		CreateHairStrandsMacroGroups(GraphBuilder, Scene, View, InstancesVisibilityType, View.HairStrandsViewData);
+		CreateHairStrandsMacroGroups(GraphBuilder, Scene, View, CullingResults, View.HairStrandsViewData);
 
 		// Voxelization and Deep Opacity Maps
 		VoxelizeHairStrands(GraphBuilder, Scene, View, InstanceCullingManager, PreViewStereoCorrection);
diff --git a/Engine/Source/Runtime/Renderer/Private/HairStrands/HairStrandsRendering.h b/Engine/Source/Runtime/Renderer/Private/HairStrands/HairStrandsRendering.h
--- a/Engine/Source/Runtime/Renderer/Private/HairStrands/HairStrandsRendering.h
+++ b/Engine/Source/Runtime/Renderer/Private/HairStrands/HairStrandsRendering.h
@@ -29,7 +29,7 @@ void RenderHairPrePass(
 	const FSceneTextures& SceneTextures,
 	TArray<FViewInfo>& Views,
 	FInstanceCullingManager& InstanceCullingManager,
-	const TArray<EHairInstanceVisibilityType>& InstancesVisibilityType);
+	const FHairInstanceCullingResults& CullingResults);
 
 void RenderHairBasePass(
 	FRDGBuilder& GraphBuilder,
diff --git a/Engine/Source/Runtime/Renderer/Private/HairStrands/HairStrandsVisibility.cpp b/Engine/Source/Runtime/Renderer/Private/HairStrands/HairStrandsVisibility.cpp
--- a/Engine/Source/Runtime/Renderer/Private/HairStrands/HairStrandsVisibility.cpp
+++ b/Engine/Source/Runtime/Renderer/Private/HairStrands/HairStrandsVisibility.cpp
@@ -4404,8 +4404,8 @@ void DrawHitProxies(
 	// Geometry won't be updated for proxy view
 	const FIntPoint Resolution = HitProxyTexture->Desc.Extent;
 	FHairStrandsViewData HairStrandsViewData;
-	TArray<EHairInstanceVisibilityType> EmptyInstancesVisibilityType;
-	CreateHairStrandsMacroGroups(GraphBuilder, &Scene, View, EmptyInstancesVisibilityType, HairStrandsViewData, false /*bBuildGPUAABB*/);
+	FHairInstanceCullingResults EmptyCullingResults;
+	CreateHairStrandsMacroGroups(GraphBuilder, &Scene, View, EmptyCullingResults, HairStrandsViewData, false /*bBuildGPUAABB*/);
 
 	// We don't compute the transmittance texture as there is no need for picking.
 	FRDGTextureRef DummyTransmittanceTexture = GraphBuilder.CreateTexture(FRDGTextureDesc::Create2D(Resolution, PF_R32_FLOAT, FClearValueBinding::White, ETextureCreateFlags::ShaderResource | ETextureCreateFlags::UAV), TEXT("Hair.DummyTransmittanceTextureForHitProxyId"));
diff --git a/Engine/Source/Runtime/Renderer/Public/HairStrandsInterface.h b/Engine/Source/Runtime/Renderer/Public/HairStrandsInterface.h
--- a/Engine/Source/Runtime/Renderer/Public/HairStrandsInterface.h
+++ b/Engine/Source/Runtime/Renderer/Public/HairStrandsInterface.h
@@ -447,6 +447,39 @@ struct FHairTransientResources
 	const FCachedGeometry& GetMeshLODData(uint32 InRegisteredIndex, bool bSim) const { if (bSim) { check(SimMeshDatas.IsValidIndex(InRegisteredIndex)); return SimMeshDatas[InRegisteredIndex]; } else {check(RenMeshDatas.IsValidIndex(InRegisteredIndex)); return RenMeshDatas[InRegisteredIndex];} }
 };
 
+////////////////////////////////////////////////////////////////////////////////////////////////////
+// FHairInstanceCullingResults
+
+enum class EHairInstanceVisibilityType : uint8
+{
+	NotVisible = 0,
+	StrandsPrimaryView = 1,
+	StrandsShadowView = 2,
+	CardsOrMeshesPrimaryView = 3,
+	CardsOrMeshesShadowView = 4,
+	Count
+};
+
+struct FHairInstanceCullingResults
+{
+	// List of instances visible in any views. 
+	// This is used to know what instance needs to run interpolation
+	FHairStrandsInstances VisibleInAnyViews_Strands; // Primary & Shadow
+	FHairStrandsInstances VisibleInAnyViews_CardsOrMeshes_Primary;
+	FHairStrandsInstances VisibleInAnyViews_CardsOrMeshes_Shadow;
+
+	// Visibility per instance per view 
+	// Store up to 8 view culling result per uint
+	TArray<uint32> InstancesVisibilityType;
+
+	void InitVisibility(uint32 InstanceCount);
+	void Add(const FViewInfo& View, const FHairStrandsInstance& Instance, EHairInstanceVisibilityType Flag);
+	bool IsVisibleInAnyViews(const FHairStrandsInstance& Instance) const;
+	bool Is(const FViewInfo& View, const FHairStrandsInstance& Instance, EHairInstanceVisibilityType Flag) const;
+	EHairInstanceVisibilityType Is(const FViewInfo& View, const FHairStrandsInstance& Instance) const;
+	RENDERER_API EHairInstanceVisibilityType Is(const FSceneView& View, const FHairStrandsInstance& Instance) const;
+};
+
 ////////////////////////////////////////////////////////////////////////////////////////////////////
 // Helpers
 struct FPointPerCurveDispatchInfo
@@ -521,16 +554,6 @@ enum class EHairStrandsBookmark : uint8
 	ProcessGuideDeformation
 };
 
-enum class EHairInstanceVisibilityType : uint8
-{
-	NotVisible = 0,
-	StrandsPrimaryView = 1,
-	StrandsShadowView = 2,
-	CardsOrMeshesPrimaryView = 3,
-	CardsOrMeshesShadowView = 4,
-	Count
-};
-
 struct FHairStrandsBookmarkParameters
 {
 	FShaderPrintData* ShaderPrintData = nullptr;
@@ -539,13 +562,10 @@ struct FHairStrandsBookmarkParameters
 	uint32 ViewUniqueID = ~0; // View 0
 	FIntRect ViewRect; // View 0
 
-	FHairStrandsInstances VisibleStrands; // Primary & Shadow
-	FHairStrandsInstances VisibleCardsOrMeshes_Primary;
-	FHairStrandsInstances VisibleCardsOrMeshes_Shadow;
+	FHairInstanceCullingResults CullingResults;
 	FHairTransientResources* TransientResources = nullptr;
 
 	FHairStrandsInstances* Instances = nullptr;
-	TArray<EHairInstanceVisibilityType> InstancesVisibilityType;
 	const FSceneView* View = nullptr;// // View 0
 	FSceneInterface* Scene = nullptr;
 	TArray<const FSceneView*> AllViews;
