FastGeo - Fix crash on stale distance field data that could occur on visibility changes

diff --git a/Engine/Plugins/Experimental/FastGeoStreaming/Source/FastGeoStreaming/Private/FastGeoComponentCluster.cpp b/Engine/Plugins/Experimental/FastGeoStreaming/Source/FastGeoStreaming/Private/FastGeoComponentCluster.cpp
--- a/Engine/Plugins/Experimental/FastGeoStreaming/Source/FastGeoStreaming/Private/FastGeoComponentCluster.cpp
+++ b/Engine/Plugins/Experimental/FastGeoStreaming/Source/FastGeoStreaming/Private/FastGeoComponentCluster.cpp
@@ -2,6 +2,7 @@
 
 #include "FastGeoComponentCluster.h"
 #include "FastGeoContainer.h"
+#include "SceneInterface.h"
 
 const FFastGeoElementType FFastGeoComponentCluster::Type(&IFastGeoElement::Type);
 
@@ -122,16 +123,16 @@ void FFastGeoComponentCluster::UpdateVisibility()
 
 	if (!ShowProxies.IsEmpty() || !HideProxies.IsEmpty())
 	{
-		ENQUEUE_RENDER_COMMAND()([ShowProxies = MoveTemp(ShowProxies), HideProxies = MoveTemp(HideProxies)](FRHICommandListBase&)
+		ENQUEUE_RENDER_COMMAND()([ShowProxies = MoveTemp(ShowProxies), HideProxies = MoveTemp(HideProxies)](FRHICommandListBase&) mutable
 		{
-			for (FPrimitiveSceneProxy* Proxy : ShowProxies)
+			if (!ShowProxies.IsEmpty())
 			{
-				Proxy->SetDrawnInGame_RenderThread(true);
+				ShowProxies[0]->GetScene().UpdatePrimitivesDrawnInGame_RenderThread(ShowProxies, true);
 			}
 
-			for (FPrimitiveSceneProxy* Proxy : HideProxies)
+			if (!HideProxies.IsEmpty())
 			{
-				Proxy->SetDrawnInGame_RenderThread(false);
+				HideProxies[0]->GetScene().UpdatePrimitivesDrawnInGame_RenderThread(HideProxies, false);
 			}
 		});
 	}
diff --git a/Engine/Source/Runtime/Engine/Public/SceneInterface.h b/Engine/Source/Runtime/Engine/Public/SceneInterface.h
--- a/Engine/Source/Runtime/Engine/Public/SceneInterface.h
+++ b/Engine/Source/Runtime/Engine/Public/SceneInterface.h
@@ -555,6 +555,8 @@ public:
 	UE_DEPRECATED(5.6, "StaticMeshBounds is no longer a parameter of this interface, please update FPrimitiveSceneDesc version.")
 	virtual void UpdatePrimitiveInstances(FInstancedStaticMeshSceneDesc* Primitive) = 0;	
 
+	virtual void UpdatePrimitivesDrawnInGame_RenderThread(TArrayView<FPrimitiveSceneProxy*> PrimitiveSceneProxies, bool bDrawnInGame) = 0;
+
 	/**
 	 * Release this scene and remove it from the rendering thread
 	 */
diff --git a/Engine/Source/Runtime/Renderer/Private/RendererScene.cpp b/Engine/Source/Runtime/Renderer/Private/RendererScene.cpp
--- a/Engine/Source/Runtime/Renderer/Private/RendererScene.cpp
+++ b/Engine/Source/Runtime/Renderer/Private/RendererScene.cpp
@@ -2294,6 +2294,29 @@ void FScene::UpdatePrimitiveInstancesFromCompute(FPrimitiveSceneInfo* PrimitiveS
 	PrimitiveUpdates.Enqueue(PrimitiveSceneInfo, MoveTemp(UpdateCommand));
 }
 
+void FScene::UpdatePrimitivesDrawnInGame_RenderThread(TArrayView<FPrimitiveSceneProxy*> InPrimitiveSceneProxies, bool bDrawnInGame)
+{
+	check(IsInRenderingThread());
+
+	for (FPrimitiveSceneProxy* PrimitiveSceneProxy : InPrimitiveSceneProxies)
+	{
+		if (PrimitiveSceneProxy->IsDrawnInGame() != bDrawnInGame)
+		{
+			if (!bDrawnInGame)
+			{
+				DistanceFieldSceneData.RemovePrimitive(PrimitiveSceneProxy->GetPrimitiveSceneInfo());
+			}
+
+			PrimitiveSceneProxy->SetDrawnInGame_RenderThread(bDrawnInGame);
+
+			if (bDrawnInGame)
+			{
+				DistanceFieldSceneData.AddPrimitive(PrimitiveSceneProxy->GetPrimitiveSceneInfo());
+			}
+		}
+	}
+}
+
 void FScene::UpdatePrimitiveInstances(FInstancedStaticMeshSceneDesc* Primitive)
 {
 	UpdatePrimitiveInstances(*Primitive);
@@ -7160,6 +7183,9 @@ public:
 
 	virtual void StartUpdatePrimitiveTransform(int32 NumPrimitives) override {}
 	virtual void FinishUpdatePrimitiveTransform() override {}
+
+	virtual void UpdatePrimitivesDrawnInGame_RenderThread(TArrayView<FPrimitiveSceneProxy*> InPrimitiveSceneProxies, bool bDrawnInGame) override {}
+
 private:
 	UWorld* World;
 	class FFXSystemInterface* FXSystem;
diff --git a/Engine/Source/Runtime/Renderer/Private/ScenePrivate.h b/Engine/Source/Runtime/Renderer/Private/ScenePrivate.h
--- a/Engine/Source/Runtime/Renderer/Private/ScenePrivate.h
+++ b/Engine/Source/Runtime/Renderer/Private/ScenePrivate.h
@@ -3877,6 +3877,9 @@ public:
 	/** Returns true if the scene requires the use of debug material permutations. */
 	bool RequiresDebugMaterials() const { return PersistentViewStateDebugFlags != 0; }
 
+     /** Update the DrawnInGame flag on the provided scene proxies */
+	virtual void UpdatePrimitivesDrawnInGame_RenderThread(TArrayView<FPrimitiveSceneProxy*> PrimitiveSceneProxies, bool bDrawnInGame) override;
+
 protected:
 
 private:
