Update some VSM options for best current paths:
- Remove ability to turn off StaticSeparate as it is fairly integral to how the fast path caching logic works now
- Ensure r.shadow.virtual.cache 0 also enables the fast "uncached" paths on the nanite view
- Remove mode "1" for VSM nanite and non-nanite HZB. Only two pass occlusion culling on current frame is supported.

Not yet enabling receiver mask by default due to some memory overhead, but that is the intended path moving forward.

diff --git a/Engine/Shaders/Private/VirtualShadowMaps/VirtualShadowMapBuildPerPageDrawCommands.usf b/Engine/Shaders/Private/VirtualShadowMaps/VirtualShadowMapBuildPerPageDrawCommands.usf
--- a/Engine/Shaders/Private/VirtualShadowMaps/VirtualShadowMapBuildPerPageDrawCommands.usf
+++ b/Engine/Shaders/Private/VirtualShadowMaps/VirtualShadowMapBuildPerPageDrawCommands.usf
@@ -45,12 +45,6 @@ StructuredBuffer<uint> DrawCommandDescs;
 
 RWBuffer<uint> DrawIndirectArgsBufferOut;
 
-#if USE_HZB_OCCLUSION
-// Mode 1 == previous frame HZB, 
-// Mode 2 == this frame HZB (Nanite only).
-uint HZBMode;
-#endif
-
 // TODO: Move to common header
 // Get the area of an "inclusive" rect (which means that the max is inside the rect), also guards against negative area (where min > max)
 uint GetInclusiveRectArea(uint4 Rect)
@@ -247,9 +241,8 @@ void CullPerPageDrawCommandsCs(uint3 GroupId : SV_GroupID, int GroupThreadIndex
 
 			// TODO: Set Cull.bIsStaticGeometry and dynamic depth range culling?
 
-			// Not safe to use the dynamic occlusion for modes other than 2 (because that is approximate occlusion only due to lack of two-pass rendering).
 #if USE_HZB_OCCLUSION
-			bool bUseStaticOcclusion = bCacheAsStatic || HZBMode != 2;
+			bool bUseStaticOcclusion = bCacheAsStatic;
 #endif
 
 			// TODO: there seems to be some code sharing that could be enabled by switching on VIRTUAL_TEXTURE_TARGET for FBoxCull
@@ -312,20 +305,9 @@ void CullPerPageDrawCommandsCs(uint3 GroupId : SV_GroupID, int GroupThreadIndex
 #if USE_HZB_OCCLUSION
 						if (Cull.bViewHZB)
 						{
-							// Assume:
-							//  1. Virtual address space has not changed (if so it'd have been invalidated)
-							//  2. Physical pages that are chached have not changed (guaranteed by persistent caching)
-							//  OR (mode 2): Using this frame's HZB from Nanite shadow passes.
-							// Then we can:
-							//  1. Iterate the clipped page rect
-							//  2. Test against only valid pages which enables
-							//  3. rejecting any that don't pass HZB at least once (or touches an uncached page)
-
-							// In Mode 2 we are using the current-frame HZB and can use any page whether it is cached or not.
+							// We are using the current-frame HZB and can use any page whether it is cached or not
 							// TODO: Figure out which flag(s) we actually want here. Also need to check the non-nanite one?
-							// TODO: for mode 1, we need to use the previous view transform and rect.
-							uint VisiblePageMask = HZBMode == 1 ? VSM_FLAG_ANY_UNCACHED : 0U;
-							if (!IsVisibleMaskedHZB(HZBMode == 1 ? MipView.TargetPrevLayerIndex : VirtualShadowMapId, MipLevel, Rect, true, HZBMode == 2, VisiblePageMask, 0xFFFFFFFF, bUseStaticOcclusion))
+							if (!IsVisibleMaskedHZB(VirtualShadowMapId, MipLevel, Rect, true, true, 0U, 0xFFFFFFFF, bUseStaticOcclusion))
 							{
 								StatsBufferInterlockedInc(VSM_STAT_NON_NANITE_INSTANCES_HZB_CULLED);
 								continue;
@@ -335,7 +317,9 @@ void CullPerPageDrawCommandsCs(uint3 GroupId : SV_GroupID, int GroupThreadIndex
 
 						uint NumMappedPages = 0U;
 						{
-							const uint MarkPageDirtyFlags = VirtualShadowMapGetMarkPageDirtyFlags(bInvalidatePages, bCacheAsStatic, Cull.bIsViewUncached, bAllowWPO);
+							uint MarkPageDirtyFlags = VirtualShadowMapGetMarkPageDirtyFlags(bInvalidatePages, bCacheAsStatic, Cull.bIsViewUncached, bAllowWPO);
+							// We never included non-nanite in the HZB so remove that flag if present
+							MarkPageDirtyFlags = MarkPageDirtyFlags & ~(VSM_MARK_PAGE_DIRTY_FLAG_HZB);
 
 							uint2 RectPagesSize = (RectPages.zw + 1u) - RectPages.xy;
 							bool bIsSmallJob = RectPagesSize.x * RectPagesSize.y <= MAX_SINGLE_THREAD_MARKING_AREA;
diff --git a/Engine/Shaders/Private/VirtualShadowMaps/VirtualShadowMapPageAccessCommon.ush b/Engine/Shaders/Private/VirtualShadowMaps/VirtualShadowMapPageAccessCommon.ush
--- a/Engine/Shaders/Private/VirtualShadowMaps/VirtualShadowMapPageAccessCommon.ush
+++ b/Engine/Shaders/Private/VirtualShadowMaps/VirtualShadowMapPageAccessCommon.ush
@@ -349,11 +349,6 @@ FPageInfo UnpackPageInfo(uint PackedData)
 	return PageInfo;
 }
 
-bool VirtualShadowMapShouldCacheStaticSeparately()
-{
-	return VirtualShadowMap.StaticCachedArrayIndex > 0;
-}
-
 uint GetVirtualShadowMapStaticArrayIndex()
 {
 	return VirtualShadowMap.StaticCachedArrayIndex;
diff --git a/Engine/Shaders/Private/VirtualShadowMaps/VirtualShadowMapPhysicalPageManagement.usf b/Engine/Shaders/Private/VirtualShadowMaps/VirtualShadowMapPhysicalPageManagement.usf
--- a/Engine/Shaders/Private/VirtualShadowMaps/VirtualShadowMapPhysicalPageManagement.usf
+++ b/Engine/Shaders/Private/VirtualShadowMaps/VirtualShadowMapPhysicalPageManagement.usf
@@ -312,7 +312,7 @@ void UpdatePhysicalPages(uint3 Index : SV_DispatchThreadID)
 					{
 						if ((PrevPhysicalFlags & VSM_EXTENDED_FLAG_ANY_INVALIDATED) != 0)
 						{
-							if (VirtualShadowMapShouldCacheStaticSeparately() && ((PrevPhysicalFlags & VSM_EXTENDED_FLAG_INVALIDATE_STATIC) == 0))
+							if ((PrevPhysicalFlags & VSM_EXTENDED_FLAG_INVALIDATE_STATIC) == 0)
 							{
 								// ONLY dynamic is invalidated, static can remain cached
 								NextPageFlags |= VSM_FLAG_DYNAMIC_UNCACHED;
@@ -811,8 +811,7 @@ void SelectPagesToInitializeCS(uint PhysicalPageIndex : SV_DispatchThreadID)
 		EmitPageToProcess(OutInitializePagesIndirectArgsBuffer, OutPhysicalPagesToInitialize, PhysicalPageIndex);
 		StatsBufferInterlockedInc(VSM_STAT_NUM_PAGES_TO_CLEAR);
 
-		if (VirtualShadowMapShouldCacheStaticSeparately() &&
-			bStaticUncached &&
+		if (bStaticUncached &&
 			(MetaData.Flags & VSM_EXTENDED_FLAG_VIEW_UNCACHED) == 0U)
 		{
 			EmitPageToProcess(OutInitializePagesIndirectArgsBuffer, OutPhysicalPagesToInitialize, PhysicalPageIndex + VirtualShadowMap.MaxPhysicalPages);
@@ -873,7 +872,7 @@ void InitializePhysicalPagesIndirectCS(uint2 TileThreadID : SV_GroupThreadID, ui
 	uint3 BasePos = GetTileBasePos(TileThreadID, GroupIndex, PhysicalPagesToInitialize, MetaData);
 	bool bStaticCached = (MetaData.Flags & VSM_FLAG_STATIC_UNCACHED) == 0U;
 
-	if (bStaticCached && VirtualShadowMapShouldCacheStaticSeparately() && (MetaData.Flags & VSM_EXTENDED_FLAG_VIEW_UNCACHED) == 0U)
+	if (bStaticCached && (MetaData.Flags & VSM_EXTENDED_FLAG_VIEW_UNCACHED) == 0U)
 	{
 		// Initialize from the static page data
 		checkSlow(BasePos.z == 0U);
diff --git a/Engine/Shaders/Private/VirtualShadowMaps/VirtualShadowMapPrintStats.usf b/Engine/Shaders/Private/VirtualShadowMaps/VirtualShadowMapPrintStats.usf
--- a/Engine/Shaders/Private/VirtualShadowMaps/VirtualShadowMapPrintStats.usf
+++ b/Engine/Shaders/Private/VirtualShadowMaps/VirtualShadowMapPrintStats.usf
@@ -142,7 +142,6 @@ void PrintStats()
 		PrintLeftAlign(Context, InStatsBuffer[VSM_STAT_NUM_HZB_PAGES_BUILT]);
 		Newline(Context);
 
-		if (VirtualShadowMapShouldCacheStaticSeparately())
 		{
 			Context.Pos.x = HeadlineX;
 			Print(Context, TEXT("Static Cached"), FontWhite);
@@ -180,26 +179,6 @@ void PrintStats()
 			PrintLeftAlign(Context, InStatsBuffer[VSM_STAT_NUM_PAGES_TO_MERGE]);
 			Newline(Context);
 		}
-		else
-		{
-			Context.Pos.x = HeadlineX;
-			Print(Context, TEXT("Cached"), FontWhite);
-			Context.Pos.x = ValueX;
-			PrintLeftAlign(Context, InStatsBuffer[VSM_STAT_STATIC_CACHED_PAGES]);
-			Newline(Context);
-
-			Context.Pos.x = ItemX;
-			Print(Context, TEXT("Invalidated"), FontWhite);
-			Context.Pos.x = ValueX;
-			PrintLeftAlign(Context, InStatsBuffer[VSM_STAT_STATIC_INVALIDATED_PAGES]);
-			Newline(Context);
-
-			Context.Pos.x = ItemX;
-			Print(Context, TEXT("Considered for WPO"), FontWhite);
-			Context.Pos.x = ValueX;
-			PrintLeftAlign(Context, InStatsBuffer[VSM_STAT_WPO_CONSIDERED_PAGES]);
-			Newline(Context);
-		}
 
 		Context.Pos.x = HeadlineX;
 		Print(Context, TEXT("Global Resolution Bias   "), FontWhite);
diff --git a/Engine/Source/Runtime/Apple/MetalRHI/Private/MetalRHI.cpp b/Engine/Source/Runtime/Apple/MetalRHI/Private/MetalRHI.cpp
--- a/Engine/Source/Runtime/Apple/MetalRHI/Private/MetalRHI.cpp
+++ b/Engine/Source/Runtime/Apple/MetalRHI/Private/MetalRHI.cpp
@@ -478,19 +478,6 @@ FMetalDynamicRHI::FMetalDynamicRHI(ERHIFeatureLevel::Type RequestedFeatureLevel)
 
     GRHIBindlessSupport = RHIGetBindlessSupport(GMaxRHIShaderPlatform);
 	
-	if(GRHIAdapterName.Contains("Apple"))
-	{
-		if(GRHIBindlessSupport == ERHIBindlessSupport::Unsupported)
-		{
-			// Switch back to single page allocation for VSM (Metal does not support atomic operations on Texture2DArrays...).
-			IConsoleVariable* VSMCacheStaticSeparateCVar = IConsoleManager::Get().FindConsoleVariable(TEXT("r.Shadow.Virtual.Cache.StaticSeparate"));
-			if (VSMCacheStaticSeparateCVar != nullptr)
-			{
-				VSMCacheStaticSeparateCVar->Set(0);
-			}
-		}
-	}
-
 #if PLATFORM_SUPPORTS_MESH_SHADERS
     GRHISupportsMeshShadersTier0 = RHISupportsMeshShadersTier0(GMaxRHIShaderPlatform);
     GRHISupportsMeshShadersTier1 = RHISupportsMeshShadersTier1(GMaxRHIShaderPlatform);
diff --git a/Engine/Source/Runtime/Renderer/Private/VirtualShadowMaps/VirtualShadowMapArray.cpp b/Engine/Source/Runtime/Renderer/Private/VirtualShadowMaps/VirtualShadowMapArray.cpp
--- a/Engine/Source/Runtime/Renderer/Private/VirtualShadowMaps/VirtualShadowMapArray.cpp
+++ b/Engine/Source/Runtime/Renderer/Private/VirtualShadowMaps/VirtualShadowMapArray.cpp
@@ -27,8 +27,6 @@
 #include "Froxel/Froxel.h"
 #include "PostProcess/DiaphragmDOF.h"
 
-#define DEBUG_ALLOW_STATIC_SEPARATE_WITHOUT_CACHING 0
-
 IMPLEMENT_STATIC_UNIFORM_BUFFER_SLOT(VirtualShadowMapUbSlot);
 
 IMPLEMENT_STATIC_UNIFORM_BUFFER_STRUCT(FVirtualShadowMapUniformParameters, "VirtualShadowMap", VirtualShadowMapUbSlot);
@@ -95,18 +93,10 @@ TAutoConsoleVariable<int32> CVarMaxPhysicalPages(
 	TEXT("More space for pages means more memory usage, but allows for higher resolution shadows.\n")
 	TEXT("Ideally this value is large enough to fit enough pages for all the lights in the scene, but not too large to waste memory.\n")
 	TEXT("Enable 'ShowStats' to see how many pages are allocated in the pool right now.\n")
-	TEXT("For more page pool control, see the 'ResolutionLodBias*', 'DynamicRes.*' and 'Cache.StaticSeparate' cvars."),
+	TEXT("For more page pool control, see the 'ResolutionLodBias*' cvars."),
 	ECVF_Scalability | ECVF_RenderThreadSafe
 );
 
-TAutoConsoleVariable<int32> CVarCacheStaticSeparate(
-	TEXT("r.Shadow.Virtual.Cache.StaticSeparate"),
-	1,
-	TEXT("When enabled, caches static objects in separate pages from dynamic objects.\n")
-	TEXT("This can improve performance in largely static scenes, but doubles the memory cost of the physical page pool."),
-	ECVF_RenderThreadSafe
-);
-
 TAutoConsoleVariable<int32> CVarBuildDynamicHZB(
 	TEXT("r.Shadow.Virtual.DynamicHZB"),
 	0,
@@ -239,11 +229,8 @@ FAutoConsoleVariableRef CVarEnableNonNaniteVSM(
 
 static TAutoConsoleVariable<int32> CVarNonNaniteVsmUseHzb(
 	TEXT("r.Shadow.Virtual.NonNanite.UseHZB"),
-	2,
-	TEXT("Cull Non-Nanite instances using HZB.\n")
-	TEXT("  Set to 0 to disable.\n")
-	TEXT("  Set to 1 to use HZB from previous frame. Can incorrectly cull in some cases due to outdated data.\n")
-	TEXT("  Set to 2 to use two-pass Nanite culling with HZB from the current frame."),
+	1,
+	TEXT("Enable two-pass Nanite culling with HZB from the current frame."),
 	ECVF_RenderThreadSafe);
 
 TAutoConsoleVariable<int32> CVarVirtualShadowOnePassProjectionMaxLights(
@@ -381,11 +368,9 @@ TAutoConsoleVariable<int32> CVarLargeInstancePageAreaThreshold(
 
 static TAutoConsoleVariable<int32> CVarShadowsVirtualUseHZB(
 	TEXT("r.Shadow.Virtual.UseHZB"),
-	2,
-	TEXT("Enables HZB for (Nanite) Virtual Shadow Maps - Non-Nanite unfortunately has a separate flag with different semantics: r.Shadow.Virtual.NonNanite.UseHZB.\n")
-	TEXT(" 0 - No HZB occlusion culling\n")
-	TEXT(" 1 - Approximate Single-pass HZB occlusion culling (using previous frame HZB)\n")
-	TEXT(" 2 - Two-pass occlusion culling (default)."),
+	1,
+	TEXT("Enables two pass occlusion culling for (Nanite) Virtual Shadow Maps\n")
+	TEXT("Non-Nanite has a separate flag: r.Shadow.Virtual.NonNanite.UseHZB."),
 	ECVF_RenderThreadSafe);
 
 static TAutoConsoleVariable<int32> CVarShadowsVirtualForceFullHZBUpdate(
@@ -627,7 +612,6 @@ void FVirtualShadowMapArray::Initialize(
 
 	bCullBackfacingPixels = CVarCullBackfacingPixels.GetValueOnRenderThread() != 0;
 	bUseHzbOcclusion = CVarShadowsVirtualUseHZB.GetValueOnRenderThread() != 0;
-	bUseTwoPassHzbOcclusion = CVarShadowsVirtualUseHZB.GetValueOnRenderThread() == 2;
 	UniformParameters.NumFullShadowMaps = 0;
 	UniformParameters.NumSinglePageShadowMaps = 0;
 	UniformParameters.NumShadowMapSlots = 0;
@@ -693,24 +677,19 @@ void FVirtualShadowMapArray::Initialize(
 		uint32 PhysicalPagesY = FMath::DivideAndRoundUp((uint32)FMath::Max(1, MaxPhysicalPages), PhysicalPagesX);	
 
 		UniformParameters.MaxPhysicalPages = PhysicalPagesX * PhysicalPagesY;
-				
-		if (CVarCacheStaticSeparate.GetValueOnRenderThread() != 0)
+		
+		if (CacheManager->IsCacheEnabled())
 		{
-			#if !DEBUG_ALLOW_STATIC_SEPARATE_WITHOUT_CACHING
-			if (CacheManager->IsCacheEnabled())
-			#endif
+			// Only set up the dynamic separate HZB build allocation if 
+			// 1. caching is enabled & we cache static separate (otherwise they are combined anyway)
+			// 2. the cvar is enabled
+			if (CVarBuildDynamicHZB.GetValueOnRenderThread() != 0)
 			{
-				// Only set up the dynamic separate HZB build allocation if 
-				// 1. caching is enabled & we cache static separate (otherwise they are combined anyway)
-				// 2. the cvar is enabled
-				if (CVarBuildDynamicHZB.GetValueOnRenderThread() != 0)
-				{
-					UniformParameters.StaticHZBArrayIndex = 1;
-				}
-
-				// Enable separate static caching in the second texture array element
-				UniformParameters.StaticCachedArrayIndex = 1;
+				UniformParameters.StaticHZBArrayIndex = 1;
 			}
+
+			// Enable separate static caching in the second texture array element
+			UniformParameters.StaticCachedArrayIndex = 1;
 		}
 
 		uint32 PhysicalX = PhysicalPagesX * FVirtualShadowMap::PageSize;
@@ -2833,7 +2812,7 @@ public:
 		SHADER_PARAMETER_RDG_TEXTURE(Texture2DArray, HZBTextureArray)
 		SHADER_PARAMETER_SAMPLER(SamplerState, HZBSampler)
 		SHADER_PARAMETER(FVector2f, HZBSize)
-		SHADER_PARAMETER(uint32, HZBMode)
+		SHADER_PARAMETER(uint32, bUseHZB)
 	END_SHADER_PARAMETER_STRUCT()
 
 
@@ -3325,7 +3304,7 @@ void FVirtualShadowMapArray::RenderVirtualShadowMapsNanite(
 		Nanite::FConfiguration CullingConfig = { 0 };
 		CullingConfig.bIsShadowPass = true;
 		CullingConfig.bUpdateStreaming = bUpdateNaniteStreaming;
-		CullingConfig.bTwoPassOcclusion = UseTwoPassHzbOcclusion();
+		CullingConfig.bTwoPassOcclusion = UseHzbOcclusion();
 		CullingConfig.bExtractStats = Nanite::IsStatFilterActive(VirtualFilterName);
 		CullingConfig.SetViewFlags(SceneView);
 
@@ -3380,14 +3359,13 @@ void FVirtualShadowMapArray::RenderVirtualShadowMapsNonNanite(FRDGBuilder& Graph
 	RDG_EVENT_SCOPE(GraphBuilder, "RenderVirtualShadowMaps(Non-Nanite)");
 
 	FGPUScene& GPUScene = Scene.GPUScene;
-		
-	int32 HZBMode = CVarNonNaniteVsmUseHzb.GetValueOnRenderThread();
-	// When disabling Nanite, there may be stale data in the Nanite-HZB causing incorrect culling.
-	if (!bHZBBuiltThisFrame)
+	
+	FRDGTextureRef HZBTextureArray = nullptr;
+	// When disabling Nanite, there may be stale data in the Nanite-HZB causing incorrect culling.	
+	if (bHZBBuiltThisFrame && HZBPhysicalArrayRDG && CVarNonNaniteVsmUseHzb.GetValueOnRenderThread() != 0)
 	{
-		HZBMode = 0; /* Disable HZB culling */
+		HZBTextureArray = HZBPhysicalArrayRDG;
 	}
-	FRDGTextureRef HZBTextureArray = (HZBMode > 0 && CacheManager->IsHZBDataAvailable()) ? HZBPhysicalArrayRDG : nullptr;
 
 	TArray<FVSMCullingBatchInfo, SceneRenderingAllocator> UnBatchedVSMCullingBatchInfo;
 	TArray<FProjectedShadowInfo*, SceneRenderingAllocator> BatchedVirtualSmMeshCommandPasses;
@@ -3519,20 +3497,9 @@ void FVirtualShadowMapArray::RenderVirtualShadowMapsNonNanite(FRDGBuilder& Graph
 	FCullPerPageDrawCommandsCs::FHZBShaderParameters HZBShaderParameters;
 	if (HZBTextureArray)
 	{
-		if (HZBMode == 2)
-		{
-			// Mode 2 uses the current frame HZB & page table.
-			HZBShaderParameters.HZBPageTable = PageTableRDG;
-			HZBShaderParameters.HZBPageFlags = PageFlagsRDG;
-			HZBShaderParameters.HZBPageRectBounds = GraphBuilder.CreateSRV(AllocatedPageRectBoundsRDG);	// TODO: Uncached?
-		}
-		else
-		{
-			const FVirtualShadowMapArrayFrameData& PrevBuffers = CacheManager->GetPrevBuffers();
-			HZBShaderParameters.HZBPageTable = GraphBuilder.RegisterExternalTexture(PrevBuffers.PageTable, TEXT("Shadow.Virtual.PrevPageTable"));
-			HZBShaderParameters.HZBPageFlags = GraphBuilder.RegisterExternalTexture(PrevBuffers.PageFlags, TEXT("Shadow.Virtual.PrevPageFlags"));
-			HZBShaderParameters.HZBPageRectBounds = GraphBuilder.CreateSRV(GraphBuilder.RegisterExternalBuffer(PrevBuffers.AllocatedPageRectBounds, TEXT("Shadow.Virtual.PrevAllocatedPageRectBounds")));	// TODO: Uncached?
-		}
+		HZBShaderParameters.HZBPageTable = PageTableRDG;
+		HZBShaderParameters.HZBPageFlags = PageFlagsRDG;
+		HZBShaderParameters.HZBPageRectBounds = GraphBuilder.CreateSRV(AllocatedPageRectBoundsRDG);	// TODO: Uncached?
 		check(HZBShaderParameters.HZBPageTable);
 		check(HZBShaderParameters.HZBPageFlags);
 		check(HZBShaderParameters.HZBPageRectBounds);
@@ -3540,7 +3507,6 @@ void FVirtualShadowMapArray::RenderVirtualShadowMapsNonNanite(FRDGBuilder& Graph
 		HZBShaderParameters.HZBTextureArray = HZBTextureArray;
 		HZBShaderParameters.HZBSize = HZBTextureArray->Desc.Extent;
 		HZBShaderParameters.HZBSampler = TStaticSamplerState< SF_Point, AM_Clamp, AM_Clamp, AM_Clamp >::GetRHI();
-		HZBShaderParameters.HZBMode = HZBMode;
 	}
 	else
 	{
diff --git a/Engine/Source/Runtime/Renderer/Private/VirtualShadowMaps/VirtualShadowMapArray.h b/Engine/Source/Runtime/Renderer/Private/VirtualShadowMaps/VirtualShadowMapArray.h
--- a/Engine/Source/Runtime/Renderer/Private/VirtualShadowMaps/VirtualShadowMapArray.h
+++ b/Engine/Source/Runtime/Renderer/Private/VirtualShadowMaps/VirtualShadowMapArray.h
@@ -424,7 +424,6 @@ public:
 
 	//
 	bool UseHzbOcclusion() const { return bUseHzbOcclusion; }
-	bool UseTwoPassHzbOcclusion() const { return bUseTwoPassHzbOcclusion; }
 
 	/**
 	 * Helper function to add clamping when interpolating the LOD resolution biases to ensure the bias for moving lights can never be lower than the one for not.
@@ -526,7 +525,6 @@ private:
 	FScene &Scene;
 	//
 	bool bUseHzbOcclusion = true;
-	bool bUseTwoPassHzbOcclusion = true;
 
 	bool bInitialized = false;
 
diff --git a/Engine/Source/Runtime/Renderer/Private/VirtualShadowMaps/VirtualShadowMapCacheManager.cpp b/Engine/Source/Runtime/Renderer/Private/VirtualShadowMaps/VirtualShadowMapCacheManager.cpp
--- a/Engine/Source/Runtime/Renderer/Private/VirtualShadowMaps/VirtualShadowMapCacheManager.cpp
+++ b/Engine/Source/Runtime/Renderer/Private/VirtualShadowMaps/VirtualShadowMapCacheManager.cpp
@@ -93,7 +93,7 @@ static TAutoConsoleVariable<int32> CVarMaxLightAgeSinceLastRequest(
 	ECVF_RenderThreadSafe);
 
 static TAutoConsoleVariable<int32> CVarFramesStaticThreshold(
-	TEXT("r.Shadow.Virtual.Cache.StaticSeparate.FramesStaticThreshold"),
+	TEXT("r.Shadow.Virtual.Cache.FramesStaticThreshold"),
 	100,
 	TEXT("Number of frames without an invalidation before an object will transition to static caching."),
 	ECVF_RenderThreadSafe);
@@ -321,20 +321,26 @@ void FVirtualShadowMapPerLightCacheEntry::UpdateClipmap(
 	ClipmapCacheKey.LightDirection = LightDirection;
 	ClipmapCacheKey.FirstLevel = FirstLevel;
 
-	bool bNewIsUncached = bForceInvalidate || Prev.RenderedFrameNumber < 0;
-
-	// On transition between uncached <-> cached we must invalidate since the static pages may not be initialized
-	if (bNewIsUncached != bIsUncached)
+	// If we swap receiver mask on/off we also may need to invalidate due to incomplete pages
+	if (bInUseReceiverMask != bUseReceiverMask)
 	{
 		Prev.RenderedFrameNumber = -1;
+		bUseReceiverMask = bInUseReceiverMask;
 	}
-	bIsUncached = bNewIsUncached;
 
-	// If we swap receiver mask on/off we also may need to invalidate due to incomplete pages
-	if (bInUseReceiverMask != bUseReceiverMask)
+	// If the cache was invalidated for any reason (light movement, etc), we render the next frame as
+	// uncached as this is more efficient. Thus continuously moving lights will automatically take the
+	// uncached path always without needing to explicitly set ForceInvalidateDirectional. After one static
+	// frame though we will swap back so that we can begin establishing static cache data. Thus it is still
+	// useful to explicitly set ForceInvalidateDirectional in cases where the light is invalidating frequently
+	// but not every single frame to keep the performance consistent.
+	bool bNewIsUncached = Prev.RenderedFrameNumber < 0;
+
+	// On transition between uncached <-> cached we must invalidate since the static pages may not be initialized
+	if (bNewIsUncached != bIsUncached)
 	{
 		Prev.RenderedFrameNumber = -1;
-		bUseReceiverMask = bInUseReceiverMask;
+		bIsUncached = bNewIsUncached;
 	}
 
 	LightOrigin = FVector(0, 0, 0);
diff --git a/Engine/Source/Runtime/Renderer/Private/VirtualShadowMaps/VirtualShadowMapClipmap.cpp b/Engine/Source/Runtime/Renderer/Private/VirtualShadowMaps/VirtualShadowMapClipmap.cpp
--- a/Engine/Source/Runtime/Renderer/Private/VirtualShadowMaps/VirtualShadowMapClipmap.cpp
+++ b/Engine/Source/Runtime/Renderer/Private/VirtualShadowMaps/VirtualShadowMapClipmap.cpp
@@ -311,11 +311,14 @@ FVirtualShadowMapClipmap::FVirtualShadowMapClipmap(
 
 	VirtualShadowMapId = VirtualShadowMapArray.Allocate(false, LevelCount);
 
+	// Enable various fast paths if we are not storing cache data
+	bool bUncached = Config.bForceInvalidate || !VirtualShadowMapArrayCacheManager->IsCacheEnabled();
+
 	// TODO: We need a light/cache entry for every light/VSM now, but scene captures may not have persistent view state for indexing
 	// This is likely to all change as we refactor the multiple cache managers stuff anyways; for now this is hopefully safe since they do have separate copies
 	const uint32 UniqueViewKey = InDependentView->ViewState ? InDependentView->ViewState->GetViewKey() : 0U;
 	PerLightCacheEntry = VirtualShadowMapArrayCacheManager->FindCreateLightCacheEntry(LightSceneInfo.Id, UniqueViewKey, LevelCount, Config.ShadowTypeId);
-	PerLightCacheEntry->UpdateClipmap(LightDirection, FirstLevel, Config.bForceInvalidate, Config.bUseReceiverMask);
+	PerLightCacheEntry->UpdateClipmap(LightDirection, FirstLevel, bUncached, Config.bUseReceiverMask);
 
 	const int RadiusLn = GetLevelRadius(static_cast<float>(LastLevel));
 	const int RadiiPerLevel = 4;
@@ -425,6 +428,8 @@ FVirtualShadowMapClipmap::FVirtualShadowMapClipmap(
 		Level.DynamicDepthCullRange.X = 0.0f;
 		Level.DynamicDepthCullRange.Y = FLT_MAX;
 
+		// TODO: Maybe we remove this path if we end up swapping things to uncached dynamically
+		// This path changes the depth range and could have visual effects in some edge cases
 		if (Config.bForceInvalidate && Config.bCullDynamicTightly)
 		{
 			// Far plane in the clipmap should be just the end of the visible range
