Fix crash with RT view mode when view state is nullptr on the view.

diff --git a/Engine/Shaders/Shared/SubstrateVisualizeDefinitions.h b/Engine/Shaders/Shared/SubstrateVisualizeDefinitions.h
--- a/Engine/Shaders/Shared/SubstrateVisualizeDefinitions.h
+++ b/Engine/Shaders/Shared/SubstrateVisualizeDefinitions.h
@@ -104,6 +104,7 @@ struct FSubstratePixelDebugData
 };
 
 #ifndef __cplusplus
+uint SubstrateDebugDataSizeInUints;
 RWStructuredBuffer<int> SubstrateDebugDataUAV;
 struct FSubstrateDebugDataSerializer
 {
@@ -111,12 +112,18 @@ struct FSubstrateDebugDataSerializer
 
 	void Serialize(int Data)
 	{
-		SubstrateDebugDataUAV[WriteIndex++] = Data;
+		if (WriteIndex < SubstrateDebugDataSizeInUints)
+		{
+			SubstrateDebugDataUAV[WriteIndex++] = Data;
+		}
 	}
 
 	void Serialize(float Data)
 	{
-		SubstrateDebugDataUAV[WriteIndex++] = asuint(Data);
+		if (WriteIndex < SubstrateDebugDataSizeInUints)
+		{
+			SubstrateDebugDataUAV[WriteIndex++] = asuint(Data);
+		}
 	}
 };
 #else
diff --git a/Engine/Source/Runtime/Renderer/Private/RayTracing/RayTracingDebug.cpp b/Engine/Source/Runtime/Renderer/Private/RayTracing/RayTracingDebug.cpp
--- a/Engine/Source/Runtime/Renderer/Private/RayTracing/RayTracingDebug.cpp
+++ b/Engine/Source/Runtime/Renderer/Private/RayTracing/RayTracingDebug.cpp
@@ -168,6 +168,7 @@ BEGIN_SHADER_PARAMETER_STRUCT(FRayTracingDebugParameters, )
 	SHADER_PARAMETER(float, TriangleHitCountPerInstanceMaxThreshold)
 	SHADER_PARAMETER(uint32, TopKMostHitInstances)
 	SHADER_PARAMETER(uint32, NumTotalInstances)
+	SHADER_PARAMETER(uint32, SubstrateDebugDataSizeInUints)
 	SHADER_PARAMETER_RDG_BUFFER_SRV(RaytracingAccelerationStructure, TLAS)
 	SHADER_PARAMETER_RDG_BUFFER_SRV(RaytracingAccelerationStructure, FarFieldTLAS)
 	SHADER_PARAMETER_RDG_TEXTURE_UAV(RWTexture2D<float4>, Output)
@@ -1554,8 +1555,9 @@ void RenderRayTracingDebug(FRDGBuilder& GraphBuilder, const FScene& Scene, const
 	FRDGTextureRef SceneColorTexture = SceneTextures.Color.Target;
 
 	const uint32 DebugVisualizationMode = GetRaytracingDebugViewModeID(View);
+	const bool bSubstratePixelDebugEnable = DebugVisualizationMode == RAY_TRACING_DEBUG_VIZ_SUBSTRATE_DATA && View.ViewState && Substrate::IsSubstrateEnabled();
 
-	if (DebugVisualizationMode == RAY_TRACING_DEBUG_VIZ_SUBSTRATE_DATA)
+	if (bSubstratePixelDebugEnable)
 	{
 		ShaderPrint::SetEnabled(true);
 		ShaderPrint::RequestSpaceForLines(1024);
@@ -1713,9 +1715,18 @@ void RenderRayTracingDebug(FRDGBuilder& GraphBuilder, const FScene& Scene, const
 	SharedParameters.NumTotalInstances = NumInstances;
 	ShaderPrint::SetParameters(GraphBuilder, View.ShaderPrintData, SharedParameters.ShaderPrintUniformBuffer);
 
-	FSubstrateViewDebugData& SubstrateViewDebugData = View.ViewState->GetSubstrateViewDebugData();
-	FSubstrateViewDebugData::FTransientPixelDebugBuffer NewSubstratePixelDebugBuffer = SubstrateViewDebugData.CreateTransientPixelDebugBuffer(GraphBuilder);
-	SharedParameters.SubstrateDebugDataUAV = NewSubstratePixelDebugBuffer.PixelMaterialDebugDataUAV;
+	FSubstrateViewDebugData::FTransientPixelDebugBuffer SubstratePixelDebugBuffer;
+	if (bSubstratePixelDebugEnable)
+	{
+		FSubstrateViewDebugData& SubstrateViewDebugData = View.ViewState->GetSubstrateViewDebugData();
+		SubstratePixelDebugBuffer = SubstrateViewDebugData.CreateTransientPixelDebugBuffer(GraphBuilder);
+	}
+	else
+	{
+		SubstratePixelDebugBuffer = FSubstrateViewDebugData::CreateDummyPixelDebugBuffer(GraphBuilder);
+	}
+	SharedParameters.SubstrateDebugDataSizeInUints = SubstratePixelDebugBuffer.DebugDataSizeInUints;
+	SharedParameters.SubstrateDebugDataUAV = SubstratePixelDebugBuffer.DebugDataUAV;
 
 	// If we don't output depth, create dummy 1x1 texture
 	const bool bOutputDepth = DebugVisualizationMode == RAY_TRACING_DEBUG_VIZ_INSTANCE_OVERLAP;
@@ -1862,9 +1873,9 @@ void RenderRayTracingDebug(FRDGBuilder& GraphBuilder, const FScene& Scene, const
 		DrawInstanceOverlap(GraphBuilder, Scene, View, SceneColorTexture, OutputDepthTexture);
 	}
 
-	if (DebugVisualizationMode == RAY_TRACING_DEBUG_VIZ_SUBSTRATE_DATA)
+	if (bSubstratePixelDebugEnable)
 	{
-		Substrate::AddProcessAndPrintSubstrateMaterialPropertiesPasses(GraphBuilder, View, SceneColorTexture, ShaderPlatform, NewSubstratePixelDebugBuffer);
+		Substrate::AddProcessAndPrintSubstrateMaterialPropertiesPasses(GraphBuilder, View, SceneColorTexture, ShaderPlatform, SubstratePixelDebugBuffer);
 	}
 }
 
diff --git a/Engine/Source/Runtime/Renderer/Private/Substrate/Substrate.h b/Engine/Source/Runtime/Renderer/Private/Substrate/Substrate.h
--- a/Engine/Source/Runtime/Renderer/Private/Substrate/Substrate.h
+++ b/Engine/Source/Runtime/Renderer/Private/Substrate/Substrate.h
@@ -203,18 +203,18 @@ struct FSubstrateViewDebugData
 
 	FSubstrateViewDebugData();
 
-	struct FTransientPixelDebugBuffer
+	struct FTransientDebugBuffer
 	{
-		FRDGBufferRef PixelMaterialDebugData = nullptr;
-		FRDGBufferUAVRef PixelMaterialDebugDataUAV = nullptr;
+		uint32 DebugDataSizeInUints = 0;
+		FRDGBufferRef DebugData = nullptr;
+		FRDGBufferUAVRef DebugDataUAV = nullptr;
 	};
+
+	struct FTransientPixelDebugBuffer : FTransientDebugBuffer {};
 	FTransientPixelDebugBuffer CreateTransientPixelDebugBuffer(FRDGBuilder& GraphBuilder);
+	static FTransientPixelDebugBuffer CreateDummyPixelDebugBuffer(FRDGBuilder& GraphBuilder);
 
-	struct FTransientSystemInfoDebugBuffer
-	{
-		FRDGBufferRef SystemInfoDebugData = nullptr;
-		FRDGBufferUAVRef SystemInfoDebugDataUAV = nullptr;
-	};
+	struct FTransientSystemInfoDebugBuffer : FTransientDebugBuffer {};
 	FTransientSystemInfoDebugBuffer CreateTransientSystemInfoDebugBuffer(FRDGBuilder& GraphBuilder);
 
 	void SafeRelease();
diff --git a/Engine/Source/Runtime/Renderer/Private/Substrate/SubstrateVisualize.cpp b/Engine/Source/Runtime/Renderer/Private/Substrate/SubstrateVisualize.cpp
--- a/Engine/Source/Runtime/Renderer/Private/Substrate/SubstrateVisualize.cpp
+++ b/Engine/Source/Runtime/Renderer/Private/Substrate/SubstrateVisualize.cpp
@@ -47,6 +47,7 @@ class FMaterialPrintInfoCS : public FGlobalShader
 
 	BEGIN_SHADER_PARAMETER_STRUCT(FParameters, )
 		SHADER_PARAMETER(uint32, bOverrideCursorPosition)
+		SHADER_PARAMETER(uint32, SubstrateDebugDataSizeInUints)
 		SHADER_PARAMETER_STRUCT_REF(FViewUniformShaderParameters, ViewUniformBuffer)
 		SHADER_PARAMETER_RDG_UNIFORM_BUFFER(FSubstrateGlobalUniformParameters, Substrate)
 		SHADER_PARAMETER_STRUCT_INCLUDE(FSceneTextureParameters, SceneTextures)
@@ -135,6 +136,7 @@ class FSubstrateSystemInfoCS : public FGlobalShader
 	using FPermutationDomain = TShaderPermutationDomain<>;
 
 	BEGIN_SHADER_PARAMETER_STRUCT(FParameters, )
+		SHADER_PARAMETER(uint32, SubstrateDebugDataSizeInUints)
 		SHADER_PARAMETER_RDG_BUFFER_SRV(Buffer<uint>, ClassificationTileDrawIndirectBuffer)
 		SHADER_PARAMETER_STRUCT_REF(FViewUniformShaderParameters, ViewUniformBuffer)
 		SHADER_PARAMETER_RDG_UNIFORM_BUFFER(FSubstrateGlobalUniformParameters, Substrate)
@@ -927,13 +929,13 @@ void AddProcessAndPrintSubstrateMaterialPropertiesPasses(FRDGBuilder& GraphBuild
 	////////////////////////////////////////////////////////////////////////////////////////////////////
 	// Copy the debug data on GPU into a readback buffer for debug display when available later.
 	TSharedPtr<FRHIGPUBufferReadback> NewReadBack = MakeShared<FRHIGPUBufferReadback>(TEXT("PixelMaterialDebugDataReadback"));
-	AddEnqueueCopyPass(GraphBuilder, NewReadBack.Get(), NewSubstratePixelDebugBuffer.PixelMaterialDebugData, SubstrateViewDebugData.PixelMaterialDebugDataSizeBytes);
+	AddEnqueueCopyPass(GraphBuilder, NewReadBack.Get(), NewSubstratePixelDebugBuffer.DebugData, SubstrateViewDebugData.PixelMaterialDebugDataSizeBytes);
 	SubstrateViewDebugData.PixelMaterialDebugDataReadbackQueries.Enqueue(NewReadBack);
 }
 
 static void AddVisualizeMaterialPropertiesPasses(FRDGBuilder& GraphBuilder, const FViewInfo& View, FScreenPassTexture& ScreenPassSceneColor, EShaderPlatform Platform)
 {
-	if (!FMaterialPrintInfoCS::IsSupported(Platform)) return;
+	if (!FMaterialPrintInfoCS::IsSupported(Platform) || !View.ViewState) return;
 	FSubstrateViewDebugData& SubstrateViewDebugData = View.ViewState->GetSubstrateViewDebugData();
 
 	// Force ShaderPrint on.
@@ -944,14 +946,14 @@ static void AddVisualizeMaterialPropertiesPasses(FRDGBuilder& GraphBuilder, cons
 	////////////////////////////////////////////////////////////////////////////////////////////////////
 	// Get a new pixel material buffer and render debug data into it
 	FSubstrateViewDebugData::FTransientPixelDebugBuffer NewSubstratePixelDebugBuffer = SubstrateViewDebugData.CreateTransientPixelDebugBuffer(GraphBuilder);
-
 	{
 		FMaterialPrintInfoCS::FParameters* PassParameters = GraphBuilder.AllocParameters<FMaterialPrintInfoCS::FParameters>();
 		PassParameters->bOverrideCursorPosition = WITH_EDITOR ? 0u : 1u;
 		PassParameters->ViewUniformBuffer = View.ViewUniformBuffer;
 		PassParameters->Substrate = Substrate::BindSubstrateGlobalUniformParameters(View);
 		PassParameters->SceneTextures = GetSceneTextureParameters(GraphBuilder, View);
-		PassParameters->SubstrateDebugDataUAV = NewSubstratePixelDebugBuffer.PixelMaterialDebugDataUAV;
+		PassParameters->SubstrateDebugDataSizeInUints = NewSubstratePixelDebugBuffer.DebugDataSizeInUints;
+		PassParameters->SubstrateDebugDataUAV = NewSubstratePixelDebugBuffer.DebugDataUAV;
 		ShaderPrint::SetParameters(GraphBuilder, View.ShaderPrintData, PassParameters->ShaderPrintParameters);
 
 		TShaderMapRef<FMaterialPrintInfoCS> ComputeShader(View.ShaderMap);
@@ -996,7 +998,7 @@ bool Is8bitTileCoordEnabled();
 
 static void AddVisualizeSystemInfoPasses(FRDGBuilder& GraphBuilder, const FViewInfo& View, FScreenPassTexture& ScreenPassSceneColor, EShaderPlatform Platform)
 {
-	if (!FSubstrateSystemInfoCS::IsSupported(Platform))
+	if (!FSubstrateSystemInfoCS::IsSupported(Platform) || !View.ViewState)
 	{
 		return;
 	}
@@ -1019,7 +1021,8 @@ static void AddVisualizeSystemInfoPasses(FRDGBuilder& GraphBuilder, const FViewI
 	// Get a new pixel material buffer and render debug data into it
 	FSubstrateViewDebugData& SubstrateViewDebugData = View.ViewState->GetSubstrateViewDebugData();
 	FSubstrateViewDebugData::FTransientSystemInfoDebugBuffer NewSubstrateSystemInfoDebugBuffer = SubstrateViewDebugData.CreateTransientSystemInfoDebugBuffer(GraphBuilder);
-	PassParameters->SubstrateDebugDataUAV = NewSubstrateSystemInfoDebugBuffer.SystemInfoDebugDataUAV;
+	PassParameters->SubstrateDebugDataSizeInUints = NewSubstrateSystemInfoDebugBuffer.DebugDataSizeInUints;
+	PassParameters->SubstrateDebugDataUAV = NewSubstrateSystemInfoDebugBuffer.DebugDataUAV;
 
 	TShaderMapRef<FSubstrateSystemInfoCS> ComputeShader(View.ShaderMap);
 	FComputeShaderUtils::AddPass(GraphBuilder, RDG_EVENT_NAME("Substrate::VisualizeSystemInfo"), ComputeShader, PassParameters, FIntVector(1, 1, 1));
@@ -1305,7 +1308,7 @@ static void AddVisualizeSystemInfoPasses(FRDGBuilder& GraphBuilder, const FViewI
 	////////////////////////////////////////////////////////////////////////////////////////////////////
 	// Copy the debug data on GPU into a readback buffer for debug display when available later.
 	TSharedPtr<FRHIGPUBufferReadback> NewReadBack = MakeShared<FRHIGPUBufferReadback>(TEXT("SystemInfoDebugDataReadback"));
-	AddEnqueueCopyPass(GraphBuilder, NewReadBack.Get(), NewSubstrateSystemInfoDebugBuffer.SystemInfoDebugData, SubstrateViewDebugData.SystemInfoDebugDataSizeBytes);
+	AddEnqueueCopyPass(GraphBuilder, NewReadBack.Get(), NewSubstrateSystemInfoDebugBuffer.DebugData, SubstrateViewDebugData.SystemInfoDebugDataSizeBytes);
 	SubstrateViewDebugData.SystemInfoDebugDataReadbackQueries.Enqueue(NewReadBack);
 }
 
@@ -1463,38 +1466,39 @@ FSubstrateViewDebugData::FSubstrateViewDebugData()
 	//
 }
 
-FSubstrateViewDebugData::FTransientPixelDebugBuffer FSubstrateViewDebugData::CreateTransientPixelDebugBuffer(FRDGBuilder& GraphBuilder)
+static FSubstrateViewDebugData::FTransientDebugBuffer InternalCreateTransientPixelDebugBuffer(FRDGBuilder& GraphBuilder, uint32 DataSizeInBytes, const TCHAR* Name)
 {
-	FSubstrateViewDebugData::FTransientPixelDebugBuffer PixelDebugBuffer;
+	FSubstrateViewDebugData::FTransientDebugBuffer Out;
 
 	const uint32 SizeOfUint = sizeof(uint32);
-	const uint32 NumUints = FMath::DivideAndRoundUp(uint32(sizeof(FSubstratePixelDebugData)), SizeOfUint);
-	PixelMaterialDebugDataSizeBytes = NumUints * SizeOfUint;
+	const uint32 NumUints = FMath::DivideAndRoundUp(DataSizeInBytes, SizeOfUint);
 
 	FRDGBufferDesc RDGBufferDesc = FRDGBufferDesc::CreateStructuredDesc(SizeOfUint, NumUints);
 	RDGBufferDesc.Usage |= BUF_SourceCopy;
-	PixelDebugBuffer.PixelMaterialDebugData = GraphBuilder.CreateBuffer(RDGBufferDesc, TEXT("PixelMaterialDebugData"));
+	Out.DebugDataSizeInUints = NumUints;
+	Out.DebugData = GraphBuilder.CreateBuffer(RDGBufferDesc, Name);
+	Out.DebugDataUAV = GraphBuilder.CreateUAV(FRDGBufferUAVDesc(Out.DebugData, PF_R32_UINT));
+	return Out;
+}
 
-	PixelDebugBuffer.PixelMaterialDebugDataUAV = GraphBuilder.CreateUAV(FRDGBufferUAVDesc(PixelDebugBuffer.PixelMaterialDebugData, PF_R32_UINT));
+FSubstrateViewDebugData::FTransientPixelDebugBuffer FSubstrateViewDebugData::CreateTransientPixelDebugBuffer(FRDGBuilder& GraphBuilder)
+{
+	FSubstrateViewDebugData::FTransientDebugBuffer Out = InternalCreateTransientPixelDebugBuffer(GraphBuilder, sizeof(FSubstratePixelDebugData), TEXT("PixelMaterialDebugData"));
+	PixelMaterialDebugDataSizeBytes = Out.DebugDataSizeInUints * sizeof(uint32);
+	return FSubstrateViewDebugData::FTransientPixelDebugBuffer(Out);
+}
 
-	return PixelDebugBuffer;
+FSubstrateViewDebugData::FTransientPixelDebugBuffer FSubstrateViewDebugData::CreateDummyPixelDebugBuffer(FRDGBuilder& GraphBuilder)
+{
+	FSubstrateViewDebugData::FTransientDebugBuffer Out = InternalCreateTransientPixelDebugBuffer(GraphBuilder, sizeof(uint32), TEXT("DummyPixelMaterialDebugData"));
+	return FSubstrateViewDebugData::FTransientPixelDebugBuffer(Out);
 }
 
 FSubstrateViewDebugData::FTransientSystemInfoDebugBuffer FSubstrateViewDebugData::CreateTransientSystemInfoDebugBuffer(FRDGBuilder& GraphBuilder)
 {
-	FSubstrateViewDebugData::FTransientSystemInfoDebugBuffer SystemInfoBuffer;
-
-	const uint32 SizeOfUint = sizeof(uint32);
-	const uint32 NumUints = SUBSTRATE_TILE_TYPE_COUNT;
-	SystemInfoDebugDataSizeBytes = NumUints * SizeOfUint;
-
-	FRDGBufferDesc RDGBufferDesc = FRDGBufferDesc::CreateStructuredDesc(SizeOfUint, NumUints);
-	RDGBufferDesc.Usage |= BUF_SourceCopy;
-	SystemInfoBuffer.SystemInfoDebugData = GraphBuilder.CreateBuffer(RDGBufferDesc, TEXT("SystemInfoDebugData"));
-
-	SystemInfoBuffer.SystemInfoDebugDataUAV = GraphBuilder.CreateUAV(FRDGBufferUAVDesc(SystemInfoBuffer.SystemInfoDebugData, PF_R32_UINT));
-
-	return SystemInfoBuffer;
+	FSubstrateViewDebugData::FTransientDebugBuffer Out = InternalCreateTransientPixelDebugBuffer(GraphBuilder, SUBSTRATE_TILE_TYPE_COUNT * sizeof(uint32), TEXT("SystemInfoDebugData"));
+	SystemInfoDebugDataSizeBytes = Out.DebugDataSizeInUints * sizeof(uint32);
+	return FSubstrateViewDebugData::FTransientSystemInfoDebugBuffer(Out);
 }
 
 void FSubstrateViewDebugData::SafeRelease()
