Fix compilation issue on Android.

diff --git a/Engine/Source/Runtime/Renderer/Private/Substrate/SubstrateVisualize.cpp b/Engine/Source/Runtime/Renderer/Private/Substrate/SubstrateVisualize.cpp
--- a/Engine/Source/Runtime/Renderer/Private/Substrate/SubstrateVisualize.cpp
+++ b/Engine/Source/Runtime/Renderer/Private/Substrate/SubstrateVisualize.cpp
@@ -1466,9 +1466,10 @@ FSubstrateViewDebugData::FSubstrateViewDebugData()
 	//
 }
 
-static FSubstrateViewDebugData::FTransientDebugBuffer InternalCreateTransientPixelDebugBuffer(FRDGBuilder& GraphBuilder, uint32 DataSizeInBytes, const TCHAR* Name)
+template<typename T>
+static T InternalCreateTransientPixelDebugBuffer(FRDGBuilder& GraphBuilder, uint32 DataSizeInBytes, const TCHAR* Name)
 {
-	FSubstrateViewDebugData::FTransientDebugBuffer Out;
+	T Out;
 
 	const uint32 SizeOfUint = sizeof(uint32);
 	const uint32 NumUints = FMath::DivideAndRoundUp(DataSizeInBytes, SizeOfUint);
@@ -1483,22 +1484,21 @@ static FSubstrateViewDebugData::FTransientDebugBuffer InternalCreateTransientPix
 
 FSubstrateViewDebugData::FTransientPixelDebugBuffer FSubstrateViewDebugData::CreateTransientPixelDebugBuffer(FRDGBuilder& GraphBuilder)
 {
-	FSubstrateViewDebugData::FTransientDebugBuffer Out = InternalCreateTransientPixelDebugBuffer(GraphBuilder, sizeof(FSubstratePixelDebugData), TEXT("PixelMaterialDebugData"));
+	FTransientPixelDebugBuffer Out = InternalCreateTransientPixelDebugBuffer<FSubstrateViewDebugData::FTransientPixelDebugBuffer>(GraphBuilder, sizeof(FSubstratePixelDebugData), TEXT("PixelMaterialDebugData"));
 	PixelMaterialDebugDataSizeBytes = Out.DebugDataSizeInUints * sizeof(uint32);
-	return FSubstrateViewDebugData::FTransientPixelDebugBuffer(Out);
+	return Out;
 }
 
 FSubstrateViewDebugData::FTransientPixelDebugBuffer FSubstrateViewDebugData::CreateDummyPixelDebugBuffer(FRDGBuilder& GraphBuilder)
 {
-	FSubstrateViewDebugData::FTransientDebugBuffer Out = InternalCreateTransientPixelDebugBuffer(GraphBuilder, sizeof(uint32), TEXT("DummyPixelMaterialDebugData"));
-	return FSubstrateViewDebugData::FTransientPixelDebugBuffer(Out);
+	return InternalCreateTransientPixelDebugBuffer<FSubstrateViewDebugData::FTransientPixelDebugBuffer>(GraphBuilder, sizeof(uint32), TEXT("DummyPixelMaterialDebugData"));
 }
 
 FSubstrateViewDebugData::FTransientSystemInfoDebugBuffer FSubstrateViewDebugData::CreateTransientSystemInfoDebugBuffer(FRDGBuilder& GraphBuilder)
 {
-	FSubstrateViewDebugData::FTransientDebugBuffer Out = InternalCreateTransientPixelDebugBuffer(GraphBuilder, SUBSTRATE_TILE_TYPE_COUNT * sizeof(uint32), TEXT("SystemInfoDebugData"));
+	FSubstrateViewDebugData::FTransientSystemInfoDebugBuffer Out = InternalCreateTransientPixelDebugBuffer<FSubstrateViewDebugData::FTransientSystemInfoDebugBuffer>(GraphBuilder, SUBSTRATE_TILE_TYPE_COUNT * sizeof(uint32), TEXT("SystemInfoDebugData"));
 	SystemInfoDebugDataSizeBytes = Out.DebugDataSizeInUints * sizeof(uint32);
-	return FSubstrateViewDebugData::FTransientSystemInfoDebugBuffer(Out);
+	return Out;
 }
 
 void FSubstrateViewDebugData::SafeRelease()
