Fixed crash due to bad call to GenerateQuarterResDepthMinMaxTexture.

#tests PC editor + xsx
#rnx

diff --git a/Engine/Source/Runtime/Renderer/Private/DeferredShadingRenderer.cpp b/Engine/Source/Runtime/Renderer/Private/DeferredShadingRenderer.cpp
--- a/Engine/Source/Runtime/Renderer/Private/DeferredShadingRenderer.cpp
+++ b/Engine/Source/Runtime/Renderer/Private/DeferredShadingRenderer.cpp
@@ -3036,15 +3036,15 @@ void FDeferredShadingSceneRenderer::Render(FRDGBuilder& GraphBuilder, const FSce
 
 		GraphBuilder.FlushSetupQueue();
 
-		if (bShouldRenderVolumetricCloud && IsVolumetricRenderTargetEnabled() && HalfResolutionDepthCheckerboardMinMaxTexture==nullptr && !bHasRayTracedOverlay)
-		{
-			GenerateQuarterResDepthMinMaxTexture(GraphBuilder, Views, SceneTextures.Depth.Resolve);
-		}
-
 		if (bShouldRenderVolumetricCloud && !bHasRayTracedOverlay)
 		{
 			if (!bAsyncComputeVolumetricCloud)
 			{
+				if(IsVolumetricRenderTargetEnabled())
+				{
+					GenerateQuarterResDepthMinMaxTexture(GraphBuilder, Views, SceneTextures.Depth.Resolve);
+				}
+
 				// Generate the volumetric cloud render target
 				bool bSkipVolumetricRenderTarget = false;
 				bool bSkipPerPixelTracing = true;
