Volumetric cloud missing file from 41567321

#rnx

diff --git a/Engine/Source/Runtime/Renderer/Private/VolumetricCloudRendering.cpp b/Engine/Source/Runtime/Renderer/Private/VolumetricCloudRendering.cpp
--- a/Engine/Source/Runtime/Renderer/Private/VolumetricCloudRendering.cpp
+++ b/Engine/Source/Runtime/Renderer/Private/VolumetricCloudRendering.cpp
@@ -88,7 +88,7 @@ static TAutoConsoleVariable<int32> CVarVolumetricCloudApplyFogOnAllPixel(
 // - So ApplyFogLate=1 allows to to have fog on clouds applied in (4). Which is fine for out single depth operation.
 static TAutoConsoleVariable<int32> CVarVolumetricCloudApplyFogLate(
 	TEXT("r.VolumetricCloud.ApplyFogLate"), 0,
-	TEXT("When true, fog will be applied on cloud later in the frame. This applies fog on cloud without visible latency due to reprojection and reconstruction. It is a little bit more expenssive but the only correct look when Lumen is enabled and tracing is run on async compute using r.VolumetricRenderTarget.PreferAsyncCompute=1."),
+	TEXT("When true, fog will be applied on cloud later in the frame. This applies fog on cloud without visible latency due to reprojection and reconstruction. It is a bit more expenssive (upsampling pass cost is roughly x2) but the only correct look when Lumen is enabled and tracing is run on async compute using r.VolumetricRenderTarget.PreferAsyncCompute=1."),
 	ECVF_RenderThreadSafe);
 
 ////////////////////////////////////////////////////////////////////////// Shadow tracing
