Path Tracing - added r.PathTracing.InvalidateOnMaterialParameterChange which allows to ignore material parameter changes during path tracing output accumulation

diff --git a/Engine/Source/Runtime/Engine/Private/Materials/MaterialInstance.cpp b/Engine/Source/Runtime/Engine/Private/Materials/MaterialInstance.cpp
--- a/Engine/Source/Runtime/Engine/Private/Materials/MaterialInstance.cpp
+++ b/Engine/Source/Runtime/Engine/Private/Materials/MaterialInstance.cpp
@@ -552,7 +552,7 @@ void GameThread_UpdateMIParameter(const UMaterialInstance* Instance, const Param
 			// check if this material has any relevance to path tracing
 			if (Domain != MD_PostProcess && Domain != MD_UI && !Material->bUsedWithEditorCompositing)
 			{
-				GetRendererModule().InvalidatePathTracedOutput();
+				GetRendererModule().InvalidatePathTracedOutput(PathTracing::EInvalidateReason::UpdateMaterialParameter);
 			}
 		}
 		FMaterialInstanceResource* Resource = Instance->Resource;
diff --git a/Engine/Source/Runtime/Engine/Public/SceneInterface.h b/Engine/Source/Runtime/Engine/Public/SceneInterface.h
--- a/Engine/Source/Runtime/Engine/Public/SceneInterface.h
+++ b/Engine/Source/Runtime/Engine/Public/SceneInterface.h
@@ -11,6 +11,7 @@
 #include "PrimitiveComponentId.h"
 #include "SceneViewOwner.h"
 #include "ShowFlags.h"
+#include "PathTracingOutputInvalidateReason.h"
 
 class AWorldSettings;
 struct FAnimBankDesc;
@@ -308,7 +309,7 @@ public:
 	virtual void RequestPreloadRuntimeVirtualTexture(class URuntimeVirtualTextureComponent* Component, FBoxSphereBounds const& WorldBounds, int32 Level) {}
 
 	/** Mark scene as needing to restart path tracer accumulation. */
-	virtual void InvalidatePathTracedOutput() {}
+	virtual void InvalidatePathTracedOutput(PathTracing::EInvalidateReason InvalidateReason = PathTracing::EInvalidateReason::Uncategorized) {}
 
 	/**  Invalidates Lumen surface cache and forces it to be refreshed. Useful to make material updates more responsive. */
 	virtual void InvalidateLumenSurfaceCache_GameThread(UPrimitiveComponent* Component) {}
diff --git a/Engine/Source/Runtime/RenderCore/Public/PathTracingOutputInvalidateReason.h b/Engine/Source/Runtime/RenderCore/Public/PathTracingOutputInvalidateReason.h
new file mode 100644
--- /dev/null
+++ b/Engine/Source/Runtime/RenderCore/Public/PathTracingOutputInvalidateReason.h
@@ -0,0 +1,12 @@
+// Copyright Epic Games, Inc. All Rights Reserved.
+
+#pragma once
+
+namespace PathTracing
+{
+	enum class EInvalidateReason
+	{
+		UpdateMaterialParameter,
+		Uncategorized
+	};
+}
\ No newline at end of file
diff --git a/Engine/Source/Runtime/RenderCore/Public/RendererInterface.h b/Engine/Source/Runtime/RenderCore/Public/RendererInterface.h
--- a/Engine/Source/Runtime/RenderCore/Public/RendererInterface.h
+++ b/Engine/Source/Runtime/RenderCore/Public/RendererInterface.h
@@ -17,6 +17,7 @@
 #include "UniformBuffer.h"
 #include "VirtualTexturing.h"
 #include "RenderGraphDefinitions.h"
+#include "PathTracingOutputInvalidateReason.h"
 
 class FCanvas;
 class FCanvasRenderContext;
@@ -867,7 +868,7 @@ public:
 	virtual IScenePrimitiveRenderingContext* BeginScenePrimitiveRendering(FRDGBuilder& GraphBuilder, FSceneInterface& Scene) = 0;
 
 	/** Mark all the current scenes as needing to restart path tracer accumulation */
-	virtual void InvalidatePathTracedOutput() = 0;
+	virtual void InvalidatePathTracedOutput(PathTracing::EInvalidateReason InvalidateReason = PathTracing::EInvalidateReason::Uncategorized) = 0;
 
 	/** Experimental:  Render multiple view families in a single scene render call.  All families must reference the same FScene.  Scene Capture not yet supported. */
 	virtual void BeginRenderingViewFamilies(FCanvas* Canvas, TConstArrayView<FSceneViewFamily*> ViewFamilies) = 0;
diff --git a/Engine/Source/Runtime/Renderer/Private/PathTracing.cpp b/Engine/Source/Runtime/Renderer/Private/PathTracing.cpp
--- a/Engine/Source/Runtime/Renderer/Private/PathTracing.cpp
+++ b/Engine/Source/Runtime/Renderer/Private/PathTracing.cpp
@@ -576,6 +576,23 @@ TAutoConsoleVariable<int32> CVarPathTracingDebug(
 	ECVF_RenderThreadSafe | ECVF_ReadOnly
 );
 
+TAutoConsoleVariable<bool> CVarPathTracingInvalidateOnMaterialParameterChange(
+	TEXT("r.PathTracing.InvalidateOnMaterialParameterChange"),
+	true,
+	TEXT("Whether to invalidate path tracer output and restart accumulaion when a material parameter changes."),
+	ECVF_RenderThreadSafe
+);
+
+bool PathTracing::IsOutputInvalidateAllowed(EInvalidateReason InvalidateReason)
+{
+	if (InvalidateReason == EInvalidateReason::UpdateMaterialParameter)
+	{
+		return CVarPathTracingInvalidateOnMaterialParameterChange.GetValueOnAnyThread();
+	}
+
+	return true;
+}
+
 BEGIN_SHADER_PARAMETER_STRUCT(FPathTracingData, )
 	SHADER_PARAMETER(float, BlendFactor)
 	SHADER_PARAMETER(float, BaseExposure)
diff --git a/Engine/Source/Runtime/Renderer/Private/PathTracing.h b/Engine/Source/Runtime/Renderer/Private/PathTracing.h
--- a/Engine/Source/Runtime/Renderer/Private/PathTracing.h
+++ b/Engine/Source/Runtime/Renderer/Private/PathTracing.h
@@ -4,6 +4,7 @@
 
 #include "PathTracingResources.h"
 #include "ShaderParameterMacros.h"
+#include "PathTracingOutputInvalidateReason.h"
 
 // this struct holds skylight parameters
 BEGIN_SHADER_PARAMETER_STRUCT(FPathTracingSkylight, )
@@ -40,4 +41,5 @@ namespace PathTracing
 	bool UsesReferenceDOF(const FViewInfo& View);
 	bool NeedsAntiAliasing(const FViewInfo& View);
 	bool NeedsTonemapping();
+	bool IsOutputInvalidateAllowed(EInvalidateReason InvalidateReason);
 }
diff --git a/Engine/Source/Runtime/Renderer/Private/RendererModule.h b/Engine/Source/Runtime/Renderer/Private/RendererModule.h
--- a/Engine/Source/Runtime/Renderer/Private/RendererModule.h
+++ b/Engine/Source/Runtime/Renderer/Private/RendererModule.h
@@ -142,7 +142,7 @@ public:
 	virtual IScenePrimitiveRenderingContext* BeginScenePrimitiveRendering(FRDGBuilder& GraphBuilder, FSceneViewFamily* ViewFamily) override;
 	virtual IScenePrimitiveRenderingContext* BeginScenePrimitiveRendering(FRDGBuilder& GraphBuilder, FSceneInterface& Scene) override;
 
-	virtual void InvalidatePathTracedOutput() override;
+	virtual void InvalidatePathTracedOutput(PathTracing::EInvalidateReason InvalidateReason = PathTracing::EInvalidateReason::Uncategorized) override;
 
 	virtual void BeginRenderingViewFamilies(FCanvas* Canvas, TConstArrayView<FSceneViewFamily*> ViewFamilies) override;
 
diff --git a/Engine/Source/Runtime/Renderer/Private/RendererScene.cpp b/Engine/Source/Runtime/Renderer/Private/RendererScene.cpp
--- a/Engine/Source/Runtime/Renderer/Private/RendererScene.cpp
+++ b/Engine/Source/Runtime/Renderer/Private/RendererScene.cpp
@@ -104,6 +104,7 @@
 #include "Skinning/AnimBankTransformProvider.h"
 #include "Skinning/SkinningTransformProvider.h"
 #include "SceneViewStateSystemMemory.h"
+#include "PathTracing.h"
 
 #if RHI_RAYTRACING
 #include "Nanite/NaniteRayTracing.h"
@@ -4023,10 +4024,13 @@ void FScene::RequestPreloadRuntimeVirtualTexture(class URuntimeVirtualTextureCom
 	}
 }
 
-void FScene::InvalidatePathTracedOutput()
+void FScene::InvalidatePathTracedOutput(PathTracing::EInvalidateReason InvalidateReason)
 {
-	// NOTE: this is an atomic, so this function is ok to call from any thread
-	++PathTracingInvalidationCounter;
+	if (PathTracing::IsOutputInvalidateAllowed(InvalidateReason))
+	{
+		// NOTE: this is an atomic, so this function is ok to call from any thread
+		++PathTracingInvalidationCounter;
+	}
 }
 
 void FScene::InvalidateLumenSurfaceCache_GameThread(UPrimitiveComponent* Component)
@@ -7623,7 +7627,7 @@ TUniquePtr<ISceneRenderBuilder> FRendererModule::CreateSceneRenderBuilder(FScene
 	return MakeUnique<FSceneRenderBuilder>(SceneInterface);
 }
 
-void FRendererModule::InvalidatePathTracedOutput()
+void FRendererModule::InvalidatePathTracedOutput(PathTracing::EInvalidateReason InvalidateReason)
 {
 	// AllocatedScenes is managed by the game thread
 
@@ -7635,7 +7639,7 @@ void FRendererModule::InvalidatePathTracedOutput()
 	check(IsInGameThread() || IsInParallelGameThread());
 	for (FSceneInterface* Scene : AllocatedScenes)
 	{
-		Scene->InvalidatePathTracedOutput();
+		Scene->InvalidatePathTracedOutput(InvalidateReason);
 	}
 }
 
diff --git a/Engine/Source/Runtime/Renderer/Private/ScenePrivate.h b/Engine/Source/Runtime/Renderer/Private/ScenePrivate.h
--- a/Engine/Source/Runtime/Renderer/Private/ScenePrivate.h
+++ b/Engine/Source/Runtime/Renderer/Private/ScenePrivate.h
@@ -3416,7 +3416,7 @@ public:
 	virtual void GetRuntimeVirtualTextureHidePrimitiveMask(uint8& bHideMaskEditor, uint8& bHideMaskGame) const override;
 	virtual void InvalidateRuntimeVirtualTexture(class URuntimeVirtualTextureComponent* Component, FBoxSphereBounds const& WorldBounds, EVTInvalidatePriority InvalidatePriority) override;
 	virtual void RequestPreloadRuntimeVirtualTexture(class URuntimeVirtualTextureComponent* Component, FBoxSphereBounds const& WorldBounds, int32 Level) override;
-	virtual void InvalidatePathTracedOutput() override;
+	virtual void InvalidatePathTracedOutput(PathTracing::EInvalidateReason InvalidateReason = PathTracing::EInvalidateReason::Uncategorized) override;
 	virtual void InvalidateLumenSurfaceCache_GameThread(UPrimitiveComponent* Component) override;
 	virtual void GetPrimitiveUniformShaderParameters_RenderThread(const FPrimitiveSceneInfo* PrimitiveSceneInfo, bool& bHasPrecomputedVolumetricLightmap, FMatrix& PreviousLocalToWorld, int32& SingleCaptureIndex, bool& bOutputVelocity) const override;
 	virtual void UpdateLightTransform(ULightComponent* Light) override;
