Preventing crash by temporarily disabling overlapping volume support

diff --git a/Engine/Source/Runtime/Renderer/Private/HeterogeneousVolumes/HeterogeneousVolumesLiveShadingPipeline.cpp b/Engine/Source/Runtime/Renderer/Private/HeterogeneousVolumes/HeterogeneousVolumesLiveShadingPipeline.cpp
--- a/Engine/Source/Runtime/Renderer/Private/HeterogeneousVolumes/HeterogeneousVolumesLiveShadingPipeline.cpp
+++ b/Engine/Source/Runtime/Renderer/Private/HeterogeneousVolumes/HeterogeneousVolumesLiveShadingPipeline.cpp
@@ -142,8 +142,13 @@ namespace HeterogeneousVolumes
 
 	bool SupportsOverlappingVolumes()
 	{
+#if 1
+		// UE-268733: SupportOverlappingVolumes current causes crash
+		return false;
+#else
 		return ShouldCompositeHeterogeneousVolumesWithTranslucency() &&
 			(CVarHeterogeneousVolumesSupportOverlappingVolumes.GetValueOnRenderThread() != 0);
+#endif
 	}
 
 	bool UseExistenceMask()
