[SSS] Ignore subsurface tiles if all pixels' scattering is inside itself. This reduces a lot of cost if large SSP coverage is used.
* In the tiling pass, we record two additional tile type `Passthrough` for subpixel SSS, and `AllNonPassthrough` that we infer 'AFIS' and 'Separable'.
* The setup pass becomes two indirect passes:
1. Main. Use AllNonPassthrough tiles as usual.
2. Passthrough. subpixel SSS are directly written to the last texture before recombine.
* Add CVar `r.SSS.Subpixel.Threshold` (default 1) as a control used to determine how a pixel is regarded as Subpixel SSS. E.g., Threshold=0.995 means 99.5% scattering is inside one pixel, we ignore scattering logic.

#tests test pass engine test on PC and internal projects.

diff --git a/Engine/Shaders/Private/BurleyNormalizedSSSCommon.ush b/Engine/Shaders/Private/BurleyNormalizedSSSCommon.ush
--- a/Engine/Shaders/Private/BurleyNormalizedSSSCommon.ush
+++ b/Engine/Shaders/Private/BurleyNormalizedSSSCommon.ush
@@ -156,6 +156,95 @@ float3 GetCDF3D(float3 D, float X)
 	return 1 - 0.25 * exp(-X / D) - 0.75 * exp(-X / (3 * D));
 }
 
+float GetCDF(float D, float X, float XI)
+{
+	return 1 - 0.25*exp(-X / D) - 0.75*exp(-X / (3 * D)) - XI;
+}
+float GetCDFDeriv1(float D, float X)
+{
+	return 0.25 / D * (exp(-X / D) + exp(-X / (3 * D)));
+}
+float GetCDFDeriv1InversD(float InvD, float X)
+{
+	return 0.25 * InvD * (exp(-X*InvD)+exp(-3*X*InvD));
+}
+
+float GetCDFDeriv2(float D, float X)
+{
+	return exp(-X / D)*(-0.0833333*exp(2 * X / (3 * D)) - 0.25) / (D*D);
+}
+
+float RadiusRootFindAnalytic(float D, float RandomNumber)
+{
+	// solve for:
+	//     F(X) = .25*exp(-X) + .75*exp(-X/3)
+	//     cdf(X) = 1 - F(X/D)
+	// so to find X s.t.
+	//     cdf(X) = RandomNumber
+	// we need to solve:
+	//     1 - RandomNumber = F(X/D)
+	//     Finv(1 - RandomNumber) * D = X
+	// Finv(X) actually has an analytic solution:
+	//     G = 1 + 4*U*(2*U + sqrt(1+4*U*U));
+	//     Inv3 = 1.0 / 3.0;
+	//     X/D = 3 * log((1+pow(G,-Inv3)+pow(G,Inv3))/(4*U));
+
+	float U = 1 - RandomNumber;
+	float G = 1.0f + 4*U*(2*U + sqrt(1+4*U*U));
+	float Inv3 = 1.0f / 3.0f;
+	float X = D * (3 * log( (1+pow(G,-Inv3)+pow(G,Inv3))/(4*U)));
+	return X;
+}
+
+// Without clamp, the result will be zero which will lead to infinity for R/pdf. and Nan for the weighted sum.
+// We need to clamp pdf to a very small value. this clamp is based on value below
+#define CLAMP_PDF 0.00001
+
+// Given D and a random number, use root finding method to calculate the radius
+// in meters
+float RadiusRootFinding(float D, float RandomNumber, float X0)
+{
+	// Make the distribution correct.
+	// r=cdf^(-1)(p)     if p<1-CLAMP_PDF
+	//  =cdf^(-1)(1-ClAMP_PDF) P>=1-CLAMP_PDF
+	// RandomNumber = clamp(RandomNumber,0,1-CLAMP_PDF);
+
+	const int NumOfIteration = 10;
+	float Xn = X0;
+
+	UNROLL for (int i = 0; i < NumOfIteration; ++i)
+	{
+		float Fx = GetCDF(D, Xn, RandomNumber);
+		float DFx = GetCDFDeriv1(D, Xn);
+		float DFxx = GetCDFDeriv2(D, Xn);
+		Xn = Xn - 2 * Fx*DFx / (2 * DFx*DFx - Fx * DFxx);
+	}
+
+	return Xn;
+}
+
+// Brian's approximation.
+float RadiusRootFindByApproximation(float D, float RandomNumber)
+{
+	return D * ((2 - 2.6)*RandomNumber - 2)*log(1 - RandomNumber);
+}
+
+// Get the probability to sample a disk.
+float GetPdf(float Radius, float L, float S)
+{
+	//without clamp, the result will be zero which will lead to infinity for R/pdf. and Nan for the weighted sum.
+	//we need to clamp this to a very small pdf.
+
+	float Pdf = GetCDFDeriv1(L / S, Radius);
+	return  max(Pdf, CLAMP_PDF);
+}
+
+float GetPdfInvD(float Radius, float InvD)
+{
+	float Pdf = GetCDFDeriv1InversD(InvD,Radius);
+	return max(Pdf, CLAMP_PDF);
+}
+
 //@TODO Revisit here after offline comparison.
 //call this function in TransmissionCommon.ush for Debug. No transmission tint is applied.
 
diff --git a/Engine/Shaders/Private/PostProcessSubsurface.usf b/Engine/Shaders/Private/PostProcessSubsurface.usf
--- a/Engine/Shaders/Private/PostProcessSubsurface.usf
+++ b/Engine/Shaders/Private/PostProcessSubsurface.usf
@@ -15,9 +15,6 @@
 
 #include "ScreenPass.ush"
 
-// x:Radius*DistanceToProjectionWindow/KernelSize*0.5, y:DistanceToProjectionWindow, z:OverrideNumSamples, w: MinGenerateMipsTileCount
-float4 SubsurfaceParams;
-
 Texture2D SubsurfaceInput0_Texture;
 Texture2D SubsurfaceInput1_Texture;
 Texture2D SubsurfaceInput2_Texture;
@@ -47,6 +44,10 @@ SCREEN_PASS_TEXTURE_VIEWPORT(Output)
 
 RWTexture2D<uint> ProfileIdTexture;
 
+//SUBSURFACE_SETUP_PASS
+#define SUBSURFACE_SETUP_PASS_MAIN        0
+#define SUBSURFACE_SETUP_PASS_PASSTHROUGH 1
+
 #define SUBSURFACE_PASS_ONE 0
 #define SUBSURFACE_PASS_TWO 1
 
@@ -84,6 +85,10 @@ RWTexture2D<uint> ProfileIdTexture;
 	#define SUBSURFACE_RECOMBINE 0
 #endif
 
+#ifndef SUBSURFACE_SETUP_PASS
+    #define SUBSURFACE_SETUP_PASS 0
+#endif
+
 #ifndef CHECKERBOARD_NEIGHBOR_SSS_VALIDATION
 	#define CHECKERBOARD_NEIGHBOR_SSS_VALIDATION SUBSURFACE_RECOMBINE
 	#if CHECKERBOARD_NEIGHBOR_SSS_VALIDATION
@@ -190,8 +195,6 @@ float4 GetSceneColor(float2 BufferUV)
 	return Texture2DSampleLevel(SubsurfaceInput0_Texture, SharedSubsurfaceSampler0, BufferUV, 0);
 }
 
-#include "BurleyNormalizedSSSCommon.ush"
-
 #if SUBTRATE_GBUFFER_FORMAT==1
 
 // Calculate the Subsurface profile kernel on the fly using integration instead of LUT.
@@ -243,64 +246,6 @@ half4 GetSubsurfaceProfileKernel(int SampleIndex, float3 D, float ScatterRadiusI
 }
 #endif
 
-// fetch surface albedo and diffuse mean free path
-// The diffuse mean free path is modulated by world scale and a custom term passed in
-// which is the opacity texture.
-#if SUBTRATE_GBUFFER_FORMAT==1
-FBurleyParameter GetBurleyParameters(FSubstrateSubsurfaceData SSSData)
-{
-	FBurleyParameter Out = (FBurleyParameter)0;
-	if (SubstrateSubSurfaceHeaderGetIsProfile(SSSData.Header))
-	{
-		uint ProfileId = SubstrateSubSurfaceHeaderGetProfileId(SSSData.Header);
-		float ProfileRadiusScale = SubstrateSubSurfaceHeaderGetProfileRadiusScale(SSSData.Header);
-
-		Out.SurfaceAlbedo = GetSubsurfaceProfileSurfaceAlbedo(ProfileId);
-		Out.DiffuseMeanFreePath = DecodeDiffuseMeanFreePath(GetSubsurfaceProfileDiffuseMeanFreePath(ProfileId));
-		Out.WorldUnitScale = DecodeWorldUnitScale(GetSubsurfaceProfileWorldUnitScale(ProfileId)) * BURLEY_CM_2_MM;
-
-		Out.SurfaceOpacity = ProfileRadiusScale;
-		Out.DiffuseMeanFreePath *= Out.SurfaceOpacity;
-	}
-	else if (SubstrateSubSurfaceHeaderGetUseDiffusion(SSSData.Header))
-	{
-		float3 SSSMFP = SubstrateSubSurfaceHeaderGetMFP(SSSData.Header);
-		float3 SSSBaseColor = SubstrateSubsurfaceExtrasGetBaseColor(SSSData.Extras);
-
-		// Decode the MFP stored as Float111110 in a uint over a RGBA8 unorm
-		const float3 MFP = SSSMFP * BURLEY_CM_2_MM; // Stored DMFP is in cm, converted into mm
-
-		const float MaxMFP = max3(MFP.x, MFP.y, MFP.z);
-		const uint IndexOfMaxMfp = (MFP.r == MaxMFP) ? 0 : ((MFP.g == MaxMFP) ? 1 : 2);
-
-		const float3 DMFP = GetDMFPFromMFPApprox(SSSBaseColor, MFP);
-
-		Out.SurfaceAlbedo = float4(SSSBaseColor, SSSBaseColor[IndexOfMaxMfp]);
-		Out.DiffuseMeanFreePath = float4(DMFP, DMFP[IndexOfMaxMfp]);
-		Out.WorldUnitScale = 1.0f;	// no world unit scale change.
-
-		// SUBSTRATE_TODO smooth transistion based on pixel footprint w.r.t. MFPRadius and pixel footprint
-		Out.SurfaceOpacity = 1.0f;
-		Out.DiffuseMeanFreePath *= Out.SurfaceOpacity;
-	}
-	return Out;
-}
-#else
-FBurleyParameter GetBurleyParameters(uint SubsurfaceProfileInt, FGBufferData GBuffer)
-{
-	FBurleyParameter Out = (FBurleyParameter)0;
-	Out.SurfaceAlbedo = GetSubsurfaceProfileSurfaceAlbedo(SubsurfaceProfileInt);
-	Out.DiffuseMeanFreePath = DecodeDiffuseMeanFreePath(GetSubsurfaceProfileDiffuseMeanFreePath(SubsurfaceProfileInt));
-	Out.WorldUnitScale = DecodeWorldUnitScale(GetSubsurfaceProfileWorldUnitScale(SubsurfaceProfileInt)) * BURLEY_CM_2_MM;
-
-	// Mask should always be true, we can probably ignore it?
-	Out.SurfaceOpacity = UseSubsurfaceProfile(GBuffer.ShadingModelID) ? GBuffer.CustomData.a : 0.0f;
-	Out.DiffuseMeanFreePath *= Out.SurfaceOpacity;
-
-	return Out;
-}
-#endif
-
 // BufferUV back to grid posistion of the input
 uint2 ConvertUVPos2Input0GridPos(float2 BufferUV)
 {
@@ -672,7 +617,10 @@ void FillMipsConditionBufferCS()
 [numthreads(8, 1, 1)]
 void BuildIndirectDispatchArgsCS(uint2 DispatchThreadId : SV_DispatchThreadID)
 {
-	if (DispatchThreadId.y == 0 && DispatchThreadId.x < SUBSURFACE_TILE_TYPE_COUNT && DispatchThreadId.x != SUBSURFACE_TILE_TYPE_ALL)
+	if (DispatchThreadId.y == 0 && DispatchThreadId.x < SUBSURFACE_TILE_TYPE_COUNT 
+		&& DispatchThreadId.x != SUBSURFACE_TILE_TYPE_PASSTHROUGH
+		&& DispatchThreadId.x != SUBSURFACE_TILE_TYPE_ALLNONPASSTHROUGH
+		&& DispatchThreadId.x != SUBSURFACE_TILE_TYPE_ALL)
 	{
 		uint TileType = DispatchThreadId.x;
 		const uint TileCount = TileTypeCountBuffer[TileType];
@@ -909,6 +857,10 @@ void SetupIndirectCS(
 	if (all(Pos < Output_ViewportMax))
 	{
 		Type |= TYPE_IN_VIEWPORT;
+#if (SUBSURFACE_SETUP_PASS == SUBSURFACE_SETUP_PASS_PASSTHROUGH)
+		//write to the final buffer, need mask to recombine.
+		OutColor.a = GetMaskFromDepthInAlpha(OutColor.a);
+#endif
 		SetupTexture[Pos] = OutColor;
 
 #if ENABLE_PROFILE_ID_CACHE
@@ -916,7 +868,7 @@ void SetupIndirectCS(
 #endif
 	}
 
-#if !SUBSURFACE_FORCE_SEPARABLE
+#if !(SUBSURFACE_FORCE_SEPARABLE || (SUBSURFACE_SETUP_PASS == SUBSURFACE_SETUP_PASS_PASSTHROUGH))
 	AddSubsurfaceComputeGroup(GI, Type, G_ID.xy);
 #endif
 }
diff --git a/Engine/Shaders/Private/PostProcessSubsurfaceCommon.ush b/Engine/Shaders/Private/PostProcessSubsurfaceCommon.ush
--- a/Engine/Shaders/Private/PostProcessSubsurfaceCommon.ush
+++ b/Engine/Shaders/Private/PostProcessSubsurfaceCommon.ush
@@ -8,15 +8,129 @@
 
 #include "Common.ush"
 #include "SubsurfaceProfileCommon.ush"
+#include "BurleyNormalizedSSSCommon.ush"
 
 // the same to FSubsurfaceTiles::ETileType
 #define SUBSURFACE_TILE_TYPE_AFIS		0
 #define SUBSURFACE_TILE_TYPE_SEPARABLE	1
-#define SUBSURFACE_TILE_TYPE_ALL		2
-#define SUBSURFACE_TILE_TYPE_COUNT		3
+#define SUBSURFACE_TILE_TYPE_PASSTHROUGH	2
+#define SUBSURFACE_TILE_TYPE_ALLNONPASSTHROUGH	3
+#define SUBSURFACE_TILE_TYPE_ALL		4
+#define SUBSURFACE_TILE_TYPE_COUNT		5
 
 #define TYPE_SEPARABLE   0x1
 #define TYPE_BURLEY      0x2
 #define TYPE_IN_VIEWPORT 0x4
 
-groupshared uint SubsurfaceTypeFlag;
\ No newline at end of file
+// x:Radius*DistanceToProjectionWindow/KernelSize*0.5, y:DistanceToProjectionWindow, z:OverrideNumSamples, w: MinGenerateMipsTileCount
+float4 SubsurfaceParams;
+
+groupshared uint SubsurfaceTypeFlag;
+
+// fetch surface albedo and diffuse mean free path
+// The diffuse mean free path is modulated by world scale and a custom term passed in
+// which is the opacity texture.
+#if SUBTRATE_GBUFFER_FORMAT==1
+FBurleyParameter GetBurleyParameters(FSubstrateSubsurfaceData SSSData)
+{
+	FBurleyParameter Out = (FBurleyParameter)0;
+	if (SubstrateSubSurfaceHeaderGetIsProfile(SSSData.Header))
+	{
+		uint ProfileId = SubstrateSubSurfaceHeaderGetProfileId(SSSData.Header);
+		float ProfileRadiusScale = SubstrateSubSurfaceHeaderGetProfileRadiusScale(SSSData.Header);
+
+		Out.SurfaceAlbedo = GetSubsurfaceProfileSurfaceAlbedo(ProfileId);
+		Out.DiffuseMeanFreePath = DecodeDiffuseMeanFreePath(GetSubsurfaceProfileDiffuseMeanFreePath(ProfileId));
+		Out.WorldUnitScale = DecodeWorldUnitScale(GetSubsurfaceProfileWorldUnitScale(ProfileId)) * BURLEY_CM_2_MM;
+
+		Out.SurfaceOpacity = ProfileRadiusScale;
+		Out.DiffuseMeanFreePath *= Out.SurfaceOpacity;
+	}
+	else if (SubstrateSubSurfaceHeaderGetUseDiffusion(SSSData.Header))
+	{
+		float3 SSSMFP = SubstrateSubSurfaceHeaderGetMFP(SSSData.Header);
+		float3 SSSBaseColor = SubstrateSubsurfaceExtrasGetBaseColor(SSSData.Extras);
+
+		// Decode the MFP stored as Float111110 in a uint over a RGBA8 unorm
+		const float3 MFP = SSSMFP * BURLEY_CM_2_MM; // Stored DMFP is in cm, converted into mm
+
+		const float MaxMFP = max3(MFP.x, MFP.y, MFP.z);
+		const uint IndexOfMaxMfp = (MFP.r == MaxMFP) ? 0 : ((MFP.g == MaxMFP) ? 1 : 2);
+
+		const float3 DMFP = GetDMFPFromMFPApprox(SSSBaseColor, MFP);
+
+		Out.SurfaceAlbedo = float4(SSSBaseColor, SSSBaseColor[IndexOfMaxMfp]);
+		Out.DiffuseMeanFreePath = float4(DMFP, DMFP[IndexOfMaxMfp]);
+		Out.WorldUnitScale = 1.0f;	// no world unit scale change.
+
+		// SUBSTRATE_TODO smooth transistion based on pixel footprint w.r.t. MFPRadius and pixel footprint
+		Out.SurfaceOpacity = 1.0f;
+		Out.DiffuseMeanFreePath *= Out.SurfaceOpacity;
+	}
+	return Out;
+}
+#else
+FBurleyParameter GetBurleyParameters(uint SubsurfaceProfileInt, FGBufferData GBuffer)
+{
+	FBurleyParameter Out = (FBurleyParameter)0;
+	Out.SurfaceAlbedo = GetSubsurfaceProfileSurfaceAlbedo(SubsurfaceProfileInt);
+	Out.DiffuseMeanFreePath = DecodeDiffuseMeanFreePath(GetSubsurfaceProfileDiffuseMeanFreePath(SubsurfaceProfileInt));
+	Out.WorldUnitScale = DecodeWorldUnitScale(GetSubsurfaceProfileWorldUnitScale(SubsurfaceProfileInt)) * BURLEY_CM_2_MM;
+
+	// Mask should always be true, we can probably ignore it?
+	Out.SurfaceOpacity = UseSubsurfaceProfile(GBuffer.ShadingModelID) ? GBuffer.CustomData.a : 0.0f;
+	Out.DiffuseMeanFreePath *= Out.SurfaceOpacity;
+
+	return Out;
+}
+#endif
+
+// The extent is for the subsurface texture.
+float2 CalculateBurleyScale(float WorldUnitScale, float DepthAtCenter, float2 ViewportSize, float2 Extent, float2 ExtentInverse)
+{
+	float2 BurleyScale = WorldUnitScale;
+
+	float SSSScaleX = SubsurfaceParams.x;
+	BurleyScale *= SSSScaleX / DepthAtCenter;
+
+	// cast from cm to mm for depth, and remove the effect of SUBSURFACE_KERNEL_SIZE. 
+	BurleyScale *= SUBSURFACE_KERNEL_SIZE / BURLEY_CM_2_MM;
+
+	// account for Screen Percentage/Dyanmic Resolution Scaling
+	BurleyScale *= (ViewportSize.x * ExtentInverse.x);
+	BurleyScale.y *= (Extent.x * ExtentInverse.y);
+
+	return BurleyScale;
+}
+
+// Given the Depth and the BurleyParameter, figure out the actual radius of the center pixel in MM,
+// taking into account the depth and screen dimensions.
+float CalculateCenterSampleRadiusInMM(FBurleyParameter BurleyParameter, float2 BurleyScale, float2 ExtentInverse)
+{
+	float DiffuseMeanFreePath = GetDiffuseMeanFreePathForSampling(BurleyParameter.DiffuseMeanFreePath);
+
+	float A = GetComponentForScalingFactorEstimation(BurleyParameter.SurfaceAlbedo);
+	float S = GetScalingFactor(A);
+	float3 S3D = GetScalingFactor3D(BurleyParameter.SurfaceAlbedo.xyz);
+
+	// In the reference function, UVOffset = BurleyScale * RadiusInMM
+	//      float2 UVOffset = BurleyScale*BurleySampleInfo.RadiusInMM;
+	// So, given the UV offset, we can find the distance in mm as:
+	//      float DistInMM = UvOffset.x/BurleyScale.x + UvOffset.y/BurleyScale.y;
+	// But for stability, we can just average them.
+	float CenterSampleRadiusInMM = 0.5f * (ExtentInverse.x/BurleyScale.x + ExtentInverse.y/BurleyScale.y);
+
+	return CenterSampleRadiusInMM;
+}
+
+float CalculateCenterSampleCdf(FBurleyParameter BurleyParameter, float CenterSampleRadiusInMM)
+{
+	float DiffuseMeanFreePathForSampling = GetDiffuseMeanFreePathForSampling(BurleyParameter.DiffuseMeanFreePath);
+	float A = GetComponentForScalingFactorEstimation(BurleyParameter.SurfaceAlbedo);
+	float S = GetScalingFactor(A);
+
+	float D = DiffuseMeanFreePathForSampling / S;
+	float CenterSampleRadiusCdf = GetCDF(D.x,CenterSampleRadiusInMM,0);
+
+	return CenterSampleRadiusCdf;
+}
diff --git a/Engine/Shaders/Private/PostProcessSubsurfaceTile.usf b/Engine/Shaders/Private/PostProcessSubsurfaceTile.usf
--- a/Engine/Shaders/Private/PostProcessSubsurfaceTile.usf
+++ b/Engine/Shaders/Private/PostProcessSubsurfaceTile.usf
@@ -126,11 +126,16 @@ void SubsurfaceTileFallbackScreenPassVS(
 SCREEN_PASS_TEXTURE_VIEWPORT(SubsurfaceInput0)
 SCREEN_PASS_TEXTURE_VIEWPORT(Output)
 
-groupshared bool ContainsSSS[SUBSURFACE_TILE_SIZE * SUBSURFACE_TILE_SIZE];
+groupshared uint2 ContainsSSS[SUBSURFACE_TILE_SIZE * SUBSURFACE_TILE_SIZE];
+groupshared uint SubsurfacePassthroughFlag;
 
 int MaxGroupCount;
 int2 TiledViewRes;
+float SubsurfaceSubpixelThreshold;
+
 RWStructuredBuffer<uint> TileMaskBufferOut;
+RWStructuredBuffer<uint> TileMaskPassthroughBufferOut;
+RWStructuredBuffer<uint> TileMaskNonPassthroughBufferOut;
 
 //Duplicate from PostProcessSubsurface.usf
 #if SUBTRATE_GBUFFER_FORMAT==1
@@ -155,7 +160,39 @@ float HasSubsurfaceForOnePixel(float2 UVSceneColor)
 	return select(bHasSSSDiffusion, 1.0f, 0.0f);
 }
 
-// Simpler logic from the SSS::Setup pass
+bool IsEffectiveSSS(float2 BufferUV, float HasSSSDiffusion)
+{
+	if (HasSSSDiffusion == 0.0f)
+	{
+		return false;
+	}
+
+#if SUBTRATE_GBUFFER_FORMAT==1
+	const FSubstrateSubsurfaceData SSSData = LoadSubstrateSSSData(BufferUV);
+	const FSubstrateTopLayerData TopLayerData = LoadSubstrateTopLayerData(BufferUV);
+	const uint SubsurfaceProfileInt = SubstrateSubSurfaceHeaderGetProfileId(SSSData.Header);
+	const FBurleyParameter BurleyParameter = GetBurleyParameters(SSSData);
+#else
+	const FScreenSpaceData ScreenSpaceData = GetScreenSpaceData(BufferUV);
+	const uint SubsurfaceProfileInt = ExtractSubsurfaceProfileInt(ScreenSpaceData.GBuffer);
+	const FBurleyParameter BurleyParameter = GetBurleyParameters(SubsurfaceProfileInt, ScreenSpaceData.GBuffer);
+#endif
+
+	float DepthAtCenter = CalcSceneDepth(BufferUV);
+
+	if (DepthAtCenter <= 0)
+	{
+		return false;
+	}
+
+	float2 BurleyScale = CalculateBurleyScale(BurleyParameter.WorldUnitScale, DepthAtCenter, Output_ViewportSize, Output_Extent, Output_ExtentInverse);
+	float CenterSampleRadiusInMM = CalculateCenterSampleRadiusInMM(BurleyParameter, BurleyScale, Output_ExtentInverse);
+	float CenterSampleCdf = CalculateCenterSampleCdf(BurleyParameter, CenterSampleRadiusInMM);
+
+	// ignore the pixel if the scattering is inside itself.
+	return select(CenterSampleCdf < SubsurfaceSubpixelThreshold, true, false);
+}
+
 [numthreads(SUBSURFACE_TILE_SIZE, SUBSURFACE_TILE_SIZE, 1)]
 void SSSTileCategorisationMarkCS(uint3 DT_ID : SV_DispatchThreadID, 
 	uint3 G_ID : SV_GroupID, 
@@ -165,52 +202,88 @@ void SSSTileCategorisationMarkCS(uint3 DT_ID : SV_DispatchThreadID,
 	if (GI == 0)
 	{
 		SubsurfaceTypeFlag = 0;
+		SubsurfacePassthroughFlag = 0;
 	}
 	GroupMemoryBarrierWithGroupSync();
+
 	uint2 Pos = DT_ID.xy + Output_ViewportMin;
 	float2 BufferUV = Output_ExtentInverse * (float2(Pos) + 0.5f);
 	uint Type = 0;
 	
 #if SUBSURFACE_HALF_RES
-	float A = HasSubsurfaceForOnePixel(min(BufferUV + float2(-0.5, 0.5f) * SubsurfaceInput0_ExtentInverse, SubsurfaceInput0_UVViewportBilinearMax));
-	float B = HasSubsurfaceForOnePixel(min(BufferUV + float2( 0.5, 0.5f) * SubsurfaceInput0_ExtentInverse, SubsurfaceInput0_UVViewportBilinearMax));
-	float C = HasSubsurfaceForOnePixel(min(BufferUV + float2( 0.5,-0.5f) * SubsurfaceInput0_ExtentInverse, SubsurfaceInput0_UVViewportBilinearMax));
-	float D = HasSubsurfaceForOnePixel(min(BufferUV + float2(-0.5,-0.5f) * SubsurfaceInput0_ExtentInverse, SubsurfaceInput0_UVViewportBilinearMax));
+	float2 UVA = min(BufferUV + float2(-0.5, 0.5f) * SubsurfaceInput0_ExtentInverse, SubsurfaceInput0_UVViewportBilinearMax);
+	float2 UVB = min(BufferUV + float2( 0.5, 0.5f) * SubsurfaceInput0_ExtentInverse, SubsurfaceInput0_UVViewportBilinearMax);
+	float2 UVC = min(BufferUV + float2( 0.5,-0.5f) * SubsurfaceInput0_ExtentInverse, SubsurfaceInput0_UVViewportBilinearMax);
+	float2 UVD = min(BufferUV + float2(-0.5,-0.5f) * SubsurfaceInput0_ExtentInverse, SubsurfaceInput0_UVViewportBilinearMax);
+
+	float A = HasSubsurfaceForOnePixel(UVA);
+	float B = HasSubsurfaceForOnePixel(UVB);
+	float C = HasSubsurfaceForOnePixel(UVC);
+	float D = HasSubsurfaceForOnePixel(UVD);
+
 	float Sum = (A + B) + (C + D);
 	// Each alpha value of A/B/C/D should be either 1.0f or 0.0, so a check of 0.5 should be safe
 	bool bHasSubsurface = (Sum >= .50f);
 #else
-	bool bHasSubsurface = bool(HasSubsurfaceForOnePixel(BufferUV));
+	float A = HasSubsurfaceForOnePixel(BufferUV);
+	bool bHasSubsurface = bool(A);
 #endif
 	bHasSubsurface = bHasSubsurface && all(Pos < Output_ViewportMax);
-	bool bWriteTile = false;
+	bool bIsPassthroughPixel = false;
+
+	BRANCH
+	if (bHasSubsurface)
+	{
+#if SUBSURFACE_HALF_RES
+		bool bEffectiveSSSA = IsEffectiveSSS(UVA, A);
+		bool bEffectiveSSSB = IsEffectiveSSS(UVB, B);
+		bool bEffectiveSSSC = IsEffectiveSSS(UVC, C);
+		bool bEffectiveSSSD = IsEffectiveSSS(UVD, D);
+
+		bIsPassthroughPixel = !((bEffectiveSSSA | bEffectiveSSSB | bEffectiveSSSC | bEffectiveSSSD));
+#else
+		bool bEffectiveSSS = IsEffectiveSSS(BufferUV, A);
+		bIsPassthroughPixel = !bEffectiveSSS;
+#endif
+	}
+
+	bool2 bWriteTile = false;
 
 #if COMPILER_SUPPORTS_WAVE_VOTE
-	const bool bAnyActiveSSSPixelsInWave = WaveActiveAnyTrue(bHasSubsurface);
-	if (WaveIsFirstLane() && bAnyActiveSSSPixelsInWave)
+	const uint NumOfActiveSSSPixelsInWave = WaveActiveCountBits(bHasSubsurface);
+	const uint NumOfActivePassthroughPixdelsInWave = WaveActiveCountBits(bIsPassthroughPixel);
+	if (WaveIsFirstLane() && NumOfActiveSSSPixelsInWave)
 	{
 		uint IndexToStore = 0;
-		InterlockedAdd(SubsurfaceTypeFlag, 1, IndexToStore);
+		InterlockedAdd(SubsurfaceTypeFlag, NumOfActiveSSSPixelsInWave, IndexToStore);
 	}
+
+	if (WaveIsFirstLane() && NumOfActivePassthroughPixdelsInWave > 0)
+	{
+		uint IndexToStore = 0;
+		InterlockedAdd(SubsurfacePassthroughFlag, NumOfActivePassthroughPixdelsInWave, IndexToStore);
+	}
+
 	GroupMemoryBarrierWithGroupSync();
 	if (GI == 0)
 	{
-		bWriteTile = SubsurfaceTypeFlag > 0;
+		bWriteTile = bool2(SubsurfaceTypeFlag > 0 ,
+			               SubsurfacePassthroughFlag == SubsurfaceTypeFlag);
 	}
 #else
 
 	const int LinearIndex = GT_ID.y * SUBSURFACE_TILE_SIZE + GT_ID.x;
-	ContainsSSS[LinearIndex] = bHasSubsurface;
+	ContainsSSS[LinearIndex] = uint2(bHasSubsurface,bIsPassthroughPixel);
 	GroupMemoryBarrierWithGroupSync();
 
 	if (LinearIndex < 32) // 8*8 = 64 elements to merge
 	{
-		ContainsSSS[LinearIndex] = ContainsSSS[LinearIndex] || ContainsSSS[LinearIndex + 32];
+		ContainsSSS[LinearIndex] = ContainsSSS[LinearIndex] + ContainsSSS[LinearIndex + 32];
 	}
 	GroupMemoryBarrierWithGroupSync();
 	if (LinearIndex < 16)
 	{
-		ContainsSSS[LinearIndex] = ContainsSSS[LinearIndex] || ContainsSSS[LinearIndex + 16];
+		ContainsSSS[LinearIndex] = ContainsSSS[LinearIndex] + ContainsSSS[LinearIndex + 16];
 	}
 	GroupMemoryBarrierWithGroupSync();
 
@@ -219,35 +292,51 @@ void SSSTileCategorisationMarkCS(uint3 DT_ID : SV_DispatchThreadID,
 
 	if (LinearIndex < 8)
 	{
-		ContainsSSS[LinearIndex] = ContainsSSS[LinearIndex] || ContainsSSS[LinearIndex + 8];
+		ContainsSSS[LinearIndex] = ContainsSSS[LinearIndex] + ContainsSSS[LinearIndex + 8];
 	}
 	GroupMemoryBarrierWithGroupSync();
 	if (LinearIndex < 4)
 	{
-		ContainsSSS[LinearIndex] = ContainsSSS[LinearIndex] || ContainsSSS[LinearIndex + 4];
+		ContainsSSS[LinearIndex] = ContainsSSS[LinearIndex] + ContainsSSS[LinearIndex + 4];
 	}
 	GroupMemoryBarrierWithGroupSync();
 	if (LinearIndex < 2)
 	{
-		ContainsSSS[LinearIndex] = ContainsSSS[LinearIndex] || ContainsSSS[LinearIndex + 2];
+		ContainsSSS[LinearIndex] = ContainsSSS[LinearIndex] + ContainsSSS[LinearIndex + 2];
 	}
 	GroupMemoryBarrierWithGroupSync();
 	if (LinearIndex < 1)
 	{
-		bWriteTile = ContainsSSS[LinearIndex] || ContainsSSS[LinearIndex + 1];
+		bWriteTile.x = ContainsSSS[LinearIndex] > 0;
+		bWriteTile.y = ContainsSSS[LinearIndex + 1] ==  ContainsSSS[LinearIndex];
 	}
 #endif
 
-	if (bWriteTile)
-	{
-		// Set bit to indicate tile is occupied.
-		uint MaskLinearIndex = TiledViewRes.x * G_ID.y + G_ID.x;
+	// Set bit to indicate tile is occupied.
+	uint MaskLinearIndex = TiledViewRes.x * G_ID.y + G_ID.x;
+
+	// All subsurface tiles
+	if (bWriteTile.x)
+	{	
 		InterlockedOr(TileMaskBufferOut[MaskLinearIndex / 32U], 1U << (MaskLinearIndex % 32U));
 	}
+
+	// passthrough tiles
+	if (bWriteTile.x && bWriteTile.y)
+	{
+		InterlockedOr(TileMaskPassthroughBufferOut[MaskLinearIndex / 32U], 1U << (MaskLinearIndex % 32U));
+	}
+
+	// non-passthrough subsurface tiles
+	if (bWriteTile.x && !bWriteTile.y)
+	{
+		InterlockedOr(TileMaskNonPassthroughBufferOut[MaskLinearIndex / 32U], 1U << (MaskLinearIndex % 32U));
+	}
+
 }
 
 int VertexCountPerInstanceIndirect;
-
+int TileType;
 StructuredBuffer<uint> TileMaskBuffer;
 
 RWBuffer<uint> RWSSSTileListData;
@@ -270,8 +359,7 @@ void SSSTileClassificationBuildListsCS(
 	uint2 GroupThreadId : SV_GroupThreadID)
 {
 	const uint LinearThreadIndex = GroupThreadId.y * SUBSURFACE_TILE_SIZE + GroupThreadId.x;
-	const uint TileType = SUBSURFACE_TILE_TYPE_ALL;// Only interested in generating the SSS tiles instead of differentiating the type
-
+	
 	if (all(DispatchThreadId == 0))
 	{
 		const uint TileCount = 0; // Instance count will be accumulated later with locked add.
diff --git a/Engine/Shaders/Private/SubsurfaceBurleyNormalized.ush b/Engine/Shaders/Private/SubsurfaceBurleyNormalized.ush
--- a/Engine/Shaders/Private/SubsurfaceBurleyNormalized.ush
+++ b/Engine/Shaders/Private/SubsurfaceBurleyNormalized.ush
@@ -157,95 +157,6 @@ float2 Generate2DRandomNumber(int3 Seed)
 #endif
 }
 
-float GetCDF(float D, float X, float XI)
-{
-	return 1 - 0.25*exp(-X / D) - 0.75*exp(-X / (3 * D)) - XI;
-}
-float GetCDFDeriv1(float D, float X)
-{
-	return 0.25 / D * (exp(-X / D) + exp(-X / (3 * D)));
-}
-float GetCDFDeriv1InversD(float InvD, float X)
-{
-	return 0.25 * InvD * (exp(-X*InvD)+exp(-3*X*InvD));
-}
-
-float GetCDFDeriv2(float D, float X)
-{
-	return exp(-X / D)*(-0.0833333*exp(2 * X / (3 * D)) - 0.25) / (D*D);
-}
-
-float RadiusRootFindAnalytic(float D, float RandomNumber)
-{
-	// solve for:
-	//     F(X) = .25*exp(-X) + .75*exp(-X/3)
-	//     cdf(X) = 1 - F(X/D)
-	// so to find X s.t.
-	//     cdf(X) = RandomNumber
-	// we need to solve:
-	//     1 - RandomNumber = F(X/D)
-	//     Finv(1 - RandomNumber) * D = X
-	// Finv(X) actually has an analytic solution:
-	//     G = 1 + 4*U*(2*U + sqrt(1+4*U*U));
-	//     Inv3 = 1.0 / 3.0;
-	//     X/D = 3 * log((1+pow(G,-Inv3)+pow(G,Inv3))/(4*U));
-
-	float U = 1 - RandomNumber;
-	float G = 1.0f + 4*U*(2*U + sqrt(1+4*U*U));
-	float Inv3 = 1.0f / 3.0f;
-	float X = D * (3 * log( (1+pow(G,-Inv3)+pow(G,Inv3))/(4*U)));
-	return X;
-}
-
-// Without clamp, the result will be zero which will lead to infinity for R/pdf. and Nan for the weighted sum.
-// We need to clamp pdf to a very small value. this clamp is based on value below
-#define CLAMP_PDF 0.00001
-
-// Given D and a random number, use root finding method to calculate the radius
-// in meters
-float RadiusRootFinding(float D, float RandomNumber, float X0)
-{
-	// Make the distribution correct.
-	// r=cdf^(-1)(p)     if p<1-CLAMP_PDF
-	//  =cdf^(-1)(1-ClAMP_PDF) P>=1-CLAMP_PDF
-	// RandomNumber = clamp(RandomNumber,0,1-CLAMP_PDF);
- 
-	const int NumOfIteration = 10;
-	float Xn = X0;
-
-	UNROLL for (int i = 0; i < NumOfIteration; ++i)
-	{
-		float Fx = GetCDF(D, Xn, RandomNumber);
-		float DFx = GetCDFDeriv1(D, Xn);
-		float DFxx = GetCDFDeriv2(D, Xn);
-		Xn = Xn - 2 * Fx*DFx / (2 * DFx*DFx - Fx * DFxx);
-	}
-
-	return Xn;
-}
-
-// Brian's approximation.
-float RadiusRootFindByApproximation(float D, float RandomNumber)
-{
-	return D * ((2 - 2.6)*RandomNumber - 2)*log(1 - RandomNumber);
-}
-
-// Get the probability to sample a disk.
-float GetPdf(float Radius, float L, float S)
-{
-	//without clamp, the result will be zero which will lead to infinity for R/pdf. and Nan for the weighted sum.
-	//we need to clamp this to a very small pdf.
-
-	float Pdf = GetCDFDeriv1(L / S, Radius);
-	return  max(Pdf, CLAMP_PDF);
-}
-
-float GetPdfInvD(float Radius, float InvD)
-{
-	float Pdf = GetCDFDeriv1InversD(InvD,Radius);
-	return max(Pdf, CLAMP_PDF);
-}
-
 struct FBurleySampleInfo
 {
 	float RadiusInMM;
@@ -255,7 +166,6 @@ struct FBurleySampleInfo
 	float SinTheta;
 };
 
-
 #define FIBONACCI_SEQUENCE_ANGLE(x) (((float(x) + 0.5)*(1 + sqrt(5))*0.5) * 2 * PI)
 
 // angle, cosine, and sine
@@ -851,20 +761,6 @@ float3 CalculateCenterSampleWeight(float Depth, FBurleyParameter BurleyParameter
 	return CenterSampleWeight;
 }
 
-float CalculateCenterSampleCdf(FBurleyParameter BurleyParameter, float Depth)
-{
-	float CenterSampleRadiusInMM = CalculateCenterSampleRadiusInMM(BurleyParameter, Depth);
-
-	float DiffuseMeanFreePathForSampling = GetDiffuseMeanFreePathForSampling(BurleyParameter.DiffuseMeanFreePath);
-	float A = GetComponentForScalingFactorEstimation(BurleyParameter.SurfaceAlbedo);
-	float S = GetScalingFactor(A);
-
-	float D = DiffuseMeanFreePathForSampling / S;
-	float CenterSampleRadiusCdf = GetCDF(D.x,CenterSampleRadiusInMM,0);
-
-	return CenterSampleRadiusCdf;
-}
-
 void UpdateSeed(int3 Seed3D, inout int StartSeed)
 {
 	/*To make R2Sequence work, we need to rebase the R2 sequence start index to a new one uniformly in the R2 space,
@@ -958,7 +854,8 @@ float4 BurleyNormalizedSS(float2 BufferUV, uint2 GroupThreadID)
 	 **************************************************************************************/
 
 #if REWEIGHT_CENTER_SAMPLE
-	float CenterSampleRadiusCdf = CalculateCenterSampleCdf(BurleyParameter, DepthAtDiscCenter);
+	float CenterSampleRadiusInMM = CalculateCenterSampleRadiusInMM(BurleyParameter, BurleyScale, SubsurfaceInput0_ExtentInverse);
+	float CenterSampleRadiusCdf = CalculateCenterSampleCdf(BurleyParameter, CenterSampleRadiusInMM);
 	float3 CenterSampleWeight = CalculateCenterSampleWeight(DepthAtDiscCenter, BurleyParameter);
 #endif
 
diff --git a/Engine/Source/Runtime/Renderer/Private/PostProcess/PostProcessSubsurface.cpp b/Engine/Source/Runtime/Renderer/Private/PostProcess/PostProcessSubsurface.cpp
--- a/Engine/Source/Runtime/Renderer/Private/PostProcess/PostProcessSubsurface.cpp
+++ b/Engine/Source/Runtime/Renderer/Private/PostProcess/PostProcessSubsurface.cpp
@@ -146,6 +146,13 @@ namespace
 		TEXT("4000. (default) The minimal number of tiles to trigger subsurface radiance mip generation. Set to zero to always generate mips (Experimental value)"),
 		ECVF_RenderThreadSafe);
 
+	TAutoConsoleVariable<float> CVarSSSSubpixelThreshold(
+		TEXT("r.SSS.Subpixel.Threshold"),
+		1.0f,
+		TEXT("If a scattering is always within 1 pixel footage, there is no need to scatter. This value determines at which percentage the scattering is considered all inside\n")
+		TEXT("the texel itself. Default is below 1. Once it reaches this value, we stop doing subsurface scattering. The logic is per tile."),
+		ECVF_RenderThreadSafe);
+
 	TAutoConsoleVariable<float> CVarSubSurfaceColorAsTannsmittanceAtDistance(
 		TEXT("r.SSS.SubSurfaceColorAsTansmittanceAtDistance"),
 		0.15f,
@@ -560,6 +567,13 @@ public:
 		MAX
 	};
 
+	enum class ESubsurfaceSetupPass : uint32
+	{
+		PassOneMain,    // Subsurface Setup. Write to Setup texture
+		PassTwoPassThrough,	// PassThrough to final texture before recombine (due to Subpixel).
+		MAX
+	};
+
 	static const TCHAR* GetEventName(EQuality Quality)
 	{
 		static const TCHAR* const kEventNames[] = {
@@ -570,12 +584,23 @@ public:
 		return kEventNames[int32(Quality)];
 	}
 
+	static const TCHAR* GetPassName(ESubsurfaceSetupPass Pass)
+	{
+		static const TCHAR* const kEventNames[] = {
+			TEXT("Main"),
+			TEXT("Passthrough"),
+		};
+		static_assert(UE_ARRAY_COUNT(kEventNames) == int32(ESubsurfaceSetupPass::MAX), "Fix me");
+		return kEventNames[int32(Pass)];
+	}
+
 	class FDimensionHalfRes : SHADER_PERMUTATION_BOOL("SUBSURFACE_HALF_RES");
 	class FDimensionCheckerboard : SHADER_PERMUTATION_BOOL("SUBSURFACE_PROFILE_CHECKERBOARD");
 	class FRunningInSeparable : SHADER_PERMUTATION_BOOL("SUBSURFACE_FORCE_SEPARABLE");
 	class FDimensionEnableProfileIDCache : SHADER_PERMUTATION_BOOL("ENABLE_PROFILE_ID_CACHE");
 	class FDimensionQuality : SHADER_PERMUTATION_ENUM_CLASS("SUBSURFACE_RECOMBINE_QUALITY", EQuality);
-	using FPermutationDomain = TShaderPermutationDomain<FDimensionHalfRes, FDimensionCheckerboard, FRunningInSeparable, FDimensionEnableProfileIDCache, FDimensionQuality>;
+	class FDimensionSetupPass :SHADER_PERMUTATION_ENUM_CLASS("SUBSURFACE_SETUP_PASS", ESubsurfaceSetupPass);
+	using FPermutationDomain = TShaderPermutationDomain<FDimensionHalfRes, FDimensionCheckerboard, FRunningInSeparable, FDimensionEnableProfileIDCache, FDimensionQuality, FDimensionSetupPass>;
 };
 
 IMPLEMENT_GLOBAL_SHADER(FSubsurfaceIndirectDispatchSetupCS, "/Engine/Private/PostProcessSubsurface.usf", "SetupIndirectCS", SF_Compute);
@@ -761,7 +786,7 @@ public:
 	{
 		if (bForceSeparable)
 		{
-			return FSubsurfaceTiles::ETileType::All;
+			return FSubsurfaceTiles::ETileType::AllNonPassthrough;
 		}
 
 		FSubsurfaceTiles::ETileType TileType = FSubsurfaceTiles::ETileType::AFIS;
@@ -1118,55 +1143,98 @@ void AddSubsurfaceViewPass(
 		
 		Tiles = ClassifySSSTiles(GraphBuilder, View, SceneTextures, SceneViewportParameters, SubsurfaceViewportParameters, bHalfRes);
 
+		// Pre-allocate black UAV for the second buffer if all tiles are rendered in half resolution.
+		// Need to access this texture to directly setup passthrough pixel.
+		if (bHalfRes)
+		{
+			SubsurfaceSubpassTwoTex = CreateBlackUAVTexture(
+				GraphBuilder,
+				SubsurfaceTextureDescriptor,
+				TEXT("Subsurface.SubpassTwoTex"),
+				RDG_EVENT_NAME("SSS::ClearUAV(%s, %s.TileCount > 0)", TEXT("Subsurface.SubpassTwoTex"), ToString(FSubsurfaceTiles::ETileType::All)),
+				SubsurfaceViewport,
+				Tiles.TileTypeCountBuffer,
+				ToIndex(FSubsurfaceTiles::ETileType::All));
+		}
+		else
+		{
+			SubsurfaceSubpassTwoTex = GraphBuilder.CreateTexture(SubsurfaceTextureDescriptor, TEXT("Subsurface.SubpassTwoTex"));
+		}
+
 		FSubsurfaceUniformRef UniformBuffer = CreateUniformBuffer(View, Tiles.TileCount);
 
 		// Call the indirect setup (with tile classification)
 		{
-			FRDGTextureUAVDesc SetupTextureOutDesc(SetupTexture, 0);
-
 			typedef FSubsurfaceIndirectDispatchSetupCS SHADER;
-			SHADER::FParameters* PassParameters = GraphBuilder.AllocParameters<SHADER::FParameters>();
-			PassParameters->Subsurface = SubsurfaceCommonParameters;
-			PassParameters->Output = SubsurfaceViewportParameters;
-			PassParameters->SubsurfaceInput0 = GetSubsurfaceInput(bReadSeparatedSubSurfaceSceneColor ? SeparatedSubSurfaceSceneColor : SceneColorTexture, SceneViewportParameters);
-			PassParameters->SubsurfaceSampler0 = PointClampSampler;
-			PassParameters->SetupTexture = GraphBuilder.CreateUAV(SetupTextureOutDesc);
-			if (bUseProfileIdCache)
+
+			struct FSubsurfaceSetupPassInfo
 			{
-				PassParameters->ProfileIdTexture = GraphBuilder.CreateUAV(FRDGTextureUAVDesc(ProfileIdTexture));
-			}
-			PassParameters->RWBurleyGroupBuffer = GraphBuilder.CreateUAV(Tiles.TileDataBuffer[ToIndex(FSubsurfaceTiles::ETileType::AFIS)], EPixelFormat::PF_R32_UINT);
-			PassParameters->RWSeparableGroupBuffer = GraphBuilder.CreateUAV(Tiles.TileDataBuffer[ToIndex(FSubsurfaceTiles::ETileType::SEPARABLE)], EPixelFormat::PF_R32_UINT);
-			PassParameters->RWTileTypeCountBuffer = GraphBuilder.CreateUAV(Tiles.TileTypeCountBuffer, EPixelFormat::PF_R32_UINT);
-			PassParameters->IndirectDispatchArgsBuffer = Tiles.TileIndirectDispatchBuffer;
-			PassParameters->GroupBuffer = Tiles.GetTileBufferSRV(FSubsurfaceTiles::ETileType::All);
+				FSubsurfaceSetupPassInfo(FSubsurfaceTiles::ETileType InTileType, FRDGTextureRef InOutput, SHADER::ESubsurfaceSetupPass InPass)
+					: TileType(InTileType), Output(InOutput), Pass(InPass)
+				{}
 
-			PassParameters->SubsurfaceUniformParameters = UniformBuffer;
-			PassParameters->Substrate = Substrate::BindSubstrateGlobalUniformParameters(View);
+				FSubsurfaceTiles::ETileType TileType;
+				FRDGTextureRef Output;
+				SHADER::ESubsurfaceSetupPass Pass;
+			};
 
-			SHADER::EQuality Quality = static_cast<SHADER::EQuality>(FSubsurfaceRecombinePS::GetQuality(View));
+			const int NumOfSubsurfaceSetupPass = 2;
 
-			SHADER::FPermutationDomain ComputeShaderPermutationVector;
-			ComputeShaderPermutationVector.Set<SHADER::FDimensionHalfRes>(bHalfRes);
-			ComputeShaderPermutationVector.Set<SHADER::FDimensionCheckerboard>(bCheckerboard);
-			ComputeShaderPermutationVector.Set<SHADER::FRunningInSeparable>(bForceRunningInSeparable);
-			ComputeShaderPermutationVector.Set<SHADER::FDimensionEnableProfileIDCache>(bUseProfileIdCache);
-			ComputeShaderPermutationVector.Set<SHADER::FDimensionQuality>(Quality);
-			TShaderMapRef<SHADER> ComputeShader(View.ShaderMap, ComputeShaderPermutationVector);
+			const FSubsurfaceSetupPassInfo SubsurfaceSetupPassInfos[NumOfSubsurfaceSetupPass] =
+			{
+				{FSubsurfaceTiles::ETileType::AllNonPassthrough, SetupTexture           , SHADER::ESubsurfaceSetupPass::PassOneMain},      // Main setup
+				{FSubsurfaceTiles::ETileType::PASSTHROUGH,       SubsurfaceSubpassTwoTex, SHADER::ESubsurfaceSetupPass::PassTwoPassThrough} // Passthrough SSS for subpixel SSS
+			};
 
-			FComputeShaderUtils::AddPass(
-				GraphBuilder,
-				RDG_EVENT_NAME("SSS::Setup(%s%s%s %s) %dx%d",
-					ComputeShaderPermutationVector.Get<SHADER::FDimensionHalfRes>() ? TEXT(" HalfRes") : TEXT(""),
-					ComputeShaderPermutationVector.Get<SHADER::FDimensionCheckerboard>() ? TEXT(" Checkerboard") : TEXT(""),
-					ComputeShaderPermutationVector.Get<SHADER::FRunningInSeparable>() ? TEXT(" RunningInSeparable") : TEXT(""),
-					SHADER::GetEventName(ComputeShaderPermutationVector.Get<SHADER::FDimensionQuality>()),
-					SubsurfaceViewport.Extent.X,
-					SubsurfaceViewport.Extent.Y),
-				ComputeShader,
-				PassParameters,
-				PassParameters->IndirectDispatchArgsBuffer, 
-				Tiles.GetIndirectDispatchArgOffset(FSubsurfaceTiles::ETileType::All));
+			for (int PassIndex = 0; PassIndex < NumOfSubsurfaceSetupPass; ++PassIndex)
+			{
+				FSubsurfaceSetupPassInfo PassInfo = SubsurfaceSetupPassInfos[PassIndex];
+
+				SHADER::FParameters* PassParameters = GraphBuilder.AllocParameters<SHADER::FParameters>();
+				PassParameters->Subsurface = SubsurfaceCommonParameters;
+				PassParameters->Output = SubsurfaceViewportParameters;
+				PassParameters->SubsurfaceInput0 = GetSubsurfaceInput(bReadSeparatedSubSurfaceSceneColor ? SeparatedSubSurfaceSceneColor : SceneColorTexture, SceneViewportParameters);
+				PassParameters->SubsurfaceSampler0 = PointClampSampler;
+				PassParameters->SetupTexture = GraphBuilder.CreateUAV(FRDGTextureUAVDesc(PassInfo.Output, 0));
+				if (bUseProfileIdCache)
+				{
+					PassParameters->ProfileIdTexture = GraphBuilder.CreateUAV(FRDGTextureUAVDesc(ProfileIdTexture));
+				}
+				PassParameters->RWBurleyGroupBuffer = GraphBuilder.CreateUAV(Tiles.GetTileBuffer(FSubsurfaceTiles::ETileType::AFIS), EPixelFormat::PF_R32_UINT);
+				PassParameters->RWSeparableGroupBuffer = GraphBuilder.CreateUAV(Tiles.GetTileBuffer(FSubsurfaceTiles::ETileType::SEPARABLE), EPixelFormat::PF_R32_UINT);
+				PassParameters->RWTileTypeCountBuffer = GraphBuilder.CreateUAV(Tiles.TileTypeCountBuffer, EPixelFormat::PF_R32_UINT);
+				PassParameters->IndirectDispatchArgsBuffer = Tiles.TileIndirectDispatchBuffer;
+				PassParameters->GroupBuffer = Tiles.GetTileBufferSRV(PassInfo.TileType);
+
+				PassParameters->SubsurfaceUniformParameters = UniformBuffer;
+				PassParameters->Substrate = Substrate::BindSubstrateGlobalUniformParameters(View);
+
+				SHADER::EQuality Quality = static_cast<SHADER::EQuality>(FSubsurfaceRecombinePS::GetQuality(View));
+
+				SHADER::FPermutationDomain ComputeShaderPermutationVector;
+				ComputeShaderPermutationVector.Set<SHADER::FDimensionHalfRes>(bHalfRes);
+				ComputeShaderPermutationVector.Set<SHADER::FDimensionCheckerboard>(bCheckerboard);
+				ComputeShaderPermutationVector.Set<SHADER::FRunningInSeparable>(bForceRunningInSeparable);
+				ComputeShaderPermutationVector.Set<SHADER::FDimensionEnableProfileIDCache>(bUseProfileIdCache);
+				ComputeShaderPermutationVector.Set<SHADER::FDimensionQuality>(Quality);
+				ComputeShaderPermutationVector.Set<SHADER::FDimensionSetupPass>(PassInfo.Pass);
+				TShaderMapRef<SHADER> ComputeShader(View.ShaderMap, ComputeShaderPermutationVector);
+
+				FComputeShaderUtils::AddPass(
+					GraphBuilder,
+					RDG_EVENT_NAME("SSS::Setup(%s%s%s %s %s) %dx%d",
+						ComputeShaderPermutationVector.Get<SHADER::FDimensionHalfRes>() ? TEXT(" HalfRes") : TEXT(""),
+						ComputeShaderPermutationVector.Get<SHADER::FDimensionCheckerboard>() ? TEXT(" Checkerboard") : TEXT(""),
+						ComputeShaderPermutationVector.Get<SHADER::FRunningInSeparable>() ? TEXT(" RunningInSeparable") : TEXT(""),
+						SHADER::GetPassName(ComputeShaderPermutationVector.Get<SHADER::FDimensionSetupPass>()),
+						SHADER::GetEventName(ComputeShaderPermutationVector.Get<SHADER::FDimensionQuality>()),
+						SubsurfaceViewport.Extent.X,
+						SubsurfaceViewport.Extent.Y),
+					ComputeShader,
+					PassParameters,
+					PassParameters->IndirectDispatchArgsBuffer,
+					Tiles.GetIndirectDispatchArgOffset(PassInfo.TileType));
+			}
 		}
 
 	
@@ -1207,9 +1275,9 @@ void AddSubsurfaceViewPass(
 			}
 		}
 
-		// Allocate auxiliary buffers
+		// Allocate additional auxiliary buffers
 		{
-			const int32 Offset = ToIndex(bForceRunningInSeparable ? FSubsurfaceTiles::ETileType::All : FSubsurfaceTiles::ETileType::SEPARABLE);
+			const int32 Offset = ToIndex(bForceRunningInSeparable ? FSubsurfaceTiles::ETileType::AllNonPassthrough : FSubsurfaceTiles::ETileType::SEPARABLE);
 			// Pre-allocate black UAV if we have separable tiles.
 			SubsurfaceSubpassOneTex = CreateBlackUAVTexture(
 				GraphBuilder, 
@@ -1220,23 +1288,6 @@ void AddSubsurfaceViewPass(
 				Tiles.TileTypeCountBuffer,
 				Offset);
 
-			// Pre-allocate black UAV for the second buffer if the separable tiles are rendered in half resolution.
-			if (bHalfRes)
-			{
-				SubsurfaceSubpassTwoTex = CreateBlackUAVTexture(
-					GraphBuilder,
-					SubsurfaceTextureDescriptor,
-					TEXT("Subsurface.SubpassTwoTex"),
-					RDG_EVENT_NAME("SSS::ClearUAV(%s, %s.TileCount > 0)", TEXT("Subsurface.SubpassTwoTex"), ToString(FSubsurfaceTiles::ETileType::All)),
-					SubsurfaceViewport,
-					Tiles.TileTypeCountBuffer,
-					ToIndex(FSubsurfaceTiles::ETileType::All));
-			}
-			else
-			{
-				SubsurfaceSubpassTwoTex = GraphBuilder.CreateTexture(SubsurfaceTextureDescriptor, TEXT("Subsurface.SubpassTwoTex"));
-			}
-
 			if (!bForceRunningInSeparable)
 			{
 				NewQualityHistoryTexture = GraphBuilder.CreateTexture(SubsurfaceTextureDescriptor, TEXT("Subsurface.QualityHistoryState"));
@@ -1275,7 +1326,7 @@ void AddSubsurfaceViewPass(
 
 			const FRDGBufferSRVRef SubsurfaceBufferUsage[] = { 
 				Tiles.GetTileBufferSRV(FSubsurfaceTiles::ETileType::AFIS), 
-				Tiles.GetTileBufferSRV(bForceRunningInSeparable ? FSubsurfaceTiles::ETileType::All : FSubsurfaceTiles::ETileType::SEPARABLE) };
+				Tiles.GetTileBufferSRV(bForceRunningInSeparable ? FSubsurfaceTiles::ETileType::AllNonPassthrough : FSubsurfaceTiles::ETileType::SEPARABLE) };
 
 			//Dispatch the two phase for both SSS
 			for (int PassIndex = 0; PassIndex < NumOfSubsurfacePass; ++PassIndex)
diff --git a/Engine/Source/Runtime/Renderer/Private/SubsurfaceTiles.cpp b/Engine/Source/Runtime/Renderer/Private/SubsurfaceTiles.cpp
--- a/Engine/Source/Runtime/Renderer/Private/SubsurfaceTiles.cpp
+++ b/Engine/Source/Runtime/Renderer/Private/SubsurfaceTiles.cpp
@@ -3,6 +3,7 @@
 #include "SubsurfaceTiles.h"
 #include "SceneRendering.h"
 #include "DataDrivenShaderPlatformInfo.h"
+#include "Engine/SubsurfaceProfile.h"
 
 bool FSubsurfaceTilePassVS::ShouldCompilePermutation(const FGlobalShaderPermutationParameters& Parameters)
 {
@@ -136,6 +137,8 @@ const TCHAR* ToString(FSubsurfaceTiles::ETileType Type)
 	case FSubsurfaceTiles::ETileType::All:			return TEXT("SSS(All)");
 	case FSubsurfaceTiles::ETileType::AFIS:		return TEXT("SSS(AFIS)");
 	case FSubsurfaceTiles::ETileType::SEPARABLE:		return TEXT("SSS(Separable)");
+	case FSubsurfaceTiles::ETileType::PASSTHROUGH:      return TEXT("SSS(Passthrough)");
+	case FSubsurfaceTiles::ETileType::AllNonPassthrough:      return TEXT("SSS(All,NonPassthrough)");
 	default:											return TEXT("Unknown");
 	}
 }
@@ -152,10 +155,15 @@ class FSSSTileCategorisationMarkCS : public FGlobalShader
 		SHADER_PARAMETER_RDG_UNIFORM_BUFFER(FSceneTextureUniformParameters, SceneTextures)
 		SHADER_PARAMETER_RDG_UNIFORM_BUFFER(FSubstrateGlobalUniformParameters, Substrate)
 		SHADER_PARAMETER_STRUCT_INCLUDE(FViewShaderParameters, View)
+		SHADER_PARAMETER(FVector4f, SubsurfaceParams)
+		SHADER_PARAMETER(float, SubsurfaceSubpixelThreshold)
 		SHADER_PARAMETER_STRUCT(FScreenPassTextureViewportParameters, Output)
 		SHADER_PARAMETER_STRUCT(FScreenPassTextureViewportParameters, SubsurfaceInput0)
 		SHADER_PARAMETER(FIntPoint, TiledViewRes)
+		SHADER_PARAMETER_RDG_BUFFER_SRV(Buffer<uint>, GroupBuffer)
 		SHADER_PARAMETER_RDG_BUFFER_UAV(RWStructuredBuffer<uint>, TileMaskBufferOut)
+		SHADER_PARAMETER_RDG_BUFFER_UAV(RWStructuredBuffer<uint>, TileMaskPassthroughBufferOut)
+	    SHADER_PARAMETER_RDG_BUFFER_UAV(RWStructuredBuffer<uint>, TileMaskNonPassthroughBufferOut)
 	END_SHADER_PARAMETER_STRUCT()
 
 	static FPermutationDomain RemapPermutation(FPermutationDomain PermutationVector)
@@ -170,9 +178,10 @@ class FSSSTileCategorisationMarkCS : public FGlobalShader
 
 	static void ModifyCompilationEnvironment(const FGlobalShaderPermutationParameters& Parameters, FShaderCompilerEnvironment& OutEnvironment)
 	{
+		FGlobalShader::ModifyCompilationEnvironment(Parameters, OutEnvironment);
 		OutEnvironment.SetDefine(TEXT("TILE_CATERGORISATION_SHADER"), 1);
 		OutEnvironment.SetDefine(TEXT("SUBSURFACE_TILE_SIZE"), FSubsurfaceTiles::TileSize);
-		FGlobalShader::ModifyCompilationEnvironment(Parameters, OutEnvironment);
+		OutEnvironment.SetDefine(TEXT("SUBSURFACE_KERNEL_SIZE"), SUBSURFACE_KERNEL_SIZE);
 	}
 };
 
@@ -187,6 +196,7 @@ class FSSSTileClassificationBuildListsCS : public FGlobalShader
 		SHADER_PARAMETER_STRUCT_INCLUDE(FViewShaderParameters, View)
 		SHADER_PARAMETER(uint32, VertexCountPerInstanceIndirect)
 		SHADER_PARAMETER(FIntPoint, TiledViewRes)
+		SHADER_PARAMETER(int, TileType)
 		SHADER_PARAMETER_RDG_BUFFER_SRV(StructuredBuffer<uint>, TileMaskBuffer)
 		SHADER_PARAMETER_RDG_BUFFER_UAV(RWBuffer<uint>, RWIndirectDrawArgsBuffer)
 		SHADER_PARAMETER_RDG_BUFFER_UAV(RWBuffer<uint>, RWIndirectDispatchArgsBuffer)
@@ -209,11 +219,38 @@ class FSSSTileClassificationBuildListsCS : public FGlobalShader
 		FGlobalShader::ModifyCompilationEnvironment(Parameters, OutEnvironment);
 		OutEnvironment.SetDefine(TEXT("TILE_CATERGORISATION_SHADER"), 1);
 		OutEnvironment.SetDefine(TEXT("SUBSURFACE_TILE_SIZE"), GetGroupSize());
+		OutEnvironment.SetDefine(TEXT("SUBSURFACE_KERNEL_SIZE"), SUBSURFACE_KERNEL_SIZE);
 	}
 };
 
 IMPLEMENT_GLOBAL_SHADER(FSSSTileClassificationBuildListsCS, "/Engine/Private/PostProcessSubsurfaceTile.usf", "SSSTileClassificationBuildListsCS", SF_Compute);
 
+// Returns the [0, N] clamped value of the 'r.SSS.Scale' CVar.
+static float GetSubsurfaceRadiusScale()
+{
+	static const auto CVar = IConsoleManager::Get().FindTConsoleVariableDataFloat(TEXT("r.SSS.Scale"));
+	check(CVar);
+
+	return FMath::Max(0.0f, CVar->GetValueOnRenderThread());
+}
+
+static float GetSubsurfaceSubpixelThreshold()
+{
+	static const auto CVar = IConsoleManager::Get().FindTConsoleVariableDataFloat(TEXT("r.SSS.Subpixel.Threshold"));
+	check(CVar);
+
+	return FMath::Max(0.0f, CVar->GetValueOnRenderThread());
+}
+
+//@TODO: remove the duplicate from PostProcessSubsurface.cpp
+FVector4f GetSubsurfaceParams(const FViewInfo& View)
+{
+	const float DistanceToProjectionWindow = View.ViewMatrices.GetProjectionMatrix().M[0][0];
+	const float SSSScaleZ = DistanceToProjectionWindow * GetSubsurfaceRadiusScale();
+	const float SSSScaleX = SSSScaleZ / SUBSURFACE_KERNEL_SIZE * 0.5f;
+
+	return FVector4f(SSSScaleX, SSSScaleZ, 0, 0);
+}
 
 /**
 * Code adapted from ScreenSpaceReflectionTiles to reduce Setup SSS cost.
@@ -245,6 +282,10 @@ FSubsurfaceTiles ClassifySSSTiles(FRDGBuilder& GraphBuilder, const FViewInfo& Vi
 			GraphBuilder.CreateBuffer(FRDGBufferDesc::CreateBufferDesc(sizeof(uint32), Result.TileCount), TEXT("Subsurface.TileDataBuffer(AFIS)"));
 		Result.TileDataBuffer[ToIndex(FSubsurfaceTiles::ETileType::SEPARABLE)] = 
 			GraphBuilder.CreateBuffer(FRDGBufferDesc::CreateBufferDesc(sizeof(uint32), Result.TileCount), TEXT("Subsurface.TileDataBuffer(SEPARABLE)"));
+		Result.TileDataBuffer[ToIndex(FSubsurfaceTiles::ETileType::PASSTHROUGH)] = 
+			GraphBuilder.CreateBuffer(FRDGBufferDesc::CreateBufferDesc(sizeof(uint32), Result.TileCount), TEXT("Subsurface.TileDataBuffer(PASSTHROUGH)"));
+		Result.TileDataBuffer[ToIndex(FSubsurfaceTiles::ETileType::AllNonPassthrough)] = 
+			GraphBuilder.CreateBuffer(FRDGBufferDesc::CreateBufferDesc(sizeof(uint32), Result.TileCount), TEXT("Subsurface.TileDataBuffer(AllNonPassthrough)"));
 
 		for (uint32 BufferIt = 0; BufferIt < FSubsurfaceTiles::TileTypeCount; ++BufferIt)
 		{
@@ -254,7 +295,7 @@ FSubsurfaceTiles ClassifySSSTiles(FRDGBuilder& GraphBuilder, const FViewInfo& Vi
 			}
 		}
 
-		// We write to FSubsurfaceTiles::All,
+		// We write to FSubsurfaceTiles::All, FSubsurfaceTiles::AllNonPassthrough, and FSubsurfaceTiles::Passthrough
 		FRDGBufferRef TileListDataBuffer = Result.TileDataBuffer[ToIndex(FSubsurfaceTiles::ETileType::All)];//GraphBuilder.CreateBuffer(FRDGBufferDesc::CreateBufferDesc(sizeof(uint32), Result.TileDimension.X * Result.TileDimension.Y), TEXT("Subsurface.TileListDataBuffer"));
 		FRDGBufferSRVRef TileListDataBufferSRV = Result.TileDataSRV[ToIndex(FSubsurfaceTiles::ETileType::All)];//GraphBuilder.CreateSRV(TileListDataBuffer, PF_R32_UINT);
 
@@ -264,11 +305,18 @@ FSubsurfaceTiles ClassifySSSTiles(FRDGBuilder& GraphBuilder, const FViewInfo& Vi
 		// Allocate buffer with 1 bit / tile
 		FRDGBufferRef TileMaskBuffer = GraphBuilder.CreateBuffer(FRDGBufferDesc::CreateStructuredDesc(sizeof(uint32), FMath::DivideAndRoundUp(Result.TileDimension.X * Result.TileDimension.Y, 32)), TEXT("SSR.Classify.TileMaskBuffer"));
 		FRDGBufferUAVRef TileMaskBufferUAV = GraphBuilder.CreateUAV(TileMaskBuffer);
+		FRDGBufferRef TileMaskPassthroughBuffer = GraphBuilder.CreateBuffer(FRDGBufferDesc::CreateStructuredDesc(sizeof(uint32), FMath::DivideAndRoundUp(Result.TileDimension.X * Result.TileDimension.Y, 32)), TEXT("SSR.Classify.TileMaskPassthroughBuffer"));
+		FRDGBufferUAVRef TileMaskPassthroughBufferUAV = GraphBuilder.CreateUAV(TileMaskPassthroughBuffer);
+		FRDGBufferRef TileMaskNonPassthroughBuffer = GraphBuilder.CreateBuffer(FRDGBufferDesc::CreateStructuredDesc(sizeof(uint32), FMath::DivideAndRoundUp(Result.TileDimension.X * Result.TileDimension.Y, 32)), TEXT("SSR.Classify.TileMaskNonPassthroughBuffer"));
+		FRDGBufferUAVRef TileMaskNonPassthroughBufferUAV = GraphBuilder.CreateUAV(TileMaskNonPassthroughBuffer);
 		
 		AddClearUAVPass(GraphBuilder, TileMaskBufferUAV, 0);
+		AddClearUAVPass(GraphBuilder, TileMaskPassthroughBufferUAV, 0);
+		AddClearUAVPass(GraphBuilder, TileMaskNonPassthroughBufferUAV, 0);
+
 		AddClearUAVPass(GraphBuilder, GraphBuilder.CreateUAV(Result.TileTypeCountBuffer,EPixelFormat::PF_R32_UINT), 0);
 
-		// Mark used tiles based on SHADING_MODEL_ID only
+		// Mark used tiles based on SHADING_MODEL_ID, and whether it is subpixel scattering only.
 		{
 			typedef FSSSTileCategorisationMarkCS SHADER;
 			SHADER::FPermutationDomain PermutationVector;
@@ -278,11 +326,16 @@ FSubsurfaceTiles ClassifySSSTiles(FRDGBuilder& GraphBuilder, const FViewInfo& Vi
 			SHADER::FParameters* PassParameters = GraphBuilder.AllocParameters<SHADER::FParameters>();
 			PassParameters->SceneTextures = SceneTextures.UniformBuffer;
 			PassParameters->View = View.GetShaderParameters();
+			PassParameters->SubsurfaceParams = GetSubsurfaceParams(View);
+			PassParameters->SubsurfaceSubpixelThreshold = GetSubsurfaceSubpixelThreshold();
 			PassParameters->Substrate = Substrate::BindSubstrateGlobalUniformParameters(View);
 			PassParameters->Output = SubsurfaceViewportParameters;
 			PassParameters->SubsurfaceInput0 = SceneViewportParameters;
 			PassParameters->TiledViewRes = Result.TileDimension;
+
 			PassParameters->TileMaskBufferOut = TileMaskBufferUAV;
+			PassParameters->TileMaskPassthroughBufferOut = TileMaskPassthroughBufferUAV;
+			PassParameters->TileMaskNonPassthroughBufferOut = TileMaskNonPassthroughBufferUAV;
 
 			FComputeShaderUtils::AddPass(
 				GraphBuilder,
@@ -296,27 +349,57 @@ FSubsurfaceTiles ClassifySSSTiles(FRDGBuilder& GraphBuilder, const FViewInfo& Vi
 		// Build compacted and coherent z-order tile lists from bit-marked tiles
 		{
 			typedef FSSSTileClassificationBuildListsCS SHADER;
-			TShaderMapRef<SHADER> ComputeShader(View.ShaderMap);
 
-			SHADER::FParameters* PassParameters = GraphBuilder.AllocParameters<SHADER::FParameters>();
+			struct FSubsurfaceBuildListsInfo
+			{
+				FSubsurfaceBuildListsInfo( FSubsurfaceTiles::ETileType inTileType, FRDGBufferRef inInput, FRDGBufferRef inOutput)
+					: Name(ToString(inTileType)), TileType(inTileType), Input(inInput), Output(inOutput)
+				{}
 
-			PassParameters->View = View.GetShaderParameters();
-			PassParameters->TiledViewRes = Result.TileDimension;
-			PassParameters->VertexCountPerInstanceIndirect = GRHISupportsRectTopology ? 4 : 6;
-			PassParameters->TileMaskBuffer = GraphBuilder.CreateSRV(TileMaskBuffer);
-
-			PassParameters->RWIndirectDrawArgsBuffer = DrawIndirectParametersBufferUAV;
-			PassParameters->RWIndirectDispatchArgsBuffer = DispatchIndirectParametersBufferUAV;
-			PassParameters->RWTileTypeCountBuffer = GraphBuilder.CreateUAV(Result.TileTypeCountBuffer, PF_R32_UINT);
-			PassParameters->RWSSSTileListDataBuffer = GraphBuilder.CreateUAV(TileListDataBuffer, PF_R32_UINT);
-			
-			FComputeShaderUtils::AddPass(
-				GraphBuilder,
-				RDG_EVENT_NAME("SSS::TileCategorisationBuildList(All)"),
-				ComputeShader,
-				PassParameters,
-				FComputeShaderUtils::GetGroupCount(Result.TileDimension, SHADER::GetGroupSize())
-			);
+				const TCHAR* Name;
+				FSubsurfaceTiles::ETileType TileType;
+				FRDGBufferRef Input;
+				FRDGBufferRef Output;
+			};
+
+			const int NumOfBuildListsPass = 3;
+
+			const FSubsurfaceBuildListsInfo SubsurfaceBuildListsInfo[NumOfBuildListsPass] =
+			{
+				{FSubsurfaceTiles::ETileType::All,               TileMaskBuffer,               Result.GetTileBuffer(FSubsurfaceTiles::ETileType::All)},
+				{FSubsurfaceTiles::ETileType::PASSTHROUGH,       TileMaskPassthroughBuffer,    Result.GetTileBuffer(FSubsurfaceTiles::ETileType::PASSTHROUGH)},
+				{FSubsurfaceTiles::ETileType::AllNonPassthrough, TileMaskNonPassthroughBuffer, Result.GetTileBuffer(FSubsurfaceTiles::ETileType::AllNonPassthrough)}
+			};
+
+			for (int PassIndex = 0 ; PassIndex < NumOfBuildListsPass ; ++PassIndex)
+			{
+				const FSubsurfaceBuildListsInfo PassInfo = SubsurfaceBuildListsInfo[PassIndex];
+
+				TShaderMapRef<SHADER> ComputeShader(View.ShaderMap);
+
+				SHADER::FParameters* PassParameters = GraphBuilder.AllocParameters<SHADER::FParameters>();
+
+				PassParameters->View = View.GetShaderParameters();
+				PassParameters->TiledViewRes = Result.TileDimension;
+				PassParameters->VertexCountPerInstanceIndirect = GRHISupportsRectTopology ? 4 : 6;
+				PassParameters->TileMaskBuffer = GraphBuilder.CreateSRV(PassInfo.Input);
+
+				PassParameters->TileType = ToIndex(PassInfo.TileType);
+				PassParameters->RWIndirectDrawArgsBuffer = DrawIndirectParametersBufferUAV;
+				PassParameters->RWIndirectDispatchArgsBuffer = DispatchIndirectParametersBufferUAV;
+				
+				PassParameters->RWTileTypeCountBuffer = GraphBuilder.CreateUAV(Result.TileTypeCountBuffer, PF_R32_UINT);
+				PassParameters->RWSSSTileListDataBuffer = GraphBuilder.CreateUAV(PassInfo.Output, PF_R32_UINT);
+
+				FComputeShaderUtils::AddPass(
+					GraphBuilder,
+					RDG_EVENT_NAME("SSS::TileCategorisationBuildList(%s)",
+						PassInfo.Name),
+					ComputeShader,
+					PassParameters,
+					FComputeShaderUtils::GetGroupCount(Result.TileDimension, SHADER::GetGroupSize())
+				);
+			}
 		}
 	}
 	return Result;
diff --git a/Engine/Source/Runtime/Renderer/Private/SubsurfaceTiles.h b/Engine/Source/Runtime/Renderer/Private/SubsurfaceTiles.h
--- a/Engine/Source/Runtime/Renderer/Private/SubsurfaceTiles.h
+++ b/Engine/Source/Runtime/Renderer/Private/SubsurfaceTiles.h
@@ -20,7 +20,9 @@ struct FSubsurfaceTiles
 	// Support two types of acceleration. It is independent of SSS profile
 	// AFIS: adaptive filtered importance sampling
 	// SEPARABLE: two pass separable filter
-	enum class ETileType : uint8 {AFIS, SEPARABLE, All, Count};
+	// PASSTHROUGH: directly pass through because of subpixel SSS
+	// AllNonPassthrough: all subsurface tiles but without passthrough tiles
+	enum class ETileType : uint8 {AFIS, SEPARABLE, PASSTHROUGH, AllNonPassthrough, All, Count};
 	static const uint32 TileTypeCount = uint32(ETileType::Count);
 
 	FIntPoint			BufferResolution = FIntPoint(0, 0);
