SkyLight - real time capture cloud face 0 now overlap with all other faces.

#rnx
#tests PC editor dx12 and console for compute overlap/barriers

diff --git a/Engine/Source/Runtime/Renderer/Private/ReflectionEnvironmentRealTimeCapture.cpp b/Engine/Source/Runtime/Renderer/Private/ReflectionEnvironmentRealTimeCapture.cpp
--- a/Engine/Source/Runtime/Renderer/Private/ReflectionEnvironmentRealTimeCapture.cpp
+++ b/Engine/Source/Runtime/Renderer/Private/ReflectionEnvironmentRealTimeCapture.cpp
@@ -625,20 +625,52 @@ void FScene::AllocateAndCaptureFrameSkyEnvMap(
 			FRDGTextureRef LowResCloudTextureArray[6] = { nullptr, nullptr, nullptr, nullptr, nullptr, nullptr };
 			if (bShouldRenderVolumetricCloud && bExecuteCloud)
 			{
-				for (int32 CubeFace = StartCubeFace; CubeFace < EndCubeFace; CubeFace++)
+				if (CloudRC.CloudInfo->GetVolumetricCloudSceneProxy().bVisibleInRealTimeSkyCaptures)
 				{
-					RDG_EVENT_SCOPE(GraphBuilder, "Capture Volumetric Cloud Face=%d", CubeFace);
-					const uint32 CubeWidthLow = CubeWidth / CloudResolutionDivider;
+					RDG_EVENT_SCOPE(GraphBuilder, "Capture Volumetric Cloud Face = %d - %d", StartCubeFace, EndCubeFace - 1);
 
-					// Create a texture to render the clouds at a given resolution
-					FRDGTextureRef& LowResCloudTexture = LowResCloudTextureArray[CubeFace];
+					const uint32 CubeWidthLow = CubeWidth / CloudResolutionDivider;
 					FRDGTextureDesc LowResCloudTexDesc = FRDGTextureDesc::Create2D(FIntPoint(CubeWidthLow, CubeWidthLow), PF_FloatRGBA,
 						FClearValueBinding::Black,
-						TexCreate_RenderTargetable | TexCreate_UAV | TexCreate_ShaderResource);
-					LowResCloudTexture = GraphBuilder.CreateTexture(LowResCloudTexDesc, TEXT("SkyLight.LowResCloudTexture"));
+						TexCreate_RenderTargetable | TexCreate_UAV | TexCreate_ShaderResource | TexCreate_NoFastClear);
+
+					struct FLowResCloudTextureArrayUAV
+					{
+						FRDGTextureUAVRef UAVs[6] = { nullptr, nullptr, nullptr, nullptr, nullptr, nullptr };
+					};
+					FLowResCloudTextureArrayUAV LowResCloudTextureArrayUAV;
+
+					// First allocate a texture to render the clouds at a given resolution and UAVs that skips barriers to have compute work overlapping.
+					for (int32 CubeFace = StartCubeFace; CubeFace < EndCubeFace; CubeFace++)
+					{
+						FRDGTextureRef& LowResCloudTexture = LowResCloudTextureArray[CubeFace];
+						LowResCloudTexture = GraphBuilder.CreateTexture(LowResCloudTexDesc, TEXT("SkyLight.LowResCloudTexture2D"));
+						LowResCloudTextureArrayUAV.UAVs[CubeFace] = GraphBuilder.CreateUAV(LowResCloudTexture, ERDGUnorderedAccessViewFlags::SkipBarrier);
+					}
+
+					// Clear all the textures.
+					for (int32 CubeFace = StartCubeFace; CubeFace < EndCubeFace; CubeFace++)
+					{
+						FRDGTextureRef& LowResCloudTexture = LowResCloudTextureArray[CubeFace];
+						AddClearUAVPass(GraphBuilder, LowResCloudTextureArrayUAV.UAVs[CubeFace], FLinearColor(0.0f, 0.0f, 0.0f, 1.0f));
+					}
+
+					// Transitions from compute to compute for the cache
+					AddPass(GraphBuilder, RDG_EVENT_NAME("Barrier"), [LowResCloudTextureArrayUAV = LowResCloudTextureArrayUAV, StartCubeFace, EndCubeFace](FRHICommandListImmediate& RHICmdList)
+					{
+						// Clear virtual texture feedback to default value
+						for (int32 CubeFace = StartCubeFace; CubeFace < EndCubeFace; CubeFace++)
+						{
+							FRDGTextureUAVRef LowResCloudTextureUAV = LowResCloudTextureArrayUAV.UAVs[CubeFace];
+							RHICmdList.Transition(FRHITransitionInfo(LowResCloudTextureUAV->GetRHI(), ERHIAccess::UAVCompute, ERHIAccess::UAVCompute));
+						}
+					});
 
-					if (CloudRC.CloudInfo->GetVolumetricCloudSceneProxy().bVisibleInRealTimeSkyCaptures)
+					// Render clouds into the texture using overlapping compute work thanks to [REDACTED_PERSON].
+					for (int32 CubeFace = StartCubeFace; CubeFace < EndCubeFace; CubeFace++)
 					{
+						FRDGTextureRef& LowResCloudTexture = LowResCloudTextureArray[CubeFace];
+
 						FMatrix LowResCubeProjectionMatrix;
 						FViewInfo* LowResCubeViewPtr = nullptr;
 						CreateMainViewSnapshotForRealTimeCapture(LowResCubeViewPtr, LowResCubeProjectionMatrix, CubeWidthLow);
@@ -653,9 +685,9 @@ void FScene::AllocateAndCaptureFrameSkyEnvMap(
 						LowResCubeView.CreateViewUniformBuffers(*LowResCubeView.CachedViewUniformShaderParameters);
 						CloudRC.ViewUniformBuffer = LowResCubeView.ViewUniformBuffer;
 
-
 						// Render
-						CloudRC.RenderTargets[0] = FRenderTargetBinding(LowResCloudTexture, ERenderTargetLoadAction::EClear);
+						CloudRC.RenderTargets[0] = FRenderTargetBinding(LowResCloudTexture, ERenderTargetLoadAction::ENoAction);
+						CloudRC.ComputeOverlapCloudColorTexture2DUAVWithoutBarrier = LowResCloudTextureArrayUAV.UAVs[CubeFace];
 						CloudRC.bDisableCloudBlending = true;
 
 						FCloudShadowAOData CloudShadowAOData;
@@ -671,17 +703,24 @@ void FScene::AllocateAndCaptureFrameSkyEnvMap(
 
 						SceneRenderer.RenderVolumetricCloudsInternal(GraphBuilder, CloudRC, InstanceCullingManager, FIntPoint(CubeWidthLow, CubeWidthLow));
 					}
-					else
+
+					// Final transition to SRV mask for the cloud texture to be composited into the captured cube map.
+					AddPass(GraphBuilder, RDG_EVENT_NAME("Barrier"), [LowResCloudTextureArrayUAV = LowResCloudTextureArrayUAV, StartCubeFace, EndCubeFace](FRHICommandListImmediate& RHICmdList)
 					{
-						AddClearRenderTargetPass(GraphBuilder, LowResCloudTexture);
-					}
+						// Clear virtual texture feedback to default value
+						for (int32 CubeFace = StartCubeFace; CubeFace < EndCubeFace; CubeFace++)
+						{
+							FRDGTextureUAVRef LowResCloudTextureUAV = LowResCloudTextureArrayUAV.UAVs[CubeFace];
+							RHICmdList.Transition(FRHITransitionInfo(LowResCloudTextureUAV->GetRHI(), ERHIAccess::UAVCompute, ERHIAccess::SRVMask));
+						}
+					});
 				}
 			}
 
 			const bool bAlwaysClearColorBuffer = CVarRealTimeReflectionCaptureAlwaysClearColorBuffer.GetValueOnRenderThread() > 0;
 			for (int32 CubeFace = StartCubeFace; CubeFace < EndCubeFace; CubeFace++)
 			{
-				RDG_EVENT_SCOPE(GraphBuilder, "Capture Face=%d", CubeFace);
+				RDG_EVENT_SCOPE(GraphBuilder, "Capture Face = %d", CubeFace);
 
 				SkyRC.RenderTargets[0] = FRenderTargetBinding(SkyCubeTexture, ERenderTargetLoadAction::ELoad, 0, CubeFace);
 
@@ -742,7 +781,7 @@ void FScene::AllocateAndCaptureFrameSkyEnvMap(
 						{
 							if (bIsMobilePlatform)
 							{
-								RDG_EVENT_SCOPE(GraphBuilder, "Capture Sky Materials", CubeFace);
+								RDG_EVENT_SCOPE(GraphBuilder, "Capture Sky Dome Materials - Mobile");
 								auto* PassParameters = GraphBuilder.AllocParameters<FMobileCaptureSkyMeshReflectionPassParameters>();
 								PassParameters->View = CubeView.GetShaderParameters();
 								PassParameters->RenderTargets = SkyRC.RenderTargets;
@@ -754,7 +793,7 @@ void FScene::AllocateAndCaptureFrameSkyEnvMap(
 									FRDGTextureDesc CubeDepthTextureDesc = FRDGTextureDesc::Create2D(FIntPoint(CubeWidth, CubeWidth), PF_DepthStencil,
 										MainView.GetSceneTexturesConfig().DepthClearValue,
 										TexCreate_DepthStencilTargetable | TexCreate_ShaderResource);
-									CubeDepthTexture = GraphBuilder.CreateTexture(CubeDepthTextureDesc, TEXT("SkyLight.CubeDepthTexture"));
+									CubeDepthTexture = GraphBuilder.CreateTexture(CubeDepthTextureDesc, TEXT("SkyLight.CubeDepthTexture2D"));
 									PassParameters->RenderTargets.DepthStencil = FDepthStencilBinding(CubeDepthTexture, ERenderTargetLoadAction::EClear, FExclusiveDepthStencil::DepthWrite_StencilNop);
 								}
 					
@@ -820,7 +859,7 @@ void FScene::AllocateAndCaptureFrameSkyEnvMap(
 							}
 							else
 							{
-								RDG_EVENT_SCOPE(GraphBuilder, "Capture Sky Materials", CubeFace);
+								RDG_EVENT_SCOPE(GraphBuilder, "Capture Sky Dome Materials");
 								auto* PassParameters = GraphBuilder.AllocParameters<FCaptureSkyMeshReflectionPassParameters>();
 								PassParameters->View = CubeView.GetShaderParameters();
 								PassParameters->RenderTargets = SkyRC.RenderTargets;
@@ -897,12 +936,13 @@ void FScene::AllocateAndCaptureFrameSkyEnvMap(
 						else if (!bIsMobilePlatform)
 						{
 							// TODO: mobile should never get there
-							RDG_EVENT_SCOPE(GraphBuilder, "Capture Sky Raw", CubeFace);
+							RDG_EVENT_SCOPE(GraphBuilder, "Capture Sky - SkyAtmosphere");
 							FSceneTextureShaderParameters SceneTextures = CreateSceneTextureShaderParameters(GraphBuilder, &SceneRenderer.GetActiveSceneTextures(), SceneRenderer.FeatureLevel, ESceneTextureSetupMode::SceneDepth);
 							SceneRenderer.RenderSkyAtmosphereInternal(GraphBuilder, SceneTextures, SkyRC);
 						}
 						else if (bIsMobilePlatform)	// The SkyAtmosphere full screen pass shader is not compiled on mobile, so clear the capture to black.
 						{
+							RDG_EVENT_SCOPE(GraphBuilder, "Capture Sky - Clear");
 							// No sky dome mesh and no sky rendering on mobile so let's cleat to black.
 							ClearCubeFace(SkyCubeTexture, CubeFace);
 						}
@@ -911,6 +951,8 @@ void FScene::AllocateAndCaptureFrameSkyEnvMap(
 						
 						if (Scene && Scene->HasAnyExponentialHeightFog() && Scene->ExponentialFogs[0].bVisibleInRealTimeSkyCaptures)
 						{
+							RDG_EVENT_SCOPE(GraphBuilder, "Capture Height Fog");
+
 							FRenderRealTimeReflectionHeightFogVS::FPermutationDomain VsPermutationVector;
 							TShaderMapRef<FRenderRealTimeReflectionHeightFogVS> VertexShader(GetGlobalShaderMap(SkyRC.FeatureLevel), VsPermutationVector);
 
@@ -960,6 +1002,7 @@ void FScene::AllocateAndCaptureFrameSkyEnvMap(
 					}
 					else if(!bAlwaysClearColorBuffer) // Only clear if not done before.
 					{
+						RDG_EVENT_SCOPE(GraphBuilder, "Capture Sky - Clear");
 						ClearCubeFace(SkyCubeTexture, CubeFace);
 					}
 				}
@@ -968,6 +1011,8 @@ void FScene::AllocateAndCaptureFrameSkyEnvMap(
 				FRDGTextureRef& LowResCloudTexture = LowResCloudTextureArray[CubeFace];
 				if ((SkyLight->bLowerHemisphereIsSolidColor || LowResCloudTexture !=nullptr) && bExecuteCloud)
 				{
+					RDG_EVENT_SCOPE(GraphBuilder, "ApplyLowerHemisphereColor");
+
 					FRenderRealTimeReflectionHeightFogVS::FPermutationDomain VsPermutationVector;
 					TShaderMapRef<FRenderRealTimeReflectionHeightFogVS> VertexShader(GetGlobalShaderMap(SkyRC.FeatureLevel), VsPermutationVector);
 
@@ -1028,6 +1073,8 @@ void FScene::AllocateAndCaptureFrameSkyEnvMap(
 
 	auto RenderCubeFaces_GenCubeMips = [&](uint32 CubeMipStart, uint32 CubeMipEnd, TRefCountPtr<IPooledRenderTarget>& SkyRenderTarget)
 	{
+		RDG_EVENT_SCOPE(GraphBuilder, "Capture GenCubeMips Mip = %d - %d", CubeMipStart, CubeMipEnd);
+
 		// on mobile platforms use the pixel shader implementation as compute shaders might not be optimal
 		const bool bIsMobilePlatform = IsMobilePlatform(MainView.GetShaderPlatform());
 		if (bIsMobilePlatform)
@@ -1086,6 +1133,8 @@ void FScene::AllocateAndCaptureFrameSkyEnvMap(
 
 	auto RenderCubeFaces_SpecularConvolution = [&](uint32 CubeMipStart, uint32 CubeMipEnd, uint32 FaceStart, uint32 FaceCount, TRefCountPtr<IPooledRenderTarget>& DstRenderTarget, TRefCountPtr<IPooledRenderTarget>& SrcRenderTarget)
 	{
+		RDG_EVENT_SCOPE(GraphBuilder, "Capture SpecularConvolution Face = %d - %d Mip = %d - %d", FaceStart, FaceStart + FaceCount - 1, CubeMipStart, CubeMipEnd);
+
 		check((FaceStart + FaceCount) <= 6);
 		FRDGTextureRef RDGSrcRenderTarget = GraphBuilder.RegisterExternalTexture(SrcRenderTarget);
 		FRDGTextureRef RDGDstRenderTarget = GraphBuilder.RegisterExternalTexture(DstRenderTarget);
@@ -1134,6 +1183,8 @@ void FScene::AllocateAndCaptureFrameSkyEnvMap(
 
 	auto RenderCubeFaces_DiffuseIrradiance = [&](TRefCountPtr<IPooledRenderTarget>& SourceCubemap)
 	{
+		RDG_EVENT_SCOPE(GraphBuilder, "Capture DiffuseIrradiance");
+
 		FRDGTextureRef SourceCubemapTexture = GraphBuilder.RegisterExternalTexture(SourceCubemap);
 		FRDGTextureSRVRef SourceCubemapTextureSRV = GraphBuilder.CreateSRV(FRDGTextureSRVDesc::Create(SourceCubemapTexture));
 		
diff --git a/Engine/Source/Runtime/Renderer/Private/VolumetricCloudRendering.cpp b/Engine/Source/Runtime/Renderer/Private/VolumetricCloudRendering.cpp
--- a/Engine/Source/Runtime/Renderer/Private/VolumetricCloudRendering.cpp
+++ b/Engine/Source/Runtime/Renderer/Private/VolumetricCloudRendering.cpp
@@ -2330,14 +2330,21 @@ static void GetOutputTexturesWithFallback(
 		}
 	}
 
-	FRDGTextureRef CloudColor0Texture = CloudRC.RenderTargets[0].GetTexture();
-	if (!CloudColor0Texture || CloudColor0Texture->Desc.IsTextureCube())
+	if (CloudRC.ComputeOverlapCloudColorTexture2DUAVWithoutBarrier)
 	{
-		CloudColor0TextureUAV = CloudRC.DefaultCloudColor02DTextureUAV;
+		CloudColor0TextureUAV = CloudRC.ComputeOverlapCloudColorTexture2DUAVWithoutBarrier;
 	}
 	else
 	{
-		CloudColor0TextureUAV = GraphBuilder.CreateUAV(CloudColor0Texture, ERDGUnorderedAccessViewFlags::SkipBarrier);
+		FRDGTextureRef CloudColor0Texture = CloudRC.RenderTargets[0].GetTexture();
+		if (!CloudColor0Texture || CloudColor0Texture->Desc.IsTextureCube())
+		{
+			CloudColor0TextureUAV = CloudRC.DefaultCloudColor02DTextureUAV;
+		}
+		else
+		{
+			CloudColor0TextureUAV = GraphBuilder.CreateUAV(CloudColor0Texture, ERDGUnorderedAccessViewFlags::SkipBarrier);
+		}
 	}
 
 	FRDGTextureRef CloudColor1Texture = CloudRC.SecondaryCloudTracingDataTexture;
diff --git a/Engine/Source/Runtime/Renderer/Private/VolumetricCloudRendering.h b/Engine/Source/Runtime/Renderer/Private/VolumetricCloudRendering.h
--- a/Engine/Source/Runtime/Renderer/Private/VolumetricCloudRendering.h
+++ b/Engine/Source/Runtime/Renderer/Private/VolumetricCloudRendering.h
@@ -173,6 +173,7 @@ struct FCloudRenderContext
 	FRDGTextureUAVRef DefaultCloudAlphaHoldoutUAV = nullptr;
 
 	FRDGTextureUAVRef ComputeOverlapCloudColorCubeTextureUAVWithoutBarrier = nullptr;
+	FRDGTextureUAVRef ComputeOverlapCloudColorTexture2DUAVWithoutBarrier = nullptr;
 
 	FCloudRenderContext();
 
