[SSS] Fix redefinition for different linkage.

diff --git a/Engine/Source/Runtime/Renderer/Private/SubsurfaceTiles.cpp b/Engine/Source/Runtime/Renderer/Private/SubsurfaceTiles.cpp
--- a/Engine/Source/Runtime/Renderer/Private/SubsurfaceTiles.cpp
+++ b/Engine/Source/Runtime/Renderer/Private/SubsurfaceTiles.cpp
@@ -226,7 +226,7 @@ class FSSSTileClassificationBuildListsCS : public FGlobalShader
 IMPLEMENT_GLOBAL_SHADER(FSSSTileClassificationBuildListsCS, "/Engine/Private/PostProcessSubsurfaceTile.usf", "SSSTileClassificationBuildListsCS", SF_Compute);
 
 // Returns the [0, N] clamped value of the 'r.SSS.Scale' CVar.
-static float GetSubsurfaceRadiusScale()
+static float GetSubsurfaceRadiusScaleForTiling()
 {
 	static const auto CVar = IConsoleManager::Get().FindTConsoleVariableDataFloat(TEXT("r.SSS.Scale"));
 	check(CVar);
@@ -246,7 +246,7 @@ static float GetSubsurfaceSubpixelThreshold()
 FVector4f GetSubsurfaceParams(const FViewInfo& View)
 {
 	const float DistanceToProjectionWindow = View.ViewMatrices.GetProjectionMatrix().M[0][0];
-	const float SSSScaleZ = DistanceToProjectionWindow * GetSubsurfaceRadiusScale();
+	const float SSSScaleZ = DistanceToProjectionWindow * GetSubsurfaceRadiusScaleForTiling();
 	const float SSSScaleX = SSSScaleZ / SUBSURFACE_KERNEL_SIZE * 0.5f;
 
 	return FVector4f(SSSScaleX, SSSScaleZ, 0, 0);
