Add missing CSV marker for virtual texture update.
This was lost in recent change to move FinalizeRequests() into its own call.
#rnx

diff --git a/Engine/Source/Runtime/Renderer/Private/VT/VirtualTextureSystem.cpp b/Engine/Source/Runtime/Renderer/Private/VT/VirtualTextureSystem.cpp
--- a/Engine/Source/Runtime/Renderer/Private/VT/VirtualTextureSystem.cpp
+++ b/Engine/Source/Runtime/Renderer/Private/VT/VirtualTextureSystem.cpp
@@ -2525,6 +2525,7 @@ void FVirtualTextureSystem::SubmitRequests(FRHICommandList& RHICmdList, ERHIFeat
 
 void FVirtualTextureSystem::FinalizeRequests(FRDGBuilder& GraphBuilder)
 {
+	RDG_CSV_STAT_EXCLUSIVE_SCOPE(GraphBuilder, VirtualTextureSystem_Update);
 	TRACE_CPUPROFILER_EVENT_SCOPE(FVirtualTextureSystem::FinalizeRequests);
 
 	// Finalize requests
