Fix comments for r.AmbientOcclusion.Compute which also works on PC

#tests Validated working on PC

diff --git a/Engine/Source/Runtime/Renderer/Private/CompositionLighting/PostProcessAmbientOcclusion.cpp b/Engine/Source/Runtime/Renderer/Private/CompositionLighting/PostProcessAmbientOcclusion.cpp
--- a/Engine/Source/Runtime/Renderer/Private/CompositionLighting/PostProcessAmbientOcclusion.cpp
+++ b/Engine/Source/Runtime/Renderer/Private/CompositionLighting/PostProcessAmbientOcclusion.cpp
@@ -32,7 +32,6 @@ static TAutoConsoleVariable<int32> CVarAmbientOcclusionCompute(
 	TEXT("If SSAO should use ComputeShader (not available on all platforms) or PixelShader.\n")
 	TEXT("The [Async] Compute Shader version is WIP, not optimized, requires hardware support (not mobile/DX10/OpenGL3),\n")
 	TEXT("does not use normals which allows it to run right after EarlyZPass (better performance when used with AyncCompute)\n")
-	TEXT("AyncCompute is currently only functional on PS4.\n")
 	TEXT(" 0: PixelShader (default)\n")
 	TEXT(" 1: (WIP) Use ComputeShader if possible, otherwise fall back to '0'\n")
 	TEXT(" 2: (WIP) Use AsyncCompute if efficient, otherwise fall back to '1'\n")
