VolumetricCloud - trivialiy reject pixels during upsampling to full res based on depth / start tracing distance.
This help with perf a bit + trivially fix some leaking of cloud onto forground objects.

#rnx
#tests PC editor dx12

diff --git a/Engine/Shaders/Private/VolumetricRenderTarget.usf b/Engine/Shaders/Private/VolumetricRenderTarget.usf
--- a/Engine/Shaders/Private/VolumetricRenderTarget.usf
+++ b/Engine/Shaders/Private/VolumetricRenderTarget.usf
@@ -544,6 +544,7 @@ uint4 VolumetricTextureValidCoordRect;
 float4 VolumetricTextureValidUvRect;
 uint ForwardShadingEnable;
 uint OutputAlphaHoldout;
+float VolumeTracingStartDistanceFromCamera;
 
 #if PERMUTATION_RENDER_UNDERWATER_BUFFER || PERMUTATION_COMPOSE_WITH_WATER
 Texture2D WaterLinearDepthTexture;
@@ -802,6 +803,13 @@ void ComposeVolumetricRTOverScenePS(
 	float PixelDistanceFromView = length(PrimaryView.TranslatedWorldCameraOrigin - ScreenTranslatedWorldPosition);
 	float PixelDistanceFromViewKm = PixelDistanceFromView * CENTIMETER_TO_KILOMETER;
 
+	if (PixelDistanceFromViewKm < VolumeTracingStartDistanceFromCamera)
+	{
+		// This pixel is closer to the start tracing distance so simple discard it.
+		// This also allows to easily remove cloud leaking on forground meshes.
+		clip(-1.0f);
+		return;
+	}
 
 
 #if PERMUTATION_COMPOSE_WITH_WATER
diff --git a/Engine/Source/Runtime/Renderer/Private/VolumetricCloudRendering.cpp b/Engine/Source/Runtime/Renderer/Private/VolumetricCloudRendering.cpp
--- a/Engine/Source/Runtime/Renderer/Private/VolumetricCloudRendering.cpp
+++ b/Engine/Source/Runtime/Renderer/Private/VolumetricCloudRendering.cpp
@@ -2749,6 +2749,7 @@ bool FSceneRenderer::RenderVolumetricCloud(
 				if (bShouldViewRenderVolumetricCloudRenderTarget)
 				{
 					FVolumetricRenderTargetViewStateData& VRT = ViewInfo.ViewState->VolumetricCloudRenderTarget;
+					VRT.SetStartTracingDistance(CloudSceneProxy.TracingStartDistanceFromCamera);
 					DestinationRT = VRT.GetOrCreateVolumetricTracingRT(GraphBuilder);
 					DestinationRTDepth = VRT.GetOrCreateVolumetricTracingRTDepth(GraphBuilder);
 
diff --git a/Engine/Source/Runtime/Renderer/Private/VolumetricRenderTarget.cpp b/Engine/Source/Runtime/Renderer/Private/VolumetricRenderTarget.cpp
--- a/Engine/Source/Runtime/Renderer/Private/VolumetricRenderTarget.cpp
+++ b/Engine/Source/Runtime/Renderer/Private/VolumetricRenderTarget.cpp
@@ -184,6 +184,7 @@ FVolumetricRenderTargetViewStateData::FVolumetricRenderTargetViewStateData()
 	, bHoldoutValid(false)
 	, bValid(false)
 	, PreViewExposure(1.0f)
+	, StartTracingDistance(0.0f)
 	, FrameId(0)
 	, NoiseFrameIndex(0)
 	, NoiseFrameIndexModPattern(0)
@@ -236,6 +237,7 @@ void FVolumetricRenderTargetViewStateData::Initialise(
 		bFirstTimeUsed = false;
 		bHistoryValid = false;
 		PreViewExposure = 1.0f;
+		StartTracingDistance = 0.0f;
 		FrameId = 0;
 		NoiseFrameIndex = 0;
 		NoiseFrameIndexModPattern = 0;
@@ -322,6 +324,7 @@ void FVolumetricRenderTargetViewStateData::Reset()
 	bHoldoutValid = false;
 	bValid = false;
 	PreViewExposure = 1.0f;
+	StartTracingDistance = 0.0f;
 	FrameId = 0;
 	NoiseFrameIndex = 0;
 	NoiseFrameIndexModPattern = 0;
@@ -887,6 +890,7 @@ class FComposeVolumetricRTOverScenePS : public FGlobalShader
 		RENDER_TARGET_BINDING_SLOTS()
 		SHADER_PARAMETER(uint32, ForwardShadingEnable)
 		SHADER_PARAMETER(uint32, OutputAlphaHoldout)
+		SHADER_PARAMETER(float, VolumeTracingStartDistanceFromCamera)
 		SHADER_PARAMETER(float, UvOffsetSampleAcceptanceWeight)
 		SHADER_PARAMETER(FVector4f, VolumetricTextureSizeAndInvSize)
 		SHADER_PARAMETER(FVector2f, FullResolutionToVolumetricBufferResolutionScale)
@@ -1056,6 +1060,7 @@ void ComposeVolumetricRenderTargetOverScene(
 		PassParameters->Substrate = Substrate::BindSubstrateGlobalUniformParameters(ViewInfo);
 		PassParameters->OutputAlphaHoldout = 0;
 		PassParameters->ForwardShadingEnable = bForwardShading ? 1 : 0;
+		PassParameters->VolumeTracingStartDistanceFromCamera = VolumetricCloudRT.GetStartTracingDistance();
 		FVector2f DummyUVScale;
 		GetTextureSafeUvCoordBound(VolumetricViewRect, PassParameters->VolumetricTexture, PassParameters->VolumetricTextureValidCoordRect, PassParameters->VolumetricTextureValidUvRect, DummyUVScale);
 
@@ -1178,6 +1183,7 @@ void ComposeVolumetricRenderTargetOverSceneUnderWater(
 		PassParameters->Substrate = Substrate::BindSubstrateGlobalUniformParameters(ViewInfo);
 		PassParameters->OutputAlphaHoldout = 0;
 		PassParameters->ForwardShadingEnable = 0;
+		PassParameters->VolumeTracingStartDistanceFromCamera = VolumetricCloudRT.GetStartTracingDistance();
 
 		if (bShouldVolumetricCloudsApplyFogDuringReconstruction)
 		{
@@ -1260,6 +1266,7 @@ void ComposeVolumetricRenderTargetOverSceneForVisualization(
 		PassParameters->Substrate = Substrate::BindSubstrateGlobalUniformParameters(ViewInfo);
 		PassParameters->OutputAlphaHoldout = 0;
 		PassParameters->ForwardShadingEnable = 0;
+		PassParameters->VolumeTracingStartDistanceFromCamera = VolumetricCloudRT.GetStartTracingDistance();;
 		FVector2f DummyUVScale;
 		GetTextureSafeUvCoordBound(VolumetricViewRect, PassParameters->VolumetricTexture, PassParameters->VolumetricTextureValidCoordRect, PassParameters->VolumetricTextureValidUvRect, DummyUVScale);
 
diff --git a/Engine/Source/Runtime/Renderer/Private/VolumetricRenderTargetViewStateData.h b/Engine/Source/Runtime/Renderer/Private/VolumetricRenderTargetViewStateData.h
--- a/Engine/Source/Runtime/Renderer/Private/VolumetricRenderTargetViewStateData.h
+++ b/Engine/Source/Runtime/Renderer/Private/VolumetricRenderTargetViewStateData.h
@@ -35,11 +35,20 @@ public:
 		PreViewExposure = ViewExposure;
 	}
 
-	float GetPrevViewExposure()
+	float GetPrevViewExposure() const 
 	{
 		return PreViewExposure;
 	}
 
+	void SetStartTracingDistance(float InStartTracingDistance)
+	{
+		StartTracingDistance = InStartTracingDistance;
+	}
+	float GetStartTracingDistance() const
+	{
+		return StartTracingDistance;
+	}
+
 	FRDGTextureRef GetOrCreateVolumetricTracingRT(FRDGBuilder& GraphBuilder);
 	FRDGTextureRef GetOrCreateVolumetricSecondaryTracingRT(FRDGBuilder& GraphBuilder);
 	FRDGTextureRef GetOrCreateVolumetricTracingRTDepth(FRDGBuilder& GraphBuilder);
@@ -97,6 +106,7 @@ private:
 	bool bHoldoutValid;
 	bool bValid;
 	float PreViewExposure;
+	float StartTracingDistance;	// The distance at which the tracing starts, and thus the composition can be clipped for pixel closer to that distance.
 
 	int32 FrameId;
 	uint32 NoiseFrameIndex;	// This is only incremented once all Volumetric render target samples have been iterated
