[HWRT] Add support for using vertex mask when doing dynamic geometry update:
* When material mask is used index buffer is not used and all vertices are duplicated during the update.
* Masked triangle is hidden if any of its vertices is masked.
* Triangles are hidden by moving all its vertices to the same position.
* Currently only supported by landscape.

#tests replayrun
#rnx

diff --git a/Engine/Shaders/Private/LandscapeVertexFactory.ush b/Engine/Shaders/Private/LandscapeVertexFactory.ush
--- a/Engine/Shaders/Private/LandscapeVertexFactory.ush
+++ b/Engine/Shaders/Private/LandscapeVertexFactory.ush
@@ -10,6 +10,8 @@
 
 #define VERTEX_FACTORY_MODIFIES_INTERPOLATION 1
 
+#define VF_SUPPORTS_RAYTRACING_VERTEX_MASK 1
+
 // Set by FLightMapTexturePolicy
 #include "/Engine/Generated/UniformBuffers/PrecomputedLightingBuffer.ush"
 #include "LightmapData.ush"
diff --git a/Engine/Shaders/Private/RayTracing/RayTracingDynamicMesh.usf b/Engine/Shaders/Private/RayTracing/RayTracingDynamicMesh.usf
--- a/Engine/Shaders/Private/RayTracing/RayTracingDynamicMesh.usf
+++ b/Engine/Shaders/Private/RayTracing/RayTracingDynamicMesh.usf
@@ -11,34 +11,12 @@
 #include "/Engine/Private/ComputeShaderUtils.ush"
 #include "/Engine/Generated/Material.ush"
 #include "/Engine/Generated/VertexFactory.ush"
- 
-struct FRayTracingDynamicMeshVSToGS
-{
-	float4 WorldPosition : SV_POSITION;
-};
-
-// Must match GetUniformMeshStreamOutLayout and TRIANGLE_VERTEX_DATA_STRIDE
-struct FRayTracingDynamicMeshVertex
-{
-	float4 WorldPosition : SV_POSITION;
-};
-
-[maxvertexcount(3)]
-void RayTracingDynamicMeshGS(triangle FRayTracingDynamicMeshVSToGS Inputs[3], inout TriangleStream<FRayTracingDynamicMeshVertex> OutStream)
-{ 
-	for (int i = 0; i < 3; i++)
-	{
-		FRayTracingDynamicMeshVSToGS Input = Inputs[i];
-
-		FRayTracingDynamicMeshVertex Vertex;
-		Vertex.WorldPosition = Input.WorldPosition;
 
-		OutStream.Append(Vertex);
-	}
-}
+#define PER_VERTEX_MASK_ENABLED (USE_VERTEX_MASK && VF_SUPPORTS_RAYTRACING_VERTEX_MASK && MATERIALBLENDING_MASKED && RAY_TRACING_DYNAMIC_MESH_IN_LOCAL_SPACE)
 
 #if COMPUTESHADER
 uint UsingIndirectDraw;
+uint MaxNumThreads;
 uint NumVertices;
 uint MinVertexIndex;
 uint PrimitiveId;
@@ -49,71 +27,129 @@ float4x4 WorldToInstance;
 RWBuffer<float> RWVertexPositions;
 int bApplyWorldPositionOffset;
 
+Buffer<uint> IndexBuffer;
+uint IndexBufferOffset;
+
+void HideTriangle(uint TriangleIndex, float3 ReferenceLocalPositon)
+{
+	for (int VertexIndex = 0; VertexIndex < 3; VertexIndex++)
+	{		
+		const uint OutputVertexIndex = OutputVertexBaseIndex + (TriangleIndex * 3 + VertexIndex) * 3;
+
+		RWVertexPositions[OutputVertexIndex + 0] = ReferenceLocalPositon.x;
+		RWVertexPositions[OutputVertexIndex + 1] = ReferenceLocalPositon.y;
+		RWVertexPositions[OutputVertexIndex + 2] = ReferenceLocalPositon.z;
+	}
+}
+
 [numthreads(64, 1, 1)]
 void RayTracingDynamicGeometryConverterCS(uint3 GroupId : SV_GroupID, int GroupThreadIndex : SV_GroupIndex)
 { 
 	uint LinearThreadId = GetUnWrappedDispatchThreadId(GroupId, GroupThreadIndex, 64);
 
-	if (LinearThreadId >= NumVertices) return;
+	if (LinearThreadId >= MaxNumThreads) return;
 
 	ResolvedView = ResolveView();
 
 	bool bUsingIndirectDraw = UsingIndirectDraw != 0;
 	uint NumActualVertices = bUsingIndirectDraw ? GetNumRayTracingDynamicMeshVerticesIndirect() : NumVertices;
 
-	uint VertexIndex = MinVertexIndex + LinearThreadId;
-
-	if (LinearThreadId < NumActualVertices)
-	{
-		FVertexFactoryInput Input = LoadVertexFactoryInputForDynamicUpdate(VertexIndex / 3, VertexIndex % 3, PrimitiveId, InstanceId);
+#if PER_VERTEX_MASK_ENABLED
+	uint3 Indices;
+	Indices.x = IndexBuffer[IndexBufferOffset + LinearThreadId * 3];
+	Indices.y = IndexBuffer[IndexBufferOffset + LinearThreadId * 3 + 1];
+	Indices.z = IndexBuffer[IndexBufferOffset + LinearThreadId * 3 + 2];
 
-#if USE_INSTANCING
-		Input.InstanceId = InstanceId;
-#elif USE_INSTANCE_CULLING
-		Input.DrawInstanceId = InstanceId;
-#endif
-		FVertexFactoryIntermediates VFIntermediates = GetVertexFactoryIntermediates(Input);
+	bool bHideTriangle = false;
+	float3 ReferenceLocalPosition = 0.0f;	
 
-		float4 TranslatedWorldPosition = VertexFactoryGetWorldPosition(Input, VFIntermediates);
-		float3x3 TangentToLocal = VertexFactoryGetTangentToLocal(Input, VFIntermediates);
-		FMaterialVertexParameters VertexParameters = GetMaterialVertexParameters(Input, VFIntermediates, TranslatedWorldPosition.xyz, TangentToLocal);
+	for(int Index = 0; Index < 3; Index++)
+	{
+		const uint VertexIndex = Indices[Index];
+		const uint OutputVertexIndex = OutputVertexBaseIndex + (LinearThreadId * 3 + Index) * 3;
+#else
+		const uint VertexIndex = LinearThreadId;
+		const uint OutputVertexIndex = OutputVertexBaseIndex + VertexIndex * 3;
+#endif // PER_VERTEX_MASK_ENABLED
 
-		if (bApplyWorldPositionOffset)
+		if (VertexIndex < NumActualVertices)
 		{
-			// We must guarantee that no NaNs are produced by WPO to avoid a possibility of producing invalid BLAS.
-			// DXR specification allows NaN positions to mark "inactive" primitives, but only when X component each vertex is NaN.
-			// This is impossible to guarantee here, since WPO is evaluated per vertex and not per triangle.
-			TranslatedWorldPosition.xyz += MakeFinite(GetMaterialWorldPositionOffset(VertexParameters));
+			const uint ActualVertexIndex = MinVertexIndex + VertexIndex;
+			FVertexFactoryInput Input = LoadVertexFactoryInputForDynamicUpdate(ActualVertexIndex / 3, ActualVertexIndex % 3, PrimitiveId, InstanceId);
+
+	#if USE_INSTANCING
+			Input.InstanceId = InstanceId;
+	#elif USE_INSTANCE_CULLING
+			Input.DrawInstanceId = InstanceId;
+	#endif
+			FVertexFactoryIntermediates VFIntermediates = GetVertexFactoryIntermediates(Input);
+
+			float4 TranslatedWorldPosition = VertexFactoryGetWorldPosition(Input, VFIntermediates);
+			float3x3 TangentToLocal = VertexFactoryGetTangentToLocal(Input, VFIntermediates);
+			FMaterialVertexParameters VertexParameters = GetMaterialVertexParameters(Input, VFIntermediates, TranslatedWorldPosition.xyz, TangentToLocal);
+
+			if (bApplyWorldPositionOffset)
+			{
+				// We must guarantee that no NaNs are produced by WPO to avoid a possibility of producing invalid BLAS.
+				// DXR specification allows NaN positions to mark "inactive" primitives, but only when X component each vertex is NaN.
+				// This is impossible to guarantee here, since WPO is evaluated per vertex and not per triangle.
+				TranslatedWorldPosition.xyz += MakeFinite(GetMaterialWorldPositionOffset(VertexParameters));
+			}
+
+			float3 LocalPosition;
+	#if USE_INSTANCING || USE_INSTANCE_CULLING
+			//Reverse transform to neutral space
+			float3 WorldPosition = TranslatedWorldPosition.xyz - DFHackToFloat(ResolvedView.PreViewTranslation);
+			LocalPosition = mul(float4(WorldPosition, 1), WorldToInstance).xyz;
+	#elif RAY_TRACING_DYNAMIC_MESH_IN_LOCAL_SPACE
+			// Move the point back into local space because the RT instance will be placed there
+			// This the recommended default behavior so that the transform applied by the VertexFactory and the RT transform are always the same
+			float4x4 TranslatedWorldToLocal = DFFastToTranslatedWorld(VertexFactoryGetWorldToLocal(VFIntermediates), ResolvedView.PreViewTranslation);
+			LocalPosition = mul(TranslatedWorldPosition, TranslatedWorldToLocal).xyz;
+	#elif RAY_TRACING_DYNAMIC_MESH_IN_WORLD_SPACE
+			// There are a few geometry types which natively generate positions in world space and therefore not need any extra processing here.
+			// This behavior must be explicitly opted-in to so as to make the code more self-documenting and avoid un-intended mismatched transforms.
+			LocalPosition = TranslatedWorldPosition.xyz - DFHackToFloat(ResolvedView.PreViewTranslation);
+	#else
+	#error "Are you sure you want mesh vertices to be in world space? Please update the VertexFactory to report its requirement!"
+			LocalPosition = TranslatedWorldPosition.xyz - DFHackToFloat(ResolvedView.PreViewTranslation);
+	#endif
+
+	#if PER_VERTEX_MASK_ENABLED
+			FVertexFactoryInterpolantsVSToPS Interpolants = VertexFactoryGetInterpolantsVSToPS(Input, VFIntermediates, VertexParameters);
+
+			float4 SvPosition = 0;
+			FMaterialPixelParameters MaterialPixelParameters = GetMaterialPixelParameters(Interpolants, SvPosition);
+
+			FPixelMaterialInputs PixelMaterialInputs = (FPixelMaterialInputs)0;
+			CalcMaterialParameters(MaterialPixelParameters, PixelMaterialInputs, SvPosition, true);
+		
+			const float Opacity = GetMaterialMask(PixelMaterialInputs);
+			if (Opacity < 0.0f)
+			{
+				bHideTriangle = true;
+				ReferenceLocalPosition = LocalPosition;
+			}
+	#endif		
+
+			RWVertexPositions[OutputVertexIndex + 0] = LocalPosition.x;
+			RWVertexPositions[OutputVertexIndex + 1] = LocalPosition.y;
+			RWVertexPositions[OutputVertexIndex + 2] = LocalPosition.z;
+		}
+		else
+		{
+			RWVertexPositions[OutputVertexIndex + 0] = asfloat(0xFFFFFFFF);
+			RWVertexPositions[OutputVertexIndex + 1] = 0;
+			RWVertexPositions[OutputVertexIndex + 2] = 0;
 		}
 
-		float3 LocalPosition;
-#if USE_INSTANCING || USE_INSTANCE_CULLING
-		//Reverse transform to neutral space
-		float3 WorldPosition = TranslatedWorldPosition.xyz - DFHackToFloat(ResolvedView.PreViewTranslation);
-		LocalPosition = mul(float4(WorldPosition, 1), WorldToInstance).xyz;
-#elif RAY_TRACING_DYNAMIC_MESH_IN_LOCAL_SPACE
-		// Move the point back into local space because the RT instance will be placed there
-		// This the recommended default behavior so that the transform applied by the VertexFactory and the RT transform are always the same
-		float4x4 TranslatedWorldToLocal = DFFastToTranslatedWorld(VertexFactoryGetWorldToLocal(VFIntermediates), ResolvedView.PreViewTranslation);
-		LocalPosition = mul(TranslatedWorldPosition, TranslatedWorldToLocal).xyz;
-#elif RAY_TRACING_DYNAMIC_MESH_IN_WORLD_SPACE
-		// There are a few geometry types which natively generate positions in world space and therefore not need any extra processing here.
-		// This behavior must be explicitly opted-in to so as to make the code more self-documenting and avoid un-intended mismatched transforms.
-		LocalPosition = TranslatedWorldPosition.xyz - DFHackToFloat(ResolvedView.PreViewTranslation);
-#else
-#error "Are you sure you want mesh vertices to be in world space? Please update the VertexFactory to report its requirement!"
-		LocalPosition = TranslatedWorldPosition.xyz - DFHackToFloat(ResolvedView.PreViewTranslation);
-#endif
-
-		RWVertexPositions[OutputVertexBaseIndex + VertexIndex * 3 + 0] = LocalPosition.x;
-		RWVertexPositions[OutputVertexBaseIndex + VertexIndex * 3 + 1] = LocalPosition.y;
-		RWVertexPositions[OutputVertexBaseIndex + VertexIndex * 3 + 2] = LocalPosition.z;
+#if PER_VERTEX_MASK_ENABLED
 	}
-	else
+
+	if (bHideTriangle)
 	{
-		RWVertexPositions[OutputVertexBaseIndex + VertexIndex * 3 + 0] = asfloat(0xFFFFFFFF);
-		RWVertexPositions[OutputVertexBaseIndex + VertexIndex * 3 + 1] = 0;
-		RWVertexPositions[OutputVertexBaseIndex + VertexIndex * 3 + 2] = 0;
+		HideTriangle(LinearThreadId, ReferenceLocalPosition);
 	}
+#endif
 }
 #endif
diff --git a/Engine/Source/Runtime/Engine/Public/SceneManagement.h b/Engine/Source/Runtime/Engine/Public/SceneManagement.h
--- a/Engine/Source/Runtime/Engine/Public/SceneManagement.h
+++ b/Engine/Source/Runtime/Engine/Public/SceneManagement.h
@@ -1529,6 +1529,9 @@ struct FRayTracingDynamicGeometryUpdateParams
 
 	uint32 InstanceId = 0;
 	FMatrix44f WorldToInstance = FMatrix44f::Identity;
+
+	bool bAlphaMasked = false;
+	FBufferRHIRef IndexBuffer = nullptr;
 };
 
 struct UE_DEPRECATED(5.5, "Use FRayTracingInstanceCollector instead.") FRayTracingMaterialGatheringContext
diff --git a/Engine/Source/Runtime/Landscape/Private/LandscapeRender.cpp b/Engine/Source/Runtime/Landscape/Private/LandscapeRender.cpp
--- a/Engine/Source/Runtime/Landscape/Private/LandscapeRender.cpp
+++ b/Engine/Source/Runtime/Landscape/Private/LandscapeRender.cpp
@@ -323,6 +323,13 @@ static FAutoConsoleVariableRef CVarLandscapeRayTracingVertexLODBias(
 	ECVF_RenderThreadSafe
 );
 
+static TAutoConsoleVariable<int32> CVarRayTracingLandscapeVertexMask(
+	TEXT("r.RayTracing.Geometry.Landscape.VertexMask"),
+	1,
+	TEXT("If landscape should use vertex mask when building ray tracing dynamic geometry"),
+	FConsoleVariableDelegate::CreateStatic(&OnCVarNeedingRenderStateInvalidationChanged),
+	ECVF_RenderThreadSafe
+);
 
 // ----------------------------------------------------------------------------------
 
@@ -3154,11 +3161,11 @@ void FLandscapeComponentSceneProxy::GetDynamicRayTracingInstances(FRayTracingIns
 		}
 	}
 
+	const FMaterialRenderProxy* FallbackMaterialRenderProxyPtr = nullptr;
+	const FMaterial& Material = SelectedMaterial->GetMaterialWithFallback(GetScene().GetFeatureLevel(), FallbackMaterialRenderProxyPtr);
+
 	if (GLandscapeRayTracingGeometryDetectTextureStreaming > 0)
 	{
-		const FMaterialRenderProxy* FallbackMaterialRenderProxyPtr = nullptr;
-		const FMaterial& Material = SelectedMaterial->GetMaterialWithFallback(GetScene().GetFeatureLevel(), FallbackMaterialRenderProxyPtr);
-
 		if (Material.GetRenderingThreadShaderMap()->UsesWorldPositionOffset())
 		{
 			const FMaterialRenderProxy* MaterialRenderProxy = FallbackMaterialRenderProxyPtr ? FallbackMaterialRenderProxyPtr : SelectedMaterial;
@@ -3231,6 +3238,7 @@ void FLandscapeComponentSceneProxy::GetDynamicRayTracingInstances(FRayTracingIns
 							SectionRayTracingState.RayTracingDynamicVertexBuffer.Release();
 						}
 
+						const bool bAlphaMasked = Material.IsMasked() && (CVarRayTracingLandscapeVertexMask.GetValueOnRenderThread() != 0);
 						Collector.AddRayTracingGeometryUpdate(
 							FRayTracingDynamicGeometryUpdateParams
 							{
@@ -3241,7 +3249,11 @@ void FLandscapeComponentSceneProxy::GetDynamicRayTracingInstances(FRayTracingIns
 								(uint32)FMath::Square(RayTracingSectionLODVertCount - 1) * 2,
 								&SectionRayTracingState.Geometry,
 								&SectionRayTracingState.RayTracingDynamicVertexBuffer,
-								true
+								true,
+								0,
+								FMatrix44f::Identity,
+								bAlphaMasked,
+								SharedBuffers->ZeroOffsetIndexBuffers[CurrentLODIndex]->IndexBufferRHI
 							}
 						);
 					}
diff --git a/Engine/Source/Runtime/Renderer/Private/RayTracing/RayTracingDynamicGeometryUpdateManager.cpp b/Engine/Source/Runtime/Renderer/Private/RayTracing/RayTracingDynamicGeometryUpdateManager.cpp
--- a/Engine/Source/Runtime/Renderer/Private/RayTracing/RayTracingDynamicGeometryUpdateManager.cpp
+++ b/Engine/Source/Runtime/Renderer/Private/RayTracing/RayTracingDynamicGeometryUpdateManager.cpp
@@ -61,6 +61,11 @@ static TAutoConsoleVariable<int32> CVarRTDynGeomForceBuildMinUpdatesSinceLastBui
 class FRayTracingDynamicGeometryConverterCS : public FMeshMaterialShader
 {
 	DECLARE_SHADER_TYPE(FRayTracingDynamicGeometryConverterCS, MeshMaterial);
+
+	class FVertexMask : SHADER_PERMUTATION_BOOL("USE_VERTEX_MASK");
+
+	using FPermutationDomain = TShaderPermutationDomain<FVertexMask>;
+
 public:
 	FRayTracingDynamicGeometryConverterCS(const FMeshMaterialShaderType::CompiledShaderInitializerType& Initializer)
 		: FMeshMaterialShader(Initializer)
@@ -70,19 +75,39 @@ public:
 		RWVertexPositions.Bind(Initializer.ParameterMap, TEXT("RWVertexPositions"));
 		UsingIndirectDraw.Bind(Initializer.ParameterMap, TEXT("UsingIndirectDraw"));
 		NumVertices.Bind(Initializer.ParameterMap, TEXT("NumVertices"));
+		MaxNumThreads.Bind(Initializer.ParameterMap, TEXT("MaxNumThreads"));
 		MinVertexIndex.Bind(Initializer.ParameterMap, TEXT("MinVertexIndex"));
 		PrimitiveId.Bind(Initializer.ParameterMap, TEXT("PrimitiveId"));
 		OutputVertexBaseIndex.Bind(Initializer.ParameterMap, TEXT("OutputVertexBaseIndex"));
 		bApplyWorldPositionOffset.Bind(Initializer.ParameterMap, TEXT("bApplyWorldPositionOffset"));
 		InstanceId.Bind(Initializer.ParameterMap, TEXT("InstanceId"));
 		WorldToInstance.Bind(Initializer.ParameterMap, TEXT("WorldToInstance"));
+		IndexBuffer.Bind(Initializer.ParameterMap, TEXT("IndexBuffer"));
+		IndexBufferOffset.Bind(Initializer.ParameterMap, TEXT("IndexBufferOffset"));
 	}
 
 	FRayTracingDynamicGeometryConverterCS() = default;
 
 	static bool ShouldCompilePermutation(const FMeshMaterialShaderPermutationParameters& Parameters)
 	{
-		return Parameters.VertexFactoryType->SupportsRayTracingDynamicGeometry() && IsRayTracingEnabledForProject(Parameters.Platform) && RHISupportsRayTracing(Parameters.Platform);
+		const FPermutationDomain PermutationVector(Parameters.PermutationId);
+
+		if (!IsRayTracingEnabledForProject(Parameters.Platform))
+		{
+			return false;
+		}
+
+		if (!Parameters.VertexFactoryType->SupportsRayTracingDynamicGeometry())
+		{
+			return false;
+		}
+
+		if (PermutationVector.Get<FVertexMask>())
+		{
+			return Parameters.MaterialParameters.BlendMode == BLEND_Masked;
+		}
+		
+		return true;
 	}
 
 	static void ModifyCompilationEnvironment(const FMaterialShaderPermutationParameters& Parameters, FShaderCompilerEnvironment& OutEnvironment)
@@ -122,6 +147,7 @@ public:
 
 	LAYOUT_FIELD(FShaderResourceParameter, RWVertexPositions);
 	LAYOUT_FIELD(FShaderParameter, UsingIndirectDraw);
+	LAYOUT_FIELD(FShaderParameter, MaxNumThreads);
 	LAYOUT_FIELD(FShaderParameter, NumVertices);
 	LAYOUT_FIELD(FShaderParameter, MinVertexIndex);
 	LAYOUT_FIELD(FShaderParameter, PrimitiveId);
@@ -129,6 +155,8 @@ public:
 	LAYOUT_FIELD(FShaderParameter, OutputVertexBaseIndex);
 	LAYOUT_FIELD(FShaderParameter, InstanceId);
 	LAYOUT_FIELD(FShaderParameter, WorldToInstance);
+	LAYOUT_FIELD(FShaderResourceParameter, IndexBuffer);
+	LAYOUT_FIELD(FShaderParameter, IndexBufferOffset);
 };
 
 IMPLEMENT_MATERIAL_SHADER_TYPE(, FRayTracingDynamicGeometryConverterCS, TEXT("/Engine/Private/RayTracing/RayTracingDynamicMesh.usf"), TEXT("RayTracingDynamicGeometryConverterCS"), SF_Compute);
@@ -269,7 +297,18 @@ void FRayTracingDynamicGeometryUpdateManager::AddDynamicGeometryToUpdate(
 )
 {
 	FRayTracingGeometry& Geometry = *UpdateParams.Geometry;
-	uint32 NumMaxVertices = UpdateParams.NumVertices;
+
+	uint32 NumVertices = UpdateParams.NumVertices;
+	uint32 VertexBufferSize = UpdateParams.VertexBufferSize;
+
+	if (UpdateParams.bAlphaMasked)
+	{
+		check(UpdateParams.IndexBuffer);
+		check(!UpdateParams.bUsingIndirectDraw);
+
+		NumVertices = UpdateParams.NumTriangles * 3;
+		VertexBufferSize = NumVertices * (uint32)sizeof(FVector3f);
+	}	
 
 	FRWBuffer* RWBuffer = UpdateParams.Buffer;
 	uint32 VertexBufferOffset = 0;
@@ -291,20 +330,20 @@ void FRayTracingDynamicGeometryUpdateManager::AddDynamicGeometryToUpdate(
 		// If update params didn't provide a buffer then use a shared vertex position buffer
 		if (RWBuffer == nullptr)
 		{
-			RWBuffer = AllocateSharedBuffer(RHICmdList, UpdateParams.VertexBufferSize, VertexBufferOffset);
+			RWBuffer = AllocateSharedBuffer(RHICmdList, VertexBufferSize, VertexBufferOffset);
 			bUseSharedVertexBuffer = true;
 		}
 		check(IsAligned(VertexBufferOffset, RHI_RAW_VIEW_ALIGNMENT));
 		
-		AddDispatchCommands(Scene, View, PrimitiveSceneProxy, UpdateParams, PrimitiveId, RWBuffer, VertexBufferOffset, GeometryBuildParams);
+		AddDispatchCommands(RHICmdList, Scene, View, PrimitiveSceneProxy, UpdateParams, PrimitiveId, RWBuffer, NumVertices, VertexBufferOffset, VertexBufferSize, GeometryBuildParams);
 	}
 
 	bool bRefit = true;
 
 	// Optionally resize the buffer when not shared (could also be lazy allocated and still empty)
-	if (!bUseSharedVertexBuffer && RWBuffer && RWBuffer->NumBytes != UpdateParams.VertexBufferSize)
+	if (!bUseSharedVertexBuffer && RWBuffer && RWBuffer->NumBytes != VertexBufferSize)
 	{
-		RWBuffer->Initialize(RHICmdList, TEXT("FRayTracingDynamicGeometryUpdateManager::RayTracingDynamicVertexBuffer"), sizeof(float), UpdateParams.VertexBufferSize / sizeof(float), PF_R32_FLOAT, BUF_UnorderedAccess | BUF_ShaderResource);
+		RWBuffer->Initialize(RHICmdList, TEXT("FRayTracingDynamicGeometryUpdateManager::RayTracingDynamicVertexBuffer"), sizeof(float), VertexBufferSize / sizeof(float), PF_R32_FLOAT, BUF_UnorderedAccess | BUF_ShaderResource);
 		bRefit = false;
 	}
 
@@ -327,11 +366,16 @@ void FRayTracingDynamicGeometryUpdateManager::AddDynamicGeometryToUpdate(
 		Geometry.Initializer.Segments.Empty();
 		FRayTracingGeometrySegment Segment;
 		Segment.NumPrimitives = UpdateParams.NumTriangles;
-		Segment.MaxVertices = UpdateParams.NumVertices;
+		Segment.MaxVertices = NumVertices;
 		Geometry.Initializer.Segments.Add(Segment);
 		bRefit = false;
 	}
 
+	if (UpdateParams.bAlphaMasked)
+	{
+		Geometry.Initializer.IndexBuffer = nullptr;
+	}
+
 	if (RWBuffer)
 	{
 		for (FRayTracingGeometrySegment& Segment : Geometry.Initializer.Segments)
@@ -419,13 +463,16 @@ FRWBuffer* FRayTracingDynamicGeometryUpdateManager::AllocateSharedBuffer(FRHICom
 }
 
 void FRayTracingDynamicGeometryUpdateManager::AddDispatchCommands(
+	FRHICommandListBase& RHICmdList,
 	const FScene* Scene, 
 	const FSceneView* View, 
 	const FPrimitiveSceneProxy* PrimitiveSceneProxy, 
 	const FRayTracingDynamicGeometryUpdateParams& UpdateParams,
 	uint32 PrimitiveId,
 	FRWBuffer* RWBuffer,
+	uint32 NumVertices,
 	uint32 VertexBufferOffset,
+	uint32 VertexBufferSize,
 	FRayTracingDynamicGeometryUpdateManager::FRayTracingDynamicGeometryBuildParams& GeometryBuildParams)
 {	
 	const int32 PSOCollectorIndex = FPSOCollectorCreateManager::GetIndex(EShadingPath::Deferred, RayTracingDynamicGeometryPSOCollectorName);
@@ -452,8 +499,11 @@ void FRayTracingDynamicGeometryUpdateManager::AddDispatchCommands(
 
 				FMeshComputeDispatchCommand DispatchCmd;
 
+				FRayTracingDynamicGeometryConverterCS::FPermutationDomain PermutationVectorCS;
+				PermutationVectorCS.Set<FRayTracingDynamicGeometryConverterCS::FVertexMask>(UpdateParams.bAlphaMasked);
+
 				FMaterialShaderTypes ShaderTypes;
-				ShaderTypes.AddShaderType<FRayTracingDynamicGeometryConverterCS>();
+				ShaderTypes.AddShaderType<FRayTracingDynamicGeometryConverterCS>(PermutationVectorCS.ToDimensionValueId());
 
 				FMaterialShaders MaterialShaders;
 				if (Material.TryGetShaders(ShaderTypes, MeshBatch.VertexFactory->GetType(), MaterialShaders))
@@ -477,19 +527,18 @@ void FRayTracingDynamicGeometryUpdateManager::AddDispatchCommands(
 					FVertexInputStreamArray DummyArray;
 					FMeshMaterialShader::GetElementShaderBindings(Shader, Scene, View, MeshBatch.VertexFactory, EVertexInputStreamType::Default, Scene->GetFeatureLevel(), PrimitiveSceneProxy, MeshBatch, MeshBatch.Elements[0], ShaderElementData, SingleShaderBindings, DummyArray);
 
-					DispatchCmd.TargetBuffer = RWBuffer;
-					DispatchCmd.NumMaxVertices = UpdateParams.NumVertices;
+					DispatchCmd.TargetBuffer = RWBuffer;					
 
 					// Setup the loose parameters directly on the binding
 					uint32 OutputVertexBaseIndex = VertexBufferOffset / sizeof(float);
 					uint32 MinVertexIndex = MeshBatch.Elements[0].MinVertexIndex;
-					uint32 NumCPUVertices = UpdateParams.NumVertices;
+					uint32 NumCPUVertices = NumVertices;
 					if (MeshBatch.Elements[0].MinVertexIndex < MeshBatch.Elements[0].MaxVertexIndex)
 					{
 						NumCPUVertices = 1 + MeshBatch.Elements[0].MaxVertexIndex - MeshBatch.Elements[0].MinVertexIndex;
 					}
 
-					const uint32 VertexBufferNumElements = UpdateParams.VertexBufferSize / sizeof(FVector3f) - MinVertexIndex;
+					const uint32 VertexBufferNumElements = VertexBufferSize / sizeof(FVector3f) - MinVertexIndex;
 					if (!ensureMsgf(NumCPUVertices <= VertexBufferNumElements,
 					                TEXT("%s: Vertex buffer contains %d vertices, but RayTracingDynamicGeometryConverterCS dispatch command expects at least %d."),
 					                *PrimitiveSceneProxy->GetOwnerName().ToString(), VertexBufferNumElements, NumCPUVertices))
@@ -508,6 +557,34 @@ void FRayTracingDynamicGeometryUpdateManager::AddDispatchCommands(
 					SingleShaderBindings.Add(Shader->InstanceId, UpdateParams.InstanceId);
 					SingleShaderBindings.Add(Shader->WorldToInstance, UpdateParams.WorldToInstance);
 
+					if (UpdateParams.bAlphaMasked)
+					{
+						FRHIBuffer* IndexBufferRHI = UpdateParams.IndexBuffer;
+
+						const uint32 IndexStride = IndexBufferRHI->GetStride();
+						const uint32 NumTriangles = UpdateParams.NumTriangles;
+						const uint32 IndexBufferOffset = UpdateParams.Geometry->Initializer.IndexBufferOffset / IndexStride + MeshBatch.Elements[0].FirstIndex;
+
+						SingleShaderBindings.Add(Shader->IndexBuffer,
+							RHICmdList.CreateShaderResourceView(IndexBufferRHI,
+								FRHIViewDesc::CreateBufferSRV()
+								.SetType(FRHIViewDesc::EBufferType::Typed)
+								.SetFormat(IndexStride == 4 ? PF_R32_UINT : PF_R16_UINT))
+						);
+
+						SingleShaderBindings.Add(Shader->MaxNumThreads, NumTriangles);
+						SingleShaderBindings.Add(Shader->IndexBufferOffset, IndexBufferOffset);
+
+						DispatchCmd.NumThreads = NumTriangles;
+					}
+					else
+					{
+						SingleShaderBindings.Add(Shader->MaxNumThreads, NumCPUVertices);
+						SingleShaderBindings.Add(Shader->IndexBufferOffset, 0);
+
+						DispatchCmd.NumThreads = NumVertices;
+					}				
+
 				#if MESH_DRAW_COMMAND_DEBUG_DATA
 					ShaderBindings.Finalize(&MeshProcessorShaders);
 				#endif
@@ -885,7 +962,7 @@ void FRayTracingDynamicGeometryUpdateManager::DispatchUpdates(FRHICommandList& R
 				Cmd.ShaderBindings.SetParameters(BatchedParameters, &ShaderBindingState);
 				RHICmdList.SetBatchedShaderParameters(CurrentShader, BatchedParameters);
 
-				const FIntVector NumWrappedThreadGroups = FComputeShaderUtils::GetGroupCountWrapped(Cmd.NumMaxVertices, 64);
+				const FIntVector NumWrappedThreadGroups = FComputeShaderUtils::GetGroupCountWrapped(Cmd.NumThreads, 64);
 				RHICmdList.DispatchComputeShader(NumWrappedThreadGroups.X, NumWrappedThreadGroups.Y, NumWrappedThreadGroups.Z);
 			}
 
diff --git a/Engine/Source/Runtime/Renderer/Public/RayTracingDynamicGeometryUpdateManager.h b/Engine/Source/Runtime/Renderer/Public/RayTracingDynamicGeometryUpdateManager.h
--- a/Engine/Source/Runtime/Renderer/Public/RayTracingDynamicGeometryUpdateManager.h
+++ b/Engine/Source/Runtime/Renderer/Public/RayTracingDynamicGeometryUpdateManager.h
@@ -27,7 +27,7 @@ struct FMeshComputeDispatchCommand
 	FMeshDrawShaderBindings ShaderBindings;
 	TShaderRef<FRayTracingDynamicGeometryConverterCS> MaterialShader;
 
-	uint32 NumMaxVertices;
+	uint32 NumThreads;
 	uint32 NumCPUVertices;
 	FRWBuffer* TargetBuffer;
 };
@@ -91,13 +91,16 @@ private:
 		int32 SegmentOffset = -1;
 	};
 
-	void AddDispatchCommands(const FScene* Scene, 
+	void AddDispatchCommands(FRHICommandListBase& RHICmdList,
+							 const FScene* Scene, 
 	                         const FSceneView* View, 
 	                         const FPrimitiveSceneProxy* PrimitiveSceneProxy, 
 	                         const FRayTracingDynamicGeometryUpdateParams& Params,
 	                         uint32 PrimitiveId,
 							 FRWBuffer* RWBuffer,
+							 uint32 NumVertices,
 							 uint32 VertexBufferOffset,
+							 uint32 VertexBufferSize,
 	                         FRayTracingDynamicGeometryBuildParams& BuildParams);
 	FRWBuffer* AllocateSharedBuffer(FRHICommandListBase& RHICmdList, uint32 VertexBufferSize, uint32& OutVertexBufferOffset);
 
