Add Substrate material buffer clear option for debugging purpose.

diff --git a/Engine/Source/Runtime/Renderer/Private/Substrate/Substrate.cpp b/Engine/Source/Runtime/Renderer/Private/Substrate/Substrate.cpp
--- a/Engine/Source/Runtime/Renderer/Private/Substrate/Substrate.cpp
+++ b/Engine/Source/Runtime/Renderer/Private/Substrate/Substrate.cpp
@@ -75,6 +75,12 @@ static TAutoConsoleVariable<int32> CVarSubstrateStochasticLightingActive(
 	TEXT("Activate stochastic lighting to get better performance (runtime toggle for debugging). Requires r.Substrate.StochasticLighting to be enabled (which is read-only)"),
 	ECVF_RenderThreadSafe);
 
+static TAutoConsoleVariable<int32> CVarSubstrateClearMaterialBuffer(
+	TEXT("r.Substrate.Debug.ClearMaterialBuffer"),
+	0,
+	TEXT("Clear Substrate material buffer before writing to it for debugging purpose"),
+	ECVF_RenderThreadSafe);
+
 IMPLEMENT_GLOBAL_SHADER_PARAMETER_STRUCT(FSubstrateGlobalUniformParameters, "Substrate");
 
 void FSubstrateViewData::Reset()
@@ -571,6 +577,14 @@ void InitialiseSubstrateFrameSceneData(FRDGBuilder& GraphBuilder, FSceneRenderer
 		// See AppendSubstrateMRTs
 		check(SUBSTRATE_BASE_PASS_MRT_OUTPUT_COUNT <= (SliceCount - SliceCountSSS - SliceCountAdvDebug)); // We want enough slice for MRTs but also do not want the SSSData to be a MRT.
 		Out.MaterialTextureArrayUAVWithoutRTs = GraphBuilder.CreateUAV(FRDGTextureUAVDesc(Out.MaterialTextureArray, 0, PF_Unknown, SUBSTRATE_BASE_PASS_MRT_OUTPUT_COUNT, SliceCount - SUBSTRATE_BASE_PASS_MRT_OUTPUT_COUNT));
+
+		if (CVarSubstrateClearMaterialBuffer.GetValueOnRenderThread() > 0)
+		{
+			for (uint32 SliceIt=0; SliceIt<SliceCount; ++SliceIt)
+			{
+				AddClearUAVPass(GraphBuilder, GraphBuilder.CreateUAV(FRDGTextureUAVDesc(Out.MaterialTextureArray, 0, PF_R32_UINT, SliceIt, 1)), 0u);
+			}
+		}
 	}
 
 	// Rough diffuse model
