[OIT] Add triangle sorting support for skinned mesh.

diff --git a/Engine/Source/Runtime/Engine/Classes/Components/SkinnedMeshComponent.h b/Engine/Source/Runtime/Engine/Classes/Components/SkinnedMeshComponent.h
--- a/Engine/Source/Runtime/Engine/Classes/Components/SkinnedMeshComponent.h
+++ b/Engine/Source/Runtime/Engine/Classes/Components/SkinnedMeshComponent.h
@@ -883,6 +883,10 @@ public:
 	uint8 bIgnoreMasterPoseComponentLOD : 1;
 #endif // WITH_EDITORONLY_DATA
 
+	/** Enable dynamic sort mesh's triangles to remove ordering issue when rendered with a translucent material */
+	UPROPERTY(EditAnywhere, AdvancedDisplay, BlueprintReadOnly, Category = Lighting, meta = (UIMin = "0", UIMax = "1", DisplayName = "Sort Triangles"))
+	uint8 bSortTriangles : 1;
+
 protected:
 	/** Are we using double buffered ComponentSpaceTransforms */
 	uint8 bDoubleBufferedComponentSpaceTransforms : 1;
diff --git a/Engine/Source/Runtime/Engine/Private/GPUSkinVertexFactory.cpp b/Engine/Source/Runtime/Engine/Private/GPUSkinVertexFactory.cpp
--- a/Engine/Source/Runtime/Engine/Private/GPUSkinVertexFactory.cpp
+++ b/Engine/Source/Runtime/Engine/Private/GPUSkinVertexFactory.cpp
@@ -1055,6 +1055,7 @@ IMPLEMENT_GPUSKINNING_VERTEX_FACTORY_TYPE(TGPUSkinVertexFactory, "/Engine/Privat
 	| EVertexFactoryFlags::SupportsDynamicLighting
 	| EVertexFactoryFlags::SupportsPSOPrecaching
 	| EVertexFactoryFlags::SupportsCachingMeshDrawCommands
+	| EVertexFactoryFlags::SupportsTriangleSorting
 );
 
 /*-----------------------------------------------------------------------------
diff --git a/Engine/Source/Runtime/Engine/Private/LocalVertexFactory.cpp b/Engine/Source/Runtime/Engine/Private/LocalVertexFactory.cpp
--- a/Engine/Source/Runtime/Engine/Private/LocalVertexFactory.cpp
+++ b/Engine/Source/Runtime/Engine/Private/LocalVertexFactory.cpp
@@ -589,4 +589,5 @@ IMPLEMENT_VERTEX_FACTORY_TYPE(FLocalVertexFactory,"/Engine/Private/LocalVertexFa
 	| EVertexFactoryFlags::SupportsPSOPrecaching
 	| EVertexFactoryFlags::SupportsGPUSkinPassThrough
 	| EVertexFactoryFlags::SupportsLumenMeshCards
+	| EVertexFactoryFlags::SupportsTriangleSorting
 );
diff --git a/Engine/Source/Runtime/Engine/Private/SkeletalMeshSceneProxy.cpp b/Engine/Source/Runtime/Engine/Private/SkeletalMeshSceneProxy.cpp
--- a/Engine/Source/Runtime/Engine/Private/SkeletalMeshSceneProxy.cpp
+++ b/Engine/Source/Runtime/Engine/Private/SkeletalMeshSceneProxy.cpp
@@ -140,6 +140,9 @@ FSkeletalMeshSceneProxy::FSkeletalMeshSceneProxy(const FSkinnedMeshSceneProxyDes
 
 	const bool bForceDefaultMaterial = InMeshDesc.ShouldRenderProxyFallbackToDefaultMaterial();
 
+	// Enable dynamic triangle reordering to remove/reduce sorting issue when rendered with a translucent material (i.e., order-independent-transparency)
+	bSupportsSortedTriangles = InMeshDesc.bSortTriangles;
+
 	// setup materials and performance classification for each LOD.
 	extern bool GForceDefaultMaterial;
 	const bool bCastShadow = InMeshDesc.CastShadow;
diff --git a/Engine/Source/Runtime/Engine/Private/SkinnedMeshSceneProxyDesc.cpp b/Engine/Source/Runtime/Engine/Private/SkinnedMeshSceneProxyDesc.cpp
--- a/Engine/Source/Runtime/Engine/Private/SkinnedMeshSceneProxyDesc.cpp
+++ b/Engine/Source/Runtime/Engine/Private/SkinnedMeshSceneProxyDesc.cpp
@@ -75,6 +75,7 @@ void FSkinnedMeshSceneProxyDesc::InitializeFromSkinnedMeshComponent(const USkinn
 	SelectedEditorSection = InComponent->GetSelectedEditorSection();
 	SelectedEditorMaterial = InComponent->GetSelectedEditorMaterial();
 #endif
+	bSortTriangles = InComponent->bSortTriangles;
 }
 
 USkinnedAsset* FSkinnedMeshSceneProxyDesc::GetSkinnedAsset() const
diff --git a/Engine/Source/Runtime/Engine/Public/GPUSkinVertexFactory.h b/Engine/Source/Runtime/Engine/Public/GPUSkinVertexFactory.h
--- a/Engine/Source/Runtime/Engine/Public/GPUSkinVertexFactory.h
+++ b/Engine/Source/Runtime/Engine/Public/GPUSkinVertexFactory.h
@@ -400,6 +400,11 @@ public:
 	virtual const uint32 GetColorIndexMask() const								{ return Data.IsValid() ? Data->ColorIndexMask : 0; }
 	virtual bool IsMorphTarget() const											{ return Data.IsValid() ? Data->bMorphTarget : false; }
 
+	virtual FShaderResourceViewRHIRef GetTriangleSortingPositionSRV() const override
+	{
+		return GetPositionsSRV();
+	}
+
 	void UpdateUniformBuffer(FRHICommandListBase& RHICmdList);
 
 	FRHIUniformBuffer* GetUniformBuffer() const
diff --git a/Engine/Source/Runtime/Engine/Public/LocalVertexFactory.h b/Engine/Source/Runtime/Engine/Public/LocalVertexFactory.h
--- a/Engine/Source/Runtime/Engine/Public/LocalVertexFactory.h
+++ b/Engine/Source/Runtime/Engine/Public/LocalVertexFactory.h
@@ -114,6 +114,11 @@ public:
 		VertexStreams.Add(FVertexInputStream(ColorStreamIndex, 0, ColorVertexBuffer->VertexBufferRHI));
 	}
 
+	virtual FShaderResourceViewRHIRef GetTriangleSortingPositionSRV() const override
+	{
+		return GetPositionsSRV();
+	}
+
 	inline FRHIShaderResourceView* GetPositionsSRV() const
 	{
 		return Data.PositionComponentSRV;
diff --git a/Engine/Source/Runtime/Engine/Public/SkinnedMeshSceneProxyDesc.h b/Engine/Source/Runtime/Engine/Public/SkinnedMeshSceneProxyDesc.h
--- a/Engine/Source/Runtime/Engine/Public/SkinnedMeshSceneProxyDesc.h
+++ b/Engine/Source/Runtime/Engine/Public/SkinnedMeshSceneProxyDesc.h
@@ -91,6 +91,7 @@ struct FSkinnedMeshSceneProxyDesc : public FPrimitiveSceneProxyDesc
 #if WITH_EDITORONLY_DATA
 	uint8 bClothPainting : 1 = false;
 #endif
+	uint32 bSortTriangles :1 = false;
 
 	/* SkeletalMesh Archetype Data End */
 
diff --git a/Engine/Source/Runtime/RenderCore/Public/VertexFactory.h b/Engine/Source/Runtime/RenderCore/Public/VertexFactory.h
--- a/Engine/Source/Runtime/RenderCore/Public/VertexFactory.h
+++ b/Engine/Source/Runtime/RenderCore/Public/VertexFactory.h
@@ -155,6 +155,7 @@ enum class EVertexFactoryFlags : uint32
 	SupportsLumenMeshCards                = 1u << 18,
 	SupportsLandscape                     = 1u << 19,
 	IncludesFwdShader                     = 1u << 20,
+	SupportsTriangleSorting               = 1u << 21,
 };
 ENUM_CLASS_FLAGS(EVertexFactoryFlags);
 
@@ -420,6 +421,7 @@ public:
 	bool SupportsLumenMeshCards() const                { return HasFlags(EVertexFactoryFlags::SupportsLumenMeshCards); }
 	bool SupportsLandscape() const                     { return HasFlags(EVertexFactoryFlags::SupportsLandscape); }
 	bool IncludesFwdShaderFile() const                 { return HasFlags(EVertexFactoryFlags::IncludesFwdShader); }
+	bool SupportsTriangleSorting() const               { return HasFlags(EVertexFactoryFlags::SupportsTriangleSorting); }
 
 	bool SupportsManualVertexFetch(ERHIFeatureLevel::Type InFeatureLevel) const 
 	{
@@ -721,6 +723,16 @@ public:
 		return GetType()->SupportsNullPixelShader();
 	}
 
+	bool SupportsTriangleSorting() const
+	{
+		return GetType()->SupportsTriangleSorting();
+	}
+
+	virtual FShaderResourceViewRHIRef GetTriangleSortingPositionSRV() const
+	{
+		return nullptr;
+	}
+
 #if WITH_EDITORONLY_DATA
 	virtual bool IsCoarseProxyMesh() const { return false; }
 #endif
diff --git a/Engine/Source/Runtime/Renderer/Private/OIT/OIT.cpp b/Engine/Source/Runtime/Renderer/Private/OIT/OIT.cpp
--- a/Engine/Source/Runtime/Renderer/Private/OIT/OIT.cpp
+++ b/Engine/Source/Runtime/Renderer/Private/OIT/OIT.cpp
@@ -660,8 +660,6 @@ static void AddOITSortTriangleIndexPass(
 	FTriangleSortingOrder SortType,
 	FOITDebugData& DebugData)
 {
-	static const FVertexFactoryType* CompatibleVF = FVertexFactoryType::GetVFByName(TEXT("FLocalVertexFactory"));
-
 	// Fat format: PF_R32G32_UINT | Compact format: PF_R32_UINT
 	const EPixelFormat PackedFormat = PF_R32_UINT;
 	const uint32 PackedFormatInBytes = 4;
@@ -669,7 +667,7 @@ static void AddOITSortTriangleIndexPass(
 	const bool bIsValid = 
 		MeshBatch.Mesh != nullptr && 
 		MeshBatch.Mesh->VertexFactory != nullptr &&
-		MeshBatch.Mesh->VertexFactory->GetType()->GetHashedName() == CompatibleVF->GetHashedName() &&
+		MeshBatch.Mesh->VertexFactory->SupportsTriangleSorting() &&
 		MeshBatch.Mesh->Elements.Num() > 0 && 
 		MeshBatch.Mesh->Elements[0].DynamicIndexBuffer.IndexBuffer != nullptr;
 	if (!bIsValid)
@@ -689,9 +687,7 @@ static void AddOITSortTriangleIndexPass(
 		SortedIndexBuffer->InitResource(GraphBuilder.RHICmdList);
 	}
 
-	const FLocalVertexFactory* VF = (const FLocalVertexFactory*)MeshBatch.Mesh->VertexFactory;
-	if (!VF) { return; }
-	const FShaderResourceViewRHIRef VertexPosition = VF->GetPositionsSRV();
+	const FShaderResourceViewRHIRef VertexPosition = MeshBatch.Mesh->VertexFactory->GetTriangleSortingPositionSRV();
 	if (!VertexPosition) { return; }
 
 	const FSortedIndexBuffer* OITIndexBuffer = (const FSortedIndexBuffer*)MeshBatch.Mesh->Elements[0].DynamicIndexBuffer.IndexBuffer;
@@ -896,14 +892,7 @@ namespace OIT
 
 	bool IsCompatible(const FMeshBatch& InMesh, ERHIFeatureLevel::Type InFeatureLevel)
 	{
-		// Only support local vertex factory at the moment as we need to have direct access to the position
-		static const FVertexFactoryType* CompatibleVF = FVertexFactoryType::GetVFByName(TEXT("FLocalVertexFactory"));
-
-		if (InMesh.IsTranslucent(InFeatureLevel))
-		{
-			return InMesh.VertexFactory->GetType()->GetHashedName() == CompatibleVF->GetHashedName();
-		}
-		return false;
+		return InMesh.IsTranslucent(InFeatureLevel) && InMesh.VertexFactory && InMesh.VertexFactory->SupportsTriangleSorting();
 	}
 
 	void AddSortTrianglesPass(FRDGBuilder& GraphBuilder, const FViewInfo& View, FOITSceneData& OITSceneData, FTriangleSortingOrder SortType)
@@ -1056,4 +1045,4 @@ namespace OIT
 			}
 		}
 	}
-}
\ No newline at end of file
+}
