Move stat for virtual texture feedback processing into local function, to reduce duplication at callers.
#rnx

diff --git a/Engine/Source/Runtime/Renderer/Private/DeferredShadingRenderer.cpp b/Engine/Source/Runtime/Renderer/Private/DeferredShadingRenderer.cpp
--- a/Engine/Source/Runtime/Renderer/Private/DeferredShadingRenderer.cpp
+++ b/Engine/Source/Runtime/Renderer/Private/DeferredShadingRenderer.cpp
@@ -306,7 +306,6 @@ DECLARE_GPU_STAT(PostRenderOpsFX);
 DECLARE_GPU_STAT_NAMED(Unaccounted, TEXT("[unaccounted]"));
 DECLARE_GPU_STAT(WaterRendering);
 DECLARE_GPU_STAT(HairRendering);
-DEFINE_GPU_DRAWCALL_STAT(VirtualTextureUpdate);
 DECLARE_GPU_STAT(UploadDynamicBuffers);
 DECLARE_GPU_STAT(PostOpaqueExtensions);
 DEFINE_GPU_STAT(CustomRenderPasses);
@@ -3596,9 +3595,6 @@ void FDeferredShadingSceneRenderer::Render(FRDGBuilder& GraphBuilder, const FSce
 
 		if (bUseVirtualTexturing)
 		{
-			RDG_EVENT_SCOPE_STAT(GraphBuilder, VirtualTextureUpdate, "VirtualTextureUpdate");
-			RDG_GPU_STAT_SCOPE(GraphBuilder, VirtualTextureUpdate);
-
 			VirtualTexture::EndFeedback(GraphBuilder);
 		}
 
diff --git a/Engine/Source/Runtime/Renderer/Private/MobileShadingRenderer.cpp b/Engine/Source/Runtime/Renderer/Private/MobileShadingRenderer.cpp
--- a/Engine/Source/Runtime/Renderer/Private/MobileShadingRenderer.cpp
+++ b/Engine/Source/Runtime/Renderer/Private/MobileShadingRenderer.cpp
@@ -1558,8 +1558,6 @@ void FMobileSceneRenderer::Render(FRDGBuilder& GraphBuilder, const FSceneRenderU
 
 	if (bUseVirtualTexturing)
 	{
-		RDG_EVENT_SCOPE_STAT(GraphBuilder, VirtualTextureUpdate, "VirtualTextureUpdate");
-		RDG_GPU_STAT_SCOPE(GraphBuilder, VirtualTextureUpdate);
 		VirtualTexture::EndFeedback(GraphBuilder);
 	}
 
diff --git a/Engine/Source/Runtime/Renderer/Private/Renderer.cpp b/Engine/Source/Runtime/Renderer/Private/Renderer.cpp
--- a/Engine/Source/Runtime/Renderer/Private/Renderer.cpp
+++ b/Engine/Source/Runtime/Renderer/Private/Renderer.cpp
@@ -578,9 +578,6 @@ void FRendererModule::DrawTileMesh(FCanvasRenderContext& RenderContext, FMeshPas
 
 		if (bUseVirtualTexturing)
 		{
-			RDG_EVENT_SCOPE_STAT(GraphBuilder, VirtualTextureUpdate, "VirtualTextureUpdate");
-			RDG_GPU_STAT_SCOPE(GraphBuilder, VirtualTextureUpdate);
-
 			VirtualTexture::EndFeedback(GraphBuilder);
 		}
 	}
diff --git a/Engine/Source/Runtime/Renderer/Private/SceneRendering.h b/Engine/Source/Runtime/Renderer/Private/SceneRendering.h
--- a/Engine/Source/Runtime/Renderer/Private/SceneRendering.h
+++ b/Engine/Source/Runtime/Renderer/Private/SceneRendering.h
@@ -116,7 +116,6 @@ struct FSceneCaptureLogUtils
 
 extern bool ShouldUseStereoLumenOptimizations();
 
-DECLARE_GPU_DRAWCALL_STAT_EXTERN(VirtualTextureUpdate);
 DECLARE_GPU_STAT_NAMED_EXTERN(Postprocessing, TEXT("Postprocessing"));
 DECLARE_GPU_STAT_NAMED_EXTERN(CustomRenderPasses, TEXT("CustomRenderPasses"));
 
diff --git a/Engine/Source/Runtime/Renderer/Private/VT/VirtualTextureFeedbackResource.cpp b/Engine/Source/Runtime/Renderer/Private/VT/VirtualTextureFeedbackResource.cpp
--- a/Engine/Source/Runtime/Renderer/Private/VT/VirtualTextureFeedbackResource.cpp
+++ b/Engine/Source/Runtime/Renderer/Private/VT/VirtualTextureFeedbackResource.cpp
@@ -11,6 +11,8 @@
 #include "VT/VirtualTextureScalability.h"
 #include "VT/VirtualTextureSystem.h"
 
+DECLARE_GPU_STAT(VirtualTextureUpdate);
+
 int32 GVirtualTextureFeedbackDefaultBufferSize = 4 * 1024;
 static FAutoConsoleVariableRef CVarVirtualTextureFeedbackDefaultBufferSize(
 	TEXT("r.vt.FeedbackDefaultBufferSize"),
@@ -109,6 +111,10 @@ void FVirtualTextureFeedbackBufferResource::End(FRDGBuilder& GraphBuilder)
 	}
 	bIsInBeginEndScope = false;
 
+	// VirtualTextureFeedback would be a more descriptive stat name, but VirtualTextureUpdate was used historically and some profile tools may depend on that.
+	RDG_EVENT_SCOPE_STAT(GraphBuilder, VirtualTextureUpdate, "VirtualTextureUpdate");
+	RDG_GPU_STAT_SCOPE(GraphBuilder, VirtualTextureUpdate);
+
 	AddPass(GraphBuilder, RDG_EVENT_NAME("VirtualTextureFeedbackTransition"), [UAV = UAV](FRHICommandListImmediate& RHICmdList)
 	{
 		RHICmdList.EndUAVOverlap(UAV);
