Fix nullptr deref when rendering runtime virtual texture thumbnail with virtual textures disabled.
#rnx

diff --git a/Engine/Source/Runtime/Renderer/Private/VT/RuntimeVirtualTextureRender.cpp b/Engine/Source/Runtime/Renderer/Private/VT/RuntimeVirtualTextureRender.cpp
--- a/Engine/Source/Runtime/Renderer/Private/VT/RuntimeVirtualTextureRender.cpp
+++ b/Engine/Source/Runtime/Renderer/Private/VT/RuntimeVirtualTextureRender.cpp
@@ -1504,7 +1504,10 @@ namespace RuntimeVirtualTexture
 		const float WorldToPixel = View->ViewRect.Width() * RcpWorldSize;
 
 		TArray<int32> PrimitiveIndices;
-		Scene->GetExtensionPtr<FRuntimeVirtualTextureSceneExtension>()->GetPrimitivesForRuntimeVirtualTexture(Scene, RuntimeVirtualTextureId, PrimitiveIndices);
+		if (FRuntimeVirtualTextureSceneExtension const* SceneExtension = Scene->GetExtensionPtr<FRuntimeVirtualTextureSceneExtension>())
+		{
+			SceneExtension->GetPrimitivesForRuntimeVirtualTexture(Scene, RuntimeVirtualTextureId, PrimitiveIndices);
+		}
 
 		for (const int32 PrimitiveIndex : PrimitiveIndices)
  		{
