Change default r.VT.AVT.AgeToFree from 0 (off) to 1000.
This prevents the AVT locking too many pages in their respective physical pools due to never releasing. That in turn can lead to unwanted oversubscription.
Add r.VT.AVT.PoolResidencyToFree cvar which prevents the AVT free from kicking in unless we reach a given physical pool residency.
We only want to pay the cost of free when necessary. Also delaying free to when residency is high makes it likely that we will have very few mapped pages left to remap. Remapping many pages is a large part of the cost.

diff --git a/Engine/Source/Runtime/Renderer/Private/VT/AdaptiveVirtualTexture.cpp b/Engine/Source/Runtime/Renderer/Private/VT/AdaptiveVirtualTexture.cpp
--- a/Engine/Source/Runtime/Renderer/Private/VT/AdaptiveVirtualTexture.cpp
+++ b/Engine/Source/Runtime/Renderer/Private/VT/AdaptiveVirtualTexture.cpp
@@ -26,14 +26,22 @@ static TAutoConsoleVariable<int32> CVarAVTMaxFreePerFrame(
 static TAutoConsoleVariable<int32> CVarAVTMaxPageResidency(
 	TEXT("r.VT.AVT.MaxPageResidency"),
 	75,
-	TEXT("Percentage of page table to allocate before we start freeing to make space"),
+	TEXT("Percentage of page table to allocate before we start force freeing to make space"),
 	ECVF_RenderThreadSafe
 );
 
 static TAutoConsoleVariable<int32> CVarAVTAgeToFree(
 	TEXT("r.VT.AVT.AgeToFree"),
-	0,
-	TEXT("Number of frames for an allocation to be unused before it is considered for free"),
+	1000,
+	TEXT("Number of frames for an allocation to be unused before it is considered for freeing by age"),
+	ECVF_RenderThreadSafe
+);
+
+static TAutoConsoleVariable<int32> CVarAVTMinPoolResidencyToFree(
+	TEXT("r.VT.AVT.PoolResidencyToFree"),
+	85,
+	TEXT("We consider freeing allocations by age if page residency percentage of physical pool is greater than this\n")
+	TEXT("Only freeing allocations at high residency reduces the likely number of mapped pages to remap, which can be expensive"),
 	ECVF_RenderThreadSafe
 );
 
@@ -665,8 +673,9 @@ bool FAdaptiveVirtualTexture::FreeLRU(FRHICommandListBase& RHICmdList, FVirtualT
 	else
 	{
 		// Reallocate to the selected level.
+		// Maintain the last used frame so that we can continue to deallocate levels.
 		const uint32 GridIndex = AllocationSlots[AllocationIndex].GridIndex;
-		Allocate(RHICmdList, InSystem, GridIndex, AllocationIndex, NewLevel, InFrame);
+		Allocate(RHICmdList, InSystem, GridIndex, AllocationIndex, NewLevel, LastFrameUsed);
 	}
 
 	return true;
@@ -674,18 +683,32 @@ bool FAdaptiveVirtualTexture::FreeLRU(FRHICommandListBase& RHICmdList, FVirtualT
 
 void FAdaptiveVirtualTexture::UpdateAllocations(FVirtualTextureSystem* InSystem, FRHICommandListImmediate& RHICmdList, uint32 InFrame)
 {
+	TRACE_CPUPROFILER_EVENT_SCOPE(FAdaptiveVirtualTexture::UpdateAllocations);
+
 	if (RequestsToMap.Num() == 0)
 	{
-		// Free old unused pages if there is no other work to do.
+		// Free old unused allocations if there is no other work to do.
 		const uint32 FrameAgeToFree = CVarAVTAgeToFree.GetValueOnRenderThread();
 		if (FrameAgeToFree > 0)
 		{
-			const int32 NumToFree = FMath::Min(NumAllocated, CVarAVTMaxFreePerFrame.GetValueOnRenderThread());
+			FVirtualTexturePhysicalSpace const* PhysicalSpace = AllocatedVirtualTextureLowMips->GetPhysicalSpaceForPageTableLayer(0);
+			FTexturePagePool const& Pool = PhysicalSpace->GetPagePool();
 
-			bool bFreeSuccess = true;
-			for (int32 FreeCount = 0; bFreeSuccess && FreeCount < NumToFree; FreeCount++)
+			const uint32 PageFreeThreshold = FMath::Max(VirtualTextureScalability::GetPageFreeThreshold(), 0u);
+			const uint32 FrameMinusThreshold = InFrame > PageFreeThreshold ? InFrame - PageFreeThreshold : 0;
+			const uint32 NumVisiblePages = Pool.GetNumVisiblePages(FrameMinusThreshold);
+			const uint32 NumPages = Pool.GetNumPages();
+			const uint32 MinPoolResidencyToFree = CVarAVTMinPoolResidencyToFree.GetValueOnRenderThread();
+
+			if (NumVisiblePages * 100 > MinPoolResidencyToFree * NumPages)
 			{
-				bFreeSuccess = FreeLRU(RHICmdList, InSystem, InFrame, FrameAgeToFree);
+				const int32 NumToFree = FMath::Min(NumAllocated, CVarAVTMaxFreePerFrame.GetValueOnRenderThread());
+
+				bool bFreeSuccess = true;
+				for (int32 FreeCount = 0; bFreeSuccess && FreeCount < NumToFree; FreeCount++)
+				{
+					bFreeSuccess = FreeLRU(RHICmdList, InSystem, InFrame, FrameAgeToFree);
+				}
 			}
 		}
 	}
