[SSS] Fix artifact on Mac. Write indirect argument buffer array in a new pass instead of using atomics.

#tests Engine test on windows and simple repro on Mac
#rnx

diff --git a/Engine/Shaders/Private/PostProcessSubsurfaceTile.usf b/Engine/Shaders/Private/PostProcessSubsurfaceTile.usf
--- a/Engine/Shaders/Private/PostProcessSubsurfaceTile.usf
+++ b/Engine/Shaders/Private/PostProcessSubsurfaceTile.usf
@@ -318,8 +318,9 @@ void SSSTileCategorisationMarkCS(uint3 DT_ID : SV_DispatchThreadID,
 	GroupMemoryBarrierWithGroupSync();
 	if (LinearIndex < 1)
 	{
-		bWriteTile.x = ContainsSSS[LinearIndex] > 0;
-		bWriteTile.y = ContainsSSS[LinearIndex + 1] ==  ContainsSSS[LinearIndex];
+		uint2 ContainsSSSResult = ContainsSSS[LinearIndex] + ContainsSSS[LinearIndex + 1];
+		bWriteTile.x = ContainsSSSResult.x > 0;
+		bWriteTile.y = ContainsSSSResult.x ==  ContainsSSSResult.y;
 	}
 #endif
 
@@ -346,14 +347,11 @@ void SSSTileCategorisationMarkCS(uint3 DT_ID : SV_DispatchThreadID,
 
 }
 
-int VertexCountPerInstanceIndirect;
 int TileType;
 StructuredBuffer<uint> TileMaskBuffer;
 
 RWBuffer<uint> RWSSSTileListData;
 RWBuffer<uint> RWTileTypeCountBuffer;
-RWBuffer<uint> RWIndirectDrawArgsBuffer;
-RWBuffer<uint> RWIndirectDispatchArgsBuffer;
 RWBuffer<uint> RWSSSTileListDataBuffer;
 
 groupshared uint SharedNumTiles;
@@ -370,21 +368,6 @@ void SSSTileClassificationBuildListsCS(
 	uint2 GroupThreadId : SV_GroupThreadID)
 {
 	const uint LinearThreadIndex = GroupThreadId.y * SUBSURFACE_TILE_SIZE + GroupThreadId.x;
-	
-	if (all(DispatchThreadId == 0))
-	{
-		const uint TileCount = 0; // Instance count will be accumulated later with locked add.
-		
-		// Indirect draw
-		RWIndirectDrawArgsBuffer[TileType * DRAW_INDIRECT_UINT_COUNT + 0] = VertexCountPerInstanceIndirect;		// VertexCountPerInstance
-		RWIndirectDrawArgsBuffer[TileType * DRAW_INDIRECT_UINT_COUNT + 1] = TileCount;  // InstanceCount
-		RWIndirectDrawArgsBuffer[TileType * DRAW_INDIRECT_UINT_COUNT + 2] = 0;  // StartVertexLocation
-		RWIndirectDrawArgsBuffer[TileType * DRAW_INDIRECT_UINT_COUNT + 3] = 0;  // StartInstanceLocation
-
-		// Indirect dispatch initialize other dimensions
-		const uint IndirectDispatchSize = TileCount;
-		WriteDispatchIndirectArgs(RWIndirectDispatchArgsBuffer, TileType, IndirectDispatchSize, 1, 1);
-	}
 
 	if (LinearThreadIndex == 0)
 	{
@@ -422,12 +405,7 @@ void SSSTileClassificationBuildListsCS(
 	// Allocate space in the tile list
 	if (LinearThreadIndex == 0 && SharedNumTiles > 0)
 	{
-		// Still need to track the total number of tiles ?.
-		uint IndexToStore = 0;
-		InterlockedAdd(RWTileTypeCountBuffer[TileType], SharedNumTiles, IndexToStore);
-
-		InterlockedAdd(RWIndirectDrawArgsBuffer[TileType * DRAW_INDIRECT_UINT_COUNT + 1], SharedNumTiles, SharedGlobalTileOffset);
-		InterlockedAdd(RWIndirectDispatchArgsBuffer[TileType * DISPATCH_INDIRECT_UINT_COUNT + 0], SharedNumTiles);
+		InterlockedAdd(RWTileTypeCountBuffer[TileType], SharedNumTiles, SharedGlobalTileOffset);
 	}
 
 	GroupMemoryBarrierWithGroupSync();
@@ -439,4 +417,32 @@ void SSSTileClassificationBuildListsCS(
 	}
 }
 
+int VertexCountPerInstanceIndirect;
+Buffer<uint> TileTypeCountBuffer;
+RWBuffer<uint> RWIndirectDrawArgsBuffer;
+RWBuffer<uint> RWIndirectDispatchArgsBuffer;
+
+[numthreads(8, 1, 1)]
+void SubsurfaceTileBuildIndirectDispatchArgsCS(uint2 DispatchThreadId : SV_DispatchThreadID)
+{
+	if (DispatchThreadId.y == 0 && DispatchThreadId.x < SUBSURFACE_TILE_TYPE_COUNT 
+	    && (DispatchThreadId.x == SUBSURFACE_TILE_TYPE_PASSTHROUGH
+		|| DispatchThreadId.x == SUBSURFACE_TILE_TYPE_ALLNONPASSTHROUGH
+		|| DispatchThreadId.x == SUBSURFACE_TILE_TYPE_ALL))
+	{
+		uint TileType = DispatchThreadId.x;
+		const uint TileCount = TileTypeCountBuffer[TileType];
+
+		// Indirect draw
+		RWIndirectDrawArgsBuffer[TileType * DRAW_INDIRECT_UINT_COUNT + 0] = VertexCountPerInstanceIndirect;  // VertexCountPerInstance
+		RWIndirectDrawArgsBuffer[TileType * DRAW_INDIRECT_UINT_COUNT + 1] = TileCount;  // InstanceCount
+		RWIndirectDrawArgsBuffer[TileType * DRAW_INDIRECT_UINT_COUNT + 2] = 0;  // StartVertexLocation
+		RWIndirectDrawArgsBuffer[TileType * DRAW_INDIRECT_UINT_COUNT + 3] = 0;  // StartInstanceLocation
+
+		// Indirect dispatch
+		const uint IndirectDispatchSize = TileCount;
+		WriteDispatchIndirectArgs(RWIndirectDispatchArgsBuffer, TileType, IndirectDispatchSize, 1, 1);
+	}
+}
+
 #endif
\ No newline at end of file
diff --git a/Engine/Source/Runtime/Renderer/Private/SubsurfaceTiles.cpp b/Engine/Source/Runtime/Renderer/Private/SubsurfaceTiles.cpp
--- a/Engine/Source/Runtime/Renderer/Private/SubsurfaceTiles.cpp
+++ b/Engine/Source/Runtime/Renderer/Private/SubsurfaceTiles.cpp
@@ -35,7 +35,6 @@ FSubsurfaceTilePassVS::FParameters GetSubsurfaceTileParameters(const FScreenPass
 	return Out;
 }
 
-
 class FClearUAVBuildIndirectDispatchBufferCS : public FGlobalShader
 {
 public:
@@ -194,12 +193,9 @@ class FSSSTileClassificationBuildListsCS : public FGlobalShader
 
 	BEGIN_SHADER_PARAMETER_STRUCT(FParameters, )
 		SHADER_PARAMETER_STRUCT_INCLUDE(FViewShaderParameters, View)
-		SHADER_PARAMETER(uint32, VertexCountPerInstanceIndirect)
 		SHADER_PARAMETER(FIntPoint, TiledViewRes)
 		SHADER_PARAMETER(int, TileType)
 		SHADER_PARAMETER_RDG_BUFFER_SRV(StructuredBuffer<uint>, TileMaskBuffer)
-		SHADER_PARAMETER_RDG_BUFFER_UAV(RWBuffer<uint>, RWIndirectDrawArgsBuffer)
-		SHADER_PARAMETER_RDG_BUFFER_UAV(RWBuffer<uint>, RWIndirectDispatchArgsBuffer)
 		SHADER_PARAMETER_RDG_BUFFER_UAV(RWBuffer<uint>, RWTileTypeCountBuffer)
 		SHADER_PARAMETER_RDG_BUFFER_UAV(RWBuffer<uint>, RWSSSTileListDataBuffer)
 	END_SHADER_PARAMETER_STRUCT()
@@ -225,6 +221,36 @@ class FSSSTileClassificationBuildListsCS : public FGlobalShader
 
 IMPLEMENT_GLOBAL_SHADER(FSSSTileClassificationBuildListsCS, "/Engine/Private/PostProcessSubsurfaceTile.usf", "SSSTileClassificationBuildListsCS", SF_Compute);
 
+class FSubsurfaceTileBuildIndirectDispatchArgsCS : public FGlobalShader
+{
+public:
+	DECLARE_GLOBAL_SHADER(FSubsurfaceTileBuildIndirectDispatchArgsCS);
+	SHADER_USE_PARAMETER_STRUCT(FSubsurfaceTileBuildIndirectDispatchArgsCS, FGlobalShader)
+
+	BEGIN_SHADER_PARAMETER_STRUCT(FParameters, )
+		SHADER_PARAMETER(uint32, VertexCountPerInstanceIndirect)
+		SHADER_PARAMETER_RDG_BUFFER_UAV(RWBuffer<uint>, RWIndirectDispatchArgsBuffer)
+		SHADER_PARAMETER_RDG_BUFFER_UAV(RWBuffer<uint>, RWIndirectDrawArgsBuffer)
+		SHADER_PARAMETER_RDG_BUFFER_SRV(Buffer<uint>, TileTypeCountBuffer)
+	END_SHADER_PARAMETER_STRUCT()
+
+	static void ModifyCompilationEnvironment(const FGlobalShaderPermutationParameters& Parameters, FShaderCompilerEnvironment& OutEnvironment)
+	{
+		FGlobalShader::ModifyCompilationEnvironment(Parameters, OutEnvironment);
+		OutEnvironment.SetDefine(TEXT("TILE_CATERGORISATION_SHADER"), 1);
+		OutEnvironment.SetDefine(TEXT("SUBSURFACE_TILE_SIZE"),  FSubsurfaceTiles::TileSize);
+		OutEnvironment.SetDefine(TEXT("SUBSURFACE_KERNEL_SIZE"), SUBSURFACE_KERNEL_SIZE);
+
+	}
+
+	static bool ShouldCompilePermutation(const FGlobalShaderPermutationParameters& Parameters)
+	{
+		return IsFeatureLevelSupported(Parameters.Platform, ERHIFeatureLevel::SM5);
+	}
+};
+
+IMPLEMENT_GLOBAL_SHADER(FSubsurfaceTileBuildIndirectDispatchArgsCS, "/Engine/Private/PostprocessSubsurfaceTile.usf", "SubsurfaceTileBuildIndirectDispatchArgsCS", SF_Compute);
+
 // Returns the [0, N] clamped value of the 'r.SSS.Scale' CVar.
 static float GetSubsurfaceRadiusScaleForTiling()
 {
@@ -271,6 +297,7 @@ FSubsurfaceTiles ClassifySSSTiles(FRDGBuilder& GraphBuilder, const FViewInfo& Vi
 		Result.TileCount = Result.TileDimension.X * Result.TileDimension.Y;
 		Result.bRectPrimitive = GRHISupportsRectTopology;
 		Result.TileTypeCountBuffer = GraphBuilder.CreateBuffer(FRDGBufferDesc::CreateBufferDesc(sizeof(uint32), FSubsurfaceTiles::TileTypeCount), TEXT("Subsurface.TileCountBuffer"));
+		Result.TileTypeCountSRV = GraphBuilder.CreateSRV(Result.TileTypeCountBuffer,EPixelFormat::PF_R32_UINT);
 		Result.TileIndirectDispatchBuffer = 
 			GraphBuilder.CreateBuffer(FRDGBufferDesc::CreateIndirectDesc<FRHIDispatchIndirectParameters>(FSubsurfaceTiles::TileTypeCount), TEXT("Subsurface.TileIndirectDispatchBuffer"));
 		Result.TileIndirectDrawBuffer =
@@ -381,13 +408,8 @@ FSubsurfaceTiles ClassifySSSTiles(FRDGBuilder& GraphBuilder, const FViewInfo& Vi
 
 				PassParameters->View = View.GetShaderParameters();
 				PassParameters->TiledViewRes = Result.TileDimension;
-				PassParameters->VertexCountPerInstanceIndirect = GRHISupportsRectTopology ? 4 : 6;
-				PassParameters->TileMaskBuffer = GraphBuilder.CreateSRV(PassInfo.Input);
-
 				PassParameters->TileType = ToIndex(PassInfo.TileType);
-				PassParameters->RWIndirectDrawArgsBuffer = DrawIndirectParametersBufferUAV;
-				PassParameters->RWIndirectDispatchArgsBuffer = DispatchIndirectParametersBufferUAV;
-				
+				PassParameters->TileMaskBuffer = GraphBuilder.CreateSRV(PassInfo.Input);
 				PassParameters->RWTileTypeCountBuffer = GraphBuilder.CreateUAV(Result.TileTypeCountBuffer, PF_R32_UINT);
 				PassParameters->RWSSSTileListDataBuffer = GraphBuilder.CreateUAV(PassInfo.Output, PF_R32_UINT);
 
@@ -400,6 +422,18 @@ FSubsurfaceTiles ClassifySSSTiles(FRDGBuilder& GraphBuilder, const FViewInfo& Vi
 					FComputeShaderUtils::GetGroupCount(Result.TileDimension, SHADER::GetGroupSize())
 				);
 			}
+
+			// Setup the indirect dispatch & draw arguments for All, PASSTHROUGH, and AllNonPassthrough.
+			{
+				typedef FSubsurfaceTileBuildIndirectDispatchArgsCS ARGSETUPSHADER;
+				ARGSETUPSHADER::FParameters* PassParameters = GraphBuilder.AllocParameters<ARGSETUPSHADER::FParameters>();
+				PassParameters->VertexCountPerInstanceIndirect = GRHISupportsRectTopology ? 4 : 6;
+				PassParameters->RWIndirectDispatchArgsBuffer = DispatchIndirectParametersBufferUAV;
+				PassParameters->RWIndirectDrawArgsBuffer = DrawIndirectParametersBufferUAV;
+				PassParameters->TileTypeCountBuffer = Result.TileTypeCountSRV;
+				TShaderMapRef<ARGSETUPSHADER> ComputeShader(View.ShaderMap);
+				FComputeShaderUtils::AddPass(GraphBuilder, FRDGEventName(TEXT("SSS::BuildIndirectArgs(Dispatch & Draw)")), ComputeShader, PassParameters, FIntVector(1, 1, 1));
+			}
 		}
 	}
 	return Result;
