Fix potential nullptr deref during ODSC usage.
#rnx

diff --git a/Engine/Source/Runtime/Renderer/Private/PrimitiveSceneInfo.cpp b/Engine/Source/Runtime/Renderer/Private/PrimitiveSceneInfo.cpp
--- a/Engine/Source/Runtime/Renderer/Private/PrimitiveSceneInfo.cpp
+++ b/Engine/Source/Runtime/Renderer/Private/PrimitiveSceneInfo.cpp
@@ -2485,9 +2485,14 @@ FString FPrimitiveSceneInfo::GetOwnerActorNameOrLabelForDebuggingOnly() const
 	return FString(TEXT("Unknown Object"));
 }
 
-UPrimitiveComponent* FPrimitiveSceneInfo::GetComponentForDebugOnly() const 
-{ 
-	return Cast<UPrimitiveComponent>(PrimitiveComponentInterfaceForDebuggingOnly->GetUObject()); 
+UPrimitiveComponent* FPrimitiveSceneInfo::GetComponentForDebugOnly() const
+{
+	if (PrimitiveComponentInterfaceForDebuggingOnly)
+	{
+		return Cast<UPrimitiveComponent>(PrimitiveComponentInterfaceForDebuggingOnly->GetUObject());
+	}
+
+	return nullptr;
 }
 
 IPrimitiveComponent* FPrimitiveSceneInfo::GetComponentInterfaceForDebugOnly() const 
