More robust handling of case where we place duplicate RVT components in a scene.
Only one component will "own" the RVT asset, and that is already well handled.
But both components will keep an associated scene proxy, and currently one will contain orphaned pointers to a deleted resource.
This can cause a crash when calling RequestPreload() on the scene proxy.
The fix is to detect the case and call the new MarkUnused() function on the duplicate scene proxy to clear orphaned pointers.
We also log a warning about the bad data setup that causes this.
#rnx

diff --git a/Engine/Source/Runtime/Renderer/Private/RendererScene.cpp b/Engine/Source/Runtime/Renderer/Private/RendererScene.cpp
--- a/Engine/Source/Runtime/Renderer/Private/RendererScene.cpp
+++ b/Engine/Source/Runtime/Renderer/Private/RendererScene.cpp
@@ -3089,6 +3089,20 @@ static void UpdateRuntimeVirtualTextureHidePrimitiveMasks(TSparseArray<FRuntimeV
 	}
 }
 
+static void FixMultipleRuntimeVirtualTextureUsage(TSparseArray<FRuntimeVirtualTextureSceneProxy*> const& InProxies, FRuntimeVirtualTextureSceneProxy const* SceneProxy)
+{
+	// If we add a runtime virtual texture component to a scene which uses the same asset as another, then we need to reset the first one.
+	// This should only happen for bad data setups, so warn here.
+	for (TSparseArray<FRuntimeVirtualTextureSceneProxy*>::TConstIterator It(InProxies); It; ++It)
+	{
+		if (*It != SceneProxy && (*It)->RuntimeVirtualTextureId == SceneProxy->RuntimeVirtualTextureId)
+		{
+			UE_LOG(LogRenderer, Warning, TEXT("Found two RVT components in a scene using the same RVT asset."));
+			(*It)->MarkUnused();
+		}
+	}
+}
+
 void FScene::AddRuntimeVirtualTexture(URuntimeVirtualTextureComponent* Component)
 {
 	if (Component->SceneProxy == nullptr)
@@ -3103,6 +3117,8 @@ void FScene::AddRuntimeVirtualTexture(URuntimeVirtualTextureComponent* Component
 		{
 			SceneProxy->SceneIndex = Scene->RuntimeVirtualTextures.Add(SceneProxy);
 			
+			FixMultipleRuntimeVirtualTextureUsage(Scene->RuntimeVirtualTextures, SceneProxy);
+
 			UpdateRuntimeVirtualTextureHidePrimitiveMasks(Scene->RuntimeVirtualTextures, Scene->bRuntimeVirtualTexturePrimitiveHideEditor, Scene->bRuntimeVirtualTexturePrimitiveHideGame);
 		});
 	}
diff --git a/Engine/Source/Runtime/Renderer/Private/VT/RuntimeVirtualTextureSceneProxy.cpp b/Engine/Source/Runtime/Renderer/Private/VT/RuntimeVirtualTextureSceneProxy.cpp
--- a/Engine/Source/Runtime/Renderer/Private/VT/RuntimeVirtualTextureSceneProxy.cpp
+++ b/Engine/Source/Runtime/Renderer/Private/VT/RuntimeVirtualTextureSceneProxy.cpp
@@ -157,6 +157,14 @@ void FRuntimeVirtualTextureSceneProxy::Release()
 	}
 }
 
+void FRuntimeVirtualTextureSceneProxy::MarkUnused()
+{
+	VirtualTexture = nullptr;
+	ProducerHandle = {};
+	AllocatedVirtualTexture = nullptr;
+	SpaceID = -1;
+}
+
 /** Transform world bounds into Virtual Texture UV space. */
 static FBox2D GetUVRectFromWorldBounds(FTransform const& InTransform, FBoxSphereBounds const& InBounds)
 {
@@ -174,63 +182,63 @@ static FBox2D GetUVRectFromWorldBounds(FTransform const& InTransform, FBoxSphere
 
 void FRuntimeVirtualTextureSceneProxy::Dirty(FBoxSphereBounds const& InBounds, EVTInvalidatePriority InInvalidatePriority)
 {
-	const FBox2D UVRect = GetUVRectFromWorldBounds(Transform, InBounds);
-
-	// Convert to Texel coordinate space
-	const FIntRect TextureRect(0, 0, VirtualTextureSize.X, VirtualTextureSize.Y);
-	FIntRect TexelRect(
-		FMath::FloorToInt(UVRect.Min.X * VirtualTextureSize.X),
-		FMath::FloorToInt(UVRect.Min.Y * VirtualTextureSize.Y),
-		FMath::CeilToInt(UVRect.Max.X * VirtualTextureSize.X),
-		FMath::CeilToInt(UVRect.Max.Y * VirtualTextureSize.Y));
-	TexelRect.Clip(TextureRect);
-
-	// Only add rect if it has some area
-	if (TexelRect.Min != TexelRect.Max)
+	// If Producer handle is not initialized yet it's safe to do nothing because we won't have rendered anything to the VT that needs flushing.
+	if (ProducerHandle.PackedValue != 0)
 	{
-		FDirtyRect DirtyRect{ .Rect = TexelRect, .InvalidatePriority = InInvalidatePriority };
-		const bool bFirst = DirtyRects.Add(DirtyRect) == 0;
-		if (bFirst)
-		{
-			CombinedDirtyRect = DirtyRect;
-		}
-		else
+		const FBox2D UVRect = GetUVRectFromWorldBounds(Transform, InBounds);
+
+		// Convert to Texel coordinate space
+		const FIntRect TextureRect(0, 0, VirtualTextureSize.X, VirtualTextureSize.Y);
+		FIntRect TexelRect(
+			FMath::FloorToInt(UVRect.Min.X * VirtualTextureSize.X),
+			FMath::FloorToInt(UVRect.Min.Y * VirtualTextureSize.Y),
+			FMath::CeilToInt(UVRect.Max.X * VirtualTextureSize.X),
+			FMath::CeilToInt(UVRect.Max.Y * VirtualTextureSize.Y));
+		TexelRect.Clip(TextureRect);
+
+		// Only add rect if it has some area
+		if (TexelRect.Min != TexelRect.Max)
 		{
-			CombinedDirtyRect.Union(DirtyRect);
+			FDirtyRect DirtyRect{ .Rect = TexelRect, .InvalidatePriority = InInvalidatePriority };
+			const bool bFirst = DirtyRects.Add(DirtyRect) == 0;
+			if (bFirst)
+			{
+				CombinedDirtyRect = DirtyRect;
+			}
+			else
+			{
+				CombinedDirtyRect.Union(DirtyRect);
+			}
 		}
 	}
 }
 
 void FRuntimeVirtualTextureSceneProxy::FlushDirtyPages()
 {
-	// If Producer handle is not initialized yet it's safe to do nothing because we won't have rendered anything to the VT that needs flushing.
-	if (ProducerHandle.PackedValue != 0)
+	// Don't do any work if we won't mark anything dirty.
+	if (MaxDirtyLevel >= 0 && CombinedDirtyRect.Rect.Width() != 0 && CombinedDirtyRect.Rect.Height() != 0)
 	{
-		// Don't do any work if we won't mark anything dirty.
-		if (MaxDirtyLevel >= 0 && CombinedDirtyRect.Rect.Width() != 0 && CombinedDirtyRect.Rect.Height() != 0)
+		// Keeping visible pages mapped reduces update flicker due to the latency in the unmap/feedback/map sequence.
+		// But it potentially creates more page update work since more pages may get updated.
+		const uint32 MaxAgeToKeepMapped = VirtualTextureScalability::GetKeepDirtyPageMappedFrameThreshold();
+
+		//todo[vt]: 
+		// Profile to work out best heuristic for when we should use the CombinedDirtyRect
+		// Also consider using some other structure to represent dirty area such as a course 2D bitfield
+		bool bCombinedFlush = (DirtyRects.Num() > 2 || CombinedDirtyRect.Rect == FIntRect(0, 0, VirtualTextureSize.X, VirtualTextureSize.Y));
+		// Don't use the combined rect if one of the dirty rects is prioritized, because that would would leave all of the pages being covered by the combined rect to get prioritized, 
+		//  which would give exactly the opposite result of what we're trying to achieve, since we need to keep the number of prioritized pages to remain low :  
+		bCombinedFlush &= (CombinedDirtyRect.InvalidatePriority == EVTInvalidatePriority::Normal);
+
+		if (bCombinedFlush)
+		{
+			FVirtualTextureSystem::Get().FlushCache(ProducerHandle, SpaceID, CombinedDirtyRect.Rect, MaxDirtyLevel, MaxAgeToKeepMapped, EVTInvalidatePriority::Normal);
+		}
+		else
 		{
-			// Keeping visible pages mapped reduces update flicker due to the latency in the unmap/feedback/map sequence.
-			// But it potentially creates more page update work since more pages may get updated.
-			const uint32 MaxAgeToKeepMapped = VirtualTextureScalability::GetKeepDirtyPageMappedFrameThreshold();
-
-			//todo[vt]: 
-			// Profile to work out best heuristic for when we should use the CombinedDirtyRect
-			// Also consider using some other structure to represent dirty area such as a course 2D bitfield
-			bool bCombinedFlush = (DirtyRects.Num() > 2 || CombinedDirtyRect.Rect == FIntRect(0, 0, VirtualTextureSize.X, VirtualTextureSize.Y));
-			// Don't use the combined rect if one of the dirty rects is prioritized, because that would would leave all of the pages being covered by the combined rect to get prioritized, 
-			//  which would give exactly the opposite result of what we're trying to achieve, since we need to keep the number of prioritized pages to remain low :  
-			bCombinedFlush &= (CombinedDirtyRect.InvalidatePriority == EVTInvalidatePriority::Normal);
-
-			if (bCombinedFlush)
+			for (const FDirtyRect& DirtyRect : DirtyRects)
 			{
-				FVirtualTextureSystem::Get().FlushCache(ProducerHandle, SpaceID, CombinedDirtyRect.Rect, MaxDirtyLevel, MaxAgeToKeepMapped, EVTInvalidatePriority::Normal);
-			}
-			else
-			{
-				for (const FDirtyRect& DirtyRect : DirtyRects)
-				{
-					FVirtualTextureSystem::Get().FlushCache(ProducerHandle, SpaceID, DirtyRect.Rect, MaxDirtyLevel, MaxAgeToKeepMapped, DirtyRect.InvalidatePriority);
-				}
+				FVirtualTextureSystem::Get().FlushCache(ProducerHandle, SpaceID, DirtyRect.Rect, MaxDirtyLevel, MaxAgeToKeepMapped, DirtyRect.InvalidatePriority);
 			}
 		}
 	}
@@ -241,9 +249,12 @@ void FRuntimeVirtualTextureSceneProxy::FlushDirtyPages()
 
 void FRuntimeVirtualTextureSceneProxy::RequestPreload(FBoxSphereBounds const& InBounds, int32 InLevel)
 {
-	const FBox2D UVRect = GetUVRectFromWorldBounds(Transform, InBounds);
-
-	FVirtualTextureSystem::Get().RequestTiles(AllocatedVirtualTexture, FVector2D::One(), FVector2D::Zero(), FVector2D::One(), UVRect.Min, UVRect.Max, InLevel);
+	// If Producer handle is not initialized yet it's safe to do nothing.
+	if (ProducerHandle.PackedValue != 0)
+	{
+		const FBox2D UVRect = GetUVRectFromWorldBounds(Transform, InBounds);
+		FVirtualTextureSystem::Get().RequestTiles(AllocatedVirtualTexture, FVector2D::One(), FVector2D::Zero(), FVector2D::One(), UVRect.Min, UVRect.Max, InLevel);
+	}
 }
 
 #undef LOCTEXT_NAMESPACE
diff --git a/Engine/Source/Runtime/Renderer/Private/VT/RuntimeVirtualTextureSceneProxy.h b/Engine/Source/Runtime/Renderer/Private/VT/RuntimeVirtualTextureSceneProxy.h
--- a/Engine/Source/Runtime/Renderer/Private/VT/RuntimeVirtualTextureSceneProxy.h
+++ b/Engine/Source/Runtime/Renderer/Private/VT/RuntimeVirtualTextureSceneProxy.h
@@ -23,6 +23,9 @@ public:
 	 */
 	void Release();
 
+	/** Mark the object as no longer used. This can happen in the edge case where two components in a scene use the same runtime virtual texture asset. */
+	void MarkUnused();
+
 	/**
 	 * Mark an area of the associated runtime virtual texture as dirty. 
 	 * @param Bounds World-space bounds of the area of the runtime virtual texture to invalidate
