[Hair] Fix material with UseHairStrands flag to issue a non-necessary draw call when used on hair strands geometry.

diff --git a/Engine/Source/Runtime/Renderer/Private/SceneVisibility.cpp b/Engine/Source/Runtime/Renderer/Private/SceneVisibility.cpp
--- a/Engine/Source/Runtime/Renderer/Private/SceneVisibility.cpp
+++ b/Engine/Source/Runtime/Renderer/Private/SceneVisibility.cpp
@@ -1935,7 +1935,8 @@ void FRelevancePacket::ComputeRelevance(FDynamicPrimitiveIndexList& DynamicPrimi
 				VisibleDynamicPrimitivesWithSimpleLights.AddPrim(PrimitiveSceneInfo);
 			}
 		}
-		else if (bHairStrandsRelevance)
+		// Filter with !bStaticRelevance to avoid regular/non-strands geometry to be added two times.
+		else if (!bStaticRelevance && bHairStrandsRelevance)
 		{
 			AddDynamicPrimitive(BitIndex);
 		}
