[shaders] adjust FPathTracingRG jobs to extra high priority so they are executed individually instead of in a batch. this on average reduces the wall time of global shader compilation on editor startup

diff --git a/Engine/Source/Runtime/Renderer/Private/PathTracing.cpp b/Engine/Source/Runtime/Renderer/Private/PathTracing.cpp
--- a/Engine/Source/Runtime/Renderer/Private/PathTracing.cpp
+++ b/Engine/Source/Runtime/Renderer/Private/PathTracing.cpp
@@ -1298,6 +1298,12 @@ class FPathTracingRG : public FGlobalShader
 		return RayTracing::GetShaderBindingLayout(Parameters.Platform);
 	}
 
+	static EShaderCompileJobPriority GetOverrideJobPriority()
+	{
+		// FPathTracingRG takes up to 20s on average on D3D SM6, and 30s on Vulkan SM6
+		return EShaderCompileJobPriority::ExtraHigh;
+	}
+
 	BEGIN_SHADER_PARAMETER_STRUCT(FParameters, )
 		SHADER_PARAMETER_RDG_TEXTURE_UAV(RWTexture2D<float4>, RadianceTexture)
 		SHADER_PARAMETER_RDG_TEXTURE_UAV(RWTexture2D<float2>, VarianceTexture)
