Fix for transmittance texture transition issue.

diff --git a/Engine/Source/Runtime/Renderer/Private/SkyAtmosphereRendering.cpp b/Engine/Source/Runtime/Renderer/Private/SkyAtmosphereRendering.cpp
--- a/Engine/Source/Runtime/Renderer/Private/SkyAtmosphereRendering.cpp
+++ b/Engine/Source/Runtime/Renderer/Private/SkyAtmosphereRendering.cpp
@@ -1777,7 +1777,7 @@ void FSkyAtmospherePendingRDGResources::CommitToSceneAndViewUniformBuffers(FRDGB
 	}
 
 	check(TransmittanceLut);
-	SkyInfo.GetTransmittanceLutTexture() = ConvertToExternalAccessTexture(GraphBuilder, ExternalAccessQueue, TransmittanceLut);
+	SkyInfo.GetTransmittanceLutTexture() = ConvertToExternalAccessTexture(GraphBuilder, ExternalAccessQueue, TransmittanceLut, ERHIAccess::SRVMask, ERHIPipeline::All);
 }
 
 FSkyAtmosphereRenderContext::FSkyAtmosphereRenderContext()
diff --git a/Engine/Source/Runtime/Renderer/Private/VolumetricCloudRendering.cpp b/Engine/Source/Runtime/Renderer/Private/VolumetricCloudRendering.cpp
--- a/Engine/Source/Runtime/Renderer/Private/VolumetricCloudRendering.cpp
+++ b/Engine/Source/Runtime/Renderer/Private/VolumetricCloudRendering.cpp
@@ -2927,7 +2927,7 @@ bool FSceneRenderer::RenderVolumetricCloud(
 					// Disable blending if volumetric render target is used, or if height fog is present (meaning we are going to render in a separate render target before blending over the scene)
 					SkyRC.bDisableBlending = (bShouldViewRenderVolumetricCloudRenderTarget || bHeightFogInScene) ? true : false;
 
-					SkyRC.TransmittanceLut = GraphBuilder.RegisterExternalTexture(SkyInfo.GetTransmittanceLutTexture());
+					SkyRC.TransmittanceLut = SkyInfo.GetTransmittanceLutTexture(GraphBuilder);
 					SkyRC.MultiScatteredLuminanceLut = GraphBuilder.RegisterExternalTexture(SkyInfo.GetMultiScatteredLuminanceLutTexture());
 
 					// Select the AerialPersepctiveOnCloud mode and set required parameters.
