[Lit Translucency] Fix ensure when switching preview platforms.

#rnx

diff --git a/Engine/Source/Runtime/Renderer/Private/TranslucentLighting.cpp b/Engine/Source/Runtime/Renderer/Private/TranslucentLighting.cpp
--- a/Engine/Source/Runtime/Renderer/Private/TranslucentLighting.cpp
+++ b/Engine/Source/Runtime/Renderer/Private/TranslucentLighting.cpp
@@ -198,7 +198,7 @@ static const TCHAR* TranslucentLightingMaterialPSOCollectorName = TEXT("Transluc
 
 float GetTranslucencyLightingVolumePositionOffsetRadius()
 {
-	return FMath::Max(0.0f, CVarTranslucencyLightingVolumePositionOffsetRadius.GetValueOnRenderThread());
+	return FMath::Max(0.0f, CVarTranslucencyLightingVolumePositionOffsetRadius.GetValueOnAnyThread());
 }
 
 /** Function returning current translucency lighting volume dimensions. */
