Setting up good default for 60hz games and real time sky light capture.

#rnx
#tests PC editor

diff --git a/Engine/Source/Runtime/Renderer/Private/ReflectionEnvironmentRealTimeCapture.cpp b/Engine/Source/Runtime/Renderer/Private/ReflectionEnvironmentRealTimeCapture.cpp
--- a/Engine/Source/Runtime/Renderer/Private/ReflectionEnvironmentRealTimeCapture.cpp
+++ b/Engine/Source/Runtime/Renderer/Private/ReflectionEnvironmentRealTimeCapture.cpp
@@ -40,7 +40,7 @@ static TAutoConsoleVariable<int32> CVarRealTimeReflectionCaptureTimeSlicing(
 	ECVF_RenderThreadSafe | ECVF_Scalability);
 
 static TAutoConsoleVariable<int32> CVarRealTimeReflectionCaptureTimeSlicingSkyCloudCubeFacePerFrame(
-	TEXT("r.SkyLight.RealTimeReflectionCapture.TimeSlice.SkyCloudCubeFacePerFrame"), 6,
+	TEXT("r.SkyLight.RealTimeReflectionCapture.TimeSlice.SkyCloudCubeFacePerFrame"), 2, // Default preferred for 60hz. 30hz applications can use 6 to not have any different sky lighting in different faces when the sun moves fast.
 	TEXT("When enabled, the real-time sky light capture, when time sliced, will not render cloud in all cube face in a single frame; but one face per frame. That is to distribute the cloud tracing cost even more, but will add latency and potentially can result in lighting discrepancy between faces if the sun is moving fast. Value in [1,6]."),
 	ECVF_RenderThreadSafe | ECVF_Scalability);
 
@@ -60,7 +60,7 @@ static TAutoConsoleVariable<int32> CVarRealTimeReflectionCaptureAlwaysClearColor
 	ECVF_RenderThreadSafe);
 
 static TAutoConsoleVariable<int32> CVarRealTimeReflectionCaptureVolumetricCloudResolutionDivider(
-	TEXT("r.SkyLight.RealTimeReflectionCapture.VolumetricCloudResolutionDivider"), 1,
+	TEXT("r.SkyLight.RealTimeReflectionCapture.VolumetricCloudResolutionDivider"), 2, // Default preferred for 60hz.
 	TEXT("The divider applied on the resolution when capturing cloud in the real time sky light. For instance, a value of 2 will render cloud at half resolution and thus will be faster to render."),
 	ECVF_RenderThreadSafe | ECVF_Scalability);
 
