[VSM] Page table and flag allocation optimization.
- Move invalidation flags to (Prev)PageFlags instead of page request flags & make the latter transient and uint8.
- Make receiver mask transient as these are not used for invalidation.

#rnx
#tests Only affects peak memory and some performance when receiver masks are enabled.

diff --git a/Engine/Shaders/Private/VirtualShadowMaps/VirtualShadowMapCacheInvalidation.ush b/Engine/Shaders/Private/VirtualShadowMaps/VirtualShadowMapCacheInvalidation.ush
--- a/Engine/Shaders/Private/VirtualShadowMaps/VirtualShadowMapCacheInvalidation.ush
+++ b/Engine/Shaders/Private/VirtualShadowMaps/VirtualShadowMapCacheInvalidation.ush
@@ -15,9 +15,9 @@ VirtualShadowMapCacheInvalidation.ush:
 
 RWStructuredBuffer<FPhysicalPageMetaData> PhysicalPageMetaDataOut;
 
-// We write the invalidation flags to the page request flags buffer to be consumed by the next
+// We write the invalidation flags to the prev page flags buffer to be consumed by the next
 // physical page update pass.
-RWTexture2D<uint> OutPageRequestFlags;
+RWTexture2D<uint> OutPrevPageFlags;
 
 bool IsValidInstanceAndCastShadows(FInstanceSceneData InstanceSceneData)
 {
@@ -137,7 +137,7 @@ void InvalidateInstancePages(FVirtualShadowMapProjectionShaderData ProjectionDat
 								// We don't look up the page table here because pages that weren't requested this frame
 								// are not in the page table, but still need to be invalidated. Thus we mark a bit in
 								// the page request flags instead which will get picked up during the next page update pass.
-								InterlockedOr(OutPageRequestFlags[PageFlagOffset.GetResourceAddress()], InvalidationFlags);
+								InterlockedOr(OutPrevPageFlags[PageFlagOffset.GetResourceAddress()], InvalidationFlags);
 							}
 						}
 					}
diff --git a/Engine/Shaders/Private/VirtualShadowMaps/VirtualShadowMapPhysicalPageManagement.usf b/Engine/Shaders/Private/VirtualShadowMaps/VirtualShadowMapPhysicalPageManagement.usf
--- a/Engine/Shaders/Private/VirtualShadowMaps/VirtualShadowMapPhysicalPageManagement.usf
+++ b/Engine/Shaders/Private/VirtualShadowMaps/VirtualShadowMapPhysicalPageManagement.usf
@@ -149,7 +149,7 @@ StructuredBuffer<FNextVirtualShadowMapData> NextVirtualShadowMapData;
 uint NextVirtualShadowMapDataCount;
 
 // To propogate any invalidation flags to the physical page flags
-Texture2D<uint> PrevPageRequestFlags;
+Texture2D<uint> PrevPageFlags;
 
 [numthreads(VSM_DEFAULT_CS_GROUP_X, 1, 1)]
 void UpdatePhysicalPageAddresses(uint3 Index : SV_DispatchThreadID)
@@ -203,7 +203,7 @@ void UpdatePhysicalPageAddresses(uint3 Index : SV_DispatchThreadID)
 	{
 #if HAS_CACHE_DATA
 		// Propogate any invalidation flags from the previous page requests to the physical page
-		const uint PrevFlags = PrevPageRequestFlags[PrevGlobalPageOffset.GetResourceAddress()];
+		const uint PrevFlags = PrevPageFlags[PrevGlobalPageOffset.GetResourceAddress()];
 		const uint InvalidationFlags = PrevFlags & VSM_EXTENDED_FLAG_ANY_INVALIDATED;
 		if (InvalidationFlags != 0)
 		{
diff --git a/Engine/Source/Runtime/Renderer/Private/VirtualShadowMaps/VirtualShadowMapArray.cpp b/Engine/Source/Runtime/Renderer/Private/VirtualShadowMaps/VirtualShadowMapArray.cpp
--- a/Engine/Source/Runtime/Renderer/Private/VirtualShadowMaps/VirtualShadowMapArray.cpp
+++ b/Engine/Source/Runtime/Renderer/Private/VirtualShadowMaps/VirtualShadowMapArray.cpp
@@ -1043,7 +1043,7 @@ class FUpdatePhysicalPageAddresses : public FVirtualShadowMapPageManagementShade
 		SHADER_PARAMETER_RDG_BUFFER_SRV(StructuredBuffer<FNextVirtualShadowMapData>,	NextVirtualShadowMapData )
 		SHADER_PARAMETER(uint32,														NextVirtualShadowMapDataCount )		
 		// Required if using FHasCacheDataDim
-		SHADER_PARAMETER_RDG_TEXTURE(Texture2D<uint>, 									PrevPageRequestFlags)
+		SHADER_PARAMETER_RDG_TEXTURE(Texture2D<uint>, 									PrevPageFlags)
 		// Required if using FGenerateStatsDim
 		SHADER_PARAMETER_RDG_BUFFER_UAV( RWStructuredBuffer< uint >,					OutStatsBuffer )
 	END_SHADER_PARAMETER_STRUCT()
@@ -1758,12 +1758,12 @@ void FVirtualShadowMapArray::UpdatePhysicalPageAddresses(FRDGBuilder& GraphBuild
 
 	FUpdatePhysicalPageAddresses::FPermutationDomain PermutationVector;
 
-	TRefCountPtr<IPooledRenderTarget> PrevPageRequestFlags = CacheManager->GetPrevBuffers().PageRequestFlags;
-	if (PrevPageRequestFlags != nullptr)
+	TRefCountPtr<IPooledRenderTarget> PrevPageFlags = CacheManager->GetPrevBuffers().PageFlags;
+	if (PrevPageFlags != nullptr)
 	{
-		PassParameters->PrevPageRequestFlags = GraphBuilder.RegisterExternalTexture(PrevPageRequestFlags, TEXT("Shadow.Virtual.PrevPageRequestFlags"));
+		PassParameters->PrevPageFlags = GraphBuilder.RegisterExternalTexture(PrevPageFlags, TEXT("Shadow.Virtual.PrevPageFlags"));
 	}
-	PermutationVector.Set<FUpdatePhysicalPageAddresses::FHasCacheDataDim>(PrevPageRequestFlags != nullptr);
+	PermutationVector.Set<FUpdatePhysicalPageAddresses::FHasCacheDataDim>(PrevPageFlags != nullptr);
 
 	SetStatsArgsAndPermutation<FUpdatePhysicalPageAddresses>(ShouldGenerateStats(), StatsBufferUAV, PassParameters, PermutationVector);
 
@@ -1930,14 +1930,14 @@ void FVirtualShadowMapArray::BuildPageAllocations(
 
 	UniformParameters.PageTableTextureSizeInvSize = FVector4f(FVector2f(PageTableTextureSize), FVector2f(1.0f, 1.0f) / FVector2f(PageTableTextureSize));
 
-	auto AllocatePageTable = [&](bool bAllocateMipLevels, EPixelFormat PixelFormat, const TCHAR* DebugName, uint32 SampleStride = 1u)
+	auto AllocatePageTable = [&](bool bAllocateMipLevels, EPixelFormat PixelFormat, ETextureCreateFlags ExtraFlags, const TCHAR* DebugName, uint32 SampleStride = 1u)
 	{
 		return GraphBuilder.CreateTexture(
 			FRDGTextureDesc::Create2D(
 				PageTableTextureSize * SampleStride, 
 				PixelFormat, 
 				FClearValueBinding(0), 
-				TexCreate_ShaderResource | TexCreate_UAV | TexCreate_AtomicCompatible,
+				TexCreate_ShaderResource | TexCreate_UAV | ExtraFlags,
 				// Technically one Hmip too little but the 1x1 level is meaningless as it mixes the info from all the various mip levels. 
 				bAllocateMipLevels ? (VSM_LOG2_PAGE_SIZE + SampleStride / 2u): 1u), 
 			DebugName);
@@ -1996,15 +1996,12 @@ void FVirtualShadowMapArray::BuildPageAllocations(
 		return PageTableDummyRDG;
 	};
 
-
-
-	// TODO: should be uint8 - but atomics prevent this
-	PageRequestFlagsRDG = AllocatePageTable(false, PF_R32_UINT, TEXT("Shadow.Virtual.PageRequestFlags"));
+	PageRequestFlagsRDG = AllocatePageTable(false, PF_R8_UINT, ETextureCreateFlags::None, TEXT("Shadow.Virtual.PageRequestFlags"));
 	AddClearPageTableUAVPass(GraphBuilder, PageRequestFlagsRDG, 0u);
 	if (UniformParameters.bEnableReceiverMasks)
 	{
 		// TODO: should be uint16 - but atomics prevent this
-		PageReceiverMasksRDG = AllocatePageTable(true, PF_R32_UINT, TEXT("Shadow.Virtual.PageReceiverMasks"), 2u);
+		PageReceiverMasksRDG = AllocatePageTable(true, PF_R32_UINT, ETextureCreateFlags::AtomicCompatible, TEXT("Shadow.Virtual.PageReceiverMasks"), 2u);
 		AddClearPageTableUAVPass(GraphBuilder, PageReceiverMasksRDG, 0u, 2u);
 	}
 	else
@@ -2313,9 +2310,9 @@ void FVirtualShadowMapArray::BuildPageAllocations(
 		}
 	}
 
-	PageTableRDG = AllocatePageTable(false, PF_R32_UINT, TEXT("Shadow.Virtual.PageTable"));		
+	PageTableRDG = AllocatePageTable(false, PF_R32_UINT, ETextureCreateFlags::AtomicCompatible, TEXT("Shadow.Virtual.PageTable"));		
 	// TODO: should be uint8 - but atomics prevent this
-	PageFlagsRDG = AllocatePageTable(true, PF_R32_UINT, TEXT("Shadow.Virtual.PageFlags"));
+	PageFlagsRDG = AllocatePageTable(true, PF_R32_UINT, ETextureCreateFlags::AtomicCompatible, TEXT("Shadow.Virtual.PageFlags"));
 
 	AddClearPageTableUAVPass(GraphBuilder, PageTableRDG, 0u);
 	AddClearPageTableUAVPass(GraphBuilder, PageFlagsRDG, 0u);
diff --git a/Engine/Source/Runtime/Renderer/Private/VirtualShadowMaps/VirtualShadowMapCacheManager.cpp b/Engine/Source/Runtime/Renderer/Private/VirtualShadowMaps/VirtualShadowMapCacheManager.cpp
--- a/Engine/Source/Runtime/Renderer/Private/VirtualShadowMaps/VirtualShadowMapCacheManager.cpp
+++ b/Engine/Source/Runtime/Renderer/Private/VirtualShadowMaps/VirtualShadowMapCacheManager.cpp
@@ -1075,8 +1075,7 @@ bool FVirtualShadowMapArrayCacheManager::IsCacheDataAvailable()
 		PrevBuffers.UncachedPageRectBounds &&
 		PrevBuffers.AllocatedPageRectBounds &&
 		PrevBuffers.ProjectionData &&
-		PrevBuffers.PhysicalPageLists &&
-		PrevBuffers.PageRequestFlags;
+		PrevBuffers.PhysicalPageLists;
 }
 
 bool FVirtualShadowMapArrayCacheManager::IsHZBDataAvailable()
@@ -1271,13 +1270,11 @@ void FVirtualShadowMapArrayCacheManager::ExtractFrameData(
 		GraphBuilder.QueueBufferExtraction(VirtualShadowMapArray.UncachedPageRectBoundsRDG, &PrevBuffers.UncachedPageRectBounds);
 		GraphBuilder.QueueBufferExtraction(VirtualShadowMapArray.AllocatedPageRectBoundsRDG, &PrevBuffers.AllocatedPageRectBounds);
 		GraphBuilder.QueueTextureExtraction(VirtualShadowMapArray.PageFlagsRDG, &PrevBuffers.PageFlags);
-		GraphBuilder.QueueTextureExtraction(VirtualShadowMapArray.PageReceiverMasksRDG, &PrevBuffers.PageReceiverMasks);
 
 		if (IsCacheEnabled())
 		{
 			GraphBuilder.QueueBufferExtraction(VirtualShadowMapArray.ProjectionDataRDG, &PrevBuffers.ProjectionData);
 			GraphBuilder.QueueBufferExtraction(VirtualShadowMapArray.PhysicalPageListsRDG, &PrevBuffers.PhysicalPageLists);
-			GraphBuilder.QueueTextureExtraction(VirtualShadowMapArray.PageRequestFlagsRDG, &PrevBuffers.PageRequestFlags);
 			
 			// Store but drop any temp references embedded in the uniform parameters this frame
 			PrevUniformParameters = VirtualShadowMapArray.UniformParameters;
@@ -1495,7 +1492,7 @@ BEGIN_SHADER_PARAMETER_STRUCT(FInvalidatePagesParameters, )
 	SHADER_PARAMETER_RDG_UNIFORM_BUFFER(FVirtualShadowMapUniformParameters, VirtualShadowMap)
 	SHADER_PARAMETER_RDG_UNIFORM_BUFFER(FSceneUniformParameters, Scene)
 	SHADER_PARAMETER_RDG_BUFFER_UAV(RWStructuredBuffer<FPhysicalPageMetaData>, PhysicalPageMetaDataOut)
-	SHADER_PARAMETER_RDG_TEXTURE_UAV(RWTexture2D<uint>, OutPageRequestFlags)
+	SHADER_PARAMETER_RDG_TEXTURE_UAV(RWTexture2D<uint>, OutPrevPageFlags)
 
 	// When USE_HZB_OCCLUSION
 	SHADER_PARAMETER_RDG_TEXTURE(Texture2D<uint>, HZBPageTable)
@@ -1725,11 +1722,14 @@ FVirtualShadowMapArrayCacheManager::FInvalidationPassCommon FVirtualShadowMapArr
 		UniformParameters->ProjectionData = RegExtCreateSrv(PrevBuffers.ProjectionData, TEXT("Shadow.Virtual.PrevProjectionData"));
 		UniformParameters->PageTable = GraphBuilder.RegisterExternalTexture(PrevBuffers.PageTable, TEXT("Shadow.Virtual.PrevPageTable"));
 		UniformParameters->PageFlags = GraphBuilder.RegisterExternalTexture(PrevBuffers.PageFlags, TEXT("Shadow.Virtual.PrevPageFlags"));
-		UniformParameters->PageReceiverMasks = GraphBuilder.RegisterExternalTexture(PrevBuffers.PageReceiverMasks, TEXT("Shadow.Virtual.PrevPageReceiverMasks"));
 		UniformParameters->UncachedPageRectBounds = RegExtCreateSrv(PrevBuffers.UncachedPageRectBounds, TEXT("Shadow.Virtual.PrevUncachedPageRectBounds"));
 		UniformParameters->AllocatedPageRectBounds = RegExtCreateSrv(PrevBuffers.AllocatedPageRectBounds, TEXT("Shadow.Virtual.PrevAllocatedPageRectBounds"));
 		UniformParameters->CachePrimitiveAsDynamic = GraphBuilder.CreateSRV(UploadCachePrimitiveAsDynamic(GraphBuilder));
 
+		// We don't use reveiver masks for invalidation & probably never will as it is coupled with always-invalidate-dynamic
+		UniformParameters->PageReceiverMasks = GSystemTextures.GetDefaultTexture2D(GraphBuilder, PF_R32_UINT, 0xFFFFFFFFu);
+		UniformParameters->bEnableReceiverMasks = false;
+
 		// Unused in this path... may be a better way to handle this
 		UniformParameters->PhysicalPagePool = GSystemTextures.GetZeroUIntArrayAtomicCompatDummy(GraphBuilder);
 		FRDGBufferSRVRef Uint32SRVDummy = GraphBuilder.CreateSRV(GSystemTextures.GetDefaultStructuredBuffer(GraphBuilder, sizeof(uint32)));
@@ -1753,7 +1753,7 @@ void FVirtualShadowMapArrayCacheManager::SetInvalidateInstancePagesParameters(
 	PassParameters->VirtualShadowMap = InvalidationPassCommon.VirtualShadowMapUniformBuffer;
 	PassParameters->Scene = InvalidationPassCommon.SceneUniformBuffer;
 	PassParameters->PhysicalPageMetaDataOut = GraphBuilder.CreateUAV(GraphBuilder.RegisterExternalBuffer(PhysicalPageMetaData));
-	PassParameters->OutPageRequestFlags = GraphBuilder.CreateUAV(GraphBuilder.RegisterExternalTexture(PrevBuffers.PageRequestFlags));
+	PassParameters->OutPrevPageFlags = GraphBuilder.CreateUAV(GraphBuilder.RegisterExternalTexture(PrevBuffers.PageFlags));
 
 	const bool bUseHZB = (CVarCacheVsmUseHzb.GetValueOnRenderThread() != 0);
 	const TRefCountPtr<IPooledRenderTarget> HZBPhysical = (bUseHZB && HZBPhysicalPagePoolArray) ? HZBPhysicalPagePoolArray : nullptr;
diff --git a/Engine/Source/Runtime/Renderer/Private/VirtualShadowMaps/VirtualShadowMapCacheManager.h b/Engine/Source/Runtime/Renderer/Private/VirtualShadowMaps/VirtualShadowMapCacheManager.h
--- a/Engine/Source/Runtime/Renderer/Private/VirtualShadowMaps/VirtualShadowMapCacheManager.h
+++ b/Engine/Source/Runtime/Renderer/Private/VirtualShadowMaps/VirtualShadowMapCacheManager.h
@@ -231,8 +231,6 @@ struct FVirtualShadowMapArrayFrameData
 	TRefCountPtr<FRDGPooledBuffer>				AllocatedPageRectBounds;
 	TRefCountPtr<FRDGPooledBuffer>				ProjectionData;
 	TRefCountPtr<FRDGPooledBuffer>				PhysicalPageLists;
-	TRefCountPtr<IPooledRenderTarget>			PageRequestFlags;
-	TRefCountPtr<IPooledRenderTarget>			PageReceiverMasks;
 	uint64 GetGPUSizeBytes(bool bLogSizes) const;
 };
 
