Split virtual texture Finalize into two phases.
RenderFinalize() does not write to the physical texture, but can read from virtual textures. This is a phase that allows any page rendering which may need to sample virtual textures.
Finalise() must write to the physical textures, but cannot sample from them.
This seperation allows us to make transitions to page table and physical textures at the correct times.
Implemented the seperation for runtime virtual textures.
At the same time implemented batching for runtime virtual textures so that page rendering (RenderFinalize()) and compression (Finalize()) can both be batched seperately for better GPU performance.
Also pass the SceneRenderer down to RenderFinalize() and RenderPages(). This is simpler than current method of finding it on the RDG blackboard.
#rnx

diff --git a/Engine/Plugins/Experimental/VirtualHeightfieldMesh/Source/VirtualHeightfieldMeshEditor/Private/HeightfieldMinMaxTextureBuild.cpp b/Engine/Plugins/Experimental/VirtualHeightfieldMesh/Source/VirtualHeightfieldMeshEditor/Private/HeightfieldMinMaxTextureBuild.cpp
--- a/Engine/Plugins/Experimental/VirtualHeightfieldMesh/Source/VirtualHeightfieldMeshEditor/Private/HeightfieldMinMaxTextureBuild.cpp
+++ b/Engine/Plugins/Experimental/VirtualHeightfieldMesh/Source/VirtualHeightfieldMeshEditor/Private/HeightfieldMinMaxTextureBuild.cpp
@@ -4,6 +4,7 @@
 
 #include "Components/RuntimeVirtualTextureComponent.h"
 #include "ContentStreaming.h"
+#include "EngineModule.h"
 #include "HeightfieldMinMaxRender.h"
 #include "HeightfieldMinMaxTexture.h"
 #include "Misc/ScopedSlowTask.h"
@@ -193,12 +194,10 @@ namespace VirtualHeightfieldMesh
 				{
 					// Rendering one page at a time, but could batch here?
 					FRDGBuilder GraphBuilder(RHICmdList);
-
-					const FBox2D TileBox(FVector2D(0, 0), FVector2D(TileSize, TileSize));
-					const FIntRect TileRect(0, 0, TileSize, TileSize);
+					FScenePrimitiveRenderingContextScopeHelper RenderingScope(GetRendererModule().BeginScenePrimitiveRendering(GraphBuilder, *Scene));
 
 					RuntimeVirtualTexture::FRenderPageBatchDesc Desc;
-					Desc.Scene = Scene->GetRenderScene();
+					Desc.SceneRenderer = RenderingScope.ScenePrimitiveRenderingContext->GetSceneRenderer();
 					Desc.RuntimeVirtualTextureId = RuntimeVirtualTextureId;
 					Desc.UVToWorld = Transform;
 					Desc.WorldBounds = Bounds;
@@ -211,13 +210,13 @@ namespace VirtualHeightfieldMesh
 					Desc.Targets[0].Texture = RenderTileResources.GetTileRenderTarget()->GetRHI();
 					Desc.Targets[1].Texture = nullptr;
 					Desc.Targets[2].Texture = nullptr;
-					Desc.PageDescs[0].DestBox[0] = TileBox;
-					Desc.PageDescs[0].DestBox[1] = TileBox;
-					Desc.PageDescs[0].DestBox[2] = TileBox;
+					Desc.PageDescs[0].DestRect[0] = FIntRect(0, 0, TileSize, TileSize);
+					Desc.PageDescs[0].DestRect[1] = FIntRect();
+					Desc.PageDescs[0].DestRect[2] = FIntRect();
 					Desc.PageDescs[0].UVRange = UVRange;
 					Desc.PageDescs[0].vLevel = MipLevel;
 
-					RenderPagesStandAlone(GraphBuilder, Desc);
+					RenderPages(GraphBuilder, Desc);
 
 					// Downsample page to texel in output
 					FRDGTextureRef SrcTexture = GraphBuilder.RegisterExternalTexture(RenderTileResources.GetTileRenderTarget());
diff --git a/Engine/Shaders/Private/BlockCompressionCommon.ush b/Engine/Shaders/Private/BlockCompressionCommon.ush
--- a/Engine/Shaders/Private/BlockCompressionCommon.ush
+++ b/Engine/Shaders/Private/BlockCompressionCommon.ush
@@ -7,8 +7,20 @@
 
 #pragma once
 
+#ifndef BLOCK_COMPRESS_SRC_TEXTURE_ARRAY
+#define BLOCK_COMPRESS_SRC_TEXTURE_ARRAY 0
+#endif
+
+#if BLOCK_COMPRESS_SRC_TEXTURE_ARRAY
+#define TextureType Texture2DArray
+#define CoordType float3
+#else
+#define TextureType Texture2D
+#define CoordType float2
+#endif
+
 // Read a 4x4 color block ready for BC compression
-void ReadBlockRGB(Texture2D<float4> SourceTexture, SamplerState TextureSampler, float2 UV, float2 TexelUVSize, out float3 Block[16])
+void ReadBlockRGB(TextureType<float4> SourceTexture, SamplerState TextureSampler, CoordType UV, float2 TexelUVSize, out float3 Block[16])
 {
 	{
 		float4 Red = SourceTexture.GatherRed(TextureSampler, UV);
@@ -20,7 +32,8 @@ void ReadBlockRGB(Texture2D<float4> SourceTexture, SamplerState TextureSampler,
 		Block[5] = float3(Red[1], Green[1], Blue[1]);
 	}
 	{
-		float2 UVOffset = UV + float2(2.f * TexelUVSize.x, 0);
+		CoordType UVOffset = UV;
+		UVOffset.x += 2.f * TexelUVSize.x;
 		float4 Red = SourceTexture.GatherRed(TextureSampler, UVOffset);
 		float4 Green = SourceTexture.GatherGreen(TextureSampler, UVOffset);
 		float4 Blue = SourceTexture.GatherBlue(TextureSampler, UVOffset);
@@ -30,7 +43,8 @@ void ReadBlockRGB(Texture2D<float4> SourceTexture, SamplerState TextureSampler,
 		Block[7] = float3(Red[1], Green[1], Blue[1]);
 	}
 	{
-		float2 UVOffset = UV + float2(0, 2.f * TexelUVSize.y);
+		CoordType UVOffset = UV;
+		UVOffset.y += 2.f * TexelUVSize.y;
 		float4 Red = SourceTexture.GatherRed(TextureSampler, UVOffset);
 		float4 Green = SourceTexture.GatherGreen(TextureSampler, UVOffset);
 		float4 Blue = SourceTexture.GatherBlue(TextureSampler, UVOffset);
@@ -40,7 +54,8 @@ void ReadBlockRGB(Texture2D<float4> SourceTexture, SamplerState TextureSampler,
 		Block[13] = float3(Red[1], Green[1], Blue[1]);
 	}
 	{
-		float2 UVOffset = UV + 2.f * TexelUVSize;
+		CoordType UVOffset = UV;
+		UVOffset.xy += 2.f * TexelUVSize;
 		float4 Red = SourceTexture.GatherRed(TextureSampler, UVOffset);
 		float4 Green = SourceTexture.GatherGreen(TextureSampler, UVOffset);
 		float4 Blue = SourceTexture.GatherBlue(TextureSampler, UVOffset);
@@ -52,7 +67,7 @@ void ReadBlockRGB(Texture2D<float4> SourceTexture, SamplerState TextureSampler,
 }
 
 // Read a 4x4 alpha block ready for BC compression
-void ReadBlockAlpha(Texture2D<float4> SourceTexture, SamplerState TextureSampler, float2 UV, float2 TexelUVSize, out float Block[16])
+void ReadBlockAlpha(TextureType<float4> SourceTexture, SamplerState TextureSampler, CoordType UV, float2 TexelUVSize, out float Block[16])
 {
 	{
 		float4 Alpha = SourceTexture.GatherAlpha(TextureSampler, UV);
@@ -62,7 +77,8 @@ void ReadBlockAlpha(Texture2D<float4> SourceTexture, SamplerState TextureSampler
 		Block[5] = Alpha[1];
 	}
 	{
-		float2 UVOffset = UV + float2(2.f * TexelUVSize.x, 0);
+		CoordType UVOffset = UV;
+		UVOffset.x += 2.f * TexelUVSize.x;
 		float4 Alpha = SourceTexture.GatherAlpha(TextureSampler, UVOffset);
 		Block[2] = Alpha[3];
 		Block[3] = Alpha[2];
@@ -70,7 +86,8 @@ void ReadBlockAlpha(Texture2D<float4> SourceTexture, SamplerState TextureSampler
 		Block[7] = Alpha[1];
 	}
 	{
-		float2 UVOffset = UV + float2(0, 2.f * TexelUVSize.y);
+		CoordType UVOffset = UV;
+		UVOffset.y += 2.f * TexelUVSize.y;
 		float4 Alpha = SourceTexture.GatherAlpha(TextureSampler, UVOffset);
 		Block[8] = Alpha[3];
 		Block[9] = Alpha[2];
@@ -78,7 +95,8 @@ void ReadBlockAlpha(Texture2D<float4> SourceTexture, SamplerState TextureSampler
 		Block[13] = Alpha[1];
 	}
 	{
-		float2 UVOffset = UV + 2.f * TexelUVSize;
+		CoordType UVOffset = UV;
+		UVOffset.xy += 2.f * TexelUVSize;
 		float4 Alpha = SourceTexture.GatherAlpha(TextureSampler, UVOffset);
 		Block[10] = Alpha[3];
 		Block[11] = Alpha[2];
@@ -87,7 +105,7 @@ void ReadBlockAlpha(Texture2D<float4> SourceTexture, SamplerState TextureSampler
 	}
 }
 
-void ReadBlockRGBA(Texture2D<float4> SourceTexture, SamplerState TextureSampler, float2 UV, float2 TexelUVSize, out float4 Block[16])
+void ReadBlockRGBA(TextureType<float4> SourceTexture, SamplerState TextureSampler, CoordType UV, float2 TexelUVSize, out float4 Block[16])
 {
 	{
 		float4 Red = SourceTexture.GatherRed(TextureSampler, UV);
@@ -101,7 +119,8 @@ void ReadBlockRGBA(Texture2D<float4> SourceTexture, SamplerState TextureSampler,
 		Block[5] = float4(Red[1], Green[1], Blue[1], Alpha[1]);
 	}
 	{
-		float2 UVOffset = UV + float2(2.f * TexelUVSize.x, 0);
+		CoordType UVOffset = UV;
+		UVOffset.x += 2.f * TexelUVSize.x;
 		float4 Red = SourceTexture.GatherRed(TextureSampler, UVOffset);
 		float4 Green = SourceTexture.GatherGreen(TextureSampler, UVOffset);
 		float4 Blue = SourceTexture.GatherBlue(TextureSampler, UVOffset);
@@ -113,7 +132,8 @@ void ReadBlockRGBA(Texture2D<float4> SourceTexture, SamplerState TextureSampler,
 		Block[7] = float4(Red[1], Green[1], Blue[1], Alpha[1]);
 	}
 	{
-		float2 UVOffset = UV + float2(0, 2.f * TexelUVSize.y);
+		CoordType UVOffset = UV;
+		UVOffset.y += 2.f * TexelUVSize.y;
 		float4 Red = SourceTexture.GatherRed(TextureSampler, UVOffset);
 		float4 Green = SourceTexture.GatherGreen(TextureSampler, UVOffset);
 		float4 Blue = SourceTexture.GatherBlue(TextureSampler, UVOffset);
@@ -124,7 +144,8 @@ void ReadBlockRGBA(Texture2D<float4> SourceTexture, SamplerState TextureSampler,
 		Block[13] = float4(Red[1], Green[1], Blue[1], Alpha[1]);
 	}
 	{
-		float2 UVOffset = UV + 2.f * TexelUVSize;
+		CoordType UVOffset = UV;
+		UVOffset.xy += 2.f * TexelUVSize;
 		float4 Red = SourceTexture.GatherRed(TextureSampler, UVOffset);
 		float4 Green = SourceTexture.GatherGreen(TextureSampler, UVOffset);
 		float4 Blue = SourceTexture.GatherBlue(TextureSampler, UVOffset);
@@ -137,7 +158,7 @@ void ReadBlockRGBA(Texture2D<float4> SourceTexture, SamplerState TextureSampler,
 }
 
 // Read a 4x4 color of channel X block ready for BC compression
-void ReadBlockX(Texture2D<float4> SourceTexture, SamplerState TextureSampler, float2 UV, float2 TexelUVSize, out float Block[16])
+void ReadBlockX(TextureType<float4> SourceTexture, SamplerState TextureSampler, CoordType UV, float2 TexelUVSize, out float Block[16])
 {
 	{
 		float4 Red = SourceTexture.GatherRed(TextureSampler, UV);
@@ -147,7 +168,8 @@ void ReadBlockX(Texture2D<float4> SourceTexture, SamplerState TextureSampler, fl
 		Block[5] = Red[1];
 	}
 	{
-		float2 UVOffset = UV + float2(2.f * TexelUVSize.x, 0);
+		CoordType UVOffset = UV;
+		UVOffset.x += 2.f * TexelUVSize.x;
 		float4 Red = SourceTexture.GatherRed(TextureSampler, UVOffset);
 		Block[2] = Red[3];
 		Block[3] = Red[2];
@@ -155,7 +177,8 @@ void ReadBlockX(Texture2D<float4> SourceTexture, SamplerState TextureSampler, fl
 		Block[7] = Red[1];
 	}
 	{
-		float2 UVOffset = UV + float2(0, 2.f * TexelUVSize.y);
+		CoordType UVOffset = UV;
+		UVOffset.y += 2.f * TexelUVSize.y;
 		float4 Red = SourceTexture.GatherRed(TextureSampler, UVOffset);
 		Block[8] = Red[3];
 		Block[9] = Red[2];
@@ -163,7 +186,8 @@ void ReadBlockX(Texture2D<float4> SourceTexture, SamplerState TextureSampler, fl
 		Block[13] = Red[1];
 	}
 	{
-		float2 UVOffset = UV + 2.f * TexelUVSize;
+		CoordType UVOffset = UV;
+		UVOffset.xy += 2.f * TexelUVSize;
 		float4 Red = SourceTexture.GatherRed(TextureSampler, UVOffset);
 		Block[10] = Red[3];
 		Block[11] = Red[2];
@@ -173,7 +197,7 @@ void ReadBlockX(Texture2D<float4> SourceTexture, SamplerState TextureSampler, fl
 }
 
 // Read a 4x4 block of XY channels from a normal texture ready for BC5 compression.
-void ReadBlockXY(Texture2D<float4> SourceTexture, SamplerState TextureSampler, float2 UV, float2 TexelUVSize, out float BlockX[16], out float BlockY[16])
+void ReadBlockXY(TextureType<float4> SourceTexture, SamplerState TextureSampler, CoordType UV, float2 TexelUVSize, out float BlockX[16], out float BlockY[16])
 {
 	{
 		float4 Red = SourceTexture.GatherRed(TextureSampler, UV);
@@ -184,7 +208,8 @@ void ReadBlockXY(Texture2D<float4> SourceTexture, SamplerState TextureSampler, f
 		BlockX[5] = Red[1]; BlockY[5] = Green[1];
 	}
 	{
-		float2 UVOffset = UV + float2(2.f * TexelUVSize.x, 0);
+		CoordType UVOffset = UV;
+		UVOffset.x += 2.f * TexelUVSize.x;
 		float4 Red = SourceTexture.GatherRed(TextureSampler, UVOffset);
 		float4 Green = SourceTexture.GatherGreen(TextureSampler, UVOffset);
 		BlockX[2] = Red[3]; BlockY[2] = Green[3];
@@ -193,7 +218,8 @@ void ReadBlockXY(Texture2D<float4> SourceTexture, SamplerState TextureSampler, f
 		BlockX[7] = Red[1]; BlockY[7] = Green[1];
 	}
 	{
-		float2 UVOffset = UV + float2(0, 2.f * TexelUVSize.y);
+		CoordType UVOffset = UV;
+		UVOffset.y += 2.f * TexelUVSize.y;
 		float4 Red = SourceTexture.GatherRed(TextureSampler, UVOffset);
 		float4 Green = SourceTexture.GatherGreen(TextureSampler, UVOffset);
 		BlockX[8] = Red[3]; BlockY[8] = Green[3];
@@ -202,7 +228,8 @@ void ReadBlockXY(Texture2D<float4> SourceTexture, SamplerState TextureSampler, f
 		BlockX[13] = Red[1]; BlockY[13] = Green[1];
 	}
 	{
-		float2 UVOffset = UV + 2.f * TexelUVSize;
+		CoordType UVOffset = UV;
+		UVOffset.xy += 2.f * TexelUVSize;
 		float4 Red = SourceTexture.GatherRed(TextureSampler, UVOffset);
 		float4 Green = SourceTexture.GatherGreen(TextureSampler, UVOffset);
 		BlockX[10] = Red[3]; BlockY[10] = Green[3];
@@ -213,7 +240,7 @@ void ReadBlockXY(Texture2D<float4> SourceTexture, SamplerState TextureSampler, f
 }
 
 // Read a 4x4 block of XYA channels ready for BC4/BC5 compression.
-void ReadBlockXYA(Texture2D<float4> SourceTexture, SamplerState TextureSampler, float2 UV, float2 TexelUVSize, out float BlockX[16], out float BlockY[16], out float BlockA[16])
+void ReadBlockXYA(TextureType<float4> SourceTexture, SamplerState TextureSampler, CoordType UV, float2 TexelUVSize, out float BlockX[16], out float BlockY[16], out float BlockA[16])
 {
 	{
 		float4 Red = SourceTexture.GatherRed(TextureSampler, UV);
@@ -225,7 +252,8 @@ void ReadBlockXYA(Texture2D<float4> SourceTexture, SamplerState TextureSampler,
 		BlockX[5] = Red[1]; BlockY[5] = Green[1]; BlockA[5] = Blue[1];
 	}
 	{
-		float2 UVOffset = UV + float2(2.f * TexelUVSize.x, 0);
+		CoordType UVOffset = UV;
+		UVOffset.x += 2.f * TexelUVSize.x;
 		float4 Red = SourceTexture.GatherRed(TextureSampler, UVOffset);
 		float4 Green = SourceTexture.GatherGreen(TextureSampler, UVOffset);
 		float4 Blue = SourceTexture.GatherBlue(TextureSampler, UVOffset);
@@ -235,7 +263,8 @@ void ReadBlockXYA(Texture2D<float4> SourceTexture, SamplerState TextureSampler,
 		BlockX[7] = Red[1]; BlockY[7] = Green[1]; BlockA[7] = Blue[1];
 	}
 	{
-		float2 UVOffset = UV + float2(0, 2.f * TexelUVSize.y);
+		CoordType UVOffset = UV;
+		UVOffset.y += 2.f * TexelUVSize.y;
 		float4 Red = SourceTexture.GatherRed(TextureSampler, UVOffset);
 		float4 Green = SourceTexture.GatherGreen(TextureSampler, UVOffset);
 		float4 Blue = SourceTexture.GatherBlue(TextureSampler, UVOffset);
@@ -245,7 +274,8 @@ void ReadBlockXYA(Texture2D<float4> SourceTexture, SamplerState TextureSampler,
 		BlockX[13] = Red[1]; BlockY[13] = Green[1]; BlockA[13] = Blue[1];
 	}
 	{
-		float2 UVOffset = UV + 2.f * TexelUVSize;
+		CoordType UVOffset = UV;
+		UVOffset.xy += 2.f * TexelUVSize;
 		float4 Red = SourceTexture.GatherRed(TextureSampler, UVOffset);
 		float4 Green = SourceTexture.GatherGreen(TextureSampler, UVOffset);
 		float4 Blue = SourceTexture.GatherBlue(TextureSampler, UVOffset);
@@ -254,4 +284,7 @@ void ReadBlockXYA(Texture2D<float4> SourceTexture, SamplerState TextureSampler,
 		BlockX[14] = Red[0]; BlockY[14] = Green[0]; BlockA[14] = Blue[0];
 		BlockX[15] = Red[1]; BlockY[15] = Green[1]; BlockA[15] = Blue[1];
 	}
-}
\ No newline at end of file
+}
+
+#undef TextureType
+#undef CoordType
diff --git a/Engine/Shaders/Private/VirtualTextureCompress.usf b/Engine/Shaders/Private/VirtualTextureCompress.usf
--- a/Engine/Shaders/Private/VirtualTextureCompress.usf
+++ b/Engine/Shaders/Private/VirtualTextureCompress.usf
@@ -3,7 +3,6 @@
 /*=============================================================================
 	VirtualTextureCompress.usf: 
 	Compression pass for runtime virtual texture
-	todo[vt]: Lots of possible optimizations to do...
 =============================================================================*/
 
 #include "Common.ush"
@@ -30,26 +29,44 @@
 #define TYPE_64BIT uint2
 #endif
 
-uint4 SourceRect;
-uint2 DestPos0;
-uint2 DestPos1;
-uint2 DestPos2;
+#ifndef USE_SRC_TEXTURE_ARRAY
+#define USE_SRC_TEXTURE_ARRAY 0
+#endif
+
+#ifndef USE_DST_TEXTURE_ARRAY
+#define USE_DST_TEXTURE_ARRAY 0
+#endif
+
+#if USE_SRC_TEXTURE_ARRAY
+#define SourceTextureType Texture2DArray
+#else
+#define SourceTextureType Texture2D
+#endif
+
+#if USE_DST_TEXTURE_ARRAY
+#define DestTextureType RWTexture2DArray
+#else
+#define DestTextureType RWTexture2D
+#endif
 
-Texture2D<float4> RenderTexture0;
+SourceTextureType<float4> RenderTexture0;
 SamplerState TextureSampler0;
-Texture2D<float4> RenderTexture1;
+SourceTextureType<float4> RenderTexture1;
 SamplerState TextureSampler1;
-Texture2D<float4> RenderTexture2;
+SourceTextureType<float4> RenderTexture2;
 SamplerState TextureSampler2;
 
 // Outputs can be 64bit or 128 bit depending on formats.
 // The 64 bit type can be uint2 or uint4 depending on platform support.
-RWTexture2D<TYPE_64BIT> OutCompressTexture0_64bit;
-RWTexture2D<TYPE_64BIT> OutCompressTexture1_64bit;
-RWTexture2D<TYPE_64BIT> OutCompressTexture2_64bit;
-RWTexture2D<uint4> OutCompressTexture0_128bit;
-RWTexture2D<uint4> OutCompressTexture1_128bit;
-RWTexture2D<uint4> OutCompressTexture2_128bit;
+DestTextureType<TYPE_64BIT> OutCompressTexture0_64bit;
+DestTextureType<TYPE_64BIT> OutCompressTexture1_64bit;
+DestTextureType<TYPE_64BIT> OutCompressTexture2_64bit;
+DestTextureType<uint4> OutCompressTexture0_128bit;
+DestTextureType<uint4> OutCompressTexture1_128bit;
+DestTextureType<uint4> OutCompressTexture2_128bit;
+
+uint4 SourceRect;
+DECLARE_SCALAR_ARRAY(uint, DestPos, MAX_BATCH_SIZE * MAX_DST_LAYERS * 2);
 
 #if PACK_RG32_RGBA16
 /** Pack RG32 to RGBA16 for output. */
@@ -59,7 +76,7 @@ uint4 PackRG32(in uint2 u2)
 	uint y0 = (u2.y >> 16) & 0x0000ffff;
 	uint x1 = (u2.x) & 0x0000ffff;
 	uint y1 = (u2.y) & 0x0000ffff;
-	return uint4(x1,x0,y1,y0);
+	return uint4(x1, x0, y1, y0);
 }
 #else
 uint2 PackRG32(in uint2 u2)
@@ -68,6 +85,24 @@ uint2 PackRG32(in uint2 u2)
 }
 #endif // PACK_RG32_RGBA16
 
+#if USE_SRC_TEXTURE_ARRAY
+float3 GetSourceUV(float2 UV, uint3 ThreadId) { return float3(UV, ThreadId.z); }
+#else
+float2 GetSourceUV(float2 UV, uint3 ThreadId) { return UV; }
+#endif
+
+#if USE_DST_TEXTURE_ARRAY
+uint3 GetDestPos(uint LayerIndex, uint3 ThreadId) { return ThreadId; }
+#else
+uint2 GetDestPos(uint LayerIndex, uint3 ThreadId) 
+{
+	uint Index = (ThreadId.z * MAX_DST_LAYERS + LayerIndex) * 2;
+	uint2 BasePos = uint2(GET_SCALAR_ARRAY_ELEMENT(DestPos, Index), GET_SCALAR_ARRAY_ELEMENT(DestPos, (Index + 1)));
+	return BasePos + ThreadId.xy; 
+}
+#endif
+
+#if VERTEXSHADER
 /** Trivial vertex shader for copy shaders. */
 void CopyVS(in uint VertexId : SV_VertexID, out float4 OutPosition : SV_POSITION, out float2 OutTexCoord : TEXCOORD0)
 {
@@ -75,7 +110,9 @@ void CopyVS(in uint VertexId : SV_VertexID, out float4 OutPosition : SV_POSITION
 	OutPosition.xy = float2(OutTexCoord.x, 1.f - OutTexCoord.y) * 2.f - 1.f;
 	OutPosition.zw = float2(0, 1);
 }
+#endif // VERTEXSHADER
 
+#if COMPUTESHADER
 /** Compress base color to a single BC1 target. */
 [numthreads(8,8,1)]
 void CompressBaseColorCS(
@@ -91,17 +128,18 @@ void CompressBaseColorCS(
 	float2 SampleUV = (float2(SamplePos) + 0.5f) * TexelUVSize;
 
 	float3 BlockBaseColor[16];
-	ReadBlockRGB(RenderTexture0, TextureSampler0, SampleUV, TexelUVSize, BlockBaseColor);
+	ReadBlockRGB(RenderTexture0, TextureSampler0, GetSourceUV(SampleUV, ThreadId), TexelUVSize, BlockBaseColor);
 #if ASTC_PROFIlE
-	OutCompressTexture0_64bit[ThreadId.xy + DestPos0] = PackRG32(CompressBlock_ETC2_SRGB(BlockBaseColor));
+	OutCompressTexture0_64bit[GetDestPos(0, ThreadId)] = PackRG32(CompressBlock_ETC2_SRGB(BlockBaseColor));
 #elif ETC_PROFILE
-	OutCompressTexture0_64bit[ThreadId.xy + DestPos0] = PackRG32(CompressBlock_ETC2_SRGB(BlockBaseColor));
+	OutCompressTexture0_64bit[GetDestPos(0, ThreadId)] = PackRG32(CompressBlock_ETC2_SRGB(BlockBaseColor));
 #else
-	OutCompressTexture0_64bit[ThreadId.xy + DestPos0] = PackRG32(CompressBC1Block_SRGB(BlockBaseColor));
+	OutCompressTexture0_64bit[GetDestPos(0, ThreadId)] = PackRG32(CompressBC1Block_SRGB(BlockBaseColor));
 #endif
-
 }
+#endif // COMPUTESHADER
 
+#if PIXELSHADER
 /** Simple copy of base color used for thumbnails. */
 void CopyBaseColorPS(
 	in float4 InPosition : SV_POSITION, 
@@ -111,7 +149,9 @@ void CopyBaseColorPS(
 	float3 BaseColor = RenderTexture0.SampleLevel(TextureSampler0, InTexCoord, 0).xyz;
 	OutColor0 = float4(BaseColor, 1.f);
 }
+#endif // PIXELSHADER
 
+#if COMPUTESHADER
 /** Compress base color, normal, roughness and specular to BC3 target pair. */
 [numthreads(8,8,1)]
 void CompressBaseColorNormalSpecularCS(
@@ -127,32 +167,33 @@ void CompressBaseColorNormalSpecularCS(
 	float2 SampleUV = (float2(SamplePos) + 0.5f) * TexelUVSize;
 
 	float3 BlockBaseColor[16];
-	ReadBlockRGB(RenderTexture0, TextureSampler0, SampleUV, TexelUVSize, BlockBaseColor);
+	ReadBlockRGB(RenderTexture0, TextureSampler0, GetSourceUV(SampleUV, ThreadId), TexelUVSize, BlockBaseColor);
 
 	float BlockNormalX[16];
 	float BlockNormalY[16];
-	ReadBlockXY(RenderTexture1, TextureSampler1, SampleUV, TexelUVSize, BlockNormalX, BlockNormalY);
+	ReadBlockXY(RenderTexture1, TextureSampler1, GetSourceUV(SampleUV, ThreadId), TexelUVSize, BlockNormalX, BlockNormalY);
 
 	float3 BlockSpecular[16];
-	ReadBlockRGB(RenderTexture2, TextureSampler2, SampleUV, TexelUVSize, BlockSpecular);
+	ReadBlockRGB(RenderTexture2, TextureSampler2, GetSourceUV(SampleUV, ThreadId), TexelUVSize, BlockSpecular);
 
 	for (int i=0; i<16; i++)
 	{
 		BlockSpecular[i].z = round(BlockSpecular[i].z);
 	}
 #if ASTC_PROFILE
-	OutCompressTexture0_128bit[ThreadId.xy + DestPos0] = CompressBlock_ASTC_SRGBA(BlockBaseColor, BlockNormalX);
-	OutCompressTexture1_128bit[ThreadId.xy + DestPos1] = CompressBlock_ASTC_RGBA(BlockSpecular, BlockNormalY);
+	OutCompressTexture0_128bit[GetDestPos(0, ThreadId)] = CompressBlock_ASTC_SRGBA(BlockBaseColor, BlockNormalX);
+	OutCompressTexture1_128bit[GetDestPos(1, ThreadId)] = CompressBlock_ASTC_RGBA(BlockSpecular, BlockNormalY);
 #elif ETC_PROFILE
-	OutCompressTexture0_128bit[ThreadId.xy + DestPos0] = CompressBlock_ETC2_SRGBA(BlockBaseColor, BlockNormalX);
-	OutCompressTexture1_128bit[ThreadId.xy + DestPos1] = CompressBlock_ETC2_RGBA(BlockSpecular, BlockNormalY);
+	OutCompressTexture0_128bit[GetDestPos(0, ThreadId)] = CompressBlock_ETC2_SRGBA(BlockBaseColor, BlockNormalX);
+	OutCompressTexture1_128bit[GetDestPos(1, ThreadId)] = CompressBlock_ETC2_RGBA(BlockSpecular, BlockNormalY);
 #else
-	OutCompressTexture0_128bit[ThreadId.xy + DestPos0] = CompressBC3Block_SRGB(BlockBaseColor, BlockNormalX);
-	OutCompressTexture1_128bit[ThreadId.xy + DestPos1] = CompressBC3Block(BlockSpecular, BlockNormalY);
+	OutCompressTexture0_128bit[GetDestPos(0, ThreadId)] = CompressBC3Block_SRGB(BlockBaseColor, BlockNormalX);
+	OutCompressTexture1_128bit[GetDestPos(1, ThreadId)] = CompressBC3Block(BlockSpecular, BlockNormalY);
 #endif
-
 }
+#endif // COMPUTESHADER
 
+#if COMPUTESHADER
 /** Compress base color, normal, roughness to BC1 & BC3. */
 [numthreads(8, 8, 1)]
 void CompressBaseColorNormalRoughnessCS(
@@ -170,25 +211,26 @@ void CompressBaseColorNormalRoughnessCS(
 	// Note that this format uses non sRGB BaseColor format with manual encode/decode.
 	// So that we don't use _SRGB forms of the block compression functions below.
 	float3 BlockBaseColor[16];
-	ReadBlockRGB(RenderTexture0, TextureSampler0, SampleUV, TexelUVSize, BlockBaseColor);
+	ReadBlockRGB(RenderTexture0, TextureSampler0, GetSourceUV(SampleUV, ThreadId), TexelUVSize, BlockBaseColor);
 
 	float3 BlockNormalXRougnnessY[16];
 	const float BlockNormalZ[16]={ 1.0, 1.0, 1.0, 1.0, 1.0, 1.0, 1.0, 1.0, 
 								   1.0, 1.0, 1.0, 1.0, 1.0, 1.0, 1.0, 1.0};
-	ReadBlockRGB(RenderTexture1, TextureSampler2, SampleUV, TexelUVSize, BlockNormalXRougnnessY);
+	ReadBlockRGB(RenderTexture1, TextureSampler2, GetSourceUV(SampleUV, ThreadId), TexelUVSize, BlockNormalXRougnnessY);
 #if ASTC_PROFILE
-	OutCompressTexture0_64bit[ThreadId.xy + DestPos0] = PackRG32(CompressBlock_ETC2_RGB(BlockBaseColor));
-	OutCompressTexture1_128bit[ThreadId.xy + DestPos1] = CompressBlock_ASTC_RGBA(BlockNormalXRougnnessY, BlockNormalZ);
+	OutCompressTexture0_64bit[GetDestPos(0, ThreadId)] = PackRG32(CompressBlock_ETC2_RGB(BlockBaseColor));
+	OutCompressTexture1_128bit[GetDestPos(1, ThreadId)] = CompressBlock_ASTC_RGBA(BlockNormalXRougnnessY, BlockNormalZ);
 #elif ETC_PROFILE
-	OutCompressTexture0_64bit[ThreadId.xy + DestPos0] = PackRG32(CompressBlock_ETC2_RGB(BlockBaseColor));
-	OutCompressTexture1_128bit[ThreadId.xy + DestPos1] = CompressBlock_ETC2_RGBA(BlockNormalXRougnnessY, BlockNormalZ);
+	OutCompressTexture0_64bit[GetDestPos(0, ThreadId)] = PackRG32(CompressBlock_ETC2_RGB(BlockBaseColor));
+	OutCompressTexture1_128bit[GetDestPos(1, ThreadId)] = CompressBlock_ETC2_RGBA(BlockNormalXRougnnessY, BlockNormalZ);
 #else
-	OutCompressTexture0_64bit[ThreadId.xy + DestPos0] = PackRG32(CompressBC1Block(BlockBaseColor));
-	OutCompressTexture1_128bit[ThreadId.xy + DestPos1] = CompressBC3Block(BlockNormalXRougnnessY, BlockNormalZ);
+	OutCompressTexture0_64bit[GetDestPos(0, ThreadId)] = PackRG32(CompressBC1Block(BlockBaseColor));
+	OutCompressTexture1_128bit[GetDestPos(1, ThreadId)] = CompressBC3Block(BlockNormalXRougnnessY, BlockNormalZ);
 #endif
-
 }
+#endif // COMPUTESHADER
 
+#if PIXELSHADER
 /** Copy path used when we disable texture compression, because we need to keep the same final channel layout. */
 void CopyBaseColorNormalSpecularPS(
 	in float4 InPosition : SV_POSITION, 
@@ -205,7 +247,9 @@ void CopyBaseColorNormalSpecularPS(
 	OutColor0 = float4(BaseColor, NormalXY.x);
 	OutColor1 = float4(RoughnessSpecularNormalZ, NormalXY.y);
 }
+#endif // PIXELSHADER
 
+#if COMPUTESHADER
 /** Compress base color, normal, roughness and specular to BC3, BC5, BC1 target set. */
 [numthreads(8,8,1)]
 void CompressBaseColorNormalSpecularYCoCgCS(
@@ -221,14 +265,14 @@ void CompressBaseColorNormalSpecularYCoCgCS(
 	float2 SampleUV = (float2(SamplePos) + 0.5f) * TexelUVSize;
 
 	float3 BlockBaseColor[16];
-	ReadBlockRGB(RenderTexture0, TextureSampler0, SampleUV, TexelUVSize, BlockBaseColor);
+	ReadBlockRGB(RenderTexture0, TextureSampler0, GetSourceUV(SampleUV, ThreadId), TexelUVSize, BlockBaseColor);
 
 	float BlockNormalX[16];
 	float BlockNormalY[16];
-	ReadBlockXY(RenderTexture1, TextureSampler1, SampleUV, TexelUVSize, BlockNormalX, BlockNormalY);
+	ReadBlockXY(RenderTexture1, TextureSampler1, GetSourceUV(SampleUV, ThreadId), TexelUVSize, BlockNormalX, BlockNormalY);
 
 	float3 BlockSpecular[16];
-	ReadBlockRGB(RenderTexture2, TextureSampler2, SampleUV, TexelUVSize, BlockSpecular);
+	ReadBlockRGB(RenderTexture2, TextureSampler2, GetSourceUV(SampleUV, ThreadId), TexelUVSize, BlockSpecular);
 
 	for (int i=0; i<16; i++)
 	{
@@ -236,20 +280,22 @@ void CompressBaseColorNormalSpecularYCoCgCS(
 	}
 
 #if ASTC_PROFILE
-	OutCompressTexture0_128bit[ThreadId.xy + DestPos0] = CompressBlock_ASTC_YCoCg(BlockBaseColor);
-	OutCompressTexture1_128bit[ThreadId.xy + DestPos1] = CompressBlock_ETC2_RG(BlockNormalX, BlockNormalY);
-	OutCompressTexture2_64bit[ThreadId.xy + DestPos2] = PackRG32(CompressBlock_ETC2_RGB(BlockSpecular));
+	OutCompressTexture0_128bit[GetDestPos(0, ThreadId)] = CompressBlock_ASTC_YCoCg(BlockBaseColor);
+	OutCompressTexture1_128bit[GetDestPos(1, ThreadId)] = CompressBlock_ETC2_RG(BlockNormalX, BlockNormalY);
+	OutCompressTexture2_64bit[GetDestPos(2, ThreadId)] = PackRG32(CompressBlock_ETC2_RGB(BlockSpecular));
 #elif ETC_PROFILE
-	OutCompressTexture0_128bit[ThreadId.xy + DestPos0] = CompressBlock_ETC2_YCoCg(BlockBaseColor);
-	OutCompressTexture1_128bit[ThreadId.xy + DestPos1] = CompressBlock_ETC2_RG(BlockNormalX, BlockNormalY);
-	OutCompressTexture2_64bit[ThreadId.xy + DestPos2] = PackRG32(CompressBlock_ETC2_RGB(BlockSpecular));
+	OutCompressTexture0_128bit[GetDestPos(0, ThreadId)] = CompressBlock_ETC2_YCoCg(BlockBaseColor);
+	OutCompressTexture1_128bit[GetDestPos(1, ThreadId)] = CompressBlock_ETC2_RG(BlockNormalX, BlockNormalY);
+	OutCompressTexture2_64bit[GetDestPos(2, ThreadId)] = PackRG32(CompressBlock_ETC2_RGB(BlockSpecular));
 #else
-	OutCompressTexture0_128bit[ThreadId.xy + DestPos0] = CompressBC3BlockYCoCg(BlockBaseColor);
-	OutCompressTexture1_128bit[ThreadId.xy + DestPos1] = CompressBC5Block(BlockNormalX, BlockNormalY);
-	OutCompressTexture2_64bit[ThreadId.xy + DestPos2] = PackRG32(CompressBC1Block(BlockSpecular));
+	OutCompressTexture0_128bit[GetDestPos(0, ThreadId)] = CompressBC3BlockYCoCg(BlockBaseColor);
+	OutCompressTexture1_128bit[GetDestPos(1, ThreadId)] = CompressBC5Block(BlockNormalX, BlockNormalY);
+	OutCompressTexture2_64bit[GetDestPos(2, ThreadId)] = PackRG32(CompressBC1Block(BlockSpecular));
 #endif
 }
+#endif // COMPUTESHADER
 
+#if PIXELSHADER
 /** Copy path used when we disable texture compression, because we need to keep the same final channel layout. */
 void CopyBaseColorNormalSpecularYCoCgPS(
 	in float4 InPosition : SV_POSITION, 
@@ -268,7 +314,9 @@ void CopyBaseColorNormalSpecularYCoCgPS(
 	OutColor1 = float4(NormalXY, 0, 1);
 	OutColor2 = float4(RoughnessSpecularNormalZ, 1);
 }
+#endif // PIXELSHADER
 
+#if COMPUTESHADER
 /** Compress base color, normal, roughness, specular and mask to BC3, BC5, BC3 target set. */
 [numthreads(8,8,1)]
 void CompressBaseColorNormalSpecularMaskYCoCgCS(
@@ -282,15 +330,15 @@ void CompressBaseColorNormalSpecularMaskYCoCgCS(
 	float2 SampleUV = (float2(SamplePos) + 0.5f) * TexelUVSize;
 
 	float3 BlockBaseColor[16];
-	ReadBlockRGB(RenderTexture0, TextureSampler0, SampleUV, TexelUVSize, BlockBaseColor);
+	ReadBlockRGB(RenderTexture0, TextureSampler0, GetSourceUV(SampleUV, ThreadId), TexelUVSize, BlockBaseColor);
 
 	float BlockNormalX[16];
 	float BlockNormalY[16];
 	float BlockMask[16];
-	ReadBlockXYA(RenderTexture1, TextureSampler1, SampleUV, TexelUVSize, BlockNormalX, BlockNormalY, BlockMask);
+	ReadBlockXYA(RenderTexture1, TextureSampler1, GetSourceUV(SampleUV, ThreadId), TexelUVSize, BlockNormalX, BlockNormalY, BlockMask);
 
 	float3 BlockSpecular[16];
-	ReadBlockRGB(RenderTexture2, TextureSampler2, SampleUV, TexelUVSize, BlockSpecular);
+	ReadBlockRGB(RenderTexture2, TextureSampler2, GetSourceUV(SampleUV, ThreadId), TexelUVSize, BlockSpecular);
 
 	for (int i=0; i<16; i++)
 	{
@@ -298,20 +346,22 @@ void CompressBaseColorNormalSpecularMaskYCoCgCS(
 	}
 
 #if ASTC_PROFILE
-	OutCompressTexture0_128bit[ThreadId.xy + DestPos0] = CompressBlock_ASTC_YCoCg(BlockBaseColor);
-	OutCompressTexture1_128bit[ThreadId.xy + DestPos1] = CompressBlock_ETC2_RG(BlockNormalX, BlockNormalY);
-	OutCompressTexture2_128bit[ThreadId.xy + DestPos2] = CompressBlock_ASTC_RGBA(BlockSpecular, BlockMask);
+	OutCompressTexture0_128bit[GetDestPos(0, ThreadId)] = CompressBlock_ASTC_YCoCg(BlockBaseColor);
+	OutCompressTexture1_128bit[GetDestPos(1, ThreadId)] = CompressBlock_ETC2_RG(BlockNormalX, BlockNormalY);
+	OutCompressTexture2_128bit[GetDestPos(2, ThreadId)] = CompressBlock_ASTC_RGBA(BlockSpecular, BlockMask);
 #elif ETC_PROFILE
-	OutCompressTexture0_128bit[ThreadId.xy + DestPos0] = CompressBlock_ETC2_YCoCg(BlockBaseColor);
-	OutCompressTexture1_128bit[ThreadId.xy + DestPos1] = CompressBlock_ETC2_RG(BlockNormalX, BlockNormalY);
-	OutCompressTexture2_128bit[ThreadId.xy + DestPos2] = CompressBlock_ETC2_RGBA(BlockSpecular, BlockMask);
+	OutCompressTexture0_128bit[GetDestPos(0, ThreadId)] = CompressBlock_ETC2_YCoCg(BlockBaseColor);
+	OutCompressTexture1_128bit[GetDestPos(1, ThreadId)] = CompressBlock_ETC2_RG(BlockNormalX, BlockNormalY);
+	OutCompressTexture2_128bit[GetDestPos(2, ThreadId)] = CompressBlock_ETC2_RGBA(BlockSpecular, BlockMask);
 #else
-	OutCompressTexture0_128bit[ThreadId.xy + DestPos0] = CompressBC3BlockYCoCg(BlockBaseColor);
-	OutCompressTexture1_128bit[ThreadId.xy + DestPos1] = CompressBC5Block(BlockNormalX, BlockNormalY);
-	OutCompressTexture2_128bit[ThreadId.xy + DestPos2] = CompressBC3Block(BlockSpecular, BlockMask);
+	OutCompressTexture0_128bit[GetDestPos(0, ThreadId)] = CompressBC3BlockYCoCg(BlockBaseColor);
+	OutCompressTexture1_128bit[GetDestPos(1, ThreadId)] = CompressBC5Block(BlockNormalX, BlockNormalY);
+	OutCompressTexture2_128bit[GetDestPos(2, ThreadId)] = CompressBC3Block(BlockSpecular, BlockMask);
 #endif
 }
+#endif // COMPUTESHADER
 
+#if PIXELSHADER
 /** Copy path used when we disable texture compression, because we need to keep the same final channel layout. */
 void CopyBaseColorNormalSpecularMaskYCoCgPS(
 	in float4 InPosition : SV_POSITION, 
@@ -330,7 +380,9 @@ void CopyBaseColorNormalSpecularMaskYCoCgPS(
 	OutColor1 = float4(NormalXYMask.xy, 0, 1);
 	OutColor2 = float4(RoughnessSpecularNormalZ, NormalXYMask.z);
 }
+#endif // PIXELSHADER
 
+#if COMPUTESHADER
 /** Compress 4 channel mask to BC3. */
 [numthreads(8,8,1)]
 void CompressMask4CS(
@@ -346,20 +398,22 @@ void CompressMask4CS(
 	float2 SampleUV = (float2(SamplePos) + 0.5f) * TexelUVSize;
 
 	float3 BlockXYZ[16];
-	ReadBlockRGB(RenderTexture0, TextureSampler0, SampleUV, TexelUVSize, BlockXYZ);
+	ReadBlockRGB(RenderTexture0, TextureSampler0, GetSourceUV(SampleUV, ThreadId), TexelUVSize, BlockXYZ);
 
 	float BlockW[16];
-	ReadBlockX(RenderTexture1, TextureSampler0, SampleUV, TexelUVSize, BlockW);
+	ReadBlockX(RenderTexture1, TextureSampler0, GetSourceUV(SampleUV, ThreadId), TexelUVSize, BlockW);
 
 #if ASTC_PROFILE
-	OutCompressTexture0_128bit[ThreadId.xy + DestPos0] = CompressBlock_ASTC_RGBA(BlockXYZ, BlockW);
+	OutCompressTexture0_128bit[GetDestPos(0, ThreadId)] = CompressBlock_ASTC_RGBA(BlockXYZ, BlockW);
 #elif ETC_PROFILE
-	OutCompressTexture0_128bit[ThreadId.xy + DestPos0] = CompressBlock_ETC2_RGBA(BlockXYZ, BlockW);
+	OutCompressTexture0_128bit[GetDestPos(0, ThreadId)] = CompressBlock_ETC2_RGBA(BlockXYZ, BlockW);
 #else
-	OutCompressTexture0_128bit[ThreadId.xy + DestPos0] = CompressBC3Block(BlockXYZ, BlockW);
+	OutCompressTexture0_128bit[GetDestPos(0, ThreadId)] = CompressBC3Block(BlockXYZ, BlockW);
 #endif
 }
+#endif // COMPUTESHADER
 
+#if PIXELSHADER
 /** Copy path used when we disable texture compression, because we need to keep the same final channel layout. */
 void CopyMask4PS(
 	in float4 InPosition : SV_POSITION, 
@@ -371,7 +425,9 @@ void CopyMask4PS(
 	
 	OutColor0 = float4(MaskXYZ, MaskW);
 }
+#endif // PIXELSHADER
 
+#if COMPUTESHADER
 /** Compress displacement to a single BC4 target. */
 [numthreads(8, 8, 1)]
 void CompressDisplacementCS(
@@ -387,16 +443,18 @@ void CompressDisplacementCS(
 	float2 SampleUV = (float2(SamplePos) + 0.5f) * TexelUVSize;
 
 	float Block[16];
-	ReadBlockX(RenderTexture0, TextureSampler0, SampleUV, TexelUVSize, Block);
+	ReadBlockX(RenderTexture0, TextureSampler0, GetSourceUV(SampleUV, ThreadId), TexelUVSize, Block);
 #if ASTC_PROFILE
-	OutCompressTexture0_64bit[ThreadId.xy + DestPos0] = PackRG32(CompressBlock_ETC2_Alpha(Block));
+	OutCompressTexture0_64bit[GetDestPos(0, ThreadId)] = PackRG32(CompressBlock_ETC2_Alpha(Block));
 #elif ETC_PROFILE
-	OutCompressTexture0_64bit[ThreadId.xy + DestPos0] = PackRG32(CompressBlock_ETC2_Alpha(Block));
+	OutCompressTexture0_64bit[GetDestPos(0, ThreadId)] = PackRG32(CompressBlock_ETC2_Alpha(Block));
 #else
-	OutCompressTexture0_64bit[ThreadId.xy + DestPos0] = PackRG32(CompressBC4Block(Block));
+	OutCompressTexture0_64bit[GetDestPos(0, ThreadId)] = PackRG32(CompressBC4Block(Block));
 #endif
 }
+#endif // COMPUTESHADER
 
+#if PIXELSHADER
 /** Copy world height or displacement (without compression)
   * This path is used when want to copy the world height or displacement data into a thumbnail. */
 void CopyWorldHeightPS(
@@ -407,3 +465,4 @@ void CopyWorldHeightPS(
 	float WorldHeight = RenderTexture0.SampleLevel(TextureSampler0, InTexCoord, 0).r;
 	OutColor0 = float4(WorldHeight.xxx, 1.f);
 }
+#endif // PIXELSHADER
diff --git a/Engine/Source/Editor/VirtualTexturingEditor/Private/RuntimeVirtualTextureBuildStreamingMips.cpp b/Engine/Source/Editor/VirtualTexturingEditor/Private/RuntimeVirtualTextureBuildStreamingMips.cpp
--- a/Engine/Source/Editor/VirtualTexturingEditor/Private/RuntimeVirtualTextureBuildStreamingMips.cpp
+++ b/Engine/Source/Editor/VirtualTexturingEditor/Private/RuntimeVirtualTextureBuildStreamingMips.cpp
@@ -6,6 +6,7 @@
 #include "Components/RuntimeVirtualTextureComponent.h"
 #include "ComponentRecreateRenderStateContext.h"
 #include "ContentStreaming.h"
+#include "EngineModule.h"
 #include "Engine/TextureRenderTarget2D.h"
 #include "Engine/World.h"
 #include "Materials/Material.h"
@@ -259,14 +260,14 @@ namespace RuntimeVirtualTexture
 					CustomMaterialData, FixedColor,
 					&RenderTargetsCurrentState](FRHICommandListImmediate& RHICmdList)
 				{
-					const FBox2D TileBox(FVector2D(0, 0), FVector2D(TileSize, TileSize));
 					const FIntRect TileRect(0, 0, TileSize, TileSize);
 
 					{
 						FRDGBuilder GraphBuilder(RHICmdList);
+						FScenePrimitiveRenderingContextScopeHelper RenderingScope(GetRendererModule().BeginScenePrimitiveRendering(GraphBuilder, *Scene));
 
 						RuntimeVirtualTexture::FRenderPageBatchDesc Desc;
-						Desc.Scene = Scene->GetRenderScene();
+						Desc.SceneRenderer = RenderingScope.ScenePrimitiveRenderingContext->GetSceneRenderer();
 						Desc.RuntimeVirtualTextureId = RuntimeVirtualTextureId;
 						Desc.UVToWorld = Transform;
 						Desc.WorldBounds = Bounds;
@@ -283,12 +284,12 @@ namespace RuntimeVirtualTexture
 							Desc.Targets[Layer].TextureAccessBefore = RenderTargetsCurrentState;
 							// Leave the render targets in CopySrc state since we'll read from them right after, to avoid a useless transition :
 							Desc.Targets[Layer].TextureAccessAfter = ERHIAccess::CopySrc;
-							Desc.PageDescs[0].DestBox[Layer] = TileBox;
+							Desc.PageDescs[0].DestRect[Layer] = TileRect;
 						}
 						Desc.PageDescs[0].UVRange = UVRange;
 						Desc.PageDescs[0].vLevel = IntCastChecked<uint8>(RenderLevel);
 
-						RuntimeVirtualTexture::RenderPagesStandAlone(GraphBuilder, Desc);
+						RuntimeVirtualTexture::RenderPages(GraphBuilder, Desc);
 
 						GraphBuilder.Execute();
 					}
diff --git a/Engine/Source/Editor/VirtualTexturingEditor/Private/RuntimeVirtualTextureThumbnailRenderer.cpp b/Engine/Source/Editor/VirtualTexturingEditor/Private/RuntimeVirtualTextureThumbnailRenderer.cpp
--- a/Engine/Source/Editor/VirtualTexturingEditor/Private/RuntimeVirtualTextureThumbnailRenderer.cpp
+++ b/Engine/Source/Editor/VirtualTexturingEditor/Private/RuntimeVirtualTextureThumbnailRenderer.cpp
@@ -3,6 +3,7 @@
 #include "RuntimeVirtualTextureThumbnailRenderer.h"
 
 #include "Components/RuntimeVirtualTextureComponent.h"
+#include "EngineModule.h"
 #include "Engine/World.h"
 #include "Materials/MaterialRenderProxy.h"
 #include "MaterialShared.h"
@@ -87,7 +88,7 @@ void URuntimeVirtualTextureThumbnailRenderer::Draw(UObject* Object, int32 X, int
 		World->SendAllEndOfFrameUpdates();
 	}
 
-	const FBox2D DestBox = FBox2D(FVector2D(X, Y), FVector2D(Width, Height));
+	const FIntRect DestRect = FIntRect(X, Y, X + Width, Y + Height);
 	const FTransform Transform = RuntimeVirtualTextureComponent->GetComponentTransform();
 	const FBox Bounds = RuntimeVirtualTextureComponent->Bounds.GetBox();
 	const ERuntimeVirtualTextureMaterialType MaterialType = RuntimeVirtualTexture->GetMaterialType();
@@ -102,14 +103,15 @@ void URuntimeVirtualTextureThumbnailRenderer::Draw(UObject* Object, int32 X, int
 	UE::RenderCommandPipe::FSyncScope SyncScope;
 
 	ENQUEUE_RENDER_COMMAND(BakeStreamingTextureTileCommand)(
-		[Scene, RuntimeVirtualTextureId, MaterialType, RenderTarget, DestBox, Transform, Bounds, MaxLevel, CustomMaterialData](FRHICommandListImmediate& RHICmdList)
+		[Scene, RuntimeVirtualTextureId, MaterialType, RenderTarget, DestRect, Transform, Bounds, MaxLevel, CustomMaterialData](FRHICommandListImmediate& RHICmdList)
 	{
 		FMaterialRenderProxy::UpdateDeferredCachedUniformExpressions();
 
 		FRDGBuilder GraphBuilder(RHICmdList);
+		FScenePrimitiveRenderingContextScopeHelper RenderingScope(GetRendererModule().BeginScenePrimitiveRendering(GraphBuilder, *Scene));
 
 		RuntimeVirtualTexture::FRenderPageBatchDesc Desc;
-		Desc.Scene = Scene->GetRenderScene();
+		Desc.SceneRenderer = RenderingScope.ScenePrimitiveRenderingContext->GetSceneRenderer();
 		Desc.RuntimeVirtualTextureId = RuntimeVirtualTextureId;
 		Desc.UVToWorld = Transform;
 		Desc.WorldBounds = Bounds;
@@ -124,11 +126,11 @@ void URuntimeVirtualTextureThumbnailRenderer::Draw(UObject* Object, int32 X, int
 		// The target is in RTV state and should be left in RTV state (RenderPage needs the input/output state) : 
 		Desc.Targets[0].TextureAccessBefore = ERHIAccess::RTV;
 		Desc.Targets[0].TextureAccessAfter = ERHIAccess::RTV;
-		Desc.PageDescs[0].DestBox[0] = DestBox;
+		Desc.PageDescs[0].DestRect[0] = DestRect;
 		Desc.PageDescs[0].UVRange = FBox2D(FVector2D(0, 0), FVector2D(1, 1));
 		Desc.PageDescs[0].vLevel = IntCastChecked<uint8>(MaxLevel);
 
-		RuntimeVirtualTexture::RenderPagesStandAlone(GraphBuilder, Desc);
+		RuntimeVirtualTexture::RenderPages(GraphBuilder, Desc);
 
 		GraphBuilder.Execute();
 	});
diff --git a/Engine/Source/Runtime/RenderCore/Public/VirtualTexturing.h b/Engine/Source/Runtime/RenderCore/Public/VirtualTexturing.h
--- a/Engine/Source/Runtime/RenderCore/Public/VirtualTexturing.h
+++ b/Engine/Source/Runtime/RenderCore/Public/VirtualTexturing.h
@@ -24,6 +24,7 @@ class FRHICommandListImmediate;
 class FRHIShaderResourceView;
 class FRHITexture;
 class FRHIUnorderedAccessView;
+class ISceneRenderer;
 namespace ERHIFeatureLevel { enum Type : int; }
 enum class EVTProducerPriority : uint8;
 enum class EVTInvalidatePriority : uint8;
@@ -189,9 +190,22 @@ struct FVTProducerDescription
 
 typedef void (FVTProducerDestroyedFunction)(const FVirtualTextureProducerHandle& InHandle, void* Baton);
 
+/** 
+ * A Finalizer is an object that does the final work of filling the physical textures. 
+ * The work for all finalizers is scheduled at a specific point in the frame where we can write to the physical texture without hazards, and
+ * where the virtual texture page tables are also updated.
+ * The finalizer work may be split into two parts.
+ * RenderFinalize() does not write to the physical texture, but can read from virtual textures. This is a phase that allows any 
+ * page rendering which may need to sample virtual textures. Runtime virtual textures and material systems require this.
+ * Finalise() must write to the physical textures, but cannot sample from them.
+ * All finalizers need to implement Finalize() but only ones that need to sample virtual textures need to implement RenderFinalize().
+ */
 class IVirtualTextureFinalizer
 {
 public:
+	/** Finalize work that has read only access to the virtual texture physical pools. */
+	virtual void RenderFinalize(FRDGBuilder& GraphBuilder, ISceneRenderer* SceneRenderingContext) {}
+	/** Finalize work that has write only access to the virtual texture physical pools. */
 	virtual void Finalize(FRDGBuilder& GraphBuilder) = 0;
 };
 
diff --git a/Engine/Source/Runtime/Renderer/Private/DeferredShadingRenderer.cpp b/Engine/Source/Runtime/Renderer/Private/DeferredShadingRenderer.cpp
--- a/Engine/Source/Runtime/Renderer/Private/DeferredShadingRenderer.cpp
+++ b/Engine/Source/Runtime/Renderer/Private/DeferredShadingRenderer.cpp
@@ -1499,7 +1499,7 @@ void FDeferredShadingSceneRenderer::Render(FRDGBuilder& GraphBuilder, const FSce
 		FVirtualTextureUpdateSettings Settings;
 		Settings.EnableThrottling(!ViewFamily.bOverrideVirtualTextureThrottle);
 
-		VirtualTextureUpdater = FVirtualTextureSystem::Get().BeginUpdate(GraphBuilder, FeatureLevel, Scene, Settings);
+		VirtualTextureUpdater = FVirtualTextureSystem::Get().BeginUpdate(GraphBuilder, FeatureLevel, this, Settings);
 		VirtualTextureFeedbackBegin(GraphBuilder, Views, GetActiveSceneTexturesConfig().Extent);
 	}
 
@@ -2008,7 +2008,7 @@ void FDeferredShadingSceneRenderer::Render(FRDGBuilder& GraphBuilder, const FSce
 	// Virtual texturing isn't needed for depth prepass
 	if (bUseVirtualTexturing && RendererOutput != ERendererOutput::DepthPrepassOnly)
 	{
-		FVirtualTextureSystem::Get().FinalizeRequests(GraphBuilder);
+		FVirtualTextureSystem::Get().FinalizeRequests(GraphBuilder, this);
 	}
 
 	{
diff --git a/Engine/Source/Runtime/Renderer/Private/MobileShadingRenderer.cpp b/Engine/Source/Runtime/Renderer/Private/MobileShadingRenderer.cpp
--- a/Engine/Source/Runtime/Renderer/Private/MobileShadingRenderer.cpp
+++ b/Engine/Source/Runtime/Renderer/Private/MobileShadingRenderer.cpp
@@ -1106,7 +1106,7 @@ void FMobileSceneRenderer::Render(FRDGBuilder& GraphBuilder, const FSceneRenderU
 		FVirtualTextureUpdateSettings Settings;
 		Settings.EnableThrottling(!ViewFamily.bOverrideVirtualTextureThrottle);
 
-		VirtualTextureUpdater = FVirtualTextureSystem::Get().BeginUpdate(GraphBuilder, FeatureLevel, Scene, Settings);
+		VirtualTextureUpdater = FVirtualTextureSystem::Get().BeginUpdate(GraphBuilder, FeatureLevel, this, Settings);
 		VirtualTextureFeedbackBegin(GraphBuilder, Views, SceneTexturesConfig.Extent);
 	}
 
@@ -1187,7 +1187,7 @@ void FMobileSceneRenderer::Render(FRDGBuilder& GraphBuilder, const FSceneRenderU
 		if (bUseVirtualTexturing)
 		{
 			FVirtualTextureSystem::Get().EndUpdate(GraphBuilder, MoveTemp(VirtualTextureUpdater), FeatureLevel);
-			FVirtualTextureSystem::Get().FinalizeRequests(GraphBuilder);
+			FVirtualTextureSystem::Get().FinalizeRequests(GraphBuilder, this);
 		}
 
 		if (bDeferredShading ||
diff --git a/Engine/Source/Runtime/Renderer/Private/Renderer.cpp b/Engine/Source/Runtime/Renderer/Private/Renderer.cpp
--- a/Engine/Source/Runtime/Renderer/Private/Renderer.cpp
+++ b/Engine/Source/Runtime/Renderer/Private/Renderer.cpp
@@ -359,16 +359,11 @@ void FRendererModule::DrawTileMesh(FCanvasRenderContext& RenderContext, FMeshPas
 		const auto FeatureLevel = View.GetFeatureLevel();
 		const EShadingPath ShadingPath = GetFeatureLevelShadingPath(FeatureLevel);
 
-		FScene* Scene = nullptr;
-		if (ViewFamily->Scene)
-		{
-			Scene = ViewFamily->Scene->GetRenderScene();
-		}
+		FScene* Scene = ViewFamily->Scene ? ViewFamily->Scene->GetRenderScene() : nullptr;
 
 		Mesh.MaterialRenderProxy->UpdateUniformExpressionCacheIfNeeded(FeatureLevel);
 		FMaterialRenderProxy::UpdateDeferredCachedUniformExpressions();
 
-
 		FRDGBuilder& GraphBuilder = RenderContext.GraphBuilder;
 		
 		FSceneUniformBuffer& SceneUniforms = *CreateSinglePrimitiveSceneUniformBuffer(GraphBuilder, FeatureLevel, Mesh);
@@ -378,14 +373,15 @@ void FRendererModule::DrawTileMesh(FCanvasRenderContext& RenderContext, FMeshPas
 			FRDGSystemTextures::Create(GraphBuilder);
 		}
 
-		// Materials sampling VTs need FVirtualTextureSystem to be updated before being rendered
+		// Materials sampling VTs need FVirtualTextureSystem to be updated before being rendered.
+		// Note that we any VTs dependent on having a SceneRenderer (eg RVT) cannot be warmed up here.
 		const FMaterial& MeshMaterial = Mesh.MaterialRenderProxy->GetIncompleteMaterialWithFallback(FeatureLevel);
 		const bool bUseVirtualTexturing = UseVirtualTexturing(View.GetShaderPlatform()) && !MeshMaterial.GetUniformVirtualTextureExpressions().IsEmpty();
 		if (bUseVirtualTexturing)
 		{
 			FVirtualTextureUpdateSettings Settings;
 			Settings.EnableThrottling(false);
-			FVirtualTextureSystem::Get().Update(GraphBuilder, FeatureLevel, Scene, Settings);
+			FVirtualTextureSystem::Get().Update(GraphBuilder, FeatureLevel, nullptr, Settings);
 
 			VirtualTextureFeedbackBegin(GraphBuilder, TArrayView<const FViewInfo>(&View, 1), RenderContext.GetViewportRect().Size());
 		}
diff --git a/Engine/Source/Runtime/Renderer/Private/SceneRendering.h b/Engine/Source/Runtime/Renderer/Private/SceneRendering.h
--- a/Engine/Source/Runtime/Renderer/Private/SceneRendering.h
+++ b/Engine/Source/Runtime/Renderer/Private/SceneRendering.h
@@ -2004,6 +2004,7 @@ public:
 	static FSceneRendererBase* GetActiveInstance(FRDGBuilder& GraphBuilder);
 
 	// ISceneRenderer interface
+	FScene* GetScene() final override { return Scene; }
 	const FSceneUniformBuffer& GetSceneUniforms() const final override { return SceneUniforms; }
 	FSceneUniformBuffer& GetSceneUniforms() final override { return SceneUniforms; }
 
diff --git a/Engine/Source/Runtime/Renderer/Private/VT/RuntimeVirtualTextureProducer.cpp b/Engine/Source/Runtime/Renderer/Private/VT/RuntimeVirtualTextureProducer.cpp
--- a/Engine/Source/Runtime/Renderer/Private/VT/RuntimeVirtualTextureProducer.cpp
+++ b/Engine/Source/Runtime/Renderer/Private/VT/RuntimeVirtualTextureProducer.cpp
@@ -38,13 +38,12 @@ void FRuntimeVirtualTextureFinalizer::AddTile(FTileEntry& Tile)
 	Tiles.Add(Tile);
 }
 
-void FRuntimeVirtualTextureFinalizer::Finalize(FRDGBuilder& GraphBuilder)
+void FRuntimeVirtualTextureFinalizer::RenderFinalize(FRDGBuilder& GraphBuilder, ISceneRenderer* SceneRenderer)
 {
-	RDG_EVENT_SCOPE(GraphBuilder, "RuntimeVirtualTextureFinalize");
 	RDG_GPU_MASK_SCOPE(GraphBuilder, FRHIGPUMask::All());
 
 	RuntimeVirtualTexture::FRenderPageBatchDesc RenderPageBatchDesc;
-	RenderPageBatchDesc.Scene = Scene->GetRenderScene();
+	RenderPageBatchDesc.SceneRenderer = SceneRenderer;
 	RenderPageBatchDesc.RuntimeVirtualTextureId = RuntimeVirtualTextureId;
 	RenderPageBatchDesc.UVToWorld = UVToWorld;
 	RenderPageBatchDesc.WorldBounds = WorldBounds;
@@ -71,10 +70,10 @@ void FRuntimeVirtualTextureFinalizer::Finalize(FRDGBuilder& GraphBuilder)
 			bBreakBatchForTextures |= (RenderPageBatchDesc.Targets[LayerIndex].Texture != Entry.Targets[LayerIndex].TextureRHI);
 		}
 
-		if (BatchSize == RuntimeVirtualTexture::EMaxRenderPageBatch || bBreakBatchForTextures)
+		if (BatchSize == RuntimeVirtualTexture::MaxRenderPageBatch || bBreakBatchForTextures)
 		{
 			RenderPageBatchDesc.NumPageDescs = BatchSize;
-			RuntimeVirtualTexture::RenderPages(GraphBuilder, RenderPageBatchDesc);
+			Batches.Add(RuntimeVirtualTexture::InitPageBatch(GraphBuilder, RenderPageBatchDesc));
 			BatchSize = 0;
 		}
 
@@ -105,19 +104,32 @@ void FRuntimeVirtualTextureFinalizer::Finalize(FRDGBuilder& GraphBuilder)
 		const int32 TileSize = Desc.TileSize + 2 * Desc.TileBorderSize;
 		for (int LayerIndex = 0; LayerIndex < RuntimeVirtualTexture::MaxTextureLayers; ++LayerIndex)
 		{
-			const FVector2D DestinationBoxStart0(Entry.Targets[LayerIndex].pPageLocation.X * TileSize, Entry.Targets[LayerIndex].pPageLocation.Y * TileSize);
-			RenderPageDesc.DestBox[LayerIndex] = FBox2D(DestinationBoxStart0, DestinationBoxStart0 + FVector2D(TileSize, TileSize));
+			const FIntPoint DestinationRectStart0(Entry.Targets[LayerIndex].pPageLocation.X * TileSize, Entry.Targets[LayerIndex].pPageLocation.Y * TileSize);
+			RenderPageDesc.DestRect[LayerIndex] = FIntRect(DestinationRectStart0, DestinationRectStart0 + FIntPoint(TileSize, TileSize));
 		}
-
 	}
 
 	if (BatchSize > 0)
 	{
 		RenderPageBatchDesc.NumPageDescs = BatchSize;
-		RuntimeVirtualTexture::RenderPages(GraphBuilder, RenderPageBatchDesc);
+		Batches.Add(RuntimeVirtualTexture::InitPageBatch(GraphBuilder, RenderPageBatchDesc));
+	}
+
+	for (RuntimeVirtualTexture::FBatchRenderContext const* Batch : Batches)
+	{
+		RuntimeVirtualTexture::RenderPageBatch(GraphBuilder, *Batch);
+	}
+}
+
+void FRuntimeVirtualTextureFinalizer::Finalize(FRDGBuilder& GraphBuilder)
+{
+	for (RuntimeVirtualTexture::FBatchRenderContext const* Batch : Batches)
+	{
+		RuntimeVirtualTexture::FinalizePageBatch(GraphBuilder, *Batch);
 	}
 
 	Tiles.SetNumUnsafeInternal(0);
+	Batches.SetNumUnsafeInternal(0);
 }
 
 FRuntimeVirtualTextureProducer::FRuntimeVirtualTextureProducer(
diff --git a/Engine/Source/Runtime/Renderer/Private/VT/RuntimeVirtualTextureProducer.h b/Engine/Source/Runtime/Renderer/Private/VT/RuntimeVirtualTextureProducer.h
--- a/Engine/Source/Runtime/Renderer/Private/VT/RuntimeVirtualTextureProducer.h
+++ b/Engine/Source/Runtime/Renderer/Private/VT/RuntimeVirtualTextureProducer.h
@@ -10,6 +10,7 @@
 enum class ERuntimeVirtualTextureMaterialType : uint8;
 class FRHITexture;
 class FSceneInterface;
+namespace RuntimeVirtualTexture { class FBatchRenderContext; }
 
 /** IVirtualTextureFinalizer implementation that renders the virtual texture pages on demand. */
 class FRuntimeVirtualTextureFinalizer : public IVirtualTextureFinalizer
@@ -42,6 +43,7 @@ public:
 	void AddTile(FTileEntry& Tile);
 
 	//~ Begin IVirtualTextureFinalizer Interface.
+	virtual void RenderFinalize(FRDGBuilder& GraphBuilder, ISceneRenderer* SceneRenderer) override;
 	virtual void Finalize(FRDGBuilder& GraphBuilder) override;
 	//~ End IVirtualTextureFinalizer Interface.
 
@@ -64,6 +66,8 @@ private:
 	FVector4f CustomMaterialData;
 	/** Array of tiles in the queue to finalize. */
 	TArray<FTileEntry> Tiles;
+	/** Array of batch render contexts used during finalize. */
+	TArray<RuntimeVirtualTexture::FBatchRenderContext const*> Batches;
 };
 
 /** IVirtualTexture implementation that is handling runtime rendered page data requests. */
diff --git a/Engine/Source/Runtime/Renderer/Private/VT/RuntimeVirtualTextureRender.cpp b/Engine/Source/Runtime/Renderer/Private/VT/RuntimeVirtualTextureRender.cpp
--- a/Engine/Source/Runtime/Renderer/Private/VT/RuntimeVirtualTextureRender.cpp
+++ b/Engine/Source/Runtime/Renderer/Private/VT/RuntimeVirtualTextureRender.cpp
@@ -10,23 +10,23 @@
 #include "GPUScene.h"
 #include "MaterialShader.h"
 #include "MeshPassProcessor.h"
+#include "MeshPassProcessor.inl"
+#include "PostProcess/SceneRenderTargets.h"
 #include "ProfilingDebugging/CsvProfiler.h"
+#include "RenderCaptureInterface.h"
 #include "RenderGraphBuilder.h"
 #include "RenderGraphUtils.h"
 #include "RenderUtils.h"
-#include "ShaderPlatformCachedIniValue.h"
+#include "RHIResourceUtils.h"
 #include "ScenePrivate.h"
-#include "PostProcess/SceneRenderTargets.h"
+#include "SceneRendering.h"
+#include "ShaderPlatformCachedIniValue.h"
 #include "ShaderBaseClasses.h"
+#include "SimpleMeshDrawCommandPass.h"
+#include "StaticMeshBatch.h"
 #include "VT/RuntimeVirtualTexture.h"
 #include "VT/RuntimeVirtualTextureSceneExtension.h"
 #include "VT/RuntimeVirtualTextureSceneProxy.h"
-#include "MeshPassProcessor.inl"
-#include "RenderCaptureInterface.h"
-#include "SimpleMeshDrawCommandPass.h"
-#include "StaticMeshBatch.h"
-#include "SceneRendering.h"
-#include "RHIResourceUtils.h"
 
 CSV_DECLARE_CATEGORY_EXTERN(VirtualTexturing);
 
@@ -557,9 +557,6 @@ namespace RuntimeVirtualTexture
 			SHADER_PARAMETER_RDG_UNIFORM_BUFFER(FSceneUniformParameters, Scene)
 			SHADER_PARAMETER_RDG_UNIFORM_BUFFER(FRuntimeVirtualTexturePassParameters, RuntimeVirtualTexturePassParameters)
 			SHADER_PARAMETER_STRUCT_INCLUDE(FInstanceCullingDrawParams, InstanceCullingDrawParams)
-			RDG_TEXTURE_ACCESS(PhysicalTexture0, ERHIAccess::SRVGraphics)
-			RDG_TEXTURE_ACCESS(PhysicalTexture1, ERHIAccess::SRVGraphics)
-			RDG_TEXTURE_ACCESS(PhysicalTexture2, ERHIAccess::SRVGraphics)
 			RENDER_TARGET_BINDING_SLOTS()
 		END_SHADER_PARAMETER_STRUCT()
 
@@ -779,6 +776,7 @@ namespace RuntimeVirtualTexture
 			default: return TGetBlendStateFromAttributeMask<31>();
 			}
 		}
+
 	public:
 		static bool ShouldCompilePermutation(const FMeshMaterialShaderPermutationParameters& Parameters)
 		{
@@ -965,24 +963,64 @@ namespace RuntimeVirtualTexture
 			}
 		}
 
+		/** Initializer description for the graph setup. */
 		struct FInitDesc
 		{
 			ERHIFeatureLevel::Type FeatureLevel;
 			ERuntimeVirtualTextureMaterialType MaterialType;
 			FIntPoint TextureSize;
+			int32 PageCount = 1;
 			EPixelFormat OutputFormat0 = PF_Unknown;
 			TArray<TRefCountPtr<IPooledRenderTarget>> OutputTargets;
 			bool bClearTextures = false;
 			bool bIsThumbnails = false;
 
-			FInitDesc(ERHIFeatureLevel::Type InFeatureLevel, ERuntimeVirtualTextureMaterialType InMaterialType, FIntPoint InTextureSize)
-				: FeatureLevel(InFeatureLevel)
-				, MaterialType(InMaterialType)
-				, TextureSize(InTextureSize)
-			{}
+			/** Initialize from a page batch description. */
+			FInitDesc(FRenderPageBatchDesc const& InDesc)
+			{
+				check(InDesc.SceneRenderer != nullptr && InDesc.SceneRenderer->GetScene() != nullptr);
+				FeatureLevel = InDesc.SceneRenderer->GetScene()->GetFeatureLevel();
+				MaterialType = InDesc.MaterialType;
+				PageCount = InDesc.NumPageDescs;
+				TextureSize = InDesc.PageDescs[0].DestRect[0].Size();
+				OutputFormat0 = InDesc.Targets[0].Texture != nullptr ? InDesc.Targets[0].Texture->GetFormat() : PF_Unknown;
+				OutputTargets.Add(InDesc.Targets[0].PooledRenderTarget);
+				OutputTargets.Add(InDesc.Targets[1].PooledRenderTarget);
+				OutputTargets.Add(InDesc.Targets[2].PooledRenderTarget);
+				bClearTextures = InDesc.bClearTextures;
+				bIsThumbnails = InDesc.bIsThumbnails;
+			}
 		};
 
-		FRenderGraphSetup(FRDGBuilder& GraphBuilder, FInitDesc const& Desc)
+		/** CreateTextureDesc() creates a texture2Darray if we have page batch size > 1 or a simple texture2D otherwise. */
+		static FRDGTextureDesc CreateTextureDesc(FIntPoint Size, EPixelFormat Format, FClearValueBinding ClearValue, ETextureCreateFlags Flags, uint16 ArraySize)
+		{
+			if (ArraySize > 1)
+			{
+				return FRDGTextureDesc::Create2DArray(Size, Format, ClearValue, Flags | TexCreate_TargetArraySlicesIndependently, ArraySize);
+			}
+			else
+			{
+				return FRDGTextureDesc::Create2D(Size, Format, ClearValue, Flags);
+			}
+		}
+
+		/** CreateTextureSRV() createss an SRV for a single slice if Texture is a texture array. */
+		static FRDGTextureSRVRef CreateTextureSRV(FRDGBuilder& GraphBuilder, FRDGTextureRef Texture, int32 ArraySlice)
+		{
+			if (Texture == nullptr)
+			{
+				return nullptr;
+			}
+			if (ArraySlice >= 0)
+			{
+				return GraphBuilder.CreateSRV(FRDGTextureSRVDesc::CreateForSlice(Texture, ArraySlice));
+			}
+			return GraphBuilder.CreateSRV(Texture);
+		}
+
+		/** Initialize the graph setup. */
+		void Init(FRDGBuilder& GraphBuilder, FInitDesc const& Desc)
 		{
 			bRenderPass = Desc.OutputFormat0 != PF_Unknown;
 			bCopyThumbnailPass = bRenderPass && Desc.bIsThumbnails;
@@ -1007,12 +1045,14 @@ namespace RuntimeVirtualTexture
 			const EPixelFormat Compressed64BitFormat = UseRGBA16(GMaxRHIShaderPlatform) ? PF_R16G16B16A16_UINT : PF_R32G32_UINT;
 			const EPixelFormat Compressed128BitFormat = PF_R32G32B32A32_UINT;
 
+			const int32 PageCount = Desc.PageCount;
+
 			switch (Desc.MaterialType)
 			{
 			case ERuntimeVirtualTextureMaterialType::BaseColor:
 				if (bRenderPass)
 				{
-					OutputAlias0 = RenderTexture0 = GraphBuilder.CreateTexture(FRDGTextureDesc::Create2D(Desc.TextureSize, PF_B8G8R8A8, FClearValueBinding::Black, RTCreateFlags | RTSrgbFlags), TEXT("RenderTexture0"));
+					OutputAlias0 = RenderTexture0 = GraphBuilder.CreateTexture(CreateTextureDesc(Desc.TextureSize, PF_B8G8R8A8, FClearValueBinding::Black, RTCreateFlags | RTSrgbFlags, PageCount), TEXT("RenderTexture0"));
 				}
 				if (bCompressPass)
 				{
@@ -1024,20 +1064,20 @@ namespace RuntimeVirtualTexture
 					}
 					else
 					{
-						OutputAlias0 = CompressTexture0 = GraphBuilder.CreateTexture(FRDGTextureDesc::Create2D(Desc.TextureSize / 4, Compressed64BitFormat, FClearValueBinding::None, TexCreate_UAV), TEXT("CompressTexture0"));
+						OutputAlias0 = CompressTexture0 = GraphBuilder.CreateTexture(CreateTextureDesc(Desc.TextureSize / 4, Compressed64BitFormat, FClearValueBinding::None, TexCreate_UAV, PageCount), TEXT("CompressTexture0"));
 						CompressTextureUAV0_64bit = GraphBuilder.CreateUAV(FRDGTextureUAVDesc(CompressTexture0));
 					}
 				}
 				if (bCopyThumbnailPass)
 				{
-					OutputAlias0 = CopyTexture0 = GraphBuilder.CreateTexture(FRDGTextureDesc::Create2D(Desc.TextureSize, PF_B8G8R8A8, FClearValueBinding::Black, RTCreateFlags | RTSrgbFlags), TEXT("CopyTexture0"));
+					OutputAlias0 = CopyTexture0 = GraphBuilder.CreateTexture(CreateTextureDesc(Desc.TextureSize, PF_B8G8R8A8, FClearValueBinding::Black, RTCreateFlags | RTSrgbFlags, PageCount), TEXT("CopyTexture0"));
 				}
 				break;
 			case ERuntimeVirtualTextureMaterialType::BaseColor_Normal_Roughness:
 				if (bRenderPass)
 				{
-					OutputAlias0 = RenderTexture0 = GraphBuilder.CreateTexture(FRDGTextureDesc::Create2D(Desc.TextureSize, bLQFormat ? PF_R5G6B5_UNORM : PF_B8G8R8A8, FClearValueBinding::Black, RTCreateFlags), TEXT("RenderTexture0"));
-					OutputAlias1 = RenderTexture1 = GraphBuilder.CreateTexture(FRDGTextureDesc::Create2D(Desc.TextureSize, bLQFormat ? PF_R5G6B5_UNORM : PF_B8G8R8A8, FClearValueBinding::Black, RTCreateFlags), TEXT("RenderTexture1"));
+					OutputAlias0 = RenderTexture0 = GraphBuilder.CreateTexture(CreateTextureDesc(Desc.TextureSize, bLQFormat ? PF_R5G6B5_UNORM : PF_B8G8R8A8, FClearValueBinding::Black, RTCreateFlags, PageCount), TEXT("RenderTexture0"));
+					OutputAlias1 = RenderTexture1 = GraphBuilder.CreateTexture(CreateTextureDesc(Desc.TextureSize, bLQFormat ? PF_R5G6B5_UNORM : PF_B8G8R8A8, FClearValueBinding::Black, RTCreateFlags, PageCount), TEXT("RenderTexture1"));
 				}
 				if (bCompressPass)
 				{
@@ -1051,24 +1091,24 @@ namespace RuntimeVirtualTexture
 					}
 					else
 					{
-						OutputAlias0 = CompressTexture0 = GraphBuilder.CreateTexture(FRDGTextureDesc::Create2D(Desc.TextureSize / 4, Compressed64BitFormat, FClearValueBinding::None, TexCreate_UAV), TEXT("CompressTexture0"));
+						OutputAlias0 = CompressTexture0 = GraphBuilder.CreateTexture(CreateTextureDesc(Desc.TextureSize / 4, Compressed64BitFormat, FClearValueBinding::None, TexCreate_UAV, PageCount), TEXT("CompressTexture0"));
 						CompressTextureUAV0_64bit = GraphBuilder.CreateUAV(FRDGTextureUAVDesc(CompressTexture0));
-						OutputAlias1 = CompressTexture1 = GraphBuilder.CreateTexture(FRDGTextureDesc::Create2D(Desc.TextureSize / 4, Compressed128BitFormat, FClearValueBinding::None, TexCreate_UAV), TEXT("CompressTexture1"));
+						OutputAlias1 = CompressTexture1 = GraphBuilder.CreateTexture(CreateTextureDesc(Desc.TextureSize / 4, Compressed128BitFormat, FClearValueBinding::None, TexCreate_UAV, PageCount), TEXT("CompressTexture1"));
 						CompressTextureUAV1_128bit = GraphBuilder.CreateUAV(FRDGTextureUAVDesc(CompressTexture1));
 					}
 				}
 				if (bCopyThumbnailPass)
 				{
-					OutputAlias0 = CopyTexture0 = GraphBuilder.CreateTexture(FRDGTextureDesc::Create2D(Desc.TextureSize, bLQFormat ? PF_R5G6B5_UNORM : PF_B8G8R8A8, FClearValueBinding::Black, RTCreateFlags), TEXT("CopyTexture0"));
+					OutputAlias0 = CopyTexture0 = GraphBuilder.CreateTexture(CreateTextureDesc(Desc.TextureSize, bLQFormat ? PF_R5G6B5_UNORM : PF_B8G8R8A8, FClearValueBinding::Black, RTCreateFlags, PageCount), TEXT("CopyTexture0"));
 					OutputAlias1 = nullptr;
 				}
 				break;
 			case ERuntimeVirtualTextureMaterialType::BaseColor_Normal_Specular:
 				if (bRenderPass)
 				{
-					OutputAlias0 = RenderTexture0 = GraphBuilder.CreateTexture(FRDGTextureDesc::Create2D(Desc.TextureSize, PF_B8G8R8A8, FClearValueBinding::Black, RTCreateFlags | RTSrgbFlags), TEXT("RenderTexture0"));
-					RenderTexture1 = GraphBuilder.CreateTexture(FRDGTextureDesc::Create2D(Desc.TextureSize, PF_B8G8R8A8, FClearValueBinding::DefaultNormal8Bit, RTCreateFlags), TEXT("RenderTexture1"));
-					RenderTexture2 = GraphBuilder.CreateTexture(FRDGTextureDesc::Create2D(Desc.TextureSize, PF_B8G8R8A8, FClearValueBinding::DefaultNormal8Bit, RTCreateFlags), TEXT("RenderTexture2"));
+					OutputAlias0 = RenderTexture0 = GraphBuilder.CreateTexture(CreateTextureDesc(Desc.TextureSize, PF_B8G8R8A8, FClearValueBinding::Black, RTCreateFlags | RTSrgbFlags, PageCount), TEXT("RenderTexture0"));
+					RenderTexture1 = GraphBuilder.CreateTexture(CreateTextureDesc(Desc.TextureSize, PF_B8G8R8A8, FClearValueBinding::DefaultNormal8Bit, RTCreateFlags, PageCount), TEXT("RenderTexture1"));
+					RenderTexture2 = GraphBuilder.CreateTexture(CreateTextureDesc(Desc.TextureSize, PF_B8G8R8A8, FClearValueBinding::DefaultNormal8Bit, RTCreateFlags, PageCount), TEXT("RenderTexture2"));
 				}
 				if (bCompressPass)
 				{
@@ -1082,28 +1122,28 @@ namespace RuntimeVirtualTexture
 					}
 					else
 					{
-						OutputAlias0 = CompressTexture0 = GraphBuilder.CreateTexture(FRDGTextureDesc::Create2D(Desc.TextureSize / 4, Compressed128BitFormat, FClearValueBinding::None, TexCreate_UAV), TEXT("CompressTexture0"));
+						OutputAlias0 = CompressTexture0 = GraphBuilder.CreateTexture(CreateTextureDesc(Desc.TextureSize / 4, Compressed128BitFormat, FClearValueBinding::None, TexCreate_UAV, PageCount), TEXT("CompressTexture0"));
 						CompressTextureUAV0_128bit = GraphBuilder.CreateUAV(FRDGTextureUAVDesc(CompressTexture0));
-						OutputAlias1 = CompressTexture1 = GraphBuilder.CreateTexture(FRDGTextureDesc::Create2D(Desc.TextureSize / 4, Compressed128BitFormat, FClearValueBinding::None, TexCreate_UAV), TEXT("CompressTexture1"));
+						OutputAlias1 = CompressTexture1 = GraphBuilder.CreateTexture(CreateTextureDesc(Desc.TextureSize / 4, Compressed128BitFormat, FClearValueBinding::None, TexCreate_UAV, PageCount), TEXT("CompressTexture1"));
 						CompressTextureUAV1_128bit = GraphBuilder.CreateUAV(FRDGTextureUAVDesc(CompressTexture1));
 					}
 				}
 				if (bCopyPass)
 				{
-					OutputAlias0 = CopyTexture0 = GraphBuilder.CreateTexture(FRDGTextureDesc::Create2D(Desc.TextureSize, PF_B8G8R8A8, FClearValueBinding::Black, RTCreateFlags | RTSrgbFlags), TEXT("CopyTexture0"));
-					OutputAlias1 = CopyTexture1 = GraphBuilder.CreateTexture(FRDGTextureDesc::Create2D(Desc.TextureSize, PF_B8G8R8A8, FClearValueBinding::Black, RTCreateFlags), TEXT("CopyTexture1"));
+					OutputAlias0 = CopyTexture0 = GraphBuilder.CreateTexture(CreateTextureDesc(Desc.TextureSize, PF_B8G8R8A8, FClearValueBinding::Black, RTCreateFlags | RTSrgbFlags, PageCount), TEXT("CopyTexture0"));
+					OutputAlias1 = CopyTexture1 = GraphBuilder.CreateTexture(CreateTextureDesc(Desc.TextureSize, PF_B8G8R8A8, FClearValueBinding::Black, RTCreateFlags, PageCount), TEXT("CopyTexture1"));
 				}
 				if (bCopyThumbnailPass)
 				{
-					OutputAlias0 = CopyTexture0 = GraphBuilder.CreateTexture(FRDGTextureDesc::Create2D(Desc.TextureSize, PF_B8G8R8A8, FClearValueBinding::Black, RTCreateFlags | RTSrgbFlags), TEXT("CopyTexture0"));
+					OutputAlias0 = CopyTexture0 = GraphBuilder.CreateTexture(CreateTextureDesc(Desc.TextureSize, PF_B8G8R8A8, FClearValueBinding::Black, RTCreateFlags | RTSrgbFlags, PageCount), TEXT("CopyTexture0"));
 				}
 				break;
 			case ERuntimeVirtualTextureMaterialType::BaseColor_Normal_Specular_YCoCg:
 				if (bRenderPass)
 				{
-					OutputAlias0 = RenderTexture0 = GraphBuilder.CreateTexture(FRDGTextureDesc::Create2D(Desc.TextureSize, PF_B8G8R8A8, FClearValueBinding::Black, RTCreateFlags | RTSrgbFlags), TEXT("RenderTexture0"));
-					RenderTexture1 = GraphBuilder.CreateTexture(FRDGTextureDesc::Create2D(Desc.TextureSize, PF_B8G8R8A8, FClearValueBinding::DefaultNormal8Bit, RTCreateFlags), TEXT("RenderTexture1"));
-					RenderTexture2 = GraphBuilder.CreateTexture(FRDGTextureDesc::Create2D(Desc.TextureSize, PF_B8G8R8A8, FClearValueBinding::DefaultNormal8Bit, RTCreateFlags), TEXT("RenderTexture2"));
+					OutputAlias0 = RenderTexture0 = GraphBuilder.CreateTexture(CreateTextureDesc(Desc.TextureSize, PF_B8G8R8A8, FClearValueBinding::Black, RTCreateFlags | RTSrgbFlags, PageCount), TEXT("RenderTexture0"));
+					RenderTexture1 = GraphBuilder.CreateTexture(CreateTextureDesc(Desc.TextureSize, PF_B8G8R8A8, FClearValueBinding::DefaultNormal8Bit, RTCreateFlags, PageCount), TEXT("RenderTexture1"));
+					RenderTexture2 = GraphBuilder.CreateTexture(CreateTextureDesc(Desc.TextureSize, PF_B8G8R8A8, FClearValueBinding::DefaultNormal8Bit, RTCreateFlags, PageCount), TEXT("RenderTexture2"));
 				}
 				if (bCompressPass)
 				{
@@ -1120,31 +1160,31 @@ namespace RuntimeVirtualTexture
 					}
 					else
 					{
-						OutputAlias0 = CompressTexture0 = GraphBuilder.CreateTexture(FRDGTextureDesc::Create2D(Desc.TextureSize / 4, Compressed128BitFormat, FClearValueBinding::None, TexCreate_UAV), TEXT("CompressTexture0"));
+						OutputAlias0 = CompressTexture0 = GraphBuilder.CreateTexture(CreateTextureDesc(Desc.TextureSize / 4, Compressed128BitFormat, FClearValueBinding::None, TexCreate_UAV, PageCount), TEXT("CompressTexture0"));
 						CompressTextureUAV0_128bit = GraphBuilder.CreateUAV(FRDGTextureUAVDesc(CompressTexture0));
-						OutputAlias1 = CompressTexture1 = GraphBuilder.CreateTexture(FRDGTextureDesc::Create2D(Desc.TextureSize / 4, Compressed128BitFormat, FClearValueBinding::None, TexCreate_UAV), TEXT("CompressTexture1"));
+						OutputAlias1 = CompressTexture1 = GraphBuilder.CreateTexture(CreateTextureDesc(Desc.TextureSize / 4, Compressed128BitFormat, FClearValueBinding::None, TexCreate_UAV, PageCount), TEXT("CompressTexture1"));
 						CompressTextureUAV1_128bit = GraphBuilder.CreateUAV(FRDGTextureUAVDesc(CompressTexture1));
-						OutputAlias2 = CompressTexture2 = GraphBuilder.CreateTexture(FRDGTextureDesc::Create2D(Desc.TextureSize / 4, Compressed64BitFormat, FClearValueBinding::None, TexCreate_UAV), TEXT("CompressTexture2"));
+						OutputAlias2 = CompressTexture2 = GraphBuilder.CreateTexture(CreateTextureDesc(Desc.TextureSize / 4, Compressed64BitFormat, FClearValueBinding::None, TexCreate_UAV, PageCount), TEXT("CompressTexture2"));
 						CompressTextureUAV2_64bit = GraphBuilder.CreateUAV(FRDGTextureUAVDesc(CompressTexture2));
 					}
 				}
 				if (bCopyPass)
 				{
-					OutputAlias0 = CopyTexture0 = GraphBuilder.CreateTexture(FRDGTextureDesc::Create2D(Desc.TextureSize, PF_B8G8R8A8, FClearValueBinding::Black, RTCreateFlags), TEXT("CopyTexture0"));
-					OutputAlias1 = CopyTexture1 = GraphBuilder.CreateTexture(FRDGTextureDesc::Create2D(Desc.TextureSize, PF_B8G8R8A8, FClearValueBinding::Black, RTCreateFlags), TEXT("CopyTexture1"));
-					OutputAlias2 = CopyTexture2 = GraphBuilder.CreateTexture(FRDGTextureDesc::Create2D(Desc.TextureSize, PF_B8G8R8A8, FClearValueBinding::Black, RTCreateFlags), TEXT("CopyTexture2"));
+					OutputAlias0 = CopyTexture0 = GraphBuilder.CreateTexture(CreateTextureDesc(Desc.TextureSize, PF_B8G8R8A8, FClearValueBinding::Black, RTCreateFlags, PageCount), TEXT("CopyTexture0"));
+					OutputAlias1 = CopyTexture1 = GraphBuilder.CreateTexture(CreateTextureDesc(Desc.TextureSize, PF_B8G8R8A8, FClearValueBinding::Black, RTCreateFlags, PageCount), TEXT("CopyTexture1"));
+					OutputAlias2 = CopyTexture2 = GraphBuilder.CreateTexture(CreateTextureDesc(Desc.TextureSize, PF_B8G8R8A8, FClearValueBinding::Black, RTCreateFlags, PageCount), TEXT("CopyTexture2"));
 				}
 				if (bCopyThumbnailPass)
 				{
-					OutputAlias0 = CopyTexture0 = GraphBuilder.CreateTexture(FRDGTextureDesc::Create2D(Desc.TextureSize, PF_B8G8R8A8, FClearValueBinding::Black, RTCreateFlags | RTSrgbFlags), TEXT("CopyTexture0"));
+					OutputAlias0 = CopyTexture0 = GraphBuilder.CreateTexture(CreateTextureDesc(Desc.TextureSize, PF_B8G8R8A8, FClearValueBinding::Black, RTCreateFlags | RTSrgbFlags, PageCount), TEXT("CopyTexture0"));
 				}
 				break;
 			case ERuntimeVirtualTextureMaterialType::BaseColor_Normal_Specular_Mask_YCoCg:
 				if (bRenderPass)
 				{
-					OutputAlias0 = RenderTexture0 = GraphBuilder.CreateTexture(FRDGTextureDesc::Create2D(Desc.TextureSize, PF_B8G8R8A8, FClearValueBinding::Black, RTCreateFlags | RTSrgbFlags), TEXT("RenderTexture0"));
-					RenderTexture1 = GraphBuilder.CreateTexture(FRDGTextureDesc::Create2D(Desc.TextureSize, PF_B8G8R8A8, FClearValueBinding::DefaultNormal8Bit, RTCreateFlags), TEXT("RenderTexture1"));
-					RenderTexture2 = GraphBuilder.CreateTexture(FRDGTextureDesc::Create2D(Desc.TextureSize, PF_B8G8R8A8, FClearValueBinding::DefaultNormal8Bit, RTCreateFlags), TEXT("RenderTexture2"));
+					OutputAlias0 = RenderTexture0 = GraphBuilder.CreateTexture(CreateTextureDesc(Desc.TextureSize, PF_B8G8R8A8, FClearValueBinding::Black, RTCreateFlags | RTSrgbFlags, PageCount), TEXT("RenderTexture0"));
+					RenderTexture1 = GraphBuilder.CreateTexture(CreateTextureDesc(Desc.TextureSize, PF_B8G8R8A8, FClearValueBinding::DefaultNormal8Bit, RTCreateFlags, PageCount), TEXT("RenderTexture1"));
+					RenderTexture2 = GraphBuilder.CreateTexture(CreateTextureDesc(Desc.TextureSize, PF_B8G8R8A8, FClearValueBinding::DefaultNormal8Bit, RTCreateFlags, PageCount), TEXT("RenderTexture2"));
 				}
 				if (bCompressPass)
 				{
@@ -1161,30 +1201,30 @@ namespace RuntimeVirtualTexture
 					}
 					else
 					{
-						OutputAlias0 = CompressTexture0 = GraphBuilder.CreateTexture(FRDGTextureDesc::Create2D(Desc.TextureSize / 4, Compressed128BitFormat, FClearValueBinding::None, TexCreate_UAV), TEXT("CompressTexture0"));
+						OutputAlias0 = CompressTexture0 = GraphBuilder.CreateTexture(CreateTextureDesc(Desc.TextureSize / 4, Compressed128BitFormat, FClearValueBinding::None, TexCreate_UAV, PageCount), TEXT("CompressTexture0"));
 						CompressTextureUAV0_128bit = GraphBuilder.CreateUAV(FRDGTextureUAVDesc(CompressTexture0));
-						OutputAlias1 = CompressTexture1 = GraphBuilder.CreateTexture(FRDGTextureDesc::Create2D(Desc.TextureSize / 4, Compressed128BitFormat, FClearValueBinding::None, TexCreate_UAV), TEXT("CompressTexture1"));
+						OutputAlias1 = CompressTexture1 = GraphBuilder.CreateTexture(CreateTextureDesc(Desc.TextureSize / 4, Compressed128BitFormat, FClearValueBinding::None, TexCreate_UAV, PageCount), TEXT("CompressTexture1"));
 						CompressTextureUAV1_128bit = GraphBuilder.CreateUAV(FRDGTextureUAVDesc(CompressTexture1));
-						OutputAlias2 = CompressTexture2 = GraphBuilder.CreateTexture(FRDGTextureDesc::Create2D(Desc.TextureSize / 4, Compressed128BitFormat, FClearValueBinding::None, TexCreate_UAV), TEXT("CompressTexture2"));
+						OutputAlias2 = CompressTexture2 = GraphBuilder.CreateTexture(CreateTextureDesc(Desc.TextureSize / 4, Compressed128BitFormat, FClearValueBinding::None, TexCreate_UAV, PageCount), TEXT("CompressTexture2"));
 						CompressTextureUAV2_128bit = GraphBuilder.CreateUAV(FRDGTextureUAVDesc(CompressTexture2));
 					}
 				}
 				if (bCopyPass)
 				{
-					OutputAlias0 = CopyTexture0 = GraphBuilder.CreateTexture(FRDGTextureDesc::Create2D(Desc.TextureSize, PF_B8G8R8A8, FClearValueBinding::Black, RTCreateFlags), TEXT("CopyTexture0"));
-					OutputAlias1 = CopyTexture1 = GraphBuilder.CreateTexture(FRDGTextureDesc::Create2D(Desc.TextureSize, PF_B8G8R8A8, FClearValueBinding::Black, RTCreateFlags), TEXT("CopyTexture1"));
-					OutputAlias2 = CopyTexture2 = GraphBuilder.CreateTexture(FRDGTextureDesc::Create2D(Desc.TextureSize, PF_B8G8R8A8, FClearValueBinding::Black, RTCreateFlags), TEXT("CopyTexture2"));
+					OutputAlias0 = CopyTexture0 = GraphBuilder.CreateTexture(CreateTextureDesc(Desc.TextureSize, PF_B8G8R8A8, FClearValueBinding::Black, RTCreateFlags, PageCount), TEXT("CopyTexture0"));
+					OutputAlias1 = CopyTexture1 = GraphBuilder.CreateTexture(CreateTextureDesc(Desc.TextureSize, PF_B8G8R8A8, FClearValueBinding::Black, RTCreateFlags, PageCount), TEXT("CopyTexture1"));
+					OutputAlias2 = CopyTexture2 = GraphBuilder.CreateTexture(CreateTextureDesc(Desc.TextureSize, PF_B8G8R8A8, FClearValueBinding::Black, RTCreateFlags, PageCount), TEXT("CopyTexture2"));
 				}
 				if (bCopyThumbnailPass)
 				{
-					OutputAlias0 = CopyTexture0 = GraphBuilder.CreateTexture(FRDGTextureDesc::Create2D(Desc.TextureSize, PF_B8G8R8A8, FClearValueBinding::Black, RTCreateFlags | RTSrgbFlags), TEXT("CopyTexture0"));
+					OutputAlias0 = CopyTexture0 = GraphBuilder.CreateTexture(CreateTextureDesc(Desc.TextureSize, PF_B8G8R8A8, FClearValueBinding::Black, RTCreateFlags | RTSrgbFlags, PageCount), TEXT("CopyTexture0"));
 				}
 				break;
 			case ERuntimeVirtualTextureMaterialType::Mask4:
 				if (bRenderPass)
 				{
-					OutputAlias0 = RenderTexture0 = GraphBuilder.CreateTexture(FRDGTextureDesc::Create2D(Desc.TextureSize, PF_B8G8R8A8, FClearValueBinding::Black, RTCreateFlags), TEXT("RenderTexture0"));
-					OutputAlias1 = RenderTexture1 = GraphBuilder.CreateTexture(FRDGTextureDesc::Create2D(Desc.TextureSize, PF_B8G8R8A8, FClearValueBinding::Black, RTCreateFlags), TEXT("RenderTexture1"));
+					OutputAlias0 = RenderTexture0 = GraphBuilder.CreateTexture(CreateTextureDesc(Desc.TextureSize, PF_B8G8R8A8, FClearValueBinding::Black, RTCreateFlags, PageCount), TEXT("RenderTexture0"));
+					OutputAlias1 = RenderTexture1 = GraphBuilder.CreateTexture(CreateTextureDesc(Desc.TextureSize, PF_B8G8R8A8, FClearValueBinding::Black, RTCreateFlags, PageCount), TEXT("RenderTexture1"));
 				}
 				if (bCompressPass)
 				{
@@ -1196,29 +1236,29 @@ namespace RuntimeVirtualTexture
 					}
 					else
 					{
-						OutputAlias0 = CompressTexture0 = GraphBuilder.CreateTexture(FRDGTextureDesc::Create2D(Desc.TextureSize / 4, Compressed128BitFormat, FClearValueBinding::None, TexCreate_UAV), TEXT("CompressTexture0"));
+						OutputAlias0 = CompressTexture0 = GraphBuilder.CreateTexture(CreateTextureDesc(Desc.TextureSize / 4, Compressed128BitFormat, FClearValueBinding::None, TexCreate_UAV, PageCount), TEXT("CompressTexture0"));
 						CompressTextureUAV0_128bit = GraphBuilder.CreateUAV(FRDGTextureUAVDesc(CompressTexture0));
 					}
 				}
 				if (bCopyPass || bCopyThumbnailPass)
 				{
-					OutputAlias0 = CopyTexture0 = GraphBuilder.CreateTexture(FRDGTextureDesc::Create2D(Desc.TextureSize, PF_B8G8R8A8, FClearValueBinding::Black, RTCreateFlags), TEXT("CopyTexture0"));
+					OutputAlias0 = CopyTexture0 = GraphBuilder.CreateTexture(CreateTextureDesc(Desc.TextureSize, PF_B8G8R8A8, FClearValueBinding::Black, RTCreateFlags, PageCount), TEXT("CopyTexture0"));
 				}
 				break;
 			case ERuntimeVirtualTextureMaterialType::WorldHeight:
 				if (bRenderPass)
 				{
-					OutputAlias0 = RenderTexture0 = GraphBuilder.CreateTexture(FRDGTextureDesc::Create2D(Desc.TextureSize, PF_G16, FClearValueBinding::Black, RTCreateFlags), TEXT("RenderTexture0"));
+					OutputAlias0 = RenderTexture0 = GraphBuilder.CreateTexture(CreateTextureDesc(Desc.TextureSize, PF_G16, FClearValueBinding::Black, RTCreateFlags, PageCount), TEXT("RenderTexture0"));
 				}
 				if (bCopyThumbnailPass)
 				{
-					OutputAlias0 = CopyTexture0 = GraphBuilder.CreateTexture(FRDGTextureDesc::Create2D(Desc.TextureSize, PF_B8G8R8A8, FClearValueBinding::Black, RTCreateFlags), TEXT("CopyTexture0"));
+					OutputAlias0 = CopyTexture0 = GraphBuilder.CreateTexture(CreateTextureDesc(Desc.TextureSize, PF_B8G8R8A8, FClearValueBinding::Black, RTCreateFlags, PageCount), TEXT("CopyTexture0"));
 				}
 				break;
 			case ERuntimeVirtualTextureMaterialType::Displacement:
 				if (bRenderPass)
 				{
-					OutputAlias0 = RenderTexture0 = GraphBuilder.CreateTexture(FRDGTextureDesc::Create2D(Desc.TextureSize, PF_G16, FClearValueBinding::Black, RTCreateFlags), TEXT("RenderTexture0"));
+					OutputAlias0 = RenderTexture0 = GraphBuilder.CreateTexture(CreateTextureDesc(Desc.TextureSize, PF_G16, FClearValueBinding::Black, RTCreateFlags, PageCount), TEXT("RenderTexture0"));
 				}
 				if (bCompressPass)
 				{
@@ -1230,13 +1270,13 @@ namespace RuntimeVirtualTexture
 					}
 					else
 					{
-						OutputAlias0 = CompressTexture0 = GraphBuilder.CreateTexture(FRDGTextureDesc::Create2D(Desc.TextureSize / 4, Compressed64BitFormat, FClearValueBinding::None, TexCreate_UAV), TEXT("CompressTexture0"));
+						OutputAlias0 = CompressTexture0 = GraphBuilder.CreateTexture(CreateTextureDesc(Desc.TextureSize / 4, Compressed64BitFormat, FClearValueBinding::None, TexCreate_UAV, PageCount), TEXT("CompressTexture0"));
 						CompressTextureUAV0_64bit = GraphBuilder.CreateUAV(FRDGTextureUAVDesc(CompressTexture0));
 					}
 				}
 				if (bCopyThumbnailPass)
 				{
-					OutputAlias0 = CopyTexture0 = GraphBuilder.CreateTexture(FRDGTextureDesc::Create2D(Desc.TextureSize, PF_B8G8R8A8, FClearValueBinding::Black, RTCreateFlags), TEXT("CopyTexture0"));
+					OutputAlias0 = CopyTexture0 = GraphBuilder.CreateTexture(CreateTextureDesc(Desc.TextureSize, PF_B8G8R8A8, FClearValueBinding::Black, RTCreateFlags, PageCount), TEXT("CopyTexture0"));
 				}
 				break;
 			}
@@ -1248,7 +1288,7 @@ namespace RuntimeVirtualTexture
 		bool bCopyPass = false;
 		bool bCopyThumbnailPass = false;
 		bool bDirectAliasing = false;
-		
+
 		/** Render graph textures needed for this virtual texture layout. */
 		FRDGTextureRef RenderTexture0 = nullptr;
 		FRDGTextureRef RenderTexture1 = nullptr;
@@ -1272,6 +1312,18 @@ namespace RuntimeVirtualTexture
 		FRDGTextureRef OutputAlias2 = nullptr;
 	};
 
+	/** 
+	 * Context for rendering a batch of pages. 
+	 * Holds the batch description and the render graph allocations.
+	 * Allows us to maintain state across RenderFinalize() and Finalize() calls.
+	 */
+	class FBatchRenderContext
+	{
+	public:
+		FRenderGraphSetup GraphSetup;
+		FRenderPageBatchDesc BatchDesc;
+		bool bAllowCachedMeshDrawCommands = true;
+	};
 
 	/** Mesh processor for rendering static meshes to the virtual texture */
 	class FRuntimeVirtualTextureMeshProcessor : public FSceneRenderingAllocatorObject<FRuntimeVirtualTextureMeshProcessor>, public FMeshPassProcessor
@@ -1602,9 +1654,7 @@ namespace RuntimeVirtualTexture
 	public:
 		BEGIN_SHADER_PARAMETER_STRUCT(FParameters, )
 			SHADER_PARAMETER(FUintVector4, SourceRect)
-			SHADER_PARAMETER(FUintVector2, DestPos0)
-			SHADER_PARAMETER(FUintVector2, DestPos1)
-			SHADER_PARAMETER(FUintVector2, DestPos2)
+			SHADER_PARAMETER_SCALAR_ARRAY(int32, DestPos, [MaxRenderPageBatch * MaxTextureLayers * 2])
 			SHADER_PARAMETER_STRUCT_REF(FEtcParameters, EtcParameters)
 			SHADER_PARAMETER_STRUCT_REF(FAstcParameters, AstcParameters)
 			SHADER_PARAMETER_RDG_TEXTURE(Texture2D<float4>, RenderTexture0)
@@ -1621,20 +1671,35 @@ namespace RuntimeVirtualTexture
 			SHADER_PARAMETER_RDG_TEXTURE_UAV(RWTexture2D<uint4>, OutCompressTexture2_128bit)
 		END_SHADER_PARAMETER_STRUCT()
 
+		class FUseSrcTextureArray : SHADER_PERMUTATION_BOOL("USE_SRC_TEXTURE_ARRAY");
+		class FUseDstTextureArray : SHADER_PERMUTATION_BOOL("USE_DST_TEXTURE_ARRAY");
 		class FAstcHighProfile : SHADER_PERMUTATION_BOOL("ASTC_HIGH_PROFILE");
-		using FPermutationDomain = TShaderPermutationDomain<FAstcHighProfile>;
+		using FPermutationDomain = TShaderPermutationDomain<FUseSrcTextureArray, FUseDstTextureArray, FAstcHighProfile>;
 
 		static void ModifyCompilationEnvironment(const FGlobalShaderPermutationParameters& Parameters, FShaderCompilerEnvironment& OutEnvironment)
 		{
 			FGlobalShader::ModifyCompilationEnvironment(Parameters, OutEnvironment);
+			
 			OutEnvironment.SetDefine(TEXT("ETC_PROFILE"), UseEtcProfile(Parameters.Platform) ? 1 : 0);
 			OutEnvironment.SetDefine(TEXT("ASTC_PROFILE"), UseAstcProfile(Parameters.Platform) ? 1 : 0);
 			OutEnvironment.SetDefine(TEXT("PACK_RG32_RGBA16"), UseRGBA16(Parameters.Platform) ? 1 : 0);
+			
+			OutEnvironment.SetDefine(TEXT("MAX_BATCH_SIZE"), MaxRenderPageBatch);
+			OutEnvironment.SetDefine(TEXT("MAX_DST_LAYERS"), MaxTextureLayers);
+
+			const FPermutationDomain PermutationVector(Parameters.PermutationId);
+			const bool bUseSrcTextureArray = PermutationVector.Get<FUseSrcTextureArray>();
+			OutEnvironment.SetDefine(TEXT("BLOCK_COMPRESS_SRC_TEXTURE_ARRAY"), bUseSrcTextureArray ? 1 : 0);
 		}
 
 		static bool ShouldCompilePermutation(const FGlobalShaderPermutationParameters& Parameters)
 		{
 			FPermutationDomain PermutationVector(Parameters.PermutationId);
+			if (!PermutationVector.Get<FUseSrcTextureArray>() && PermutationVector.Get<FUseDstTextureArray>())
+			{
+				// No compress pass goes from simple source texture to destination array texture.
+				return false;
+			}
 			if (PermutationVector.Get<FAstcHighProfile>())
 			{
 				return UseAstcHighProfile(Parameters.Platform);
@@ -1683,11 +1748,13 @@ namespace RuntimeVirtualTexture
 
 	/** Add the BC compression pass to the graph. */
 	template< ERuntimeVirtualTextureMaterialType MaterialType >
-	void AddCompressPass(FRDGBuilder& GraphBuilder, ERHIFeatureLevel::Type FeatureLevel, FShader_VirtualTextureCompress::FParameters* Parameters, FIntVector GroupCount)
+	void AddCompressPass(FRDGBuilder& GraphBuilder, ERHIFeatureLevel::Type FeatureLevel, FShader_VirtualTextureCompress::FParameters* Parameters, FIntVector GroupCount, bool bDirectAliasing)
 	{
 		FGlobalShaderMap* GlobalShaderMap = GetGlobalShaderMap(FeatureLevel);
 
 		FShader_VirtualTextureCompress::FPermutationDomain PermutationVector;
+		PermutationVector.Set<FShader_VirtualTextureCompress::FUseSrcTextureArray>(GroupCount.Z > 1);
+		PermutationVector.Set<FShader_VirtualTextureCompress::FUseDstTextureArray>(GroupCount.Z > 1 && !bDirectAliasing);
 		PermutationVector.Set<FShader_VirtualTextureCompress::FAstcHighProfile>(UseAstcHighProfile(GShaderPlatformForFeatureLevel[FeatureLevel]));
 		TShaderMapRef< FShader_VirtualTextureCompress_CS< MaterialType > > ComputeShader(GlobalShaderMap, PermutationVector);
 
@@ -1698,33 +1765,40 @@ namespace RuntimeVirtualTexture
 	}
 
 	/** Set up the BC compression pass for the given MaterialType. */
-	void AddCompressPass(FRDGBuilder& GraphBuilder, ERHIFeatureLevel::Type FeatureLevel, FShader_VirtualTextureCompress::FParameters* Parameters, FIntPoint TextureSize, ERuntimeVirtualTextureMaterialType MaterialType)
+	void AddCompressPass(
+		FRDGBuilder& GraphBuilder, 
+		ERHIFeatureLevel::Type FeatureLevel, 
+		FShader_VirtualTextureCompress::FParameters* Parameters, 
+		FIntPoint TextureSize, 
+		int32 NumSlices, 
+		ERuntimeVirtualTextureMaterialType MaterialType, 
+		bool bDirectAliasing)
 	{
-		const FIntVector GroupCount(((TextureSize.X / 4) + 7) / 8, ((TextureSize.Y / 4) + 7) / 8, 1);
+		const FIntVector GroupCount(((TextureSize.X / 4) + 7) / 8, ((TextureSize.Y / 4) + 7) / 8, NumSlices);
 
 		// Dispatch using the shader variation for our MaterialType
 		switch (MaterialType)
 		{
 		case ERuntimeVirtualTextureMaterialType::BaseColor:
-			AddCompressPass<ERuntimeVirtualTextureMaterialType::BaseColor>(GraphBuilder, FeatureLevel, Parameters, GroupCount);
+			AddCompressPass<ERuntimeVirtualTextureMaterialType::BaseColor>(GraphBuilder, FeatureLevel, Parameters, GroupCount, bDirectAliasing);
 			break;
 		case ERuntimeVirtualTextureMaterialType::BaseColor_Normal_Specular:
-			AddCompressPass<ERuntimeVirtualTextureMaterialType::BaseColor_Normal_Specular>(GraphBuilder, FeatureLevel, Parameters, GroupCount);
+			AddCompressPass<ERuntimeVirtualTextureMaterialType::BaseColor_Normal_Specular>(GraphBuilder, FeatureLevel, Parameters, GroupCount, bDirectAliasing);
 			break;
 		case ERuntimeVirtualTextureMaterialType::BaseColor_Normal_Roughness:
-			AddCompressPass<ERuntimeVirtualTextureMaterialType::BaseColor_Normal_Roughness>(GraphBuilder, FeatureLevel, Parameters, GroupCount);
+			AddCompressPass<ERuntimeVirtualTextureMaterialType::BaseColor_Normal_Roughness>(GraphBuilder, FeatureLevel, Parameters, GroupCount, bDirectAliasing);
 			break;
 		case ERuntimeVirtualTextureMaterialType::BaseColor_Normal_Specular_YCoCg:
-			AddCompressPass<ERuntimeVirtualTextureMaterialType::BaseColor_Normal_Specular_YCoCg>(GraphBuilder, FeatureLevel, Parameters, GroupCount);
+			AddCompressPass<ERuntimeVirtualTextureMaterialType::BaseColor_Normal_Specular_YCoCg>(GraphBuilder, FeatureLevel, Parameters, GroupCount, bDirectAliasing);
 			break;
 		case ERuntimeVirtualTextureMaterialType::BaseColor_Normal_Specular_Mask_YCoCg:
-			AddCompressPass<ERuntimeVirtualTextureMaterialType::BaseColor_Normal_Specular_Mask_YCoCg>(GraphBuilder, FeatureLevel, Parameters, GroupCount);
+			AddCompressPass<ERuntimeVirtualTextureMaterialType::BaseColor_Normal_Specular_Mask_YCoCg>(GraphBuilder, FeatureLevel, Parameters, GroupCount, bDirectAliasing);
 			break;
 		case ERuntimeVirtualTextureMaterialType::Mask4:
-			AddCompressPass<ERuntimeVirtualTextureMaterialType::Mask4>(GraphBuilder, FeatureLevel, Parameters, GroupCount);
+			AddCompressPass<ERuntimeVirtualTextureMaterialType::Mask4>(GraphBuilder, FeatureLevel, Parameters, GroupCount, bDirectAliasing);
 			break;
 		case ERuntimeVirtualTextureMaterialType::Displacement:
-			AddCompressPass<ERuntimeVirtualTextureMaterialType::Displacement>(GraphBuilder, FeatureLevel, Parameters, GroupCount);
+			AddCompressPass<ERuntimeVirtualTextureMaterialType::Displacement>(GraphBuilder, FeatureLevel, Parameters, GroupCount, bDirectAliasing);
 			break;
 		}
 	}
@@ -1737,11 +1811,11 @@ namespace RuntimeVirtualTexture
 		BEGIN_SHADER_PARAMETER_STRUCT(FParameters, )
 			RENDER_TARGET_BINDING_SLOTS()
 			SHADER_PARAMETER(FIntVector4, DestRect)
-			SHADER_PARAMETER_RDG_TEXTURE(Texture2D<float4>, RenderTexture0)
+			SHADER_PARAMETER_RDG_TEXTURE_SRV(Texture2D<float4>, RenderTexture0)
 			SHADER_PARAMETER_SAMPLER(SamplerState, TextureSampler0)
-			SHADER_PARAMETER_RDG_TEXTURE(Texture2D<float4>, RenderTexture1)
+			SHADER_PARAMETER_RDG_TEXTURE_SRV(Texture2D<float4>, RenderTexture1)
 			SHADER_PARAMETER_SAMPLER(SamplerState, TextureSampler1)
-			SHADER_PARAMETER_RDG_TEXTURE(Texture2D<float4>, RenderTexture2)
+			SHADER_PARAMETER_RDG_TEXTURE_SRV(Texture2D<float4>, RenderTexture2)
 			SHADER_PARAMETER_SAMPLER(SamplerState, TextureSampler2)
 		END_SHADER_PARAMETER_STRUCT()
 
@@ -1753,6 +1827,13 @@ namespace RuntimeVirtualTexture
 		{
 			Bindings.BindForLegacyShaderParameters(this, Initializer.PermutationId, Initializer.ParameterMap, *FParameters::FTypeInfo::GetStructMetadata());
 		}
+
+		static void ModifyCompilationEnvironment(const FGlobalShaderPermutationParameters& Parameters, FShaderCompilerEnvironment& OutEnvironment)
+		{
+			FGlobalShader::ModifyCompilationEnvironment(Parameters, OutEnvironment);
+			OutEnvironment.SetDefine(TEXT("MAX_BATCH_SIZE"), 1);
+			OutEnvironment.SetDefine(TEXT("MAX_DST_LAYERS"), 1);
+		}
 	};
 
 	class FShader_VirtualTextureCopy_VS : public FShader_VirtualTextureCopy
@@ -1873,11 +1954,19 @@ namespace RuntimeVirtualTexture
 	END_SHADER_PARAMETER_STRUCT()
 
 	/** Set up the copy to final output physical texture. */
-	void AddCopyToOutputPass(FRDGBuilder& GraphBuilder, FRDGTextureRef InputTexture, FRHITexture* OutputTexture, ERHIAccess OutputTextureAccessBefore, ERHIAccess OutputTextureAccessAfter, FBox2D const& DestBox)
+	void AddCopyToOutputPass(
+		FRDGBuilder& GraphBuilder, 
+		FRDGTextureRef InputTexture, 
+		FRHITexture* OutputTexture, 
+		ERHIAccess OutputTextureAccessBefore, 
+		ERHIAccess OutputTextureAccessAfter, 
+		int32 SourceSliceIndex, 
+		FIntPoint const& DestPosition)
 	{
 		FRHICopyTextureInfo CopyInfo;
 		CopyInfo.Size = InputTexture->Desc.GetSize();
-		CopyInfo.DestPosition = FIntVector(DestBox.Min.X, DestBox.Min.Y, 0);
+		CopyInfo.SourceSliceIndex = SourceSliceIndex;
+		CopyInfo.DestPosition = FIntVector(DestPosition.X, DestPosition.Y, 0);
 
 		FCopyToOutputParameters* Parameters = GraphBuilder.AllocParameters<FCopyToOutputParameters>();
 		Parameters->Input = InputTexture;
@@ -1887,18 +1976,22 @@ namespace RuntimeVirtualTexture
 			Parameters,
 			ERDGPassFlags::Copy | ERDGPassFlags::NeverCull,
 			[InputTexture, OutputTexture, OutputTextureAccessBefore, OutputTextureAccessAfter, CopyInfo](FRDGAsyncTask, FRHICommandList& RHICmdList)
+		{
+			if (OutputTextureAccessBefore != ERHIAccess::Unknown)
 			{
 				RHICmdList.Transition(FRHITransitionInfo(OutputTexture, OutputTextureAccessBefore, ERHIAccess::CopyDest));
-				RHICmdList.CopyTexture(InputTexture->GetRHI(), OutputTexture, CopyInfo);
-				RHICmdList.Transition(FRHITransitionInfo(OutputTexture, ERHIAccess::CopyDest, OutputTextureAccessAfter));
-			});
-	}
+			}
 
-	bool IsSceneReadyToRender(FSceneInterface* Scene)
-	{
-		return Scene != nullptr && Scene->GetRenderScene() != nullptr && Scene->GetRenderScene()->GPUScene.IsRendering();
+			RHICmdList.CopyTexture(InputTexture->GetRHI(), OutputTexture, CopyInfo);
+
+			if (OutputTextureAccessAfter != ERHIAccess::Unknown)
+			{
+				RHICmdList.Transition(FRHITransitionInfo(OutputTexture, ERHIAccess::CopyDest, OutputTextureAccessAfter));
+			}
+		});
 	}
 
+	/** Get the debug color to use for a given mip level. Returns InDefaultColor if mip debugging is disabled. */
 	FLinearColor GetDebugMipLevelColor(uint32 InLevel, FLinearColor InDefaultColor)
 	{
 		const int32 MipColorMode = CVarVTMipColors.GetValueOnRenderThread();
@@ -1919,57 +2012,35 @@ namespace RuntimeVirtualTexture
 		return InDefaultColor;
 	}
 
-	void RenderPage(
-		FRDGBuilder& GraphBuilder,
-		FScene* Scene,
-		ISceneRenderer* SceneRenderer,
-		int32 RuntimeVirtualTextureId,
-		ERuntimeVirtualTextureMaterialType MaterialType,
-		bool bClearTextures,
-		bool bIsThumbnails,
-		bool bAllowCachedMeshDrawCommands,
-		FRHITexture* OutputTexture0,		// todo[vt]: Only use IPooledRenderTarget or FRDGTextureRef, not raw RHI textures.
-		ERHIAccess OutputTextureAccessBefore0, 
-		ERHIAccess OutputTextureAccessAfter0, 
-		IPooledRenderTarget* OutputTarget0,
-		FBox2D const& DestBox0,
-		FRHITexture* OutputTexture1,
-		ERHIAccess OutputTextureAccessBefore1, 
-		ERHIAccess OutputTextureAccessAfter1,
-		IPooledRenderTarget* OutputTarget1,
-		FBox2D const& DestBox1,
-		FRHITexture* OutputTexture2, 
-		ERHIAccess OutputTextureAccessBefore2, 
-		ERHIAccess OutputTextureAccessAfter2, 
-		IPooledRenderTarget* OutputTarget2,
-		FBox2D const& DestBox2,
-		FTransform const& UVToWorld,
-		FBox const& WorldBounds,
-		FBox2D const& UVRange,
-		uint8 vLevel,
-		uint8 MaxLevel,
-		FLinearColor const& FixedColor, 
-		FVector4f const& CustomMaterialData)
+	/** 
+	 * Render a single page from a batch. 
+	 * todo[vt]: Can we add some batch rendering mesh pass where all prerequesite BuildRenderingCommands/Compute phases are batched and then all Graphics draws are batched.
+	 */
+	void RenderPage(FRDGBuilder& GraphBuilder, FBatchRenderContext const& BatchRenderContext, int32 PageIndex)
 	{
-		RDG_EVENT_SCOPE(GraphBuilder, "VirtualTextureDynamicCache");
 		CSV_CUSTOM_STAT(VirtualTexturing, RenderedPages, 1, ECsvCustomStatOp::Accumulate);
 		INC_DWORD_STAT_BY(STAT_RenderedPages, 1);
 
-		// Initialize a temporary view required for the material render pass
-		//todo[vt]: Some of this, such as ViewRotationMatrix, can be computed once in the Finalizer and passed down.
-		//todo[vt]: Have specific shader variations and setup for different output texture configs
+		FRenderGraphSetup const& GraphSetup = BatchRenderContext.GraphSetup;
+		FRenderPageBatchDesc const& BatchDesc = BatchRenderContext.BatchDesc;
+		FRenderPageDesc const& PageDesc = BatchDesc.PageDescs[PageIndex];
+		FScene* Scene = BatchDesc.SceneRenderer->GetScene();
+
+		// Initialize the view required for the material render pass
 		FSceneViewFamily::ConstructionValues ViewFamilyInit(nullptr, Scene, FEngineShowFlags(ESFIM_Game));
 		ViewFamilyInit.SetTime(FGameTime());
 		FViewFamilyInfo& ViewFamily = *GraphBuilder.AllocObject<FViewFamilyInfo>(ViewFamilyInit);
-		ViewFamily.SetSceneRenderer(SceneRenderer);
+		ViewFamily.SetSceneRenderer(BatchDesc.SceneRenderer);
 
 		FSceneViewInitOptions ViewInitOptions;
 		ViewInitOptions.ViewFamily = &ViewFamily;
 
-		const FIntPoint TextureSize = (DestBox0.Max - DestBox0.Min).IntPoint();
+		const FIntPoint TextureSize = PageDesc.DestRect[0].Size();
 		ViewInitOptions.SetViewRectangle(FIntRect(FIntPoint(0, 0), TextureSize));
 
+		FBox2D const& UVRange = PageDesc.UVRange;
 		const FVector UVCenter = FVector(UVRange.GetCenter(), 0.f);
+		FTransform const& UVToWorld = BatchDesc.UVToWorld;
 		const FVector CameraLookAt = UVToWorld.TransformPosition(UVCenter);
 		const float BoundBoxZ = UVToWorld.GetScale3D().Z;
 		const FVector CameraPos = CameraLookAt + BoundBoxZ * UVToWorld.GetUnitAxis(EAxis::Z);
@@ -1991,8 +2062,11 @@ namespace RuntimeVirtualTexture
 		const float ZOffset = -NearPlane;
 		ViewInitOptions.ProjectionMatrix = FReversedZOrthoMatrix(OrthoWidth, OrthoHeight, ZScale, ZOffset);
 
+		const uint8 vLevel = PageDesc.vLevel;
+		const uint8 MaxLevel = BatchDesc.MaxLevel;
 		const FVector4f MipLevelParameter = FVector4f((float)vLevel, (float)MaxLevel, OrthoWidth / (float)TextureSize.X, OrthoHeight / (float)TextureSize.Y);
 		
+		FBox const& WorldBounds = BatchDesc.WorldBounds;
 		const float HeightRange = FMath::Max<float>(WorldBounds.Max.Z - WorldBounds.Min.Z, 1.f);
 		const FVector2D WorldHeightPackParameter = FVector2D(1.f / HeightRange, -WorldBounds.Min.Z / HeightRange);
 
@@ -2008,178 +2082,219 @@ namespace RuntimeVirtualTexture
 		View->SetupUniformBufferParameters(nullptr, 0, *View->CachedViewUniformShaderParameters);
 		View->ViewUniformBuffer = TUniformBufferRef<FViewUniformShaderParameters>::CreateUniformBufferImmediate(*View->CachedViewUniformShaderParameters, UniformBuffer_SingleFrame);
 
-		// Build graph
-		FRenderGraphSetup::FInitDesc GraphSetupDesc(Scene->GetFeatureLevel(), MaterialType, TextureSize);
-		GraphSetupDesc.OutputFormat0 = OutputTexture0 != nullptr ? OutputTexture0->GetFormat() : PF_Unknown;
-		GraphSetupDesc.OutputTargets.Add(OutputTarget0);
-		GraphSetupDesc.OutputTargets.Add(OutputTarget1);
-		GraphSetupDesc.OutputTargets.Add(OutputTarget2);
-		GraphSetupDesc.bClearTextures = bClearTextures;
-		GraphSetupDesc.bIsThumbnails = bIsThumbnails;
-		FRenderGraphSetup GraphSetup(GraphBuilder, GraphSetupDesc);
-
-		// Draw Pass
-		if (GraphSetup.bRenderPass)
 		{
-			RDG_EVENT_SCOPE(GraphBuilder, "RenderPass");
+			RenderCaptureInterface::FScopedCapture RenderCapture((RenderCaptureNextRVTPagesDraws != 0), GraphBuilder, TEXT("RenderRVTPage"));
+			RenderCaptureNextRVTPagesDraws = FMath::Max(RenderCaptureNextRVTPagesDraws - 1, 0);
 
-			ERenderTargetLoadAction LoadAction = bClearTextures ? ERenderTargetLoadAction::EClear : ERenderTargetLoadAction::ENoAction;
+			ERenderTargetLoadAction LoadAction = BatchDesc.bClearTextures ? ERenderTargetLoadAction::EClear : ERenderTargetLoadAction::ENoAction;
 			FShader_VirtualTextureMaterialDraw::FParameters* PassParameters = GraphBuilder.AllocParameters<FShader_VirtualTextureMaterialDraw::FParameters>();
 			PassParameters->View = View->ViewUniformBuffer;
-			PassParameters->Scene = SceneRenderer->GetSceneUniformBufferRef(GraphBuilder);
+			PassParameters->Scene = BatchDesc.SceneRenderer->GetSceneUniformBufferRef(GraphBuilder);
 
 			FRuntimeVirtualTexturePassParameters* RuntimeVirtualTexturePassParameters = GraphBuilder.AllocParameters<FRuntimeVirtualTexturePassParameters>();
 			RuntimeVirtualTexturePassParameters->MipLevel = MipLevelParameter;
-			RuntimeVirtualTexturePassParameters->CustomMaterialData = CustomMaterialData;
-			RuntimeVirtualTexturePassParameters->DebugParams = GetDebugMipLevelColor(vLevel, FixedColor);
+			RuntimeVirtualTexturePassParameters->CustomMaterialData = BatchDesc.CustomMaterialData;
+			RuntimeVirtualTexturePassParameters->DebugParams = GetDebugMipLevelColor(vLevel, BatchDesc.FixedColor);
 			RuntimeVirtualTexturePassParameters->PackHeight = FVector2f(WorldHeightPackParameter);	// LWC_TODO: Precision loss
  			PassParameters->RuntimeVirtualTexturePassParameters = GraphBuilder.CreateUniformBuffer(RuntimeVirtualTexturePassParameters);
 
-			// When direct aliasing the output physical texture for compression, the texture gets managed by RDG.
-			// The texture needs to be transitioned to UAV during the compression pass.
-			// But we may also be accessing the physical texture in a material in this render pass (using the raw RHI SRV).
-			// So we need to declare the physical textures as SRV here to avoid RDG hoisting their UAV transition to before this pass.
-			const bool bFixEarlyUAVTransition = GraphSetup.bCompressPass && GraphSetup.bDirectAliasing;
-			PassParameters->PhysicalTexture0 = bFixEarlyUAVTransition ? GraphSetup.CompressTexture0 : nullptr;
-			PassParameters->PhysicalTexture1 = bFixEarlyUAVTransition ? GraphSetup.CompressTexture1 : nullptr;
-			PassParameters->PhysicalTexture2 = bFixEarlyUAVTransition ? GraphSetup.CompressTexture2 : nullptr;
-			PassParameters->RenderTargets[0] = GraphSetup.RenderTexture0 ? FRenderTargetBinding(GraphSetup.RenderTexture0, LoadAction) : FRenderTargetBinding();
-			PassParameters->RenderTargets[1] = GraphSetup.RenderTexture1 ? FRenderTargetBinding(GraphSetup.RenderTexture1, LoadAction) : FRenderTargetBinding();
-			PassParameters->RenderTargets[2] = GraphSetup.RenderTexture2 ? FRenderTargetBinding(GraphSetup.RenderTexture2, LoadAction) : FRenderTargetBinding();
+			PassParameters->RenderTargets[0] = GraphSetup.RenderTexture0 ? FRenderTargetBinding(GraphSetup.RenderTexture0, LoadAction, 0, PageIndex) : FRenderTargetBinding();
+			PassParameters->RenderTargets[1] = GraphSetup.RenderTexture1 ? FRenderTargetBinding(GraphSetup.RenderTexture1, LoadAction, 0, PageIndex) : FRenderTargetBinding();
+			PassParameters->RenderTargets[2] = GraphSetup.RenderTexture2 ? FRenderTargetBinding(GraphSetup.RenderTexture2, LoadAction, 0, PageIndex) : FRenderTargetBinding();
     		
-			AddSimpleMeshPass(GraphBuilder, PassParameters, Scene, *View, nullptr, RDG_EVENT_NAME("VirtualTextureDraw"), View->ViewRect,
+			const int32 RuntimeVirtualTextureId = BatchDesc.RuntimeVirtualTextureId;
+			const ERuntimeVirtualTextureMaterialType MaterialType = BatchDesc.MaterialType;
+			const bool bAllowCachedMeshDrawCommands = BatchRenderContext.bAllowCachedMeshDrawCommands;
+
+			AddSimpleMeshPass(GraphBuilder, PassParameters, Scene, *View, nullptr, RDG_EVENT_NAME("VirtualTextureDraw"), View->ViewRect, ERDGPassFlags::Raster | ERDGPassFlags::NeverMerge,
 	    	[Scene, View, MaterialType, RuntimeVirtualTextureId, vLevel, MaxLevel, bAllowCachedMeshDrawCommands](FDynamicPassMeshDrawListContext* DynamicMeshPassContext)
 	    	{
 	    		GatherMeshesToDraw(DynamicMeshPassContext, Scene, View, MaterialType, RuntimeVirtualTextureId, vLevel, MaxLevel, bAllowCachedMeshDrawCommands);
 	    	});
 		}
+	}
 
-		// Compression Pass
-		if (GraphSetup.bCompressPass)
-		{
-			FShader_VirtualTextureCompress::FParameters* PassParameters = GraphBuilder.AllocParameters<FShader_VirtualTextureCompress::FParameters>();
-			PassParameters->SourceRect = FUintVector4(0, 0, TextureSize.X, TextureSize.Y);
-			PassParameters->DestPos0 = GraphSetup.bDirectAliasing ? FUintVector2(DestBox0.Min.X / 4, DestBox0.Min.Y / 4) : FUintVector2(0, 0);
-			PassParameters->DestPos1 = GraphSetup.bDirectAliasing ? FUintVector2(DestBox1.Min.X / 4, DestBox1.Min.Y / 4) : FUintVector2(0, 0);
-			PassParameters->DestPos2 = GraphSetup.bDirectAliasing ? FUintVector2(DestBox2.Min.X / 4, DestBox2.Min.Y / 4) : FUintVector2(0, 0);
-			PassParameters->EtcParameters = GetEtcParametersUniformBufferRef();
-			PassParameters->AstcParameters = GetAstcParametersUniformBufferRef();
-			PassParameters->RenderTexture0 = GraphSetup.RenderTexture0;
-			PassParameters->TextureSampler0 = TStaticSamplerState<SF_Point, AM_Clamp, AM_Clamp, AM_Clamp>::GetRHI();
-			PassParameters->RenderTexture1 = GraphSetup.RenderTexture1;
-			PassParameters->TextureSampler1 = TStaticSamplerState<SF_Point, AM_Clamp, AM_Clamp, AM_Clamp>::GetRHI();
-			PassParameters->RenderTexture2 = GraphSetup.RenderTexture2;
-			PassParameters->TextureSampler2 = TStaticSamplerState<SF_Point, AM_Clamp, AM_Clamp, AM_Clamp>::GetRHI();
-			PassParameters->OutCompressTexture0_64bit = GraphSetup.CompressTextureUAV0_64bit;
-			PassParameters->OutCompressTexture1_64bit = GraphSetup.CompressTextureUAV1_64bit;
-			PassParameters->OutCompressTexture2_64bit = GraphSetup.CompressTextureUAV2_64bit;
-			PassParameters->OutCompressTexture0_128bit = GraphSetup.CompressTextureUAV0_128bit;
-			PassParameters->OutCompressTexture1_128bit = GraphSetup.CompressTextureUAV1_128bit;
-			PassParameters->OutCompressTexture2_128bit = GraphSetup.CompressTextureUAV2_128bit;
-
-			AddCompressPass(GraphBuilder, Scene->GetFeatureLevel(), PassParameters, TextureSize, MaterialType);
-		}
+	/** 
+	 * Copy a single rendered page doing any attribute packing.
+	 * This path is rarely used, since most use cases want to compress the results of rendering.
+	 * We use a pixel shader, but could use a compute shader which batches multiple pages.
+	 */
+	void CopyPage(FRDGBuilder& GraphBuilder, FBatchRenderContext const& BatchRenderContext, int32 PageIndex)
+	{
+		FRenderGraphSetup const& GraphSetup = BatchRenderContext.GraphSetup;
+		FRenderPageBatchDesc const& BatchDesc = BatchRenderContext.BatchDesc;
+		FRenderPageDesc const& PageDesc = BatchDesc.PageDescs[PageIndex];
+		const int32 ArraySlice = BatchDesc.NumPageDescs > 1 ? PageIndex : -1;
+		FScene* Scene = BatchDesc.SceneRenderer->GetScene();
+		const FIntPoint TextureSize = PageDesc.DestRect[0].Size();
 		
-		// Copy Pass
-		if (GraphSetup.bCopyPass || GraphSetup.bCopyThumbnailPass)
-		{
-			FShader_VirtualTextureCopy::FParameters* PassParameters = GraphBuilder.AllocParameters<FShader_VirtualTextureCopy::FParameters>();
-			PassParameters->RenderTargets[0] = GraphSetup.CopyTexture0 ? FRenderTargetBinding(GraphSetup.CopyTexture0, ERenderTargetLoadAction::ENoAction) : FRenderTargetBinding();
-			PassParameters->RenderTargets[1] = GraphSetup.CopyTexture1 ? FRenderTargetBinding(GraphSetup.CopyTexture1, ERenderTargetLoadAction::ENoAction) : FRenderTargetBinding();
-			PassParameters->RenderTargets[2] = GraphSetup.CopyTexture2 ? FRenderTargetBinding(GraphSetup.CopyTexture2, ERenderTargetLoadAction::ENoAction) : FRenderTargetBinding();
-			PassParameters->DestRect = FIntVector4(0, 0, TextureSize.X, TextureSize.Y);
-			PassParameters->RenderTexture0 = GraphSetup.RenderTexture0;
-			PassParameters->TextureSampler0 = TStaticSamplerState<SF_Point, AM_Clamp, AM_Clamp, AM_Clamp>::GetRHI();
-			PassParameters->RenderTexture1 = GraphSetup.RenderTexture1;
-			PassParameters->TextureSampler1 = TStaticSamplerState<SF_Point, AM_Clamp, AM_Clamp, AM_Clamp>::GetRHI();
-			PassParameters->RenderTexture2 = GraphSetup.RenderTexture2;
-			PassParameters->TextureSampler2 = TStaticSamplerState<SF_Point, AM_Clamp, AM_Clamp, AM_Clamp>::GetRHI();
-
-			if (GraphSetup.bCopyPass)
-			{
-				AddCopyPass(GraphBuilder, Scene->GetFeatureLevel(), PassParameters, TextureSize, MaterialType);
-			}
-			else
-			{
-				AddCopyThumbnailPass(GraphBuilder, Scene->GetFeatureLevel(), PassParameters, TextureSize, MaterialType);
-			}
-		}
+		FShader_VirtualTextureCopy::FParameters* PassParameters = GraphBuilder.AllocParameters<FShader_VirtualTextureCopy::FParameters>();
+		PassParameters->RenderTargets[0] = GraphSetup.CopyTexture0 ? FRenderTargetBinding(GraphSetup.CopyTexture0, ERenderTargetLoadAction::ENoAction, 0, PageIndex) : FRenderTargetBinding();
+		PassParameters->RenderTargets[1] = GraphSetup.CopyTexture1 ? FRenderTargetBinding(GraphSetup.CopyTexture1, ERenderTargetLoadAction::ENoAction, 0, PageIndex) : FRenderTargetBinding();
+		PassParameters->RenderTargets[2] = GraphSetup.CopyTexture2 ? FRenderTargetBinding(GraphSetup.CopyTexture2, ERenderTargetLoadAction::ENoAction, 0, PageIndex) : FRenderTargetBinding();
+		PassParameters->DestRect = FIntVector4(0, 0, TextureSize.X, TextureSize.Y);
+		PassParameters->RenderTexture0 = FRenderGraphSetup::CreateTextureSRV(GraphBuilder, GraphSetup.RenderTexture0, ArraySlice);
+		PassParameters->TextureSampler0 = TStaticSamplerState<SF_Point, AM_Clamp, AM_Clamp, AM_Clamp>::GetRHI();
+		PassParameters->RenderTexture1 = FRenderGraphSetup::CreateTextureSRV(GraphBuilder, GraphSetup.RenderTexture1, ArraySlice);
+		PassParameters->TextureSampler1 = TStaticSamplerState<SF_Point, AM_Clamp, AM_Clamp, AM_Clamp>::GetRHI();
+		PassParameters->RenderTexture2 = FRenderGraphSetup::CreateTextureSRV(GraphBuilder, GraphSetup.RenderTexture2, ArraySlice);
+		PassParameters->TextureSampler2 = TStaticSamplerState<SF_Point, AM_Clamp, AM_Clamp, AM_Clamp>::GetRHI();
 
-		// Copy to Output for each output texture
-		if (GraphSetup.OutputAlias0 != nullptr && OutputTexture0 != nullptr)
-		{
-			AddCopyToOutputPass(GraphBuilder, GraphSetup.OutputAlias0, OutputTexture0, OutputTextureAccessBefore0, OutputTextureAccessAfter0, DestBox0);
-		}
-		if (GraphSetup.OutputAlias1 != nullptr && OutputTexture1 != nullptr)
+		if (GraphSetup.bCopyPass)
 		{
-			AddCopyToOutputPass(GraphBuilder, GraphSetup.OutputAlias1, OutputTexture1, OutputTextureAccessBefore1, OutputTextureAccessAfter1, DestBox1);
+			AddCopyPass(GraphBuilder, Scene->GetFeatureLevel(), PassParameters, TextureSize, BatchDesc.MaterialType);
 		}
-		if (GraphSetup.OutputAlias2 != nullptr && OutputTexture2 != nullptr)
+		else
 		{
-			AddCopyToOutputPass(GraphBuilder, GraphSetup.OutputAlias2, OutputTexture2, OutputTextureAccessBefore2, OutputTextureAccessAfter2, DestBox2);
+			AddCopyThumbnailPass(GraphBuilder, Scene->GetFeatureLevel(), PassParameters, TextureSize, BatchDesc.MaterialType);
 		}
 	}
 
-	void RenderPagesInternal(FRDGBuilder& GraphBuilder, FRenderPageBatchDesc const& InDesc, ISceneRenderer* SceneRenderer, bool bAllowCachedMeshDrawCommands)
+	/** Compress all pages in a batch. */
+	void CompressPages(FRDGBuilder& GraphBuilder, FBatchRenderContext const& BatchRenderContext)
+	{
+		FRenderGraphSetup const& GraphSetup = BatchRenderContext.GraphSetup;
+		FRenderPageBatchDesc const& BatchDesc = BatchRenderContext.BatchDesc;
+
+		FScene* Scene = BatchDesc.SceneRenderer->GetScene();
+		const FIntPoint TextureSize = BatchDesc.PageDescs[0].DestRect[0].Size();
+
+		FShader_VirtualTextureCompress::FParameters* PassParameters = GraphBuilder.AllocParameters<FShader_VirtualTextureCompress::FParameters>();
+		PassParameters->SourceRect = FUintVector4(0, 0, TextureSize.X, TextureSize.Y);
+		PassParameters->EtcParameters = GetEtcParametersUniformBufferRef();
+		PassParameters->AstcParameters = GetAstcParametersUniformBufferRef();
+		PassParameters->RenderTexture0 = GraphSetup.RenderTexture0;
+		PassParameters->TextureSampler0 = TStaticSamplerState<SF_Point, AM_Clamp, AM_Clamp, AM_Clamp>::GetRHI();
+		PassParameters->RenderTexture1 = GraphSetup.RenderTexture1;
+		PassParameters->TextureSampler1 = TStaticSamplerState<SF_Point, AM_Clamp, AM_Clamp, AM_Clamp>::GetRHI();
+		PassParameters->RenderTexture2 = GraphSetup.RenderTexture2;
+		PassParameters->TextureSampler2 = TStaticSamplerState<SF_Point, AM_Clamp, AM_Clamp, AM_Clamp>::GetRHI();
+		PassParameters->OutCompressTexture0_64bit = GraphSetup.CompressTextureUAV0_64bit;
+		PassParameters->OutCompressTexture1_64bit = GraphSetup.CompressTextureUAV1_64bit;
+		PassParameters->OutCompressTexture2_64bit = GraphSetup.CompressTextureUAV2_64bit;
+		PassParameters->OutCompressTexture0_128bit = GraphSetup.CompressTextureUAV0_128bit;
+		PassParameters->OutCompressTexture1_128bit = GraphSetup.CompressTextureUAV1_128bit;
+		PassParameters->OutCompressTexture2_128bit = GraphSetup.CompressTextureUAV2_128bit;
+
+		for (int32 PageIndex = 0; PageIndex < BatchDesc.NumPageDescs; ++PageIndex)
+		{
+			FRenderPageDesc const& PageDesc = BatchDesc.PageDescs[PageIndex];
+			for (int32 LayerIndex = 0; LayerIndex < MaxTextureLayers; ++LayerIndex)
+			{
+				const int32 WriteIndex = (PageIndex * MaxTextureLayers + LayerIndex) * 2;
+
+				// Direct aliasing case assumes needs to adjust dest position for BC block size.
+				const int32 DestPosX = GraphSetup.bDirectAliasing ? PageDesc.DestRect[LayerIndex].Min.X / 4 : 0;
+				const int32 DestPosY = GraphSetup.bDirectAliasing ? PageDesc.DestRect[LayerIndex].Min.Y / 4 : 0;
+
+				GET_SCALAR_ARRAY_ELEMENT(PassParameters->DestPos, WriteIndex) = DestPosX;
+				GET_SCALAR_ARRAY_ELEMENT(PassParameters->DestPos, WriteIndex + 1) = DestPosY;
+			}
+		}
+
+		AddCompressPass(GraphBuilder, Scene->GetFeatureLevel(), PassParameters, TextureSize, BatchDesc.NumPageDescs, BatchDesc.MaterialType, GraphSetup.bDirectAliasing);
+	}
+	
+	/** Copy all pages in a batch to the final output textures. */
+	void CopyPagesToOutput(FRDGBuilder& GraphBuilder, FBatchRenderContext const& BatchRenderContext)
 	{
-		check(InDesc.NumPageDescs <= EMaxRenderPageBatch);
+		FRenderGraphSetup const& GraphSetup = BatchRenderContext.GraphSetup;
+		FRenderPageBatchDesc const& BatchDesc = BatchRenderContext.BatchDesc;
+
+		FRDGTextureRef OutputAlias[MaxTextureLayers] = { GraphSetup.OutputAlias0, GraphSetup.OutputAlias1, GraphSetup.OutputAlias2 };
+		FRHITexture* OutputTexture[MaxTextureLayers] = { BatchDesc.Targets[0].Texture, BatchDesc.Targets[1].Texture, BatchDesc.Targets[2].Texture };
+		ERHIAccess OutputTextureAccessBefore[MaxTextureLayers] = { BatchDesc.Targets[0].TextureAccessBefore, BatchDesc.Targets[1].TextureAccessBefore, BatchDesc.Targets[2].TextureAccessBefore };
+		ERHIAccess OutputTextureAccessAfter[MaxTextureLayers] = { BatchDesc.Targets[0].TextureAccessAfter, BatchDesc.Targets[1].TextureAccessAfter, BatchDesc.Targets[2].TextureAccessAfter };
 
-		if (InDesc.NumPageDescs > 0)
+		for (int32 PageIndex = 0; PageIndex < BatchDesc.NumPageDescs; ++PageIndex)
 		{
-			RenderCaptureInterface::FScopedCapture RenderCapture((RenderCaptureNextRVTPagesDraws != 0), GraphBuilder, TEXT("RenderRVTPages"));
-			RenderCaptureNextRVTPagesDraws = FMath::Max(RenderCaptureNextRVTPagesDraws - 1, 0);
+			FRenderPageDesc const& PageDesc = BatchDesc.PageDescs[PageIndex];
 
-			for (int32 PageIndex = 0; PageIndex < InDesc.NumPageDescs; ++PageIndex)
+			for (int32 LayerIndex = 0; LayerIndex < MaxTextureLayers; ++LayerIndex)
 			{
-				FRenderPageDesc const& PageDesc = InDesc.PageDescs[PageIndex];
-
-				RenderPage(
-					GraphBuilder,
-					InDesc.Scene,
-					SceneRenderer,
-					InDesc.RuntimeVirtualTextureId,
-					InDesc.MaterialType,
-					InDesc.bClearTextures,
-					InDesc.bIsThumbnails,
-					bAllowCachedMeshDrawCommands,
-					InDesc.Targets[0].Texture, InDesc.Targets[0].TextureAccessBefore, InDesc.Targets[0].TextureAccessAfter, InDesc.Targets[0].PooledRenderTarget, PageDesc.DestBox[0],
-					InDesc.Targets[1].Texture, InDesc.Targets[1].TextureAccessBefore, InDesc.Targets[1].TextureAccessAfter, InDesc.Targets[1].PooledRenderTarget, PageDesc.DestBox[1],
-					InDesc.Targets[2].Texture, InDesc.Targets[2].TextureAccessBefore, InDesc.Targets[2].TextureAccessAfter, InDesc.Targets[2].PooledRenderTarget, PageDesc.DestBox[2],
-					InDesc.UVToWorld,
-					InDesc.WorldBounds,
-					PageDesc.UVRange,
-					PageDesc.vLevel,
-					InDesc.MaxLevel,
-					InDesc.FixedColor,
-					InDesc.CustomMaterialData);
+				if (OutputAlias[LayerIndex] != nullptr && OutputTexture[LayerIndex] != nullptr)
+				{
+					const ERHIAccess TransitionBefore = PageIndex == 0 ? OutputTextureAccessBefore[LayerIndex] : ERHIAccess::Unknown;
+					const ERHIAccess TransitionAfter = PageIndex == BatchDesc.NumPageDescs - 1 ? OutputTextureAccessAfter[LayerIndex] : ERHIAccess::Unknown;
+					
+					AddCopyToOutputPass(GraphBuilder, OutputAlias[LayerIndex], OutputTexture[LayerIndex], TransitionBefore, TransitionAfter, PageIndex, PageDesc.DestRect[LayerIndex].Min);
+				}
 			}
 		}
 	}
 
-	void RenderPagesStandAlone(FRDGBuilder& GraphBuilder, FRenderPageBatchDesc const& InDesc)
+	bool IsSceneReadyToRender(FSceneInterface* Scene)
 	{
-		check(InDesc.Scene != nullptr);
-		FScenePrimitiveRenderingContextScopeHelper RenderingScope(GetRendererModule().BeginScenePrimitiveRendering(GraphBuilder, *InDesc.Scene));
+		return Scene != nullptr && Scene->GetRenderScene() != nullptr && Scene->GetRenderScene()->GPUScene.IsRendering();
+	}
 
-		// Disable MDC caching because we can't guarantee that primitives associated with the scene have been recreated (e.g. by World->SendAllEndOfFrameUpdates()).
-		const bool bAllowCachedMeshDrawCommands = false;
-		RenderPagesInternal(GraphBuilder, InDesc, RenderingScope.ScenePrimitiveRenderingContext->GetSceneRenderer(), bAllowCachedMeshDrawCommands);
+	FBatchRenderContext const* InitPageBatch(FRDGBuilder& GraphBuilder, FRenderPageBatchDesc const& InDesc)
+	{
+		FBatchRenderContext* Context = GraphBuilder.AllocObject<FBatchRenderContext>();
+		Context->GraphSetup.Init(GraphBuilder, FRenderGraphSetup::FInitDesc(InDesc));
+		Context->BatchDesc = InDesc;
+		return Context;
 	}
 
-	void RenderPages(FRDGBuilder& GraphBuilder, FRenderPageBatchDesc const& InDesc)
+	void RenderPageBatch(FRDGBuilder& GraphBuilder, FBatchRenderContext const& InContext)
 	{
-		check(InDesc.Scene != nullptr);
-		if (ensure(InDesc.Scene->GPUScene.IsRendering()))
+		FRenderGraphSetup const& GraphSetup = InContext.GraphSetup;
+		FRenderPageBatchDesc const& Desc = InContext.BatchDesc;
+
+		if (GraphSetup.bRenderPass)
+		{
+			for (int32 PageIndex = 0; PageIndex < Desc.NumPageDescs; ++PageIndex)
+			{
+				RenderPage(GraphBuilder, InContext, PageIndex);
+			}
+		}
+		
+		if (GraphSetup.bCopyPass || GraphSetup.bCopyThumbnailPass)
 		{
-			FSceneRendererBase* SceneRenderer = FSceneRendererBase::GetActiveInstance(GraphBuilder);
-			checkf(SceneRenderer, TEXT("RenderPages needs to be called from within the context of a scene renderer."));
-			checkf(SceneRenderer->Scene == InDesc.Scene, TEXT("Provided scene doesn't match the parent scene renderer's scene."));
+			for (int32 PageIndex = 0; PageIndex < Desc.NumPageDescs; ++PageIndex)
+			{
+				CopyPage(GraphBuilder, InContext, PageIndex);
+			}
+		}
 
-			const bool bAllowCachedMeshDrawCommands = true;
-			RenderPagesInternal(GraphBuilder, InDesc, SceneRenderer, bAllowCachedMeshDrawCommands);
+		// Batch compress pages now if not direct aliasing the final output texture.
+		// This can reduce the memory high water mark.
+		// If we are direct aliasing then we must defer compression to FinalizePageBatch().
+		if (GraphSetup.bCompressPass && !InContext.GraphSetup.bDirectAliasing)
+		{
+			CompressPages(GraphBuilder, InContext);
 		}
 	}
 
+	void FinalizePageBatch(FRDGBuilder& GraphBuilder, FBatchRenderContext const& InContext)
+	{
+		FRenderGraphSetup const& GraphSetup = InContext.GraphSetup;
+
+		if (GraphSetup.bCompressPass && GraphSetup.bDirectAliasing)
+		{
+			CompressPages(GraphBuilder, InContext);
+		}
+		
+		if (!GraphSetup.bDirectAliasing)
+		{
+			CopyPagesToOutput(GraphBuilder, InContext);
+		}
+	}
+
+	void RenderPages(FRDGBuilder& GraphBuilder, FRenderPageBatchDesc const& InDesc)
+	{
+		FBatchRenderContext Context;
+		Context.GraphSetup.Init(GraphBuilder, FRenderGraphSetup::FInitDesc(InDesc));
+		Context.BatchDesc = InDesc;
+
+		// Disable MDC caching for this standalone path because we can't guarantee that primitives 
+		// associated with the scene have been recreated (e.g. by World->SendAllEndOfFrameUpdates()).
+		Context.bAllowCachedMeshDrawCommands = false;
+
+		RenderPageBatch(GraphBuilder, Context);
+		FinalizePageBatch(GraphBuilder, Context);
+	}
+
 	/** This function is deprecated! */
 	uint32 GetRuntimeVirtualTextureSceneIndex_GameThread(URuntimeVirtualTextureComponent* InComponent)
 	{
diff --git a/Engine/Source/Runtime/Renderer/Private/VT/VirtualTextureSystem.cpp b/Engine/Source/Runtime/Renderer/Private/VT/VirtualTextureSystem.cpp
--- a/Engine/Source/Runtime/Renderer/Private/VT/VirtualTextureSystem.cpp
+++ b/Engine/Source/Runtime/Renderer/Private/VT/VirtualTextureSystem.cpp
@@ -1285,7 +1285,7 @@ void FVirtualTextureSystem::LoadPendingTiles(FRDGBuilder& GraphBuilder, ERHIFeat
 		// No need to sort requests, since we're submitting all of them here (no throttling)
 		AllocateResources(GraphBuilder);
 		SubmitRequests(GraphBuilder.RHICmdList, FeatureLevel, Allocator, Settings, RequestList, false);
-		FinalizeRequests(GraphBuilder);
+		FinalizeRequests(GraphBuilder, nullptr);
 
 		// Swap any failed tile locks to be picked up by next BeginUpdate()
 		TilesToLock = MoveTemp(TilesToLockForNextFrame);
@@ -2523,27 +2523,41 @@ void FVirtualTextureSystem::SubmitRequests(FRHICommandList& RHICmdList, ERHIFeat
 	}
 }
 
-void FVirtualTextureSystem::FinalizeRequests(FRDGBuilder& GraphBuilder)
+void FVirtualTextureSystem::FinalizeRequests(FRDGBuilder& GraphBuilder, ISceneRenderer* SceneRenderer)
 {
 	RDG_CSV_STAT_EXCLUSIVE_SCOPE(GraphBuilder, VirtualTextureSystem_Update);
 	TRACE_CPUPROFILER_EVENT_SCOPE(FVirtualTextureSystem::FinalizeRequests);
-	
 	RDG_GPU_STAT_SCOPE(GraphBuilder, VirtualTexture);
 
+	FRDGExternalAccessQueue ExternalAccessQueue;
+
 	// Finalize requests
 	{
 		SCOPE_CYCLE_COUNTER(STAT_ProcessRequests_Finalize);
 		RDG_EVENT_SCOPE(GraphBuilder, "VirtualTextureFinalizeRequests");
 		TRACE_CPUPROFILER_EVENT_SCOPE(FVirtualTextureSystem::Finalize);
-		
+
+		for (IVirtualTextureFinalizer* VTFinalizer : Finalizers)
+		{
+			VTFinalizer->RenderFinalize(GraphBuilder, SceneRenderer);
+		}
+
 		for (IVirtualTextureFinalizer* VTFinalizer : Finalizers)
 		{
 			VTFinalizer->Finalize(GraphBuilder);
 		}
+		
 		Finalizers.Reset();
-	}
 
-	FRDGExternalAccessQueue ExternalAccessQueue;
+		// Transition any touched physical textures
+		for (FVirtualTexturePhysicalSpace* PhysicalSpace : PhysicalSpaces)
+		{
+			if (PhysicalSpace != nullptr)
+			{
+				PhysicalSpace->FinalizeTextures(GraphBuilder, ExternalAccessQueue);
+			}
+		}
+	}
 
 	// Update page tables
 	{
@@ -2560,14 +2574,6 @@ void FVirtualTextureSystem::FinalizeRequests(FRDGBuilder& GraphBuilder)
 		}
 	}
 
-	for (FVirtualTexturePhysicalSpace* PhysicalSpace : PhysicalSpaces)
-	{
-		if (PhysicalSpace != nullptr)
-		{
-			PhysicalSpace->FinalizeTextures(GraphBuilder, ExternalAccessQueue);
-		}
-	}
-
 	ExternalAccessQueue.Submit(GraphBuilder);
 
 	Producers.NotifyRequestsCompleted();
@@ -2821,7 +2827,7 @@ void FVirtualTextureSystem::CallPendingCallbacks()
 	Producers.CallPendingCallbacks();
 }
 
-TUniquePtr<FVirtualTextureUpdater> FVirtualTextureSystem::BeginUpdate(FRDGBuilder& GraphBuilder, ERHIFeatureLevel::Type FeatureLevel, FScene* Scene, const FVirtualTextureUpdateSettings& Settings)
+TUniquePtr<FVirtualTextureUpdater> FVirtualTextureSystem::BeginUpdate(FRDGBuilder& GraphBuilder, ERHIFeatureLevel::Type FeatureLevel, ISceneRenderer* SceneRenderer, const FVirtualTextureUpdateSettings& Settings)
 {
 	check(IsInRenderingThread());
 	check(!bUpdating);
@@ -2905,6 +2911,7 @@ TUniquePtr<FVirtualTextureUpdater> FVirtualTextureSystem::BeginUpdate(FRDGBuilde
 	}
 
 	// Flush any dirty runtime virtual textures for the current scene
+	FScene* Scene = SceneRenderer != nullptr ? SceneRenderer->GetScene() : nullptr;
 	if (Scene != nullptr)
 	{
 		// Only flush if we know that there is GPU feedback available to refill the visible data this frame
@@ -2999,17 +3006,17 @@ void FVirtualTextureSystem::EndUpdate(FRDGBuilder& GraphBuilder, TUniquePtr<FVir
 	TilesToLock = MoveTemp(TilesToLockForNextFrame);
 }
 
-void FVirtualTextureSystem::Update(FRDGBuilder& GraphBuilder, ERHIFeatureLevel::Type FeatureLevel, FScene* Scene, const FVirtualTextureUpdateSettings& InSettings)
+void FVirtualTextureSystem::Update(FRDGBuilder& GraphBuilder, ERHIFeatureLevel::Type FeatureLevel, ISceneRenderer* SceneRenderer, const FVirtualTextureUpdateSettings& InSettings)
 {
 	CallPendingCallbacks();
 
 	FVirtualTextureUpdateSettings Settings = InSettings;
 	Settings.EnableAsyncTasks(false);
 
-	TUniquePtr<FVirtualTextureUpdater> Updater = BeginUpdate(GraphBuilder, FeatureLevel, Scene, Settings);
+	TUniquePtr<FVirtualTextureUpdater> Updater = BeginUpdate(GraphBuilder, FeatureLevel, SceneRenderer, Settings);
 	EndUpdate(GraphBuilder, MoveTemp(Updater), FeatureLevel);
 
-	FinalizeRequests(GraphBuilder);
+	FinalizeRequests(GraphBuilder, SceneRenderer);
 }
 
 void FVirtualTextureSystem::ReleasePendingResources()
diff --git a/Engine/Source/Runtime/Renderer/Private/VT/VirtualTextureSystem.h b/Engine/Source/Runtime/Renderer/Private/VT/VirtualTextureSystem.h
--- a/Engine/Source/Runtime/Renderer/Private/VT/VirtualTextureSystem.h
+++ b/Engine/Source/Runtime/Renderer/Private/VT/VirtualTextureSystem.h
@@ -32,6 +32,7 @@ struct FFeedbackAnalysisParameters;
 struct FAddRequestedTilesParameters;
 struct FGatherRequestsParameters;
 struct FPageUpdateBuffer;
+class ISceneRenderer;
 
 extern uint32 GetTypeHash(const FAllocatedVTDescription& Description);
 
@@ -111,16 +112,16 @@ public:
 
 	uint32 GetFrame() const { return Frame; }
 
-	void Update(FRDGBuilder& GraphBuilder, ERHIFeatureLevel::Type FeatureLevel, FScene* Scene, const FVirtualTextureUpdateSettings& Settings);
+	void Update(FRDGBuilder& GraphBuilder, ERHIFeatureLevel::Type FeatureLevel, ISceneRenderer* SceneRenderer, const FVirtualTextureUpdateSettings& Settings);
 
 	// Called in FSceneRenderer::UpdateScene, before FVirtualTextureSystem::BeginUpdate -- see comments there for more info
 	void CallPendingCallbacks();
 
-	TUniquePtr<FVirtualTextureUpdater> BeginUpdate(FRDGBuilder& GraphBuilder, ERHIFeatureLevel::Type FeatureLevel, FScene* Scene, const FVirtualTextureUpdateSettings& Settings);
+	TUniquePtr<FVirtualTextureUpdater> BeginUpdate(FRDGBuilder& GraphBuilder, ERHIFeatureLevel::Type FeatureLevel, ISceneRenderer* SceneRenderer, const FVirtualTextureUpdateSettings& Settings);
 	void WaitForTasks(FVirtualTextureUpdater* Updater);
 	void EndUpdate(FRDGBuilder& GraphBuilder, TUniquePtr<FVirtualTextureUpdater>&& Updater, ERHIFeatureLevel::Type FeatureLevel);
 
-	void FinalizeRequests(FRDGBuilder& GraphBuilder);
+	void FinalizeRequests(FRDGBuilder& GraphBuilder, ISceneRenderer* SceneRenderer);
 
 	void ReleasePendingResources();
 
diff --git a/Engine/Source/Runtime/Renderer/Public/SceneRendererInterface.h b/Engine/Source/Runtime/Renderer/Public/SceneRendererInterface.h
--- a/Engine/Source/Runtime/Renderer/Public/SceneRendererInterface.h
+++ b/Engine/Source/Runtime/Renderer/Public/SceneRendererInterface.h
@@ -5,6 +5,7 @@
 #include "SceneView.h"
 #include "ShaderParameterStructDeclaration.h"
 
+class FScene;
 class FSceneUniformBuffer;
 class FPrimitiveSceneInfo;
 struct FPersistentPrimitiveIndex;
@@ -47,6 +48,7 @@ class ISceneRenderer
 public:
 	virtual ~ISceneRenderer() = default;
 
+	virtual FScene* GetScene() = 0;
 	virtual const FSceneUniformBuffer& GetSceneUniforms() const = 0;
 	virtual FSceneUniformBuffer& GetSceneUniforms() = 0;
 
diff --git a/Engine/Source/Runtime/Renderer/Public/VT/RuntimeVirtualTextureRender.h b/Engine/Source/Runtime/Renderer/Public/VT/RuntimeVirtualTextureRender.h
--- a/Engine/Source/Runtime/Renderer/Public/VT/RuntimeVirtualTextureRender.h
+++ b/Engine/Source/Runtime/Renderer/Public/VT/RuntimeVirtualTextureRender.h
@@ -3,20 +3,21 @@
 #pragma once
 
 #include "CoreMinimal.h"
-#include "VT/RuntimeVirtualTextureEnum.h"
 #include "RHIAccess.h"
+#include "VT/RuntimeVirtualTextureEnum.h"
 
 class FRHITexture;
 class FRDGBuilder;
 class FScene;
 class FSceneInterface;
 struct IPooledRenderTarget;
+class ISceneRenderer;
 class URuntimeVirtualTextureComponent;
 
 namespace RuntimeVirtualTexture
 {
 	/** Enum for our maximum RenderPages() batch size. */
-	enum { EMaxRenderPageBatch = 8 };
+	enum { MaxRenderPageBatch = 8 };
 
 	/** Structure containing a texture layer target description for a call for RenderPages(). */
 	struct FRenderPageTarget
@@ -41,13 +42,22 @@ namespace RuntimeVirtualTexture
 		/** UV range to render in virtual texture space. */
 		FBox2D UVRange;
 		/** Destination box to render in texel space of the target physical texture. */
-		FBox2D DestBox[RuntimeVirtualTexture::MaxTextureLayers];
+		FIntRect DestRect[MaxTextureLayers];
+		UE_DEPRECATED(5.6, "Use DestRect instead.")
+		FBox2D DestBox[MaxTextureLayers];
+	
+		// Disable deprecation warnings for implicit copying of deprecated members.
+		PRAGMA_DISABLE_DEPRECATION_WARNINGS
+		FRenderPageDesc& operator=(FRenderPageDesc const& RHS) = default;
+		PRAGMA_ENABLE_DEPRECATION_WARNINGS
 	};
 
 	/** A description of a batch of pages to be rendered with a single call to RenderPages(). */
 	struct FRenderPageBatchDesc
 	{
 		/** Scene to use when rendering the batch. */
+		ISceneRenderer* SceneRenderer;
+		UE_DEPRECATED(5.6, "Use SceneRenderer instead.")
 		FScene* Scene;
 		/** Unique object ID of the runtime virtual texture that we are rendering. */
 		int32 RuntimeVirtualTextureId;
@@ -71,25 +81,40 @@ namespace RuntimeVirtualTexture
 		FVector4f CustomMaterialData;
 
 		/** Physical texture targets to render to. */
-		FRenderPageTarget Targets[RuntimeVirtualTexture::MaxTextureLayers];
+		FRenderPageTarget Targets[MaxTextureLayers];
 		
 		/** Number of pages to render. */
 		int32 NumPageDescs;
 		/** Page descriptions for each page in the batch. */
-		FRenderPageDesc PageDescs[EMaxRenderPageBatch];
+		FRenderPageDesc PageDescs[MaxRenderPageBatch];
+
+		// Disable deprecation warnings for implicit copying of deprecated members.
+		PRAGMA_DISABLE_DEPRECATION_WARNINGS
+		FRenderPageBatchDesc& operator=(FRenderPageBatchDesc const& RHS) = default;
+		PRAGMA_ENABLE_DEPRECATION_WARNINGS
 	};
 
+	/** Opaque class with context for rendering a batch of pages. */
+	class FBatchRenderContext;
+
 	/** Returns true if the scene is initialized for rendering to runtime virtual textures. Always check this before calling RenderPages(). */
 	RENDERER_API bool IsSceneReadyToRender(FSceneInterface* Scene);
 
 	/** Render a batch of pages for a runtime virtual texture. */
 	RENDERER_API void RenderPages(FRDGBuilder& GraphBuilder, FRenderPageBatchDesc const& InDesc);
 
-	/** 
-	 * Render a batch of pages for a runtime virtual texture. 
-	 * Performs the required scene rendering setup such that it can be called from a render thread task.
-	 */
-	RENDERER_API void RenderPagesStandAlone(FRDGBuilder& GraphBuilder, FRenderPageBatchDesc const& InDesc);
+	/** Get a context for rendering a batch of pages. */
+	FBatchRenderContext const* InitPageBatch(FRDGBuilder& GraphBuilder, FRenderPageBatchDesc const& InDesc);
+	/** Do rendering for all pages in a context. */
+	void RenderPageBatch(FRDGBuilder& GraphBuilder, FBatchRenderContext const& InBatch);
+	/** Finalize all pages in a context to their final physical location. */
+	void FinalizePageBatch(FRDGBuilder& GraphBuilder, FBatchRenderContext const& InBatch);
+
+	UE_DEPRECATED(5.6, "Use RenderPages().")
+	inline void RenderPagesStandAlone(FRDGBuilder& GraphBuilder, FRenderPageBatchDesc const& InDesc)
+	{
+		RenderPages(GraphBuilder, InDesc);
+	}
 
 #if WITH_EDITOR
 	UE_DEPRECATED(5.6, "Use RuntimeVirtualTextureId everywhere that you would previously use a SceneIndex.")
