PR #13079: Fixed CalculateNearFarDepthExcludingSky()

#tests pc editor dx12 with r.D3D12.Depth24Bit=0/1
#rnx

diff --git a/Engine/Source/Runtime/Renderer/Private/RendererUtils.cpp b/Engine/Source/Runtime/Renderer/Private/RendererUtils.cpp
--- a/Engine/Source/Runtime/Renderer/Private/RendererUtils.cpp
+++ b/Engine/Source/Runtime/Renderer/Private/RendererUtils.cpp
@@ -201,16 +201,32 @@ FDepthBounds::FDepthBoundsValues FDepthBounds::CalculateNearFarDepthExcludingSky
 
 	if (bool(ERHIZBuffer::IsInverted))
 	{
-		//const float SmallestFloatAbove0 = 1.1754943508e-38;		// 32bit float depth
-		const float SmallestFloatAbove0 = 1.0f / 16777215.0f;		// 24bit norm depth
+		float SmallestFloatAbove0;
+
+		if (GPixelFormats[PF_DepthStencil].bIs24BitUnormDepthStencil)
+		{
+			SmallestFloatAbove0 = 1.0f / 16777215.0f;			// 24bit norm depth
+		}
+		else
+		{
+			SmallestFloatAbove0 = 1.1754943508e-38;				// 32bit float depth
+		}
 
 		Values.MinDepth = SmallestFloatAbove0;
 		Values.MaxDepth = float(ERHIZBuffer::NearPlane);
 	}
 	else
 	{
-		//const float SmallestFloatBelow1 = 0.9999999404;			// 32bit float depth
-		const float SmallestFloatBelow1 = 16777214.0f / 16777215.0f;// 24bit norm depth
+		float SmallestFloatBelow1;
+
+		if (GPixelFormats[PF_DepthStencil].bIs24BitUnormDepthStencil)
+		{
+			SmallestFloatBelow1 = 16777214.0f / 16777215.0f;	// 24bit norm depth
+		}
+		else
+		{
+			SmallestFloatBelow1 = 0.9999999404;					// 32bit float depth
+		}
 
 		Values.MinDepth = float(ERHIZBuffer::NearPlane);
 		Values.MaxDepth = SmallestFloatBelow1;
