Fix for ReflectionEnvironmentAndSky pass to leverage depth bound test (depth buffer is now bound).

#tests pc editor dx12 with r.D3D12.Depth24Bit=0/1
#rnx

diff --git a/Engine/Source/Runtime/Renderer/Private/IndirectLightRendering.cpp b/Engine/Source/Runtime/Renderer/Private/IndirectLightRendering.cpp
--- a/Engine/Source/Runtime/Renderer/Private/IndirectLightRendering.cpp
+++ b/Engine/Source/Runtime/Renderer/Private/IndirectLightRendering.cpp
@@ -1847,6 +1847,13 @@ static void AddSkyReflectionPass(
 		PassParameters->RenderTargets[2] = FRenderTargetBinding(Scene->SubstrateSceneData.SeparatedSubSurfaceSceneColor, ERenderTargetLoadAction::ELoad);
 	}
 
+	const bool bUseDepthBounds = GSupportsDepthBoundsTest;
+	if (bUseDepthBounds)
+	{
+		PassParameters->RenderTargets.DepthStencil = FDepthStencilBinding(SceneTextures.Depth.Target,
+			ERenderTargetLoadAction::ELoad, ERenderTargetLoadAction::ELoad, FExclusiveDepthStencil::DepthRead_StencilRead);
+	}
+
 	// Bind hair data
 	const bool bCheckerboardSubsurfaceRendering = IsSubsurfaceCheckerboardFormat(SceneColorTexture.Target->Desc.Format, View);
 	const bool bHasHairCardsOrMeshes = View.HairCardsMeshElements.Num() > 0;
@@ -1881,7 +1888,7 @@ static void AddSkyReflectionPass(
 			bSubstrateEnabled ? ToString(SubstrateTileMaterialType) : TEXT("Off")),
 		PassParameters,
 		ERDGPassFlags::Raster,
-		[PassParameters, &View, PixelShader, bCheckerboardSubsurfaceRendering, SubstrateTileMaterialType,
+		[PassParameters, &View, PixelShader, bCheckerboardSubsurfaceRendering, SubstrateTileMaterialType, bUseDepthBounds,
 		SubstrateTilePassVertexShader, bSubstrateEnabled, SubstrateTilePrimitiveType](FRDGAsyncTask, FRHICommandList& InRHICmdList)
 	{
 		InRHICmdList.SetViewport(View.ViewRect.Min.X, View.ViewRect.Min.Y, 0.0f, View.ViewRect.Max.X, View.ViewRect.Max.Y, 1.0f);
@@ -1917,7 +1924,7 @@ static void AddSkyReflectionPass(
 			}
 		}
 
-		if (GSupportsDepthBoundsTest)
+		if (bUseDepthBounds)
 		{
 			// We do not want to process sky pixels so we take advantage of depth bound test when available to skip launching pointless GPU wavefront/work.
 			GraphicsPSOInit.bDepthBounds = true;
