Custom Render Passes now disable MSAA.  Fixes bug where Water CRPs produce a crash when MSAA is enabled.

diff --git a/Engine/Source/Runtime/Engine/Private/SceneTexturesConfig.cpp b/Engine/Source/Runtime/Engine/Private/SceneTexturesConfig.cpp
--- a/Engine/Source/Runtime/Engine/Private/SceneTexturesConfig.cpp
+++ b/Engine/Source/Runtime/Engine/Private/SceneTexturesConfig.cpp
@@ -280,7 +280,12 @@ void FSceneTexturesConfig::Init(const FSceneTexturesConfigInitSettings& InitSett
 	ShadingPath				= GetFeatureLevelShadingPath(FeatureLevel);
 	ShaderPlatform			= GetFeatureLevelShaderPlatform(FeatureLevel);
 	Extent					= InitSettings.Extent;
+
+	// Note that the GetDefaultMSAACount call is mirrored in the FSceneRenderer constructor where custom render passes are initialized.
+	// Custom render passes don't support MSAA, and use a different config with MSAA disabled, when its enabled for the main view.  If
+	// this calculation of NumSamples changes, that corresponding code would need to be updated.
 	NumSamples				= GetDefaultMSAACount(FeatureLevel, GDynamicRHI->RHIGetPlatformTextureMaxSampleCount());
+
 	EditorPrimitiveNumSamples = GetEditorPrimitiveNumSamples(FeatureLevel);
 	ColorFormat				= PF_Unknown;
 	ColorClearValue			= FClearValueBinding::Black;
diff --git a/Engine/Source/Runtime/Renderer/Private/DeferredShadingRenderer.cpp b/Engine/Source/Runtime/Renderer/Private/DeferredShadingRenderer.cpp
--- a/Engine/Source/Runtime/Renderer/Private/DeferredShadingRenderer.cpp
+++ b/Engine/Source/Runtime/Renderer/Private/DeferredShadingRenderer.cpp
@@ -2066,7 +2066,7 @@ void FDeferredShadingSceneRenderer::Render(FRDGBuilder& GraphBuilder, const FSce
 	const bool bHairStrandsEnable = HairStrandsBookmarkParameters.HasInstances() && Views.Num() > 0 && IsHairStrandsEnabled(EHairStrandsShaderType::Strands, Platform);
 	const bool bForceVelocityOutput = bHairStrandsEnable || ShouldRenderDistortion();
 
-	auto RenderPrepassAndVelocity = [&](auto& InViews, auto& InNaniteBasePassVisibility, auto& NaniteRasterResults, auto& PrimaryNaniteViews)
+	auto RenderPrepassAndVelocity = [&](auto& InViews, auto& InNaniteBasePassVisibility, auto& NaniteRasterResults, auto& PrimaryNaniteViews, FSceneTextures& LocalSceneTextures)
 	{
 		if (FGPUSkinCache* GPUSkinCache = Scene->GetGPUSkinCache())
 		{
@@ -2081,24 +2081,24 @@ void FDeferredShadingSceneRenderer::Render(FRDGBuilder& GraphBuilder, const FSce
 				: ERenderTargetLoadAction::EClear;
 
 			const ERenderTargetLoadAction DepthLoadAction = ERenderTargetLoadAction::EClear;
-			AddClearDepthStencilPass(GraphBuilder, SceneTextures.Depth.Target, DepthLoadAction, StencilLoadAction);
+			AddClearDepthStencilPass(GraphBuilder, LocalSceneTextures.Depth.Target, DepthLoadAction, StencilLoadAction);
 
 			// Draw the scene pre-pass / early z pass, populating the scene depth buffer and HiZ
 			if (bNeedsPrePass)
 			{
-				RenderPrePass(GraphBuilder, InViews, SceneTextures.Depth.Target, InstanceCullingManager, &FirstStageDepthBuffer);
+				RenderPrePass(GraphBuilder, InViews, LocalSceneTextures.Depth.Target, InstanceCullingManager, &FirstStageDepthBuffer);
 			}
 			else
 			{
 				// We didn't do the prepass, but we still want the HMD mask if there is one
-				RenderPrePassHMD(GraphBuilder, InViews, SceneTextures.Depth.Target);
+				RenderPrePassHMD(GraphBuilder, InViews, LocalSceneTextures.Depth.Target);
 			}
 
 			// special pass for DDM_AllOpaqueNoVelocity, which uses the velocity pass to finish the early depth pass write
 			if (bShouldRenderVelocities && Scene->EarlyZPassMode == DDM_AllOpaqueNoVelocity && RendererOutput == ERendererOutput::FinalSceneColor)
 			{
 				// Render the velocities of movable objects.  Don't bind the velocity render target for custom render passes (it's not used downstream), to avoid needing to clear it again.
-				RenderVelocities(GraphBuilder, InViews, SceneTextures, EVelocityPass::Opaque, bForceVelocityOutput, /*bBindRenderTarget=*/ InViews[0].CustomRenderPass == nullptr);
+				RenderVelocities(GraphBuilder, InViews, LocalSceneTextures, EVelocityPass::Opaque, bForceVelocityOutput, /*bBindRenderTarget=*/ InViews[0].CustomRenderPass == nullptr);
 			}
 		}
 
@@ -2107,24 +2107,24 @@ void FDeferredShadingSceneRenderer::Render(FRDGBuilder& GraphBuilder, const FSce
 
 			if (bNaniteEnabled && InViews.Num() > 0)
 			{
-				RenderNanite(GraphBuilder, InViews, SceneTextures, bIsEarlyDepthComplete, InNaniteBasePassVisibility, NaniteRasterResults, PrimaryNaniteViews, FirstStageDepthBuffer);
+				RenderNanite(GraphBuilder, InViews, LocalSceneTextures, bIsEarlyDepthComplete, InNaniteBasePassVisibility, NaniteRasterResults, PrimaryNaniteViews, FirstStageDepthBuffer);
 			}
 		}
 
 		if (FirstStageDepthBuffer)
 		{
-			SceneTextures.PartialDepth = FirstStageDepthBuffer;
-			AddResolveSceneDepthPass(GraphBuilder, InViews, SceneTextures.PartialDepth);
+			LocalSceneTextures.PartialDepth = FirstStageDepthBuffer;
+			AddResolveSceneDepthPass(GraphBuilder, InViews, LocalSceneTextures.PartialDepth);
 		}
 		else
 		{
 			// Setup default partial depth to be scene depth so that it also works on transparent emitter when partial depth has not been generated.
-			SceneTextures.PartialDepth = SceneTextures.Depth;
+			LocalSceneTextures.PartialDepth = LocalSceneTextures.Depth;
 		}
-		SceneTextures.SetupMode = ESceneTextureSetupMode::SceneDepth;
-		SceneTextures.UniformBuffer = CreateSceneTextureUniformBuffer(GraphBuilder, &SceneTextures, FeatureLevel, SceneTextures.SetupMode);
+		LocalSceneTextures.SetupMode = ESceneTextureSetupMode::SceneDepth;
+		LocalSceneTextures.UniformBuffer = CreateSceneTextureUniformBuffer(GraphBuilder, &LocalSceneTextures, FeatureLevel, LocalSceneTextures.SetupMode);
 
-		AddResolveSceneDepthPass(GraphBuilder, InViews, SceneTextures.Depth);
+		AddResolveSceneDepthPass(GraphBuilder, InViews, LocalSceneTextures.Depth);
 	};
 
 	FDBufferTextures DBufferTextures = CreateDBufferTextures(GraphBuilder, SceneTextures.Config.Extent, ShaderPlatform);
@@ -2139,10 +2139,22 @@ void FDeferredShadingSceneRenderer::Render(FRDGBuilder& GraphBuilder, const FSce
 		RDG_EVENT_SCOPE_STAT(GraphBuilder, CustomRenderPasses, "CustomRenderPasses");
 		RDG_GPU_STAT_SCOPE(GraphBuilder, CustomRenderPasses);
 
+		// If the main view family has MSAA enabled, initialize and use the separate non-MSAA version of the FSceneTextures stored in the first
+		// custom render pass (also used by the other custom render passes).
+		FSceneTextures* CustomRenderPassSceneTextures = &SceneTextures;
+		if (SceneTextures.Config.NumSamples > 1)
+		{
+			FSceneTextures::InitializeViewFamily(GraphBuilder, CustomRenderPassInfos[0].ViewFamily, FamilySize);
+			CustomRenderPassSceneTextures = const_cast<FSceneTextures*>(&CustomRenderPassInfos[0].ViewFamily.GetSceneTextures());
+			
+			// Make sure a separate FSceneTextures structure was allocated in the FSceneRenderer constructor when custom render passes were initialized!
+			check(CustomRenderPassSceneTextures != &SceneTextures);
+		}
+
 		// We want to reset the scene texture uniform buffer to its original state after custom render passes,
 		// so they can't affect downstream rendering.
-		ESceneTextureSetupMode OriginalSceneTextureSetupMode = SceneTextures.SetupMode;
-		TRDGUniformBufferRef<FSceneTextureUniformParameters> OriginalSceneTextureUniformBuffer = SceneTextures.UniformBuffer;
+		ESceneTextureSetupMode OriginalSceneTextureSetupMode = CustomRenderPassSceneTextures->SetupMode;
+		TRDGUniformBufferRef<FSceneTextureUniformParameters> OriginalSceneTextureUniformBuffer = CustomRenderPassSceneTextures->UniformBuffer;
 
 		for (int32 i = 0; i < CustomRenderPassInfos.Num(); ++i)
 		{
@@ -2162,34 +2174,34 @@ void FDeferredShadingSceneRenderer::Render(FRDGBuilder& GraphBuilder, const FSce
 				TArray<Nanite::FRasterResults, TInlineAllocator<2>> NaniteRasterResults;
 				TArray<Nanite::FPackedView, SceneRenderingAllocator> PrimaryNaniteViews;
 				FNaniteBasePassVisibility DummyNaniteBasePassVisibility;
-				RenderPrepassAndVelocity(CustomRenderPassViews, DummyNaniteBasePassVisibility, NaniteRasterResults, PrimaryNaniteViews);
+				RenderPrepassAndVelocity(CustomRenderPassViews, DummyNaniteBasePassVisibility, NaniteRasterResults, PrimaryNaniteViews, *CustomRenderPassSceneTextures);
 
 				const FSingleLayerWaterPrePassResult* SingleLayerWaterPrePassResult = nullptr;
 				if (ShouldRenderSingleLayerWaterDepthPrepass(CustomRenderPassViews))
 				{
-					SingleLayerWaterPrePassResult = RenderSingleLayerWaterDepthPrepass(GraphBuilder, CustomRenderPassViews, SceneTextures, ESingleLayerWaterPrepassLocation::BeforeBasePass, NaniteRasterResults);
+					SingleLayerWaterPrePassResult = RenderSingleLayerWaterDepthPrepass(GraphBuilder, CustomRenderPassViews, *CustomRenderPassSceneTextures, ESingleLayerWaterPrepassLocation::BeforeBasePass, NaniteRasterResults);
 				}
 
 				const FSceneCaptureCustomRenderPassUserData& SceneCaptureUserData = FSceneCaptureCustomRenderPassUserData::Get(CustomRenderPass);
 
 				if (CustomRenderPass->GetRenderMode() == FCustomRenderPassBase::ERenderMode::DepthAndBasePass)
 				{
-					SceneTextures.SetupMode |= ESceneTextureSetupMode::SceneColor;
-					SceneTextures.UniformBuffer = CreateSceneTextureUniformBuffer(GraphBuilder, &SceneTextures, FeatureLevel, SceneTextures.SetupMode);
+					CustomRenderPassSceneTextures->SetupMode |= ESceneTextureSetupMode::SceneColor;
+					CustomRenderPassSceneTextures->UniformBuffer = CreateSceneTextureUniformBuffer(GraphBuilder, CustomRenderPassSceneTextures, FeatureLevel, CustomRenderPassSceneTextures->SetupMode);
 
 					if (bNaniteEnabled)
 					{
 						Nanite::BuildShadingCommands(GraphBuilder, *Scene, ENaniteMeshPass::BasePass, NaniteBasePassShadingCommands, Nanite::EBuildShadingCommandsMode::Custom);
 					}
 
-					RenderBasePass(*this, GraphBuilder, CustomRenderPassViews, SceneTextures, DBufferTextures, BasePassDepthStencilAccess, /*ForwardScreenSpaceShadowMaskTexture=*/nullptr, InstanceCullingManager, bNaniteEnabled, NaniteBasePassShadingCommands, NaniteRasterResults);
+					RenderBasePass(*this, GraphBuilder, CustomRenderPassViews, *CustomRenderPassSceneTextures, DBufferTextures, BasePassDepthStencilAccess, /*ForwardScreenSpaceShadowMaskTexture=*/nullptr, InstanceCullingManager, bNaniteEnabled, NaniteBasePassShadingCommands, NaniteRasterResults);
 
 					if (ShouldRenderSingleLayerWater(CustomRenderPassViews))
 					{
 						// GBuffer code paths in RenderSingleLayerWater don't use the bIsCameraUnderWater flag, so just pass in false.  Normally this is
 						// computed by a render extension, but those aren't run for custom render passes.
 						FSceneWithoutWaterTextures SceneWithoutWaterTextures;
-						RenderSingleLayerWater(GraphBuilder, CustomRenderPassViews, SceneTextures, SingleLayerWaterPrePassResult, /*bShouldRenderVolumetricCloud=*/false, SceneWithoutWaterTextures, LumenFrameTemporaries, /*bIsCameraUnderWater=*/false);
+						RenderSingleLayerWater(GraphBuilder, CustomRenderPassViews, *CustomRenderPassSceneTextures, SingleLayerWaterPrePassResult, /*bShouldRenderVolumetricCloud=*/false, SceneWithoutWaterTextures, LumenFrameTemporaries, /*bIsCameraUnderWater=*/false);
 					}
 
 					FCustomRenderPassBase::ERenderOutput RenderOutput = CustomRenderPass->GetRenderOutput();
@@ -2197,8 +2209,8 @@ void FDeferredShadingSceneRenderer::Render(FRDGBuilder& GraphBuilder, const FSce
 						!SceneCaptureUserData.UserSceneTextureBaseColor.IsNone() || !SceneCaptureUserData.UserSceneTextureNormal.IsNone() || !SceneCaptureUserData.UserSceneTextureSceneColor.IsNone())
 					{
 						// CopySceneCaptureComponentToTarget uses scene texture uniforms
-						SceneTextures.SetupMode |= ESceneTextureSetupMode::GBuffers;
-						SceneTextures.UniformBuffer = CreateSceneTextureUniformBuffer(GraphBuilder, &SceneTextures, FeatureLevel, SceneTextures.SetupMode);
+						CustomRenderPassSceneTextures->SetupMode |= ESceneTextureSetupMode::GBuffers;
+						CustomRenderPassSceneTextures->UniformBuffer = CreateSceneTextureUniformBuffer(GraphBuilder, CustomRenderPassSceneTextures, FeatureLevel, CustomRenderPassSceneTextures->SetupMode);
 					}
 
 					if (CustomRenderPass->IsTranslucentIncluded())
@@ -2217,14 +2229,15 @@ void FDeferredShadingSceneRenderer::Render(FRDGBuilder& GraphBuilder, const FSce
 							View.ReflectionCaptureUniformBuffer = EmptyReflectionCaptureUniformBuffer;
 						}
 
-						RenderTranslucency(*this, GraphBuilder, SceneTextures, TranslucencyLightingVolumeTextures, &TranslucencyResourceMap, CustomRenderPassViews, ETranslucencyView::AboveWater, CustomTranslucencyDimensions, InstanceCullingManager, bStandardTranslucentCanRenderSeparate, TranslucencySharedDepthTexture);
+						RenderTranslucency(*this, GraphBuilder, *CustomRenderPassSceneTextures, TranslucencyLightingVolumeTextures, &TranslucencyResourceMap, CustomRenderPassViews, ETranslucencyView::AboveWater, CustomTranslucencyDimensions, InstanceCullingManager, bStandardTranslucentCanRenderSeparate, TranslucencySharedDepthTexture);
 					}
 				}
 
-				CopySceneCaptureComponentToTarget(GraphBuilder, SceneTextures, CustomRenderPass->GetRenderTargetTexture(), ViewFamily, CustomRenderPassViews);
+				CopySceneCaptureComponentToTarget(GraphBuilder, *CustomRenderPassSceneTextures, CustomRenderPass->GetRenderTargetTexture(), ViewFamily, CustomRenderPassViews);
 
 				if (!SceneCaptureUserData.UserSceneTextureBaseColor.IsNone())
 				{
+					// User Scene Textures are stored to "SceneTextures" for downstream use, not in "CustomRenderPassSceneTextures", only used during Custom Render Pass rendering
 					bool bFirstRender;
 					FRDGTextureRef BaseColorSceneTexture = SceneTextures.FindOrAddUserSceneTexture(GraphBuilder, 0, SceneCaptureUserData.UserSceneTextureBaseColor, SceneCaptureUserData.SceneTextureDivisor, bFirstRender, nullptr, CustomRenderPassViews[0].ViewRect);
 #if !(UE_BUILD_SHIPPING)
@@ -2232,7 +2245,7 @@ void FDeferredShadingSceneRenderer::Render(FRDGBuilder& GraphBuilder, const FSce
 #endif
 
 					CustomRenderPass->OverrideRenderOutput(FCustomRenderPassBase::ERenderOutput::BaseColor);
-					CopySceneCaptureComponentToTarget(GraphBuilder, SceneTextures, BaseColorSceneTexture, ViewFamily, CustomRenderPassViews);
+					CopySceneCaptureComponentToTarget(GraphBuilder, *CustomRenderPassSceneTextures, BaseColorSceneTexture, ViewFamily, CustomRenderPassViews);
 				}
 
 				if (!SceneCaptureUserData.UserSceneTextureNormal.IsNone())
@@ -2244,7 +2257,7 @@ void FDeferredShadingSceneRenderer::Render(FRDGBuilder& GraphBuilder, const FSce
 #endif
 
 					CustomRenderPass->OverrideRenderOutput(FCustomRenderPassBase::ERenderOutput::Normal);
-					CopySceneCaptureComponentToTarget(GraphBuilder, SceneTextures, NormalSceneTexture, ViewFamily, CustomRenderPassViews);
+					CopySceneCaptureComponentToTarget(GraphBuilder, *CustomRenderPassSceneTextures, NormalSceneTexture, ViewFamily, CustomRenderPassViews);
 				}
 
 				if (!SceneCaptureUserData.UserSceneTextureSceneColor.IsNone())
@@ -2256,7 +2269,7 @@ void FDeferredShadingSceneRenderer::Render(FRDGBuilder& GraphBuilder, const FSce
 #endif
 
 					CustomRenderPass->OverrideRenderOutput(FCustomRenderPassBase::ERenderOutput::SceneColorAndAlpha);
-					CopySceneCaptureComponentToTarget(GraphBuilder, SceneTextures, SceneColorSceneTexture, ViewFamily, CustomRenderPassViews);
+					CopySceneCaptureComponentToTarget(GraphBuilder, *CustomRenderPassSceneTextures, SceneColorSceneTexture, ViewFamily, CustomRenderPassViews);
 				}
 
 				CustomRenderPass->PostRender(GraphBuilder);
@@ -2276,14 +2289,14 @@ void FDeferredShadingSceneRenderer::Render(FRDGBuilder& GraphBuilder, const FSce
 			CustomRenderPass->EndPass(GraphBuilder);
 
 			// Restore original scene texture uniforms
-			SceneTextures.SetupMode = OriginalSceneTextureSetupMode;
-			SceneTextures.UniformBuffer = OriginalSceneTextureUniformBuffer;
+			CustomRenderPassSceneTextures->SetupMode = OriginalSceneTextureSetupMode;
+			CustomRenderPassSceneTextures->UniformBuffer = OriginalSceneTextureUniformBuffer;
 		}
 	}
 
 	TArray<Nanite::FRasterResults, TInlineAllocator<2>> NaniteRasterResults;
 	TArray<Nanite::FPackedView, SceneRenderingAllocator> PrimaryNaniteViews;
-	RenderPrepassAndVelocity(Views, NaniteBasePassVisibility, NaniteRasterResults, PrimaryNaniteViews);
+	RenderPrepassAndVelocity(Views, NaniteBasePassVisibility, NaniteRasterResults, PrimaryNaniteViews, SceneTextures);
 
 	// Run Nanite compute commands early in the frame to allow some task overlap on the CPU until the base pass runs.
 	if (bNaniteEnabled && RendererOutput != ERendererOutput::DepthPrepassOnly && !bHasRayTracedOverlay)
diff --git a/Engine/Source/Runtime/Renderer/Private/SceneRendering.cpp b/Engine/Source/Runtime/Renderer/Private/SceneRendering.cpp
--- a/Engine/Source/Runtime/Renderer/Private/SceneRendering.cpp
+++ b/Engine/Source/Runtime/Renderer/Private/SceneRendering.cpp
@@ -2585,6 +2585,11 @@ FViewFamilyInfo::~FViewFamilyInfo()
 	}
 }
 
+FSceneRenderer::FCustomRenderPassInfo::FCustomRenderPassInfo(const FSceneViewFamily::ConstructionValues& CVS)
+	: ViewFamily(CVS)
+{
+}
+
 FSceneRenderer::FCustomRenderPassInfo::FCustomRenderPassInfo(const FSceneViewFamily::ConstructionValues& CVS, const FViewFamilyInfo& MainViewFamily)
 	: ViewFamily(CVS, MainViewFamily)
 {
@@ -2927,11 +2932,34 @@ FSceneRenderer::FSceneRenderer(const FSceneViewFamily* InViewFamily, FHitProxyCo
 		// color where it isn't used, and translucency should be disabled as an unnecessary perf cost.
 		FamilyCVS.EngineShowFlags.SetTranslucency(CustomRenderPass->IsTranslucentIncluded());
 
-		FCustomRenderPassInfo& CustomRenderPassInfo = CustomRenderPassInfos.Emplace_GetRef(FamilyCVS, ViewFamily);
-		CustomRenderPassInfo.CustomRenderPass = CustomRenderPass;
-		CustomRenderPassInfo.ViewFamily.Time = ViewFamily.Time;
-		CustomRenderPassInfo.ViewFamily.SetSceneRenderer(this);
-		CustomRenderPassInfo.ViewFamily.bIsSceneTextureSizedCapture = SceneCaptureUserData.bMainViewResolution;
+		// The GetDefaultMSAACount call below must match the logic that initializes NumSamples in FSceneTexturesConfig::Init.  Downstream
+		// code uses NumSamples to detect MSAA, after the config has been initialized.
+		FCustomRenderPassInfo* CustomRenderPassInfo;
+		if (GetDefaultMSAACount(Scene->GetFeatureLevel(), GDynamicRHI->RHIGetPlatformTextureMaxSampleCount()) > 1)
+		{
+			// If main pass is MSAA, create a view family with its own FSceneTextures, by calling the constructor that doesn't accept another view family.
+			// MSAA isn't supported for custom render passes, so we'll initialize this one without MSAA (see FSceneRenderer::OnRenderBegin, where we set
+			// NumSamples to 1 for Custom Render Pass scene texture configs).  First custom render pass gets its own, subsequent custom render passes share
+			// with the first custom render pass.
+			if (CustomRenderPassInfos.Num() == 0)
+			{
+				CustomRenderPassInfo = &CustomRenderPassInfos.Emplace_GetRef(FamilyCVS);
+			}
+			else
+			{
+				CustomRenderPassInfo = &CustomRenderPassInfos.Emplace_GetRef(FamilyCVS, CustomRenderPassInfos[0].ViewFamily);
+			}
+		}
+		else
+		{
+			// Passing a View Family to the constructor shares its FSceneTextures struct.
+			CustomRenderPassInfo = &CustomRenderPassInfos.Emplace_GetRef(FamilyCVS, ViewFamily);
+		}
+
+		CustomRenderPassInfo->CustomRenderPass = CustomRenderPass;
+		CustomRenderPassInfo->ViewFamily.Time = ViewFamily.Time;
+		CustomRenderPassInfo->ViewFamily.SetSceneRenderer(this);
+		CustomRenderPassInfo->ViewFamily.bIsSceneTextureSizedCapture = SceneCaptureUserData.bMainViewResolution;
 
 		FSceneViewInitOptions ViewInitOptions;
 		ViewInitOptions.SceneViewStateInterface = PassInput.ViewStateInterface;
@@ -2940,14 +2968,14 @@ FSceneRenderer::FSceneRenderer(const FSceneViewFamily* InViewFamily, FHitProxyCo
 		ViewInitOptions.ViewRotationMatrix = PassInput.ViewRotationMatrix;
 		ViewInitOptions.ProjectionMatrix = PassInput.ProjectionMatrix;
 		ViewInitOptions.bIsSceneCapture = PassInput.bIsSceneCapture;
-		ViewInitOptions.ViewFamily = &CustomRenderPassInfo.ViewFamily;
+		ViewInitOptions.ViewFamily = &CustomRenderPassInfo->ViewFamily;
 		ViewInitOptions.ViewActor = PassInput.ViewActor;
 		ViewInitOptions.ShowOnlyPrimitives = PassInput.ShowOnlyPrimitives;
 		ViewInitOptions.HiddenPrimitives = PassInput.HiddenPrimitives;
 
 		FSceneView NewView(ViewInitOptions);
-		FViewInfo* ViewInfo = &CustomRenderPassInfo.Views.Emplace_GetRef(&NewView);
-		CustomRenderPassInfo.ViewFamily.Views.Add(ViewInfo);
+		FViewInfo* ViewInfo = &CustomRenderPassInfo->Views.Emplace_GetRef(&NewView);
+		CustomRenderPassInfo->ViewFamily.Views.Add(ViewInfo);
 
 		if (PassInput.bOverridesPostVolumeUserFlags)
 		{
@@ -3902,6 +3930,11 @@ IVisibilityTaskData* FSceneRenderer::OnRenderBegin(FRDGBuilder& GraphBuilder, co
 				for (FCustomRenderPassInfo& CustomRenderPass : Renderer->CustomRenderPassInfos)
 				{
 					CustomRenderPass.ViewFamily.SceneTexturesConfig = Renderer->ViewFamily.SceneTexturesConfig;
+
+					// Custom Render Passes don't support MSAA.  If MSAA is enabled, the first Custom Render Pass will allocate a separate non-MSAA
+					// FSceneTextures, initialized using this config (see logic in the FSceneRenderer constructor that fills in CustomRenderPassInfos).
+					CustomRenderPass.ViewFamily.SceneTexturesConfig.NumSamples = 1;
+					CustomRenderPass.ViewFamily.SceneTexturesConfig.EditorPrimitiveNumSamples = 1;
 				}
 	
 				Renderer->PrepareViewStateForVisibility(SceneTexturesConfig);
diff --git a/Engine/Source/Runtime/Renderer/Private/SceneRendering.h b/Engine/Source/Runtime/Renderer/Private/SceneRendering.h
--- a/Engine/Source/Runtime/Renderer/Private/SceneRendering.h
+++ b/Engine/Source/Runtime/Renderer/Private/SceneRendering.h
@@ -2059,6 +2059,7 @@ public:
 	/** Information of a custom render pass that renders as part of the main renderer. */
 	struct FCustomRenderPassInfo
 	{
+		FCustomRenderPassInfo(const FSceneViewFamily::ConstructionValues& CVS);
 		FCustomRenderPassInfo(const FSceneViewFamily::ConstructionValues& CVS, const FViewFamilyInfo& MainViewFamily);
 
 		/** Custom render pass that render as part of the main renderer. */
diff --git a/Engine/Source/Runtime/Renderer/Private/SceneTextures.cpp b/Engine/Source/Runtime/Renderer/Private/SceneTextures.cpp
--- a/Engine/Source/Runtime/Renderer/Private/SceneTextures.cpp
+++ b/Engine/Source/Runtime/Renderer/Private/SceneTextures.cpp
@@ -782,7 +782,9 @@ void FSceneTextures::InitializeViewFamily(FRDGBuilder& GraphBuilder, FViewFamily
 #endif
 
 #if SUPPORTS_VISUALIZE_TEXTURE
-	if (GVisualizeTexture.IsRequestedView())
+	// This function can be called for a Custom Render Pass view family.  Skip those, by checking if this is the
+	// main view family (the first view is the same as the first one in "AllViews").
+	if (GVisualizeTexture.IsRequestedView() && ViewFamily.Views[0] == ViewFamily.AllViews[0])
 	{
 		TArray<FIntRect> FamilyViewRects;
 		FamilyViewRects.SetNumUninitialized(ViewFamily.Views.Num());
