Fix a crash in GPUSkinCache when changing preview platform. This happens when GPUSkinCache is not available on the preview platform
#tests Tested in editor and on console

diff --git a/Engine/Source/Runtime/Renderer/Private/RendererScene.cpp b/Engine/Source/Runtime/Renderer/Private/RendererScene.cpp
--- a/Engine/Source/Runtime/Renderer/Private/RendererScene.cpp
+++ b/Engine/Source/Runtime/Renderer/Private/RendererScene.cpp
@@ -5116,7 +5116,8 @@ void FScene::Update(FRDGBuilder& GraphBuilder, const FUpdateParameters& Paramete
 		GPUSkinUpdateTask = GPUSkinTasks.SkinCache;
 	}
 
-	if (GPUSkinCache)
+	// Check if the skin cache is available here, because we might be midway through a feature level change
+	if (GPUSkinCache && IsGPUSkinCacheAvailable(GetFeatureLevelShaderPlatform(GetFeatureLevel())))
 	{
 		GPUSkinCacheTask = GPUSkinCache->Dispatch(GraphBuilder, GPUSkinTasks.SkinCache, GPUSkinCachePipeline);
 	}
