Cloud investigation

#tests PC editor split screen two player
#rnx

diff --git a/Engine/Source/Runtime/Renderer/Private/VolumetricCloudRendering.cpp b/Engine/Source/Runtime/Renderer/Private/VolumetricCloudRendering.cpp
--- a/Engine/Source/Runtime/Renderer/Private/VolumetricCloudRendering.cpp
+++ b/Engine/Source/Runtime/Renderer/Private/VolumetricCloudRendering.cpp
@@ -2223,7 +2223,7 @@ static TRDGUniformBufferRef<FRenderVolumetricCloudGlobalParameters> CreateCloudP
 	if (CloudRC.bShouldViewRenderVolumetricRenderTarget && MainView.ViewState)
 	{
 		FVolumetricRenderTargetViewStateData& VRT = MainView.ViewState->VolumetricCloudRenderTarget;
-		const uint32 VolumetricReconstructRTDownsampleFactor = VRT.GetVolumetricReconstructRTDownsampleFactor();
+		const uint32 VolumetricReconstructRTDownsampleFactor = VRT.GetVolumetricTracingRTDownsampleFactor();
 
 		// Use the main view to get the target rect. The clouds are reconstructed at the same resolution as the depth buffer read when tracing.
 		VolumetricCloudParams.SceneDepthTextureMinMaxCoord = FUintVector4(
@@ -2789,8 +2789,8 @@ bool FSceneRenderer::RenderVolumetricCloud(
 					CloudRC.RenderTargets[1] = FRenderTargetBinding(DestinationRTDepth, ERenderTargetLoadAction::ENoAction);
 					CloudRC.TracingCoordToZbufferCoordScaleBias = VRT.GetTracingCoordToZbufferCoordScaleBias();
 					// Also take into account the view rect min to be able to read correct depth
-					CloudRC.TracingCoordToZbufferCoordScaleBias.Z += ViewInfo.CachedViewUniformShaderParameters->ViewRectMin.X / ((VRT.GetMode() == 0 || VRT.GetMode() == 3) ? 2 : 1);
-					CloudRC.TracingCoordToZbufferCoordScaleBias.W += ViewInfo.CachedViewUniformShaderParameters->ViewRectMin.Y / ((VRT.GetMode() == 0 || VRT.GetMode() == 3) ? 2 : 1);
+					CloudRC.TracingCoordToZbufferCoordScaleBias.Z += ViewInfo.CachedViewUniformShaderParameters->ViewRectMin.X / ((VRT.GetMode() == 0 || VRT.GetMode() == 1) ? 2 : 1);
+					CloudRC.TracingCoordToZbufferCoordScaleBias.W += ViewInfo.CachedViewUniformShaderParameters->ViewRectMin.Y / ((VRT.GetMode() == 0 || VRT.GetMode() == 1) ? 2 : 1);
 					CloudRC.TracingCoordToFullResPixelCoordScaleBias = VRT.GetTracingCoordToFullResPixelCoordScaleBias();
 
 					if (bAccumulateAlphaHoldOut)
diff --git a/Engine/Source/Runtime/Renderer/Private/VolumetricRenderTarget.cpp b/Engine/Source/Runtime/Renderer/Private/VolumetricRenderTarget.cpp
--- a/Engine/Source/Runtime/Renderer/Private/VolumetricRenderTarget.cpp
+++ b/Engine/Source/Runtime/Renderer/Private/VolumetricRenderTarget.cpp
@@ -563,7 +563,7 @@ FUintVector4 FVolumetricRenderTargetViewStateData::GetTracingCoordToZbufferCoord
 	
 	if (Mode == 2 || Mode == 3)
 	{
-		// In this case, the source depth buffer full resolution depth buffer is the full resolution scene one
+		// In this case, the source depth buffer is the full resolution scene one
 		const uint32 CombinedDownsampleFactor = InvRenderTargetScale * VolumetricReconstructRTDownsampleFactor * VolumetricTracingRTDownsampleFactor;
 		return FUintVector4(CombinedDownsampleFactor, CombinedDownsampleFactor,																// Scale is the combined downsample factor
 			CurrentPixelOffset.X * VolumetricReconstructRTDownsampleFactor, CurrentPixelOffset.Y * VolumetricReconstructRTDownsampleFactor);// Each sample will then sample from full res according to reconstructed RT offset times its downsample factor
