[PathTracing] Output variance on demand instead of always capturing to reduce default cost.

#tests Hillside, ppm with variance, switching denoisers
#rnx

diff --git a/Engine/Plugins/Experimental/NFORDenoise/Source/NFORDenoise/Private/NFORDenoise.cpp b/Engine/Plugins/Experimental/NFORDenoise/Source/NFORDenoise/Private/NFORDenoise.cpp
--- a/Engine/Plugins/Experimental/NFORDenoise/Source/NFORDenoise/Private/NFORDenoise.cpp
+++ b/Engine/Plugins/Experimental/NFORDenoise/Source/NFORDenoise/Private/NFORDenoise.cpp
@@ -186,6 +186,8 @@ public:
 
 	const TCHAR* GetDebugName() const override { return *DebugName; }
 
+	virtual bool NeedVarianceTexture() const override { return true;}
+
 	FOutputs AddPasses(FRDGBuilder& GraphBuilder, const FSceneView& View, const FInputs& Inputs) const override
 	{
 		AddCopyTexturePass(GraphBuilder,Inputs.ColorTex, Inputs.OutputTex);
diff --git a/Engine/Source/Runtime/Renderer/Private/PathTracingSpatialTemporalDenoising.cpp b/Engine/Source/Runtime/Renderer/Private/PathTracingSpatialTemporalDenoising.cpp
--- a/Engine/Source/Runtime/Renderer/Private/PathTracingSpatialTemporalDenoising.cpp
+++ b/Engine/Source/Runtime/Renderer/Private/PathTracingSpatialTemporalDenoising.cpp
@@ -145,17 +145,17 @@ namespace {
 	TAutoConsoleVariable<int32> CVarPathTracingDenoiserPrepassVarianceType(
 		TEXT("r.PathTracing.Denoiser.Prepass.VarianceType"),
 		1,
-		TEXT("Select the per-pixel variance type:")
-		TEXT("0: Multiple channel (RGB) variance for radiance")
-		TEXT("1: Combined single channel variance for radiance, albedo and normal"),
+		TEXT("Select the per-pixel variance type:\n")
+		TEXT("0: Multiple channel (RGB) variance for radiance\n")
+		TEXT("1: Combined single channel variance for radiance, albedo and normal\n"),
 		ECVF_RenderThreadSafe
 	);
 
 	TAutoConsoleVariable<int32> CVarPathTracingDenoiserPrepassOutputVarianceTexture(
 		TEXT("r.PathTracing.Denoiser.Prepass.OutputVarianceTexture"),
-		1,
-		TEXT("0: No variance texture will be generated in Prepass")
-		TEXT("1: Output variance texture to denoisers, or the postprocess material usually used by MRQ")
+		0,
+		TEXT("0: Output variance texture on demand\n")
+		TEXT("1: Always output variance texture to denoisers, or the postprocess material usually used by MRQ\n")
 	);
 
 	TAutoConsoleVariable<int32> CVarPathTracingDenoiserPrepassRankedLuminanceVariance(
@@ -643,18 +643,37 @@ static bool ShouldEnableSubpixelOffset(uint32 MipLevel)
 	return bShouldEnableSubpixelOffset;
 }
 
+static bool IsVarianceTextureRequiredForCurrentDenoiser()
+{
+	bool NeedVarianceTexture = false;
+
+	bool bMatch = true;
+	using UE::Renderer::Private::IPathTracingSpatialTemporalDenoiser;
+	FString DenoiserName = CVarPathTracingDenoiserName.GetValueOnRenderThread();
+	IPathTracingSpatialTemporalDenoiser* ActiveSpatialTemporalDenoiser = FDenoiserManager::Get().GetSpatialTemporalDenoiser(DenoiserName, bMatch);
+
+	if (ActiveSpatialTemporalDenoiser)
+	{
+		NeedVarianceTexture = ActiveSpatialTemporalDenoiser->NeedVarianceTexture();
+	}
+	
+	return NeedVarianceTexture;
+}
+
 static bool ShouldPrepassOutputVarianceTexture(const FViewInfo& View)
 {
 	static const auto CVarOutputPostProcessResources = 
-		IConsoleManager::Get().FindTConsoleVariableDataInt(TEXT("r.PathTracing.OutputPostProcessResources"));
+		IConsoleManager::Get().FindTConsoleVariableDataBool(TEXT("r.PathTracing.OutputPostProcessResources"));
 	const bool bOutputPostProcessResources = CVarOutputPostProcessResources ?
-		(CVarOutputPostProcessResources->GetValueOnRenderThread() != 0) : false;
+		CVarOutputPostProcessResources->GetValueOnRenderThread() : false;
 	
-	// Variance texture will be available if post process requires and we allow output post process resource in path tracing
-	// or when we allow prepass to output so it can be accessed by denoisers based on variance.
+	// Variance texture will be available if 
+	// 1. when we always output so it can be accessed for debug
+	// 2. post process requires and we allow output post process resource in path tracing
+	// 3. the denoiser requires a variance texture.
 	return CVarPathTracingDenoiserPrepassOutputVarianceTexture.GetValueOnRenderThread() != 0 ||
-		(bOutputPostProcessResources && 
-		IsPathTracingVarianceTextureRequiredInPostProcessMaterial(View));
+		(bOutputPostProcessResources && IsPathTracingVarianceTextureRequiredInPostProcessMaterial(View)) ||
+		IsVarianceTextureRequiredForCurrentDenoiser();
 }
 
 static constexpr uint32 kMipDiffDelta = 2;
@@ -1083,7 +1102,7 @@ static void PathTracingDenoiserPlugin(FRDGBuilder& GraphBuilder,
 			Inputs.VarianceTex = VarianceTexture;
 			
 			// Set flow texture to black
-			Inputs.FlowTex = GraphBuilder.CreateTexture(VarianceTexture->Desc, TEXT("PathTracing.OpticalFlow"));
+			Inputs.FlowTex = GraphBuilder.CreateTexture(InputTexture->Desc, TEXT("PathTracing.OpticalFlow"));
 			AddClearRenderTargetPass(GraphBuilder, Inputs.FlowTex, FLinearColor::Black);
 			
 			Inputs.PreviousOutputTex = InputTexture;
diff --git a/Engine/Source/Runtime/Renderer/Public/PathTracingDenoiser.h b/Engine/Source/Runtime/Renderer/Public/PathTracingDenoiser.h
--- a/Engine/Source/Runtime/Renderer/Public/PathTracingDenoiser.h
+++ b/Engine/Source/Runtime/Renderer/Public/PathTracingDenoiser.h
@@ -97,6 +97,8 @@ public:
 
 	virtual bool NeedTextureCreateExtraFlags() const { return false; }
 
+	virtual bool NeedVarianceTexture() const { return false;}
+
 	/** Adds the necessary passes into RDG for denoising. */
 	virtual FOutputs AddPasses(FRDGBuilder& GraphBuilder, const FSceneView& View, const FInputs& Inputs) const = 0;
 
