Path Tracer: Allow post-process materials to be specialized for the path tracer

This adds a permutation of the PostProcess Material that is used when the path tracer is enabled on the view. This can be used to account for differences in how the path tracer creates the image.

Guard reads to the path tracing specific buffers to this new permutation.

diff --git a/Engine/Shaders/Private/Common.ush b/Engine/Shaders/Private/Common.ush
--- a/Engine/Shaders/Private/Common.ush
+++ b/Engine/Shaders/Private/Common.ush
@@ -2135,7 +2135,7 @@ bool GetRayTracingQualitySwitch()
 #endif
 }
 
-#if !PATH_TRACING
+#if !PATH_TRACING && !PATH_TRACING_POST_PROCESS_MATERIAL
 // Default implementations for non-path traced shaders
 // The actual implementations used for the path tracer exist in PathTracingShaderUtils.h to minimize cruft in Common.ush
 bool GetPathTracingQualitySwitch() { return false;  }
diff --git a/Engine/Shaders/Private/MaterialTemplate.ush b/Engine/Shaders/Private/MaterialTemplate.ush
--- a/Engine/Shaders/Private/MaterialTemplate.ush
+++ b/Engine/Shaders/Private/MaterialTemplate.ush
@@ -3421,8 +3421,7 @@ float4 MaterialExpressionDBufferTextureLookup(FMaterialPixelParameters Parameter
 MaterialFloat2 GetDefaultPathTracingBufferTextureUV(FMaterialPixelParameters Parameters, const uint PathTracingBufferTextureIndex)
 {
 	float2 ViewportUV = GetViewportUV(Parameters);
-#if POST_PROCESS_MATERIAL && SHADING_PATH_DEFERRED && MATERIAL_PATH_TRACING_BUFFER_READ
-	
+#if POST_PROCESS_MATERIAL && SHADING_PATH_DEFERRED && MATERIAL_PATH_TRACING_BUFFER_READ && PATH_TRACING_POST_PROCESS_MATERIAL
 	switch (PathTracingBufferTextureIndex)
 	{
 	case 0: return ViewportUV * PathTracingPostProcessInput_0_UVViewportSize + PathTracingPostProcessInput_0_UVViewportMin;
@@ -3437,7 +3436,7 @@ MaterialFloat2 GetDefaultPathTracingBufferTextureUV(FMaterialPixelParameters Par
 
 float4 MaterialExpressionPathTracingBufferTextureLookup(FMaterialPixelParameters Parameters, float2 BufferUV, int PathTracingBufferTextureIndex)
 {
-#if POST_PROCESS_MATERIAL && SHADING_PATH_DEFERRED && MATERIAL_PATH_TRACING_BUFFER_READ
+#if POST_PROCESS_MATERIAL && SHADING_PATH_DEFERRED && MATERIAL_PATH_TRACING_BUFFER_READ && PATH_TRACING_POST_PROCESS_MATERIAL
 	switch (PathTracingBufferTextureIndex)
 	{
 	case 0:	return Texture2DSample(PathTracingPostProcessInput_0_Texture, PathTracingPostProcessInput_0_Sampler, BufferUV);
diff --git a/Engine/Shaders/Private/PostProcessMaterialShaders.usf b/Engine/Shaders/Private/PostProcessMaterialShaders.usf
--- a/Engine/Shaders/Private/PostProcessMaterialShaders.usf
+++ b/Engine/Shaders/Private/PostProcessMaterialShaders.usf
@@ -90,6 +90,16 @@ SamplerState PathTracingPostProcessInput_3_Sampler;
 SamplerState PathTracingPostProcessInput_4_Sampler;
 #endif
 
+#if PATH_TRACING_POST_PROCESS_MATERIAL
+// Allow users to customize the behavior of the material for the path tracing with PathTraicngQualitySwitch node
+bool GetPathTracingQualitySwitch() { return true;  }
+bool GetPathTracingIsShadow() { return false; }
+bool GetPathTracingIsIndirectDiffuse() { return false; }
+bool GetPathTracingIsIndirectSpecular() { return false; }
+bool GetPathTracingIsIndirectVolume() { return false; }
+#endif
+
+
 #include "NeuralPostProcessCommon.usf"
 
 #define EYE_ADAPTATION_LOOSE_PARAMETERS 1
diff --git a/Engine/Source/Runtime/Renderer/Private/PathTracing.cpp b/Engine/Source/Runtime/Renderer/Private/PathTracing.cpp
--- a/Engine/Source/Runtime/Renderer/Private/PathTracing.cpp
+++ b/Engine/Source/Runtime/Renderer/Private/PathTracing.cpp
@@ -810,6 +810,13 @@ int GetPathTracingVisualizationMode()
 
 namespace PathTracing
 {
+	bool ShouldCompilePathTracingShadersForProject(EShaderPlatform ShaderPlatform)
+	{
+		return ShouldCompileRayTracingShadersForProject(ShaderPlatform) &&
+			FDataDrivenShaderPlatformInfo::GetSupportsPathTracing(ShaderPlatform) &&
+			CVarPathTracing.GetValueOnAnyThread() != 0;
+	}
+
 	bool UsesDecals(const FSceneViewFamily& ViewFamily)
 	{
 		return ViewFamily.EngineShowFlags.Decals;
@@ -967,13 +974,6 @@ static void PreparePathTracingData(const FScene* Scene, const FViewInfo& View, F
 	}
 }
 
-static bool ShouldCompilePathTracingShadersForProject(EShaderPlatform ShaderPlatform)
-{
-	return ShouldCompileRayTracingShadersForProject(ShaderPlatform) &&
-			FDataDrivenShaderPlatformInfo::GetSupportsPathTracing(ShaderPlatform) &&
-			CVarPathTracing.GetValueOnAnyThread() != 0;
-}
-
 static bool ShouldCompileGPULightmassShadersForProject(EShaderPlatform ShaderPlatform)
 {
 #if WITH_EDITOR
@@ -1076,7 +1076,7 @@ class FPathTracingBuildLightGridCS : public FGlobalShader
 
 	static bool ShouldCompilePermutation(const FGlobalShaderPermutationParameters& Parameters)
 	{
-		return ShouldCompilePathTracingShadersForProject(Parameters.Platform) ||
+		return PathTracing::ShouldCompilePathTracingShadersForProject(Parameters.Platform) ||
 			   ShouldCompileGPULightmassShadersForProject(Parameters.Platform);
 	}
 
@@ -1276,7 +1276,7 @@ class FPathTracingRG : public FGlobalShader
 				return false;
 			}
 		}
-		return ShouldCompilePathTracingShadersForProject(Parameters.Platform);
+		return PathTracing::ShouldCompilePathTracingShadersForProject(Parameters.Platform);
 	}
 
 	static void ModifyCompilationEnvironment(const FGlobalShaderPermutationParameters& Parameters, FShaderCompilerEnvironment& OutEnvironment)
@@ -1379,7 +1379,7 @@ class FPathTracingDebugRG : public FGlobalShader
 
 	static bool ShouldCompilePermutation(const FGlobalShaderPermutationParameters& Parameters)
 	{
-		return ShouldCompilePathTracingShadersForProject(Parameters.Platform);
+		return PathTracing::ShouldCompilePathTracingShadersForProject(Parameters.Platform);
 	}
 
 	static void ModifyCompilationEnvironment(const FGlobalShaderPermutationParameters& Parameters, FShaderCompilerEnvironment& OutEnvironment)
@@ -1462,7 +1462,7 @@ class FPathTracingInitExtinctionCoefficientRG : public FGlobalShader
 
 	static bool ShouldCompilePermutation(const FGlobalShaderPermutationParameters& Parameters)
 	{
-		return ShouldCompilePathTracingShadersForProject(Parameters.Platform);
+		return PathTracing::ShouldCompilePathTracingShadersForProject(Parameters.Platform);
 	}
 
 	static ERayTracingPayloadType GetRayTracingPayloadType(const int32 PermutationId)
@@ -1493,7 +1493,7 @@ class FPathTracingSwizzleScanlinesCS : public FGlobalShader
 
 	static bool ShouldCompilePermutation(const FGlobalShaderPermutationParameters& Parameters)
 	{
-		return ShouldCompilePathTracingShadersForProject(Parameters.Platform);
+		return PathTracing::ShouldCompilePathTracingShadersForProject(Parameters.Platform);
 	}
 
 	static void ModifyCompilationEnvironment(const FGlobalShaderPermutationParameters& Parameters, FShaderCompilerEnvironment& OutEnvironment)
@@ -1521,7 +1521,7 @@ class FPathTracingBuildAtmosphereOpticalDepthLUTCS : public FGlobalShader
 
 	static bool ShouldCompilePermutation(const FGlobalShaderPermutationParameters& Parameters)
 	{
-		return ShouldCompilePathTracingShadersForProject(Parameters.Platform);
+		return PathTracing::ShouldCompilePathTracingShadersForProject(Parameters.Platform);
 	}
 
 	static void ModifyCompilationEnvironment(const FGlobalShaderPermutationParameters& Parameters, FShaderCompilerEnvironment& OutEnvironment)
@@ -1614,7 +1614,7 @@ class FPathTracingBuildCloudAccelerationMapCS : public FMeshMaterialShader
 
 	static bool ShouldCompilePermutation(const FMeshMaterialShaderPermutationParameters& Parameters)
 	{
-		return ShouldCompilePathTracingShadersForProject(Parameters.Platform) &&
+		return PathTracing::ShouldCompilePathTracingShadersForProject(Parameters.Platform) &&
 			Parameters.MaterialParameters.bIsUsedWithVolumetricCloud &&
 			Parameters.MaterialParameters.MaterialDomain == MD_Volume;
 	}
@@ -1653,7 +1653,7 @@ public:
 
 	static bool ShouldCompilePermutation(const FMaterialShaderPermutationParameters& Parameters)
 	{
-		return ShouldCompilePathTracingShadersForProject(Parameters.Platform) &&
+		return PathTracing::ShouldCompilePathTracingShadersForProject(Parameters.Platform) &&
 			ShouldCompileRayTracingCallableShadersForProject(Parameters.Platform) &&
 			Parameters.MaterialParameters.bIsUsedWithVolumetricCloud &&
 			Parameters.MaterialParameters.MaterialDomain == MD_Volume;
@@ -1737,7 +1737,7 @@ public:
 
 	static bool ShouldCompilePermutation(const FGlobalShaderPermutationParameters& Parameters)
 	{
-		return ShouldCompilePathTracingShadersForProject(Parameters.Platform) &&
+		return PathTracing::ShouldCompilePathTracingShadersForProject(Parameters.Platform) &&
 			   ShouldCompileRayTracingCallableShadersForProject(Parameters.Platform);
 	}
 
@@ -1946,7 +1946,7 @@ class FPathTracingBuildAdaptiveErrorTextureCS : public FGlobalShader
 
 	static bool ShouldCompilePermutation(const FGlobalShaderPermutationParameters& Parameters)
 	{
-		return ShouldCompilePathTracingShadersForProject(Parameters.Platform);
+		return PathTracing::ShouldCompilePathTracingShadersForProject(Parameters.Platform);
 	}
 
 	static void ModifyCompilationEnvironment(const FGlobalShaderPermutationParameters& Parameters, FShaderCompilerEnvironment& OutEnvironment)
@@ -1974,7 +1974,7 @@ class FPathTracingAdaptiveStartCS : public FGlobalShader
 
 	static bool ShouldCompilePermutation(const FGlobalShaderPermutationParameters& Parameters)
 	{
-		return ShouldCompilePathTracingShadersForProject(Parameters.Platform);
+		return PathTracing::ShouldCompilePathTracingShadersForProject(Parameters.Platform);
 	}
 
 	static void ModifyCompilationEnvironment(const FGlobalShaderPermutationParameters& Parameters, FShaderCompilerEnvironment& OutEnvironment)
@@ -2020,7 +2020,7 @@ public:
 		}
 		else
 		{
-			return ShouldCompilePathTracingShadersForProject(Parameters.Platform);
+			return PathTracing::ShouldCompilePathTracingShadersForProject(Parameters.Platform);
 		}
 	}
 
@@ -2129,7 +2129,7 @@ class FPathTracingLightingMS : public FMaterialShader
 public:
 	static bool ShouldCompilePermutation(const FMaterialShaderPermutationParameters& Parameters)
 	{
-		return Parameters.MaterialParameters.MaterialDomain == MD_LightFunction && ShouldCompilePathTracingShadersForProject(Parameters.Platform);
+		return Parameters.MaterialParameters.MaterialDomain == MD_LightFunction && PathTracing::ShouldCompilePathTracingShadersForProject(Parameters.Platform);
 	}
 
 	FPathTracingLightingMS() {}
@@ -2335,7 +2335,7 @@ public:
 				// or the user did not request the extra permutations to be compiled (default)
 				return false;
 			}
-			return ShouldCompilePathTracingShadersForProject(Parameters.Platform);
+			return PathTracing::ShouldCompilePathTracingShadersForProject(Parameters.Platform);
 		}
 	}
 
@@ -2437,7 +2437,7 @@ class TPathTracingDefaultHitGroup : public FGlobalShader
 		}
 		else
 		{
-			return ShouldCompilePathTracingShadersForProject(Parameters.Platform);
+			return PathTracing::ShouldCompilePathTracingShadersForProject(Parameters.Platform);
 		}
 	}
 
@@ -3211,7 +3211,7 @@ void FDeferredShadingSceneRenderer::PreparePathTracing(const FViewInfo& View, co
 
 	const FSceneViewFamily& ViewFamily = *View.Family;
 	if (ViewFamily.EngineShowFlags.PathTracing
-		&& ShouldCompilePathTracingShadersForProject(ViewFamily.GetShaderPlatform()))
+		&& PathTracing::ShouldCompilePathTracingShadersForProject(ViewFamily.GetShaderPlatform()))
 	{
 		if (GetPathTracingVisualizationMode() >= 0)
 		{
diff --git a/Engine/Source/Runtime/Renderer/Private/PathTracing.h b/Engine/Source/Runtime/Renderer/Private/PathTracing.h
--- a/Engine/Source/Runtime/Renderer/Private/PathTracing.h
+++ b/Engine/Source/Runtime/Renderer/Private/PathTracing.h
@@ -36,6 +36,7 @@ void PreparePathTracingCloudMaterial(FRDGBuilder& GraphBuilder, FScene* Scene, T
 
 namespace PathTracing
 {
+	bool ShouldCompilePathTracingShadersForProject(EShaderPlatform ShaderPlatform);
 	bool UsesDecals(const FSceneViewFamily& ViewFamily);
 	bool UsesReferenceAtmosphere(const FViewInfo& View);
 	bool UsesReferenceDOF(const FViewInfo& View);
diff --git a/Engine/Source/Runtime/Renderer/Private/PostProcess/PostProcessMaterial.cpp b/Engine/Source/Runtime/Renderer/Private/PostProcess/PostProcessMaterial.cpp
--- a/Engine/Source/Runtime/Renderer/Private/PostProcess/PostProcessMaterial.cpp
+++ b/Engine/Source/Runtime/Renderer/Private/PostProcess/PostProcessMaterial.cpp
@@ -30,6 +30,7 @@
 #include "Substrate/Substrate.h"
 #include "SingleLayerWaterRendering.h"
 #include "Engine/NeuralProfile.h"
+#include "PathTracing.h"
 
 namespace
 {
@@ -304,7 +305,8 @@ public:
 
 	class FManualStencilTestDim : SHADER_PERMUTATION_BOOL("MANUAL_STENCIL_TEST");
 	class FNeuralPostProcessPrePass : SHADER_PERMUTATION_BOOL("NEURAL_POSTPROCESS_PREPASS");
-	using FPermutationDomain = TShaderPermutationDomain<FManualStencilTestDim,FNeuralPostProcessPrePass>;
+	class FPathTracingDim : SHADER_PERMUTATION_BOOL("PATH_TRACING_POST_PROCESS_MATERIAL");
+	using FPermutationDomain = TShaderPermutationDomain<FManualStencilTestDim,FNeuralPostProcessPrePass,FPathTracingDim>;
 
 	FPostProcessMaterialPS() = default;
 	FPostProcessMaterialPS(const ShaderMetaType::CompiledShaderInitializerType& Initializer)
@@ -324,7 +326,19 @@ public:
 		// See comments in CustomDepthRendering.h for more details.
 		if (PermutationVector.Get<FManualStencilTestDim>())
 		{
-			return Parameters.MaterialParameters.bIsStencilTestEnabled && DoesPlatformSupportNanite(Parameters.Platform);
+			if (!(Parameters.MaterialParameters.bIsStencilTestEnabled && DoesPlatformSupportNanite(Parameters.Platform)))
+			{
+				return false;
+			}
+		}
+
+		// Only enable the path tracing specialization when path tracing is enabled on the project
+		if (PermutationVector.Get<FPathTracingDim>())
+		{
+			if (!PathTracing::ShouldCompilePathTracingShadersForProject(Parameters.Platform))
+			{
+				return false;
+			}
 		}
 
 		return true;
@@ -357,6 +371,7 @@ static bool GetMaterialShaders(
 	const FMaterial& Material, 
 	bool bManualStencilTest, 
 	bool bNeuralPostProcessPrepass,
+	bool bPathTracingEnabled,
 	TShaderRef<FPostProcessMaterialVS>& OutVertexShader,
 	TShaderRef<FPostProcessMaterialPS>& OutPixelShader)
 {
@@ -365,6 +380,7 @@ static bool GetMaterialShaders(
 	FPostProcessMaterialPS::FPermutationDomain PermutationVectorPS;
 	PermutationVectorPS.Set<FPostProcessMaterialPS::FManualStencilTestDim>(bManualStencilTest);
 	PermutationVectorPS.Set<FPostProcessMaterialPS::FNeuralPostProcessPrePass>(bNeuralPostProcessPrepass);
+	PermutationVectorPS.Set<FPostProcessMaterialPS::FPathTracingDim>(bPathTracingEnabled);
 
 	ShaderTypes.AddShaderType<FPostProcessMaterialVS>();
 	ShaderTypes.AddShaderType<FPostProcessMaterialPS>(PermutationVectorPS.ToDimensionValueId());
@@ -384,6 +400,7 @@ static bool GetMaterialShaders(
 static void GetMaterialInfo(
 	const UMaterialInterface* InMaterialInterface,
 	ERHIFeatureLevel::Type InFeatureLevel,
+	const bool bPathTracingEnabled,
 	const FPostProcessMaterialInputs& Inputs,
 	const FMaterial*& OutMaterial,
 	const FMaterialRenderProxy*& OutMaterialProxy,
@@ -403,7 +420,7 @@ static void GetMaterialInfo(
 		if (Material && Material->GetMaterialDomain() == MD_PostProcess)
 		{
 			const bool bManualStencilTest = Inputs.bManualStencilTest && Material->IsStencilTestEnabled();
-			if (GetMaterialShaders(*Material, bManualStencilTest, bNeuralPostProcessPrepass, OutVertexShader, OutPixelShader))
+			if (GetMaterialShaders(*Material, bManualStencilTest, bNeuralPostProcessPrepass, bPathTracingEnabled, OutVertexShader, OutPixelShader))
 			{
 				break;
 			}
@@ -574,13 +591,14 @@ void AddNeuralPostProcessPass(
 	const FScreenPassTexture SceneColor = FScreenPassTexture::CopyFromSlice(GraphBuilder, Inputs.GetInput(EPostProcessMaterialInput::SceneColor));
 
 	const ERHIFeatureLevel::Type FeatureLevel = View.GetFeatureLevel();
+	const bool bPathTracingEnabled = View.Family->EngineShowFlags.PathTracing;
 
 	const FMaterial* Material = nullptr;
 	const FMaterialRenderProxy* MaterialRenderProxy = nullptr;
 	const FMaterialShaderMap* MaterialShaderMap = nullptr;
 	TShaderRef<FPostProcessMaterialVS> NeuralPostProcessPassVertexShader;
 	TShaderRef<FPostProcessMaterialPS> NeuralPostProcessPassPixelShader;
-	GetMaterialInfo(MaterialInterface, FeatureLevel, Inputs, Material, MaterialRenderProxy, MaterialShaderMap, NeuralPostProcessPassVertexShader, NeuralPostProcessPassPixelShader, true);
+	GetMaterialInfo(MaterialInterface, FeatureLevel, bPathTracingEnabled, Inputs, Material, MaterialRenderProxy, MaterialShaderMap, NeuralPostProcessPassVertexShader, NeuralPostProcessPassPixelShader, true);
 	
 	check(NeuralPostProcessPassVertexShader.IsValid());
 	check(NeuralPostProcessPassPixelShader.IsValid());
@@ -723,13 +741,14 @@ FScreenPassTexture AddPostProcessMaterialPass(
 	Inputs.Validate();
 
 	const ERHIFeatureLevel::Type FeatureLevel = View.GetFeatureLevel();
+	const bool bPathTracingEnabled = View.Family->EngineShowFlags.PathTracing;
 
 	const FMaterial* Material = nullptr;
 	const FMaterialRenderProxy* MaterialRenderProxy = nullptr;
 	const FMaterialShaderMap* MaterialShaderMap = nullptr;
 	TShaderRef<FPostProcessMaterialVS> VertexShader;
 	TShaderRef<FPostProcessMaterialPS> PixelShader;
-	GetMaterialInfo(MaterialInterface, FeatureLevel, Inputs, Material, MaterialRenderProxy, MaterialShaderMap, VertexShader, PixelShader);
+	GetMaterialInfo(MaterialInterface, FeatureLevel, bPathTracingEnabled, Inputs, Material, MaterialRenderProxy, MaterialShaderMap, VertexShader, PixelShader);
 
 	EBlendableLocation BlendableLocation = MaterialRenderProxy->GetBlendableLocation(Material);
 	const FScreenPassTextureSlice SceneColorOutput = Inputs.GetSceneColorOutput(BlendableLocation);
@@ -1558,8 +1577,11 @@ void FPostProcessMaterialPSOCollector::CollectPSOInitializers(
 		TShaderRef<FPostProcessMaterialVS> VertexShader;
 		TShaderRef<FPostProcessMaterialPS> PixelShader;
 
+		// PSO Precaching will not take into account the following features for now as they are less frequently used
 		bool bNeuralPostProcessPrepass = false;
-		if (!GetMaterialShaders(Material, bManualStencilTest, bNeuralPostProcessPrepass, VertexShader, PixelShader))
+		bool bPathTracingEnabled = false;
+
+		if (!GetMaterialShaders(Material, bManualStencilTest, bNeuralPostProcessPrepass, bPathTracingEnabled, VertexShader, PixelShader))
 		{
 			return;
 		}
