Path Tracer: Fix compile failure on platforms without raytracing support

#rnx

diff --git a/Engine/Source/Runtime/Renderer/Private/PathTracing.cpp b/Engine/Source/Runtime/Renderer/Private/PathTracing.cpp
--- a/Engine/Source/Runtime/Renderer/Private/PathTracing.cpp
+++ b/Engine/Source/Runtime/Renderer/Private/PathTracing.cpp
@@ -4559,6 +4559,11 @@ void FDeferredShadingSceneRenderer::RenderPathTracing(
 
 namespace PathTracing
 {
+	bool ShouldCompilePathTracingShadersForProject(EShaderPlatform ShaderPlatform)
+	{
+		return false;
+	}
+
 	bool UsesDecals(const FSceneViewFamily& ViewFamily)
 	{
 		return false;
