Fix incorrect transition for virtual texture feedback debug.
#rnx

diff --git a/Engine/Source/Runtime/Renderer/Private/VT/VirtualTextureFeedbackResource.cpp b/Engine/Source/Runtime/Renderer/Private/VT/VirtualTextureFeedbackResource.cpp
--- a/Engine/Source/Runtime/Renderer/Private/VT/VirtualTextureFeedbackResource.cpp
+++ b/Engine/Source/Runtime/Renderer/Private/VT/VirtualTextureFeedbackResource.cpp
@@ -256,7 +256,7 @@ FRDGBuffer* FVirtualTextureFeedbackBufferResource::ResolveExtendedDebugBuffer(FR
 	AddPass(GraphBuilder, RDG_EVENT_NAME("VirtualTextureFeedbackTransitionBeforeExtract"), [UAV = UAV](FRHICommandList& RHICmdList)
 	{
 		RHICmdList.EndUAVOverlap(UAV);
-		RHICmdList.Transition(FRHITransitionInfo(UAV, ERHIAccess::UAVMask, ERHIAccess::SRVCompute));
+		RHICmdList.Transition(FRHITransitionInfo(UAV, ERHIAccess::UAVMask, ERHIAccess::CopySrc));
 	});
 
 	// Copy the extended debug payload.
@@ -270,7 +270,7 @@ FRDGBuffer* FVirtualTextureFeedbackBufferResource::ResolveExtendedDebugBuffer(FR
 	// Transition feedback to writing in case we have any subsequent feedback passes (that we won't have captured debug info for).
 	AddPass(GraphBuilder, RDG_EVENT_NAME("VirtualTextureFeedbackTransitionAfterExtract"), [UAV = UAV](FRHICommandList& RHICmdList)
 	{
-		RHICmdList.Transition(FRHITransitionInfo(UAV, ERHIAccess::SRVCompute, ERHIAccess::UAVMask));
+		RHICmdList.Transition(FRHITransitionInfo(UAV, ERHIAccess::CopySrc, ERHIAccess::UAVMask));
 		RHICmdList.BeginUAVOverlap(UAV);
 	});
 
