Remove disabled/deprecated virtual texture test.
#rnx

diff --git a/Engine/Source/Runtime/Renderer/Private/ShaderBaseClasses.cpp b/Engine/Source/Runtime/Renderer/Private/ShaderBaseClasses.cpp
--- a/Engine/Source/Runtime/Renderer/Private/ShaderBaseClasses.cpp
+++ b/Engine/Source/Runtime/Renderer/Private/ShaderBaseClasses.cpp
@@ -10,9 +10,6 @@
 #include "RendererModule.h"
 #include "ScenePrivate.h"
 #include "ParameterCollection.h"
-#include "VT/VirtualTextureTest.h"
-#include "VT/VirtualTextureSpace.h"
-#include "VT/VirtualTextureSystem.h"
 
 IMPLEMENT_TYPE_LAYOUT(FMaterialShader);
 IMPLEMENT_TYPE_LAYOUT(FMeshMaterialShader);
diff --git a/Engine/Source/Runtime/Renderer/Private/VT/VirtualTextureTest.cpp b/Engine/Source/Runtime/Renderer/Private/VT/VirtualTextureTest.cpp
deleted file mode 100644
diff --git a/Engine/Source/Runtime/Renderer/Private/VT/VirtualTextureTest.h b/Engine/Source/Runtime/Renderer/Private/VT/VirtualTextureTest.h
deleted file mode 100644
