Runtime virtual texture always uses PooledRenderTarget and no longer uses the low level RHI Texture.
This simplifies transition handling.
#rnx

diff --git a/Engine/Plugins/Experimental/VirtualHeightfieldMesh/Source/VirtualHeightfieldMeshEditor/Private/HeightfieldMinMaxTextureBuild.cpp b/Engine/Plugins/Experimental/VirtualHeightfieldMesh/Source/VirtualHeightfieldMeshEditor/Private/HeightfieldMinMaxTextureBuild.cpp
--- a/Engine/Plugins/Experimental/VirtualHeightfieldMesh/Source/VirtualHeightfieldMeshEditor/Private/HeightfieldMinMaxTextureBuild.cpp
+++ b/Engine/Plugins/Experimental/VirtualHeightfieldMesh/Source/VirtualHeightfieldMeshEditor/Private/HeightfieldMinMaxTextureBuild.cpp
@@ -207,9 +207,9 @@ namespace VirtualHeightfieldMesh
 					Desc.bIsThumbnails = false;
 					Desc.FixedColor = FLinearColor::Transparent;
 					Desc.NumPageDescs = 1;
-					Desc.Targets[0].Texture = RenderTileResources.GetTileRenderTarget()->GetRHI();
-					Desc.Targets[1].Texture = nullptr;
-					Desc.Targets[2].Texture = nullptr;
+					Desc.Targets[0].PooledRenderTarget = RenderTileResources.GetTileRenderTarget();
+					Desc.Targets[1].PooledRenderTarget = nullptr;
+					Desc.Targets[2].PooledRenderTarget = nullptr;
 					Desc.PageDescs[0].DestRect[0] = FIntRect(0, 0, TileSize, TileSize);
 					Desc.PageDescs[0].DestRect[1] = FIntRect();
 					Desc.PageDescs[0].DestRect[2] = FIntRect();
diff --git a/Engine/Source/Editor/VirtualTexturingEditor/Private/RuntimeVirtualTextureBuildStreamingMips.cpp b/Engine/Source/Editor/VirtualTexturingEditor/Private/RuntimeVirtualTextureBuildStreamingMips.cpp
--- a/Engine/Source/Editor/VirtualTexturingEditor/Private/RuntimeVirtualTextureBuildStreamingMips.cpp
+++ b/Engine/Source/Editor/VirtualTexturingEditor/Private/RuntimeVirtualTextureBuildStreamingMips.cpp
@@ -14,6 +14,7 @@
 #include "Materials/MaterialInterface.h"
 #include "Misc/ScopedSlowTask.h"
 #include "RenderGraphBuilder.h"
+#include "RenderGraphUtils.h"
 #include "RendererInterface.h"
 #include "RenderTargetPool.h"
 #include "SceneInterface.h"
@@ -58,16 +59,13 @@ namespace
 
 			for (int32 Layer = 0; Layer < NumLayers; ++Layer)
 			{
-				FRHITextureCreateDesc Desc =
-					FRHITextureCreateDesc::Create2D(TEXT("FTileRenderResources"), TileSize, TileSize, LayerFormats[Layer]);
+				FPooledRenderTargetDesc RenderTargetDesc = FPooledRenderTargetDesc::Create2DDesc(
+					FIntPoint(TileSize, TileSize), LayerFormats[Layer], FClearValueBinding::None, TexCreate_None, TexCreate_ShaderResource, false);
+				GRenderTargetPool.FindFreeElement(RHICmdList, RenderTargetDesc, RenderTargets[Layer], TEXT("FTileRenderResources"));
 
-				Desc.SetFlags(ETextureCreateFlags::RenderTargetable);
-				RenderTargets[Layer] = RHICreateTexture(Desc);
-
-				Desc.SetFlags(ETextureCreateFlags::CPUReadback);
-				// The staging texture will only ever be in CopyDest state so set it right away : 
-				Desc.SetInitialState(ERHIAccess::CopyDest); 
-				StagingTextures[Layer] = RHICreateTexture(Desc);
+				FRHITextureCreateDesc StagingTextureDesc = FRHITextureCreateDesc::Create2D(TEXT("FTileRenderResources"), TileSize, TileSize, LayerFormats[Layer]);
+				StagingTextureDesc.SetFlags(ETextureCreateFlags::CPUReadback);
+				StagingTextures[Layer] = RHICreateTexture(StagingTextureDesc);
 			}
 
 			Fence = RHICreateGPUFence(TEXT("Runtime Virtual Texture Build"));
@@ -87,8 +85,8 @@ namespace
 		EPixelFormat GetLayerFormat(int32 Index) const { return LayerFormats[Index]; }
 		int64 GetLayerOffset(int32 Index) const { return LayerOffsets[Index]; }
 
-		FRHITexture* GetRenderTarget(int32 Index) const { return Index < NumLayers ? RenderTargets[Index] : nullptr; }
-		FRHITexture* GetStagingTexture(int32 Index) const { return Index < NumLayers ? StagingTextures[Index] : nullptr; }
+		TRefCountPtr<IPooledRenderTarget> GetRenderTarget(int32 Index) const { return RenderTargets[Index]; }
+		FRHITexture* GetStagingTexture(int32 Index) const { return StagingTextures[Index]; }
 		FRHIGPUFence* GetFence() const { return Fence; }
 
 	private:
@@ -99,7 +97,7 @@ namespace
 		TArray<EPixelFormat> LayerFormats;
 		TArray<int64> LayerOffsets;
 
-		TArray<FTextureRHIRef> RenderTargets;
+		TArray<TRefCountPtr<IPooledRenderTarget>> RenderTargets;
 		TArray<FTextureRHIRef> StagingTextures;
 		FGPUFenceRHIRef Fence;
 	};
@@ -227,7 +225,6 @@ namespace RuntimeVirtualTexture
 		FinalPixels.SetNumUninitialized(RenderTileResourcesBytes);
 
 		UE::RenderCommandPipe::FSyncScope SyncScope;
-		ERHIAccess RenderTargetsCurrentState = ERHIAccess::RTV;
 
 		// Iterate over all tiles and render/store each one to the final image
 		for (int32 TileY = 0; TileY < NumTilesY && !Task.ShouldCancel(); TileY++)
@@ -257,8 +254,7 @@ namespace RuntimeVirtualTexture
 					TileX, TileY,
 					TileSize, ImageSizeX, ImageSizeY, 
 					&FinalPixels,
-					CustomMaterialData, FixedColor,
-					&RenderTargetsCurrentState](FRHICommandListImmediate& RHICmdList)
+					CustomMaterialData, FixedColor] (FRHICommandListImmediate& RHICmdList)
 				{
 					const FIntRect TileRect(0, 0, TileSize, TileSize);
 
@@ -280,10 +276,7 @@ namespace RuntimeVirtualTexture
 						Desc.NumPageDescs = 1;
 						for (int32 Layer = 0; Layer < NumLayers; Layer++)
 						{
-							Desc.Targets[Layer].Texture = RenderTileResources.GetRenderTarget(Layer);
-							Desc.Targets[Layer].TextureAccessBefore = RenderTargetsCurrentState;
-							// Leave the render targets in CopySrc state since we'll read from them right after, to avoid a useless transition :
-							Desc.Targets[Layer].TextureAccessAfter = ERHIAccess::CopySrc;
+							Desc.Targets[Layer].PooledRenderTarget = RenderTileResources.GetRenderTarget(Layer);
 							Desc.PageDescs[0].DestRect[Layer] = TileRect;
 						}
 						Desc.PageDescs[0].UVRange = UVRange;
@@ -291,15 +284,16 @@ namespace RuntimeVirtualTexture
 
 						RuntimeVirtualTexture::RenderPages(GraphBuilder, Desc);
 
-						GraphBuilder.Execute();
-					}
+						// Copy to staging
+						for (int32 Layer = 0; Layer < NumLayers; Layer++)
+						{
+							FRDGTextureRef RenderTarget = GraphBuilder.RegisterExternalTexture(RenderTileResources.GetRenderTarget(Layer), ERDGTextureFlags::None);
+							FRDGTextureRef StagingTexture = GraphBuilder.RegisterExternalTexture(CreateRenderTarget(RenderTileResources.GetStagingTexture(Layer), TEXT("StagingTexture")));
+							AddCopyTexturePass(GraphBuilder, RenderTarget, StagingTexture, FRHICopyTextureInfo());
+						}
 
-					// Copy to staging
-					for (int32 Layer = 0; Layer < NumLayers; Layer++)
-					{
-						RHICmdList.CopyTexture(RenderTileResources.GetRenderTarget(Layer), RenderTileResources.GetStagingTexture(Layer), FRHICopyTextureInfo());
+						GraphBuilder.Execute();
 					}
-					RenderTargetsCurrentState = ERHIAccess::CopySrc;
 
 					RenderTileResources.GetFence()->Clear();
 					RHICmdList.WriteGPUFence(RenderTileResources.GetFence());
diff --git a/Engine/Source/Editor/VirtualTexturingEditor/Private/RuntimeVirtualTextureThumbnailRenderer.cpp b/Engine/Source/Editor/VirtualTexturingEditor/Private/RuntimeVirtualTextureThumbnailRenderer.cpp
--- a/Engine/Source/Editor/VirtualTexturingEditor/Private/RuntimeVirtualTextureThumbnailRenderer.cpp
+++ b/Engine/Source/Editor/VirtualTexturingEditor/Private/RuntimeVirtualTextureThumbnailRenderer.cpp
@@ -107,6 +107,8 @@ void URuntimeVirtualTextureThumbnailRenderer::Draw(UObject* Object, int32 X, int
 	{
 		FMaterialRenderProxy::UpdateDeferredCachedUniformExpressions();
 
+		TRefCountPtr<IPooledRenderTarget> PooledRenderTarget = CreateRenderTarget(RenderTarget->GetRenderTargetTexture(), TEXT("RenderTarget"));
+
 		FRDGBuilder GraphBuilder(RHICmdList);
 		FScenePrimitiveRenderingContextScopeHelper RenderingScope(GetRendererModule().BeginScenePrimitiveRendering(GraphBuilder, *Scene));
 
@@ -122,10 +124,7 @@ void URuntimeVirtualTextureThumbnailRenderer::Draw(UObject* Object, int32 X, int
 		Desc.FixedColor = FLinearColor::Transparent;
 		Desc.CustomMaterialData = CustomMaterialData;
 		Desc.NumPageDescs = 1;
-		Desc.Targets[0].Texture = RenderTarget->GetRenderTargetTexture();
-		// The target is in RTV state and should be left in RTV state (RenderPage needs the input/output state) : 
-		Desc.Targets[0].TextureAccessBefore = ERHIAccess::RTV;
-		Desc.Targets[0].TextureAccessAfter = ERHIAccess::RTV;
+		Desc.Targets[0].PooledRenderTarget = PooledRenderTarget;
 		Desc.PageDescs[0].DestRect[0] = DestRect;
 		Desc.PageDescs[0].UVRange = FBox2D(FVector2D(0, 0), FVector2D(1, 1));
 		Desc.PageDescs[0].vLevel = IntCastChecked<uint8>(MaxLevel);
diff --git a/Engine/Source/Runtime/Renderer/Private/VT/RuntimeVirtualTextureProducer.cpp b/Engine/Source/Runtime/Renderer/Private/VT/RuntimeVirtualTextureProducer.cpp
--- a/Engine/Source/Runtime/Renderer/Private/VT/RuntimeVirtualTextureProducer.cpp
+++ b/Engine/Source/Runtime/Renderer/Private/VT/RuntimeVirtualTextureProducer.cpp
@@ -56,7 +56,6 @@ void FRuntimeVirtualTextureFinalizer::RenderFinalize(FRDGBuilder& GraphBuilder,
 	
 	for (int LayerIndex = 0; LayerIndex < RuntimeVirtualTexture::MaxTextureLayers; ++LayerIndex)
 	{
-		RenderPageBatchDesc.Targets[LayerIndex].Texture = Tiles[0].Targets[LayerIndex].TextureRHI != nullptr ? Tiles[0].Targets[LayerIndex].TextureRHI->GetTexture2D() : nullptr;
 		RenderPageBatchDesc.Targets[LayerIndex].PooledRenderTarget = Tiles[0].Targets[LayerIndex].PooledRenderTarget;
 	}
 
@@ -67,7 +66,7 @@ void FRuntimeVirtualTextureFinalizer::RenderFinalize(FRDGBuilder& GraphBuilder,
 		for (int LayerIndex = 0; LayerIndex < RuntimeVirtualTexture::MaxTextureLayers; ++LayerIndex)
 		{
 			// This should never happen which is why we don't bother sorting to maximize batch size
-			bBreakBatchForTextures |= (RenderPageBatchDesc.Targets[LayerIndex].Texture != Entry.Targets[LayerIndex].TextureRHI);
+			bBreakBatchForTextures |= (RenderPageBatchDesc.Targets[LayerIndex].PooledRenderTarget != Entry.Targets[LayerIndex].PooledRenderTarget);
 		}
 
 		if (BatchSize == RuntimeVirtualTexture::MaxRenderPageBatch || bBreakBatchForTextures)
@@ -81,7 +80,6 @@ void FRuntimeVirtualTextureFinalizer::RenderFinalize(FRDGBuilder& GraphBuilder,
 		{
 			for (int LayerIndex = 0; LayerIndex < RuntimeVirtualTexture::MaxTextureLayers; ++LayerIndex)
 			{
-				RenderPageBatchDesc.Targets[LayerIndex].Texture = Entry.Targets[LayerIndex].TextureRHI != nullptr ? Entry.Targets[LayerIndex].TextureRHI->GetTexture2D() : nullptr;
 				RenderPageBatchDesc.Targets[LayerIndex].PooledRenderTarget = Entry.Targets[LayerIndex].PooledRenderTarget;
 			}
 		}
@@ -196,7 +194,7 @@ IVirtualTextureFinalizer* FRuntimeVirtualTextureProducer::ProducePageData(
 
 	for (int LayerIndex = 0; LayerIndex < RuntimeVirtualTexture::MaxTextureLayers; ++LayerIndex)
 	{
-		if (TargetLayers[LayerIndex].TextureRHI != nullptr)
+		if (TargetLayers[LayerIndex].PooledRenderTarget != nullptr)
 		{
 			Tile.Targets[LayerIndex] = TargetLayers[LayerIndex];
 		}
diff --git a/Engine/Source/Runtime/Renderer/Private/VT/RuntimeVirtualTextureRender.cpp b/Engine/Source/Runtime/Renderer/Private/VT/RuntimeVirtualTextureRender.cpp
--- a/Engine/Source/Runtime/Renderer/Private/VT/RuntimeVirtualTextureRender.cpp
+++ b/Engine/Source/Runtime/Renderer/Private/VT/RuntimeVirtualTextureRender.cpp
@@ -970,7 +970,6 @@ namespace RuntimeVirtualTexture
 			ERuntimeVirtualTextureMaterialType MaterialType;
 			FIntPoint TextureSize;
 			int32 PageCount = 1;
-			EPixelFormat OutputFormat0 = PF_Unknown;
 			TArray<TRefCountPtr<IPooledRenderTarget>> OutputTargets;
 			bool bClearTextures = false;
 			bool bIsThumbnails = false;
@@ -983,7 +982,6 @@ namespace RuntimeVirtualTexture
 				MaterialType = InDesc.MaterialType;
 				PageCount = InDesc.NumPageDescs;
 				TextureSize = InDesc.PageDescs[0].DestRect[0].Size();
-				OutputFormat0 = InDesc.Targets[0].Texture != nullptr ? InDesc.Targets[0].Texture->GetFormat() : PF_Unknown;
 				OutputTargets.Add(InDesc.Targets[0].PooledRenderTarget);
 				OutputTargets.Add(InDesc.Targets[1].PooledRenderTarget);
 				OutputTargets.Add(InDesc.Targets[2].PooledRenderTarget);
@@ -1022,10 +1020,12 @@ namespace RuntimeVirtualTexture
 		/** Initialize the graph setup. */
 		void Init(FRDGBuilder& GraphBuilder, FInitDesc const& Desc)
 		{
-			bRenderPass = Desc.OutputFormat0 != PF_Unknown;
+			const EPixelFormat OutputFormat0 = Desc.OutputTargets[0].IsValid() ? Desc.OutputTargets[0]->GetRHI()->GetFormat() : PF_Unknown;
+
+			bRenderPass = OutputFormat0 != PF_Unknown;
 			bCopyThumbnailPass = bRenderPass && Desc.bIsThumbnails;
-			const bool bCompressedFormat = GPixelFormats[Desc.OutputFormat0].BlockSizeX == 4 && GPixelFormats[Desc.OutputFormat0].BlockSizeY == 4;
-			const bool bLQFormat = Desc.OutputFormat0 == PF_R5G6B5_UNORM;
+			const bool bCompressedFormat = GPixelFormats[OutputFormat0].BlockSizeX == 4 && GPixelFormats[OutputFormat0].BlockSizeY == 4;
+			const bool bLQFormat = OutputFormat0 == PF_R5G6B5_UNORM;
 			bCompressPass = bRenderPass && !bCopyThumbnailPass && bCompressedFormat;
 			bCopyPass = bRenderPass && !bCopyThumbnailPass && !bCompressPass && 
 				(Desc.MaterialType == ERuntimeVirtualTextureMaterialType::BaseColor_Normal_Specular || Desc.MaterialType == ERuntimeVirtualTextureMaterialType::BaseColor_Normal_Specular_YCoCg || Desc.MaterialType == ERuntimeVirtualTextureMaterialType::BaseColor_Normal_Specular_Mask_YCoCg || Desc.MaterialType == ERuntimeVirtualTextureMaterialType::Mask4);
@@ -1033,6 +1033,10 @@ namespace RuntimeVirtualTexture
 			// Use direct aliasing for compression pass on platforms that support it.
 			bDirectAliasing = bCompressedFormat && GRHISupportsUAVFormatAliasing && CVarVTDirectCompress.GetValueOnRenderThread() != 0;
 
+			// ForceImmediateFirstBarrier so that UAV transitions for the output targets aren't hoisted above
+			// Finalize() and into RenderFinalize() or earlier where they will be incorrect for virtual texture sampling.
+			const ERDGTextureFlags ExternalTextureFlags = ERDGTextureFlags::ForceImmediateFirstBarrier;
+
 			// Some problems happen when we don't use ERenderTargetLoadAction::EClear:
 			// * Some RHI need explicit flag to avoid a fast clear (TexCreate_NoFastClear).
 			// * DX12 RHI has a bug with RDG transient allocator (UE-173023) so we use TexCreate_Shared to avoid that.
@@ -1058,7 +1062,7 @@ namespace RuntimeVirtualTexture
 				{
 					if (bDirectAliasing)
 					{
-						CompressTexture0 = GraphBuilder.RegisterExternalTexture(Desc.OutputTargets[0], ERDGTextureFlags::None);
+						CompressTexture0 = GraphBuilder.RegisterExternalTexture(Desc.OutputTargets[0], ExternalTextureFlags);
 						CompressTextureUAV0_64bit = GraphBuilder.CreateUAV(FRDGTextureUAVDesc(CompressTexture0, 0, Compressed64BitFormat));
 						OutputAlias0 = nullptr;
 					}
@@ -1083,9 +1087,9 @@ namespace RuntimeVirtualTexture
 				{
 					if (bDirectAliasing)
 					{
-						CompressTexture0 = GraphBuilder.RegisterExternalTexture(Desc.OutputTargets[0], ERDGTextureFlags::None);
+						CompressTexture0 = GraphBuilder.RegisterExternalTexture(Desc.OutputTargets[0], ExternalTextureFlags);
 						CompressTextureUAV0_64bit = GraphBuilder.CreateUAV(FRDGTextureUAVDesc(CompressTexture0, 0, Compressed64BitFormat));
-						CompressTexture1 = GraphBuilder.RegisterExternalTexture(Desc.OutputTargets[1], ERDGTextureFlags::None);
+						CompressTexture1 = GraphBuilder.RegisterExternalTexture(Desc.OutputTargets[1], ExternalTextureFlags);
 						CompressTextureUAV1_128bit = GraphBuilder.CreateUAV(FRDGTextureUAVDesc(CompressTexture1, 0, Compressed128BitFormat));
 						OutputAlias0 = OutputAlias1 = nullptr;
 					}
@@ -1114,9 +1118,9 @@ namespace RuntimeVirtualTexture
 				{
 					if (bDirectAliasing)
 					{
-						CompressTexture0 = GraphBuilder.RegisterExternalTexture(Desc.OutputTargets[0], ERDGTextureFlags::None);
+						CompressTexture0 = GraphBuilder.RegisterExternalTexture(Desc.OutputTargets[0], ExternalTextureFlags);
 						CompressTextureUAV0_128bit = GraphBuilder.CreateUAV(FRDGTextureUAVDesc(CompressTexture0, 0, Compressed128BitFormat));
-						CompressTexture1 = GraphBuilder.RegisterExternalTexture(Desc.OutputTargets[1], ERDGTextureFlags::None);
+						CompressTexture1 = GraphBuilder.RegisterExternalTexture(Desc.OutputTargets[1], ExternalTextureFlags);
 						CompressTextureUAV1_128bit = GraphBuilder.CreateUAV(FRDGTextureUAVDesc(CompressTexture1, 0, Compressed128BitFormat));
 						OutputAlias0 = OutputAlias1 = nullptr;
 					}
@@ -1149,11 +1153,11 @@ namespace RuntimeVirtualTexture
 				{
 					if (bDirectAliasing)
 					{
-						CompressTexture0 = GraphBuilder.RegisterExternalTexture(Desc.OutputTargets[0], ERDGTextureFlags::None);
+						CompressTexture0 = GraphBuilder.RegisterExternalTexture(Desc.OutputTargets[0], ExternalTextureFlags);
 						CompressTextureUAV0_128bit = GraphBuilder.CreateUAV(FRDGTextureUAVDesc(CompressTexture0, 0, Compressed128BitFormat));
-						CompressTexture1 = GraphBuilder.RegisterExternalTexture(Desc.OutputTargets[1], ERDGTextureFlags::None);
+						CompressTexture1 = GraphBuilder.RegisterExternalTexture(Desc.OutputTargets[1], ExternalTextureFlags);
 						CompressTextureUAV1_128bit = GraphBuilder.CreateUAV(FRDGTextureUAVDesc(CompressTexture1, 0, Compressed128BitFormat));
-						CompressTexture2 = GraphBuilder.RegisterExternalTexture(Desc.OutputTargets[2], ERDGTextureFlags::None);
+						CompressTexture2 = GraphBuilder.RegisterExternalTexture(Desc.OutputTargets[2], ExternalTextureFlags);
 						CompressTextureUAV2_64bit = GraphBuilder.CreateUAV(FRDGTextureUAVDesc(CompressTexture2, 0, Compressed64BitFormat));
 
 						OutputAlias0 = OutputAlias1 = OutputAlias2 = nullptr;
@@ -1190,11 +1194,11 @@ namespace RuntimeVirtualTexture
 				{
 					if (bDirectAliasing)
 					{
-						CompressTexture0 = GraphBuilder.RegisterExternalTexture(Desc.OutputTargets[0], ERDGTextureFlags::None);
+						CompressTexture0 = GraphBuilder.RegisterExternalTexture(Desc.OutputTargets[0], ExternalTextureFlags);
 						CompressTextureUAV0_128bit = GraphBuilder.CreateUAV(FRDGTextureUAVDesc(CompressTexture0, 0, Compressed128BitFormat));
-						CompressTexture1 = GraphBuilder.RegisterExternalTexture(Desc.OutputTargets[1], ERDGTextureFlags::None);
+						CompressTexture1 = GraphBuilder.RegisterExternalTexture(Desc.OutputTargets[1], ExternalTextureFlags);
 						CompressTextureUAV1_128bit = GraphBuilder.CreateUAV(FRDGTextureUAVDesc(CompressTexture1, 0, Compressed128BitFormat));
-						CompressTexture2 = GraphBuilder.RegisterExternalTexture(Desc.OutputTargets[2], ERDGTextureFlags::None);
+						CompressTexture2 = GraphBuilder.RegisterExternalTexture(Desc.OutputTargets[2], ExternalTextureFlags);
 						CompressTextureUAV2_128bit = GraphBuilder.CreateUAV(FRDGTextureUAVDesc(CompressTexture2, 0, Compressed128BitFormat));
 
 						OutputAlias0 = OutputAlias1 = OutputAlias2 = nullptr;
@@ -1230,7 +1234,7 @@ namespace RuntimeVirtualTexture
 				{
 					if (bDirectAliasing)
 					{
-						CompressTexture0 = GraphBuilder.RegisterExternalTexture(Desc.OutputTargets[0], ERDGTextureFlags::None);
+						CompressTexture0 = GraphBuilder.RegisterExternalTexture(Desc.OutputTargets[0], ExternalTextureFlags);
 						CompressTextureUAV0_128bit = GraphBuilder.CreateUAV(FRDGTextureUAVDesc(CompressTexture0, 0, Compressed128BitFormat));
 						OutputAlias0 = nullptr;
 					}
@@ -1264,7 +1268,7 @@ namespace RuntimeVirtualTexture
 				{
 					if (bDirectAliasing)
 					{
-						CompressTexture0 = GraphBuilder.RegisterExternalTexture(Desc.OutputTargets[0], ERDGTextureFlags::None);
+						CompressTexture0 = GraphBuilder.RegisterExternalTexture(Desc.OutputTargets[0], ExternalTextureFlags);
 						CompressTextureUAV0_64bit = GraphBuilder.CreateUAV(FRDGTextureUAVDesc(CompressTexture0, 0, Compressed64BitFormat));
 						OutputAlias0 = nullptr;
 					}
@@ -1280,6 +1284,19 @@ namespace RuntimeVirtualTexture
 				}
 				break;
 			}
+
+			if (OutputAlias0)
+			{
+				TargetTexture0 = GraphBuilder.RegisterExternalTexture(Desc.OutputTargets[0], ExternalTextureFlags);
+			}
+			if (OutputAlias1)
+			{
+				TargetTexture1 = GraphBuilder.RegisterExternalTexture(Desc.OutputTargets[1], ExternalTextureFlags);
+			}
+			if (OutputAlias2)
+			{
+				TargetTexture2 = GraphBuilder.RegisterExternalTexture(Desc.OutputTargets[2], ExternalTextureFlags);
+			}
 		}
 
 		/** Flags to express what passes we need for this virtual texture layout. */
@@ -1310,6 +1327,10 @@ namespace RuntimeVirtualTexture
 		FRDGTextureRef OutputAlias0 = nullptr;
 		FRDGTextureRef OutputAlias1 = nullptr;
 		FRDGTextureRef OutputAlias2 = nullptr;
+		/** If we have output aliases, then these will containg the final physical texture targets. */
+		FRDGTextureRef TargetTexture0 = nullptr;
+		FRDGTextureRef TargetTexture1 = nullptr;
+		FRDGTextureRef TargetTexture2 = nullptr;
 	};
 
 	/** 
@@ -1949,48 +1970,6 @@ namespace RuntimeVirtualTexture
 		}
 	}
 
-	BEGIN_SHADER_PARAMETER_STRUCT(FCopyToOutputParameters, )
-		RDG_TEXTURE_ACCESS(Input, ERHIAccess::CopySrc)
-	END_SHADER_PARAMETER_STRUCT()
-
-	/** Set up the copy to final output physical texture. */
-	void AddCopyToOutputPass(
-		FRDGBuilder& GraphBuilder, 
-		FRDGTextureRef InputTexture, 
-		FRHITexture* OutputTexture, 
-		ERHIAccess OutputTextureAccessBefore, 
-		ERHIAccess OutputTextureAccessAfter, 
-		int32 SourceSliceIndex, 
-		FIntPoint const& DestPosition)
-	{
-		FRHICopyTextureInfo CopyInfo;
-		CopyInfo.Size = InputTexture->Desc.GetSize();
-		CopyInfo.SourceSliceIndex = SourceSliceIndex;
-		CopyInfo.DestPosition = FIntVector(DestPosition.X, DestPosition.Y, 0);
-
-		FCopyToOutputParameters* Parameters = GraphBuilder.AllocParameters<FCopyToOutputParameters>();
-		Parameters->Input = InputTexture;
-
-		GraphBuilder.AddPass(
-			RDG_EVENT_NAME("VirtualTextureCopyToOutput"),
-			Parameters,
-			ERDGPassFlags::Copy | ERDGPassFlags::NeverCull,
-			[InputTexture, OutputTexture, OutputTextureAccessBefore, OutputTextureAccessAfter, CopyInfo](FRDGAsyncTask, FRHICommandList& RHICmdList)
-		{
-			if (OutputTextureAccessBefore != ERHIAccess::Unknown)
-			{
-				RHICmdList.Transition(FRHITransitionInfo(OutputTexture, OutputTextureAccessBefore, ERHIAccess::CopyDest));
-			}
-
-			RHICmdList.CopyTexture(InputTexture->GetRHI(), OutputTexture, CopyInfo);
-
-			if (OutputTextureAccessAfter != ERHIAccess::Unknown)
-			{
-				RHICmdList.Transition(FRHITransitionInfo(OutputTexture, ERHIAccess::CopyDest, OutputTextureAccessAfter));
-			}
-		});
-	}
-
 	/** Get the debug color to use for a given mip level. Returns InDefaultColor if mip debugging is disabled. */
 	FLinearColor GetDebugMipLevelColor(uint32 InLevel, FLinearColor InDefaultColor)
 	{
@@ -2199,12 +2178,15 @@ namespace RuntimeVirtualTexture
 	void CopyPagesToOutput(FRDGBuilder& GraphBuilder, FBatchRenderContext const& BatchRenderContext)
 	{
 		FRenderGraphSetup const& GraphSetup = BatchRenderContext.GraphSetup;
-		FRenderPageBatchDesc const& BatchDesc = BatchRenderContext.BatchDesc;
+		if (GraphSetup.OutputAlias0 == nullptr && GraphSetup.OutputAlias1 == nullptr && GraphSetup.OutputAlias2 == nullptr)
+		{
+			return;
+		}
 
-		FRDGTextureRef OutputAlias[MaxTextureLayers] = { GraphSetup.OutputAlias0, GraphSetup.OutputAlias1, GraphSetup.OutputAlias2 };
-		FRHITexture* OutputTexture[MaxTextureLayers] = { BatchDesc.Targets[0].Texture, BatchDesc.Targets[1].Texture, BatchDesc.Targets[2].Texture };
-		ERHIAccess OutputTextureAccessBefore[MaxTextureLayers] = { BatchDesc.Targets[0].TextureAccessBefore, BatchDesc.Targets[1].TextureAccessBefore, BatchDesc.Targets[2].TextureAccessBefore };
-		ERHIAccess OutputTextureAccessAfter[MaxTextureLayers] = { BatchDesc.Targets[0].TextureAccessAfter, BatchDesc.Targets[1].TextureAccessAfter, BatchDesc.Targets[2].TextureAccessAfter };
+		FRenderPageBatchDesc const& BatchDesc = BatchRenderContext.BatchDesc;
+		const FRDGTextureRef SourceTexture[MaxTextureLayers] = { GraphSetup.OutputAlias0, GraphSetup.OutputAlias1, GraphSetup.OutputAlias2 };
+		const FRDGTextureRef DestTexture[MaxTextureLayers] = { GraphSetup.TargetTexture0, GraphSetup.TargetTexture1, GraphSetup.TargetTexture2 };
+		const FIntVector CopySize = SourceTexture[0] ? SourceTexture[0]->Desc.GetSize() : FIntVector(0, 0, 0);
 
 		for (int32 PageIndex = 0; PageIndex < BatchDesc.NumPageDescs; ++PageIndex)
 		{
@@ -2212,12 +2194,14 @@ namespace RuntimeVirtualTexture
 
 			for (int32 LayerIndex = 0; LayerIndex < MaxTextureLayers; ++LayerIndex)
 			{
-				if (OutputAlias[LayerIndex] != nullptr && OutputTexture[LayerIndex] != nullptr)
+				if (SourceTexture[LayerIndex] != nullptr && DestTexture[LayerIndex] != nullptr)
 				{
-					const ERHIAccess TransitionBefore = PageIndex == 0 ? OutputTextureAccessBefore[LayerIndex] : ERHIAccess::Unknown;
-					const ERHIAccess TransitionAfter = PageIndex == BatchDesc.NumPageDescs - 1 ? OutputTextureAccessAfter[LayerIndex] : ERHIAccess::Unknown;
-					
-					AddCopyToOutputPass(GraphBuilder, OutputAlias[LayerIndex], OutputTexture[LayerIndex], TransitionBefore, TransitionAfter, PageIndex, PageDesc.DestRect[LayerIndex].Min);
+					FRHICopyTextureInfo CopyInfo;
+					CopyInfo.Size = CopySize;
+					CopyInfo.SourceSliceIndex = PageIndex;
+					CopyInfo.DestPosition = FIntVector(PageDesc.DestRect[LayerIndex].Min.X, PageDesc.DestRect[LayerIndex].Min.Y, 0);
+
+					AddCopyTexturePass(GraphBuilder, SourceTexture[LayerIndex], DestTexture[LayerIndex], CopyInfo);
 				}
 			}
 		}
diff --git a/Engine/Source/Runtime/Renderer/Private/VT/RuntimeVirtualTextureSceneProxy.cpp b/Engine/Source/Runtime/Renderer/Private/VT/RuntimeVirtualTextureSceneProxy.cpp
--- a/Engine/Source/Runtime/Renderer/Private/VT/RuntimeVirtualTextureSceneProxy.cpp
+++ b/Engine/Source/Runtime/Renderer/Private/VT/RuntimeVirtualTextureSceneProxy.cpp
@@ -2,9 +2,10 @@
 
 #include "VT/RuntimeVirtualTextureSceneProxy.h"
 
-#include "ComponentRecreateRenderStateContext.h"
 #include "Components/RuntimeVirtualTextureComponent.h"
 #include "RendererOnScreenNotification.h"
+#include "SceneInterface.h"
+#include "SceneUtils.h"
 #include "VirtualTextureSystem.h"
 #include "VT/RuntimeVirtualTextureEnum.h"
 #include "VT/RuntimeVirtualTexture.h"
@@ -12,7 +13,6 @@
 #include "VT/VirtualTexture.h"
 #include "VT/VirtualTextureBuilder.h"
 #include "VT/VirtualTextureScalability.h"
-#include "SceneUtils.h"
 
 #define LOCTEXT_NAMESPACE "VirtualTexture"
 
@@ -126,11 +126,11 @@ FRuntimeVirtualTextureSceneProxy::FRuntimeVirtualTextureSceneProxy(URuntimeVirtu
 			// Store the ProducerHandle, SpaceID and AllocatedVirtualTexture immediately after virtual texture is initialized.
 			ENQUEUE_RENDER_COMMAND(GetProducerHandle)(
 				[this, VirtualTexturePtr = VirtualTexture](FRHICommandList& RHICmdList)
-				{
-					ProducerHandle = VirtualTexturePtr->GetProducerHandle();
-					AllocatedVirtualTexture = VirtualTexturePtr->GetAllocatedVirtualTexture();
-					SpaceID = AllocatedVirtualTexture->GetSpaceID();
-				});
+			{
+				ProducerHandle = VirtualTexturePtr->GetProducerHandle();
+				AllocatedVirtualTexture = VirtualTexturePtr->GetAllocatedVirtualTexture();
+				SpaceID = AllocatedVirtualTexture->GetSpaceID();
+			});
 		}
 		else
 		{
diff --git a/Engine/Source/Runtime/Renderer/Public/VT/RuntimeVirtualTextureRender.h b/Engine/Source/Runtime/Renderer/Public/VT/RuntimeVirtualTextureRender.h
--- a/Engine/Source/Runtime/Renderer/Public/VT/RuntimeVirtualTextureRender.h
+++ b/Engine/Source/Runtime/Renderer/Public/VT/RuntimeVirtualTextureRender.h
@@ -22,16 +22,20 @@ namespace RuntimeVirtualTexture
 	/** Structure containing a texture layer target description for a call for RenderPages(). */
 	struct FRenderPageTarget
 	{
-		/** Physical texture to render to. */
-		FRHITexture* Texture = nullptr;
+		/** Physical target to render to. */
+		IPooledRenderTarget* PooledRenderTarget = nullptr;
 
-		/** State of the physical texture to render to when it's being rendered */
+		UE_DEPRECATED(5.6, "Use PooledRenderTarget instead.")
+		FRHITexture* Texture = nullptr;
+		UE_DEPRECATED(5.6, "PooledRenderTarget tracks its own state.")
 		ERHIAccess TextureAccessBefore = ERHIAccess::SRVMask;
-
-		/** State the physical texture to render should be left in after it's being rendered */
+		UE_DEPRECATED(5.6, "PooledRenderTarget tracks its own state.")
 		ERHIAccess TextureAccessAfter = ERHIAccess::SRVMask;
 
-		IPooledRenderTarget* PooledRenderTarget = nullptr;
+		// Disable deprecation warnings for implicit copying of deprecated members.
+		PRAGMA_DISABLE_DEPRECATION_WARNINGS
+		FRenderPageTarget& operator=(FRenderPageTarget const& RHS) = default;
+		PRAGMA_ENABLE_DEPRECATION_WARNINGS
 	};
 
 	/** A single page description. Multiple of these can be placed in a single FRenderPageBatchDesc batch description. */
