Remove direct access to VT physical texture RHI.
Should use PooledRenderTarget or RHI SRV only.
This is part of ensuring correct transition usage.
#rnx

diff --git a/Engine/Source/Runtime/Renderer/Private/VT/AllocatedVirtualTexture.cpp b/Engine/Source/Runtime/Renderer/Private/VT/AllocatedVirtualTexture.cpp
--- a/Engine/Source/Runtime/Renderer/Private/VT/AllocatedVirtualTexture.cpp
+++ b/Engine/Source/Runtime/Renderer/Private/VT/AllocatedVirtualTexture.cpp
@@ -442,8 +442,14 @@ FRHITexture* FAllocatedVirtualTexture::GetPhysicalTexture(uint32 InLayerIndex) c
 {
 	if (InLayerIndex < Description.NumTextureLayers)
 	{
-		const FVirtualTexturePhysicalSpace* PhysicalSpace = UniquePageTableLayers[TextureLayers[InLayerIndex].UniquePageTableLayerIndex].PhysicalSpace;
-		return PhysicalSpace ? PhysicalSpace->GetPhysicalTexture(TextureLayers[InLayerIndex].PhysicalTextureIndex) : nullptr;
+		if (const FVirtualTexturePhysicalSpace* PhysicalSpace = UniquePageTableLayers[TextureLayers[InLayerIndex].UniquePageTableLayerIndex].PhysicalSpace)
+		{
+			TRefCountPtr<IPooledRenderTarget> PooledRenderTarget = PhysicalSpace->GetPhysicalTexture(TextureLayers[InLayerIndex].PhysicalTextureIndex);
+			if (PooledRenderTarget.IsValid())
+			{
+				return PooledRenderTarget->GetRHI();
+			}
+		}
 	}
 	return nullptr;
 }
diff --git a/Engine/Source/Runtime/Renderer/Private/VT/VirtualTexturePhysicalSpace.h b/Engine/Source/Runtime/Renderer/Private/VT/VirtualTexturePhysicalSpace.h
--- a/Engine/Source/Runtime/Renderer/Private/VT/VirtualTexturePhysicalSpace.h
+++ b/Engine/Source/Runtime/Renderer/Private/VT/VirtualTexturePhysicalSpace.h
@@ -123,23 +123,17 @@ public:
 	inline uint32 ReleaseResourceRef() { check(NumResourceRefs > 0u); return --NumResourceRefs; }
 	inline uint32 GetResourceRefCount() const { return NumResourceRefs; }
 
-	FRHITexture* GetPhysicalTexture(int32 Layer) const
+	inline TRefCountPtr<IPooledRenderTarget> GetPhysicalTexture(int32 Layer) const
 	{
-		check(PooledRenderTarget[Layer].IsValid());
-		return PooledRenderTarget[Layer]->GetRHI();
+		return PooledRenderTarget[Layer];
 	}
 
-	FRHIShaderResourceView* GetPhysicalTextureSRV(int32 Layer, bool bSRGB) const
+	inline FRHIShaderResourceView* GetPhysicalTextureSRV(int32 Layer, bool bSRGB) const
 	{
 		return bSRGB ? TextureSRV_SRGB[Layer] : TextureSRV[Layer];
 	}
 
-	TRefCountPtr<IPooledRenderTarget> GetPhysicalTexturePooledRenderTarget(int32 Layer) const
-	{
-		check(PooledRenderTarget[Layer].IsValid());
-		return PooledRenderTarget[Layer];
-	}
-
+	/** Gather the necessary physical texture transitions back to SRVMask after the producer Finalize() point in the frame. */
 	void FinalizeTextures(FRDGBuilder& GraphBuilder, FRDGExternalAccessQueue& ExternalAccessQueue);
 
 	/** Update internal tracking of residency. This is used to update stats and to calculate a mip bias to keep within the pool budget. */
diff --git a/Engine/Source/Runtime/Renderer/Private/VT/VirtualTextureSystem.cpp b/Engine/Source/Runtime/Renderer/Private/VT/VirtualTextureSystem.cpp
--- a/Engine/Source/Runtime/Renderer/Private/VT/VirtualTextureSystem.cpp
+++ b/Engine/Source/Runtime/Renderer/Private/VT/VirtualTextureSystem.cpp
@@ -2106,8 +2106,7 @@ int32 FVirtualTextureSystem::SubmitRequestsFromLocalTileRequests(FConcurrentLine
 				{
 					if (CurrentProducer->GetPhysicalGroupIndexForTextureLayer(ProducerLayerIndex) == ProducerPhysicalGroupIndex)
 					{
-						ProduceTarget[ProducerLayerIndex].TextureRHI = PhysicalSpace->GetPhysicalTexture(PhysicalLocalTextureIndex);
-						ProduceTarget[ProducerLayerIndex].PooledRenderTarget = PhysicalSpace->GetPhysicalTexturePooledRenderTarget(PhysicalLocalTextureIndex);
+						ProduceTarget[ProducerLayerIndex].PooledRenderTarget = PhysicalSpace->GetPhysicalTexture(PhysicalLocalTextureIndex);
 						ProduceTarget[ProducerLayerIndex].pPageLocation = PhysicalSpace->GetPhysicalLocation(pAddress);
 						LayerMask |= 1 << ProducerLayerIndex;
 						PhysicalLocalTextureIndex++;
@@ -2328,10 +2327,8 @@ void FVirtualTextureSystem::SubmitRequests(FRHICommandList& RHICmdList, ERHIFeat
 							{
 								if (Producer.GetPhysicalGroupIndexForTextureLayer(ProducerLayerIndex) == ProducerPhysicalGroupIndex)
 								{
-									ProduceTarget[ProducerLayerIndex].TextureRHI = PhysicalSpace->GetPhysicalTexture(PhysicalLocalTextureIndex);
-									ProduceTarget[ProducerLayerIndex].PooledRenderTarget = PhysicalSpace->GetPhysicalTexturePooledRenderTarget(PhysicalLocalTextureIndex);
+									ProduceTarget[ProducerLayerIndex].PooledRenderTarget = PhysicalSpace->GetPhysicalTexture(PhysicalLocalTextureIndex);
 									ProduceTarget[ProducerLayerIndex].pPageLocation = PhysicalSpace->GetPhysicalLocation(pAddress);
-									
 									PhysicalLocalTextureIndex++;
 
 									Allocate_pAddress[ProducerPhysicalGroupIndex] = pAddress;
@@ -2376,8 +2373,7 @@ void FVirtualTextureSystem::SubmitRequests(FRHICommandList& RHICmdList, ERHIFeat
 							{
 								if (Producer.GetPhysicalGroupIndexForTextureLayer(ProducerLayerIndex) == ProducerPhysicalGroupIndex)
 								{
-									ProduceTarget[ProducerLayerIndex].TextureRHI = PhysicalSpace->GetPhysicalTexture(PhysicalLocalTextureIndex);
-									ProduceTarget[ProducerLayerIndex].PooledRenderTarget = PhysicalSpace->GetPhysicalTexturePooledRenderTarget(PhysicalLocalTextureIndex);
+									ProduceTarget[ProducerLayerIndex].PooledRenderTarget = PhysicalSpace->GetPhysicalTexture(PhysicalLocalTextureIndex);
 									ProduceTarget[ProducerLayerIndex].pPageLocation = PhysicalSpace->GetPhysicalLocation(pAddress);
 									PhysicalLocalTextureIndex++;
 								}
@@ -2532,6 +2528,7 @@ void FVirtualTextureSystem::FinalizeRequests(FRDGBuilder& GraphBuilder, ISceneRe
 	FRDGExternalAccessQueue ExternalAccessQueue;
 
 	// Finalize requests
+	if (Finalizers.Num())
 	{
 		SCOPE_CYCLE_COUNTER(STAT_ProcessRequests_Finalize);
 		RDG_EVENT_SCOPE(GraphBuilder, "VirtualTextureFinalizeRequests");
