Mark page table textures as ForceImmediateFirstBarrier in the RDG graph so that render target view transitions for aren't hoisted into RenderFinalize() or earlier where they will be incorrect for virtual texture sampling.
#rnx

diff --git a/Engine/Source/Runtime/Renderer/Private/VT/VirtualTextureSpace.cpp b/Engine/Source/Runtime/Renderer/Private/VT/VirtualTextureSpace.cpp
--- a/Engine/Source/Runtime/Renderer/Private/VT/VirtualTextureSpace.cpp
+++ b/Engine/Source/Runtime/Renderer/Private/VT/VirtualTextureSpace.cpp
@@ -499,7 +499,9 @@ void FVirtualTextureSpace::ApplyUpdates(FVirtualTextureSystem* System, FRDGBuild
 		const uint32 LayerInTexture = LayerIndex % LayersPerPageTableTexture;
 		FTextureEntry& PageTableEntry = PageTable[TextureIndex];
 
-		FRDGTextureRef PageTableTexture = GraphBuilder.RegisterExternalTexture(PageTableEntry.RenderTarget, ERDGTextureFlags::None);
+		// ForceImmediateFirstBarrier so that RTV transitions for the page table textures aren't hoisted
+		// into RenderFinalize() or earlier where they will be incorrect for virtual texture sampling.
+		FRDGTextureRef PageTableTexture = GraphBuilder.RegisterExternalTexture(PageTableEntry.RenderTarget, ERDGTextureFlags::ForceImmediateFirstBarrier);
 		GraphBuilder.UseInternalAccessMode(PageTableTexture);
 
 		// Use color write mask to update the proper page table entry for this layer
