Fix nullptr assertion in RDG debug code.
#rnx

diff --git a/Engine/Source/Runtime/Renderer/Private/VT/RuntimeVirtualTextureRender.cpp b/Engine/Source/Runtime/Renderer/Private/VT/RuntimeVirtualTextureRender.cpp
--- a/Engine/Source/Runtime/Renderer/Private/VT/RuntimeVirtualTextureRender.cpp
+++ b/Engine/Source/Runtime/Renderer/Private/VT/RuntimeVirtualTextureRender.cpp
@@ -1285,15 +1285,15 @@ namespace RuntimeVirtualTexture
 				break;
 			}
 
-			if (OutputAlias0)
+			if (OutputAlias0 && Desc.OutputTargets[0])
 			{
 				TargetTexture0 = GraphBuilder.RegisterExternalTexture(Desc.OutputTargets[0], ExternalTextureFlags);
 			}
-			if (OutputAlias1)
+			if (OutputAlias1 && Desc.OutputTargets[1])
 			{
 				TargetTexture1 = GraphBuilder.RegisterExternalTexture(Desc.OutputTargets[1], ExternalTextureFlags);
 			}
-			if (OutputAlias2)
+			if (OutputAlias2 && Desc.OutputTargets[2])
 			{
 				TargetTexture2 = GraphBuilder.RegisterExternalTexture(Desc.OutputTargets[2], ExternalTextureFlags);
 			}
@@ -2080,7 +2080,7 @@ namespace RuntimeVirtualTexture
 			PassParameters->RenderTargets[0] = GraphSetup.RenderTexture0 ? FRenderTargetBinding(GraphSetup.RenderTexture0, LoadAction, 0, PageIndex) : FRenderTargetBinding();
 			PassParameters->RenderTargets[1] = GraphSetup.RenderTexture1 ? FRenderTargetBinding(GraphSetup.RenderTexture1, LoadAction, 0, PageIndex) : FRenderTargetBinding();
 			PassParameters->RenderTargets[2] = GraphSetup.RenderTexture2 ? FRenderTargetBinding(GraphSetup.RenderTexture2, LoadAction, 0, PageIndex) : FRenderTargetBinding();
-    		
+
 			const int32 RuntimeVirtualTextureId = BatchDesc.RuntimeVirtualTextureId;
 			const ERuntimeVirtualTextureMaterialType MaterialType = BatchDesc.MaterialType;
 			const bool bAllowCachedMeshDrawCommands = BatchRenderContext.bAllowCachedMeshDrawCommands;
