Add a workaround fix for occasional RVT page corruption that has been seen in the wild.
We haven't been able to reproduce at Epic, but testing with licensees has shown that inserting a single call to ID3D12GraphicsCommandList::SetPipelineState() is enough to fix the issue without significant overhead.
Investigations will continue to track down the root cause.
The workaround is enabled by default and can be toggled with the console variable "r.VT.RVT.PageCorruptionFix".

diff --git a/Engine/Source/Runtime/Renderer/Private/VT/RuntimeVirtualTextureRender.cpp b/Engine/Source/Runtime/Renderer/Private/VT/RuntimeVirtualTextureRender.cpp
--- a/Engine/Source/Runtime/Renderer/Private/VT/RuntimeVirtualTextureRender.cpp
+++ b/Engine/Source/Runtime/Renderer/Private/VT/RuntimeVirtualTextureRender.cpp
@@ -53,6 +53,12 @@ namespace RuntimeVirtualTexture
 		TEXT("Compress texture data direct to the physical texture on platforms that support it."),
 		ECVF_RenderThreadSafe);
 
+	static TAutoConsoleVariable<bool> CVarVTApplyPageCorruptionFix(
+		TEXT("r.VT.RVT.PageCorruptionFix"),
+		true,
+		TEXT("Apply change that has been found to fix some rare page corruption on PC."),
+		ECVF_RenderThreadSafe);
+
     int32 RenderCaptureNextRVTPagesDraws = 0;
     static FAutoConsoleVariableRef CVarRenderCaptureNextRVTPagesDraws(
 	    TEXT("r.VT.RenderCaptureNextPagesDraws"),
@@ -1970,6 +1976,41 @@ namespace RuntimeVirtualTexture
 		}
 	}
 
+	/** 
+	 * Fixes an issue where runtime virtual texture page corruption causes square artifacts.
+	 * Repro of the bug is rare. But it's been found that inserting a single call to ID3D12GraphicsCommandList::SetPipelineState()
+	 * before rendering any RVT pages resolves the issue.
+	 */
+	void ApplyPageCorruptionFix(FRDGBuilder& GraphBuilder, int32 InPageIndex, EShaderPlatform InShaderPlatform)
+	{
+		if (InPageIndex == 0 && CVarVTApplyPageCorruptionFix.GetValueOnRenderThread() && IsPCPlatform(InShaderPlatform) && IsD3DPlatform(InShaderPlatform))
+		{
+			GraphBuilder.AddPass(
+				RDG_EVENT_NAME("VirtualTextureDummyPass"),
+				ERDGPassFlags::NeverCull | ERDGPassFlags::NeverMerge,
+				[](FRHICommandList& RHICmdList)
+			{
+				FGlobalShaderMap* GlobalShaderMap = GetGlobalShaderMap(GMaxRHIFeatureLevel);
+				TShaderMapRef<FShader_VirtualTextureCopy_VS> VertexShader(GlobalShaderMap);
+				TShaderMapRef<FShader_VirtualTextureCopy_PS<ERuntimeVirtualTextureMaterialType::BaseColor>> PixelShader(GlobalShaderMap);
+
+				FGraphicsPipelineStateInitializer GraphicsPSOInit;
+				GraphicsPSOInit.RenderTargetsEnabled = 1;
+				GraphicsPSOInit.RenderTargetFormats[0] = PF_B8G8R8A8;
+				GraphicsPSOInit.RenderTargetFlags[0] = TexCreate_None;
+				GraphicsPSOInit.NumSamples = 1;
+				GraphicsPSOInit.DepthStencilState = TStaticDepthStencilState<false, CF_Always>::GetRHI();
+				GraphicsPSOInit.BlendState = TStaticBlendState<CW_RGBA, BO_Add, BF_One, BF_Zero, BO_Add, BF_One, BF_Zero>::GetRHI();
+				GraphicsPSOInit.RasterizerState = TStaticRasterizerState<>::GetRHI();
+				GraphicsPSOInit.PrimitiveType = PT_TriangleList;
+				GraphicsPSOInit.BoundShaderState.VertexDeclarationRHI = GetVertexDeclarationFVector4();
+				GraphicsPSOInit.BoundShaderState.VertexShaderRHI = VertexShader.GetVertexShader();
+				GraphicsPSOInit.BoundShaderState.PixelShaderRHI = PixelShader.GetPixelShader();
+				SetGraphicsPipelineState(RHICmdList, GraphicsPSOInit, 0, EApplyRendertargetOption::DoNothing);
+			});
+		}
+	}
+
 	/** Get the debug color to use for a given mip level. Returns InDefaultColor if mip debugging is disabled. */
 	FLinearColor GetDebugMipLevelColor(uint32 InLevel, FLinearColor InDefaultColor)
 	{
@@ -2065,6 +2106,8 @@ namespace RuntimeVirtualTexture
 			RenderCaptureInterface::FScopedCapture RenderCapture((RenderCaptureNextRVTPagesDraws != 0), GraphBuilder, TEXT("RenderRVTPage"));
 			RenderCaptureNextRVTPagesDraws = FMath::Max(RenderCaptureNextRVTPagesDraws - 1, 0);
 
+			ApplyPageCorruptionFix(GraphBuilder, PageIndex, Scene->GetShaderPlatform());
+
 			ERenderTargetLoadAction LoadAction = BatchDesc.bClearTextures ? ERenderTargetLoadAction::EClear : ERenderTargetLoadAction::ENoAction;
 			FShader_VirtualTextureMaterialDraw::FParameters* PassParameters = GraphBuilder.AllocParameters<FShader_VirtualTextureMaterialDraw::FParameters>();
 			PassParameters->View = View->ViewUniformBuffer;
