[NFOR] Fix temporal denoiser failure due to history always invalidated and variance not gathered.

#tests test done using MRG, NFOR enabled
#rnx

diff --git a/Engine/Source/Runtime/Renderer/Private/PathTracing.cpp b/Engine/Source/Runtime/Renderer/Private/PathTracing.cpp
--- a/Engine/Source/Runtime/Renderer/Private/PathTracing.cpp
+++ b/Engine/Source/Runtime/Renderer/Private/PathTracing.cpp
@@ -3546,7 +3546,7 @@ void FDeferredShadingSceneRenderer::RenderPathTracing(
 	Config.PathTracingData.MaxSamples = MaxSPP;
 
 	FPathTracingState* PathTracingState = GetPathTracingStateFromView(View);
-	const bool bFirstTime = !PathTracingState->RadianceRT.IsValid(); // we just initialized (or reset) the option state for this view -- don't bother comparing in this case
+	const bool bFirstTime = !PathTracingState->RadianceRT.IsValid() && !PathTracingState->LastRadianceRT.IsValid(); // we just initialized (or reset) the option state for this view -- don't bother comparing in this case
 
 	if (!bFirstTime && Config.UseMISCompensation != PathTracingState->LastConfig.UseMISCompensation)
 	{
diff --git a/Engine/Source/Runtime/Renderer/Private/PathTracingSpatialTemporalDenoising.cpp b/Engine/Source/Runtime/Renderer/Private/PathTracingSpatialTemporalDenoising.cpp
--- a/Engine/Source/Runtime/Renderer/Private/PathTracingSpatialTemporalDenoising.cpp
+++ b/Engine/Source/Runtime/Renderer/Private/PathTracingSpatialTemporalDenoising.cpp
@@ -646,12 +646,19 @@ static bool ShouldEnableSubpixelOffset(uint32 MipLevel)
 static bool IsVarianceTextureRequiredForCurrentDenoiser()
 {
 	bool NeedVarianceTexture = false;
-
-	bool bMatch = true;
 	using UE::Renderer::Private::IPathTracingSpatialTemporalDenoiser;
-	FString DenoiserName = CVarPathTracingDenoiserName.GetValueOnRenderThread();
-	IPathTracingSpatialTemporalDenoiser* ActiveSpatialTemporalDenoiser = FDenoiserManager::Get().GetSpatialTemporalDenoiser(DenoiserName, bMatch);
+	IPathTracingSpatialTemporalDenoiser* ActiveSpatialTemporalDenoiser = nullptr;
 
+	if (GetSpatialDenoiserType() == ESpatialDenoiserType::SPATIAL_TEMPORAL_DENOISER_PLUGIN)
+	{
+		ActiveSpatialTemporalDenoiser = GetActiveSpatialTemporalDenoiser();
+	}
+	else
+	{
+		FString DenoiserName = CVarPathTracingDenoiserName.GetValueOnRenderThread();
+		ActiveSpatialTemporalDenoiser = FDenoiserManager::Get().GetSpatialTemporalDenoiser(DenoiserName, /*bMatch*/true);
+	}
+	
 	if (ActiveSpatialTemporalDenoiser)
 	{
 		NeedVarianceTexture = ActiveSpatialTemporalDenoiser->NeedVarianceTexture();
