Disable workaround fix for occasional RVT page corruption by default.
#rnx

diff --git a/Engine/Source/Runtime/Renderer/Private/VT/RuntimeVirtualTextureRender.cpp b/Engine/Source/Runtime/Renderer/Private/VT/RuntimeVirtualTextureRender.cpp
--- a/Engine/Source/Runtime/Renderer/Private/VT/RuntimeVirtualTextureRender.cpp
+++ b/Engine/Source/Runtime/Renderer/Private/VT/RuntimeVirtualTextureRender.cpp
@@ -55,7 +55,7 @@ namespace RuntimeVirtualTexture
 
 	static TAutoConsoleVariable<bool> CVarVTApplyPageCorruptionFix(
 		TEXT("r.VT.RVT.PageCorruptionFix"),
-		true,
+		false,
 		TEXT("Apply change that has been found to fix some rare page corruption on PC."),
 		ECVF_RenderThreadSafe);
 
