Modify workaround fix for occasional RVT page corruption that has been seen in the wild.
Use the render pass prologue to apply the workaround instead of adding an additional pass.
#rnx

diff --git a/Engine/Source/Runtime/Renderer/Private/VT/RuntimeVirtualTextureRender.cpp b/Engine/Source/Runtime/Renderer/Private/VT/RuntimeVirtualTextureRender.cpp
--- a/Engine/Source/Runtime/Renderer/Private/VT/RuntimeVirtualTextureRender.cpp
+++ b/Engine/Source/Runtime/Renderer/Private/VT/RuntimeVirtualTextureRender.cpp
@@ -1976,39 +1976,34 @@ namespace RuntimeVirtualTexture
 		}
 	}
 
-	/** 
-	 * Fixes an issue where runtime virtual texture page corruption causes square artifacts.
-	 * Repro of the bug is rare. But it's been found that inserting a single call to ID3D12GraphicsCommandList::SetPipelineState()
-	 * before rendering any RVT pages resolves the issue.
-	 */
-	void ApplyPageCorruptionFix(FRDGBuilder& GraphBuilder, int32 InPageIndex, EShaderPlatform InShaderPlatform)
+	/** Mesh render pass prologue to set the viewport. Also applies a page corruption workaround when that is enabled. */
+	void MeshPassPrologue(FRHICommandList& RHICmdList, FIntRect const& InViewRect, int32 InPageIndex, EShaderPlatform InShaderPlatform)
 	{
 		if (InPageIndex == 0 && CVarVTApplyPageCorruptionFix.GetValueOnRenderThread() && IsPCPlatform(InShaderPlatform) && IsD3DPlatform(InShaderPlatform))
 		{
-			GraphBuilder.AddPass(
-				RDG_EVENT_NAME("VirtualTextureDummyPass"),
-				ERDGPassFlags::NeverCull | ERDGPassFlags::NeverMerge,
-				[](FRHICommandList& RHICmdList)
-			{
-				FGlobalShaderMap* GlobalShaderMap = GetGlobalShaderMap(GMaxRHIFeatureLevel);
-				TShaderMapRef<FShader_VirtualTextureCopy_VS> VertexShader(GlobalShaderMap);
-				TShaderMapRef<FShader_VirtualTextureCopy_PS<ERuntimeVirtualTextureMaterialType::BaseColor>> PixelShader(GlobalShaderMap);
-
-				FGraphicsPipelineStateInitializer GraphicsPSOInit;
-				GraphicsPSOInit.RenderTargetsEnabled = 1;
-				GraphicsPSOInit.RenderTargetFormats[0] = PF_B8G8R8A8;
-				GraphicsPSOInit.RenderTargetFlags[0] = TexCreate_None;
-				GraphicsPSOInit.NumSamples = 1;
-				GraphicsPSOInit.DepthStencilState = TStaticDepthStencilState<false, CF_Always>::GetRHI();
-				GraphicsPSOInit.BlendState = TStaticBlendState<CW_RGBA, BO_Add, BF_One, BF_Zero, BO_Add, BF_One, BF_Zero>::GetRHI();
-				GraphicsPSOInit.RasterizerState = TStaticRasterizerState<>::GetRHI();
-				GraphicsPSOInit.PrimitiveType = PT_TriangleList;
-				GraphicsPSOInit.BoundShaderState.VertexDeclarationRHI = GetVertexDeclarationFVector4();
-				GraphicsPSOInit.BoundShaderState.VertexShaderRHI = VertexShader.GetVertexShader();
-				GraphicsPSOInit.BoundShaderState.PixelShaderRHI = PixelShader.GetPixelShader();
-				SetGraphicsPipelineState(RHICmdList, GraphicsPSOInit, 0, EApplyRendertargetOption::DoNothing);
-			});
+			// Workaround fix for an issue where runtime virtual texture page corruption causes square artifacts.
+			// Repro of the bug is rare. But it's been found that inserting a single call to ID3D12GraphicsCommandList::SetPipelineState()
+			// before rendering any RVT pages resolves the issue.
+			FGlobalShaderMap* GlobalShaderMap = GetGlobalShaderMap(GMaxRHIFeatureLevel);
+			TShaderMapRef<FShader_VirtualTextureCopy_VS> VertexShader(GlobalShaderMap);
+			TShaderMapRef<FShader_VirtualTextureCopy_PS<ERuntimeVirtualTextureMaterialType::BaseColor>> PixelShader(GlobalShaderMap);
+
+			FGraphicsPipelineStateInitializer GraphicsPSOInit;
+			GraphicsPSOInit.RenderTargetsEnabled = 1;
+			GraphicsPSOInit.RenderTargetFormats[0] = PF_B8G8R8A8;
+			GraphicsPSOInit.RenderTargetFlags[0] = TexCreate_None;
+			GraphicsPSOInit.NumSamples = 1;
+			GraphicsPSOInit.DepthStencilState = TStaticDepthStencilState<false, CF_Always>::GetRHI();
+			GraphicsPSOInit.BlendState = TStaticBlendState<CW_RGBA, BO_Add, BF_One, BF_Zero, BO_Add, BF_One, BF_Zero>::GetRHI();
+			GraphicsPSOInit.RasterizerState = TStaticRasterizerState<>::GetRHI();
+			GraphicsPSOInit.PrimitiveType = PT_TriangleList;
+			GraphicsPSOInit.BoundShaderState.VertexDeclarationRHI = GetVertexDeclarationFVector4();
+			GraphicsPSOInit.BoundShaderState.VertexShaderRHI = VertexShader.GetVertexShader();
+			GraphicsPSOInit.BoundShaderState.PixelShaderRHI = PixelShader.GetPixelShader();
+			SetGraphicsPipelineState(RHICmdList, GraphicsPSOInit, 0, EApplyRendertargetOption::DoNothing);
 		}
+
+		RHICmdList.SetViewport(static_cast<float>(InViewRect.Min.X), static_cast<float>(InViewRect.Min.Y), 0.0f, static_cast<float>(InViewRect.Max.X), static_cast<float>(InViewRect.Max.Y), 1.0f);
 	}
 
 	/** Get the debug color to use for a given mip level. Returns InDefaultColor if mip debugging is disabled. */
@@ -2106,8 +2101,6 @@ namespace RuntimeVirtualTexture
 			RenderCaptureInterface::FScopedCapture RenderCapture((RenderCaptureNextRVTPagesDraws != 0), GraphBuilder, TEXT("RenderRVTPage"));
 			RenderCaptureNextRVTPagesDraws = FMath::Max(RenderCaptureNextRVTPagesDraws - 1, 0);
 
-			ApplyPageCorruptionFix(GraphBuilder, PageIndex, Scene->GetShaderPlatform());
-
 			ERenderTargetLoadAction LoadAction = BatchDesc.bClearTextures ? ERenderTargetLoadAction::EClear : ERenderTargetLoadAction::ENoAction;
 			FShader_VirtualTextureMaterialDraw::FParameters* PassParameters = GraphBuilder.AllocParameters<FShader_VirtualTextureMaterialDraw::FParameters>();
 			PassParameters->View = View->ViewUniformBuffer;
@@ -2128,11 +2121,15 @@ namespace RuntimeVirtualTexture
 			const ERuntimeVirtualTextureMaterialType MaterialType = BatchDesc.MaterialType;
 			const bool bAllowCachedMeshDrawCommands = BatchRenderContext.bAllowCachedMeshDrawCommands;
 
-			AddSimpleMeshPass(GraphBuilder, PassParameters, Scene, *View, nullptr, RDG_EVENT_NAME("VirtualTextureDraw"), View->ViewRect, ERDGPassFlags::Raster | ERDGPassFlags::NeverMerge,
+			AddSimpleMeshPass(GraphBuilder, PassParameters, Scene, *View, nullptr, RDG_EVENT_NAME("VirtualTextureDraw"), ERDGPassFlags::Raster | ERDGPassFlags::NeverMerge,
 	    	[Scene, View, MaterialType, RuntimeVirtualTextureId, vLevel, MaxLevel, bAllowCachedMeshDrawCommands](FDynamicPassMeshDrawListContext* DynamicMeshPassContext)
 	    	{
 	    		GatherMeshesToDraw(DynamicMeshPassContext, Scene, View, MaterialType, RuntimeVirtualTextureId, vLevel, MaxLevel, bAllowCachedMeshDrawCommands);
-	    	});
+	    	}, 
+			[ViewRect = View->ViewRect, PageIndex, ShaderPlatform = Scene->GetShaderPlatform()](FRHICommandList& RHICmdList)
+			{
+				MeshPassPrologue(RHICmdList, ViewRect, PageIndex, ShaderPlatform);
+			});
 		}
 	}
 
