Added CheckAndSanitizePrimitiveBounds to UpdatePrimitiveInstances as otherwise invalid primitive bounds may enter the renderer and cause crashes.

#rnx

diff --git a/Engine/Source/Runtime/Renderer/Private/RendererScene.cpp b/Engine/Source/Runtime/Renderer/Private/RendererScene.cpp
--- a/Engine/Source/Runtime/Renderer/Private/RendererScene.cpp
+++ b/Engine/Source/Runtime/Renderer/Private/RendererScene.cpp
@@ -1701,7 +1701,7 @@ void FScene::UpdatePrimitiveInstances(UPrimitiveComponent* Primitive)
 	UpdateParams.PrimitiveSceneProxy = Primitive->GetSceneProxy();
 	UpdateParams.WorldBounds = Primitive->Bounds;
 	UpdateParams.LocalBounds = static_cast<UPrimitiveComponent*>(Primitive)->GetLocalBounds();
-
+	CheckAndSanitizePrimitiveBounds(UpdateParams.WorldBounds, Primitive);
 	UpdatePrimitiveInstances(UpdateParams);
 }
 
@@ -1792,7 +1792,7 @@ void FScene::UpdatePrimitiveInstances(FPrimitiveSceneDesc* Primitive)
 	UpdateParams.PrimitiveSceneProxy = Primitive->GetSceneProxy();
 	UpdateParams.WorldBounds = Primitive->GetBounds();
 	UpdateParams.LocalBounds = Primitive->GetLocalBounds();
-
+	CheckAndSanitizePrimitiveBounds(UpdateParams.WorldBounds, Primitive);
 	UpdatePrimitiveInstances(UpdateParams);
 }
 
