Heterogeneous Volumes: Re-enstate overlapping volume support for the AVSM pipeline.

Overlapping volumes are not supported for the Beer Shadow Map pipeline due to its iterative construction.

Prevent binding crash by cascading through different initialization behaviors for the AVSM CameraMap Uniform Buffer:
1) Live Shading pipeline: A Uniform buffer has not yet been created, but parameters have been initialized. Needed for overlapping volume support.
2) Base Pass: A Uniform buffer has been created.

diff --git a/Engine/Source/Runtime/Renderer/Private/HeterogeneousVolumes/HeterogeneousVolumes.cpp b/Engine/Source/Runtime/Renderer/Private/HeterogeneousVolumes/HeterogeneousVolumes.cpp
--- a/Engine/Source/Runtime/Renderer/Private/HeterogeneousVolumes/HeterogeneousVolumes.cpp
+++ b/Engine/Source/Runtime/Renderer/Private/HeterogeneousVolumes/HeterogeneousVolumes.cpp
@@ -1591,6 +1591,7 @@ namespace HeterogeneousVolumes
 			FSceneViewState* ViewState = Views[ViewIndex].ViewState;
 			if (ViewState)
 			{
+				ViewState->AdaptiveVolumetricCameraMapUniformBufferParameters = nullptr;
 				DestroyAdaptiveVolumetricShadowMapUniformBuffer(ViewState->AdaptiveVolumetricCameraMapUniformBuffer);
 
 				for (auto Itr = ViewState->AdaptiveVolumetricShadowMapUniformBufferMap.begin(); Itr != ViewState->AdaptiveVolumetricShadowMapUniformBufferMap.end(); ++Itr)
diff --git a/Engine/Source/Runtime/Renderer/Private/HeterogeneousVolumes/HeterogeneousVolumesLiveShadingPipeline.cpp b/Engine/Source/Runtime/Renderer/Private/HeterogeneousVolumes/HeterogeneousVolumesLiveShadingPipeline.cpp
--- a/Engine/Source/Runtime/Renderer/Private/HeterogeneousVolumes/HeterogeneousVolumesLiveShadingPipeline.cpp
+++ b/Engine/Source/Runtime/Renderer/Private/HeterogeneousVolumes/HeterogeneousVolumesLiveShadingPipeline.cpp
@@ -142,13 +142,8 @@ namespace HeterogeneousVolumes
 
 	bool SupportsOverlappingVolumes()
 	{
-#if 1
-		// UE-268733: SupportOverlappingVolumes current causes crash
-		return false;
-#else
-		return ShouldCompositeHeterogeneousVolumesWithTranslucency() &&
+		return (GetTranslucencyCompositingMode() == ETranslucencyCompositingMode::VolumetricShadowMap) &&
 			(CVarHeterogeneousVolumesSupportOverlappingVolumes.GetValueOnRenderThread() != 0);
-#endif
 	}
 
 	bool UseExistenceMask()
@@ -965,6 +960,8 @@ FAdaptiveVolumetricShadowMapParameters GetAdaptiveVolumetricShadowMapParametersF
 		Parameters.LinkedListBuffer = UniformBuffer->GetParameters()->LinkedListBuffer;
 		Parameters.IndirectionBuffer = UniformBuffer->GetParameters()->IndirectionBuffer;
 		Parameters.SampleBuffer = UniformBuffer->GetParameters()->SampleBuffer;
+		Parameters.RadianceTexture = UniformBuffer->GetParameters()->RadianceTexture;
+		Parameters.TextureSampler = UniformBuffer->GetParameters()->TextureSampler;
 	}
 
 	return Parameters;
@@ -2779,4 +2776,4 @@ void RenderAdaptiveVolumetricCameraMapWithLiveShading(
 			View.ViewState->AdaptiveVolumetricCameraMapUniformBufferParameters
 		);
 	}
-}
\ No newline at end of file
+}
diff --git a/Engine/Source/Runtime/Renderer/Private/HeterogeneousVolumes/HeterogeneousVolumesVoxelGridPipeline.cpp b/Engine/Source/Runtime/Renderer/Private/HeterogeneousVolumes/HeterogeneousVolumesVoxelGridPipeline.cpp
--- a/Engine/Source/Runtime/Renderer/Private/HeterogeneousVolumes/HeterogeneousVolumesVoxelGridPipeline.cpp
+++ b/Engine/Source/Runtime/Renderer/Private/HeterogeneousVolumes/HeterogeneousVolumesVoxelGridPipeline.cpp
@@ -3404,8 +3404,7 @@ namespace HeterogeneousVolumes
 			UniformBufferParameters->TextureSampler = TStaticSamplerState<SF_Bilinear>::GetRHI();
 		}
 
-		TRDGUniformBufferRef<FAdaptiveVolumetricShadowMapUniformBufferParameters> AdaptiveVolumetricShadowMapUniformBuffer = GraphBuilder.CreateUniformBuffer(UniformBufferParameters);
-		return AdaptiveVolumetricShadowMapUniformBuffer;
+		return GraphBuilder.CreateUniformBuffer(UniformBufferParameters);
 	}
 
 	void DestroyAdaptiveVolumetricShadowMapUniformBuffer(
@@ -3440,7 +3439,14 @@ namespace HeterogeneousVolumes
 	{
 		if (ViewState && ShouldCompositeHeterogeneousVolumesWithTranslucency())
 		{
-			return ViewState->AdaptiveVolumetricCameraMapUniformBuffer;
+			if (ViewState->AdaptiveVolumetricCameraMapUniformBuffer != nullptr)
+			{
+				return ViewState->AdaptiveVolumetricCameraMapUniformBuffer;
+			}
+			else if (ViewState->AdaptiveVolumetricCameraMapUniformBufferParameters != nullptr)
+			{
+				return GraphBuilder.CreateUniformBuffer(ViewState->AdaptiveVolumetricCameraMapUniformBufferParameters);
+			}
 		}
 		
 		return CreateEmptyAdaptiveVolumetricShadowMapUniformBuffer(GraphBuilder);
