[Renderer]
* Removed extern GetTypeHash function in VirtualTextureSystem since it is a hidden friend where implemented and this confuses header unit compilations. Non-unity still seems to compile with this removed

diff --git a/Engine/Source/Runtime/Renderer/Private/VT/VirtualTextureSystem.h b/Engine/Source/Runtime/Renderer/Private/VT/VirtualTextureSystem.h
--- a/Engine/Source/Runtime/Renderer/Private/VT/VirtualTextureSystem.h
+++ b/Engine/Source/Runtime/Renderer/Private/VT/VirtualTextureSystem.h
@@ -34,8 +34,6 @@ struct FGatherRequestsParameters;
 struct FPageUpdateBuffer;
 class ISceneRenderer;
 
-extern uint32 GetTypeHash(const FAllocatedVTDescription& Description);
-
 struct FVirtualTextureUpdateSettings
 {
 	FVirtualTextureUpdateSettings();
