Fixed incorrect RHI for RDG readback pass used for debug, which tripped checks if turning on r.SceneCulling.ValidateGPUData

#rnx

diff --git a/Engine/Source/Runtime/Renderer/Private/RendererPrivateUtils.h b/Engine/Source/Runtime/Renderer/Private/RendererPrivateUtils.h
--- a/Engine/Source/Runtime/Renderer/Private/RendererPrivateUtils.h
+++ b/Engine/Source/Runtime/Renderer/Private/RendererPrivateUtils.h
@@ -92,7 +92,7 @@ void AddBufferLockReadbackPass(FRDGBuilder& GraphBuilder, TRefCountPtr<FRDGPoole
 	GraphBuilder.AddPass(
 		RDG_EVENT_NAME("BufferLockReadbackPass"),
 		ERDGPassFlags::None,
-		[ReadbackProcessingLambdaType=MoveTemp(ReadbackProcessingLambda), SourceBufferRHI, NumBytes](FRHICommandList& RHICmdList)
+		[ReadbackProcessingLambdaType=MoveTemp(ReadbackProcessingLambda), SourceBufferRHI, NumBytes](FRHICommandListImmediate& RHICmdList)
 	{
 		const void *GPUData = (uint32*)RHICmdList.LockBuffer(SourceBufferRHI, 0, NumBytes, RLM_ReadOnly);
 		ReadbackProcessingLambdaType(GPUData);
