Fixed: Issues with a separate translucency rendeing on mobile.
Separate translucency materials will now always sample Depth as a texture instead of trying to use FBF. When separate translucecy feature is switched off, materials will render in base-pass as a regular translucency but depth sampling will not work for them.

diff --git a/Engine/Source/Runtime/Renderer/Private/MobileBasePass.cpp b/Engine/Source/Runtime/Renderer/Private/MobileBasePass.cpp
--- a/Engine/Source/Runtime/Renderer/Private/MobileBasePass.cpp
+++ b/Engine/Source/Runtime/Renderer/Private/MobileBasePass.cpp
@@ -928,9 +928,22 @@ void FMobileBasePassMeshProcessor::CollectPSOInitializersForLMPolicy(
 	}
 
 	// subpass info set during the submission of the draws in mobile deferred renderer.
-	uint8 SubpassIndex = bTranslucentBasePass ? (bDeferredShading ? 2 : 1) : 0;
+	uint8 SubpassIndex = 0;
 	ESubpassHint SubpassHint = GetSubpassHint(GMaxRHIShaderPlatform, bDeferredShading, RenderTargetsInfo.MultiViewCount > 1, RenderTargetsInfo.NumSamples);
-
+	if (bTranslucentBasePass)
+	{
+		if (MeshPassType == EMeshPass::TranslucencyAfterDOF)
+		{
+			// Separate translucency in subpass 0
+			SubpassIndex = 0;
+			SubpassHint = ESubpassHint::None;
+		}
+		else
+		{
+			SubpassIndex = bDeferredShading ? 2 : 1;
+		}
+	}
+	
 	AddGraphicsPipelineStateInitializer(
 		VertexFactoryData,
 		MaterialResource,
diff --git a/Engine/Source/Runtime/Renderer/Private/MobileBasePassRendering.cpp b/Engine/Source/Runtime/Renderer/Private/MobileBasePassRendering.cpp
--- a/Engine/Source/Runtime/Renderer/Private/MobileBasePassRendering.cpp
+++ b/Engine/Source/Runtime/Renderer/Private/MobileBasePassRendering.cpp
@@ -224,7 +224,8 @@ void MobileBasePassModifyCompilationEnvironment(const FMaterialShaderPermutation
 	}
 
 	const bool bMobileForceDepthRead = MobileUsesFullDepthPrepass(Parameters.Platform);
-	OutEnvironment.SetDefine(TEXT("IS_MOBILE_DEPTHREAD_SUBPASS"), bTranslucentMaterial && !bMobileForceDepthRead ? 1u : 0u);
+	// separate translucency must sample depth texture, instead of using FBF
+	OutEnvironment.SetDefine(TEXT("IS_MOBILE_DEPTHREAD_SUBPASS"), bTranslucentMaterial && !bMobileForceDepthRead && !Parameters.MaterialParameters.bIsMobileSeparateTranslucencyEnabled ? 1u : 0u);
 	// translucency is in the same subpass with deferred shading shaders, so it has access to GBuffer
 	const bool bDeferredShadingSubpass = (bDeferredShadingEnabled && bTranslucentMaterial && !Parameters.MaterialParameters.bIsMobileSeparateTranslucencyEnabled);
 	OutEnvironment.SetDefine(TEXT("IS_MOBILE_DEFERREDSHADING_SUBPASS"), bDeferredShadingSubpass ? 1u : 0u);
diff --git a/Engine/Source/Runtime/Renderer/Private/MobileSeparateTranslucencyPass.cpp b/Engine/Source/Runtime/Renderer/Private/MobileSeparateTranslucencyPass.cpp
--- a/Engine/Source/Runtime/Renderer/Private/MobileSeparateTranslucencyPass.cpp
+++ b/Engine/Source/Runtime/Renderer/Private/MobileSeparateTranslucencyPass.cpp
@@ -36,48 +36,28 @@ bool IsMobileSeparateTranslucencyActive(const FViewInfo& View)
 
 void AddMobileSeparateTranslucencyPass(FRDGBuilder& GraphBuilder, FScene* Scene, const FViewInfo& View, const FMobileSeparateTranslucencyInputs& Inputs)
 {
-	// Whether BasePass uses single pass rendering with framebuffer fetch
-	const bool bSinglePassRendering = !Inputs.bRequiresMultiPass;
-
 	FMobileSeparateTranslucencyPassParameters* PassParameters = GraphBuilder.AllocParameters<FMobileSeparateTranslucencyPassParameters>();
 	PassParameters->RenderTargets[0] = FRenderTargetBinding(Inputs.SceneColor.Texture, ERenderTargetLoadAction::ELoad);
-	if (bSinglePassRendering)
-	{
-		if (MobileRequiresSceneDepthAux(View.GetShaderPlatform()))
-		{
-			PassParameters->RenderTargets[1] = FRenderTargetBinding(Inputs.SceneDepthAux.Texture, ERenderTargetLoadAction::ELoad);
-		}
-		PassParameters->RenderTargets.SubpassHint = ESubpassHint::DepthReadSubpass;
-	}
 	PassParameters->RenderTargets.DepthStencil = FDepthStencilBinding(Inputs.SceneDepth.Texture, ERenderTargetLoadAction::ELoad, ERenderTargetLoadAction::ELoad, FExclusiveDepthStencil::DepthRead_StencilRead);
 	PassParameters->View = View.ViewUniformBuffer;
 	EMobileSceneTextureSetupMode SetupMode = EMobileSceneTextureSetupMode::SceneDepth | EMobileSceneTextureSetupMode::SceneDepthAux | EMobileSceneTextureSetupMode::CustomDepth;
 	PassParameters->MobileBasePass = CreateMobileBasePassUniformBuffer(GraphBuilder, View, EMobileBasePass::Translucent, SetupMode);
 
 	auto* Pass = const_cast<FViewInfo&>(View).ParallelMeshDrawCommandPasses[EMeshPass::TranslucencyAfterDOF];
-
 	if (Pass)
 	{
 		Pass->BuildRenderingCommands(GraphBuilder, Scene->GPUScene, PassParameters->InstanceCullingDrawParams);
-	}
 
-	// Pass always added due to NextSubpass.
-	GraphBuilder.AddPass(
-		RDG_EVENT_NAME("SeparateTranslucency %dx%d", View.ViewRect.Width(), View.ViewRect.Height()),
-		PassParameters,
-		ERDGPassFlags::Raster,
-		[&View, Pass, PassParameters, bSinglePassRendering](FRHICommandList& RHICmdList)
-	{
-		if (bSinglePassRendering)
-		{
-			RHICmdList.NextSubpass();
-		}
+		GraphBuilder.AddPass(
+			RDG_EVENT_NAME("SeparateTranslucency %dx%d", View.ViewRect.Width(), View.ViewRect.Height()),
+			PassParameters,
+			ERDGPassFlags::Raster,
+			[&View, Pass, PassParameters](FRHICommandList& RHICmdList)
+			{
+				// Set the view family's render target/viewport.
+				RHICmdList.SetViewport(View.ViewRect.Min.X, View.ViewRect.Min.Y, 0.0f, View.ViewRect.Max.X, View.ViewRect.Max.Y, 1.0f);
+				Pass->Draw(RHICmdList, &PassParameters->InstanceCullingDrawParams);
+			});
+	}
 
-		if (Pass)
-		{
-			// Set the view family's render target/viewport.
-			RHICmdList.SetViewport(View.ViewRect.Min.X, View.ViewRect.Min.Y, 0.0f, View.ViewRect.Max.X, View.ViewRect.Max.Y, 1.0f);
-			Pass->Draw(RHICmdList, &PassParameters->InstanceCullingDrawParams);
-		}
-	});
 }
\ No newline at end of file
diff --git a/Engine/Source/Runtime/Renderer/Private/MobileSeparateTranslucencyPass.h b/Engine/Source/Runtime/Renderer/Private/MobileSeparateTranslucencyPass.h
--- a/Engine/Source/Runtime/Renderer/Private/MobileSeparateTranslucencyPass.h
+++ b/Engine/Source/Runtime/Renderer/Private/MobileSeparateTranslucencyPass.h
@@ -11,7 +11,6 @@ struct FMobileSeparateTranslucencyInputs
 	FScreenPassTexture SceneColor;
 	FScreenPassTexture SceneDepthAux;
 	FScreenPassTexture SceneDepth;
-	bool bRequiresMultiPass;
 };
 
 void AddMobileSeparateTranslucencyPass(FRDGBuilder& GraphBuilder, FScene* Scene, const FViewInfo& View, const FMobileSeparateTranslucencyInputs& Inputs);
diff --git a/Engine/Source/Runtime/Renderer/Private/MobileShadingRenderer.cpp b/Engine/Source/Runtime/Renderer/Private/MobileShadingRenderer.cpp
--- a/Engine/Source/Runtime/Renderer/Private/MobileShadingRenderer.cpp
+++ b/Engine/Source/Runtime/Renderer/Private/MobileShadingRenderer.cpp
@@ -1526,7 +1526,6 @@ void FMobileSceneRenderer::Render(FRDGBuilder& GraphBuilder, const FSceneRenderU
 					FMobilePostProcessingInputs PostProcessingInputs;
 					PostProcessingInputs.ViewFamilyTexture = ViewFamilyTexture;
 					PostProcessingInputs.SceneTextures = CreateMobileSceneTextureUniformBuffer(GraphBuilder, &SceneTextures, EMobileSceneTextureSetupMode::All);
-					PostProcessingInputs.bRequiresMultiPass = bRequiresMultiPass;
 
 					for (int32 ViewIndex = 0; ViewIndex < Views.Num(); ViewIndex++)
 					{
diff --git a/Engine/Source/Runtime/Renderer/Private/PostProcess/PostProcessing.cpp b/Engine/Source/Runtime/Renderer/Private/PostProcess/PostProcessing.cpp
--- a/Engine/Source/Runtime/Renderer/Private/PostProcess/PostProcessing.cpp
+++ b/Engine/Source/Runtime/Renderer/Private/PostProcess/PostProcessing.cpp
@@ -3033,7 +3033,6 @@ void AddMobilePostProcessingPasses(FRDGBuilder& GraphBuilder, FScene* Scene, con
 			SeparateTranslucencyInputs.SceneColor = SceneColor;
 			SeparateTranslucencyInputs.SceneDepthAux = SceneDepthAux;
 			SeparateTranslucencyInputs.SceneDepth = SceneDepth;
-			SeparateTranslucencyInputs.bRequiresMultiPass = Inputs.bRequiresMultiPass;
 			
 			AddMobileSeparateTranslucencyPass(GraphBuilder, Scene, View, SeparateTranslucencyInputs);
 		}
diff --git a/Engine/Source/Runtime/Renderer/Private/PostProcess/PostProcessing.h b/Engine/Source/Runtime/Renderer/Private/PostProcess/PostProcessing.h
--- a/Engine/Source/Runtime/Renderer/Private/PostProcess/PostProcessing.h
+++ b/Engine/Source/Runtime/Renderer/Private/PostProcess/PostProcessing.h
@@ -56,7 +56,6 @@ struct FMobilePostProcessingInputs
 	TRDGUniformBufferRef<FMobileSceneTextureUniformParameters> SceneTextures = nullptr;
 	FRDGTextureRef ViewFamilyTexture = nullptr;
 	FRDGTextureRef ViewFamilyDepthTexture = nullptr;
-	bool bRequiresMultiPass = false;
 
 	void Validate() const
 	{
