Fix up virtual texture GPU stats for the new GPU profiler.
#rnx

diff --git a/Engine/Source/Runtime/Renderer/Private/VT/VirtualTextureSystem.cpp b/Engine/Source/Runtime/Renderer/Private/VT/VirtualTextureSystem.cpp
--- a/Engine/Source/Runtime/Renderer/Private/VT/VirtualTextureSystem.cpp
+++ b/Engine/Source/Runtime/Renderer/Private/VT/VirtualTextureSystem.cpp
@@ -2523,7 +2523,6 @@ void FVirtualTextureSystem::FinalizeRequests(FRDGBuilder& GraphBuilder, ISceneRe
 {
 	RDG_CSV_STAT_EXCLUSIVE_SCOPE(GraphBuilder, VirtualTextureSystem_Update);
 	TRACE_CPUPROFILER_EVENT_SCOPE(FVirtualTextureSystem::FinalizeRequests);
-	RDG_GPU_STAT_SCOPE(GraphBuilder, VirtualTexture);
 
 	FRDGExternalAccessQueue ExternalAccessQueue;
 
@@ -2531,7 +2530,8 @@ void FVirtualTextureSystem::FinalizeRequests(FRDGBuilder& GraphBuilder, ISceneRe
 	if (Finalizers.Num())
 	{
 		SCOPE_CYCLE_COUNTER(STAT_ProcessRequests_Finalize);
-		RDG_EVENT_SCOPE(GraphBuilder, "VirtualTextureFinalizeRequests");
+		RDG_EVENT_SCOPE_STAT(GraphBuilder, VirtualTexture, "VirtualTextureFinalizeRequests");
+		RDG_GPU_STAT_SCOPE(GraphBuilder, VirtualTexture);
 		TRACE_CPUPROFILER_EVENT_SCOPE(FVirtualTextureSystem::Finalize);
 
 		for (IVirtualTextureFinalizer* VTFinalizer : Finalizers)
@@ -2559,7 +2559,8 @@ void FVirtualTextureSystem::FinalizeRequests(FRDGBuilder& GraphBuilder, ISceneRe
 	// Update page tables
 	{
 		SCOPE_CYCLE_COUNTER(STAT_PageTableUpdates);
-		RDG_EVENT_SCOPE(GraphBuilder, "VirtualTexturePageTableUpdates");
+		RDG_EVENT_SCOPE_STAT(GraphBuilder, VirtualTexture, "VirtualTexturePageTableUpdates");
+		RDG_GPU_STAT_SCOPE(GraphBuilder, VirtualTexture);
 		TRACE_CPUPROFILER_EVENT_SCOPE(FVirtualTextureSystem::ApplyUpdates);
 
 		for (uint32 ID = 0; ID < MaxSpaces; ID++)
