RayTracing - remove separate far field build mode CVar (r.RayTracing.Scene.FarFieldBuildMode), as reusing r.RayTracing.Scene.BuildMode seems to be a better default. It will also allow to batch TLAS builds in the future.

diff --git a/Engine/Source/Runtime/Renderer/Private/RayTracing/RayTracingScene.cpp b/Engine/Source/Runtime/Renderer/Private/RayTracing/RayTracingScene.cpp
--- a/Engine/Source/Runtime/Renderer/Private/RayTracing/RayTracingScene.cpp
+++ b/Engine/Source/Runtime/Renderer/Private/RayTracing/RayTracingScene.cpp
@@ -25,15 +25,6 @@ static TAutoConsoleVariable<int32> CVarRayTracingSceneBuildMode(
 	ECVF_RenderThreadSafe | ECVF_Scalability
 );
 
-static TAutoConsoleVariable<int32> CVarRayTracingSceneFarFieldBuildMode(
-	TEXT("r.RayTracing.Scene.FarFieldBuildMode"),
-	1,
-	TEXT("Controls the mode in which Far Field ray tracing scene is built:\n")
-	TEXT(" 0: Fast build\n")
-	TEXT(" 1: Fast trace (default)\n"),
-	ECVF_RenderThreadSafe | ECVF_Scalability
-);
-
 static TAutoConsoleVariable<bool> CVarRayTracingSceneUseTracingFeedback(
 	TEXT("r.RayTracing.Scene.UseTracingFeedback"),
 	false,
@@ -144,10 +135,6 @@ void FRayTracingScene::Create(FRDGBuilder& GraphBuilder, const FViewInfo& View,
 		? ERayTracingAccelerationStructureFlags::FastTrace
 		: ERayTracingAccelerationStructureFlags::FastBuild;
 
-	const ERayTracingAccelerationStructureFlags FarFieldBuildFlags = CVarRayTracingSceneFarFieldBuildMode.GetValueOnRenderThread()
-		? ERayTracingAccelerationStructureFlags::FastTrace
-		: ERayTracingAccelerationStructureFlags::FastBuild;
-
 	if (!bInitializationDataBuilt)
 	{
 		BuildInitializationData();
@@ -193,7 +180,7 @@ void FRayTracingScene::Create(FRDGBuilder& GraphBuilder, const FViewInfo& View,
 			FRayTracingSceneInitializer Initializer;
 			Initializer.DebugName = Layer.Name;
 			Initializer.MaxNumInstances = NumNativeInstances;
-			Initializer.BuildFlags = LayerIndex == uint8(ERayTracingSceneLayer::FarField) ? FarFieldBuildFlags : BuildFlags;
+			Initializer.BuildFlags = BuildFlags;
 			Layer.RayTracingSceneRHI = RHICreateRayTracingScene(MoveTemp(Initializer));
 		}
 
