Post Processing: Allow lens flare to remain enabled even if bloom is disabled.

* Lens flare intensity is no longer scaled by bloom intensity.
* Backward compatibility is maintained by updating the default lens flare intensity value to match the default bloom intensity (0.675f).
* Additional logic on load to maintain the same *effective* intensity on older serialized settings.

diff --git a/Engine/Source/Runtime/Core/Public/UObject/UE5ReleaseStreamObjectVersions.inl b/Engine/Source/Runtime/Core/Public/UObject/UE5ReleaseStreamObjectVersions.inl
--- a/Engine/Source/Runtime/Core/Public/UObject/UE5ReleaseStreamObjectVersions.inl
+++ b/Engine/Source/Runtime/Core/Public/UObject/UE5ReleaseStreamObjectVersions.inl
@@ -165,5 +165,8 @@ UE5_RELEASE_VERSION(OrthographicAutoNearFarPlane, 51)
 // Fix a bug where BlendMask counts could get out of sync with BlendPose counts.
 UE5_RELEASE_VERSION(AnimLayeredBoneBlendMasksFix, 52)
 
+// Separated lens flare from bloom intensity
+UE5_RELEASE_VERSION(BloomIndependentLensFlare, 53)
+
 // -----<new versions can be added above this line>-------------------------------------------------
 #undef UE5_RELEASE_VERSION
\ No newline at end of file
diff --git a/Engine/Source/Runtime/Engine/Private/Scene.cpp b/Engine/Source/Runtime/Engine/Private/Scene.cpp
--- a/Engine/Source/Runtime/Engine/Private/Scene.cpp
+++ b/Engine/Source/Runtime/Engine/Private/Scene.cpp
@@ -498,7 +498,7 @@ FPostProcessSettings::FPostProcessSettings()
 	LocalExposureMiddleGreyBias = 0.0f;
 
 	// next value might get overwritten by r.DefaultFeature.LensFlare
-	LensFlareIntensity = 1.0f;
+	LensFlareIntensity = 0.675f;
 	LensFlareTint = FLinearColor(1.0f, 1.0f, 1.0f);
 	LensFlareBokehSize = 3.0f;
 	LensFlareThreshold = 8.0f;
@@ -664,6 +664,12 @@ bool FPostProcessSettings::Serialize(FArchive& Ar)
 	Ar.UsingCustomVersion(FReleaseObjectVersion::GUID);
 	Ar.UsingCustomVersion(FUE5ReleaseStreamObjectVersion::GUID);
 
+	if (Ar.CustomVer(FUE5ReleaseStreamObjectVersion::GUID) < FUE5ReleaseStreamObjectVersion::BloomIndependentLensFlare)
+	{
+		// For backwards compatibility, we must preserve the pre-UE 5.6 default value *before* serialization.
+		LensFlareIntensity = 1.0f;
+	}
+
 	// Don't actually serialize, just write the custom version for PostSerialize
 	return false;
 }
@@ -928,6 +934,12 @@ void FPostProcessSettings::PostSerialize(const FArchive& Ar)
 				}
 			}
 		}
+
+		if (Ar.CustomVer(FUE5ReleaseStreamObjectVersion::GUID) < FUE5ReleaseStreamObjectVersion::BloomIndependentLensFlare)
+		{
+			// Lens flare intensity is now independent from bloom, but we preserve the bloom intensity scale on settings before UE 5.6.
+			LensFlareIntensity *= FMath::Max(0.0f, BloomIntensity);
+		}
 	}
 }
 #endif
diff --git a/Engine/Source/Runtime/Renderer/Private/PostProcess/PostProcessLensFlares.cpp b/Engine/Source/Runtime/Renderer/Private/PostProcess/PostProcessLensFlares.cpp
--- a/Engine/Source/Runtime/Renderer/Private/PostProcess/PostProcessLensFlares.cpp
+++ b/Engine/Source/Runtime/Renderer/Private/PostProcess/PostProcessLensFlares.cpp
@@ -406,7 +406,7 @@ FScreenPassTexture AddLensFlaresPass(
 	LensFlareInputs.TintColorsPerFlare = Settings.LensFlareTints;
 	LensFlareInputs.TintColor = Settings.LensFlareTint;
 	LensFlareInputs.BokehSizePercent = Settings.LensFlareBokehSize;
-	LensFlareInputs.Intensity = Settings.LensFlareIntensity * Settings.BloomIntensity;
+	LensFlareInputs.Intensity = Settings.LensFlareIntensity;
 	LensFlareInputs.Threshold = Settings.LensFlareThreshold;
 
 	// If a bloom output texture isn't available, substitute the half resolution scene color instead, but disable bloom
diff --git a/Engine/Source/Runtime/Renderer/Private/PostProcess/PostProcessing.cpp b/Engine/Source/Runtime/Renderer/Private/PostProcess/PostProcessing.cpp
--- a/Engine/Source/Runtime/Renderer/Private/PostProcess/PostProcessing.cpp
+++ b/Engine/Source/Runtime/Renderer/Private/PostProcess/PostProcessing.cpp
@@ -771,7 +771,7 @@ void AddPostProcessingPasses(
 
 		PassSequence.Finalize();
 
-		const bool bLensFlareEnabled = bBloomEnabled && IsLensFlaresEnabled(View);
+		const bool bLensFlareEnabled = IsLensFlaresEnabled(View) && !bVisualizeTSR;
 		const bool bFFTBloomEnabled = bBloomEnabled && IsFFTBloomEnabled(View);
 
 		const bool bBasicEyeAdaptationEnabled = bEyeAdaptationEnabled && (AutoExposureMethod == EAutoExposureMethod::AEM_Basic);
@@ -786,7 +786,7 @@ void AddPostProcessingPasses(
 		const bool bProduceSceneColorChain = (
 			bBasicEyeAdaptationEnabled ||
 			(bBloomEnabled && !bFFTBloomEnabled) ||
-			(bLensFlareEnabled && bFFTBloomEnabled) ||
+			bLensFlareEnabled ||
 			bLocalExposureBlurredLum);
 		extern int32 GSSRHalfResSceneColor;
 		const bool bNeedBeforeBloomHalfRes    = (!bProcessQuarterResolution && !bProcessEighthResolution) || (bFFTBloomEnabled && FFTBloomResolutionFraction > 0.25f && FFTBloomResolutionFraction <= 0.5f) || (ReflectionsMethod == EReflectionsMethod::SSR && !View.bStatePrevViewInfoIsReadOnly && GSSRHalfResSceneColor);
@@ -1281,99 +1281,95 @@ void AddPostProcessingPasses(
 
 		FScreenPassTexture Bloom;
 		FRDGBufferRef SceneColorApplyParameters = nullptr;
-		if (bBloomEnabled)
 		{
-			const FTextureDownsampleChain* LensFlareSceneDownsampleChain;
+			const FTextureDownsampleChain* LensFlareSceneDownsampleChain = &SceneDownsampleChain;
 
 			FTextureDownsampleChain BloomDownsampleChain;
 
-			if (bFFTBloomEnabled)
+			if (bBloomEnabled)
 			{
-				LensFlareSceneDownsampleChain = &SceneDownsampleChain;
+				if (bFFTBloomEnabled)
+				{
+					float InputResolutionFraction;
+					FScreenPassTextureSlice InputSceneColor;
 
-				float InputResolutionFraction;
-				FScreenPassTextureSlice InputSceneColor;
+					if (FFTBloomResolutionFraction <= 0.125f)
+					{
+						InputSceneColor = EighthResSceneColor;
+						InputResolutionFraction = 0.125f;
+					}
+					else if (FFTBloomResolutionFraction <= 0.25f)
+					{
+						InputSceneColor = QuarterResSceneColor;
+						InputResolutionFraction = 0.25f;
+					}
+					else if (FFTBloomResolutionFraction <= 0.5f)
+					{
+						InputSceneColor = HalfResSceneColor;
+						InputResolutionFraction = 0.5f;
+					}
+					else
+					{
+						InputSceneColor = SceneColorSlice;
+						InputResolutionFraction = 1.0f;
+					}
 
-				if (FFTBloomResolutionFraction <= 0.125f)
-				{
-					InputSceneColor = EighthResSceneColor;
-					InputResolutionFraction = 0.125f;
-				}
-				else if (FFTBloomResolutionFraction <= 0.25f)
-				{
-					InputSceneColor = QuarterResSceneColor;
-					InputResolutionFraction = 0.25f;
-				}
-				else if (FFTBloomResolutionFraction <= 0.5f)
-				{
-					InputSceneColor = HalfResSceneColor;
-					InputResolutionFraction = 0.5f;
+					FFFTBloomOutput Outputs = AddFFTBloomPass(
+						GraphBuilder,
+						View,
+						InputSceneColor,
+						InputResolutionFraction,
+						EyeAdaptationParameters,
+						EyeAdaptationBuffer,
+						LocalExposureParameters,
+						CVarBloomApplyLocalExposure.GetValueOnRenderThread() ? LocalExposureBilateralGridTexture : nullptr,
+						LocalExposureBlurredLogLumTexture);
+
+					Bloom = Outputs.BloomTexture;
+					SceneColorApplyParameters = Outputs.SceneColorApplyParameters;
 				}
 				else
 				{
-					InputSceneColor = SceneColorSlice;
-					InputResolutionFraction = 1.0f;
-				}
+					const bool bApplyLocalExposureToBloom =
+						CVarBloomApplyLocalExposure.GetValueOnRenderThread()
+						&& View.FinalPostProcessSettings.LocalExposureMethod == ELocalExposureMethod::Bilateral
+						&& LocalExposureBilateralGridTexture != nullptr;
 
-				FFFTBloomOutput Outputs = AddFFTBloomPass(
-					GraphBuilder, 
-					View,
-					InputSceneColor,
-					InputResolutionFraction,
-					EyeAdaptationParameters,
-					EyeAdaptationBuffer,
-					LocalExposureParameters,
-					CVarBloomApplyLocalExposure.GetValueOnRenderThread() ? LocalExposureBilateralGridTexture : nullptr,
-					LocalExposureBlurredLogLumTexture);
+					const bool bBloomSetupRequiredEnabled = View.FinalPostProcessSettings.BloomThreshold > -1.0f || bApplyLocalExposureToBloom;
 
-				Bloom = Outputs.BloomTexture;
-				SceneColorApplyParameters = Outputs.SceneColorApplyParameters;
-			}
-			else
-			{
-				const bool bApplyLocalExposureToBloom = 
-					CVarBloomApplyLocalExposure.GetValueOnRenderThread()
-					&& View.FinalPostProcessSettings.LocalExposureMethod == ELocalExposureMethod::Bilateral
-					&& LocalExposureBilateralGridTexture != nullptr;
-
-				const bool bBloomSetupRequiredEnabled = View.FinalPostProcessSettings.BloomThreshold > -1.0f || bApplyLocalExposureToBloom;
+					// Reuse the main scene downsample chain if setup isn't required for gaussian bloom.
+					if (!SceneDownsampleChain.IsInitialized() || bBloomSetupRequiredEnabled)
+					{
+						FScreenPassTextureSlice DownsampleInput = bProcessEighthResolution ? EighthResSceneColor : (bProcessQuarterResolution ? QuarterResSceneColor : HalfResSceneColor);
 
-				// Reuse the main scene downsample chain if setup isn't required for gaussian bloom.
-				if (SceneDownsampleChain.IsInitialized() && !bBloomSetupRequiredEnabled)
-				{
-					LensFlareSceneDownsampleChain = &SceneDownsampleChain;
-				}
-				else
-				{
-					FScreenPassTextureSlice DownsampleInput = bProcessEighthResolution ? EighthResSceneColor : (bProcessQuarterResolution ? QuarterResSceneColor : HalfResSceneColor);
+						if (bBloomSetupRequiredEnabled)
+						{
+							const float BloomThreshold = View.FinalPostProcessSettings.BloomThreshold;
 
-					if (bBloomSetupRequiredEnabled)
-					{
-						const float BloomThreshold = View.FinalPostProcessSettings.BloomThreshold;
+							FBloomSetupInputs SetupPassInputs;
+							SetupPassInputs.SceneColor = DownsampleInput;
+							SetupPassInputs.EyeAdaptationBuffer = EyeAdaptationBuffer;
+							SetupPassInputs.EyeAdaptationParameters = &EyeAdaptationParameters;
+							SetupPassInputs.Threshold = BloomThreshold;
 
-						FBloomSetupInputs SetupPassInputs;
-						SetupPassInputs.SceneColor = DownsampleInput;
-						SetupPassInputs.EyeAdaptationBuffer = EyeAdaptationBuffer;
-						SetupPassInputs.EyeAdaptationParameters = &EyeAdaptationParameters;
-						SetupPassInputs.Threshold = BloomThreshold;
+							if (bApplyLocalExposureToBloom)
+							{
+								SetupPassInputs.LocalExposureParameters = &LocalExposureParameters;
+								SetupPassInputs.LocalExposureTexture = LocalExposureBilateralGridTexture;
+								SetupPassInputs.BlurredLogLuminanceTexture = LocalExposureBlurredLogLumTexture;
+							}
 
-						if (bApplyLocalExposureToBloom)
-						{
-							SetupPassInputs.LocalExposureParameters = &LocalExposureParameters;
-							SetupPassInputs.LocalExposureTexture = LocalExposureBilateralGridTexture;
-							SetupPassInputs.BlurredLogLuminanceTexture = LocalExposureBlurredLogLumTexture;
+							DownsampleInput = FScreenPassTextureSlice::CreateFromScreenPassTexture(GraphBuilder, AddBloomSetupPass(GraphBuilder, View, SetupPassInputs));
 						}
 
-						DownsampleInput = FScreenPassTextureSlice::CreateFromScreenPassTexture(GraphBuilder, AddBloomSetupPass(GraphBuilder, View, SetupPassInputs));
-					}
+						const bool bLogLumaInAlpha = false;
+						BloomDownsampleChain.Init(GraphBuilder, View, EyeAdaptationParameters, DownsampleInput, DownsampleChainQuality, (uint32)EBloomQuality::MAX, bLogLumaInAlpha, nullptr, bProcessEighthResolution ? 3 : (bProcessQuarterResolution ? 2 : 1));
 
-					const bool bLogLumaInAlpha = false;
-					BloomDownsampleChain.Init(GraphBuilder, View, EyeAdaptationParameters, DownsampleInput, DownsampleChainQuality, (uint32)EBloomQuality::MAX, bLogLumaInAlpha, nullptr, bProcessEighthResolution ? 3 : (bProcessQuarterResolution ? 2 : 1));
+						LensFlareSceneDownsampleChain = &BloomDownsampleChain;
+					}
 
-					LensFlareSceneDownsampleChain = &BloomDownsampleChain;
+					Bloom = AddGaussianBloomPasses(GraphBuilder, View, LensFlareSceneDownsampleChain);
 				}
-
-				Bloom = AddGaussianBloomPasses(GraphBuilder, View, LensFlareSceneDownsampleChain);
 			}
 
 			if (bLensFlareEnabled)
