Adding missing MeshPass bindings of uniform buffers.

diff --git a/Engine/Source/Runtime/Renderer/Private/HeterogeneousVolumes/HeterogeneousVolumesLiveShadingPipeline.cpp b/Engine/Source/Runtime/Renderer/Private/HeterogeneousVolumes/HeterogeneousVolumesLiveShadingPipeline.cpp
--- a/Engine/Source/Runtime/Renderer/Private/HeterogeneousVolumes/HeterogeneousVolumesLiveShadingPipeline.cpp
+++ b/Engine/Source/Runtime/Renderer/Private/HeterogeneousVolumes/HeterogeneousVolumesLiveShadingPipeline.cpp
@@ -650,9 +650,13 @@ void AddComputePass(
 				SingleShaderBindings.Add(ComputeShader->template GetUniformBufferParameter<FDeferredLightUniformStruct>(), PassParameters->DeferredLight.GetUniformBuffer());
 				SingleShaderBindings.Add(ComputeShader->template GetUniformBufferParameter<FForwardLightUniformParameters>(), PassParameters->ForwardLightStruct.GetUniformBuffer()->GetRHIRef());
 				SingleShaderBindings.Add(ComputeShader->template GetUniformBufferParameter<FVirtualShadowMapUniformParameters>(), PassParameters->VirtualShadowMapSamplingParameters.VirtualShadowMap.GetUniformBuffer()->GetRHIRef());
+				SingleShaderBindings.Add(ComputeShader->template GetUniformBufferParameter<FAdaptiveVolumetricShadowMaps>(), PassParameters->AVSM.GetUniformBuffer()->GetRHIRef());
+				SingleShaderBindings.Add(ComputeShader->template GetUniformBufferParameter<FAdaptiveVolumetricShadowMapUniformBufferParameters>(), PassParameters->AVSM.GetUniformBuffer()->GetRHIRef());
 				if constexpr (bWithLumen)
 				{
+					SingleShaderBindings.Add(ComputeShader->template GetUniformBufferParameter<FFogUniformParameters>(), PassParameters->FogStruct.GetUniformBuffer()->GetRHIRef());
 					SingleShaderBindings.Add(ComputeShader->template GetUniformBufferParameter<FLumenTranslucencyLightingUniforms>(), PassParameters->LumenGIVolumeStruct.GetUniformBuffer()->GetRHIRef());
+					SingleShaderBindings.Add(ComputeShader->template GetUniformBufferParameter<FBlueNoise>(), PassParameters->BlueNoise.GetUniformBuffer());
 				}
 				ComputeShader->ShaderLooseBindings.SetParameters(SingleShaderBindings, PassParameters);
 				ShaderBindings.Finalize(&PassShaders);
@@ -2141,6 +2145,7 @@ bool RenderVolumetricShadowMapForLightForHeterogeneousVolumeWithLiveShading(
 				{
 					FMeshDrawSingleShaderBindings SingleShaderBindings = ShaderBindings.GetSingleShaderBindings(SF_Compute);
 					ComputeShader->GetShaderBindings(Scene, Scene->GetFeatureLevel(), nullptr, *MaterialRenderProxy, Material, ShaderElementData, SingleShaderBindings);
+					SingleShaderBindings.Add(ComputeShader->template GetUniformBufferParameter<FBlueNoise>(), PassParameters->BlueNoise.GetUniformBuffer());
 					ComputeShader->ShaderLooseBindings.SetParameters(SingleShaderBindings, PassParameters);
 					ShaderBindings.Finalize(&PassShaders);
 				}
